HeterogeneousVolumes: Voxel Grid Pipeline constructs sparse ortho and frustum voxel grids, for rendering overlapping and preshaded volumes.

r.HeterogeneousVolumes.OrthoGrid (default = 1) Enables ortho grid acceleration structure build
r.HeterogeneousVolumes.OrthoGrid.ShadingRate (default = 4) Determines native voxel resolution for ortho grid. A value of 1 roughly corresponds to voxels of pixel width.
r.HeterogeneousVolumes.OrthoGrid.MaxBottomLevelMemoryInMegabytes (default = 128) Determines the bottom-level memory available for the ortho grid.

r.HeterogeneousVolumes.FrustumGrid (default = 1) Enables frustum grid acceleration structure build
r.HeterogeneousVolumes.FrustumGrid.ShadingRate (default = 4) Determines native voxel resolution for frustum grid. A value of 1 corresponds to voxels of pixel width.
r.HeterogeneousVolumes.FrustumGrid.DepthSliceCount (default = 512) Determines the number of depth slices when building the frustum grid. Depth slice count is inversely proportional to shading rate.
r.HeterogeneousVolumes.FrustumGrid.MaxBottomLevelMemoryInMegabytes (default = 128) Determines the bottom-level memory available for the frustum grid.

r.HeterogeneousVolumes.Tessellation.MinimumVoxelSizeInFrustum (default = 1) Determines the minimum voxel size when building the ortho grid.

diff --git a/Engine/Shaders/Private/HeterogeneousVolumes/HeterogeneousVolumesFrustumVoxelGrid.usf b/Engine/Shaders/Private/HeterogeneousVolumes/HeterogeneousVolumesFrustumVoxelGrid.usf
new file mode 100644
--- /dev/null
+++ b/Engine/Shaders/Private/HeterogeneousVolumes/HeterogeneousVolumesFrustumVoxelGrid.usf
@@ -0,0 +1,77 @@
+// Copyright Epic Games, Inc. All Rights Reserved.
+
+#include "HeterogeneousVolumesVoxelGridTypes.ush"
+
+#include "../Common.ush"
+#include "../DeferredShadingCommon.ush"
+#include "../ComputeShaderUtils.ush"
+#include "HeterogeneousVolumesTracingUtils.ush"
+#include "HeterogeneousVolumesVoxelGridTraversal.ush"
+
+int3 TopLevelGridResolution;
+
+StructuredBuffer<FTopLevelGridData> TopLevelGridBuffer;
+StructuredBuffer<FScalarGridData> ExtinctionGridBuffer;
+StructuredBuffer<FVectorGridData> EmissionGridBuffer;
+StructuredBuffer<FVectorGridData> ScatteringGridBuffer;
+
+float3 PrimitiveWorldBoundsMin;
+float3 PrimitiveWorldBoundsMax;
+
+RWStructuredBuffer<FTopLevelGridData> RWTopLevelGridBuffer;
+
+float4x4 ViewToWorld;
+int3 VoxelDimensions;
+
+float HalfFOV;
+float NearPlaneDepth;
+float FarPlaneDepth;
+
+bool BoxIntersectsVoxel(float3 WorldBoxCenter, float WorldBoxRadius, uint3 VoxelIndex)
+{
+	float3 VoxelPosMin = VoxelIndex;
+	float3 VoxelPosMax = (VoxelIndex + 1);
+
+	float3 ViewPosMin = VoxelToView(VoxelPosMin, VoxelDimensions, NearPlaneDepth, FarPlaneDepth);
+	float3 ViewPosMax = VoxelToView(VoxelPosMax, VoxelDimensions, NearPlaneDepth, FarPlaneDepth);
+
+	float3 WorldPosMin = mul(float4(ViewPosMin, 1), ViewToWorld).xyz;
+	float3 WorldPosMax = mul(float4(ViewPosMax, 1), ViewToWorld).xyz;
+
+	float3 WorldPosCenter = (WorldPosMin + WorldPosMax) * 0.5;
+	float WorldRadius = length((WorldPosMax - WorldPosMin) * 0.5);
+
+	float DistanceToCenters = length(WorldBoxCenter - WorldPosCenter);
+	float SumOfRadii = WorldBoxRadius + WorldRadius;
+	return DistanceToCenters < SumOfRadii;
+}
+
+// Mark voxels which intersect a given primitive
+[numthreads(THREADGROUP_SIZE_3D, THREADGROUP_SIZE_3D, THREADGROUP_SIZE_3D)]
+void MarkTopLevelGridVoxelsForFrustumGridCS(
+	uint3 GroupId : SV_GroupID,
+	uint3 GroupThreadId : SV_GroupThreadID,
+	uint3 DispatchThreadId : SV_DispatchThreadID
+)
+{
+	uint3 VoxelIndex = DispatchThreadId;
+	if (any(VoxelIndex >= TopLevelGridResolution))
+	{
+		return;
+	}
+
+	uint LinearIndex = GetLinearIndex(VoxelIndex, TopLevelGridResolution);
+	FTopLevelGridData TopLevelGridData = RWTopLevelGridBuffer[LinearIndex];
+	uint3 VoxelResolution = GetBottomLevelVoxelResolution(TopLevelGridData);
+
+	float3 WorldPrimitiveOrigin = (PrimitiveWorldBoundsMax + PrimitiveWorldBoundsMin) * 0.5;
+	float3 WorldPrimitiveExtent = (PrimitiveWorldBoundsMax - PrimitiveWorldBoundsMin) * 0.5;
+	if (BoxIntersectsVoxel(WorldPrimitiveOrigin, length(WorldPrimitiveExtent), VoxelIndex))
+	{
+		VoxelResolution = 4;
+	}
+
+	SetBottomLevelIndex(TopLevelGridData, EMPTY_VOXEL_INDEX);
+	SetBottomLevelVoxelResolution(TopLevelGridData, VoxelResolution);
+	RWTopLevelGridBuffer[LinearIndex] = TopLevelGridData;
+}
diff --git a/Engine/Shaders/Private/HeterogeneousVolumes/HeterogeneousVolumesFrustumVoxelGridUtils.ush b/Engine/Shaders/Private/HeterogeneousVolumes/HeterogeneousVolumesFrustumVoxelGridUtils.ush
new file mode 100644
--- /dev/null
+++ b/Engine/Shaders/Private/HeterogeneousVolumes/HeterogeneousVolumesFrustumVoxelGridUtils.ush
@@ -0,0 +1,161 @@
+// Copyright Epic Games, Inc. All Rights Reserved.
+
+float Exponential(float sigma, float x)
+{
+	// X = 1 - e^(-sigma * x)
+	return 1 - exp(-sigma * x);
+}
+
+float InverseExponential(float sigma, float X)
+{
+	// X = 1 - e^(-sigma * x)
+	// x = ln(1 - X) / -sigma
+	return log(1 - X) / -sigma;
+}
+
+float FindSigma(float x, float X)
+{
+	// X = 1 - e^(-sigma * x)
+	// e^(-sigma * x) = 1 - X
+	// -sigma * x = ln(1 - X)
+	// sigma = -ln(1 - X) / x
+	return -log(1 - X) / x;
+}
+
+float ExponentialRemap(float ViewDepth, float NearPlaneDepth, float FarPlaneDepth)
+{
+	float Sigma = FindSigma(FarPlaneDepth, 0.9999);
+	float NormalizedDepth = Exponential(Sigma, ViewDepth);
+	return NormalizedDepth;
+}
+
+float InverseExponentialRemap(float NormalizedDepth, float NearPlaneDepth, float FarPlaneDepth)
+{
+	float Sigma = FindSigma(FarPlaneDepth, 0.9999);
+	float ViewDepth = InverseExponential(Sigma, NormalizedDepth);
+	return ViewDepth;
+}
+
+float LinearRemap(float ViewDepth, float NearPlaneDepth, float FarPlaneDepth)
+{
+	float NormalizedDepth = (ViewDepth - NearPlaneDepth) / (FarPlaneDepth - NearPlaneDepth);
+	return NormalizedDepth;
+}
+
+float InverseLinearRemap(float NormalizedDepth, float NearPlaneDepth, float FarPlaneDepth)
+{
+	float ViewDepth = NormalizedDepth * (FarPlaneDepth - NearPlaneDepth) + NearPlaneDepth;
+	return ViewDepth;
+}
+
+float SquareDistanceRemap(float ViewDepth, float NearPlaneDepth, float FarPlaneDepth)
+{
+	float NormalizedDepth = ViewDepth / 20.0;
+	if (ViewDepth > 2.0)
+	{
+		float LogDepth = log2(ViewDepth);
+		float LogFarPlaneDepth = log2(FarPlaneDepth);
+		NormalizedDepth = LogDepth / LogFarPlaneDepth + 0.1;
+	}
+
+	return NormalizedDepth;
+}
+
+float InverseSquareDistanceRemap(float NormalizedDepth, float NearPlaneDepth, float FarPlaneDepth)
+{
+	float ViewDepth = NormalizedDepth * 20.0;
+	if (ViewDepth > 2.0)
+	{
+		float LogFarPlaneDepth = log2(FarPlaneDepth);
+		float LogDepth = (NormalizedDepth - 0.1) * LogFarPlaneDepth;
+		ViewDepth = pow(2, LogDepth);
+	}
+
+	return ViewDepth;
+}
+
+float3 ScreenToVoxel(float3 ScreenPos, int3 VoxelDimensions)
+{
+	float3 NDC = float3(ScreenPos.xy * 0.5 + 0.5, ScreenPos.z);
+	float3 VoxelPos = NDC * VoxelDimensions;
+	return VoxelPos;
+}
+
+float3 VoxelToScreen(float3 VoxelPos, int3 VoxelDimensions)
+{
+	float3 NDC = VoxelPos / VoxelDimensions;
+	float3 ScreenPos = float3(NDC.xy * 2.0 - 1.0, NDC.z);
+	return ScreenPos;
+}
+
+float PerspectiveRemap(float ViewDepth, int3 VoxelDimensions)
+{
+	float3 VoxelPos = ScreenToVoxel(float3(0, 0, ViewDepth), VoxelDimensions);
+	float NormalizedDepth = VoxelPos.z / VoxelDimensions.z;
+	return NormalizedDepth;
+}
+
+float InversePerspectiveRemap(float NormalizedDepth, int3 VoxelDimensions)
+{
+	float3 VoxelPos = ScreenToVoxel(float3(0, 0, NormalizedDepth), VoxelDimensions);
+	float ViewDepth = VoxelPos.z;
+	return ViewDepth;
+}
+
+float3 ViewToVoxel(float3 ViewPos, int3 VoxelDimensions, float NearPlaneDepth, float FarPlaneDepth)
+{
+	//float4 ClipPos = mul(float4(ViewPos, 1), ViewToClip);
+	//float3 ScreenPos = ClipPos.xyz / ClipPos.w;
+	float2 AspectRatio = VoxelDimensions.xy / float(VoxelDimensions.x);
+	float2 ScreenPos = ViewPos.xy / (AspectRatio * ViewPos.z);
+
+#if 0
+	float NormalizedDepth = SquareDistanceRemap(ViewPos.z, NearPlaneDepth, FarPlaneDepth);
+#elif 0
+	float NormalizedDepth = ExponentialRemap(ViewPos.z, NearPlaneDepth, FarPlaneDepth);
+#elif 1
+	float NormalizedDepth = LinearRemap(ViewPos.z, NearPlaneDepth, FarPlaneDepth);
+#else
+	float NormalizedDepth = PerspectiveRemap(ViewPos.z, VoxelDimensions);
+#endif
+
+	float3 NDC = float3(ScreenPos.xy * 0.5 + 0.5, NormalizedDepth);
+	float3 VoxelPos = NDC * VoxelDimensions;
+	return VoxelPos;
+}
+
+float3 VoxelToView(float3 VoxelPos, int3 VoxelDimensions, float NearPlaneDepth, float FarPlaneDepth)
+{
+	float3 NDC = VoxelPos / VoxelDimensions;
+	float2 ScreenPos = NDC.xy * 2.0 - 1.0;
+#if 0
+	float ViewDepth = InverseSquareDistanceRemap(NDC.z, NearPlaneDepth, FarPlaneDepth);
+#elif 0
+	float ViewDepth = InverseExponentialRemap(NDC.z, NearPlaneDepth, FarPlaneDepth);
+#elif 1
+	float ViewDepth = InverseLinearRemap(NDC.z, NearPlaneDepth, FarPlaneDepth);
+#else
+	float ViewDepth = InversePerspectiveRemap(NDC.z, VoxelDimensions);
+#endif
+
+	// TODO: Guarantee no view-space translation exists..
+	float2 AspectRatio = VoxelDimensions.xy / float(VoxelDimensions.x);
+	float3 ViewPos = float3(ScreenPos * AspectRatio * ViewDepth, ViewDepth);
+	return ViewPos;
+}
+
+float2 ClipRay(float3 Origin, float3 Direction, float TMin, float TMax, float4 Plane)
+{
+	float3 Normal = Plane.xyz;
+	float d = Plane.w;
+	float HitT = -(d + dot(Normal, Origin)) * rcp(dot(Normal, Direction));
+
+	float2 RayT = float2(TMin, TMax);
+	if ((HitT > RayT.x) && (HitT < RayT.y))
+	{
+		float SignedDistance = dot(Origin, Normal) + d;
+		RayT = SignedDistance < 0 ? float2(HitT, RayT.y) : float2(RayT.x, HitT);
+	}
+
+	return RayT;
+}
\ No newline at end of file
diff --git a/Engine/Shaders/Private/HeterogeneousVolumes/HeterogeneousVolumesMaterialBakingPipeline.usf b/Engine/Shaders/Private/HeterogeneousVolumes/HeterogeneousVolumesMaterialBakingPipeline.usf
--- a/Engine/Shaders/Private/HeterogeneousVolumes/HeterogeneousVolumesMaterialBakingPipeline.usf
+++ b/Engine/Shaders/Private/HeterogeneousVolumes/HeterogeneousVolumesMaterialBakingPipeline.usf
@@ -6,6 +6,10 @@
 
 #include "/Engine/Generated/Material.ush"
 #include "../PositionReconstructionCommon.ush"
+#include "../MortonCode.ush"
+#include "../ComputeShaderUtils.ush"
+#include "../Hash.ush"
+#include "../HashTable.ush"
 #include "HeterogeneousVolumesLiveShadingUtils.ush"
 
 #ifndef THREADGROUP_SIZE_3D
@@ -60,4 +64,650 @@ void HeterogeneousVolumesBakeMaterialCS(
 		RWEmissionTexture[VoxelIndex] = SampleEmissive(PixelMaterialInputs);
 		RWAlbedoTexture[VoxelIndex] = SampleAlbedo(PixelMaterialInputs);
 	}
+}
+
+// Global transmittance volume
+
+#include "HeterogeneousVolumesVoxelGridTypes.ush"
+#include "HeterogeneousVolumesVoxelGridUtils.ush"
+
+int3 TopLevelGridResolution;
+float3 TopLevelGridWorldBoundsMin;
+float3 TopLevelGridWorldBoundsMax;
+
+float3 PrimitiveWorldBoundsMin;
+float3 PrimitiveWorldBoundsMax;
+
+int BottomLevelGridBufferSize;
+
+void CalcVoxelBounds(float3 VoxelIndex, inout float3 VoxelBoundsMin, inout float3 VoxelBoundsMax)
+{
+	float3 TopLevelGridWorldBoundsExtent = TopLevelGridWorldBoundsMax - TopLevelGridWorldBoundsMin;
+	VoxelBoundsMin = TopLevelGridWorldBoundsMin + TopLevelGridWorldBoundsExtent * (VoxelIndex / TopLevelGridResolution);
+	VoxelBoundsMax = TopLevelGridWorldBoundsMin + TopLevelGridWorldBoundsExtent * ((VoxelIndex + 1) / TopLevelGridResolution);
+}
+
+float3 GetVoxelCenter(float3 VoxelIndex, inout float3 VoxelBoundsMin, inout float3 VoxelBoundsMax)
+{
+	float3 TopLevelGridWorldBoundsExtent = TopLevelGridWorldBoundsMax - TopLevelGridWorldBoundsMin;
+	return TopLevelGridWorldBoundsMin + TopLevelGridWorldBoundsExtent * ((VoxelIndex + 0.5) / TopLevelGridResolution);
+}
+
+bool PrimitiveIntersectsVoxel(float3 VoxelBoundsMin, float3 VoxelBoundsMax)
+{
+	if (any(PrimitiveWorldBoundsMin > VoxelBoundsMax) || any(VoxelBoundsMin > PrimitiveWorldBoundsMax))
+	{
+		return false;
+	}
+
+	return true;
+}
+
+struct FRasterTileData
+{
+	uint TopLevelGridLinearIndex;
+	uint BottomLevelGridLinearOffset;
+};
+
+Buffer<uint> RasterTileAllocatorBuffer;
+StructuredBuffer<FRasterTileData> RasterTileBuffer;
+
+RWStructuredBuffer<uint> RWTopLevelGridBuffer;
+RWBuffer<uint> RWBottomLevelGridAllocatorBuffer;
+
+RWStructuredBuffer<FScalarGridData> RWExtinctionGridBuffer;
+RWStructuredBuffer<FVectorGridData> RWEmissionGridBuffer;
+RWStructuredBuffer<FVectorGridData> RWScatteringGridBuffer;
+
+// Indirection Grid
+int FixedBottomLevelResolution;
+RWStructuredBuffer<FTopLevelGridData> RWIndirectionGridBuffer;
+
+// Hash index to voxel index
+RWBuffer<uint> RWHashToVoxelBuffer;
+
+// Sampling mode
+int bSampleAtVertices;
+
+float HomogeneousThreshold;
+
+groupshared float3 Extinction[THREADGROUP_SIZE_1D];
+groupshared float3 ExtinctionSum[THREADGROUP_SIZE_1D];
+groupshared float ExtinctionVariance[THREADGROUP_SIZE_1D];
+groupshared float ZeroCrossings[THREADGROUP_SIZE_1D];
+groupshared FTopLevelGridData AllocatedGridData;
+groupshared int3 AllocatedVoxelResolution;
+groupshared int AllocatedVoxelCount;
+
+float GetZeroThreshold()
+{
+	return 1.0e-6;
+}
+
+float CalcSampleVariance(float3 MeanValue, float3 SampleValue)
+{
+	float3 SampleVariance = MeanValue - SampleValue;
+	return dot(SampleVariance, SampleVariance);
+}
+
+float CalcZeroCrossing(float3 SampleValue)
+{
+	float EmptyVoxelThreshold = 1.0e-5;
+	return all(SampleValue < EmptyVoxelThreshold) ? 1.0 : 0.0;
+}
+
+uint2 CalcStencil(uint VoxelCount)
+{
+	//int x = i % 4;
+	//int y = i / 4;
+	//int z = i % 16;
+	uint2 Stencil = 0;
+	
+	float ZeroThreshold = GetZeroThreshold();
+	for (int i = 0; i < min(VoxelCount, 32); ++i)
+	{
+		int bit = i;
+		Stencil[0] |= any(Extinction[i] > ZeroThreshold) ? 1 << bit : 0;
+	}
+	
+	for (i = 32; i < min(VoxelCount, 64); ++i)
+	{
+		int bit = i - 32;
+		Stencil[1] |= any(Extinction[i] > ZeroThreshold) ? 1 << bit : 0;
+	}
+
+	return Stencil;
+}
+
+float ExtinctionMax()
+{
+	return 0.01;
+}
+
+float ExtinctionFixedPointCeiling()
+{
+	//int QuantizeBits = 16;
+	int QuantizeBits = 8;
+	return 1 << QuantizeBits;
+}
+
+// TODO: Parallel
+uint QuantizeOpticalThickness(uint VoxelCount)
+{
+	float ZeroThreshold = GetZeroThreshold();
+
+	float3 ExtinctionMean = 0;
+	uint NonZeroCount = 0;
+	for (uint i = 0; i < VoxelCount; ++i)
+	{
+		if (any(Extinction[i] > ZeroThreshold))
+		{
+			ExtinctionMean += Extinction[i];
+			NonZeroCount++;
+		}
+	}
+
+	if (NonZeroCount > 0)
+	{
+		ExtinctionMean /= NonZeroCount;
+	}
+
+	return Luminance(ExtinctionMean) / ExtinctionMax() * ExtinctionFixedPointCeiling();
+}
+
+// TODO: Parallel
+float3 QuantizeExtinction(uint VoxelCount)
+{
+	float ZeroThreshold = GetZeroThreshold();
+
+	float3 ExtinctionMean = 0;
+	uint NonZeroCount = 0;
+	for (uint i = 0; i < VoxelCount; ++i)
+	{
+		if (any(Extinction[i] > ZeroThreshold))
+		{
+			ExtinctionMean += Extinction[i];
+			NonZeroCount++;
+		}
+	}
+
+	if (NonZeroCount > 0)
+	{
+		ExtinctionMean /= NonZeroCount;
+	}
+
+	uint3 QuantizedValue = ExtinctionMean / ExtinctionMax() * ExtinctionFixedPointCeiling();
+	return QuantizedValue / ExtinctionFixedPointCeiling() * ExtinctionMax();
+}
+
+uint CalcBottomLevelHashKey(uint VoxelCount)
+{
+	uint HashKey = 0;
+
+	// Hash Resolution
+	HashKey = MurmurMix(VoxelCount);
+
+	// Hash quantized average optical thickness
+	//HashKey = MurmurAdd(HashKey, QuantizeOpticalThickness(VoxelCount));
+	float3 QuantizedExtinction = QuantizeExtinction(VoxelCount);
+	uint QuantizedValue = Luminance(QuantizedExtinction / ExtinctionMax()) * ExtinctionFixedPointCeiling();
+	HashKey = MurmurAdd(HashKey, QuantizedValue);
+
+	// Hash bottom-voxel stencil
+	uint2 Stencil = CalcStencil(VoxelCount);
+	HashKey = MurmurAdd(HashKey, Stencil[0]);
+	HashKey = MurmurAdd(HashKey, Stencil[1]);
+
+	return HashKey;
+}
+
+bool HashTableFind2(uint Key, out uint Index)
+{
+	// Zero is reserved as invalid
+	Key++;
+
+	uint NumLoops = 0;
+	LOOP
+		[allow_uav_condition]
+	for (Index = MurmurMix(Key);; Index++)
+	{
+		//Belt and braces safety code to prevent inf loops if tables are malformed.
+		if (++NumLoops > HashTableSize)
+		{
+			break;
+		}
+
+		//Index &= HashTableSize - 1;
+		Index = Index % HashTableSize;
+
+		uint StoredKey = RWHashTable[Index];
+		if (StoredKey != Key)
+		{
+			if (StoredKey != 0)
+				continue;
+		}
+		else
+		{
+			return true;
+		}
+
+		break;
+	}
+
+	return false;
+}
+
+[numthreads(THREADGROUP_SIZE_3D, THREADGROUP_SIZE_3D, THREADGROUP_SIZE_3D)]
+void RasterizeBottomLevelGridCS(
+	uint3 GroupId : SV_GroupID,
+	uint3 GroupThreadId : SV_GroupThreadID
+)
+{
+	uint RasterTileIndex = GetUnWrappedDispatchGroupId(GroupId);
+	if (RasterTileIndex >= RasterTileAllocatorBuffer[0])
+	{
+		return;
+	}
+
+	uint LinearThreadIndex = MortonEncode3(GroupThreadId);
+	Extinction[LinearThreadIndex] = 0.0;
+
+	uint TopLevelGridLinearIndex = RasterTileBuffer[RasterTileIndex].TopLevelGridLinearIndex;
+	uint3 TopLevelGridVoxelIndex = GetVoxelIndex(TopLevelGridLinearIndex, TopLevelGridResolution);
+
+	float3 TopLevelGridVoxelWorldBoundsMin;
+	float3 TopLevelGridVoxelWorldBoundsMax;
+	CalcVoxelBounds(TopLevelGridVoxelIndex, TopLevelGridVoxelWorldBoundsMin, TopLevelGridVoxelWorldBoundsMax);
+	float3 TopLevelGridVoxelWorldBoundsExtent = TopLevelGridVoxelWorldBoundsMax - TopLevelGridVoxelWorldBoundsMin;
+
+	// Note: Raster tiles map to non-empty voxels by definition
+	if (PrimitiveIntersectsVoxel(TopLevelGridVoxelWorldBoundsMin, TopLevelGridVoxelWorldBoundsMax))
+	{
+		// Map thread id to voxel-space position
+		uint BottomLevelGridLinearOffset = RasterTileBuffer[RasterTileIndex].BottomLevelGridLinearOffset;
+
+		FTopLevelGridData TopLevelGridData = CreateTopLevelGridData(RWTopLevelGridBuffer[TopLevelGridLinearIndex]);
+		int3 BottomLevelVoxelResolution = GetBottomLevelVoxelResolution(TopLevelGridData);
+
+	#if DIM_ENABLE_INDIRECTION_GRID
+		int3 BottomLevelVoxelIndex = MortonDecode3(BottomLevelGridLinearOffset) * THREADGROUP_SIZE_3D + GroupThreadId;
+		BottomLevelVoxelResolution *= FixedBottomLevelResolution;
+	#else
+		int3 BottomLevelVoxelIndex = GetVoxelIndex(BottomLevelGridLinearOffset * THREADGROUP_SIZE_1D + MortonEncode3(GroupThreadId), BottomLevelVoxelResolution);
+	#endif
+
+		float3 Emission = 0;
+		float3 Albedo = 0;
+		if (all(BottomLevelVoxelIndex < BottomLevelVoxelResolution))
+		{
+#if 1
+			float3 UVW;
+			if (bSampleAtVertices)
+			{
+				UVW = BottomLevelVoxelIndex / float3(BottomLevelVoxelResolution - 1);
+			}
+			else
+			{
+				UVW = float3(BottomLevelVoxelIndex + 0.5) / float3(BottomLevelVoxelResolution);
+			}
+			float3 WorldPosition = UVW * TopLevelGridVoxelWorldBoundsExtent + TopLevelGridVoxelWorldBoundsMin;
+
+			// Setup evaluation context
+			FMaterialPixelParameters MaterialParameters = MakeInitializedMaterialPixelParameters();
+			MaterialParameters.PrimitiveId = PrimitiveId;
+			MaterialParameters.AbsoluteWorldPosition = LWCPromote(WorldPosition);
+
+			// Evaluate material graph
+			FPixelMaterialInputs PixelMaterialInputs;
+			CalcPixelMaterialInputs(MaterialParameters, PixelMaterialInputs);
+
+			// Extract volume rendering coefficients
+			Extinction[LinearThreadIndex] = SampleExtinctionCoefficients(PixelMaterialInputs);
+			Emission = SampleEmissive(PixelMaterialInputs);
+			//Albedo = SampleAlbedo(PixelMaterialInputs);
+			Albedo = SampleAlbedo(PixelMaterialInputs) * Extinction[LinearThreadIndex];
+#else
+			Extinction[LinearThreadIndex] = 0.01;
+#endif
+		}
+
+	#if DIM_ENABLE_INDIRECTION_GRID
+		AllocatedVoxelResolution = FixedBottomLevelResolution;
+	#else
+		AllocatedVoxelResolution = BottomLevelVoxelResolution;
+	#endif
+		AllocatedVoxelCount = AllocatedVoxelResolution.x * AllocatedVoxelResolution.y * AllocatedVoxelResolution.z;
+
+		GroupMemoryBarrier();
+
+		// Sum contribution
+#if 0
+		ExtinctionSum[0] = 0;
+		for (uint Index = 0; Index < THREADGROUP_SIZE_1D; ++Index)
+		{
+			ExtinctionSum[0] += Extinction[Index];
+		}
+#else
+		ExtinctionSum[LinearThreadIndex] = Extinction[LinearThreadIndex];
+		// Parallel Sum
+		for (int NeighborOffset = 1; NeighborOffset < THREADGROUP_SIZE_1D; NeighborOffset = (NeighborOffset << 1))
+		{
+			int NeighborIndex = LinearThreadIndex + NeighborOffset;
+			ExtinctionSum[LinearThreadIndex] += (NeighborIndex < THREADGROUP_SIZE_1D) ? ExtinctionSum[NeighborIndex] : 0.0;
+
+			//GroupMemoryBarrierWithGroupSync();
+			GroupMemoryBarrier();
+		}
+#endif
+
+		GroupMemoryBarrierWithGroupSync();
+
+#if DIM_ENABLE_HOMOGENEOUS_AGGREGATION
+		// Calculate variance
+		float3 ExtinctionMean = ExtinctionSum[0] * rcp(AllocatedVoxelCount);
+		ExtinctionVariance[LinearThreadIndex] = CalcSampleVariance(ExtinctionMean, Extinction[LinearThreadIndex]);
+		// Reuse ExtinctionSum..
+		ZeroCrossings[LinearThreadIndex] = CalcZeroCrossing(Extinction[LinearThreadIndex]);
+
+		// Parallel Sum
+		for (NeighborOffset = 1; NeighborOffset < AllocatedVoxelCount; NeighborOffset = (NeighborOffset << 1))
+		{
+			int NeighborIndex = LinearThreadIndex + NeighborOffset;
+			ExtinctionVariance[LinearThreadIndex] += (NeighborIndex < AllocatedVoxelCount) ? ExtinctionVariance[NeighborIndex] : 0.0;
+			ZeroCrossings[LinearThreadIndex] += (NeighborIndex < AllocatedVoxelCount) ? ZeroCrossings[NeighborIndex] : 0.0;
+
+			//GroupMemoryBarrierWithGroupSync();
+			GroupMemoryBarrier();
+		}
+#if 0	
+		// Linear Sum
+		for (uint Index = 0; Index < AllocatedVoxelCount; ++Index)
+		{
+			ExtinctionVariance[0] += ExtinctionVariance[Index];
+		}
+#endif
+
+		GroupMemoryBarrierWithGroupSync();
+		ExtinctionVariance[0] *= rcp(AllocatedVoxelCount - 1);
+
+		float StdRelError = sqrt(ExtinctionVariance[0] * rcp(AllocatedVoxelCount)) / ExtinctionMean;
+
+		// Declare low-variance to constitute a region of homogeneity
+		if (all(GroupThreadId == 0) && (StdRelError < HomogeneousThreshold) && (ZeroCrossings[0] < 4))
+		{
+			Extinction[0] = ExtinctionMean;
+			AllocatedVoxelCount = AllocatedVoxelResolution = 1;
+		}
+#endif // #if DIM_ENABLE_HOMOGENEOUS_AGGREGATION
+
+		bool bAllocate = true;
+		uint BottomLevelIndex = EMPTY_VOXEL_INDEX;
+
+	#if DIM_ENABLE_INDIRECTION_GRID
+		uint IndirectionIndexOffset = GetBottomLevelIndex(TopLevelGridData);
+		uint IndirectionGridLinearOffset = IndirectionIndexOffset + BottomLevelGridLinearOffset;
+
+		FTopLevelGridData IndirectionData = RWIndirectionGridBuffer[IndirectionGridLinearOffset];
+		bool bWasAlreadyAllocated = IsBottomLevelAllocated(IndirectionData);
+		if (all(GroupThreadId == 0) && any(ExtinctionSum[0] > 1.0e-4) && !bWasAlreadyAllocated)
+	#else
+		bool bWasAlreadyAllocated = IsBottomLevelAllocated(TopLevelGridData);
+		if (all(GroupThreadId == 0) && any(ExtinctionSum[0] > GetZeroThreshold()) && !bWasAlreadyAllocated)
+	#endif
+		{
+			// Voxel Hashing
+			uint HashIndex = 0;
+	#if DIM_ENABLE_VOXEL_HASHING
+			if (AllocatedVoxelCount == 1)
+			{
+				uint HashKey = CalcBottomLevelHashKey(AllocatedVoxelCount);
+				bAllocate = !HashTableFind2(HashKey, HashIndex);
+			}
+	#endif // DIM_ENABLE_VOXEL_HASHING
+
+#if 0
+			// DEBUG: Only process aggregated voxels..
+			if (AllocatedVoxelCount != 1)
+			{
+				BottomLevelIndex = EMPTY_VOXEL_INDEX;
+				AllocatedVoxelCount = AllocatedVoxelResolution = 0;
+				bAllocate = false;
+			}
+#endif
+
+			// Otherwise allocate
+			if (bAllocate)
+			{
+				InterlockedAdd(RWBottomLevelGridAllocatorBuffer[0], AllocatedVoxelCount, BottomLevelIndex);
+
+				// Guard against over allocation
+				if (BottomLevelIndex + AllocatedVoxelCount > BottomLevelGridBufferSize)
+				{
+					// Declare the buffer to be filled
+					uint Dummy;
+					InterlockedExchange(RWBottomLevelGridAllocatorBuffer[0], BottomLevelGridBufferSize, Dummy);
+
+					BottomLevelIndex = EMPTY_VOXEL_INDEX;
+					AllocatedVoxelCount = AllocatedVoxelResolution = 0;
+				}
+			}
+	#if DIM_ENABLE_VOXEL_HASHING
+			else
+			{
+				BottomLevelIndex = RWHashToVoxelBuffer[HashIndex];
+			}
+	#endif // DIM_ENABLE_VOXEL_HASHING
+
+			// Update the index with the properly allocated index
+			AllocatedGridData = TopLevelGridData;
+			SetBottomLevelIndex(AllocatedGridData, BottomLevelIndex);
+			SetBottomLevelVoxelResolution(AllocatedGridData, AllocatedVoxelResolution);
+
+	#if DIM_ENABLE_INDIRECTION_GRID
+			RWIndirectionGridBuffer[IndirectionGridLinearOffset] = AllocatedGridData;
+	#else
+			InterlockedCompareStore(RWTopLevelGridBuffer[TopLevelGridLinearIndex], TopLevelGridData.PackedData[0], AllocatedGridData.PackedData[0]);
+	#endif
+		}
+
+		AllMemoryBarrierWithGroupSync();
+		//AllMemoryBarrier();
+
+#if DIM_ENABLE_VOXEL_HASHING
+		// Compositing with hashed voxels should require copy-on-write, but it is not currently supported
+		if (bAllocate && LinearThreadIndex < AllocatedVoxelCount)
+#else
+		if (LinearThreadIndex < AllocatedVoxelCount)
+#endif
+		{
+	#if DIM_ENABLE_INDIRECTION_GRID
+			FTopLevelGridData IndirectionData = RWIndirectionGridBuffer[IndirectionGridLinearOffset];
+
+			if (IsBottomLevelAllocated(IndirectionData))
+			{
+				uint BottomLevelVoxelLinearIndex = GetBottomLevelIndex(IndirectionData) + LinearThreadIndex;
+
+				if (bWasAlreadyAllocated)
+				{
+					Extinction[LinearThreadIndex] += GetExtinction(RWExtinctionGridBuffer[BottomLevelVoxelLinearIndex]);
+					Emission += GetEmission(RWEmissionGridBuffer[BottomLevelVoxelLinearIndex]);
+					Albedo += GetAlbedo(RWScatteringGridBuffer[BottomLevelVoxelLinearIndex]);
+				}
+
+				SetExtinction(RWExtinctionGridBuffer[BottomLevelVoxelLinearIndex], Extinction[LinearThreadIndex]);
+				SetEmission(RWEmissionGridBuffer[BottomLevelVoxelLinearIndex], Emission);
+				SetAlbedo(RWScatteringGridBuffer[BottomLevelVoxelLinearIndex], Albedo);
+			}
+	#else
+			FTopLevelGridData TopLevelGridData = CreateTopLevelGridData(RWTopLevelGridBuffer[TopLevelGridLinearIndex]);
+
+			if (IsBottomLevelAllocated(TopLevelGridData))
+			{
+				uint BottomLevelVoxelLinearIndex = GetBottomLevelIndex(TopLevelGridData) + MortonEncode3(BottomLevelVoxelIndex);
+
+				if (bWasAlreadyAllocated)
+				{
+					Extinction[LinearThreadIndex] += GetExtinction(RWExtinctionGridBuffer[BottomLevelVoxelLinearIndex]);
+					Emission += GetEmission(RWEmissionGridBuffer[BottomLevelVoxelLinearIndex]);
+					Albedo += GetAlbedo(RWScatteringGridBuffer[BottomLevelVoxelLinearIndex]);
+				}
+
+				SetExtinction(RWExtinctionGridBuffer[BottomLevelVoxelLinearIndex], Extinction[LinearThreadIndex]);
+				SetEmission(RWEmissionGridBuffer[BottomLevelVoxelLinearIndex], Emission);
+				SetAlbedo(RWScatteringGridBuffer[BottomLevelVoxelLinearIndex], Albedo);
+			}
+	#endif
+		}
+
+	#if DIM_ENABLE_VOXEL_HASHING
+		// Update deduplication structure
+		if (all(GroupThreadId == 0) && AllocatedVoxelCount == 1)
+		//if (all(GroupThreadId == 0) && any(ExtinctionSum[0] > GetZeroThreshold()) && bAllocate)
+		{
+			uint HashKey = CalcBottomLevelHashKey(AllocatedVoxelCount);
+			uint HashIndex = 0;
+			if (HashTableAdd(HashKey, HashIndex))
+			{
+				uint Dummy;
+				InterlockedExchange(RWHashToVoxelBuffer[HashIndex], BottomLevelIndex, Dummy);
+			}
+		}
+	#endif // DIM_ENABLE_VOXEL_HASHING
+	}
+}
+
+int3 VoxelDimensions;
+float4x4 ViewToWorld;
+
+float HalfFOV;
+float NearPlaneDepth;
+float FarPlaneDepth;
+
+#include "HeterogeneousVolumesFrustumVoxelGridUtils.ush"
+#include "../BlueNoise.ush"
+
+groupshared float GSExtinctionSum[THREADGROUP_SIZE_1D];
+
+bool EvaluationPositionIntersectsPrimitive(float3 WorldPosition)
+{
+	return all(WorldPosition > PrimitiveWorldBoundsMin) && all(WorldPosition < PrimitiveWorldBoundsMax);
+}
+
+[numthreads(THREADGROUP_SIZE_3D, THREADGROUP_SIZE_3D, THREADGROUP_SIZE_3D)]
+void RasterizeBottomLevelFroxelGridCS(
+	uint3 GroupId : SV_GroupID,
+	uint3 GroupThreadId : SV_GroupThreadID
+)
+{
+	uint RasterTileIndex = GetUnWrappedDispatchGroupId(GroupId);
+	if (RasterTileIndex >= RasterTileAllocatorBuffer[0])
+	{
+		return;
+	}
+
+	uint TopLevelGridLinearIndex = RasterTileBuffer[RasterTileIndex].TopLevelGridLinearIndex;
+	FTopLevelGridData TopLevelGridData = CreateTopLevelGridData(RWTopLevelGridBuffer[TopLevelGridLinearIndex]);
+	int3 BottomLevelVoxelResolution = GetBottomLevelVoxelResolution(TopLevelGridData);
+
+	// Setup evaluation context
+	float3 TopLevelVoxelPos = GetVoxelIndex(TopLevelGridLinearIndex, TopLevelGridResolution);
+	float3 BottomLevelVoxelPos = TopLevelVoxelPos * BottomLevelVoxelResolution + GroupThreadId;
+
+	float2 JitterXY = 0.5;
+	//float JitterZ = 0.5;
+	float JitterZ = BlueNoiseScalar(BottomLevelVoxelPos.xy, View.StateFrameIndex);
+	float3 Jitter = float3(JitterXY, JitterZ);
+
+	float3 LocalVoxelPos = GroupThreadId + Jitter;
+	float3 VoxelPosition = TopLevelVoxelPos + LocalVoxelPos / BottomLevelVoxelResolution;
+
+	float3 ViewPos = VoxelToView(VoxelPosition, VoxelDimensions, NearPlaneDepth, FarPlaneDepth);
+	float3 WorldPosition = mul(float4(ViewPos, 1), ViewToWorld).xyz;
+
+	float3 Extinction = 0;
+	float3 Emission = 0;
+	float3 Albedo = 0;
+
+	if (EvaluationPositionIntersectsPrimitive(WorldPosition))
+	{
+		FMaterialPixelParameters MaterialParameters = MakeInitializedMaterialPixelParameters();
+		MaterialParameters.PrimitiveId = PrimitiveId;
+		MaterialParameters.AbsoluteWorldPosition = LWCPromote(WorldPosition);
+
+		// Evaluate material graph
+		FPixelMaterialInputs PixelMaterialInputs;
+		CalcPixelMaterialInputs(MaterialParameters, PixelMaterialInputs);
+
+		// Rasterize coefficients
+		Extinction = SampleExtinctionCoefficients(PixelMaterialInputs);
+		Emission = SampleEmissive(PixelMaterialInputs);
+		Albedo = SampleAlbedo(PixelMaterialInputs) * Extinction;
+	}
+
+	// Aggregate in group-shared memory
+	uint LinearThreadIndex = MortonEncode3(GroupThreadId);
+	GSExtinctionSum[LinearThreadIndex] = Luminance(Extinction);
+	GroupMemoryBarrierWithGroupSync();
+
+	// Parallel Sum
+	for (int NeighborOffset = 1; NeighborOffset < THREADGROUP_SIZE_1D; NeighborOffset = (NeighborOffset << 1))
+	{
+		int NeighborIndex = LinearThreadIndex + NeighborOffset;
+		GSExtinctionSum[LinearThreadIndex] += (NeighborIndex < THREADGROUP_SIZE_1D) ? GSExtinctionSum[NeighborIndex] : 0.0;
+
+		GroupMemoryBarrierWithGroupSync();
+	}
+
+	// Allocate if non-zero voxel data
+	bool bWasAlreadyAllocated = IsBottomLevelAllocated(TopLevelGridData);
+	AllocatedVoxelCount = BottomLevelVoxelResolution.x * BottomLevelVoxelResolution.y * BottomLevelVoxelResolution.z;
+	float ZeroThreshold = 1.0e-6;
+	if (all(GroupThreadId == 0) && any(GSExtinctionSum[0] > ZeroThreshold) && !bWasAlreadyAllocated)
+	{
+		uint3 AllocatedVoxelResolution = BottomLevelVoxelResolution;
+		AllocatedVoxelCount = AllocatedVoxelResolution.x * AllocatedVoxelResolution.y * AllocatedVoxelResolution.z;
+
+		uint BottomLevelIndex = 0;
+		InterlockedAdd(RWBottomLevelGridAllocatorBuffer[0], AllocatedVoxelCount, BottomLevelIndex);
+
+		// Guard against over allocation
+		if (BottomLevelIndex + AllocatedVoxelCount > BottomLevelGridBufferSize)
+		{
+			// Declare the buffer to be filled
+			uint Dummy;
+			InterlockedExchange(RWBottomLevelGridAllocatorBuffer[0], BottomLevelGridBufferSize, Dummy);
+
+			BottomLevelIndex = EMPTY_VOXEL_INDEX;
+			AllocatedVoxelResolution = AllocatedVoxelCount = 0;
+		}
+
+		// Update the index with the properly allocated index
+		FTopLevelGridData AllocatedGridData = TopLevelGridData;
+		SetBottomLevelIndex(AllocatedGridData, BottomLevelIndex);
+		SetBottomLevelVoxelResolution(AllocatedGridData, AllocatedVoxelResolution);
+
+		InterlockedCompareStore(RWTopLevelGridBuffer[TopLevelGridLinearIndex], TopLevelGridData.PackedData[0], AllocatedGridData.PackedData[0]);
+	}
+
+	AllMemoryBarrierWithGroupSync();
+
+	if (LinearThreadIndex < AllocatedVoxelCount)
+	{
+		FTopLevelGridData TopLevelGridData = CreateTopLevelGridData(RWTopLevelGridBuffer[TopLevelGridLinearIndex]);
+		if (IsBottomLevelAllocated(TopLevelGridData))
+		{
+			uint BottomLevelVoxelLinearIndex = GetBottomLevelIndex(TopLevelGridData) + LinearThreadIndex;
+
+			if (bWasAlreadyAllocated)
+			{
+				Extinction += GetExtinction(RWExtinctionGridBuffer[BottomLevelVoxelLinearIndex]);
+				Emission += GetEmission(RWEmissionGridBuffer[BottomLevelVoxelLinearIndex]);
+				Albedo += GetAlbedo(RWScatteringGridBuffer[BottomLevelVoxelLinearIndex]);
+			}
+
+			SetExtinction(RWExtinctionGridBuffer[BottomLevelVoxelLinearIndex], Extinction);
+			SetEmission(RWEmissionGridBuffer[BottomLevelVoxelLinearIndex], Emission);
+			SetAlbedo(RWScatteringGridBuffer[BottomLevelVoxelLinearIndex], Albedo);
+		}
+	}
 }
\ No newline at end of file
diff --git a/Engine/Shaders/Private/HeterogeneousVolumes/HeterogeneousVolumesVoxelGridPipeline.usf b/Engine/Shaders/Private/HeterogeneousVolumes/HeterogeneousVolumesVoxelGridPipeline.usf
new file mode 100644
--- /dev/null
+++ b/Engine/Shaders/Private/HeterogeneousVolumes/HeterogeneousVolumesVoxelGridPipeline.usf
@@ -0,0 +1,656 @@
+// Copyright Epic Games, Inc. All Rights Reserved.
+
+#include "HeterogeneousVolumesVoxelGridTypes.ush"
+
+#include "../Common.ush"
+#include "../DeferredShadingCommon.ush"
+#include "../ComputeShaderUtils.ush"
+#include "HeterogeneousVolumesTracingUtils.ush"
+#include "HeterogeneousVolumesVoxelGridTraversal.ush"
+
+#ifndef THREADGROUP_SIZE_1D
+#define THREADGROUP_SIZE_1D 1
+#endif // THREADGROUP_SIZE_1D
+
+#ifndef THREADGROUP_SIZE_2D
+#define THREADGROUP_SIZE_2D 1
+#endif // THREADGROUP_SIZE_2D
+
+#ifndef THREADGROUP_SIZE_3D
+#define THREADGROUP_SIZE_3D 1
+#endif // THREADGROUP_SIZE_3D
+
+#ifndef ENABLE_ADAPTIVE_VOXEL_GRID_UNIFORM_BUFFER
+#define ENABLE_ADAPTIVE_VOXEL_GRID_UNIFORM_BUFFER 0
+#endif // ENABLE_ADAPTIVE_VOXEL_GRID_UNIFORM_BUFFER
+
+#ifndef DIM_ENABLE_INDIRECTION_GRID
+#define DIM_ENABLE_INDIRECTION_GRID 0
+#endif // DIM_ENABLE_INDIRECTION_GRID
+
+#if ENABLE_ADAPTIVE_VOXEL_GRID_UNIFORM_BUFFER
+	#define TopLevelGridResolution OrthoGridUniformBuffer.TopLevelGridResolution
+	#define TopLevelGridWorldBoundsMin OrthoGridUniformBuffer.TopLevelGridWorldBoundsMin
+	#define TopLevelGridWorldBoundsMax OrthoGridUniformBuffer.TopLevelGridWorldBoundsMax
+
+	#define TopLevelGridBitmaskBuffer OrthoGridUniformBuffer.TopLevelGridBitmaskBuffer
+	#define TopLevelGridBuffer OrthoGridUniformBuffer.TopLevelGridBuffer
+	#define IndirectionGridBuffer OrthoGridUniformBuffer.IndirectionGridBuffer
+
+	#define ExtinctionGridBuffer OrthoGridUniformBuffer.ExtinctionGridBuffer
+	#define EmissionGridBuffer OrthoGridUniformBuffer.EmissionGridBuffer
+	#define ScatteringGridBuffer OrthoGridUniformBuffer.ScatteringGridBuffer
+
+#else
+	int3 TopLevelGridResolution;
+	float3 TopLevelGridWorldBoundsMin;
+	float3 TopLevelGridWorldBoundsMax;
+
+	StructuredBuffer<FTopLevelGridBitmaskData> TopLevelGridBitmaskBuffer;
+	StructuredBuffer<FTopLevelGridData> TopLevelGridBuffer;
+	StructuredBuffer<FTopLevelGridData> IndirectionGridBuffer;
+	StructuredBuffer<FScalarGridData> ExtinctionGridBuffer;
+	StructuredBuffer<FVectorGridData> EmissionGridBuffer;
+	StructuredBuffer<FVectorGridData> ScatteringGridBuffer;
+
+#endif // !ENABLE_ADAPTIVE_VOXEL_GRID_UNIFORM_BUFFER
+
+float3 PrimitiveWorldBoundsMin;
+float3 PrimitiveWorldBoundsMax;
+
+float ShadingRate;
+float MinVoxelSizeInFrustum;
+float MinVoxelSizeOutOfFrustum;
+
+RWStructuredBuffer<FTopLevelGridData> RWTopLevelGridBuffer;
+
+void CalcVoxelBounds(float3 VoxelIndex, float3 GridResolution, float3 WorldBoundsMin, float3 WorldBoundsMax,
+	inout float3 VoxelBoundsMin, inout float3 VoxelBoundsMax)
+{
+	float3 WorldBoundsExtent = WorldBoundsMax - WorldBoundsMin;
+	VoxelBoundsMin = WorldBoundsMin + WorldBoundsExtent * (VoxelIndex / GridResolution);
+	VoxelBoundsMax = WorldBoundsMin + WorldBoundsExtent * ((VoxelIndex + 1) / GridResolution);
+}
+
+void CalcTopLevelVoxelBounds(float3 VoxelIndex, inout float3 VoxelBoundsMin, inout float3 VoxelBoundsMax)
+{
+	CalcVoxelBounds(VoxelIndex, TopLevelGridResolution, TopLevelGridWorldBoundsMin, TopLevelGridWorldBoundsMax,
+		VoxelBoundsMin, VoxelBoundsMax);
+}
+
+float3 GetVoxelCenter(float3 VoxelIndex, inout float3 VoxelBoundsMin, inout float3 VoxelBoundsMax)
+{
+	float3 TopLevelGridWorldBoundsExtent = TopLevelGridWorldBoundsMax - TopLevelGridWorldBoundsMin;
+	return TopLevelGridWorldBoundsMin + TopLevelGridWorldBoundsExtent * ((VoxelIndex + 0.5) / TopLevelGridResolution);
+}
+
+bool PrimitiveIntersectsVoxel(float3 VoxelBoundsMin, float3 VoxelBoundsMax)
+{
+	if (any(PrimitiveWorldBoundsMin > VoxelBoundsMax) || any(VoxelBoundsMin > PrimitiveWorldBoundsMax))
+	{
+		return false;
+	}
+
+	return true;
+}
+
+// Copied from LiveShadingPipeline.usf
+float4 CreatePlane(
+	float3 P_000,
+	float3 P_100,
+	float3 P_010
+)
+{
+	float3 X = P_100 - P_000;
+	float3 Y = P_010 - P_000;
+
+	float3 N = normalize(cross(X, Y));
+	float d = -dot(N, P_000);
+	return float4(N, d);
+}
+
+float CalcSignedDistance(float3 P, float4 Plane)
+{
+	return dot(P, Plane.xyz) + Plane.w;
+}
+
+bool BoxIntersectsFrustum(float3 BoxOrigin, float3 BoxExtent, float4 Planes[5])
+{
+	for (int i = 0; i < 5; ++i)
+	{
+		float Distance = CalcSignedDistance(BoxOrigin, Planes[i]);
+		float ProjectedExtent = dot(BoxExtent, abs(Planes[i].xyz));
+
+		if (Distance > ProjectedExtent)
+		{
+			return false;
+		}
+	}
+
+	return true;
+}
+
+bool BoxIntersectsViewFrustum(float3 BoxOrigin, float3 BoxExtent)
+{
+	float Epsilon = 1.0e-6;
+#if HAS_INVERTED_Z_BUFFER
+	float NearDepth = 1.0;
+	float FarDepth = Epsilon;
+#else
+	float NearDepth = 0.0;
+	float FarDepth = 1.0 - Epsilon;
+#endif // HAS_INVERTED_Z_BUFFER
+
+	// Project frustum vertices into world-space
+	float4 Viewport = float4(0.0, 0.0, View.ViewSizeAndInvSize.xy);
+	float3 W_000 = LWCHackToFloat(SvPositionToWorld(float4(Viewport.xw, NearDepth, 1)));
+	float3 W_100 = LWCHackToFloat(SvPositionToWorld(float4(Viewport.zw, NearDepth, 1)));
+	float3 W_010 = LWCHackToFloat(SvPositionToWorld(float4(Viewport.xy, NearDepth, 1)));
+	float3 W_110 = LWCHackToFloat(SvPositionToWorld(float4(Viewport.zy, NearDepth, 1)));
+	float3 W_001 = LWCHackToFloat(SvPositionToWorld(float4(Viewport.xw, FarDepth, 1)));
+	float3 W_101 = LWCHackToFloat(SvPositionToWorld(float4(Viewport.zw, FarDepth, 1)));
+	float3 W_011 = LWCHackToFloat(SvPositionToWorld(float4(Viewport.xy, FarDepth, 1)));
+	float3 W_111 = LWCHackToFloat(SvPositionToWorld(float4(Viewport.zy, FarDepth, 1)));
+
+	float4 Planes[5] = {
+		// Front
+		CreatePlane(W_000, W_010, W_100),
+		// Back
+		//CreatePlane(W_101, W_001, W_111),
+		// +X
+		CreatePlane(W_100, W_110, W_101),
+		// +Y
+		CreatePlane(W_010, W_011, W_110),
+		// -X
+		CreatePlane(W_000, W_001, W_010),
+		// -Y
+		CreatePlane(W_000, W_100, W_001)
+	};
+
+	return BoxIntersectsFrustum(BoxOrigin, BoxExtent, Planes);
+}
+
+[numthreads(THREADGROUP_SIZE_3D, THREADGROUP_SIZE_3D, THREADGROUP_SIZE_3D)]
+void TopLevelGridCalculateVoxelSize(
+	uint3 GroupId : SV_GroupID,
+	uint3 GroupThreadId : SV_GroupThreadID,
+	uint3 DispatchThreadId : SV_DispatchThreadID
+)
+{
+	uint3 VoxelIndex = DispatchThreadId;
+	if (any(VoxelIndex >= TopLevelGridResolution))
+	{
+		return;
+	}
+
+	float3 VoxelBoundsMin;
+	float3 VoxelBoundsMax;
+	CalcTopLevelVoxelBounds(VoxelIndex, VoxelBoundsMin, VoxelBoundsMax);
+
+	float3 WorldVoxelCenter = (VoxelBoundsMin + VoxelBoundsMax) * 0.5;
+	float3 WorldVoxelExtent = (VoxelBoundsMax - VoxelBoundsMin) * 0.5;
+
+	if (PrimitiveIntersectsVoxel(VoxelBoundsMin, VoxelBoundsMax))
+	{
+		// Initialize the voxel-rate based on view distance
+		float2 TanHalfFOV = GetTanHalfFieldOfView();
+		int HalfWidth = View.ViewSizeAndInvSize.x * 0.5;
+		float PixelWidth = TanHalfFOV.x / HalfWidth;
+
+		float3 WorldCameraOrigin = LWCHackToFloat(LWCSubtract(View.TranslatedWorldCameraOrigin, PrimaryView.PreViewTranslation));
+		float Distance = length(WorldVoxelCenter - WorldCameraOrigin) - length(WorldVoxelExtent);
+		float VoxelWidth = Distance * PixelWidth * ShadingRate;
+
+		// Frustum intersection determines which min voxel-size to clip by
+		float3 PrimitiveWorldBoundsOrigin = (PrimitiveWorldBoundsMax + PrimitiveWorldBoundsMin) * 0.5;
+		float3 PrimitiveWorldBoundsExtent = (PrimitiveWorldBoundsMax - PrimitiveWorldBoundsMin) * 0.5;
+		float MinVoxelSize = BoxIntersectsViewFrustum(WorldVoxelCenter, WorldVoxelExtent) ? MinVoxelSizeInFrustum : MinVoxelSizeOutOfFrustum;
+		float VoxelSize = max(VoxelWidth, MinVoxelSize);
+
+		// Store minimum voxel rate temporarily as the bottom-level index
+		uint LinearIndex = GetLinearIndex(VoxelIndex, TopLevelGridResolution);
+		if (!IsBottomLevelEmpty(RWTopLevelGridBuffer[LinearIndex]))
+		{
+			float PrevVoxelSize = GetVoxelSize(RWTopLevelGridBuffer[LinearIndex]);
+			VoxelSize = min(VoxelSize, PrevVoxelSize);
+		}
+
+		FTopLevelGridData TopLevelGridData;
+		{
+			SetVoxelSize(TopLevelGridData, VoxelSize);
+		}
+		RWTopLevelGridBuffer[LinearIndex] = TopLevelGridData;
+	}
+}
+
+int3 NextPowerOfTwo(int3 Value)
+{
+	Value -= 1;
+	Value |= Value >> 1;
+	Value |= Value >> 2;
+	Value |= Value >> 4;
+	Value |= Value >> 8;
+	Value |= Value >> 16;
+	Value += 1;
+
+	return Value;
+}
+
+int MaxVoxelResolution;
+int bSampleAtVertices;
+
+// Indirection Grid
+int IndirectionGridBufferSize;
+RWBuffer<uint> RWIndirectionGridAllocatorBuffer;
+RWStructuredBuffer<FTopLevelGridData> RWIndirectionGridBuffer;
+
+[numthreads(THREADGROUP_SIZE_3D, THREADGROUP_SIZE_3D, THREADGROUP_SIZE_3D)]
+void AllocateBottomLevelGrid(
+	uint3 GroupId : SV_GroupID,
+	uint3 GroupThreadId : SV_GroupThreadID,
+	uint3 DispatchThreadId : SV_DispatchThreadID
+)
+{
+	uint3 VoxelIndex = DispatchThreadId;
+	if (any(VoxelIndex >= TopLevelGridResolution))
+	{
+		return;
+	}
+
+	uint LinearIndex = GetLinearIndex(VoxelIndex, TopLevelGridResolution);
+	if (IsBottomLevelAllocated(RWTopLevelGridBuffer[LinearIndex]))
+	{
+		float VoxelSize = GetVoxelSize(RWTopLevelGridBuffer[LinearIndex]);
+
+		float3 VoxelBoundsMin;
+		float3 VoxelBoundsMax;
+		CalcTopLevelVoxelBounds(VoxelIndex, VoxelBoundsMin, VoxelBoundsMax);
+
+		// Calculate bottom-level voxel resolution and round to nearest power-of-two
+		float3 VoxelBoundsExtent = VoxelBoundsMax - VoxelBoundsMin;
+		int3 VoxelResolution = NextPowerOfTwo(ceil(VoxelBoundsExtent / VoxelSize));
+		if (bSampleAtVertices)
+		{
+			VoxelResolution = max(VoxelResolution, 2);
+		}
+
+		// DEBUG: OVERRIDE RESOLUTION..
+		// VoxelResolution = MaxVoxelResolution;
+		//VoxelResolution = max(VoxelResolution.x, max(VoxelResolution.y, VoxelResolution.z));
+
+#if DIM_ENABLE_INDIRECTION_GRID
+		// Allocate indirection voxels
+		int3 BottomLevelVoxelResolution = MaxVoxelResolution;
+		int3 IndirectionGridResolution = ceil(VoxelResolution / BottomLevelVoxelResolution);
+		// Encoding restricts dimension to be a power-of-2
+		//IndirectionGridResolution = NextPowerOfTwo(IndirectionGridResolution);
+		IndirectionGridResolution.x = clamp(IndirectionGridResolution.x, 1, 4);
+		IndirectionGridResolution.y = clamp(IndirectionGridResolution.y, 1, 4);
+		IndirectionGridResolution.z = clamp(IndirectionGridResolution.z, 1, 4);
+		// Force regular sized voxels so the delineation is tractable for the allocator
+		IndirectionGridResolution = max(IndirectionGridResolution.x, max(IndirectionGridResolution.y, IndirectionGridResolution.z));
+
+		int IndirectionVoxelCount = IndirectionGridResolution.x * IndirectionGridResolution.y * IndirectionGridResolution.z;
+
+		uint IndirectionGridIndex = EMPTY_VOXEL_INDEX;
+		InterlockedAdd(RWIndirectionGridAllocatorBuffer[0], IndirectionVoxelCount, IndirectionGridIndex);
+
+		if (IndirectionGridIndex + IndirectionVoxelCount > IndirectionGridBufferSize)
+		{
+			// Declare the buffer to be filled
+			uint Dummy;
+			InterlockedExchange(RWIndirectionGridAllocatorBuffer[0], IndirectionGridBufferSize, Dummy);
+
+			IndirectionGridIndex = EMPTY_VOXEL_INDEX;
+			IndirectionGridResolution = 0;
+		}
+
+		FTopLevelGridData TopLevelGridData;
+		{
+			SetBottomLevelIndex(TopLevelGridData, IndirectionGridIndex);
+			SetBottomLevelVoxelResolution(TopLevelGridData, IndirectionGridResolution);
+		}
+		RWTopLevelGridBuffer[LinearIndex] = TopLevelGridData;
+
+#else // DIM_ENABLE_INDIRECTION_GRID
+		uint BottomLevelIndex = EMPTY_VOXEL_INDEX;
+
+		VoxelResolution.x = clamp(VoxelResolution.x, 1, MaxVoxelResolution);
+		VoxelResolution.y = clamp(VoxelResolution.y, 1, MaxVoxelResolution);
+		VoxelResolution.z = clamp(VoxelResolution.z, 1, MaxVoxelResolution);
+		// Force regular sized voxels so the delineation is tractable for the allocator
+		VoxelResolution = max(VoxelResolution.x, max(VoxelResolution.y, VoxelResolution.z));
+
+		FTopLevelGridData TopLevelGridData = RWTopLevelGridBuffer[LinearIndex];
+		{
+			SetBottomLevelIndex(TopLevelGridData, BottomLevelIndex);
+			SetBottomLevelVoxelResolution(TopLevelGridData, VoxelResolution);
+		}
+		RWTopLevelGridBuffer[LinearIndex] = TopLevelGridData;
+#endif // DIM_ENABLE_INDIRECTION_GRID
+	}
+}
+
+struct FRasterTileData
+{
+	uint TopLevelGridLinearIndex;
+	uint BottomLevelGridLinearOffset;
+};
+
+int RasterTileVoxelResolution;
+int MaxNumRasterTiles;
+
+RWBuffer<uint> RWRasterTileAllocatorBuffer;
+RWStructuredBuffer<FRasterTileData> RWRasterTileBuffer;
+
+[numthreads(THREADGROUP_SIZE_3D, THREADGROUP_SIZE_3D, THREADGROUP_SIZE_3D)]
+void GenerateRasterTiles(
+	uint3 GroupId : SV_GroupID,
+	uint3 GroupThreadId : SV_GroupThreadID,
+	uint3 DispatchThreadId : SV_DispatchThreadID
+)
+{
+	uint3 VoxelIndex = DispatchThreadId;
+	if (any(VoxelIndex >= TopLevelGridResolution))
+	{
+		return;
+	}
+
+	int LinearIndex = GetLinearIndex(VoxelIndex, TopLevelGridResolution);
+	if (!IsBottomLevelEmpty(TopLevelGridBuffer[LinearIndex]))
+	{
+		int3 BottomLevelResolution = GetBottomLevelVoxelResolution(TopLevelGridBuffer[LinearIndex]);
+		int BottomLevelLinearResolution = BottomLevelResolution.x * BottomLevelResolution.y * BottomLevelResolution.z;
+
+#if DIM_ENABLE_INDIRECTION_GRID
+		int RasterTileCount = BottomLevelLinearResolution;
+
+#else // DIM_ENABLE_INDIRECTION_GRID
+		int3 RasterTileCount3 = ceil(float3(BottomLevelResolution) / RasterTileVoxelResolution);
+		int RasterTileCount = RasterTileCount3.x * RasterTileCount3.y * RasterTileCount3.z;
+#endif // DIM_ENABLE_INDIRECTION_GRID
+
+		if (RasterTileCount > 0)
+		{
+			uint RasterTileOffset;
+			InterlockedAdd(RWRasterTileAllocatorBuffer[0], RasterTileCount, RasterTileOffset);
+
+			// Prevent over allocation
+			if (RasterTileOffset + RasterTileCount <= MaxNumRasterTiles)
+			{
+				for (uint Index = 0; Index < RasterTileCount; ++Index)
+				{
+					uint RasterTileIndex = RasterTileOffset + Index;
+					RWRasterTileBuffer[RasterTileIndex].TopLevelGridLinearIndex = LinearIndex;
+					RWRasterTileBuffer[RasterTileIndex].BottomLevelGridLinearOffset = Index;
+				}
+			}
+			else
+			{
+				// Declare the buffer to be filled
+				uint Dummy;
+				InterlockedExchange(RWRasterTileAllocatorBuffer[0], MaxNumRasterTiles, Dummy);
+			}
+		}
+	}
+}
+
+int3 MaxDispatchThreadGroupsPerDimension;
+Buffer<uint> RasterTileAllocatorBuffer;
+RWBuffer<uint> RWRasterizeBottomLevelGridIndirectArgsBuffer;
+
+uint3 GetGroupCountWrapped(uint TargetGroupCount)
+{
+	const uint WrappedGroupStride = 128;
+
+	uint3 GroupCount = uint3(TargetGroupCount, 1, 1);
+
+	if (GroupCount.x > MaxDispatchThreadGroupsPerDimension.x)
+	{
+		GroupCount.y = (GroupCount.x + WrappedGroupStride - 1) / WrappedGroupStride;
+		GroupCount.x = WrappedGroupStride;
+	}
+	if (GroupCount.y > MaxDispatchThreadGroupsPerDimension.y)
+	{
+		GroupCount.z = (GroupCount.y + WrappedGroupStride - 1) / WrappedGroupStride;
+		GroupCount.y = WrappedGroupStride;
+	}
+
+	return GroupCount;
+}
+
+[numthreads(1, 1, 1)]
+void SetRasterizeBottomLevelGridIndirectArgs()
+{
+	uint3 GroupCount = GetGroupCountWrapped(RasterTileAllocatorBuffer[0]);
+	RWRasterizeBottomLevelGridIndirectArgsBuffer[0] = GroupCount.x;
+	RWRasterizeBottomLevelGridIndirectArgsBuffer[1] = GroupCount.y;
+	RWRasterizeBottomLevelGridIndirectArgsBuffer[2] = GroupCount.z;
+}
+
+RWStructuredBuffer<FTopLevelGridBitmaskData> RWTopLevelGridBitmaskBuffer;
+
+[numthreads(THREADGROUP_SIZE_3D, THREADGROUP_SIZE_3D, THREADGROUP_SIZE_3D)]
+void TopLevelCreateBitmaskCS(
+	uint3 GroupId : SV_GroupID,
+	uint3 GroupThreadId : SV_GroupThreadID,
+	uint3 DispatchThreadId : SV_DispatchThreadID
+)
+{
+	int3 VoxelIndex = DispatchThreadId;
+	if (any(VoxelIndex >= TopLevelGridResolution))
+	{
+		return;
+	}
+
+	int LinearIndex = GetLinearIndex(VoxelIndex, TopLevelGridResolution);
+	if (!IsBottomLevelEmpty(TopLevelGridBuffer[LinearIndex]))
+	{
+		FTopLevelGridBitmaskData Bitmask = (FTopLevelGridBitmaskData) 0;
+
+		uint BitmaskIndex = MortonEncode3(GroupThreadId);
+		if (BitmaskIndex < 32)
+		{
+			Bitmask.PackedData[0] = 1u << BitmaskIndex;
+		}
+		else
+		{
+			Bitmask.PackedData[1] = 1u << (BitmaskIndex - 32u);
+		}
+
+		uint GroupIdLinearIndex = MortonEncode3(GroupId);
+		InterlockedOr(RWTopLevelGridBitmaskBuffer[GroupIdLinearIndex].PackedData[0], Bitmask.PackedData[0]);
+		InterlockedOr(RWTopLevelGridBitmaskBuffer[GroupIdLinearIndex].PackedData[1], Bitmask.PackedData[1]);
+	}
+}
+
+RWStructuredBuffer<FScalarGridData> RWMajorantVoxelGridBuffer;
+
+[numthreads(THREADGROUP_SIZE_3D, THREADGROUP_SIZE_3D, THREADGROUP_SIZE_3D)]
+void BuildMajorantVoxelGridCS(
+	uint3 GroupId : SV_GroupID,
+	uint3 GroupThreadId : SV_GroupThreadID,
+	uint3 DispatchThreadId : SV_DispatchThreadID
+)
+{
+	int3 VoxelIndex = DispatchThreadId;
+	if (any(VoxelIndex >= TopLevelGridResolution))
+	{
+		return;
+	}
+
+	FMajorantData MajorantData = CreateMajorantData();
+
+	int LinearIndex = GetLinearIndex(VoxelIndex, TopLevelGridResolution);
+	FTopLevelGridData TopLevelGridData = TopLevelGridBuffer[LinearIndex];
+	if (!IsBottomLevelEmpty(TopLevelGridData))
+	{
+		int3 VoxelResolution = GetBottomLevelVoxelResolution(TopLevelGridData);
+		int BottomLevelVoxelCount = VoxelResolution.x * VoxelResolution.y * VoxelResolution.z;
+#if DIM_ENABLE_INDIRECTION_GRID
+		uint IndirectionGridIndex = GetBottomLevelIndex(TopLevelGridData);
+		for (int Index = 0; Index < BottomLevelVoxelCount; ++Index)
+		{
+			FTopLevelGridData IndirectionData = IndirectionGridBuffer[IndirectionGridIndex + Index];
+			uint IndirectionIndex = GetBottomLevelIndex(IndirectionData);
+			int3 IndirectionResolution = GetBottomLevelVoxelResolution(IndirectionData);
+			int IndirectionVoxelCount = IndirectionResolution.x * IndirectionResolution.y * IndirectionResolution.z;
+			for (uint VoxelIndex = 0; VoxelIndex < IndirectionVoxelCount; ++VoxelIndex)
+			{
+				uint BottomLevelIndex = IndirectionIndex + VoxelIndex;
+				float3 Extinction = GetExtinction(ExtinctionGridBuffer[BottomLevelIndex]);
+				UpdateMajorantData(MajorantData, Luminance(Extinction));
+			}
+		}
+#else // DIM_ENABLE_INDIRECTION_GRID
+		for (int VoxelIndex = 0; VoxelIndex < BottomLevelVoxelCount; ++VoxelIndex)
+		{
+			int BottomLevelIndex = GetBottomLevelIndex(TopLevelGridData) + VoxelIndex;
+			float3 Extinction = GetExtinction(ExtinctionGridBuffer[BottomLevelIndex]);
+			UpdateMajorantData(MajorantData, Luminance(Extinction));
+		}
+#endif // DIM_ENABLE_INDIRECTION_GRID
+	}
+
+	SetMajorantData(RWMajorantVoxelGridBuffer[LinearIndex], MajorantData);
+}
+
+// Ray data
+int bJitter;
+float MaxTraceDistance;
+float StepSize;
+int MaxStepCount;
+
+// Dispatch data
+int3 GroupCount;
+
+// Output
+RWTexture2D<float4> RWLightingTexture;
+
+[numthreads(THREADGROUP_SIZE_2D, THREADGROUP_SIZE_2D, 1)]
+void RenderTransmittanceWithVoxelGridCS(
+	uint2 GroupThreadId : SV_GroupThreadID,
+	uint2 DispatchThreadId : SV_DispatchThreadID
+)
+{
+	float3 Radiance = 0.0;
+	float3 Transmittance = 1.0;
+
+	// Create screen ray
+	if (any(DispatchThreadId.xy >= View.ViewSizeAndInvSize.xy))
+	{
+		return;
+	}
+	uint LinearIndex = DispatchThreadId.y * GroupCount.x * THREADGROUP_SIZE_2D + DispatchThreadId.x;
+	uint2 PixelCoord = DispatchThreadId.xy + View.ViewRectMin.xy;
+
+	// Extract depth
+	float DeviceZ = SceneDepthTexture.Load(int3(PixelCoord, 0)).r;
+#if HAS_INVERTED_Z_BUFFER
+	DeviceZ = max(0.000000000001, DeviceZ);
+#endif // HAS_INVERTED_Z_BUFFER
+
+	// Clip trace distance
+	float SceneDepth = min(ConvertFromDeviceZ(DeviceZ), MaxTraceDistance);
+	DeviceZ = ConvertToDeviceZ(SceneDepth);
+
+	float3 RayOrigin = View.TranslatedWorldCameraOrigin;
+	float3 RayEnd = SvPositionToTranslatedWorld(float4(PixelCoord + 0.5, DeviceZ, 1));
+	float3 RayDirection = RayEnd - RayOrigin;
+	float RayTMin = 0.0;
+	float RayTMax = length(RayDirection);
+	RayDirection /= RayTMax;
+
+	float FrustumRayTMin = RayTMin;
+	float FrustumRayTMax = RayTMax;
+	if (FrustumGridUniformBuffer.bUseFrustumGrid)
+	{
+		float ViewNearZ = min(ConvertFromDeviceZ(DeviceZ), FrustumGridUniformBuffer.NearPlaneDepth);
+		float DeviceNearZ = ConvertToDeviceZ(ViewNearZ);
+		float ViewFarZ = min(ConvertFromDeviceZ(DeviceZ), FrustumGridUniformBuffer.FarPlaneDepth);
+		float DeviceFarZ = ConvertToDeviceZ(ViewFarZ);
+
+		// TODO: Shared code is imprecise, in comparison with commented out frustum-trimmed ray. Investigate...
+#if 0
+		float3 RayBegin = SvPositionToTranslatedWorld(float4(PixelCoord + 0.5, DeviceNearZ, 1));
+		float3 RayEnd = SvPositionToTranslatedWorld(float4(PixelCoord + 0.5, DeviceFarZ, 1));
+		float3 RayDirection = RayEnd - RayBegin;
+		FrustumRayTMin = 0.0;
+		FrustumRayTMax = length(RayDirection);
+		RayDirection /= FrustumRayTMax;
+
+		Transmittance = TraceFrustumVoxelGrid(RayBegin, RayDirection, FrustumRayTMin, FrustumRayTMax);
+		float RayTOffset = length(RayBegin - RayOrigin);
+		FrustumRayTMin += RayTOffset;
+		FrustumRayTMax += RayTOffset;
+#else
+		float3 WorldRayOrigin = LWCHackToFloat(SvPositionToWorld(float4(PixelCoord + 0.5, DeviceNearZ, 1)));
+		float3 WorldRayEnd = LWCHackToFloat(SvPositionToWorld(float4(PixelCoord + 0.5, DeviceFarZ, 1)));
+		float3 WorldRayDirection = WorldRayEnd - WorldRayOrigin;
+		float WorldRayLength = length(WorldRayDirection);
+		WorldRayDirection /= WorldRayLength;
+
+
+		// Convert to view-space
+		float3 ViewRayOrigin = mul(float4(WorldRayOrigin, 1), FrustumGridUniformBuffer.WorldToView).xyz;
+		float3 ViewRayEnd = mul(float4(WorldRayEnd, 1), FrustumGridUniformBuffer.WorldToView).xyz;
+		float3 ViewRayDirection = ViewRayEnd - ViewRayOrigin;
+		float ViewRayLength = length(ViewRayDirection);
+		ViewRayDirection /= ViewRayLength;
+
+		// Convert to voxel-space
+		int3 VoxelDimensions = FrustumGridUniformBuffer.VoxelDimensions;
+		float NearPlaneDepth = FrustumGridUniformBuffer.NearPlaneDepth;
+		float FarPlaneDepth = FrustumGridUniformBuffer.FarPlaneDepth;
+		float3 VoxelRayOrigin = ViewToVoxel(ViewRayOrigin, VoxelDimensions, NearPlaneDepth, FarPlaneDepth);
+		float3 VoxelRayEnd = ViewToVoxel(ViewRayOrigin + ViewRayDirection * ViewRayLength, VoxelDimensions, NearPlaneDepth, FarPlaneDepth);
+
+		float3 VoxelRayDirection = VoxelRayEnd - VoxelRayOrigin;
+		float VoxelRayLength = length(VoxelRayDirection);
+		VoxelRayDirection /= VoxelRayLength;
+
+		float VoxelRayTMin = 0.0;
+		float VoxelRayTMax = VoxelRayLength;
+
+		// Traversal
+		{
+			MarchTopLevelFrustumVoxelGrid(
+				VoxelRayOrigin,
+				VoxelRayDirection,
+				VoxelRayTMin,
+				VoxelRayTMax,
+				WorldRayLength,
+				FrustumGridUniformBuffer.TopLevelFroxelGridResolution,
+				FrustumGridUniformBuffer.TopLevelFroxelGridBuffer,
+				FrustumGridUniformBuffer.ExtinctionFroxelGridBuffer,
+				Transmittance
+			);
+		}
+
+		float3 RayBegin = SvPositionToTranslatedWorld(float4(PixelCoord + 0.5, DeviceNearZ, 1));
+		float RayTOffset = length(RayBegin - RayOrigin);
+		FrustumRayTMin += RayTOffset;
+		FrustumRayTMax += RayTOffset;
+#endif
+	}
+
+	if (OrthoGridUniformBuffer.bUseOrthoGrid)
+	{
+		if (FrustumGridUniformBuffer.bUseFrustumGrid && (FrustumRayTMin < FrustumRayTMax))
+		{
+			Transmittance *= TraceOrthoVoxelGrid(RayOrigin, RayDirection, RayTMin, FrustumRayTMin);
+			Transmittance *= TraceOrthoVoxelGrid(RayOrigin, RayDirection, FrustumRayTMax, RayTMax);
+		}
+		else
+		{
+			Transmittance = TraceOrthoVoxelGrid(RayOrigin, RayDirection, RayTMin, RayTMax);
+		}
+	}
+
+	// Debug
+	Radiance = Transmittance;
+
+	// Output..
+	float3 Opacity = 1.0 - Transmittance;
+	RWLightingTexture[PixelCoord] += float4(Radiance, Luminance(Opacity));
+}
\ No newline at end of file
diff --git a/Engine/Shaders/Private/HeterogeneousVolumes/HeterogeneousVolumesVoxelGridTraversal.ush b/Engine/Shaders/Private/HeterogeneousVolumes/HeterogeneousVolumesVoxelGridTraversal.ush
new file mode 100644
--- /dev/null
+++ b/Engine/Shaders/Private/HeterogeneousVolumes/HeterogeneousVolumesVoxelGridTraversal.ush
@@ -0,0 +1,529 @@
+// Copyright Epic Games, Inc. All Rights Reserved.
+
+#pragma once
+
+#include "HeterogeneousVolumesVoxelGridUtils.ush"
+#include "HeterogeneousVolumesFrustumVoxelGridUtils.ush"
+
+void MarchBottomLevelVoxelGrid(
+	float3 VoxelRayOrigin,
+	float3 VoxelRayDirection,
+	float VoxelRayTMin,
+	float VoxelRayTMax,
+	float WorldRayHitSpan,
+	float3 TopLevelVoxelStartPos,
+	float3 TopLevelVoxelEndPos,
+	FTopLevelGridData TopLevelGridData,
+	StructuredBuffer<FScalarGridData> ExtinctionGridBuffer,
+	inout float3 Transmittance
+)
+{
+	uint BottomLevelVoxelIndex = GetBottomLevelIndex(TopLevelGridData);
+	int3 BottomLevelVoxelResolution = GetBottomLevelVoxelResolution(TopLevelGridData);
+	float TopLevelVoxelRayLength = length(TopLevelVoxelEndPos - TopLevelVoxelStartPos);
+
+	// Transform to local-space
+	float3 MinTopLevelVoxelPos = floor((TopLevelVoxelStartPos + TopLevelVoxelEndPos) * 0.5);
+	float3 LocalRayOrigin = (TopLevelVoxelStartPos - MinTopLevelVoxelPos) * BottomLevelVoxelResolution;
+	float3 LocalRayEnd = (TopLevelVoxelEndPos - MinTopLevelVoxelPos) * BottomLevelVoxelResolution;
+
+	float3 LocalRayDirection = LocalRayEnd - LocalRayOrigin;
+	float LocalRayLength = length(LocalRayDirection);
+	LocalRayDirection /= LocalRayLength;
+	float3 LocalRayDirectionInv = 1.0 / LocalRayDirection;
+
+	float LocalRayMarchT = 0.0;
+	float3 LocalRayMarchPos = LocalRayOrigin;
+
+	float3 LocalBoundsPos;
+	LocalBoundsPos.x = sign(LocalRayDirection.x) > 0 ? floor(LocalRayMarchPos.x) + 1 : ceil(LocalRayMarchPos.x) - 1;
+	LocalBoundsPos.y = sign(LocalRayDirection.y) > 0 ? floor(LocalRayMarchPos.y) + 1 : ceil(LocalRayMarchPos.y) - 1;
+	LocalBoundsPos.z = sign(LocalRayDirection.z) > 0 ? floor(LocalRayMarchPos.z) + 1 : ceil(LocalRayMarchPos.z) - 1;
+
+	float3 LocalBoundsHitT = (LocalBoundsPos - LocalRayMarchPos) * LocalRayDirectionInv;
+
+	// Remove floating-point rounding error
+	float Epsilon = 1.0e-4;
+	if (LocalBoundsHitT.x <= Epsilon) LocalBoundsHitT.x += abs(LocalRayDirectionInv.x);
+	if (LocalBoundsHitT.y <= Epsilon) LocalBoundsHitT.y += abs(LocalRayDirectionInv.y);
+	if (LocalBoundsHitT.z <= Epsilon) LocalBoundsHitT.z += abs(LocalRayDirectionInv.z);
+
+	float3 OpacityThreshold = 1.0e-4;
+	while ((LocalRayMarchT < LocalRayLength) && all(Transmittance > OpacityThreshold))
+	{
+		float LocalRayMarchDeltaT = 0.0;
+		if (LocalBoundsHitT.x < LocalBoundsHitT.y)
+		{
+			if (LocalBoundsHitT.x < LocalBoundsHitT.z)
+			{
+				LocalRayMarchDeltaT = LocalBoundsHitT.x - LocalRayMarchT;
+				LocalBoundsHitT.x += abs(LocalRayDirectionInv.x);
+			}
+			else
+			{
+				LocalRayMarchDeltaT = LocalBoundsHitT.z - LocalRayMarchT;
+				LocalBoundsHitT.z += abs(LocalRayDirectionInv.z);
+			}
+		}
+		else
+		{
+			if (LocalBoundsHitT.y < LocalBoundsHitT.z)
+			{
+				LocalRayMarchDeltaT = LocalBoundsHitT.y - LocalRayMarchT;
+				LocalBoundsHitT.y += abs(LocalRayDirectionInv.y);
+			}
+			else
+			{
+				LocalRayMarchDeltaT = LocalBoundsHitT.z - LocalRayMarchT;
+				LocalBoundsHitT.z += abs(LocalRayDirectionInv.z);
+			}
+		}
+
+		// Clip voxel delta-t by overall voxel ray-length.
+		if (LocalRayMarchT + LocalRayMarchDeltaT > LocalRayLength)
+		{
+			LocalRayMarchDeltaT = LocalRayLength - LocalRayMarchT;
+		}
+		float3 LocalEndPos = LocalRayMarchPos + LocalRayDirection * LocalRayMarchDeltaT;
+		float3 LocalSamplePos = clamp((LocalRayMarchPos + LocalEndPos) * 0.5, 0.0, BottomLevelVoxelResolution - 1.0);
+		uint LocalSampleLinearPos = BottomLevelVoxelIndex + MortonEncode3(LocalSamplePos);
+
+		float3 Extinction = GetExtinction(ExtinctionGridBuffer[LocalSampleLinearPos]);
+		float WorldRayMarchDeltaT = max((LocalRayMarchDeltaT / LocalRayLength) * TopLevelVoxelRayLength * WorldRayHitSpan / VoxelRayTMax, 0.0);
+		Transmittance *= saturate(exp(-Extinction * WorldRayMarchDeltaT));
+
+		LocalRayMarchT += LocalRayMarchDeltaT;
+		LocalRayMarchPos = LocalEndPos;
+	}
+}
+
+void MarchIndirectionOrthoVoxelGrid(
+	float3 VoxelRayOrigin,
+	float3 VoxelRayDirection,
+	float VoxelRayTMin,
+	float VoxelRayTMax,
+	float WorldRayHitSpan,
+	float3 TopLevelVoxelStartPos,
+	float3 TopLevelVoxelEndPos,
+	FTopLevelGridData TopLevelGridData,
+	StructuredBuffer<FTopLevelGridData> IndirectionGridBuffer,
+	StructuredBuffer<FScalarGridData> ExtinctionGridBuffer,
+	inout float3 Transmittance
+)
+{
+	uint BottomLevelVoxelIndex = GetBottomLevelIndex(TopLevelGridData);
+	int3 BottomLevelVoxelResolution = GetBottomLevelVoxelResolution(TopLevelGridData);
+	float TopLevelVoxelRayLength = length(TopLevelVoxelEndPos - TopLevelVoxelStartPos);
+
+	// Transform to local-space
+	float3 MinTopLevelVoxelPos = floor((TopLevelVoxelStartPos + TopLevelVoxelEndPos) * 0.5);
+	float3 LocalRayOrigin = (TopLevelVoxelStartPos - MinTopLevelVoxelPos) * BottomLevelVoxelResolution;
+	float3 LocalRayEnd = (TopLevelVoxelEndPos - MinTopLevelVoxelPos) * BottomLevelVoxelResolution;
+
+	float3 LocalRayDirection = LocalRayEnd - LocalRayOrigin;
+	float LocalRayLength = length(LocalRayDirection);
+	LocalRayDirection /= LocalRayLength;
+	float3 LocalRayDirectionInv = 1.0 / LocalRayDirection;
+
+	float LocalRayMarchT = 0.0;
+	float3 LocalRayMarchPos = LocalRayOrigin;
+
+	float3 LocalBoundsPos;
+	LocalBoundsPos.x = sign(LocalRayDirection.x) > 0 ? floor(LocalRayMarchPos.x) + 1 : ceil(LocalRayMarchPos.x) - 1;
+	LocalBoundsPos.y = sign(LocalRayDirection.y) > 0 ? floor(LocalRayMarchPos.y) + 1 : ceil(LocalRayMarchPos.y) - 1;
+	LocalBoundsPos.z = sign(LocalRayDirection.z) > 0 ? floor(LocalRayMarchPos.z) + 1 : ceil(LocalRayMarchPos.z) - 1;
+
+	float3 LocalBoundsHitT = (LocalBoundsPos - LocalRayMarchPos) * LocalRayDirectionInv;
+
+	// Remove floating-point rounding error
+	float Epsilon = 1.0e-4;
+	if (LocalBoundsHitT.x <= Epsilon) LocalBoundsHitT.x += abs(LocalRayDirectionInv.x);
+	if (LocalBoundsHitT.y <= Epsilon) LocalBoundsHitT.y += abs(LocalRayDirectionInv.y);
+	if (LocalBoundsHitT.z <= Epsilon) LocalBoundsHitT.z += abs(LocalRayDirectionInv.z);
+
+	float3 OpacityThreshold = 1.0e-4;
+	while ((LocalRayMarchT < LocalRayLength) && all(Transmittance > OpacityThreshold))
+	{
+		float LocalRayMarchDeltaT = 0.0;
+		if (LocalBoundsHitT.x < LocalBoundsHitT.y)
+		{
+			if (LocalBoundsHitT.x < LocalBoundsHitT.z)
+			{
+				LocalRayMarchDeltaT = LocalBoundsHitT.x - LocalRayMarchT;
+				LocalBoundsHitT.x += abs(LocalRayDirectionInv.x);
+			}
+			else
+			{
+				LocalRayMarchDeltaT = LocalBoundsHitT.z - LocalRayMarchT;
+				LocalBoundsHitT.z += abs(LocalRayDirectionInv.z);
+			}
+		}
+		else
+		{
+			if (LocalBoundsHitT.y < LocalBoundsHitT.z)
+			{
+				LocalRayMarchDeltaT = LocalBoundsHitT.y - LocalRayMarchT;
+				LocalBoundsHitT.y += abs(LocalRayDirectionInv.y);
+			}
+			else
+			{
+				LocalRayMarchDeltaT = LocalBoundsHitT.z - LocalRayMarchT;
+				LocalBoundsHitT.z += abs(LocalRayDirectionInv.z);
+			}
+		}
+
+		// Clip voxel delta-t by overall voxel ray-length.
+		if (LocalRayMarchT + LocalRayMarchDeltaT > LocalRayLength)
+		{
+			LocalRayMarchDeltaT = LocalRayLength - LocalRayMarchT;
+		}
+
+		float3 LocalEndPos = LocalRayMarchPos + LocalRayDirection * LocalRayMarchDeltaT;
+		float3 LocalSamplePos = clamp((LocalRayMarchPos + LocalEndPos) * 0.5, 0.0, BottomLevelVoxelResolution - 1.0);
+		uint LocalSampleLinearPos = BottomLevelVoxelIndex + MortonEncode3(LocalSamplePos);
+		float LocalScaleFactor = TopLevelVoxelRayLength / LocalRayLength;
+
+		FTopLevelGridData IndirectionData = IndirectionGridBuffer[LocalSampleLinearPos];
+		if (IsBottomLevelAllocated(IndirectionData))
+		{
+			MarchBottomLevelVoxelGrid(
+				LocalRayOrigin,
+				LocalRayDirection,
+				VoxelRayTMin,
+				VoxelRayTMax,
+				WorldRayHitSpan * LocalScaleFactor,
+				LocalRayMarchPos,
+				LocalEndPos,
+				IndirectionData,
+				ExtinctionGridBuffer,
+				Transmittance
+			);
+		}
+
+		LocalRayMarchT += LocalRayMarchDeltaT;
+		LocalRayMarchPos = LocalEndPos;
+	}
+}
+
+void MarchTopLevelOrthoVoxelGrid(
+	float3 VoxelRayOrigin,
+	float3 VoxelRayDirection,
+	float VoxelRayTMin,
+	float VoxelRayTMax,
+	float WorldRayHitSpan,
+	bool bEnableIndirectionGrid,
+	int3 TopLevelGridResolution,
+	StructuredBuffer<FTopLevelGridData> TopLevelGridBuffer,
+	StructuredBuffer<FTopLevelGridData> IndirectionGridBuffer,
+	StructuredBuffer<FScalarGridData> ExtinctionGridBuffer,
+	inout float3 Transmittance
+)
+{
+	float3 VoxelRayDirectionInv = 1.0 / VoxelRayDirection;
+
+	// Grid intersection
+	float VoxelRayMarchT = VoxelRayTMin;
+	float3 VoxelRayMarchPos = VoxelRayOrigin;
+
+	float3 VoxelBoundsPos;
+	VoxelBoundsPos.x = sign(VoxelRayDirection.x) > 0 ? floor(VoxelRayMarchPos.x) + 1 : ceil(VoxelRayMarchPos.x) - 1;
+	VoxelBoundsPos.y = sign(VoxelRayDirection.y) > 0 ? floor(VoxelRayMarchPos.y) + 1 : ceil(VoxelRayMarchPos.y) - 1;
+	VoxelBoundsPos.z = sign(VoxelRayDirection.z) > 0 ? floor(VoxelRayMarchPos.z) + 1 : ceil(VoxelRayMarchPos.z) - 1;
+
+	float3 VoxelBoundsHitT = (VoxelBoundsPos - VoxelRayMarchPos) * VoxelRayDirectionInv;
+
+	// Remove floating-point rounding error
+	float Epsilon = 1.0e-4;
+	if (VoxelBoundsHitT.x <= Epsilon) VoxelBoundsHitT.x += abs(VoxelRayDirectionInv.x);
+	if (VoxelBoundsHitT.y <= Epsilon) VoxelBoundsHitT.y += abs(VoxelRayDirectionInv.y);
+	if (VoxelBoundsHitT.z <= Epsilon) VoxelBoundsHitT.z += abs(VoxelRayDirectionInv.z);
+
+	float3 OpacityThreshold = 1.0e-4;
+	while ((VoxelRayMarchT < VoxelRayTMax) && all(Transmittance > OpacityThreshold))
+	{
+		float VoxelRayMarchDeltaT = 0.0;
+		if (VoxelBoundsHitT.x < VoxelBoundsHitT.y)
+		{
+			if (VoxelBoundsHitT.x < VoxelBoundsHitT.z)
+			{
+				VoxelRayMarchDeltaT = VoxelBoundsHitT.x - VoxelRayMarchT;
+				VoxelBoundsHitT.x += abs(VoxelRayDirectionInv.x);
+			}
+			else
+			{
+				VoxelRayMarchDeltaT = VoxelBoundsHitT.z - VoxelRayMarchT;
+				VoxelBoundsHitT.z += abs(VoxelRayDirectionInv.z);
+			}
+		}
+		else
+		{
+			if (VoxelBoundsHitT.y < VoxelBoundsHitT.z)
+			{
+				VoxelRayMarchDeltaT = VoxelBoundsHitT.y - VoxelRayMarchT;
+				VoxelBoundsHitT.y += abs(VoxelRayDirectionInv.y);
+			}
+			else
+			{
+				VoxelRayMarchDeltaT = VoxelBoundsHitT.z - VoxelRayMarchT;
+				VoxelBoundsHitT.z += abs(VoxelRayDirectionInv.z);
+			}
+		}
+
+		// Clip voxel delta-t by overall voxel ray-length.
+		if (VoxelRayMarchT + VoxelRayMarchDeltaT > VoxelRayTMax)
+		{
+			VoxelRayMarchDeltaT = VoxelRayTMax - VoxelRayMarchT;
+		}
+
+		// Sample at the midpoint of the voxel ray
+		float3 VoxelEndPos = VoxelRayMarchPos + VoxelRayDirection * VoxelRayMarchDeltaT;
+		float3 VoxelSamplePos = (VoxelRayMarchPos + VoxelEndPos) * 0.5;
+		uint VoxelSampleLinearPos = GetLinearIndex(VoxelSamplePos, TopLevelGridResolution);
+
+		int3 BottomLevelVoxelResolution = GetBottomLevelVoxelResolution(TopLevelGridBuffer[VoxelSampleLinearPos]);
+		if (IsBottomLevelAllocated(TopLevelGridBuffer[VoxelSampleLinearPos]))
+		{
+			if (bEnableIndirectionGrid)
+			{
+				MarchIndirectionOrthoVoxelGrid(
+					VoxelRayOrigin,
+					VoxelRayDirection,
+					VoxelRayTMin,
+					VoxelRayTMax,
+					WorldRayHitSpan,
+					VoxelRayMarchPos,
+					VoxelEndPos,
+					TopLevelGridBuffer[VoxelSampleLinearPos],
+					IndirectionGridBuffer,
+					ExtinctionGridBuffer,
+					Transmittance
+				);
+			}
+			else if (all(BottomLevelVoxelResolution > 1))
+			{
+				MarchBottomLevelVoxelGrid(
+					VoxelRayOrigin,
+					VoxelRayDirection,
+					VoxelRayTMin,
+					VoxelRayTMax,
+					WorldRayHitSpan,
+					VoxelRayMarchPos,
+					VoxelEndPos,
+					TopLevelGridBuffer[VoxelSampleLinearPos],
+					ExtinctionGridBuffer,
+					Transmittance
+				);
+			}
+			else
+			{
+				float3 Extinction = GetExtinction(TopLevelGridBuffer[VoxelSampleLinearPos], ExtinctionGridBuffer);
+				float WorldRayMarchDeltaT = max(VoxelRayMarchDeltaT * WorldRayHitSpan / VoxelRayTMax, 0.0);
+				Transmittance *= saturate(exp(-Extinction * WorldRayMarchDeltaT));
+			}
+		}
+
+		VoxelRayMarchT += VoxelRayMarchDeltaT;
+		VoxelRayMarchPos = VoxelEndPos;
+	}
+}
+
+float3 TraceOrthoVoxelGrid(float3 Origin, float3 Direction, float TMin, float TMax)
+{
+	float3 Transmittance = 1.0;
+
+	// Intersect world-space AABB to clip ray
+	float3 WorldBoundsMin = OrthoGridUniformBuffer.TopLevelGridWorldBoundsMin;
+	float3 WorldBoundsMax = OrthoGridUniformBuffer.TopLevelGridWorldBoundsMax;
+
+	float3 WorldRayOrigin = LWCHackToFloat(LWCSubtract(Origin, PrimaryView.PreViewTranslation));
+	float2 ClipHitT = IntersectAABB(WorldRayOrigin, Direction, TMin, TMax, WorldBoundsMin, WorldBoundsMax);
+	if (ClipHitT.x > ClipHitT.y)
+	{
+		return Transmittance;
+	}
+
+	float3 WorldRayBegin = WorldRayOrigin + Direction * ClipHitT.x;
+	float3 WorldRayEnd = WorldRayOrigin + Direction * ClipHitT.y;
+	float WorldRayLength = ClipHitT.y - ClipHitT.x;
+
+	// Transform to voxel-space
+	float3 WorldBoundsExtent = WorldBoundsMax - WorldBoundsMin;
+	float3 VoxelRayBegin = (WorldRayBegin - WorldBoundsMin) / WorldBoundsExtent * OrthoGridUniformBuffer.TopLevelGridResolution;
+	float3 VoxelRayEnd = (WorldRayEnd - WorldBoundsMin) / WorldBoundsExtent * OrthoGridUniformBuffer.TopLevelGridResolution;
+	float3 VoxelRayDirection = VoxelRayEnd - VoxelRayBegin;
+	float VoxelRayTMin = 0.0;
+	float VoxelRayTMax = length(VoxelRayDirection);
+	VoxelRayDirection /= VoxelRayTMax;
+
+	MarchTopLevelOrthoVoxelGrid(
+		VoxelRayBegin,
+		VoxelRayDirection,
+		VoxelRayTMin,
+		VoxelRayTMax,
+		WorldRayLength,
+		OrthoGridUniformBuffer.bEnableIndirectionGrid,
+		OrthoGridUniformBuffer.TopLevelGridResolution,
+		OrthoGridUniformBuffer.TopLevelGridBuffer,
+		OrthoGridUniformBuffer.IndirectionGridBuffer,
+		OrthoGridUniformBuffer.ExtinctionGridBuffer,
+		Transmittance
+	);
+
+	return Transmittance;
+}
+
+void MarchTopLevelFrustumVoxelGrid(
+	float3 VoxelRayOrigin,
+	float3 VoxelRayDirection,
+	float VoxelRayTMin,
+	float VoxelRayTMax,
+	float WorldRayHitSpan,
+	int3 TopLevelGridResolution,
+	StructuredBuffer<FTopLevelGridData> TopLevelGridBuffer,
+	StructuredBuffer<FScalarGridData> ExtinctionGridBuffer,
+	inout float3 Transmittance
+)
+{
+	float3 VoxelRayDirectionInv = 1.0 / VoxelRayDirection;
+
+	// Grid intersection
+	float VoxelRayMarchT = VoxelRayTMin;
+	float3 VoxelRayMarchPos = VoxelRayOrigin;
+
+	float3 VoxelBoundsPos;
+	VoxelBoundsPos.x = sign(VoxelRayDirection.x) > 0 ? floor(VoxelRayMarchPos.x) + 1 : ceil(VoxelRayMarchPos.x) - 1;
+	VoxelBoundsPos.y = sign(VoxelRayDirection.y) > 0 ? floor(VoxelRayMarchPos.y) + 1 : ceil(VoxelRayMarchPos.y) - 1;
+	VoxelBoundsPos.z = sign(VoxelRayDirection.z) > 0 ? floor(VoxelRayMarchPos.z) + 1 : ceil(VoxelRayMarchPos.z) - 1;
+
+	float3 VoxelBoundsHitT = (VoxelBoundsPos - VoxelRayMarchPos) * VoxelRayDirectionInv;
+
+	// Remove floating-point rounding error
+	float Epsilon = 1.0e-4;
+	if (VoxelBoundsHitT.x <= Epsilon) VoxelBoundsHitT.x += abs(VoxelRayDirectionInv.x);
+	if (VoxelBoundsHitT.y <= Epsilon) VoxelBoundsHitT.y += abs(VoxelRayDirectionInv.y);
+	if (VoxelBoundsHitT.z <= Epsilon) VoxelBoundsHitT.z += abs(VoxelRayDirectionInv.z);
+
+	float3 OpacityThreshold = 1.0e-4;
+	while ((VoxelRayMarchT < VoxelRayTMax) && all(Transmittance > OpacityThreshold))
+	{
+		float VoxelRayMarchDeltaT = 0.0;
+		if (VoxelBoundsHitT.x < VoxelBoundsHitT.y)
+		{
+			if (VoxelBoundsHitT.x < VoxelBoundsHitT.z)
+			{
+				VoxelRayMarchDeltaT = VoxelBoundsHitT.x - VoxelRayMarchT;
+				VoxelBoundsHitT.x += abs(VoxelRayDirectionInv.x);
+			}
+			else
+			{
+				VoxelRayMarchDeltaT = VoxelBoundsHitT.z - VoxelRayMarchT;
+				VoxelBoundsHitT.z += abs(VoxelRayDirectionInv.z);
+			}
+		}
+		else
+		{
+			if (VoxelBoundsHitT.y < VoxelBoundsHitT.z)
+			{
+				VoxelRayMarchDeltaT = VoxelBoundsHitT.y - VoxelRayMarchT;
+				VoxelBoundsHitT.y += abs(VoxelRayDirectionInv.y);
+			}
+			else
+			{
+				VoxelRayMarchDeltaT = VoxelBoundsHitT.z - VoxelRayMarchT;
+				VoxelBoundsHitT.z += abs(VoxelRayDirectionInv.z);
+			}
+		}
+
+		// Clip voxel delta-t by overall voxel ray-length.
+		if (VoxelRayMarchT + VoxelRayMarchDeltaT > VoxelRayTMax)
+		{
+			VoxelRayMarchDeltaT = VoxelRayTMax - VoxelRayMarchT;
+		}
+
+		// Sample at the midpoint of the voxel ray
+		float3 VoxelEndPos = VoxelRayMarchPos + VoxelRayDirection * VoxelRayMarchDeltaT;
+		float3 VoxelSamplePos = (VoxelRayMarchPos + VoxelEndPos) * 0.5;
+		uint VoxelSampleLinearPos = GetLinearIndex(VoxelSamplePos, TopLevelGridResolution);
+
+		int3 BottomLevelVoxelResolution = GetBottomLevelVoxelResolution(TopLevelGridBuffer[VoxelSampleLinearPos]);
+		if (IsBottomLevelAllocated(TopLevelGridBuffer[VoxelSampleLinearPos]))
+		{
+			if (all(BottomLevelVoxelResolution > 1))
+			{
+				MarchBottomLevelVoxelGrid(
+					VoxelRayOrigin,
+					VoxelRayDirection,
+					VoxelRayTMin,
+					VoxelRayTMax,
+					WorldRayHitSpan,
+					VoxelRayMarchPos,
+					VoxelEndPos,
+					TopLevelGridBuffer[VoxelSampleLinearPos],
+					ExtinctionGridBuffer,
+					Transmittance
+				);
+			}
+			else
+			{
+				float3 Extinction = GetExtinction(TopLevelGridBuffer[VoxelSampleLinearPos], ExtinctionGridBuffer);
+				float WorldRayMarchDeltaT = max(VoxelRayMarchDeltaT * WorldRayHitSpan / VoxelRayTMax, 0.0);
+				Transmittance *= saturate(exp(-Extinction * WorldRayMarchDeltaT));
+			}
+		}
+
+		VoxelRayMarchT += VoxelRayMarchDeltaT;
+		VoxelRayMarchPos = VoxelEndPos;
+	}
+}
+
+float3 TraceFrustumVoxelGrid(float3 Origin, float3 Direction, inout float TMin, inout float TMax)
+{
+	float3 Transmittance = 1.0;
+
+	float3 WorldRayOrigin = LWCHackToFloat(LWCSubtract(Origin, PrimaryView.PreViewTranslation));
+	float3 ViewRayOrigin = mul(float4(WorldRayOrigin, 1), FrustumGridUniformBuffer.WorldToView).xyz;
+	float3 ViewRayDirection = mul(float4(WorldRayOrigin + Direction, 1), FrustumGridUniformBuffer.WorldToView).xyz;
+	ViewRayDirection -= ViewRayOrigin;
+
+	int3 VoxelDimensions = FrustumGridUniformBuffer.VoxelDimensions;
+	float NearPlaneDepth = FrustumGridUniformBuffer.NearPlaneDepth;
+	float FarPlaneDepth = FrustumGridUniformBuffer.FarPlaneDepth;
+	float3 VoxelRayOrigin = ViewToVoxel(ViewRayOrigin, VoxelDimensions, NearPlaneDepth, FarPlaneDepth);
+	float3 VoxelRayDirection = ViewToVoxel(ViewRayOrigin + ViewRayDirection, VoxelDimensions, NearPlaneDepth, FarPlaneDepth);
+	VoxelRayDirection -= VoxelRayOrigin;
+	float UnitWorldToVoxelScale = length(VoxelRayDirection);
+	VoxelRayDirection /= UnitWorldToVoxelScale;
+
+	// Intersect voxel-space AABB to clip ray
+	float VoxelRayTMin = TMin * UnitWorldToVoxelScale;
+	float VoxelRayTMax = (TMax == POSITIVE_INFINITY) ? POSITIVE_INFINITY : TMax * UnitWorldToVoxelScale;
+	float2 ClipHitT = IntersectAABB(VoxelRayOrigin, VoxelRayDirection, VoxelRayTMin, VoxelRayTMax, 0, VoxelDimensions);
+	if (ClipHitT.x > ClipHitT.y)
+	{
+		TMin = TMax = 0;
+		return Transmittance;
+	}
+
+	VoxelRayOrigin += VoxelRayDirection * ClipHitT.x;
+	VoxelRayTMin = 0.0;
+	VoxelRayTMax = ClipHitT.y - ClipHitT.x;
+
+	float UnitVoxelToWorldScale = 1.0 / UnitWorldToVoxelScale;
+	float WorldRayLength = VoxelRayTMax * UnitVoxelToWorldScale;
+
+	MarchTopLevelFrustumVoxelGrid(
+		VoxelRayOrigin,
+		VoxelRayDirection,
+		VoxelRayTMin,
+		VoxelRayTMax,
+		WorldRayLength,
+		FrustumGridUniformBuffer.TopLevelFroxelGridResolution,
+		FrustumGridUniformBuffer.TopLevelFroxelGridBuffer,
+		FrustumGridUniformBuffer.ExtinctionFroxelGridBuffer,
+		Transmittance
+	);
+
+	TMin = ClipHitT.x * UnitVoxelToWorldScale;
+	TMax = ClipHitT.y * UnitVoxelToWorldScale;
+	return Transmittance;
+}
diff --git a/Engine/Shaders/Private/HeterogeneousVolumes/HeterogeneousVolumesVoxelGridTypes.ush b/Engine/Shaders/Private/HeterogeneousVolumes/HeterogeneousVolumesVoxelGridTypes.ush
new file mode 100644
--- /dev/null
+++ b/Engine/Shaders/Private/HeterogeneousVolumes/HeterogeneousVolumesVoxelGridTypes.ush
@@ -0,0 +1,21 @@
+// Copyright Epic Games, Inc. All Rights Reserved.
+
+struct FScalarGridData
+{
+	uint PackedData[2];
+};
+
+struct FVectorGridData
+{
+	uint PackedData[2];
+};
+
+struct FTopLevelGridData
+{
+	uint PackedData[1];
+};
+
+struct FTopLevelGridBitmaskData
+{
+	uint PackedData[2];
+};
diff --git a/Engine/Shaders/Private/HeterogeneousVolumes/HeterogeneousVolumesVoxelGridUtils.ush b/Engine/Shaders/Private/HeterogeneousVolumes/HeterogeneousVolumesVoxelGridUtils.ush
new file mode 100644
--- /dev/null
+++ b/Engine/Shaders/Private/HeterogeneousVolumes/HeterogeneousVolumesVoxelGridUtils.ush
@@ -0,0 +1,222 @@
+// Copyright Epic Games, Inc. All Rights Reserved.
+
+#include "../MortonCode.ush"
+
+#define EMPTY_VOXEL_INDEX 0x1FFFFFFF
+
+struct FMajorantData
+{
+	float Minorant;
+	float Majorant;
+};
+
+FMajorantData CreateMajorantData()
+{
+	FMajorantData MajorantData;
+	MajorantData.Minorant = 1.0e4;
+	MajorantData.Majorant = 0.0;
+	return MajorantData;
+}
+
+void UpdateMajorantData(inout FMajorantData MajorantData, float Value)
+{
+	MajorantData.Minorant = min(MajorantData.Minorant, Value);
+	MajorantData.Majorant = max(MajorantData.Majorant, Value);
+}
+
+void UpdateMajorantData(inout FMajorantData MajorantData, FMajorantData Value)
+{
+	MajorantData.Minorant = min(MajorantData.Minorant, Value.Minorant);
+	MajorantData.Majorant = max(MajorantData.Majorant, Value.Majorant);
+}
+
+FMajorantData GetMajorantData(FScalarGridData MajorantGridData)
+{
+	FMajorantData Data;
+	Data.Majorant = f16tof32(MajorantGridData.PackedData[0] & 0xFFFF);
+	Data.Minorant = f16tof32(MajorantGridData.PackedData[0] >> 16);
+	return Data;
+}
+
+void SetMajorantData(inout FScalarGridData MajorantGridData, FMajorantData MajorantData)
+{
+	MajorantGridData.PackedData[0] = (f32tof16(MajorantData.Minorant) << 16) | f32tof16(MajorantData.Majorant);
+}
+
+float3 GetExtinction(FScalarGridData ExtinctionGridData)
+{
+	float3 Extinction = 0.0;
+	Extinction.x = f16tof32(ExtinctionGridData.PackedData[0]);
+	Extinction.y = f16tof32(ExtinctionGridData.PackedData[0] >> 16);
+	Extinction.z = f16tof32(ExtinctionGridData.PackedData[1]);
+	return Extinction;
+}
+
+void SetExtinction(inout FScalarGridData ExtinctionGridData, float3 Extinction)
+{
+	ExtinctionGridData.PackedData[0] = f32tof16(Extinction.x) | (f32tof16(Extinction.y) << 16);
+	ExtinctionGridData.PackedData[1] = f32tof16(Extinction.z) | (ExtinctionGridData.PackedData[1] & 0xFFFF0000);
+}
+
+float3 GetEmission(FVectorGridData EmissionGridData)
+{
+	float3 Emission = 0.0;
+	Emission.x = f16tof32(EmissionGridData.PackedData[0]);
+	Emission.y = f16tof32(EmissionGridData.PackedData[0] >> 16);
+	Emission.z = f16tof32(EmissionGridData.PackedData[1]);
+	return Emission;
+}
+
+void SetEmission(inout FVectorGridData EmissionGridData, float3 Emission)
+{
+	EmissionGridData.PackedData[0] = f32tof16(Emission.x) | (f32tof16(Emission.y) << 16);
+	EmissionGridData.PackedData[1] = f32tof16(Emission.z) | (EmissionGridData.PackedData[1] & 0xFFFF0000);
+}
+
+float3 GetAlbedo(FVectorGridData AlbedoGridData)
+{
+	float3 Albedo = 0.0;
+	Albedo.x = f16tof32(AlbedoGridData.PackedData[0]);
+	Albedo.y = f16tof32(AlbedoGridData.PackedData[0] >> 16);
+	Albedo.z = f16tof32(AlbedoGridData.PackedData[1]);
+	return Albedo;
+}
+
+void SetAlbedo(inout FVectorGridData AlbedoGridData, float3 Albedo)
+{
+	AlbedoGridData.PackedData[0] = f32tof16(Albedo.x) | (f32tof16(Albedo.y) << 16);
+	AlbedoGridData.PackedData[1] = f32tof16(Albedo.z) | (AlbedoGridData.PackedData[1] & 0xFFFF0000);
+}
+
+FTopLevelGridData CreateTopLevelGridData(uint PackedData)
+{
+	FTopLevelGridData TopLevelGridData;
+	TopLevelGridData.PackedData[0] = PackedData;
+	return TopLevelGridData;
+}
+
+uint MortonEncode3(uint3 Voxel)
+{
+	return MortonCode3(Voxel.x) | MortonCode3(Voxel.y) << 1 | MortonCode3(Voxel.z) << 2;
+}
+
+uint3 MortonDecode3(uint Morton)
+{
+	uint3 Voxel = uint3(ReverseMortonCode3(Morton), ReverseMortonCode3(Morton >> 1), ReverseMortonCode3(Morton >> 2));
+	return Voxel;
+}
+
+int GetLinearIndex(uint3 VoxelIndex, int3 TopLevelGridResolution)
+{
+	int SliceSize = TopLevelGridResolution.x * TopLevelGridResolution.y;
+
+	return 
+		VoxelIndex.z * SliceSize +
+		VoxelIndex.y * TopLevelGridResolution.x +
+		VoxelIndex.x;
+}
+
+int3 GetVoxelIndex(int LinearIndex, int3 TopLevelGridResolution)
+{
+	int SliceSize = TopLevelGridResolution.x * TopLevelGridResolution.y;
+	int SliceIndex = LinearIndex / SliceSize;
+	int SliceRemainder = LinearIndex - SliceIndex * SliceSize;
+	
+	return int3(
+		SliceRemainder % TopLevelGridResolution.x,
+		SliceRemainder / TopLevelGridResolution.x,
+		SliceIndex
+	);
+}
+
+uint GetBottomLevelIndex(FTopLevelGridData TopLevelGridData)
+{
+	// Maximum addressable space is equivalent to 1024 x 1024 x 512 top-level volume
+	return TopLevelGridData.PackedData[0] >> 3;
+}
+
+void SetBottomLevelIndex(inout FTopLevelGridData TopLevelGridData, uint BottomLevelIndex)
+{
+	uint Index = BottomLevelIndex;
+	uint ResolutionExponent = TopLevelGridData.PackedData[0] & 0x7;
+	TopLevelGridData.PackedData[0] = (Index << 3) | ResolutionExponent;
+}
+
+bool IsBottomLevelAllocated(FTopLevelGridData TopLevelGridData)
+{
+	return GetBottomLevelIndex(TopLevelGridData) != EMPTY_VOXEL_INDEX;
+}
+
+bool IsBottomLevelEmpty(FTopLevelGridData TopLevelGridData)
+{
+	uint ResolutionExponent = TopLevelGridData.PackedData[0] & 0x7;
+	return (ResolutionExponent == 0u);
+}
+
+int3 GetBottomLevelVoxelResolution(FTopLevelGridData TopLevelGridData)
+{
+	uint ResolutionExponent = TopLevelGridData.PackedData[0] & 0x7;
+	int3 Resolution = ResolutionExponent > 0 ? (1u << (ResolutionExponent - 1)) : 0;
+
+	return Resolution;
+}
+
+void SetBottomLevelVoxelResolution(inout FTopLevelGridData TopLevelGridData, int3 BottomLevelVoxelResolution)
+{
+	uint Index = GetBottomLevelIndex(TopLevelGridData);
+	uint ResolutionExponent = int(BottomLevelVoxelResolution.x > 0 ? uint(log2(BottomLevelVoxelResolution.x)) + 1 : 0) & 0x7;
+	TopLevelGridData.PackedData[0] = (Index << 3) | ResolutionExponent;
+}
+
+float GetVoxelSize(FTopLevelGridData TopLevelGridData)
+{
+	return asfloat(TopLevelGridData.PackedData[0]);
+}
+
+void SetVoxelSize(inout FTopLevelGridData TopLevelGridData, float VoxelSize)
+{
+	TopLevelGridData.PackedData[0] = asint(VoxelSize);
+}
+
+// Convenience functions to retrieve the first element of a tile
+float3 GetExtinction(FTopLevelGridData TopLevelGridData, StructuredBuffer<FScalarGridData> ExtinctionGridBuffer)
+{
+	float3 Extinction = GetExtinction(ExtinctionGridBuffer[GetBottomLevelIndex(TopLevelGridData)]);
+	return Extinction;
+}
+
+float3 GetExtinction(FTopLevelGridData TopLevelGridData, StructuredBuffer<FTopLevelGridData> IndirectionGridBuffer, StructuredBuffer<FScalarGridData> ExtinctionGridBuffer)
+{
+	uint IndirectionIndex = GetBottomLevelIndex(TopLevelGridData);
+	uint BottomLevelIndex = GetBottomLevelIndex(IndirectionGridBuffer[IndirectionIndex]);
+	float3 Extinction = GetExtinction(ExtinctionGridBuffer[BottomLevelIndex]);
+	return Extinction;
+}
+
+float3 GetEmission(FTopLevelGridData TopLevelGridData, StructuredBuffer<FVectorGridData> EmissionGridBuffer)
+{
+	float3 Emission = GetEmission(EmissionGridBuffer[GetBottomLevelIndex(TopLevelGridData)]);
+	return Emission;
+}
+
+float3 GetEmission(FTopLevelGridData TopLevelGridData, StructuredBuffer<FTopLevelGridData> IndirectionGridBuffer, StructuredBuffer<FVectorGridData> EmissionGridBuffer)
+{
+	uint IndirectionIndex = GetBottomLevelIndex(TopLevelGridData);
+	uint BottomLevelIndex = GetBottomLevelIndex(IndirectionGridBuffer[IndirectionIndex]);
+	float3 Emission = GetEmission(EmissionGridBuffer[BottomLevelIndex]);
+	return Emission;
+}
+
+float3 GetAlbedo(FTopLevelGridData TopLevelGridData, StructuredBuffer<FVectorGridData> ScatteringGridBuffer)
+{
+	float3 Albedo = GetAlbedo(ScatteringGridBuffer[GetBottomLevelIndex(TopLevelGridData)]);
+	return Albedo;
+}
+
+float3 GetAlbedo(FTopLevelGridData TopLevelGridData, StructuredBuffer<FTopLevelGridData> IndirectionGridBuffer, StructuredBuffer<FVectorGridData> ScatteringGridBuffer)
+{
+	uint IndirectionIndex = GetBottomLevelIndex(TopLevelGridData);
+	uint BottomLevelIndex = GetBottomLevelIndex(IndirectionGridBuffer[IndirectionIndex]);
+	float3 Albedo = GetAlbedo(ScatteringGridBuffer[BottomLevelIndex]);
+	return Albedo;
+}
diff --git a/Engine/Source/Runtime/Renderer/Private/HeterogeneousVolumes/HeterogeneousVolumes.cpp b/Engine/Source/Runtime/Renderer/Private/HeterogeneousVolumes/HeterogeneousVolumes.cpp
--- a/Engine/Source/Runtime/Renderer/Private/HeterogeneousVolumes/HeterogeneousVolumes.cpp
+++ b/Engine/Source/Runtime/Renderer/Private/HeterogeneousVolumes/HeterogeneousVolumes.cpp
@@ -391,33 +391,59 @@ void FDeferredShadingSceneRenderer::RenderHeterogeneousVolumes(
 	RDG_GPU_STAT_SCOPE(GraphBuilder, HeterogeneousVolumesStat);
 	SCOPED_NAMED_EVENT(HeterogeneousVolumes, FColor::Emerald);
 
+	// Build global transmittance structure
+	TRDGUniformBufferRef<FOrthoVoxelGridUniformBufferParameters> OrthoGridUniformBuffer;
+	TRDGUniformBufferRef<FFrustumVoxelGridUniformBufferParameters> FrustumGridUniformBuffer;
+
+	if (HeterogeneousVolumes::GetDebugMode() != 0)
+	{
+		BuildOrthoVoxelGrid(GraphBuilder, Scene, Views, OrthoGridUniformBuffer);
+		BuildFrustumVoxelGrid(GraphBuilder, Scene, Views[0], FrustumGridUniformBuffer);
+	}
+
 	for (int32 ViewIndex = 0; ViewIndex < Views.Num(); ViewIndex++)
 	{
 		FViewInfo& View = Views[ViewIndex];
 
+		// Per-view??
+		FRDGTextureDesc Desc = SceneTextures.Color.Target->Desc;
+		Desc.Format = PF_FloatRGBA;
+		Desc.Flags &= ~(TexCreate_FastVRAM);
+		FRDGTextureRef HeterogeneousVolumeRadiance = GraphBuilder.CreateTexture(Desc, TEXT("HeterogeneousVolumes"));
+		AddClearUAVPass(GraphBuilder, GraphBuilder.CreateUAV(HeterogeneousVolumeRadiance), FLinearColor::Transparent);
+
 		if (ShouldRenderHeterogeneousVolumesForView(View))
 		{
-			FRDGTextureDesc Desc = SceneTextures.Color.Target->Desc;
-			Desc.Format = PF_FloatRGBA;
-			Desc.Flags &= ~(TexCreate_FastVRAM);
-			FRDGTextureRef HeterogeneousVolumeRadiance = GraphBuilder.CreateTexture(Desc, TEXT("HeterogeneousVolumes"));
-			AddClearUAVPass(GraphBuilder, GraphBuilder.CreateUAV(HeterogeneousVolumeRadiance), FLinearColor::Transparent);
-
-			for (int32 MeshBatchIndex = 0; MeshBatchIndex < View.HeterogeneousVolumesMeshBatches.Num(); ++MeshBatchIndex)
+			if (HeterogeneousVolumes::GetDebugMode() == 1)
+			{
+				RenderTransmittanceWithVoxelGrid(
+					GraphBuilder,
+					SceneTextures,
+					Scene,
+					ViewFamily,
+					Views[ViewIndex],
+					OrthoGridUniformBuffer,
+					FrustumGridUniformBuffer,
+					HeterogeneousVolumeRadiance
+				);
+			}
+			else
 			{
-				const FMeshBatch* Mesh = View.HeterogeneousVolumesMeshBatches[MeshBatchIndex].Mesh;
-				const FPrimitiveSceneProxy* PrimitiveSceneProxy = View.HeterogeneousVolumesMeshBatches[MeshBatchIndex].Proxy;
-				if (!ShouldRenderMeshBatchWithHeterogeneousVolumes(Mesh, PrimitiveSceneProxy, View.GetFeatureLevel()))
+				for (int32 MeshBatchIndex = 0; MeshBatchIndex < View.HeterogeneousVolumesMeshBatches.Num(); ++MeshBatchIndex)
 				{
-					continue;
-				}
+					const FMeshBatch* Mesh = View.HeterogeneousVolumesMeshBatches[MeshBatchIndex].Mesh;
+					const FPrimitiveSceneProxy* PrimitiveSceneProxy = View.HeterogeneousVolumesMeshBatches[MeshBatchIndex].Proxy;
+					if (!ShouldRenderMeshBatchWithHeterogeneousVolumes(Mesh, PrimitiveSceneProxy, View.GetFeatureLevel()))
+					{
+						continue;
+					}
 
-				const FMaterialRenderProxy* MaterialRenderProxy = Mesh->MaterialRenderProxy;
-				const FPrimitiveSceneInfo* PrimitiveSceneInfo = PrimitiveSceneProxy->GetPrimitiveSceneInfo();
-				const int32 PrimitiveId = PrimitiveSceneInfo->GetIndex();
-				const FBoxSphereBounds LocalBoxSphereBounds = PrimitiveSceneProxy->GetLocalBounds();
+					const FMaterialRenderProxy* MaterialRenderProxy = Mesh->MaterialRenderProxy;
+					const FPrimitiveSceneInfo* PrimitiveSceneInfo = PrimitiveSceneProxy->GetPrimitiveSceneInfo();
+					const int32 PrimitiveId = PrimitiveSceneInfo->GetIndex();
+					const FBoxSphereBounds LocalBoxSphereBounds = PrimitiveSceneProxy->GetLocalBounds();
 
-				RDG_EVENT_SCOPE(GraphBuilder, "%s", *PrimitiveSceneProxy->GetResourceName().ToString());
+					RDG_EVENT_SCOPE(GraphBuilder, "%s", *PrimitiveSceneProxy->GetResourceName().ToString());
 
 				// Allocate transmittance volume
 				// TODO: Allow option for scalar transmittance to conserve bandwidth
@@ -433,53 +459,54 @@ void FDeferredShadingSceneRenderer::RenderHeterogeneousVolumes(
 				FRDGTextureRef LightingCacheTexture = GraphBuilder.CreateTexture(LightingCacheDesc, TEXT("HeterogeneousVolumes.LightingCacheTexture"));
 				AddClearUAVPass(GraphBuilder, GraphBuilder.CreateUAV(LightingCacheTexture), FLinearColor::Black);
 
-				// Material baking executes a pre-shading pipeline
-				if (CVarHeterogeneousVolumesPreshading.GetValueOnRenderThread())
-				{
-					RenderWithPreshading(
-						GraphBuilder,
-						SceneTextures,
-						Scene,
-						ViewFamily,
-						View,
-						// Shadow Data
-						VisibleLightInfos,
-						VirtualShadowMapArray,
-						// Object data
-						PrimitiveSceneProxy,
-						MaterialRenderProxy,
-						PrimitiveId,
-						LocalBoxSphereBounds,
-						// Transmittance accleration
-						LightingCacheTexture,
-						// Output
-						HeterogeneousVolumeRadiance
-					);
-				}
-				// Otherwise execute a live-shading pipeline
-				else
-				{
-					RenderWithLiveShading(
-						GraphBuilder,
-						SceneTextures,
-						Scene,
-						View,
-						// Shadow Data
-						VisibleLightInfos,
-						VirtualShadowMapArray,
-						// Object Data
-						PrimitiveSceneProxy,
-						MaterialRenderProxy,
-						PrimitiveId,
-						LocalBoxSphereBounds,
-						// Transmittance accleration
-						LightingCacheTexture,
-						// Output
-						HeterogeneousVolumeRadiance
-					);
+					// Material baking executes a pre-shading pipeline
+					if (CVarHeterogeneousVolumesPreshading.GetValueOnRenderThread())
+					{
+						RenderWithPreshading(
+							GraphBuilder,
+							SceneTextures,
+							Scene,
+							ViewFamily,
+							View,
+							// Shadow Data
+							VisibleLightInfos,
+							VirtualShadowMapArray,
+							// Object data
+							PrimitiveSceneProxy,
+							MaterialRenderProxy,
+							PrimitiveId,
+							LocalBoxSphereBounds,
+							// Transmittance accleration
+							LightingCacheTexture,
+							// Output
+							HeterogeneousVolumeRadiance
+						);
+					}
+					// Otherwise execute a live-shading pipeline
+					else
+					{
+						RenderWithLiveShading(
+							GraphBuilder,
+							SceneTextures,
+							Scene,
+							View,
+							// Shadow Data
+							VisibleLightInfos,
+							VirtualShadowMapArray,
+							// Object Data
+							PrimitiveSceneProxy,
+							MaterialRenderProxy,
+							PrimitiveId,
+							LocalBoxSphereBounds,
+							// Transmittance accleration
+							LightingCacheTexture,
+							// Output
+							HeterogeneousVolumeRadiance
+						);
+					}
 				}
-			}
 
+			}
 			View.HeterogeneousVolumeRadiance = HeterogeneousVolumeRadiance;
 		}
 	}
@@ -572,4 +599,4 @@ void FDeferredShadingSceneRenderer::CompositeHeterogeneousVolumes(
 				GroupCount);
 		}
 	}
-}
\ No newline at end of file
+}
diff --git a/Engine/Source/Runtime/Renderer/Private/HeterogeneousVolumes/HeterogeneousVolumes.h b/Engine/Source/Runtime/Renderer/Private/HeterogeneousVolumes/HeterogeneousVolumes.h
--- a/Engine/Source/Runtime/Renderer/Private/HeterogeneousVolumes/HeterogeneousVolumes.h
+++ b/Engine/Source/Runtime/Renderer/Private/HeterogeneousVolumes/HeterogeneousVolumes.h
@@ -6,6 +6,7 @@
 #include "RendererInterface.h"
 #include "ShaderParameterMacros.h"
 
+class IHeterogeneousVolumeInterface;
 class FLightSceneInfo;
 class FPrimitiveSceneProxy;
 class FRayTracingScene;
@@ -51,12 +52,16 @@ namespace HeterogeneousVolumes
 	float GetMaxShadowTraceDistance();
 	float GetStepSize();
 	float GetMaxStepCount();
+	float GetMinimumVoxelSizeInFrustum();
+	float GetMinimumVoxelSizeOutsideFrustum();
 
 	int32 GetMipLevel();
 	int32 GetDebugMode();
 	int32 GetLightingCacheMode();
 	uint32 GetSparseVoxelMipBias();
-
+	int32 GetBottomLevelGridResolution();
+	int32 GetIndirectionGridResolution();
+	
 	bool ShouldJitter();
 	bool ShouldRefineSparseVoxels();
 	bool UseHardwareRayTracing();
@@ -66,6 +71,9 @@ namespace HeterogeneousVolumes
 	bool UseLightingCacheForInscattering();
 	bool UseLightingCacheForTransmittance();
 
+	bool EnableIndirectionGrid();
+	bool EnableLinearInterpolation();
+
 	// Convenience Utils
 	const IHeterogeneousVolumeInterface* GetInterface(const FPrimitiveSceneProxy* PrimitiveSceneProxy);
 
@@ -114,8 +122,164 @@ BEGIN_SHADER_PARAMETER_STRUCT(FLightingCacheParameters, )
 	SHADER_PARAMETER_RDG_TEXTURE(Texture3D, LightingCacheTexture)
 END_SHADER_PARAMETER_STRUCT()
 
+// Adaptive Voxel Grid structures
+
+struct FTopLevelGridBitmaskData
+{
+	uint32 PackedData[2];
+};
+
+struct FTopLevelGridData
+{
+	uint32 PackedData[1];
+};
+
+struct FScalarGridData
+{
+	uint32 PackedData[2];
+};
+
+struct FVectorGridData
+{
+	uint32 PackedData[2];
+};
+
+
+BEGIN_UNIFORM_BUFFER_STRUCT(FOrthoVoxelGridUniformBufferParameters, )
+	SHADER_PARAMETER(FVector3f, TopLevelGridWorldBoundsMin)
+	SHADER_PARAMETER(FVector3f, TopLevelGridWorldBoundsMax)
+	SHADER_PARAMETER(FIntVector, TopLevelGridResolution)
+
+	SHADER_PARAMETER(int32, bUseOrthoGrid)
+	SHADER_PARAMETER(int32, bUseMajorantGrid)
+	SHADER_PARAMETER(int32, bEnableIndirectionGrid)
+
+	SHADER_PARAMETER_RDG_BUFFER_SRV(StructuredBuffer<FTopLevelGridBitmaskData>, TopLevelGridBitmaskBuffer)
+	SHADER_PARAMETER_RDG_BUFFER_SRV(StructuredBuffer<FTopLevelGridData>, TopLevelGridBuffer)
+	SHADER_PARAMETER_RDG_BUFFER_SRV(StructuredBuffer<FTopLevelGridData>, IndirectionGridBuffer)
+	SHADER_PARAMETER_RDG_BUFFER_SRV(StructuredBuffer<FScalarGridData>, ExtinctionGridBuffer)
+	SHADER_PARAMETER_RDG_BUFFER_SRV(StructuredBuffer<FVectorGridData>, EmissionGridBuffer)
+	SHADER_PARAMETER_RDG_BUFFER_SRV(StructuredBuffer<FVectorGridData>, ScatteringGridBuffer)
+	SHADER_PARAMETER_RDG_BUFFER_SRV(StructuredBuffer<FScalarGridData>, MajorantGridBuffer)
+END_UNIFORM_BUFFER_STRUCT()
+
+BEGIN_UNIFORM_BUFFER_STRUCT(FFrustumVoxelGridUniformBufferParameters, )
+	SHADER_PARAMETER(FMatrix44f, WorldToClip)
+	SHADER_PARAMETER(FMatrix44f, ClipToWorld)
+
+	SHADER_PARAMETER(FMatrix44f, WorldToView)
+	SHADER_PARAMETER(FMatrix44f, ViewToWorld)
+
+	SHADER_PARAMETER(FMatrix44f, ViewToClip)
+	SHADER_PARAMETER(FMatrix44f, ClipToView)
+
+	SHADER_PARAMETER(FVector3f, TopLevelGridWorldBoundsMin)
+	SHADER_PARAMETER(FVector3f, TopLevelGridWorldBoundsMax)
+	SHADER_PARAMETER(FIntVector, TopLevelFroxelGridResolution)
+	SHADER_PARAMETER(FIntVector, VoxelDimensions)
+
+	SHADER_PARAMETER(int32, bUseFrustumGrid)
+
+	SHADER_PARAMETER(float, NearPlaneDepth)
+	SHADER_PARAMETER(float, FarPlaneDepth)
+
+	SHADER_PARAMETER_ARRAY(FVector4f, ViewFrustumPlanes, [6])
+
+	SHADER_PARAMETER_RDG_BUFFER_SRV(StructuredBuffer<FTopLevelGridData>, TopLevelFroxelGridBuffer)
+	SHADER_PARAMETER_RDG_BUFFER_SRV(StructuredBuffer<FScalarGridData>, ExtinctionFroxelGridBuffer)
+	SHADER_PARAMETER_RDG_BUFFER_SRV(StructuredBuffer<FVectorGridData>, EmissionFroxelGridBuffer)
+	SHADER_PARAMETER_RDG_BUFFER_SRV(StructuredBuffer<FVectorGridData>, ScatteringFroxelGridBuffer)
+END_UNIFORM_BUFFER_STRUCT()
+
 // Render specializations
 
+void BuildOrthoVoxelGrid(
+	FRDGBuilder& GraphBuilder,
+	const FScene* Scene,
+	/*const*/ TArray<FViewInfo>& Views,
+	TRDGUniformBufferRef<FOrthoVoxelGridUniformBufferParameters>& OrthoVoxelGridUniformBuffer
+);
+
+void BuildFrustumVoxelGrid(
+	FRDGBuilder& GraphBuilder,
+	const FScene* Scene,
+	const FViewInfo& View,
+	TRDGUniformBufferRef<FFrustumVoxelGridUniformBufferParameters>& FrustumVoxelGridUniformBuffer
+);
+
+struct FAdaptiveFrustumGridParameterCache
+{
+	FMatrix44f WorldToClip;
+	FMatrix44f ClipToWorld;
+
+	FMatrix44f WorldToView;
+	FMatrix44f ViewToWorld;
+
+	FMatrix44f ViewToClip;
+	FMatrix44f ClipToView;
+
+	FVector3f TopLevelGridWorldBoundsMin;
+	FVector3f TopLevelGridWorldBoundsMax;
+	FIntVector TopLevelGridResolution;
+	FIntVector VoxelDimensions;
+
+	int32 bUseFrustumGrid;
+
+	float NearPlaneDepth;
+	float FarPlaneDepth;
+
+	FVector4f ViewFrustumPlanes[6];
+
+	TRefCountPtr<FRDGPooledBuffer> TopLevelGridBuffer;
+	TRefCountPtr<FRDGPooledBuffer> ExtinctionGridBuffer;
+	TRefCountPtr<FRDGPooledBuffer> EmissionGridBuffer;
+	TRefCountPtr<FRDGPooledBuffer> ScatteringGridBuffer;
+};
+
+struct FAdaptiveOrthoGridParameterCache
+{
+	FVector3f TopLevelGridWorldBoundsMin;
+	FVector3f TopLevelGridWorldBoundsMax;
+	FIntVector TopLevelGridResolution;
+	int32 bUseOrthoGrid;
+	int32 bUseMajorantGrid;
+	int32 bEnableIndirectionGrid;
+
+	TRefCountPtr<FRDGPooledBuffer> TopLevelGridBitmaskBuffer;
+	TRefCountPtr<FRDGPooledBuffer> TopLevelGridBuffer;
+	TRefCountPtr<FRDGPooledBuffer> IndirectionGridBuffer;
+
+	TRefCountPtr<FRDGPooledBuffer> ExtinctionGridBuffer;
+	TRefCountPtr<FRDGPooledBuffer> EmissionGridBuffer;
+	TRefCountPtr<FRDGPooledBuffer> ScatteringGridBuffer;
+
+	TRefCountPtr<FRDGPooledBuffer> MajorantGridBuffer;
+};
+
+void ExtractFrustumVoxelGridUniformBuffer(
+	FRDGBuilder& GraphBuilder,
+	const TRDGUniformBufferRef<FFrustumVoxelGridUniformBufferParameters>& FrustumGridUniformBuffer,
+	FAdaptiveFrustumGridParameterCache& AdaptiveFrustumGridParameterCache
+);
+
+void RegisterExternalFrustumVoxelGridUniformBuffer(
+	FRDGBuilder& GraphBuilder,
+	const FAdaptiveFrustumGridParameterCache& AdaptiveFrustumGridParameterCache,
+	TRDGUniformBufferRef<FFrustumVoxelGridUniformBufferParameters>& FrustumGridUniformBuffer
+);
+
+void ExtractOrthoVoxelGridUniformBuffer(
+	FRDGBuilder& GraphBuilder,
+	const TRDGUniformBufferRef<FOrthoVoxelGridUniformBufferParameters>& OrthoGridUniformBuffer,
+	FAdaptiveOrthoGridParameterCache& AdaptiveOrthoGridParameterCache
+);
+
+void RegisterExternalOrthoVoxelGridUniformBuffer(
+	FRDGBuilder& GraphBuilder,
+	const FAdaptiveOrthoGridParameterCache& AdaptiveOrthoGridParameterCache,
+	TRDGUniformBufferRef<FOrthoVoxelGridUniformBufferParameters>& OrthoGridUniformBuffer
+);
+
 void RenderWithLiveShading(
 	FRDGBuilder& GraphBuilder,
 	const FSceneTextures& SceneTextures,
@@ -156,6 +320,19 @@ void RenderWithPreshading(
 	FRDGTextureRef& HeterogeneousVolumeRadiance
 );
 
+void RenderTransmittanceWithVoxelGrid(
+	FRDGBuilder& GraphBuilder,
+	// Scene data
+	const FSceneTextures& SceneTextures,
+	FScene* Scene,
+	const FSceneViewFamily& ViewFamily,
+	FViewInfo& View,
+	const TRDGUniformBufferRef<FOrthoVoxelGridUniformBufferParameters>& OrthoGridUniformBuffer,
+	const TRDGUniformBufferRef<FFrustumVoxelGridUniformBufferParameters>& FrustumGridUniformBuffer,
+	// Output
+	FRDGTextureRef& HeterogeneousVolumeRadiance
+);
+
 // Preshading Pipeline
 void ComputeHeterogeneousVolumeBakeMaterial(
 	FRDGBuilder& GraphBuilder,
diff --git a/Engine/Source/Runtime/Renderer/Private/HeterogeneousVolumes/HeterogeneousVolumesVoxelGridPipeline.cpp b/Engine/Source/Runtime/Renderer/Private/HeterogeneousVolumes/HeterogeneousVolumesVoxelGridPipeline.cpp
new file mode 100644
--- /dev/null
+++ b/Engine/Source/Runtime/Renderer/Private/HeterogeneousVolumes/HeterogeneousVolumesVoxelGridPipeline.cpp
@@ -0,0 +1,2381 @@
+// Copyright Epic Games, Inc. All Rights Reserved.
+
+#include "HeterogeneousVolumes.h"
+
+#include "HeterogeneousVolumeInterface.h"
+#include "LocalVertexFactory.h"
+#include "PixelShaderUtils.h"
+#include "RayTracingDefinitions.h"
+#include "RayTracingInstance.h"
+#include "RayTracingInstanceBufferUtil.h"
+#include "RendererPrivate.h"
+#include "ScenePrivate.h"
+#include "SceneManagement.h"
+#include "BlueNoise.h"
+
+static TAutoConsoleVariable<int32> CVarHeterogeneousVolumesEnableFrustumVoxelGrid(
+	TEXT("r.HeterogeneousVolumes.FrustumGrid"),
+	1,
+	TEXT("Enables a frustum voxel grid (Default = 1)"),
+	ECVF_RenderThreadSafe
+);
+
+static TAutoConsoleVariable<float> CVarHeterogeneousVolumesFrustumGridShadingRate(
+	TEXT("r.HeterogeneousVolumes.FrustumGrid.ShadingRate"),
+	4.0,
+	TEXT("The voxel tessellation rate, in pixel-space (Default = 4.0)"),
+	ECVF_RenderThreadSafe
+);
+
+static TAutoConsoleVariable<float> CVarHeterogeneousVolumesFrustumGridNearPlaneDistance(
+	TEXT("r.HeterogeneousVolumes.FrustumGrid.NearPlaneDistance"),
+	1.0,
+	TEXT("Sets near-plane distance for the frustum grid (Default = 1.0)\n"),
+	ECVF_RenderThreadSafe
+);
+
+static TAutoConsoleVariable<float> CVarHeterogeneousVolumesFrustumGridFarPlaneDistance(
+	TEXT("r.HeterogeneousVolumes.FrustumGrid.FarPlaneDistance"),
+	-1.0,
+	TEXT("Sets far-plane distance for the frustum grid (Default = -1.0)\n"),
+	ECVF_RenderThreadSafe
+);
+
+static TAutoConsoleVariable<int32> CVarHeterogeneousVolumesFrustumGridDepthSliceCount(
+	TEXT("r.HeterogeneousVolumes.FrustumGrid.DepthSliceCount"),
+	512,
+	TEXT("The number of depth slices (Default = 512)"),
+	ECVF_RenderThreadSafe
+);
+
+static TAutoConsoleVariable<int32> CVarHeterogeneousVolumesFrustumGridMaxBottomLevelMemoryInMegabytes(
+	TEXT("r.HeterogeneousVolumes.FrustumGrid.MaxBottomLevelMemoryInMegabytes"),
+	128,
+	TEXT("The minimum voxel size (Default = 128)"),
+	ECVF_RenderThreadSafe
+);
+
+static TAutoConsoleVariable<int32> CVarHeterogeneousVolumesEnableOrthoVoxelGrid(
+	TEXT("r.HeterogeneousVolumes.OrthoGrid"),
+	1,
+	TEXT("Enables an ortho voxel grid (Default = 1)"),
+	ECVF_RenderThreadSafe
+);
+
+static TAutoConsoleVariable<float> CVarHeterogeneousVolumesOrthoGridShadingRate(
+	TEXT("r.HeterogeneousVolumes.OrthoGrid.ShadingRate"),
+	4.0,
+	TEXT("The voxel tessellation rate, in pixel-space (Default = 4.0)"),
+	ECVF_RenderThreadSafe
+);
+
+static TAutoConsoleVariable<int32> CVarHeterogeneousVolumesOrthoGridMaxBottomLevelMemoryInMegabytes(
+	TEXT("r.HeterogeneousVolumes.OrthoGrid.MaxBottomLevelMemoryInMegabytes"),
+	128,
+	TEXT("The minimum voxel size (Default = 128)"),
+	ECVF_RenderThreadSafe
+);
+
+static TAutoConsoleVariable<int32> CVarHeterogeneousVolumesMarchingMode(
+	TEXT("r.HeterogeneousVolumes.Debug.MarchingMode"),
+	1,
+	TEXT("The marching mode (Default = 0)\n")
+	TEXT("0: Ray Marching (dt=StepSize)\n")
+	TEXT("1: Naive DDA\n")
+	TEXT("2: Optimized DDA\n")
+	TEXT("3: Optimized DDA w/ bitmask\n"),
+	ECVF_RenderThreadSafe
+);
+
+static TAutoConsoleVariable<int32> CVarHeterogeneousVolumesBottomLevelGridResolution(
+	TEXT("r.HeterogeneousVolumes.Tessellation.BottomLevelGrid.Resolution"),
+	4,
+	TEXT("Determines intra-tile bottom-level grid resolution (Default = 4)"),
+	ECVF_RenderThreadSafe
+);
+
+static TAutoConsoleVariable<int32> CVarHeterogeneousVolumesBottomLevelGridVoxelHashing(
+	TEXT("r.HeterogeneousVolumes.Tessellation.BottomLevelGrid.VoxelHashing"),
+	0,
+	TEXT("Enables bottom-level voxel hashing for deduplication (Default = 0)\n"),
+	ECVF_RenderThreadSafe
+);
+
+static TAutoConsoleVariable<float> CVarHeterogeneousVolumesMinimumVoxelSizeInFrustum(
+	TEXT("r.HeterogeneousVolumes.Tessellation.MinimumVoxelSizeInFrustum"),
+	1.0,
+	TEXT("The minimum voxel size (Default = 1.0)"),
+	ECVF_RenderThreadSafe
+);
+
+static TAutoConsoleVariable<float> CVarHeterogeneousVolumesMinimumVoxelSizeOutsideFrustum(
+	TEXT("r.HeterogeneousVolumes.Tessellation.MinimumVoxelSizeOutsideFrustum"),
+	100.0,
+	TEXT("The minimum voxel size (Default = 100.0)"),
+	ECVF_RenderThreadSafe
+);
+
+static TAutoConsoleVariable<int32> CVarHeterogeneousVolumesBottomLevelGridVoxelHashingMemoryInMegabytes(
+	TEXT("r.HeterogeneousVolumes.Tessellation.BottomLevelGrid.VoxelHashingMemoryInMegabytes"),
+	64,
+	TEXT("Enables bottom-level voxel hashing for deduplication (Default = 64)\n"),
+	ECVF_RenderThreadSafe
+);
+
+static TAutoConsoleVariable<int32> CVarHeterogeneousVolumesBottomLevelGridHomogeneousAggregation(
+	TEXT("r.HeterogeneousVolumes.Tessellation.BottomLevelGrid.HomogeneousAggregation"),
+	1,
+	TEXT("Enables bottom-level voxel homogeneous aggregation (Default = 1)\n"),
+	ECVF_RenderThreadSafe
+);
+
+static TAutoConsoleVariable<float> CVarHeterogeneousVolumesBottomLevelGridHomogeneousAggregationThreshold(
+	TEXT("r.HeterogeneousVolumes.Tessellation.BottomLevelGrid.HomogeneousAggregationThreshold"),
+	1.0e-3,
+	TEXT("Threshold for bottom-level voxel homogeneous aggregation (Default = 1.0e-3)\n"),
+	ECVF_RenderThreadSafe
+);
+
+static TAutoConsoleVariable<int32> CVarHeterogeneousVolumesIndirectionGrid(
+	TEXT("r.HeterogeneousVolumes.Tessellation.IndirectionGrid"),
+	1,
+	TEXT("Enables lazy allocation of bottom-level memory (Default = 1)\n"),
+	ECVF_RenderThreadSafe
+);
+
+static TAutoConsoleVariable<int32> CVarHeterogeneousVolumesIndirectionGridResolution(
+	TEXT("r.HeterogeneousVolumes.Tessellation.IndirectionGrid.Resolution"),
+	4,
+	TEXT("Determines intra-tile indirection grid resolution (Default = 4)"),
+	ECVF_RenderThreadSafe
+);
+
+static TAutoConsoleVariable<int32> CVarHeterogeneousVolumesFarPlaneAutoTransition(
+	TEXT("r.HeterogeneousVolumes.Tessellation.FarPlaneAutoTransition"),
+	1,
+	TEXT("Enables auto transitioning of far-plane distance, based on projected minimum voxel size (Default = 1)\n"),
+	ECVF_RenderThreadSafe
+);
+
+static TAutoConsoleVariable<int32> CVarHeterogeneousVolumesEnableTopLevelBitmask(
+	TEXT("r.HeterogeneousVolumes.Tessellation.TopLevelBitmask"),
+	0,
+	TEXT("Enables top-level bitmask to accelerate grid traversal (Default = 0)\n"),
+	ECVF_RenderThreadSafe
+);
+
+static TAutoConsoleVariable<int32> CVarHeterogeneousVolumesEnableMajorantGrid(
+	TEXT("r.HeterogeneousVolumes.Tessellation.MajorantGrid"),
+	0,
+	TEXT("Enables building majorant grids to accelerate volume tracking (Default = 0)\n"),
+	ECVF_RenderThreadSafe
+);
+
+IMPLEMENT_UNIFORM_BUFFER_STRUCT(FOrthoVoxelGridUniformBufferParameters, "OrthoGridUniformBuffer");
+IMPLEMENT_UNIFORM_BUFFER_STRUCT(FFrustumVoxelGridUniformBufferParameters, "FrustumGridUniformBuffer");
+
+struct FRasterTileData
+{
+	uint32 TopLevelGridLinearIndex;
+	uint32 BottomLevelGridLinearOffset;
+};
+
+namespace HeterogeneousVolumes
+{
+	bool EnableFrustumVoxelGrid()
+	{
+		return CVarHeterogeneousVolumesEnableFrustumVoxelGrid.GetValueOnRenderThread() != 0;
+	}
+
+	float GetShadingRateForFrustumGrid()
+	{
+		return FMath::Max(CVarHeterogeneousVolumesFrustumGridShadingRate.GetValueOnRenderThread(), 0.1);
+	}
+
+	float GetNearPlaneDistanceForFrustumGrid()
+	{
+		return CVarHeterogeneousVolumesFrustumGridNearPlaneDistance.GetValueOnRenderThread();
+	}
+
+	float GetFarPlaneDistanceForFrustumGrid()
+	{
+		return CVarHeterogeneousVolumesFrustumGridFarPlaneDistance.GetValueOnRenderThread();
+	}
+
+	int32 GetDepthSliceCountForFrustumGrid()
+	{
+		return FMath::Max(CVarHeterogeneousVolumesFrustumGridDepthSliceCount.GetValueOnRenderThread(), 1);
+	}
+
+	int32 GetMaxBottomLevelMemoryInMegabytesForFrustumGrid()
+	{
+		return FMath::Max(CVarHeterogeneousVolumesFrustumGridMaxBottomLevelMemoryInMegabytes.GetValueOnRenderThread(), 1);
+	}
+
+	bool EnableOrthoVoxelGrid()
+	{
+		return CVarHeterogeneousVolumesEnableOrthoVoxelGrid.GetValueOnRenderThread() != 0;
+	}
+
+	float GetShadingRateForOrthoGrid()
+	{
+		return FMath::Max(CVarHeterogeneousVolumesOrthoGridShadingRate.GetValueOnRenderThread(), 0.1);
+	}
+
+	int32 GetMaxBottomLevelMemoryInMegabytesForOrthoGrid()
+	{
+		return FMath::Max(CVarHeterogeneousVolumesOrthoGridMaxBottomLevelMemoryInMegabytes.GetValueOnRenderThread(), 1);
+	}
+
+	bool EnableFarPlaneAutoTransition()
+	{
+		return CVarHeterogeneousVolumesFarPlaneAutoTransition.GetValueOnRenderThread() != 0;
+	}
+
+	float GetMinimumVoxelSizeInFrustum()
+	{
+		return FMath::Max(CVarHeterogeneousVolumesMinimumVoxelSizeInFrustum.GetValueOnRenderThread(), 0.01);
+	}
+
+	float GetMinimumVoxelSizeOutsideFrustum()
+	{
+		return FMath::Max(CVarHeterogeneousVolumesMinimumVoxelSizeOutsideFrustum.GetValueOnRenderThread(), 0.01);
+	}
+
+	int32 GetMarchingMode()
+	{
+		return FMath::Clamp(CVarHeterogeneousVolumesMarchingMode.GetValueOnRenderThread(), 0, 3);
+	}
+
+	bool EnableVoxelHashing()
+	{
+		//return CVarHeterogeneousVolumesBottomLevelGridVoxelHashing.GetValueOnRenderThread() != 0;
+		return false;
+	}
+
+	uint32 GetVoxelHashingMemoryInMegabytes()
+	{
+		return CVarHeterogeneousVolumesBottomLevelGridVoxelHashingMemoryInMegabytes.GetValueOnRenderThread();
+	}
+
+	bool EnableHomogeneousAggregation()
+	{
+		return CVarHeterogeneousVolumesBottomLevelGridHomogeneousAggregation.GetValueOnRenderThread() != 0;
+	}
+
+	float GetHomogeneousAggregationThreshold()
+	{
+		return CVarHeterogeneousVolumesBottomLevelGridHomogeneousAggregationThreshold.GetValueOnRenderThread();
+	}
+
+	bool EnableIndirectionGrid()
+	{
+		return CVarHeterogeneousVolumesIndirectionGrid.GetValueOnRenderThread() != 0;
+	}
+
+	int32 GetBottomLevelGridResolution()
+	{
+		return FMath::Clamp(CVarHeterogeneousVolumesBottomLevelGridResolution.GetValueOnRenderThread(), 1, 4);
+	}
+
+	int32 GetIndirectionGridResolution()
+	{
+		return FMath::Clamp(CVarHeterogeneousVolumesIndirectionGridResolution.GetValueOnRenderThread(), 1, 4);
+	}
+
+	bool EnableLinearInterpolation()
+	{
+		return false;
+	}
+
+	bool EnableTopLevelBitmask()
+	{
+		return CVarHeterogeneousVolumesEnableTopLevelBitmask.GetValueOnRenderThread() != 0;
+	}
+
+	bool EnableMajorantGrid()
+	{
+		return CVarHeterogeneousVolumesEnableMajorantGrid.GetValueOnRenderThread() != 0;
+	}
+}
+
+class FMarkTopLevelGridVoxelsForFrustumGrid : public FGlobalShader
+{
+	DECLARE_GLOBAL_SHADER(FMarkTopLevelGridVoxelsForFrustumGrid);
+	SHADER_USE_PARAMETER_STRUCT(FMarkTopLevelGridVoxelsForFrustumGrid, FGlobalShader);
+
+	using FPermutationDomain = TShaderPermutationDomain<>;
+
+	BEGIN_SHADER_PARAMETER_STRUCT(FParameters, )
+		SHADER_PARAMETER_STRUCT_REF(FViewUniformShaderParameters, View)
+
+		SHADER_PARAMETER(FIntVector, TopLevelGridResolution)
+
+		SHADER_PARAMETER(FVector3f, PrimitiveWorldBoundsMin)
+		SHADER_PARAMETER(FVector3f, PrimitiveWorldBoundsMax)
+
+		SHADER_PARAMETER(FMatrix44f, ViewToWorld)
+		SHADER_PARAMETER(float, HalfFOV)
+		SHADER_PARAMETER(float, NearPlaneDepth)
+		SHADER_PARAMETER(float, FarPlaneDepth)
+
+		SHADER_PARAMETER(FIntVector, VoxelDimensions)
+
+		SHADER_PARAMETER_RDG_BUFFER_UAV(RWStructuredBuffer<FTopLevelGridData>, RWTopLevelGridBuffer)
+	END_SHADER_PARAMETER_STRUCT()
+
+	static bool ShouldCompilePermutation(
+		const FGlobalShaderPermutationParameters& Parameters
+	)
+	{
+		return DoesPlatformSupportHeterogeneousVolumes(Parameters.Platform);
+	}
+
+	static void ModifyCompilationEnvironment(
+		const FGlobalShaderPermutationParameters& Parameters,
+		FShaderCompilerEnvironment& OutEnvironment
+	)
+	{
+		FMaterialShader::ModifyCompilationEnvironment(Parameters, OutEnvironment);
+		OutEnvironment.SetDefine(TEXT("THREADGROUP_SIZE_1D"), GetThreadGroupSize1D());
+		OutEnvironment.SetDefine(TEXT("THREADGROUP_SIZE_3D"), GetThreadGroupSize3D());
+	}
+
+	static int32 GetThreadGroupSize1D() { return GetThreadGroupSize3D() * GetThreadGroupSize3D() * GetThreadGroupSize3D(); }
+	static int32 GetThreadGroupSize3D() { return 4; }
+};
+
+IMPLEMENT_MATERIAL_SHADER_TYPE(, FMarkTopLevelGridVoxelsForFrustumGrid, TEXT("/Engine/Private/HeterogeneousVolumes/HeterogeneousVolumesFrustumVoxelGrid.usf"), TEXT("MarkTopLevelGridVoxelsForFrustumGridCS"), SF_Compute);
+
+class FRasterizeBottomLevelFroxelGridCS : public FMeshMaterialShader
+{
+	DECLARE_SHADER_TYPE(FRasterizeBottomLevelFroxelGridCS, MeshMaterial);
+
+	using FPermutationDomain = TShaderPermutationDomain<>;
+
+	BEGIN_SHADER_PARAMETER_STRUCT(FParameters, )
+		// Scene data
+		SHADER_PARAMETER_STRUCT_REF(FViewUniformShaderParameters, View)
+
+		// Primitive data
+		SHADER_PARAMETER(FVector3f, PrimitiveWorldBoundsMin)
+		SHADER_PARAMETER(FVector3f, PrimitiveWorldBoundsMax)
+		SHADER_PARAMETER(FMatrix44f, LocalToWorld)
+		SHADER_PARAMETER(FMatrix44f, WorldToLocal)
+		SHADER_PARAMETER(FVector3f, LocalBoundsOrigin)
+		SHADER_PARAMETER(FVector3f, LocalBoundsExtent)
+		SHADER_PARAMETER(int32, PrimitiveId)
+
+		// Volume data
+		SHADER_PARAMETER(FIntVector, TopLevelGridResolution)
+		SHADER_PARAMETER(FIntVector, VoxelDimensions)
+		SHADER_PARAMETER(FMatrix44f, ViewToWorld)
+		SHADER_PARAMETER(float, HalfFOV)
+		SHADER_PARAMETER(float, NearPlaneDepth)
+		SHADER_PARAMETER(float, FarPlaneDepth)
+
+		SHADER_PARAMETER(int, BottomLevelGridBufferSize)
+
+		// Sampling data
+		SHADER_PARAMETER_STRUCT_REF(FBlueNoise, BlueNoise)
+
+		// Raster tile data
+		SHADER_PARAMETER_RDG_BUFFER_SRV(Buffer<uint>, RasterTileAllocatorBuffer)
+		SHADER_PARAMETER_RDG_BUFFER_SRV(StructuredBuffer<FRasterTileData>, RasterTileBuffer)
+
+		// Indirect args
+		RDG_BUFFER_ACCESS(IndirectArgs, ERHIAccess::IndirectArgs)
+
+		// Grid data
+		//SHADER_PARAMETER_RDG_BUFFER_SRV(StructuredBuffer<FTopLevelGridData>, TopLevelGridBuffer)
+
+		SHADER_PARAMETER_RDG_BUFFER_UAV(RWBuffer<uint>, RWBottomLevelGridAllocatorBuffer)
+		SHADER_PARAMETER_RDG_BUFFER_UAV(RWStructuredBuffer<uint>, RWTopLevelGridBuffer)
+
+		SHADER_PARAMETER_RDG_BUFFER_UAV(RWStructuredBuffer<FScalarGridData>, RWExtinctionGridBuffer)
+		SHADER_PARAMETER_RDG_BUFFER_UAV(RWStructuredBuffer<FVectorGridData>, RWEmissionGridBuffer)
+		SHADER_PARAMETER_RDG_BUFFER_UAV(RWStructuredBuffer<FVectorGridData>, RWScatteringGridBuffer)
+	END_SHADER_PARAMETER_STRUCT()
+
+	FRasterizeBottomLevelFroxelGridCS() = default;
+
+	FRasterizeBottomLevelFroxelGridCS(
+		const FMeshMaterialShaderType::CompiledShaderInitializerType& Initializer
+	)
+		: FMeshMaterialShader(Initializer)
+	{
+		Bindings.BindForLegacyShaderParameters(
+			this,
+			Initializer.PermutationId,
+			Initializer.ParameterMap,
+			*FParameters::FTypeInfo::GetStructMetadata(),
+			// Don't require full bindings, we use FMaterialShader::SetParameters
+			false
+		);
+	}
+
+	static bool ShouldCompilePermutation(
+		const FMaterialShaderPermutationParameters& Parameters
+	)
+	{
+		return DoesPlatformSupportHeterogeneousVolumes(Parameters.Platform)
+			&& DoesMaterialShaderSupportHeterogeneousVolumes(Parameters.MaterialParameters);
+	}
+
+	static FPermutationDomain RemapPermutation(FPermutationDomain PermutationVector)
+	{
+		return PermutationVector;
+	}
+
+	static void ModifyCompilationEnvironment(
+		const FMaterialShaderPermutationParameters& Parameters,
+		FShaderCompilerEnvironment& OutEnvironment
+	)
+	{
+		FMaterialShader::ModifyCompilationEnvironment(Parameters, OutEnvironment);
+		OutEnvironment.SetDefine(TEXT("THREADGROUP_SIZE_1D"), GetThreadGroupSize1D());
+		OutEnvironment.SetDefine(TEXT("THREADGROUP_SIZE_3D"), GetThreadGroupSize3D());
+
+		// This shader takes a very long time to compile with FXC, so we pre-compile it with DXC first and then forward the optimized HLSL to FXC.
+		OutEnvironment.CompilerFlags.Add(CFLAG_PrecompileWithDXC);
+		OutEnvironment.CompilerFlags.Add(CFLAG_AllowTypedUAVLoads);
+
+		OutEnvironment.SetDefine(TEXT("GET_PRIMITIVE_DATA_OVERRIDE"), 1);
+	}
+
+	void SetParameters(
+		FRHIComputeCommandList& RHICmdList,
+		FRHIComputeShader* ShaderRHI,
+		const FViewInfo& View,
+		const FMaterialRenderProxy* MaterialProxy,
+		const FMaterial& Material
+	)
+	{
+		FMaterialShader::SetViewParameters(RHICmdList, ShaderRHI, View, View.ViewUniformBuffer);
+		FMaterialShader::SetParameters(RHICmdList, ShaderRHI, MaterialProxy, Material, View);
+	}
+
+	static int32 GetThreadGroupSize1D() { return GetThreadGroupSize3D() * GetThreadGroupSize3D() * GetThreadGroupSize3D(); }
+	static int32 GetThreadGroupSize3D() { return 4; }
+};
+
+IMPLEMENT_MATERIAL_SHADER_TYPE(, FRasterizeBottomLevelFroxelGridCS, TEXT("/Engine/Private/HeterogeneousVolumes/HeterogeneousVolumesMaterialBakingPipeline.usf"), TEXT("RasterizeBottomLevelFroxelGridCS"), SF_Compute);
+
+class FTopLevelGridCalculateVoxelSize : public FGlobalShader
+{
+	DECLARE_GLOBAL_SHADER(FTopLevelGridCalculateVoxelSize);
+	SHADER_USE_PARAMETER_STRUCT(FTopLevelGridCalculateVoxelSize, FGlobalShader);
+
+	using FPermutationDomain = TShaderPermutationDomain<>;
+
+	BEGIN_SHADER_PARAMETER_STRUCT(FParameters, )
+		SHADER_PARAMETER_STRUCT_REF(FViewUniformShaderParameters, View)
+
+		SHADER_PARAMETER(FIntVector, TopLevelGridResolution)
+		SHADER_PARAMETER(FVector3f, TopLevelGridWorldBoundsMin)
+		SHADER_PARAMETER(FVector3f, TopLevelGridWorldBoundsMax)
+
+		SHADER_PARAMETER(FVector3f, PrimitiveWorldBoundsMin)
+		SHADER_PARAMETER(FVector3f, PrimitiveWorldBoundsMax)
+		
+		SHADER_PARAMETER(float, ShadingRate)
+		SHADER_PARAMETER(float, MinVoxelSizeInFrustum)
+		SHADER_PARAMETER(float, MinVoxelSizeOutOfFrustum)
+
+		SHADER_PARAMETER_RDG_BUFFER_UAV(RWStructuredBuffer<FTopLevelGridData>, RWTopLevelGridBuffer)
+	END_SHADER_PARAMETER_STRUCT()
+
+	static bool ShouldCompilePermutation(
+		const FGlobalShaderPermutationParameters& Parameters
+	)
+	{
+		return DoesPlatformSupportHeterogeneousVolumes(Parameters.Platform);
+	}
+
+	static void ModifyCompilationEnvironment(
+		const FGlobalShaderPermutationParameters& Parameters,
+		FShaderCompilerEnvironment& OutEnvironment
+	)
+	{
+		FMaterialShader::ModifyCompilationEnvironment(Parameters, OutEnvironment);
+		OutEnvironment.SetDefine(TEXT("THREADGROUP_SIZE_1D"), GetThreadGroupSize1D());
+		OutEnvironment.SetDefine(TEXT("THREADGROUP_SIZE_3D"), GetThreadGroupSize3D());
+	}
+
+	static int32 GetThreadGroupSize1D() { return GetThreadGroupSize3D() * GetThreadGroupSize3D() * GetThreadGroupSize3D(); }
+	static int32 GetThreadGroupSize3D() { return 4; }
+};
+
+IMPLEMENT_MATERIAL_SHADER_TYPE(, FTopLevelGridCalculateVoxelSize, TEXT("/Engine/Private/HeterogeneousVolumes/HeterogeneousVolumesVoxelGridPipeline.usf"), TEXT("TopLevelGridCalculateVoxelSize"), SF_Compute);
+
+class FAllocateBottomLevelGrid : public FGlobalShader
+{
+	DECLARE_GLOBAL_SHADER(FAllocateBottomLevelGrid);
+	SHADER_USE_PARAMETER_STRUCT(FAllocateBottomLevelGrid, FGlobalShader);
+
+	class FEnableIndirectionGrid : SHADER_PERMUTATION_BOOL("DIM_ENABLE_INDIRECTION_GRID");
+	using FPermutationDomain = TShaderPermutationDomain<FEnableIndirectionGrid>;
+
+	BEGIN_SHADER_PARAMETER_STRUCT(FParameters, )
+		SHADER_PARAMETER_STRUCT_REF(FViewUniformShaderParameters, View)
+
+		SHADER_PARAMETER(FIntVector, TopLevelGridResolution)
+		SHADER_PARAMETER(FVector3f, TopLevelGridWorldBoundsMin)
+		SHADER_PARAMETER(FVector3f, TopLevelGridWorldBoundsMax)
+
+		SHADER_PARAMETER(int, MaxVoxelResolution)
+		SHADER_PARAMETER(int, bSampleAtVertices)
+
+		SHADER_PARAMETER_RDG_BUFFER_UAV(RWStructuredBuffer<FTopLevelGridData>, RWTopLevelGridBuffer)
+		SHADER_PARAMETER_RDG_BUFFER_UAV(RWBuffer<uint>, RWBottomLevelGridAllocatorBuffer)
+
+		// Indirection Grid
+		SHADER_PARAMETER(int, IndirectionGridBufferSize)
+		SHADER_PARAMETER_RDG_BUFFER_UAV(RWStructuredBuffer<FTopLevelGridData>, RWIndirectionGridBuffer)
+		SHADER_PARAMETER_RDG_BUFFER_UAV(RWBuffer<uint>, RWIndirectionGridAllocatorBuffer)
+	END_SHADER_PARAMETER_STRUCT()
+
+	static bool ShouldCompilePermutation(
+		const FGlobalShaderPermutationParameters& Parameters
+	)
+	{
+		return DoesPlatformSupportHeterogeneousVolumes(Parameters.Platform);
+	}
+
+	static void ModifyCompilationEnvironment(
+		const FGlobalShaderPermutationParameters& Parameters,
+		FShaderCompilerEnvironment& OutEnvironment
+	)
+	{
+		FMaterialShader::ModifyCompilationEnvironment(Parameters, OutEnvironment);
+		OutEnvironment.SetDefine(TEXT("THREADGROUP_SIZE_1D"), GetThreadGroupSize1D());
+		OutEnvironment.SetDefine(TEXT("THREADGROUP_SIZE_3D"), GetThreadGroupSize3D());
+	}
+
+	static int32 GetThreadGroupSize1D() { return GetThreadGroupSize3D() * GetThreadGroupSize3D() * GetThreadGroupSize3D(); }
+	static int32 GetThreadGroupSize3D() { return 4; }
+};
+
+IMPLEMENT_MATERIAL_SHADER_TYPE(, FAllocateBottomLevelGrid, TEXT("/Engine/Private/HeterogeneousVolumes/HeterogeneousVolumesVoxelGridPipeline.usf"), TEXT("AllocateBottomLevelGrid"), SF_Compute);
+
+class FGenerateRasterTiles : public FGlobalShader
+{
+	DECLARE_GLOBAL_SHADER(FGenerateRasterTiles);
+	SHADER_USE_PARAMETER_STRUCT(FGenerateRasterTiles, FGlobalShader);
+
+	class FEnableIndirectionGrid : SHADER_PERMUTATION_BOOL("DIM_ENABLE_INDIRECTION_GRID");
+	using FPermutationDomain = TShaderPermutationDomain<FEnableIndirectionGrid>;
+
+	BEGIN_SHADER_PARAMETER_STRUCT(FParameters, )
+		//SHADER_PARAMETER_STRUCT_REF(FViewUniformShaderParameters, View)
+
+		SHADER_PARAMETER(FIntVector, TopLevelGridResolution)
+		SHADER_PARAMETER(int, RasterTileVoxelResolution)
+		SHADER_PARAMETER(int, MaxNumRasterTiles)
+		//SHADER_PARAMETER(FVector3f, TopLevelGridWorldBoundsMin)
+		//SHADER_PARAMETER(FVector3f, TopLevelGridWorldBoundsMax)
+
+		SHADER_PARAMETER_RDG_BUFFER_SRV(StructuredBuffer<FTopLevelGridData>, TopLevelGridBuffer)
+
+		SHADER_PARAMETER_RDG_BUFFER_UAV(RWBuffer<uint>, RWRasterTileAllocatorBuffer)
+		SHADER_PARAMETER_RDG_BUFFER_UAV(RWStructuredBuffer<FRasterTileData>, RWRasterTileBuffer)
+	END_SHADER_PARAMETER_STRUCT()
+
+	static bool ShouldCompilePermutation(
+		const FGlobalShaderPermutationParameters& Parameters
+	)
+	{
+		return DoesPlatformSupportHeterogeneousVolumes(Parameters.Platform);
+	}
+
+	static void ModifyCompilationEnvironment(
+		const FGlobalShaderPermutationParameters& Parameters,
+		FShaderCompilerEnvironment& OutEnvironment
+	)
+	{
+		FMaterialShader::ModifyCompilationEnvironment(Parameters, OutEnvironment);
+		OutEnvironment.SetDefine(TEXT("THREADGROUP_SIZE_1D"), GetThreadGroupSize1D());
+		OutEnvironment.SetDefine(TEXT("THREADGROUP_SIZE_3D"), GetThreadGroupSize3D());
+	}
+
+	static int32 GetThreadGroupSize1D() { return GetThreadGroupSize3D() * GetThreadGroupSize3D() * GetThreadGroupSize3D(); }
+	static int32 GetThreadGroupSize3D() { return 4; }
+};
+
+IMPLEMENT_MATERIAL_SHADER_TYPE(, FGenerateRasterTiles, TEXT("/Engine/Private/HeterogeneousVolumes/HeterogeneousVolumesVoxelGridPipeline.usf"), TEXT("GenerateRasterTiles"), SF_Compute);
+
+class FSetRasterizeBottomLevelGridIndirectArgs : public FGlobalShader
+{
+	DECLARE_GLOBAL_SHADER(FSetRasterizeBottomLevelGridIndirectArgs);
+	SHADER_USE_PARAMETER_STRUCT(FSetRasterizeBottomLevelGridIndirectArgs, FGlobalShader);
+
+	using FPermutationDomain = TShaderPermutationDomain<>;
+
+	BEGIN_SHADER_PARAMETER_STRUCT(FParameters, )
+		SHADER_PARAMETER(FIntVector, MaxDispatchThreadGroupsPerDimension)
+
+		SHADER_PARAMETER_RDG_BUFFER_SRV(Buffer<uint>, RasterTileAllocatorBuffer)
+		SHADER_PARAMETER_RDG_BUFFER_UAV(RWBuffer<uint>, RWRasterizeBottomLevelGridIndirectArgsBuffer)
+	END_SHADER_PARAMETER_STRUCT()
+
+	static bool ShouldCompilePermutation(
+		const FGlobalShaderPermutationParameters& Parameters
+	)
+	{
+		return DoesPlatformSupportHeterogeneousVolumes(Parameters.Platform);
+	}
+
+	static void ModifyCompilationEnvironment(
+		const FGlobalShaderPermutationParameters& Parameters,
+		FShaderCompilerEnvironment& OutEnvironment
+	)
+	{
+		FMaterialShader::ModifyCompilationEnvironment(Parameters, OutEnvironment);
+		OutEnvironment.SetDefine(TEXT("THREADGROUP_SIZE_1D"), GetThreadGroupSize1D());
+		OutEnvironment.SetDefine(TEXT("THREADGROUP_SIZE_3D"), GetThreadGroupSize3D());
+	}
+
+	static int32 GetThreadGroupSize1D() { return GetThreadGroupSize3D() * GetThreadGroupSize3D() * GetThreadGroupSize3D(); }
+	static int32 GetThreadGroupSize3D() { return 4; }
+};
+
+IMPLEMENT_MATERIAL_SHADER_TYPE(, FSetRasterizeBottomLevelGridIndirectArgs, TEXT("/Engine/Private/HeterogeneousVolumes/HeterogeneousVolumesVoxelGridPipeline.usf"), TEXT("SetRasterizeBottomLevelGridIndirectArgs"), SF_Compute);
+
+class FRasterizeBottomLevelGrid : public FMeshMaterialShader
+{
+	DECLARE_SHADER_TYPE(FRasterizeBottomLevelGrid, MeshMaterial);
+
+	class FEnableIndirectionGrid : SHADER_PERMUTATION_BOOL("DIM_ENABLE_INDIRECTION_GRID");
+	//class FEnableVoxelHashing : SHADER_PERMUTATION_BOOL("DIM_ENABLE_VOXEL_HASHING");
+	class FEnableHomogeneousAggregation : SHADER_PERMUTATION_BOOL("DIM_ENABLE_HOMOGENEOUS_AGGREGATION");
+	using FPermutationDomain = TShaderPermutationDomain<FEnableIndirectionGrid, FEnableHomogeneousAggregation>;
+
+	BEGIN_SHADER_PARAMETER_STRUCT(FParameters, )
+		// Scene data
+		SHADER_PARAMETER_STRUCT_REF(FViewUniformShaderParameters, View)
+
+		// Primitive data
+		SHADER_PARAMETER(FVector3f, PrimitiveWorldBoundsMin)
+		SHADER_PARAMETER(FVector3f, PrimitiveWorldBoundsMax)
+		SHADER_PARAMETER(FMatrix44f, LocalToWorld)
+		SHADER_PARAMETER(FMatrix44f, WorldToLocal)
+		SHADER_PARAMETER(FVector3f, LocalBoundsOrigin)
+		SHADER_PARAMETER(FVector3f, LocalBoundsExtent)
+		SHADER_PARAMETER(int32, PrimitiveId)
+
+		// Volume data
+		SHADER_PARAMETER(FIntVector, TopLevelGridResolution)
+		SHADER_PARAMETER(FVector3f, TopLevelGridWorldBoundsMin)
+		SHADER_PARAMETER(FVector3f, TopLevelGridWorldBoundsMax)
+
+		SHADER_PARAMETER(int, BottomLevelGridBufferSize)
+
+		// Volume sample mode
+		SHADER_PARAMETER(int, bSampleAtVertices)
+
+		// Raster tile data
+		SHADER_PARAMETER_RDG_BUFFER_SRV(Buffer<uint>, RasterTileAllocatorBuffer)
+		SHADER_PARAMETER_RDG_BUFFER_SRV(StructuredBuffer<FRasterTileData>, RasterTileBuffer)
+
+		// Indirect args
+		RDG_BUFFER_ACCESS(IndirectArgs, ERHIAccess::IndirectArgs)
+
+		// Grid data
+		SHADER_PARAMETER_RDG_BUFFER_SRV(StructuredBuffer<FTopLevelGridData>, TopLevelGridBuffer)
+
+		SHADER_PARAMETER_RDG_BUFFER_UAV(RWBuffer<uint>, RWBottomLevelGridAllocatorBuffer)
+		SHADER_PARAMETER_RDG_BUFFER_UAV(RWStructuredBuffer<uint>, RWTopLevelGridBuffer)
+		SHADER_PARAMETER_RDG_BUFFER_UAV(RWStructuredBuffer<FScalarGridData>, RWExtinctionGridBuffer)
+		SHADER_PARAMETER_RDG_BUFFER_UAV(RWStructuredBuffer<FVectorGridData>, RWEmissionGridBuffer)
+		SHADER_PARAMETER_RDG_BUFFER_UAV(RWStructuredBuffer<FVectorGridData>, RWScatteringGridBuffer)
+
+		// Indirection Grid
+		SHADER_PARAMETER(int, IndirectionGridBufferSize)
+		SHADER_PARAMETER(int, FixedBottomLevelResolution)
+		SHADER_PARAMETER_RDG_BUFFER_UAV(RWStructuredBuffer<FTopLevelGridData>, RWIndirectionGridBuffer)
+		SHADER_PARAMETER_RDG_BUFFER_UAV(RWBuffer<uint>, RWIndirectionGridAllocatorBuffer)
+
+		// Hash Table
+		SHADER_PARAMETER_RDG_BUFFER_UAV(RWBuffer<uint>, RWHashTable)
+		SHADER_PARAMETER_RDG_BUFFER_UAV(RWBuffer<uint>, RWHashToVoxelBuffer)
+		SHADER_PARAMETER(int, HashTableSize)
+
+		// Homogeneous aggregation
+		SHADER_PARAMETER(float, HomogeneousThreshold)
+	END_SHADER_PARAMETER_STRUCT()
+	
+	FRasterizeBottomLevelGrid() = default;
+
+	FRasterizeBottomLevelGrid(
+		const FMeshMaterialShaderType::CompiledShaderInitializerType & Initializer
+	)
+		: FMeshMaterialShader(Initializer)
+	{
+		Bindings.BindForLegacyShaderParameters(
+			this,
+			Initializer.PermutationId,
+			Initializer.ParameterMap,
+			*FParameters::FTypeInfo::GetStructMetadata(),
+			// Don't require full bindings, we use FMaterialShader::SetParameters
+			false
+		);
+	}
+
+	static bool ShouldCompilePermutation(
+		const FMaterialShaderPermutationParameters& Parameters
+	)
+	{
+		return DoesPlatformSupportHeterogeneousVolumes(Parameters.Platform)
+			&& DoesMaterialShaderSupportHeterogeneousVolumes(Parameters.MaterialParameters);
+	}
+
+	static FPermutationDomain RemapPermutation(FPermutationDomain PermutationVector)
+	{
+		return PermutationVector;
+	}
+
+	static void ModifyCompilationEnvironment(
+		const FMaterialShaderPermutationParameters& Parameters,
+		FShaderCompilerEnvironment& OutEnvironment
+	)
+	{
+		FMaterialShader::ModifyCompilationEnvironment(Parameters, OutEnvironment);
+		OutEnvironment.SetDefine(TEXT("THREADGROUP_SIZE_1D"), GetThreadGroupSize1D());
+		OutEnvironment.SetDefine(TEXT("THREADGROUP_SIZE_3D"), GetThreadGroupSize3D());
+
+		// This shader takes a very long time to compile with FXC, so we pre-compile it with DXC first and then forward the optimized HLSL to FXC.
+		OutEnvironment.CompilerFlags.Add(CFLAG_PrecompileWithDXC);
+		OutEnvironment.CompilerFlags.Add(CFLAG_AllowTypedUAVLoads);
+
+		OutEnvironment.SetDefine(TEXT("GET_PRIMITIVE_DATA_OVERRIDE"), 1);
+	}
+
+	void SetParameters(
+		FRHIComputeCommandList& RHICmdList,
+		FRHIComputeShader* ShaderRHI,
+		const FViewInfo& View,
+		const FMaterialRenderProxy* MaterialProxy,
+		const FMaterial& Material
+	)
+	{
+		FMaterialShader::SetViewParameters(RHICmdList, ShaderRHI, View, View.ViewUniformBuffer);
+		FMaterialShader::SetParameters(RHICmdList, ShaderRHI, MaterialProxy, Material, View);
+	}
+
+	static int32 GetThreadGroupSize1D() { return GetThreadGroupSize3D() * GetThreadGroupSize3D() * GetThreadGroupSize3D(); }
+	static int32 GetThreadGroupSize3D() { return 4; }
+};
+
+IMPLEMENT_MATERIAL_SHADER_TYPE(, FRasterizeBottomLevelGrid, TEXT("/Engine/Private/HeterogeneousVolumes/HeterogeneousVolumesMaterialBakingPipeline.usf"), TEXT("RasterizeBottomLevelGridCS"), SF_Compute);
+
+class FTopLevelCreateBitmaskCS : public FGlobalShader
+{
+	DECLARE_GLOBAL_SHADER(FTopLevelCreateBitmaskCS);
+	SHADER_USE_PARAMETER_STRUCT(FTopLevelCreateBitmaskCS, FGlobalShader);
+
+	using FPermutationDomain = TShaderPermutationDomain<>;
+
+	BEGIN_SHADER_PARAMETER_STRUCT(FParameters, )
+		SHADER_PARAMETER(FIntVector, TopLevelGridResolution)
+
+		// Grid data
+		SHADER_PARAMETER_RDG_BUFFER_SRV(StructuredBuffer<FTopLevelGridData>, TopLevelGridBuffer)
+
+		// Output
+		SHADER_PARAMETER_RDG_BUFFER_UAV(RWStructuredBuffer<FTopLevelGridBitmaskData>, RWTopLevelGridBitmaskBuffer)
+	END_SHADER_PARAMETER_STRUCT()
+
+	static bool ShouldCompilePermutation(
+		const FGlobalShaderPermutationParameters& Parameters
+	)
+	{
+		return DoesPlatformSupportHeterogeneousVolumes(Parameters.Platform);
+	}
+
+	static void ModifyCompilationEnvironment(
+		const FGlobalShaderPermutationParameters& Parameters,
+		FShaderCompilerEnvironment& OutEnvironment
+	)
+	{
+		FMaterialShader::ModifyCompilationEnvironment(Parameters, OutEnvironment);
+		OutEnvironment.SetDefine(TEXT("THREADGROUP_SIZE_1D"), GetThreadGroupSize1D());
+		OutEnvironment.SetDefine(TEXT("THREADGROUP_SIZE_3D"), GetThreadGroupSize3D());
+	}
+
+	static int32 GetThreadGroupSize1D() { return GetThreadGroupSize3D() * GetThreadGroupSize3D() * GetThreadGroupSize3D(); }
+	static int32 GetThreadGroupSize3D() { return 4; }
+};
+
+IMPLEMENT_MATERIAL_SHADER_TYPE(, FTopLevelCreateBitmaskCS, TEXT("/Engine/Private/HeterogeneousVolumes/HeterogeneousVolumesVoxelGridPipeline.usf"), TEXT("TopLevelCreateBitmaskCS"), SF_Compute);
+
+void CreateTopLevelBitmask(
+	FRDGBuilder& GraphBuilder,
+	TArray<FViewInfo>& Views,
+	FIntVector TopLevelGridResolution,
+	FRDGBufferRef TopLevelGridBuffer,
+	FRDGBufferRef& TopLevelGridBitmaskBuffer
+)
+{
+	if (!HeterogeneousVolumes::EnableTopLevelBitmask())
+	{
+		TopLevelGridBitmaskBuffer = GSystemTextures.GetDefaultStructuredBuffer(GraphBuilder, sizeof(FTopLevelGridBitmaskData));
+	}
+	else
+	{
+		FIntVector BitmaskGridResolution = FIntVector(
+			FMath::DivideAndRoundUp(TopLevelGridResolution.X, 4),
+			FMath::DivideAndRoundUp(TopLevelGridResolution.Y, 4),
+			FMath::DivideAndRoundUp(TopLevelGridResolution.Z, 4)
+		);
+		uint32 TopLevelGridBitmaskBufferSize = BitmaskGridResolution.X * BitmaskGridResolution.Y * BitmaskGridResolution.Z;
+
+		TopLevelGridBitmaskBuffer = GraphBuilder.CreateBuffer(
+			FRDGBufferDesc::CreateStructuredDesc(sizeof(FTopLevelGridBitmaskData), TopLevelGridBitmaskBufferSize),
+			TEXT("HeterogeneousVolumes.BottomLevelGridBuffer")
+		);
+		AddClearUAVPass(GraphBuilder, GraphBuilder.CreateUAV(TopLevelGridBitmaskBuffer), 0x0);
+
+		FTopLevelCreateBitmaskCS::FParameters* PassParameters = GraphBuilder.AllocParameters<FTopLevelCreateBitmaskCS::FParameters>();
+		{
+			PassParameters->TopLevelGridResolution = TopLevelGridResolution;
+			PassParameters->TopLevelGridBuffer = GraphBuilder.CreateSRV(TopLevelGridBuffer);
+			PassParameters->RWTopLevelGridBitmaskBuffer = GraphBuilder.CreateUAV(TopLevelGridBitmaskBuffer);
+		}
+
+		//FTopLevelCreateBitmaskCS::FPermutationDomain PermutationVector;
+		const FGlobalShaderMap* GlobalShaderMap = Views[0].ShaderMap;
+		TShaderRef<FTopLevelCreateBitmaskCS> ComputeShader = GlobalShaderMap->GetShader<FTopLevelCreateBitmaskCS>();
+		FComputeShaderUtils::AddPass(
+			GraphBuilder,
+			RDG_EVENT_NAME("RenderTransmittanceTopLevelGridCS"),
+			ERDGPassFlags::Compute | ERDGPassFlags::NeverCull,
+			ComputeShader,
+			PassParameters,
+			BitmaskGridResolution
+		);
+	}
+}
+
+class FBuildMajorantVoxelGridCS : public FGlobalShader
+{
+	DECLARE_GLOBAL_SHADER(FBuildMajorantVoxelGridCS);
+	SHADER_USE_PARAMETER_STRUCT(FBuildMajorantVoxelGridCS, FGlobalShader);
+
+	class FEnableIndirectionGrid : SHADER_PERMUTATION_BOOL("DIM_ENABLE_INDIRECTION_GRID");
+	using FPermutationDomain = TShaderPermutationDomain<FEnableIndirectionGrid>;
+
+	BEGIN_SHADER_PARAMETER_STRUCT(FParameters, )
+		SHADER_PARAMETER(FIntVector, TopLevelGridResolution)
+
+		// Grid data
+		SHADER_PARAMETER_RDG_BUFFER_SRV(StructuredBuffer<FTopLevelGridData>, TopLevelGridBuffer)
+		SHADER_PARAMETER_RDG_BUFFER_SRV(StructuredBuffer<FTopLevelGridData>, IndirectionGridBuffer)
+		SHADER_PARAMETER_RDG_BUFFER_SRV(StructuredBuffer<FScalarGridData>, ExtinctionGridBuffer)
+
+		// Output
+		SHADER_PARAMETER_RDG_BUFFER_UAV(RWStructuredBuffer<FScalarGridData>, RWMajorantVoxelGridBuffer)
+	END_SHADER_PARAMETER_STRUCT()
+
+	static bool ShouldCompilePermutation(
+		const FGlobalShaderPermutationParameters& Parameters
+	)
+	{
+		return DoesPlatformSupportHeterogeneousVolumes(Parameters.Platform);
+	}
+
+	static void ModifyCompilationEnvironment(
+		const FGlobalShaderPermutationParameters& Parameters,
+		FShaderCompilerEnvironment& OutEnvironment
+	)
+	{
+		FMaterialShader::ModifyCompilationEnvironment(Parameters, OutEnvironment);
+		OutEnvironment.SetDefine(TEXT("THREADGROUP_SIZE_1D"), GetThreadGroupSize1D());
+		OutEnvironment.SetDefine(TEXT("THREADGROUP_SIZE_3D"), GetThreadGroupSize3D());
+	}
+
+	static int32 GetThreadGroupSize1D() { return GetThreadGroupSize3D() * GetThreadGroupSize3D() * GetThreadGroupSize3D(); }
+	static int32 GetThreadGroupSize3D() { return 4; }
+};
+
+IMPLEMENT_MATERIAL_SHADER_TYPE(, FBuildMajorantVoxelGridCS, TEXT("/Engine/Private/HeterogeneousVolumes/HeterogeneousVolumesVoxelGridPipeline.usf"), TEXT("BuildMajorantVoxelGridCS"), SF_Compute);
+
+void BuildMajorantVoxelGrid(
+	FRDGBuilder& GraphBuilder,
+	const FScene* Scene,
+	const FIntVector& TopLevelGridResolution,
+	FRDGBufferRef TopLevelGridBuffer,
+	FRDGBufferRef IndirectionGridBuffer,
+	FRDGBufferRef ExtinctionGridBuffer,
+	FRDGBufferRef& MajorantVoxelGridBuffer
+)
+{
+	if (!HeterogeneousVolumes::EnableMajorantGrid())
+	{
+		MajorantVoxelGridBuffer = GSystemTextures.GetDefaultStructuredBuffer(GraphBuilder, sizeof(FScalarGridData));
+	}
+	else
+	{
+		uint32 MajorantVoxelGridBufferSize = TopLevelGridResolution.X * TopLevelGridResolution.Y * TopLevelGridResolution.Z;
+		MajorantVoxelGridBuffer = GraphBuilder.CreateBuffer(
+			FRDGBufferDesc::CreateStructuredDesc(sizeof(FScalarGridData), MajorantVoxelGridBufferSize),
+			TEXT("HeterogeneousVolumes.MajorantVoxelGridBuffer")
+		);
+
+		FBuildMajorantVoxelGridCS::FParameters* PassParameters = GraphBuilder.AllocParameters<FBuildMajorantVoxelGridCS::FParameters>();
+		{
+			PassParameters->TopLevelGridResolution = TopLevelGridResolution;
+			PassParameters->TopLevelGridBuffer = GraphBuilder.CreateSRV(TopLevelGridBuffer);
+			PassParameters->IndirectionGridBuffer = GraphBuilder.CreateSRV(IndirectionGridBuffer);
+			PassParameters->ExtinctionGridBuffer = GraphBuilder.CreateSRV(ExtinctionGridBuffer);
+			PassParameters->RWMajorantVoxelGridBuffer = GraphBuilder.CreateUAV(MajorantVoxelGridBuffer);
+		}
+
+		const FGlobalShaderMap* GlobalShaderMap = GetGlobalShaderMap(Scene->GetFeatureLevel());
+
+		FBuildMajorantVoxelGridCS::FPermutationDomain PermutationVector;
+		PermutationVector.Set<FBuildMajorantVoxelGridCS::FEnableIndirectionGrid>(HeterogeneousVolumes::EnableIndirectionGrid());
+		TShaderRef<FBuildMajorantVoxelGridCS> ComputeShader = GlobalShaderMap->GetShader<FBuildMajorantVoxelGridCS>(PermutationVector);
+
+		FComputeShaderUtils::AddPass(
+			GraphBuilder,
+			RDG_EVENT_NAME("RenderTransmittanceTopLevelGridCS"),
+			ERDGPassFlags::Compute | ERDGPassFlags::NeverCull,
+			ComputeShader,
+			PassParameters,
+			TopLevelGridResolution
+		);
+	}
+}
+
+class FRenderTransmittanceWithVoxelGridCS : public FGlobalShader
+{
+	DECLARE_GLOBAL_SHADER(FRenderTransmittanceWithVoxelGridCS);
+	SHADER_USE_PARAMETER_STRUCT(FRenderTransmittanceWithVoxelGridCS, FGlobalShader);
+
+	using FPermutationDomain = TShaderPermutationDomain<>;
+
+	BEGIN_SHADER_PARAMETER_STRUCT(FParameters, )
+		// Scene data
+		SHADER_PARAMETER_STRUCT_REF(FViewUniformShaderParameters, View)
+		SHADER_PARAMETER_STRUCT_INCLUDE(FSceneTextureParameters, SceneTextures)
+
+		// Ray data
+		SHADER_PARAMETER(int, bJitter)
+		SHADER_PARAMETER(float, MaxTraceDistance)
+		SHADER_PARAMETER(float, StepSize)
+		SHADER_PARAMETER(int, MaxStepCount)
+
+		// Voxel grids
+		SHADER_PARAMETER_RDG_UNIFORM_BUFFER(FOrthoVoxelGridUniformBufferParameters, OrthoGridUniformBuffer)
+		SHADER_PARAMETER_RDG_UNIFORM_BUFFER(FFrustumVoxelGridUniformBufferParameters, FrustumGridUniformBuffer)
+
+		// Dispatch data
+		SHADER_PARAMETER(FIntVector, GroupCount)
+
+		// Output
+		SHADER_PARAMETER_RDG_TEXTURE_UAV(RWTexture2D<float4>, RWLightingTexture)
+	END_SHADER_PARAMETER_STRUCT()
+
+	static bool ShouldCompilePermutation(
+		const FGlobalShaderPermutationParameters& Parameters
+	)
+	{
+		return DoesPlatformSupportHeterogeneousVolumes(Parameters.Platform);
+	}
+
+	static void ModifyCompilationEnvironment(
+		const FGlobalShaderPermutationParameters& Parameters,
+		FShaderCompilerEnvironment& OutEnvironment
+	)
+	{
+		FMaterialShader::ModifyCompilationEnvironment(Parameters, OutEnvironment);
+		OutEnvironment.SetDefine(TEXT("THREADGROUP_SIZE_1D"), GetThreadGroupSize1D());
+		OutEnvironment.SetDefine(TEXT("THREADGROUP_SIZE_2D"), GetThreadGroupSize2D());
+		OutEnvironment.SetDefine(TEXT("THREADGROUP_SIZE_3D"), GetThreadGroupSize3D());
+
+		OutEnvironment.CompilerFlags.Add(CFLAG_AllowTypedUAVLoads);
+
+		OutEnvironment.SetDefine(TEXT("ENABLE_ADAPTIVE_VOXEL_GRID_UNIFORM_BUFFER"), 1);
+	}
+
+	static int32 GetThreadGroupSize1D() { return GetThreadGroupSize3D() * GetThreadGroupSize3D() * GetThreadGroupSize3D(); }
+	static int32 GetThreadGroupSize2D() { return 8; }
+	static int32 GetThreadGroupSize3D() { return 4; }
+};
+
+IMPLEMENT_MATERIAL_SHADER_TYPE(, FRenderTransmittanceWithVoxelGridCS, TEXT("/Engine/Private/HeterogeneousVolumes/HeterogeneousVolumesVoxelGridPipeline.usf"), TEXT("RenderTransmittanceWithVoxelGridCS"), SF_Compute);
+
+void CalcViewBoundsAndMinimumVoxelSize(
+	FRDGBuilder& GraphBuilder,
+	const FViewInfo& View,
+	FBoxSphereBounds& TopLevelGridBounds,
+	float& MinimumVoxelSize
+)
+{
+	TopLevelGridBounds = FBoxSphereBounds(ForceInit);
+	MinimumVoxelSize = HeterogeneousVolumes::GetMinimumVoxelSizeOutsideFrustum();
+
+	// Build view bounds
+	FVector WorldCameraOrigin = View.ViewMatrices.GetViewOrigin();
+	FBoxSphereBounds WorldCameraBounds(FSphere(WorldCameraOrigin, HeterogeneousVolumes::GetMaxTraceDistance()));
+
+	float TanHalfFOV = FMath::Tan(View.ViewMatrices.ComputeHalfFieldOfViewPerAxis().X);
+	int32 HalfWidth = View.ViewRect.Width() * 0.5;
+	float PixelWidth = TanHalfFOV / HalfWidth;
+
+	for (int32 MeshBatchIndex = 0; MeshBatchIndex < View.HeterogeneousVolumesMeshBatches.Num(); ++MeshBatchIndex)
+	{
+		// Only Niagara mesh particles bound to volume materials
+		const FMeshBatch* Mesh = View.HeterogeneousVolumesMeshBatches[MeshBatchIndex].Mesh;
+		const FPrimitiveSceneProxy* PrimitiveSceneProxy = View.HeterogeneousVolumesMeshBatches[MeshBatchIndex].Proxy;
+		if (!ShouldRenderMeshBatchWithHeterogeneousVolumes(Mesh, PrimitiveSceneProxy, View.GetFeatureLevel()))
+		{
+			continue;
+		}
+
+		// Only incorporate the primitive if it intersects with the canera bounding sphere where radius=MaxTraceDistance
+		const FBoxSphereBounds& PrimitiveBounds = PrimitiveSceneProxy->GetBounds();
+		if (View.ViewFrustum.IntersectBox(PrimitiveBounds.Origin, PrimitiveBounds.BoxExtent))
+		{
+			TopLevelGridBounds = Union(TopLevelGridBounds, PrimitiveBounds);
+
+			if (View.ViewFrustum.IntersectBox(TopLevelGridBounds.Origin, TopLevelGridBounds.BoxExtent))
+			{
+				const IHeterogeneousVolumeInterface* HeterogeneousVolumeInterface = HeterogeneousVolumes::GetInterface(PrimitiveSceneProxy);
+
+				// Bandlimit minimum voxel size request with projected voxel size, based on shading rate
+				FVector VoxelCenter = PrimitiveBounds.Origin;
+				float Distance = FMath::Max(FVector(PrimitiveBounds.Origin - WorldCameraOrigin).Length() - TopLevelGridBounds.BoxExtent.Length(), 0.0);
+				float VoxelWidth = Distance * PixelWidth * HeterogeneousVolumes::GetShadingRateForFrustumGrid();
+
+				float PerVolumeMinimumVoxelSize = FMath::Max(VoxelWidth, HeterogeneousVolumeInterface->GetMinimumVoxelSize());
+				MinimumVoxelSize = FMath::Min(PerVolumeMinimumVoxelSize, MinimumVoxelSize);
+			}
+		}
+	}
+}
+
+void CalcGlobalBoundsAndMinimumVoxelSize(
+	FRDGBuilder& GraphBuilder,
+	const TArray<FViewInfo>& Views,
+	FBoxSphereBounds& TopLevelGridBounds,
+	float& GlobalMinimumVoxelSize
+)
+{
+	TopLevelGridBounds = FBoxSphereBounds(ForceInit);
+	GlobalMinimumVoxelSize = HeterogeneousVolumes::GetMinimumVoxelSizeOutsideFrustum();
+
+	// Cycle through all Volume PrimitiveSceneProxies to collect bounds information and minimum voxel-size
+	for (int32 ViewIndex = 0; ViewIndex < Views.Num(); ViewIndex++)
+	{
+		float ViewMinimumVoxelSize = HeterogeneousVolumes::GetMinimumVoxelSizeOutsideFrustum();
+		FBoxSphereBounds AggregatePrimitiveBounds = FBoxSphereBounds(ForceInit);
+
+		// Build view bounds
+		const FViewInfo& View = Views[ViewIndex];
+		FVector WorldCameraOrigin = View.ViewMatrices.GetViewOrigin();
+		FBoxSphereBounds WorldCameraBounds(FSphere(WorldCameraOrigin, HeterogeneousVolumes::GetMaxTraceDistance()));
+
+		float TanHalfFOV = FMath::Tan(View.ViewMatrices.ComputeHalfFieldOfViewPerAxis().X);
+		int32 HalfWidth = View.ViewRect.Width() * 0.5;
+		float PixelWidth = TanHalfFOV / HalfWidth;
+
+		for (int32 MeshBatchIndex = 0; MeshBatchIndex < View.HeterogeneousVolumesMeshBatches.Num(); ++MeshBatchIndex)
+		{
+			// Only Niagara mesh particles bound to volume materials
+			const FMeshBatch* Mesh = View.HeterogeneousVolumesMeshBatches[MeshBatchIndex].Mesh;
+			const FPrimitiveSceneProxy* PrimitiveSceneProxy = View.HeterogeneousVolumesMeshBatches[MeshBatchIndex].Proxy;
+			if (!ShouldRenderMeshBatchWithHeterogeneousVolumes(Mesh, PrimitiveSceneProxy, View.GetFeatureLevel()))
+			{
+				continue;
+			}
+
+			// Only incorporate the primitive if it intersects with the canera bounding sphere where radius=MaxTraceDistance
+			const FBoxSphereBounds& PrimitiveBounds = PrimitiveSceneProxy->GetBounds();
+			if (View.ViewFrustum.IntersectBox(PrimitiveBounds.Origin, PrimitiveBounds.BoxExtent))
+			{
+				AggregatePrimitiveBounds = Union(AggregatePrimitiveBounds, PrimitiveBounds);
+
+				if (View.ViewFrustum.IntersectBox(AggregatePrimitiveBounds.Origin, AggregatePrimitiveBounds.BoxExtent))
+				{
+					const IHeterogeneousVolumeInterface* HeterogeneousVolumeInterface = HeterogeneousVolumes::GetInterface(PrimitiveSceneProxy);
+
+					// Bandlimit minimum voxel size request with projected voxel size, based on shading rate
+					FVector VoxelCenter = PrimitiveBounds.Origin;
+					float Distance = FMath::Max(FVector(PrimitiveBounds.Origin - WorldCameraOrigin).Length() - AggregatePrimitiveBounds.BoxExtent.Length(), 0.0);
+					float VoxelWidth = Distance * PixelWidth * HeterogeneousVolumes::GetShadingRateForOrthoGrid();
+
+					float PerVolumeMinimumVoxelSize = FMath::Max(VoxelWidth, HeterogeneousVolumeInterface->GetMinimumVoxelSize());
+					ViewMinimumVoxelSize = FMath::Min(PerVolumeMinimumVoxelSize, ViewMinimumVoxelSize);
+				}
+			}
+			// TODO: Out-of-frustum minimum voxel size per-primitive?
+			// else if (FBoxSphereBounds::BoxesIntersect(WorldCameraBounds, PrimitiveBounds))
+		}
+
+		// Clamp per-view minimum voxel-size to in-frustum maximum
+		if (View.ViewFrustum.IntersectBox(AggregatePrimitiveBounds.Origin, AggregatePrimitiveBounds.BoxExtent))
+		{
+			ViewMinimumVoxelSize = FMath::Max(ViewMinimumVoxelSize, HeterogeneousVolumes::GetMinimumVoxelSizeInFrustum());
+		}
+
+		// Update global bounds and minimum voxel-size
+		TopLevelGridBounds = Union(TopLevelGridBounds, AggregatePrimitiveBounds);
+		GlobalMinimumVoxelSize = FMath::Min(GlobalMinimumVoxelSize, ViewMinimumVoxelSize);
+	}
+}
+
+void ExtractFrustumVoxelGridUniformBuffer(
+	FRDGBuilder& GraphBuilder,
+	const TRDGUniformBufferRef<FFrustumVoxelGridUniformBufferParameters>& FrustumGridUniformBuffer,
+	FAdaptiveFrustumGridParameterCache& AdaptiveFrustumGridParameterCache
+)
+{
+	const TRDGParameterStruct<FFrustumVoxelGridUniformBufferParameters>& Parameters = FrustumGridUniformBuffer->GetParameters();
+
+	AdaptiveFrustumGridParameterCache.TopLevelGridWorldBoundsMin = Parameters->TopLevelGridWorldBoundsMin;
+	AdaptiveFrustumGridParameterCache.TopLevelGridWorldBoundsMax = Parameters->TopLevelGridWorldBoundsMax;
+	AdaptiveFrustumGridParameterCache.TopLevelGridResolution = Parameters->TopLevelFroxelGridResolution;
+	AdaptiveFrustumGridParameterCache.VoxelDimensions = Parameters->VoxelDimensions;
+	AdaptiveFrustumGridParameterCache.bUseFrustumGrid = Parameters->bUseFrustumGrid;
+	AdaptiveFrustumGridParameterCache.NearPlaneDepth = Parameters->NearPlaneDepth;
+	AdaptiveFrustumGridParameterCache.FarPlaneDepth = Parameters->FarPlaneDepth;
+
+
+	AdaptiveFrustumGridParameterCache.WorldToClip = Parameters->WorldToClip;
+	AdaptiveFrustumGridParameterCache.ClipToWorld = Parameters->ClipToWorld;
+	AdaptiveFrustumGridParameterCache.WorldToView = Parameters->WorldToView;
+	AdaptiveFrustumGridParameterCache.ViewToWorld = Parameters->ViewToWorld;
+	AdaptiveFrustumGridParameterCache.ViewToClip = Parameters->ViewToClip;
+	AdaptiveFrustumGridParameterCache.ClipToView = Parameters->ClipToView;
+
+	for (int i = 0; i < 6; ++i)
+	{
+		AdaptiveFrustumGridParameterCache.ViewFrustumPlanes[i] = Parameters->ViewFrustumPlanes[i];
+	}
+
+	GraphBuilder.QueueBufferExtraction(Parameters->TopLevelFroxelGridBuffer->GetParent(), &AdaptiveFrustumGridParameterCache.TopLevelGridBuffer);
+	GraphBuilder.QueueBufferExtraction(Parameters->ExtinctionFroxelGridBuffer->GetParent(), &AdaptiveFrustumGridParameterCache.ExtinctionGridBuffer);
+	GraphBuilder.QueueBufferExtraction(Parameters->EmissionFroxelGridBuffer->GetParent(), &AdaptiveFrustumGridParameterCache.EmissionGridBuffer);
+	GraphBuilder.QueueBufferExtraction(Parameters->ScatteringFroxelGridBuffer->GetParent(), &AdaptiveFrustumGridParameterCache.ScatteringGridBuffer);
+}
+
+void RegisterExternalFrustumVoxelGridUniformBuffer(
+	FRDGBuilder& GraphBuilder,
+	const FAdaptiveFrustumGridParameterCache& AdaptiveFrustumGridParameterCache,
+	TRDGUniformBufferRef<FFrustumVoxelGridUniformBufferParameters>& FrustumGridUniformBuffer
+)
+{
+	FFrustumVoxelGridUniformBufferParameters* UniformBufferParameters = GraphBuilder.AllocParameters<FFrustumVoxelGridUniformBufferParameters>();
+	{
+		UniformBufferParameters->WorldToClip = AdaptiveFrustumGridParameterCache.WorldToClip;
+		UniformBufferParameters->ClipToWorld = AdaptiveFrustumGridParameterCache.ClipToWorld;
+
+		UniformBufferParameters->WorldToView = AdaptiveFrustumGridParameterCache.WorldToView;
+		UniformBufferParameters->ViewToWorld = AdaptiveFrustumGridParameterCache.ViewToWorld;
+
+		UniformBufferParameters->ViewToClip = AdaptiveFrustumGridParameterCache.ViewToClip;
+		UniformBufferParameters->ClipToView = AdaptiveFrustumGridParameterCache.ClipToView;
+
+		UniformBufferParameters->TopLevelGridWorldBoundsMin = AdaptiveFrustumGridParameterCache.TopLevelGridWorldBoundsMin;
+		UniformBufferParameters->TopLevelGridWorldBoundsMax = AdaptiveFrustumGridParameterCache.TopLevelGridWorldBoundsMax;
+		UniformBufferParameters->TopLevelFroxelGridResolution = AdaptiveFrustumGridParameterCache.TopLevelGridResolution;
+		UniformBufferParameters->VoxelDimensions = AdaptiveFrustumGridParameterCache.TopLevelGridResolution;
+		UniformBufferParameters->bUseFrustumGrid = AdaptiveFrustumGridParameterCache.bUseFrustumGrid;
+		UniformBufferParameters->NearPlaneDepth = AdaptiveFrustumGridParameterCache.NearPlaneDepth;
+		UniformBufferParameters->FarPlaneDepth = AdaptiveFrustumGridParameterCache.FarPlaneDepth;
+
+		// Frustum assignment
+		for (int i = 0; i < 6; ++i)
+		{
+			UniformBufferParameters->ViewFrustumPlanes[i] = AdaptiveFrustumGridParameterCache.ViewFrustumPlanes[i];
+		}
+
+		UniformBufferParameters->TopLevelFroxelGridBuffer = GraphBuilder.CreateSRV(GraphBuilder.RegisterExternalBuffer(AdaptiveFrustumGridParameterCache.TopLevelGridBuffer));
+		UniformBufferParameters->ExtinctionFroxelGridBuffer = GraphBuilder.CreateSRV(GraphBuilder.RegisterExternalBuffer(AdaptiveFrustumGridParameterCache.ExtinctionGridBuffer));
+		UniformBufferParameters->EmissionFroxelGridBuffer = GraphBuilder.CreateSRV(GraphBuilder.RegisterExternalBuffer(AdaptiveFrustumGridParameterCache.EmissionGridBuffer));
+		UniformBufferParameters->ScatteringFroxelGridBuffer = GraphBuilder.CreateSRV(GraphBuilder.RegisterExternalBuffer(AdaptiveFrustumGridParameterCache.ScatteringGridBuffer));
+	}
+	FrustumGridUniformBuffer = GraphBuilder.CreateUniformBuffer(UniformBufferParameters);
+}
+
+void CreateEmptyFrustumVoxelGridUniformBuffer(
+	FRDGBuilder& GraphBuilder,
+	TRDGUniformBufferRef<FFrustumVoxelGridUniformBufferParameters>& FrustumGridUniformBuffer
+)
+{
+	FFrustumVoxelGridUniformBufferParameters* UniformBufferParameters = GraphBuilder.AllocParameters<FFrustumVoxelGridUniformBufferParameters>();
+	{
+		UniformBufferParameters->WorldToClip = FMatrix44f::Identity;
+		UniformBufferParameters->ClipToWorld = FMatrix44f::Identity;
+		UniformBufferParameters->WorldToView = FMatrix44f::Identity;
+		UniformBufferParameters->ViewToWorld = FMatrix44f::Identity;
+		UniformBufferParameters->ViewToClip = FMatrix44f::Identity;
+		UniformBufferParameters->ClipToView = FMatrix44f::Identity;
+
+		UniformBufferParameters->TopLevelFroxelGridBuffer = GraphBuilder.CreateSRV(GSystemTextures.GetDefaultStructuredBuffer(GraphBuilder, sizeof(FTopLevelGridData)));
+		UniformBufferParameters->ExtinctionFroxelGridBuffer = GraphBuilder.CreateSRV(GSystemTextures.GetDefaultStructuredBuffer(GraphBuilder, sizeof(FScalarGridData)));
+		UniformBufferParameters->EmissionFroxelGridBuffer = GraphBuilder.CreateSRV(GSystemTextures.GetDefaultStructuredBuffer(GraphBuilder, sizeof(FVectorGridData)));
+		UniformBufferParameters->ScatteringFroxelGridBuffer = GraphBuilder.CreateSRV(GSystemTextures.GetDefaultStructuredBuffer(GraphBuilder, sizeof(FVectorGridData)));
+
+		UniformBufferParameters->TopLevelGridWorldBoundsMin = FVector3f(0);
+		UniformBufferParameters->TopLevelGridWorldBoundsMax = FVector3f(0);
+		UniformBufferParameters->TopLevelFroxelGridResolution = FIntVector(0);
+		UniformBufferParameters->VoxelDimensions = FIntVector(0);
+		UniformBufferParameters->bUseFrustumGrid = false;
+		UniformBufferParameters->NearPlaneDepth = 0.0;
+		UniformBufferParameters->FarPlaneDepth = 0.0;
+	}
+	FrustumGridUniformBuffer = GraphBuilder.CreateUniformBuffer(UniformBufferParameters);
+}
+
+void ClipNearFarDistances(const FViewInfo& View, const FBoxSphereBounds& TopLevelGridBounds, float& NearPlaneDistance, float& FarPlaneDistance)
+{
+	// Determine near and far planes, in world-space
+	float d = -FVector::DotProduct(View.GetViewDirection(), View.ViewLocation);
+
+	// Analyze the input volumes and determine the near/far extents
+	FVector Center = TopLevelGridBounds.GetSphere().Center;
+
+	// Project center onto the camera view-plane
+	float SignedDistance = FVector::DotProduct(View.GetViewDirection(), Center) + d;
+
+	float Radius = TopLevelGridBounds.GetSphere().W;
+	float NearDistance = SignedDistance - Radius;
+	float FarDistance = SignedDistance + Radius;
+
+	if (NearPlaneDistance < 0.0)
+	{
+		NearPlaneDistance = NearDistance;
+	}
+
+	if (FarPlaneDistance < 0.0)
+	{
+		FarPlaneDistance = FarDistance;
+	}
+
+	NearPlaneDistance = FMath::Max(NearPlaneDistance, 0.01);
+	FarPlaneDistance = FMath::Max(FarPlaneDistance, NearDistance + 1.0);
+}
+
+void CalculateTopLevelGridResolution(
+	FBoxSphereBounds TopLevelGridBounds,
+	float GlobalMinimumVoxelSize,
+	FIntVector& TopLevelGridResolution
+)
+{
+	// Bound Top-level grid resolution to cover fully allocated child grids at minimum voxel size
+	FIntVector CombinedChildGridResolution = FIntVector(HeterogeneousVolumes::GetBottomLevelGridResolution());
+	if (HeterogeneousVolumes::EnableIndirectionGrid())
+	{
+		CombinedChildGridResolution *= HeterogeneousVolumes::GetIndirectionGridResolution();
+	}
+
+	FVector TopLevelGridVoxelSize = FVector(CombinedChildGridResolution) * GlobalMinimumVoxelSize;
+	FVector TopLevelGridResolutionAsFloat = (TopLevelGridBounds.BoxExtent * 2.0) / TopLevelGridVoxelSize;
+
+	TopLevelGridResolution.X = FMath::CeilToInt(TopLevelGridResolutionAsFloat.X);
+	TopLevelGridResolution.Y = FMath::CeilToInt(TopLevelGridResolutionAsFloat.Y);
+	TopLevelGridResolution.Z = FMath::CeilToInt(TopLevelGridResolutionAsFloat.Z);
+
+	// Clamp to a moderate limit to also handle indirection grid allocation
+	TopLevelGridResolution.X = FMath::Clamp(TopLevelGridResolution.X, 1, 512);
+	TopLevelGridResolution.Y = FMath::Clamp(TopLevelGridResolution.Y, 1, 512);
+	TopLevelGridResolution.Z = FMath::Clamp(TopLevelGridResolution.Z, 1, 512);
+}
+
+void CalculateVoxelSize(
+	FRDGBuilder& GraphBuilder,
+	const FScene* Scene,
+	/*const*/ TArray<FViewInfo>& Views,
+	FBoxSphereBounds TopLevelGridBounds,
+	FIntVector TopLevelGridResolution,
+	FRDGBufferRef& TopLevelGridBuffer
+)
+{
+	TopLevelGridBuffer = GraphBuilder.CreateBuffer(
+		FRDGBufferDesc::CreateStructuredDesc(sizeof(FTopLevelGridData), TopLevelGridResolution.X * TopLevelGridResolution.Y * TopLevelGridResolution.Z),
+		TEXT("HeterogeneousVolumes.TopLevelGridBuffer")
+	);
+	AddClearUAVPass(GraphBuilder, GraphBuilder.CreateUAV(TopLevelGridBuffer), 0xFFFFFFF8);
+
+	for (int32 ViewIndex = 0; ViewIndex < Views.Num(); ViewIndex++)
+	{
+		// Build view
+		const FViewInfo& View = Views[ViewIndex];
+		FVector WorldCameraOrigin = View.ViewMatrices.GetViewOrigin();
+		FBoxSphereBounds WorldCameraBounds(FSphere(WorldCameraOrigin, HeterogeneousVolumes::GetMaxTraceDistance()));
+
+		for (int32 MeshBatchIndex = 0; MeshBatchIndex < View.HeterogeneousVolumesMeshBatches.Num(); ++MeshBatchIndex)
+		{
+			const FPrimitiveSceneProxy* PrimitiveSceneProxy = View.HeterogeneousVolumesMeshBatches[MeshBatchIndex].Proxy;
+			const FBoxSphereBounds& PrimitiveBounds = PrimitiveSceneProxy->GetBounds();
+			const IHeterogeneousVolumeInterface* HeterogeneousVolumeInterface = HeterogeneousVolumes::GetInterface(PrimitiveSceneProxy);
+
+			FTopLevelGridCalculateVoxelSize::FParameters* PassParameters = GraphBuilder.AllocParameters<FTopLevelGridCalculateVoxelSize::FParameters>();
+			{
+				PassParameters->View = View.ViewUniformBuffer;
+
+				PassParameters->TopLevelGridResolution = TopLevelGridResolution;
+				PassParameters->TopLevelGridWorldBoundsMin = FVector3f(TopLevelGridBounds.Origin - TopLevelGridBounds.BoxExtent);
+				PassParameters->TopLevelGridWorldBoundsMax = FVector3f(TopLevelGridBounds.Origin + TopLevelGridBounds.BoxExtent);
+
+				PassParameters->PrimitiveWorldBoundsMin = FVector3f(PrimitiveBounds.Origin - PrimitiveBounds.BoxExtent);
+				PassParameters->PrimitiveWorldBoundsMax = FVector3f(PrimitiveBounds.Origin + PrimitiveBounds.BoxExtent);
+
+				PassParameters->ShadingRate = HeterogeneousVolumes::GetShadingRateForOrthoGrid();
+				PassParameters->MinVoxelSizeInFrustum = FMath::Max(HeterogeneousVolumeInterface->GetMinimumVoxelSize(), HeterogeneousVolumes::GetMinimumVoxelSizeInFrustum());
+				PassParameters->MinVoxelSizeOutOfFrustum = HeterogeneousVolumes::GetMinimumVoxelSizeOutsideFrustum();
+
+				PassParameters->RWTopLevelGridBuffer = GraphBuilder.CreateUAV(TopLevelGridBuffer);
+			}
+
+			FIntVector GroupCount;
+			GroupCount.X = FMath::DivideAndRoundUp(TopLevelGridResolution.X, FTopLevelGridCalculateVoxelSize::GetThreadGroupSize3D());
+			GroupCount.Y = FMath::DivideAndRoundUp(TopLevelGridResolution.Y, FTopLevelGridCalculateVoxelSize::GetThreadGroupSize3D());
+			GroupCount.Z = FMath::DivideAndRoundUp(TopLevelGridResolution.Z, FTopLevelGridCalculateVoxelSize::GetThreadGroupSize3D());
+
+			FTopLevelGridCalculateVoxelSize::FPermutationDomain PermutationVector;
+			TShaderRef<FTopLevelGridCalculateVoxelSize> ComputeShader = View.ShaderMap->GetShader<FTopLevelGridCalculateVoxelSize>(PermutationVector);
+			FComputeShaderUtils::AddPass(
+				GraphBuilder,
+				RDG_EVENT_NAME("TopLevelGridCalculateVoxelSize"),
+				ERDGPassFlags::Compute | ERDGPassFlags::NeverCull,
+				ComputeShader,
+				PassParameters,
+				GroupCount
+			);
+		}
+	}
+}
+
+void MarkTopLevelGrid(
+	FRDGBuilder& GraphBuilder,
+	const FScene* Scene,
+	FBoxSphereBounds TopLevelGridBounds,
+	FIntVector TopLevelGridResolution,
+	FRDGBufferRef& TopLevelGridBuffer,
+	FRDGBufferRef& IndirectionGridBuffer
+)
+{
+	int32 IndirectionGridBufferSize = TopLevelGridResolution.X * TopLevelGridResolution.Y * TopLevelGridResolution.Z * 16 * sizeof(FTopLevelGridData);
+	IndirectionGridBuffer = GraphBuilder.CreateBuffer(
+		FRDGBufferDesc::CreateStructuredDesc(sizeof(FTopLevelGridData), IndirectionGridBufferSize),
+		TEXT("HeterogeneousVolumes.OrthoGrid.IndirectionGridBuffer")
+	);
+	AddClearUAVPass(GraphBuilder, GraphBuilder.CreateUAV(IndirectionGridBuffer), 0xFFFFFFF8);
+
+	FRDGBufferRef IndirectionGridAllocatorBuffer = GraphBuilder.CreateBuffer(
+		FRDGBufferDesc::CreateBufferDesc(sizeof(uint32), 2),
+		TEXT("HeterogeneousVolume.OrthoGrid.IndirectionGridAllocatorBuffer")
+	);
+	AddClearUAVPass(GraphBuilder, GraphBuilder.CreateUAV(IndirectionGridAllocatorBuffer, PF_R32_UINT), 0);
+
+	FAllocateBottomLevelGrid::FParameters* PassParameters = GraphBuilder.AllocParameters<FAllocateBottomLevelGrid::FParameters>();
+	{
+		//PassParameters->View = View.ViewUniformBuffer;
+		PassParameters->TopLevelGridResolution = TopLevelGridResolution;
+		PassParameters->TopLevelGridWorldBoundsMin = FVector3f(TopLevelGridBounds.Origin - TopLevelGridBounds.BoxExtent);
+		PassParameters->TopLevelGridWorldBoundsMax = FVector3f(TopLevelGridBounds.Origin + TopLevelGridBounds.BoxExtent);
+
+		PassParameters->MaxVoxelResolution = HeterogeneousVolumes::GetBottomLevelGridResolution();
+		PassParameters->bSampleAtVertices = HeterogeneousVolumes::EnableLinearInterpolation();
+
+		PassParameters->RWTopLevelGridBuffer = GraphBuilder.CreateUAV(TopLevelGridBuffer);
+
+		// Indirection Grid
+		PassParameters->IndirectionGridBufferSize = IndirectionGridBufferSize;
+		PassParameters->RWIndirectionGridBuffer = GraphBuilder.CreateUAV(IndirectionGridBuffer);
+		PassParameters->RWIndirectionGridAllocatorBuffer = GraphBuilder.CreateUAV(IndirectionGridAllocatorBuffer, PF_R32_UINT);
+	}
+
+	FIntVector GroupCount;
+	GroupCount.X = FMath::DivideAndRoundUp(TopLevelGridResolution.X, FAllocateBottomLevelGrid::GetThreadGroupSize3D());
+	GroupCount.Y = FMath::DivideAndRoundUp(TopLevelGridResolution.Y, FAllocateBottomLevelGrid::GetThreadGroupSize3D());
+	GroupCount.Z = FMath::DivideAndRoundUp(TopLevelGridResolution.Z, FAllocateBottomLevelGrid::GetThreadGroupSize3D());
+
+	const FGlobalShaderMap* GlobalShaderMap = GetGlobalShaderMap(Scene->GetFeatureLevel());
+
+	FAllocateBottomLevelGrid::FPermutationDomain PermutationVector;
+	PermutationVector.Set<FAllocateBottomLevelGrid::FEnableIndirectionGrid>(HeterogeneousVolumes::EnableIndirectionGrid());
+	TShaderRef<FAllocateBottomLevelGrid> ComputeShader = GlobalShaderMap->GetShader<FAllocateBottomLevelGrid>(PermutationVector);
+	FComputeShaderUtils::AddPass(
+		GraphBuilder,
+		RDG_EVENT_NAME("AllocateBottomLevelGrid"),
+		ERDGPassFlags::Compute | ERDGPassFlags::NeverCull,
+		ComputeShader,
+		PassParameters,
+		GroupCount
+	);
+}
+
+void GenerateRasterTiles(
+	FRDGBuilder& GraphBuilder,
+	const FScene* Scene,
+	bool bEnableIndirectionGrid,
+	FIntVector TopLevelGridResolution,
+	FRDGBufferRef TopLevelGridBuffer,
+	FRDGBufferRef& RasterTileBuffer,
+	FRDGBufferRef& RasterTileAllocatorBuffer
+)
+{
+	const uint32 RasterTileVoxelResolution = HeterogeneousVolumes::GetBottomLevelGridResolution();
+	uint32 TileFactor = 64;
+	uint32 MaxNumRasterTiles = TileFactor * TopLevelGridResolution.X * TopLevelGridResolution.Y * TopLevelGridResolution.Z;
+	RasterTileBuffer = GraphBuilder.CreateBuffer(
+		FRDGBufferDesc::CreateStructuredDesc(sizeof(FRasterTileData), MaxNumRasterTiles),
+		TEXT("HeterogeneousVolumes.OrthoGrid.RasterTileBuffer")
+	);
+
+	RasterTileAllocatorBuffer = GraphBuilder.CreateBuffer(
+		FRDGBufferDesc::CreateBufferDesc(sizeof(uint32), 1),
+		TEXT("HeterogeneousVolume.OrthoGrid.RasterTileAllocatorBuffer")
+	);
+	AddClearUAVPass(GraphBuilder, GraphBuilder.CreateUAV(RasterTileAllocatorBuffer, PF_R32_UINT), 0);
+
+	{
+		FGenerateRasterTiles::FParameters* PassParameters = GraphBuilder.AllocParameters<FGenerateRasterTiles::FParameters>();
+		{
+			PassParameters->TopLevelGridResolution = TopLevelGridResolution;
+			PassParameters->TopLevelGridBuffer = GraphBuilder.CreateSRV(TopLevelGridBuffer);
+
+			PassParameters->RasterTileVoxelResolution = RasterTileVoxelResolution;
+			PassParameters->MaxNumRasterTiles = MaxNumRasterTiles;
+
+			PassParameters->RWRasterTileAllocatorBuffer = GraphBuilder.CreateUAV(RasterTileAllocatorBuffer, PF_R32_UINT);
+			PassParameters->RWRasterTileBuffer = GraphBuilder.CreateUAV(RasterTileBuffer);
+		}
+
+		FIntVector GroupCount;
+		GroupCount.X = FMath::DivideAndRoundUp(TopLevelGridResolution.X, FTopLevelGridCalculateVoxelSize::GetThreadGroupSize3D());
+		GroupCount.Y = FMath::DivideAndRoundUp(TopLevelGridResolution.Y, FTopLevelGridCalculateVoxelSize::GetThreadGroupSize3D());
+		GroupCount.Z = FMath::DivideAndRoundUp(TopLevelGridResolution.Z, FTopLevelGridCalculateVoxelSize::GetThreadGroupSize3D());
+
+		const FGlobalShaderMap* GlobalShaderMap = GetGlobalShaderMap(Scene->GetFeatureLevel());
+
+		FGenerateRasterTiles::FPermutationDomain PermutationVector;
+		PermutationVector.Set<FGenerateRasterTiles::FEnableIndirectionGrid>(bEnableIndirectionGrid);
+		TShaderRef<FGenerateRasterTiles> ComputeShader = GlobalShaderMap->GetShader<FGenerateRasterTiles>(PermutationVector);
+		FComputeShaderUtils::AddPass(
+			GraphBuilder,
+			RDG_EVENT_NAME("GenerateRasterTiles"),
+			ERDGPassFlags::Compute | ERDGPassFlags::NeverCull,
+			ComputeShader,
+			PassParameters,
+			GroupCount
+		);
+	}
+}
+
+void CalculateTopLevelGridResolutionForFrustumGrid(
+	const FViewInfo& View,
+	float MinimumVoxelSize,
+	float NearPlaneDistance,
+	float& FarPlaneDistance,
+	FIntVector& TopLevelGridResolution
+)
+{
+	// Determine top-level grid resolution
+	// Build a grid where the resolution is proportional to a bottom-level grid that is using all available memory
+	int32 BottomLevelGridResolution = HeterogeneousVolumes::GetBottomLevelGridResolution();
+
+	float ShadingRate = HeterogeneousVolumes::GetShadingRateForFrustumGrid();
+	int32 Width = FMath::CeilToInt32(View.ViewRect.Width() / ShadingRate);
+	int32 Height = FMath::CeilToInt32(View.ViewRect.Height() / ShadingRate);
+
+	// Use view frustum field-of-view and preferred minimum voxel size to find optimal far-plane distance
+	if (HeterogeneousVolumes::EnableOrthoVoxelGrid() &&
+		HeterogeneousVolumes::EnableFarPlaneAutoTransition() &&
+		ShadingRate >= HeterogeneousVolumes::GetShadingRateForOrthoGrid())
+	{
+		float GlobalMinimumVoxelSize = FMath::Max(MinimumVoxelSize, HeterogeneousVolumes::GetMinimumVoxelSizeInFrustum());
+
+		float Theta = FMath::DegreesToRadians(View.FOV * 0.5);
+		float PixelTheta = (2.0 * Theta) / Width;
+		float TanOfPixel = FMath::Tan(PixelTheta);
+		float OptimalFarPlaneDepth = GlobalMinimumVoxelSize / TanOfPixel;
+		FarPlaneDistance = FMath::Min(FarPlaneDistance, OptimalFarPlaneDepth);
+	}
+
+	// Depth slices should not be smaller than the declared minimum voxel size
+	int32 MaxDepth = FMath::CeilToInt32((FarPlaneDistance - NearPlaneDistance) / MinimumVoxelSize);
+	int32 Depth = FMath::Min(FMath::CeilToInt32(HeterogeneousVolumes::GetDepthSliceCountForFrustumGrid() / ShadingRate), MaxDepth);
+
+	TopLevelGridResolution = FIntVector(Width, Height, Depth);
+	TopLevelGridResolution.X = FMath::DivideAndRoundUp(TopLevelGridResolution.X, BottomLevelGridResolution);
+	TopLevelGridResolution.Y = FMath::DivideAndRoundUp(TopLevelGridResolution.Y, BottomLevelGridResolution);
+	TopLevelGridResolution.Z = FMath::DivideAndRoundUp(TopLevelGridResolution.Z, BottomLevelGridResolution);
+}
+
+void MarkTopLevelGridVoxelsForFrustumGrid(
+	FRDGBuilder& GraphBuilder,
+	const FViewInfo& View,
+	FIntVector TopLevelGridResolution,
+	float NearPlaneDistance,
+	float FarPlaneDistance,
+	const FMatrix& ViewToWorld,
+	FRDGBufferRef& TopLevelGridBuffer
+)
+{
+	int32 TopLevelVoxelCount = TopLevelGridResolution.X * TopLevelGridResolution.Y * TopLevelGridResolution.Z;
+
+	TopLevelGridBuffer = GraphBuilder.CreateBuffer(
+		FRDGBufferDesc::CreateStructuredDesc(sizeof(FTopLevelGridData), TopLevelVoxelCount),
+		TEXT("HeterogeneousVolumes.FrustumGrid.TopLevelGridBuffer")
+	);
+	AddClearUAVPass(GraphBuilder, GraphBuilder.CreateUAV(TopLevelGridBuffer), 0xFFFFFFF8);
+
+	for (int32 MeshBatchIndex = 0; MeshBatchIndex < View.HeterogeneousVolumesMeshBatches.Num(); ++MeshBatchIndex)
+	{
+		const FPrimitiveSceneProxy* PrimitiveSceneProxy = View.HeterogeneousVolumesMeshBatches[MeshBatchIndex].Proxy;
+		const FBoxSphereBounds& PrimitiveBounds = PrimitiveSceneProxy->GetBounds();
+		const IHeterogeneousVolumeInterface* HeterogeneousVolumeInterface = HeterogeneousVolumes::GetInterface(PrimitiveSceneProxy);
+
+		FMarkTopLevelGridVoxelsForFrustumGrid::FParameters* PassParameters = GraphBuilder.AllocParameters<FMarkTopLevelGridVoxelsForFrustumGrid::FParameters>();
+		{
+			PassParameters->View = View.ViewUniformBuffer;
+
+			PassParameters->PrimitiveWorldBoundsMin = FVector3f(PrimitiveBounds.Origin - PrimitiveBounds.BoxExtent);
+			PassParameters->PrimitiveWorldBoundsMax = FVector3f(PrimitiveBounds.Origin + PrimitiveBounds.BoxExtent);
+
+			PassParameters->ViewToWorld = FMatrix44f(ViewToWorld);
+			PassParameters->HalfFOV = FMath::DegreesToRadians(View.FOV * 0.5);
+			PassParameters->NearPlaneDepth = NearPlaneDistance;
+			PassParameters->FarPlaneDepth = FarPlaneDistance;
+
+			PassParameters->VoxelDimensions = TopLevelGridResolution;
+			PassParameters->TopLevelGridResolution = TopLevelGridResolution;
+
+			PassParameters->RWTopLevelGridBuffer = GraphBuilder.CreateUAV(TopLevelGridBuffer);
+		}
+
+		FIntVector GroupCount;
+		GroupCount.X = FMath::DivideAndRoundUp(TopLevelGridResolution.X, FMarkTopLevelGridVoxelsForFrustumGrid::GetThreadGroupSize3D());
+		GroupCount.Y = FMath::DivideAndRoundUp(TopLevelGridResolution.Y, FMarkTopLevelGridVoxelsForFrustumGrid::GetThreadGroupSize3D());
+		GroupCount.Z = FMath::DivideAndRoundUp(TopLevelGridResolution.Z, FMarkTopLevelGridVoxelsForFrustumGrid::GetThreadGroupSize3D());
+
+		FMarkTopLevelGridVoxelsForFrustumGrid::FPermutationDomain PermutationVector;
+		TShaderRef<FMarkTopLevelGridVoxelsForFrustumGrid> ComputeShader = View.ShaderMap->GetShader<FMarkTopLevelGridVoxelsForFrustumGrid>(PermutationVector);
+		FComputeShaderUtils::AddPass(
+			GraphBuilder,
+			RDG_EVENT_NAME("MarkTopLevelGridVoxelsForFrustumGrid"),
+			ERDGPassFlags::Compute | ERDGPassFlags::NeverCull,
+			ComputeShader,
+			PassParameters,
+			GroupCount
+		);
+	}
+}
+
+void RasterizeVolumesIntoFrustumVoxelGrid(
+	FRDGBuilder& GraphBuilder,
+	const FScene* Scene,
+	const FViewInfo& View,
+	// Transform data
+	FMatrix& ViewToWorld,
+	float NearPlaneDistance,
+	float FarPlaneDistance,
+	// Raster tile
+	FRDGBufferRef RasterTileBuffer,
+	FRDGBufferRef RasterTileAllocatorBuffer,
+	// Top-level grid
+	FBoxSphereBounds TopLevelGridBounds,
+	FIntVector TopLevelGridResolution,
+	FRDGBufferRef& TopLevelGridBuffer,
+	// Bottom-level grid
+	FRDGBufferRef& ExtinctionGridBuffer,
+	FRDGBufferRef& EmissionGridBuffer,
+	FRDGBufferRef& ScatteringGridBuffer
+)
+{
+	//Setup indirect dispatch
+	FRDGBufferRef RasterizeBottomLevelGridIndirectArgsBuffer = GraphBuilder.CreateBuffer(
+		FRDGBufferDesc::CreateIndirectDesc<FRHIDispatchIndirectParameters>(1),
+		TEXT("HeterogeneousVolume.FrustumGrid.RasterizeBottomLevelGridIndirectArgs")
+	);
+
+	FSetRasterizeBottomLevelGridIndirectArgs::FParameters* IndirectArgsPassParameters = GraphBuilder.AllocParameters<FSetRasterizeBottomLevelGridIndirectArgs::FParameters>();
+	{
+		IndirectArgsPassParameters->MaxDispatchThreadGroupsPerDimension = GRHIMaxDispatchThreadGroupsPerDimension;
+		IndirectArgsPassParameters->RasterTileAllocatorBuffer = GraphBuilder.CreateSRV(RasterTileAllocatorBuffer, PF_R32_UINT);
+		IndirectArgsPassParameters->RWRasterizeBottomLevelGridIndirectArgsBuffer = GraphBuilder.CreateUAV(RasterizeBottomLevelGridIndirectArgsBuffer, PF_R32_UINT);
+	}
+
+	FIntVector GroupCount(1, 1, 1);
+	TShaderRef<FSetRasterizeBottomLevelGridIndirectArgs> ComputeShader = View.ShaderMap->GetShader<FSetRasterizeBottomLevelGridIndirectArgs>();
+	FComputeShaderUtils::AddPass(
+		GraphBuilder,
+		RDG_EVENT_NAME("SetRasterizeBottomLevelGridIndirectArgs"),
+		ERDGPassFlags::Compute,
+		ComputeShader,
+		IndirectArgsPassParameters,
+		GroupCount
+	);
+
+	// Pre-allocate bottom-level voxel grid pool based on user-defined budget
+	int32 MaxBottomLevelVoxelCount = (HeterogeneousVolumes::GetMaxBottomLevelMemoryInMegabytesForFrustumGrid() * 1e6) / sizeof(FVectorGridData);
+	int32 BottomLevelGridBufferSize = MaxBottomLevelVoxelCount;
+
+	ExtinctionGridBuffer = GraphBuilder.CreateBuffer(
+		FRDGBufferDesc::CreateStructuredDesc(sizeof(FScalarGridData), BottomLevelGridBufferSize),
+		TEXT("HeterogeneousVolumes.FrustumGrid.ExtinctionGridBuffer")
+	);
+	EmissionGridBuffer = GraphBuilder.CreateBuffer(
+		FRDGBufferDesc::CreateStructuredDesc(sizeof(FVectorGridData), BottomLevelGridBufferSize),
+		TEXT("HeterogeneousVolumes.FrustumGrid.EmissionGridBuffer")
+	);
+	ScatteringGridBuffer = GraphBuilder.CreateBuffer(
+		FRDGBufferDesc::CreateStructuredDesc(sizeof(FVectorGridData), BottomLevelGridBufferSize),
+		TEXT("HeterogeneousVolumes.FrustumGrid.ScatteringGridBuffer")
+	);
+
+	FRDGBufferRef BottomLevelGridAllocatorBuffer = GraphBuilder.CreateBuffer(
+		FRDGBufferDesc::CreateBufferDesc(sizeof(uint32), 2),
+		TEXT("HeterogeneousVolume.FrustumGrid.BottomLevelGridAllocatorBuffer")
+	);
+	AddClearUAVPass(GraphBuilder, GraphBuilder.CreateUAV(BottomLevelGridAllocatorBuffer, PF_R32_UINT), 0);
+
+	// Rasterize volumes into bottom-level grid
+	for (int32 MeshBatchIndex = 0; MeshBatchIndex < View.HeterogeneousVolumesMeshBatches.Num(); ++MeshBatchIndex)
+	{
+		const FMeshBatch* Mesh = View.HeterogeneousVolumesMeshBatches[MeshBatchIndex].Mesh;
+		const FPrimitiveSceneProxy* PrimitiveSceneProxy = View.HeterogeneousVolumesMeshBatches[MeshBatchIndex].Proxy;
+		const FMaterialRenderProxy* MaterialRenderProxy = Mesh->MaterialRenderProxy;
+		if (!ShouldRenderMeshBatchWithHeterogeneousVolumes(Mesh, PrimitiveSceneProxy, View.GetFeatureLevel()))
+		{
+			continue;
+		}
+
+		const FPrimitiveSceneInfo* PrimitiveSceneInfo = PrimitiveSceneProxy->GetPrimitiveSceneInfo();
+		const int32 PrimitiveId = PrimitiveSceneInfo->GetIndex();
+		const FBoxSphereBounds LocalBoxSphereBounds = PrimitiveSceneProxy->GetLocalBounds();
+		const FBoxSphereBounds PrimitiveBounds = PrimitiveSceneProxy->GetBounds();
+		const FMaterial& Material = MaterialRenderProxy->GetMaterialWithFallback(View.GetFeatureLevel(), MaterialRenderProxy);
+
+		FRasterizeBottomLevelFroxelGridCS::FParameters* PassParameters = GraphBuilder.AllocParameters<FRasterizeBottomLevelFroxelGridCS::FParameters>();
+		{
+			// Scene data
+			PassParameters->View = View.ViewUniformBuffer;
+
+			// Primitive data
+			PassParameters->PrimitiveWorldBoundsMin = FVector3f(PrimitiveBounds.Origin - PrimitiveBounds.BoxExtent);
+			PassParameters->PrimitiveWorldBoundsMax = FVector3f(PrimitiveBounds.Origin + PrimitiveBounds.BoxExtent);
+			FMatrix44f LocalToWorld = FMatrix44f(PrimitiveSceneProxy->GetLocalToWorld());
+			PassParameters->LocalToWorld = LocalToWorld;
+			PassParameters->WorldToLocal = LocalToWorld.Inverse();
+			PassParameters->LocalBoundsOrigin = FVector3f(LocalBoxSphereBounds.Origin);
+			PassParameters->LocalBoundsExtent = FVector3f(LocalBoxSphereBounds.BoxExtent);
+			PassParameters->PrimitiveId = PrimitiveId;
+
+			// Volume data
+			PassParameters->TopLevelGridResolution = TopLevelGridResolution;
+			PassParameters->VoxelDimensions = TopLevelGridResolution;
+			PassParameters->ViewToWorld = FMatrix44f(ViewToWorld);
+			PassParameters->HalfFOV = FMath::DegreesToRadians(View.FOV * 0.5);
+			PassParameters->NearPlaneDepth = NearPlaneDistance;
+			PassParameters->FarPlaneDepth = FarPlaneDistance;
+
+			// Sampling data
+			FBlueNoise BlueNoise = GetBlueNoiseGlobalParameters();
+			PassParameters->BlueNoise = CreateUniformBufferImmediate(BlueNoise, EUniformBufferUsage::UniformBuffer_SingleDraw);
+
+			// Raster tile data
+			PassParameters->RasterTileAllocatorBuffer = GraphBuilder.CreateSRV(RasterTileAllocatorBuffer, PF_R32_UINT);
+			PassParameters->RasterTileBuffer = GraphBuilder.CreateSRV(RasterTileBuffer);
+
+			// Indirect args
+			PassParameters->IndirectArgs = RasterizeBottomLevelGridIndirectArgsBuffer;
+
+			// Grid data
+			PassParameters->RWBottomLevelGridAllocatorBuffer = GraphBuilder.CreateUAV(BottomLevelGridAllocatorBuffer, PF_R32_UINT);
+			PassParameters->RWTopLevelGridBuffer = GraphBuilder.CreateUAV(TopLevelGridBuffer, PF_R32_UINT);
+
+			PassParameters->RWExtinctionGridBuffer = GraphBuilder.CreateUAV(ExtinctionGridBuffer);
+			PassParameters->RWEmissionGridBuffer = GraphBuilder.CreateUAV(EmissionGridBuffer);
+			PassParameters->RWScatteringGridBuffer = GraphBuilder.CreateUAV(ScatteringGridBuffer);
+			PassParameters->BottomLevelGridBufferSize = BottomLevelGridBufferSize;
+		}
+
+		GraphBuilder.AddPass(
+			RDG_EVENT_NAME("FrustumGrid.RasterizeBottomLevelGrid"),
+			PassParameters,
+			ERDGPassFlags::Compute,
+			// Why is scene explicitly copied??
+			[PassParameters, LocalScene = Scene, &View, MaterialRenderProxy, &Material](FRHIComputeCommandList& RHICmdList)
+			{
+				FRasterizeBottomLevelFroxelGridCS::FPermutationDomain PermutationVector;
+				TShaderRef<FRasterizeBottomLevelFroxelGridCS> ComputeShader = Material.GetShader<FRasterizeBottomLevelFroxelGridCS>(&FLocalVertexFactory::StaticType, PermutationVector, false);
+
+				if (!ComputeShader.IsNull())
+				{
+					ClearUnusedGraphResources(ComputeShader, PassParameters);
+
+					FMeshPassProcessorRenderState DrawRenderState;
+
+					FMeshMaterialShaderElementData ShaderElementData;
+					ShaderElementData.FadeUniformBuffer = GDistanceCullFadedInUniformBuffer.GetUniformBufferRHI();
+					ShaderElementData.DitherUniformBuffer = GDitherFadedInUniformBuffer.GetUniformBufferRHI();
+
+					FMeshProcessorShaders PassShaders;
+					PassShaders.ComputeShader = ComputeShader;
+
+					FMeshDrawShaderBindings ShaderBindings;
+					{
+						ShaderBindings.Initialize(PassShaders);
+
+						int32 DataOffset = 0;
+						FMeshDrawSingleShaderBindings SingleShaderBindings = ShaderBindings.GetSingleShaderBindings(SF_Compute, DataOffset);
+						ComputeShader->GetShaderBindings(LocalScene, LocalScene->GetFeatureLevel(), nullptr, *MaterialRenderProxy, Material, DrawRenderState, ShaderElementData, SingleShaderBindings);
+
+						ShaderBindings.Finalize(&PassShaders);
+					}
+					ShaderBindings.SetOnCommandList(RHICmdList, ComputeShader.GetComputeShader());
+
+					FComputeShaderUtils::DispatchIndirect(RHICmdList, ComputeShader, *PassParameters, PassParameters->IndirectArgs->GetIndirectRHICallBuffer(), 0);
+				}
+			}
+		);
+	}
+}
+
+void BuildFrustumVoxelGrid(
+	FRDGBuilder& GraphBuilder,
+	const FScene* Scene,
+	const FViewInfo& View,
+	TRDGUniformBufferRef<FFrustumVoxelGridUniformBufferParameters>& FrustumGridUniformBuffer
+)
+{
+	if (!ShouldRenderHeterogeneousVolumesForView(View) || !HeterogeneousVolumes::EnableFrustumVoxelGrid())
+	{
+		CreateEmptyFrustumVoxelGridUniformBuffer(GraphBuilder, FrustumGridUniformBuffer);
+		return;
+	}
+
+	RDG_EVENT_SCOPE(GraphBuilder, "Frustum Grid Build");
+
+	// Determine the minimum voxel size for the scene, based on screen projection or user-defined minima
+	FBoxSphereBounds TopLevelGridBounds;
+	float MinimumVoxelSize;
+	CalcViewBoundsAndMinimumVoxelSize(GraphBuilder, View, TopLevelGridBounds, MinimumVoxelSize);
+
+	if (TopLevelGridBounds.SphereRadius == 0)
+	{
+		CreateEmptyFrustumVoxelGridUniformBuffer(GraphBuilder, FrustumGridUniformBuffer);
+		return;
+	}
+
+	float NearPlaneDistance = HeterogeneousVolumes::GetNearPlaneDistanceForFrustumGrid();
+	float FarPlaneDistance = HeterogeneousVolumes::GetFarPlaneDistanceForFrustumGrid();
+	ClipNearFarDistances(View, TopLevelGridBounds, NearPlaneDistance, FarPlaneDistance);
+
+	FIntVector TopLevelGridResolution;
+	CalculateTopLevelGridResolutionForFrustumGrid(
+		View,
+		MinimumVoxelSize,
+		NearPlaneDistance,
+		FarPlaneDistance,
+		TopLevelGridResolution
+	);
+
+	// Construct top-level grid over global bounding domain with some pre-determined resolution
+	int32 TopLevelVoxelCount = TopLevelGridResolution.X * TopLevelGridResolution.Y * TopLevelGridResolution.Z;
+	if (TopLevelVoxelCount == 0)
+	{
+		CreateEmptyFrustumVoxelGridUniformBuffer(GraphBuilder, FrustumGridUniformBuffer);
+		return;
+	}
+
+	FMatrix WorldToView = View.ViewMatrices.GetViewMatrix();
+	FMatrix ViewToWorld = View.ViewMatrices.GetInvViewMatrix();
+
+	// Mark top-level voxels for rasterization
+	FRDGBufferRef TopLevelGridBuffer;
+	MarkTopLevelGridVoxelsForFrustumGrid(
+		GraphBuilder,
+		View,
+		TopLevelGridResolution,
+		NearPlaneDistance,
+		FarPlaneDistance,
+		ViewToWorld,
+		TopLevelGridBuffer
+	);
+
+	// Generate raster tiles of approximately equal work
+	bool bEnableIndirectionGrid = false;
+	FRDGBufferRef RasterTileBuffer;
+	FRDGBufferRef RasterTileAllocatorBuffer;
+	GenerateRasterTiles(
+		GraphBuilder,
+		Scene,
+		bEnableIndirectionGrid,
+		// Grid data
+		TopLevelGridResolution,
+		TopLevelGridBuffer,
+		// Tile data
+		RasterTileBuffer,
+		RasterTileAllocatorBuffer
+	);
+
+	FRDGBufferRef ExtinctionGridBuffer;
+	FRDGBufferRef EmissionGridBuffer;
+	FRDGBufferRef ScatteringGridBuffer;
+	RasterizeVolumesIntoFrustumVoxelGrid(
+		GraphBuilder,
+		Scene,
+		View,
+		ViewToWorld,
+		NearPlaneDistance,
+		FarPlaneDistance,
+		// Raster tile
+		RasterTileBuffer,
+		RasterTileAllocatorBuffer,
+		// Top-level grid
+		TopLevelGridBounds,
+		TopLevelGridResolution,
+		TopLevelGridBuffer,
+		// Bottom-level grid
+		ExtinctionGridBuffer,
+		EmissionGridBuffer,
+		ScatteringGridBuffer
+	);
+
+	// Create Adpative Voxel Grid uniform buffer
+	FFrustumVoxelGridUniformBufferParameters* UniformBufferParameters = GraphBuilder.AllocParameters<FFrustumVoxelGridUniformBufferParameters>();
+	{
+		FMatrix ViewToClip = FPerspectiveMatrix(
+			FMath::DegreesToRadians(View.FOV * 0.5),
+			TopLevelGridResolution.X,
+			TopLevelGridResolution.Y,
+			NearPlaneDistance,
+			FarPlaneDistance
+		);
+		FMatrix ClipToView = ViewToClip.Inverse();
+		UniformBufferParameters->ViewToClip = FMatrix44f(ViewToClip);
+		UniformBufferParameters->ClipToView = FMatrix44f(ClipToView);
+
+		FMatrix WorldToClip = WorldToView * ViewToClip;
+		FMatrix ClipToWorld = ClipToView * ViewToWorld;
+		UniformBufferParameters->WorldToClip = FMatrix44f(WorldToClip);
+		UniformBufferParameters->ClipToWorld = FMatrix44f(ClipToWorld);
+
+		UniformBufferParameters->WorldToView = FMatrix44f(WorldToView);
+		UniformBufferParameters->ViewToWorld = FMatrix44f(ViewToWorld);
+
+		UniformBufferParameters->TopLevelGridWorldBoundsMin = FVector3f(TopLevelGridBounds.Origin - TopLevelGridBounds.BoxExtent);
+		UniformBufferParameters->TopLevelGridWorldBoundsMax = FVector3f(TopLevelGridBounds.Origin + TopLevelGridBounds.BoxExtent);
+		UniformBufferParameters->TopLevelFroxelGridResolution = TopLevelGridResolution;
+		UniformBufferParameters->VoxelDimensions = TopLevelGridResolution;
+
+		UniformBufferParameters->bUseFrustumGrid = HeterogeneousVolumes::EnableFrustumVoxelGrid();
+		UniformBufferParameters->NearPlaneDepth = NearPlaneDistance;
+		UniformBufferParameters->FarPlaneDepth = FarPlaneDistance;
+
+		// Frustum assignment
+		{
+			// Near/Far plane definition is reversed when explicitly specifying clipping planes
+			FPlane NearPlane;
+			ViewToClip.GetFrustumFarPlane(NearPlane);
+			UniformBufferParameters->ViewFrustumPlanes[0] = FVector4f(NearPlane.X, NearPlane.Y, NearPlane.Z, NearPlane.W);
+
+			FPlane FarPlane;
+			ViewToClip.GetFrustumNearPlane(FarPlane);
+			UniformBufferParameters->ViewFrustumPlanes[1] = FVector4f(FarPlane.X, FarPlane.Y, FarPlane.Z, FarPlane.W);
+
+			FPlane LeftPlane;
+			ViewToClip.GetFrustumLeftPlane(LeftPlane);
+			UniformBufferParameters->ViewFrustumPlanes[2] = FVector4f(LeftPlane.X, LeftPlane.Y, LeftPlane.Z, LeftPlane.W);
+
+			FPlane RightPlane;
+			ViewToClip.GetFrustumRightPlane(RightPlane);
+			UniformBufferParameters->ViewFrustumPlanes[3] = FVector4f(RightPlane.X, RightPlane.Y, RightPlane.Z, RightPlane.W);
+
+			FPlane TopPlane;
+			ViewToClip.GetFrustumTopPlane(TopPlane);
+			UniformBufferParameters->ViewFrustumPlanes[4] = FVector4f(TopPlane.X, TopPlane.Y, TopPlane.Z, TopPlane.W);
+
+			FPlane BottomPlane;
+			ViewToClip.GetFrustumBottomPlane(BottomPlane);
+			UniformBufferParameters->ViewFrustumPlanes[5] = FVector4f(BottomPlane.X, BottomPlane.Y, BottomPlane.Z, BottomPlane.W);
+		}
+
+		UniformBufferParameters->TopLevelFroxelGridBuffer = GraphBuilder.CreateSRV(TopLevelGridBuffer);
+		UniformBufferParameters->ExtinctionFroxelGridBuffer = GraphBuilder.CreateSRV(ExtinctionGridBuffer);
+		UniformBufferParameters->EmissionFroxelGridBuffer = GraphBuilder.CreateSRV(EmissionGridBuffer);
+		UniformBufferParameters->ScatteringFroxelGridBuffer = GraphBuilder.CreateSRV(ScatteringGridBuffer);
+	}
+	FrustumGridUniformBuffer = GraphBuilder.CreateUniformBuffer(UniformBufferParameters);
+}
+
+void ExtractOrthoVoxelGridUniformBuffer(
+	FRDGBuilder& GraphBuilder,
+	const TRDGUniformBufferRef<FOrthoVoxelGridUniformBufferParameters>& OrthoGridUniformBuffer,
+	FAdaptiveOrthoGridParameterCache& AdaptiveOrthoGridParameterCache
+)
+{
+	const TRDGParameterStruct<FOrthoVoxelGridUniformBufferParameters>& Parameters = OrthoGridUniformBuffer->GetParameters();
+
+	AdaptiveOrthoGridParameterCache.TopLevelGridWorldBoundsMin = Parameters->TopLevelGridWorldBoundsMin;
+	AdaptiveOrthoGridParameterCache.TopLevelGridWorldBoundsMax = Parameters->TopLevelGridWorldBoundsMax;
+	AdaptiveOrthoGridParameterCache.TopLevelGridResolution = Parameters->TopLevelGridResolution;
+	AdaptiveOrthoGridParameterCache.bUseOrthoGrid = Parameters->bUseOrthoGrid;
+	AdaptiveOrthoGridParameterCache.bUseMajorantGrid = Parameters->bUseMajorantGrid;
+	AdaptiveOrthoGridParameterCache.bEnableIndirectionGrid = Parameters->bEnableIndirectionGrid;
+
+	GraphBuilder.QueueBufferExtraction(Parameters->TopLevelGridBitmaskBuffer->GetParent(), &AdaptiveOrthoGridParameterCache.TopLevelGridBitmaskBuffer);
+	GraphBuilder.QueueBufferExtraction(Parameters->TopLevelGridBuffer->GetParent(), &AdaptiveOrthoGridParameterCache.TopLevelGridBuffer);
+	GraphBuilder.QueueBufferExtraction(Parameters->IndirectionGridBuffer->GetParent(), &AdaptiveOrthoGridParameterCache.IndirectionGridBuffer);
+
+	GraphBuilder.QueueBufferExtraction(Parameters->ExtinctionGridBuffer->GetParent(), &AdaptiveOrthoGridParameterCache.ExtinctionGridBuffer);
+	GraphBuilder.QueueBufferExtraction(Parameters->EmissionGridBuffer->GetParent(), &AdaptiveOrthoGridParameterCache.EmissionGridBuffer);
+	GraphBuilder.QueueBufferExtraction(Parameters->ScatteringGridBuffer->GetParent(), &AdaptiveOrthoGridParameterCache.ScatteringGridBuffer);
+	GraphBuilder.QueueBufferExtraction(Parameters->MajorantGridBuffer->GetParent(), &AdaptiveOrthoGridParameterCache.MajorantGridBuffer);
+
+}
+
+void RegisterExternalOrthoVoxelGridUniformBuffer(
+	FRDGBuilder& GraphBuilder,
+	const FAdaptiveOrthoGridParameterCache& AdaptiveOrthoGridParameterCache,
+	TRDGUniformBufferRef<FOrthoVoxelGridUniformBufferParameters>& OrthoGridUniformBuffer
+)
+{
+	FOrthoVoxelGridUniformBufferParameters* UniformBufferParameters = GraphBuilder.AllocParameters<FOrthoVoxelGridUniformBufferParameters>();
+	{
+		UniformBufferParameters->TopLevelGridWorldBoundsMin = AdaptiveOrthoGridParameterCache.TopLevelGridWorldBoundsMin;
+		UniformBufferParameters->TopLevelGridWorldBoundsMax = AdaptiveOrthoGridParameterCache.TopLevelGridWorldBoundsMax;
+		UniformBufferParameters->TopLevelGridResolution = AdaptiveOrthoGridParameterCache.TopLevelGridResolution;
+		UniformBufferParameters->bUseOrthoGrid = AdaptiveOrthoGridParameterCache.bUseOrthoGrid;
+		UniformBufferParameters->bUseMajorantGrid = AdaptiveOrthoGridParameterCache.bUseMajorantGrid;
+		UniformBufferParameters->bEnableIndirectionGrid = AdaptiveOrthoGridParameterCache.bEnableIndirectionGrid;
+
+		UniformBufferParameters->TopLevelGridBitmaskBuffer = GraphBuilder.CreateSRV(GraphBuilder.RegisterExternalBuffer(AdaptiveOrthoGridParameterCache.TopLevelGridBitmaskBuffer));
+		UniformBufferParameters->TopLevelGridBuffer = GraphBuilder.CreateSRV(GraphBuilder.RegisterExternalBuffer(AdaptiveOrthoGridParameterCache.TopLevelGridBuffer));
+		UniformBufferParameters->IndirectionGridBuffer = GraphBuilder.CreateSRV(GraphBuilder.RegisterExternalBuffer(AdaptiveOrthoGridParameterCache.IndirectionGridBuffer));
+
+		UniformBufferParameters->ExtinctionGridBuffer = GraphBuilder.CreateSRV(GraphBuilder.RegisterExternalBuffer(AdaptiveOrthoGridParameterCache.ExtinctionGridBuffer));
+		UniformBufferParameters->EmissionGridBuffer = GraphBuilder.CreateSRV(GraphBuilder.RegisterExternalBuffer(AdaptiveOrthoGridParameterCache.EmissionGridBuffer));
+		UniformBufferParameters->ScatteringGridBuffer = GraphBuilder.CreateSRV(GraphBuilder.RegisterExternalBuffer(AdaptiveOrthoGridParameterCache.ScatteringGridBuffer));
+		UniformBufferParameters->MajorantGridBuffer = GraphBuilder.CreateSRV(GraphBuilder.RegisterExternalBuffer(AdaptiveOrthoGridParameterCache.MajorantGridBuffer));
+	}
+	OrthoGridUniformBuffer = GraphBuilder.CreateUniformBuffer(UniformBufferParameters);
+}
+
+void CreateEmptyOrthoVoxelGridUniformBuffer(
+	FRDGBuilder& GraphBuilder,
+	TRDGUniformBufferRef<FOrthoVoxelGridUniformBufferParameters>& OrthoGridUniformBuffer
+)
+{
+	FOrthoVoxelGridUniformBufferParameters* OrthoGridUniformBufferParameters = GraphBuilder.AllocParameters<FOrthoVoxelGridUniformBufferParameters>();
+	{
+		OrthoGridUniformBufferParameters->TopLevelGridWorldBoundsMin = FVector3f(0);
+		OrthoGridUniformBufferParameters->TopLevelGridWorldBoundsMax = FVector3f(0);
+		OrthoGridUniformBufferParameters->TopLevelGridResolution = FIntVector(0);
+
+		OrthoGridUniformBufferParameters->TopLevelGridBitmaskBuffer = GraphBuilder.CreateSRV(GSystemTextures.GetDefaultStructuredBuffer(GraphBuilder, sizeof(FTopLevelGridBitmaskData)));
+		OrthoGridUniformBufferParameters->TopLevelGridBuffer = GraphBuilder.CreateSRV(GSystemTextures.GetDefaultStructuredBuffer(GraphBuilder, sizeof(FTopLevelGridData)));
+		OrthoGridUniformBufferParameters->IndirectionGridBuffer = GraphBuilder.CreateSRV(GSystemTextures.GetDefaultStructuredBuffer(GraphBuilder, sizeof(FTopLevelGridData)));
+		OrthoGridUniformBufferParameters->ExtinctionGridBuffer = GraphBuilder.CreateSRV(GSystemTextures.GetDefaultStructuredBuffer(GraphBuilder, sizeof(FScalarGridData)));
+		OrthoGridUniformBufferParameters->EmissionGridBuffer = GraphBuilder.CreateSRV(GSystemTextures.GetDefaultStructuredBuffer(GraphBuilder, sizeof(FVectorGridData)));
+		OrthoGridUniformBufferParameters->ScatteringGridBuffer = GraphBuilder.CreateSRV(GSystemTextures.GetDefaultStructuredBuffer(GraphBuilder, sizeof(FVectorGridData)));
+
+		OrthoGridUniformBufferParameters->bUseOrthoGrid = false;
+		OrthoGridUniformBufferParameters->bUseMajorantGrid = false;
+		OrthoGridUniformBufferParameters->bEnableIndirectionGrid = false;
+		OrthoGridUniformBufferParameters->MajorantGridBuffer = GraphBuilder.CreateSRV(GSystemTextures.GetDefaultStructuredBuffer(GraphBuilder, sizeof(FScalarGridData)));
+	}
+	OrthoGridUniformBuffer = GraphBuilder.CreateUniformBuffer(OrthoGridUniformBufferParameters);
+}
+
+void RasterizeVolumesIntoOrthoVoxelGrid(
+	FRDGBuilder& GraphBuilder,
+	const FScene* Scene,
+	const TArray<FViewInfo>& Views,
+	// Raster tile
+	FRDGBufferRef RasterTileBuffer,
+	FRDGBufferRef RasterTileAllocatorBuffer,
+	// Top-level grid
+	FBoxSphereBounds TopLevelGridBounds,
+	FIntVector TopLevelGridResolution,
+	FRDGBufferRef& TopLevelGridBuffer,
+	// Indirection grid
+	FRDGBufferRef& IndirectionGridBuffer,
+	// Bottom-level grid
+	FRDGBufferRef& ExtinctionGridBuffer,
+	FRDGBufferRef& EmissionGridBuffer,
+	FRDGBufferRef& ScatteringGridBuffer
+)
+{
+	// Setup indirect dispatch
+	FRDGBufferRef RasterizeBottomLevelGridIndirectArgsBuffer = GraphBuilder.CreateBuffer(
+		FRDGBufferDesc::CreateIndirectDesc<FRHIDispatchIndirectParameters>(1),
+		TEXT("HeterogeneousVolume.RasterizeBottomLevelGridIndirectArgs")
+	);
+
+	FSetRasterizeBottomLevelGridIndirectArgs::FParameters* IndirectArgsPassParameters = GraphBuilder.AllocParameters<FSetRasterizeBottomLevelGridIndirectArgs::FParameters>();
+	{
+		IndirectArgsPassParameters->MaxDispatchThreadGroupsPerDimension = GRHIMaxDispatchThreadGroupsPerDimension;
+		IndirectArgsPassParameters->RasterTileAllocatorBuffer = GraphBuilder.CreateSRV(RasterTileAllocatorBuffer, PF_R32_UINT);
+		IndirectArgsPassParameters->RWRasterizeBottomLevelGridIndirectArgsBuffer = GraphBuilder.CreateUAV(RasterizeBottomLevelGridIndirectArgsBuffer, PF_R32_UINT);
+	}
+
+	const FGlobalShaderMap* GlobalShaderMap = GetGlobalShaderMap(Scene->GetFeatureLevel());
+
+	FIntVector GroupCount(1, 1, 1);
+	TShaderRef<FSetRasterizeBottomLevelGridIndirectArgs> ComputeShader = GlobalShaderMap->GetShader<FSetRasterizeBottomLevelGridIndirectArgs>();
+	FComputeShaderUtils::AddPass(
+		GraphBuilder,
+		RDG_EVENT_NAME("SetRasterizeBottomLevelGridIndirectArgs"),
+		ERDGPassFlags::Compute,
+		ComputeShader,
+		IndirectArgsPassParameters,
+		GroupCount
+	);
+
+	int32 MaxBottomLevelVoxelCount = (HeterogeneousVolumes::GetMaxBottomLevelMemoryInMegabytesForOrthoGrid() * 1e6) / sizeof(FVectorGridData);
+	int32 BottomLevelGridBufferSize = MaxBottomLevelVoxelCount;
+
+	ExtinctionGridBuffer = GraphBuilder.CreateBuffer(
+		FRDGBufferDesc::CreateStructuredDesc(sizeof(FScalarGridData), BottomLevelGridBufferSize),
+		TEXT("HeterogeneousVolumes.OrthoGrid.ExtinctionGridBuffer")
+	);
+	EmissionGridBuffer = GraphBuilder.CreateBuffer(
+		FRDGBufferDesc::CreateStructuredDesc(sizeof(FVectorGridData), BottomLevelGridBufferSize),
+		TEXT("HeterogeneousVolumes.OrthoGrid.EmissionGridBuffer")
+	);
+	ScatteringGridBuffer = GraphBuilder.CreateBuffer(
+		FRDGBufferDesc::CreateStructuredDesc(sizeof(FVectorGridData), BottomLevelGridBufferSize),
+		TEXT("HeterogeneousVolumes.OrthoGrid.ScatteringGridBuffer")
+	);
+
+	FRDGBufferRef BottomLevelGridAllocatorBuffer = GraphBuilder.CreateBuffer(
+		FRDGBufferDesc::CreateBufferDesc(sizeof(uint32), 2),
+		TEXT("HeterogeneousVolume.OrthoGrid.BottomLevelGridAllocatorBuffer")
+	);
+	AddClearUAVPass(GraphBuilder, GraphBuilder.CreateUAV(BottomLevelGridAllocatorBuffer, PF_R32_UINT), 0);
+
+#if 0
+	// Enable bottom-level voxel hashing
+	uint32 HashTableBufferSize = 1;
+	if (HeterogeneousVolumes::EnableVoxelHashing())
+	{
+		HashTableBufferSize = HeterogeneousVolumes::GetVoxelHashingMemoryInMegabytes() * 1.0e6;
+	}
+
+	FRDGBufferRef HashTableBuffer = GraphBuilder.CreateBuffer(
+		FRDGBufferDesc::CreateBufferDesc(sizeof(uint32), HashTableBufferSize),
+		TEXT("HeterogeneousVolume.OrthoGrid.HashTableBufferSize")
+	);
+	AddClearUAVPass(GraphBuilder, GraphBuilder.CreateUAV(HashTableBuffer, PF_R32_UINT), 0);
+
+	FRDGBufferRef HashToVoxelBuffer = GraphBuilder.CreateBuffer(
+		FRDGBufferDesc::CreateBufferDesc(sizeof(uint32), HashTableBufferSize),
+		TEXT("HeterogeneousVolume.OrthoGrid.HashToVoxelBuffer")
+	);
+#endif
+
+	for (int32 ViewIndex = 0; ViewIndex < Views.Num(); ViewIndex++)
+	{
+		// Build view
+		const FViewInfo& View = Views[ViewIndex];
+		FVector WorldCameraOrigin = View.ViewMatrices.GetViewOrigin();
+		FBoxSphereBounds WorldCameraBounds(FSphere(WorldCameraOrigin, HeterogeneousVolumes::GetMaxTraceDistance()));
+
+		for (int32 MeshBatchIndex = 0; MeshBatchIndex < View.HeterogeneousVolumesMeshBatches.Num(); ++MeshBatchIndex)
+		{
+			const FMeshBatch* Mesh = View.HeterogeneousVolumesMeshBatches[MeshBatchIndex].Mesh;
+			const FPrimitiveSceneProxy* PrimitiveSceneProxy = View.HeterogeneousVolumesMeshBatches[MeshBatchIndex].Proxy;
+			const FMaterialRenderProxy* MaterialRenderProxy = Mesh->MaterialRenderProxy;
+			if (!ShouldRenderMeshBatchWithHeterogeneousVolumes(Mesh, PrimitiveSceneProxy, View.GetFeatureLevel()))
+			{
+				continue;
+			}
+
+			const FPrimitiveSceneInfo* PrimitiveSceneInfo = PrimitiveSceneProxy->GetPrimitiveSceneInfo();
+			const int32 PrimitiveId = PrimitiveSceneInfo->GetIndex();
+			const FBoxSphereBounds LocalBoxSphereBounds = PrimitiveSceneProxy->GetLocalBounds();
+			const FBoxSphereBounds PrimitiveBounds = PrimitiveSceneProxy->GetBounds();
+			const FMaterial& Material = MaterialRenderProxy->GetMaterialWithFallback(View.GetFeatureLevel(), MaterialRenderProxy);
+
+			FRasterizeBottomLevelGrid::FParameters* PassParameters = GraphBuilder.AllocParameters<FRasterizeBottomLevelGrid::FParameters>();
+			{
+				// Scene data
+				PassParameters->View = View.ViewUniformBuffer;
+
+				// Primitive data
+				PassParameters->PrimitiveWorldBoundsMin = FVector3f(PrimitiveBounds.Origin - PrimitiveBounds.BoxExtent);
+				PassParameters->PrimitiveWorldBoundsMax = FVector3f(PrimitiveBounds.Origin + PrimitiveBounds.BoxExtent);
+
+				FMatrix44f LocalToWorld = FMatrix44f(PrimitiveSceneProxy->GetLocalToWorld());
+				PassParameters->LocalToWorld = LocalToWorld;
+				PassParameters->WorldToLocal = LocalToWorld.Inverse();
+				PassParameters->LocalBoundsOrigin = FVector3f(LocalBoxSphereBounds.Origin);
+				PassParameters->LocalBoundsExtent = FVector3f(LocalBoxSphereBounds.BoxExtent);
+				PassParameters->PrimitiveId = PrimitiveId;
+
+				// Volume data
+				PassParameters->TopLevelGridResolution = TopLevelGridResolution;
+				PassParameters->TopLevelGridWorldBoundsMin = FVector3f(TopLevelGridBounds.Origin - TopLevelGridBounds.BoxExtent);
+				PassParameters->TopLevelGridWorldBoundsMax = FVector3f(TopLevelGridBounds.Origin + TopLevelGridBounds.BoxExtent);
+
+				// Unify with "object" definition??
+				PassParameters->PrimitiveWorldBoundsMin = FVector3f(PrimitiveBounds.Origin - PrimitiveBounds.BoxExtent);
+				PassParameters->PrimitiveWorldBoundsMax = FVector3f(PrimitiveBounds.Origin + PrimitiveBounds.BoxExtent);
+
+				PassParameters->BottomLevelGridBufferSize = BottomLevelGridBufferSize;
+
+				// Raster tile data
+				PassParameters->RasterTileAllocatorBuffer = GraphBuilder.CreateSRV(RasterTileAllocatorBuffer, PF_R32_UINT);
+				PassParameters->RasterTileBuffer = GraphBuilder.CreateSRV(RasterTileBuffer);
+
+				// Indirect args
+				PassParameters->IndirectArgs = RasterizeBottomLevelGridIndirectArgsBuffer;
+
+				// Sampling mode
+				PassParameters->bSampleAtVertices = HeterogeneousVolumes::EnableLinearInterpolation();
+
+				// Grid data
+				PassParameters->TopLevelGridBuffer = GraphBuilder.CreateSRV(TopLevelGridBuffer);
+				PassParameters->RWBottomLevelGridAllocatorBuffer = GraphBuilder.CreateUAV(BottomLevelGridAllocatorBuffer, PF_R32_UINT);
+				PassParameters->RWTopLevelGridBuffer = GraphBuilder.CreateUAV(TopLevelGridBuffer, PF_R32_UINT);
+				PassParameters->RWExtinctionGridBuffer = GraphBuilder.CreateUAV(ExtinctionGridBuffer);
+				PassParameters->RWEmissionGridBuffer = GraphBuilder.CreateUAV(EmissionGridBuffer);
+				PassParameters->RWScatteringGridBuffer = GraphBuilder.CreateUAV(ScatteringGridBuffer);
+
+				// Indirection Grid
+				PassParameters->FixedBottomLevelResolution = HeterogeneousVolumes::GetBottomLevelGridResolution();
+				PassParameters->RWIndirectionGridBuffer = GraphBuilder.CreateUAV(IndirectionGridBuffer);
+
+#if 0
+				// Hash table
+				if (HeterogeneousVolumes::EnableVoxelHashing())
+				{
+					PassParameters->RWHashTable = GraphBuilder.CreateUAV(HashTableBuffer);
+					PassParameters->RWHashToVoxelBuffer = GraphBuilder.CreateUAV(HashToVoxelBuffer);
+					PassParameters->HashTableSize = HashTableBufferSize;
+				}
+#endif
+				PassParameters->HomogeneousThreshold = HeterogeneousVolumes::GetHomogeneousAggregationThreshold();
+			}
+
+			GraphBuilder.AddPass(
+				RDG_EVENT_NAME("RasterizeBottomLevelGrid"),
+				PassParameters,
+				ERDGPassFlags::Compute,
+				// Why is scene explicitly copied?
+				[PassParameters, LocalScene = Scene, &View, MaterialRenderProxy, &Material](FRHIComputeCommandList& RHICmdList)
+				{
+					FRasterizeBottomLevelGrid::FPermutationDomain PermutationVector;
+					PermutationVector.Set<FRasterizeBottomLevelGrid::FEnableIndirectionGrid>(HeterogeneousVolumes::EnableIndirectionGrid());
+					//PermutationVector.Set<FRasterizeBottomLevelGrid::FEnableVoxelHashing>(HeterogeneousVolumes::EnableVoxelHashing());
+					PermutationVector.Set<FRasterizeBottomLevelGrid::FEnableHomogeneousAggregation>(HeterogeneousVolumes::EnableHomogeneousAggregation());
+					TShaderRef<FRasterizeBottomLevelGrid> ComputeShader = Material.GetShader<FRasterizeBottomLevelGrid>(&FLocalVertexFactory::StaticType, PermutationVector, false);
+
+
+					if (!ComputeShader.IsNull())
+					{
+						ClearUnusedGraphResources(ComputeShader, PassParameters);
+
+						FMeshPassProcessorRenderState DrawRenderState;
+
+						FMeshMaterialShaderElementData ShaderElementData;
+						ShaderElementData.FadeUniformBuffer = GDistanceCullFadedInUniformBuffer.GetUniformBufferRHI();
+						ShaderElementData.DitherUniformBuffer = GDitherFadedInUniformBuffer.GetUniformBufferRHI();
+
+						FMeshProcessorShaders PassShaders;
+						PassShaders.ComputeShader = ComputeShader;
+
+						FMeshDrawShaderBindings ShaderBindings;
+						{
+							ShaderBindings.Initialize(PassShaders);
+
+							int32 DataOffset = 0;
+							FMeshDrawSingleShaderBindings SingleShaderBindings = ShaderBindings.GetSingleShaderBindings(SF_Compute, DataOffset);
+							ComputeShader->GetShaderBindings(LocalScene, LocalScene->GetFeatureLevel(), nullptr, *MaterialRenderProxy, Material, DrawRenderState, ShaderElementData, SingleShaderBindings);
+
+							ShaderBindings.Finalize(&PassShaders);
+						}
+						ShaderBindings.SetOnCommandList(RHICmdList, ComputeShader.GetComputeShader());
+
+						FComputeShaderUtils::DispatchIndirect(RHICmdList, ComputeShader, *PassParameters, PassParameters->IndirectArgs->GetIndirectRHICallBuffer(), 0);
+					}
+				}
+			);
+		}
+	}
+}
+
+void BuildOrthoVoxelGrid(
+	FRDGBuilder& GraphBuilder,
+	const FScene* Scene,
+	/*const*/ TArray<FViewInfo>& Views,
+	TRDGUniformBufferRef<FOrthoVoxelGridUniformBufferParameters>& OrthoGridUniformBuffer
+)
+{
+	if (!ShouldRenderHeterogeneousVolumes(Scene) || !HeterogeneousVolumes::EnableOrthoVoxelGrid())
+	{
+		CreateEmptyOrthoVoxelGridUniformBuffer(GraphBuilder, OrthoGridUniformBuffer);
+		return;
+	}
+	check(!Views.IsEmpty());
+
+	RDG_EVENT_SCOPE(GraphBuilder, "Ortho Grid Build");
+
+	// Collect global bounds
+	FBoxSphereBounds TopLevelGridBounds;
+	float GlobalMinimumVoxelSize;
+	CalcGlobalBoundsAndMinimumVoxelSize(
+		GraphBuilder,
+		Views,
+		TopLevelGridBounds,
+		GlobalMinimumVoxelSize
+	);
+
+	if (TopLevelGridBounds.SphereRadius == 0)
+	{
+		CreateEmptyOrthoVoxelGridUniformBuffer(GraphBuilder, OrthoGridUniformBuffer);
+		return;
+	}
+
+	// Determine top-level grid resolution
+	FIntVector TopLevelGridResolution;
+	CalculateTopLevelGridResolution(
+		TopLevelGridBounds,
+		GlobalMinimumVoxelSize,
+		TopLevelGridResolution
+	);
+
+	// Calculate the preferred voxel size for each bottom-level grid in a top-level cell
+	FRDGBufferRef TopLevelGridBuffer;
+	CalculateVoxelSize(
+		GraphBuilder,
+		Scene,
+		Views,
+		TopLevelGridBounds,
+		TopLevelGridResolution,
+		TopLevelGridBuffer
+	);
+
+	// Allocate bottom-level grid
+	FRDGBufferRef IndirectionGridBuffer;
+	MarkTopLevelGrid(
+		GraphBuilder,
+		Scene,
+		// Grid data
+		TopLevelGridBounds,
+		TopLevelGridResolution,
+		TopLevelGridBuffer,
+		IndirectionGridBuffer
+	);
+
+	// Generate raster tiles
+	FRDGBufferRef RasterTileBuffer;
+	FRDGBufferRef RasterTileAllocatorBuffer;
+	GenerateRasterTiles(
+		GraphBuilder,
+		Scene,
+		HeterogeneousVolumes::EnableIndirectionGrid(),
+		// Grid data
+		TopLevelGridResolution,
+		TopLevelGridBuffer,
+		// Tile data
+		RasterTileBuffer,
+		RasterTileAllocatorBuffer
+	);
+
+	// Volume rasterization
+	FRDGBufferRef ExtinctionGridBuffer;
+	FRDGBufferRef EmissionGridBuffer;
+	FRDGBufferRef ScatteringGridBuffer;
+	RasterizeVolumesIntoOrthoVoxelGrid(
+		GraphBuilder,
+		Scene,
+		Views,
+		// Tile data
+		RasterTileBuffer,
+		RasterTileAllocatorBuffer,
+		// Grid data
+		TopLevelGridBounds,
+		TopLevelGridResolution,
+		TopLevelGridBuffer,
+		IndirectionGridBuffer,
+		ExtinctionGridBuffer,
+		EmissionGridBuffer,
+		ScatteringGridBuffer
+	);
+
+	FRDGBufferRef TopLevelGridBitmaskBuffer;
+	CreateTopLevelBitmask(GraphBuilder, Views, TopLevelGridResolution, TopLevelGridBuffer, TopLevelGridBitmaskBuffer);
+
+	FRDGBufferRef MajorantGridBuffer;
+	BuildMajorantVoxelGrid(GraphBuilder, Scene, TopLevelGridResolution, TopLevelGridBuffer, IndirectionGridBuffer, ExtinctionGridBuffer, MajorantGridBuffer);
+
+	// Create Adpative Voxel Grid uniform buffer
+	FOrthoVoxelGridUniformBufferParameters* OrthoGridUniformBufferParameters = GraphBuilder.AllocParameters<FOrthoVoxelGridUniformBufferParameters>();
+	{
+		OrthoGridUniformBufferParameters->TopLevelGridWorldBoundsMin = FVector3f(TopLevelGridBounds.Origin - TopLevelGridBounds.BoxExtent);
+		OrthoGridUniformBufferParameters->TopLevelGridWorldBoundsMax = FVector3f(TopLevelGridBounds.Origin + TopLevelGridBounds.BoxExtent);
+		OrthoGridUniformBufferParameters->TopLevelGridResolution = TopLevelGridResolution;
+
+		OrthoGridUniformBufferParameters->TopLevelGridBitmaskBuffer = GraphBuilder.CreateSRV(TopLevelGridBitmaskBuffer);
+		OrthoGridUniformBufferParameters->TopLevelGridBuffer = GraphBuilder.CreateSRV(TopLevelGridBuffer);
+		OrthoGridUniformBufferParameters->IndirectionGridBuffer = GraphBuilder.CreateSRV(IndirectionGridBuffer);
+		OrthoGridUniformBufferParameters->ExtinctionGridBuffer = GraphBuilder.CreateSRV(ExtinctionGridBuffer);
+		OrthoGridUniformBufferParameters->EmissionGridBuffer = GraphBuilder.CreateSRV(EmissionGridBuffer);
+		OrthoGridUniformBufferParameters->ScatteringGridBuffer = GraphBuilder.CreateSRV(ScatteringGridBuffer);
+
+		OrthoGridUniformBufferParameters->bUseOrthoGrid = HeterogeneousVolumes::EnableOrthoVoxelGrid();
+		OrthoGridUniformBufferParameters->bUseMajorantGrid = HeterogeneousVolumes::EnableMajorantGrid();
+		OrthoGridUniformBufferParameters->bEnableIndirectionGrid = HeterogeneousVolumes::EnableIndirectionGrid();
+		OrthoGridUniformBufferParameters->MajorantGridBuffer = GraphBuilder.CreateSRV(MajorantGridBuffer);
+	}
+	OrthoGridUniformBuffer = GraphBuilder.CreateUniformBuffer(OrthoGridUniformBufferParameters);
+}
+
+void RenderTransmittanceWithVoxelGrid(
+	FRDGBuilder& GraphBuilder,
+	// Scene data
+	const FSceneTextures& SceneTextures,
+	FScene* Scene,
+	const FSceneViewFamily& ViewFamily,
+	FViewInfo& View,
+	const TRDGUniformBufferRef<FOrthoVoxelGridUniformBufferParameters>& OrthoGridUniformBuffer,
+	const TRDGUniformBufferRef<FFrustumVoxelGridUniformBufferParameters>& FrustumGridUniformBuffer,
+	// Output
+	FRDGTextureRef& HeterogeneousVolumeRadiance
+)
+{
+	uint32 GroupCountX = FMath::DivideAndRoundUp(View.ViewRect.Size().X, FRenderTransmittanceWithVoxelGridCS::GetThreadGroupSize2D());
+	uint32 GroupCountY = FMath::DivideAndRoundUp(View.ViewRect.Size().Y, FRenderTransmittanceWithVoxelGridCS::GetThreadGroupSize2D());
+	FIntVector GroupCount = FIntVector(GroupCountX, GroupCountY, 1);
+
+	FRenderTransmittanceWithVoxelGridCS::FParameters* PassParameters = GraphBuilder.AllocParameters<FRenderTransmittanceWithVoxelGridCS::FParameters>();
+	{
+		// Scene data
+		PassParameters->View = View.ViewUniformBuffer;
+		PassParameters->SceneTextures = GetSceneTextureParameters(GraphBuilder, SceneTextures);
+
+		// Ray data
+		PassParameters->bJitter = HeterogeneousVolumes::ShouldJitter();
+		PassParameters->MaxTraceDistance = HeterogeneousVolumes::GetMaxTraceDistance();
+
+		// Volume data
+		PassParameters->OrthoGridUniformBuffer = OrthoGridUniformBuffer;
+		PassParameters->FrustumGridUniformBuffer = FrustumGridUniformBuffer;
+
+		// Dispatch data
+		PassParameters->GroupCount = GroupCount;
+
+		// Output
+		PassParameters->RWLightingTexture = GraphBuilder.CreateUAV(HeterogeneousVolumeRadiance);
+	}
+
+	FRenderTransmittanceWithVoxelGridCS::FPermutationDomain PermutationVector;
+	TShaderRef<FRenderTransmittanceWithVoxelGridCS> ComputeShader = View.ShaderMap->GetShader<FRenderTransmittanceWithVoxelGridCS>(PermutationVector);
+	FComputeShaderUtils::AddPass(
+		GraphBuilder,
+		RDG_EVENT_NAME("RenderTransmittanceWithVoxelGridCS"),
+		ERDGPassFlags::Compute | ERDGPassFlags::NeverCull,
+		ComputeShader,
+		PassParameters,
+		GroupCount
+	);
+}
\ No newline at end of file
