More work towards persisting VSMs and associated physical pages and moving things more towards the shadow scene and shadow renderer
- Make the cache entries effectively the persistent data structure for light info that always exists
- Allow cache entries to live some number of frames even when they are unreferenced. Reallocate their VSM IDs and keep projection data as needed
- Lots of cleanup and refactoring things into simpler paths and avoiding redoing various work every frame

Lots more to come but this seems like a good place to checkpoint

Notes:
- There are some theoretical correctness issues here as we re-add the VSM IDs and page tables to allow the cached pages to live on but we do not recreate the render views to allow Nanite/Non-Nanite to render any things "on top" of cached pages, as we currently need without the future proposed changes to invalidation.
- Invalidations that are triggered as side-effects of rendering (such as WPO and Nanite) are going to be a bit of a problem as we don't track the render views for the cached pages. Almost certainly we want offscreen invalidations to just flag the given page but otherwise just let it drop, but we can't lose track of cases like WPO that would be invalidating just because the light is offscreen. Will need some solution for this that ideally doesn't have the full overhead of the culling pipeline.

diff --git a/Engine/Shaders/Private/VirtualShadowMaps/VirtualShadowMapPageAccessCommon.ush b/Engine/Shaders/Private/VirtualShadowMaps/VirtualShadowMapPageAccessCommon.ush
--- a/Engine/Shaders/Private/VirtualShadowMaps/VirtualShadowMapPageAccessCommon.ush
+++ b/Engine/Shaders/Private/VirtualShadowMaps/VirtualShadowMapPageAccessCommon.ush
@@ -40,7 +40,7 @@
 #define VSM_PAGE_FLAGS_BITS_PER_HMIP (4U)
 #define VSM_PAGE_FLAGS_BITS_MASK     ((1U<<VSM_PAGE_FLAGS_BITS_PER_HMIP)-1U)
 
-bool IsSinglePageVirtualShadowMap(uint VirtualShadowMapId)
+bool IsSinglePageVirtualShadowMap(int VirtualShadowMapId)
 {
 	return VirtualShadowMapId < VSM_MAX_SINGLE_PAGE_SHADOW_MAPS;
 }
diff --git a/Engine/Source/Runtime/Renderer/Private/DeferredShadingRenderer.cpp b/Engine/Source/Runtime/Renderer/Private/DeferredShadingRenderer.cpp
--- a/Engine/Source/Runtime/Renderer/Private/DeferredShadingRenderer.cpp
+++ b/Engine/Source/Runtime/Renderer/Private/DeferredShadingRenderer.cpp
@@ -3708,9 +3708,7 @@ void FDeferredShadingSceneRenderer::Render(FRDGBuilder& GraphBuilder)
 
 				FFrontLayerTranslucencyData FrontLayerTranslucencyData = RenderFrontLayerTranslucency(GraphBuilder, Views, SceneTextures, true /*VSM page marking*/);
 
-				// TODO: propagate config in some other way (along with shadowing lights flowing from persistent setup to the VSM setup)
-				auto GetLightMobilityFactor = [this](int32 LightId)->float { return Scene->ShadowScene->GetLightMobilityFactor(LightId); };
-				VirtualShadowMapArray.BuildPageAllocations(GraphBuilder, GetActiveSceneTextures(), Views, ViewFamily.EngineShowFlags, SortedLightSet, VisibleLightInfos, GetLightMobilityFactor, SingleLayerWaterPrePassResult, FrontLayerTranslucencyData);
+				VirtualShadowMapArray.BuildPageAllocations(GraphBuilder, GetActiveSceneTextures(), Views, ViewFamily.EngineShowFlags, SortedLightSet, VisibleLightInfos, SingleLayerWaterPrePassResult, FrontLayerTranslucencyData);
 			}
 
 			RenderShadowDepthMaps(GraphBuilder, InstanceCullingManager, ExternalAccessQueue);
diff --git a/Engine/Source/Runtime/Renderer/Private/SceneRendering.cpp b/Engine/Source/Runtime/Renderer/Private/SceneRendering.cpp
--- a/Engine/Source/Runtime/Renderer/Private/SceneRendering.cpp
+++ b/Engine/Source/Runtime/Renderer/Private/SceneRendering.cpp
@@ -472,11 +472,11 @@ TSharedPtr<FVirtualShadowMapClipmap> FVisibleLightInfo::FindShadowClipmapForView
 	return TSharedPtr<FVirtualShadowMapClipmap>();
 }
 
-int32 FVisibleLightInfo::GetVirtualShadowMapId( const FViewInfo* View ) const
+int32 FVisibleLightInfo::GetVirtualShadowMapId(const FViewInfo* View) const
 {
-	if( VirtualShadowMapClipmaps.Num() )
+	if (VirtualShadowMapClipmaps.Num())
 	{
-		return FindShadowClipmapForView( View )->GetVirtualShadowMap(0)->ID;
+		return FindShadowClipmapForView(View)->GetVirtualShadowMapId();
 	}
 	else
 	{
diff --git a/Engine/Source/Runtime/Renderer/Private/ShadowRendering.h b/Engine/Source/Runtime/Renderer/Private/ShadowRendering.h
--- a/Engine/Source/Runtime/Renderer/Private/ShadowRendering.h
+++ b/Engine/Source/Runtime/Renderer/Private/ShadowRendering.h
@@ -435,8 +435,8 @@ public:
 	EMeshPass::Type MeshPassTargetType = EMeshPass::CSMShadowDepth;
 
 	EShadowMeshSelection MeshSelectionMask = EShadowMeshSelection::All;
-	/** */
-	TArray< FVirtualShadowMap*, TInlineAllocator<6> > VirtualShadowMaps;
+
+	int32 VirtualShadowMapId = INDEX_NONE;
 	TSharedPtr<FVirtualShadowMapPerLightCacheEntry> VirtualShadowMapPerLightCacheEntry;
 
 	/** View projection matrices for each cubemap face, used by one pass point light shadows. */
@@ -768,7 +768,7 @@ public:
 		return FShadowDepthType(bDirectionalLight, bOnePassPointLightShadow);
 	}
 
-	bool HasVirtualShadowMap() const { return VirtualShadowMaps.Num() > 0 || VirtualShadowMapClipmap.IsValid(); }
+	bool HasVirtualShadowMap() const { return VirtualShadowMapId != INDEX_NONE || VirtualShadowMapClipmap.IsValid(); }
 
 	FParallelMeshDrawCommandPass& GetShadowDepthPass() { return ShadowDepthPass; }
 
diff --git a/Engine/Source/Runtime/Renderer/Private/ShadowSetup.cpp b/Engine/Source/Runtime/Renderer/Private/ShadowSetup.cpp
--- a/Engine/Source/Runtime/Renderer/Private/ShadowSetup.cpp
+++ b/Engine/Source/Runtime/Renderer/Private/ShadowSetup.cpp
@@ -3973,7 +3973,7 @@ void FSceneRenderer::CreateWholeSceneProjectedShadow(
 								FBoxSphereBounds const& Bounds = PrimitiveSceneInfo->Proxy->GetBounds();
 								if (IntersectsConvexHulls(LightViewFrustumConvexHulls, Bounds))
 								{
-									if (ProjectedShadowInfo->AddSubjectPrimitive(TaskData, Interaction->GetPrimitiveSceneInfo(), Views, false) 
+									if (ProjectedShadowInfo->AddSubjectPrimitive(TaskData, Interaction->GetPrimitiveSceneInfo(), Views, false)
 										&& VirtualSmCacheEntry.IsValid())
 									{
 										VirtualSmCacheEntry->OnPrimitiveRendered(PrimitiveSceneInfo);
@@ -4007,12 +4007,6 @@ void FSceneRenderer::CreateWholeSceneProjectedShadow(
 
 					TSharedPtr<FVirtualShadowMapPerLightCacheEntry> VirtualSmPerLightCacheEntry = ShadowSceneRenderer->AddLocalLightShadow(ProjectedShadowInitializer, ProjectedShadowInfo, LightSceneInfo, MaxScreenRadius);
 
-					if (VirtualSmPerLightCacheEntry)
-					{
-						ProjectedShadowInfo->VirtualShadowMapPerLightCacheEntry = VirtualSmPerLightCacheEntry;
-						ProjectedShadowInfo->bShouldRenderVSM = !VirtualSmPerLightCacheEntry->IsFullyCached();
-					}
-
 					ProjectedShadowInfo->MeshPassTargetType = EMeshPass::VSMShadowDepth;
 					ProjectedShadowInfo->MeshSelectionMask = EShadowMeshSelection::VSM;
 
@@ -4031,7 +4025,7 @@ void FSceneRenderer::CreateWholeSceneProjectedShadow(
 						AddInteractingPrimitives(LightSceneInfo->GetDynamicInteractionStaticPrimitiveList(), ProjectedShadowInfo, bContainsNaniteSubjects, EmptyHull, VirtualSmPerLightCacheEntry);
 						ProjectedShadowInfo->bContainsNaniteSubjects = bContainsNaniteSubjects;
 					}
-					VisibleLightInfo.VirtualShadowMapId = ProjectedShadowInfo->VirtualShadowMaps[0]->ID;
+					VisibleLightInfo.VirtualShadowMapId = ProjectedShadowInfo->VirtualShadowMapId;
 					VisibleLightInfo.AllProjectedShadows.Add(ProjectedShadowInfo);
 				}
 				
@@ -5184,13 +5178,9 @@ void FSceneRenderer::AddViewDependentWholeSceneShadowsForView(
 
 			if (bNeedsVirtualShadowMap)
 			{
-				FVector LightDirection = LightSceneInfo.Proxy->GetDirection();
-				FMatrix WorldToLight = FInverseRotationMatrix(LightDirection.GetSafeNormal().Rotation());
-
 				TSharedPtr<FVirtualShadowMapClipmap> VirtualShadowMapClipmap = TSharedPtr<FVirtualShadowMapClipmap>(new FVirtualShadowMapClipmap(
 					VirtualShadowMapArray,
 					LightSceneInfo,
-					WorldToLight,
 					View.ViewMatrices,
 					View.ViewRect.Size(),
 					&View,
@@ -5215,7 +5205,7 @@ void FSceneRenderer::AddViewDependentWholeSceneShadowsForView(
 				}
 
 				// NOTE: If there are multiple camera views this will simply be associated with "one of them"
-				VisibleLightInfo.VirtualShadowMapId = VirtualShadowMapClipmap->GetVirtualShadowMap(0)->ID;
+				VisibleLightInfo.VirtualShadowMapId = VirtualShadowMapClipmap->GetVirtualShadowMapId();
 			}
 		}
 	}
diff --git a/Engine/Source/Runtime/Renderer/Private/Shadows/ShadowSceneRenderer.cpp b/Engine/Source/Runtime/Renderer/Private/Shadows/ShadowSceneRenderer.cpp
--- a/Engine/Source/Runtime/Renderer/Private/Shadows/ShadowSceneRenderer.cpp
+++ b/Engine/Source/Runtime/Renderer/Private/Shadows/ShadowSceneRenderer.cpp
@@ -59,6 +59,20 @@ TAutoConsoleVariable<int32> CVarVirtualShadowOnePassProjection(
 	ECVF_Scalability | ECVF_RenderThreadSafe
 );
 
+static TAutoConsoleVariable<float> CVarResolutionLodBiasLocal(
+	TEXT("r.Shadow.Virtual.ResolutionLodBiasLocal"),
+	0.0f,
+	TEXT("Bias applied to LOD calculations for local lights. -1.0 doubles resolution, 1.0 halves it and so on."),
+	ECVF_Scalability | ECVF_RenderThreadSafe
+);
+
+static TAutoConsoleVariable<float> CVarResolutionLodBiasLocalMoving(
+	TEXT("r.Shadow.Virtual.ResolutionLodBiasLocalMoving"),
+	1.0f,
+	TEXT("Bias applied to LOD calculations for moving local lights. -1.0 doubles resolution, 1.0 halves it and so on."),
+	ECVF_Scalability | ECVF_RenderThreadSafe
+);
+
 DECLARE_DWORD_COUNTER_STAT(TEXT("VSM Total Raster Bins"), STAT_VSMNaniteBasePassTotalRasterBins, STATGROUP_ShadowRendering);
 DECLARE_DWORD_COUNTER_STAT(TEXT("VSM Total Shading Draws"), STAT_VSMNaniteBasePassTotalShadingDraws, STATGROUP_ShadowRendering);
 
@@ -85,6 +99,48 @@ float FShadowSceneRenderer::ComputeNaniteShadowsLODScaleFactor()
 	return FMath::Pow(2.0f, -CVarNaniteShadowsLODBias.GetValueOnRenderThread());
 }
 
+static float GetResolutionLODBiasLocal(float LightMobilityFactor)
+{
+	return FVirtualShadowMapArray::InterpolateResolutionBias(
+		CVarResolutionLodBiasLocal.GetValueOnRenderThread(),
+		CVarResolutionLodBiasLocalMoving.GetValueOnRenderThread(),
+		LightMobilityFactor);
+}
+
+FVirtualShadowMapProjectionShaderData FShadowSceneRenderer::GetLocalLightProjectionShaderData(
+	float ResolutionLODBiasLocal,
+	const FProjectedShadowInfo* ProjectedShadowInfo,
+	int32 MapIndex) const
+{
+	TSharedPtr<FVirtualShadowMapPerLightCacheEntry> CacheEntry = ProjectedShadowInfo->VirtualShadowMapPerLightCacheEntry;
+	check(CacheEntry.IsValid());
+
+	int32 VirtualShadowMapId = ProjectedShadowInfo->VirtualShadowMapId + MapIndex;
+	bool bIsSinglePageSM = FVirtualShadowMapArray::IsSinglePage(VirtualShadowMapId);
+	check(VirtualShadowMapId != INDEX_NONE && CacheEntry->Current.bIsDistantLight == bIsSinglePageSM);
+
+	uint32 Flags = bIsSinglePageSM ? VSM_PROJ_FLAG_CURRENT_DISTANT_LIGHT : 0U;
+	Flags |= CacheEntry->IsUncached() ? VSM_PROJ_FLAG_UNCACHED : 0U;
+
+	const FViewMatrices ViewMatrices = ProjectedShadowInfo->GetShadowDepthRenderingViewMatrices(MapIndex, true);
+	const FLargeWorldRenderPosition PreViewTranslation(ProjectedShadowInfo->PreShadowTranslation);
+
+	FVirtualShadowMapProjectionShaderData Data; 
+	Data.TranslatedWorldToShadowViewMatrix		= FMatrix44f(ViewMatrices.GetTranslatedViewMatrix());	// LWC_TODO: Precision loss?
+	Data.ShadowViewToClipMatrix					= FMatrix44f(ViewMatrices.GetProjectionMatrix());
+	Data.TranslatedWorldToShadowUVMatrix		= FMatrix44f(CalcTranslatedWorldToShadowUVMatrix( ViewMatrices.GetTranslatedViewMatrix(), ViewMatrices.GetProjectionMatrix() ));
+	Data.TranslatedWorldToShadowUVNormalMatrix	= FMatrix44f(CalcTranslatedWorldToShadowUVNormalMatrix( ViewMatrices.GetTranslatedViewMatrix(), ViewMatrices.GetProjectionMatrix() ));
+	Data.PreViewTranslationLWCTile				= PreViewTranslation.GetTile();
+	Data.PreViewTranslationLWCOffset			= PreViewTranslation.GetOffset();
+	Data.LightType								= ProjectedShadowInfo->GetLightSceneInfo().Proxy->GetLightType();
+	Data.LightSourceRadius						= ProjectedShadowInfo->GetLightSceneInfo().Proxy->GetSourceRadius();
+	Data.ResolutionLodBias						= ResolutionLODBiasLocal;
+	Data.LightRadius							= ProjectedShadowInfo->GetLightSceneInfo().Proxy->GetRadius();
+	Data.Flags									= Flags;
+
+	return Data;
+}
+
 TSharedPtr<FVirtualShadowMapPerLightCacheEntry> FShadowSceneRenderer::AddLocalLightShadow(const FWholeSceneProjectedShadowInitializer& ProjectedShadowInitializer, FProjectedShadowInfo* ProjectedShadowInfo, FLightSceneInfo* LightSceneInfo, float MaxScreenRadius)
 {
 	FVirtualShadowMapArrayCacheManager* CacheManager = VirtualShadowMapArray.CacheManager;
@@ -94,50 +150,49 @@ TSharedPtr<FVirtualShadowMapPerLightCacheEntry> FShadowSceneRenderer::AddLocalLi
 	LocalLightShadowFrameSetup.ProjectedShadowInfo = ProjectedShadowInfo;
 	LocalLightShadowFrameSetup.LightSceneInfo = LightSceneInfo;
 
+	const float ResolutionLODBiasLocal = GetResolutionLODBiasLocal(ShadowScene.GetLightMobilityFactor(LightSceneInfo->Id));
+
 	// Single page res, at this point we force the VSM to be single page
 	// TODO: this computation does not match up with page marking logic super-well, particularly for long spot lights,
 	//       we can absolutely mirror the page marking calc better, just unclear how much it helps. 
 	//       Also possible to feed back from gpu - which would be more accurate wrt partially visible lights (e.g., a spot going through the ground).
-	//       Of course this creates jumps if visibility changes, which may or may not create unsolvable artifacts.
-	const float BiasedFootprintThreshold = float(FVirtualShadowMap::PageSize) * FMath::Exp2(VirtualShadowMapArray.GetResolutionLODBiasLocal(ShadowScene.GetLightMobilityFactor(LightSceneInfo->Id)));
+	//       Of course this creates jumps if visibility changes, which may or may not create unsolvable artifacts.	
+	const float BiasedFootprintThreshold = float(FVirtualShadowMap::PageSize) * FMath::Exp2(ResolutionLODBiasLocal);
 	const bool bIsDistantLight = CVarDistantLightMode.GetValueOnRenderThread() != 0
 		&& (MaxScreenRadius <= BiasedFootprintThreshold || CVarDistantLightMode.GetValueOnRenderThread() == 2);
 
-
 	const int32 NumMaps = ProjectedShadowInitializer.bOnePassPointLightShadow ? 6 : 1;
+	TSharedPtr<FVirtualShadowMapPerLightCacheEntry> PerLightCacheEntry = CacheManager->FindCreateLightCacheEntry(LightSceneInfo->Id, 0, NumMaps);
+			
+	const float DistantLightForceCacheFootprintFraction = FMath::Clamp(CVarDistantLightForceCacheFootprintFraction.GetValueOnRenderThread(), 0.0f, 1.0f);
+	bool bShouldForceTimeSliceDistantUpdate = (bIsDistantLight && MaxScreenRadius <= BiasedFootprintThreshold * DistantLightForceCacheFootprintFraction);
+	LocalLightShadowFrameSetup.PerLightCacheEntry = PerLightCacheEntry;
+	bool bIsCached = PerLightCacheEntry->UpdateLocal(ProjectedShadowInitializer, bIsDistantLight, CacheManager->IsCacheEnabled(), !bShouldForceTimeSliceDistantUpdate);
+
+	// Update info on the ProjectionShadowInfo; eventually this should all move into local data structures here
+	const int32 VirtualShadowMapId = VirtualShadowMapArray.Allocate(bIsDistantLight, NumMaps);
+	ProjectedShadowInfo->VirtualShadowMapId = VirtualShadowMapId;
+	ProjectedShadowInfo->VirtualShadowMapPerLightCacheEntry = PerLightCacheEntry;
+	ProjectedShadowInfo->bShouldRenderVSM = !PerLightCacheEntry->IsFullyCached();
+
 	for (int32 Index = 0; Index < NumMaps; ++Index)
 	{
-		FVirtualShadowMap* VirtualShadowMap = VirtualShadowMapArray.Allocate(bIsDistantLight);
-		// TODO: redundant
-		ProjectedShadowInfo->VirtualShadowMaps.Add(VirtualShadowMap);
-		LocalLightShadowFrameSetup.VirtualShadowMaps.Add(VirtualShadowMap);
+		const int32 FaceVirtualShadowMapId = VirtualShadowMapId + Index;
+		FVirtualShadowMapCacheEntry& VirtualSmCacheEntry = PerLightCacheEntry->ShadowMapEntries[Index];
+		VirtualSmCacheEntry.Update(VirtualShadowMapArray, *PerLightCacheEntry, FaceVirtualShadowMapId);
+		// Update projection data
+		VirtualSmCacheEntry.ProjectionData = GetLocalLightProjectionShaderData(ResolutionLODBiasLocal, ProjectedShadowInfo, Index);
 	}
 
-	TSharedPtr<FVirtualShadowMapPerLightCacheEntry> PerLightCacheEntry = CacheManager->FindCreateLightCacheEntry(LightSceneInfo->Id);
-	if (PerLightCacheEntry.IsValid())
+	// Only round-robin those that were not invalidated.
+	if (bIsDistantLight && bIsCached)
 	{
-		const float DistantLightForceCacheFootprintFraction = FMath::Clamp(CVarDistantLightForceCacheFootprintFraction.GetValueOnRenderThread(), 0.0f, 1.0f);
-		bool bShouldForceTimeSliceDistantUpdate = (bIsDistantLight && MaxScreenRadius <= BiasedFootprintThreshold * DistantLightForceCacheFootprintFraction);
-		LocalLightShadowFrameSetup.PerLightCacheEntry = PerLightCacheEntry;
-		bool bIsCached = PerLightCacheEntry->UpdateLocal(ProjectedShadowInitializer, bIsDistantLight, CacheManager->IsCacheEnabled(), !bShouldForceTimeSliceDistantUpdate);
-
-		for (int32 Index = 0; Index < NumMaps; ++Index)
-		{
-			FVirtualShadowMap* VirtualShadowMap = LocalLightShadowFrameSetup.VirtualShadowMaps[Index];
-
-			TSharedPtr<FVirtualShadowMapCacheEntry> VirtualSmCacheEntry = PerLightCacheEntry->FindCreateShadowMapEntry(Index);
-			VirtualSmCacheEntry->UpdateLocal(VirtualShadowMap->ID, *PerLightCacheEntry);
-			VirtualShadowMap->VirtualShadowMapCacheEntry = VirtualSmCacheEntry;
-		}
-		// Only round-robin those that were not invalidated.
-		if (bIsDistantLight && bIsCached)
-		{
-			// This priority could be calculated based also on whether the light has actually been invalidated or not (currently not tracked on host).
-			// E.g., all things being equal update those with an animated mesh in, for example. Plus don't update those the don't need it at all.
-			int32 FramesSinceLastRender = int32(Scene.GetFrameNumber()) - int32(PerLightCacheEntry->GetLastScheduledFrameNumber());
-			DistantLightUpdateQueue.Add(-FramesSinceLastRender, LocalLightShadowIndex);
-		}
+		// This priority could be calculated based also on whether the light has actually been invalidated or not (currently not tracked on host).
+		// E.g., all things being equal update those with an animated mesh in, for example. Plus don't update those the don't need it at all.
+		int32 FramesSinceLastRender = int32(Scene.GetFrameNumber()) - int32(PerLightCacheEntry->GetLastScheduledFrameNumber());
+		DistantLightUpdateQueue.Add(-FramesSinceLastRender, LocalLightShadowIndex);
 	}
+
 	return PerLightCacheEntry;
 }
 
diff --git a/Engine/Source/Runtime/Renderer/Private/Shadows/ShadowSceneRenderer.h b/Engine/Source/Runtime/Renderer/Private/Shadows/ShadowSceneRenderer.h
--- a/Engine/Source/Runtime/Renderer/Private/Shadows/ShadowSceneRenderer.h
+++ b/Engine/Source/Runtime/Renderer/Private/Shadows/ShadowSceneRenderer.h
@@ -82,6 +82,8 @@ public:
 	}
 
 private:
+	FVirtualShadowMapProjectionShaderData GetLocalLightProjectionShaderData(float ResolutionLODBiasLocal, const FProjectedShadowInfo* ProjectedShadowInfo, int32 MapIndex) const;
+
 	/**
 	 * Select the budgeted set of distant lights to update this frame.
 	 */
@@ -89,11 +91,10 @@ private:
 
 	struct FLocalLightShadowFrameSetup
 	{
-		TArray<FVirtualShadowMap*, TInlineAllocator<6>> VirtualShadowMaps;
+		TSharedPtr<FVirtualShadowMapPerLightCacheEntry> PerLightCacheEntry;
 		// link to legacy system stuff, to be removed in due time
 		FProjectedShadowInfo* ProjectedShadowInfo = nullptr;
 		FLightSceneInfo* LightSceneInfo = nullptr;
-		TSharedPtr<FVirtualShadowMapPerLightCacheEntry> PerLightCacheEntry;
 	};
 
 	// TODO: maybe we want to keep these in a 1:1 sparse array wrt the light scene infos, for easy crossreference & GPU access (maybe)?
diff --git a/Engine/Source/Runtime/Renderer/Private/VirtualShadowMaps/VirtualShadowMapArray.cpp b/Engine/Source/Runtime/Renderer/Private/VirtualShadowMaps/VirtualShadowMapArray.cpp
--- a/Engine/Source/Runtime/Renderer/Private/VirtualShadowMaps/VirtualShadowMapArray.cpp
+++ b/Engine/Source/Runtime/Renderer/Private/VirtualShadowMaps/VirtualShadowMapArray.cpp
@@ -31,6 +31,8 @@ IMPLEMENT_STATIC_UNIFORM_BUFFER_STRUCT(FVirtualShadowMapUniformParameters, "Virt
 CSV_DEFINE_CATEGORY(VSM, false);
 
 DECLARE_DWORD_COUNTER_STAT(TEXT("VSM Nanite Views (Primary)"), STAT_VSMNaniteViewsPrimary, STATGROUP_ShadowRendering);
+DECLARE_DWORD_COUNTER_STAT(TEXT("VSM Single Page Count"), STAT_VSMSinglePageCount, STATGROUP_ShadowRendering);
+DECLARE_DWORD_COUNTER_STAT(TEXT("VSM Full Count"), STAT_VSMFullCount, STATGROUP_ShadowRendering);
 
 extern int32 GForceInvalidateDirectionalVSM;
 extern TAutoConsoleVariable<float> CVarNaniteMaxPixelsPerEdge;
@@ -85,19 +87,6 @@ static TAutoConsoleVariable<int32> CVarShowStats(
 	ECVF_RenderThreadSafe
 );
 
-static TAutoConsoleVariable<float> CVarResolutionLodBiasLocal(
-	TEXT("r.Shadow.Virtual.ResolutionLodBiasLocal"),
-	0.0f,
-	TEXT("Bias applied to LOD calculations for local lights. -1.0 doubles resolution, 1.0 halves it and so on."),
-	ECVF_Scalability | ECVF_RenderThreadSafe
-);
-
-static TAutoConsoleVariable<float> CVarResolutionLodBiasLocalMoving(
-	TEXT("r.Shadow.Virtual.ResolutionLodBiasLocalMoving"),
-	1.0f,
-	TEXT("Bias applied to LOD calculations for moving local lights. -1.0 doubles resolution, 1.0 halves it and so on."),
-	ECVF_Scalability | ECVF_RenderThreadSafe
-);
 static TAutoConsoleVariable<float> CVarPageDilationBorderSizeDirectional(
 	TEXT("r.Shadow.Virtual.PageDilationBorderSizeDirectional"),
 	0.05f,
@@ -367,48 +356,23 @@ FVirtualShadowMapArray::FVirtualShadowMapArray(FScene& InScene)
 {
 }
 
-// Utility to create metadata linking prev/next shadow map data from frame to frame
-// This is currently necessary since the VirtualShadowMapId's can change
-struct FNextVirtualShadowMapBuffers
-{
-	FRDGBufferRef NextData;		// FNextVirtualShadowMapData
-	uint32 NextVirtualShadowMapDataCount = 0;
-};
-
-static FNextVirtualShadowMapBuffers CreatePrevNextVirtualShadowMapData(
-	FRDGBuilder& GraphBuilder,
-	FVirtualShadowMapArrayCacheManager* CacheManager,
-	const TArray<TUniquePtr<FVirtualShadowMap>, SceneRenderingAllocator>& ShadowMaps)
+void FVirtualShadowMapArray::UpdateNextData(int32 PrevVirtualShadowMapId, int32 NextVirtualShadowMapId, FInt32Point PageOffset)
 {
-	TArray<FNextVirtualShadowMapData, SceneRenderingAllocator> NextData;
-	NextData.Reserve(ShadowMaps.Num());		// Not exact but a reasonable guess
-
 	// Fill in any slots with empty mappings
 	FNextVirtualShadowMapData EmptyData;
 	EmptyData.NextVirtualShadowMapId = INDEX_NONE;
 	EmptyData.PageAddressOffset = FIntVector2(0, 0);
 
-	for (int32 VirtualShadowMapId = 0; VirtualShadowMapId < ShadowMaps.Num(); ++VirtualShadowMapId)
+	// TODO: Some ways to optimize this
+	// Can't use SetNum because we need the empty item initializer which doesn't fit nicely with our shared HLSL definition right now
+	NextData.Reserve(PrevVirtualShadowMapId);
+	while (PrevVirtualShadowMapId >= NextData.Num())
 	{
-		TSharedPtr<FVirtualShadowMapCacheEntry> CacheEntry = ShadowMaps[VirtualShadowMapId].IsValid() ? ShadowMaps[VirtualShadowMapId]->VirtualShadowMapCacheEntry : nullptr;
-		if (CacheEntry != nullptr && CacheEntry->PrevVirtualShadowMapId != INDEX_NONE)
-		{
-			// TODO: Some ways to optimize this
-			while (CacheEntry->PrevVirtualShadowMapId >= NextData.Num())
-			{
-				NextData.Add(EmptyData);
-			}
-
-			FInt32Point PageOffset = CacheEntry->GetPreviousToCurrentPageOffset();
-			NextData[CacheEntry->PrevVirtualShadowMapId].NextVirtualShadowMapId = VirtualShadowMapId;
-			NextData[CacheEntry->PrevVirtualShadowMapId].PageAddressOffset = FIntVector2(PageOffset.X, PageOffset.Y);
-		}
+		NextData.Add(EmptyData);
 	}
 
-	FNextVirtualShadowMapBuffers Result;
-	Result.NextData = CreateStructuredBuffer(GraphBuilder, TEXT("Shadow.Virtual.NextVirtualShadowMapData"), NextData);
-	Result.NextVirtualShadowMapDataCount = NextData.Num();
-	return Result;
+	NextData[PrevVirtualShadowMapId].NextVirtualShadowMapId = NextVirtualShadowMapId;
+	NextData[PrevVirtualShadowMapId].PageAddressOffset = FIntVector2(PageOffset.X, PageOffset.Y);
 }
 
 void FVirtualShadowMapArray::Initialize(FRDGBuilder& GraphBuilder, FVirtualShadowMapArrayCacheManager* InCacheManager, bool bInEnabled, bool bIsSceneCapture)
@@ -447,7 +411,7 @@ void FVirtualShadowMapArray::Initialize(FRDGBuilder& GraphBuilder, FVirtualShado
 	if (bEnabled)
 	{
 		// Always reserve IDs for the single-page SMs.
-		ShadowMaps.SetNum(VSM_MAX_SINGLE_PAGE_SHADOW_MAPS);
+		NumShadowMapSlots = VSM_MAX_SINGLE_PAGE_SHADOW_MAPS;
 
 		// Fixed physical page pool width, we adjust the height to accomodate the requested maximum
 		// NOTE: Row size in pages has to be POT since we use mask & shift in place of integer ops
@@ -517,30 +481,26 @@ void FVirtualShadowMapArray::Initialize(FRDGBuilder& GraphBuilder, FVirtualShado
 	UpdateCachedUniformBuffer(GraphBuilder);
 }
 
-FVirtualShadowMap* FVirtualShadowMapArray::Allocate(bool bSinglePageShadowMap)
+int32 FVirtualShadowMapArray::Allocate(bool bSinglePageShadowMap, int32 Count)
 {
 	check(IsEnabled());
-	FVirtualShadowMap* SM = nullptr;
+	int32 VirtualShadowMapId = INDEX_NONE;
 	if (bSinglePageShadowMap)
 	{
-		if (ensure(NumSinglePageSms < VSM_MAX_SINGLE_PAGE_SHADOW_MAPS))
+		if (ensure((NumSinglePageShadowMaps + Count) <= VSM_MAX_SINGLE_PAGE_SHADOW_MAPS))
 		{
-			SM = new FVirtualShadowMap(NumSinglePageSms, true);
-			ShadowMaps[NumSinglePageSms++] = TUniquePtr<FVirtualShadowMap>(SM);
+			VirtualShadowMapId = NumSinglePageShadowMaps;
+			NumSinglePageShadowMaps += Count;
 		}
 	}
 	else
 	{
-		SM = new FVirtualShadowMap(ShadowMaps.Num(), false);
-		ShadowMaps.Emplace(SM);
+		// Full shadow maps come after single page shadow maps
+		VirtualShadowMapId = NumShadowMapSlots;
+		NumShadowMapSlots += Count;
 	}
 
-	return SM;
-}
-
-float FVirtualShadowMapArray::GetResolutionLODBiasLocal(float LightMobilityFactor) const
-{
-	return InterpolateResolutionBias(CVarResolutionLodBiasLocal.GetValueOnRenderThread(), CVarResolutionLodBiasLocalMoving.GetValueOnRenderThread(), LightMobilityFactor);
+	return VirtualShadowMapId;
 }
 
 FVirtualShadowMapArray::~FVirtualShadowMapArray()
@@ -1112,10 +1072,53 @@ static FRDGTextureRef CreateDebugVisualizationTexture(FRDGBuilder& GraphBuilder,
 	return Texture;
 }
 
-static uint32 GetShadowMapsToAllocate(uint32 NumShadowMaps)
+void FVirtualShadowMapArray::UpdateVisualizeLight(
+	const TConstArrayView<FViewInfo> &Views,
+	const TConstArrayView<FVisibleLightInfo>& VisibleLightInfos)
 {
-	// Round up to powers of two to be friendlier to the buffer pool
-	return FMath::RoundUpToPowerOfTwo(FMath::Max(64U, NumShadowMaps));
+#if !UE_BUILD_SHIPPING
+	for (int32 LightId = 0; LightId < VisibleLightInfos.Num(); ++LightId)
+	{
+		const FVisibleLightInfo& VisibleLightInfo = VisibleLightInfos[LightId];
+		for (const TSharedPtr<FVirtualShadowMapClipmap>& Clipmap : VisibleLightInfo.VirtualShadowMapClipmaps)
+		{
+			for (int32 ViewIndex = 0; ViewIndex < Views.Num(); ViewIndex++)
+			{
+				if (Clipmap->GetDependentView() == Views[ViewIndex].GetPrimaryView())
+				{
+					VisualizeLight[ViewIndex].CheckLight(Clipmap->GetLightSceneInfo().Proxy, Clipmap->GetVirtualShadowMapId());
+				}
+			}
+		}
+
+		for (FProjectedShadowInfo* ProjectedShadowInfo : VisibleLightInfo.AllProjectedShadows)
+		{
+			// NOTE: Specifically checking the VirtualShadowMapId vs HasVirtualShadowMap() here as clipmaps are handled above
+			if (ProjectedShadowInfo->VirtualShadowMapId != INDEX_NONE)
+			{
+				check(ProjectedShadowInfo->CascadeSettings.ShadowSplitIndex == INDEX_NONE);		// We use clipmaps for virtual shadow maps, not cascades
+
+				// NOTE: Virtual shadow maps are never atlased, but verify our assumptions
+				{
+					const FVector4f ClipToShadowUV = ProjectedShadowInfo->GetClipToShadowBufferUvScaleBias();
+					check(ProjectedShadowInfo->BorderSize == 0);
+					check(ProjectedShadowInfo->X == 0);
+					check(ProjectedShadowInfo->Y == 0);
+					const FIntRect ShadowViewRect = ProjectedShadowInfo->GetInnerViewRect();
+					check(ShadowViewRect.Min.X == 0);
+					check(ShadowViewRect.Min.Y == 0);
+					check(ShadowViewRect.Max.X == FVirtualShadowMap::VirtualMaxResolutionXY);
+					check(ShadowViewRect.Max.Y == FVirtualShadowMap::VirtualMaxResolutionXY);
+				}
+
+				for (int32 ViewIndex = 0; ViewIndex < Views.Num(); ViewIndex++)
+				{
+					VisualizeLight[ViewIndex].CheckLight(ProjectedShadowInfo->GetLightSceneInfo().Proxy, ProjectedShadowInfo->VirtualShadowMapId);
+				}
+			}
+		}
+	}
+#endif
 }
 
 void FVirtualShadowMapArray::AppendPhysicalPageList(FRDGBuilder& GraphBuilder, bool bEmptyToAvailable)
@@ -1148,6 +1151,25 @@ void FVirtualShadowMapArray::AppendPhysicalPageList(FRDGBuilder& GraphBuilder, b
 	);
 }
 
+void FVirtualShadowMapArray::UploadProjectionData(FRDGBuilder& GraphBuilder)
+{
+	// Create large enough to hold all the unused elements too (wastes GPU memory but allows direct indexing via the ID)
+	uint32 DataSize = sizeof(FVirtualShadowMapProjectionShaderData) * uint32(GetNumShadowMapSlots());
+	FRDGBufferDesc Desc;
+	Desc.Usage = EBufferUsageFlags::UnorderedAccess | EBufferUsageFlags::ShaderResource | EBufferUsageFlags::ByteAddressBuffer | EBufferUsageFlags::StructuredBuffer;
+	Desc.BytesPerElement = 4;
+	Desc.NumElements = DataSize / 4U;
+	ProjectionDataRDG = GraphBuilder.CreateBuffer(Desc, TEXT("Shadow.Virtual.ProjectionData"));
+		
+	FRDGScatterUploadBuffer Uploader;
+	Uploader.Init(GraphBuilder, GetNumShadowMaps(), sizeof(FVirtualShadowMapProjectionShaderData), false, TEXT("Shadow.Virtual.ProjectionData.UploadBuffer"));
+
+	// Upload data for all cached entries, even if they are not referenced this frame
+	CacheManager->UploadProjectionData(Uploader);
+
+	Uploader.ResourceUploadTo(GraphBuilder, ProjectionDataRDG);
+}
+
 void FVirtualShadowMapArray::BuildPageAllocations(
 	FRDGBuilder& GraphBuilder,
 	const FMinimalSceneTextures& SceneTextures,
@@ -1155,12 +1177,15 @@ void FVirtualShadowMapArray::BuildPageAllocations(
 	const FEngineShowFlags& EngineShowFlags,
 	const FSortedLightSetSceneInfo& SortedLightsInfo,
 	const TConstArrayView<FVisibleLightInfo>& VisibleLightInfos,
-	const TFunctionRef<float(int32 LightId)>& GetLightMobilityFactor,
 	const FSingleLayerWaterPrePassResult* SingleLayerWaterPrePassResult,
 	const FFrontLayerTranslucencyData& FrontLayerTranslucencyData)
 {
 	check(IsEnabled());
 
+	// Before we check on what shadow maps are present, let the cache manager update any that may not be referenced this frame
+	// but may still have cached pages.
+	CacheManager->UpdateUnreferencedCacheEntries(*this);
+
 	if (GetNumShadowMaps() == 0 || Views.Num() == 0)
 	{
 		// Nothing to do
@@ -1212,112 +1237,19 @@ void FVirtualShadowMapArray::BuildPageAllocations(
 			}
 		}
 	}
-#endif //!UE_BUILD_SHIPPING
-		
-	// Store shadow map projection data for each virtual shadow map
-	FRDGScatterUploadBuffer ProjectionDataUploader;
-	ProjectionDataUploader.Init(GraphBuilder, GetNumShadowMaps(), sizeof(FVirtualShadowMapProjectionShaderData), false, TEXT("Shadow.Virtual.ProjectionData.UploadBuffer"));
-
-	for (int32 LightId = 0; LightId < VisibleLightInfos.Num(); ++LightId)
-	{
-		const FVisibleLightInfo& VisibleLightInfo = VisibleLightInfos[LightId];
-		for (const TSharedPtr<FVirtualShadowMapClipmap>& Clipmap : VisibleLightInfo.VirtualShadowMapClipmaps)
-		{
-			// NOTE: Shader assumes all levels from a given clipmap are contiguous
-			int32 ClipmapID = Clipmap->GetVirtualShadowMap()->ID;
-			for (int32 ClipmapLevel = 0; ClipmapLevel < Clipmap->GetLevelCount(); ++ClipmapLevel)
-			{
-				*reinterpret_cast<FVirtualShadowMapProjectionShaderData*>(ProjectionDataUploader.Add_GetRef(ClipmapID + ClipmapLevel)) = Clipmap->GetProjectionShaderData(ClipmapLevel);
-			}
-
-			if (bDebugOutputEnabled)
-			{
-				for (int32 ViewIndex = 0; ViewIndex < Views.Num(); ViewIndex++)
-				{
-					if (Clipmap->GetDependentView() == Views[ViewIndex].GetPrimaryView())
-					{
-						VisualizeLight[ViewIndex].CheckLight(Clipmap->GetLightSceneInfo().Proxy, ClipmapID);
-					}
-				}
-			}
-		}
-
-
-		for (FProjectedShadowInfo* ProjectedShadowInfo : VisibleLightInfo.AllProjectedShadows)
-		{
-			if (ProjectedShadowInfo->VirtualShadowMaps.Num() > 0)
-			{
-				const float ResolutionLodBiasLocal = GetResolutionLODBiasLocal(GetLightMobilityFactor(LightId));
-				check(ProjectedShadowInfo->CascadeSettings.ShadowSplitIndex == INDEX_NONE);		// We use clipmaps for virtual shadow maps, not cascades
-
-				// NOTE: Virtual shadow maps are never atlased, but verify our assumptions
-				{
-					const FVector4f ClipToShadowUV = ProjectedShadowInfo->GetClipToShadowBufferUvScaleBias();
-					check(ProjectedShadowInfo->BorderSize == 0);
-					check(ProjectedShadowInfo->X == 0);
-					check(ProjectedShadowInfo->Y == 0);
-					const FIntRect ShadowViewRect = ProjectedShadowInfo->GetInnerViewRect();
-					check(ShadowViewRect.Min.X == 0);
-					check(ShadowViewRect.Min.Y == 0);
-					check(ShadowViewRect.Max.X == FVirtualShadowMap::VirtualMaxResolutionXY);
-					check(ShadowViewRect.Max.Y == FVirtualShadowMap::VirtualMaxResolutionXY);
-				}
-				TSharedPtr<FVirtualShadowMapPerLightCacheEntry> CacheEntry = CacheManager->FindCreateLightCacheEntry(ProjectedShadowInfo->GetLightSceneInfo().Id);
-
-				int32 NumMaps = ProjectedShadowInfo->bOnePassPointLightShadow ? 6 : 1;
-				for(int32 Index = 0; Index < NumMaps; ++Index)
-				{
-					check(!CacheEntry.IsValid() || CacheEntry->Current.bIsDistantLight == ProjectedShadowInfo->VirtualShadowMaps[Index]->bIsSinglePageSM);
-					
-					uint32 Flags = ProjectedShadowInfo->VirtualShadowMaps[Index]->bIsSinglePageSM ? VSM_PROJ_FLAG_CURRENT_DISTANT_LIGHT : 0U;
-					Flags |= CacheEntry.IsValid() && CacheEntry->IsUncached() ? VSM_PROJ_FLAG_UNCACHED : 0U;
-					int32 ID = ProjectedShadowInfo->VirtualShadowMaps[Index]->ID;
-
-					const FViewMatrices ViewMatrices = ProjectedShadowInfo->GetShadowDepthRenderingViewMatrices(Index, true );
-					const FLargeWorldRenderPosition PreViewTranslation(ProjectedShadowInfo->PreShadowTranslation);
-
-					FVirtualShadowMapProjectionShaderData Data; 
-					Data.TranslatedWorldToShadowViewMatrix		= FMatrix44f(ViewMatrices.GetTranslatedViewMatrix());	// LWC_TODO: Precision loss?
-					Data.ShadowViewToClipMatrix					= FMatrix44f(ViewMatrices.GetProjectionMatrix());
-					Data.TranslatedWorldToShadowUVMatrix		= FMatrix44f(CalcTranslatedWorldToShadowUVMatrix( ViewMatrices.GetTranslatedViewMatrix(), ViewMatrices.GetProjectionMatrix() ));
-					Data.TranslatedWorldToShadowUVNormalMatrix	= FMatrix44f(CalcTranslatedWorldToShadowUVNormalMatrix( ViewMatrices.GetTranslatedViewMatrix(), ViewMatrices.GetProjectionMatrix() ));
-					Data.PreViewTranslationLWCTile				= PreViewTranslation.GetTile();
-					Data.PreViewTranslationLWCOffset			= PreViewTranslation.GetOffset();
-					Data.LightType								= ProjectedShadowInfo->GetLightSceneInfo().Proxy->GetLightType();
-					Data.LightSourceRadius						= ProjectedShadowInfo->GetLightSceneInfo().Proxy->GetSourceRadius();
-					Data.ResolutionLodBias						= ResolutionLodBiasLocal;
-					Data.LightRadius							= ProjectedShadowInfo->GetLightSceneInfo().Proxy->GetRadius();
-					Data.Flags									= Flags;
-					ProjectionDataUploader.Add(ID, &Data);
-				}
-
-				if (bDebugOutputEnabled)
-				{
-					for (int32 ViewIndex = 0; ViewIndex < Views.Num(); ViewIndex++)
-					{
-						VisualizeLight[ViewIndex].CheckLight(ProjectedShadowInfo->GetLightSceneInfo().Proxy, ProjectedShadowInfo->VirtualShadowMaps[0]->ID);
-					}
-				}
-			}
-		}
-	}
 
-	{
-		// Create large enough to hold all the unused elements too (wastes GPU memory but allows direct indexing via the ID)
-		uint32 DataSize = sizeof(FVirtualShadowMapProjectionShaderData) * uint32(ShadowMaps.Num());
+	UpdateVisualizeLight(Views, VisibleLightInfos);
+#endif //!UE_BUILD_SHIPPING
 
-		FRDGBufferDesc Desc;
-		Desc.Usage = EBufferUsageFlags::UnorderedAccess | EBufferUsageFlags::ShaderResource | EBufferUsageFlags::ByteAddressBuffer | EBufferUsageFlags::StructuredBuffer;
-		Desc.BytesPerElement = 4;
-		Desc.NumElements = DataSize / 4U;
+	UploadProjectionData(GraphBuilder);
 
-		ProjectionDataRDG = GraphBuilder.CreateBuffer(Desc, TEXT("Shadow.Virtual.ProjectionData"));
-	}
-	ProjectionDataUploader.ResourceUploadTo(GraphBuilder, ProjectionDataRDG);
+	// Stats
+	SET_DWORD_STAT(STAT_VSMSinglePageCount, GetNumSinglePageShadowMaps());
+	SET_DWORD_STAT(STAT_VSMFullCount, GetNumFullShadowMaps());
 
 	UniformParameters.NumFullShadowMaps = GetNumFullShadowMaps();
 	UniformParameters.NumSinglePageShadowMaps = GetNumSinglePageShadowMaps();
-	UniformParameters.NumShadowMapSlots = ShadowMaps.Num();
+	UniformParameters.NumShadowMapSlots = GetNumShadowMapSlots();
 	UniformParameters.ProjectionData = GraphBuilder.CreateSRV(ProjectionDataRDG);
 
 	UniformParameters.bExcludeNonNaniteFromCoarsePages = !CVarCoarsePagesIncludeNonNanite.GetValueOnRenderThread();
@@ -1346,7 +1278,6 @@ void FVirtualShadowMapArray::BuildPageAllocations(
 	}
 	
 	// We potentially over-allocate these to avoid too many different allocation sizes each frame
-	const uint32 NumShadowMapsToAllocate = GetShadowMapsToAllocate(GetNumShadowMaps());
 	const int32 NumPageFlagsToAllocate = FMath::RoundUpToPowerOfTwo(FMath::Max(128 * 1024, GetNumFullShadowMaps() * int32(FVirtualShadowMap::PageTableSize) + int32(VSM_MAX_SINGLE_PAGE_SHADOW_MAPS)));
 
 	// Create and clear the requested page flags	
@@ -1361,7 +1292,7 @@ void FVirtualShadowMapArray::BuildPageAllocations(
 	const uint32 PhysicalPageListsCount = 4;		// See VirtualShadowMapPhysicalPageManagement.usf
 	PhysicalPageListsRDG = GraphBuilder.CreateBuffer(FRDGBufferDesc::CreateStructuredDesc(sizeof(int32), PhysicalPageListsCount * ItemsPerPhysicalPageList), TEXT("Shadow.Virtual.PhysicalPageLists"));
 
-	const uint32 NumPageRects = ShadowMaps.Num() * FVirtualShadowMap::MaxMipLevels;
+	const uint32 NumPageRects = GetNumShadowMapSlots() * FVirtualShadowMap::MaxMipLevels;
 	const uint32 NumPageRectsToAllocate = FMath::RoundUpToPowerOfTwo(NumPageRects);
 	PageRectBoundsRDG = GraphBuilder.CreateBuffer(FRDGBufferDesc::CreateStructuredDesc(sizeof(FIntVector4), NumPageRectsToAllocate), TEXT("Shadow.Virtual.PageRectBounds"));
 	{
@@ -1398,7 +1329,7 @@ void FVirtualShadowMapArray::BuildPageAllocations(
 				if (Clipmap->GetDependentView() == Views[ViewIndex].GetPrimaryView())
 				{
 					// NOTE: Shader assumes all levels from a given clipmap are contiguous
-					int32 ClipmapID = Clipmap->GetVirtualShadowMap()->ID;
+					int32 ClipmapID = Clipmap->GetVirtualShadowMapId();
 					DirectionalLightIds.Add(ClipmapID);
 				}
 			}
@@ -1589,14 +1520,9 @@ void FVirtualShadowMapArray::BuildPageAllocations(
 	AddClearUAVPass(GraphBuilder, GraphBuilder.CreateUAV(PageFlagsRDG), 0);
 
 	// Update cached or newly invalidated pages with respect to the new requests
-	{
+	{	
 		// Cached data from previous frames is available and valid
 		const bool bCacheDataAvailable = CacheManager->IsCacheDataAvailable();
-		FNextVirtualShadowMapBuffers NextShadowData;
-		if (bCacheDataAvailable)
-		{
-			NextShadowData = CreatePrevNextVirtualShadowMapData(GraphBuilder, CacheManager, ShadowMaps);
-		}
 
 		FUpdatePhysicalPages::FParameters* PassParameters = GraphBuilder.AllocParameters<FUpdatePhysicalPages::FParameters>();
 		PassParameters->VirtualShadowMap		= GetUncachedUniformBuffer(GraphBuilder);
@@ -1605,12 +1531,15 @@ void FVirtualShadowMapArray::BuildPageAllocations(
 
 		if (bCacheDataAvailable)
 		{
+			// Upload our prev -> next shadow data mapping (FNextVirtualShadowMapData) to the GPU
+			FRDGBufferRef NextVirtualShadowMapData = CreateStructuredBuffer(GraphBuilder, TEXT("Shadow.Virtual.NextVirtualShadowMapData"), NextData);
+
 			PassParameters->PageRequestFlags		 = GraphBuilder.CreateSRV(PageRequestFlagsRDG);
 			PassParameters->OutPageTable			 = GraphBuilder.CreateUAV(PageTableRDG);
 			PassParameters->OutPageFlags			 = GraphBuilder.CreateUAV(PageFlagsRDG);
 			PassParameters->PrevPhysicalPageLists    = GraphBuilder.CreateSRV(GraphBuilder.RegisterExternalBuffer(CacheManager->GetPrevBuffers().PhysicalPageLists));
-			PassParameters->NextVirtualShadowMapData = GraphBuilder.CreateSRV(NextShadowData.NextData);
-			PassParameters->NextVirtualShadowMapDataCount = NextShadowData.NextVirtualShadowMapDataCount;
+			PassParameters->NextVirtualShadowMapData = GraphBuilder.CreateSRV(NextVirtualShadowMapData);
+			PassParameters->NextVirtualShadowMapDataCount = NextData.Num();
 			PassParameters->bKeepOnlyRequestedPages  = CVarCacheKeepOnlyRequestedPages.GetValueOnRenderThread();
 			PassParameters->bDynamicPageInvalidation = 1;
 #if !UE_BUILD_SHIPPING
@@ -1960,7 +1889,6 @@ void FVirtualShadowMapArray::CreateMipViews( TArray<Nanite::FPackedView, SceneRe
 	// strategy: 
 	// 1. Use the cull pass to generate copies of every node for every view needed.
 	// [2. Fabricate a HZB array?]
-	ensure(Views.Num() <= ShadowMaps.Num());
 	
 	const int32 NumPrimaryViews = Views.Num();
 
@@ -1975,7 +1903,7 @@ void FVirtualShadowMapArray::CreateMipViews( TArray<Nanite::FPackedView, SceneRe
 	{
 		const Nanite::FPackedView& PrimaryView = Views[ViewIndex];
 		
-		ensure( PrimaryView.TargetLayerIdX_AndMipLevelY_AndNumMipLevelsZ.X >= 0 && PrimaryView.TargetLayerIdX_AndMipLevelY_AndNumMipLevelsZ.X < ShadowMaps.Num() );
+		ensure( PrimaryView.TargetLayerIdX_AndMipLevelY_AndNumMipLevelsZ.X >= 0 && PrimaryView.TargetLayerIdX_AndMipLevelY_AndNumMipLevelsZ.X < GetNumShadowMapSlots() );
 		ensure( PrimaryView.TargetLayerIdX_AndMipLevelY_AndNumMipLevelsZ.Y == 0 );
 		ensure( PrimaryView.TargetLayerIdX_AndMipLevelY_AndNumMipLevelsZ.Z > 0 && PrimaryView.TargetLayerIdX_AndMipLevelY_AndNumMipLevelsZ.Z <= FVirtualShadowMap::MaxMipLevels );
 		
@@ -2984,9 +2912,9 @@ void FVirtualShadowMapArray::RenderVirtualShadowMapsNonNanite(FRDGBuilder& Graph
 
 			ShaderPrint::SetParameters(GraphBuilder, View.ShaderPrintData, PassParameters->ShaderPrintStruct);
 			//PassParameters->VirtualShadowMap = GetUncachedUniformBuffer(GraphBuilder);
-			PassParameters->ShadowMapIdRangeStart = Clipmap->GetVirtualShadowMap(0)->ID;
+			PassParameters->ShadowMapIdRangeStart = Clipmap->GetVirtualShadowMapId();
 			// Note: assumes range!
-			PassParameters->ShadowMapIdRangeEnd = Clipmap->GetVirtualShadowMap(0)->ID + Clipmap->GetLevelCount();
+			PassParameters->ShadowMapIdRangeEnd = Clipmap->GetVirtualShadowMapId() + Clipmap->GetLevelCount();
 			PassParameters->PageRectBounds = GraphBuilder.CreateSRV(PageRectBoundsRDG);
 
 			auto ComputeShader = View.ShaderMap->GetShader<FVirtualSmPrintClipmapStatsCS>();
@@ -3209,10 +3137,10 @@ uint32 FVirtualShadowMapArray::AddRenderViews(const TSharedPtr<FVirtualShadowMap
 
 	for (int32 ClipmapLevelIndex = 0; ClipmapLevelIndex < Clipmap->GetLevelCount(); ++ClipmapLevelIndex)
 	{
-		FVirtualShadowMap* VirtualShadowMap = Clipmap->GetVirtualShadowMap(ClipmapLevelIndex);
+		int32 VirtualShadowMapId = Clipmap->GetVirtualShadowMapId(ClipmapLevelIndex);
 
 		Nanite::FPackedViewParams Params = BaseParams;
-		Params.TargetLayerIndex = VirtualShadowMap->ID;
+		Params.TargetLayerIndex = VirtualShadowMapId;
 		Params.ViewMatrices = Clipmap->GetViewMatrices(ClipmapLevelIndex);
 		Params.PrevTargetLayerIndex = INDEX_NONE;
 		Params.PrevViewMatrices = Params.ViewMatrices;
@@ -3220,11 +3148,11 @@ uint32 FVirtualShadowMapArray::AddRenderViews(const TSharedPtr<FVirtualShadowMap
 
 		if (CacheEntry)
 		{
-			TSharedPtr<FVirtualShadowMapCacheEntry> LevelEntry = CacheEntry->ShadowMapEntries[Clipmap->GetClipmapLevel(ClipmapLevelIndex)];
+			FVirtualShadowMapCacheEntry& LevelEntry = CacheEntry->ShadowMapEntries[ClipmapLevelIndex];
 
 			if (bSetHZBParams)
 			{
-				LevelEntry->SetHZBViewParams(Params);
+				LevelEntry.SetHZBViewParams(Params);
 			}
 
 			// If we're going to generate a new HZB this frame, save the associated metadata
@@ -3234,7 +3162,7 @@ uint32 FVirtualShadowMapArray::AddRenderViews(const TSharedPtr<FVirtualShadowMap
 				HZBMeta.TargetLayerIndex = Params.TargetLayerIndex;
 				HZBMeta.ViewMatrices = Params.ViewMatrices;
 				HZBMeta.ViewRect = Params.ViewRect;
-				LevelEntry->CurrentHZBMetadata = HZBMeta;
+				LevelEntry.CurrentHZBMetadata = HZBMeta;
 			}
 		}
 
@@ -3289,31 +3217,28 @@ uint32 FVirtualShadowMapArray::AddRenderViews(const FProjectedShadowInfo* Projec
 	BaseParams.DrawDistanceOrigin = ClosestCullingViewOrigin;
 
 	TSharedPtr<FVirtualShadowMapPerLightCacheEntry> CacheEntry = ProjectedShadowInfo->VirtualShadowMapPerLightCacheEntry;
-	if (CacheEntry.IsValid())
-	{
-		CacheEntry->MarkRendered(Scene.GetFrameNumber());
-
-		BaseParams.Flags |= CacheEntry->IsUncached() ? NANITE_VIEW_FLAG_UNCACHED : 0U;
-	}
+	check(CacheEntry.IsValid())
+	CacheEntry->MarkRendered(Scene.GetFrameNumber());
+	BaseParams.Flags |= CacheEntry->IsUncached() ? NANITE_VIEW_FLAG_UNCACHED : 0U;
 
 	int32 NumMaps = ProjectedShadowInfo->bOnePassPointLightShadow ? 6 : 1;
 	for (int32 Index = 0; Index < NumMaps; ++Index)
 	{
-		FVirtualShadowMap* VirtualShadowMap = ProjectedShadowInfo->VirtualShadowMaps[Index];
+		int32 VirtualShadowMapId = ProjectedShadowInfo->VirtualShadowMapId + Index;
 
 		Nanite::FPackedViewParams Params = BaseParams;
-		Params.TargetLayerIndex = VirtualShadowMap->ID;
+		Params.TargetLayerIndex = VirtualShadowMapId;
 		Params.ViewMatrices = ProjectedShadowInfo->GetShadowDepthRenderingViewMatrices(Index, true);
 		Params.RangeBasedCullingDistance = ProjectedShadowInfo->GetLightSceneInfo().Proxy->GetRadius();
 		// TODO: MaxPixelsPerEdgeMultipler
 
 		if (CacheEntry)
 		{
-			TSharedPtr<FVirtualShadowMapCacheEntry> LevelEntry = CacheEntry->ShadowMapEntries[Index];
+			FVirtualShadowMapCacheEntry& LevelEntry = CacheEntry->ShadowMapEntries[Index];
 			
 			if (bSetHZBParams)
 			{
-				LevelEntry->SetHZBViewParams(Params);
+				LevelEntry.SetHZBViewParams(Params);
 			}
 
 			// If we're going to generate a new HZB this frame, save the associated metadata
@@ -3323,7 +3248,7 @@ uint32 FVirtualShadowMapArray::AddRenderViews(const FProjectedShadowInfo* Projec
 				HZBMeta.TargetLayerIndex = Params.TargetLayerIndex;
 				HZBMeta.ViewMatrices = Params.ViewMatrices;
 				HZBMeta.ViewRect = Params.ViewRect;
-				LevelEntry->CurrentHZBMetadata = HZBMeta;
+				LevelEntry.CurrentHZBMetadata = HZBMeta;
 			}
 		}
 
diff --git a/Engine/Source/Runtime/Renderer/Private/VirtualShadowMaps/VirtualShadowMapArray.h b/Engine/Source/Runtime/Renderer/Private/VirtualShadowMaps/VirtualShadowMapArray.h
--- a/Engine/Source/Runtime/Renderer/Private/VirtualShadowMaps/VirtualShadowMapArray.h
+++ b/Engine/Source/Runtime/Renderer/Private/VirtualShadowMaps/VirtualShadowMapArray.h
@@ -73,14 +73,9 @@ public:
 	
 	static_assert(MaxMipLevels <= 8, ">8 mips requires more PageFlags bits. See VSM_PAGE_FLAGS_BITS_PER_HMIP in PageAccessCommon.ush");
 
-	FVirtualShadowMap(uint32 InID, bool bInIsSinglePageSM) : ID(InID), bIsSinglePageSM(bInIsSinglePageSM)
-	{
-	}
-
-	const int32 ID = INDEX_NONE;
-	const bool bIsSinglePageSM;
-
-	TSharedPtr<FVirtualShadowMapCacheEntry> VirtualShadowMapCacheEntry;
+	// TODO: Currently only used for these constants... probably rename the cache structure to virtual shadow map now instead
+private:
+	FVirtualShadowMap() {}
 };
 
 // Useful data for both the page mapping shader and the projection shader
@@ -220,16 +215,30 @@ public:
 		return bEnabled;
 	}
 
-	FVirtualShadowMap* Allocate(bool bSinglePageShadowMap);
+	// Returns the first in a continuously allocated range of new VirtualShadowMapIds
+	int32 Allocate(bool bSinglePageShadowMap, int32 Count);
+
+	// TODO: Can probably make this 1:1 with allocate directly
+	void UpdateNextData(int32 PrevVirtualShadowMapId, int32 CurrentVirtualShadowMapId, FInt32Point PageOffset = FInt32Point(0, 0));
+
+	static bool IsSinglePage(int VirtualShadowMapId)
+	{
+		return (VirtualShadowMapId < VSM_MAX_SINGLE_PAGE_SHADOW_MAPS);
+	}
+
+	int32 GetNumShadowMapSlots() const
+	{
+		return NumShadowMapSlots;
+	}
 
 	int32 GetNumFullShadowMaps() const
 	{
-		return FMath::Max(ShadowMaps.Num() - int32(VSM_MAX_SINGLE_PAGE_SHADOW_MAPS), 0);
+		return FMath::Max(GetNumShadowMapSlots() - int32(VSM_MAX_SINGLE_PAGE_SHADOW_MAPS), 0);
 	}
 
 	int32 GetNumSinglePageShadowMaps() const
 	{
-		return NumSinglePageSms;
+		return NumSinglePageShadowMaps;
 	}
 
 	/**
@@ -241,11 +250,6 @@ public:
 		return GetNumFullShadowMaps() + GetNumSinglePageShadowMaps();
 	}
 
-	/**
-	 * Get configured LOD bias for the local lights, LightMobilityFactor [0,1] selects between (by interpolation) the two LOD biases for static (non-moving) and moving lights.
-	 */
-	float GetResolutionLODBiasLocal(float LightMobilityFactor) const;
-
 	// Raw size of the physical pool, including both static and dynamic pages (if enabled)
 	FIntPoint GetPhysicalPoolSize() const;
 	// Size of HZB (level 0)
@@ -270,7 +274,6 @@ public:
 		const FEngineShowFlags& EngineShowFlags,
 		const FSortedLightSetSceneInfo& SortedLights, 
 		const TConstArrayView<FVisibleLightInfo>& VisibleLightInfos,
-		const TFunctionRef<float(int32 LightId)>& GetLightMobilityFactor, // TODO: propagate setup in a better(tm) way
 		const FSingleLayerWaterPrePassResult* SingleLayerWaterPrePassResult,
 		const FFrontLayerTranslucencyData& FrontLayerTranslucencyData);
 
@@ -384,15 +387,25 @@ public:
 	TArray<FVirtualShadowMapVisualizeLightSearch> VisualizeLight;
 
 private:
+	void UpdateVisualizeLight(
+		const TConstArrayView<FViewInfo> &Views,
+		const TConstArrayView<FVisibleLightInfo>& VisibleLightInfos);
+
+	void UploadProjectionData(FRDGBuilder& GraphBuilder);
+
 	void AppendPhysicalPageList(FRDGBuilder& GraphBuilder, bool bEmptyToAvailable);
 
 	uint32 AddRenderViews(const TSharedPtr<FVirtualShadowMapClipmap>& Clipmap, float LODScaleFactor, bool bSetHzbParams, bool bUpdateHZBMetaData, const FVector &CullingViewOrigin, TArray<Nanite::FPackedView, SceneRenderingAllocator>& OutVirtualShadowViews);
 
 	TRDGUniformBufferRef<FVirtualShadowMapUniformParameters> GetUncachedUniformBuffer(FRDGBuilder& GraphBuilder) const;
 	void UpdateCachedUniformBuffer(FRDGBuilder& GraphBuilder);
+			
+	// Track mapping of previous VSM data -> current frame VSM data
+	// This is primarily an indirection that allows us to reallocate/repack VirtualShadowMapIds each frame
+	TArray<FNextVirtualShadowMapData, SceneRenderingAllocator> NextData;
 
-	TArray<TUniquePtr<FVirtualShadowMap>, SceneRenderingAllocator> ShadowMaps;
-	int32 NumSinglePageSms = 0;
+	int32 NumShadowMapSlots = 0;
+	int32 NumSinglePageShadowMaps = 0;
 
 	// Cached copy of the latest uniform parameters
 	// Gets created in dummy form at initialization time, then updated after VSM data is computed
diff --git a/Engine/Source/Runtime/Renderer/Private/VirtualShadowMaps/VirtualShadowMapCacheManager.cpp b/Engine/Source/Runtime/Renderer/Private/VirtualShadowMaps/VirtualShadowMapCacheManager.cpp
--- a/Engine/Source/Runtime/Renderer/Private/VirtualShadowMaps/VirtualShadowMapCacheManager.cpp
+++ b/Engine/Source/Runtime/Renderer/Private/VirtualShadowMaps/VirtualShadowMapCacheManager.cpp
@@ -17,6 +17,7 @@
 #define VSM_LOG_STATIC_CACHING 0
 
 CSV_DECLARE_CATEGORY_EXTERN(VSM);
+DECLARE_DWORD_COUNTER_STAT(TEXT("VSM Unreferenced Lights"), STAT_VSMUnreferencedLights, STATGROUP_ShadowRendering);
 
 static TAutoConsoleVariable<int32> CVarAccumulateStats(
 	TEXT("r.Shadow.Virtual.AccumulateStats"),
@@ -70,6 +71,13 @@ static FAutoConsoleVariableRef  CVarForceInvalidateDirectionalVSM(
 	TEXT("Forces the clipmap to always invalidate, useful to emulate a moving sun to avoid misrepresenting cache performance."),
 	ECVF_RenderThreadSafe);
 
+static TAutoConsoleVariable<int32> CVarCacheMaxUnreferencedLightAge(
+	TEXT("r.Shadow.Virtual.Cache.MaxUnreferencedLightAge"),
+	0,
+	TEXT("The number of frames to keep around cached pages from lights that are unreferenced (usually due to being offscreen or otherwise culled). 0=disabled.\n")
+	TEXT("Higher values potentially allow for more physical page reuse when cache space is available, but setting this too high can add overhead due to maintaining the extra entries."),
+	ECVF_RenderThreadSafe
+);
 
 static TAutoConsoleVariable<int32> CVarVSMReservedResource(
 	TEXT("r.Shadow.Virtual.AllocatePagePoolAsReservedResource"),
@@ -78,20 +86,20 @@ static TAutoConsoleVariable<int32> CVarVSMReservedResource(
 	ECVF_RenderThreadSafe
 );
 
-void FVirtualShadowMapCacheEntry::UpdateClipmap(
+void FVirtualShadowMapCacheEntry::UpdateClipmapLevel(
+	FVirtualShadowMapArray& VirtualShadowMapArray,
+	const FVirtualShadowMapPerLightCacheEntry& PerLightEntry,
 	int32 VirtualShadowMapId,
 	FInt64Point PageSpaceLocation,
-	FInt64Point ClipmapCornerOffset,
 	double LevelRadius,
 	double ViewCenterZ,
-	// NOTE: ViewRadiusZ must be constant for a given clipmap level
-	double ViewRadiusZ,
-	const FVirtualShadowMapPerLightCacheEntry& PerLightEntry)
+	double ViewRadiusZ)
 {
-	PrevPageSpaceLocation = CurrentPageSpaceLocation;
+	int32 PrevVirtualShadowMapId = CurrentVirtualShadowMapId;
+	FInt64Point PrevPageSpaceLocation = Clipmap.PageSpaceLocation;
 	PrevHZBMetadata = CurrentHZBMetadata;
 
-	bool bCacheValid = (CurrentVirtualShadowMapId != INDEX_NONE);
+	bool bCacheValid = (PrevVirtualShadowMapId != INDEX_NONE);
 	
 	if (bCacheValid && GClipmapPanning == 0)
 	{
@@ -116,60 +124,62 @@ void FVirtualShadowMapCacheEntry::UpdateClipmap(
 	}
 
 	// Not valid if it was never rendered
-	if (PerLightEntry.Prev.RenderedFrameNumber < 0)
-	{
-		bCacheValid = false;
-	}
+	bCacheValid = bCacheValid && (PerLightEntry.Prev.RenderedFrameNumber >= 0);
 
-	bool bRadiusMatches = (ViewRadiusZ == Clipmap.ViewRadiusZ);
+	// Not valid if radius has changed
+	bCacheValid = bCacheValid && (ViewRadiusZ == Clipmap.ViewRadiusZ);
 
-	if (bCacheValid && bRadiusMatches)
+	if (!bCacheValid)
 	{
-		PrevVirtualShadowMapId = CurrentVirtualShadowMapId;
+		Clipmap.ViewCenterZ = ViewCenterZ;
+		Clipmap.ViewRadiusZ = ViewRadiusZ;
 	}
 	else
 	{
-		if (bCacheValid && !bRadiusMatches)
-		{
-			// These really should be exact by construction currently
-			UE_LOG(LogRenderer, Warning, TEXT("Invalidated clipmap level (VSM %d) due to Z radius mismatch"), VirtualShadowMapId);
-		}
-
-		// New cached level
-		PrevVirtualShadowMapId = INDEX_NONE;
-		Clipmap.ViewCenterZ = ViewCenterZ;
-		Clipmap.ViewRadiusZ = ViewRadiusZ;
+		// NOTE: Leave the view center and radius where they were previously for the cached page
+		FInt64Point CurrentToPreviousPageOffset(PageSpaceLocation - PrevPageSpaceLocation);
+		VirtualShadowMapArray.UpdateNextData(PrevVirtualShadowMapId, VirtualShadowMapId, FInt32Point(CurrentToPreviousPageOffset));
 	}
 	
 	CurrentVirtualShadowMapId = VirtualShadowMapId;
-	CurrentPageSpaceLocation = PageSpaceLocation;
-	// Current HZB metadata gets updated during rendering
-
-	Clipmap.PrevClipmapCornerOffset = Clipmap.CurrentClipmapCornerOffset;
-	Clipmap.CurrentClipmapCornerOffset = ClipmapCornerOffset;
+	Clipmap.PageSpaceLocation = PageSpaceLocation;
 }
 
-void FVirtualShadowMapCacheEntry::UpdateLocal(int32 VirtualShadowMapId, const FVirtualShadowMapPerLightCacheEntry& PerLightEntry)
+void FVirtualShadowMapCacheEntry::Update(
+	FVirtualShadowMapArray& VirtualShadowMapArray,
+	const FVirtualShadowMapPerLightCacheEntry& PerLightEntry,
+	int32 VirtualShadowMapId)
 {
 	// Swap previous frame data over.
-	PrevPageSpaceLocation = FInt64Point(0, 0);		// Not used for local lights
-	PrevVirtualShadowMapId = CurrentVirtualShadowMapId;
+	int32 PrevVirtualShadowMapId = CurrentVirtualShadowMapId;
+
+	// TODO: This is pretty wrong specifically for unreferenced lights, as the VSM IDs will have changed and not been updated since
+	// this gets updated by rendering! Need to figure out a better way to track this data, and probably not here...
 	PrevHZBMetadata = CurrentHZBMetadata;
+	
+	bool bCacheValid = (PrevVirtualShadowMapId != INDEX_NONE);
 
 	// Not valid if it was never rendered
-	if (PerLightEntry.Prev.RenderedFrameNumber < 0)
+	bCacheValid = bCacheValid && (PerLightEntry.Prev.RenderedFrameNumber >= 0);
+
+	if (bCacheValid)
 	{
-		PrevVirtualShadowMapId = INDEX_NONE;
+		// Invalidate on transition between single page and full
+		bool bPrevSinglePage = FVirtualShadowMapArray::IsSinglePage(PrevVirtualShadowMapId);
+		bool bCurrentSinglePage = FVirtualShadowMapArray::IsSinglePage(VirtualShadowMapId);
+		if (bPrevSinglePage != bCurrentSinglePage)
+		{
+			bCacheValid = false;
+		}
 	}
 
-	// Invalidate on transition from distant to full
-	if (!PerLightEntry.Current.bIsDistantLight && PerLightEntry.Prev.bIsDistantLight)
+	if (bCacheValid)
 	{
-		PrevVirtualShadowMapId = INDEX_NONE;
+		// Update previous/next frame mapping if we have a valid cached shadow map
+		VirtualShadowMapArray.UpdateNextData(PrevVirtualShadowMapId, VirtualShadowMapId);
 	}
 
 	CurrentVirtualShadowMapId = VirtualShadowMapId;
-	CurrentPageSpaceLocation = FInt64Point(0, 0);		// Not used for local lights
 	// Current HZB metadata gets updated during rendering
 }
 
@@ -180,28 +190,25 @@ void FVirtualShadowMapCacheEntry::SetHZBViewParams(Nanite::FPackedViewParams& Ou
 	OutParams.Flags |= NANITE_VIEW_FLAG_HZBTEST;
 }
 
-
-void FVirtualShadowMapCacheEntry::Invalidate()
-{
-	PrevVirtualShadowMapId = INDEX_NONE;
-}
-
-void FVirtualShadowMapPerLightCacheEntry::UpdateClipmap(const FMatrix& WorldToLight)
+void FVirtualShadowMapPerLightCacheEntry::UpdateClipmap(const FVector& LightDirection, int FirstLevel)
 {
 	Prev.bIsUncached = Current.bIsUncached;
 	Prev.RenderedFrameNumber = FMath::Max(Prev.RenderedFrameNumber, Current.RenderedFrameNumber);
 	Current.RenderedFrameNumber = -1;
-	if (GForceInvalidateDirectionalVSM != 0 || WorldToLight != ClipmapCacheKey.WorldToLight)
+
+	if (GForceInvalidateDirectionalVSM != 0 ||
+		LightDirection != ClipmapCacheKey.LightDirection ||
+		FirstLevel != ClipmapCacheKey.FirstLevel)
 	{
 		Prev.RenderedFrameNumber = -1;
-		//UE_LOG(LogRenderer, Display, TEXT("Invalidated clipmap level (VSM %d) due to light movement"), VirtualShadowMapId);
 	}
-	ClipmapCacheKey.WorldToLight = WorldToLight;
+	ClipmapCacheKey.LightDirection = LightDirection;
+	ClipmapCacheKey.FirstLevel = FirstLevel;
 
 	Current.bIsUncached = GForceInvalidateDirectionalVSM != 0 || Prev.RenderedFrameNumber < 0;
 
-	// On transition from uncached -> cached we must invalidate since the static pages are not initialized at this point.
-	if (!Current.bIsUncached && Prev.bIsUncached)
+	// On transition between uncached <-> cached we must invalidate since the static pages may not be initialized
+	if (Current.bIsUncached != Prev.bIsUncached)
 	{
 		Prev.RenderedFrameNumber = -1;
 	}
@@ -230,8 +237,8 @@ bool FVirtualShadowMapPerLightCacheEntry::UpdateLocal(const FProjectedShadowInit
 	Current.ScheduledFrameNumber = -1;
 	Current.bIsUncached = Prev.RenderedFrameNumber < 0;
 
-	// On transition from uncached -> cached we must invalidate since the static pages are not initialized at this point.
-	if (!Current.bIsUncached && Prev.bIsUncached)
+	// On transition between uncached <-> cached we must invalidate since the static pages may not be initialized
+	if (Current.bIsUncached != Prev.bIsUncached)
 	{
 		Prev.RenderedFrameNumber = -1;
 	}
@@ -242,11 +249,6 @@ bool FVirtualShadowMapPerLightCacheEntry::UpdateLocal(const FProjectedShadowInit
 void FVirtualShadowMapPerLightCacheEntry::Invalidate()
 {
 	Prev.RenderedFrameNumber = -1;
-
-	for (TSharedPtr<FVirtualShadowMapCacheEntry>& Entry : ShadowMapEntries)
-	{
-		Entry->Invalidate();
-	}
 }
 
 static inline uint32 EncodeInstanceInvalidationPayload(bool bInvalidateStaticPage, int32 ClipmapVirtualShadowMapId = INDEX_NONE)
@@ -288,7 +290,7 @@ void FVirtualShadowMapArrayCacheManager::FInvalidatingPrimitiveCollector::AddDyn
 	// SIDE EFFECT GLOBAL
 	GPUScene.DynamicPrimitiveInstancesToInvalidate.Reset();
 
-	for (auto& CacheEntry : Manager.PrevCacheEntries)
+	for (auto& CacheEntry : Manager.CacheEntries)
 	{
 		for (const FVirtualShadowMapPerLightCacheEntry::FInstanceRange& Range : CacheEntry.Value->PrimitiveInstancesToInvalidate)
 		{
@@ -296,12 +298,9 @@ void FVirtualShadowMapArrayCacheManager::FInvalidatingPrimitiveCollector::AddDyn
 			// TODO: maybe add permutation so we can strip the loop completely.
 			for (const auto& SmCacheEntry : CacheEntry.Value->ShadowMapEntries)
 			{
-				if (SmCacheEntry.IsValid())
-				{
-					// TODO: Do we ever need to invalidate these "post" update?
-					Instances.Add(Range.InstanceSceneDataOffset, Range.NumInstanceSceneDataEntries,
-						EncodeInstanceInvalidationPayload(Range.bInvalidateStaticPage, SmCacheEntry->CurrentVirtualShadowMapId));
-				}
+				// TODO: Do we ever need to invalidate these "post" update?
+				Instances.Add(Range.InstanceSceneDataOffset, Range.NumInstanceSceneDataEntries,
+					EncodeInstanceInvalidationPayload(Range.bInvalidateStaticPage, SmCacheEntry.CurrentVirtualShadowMapId));
 			}
 		}
 		CacheEntry.Value->PrimitiveInstancesToInvalidate.Reset();
@@ -402,7 +401,7 @@ void FVirtualShadowMapArrayCacheManager::FInvalidatingPrimitiveCollector::AddInv
 		return;
 	}
 
-	for (auto& CacheEntry : Manager.PrevCacheEntries)
+	for (auto& CacheEntry : Manager.CacheEntries)
 	{
 		TBitArray<>& CachedPrimitives = CacheEntry.Value->CachedPrimitives;
 		if (bIsNaniteMesh || (PersistentPrimitiveIndex.Index < CachedPrimitives.Num() && CachedPrimitives[PersistentPrimitiveIndex.Index]))
@@ -418,12 +417,9 @@ void FVirtualShadowMapArrayCacheManager::FInvalidatingPrimitiveCollector::AddInv
 			// TODO: maybe add permutation so we can strip the loop completely.
 			for (const auto& SmCacheEntry : CacheEntry.Value->ShadowMapEntries)
 			{
-				if (SmCacheEntry.IsValid())
-				{
-					checkSlow(SmCacheEntry->CurrentVirtualShadowMapId != INDEX_NONE);
-					uint32 Payload = EncodeInstanceInvalidationPayload(bPreviouslyCachedAsStatic, SmCacheEntry->CurrentVirtualShadowMapId);
-					Instances.Add(InstanceSceneDataOffset, NumInstanceSceneDataEntries, Payload);
-				}
+				checkSlow(SmCacheEntry.CurrentVirtualShadowMapId != INDEX_NONE);
+				uint32 Payload = EncodeInstanceInvalidationPayload(bPreviouslyCachedAsStatic, SmCacheEntry.CurrentVirtualShadowMapId);
+				Instances.Add(InstanceSceneDataOffset, NumInstanceSceneDataEntries, Payload);
 			}
 		}
 	}
@@ -527,8 +523,8 @@ FVirtualShadowMapArrayCacheManager::FVirtualShadowMapArrayCacheManager(FScene* I
 			static const auto* CVarMaxPhysicalPagesPtr = IConsoleManager::Get().FindTConsoleVariableDataInt(TEXT("r.Shadow.Virtual.MaxPhysicalPages"));
 			const int32 MaxPhysicalPages = CVarMaxPhysicalPagesPtr->GetValueOnRenderThread();
 
-				static const auto* CVarMaxPhysicalPagesSceneCapturePtr = IConsoleManager::Get().FindTConsoleVariableDataInt(TEXT("r.Shadow.Virtual.MaxPhysicalPagesSceneCapture"));
-				const int32 MaxPhysicalPagesSceneCapture = CVarMaxPhysicalPagesSceneCapturePtr->GetValueOnRenderThread();
+			static const auto* CVarMaxPhysicalPagesSceneCapturePtr = IConsoleManager::Get().FindTConsoleVariableDataInt(TEXT("r.Shadow.Virtual.MaxPhysicalPagesSceneCapture"));
+			const int32 MaxPhysicalPagesSceneCapture = CVarMaxPhysicalPagesSceneCapturePtr->GetValueOnRenderThread();
 
 #if !UE_BUILD_SHIPPING
 			if (!bLoggedPageOverflow)
@@ -714,7 +710,6 @@ void FVirtualShadowMapArrayCacheManager::FreeHZBPhysicalPool()
 void FVirtualShadowMapArrayCacheManager::Invalidate()
 {
 	// Clear the cache
-	PrevCacheEntries.Empty();
 	CacheEntries.Reset();
 
 	// Mark globally invalid until the next GPU allocation/metadata update
@@ -750,45 +745,38 @@ bool FVirtualShadowMapArrayCacheManager::IsHZBDataAvailable()
 	return HZBPhysicalPagePool && PrevBuffers.PageTable && PrevBuffers.PageFlags;
 }
 
-TSharedPtr<FVirtualShadowMapCacheEntry> FVirtualShadowMapPerLightCacheEntry::FindCreateShadowMapEntry(int32 Index)
+TSharedPtr<FVirtualShadowMapPerLightCacheEntry> FVirtualShadowMapArrayCacheManager::FindCreateLightCacheEntry(
+	int32 LightSceneId, uint32 ViewUniqueID, uint32 NumShadowMaps)
 {
-	check(Index >= 0);
-	ShadowMapEntries.SetNum(FMath::Max(Index + 1, ShadowMapEntries.Num()));
+	const uint64 CacheKey = (uint64(ViewUniqueID) << 32U) | uint64(LightSceneId);
 
-	TSharedPtr<FVirtualShadowMapCacheEntry>& EntryRef = ShadowMapEntries[Index];
+	TSharedPtr<FVirtualShadowMapPerLightCacheEntry> *LightEntryKey = CacheEntries.Find(CacheKey);
+	const int32 FrameNumber = Scene->GetFrameNumberRenderThread();
 
-	if (!EntryRef.IsValid())
+	if (LightEntryKey)
 	{
-		EntryRef = MakeShared<FVirtualShadowMapCacheEntry>();
-	}
-
-	return EntryRef;
-}
+		TSharedPtr<FVirtualShadowMapPerLightCacheEntry> LightEntry = *LightEntryKey;
 
-TSharedPtr<FVirtualShadowMapPerLightCacheEntry> FVirtualShadowMapArrayCacheManager::FindCreateLightCacheEntry(int32 LightSceneId, uint32 ViewUniqueID)
-{
-	const uint64 CacheKey = (uint64(ViewUniqueID) << 32U) | uint64(LightSceneId);
-
-	if (TSharedPtr<FVirtualShadowMapPerLightCacheEntry> *LightEntry = CacheEntries.Find(CacheKey))
-	{
-		return *LightEntry;
+		if (LightEntry->ShadowMapEntries.Num() == NumShadowMaps)
+		{
+			LightEntry->ReferencedFrameNumber = FrameNumber;
+			return LightEntry;
+		}
+		else
+		{
+			// Remove this entry and create a new one below
+			// NOTE: This should only happen for clipmaps currently on cvar changes
+			UE_LOG(LogRenderer, Display, TEXT("Virtual shadow map cache invalidated for light due to clipmap level count change"));
+			CacheEntries.Remove(CacheKey);
+		}
 	}
 
-	// Add to current frame / active set.
-	TSharedPtr<FVirtualShadowMapPerLightCacheEntry>& NewLightEntry = CacheEntries.Add(CacheKey);
+	// Make new entry for this light
+	TSharedPtr<FVirtualShadowMapPerLightCacheEntry> LightEntry = MakeShared<FVirtualShadowMapPerLightCacheEntry>(Scene->GetMaxPersistentPrimitiveIndex(), NumShadowMaps);
+	LightEntry->ReferencedFrameNumber = FrameNumber;	
+	CacheEntries.Add(CacheKey, LightEntry);
 
-	// Copy data if available
-	if (TSharedPtr<FVirtualShadowMapPerLightCacheEntry>* PrevNewLightEntry = PrevCacheEntries.Find(CacheKey))
-	{
-		NewLightEntry = *PrevNewLightEntry;
-	}
-	else
-	{
-		NewLightEntry = MakeShared<FVirtualShadowMapPerLightCacheEntry>(Scene->GetMaxPersistentPrimitiveIndex());
-	}
-
-	// return entry
-	return NewLightEntry;
+	return LightEntry;
 }
 
 void FVirtualShadowMapPerLightCacheEntry::OnPrimitiveRendered(const FPrimitiveSceneInfo* PrimitiveSceneInfo)
@@ -810,6 +798,73 @@ void FVirtualShadowMapPerLightCacheEntry::OnPrimitiveRendered(const FPrimitiveSc
 	}
 }
 
+void FVirtualShadowMapArrayCacheManager::UpdateUnreferencedCacheEntries(
+	FVirtualShadowMapArray& VirtualShadowMapArray)
+{
+	const int32 FrameNumber = Scene->GetFrameNumberRenderThread();
+	const int32 MaxAge = CVarCacheMaxUnreferencedLightAge.GetValueOnRenderThread();
+	const bool bAllowUnreferencedEntries = IsCacheEnabled() && MaxAge > 0;
+
+	TArray<uint64, SceneRenderingAllocator> EntriesToRemove;
+	for (auto& LightEntry : CacheEntries)
+	{
+		TSharedPtr<FVirtualShadowMapPerLightCacheEntry> CacheEntry = LightEntry.Value;
+		int32 Age = FrameNumber - CacheEntry->ReferencedFrameNumber;
+		check(CacheEntry->ReferencedFrameNumber >= 0);
+		check(Age >= 0);
+		if (Age == 0)
+		{
+			// Referenced this frame still, leave it alone
+		}
+		else if (bAllowUnreferencedEntries && Age < MaxAge)
+		{
+			INC_DWORD_STAT(STAT_VSMUnreferencedLights);
+
+			int PrevBaseVirtualShadowMapId = CacheEntry->ShadowMapEntries[0].CurrentVirtualShadowMapId;
+			bool bIsSinglePage = FVirtualShadowMapArray::IsSinglePage(PrevBaseVirtualShadowMapId);
+
+			// Keep the entry, reallocate new VSM IDs
+			int32 NumMaps = CacheEntry->ShadowMapEntries.Num();
+			int32 VirtualShadowMapId = VirtualShadowMapArray.Allocate(bIsSinglePage, NumMaps);
+			for (int32 Map = 0; Map < NumMaps; ++Map)
+			{
+				CacheEntry->ShadowMapEntries[Map].Update(VirtualShadowMapArray, *CacheEntry, VirtualShadowMapId + Map);
+				// NOTE: Leave the ProjectionData as whatever it was before
+				// TODO: We may want to add a flag that this is unreferenced so we can prune it from the light grid and skip it in page marking, etc...?
+				// Except in theory if we are marking things from onscreen pixels then we wouldn't have culled it... (small light culling though)?
+			}
+		}
+		else
+		{
+			// Enqueue for remove
+			EntriesToRemove.Add(LightEntry.Key);
+			//UE_LOG(LogRenderer, Display, TEXT("Removed VSM light cache entry (%d, Age %d)"), LightEntry.Key, Age);
+		}
+	}
+
+	for (uint64 Entry : EntriesToRemove)
+	{
+		int32 NumRemoved = CacheEntries.Remove(Entry);
+		check(NumRemoved > 0);
+	}
+}
+
+void FVirtualShadowMapArrayCacheManager::UploadProjectionData(FRDGScatterUploadBuffer& Uploader) const
+{
+	// If we get here, this should be non-empty, otherwise we will upload nothing and the destination buffer
+	// will end up as not having been written in RDG, which makes it upset even if it will never get referenced on the GPU.
+	check(!CacheEntries.IsEmpty());
+
+	for (const auto& LightEntry : CacheEntries)
+	{
+		TSharedPtr<FVirtualShadowMapPerLightCacheEntry> CacheEntry = LightEntry.Value;
+		for (const auto& Entry : CacheEntry->ShadowMapEntries)
+		{
+			Uploader.Add(Entry.CurrentVirtualShadowMapId, &Entry.ProjectionData);
+		}
+	}
+}
+
 class FVirtualSmCopyStatsCS : public FGlobalShader
 {
 	DECLARE_GLOBAL_SHADER(FVirtualSmCopyStatsCS);
@@ -885,8 +940,6 @@ void FVirtualShadowMapArrayCacheManager::ExtractFrameData(
 			PrevUniformParameters.PageTable = nullptr;
 			PrevUniformParameters.PhysicalPagePool = nullptr;
 		}
-		// Move cache entries to previous frame, this implicitly removes any that were not used
-		PrevCacheEntries = CacheEntries;
 
 		// propagate current-frame primitive state to cache entry
 		for (const auto& LightInfo : SceneRenderer.VisibleLightInfos)
@@ -898,8 +951,6 @@ void FVirtualShadowMapArrayCacheManager::ExtractFrameData(
 			}
 		}
 
-		CacheEntries.Reset();
-
 		ExtractStats(GraphBuilder, VirtualShadowMapArray);
 	}
 	else
@@ -1092,11 +1143,6 @@ void FVirtualShadowMapArrayCacheManager::OnSceneChange()
 	{
 		const int32 MaxPersistentPrimitiveIndex = FMath::Max(1, Scene->GetMaxPersistentPrimitiveIndex());
 
-		for (auto& CacheEntry : PrevCacheEntries)
-		{
-			CacheEntry.Value->CachedPrimitives.SetNum(MaxPersistentPrimitiveIndex, false);
-			CacheEntry.Value->RenderedPrimitives.SetNum(MaxPersistentPrimitiveIndex, false);
-		}
 		for (auto& CacheEntry : CacheEntries)
 		{
 			CacheEntry.Value->CachedPrimitives.SetNum(MaxPersistentPrimitiveIndex, false);
@@ -1108,7 +1154,6 @@ void FVirtualShadowMapArrayCacheManager::OnSceneChange()
 void FVirtualShadowMapArrayCacheManager::OnLightRemoved(int32 LightId)
 {
 	CacheEntries.Remove(LightId);
-	PrevCacheEntries.Remove(LightId);
 }
 
 /**
diff --git a/Engine/Source/Runtime/Renderer/Private/VirtualShadowMaps/VirtualShadowMapCacheManager.h b/Engine/Source/Runtime/Renderer/Private/VirtualShadowMaps/VirtualShadowMapCacheManager.h
--- a/Engine/Source/Runtime/Renderer/Private/VirtualShadowMaps/VirtualShadowMapCacheManager.h
+++ b/Engine/Source/Runtime/Renderer/Private/VirtualShadowMaps/VirtualShadowMapCacheManager.h
@@ -20,48 +20,42 @@ namespace Nanite { struct FPackedViewParams; }
 class FVirtualShadowMapCacheEntry
 {
 public:
-	void UpdateLocal(int32 VirtualShadowMapId, const FVirtualShadowMapPerLightCacheEntry &PerLightEntry);
-
-	void UpdateClipmap(int32 VirtualShadowMapId,
+	// Generic version used for local lights but also unreferenced lights
+	// Updates the VSM ID
+	void Update(
+		FVirtualShadowMapArray& VirtualShadowMapArray,
+		const FVirtualShadowMapPerLightCacheEntry &PerLightEntry,
+		int32 VirtualShadowMapId);
+
+	// Specific version of the above for clipmap levels, which have additional constraints
+	void UpdateClipmapLevel(
+		FVirtualShadowMapArray& VirtualShadowMapArray,
+		const FVirtualShadowMapPerLightCacheEntry& PerLightEntry,
+		int32 VirtualShadowMapId,
 		FInt64Point PageSpaceLocation,
-		FInt64Point ClipmapCornerOffset,
 		double LevelRadius,
 		double ViewCenterZ,
-		double ViewRadiusZ, 
-		const FVirtualShadowMapPerLightCacheEntry& PerLightEntry);
-
-	void Invalidate();
+		double ViewRadiusZ);
 
 	void SetHZBViewParams(Nanite::FPackedViewParams& OutParams);
 
-	// Note: Valid for all lights, but will return 0 for local lights currently as it is only used for clipmap panning right now
-	FInt32Point GetPreviousToCurrentPageOffset() const
-	{
-		// See clipmap construction
-		FInt32Point CurrentToPreviousOffset(Clipmap.CurrentClipmapCornerOffset - Clipmap.PrevClipmapCornerOffset);
-		return CurrentToPreviousOffset * (FVirtualShadowMap::Level0DimPagesXY >> 2U);
-	}
-
 	// Previous frame data
-	FInt64Point PrevPageSpaceLocation = FInt64Point(0, 0);
-	int32 PrevVirtualShadowMapId = INDEX_NONE;
 	FVirtualShadowMapHZBMetadata PrevHZBMetadata;
 
 	// Current frame data
-	FInt64Point CurrentPageSpaceLocation = FInt64Point(0, 0);
 	int32 CurrentVirtualShadowMapId = INDEX_NONE;
 	FVirtualShadowMapHZBMetadata CurrentHZBMetadata;
 
+	// Stores the projection shader data. This is needed for cached entries that may be unreferenced on the current frame
+	// and also avoids recomputing it every frame.
+	FVirtualShadowMapProjectionShaderData ProjectionData;
+
+	// Clipmap-specific information for panning and tracking of cached z-ranges in a given level
 	struct FClipmapInfo
 	{
-		double ViewCenterZ;
-		double ViewRadiusZ;
-
-		// Previous frame data
-		FInt64Point PrevClipmapCornerOffset = FInt64Point(0, 0);
-
-		// Current frame data
-		FInt64Point CurrentClipmapCornerOffset = FInt64Point(0, 0);
+		FInt64Point PageSpaceLocation = FInt64Point(0, 0);
+		double ViewCenterZ = 0.0;
+		double ViewRadiusZ = 0.0;
 	};
 	FClipmapInfo Clipmap;
 };
@@ -69,14 +63,13 @@ public:
 class FVirtualShadowMapPerLightCacheEntry
 {
 public:
-	FVirtualShadowMapPerLightCacheEntry(int32 MaxPersistentScenePrimitiveIndex)
+	FVirtualShadowMapPerLightCacheEntry(int32 MaxPersistentScenePrimitiveIndex, uint32 NumShadowMaps)
 		: RenderedPrimitives(false, MaxPersistentScenePrimitiveIndex)
 		, CachedPrimitives(false, MaxPersistentScenePrimitiveIndex)
 	{
+		ShadowMapEntries.SetNum(NumShadowMaps);
 	}
 
-	TSharedPtr<FVirtualShadowMapCacheEntry> FindCreateShadowMapEntry(int32 Index);
-
 	void OnPrimitiveRendered(const FPrimitiveSceneInfo* PrimitiveSceneInfo);
 	/**
 	 * The (local) VSM is fully cached if the previous frame if was distant and is distant this frame also.
@@ -90,7 +83,7 @@ public:
 
 	void MarkRendered(int32 FrameIndex) { Current.RenderedFrameNumber = FrameIndex; }
 	int32 GetLastScheduledFrameNumber() const { return Prev.ScheduledFrameNumber; }
-	void UpdateClipmap(const FMatrix& WorldToLight);
+	void UpdateClipmap(const FVector& LightDirection, int FirstLevelx);
 	/**
 	 * Returns true if the cache entry is valid (has previous state).
 	 */
@@ -101,6 +94,7 @@ public:
 	 */
 	void Invalidate();
 
+	// TODO: We probably don't need the prev/next thing anymore
 	struct FFrameState
 	{
 		bool bIsUncached = false;
@@ -112,6 +106,10 @@ public:
 	FFrameState Prev;
 	FFrameState Current;
 
+	// Tracks the last time we saw this light referenced.
+	// This allows us to keep entries alive even when they are temporarily not visible or similar, but age them out over time
+	int32 ReferencedFrameNumber = -1;
+
 	// Primitives that have been rendered (not culled) the previous frame, when a primitive transitions from being culled to not it must be rendered into the VSM
 	// Key culling reasons are small size or distance cutoff.
 	TBitArray<> RenderedPrimitives;
@@ -120,7 +118,7 @@ public:
 	TBitArray<> CachedPrimitives;
 
 	// One entry represents the cached state of a given shadow map in the set of either a clipmap(N), one cube map(6) or a regular VSM (1)
-	TArray< TSharedPtr<FVirtualShadowMapCacheEntry> > ShadowMapEntries;
+	TArray<FVirtualShadowMapCacheEntry> ShadowMapEntries;
 
 	// TODO: refactor this to not ne stored in the cache entry when we move (some) invalidaitons to the end of frame rather than in the scene primitive updates.
 	struct FInstanceRange
@@ -136,7 +134,9 @@ private:
 	FProjectedShadowInitializer LocalCacheKey;
 	struct FClipmapCacheKey
 	{
-		FMatrix	WorldToLight;
+		FVector LightDirection;
+		int FirstLevel;
+		int LevelCount;
 	};
 	FClipmapCacheKey ClipmapCacheKey;
 };
@@ -211,6 +211,17 @@ public:
 	// Invalidate the cache for all shadows, causing any pages to be rerendered
 	void Invalidate();
 
+	/**
+	 * Called before VSM builds page allocations to reallocate any lights that may not be visible this frame
+	 * but that may still have cached physical pages. We reallocate new VSM each frame for these to allow the associated
+	 * physical pages to live through short periods of being offscreen or otherwise culled. This function also removes
+	 * entries that are too old.
+	 */
+	void UpdateUnreferencedCacheEntries(FVirtualShadowMapArray& VirtualShadowMapArray);
+
+	// Must be called *after* calling UpdateUnreferencedCacheEntries - TODO: Perhaps merge the two to enforce
+	void UploadProjectionData(FRDGScatterUploadBuffer& Uploader) const;
+
 	/**
 	 * Call at end of frame to extract resouces from the virtual SM array to preserve to next frame.
 	 * If bCachingEnabled is false, all previous frame data is dropped and cache (and HZB!) data will not be available for the next frame.
@@ -223,7 +234,7 @@ public:
 	/**
 	 * Finds an existing cache entry and moves to the active set or creates a fresh one.
 	 */
-	TSharedPtr<FVirtualShadowMapPerLightCacheEntry> FindCreateLightCacheEntry(int32 LightSceneId, uint32 ViewUniqueID = 0U);
+	TSharedPtr<FVirtualShadowMapPerLightCacheEntry> FindCreateLightCacheEntry(int32 LightSceneId, uint32 ViewUniqueID, uint32 NumShadowMaps);
 
 	bool IsCacheEnabled();
 	bool IsCacheDataAvailable();
@@ -349,7 +360,6 @@ private:
 
 	// Index the Cache entries by the light ID
 	TMap< uint64, TSharedPtr<FVirtualShadowMapPerLightCacheEntry> > CacheEntries;
-	TMap< uint64, TSharedPtr<FVirtualShadowMapPerLightCacheEntry> > PrevCacheEntries;
 
 	// Marked after successfully allocating new physical pages
 	// Cleared if any global invalidation happens, in which case the next VSM update will not consider cached pages
diff --git a/Engine/Source/Runtime/Renderer/Private/VirtualShadowMaps/VirtualShadowMapClipmap.cpp b/Engine/Source/Runtime/Renderer/Private/VirtualShadowMaps/VirtualShadowMapClipmap.cpp
--- a/Engine/Source/Runtime/Renderer/Private/VirtualShadowMaps/VirtualShadowMapClipmap.cpp
+++ b/Engine/Source/Runtime/Renderer/Private/VirtualShadowMaps/VirtualShadowMapClipmap.cpp
@@ -87,7 +87,6 @@ static float GetLevelRadius(int32 Level)
 FVirtualShadowMapClipmap::FVirtualShadowMapClipmap(
 	FVirtualShadowMapArray& VirtualShadowMapArray,
 	const FLightSceneInfo& InLightSceneInfo,
-	const FMatrix& WorldToLightRotationMatrix,
 	const FViewMatrices& CameraViewMatrices,
 	FIntPoint CameraViewRectSize,
 	const FViewInfo* InDependentView,
@@ -95,10 +94,12 @@ FVirtualShadowMapClipmap::FVirtualShadowMapClipmap(
 	: LightSceneInfo(InLightSceneInfo),
 	  DependentView(InDependentView)
 {
-	check(WorldToLightRotationMatrix.GetOrigin() == FVector(0, 0, 0));	// Should not contain translation or scaling
-
 	FVirtualShadowMapArrayCacheManager* VirtualShadowMapArrayCacheManager = VirtualShadowMapArray.CacheManager;
 
+	// This math can obviously be simplified/optimized but not a big deal right now
+	const FVector LightDirection = LightSceneInfo.Proxy->GetDirection();
+	const FMatrix WorldToLightRotationMatrix = FInverseRotationMatrix(LightDirection.GetSafeNormal().Rotation());
+
 	const FMatrix FaceMatrix(
 		FPlane( 0, 0, 1, 0 ),
 		FPlane( 0, 1, 0, 0 ),
@@ -134,16 +135,15 @@ FVirtualShadowMapClipmap::FVirtualShadowMapClipmap(
 	LevelData.Empty();
 	LevelData.AddDefaulted(LevelCount);
 
+	VirtualShadowMapId = VirtualShadowMapArray.Allocate(false, LevelCount);
+
 	WorldOrigin = CameraViewMatrices.GetViewOrigin();
 
-	if (InDependentView && InDependentView->ViewState)	// scene capture views don't have a persistent view state
-	{
-		PerLightCacheEntry = VirtualShadowMapArrayCacheManager->FindCreateLightCacheEntry(LightSceneInfo.Id, InDependentView->ViewState->GetViewKey());
-		if (PerLightCacheEntry.IsValid())
-		{
-			PerLightCacheEntry->UpdateClipmap(WorldToLightRotationMatrix);
-		}
-	}
+	// TODO: We need a light/cache entry for every light/VSM now, but scene captures may not have persistent view state for indexing
+	// This is likely to all change as we refactor the multiple cache managers stuff anyways; for now this is hopefully safe since they do have separate copies
+	const uint32 UniqueViewKey = InDependentView->ViewState ? InDependentView->ViewState->GetViewKey() : 0U;
+	PerLightCacheEntry = VirtualShadowMapArrayCacheManager->FindCreateLightCacheEntry(LightSceneInfo.Id, UniqueViewKey, LevelCount);
+	PerLightCacheEntry->UpdateClipmap(LightDirection, FirstLevel);
 
 	const int RadiusLn = GetLevelRadius(LastLevel);
 	const int RadiiPerLevel = 4;
@@ -156,32 +156,36 @@ FVirtualShadowMapClipmap::FVirtualShadowMapClipmap(
 		SnappedOriginLn.Y = OriginSnapUnitsLn.Y * RadiusLn;
 	}
 
+	// We expand the depth range of the clipmap level to allow for camera movement without having to invalidate cached shadow data
+	// (See VirtualShadowMapCacheManager::UpdateClipmap for invalidation logic.)
+	// This also better accomodates SMRT where we want to avoid stepping outside of the Z bounds of a given clipmap
+	// NOTE: It's tempting to use a single global Z range for the entire clipmap (which avoids some SMRT overhead too)
+	// but this can cause precision issues with cached pages very near the camera.
+	const double ViewRadiusZScale = CVarVirtualShadowMapClipmapZRangeScale.GetValueOnRenderThread();
+	const FVector ViewCenter = WorldToLightViewRotationMatrix.TransformPosition(WorldOrigin);
+
 	for (int32 Index = 0; Index < LevelCount; ++Index)
 	{
 		FLevelData& Level = LevelData[Index];
 		const int32 AbsoluteLevel = Index + FirstLevel;		// Absolute (virtual) level index
 
-		// TODO: Allocate these as a chunk if we continue to use one per clipmap level
-		Level.VirtualShadowMap = VirtualShadowMapArray.Allocate(false);
-		ensure(Index == 0 || (Level.VirtualShadowMap->ID == (LevelData[Index-1].VirtualShadowMap->ID + 1)));
-
 		const float RawLevelRadius = GetLevelRadius(AbsoluteLevel);
 
 		double HalfLevelDim = 2.0 * RawLevelRadius;
 		double SnapSize = RawLevelRadius;
 
-		FVector ViewCenter = WorldToLightViewRotationMatrix.TransformPosition(WorldOrigin);
+		FVector SnappedViewCenter = ViewCenter;
 		FVector2D CenterSnapUnits(
 			FMath::RoundToDouble(ViewCenter.X / SnapSize),
 			FMath::RoundToDouble(ViewCenter.Y / SnapSize));
-		ViewCenter.X = CenterSnapUnits.X * SnapSize;
-		ViewCenter.Y = CenterSnapUnits.Y * SnapSize;
+		SnappedViewCenter.X = CenterSnapUnits.X * SnapSize;
+		SnappedViewCenter.Y = CenterSnapUnits.Y * SnapSize;
 
 		FInt64Point CornerOffset;
 		CornerOffset.X = -(int64_t)CenterSnapUnits.X + (RadiiPerLevel/2);
 		CornerOffset.Y =  (int64_t)CenterSnapUnits.Y + (RadiiPerLevel/2);
 
-		const FVector SnappedWorldCenter = ViewToWorldRotationMatrix.TransformPosition(ViewCenter);
+		const FVector SnappedWorldCenter = ViewToWorldRotationMatrix.TransformPosition(SnappedViewCenter);
 
 		Level.WorldCenter = SnappedWorldCenter;
 		Level.CornerOffset = CornerOffset;
@@ -190,61 +194,53 @@ FVirtualShadowMapClipmap::FVirtualShadowMapClipmap(
 		// The reference point is WorldOrigin snapped to a grid of GetLevelRadius(LastLevel),
 		// because points on this grid are guaranteed to also be present on lower levels,
 		// therefore allowing the offsets to represented as factors of level radii without precision loss.
-		const FInt64Point SnappedPageOriginLi(-ViewCenter.X, ViewCenter.Y);
+		const FInt64Point SnappedPageOriginLi(-SnappedViewCenter.X, SnappedViewCenter.Y);
 		const FInt64Point SnappedPageOriginLn(-SnappedOriginLn.X, SnappedOriginLn.Y);
 		const FInt64Point RelativeCornerOffset = SnappedPageOriginLi - SnappedPageOriginLn + ((RadiiPerLevel / 2) * (int64_t)SnapSize);
 		Level.RelativeCornerOffset = FIntPoint(RelativeCornerOffset / SnapSize);
 
 		// Check if we have a cache entry for this level
 		// If we do and it covers our required depth range, we can use cached pages. Otherwise we need to invalidate.
-		TSharedPtr<FVirtualShadowMapCacheEntry> CacheEntry = nullptr;
-		if (PerLightCacheEntry.IsValid())
-		{
-			// NOTE: We use the absolute (virtual) level index so that the caching is robust against changes to the chosen level range
-			CacheEntry = PerLightCacheEntry->FindCreateShadowMapEntry(AbsoluteLevel);
-		}
-
-		// We expand the depth range of the clipmap level to allow for camera movement without having to invalidate cached shadow data
-		// (See VirtualShadowMapCacheManager::UpdateClipmap for invalidation logic.)
-		// This also better accomodates SMRT where we want to avoid stepping outside of the Z bounds of a given clipmap
-		// NOTE: It's tempting to use a single global Z range for the entire clipmap (which avoids some SMRT overhead too)
-		// but this can cause precision issues with cached pages very near the camera.
-		const double ViewRadiusZScale = CVarVirtualShadowMapClipmapZRangeScale.GetValueOnRenderThread();
+		FVirtualShadowMapCacheEntry* ClipmapLevelEntry = &PerLightCacheEntry->ShadowMapEntries[Index];
 
 		double ViewRadiusZ = RawLevelRadius * ViewRadiusZScale;
 		double ViewCenterDeltaZ = 0.0f;
 
-		if (CacheEntry)
-		{
-			// We snap to half the size of the VSM at each level
-			check((FVirtualShadowMap::Level0DimPagesXY & 1) == 0);
-			FInt64Point PageOffset(CornerOffset * (FVirtualShadowMap::Level0DimPagesXY >> 2));
-
-			CacheEntry->UpdateClipmap(Level.VirtualShadowMap->ID,
-				PageOffset,
-				CornerOffset,
-				RawLevelRadius,
-				ViewCenter.Z,
-				ViewRadiusZ,
-				*PerLightCacheEntry);
-
-			Level.VirtualShadowMap->VirtualShadowMapCacheEntry = CacheEntry;
-
-			// Update min/max Z based on the cached page (if present and valid)
-			// We need to ensure we use a consistent depth range as the camera moves for each level
-			ViewCenterDeltaZ = ViewCenter.Z - CacheEntry->Clipmap.ViewCenterZ;
-			ViewRadiusZ = CacheEntry->Clipmap.ViewRadiusZ;
-		}
+		// We snap to half the size of the VSM at each level
+		check((FVirtualShadowMap::Level0DimPagesXY & 1) == 0);
+		FInt64Point PageOffset(CornerOffset * (FVirtualShadowMap::Level0DimPagesXY >> 2));
+
+		ClipmapLevelEntry->UpdateClipmapLevel(
+			VirtualShadowMapArray,
+			*PerLightCacheEntry,
+			GetVirtualShadowMapId(Index),
+			PageOffset,
+			RawLevelRadius,
+			ViewCenter.Z,
+			ViewRadiusZ);
+
+		// Update min/max Z based on the cached page (if present and valid)
+		// We need to ensure we use a consistent depth range as the camera moves for each level
+		ViewCenterDeltaZ = ViewCenter.Z - ClipmapLevelEntry->Clipmap.ViewCenterZ;
+		ViewRadiusZ = ClipmapLevelEntry->Clipmap.ViewRadiusZ;
 
 		// NOTE: These values are all in regular ranges after being offset
 		const double ZScale = 0.5 / ViewRadiusZ;
 		const double ZOffset = ViewRadiusZ + ViewCenterDeltaZ;
 		Level.ViewToClip = FReversedZOrthoMatrix(HalfLevelDim, HalfLevelDim, ZScale, ZOffset);
+
+		// Update the entry with the new projection data
+		ClipmapLevelEntry->ProjectionData = ComputeProjectionShaderData(Index);
 	}
 
 	ComputeBoundingVolumes(CameraViewMatrices);
 }
 
+const FVirtualShadowMapProjectionShaderData& FVirtualShadowMapClipmap::GetProjectionShaderData(int32 ClipmapIndex) const
+{
+	return PerLightCacheEntry->ShadowMapEntries[ClipmapIndex].ProjectionData;
+}
+
 void FVirtualShadowMapClipmap::ComputeBoundingVolumes(const FViewMatrices& CameraViewMatrices)
 {
 	// We don't really do much CPU culling with clipmaps. After various testing the fact that we are culling
@@ -287,7 +283,7 @@ FViewMatrices FVirtualShadowMapClipmap::GetViewMatrices(int32 ClipmapIndex) cons
 	return FViewMatrices(Initializer);
 }
 
-FVirtualShadowMapProjectionShaderData FVirtualShadowMapClipmap::GetProjectionShaderData(int32 ClipmapIndex) const
+FVirtualShadowMapProjectionShaderData FVirtualShadowMapClipmap::ComputeProjectionShaderData(int32 ClipmapIndex) const
 {
 	check(ClipmapIndex >= 0 && ClipmapIndex < LevelData.Num());
 	const FLevelData& Level = LevelData[ClipmapIndex];
@@ -314,11 +310,7 @@ FVirtualShadowMapProjectionShaderData FVirtualShadowMapClipmap::GetProjectionSha
 	Data.ResolutionLodBias = ResolutionLodBias;
 	Data.ClipmapCornerRelativeOffset = Level.RelativeCornerOffset;
 	Data.LightSourceRadius = GetLightSceneInfo().Proxy->GetSourceRadius();
-	Data.Flags = 0U;
-	if (PerLightCacheEntry.IsValid() && PerLightCacheEntry->IsUncached())
-	{
-		Data.Flags |= VSM_PROJ_FLAG_UNCACHED;
-	}
+	Data.Flags = PerLightCacheEntry->IsUncached() ? VSM_PROJ_FLAG_UNCACHED : 0U;
 
 	return Data;
 }
diff --git a/Engine/Source/Runtime/Renderer/Private/VirtualShadowMaps/VirtualShadowMapClipmap.h b/Engine/Source/Runtime/Renderer/Private/VirtualShadowMaps/VirtualShadowMapClipmap.h
--- a/Engine/Source/Runtime/Renderer/Private/VirtualShadowMaps/VirtualShadowMapClipmap.h
+++ b/Engine/Source/Runtime/Renderer/Private/VirtualShadowMaps/VirtualShadowMapClipmap.h
@@ -25,7 +25,6 @@ public:
 	FVirtualShadowMapClipmap(
 		FVirtualShadowMapArray& VirtualShadowMapArray,
 		const FLightSceneInfo& InLightSceneInfo,
-		const FMatrix& WorldToLightRotationMatrix,
 		const FViewMatrices& CameraViewMatrices,
 		FIntPoint CameraViewRectSize,
 		const FViewInfo* InDependentView,
@@ -34,9 +33,9 @@ public:
 
 	FViewMatrices GetViewMatrices(int32 ClipmapIndex) const;
 
-	FVirtualShadowMap* GetVirtualShadowMap(int32 ClipmapIndex = 0)
+	int32 GetVirtualShadowMapId(int32 ClipmapIndex = 0) const
 	{
-		return LevelData[ClipmapIndex].VirtualShadowMap;
+		return VirtualShadowMapId + ClipmapIndex;
 	}
 
 	int32 GetLevelCount() const
@@ -60,7 +59,7 @@ public:
 		return LightSceneInfo;
 	}
 
-	FVirtualShadowMapProjectionShaderData GetProjectionShaderData(int32 ClipmapIndex) const;
+	const FVirtualShadowMapProjectionShaderData& GetProjectionShaderData(int32 ClipmapIndex) const;
 
 	FVector GetWorldOrigin() const
 	{
@@ -99,6 +98,7 @@ public:
 	TSharedPtr<FVirtualShadowMapPerLightCacheEntry>& GetCacheEntry() { return PerLightCacheEntry; }
 
 private:
+	FVirtualShadowMapProjectionShaderData ComputeProjectionShaderData(int32 ClipmapIndex) const;
 	void ComputeBoundingVolumes(const FViewMatrices& CameraViewMatrices);
 
 	const FLightSceneInfo& LightSceneInfo;
@@ -124,7 +124,6 @@ private:
 
 	struct FLevelData
 	{
-		FVirtualShadowMap* VirtualShadowMap;
 		FMatrix ViewToClip;
 		FVector WorldCenter;
 		//Offset from (0,0) to clipmap corner, in level radii
@@ -137,6 +136,7 @@ private:
 	FSphere BoundingSphere;
 	FConvexVolume ViewFrustumBounds;
 
+	int32 VirtualShadowMapId = INDEX_NONE;		// Base ID; levels are contiguous
 	TSharedPtr<FVirtualShadowMapPerLightCacheEntry> PerLightCacheEntry;
 
 	// Rendered primitives are marked during culling (through OnPrimitiveRendered being called).
diff --git a/Engine/Source/Runtime/Renderer/Private/VirtualShadowMaps/VirtualShadowMapProjection.cpp b/Engine/Source/Runtime/Renderer/Private/VirtualShadowMaps/VirtualShadowMapProjection.cpp
--- a/Engine/Source/Runtime/Renderer/Private/VirtualShadowMaps/VirtualShadowMapProjection.cpp
+++ b/Engine/Source/Runtime/Renderer/Private/VirtualShadowMaps/VirtualShadowMapProjection.cpp
@@ -568,7 +568,7 @@ void RenderVirtualShadowMapProjection(
 		InputType,
 		VirtualShadowMaskTexture,
 		Clipmap->GetLightSceneInfo().Proxy,
-		Clipmap->GetVirtualShadowMap()->ID,
+		Clipmap->GetVirtualShadowMapId(),
 		TiledVSMProjection);
 	
 	CompositeVirtualShadowMapMask(
