Fix GPUScene HZB instance culling on mobile.

diff --git a/Engine/Source/Runtime/Renderer/Private/MobileShadingRenderer.cpp b/Engine/Source/Runtime/Renderer/Private/MobileShadingRenderer.cpp
--- a/Engine/Source/Runtime/Renderer/Private/MobileShadingRenderer.cpp
+++ b/Engine/Source/Runtime/Renderer/Private/MobileShadingRenderer.cpp
@@ -451,6 +451,8 @@ void FMobileSceneRenderer::InitViews(
 
 	check(Scene);
 
+	PreVisibilityFrameSetup(GraphBuilder);
+
 	// Create GPU-side representation of the view for instance culling.
 	for (int32 ViewIndex = 0; ViewIndex < Views.Num(); ++ViewIndex)
 	{
@@ -462,7 +464,6 @@ void FMobileSceneRenderer::InitViews(
 
 	const FExclusiveDepthStencil::Type BasePassDepthStencilAccess = FExclusiveDepthStencil::DepthWrite_StencilWrite;
 
-	PreVisibilityFrameSetup(GraphBuilder);
 	if (FXSystem && FXSystem->RequiresEarlyViewUniformBuffer() && Views.IsValidIndex(0))
 	{
 		// This is to init the ViewUniformBuffer before rendering for the Niagara compute shader.
@@ -1990,6 +1991,18 @@ void FMobileSceneRenderer::RenderHZB(FRDGBuilder& GraphBuilder, FRDGTextureRef S
 
 			View.HZBMipmap0Size = FurthestHZBTexture->Desc.Extent;
 			View.HZB = FurthestHZBTexture;
+
+			if (View.ViewState)
+			{
+				if (FInstanceCullingContext::IsOcclusionCullingEnabled())
+				{
+					GraphBuilder.QueueTextureExtraction(FurthestHZBTexture, &View.ViewState->PrevFrameViewInfo.HZB);
+				}
+				else
+				{
+					View.ViewState->PrevFrameViewInfo.HZB = nullptr;
+				}
+			}
 		}
 	}
 }
