Add Scene uniform buffer to Heterogenous Volumes material shaders.

diff --git a/Engine/Source/Runtime/Renderer/Private/HeterogeneousVolumes/HeterogeneousVolumesLiveShadingPipeline.cpp b/Engine/Source/Runtime/Renderer/Private/HeterogeneousVolumes/HeterogeneousVolumesLiveShadingPipeline.cpp
--- a/Engine/Source/Runtime/Renderer/Private/HeterogeneousVolumes/HeterogeneousVolumesLiveShadingPipeline.cpp
+++ b/Engine/Source/Runtime/Renderer/Private/HeterogeneousVolumes/HeterogeneousVolumesLiveShadingPipeline.cpp
@@ -99,6 +99,7 @@ class FRenderLightingCacheWithLiveShadingCS : public FMeshMaterialShader
 		// Scene data
 		SHADER_PARAMETER_STRUCT_REF(FViewUniformShaderParameters, View)
 		SHADER_PARAMETER_STRUCT_INCLUDE(FSceneTextureParameters, SceneTextures)
+		SHADER_PARAMETER_RDG_UNIFORM_BUFFER(FSceneUniformParameters, Scene)
 
 		// Light data
 		SHADER_PARAMETER(int, bApplyEmissionAndTransmittance)
@@ -213,6 +214,7 @@ class FRenderSingleScatteringWithLiveShadingCS : public FMeshMaterialShader
 		// Scene data
 		SHADER_PARAMETER_STRUCT_REF(FViewUniformShaderParameters, View)
 		SHADER_PARAMETER_STRUCT_INCLUDE(FSceneTextureParameters, SceneTextures)
+		SHADER_PARAMETER_RDG_UNIFORM_BUFFER(FSceneUniformParameters, Scene)
 
 		// Light data
 		SHADER_PARAMETER(int, bApplyEmissionAndTransmittance)
@@ -411,6 +413,7 @@ void RenderLightingCacheWithLiveShading(
 		// Scene data
 		PassParameters->View = View.ViewUniformBuffer;
 		PassParameters->SceneTextures = GetSceneTextureParameters(GraphBuilder, SceneTextures);
+		PassParameters->Scene = View.GetSceneUniforms().GetBuffer(GraphBuilder);
 
 		// Light data
 		check(LightSceneInfo != nullptr)
@@ -554,6 +557,7 @@ void RenderSingleScatteringWithLiveShading(
 		// Scene data
 		PassParameters->View = View.ViewUniformBuffer;
 		PassParameters->SceneTextures = GetSceneTextureParameters(GraphBuilder, SceneTextures);
+		PassParameters->Scene = View.GetSceneUniforms().GetBuffer(GraphBuilder);
 
 		// Light data
 		PassParameters->bApplyEmissionAndTransmittance = bApplyEmissionAndTransmittance;
diff --git a/Engine/Source/Runtime/Renderer/Private/HeterogeneousVolumes/HeterogeneousVolumesMaterialBakingPipeline.cpp b/Engine/Source/Runtime/Renderer/Private/HeterogeneousVolumes/HeterogeneousVolumesMaterialBakingPipeline.cpp
--- a/Engine/Source/Runtime/Renderer/Private/HeterogeneousVolumes/HeterogeneousVolumesMaterialBakingPipeline.cpp
+++ b/Engine/Source/Runtime/Renderer/Private/HeterogeneousVolumes/HeterogeneousVolumesMaterialBakingPipeline.cpp
@@ -18,6 +18,7 @@ class FHeterogeneousVolumesBakeMaterialCS : public FMeshMaterialShader
 	BEGIN_SHADER_PARAMETER_STRUCT(FParameters, )
 		// Scene data
 		SHADER_PARAMETER_STRUCT_REF(FViewUniformShaderParameters, View)
+		SHADER_PARAMETER_RDG_UNIFORM_BUFFER(FSceneUniformParameters, Scene)
 
 		// Object data
 		SHADER_PARAMETER(FMatrix44f, LocalToWorld)
@@ -132,6 +133,7 @@ void ComputeHeterogeneousVolumeBakeMaterial(
 	{
 		// Scene data
 		PassParameters->View = View.ViewUniformBuffer;
+		PassParameters->Scene = View.GetSceneUniforms().GetBuffer(GraphBuilder);
 
 		// Object data
 		FMatrix44f LocalToWorld = FMatrix44f(PrimitiveSceneProxy->GetLocalToWorld());
diff --git a/Engine/Source/Runtime/Renderer/Private/HeterogeneousVolumes/HeterogeneousVolumesVoxelGridPipeline.cpp b/Engine/Source/Runtime/Renderer/Private/HeterogeneousVolumes/HeterogeneousVolumesVoxelGridPipeline.cpp
--- a/Engine/Source/Runtime/Renderer/Private/HeterogeneousVolumes/HeterogeneousVolumesVoxelGridPipeline.cpp
+++ b/Engine/Source/Runtime/Renderer/Private/HeterogeneousVolumes/HeterogeneousVolumesVoxelGridPipeline.cpp
@@ -356,6 +356,7 @@ class FRasterizeBottomLevelFroxelGridCS : public FMeshMaterialShader
 	BEGIN_SHADER_PARAMETER_STRUCT(FParameters, )
 		// Scene data
 		SHADER_PARAMETER_STRUCT_REF(FViewUniformShaderParameters, View)
+		SHADER_PARAMETER_RDG_UNIFORM_BUFFER(FSceneUniformParameters, Scene)
 
 		// Primitive data
 		SHADER_PARAMETER(FVector3f, PrimitiveWorldBoundsMin)
@@ -653,6 +654,7 @@ class FRasterizeBottomLevelGrid : public FMeshMaterialShader
 	BEGIN_SHADER_PARAMETER_STRUCT(FParameters, )
 		// Scene data
 		SHADER_PARAMETER_STRUCT_REF(FViewUniformShaderParameters, View)
+		SHADER_PARAMETER_RDG_UNIFORM_BUFFER(FSceneUniformParameters, Scene)
 
 		// Primitive data
 		SHADER_PARAMETER(FVector3f, PrimitiveWorldBoundsMin)
@@ -1656,6 +1658,7 @@ void RasterizeVolumesIntoFrustumVoxelGrid(
 		{
 			// Scene data
 			PassParameters->View = View.ViewUniformBuffer;
+			PassParameters->Scene = View.GetSceneUniforms().GetBuffer(GraphBuilder);
 
 			// Primitive data
 			PassParameters->PrimitiveWorldBoundsMin = FVector3f(PrimitiveBounds.Origin - PrimitiveBounds.BoxExtent);
@@ -2099,6 +2102,7 @@ void RasterizeVolumesIntoOrthoVoxelGrid(
 			{
 				// Scene data
 				PassParameters->View = View.ViewUniformBuffer;
+				PassParameters->Scene = View.GetSceneUniforms().GetBuffer(GraphBuilder);
 
 				// Primitive data
 				PassParameters->PrimitiveWorldBoundsMin = FVector3f(PrimitiveBounds.Origin - PrimitiveBounds.BoxExtent);
