Maintain our own sequence count for VSM cache entries since the frame count can be reused with scene captures and similar.
In this case we need to know if a given cache entry was referenced this "render".

diff --git a/Engine/Source/Runtime/Renderer/Private/VirtualShadowMaps/VirtualShadowMapCacheManager.cpp b/Engine/Source/Runtime/Renderer/Private/VirtualShadowMaps/VirtualShadowMapCacheManager.cpp
--- a/Engine/Source/Runtime/Renderer/Private/VirtualShadowMaps/VirtualShadowMapCacheManager.cpp
+++ b/Engine/Source/Runtime/Renderer/Private/VirtualShadowMaps/VirtualShadowMapCacheManager.cpp
@@ -751,7 +751,6 @@ TSharedPtr<FVirtualShadowMapPerLightCacheEntry> FVirtualShadowMapArrayCacheManag
 	const uint64 CacheKey = (uint64(ViewUniqueID) << 32U) | uint64(LightSceneId);
 
 	TSharedPtr<FVirtualShadowMapPerLightCacheEntry> *LightEntryKey = CacheEntries.Find(CacheKey);
-	const int32 FrameNumber = Scene->GetFrameNumberRenderThread();
 
 	if (LightEntryKey)
 	{
@@ -759,7 +758,7 @@ TSharedPtr<FVirtualShadowMapPerLightCacheEntry> FVirtualShadowMapArrayCacheManag
 
 		if (LightEntry->ShadowMapEntries.Num() == NumShadowMaps)
 		{
-			LightEntry->ReferencedFrameNumber = FrameNumber;
+			LightEntry->ReferencedRenderSequenceNumber = RenderSequenceNumber;
 			return LightEntry;
 		}
 		else
@@ -773,7 +772,7 @@ TSharedPtr<FVirtualShadowMapPerLightCacheEntry> FVirtualShadowMapArrayCacheManag
 
 	// Make new entry for this light
 	TSharedPtr<FVirtualShadowMapPerLightCacheEntry> LightEntry = MakeShared<FVirtualShadowMapPerLightCacheEntry>(Scene->GetMaxPersistentPrimitiveIndex(), NumShadowMaps);
-	LightEntry->ReferencedFrameNumber = FrameNumber;	
+	LightEntry->ReferencedRenderSequenceNumber = RenderSequenceNumber;	
 	CacheEntries.Add(CacheKey, LightEntry);
 
 	return LightEntry;
@@ -801,20 +800,23 @@ void FVirtualShadowMapPerLightCacheEntry::OnPrimitiveRendered(const FPrimitiveSc
 void FVirtualShadowMapArrayCacheManager::UpdateUnreferencedCacheEntries(
 	FVirtualShadowMapArray& VirtualShadowMapArray)
 {
-	const int32 FrameNumber = Scene->GetFrameNumberRenderThread();
-	const int32 MaxAge = CVarCacheMaxUnreferencedLightAge.GetValueOnRenderThread();
+	const int64 MaxAge = CVarCacheMaxUnreferencedLightAge.GetValueOnRenderThread();
 	const bool bAllowUnreferencedEntries = IsCacheEnabled() && MaxAge > 0;
 
 	TArray<uint64, SceneRenderingAllocator> EntriesToRemove;
 	for (auto& LightEntry : CacheEntries)
 	{
 		TSharedPtr<FVirtualShadowMapPerLightCacheEntry> CacheEntry = LightEntry.Value;
-		int32 Age = FrameNumber - CacheEntry->ReferencedFrameNumber;
-		check(CacheEntry->ReferencedFrameNumber >= 0);
+		// NOTE: We probably want to decouple the age from render calls at some point, but regardless we need to know
+		// in the first branch that the given light was referenced *this* render, and therefore has already been
+		// (re)allocated a current VSM ID.
+		int64 Age = RenderSequenceNumber - CacheEntry->ReferencedRenderSequenceNumber;
+		check(CacheEntry->ReferencedRenderSequenceNumber >= 0);
 		check(Age >= 0);
 		if (Age == 0)
 		{
 			// Referenced this frame still, leave it alone
+			check(CacheEntry->ShadowMapEntries.Last().CurrentVirtualShadowMapId < VirtualShadowMapArray.GetNumShadowMapSlots());
 		}
 		else if (bAllowUnreferencedEntries && Age < MaxAge)
 		{
@@ -972,6 +974,8 @@ void FVirtualShadowMapArrayCacheManager::ExtractFrameData(
 		RecentlyRemovedReadIndex = 1 - RecentlyRemovedReadIndex;
 		RecentlyRemovedFrameCounter = 0;
 	}
+
+	++RenderSequenceNumber;
 }
 
 void FVirtualShadowMapArrayCacheManager::ExtractStats(FRDGBuilder& GraphBuilder, FVirtualShadowMapArray &VirtualShadowMapArray)
diff --git a/Engine/Source/Runtime/Renderer/Private/VirtualShadowMaps/VirtualShadowMapCacheManager.h b/Engine/Source/Runtime/Renderer/Private/VirtualShadowMaps/VirtualShadowMapCacheManager.h
--- a/Engine/Source/Runtime/Renderer/Private/VirtualShadowMaps/VirtualShadowMapCacheManager.h
+++ b/Engine/Source/Runtime/Renderer/Private/VirtualShadowMaps/VirtualShadowMapCacheManager.h
@@ -100,7 +100,7 @@ public:
 		bool bIsUncached = false;
 		bool bIsDistantLight = false;
 		int32 RenderedFrameNumber = -1;
-		int32 ScheduledFrameNumber = -1;
+		int32 ScheduledFrameNumber = -1;		
 
 	};
 	FFrameState Prev;
@@ -108,7 +108,9 @@ public:
 
 	// Tracks the last time we saw this light referenced.
 	// This allows us to keep entries alive even when they are temporarily not visible or similar, but age them out over time
-	int32 ReferencedFrameNumber = -1;
+	// NOTE: This is not exactly a frame number since multiple "Render" calls may happen per frame, and we need to track
+	// them separately for the purposes of determining whether a given light is actually referenced in the current Render.
+	int64 ReferencedRenderSequenceNumber = -1;
 
 	// Primitives that have been rendered (not culled) the previous frame, when a primitive transitions from being culled to not it must be rendered into the VSM
 	// Key culling reasons are small size or distance cutoff.
@@ -374,6 +376,9 @@ private:
 	int32 RecentlyRemovedReadIndex = 0;
 	int32 RecentlyRemovedFrameCounter = 0;
 
+	// Used for tracking which cache entries are referenced in a given render call
+	int64 RenderSequenceNumber = 0;
+
 	// Stores stats over frames when activated.
 	TRefCountPtr<FRDGPooledBuffer> AccumulatedStatsBuffer;
 	bool bAccumulatingStats = false;
