Removed the currently unused FSpatialHashDebugDrawer since it is causing static analysis warnings

diff --git a/Engine/Source/Runtime/Renderer/Private/SceneCulling/SceneCulling.cpp b/Engine/Source/Runtime/Renderer/Private/SceneCulling/SceneCulling.cpp
--- a/Engine/Source/Runtime/Renderer/Private/SceneCulling/SceneCulling.cpp
+++ b/Engine/Source/Runtime/Renderer/Private/SceneCulling/SceneCulling.cpp
@@ -233,7 +233,7 @@ struct FSpatialHashNullDebugDrawer
 	inline void DrawBlock(bool bHighlight, const FBox& BlockBounds) {}
 
 };
-#if !(UE_BUILD_SHIPPING || UE_BUILD_TEST)
+#if 0
 
 struct FSpatialHashDebugDrawer
 {
@@ -270,8 +270,6 @@ struct FSpatialHashDebugDrawer
 		}
 	}
 };
-#else
-using FSpatialHashDebugDrawer = FSpatialHashNullDebugDrawer;
 #endif
 
 inline bool IsValid(const FCellHeader& CellHeader)
