Fix for volumetric fog leaking on the side of screens in some case.
It was happening when reprojecting history and applying fog on screen.
Tested with resolution hell and resource clobber with split screen.

diff --git a/Engine/Shaders/Private/HeightFogPixelShader.usf b/Engine/Shaders/Private/HeightFogPixelShader.usf
--- a/Engine/Shaders/Private/HeightFogPixelShader.usf
+++ b/Engine/Shaders/Private/HeightFogPixelShader.usf
@@ -67,6 +67,7 @@ void ExponentialPixelMain(
 	float3 CellOffset = UpsampleJitterMultiplier * Rand3D;
 
 	float3 VolumeUV = float3(((SVPos.xy - View.ViewRectMin.xy) + CellOffset.xy) * View.VolumetricFogSVPosToVolumeUV, ZSlice);
+	VolumeUV.xy = min(VolumeUV.xy, View.VolumetricFogUVMax);
 
 	float4 HeightFogInscatteringAndOpacity = CalculateHeightFog(WorldPositionRelativeToCamera);
 
diff --git a/Engine/Shaders/Private/VolumetricFog.usf b/Engine/Shaders/Private/VolumetricFog.usf
--- a/Engine/Shaders/Private/VolumetricFog.usf
+++ b/Engine/Shaders/Private/VolumetricFog.usf
@@ -537,8 +537,9 @@ void LightScatteringCS(
 	float3 HistoryUV = ComputeHistoryVolumeUV_DEPRECATED(ComputeCellTranslatedWorldPosition(GridCoordinate, .5f), UnjitteredPrevTranslatedWorldToClip);
 	float HistoryAlpha = HistoryWeight;
 
+	// We need to test with View.VolumetricFogPrevUVMax for HistoryUV because that is what we clamp with in ComputeHistoryVolumeUV_DEPRECATED.
 	FLATTEN
-	if (any(HistoryUV < 0) || any(HistoryUV > 1) || CellWasOccluded)
+	if (any(HistoryUV < 0) || any(HistoryUV >= float3(View.VolumetricFogPrevUVMax, 1.0f)) || CellWasOccluded)
 	{
 		HistoryAlpha = 0;
 	}
diff --git a/Engine/Source/Runtime/Renderer/Private/SceneRendering.cpp b/Engine/Source/Runtime/Renderer/Private/SceneRendering.cpp
--- a/Engine/Source/Runtime/Renderer/Private/SceneRendering.cpp
+++ b/Engine/Source/Runtime/Renderer/Private/SceneRendering.cpp
@@ -1482,7 +1482,7 @@ void FViewInfo::SetupUniformBufferParameters(
 		ViewUniformShaderParameters.VolumetricFogPrevUVMax = ViewState->VolumetricFogPrevUVMax;
 	}
 	ViewUniformShaderParameters.VolumetricFogScreenToResourceUV = ViewRectSize / (FVector2f(VolumetricFogResourceGridSize.X, VolumetricFogResourceGridSize.Y) * VolumetricFogResourceGridPixelSize);
-	ViewUniformShaderParameters.VolumetricFogUVMax = FVector2f(ViewRectSize.X - 0.51, ViewRectSize.Y - 0.51) / (FVector2f(VolumetricFogResourceGridSize.X, VolumetricFogResourceGridSize.Y) * VolumetricFogResourceGridPixelSize);
+	ViewUniformShaderParameters.VolumetricFogUVMax = GetVolumetricFogUVMax(ViewRectSize, VolumetricFogResourceGridSize, VolumetricFogResourceGridPixelSize);
 
 	SetupPrecomputedVolumetricLightmapUniformBufferParameters(Scene, Family->EngineShowFlags, ViewUniformShaderParameters);
 
diff --git a/Engine/Source/Runtime/Renderer/Private/VolumetricFog.cpp b/Engine/Source/Runtime/Renderer/Private/VolumetricFog.cpp
--- a/Engine/Source/Runtime/Renderer/Private/VolumetricFog.cpp
+++ b/Engine/Source/Runtime/Renderer/Private/VolumetricFog.cpp
@@ -867,6 +867,13 @@ FIntVector GetVolumetricFogViewGridSize(const FViewInfo& View, int32& OutVolumet
 	return GetVolumetricFogGridSize(View.ViewRect.Size(), OutVolumetricFogGridPixelSize);
 }
 
+FVector2f GetVolumetricFogUVMax(const FVector2f& ViewRectSize, FIntVector VolumetricFogResourceGridSize, int32 VolumetricFogResourceGridPixelSize)
+{
+	float ViewRectSizeXSafe = FMath::DivideAndRoundUp<int32>(int32(ViewRectSize.X), VolumetricFogResourceGridPixelSize) * VolumetricFogResourceGridPixelSize - (VolumetricFogResourceGridPixelSize / 2 + 1);
+	float ViewRectSizeYSafe = FMath::DivideAndRoundUp<int32>(int32(ViewRectSize.Y), VolumetricFogResourceGridPixelSize) * VolumetricFogResourceGridPixelSize - (VolumetricFogResourceGridPixelSize / 2 + 1);
+	return FVector2f(ViewRectSizeXSafe, ViewRectSizeYSafe) / (FVector2f(VolumetricFogResourceGridSize.X, VolumetricFogResourceGridSize.Y) * VolumetricFogResourceGridPixelSize);
+}
+
 FVector2f GetVolumetricFogFroxelToScreenSVPosRatio(const FViewInfo& View)
 {
 	const FIntPoint ViewRectSize = View.ViewRect.Size();
@@ -1368,7 +1375,7 @@ void FSceneRenderer::ComputeVolumetricFog(FRDGBuilder& GraphBuilder,
 			View.ViewState->PrevLightScatteringViewGridUVToViewRectVolumeUV = ViewRectSize / (FVector2f(VolumetricFogViewGridSize.X, VolumetricFogViewGridSize.Y) * VolumetricFogGridPixelSize);
 
 			View.ViewState->VolumetricFogPrevViewGridRectUVToResourceUV = FVector2f(VolumetricFogViewGridSize.X, VolumetricFogViewGridSize.Y) / FVector2f(VolumetricFogResourceGridSize.X, VolumetricFogResourceGridSize.Y);
-			View.ViewState->VolumetricFogPrevUVMax = FVector2f(ViewRectSize.X - 0.51f, ViewRectSize.Y - 0.51f) / (FVector2f(VolumetricFogResourceGridSize.X, VolumetricFogResourceGridSize.Y) * VolumetricFogGridPixelSize);
+			View.ViewState->VolumetricFogPrevUVMax = GetVolumetricFogUVMax(ViewRectSize, VolumetricFogResourceGridSize, VolumetricFogGridPixelSize);
 		}
 		else if (View.ViewState)
 		{
diff --git a/Engine/Source/Runtime/Renderer/Private/VolumetricFog.h b/Engine/Source/Runtime/Renderer/Private/VolumetricFog.h
--- a/Engine/Source/Runtime/Renderer/Private/VolumetricFog.h
+++ b/Engine/Source/Runtime/Renderer/Private/VolumetricFog.h
@@ -52,6 +52,8 @@ extern FIntVector GetVolumetricFogResourceGridSize(const FViewInfo& View, int32&
 // Grid size for the view rectangle within the allocated resource
 extern FIntVector GetVolumetricFogViewGridSize(const FViewInfo& View, int32& OutVolumetricFogGridPixelSize);
 
+extern FVector2f GetVolumetricFogUVMax(const FVector2f& ViewRectSize, FIntVector VolumetricFogResourceGridSize, int32 VolumetricFogResourceGridPixelSize);
+
 extern FVector2f GetVolumetricFogFroxelToScreenSVPosRatio(const FViewInfo& View);
 
 extern bool DoesPlatformSupportVolumetricFogVoxelization(const FStaticShaderPlatform Platform);
