Nukes tonemaper's USE_GRAIN_QUANTIZATION

* Tonemapper have too many shader permutations
* Doesn't make any meaningful performance difference on slowest platforms
* Renames the GrainQuantization to BackbufferQuantizationDithering to reduce confusion with film grain
* Offer it on all post processing scalability settings

diff --git a/Engine/Config/BaseScalability.ini b/Engine/Config/BaseScalability.ini
--- a/Engine/Config/BaseScalability.ini
+++ b/Engine/Config/BaseScalability.ini
@@ -390,7 +390,6 @@ r.BloomQuality=4
 r.Bloom.ScreenPercentage=25.000
 r.FastBlurThreshold=0
 r.Upscale.Quality=1
-r.Tonemapper.GrainQuantization=0
 r.LightShaftQuality=0
 r.Filter.SizeScale=0.6
 r.Tonemapper.Quality=0
@@ -411,7 +410,6 @@ r.BloomQuality=4
 r.Bloom.ScreenPercentage=35.355
 r.FastBlurThreshold=2
 r.Upscale.Quality=2
-r.Tonemapper.GrainQuantization=0
 r.LightShaftQuality=0
 r.Filter.SizeScale=0.7
 r.Tonemapper.Quality=2
@@ -444,7 +442,6 @@ r.BloomQuality=5
 r.Bloom.ScreenPercentage=50.000
 r.FastBlurThreshold=3
 r.Upscale.Quality=2
-r.Tonemapper.GrainQuantization=1
 r.LightShaftQuality=1
 r.Filter.SizeScale=0.8
 r.Tonemapper.Quality=5
@@ -479,7 +476,6 @@ r.BloomQuality=5
 r.Bloom.ScreenPercentage=50.000
 r.FastBlurThreshold=100
 r.Upscale.Quality=3
-r.Tonemapper.GrainQuantization=1
 r.LightShaftQuality=1
 r.Filter.SizeScale=1
 r.Tonemapper.Quality=5
@@ -517,7 +513,6 @@ r.BloomQuality=5
 r.Bloom.ScreenPercentage=100
 r.FastBlurThreshold=100
 r.Upscale.Quality=3
-r.Tonemapper.GrainQuantization=1
 r.LightShaftQuality=1
 r.Filter.SizeScale=1
 r.Tonemapper.Quality=5
diff --git a/Engine/Shaders/Private/PostProcessTonemap.usf b/Engine/Shaders/Private/PostProcessTonemap.usf
--- a/Engine/Shaders/Private/PostProcessTonemap.usf
+++ b/Engine/Shaders/Private/PostProcessTonemap.usf
@@ -73,7 +73,7 @@ float4 ScreenPosToFilmGrainTextureUV;
 	StructuredBuffer<float4> FilmGrainTextureConstants;
 #endif // SUPPORTS_FILM_GRAIN
 
-float GrainQuantization;
+float BackbufferQuantizationDithering;
 float EditorNITLevel;
 
 uint bOutputInHDR;
@@ -525,15 +525,8 @@ float4 TonemapCommonPS(
 	}
 #endif
 
-	#if USE_GRAIN_QUANTIZATION && !USE_GAMMA_ONLY
+	#if !USE_GAMMA_ONLY
 	{
-		#if DIM_OUTPUT_DEVICE == TONEMAPPER_OUTPUT_NoToneCurve
-			#error LinearNoToneCurve device requires no Grain Quantization
-		#endif
-		#if DIM_OUTPUT_DEVICE == TONEMAPPER_OUTPUT_WithToneCurve
-			#error LinearWithToneCurve device requires no Grain Quantization
-		#endif
-
 		// Compute uniform grain on [0;1]
 		half UniformGrainOP = GrainFromUV(GrainUV.zw); //TODO: verify the quality of GrainFromUV()
 
@@ -567,7 +560,7 @@ float4 TonemapCommonPS(
 			OutColor.rgb = OutColor.ggg;
 		}
 		#else
-			OutColor.rgb += FinalGrain * GrainQuantization;
+			OutColor.rgb += FinalGrain * BackbufferQuantizationDithering;
 		#endif
 	}
 	#endif
diff --git a/Engine/Source/Runtime/Core/Private/HAL/ConsoleManager.cpp b/Engine/Source/Runtime/Core/Private/HAL/ConsoleManager.cpp
--- a/Engine/Source/Runtime/Core/Private/HAL/ConsoleManager.cpp
+++ b/Engine/Source/Runtime/Core/Private/HAL/ConsoleManager.cpp
@@ -3300,13 +3300,6 @@ static TAutoConsoleVariable<int32> CVarFreeSkeletalMeshBuffers(
 		 "1: Free buffers"),
 	ECVF_RenderThreadSafe);
 
-static TAutoConsoleVariable<int32> CVarTonemapperGrainQuantization(
-	TEXT("r.Tonemapper.GrainQuantization"),
-	1,
-	TEXT("0: low (minor performance benefit)\n"
-		 "1: high (default, with high frequency pixel pattern to fight 8 bit color quantization)"),
-	ECVF_Scalability | ECVF_RenderThreadSafe);
-
 static TAutoConsoleVariable<int32> CVarDetailMode(
 	TEXT("r.DetailMode"),
 	2,
diff --git a/Engine/Source/Runtime/Renderer/Private/PostProcess/PostProcessTonemap.cpp b/Engine/Source/Runtime/Renderer/Private/PostProcess/PostProcessTonemap.cpp
--- a/Engine/Source/Runtime/Renderer/Private/PostProcess/PostProcessTonemap.cpp
+++ b/Engine/Source/Runtime/Renderer/Private/PostProcess/PostProcessTonemap.cpp
@@ -63,6 +63,12 @@ TAutoConsoleVariable<int32> CVarFilmGrainCacheTextureConstants(
 	TEXT("Wether the constants related to the film grain should be cached."),
 	ECVF_RenderThreadSafe);
 
+static TAutoConsoleVariable<int32> CVarBackbufferQuantizationDitheringOverride(
+	TEXT("r.BackbufferQuantizationDitheringOverride"), 0,
+	TEXT("Override the bitdepth in bits of each channel of the backbuffer targeted by the quantization dithering. ")
+	TEXT("Disabled by default. Instead is automatically found out by FSceneViewFamily::RenderTarget's pixel format of the backbuffer."),
+	ECVF_RenderThreadSafe);
+
 const int32 GTonemapComputeTileSizeX = 8;
 const int32 GTonemapComputeTileSizeY = 8;
 
@@ -139,14 +145,12 @@ FCommonDomain BuildCommonPermutationDomain(const FViewInfo& View, bool bGammaOnl
 
 // Desktop renderer permutation dimensions.
 class FTonemapperColorFringeDim       : SHADER_PERMUTATION_BOOL("USE_COLOR_FRINGE");
-class FTonemapperGrainQuantizationDim : SHADER_PERMUTATION_BOOL("USE_GRAIN_QUANTIZATION");
 class FTonemapperOutputDeviceDim      : SHADER_PERMUTATION_ENUM_CLASS("DIM_OUTPUT_DEVICE", EDisplayOutputFormat);
 class FTonemapperOutputLuminance	  : SHADER_PERMUTATION_BOOL("OUTPUT_LUMINANCE");
 
 using FDesktopDomain = TShaderPermutationDomain<
 	FCommonDomain,
 	FTonemapperColorFringeDim,
-	FTonemapperGrainQuantizationDim,
 	FTonemapperOutputLuminance,
 	FTonemapperOutputDeviceDim>;
 
@@ -189,12 +193,6 @@ FDesktopDomain RemapPermutation(FDesktopDomain PermutationVector, ERHIFeatureLev
 		PermutationVector.Set<FTonemapperOutputDeviceDim>(EDisplayOutputFormat::SDR_sRGB);
 	}
 
-	// Disable grain quantization for LinearNoToneCurve and LinearWithToneCurve output device
-	if (PermutationVector.Get<FTonemapperOutputDeviceDim>() == EDisplayOutputFormat::HDR_LinearNoToneCurve || PermutationVector.Get<FTonemapperOutputDeviceDim>() == EDisplayOutputFormat::HDR_LinearWithToneCurve)
-		PermutationVector.Set<FTonemapperGrainQuantizationDim>(false);
-	else
-		PermutationVector.Set<FTonemapperGrainQuantizationDim>(true);
-
 	if (FeatureLevel < ERHIFeatureLevel::SM5)
 	{
 		CommonPermutationVector.Set<FTonemapperFilmGrainDim>(false);
@@ -221,8 +219,7 @@ bool ShouldCompileDesktopPermutation(const FGlobalShaderPermutationParameters& P
 
 	if (CommonPermutationVector.Get<FTonemapperGammaOnlyDim>())
 	{
-		return !PermutationVector.Get<FTonemapperColorFringeDim>() &&
-			!PermutationVector.Get<FTonemapperGrainQuantizationDim>();
+		return !PermutationVector.Get<FTonemapperColorFringeDim>();
 	}
 	return true;
 }
@@ -337,7 +334,7 @@ BEGIN_SHADER_PARAMETER_STRUCT(FTonemapParameters, )
 	SHADER_PARAMETER(float, LUTScale)
 	SHADER_PARAMETER(float, LUTOffset)
 	SHADER_PARAMETER(float, EditorNITLevel)
-	SHADER_PARAMETER(float, GrainQuantization)
+	SHADER_PARAMETER(float, BackbufferQuantizationDithering)
 	SHADER_PARAMETER(uint32, bOutputInHDR)
 END_SHADER_PARAMETER_STRUCT()
 
@@ -654,16 +651,6 @@ FScreenPassTexture AddTonemapPass(FRDGBuilder& GraphBuilder, const FViewInfo& Vi
 	}
 	#endif
 
-	float GrainQuantization = 1.0f / 1024.0f;
-	if (View.Family->RenderTarget && View.Family->RenderTarget->GetRenderTargetTexture())
-	{
-		EPixelFormat BackbufferPixelFormat = View.Family->RenderTarget->GetRenderTargetTexture()->GetFormat();
-		if (BackbufferPixelFormat == PF_B8G8R8A8 || BackbufferPixelFormat == PF_R8G8B8A8)
-		{
-			GrainQuantization = 1.0 / 256.0;
-		}
-	}
-
 	FTonemapParameters CommonParameters;
 	CommonParameters.View = View.ViewUniformBuffer;
 
@@ -779,7 +766,32 @@ FScreenPassTexture AddTonemapPass(FRDGBuilder& GraphBuilder, const FViewInfo& Vi
 	CommonParameters.ChromaticAberrationParams = ChromaticAberrationParams;
 	CommonParameters.TonemapperParams = FVector4f(PostProcessSettings.VignetteIntensity, SharpenDiv6, 0.0f, 0.0f);
 	CommonParameters.EditorNITLevel = EditorNITLevel;
-	CommonParameters.GrainQuantization = GrainQuantization;
+	{
+		const bool bSupportsBackbufferQuantizationDithering = EDisplayOutputFormat(CommonParameters.OutputDevice.OutputDevice) != EDisplayOutputFormat::HDR_LinearNoToneCurve && EDisplayOutputFormat(CommonParameters.OutputDevice.OutputDevice) != EDisplayOutputFormat::HDR_LinearWithToneCurve;
+		const int32 BitdepthOverride = CVarBackbufferQuantizationDitheringOverride.GetValueOnRenderThread();
+
+		if (!bSupportsBackbufferQuantizationDithering)
+		{
+			CommonParameters.BackbufferQuantizationDithering = 0.0f;
+		}
+		else if (BitdepthOverride > 0)
+		{
+			CommonParameters.BackbufferQuantizationDithering = 1.0f / (FMath::Pow(2.0f, float(BitdepthOverride)) - 1.0f);
+		}
+		else
+		{
+			CommonParameters.BackbufferQuantizationDithering = 1.0f / 1023.0f;
+
+			if (View.Family->RenderTarget && View.Family->RenderTarget->GetRenderTargetTexture())
+			{
+				EPixelFormat BackbufferPixelFormat = View.Family->RenderTarget->GetRenderTargetTexture()->GetFormat();
+				if (BackbufferPixelFormat == PF_B8G8R8A8 || BackbufferPixelFormat == PF_R8G8B8A8)
+				{
+					CommonParameters.BackbufferQuantizationDithering = 1.0f / 255.0f;
+				}
+			}
+		}
+	}
 	CommonParameters.bOutputInHDR = ViewFamily.bIsHDR;
 	CommonParameters.LUTSize = LUTSize;
 	CommonParameters.InvLUTSize = 1.0f / LUTSize;
@@ -856,13 +868,6 @@ FScreenPassTexture AddTonemapPass(FRDGBuilder& GraphBuilder, const FViewInfo& Vi
 
 		if (!CommonDomain.Get<TonemapperPermutation::FTonemapperGammaOnlyDim>())
 		{
-			// Grain Quantization
-			{
-				static TConsoleVariableData<int32>* CVar = IConsoleManager::Get().FindTConsoleVariableDataInt(TEXT("r.Tonemapper.GrainQuantization"));
-				const int32 Value = CVar->GetValueOnRenderThread();
-				DesktopPermutationVector.Set<TonemapperPermutation::FTonemapperGrainQuantizationDim>(Value > 0);
-			}
-
 			DesktopPermutationVector.Set<TonemapperPermutation::FTonemapperColorFringeDim>(PostProcessSettings.SceneFringeIntensity > 0.01f);
 		}
 
