Added EMeshPassFlags::CachedMeshCommands flag to SingleLayerWater MeshProcessors to enable dynamic instancing.

#rnx

diff --git a/Engine/Source/Runtime/Renderer/Private/SingleLayerWaterRendering.cpp b/Engine/Source/Runtime/Renderer/Private/SingleLayerWaterRendering.cpp
--- a/Engine/Source/Runtime/Renderer/Private/SingleLayerWaterRendering.cpp
+++ b/Engine/Source/Runtime/Renderer/Private/SingleLayerWaterRendering.cpp
@@ -1655,7 +1655,7 @@ FMeshPassProcessor* CreateSingleLayerWaterPassProcessor(ERHIFeatureLevel::Type F
 	return new FSingleLayerWaterPassMeshProcessor(Scene, FeatureLevel, InViewIfDynamicMeshCommand, DrawRenderState, InDrawListContext);
 }
 
-REGISTER_MESHPASSPROCESSOR_AND_PSOCOLLECTOR(SingleLayerWater, CreateSingleLayerWaterPassProcessor, EShadingPath::Deferred, EMeshPass::SingleLayerWaterPass, EMeshPassFlags::MainView);
+REGISTER_MESHPASSPROCESSOR_AND_PSOCOLLECTOR(SingleLayerWater, CreateSingleLayerWaterPassProcessor, EShadingPath::Deferred, EMeshPass::SingleLayerWaterPass, EMeshPassFlags::CachedMeshCommands | EMeshPassFlags::MainView);
 
 
 class FSingleLayerWaterDepthPrepassMeshProcessor : public FSceneRenderingAllocatorObject<FSingleLayerWaterDepthPrepassMeshProcessor>, public FMeshPassProcessor
@@ -1930,4 +1930,4 @@ FMeshPassProcessor* CreateSingleLayerWaterDepthPrepassProcessor(ERHIFeatureLevel
 	return new FSingleLayerWaterDepthPrepassMeshProcessor(Scene, FeatureLevel, InViewIfDynamicMeshCommand, DrawRenderState, InDrawListContext);
 }
 
-REGISTER_MESHPASSPROCESSOR_AND_PSOCOLLECTOR(SingleLayerWaterDepthPrepass, CreateSingleLayerWaterDepthPrepassProcessor, EShadingPath::Deferred, EMeshPass::SingleLayerWaterDepthPrepass, EMeshPassFlags::MainView);
+REGISTER_MESHPASSPROCESSOR_AND_PSOCOLLECTOR(SingleLayerWaterDepthPrepass, CreateSingleLayerWaterDepthPrepassProcessor, EShadingPath::Deferred, EMeshPass::SingleLayerWaterDepthPrepass, EMeshPassFlags::CachedMeshCommands | EMeshPassFlags::MainView);
