Changes r.TSR.ShadingRejection.ExposureOffset's default

diff --git a/Engine/Source/Runtime/Renderer/Private/PostProcess/TemporalSuperResolution.cpp b/Engine/Source/Runtime/Renderer/Private/PostProcess/TemporalSuperResolution.cpp
--- a/Engine/Source/Runtime/Renderer/Private/PostProcess/TemporalSuperResolution.cpp
+++ b/Engine/Source/Runtime/Renderer/Private/PostProcess/TemporalSuperResolution.cpp
@@ -159,7 +159,7 @@ TAutoConsoleVariable<int32> CVarTSRShadingTileOverscan(
 	ECVF_Scalability | ECVF_RenderThreadSafe);
 
 TAutoConsoleVariable<float> CVarTSRShadingExposureOffset(
-	TEXT("r.TSR.ShadingRejection.ExposureOffset"), 0.0, // TODO: update engine default with new EngineTest references
+	TEXT("r.TSR.ShadingRejection.ExposureOffset"), 3.0,
 	TEXT("The shading rejection needs to have a representative idea how bright a linear color pixel ends up displayed to the user. ")
 	TEXT("And the shading rejection detect if a color become to changed to be visible in the back buffer by comparing to MeasureBackbufferLDRQuantizationError().\n")
 	TEXT("\n")
