Add ECVF_RenderThreadSafe to r.VRS.DebugForceRate to silence warning.

diff --git a/Engine/Source/Runtime/Renderer/Private/VariableRateShading/VariableRateShadingImageManager.cpp b/Engine/Source/Runtime/Renderer/Private/VariableRateShading/VariableRateShadingImageManager.cpp
--- a/Engine/Source/Runtime/Renderer/Private/VariableRateShading/VariableRateShadingImageManager.cpp
+++ b/Engine/Source/Runtime/Renderer/Private/VariableRateShading/VariableRateShadingImageManager.cpp
@@ -39,7 +39,8 @@ int GVRSDebugForceRate = -1;
 FAutoConsoleVariableRef CVarVRSDebugForceRate(
 	TEXT("r.VRS.DebugForceRate"),
 	GVRSDebugForceRate,
-	TEXT("-1 : None, 0 : Force 1x1, 1 : Force 1x2, 2 : Force 2x1, 3: Force 2x2, 4 : Force 2x4, 5 : Force 4x2, 6 : Force 4x4"));
+	TEXT("-1 : None, 0 : Force 1x1, 1 : Force 1x2, 2 : Force 2x1, 3: Force 2x2, 4 : Force 2x4, 5 : Force 4x2, 6 : Force 4x4"),
+	ECVF_RenderThreadSafe);
 
 /**
  * Shaders
