FWaterRefractionCopyPS doesn't need the view UB.

#rnx

diff --git a/Engine/Source/Runtime/Renderer/Private/SingleLayerWaterRendering.cpp b/Engine/Source/Runtime/Renderer/Private/SingleLayerWaterRendering.cpp
--- a/Engine/Source/Runtime/Renderer/Private/SingleLayerWaterRendering.cpp
+++ b/Engine/Source/Runtime/Renderer/Private/SingleLayerWaterRendering.cpp
@@ -404,7 +404,6 @@ class FWaterRefractionCopyPS : public FGlobalShader
 	SHADER_USE_PARAMETER_STRUCT(FWaterRefractionCopyPS, FGlobalShader);
 
 	BEGIN_SHADER_PARAMETER_STRUCT(FParameters, )
-		SHADER_PARAMETER_STRUCT_INCLUDE(FViewShaderParameters, View)
 		SHADER_PARAMETER_RDG_TEXTURE(Texture2D, SceneColorCopyDownsampleTexture)
 		SHADER_PARAMETER_SAMPLER(SamplerState, SceneColorCopyDownsampleSampler)
 		SHADER_PARAMETER_RDG_TEXTURE(Texture2D, SceneDepthCopyDownsampleTexture)
@@ -747,7 +746,6 @@ static FSceneWithoutWaterTextures AddCopySceneWithoutWaterPass(
 		RDG_EVENT_SCOPE_CONDITIONAL(GraphBuilder, Views.Num() > 1, "View%d", ViewIndex);
 
 		FWaterRefractionCopyPS::FParameters* PassParameters = GraphBuilder.AllocParameters<FWaterRefractionCopyPS::FParameters>();
-		PassParameters->View = View.GetShaderParameters();
 		PassParameters->SceneColorCopyDownsampleTexture = SceneColorTexture;
 		PassParameters->SceneColorCopyDownsampleSampler = TStaticSamplerState<SF_Bilinear, AM_Clamp, AM_Clamp, AM_Clamp>::GetRHI();
 		PassParameters->SceneDepthCopyDownsampleTexture = SceneDepthTexture;
