Compiles out the tonemaper's USE_BLOOM shader permutation on desktop renderer

diff --git a/Engine/Source/Runtime/Renderer/Private/PostProcess/PostProcessTonemap.cpp b/Engine/Source/Runtime/Renderer/Private/PostProcess/PostProcessTonemap.cpp
--- a/Engine/Source/Runtime/Renderer/Private/PostProcess/PostProcessTonemap.cpp
+++ b/Engine/Source/Runtime/Renderer/Private/PostProcess/PostProcessTonemap.cpp
@@ -186,15 +186,20 @@ FDesktopDomain RemapPermutation(FDesktopDomain PermutationVector, ERHIFeatureLev
 	// You most likely need Bloom anyway.
 	CommonPermutationVector.Set<FTonemapperBloomDim>(true);
 
-	// Mobile supports only sRGB and LinearNoToneCurve output
-	if (FeatureLevel <= ERHIFeatureLevel::ES3_1 &&
-		PermutationVector.Get<FTonemapperOutputDeviceDim>() != EDisplayOutputFormat::HDR_LinearNoToneCurve)
+	if (FeatureLevel >= ERHIFeatureLevel::SM5)
 	{
-		PermutationVector.Set<FTonemapperOutputDeviceDim>(EDisplayOutputFormat::SDR_sRGB);
+		// Disabling bloom on desktop renderer is very rare, not worth compiling shader permutation without.
+		CommonPermutationVector.Set<FTonemapperBloomDim>(true);
 	}
-
-	if (FeatureLevel < ERHIFeatureLevel::SM5)
+	else
 	{
+		// Mobile supports only sRGB and LinearNoToneCurve output
+		if (PermutationVector.Get<FTonemapperOutputDeviceDim>() != EDisplayOutputFormat::HDR_LinearNoToneCurve)
+		{
+			PermutationVector.Set<FTonemapperOutputDeviceDim>(EDisplayOutputFormat::SDR_sRGB);
+		}
+
+		// Mobile doesn't support film grain.
 		CommonPermutationVector.Set<FTonemapperFilmGrainDim>(false);
 	}
 
@@ -799,37 +804,11 @@ FScreenPassTexture AddTonemapPass(FRDGBuilder& GraphBuilder, const FViewInfo& Vi
 	CommonParameters.LUTOffset = 0.5f / LUTSize;
 	CommonParameters.LensPrincipalPointOffsetScale = View.LensPrincipalPointOffsetScale;
 
-	// TODO: PostProcessSettings.BloomDirtMask->GetResource() is not thread safe
-	if (PostProcessSettings.BloomDirtMask && PostProcessSettings.BloomDirtMask->GetResource() && PostProcessSettings.BloomDirtMaskIntensity > 0 && PostProcessSettings.BloomIntensity > 0.0)
-	{
-		CommonParameters.BloomDirtMaskTint = PostProcessSettings.BloomDirtMaskTint * PostProcessSettings.BloomDirtMaskIntensity / PostProcessSettings.BloomIntensity;
-		CommonParameters.BloomDirtMaskTexture = PostProcessSettings.BloomDirtMask->GetResource()->TextureRHI;
-		CommonParameters.BloomDirtMaskSampler = BilinearClampSampler;
-	}
-	else
-	{
-		CommonParameters.BloomDirtMaskTint = FLinearColor::Black;
-		CommonParameters.BloomDirtMaskTexture = GBlackTexture->TextureRHI;
-		CommonParameters.BloomDirtMaskSampler = BilinearClampSampler;
-	}
-
-	if (!FTonemapInputs::SupportsSceneColorApplyParametersBuffer(View.GetShaderPlatform()))
-	{
-		check(Inputs.SceneColorApplyParamaters == nullptr);
-	}
-	else if (Inputs.SceneColorApplyParamaters)
-	{
-		CommonParameters.SceneColorApplyParamaters = GraphBuilder.CreateSRV(Inputs.SceneColorApplyParamaters);
-	}
-	else
-	{
-		FRDGBufferRef ApplyParametersBuffer = GSystemTextures.GetDefaultStructuredBuffer(GraphBuilder, sizeof(FVector4f), FVector4f(1.0f, 1.0f, 1.0f, 1.0f));
-		CommonParameters.SceneColorApplyParamaters = GraphBuilder.CreateSRV(ApplyParametersBuffer);
-	}
-
 	// Bloom parameters
 	{
-		if (Inputs.Bloom.Texture)
+		const bool bUseBloom = Inputs.Bloom.Texture != nullptr;
+
+		if (bUseBloom)
 		{
 			const FScreenPassTextureViewport BloomViewport(Inputs.Bloom);
 			CommonParameters.ColorToBloom = FScreenTransform::ChangeTextureUVCoordinateFromTo(SceneColorViewport, BloomViewport);
@@ -846,6 +825,34 @@ FScreenPassTexture AddTonemapPass(FRDGBuilder& GraphBuilder, const FViewInfo& Vi
 			CommonParameters.BloomTexture = GSystemTextures.GetBlackDummy(GraphBuilder);
 			CommonParameters.BloomSampler = BilinearClampSampler;
 		}
+
+		if (!FTonemapInputs::SupportsSceneColorApplyParametersBuffer(View.GetShaderPlatform()))
+		{
+			check(Inputs.SceneColorApplyParamaters == nullptr);
+		}
+		else if (Inputs.SceneColorApplyParamaters)
+		{
+			CommonParameters.SceneColorApplyParamaters = GraphBuilder.CreateSRV(Inputs.SceneColorApplyParamaters);
+		}
+		else
+		{
+			FRDGBufferRef ApplyParametersBuffer = GSystemTextures.GetDefaultStructuredBuffer(GraphBuilder, sizeof(FVector4f), FVector4f(1.0f, 1.0f, 1.0f, 1.0f));
+			CommonParameters.SceneColorApplyParamaters = GraphBuilder.CreateSRV(ApplyParametersBuffer);
+		}
+
+		// TODO: PostProcessSettings.BloomDirtMask->GetResource() is not thread safe
+		if (bUseBloom && PostProcessSettings.BloomDirtMask && PostProcessSettings.BloomDirtMask->GetResource() && PostProcessSettings.BloomDirtMaskIntensity > 0 && PostProcessSettings.BloomIntensity > 0.0)
+		{
+			CommonParameters.BloomDirtMaskTint = PostProcessSettings.BloomDirtMaskTint * PostProcessSettings.BloomDirtMaskIntensity / PostProcessSettings.BloomIntensity;
+			CommonParameters.BloomDirtMaskTexture = PostProcessSettings.BloomDirtMask->GetResource()->TextureRHI;
+			CommonParameters.BloomDirtMaskSampler = BilinearClampSampler;
+		}
+		else
+		{
+			CommonParameters.BloomDirtMaskTint = FLinearColor::Black;
+			CommonParameters.BloomDirtMaskTexture = GBlackTexture->TextureRHI;
+			CommonParameters.BloomDirtMaskSampler = BilinearClampSampler;
+		}
 	}
 
 	// forward transformation from shader:
