Upgraded VT tasks to high priority as they are on the critical path.

diff --git a/Engine/Source/Runtime/Renderer/Private/VT/VirtualTextureSystem.cpp b/Engine/Source/Runtime/Renderer/Private/VT/VirtualTextureSystem.cpp
--- a/Engine/Source/Runtime/Renderer/Private/VT/VirtualTextureSystem.cpp
+++ b/Engine/Source/Runtime/Renderer/Private/VT/VirtualTextureSystem.cpp
@@ -265,7 +265,7 @@ public:
 	}
 
 	static ESubsequentsMode::Type GetSubsequentsMode() { return ESubsequentsMode::TrackSubsequents; }
-	ENamedThreads::Type GetDesiredThread() { return ENamedThreads::AnyNormalThreadNormalTask; }
+	ENamedThreads::Type GetDesiredThread() { return ENamedThreads::AnyHiPriThreadHiPriTask; }
 	FORCEINLINE TStatId GetStatId() const { RETURN_QUICK_DECLARE_CYCLE_STAT(FAddRequestedTilesTask, STATGROUP_VirtualTexturing); }
 };
 
@@ -298,7 +298,7 @@ public:
 	}
 
 	static ESubsequentsMode::Type GetSubsequentsMode() { return ESubsequentsMode::TrackSubsequents; }
-	ENamedThreads::Type GetDesiredThread() { return ENamedThreads::AnyNormalThreadNormalTask; }
+	ENamedThreads::Type GetDesiredThread() { return ENamedThreads::AnyHiPriThreadHiPriTask; }
 	FORCEINLINE TStatId GetStatId() const { RETURN_QUICK_DECLARE_CYCLE_STAT(FGatherRequestsTask, STATGROUP_VirtualTexturing); }
 };
 
@@ -2644,7 +2644,7 @@ void FVirtualTextureSystem::BeginUpdate(FRDGBuilder& GraphBuilder, FVirtualTextu
 		// Reset the request list for the gather in EndUpdate.
 		Updater->MergedRequestList->Reset(bContinousUpdates);
 
-	}, Updater->bAsyncTaskAllowed);
+	}, UE::Tasks::ETaskPriority::High, Updater->bAsyncTaskAllowed);
 }
 
 TUniquePtr<FVirtualTextureUpdater> FVirtualTextureSystem::BeginUpdate(FRDGBuilder& GraphBuilder, ERHIFeatureLevel::Type FeatureLevel, FScene* Scene, const FVirtualTextureUpdateSettings& Settings)
