Moved ViewDebugInfo to its own source file.

diff --git a/Engine/Source/Runtime/RenderCore/Private/RenderCore.cpp b/Engine/Source/Runtime/RenderCore/Private/RenderCore.cpp
--- a/Engine/Source/Runtime/RenderCore/Private/RenderCore.cpp
+++ b/Engine/Source/Runtime/RenderCore/Private/RenderCore.cpp
@@ -117,7 +117,6 @@ DEFINE_STAT(STAT_ViewRelevance);
 DEFINE_STAT(STAT_ComputeViewRelevance);
 DEFINE_STAT(STAT_OcclusionCull);
 DEFINE_STAT(STAT_UpdatePrimitiveFading);
-DEFINE_STAT(STAT_PrimitiveCull);
 DEFINE_STAT(STAT_DecompressPrecomputedOcclusion);
 DEFINE_STAT(STAT_ViewVisibilityTime);
 
diff --git a/Engine/Source/Runtime/RenderCore/Public/RenderCore.h b/Engine/Source/Runtime/RenderCore/Public/RenderCore.h
--- a/Engine/Source/Runtime/RenderCore/Public/RenderCore.h
+++ b/Engine/Source/Runtime/RenderCore/Public/RenderCore.h
@@ -66,7 +66,6 @@ DECLARE_DWORD_COUNTER_STAT_EXTERN(TEXT("Ray tracing active instances"), STAT_Ray
 
 DECLARE_CYCLE_STAT_EXTERN(TEXT("View Visibility"),STAT_ViewVisibilityTime,STATGROUP_InitViews, RENDERCORE_API);
 DECLARE_CYCLE_STAT_EXTERN(TEXT("Decompress Occlusion"),STAT_DecompressPrecomputedOcclusion,STATGROUP_InitViews, RENDERCORE_API);
-DECLARE_CYCLE_STAT_EXTERN(TEXT("Primitive Cull"), STAT_PrimitiveCull,STATGROUP_InitViews, RENDERCORE_API);
 DECLARE_CYCLE_STAT_EXTERN(TEXT("Update Fading"),STAT_UpdatePrimitiveFading,STATGROUP_InitViews, RENDERCORE_API);
 DECLARE_CYCLE_STAT_EXTERN(TEXT("Occlusion Cull"),STAT_OcclusionCull,STATGROUP_InitViews, RENDERCORE_API);
 DECLARE_CYCLE_STAT_EXTERN(TEXT("View Relevance"),STAT_ViewRelevance,STATGROUP_InitViews, RENDERCORE_API);
diff --git a/Engine/Source/Runtime/Renderer/Private/SceneVisibility.cpp b/Engine/Source/Runtime/Renderer/Private/SceneVisibility.cpp
--- a/Engine/Source/Runtime/Renderer/Private/SceneVisibility.cpp
+++ b/Engine/Source/Runtime/Renderer/Private/SceneVisibility.cpp
@@ -982,8 +982,6 @@ static void PrimitiveCull(FVisibilityTaskData& TaskData, FScene* RESTRICT Scene,
 	PrerequisiteTask = Scene->GetCacheRayTracingPrimitivesTask();
 #endif
 
-	SCOPE_CYCLE_COUNTER(STAT_PrimitiveCull);
-
 	FSceneBitArray& VisibleNodes = *TaskData.Allocator.Create<FSceneBitArray>();
 
 	if (bShouldVisibilityCull && Flags.bUseVisibilityOctree)
@@ -4243,298 +4241,6 @@ void FSceneViewState::UpdateMotionBlurTimeScale(const FViewInfo& View)
 	MotionBlurTimeScale = MotionBlurTargetDeltaTime / DeltaRealTime;
 }
 
-#if !UE_BUILD_SHIPPING
-static uint32 GetDrawCountFromPrimitiveSceneInfo(FScene* Scene, const FPrimitiveSceneInfo* PrimitiveSceneInfo)
-{
-	uint32 DrawCount = 0;
-	for (const FCachedMeshDrawCommandInfo& CachedCommand : PrimitiveSceneInfo->StaticMeshCommandInfos)
-	{
-		if (CachedCommand.MeshPass != EMeshPass::BasePass)
-			continue;
-
-		if (CachedCommand.StateBucketId != INDEX_NONE || CachedCommand.CommandIndex >= 0)
-		{
-			DrawCount++;
-		}
-	}
-
-	return DrawCount;
-}
-
-
-FViewDebugInfo FViewDebugInfo::Instance;
-
-FViewDebugInfo::FViewDebugInfo()
-{
-	bHasEverUpdated = false;
-	bIsOutdated = true;
-	bShouldUpdate = false;
-	bShouldCaptureSingleFrame = false;
-}
-
-void FViewDebugInfo::ProcessPrimitive(FPrimitiveSceneInfo* PrimitiveSceneInfo, const FViewInfo& View, FScene* Scene, const UPrimitiveComponent* DebugComponent)
-{
-	if (!DebugComponent->IsRegistered())
-	{
-		return;
-	}
-	AActor* Actor = DebugComponent->GetOwner();
-	FString FullName = DebugComponent->GetName();
-	const uint32 DrawCount = GetDrawCountFromPrimitiveSceneInfo(Scene, PrimitiveSceneInfo);
-
-	TArray<UMaterialInterface*> Materials;
-	DebugComponent->GetUsedMaterials(Materials);
-	const int32 LOD = PrimitiveSceneInfo->Proxy ? PrimitiveSceneInfo->Proxy->GetLOD(&View) : INDEX_NONE;
-	int32 Triangles = 0;
-	if (const UStaticMeshComponent* StaticMeshComponent = Cast<UStaticMeshComponent>(DebugComponent))
-	{
-		Triangles = StaticMeshComponent->GetStaticMesh()->GetNumTriangles(LOD);
-	}
-	else if (const USkeletalMeshComponent* SkeletalMeshComponent = Cast<USkeletalMeshComponent>(DebugComponent))
-	{
-		for (const FSkeletalMeshLODRenderData& RenderData : SkeletalMeshComponent->GetSkeletalMeshRenderData()->LODRenderData)
-		{
-			Triangles += RenderData.MultiSizeIndexContainer.GetIndexBuffer()->Num() / 3;
-		}
-	}
-		
-	const FPrimitiveInfo PrimitiveInfo = {
-		Actor,
-		DebugComponent->ComponentId,
-		const_cast<UPrimitiveComponent*>(DebugComponent), // This is probably a bad idea, find alternative
-		PrimitiveSceneInfo,
-		MoveTemp(Materials),
-		MoveTemp(FullName),
-		DrawCount,
-		Triangles,
-		LOD
-	};
-
-	Primitives.Add(PrimitiveInfo);
-}
-
-void FViewDebugInfo::DumpToCSV() const
-{
-	const FString OutputPath = FPaths::ProfilingDir() / TEXT("Primitives") / FString::Printf(TEXT("PrimitivesDetailed-%s.csv"), *FDateTime::Now().ToString());
-	const bool bSuppressViewer = true;
-	FDiagnosticTableViewer DrawViewer(*OutputPath, bSuppressViewer);
-	DrawViewer.AddColumn(TEXT("Name"));
-	DrawViewer.AddColumn(TEXT("ActorClass"));
-	DrawViewer.AddColumn(TEXT("Actor"));
-	DrawViewer.AddColumn(TEXT("Location"));
-	DrawViewer.AddColumn(TEXT("NumMaterials"));
-	DrawViewer.AddColumn(TEXT("Materials"));
-	DrawViewer.AddColumn(TEXT("NumDraws"));
-	DrawViewer.AddColumn(TEXT("LOD"));
-	DrawViewer.AddColumn(TEXT("Triangles"));
-	DrawViewer.CycleRow();
-
-	FRWScopeLock ScopeLock(Lock, SLT_ReadOnly);
-	const FPrimitiveSceneInfo* LastPrimitiveSceneInfo = nullptr;
-	for (const FPrimitiveInfo& Primitive : Primitives)
-	{
-		if (Primitive.PrimitiveSceneInfo != LastPrimitiveSceneInfo)
-		{
-			DrawViewer.AddColumn(*Primitive.Name);
-			DrawViewer.AddColumn(Primitive.Owner ? *Primitive.Owner->GetClass()->GetName() : TEXT(""));
-			DrawViewer.AddColumn(Primitive.Owner ? *Primitive.Owner->GetFullName() : TEXT(""));
-			DrawViewer.AddColumn(Primitive.Owner ?
-				*FString::Printf(TEXT("{%s}"), *Primitive.Owner->GetActorLocation().ToString()) : TEXT(""));
-			DrawViewer.AddColumn(*FString::Printf(TEXT("%d"), Primitive.Materials.Num()));
-			FString Materials = "[";
-			for (int i = 0; i < Primitive.Materials.Num(); i++)
-			{
-				if (Primitive.Materials[i] && Primitive.Materials[i]->GetMaterial())
-				{
-					Materials += Primitive.Materials[i]->GetMaterial()->GetName();
-				}
-				else
-				{
-					Materials += "Null";
-				}
-				
-				if (i < Primitive.Materials.Num() - 1)
-				{
-					Materials += ", ";
-				}
-			}
-			Materials += "]";
-			DrawViewer.AddColumn(*FString::Printf(TEXT("%s"), *Materials));
-			DrawViewer.AddColumn(*FString::Printf(TEXT("%d"), Primitive.DrawCount));
-			DrawViewer.AddColumn(*FString::Printf(TEXT("%d"), Primitive.LOD));
-			DrawViewer.AddColumn(*FString::Printf(TEXT("%d"), Primitive.TriangleCount));
-			DrawViewer.CycleRow();
-
-			LastPrimitiveSceneInfo = Primitive.PrimitiveSceneInfo;
-		}
-	}
-}
-
-void FViewDebugInfo::CaptureNextFrame()
-{
-	FRWScopeLock ScopeLock(Lock, SLT_Write);
-	bShouldCaptureSingleFrame = true;
-	bShouldUpdate = true;
-}
-
-void FViewDebugInfo::EnableLiveCapture()
-{
-	FRWScopeLock ScopeLock(Lock, SLT_Write);
-	bShouldCaptureSingleFrame = false;
-	bShouldUpdate = true;
-}
-
-void FViewDebugInfo::DisableLiveCapture()
-{
-	FRWScopeLock ScopeLock(Lock, SLT_Write);
-	bShouldCaptureSingleFrame = false;
-	bShouldUpdate = false;
-}
-
-bool FViewDebugInfo::HasEverUpdated() const
-{
-	FRWScopeLock ScopeLock(Lock, SLT_ReadOnly);
-	return bHasEverUpdated;
-}
-
-bool FViewDebugInfo::IsOutOfDate() const
-{
-	FRWScopeLock ScopeLock(Lock, SLT_ReadOnly);
-	return bIsOutdated;
-}
-
-void FSceneRenderer::ProcessPrimitives(const FViewInfo& View, const FViewCommands& ViewCommands) const
-{
-	FViewDebugInfo& DebugInfo = FViewDebugInfo::Instance;
-	{
-		FRWScopeLock ScopeLock(DebugInfo.Lock, SLT_Write);
-		DebugInfo.bIsOutdated = true;
-	
-		if (!DebugInfo.bShouldUpdate && !bDumpDetailedPrimitivesNextFrame)
-		{
-			return;
-		}
-
-		if (DebugInfo.bShouldCaptureSingleFrame)
-		{
-			DebugInfo.bShouldCaptureSingleFrame = false;
-			DebugInfo.bShouldUpdate = false;
-		}
-
-		// TODO: Add profiling to this function
-	
-		DebugInfo.Primitives.Empty(ViewCommands.MeshCommands[EMeshPass::BasePass].Num() + ViewCommands.DynamicMeshCommandBuildRequests[EMeshPass::BasePass].Num());
-
-		for (const FVisibleMeshDrawCommand& Mesh : ViewCommands.MeshCommands[EMeshPass::BasePass])
-		{
-			const int32 PrimitiveId = Mesh.PrimitiveIdInfo.ScenePrimitiveId;
-			if (PrimitiveId >= 0 && PrimitiveId < Scene->Primitives.Num())
-			{
-				FPrimitiveSceneInfo* PrimitiveSceneInfo = Scene->Primitives[PrimitiveId];
-				DebugInfo.ProcessPrimitive(PrimitiveSceneInfo, View, Scene, PrimitiveSceneInfo->ComponentForDebuggingOnly);
-			}
-		}
-
-		for (const FStaticMeshBatch* StaticMeshBatch : ViewCommands.DynamicMeshCommandBuildRequests[EMeshPass::BasePass])
-		{
-			FPrimitiveSceneInfo* PrimitiveSceneInfo = StaticMeshBatch->PrimitiveSceneInfo;
-			DebugInfo.ProcessPrimitive(PrimitiveSceneInfo, View, Scene, PrimitiveSceneInfo->ComponentForDebuggingOnly);
-		}
-
-		DebugInfo.bHasEverUpdated = true;
-		DebugInfo.bIsOutdated = false;
-	}
-	DebugInfo.OnUpdate.Broadcast();
-	
-	if (bDumpDetailedPrimitivesNextFrame)
-	{
-		DebugInfo.DumpToCSV();
-		bDumpDetailedPrimitivesNextFrame = false;
-	}
-	
-}
-
-void FSceneRenderer::DumpPrimitives(const FViewCommands& ViewCommands)
-{
-	if (!bDumpPrimitivesNextFrame)
-	{
-		return;
-	}
-
-	bDumpPrimitivesNextFrame = false;
-
-	struct FPrimitiveInfo
-	{
-		const FPrimitiveSceneInfo* PrimitiveSceneInfo;
-		FString Name;
-		uint32 DrawCount;
-
-		bool operator<(const FPrimitiveInfo& Other) const
-		{
-			// Sort by name to group similar assets together, then by exact primitives so we can ignore duplicates
-			const int32 NameCompare = Name.Compare(Other.Name);
-			if (NameCompare != 0)
-			{
-				return NameCompare < 0;
-			}
-
-			return PrimitiveSceneInfo < Other.PrimitiveSceneInfo;
-		}
-	};
-
-	TArray<FPrimitiveInfo> Primitives;
-	Primitives.Reserve(ViewCommands.MeshCommands[EMeshPass::BasePass].Num() + ViewCommands.DynamicMeshCommandBuildRequests[EMeshPass::BasePass].Num());
-
-	{
-		for (const FVisibleMeshDrawCommand& Mesh : ViewCommands.MeshCommands[EMeshPass::BasePass])
-		{
-			int32 PrimitiveId = Mesh.PrimitiveIdInfo.ScenePrimitiveId;
-			if (PrimitiveId >= 0 && PrimitiveId < Scene->Primitives.Num())
-			{
-				const FPrimitiveSceneInfo* PrimitiveSceneInfo = Scene->Primitives[PrimitiveId];
-				FString FullName = PrimitiveSceneInfo->ComponentForDebuggingOnly->GetFullName();
-
-				uint32 DrawCount = GetDrawCountFromPrimitiveSceneInfo(Scene, PrimitiveSceneInfo);
-
-				Primitives.Add({ PrimitiveSceneInfo, MoveTemp(FullName), DrawCount });
-			}
-		}
-
-		for (const FStaticMeshBatch* StaticMeshBatch : ViewCommands.DynamicMeshCommandBuildRequests[EMeshPass::BasePass])
-		{
-			const FPrimitiveSceneInfo* PrimitiveSceneInfo = StaticMeshBatch->PrimitiveSceneInfo;
-			FString FullName = PrimitiveSceneInfo->ComponentForDebuggingOnly->GetFullName();
-
-			uint32 DrawCount = GetDrawCountFromPrimitiveSceneInfo(Scene, PrimitiveSceneInfo);
-
-			Primitives.Add({ PrimitiveSceneInfo, MoveTemp(FullName), DrawCount });
-		}
-	}
-
-	Primitives.Sort();
-
-	const FString OutputPath = FPaths::ProfilingDir() / TEXT("Primitives") / FString::Printf(TEXT("Primitives-%s.csv"), *FDateTime::Now().ToString());
-	const bool bSuppressViewer = true;
-	FDiagnosticTableViewer DrawViewer(*OutputPath, bSuppressViewer);
-	DrawViewer.AddColumn(TEXT("Name"));
-	DrawViewer.AddColumn(TEXT("NumDraws"));
-	DrawViewer.CycleRow();
-
-	const FPrimitiveSceneInfo* LastPrimitiveSceneInfo = nullptr;
-	for (const FPrimitiveInfo& Primitive : Primitives)
-	{
-		if (Primitive.PrimitiveSceneInfo != LastPrimitiveSceneInfo)
-		{
-			DrawViewer.AddColumn(*Primitive.Name);
-			DrawViewer.AddColumn(*FString::Printf(TEXT("%d"), Primitive.DrawCount));
-			DrawViewer.CycleRow();
-
-			LastPrimitiveSceneInfo = Primitive.PrimitiveSceneInfo;
-		}
-	}
-}
-#endif
-
 FVisibilityTaskData* FSceneRenderer::BeginInitVisibility(FRDGBuilder& GraphBuilder, const FSceneTexturesConfig& SceneTexturesConfig)
 {
 	TRACE_CPUPROFILER_EVENT_SCOPE(FSceneRenderer_BeginInitVisibility);
@@ -4822,8 +4528,7 @@ void FSceneRenderer::ComputeViewVisibility(
 		FViewCommands& ViewCommands = TaskData->ViewCommandsPerView[ViewIndex];
 
 #if !UE_BUILD_SHIPPING
-		DumpPrimitives(ViewCommands);
-		ProcessPrimitives(View, ViewCommands);
+		FViewDebugInfo::Get().ProcessPrimitives(Scene, View, ViewCommands);
 #endif
 
 		SetupMeshPass(View, BasePassDepthStencilAccess, ViewCommands, InstanceCullingManager);
diff --git a/Engine/Source/Runtime/Renderer/Private/ViewDebug.cpp b/Engine/Source/Runtime/Renderer/Private/ViewDebug.cpp
new file mode 100644
--- /dev/null
+++ b/Engine/Source/Runtime/Renderer/Private/ViewDebug.cpp
@@ -0,0 +1,324 @@
+// Copyright Epic Games, Inc. All Rights Reserved.
+
+#include "ViewDebug.h"
+
+#if !UE_BUILD_SHIPPING
+
+#include "ScenePrivate.h"
+#include "Materials/Material.h"
+#include "MaterialShared.h"
+#include "ProfilingDebugging/DiagnosticTable.h"
+#include "MeshPassProcessor.h"
+#include "Engine/StaticMesh.h"
+#include "Engine/SkeletalMesh.h"
+#include "Materials/MaterialInterface.h"
+#include "Components/PrimitiveComponent.h"
+#include "Components/StaticMeshComponent.h"
+#include "Components/SkeletalMeshComponent.h"
+#include "Rendering/SkeletalMeshRenderData.h"
+#include "StaticMeshBatch.h"
+
+static bool bDumpPrimitivesNextFrame = false;
+static bool bDumpDetailedPrimitivesNextFrame = false;
+
+static FAutoConsoleCommand CVarDumpPrimitives(
+	TEXT("DumpPrimitives"),
+	TEXT("Writes out all scene primitive names to a CSV file"),
+	FConsoleCommandDelegate::CreateStatic([] { bDumpPrimitivesNextFrame = true; }),
+	ECVF_Default);
+
+static FAutoConsoleCommand CVarDrawPrimitiveDebugData(
+	TEXT("DumpDetailedPrimitives"),
+	TEXT("Writes out all scene primitive details to a CSV file"),
+	FConsoleCommandDelegate::CreateStatic([] { bDumpDetailedPrimitivesNextFrame = !bDumpDetailedPrimitivesNextFrame; }),
+	ECVF_Default);
+
+static uint32 GetDrawCountFromPrimitiveSceneInfo(FScene* Scene, const FPrimitiveSceneInfo* PrimitiveSceneInfo)
+{
+	uint32 DrawCount = 0;
+	for (const FCachedMeshDrawCommandInfo& CachedCommand : PrimitiveSceneInfo->StaticMeshCommandInfos)
+	{
+		if (CachedCommand.MeshPass != EMeshPass::BasePass)
+			continue;
+
+		if (CachedCommand.StateBucketId != INDEX_NONE || CachedCommand.CommandIndex >= 0)
+		{
+			DrawCount++;
+		}
+	}
+
+	return DrawCount;
+}
+
+FViewDebugInfo FViewDebugInfo::Instance;
+
+FViewDebugInfo::FViewDebugInfo()
+{
+	bHasEverUpdated = false;
+	bIsOutdated = true;
+	bShouldUpdate = false;
+	bShouldCaptureSingleFrame = false;
+}
+
+void FViewDebugInfo::ProcessPrimitive(FPrimitiveSceneInfo* PrimitiveSceneInfo, const FViewInfo& View, FScene* Scene, const UPrimitiveComponent* DebugComponent)
+{
+	if (!DebugComponent->IsRegistered())
+	{
+		return;
+	}
+	AActor* Actor = DebugComponent->GetOwner();
+	FString FullName = DebugComponent->GetName();
+	const uint32 DrawCount = GetDrawCountFromPrimitiveSceneInfo(Scene, PrimitiveSceneInfo);
+
+	TArray<UMaterialInterface*> Materials;
+	DebugComponent->GetUsedMaterials(Materials);
+	const int32 LOD = PrimitiveSceneInfo->Proxy ? PrimitiveSceneInfo->Proxy->GetLOD(&View) : INDEX_NONE;
+	int32 Triangles = 0;
+	if (const UStaticMeshComponent* StaticMeshComponent = Cast<UStaticMeshComponent>(DebugComponent))
+	{
+		Triangles = StaticMeshComponent->GetStaticMesh()->GetNumTriangles(LOD);
+	}
+	else if (const USkeletalMeshComponent* SkeletalMeshComponent = Cast<USkeletalMeshComponent>(DebugComponent))
+	{
+		for (const FSkeletalMeshLODRenderData& RenderData : SkeletalMeshComponent->GetSkeletalMeshRenderData()->LODRenderData)
+		{
+			Triangles += RenderData.MultiSizeIndexContainer.GetIndexBuffer()->Num() / 3;
+		}
+	}
+
+	const FPrimitiveInfo PrimitiveInfo = {
+		Actor,
+		DebugComponent->ComponentId,
+		const_cast<UPrimitiveComponent*>(DebugComponent), // This is probably a bad idea, find alternative
+		PrimitiveSceneInfo,
+		MoveTemp(Materials),
+		MoveTemp(FullName),
+		DrawCount,
+		Triangles,
+		LOD
+	};
+
+	Primitives.Add(PrimitiveInfo);
+}
+
+void FViewDebugInfo::DumpToCSV() const
+{
+	const FString OutputPath = FPaths::ProfilingDir() / TEXT("Primitives") / FString::Printf(TEXT("PrimitivesDetailed-%s.csv"), *FDateTime::Now().ToString());
+	const bool bSuppressViewer = true;
+	FDiagnosticTableViewer DrawViewer(*OutputPath, bSuppressViewer);
+	DrawViewer.AddColumn(TEXT("Name"));
+	DrawViewer.AddColumn(TEXT("ActorClass"));
+	DrawViewer.AddColumn(TEXT("Actor"));
+	DrawViewer.AddColumn(TEXT("Location"));
+	DrawViewer.AddColumn(TEXT("NumMaterials"));
+	DrawViewer.AddColumn(TEXT("Materials"));
+	DrawViewer.AddColumn(TEXT("NumDraws"));
+	DrawViewer.AddColumn(TEXT("LOD"));
+	DrawViewer.AddColumn(TEXT("Triangles"));
+	DrawViewer.CycleRow();
+
+	FRWScopeLock ScopeLock(Lock, SLT_ReadOnly);
+	const FPrimitiveSceneInfo* LastPrimitiveSceneInfo = nullptr;
+	for (const FPrimitiveInfo& Primitive : Primitives)
+	{
+		if (Primitive.PrimitiveSceneInfo != LastPrimitiveSceneInfo)
+		{
+			DrawViewer.AddColumn(*Primitive.Name);
+			DrawViewer.AddColumn(Primitive.Owner ? *Primitive.Owner->GetClass()->GetName() : TEXT(""));
+			DrawViewer.AddColumn(Primitive.Owner ? *Primitive.Owner->GetFullName() : TEXT(""));
+			DrawViewer.AddColumn(Primitive.Owner ?
+				*FString::Printf(TEXT("{%s}"), *Primitive.Owner->GetActorLocation().ToString()) : TEXT(""));
+			DrawViewer.AddColumn(*FString::Printf(TEXT("%d"), Primitive.Materials.Num()));
+			FString Materials = "[";
+			for (int i = 0; i < Primitive.Materials.Num(); i++)
+			{
+				if (Primitive.Materials[i] && Primitive.Materials[i]->GetMaterial())
+				{
+					Materials += Primitive.Materials[i]->GetMaterial()->GetName();
+				}
+				else
+				{
+					Materials += "Null";
+				}
+
+				if (i < Primitive.Materials.Num() - 1)
+				{
+					Materials += ", ";
+				}
+			}
+			Materials += "]";
+			DrawViewer.AddColumn(*FString::Printf(TEXT("%s"), *Materials));
+			DrawViewer.AddColumn(*FString::Printf(TEXT("%d"), Primitive.DrawCount));
+			DrawViewer.AddColumn(*FString::Printf(TEXT("%d"), Primitive.LOD));
+			DrawViewer.AddColumn(*FString::Printf(TEXT("%d"), Primitive.TriangleCount));
+			DrawViewer.CycleRow();
+
+			LastPrimitiveSceneInfo = Primitive.PrimitiveSceneInfo;
+		}
+	}
+}
+
+void FViewDebugInfo::CaptureNextFrame()
+{
+	FRWScopeLock ScopeLock(Lock, SLT_Write);
+	bShouldCaptureSingleFrame = true;
+	bShouldUpdate = true;
+}
+
+void FViewDebugInfo::EnableLiveCapture()
+{
+	FRWScopeLock ScopeLock(Lock, SLT_Write);
+	bShouldCaptureSingleFrame = false;
+	bShouldUpdate = true;
+}
+
+void FViewDebugInfo::DisableLiveCapture()
+{
+	FRWScopeLock ScopeLock(Lock, SLT_Write);
+	bShouldCaptureSingleFrame = false;
+	bShouldUpdate = false;
+}
+
+bool FViewDebugInfo::HasEverUpdated() const
+{
+	FRWScopeLock ScopeLock(Lock, SLT_ReadOnly);
+	return bHasEverUpdated;
+}
+
+bool FViewDebugInfo::IsOutOfDate() const
+{
+	FRWScopeLock ScopeLock(Lock, SLT_ReadOnly);
+	return bIsOutdated;
+}
+
+void FViewDebugInfo::ProcessPrimitives(FScene* Scene, const FViewInfo& View, const FViewCommands& ViewCommands)
+{
+	DumpPrimitives(Scene, ViewCommands);
+
+	{
+		FRWScopeLock ScopeLock(Lock, SLT_Write);
+		bIsOutdated = true;
+
+		if (!bShouldUpdate && !bDumpDetailedPrimitivesNextFrame)
+		{
+			return;
+		}
+
+		if (bShouldCaptureSingleFrame)
+		{
+			bShouldCaptureSingleFrame = false;
+			bShouldUpdate = false;
+		}
+
+		// TODO: Add profiling to this function
+
+		Primitives.Empty(ViewCommands.MeshCommands[EMeshPass::BasePass].Num() + ViewCommands.DynamicMeshCommandBuildRequests[EMeshPass::BasePass].Num());
+
+		for (const FVisibleMeshDrawCommand& Mesh : ViewCommands.MeshCommands[EMeshPass::BasePass])
+		{
+			const int32 PrimitiveId = Mesh.PrimitiveIdInfo.ScenePrimitiveId;
+			if (PrimitiveId >= 0 && PrimitiveId < Scene->Primitives.Num())
+			{
+				FPrimitiveSceneInfo* PrimitiveSceneInfo = Scene->Primitives[PrimitiveId];
+				ProcessPrimitive(PrimitiveSceneInfo, View, Scene, PrimitiveSceneInfo->ComponentForDebuggingOnly);
+			}
+		}
+
+		for (const FStaticMeshBatch* StaticMeshBatch : ViewCommands.DynamicMeshCommandBuildRequests[EMeshPass::BasePass])
+		{
+			FPrimitiveSceneInfo* PrimitiveSceneInfo = StaticMeshBatch->PrimitiveSceneInfo;
+			ProcessPrimitive(PrimitiveSceneInfo, View, Scene, PrimitiveSceneInfo->ComponentForDebuggingOnly);
+		}
+
+		bHasEverUpdated = true;
+		bIsOutdated = false;
+	}
+	OnUpdate.Broadcast();
+
+	if (bDumpDetailedPrimitivesNextFrame)
+	{
+		DumpToCSV();
+		bDumpDetailedPrimitivesNextFrame = false;
+	}
+}
+
+void FViewDebugInfo::DumpPrimitives(FScene* Scene, const FViewCommands& ViewCommands)
+{
+	if (!bDumpPrimitivesNextFrame)
+	{
+		return;
+	}
+
+	bDumpPrimitivesNextFrame = false;
+
+	struct FPrimitiveInfo
+	{
+		const FPrimitiveSceneInfo* PrimitiveSceneInfo;
+		FString Name;
+		uint32 DrawCount;
+
+		bool operator<(const FPrimitiveInfo& Other) const
+		{
+			// Sort by name to group similar assets together, then by exact primitives so we can ignore duplicates
+			const int32 NameCompare = Name.Compare(Other.Name);
+			if (NameCompare != 0)
+			{
+				return NameCompare < 0;
+			}
+
+			return PrimitiveSceneInfo < Other.PrimitiveSceneInfo;
+		}
+	};
+
+	TArray<FPrimitiveInfo> Primitives;
+	Primitives.Reserve(ViewCommands.MeshCommands[EMeshPass::BasePass].Num() + ViewCommands.DynamicMeshCommandBuildRequests[EMeshPass::BasePass].Num());
+
+	{
+		for (const FVisibleMeshDrawCommand& Mesh : ViewCommands.MeshCommands[EMeshPass::BasePass])
+		{
+			int32 PrimitiveId = Mesh.PrimitiveIdInfo.ScenePrimitiveId;
+			if (PrimitiveId >= 0 && PrimitiveId < Scene->Primitives.Num())
+			{
+				const FPrimitiveSceneInfo* PrimitiveSceneInfo = Scene->Primitives[PrimitiveId];
+				FString FullName = PrimitiveSceneInfo->ComponentForDebuggingOnly->GetFullName();
+
+				uint32 DrawCount = GetDrawCountFromPrimitiveSceneInfo(Scene, PrimitiveSceneInfo);
+
+				Primitives.Add({ PrimitiveSceneInfo, MoveTemp(FullName), DrawCount });
+			}
+		}
+
+		for (const FStaticMeshBatch* StaticMeshBatch : ViewCommands.DynamicMeshCommandBuildRequests[EMeshPass::BasePass])
+		{
+			const FPrimitiveSceneInfo* PrimitiveSceneInfo = StaticMeshBatch->PrimitiveSceneInfo;
+			FString FullName = PrimitiveSceneInfo->ComponentForDebuggingOnly->GetFullName();
+
+			uint32 DrawCount = GetDrawCountFromPrimitiveSceneInfo(Scene, PrimitiveSceneInfo);
+
+			Primitives.Add({ PrimitiveSceneInfo, MoveTemp(FullName), DrawCount });
+		}
+	}
+
+	Primitives.Sort();
+
+	const FString OutputPath = FPaths::ProfilingDir() / TEXT("Primitives") / FString::Printf(TEXT("Primitives-%s.csv"), *FDateTime::Now().ToString());
+	const bool bSuppressViewer = true;
+	FDiagnosticTableViewer DrawViewer(*OutputPath, bSuppressViewer);
+	DrawViewer.AddColumn(TEXT("Name"));
+	DrawViewer.AddColumn(TEXT("NumDraws"));
+	DrawViewer.CycleRow();
+
+	const FPrimitiveSceneInfo* LastPrimitiveSceneInfo = nullptr;
+	for (const FPrimitiveInfo& Primitive : Primitives)
+	{
+		if (Primitive.PrimitiveSceneInfo != LastPrimitiveSceneInfo)
+		{
+			DrawViewer.AddColumn(*Primitive.Name);
+			DrawViewer.AddColumn(*FString::Printf(TEXT("%d"), Primitive.DrawCount));
+			DrawViewer.CycleRow();
+
+			LastPrimitiveSceneInfo = Primitive.PrimitiveSceneInfo;
+		}
+	}
+}
+#endif
\ No newline at end of file
diff --git a/Engine/Source/Runtime/Renderer/Public/PrimitiveSceneInfo.h b/Engine/Source/Runtime/Renderer/Public/PrimitiveSceneInfo.h
--- a/Engine/Source/Runtime/Renderer/Public/PrimitiveSceneInfo.h
+++ b/Engine/Source/Runtime/Renderer/Public/PrimitiveSceneInfo.h
@@ -281,6 +281,7 @@ inline bool operator == (FPersistentPrimitiveIndex A, FPersistentPrimitiveIndex
 class FPrimitiveSceneInfo : public FDeferredCleanupInterface
 {
 	friend class FSceneRenderer;
+	friend struct FViewDebugInfo;
 public:
 
 	/** The render proxy for the primitive. */
diff --git a/Engine/Source/Runtime/Renderer/Public/ViewDebug.h b/Engine/Source/Runtime/Renderer/Public/ViewDebug.h
--- a/Engine/Source/Runtime/Renderer/Public/ViewDebug.h
+++ b/Engine/Source/Runtime/Renderer/Public/ViewDebug.h
@@ -9,6 +9,9 @@
 #if !UE_BUILD_SHIPPING // TODO: Decide whether or not the struct should be entirely stripped out of shipping
 
 class UMaterialInterface;
+class FScene;
+class FViewInfo;
+class FViewCommands;
 
 DECLARE_MULTICAST_DELEGATE(FOnUpdateViewDebugInfo);
 
@@ -17,7 +20,6 @@ DECLARE_MULTICAST_DELEGATE(FOnUpdateViewDebugInfo);
  */
 struct RENDERER_API FViewDebugInfo
 {
-	friend class FSceneRenderer;
 	friend class FDrawPrimitiveDebuggerModule;
 
 private:
@@ -32,7 +34,7 @@ public:
 	 * Gets a reference to the view debug information that is used by the renderer.
 	 * @returns The debug information that is used by the renderer.
 	 */
-	static inline const FViewDebugInfo& Get()
+	static inline FViewDebugInfo& Get()
 	{
 		return Instance;
 	}
@@ -86,7 +88,10 @@ private:
 
 	void DisableLiveCapture();
 
+	static void DumpPrimitives(FScene* Scene, const FViewCommands& ViewCommands);
+
 public:
+	void ProcessPrimitives(FScene* Scene, const FViewInfo& View, const FViewCommands& ViewCommands);
 
 	/**
 	 * Writes the currently stored information out to a CSV file.
