Add support for overriding max anisotropy via a texture LOD group

diff --git a/Engine/Source/Runtime/Engine/Classes/Engine/TextureLODSettings.h b/Engine/Source/Runtime/Engine/Classes/Engine/TextureLODSettings.h
--- a/Engine/Source/Runtime/Engine/Classes/Engine/TextureLODSettings.h
+++ b/Engine/Source/Runtime/Engine/Classes/Engine/TextureLODSettings.h
@@ -42,6 +42,7 @@ struct ENGINE_API FTextureLODGroup
 		, VirtualTextureTileSizeBias(0)
 		, LossyCompressionAmount(TLCA_Default)
 		, CookPlatformTilingDisabled(false)
+		, MaxAniso(0)
 	{
 		SetupGroup();
 	}
@@ -137,6 +138,10 @@ struct ENGINE_API FTextureLODGroup
 	UPROPERTY()
 	bool CookPlatformTilingDisabled;
 
+	/** Allows us to override max anisotropy. If unspecified, uses r.MaxAnisotropy */
+	UPROPERTY()
+	int32 MaxAniso;
+
 	void SetupGroup();
 
 	bool operator==(const FTextureLODGroup& Other) const;
diff --git a/Engine/Source/Runtime/Engine/Private/Rendering/StreamableTextureResource.cpp b/Engine/Source/Runtime/Engine/Private/Rendering/StreamableTextureResource.cpp
--- a/Engine/Source/Runtime/Engine/Private/Rendering/StreamableTextureResource.cpp
+++ b/Engine/Source/Runtime/Engine/Private/Rendering/StreamableTextureResource.cpp
@@ -98,10 +98,13 @@ FStreamableTextureResource::FStreamableTextureResource(UTexture* InOwner, const
 	bSRGB = InOwner->SRGB;
 	bGreyScaleFormat = UE::TextureDefines::ShouldUseGreyScaleEditorVisualization( InOwner->CompressionSettings );
 
-	Filter = (ESamplerFilter)UDeviceProfileManager::Get().GetActiveProfile()->GetTextureLODSettings()->GetSamplerFilter(InOwner);
+	const UTextureLODSettings* TextureLODSettings = UDeviceProfileManager::Get().GetActiveProfile()->GetTextureLODSettings();
+
+	Filter = (ESamplerFilter)TextureLODSettings->GetSamplerFilter(InOwner);
 	AddressU = InOwner->GetTextureAddressX() == TA_Wrap ? AM_Wrap : (InOwner->GetTextureAddressX() == TA_Clamp ? AM_Clamp : AM_Mirror);
 	AddressV = InOwner->GetTextureAddressY() == TA_Wrap ? AM_Wrap : (InOwner->GetTextureAddressY() == TA_Clamp ? AM_Clamp : AM_Mirror);
 	AddressW = InOwner->GetTextureAddressZ() == TA_Wrap ? AM_Wrap : (InOwner->GetTextureAddressZ() == TA_Clamp ? AM_Clamp : AM_Mirror);
+	MaxAniso = TextureLODSettings->GetTextureLODGroup(InOwner->LODGroup).MaxAniso;
 
 	// Get the biggest mips size, might be different from the actual resolution (depending on NumOfResidentLODs).
 	const FTexture2DMipMap& Mip0 = PlatformData->Mips[State.AssetLODBias];
@@ -248,7 +251,8 @@ void FStreamableTextureResource::RefreshSamplerStates()
 		AddressU,
 		AddressV,
 		AddressW,
-		MipBias
+		MipBias,
+		MaxAniso
 	);
 	SamplerStateRHI = GetOrCreateSamplerState(SamplerStateInitializer);
 
diff --git a/Engine/Source/Runtime/Engine/Private/SceneManagement.cpp b/Engine/Source/Runtime/Engine/Private/SceneManagement.cpp
--- a/Engine/Source/Runtime/Engine/Private/SceneManagement.cpp
+++ b/Engine/Source/Runtime/Engine/Private/SceneManagement.cpp
@@ -1027,13 +1027,15 @@ void FSharedSamplerState::InitRHI()
 {
 	const float MipMapBias = UTexture2D::GetGlobalMipMapLODBias();
 
+	const UTextureLODSettings* TextureLODSettings = UDeviceProfileManager::Get().GetActiveProfile()->GetTextureLODSettings();
 	FSamplerStateInitializerRHI SamplerStateInitializer
 	(
-	(ESamplerFilter)UDeviceProfileManager::Get().GetActiveProfile()->GetTextureLODSettings()->GetSamplerFilter(TEXTUREGROUP_World),
+	(ESamplerFilter)TextureLODSettings->GetSamplerFilter(TEXTUREGROUP_World),
 		bWrap ? AM_Wrap : AM_Clamp,
 		bWrap ? AM_Wrap : AM_Clamp,
 		bWrap ? AM_Wrap : AM_Clamp,
-		MipMapBias
+		MipMapBias,
+		TextureLODSettings->GetTextureLODGroup(TEXTUREGROUP_World).MaxAniso
 	);
 	SamplerStateRHI = RHICreateSamplerState(SamplerStateInitializer);
 }
diff --git a/Engine/Source/Runtime/Engine/Private/TextureRenderTarget2D.cpp b/Engine/Source/Runtime/Engine/Private/TextureRenderTarget2D.cpp
--- a/Engine/Source/Runtime/Engine/Private/TextureRenderTarget2D.cpp
+++ b/Engine/Source/Runtime/Engine/Private/TextureRenderTarget2D.cpp
@@ -667,12 +667,15 @@ void FTextureRenderTarget2DResource::InitDynamicRHI()
 	}
 
 	// Create the sampler state RHI resource.
+	const UTextureLODSettings* TextureLODSettings = UDeviceProfileManager::Get().GetActiveProfile()->GetTextureLODSettings();
 	FSamplerStateInitializerRHI SamplerStateInitializer
 	(
-		(ESamplerFilter)UDeviceProfileManager::Get().GetActiveProfile()->GetTextureLODSettings()->GetSamplerFilter( Owner ),
+		(ESamplerFilter)TextureLODSettings->GetSamplerFilter( Owner ),
 		Owner->AddressX == TA_Wrap ? AM_Wrap : (Owner->AddressX == TA_Clamp ? AM_Clamp : AM_Mirror),
 		Owner->AddressY == TA_Wrap ? AM_Wrap : (Owner->AddressY == TA_Clamp ? AM_Clamp : AM_Mirror),
-		AM_Wrap
+		AM_Wrap,
+		0,
+		TextureLODSettings->GetTextureLODGroup(Owner->LODGroup).MaxAniso
 	);
 	SamplerStateRHI = GetOrCreateSamplerState( SamplerStateInitializer );
 }
diff --git a/Engine/Source/Runtime/Engine/Private/TextureRenderTarget2DArray.cpp b/Engine/Source/Runtime/Engine/Private/TextureRenderTarget2DArray.cpp
--- a/Engine/Source/Runtime/Engine/Private/TextureRenderTarget2DArray.cpp
+++ b/Engine/Source/Runtime/Engine/Private/TextureRenderTarget2DArray.cpp
@@ -268,12 +268,15 @@ void FTextureRenderTarget2DArrayResource::InitDynamicRHI()
 	}
 
 	// Create the sampler state RHI resource.
+	const UTextureLODSettings* TextureLODSettings = UDeviceProfileManager::Get().GetActiveProfile()->GetTextureLODSettings();
 	FSamplerStateInitializerRHI SamplerStateInitializer
 	(
-		(ESamplerFilter)UDeviceProfileManager::Get().GetActiveProfile()->GetTextureLODSettings()->GetSamplerFilter(Owner),
+		(ESamplerFilter)TextureLODSettings->GetSamplerFilter(Owner),
 		AM_Wrap,
 		AM_Wrap,
-		AM_Wrap
+		AM_Wrap,
+		0,
+		TextureLODSettings->GetTextureLODGroup(Owner->LODGroup).MaxAniso
 	);
 	SamplerStateRHI = GetOrCreateSamplerState(SamplerStateInitializer);
 }
diff --git a/Engine/Source/Runtime/Engine/Public/Rendering/StreamableTextureResource.h b/Engine/Source/Runtime/Engine/Public/Rendering/StreamableTextureResource.h
--- a/Engine/Source/Runtime/Engine/Public/Rendering/StreamableTextureResource.h
+++ b/Engine/Source/Runtime/Engine/Public/Rendering/StreamableTextureResource.h
@@ -115,6 +115,9 @@ protected:
 	/** Whether this texture should be updated using the virtual address mapping for each mip. */
 	bool bUsePartiallyResidentMips = false;
 
+	/** Max anisotropy. if 0, will use r.MaxAnisotropy */
+	int8 MaxAniso = 0;
+
 #if STATS
 private:
 	void CalcRequestedMipsSize();
diff --git a/Engine/Source/Runtime/Renderer/Private/SceneRendering.cpp b/Engine/Source/Runtime/Renderer/Private/SceneRendering.cpp
--- a/Engine/Source/Runtime/Renderer/Private/SceneRendering.cpp
+++ b/Engine/Source/Runtime/Renderer/Private/SceneRendering.cpp
@@ -940,8 +940,10 @@ void FViewInfo::Init()
 	// Cache TEXTUREGROUP filter settings for the render thread to create shared samplers.
 	if (IsInGameThread())
 	{
-		WorldTextureGroupSamplerFilter = (ESamplerFilter)UDeviceProfileManager::Get().GetActiveProfile()->GetTextureLODSettings()->GetSamplerFilter(TEXTUREGROUP_World);
-		TerrainWeightmapTextureGroupSamplerFilter = (ESamplerFilter)UDeviceProfileManager::Get().GetActiveProfile()->GetTextureLODSettings()->GetSamplerFilter(TEXTUREGROUP_Terrain_Weightmap);
+		const UTextureLODSettings* TextureLODSettings = UDeviceProfileManager::Get().GetActiveProfile()->GetTextureLODSettings();
+		WorldTextureGroupSamplerFilter = (ESamplerFilter)TextureLODSettings->GetSamplerFilter(TEXTUREGROUP_World);
+		TerrainWeightmapTextureGroupSamplerFilter = (ESamplerFilter)TextureLODSettings->GetSamplerFilter(TEXTUREGROUP_Terrain_Weightmap);
+		WorldTextureGroupMaxAnisotropy = TextureLODSettings->GetTextureLODGroup(TEXTUREGROUP_World).MaxAniso;
 		bIsValidTextureGroupSamplerFilters = true;
 	}
 	else
@@ -1524,8 +1526,8 @@ void FViewInfo::SetupUniformBufferParameters(
 		{
 			check(bIsValidTextureGroupSamplerFilters);
 
-			WrappedSampler = RHICreateSamplerState(FSamplerStateInitializerRHI(WorldTextureGroupSamplerFilter, AM_Wrap,  AM_Wrap,  AM_Wrap,  FinalMaterialTextureMipBias));
-			ClampedSampler = RHICreateSamplerState(FSamplerStateInitializerRHI(WorldTextureGroupSamplerFilter, AM_Clamp, AM_Clamp, AM_Clamp, FinalMaterialTextureMipBias));
+			WrappedSampler = RHICreateSamplerState(FSamplerStateInitializerRHI(WorldTextureGroupSamplerFilter, AM_Wrap,  AM_Wrap,  AM_Wrap,  FinalMaterialTextureMipBias, WorldTextureGroupMaxAnisotropy));
+			ClampedSampler = RHICreateSamplerState(FSamplerStateInitializerRHI(WorldTextureGroupSamplerFilter, AM_Clamp, AM_Clamp, AM_Clamp, FinalMaterialTextureMipBias, WorldTextureGroupMaxAnisotropy));
 		}
 
 		// At this point, a sampler must be set.
diff --git a/Engine/Source/Runtime/Renderer/Private/SceneRendering.h b/Engine/Source/Runtime/Renderer/Private/SceneRendering.h
--- a/Engine/Source/Runtime/Renderer/Private/SceneRendering.h
+++ b/Engine/Source/Runtime/Renderer/Private/SceneRendering.h
@@ -1890,6 +1890,7 @@ private:
 	// may not have a valid value if FViewInfo is created on the render thread.
 	ESamplerFilter WorldTextureGroupSamplerFilter;
 	ESamplerFilter TerrainWeightmapTextureGroupSamplerFilter;
+	int32 WorldTextureGroupMaxAnisotropy;
 	bool bIsValidTextureGroupSamplerFilters;
 
 	FSceneViewState* GetEyeAdaptationViewState() const;
