Fix for build break due to multiple symbol definitions.

diff --git a/Engine/Source/Runtime/Renderer/Private/SceneVisibility.cpp b/Engine/Source/Runtime/Renderer/Private/SceneVisibility.cpp
--- a/Engine/Source/Runtime/Renderer/Private/SceneVisibility.cpp
+++ b/Engine/Source/Runtime/Renderer/Private/SceneVisibility.cpp
@@ -324,21 +324,6 @@ static TAutoConsoleVariable<float> CVarFreezeTemporalHistories(
 	TEXT("Freezes all temporal histories as well as the temporal sequence."),
 	ECVF_RenderThreadSafe);
 
-static bool bDumpPrimitivesNextFrame = false;
-static bool bDumpDetailedPrimitivesNextFrame = false;
-
-static FAutoConsoleCommand CVarDumpPrimitives(
-	TEXT("DumpPrimitives"),
-	TEXT("Writes out all scene primitive names to a CSV file"),
-	FConsoleCommandDelegate::CreateStatic([] { bDumpPrimitivesNextFrame = true; }),
-	ECVF_Default);
-
-static FAutoConsoleCommand CVarDrawPrimitiveDebugData(
-	TEXT("DumpDetailedPrimitives"),
-	TEXT("Writes out all scene primitive details to a CSV file"),
-	FConsoleCommandDelegate::CreateStatic([] { bDumpDetailedPrimitivesNextFrame = !bDumpDetailedPrimitivesNextFrame; }),
-	ECVF_Default);
-
 #endif
 
 DECLARE_CYCLE_STAT(TEXT("Occlusion Readback"), STAT_CLMM_OcclusionReadback, STATGROUP_CommandListMarkers);
