Added overlay material boolean to FStaticMeshBatchRelevance.

diff --git a/Engine/Source/Runtime/Renderer/Private/StaticMeshBatch.cpp b/Engine/Source/Runtime/Renderer/Private/StaticMeshBatch.cpp
--- a/Engine/Source/Runtime/Renderer/Private/StaticMeshBatch.cpp
+++ b/Engine/Source/Runtime/Renderer/Private/StaticMeshBatch.cpp
@@ -43,6 +43,7 @@ FStaticMeshBatchRelevance::FStaticMeshBatchRelevance(const FStaticMeshBatch& Sta
 	, bUseSingleLayerWaterMaterial(bInUseSingleLayerWaterMaterial)
 	, bUseHairStrands(StaticMesh.UseForHairStrands(FeatureLevel))
 	, bUseAnisotropy(bInUseAnisotropy)
+	, bOverlayMaterial(StaticMesh.bOverlayMaterial)
 	, bRenderToVirtualTexture(StaticMesh.bRenderToVirtualTexture)
 	, RuntimeVirtualTextureMaterialType(StaticMesh.RuntimeVirtualTextureMaterialType)
 	, bSupportsCachingMeshDrawCommands(InbSupportsCachingMeshDrawCommands)
diff --git a/Engine/Source/Runtime/Renderer/Public/StaticMeshBatch.h b/Engine/Source/Runtime/Renderer/Public/StaticMeshBatch.h
--- a/Engine/Source/Runtime/Renderer/Public/StaticMeshBatch.h
+++ b/Engine/Source/Runtime/Renderer/Public/StaticMeshBatch.h
@@ -71,6 +71,7 @@ public:
 	uint8 bUseSingleLayerWaterMaterial : 1; // Whether this batch uses a water material or not.
 	uint8 bUseHairStrands : 1; // Whether it contains hair strands geometry.
 	uint8 bUseAnisotropy : 1; // Whether material uses anisotropy parameter.
+	uint8 bOverlayMaterial : 1; // Whether mesh uses overlay material.
 
 	/** Whether the mesh batch can be used for rendering to a virtual texture. */
 	uint8 bRenderToVirtualTexture : 1;
