Fixes some missing RTs in FTSRHistory::GetGPUSizeBytes

diff --git a/Engine/Source/Runtime/Renderer/Private/RendererScene.cpp b/Engine/Source/Runtime/Renderer/Private/RendererScene.cpp
--- a/Engine/Source/Runtime/Renderer/Private/RendererScene.cpp
+++ b/Engine/Source/Runtime/Renderer/Private/RendererScene.cpp
@@ -971,11 +971,17 @@ uint64 FTemporalAAHistory::GetGPUSizeBytes(bool bLogSizes) const
 uint64 FTSRHistory::GetGPUSizeBytes(bool bLogSizes) const
 {
 	return
+		GetRenderTargetGPUSizeBytes(Output, bLogSizes) +
 		GetRenderTargetGPUSizeBytes(ColorArray, bLogSizes) +
 		GetRenderTargetGPUSizeBytes(Metadata, bLogSizes) +
 		GetRenderTargetGPUSizeBytes(TranslucencyAlpha, bLogSizes) +
 		GetRenderTargetGPUSizeBytes(SubpixelDetails, bLogSizes) +
-		GetRenderTargetGPUSizeBytes(Guide, bLogSizes);
+		GetRenderTargetGPUSizeBytes(SubpixelDepth, bLogSizes) +
+		GetRenderTargetGPUSizeBytes(Guide, bLogSizes) +
+		GetRenderTargetGPUSizeBytes(Moire, bLogSizes) +
+		GetRenderTargetGPUSizeBytes(Velocity, bLogSizes) +
+		GetRenderTargetGPUSizeBytes(PrevOutput, bLogSizes) +
+		GetRenderTargetGPUSizeBytes(PrevColorArray, bLogSizes);
 }
 
 uint64 FScreenSpaceDenoiserHistory::GetGPUSizeBytes(bool bLogSizes) const
