Deletes FTSRCompareTranslucencyCS

diff --git a/Engine/Shaders/Private/TemporalSuperResolution/TSRCompareTranslucency.usf b/Engine/Shaders/Private/TemporalSuperResolution/TSRCompareTranslucency.usf
deleted file mode 100644
diff --git a/Engine/Shaders/Private/TemporalSuperResolution/TSRUpdateHistory.usf b/Engine/Shaders/Private/TemporalSuperResolution/TSRUpdateHistory.usf
--- a/Engine/Shaders/Private/TemporalSuperResolution/TSRUpdateHistory.usf
+++ b/Engine/Shaders/Private/TemporalSuperResolution/TSRUpdateHistory.usf
@@ -124,7 +124,6 @@ Texture2D<uint2> InputSceneStencilTexture;
 Texture2D<tsr_half4> InputSceneTranslucencyTexture;
 
 Texture2D<tsr_half2> HistoryRejectionTexture;
-Texture2D<tsr_half> TranslucencyRejectionTexture;
 
 Texture2D<float2> DilatedVelocityTexture;
 Texture2D<uint> ParallaxFactorTexture;
@@ -136,20 +135,6 @@ Texture2D<tsr_half2> HoleFilledVelocityMaskTexture;
 Texture2D<float2> GrandDilatedVelocityTexture;
 Texture2D<tsr_halfC> GrandPrevColorTexture;
 
-float2 TranslucencyInfo_Extent;
-float2 TranslucencyInfo_ExtentInverse;
-float2 TranslucencyInfo_ScreenPosToViewportScale;
-float2 TranslucencyInfo_ScreenPosToViewportBias;
-uint2  TranslucencyInfo_ViewportMin;
-uint2  TranslucencyInfo_ViewportMax;
-float2 TranslucencyInfo_ViewportSize;
-float2 TranslucencyInfo_ViewportSizeInverse;
-float2 TranslucencyInfo_UVViewportMin;
-float2 TranslucencyInfo_UVViewportMax;
-float2 TranslucencyInfo_UVViewportSize;
-float2 TranslucencyInfo_UVViewportSizeInverse;
-float2 TranslucencyInfo_UVViewportBilinearMin;
-float2 TranslucencyInfo_UVViewportBilinearMax;
 int2   TranslucencyPixelPosMin;
 int2   TranslucencyPixelPosMax;
 
@@ -158,7 +143,6 @@ FScreenTransform HistoryPixelPosToInputPPCo;
 FScreenTransform HistoryPixelPosToTranslucencyPPCo;
 FScreenTransform ScreenPosToGrandPrevHistoryBufferUV;
 float3 HistoryQuantizationError;
-float MinTranslucencyRejection;
 float HistorySampleCount;
 float HistoryHisteresis;
 float WeightClampingRejection;
@@ -531,7 +515,6 @@ void MainCS(
 	#endif
 	tsr_half2 LowFrequencyRejection = tsr_half(1.0).xx;
 	tsr_half2 LowFrequencyClamp = tsr_half(1.0).xx;
-	tsr_half2 TranslucencyRejection = tsr_half(1.0).xx;
 	tsr_half2 OutputPixelVelocity = tsr_half(0.0).xx;
 	tsr_half2 NoiseFiltering = tsr_half(0.0).xx;
 	tsr_half2x2 HoleFilledVelocityMask = dpv_interleave_mono_registers(tsr_half(1.0).xx);
@@ -755,8 +738,6 @@ void MainCS(
 			UNROLL_N(CONFIG_SAMPLES - 1)
 			for (uint NeighborId = 0; NeighborId < (CONFIG_SAMPLES - 1); NeighborId++)
 			{
-				//TranslucencyRejection = min(TranslucencyRejection, dpv_force_interleave_registers_array(RawTranslucencyRejectionNeighborhood[NeighborId]));
-			
 				SceneStencilRef |= dpv_interleave_registers_array(RawSceneStencilRefNeighborhood[NeighborId]);
 			}
 		}
@@ -774,16 +755,10 @@ void MainCS(
 		
 		// Final post processing.
 		{
-			TranslucencyRejection = max(TranslucencyRejection, tsr_half(MinTranslucencyRejection));
-		
 			// Fetch whether the pixel is responsive AA or not.
 			#if CONFIG_RESPONSIVE_STENCIL
 			{
 				bIsResponsiveAAPixel = (SceneStencilRef & ResponsiveStencilMask) != 0;
-
-				#if CONFIG_REJECT_TRANSLUCENCY_ON_RESPONSIVE_AA_ONLY
-					TranslucencyRejection = max(TranslucencyRejection, select(bIsResponsiveAAPixel, tsr_half(0.0).xx, tsr_half(1.0).xx));
-				#endif
 			}
 			#endif
 
@@ -792,16 +767,10 @@ void MainCS(
 				bIsDisoccluded = false;
 				LowFrequencyRejection = tsr_half(1.0).xx;
 				LowFrequencyClamp = tsr_half(1.0).xx;
-				TranslucencyRejection = tsr_half(1.0).xx;
 			}
 			#endif
 		}
 	}
-
-	// Reject the translucency to respond quicker on fireflies.
-	#if CONFIG_TRANSLUCENCY_REJECTION && CONFIG_ACCUMULATE_TRANSLUCENCY_SEPARATELY == 0
-		LowFrequencyRejection = min(LowFrequencyRejection, TranslucencyRejection);
-	#endif
 	
 	// Reproject history
 	tsr_halfCx2 PrevHighFrequencyColor;
@@ -1537,18 +1506,6 @@ void MainCS(
 
 		tsr_half4x2 ClampedPrevHistoryTranslucency = fastClamp(PrevHistoryTranslucency, TranslucencyMinColor, TranslucencyMaxColor);
 		
-		#if CONFIG_SCENE_COLOR_ALPHA
-			tsr_half2 HighFrequencyTranslucencyRejection = MeasureRejectionFactor(
-				PrevHistoryTranslucency, ClampedPrevHistoryTranslucency,
-				FilteredTranslucencyColor, TranslucencyMinColor, TranslucencyMaxColor,
-				tsr_half(0.0));
-		#else
-			tsr_half2 HighFrequencyTranslucencyRejection = MeasureRejectionFactor(
-				GetColorChannel(PrevHistoryTranslucency), GetColorChannel(ClampedPrevHistoryTranslucency),
-				GetColorChannel(FilteredTranslucencyColor), GetColorChannel(TranslucencyMinColor), GetColorChannel(TranslucencyMaxColor),
-				tsr_half(0.0));
-		#endif
-
 		tsr_half2 BlendFinal = tsr_half(0.2) * InputPixelAlignement;
 
 		// Antiflicker
@@ -1572,10 +1529,6 @@ void MainCS(
 			BlendFinal[1] = bIsResponsiveAAPixel[1] ? tsr_half(1.0/4.0) : BlendFinal[1];
 		#endif
 		
-		#if CONFIG_LOAD_TRANSLUCENCY_REJECTION
-			BlendFinal = max(BlendFinal, (tsr_half(1.0) - TranslucencyRejection));
-		#endif
-		
 		FLATTEN
 		if (bCameraCut)
 		{
diff --git a/Engine/Source/Runtime/Renderer/Private/PostProcess/TemporalSuperResolution.cpp b/Engine/Source/Runtime/Renderer/Private/PostProcess/TemporalSuperResolution.cpp
--- a/Engine/Source/Runtime/Renderer/Private/PostProcess/TemporalSuperResolution.cpp
+++ b/Engine/Source/Runtime/Renderer/Private/PostProcess/TemporalSuperResolution.cpp
@@ -194,18 +194,6 @@ TAutoConsoleVariable<int32> CVarTSRAsyncCompute(
 	TEXT(" 3: Run all passes on async compute;"),
 	ECVF_RenderThreadSafe);
 
-// TODO: remove this dead code
-//TAutoConsoleVariable<float> CVarTSRTranslucencyHighlightLuminance(
-//	TEXT("r.TSR.Translucency.HighlightLuminance"), -1.0f,
-//	TEXT("Sets the luminance at which translucency is considered an highlights (default=-1.0)."),
-//	ECVF_RenderThreadSafe);
-
-// TODO: remove this dead code
-//TAutoConsoleVariable<int32> CVarTSRTranslucencyPreviousFrameRejection(
-//	TEXT("r.TSR.Translucency.PreviousFrameRejection"), 0,
-//	TEXT("Enable heuristic to reject Separate translucency based on previous frame translucency."),
-//	ECVF_RenderThreadSafe);
-
 // TODO: currently dead code, that is technically relevent for 5.0's r.TSR.History.SeparateTranslucency=1 but that has too many problems on translucency contour around opaque geometry.
 //TAutoConsoleVariable<int32> CVarTSREnableResponiveAA(
 //	TEXT("r.TSR.Translucency.EnableResponiveAA"), 1,
@@ -679,30 +667,6 @@ class FTSRDecimateHistoryCS : public FTSRShader
 	END_SHADER_PARAMETER_STRUCT()
 }; // class FTSRDecimateHistoryCS
 
-class FTSRCompareTranslucencyCS : public FTSRShader
-{
-	DECLARE_GLOBAL_SHADER(FTSRCompareTranslucencyCS);
-	SHADER_USE_PARAMETER_STRUCT(FTSRCompareTranslucencyCS, FTSRShader);
-
-	BEGIN_SHADER_PARAMETER_STRUCT(FParameters, )
-		SHADER_PARAMETER_STRUCT_INCLUDE(FTSRCommonParameters, CommonParameters)
-
-		SHADER_PARAMETER_STRUCT(FScreenPassTextureViewportParameters, TranslucencyInfo)
-		SHADER_PARAMETER_STRUCT(FScreenPassTextureViewportParameters, PrevTranslucencyInfo)
-		SHADER_PARAMETER(float, PrevTranslucencyPreExposureCorrection)
-		SHADER_PARAMETER(float, TranslucencyHighlightLuminance)
-
-		SHADER_PARAMETER(FScreenTransform, ScreenPosToPrevTranslucencyTextureUV)
-
-		SHADER_PARAMETER_RDG_TEXTURE(Texture2D, DilatedVelocityTexture)
-		SHADER_PARAMETER_RDG_TEXTURE(Texture2D, TranslucencyTexture)
-		SHADER_PARAMETER_RDG_TEXTURE(Texture2D, PrevTranslucencyTexture)
-
-		SHADER_PARAMETER_RDG_TEXTURE_UAV(RWTexture2D, TranslucencyRejectionOutput)
-		SHADER_PARAMETER_RDG_TEXTURE_UAV(RWTexture2DArray, DebugOutput)
-	END_SHADER_PARAMETER_STRUCT()
-}; // class FTSRCompareTranslucencyCS
-
 class FTSRRejectShadingCS : public FTSRShader
 {
 	DECLARE_GLOBAL_SHADER(FTSRRejectShadingCS);
@@ -889,7 +853,6 @@ class FTSRUpdateHistoryCS : public FTSRShader
 		SHADER_PARAMETER_RDG_TEXTURE(Texture2D, InputSceneTranslucencyTexture)
 
 		SHADER_PARAMETER_RDG_TEXTURE(Texture2D, HistoryRejectionTexture)
-		SHADER_PARAMETER_RDG_TEXTURE(Texture2D, TranslucencyRejectionTexture)
 
 		SHADER_PARAMETER_RDG_TEXTURE(Texture2D, DilatedVelocityTexture)
 		SHADER_PARAMETER_RDG_TEXTURE(Texture2D, GrandDilatedVelocityTexture)
@@ -899,7 +862,6 @@ class FTSRUpdateHistoryCS : public FTSRShader
 		SHADER_PARAMETER_RDG_TEXTURE(Texture2D, NoiseFilteringTexture)
 		SHADER_PARAMETER_RDG_TEXTURE(Texture2D, HoleFilledVelocityMaskTexture)
 
-		SHADER_PARAMETER_STRUCT(FScreenPassTextureViewportParameters, TranslucencyInfo)
 		SHADER_PARAMETER(FIntPoint, TranslucencyPixelPosMin)
 		SHADER_PARAMETER(FIntPoint, TranslucencyPixelPosMax)
 
@@ -908,7 +870,6 @@ class FTSRUpdateHistoryCS : public FTSRShader
 		SHADER_PARAMETER(FScreenTransform, HistoryPixelPosToTranslucencyPPCo)
 
 		SHADER_PARAMETER(FVector3f, HistoryQuantizationError)
-		SHADER_PARAMETER(float, MinTranslucencyRejection)
 		SHADER_PARAMETER(float, HistorySampleCount)
 		SHADER_PARAMETER(float, HistoryHisteresis)
 		SHADER_PARAMETER(float, WeightClampingRejection)
@@ -1058,7 +1019,6 @@ IMPLEMENT_GLOBAL_SHADER(FTSRClearPrevTexturesCS,     "/Engine/Private/TemporalSu
 IMPLEMENT_GLOBAL_SHADER(FTSRForwardScatterDepthCS,   "/Engine/Private/TemporalSuperResolution/TSRForwardScatterDepth.usf",   "MainCS", SF_Compute);
 IMPLEMENT_GLOBAL_SHADER(FTSRDilateVelocityCS,        "/Engine/Private/TemporalSuperResolution/TSRDilateVelocity.usf",        "MainCS", SF_Compute);
 IMPLEMENT_GLOBAL_SHADER(FTSRDecimateHistoryCS,       "/Engine/Private/TemporalSuperResolution/TSRDecimateHistory.usf",       "MainCS", SF_Compute);
-IMPLEMENT_GLOBAL_SHADER(FTSRCompareTranslucencyCS,   "/Engine/Private/TemporalSuperResolution/TSRCompareTranslucency.usf",   "MainCS", SF_Compute);
 IMPLEMENT_GLOBAL_SHADER(FTSRRejectShadingCS,         "/Engine/Private/TemporalSuperResolution/TSRRejectShading.usf",         "MainCS", SF_Compute);
 IMPLEMENT_GLOBAL_SHADER(FTSRSpatialAntiAliasingCS,   "/Engine/Private/TemporalSuperResolution/TSRSpatialAntiAliasing.usf",   "MainCS", SF_Compute);
 IMPLEMENT_GLOBAL_SHADER(FTSRFilterAntiAliasingCS,    "/Engine/Private/TemporalSuperResolution/TSRFilterAntiAliasing.usf",    "MainCS", SF_Compute);
@@ -1201,8 +1161,6 @@ ITemporalUpscaler::FOutputs AddTemporalSuperResolutionPasses(
 
 	bool bIsSeperateTranslucyTexturesValid = PassInputs.PostDOFTranslucencyResources.IsValid();
 
-	const bool bRejectSeparateTranslucency = false; // bIsSeperateTranslucyTexturesValid && CVarTSRTranslucencyPreviousFrameRejection.GetValueOnRenderThread() != 0;
-
 	EPixelFormat ColorFormat = bSupportsAlpha ? PF_FloatRGBA : PF_FloatR11G11B10;
 	EPixelFormat HistoryColorFormat = (CVarTSRR11G11B10History.GetValueOnRenderThread() != 0 && !bSupportsAlpha) ? PF_FloatR11G11B10 : PF_FloatRGBA;
 
@@ -1873,64 +1831,6 @@ ITemporalUpscaler::FOutputs AddTemporalSuperResolutionPasses(
 		DilatedVelocityTexture = HoleFilledVelocityTexture;
 	}
 
-	FRDGTextureRef TranslucencyRejectionTexture = nullptr;
-	if (bRejectSeparateTranslucency && View.PrevViewInfo.SeparateTranslucency != nullptr)
-	{
-		FRDGTextureRef PrevTranslucencyTexture;
-		FScreenPassTextureViewport PrevTranslucencyViewport;
-
-		if (View.PrevViewInfo.SeparateTranslucency)
-		{
-
-			PrevTranslucencyTexture = GraphBuilder.RegisterExternalTexture(View.PrevViewInfo.SeparateTranslucency);
-			PrevTranslucencyViewport = FScreenPassTextureViewport(PrevTranslucencyTexture->Desc.Extent, View.PrevViewInfo.SeparateTranslucencyViewRect);
-
-		}
-		else
-		{
-			PrevTranslucencyTexture = BlackAlphaOneDummy;
-			PrevTranslucencyViewport = FScreenPassTextureViewport(FIntPoint(1, 1), FIntRect(FIntPoint(0, 0), FIntPoint(1, 1)));
-		}
-
-		{
-			FRDGTextureDesc Desc = FRDGTextureDesc::Create2D(
-				InputExtent,
-				PF_R8,
-				FClearValueBinding::None,
-				/* InFlags = */ TexCreate_ShaderResource | TexCreate_UAV);
-
-			TranslucencyRejectionTexture = GraphBuilder.CreateTexture(Desc, TEXT("TSR.TranslucencyRejection"));
-		}
-
-		FTSRCompareTranslucencyCS::FParameters* PassParameters = GraphBuilder.AllocParameters<FTSRCompareTranslucencyCS::FParameters>();
-		PassParameters->CommonParameters = CommonParameters;
-
-		PassParameters->TranslucencyInfo = GetScreenPassTextureViewportParameters(FScreenPassTextureViewport(
-			SeparateTranslucencyTexture->Desc.Extent, SeparateTranslucencyRect));
-		PassParameters->PrevTranslucencyInfo = GetScreenPassTextureViewportParameters(PrevTranslucencyViewport);
-		PassParameters->PrevTranslucencyPreExposureCorrection = PrevHistoryParameters.HistoryPreExposureCorrection;
-		PassParameters->TranslucencyHighlightLuminance = -1.0; //CVarTSRTranslucencyHighlightLuminance.GetValueOnRenderThread();
-
-		PassParameters->ScreenPosToPrevTranslucencyTextureUV = FScreenTransform::ChangeTextureBasisFromTo(
-			PrevTranslucencyViewport, FScreenTransform::ETextureBasis::ScreenPosition, FScreenTransform::ETextureBasis::TextureUV);
-
-		PassParameters->DilatedVelocityTexture = DilatedVelocityTexture;
-		PassParameters->TranslucencyTexture = SeparateTranslucencyTexture;
-		PassParameters->PrevTranslucencyTexture = PrevTranslucencyTexture;
-
-		PassParameters->TranslucencyRejectionOutput = GraphBuilder.CreateUAV(TranslucencyRejectionTexture);
-		PassParameters->DebugOutput = CreateDebugUAV(InputExtent, TEXT("Debug.TSR.CompareTranslucency"));
-
-		TShaderMapRef<FTSRCompareTranslucencyCS> ComputeShader(View.ShaderMap);
-		FComputeShaderUtils::AddPass(
-			GraphBuilder,
-			RDG_EVENT_NAME("TSR CompareTranslucency %dx%d", InputRect.Width(), InputRect.Height()),
-			AsyncComputePasses >= 3 ? ERDGPassFlags::AsyncCompute : ERDGPassFlags::Compute,
-			ComputeShader,
-			PassParameters,
-			FComputeShaderUtils::GetGroupCount(InputRect.Size(), 8));
-	}
-
 	// Reject the history with frequency decomposition.
 	FRDGTextureRef HistoryRejectionTexture;
 	FRDGTextureRef InputSceneColorLdrLumaTexture = nullptr;
@@ -2106,7 +2006,6 @@ ITemporalUpscaler::FOutputs AddTemporalSuperResolutionPasses(
 			FRDGTextureSRVDesc::CreateWithPixelFormat(PassInputs.SceneDepthTexture, PF_X24_G8));
 		PassParameters->InputSceneTranslucencyTexture = SeparateTranslucencyTexture;
 		PassParameters->HistoryRejectionTexture = HistoryRejectionTexture;
-		PassParameters->TranslucencyRejectionTexture = TranslucencyRejectionTexture ? TranslucencyRejectionTexture : BlackDummy;
 
 		PassParameters->DilatedVelocityTexture = DilatedVelocityTexture;
 		PassParameters->GrandDilatedVelocityTexture = GrandHoleFilledVelocityTexture;
@@ -2116,17 +2015,14 @@ ITemporalUpscaler::FOutputs AddTemporalSuperResolutionPasses(
 		PassParameters->NoiseFilteringTexture = NoiseFilteringTexture;
 		PassParameters->HoleFilledVelocityMaskTexture = HoleFilledVelocityMaskTexture;
 
-		PassParameters->TranslucencyInfo = GetScreenPassTextureViewportParameters(FScreenPassTextureViewport(
-			PassParameters->TranslucencyRejectionTexture->Desc.Extent, SeparateTranslucencyRect));
-		PassParameters->TranslucencyPixelPosMin = PassParameters->TranslucencyInfo.ViewportMin;
-		PassParameters->TranslucencyPixelPosMax = PassParameters->TranslucencyInfo.ViewportMax - 1;
+		PassParameters->TranslucencyPixelPosMin = SeparateTranslucencyRect.Min;
+		PassParameters->TranslucencyPixelPosMax = SeparateTranslucencyRect.Max - 1;
 
 		FScreenTransform HistoryPixelPosToViewportUV = (FScreenTransform::Identity + 0.5f) * CommonParameters.HistoryInfo.ViewportSizeInverse;
 		PassParameters->HistoryPixelPosToScreenPos = HistoryPixelPosToViewportUV * FScreenTransform::ViewportUVToScreenPos;
 		PassParameters->HistoryPixelPosToInputPPCo = HistoryPixelPosToViewportUV * CommonParameters.InputInfo.ViewportSize + CommonParameters.InputJitter + CommonParameters.InputPixelPosMin;
-		PassParameters->HistoryPixelPosToTranslucencyPPCo = HistoryPixelPosToViewportUV * PassParameters->TranslucencyInfo.ViewportSize + CommonParameters.InputJitter * PassParameters->TranslucencyInfo.ViewportSize / CommonParameters.InputInfo.ViewportSize + SeparateTranslucencyRect.Min;
+		PassParameters->HistoryPixelPosToTranslucencyPPCo = HistoryPixelPosToViewportUV * SeparateTranslucencyRect.Size() + CommonParameters.InputJitter * SeparateTranslucencyRect.Size() / CommonParameters.InputInfo.ViewportSize + SeparateTranslucencyRect.Min;
 		PassParameters->HistoryQuantizationError = ComputePixelFormatQuantizationError((History.ColorArray ? History.ColorArray : History.Output)->Desc.Format);
-		PassParameters->MinTranslucencyRejection = TranslucencyRejectionTexture == nullptr ? 1.0 : 0.0;
 
 		// All parameters to control the sample count in history.
 		PassParameters->HistorySampleCount = MaxHistorySampleCount / OutputToHistoryResolutionFractionSquare;
@@ -2345,14 +2241,6 @@ ITemporalUpscaler::FOutputs AddTemporalSuperResolutionPasses(
 			}
 		}
 
-		// Extract the translucency buffer to compare it with next frame
-		if (bRejectSeparateTranslucency)
-		{
-			GraphBuilder.QueueTextureExtraction(
-				SeparateTranslucencyTexture, &View.ViewState->PrevFrameViewInfo.SeparateTranslucency);
-			View.ViewState->PrevFrameViewInfo.SeparateTranslucencyViewRect = InputRect;
-		}
-
 		// Extract the output for next frame SSR so that separate translucency shows up in SSR.
 		{
 			GraphBuilder.QueueTextureExtraction(
diff --git a/Engine/Source/Runtime/Renderer/Private/RendererScene.cpp b/Engine/Source/Runtime/Renderer/Private/RendererScene.cpp
--- a/Engine/Source/Runtime/Renderer/Private/RendererScene.cpp
+++ b/Engine/Source/Runtime/Renderer/Private/RendererScene.cpp
@@ -1011,7 +1011,6 @@ uint64 FPreviousViewInfo::GetGPUSizeBytes(bool bLogSizes) const
 		GetRenderTargetGPUSizeBytes(CompressedOpaqueDepth, bLogSizes) +
 		GetRenderTargetGPUSizeBytes(CompressedOpaqueShadingModel, bLogSizes) +
 		GetRenderTargetGPUSizeBytes(ScreenSpaceRayTracingInput, bLogSizes) +
-		GetRenderTargetGPUSizeBytes(SeparateTranslucency, bLogSizes) +
 		TemporalAAHistory.GetGPUSizeBytes(bLogSizes) +
 		TSRHistory.GetGPUSizeBytes(bLogSizes) +
 		GetRenderTargetGPUSizeBytes(HalfResTemporalAAHistory, bLogSizes) +
diff --git a/Engine/Source/Runtime/Renderer/Private/SceneRendering.h b/Engine/Source/Runtime/Renderer/Private/SceneRendering.h
--- a/Engine/Source/Runtime/Renderer/Private/SceneRendering.h
+++ b/Engine/Source/Runtime/Renderer/Private/SceneRendering.h
@@ -1136,10 +1136,6 @@ struct FPreviousViewInfo
 	// Bleed free scene color to use for screen space ray tracing.
 	TRefCountPtr<IPooledRenderTarget> ScreenSpaceRayTracingInput;
 
-	// Separate translucency of previous frame to detect translucency change in TAA
-	FIntRect SeparateTranslucencyViewRect;
-	TRefCountPtr<IPooledRenderTarget> SeparateTranslucency;
-
 	// Temporal AA result of last frame
 	FTemporalAAHistory TemporalAAHistory;
 
