Works arround a weird shader compilation failure in spirv_hlsl affecting D3D11

diff --git a/Engine/Source/Runtime/Renderer/Private/PostProcess/TemporalSuperResolution.cpp b/Engine/Source/Runtime/Renderer/Private/PostProcess/TemporalSuperResolution.cpp
--- a/Engine/Source/Runtime/Renderer/Private/PostProcess/TemporalSuperResolution.cpp
+++ b/Engine/Source/Runtime/Renderer/Private/PostProcess/TemporalSuperResolution.cpp
@@ -862,6 +862,7 @@ class FTSRUpdateHistoryCS : public FTSRShader
 		SHADER_PARAMETER_RDG_TEXTURE(Texture2D, NoiseFilteringTexture)
 		SHADER_PARAMETER_RDG_TEXTURE(Texture2D, HoleFilledVelocityMaskTexture)
 
+		SHADER_PARAMETER_STRUCT(FScreenPassTextureViewportParameters, TranslucencyInfo)
 		SHADER_PARAMETER(FIntPoint, TranslucencyPixelPosMin)
 		SHADER_PARAMETER(FIntPoint, TranslucencyPixelPosMax)
 
