Fix ElectricDream cooking issue.

diff --git a/Engine/Source/Runtime/Renderer/Private/Strata/Strata.cpp b/Engine/Source/Runtime/Renderer/Private/Strata/Strata.cpp
--- a/Engine/Source/Runtime/Renderer/Private/Strata/Strata.cpp
+++ b/Engine/Source/Runtime/Renderer/Private/Strata/Strata.cpp
@@ -320,14 +320,17 @@ static void RecordStrataAnalytics(EShaderPlatform InPlatform)
 	}
 }
 
-static EPixelFormat GetTopLayerTextureFormat()
+static EPixelFormat GetTopLayerTextureFormat(bool bUseDBufferPass)
 {
 	const bool bStrataHighQualityNormal = GetNormalQuality() > 0;
 
 	// High quality normal is not supported on platforms that do not support R32G32 UAV load.
 	// This is dues to the way Strata account for decals. See FStrataDBufferPassCS, updating TopLayerTexture this way.
 	// If you encounter this check, you must disable high quality normal for Strata (material shaders must be recompiled to account for that).
-	check(!bStrataHighQualityNormal || (bStrataHighQualityNormal && UE::PixelFormat::HasCapabilities(PF_R32G32_UINT, EPixelFormatCapabilities::TypedUAVLoad)));
+	if (bUseDBufferPass)
+	{
+		check(!bStrataHighQualityNormal || (bStrataHighQualityNormal && UE::PixelFormat::HasCapabilities(PF_R32G32_UINT, EPixelFormatCapabilities::TypedUAVLoad)));
+	}
 
 	return bStrataHighQualityNormal ? PF_R32G32_UINT : PF_R32_UINT;
 }
@@ -353,6 +356,7 @@ void InitialiseStrataFrameSceneData(FRDGBuilder& GraphBuilder, FSceneRenderer& S
 	// Compute the max byte per pixels required by the views
 	bool bNeedBSDFOffsets = false;
 	bool bNeedUAV = false;
+	bool bUseDBufferPass = false;
 
 	FIntPoint MaterialBufferSizeXY;
 	UpdateMaterialBufferToTiledResolution(FIntPoint(1, 1), MaterialBufferSizeXY);
@@ -373,6 +377,7 @@ void InitialiseStrataFrameSceneData(FRDGBuilder& GraphBuilder, FSceneRenderer& S
 			bNeedBSDFOffsets = bNeedBSDFOffsets || NeedBSDFOffsets(SceneRenderer.Scene, View);
 			bNeedUAV = bNeedUAV || IsDBufferPassEnabled(View.GetShaderPlatform()) || NaniteComputeMaterialsSupported();
 			Out.ViewsMaxBytesPerPixel = FMath::Max(Out.ViewsMaxBytesPerPixel, View.StrataViewData.MaxBytesPerPixel);
+			bUseDBufferPass = bUseDBufferPass || IsDBufferPassEnabled(View.GetShaderPlatform());
 
 			// Only use primary views max. byte per pixel as reflection/capture views can bias allocation requirement when using growing-only mode
 			if (!View.bIsPlanarReflection && !View.bIsReflectionCapture && !View.bIsSceneCapture)
@@ -408,7 +413,7 @@ void InitialiseStrataFrameSceneData(FRDGBuilder& GraphBuilder, FSceneRenderer& S
 
 		// Top layer texture
 		{
-			Out.TopLayerTexture = GraphBuilder.CreateTexture(FRDGTextureDesc::Create2D(SceneTextureExtent, GetTopLayerTextureFormat(), FClearValueBinding::Black, TexCreate_RenderTargetable | TexCreate_ShaderResource | TexCreate_FastVRAM | (bNeedUAV ? TexCreate_UAV : TexCreate_None)), TEXT("Substrate.TopLayerTexture"));
+			Out.TopLayerTexture = GraphBuilder.CreateTexture(FRDGTextureDesc::Create2D(SceneTextureExtent, GetTopLayerTextureFormat(bUseDBufferPass), FClearValueBinding::Black, TexCreate_RenderTargetable | TexCreate_ShaderResource | TexCreate_FastVRAM | (bNeedUAV ? TexCreate_UAV : TexCreate_None)), TEXT("Substrate.TopLayerTexture"));
 		}
 
 		// Separated subsurface and rough refraction textures
@@ -1215,7 +1220,7 @@ void SetBasePassRenderTargetOutputFormat(const EShaderPlatform Platform, const F
 		}
 
 		// Add another MRT for Strata top layer information
-		OutEnvironment.SetRenderTargetOutputFormat(BufferInfo.NumTargets + STRATA_BASE_PASS_MRT_OUTPUT_COUNT, GetTopLayerTextureFormat());
+		OutEnvironment.SetRenderTargetOutputFormat(BufferInfo.NumTargets + STRATA_BASE_PASS_MRT_OUTPUT_COUNT, GetTopLayerTextureFormat(IsDBufferPassEnabled(Platform)));
 	}
 }
 
