Added missing header guard.

diff --git a/Engine/Source/Runtime/Renderer/Private/SceneVisibilityPrivate.h b/Engine/Source/Runtime/Renderer/Private/SceneVisibilityPrivate.h
--- a/Engine/Source/Runtime/Renderer/Private/SceneVisibilityPrivate.h
+++ b/Engine/Source/Runtime/Renderer/Private/SceneVisibilityPrivate.h
@@ -1,5 +1,7 @@
 // Copyright Epic Games, Inc. All Rights Reserved.
 
+#pragma once
+
 #include "SceneVisibility.h"
 #include "ScenePrivate.h"
 #include "Containers/ConsumeAllMpmcQueue.h"
