One-line fix for VRS bug when tonemapper compute pass is enabled.

diff --git a/Engine/Source/Runtime/Renderer/Private/PostProcess/PostProcessTonemap.cpp b/Engine/Source/Runtime/Renderer/Private/PostProcess/PostProcessTonemap.cpp
--- a/Engine/Source/Runtime/Renderer/Private/PostProcess/PostProcessTonemap.cpp
+++ b/Engine/Source/Runtime/Renderer/Private/PostProcess/PostProcessTonemap.cpp
@@ -975,11 +975,11 @@ FScreenPassTexture AddTonemapPass(FRDGBuilder& GraphBuilder, const FViewInfo& Vi
 			SetShaderParameters(RHICmdList, VertexShader, VertexShader.GetVertexShader(), PassParameters->Tonemap);
 			SetShaderParameters(RHICmdList, PixelShader, PixelShader.GetPixelShader(), *PassParameters);
 		});
+	}
 
-		if (OutputLuminance && View.ViewState)
-		{
-			GraphBuilder.QueueTextureExtraction(OutputLuminance, &View.ViewState->PrevFrameViewInfo.LuminanceHistory);
-		}
+	if (OutputLuminance && View.ViewState)
+	{
+		GraphBuilder.QueueTextureExtraction(OutputLuminance, &View.ViewState->PrevFrameViewInfo.LuminanceHistory);
 	}
 
 	return MoveTemp(Output);
