PSO precaching: Added support for a mobile specific render passes

diff --git a/Engine/Source/Runtime/Engine/Private/SceneTexturesConfig.cpp b/Engine/Source/Runtime/Engine/Private/SceneTexturesConfig.cpp
--- a/Engine/Source/Runtime/Engine/Private/SceneTexturesConfig.cpp
+++ b/Engine/Source/Runtime/Engine/Private/SceneTexturesConfig.cpp
@@ -67,7 +67,7 @@ static EPixelFormat GetMobileSceneColorFormat(bool bRequiresAlphaChannel)
 
 	EPixelFormat Format = DefaultColorFormat;
 	static const auto CVar = IConsoleManager::Get().FindTConsoleVariableDataInt(TEXT("r.Mobile.SceneColorFormat"));
-	int32 MobileSceneColor = CVar->GetValueOnRenderThread();
+	int32 MobileSceneColor = CVar->GetValueOnAnyThread();
 	switch (MobileSceneColor)
 	{
 	case 1:
@@ -211,6 +211,43 @@ static uint32 GetEditorPrimitiveNumSamples(ERHIFeatureLevel::Type FeatureLevel)
 	return SampleCount;
 }
 
+static void SetupMobileGBufferFlags(FGBufferBindings GBufferBindings[GBL_Num], bool bRequiresMultiPass)
+{
+	ETextureCreateFlags AddFlags = TexCreate_InputAttachmentRead;
+
+	if (!bRequiresMultiPass)
+	{
+		// use memoryless GBuffer when possible
+		AddFlags |= TexCreate_Memoryless;
+	}
+
+	for (uint32 Layout = 0; Layout < GBL_Num; ++Layout)
+	{
+		FGBufferBindings& Bindings = GBufferBindings[Layout];
+
+		Bindings.GBufferA.Flags |= AddFlags;
+		Bindings.GBufferB.Flags |= AddFlags;
+		Bindings.GBufferC.Flags |= AddFlags;
+		Bindings.GBufferD.Flags |= AddFlags;
+		Bindings.GBufferE.Flags |= AddFlags;
+
+		// Mobile uses FBF/subpassLoad to fetch data from GBuffer, and FBF does not always work with sRGB targets 
+		Bindings.GBufferA.Flags &= (~TexCreate_SRGB);
+		Bindings.GBufferB.Flags &= (~TexCreate_SRGB);
+		Bindings.GBufferC.Flags &= (~TexCreate_SRGB);
+		Bindings.GBufferD.Flags &= (~TexCreate_SRGB);
+		Bindings.GBufferE.Flags &= (~TexCreate_SRGB);
+
+		// Input attachments with PF_R8G8B8A8 has better support on mobile than PF_B8G8R8A8
+		auto OverrideB8G8R8A8 = [](FGBufferBinding& Binding) { if (Binding.Format == PF_B8G8R8A8) Binding.Format = PF_R8G8B8A8; };
+		OverrideB8G8R8A8(Bindings.GBufferA);
+		OverrideB8G8R8A8(Bindings.GBufferB);
+		OverrideB8G8R8A8(Bindings.GBufferC);
+		OverrideB8G8R8A8(Bindings.GBufferD);
+		OverrideB8G8R8A8(Bindings.GBufferE);
+	}
+}
+
 void FSceneTexturesConfig::Init(const FSceneTexturesConfigInitSettings& InitSettings)
 {
 	FeatureLevel			= InitSettings.FeatureLevel;
@@ -261,6 +298,12 @@ void FSceneTexturesConfig::Init(const FSceneTexturesConfigInitSettings& InitSett
 
 			BindingCache.GBufferParams = DefaultParams;
 			BindingCache.bInitialized = true;
+
+			if (ShadingPath == EShadingPath::Mobile)
+			{
+				bool bRequiresMultiPass = false; // This will be overriden later when in the renderer InitViews. Does not affect PSO state
+				::SetupMobileGBufferFlags(BindingCache.Bindings, bRequiresMultiPass);
+			}
 		}
 
 		for (uint32 Layout = 0; Layout < GBL_Num; ++Layout)
@@ -270,6 +313,11 @@ void FSceneTexturesConfig::Init(const FSceneTexturesConfigInitSettings& InitSett
 	}
 }
 
+void FSceneTexturesConfig::SetupMobileGBufferFlags(bool bRequiresMultiPass)
+{
+	::SetupMobileGBufferFlags(GBufferBindings, bRequiresMultiPass);
+}
+
 void FSceneTexturesConfig::BuildSceneColorAndDepthFlags()
 {
     GetSceneColorFormatAndCreateFlags(FeatureLevel, bRequiresAlphaChannel, ExtraSceneColorCreateFlags, NumSamples, bMemorylessMSAA, ColorFormat, ColorCreateFlags);
diff --git a/Engine/Source/Runtime/Engine/Public/SceneTexturesConfig.h b/Engine/Source/Runtime/Engine/Public/SceneTexturesConfig.h
--- a/Engine/Source/Runtime/Engine/Public/SceneTexturesConfig.h
+++ b/Engine/Source/Runtime/Engine/Public/SceneTexturesConfig.h
@@ -128,6 +128,7 @@ struct ENGINE_API FSceneTexturesConfig
 	void Init(const FSceneTexturesConfigInitSettings& InitSettings);
     void BuildSceneColorAndDepthFlags();
 	uint32 GetGBufferRenderTargetsInfo(FGraphicsPipelineRenderTargetsInfo& RenderTargetsInfo, EGBufferLayout Layout = GBL_Default) const;
+	void SetupMobileGBufferFlags(bool bRequiresMultiPass);
 
 	FORCEINLINE bool IsValid() const
 	{
diff --git a/Engine/Source/Runtime/RenderCore/Private/GBufferInfo.cpp b/Engine/Source/Runtime/RenderCore/Private/GBufferInfo.cpp
--- a/Engine/Source/Runtime/RenderCore/Private/GBufferInfo.cpp
+++ b/Engine/Source/Runtime/RenderCore/Private/GBufferInfo.cpp
@@ -8,6 +8,8 @@
 #include "RenderUtils.h"
 #include "HAL/IConsoleManager.h"
 #include "DataDrivenShaderPlatformInfo.h"
+#include "RenderUtils.h"
+#include "RHIGlobals.h"
 
 // This is very ugly, but temporary. We can't include EngineTypes.h because this file is ShaderCore and that
 // file is Engine. But since we will remove this flag anyways, we will copy/paste the enums for now, but remove
@@ -163,6 +165,9 @@ FGBufferBinding FindGBufferBindingByName(const FGBufferInfo& GBufferInfo, const
 		case GBT_Float_16_16_16_16:
 			PixelFormat = PF_FloatRGBA;
 			break;
+		case GBT_Float_32:
+			PixelFormat = PF_R32_FLOAT;
+			break;
 		case GBT_Invalid:
 		default:
 			check(0);
@@ -542,10 +547,46 @@ FGBufferInfo RENDERCORE_API FetchLegacyGBufferInfo(const FGBufferParams& Params)
 }
 
 
+FGBufferInfo RENDERCORE_API FetchMobileGBufferInfo(const FGBufferParams& Params)
+{
+	// for now, take desktop and trim out what we do not require.
+	FGBufferInfo Info = FetchLegacyGBufferInfo(Params);
+	
+	// If we are using GL and don't have FBF support, use PLS
+	bool bUsingPixelLocalStorage = IsAndroidOpenGLESPlatform(Params.ShaderPlatform) && GSupportsPixelLocalStorage && GSupportsShaderDepthStencilFetch;
+	if (bUsingPixelLocalStorage)
+	{
+		Info.NumTargets = 1;
+	}
+	else
+	{
+		Info.NumTargets = 4;
+		if (MobileUsesExtenedGBuffer(Params.ShaderPlatform))
+		{
+			Info.NumTargets++;
+		}
+		
+		if (MobileRequiresSceneDepthAux(Params.ShaderPlatform))
+		{
+			// if used, mobile deferred is always F32.
+			Info.Targets[Info.NumTargets].Init(GBT_Float_32, TEXT("DepthAux"), false, true, true, false);
+			Info.NumTargets++;
+		}
+	}
+
+	return Info;
+}
+
 FGBufferInfo RENDERCORE_API FetchFullGBufferInfo(const FGBufferParams& Params)
 {
 	// For now, we are only doing legacy. But next, we will have a switch between the old and new formats.
-	FGBufferInfo Ret = FetchLegacyGBufferInfo(Params);
-	return Ret;
+	if (IsMobilePlatform(Params.ShaderPlatform))
+	{
+		return FetchMobileGBufferInfo(Params);
+	}
+	else
+	{
+		return FetchLegacyGBufferInfo(Params);
+	}
 }
 
diff --git a/Engine/Source/Runtime/RenderCore/Public/GBufferInfo.h b/Engine/Source/Runtime/RenderCore/Public/GBufferInfo.h
--- a/Engine/Source/Runtime/RenderCore/Public/GBufferInfo.h
+++ b/Engine/Source/Runtime/RenderCore/Public/GBufferInfo.h
@@ -89,7 +89,8 @@ enum EGBufferType
 	GBT_Unorm_10_10_10_2,
 	GBT_Unorm_16_16_16_16,
 	GBT_Float_16_16,
-	GBT_Float_16_16_16_16,	
+	GBT_Float_16_16_16_16,
+	GBT_Float_32,
 	GBT_Num
 };
 
diff --git a/Engine/Source/Runtime/Renderer/Private/CustomDepthRendering.cpp b/Engine/Source/Runtime/Renderer/Private/CustomDepthRendering.cpp
--- a/Engine/Source/Runtime/Renderer/Private/CustomDepthRendering.cpp
+++ b/Engine/Source/Runtime/Renderer/Private/CustomDepthRendering.cpp
@@ -739,4 +739,4 @@ FMeshPassProcessor* CreateCustomDepthPassProcessor(ERHIFeatureLevel::Type Featur
 }
 
 REGISTER_MESHPASSPROCESSOR_AND_PSOCOLLECTOR(RegisterCustomDepthPass, CreateCustomDepthPassProcessor, EShadingPath::Deferred, EMeshPass::CustomDepth, EMeshPassFlags::MainView);
-FRegisterPassProcessorCreateFunction RegisterMobileCustomDepthPass(&CreateCustomDepthPassProcessor, EShadingPath::Mobile, EMeshPass::CustomDepth, EMeshPassFlags::MainView);
+REGISTER_MESHPASSPROCESSOR_AND_PSOCOLLECTOR(RegisterMobileCustomDepthPass, CreateCustomDepthPassProcessor, EShadingPath::Mobile, EMeshPass::CustomDepth, EMeshPassFlags::MainView);
diff --git a/Engine/Source/Runtime/Renderer/Private/DepthRendering.cpp b/Engine/Source/Runtime/Renderer/Private/DepthRendering.cpp
--- a/Engine/Source/Runtime/Renderer/Private/DepthRendering.cpp
+++ b/Engine/Source/Runtime/Renderer/Private/DepthRendering.cpp
@@ -1193,7 +1193,7 @@ FMeshPassProcessor* CreateDepthPassProcessor(ERHIFeatureLevel::Type FeatureLevel
 }
 
 REGISTER_MESHPASSPROCESSOR_AND_PSOCOLLECTOR(DepthPass, CreateDepthPassProcessor, EShadingPath::Deferred, EMeshPass::DepthPass, EMeshPassFlags::CachedMeshCommands | EMeshPassFlags::MainView);
-FRegisterPassProcessorCreateFunction RegisterMobileDepthPass(&CreateDepthPassProcessor, EShadingPath::Mobile, EMeshPass::DepthPass, EMeshPassFlags::CachedMeshCommands | EMeshPassFlags::MainView);
+REGISTER_MESHPASSPROCESSOR_AND_PSOCOLLECTOR(MobileDepthPass, CreateDepthPassProcessor, EShadingPath::Mobile, EMeshPass::DepthPass, EMeshPassFlags::CachedMeshCommands | EMeshPassFlags::MainView);
 
 FMeshPassProcessor* CreateSecondStageDepthPassProcessor(ERHIFeatureLevel::Type FeatureLevel, const FScene* Scene, const FSceneView* InViewIfDynamicMeshCommand, FMeshPassDrawListContext* InDrawListContext)
 {
diff --git a/Engine/Source/Runtime/Renderer/Private/MobileBasePass.cpp b/Engine/Source/Runtime/Renderer/Private/MobileBasePass.cpp
--- a/Engine/Source/Runtime/Renderer/Private/MobileBasePass.cpp
+++ b/Engine/Source/Runtime/Renderer/Private/MobileBasePass.cpp
@@ -50,7 +50,7 @@ bool MobileUsesNoLightMapPermutation(const FMeshMaterialShaderPermutationParamet
 template <ELightMapPolicyType Policy, bool bEnableLocalLights>
 bool GetUniformMobileBasePassShaders(
 	const FMaterial& Material, 
-	FVertexFactoryType* VertexFactoryType, 
+	const FVertexFactoryType* VertexFactoryType, 
 	bool bEnableSkyLight,
 	TShaderRef<TMobileBasePassVSPolicyParamType<FUniformLightMapPolicy>>& VertexShader,
 	TShaderRef<TMobileBasePassPSPolicyParamType<FUniformLightMapPolicy>>& PixelShader
@@ -102,7 +102,7 @@ template <bool bEnableLocalLights>
 bool GetMobileBasePassShaders(
 	ELightMapPolicyType LightMapPolicyType, 
 	const FMaterial& Material, 
-	FVertexFactoryType* VertexFactoryType, 
+	const FVertexFactoryType* VertexFactoryType, 
 	bool bEnableSkyLight,
 	TShaderRef<TMobileBasePassVSPolicyParamType<FUniformLightMapPolicy>>& VertexShader,
 	TShaderRef<TMobileBasePassPSPolicyParamType<FUniformLightMapPolicy>>& PixelShader
@@ -140,7 +140,7 @@ bool MobileBasePass::GetShaders(
 	ELightMapPolicyType LightMapPolicyType,
 	bool bEnableLocalLights,
 	const FMaterial& MaterialResource,
-	FVertexFactoryType* VertexFactoryType,
+	const FVertexFactoryType* VertexFactoryType,
 	bool bEnableSkyLight, 
 	TShaderRef<TMobileBasePassVSPolicyParamType<FUniformLightMapPolicy>>& VertexShader,
 	TShaderRef<TMobileBasePassPSPolicyParamType<FUniformLightMapPolicy>>& PixelShader)
@@ -213,23 +213,22 @@ ELightMapPolicyType MobileBasePass::SelectMeshLightmapPolicy(
 	const FMeshBatch& Mesh, 
 	const FPrimitiveSceneProxy* PrimitiveSceneProxy,
 	const FLightSceneInfo* MobileDirectionalLight,
-	FMaterialShadingModelField ShadingModels, 
 	bool bPrimReceivesCSM,
 	bool bUsesDeferredShading,
-	bool bIsTranslucent,
-	ERHIFeatureLevel::Type FeatureLevel)
+	bool bIsLitMaterial,
+	bool bIsTranslucent)
 {
 	// Unlit uses NoLightmapPolicy with 0 point lights
 	ELightMapPolicyType SelectedLightmapPolicy = LMP_NO_LIGHTMAP;
 	
-	const bool bIsLitMaterial = ShadingModels.IsLit();
 	if (bIsLitMaterial)
 	{
+		constexpr ERHIFeatureLevel::Type FeatureLevel = ERHIFeatureLevel::ES3_1;
 		const FReadOnlyCVARCache& ReadOnlyCVARCache = FReadOnlyCVARCache::Get();
 
 		if (!ReadOnlyCVARCache.bAllowStaticLighting || (ReadOnlyCVARCache.bMobileEnableNoPrecomputedLightingCSMShader && Scene && Scene->GetForceNoPrecomputedLighting()))
 		{
-			if (!bIsTranslucent && !ShadingModels.HasShadingModel(MSM_SingleLayerWater))
+			if (!bIsTranslucent)
 			{
 				// Whether to use a single CSM permutation with a branch in the shader
 				bPrimReceivesCSM |= MobileUseCSMShaderBranch();
@@ -341,6 +340,93 @@ ELightMapPolicyType MobileBasePass::SelectMeshLightmapPolicy(
 	return SelectedLightmapPolicy;
 }
 
+typedef TArray<ELightMapPolicyType, TInlineAllocator<4>> FMobileLightMapPolicyTypeList;
+
+static FMobileLightMapPolicyTypeList GetUniformLightMapPolicyTypeForPSOCollection(bool bLitMaterial, bool bTranslucent, bool bUsesDeferredShading, bool bCanReceiveCSM, bool bMovable)
+{
+	FMobileLightMapPolicyTypeList Result;
+	
+	if (bLitMaterial)
+	{
+		const FReadOnlyCVARCache& ReadOnlyCVARCache = FReadOnlyCVARCache::Get();
+		
+		if (!ReadOnlyCVARCache.bAllowStaticLighting)
+		{
+			if (bUsesDeferredShading || bTranslucent || !MobileUseCSMShaderBranch())
+			{
+				Result.Add(LMP_NO_LIGHTMAP);
+			}
+						
+			if (!bTranslucent && !bUsesDeferredShading)
+			{
+				// permutation that can receive CSM
+				Result.Add(LMP_MOBILE_DIRECTIONAL_LIGHT_CSM);
+			}
+		}
+		else
+		{
+			if (ReadOnlyCVARCache.bEnableLowQualityLightmaps)
+			{
+				if (ReadOnlyCVARCache.bMobileEnableStaticAndCSMShadowReceivers && !bUsesDeferredShading)
+				{
+					if (ReadOnlyCVARCache.bMobileAllowDistanceFieldShadows)
+					{
+						Result.Add(LMP_MOBILE_DISTANCE_FIELD_SHADOWS_LIGHTMAP_AND_CSM);
+					}
+
+					if (ReadOnlyCVARCache.bMobileAllowMovableDirectionalLights)
+					{
+						Result.Add(LMP_MOBILE_MOVABLE_DIRECTIONAL_LIGHT_CSM_WITH_LIGHTMAP);
+					}
+
+					Result.Add(LMP_MOBILE_DIRECTIONAL_LIGHT_CSM_AND_LIGHTMAP);
+				}
+
+				if (ReadOnlyCVARCache.bMobileAllowDistanceFieldShadows)
+				{
+					Result.Add(LMP_MOBILE_DISTANCE_FIELD_SHADOWS_AND_LQ_LIGHTMAP);
+				}
+
+				if (ReadOnlyCVARCache.bMobileAllowMovableDirectionalLights)
+				{
+					if (bUsesDeferredShading)
+					{
+						Result.Add(LMP_LQ_LIGHTMAP);
+					}
+					else
+					{
+						Result.Add(LMP_MOBILE_MOVABLE_DIRECTIONAL_LIGHT_WITH_LIGHTMAP);
+					}
+				}
+				else
+				{
+					Result.Add(LMP_LQ_LIGHTMAP);
+				}
+			}
+						
+			// ILC/LVM
+			if (bMovable)
+			{
+				if (!bUsesDeferredShading && ReadOnlyCVARCache.bMobileEnableStaticAndCSMShadowReceivers)
+				{
+					Result.Add(LMP_MOBILE_DIRECTIONAL_LIGHT_CSM_AND_SH_INDIRECT);
+				}
+				else
+				{
+					Result.Add(LMP_MOBILE_DIRECTIONAL_LIGHT_AND_SH_INDIRECT);
+				}
+			}
+		}
+	}
+	else
+	{
+		// Unlit materials
+		Result.Add(LMP_NO_LIGHTMAP);
+	}
+
+	return Result;
+}
+
 void MobileBasePass::SetOpaqueRenderState(FMeshPassProcessorRenderState& DrawRenderState, const FPrimitiveSceneProxy* PrimitiveSceneProxy, const FMaterial& Material, FMaterialShadingModelField ShadingModels, bool bEnableReceiveDecalOutput, bool bUsesDeferredShading)
 {
 	uint8 StencilValue = 0;
@@ -626,48 +712,49 @@ FMobileBasePassMeshProcessor::FMobileBasePassMeshProcessor(
 	, TranslucencyPassType(InTranslucencyPassType)
 	, Flags(InFlags)
 	, bTranslucentBasePass(InTranslucencyPassType != ETranslucencyPass::TPT_MAX)
-	, bUsesDeferredShading(!bTranslucentBasePass && IsMobileDeferredShadingEnabled(GetFeatureLevelShaderPlatform(InFeatureLevel)))
+	, bDeferredShading(IsMobileDeferredShadingEnabled(GetFeatureLevelShaderPlatform(InFeatureLevel)))
+	, bPassUsesDeferredShading(bDeferredShading && !bTranslucentBasePass)
 {
 }
 
-bool FMobileBasePassMeshProcessor::TryAddMeshBatch(const FMeshBatch& RESTRICT MeshBatch, uint64 BatchElementMask, const FPrimitiveSceneProxy* RESTRICT PrimitiveSceneProxy, int32 StaticMeshId, const FMaterialRenderProxy& MaterialRenderProxy, const FMaterial& Material)
+bool FMobileBasePassMeshProcessor::ShouldDraw(const FMaterial& Material) const
 {
 	const FMaterialShadingModelField ShadingModels = Material.GetShadingModels();
-	const bool bIsMasked = IsMaskedBlendMode(Material);
-	const bool bIsTranslucent = IsTranslucentBlendMode(Material);
-	const bool bUsesWaterMaterial = ShadingModels.HasShadingModel(MSM_SingleLayerWater); // Water goes into the translucent pass
-	const bool bCanReceiveCSM = ((Flags & EFlags::CanReceiveCSM) == EFlags::CanReceiveCSM);
-
-	bool bResult = true;
+	const bool bIsTranslucent = IsTranslucentBlendMode(Material.GetBlendMode()) || ShadingModels.HasShadingModel(MSM_SingleLayerWater); // Water goes into the translucent pass;
+	const bool bCanReceiveCSM = ((Flags & FMobileBasePassMeshProcessor::EFlags::CanReceiveCSM) == FMobileBasePassMeshProcessor::EFlags::CanReceiveCSM);
 	if (bTranslucentBasePass)
 	{
 		// Skipping TPT_TranslucencyAfterDOFModulate. That pass is only needed for Dual Blending, which is not supported on Mobile.
-		bool bShouldDraw = (bIsTranslucent || bUsesWaterMaterial) &&
-		(   TranslucencyPassType == ETranslucencyPass::TPT_AllTranslucency
-		|| (TranslucencyPassType == ETranslucencyPass::TPT_TranslucencyStandard && !Material.IsMobileSeparateTranslucencyEnabled())
-		|| (TranslucencyPassType == ETranslucencyPass::TPT_TranslucencyAfterDOF && Material.IsMobileSeparateTranslucencyEnabled()));
+		bool bShouldDraw = bIsTranslucent &&
+			(TranslucencyPassType == ETranslucencyPass::TPT_AllTranslucency
+				|| (TranslucencyPassType == ETranslucencyPass::TPT_TranslucencyStandard && !Material.IsMobileSeparateTranslucencyEnabled())
+				|| (TranslucencyPassType == ETranslucencyPass::TPT_TranslucencyAfterDOF && Material.IsMobileSeparateTranslucencyEnabled()));
 
-		if (bShouldDraw)
-		{
-			check(bCanReceiveCSM == false);
-			const FLightSceneInfo* MobileDirectionalLight = MobileBasePass::GetDirectionalLightInfo(Scene, PrimitiveSceneProxy);
-			// Opaque meshes used for mobile pixel projected reflection could receive CSM in translucent pass.
-			ELightMapPolicyType LightmapPolicyType = MobileBasePass::SelectMeshLightmapPolicy(Scene, MeshBatch, PrimitiveSceneProxy, MobileDirectionalLight, ShadingModels, bCanReceiveCSM, false, bIsTranslucent, FeatureLevel);
-			bResult = Process(MeshBatch, BatchElementMask, StaticMeshId, PrimitiveSceneProxy, MaterialRenderProxy, Material, bIsMasked, bIsTranslucent, ShadingModels, LightmapPolicyType, bCanReceiveCSM, MeshBatch.LCI);
-		}
+		check(!bShouldDraw || bCanReceiveCSM == false);
+		return bShouldDraw;
 	}
 	else
 	{
 		// opaque materials.
-		if (!bIsTranslucent && !bUsesWaterMaterial)
-		{
-			const FLightSceneInfo* MobileDirectionalLight = MobileBasePass::GetDirectionalLightInfo(Scene, PrimitiveSceneProxy);
-			ELightMapPolicyType LightmapPolicyType = MobileBasePass::SelectMeshLightmapPolicy(Scene, MeshBatch, PrimitiveSceneProxy, MobileDirectionalLight, ShadingModels, bCanReceiveCSM, bUsesDeferredShading, bIsTranslucent, FeatureLevel);
-			bResult = Process(MeshBatch, BatchElementMask, StaticMeshId, PrimitiveSceneProxy, MaterialRenderProxy, Material, bIsMasked, bIsTranslucent, ShadingModels, LightmapPolicyType, bCanReceiveCSM, MeshBatch.LCI);
-		}
+		return !bIsTranslucent;
 	}
+}
 
-	return bResult;
+bool FMobileBasePassMeshProcessor::TryAddMeshBatch(const FMeshBatch& RESTRICT MeshBatch, uint64 BatchElementMask, const FPrimitiveSceneProxy* RESTRICT PrimitiveSceneProxy, int32 StaticMeshId, const FMaterialRenderProxy& MaterialRenderProxy, const FMaterial& Material)
+{
+	if (ShouldDraw(Material))
+	{
+		const FMaterialShadingModelField ShadingModels = Material.GetShadingModels();
+		const bool bCanReceiveCSM = ((Flags & EFlags::CanReceiveCSM) == EFlags::CanReceiveCSM);
+		const EBlendMode BlendMode = Material.GetBlendMode();
+		const bool bIsLitMaterial = ShadingModels.IsLit();
+		const bool bIsTranslucent = IsTranslucentBlendMode(BlendMode) || ShadingModels.HasShadingModel(MSM_SingleLayerWater); // Water goes into the translucent pass;
+		const bool bIsMasked = IsMaskedBlendMode(Material);
+		const FLightSceneInfo* MobileDirectionalLight = MobileBasePass::GetDirectionalLightInfo(Scene, PrimitiveSceneProxy);
+		ELightMapPolicyType LightmapPolicyType = MobileBasePass::SelectMeshLightmapPolicy(Scene, MeshBatch, PrimitiveSceneProxy, MobileDirectionalLight, bCanReceiveCSM, bPassUsesDeferredShading, bIsLitMaterial, bIsTranslucent);
+		return Process(MeshBatch, BatchElementMask, StaticMeshId, PrimitiveSceneProxy, MaterialRenderProxy, Material, bIsMasked, bIsTranslucent, ShadingModels, LightmapPolicyType, bCanReceiveCSM, MeshBatch.LCI);
+	}
+	return true;
 }
 
 void FMobileBasePassMeshProcessor::AddMeshBatch(const FMeshBatch& RESTRICT MeshBatch, uint64 BatchElementMask, const FPrimitiveSceneProxy* RESTRICT PrimitiveSceneProxy, int32 StaticMeshId)
@@ -725,7 +812,7 @@ bool FMobileBasePassMeshProcessor::Process(
 	if (Scene && PrimitiveSceneProxy && ShadingModels.IsLit())
 	{
 		// Whether material with a forward shading needs clustered lights
-		if (!bUsesDeferredShading && MobileForwardEnableLocalLights(Scene->GetShaderPlatform()))
+		if (!bPassUsesDeferredShading && MobileForwardEnableLocalLights(Scene->GetShaderPlatform()))
 		{
 			bEnableLocalLights = PrimitiveSceneProxy->GetPrimitiveSceneInfo()->NumMobileDynamicLocalLights > 0;
 		}
@@ -760,7 +847,7 @@ bool FMobileBasePassMeshProcessor::Process(
 		else
 		{
 			const bool bEnableReceiveDecalOutput = ((Flags & EFlags::CanUseDepthStencil) == EFlags::CanUseDepthStencil);
-			MobileBasePass::SetOpaqueRenderState(DrawRenderState, PrimitiveSceneProxy, MaterialResource, ShadingModels, bEnableReceiveDecalOutput && IsMobileHDR(), bUsesDeferredShading);
+			MobileBasePass::SetOpaqueRenderState(DrawRenderState, PrimitiveSceneProxy, MaterialResource, ShadingModels, bEnableReceiveDecalOutput && IsMobileHDR(), bPassUsesDeferredShading);
 		}
 	}
 
@@ -807,29 +894,170 @@ bool FMobileBasePassMeshProcessor::Process(
 	return true;
 }
 
+void FMobileBasePassMeshProcessor::CollectPSOInitializersForLMPolicy(
+	const FPSOPrecacheVertexFactoryData& VertexFactoryData,
+	const FMeshPassProcessorRenderState& RESTRICT DrawRenderState,
+	const FGraphicsPipelineRenderTargetsInfo& RESTRICT RenderTargetsInfo,
+	const FMaterial& RESTRICT MaterialResource,
+	const bool bEnableSkyLight,
+	const bool bEnableLocalLights,
+	const ELightMapPolicyType LightMapPolicyType,
+	ERasterizerFillMode MeshFillMode,
+	ERasterizerCullMode MeshCullMode,
+	EPrimitiveType PrimitiveType,
+	TArray<FPSOPrecacheData>& PSOInitializers)
+{
+	TMeshProcessorShaders<
+		TMobileBasePassVSPolicyParamType<FUniformLightMapPolicy>,
+		TMobileBasePassPSPolicyParamType<FUniformLightMapPolicy>> BasePassShaders;
+
+	if (!MobileBasePass::GetShaders(
+		LightMapPolicyType,
+		bEnableLocalLights,
+		MaterialResource,
+		VertexFactoryData.VertexFactoryType,
+		bEnableSkyLight,
+		BasePassShaders.VertexShader,
+		BasePassShaders.PixelShader))
+	{
+		return;
+	}
+
+	// subpass info set during the submission of the draws in mobile deferred renderer.
+	uint8 SubpassIndex = bTranslucentBasePass ? (bDeferredShading ? 2 : 1) : 0;
+	ESubpassHint SubpassHint = bDeferredShading ? ESubpassHint::DeferredShadingSubpass : ESubpassHint::DepthReadSubpass;
+
+	AddGraphicsPipelineStateInitializer(
+		VertexFactoryData,
+		MaterialResource,
+		DrawRenderState,
+		RenderTargetsInfo,
+		BasePassShaders,
+		MeshFillMode,
+		MeshCullMode,
+		PrimitiveType,
+		EMeshPassFeatures::Default,
+		SubpassHint,
+		SubpassIndex,
+		true /*bRequired*/,
+		PSOInitializers);
+}
+
+void FMobileBasePassMeshProcessor::CollectPSOInitializers(const FSceneTexturesConfig& SceneTexturesConfig, const FMaterial& Material, const FPSOPrecacheVertexFactoryData& VertexFactoryData, const FPSOPrecacheParams& PreCacheParams, TArray<FPSOPrecacheData>& PSOInitializers)
+{
+	static IConsoleVariable* CVarPSOPrecacheTranslucencyAllPass = IConsoleManager::Get().FindConsoleVariable(TEXT("r.PSOPrecache.TranslucencyAllPass"));
+	// PSO precaching enabled for TranslucencyAll
+	if (MeshPassType == EMeshPass::TranslucencyAll && CVarPSOPrecacheTranslucencyAllPass->GetInt() == 0)
+	{
+		return;
+	}
+	
+	// Check if material should be rendered
+	if (!PreCacheParams.bRenderInMainPass || !ShouldDraw(Material))
+	{
+		return;
+	}
+
+	// Determine the mesh's material and blend mode.
+	EShaderPlatform ShaderPlatform = GetFeatureLevelShaderPlatform(Material.GetFeatureLevel());
+	const FMeshDrawingPolicyOverrideSettings OverrideSettings = ComputeMeshOverrideSettings(PreCacheParams);
+	const ERasterizerFillMode MeshFillMode = ComputeMeshFillMode(Material, OverrideSettings);
+	ERasterizerCullMode MeshCullMode = ComputeMeshCullMode(Material, OverrideSettings);
+	const EBlendMode BlendMode = Material.GetBlendMode();
+	const FMaterialShadingModelField ShadingModels = Material.GetShadingModels();
+	const bool bLitMaterial = ShadingModels.IsLit();
+
+	bool bMovable = PreCacheParams.Mobility == EComponentMobility::Movable || PreCacheParams.Mobility == EComponentMobility::Stationary;
+	bool bDitheredLODTransition = !bMovable && Material.IsDitheredLODTransition() && !PreCacheParams.bForceLODModel;
+
+	// Setup the draw state
+	FGraphicsPipelineRenderTargetsInfo RenderTargetsInfo;
+	RenderTargetsInfo.NumSamples = SceneTexturesConfig.NumSamples;
+
+	FMeshPassProcessorRenderState DrawRenderState(PassDrawRenderState);
+	EPixelFormat SceneColorFormat = SceneTexturesConfig.ColorFormat;
+	ETextureCreateFlags SceneColorCreateFlags = SceneTexturesConfig.ColorCreateFlags;
+
+	const bool bMaskedInEarlyPass = MaskedInEarlyPass(ShaderPlatform);
+
+	FExclusiveDepthStencil ExclusiveDepthStencil = (bTranslucentBasePass || bMaskedInEarlyPass) ? 
+		FExclusiveDepthStencil::DepthRead_StencilRead : 
+		FExclusiveDepthStencil::DepthWrite_StencilWrite;
+
+	SetupGBufferRenderTargetInfo(SceneTexturesConfig, RenderTargetsInfo, false /*bSetupDepthStencil*/);
+	SetupDepthStencilInfo(PF_DepthStencil, SceneTexturesConfig.DepthCreateFlags, ERenderTargetLoadAction::ELoad,
+		ERenderTargetLoadAction::ELoad, ExclusiveDepthStencil, RenderTargetsInfo);
+
+	if (bTranslucentBasePass)
+	{
+		MobileBasePass::SetTranslucentRenderState(DrawRenderState, Material, ShadingModels);
+	}
+	//else if((MeshBatch.bUseForDepthPass && Scene->EarlyZPassMode == DDM_AllOpaque) || bMaskedInEarlyPass)
+	//{
+	//	DrawRenderState.SetDepthStencilState(TStaticDepthStencilState<false, CF_Equal>::GetRHI());
+	//}
+	else
+	{
+		MobileBasePass::SetOpaqueRenderState(DrawRenderState, nullptr, Material, ShadingModels, true, bPassUsesDeferredShading);
+	}
+
+	const bool bUseLocalLightPermutation = bLitMaterial && (!bPassUsesDeferredShading && MobileForwardEnableLocalLights(ShaderPlatform));
+	const bool bCanReceiveCSM = ((Flags & FMobileBasePassMeshProcessor::EFlags::CanReceiveCSM) == FMobileBasePassMeshProcessor::EFlags::CanReceiveCSM);
+
+	FMobileLightMapPolicyTypeList UniformLightMapPolicyTypes = GetUniformLightMapPolicyTypeForPSOCollection(bLitMaterial, bTranslucentBasePass, bPassUsesDeferredShading, bCanReceiveCSM, bMovable);
+	
+	for (ELightMapPolicyType LightMapPolicyType : UniformLightMapPolicyTypes)
+	{
+		// SkyLight OFF
+		bool bEnableSkyLight = false;
+		if (MobileBasePass::UseSkylightPermutation(bEnableSkyLight, FReadOnlyCVARCache::Get().MobileSkyLightPermutation) || !bLitMaterial)
+		{
+			CollectPSOInitializersForLMPolicy(VertexFactoryData, DrawRenderState, RenderTargetsInfo, Material, bEnableSkyLight, false /*bEnableLocalLights*/, LightMapPolicyType, MeshFillMode, MeshCullMode, (EPrimitiveType)PreCacheParams.PrimitiveType, PSOInitializers);
+			if (bUseLocalLightPermutation)
+			{
+				CollectPSOInitializersForLMPolicy(VertexFactoryData, DrawRenderState, RenderTargetsInfo, Material, bEnableSkyLight, true /*bEnableLocalLights*/, LightMapPolicyType, MeshFillMode, MeshCullMode, (EPrimitiveType)PreCacheParams.PrimitiveType, PSOInitializers);
+			}
+		}
+
+		// SkyLight ON
+		bEnableSkyLight = true;
+		if (MobileBasePass::UseSkylightPermutation(bEnableSkyLight, FReadOnlyCVARCache::Get().MobileSkyLightPermutation) && bLitMaterial)
+		{
+			CollectPSOInitializersForLMPolicy(VertexFactoryData, DrawRenderState, RenderTargetsInfo, Material, bEnableSkyLight, false /*bEnableLocalLights*/, LightMapPolicyType, MeshFillMode, MeshCullMode, (EPrimitiveType)PreCacheParams.PrimitiveType, PSOInitializers);
+			if (bUseLocalLightPermutation)
+			{
+				CollectPSOInitializersForLMPolicy(VertexFactoryData, DrawRenderState, RenderTargetsInfo, Material, bEnableSkyLight, true /*bEnableLocalLights*/, LightMapPolicyType, MeshFillMode, MeshCullMode, (EPrimitiveType)PreCacheParams.PrimitiveType, PSOInitializers);
+			}
+		}
+	}
+}
+
 FMeshPassProcessor* CreateMobileBasePassProcessor(ERHIFeatureLevel::Type FeatureLevel, const FScene* Scene, const FSceneView* InViewIfDynamicMeshCommand, FMeshPassDrawListContext* InDrawListContext)
 {
 	FMeshPassProcessorRenderState PassDrawRenderState;
 	PassDrawRenderState.SetBlendState(TStaticBlendStateWriteMask<CW_RGBA>::GetRHI());
-	PassDrawRenderState.SetDepthStencilAccess(Scene->DefaultBasePassDepthStencilAccess);
+	const FExclusiveDepthStencil::Type DefaultBasePassDepthStencilAccess = FScene::GetDefaultBasePassDepthStencilAccess(FeatureLevel);
+	PassDrawRenderState.SetDepthStencilAccess(DefaultBasePassDepthStencilAccess);
 	PassDrawRenderState.SetDepthStencilState(TStaticDepthStencilState<true, CF_DepthNearOrEqual>::GetRHI());
 
 	const FMobileBasePassMeshProcessor::EFlags Flags = FMobileBasePassMeshProcessor::EFlags::CanUseDepthStencil
-		| (MobileBasePassAlwaysUsesCSM(Scene->GetShaderPlatform()) ? FMobileBasePassMeshProcessor::EFlags::CanReceiveCSM : FMobileBasePassMeshProcessor::EFlags::None);
+		| (MobileBasePassAlwaysUsesCSM(GShaderPlatformForFeatureLevel[FeatureLevel]) ? FMobileBasePassMeshProcessor::EFlags::CanReceiveCSM : FMobileBasePassMeshProcessor::EFlags::None);
 
 	return new FMobileBasePassMeshProcessor(EMeshPass::BasePass, Scene, FeatureLevel, InViewIfDynamicMeshCommand, PassDrawRenderState, InDrawListContext, Flags);
 }
 
 FMeshPassProcessor* CreateMobileBasePassCSMProcessor(ERHIFeatureLevel::Type FeatureLevel, const FScene* Scene, const FSceneView* InViewIfDynamicMeshCommand, FMeshPassDrawListContext* InDrawListContext)
 {
+	const FExclusiveDepthStencil::Type DefaultBasePassDepthStencilAccess = FScene::GetDefaultBasePassDepthStencilAccess(FeatureLevel);
+
 	FMeshPassProcessorRenderState PassDrawRenderState;
 	PassDrawRenderState.SetBlendState(TStaticBlendStateWriteMask<CW_RGBA>::GetRHI());
-	PassDrawRenderState.SetDepthStencilAccess(Scene->DefaultBasePassDepthStencilAccess);
+	PassDrawRenderState.SetDepthStencilAccess(DefaultBasePassDepthStencilAccess);
 	PassDrawRenderState.SetDepthStencilState(TStaticDepthStencilState<true, CF_DepthNearOrEqual>::GetRHI());
 
 	// By default this processor will not cache anything. Only enable it when CSM culling is active
 	FMobileBasePassMeshProcessor::EFlags Flags = FMobileBasePassMeshProcessor::EFlags::DoNotCache;
-	if (!MobileBasePassAlwaysUsesCSM(Scene->GetShaderPlatform()))
+	if (!MobileBasePassAlwaysUsesCSM(GShaderPlatformForFeatureLevel[FeatureLevel]))
 	{
 		Flags = FMobileBasePassMeshProcessor::EFlags::CanReceiveCSM | FMobileBasePassMeshProcessor::EFlags::CanUseDepthStencil;
 	}
@@ -871,9 +1099,9 @@ FMeshPassProcessor* CreateMobileTranslucencyAllPassProcessor(ERHIFeatureLevel::T
 	return new FMobileBasePassMeshProcessor(EMeshPass::TranslucencyAll, Scene, FeatureLevel, InViewIfDynamicMeshCommand, PassDrawRenderState, InDrawListContext, Flags, ETranslucencyPass::TPT_AllTranslucency);
 }
 
-FRegisterPassProcessorCreateFunction RegisterMobileBasePass(&CreateMobileBasePassProcessor, EShadingPath::Mobile, EMeshPass::BasePass, EMeshPassFlags::CachedMeshCommands | EMeshPassFlags::MainView);
-FRegisterPassProcessorCreateFunction RegisterMobileBasePassCSM(&CreateMobileBasePassCSMProcessor, EShadingPath::Mobile, EMeshPass::MobileBasePassCSM, EMeshPassFlags::CachedMeshCommands | EMeshPassFlags::MainView);
-FRegisterPassProcessorCreateFunction RegisterMobileTranslucencyAllPass(&CreateMobileTranslucencyAllPassProcessor, EShadingPath::Mobile, EMeshPass::TranslucencyAll, EMeshPassFlags::CachedMeshCommands | EMeshPassFlags::MainView);
-FRegisterPassProcessorCreateFunction RegisterMobileTranslucencyStandardPass(&CreateMobileTranslucencyStandardPassProcessor, EShadingPath::Mobile, EMeshPass::TranslucencyStandard, EMeshPassFlags::CachedMeshCommands | EMeshPassFlags::MainView);
-FRegisterPassProcessorCreateFunction RegisterMobileTranslucencyAfterDOFPass(&CreateMobileTranslucencyAfterDOFProcessor, EShadingPath::Mobile, EMeshPass::TranslucencyAfterDOF, EMeshPassFlags::CachedMeshCommands | EMeshPassFlags::MainView);
+REGISTER_MESHPASSPROCESSOR_AND_PSOCOLLECTOR(MobileBasePass, 			CreateMobileBasePassProcessor, 			EShadingPath::Mobile, EMeshPass::BasePass, 		EMeshPassFlags::CachedMeshCommands | EMeshPassFlags::MainView);
+REGISTER_MESHPASSPROCESSOR_AND_PSOCOLLECTOR(MobileBasePassCSM,			CreateMobileBasePassCSMProcessor,		EShadingPath::Mobile, EMeshPass::MobileBasePassCSM, 	EMeshPassFlags::CachedMeshCommands | EMeshPassFlags::MainView);
+REGISTER_MESHPASSPROCESSOR_AND_PSOCOLLECTOR(MobileTranslucencyAllPass,		CreateMobileTranslucencyAllPassProcessor,	EShadingPath::Mobile, EMeshPass::TranslucencyAll, 	EMeshPassFlags::CachedMeshCommands | EMeshPassFlags::MainView);
+REGISTER_MESHPASSPROCESSOR_AND_PSOCOLLECTOR(MobileTranslucencyStandardPass,	CreateMobileTranslucencyStandardPassProcessor,	EShadingPath::Mobile, EMeshPass::TranslucencyStandard, 	EMeshPassFlags::CachedMeshCommands | EMeshPassFlags::MainView);
+REGISTER_MESHPASSPROCESSOR_AND_PSOCOLLECTOR(MobileTranslucencyAfterDOFPass,	CreateMobileTranslucencyAfterDOFProcessor,	EShadingPath::Mobile, EMeshPass::TranslucencyAfterDOF, 	EMeshPassFlags::CachedMeshCommands | EMeshPassFlags::MainView);
 // Skipping EMeshPass::TranslucencyAfterDOFModulate because dual blending is not supported on mobile
diff --git a/Engine/Source/Runtime/Renderer/Private/MobileBasePassRendering.h b/Engine/Source/Runtime/Renderer/Private/MobileBasePassRendering.h
--- a/Engine/Source/Runtime/Renderer/Private/MobileBasePassRendering.h
+++ b/Engine/Source/Runtime/Renderer/Private/MobileBasePassRendering.h
@@ -312,17 +312,16 @@ namespace MobileBasePass
 		const FMeshBatch& MeshBatch, 
 		const FPrimitiveSceneProxy* PrimitiveSceneProxy, 
 		const FLightSceneInfo* MobileDirectionalLight, 
-		FMaterialShadingModelField ShadingModels, 
 		bool bPrimReceivesCSM, 
 		bool bUsedDeferredShading,
-		bool bIsTranslucent,
-		ERHIFeatureLevel::Type FeatureLevel);
+		bool bIsLitMaterial,
+		bool bIsTranslucent);
 
 	bool GetShaders(
 		ELightMapPolicyType LightMapPolicyType,
 		bool bEnableLocalLights,
 		const FMaterial& MaterialResource,
-		FVertexFactoryType* VertexFactoryType,
+		const FVertexFactoryType* VertexFactoryType,
 		bool bEnableSkyLight, 
 		TShaderRef<TMobileBasePassVSPolicyParamType<FUniformLightMapPolicy>>& VertexShader,
 		TShaderRef<TMobileBasePassPSPolicyParamType<FUniformLightMapPolicy>>& PixelShader);
@@ -333,20 +332,19 @@ namespace MobileBasePass
 
 	void SetOpaqueRenderState(FMeshPassProcessorRenderState& DrawRenderState, const FPrimitiveSceneProxy* PrimitiveSceneProxy, const FMaterial& Material, FMaterialShadingModelField ShadingModels, bool bEnableReceiveDecalOutput, bool bUsesDeferredShading);
 	void SetTranslucentRenderState(FMeshPassProcessorRenderState& DrawRenderState, const FMaterial& Material, FMaterialShadingModelField ShadingModels);
-};
-
 
-inline bool UseSkylightPermutation(bool bEnableSkyLight, int32 MobileSkyLightPermutationOptions)
-{
-	if (bEnableSkyLight)
+	inline bool UseSkylightPermutation(bool bEnableSkyLight, int32 MobileSkyLightPermutationOptions)
 	{
-		return MobileSkyLightPermutationOptions == 0 || MobileSkyLightPermutationOptions == 2;
-	}
-	else
-	{
-		return MobileSkyLightPermutationOptions == 0 || MobileSkyLightPermutationOptions == 1;
+		if (bEnableSkyLight)
+		{
+			return MobileSkyLightPermutationOptions == 0 || MobileSkyLightPermutationOptions == 2;
+		}
+		else
+		{
+			return MobileSkyLightPermutationOptions == 0 || MobileSkyLightPermutationOptions == 1;
+		}
 	}
-}
+};
 
 template< typename LightMapPolicyType, EOutputFormat OutputFormat, bool bEnableSkyLight, bool bEnableLocalLights>
 class TMobileBasePassPS : public TMobileBasePassPSBaseType<LightMapPolicyType>
@@ -363,7 +361,7 @@ public:
 		// Only compile skylight version for lit materials
 		const bool bShouldCacheBySkylight = !bEnableSkyLight || bIsLit;
 		// Only compile skylight permutations when they are enabled
-		if (bIsLit && !UseSkylightPermutation(bEnableSkyLight, MobileSkyLightPermutationOptions))
+		if (bIsLit && !MobileBasePass::UseSkylightPermutation(bEnableSkyLight, MobileSkyLightPermutationOptions))
 		{
 			return false;
 		}
@@ -447,9 +445,24 @@ public:
 	virtual void AddMeshBatch(const FMeshBatch& RESTRICT MeshBatch, uint64 BatchElementMask, const FPrimitiveSceneProxy* RESTRICT PrimitiveSceneProxy, int32 StaticMeshId = -1) override final;
 
 	FMeshPassProcessorRenderState PassDrawRenderState;
+	virtual void CollectPSOInitializers(const FSceneTexturesConfig& SceneTexturesConfig, const FMaterial& Material, const FPSOPrecacheVertexFactoryData& VertexFactoryData, const FPSOPrecacheParams& PreCacheParams, TArray<FPSOPrecacheData>& PSOInitializers) override;
+
+	void CollectPSOInitializersForLMPolicy(
+		const FPSOPrecacheVertexFactoryData& VertexFactoryData,
+		const FMeshPassProcessorRenderState& RESTRICT DrawRenderState,
+		const FGraphicsPipelineRenderTargetsInfo& RESTRICT RenderTargetsInfo,
+		const FMaterial& RESTRICT MaterialResource,
+		const bool bRenderSkylight,
+		const bool bEnableLocalLights,
+		const ELightMapPolicyType LightMapPolicyType,
+		ERasterizerFillMode MeshFillMode,
+		ERasterizerCullMode MeshCullMode,
+		EPrimitiveType PrimitiveType,
+		TArray<FPSOPrecacheData>& PSOInitializers);
 
 private:
 	FMaterialShadingModelField FilterShadingModelsMask(const FMaterialShadingModelField& ShadingModels) const;
+	bool ShouldDraw(const class FMaterial& Material) const;
 
 	bool TryAddMeshBatch(const FMeshBatch& RESTRICT MeshBatch, uint64 BatchElementMask, const FPrimitiveSceneProxy* RESTRICT PrimitiveSceneProxy, int32 StaticMeshId, const FMaterialRenderProxy& MaterialRenderProxy, const FMaterial& Material);
 
@@ -470,7 +483,10 @@ private:
 	const ETranslucencyPass::Type TranslucencyPassType;
 	const EFlags Flags;
 	const bool bTranslucentBasePass;
-	const bool bUsesDeferredShading;
+	// Whether renderer uses deferred shading
+	const bool bDeferredShading; 
+	// Whether this pass uses deferred shading
+	const bool bPassUsesDeferredShading; 
 };
 
 ENUM_CLASS_FLAGS(FMobileBasePassMeshProcessor::EFlags);
diff --git a/Engine/Source/Runtime/Renderer/Private/MobileShadingRenderer.cpp b/Engine/Source/Runtime/Renderer/Private/MobileShadingRenderer.cpp
--- a/Engine/Source/Runtime/Renderer/Private/MobileShadingRenderer.cpp
+++ b/Engine/Source/Runtime/Renderer/Private/MobileShadingRenderer.cpp
@@ -219,40 +219,6 @@ static bool PostProcessUsesSceneDepth(const FViewInfo& View)
 	return false;
 }
 
-static void SetupGBufferFlags(FSceneTexturesConfig& SceneTexturesConfig, bool bRequiresMultiPass)
-{
-	ETextureCreateFlags AddFlags = TexCreate_InputAttachmentRead;
-	if (!bRequiresMultiPass)
-	{
-		// use memoryless GBuffer when possible
-		AddFlags |= TexCreate_Memoryless;
-	}
-
-	for (FGBufferBindings& Bindings : SceneTexturesConfig.GBufferBindings)
-	{
-		Bindings.GBufferA.Flags |= AddFlags;
-		Bindings.GBufferB.Flags |= AddFlags;
-		Bindings.GBufferC.Flags |= AddFlags;
-		Bindings.GBufferD.Flags |= AddFlags;
-		Bindings.GBufferE.Flags |= AddFlags;
-
-		// Mobile uses FBF/subpassLoad to fetch data from GBuffer, and FBF does not always work with sRGB targets 
-		Bindings.GBufferA.Flags &= (~TexCreate_SRGB);
-		Bindings.GBufferB.Flags &= (~TexCreate_SRGB);
-		Bindings.GBufferC.Flags &= (~TexCreate_SRGB);
-		Bindings.GBufferD.Flags &= (~TexCreate_SRGB);
-		Bindings.GBufferE.Flags &= (~TexCreate_SRGB);
-
-		// Input attachments with PF_R8G8B8A8 has better support on mobile than PF_B8G8R8A8
-		auto OverrideB8G8R8A8 = [](FGBufferBinding& Binding) { if (Binding.Format == PF_B8G8R8A8) Binding.Format = PF_R8G8B8A8; };
-		OverrideB8G8R8A8(Bindings.GBufferA);
-		OverrideB8G8R8A8(Bindings.GBufferB);
-		OverrideB8G8R8A8(Bindings.GBufferC);
-		OverrideB8G8R8A8(Bindings.GBufferD);
-		OverrideB8G8R8A8(Bindings.GBufferE);
-	}
-}
-
 static void PollOcclusionQueriesPass(FRDGBuilder& GraphBuilder)
 {
 	AddPass(GraphBuilder, RDG_EVENT_NAME("PollOcclusionQueries"), [](FRHICommandListImmediate& RHICmdList)
@@ -557,7 +523,8 @@ void FMobileSceneRenderer::InitViews(
 	const bool bSeparateTranslucencyActive = IsMobileSeparateTranslucencyActive(Views.GetData(), Views.Num()); 
 	const bool bPostProcessUsesSceneDepth = PostProcessUsesSceneDepth(Views[0]) || IsMobileDistortionActive(Views[0]);
 	const bool bRequireSeparateViewPass = Views.Num() > 1 && !Views[0].bIsMobileMultiViewEnabled;
-	bRequiresMultiPass = RequiresMultiPass(Views[0]);
+	bRequiresMultiPass = RequiresMultiPass(NumMSAASamples, ShaderPlatform);
+
 	bKeepDepthContent =
 		bRequiresMultiPass ||
 		bForceDepthResolve ||
@@ -590,12 +557,10 @@ void FMobileSceneRenderer::InitViews(
 	// If we render in a single pass MSAA targets can be memoryless
     SceneTexturesConfig.bMemorylessMSAA = !(bRequiresMultiPass || bShouldCompositeEditorPrimitives || bRequireSeparateViewPass);
     SceneTexturesConfig.NumSamples = NumMSAASamples;
-    
     SceneTexturesConfig.BuildSceneColorAndDepthFlags();
-    
 	if (bDeferredShading) 
 	{
-		SetupGBufferFlags(SceneTexturesConfig, bRequiresMultiPass || GraphBuilder.IsDumpingFrame() || bRequireSeparateViewPass);
+		SceneTexturesConfig.SetupMobileGBufferFlags(bRequiresMultiPass || GraphBuilder.IsDumpingFrame() || bRequireSeparateViewPass);
 	}
 
 	// Update the pixel projected reflection extent according to the settings in the PlanarReflectionComponent.
@@ -1792,7 +1757,7 @@ int32 FMobileSceneRenderer::ComputeNumOcclusionQueriesToBatch() const
 }
 
 // Whether we need a separate render-passes for translucency, decals etc
-bool FMobileSceneRenderer::RequiresMultiPass(const FViewInfo& View) const
+bool FMobileSceneRenderer::RequiresMultiPass(int32 NumMSAASamples, EShaderPlatform ShaderPlatform)
 {
 	// Vulkan uses subpasses
 	if (IsVulkanPlatform(ShaderPlatform))
diff --git a/Engine/Source/Runtime/Renderer/Private/SceneRendering.h b/Engine/Source/Runtime/Renderer/Private/SceneRendering.h
--- a/Engine/Source/Runtime/Renderer/Private/SceneRendering.h
+++ b/Engine/Source/Runtime/Renderer/Private/SceneRendering.h
@@ -2672,6 +2672,9 @@ public:
 
 	virtual bool AllowSimpleLights() const override;
 
+	/** Whether platform requires multiple render-passes for SceneColor rendering */
+	static bool RequiresMultiPass(int32 NumMSAASamples, EShaderPlatform ShaderPlatform);
+
 protected:
 	/** Finds the visible dynamic shadows for each view. */
 	FDynamicShadowsTaskData* InitDynamicShadows(FRDGBuilder& GraphBuilder, FInstanceCullingManager& FInstanceCullingManager, FRDGExternalAccessQueue& ExternalAccessQueue);
@@ -2728,8 +2731,6 @@ protected:
 	/** Computes how many queries will be issued this frame */
 	int32 ComputeNumOcclusionQueriesToBatch() const;
 
-	/** Whether platform requires multiple render-passes for SceneColor rendering */
-	bool RequiresMultiPass(const FViewInfo& View) const;
 
 	/** Renders decals. */
 	void RenderDecals(FRHICommandList& RHICmdList, const FViewInfo& View);
diff --git a/Engine/Source/Runtime/Renderer/Private/ShadowDepthRendering.cpp b/Engine/Source/Runtime/Renderer/Private/ShadowDepthRendering.cpp
--- a/Engine/Source/Runtime/Renderer/Private/ShadowDepthRendering.cpp
+++ b/Engine/Source/Runtime/Renderer/Private/ShadowDepthRendering.cpp
@@ -2289,4 +2289,4 @@ FMeshPassProcessor* CreateVSMShadowDepthPassProcessor(ERHIFeatureLevel::Type Fea
 }
 
 REGISTER_MESHPASSPROCESSOR_AND_PSOCOLLECTOR(VSMShadowDepthPass, CreateVSMShadowDepthPassProcessor, EShadingPath::Deferred, EMeshPass::VSMShadowDepth, EMeshPassFlags::CachedMeshCommands);
-FRegisterPassProcessorCreateFunction RegisterMobileCSMShadowDepthPass(&CreateCSMShadowDepthPassProcessor, EShadingPath::Mobile, EMeshPass::CSMShadowDepth, EMeshPassFlags::CachedMeshCommands);
+REGISTER_MESHPASSPROCESSOR_AND_PSOCOLLECTOR(CSMMobileShadowDepthPass, CreateCSMShadowDepthPassProcessor, EShadingPath::Mobile, EMeshPass::CSMShadowDepth, EMeshPassFlags::CachedMeshCommands);
diff --git a/Engine/Source/Runtime/Renderer/Public/MeshPassProcessor.h b/Engine/Source/Runtime/Renderer/Public/MeshPassProcessor.h
--- a/Engine/Source/Runtime/Renderer/Public/MeshPassProcessor.h
+++ b/Engine/Source/Runtime/Renderer/Public/MeshPassProcessor.h
@@ -1978,6 +1978,22 @@ public:
 		bool bRequired,
 		TArray<FPSOPrecacheData>& PSOInitializers);
 
+	template<typename PassShadersType>
+	void AddGraphicsPipelineStateInitializer(
+		const FPSOPrecacheVertexFactoryData& VertexFactoryData,
+		const FMaterial& RESTRICT MaterialResource,
+		const FMeshPassProcessorRenderState& RESTRICT DrawRenderState,
+		const FGraphicsPipelineRenderTargetsInfo& RESTRICT RenderTargetsInfo,
+		PassShadersType PassShaders,
+		ERasterizerFillMode MeshFillMode,
+		ERasterizerCullMode MeshCullMode,
+		EPrimitiveType PrimitiveType,
+		EMeshPassFeatures MeshPassFeatures,
+		ESubpassHint SubpassHint,
+		uint8 SubpassIndex,
+		bool bRequired,
+		TArray<FPSOPrecacheData>& PSOInitializers);
+
 protected:
 	RENDERER_API FMeshDrawCommandPrimitiveIdInfo GetDrawCommandPrimitiveId(
 		const FPrimitiveSceneInfo* RESTRICT PrimitiveSceneInfo,
diff --git a/Engine/Source/Runtime/Renderer/Public/MeshPassProcessor.inl b/Engine/Source/Runtime/Renderer/Public/MeshPassProcessor.inl
--- a/Engine/Source/Runtime/Renderer/Public/MeshPassProcessor.inl
+++ b/Engine/Source/Runtime/Renderer/Public/MeshPassProcessor.inl
@@ -190,6 +190,36 @@ void FMeshPassProcessor::BuildMeshDrawCommands(
 	}
 }
 
+template<typename PassShadersType>
+void FMeshPassProcessor::AddGraphicsPipelineStateInitializer(
+	const FPSOPrecacheVertexFactoryData& VertexFactoryData,
+	const FMaterial& RESTRICT MaterialResource,
+	const FMeshPassProcessorRenderState& RESTRICT DrawRenderState,
+	const FGraphicsPipelineRenderTargetsInfo& RESTRICT RenderTargetsInfo,
+	PassShadersType PassShaders,
+	ERasterizerFillMode MeshFillMode,
+	ERasterizerCullMode MeshCullMode,
+	EPrimitiveType PrimitiveType,
+	EMeshPassFeatures MeshPassFeatures,
+	bool bRequired,
+	TArray<FPSOPrecacheData>& PSOInitializers)
+{
+	AddGraphicsPipelineStateInitializer(
+		VertexFactoryData,
+		MaterialResource,
+		DrawRenderState,
+		RenderTargetsInfo,
+		PassShaders,
+		MeshFillMode,
+		MeshCullMode,
+		PrimitiveType,
+		MeshPassFeatures,
+		ESubpassHint::None,
+		0,
+		bRequired,
+		PSOInitializers);
+}
+
 template<typename PassShadersType>
 void FMeshPassProcessor::AddGraphicsPipelineStateInitializer(
 	const FPSOPrecacheVertexFactoryData& VertexFactoryData,
@@ -201,6 +231,8 @@ void FMeshPassProcessor::AddGraphicsPipelineStateInitializer(
 	ERasterizerCullMode MeshCullMode,
 	EPrimitiveType PrimitiveType,
 	EMeshPassFeatures MeshPassFeatures, 
+	ESubpassHint SubpassHint,
+	uint8 SubpassIndex,
 	bool bRequired,
 	TArray<FPSOPrecacheData>& PSOInitializers)
 {
@@ -249,6 +281,8 @@ void FMeshPassProcessor::AddGraphicsPipelineStateInitializer(
 	// NOTE: AsGraphicsPipelineStateInitializer will create the RHIShaders internally if they are not cached yet
 	FGraphicsPipelineStateInitializer PipelineStateInitializer = MinimalPipelineStateInitializer.AsGraphicsPipelineStateInitializer();
 	ApplyTargetsInfo(PipelineStateInitializer, RenderTargetsInfo);
+	PipelineStateInitializer.SubpassHint = SubpassHint;
+	PipelineStateInitializer.SubpassIndex = SubpassIndex;
 
 	FPSOPrecacheData PSOPrecacheData;
 	PSOPrecacheData.bRequired = bRequired;
