PSO precahing: Do not precache depth PSO when depth pre-pass is not used on mobile

diff --git a/Engine/Source/Runtime/Renderer/Private/DepthRendering.cpp b/Engine/Source/Runtime/Renderer/Private/DepthRendering.cpp
--- a/Engine/Source/Runtime/Renderer/Private/DepthRendering.cpp
+++ b/Engine/Source/Runtime/Renderer/Private/DepthRendering.cpp
@@ -1066,6 +1066,12 @@ void FDepthPassMeshProcessor::CollectPSOInitializers(const FSceneTexturesConfig&
 		return;
 	}
 
+	// Mobile rendering does not use depth prepass by default
+	if (FeatureLevel == ERHIFeatureLevel::ES3_1 && EarlyZPassMode == DDM_None)
+	{
+		return;
+	}
+
 	const bool bIsTranslucent = IsTranslucentBlendMode(Material);
 
 	// Early out if translucent or material shouldn't be used during this pass
