Fixed compile error for a monolithic config

diff --git a/Engine/Source/Runtime/Renderer/Private/MobileBasePass.cpp b/Engine/Source/Runtime/Renderer/Private/MobileBasePass.cpp
--- a/Engine/Source/Runtime/Renderer/Private/MobileBasePass.cpp
+++ b/Engine/Source/Runtime/Renderer/Private/MobileBasePass.cpp
@@ -945,9 +945,9 @@ void FMobileBasePassMeshProcessor::CollectPSOInitializersForLMPolicy(
 
 void FMobileBasePassMeshProcessor::CollectPSOInitializers(const FSceneTexturesConfig& SceneTexturesConfig, const FMaterial& Material, const FPSOPrecacheVertexFactoryData& VertexFactoryData, const FPSOPrecacheParams& PreCacheParams, TArray<FPSOPrecacheData>& PSOInitializers)
 {
-	static IConsoleVariable* CVarPSOPrecacheTranslucencyAllPass = IConsoleManager::Get().FindConsoleVariable(TEXT("r.PSOPrecache.TranslucencyAllPass"));
+	static IConsoleVariable* PSOPrecacheTranslucencyAllPass = IConsoleManager::Get().FindConsoleVariable(TEXT("r.PSOPrecache.TranslucencyAllPass"));
 	// PSO precaching enabled for TranslucencyAll
-	if (MeshPassType == EMeshPass::TranslucencyAll && CVarPSOPrecacheTranslucencyAllPass->GetInt() == 0)
+	if (MeshPassType == EMeshPass::TranslucencyAll && PSOPrecacheTranslucencyAllPass->GetInt() == 0)
 	{
 		return;
 	}
