Make r.RayTracing.Shadows.EnableMaterials render thread safe and set ray gen shader when it's disabled.

diff --git a/Engine/Source/Runtime/Renderer/Private/RayTracing/RayTracingShadows.cpp b/Engine/Source/Runtime/Renderer/Private/RayTracing/RayTracingShadows.cpp
--- a/Engine/Source/Runtime/Renderer/Private/RayTracing/RayTracingShadows.cpp
+++ b/Engine/Source/Runtime/Renderer/Private/RayTracing/RayTracingShadows.cpp
@@ -33,7 +33,8 @@ static int32 GRayTracingShadowsEnableMaterials = 1;
 static FAutoConsoleVariableRef CVarRayTracingShadowsEnableMaterials(
 	TEXT("r.RayTracing.Shadows.EnableMaterials"),
 	GRayTracingShadowsEnableMaterials,
-	TEXT("Enables material shader binding for shadow rays. If this is disabled, then a default trivial shader is used. (default = 1)")
+	TEXT("Enables material shader binding for shadow rays. If this is disabled, then a default trivial shader is used. (default = 1)"),
+	ECVF_RenderThreadSafe
 );
 
 static float GRayTracingShadowsAvoidSelfIntersectionTraceDistance = 0.0f;
@@ -506,6 +507,9 @@ void FDeferredShadingSceneRenderer::RenderRayTracingShadows(
 
 						Initializer.MaxPayloadSizeInBytes = GetRayTracingPayloadTypeMaxSize(ERayTracingPayloadType::RayTracingMaterial);
 
+						FRHIRayTracingShader* RayGenShaderTable[] = { RayGenerationShader.GetRayTracingShader() };
+						Initializer.SetRayGenShaderTable(RayGenShaderTable);
+
 						FRHIRayTracingShader* HitGroupTable[] = { GetRayTracingDefaultOpaqueShader(View.ShaderMap) };
 						Initializer.SetHitGroupTable(HitGroupTable);
 						Initializer.bAllowHitGroupIndexing = false; // Use the same hit shader for all geometry in the scene by disabling SBT indexing.
