Fix planar reflection crash on SteamVR.

Avoid drawing the HiddenAreaMesh for planar reflection captures as they're not drawing to the HMD lenses.

diff --git a/Engine/Source/Runtime/Renderer/Private/DepthRendering.cpp b/Engine/Source/Runtime/Renderer/Private/DepthRendering.cpp
--- a/Engine/Source/Runtime/Renderer/Private/DepthRendering.cpp
+++ b/Engine/Source/Runtime/Renderer/Private/DepthRendering.cpp
@@ -693,7 +693,10 @@ void FDeferredShadingSceneRenderer::RenderPrePassHMD(FRDGBuilder& GraphBuilder,
 
 	for (const FViewInfo& View : Views)
 	{
-		if (IStereoRendering::IsStereoEyeView(View))
+		// Don't draw the hidden area mesh in scene captures as they are not displayed
+		// through the HMD lenses.
+		const bool bIsCapture = View.bIsSceneCapture || View.bIsPlanarReflection;
+		if (IStereoRendering::IsStereoEyeView(View) && !bIsCapture)
 		{
 			RDG_GPU_MASK_SCOPE(GraphBuilder, View.GPUMask);
 
