Fixed bug miscounting the deleted instances in the temp cell (leading to IDs not being removed)
- added validation of stored instance IDs.
- added some more detail logging & removed the accumulated string (disabled).

diff --git a/Engine/Source/Runtime/Renderer/Private/SceneCulling/SceneCulling.cpp b/Engine/Source/Runtime/Renderer/Private/SceneCulling/SceneCulling.cpp
--- a/Engine/Source/Runtime/Renderer/Private/SceneCulling/SceneCulling.cpp
+++ b/Engine/Source/Runtime/Renderer/Private/SceneCulling/SceneCulling.cpp
@@ -119,6 +119,12 @@ static TAutoConsoleVariable<int32> CVarSmallFootprintCellCountThreshold(
 	TEXT("Queries with a smaller footprint (in number of cells in the lowest level) go down the footprint based path."), 
 	ECVF_RenderThreadSafe);
 
+static TAutoConsoleVariable<int32> CVarValidateAllInstanceAllocations(
+	TEXT("r.SceneCulling.ValidateAllInstanceAllocations"), 
+	0, 
+	TEXT("Perform validation of all instance IDs stored in the grid. This is very slow."), 
+	ECVF_RenderThreadSafe);
+
 #if SC_ENABLE_GPU_DATA_VALIDATION
 
 static TAutoConsoleVariable<int32> CVarValidateGPUData(
@@ -926,6 +932,7 @@ public:
 			// Copied chunk data (only whole chunks)
 			FTempChunks RetainedChunks; 
 
+			BUILDER_LOG_SCOPE("FinalizeChunks(Index: %d RemovedInstanceCount %d):", CellOffset, RemovedInstanceCount);
 #if OLA_TODO
 			// Handle complete removal case efficiently
 			if (RemovedInstanceCount == TotalCellInstances)
@@ -953,6 +960,7 @@ public:
 
 					if (bCopyExistingChunk)
 					{
+						BUILDER_LOG("Chunk-Retained( NumItems: %d)", NumItems);
 						AddExistingChunk(RetainedChunks, PackedChunkData);
 						continue;
 					}
@@ -963,9 +971,14 @@ public:
 						uint32 FirstInstanceDataOffset = PackedChunkData & INSTANCE_HIERARCHY_ITEM_CHUNK_COMPRESSED_PAYLOAD_MASK;
 						if (!Builder.IsMarkedForRemove(FirstInstanceDataOffset))
 						{
-							NumRemoved += 64;
+							BUILDER_LOG("Chunk-Retained (FirstInstanceDataOffset: %d)", FirstInstanceDataOffset);
 							AddExistingChunk(RetainedChunks, PackedChunkData);
 						}
+						else
+						{
+							BUILDER_LOG("Chunk-Removed (FirstInstanceDataOffset: %d)", FirstInstanceDataOffset);
+							NumRemoved += 64;
+						}
 					}
 					// 2. elsewise, loop over the items in the chunk and copy not-deleted ones.
 					else
@@ -995,15 +1008,21 @@ public:
 						}
 						else
 						{
+							BUILDER_LOG_LIST("Processed(%d):", NumItems);
 							// 3.3., otherwise, we must copy the surviving IDs
 							for (uint32 ItemIndex = 0u; ItemIndex < NumItems; ++ItemIndex)
 							{
 								uint32 InstanceId = Builder.SceneCulling.PackedCellData[ItemDataOffset + ItemIndex];
 								if ((DeletionMask & (1ull << ItemIndex)) == 0ull)
 								{
+									BUILDER_LOG_LIST_APPEND("K: %d", InstanceId);
 									UPDATE_BUILDER_STAT(Builder, CopiedIdCount, 1);
 									Add(Builder, InstanceId);
 								}
+								else
+								{
+									BUILDER_LOG_LIST_APPEND("D: %d", InstanceId);
+								}
 							}
 
 							// Mark the chunk as not in use.
@@ -1573,8 +1592,10 @@ public:
 			int32 CacheIndex = PrimitiveState.Payload;
 			const FCellIndexCacheEntry &CellIndexCacheEntry = GetCacheEntry(CacheIndex);
 
+			BUILDER_LOG_LIST("MarkForRemove(%d):", CellIndexCacheEntry.Items.Num());
 			for (FCellIndexCacheEntry::FItem Item : CellIndexCacheEntry.Items)
 			{
+				BUILDER_LOG_LIST_APPEND("(%d, %d)", Item.CellIndex, Item.NumInstances);
 				MarkCellForRemove(Item.CellIndex, Item.NumInstances);
 			}
 
@@ -1849,7 +1870,7 @@ public:
 
 	void UploadToGPU(FRDGBuilder& GraphBuilder)
 	{
-		BUILDER_LOG("UploadToGPU");
+		BUILDER_LOG("UploadToGPU %d", ItemChunkUploader.GetNumScatters());
 
 		INC_DWORD_STAT_BY(STAT_SceneCulling_UploadedChunks, ItemChunkUploader.GetNumScatters())
 		INC_DWORD_STAT_BY(STAT_SceneCulling_UploadedCells, CellHeaderUploader.GetNumScatters());
@@ -1904,7 +1925,7 @@ public:
 		SCOPED_NAMED_EVENT(SceneCulling_Update_Post, FColor::Emerald);
 		SCOPE_CYCLE_COUNTER(STAT_SceneCulling_Update_Post);
 
-		BUILDER_LOG_SCOPE("ProcessPostSceneUpdate: %d/%d", ScenePostUpdateData.UpdatedPrimitives.Num(), ScenePostUpdateData.AddedPrimitives.Num());
+		BUILDER_LOG_SCOPE("ProcessPostSceneUpdate: %d/%d", ScenePostUpdateData.UpdatedPrimitiveIds.Num(), ScenePostUpdateData.AddedPrimitiveIds.Num());
 
 		SceneCulling.PrimitiveStates.SetNum(SceneCulling.Scene.GetMaxPersistentPrimitiveIndex(), false);
 
@@ -2017,7 +2038,6 @@ public:
 
 #if SC_ENABLE_DETAILED_LOGGING
 	bool bIsLoggingEnabled = false;
-	FString LogStr;
 	FString SceneTag;
 	int32 LogStrIndent = 0;
 public:
@@ -2030,23 +2050,19 @@ public:
 		if (bIsLoggingEnabled)
 		{
 			FString Indent = FString::ChrN(LogStrIndent, TEXT(' '));
-			LogStr.Appendf(TEXT("%s%s\n"), *Indent, *Item);
 			UE_LOG(LogTemp, Warning, TEXT("[%s]%s%s"), *SceneTag, *Indent, *Item);
 		}
 	}
 	inline void EndLogging()
 	{
+		check(GBuilderForLogging == this);
+		GBuilderForLogging = nullptr;
 		if (bIsLoggingEnabled)
 		{
-			check(GBuilderForLogging == this);
-			GBuilderForLogging = nullptr;
-
 			LogIndent(-1);
 			AddLog(TEXT("Log-Scope-End"));
-			UE_LOG(LogTemp, Warning, TEXT("\n\n%s\n%s\n\n"), *SceneTag, *LogStr);
 		}
 		LogStrIndent = 0;
-		LogStr.Reset();
 	}
 #endif
 };
@@ -2142,7 +2158,7 @@ void FSceneCulling::FUpdater::OnPreSceneUpdate(FRDGBuilder& GraphBuilder, const
 	{
 		SCOPED_NAMED_EVENT(SceneCulling_Update, FColor::Emerald);
 		SCOPE_CYCLE_COUNTER(STAT_SceneCulling_Update_Pre);
-		BUILDER_LOG_SCOPE("OnPreSceneUpdate: %d", ScenePreUpdateData.RemovedPrimitives.Num());
+		BUILDER_LOG_SCOPE("OnPreSceneUpdate: %d", ScenePreUpdateData.RemovedPrimitiveIds.Num());
 
 		check(DebugTaskCounter++ == 0);
 		for (int32 Index = 0; Index < ScenePreUpdateData.RemovedPrimitiveIds.Num(); ++Index)
@@ -2197,7 +2213,7 @@ void FSceneCulling::FUpdater::FinalizeAndClear(FRDGBuilder& GraphBuilder, bool b
 	{
 		SCOPED_NAMED_EVENT(SceneCulling_Update_FinalizeAndClear, FColor::Emerald);
 		SCOPE_CYCLE_COUNTER(STAT_SceneCulling_Update_FinalizeAndClear);
-		BUILDER_LOG_SCOPE("FinalizeAndClear");
+		BUILDER_LOG_SCOPE("FinalizeAndClear %d", 1);
 
 		PostUpdateTaskHandle.Wait();
 
@@ -2222,7 +2238,58 @@ void FSceneCulling::EndUpdate(FRDGBuilder& GraphBuilder, bool bPublishStats)
 	if (bIsEnabled)
 	{
 		Updater.FinalizeAndClear(GraphBuilder, bPublishStats);
+#if DO_CHECK
+		ValidateAllInstanceAllocations();
+#endif
+	}
+}
+
+void FSceneCulling::ValidateAllInstanceAllocations()
+{
+#if DO_CHECK
+	if (CVarValidateAllInstanceAllocations.GetValueOnRenderThread() != 0)
+	{
+		for (TConstSetBitIterator<> BitIt(CellOccupancyMask); BitIt; ++BitIt)
+		{
+			uint32 CellId = uint32(BitIt.GetIndex());
+			FCellHeader CellHeader = CellHeaders[CellId];
+			check(IsValid(CellHeader));
+
+			for (uint32 ChunkIndex = 0; ChunkIndex < CellHeader.NumItemChunks; ++ChunkIndex)
+			{
+				uint32 PackedChunkData = PackedCellChunkData[CellHeader.ItemChunksOffset + ChunkIndex];
+				const bool bIsCompressed = (PackedChunkData & INSTANCE_HIERARCHY_ITEM_CHUNK_COMPRESSED_FLAG) != 0u;
+				const uint32 NumItems = bIsCompressed ? 64u : PackedChunkData >> INSTANCE_HIERARCHY_ITEM_CHUNK_COUNT_SHIFT;
+				const bool bIsFullChunk = NumItems == 64u;
+
+				// 1. if it is a compressed chunk and contains any index, we may assume it is to be removed entirely.
+				if (bIsCompressed)
+				{
+					uint32 InstanceDataOffset = PackedChunkData & INSTANCE_HIERARCHY_ITEM_CHUNK_COMPRESSED_PAYLOAD_MASK;
+					for (uint32 ItemIndex = 0u; ItemIndex < 64u; ++ItemIndex)
+					{
+						uint32 InstanceId = InstanceDataOffset + ItemIndex;
+						check(!Scene.GPUScene.GetInstanceSceneDataAllocator().IsFree(InstanceId));
+					}
+				}
+				else
+				{
+					uint32 ChunkId =  (PackedChunkData & INSTANCE_HIERARCHY_ITEM_CHUNK_ID_MASK);
+					uint32 ItemDataOffset = ChunkId * INSTANCE_HIERARCHY_MAX_CHUNK_SIZE;
+
+					// 3.1. scan chunk for deleted items
+					uint64 DeletionMask = 0ull;
+
+					for (uint32 ItemIndex = 0u; ItemIndex < NumItems; ++ItemIndex)
+					{
+						uint32 InstanceId = PackedCellData[ItemDataOffset + ItemIndex];
+						check(!Scene.GPUScene.GetInstanceSceneDataAllocator().IsFree(InstanceId));
+					}
+				}
+			}
+		}
 	}
+#endif
 }
 
 UE::Tasks::FTask FSceneCulling::GetUpdateTaskHandle() const
diff --git a/Engine/Source/Runtime/Renderer/Private/SceneCulling/SceneCulling.h b/Engine/Source/Runtime/Renderer/Private/SceneCulling/SceneCulling.h
--- a/Engine/Source/Runtime/Renderer/Private/SceneCulling/SceneCulling.h
+++ b/Engine/Source/Runtime/Renderer/Private/SceneCulling/SceneCulling.h
@@ -102,6 +102,7 @@ public:
 	using FFootprint32 = FSpatialHash::FFootprint32;
 	using FFootprint64 = FSpatialHash::FFootprint64;	
 private:
+	void ValidateAllInstanceAllocations();
 
 	void Empty();
 
