Path Tracing: Fix error from materials that are translucent but flagged to cast shadows as masked

The IsMasked() method includes translucent materials which cast shadows as masked. This is not a condition the path tracer cares about, so modify the logic to only look at the blend mode.

diff --git a/Engine/Source/Runtime/Renderer/Private/PathTracing.cpp b/Engine/Source/Runtime/Renderer/Private/PathTracing.cpp
--- a/Engine/Source/Runtime/Renderer/Private/PathTracing.cpp
+++ b/Engine/Source/Runtime/Renderer/Private/PathTracing.cpp
@@ -1270,7 +1270,7 @@ static bool NeedsAnyHitShader(bool bIsMasked, bool bIsDitherMasked)
 
 static bool NeedsAnyHitShader(const FMaterial& RESTRICT MaterialResource)
 {
-	return NeedsAnyHitShader(MaterialResource.IsMasked(), MaterialResource.IsDitherMasked());
+	return NeedsAnyHitShader(MaterialResource.GetBlendMode() == BLEND_Masked, MaterialResource.IsDitherMasked());
 }
 
 template<bool UseAnyHitShader, bool UseIntersectionShader, bool IsGPULightmass, bool SimplifyStrata>
@@ -1301,7 +1301,7 @@ public:
 			// This material is only for surfaces at the moment
 			return false;
 		}
-		if (NeedsAnyHitShader(Parameters.MaterialParameters.bIsMasked, Parameters.MaterialParameters.bIsDitherMasked) != UseAnyHitShader)
+		if (NeedsAnyHitShader(Parameters.MaterialParameters.BlendMode == BLEND_Masked, Parameters.MaterialParameters.bIsDitherMasked) != UseAnyHitShader)
 		{
 			return false;
 		}
