Fixed typo in FViewInfo::IsDistanceCulled_AnyThread() causing it to execute on primitives that should have been excluded

diff --git a/Engine/Source/Runtime/Renderer/Private/SceneVisibility.cpp b/Engine/Source/Runtime/Renderer/Private/SceneVisibility.cpp
--- a/Engine/Source/Runtime/Renderer/Private/SceneVisibility.cpp
+++ b/Engine/Source/Runtime/Renderer/Private/SceneVisibility.cpp
@@ -392,7 +392,7 @@ bool FViewInfo::IsDistanceCulled_AnyThread(float DistanceSquared, float MinDrawD
 	const float MaxDrawDistance = InMaxDrawDistance * MaxDrawDistanceScale;
 
 	bool bHasMaxDrawDistance = InMaxDrawDistance != FLT_MAX;
-	bool bHasMinDrawDistance = InMaxDrawDistance > 0;
+	bool bHasMinDrawDistance = MinDrawDistance > 0;
 	bOutMayBeFading = false;
 
 
