Removed DefinitelyUnoccluded check from updating component render times. There's no measurable performance difference and it's technically incorrect as-is, since a component that is part of a batched occlusion query won't have its render time updated even if it's visible.

diff --git a/Engine/Source/Runtime/Renderer/Private/SceneRendering.cpp b/Engine/Source/Runtime/Renderer/Private/SceneRendering.cpp
--- a/Engine/Source/Runtime/Renderer/Private/SceneRendering.cpp
+++ b/Engine/Source/Runtime/Renderer/Private/SceneRendering.cpp
@@ -862,8 +862,6 @@ void FViewInfo::Init()
 	ShadingModelMaskInView = 0;
 	bSceneHasSkyMaterial = 0;
 	bHasSingleLayerWaterMaterial = 0;
-	bHasTranslucencySeparateModulation = 0;
-	bHasStandardTranslucencyModulation = 0;
 	AutoBeforeDOFTranslucencyBoundary = 0.0f;
 	bUsesSecondStageDepthPass = 0;
 
diff --git a/Engine/Source/Runtime/Renderer/Private/SceneRendering.h b/Engine/Source/Runtime/Renderer/Private/SceneRendering.h
--- a/Engine/Source/Runtime/Renderer/Private/SceneRendering.h
+++ b/Engine/Source/Runtime/Renderer/Private/SceneRendering.h
@@ -1271,9 +1271,6 @@ public:
 	/** A map from primitive ID to a boolean ray tracing visibility value. */
 	FSceneBitArray PrimitiveRayTracingVisibilityMap;
 
-	/** Bit set when a primitive is known to be un-occluded. */
-	FSceneBitArray PrimitiveDefinitelyUnoccludedMap;
-
 	/** A map from primitive ID to a boolean is fading value. */
 	FSceneBitArray PotentiallyFadingPrimitiveMap;
 
@@ -1500,11 +1497,6 @@ public:
 	 * true if the scene has at least one mesh with a material tagged as water visible in a view.
 	 */
 	uint32 bHasSingleLayerWaterMaterial : 1;
-	/**
-	 * true if the scene has at least one mesh with a material that needs dual blending AND is applied post DOF. 
-	 * If true, that means we need to run the post-dof separate modulation render pass.
-	 */
-	uint32 bHasTranslucencySeparateModulation : 1;
 	/**
 	 * True if the scene has at least one mesh that needs to sample form the first stage depth buffer.
 	 * And as such will need to render in the second stage depth buffer after the first stage depth buffer is copied.
@@ -1512,12 +1504,6 @@ public:
 	 */
 	uint32 bUsesSecondStageDepthPass : 1;
 
-	/**
-	 * true if the scene has at least one mesh with a material that needs dual blending AND is applied before DOF. 
-	 * If true, that means we need to run the before-dof separate modulation render pass.
-	 */
-	uint32 bHasStandardTranslucencyModulation : 1;
-
 	/** Whether post DOF translucency should be rendered before DOF if primitive bounds behind DOF's focus distance. */
 	float AutoBeforeDOFTranslucencyBoundary;
 
diff --git a/Engine/Source/Runtime/Renderer/Private/SceneVisibility.cpp b/Engine/Source/Runtime/Renderer/Private/SceneVisibility.cpp
--- a/Engine/Source/Runtime/Renderer/Private/SceneVisibility.cpp
+++ b/Engine/Source/Runtime/Renderer/Private/SceneVisibility.cpp
@@ -1135,9 +1135,7 @@ void FRelevancePacket::Finalize()
 	WriteView.bTranslucentSurfaceLighting |= bTranslucentSurfaceLighting;
 	WriteView.bSceneHasSkyMaterial |= bSceneHasSkyMaterial;
 	WriteView.bHasSingleLayerWaterMaterial |= bHasSingleLayerWaterMaterial;
-	WriteView.bHasTranslucencySeparateModulation |= bHasTranslucencySeparateModulation;
 	WriteView.bUsesSecondStageDepthPass |= bUsesSecondStageDepthPass && ShadingPath != EShadingPath::Mobile;
-	WriteView.bHasStandardTranslucencyModulation |= bHasStandardTranslucencyModulation;
 	VisibleDynamicPrimitivesWithSimpleLights.AppendTo(WriteView.VisibleDynamicPrimitivesWithSimpleLights);
 	WriteView.NumVisibleDynamicPrimitives += NumVisibleDynamicPrimitives;
 	WriteView.NumVisibleDynamicEditorPrimitives += NumVisibleDynamicEditorPrimitives;
@@ -1210,8 +1208,6 @@ void FRelevancePacket::ComputeRelevance(FDynamicPrimitiveIndexList& DynamicPrimi
 	StrataBSDFCountMask = 0;
 	bSceneHasSkyMaterial = 0;
 	bHasSingleLayerWaterMaterial = 0;
-	bHasTranslucencySeparateModulation = 0;
-	bHasStandardTranslucencyModulation = 0;
 	bUsesSecondStageDepthPass = 0;
 	bUsesLightingChannels = false;
 	bTranslucentSurfaceLighting = false;
@@ -1773,8 +1769,6 @@ void FRelevancePacket::ComputeRelevance(FDynamicPrimitiveIndexList& DynamicPrimi
 		bUsesCustomStencil |= (ViewRelevance.CustomDepthStencilUsageMask & (1 << 1)) > 0;
 		bSceneHasSkyMaterial |= ViewRelevance.bUsesSkyMaterial;
 		bHasSingleLayerWaterMaterial |= ViewRelevance.bUsesSingleLayerWaterMaterial;
-		bHasStandardTranslucencyModulation |= ViewRelevance.bNormalTranslucency && ViewRelevance.bTranslucencyModulate && View.Family->AllowStandardTranslucencySeparated();
-		bHasTranslucencySeparateModulation |= ViewRelevance.bSeparateTranslucency && ViewRelevance.bTranslucencyModulate;
 		bUsesSecondStageDepthPass |= ViewRelevance.bRenderInSecondStageDepthPass && ShadingPath!=EShadingPath::Mobile;
 
 		if (ViewRelevance.bRenderCustomDepth)
@@ -1803,13 +1797,8 @@ void FRelevancePacket::ComputeRelevance(FDynamicPrimitiveIndexList& DynamicPrimi
 
 		PrimitiveSceneInfo->LastRenderTime = CurrentWorldTime;
 
-		// If the primitive is definitely unoccluded or if in Wireframe mode and the primitive is estimated
-		// to be unoccluded, then update the primitive components's LastRenderTime 
-		// on the game thread. This signals that the primitive is visible.
-		if (View.PrimitiveDefinitelyUnoccludedMap[BitIndex] || (View.Family->EngineShowFlags.Wireframe && View.PrimitiveVisibilityMap[BitIndex]))
-		{
-			PrimitiveSceneInfo->UpdateComponentLastRenderTime(CurrentWorldTime, /*bUpdateLastRenderTimeOnScreen=*/true);
-		}
+		const bool bUpdateLastRenderTimeOnScreen = true;
+		PrimitiveSceneInfo->UpdateComponentLastRenderTime(CurrentWorldTime, bUpdateLastRenderTimeOnScreen);
 
 		// Cache the nearest reflection proxy if needed
 		if (PrimitiveSceneInfo->NeedsReflectionCaptureUpdate())
@@ -2568,10 +2557,6 @@ bool FGPUOcclusionPacket::OcclusionCullPrimitive(VisitorType& Visitor, FOcclusio
 				bIsVisible = false;
 				Result.NumCulledPrimitives++;
 			}
-			else if (bOcclusionStateIsDefinite)
-			{
-				View.PrimitiveDefinitelyUnoccludedMap[Index] = true;
-			}
 		}
 	}
 
@@ -2586,10 +2571,6 @@ bool FGPUOcclusionPacket::OcclusionCullPrimitive(VisitorType& Visitor, FOcclusio
 			bIsVisible = false;
 			Result.NumCulledPrimitives++;
 		}
-		else if (bAllSubOcclusionStateIsDefinite)
-		{
-			View.PrimitiveDefinitelyUnoccludedMap[Index] = true;
-		}
 	}
 
 	return bIsVisible;
@@ -2813,7 +2794,6 @@ bool FGPUOcclusionParallelPacket::AddPrimitive(int32 PrimitiveIndex)
 		return true;
 	}
 
-	View.PrimitiveDefinitelyUnoccludedMap[PrimitiveIndex] = true;
 	return false;
 }
 
@@ -2987,10 +2967,6 @@ void FGPUOcclusionSerial::AddPrimitives(FPrimitiveRange PrimitiveRange)
 		{
 			Packet.OcclusionCullPrimitive(ProcessVisitor, OcclusionCullResult, BitIt.GetIndex());
 		}
-		else
-		{
-			View.PrimitiveDefinitelyUnoccludedMap.AccessCorrespondingBit(BitIt) = true;
-		}
 	}
 }
 
@@ -3235,7 +3211,6 @@ void FVisibilityViewPacket::BeginInitVisibility()
 	View.PrimitiveVisibilityMap.Init(false, Scene.Primitives.Num());
 	View.PrimitiveRayTracingVisibilityMap.Init(false, Scene.Primitives.Num());
 	View.DynamicMeshEndIndices.SetNumZeroed(Scene.Primitives.Num());
-	View.PrimitiveDefinitelyUnoccludedMap.Init(false, Scene.Primitives.Num());
 	View.PotentiallyFadingPrimitiveMap.Init(false, Scene.Primitives.Num());
 	View.PrimitiveFadeUniformBuffers.AddZeroed(Scene.Primitives.Num());
 	View.PrimitiveFadeUniformBufferMap.Init(false, Scene.Primitives.Num());
diff --git a/Engine/Source/Runtime/Renderer/Private/SceneVisibilityPrivate.h b/Engine/Source/Runtime/Renderer/Private/SceneVisibilityPrivate.h
--- a/Engine/Source/Runtime/Renderer/Private/SceneVisibilityPrivate.h
+++ b/Engine/Source/Runtime/Renderer/Private/SceneVisibilityPrivate.h
@@ -665,8 +665,6 @@ private:
 	bool bUsesCustomStencil = false;
 	bool bSceneHasSkyMaterial = false;
 	bool bHasSingleLayerWaterMaterial = false;
-	bool bHasTranslucencySeparateModulation = false;
-	bool bHasStandardTranslucencyModulation = false;
 	bool bUsesSecondStageDepthPass = false;
 	bool bAddLightmapDensityCommands = false;
 	bool bComputeRelevanceTaskLaunched = false;
