Implements r.Test.FreezeTemporalHistories.Progress

diff --git a/Engine/Source/Runtime/Renderer/Private/SceneVisibility.cpp b/Engine/Source/Runtime/Renderer/Private/SceneVisibility.cpp
--- a/Engine/Source/Runtime/Renderer/Private/SceneVisibility.cpp
+++ b/Engine/Source/Runtime/Renderer/Private/SceneVisibility.cpp
@@ -347,6 +347,11 @@ static TAutoConsoleVariable<float> CVarFreezeTemporalHistories(
 	TEXT("Freezes all temporal histories as well as the temporal sequence."),
 	ECVF_RenderThreadSafe);
 
+static TAutoConsoleVariable<float> CVarFreezeTemporalHistoriesProgress(
+	TEXT("r.Test.FreezeTemporalHistories.Progress"), 0,
+	TEXT("Progress the temporal histories by one frame when modified."),
+	ECVF_RenderThreadSafe);
+
 #endif
 
 DECLARE_CYCLE_STAT(TEXT("Occlusion Readback"), STAT_CLMM_OcclusionReadback, STATGROUP_CommandListMarkers);
@@ -4077,6 +4082,14 @@ void FSceneRenderer::PrepareViewStateForVisibility(const FSceneTexturesConfig& S
 	const bool bFreezeTemporalSequences = false;
 #else
 	bool bFreezeTemporalHistories = CVarFreezeTemporalHistories.GetValueOnRenderThread() != 0;
+
+	static int32 CurrentFreezeTemporalHistoriesProgress = 0;
+	if (CurrentFreezeTemporalHistoriesProgress != CVarFreezeTemporalHistoriesProgress.GetValueOnRenderThread())
+	{
+		bFreezeTemporalHistories = false;
+		CurrentFreezeTemporalHistoriesProgress = CVarFreezeTemporalHistoriesProgress.GetValueOnRenderThread();
+	}
+
 	bool bFreezeTemporalSequences = bFreezeTemporalHistories || CVarFreezeTemporalSequences.GetValueOnRenderThread() != 0;
 #endif
 
