Lumen - fixed material AO. Don't apply material AO on top of scene color, as Lumen handles it internally.
* Without async it was applied on top of emissive
* With async it was applied on top of emissive and direct lighting

diff --git a/Engine/Shaders/Private/DiffuseIndirectComposite.usf b/Engine/Shaders/Private/DiffuseIndirectComposite.usf
--- a/Engine/Shaders/Private/DiffuseIndirectComposite.usf
+++ b/Engine/Shaders/Private/DiffuseIndirectComposite.usf
@@ -348,9 +348,14 @@ void MainPS(
 	// Sample the ambient occlusion that is dynamically generated every frame.
 	const float DynamicAmbientOcclusion = AmbientOcclusionTexture.SampleLevel(AmbientOcclusionSampler, SceneBufferUV, 0).r;
 
-	// Compute the final ambient occlusion to be applied.
-	const float AOMask = IsValid(Material);
-	const float FinalAmbientOcclusion = lerp(1.0f, Material.MaterialAO * DynamicAmbientOcclusion, AOMask * AmbientOcclusionStaticFraction);
+	// Compute the final ambient occlusion to be applied. Lumen handles material AO internally.
+	float FinalAmbientOcclusion = 1.0f;
+	#if DIM_APPLY_DIFFUSE_INDIRECT != DIM_APPLY_DIFFUSE_INDIRECT_SCREEN_PROBE_GATHER
+	{
+		float AOMask = IsValid(Material);
+		FinalAmbientOcclusion = lerp(1.0f, Material.MaterialAO * DynamicAmbientOcclusion, AOMask * AmbientOcclusionStaticFraction);
+	}
+	#endif
 
 	const float3 TranslatedWorldPosition = mul(float4(GetScreenPositionForProjectionType(ScreenPosition, Material.SceneDepth), Material.SceneDepth, 1), View.ScreenToTranslatedWorld).xyz;
 
diff --git a/Engine/Source/Runtime/Renderer/Private/IndirectLightRendering.cpp b/Engine/Source/Runtime/Renderer/Private/IndirectLightRendering.cpp
--- a/Engine/Source/Runtime/Renderer/Private/IndirectLightRendering.cpp
+++ b/Engine/Source/Runtime/Renderer/Private/IndirectLightRendering.cpp
@@ -1326,6 +1326,7 @@ void FDeferredShadingSceneRenderer::RenderDiffuseIndirectAndAmbientOcclusion(
 			PermutationVector.Set<FDiffuseIndirectCompositePS::FStrataTileType>(EStrataTileType::EComplex);
 
 			bool bUpscale = false;
+			bool bApplyAOToSceneColor = true;
 
 			if (DenoiserOutputs.Textures[0])
 			{
@@ -1343,6 +1344,7 @@ void FDeferredShadingSceneRenderer::RenderDiffuseIndirectAndAmbientOcclusion(
 						PermutationVector.Set<FDiffuseIndirectCompositePS::FStrataTileType>(TileType);
 					}
 					DiffuseIndirectSampling = TEXT("ScreenProbeGather");
+					bApplyAOToSceneColor = false;
 				}
 				else
 				{
@@ -1361,9 +1363,23 @@ void FDeferredShadingSceneRenderer::RenderDiffuseIndirectAndAmbientOcclusion(
 			FRHIBlendState* BlendState;
 			if (!bEnableCopyPass)
 			{
-				BlendState = PermutationVector.Get<FDiffuseIndirectCompositePS::FApplyDiffuseIndirectDim>() > 0 ?
-					TStaticBlendState<CW_RGBA, BO_Add, BF_One, BF_Source1Color, BO_Add, BF_One, BF_Source1Alpha>::GetRHI() :
-					TStaticBlendState<CW_RGBA, BO_Add, BF_Zero, BF_SourceColor, BO_Add, BF_Zero, BF_SourceAlpha>::GetRHI();
+				if (PermutationVector.Get<FDiffuseIndirectCompositePS::FApplyDiffuseIndirectDim>() > 0)
+				{
+					if (bApplyAOToSceneColor)
+					{
+						// FinalColor = SceneColor * AO + Indirect;
+						BlendState = TStaticBlendState<CW_RGBA, BO_Add, BF_One, BF_Source1Color, BO_Add, BF_One, BF_Source1Alpha>::GetRHI();
+					}
+					else
+					{
+						// FinalColor = SceneColor + Indirect;
+						BlendState = TStaticBlendState<CW_RGBA, BO_Add, BF_One, BF_One, BO_Add, BF_One, BF_One>::GetRHI();
+					}
+				}		
+				else
+				{
+					BlendState = TStaticBlendState<CW_RGBA, BO_Add, BF_Zero, BF_SourceColor, BO_Add, BF_Zero, BF_SourceAlpha>::GetRHI();
+				}
 			}
 			else
 			{
