Substrate - added support for planar refelctions.

diff --git a/Engine/Shaders/Private/PlanarReflectionShaders.usf b/Engine/Shaders/Private/PlanarReflectionShaders.usf
--- a/Engine/Shaders/Private/PlanarReflectionShaders.usf
+++ b/Engine/Shaders/Private/PlanarReflectionShaders.usf
@@ -8,6 +8,7 @@
 #include "DeferredShadingCommon.ush"
 #include "PlanarReflectionShared.ush"
 #include "SceneTextureParameters.ush"
+#include "Strata/Strata.ush"
 
 Texture2D SceneColorInputTexture;
 SamplerState SceneColorInputSampler;
@@ -140,6 +141,31 @@ void PlanarReflectionPS(
 	ResolvedView = ResolveView();
 	OutColor = 0;
 
+
+#if STRATA_ENABLED
+
+	const FStrataTopLayerData TopLayerData = StrataUnpackTopLayerData(Strata.TopLayerTexture.Load(uint3(SVPos.xy, 0)));
+
+	float DeviceZ = SampleDeviceZFromSceneTextures(BufferUV);
+	float SceneDepth = ConvertFromDeviceZ(DeviceZ);
+
+	float3 TranslatedWorldPosition = ResolvedView.TranslatedWorldCameraOrigin + ScreenVector * SceneDepth;
+
+	const float Roughness = TopLayerData.Roughness;
+	const float3 WorldNormal = TopLayerData.WorldNormal;
+	const bool bNoMaterial = !IsStrataMaterial(TopLayerData);
+
+	// Skip unlit pixels (skybox)
+	BRANCH
+	if (!bNoMaterial)
+	{
+		OutColor = ComputePlanarReflections(TranslatedWorldPosition, WorldNormal, Roughness, PlanarReflectionStruct.PlanarReflectionSampler);
+
+		OutColor.xyz *= View.PreExposure;
+	}
+
+#else
+
 	FGBufferData GBuffer = GetGBufferDataFromSceneTextures(BufferUV);
 
 	float3 TranslatedWorldPosition = ResolvedView.TranslatedWorldCameraOrigin + ScreenVector * GBuffer.Depth;
@@ -152,6 +178,9 @@ void PlanarReflectionPS(
 
 		OutColor.xyz *= View.PreExposure;
 	}
+
+#endif
+
 }
 
 #endif
\ No newline at end of file
diff --git a/Engine/Source/Runtime/Renderer/Private/PlanarReflectionRendering.cpp b/Engine/Source/Runtime/Renderer/Private/PlanarReflectionRendering.cpp
--- a/Engine/Source/Runtime/Renderer/Private/PlanarReflectionRendering.cpp
+++ b/Engine/Source/Runtime/Renderer/Private/PlanarReflectionRendering.cpp
@@ -35,6 +35,7 @@
 #include "ClearQuad.h"
 #include "SceneTextureParameters.h"
 #include "SceneViewExtension.h"
+#include "Strata/Strata.h"
 
 void SetupPlanarReflectionUniformParameters(const class FSceneView& View, const FPlanarReflectionSceneProxy* ReflectionSceneProxy, FPlanarReflectionUniformParameters& OutParameters)
 {
@@ -797,6 +798,7 @@ class FPlanarReflectionPS : public FGlobalShader
 	BEGIN_SHADER_PARAMETER_STRUCT(FParameters, )
 		SHADER_PARAMETER_STRUCT_INCLUDE(FSceneTextureParameters, SceneTextures)
 
+		SHADER_PARAMETER_RDG_UNIFORM_BUFFER(FStrataGlobalUniformParameters, Strata)
 		SHADER_PARAMETER_STRUCT_REF(FViewUniformShaderParameters, ViewUniformBuffer)
 		SHADER_PARAMETER_STRUCT_REF(FPlanarReflectionUniformParameters, PlanarReflectionParameters)
 
@@ -872,6 +874,7 @@ void FDeferredShadingSceneRenderer::RenderDeferredPlanarReflections(FRDGBuilder&
 	PassParameters->SceneTextures.GBufferFTexture = SceneTextures.GBufferFTexture;
 	PassParameters->SceneTextures.GBufferVelocityTexture = SceneTextures.GBufferVelocityTexture;
 
+	PassParameters->Strata = Strata::BindStrataGlobalUniformParameters(View);
 	PassParameters->ViewUniformBuffer = View.ViewUniformBuffer;
 	PassParameters->RenderTargets[0] = FRenderTargetBinding(
 		ReflectionsOutputTexture, bClearReflectionsOutputTexture ? ERenderTargetLoadAction::EClear : ERenderTargetLoadAction::ELoad);
