Make FAtlasTextureInvalidationScope public so that external code (plugins) can invalidate a particular texture (e.g., for CPU data upload).

diff --git a/Engine/Source/Runtime/Renderer/Private/RectLightTextureManager.h b/Engine/Source/Runtime/Renderer/Private/RectLightTextureManager.h
--- a/Engine/Source/Runtime/Renderer/Private/RectLightTextureManager.h
+++ b/Engine/Source/Runtime/Renderer/Private/RectLightTextureManager.h
@@ -39,7 +39,7 @@ RENDERER_API void AddDebugPass(FRDGBuilder& GraphBuilder, const FViewInfo& View,
 
 // Scope for invalidating a particular texture 
 // This ensures the atlas contains the latest version of the texture and filter it
-struct FAtlasTextureInvalidationScope
+struct RENDERER_API FAtlasTextureInvalidationScope
 {
 	FAtlasTextureInvalidationScope(UTexture* In);
 	~FAtlasTextureInvalidationScope();
