Fix rect light support with MediaTexture.

These textures, when assigned to a rect light, were missing invalidation to refresh their atlas content.

diff --git a/Engine/Source/Runtime/MediaAssets/MediaAssets.Build.cs b/Engine/Source/Runtime/MediaAssets/MediaAssets.Build.cs
--- a/Engine/Source/Runtime/MediaAssets/MediaAssets.Build.cs
+++ b/Engine/Source/Runtime/MediaAssets/MediaAssets.Build.cs
@@ -22,7 +22,8 @@ namespace UnrealBuildTool.Rules
 				new string[] {
 					"RenderCore",
 					"RHI",
-					"ColorManagement"
+					"ColorManagement",
+					"Renderer",
 				});
 
 			if (Target.bBuildEditor)
diff --git a/Engine/Source/Runtime/MediaAssets/Private/Assets/MediaTexture.cpp b/Engine/Source/Runtime/MediaAssets/Private/Assets/MediaTexture.cpp
--- a/Engine/Source/Runtime/MediaAssets/Private/Assets/MediaTexture.cpp
+++ b/Engine/Source/Runtime/MediaAssets/Private/Assets/MediaTexture.cpp
@@ -18,6 +18,7 @@
 #include "Misc/MediaTextureResource.h"
 #include "IMediaTextureSample.h"
 
+#include "RectLightTexture.h"
 
 /* Local helpers
  *****************************************************************************/
@@ -467,12 +468,16 @@ void UMediaTexture::TickResource(FTimespan Timecode)
 	FMediaTextureResource* ResourceParam = (FMediaTextureResource*)GetResource();
 
 	const ERenderMode RenderModeParam = GetRenderMode();
+	const UTexture* TexturePtrNotDeferenced = this;
 
 	ENQUEUE_RENDER_COMMAND(MediaTextureResourceRender)(
-		[ResourceParam, RenderParams, RenderModeParam](FRHICommandListImmediate& RHICmdList)
+		[ResourceParam, RenderParams, RenderModeParam, TexturePtrNotDeferenced](FRHICommandListImmediate& RHICmdList)
 		{
 			check(ResourceParam);
 
+			// Lock/Enqueue rect atlas refresh if that texture is used by a rect. light
+			RectLightAtlas::FAtlasTextureInvalidationScope InvalidationScope(TexturePtrNotDeferenced);
+
 			if (RenderModeParam == ERenderMode::JustInTime)
 			{
 				// Cache the render params if this is a just in time render mode. 
diff --git a/Engine/Source/Runtime/Renderer/Private/RectLightTextureManager.cpp b/Engine/Source/Runtime/Renderer/Private/RectLightTextureManager.cpp
--- a/Engine/Source/Runtime/Renderer/Private/RectLightTextureManager.cpp
+++ b/Engine/Source/Runtime/Renderer/Private/RectLightTextureManager.cpp
@@ -19,6 +19,7 @@
 #include "TextureLayout.h"
 #include "CommonRenderResources.h"
 #include "ScreenPass.h"
+#include "RectLightTexture.h"
 
 ///////////////////////////////////////////////////////////////////////////////////////////////////
 // Possible improvements:
@@ -1705,7 +1706,6 @@ void AddDebugPass(FRDGBuilder& GraphBuilder, const FViewInfo& View, FRDGTextureR
 	}
 }
 
-
 FRHITexture* GetAtlasTexture()
 {
 	return GRectLightTextureManager.AtlasTexture ? GRectLightTextureManager.AtlasTexture->GetRHI() : nullptr;
@@ -1713,7 +1713,7 @@ FRHITexture* GetAtlasTexture()
 
 // Scope object allowing to force refresh of a slot texture, and locking/preventing 
 // the atlas update during the 'update'/'capture'
-FAtlasTextureInvalidationScope::FAtlasTextureInvalidationScope(UTexture* In)
+FAtlasTextureInvalidationScope::FAtlasTextureInvalidationScope(const UTexture* In)
 {
 	if (In)
 	{
@@ -1728,7 +1728,7 @@ FAtlasTextureInvalidationScope::FAtlasTextureInvalidationScope(UTexture* In)
 				// Sanity check, allow a single lock/capture refresh at a time.
 				check(GRectLightTextureManager.bLock == false);
 
-				Texture = In;
+				bLocked = true;
 				Slot.bForceRefresh = true;
 				GRectLightTextureManager.bLock = true;
 				return;
@@ -1739,12 +1739,13 @@ FAtlasTextureInvalidationScope::FAtlasTextureInvalidationScope(UTexture* In)
 
 FAtlasTextureInvalidationScope::~FAtlasTextureInvalidationScope()
 {
-	if (Texture)
+	if (bLocked)
 	{
 		// Sanity check, allow a single lock/capture refresh at a time.
 		check(GRectLightTextureManager.bLock == true);
 
 		GRectLightTextureManager.bLock = false;
+		bLocked = false;
 	}
 }
 
diff --git a/Engine/Source/Runtime/Renderer/Private/RectLightTextureManager.h b/Engine/Source/Runtime/Renderer/Private/RectLightTextureManager.h
--- a/Engine/Source/Runtime/Renderer/Private/RectLightTextureManager.h
+++ b/Engine/Source/Runtime/Renderer/Private/RectLightTextureManager.h
@@ -37,15 +37,4 @@ RENDERER_API void UpdateAtlasTexture(FRDGBuilder& GraphBuilder, const ERHIFeatur
 // Return the rect light atlas debug pass
 RENDERER_API void AddDebugPass(FRDGBuilder& GraphBuilder, const FViewInfo& View, FRDGTextureRef OutputTexture);
 
-// Scope for invalidating a particular texture 
-// This ensures the atlas contains the latest version of the texture and filter it
-struct RENDERER_API FAtlasTextureInvalidationScope
-{
-	FAtlasTextureInvalidationScope(UTexture* In);
-	~FAtlasTextureInvalidationScope();
-	FAtlasTextureInvalidationScope(const FAtlasTextureInvalidationScope&) = delete;
-	FAtlasTextureInvalidationScope& operator=(const FAtlasTextureInvalidationScope&) = delete;
-	UTexture* Texture = nullptr;
-};
-
 } 
diff --git a/Engine/Source/Runtime/Renderer/Private/SceneCaptureRendering.cpp b/Engine/Source/Runtime/Renderer/Private/SceneCaptureRendering.cpp
--- a/Engine/Source/Runtime/Renderer/Private/SceneCaptureRendering.cpp
+++ b/Engine/Source/Runtime/Renderer/Private/SceneCaptureRendering.cpp
@@ -40,7 +40,7 @@
 #include "Rendering/MotionVectorSimulation.h"
 #include "SceneViewExtension.h"
 #include "GenerateMips.h"
-#include "RectLightTextureManager.h"
+#include "RectLightTexture.h"
 #include "Materials/MaterialRenderProxy.h"
 
 #if WITH_EDITOR
diff --git a/Engine/Source/Runtime/Renderer/Public/RectLightTexture.h b/Engine/Source/Runtime/Renderer/Public/RectLightTexture.h
new file mode 100644
--- /dev/null
+++ b/Engine/Source/Runtime/Renderer/Public/RectLightTexture.h
@@ -0,0 +1,21 @@
+// Copyright Epic Games, Inc. All Rights Reserved.
+
+#pragma once
+
+class UTexture;
+
+namespace RectLightAtlas
+{
+
+// Scope for invalidating a particular texture 
+// This ensures the atlas contains the latest version of the texture and filter it
+struct RENDERER_API FAtlasTextureInvalidationScope
+{
+	FAtlasTextureInvalidationScope(const UTexture* In);
+	~FAtlasTextureInvalidationScope();
+	FAtlasTextureInvalidationScope(const FAtlasTextureInvalidationScope&) = delete;
+	FAtlasTextureInvalidationScope& operator=(const FAtlasTextureInvalidationScope&) = delete;
+	bool bLocked = false;
+};
+
+} 
\ No newline at end of file
