Limit PrimitiveOctree bounds to UE_OLD_HALF_WORLD_MAX, similar to CL 22089457

diff --git a/Engine/Source/Runtime/Renderer/Private/RendererScene.cpp b/Engine/Source/Runtime/Renderer/Private/RendererScene.cpp
--- a/Engine/Source/Runtime/Renderer/Private/RendererScene.cpp
+++ b/Engine/Source/Runtime/Renderer/Private/RendererScene.cpp
@@ -1857,7 +1857,7 @@ FScene::FScene(UWorld* InWorld, bool bInRequiresHitProxies, bool bInIsEditorScen
 ,	VolumetricCloud(NULL)
 ,	PrecomputedVisibilityHandler(NULL)
 ,	LocalShadowCastingLightOctree(FVector::ZeroVector, UE_OLD_HALF_WORLD_MAX)
-,	PrimitiveOctree(FVector::ZeroVector, HALF_WORLD_MAX)
+,	PrimitiveOctree(FVector::ZeroVector, UE_OLD_HALF_WORLD_MAX)
 ,	bRequiresHitProxies(bInRequiresHitProxies)
 ,	bIsEditorScene(bInIsEditorScene)
 ,	NumUncachedStaticLightingInteractions(0)
