When possible support decals on mobile for projects with a "Mobile HDR" off

diff --git a/Engine/Source/Runtime/Renderer/Private/MobileDecalRendering.cpp b/Engine/Source/Runtime/Renderer/Private/MobileDecalRendering.cpp
--- a/Engine/Source/Runtime/Renderer/Private/MobileDecalRendering.cpp
+++ b/Engine/Source/Runtime/Renderer/Private/MobileDecalRendering.cpp
@@ -16,13 +16,35 @@
 #include "DecalRenderingCommon.h"
 #include "DecalRenderingShared.h"
 #include "RenderCore.h"
+#include "DataDrivenShaderPlatformInfo.h"
 
 void RenderMeshDecalsMobile(FRHICommandList& RHICmdList, const FViewInfo& View, EDecalRenderStage DecalRenderStage, EDecalRenderTargetMode RenderTargetMode);
 extern void RenderDeferredDecalsMobile(FRHICommandList& RHICmdList, const FScene& Scene, const FViewInfo& View, EDecalRenderStage DecalRenderStage, EDecalRenderTargetMode RenderTargetMode);
 
+static bool DoesPlatformSupportDecals(EShaderPlatform ShaderPlatform)
+{
+	if (!IsMobileHDR())
+	{
+		// Vulkan uses sub-pass to fetch SceneDepth
+		if (IsVulkanPlatform(ShaderPlatform) ||
+			IsSimulatedPlatform(ShaderPlatform) ||
+			// Some Androids support SceneDepth fetch
+			(IsAndroidOpenGLESPlatform(ShaderPlatform) && GSupportsShaderDepthStencilFetch))
+		{
+			return true;
+		}
+
+		// Metal needs DepthAux to fetch depth, and its not availle in LDR mode
+		return false;
+	}
+
+	// HDR always supports decals
+	return true;
+}
+
 void FMobileSceneRenderer::RenderDecals(FRHICommandList& RHICmdList, const FViewInfo& View)
 {
-	if (!IsMobileHDR() || !ViewFamily.EngineShowFlags.Decals || View.bIsPlanarReflection)
+	if (!DoesPlatformSupportDecals(View.GetShaderPlatform()) || !ViewFamily.EngineShowFlags.Decals || View.bIsPlanarReflection)
 	{
 		return;
 	}
