Remove unused files.

diff --git a/Engine/Source/Runtime/Renderer/Private/HairStrands/HairStrandsDebug.cpp b/Engine/Source/Runtime/Renderer/Private/HairStrands/HairStrandsDebug.cpp
--- a/Engine/Source/Runtime/Renderer/Private/HairStrands/HairStrandsDebug.cpp
+++ b/Engine/Source/Runtime/Renderer/Private/HairStrands/HairStrandsDebug.cpp
@@ -9,7 +9,6 @@
 #include "HairStrandsRendering.h"
 #include "HairStrandsVisibility.h"
 #include "HairStrandsInterface.h"
-#include "HairStrandsMeshProjection.h"
 #include "HairStrandsTile.h"
 #include "HairStrandsData.h"
 
diff --git a/Engine/Source/Runtime/Renderer/Private/HairStrands/HairStrandsInterface.cpp b/Engine/Source/Runtime/Renderer/Private/HairStrands/HairStrandsInterface.cpp
--- a/Engine/Source/Runtime/Renderer/Private/HairStrands/HairStrandsInterface.cpp
+++ b/Engine/Source/Runtime/Renderer/Private/HairStrands/HairStrandsInterface.cpp
@@ -6,7 +6,6 @@
 
 #include "HairStrandsInterface.h"
 #include "HairStrandsRendering.h"
-#include "HairStrandsMeshProjection.h"
 #include "HairStrandsData.h"
 
 #include "DataDrivenShaderPlatformInfo.h"
diff --git a/Engine/Source/Runtime/Renderer/Private/HairStrands/HairStrandsMeshProjection.cpp b/Engine/Source/Runtime/Renderer/Private/HairStrands/HairStrandsMeshProjection.cpp
deleted file mode 100644
diff --git a/Engine/Source/Runtime/Renderer/Private/HairStrands/HairStrandsMeshProjection.h b/Engine/Source/Runtime/Renderer/Private/HairStrands/HairStrandsMeshProjection.h
deleted file mode 100644
