Fixed flickering pre-shadows caused by reodering bug introduced in the parallel init views refactor.

diff --git a/Engine/Source/Runtime/Renderer/Private/SceneRendering.cpp b/Engine/Source/Runtime/Renderer/Private/SceneRendering.cpp
--- a/Engine/Source/Runtime/Renderer/Private/SceneRendering.cpp
+++ b/Engine/Source/Runtime/Renderer/Private/SceneRendering.cpp
@@ -2193,21 +2193,18 @@ void FViewInfo::DestroyAllSnapshots(FParallelMeshDrawCommandPass::EWaitThread Wa
 
 FInt32Range FViewInfo::GetDynamicMeshElementRange(uint32 PrimitiveIndex) const
 {
-	int32 Start = 0;	// inclusive
-	int32 AfterEnd = 0;	// exclusive
-
 	// DynamicMeshEndIndices contains valid values only for visible primitives with bDynamicRelevance.
 	if (PrimitiveVisibilityMap[PrimitiveIndex])
 	{
 		const FPrimitiveViewRelevance& ViewRelevance = PrimitiveViewRelevanceMap[PrimitiveIndex];
 		if (ViewRelevance.bDynamicRelevance)
 		{
-			Start = (PrimitiveIndex == 0) ? 0 : DynamicMeshEndIndices[PrimitiveIndex - 1];
-			AfterEnd = DynamicMeshEndIndices[PrimitiveIndex];
+
+			return FInt32Range(DynamicMeshElementRanges[PrimitiveIndex].X, DynamicMeshElementRanges[PrimitiveIndex].Y);
 		}
 	}
 
-	return FInt32Range(Start, AfterEnd);
+	return FInt32Range::Empty();
 }
 
 FRDGTextureRef FViewInfo::GetVolumetricCloudTexture(FRDGBuilder& GraphBuilder) const
diff --git a/Engine/Source/Runtime/Renderer/Private/SceneRendering.h b/Engine/Source/Runtime/Renderer/Private/SceneRendering.h
--- a/Engine/Source/Runtime/Renderer/Private/SceneRendering.h
+++ b/Engine/Source/Runtime/Renderer/Private/SceneRendering.h
@@ -1375,7 +1375,7 @@ public:
 	TArray<FMeshBatchAndRelevance,SceneRenderingAllocator> DynamicMeshElements;
 
 	/* [PrimitiveIndex] = end index index in DynamicMeshElements[], to support GetDynamicMeshElementRange(). Contains valid values only for visible primitives with bDynamicRelevance. */
-	TArray<uint32, SceneRenderingAllocator> DynamicMeshEndIndices;
+	TArray<FInt32Vector2, SceneRenderingAllocator> DynamicMeshElementRanges;
 
 	/** Hair strands & cards dynamic mesh element. */
 	TArray<FMeshBatchAndRelevance, SceneRenderingAllocator> HairStrandsMeshElements;
diff --git a/Engine/Source/Runtime/Renderer/Private/SceneVisibility.cpp b/Engine/Source/Runtime/Renderer/Private/SceneVisibility.cpp
--- a/Engine/Source/Runtime/Renderer/Private/SceneVisibility.cpp
+++ b/Engine/Source/Runtime/Renderer/Private/SceneVisibility.cpp
@@ -3215,7 +3215,7 @@ void FVisibilityViewPacket::BeginInitVisibility()
 	// Allocate the view's visibility maps.
 	View.PrimitiveVisibilityMap.Init(false, Scene.Primitives.Num());
 	View.PrimitiveRayTracingVisibilityMap.Init(false, Scene.Primitives.Num());
-	View.DynamicMeshEndIndices.SetNumZeroed(Scene.Primitives.Num());
+	View.DynamicMeshElementRanges.SetNumZeroed(Scene.Primitives.Num());
 	View.PotentiallyFadingPrimitiveMap.Init(false, Scene.Primitives.Num());
 	View.PrimitiveFadeUniformBuffers.AddZeroed(Scene.Primitives.Num());
 	View.PrimitiveFadeUniformBufferMap.Init(false, Scene.Primitives.Num());
@@ -3653,6 +3653,13 @@ void FVisibilityTaskData::GatherDynamicMeshElements()
 
 void FVisibilityTaskData::GatherDynamicMeshElementsForPrimitive(int32 PrimitiveIndex, uint8 ViewMask)
 {
+	TArray<int32, TInlineAllocator<4>> MeshBatchCountBefore;
+	MeshBatchCountBefore.SetNumUninitialized(Views.Num());
+	for (int32 ViewIndex = 0; ViewIndex < Views.Num(); ViewIndex++)
+	{
+		MeshBatchCountBefore[ViewIndex] = MeshCollector.GetMeshBatchCount(ViewIndex);
+	}
+
 	FPrimitiveSceneInfo* PrimitiveSceneInfo = Scene.Primitives[PrimitiveIndex];
 	MeshCollector.SetPrimitive(PrimitiveSceneInfo->Proxy, PrimitiveSceneInfo->DefaultDynamicHitProxyId);
 
@@ -3660,10 +3667,9 @@ void FVisibilityTaskData::GatherDynamicMeshElementsForPrimitive(int32 PrimitiveI
 
 	for (int32 ViewIndex = 0; ViewIndex < Views.Num(); ViewIndex++)
 	{
+		FViewInfo& View = Views[ViewIndex];
 		if (ViewMask & (1 << ViewIndex))
 		{
-			const FViewInfo& View = Views[ViewIndex];
-
 			FDynamicPrimitive& DynamicPrimitive = DynamicMeshElements.DynamicPrimitives.Emplace_GetRef();
 			DynamicPrimitive.PrimitiveIndex = PrimitiveIndex;
 			DynamicPrimitive.ViewIndex = ViewIndex;
@@ -3672,11 +3678,8 @@ void FVisibilityTaskData::GatherDynamicMeshElementsForPrimitive(int32 PrimitiveI
 
 			DynamicMeshElements.LastElementPerView[ViewIndex] = View.DynamicMeshElements.Num();
 		}
-	}
 
-	for (int32 ViewIndex = 0; ViewIndex < Views.Num(); ViewIndex++)
-	{
-		Views[ViewIndex].DynamicMeshEndIndices[PrimitiveIndex] = MeshCollector.GetMeshBatchCount(ViewIndex);
+		View.DynamicMeshElementRanges[PrimitiveIndex] = FInt32Vector2(MeshBatchCountBefore[ViewIndex], MeshCollector.GetMeshBatchCount(ViewIndex));
 	}
 };
 
