Multiply OneOverPreExposure in BasicEyeAdaptationSetupPS instead of doing it in BasicEyeAdaptationCS
- making base eye adaptation path more consistent with histogram path

diff --git a/Engine/Shaders/Private/PostProcessEyeAdaptation.usf b/Engine/Shaders/Private/PostProcessEyeAdaptation.usf
--- a/Engine/Shaders/Private/PostProcessEyeAdaptation.usf
+++ b/Engine/Shaders/Private/PostProcessEyeAdaptation.usf
@@ -174,7 +174,7 @@ void BasicEyeAdaptationSetupPS(
 	OutColor = Texture2DSample(ColorTexture, ColorSampler, UV);
 
 	// Use max to ensure intensity is never zero (so the following log is well behaved)
-	const float Intensity = CalculateEyeAdaptationLuminance(OutColor.xyz);
+	const float Intensity = CalculateEyeAdaptationLuminance(OutColor.xyz * View.OneOverPreExposure);
 	const float LogIntensity = clamp(log2(Intensity),-10.0f,20.0f);
 
 	// Store log intensity in the alpha channel: scale to 0,1 range.
@@ -245,7 +245,7 @@ void BasicEyeAdaptationCS(uint GIndex : SV_GroupIndex, uint2 GTId : SV_GroupThre
 	LogLumAve = (LogLumAve - EyeAdaptation_HistogramBias) / EyeAdaptation_HistogramScale;
 	
 	// Convert LogLuminanceAverage to Average Intensity
-	const float AverageSceneLuminance = View.OneOverPreExposure * exp2(LogLumAve);
+	const float AverageSceneLuminance = exp2(LogLumAve);
 	
 	// The Exposure Scale (and thus intensity) used in the previous frame
 	const float OldExposureScale = EyeAdaptationBuffer[0].x;
diff --git a/Engine/Source/Runtime/Renderer/Private/PostProcess/PostProcessEyeAdaptation.cpp b/Engine/Source/Runtime/Renderer/Private/PostProcess/PostProcessEyeAdaptation.cpp
--- a/Engine/Source/Runtime/Renderer/Private/PostProcess/PostProcessEyeAdaptation.cpp
+++ b/Engine/Source/Runtime/Renderer/Private/PostProcess/PostProcessEyeAdaptation.cpp
@@ -850,6 +850,7 @@ public:
 	SHADER_USE_PARAMETER_STRUCT(FBasicEyeAdaptationSetupPS, FGlobalShader);
 
 	BEGIN_SHADER_PARAMETER_STRUCT(FParameters, )
+		SHADER_PARAMETER_STRUCT_REF(FViewUniformShaderParameters, View)
 		SHADER_PARAMETER_STRUCT(FEyeAdaptationParameters, EyeAdaptation)
 		SHADER_PARAMETER_RDG_TEXTURE(Texture2D, ColorTexture)
 		SHADER_PARAMETER_SAMPLER(SamplerState, ColorSampler)
@@ -883,6 +884,7 @@ FScreenPassTexture AddBasicEyeAdaptationSetupPass(
 	const FScreenPassTextureViewport Viewport(SceneColor);
 
 	FBasicEyeAdaptationSetupPS::FParameters* PassParameters = GraphBuilder.AllocParameters<FBasicEyeAdaptationSetupPS::FParameters>();
+	PassParameters->View = View.ViewUniformBuffer;
 	PassParameters->EyeAdaptation = EyeAdaptationParameters;
 	PassParameters->ColorTexture = SceneColor.Texture;
 	PassParameters->ColorSampler = TStaticSamplerState<SF_Point, AM_Clamp, AM_Clamp, AM_Clamp>::GetRHI();
