[scene culling] Defer outstanding GPU-uploads until a future call to begin update.
- This is a more robust mechanism that doesn't rely on finding all places where it is not done at some point before RDG execute.
- Also: added some csv stats & force inlined some updater functions.

diff --git a/Engine/Source/Runtime/Renderer/Private/RendererScene.cpp b/Engine/Source/Runtime/Renderer/Private/RendererScene.cpp
--- a/Engine/Source/Runtime/Renderer/Private/RendererScene.cpp
+++ b/Engine/Source/Runtime/Renderer/Private/RendererScene.cpp
@@ -6574,13 +6574,6 @@ void FScene::UpdateAllPrimitiveSceneInfos(FRDGBuilder& GraphBuilder, EUpdateAllP
 		checkSlow(PersistentPrimitiveIdToIndexMap[PrimitiveSceneInfo->PersistentIndex.Index] == PrimitiveSceneInfo->PackedIndex);
 	}
 #endif
-
-	// No async going on, we must sync here or there will be buildup of accumulated & unuploaded crap.
-	// As a more robust alternative we may stipulate a list of callbacks that need to be done before the graph builder is executed. These may do cleanup in case no one did along the way.
-	if (AsyncOps == EUpdateAllPrimitiveSceneInfosAsyncOps::None)
-	{
-		SceneCulling->EndUpdate(GraphBuilder, false);
-	}
 }
 
 void FScene::CreateLightPrimitiveInteractionsForPrimitive(FPrimitiveSceneInfo* PrimitiveInfo)
diff --git a/Engine/Source/Runtime/Renderer/Private/SceneCulling/SceneCulling.cpp b/Engine/Source/Runtime/Renderer/Private/SceneCulling/SceneCulling.cpp
--- a/Engine/Source/Runtime/Renderer/Private/SceneCulling/SceneCulling.cpp
+++ b/Engine/Source/Runtime/Renderer/Private/SceneCulling/SceneCulling.cpp
@@ -68,6 +68,9 @@ DECLARE_DWORD_COUNTER_STAT(TEXT("Update Copied Id Count"), STAT_SceneCulling_Cop
 	#define UPDATE_BUILDER_STAT(Builder, StatId, Num) 
 #endif
 
+CSV_DEFINE_CATEGORY(SceneCulling, true);
+
+
 LLM_DECLARE_TAG_API(SceneCulling, RENDERER_API);
 DECLARE_LLM_MEMORY_STAT(TEXT("SceneCulling"), STAT_SceneCullingLLM, STATGROUP_LLMFULL);
 DECLARE_LLM_MEMORY_STAT(TEXT("SceneCulling"), STAT_SceneCullingSummaryLLM, STATGROUP_LLM);
@@ -598,7 +601,7 @@ void FSceneCulling::Empty()
 /**
  * Produce a world-space bounding sphere for an instance given local bounds and transforms.
  */
-inline FVector4d TransformBounds(VectorRegister4f VecOrigin, VectorRegister4f VecExtent, const FRenderTransform& LocalToPrimitive, const FMatrix44f& PrimitiveToWorldRotation, VectorRegister4Double PrimitiveToWorldTranslationVec)
+FORCEINLINE_DEBUGGABLE FVector4d TransformBounds(VectorRegister4f VecOrigin, VectorRegister4f VecExtent, const FRenderTransform& LocalToPrimitive, const FMatrix44f& PrimitiveToWorldRotation, VectorRegister4Double PrimitiveToWorldTranslationVec)
 {
 	// 1. Matrix Concat and bounds all in one
 	const VectorRegister4f B0 = VectorLoadAligned(PrimitiveToWorldRotation.M[0]);
@@ -659,14 +662,14 @@ inline FVector4d TransformBounds(VectorRegister4f VecOrigin, VectorRegister4f Ve
 
 struct FBoundsTransformerUniqueBounds
 {
-	inline FBoundsTransformerUniqueBounds(const FMatrix44d& PrimitiveToWorld, FPrimitiveSceneProxy* InSceneProxy)
+	FORCEINLINE_DEBUGGABLE FBoundsTransformerUniqueBounds(const FMatrix44d& PrimitiveToWorld, FPrimitiveSceneProxy* InSceneProxy)
 		: SceneProxy(InSceneProxy)
 	{
 		PrimitiveToWorldRotation = FMatrix44f(PrimitiveToWorld.RemoveTranslation());
 		PrimitiveToWorldTranslationVec = VectorLoadAligned(PrimitiveToWorld.M[3]);
 	}
 
-	inline FVector4d TransformBounds(int32 InstanceIndex, const FInstanceSceneData& PrimitiveInstance)
+	FORCEINLINE_DEBUGGABLE FVector4d TransformBounds(int32 InstanceIndex, const FInstanceSceneData& PrimitiveInstance)
 	{
 		const FRenderBounds InstanceBounds = SceneProxy->GetInstanceLocalBounds(InstanceIndex);
 		const VectorRegister4f VecMin = VectorLoadFloat3(&InstanceBounds.Min);
@@ -685,7 +688,7 @@ struct FBoundsTransformerUniqueBounds
 
 struct FBoundsTransformerSharedBounds
 {
-	inline FBoundsTransformerSharedBounds(const FMatrix44d& PrimitiveToWorld, FPrimitiveSceneProxy* SceneProxy)
+	FORCEINLINE_DEBUGGABLE FBoundsTransformerSharedBounds(const FMatrix44d& PrimitiveToWorld, FPrimitiveSceneProxy* SceneProxy)
 	{
 		PrimitiveToWorldRotation = FMatrix44f(PrimitiveToWorld.RemoveTranslation());
 		PrimitiveToWorldTranslationVec = VectorLoadAligned(PrimitiveToWorld.M[3]);
@@ -698,7 +701,7 @@ struct FBoundsTransformerSharedBounds
 		VecExtent = VectorMultiply(VectorSubtract(VecMax, VecMin), Half);
 	}
 
-	inline FVector4d TransformBounds(int32 InstanceIndex, const FInstanceSceneData& PrimitiveInstance)
+	FORCEINLINE_DEBUGGABLE FVector4d TransformBounds(int32 InstanceIndex, const FInstanceSceneData& PrimitiveInstance)
 	{
 		return ::TransformBounds(VecOrigin, VecExtent, PrimitiveInstance.LocalToPrimitive, PrimitiveToWorldRotation, PrimitiveToWorldTranslationVec);
 	}
@@ -712,14 +715,14 @@ struct FBoundsTransformerSharedBounds
 template <typename BoundsTransformerType>
 struct FHashLocationComputerFromBounds
 {
-	FHashLocationComputerFromBounds(const TConstArrayView<FInstanceSceneData> InInstanceSceneData, FSceneCulling::FSpatialHash& InSpatialHash, const FMatrix44d& PrimitiveToWorld, FPrimitiveSceneProxy* SceneProxy)
+	FORCEINLINE_DEBUGGABLE FHashLocationComputerFromBounds(const TConstArrayView<FInstanceSceneData> InInstanceSceneData, FSceneCulling::FSpatialHash& InSpatialHash, const FMatrix44d& PrimitiveToWorld, FPrimitiveSceneProxy* SceneProxy)
 		: BoundsTransformer(PrimitiveToWorld, SceneProxy)
 		, InstanceSceneData(InInstanceSceneData)
 		, SpatialHash(InSpatialHash)
 	{
 	}
 
-	FSceneCulling::FLocation64 CalcLoc(int32 InstanceIndex)
+	FORCEINLINE_DEBUGGABLE FSceneCulling::FLocation64 CalcLoc(int32 InstanceIndex)
 	{
 		const FInstanceSceneData& PrimitiveInstance = InstanceSceneData[InstanceIndex];
 		FVector4d InstanceWorldBound = BoundsTransformer.TransformBounds(InstanceIndex, PrimitiveInstance);
@@ -813,7 +816,7 @@ public:
 	{
 		int32 CurrentChunkCount = MaxChunkSize;
 		int32 CurrentChunkId = INDEX_NONE;
-		inline void EmitCurrentChunk()
+		FORCEINLINE_DEBUGGABLE void EmitCurrentChunk()
 		{
 			if (CurrentChunkId != INDEX_NONE)
 			{
@@ -822,7 +825,7 @@ public:
 			CurrentChunkId = INDEX_NONE;
 		}
 
-		inline void AddCompressedChunk(FSceneCullingBuilder& Builder, uint32 StartInstanceId)
+		FORCEINLINE_DEBUGGABLE void AddCompressedChunk(FSceneCullingBuilder& Builder, uint32 StartInstanceId)
 		{
 			Builder.TotalCellChunkDataCount += 1;
 			// Add directly to chunk headers to not upset current chunk packing
@@ -830,7 +833,7 @@ public:
 
 			//Builder->SceneCulling.InstanceIdToCellDataSlot[StartInstanceId] = FCellDataSlot { true, 0 };
 		}
-		inline void Add(FSceneCullingBuilder& Builder, uint32 InstanceId)
+		FORCEINLINE_DEBUGGABLE void Add(FSceneCullingBuilder& Builder, uint32 InstanceId)
 		{
 			TArray<uint32>& PackedCellData = Builder.SceneCulling.PackedCellData;
 			if (CurrentChunkCount >= MaxChunkSize)
@@ -846,7 +849,7 @@ public:
 			PackedCellData[ItemOffset] = InstanceId;
 		}
 
-		inline void AddRange(FSceneCullingBuilder& Builder, int32 InInstanceIdOffset, int32 InInstanceIdCount)
+		FORCEINLINE_DEBUGGABLE void AddRange(FSceneCullingBuilder& Builder, int32 InInstanceIdOffset, int32 InInstanceIdCount)
 		{
 			// First add all compressible ranges.
 			int32 InstanceIdOffset = InInstanceIdOffset;
@@ -868,7 +871,7 @@ public:
 		}
 
 
-		uint32 ReallocateChunkRange(FSceneCullingBuilder &Builder, uint32 PrevItemChunksOffset, uint32 PrevNumItemChunks)
+		FORCEINLINE_DEBUGGABLE uint32 ReallocateChunkRange(FSceneCullingBuilder &Builder, uint32 PrevItemChunksOffset, uint32 PrevNumItemChunks)
 		{
 			uint32 ItemChunksOffset = PrevItemChunksOffset;
 			// TODO: round allocation size to POT, or some multiple to reduce reallocations & fragmentation?
@@ -888,7 +891,7 @@ public:
 			return ItemChunksOffset;
 		}
 
-		void OutputChunkIds(FSceneCullingBuilder &Builder, uint32 ItemChunksOffset)
+		FORCEINLINE_DEBUGGABLE void OutputChunkIds(FSceneCullingBuilder &Builder, uint32 ItemChunksOffset)
 		{
 			int32 NumIds = PackedChunkIds.Num();
 			// 3. Copy the list of chunk IDs
@@ -904,7 +907,7 @@ public:
 			}
 		}
 
-		bool IsEmpty() const { return PackedChunkIds.IsEmpty(); }
+		FORCEINLINE_DEBUGGABLE bool IsEmpty() const { return PackedChunkIds.IsEmpty(); }
 
 		TArray<uint32, TInlineAllocator<32, SceneRenderingAllocator>> PackedChunkIds;
 	};
@@ -926,7 +929,7 @@ public:
 		int32 RemovedInstanceCount = 0;
 		FCellHeader PrevCellHeader = FCellHeader { InvalidCellFlag, InvalidCellFlag };
 
-		inline void FinalizeChunks(FSceneCullingBuilder& Builder)
+		FORCEINLINE_DEBUGGABLE void FinalizeChunks(FSceneCullingBuilder& Builder)
 		{
 			// TODO: pass in and share/reuse space? 
 			// Copied chunk data (only whole chunks)
@@ -1082,7 +1085,7 @@ public:
 			ItemChunksOffset = ReallocateChunkRange(Builder, IsValid(PrevCellHeader) ? PrevCellHeader.ItemChunksOffset : INDEX_NONE, PrevCellHeader.NumItemChunks);
 		}
 
-		inline void AddExistingChunk(FTempChunks& RetainedChunks, uint32 ExistingPackedChunk)
+		FORCEINLINE_DEBUGGABLE void AddExistingChunk(FTempChunks& RetainedChunks, uint32 ExistingPackedChunk)
 		{
 			// Builder.TotalCellChunkDataCount += 1;
 			// Add directly to chunk headers to not upset current chunk packing
@@ -1090,7 +1093,7 @@ public:
 		}
 	};
 
-	inline FHashElementId FindOrAddBlock(const FSceneCulling::FLocation64& BlockLoc)
+	FORCEINLINE_DEBUGGABLE FHashElementId FindOrAddBlock(const FSceneCulling::FLocation64& BlockLoc)
 	{
 		// Update cached block ID / loc on the assumption that when adding many instances they will often hit the same block
 		if (!CachedBlockId.IsValid() || !(CachedBlockLoc == BlockLoc))
@@ -1108,7 +1111,7 @@ public:
 		return CachedBlockId;
 	}
 
-	inline FTempCell& GetOrAddTempCell(const FSceneCulling::FLocation64& InstanceCellLoc)
+	FORCEINLINE_DEBUGGABLE FTempCell& GetOrAddTempCell(const FSceneCulling::FLocation64& InstanceCellLoc)
 	{
 		// Address of the cell-block touched
 		FSceneCulling::FLocation64 BlockLoc = ToLevelRelative(InstanceCellLoc, CellBlockDimLog2);
@@ -1148,7 +1151,7 @@ public:
 		return GetOrAddTempCell(CellOffset);
 	}
 
-	inline FTempCell& GetOrAddTempCell(int32 CellIndex)
+	FORCEINLINE_DEBUGGABLE FTempCell& GetOrAddTempCell(int32 CellIndex)
 	{
 		FCellHeader CellHeader = SceneCulling.CellHeaders[CellIndex];
 
@@ -1174,7 +1177,7 @@ public:
 		return TempCells[CellHeader.ItemChunksOffset];
 	}
 
-	inline int32 AddRange(const FSceneCulling::FLocation64& InstanceCellLoc, int32 InInstanceIdOffset, int32 InInstanceIdCount)
+	FORCEINLINE_DEBUGGABLE int32 AddRange(const FSceneCulling::FLocation64& InstanceCellLoc, int32 InInstanceIdOffset, int32 InInstanceIdCount)
 	{
 		UPDATE_BUILDER_STAT(*this, RangeCount, 1);
 
@@ -1185,7 +1188,7 @@ public:
 		return TempCell.CellOffset;
 	}
 
-	inline int32 AddToCell(const FSceneCulling::FLocation64& InstanceCellLoc, int32 InstanceId)
+	FORCEINLINE_DEBUGGABLE int32 AddToCell(const FSceneCulling::FLocation64& InstanceCellLoc, int32 InstanceId)
 	{
 		FTempCell& TempCell = GetOrAddTempCell(InstanceCellLoc);
 		TempCell.Add(*this, uint32(InstanceId));
@@ -1193,7 +1196,7 @@ public:
 	}
 
 	template <typename HashLocationComputerType>
-	inline void BuildInstanceRange(int32 InstanceDataOffset, int32 NumInstances, const TConstArrayView<FInstanceSceneData> InstanceSceneData, HashLocationComputerType HashLocationComputer, FSceneCulling::FCellIndexCacheEntry &CellIndexCacheEntry, bool bCompressRLE)
+	FORCEINLINE_DEBUGGABLE void BuildInstanceRange(int32 InstanceDataOffset, int32 NumInstances, const TConstArrayView<FInstanceSceneData> InstanceSceneData, HashLocationComputerType HashLocationComputer, FSceneCulling::FCellIndexCacheEntry &CellIndexCacheEntry, bool bCompressRLE)
 	{
 		FSceneCulling::FLocation64 PrevInstanceCellLoc;
 		int32 SameInstanceLocRunCount = 0;
@@ -1234,7 +1237,7 @@ public:
 		}
 	}
 
-	inline FHashElementId GetBlockId(const FSceneCulling::FLocation64& BlockLoc)
+	FORCEINLINE_DEBUGGABLE FHashElementId GetBlockId(const FSceneCulling::FLocation64& BlockLoc)
 	{
 		// Update cached block ID / loc on the assumption that when adding many instances they will often hit the same block
 		if (!CachedBlockId.IsValid() || !(CachedBlockLoc == BlockLoc))
@@ -1248,7 +1251,7 @@ public:
 		return CachedBlockId;
 	}
 
-	inline int32 GetCellIndex(const FSceneCulling::FLocation64 &CellLoc)
+	FORCEINLINE_DEBUGGABLE int32 GetCellIndex(const FSceneCulling::FLocation64 &CellLoc)
 	{
 		if (CachedCellIdIndex == INDEX_NONE || !(CachedCellLoc == CellLoc))
 		{
@@ -1282,6 +1285,7 @@ public:
 
 		for (FTempCell& TempCell : TempCells)
 		{
+			TotalRemovedInstances += TempCell.RemovedInstanceCount;
 			TempCell.FinalizeChunks(*this);
 		}
 
@@ -1349,7 +1353,7 @@ public:
 			BUILDER_LOG("CellHeaderUploader { %u, %u} -> %d", CellHeader.ItemChunksOffset, CellHeader.NumItemChunks, TempCell.CellOffset);
 		}
 
-		TempCells.Reset();
+		TempCells.Empty();
 	}
 
 	inline int32 AllocateCacheEntry()
@@ -1440,6 +1444,8 @@ public:
 		const int32 NumInstances = PrimitiveSceneInfo->GetNumInstanceSceneDataEntries();
 		const int32 InstanceDataOffset = PrimitiveSceneInfo->GetInstanceSceneDataOffset();
 
+		TotalAddedInstances += NumInstances;
+
 		// leave in unknown state
 		if (NumInstances == 0)
 		{
@@ -1645,6 +1651,8 @@ public:
 		const int32 InstanceDataOffset = PrimitiveSceneInfo->GetInstanceSceneDataOffset();
 		const int32 NumInstances = PrimitiveSceneInfo->GetNumInstanceSceneDataEntries();
 
+		TotalUpdatedInstances += NumInstances;
+	 
 		// Just leave on the list
 		if (PrevPrimitiveState.State == FPrimitiveState::UnCullable)
 		{
@@ -1934,6 +1942,7 @@ public:
 		{
 			UpdateInstances(ScenePostUpdateData.UpdatedPrimitiveIds[Index], ScenePostUpdateData.UpdatedPrimitiveSceneInfos[Index]);
 		}	
+
 		// Next process all added and added ones.
 		for (int32 Index = 0; Index < ScenePostUpdateData.AddedPrimitiveIds.Num(); ++Index)
 		{
@@ -2005,6 +2014,16 @@ public:
 				BUILDER_LOG("Invalid block index no upload or update: %d", BlockIndex);
 			}
 		}
+
+		// Make sure to release any references to data allocated through SceneRenderingAllocator
+		// TODO: maybe change to default allocator? Might make sense to keep the data allocated anyway.
+		DirtyChunks.Empty();
+		DirtyBlocks.Empty();
+		RemovedInstanceFlags.Empty();
+
+		CSV_CUSTOM_STAT(SceneCulling, NumUpdatedInstances,  TotalUpdatedInstances, ECsvCustomStatOp::Accumulate);
+		CSV_CUSTOM_STAT(SceneCulling, NumAddedInstances,  TotalAddedInstances, ECsvCustomStatOp::Accumulate);
+		CSV_CUSTOM_STAT(SceneCulling, NumRemovedInstances,  TotalRemovedInstances, ECsvCustomStatOp::Accumulate);
 	}
 
 	using FRemovedIntanceFlags = TBitArray<SceneRenderingAllocator>;
@@ -2036,6 +2055,10 @@ public:
 	TStructuredBufferScatterUploader<FCellHeader> CellHeaderUploader;
 	TStructuredBufferScatterUploader<uint32> ItemChunkUploader;
 
+	int32 TotalUpdatedInstances = 0;
+	int32 TotalAddedInstances = 0;
+	int32 TotalRemovedInstances = 0;
+
 #if SC_ENABLE_DETAILED_LOGGING
 	bool bIsLoggingEnabled = false;
 	FString SceneTag;
@@ -2111,11 +2134,17 @@ FSceneCulling::FUpdater::~FUpdater()
 	if (Implementation != nullptr)
 	{
 		PostUpdateTaskHandle.Wait();
+		delete Implementation;
 	}
 }
 
 FSceneCulling::FUpdater &FSceneCulling::BeginUpdate(FRDGBuilder& GraphBuilder)
 {	
+	if (Updater.Implementation != nullptr)
+	{
+		Updater.FinalizeAndClear(GraphBuilder, false);
+	}
+
 	check(Updater.Implementation == nullptr);
 
 	// Note: this only works in concert with the FGlobalComponentRecreateRenderStateContext on the CVarSceneCulling callback ensuring all geometry is re-registered
@@ -2127,7 +2156,7 @@ FSceneCulling::FUpdater &FSceneCulling::BeginUpdate(FRDGBuilder& GraphBuilder)
 
 	if (bIsEnabled)
 	{
-		Updater.Implementation = GraphBuilder.AllocObject<FSceneCullingBuilder>(*this);
+		Updater.Implementation = new FSceneCullingBuilder(*this);
 	}
 	else
 	{
@@ -2227,8 +2256,7 @@ void FSceneCulling::FUpdater::FinalizeAndClear(FRDGBuilder& GraphBuilder, bool b
 #if SC_ENABLE_DETAILED_LOGGING
 		Implementation->EndLogging();
 #endif
-
-		// We don't delete here, the RDG object allocator frees the data and we need it alive until RDG uploads are done.
+		delete Implementation;
 		Implementation = nullptr;
 	}
 }
diff --git a/Engine/Source/Runtime/Renderer/Private/SceneRendering.cpp b/Engine/Source/Runtime/Renderer/Private/SceneRendering.cpp
--- a/Engine/Source/Runtime/Renderer/Private/SceneRendering.cpp
+++ b/Engine/Source/Runtime/Renderer/Private/SceneRendering.cpp
@@ -4471,9 +4471,6 @@ static void RenderViewFamilies_RenderThread(FRHICommandListImmediate& RHICmdList
 			SceneRenderer->Render(GraphBuilder);
 		}
 
-		// Must call this to make sure that scene changes are uploaded even if there was no rendering done that used the hierarchy.
-		Scene->SceneCulling->EndUpdate(GraphBuilder, false);
-
 		SceneRenderer->FlushCrossGPUFences(GraphBuilder);
 
 		GraphBuilder.Execute();
