Change GPU-Scene data allocators over to FSpanAllocator such that they can shrink as well as grow (r.GPUScene.UseGrowOnlyAllocationPolicy can be used to disable the new behavior).
 - also change the free/allocate order for updated instances such that all free operations in the scene update are batched.

diff --git a/Engine/Source/Runtime/Engine/Private/SpanAllocator.cpp b/Engine/Source/Runtime/Engine/Private/SpanAllocator.cpp
--- a/Engine/Source/Runtime/Engine/Private/SpanAllocator.cpp
+++ b/Engine/Source/Runtime/Engine/Private/SpanAllocator.cpp
@@ -3,6 +3,55 @@
 #include "SpanAllocator.h"
 #include "CoreGlobals.h"
 
+FSpanAllocator::FSpanAllocator(bool bInGrowOnly) 
+	: CurrentMaxSize(0)
+	, PeakMaxSize(0)
+	, FirstNonEmptySpan(0)
+	, bGrowOnly(bInGrowOnly)
+{}
+
+
+void FSpanAllocator::Reset()
+{
+	CurrentMaxSize = 0;
+	PeakMaxSize = 0;
+	FirstNonEmptySpan = 0;
+	FreeSpans.Reset();
+	PendingFreeSpans.Reset();
+}
+
+void FSpanAllocator::Empty()
+{
+	CurrentMaxSize = 0;
+	PeakMaxSize = 0;
+	FirstNonEmptySpan = 0;
+	FreeSpans.Empty();
+	PendingFreeSpans.Empty();
+}
+
+#if DO_CHECK
+
+bool FSpanAllocator::IsFree(int32 Index) const
+{
+	// If outside the max size, it is considered free as the allocator can grow at will
+	if (Index >= CurrentMaxSize)
+	{
+		return true;
+	}
+
+	for (const auto& FreeSpan : FreeSpans)
+	{
+		if (Index >= FreeSpan.StartOffset && Index < FreeSpan.StartOffset + FreeSpan.Num)
+		{
+			return true;
+		}
+	}
+
+	return false;
+}
+
+#endif // DO_CHECK
+
 void FSpanAllocator::Consolidate()
 {
 	// Consolidation 
@@ -57,9 +106,9 @@ void FSpanAllocator::Consolidate()
 	}
 
 	// Trim last span
-	if (!FreeSpansTmp.IsEmpty() && FreeSpansTmp.Last().StartOffset + FreeSpansTmp.Last().Num == MaxSize)
+	if (!FreeSpansTmp.IsEmpty() && FreeSpansTmp.Last().StartOffset + FreeSpansTmp.Last().Num == CurrentMaxSize)
 	{
-		MaxSize -= FreeSpansTmp.Last().Num;
+		CurrentMaxSize -= FreeSpansTmp.Last().Num;
 		FreeSpansTmp.Pop(false);
 	}
 
diff --git a/Engine/Source/Runtime/Engine/Public/SpanAllocator.h b/Engine/Source/Runtime/Engine/Public/SpanAllocator.h
--- a/Engine/Source/Runtime/Engine/Public/SpanAllocator.h
+++ b/Engine/Source/Runtime/Engine/Public/SpanAllocator.h
@@ -14,10 +14,10 @@
 class FSpanAllocator
 {
 public:
-	FSpanAllocator() 
-		: MaxSize(0)
-		, FirstNonEmptySpan(0)
-	{}
+	/**
+	 * If bInGrowOnly is true the size reported by GetMaxSize() will never decrease except when for when Reset() or Empty() is called.
+	 */
+	ENGINE_API FSpanAllocator(bool bInGrowOnly = false);
 
 	// Allocate a range.  Returns allocated StartOffset.
 	inline int32 Allocate(int32 Num = 1)
@@ -49,8 +49,10 @@ public:
 		}
 
 		// New allocation
-		int32 StartOffset = MaxSize;
-		MaxSize = MaxSize + Num;
+		int32 StartOffset = CurrentMaxSize;
+		CurrentMaxSize = CurrentMaxSize + Num;
+
+		PeakMaxSize = FMath::Max(PeakMaxSize, CurrentMaxSize);
 
 		return StartOffset;
 	}
@@ -58,29 +60,17 @@ public:
 	// Free an already allocated range.  
 	FORCEINLINE void Free(int32 BaseOffset, int32 Num = 1)
 	{
-		checkSlow(BaseOffset + Num <= MaxSize);
+		checkSlow(BaseOffset + Num <= CurrentMaxSize);
 		PendingFreeSpans.Add(FLinearAllocation{ BaseOffset, Num });
 	}
 
-	void Reset()
-	{
-		MaxSize = 0;
-		FirstNonEmptySpan = 0;
-		FreeSpans.Reset();
-		PendingFreeSpans.Reset();
-	}
+	ENGINE_API void Reset();
 
-	void Empty()
-	{
-		MaxSize = 0;
-		FirstNonEmptySpan = 0;
-		FreeSpans.Empty();
-		PendingFreeSpans.Empty();
-	}
+	ENGINE_API void Empty();
 
 	inline int32 GetSparselyAllocatedSize() const
 	{
-		int32 AllocatedSize = MaxSize;
+		int32 AllocatedSize = CurrentMaxSize;
 
 		for (int32 i = 0; i < FreeSpans.Num(); i++)
 		{
@@ -92,7 +82,7 @@ public:
 
 	FORCEINLINE int32 GetMaxSize() const
 	{
-		return MaxSize;
+		return bGrowOnly ? PeakMaxSize : CurrentMaxSize;
 	}
 
 	FORCEINLINE int32 GetNumFreeSpans() const
@@ -115,24 +105,7 @@ public:
 	 * Loop over all free spans and check if Index is in any of them.
 	 * Note: Very costly, only intended for debugging use, and probably best if under a toggle even then.
 	 */
-	FORCEINLINE bool IsFree(int32 Index) const
-	{
-		// If outside the max size, it is considered free as the allocator can grow at will
-		if (Index >= MaxSize)
-		{
-			return true;
-		}
-
-		for (const auto& FreeSpan : FreeSpans)
-		{
-			if (Index >= FreeSpan.StartOffset && Index < FreeSpan.StartOffset + FreeSpan.Num)
-			{
-				return true;
-			}
-		}
-
-		return false;
-	}
+	ENGINE_API bool IsFree(int32 Index) const;
 #endif // DO_CHECK
 
 	/**
@@ -155,13 +128,16 @@ private:
 	};
 
 	// Size of the linear range used by the allocator
-	int32 MaxSize;
+	int32 CurrentMaxSize;
+	// Peak tracked size since last Reset or Empty
+	int32 PeakMaxSize;
 	// First span in free list that is not empty, used to speed up search for free items
 	int32 FirstNonEmptySpan;
 	// Ordered free list from low to high
 	TArray<FLinearAllocation, TInlineAllocator<10>> FreeSpans;
 	// Unordered list of freed items since last consolidate
 	TArray<FLinearAllocation, TInlineAllocator<10>> PendingFreeSpans;
+	bool bGrowOnly;
 
 	ENGINE_API int32 SearchFreeList(int32 Num, int32 SearchStartIndex = 0);
 };
\ No newline at end of file
diff --git a/Engine/Source/Runtime/Renderer/Private/GPUScene.cpp b/Engine/Source/Runtime/Renderer/Private/GPUScene.cpp
--- a/Engine/Source/Runtime/Renderer/Private/GPUScene.cpp
+++ b/Engine/Source/Runtime/Renderer/Private/GPUScene.cpp
@@ -29,6 +29,9 @@
 #include "SceneDefinitions.h"
 #include "PrimitiveSceneShaderData.h"
 
+// Defaults to being disabled, enable using the command line argument: -CsvCategory GPUScene
+CSV_DEFINE_CATEGORY(GPUScene, false);
+
 #define LOG_INSTANCE_ALLOCATIONS 0
 
 static void ConstructDefault(FGPUSceneResourceParameters& GPUScene, FRDGBuilder& GraphBuilder)
@@ -47,42 +50,23 @@ static void ConstructDefault(FGPUSceneResourceParameters& GPUScene, FRDGBuilder&
 }
 IMPLEMENT_SCENE_UB_STRUCT(FGPUSceneResourceParameters, GPUScene, ConstructDefault);
 
-int32 GGPUSceneUploadEveryFrame = 0;
-FAutoConsoleVariableRef CVarGPUSceneUploadEveryFrame(
+static TAutoConsoleVariable<int32> CVarGPUSceneUploadEveryFrame(
 	TEXT("r.GPUScene.UploadEveryFrame"),
-	GGPUSceneUploadEveryFrame,
+	0,
 	TEXT("Whether to upload the entire scene's primitive data every frame.  Useful for debugging."),
 	ECVF_RenderThreadSafe
 );
 
-int32 GGPUSceneValidatePrimitiveBuffer = 0;
-FAutoConsoleVariableRef CVarGPUSceneValidatePrimitiveBuffer(
-	TEXT("r.GPUScene.ValidatePrimitiveBuffer"),
-	GGPUSceneValidatePrimitiveBuffer,
-	TEXT("Whether to readback the GPU primitive data and assert if it doesn't match the RT primitive data.  Useful for debugging."),
-	ECVF_RenderThreadSafe
-);
-
-int32 GGPUSceneValidateInstanceBuffer = 0;
-FAutoConsoleVariableRef CVarGPUSceneValidateInstanceBuffer(
-	TEXT("r.GPUScene.ValidateInstanceBuffer"),
-	GGPUSceneValidateInstanceBuffer,
-	TEXT("Whether to readback the GPU instance data and assert if it doesn't match the RT primitive data.  Useful for debugging."),
-	ECVF_RenderThreadSafe
-);
-
-int32 GGPUSceneMaxPooledUploadBufferSize = 256000;
-FAutoConsoleVariableRef CVarGGPUSceneMaxPooledUploadBufferSize(
+static TAutoConsoleVariable<int32> CVarGPUSceneMaxPooledUploadBufferSize(
 	TEXT("r.GPUScene.MaxPooledUploadBufferSize"),
-	GGPUSceneMaxPooledUploadBufferSize,
+	256000,
 	TEXT("Maximum size of GPU Scene upload buffer size to pool."),
 	ECVF_RenderThreadSafe
 );
 
-int32 GGPUSceneParallelUpdate = 0;
-FAutoConsoleVariableRef CVarGPUSceneParallelUpdate(
+static TAutoConsoleVariable<int32> CVarGPUSceneParallelUpdate(
 	TEXT("r.GPUScene.ParallelUpdate"),
-	GGPUSceneParallelUpdate,
+	0,
 	TEXT(""),
 	ECVF_RenderThreadSafe
 );
@@ -114,23 +98,14 @@ static TAutoConsoleVariable<float> CVarGPUSceneDebugDrawRange(
 	ECVF_RenderThreadSafe
 );
 
-static int32 GGPUSceneAllowDeferredAllocatorMerges = 1;
-FAutoConsoleVariableRef CVarGPUSceneAllowDeferredAllocatorMerges(
-	TEXT("r.GPUScene.AllowDeferredAllocatorMerges"),
-	GGPUSceneAllowDeferredAllocatorMerges,
-	TEXT(""),
-	ECVF_RenderThreadSafe
-);
-
-int32 GGPUSceneInstanceUploadViaCreate = 1;
-FAutoConsoleVariableRef CVarGPUSceneInstanceUploadViaCreate(
-	TEXT("r.GPUScene.InstanceUploadViaCreate"),
-	GGPUSceneInstanceUploadViaCreate,
-	TEXT("When uploading GPUScene InstanceData, upload via resource creation when the RHI supports it efficiently."),
-	ECVF_RenderThreadSafe
+static TAutoConsoleVariable<int32> CVarGPUSceneUseGrowOnlyAllocationPolicy(
+	TEXT("r.GPUScene.UseGrowOnlyAllocationPolicy"),
+	0,
+	TEXT("Deprecated 5.3. If set to 1 the allocators used for GPU-scene instances and similar will use a grow-only allocation policy to mimic the behavior in 5.2 and earlier.\n")
+	TEXT("  Disabled by default, which means that the buffers can shrink as well as grow."),
+	ECVF_RenderThreadSafe | ECVF_ReadOnly
 );
 
-
 LLM_DECLARE_TAG_API(GPUScene, RENDERER_API);
 DECLARE_LLM_MEMORY_STAT(TEXT("GPUScene"), STAT_GPUSceneLLM, STATGROUP_LLMFULL);
 DECLARE_LLM_MEMORY_STAT(TEXT("GPUScene"), STAT_GPUSceneSummaryLLM, STATGROUP_LLM);
@@ -579,6 +554,16 @@ void FGPUScene::SetEnabled(ERHIFeatureLevel::Type InFeatureLevel)
 	bIsEnabled = UseGPUScene(GMaxRHIShaderPlatform, FeatureLevel);
 }
 
+FGPUScene::FGPUScene()
+	: bUpdateAllPrimitives(false)
+	, InstanceSceneDataAllocator(CVarGPUSceneUseGrowOnlyAllocationPolicy.GetValueOnAnyThread() != 0)
+	, InstanceSceneDataSOAStride(0)
+	, InstancePayloadDataAllocator(CVarGPUSceneUseGrowOnlyAllocationPolicy.GetValueOnAnyThread() != 0)
+	, LightmapDataAllocator(CVarGPUSceneUseGrowOnlyAllocationPolicy.GetValueOnAnyThread() != 0)
+{
+}
+
+
 FGPUScene::~FGPUScene()
 {
 }
@@ -676,11 +661,23 @@ void FGPUScene::UpdateInternal(FRDGBuilder& GraphBuilder, FSceneUniformBuffer& S
 	ensure(NumScenePrimitives == Scene.Primitives.Num());
 	ensure(DynamicPrimitivesOffset >= Scene.Primitives.Num());
 
+	CSV_CUSTOM_STAT(GPUScene, InstanceAllocMaxSize, InstanceSceneDataAllocator.GetMaxSize(), ECsvCustomStatOp::Set);
+	CSV_CUSTOM_STAT(GPUScene, InstanceAllocUsedSize, InstanceSceneDataAllocator.GetSparselyAllocatedSize(), ECsvCustomStatOp::Set);
+
+	CSV_CUSTOM_STAT(GPUScene, InstancePayloadAllocMaxSize, InstancePayloadDataAllocator.GetMaxSize(), ECsvCustomStatOp::Set);
+	CSV_CUSTOM_STAT(GPUScene, InstancePayloadAllocUsedSize, InstancePayloadDataAllocator.GetSparselyAllocatedSize(), ECsvCustomStatOp::Set);
+
+	CSV_CUSTOM_STAT(GPUScene, LightmapDataAllocMaxSize, LightmapDataAllocator.GetMaxSize(), ECsvCustomStatOp::Set);
+	CSV_CUSTOM_STAT(GPUScene, LightmapDataAllocUsedSize, LightmapDataAllocator.GetSparselyAllocatedSize(), ECsvCustomStatOp::Set);
+
 	RDG_EVENT_SCOPE(GraphBuilder, "GPUScene.Update");
 
+	// Do this stat separately since it should always be on.
+	CSV_CUSTOM_STAT_GLOBAL(GPUSceneInstanceCount, float(InstanceSceneDataAllocator.GetMaxSize()), ECsvCustomStatOp::Set);
+
 	LastDeferredGPUWritePass = EGPUSceneGPUWritePass::None;
 
-	if (GGPUSceneUploadEveryFrame || bUpdateAllPrimitives)
+	if ((CVarGPUSceneUploadEveryFrame.GetValueOnRenderThread() != 0) || bUpdateAllPrimitives)
 	{
 		PrimitivesToUpdate.Reset();
 
@@ -770,8 +767,6 @@ void FGPUScene::UpdateBufferState(FRDGBuilder& GraphBuilder, FSceneUniformBuffer
 	const uint32 SizeReserve = FMath::RoundUpToPowerOfTwo(FMath::Max(DynamicPrimitivesOffset, InitialBufferSize));
 	BufferState.PrimitiveBuffer = ResizeStructuredBufferIfNeeded(GraphBuilder, PrimitiveBuffer, SizeReserve * sizeof(FPrimitiveSceneShaderData::Data), TEXT("GPUScene.PrimitiveData"));
 
-	CSV_CUSTOM_STAT_GLOBAL(GPUSceneInstanceCount, float(InstanceSceneDataAllocator.GetMaxSize()), ECsvCustomStatOp::Accumulate);
-
 	const uint32 InstanceSceneDataSizeReserve = FMath::RoundUpToPowerOfTwo(FMath::Max(InstanceSceneDataAllocator.GetMaxSize(), InitialBufferSize));
 	FResizeResourceSOAParams ResizeParams;
 	ResizeParams.NumBytes = InstanceSceneDataSizeReserve * FInstanceSceneShaderData::GetEffectiveNumBytes();
@@ -913,7 +908,7 @@ void FGPUScene::UploadGeneral(FRDGBuilder& GraphBuilder, FScene& Scene, FRDGExte
 		return;
 	}
 
-	const bool bExecuteInParallel = GGPUSceneParallelUpdate != 0 && FApp::ShouldUseThreadingForPerformance();
+	const bool bExecuteInParallel = CVarGPUSceneParallelUpdate.GetValueOnRenderThread() != 0 && FApp::ShouldUseThreadingForPerformance();
 	const bool bNaniteEnabled = DoesPlatformSupportNanite(GMaxRHIShaderPlatform);
 
 	SCOPED_NAMED_EVENT(UpdateGPUScene, FColor::Green);
@@ -1306,28 +1301,28 @@ void FGPUScene::UploadGeneral(FRDGBuilder& GraphBuilder, FScene& Scene, FRDGExte
 			Scene.NaniteMaterials[NaniteMeshPassIndex].Finish(GraphBuilder, ExternalAccessQueue, TaskContext.NaniteMaterialUploaders[NaniteMeshPassIndex]);
 		}
 	}
-
-	if (PrimitiveUploadBuffer.GetNumBytes() > (uint32)GGPUSceneMaxPooledUploadBufferSize)
+	const uint32 MaxPooledSize = uint32(CVarGPUSceneMaxPooledUploadBufferSize.GetValueOnRenderThread());
+	if (PrimitiveUploadBuffer.GetNumBytes() > MaxPooledSize)
 	{
 		PrimitiveUploadBuffer.Release();
 	}
 
-	if (InstanceSceneUploadBuffer.GetNumBytes() > (uint32)GGPUSceneMaxPooledUploadBufferSize)
+	if (InstanceSceneUploadBuffer.GetNumBytes() > MaxPooledSize)
 	{
 		InstanceSceneUploadBuffer.Release();
 	}
 
-	if (InstancePayloadUploadBuffer.GetNumBytes() > (uint32)GGPUSceneMaxPooledUploadBufferSize)
+	if (InstancePayloadUploadBuffer.GetNumBytes() > MaxPooledSize)
 	{
 		InstancePayloadUploadBuffer.Release();
 	}
 
-	if (InstanceBVHUploadBuffer.GetNumBytes() > (uint32)GGPUSceneMaxPooledUploadBufferSize)
+	if (InstanceBVHUploadBuffer.GetNumBytes() > MaxPooledSize)
 	{
 		InstanceBVHUploadBuffer.Release();
 	}
 
-	if (LightmapUploadBuffer.GetNumBytes() > (uint32)GGPUSceneMaxPooledUploadBufferSize)
+	if (LightmapUploadBuffer.GetNumBytes() > MaxPooledSize)
 	{
 		LightmapUploadBuffer.Release();
 	}
@@ -1868,27 +1863,14 @@ void FGPUScene::DebugRender(FRDGBuilder& GraphBuilder, FScene& Scene, FViewInfo&
 	}
 }
 
-
-void FGPUScene::BeginDeferAllocatorMerges()
+void FGPUScene::ConsolidateInstanceDataAllocations()
 {
-	if (GGPUSceneAllowDeferredAllocatorMerges != 0)
-	{
-		InstanceSceneDataAllocator.BeginDeferMerges();
-		InstancePayloadDataAllocator.BeginDeferMerges();
-		LightmapDataAllocator.BeginDeferMerges();
-	}
-}
+	CSV_SCOPED_TIMING_STAT_EXCLUSIVE(ConsolidateInstanceDataAllocations);
 
-
-void FGPUScene::EndDeferAllocatorMerges()
-{
-	if (GGPUSceneAllowDeferredAllocatorMerges != 0)
-	{
+	InstanceSceneDataAllocator.Consolidate();
+	InstancePayloadDataAllocator.Consolidate();
+	LightmapDataAllocator.Consolidate();
 		SCOPED_NAMED_EVENT(FGPUScene_EndDeferAllocatorMerges, FColor::Green);
-		InstanceSceneDataAllocator.EndDeferMerges();
-		InstancePayloadDataAllocator.EndDeferMerges();
-		LightmapDataAllocator.EndDeferMerges();
-	}
 }
 
 
diff --git a/Engine/Source/Runtime/Renderer/Private/GPUScene.h b/Engine/Source/Runtime/Renderer/Private/GPUScene.h
--- a/Engine/Source/Runtime/Renderer/Private/GPUScene.h
+++ b/Engine/Source/Runtime/Renderer/Private/GPUScene.h
@@ -6,6 +6,7 @@
 #include "RendererInterface.h"
 #include "PrimitiveUniformShaderParameters.h"
 #include "PrimitiveSceneInfo.h"
+#include "SpanAllocator.h"
 #include "GrowOnlySpanAllocator.h"
 #include "InstanceCulling/InstanceCullingLoadBalancer.h"
 #include "MeshBatch.h"
@@ -185,11 +186,7 @@ struct FGPUSceneBufferState
 class FGPUScene
 {
 public:
-	FGPUScene()
-		: bUpdateAllPrimitives(false)
-		, InstanceSceneDataSOAStride(0)
-	{
-	}
+	FGPUScene();
 	~FGPUScene();
 
 	void SetEnabled(ERHIFeatureLevel::Type InFeatureLevel);
@@ -299,7 +296,7 @@ public:
 	int32 GetNumPrimitives() const { return DynamicPrimitivesOffset; }
 	int32 GetNumLightmapDataItems() const { return LightmapDataAllocator.GetMaxSize(); }
 
-	const FGrowOnlySpanAllocator& GetInstanceSceneDataAllocator() const { return InstanceSceneDataAllocator; }
+	const FSpanAllocator& GetInstanceSceneDataAllocator() const { return InstanceSceneDataAllocator; }
 
 	/**
 	 * Return the GPU scene resource
@@ -313,11 +310,9 @@ public:
 	void DebugRender(FRDGBuilder& GraphBuilder, FScene& Scene, FViewInfo& View);
 
 	/**
-	 * Between these calls to FGrowOnlySpanAllocator::Free just appends the allocation to the free list, rather than trying to merge with existing allocations.
-	 * At EndDeferAllocatorMerges the free list is consolidated by sorting and merging all spans. This amortises the cost of the merge over many calls.
+	 * Manually trigger an allocator consolidate (will otherwise be done when an item is allocated).
 	 */
-	void BeginDeferAllocatorMerges();
-	void EndDeferAllocatorMerges();
+	void ConsolidateInstanceDataAllocations();
 
 	/**
 	 * Executes GPUScene writes that were deferred until a later point in scene rendering
@@ -337,12 +332,12 @@ public:
 	FRDGAsyncScatterUploadBuffer   PrimitiveUploadBuffer;
 
 	/** GPU primitive instance list */
-	FGrowOnlySpanAllocator         InstanceSceneDataAllocator;
+	FSpanAllocator		           InstanceSceneDataAllocator;
 	TRefCountPtr<FRDGPooledBuffer> InstanceSceneDataBuffer;
 	FRDGAsyncScatterUploadBuffer   InstanceSceneUploadBuffer;
 	uint32                         InstanceSceneDataSOAStride;	// Distance between arrays in float4s
 
-	FGrowOnlySpanAllocator         InstancePayloadDataAllocator;
+	FSpanAllocator                 InstancePayloadDataAllocator;
 	TRefCountPtr<FRDGPooledBuffer> InstancePayloadDataBuffer;
 	FRDGAsyncScatterUploadBuffer   InstancePayloadUploadBuffer;
 
@@ -350,9 +345,9 @@ public:
 	FRDGAsyncScatterUploadBuffer   InstanceBVHUploadBuffer;
 
 	/** GPU light map data */
-	FGrowOnlySpanAllocator         LightmapDataAllocator;
+	FSpanAllocator                 LightmapDataAllocator;
 	TRefCountPtr<FRDGPooledBuffer> LightmapDataBuffer;
-	FRDGAsyncScatterUploadBuffer        LightmapUploadBuffer;
+	FRDGAsyncScatterUploadBuffer   LightmapUploadBuffer;
 
 	struct FInstanceRange
 	{
diff --git a/Engine/Source/Runtime/Renderer/Private/RendererScene.cpp b/Engine/Source/Runtime/Renderer/Private/RendererScene.cpp
--- a/Engine/Source/Runtime/Renderer/Private/RendererScene.cpp
+++ b/Engine/Source/Runtime/Renderer/Private/RendererScene.cpp
@@ -233,126 +233,6 @@ public:
 	}
 };
 
-
-#define ENABLE_LOG_PRIMITIVE_INSTANCE_ID_STATS_TO_CSV 0
-
-#if ENABLE_LOG_PRIMITIVE_INSTANCE_ID_STATS_TO_CSV
-
-int32 GDumpPrimitiveAllocatorStatsToCSV = 0;
-FAutoConsoleVariableRef CVarDumpPrimitiveStatsToCSV(
-	TEXT("r.DumpPrimitiveStatsToCSV"),
-	GDumpPrimitiveAllocatorStatsToCSV,
-	TEXT("Dump primitive and instance stats to CSV\n")
-	TEXT(" 0 - stop recording, dump to csv and clear array.\n")
-	TEXT(" 1 - start recording into array\n")
-	TEXT(" 2 - dump to csv and continue recording without clearing array\n"),
-	ECVF_RenderThreadSafe
-);
-
-static constexpr int32 StatStride = 8;
-TArray<int32> GPrimitiveAllocatorStats;
-
-void DumpPrimitiveAllocatorStats()
-{
-	if (!GPrimitiveAllocatorStats.IsEmpty())
-	{
-		FString FileName = FPaths::ProjectLogDir() / TEXT("PrimitiveStats-") + FDateTime::Now().ToString() + TEXT(".csv");
-
-		UE_LOG(LogRenderer, Log, TEXT("Dumping primitive allocator stats trace to: '%s'"), *FileName);
-
-		FArchive* FileToLogTo = IFileManager::Get().CreateFileWriter(*FileName, false);
-		
-		ensure(FileToLogTo);
-		if (FileToLogTo)
-		{
-			static const FString StatNames[StatStride] =
-			{
-				TEXT("NumPrimitives"),
-				TEXT("MaxPersistent"),
-				TEXT("NumPersistentAllocated"),
-				TEXT("PersistentFreeListSizeBC"),
-				TEXT("PersistentPendingListSizeBC"),
-				TEXT("PersistentFreeListSize"),
-				TEXT("MaxInstances"),
-				TEXT("NumInstancesAllocated"),
-			};
-
-			// Print header
-			FString StringToPrint;
-			for (int32 Index = 0; Index < StatStride; ++Index)
-			{
-				if (!StringToPrint.IsEmpty())
-				{
-					StringToPrint += TEXT(",");
-				}
-				if (Index < int32(UE_ARRAY_COUNT(StatNames)))
-				{
-					StringToPrint.Append(StatNames[Index]);
-				}
-				else
-				{
-					StringToPrint.Appendf(TEXT("Stat_%d"), Index);
-				}
-			}
-
-			StringToPrint += TEXT("\n");
-			FileToLogTo->Serialize(TCHAR_TO_ANSI(*StringToPrint), StringToPrint.Len());
-
-			int32 Num = GPrimitiveAllocatorStats.Num() / StatStride;
-			for (int32 Ind = 0; Ind < Num; ++Ind)
-			{
-				StringToPrint.Empty();
-
-				for (int32 StatInd = 0; StatInd < StatStride; ++StatInd)
-				{
-					if (!StringToPrint.IsEmpty())
-					{
-						StringToPrint.Append(TEXT(","));
-					}
-
-					StringToPrint.Appendf(TEXT("%d"), GPrimitiveAllocatorStats[Ind * StatStride + StatInd]);
-				}
-
-				StringToPrint += TEXT("\n");
-				FileToLogTo->Serialize(TCHAR_TO_ANSI(*StringToPrint), StringToPrint.Len());
-			}
-
-
-			FileToLogTo->Close();
-		}
-	}
-}
-
-void UpdatePrimitiveAllocatorStats(int32 NumPrimitives, int32 MaxPersistent, int32 NumPersistentAllocated, int32 PersistemFreeListSize, int32 PersistemFreeListSizeBC, int32 PersistentPendingListSizeBC, int32 MaxInstances, int32 NumInstancesAllocated)
-{
-	// Dump and clear when turning off the cvar
-	if (GDumpPrimitiveAllocatorStatsToCSV == 0 && !GPrimitiveAllocatorStats.IsEmpty())
-	{
-		DumpPrimitiveAllocatorStats();
-		GPrimitiveAllocatorStats.Empty();
-	}
-
-	// Dump snapshot (and don't clear or stop) when cvar is set to 2
-	if (GDumpPrimitiveAllocatorStatsToCSV == 2 && !GPrimitiveAllocatorStats.IsEmpty())
-	{
-		DumpPrimitiveAllocatorStats();
-		GDumpPrimitiveAllocatorStatsToCSV = 1;
-	}
-
-	if (GDumpPrimitiveAllocatorStatsToCSV != 0)
-	{
-		GPrimitiveAllocatorStats.Add(NumPrimitives);
-		GPrimitiveAllocatorStats.Add(MaxPersistent);
-		GPrimitiveAllocatorStats.Add(NumPersistentAllocated);
-		GPrimitiveAllocatorStats.Add(PersistemFreeListSizeBC);
-		GPrimitiveAllocatorStats.Add(PersistentPendingListSizeBC);
-		GPrimitiveAllocatorStats.Add(PersistemFreeListSize);
-		GPrimitiveAllocatorStats.Add(MaxInstances);
-		GPrimitiveAllocatorStats.Add(NumInstancesAllocated);
-	}
-}
-#endif // ENABLE_LOG_PRIMITIVE_INSTANCE_ID_STATS_TO_CSV
-
 /**
  * Holds the info to update SpeedTree wind per unique tree object in the scene, instead of per instance
  */
@@ -5599,8 +5479,6 @@ void FScene::UpdateAllPrimitiveSceneInfos(FRDGBuilder& GraphBuilder, EUpdateAllP
 		SCOPED_NAMED_EVENT(FScene_RemovePrimitiveSceneInfos, FColor::Red);
 		SCOPE_CYCLE_COUNTER(STAT_RemoveScenePrimitiveTime);
 
-		GPUScene.BeginDeferAllocatorMerges();
-
 		for (FPrimitiveSceneInfo* PrimitiveSceneInfo : RemovedLocalPrimitiveSceneInfos)
 		{
 			// clear it up, parent is getting removed
@@ -5814,7 +5692,6 @@ void FScene::UpdateAllPrimitiveSceneInfos(FRDGBuilder& GraphBuilder, EUpdateAllP
 			RemovedLocalPrimitiveSceneInfos.RemoveAt(StartIndex, RemovedLocalPrimitiveSceneInfos.Num() - StartIndex, false);
 		}
 	
-		GPUScene.EndDeferAllocatorMerges();
 	}
 
 	bool bNeedPathTracedInvalidation = false;
@@ -5862,14 +5739,90 @@ void FScene::UpdateAllPrimitiveSceneInfos(FRDGBuilder& GraphBuilder, EUpdateAllP
 		return false;
 	};
 
+#if RHI_RAYTRACING
+	TArray<FPrimitiveSceneInfo*, SceneRenderingAllocator> RayTracingPrimitivesToUpdate;
+	RayTracingPrimitivesToUpdate.Reserve(UpdatedInstances.Num());
+	bool bUpdateCachedRayTracingInstances = false;
+#endif
+	
+	TArray<FPrimitiveSceneInfo*, SceneRenderingAllocator> UpdatedInstancesWithInstanceCountChange;
+	UpdatedInstancesWithInstanceCountChange.Reserve(UpdatedInstances.Num());
+	// Perform instance releases before adds, such that allocator consolidation can happen only once.
+	{
+		CSV_SCOPED_TIMING_STAT_EXCLUSIVE(UpdatePrimitiveInstances);
+		SCOPED_NAMED_EVENT(FScene_UpdatePrimitiveInstances, FColor::Emerald);
+		SCOPE_CYCLE_COUNTER(STAT_UpdatePrimitiveInstanceRenderThreadTime);
+
+		for (const auto& UpdateInstance : UpdatedInstances)
+		{
+			FPrimitiveSceneProxy* PrimitiveSceneProxy = UpdateInstance.Key;
+			FScopeCycleCounter Context(PrimitiveSceneProxy->GetStatId());
+			FPrimitiveSceneInfo* PrimitiveSceneInfo = PrimitiveSceneProxy->GetPrimitiveSceneInfo();
+			
+			// being added or deleted, skip update logic
+			if (PrimitiveSceneInfo->PackedIndex == INDEX_NONE || DeletedSceneInfos.Contains(PrimitiveSceneInfo))
+			{
+				continue;
+			}
+			
+			QueueFlushVirtualTexture(PrimitiveSceneInfo);
+
+			// If we recorded no adds or removes the instance count has stayed the same.  Therefore the cached mesh draw commands do not
+			// need to be updated.  In situations where the instance count changes the mesh draw command stores the instance count which
+			// would need to be updated.
+			const bool bInstanceCountChanged = (UpdateInstance.Value.CmdBuffer.NumAdds > 0) || (UpdateInstance.Value.CmdBuffer.NumRemoves > 0);
+			const bool bUpdateStaticDrawLists = !PrimitiveSceneProxy->StaticElementsAlwaysUseProxyPrimitiveUniformBuffer() || bInstanceCountChanged;
+
+			if (QueueAddToScene(PrimitiveSceneInfo))
+			{
+				PrimitiveSceneInfo->RemoveFromScene(bUpdateStaticDrawLists);
+
+				if (bUpdateStaticDrawLists)
+				{
+					QueueAddStaticMeshes(PrimitiveSceneInfo);
+				}
+				else
+				{
+#if RHI_RAYTRACING
+					RayTracingPrimitivesToUpdate.Add(PrimitiveSceneInfo);
+					bUpdateCachedRayTracingInstances = true;
+#endif
+				}
+			}
+
+			// Update the Proxy's data.
+			PrimitiveSceneProxy->UpdateInstances_RenderThread(UpdateInstance.Value.CmdBuffer, UpdateInstance.Value.WorldBounds, UpdateInstance.Value.LocalBounds, UpdateInstance.Value.StaticMeshBounds);
+
+			if (!RHISupportsVolumeTextures(GetFeatureLevel())
+				&& (PrimitiveSceneProxy->IsMovable() || PrimitiveSceneProxy->NeedsUnbuiltPreviewLighting() || PrimitiveSceneProxy->GetLightmapType() == ELightmapType::ForceVolumetric))
+			{
+				PrimitiveSceneInfo->MarkIndirectLightingCacheBufferDirty();
+			}
+
+			if (bInstanceCountChanged)
+			{
+				PrimitiveSceneInfo->FreeGPUSceneInstances();
+				UpdatedInstancesWithInstanceCountChange.Add(PrimitiveSceneInfo);
+			}
+			else
+			{
+				GPUScene.AddPrimitiveToUpdate(PrimitiveSceneInfo->PackedIndex, EPrimitiveDirtyState::ChangedAll);
+
+				DistanceFieldSceneData.UpdatePrimitive(PrimitiveSceneInfo);
+				LumenUpdatePrimitive(PrimitiveSceneInfo);
+			}
+
+			bNeedPathTracedInvalidation = bNeedPathTracedInvalidation || IsPrimitiveRelevantToPathTracing(PrimitiveSceneInfo);
+		}
+	}
+
+	GPUScene.ConsolidateInstanceDataAllocations();
+
 	{
 		CSV_SCOPED_TIMING_STAT_EXCLUSIVE(AddPrimitiveSceneInfos);
 		SCOPED_NAMED_EVENT(FScene_AddPrimitiveSceneInfos, FColor::Green);
 		SCOPE_CYCLE_COUNTER(STAT_AddScenePrimitiveRenderThreadTime);
-#if ENABLE_LOG_PRIMITIVE_INSTANCE_ID_STATS_TO_CSV
-		int32 PersistentPrimitiveFreeListSizeBeforeConsolidate = PersistentPrimitiveIdAllocator.GetNumFreeSpans();
-		int32 PersistentPrimitivePendingFreeListSizeBeforeConsolidate = PersistentPrimitiveIdAllocator.GetNumPendingFreeSpans();
-#endif
+
 		PersistentPrimitiveIdAllocator.Consolidate();
 
 		if (AddedLocalPrimitiveSceneInfos.Num())
@@ -6056,10 +6009,6 @@ void FScene::UpdateAllPrimitiveSceneInfos(FRDGBuilder& GraphBuilder, EUpdateAllP
 
 			CheckPrimitiveArrays();
 
-#if ENABLE_LOG_PRIMITIVE_INSTANCE_ID_STATS_TO_CSV
-			UpdatePrimitiveAllocatorStats(Primitives.Num(), PersistentPrimitiveIdAllocator.GetMaxSize(), PersistentPrimitiveIdAllocator.GetSparselyAllocatedSize(), PersistentPrimitiveIdAllocator.GetNumFreeSpans(), PersistentPrimitiveFreeListSizeBeforeConsolidate, PersistentPrimitivePendingFreeListSizeBeforeConsolidate, GPUScene.GetNumInstances(), GPUScene.GetInstanceSceneDataAllocator().GetSparselyAllocatedSize());
-#endif
-
 			for (int32 AddIndex = StartIndex; AddIndex < AddedLocalPrimitiveSceneInfos.Num(); AddIndex++)
 			{
 				FPrimitiveSceneInfo* PrimitiveSceneInfo = AddedLocalPrimitiveSceneInfos[AddIndex];
@@ -6193,85 +6142,25 @@ void FScene::UpdateAllPrimitiveSceneInfos(FRDGBuilder& GraphBuilder, EUpdateAllP
 		CacheManager->OnSceneChange();
 	}
 
-#if RHI_RAYTRACING
-	TArray<FPrimitiveSceneInfo*, SceneRenderingAllocator> RayTracingPrimitivesToUpdate;
-	RayTracingPrimitivesToUpdate.Reserve(UpdatedInstances.Num());
-	bool bUpdateCachedRayTracingInstances = false;
-#endif
-
+	// Re-adding step of instance update for those that had the instance count updated.
 	{
 		CSV_SCOPED_TIMING_STAT_EXCLUSIVE(UpdatePrimitiveInstances);
 		SCOPED_NAMED_EVENT(FScene_UpdatePrimitiveInstances, FColor::Emerald);
 		SCOPE_CYCLE_COUNTER(STAT_UpdatePrimitiveInstanceRenderThreadTime);
 
-		for (const auto& UpdateInstance : UpdatedInstances)
-		{
-			FPrimitiveSceneProxy* PrimitiveSceneProxy = UpdateInstance.Key;
-			if (DeletedSceneInfos.Contains(PrimitiveSceneProxy->GetPrimitiveSceneInfo()))
-			{
-				continue;
-			}
-
-			check(PrimitiveSceneProxy->GetPrimitiveSceneInfo()->PackedIndex != INDEX_NONE);
+		FPrimitiveSceneInfo::AllocateGPUSceneInstances(this, UpdatedInstancesWithInstanceCountChange);
 
-			FScopeCycleCounter Context(PrimitiveSceneProxy->GetStatId());
-			FPrimitiveSceneInfo* PrimitiveSceneInfo = PrimitiveSceneProxy->GetPrimitiveSceneInfo();
+		for (FPrimitiveSceneInfo* PrimitiveSceneInfo : UpdatedInstancesWithInstanceCountChange)
+		{
+			check(PrimitiveSceneInfo->PackedIndex != INDEX_NONE);
+			FScopeCycleCounter Context(PrimitiveSceneInfo->Proxy->GetStatId());
 			
-			QueueFlushVirtualTexture(PrimitiveSceneInfo);
 
-			// If we recorded no adds or removes the instance count has stayed the same.  Therefore the cached mesh draw commands do not
-			// need to be updated.  In situations where the instance count changes the mesh draw command stores the instance count which
-			// would need to be updated.
-			const bool bInstanceCountChanged = (UpdateInstance.Value.CmdBuffer.NumAdds > 0) || (UpdateInstance.Value.CmdBuffer.NumRemoves > 0);
-
-			const bool bUpdateStaticDrawLists = !PrimitiveSceneProxy->StaticElementsAlwaysUseProxyPrimitiveUniformBuffer() || bInstanceCountChanged;
-
-			if (QueueAddToScene(PrimitiveSceneInfo))
-			{
-				PrimitiveSceneInfo->RemoveFromScene(bUpdateStaticDrawLists);
-
-				if (bUpdateStaticDrawLists)
-				{
-					QueueAddStaticMeshes(PrimitiveSceneInfo);
-				}
-				else
-				{
-#if RHI_RAYTRACING
-					RayTracingPrimitivesToUpdate.Add(PrimitiveSceneInfo);
-					bUpdateCachedRayTracingInstances = true;
-#endif
-				}
-			}
-
-			// Update the Proxy's data.
-			PrimitiveSceneProxy->UpdateInstances_RenderThread(UpdateInstance.Value.CmdBuffer, UpdateInstance.Value.WorldBounds, UpdateInstance.Value.LocalBounds, UpdateInstance.Value.StaticMeshBounds);
-
-			if (!RHISupportsVolumeTextures(GetFeatureLevel())
-				&& (PrimitiveSceneProxy->IsMovable() || PrimitiveSceneProxy->NeedsUnbuiltPreviewLighting() || PrimitiveSceneProxy->GetLightmapType() == ELightmapType::ForceVolumetric))
-			{
-				PrimitiveSceneInfo->MarkIndirectLightingCacheBufferDirty();
-			}
-
-			if (UpdateInstance.Value.CmdBuffer.NumAdds > 0 || UpdateInstance.Value.CmdBuffer.NumRemoves > 0)
-			{
-				PrimitiveSceneInfo->FreeGPUSceneInstances();
-				FPrimitiveSceneInfo::AllocateGPUSceneInstances(this, MakeArrayView(&PrimitiveSceneInfo, 1));
-
-				DistanceFieldSceneData.RemovePrimitive(PrimitiveSceneInfo);
-				DistanceFieldSceneData.AddPrimitive(PrimitiveSceneInfo);
-
-				LumenRemovePrimitive(PrimitiveSceneInfo, PrimitiveSceneInfo->GetIndex());
-				LumenAddPrimitive(PrimitiveSceneInfo);
-			}
-			else
-			{
-				GPUScene.AddPrimitiveToUpdate(PrimitiveSceneInfo->PackedIndex, EPrimitiveDirtyState::ChangedAll);
-
-				DistanceFieldSceneData.UpdatePrimitive(PrimitiveSceneInfo);
-				LumenUpdatePrimitive(PrimitiveSceneInfo);
-			}
+			DistanceFieldSceneData.RemovePrimitive(PrimitiveSceneInfo);
+			DistanceFieldSceneData.AddPrimitive(PrimitiveSceneInfo);
 
-			bNeedPathTracedInvalidation = bNeedPathTracedInvalidation || IsPrimitiveRelevantToPathTracing(PrimitiveSceneInfo);
+			LumenRemovePrimitive(PrimitiveSceneInfo, PrimitiveSceneInfo->GetIndex());
+			LumenAddPrimitive(PrimitiveSceneInfo);
 		}
 
 #if RHI_RAYTRACING
diff --git a/Engine/Source/Runtime/Renderer/Public/PrimitiveSceneInfo.h b/Engine/Source/Runtime/Renderer/Public/PrimitiveSceneInfo.h
--- a/Engine/Source/Runtime/Renderer/Public/PrimitiveSceneInfo.h
+++ b/Engine/Source/Runtime/Renderer/Public/PrimitiveSceneInfo.h
@@ -411,9 +411,10 @@ public:
 
 	/** Allocate/Free slots for instance data in GPU-Scene */
 	static void AllocateGPUSceneInstances(FScene* Scene, const TArrayView<FPrimitiveSceneInfo*>& SceneInfos);
-	static void ReallocateGPUSceneInstances(FScene* Scene, const TArrayView<FPrimitiveSceneInfo*>& SceneInfos);
 	void FreeGPUSceneInstances();
 
+	UE_DEPRECATED(5.3, "ReallocateGPUSceneInstances has been deprecated.")
+	static void ReallocateGPUSceneInstances(FScene* Scene, const TArrayView<FPrimitiveSceneInfo*>& SceneInfos);
 
 	/** return true if we need to call LazyUpdateForRendering */
 	FORCEINLINE bool NeedsIndirectLightingCacheBufferUpdate()
