Fixed wrong color of translucent shadow volume when light is an atmospheric light.

diff --git a/Engine/Source/Runtime/Renderer/Private/ShadowRendering.cpp b/Engine/Source/Runtime/Renderer/Private/ShadowRendering.cpp
--- a/Engine/Source/Runtime/Renderer/Private/ShadowRendering.cpp
+++ b/Engine/Source/Runtime/Renderer/Private/ShadowRendering.cpp
@@ -2445,8 +2445,17 @@ void SetupTranslucentSelfShadowUniformParameters(const FProjectedShadowInfo* Sha
 
 		//@todo - support fading from both views
 		const float FadeAlpha = ShadowInfo->FadeAlphas[0];
+		FLinearColor LightColor;
+		if (LightProxy->IsUsedAsAtmosphereSunLight())
+		{
+			LightColor = LightProxy->GetSunIlluminanceAccountingForSkyAtmospherePerPixelTransmittance();
+		}
+		else
+		{
+			LightColor = LightProxy->GetColor();
+		}
 		// Incorporate the diffuse scale of 1 / PI into the light color
-		OutParameters.DirectionalLightColor = FVector4f(FVector3f(LightProxy->GetColor() * FadeAlpha / PI), FadeAlpha);
+		OutParameters.DirectionalLightColor = FVector4f(FVector3f(LightColor * FadeAlpha / PI), FadeAlpha);
 
 		OutParameters.Transmission0 = ShadowInfo->RenderTargets.ColorTargets[0]->GetRHI();
 		OutParameters.Transmission1 = ShadowInfo->RenderTargets.ColorTargets[1]->GetRHI();
