Fix incorrect stride in EyeAdaptationBuffer

diff --git a/Engine/Source/Runtime/Renderer/Private/PostProcess/PostProcessEyeAdaptation.cpp b/Engine/Source/Runtime/Renderer/Private/PostProcess/PostProcessEyeAdaptation.cpp
--- a/Engine/Source/Runtime/Renderer/Private/PostProcess/PostProcessEyeAdaptation.cpp
+++ b/Engine/Source/Runtime/Renderer/Private/PostProcess/PostProcessEyeAdaptation.cpp
@@ -1060,7 +1060,7 @@ const TRefCountPtr<FRDGPooledBuffer>& FSceneViewState::FEyeAdaptationManager::Ge
 	// Create buffer if needed.
 	if (!ExposureBufferData[BufferIndex].IsValid())
 	{
-		FRDGBufferDesc RDGBufferDesc = FRDGBufferDesc::CreateStructuredDesc(EXPOSURE_BUFFER_SIZE_IN_VECTOR4 * sizeof(FVector4f), 1);
+		FRDGBufferDesc RDGBufferDesc = FRDGBufferDesc::CreateStructuredDesc(sizeof(FVector4f), EXPOSURE_BUFFER_SIZE_IN_VECTOR4);
 		RDGBufferDesc.Usage |= BUF_SourceCopy;
 		FRDGBufferRef RDGBuffer = GraphBuilder.CreateBuffer(RDGBufferDesc, TEXT("EyeAdaptationBuffer"), ERDGBufferFlags::MultiFrame);
 
