Allow per-pass VRS CVars to affect all generators, including DebugForceRate. Updated related preview CVars to reflect this.
Also fixed a bug in which fixed foveation was using the wrong half-rate cutoffs, leading to it skipping past 2x2 rates straight to 4x4.

diff --git a/Engine/Source/Runtime/Renderer/Private/VariableRateShading/ContrastAdaptiveImageGenerator.cpp b/Engine/Source/Runtime/Renderer/Private/VariableRateShading/ContrastAdaptiveImageGenerator.cpp
--- a/Engine/Source/Runtime/Renderer/Private/VariableRateShading/ContrastAdaptiveImageGenerator.cpp
+++ b/Engine/Source/Runtime/Renderer/Private/VariableRateShading/ContrastAdaptiveImageGenerator.cpp
@@ -31,105 +31,43 @@
  * CAS Parameters
  */
 
-TAutoConsoleVariable<int32> CVarVRSContrastAdaptiveShading(
+TAutoConsoleVariable<int32> CVarCASContrastAdaptiveShading(
 	TEXT("r.VRS.ContrastAdaptiveShading"),
 	0,
 	TEXT("Enables using Variable Rate Shading based on the luminance from the previous frame's post process output \n"),
 	ECVF_RenderThreadSafe);
-TAutoConsoleVariable<float> CVarVRSEdgeThreshold(
+TAutoConsoleVariable<float> CVarCASEdgeThreshold(
 	TEXT("r.VRS.ContrastAdaptiveShading.EdgeThreshold"),
 	0.2,
 	TEXT(""),
 	ECVF_RenderThreadSafe);
-TAutoConsoleVariable<float> CVarVRSConservativeEdgeThreshold(
+TAutoConsoleVariable<float> CVarCASConservativeEdgeThreshold(
 	TEXT("r.VRS.ContrastAdaptiveShading.ConservativeEdgeThreshold"),
 	0.1,
 	TEXT(""),
 	ECVF_RenderThreadSafe);
-TAutoConsoleVariable<float> CVarVRS_HDR10CorrectionMultiplier(
+TAutoConsoleVariable<float> CVarCAS_HDR10CorrectionMultiplier(
 	TEXT("r.VRS.ContrastAdaptiveShading.HDR10CorrectionMultiplier"),
 	0.55,
 	TEXT("Approximation multiplier to account for how perceptual values are spread out in SDR vs HDR10\n"),
 	ECVF_RenderThreadSafe);
 
 /**
- * Pass Settings
+ * Debug Settings 
  */
 
-TAutoConsoleVariable<int32> CVarVRSBasePass(
-	TEXT("r.VRS.BasePass"),
-	2,
-	TEXT("Enables Variable Rate Shading for the base pass\n")
-	TEXT("0: Disabled")
-	TEXT("1: Full")
-	TEXT("2: Conservative (default)"),
-	ECVF_RenderThreadSafe);
-TAutoConsoleVariable<int32> CVarVRSTranslucency(
-	TEXT("r.VRS.Translucency"),
-	1,
-	TEXT("Enable VRS with translucency rendering.\n")
-	TEXT("0: Disabled")
-	TEXT("1: Full (default)")
-	TEXT("2: Conservative"),
-	ECVF_RenderThreadSafe);
-TAutoConsoleVariable<int32> CVarVRSNaniteEmitGBuffer(
-	TEXT("r.VRS.NaniteEmitGBuffer"),
-	2,
-	TEXT("Enable VRS with Nanite EmitGBuffer rendering.\n")
-	TEXT("0: Disabled")
-	TEXT("1: Full")
-	TEXT("2: Conservative (default)"),
-	ECVF_RenderThreadSafe);
-TAutoConsoleVariable<int32> CVarVRS_SSAO(
-	TEXT("r.VRS.SSAO"),
-	0,
-	TEXT("Enable VRS with SSAO rendering.\n")
-	TEXT("0: Disabled")
-	TEXT("1: Full")
-	TEXT("2: Conservative (default)"),
-	ECVF_RenderThreadSafe);
-TAutoConsoleVariable<int32> CVarVRS_SSR(
-	TEXT("r.VRS.SSR"),
-	2,
-	TEXT("Enable VRS with SSR (PS) rendering.\n")
-	TEXT("0: Disabled")
-	TEXT("1: Full")
-	TEXT("2: Conservative (default)"),
-	ECVF_RenderThreadSafe);
-TAutoConsoleVariable<int32> CVarVRSReflectionEnvironmentSky(
-	TEXT("r.VRS.ReflectionEnvironmentSky"),
-	2,
-	TEXT("Enable VRS with ReflectionEnvironmentAndSky (PS) rendering.\n")
-	TEXT("0: Disabled")
-	TEXT("1: Full")
-	TEXT("2: Conservative (default)"),
-	ECVF_RenderThreadSafe);
-TAutoConsoleVariable<int32> CVarVRSLightFunctions(
-	TEXT("r.VRS.LightFunctions"),
+TAutoConsoleVariable<int32> CVarCASPreview(
+	TEXT("r.VRS.ContrastAdaptiveShading.Preview"),
 	1,
-	TEXT("Enables Variable Rate Shading for light functions\n")
-	TEXT("0: Disabled")
-	TEXT("1: Full (default)")
-	TEXT("2: Conservative"),
+	TEXT("Whether to include CAS in VRS preview overlay.")
+	TEXT("0 - off, 1 - on (default)"),
 	ECVF_RenderThreadSafe);
-TAutoConsoleVariable<int32> CVarVRSDecals(
-	TEXT("r.VRS.Decals"),
-	2,
-	TEXT("Enables Variable Rate Shading for decals\n")
-	TEXT("0: Disabled")
-	TEXT("1: Full")
-	TEXT("2: Conservative (default)"),
-	ECVF_RenderThreadSafe);
-
-/**
- * Debug Settings 
- */
 
-TAutoConsoleVariable<int32> CVarCASPreviewType(
-	TEXT("r.VRS.ContrastAdaptiveShading.PreviewType"),
-	1,
-	TEXT("Include CAS in the VRS debug overlay.")
-	TEXT("0 - off, 1 - the SRI texture, 2- the conservative SRI texture, 3 - the pre-reprojection SRI texture"),
+TAutoConsoleVariable<int32> CVarCASPreviewPreReprojection(
+	TEXT("r.VRS.ContrastAdaptiveShading.PreviewPreReprojection"),
+	0,
+	TEXT("Sets CAS preview to use the pre-reprojection SRI. Overrides full vs. conservative images.")
+	TEXT("0 - off (default), 1 - on"),
 	ECVF_RenderThreadSafe);
 
 
@@ -178,9 +116,9 @@ class FCalculateShadingRateImageCS : public FGlobalShader
 	{
 		Parameters.LuminanceTexture = Luminance;
 		Parameters.ViewRect = FVector4f(ViewRect.Min.X, ViewRect.Min.Y, ViewRect.Max.X, ViewRect.Max.Y);
-		const float cEdgeThresholdCorrectionValue = bIsHDR10 ? CVarVRS_HDR10CorrectionMultiplier.GetValueOnRenderThread() : 1.0;
-		Parameters.EdgeThreshold = cEdgeThresholdCorrectionValue * CVarVRSEdgeThreshold.GetValueOnRenderThread();
-		Parameters.ConservativeEdgeThreshold = cEdgeThresholdCorrectionValue * CVarVRSConservativeEdgeThreshold.GetValueOnRenderThread();
+		const float cEdgeThresholdCorrectionValue = bIsHDR10 ? CVarCAS_HDR10CorrectionMultiplier.GetValueOnRenderThread() : 1.0;
+		Parameters.EdgeThreshold = cEdgeThresholdCorrectionValue * CVarCASEdgeThreshold.GetValueOnRenderThread();
+		Parameters.ConservativeEdgeThreshold = cEdgeThresholdCorrectionValue * CVarCASConservativeEdgeThreshold.GetValueOnRenderThread();
 		Parameters.VariableRateShadingTexture = ShadingRateImage;
 	}
 };
@@ -280,37 +218,9 @@ namespace ESRITextureType
 	{
 		return IsValidShadingRateTexture(static_cast<int32>(TextureType));
 	}
-	ESRITextureType::Type GetTextureType(FVariableRateShadingImageManager::EVRSPassType PassType)
+	ESRITextureType::Type GetTextureType(FVariableRateShadingImageManager::EVRSImageType ImageType)
 	{
-		static struct FStaticData
-		{
-			TAutoConsoleVariable<int32>*CVarByPassType[FVariableRateShadingImageManager::EVRSPassType::Num] = {};
-			FStaticData()
-			{
-				CVarByPassType[FVariableRateShadingImageManager::EVRSPassType::BasePass] = &CVarVRSBasePass;
-				CVarByPassType[FVariableRateShadingImageManager::EVRSPassType::TranslucencyAll] = &CVarVRSTranslucency;
-				CVarByPassType[FVariableRateShadingImageManager::EVRSPassType::NaniteEmitGBufferPass] = &CVarVRSNaniteEmitGBuffer;
-				CVarByPassType[FVariableRateShadingImageManager::EVRSPassType::SSAO] = &CVarVRS_SSAO;
-				CVarByPassType[FVariableRateShadingImageManager::EVRSPassType::SSR] = &CVarVRS_SSR;
-				CVarByPassType[FVariableRateShadingImageManager::EVRSPassType::ReflectionEnvironmentAndSky] = &CVarVRSReflectionEnvironmentSky;
-				CVarByPassType[FVariableRateShadingImageManager::EVRSPassType::LightFunctions] = &CVarVRSLightFunctions;
-				CVarByPassType[FVariableRateShadingImageManager::EVRSPassType::Decals] = &CVarVRSDecals;
-			}
-		} StaticData;
-		const int32 PassTypeAsInt = static_cast<int32>(PassType);
-		if (PassTypeAsInt >= 0 && PassTypeAsInt < EMeshPass::Num)
-		{
-			auto* CVar = StaticData.CVarByPassType[PassTypeAsInt];
-			if (CVar)
-			{
-				int32 TextureTypeAsInt = CVar->GetValueOnRenderThread();
-				if (IsValidShadingRateTexture(TextureTypeAsInt))
-				{
-					return static_cast<ESRITextureType::Type>(TextureTypeAsInt);
-				}
-			}
-		}
-		return ESRITextureType::None;
+		return static_cast<ESRITextureType::Type>(ImageType);
 	}
 };
 
@@ -412,7 +322,7 @@ static bool IsHDR10(const EDisplayOutputFormat& OutputFormat)
 
 static bool IsContrastAdaptiveShadingEnabled()
 {
-	return GRHISupportsAttachmentVariableRateShading && GRHIAttachmentVariableRateShadingEnabled && (CVarVRSContrastAdaptiveShading.GetValueOnRenderThread() != 0);
+	return GRHISupportsAttachmentVariableRateShading && GRHIAttachmentVariableRateShadingEnabled && (CVarCASContrastAdaptiveShading.GetValueOnRenderThread() != 0);
 }
 
 static FIntRect GetPostProcessOutputRect(const FViewInfo& ViewInfo)
@@ -550,11 +460,11 @@ void AddPrepareImageBasedVRSPass(
  * Interface Functions
  */
 
-FRDGTextureRef FContrastAdaptiveImageGenerator::GetImage(FRDGBuilder& GraphBuilder, const FViewInfo& ViewInfo, FVariableRateShadingImageManager::EVRSPassType PassType)
+FRDGTextureRef FContrastAdaptiveImageGenerator::GetImage(FRDGBuilder& GraphBuilder, const FViewInfo& ViewInfo, FVariableRateShadingImageManager::EVRSImageType ImageType)
 {
 	if (FVRSTextures::IsInitialized(GraphBuilder))
 	{
-		ESRITextureType::Type TextureType = ESRITextureType::GetTextureType(PassType);
+		ESRITextureType::Type TextureType = ESRITextureType::GetTextureType(ImageType);
 		if (TextureType != ESRITextureType::None)
 		{
 			const FVRSTextures& VRSTextures = FVRSTextures::Get(GraphBuilder);
@@ -606,17 +516,20 @@ bool FContrastAdaptiveImageGenerator::IsEnabledForView(const FSceneView& View) c
 	return IsContrastAdaptiveShadingEnabled() && !View.bIsSceneCapture && bCompatibleWithOutputType;
 }
 
-FRDGTextureRef FContrastAdaptiveImageGenerator::GetDebugImage(FRDGBuilder& GraphBuilder, const FViewInfo& ViewInfo)
+FRDGTextureRef FContrastAdaptiveImageGenerator::GetDebugImage(FRDGBuilder& GraphBuilder, const FViewInfo& ViewInfo, FVariableRateShadingImageManager::EVRSImageType ImageType)
 {
-	if (!FVRSTextures::IsInitialized(GraphBuilder))
+	if (!CVarCASPreview.GetValueOnRenderThread() || !FVRSTextures::IsInitialized(GraphBuilder))
 	{
 		return nullptr;
 	}
 
-	const FVRSTextures& VRSTextures = FVRSTextures::Get(GraphBuilder);
-
-	ESRIPreviewType::Type PreviewType = static_cast<ESRIPreviewType::Type>(CVarCASPreviewType.GetValueOnRenderThread());
+	ESRIPreviewType::Type PreviewType = static_cast<ESRIPreviewType::Type>(ImageType);
+	if (CVarCASPreviewPreReprojection.GetValueOnRenderThread() && PreviewType != ESRIPreviewType::Off)
+	{
+		PreviewType = ESRIPreviewType::BeforeReprojection;
+	}
 
+	const FVRSTextures& VRSTextures = FVRSTextures::Get(GraphBuilder);
 	switch (PreviewType)
 	{
 	case ESRIPreviewType::Projected:
diff --git a/Engine/Source/Runtime/Renderer/Private/VariableRateShading/ContrastAdaptiveImageGenerator.h b/Engine/Source/Runtime/Renderer/Private/VariableRateShading/ContrastAdaptiveImageGenerator.h
--- a/Engine/Source/Runtime/Renderer/Private/VariableRateShading/ContrastAdaptiveImageGenerator.h
+++ b/Engine/Source/Runtime/Renderer/Private/VariableRateShading/ContrastAdaptiveImageGenerator.h
@@ -13,13 +13,13 @@ class FContrastAdaptiveImageGenerator : public IVariableRateShadingImageGenerato
 {
 public:
 	virtual ~FContrastAdaptiveImageGenerator() override {};
-	virtual FRDGTextureRef GetImage(FRDGBuilder& GraphBuilder, const FViewInfo& ViewInfo, FVariableRateShadingImageManager::EVRSPassType PassType) override;
+	virtual FRDGTextureRef GetImage(FRDGBuilder& GraphBuilder, const FViewInfo& ViewInfo, FVariableRateShadingImageManager::EVRSImageType ImageType) override;
 	virtual void PrepareImages(FRDGBuilder& GraphBuilder, const FSceneViewFamily& ViewFamily, const FMinimalSceneTextures& SceneTextures) override;
 	virtual bool IsEnabledForView(const FSceneView& View) const override;
 	virtual FVariableRateShadingImageManager::EVRSSourceType GetType() const override
 	{
 		return FVariableRateShadingImageManager::EVRSSourceType::ContrastAdaptiveShading;
 	}
-	virtual FRDGTextureRef GetDebugImage(FRDGBuilder& GraphBuilder, const FViewInfo& ViewInfo) override;
+	virtual FRDGTextureRef GetDebugImage(FRDGBuilder& GraphBuilder, const FViewInfo& ViewInfo, FVariableRateShadingImageManager::EVRSImageType ImageType) override;
 };
 
diff --git a/Engine/Source/Runtime/Renderer/Private/VariableRateShading/FoveatedImageGenerator.cpp b/Engine/Source/Runtime/Renderer/Private/VariableRateShading/FoveatedImageGenerator.cpp
--- a/Engine/Source/Runtime/Renderer/Private/VariableRateShading/FoveatedImageGenerator.cpp
+++ b/Engine/Source/Runtime/Renderer/Private/VariableRateShading/FoveatedImageGenerator.cpp
@@ -120,15 +120,17 @@ public:
 
 IMPLEMENT_GLOBAL_SHADER(FComputeVariableRateShadingImageGeneration, "/Engine/Private/VariableRateShading/VRSShadingRateFoveated.usf", "GenerateShadingRateTexture", SF_Compute);
 
-FRDGTextureRef FFoveatedImageGenerator::GetImage(FRDGBuilder& GraphBuilder, const FViewInfo& ViewInfo, FVariableRateShadingImageManager::EVRSPassType PassType)
+FRDGTextureRef FFoveatedImageGenerator::GetImage(FRDGBuilder& GraphBuilder, const FViewInfo& ViewInfo, FVariableRateShadingImageManager::EVRSImageType ImageType)
 {
 	// Generator only supports up to two side-by-side views
-	if (ViewInfo.StereoViewIndex > 1)
+	if (ImageType == FVariableRateShadingImageManager::EVRSImageType::Disabled || ViewInfo.StereoViewIndex > 1)
 	{
 		return nullptr;
 	}
-
-	return CachedImage;
+	else
+	{
+		return CachedImage;
+	}
 }
 
 void FFoveatedImageGenerator::PrepareImages(FRDGBuilder& GraphBuilder, const FSceneViewFamily& ViewFamily, const FMinimalSceneTextures& SceneTextures)
@@ -136,7 +138,7 @@ void FFoveatedImageGenerator::PrepareImages(FRDGBuilder& GraphBuilder, const FSc
 
 	// VRS level parameters - pretty arbitrary right now, later should depend on device characteristics
 	static const TArray<float> kFoveationFullRateCutoffs = { 1.0f, 0.7f, 0.50f, 0.35f, 0.35f };
-	static const TArray<float> kFoveationHalfRateCutofffs = { 1.0f, 0.9f, 0.75f, 0.55f, 0.55f };
+	static const TArray<float> kFoveationHalfRateCutoffs = { 1.0f, 0.9f, 0.75f, 0.55f, 0.55f };
 	static const TArray<float> kFixedFoveationCenterX = { 0.5f, 0.5f, 0.5f, 0.5f, 0.5f };
 	static const TArray<float> kFixedFoveationCenterY = { 0.5f, 0.5f, 0.5f, 0.5f, 0.42f };
 
@@ -156,7 +158,7 @@ void FFoveatedImageGenerator::PrepareImages(FRDGBuilder& GraphBuilder, const FSc
 	}
 
 	const float FoveationFullRateCutoff = FMath::Lerp(kFoveationFullRateCutoffs[0], kFoveationFullRateCutoffs[VRSMaxLevel], VRSDynamicAmount);
-	const float FoveationHalfRateCutoff = FMath::Lerp(kFoveationHalfRateCutofffs[0], kFoveationFullRateCutoffs[VRSMaxLevel], VRSDynamicAmount);
+	const float FoveationHalfRateCutoff = FMath::Lerp(kFoveationHalfRateCutoffs[0], kFoveationHalfRateCutoffs[VRSMaxLevel], VRSDynamicAmount);
 
 	// Default to fixed center point
 	float FoveationCenterX = kFixedFoveationCenterX[VRSMaxLevel];
@@ -251,9 +253,16 @@ FVariableRateShadingImageManager::EVRSSourceType FFoveatedImageGenerator::GetTyp
 	return IsGazeTrackingEnabled() ? FVariableRateShadingImageManager::EVRSSourceType::FixedFoveation : FVariableRateShadingImageManager::EVRSSourceType::EyeTrackedFoveation;
 }
 
-FRDGTextureRef FFoveatedImageGenerator::GetDebugImage(FRDGBuilder& GraphBuilder, const FViewInfo& ViewInfo)
+FRDGTextureRef FFoveatedImageGenerator::GetDebugImage(FRDGBuilder& GraphBuilder, const FViewInfo& ViewInfo, FVariableRateShadingImageManager::EVRSImageType ImageType)
 {
-	return CachedImage;
+	if (CVarFoveationPreview.GetValueOnRenderThread() && ImageType != FVariableRateShadingImageManager::EVRSImageType::Disabled)
+	{
+		return CachedImage;
+	}
+	else
+	{
+		return nullptr;
+	}
 }
 
 float FFoveatedImageGenerator::UpdateDynamicVRSAmount()
diff --git a/Engine/Source/Runtime/Renderer/Private/VariableRateShading/FoveatedImageGenerator.h b/Engine/Source/Runtime/Renderer/Private/VariableRateShading/FoveatedImageGenerator.h
--- a/Engine/Source/Runtime/Renderer/Private/VariableRateShading/FoveatedImageGenerator.h
+++ b/Engine/Source/Runtime/Renderer/Private/VariableRateShading/FoveatedImageGenerator.h
@@ -13,11 +13,11 @@ class FFoveatedImageGenerator : public IVariableRateShadingImageGenerator
 {
 public:
 	virtual ~FFoveatedImageGenerator() override {};
-	virtual FRDGTextureRef GetImage(FRDGBuilder& GraphBuilder, const FViewInfo& ViewInfo, FVariableRateShadingImageManager::EVRSPassType PassType) override;
+	virtual FRDGTextureRef GetImage(FRDGBuilder& GraphBuilder, const FViewInfo& ViewInfo, FVariableRateShadingImageManager::EVRSImageType ImageType) override;
 	virtual void PrepareImages(FRDGBuilder& GraphBuilder, const FSceneViewFamily& ViewFamily, const FMinimalSceneTextures& SceneTextures) override;
 	virtual bool IsEnabledForView(const FSceneView& View) const override;
 	virtual FVariableRateShadingImageManager::EVRSSourceType GetType() const override;
-	virtual FRDGTextureRef GetDebugImage(FRDGBuilder& GraphBuilder, const FViewInfo& ViewInfo) override;
+	virtual FRDGTextureRef GetDebugImage(FRDGBuilder& GraphBuilder, const FViewInfo& ViewInfo, FVariableRateShadingImageManager::EVRSImageType ImageType) override;
 private:
 	FRDGTextureRef CachedImage = nullptr;
 	struct FDynamicVRSData
diff --git a/Engine/Source/Runtime/Renderer/Private/VariableRateShading/VariableRateShadingImageManager.cpp b/Engine/Source/Runtime/Renderer/Private/VariableRateShading/VariableRateShadingImageManager.cpp
--- a/Engine/Source/Runtime/Renderer/Private/VariableRateShading/VariableRateShadingImageManager.cpp
+++ b/Engine/Source/Runtime/Renderer/Private/VariableRateShading/VariableRateShadingImageManager.cpp
@@ -25,14 +25,14 @@
 TGlobalResource<FVariableRateShadingImageManager> GVRSImageManager;
 
 /**
- * CVars
+ * Basic CVars
  */
 
 TAutoConsoleVariable<int32> CVarVRSPreview(
 	TEXT("r.VRS.Preview"),
 	0,
-	TEXT("Show a debug visualiation of the VRS shading rate image texture.")
-	TEXT("0 - off, 1 - on"),
+	TEXT("Show a debug visualization of the VRS shading rate image texture.")
+	TEXT("0 - off, 1 - on, 2 - conservative (affects CAS only)"),
 	ECVF_RenderThreadSafe);
 
 int GVRSDebugForceRate = -1;
@@ -42,6 +42,75 @@ FAutoConsoleVariableRef CVarVRSDebugForceRate(
 	TEXT("-1 : None, 0 : Force 1x1, 1 : Force 1x2, 2 : Force 2x1, 3: Force 2x2, 4 : Force 2x4, 5 : Force 4x2, 6 : Force 4x4"),
 	ECVF_RenderThreadSafe);
 
+/**
+ * Pass Settings
+ */
+
+TAutoConsoleVariable<int32> CVarVRSBasePass(
+	TEXT("r.VRS.BasePass"),
+	2,
+	TEXT("Enables Variable Rate Shading for the base pass\n")
+	TEXT("0: Disabled")
+	TEXT("1: Full")
+	TEXT("2: Conservative (default)"),
+	ECVF_RenderThreadSafe);
+TAutoConsoleVariable<int32> CVarVRSTranslucency(
+	TEXT("r.VRS.Translucency"),
+	1,
+	TEXT("Enable VRS with translucency rendering.\n")
+	TEXT("0: Disabled")
+	TEXT("1: Full (default)")
+	TEXT("2: Conservative"),
+	ECVF_RenderThreadSafe);
+TAutoConsoleVariable<int32> CVarVRSNaniteEmitGBuffer(
+	TEXT("r.VRS.NaniteEmitGBuffer"),
+	2,
+	TEXT("Enable VRS with Nanite EmitGBuffer rendering.\n")
+	TEXT("0: Disabled")
+	TEXT("1: Full")
+	TEXT("2: Conservative (default)"),
+	ECVF_RenderThreadSafe);
+TAutoConsoleVariable<int32> CVarVRS_SSAO(
+	TEXT("r.VRS.SSAO"),
+	0,
+	TEXT("Enable VRS with SSAO rendering.\n")
+	TEXT("0: Disabled")
+	TEXT("1: Full")
+	TEXT("2: Conservative (default)"),
+	ECVF_RenderThreadSafe);
+TAutoConsoleVariable<int32> CVarVRS_SSR(
+	TEXT("r.VRS.SSR"),
+	2,
+	TEXT("Enable VRS with SSR (PS) rendering.\n")
+	TEXT("0: Disabled")
+	TEXT("1: Full")
+	TEXT("2: Conservative (default)"),
+	ECVF_RenderThreadSafe);
+TAutoConsoleVariable<int32> CVarVRSReflectionEnvironmentSky(
+	TEXT("r.VRS.ReflectionEnvironmentSky"),
+	2,
+	TEXT("Enable VRS with ReflectionEnvironmentAndSky (PS) rendering.\n")
+	TEXT("0: Disabled")
+	TEXT("1: Full")
+	TEXT("2: Conservative (default)"),
+	ECVF_RenderThreadSafe);
+TAutoConsoleVariable<int32> CVarVRSLightFunctions(
+	TEXT("r.VRS.LightFunctions"),
+	1,
+	TEXT("Enables Variable Rate Shading for light functions\n")
+	TEXT("0: Disabled")
+	TEXT("1: Full (default)")
+	TEXT("2: Conservative"),
+	ECVF_RenderThreadSafe);
+TAutoConsoleVariable<int32> CVarVRSDecals(
+	TEXT("r.VRS.Decals"),
+	2,
+	TEXT("Enables Variable Rate Shading for decals\n")
+	TEXT("0: Disabled")
+	TEXT("1: Full")
+	TEXT("2: Conservative (default)"),
+	ECVF_RenderThreadSafe);
+
 /**
  * Shaders
  */
@@ -243,7 +312,7 @@ FRDGTextureRef FVariableRateShadingImageManager::GetVariableRateShadingImage(FRD
 	// Use debug rate if provided
 	if (CVarVRSDebugForceRate->GetInt() >= 0)
 	{
-		return GetForceRateImage(GraphBuilder, CVarVRSDebugForceRate->GetInt());
+		return GetForceRateImage(GraphBuilder, CVarVRSDebugForceRate->GetInt(), GetImageTypeFromPassType(PassType));
 	}
 
 	// Otherwise collate all internal sources
@@ -254,7 +323,7 @@ FRDGTextureRef FVariableRateShadingImageManager::GetVariableRateShadingImage(FRD
 		FRDGTextureRef Image = nullptr;
 		if (Generator->IsEnabledForView(ViewInfo) && !EnumHasAnyFlags(VRSTypesToExclude, Generator->GetType()))
 		{
-			Image = Generator->GetImage(GraphBuilder, ViewInfo, PassType);
+			Image = Generator->GetImage(GraphBuilder, ViewInfo, GetImageTypeFromPassType(PassType));
 		}
 
 		if (Image)
@@ -362,8 +431,14 @@ TRefCountPtr<IPooledRenderTarget> FVariableRateShadingImageManager::GetMobileVar
 
 void FVariableRateShadingImageManager::DrawDebugPreview(FRDGBuilder& GraphBuilder, const FSceneViewFamily& ViewFamily, FRDGTextureRef OutputSceneColor)
 {
+	uint32 ImageTypeAsInt = CVarVRSPreview.GetValueOnRenderThread();
+	EVRSImageType PreviewImageType = EVRSImageType::Disabled;
+	if (ImageTypeAsInt >= 0 && ImageTypeAsInt <= EVRSImageType::Conservative)
+	{
+		PreviewImageType = static_cast<EVRSImageType>(ImageTypeAsInt);
+	}
 
-	if (CVarVRSPreview.GetValueOnRenderThread() != 1 || !OutputSceneColor)
+	if (PreviewImageType == EVRSImageType::Disabled || !OutputSceneColor)
 	{
 		return;
 	}
@@ -379,7 +454,7 @@ void FVariableRateShadingImageManager::DrawDebugPreview(FRDGBuilder& GraphBuilde
 			// Use debug rate if provided
 			if (CVarVRSDebugForceRate->GetInt() >= 0)
 			{
-				PreviewTexture = GetForceRateImage(GraphBuilder, CVarVRSDebugForceRate->GetInt());
+				PreviewTexture = GetForceRateImage(GraphBuilder, CVarVRSDebugForceRate->GetInt(), PreviewImageType);
 			}
 
 			// Otherwise collate debug images
@@ -393,7 +468,7 @@ void FVariableRateShadingImageManager::DrawDebugPreview(FRDGBuilder& GraphBuilde
 					FRDGTextureRef Image = nullptr;
 					if (Generator->IsEnabledForView(*View))
 					{
-						Image = Generator->GetDebugImage(GraphBuilder, *ViewInfo);
+						Image = Generator->GetDebugImage(GraphBuilder, *ViewInfo, PreviewImageType);
 					}
 
 					if (Image)
@@ -495,7 +570,7 @@ FRDGTextureRef FVariableRateShadingImageManager::CombineShadingRateImages(FRDGBu
 
 }
 
-FRDGTextureRef FVariableRateShadingImageManager::GetForceRateImage(FRDGBuilder& GraphBuilder, int RateIndex /* = 0*/)
+FRDGTextureRef FVariableRateShadingImageManager::GetForceRateImage(FRDGBuilder& GraphBuilder, int RateIndex /* = 0*/, EVRSImageType ImageType /* = EVRSImageType::Full*/)
 {
 	static const TArray<uint32> ValidShadingRates = { VRSSR_1x1, VRSSR_1x2, VRSSR_2x1, VRSSR_2x2, VRSSR_2x4, VRSSR_4x2, VRSSR_4x4 };
 
@@ -505,9 +580,43 @@ FRDGTextureRef FVariableRateShadingImageManager::GetForceRateImage(FRDGBuilder&
 		RateIndex = ValidShadingRates.Num() - 1; // Default to maximum shading rate if value exceeds valid rates
 	}
 
+	if (ImageType == EVRSImageType::Disabled)
+	{
+		RateIndex = ValidShadingRates[0]; // Force to minimum shading rate if VRS is disabled for this pass
+	}
+
 	FRDGTextureRef ForceShadingRateTexture = GraphBuilder.CreateTexture(GetSRIDesc(), TEXT("ForceShadingRateTexture"));
 	FRDGTextureUAVRef ForceShadingRateUAV = GraphBuilder.CreateUAV(ForceShadingRateTexture);
 	AddClearUAVPass(GraphBuilder, ForceShadingRateUAV, ValidShadingRates[RateIndex]);
 
 	return ForceShadingRateTexture;
+}
+
+FVariableRateShadingImageManager::EVRSImageType FVariableRateShadingImageManager::GetImageTypeFromPassType(EVRSPassType PassType)
+{
+	static struct FStaticPassToImageCVarData
+	{
+		TAutoConsoleVariable<int32>* CVarByPassType[EVRSPassType::Num] = {};
+		FStaticPassToImageCVarData()
+		{
+			CVarByPassType[EVRSPassType::BasePass] = &CVarVRSBasePass;
+			CVarByPassType[EVRSPassType::TranslucencyAll] = &CVarVRSTranslucency;
+			CVarByPassType[EVRSPassType::NaniteEmitGBufferPass] = &CVarVRSNaniteEmitGBuffer;
+			CVarByPassType[EVRSPassType::SSAO] = &CVarVRS_SSAO;
+			CVarByPassType[EVRSPassType::SSR] = &CVarVRS_SSR;
+			CVarByPassType[EVRSPassType::ReflectionEnvironmentAndSky] = &CVarVRSReflectionEnvironmentSky;
+			CVarByPassType[EVRSPassType::LightFunctions] = &CVarVRSLightFunctions;
+			CVarByPassType[EVRSPassType::Decals] = &CVarVRSDecals;
+		}
+	} StaticData;
+
+	uint32 ImageTypeAsInt = StaticData.CVarByPassType[PassType]->GetValueOnRenderThread();
+	if (ImageTypeAsInt >= 0 && ImageTypeAsInt <= EVRSImageType::Conservative)
+	{
+		return static_cast<EVRSImageType>(ImageTypeAsInt);
+	}
+	else
+	{
+		return EVRSImageType::Disabled;
+	}
 }
\ No newline at end of file
diff --git a/Engine/Source/Runtime/Renderer/Public/VariableRateShadingImageManager.h b/Engine/Source/Runtime/Renderer/Public/VariableRateShadingImageManager.h
--- a/Engine/Source/Runtime/Renderer/Public/VariableRateShadingImageManager.h
+++ b/Engine/Source/Runtime/Renderer/Public/VariableRateShadingImageManager.h
@@ -19,7 +19,7 @@ class FVariableRateShadingImageManager : public FRenderResource
 {
 public:
 	/**
-	 * Pass type used by some generators. Must be provided by passes requesting a VRS image.
+	 * Pass type used to determine requested image type based on current CVar settings.
 	 */
 	enum EVRSPassType
 	{
@@ -34,6 +34,16 @@ public:
 		Num
 	};
 
+	/**
+	 * Image type to request from generator. Only the CAS generator currently distinguishes between Full and Conservative
+	 */
+	enum EVRSImageType : uint32
+	{
+		Disabled = 0,
+		Full = 1,
+		Conservative = 2
+	};
+
 	/**
 	 * Explicit VRS source flags; each type of VRS image generated here will have it's own flag defined in this bitfield.
 	 * These flags can be used for any particular rendering path to exclude specific types of VRS generation; e.g.
@@ -105,7 +115,9 @@ private:
 	TArray<TUniquePtr<IVariableRateShadingImageGenerator>> ImageGenerators;
 
 	FRDGTextureRef CombineShadingRateImages(FRDGBuilder& GraphBuilder, const FViewInfo& ViewInfo, TArray<FRDGTextureRef> Sources);
-	FRDGTextureRef GetForceRateImage(FRDGBuilder& GraphBuilder, int RateIndex = 0);
+	FRDGTextureRef GetForceRateImage(FRDGBuilder& GraphBuilder, int RateIndex = 0, EVRSImageType ImageType = EVRSImageType::Full);
+
+	EVRSImageType GetImageTypeFromPassType(EVRSPassType PassType);
 };
 
 ENUM_CLASS_FLAGS(FVariableRateShadingImageManager::EVRSSourceType);
@@ -118,7 +130,7 @@ public:
 	virtual ~IVariableRateShadingImageGenerator() {};
 
 	// Returns cached VRS image.
-	virtual FRDGTextureRef GetImage(FRDGBuilder& GraphBuilder, const FViewInfo& ViewInfo, FVariableRateShadingImageManager::EVRSPassType PassType) = 0;
+	virtual FRDGTextureRef GetImage(FRDGBuilder& GraphBuilder, const FViewInfo& ViewInfo, FVariableRateShadingImageManager::EVRSImageType ImageType) = 0;
 
 	// Generates image(s) and saves to generator cache. Should only be run once per view per frame, in Render().
 	virtual void PrepareImages(FRDGBuilder& GraphBuilder, const FSceneViewFamily& ViewFamily, const FMinimalSceneTextures& SceneTextures) = 0;
@@ -130,5 +142,5 @@ public:
 	virtual FVariableRateShadingImageManager::EVRSSourceType GetType() const { return FVariableRateShadingImageManager::EVRSSourceType::None; };
 
 	// Get VRS image to be used w/ debug overlay
-	virtual FRDGTextureRef GetDebugImage(FRDGBuilder& GraphBuilder, const FViewInfo& ViewInfo) = 0;
+	virtual FRDGTextureRef GetDebugImage(FRDGBuilder& GraphBuilder, const FViewInfo& ViewInfo, FVariableRateShadingImageManager::EVRSImageType ImageType) = 0;
 };
\ No newline at end of file
