Fix bug in the batched instance id allocation that causes instance updates to be lost if an add is performed at the same time.

diff --git a/Engine/Source/Runtime/Renderer/Private/RendererScene.cpp b/Engine/Source/Runtime/Renderer/Private/RendererScene.cpp
--- a/Engine/Source/Runtime/Renderer/Private/RendererScene.cpp
+++ b/Engine/Source/Runtime/Renderer/Private/RendererScene.cpp
@@ -5745,14 +5745,12 @@
 		return false;
 	};
 
-#if RHI_RAYTRACING
-	TArray<FPrimitiveSceneInfo*, SceneRenderingAllocator> RayTracingPrimitivesToUpdate;
-	RayTracingPrimitivesToUpdate.Reserve(UpdatedInstances.Num());
-	bool bUpdateCachedRayTracingInstances = false;
-#endif
 	
-	TArray<FPrimitiveSceneInfo*, SceneRenderingAllocator> UpdatedInstancesWithInstanceCountChange;
-	UpdatedInstancesWithInstanceCountChange.Reserve(UpdatedInstances.Num());
+	TArray<FPrimitiveSceneInfo*, SceneRenderingAllocator> PendingAllocateInstanceIds;
+	PendingAllocateInstanceIds.Reserve(UpdatedInstances.Num() + AddedLocalPrimitiveSceneInfos.Num());
+	// All added primitive scene infos need to be allocated.
+	PendingAllocateInstanceIds.Append(AddedLocalPrimitiveSceneInfos);
+
 	// Perform instance releases before adds, such that allocator consolidation can happen only once.
 	{
 		CSV_SCOPED_TIMING_STAT_EXCLUSIVE(UpdatePrimitiveInstances);
@@ -5766,59 +5764,22 @@
 			FPrimitiveSceneInfo* PrimitiveSceneInfo = PrimitiveSceneProxy->GetPrimitiveSceneInfo();
 			
 			// being added or deleted, skip update logic
-			if (PrimitiveSceneInfo->PackedIndex == INDEX_NONE || DeletedSceneInfos.Contains(PrimitiveSceneInfo))
+			if (DeletedSceneInfos.Contains(PrimitiveSceneInfo))
 			{
 				continue;
 			}
-			
-			QueueFlushVirtualTexture(PrimitiveSceneInfo);
 
 			// If we recorded no adds or removes the instance count has stayed the same.  Therefore the cached mesh draw commands do not
 			// need to be updated.  In situations where the instance count changes the mesh draw command stores the instance count which
 			// would need to be updated.
 			const bool bInstanceCountChanged = (UpdateInstance.Value.CmdBuffer.NumAdds > 0) || (UpdateInstance.Value.CmdBuffer.NumRemoves > 0);
-			const bool bUpdateStaticDrawLists = !PrimitiveSceneProxy->StaticElementsAlwaysUseProxyPrimitiveUniformBuffer() || bInstanceCountChanged;
 
-			if (QueueAddToScene(PrimitiveSceneInfo))
+			// Append to queue if not added (if it is also added it will already be queued up)
+			if (bInstanceCountChanged && PrimitiveSceneInfo->GetIndex() != INDEX_NONE)
 			{
-				PrimitiveSceneInfo->RemoveFromScene(bUpdateStaticDrawLists);
-
-				if (bUpdateStaticDrawLists)
-				{
-					QueueAddStaticMeshes(PrimitiveSceneInfo);
-				}
-				else
-				{
-#if RHI_RAYTRACING
-					RayTracingPrimitivesToUpdate.Add(PrimitiveSceneInfo);
-					bUpdateCachedRayTracingInstances = true;
-#endif
-				}
-			}
-
-			// Update the Proxy's data.
-			PrimitiveSceneProxy->UpdateInstances_RenderThread(UpdateInstance.Value.CmdBuffer, UpdateInstance.Value.WorldBounds, UpdateInstance.Value.LocalBounds, UpdateInstance.Value.StaticMeshBounds);
-
-			if (!RHISupportsVolumeTextures(GetFeatureLevel())
-				&& (PrimitiveSceneProxy->IsMovable() || PrimitiveSceneProxy->NeedsUnbuiltPreviewLighting() || PrimitiveSceneProxy->GetLightmapType() == ELightmapType::ForceVolumetric))
-			{
-				PrimitiveSceneInfo->MarkIndirectLightingCacheBufferDirty();
-			}
-
-			if (bInstanceCountChanged)
-			{
 				PrimitiveSceneInfo->FreeGPUSceneInstances();
-				UpdatedInstancesWithInstanceCountChange.Add(PrimitiveSceneInfo);
+				PendingAllocateInstanceIds.Add(PrimitiveSceneInfo);
 			}
-			else
-			{
-				GPUScene.AddPrimitiveToUpdate(PrimitiveSceneInfo->PackedIndex, EPrimitiveDirtyState::ChangedAll);
-
-				DistanceFieldSceneData.UpdatePrimitive(PrimitiveSceneInfo);
-				LumenUpdatePrimitive(PrimitiveSceneInfo);
-			}
-
-			bNeedPathTracedInvalidation = bNeedPathTracedInvalidation || IsPrimitiveRelevantToPathTracing(PrimitiveSceneInfo);
 		}
 	}
 
@@ -6026,8 +5987,6 @@
 				PrimitiveSceneInfo->LinkAttachmentGroup();
 			}
 
-			FPrimitiveSceneInfo::AllocateGPUSceneInstances(this, TArrayView<FPrimitiveSceneInfo*>(&AddedLocalPrimitiveSceneInfos[StartIndex], AddedLocalPrimitiveSceneInfos.Num() - StartIndex));
-
 			for (int AddIndex = StartIndex; AddIndex < AddedLocalPrimitiveSceneInfos.Num(); AddIndex++)
 			{
 				FPrimitiveSceneInfo* PrimitiveSceneInfo = AddedLocalPrimitiveSceneInfos[AddIndex];
@@ -6142,31 +6101,84 @@
 		}
 	}
 
-	// handle scene changes
-	for (FVirtualShadowMapArrayCacheManager* CacheManager : VirtualShadowCacheManagers)
+
+#if RHI_RAYTRACING
+	TArray<FPrimitiveSceneInfo*, SceneRenderingAllocator> RayTracingPrimitivesToUpdate;
+	RayTracingPrimitivesToUpdate.Reserve(UpdatedInstances.Num());
+	bool bUpdateCachedRayTracingInstances = false;
+#endif
+
+	// Updating step of instance update. Note that instance IDs are not allocated at this point.
 	{
-		CacheManager->OnSceneChange();
-	}
-
-	// Re-adding step of instance update for those that had the instance count updated.
-	{
 		CSV_SCOPED_TIMING_STAT_EXCLUSIVE(UpdatePrimitiveInstances);
 		SCOPED_NAMED_EVENT(FScene_UpdatePrimitiveInstances, FColor::Emerald);
 		SCOPE_CYCLE_COUNTER(STAT_UpdatePrimitiveInstanceRenderThreadTime);
 
-		FPrimitiveSceneInfo::AllocateGPUSceneInstances(this, UpdatedInstancesWithInstanceCountChange);
-
-		for (FPrimitiveSceneInfo* PrimitiveSceneInfo : UpdatedInstancesWithInstanceCountChange)
+		for (const auto& UpdateInstance : UpdatedInstances)
 		{
-			check(PrimitiveSceneInfo->PackedIndex != INDEX_NONE);
-			FScopeCycleCounter Context(PrimitiveSceneInfo->Proxy->GetStatId());
+			FPrimitiveSceneProxy* PrimitiveSceneProxy = UpdateInstance.Key;
+			FScopeCycleCounter Context(PrimitiveSceneProxy->GetStatId());
+			FPrimitiveSceneInfo* PrimitiveSceneInfo = PrimitiveSceneProxy->GetPrimitiveSceneInfo();
 			
+			// being added or deleted, skip update logic
+			if (DeletedSceneInfos.Contains(PrimitiveSceneInfo))
+			{
+				continue;
+			}
 
-			DistanceFieldSceneData.RemovePrimitive(PrimitiveSceneInfo);
-			DistanceFieldSceneData.AddPrimitive(PrimitiveSceneInfo);
+			
+			QueueFlushVirtualTexture(PrimitiveSceneInfo);
 
-			LumenRemovePrimitive(PrimitiveSceneInfo, PrimitiveSceneInfo->GetIndex());
-			LumenAddPrimitive(PrimitiveSceneInfo);
+			// If we recorded no adds or removes the instance count has stayed the same.  Therefore the cached mesh draw commands do not
+			// need to be updated.  In situations where the instance count changes the mesh draw command stores the instance count which
+			// would need to be updated.
+			const bool bInstanceCountChanged = (UpdateInstance.Value.CmdBuffer.NumAdds > 0) || (UpdateInstance.Value.CmdBuffer.NumRemoves > 0);
+
+			const bool bUpdateStaticDrawLists = !PrimitiveSceneProxy->StaticElementsAlwaysUseProxyPrimitiveUniformBuffer() || bInstanceCountChanged;
+
+			if (QueueAddToScene(PrimitiveSceneInfo))
+			{
+				PrimitiveSceneInfo->RemoveFromScene(bUpdateStaticDrawLists);
+
+				if (bUpdateStaticDrawLists)
+				{
+					QueueAddStaticMeshes(PrimitiveSceneInfo);
+				}
+				else
+				{
+#if RHI_RAYTRACING
+					RayTracingPrimitivesToUpdate.Add(PrimitiveSceneInfo);
+					bUpdateCachedRayTracingInstances = true;
+#endif
+				}
+			}
+
+			// Update the Proxy's data.
+			PrimitiveSceneProxy->UpdateInstances_RenderThread(UpdateInstance.Value.CmdBuffer, UpdateInstance.Value.WorldBounds, UpdateInstance.Value.LocalBounds, UpdateInstance.Value.StaticMeshBounds);
+
+			if (!RHISupportsVolumeTextures(GetFeatureLevel())
+				&& (PrimitiveSceneProxy->IsMovable() || PrimitiveSceneProxy->NeedsUnbuiltPreviewLighting() || PrimitiveSceneProxy->GetLightmapType() == ELightmapType::ForceVolumetric))
+			{
+				PrimitiveSceneInfo->MarkIndirectLightingCacheBufferDirty();
+			}
+
+			if (bInstanceCountChanged)
+			{
+				DistanceFieldSceneData.RemovePrimitive(PrimitiveSceneInfo);
+				DistanceFieldSceneData.AddPrimitive(PrimitiveSceneInfo);
+
+				LumenRemovePrimitive(PrimitiveSceneInfo, PrimitiveSceneInfo->GetIndex());
+				LumenAddPrimitive(PrimitiveSceneInfo);
+			}
+			else
+			{
+				GPUScene.AddPrimitiveToUpdate(PrimitiveSceneInfo->PackedIndex, EPrimitiveDirtyState::ChangedAll);
+
+				DistanceFieldSceneData.UpdatePrimitive(PrimitiveSceneInfo);
+				LumenUpdatePrimitive(PrimitiveSceneInfo);
+			}
+
+			bNeedPathTracedInvalidation = bNeedPathTracedInvalidation || IsPrimitiveRelevantToPathTracing(PrimitiveSceneInfo);
 		}
 
 #if RHI_RAYTRACING
@@ -6176,6 +6188,15 @@
 #endif
 	}
 
+	// Allocate all instance slots. Needs to happen after the instance datas are updated since that may change the counts.
+	FPrimitiveSceneInfo::AllocateGPUSceneInstances(this, PendingAllocateInstanceIds);
+
+	// handle scene changes
+	for (FVirtualShadowMapArrayCacheManager* CacheManager : VirtualShadowCacheManagers)
+	{
+		CacheManager->OnSceneChange();
+	}
+
 	if (SceneInfosWithAddToScene.Num() > 0)
 	{
 		FPrimitiveSceneInfo::AddToScene(this, SceneInfosWithAddToScene);
