First simple LocalHeightFog volume component support.
Lots of missing bits but that is a good first step for iterations.

diff --git a/Engine/Shaders/Private/LocalHeightFog.usf b/Engine/Shaders/Private/LocalHeightFog.usf
new file mode 100644
--- /dev/null
+++ b/Engine/Shaders/Private/LocalHeightFog.usf
@@ -0,0 +1,224 @@
+// Copyright Epic Games, Inc. All Rights Reserved.
+
+#include "Common.ush"
+#include "SceneTexturesCommon.ush"
+#include "HeightFogCommon.ush"
+#include "ParticipatingMediaCommon.ush"
+
+
+struct FLocalHeightFogGPUInstanceData
+{
+	float4x4 Transform;
+	float4x4 InvTransform;
+
+	float Density;
+	float HeightFalloff;
+	float HeightOffset;
+	float RadialAttenuationPower;
+
+	float3 Albedo;
+	float  PhaseG;
+	float3 Emissive;
+	float RadialAttenuationDepth;
+};
+
+StructuredBuffer<FLocalHeightFogGPUInstanceData> LocalHeightFogInstances;
+
+
+struct FLocalHeightFogVertexOutput
+{
+	nointerpolation uint InstanceId : TEXCOORD0;
+};
+
+void LocalHeightFogSplatVS(
+	in float4 InPosition : ATTRIBUTE0,
+	uint InstanceIndex : SV_InstanceID,
+	out float4 OutPosition : SV_POSITION,
+	out FLocalHeightFogVertexOutput OtherVertexOutput
+)
+{
+	ResolvedView = ResolveView();
+
+	float4 LocalPosition = InPosition;
+
+	FLocalHeightFogGPUInstanceData FogInstance = LocalHeightFogInstances[InstanceIndex];
+
+	float4x4 FogTransform = FogInstance.Transform;
+
+	float4 WorldPosition = mul(float4(LocalPosition.xyz, 1.0f), FogTransform);
+
+	float4 TranslatedWorldPosition = float4(WorldPosition.xyz + LWCHackToFloat(PrimaryView.PreViewTranslation), 1);
+
+	OutPosition = mul(TranslatedWorldPosition, PrimaryView.TranslatedWorldToClip);
+	OtherVertexOutput.InstanceId = InstanceIndex;
+} 
+
+// This is an integral of height fog where density is defined as "DensityScale * Exp[-VerticalFallOff * z]"
+float HeightFogRayIntegral(FLocalHeightFogGPUInstanceData FogInstance, float3 CameraPos, float3 RayStart, float3 RayDir, float RayLength)
+{
+	float StartHeight = RayStart.z - FogInstance.HeightOffset;
+#if 0
+	float OpticalDepth = (FogInstance.Density / FogInstance.HeightFalloff) * exp(-StartHeight * FogInstance.HeightFalloff) * (1.0 - exp(-RayDir.z * RayLength * FogInstance.HeightFalloff)) / RayDir.z;
+#else
+	float OpticalDepth = (FogInstance.Density / FogInstance.HeightFalloff) *(exp(-StartHeight * FogInstance.HeightFalloff) - exp(-StartHeight * FogInstance.HeightFalloff -RayDir.z * RayLength * FogInstance.HeightFalloff)) / RayDir.z;
+#endif
+
+	return OpticalDepth;
+}
+
+// CameraPos can be behind of RayStart if the camera is out of the volume
+struct FFogData
+{
+	float Coverage;
+};
+FFogData ComputeFogCoverage(FLocalHeightFogGPUInstanceData FogInstance,  float3 CameraPos, float3 RayStart, float3 RayDir, float RayLength, float RadialAttenuation)
+{
+	FFogData FogData;
+	// Integral of density as extinction over a distance is optical depth.
+	float OpticalDepth	= HeightFogRayIntegral(FogInstance, CameraPos, RayStart, RayDir, RayLength);
+	float Transmittance	= exp(-OpticalDepth);
+	FogData.Coverage	= (1.0 - Transmittance) * RadialAttenuation;	// Only coverage is affected 
+	return FogData;
+}
+
+
+float3 ComputeFogInscattering(in FLocalHeightFogGPUInstanceData FogInstance, in FFogData FogData, in float3 RayDir)
+{
+	half3 InScattering = 0;// float3(0.1, 0.5, 0.5);// FogStruct.ExponentialFogColorParameter.xyz;
+
+#if PROJECT_SUPPORT_SKY_ATMOSPHERE_AFFECTS_HEIGHFOG
+	InScattering +=  Texture2DSampleLevel(View.DistantSkyLightLutTexture, View.DistantSkyLightLutTextureSampler, float2(0.5f, 0.5f), 0.0f).rgb;
+
+	// No need to test View.AtmosphereLightIlluminanceOnGroundPostTransmittance[0].a because InscatteringLightDirection.w above is doing the same test already.
+	InScattering += View.AtmosphereLightIlluminanceOnGroundPostTransmittance[0].rgb * SchlickPhase(-FogInstance.PhaseG, dot(RayDir, View.AtmosphereLightDirection[0].xyz));
+
+	if (View.AtmosphereLightIlluminanceOnGroundPostTransmittance[1].a > 0.0f) // Skip the second light when disabled.
+	{
+		InScattering += View.AtmosphereLightIlluminanceOnGroundPostTransmittance[0].rgb * SchlickPhase(-FogInstance.PhaseG, dot(RayDir, View.AtmosphereLightDirection[1].xyz));
+	}
+#endif
+
+	InScattering *= FogInstance.Albedo * View.SkyAtmosphereHeightFogContribution.xxx;
+
+	// Now account for emissive luminance after albedo and heighfog contribution has been accounted for.
+	InScattering += FogInstance.Emissive;
+
+	return InScattering * FogData.Coverage;
+}
+
+/** . */
+void LocalHeightFogSplatPS(
+	in float4 SVPos : SV_POSITION,
+	in FLocalHeightFogVertexOutput OtherVertexOutput,
+	out float4 OutColor : SV_Target0)
+{
+	ResolvedView = ResolveView();
+	uint InstanceIndex = OtherVertexOutput.InstanceId;
+
+	float3 CamRayTranslatedWorldOrigin	= PrimaryView.TranslatedWorldCameraOrigin;
+	float3 CamRayTranslatedWorldDir = normalize(SvPositionToTranslatedWorld(float4(SVPos.xy, 0.5, 1.0)));
+	float3 CamRayWorldDir = CamRayTranslatedWorldDir;
+
+	FLocalHeightFogGPUInstanceData FogInstance = LocalHeightFogInstances[InstanceIndex];
+	float4x4 Transform = FogInstance.Transform;
+	float4x4 InvTransform = FogInstance.InvTransform;
+
+	OutColor = float4(0, 0, 0, 1);
+
+	float3 CamRayWorldOrigin = CamRayTranslatedWorldOrigin - LWCHackToFloat(PrimaryView.PreViewTranslation);
+	float4 RayPosU = mul(float4(CamRayWorldOrigin, 1.0), InvTransform);
+	float4 RayDirU = mul(float4(CamRayTranslatedWorldDir, 0.0), InvTransform);
+	RayDirU.xyz = normalize(RayDirU.xyz);
+
+	float2 Ts = RayIntersectSphere(RayPosU.xyz, RayDirU.xyz, float4(0.0, 0.0, 0.0, 1.0));
+	if (any(Ts > 0.0))
+	{
+#if 1
+		// World space computations
+		float3 P0 = mul(float4(RayPosU.xyz + max(0.0, Ts.x) * RayDirU.xyz, 1.0), Transform).xyz;
+		float3 P1 = mul(float4(RayPosU.xyz + max(0.0, Ts.y) * RayDirU.xyz, 1.0), Transform).xyz;
+
+		float RadialAttenuation = 1.0f;
+#if 0	// attenuation with normalised sphere. But has artefact or vanishes to early
+		float USphereRayThickness = max(0.0, max(Ts.x,Ts.y)) - max(0.0, min(Ts.x, Ts.y));
+		float USphereLinearAttenuation = asinFast(USphereRayThickness * 0.5) * (1.0f / (PI * 0.5));
+		RadialAttenuation = pow(USphereLinearAttenuation, FogInstance.RadialAttenuationPower);
+#endif
+
+		float Length0 = length(P0 - CamRayWorldOrigin);
+		float Length1 = length(P1 - CamRayWorldOrigin);
+
+		const float DeviceZ = LookupDeviceZ(uint2(SVPos.xy));
+		float3 DepthBufferTranslatedWorldPos = SvPositionToTranslatedWorld(float4(SVPos.xy, DeviceZ, 1.0));
+		float LengthD = length(DepthBufferTranslatedWorldPos - CamRayTranslatedWorldOrigin);
+
+#if 1
+		// Attenuation based on world distance travelled in sphere, ignoring depth intersection
+		float RayLengthInSphere = max(0.0, max(Length0, Length1)) - max(0.0, min(Length0, Length1));
+		RadialAttenuation = saturate(RayLengthInSphere / FogInstance.RadialAttenuationDepth);
+		RadialAttenuation = pow(RadialAttenuation, FogInstance.RadialAttenuationPower);
+		//RadialAttenuation = 1.0 - RadialAttenuation;
+		//RadialAttenuation = 1.0 / (1.0 + exp(-RadialAttenuation * FogInstance.RadialAttenuationPower));	// Sigmoid: do better with range over -x/+x
+		//RadialAttenuation = 1.0 - RadialAttenuation;
+#endif
+
+		Length0 = min(LengthD, Length0);
+		Length1 = min(LengthD, Length1);
+		if (Length1 < Length0)
+		{
+			float LengthX = Length1;
+			Length1 = Length0;
+			Length0 = LengthX;
+		}
+		float WorldLength = max(0.0, Length1 - Length0);
+
+		if (WorldLength > 0.0)
+		{
+			FFogData FogData = ComputeFogCoverage(FogInstance, CamRayWorldOrigin, CamRayWorldOrigin + CamRayWorldDir * Length0, CamRayWorldDir, WorldLength, RadialAttenuation);
+
+			OutColor.xyz = ComputeFogInscattering(FogInstance, FogData, CamRayWorldDir);	// = Luminance
+			OutColor.a		 = 1.0 - FogData.Coverage;										// = Transmittance
+		}
+		else
+		{
+			clip(-1.0);
+		}
+
+		OutColor.xyz *= PrimaryView.PreExposure;
+#else
+		// Local space compuatations
+		const float DeviceZ = LookupDeviceZ(uint2(SVPos.xy));
+		float3 DepthBufferTranslatedWorldPos = SvPositionToTranslatedWorld(float4(SVPos.xy, DeviceZ, 1.0));
+		float3 DepthBufferWorldPos = DepthBufferTranslatedWorldPos - LWCHackToFloat(PrimaryView.PreViewTranslation).xyz;
+		float3 DepthBufferPosU = mul(float4(DepthBufferWorldPos, 1.0), InvTransform);
+		float3 LengthD = length(DepthBufferPosU - RayPosU);								// TODO account for scale!
+
+		float3 CamPosU = mul(float4(CamRayWorldOrigin, 1.0), InvTransform).xyz;
+
+		float Length0 = length((RayPosU.xyz + max(0.0, Ts.x) * RayDirU.xyz) - CamPosU);
+		float Length1 = length((RayPosU.xyz + max(0.0, Ts.y) * RayDirU.xyz) - CamPosU);
+
+		float UnitVolumeLength = max(0.0, abs(Length1 - Length0));
+
+		if (UnitVolumeLength > 0.0)
+		{
+			OutColor = ComputeFog(FogInstance, 10.0*CamPosU, 10.0*(CamPosU + RayDirU * Length0), RayDirU, 10.0*UnitVolumeLength);
+		}
+		else
+		{
+			clip(-1.0);
+		}
+
+		OutColor.xyz *= PrimaryView.PreExposure;
+
+#endif
+
+	}
+	else
+	{
+		clip(-1.0);
+	}
+}
+
+
+ 
\ No newline at end of file
diff --git a/Engine/Source/Editor/ComponentVisualizers/Private/ComponentVisualizers.cpp b/Engine/Source/Editor/ComponentVisualizers/Private/ComponentVisualizers.cpp
--- a/Engine/Source/Editor/ComponentVisualizers/Private/ComponentVisualizers.cpp
+++ b/Engine/Source/Editor/ComponentVisualizers/Private/ComponentVisualizers.cpp
@@ -14,6 +14,7 @@
 #include "Components/SpotLightComponent.h"
 #include "Components/StereoLayerComponent.h"
 #include "Components/WorldPartitionStreamingSourceComponent.h"
+#include "Components/LocalHeightFogComponent.h"
 #include "ConstraintComponentVisualizer.h"
 #include "DecalComponentVisualizer.h"
 #include "Editor/UnrealEdEngine.h"
@@ -40,6 +41,7 @@
 #include "UObject/Class.h"
 #include "UnrealEdGlobals.h"
 #include "WorldPartitionStreamingSourceComponentVisualizer.h"
+#include "LocalHeightFogComponentVisualizer.h"
 
 IMPLEMENT_MODULE( FComponentVisualizersModule, ComponentVisualizers );
 
@@ -61,6 +63,7 @@ void FComponentVisualizersModule::StartupModule()
 	RegisterComponentVisualizer(UDecalComponent::StaticClass()->GetFName(), MakeShareable(new FDecalComponentVisualizer));
 	RegisterComponentVisualizer(UStereoLayerComponent::StaticClass()->GetFName(), MakeShareable(new FStereoLayerComponentVisualizer));
 	RegisterComponentVisualizer(UWorldPartitionStreamingSourceComponent::StaticClass()->GetFName(), MakeShareable(new FWorldPartitionStreamingSourceComponentVisualizer));
+	RegisterComponentVisualizer(ULocalHeightFogComponent::StaticClass()->GetFName(), MakeShareable(new FLocalHeightFogComponentVisualizer));
 }
 
 void FComponentVisualizersModule::ShutdownModule()
diff --git a/Engine/Source/Editor/ComponentVisualizers/Private/LocalHeightFogComponentVisualizer.cpp b/Engine/Source/Editor/ComponentVisualizers/Private/LocalHeightFogComponentVisualizer.cpp
new file mode 100644
--- /dev/null
+++ b/Engine/Source/Editor/ComponentVisualizers/Private/LocalHeightFogComponentVisualizer.cpp
@@ -0,0 +1,23 @@
+// Copyright Epic Games, Inc. All Rights Reserved.
+
+#include "LocalHeightFogComponentVisualizer.h"
+
+#include "Components/ActorComponent.h"
+#include "Components/LocalHeightFogComponent.h"
+
+void FLocalHeightFogComponentVisualizer::DrawVisualization( const UActorComponent* Component, const FSceneView* View, FPrimitiveDrawInterface* PDI )
+{
+	if(View->Family->EngineShowFlags.CompositeEditorPrimitives)
+	{
+		const ULocalHeightFogComponent* LocalHeightFog = Cast<const ULocalHeightFogComponent>(Component);
+		if(LocalHeightFog != NULL)
+		{
+			FTransform LocalHeightFogTransform = LocalHeightFog->GetComponentTransform();
+
+			// Draw local fog volume spherical shape radius
+			DrawWireSphereAutoSides(PDI, LocalHeightFogTransform, FColor(200, 255, 255), 1.0f, SDPG_World);
+		}
+	}
+}
+
+
diff --git a/Engine/Source/Editor/ComponentVisualizers/Public/LocalHeightFogComponentVisualizer.h b/Engine/Source/Editor/ComponentVisualizers/Public/LocalHeightFogComponentVisualizer.h
new file mode 100644
--- /dev/null
+++ b/Engine/Source/Editor/ComponentVisualizers/Public/LocalHeightFogComponentVisualizer.h
@@ -0,0 +1,18 @@
+// Copyright Epic Games, Inc. All Rights Reserved.
+
+#pragma once
+
+#include "CoreMinimal.h"
+#include "ComponentVisualizer.h"
+
+class FPrimitiveDrawInterface;
+class FSceneView;
+class UActorComponent;
+
+class COMPONENTVISUALIZERS_API FLocalHeightFogComponentVisualizer : public FComponentVisualizer
+{
+public:
+	//~ Begin FComponentVisualizer Interface
+	virtual void DrawVisualization(const UActorComponent* Component, const FSceneView* View, FPrimitiveDrawInterface* PDI) override;
+	//~ End FComponentVisualizer Interface
+};
diff --git a/Engine/Source/Editor/LevelEditor/Private/LevelEditorOutlinerSettings.cpp b/Engine/Source/Editor/LevelEditor/Private/LevelEditorOutlinerSettings.cpp
--- a/Engine/Source/Editor/LevelEditor/Private/LevelEditorOutlinerSettings.cpp
+++ b/Engine/Source/Editor/LevelEditor/Private/LevelEditorOutlinerSettings.cpp
@@ -19,6 +19,7 @@
 #include "Animation/SkeletalMeshActor.h"
 #include "Components/SkyAtmosphereComponent.h"
 #include "Components/VolumetricCloudComponent.h"
+#include "Engine/LocalHeightFog.h"
 #include "Engine/Blueprint.h"
 #include "Engine/Brush.h"
 #include "Engine/ExponentialHeightFog.h"
@@ -276,6 +277,7 @@ void FLevelEditorOutlinerSettings::CreateDefaultFilters()
 		FindOrAddClassFilter(AExponentialHeightFog::StaticClass(), *EnvironmentFilterCategory);
 		FindOrAddClassFilter(AInstancedFoliageActor::StaticClass(), *EnvironmentFilterCategory);
 		FindOrAddClassFilter(ASkyAtmosphere::StaticClass(), *EnvironmentFilterCategory);
+		FindOrAddClassFilter(ALocalHeightFog::StaticClass(), *EnvironmentFilterCategory);
 		FindOrAddClassFilter(AVolumetricCloud::StaticClass(), *EnvironmentFilterCategory);
 	}
 	
diff --git a/Engine/Source/Editor/PlacementMode/Private/PlacementModeModule.cpp b/Engine/Source/Editor/PlacementMode/Private/PlacementModeModule.cpp
--- a/Engine/Source/Editor/PlacementMode/Private/PlacementModeModule.cpp
+++ b/Engine/Source/Editor/PlacementMode/Private/PlacementModeModule.cpp
@@ -22,6 +22,7 @@
 #include "ActorFactories/ActorFactoryEmptyActor.h"
 #include "ActorFactories/ActorFactoryPawn.h"
 #include "ActorFactories/ActorFactoryExponentialHeightFog.h"
+#include "ActorFactories/ActorFactoryLocalHeightFog.h"
 #include "ActorFactories/ActorFactorySkyAtmosphere.h"
 #include "ActorFactories/ActorFactoryVolumetricCloud.h"
 #include "ActorFactories/ActorFactoryPlayerStart.h"
@@ -207,6 +208,7 @@ void FPlacementModeModule::StartupModule()
 		Category->Items.Add(CreateID(), MakeShareable(new FPlaceableItem(*UActorFactorySkyAtmosphere::StaticClass(), SortOrder += 10)));
 		Category->Items.Add(CreateID(), MakeShareable(new FPlaceableItem(*UActorFactoryVolumetricCloud::StaticClass(), SortOrder += 10)));
 		Category->Items.Add(CreateID(), MakeShareable(new FPlaceableItem(*UActorFactoryExponentialHeightFog::StaticClass(), SortOrder += 10)));
+		Category->Items.Add(CreateID(), MakeShareable(new FPlaceableItem(*UActorFactoryLocalHeightFog::StaticClass(), SortOrder += 10)));
 		Category->Items.Add(CreateID(), MakeShareable(new FPlaceableItem(*UActorFactorySphereReflectionCapture::StaticClass(), SortOrder += 10)));
 		Category->Items.Add(CreateID(), MakeShareable(new FPlaceableItem(*UActorFactoryBoxReflectionCapture::StaticClass(), SortOrder += 10)));
 		Category->Items.Add(CreateID(), MakeShareable(new FPlaceableItem(*UActorFactoryPlanarReflection::StaticClass(), SortOrder += 10)));
diff --git a/Engine/Source/Editor/UnrealEd/Classes/ActorFactories/ActorFactoryLocalHeightFog.h b/Engine/Source/Editor/UnrealEd/Classes/ActorFactories/ActorFactoryLocalHeightFog.h
new file mode 100644
--- /dev/null
+++ b/Engine/Source/Editor/UnrealEd/Classes/ActorFactories/ActorFactoryLocalHeightFog.h
@@ -0,0 +1,15 @@
+// Copyright Epic Games, Inc. All Rights Reserved.
+
+#pragma once
+
+#include "CoreMinimal.h"
+#include "UObject/ObjectMacros.h"
+#include "ActorFactories/ActorFactory.h"
+#include "ActorFactoryLocalHeightFog.generated.h"
+
+UCLASS(MinimalAPI, config = Editor, collapsecategories, hidecategories = Object)
+class UActorFactoryLocalHeightFog : public UActorFactory
+{
+	GENERATED_UCLASS_BODY()
+};
+
diff --git a/Engine/Source/Editor/UnrealEd/Private/Factories/ActorFactory.cpp b/Engine/Source/Editor/UnrealEd/Private/Factories/ActorFactory.cpp
--- a/Engine/Source/Editor/UnrealEd/Private/Factories/ActorFactory.cpp
+++ b/Engine/Source/Editor/UnrealEd/Private/Factories/ActorFactory.cpp
@@ -31,6 +31,7 @@ ActorFactory.cpp:
 #include "ActorFactories/ActorFactoryEmptyActor.h"
 #include "ActorFactories/ActorFactoryPawn.h"
 #include "ActorFactories/ActorFactoryExponentialHeightFog.h"
+#include "ActorFactories/ActorFactoryLocalHeightFog.h"
 #include "ActorFactories/ActorFactoryNote.h"
 #include "ActorFactories/ActorFactoryPhysicsAsset.h"
 #include "ActorFactories/ActorFactoryPlaneReflectionCapture.h"
@@ -79,6 +80,7 @@ ActorFactory.cpp:
 #include "Engine/DecalActor.h"
 #include "Atmosphere/AtmosphericFog.h"
 #include "Engine/ExponentialHeightFog.h"
+#include "Engine/LocalHeightFog.h"
 #include "Engine/SkyLight.h"
 #include "Engine/DirectionalLight.h"
 #include "Engine/PointLight.h"
@@ -1874,6 +1876,16 @@ UActorFactoryExponentialHeightFog::UActorFactoryExponentialHeightFog(const FObje
 	NewActorClass = AExponentialHeightFog::StaticClass();
 }
 
+/*-----------------------------------------------------------------------------
+UActorFactoryLocalHeightFog
+-----------------------------------------------------------------------------*/
+UActorFactoryLocalHeightFog::UActorFactoryLocalHeightFog(const FObjectInitializer& ObjectInitializer)
+	: Super(ObjectInitializer)
+{
+	DisplayName = LOCTEXT("LocalHeightFogDisplayName", "Local Height Fog");
+	NewActorClass = ALocalHeightFog::StaticClass();
+}
+
 /*-----------------------------------------------------------------------------
 UActorFactoryInteractiveFoliage
 -----------------------------------------------------------------------------*/
diff --git a/Engine/Source/Runtime/Engine/Classes/Components/LocalHeightFogComponent.h b/Engine/Source/Runtime/Engine/Classes/Components/LocalHeightFogComponent.h
new file mode 100644
--- /dev/null
+++ b/Engine/Source/Runtime/Engine/Classes/Components/LocalHeightFogComponent.h
@@ -0,0 +1,75 @@
+// Copyright Epic Games, Inc. All Rights Reserved.
+
+#pragma once
+
+#include "CoreMinimal.h"
+#include "Components/SceneComponent.h"
+#include "Math/Color.h"
+#include "LocalHeightFogComponent.generated.h"
+
+class FLocalHeightFogSceneProxy;
+
+UCLASS(ClassGroup = Rendering, collapsecategories, hidecategories = (Object, Mobility, Activation, "Components|Activation"), editinlinenew, meta = (BlueprintSpawnableComponent), MinimalAPI)
+class ULocalHeightFogComponent : public USceneComponent
+{
+	GENERATED_UCLASS_BODY()
+
+	~ULocalHeightFogComponent();
+
+	/** Global density factor for this fog. */
+	UPROPERTY(EditAnywhere, BlueprintReadOnly, interp, Category = "Fog Distribution", meta = (UIMin = "0", UIMax = ".05", SliderExponent = 4.0, ClampMin = 0.0))
+	float FogDensity = 0.02f;
+
+	/** Controls how the density decreases as height increases. Smaller values make the visible transition larger. */
+	UPROPERTY(EditAnywhere, BlueprintReadOnly, interp, Category = "Fog Distribution", meta = (UIMin = "0.001", UIMax = "10.0", SliderExponent = 4.0, ClampMin = 0.001))
+	float FogHeightFalloff = 3.0f;
+
+	/** Height offset, relative to the actor Z position. */
+	UPROPERTY(EditAnywhere, BlueprintReadOnly, interp, Category = "Fog Distribution", meta = (UIMin = "-5000.0", UIMax = "5000.0"))
+	float FogHeightOffset = 0.0f;
+
+	/** Controls the radial attenuation of this fog volume. */
+	UPROPERTY(EditAnywhere, BlueprintReadOnly, interp, Category = "Fog Distribution", meta = (UIMin = "0.01", UIMax = "5.0", ClampMin = 0.01))
+	float FogRadialAttenuationPower = 1.0f;
+
+	/** Controls the radial attenuation of this fog volume. */
+	UPROPERTY(EditAnywhere, BlueprintReadOnly, interp, Category = "Fog Distribution", meta = (UIMin = "0.01", UIMax = "5000.0", SliderExponent = 4.0, ClampMin = 0.01))
+	float FogRadialAttenuationDepth = 1000.0f;
+
+	/** Controls the directionality of the scattering within this fog volume. */
+	UPROPERTY(EditAnywhere, BlueprintReadOnly, interp, Category = "Fog Shading", meta = (UIMin = "0.0", UIMax = "0.999", ClampMin = 0.0, ClampMax = 0.999))
+	float FogPhaseG = 0.8f;
+
+	/** Controls the albedo of this fog volume. */
+	UPROPERTY(EditAnywhere, BlueprintReadOnly, interp, Category = "Fog Shading", meta = (ClampMin = 0.0, ClampMax = 1.0))
+	FLinearColor FogAlbedo = FLinearColor::White;
+
+	/** Controls the emissive color of this fog volume. */
+	UPROPERTY(EditAnywhere, BlueprintReadOnly, interp, Category = "Fog Shading", meta = (ClampMin = 0.0, ClampMax = 1.0))
+	FLinearColor FogEmissive = FLinearColor::Black;
+
+public:
+
+	//~ Begin UObject Interface
+#if WITH_EDITOR
+	virtual void PostEditChangeProperty(FPropertyChangedEvent& PropertyChangedEvent) override;
+#endif // WITH_EDITOR
+	//~ End UObject Interface
+
+	/** Callback to create the rendering thread mirror. */
+	ENGINE_API FLocalHeightFogSceneProxy* CreateSceneProxy();
+
+protected:
+	//~ Begin UActorComponent Interface.
+	virtual void CreateRenderState_Concurrent(FRegisterComponentContext* Context) override;
+	virtual void SendRenderTransform_Concurrent() override;
+	virtual void DestroyRenderState_Concurrent() override;
+	//~ End UActorComponent Interface.
+
+	void SendRenderTransformCommand();
+
+private:
+
+	FLocalHeightFogSceneProxy* LocalHeightFogSceneProxy;
+};
+
diff --git a/Engine/Source/Runtime/Engine/Classes/Engine/LocalHeightFog.h b/Engine/Source/Runtime/Engine/Classes/Engine/LocalHeightFog.h
new file mode 100644
--- /dev/null
+++ b/Engine/Source/Runtime/Engine/Classes/Engine/LocalHeightFog.h
@@ -0,0 +1,34 @@
+// Copyright Epic Games, Inc. All Rights Reserved.
+
+#pragma once
+
+#include "CoreMinimal.h"
+#include "GameFramework/Info.h"
+#include "LocalHeightFog.generated.h"
+
+class ULocalHeightFogComponent;
+
+/**
+ *	Actor used to position a local fog volume in the scene.
+ *	@see https://docs.unrealengine.com/???
+ */
+UCLASS(showcategories = (Movement, Rendering, Transformation, DataLayers, "Input|MouseInput", "Input|TouchInput"), ClassGroup = Fog, hidecategories = (Info, Object, Input), MinimalAPI)
+class ALocalHeightFog : public AInfo
+{
+	GENERATED_UCLASS_BODY()
+
+private:
+#if WITH_EDITOR
+	virtual bool ActorTypeSupportsDataLayer() const override { return true; }
+#endif
+
+	/** Object used to visualize the local fog volume */
+	UPROPERTY(BlueprintReadOnly, VisibleAnywhere, Category = Fog, meta = (AllowPrivateAccess = "true"))
+	TObjectPtr<ULocalHeightFogComponent> LocalHeightFogVolume;
+
+public:
+
+	/** Returns LocalHeightFogVolume subobject **/
+	ENGINE_API ULocalHeightFogComponent* GetComponent() const { return LocalHeightFogVolume; }
+};
+
diff --git a/Engine/Source/Runtime/Engine/Private/Components/LocalHeightFogComponent.cpp b/Engine/Source/Runtime/Engine/Private/Components/LocalHeightFogComponent.cpp
new file mode 100644
--- /dev/null
+++ b/Engine/Source/Runtime/Engine/Private/Components/LocalHeightFogComponent.cpp
@@ -0,0 +1,146 @@
+// Copyright Epic Games, Inc. All Rights Reserved.
+
+#include "Components/LocalHeightFogComponent.h"
+#include "Engine/LocalHeightFog.h"
+#include "LocalHeightFogSceneProxy.h"
+#include "Components/ArrowComponent.h"
+#include "Components/DrawSphereComponent.h"
+#include "Components/BillboardComponent.h"
+
+#include UE_INLINE_GENERATED_CPP_BY_NAME(LocalHeightFogComponent)
+
+
+ULocalHeightFogComponent::ULocalHeightFogComponent(const FObjectInitializer& ObjectInitializer)
+	: Super(ObjectInitializer)
+	, LocalHeightFogSceneProxy(nullptr)
+{
+	Mobility = EComponentMobility::Movable;
+
+	// Set a size by default to be visible when droped in the scene
+	const float Size = 500.0f;
+	SetRelativeScale3D(FVector(Size, Size, Size));
+}
+
+ULocalHeightFogComponent::~ULocalHeightFogComponent()
+{
+}
+
+void ULocalHeightFogComponent::SendRenderTransformCommand()
+{
+	if (LocalHeightFogSceneProxy)
+	{
+		FTransform ComponentTransform = GetComponentTransform();
+		float HeightOffset = FogHeightOffset;
+		FLocalHeightFogSceneProxy* SceneProxy = LocalHeightFogSceneProxy;
+		ENQUEUE_RENDER_COMMAND(FUpdateLocalHeightFogSceneProxyTransformCommand)(
+			[SceneProxy, ComponentTransform, HeightOffset](FRHICommandList& RHICmdList)
+			{
+				SceneProxy->FogTransform = ComponentTransform;
+				SceneProxy->FogHeightOffset = HeightOffset + ComponentTransform.GetLocation().Z;
+			});
+	}
+}
+
+void ULocalHeightFogComponent::CreateRenderState_Concurrent(FRegisterComponentContext* Context)
+{
+	Super::CreateRenderState_Concurrent(Context);	
+	
+	bool bHidden = false;
+#if WITH_EDITORONLY_DATA
+	bHidden = GetOwner() ? GetOwner()->bHiddenEdLevel : false;
+#endif // WITH_EDITORONLY_DATA
+	if (!ShouldComponentAddToScene())
+	{
+		bHidden = true;
+	}
+
+	if (GetVisibleFlag() && !bHidden && ShouldComponentAddToScene() && ShouldRender() && IsRegistered() && (GetOuter() == NULL || !GetOuter()->HasAnyFlags(RF_ClassDefaultObject)))
+	{
+		LocalHeightFogSceneProxy = CreateSceneProxy();
+		GetWorld()->Scene->AddLocalHeightFog(LocalHeightFogSceneProxy);
+	}
+}
+
+void ULocalHeightFogComponent::SendRenderTransform_Concurrent()
+{
+	Super::SendRenderTransform_Concurrent();
+	SendRenderTransformCommand();
+}
+
+void ULocalHeightFogComponent::DestroyRenderState_Concurrent()
+{
+	Super::DestroyRenderState_Concurrent();
+
+	if (LocalHeightFogSceneProxy)
+	{
+		GetWorld()->Scene->RemoveLocalHeightFog(LocalHeightFogSceneProxy);
+
+		FLocalHeightFogSceneProxy* SceneProxy = LocalHeightFogSceneProxy;
+		ENQUEUE_RENDER_COMMAND(FDestroyLocalHeightFogSceneProxyCommand)(
+			[SceneProxy](FRHICommandList& RHICmdList)
+			{
+				delete SceneProxy;
+			});
+
+		LocalHeightFogSceneProxy = nullptr;
+	}
+}
+
+#if WITH_EDITOR
+
+void ULocalHeightFogComponent::PostEditChangeProperty(FPropertyChangedEvent& PropertyChangedEvent)
+{
+	Super::PostEditChangeProperty(PropertyChangedEvent);
+	SendRenderTransformCommand();
+}
+
+#endif // WITH_EDITOR
+
+FLocalHeightFogSceneProxy* ULocalHeightFogComponent::CreateSceneProxy()
+{
+	return new FLocalHeightFogSceneProxy(this);
+}
+
+#if WITH_EDITOR
+#include "ObjectEditorUtils.h"
+#endif
+
+ALocalHeightFog::ALocalHeightFog(const FObjectInitializer& ObjectInitializer)
+	: Super(ObjectInitializer.SetDefaultSubobjectClass<ULocalHeightFogComponent>(TEXT("NewLocalHeightFogComponent")))
+{
+	LocalHeightFogVolume = CreateDefaultSubobject<ULocalHeightFogComponent>(TEXT("LocalHeightFogComponent"));
+	RootComponent = LocalHeightFogVolume;
+
+#if WITH_EDITORONLY_DATA
+
+	if (!IsRunningCommandlet())
+	{
+		// Structure to hold one-time initialization
+		struct FConstructorStatics
+		{
+			ConstructorHelpers::FObjectFinderOptional<UTexture2D> LocalHeightFogTextureObject;
+			FName ID_LocalHeightFog;
+			FText NAME_LocalHeightFog;
+			FConstructorStatics()
+				: LocalHeightFogTextureObject(TEXT("/Engine/EditorResources/S_SkyAtmosphere"))	// TODO
+				, ID_LocalHeightFog(TEXT("Fog"))
+				, NAME_LocalHeightFog(NSLOCTEXT("SpriteCategory", "Fog", "Fog"))
+			{
+			}
+		};
+		static FConstructorStatics ConstructorStatics;
+
+		if (GetSpriteComponent())
+		{
+			GetSpriteComponent()->Sprite = ConstructorStatics.LocalHeightFogTextureObject.Get();
+			GetSpriteComponent()->SetRelativeScale3D(FVector(0.5f, 0.5f, 0.5f));
+			GetSpriteComponent()->SpriteInfo.Category = ConstructorStatics.ID_LocalHeightFog;
+			GetSpriteComponent()->SpriteInfo.DisplayName = ConstructorStatics.NAME_LocalHeightFog;
+			GetSpriteComponent()->SetupAttachment(LocalHeightFogVolume);
+		}
+	}
+#endif // WITH_EDITORONLY_DATA
+
+	PrimaryActorTick.bCanEverTick = true;
+	SetHidden(false);
+}
diff --git a/Engine/Source/Runtime/Engine/Private/LocalHeightFogSceneProxy.cpp b/Engine/Source/Runtime/Engine/Private/LocalHeightFogSceneProxy.cpp
new file mode 100644
--- /dev/null
+++ b/Engine/Source/Runtime/Engine/Private/LocalHeightFogSceneProxy.cpp
@@ -0,0 +1,29 @@
+// Copyright Epic Games, Inc. All Rights Reserved.
+
+/*=============================================================================
+	FVolumetricCloudSceneProxy implementation.
+=============================================================================*/
+
+#include "LocalHeightFogSceneProxy.h"
+#include "Components/LocalHeightFogComponent.h"
+
+
+
+FLocalHeightFogSceneProxy::FLocalHeightFogSceneProxy(const ULocalHeightFogComponent* InComponent)
+	: FogTransform(InComponent->GetComponentTransform())
+	, FogDensity(InComponent->FogDensity)
+	, FogHeightFalloff(InComponent->FogHeightFalloff)
+	, FogHeightOffset(InComponent->FogHeightOffset + InComponent->GetComponentTransform().GetLocation().Z)
+	, FogRadialAttenuationPower(InComponent->FogRadialAttenuationPower)
+	, FogRadialAttenuationDepth(InComponent->FogRadialAttenuationDepth)
+	, FogPhaseG(InComponent->FogPhaseG)
+	, FogAlbedo(InComponent->FogAlbedo)
+	, FogEmissive(InComponent->FogEmissive)
+{
+}
+
+FLocalHeightFogSceneProxy::~FLocalHeightFogSceneProxy()
+{
+}
+
+
diff --git a/Engine/Source/Runtime/Engine/Public/LocalHeightFogSceneProxy.h b/Engine/Source/Runtime/Engine/Public/LocalHeightFogSceneProxy.h
new file mode 100644
--- /dev/null
+++ b/Engine/Source/Runtime/Engine/Public/LocalHeightFogSceneProxy.h
@@ -0,0 +1,31 @@
+// Copyright Epic Games, Inc. All Rights Reserved.
+
+#pragma once
+
+#include "CoreMinimal.h"
+#include "Math/Color.h"
+
+class ULocalHeightFogComponent;
+
+/** Represents a UVolumetricCloudComponent to the rendering thread, created game side from the component. */
+class ENGINE_API FLocalHeightFogSceneProxy
+{
+public:
+
+	// Initialization constructor.
+	FLocalHeightFogSceneProxy(const ULocalHeightFogComponent* InComponent);
+	~FLocalHeightFogSceneProxy();
+
+	FTransform FogTransform;
+
+	float FogDensity;
+	float FogHeightFalloff;
+	float FogHeightOffset;
+	float FogRadialAttenuationPower;
+	float FogRadialAttenuationDepth;
+
+	float FogPhaseG;
+	FLinearColor FogAlbedo;
+	FLinearColor FogEmissive;
+private:
+};
diff --git a/Engine/Source/Runtime/Engine/Public/SceneInterface.h b/Engine/Source/Runtime/Engine/Public/SceneInterface.h
--- a/Engine/Source/Runtime/Engine/Public/SceneInterface.h
+++ b/Engine/Source/Runtime/Engine/Public/SceneInterface.h
@@ -346,6 +346,23 @@ public:
 	 */
 	virtual bool HasAnyExponentialHeightFog() const = 0;
 
+	/**
+	 * Adds a new local height fog component to the scene
+	 *
+	 * @param FogProxy - fog proxy to add
+	 */
+	virtual void AddLocalHeightFog(class FLocalHeightFogSceneProxy* FogProxy) = 0;
+	/**
+	 * Removes a local height fog component from the scene
+	 *
+	 * @param FogProxy - fog proxy to remove
+	 */
+	virtual void RemoveLocalHeightFog(class FLocalHeightFogSceneProxy* FogProxy) = 0;
+	/**
+	 * @return True if there are any local height fog potentially enabled in the scene
+	 */
+	virtual bool HasAnyLocalHeightFog() const = 0;
+
 	/**
 	 * Adds the unique volumetric cloud component to the scene
 	 *
diff --git a/Engine/Source/Runtime/Renderer/Private/DeferredShadingRenderer.cpp b/Engine/Source/Runtime/Renderer/Private/DeferredShadingRenderer.cpp
--- a/Engine/Source/Runtime/Renderer/Private/DeferredShadingRenderer.cpp
+++ b/Engine/Source/Runtime/Renderer/Private/DeferredShadingRenderer.cpp
@@ -87,6 +87,7 @@
 #include "ComponentRecreateRenderStateContext.h"
 #include "RenderCore.h"
 #include "VariableRateShadingImageManager.h"
+#include "LocalHeightFogRendering.h"
 #include "Shadows/ShadowScene.h"
 #include "Lumen/LumenHardwareRayTracingCommon.h"
 
@@ -402,6 +403,7 @@ DECLARE_CYCLE_STAT(TEXT("DeferredShadingSceneRenderer RenderLightShaftOcclusion"
 DECLARE_CYCLE_STAT(TEXT("DeferredShadingSceneRenderer RenderAtmosphere"), STAT_FDeferredShadingSceneRenderer_RenderAtmosphere, STATGROUP_SceneRendering);
 DECLARE_CYCLE_STAT(TEXT("DeferredShadingSceneRenderer RenderSkyAtmosphere"), STAT_FDeferredShadingSceneRenderer_RenderSkyAtmosphere, STATGROUP_SceneRendering);
 DECLARE_CYCLE_STAT(TEXT("DeferredShadingSceneRenderer RenderFog"), STAT_FDeferredShadingSceneRenderer_RenderFog, STATGROUP_SceneRendering);
+DECLARE_CYCLE_STAT(TEXT("DeferredShadingSceneRenderer RenderLocalHeightFog"), STAT_FDeferredShadingSceneRenderer_RenderLocalHeightFog, STATGROUP_SceneRendering);
 DECLARE_CYCLE_STAT(TEXT("DeferredShadingSceneRenderer RenderLightShaftBloom"), STAT_FDeferredShadingSceneRenderer_RenderLightShaftBloom, STATGROUP_SceneRendering);
 DECLARE_CYCLE_STAT(TEXT("DeferredShadingSceneRenderer RenderFinish"), STAT_FDeferredShadingSceneRenderer_RenderFinish, STATGROUP_SceneRendering);
 
@@ -4000,6 +4002,13 @@ void FDeferredShadingSceneRenderer::Render(FRDGBuilder& GraphBuilder)
 		SCOPE_CYCLE_COUNTER(STAT_FDeferredShadingSceneRenderer_RenderFog);
 		RenderFog(GraphBuilder, SceneTextures, LightShaftOcclusionTexture);
 	}
+	if (!bHasRayTracedOverlay && ShouldRenderLocalHeightFog(Scene, ViewFamily))
+	{
+		RDG_CSV_STAT_EXCLUSIVE_SCOPE(GraphBuilder, RenderLocalHeightFog);
+		SCOPED_NAMED_EVENT(RenderLocalHeightFog, FColor::Emerald);
+		SCOPE_CYCLE_COUNTER(STAT_FDeferredShadingSceneRenderer_RenderLocalHeightFog);
+		RenderLocalHeightFog(Scene, Views, GraphBuilder, SceneTextures, LightShaftOcclusionTexture);
+	}
 
 	// After the height fog, Draw volumetric clouds (having fog applied on them already) when using per pixel tracing,
 	if (!bHasRayTracedOverlay && bShouldRenderVolumetricCloud)
diff --git a/Engine/Source/Runtime/Renderer/Private/LocalHeightFogRendering.cpp b/Engine/Source/Runtime/Renderer/Private/LocalHeightFogRendering.cpp
new file mode 100644
--- /dev/null
+++ b/Engine/Source/Runtime/Renderer/Private/LocalHeightFogRendering.cpp
@@ -0,0 +1,235 @@
+// Copyright Epic Games, Inc. All Rights Reserved.
+
+#include "LocalHeightFogRendering.h"
+#include "ScenePrivate.h"
+#include "RendererUtils.h"
+#include "ScreenPass.h"
+#include "LocalHeightFogSceneProxy.h"
+
+
+// The runtime ON/OFF toggle
+static TAutoConsoleVariable<int32> CVarLocalHeightFog(
+	TEXT("r.LocalHeightFog"), 1,
+	TEXT("LocalHeightFog components are rendered when this is not 0, otherwise ignored.\n"),
+	ECVF_RenderThreadSafe);
+
+bool ShouldRenderLocalHeightFog(const FScene* Scene, const FSceneViewFamily& Family)
+{
+	const FEngineShowFlags EngineShowFlags = Family.EngineShowFlags;
+	if (Scene && Scene->HasAnyLocalHeightFog() && EngineShowFlags.Fog && !Family.UseDebugViewPS())
+	{
+		return CVarLocalHeightFog.GetValueOnRenderThread() > 0;
+	}
+	return false;
+}
+
+
+
+/*=============================================================================
+	FScene functions
+=============================================================================*/
+
+
+
+void FScene::AddLocalHeightFog(class FLocalHeightFogSceneProxy* FogProxy)
+{
+	check(FogProxy);
+	FScene* Scene = this;
+
+	ENQUEUE_RENDER_COMMAND(FAddLocalHeightFogCommand)(
+		[Scene, FogProxy](FRHICommandListImmediate& RHICmdList)
+		{
+			check(!Scene->LocalHeightFogs.Contains(FogProxy));
+			Scene->LocalHeightFogs.Push(FogProxy);
+		} );
+}
+
+void FScene::RemoveLocalHeightFog(class FLocalHeightFogSceneProxy* FogProxy)
+{
+	check(FogProxy);
+	FScene* Scene = this;
+
+	ENQUEUE_RENDER_COMMAND(FRemoveLocalHeightFogCommand)(
+		[Scene, FogProxy](FRHICommandListImmediate& RHICmdList)
+		{
+			Scene->LocalHeightFogs.RemoveSingle(FogProxy);
+		} );
+}
+
+bool FScene::HasAnyLocalHeightFog() const
+{ 
+	return LocalHeightFogs.Num() > 0;
+}
+
+
+/*=============================================================================
+	Local height fog rendering functions
+=============================================================================*/
+
+class FLocalHeightFogGPUInstanceData
+{
+public:
+	FMatrix44f Transform;
+	FMatrix44f InvTransform;
+
+	float Density;
+	float HeightFalloff;
+	float HeightOffset;
+	float RadialAttenuationPower;
+
+	FVector3f Albedo;
+	float PhaseG;
+	FVector3f Emissive;
+	float RadialAttenuationDepth;
+};
+
+class FLocalHeightFogVS : public FGlobalShader
+{
+	DECLARE_GLOBAL_SHADER(FLocalHeightFogVS);
+	SHADER_USE_PARAMETER_STRUCT(FLocalHeightFogVS, FGlobalShader);
+
+	BEGIN_SHADER_PARAMETER_STRUCT(FParameters, )
+		SHADER_PARAMETER_STRUCT_REF(FViewUniformShaderParameters, View)
+		SHADER_PARAMETER_RDG_BUFFER_SRV(StructuredBuffer<FLocalHeightFogGPUInstanceData>, LocalHeightFogInstances)
+	END_SHADER_PARAMETER_STRUCT()
+
+	using FPermutationDomain = TShaderPermutationDomain<>;
+
+	static bool ShouldCompilePermutation(const FGlobalShaderPermutationParameters& Parameters)
+	{
+		return true;
+	}
+};
+
+IMPLEMENT_GLOBAL_SHADER(FLocalHeightFogVS, "/Engine/Private/LocalHeightFog.usf", "LocalHeightFogSplatVS", SF_Vertex);
+
+class FLocalHeightFogPS : public FGlobalShader
+{
+	DECLARE_GLOBAL_SHADER(FLocalHeightFogPS);
+	SHADER_USE_PARAMETER_STRUCT(FLocalHeightFogPS, FGlobalShader);
+
+	BEGIN_SHADER_PARAMETER_STRUCT(FParameters, )
+		SHADER_PARAMETER_STRUCT_REF(FViewUniformShaderParameters, View)
+		SHADER_PARAMETER_RDG_BUFFER_SRV(StructuredBuffer<FLocalHeightFogGPUInstanceData>, LocalHeightFogInstances)
+	END_SHADER_PARAMETER_STRUCT()
+
+	using FPermutationDomain = TShaderPermutationDomain<>;
+
+	static bool ShouldCompilePermutation(const FGlobalShaderPermutationParameters& Parameters)
+	{
+		return true;
+	}
+};
+
+IMPLEMENT_GLOBAL_SHADER(FLocalHeightFogPS, "/Engine/Private/LocalHeightFog.usf", "LocalHeightFogSplatPS", SF_Pixel);
+
+BEGIN_SHADER_PARAMETER_STRUCT(FLocalHeightFogPassParameters, )
+	SHADER_PARAMETER_STRUCT_INCLUDE(FLocalHeightFogVS::FParameters, VS)
+	SHADER_PARAMETER_STRUCT_INCLUDE(FLocalHeightFogPS::FParameters, PS)
+	SHADER_PARAMETER_RDG_UNIFORM_BUFFER(FSceneTextureUniformParameters, SceneTextures)
+	RENDER_TARGET_BINDING_SLOTS()
+END_SHADER_PARAMETER_STRUCT()
+
+void RenderLocalHeightFog(
+	const FScene* Scene,
+	TArray<FViewInfo>& Views,
+	FRDGBuilder& GraphBuilder,
+	const FMinimalSceneTextures& SceneTextures,
+	FRDGTextureRef LightShaftOcclusionTexture)
+{
+	uint32 LocalHeightFogInstanceCount = Scene->LocalHeightFogs.Num();
+	if (LocalHeightFogInstanceCount <= 0)
+	{
+		return;
+	}
+
+	// No culling as of today
+	const uint32 AllLocalHeightFogInstanceBytes = sizeof(FLocalHeightFogGPUInstanceData) * LocalHeightFogInstanceCount;
+	FLocalHeightFogGPUInstanceData* LocalHeightFogGPUInstanceData = (FLocalHeightFogGPUInstanceData*) GraphBuilder.Alloc(AllLocalHeightFogInstanceBytes, 16);
+	FLocalHeightFogGPUInstanceData* LocalHeightFogGPUInstanceDataIt = LocalHeightFogGPUInstanceData;
+	for (FLocalHeightFogSceneProxy* LHF : Scene->LocalHeightFogs)
+	{
+		LocalHeightFogGPUInstanceDataIt->Transform = FMatrix44f(LHF->FogTransform.ToMatrixWithScale());
+		LocalHeightFogGPUInstanceDataIt->InvTransform = LocalHeightFogGPUInstanceDataIt->Transform.Inverse();
+
+		LocalHeightFogGPUInstanceDataIt->Density = LHF->FogDensity;
+		LocalHeightFogGPUInstanceDataIt->HeightFalloff = LHF->FogHeightFalloff * 0.01f;	// This scale is used to have artist author reasonable range.
+		LocalHeightFogGPUInstanceDataIt->HeightOffset = LHF->FogHeightOffset;
+		LocalHeightFogGPUInstanceDataIt->RadialAttenuationPower = LHF->FogRadialAttenuationPower;
+		LocalHeightFogGPUInstanceDataIt->RadialAttenuationDepth = LHF->FogRadialAttenuationDepth;
+
+		LocalHeightFogGPUInstanceDataIt->Albedo = FVector3f(LHF->FogAlbedo);
+		LocalHeightFogGPUInstanceDataIt->PhaseG = LHF->FogPhaseG;
+		LocalHeightFogGPUInstanceDataIt->Emissive = FVector3f(LHF->FogEmissive);
+
+		LocalHeightFogGPUInstanceDataIt++;
+	}
+	FRDGBufferRef LocalHeightFogGPUInstanceDataBuffer =  CreateStructuredBuffer(GraphBuilder, TEXT("LocalHeightFogGPUInstanceDataBuffer"), 
+		sizeof(FLocalHeightFogGPUInstanceData), LocalHeightFogInstanceCount, LocalHeightFogGPUInstanceData, AllLocalHeightFogInstanceBytes, ERDGInitialDataFlags::NoCopy);
+	FRDGBufferSRVRef LocalHeightFogGPUInstanceDataBufferSRV = GraphBuilder.CreateSRV(LocalHeightFogGPUInstanceDataBuffer);
+
+	FRDGTextureRef SceneColorTexture = SceneTextures.Color.Resolve;
+
+	for (FViewInfo& View : Views)
+	{
+		FLocalHeightFogPassParameters* PassParameters = GraphBuilder.AllocParameters<FLocalHeightFogPassParameters>();
+
+		PassParameters->VS.View = GetShaderBinding(View.ViewUniformBuffer);
+		PassParameters->VS.LocalHeightFogInstances = LocalHeightFogGPUInstanceDataBufferSRV;
+
+		PassParameters->PS.View = GetShaderBinding(View.ViewUniformBuffer);
+		PassParameters->PS.LocalHeightFogInstances = LocalHeightFogGPUInstanceDataBufferSRV;
+
+		PassParameters->SceneTextures = SceneTextures.UniformBuffer;
+		PassParameters->RenderTargets[0] = FRenderTargetBinding(SceneColorTexture, ERenderTargetLoadAction::ENoAction);
+
+		FLocalHeightFogVS::FPermutationDomain VSPermutationVector;
+		auto VertexShader = View.ShaderMap->GetShader< FLocalHeightFogVS >(VSPermutationVector);
+
+		FLocalHeightFogPS::FPermutationDomain PsPermutationVector;
+		auto PixelShader = View.ShaderMap->GetShader< FLocalHeightFogPS >(PsPermutationVector);
+
+		const FIntRect ViewRect = View.ViewRect;
+
+		ClearUnusedGraphResources(VertexShader, &PassParameters->VS);
+		ClearUnusedGraphResources(PixelShader, &PassParameters->PS);
+
+		GraphBuilder.AddPass(
+			RDG_EVENT_NAME("RenderLocalHeightFog %u inst.", LocalHeightFogInstanceCount),
+			PassParameters,
+			ERDGPassFlags::Raster,
+			[VertexShader, PixelShader, PassParameters, LocalHeightFogInstanceCount, ViewRect](FRHICommandList& RHICmdList)
+		{
+			FGraphicsPipelineStateInitializer GraphicsPSOInit;
+			RHICmdList.ApplyCachedRenderTargets(GraphicsPSOInit);
+
+			RHICmdList.SetViewport(ViewRect.Min.X, ViewRect.Min.Y, 0.0f, ViewRect.Max.X, ViewRect.Max.Y, 1.0f);
+
+			// Render back faces only since camera may intersect
+			GraphicsPSOInit.RasterizerState = TStaticRasterizerState<FM_Solid, CM_CW>::GetRHI();
+			GraphicsPSOInit.DepthStencilState = TStaticDepthStencilState<false, CF_Always>::GetRHI();
+			GraphicsPSOInit.BlendState = TStaticBlendState<CW_RGBA, BO_Add, BF_One, BF_SourceAlpha, BO_Add, BF_Zero, BF_SourceAlpha>::GetRHI();
+			GraphicsPSOInit.PrimitiveType = PT_TriangleList;
+
+			GraphicsPSOInit.BoundShaderState.VertexDeclarationRHI = GetVertexDeclarationFVector4();
+			GraphicsPSOInit.BoundShaderState.VertexShaderRHI = VertexShader.GetVertexShader();
+			GraphicsPSOInit.BoundShaderState.PixelShaderRHI = PixelShader.GetPixelShader();
+
+			SetGraphicsPipelineState(RHICmdList, GraphicsPSOInit, 0);
+
+			SetShaderParameters(RHICmdList, VertexShader, VertexShader.GetVertexShader(), PassParameters->VS);
+			SetShaderParameters(RHICmdList, PixelShader, PixelShader.GetPixelShader(), PassParameters->PS);
+
+			RHICmdList.SetStreamSource(0, GetUnitCubeVertexBuffer(), 0);
+
+			RHICmdList.DrawIndexedPrimitive( GetUnitCubeIndexBuffer()
+				, 0									//BaseVertexIndex
+				, 0									//FirstInstance
+				, 8									//uint32 NumVertices
+				, 0									//uint32 StartIndex
+				, UE_ARRAY_COUNT(GCubeIndices) / 3	//uint32 NumPrimitives
+				, LocalHeightFogInstanceCount		//uint32 NumInstances
+			);
+		});
+	}
+}
diff --git a/Engine/Source/Runtime/Renderer/Private/LocalHeightFogRendering.h b/Engine/Source/Runtime/Renderer/Private/LocalHeightFogRendering.h
new file mode 100644
--- /dev/null
+++ b/Engine/Source/Runtime/Renderer/Private/LocalHeightFogRendering.h
@@ -0,0 +1,23 @@
+// Copyright Epic Games, Inc. All Rights Reserved.
+
+#pragma once
+
+#include "EngineDefines.h"
+#include "RenderGraph.h"
+#include "RenderResource.h"
+#include "RenderGraphResources.h"
+#include "SceneView.h"
+
+class FScene;
+class FViewInfo;
+struct FMinimalSceneTextures;
+
+bool ShouldRenderLocalHeightFog(const FScene* Scene, const FSceneViewFamily& Family);
+
+void RenderLocalHeightFog(
+	const FScene* Scene,
+	TArray<FViewInfo>& Views,
+	FRDGBuilder& GraphBuilder,
+	const FMinimalSceneTextures& SceneTextures,
+	FRDGTextureRef LightShaftOcclusionTexture);
+
diff --git a/Engine/Source/Runtime/Renderer/Private/RendererScene.cpp b/Engine/Source/Runtime/Renderer/Private/RendererScene.cpp
--- a/Engine/Source/Runtime/Renderer/Private/RendererScene.cpp
+++ b/Engine/Source/Runtime/Renderer/Private/RendererScene.cpp
@@ -6637,6 +6637,10 @@ public:
 	virtual void RemoveExponentialHeightFog(class UExponentialHeightFogComponent* FogComponent) override {}
 	virtual bool HasAnyExponentialHeightFog() const override { return false; }
 
+	virtual void AddLocalHeightFog(class FLocalHeightFogSceneProxy* FogProxy) override {}
+	virtual void RemoveLocalHeightFog(class FLocalHeightFogSceneProxy* FogProxy) override {}
+	virtual bool HasAnyLocalHeightFog() const override { return false; }
+
 	virtual void AddSkyAtmosphere(FSkyAtmosphereSceneProxy* SkyAtmosphereSceneProxy, bool bStaticLightingBuilt) override {}
 	virtual void RemoveSkyAtmosphere(FSkyAtmosphereSceneProxy* SkyAtmosphereSceneProxy) override {}
 	virtual FSkyAtmosphereRenderSceneInfo* GetSkyAtmosphereSceneInfo() override { return NULL; }
diff --git a/Engine/Source/Runtime/Renderer/Private/ScenePrivate.h b/Engine/Source/Runtime/Renderer/Private/ScenePrivate.h
--- a/Engine/Source/Runtime/Renderer/Private/ScenePrivate.h
+++ b/Engine/Source/Runtime/Renderer/Private/ScenePrivate.h
@@ -3015,6 +3015,8 @@ public:
 	/** Used to track the order that skylights were enabled in. */
 	TArray<FVolumetricCloudSceneProxy*> VolumetricCloudStack;
 
+	TArray<FLocalHeightFogSceneProxy*> LocalHeightFogs;
+
 	TArray<FSparseVolumeTextureViewerSceneProxy*> SparseVolumeTextureViewers;
 
 	/** Global Field Manager */
@@ -3186,6 +3188,10 @@ public:
 	virtual void GetLightIESAtlasSlot(const FLightSceneProxy* Proxy, FLightRenderParameters* Out) override;
 	virtual void GetRectLightAtlasSlot(const FRectLightSceneProxy* Proxy, FLightRenderParameters* Out) override;
 
+	virtual void AddLocalHeightFog(class FLocalHeightFogSceneProxy* FogProxy) override;
+	virtual void RemoveLocalHeightFog(class FLocalHeightFogSceneProxy* FogProxy) override;
+	virtual bool HasAnyLocalHeightFog() const override;
+
 	virtual void AddSkyAtmosphere(FSkyAtmosphereSceneProxy* SkyAtmosphereSceneProxy, bool bStaticLightingBuilt) override;
 	virtual void RemoveSkyAtmosphere(FSkyAtmosphereSceneProxy* SkyAtmosphereSceneProxy) override;
 	virtual FSkyAtmosphereRenderSceneInfo* GetSkyAtmosphereSceneInfo() override { return SkyAtmosphere; }
