Added support for rect light volumetric fog ligth functions.

diff --git a/Engine/Shaders/Private/VolumetricFog.usf b/Engine/Shaders/Private/VolumetricFog.usf
--- a/Engine/Shaders/Private/VolumetricFog.usf
+++ b/Engine/Shaders/Private/VolumetricFog.usf
@@ -234,19 +234,11 @@ float4 LightFunctionAtlasTileMinMaxUvBound;
 Texture2D LightFunctionAtlasTexture;
 SamplerState LightFunctionAtlasSampler;
 
-float GetLocalLightFunction(float3 TranslatedWorldPosition, bool IsPointLight)
+float GetLocalLightFunction(float3 TranslatedWorldPosition, bool bSpotLight, bool bRectLight)
 {
 	float2 LightFunctionUV;
-	if (IsPointLight)
-	{
-		// Get vector from light to cell world pos.
-		float3 LightVec = normalize(TranslatedWorldPosition - DeferredLightUniforms.TranslatedWorldPosition);
-		// Transform and swizzle
-		LightVec = mul(float4(LightVec.xyz, 0), LocalLightFunctionMatrix).zyx; 
-		// Map to spherical coordinate
-		LightFunctionUV = float2((atan2Fast(LightVec.y, LightVec.x) + PI) / (2 * PI), acosFast(LightVec.z) / PI);
-	}
-	else
+
+	if(bSpotLight)
 	{
 		float4 HomogeneousShadowPosition = mul(float4(TranslatedWorldPosition, 1), LocalLightFunctionMatrix);
 		HomogeneousShadowPosition.xyz /= HomogeneousShadowPosition.w;
@@ -254,6 +246,24 @@ float GetLocalLightFunction(float3 TranslatedWorldPosition, bool IsPointLight)
 		LightFunctionUV.x = 1 - LightFunctionUV.x;
 		LightFunctionUV.y = 1 - LightFunctionUV.y;
 	}
+	else if (bRectLight)
+	{
+		float4 HomogeneousShadowPosition = mul(float4(TranslatedWorldPosition, 1), LocalLightFunctionMatrix).zyxw;
+		HomogeneousShadowPosition.xyz /= HomogeneousShadowPosition.w;
+		LightFunctionUV = HomogeneousShadowPosition.xy;
+
+		LightFunctionUV = LightFunctionUV * 0.5 + 0.5;
+		LightFunctionUV = frac(LightFunctionUV);
+	}
+	else // Point light
+	{
+		// Get vector from light to cell world pos.
+		float3 LightVec = normalize(TranslatedWorldPosition - DeferredLightUniforms.TranslatedWorldPosition);
+		// Transform and swizzle
+		LightVec = mul(float4(LightVec.xyz, 0), LocalLightFunctionMatrix).zyx;
+		// Map to spherical coordinate
+		LightFunctionUV = float2((atan2Fast(LightVec.y, LightVec.x) + PI) / (2 * PI), acosFast(LightVec.z) / PI);
+	}
 	
 	return Texture2DSampleLevel(LightFunctionAtlasTexture, LightFunctionAtlasSampler, LightFunctionAtlasTileMinMaxUvBound.xy + LightFunctionUV * (LightFunctionAtlasTileMinMaxUvBound.zw - LightFunctionAtlasTileMinMaxUvBound.xy), 0).x;
 }
@@ -362,7 +372,7 @@ void InjectShadowedLocalLightPS(
 #endif
 	
 #if	USE_LIGHT_FUNCTION
-			ShadowFactor *= GetLocalLightFunction(TranslatedWorldPosition, IsPointLight);
+			ShadowFactor *= GetLocalLightFunction(TranslatedWorldPosition, LightData.bSpotLight, LightData.bRectLight);
 #endif
 
 			OutScattering.rgb += LightColor * (PhaseFunction(PhaseG, dot(L, -CameraVector)) * CombinedAttenuation * ShadowFactor * VolumetricScatteringIntensity);
@@ -479,7 +489,7 @@ float4x4 DirectionalLightFunctionTranslatedWorldToShadow;
 Texture2D LightFunctionTexture;
 SamplerState LightFunctionSampler;
 
-float GetLightFunction(float3 TranslatedWorldPosition)
+float GetDirectionalLightFunction(float3 TranslatedWorldPosition) // directional light
 {
 	float4 HomogeneousShadowPosition = mul(float4(TranslatedWorldPosition, 1), DirectionalLightFunctionTranslatedWorldToShadow);
 	float2 LightFunctionUV = HomogeneousShadowPosition.xy * .5f + .5f;
@@ -584,7 +594,7 @@ void LightScatteringCS(
 #endif
 			}
 
-			ShadowFactor *= GetLightFunction(TranslatedWorldPosition);
+			ShadowFactor *= GetDirectionalLightFunction(TranslatedWorldPosition);
 
 			float3 DirectionalLightColor = ForwardLightData.DirectionalLightColor;
 
diff --git a/Engine/Shaders/Private/VolumetricFogLightFunction.usf b/Engine/Shaders/Private/VolumetricFogLightFunction.usf
--- a/Engine/Shaders/Private/VolumetricFogLightFunction.usf
+++ b/Engine/Shaders/Private/VolumetricFogLightFunction.usf
@@ -22,8 +22,11 @@ void Main(
 	float4 TranslatedWorldPosition = float4(0,0,0,0);
 	float3 LightVector = float3(0,0,0);
 
-#if LIGHT_TYPE == 0 
-	// Point and rect light
+	float Depth = 0.000001f;
+	float2 ScreenPosition = 0.0f;
+
+#if LIGHT_TYPE == 0			// Point light
+	
 	float Theta = (UV.x * (2*PI)) - PI;
 	float Phi = UV.y * PI;
 	float X = cos(Theta) * sin(Phi);
@@ -32,13 +35,8 @@ void Main(
 	LightVector = float3(Z, Y, X);
 	TranslatedWorldPosition = float4(LightTranslatedWorldPosition + LightVector.zyx, 1);
 
-#else
-
-	float Depth = 0.000001f;
-	float2 ScreenPosition;
-
-#if LIGHT_TYPE == 1
-	// Spot light
+#elif LIGHT_TYPE == 1		// Spot light
+	
 	ScreenPosition = (SvPosition.xy * LightFunctionTexelSize);
 	ScreenPosition.x = 1 - ScreenPosition.x;
 	ScreenPosition.y = 1 - ScreenPosition.y;
@@ -46,17 +44,31 @@ void Main(
 	// Wrap ScreenPosition to within the atlas tile
 	ScreenPosition = frac(ScreenPosition);
 
-#else // LIGHT_TYPE == 2
-	// Directional light
-	ScreenPosition = (SvPosition.xy * LightFunctionTexelSize - float2(0.5f, 0.5f)) * float2(2, -2);
-#endif
+	// Compute translated world position (inverse shadow projection) from the screen position
+	// make SvPosition appear to be rasterized with the depth from the depth buffer
+	TranslatedWorldPosition = mul(float4(ScreenPosition, Depth, 1), ShadowToTranslatedWorld);
+	TranslatedWorldPosition.xyz /= TranslatedWorldPosition.w;
+
+	// Compute the light vector containing UV from the translated world position
+	float4 Hom = mul(float4(TranslatedWorldPosition.xyz, 1), LightFunctionTranslatedWorldToLight);
+	LightVector = Hom.xyz / Hom.w;
 
+#else //LIGHT_TYPE == 2	// Directional and Rect light
+
+	ScreenPosition = (SvPosition.xy * LightFunctionTexelSize - float2(0.5f, 0.5f)) * float2(2, 2);
+	ScreenPosition.xy = frac(ScreenPosition.xy);
+
+	// Compute translated world position (inverse shadow projection) from the screen position
 	// make SvPosition appear to be rasterized with the depth from the depth buffer
 	TranslatedWorldPosition = mul(float4(ScreenPosition, Depth, 1), ShadowToTranslatedWorld);
 	TranslatedWorldPosition.xyz /= TranslatedWorldPosition.w;
 
+	// Compute the light vector containing UV from the translated world position
 	float4 Hom = mul(float4(TranslatedWorldPosition.xyz, 1), LightFunctionTranslatedWorldToLight);
 	LightVector = Hom.xyz / Hom.w;
+
+	LightVector.xyz = LightVector.zyx;	// Swap to match UV convention
+
 #endif
 	
 	float3 LightFunction = GetLightFunctionColor(LightVector, TranslatedWorldPosition.xyz);
diff --git a/Engine/Source/Runtime/Renderer/Private/VolumetricFog.cpp b/Engine/Source/Runtime/Renderer/Private/VolumetricFog.cpp
--- a/Engine/Source/Runtime/Renderer/Private/VolumetricFog.cpp
+++ b/Engine/Source/Runtime/Renderer/Private/VolumetricFog.cpp
@@ -354,8 +354,7 @@ bool LightNeedsSeparateInjectionIntoVolumetricFogForOpaqueShadow(const FViewInfo
 bool LightNeedsSeparateInjectionIntoVolumetricFogForLightFunction(const FLightSceneInfo* LightSceneInfo)
 {
 	// No directional light type because it is handled in a specific way in RenderLightFunctionForVolumetricFog.
-	// TODO: add support for rect lights.
-	return GVolumetricFogLightFunction > 0 && (LightSceneInfo->Proxy->GetLightType() == LightType_Point || LightSceneInfo->Proxy->GetLightType() == LightType_Spot);
+	return GVolumetricFogLightFunction > 0 && (LightSceneInfo->Proxy->GetLightType() == LightType_Point || LightSceneInfo->Proxy->GetLightType() == LightType_Spot || LightSceneInfo->Proxy->GetLightType() == LightType_Rect);
 }
 
 FIntPoint CalculateVolumetricFogBoundsForLight(const FSphere& LightBounds, const FViewInfo& View, FIntVector VolumetricFogGridSize, FVector GridZParams)
diff --git a/Engine/Source/Runtime/Renderer/Private/VolumetricFogLightFunction.cpp b/Engine/Source/Runtime/Renderer/Private/VolumetricFogLightFunction.cpp
--- a/Engine/Source/Runtime/Renderer/Private/VolumetricFogLightFunction.cpp
+++ b/Engine/Source/Runtime/Renderer/Private/VolumetricFogLightFunction.cpp
@@ -186,8 +186,10 @@ void FSceneRenderer::RenderLightFunctionForVolumetricFog(
 					continue;
 				}
 
-				if (LightSceneInfo->Proxy->GetLightType() == LightType_Spot 
-					|| LightSceneInfo->Proxy->GetLightType() == LightType_Point)
+				uint8 LightType = LightSceneInfo->Proxy->GetLightType();
+				if (LightType == LightType_Spot
+					|| LightType == LightType_Point
+					|| LightType == LightType_Rect)
 				{
 					LocalLightsToEvaluate.Add(LightSceneInfo);
 				}
@@ -349,17 +351,18 @@ void FSceneRenderer::RenderLightFunctionForVolumetricFog(
 	// We are going to render local light LightFunction so we allocate the atlas.
 	View.VolumetricFogResources.TransientLightFunctionTextureAtlas = new (GraphBuilder.Alloc(sizeof(FTransientLightFunctionTextureAtlas), alignof(FTransientLightFunctionTextureAtlas))) FTransientLightFunctionTextureAtlas(GraphBuilder);
 
-	// Conpute all the data required by the lights
+	// Compute all the data required by the lights
+	const FIntPoint LightFunctionResolution = GetVolumetricFogLightFunctionResolution();
 	const FMatrix TranslatedWorldToWorld = FTranslationMatrix(-View.ViewMatrices.GetPreViewTranslation());
 	for (FLightSceneInfo* LightSceneInfo : LocalLightsToEvaluate)
 	{
-		const bool bIsPointLight = LightSceneInfo->Proxy->GetLightType() == LightType_Point;
-		const FVector LightDirection = LightSceneInfo->Proxy->GetDirection().GetSafeNormal();
-		const FIntPoint LightFunctionResolution = GetVolumetricFogLightFunctionResolution();
-		FMatrix WorldToShadowMatrix = FMatrix::Identity;
+		uint8 LightType = LightSceneInfo->Proxy->GetLightType();
+
+		FVolumetricFogLocalLightFunctionInfo& LightFunctionData = LocalLightFunctionData.Add(LightSceneInfo);
+		LightFunctionData.AtlasTile = View.VolumetricFogResources.TransientLightFunctionTextureAtlas->AllocateAtlasTile();
 
 		// Spotlights needs to utilize their shadow projection when rendering the light function 
-		if (!bIsPointLight)
+		if (LightType == LightType_Spot)
 		{
 			TArray<FWholeSceneProjectedShadowInitializer, TInlineAllocator<6>> Initializers;
 			LightSceneInfo->Proxy->GetWholeSceneProjectedShadowInitializer(*View.Family, Initializers);
@@ -378,15 +381,21 @@ void FSceneRenderer::RenderLightFunctionForVolumetricFog(
 			);
 
 			FVector4f ShadowmapMinMaxValue;
-			WorldToShadowMatrix = ProjectedShadowInfo.GetWorldToShadowMatrix(ShadowmapMinMaxValue, &LightFunctionResolution);
+			LightFunctionData.LightFunctionTranslatedWorldToLightMatrix = FMatrix44f(TranslatedWorldToWorld * ProjectedShadowInfo.GetWorldToShadowMatrix(ShadowmapMinMaxValue, &LightFunctionResolution));
 		}
+		else if (LightType == LightType_Point)
+		{
+			LightFunctionData.LightFunctionTranslatedWorldToLightMatrix = FMatrix44f(TranslatedWorldToWorld * LightSceneInfo->Proxy->GetWorldToLight());
+		}
+		else if (LightType == LightType_Rect)
+		{
+			const FVector Scale = LightSceneInfo->Proxy->GetLightFunctionScale();
+			// Switch x and z so that z of the user specified scale affects the distance along the light direction
+			const FVector InverseScale = FVector(1.0 / Scale.Z, 1.0 / Scale.Y, 1.0 / Scale.X);
+			const FMatrix WorldToLight = LightSceneInfo->Proxy->GetWorldToLight() * FScaleMatrix(InverseScale);
 
-		FVolumetricFogLocalLightFunctionInfo& LightFunctionData = LocalLightFunctionData.Add(LightSceneInfo);
-		LightFunctionData.AtlasTile = View.VolumetricFogResources.TransientLightFunctionTextureAtlas->AllocateAtlasTile();
-
-		// Pass the World To Light when light type is Point because we need it instead of the World To Shadow.
-		const FMatrix WorldToLight = bIsPointLight ? LightSceneInfo->Proxy->GetWorldToLight() : WorldToShadowMatrix;
-		LightFunctionData.LightFunctionTranslatedWorldToLightMatrix = FMatrix44f(TranslatedWorldToWorld * WorldToLight);
+			LightFunctionData.LightFunctionTranslatedWorldToLightMatrix = FMatrix44f(TranslatedWorldToWorld * WorldToLight);
+		}
 	}
 
 	// And now process all the local lights, allocated atlas tile and generate the tile as needed
@@ -427,12 +436,14 @@ void FSceneRenderer::RenderLightFunctionForVolumetricFog(
 					switch (LightSceneInfo->Proxy->GetLightType())
 					{
 					case LightType_Point:
-					case LightType_Rect:
 						PermutationVector.Set<FVolumetricFogLightFunctionPS::FLightType>(0);
 						break;
 					case LightType_Spot:
 						PermutationVector.Set<FVolumetricFogLightFunctionPS::FLightType>(1);
 						break;
+					case LightType_Rect:
+						PermutationVector.Set<FVolumetricFogLightFunctionPS::FLightType>(2);
+						break;
 					default:
 						check(false)
 					}
@@ -444,7 +455,10 @@ void FSceneRenderer::RenderLightFunctionForVolumetricFog(
 
 					SetGraphicsPipelineState(RHICmdList, GraphicsPSOInit, 0);
 
-					FVolumetricFogLightFunctionPS::FParameters PS = PixelShader->GetParameters(View, LightSceneInfo, FVector2D(1.0f / LightFunctionResolution, 1.0f / LightFunctionResolution), LightFunctionInfo.LightFunctionTranslatedWorldToLightMatrix.Inverse());
+					FVolumetricFogLightFunctionPS::FParameters PS = 
+						PixelShader->GetParameters(View, LightSceneInfo, 
+							FVector2D(1.0f / LightFunctionResolution, 1.0f / LightFunctionResolution),
+							LightFunctionInfo.LightFunctionTranslatedWorldToLightMatrix.Inverse());
 
 					ClearUnusedGraphResources(PixelShader, &PS);
 					SetShaderParametersMixedPS(RHICmdList, PixelShader, PS, View, MaterialProxyForRendering);
@@ -479,7 +493,7 @@ FTransientLightFunctionTextureAtlas::FTransientLightFunctionTextureAtlas(FRDGBui
 	HalfTexelSize = 0.5f / float(AtlasTextureWidth);
 
 	FRDGTextureDesc LightFunctionTextureDesc = FRDGTextureDesc::Create2D(FIntPoint(AtlasTextureWidth, AtlasTextureWidth), PF_G8, FClearValueBinding::None, TexCreate_ShaderResource | TexCreate_RenderTargetable);
-	TransientLightFunctionAtlasTexture = GraphBuilder.CreateTexture(LightFunctionTextureDesc, TEXT("TransientLightFunctionAtlasTexture"));
+	TransientLightFunctionAtlasTexture = GraphBuilder.CreateTexture(LightFunctionTextureDesc, TEXT("VolumetricFog.LightFunctionAtlasTexture"));
 	DefaultLightFunctionAtlasItemTexture = GSystemTextures.GetWhiteDummy(GraphBuilder);
 }
 
