Fix RHI validation warnings in RaytracingScene.

- Test: -rhivalidation -onethread -rdgimmediate

diff --git a/Engine/Source/Runtime/Renderer/Private/DeferredShadingRenderer.cpp b/Engine/Source/Runtime/Renderer/Private/DeferredShadingRenderer.cpp
--- a/Engine/Source/Runtime/Renderer/Private/DeferredShadingRenderer.cpp
+++ b/Engine/Source/Runtime/Renderer/Private/DeferredShadingRenderer.cpp
@@ -2093,7 +2093,7 @@ bool FDeferredShadingSceneRenderer::DispatchRayTracingWorldUpdates(FRDGBuilder&
 	// Keep mask the same as what's already set (which will be the view mask) if TLAS updates should be masked to the view
 	RDG_GPU_MASK_SCOPE(GraphBuilder, GRayTracingMultiGpuTLASMask ? GraphBuilder.RHICmdList.GetGPUMask() : FRHIGPUMask::All());
 
-	RayTracingScene.CreateWithInitializationData(GraphBuilder, &Scene->GPUScene, ReferenceView.ViewMatrices, MoveTemp(ReferenceView.RayTracingSceneInitData));
+	RayTracingScene.CreateWithInitializationData(GraphBuilder, ReferenceView, &Scene->GPUScene, MoveTemp(ReferenceView.RayTracingSceneInitData));
 
 	const uint32 BLASScratchSize = Scene->GetRayTracingDynamicGeometryCollection()->ComputeScratchBufferSize();
 	if (BLASScratchSize > 0)
diff --git a/Engine/Source/Runtime/Renderer/Private/HeterogeneousVolumes/HeterogeneousVolumesHardwareRayTracing.cpp b/Engine/Source/Runtime/Renderer/Private/HeterogeneousVolumes/HeterogeneousVolumesHardwareRayTracing.cpp
--- a/Engine/Source/Runtime/Renderer/Private/HeterogeneousVolumes/HeterogeneousVolumesHardwareRayTracing.cpp
+++ b/Engine/Source/Runtime/Renderer/Private/HeterogeneousVolumes/HeterogeneousVolumesHardwareRayTracing.cpp
@@ -237,8 +237,7 @@ void GenerateRayTracingScene(
 
 	// Create RayTracingScene
 	const FGPUScene* EmptyGPUScene = nullptr;
-	FViewMatrices EmptyViewMatrices;
-	RayTracingScene.Create(GraphBuilder, EmptyGPUScene, EmptyViewMatrices);
+	RayTracingScene.Create(GraphBuilder, View, EmptyGPUScene);
 
 	// Build TLAS
 	FBuildTLASPassParams* PassParamsTLAS = GraphBuilder.AllocParameters<FBuildTLASPassParams>();
diff --git a/Engine/Source/Runtime/Renderer/Private/RayTracing/RayTracingScene.cpp b/Engine/Source/Runtime/Renderer/Private/RayTracing/RayTracingScene.cpp
--- a/Engine/Source/Runtime/Renderer/Private/RayTracing/RayTracingScene.cpp
+++ b/Engine/Source/Runtime/Renderer/Private/RayTracing/RayTracingScene.cpp
@@ -11,10 +11,13 @@
 #include "RenderGraphUtils.h"
 #include "RaytracingOptions.h"
 #include "PrimitiveSceneProxy.h"
+#include "SceneUniformBuffer.h"
+#include "SceneRendering.h"
 
 BEGIN_SHADER_PARAMETER_STRUCT(FBuildInstanceBufferPassParams, )
 	SHADER_PARAMETER_RDG_BUFFER_UAV(RWStructuredBuffer, InstanceBuffer)
 	SHADER_PARAMETER_RDG_BUFFER_UAV(RWStructuredBuffer, DebugInstanceGPUSceneIndexBuffer)
+	SHADER_PARAMETER_RDG_UNIFORM_BUFFER(FSceneUniformParameters, Scene)
 END_SHADER_PARAMETER_STRUCT()
 
 FRayTracingScene::FRayTracingScene()
@@ -38,9 +41,9 @@ FRayTracingSceneWithGeometryInstances FRayTracingScene::BuildInitializationData(
 		NumCallableShaderSlots);
 }
 
-void FRayTracingScene::Create(FRDGBuilder& GraphBuilder, const FGPUScene* GPUScene, const FViewMatrices& ViewMatrices)
+void FRayTracingScene::Create(FRDGBuilder& GraphBuilder, const FViewInfo& View, const FGPUScene* GPUScene)
 {
-	CreateWithInitializationData(GraphBuilder, GPUScene, ViewMatrices, BuildInitializationData());
+	CreateWithInitializationData(GraphBuilder, View, GPUScene, BuildInitializationData());
 }
 
 void FRayTracingScene::InitPreViewTranslation(const FViewMatrices& ViewMatrices)
@@ -52,7 +55,7 @@ void FRayTracingScene::InitPreViewTranslation(const FViewMatrices& ViewMatrices)
 	ViewTilePosition = AbsoluteViewOrigin.GetTile();
 }
 
-void FRayTracingScene::CreateWithInitializationData(FRDGBuilder& GraphBuilder, const FGPUScene* GPUScene, const FViewMatrices& ViewMatrices, FRayTracingSceneWithGeometryInstances SceneWithGeometryInstances)
+void FRayTracingScene::CreateWithInitializationData(FRDGBuilder& GraphBuilder, const FViewInfo& View, const FGPUScene* GPUScene, FRayTracingSceneWithGeometryInstances SceneWithGeometryInstances)
 {
 	QUICK_SCOPE_CYCLE_COUNTER(FRayTracingScene_BeginCreate);
 
@@ -177,7 +180,7 @@ void FRayTracingScene::CreateWithInitializationData(FRDGBuilder& GraphBuilder, c
 			InstanceGeometryIndices = MoveTemp(SceneWithGeometryInstances.InstanceGeometryIndices),
 			BaseUploadBufferOffsets = MoveTemp(SceneWithGeometryInstances.BaseUploadBufferOffsets),
 			RayTracingSceneRHI = RayTracingSceneRHI,
-			PreViewTranslation = ViewMatrices.GetPreViewTranslation()]()
+			PreViewTranslation = View.ViewMatrices.GetPreViewTranslation()]()
 		{
 			FTaskTagScope TaskTagScope(ETaskTag::EParallelRenderingThread);
 			FillRayTracingInstanceUploadBuffer(
@@ -201,6 +204,7 @@ void FRayTracingScene::CreateWithInitializationData(FRDGBuilder& GraphBuilder, c
 		FBuildInstanceBufferPassParams* PassParams = GraphBuilder.AllocParameters<FBuildInstanceBufferPassParams>();
 		PassParams->InstanceBuffer = GraphBuilder.CreateUAV(InstanceBuffer);
 		PassParams->DebugInstanceGPUSceneIndexBuffer = nullptr;
+		PassParams->Scene = View.GetSceneUniforms().GetBuffer(GraphBuilder);
 
 		if (bNeedsDebugInstanceGPUSceneIndexBuffer)
 		{
diff --git a/Engine/Source/Runtime/Renderer/Private/RayTracing/RayTracingScene.h b/Engine/Source/Runtime/Renderer/Private/RayTracing/RayTracingScene.h
--- a/Engine/Source/Runtime/Renderer/Private/RayTracing/RayTracingScene.h
+++ b/Engine/Source/Runtime/Renderer/Private/RayTracing/RayTracingScene.h
@@ -56,10 +56,10 @@ public:
 	// Allocates GPU memory to fit at least the current number of instances.
 	// Kicks off instance buffer build to parallel thread along with RDG pass.
 	// NOTE: SceneWithGeometryInstances is passed in by value because ownership of the internal data is taken over. Use MoveTemp at call site, if possible.
-	void CreateWithInitializationData(FRDGBuilder& GraphBuilder, const FGPUScene* GPUScene, const FViewMatrices& ViewMatrices, FRayTracingSceneWithGeometryInstances SceneWithGeometryInstances);
+	void CreateWithInitializationData(FRDGBuilder& GraphBuilder, const FViewInfo& View, const FGPUScene* GPUScene, FRayTracingSceneWithGeometryInstances SceneWithGeometryInstances);
 
 	// Backwards-compatible version of Create() which internally calls CreateRayTracingSceneWithGeometryInstances().
-	void Create(FRDGBuilder& GraphBuilder, const FGPUScene* GPUScene, const FViewMatrices& ViewMatrices);
+	void Create(FRDGBuilder& GraphBuilder, const FViewInfo& View, const FGPUScene* GPUScene);
 
 	// Resets the instance list and reserves memory for this frame.
 	void Reset(bool bInstanceDebugDataEnabled);
