Substrate - Glint - Alternative to the glint texture for more variation and less visual artefact.

diff --git a/Engine/Config/BaseEngine.ini b/Engine/Config/BaseEngine.ini
--- a/Engine/Config/BaseEngine.ini
+++ b/Engine/Config/BaseEngine.ini
@@ -162,6 +162,7 @@ PreIntegratedSkinBRDFTextureName=/Engine/EngineMaterials/PreintegratedSkinBRDF.P
 BlueNoiseScalarTextureName=/Engine/EngineMaterials/STBlueNoise_scalar_128x128x64.STBlueNoise_scalar_128x128x64
 BlueNoiseVec2TextureName=/Engine/EngineMaterials/STBlueNoise_vec2_128x128x64.STBlueNoise_vec2_128x128x64
 GlintTextureName=/Engine/EngineMaterials/Substrate/GLints/AGlintTex2dArray.AGlintTex2dArray
+GlintTexture2Name=/Engine/EngineMaterials/Substrate/GLints2/AGlintTex2dArray2.AGlintTex2dArray2
 MiniFontTextureName=/Engine/EngineMaterials/MiniFont.MiniFont
 WeightMapPlaceholderTextureName=/Engine/EngineMaterials/WeightMapPlaceholderTexture.WeightMapPlaceholderTexture
 LightMapDensityTextureName=/Engine/EngineMaterials/DefaultWhiteGrid.DefaultWhiteGrid
diff --git a/Engine/Source/Runtime/Engine/Classes/Engine/Engine.h b/Engine/Source/Runtime/Engine/Classes/Engine/Engine.h
--- a/Engine/Source/Runtime/Engine/Classes/Engine/Engine.h
+++ b/Engine/Source/Runtime/Engine/Classes/Engine/Engine.h
@@ -1373,11 +1373,19 @@ public:
 	/** Stable glint BSDF texture */
 	UPROPERTY()
 	TObjectPtr<class UTexture2DArray> GlintTexture;
+	
+	/** Stable glint BSDF texture with more variety to cover slope space and avoid circular artifact */
+	UPROPERTY()
+	TObjectPtr<class UTexture2DArray> GlintTexture2;
 
 	/** Path of the glint BSDF texture */
 	UPROPERTY(globalconfig)
 	FSoftObjectPath GlintTextureName;
 
+	/** Path of the glint BSDF texture 2 */
+	UPROPERTY(globalconfig)
+	FSoftObjectPath GlintTexture2Name;
+
 	/** Texture used to do font rendering in shaders */
 	UPROPERTY()
 	TObjectPtr<class UTexture2D> MiniFontTexture;
@@ -1921,7 +1929,7 @@ public:
 	/** Delay loading the glint texture until it is needed by the renderer.
 	* This texture is not going to be streamed to be available right away.
 	*/
-	void LoadGlintTexture();
+	void LoadGlintTextures();
 
 private:
 	#if WITH_DYNAMIC_RESOLUTION
diff --git a/Engine/Source/Runtime/Engine/Private/UnrealEngine.cpp b/Engine/Source/Runtime/Engine/Private/UnrealEngine.cpp
--- a/Engine/Source/Runtime/Engine/Private/UnrealEngine.cpp
+++ b/Engine/Source/Runtime/Engine/Private/UnrealEngine.cpp
@@ -3316,12 +3316,16 @@ void UEngine::ConditionallyLoadPreIntegratedSkinBRDFTexture()
 	}
 }
 
-void UEngine::LoadGlintTexture()
+void UEngine::LoadGlintTextures()
 {
 	if (GlintTexture == nullptr && GlintTextureName.IsValid())
 	{
 		LoadEngineTexture(GlintTexture, *GlintTextureName.ToString());
 	}
+	if (GlintTexture2 == nullptr && GlintTexture2Name.IsValid())
+	{
+		LoadEngineTexture(GlintTexture2, *GlintTexture2Name.ToString());
+	}
 }
 
 void UEngine::InitializePortalServices()
diff --git a/Engine/Source/Runtime/RenderCore/Private/RenderUtils.cpp b/Engine/Source/Runtime/RenderCore/Private/RenderUtils.cpp
--- a/Engine/Source/Runtime/RenderCore/Private/RenderUtils.cpp
+++ b/Engine/Source/Runtime/RenderCore/Private/RenderUtils.cpp
@@ -1598,6 +1598,12 @@ static TAutoConsoleVariable<int32> CVarSubstrateGlints(
 	TEXT("Enable Glint support for Strata slabs."),
 	ECVF_ReadOnly | ECVF_RenderThreadSafe);
 
+static TAutoConsoleVariable<int32> CVarSubstrateGlintsLUT(
+	TEXT("r.Substrate.GlintsLUT"),
+	1,
+	TEXT("Select one of the glint rendering LUT for testing."),
+	ECVF_RenderThreadSafe);
+
 static TAutoConsoleVariable<int32> CVarSubstrateDebugAdvancedVisualizationShaders(
 	TEXT("r.Substrate.Debug.AdvancedVisualizationShaders"),
 	0,
@@ -1702,6 +1708,11 @@ namespace Strata
 		return IsStrataEnabled() && CVarSubstrateGlints.GetValueOnAnyThread() > 0;
 	}
 
+	uint32 GlintLUTIndex()
+	{
+		return CVarSubstrateGlintsLUT.GetValueOnAnyThread()<= 0 ? 0u : 1u;
+	}
+
 	bool Is8bitTileCoordEnabled()
 	{
 		return CVarSubstrateTileCoord8Bits.GetValueOnAnyThread() > 0 ? 1 : 0;
diff --git a/Engine/Source/Runtime/RenderCore/Public/RenderUtils.h b/Engine/Source/Runtime/RenderCore/Public/RenderUtils.h
--- a/Engine/Source/Runtime/RenderCore/Public/RenderUtils.h
+++ b/Engine/Source/Runtime/RenderCore/Public/RenderUtils.h
@@ -468,6 +468,7 @@ namespace Strata
 	RENDERCORE_API bool IsStrataEnabled();
 	RENDERCORE_API bool IsRoughDiffuseEnabled();
 	RENDERCORE_API bool IsGlintEnabled();
+	RENDERCORE_API uint32 GlintLUTIndex();
 	RENDERCORE_API bool IsBackCompatibilityEnabled();
 	RENDERCORE_API bool IsDBufferPassEnabled(EShaderPlatform InPlatform);
 	RENDERCORE_API bool IsOpaqueRoughRefractionEnabled();
diff --git a/Engine/Source/Runtime/Renderer/Private/SceneRendering.cpp b/Engine/Source/Runtime/Renderer/Private/SceneRendering.cpp
--- a/Engine/Source/Runtime/Renderer/Private/SceneRendering.cpp
+++ b/Engine/Source/Runtime/Renderer/Private/SceneRendering.cpp
@@ -2517,7 +2517,7 @@ FSceneRenderer::FSceneRenderer(const FSceneViewFamily* InViewFamily, FHitProxyCo
 
 		if (Strata::IsGlintEnabled())
 		{
-			GEngine->LoadGlintTexture();
+			GEngine->LoadGlintTextures();
 		}
 
 		// Handle the FFT bloom kernel textire
diff --git a/Engine/Source/Runtime/Renderer/Private/Strata/Glint/GlintShadingLUTs.cpp b/Engine/Source/Runtime/Renderer/Private/Strata/Glint/GlintShadingLUTs.cpp
--- a/Engine/Source/Runtime/Renderer/Private/Strata/Glint/GlintShadingLUTs.cpp
+++ b/Engine/Source/Runtime/Renderer/Private/Strata/Glint/GlintShadingLUTs.cpp
@@ -118,9 +118,16 @@ void FGlintShadingLUTsStateData::Init(FRDGBuilder& GraphBuilder, FViewInfo& View
 		GlintShadingLUTsStateData.RHIGlintShadingLUTs = nullptr;
 		if (GEngine->GlintTexture)
 		{
-			GlintShadingLUTsStateData.RHIGlintShadingLUTs = GEngine->GlintTexture->GetResource()->TextureRHI->GetTexture2DArray();
-
-			GlintShadingLUTsStateData.SetDictionaryParameter(16, 64, 0.5);	// Match the embedded LUT
+			if (Strata::GlintLUTIndex() == 0)
+			{
+				GlintShadingLUTsStateData.RHIGlintShadingLUTs = GEngine->GlintTexture->GetResource()->TextureRHI->GetTexture2DArray();
+				GlintShadingLUTsStateData.SetDictionaryParameter(16, 64, 0.5);	// Match the embedded LUT
+			}
+			else
+			{
+				GlintShadingLUTsStateData.RHIGlintShadingLUTs = GEngine->GlintTexture2->GetResource()->TextureRHI->GetTexture2DArray();
+				GlintShadingLUTsStateData.SetDictionaryParameter(8, 256, 0.5);	// Match the embedded LUT
+			}
 		}
 	}
 
