Fix visibility to work with -nothreading.

diff --git a/Engine/Source/Runtime/Renderer/Private/SceneVisibility.cpp b/Engine/Source/Runtime/Renderer/Private/SceneVisibility.cpp
--- a/Engine/Source/Runtime/Renderer/Private/SceneVisibility.cpp
+++ b/Engine/Source/Runtime/Renderer/Private/SceneVisibility.cpp
@@ -3080,7 +3080,7 @@ FVisibilityTaskConfig::FVisibilityTaskConfig(const FScene& Scene, TConstArrayVie
 			}
 			else if (FSceneViewState* ViewState = static_cast<FSceneViewState*>(Views[ViewIndex].State))
 			{
-				MaxQueriesPerTask = FMath::DivideAndRoundUp(ViewState->Occlusion.NumRequestedQueries, NumWorkerThreads * NumOcclusionCullTasksPerThread);
+				MaxQueriesPerTask = FMath::DivideAndRoundUp(ViewState->Occlusion.NumRequestedQueries, FMath::Max(1u, NumWorkerThreads) * NumOcclusionCullTasksPerThread);
 			}
 
 			MaxQueriesPerTask = FMath::Max(MaxQueriesPerTask, OcclusionCull.MinQueriesPerTask);
