Change VSM cache manager over to use SceneUB.
- also clean up dangling RDG resource references in GPUScene.

diff --git a/Engine/Shaders/Private/VirtualShadowMaps/VirtualShadowMapCacheManagement.usf b/Engine/Shaders/Private/VirtualShadowMaps/VirtualShadowMapCacheManagement.usf
--- a/Engine/Shaders/Private/VirtualShadowMaps/VirtualShadowMapCacheManagement.usf
+++ b/Engine/Shaders/Private/VirtualShadowMaps/VirtualShadowMapCacheManagement.usf
@@ -19,10 +19,6 @@
 
 RWStructuredBuffer<FPhysicalPageMetaData> PhysicalPageMetaDataOut;
 
-uint InstanceSceneDataSOAStride;
-uint GPUSceneNumAllocatedInstances;
-uint GPUSceneNumAllocatedPrimitives;
-
 #if ENABLE_DEBUG_MODE
 uint bDrawBounds;
 #endif
@@ -86,10 +82,10 @@ void VirtualSmInvalidateInstancePagesCS(
 	uint InstanceId = WorkSetup.Item.InstanceDataOffset + uint(WorkSetup.LocalItemIndex);
 
 	{
-		checkSlow(InstanceId >= 0 && InstanceId < InstanceSceneDataSOAStride);
-		checkSlow(InstanceId >= 0 && InstanceId < GPUSceneNumAllocatedInstances);
+		checkSlow(InstanceId >= 0 && InstanceId < Scene.GPUScene.InstanceDataSOAStride);
+		checkSlow(InstanceId >= 0 && InstanceId < Scene.GPUScene.NumInstances);
 
-		FInstanceSceneData InstanceSceneData = GetInstanceSceneData(InstanceId, InstanceSceneDataSOAStride);
+		FInstanceSceneData InstanceSceneData = GetInstanceSceneData(InstanceId, Scene.GPUScene.InstanceDataSOAStride);
 
 		if (!InstanceSceneData.ValidInstance)
 		{
diff --git a/Engine/Source/Runtime/Renderer/Private/GPUScene.cpp b/Engine/Source/Runtime/Renderer/Private/GPUScene.cpp
--- a/Engine/Source/Runtime/Renderer/Private/GPUScene.cpp
+++ b/Engine/Source/Runtime/Renderer/Private/GPUScene.cpp
@@ -594,6 +594,7 @@ void FGPUScene::EndRender()
 	bInBeginEndBlock = false;
 	CurrentDynamicContext = nullptr;
 	BufferState = {};
+	ShaderParameters = {};
 }
 
 void FGPUScene::UpdateGPULights(FRDGBuilder& GraphBuilder, FScene& Scene)
@@ -1601,7 +1602,31 @@ bool FGPUScene::FillSceneUniformBuffer(FRDGBuilder& GraphBuilder, FSceneUniformB
 	{
 		return false;
 	}
-	return SceneUB.Set(SceneUB::GPUScene, ShaderParameters);
+
+	if (ShaderParameters.GPUScenePrimitiveSceneData != nullptr)
+	{
+		return SceneUB.Set(SceneUB::GPUScene, ShaderParameters);
+	}
+	else if (PrimitiveBuffer != nullptr)
+	{
+		FGPUSceneResourceParameters TmpParameters;
+		// Not in an active rendering context, must register the buffers and fill in the data structure.
+		TmpParameters.GPUSceneInstanceSceneData = GraphBuilder.CreateSRV(GraphBuilder.RegisterExternalBuffer(InstanceSceneDataBuffer));
+		TmpParameters.GPUSceneInstancePayloadData = GraphBuilder.CreateSRV(GraphBuilder.RegisterExternalBuffer(InstancePayloadDataBuffer));
+		TmpParameters.GPUScenePrimitiveSceneData = GraphBuilder.CreateSRV(GraphBuilder.RegisterExternalBuffer(PrimitiveBuffer));
+		TmpParameters.GPUSceneLightmapData = GraphBuilder.CreateSRV(GraphBuilder.RegisterExternalBuffer(LightmapDataBuffer));
+		TmpParameters.GPUSceneLightData = GraphBuilder.CreateSRV(GSystemTextures.GetDefaultStructuredBuffer(GraphBuilder, sizeof(FLightSceneData)));
+		TmpParameters.InstanceDataSOAStride = InstanceSceneDataSOAStride;
+		TmpParameters.NumScenePrimitives = NumScenePrimitives;
+		TmpParameters.NumInstances = InstanceSceneDataAllocator.GetMaxSize();
+		TmpParameters.GPUSceneFrameNumber = GetSceneFrameNumber();
+		return SceneUB.Set(SceneUB::GPUScene, TmpParameters);
+	}
+	else
+	{
+		// leave the dummy data in place - the gpu scene is not yet populated
+		return false;
+	}
 }
 
 void FGPUScene::AddPrimitiveToUpdate(int32 PrimitiveId, EPrimitiveDirtyState DirtyState)
diff --git a/Engine/Source/Runtime/Renderer/Private/GPUScene.h b/Engine/Source/Runtime/Renderer/Private/GPUScene.h
--- a/Engine/Source/Runtime/Renderer/Private/GPUScene.h
+++ b/Engine/Source/Runtime/Renderer/Private/GPUScene.h
@@ -301,7 +301,7 @@ public:
 	/**
 	 * Return the GPU scene resource
 	 */
-	FGPUSceneResourceParameters GetShaderParameters() const { return ShaderParameters; }
+	FGPUSceneResourceParameters GetShaderParameters() const { check(ShaderParameters.GPUScenePrimitiveSceneData != nullptr); return ShaderParameters; }
 
 	/**
 	 * Draw GPU-Scene debug info, such as bounding boxes. Call once per view at some point in the frame after GPU scene has been updated fully.
diff --git a/Engine/Source/Runtime/Renderer/Private/RendererScene.cpp b/Engine/Source/Runtime/Renderer/Private/RendererScene.cpp
--- a/Engine/Source/Runtime/Renderer/Private/RendererScene.cpp
+++ b/Engine/Source/Runtime/Renderer/Private/RendererScene.cpp
@@ -5427,6 +5427,10 @@ void FScene::UpdateAllPrimitiveSceneInfos(FRDGBuilder& GraphBuilder, EUpdateAllP
 
 	TArray<FVirtualShadowMapArrayCacheManager*, SceneRenderingAllocator> VirtualShadowCacheManagers;
 
+	// Create a SceneUB that permits access to the scene for invalidation processing.
+	FSceneUniformBuffer SceneUB;
+	GPUScene.FillSceneUniformBuffer(GraphBuilder, SceneUB);
+
 	{
 		SCOPED_NAMED_EVENT(FScene_VirtualShadowCacheUpdate, FColor::Orange);
 		GetAllVirtualShadowMapCacheManagers(VirtualShadowCacheManagers);
@@ -5460,7 +5464,7 @@ void FScene::UpdateAllPrimitiveSceneInfos(FRDGBuilder& GraphBuilder, EUpdateAllP
 
 			InvalidatingPrimitiveCollector.Finalize();
 
-			CacheManager->ProcessInvalidations(GraphBuilder, InvalidatingPrimitiveCollector);
+			CacheManager->ProcessInvalidations(GraphBuilder, SceneUB, InvalidatingPrimitiveCollector);
 		}
 	}
 
diff --git a/Engine/Source/Runtime/Renderer/Private/VirtualShadowMaps/VirtualShadowMapCacheManager.cpp b/Engine/Source/Runtime/Renderer/Private/VirtualShadowMaps/VirtualShadowMapCacheManager.cpp
--- a/Engine/Source/Runtime/Renderer/Private/VirtualShadowMaps/VirtualShadowMapCacheManager.cpp
+++ b/Engine/Source/Runtime/Renderer/Private/VirtualShadowMaps/VirtualShadowMapCacheManager.cpp
@@ -1132,7 +1132,7 @@ static uint32 GetPrimFlagsBufferSizeInDwords(int32 MaxPersistentPrimitiveIndex)
 	return FMath::RoundUpToPowerOfTwo(FMath::DivideAndRoundUp(MaxPersistentPrimitiveIndex, 32));
 }
 
-void FVirtualShadowMapArrayCacheManager::ProcessInvalidations(FRDGBuilder& GraphBuilder, const FInvalidatingPrimitiveCollector& InvalidatingPrimitiveCollector)
+void FVirtualShadowMapArrayCacheManager::ProcessInvalidations(FRDGBuilder& GraphBuilder, FSceneUniformBuffer &SceneUniformBuffer, const FInvalidatingPrimitiveCollector& InvalidatingPrimitiveCollector)
 {
 	// Always incorporate any scene removals into the "recently removed" list
 	UpdateRecentlyRemoved(InvalidatingPrimitiveCollector.GetRemovedPrimitives());
@@ -1143,7 +1143,7 @@ void FVirtualShadowMapArrayCacheManager::ProcessInvalidations(FRDGBuilder& Graph
 
 		if (!InvalidatingPrimitiveCollector.Instances.IsEmpty())
 		{
-			ProcessInvalidations(GraphBuilder, InvalidatingPrimitiveCollector.Instances);
+			ProcessInvalidations(GraphBuilder, SceneUniformBuffer, InvalidatingPrimitiveCollector.Instances);
 		}
 	}
 }
@@ -1186,13 +1186,7 @@ public:
 
 		SHADER_PARAMETER_RDG_BUFFER_UAV(RWStructuredBuffer< FPhysicalPageMetaData >, PhysicalPageMetaDataOut)
 
-		SHADER_PARAMETER_RDG_BUFFER_SRV(StructuredBuffer<float4>, GPUSceneInstanceSceneData)
-		SHADER_PARAMETER_RDG_BUFFER_SRV(StructuredBuffer<float4>, GPUSceneInstancePayloadData)
-		SHADER_PARAMETER_RDG_BUFFER_SRV(StructuredBuffer<float4>, GPUScenePrimitiveSceneData)
-		SHADER_PARAMETER(uint32, GPUSceneFrameNumber)
-		SHADER_PARAMETER(uint32, InstanceSceneDataSOAStride)
-		SHADER_PARAMETER(uint32, GPUSceneNumAllocatedInstances)
-		SHADER_PARAMETER(uint32, GPUSceneNumAllocatedPrimitives)
+		SHADER_PARAMETER_RDG_UNIFORM_BUFFER(FSceneUniformParameters, Scene)
 
 		SHADER_PARAMETER_RDG_BUFFER_SRV(StructuredBuffer< uint >, HZBPageTable)
 		SHADER_PARAMETER_RDG_BUFFER_SRV(StructuredBuffer< uint4 >, HZBPageRectBounds)
@@ -1223,7 +1217,6 @@ public:
 		FGlobalShader::ModifyCompilationEnvironment(Parameters, OutEnvironment);
 		FVirtualShadowMapArray::SetShaderDefines(OutEnvironment);
 		OutEnvironment.SetDefine(TEXT("CS_1D_GROUP_SIZE_X"), Cs1dGroupSizeX);
-		OutEnvironment.SetDefine(TEXT("USE_GLOBAL_GPU_SCENE_DATA"), 1);
 		OutEnvironment.SetDefine(TEXT("VF_SUPPORTS_PRIMITIVE_SCENE_DATA"), 1);
 		FGPUScene::FInstanceGPULoadBalancer::SetShaderDefines(OutEnvironment);
 	}
@@ -1252,6 +1245,7 @@ void FVirtualShadowMapArrayCacheManager::UpdateGPUMask(FRHIGPUMask GPUMask)
 
 void FVirtualShadowMapArrayCacheManager::ProcessInvalidations(
 	FRDGBuilder& GraphBuilder,
+	FSceneUniformBuffer &SceneUniformBuffer,
 	const FInstanceGPULoadBalancer& Instances) const
 {
 	if (Instances.IsEmpty())
@@ -1282,16 +1276,7 @@ void FVirtualShadowMapArrayCacheManager::ProcessInvalidations(
 		}
 		PassParameters->VirtualShadowMap = GraphBuilder.CreateUniformBuffer(UniformParameters);
 
-		{
-			const FGPUScene& GPUScene = Scene->GPUScene;
-			PassParameters->GPUSceneInstanceSceneData = GraphBuilder.CreateSRV(GraphBuilder.RegisterExternalBuffer(GPUScene.InstanceSceneDataBuffer));
-			PassParameters->GPUScenePrimitiveSceneData = GraphBuilder.CreateSRV(GraphBuilder.RegisterExternalBuffer(GPUScene.PrimitiveBuffer));
-			PassParameters->GPUSceneInstancePayloadData = GraphBuilder.CreateSRV(GraphBuilder.RegisterExternalBuffer(GPUScene.InstancePayloadDataBuffer));
-			PassParameters->GPUSceneFrameNumber = GPUScene.GetSceneFrameNumber();
-			PassParameters->GPUSceneNumAllocatedInstances = GPUScene.GetNumInstances();
-			PassParameters->GPUSceneNumAllocatedPrimitives = GPUScene.GetNumPrimitives();
-			PassParameters->InstanceSceneDataSOAStride = GPUScene.InstanceSceneDataSOAStride;
-		}
+		PassParameters->Scene = SceneUniformBuffer.GetBuffer(GraphBuilder);
 
 		PassParameters->PhysicalPageMetaDataOut = GraphBuilder.CreateUAV(GraphBuilder.RegisterExternalBuffer(PhysicalPageMetaData));
 		PassParameters->bDrawBounds = CVarDrawInvalidatingBounds.GetValueOnRenderThread() != 0;
diff --git a/Engine/Source/Runtime/Renderer/Private/VirtualShadowMaps/VirtualShadowMapCacheManager.h b/Engine/Source/Runtime/Renderer/Private/VirtualShadowMaps/VirtualShadowMapCacheManager.h
--- a/Engine/Source/Runtime/Renderer/Private/VirtualShadowMaps/VirtualShadowMapCacheManager.h
+++ b/Engine/Source/Runtime/Renderer/Private/VirtualShadowMaps/VirtualShadowMapCacheManager.h
@@ -298,7 +298,7 @@ public:
 		TBitArray<SceneRenderingAllocator> RemovedPrimitives;
 	};
 
-	void ProcessInvalidations(FRDGBuilder& GraphBuilder, const FInvalidatingPrimitiveCollector& InvalidatingPrimitiveCollector);
+	void ProcessInvalidations(FRDGBuilder& GraphBuilder, FSceneUniformBuffer &SceneUniformBuffer, const FInvalidatingPrimitiveCollector& InvalidatingPrimitiveCollector);
 
 	/**
 	 * Allow the cache manager to track scene changes, in particular track resizing of primitive tracking data.
@@ -328,7 +328,7 @@ public:
 #endif
 
 private:
-	void ProcessInvalidations(FRDGBuilder& GraphBuilder, const FInstanceGPULoadBalancer& Instances) const;
+	void ProcessInvalidations(FRDGBuilder& GraphBuilder, FSceneUniformBuffer &SceneUniformBuffer, const FInstanceGPULoadBalancer& Instances) const;
 
 	void ExtractStats(FRDGBuilder& GraphBuilder, FVirtualShadowMapArray &VirtualShadowMapArray);
 
