Removing the RealTime viewport requirement for running the full history pass for DistanceFieldAO. This eliminates issues in the editor where some editor viewports have RealTime Off selectively between PIE and non-PIE or have issues when windows or graph nodes are being dragged around, forcing RealTime Off.

#rnx

diff --git a/Engine/Source/Runtime/Renderer/Private/DistanceFieldLightingPost.cpp b/Engine/Source/Runtime/Renderer/Private/DistanceFieldLightingPost.cpp
--- a/Engine/Source/Runtime/Renderer/Private/DistanceFieldLightingPost.cpp
+++ b/Engine/Source/Runtime/Renderer/Private/DistanceFieldLightingPost.cpp
@@ -329,7 +329,6 @@ void UpdateHistory(
 		if (*BentNormalHistoryState 
 			&& !View.bCameraCut 
 			&& !View.bPrevTransformsReset
-			&& View.Family->bRealtimeUpdate
 			&& !GAOClearHistory
 			// If the scene render targets reallocate, toss the history so we don't read uninitialized data
 			&& (*BentNormalHistoryState)->GetDesc().Extent == BufferSize)
