Added local height fog volumes gpu stat.

diff --git a/Engine/Source/Runtime/Renderer/Private/LocalHeightFogRendering.cpp b/Engine/Source/Runtime/Renderer/Private/LocalHeightFogRendering.cpp
--- a/Engine/Source/Runtime/Renderer/Private/LocalHeightFogRendering.cpp
+++ b/Engine/Source/Runtime/Renderer/Private/LocalHeightFogRendering.cpp
@@ -23,6 +23,7 @@ bool ShouldRenderLocalHeightFog(const FScene* Scene, const FSceneViewFamily& Fam
 	return false;
 }
 
+DECLARE_GPU_STAT(LocalHeightFogVolumes);
 
 
 /*=============================================================================
@@ -138,98 +139,98 @@ void RenderLocalHeightFog(
 	FRDGTextureRef LightShaftOcclusionTexture)
 {
 	uint32 LocalHeightFogInstanceCount = Scene->LocalHeightFogs.Num();
-	if (LocalHeightFogInstanceCount <= 0)
+	if (LocalHeightFogInstanceCount > 0)
 	{
-		return;
-	}
-
-	// No culling as of today
-	const uint32 AllLocalHeightFogInstanceBytes = sizeof(FLocalHeightFogGPUInstanceData) * LocalHeightFogInstanceCount;
-	FLocalHeightFogGPUInstanceData* LocalHeightFogGPUInstanceData = (FLocalHeightFogGPUInstanceData*) GraphBuilder.Alloc(AllLocalHeightFogInstanceBytes, 16);
-	FLocalHeightFogGPUInstanceData* LocalHeightFogGPUInstanceDataIt = LocalHeightFogGPUInstanceData;
-	for (FLocalHeightFogSceneProxy* LHF : Scene->LocalHeightFogs)
-	{
-		LocalHeightFogGPUInstanceDataIt->Transform = FMatrix44f(LHF->FogTransform.ToMatrixWithScale());
-		LocalHeightFogGPUInstanceDataIt->InvTransform = LocalHeightFogGPUInstanceDataIt->Transform.Inverse();
+		RDG_GPU_STAT_SCOPE(GraphBuilder, LocalHeightFogVolumes);
+		
+		// No culling as of today
+		const uint32 AllLocalHeightFogInstanceBytes = sizeof(FLocalHeightFogGPUInstanceData) * LocalHeightFogInstanceCount;
+		FLocalHeightFogGPUInstanceData* LocalHeightFogGPUInstanceData = (FLocalHeightFogGPUInstanceData*) GraphBuilder.Alloc(AllLocalHeightFogInstanceBytes, 16);
+		FLocalHeightFogGPUInstanceData* LocalHeightFogGPUInstanceDataIt = LocalHeightFogGPUInstanceData;
+		for (FLocalHeightFogSceneProxy* LHF : Scene->LocalHeightFogs)
+		{
+			LocalHeightFogGPUInstanceDataIt->Transform = FMatrix44f(LHF->FogTransform.ToMatrixWithScale());
+			LocalHeightFogGPUInstanceDataIt->InvTransform = LocalHeightFogGPUInstanceDataIt->Transform.Inverse();
 
-		LocalHeightFogGPUInstanceDataIt->Density = LHF->FogDensity;
-		LocalHeightFogGPUInstanceDataIt->HeightFalloff = LHF->FogHeightFalloff * 0.01f;	// This scale is used to have artist author reasonable range.
-		LocalHeightFogGPUInstanceDataIt->HeightOffset = LHF->FogHeightOffset;
-		LocalHeightFogGPUInstanceDataIt->RadialAttenuation = LHF->FogRadialAttenuation;
-		LocalHeightFogGPUInstanceDataIt->RadialAttenuationHermiteTangent = 1.0f - LHF->FogRadialAttenuationSoftness;
+			LocalHeightFogGPUInstanceDataIt->Density = LHF->FogDensity;
+			LocalHeightFogGPUInstanceDataIt->HeightFalloff = LHF->FogHeightFalloff * 0.01f;	// This scale is used to have artist author reasonable range.
+			LocalHeightFogGPUInstanceDataIt->HeightOffset = LHF->FogHeightOffset;
+			LocalHeightFogGPUInstanceDataIt->RadialAttenuation = LHF->FogRadialAttenuation;
+			LocalHeightFogGPUInstanceDataIt->RadialAttenuationHermiteTangent = 1.0f - LHF->FogRadialAttenuationSoftness;
 
-		LocalHeightFogGPUInstanceDataIt->Albedo = FVector3f(LHF->FogAlbedo);
-		LocalHeightFogGPUInstanceDataIt->PhaseG = LHF->FogPhaseG;
-		LocalHeightFogGPUInstanceDataIt->Emissive = FVector3f(LHF->FogEmissive);
+			LocalHeightFogGPUInstanceDataIt->Albedo = FVector3f(LHF->FogAlbedo);
+			LocalHeightFogGPUInstanceDataIt->PhaseG = LHF->FogPhaseG;
+			LocalHeightFogGPUInstanceDataIt->Emissive = FVector3f(LHF->FogEmissive);
 
-		LocalHeightFogGPUInstanceDataIt++;
-	}
-	FRDGBufferRef LocalHeightFogGPUInstanceDataBuffer =  CreateStructuredBuffer(GraphBuilder, TEXT("LocalHeightFogGPUInstanceDataBuffer"), 
-		sizeof(FLocalHeightFogGPUInstanceData), LocalHeightFogInstanceCount, LocalHeightFogGPUInstanceData, AllLocalHeightFogInstanceBytes, ERDGInitialDataFlags::NoCopy);
-	FRDGBufferSRVRef LocalHeightFogGPUInstanceDataBufferSRV = GraphBuilder.CreateSRV(LocalHeightFogGPUInstanceDataBuffer);
+			LocalHeightFogGPUInstanceDataIt++;
+		}
+		FRDGBufferRef LocalHeightFogGPUInstanceDataBuffer =  CreateStructuredBuffer(GraphBuilder, TEXT("LocalHeightFogGPUInstanceDataBuffer"), 
+			sizeof(FLocalHeightFogGPUInstanceData), LocalHeightFogInstanceCount, LocalHeightFogGPUInstanceData, AllLocalHeightFogInstanceBytes, ERDGInitialDataFlags::NoCopy);
+		FRDGBufferSRVRef LocalHeightFogGPUInstanceDataBufferSRV = GraphBuilder.CreateSRV(LocalHeightFogGPUInstanceDataBuffer);
 
-	FRDGTextureRef SceneColorTexture = SceneTextures.Color.Resolve;
+		FRDGTextureRef SceneColorTexture = SceneTextures.Color.Resolve;
 
-	for (FViewInfo& View : Views)
-	{
-		FLocalHeightFogPassParameters* PassParameters = GraphBuilder.AllocParameters<FLocalHeightFogPassParameters>();
+		for (FViewInfo& View : Views)
+		{
+			FLocalHeightFogPassParameters* PassParameters = GraphBuilder.AllocParameters<FLocalHeightFogPassParameters>();
 
-		PassParameters->VS.View = GetShaderBinding(View.ViewUniformBuffer);
-		PassParameters->VS.LocalHeightFogInstances = LocalHeightFogGPUInstanceDataBufferSRV;
+			PassParameters->VS.View = GetShaderBinding(View.ViewUniformBuffer);
+			PassParameters->VS.LocalHeightFogInstances = LocalHeightFogGPUInstanceDataBufferSRV;
 
-		PassParameters->PS.View = GetShaderBinding(View.ViewUniformBuffer);
-		PassParameters->PS.LocalHeightFogInstances = LocalHeightFogGPUInstanceDataBufferSRV;
+			PassParameters->PS.View = GetShaderBinding(View.ViewUniformBuffer);
+			PassParameters->PS.LocalHeightFogInstances = LocalHeightFogGPUInstanceDataBufferSRV;
 
-		PassParameters->SceneTextures = SceneTextures.UniformBuffer;
-		PassParameters->RenderTargets[0] = FRenderTargetBinding(SceneColorTexture, ERenderTargetLoadAction::ENoAction);
+			PassParameters->SceneTextures = SceneTextures.UniformBuffer;
+			PassParameters->RenderTargets[0] = FRenderTargetBinding(SceneColorTexture, ERenderTargetLoadAction::ENoAction);
 
-		FLocalHeightFogVS::FPermutationDomain VSPermutationVector;
-		auto VertexShader = View.ShaderMap->GetShader< FLocalHeightFogVS >(VSPermutationVector);
+			FLocalHeightFogVS::FPermutationDomain VSPermutationVector;
+			auto VertexShader = View.ShaderMap->GetShader< FLocalHeightFogVS >(VSPermutationVector);
 
-		FLocalHeightFogPS::FPermutationDomain PsPermutationVector;
-		auto PixelShader = View.ShaderMap->GetShader< FLocalHeightFogPS >(PsPermutationVector);
+			FLocalHeightFogPS::FPermutationDomain PsPermutationVector;
+			auto PixelShader = View.ShaderMap->GetShader< FLocalHeightFogPS >(PsPermutationVector);
 
-		const FIntRect ViewRect = View.ViewRect;
+			const FIntRect ViewRect = View.ViewRect;
 
-		ClearUnusedGraphResources(VertexShader, &PassParameters->VS);
-		ClearUnusedGraphResources(PixelShader, &PassParameters->PS);
+			ClearUnusedGraphResources(VertexShader, &PassParameters->VS);
+			ClearUnusedGraphResources(PixelShader, &PassParameters->PS);
 
-		GraphBuilder.AddPass(
-			RDG_EVENT_NAME("RenderLocalHeightFog %u inst.", LocalHeightFogInstanceCount),
-			PassParameters,
-			ERDGPassFlags::Raster,
-			[VertexShader, PixelShader, PassParameters, LocalHeightFogInstanceCount, ViewRect](FRHICommandList& RHICmdList)
-		{
-			FGraphicsPipelineStateInitializer GraphicsPSOInit;
-			RHICmdList.ApplyCachedRenderTargets(GraphicsPSOInit);
+			GraphBuilder.AddPass(
+				RDG_EVENT_NAME("RenderLocalHeightFog %u inst.", LocalHeightFogInstanceCount),
+				PassParameters,
+				ERDGPassFlags::Raster,
+				[VertexShader, PixelShader, PassParameters, LocalHeightFogInstanceCount, ViewRect](FRHICommandList& RHICmdList)
+			{
+				FGraphicsPipelineStateInitializer GraphicsPSOInit;
+				RHICmdList.ApplyCachedRenderTargets(GraphicsPSOInit);
 
-			RHICmdList.SetViewport(ViewRect.Min.X, ViewRect.Min.Y, 0.0f, ViewRect.Max.X, ViewRect.Max.Y, 1.0f);
+				RHICmdList.SetViewport(ViewRect.Min.X, ViewRect.Min.Y, 0.0f, ViewRect.Max.X, ViewRect.Max.Y, 1.0f);
 
-			// Render back faces only since camera may intersect
-			GraphicsPSOInit.RasterizerState = TStaticRasterizerState<FM_Solid, CM_CW>::GetRHI();
-			GraphicsPSOInit.DepthStencilState = TStaticDepthStencilState<false, CF_Always>::GetRHI();
-			GraphicsPSOInit.BlendState = TStaticBlendState<CW_RGBA, BO_Add, BF_One, BF_SourceAlpha, BO_Add, BF_Zero, BF_SourceAlpha>::GetRHI();
-			GraphicsPSOInit.PrimitiveType = PT_TriangleList;
+				// Render back faces only since camera may intersect
+				GraphicsPSOInit.RasterizerState = TStaticRasterizerState<FM_Solid, CM_CW>::GetRHI();
+				GraphicsPSOInit.DepthStencilState = TStaticDepthStencilState<false, CF_Always>::GetRHI();
+				GraphicsPSOInit.BlendState = TStaticBlendState<CW_RGBA, BO_Add, BF_One, BF_SourceAlpha, BO_Add, BF_Zero, BF_SourceAlpha>::GetRHI();
+				GraphicsPSOInit.PrimitiveType = PT_TriangleList;
 
-			GraphicsPSOInit.BoundShaderState.VertexDeclarationRHI = GetVertexDeclarationFVector4();
-			GraphicsPSOInit.BoundShaderState.VertexShaderRHI = VertexShader.GetVertexShader();
-			GraphicsPSOInit.BoundShaderState.PixelShaderRHI = PixelShader.GetPixelShader();
+				GraphicsPSOInit.BoundShaderState.VertexDeclarationRHI = GetVertexDeclarationFVector4();
+				GraphicsPSOInit.BoundShaderState.VertexShaderRHI = VertexShader.GetVertexShader();
+				GraphicsPSOInit.BoundShaderState.PixelShaderRHI = PixelShader.GetPixelShader();
 
-			SetGraphicsPipelineState(RHICmdList, GraphicsPSOInit, 0);
+				SetGraphicsPipelineState(RHICmdList, GraphicsPSOInit, 0);
 
-			SetShaderParameters(RHICmdList, VertexShader, VertexShader.GetVertexShader(), PassParameters->VS);
-			SetShaderParameters(RHICmdList, PixelShader, PixelShader.GetPixelShader(), PassParameters->PS);
+				SetShaderParameters(RHICmdList, VertexShader, VertexShader.GetVertexShader(), PassParameters->VS);
+				SetShaderParameters(RHICmdList, PixelShader, PixelShader.GetPixelShader(), PassParameters->PS);
 
-			RHICmdList.SetStreamSource(0, GetUnitCubeVertexBuffer(), 0);
+				RHICmdList.SetStreamSource(0, GetUnitCubeVertexBuffer(), 0);
 
-			RHICmdList.DrawIndexedPrimitive( GetUnitCubeIndexBuffer()
-				, 0									//BaseVertexIndex
-				, 0									//FirstInstance
-				, 8									//uint32 NumVertices
-				, 0									//uint32 StartIndex
-				, UE_ARRAY_COUNT(GCubeIndices) / 3	//uint32 NumPrimitives
-				, LocalHeightFogInstanceCount		//uint32 NumInstances
-			);
-		});
+				RHICmdList.DrawIndexedPrimitive( GetUnitCubeIndexBuffer()
+					, 0									//BaseVertexIndex
+					, 0									//FirstInstance
+					, 8									//uint32 NumVertices
+					, 0									//uint32 StartIndex
+					, UE_ARRAY_COUNT(GCubeIndices) / 3	//uint32 NumPrimitives
+					, LocalHeightFogInstanceCount		//uint32 NumInstances
+				);
+			});
+		}
 	}
 }
