- Make sure water material received post load before precaching them
- re-enable PSO precaching for base pass after auto resolve issues which removed the PSO precaching

diff --git a/Engine/Plugins/Experimental/Water/Source/Runtime/Private/WaterBodyComponent.cpp b/Engine/Plugins/Experimental/Water/Source/Runtime/Private/WaterBodyComponent.cpp
--- a/Engine/Plugins/Experimental/Water/Source/Runtime/Private/WaterBodyComponent.cpp
+++ b/Engine/Plugins/Experimental/Water/Source/Runtime/Private/WaterBodyComponent.cpp
@@ -1541,14 +1541,17 @@ void UWaterBodyComponent::PostLoad()
 		SetupPrecachePSOParams(PrecachePSOParams);
 		if (WaterMaterial)
 		{
+			WaterMaterial->ConditionalPostLoad();
 			WaterMaterial->PrecachePSOs(&FLocalVertexFactory::StaticType, PrecachePSOParams);
 		}
 		if (UnderwaterPostProcessMaterial)
 		{
+			UnderwaterPostProcessMaterial->ConditionalPostLoad();
 			UnderwaterPostProcessMaterial->PrecachePSOs(&FLocalVertexFactory::StaticType, PrecachePSOParams);
 		}
 		if (WaterInfoMaterial)
 		{
+			WaterInfoMaterial->ConditionalPostLoad();
 			WaterInfoMaterial->PrecachePSOs(&FLocalVertexFactory::StaticType, PrecachePSOParams);
 		}
 	}
diff --git a/Engine/Source/Runtime/Renderer/Private/BasePassRendering.cpp b/Engine/Source/Runtime/Renderer/Private/BasePassRendering.cpp
--- a/Engine/Source/Runtime/Renderer/Private/BasePassRendering.cpp
+++ b/Engine/Source/Runtime/Renderer/Private/BasePassRendering.cpp
@@ -195,7 +195,7 @@ static bool IsBasePassWaitForTasksEnabled()
 static bool IsStandardTranslucenyPassSeparated()
 {
 	static const auto TranslucencyStandardSeparatedCVar = IConsoleManager::Get().FindTConsoleVariableDataInt(TEXT("r.Translucency.StandardSeparated"));
-	return TranslucencyStandardSeparatedCVar && TranslucencyStandardSeparatedCVar->GetValueOnRenderThread() != 0;
+	return TranslucencyStandardSeparatedCVar && TranslucencyStandardSeparatedCVar->GetValueOnAnyThread() != 0;
 }
 
 template<uint32 StencilRef> void SetTranslucentPassDepthStencilState(FMeshPassProcessorRenderState& DrawRenderState, bool bDisableDepthTest)
@@ -2449,10 +2449,10 @@ FMeshPassProcessor* CreateTranslucencyAllPassProcessor(ERHIFeatureLevel::Type Fe
 	return new FBasePassMeshProcessor(EMeshPass::TranslucencyAll, Scene, FeatureLevel, InViewIfDynamicMeshCommand, PassDrawRenderState, InDrawListContext, Flags, ETranslucencyPass::TPT_AllTranslucency);
 }
 
-FRegisterPassProcessorCreateFunction RegisterBasePass(&CreateBasePassProcessor, EShadingPath::Deferred, EMeshPass::BasePass, EMeshPassFlags::CachedMeshCommands | EMeshPassFlags::MainView);
-FRegisterPassProcessorCreateFunction RegisterTranslucencyStandardPass(&CreateTranslucencyStandardPassProcessor, EShadingPath::Deferred, EMeshPass::TranslucencyStandard, EMeshPassFlags::MainView);
+REGISTER_MESHPASSPROCESSOR_AND_PSOCOLLECTOR(BasePass, CreateBasePassProcessor, EShadingPath::Deferred, EMeshPass::BasePass, EMeshPassFlags::CachedMeshCommands | EMeshPassFlags::MainView);
+REGISTER_MESHPASSPROCESSOR_AND_PSOCOLLECTOR(TranslucencyStandardPass, CreateTranslucencyStandardPassProcessor, EShadingPath::Deferred, EMeshPass::TranslucencyStandard, EMeshPassFlags::MainView);
 REGISTER_MESHPASSPROCESSOR_AND_PSOCOLLECTOR(TranslucencyStandardModulatePass, CreateTranslucencyStandardModulatePassProcessor, EShadingPath::Deferred, EMeshPass::TranslucencyStandardModulate, EMeshPassFlags::MainView);
-FRegisterPassProcessorCreateFunction RegisterTranslucencyAfterDOFPass(&CreateTranslucencyAfterDOFProcessor, EShadingPath::Deferred, EMeshPass::TranslucencyAfterDOF, EMeshPassFlags::MainView);
-FRegisterPassProcessorCreateFunction RegisterTranslucencyAfterDOFModulatePass(&CreateTranslucencyAfterDOFModulateProcessor, EShadingPath::Deferred, EMeshPass::TranslucencyAfterDOFModulate, EMeshPassFlags::MainView);
-FRegisterPassProcessorCreateFunction RegisterTranslucencyAfterMotionBlurPass(&CreateTranslucencyAfterMotionBlurProcessor, EShadingPath::Deferred, EMeshPass::TranslucencyAfterMotionBlur, EMeshPassFlags::MainView);
-FRegisterPassProcessorCreateFunction RegisterTranslucencyAllPass(&CreateTranslucencyAllPassProcessor, EShadingPath::Deferred, EMeshPass::TranslucencyAll, EMeshPassFlags::MainView);
\ No newline at end of file
+REGISTER_MESHPASSPROCESSOR_AND_PSOCOLLECTOR(TranslucencyAfterDOFPass, CreateTranslucencyAfterDOFProcessor, EShadingPath::Deferred, EMeshPass::TranslucencyAfterDOF, EMeshPassFlags::MainView);
+REGISTER_MESHPASSPROCESSOR_AND_PSOCOLLECTOR(TranslucencyAfterDOFModulatePass, CreateTranslucencyAfterDOFModulateProcessor, EShadingPath::Deferred, EMeshPass::TranslucencyAfterDOFModulate, EMeshPassFlags::MainView);
+REGISTER_MESHPASSPROCESSOR_AND_PSOCOLLECTOR(TranslucencyAfterMotionBlurPass, CreateTranslucencyAfterMotionBlurProcessor, EShadingPath::Deferred, EMeshPass::TranslucencyAfterMotionBlur, EMeshPassFlags::MainView);
+REGISTER_MESHPASSPROCESSOR_AND_PSOCOLLECTOR(TranslucencyAllPass, CreateTranslucencyAllPassProcessor, EShadingPath::Deferred, EMeshPass::TranslucencyAll, EMeshPassFlags::MainView);
