Fixed crash in visibility system; initialize dither uniform buffers early so that shadow view initialization copies the correct state from FViewInfo.

diff --git a/Engine/Source/Runtime/Renderer/Private/SceneVisibility.cpp b/Engine/Source/Runtime/Renderer/Private/SceneVisibility.cpp
--- a/Engine/Source/Runtime/Renderer/Private/SceneVisibility.cpp
+++ b/Engine/Source/Runtime/Renderer/Private/SceneVisibility.cpp
@@ -4059,19 +4059,6 @@ void FSceneRenderer::PreVisibilityFrameSetup(FRDGBuilder& GraphBuilder)
 	}
 #endif
 
-	// Setup global dither fade in and fade out uniform buffers.
-	for (int32 ViewIndex = 0; ViewIndex < Views.Num(); ViewIndex++)
-	{
-		FViewInfo& View = Views[ViewIndex];
-
-		FDitherUniformShaderParameters DitherUniformShaderParameters;
-		DitherUniformShaderParameters.LODFactor = View.GetTemporalLODTransition();
-		View.DitherFadeOutUniformBuffer = FDitherUniformBufferRef::CreateUniformBufferImmediate(DitherUniformShaderParameters, UniformBuffer_SingleFrame);
-
-		DitherUniformShaderParameters.LODFactor = View.GetTemporalLODTransition() - 1.0f;
-		View.DitherFadeInUniformBuffer = FDitherUniformBufferRef::CreateUniformBufferImmediate(DitherUniformShaderParameters, UniformBuffer_SingleFrame);
-	}
-
 	for (const auto& ViewExtension : ViewFamily.ViewExtensions)
 	{
 		ViewExtension->PreInitViews_RenderThread(GraphBuilder);
@@ -4110,6 +4097,16 @@ void FSceneRenderer::PrepareViewStateForVisibility(const FSceneTexturesConfig& S
 
 		check(View.VerifyMembersChecks());
 
+		// Setup global dither fade in and fade out uniform buffers.
+		{
+			FDitherUniformShaderParameters DitherUniformShaderParameters;
+			DitherUniformShaderParameters.LODFactor = View.GetTemporalLODTransition();
+			View.DitherFadeOutUniformBuffer = FDitherUniformBufferRef::CreateUniformBufferImmediate(DitherUniformShaderParameters, UniformBuffer_SingleFrame);
+
+			DitherUniformShaderParameters.LODFactor = View.GetTemporalLODTransition() - 1.0f;
+			View.DitherFadeInUniformBuffer = FDitherUniformBufferRef::CreateUniformBufferImmediate(DitherUniformShaderParameters, UniformBuffer_SingleFrame);
+		}
+
 		// Once per render increment the occlusion frame counter.
 		if (ViewState)
 		{
