Removed unused task event from visibility system that is causing a race condition.

diff --git a/Engine/Source/Runtime/Renderer/Private/SceneVisibility.cpp b/Engine/Source/Runtime/Renderer/Private/SceneVisibility.cpp
--- a/Engine/Source/Runtime/Renderer/Private/SceneVisibility.cpp
+++ b/Engine/Source/Runtime/Renderer/Private/SceneVisibility.cpp
@@ -3181,8 +3181,6 @@ FVisibilityViewPacket::FVisibilityViewPacket(FVisibilityTaskData& InTaskData, FS
 				Relevance.CommandPipe.ReleaseNumCommands(1);
 				Tasks.OcclusionCull.Trigger();
 			}
-
-			Tasks.OcclusionCullPipe.Trigger();
 		});
 
 		// Callback for when a relevance command is queued from occlusion.
@@ -3864,7 +3862,6 @@ void FVisibilityTaskData::ProcessRenderThreadTasks(FExclusiveDepthStencil::Type
 			}
 
 			ViewPacket.Tasks.OcclusionCull.Trigger();
-			ViewPacket.Tasks.OcclusionCullPipe.Trigger();
 		}
 
 		MergeSecondaryViewVisibility();
diff --git a/Engine/Source/Runtime/Renderer/Private/SceneVisibilityPrivate.h b/Engine/Source/Runtime/Renderer/Private/SceneVisibilityPrivate.h
--- a/Engine/Source/Runtime/Renderer/Private/SceneVisibilityPrivate.h
+++ b/Engine/Source/Runtime/Renderer/Private/SceneVisibilityPrivate.h
@@ -347,7 +347,6 @@ private:
 	struct FTasks
 	{
 		UE::Tasks::FTaskEvent FrustumCull{ UE_SOURCE_LOCATION };
-		UE::Tasks::FTaskEvent OcclusionCullPipe{ UE_SOURCE_LOCATION };
 		UE::Tasks::FTaskEvent OcclusionCull{ UE_SOURCE_LOCATION };
 		UE::Tasks::FTaskEvent ComputeRelevance{ UE_SOURCE_LOCATION };
 		UE::Tasks::FTaskEvent LightVisibility{ UE_SOURCE_LOCATION };
@@ -367,8 +366,8 @@ private:
 
 		Frustum Cull          - Primitives are frustum / distance culled and visible primitives are emitted.
 		Occlusion Cull        - Primitives are culled against occluders in the scene and visible primitives are emitted.
-		Compute Relevance     - Primitives are queried for view relevance information and dynamic primitives are identified and emitted.
-		Static Mesh Filtering - Primitives with static meshes are filtered into various mesh passes.
+		Compute Relevance     - Primitives are queried for view relevance informationm dynamic primitives are identified and emitted, and static meshes are filtered
+								into various mesh passes
 
 	The following stages are then performed for all views:
 
@@ -379,7 +378,7 @@ private:
 	'command pipes' which are serial queues that run between stages to launch work for the next stage. Each view has two pipes: OcclusionCull, and Relevance. The OcclusionCull
 	pipe can either process occlusion tasks, or act as the relevance pipe if occlusion is disabled. The relevance pipe only launches relevance tasks.
 
-	The render thread performs the GDME stage which syncs the compute relevance stage for all views. Setup Mesh Passes then runs once all the static mesh filtering tasks
+	The render thread performs the GDME stage which syncs the compute relevance stage for all views. Setup Mesh Passes then runs once all the GDME tasks
 	have completed. When the renderer has only one view, GDME utilizes a command pipe to process requests from relevance as quickly as possible and achieves some overlap
 	with relevance, reducing the critical path. With multiple views, it's necessary to sync beforehand as the gather requires a view mask for each dynamic primitive. To
 	that end, GDME supports two paths for processing dynamic primitives: an index list or view bit mask. The former just supplies a list of dynamic primitives and the view
