Fixed CIS build errors and resubmit after CL25523894 backed out CL25523052

[SceneCapture] Add r.SceneCapture.DepthPrepassOptimization to speed up pre-pass depth rendering for scene capture.
This is a short-term optimization to support certain features requiring a depth pre-pass and can only utilize SceneCapture to achieve it currently.
Note turning it on means rendering after depth pre-pass (e.g. SingleLayerWater) is ignored, hence result is different from when CVar is off.

diff --git a/Engine/Source/Runtime/Renderer/Private/DeferredShadingRenderer.h b/Engine/Source/Runtime/Renderer/Private/DeferredShadingRenderer.h
--- a/Engine/Source/Runtime/Renderer/Private/DeferredShadingRenderer.h
+++ b/Engine/Source/Runtime/Renderer/Private/DeferredShadingRenderer.h
@@ -1064,6 +1064,7 @@ private:
 		const FSceneTextureParameters& SceneTextures,
 		FRDGTextureRef* OutAmbientOcclusionTexture);
 	
+	ERendererOutput GetRendererOutput() const;
 
 #if RHI_RAYTRACING
 	template <int TextureImportanceSampling>
@@ -1152,7 +1153,6 @@ private:
 	FRayTracingLocalShaderBindings* BuildLumenHardwareRayTracingMaterialBindings(FRHICommandList& RHICmdList, const FViewInfo& View, FRDGBufferRef OutHitGroupDataBuffer, bool bInlineOnly);
 	void SetupLumenHardwareRayTracingHitGroupBuffer(FRDGBuilder& GraphBuilder, FViewInfo& View);
 	void SetupLumenHardwareRayTracingUniformBuffer(FRDGBuilder& GraphBuilder, FViewInfo& View);
-
 	// #dxr_todo: UE-72565: refactor ray tracing effects to not be member functions of DeferredShadingRenderer. Register each effect at startup and just loop over them automatically
 	static void PrepareRayTracingReflections(const FViewInfo& View, const FScene& Scene, TArray<FRHIRayTracingShader*>& OutRayGenShaders);
 	static void PrepareRayTracingDeferredReflections(const FViewInfo& View, const FScene& Scene, TArray<FRHIRayTracingShader*>& OutRayGenShaders);
