[SceneCapture] Add r.SceneCapture.DepthPrepassOptimization to speed up pre-pass depth rendering for scene capture.
This is a short-term optimization to support certain features requiring a depth pre-pass and can only utilize SceneCapture to achieve it currently.
Note turning it on means rendering after depth pre-pass (e.g. SingleLayerWater) is ignored, hence result is different from when CVar is off.

diff --git a/Engine/Source/Runtime/Renderer/Private/DeferredShadingRenderer.cpp b/Engine/Source/Runtime/Renderer/Private/DeferredShadingRenderer.cpp
--- a/Engine/Source/Runtime/Renderer/Private/DeferredShadingRenderer.cpp
+++ b/Engine/Source/Runtime/Renderer/Private/DeferredShadingRenderer.cpp
@@ -100,6 +100,15 @@ extern int32 GNanitePickingDomain;
 
 extern DynamicRenderScaling::FBudget GDynamicNaniteScalingPrimary;
 
+// This is an experimental optimization switch to render pre-pass depth for scene capture without calling the entire FDeferredShadingSceneRenderer::Render()
+int32 GSceneCaptureDepthPrepassOptimization = 0;
+static FAutoConsoleVariableRef CVarSceneCaptureDepthPrepassOptimization(
+	TEXT("r.SceneCapture.DepthPrepassOptimization"),
+	GSceneCaptureDepthPrepassOptimization,
+	TEXT("Whether to apply optimized render path when capturing depth prepass for scene capture 2D. Experimental!\n")
+	TEXT("Warning: turning it on means rendering after depth pre-pass (e.g. SingleLayerWater) is ignored, hence result is different from when CVar is off.\n"),
+	ECVF_RenderThreadSafe | ECVF_Scalability);
+
 static TAutoConsoleVariable<int32> CVarClearCoatNormal(
 	TEXT("r.ClearCoatNormal"),
 	0,
@@ -2459,6 +2468,12 @@ bool FDeferredShadingSceneRenderer::IsNaniteEnabled() const
 	return UseNanite(ShaderPlatform) && ViewFamily.EngineShowFlags.NaniteMeshes && Nanite::GStreamingManager.HasResourceEntries();
 }
 
+FDeferredShadingSceneRenderer::ERendererOutput FDeferredShadingSceneRenderer::GetRendererOutput() const
+{
+	const bool bSceneCaptureDepthPrepass = Views[0].bIsSceneCapture && (ViewFamily.SceneCaptureSource == ESceneCaptureSource::SCS_SceneDepth || ViewFamily.SceneCaptureSource == ESceneCaptureSource::SCS_DeviceDepth);
+	return bSceneCaptureDepthPrepass && GSceneCaptureDepthPrepassOptimization ? ERendererOutput::DepthPrepassOnly : ERendererOutput::FinalSceneColor;
+}
+
 void FDeferredShadingSceneRenderer::Render(FRDGBuilder& GraphBuilder)
 {
 	if (!ViewFamily.EngineShowFlags.Rendering)
@@ -2466,6 +2481,9 @@ void FDeferredShadingSceneRenderer::Render(FRDGBuilder& GraphBuilder)
 		return;
 	}
 
+	// If this is scene capture rendering depth pre-pass, we'll take the shortcut function RenderSceneCaptureDepth if optimization switch is on.
+	const ERendererOutput RendererOutput = GetRendererOutput();
+
 	const bool bNaniteEnabled = IsNaniteEnabled();
 
 #if !UE_BUILD_SHIPPING
@@ -2486,7 +2504,7 @@ void FDeferredShadingSceneRenderer::Render(FRDGBuilder& GraphBuilder)
 
 	const bool bUseVirtualTexturing = UseVirtualTexturing(FeatureLevel);
 
-	if (bUseVirtualTexturing)
+	if (bUseVirtualTexturing && RendererOutput == ERendererOutput::FinalSceneColor)
 	{
 		FVirtualTextureUpdateSettings Settings;
 		Settings.EnableThrottling(!ViewFamily.bOverrideVirtualTextureThrottle);
@@ -2502,22 +2520,25 @@ void FDeferredShadingSceneRenderer::Render(FRDGBuilder& GraphBuilder)
 	GSystemTextures.InitializeTextures(GraphBuilder.RHICmdList, FeatureLevel);
 
 	{
-		InitViewTaskDatas.LumenFrameTemporaries = &LumenFrameTemporaries;
+		if (RendererOutput == ERendererOutput::FinalSceneColor)
+		{
+			InitViewTaskDatas.LumenFrameTemporaries = &LumenFrameTemporaries;
 
-		// Important that this uses consistent logic throughout the frame, so evaluate once and pass in the flag from here
-		// NOTE: Must be done after  system texture initialization
-		// TODO: This doesn't take into account the potential for split screen views with separate shadow caches
-		const bool bEnableVirtualShadowMaps = UseVirtualShadowMaps(ShaderPlatform, FeatureLevel) && ViewFamily.EngineShowFlags.DynamicShadows;
-		VirtualShadowMapArray.Initialize(GraphBuilder, Scene->GetVirtualShadowMapCache(Views[0]), bEnableVirtualShadowMaps, Views[0].bIsSceneCapture);
+			// Important that this uses consistent logic throughout the frame, so evaluate once and pass in the flag from here
+			// NOTE: Must be done after  system texture initialization
+			// TODO: This doesn't take into account the potential for split screen views with separate shadow caches
+			const bool bEnableVirtualShadowMaps = UseVirtualShadowMaps(ShaderPlatform, FeatureLevel) && ViewFamily.EngineShowFlags.DynamicShadows;
+			VirtualShadowMapArray.Initialize(GraphBuilder, Scene->GetVirtualShadowMapCache(Views[0]), bEnableVirtualShadowMaps, Views[0].bIsSceneCapture);
 
-		BeginInitDynamicShadows(InitViewTaskDatas);
+			BeginInitDynamicShadows(InitViewTaskDatas);
 
-		if (InitViewTaskDatas.LumenFrameTemporaries)
-		{
-			BeginUpdateLumenSceneTasks(GraphBuilder, *InitViewTaskDatas.LumenFrameTemporaries);
-		}
+			if (InitViewTaskDatas.LumenFrameTemporaries)
+			{
+				BeginUpdateLumenSceneTasks(GraphBuilder, *InitViewTaskDatas.LumenFrameTemporaries);
+			}
 
-		BeginGatherLumenLights(InitViewTaskDatas.LumenDirectLighting, InitViewTaskDatas.VisibilityTaskData);
+			BeginGatherLumenLights(InitViewTaskDatas.LumenDirectLighting, InitViewTaskDatas.VisibilityTaskData);
+		}
 
 		if (bNaniteEnabled)
 		{
@@ -2546,98 +2567,101 @@ void FDeferredShadingSceneRenderer::Render(FRDGBuilder& GraphBuilder)
 		}
 	}
 
-	PrepareDistanceFieldScene(GraphBuilder, ExternalAccessQueue);
+	if (RendererOutput == ERendererOutput::FinalSceneColor)
+	{
+		PrepareDistanceFieldScene(GraphBuilder, ExternalAccessQueue);
 
-	ShaderPrint::BeginViews(GraphBuilder, Views);
+		ShaderPrint::BeginViews(GraphBuilder, Views);
 
-	ON_SCOPE_EXIT
-	{
-		ShaderPrint::EndViews(Views);
-	};
+		ON_SCOPE_EXIT
+		{
+			ShaderPrint::EndViews(Views);
+		};
 
-	for (int32 ViewIndex = 0; ViewIndex < Views.Num(); ViewIndex++)
-	{
-		FViewInfo& View = Views[ViewIndex];
-		RDG_GPU_MASK_SCOPE(GraphBuilder, View.GPUMask);
+		for (int32 ViewIndex = 0; ViewIndex < Views.Num(); ViewIndex++)
+		{
+			FViewInfo& View = Views[ViewIndex];
+			RDG_GPU_MASK_SCOPE(GraphBuilder, View.GPUMask);
 
-		ShadingEnergyConservation::Init(GraphBuilder, View);
+			ShadingEnergyConservation::Init(GraphBuilder, View);
 
-		FGlintShadingLUTsStateData::Init(GraphBuilder, View);
-	}
+			FGlintShadingLUTsStateData::Init(GraphBuilder, View);
+		}
 
-	// kick off dependent scene updates 
-	Scene->ShadowScene->UpdateForRenderedFrame(GraphBuilder);
+		// kick off dependent scene updates 
+		Scene->ShadowScene->UpdateForRenderedFrame(GraphBuilder);
 
 #if RHI_RAYTRACING
-	// Initialize ray tracing flags, in case they weren't initialized in the CreateSceneRenderers code path
-	InitializeRayTracingFlags_RenderThread();
+		// Initialize ray tracing flags, in case they weren't initialized in the CreateSceneRenderers code path
+		InitializeRayTracingFlags_RenderThread();
 
-	GRayTracingGeometryManager.Tick(bHasRayTracingEnableChanged);
+		GRayTracingGeometryManager.Tick(bHasRayTracingEnableChanged);
 
-	if ((GetRayTracingMode() == ERayTracingMode::Dynamic) && bHasRayTracingEnableChanged)
-	{
-		Scene->GetRayTracingDynamicGeometryCollection()->Clear();
-	}
+		if ((GetRayTracingMode() == ERayTracingMode::Dynamic) && bHasRayTracingEnableChanged)
+		{
+			Scene->GetRayTracingDynamicGeometryCollection()->Clear();
+		}
 
-	// Now that we have updated all the PrimitiveSceneInfos, update the RayTracing mesh commands cache if needed
-	{
-		const ERayTracingMeshCommandsMode CurrentMode = ViewFamily.EngineShowFlags.PathTracing ? ERayTracingMeshCommandsMode::PATH_TRACING : ERayTracingMeshCommandsMode::RAY_TRACING;
-		bool bNaniteCoarseMeshStreamingModeChanged = false;
+		// Now that we have updated all the PrimitiveSceneInfos, update the RayTracing mesh commands cache if needed
+		{
+			const ERayTracingMeshCommandsMode CurrentMode = ViewFamily.EngineShowFlags.PathTracing ? ERayTracingMeshCommandsMode::PATH_TRACING : ERayTracingMeshCommandsMode::RAY_TRACING;
+			bool bNaniteCoarseMeshStreamingModeChanged = false;
 #if WITH_EDITOR
-		bNaniteCoarseMeshStreamingModeChanged = Nanite::FCoarseMeshStreamingManager::CheckStreamingMode();
+			bNaniteCoarseMeshStreamingModeChanged = Nanite::FCoarseMeshStreamingManager::CheckStreamingMode();
 #endif // WITH_EDITOR
-		const bool bNaniteRayTracingModeChanged = Nanite::GRayTracingManager.CheckModeChanged();
+			const bool bNaniteRayTracingModeChanged = Nanite::GRayTracingManager.CheckModeChanged();
 
-		if (CurrentMode != Scene->CachedRayTracingMeshCommandsMode || bNaniteCoarseMeshStreamingModeChanged || bNaniteRayTracingModeChanged || bHasRayTracingEnableChanged)
-		{
-			Scene->WaitForCacheRayTracingPrimitivesTask();
+			if (CurrentMode != Scene->CachedRayTracingMeshCommandsMode || bNaniteCoarseMeshStreamingModeChanged || bNaniteRayTracingModeChanged || bHasRayTracingEnableChanged)
+			{
+				Scene->WaitForCacheRayTracingPrimitivesTask();
 
-			// In some situations, we need to refresh the cached ray tracing mesh commands because they contain data about the currently bound shader. 
-			// This operation is a bit expensive but only happens once as we transition between modes which should be rare.
-			Scene->CachedRayTracingMeshCommandsMode = CurrentMode;
-			Scene->RefreshRayTracingMeshCommandCache();
-			bHasRayTracingEnableChanged = false;
-		}
+				// In some situations, we need to refresh the cached ray tracing mesh commands because they contain data about the currently bound shader. 
+				// This operation is a bit expensive but only happens once as we transition between modes which should be rare.
+				Scene->CachedRayTracingMeshCommandsMode = CurrentMode;
+				Scene->RefreshRayTracingMeshCommandCache();
+				bHasRayTracingEnableChanged = false;
+			}
 
-		if (bRefreshRayTracingInstances)
-		{
-			Scene->WaitForCacheRayTracingPrimitivesTask();
+			if (bRefreshRayTracingInstances)
+			{
+				Scene->WaitForCacheRayTracingPrimitivesTask();
 
-			// In some situations, we need to refresh the cached ray tracing instance.
-			// eg: Need to update PrimitiveRayTracingFlags
-			// This operation is a bit expensive but only happens once as we transition between modes which should be rare.
-			Scene->RefreshRayTracingInstances();
-			bRefreshRayTracingInstances = false;
-		}
+				// In some situations, we need to refresh the cached ray tracing instance.
+				// eg: Need to update PrimitiveRayTracingFlags
+				// This operation is a bit expensive but only happens once as we transition between modes which should be rare.
+				Scene->RefreshRayTracingInstances();
+				bRefreshRayTracingInstances = false;
+			}
 
-		if (bNaniteRayTracingModeChanged)
-		{
-			for (FViewInfo& View : Views)
+			if (bNaniteRayTracingModeChanged)
 			{
-				if (View.ViewState != nullptr && !View.bIsOfflineRender)
+				for (FViewInfo& View : Views)
 				{
-					// don't invalidate in the offline case because we only get one attempt at rendering each sample
-					View.ViewState->PathTracingInvalidate();
+					if (View.ViewState != nullptr && !View.bIsOfflineRender)
+					{
+						// don't invalidate in the offline case because we only get one attempt at rendering each sample
+						View.ViewState->PathTracingInvalidate();
+					}
 				}
 			}
-		}
 
-		if (bAnyRayTracingPassEnabled)
-		{
-			const int32 ReferenceViewIndex = 0;
-			FViewInfo& ReferenceView = Views[ReferenceViewIndex];
-			FGraphEventRef PrereqTask = CreateCompatibilityGraphEvent(MakeArrayView({ Scene->GetCacheRayTracingPrimitivesTask(), InitViewTaskDatas.VisibilityTaskData->GetFrustumCullTask() }));
-
-			InitViewTaskDatas.RayTracingRelevantPrimitives = Allocator.Create<FRayTracingRelevantPrimitiveTaskData>();
-			InitViewTaskDatas.RayTracingRelevantPrimitives->Task = FFunctionGraphTask::CreateAndDispatchWhenReady(
-				[Scene = Scene, &ReferenceView, &RayTracingRelevantPrimitiveList = InitViewTaskDatas.RayTracingRelevantPrimitives->List]()
+			if (bAnyRayTracingPassEnabled)
 			{
-				FTaskTagScope TaskTagScope(ETaskTag::EParallelRenderingThread);
-				GatherRayTracingRelevantPrimitives(*Scene, ReferenceView, RayTracingRelevantPrimitiveList);
-			}, TStatId(), PrereqTask, ENamedThreads::AnyNormalThreadHiPriTask);
+				const int32 ReferenceViewIndex = 0;
+				FViewInfo& ReferenceView = Views[ReferenceViewIndex];
+				FGraphEventRef PrereqTask = CreateCompatibilityGraphEvent(MakeArrayView({ Scene->GetCacheRayTracingPrimitivesTask(), InitViewTaskDatas.VisibilityTaskData->GetFrustumCullTask() }));
+
+				InitViewTaskDatas.RayTracingRelevantPrimitives = Allocator.Create<FRayTracingRelevantPrimitiveTaskData>();
+				InitViewTaskDatas.RayTracingRelevantPrimitives->Task = FFunctionGraphTask::CreateAndDispatchWhenReady(
+					[Scene = Scene, &ReferenceView, &RayTracingRelevantPrimitiveList = InitViewTaskDatas.RayTracingRelevantPrimitives->List]()
+					{
+						FTaskTagScope TaskTagScope(ETaskTag::EParallelRenderingThread);
+						GatherRayTracingRelevantPrimitives(*Scene, ReferenceView, RayTracingRelevantPrimitiveList);
+					}, TStatId(), PrereqTask, ENamedThreads::AnyNormalThreadHiPriTask);
+			}
 		}
-	}
 #endif
+	}
 
 	bool bUpdateNaniteStreaming = false;
 	bool bVisualizeNanite = false;
@@ -2669,20 +2693,23 @@ void FDeferredShadingSceneRenderer::Render(FRDGBuilder& GraphBuilder)
 
 	CSV_SCOPED_TIMING_STAT_EXCLUSIVE(RenderOther);
 
-	const bool bPathTracedAtmosphere = ViewFamily.EngineShowFlags.PathTracing && Views.Num() > 0 && Views[0].FinalPostProcessSettings.PathTracingEnableReferenceAtmosphere;
-	if (ShouldRenderSkyAtmosphere(Scene, ViewFamily.EngineShowFlags) && !bPathTracedAtmosphere)
+	if (RendererOutput == ERendererOutput::FinalSceneColor)
 	{
-		for (int32 LightIndex = 0; LightIndex < NUM_ATMOSPHERE_LIGHTS; ++LightIndex)
+		const bool bPathTracedAtmosphere = ViewFamily.EngineShowFlags.PathTracing && Views.Num() > 0 && Views[0].FinalPostProcessSettings.PathTracingEnableReferenceAtmosphere;
+		if (ShouldRenderSkyAtmosphere(Scene, ViewFamily.EngineShowFlags) && !bPathTracedAtmosphere)
 		{
-			if (Scene->AtmosphereLights[LightIndex])
+			for (int32 LightIndex = 0; LightIndex < NUM_ATMOSPHERE_LIGHTS; ++LightIndex)
 			{
-				PrepareSunLightProxy(*Scene->GetSkyAtmosphereSceneInfo(),LightIndex, *Scene->AtmosphereLights[LightIndex]);
+				if (Scene->AtmosphereLights[LightIndex])
+				{
+					PrepareSunLightProxy(*Scene->GetSkyAtmosphereSceneInfo(),LightIndex, *Scene->AtmosphereLights[LightIndex]);
+				}
 			}
 		}
-	}
-	else
-	{
-		Scene->ResetAtmosphereLightsProperties();
+		else
+		{
+			Scene->ResetAtmosphereLightsProperties();
+		}
 	}
 
 	SCOPED_NAMED_EVENT(FDeferredShadingSceneRenderer_Render, FColor::Emerald);
@@ -2696,6 +2723,7 @@ void FDeferredShadingSceneRenderer::Render(FRDGBuilder& GraphBuilder)
 	RDG_EVENT_SCOPE(GraphBuilder, "Scene");
 	RDG_GPU_STAT_SCOPE_VERBOSE(GraphBuilder, Unaccounted, *ViewFamily.ProfileDescription);
 	
+	if (RendererOutput == ERendererOutput::FinalSceneColor)
 	{
 		SCOPE_CYCLE_COUNTER(STAT_FDeferredShadingSceneRenderer_Render_Init);
 		RDG_RHI_GPU_STAT_SCOPE(GraphBuilder, AllocateRendertargets);
@@ -2766,88 +2794,92 @@ void FDeferredShadingSceneRenderer::Render(FRDGBuilder& GraphBuilder)
 	}
 
 #if RHI_RAYTRACING
-
 	const int32 ReferenceViewIndex = 0;
 	FViewInfo& ReferenceView = Views[ReferenceViewIndex];
-
-	// Prepare the scene for rendering this frame.
 	FRayTracingScene& RayTracingScene = Scene->RayTracingScene;
-	RayTracingScene.Reset(IsRayTracingInstanceDebugDataEnabled(ReferenceView)); // Resets the internal arrays, but does not release any resources.
+#endif
 
-	if (ShouldPrepareRayTracingDecals(*Scene, ViewFamily))
+	if (RendererOutput == ERendererOutput::FinalSceneColor)
 	{
-		// Calculate decal grid for ray tracing per view since decal fade is view dependent
-		// TODO: investigate reusing the same grid for all views (ie: different callable shader SBT entries for each view so fade alpha is still correct for each view)
+#if RHI_RAYTRACING
+		// Prepare the scene for rendering this frame.
+		RayTracingScene.Reset(IsRayTracingInstanceDebugDataEnabled(ReferenceView)); // Resets the internal arrays, but does not release any resources.
 
-		for (int32 ViewIndex = 0; ViewIndex < Views.Num(); ViewIndex++)
+		if (ShouldPrepareRayTracingDecals(*Scene, ViewFamily))
 		{
-			FViewInfo& View = Views[ViewIndex];
-			View.RayTracingDecalUniformBuffer = CreateRayTracingDecalData(GraphBuilder, *Scene, View, RayTracingScene.NumCallableShaderSlots);
-			View.bHasRayTracingDecals = true;
-			RayTracingScene.NumCallableShaderSlots += Scene->Decals.Num();
-		}
-	}
-	else
-	{
-		TRDGUniformBufferRef<FRayTracingDecals> NullRayTracingDecalUniformBuffer = CreateNullRayTracingDecalsUniformBuffer(GraphBuilder);
+			// Calculate decal grid for ray tracing per view since decal fade is view dependent
+			// TODO: investigate reusing the same grid for all views (ie: different callable shader SBT entries for each view so fade alpha is still correct for each view)
 
-		for (int32 ViewIndex = 0; ViewIndex < Views.Num(); ViewIndex++)
-		{
-			FViewInfo& View = Views[ViewIndex];
-			View.RayTracingDecalUniformBuffer = NullRayTracingDecalUniformBuffer;
-			View.bHasRayTracingDecals = false;
+			for (int32 ViewIndex = 0; ViewIndex < Views.Num(); ViewIndex++)
+			{
+				FViewInfo& View = Views[ViewIndex];
+				View.RayTracingDecalUniformBuffer = CreateRayTracingDecalData(GraphBuilder, *Scene, View, RayTracingScene.NumCallableShaderSlots);
+				View.bHasRayTracingDecals = true;
+				RayTracingScene.NumCallableShaderSlots += Scene->Decals.Num();
+			}
 		}
-	}
-
-	if (IsRayTracingEnabled() && RHISupportsRayTracingShaders(ViewFamily.GetShaderPlatform()))
-	{
-		// Nanite raytracing manager update must run before GPUScene update since it can modify primitive data
-		Nanite::GRayTracingManager.Update();
-
-		if (!ViewFamily.EngineShowFlags.PathTracing)
+		else
 		{
-			// get the default lighting miss shader (to implicitly fill in the MissShader library before the RT pipeline is created)
-			GetRayTracingLightingMissShader(ReferenceView.ShaderMap);
-			RayTracingScene.NumMissShaderSlots++;
-		}
+			TRDGUniformBufferRef<FRayTracingDecals> NullRayTracingDecalUniformBuffer = CreateNullRayTracingDecalsUniformBuffer(GraphBuilder);
 
-		if (ViewFamily.EngineShowFlags.LightFunctions)
-		{
-			// gather all the light functions that may be used (and also count how many miss shaders we will need)
-			FRayTracingLightFunctionMap RayTracingLightFunctionMap;
-			if (ViewFamily.EngineShowFlags.PathTracing)
+			for (int32 ViewIndex = 0; ViewIndex < Views.Num(); ViewIndex++)
 			{
-				RayTracingLightFunctionMap = GatherLightFunctionLightsPathTracing(Scene, ViewFamily.EngineShowFlags, ReferenceView.GetFeatureLevel());
+				FViewInfo& View = Views[ViewIndex];
+				View.RayTracingDecalUniformBuffer = NullRayTracingDecalUniformBuffer;
+				View.bHasRayTracingDecals = false;
 			}
-			else
+		}
+
+		if (IsRayTracingEnabled() && RHISupportsRayTracingShaders(ViewFamily.GetShaderPlatform()))
+		{
+			// Nanite raytracing manager update must run before GPUScene update since it can modify primitive data
+			Nanite::GRayTracingManager.Update();
+
+			if (!ViewFamily.EngineShowFlags.PathTracing)
 			{
-				RayTracingLightFunctionMap = GatherLightFunctionLights(Scene, ViewFamily.EngineShowFlags, ReferenceView.GetFeatureLevel());
+				// get the default lighting miss shader (to implicitly fill in the MissShader library before the RT pipeline is created)
+				GetRayTracingLightingMissShader(ReferenceView.ShaderMap);
+				RayTracingScene.NumMissShaderSlots++;
 			}
-			if (!RayTracingLightFunctionMap.IsEmpty())
+
+			if (ViewFamily.EngineShowFlags.LightFunctions)
 			{
-				// If we got some light functions in our map, store them in the RDG blackboard so downstream functions can use them.
-				// The map itself will be strictly read-only from this point on.
-				GraphBuilder.Blackboard.Create<FRayTracingLightFunctionMap>(MoveTemp(RayTracingLightFunctionMap));
+				// gather all the light functions that may be used (and also count how many miss shaders we will need)
+				FRayTracingLightFunctionMap RayTracingLightFunctionMap;
+				if (ViewFamily.EngineShowFlags.PathTracing)
+				{
+					RayTracingLightFunctionMap = GatherLightFunctionLightsPathTracing(Scene, ViewFamily.EngineShowFlags, ReferenceView.GetFeatureLevel());
+				}
+				else
+				{
+					RayTracingLightFunctionMap = GatherLightFunctionLights(Scene, ViewFamily.EngineShowFlags, ReferenceView.GetFeatureLevel());
+				}
+				if (!RayTracingLightFunctionMap.IsEmpty())
+				{
+					// If we got some light functions in our map, store them in the RDG blackboard so downstream functions can use them.
+					// The map itself will be strictly read-only from this point on.
+					GraphBuilder.Blackboard.Create<FRayTracingLightFunctionMap>(MoveTemp(RayTracingLightFunctionMap));
+				}
 			}
 		}
-	}
 #endif // RHI_RAYTRACING
 
-	// Notify the FX system that the scene is about to be rendered.
-	if (FXSystem && Views.IsValidIndex(0))
-	{
-		SCOPE_CYCLE_COUNTER(STAT_FDeferredShadingSceneRenderer_FXSystem_PreRender);
-		GraphBuilder.SetCommandListStat(GET_STATID(STAT_CLM_FXPreRender));
-		FXSystem->PreRender(GraphBuilder, GetSceneViews(), bIsFirstSceneRenderer /*bAllowGPUParticleUpdate*/);
-		if (FGPUSortManager* GPUSortManager = FXSystem->GetGPUSortManager())
+		// Notify the FX system that the scene is about to be rendered.
+		if (FXSystem && Views.IsValidIndex(0))
 		{
-			GPUSortManager->OnPreRender(GraphBuilder);
+			SCOPE_CYCLE_COUNTER(STAT_FDeferredShadingSceneRenderer_FXSystem_PreRender);
+			GraphBuilder.SetCommandListStat(GET_STATID(STAT_CLM_FXPreRender));
+			FXSystem->PreRender(GraphBuilder, GetSceneViews(), bIsFirstSceneRenderer /*bAllowGPUParticleUpdate*/);
+			if (FGPUSortManager* GPUSortManager = FXSystem->GetGPUSortManager())
+			{
+				GPUSortManager->OnPreRender(GraphBuilder);
+			}
 		}
-	}
 
 #if !(UE_BUILD_SHIPPING || UE_BUILD_TEST)
-	Scene->DebugRender(Views);
+		Scene->DebugRender(Views);
 #endif
+	}
 
 	{
 		RDG_CSV_STAT_EXCLUSIVE_SCOPE(GraphBuilder, UpdateGPUScene);
@@ -2882,18 +2914,8 @@ void FDeferredShadingSceneRenderer::Render(FRDGBuilder& GraphBuilder)
 	FSceneTextures::InitializeViewFamily(GraphBuilder, ViewFamily);
 	FSceneTextures& SceneTextures = GetActiveSceneTextures();
 
-	if (bUseVirtualTexturing)
-	{
-		// Note, should happen after the GPU-Scene update to ensure rendering to runtime virtual textures is using the correctly updated scene
-		FVirtualTextureSystem::Get().EndUpdate(GraphBuilder, MoveTemp(VirtualTextureUpdater), FeatureLevel);
-	}
-
-#if RHI_RAYTRACING
-	GatherRayTracingWorldInstancesForView(GraphBuilder, ReferenceView, RayTracingScene, InitViewTaskDatas.RayTracingRelevantPrimitives);
-#endif // RHI_RAYTRACING
-
 	const bool bUseGBuffer = IsUsingGBuffers(ShaderPlatform);
-	
+
 	const bool bRenderDeferredLighting = ViewFamily.EngineShowFlags.Lighting
 		&& FeatureLevel >= ERHIFeatureLevel::SM5
 		&& ViewFamily.EngineShowFlags.DeferredLighting
@@ -2901,33 +2923,47 @@ void FDeferredShadingSceneRenderer::Render(FRDGBuilder& GraphBuilder)
 		&& !bHasRayTracedOverlay;
 
 	bool bComputeLightGrid = false;
-	bool bAnyLumenEnabled = false;
 
+	if (RendererOutput == ERendererOutput::FinalSceneColor)
 	{
-		if (bUseGBuffer)
-		{
-			bComputeLightGrid = bRenderDeferredLighting;
-		}
-		else
+		if (bUseVirtualTexturing)
 		{
-			bComputeLightGrid = ViewFamily.EngineShowFlags.Lighting;
+			// Note, should happen after the GPU-Scene update to ensure rendering to runtime virtual textures is using the correctly updated scene
+			FVirtualTextureSystem::Get().EndUpdate(GraphBuilder, MoveTemp(VirtualTextureUpdater), FeatureLevel);
 		}
 
-		for (int32 ViewIndex = 0; ViewIndex < Views.Num(); ViewIndex++)
+#if RHI_RAYTRACING
+		GatherRayTracingWorldInstancesForView(GraphBuilder, ReferenceView, RayTracingScene, InitViewTaskDatas.RayTracingRelevantPrimitives);
+#endif // RHI_RAYTRACING
+
+		bool bAnyLumenEnabled = false;
+
 		{
-			FViewInfo& View = Views[ViewIndex];
-			bAnyLumenEnabled = bAnyLumenEnabled 
-				|| GetViewPipelineState(View).DiffuseIndirectMethod == EDiffuseIndirectMethod::Lumen
-				|| GetViewPipelineState(View).ReflectionsMethod == EReflectionsMethod::Lumen;
-		}
+			if (bUseGBuffer)
+			{
+				bComputeLightGrid = bRenderDeferredLighting;
+			}
+			else
+			{
+				bComputeLightGrid = ViewFamily.EngineShowFlags.Lighting;
+			}
+
+			for (int32 ViewIndex = 0; ViewIndex < Views.Num(); ViewIndex++)
+			{
+				FViewInfo& View = Views[ViewIndex];
+				bAnyLumenEnabled = bAnyLumenEnabled
+					|| GetViewPipelineState(View).DiffuseIndirectMethod == EDiffuseIndirectMethod::Lumen
+					|| GetViewPipelineState(View).ReflectionsMethod == EReflectionsMethod::Lumen;
+			}
 
-		bComputeLightGrid |= (
-			ShouldRenderVolumetricFog() ||
-			VolumetricCloudWantsToSampleLocalLights(Scene, ViewFamily.EngineShowFlags) ||
-			ViewFamily.ViewMode != VMI_Lit ||
-			bAnyLumenEnabled ||
-			VirtualShadowMapArray.IsEnabled() ||
-			ShouldVisualizeLightGrid());
+			bComputeLightGrid |= (
+				ShouldRenderVolumetricFog() ||
+				VolumetricCloudWantsToSampleLocalLights(Scene, ViewFamily.EngineShowFlags) ||
+				ViewFamily.ViewMode != VMI_Lit ||
+				bAnyLumenEnabled ||
+				VirtualShadowMapArray.IsEnabled() ||
+				ShouldVisualizeLightGrid());
+		}
 	}
 
 	// force using occ queries for wireframe if rendering is parented or frozen in the first view
@@ -2966,7 +3002,10 @@ void FDeferredShadingSceneRenderer::Render(FRDGBuilder& GraphBuilder)
 
 			const TSet<uint32> ModifiedResources = Nanite::GStreamingManager.GetAndClearModifiedResources();
 #if RHI_RAYTRACING
-			Nanite::GRayTracingManager.RequestUpdates(ModifiedResources);
+			if (RendererOutput == ERendererOutput::FinalSceneColor)
+			{
+				Nanite::GRayTracingManager.RequestUpdates(ModifiedResources);
+			}
 #endif
 		}
 	}
@@ -2987,7 +3026,7 @@ void FDeferredShadingSceneRenderer::Render(FRDGBuilder& GraphBuilder)
 	}
 
 	FHairStrandsBookmarkParameters& HairStrandsBookmarkParameters = *GraphBuilder.AllocObject<FHairStrandsBookmarkParameters>();
-	if (IsHairStrandsEnabled(EHairStrandsShaderType::All, Scene->GetShaderPlatform()))
+	if (IsHairStrandsEnabled(EHairStrandsShaderType::All, Scene->GetShaderPlatform()) && RendererOutput == ERendererOutput::FinalSceneColor)
 	{
 		HairStrandsBookmarkParameters = CreateHairStrandsBookmarkParameters(Scene, Views[0]);
 		RunHairStrandsBookmark(GraphBuilder, EHairStrandsBookmark::ProcessTasks, HairStrandsBookmarkParameters);
@@ -3050,7 +3089,7 @@ void FDeferredShadingSceneRenderer::Render(FRDGBuilder& GraphBuilder)
 		GraphBuilder.SetCommandListStat(GET_STATID(STAT_CLM_AfterPrePass));
 
 		// special pass for DDM_AllOpaqueNoVelocity, which uses the velocity pass to finish the early depth pass write
-		if (bShouldRenderVelocities && Scene->EarlyZPassMode == DDM_AllOpaqueNoVelocity)
+		if (bShouldRenderVelocities && Scene->EarlyZPassMode == DDM_AllOpaqueNoVelocity && RendererOutput == ERendererOutput::FinalSceneColor)
 		{
 			// Render the velocities of movable objects
 			GraphBuilder.SetCommandListStat(GET_STATID(STAT_CLM_Velocity));
@@ -3295,47 +3334,13 @@ void FDeferredShadingSceneRenderer::Render(FRDGBuilder& GraphBuilder)
 
 	AddResolveSceneDepthPass(GraphBuilder, Views, SceneTextures.Depth);
 
-	GVRSImageManager.PrepareImageBasedVRS(GraphBuilder, ViewFamily, SceneTextures);
-
-	if (!IsForwardShadingEnabled(ShaderPlatform))
-	{
-		// Dynamic shadows are synced later when using the deferred path to make more headroom for tasks.
-		FinishInitDynamicShadows(GraphBuilder, InitViewTaskDatas.DynamicShadows, InstanceCullingManager, ExternalAccessQueue);
-	}
-
 	FComputeLightGridOutput ComputeLightGridOutput = {};
 
-	// NOTE: The ordering of the lights is used to select sub-sets for different purposes, e.g., those that support clustered deferred.
-	FSortedLightSetSceneInfo& SortedLightSet = *GraphBuilder.AllocObject<FSortedLightSetSceneInfo>();
-	{
-		RDG_CSV_STAT_EXCLUSIVE_SCOPE(GraphBuilder, SortLights);
-		RDG_GPU_STAT_SCOPE(GraphBuilder, SortLights);
-		ComputeLightGridOutput = GatherLightsAndComputeLightGrid(GraphBuilder, bComputeLightGrid, SortedLightSet);
-	}
-
-	CSV_CUSTOM_STAT(LightCount, All,  float(SortedLightSet.SortedLights.Num()), ECsvCustomStatOp::Set);
-	CSV_CUSTOM_STAT(LightCount, Batched, float(SortedLightSet.UnbatchedLightStart), ECsvCustomStatOp::Set);
-	CSV_CUSTOM_STAT(LightCount, Unbatched, float(SortedLightSet.SortedLights.Num()) - float(SortedLightSet.UnbatchedLightStart), ECsvCustomStatOp::Set);
-
-	FCompositionLighting CompositionLighting(Views, SceneTextures, [this] (int32 ViewIndex)
+	FCompositionLighting CompositionLighting(Views, SceneTextures, [this](int32 ViewIndex)
 	{
 		return GetViewPipelineState(Views[ViewIndex]).AmbientOcclusionMethod == EAmbientOcclusionMethod::SSAO;
 	});
 
-	const bool bShouldRenderSkyAtmosphere = ShouldRenderSkyAtmosphere(Scene, ViewFamily.EngineShowFlags);
-	const ESkyAtmospherePassLocation SkyAtmospherePassLocation = GetSkyAtmospherePassLocation();
-	const bool bShouldRenderVolumetricCloudBase = ShouldRenderVolumetricCloud(Scene, ViewFamily.EngineShowFlags);
-	const bool bShouldRenderVolumetricCloud = bShouldRenderVolumetricCloudBase && (!ViewFamily.EngineShowFlags.VisualizeVolumetricCloudConservativeDensity && !ViewFamily.EngineShowFlags.VisualizeVolumetricCloudEmptySpaceSkipping);
-	const bool bShouldVisualizeVolumetricCloud = bShouldRenderVolumetricCloudBase && (!!ViewFamily.EngineShowFlags.VisualizeVolumetricCloudConservativeDensity || !!ViewFamily.EngineShowFlags.VisualizeVolumetricCloudEmptySpaceSkipping);
-	bool bAsyncComputeVolumetricCloud = IsVolumetricRenderTargetEnabled() && IsVolumetricRenderTargetAsyncCompute();
-	bool bVolumetricRenderTargetRequired = bShouldRenderVolumetricCloud && !bHasRayTracedOverlay;
-
-	if (SkyAtmospherePassLocation == ESkyAtmospherePassLocation::BeforeOcclusion && bShouldRenderSkyAtmosphere)
-	{
-		// Generate the Sky/Atmosphere look up tables
-		RenderSkyAtmosphereLookUpTables(GraphBuilder, ExternalAccessQueue);
-	}
-
 	const auto RenderOcclusionLambda = [&]()
 	{
 		const int32 AsyncComputeMode = CVarSceneDepthHZBAsyncCompute.GetValueOnRenderThread();
@@ -3347,1057 +3352,1104 @@ void FDeferredShadingSceneRenderer::Render(FRDGBuilder& GraphBuilder)
 			AsyncComputeParams.Prerequisite = ComputeLightGridOutput.CompactLinksPass;
 		}
 
-		RenderOcclusion(GraphBuilder, SceneTextures, bIsOcclusionTesting, 
+		RenderOcclusion(GraphBuilder, SceneTextures, bIsOcclusionTesting,
 			bAsyncCompute ? &AsyncComputeParams : nullptr);
 
 		CompositionLighting.ProcessAfterOcclusion(GraphBuilder);
 	};
 
-	// Early occlusion queries
-	const bool bOcclusionBeforeBasePass = ((DepthPass.EarlyZPassMode == EDepthDrawingMode::DDM_AllOccluders) || bIsEarlyDepthComplete);
+	const bool bShouldRenderSkyAtmosphere = ShouldRenderSkyAtmosphere(Scene, ViewFamily.EngineShowFlags);
+	const ESkyAtmospherePassLocation SkyAtmospherePassLocation = GetSkyAtmospherePassLocation();
+	const bool bShouldRenderVolumetricCloudBase = ShouldRenderVolumetricCloud(Scene, ViewFamily.EngineShowFlags);
+	const bool bShouldRenderVolumetricCloud = bShouldRenderVolumetricCloudBase && (!ViewFamily.EngineShowFlags.VisualizeVolumetricCloudConservativeDensity && !ViewFamily.EngineShowFlags.VisualizeVolumetricCloudEmptySpaceSkipping);
+	const bool bShouldVisualizeVolumetricCloud = bShouldRenderVolumetricCloudBase && (!!ViewFamily.EngineShowFlags.VisualizeVolumetricCloudConservativeDensity || !!ViewFamily.EngineShowFlags.VisualizeVolumetricCloudEmptySpaceSkipping);
+	bool bAsyncComputeVolumetricCloud = IsVolumetricRenderTargetEnabled() && IsVolumetricRenderTargetAsyncCompute();
+	bool bVolumetricRenderTargetRequired = bShouldRenderVolumetricCloud && !bHasRayTracedOverlay;
 
-	if (bOcclusionBeforeBasePass)
+	FRDGTextureRef ViewFamilyTexture = TryCreateViewFamilyTexture(GraphBuilder, ViewFamily);
+	if (RendererOutput == ERendererOutput::DepthPrepassOnly)
 	{
 		RenderOcclusionLambda();
-	}
 
-	// End early occlusion queries
-
-	BeginAsyncDistanceFieldShadowProjections(GraphBuilder, SceneTextures, InitViewTaskDatas.DynamicShadows);
+		if (bUpdateNaniteStreaming)
+		{
+			Nanite::GStreamingManager.SubmitFrameStreamingRequests(GraphBuilder);
+		}
 
-	if (bShouldRenderVolumetricCloudBase)
-	{
-		InitVolumetricRenderTargetForViews(GraphBuilder, Views);
+		CopySceneCaptureComponentToTarget(GraphBuilder, SceneTextures, ViewFamilyTexture, ViewFamily, Views);
 	}
 	else
 	{
-		ResetVolumetricRenderTargetForViews(GraphBuilder, Views);
-	}
-
-	InitVolumetricCloudsForViews(GraphBuilder, bShouldRenderVolumetricCloudBase, InstanceCullingManager);
-
-	// Generate sky LUTs
-	// TODO: Valid shadow maps (for volumetric light shafts) have not yet been generated at this point in the frame. Need to resolve dependency ordering!
-	// This also must happen before the BasePass for Sky material to be able to sample valid LUTs.
-	if (SkyAtmospherePassLocation == ESkyAtmospherePassLocation::BeforeBasePass && bShouldRenderSkyAtmosphere)
-	{
-		// Generate the Sky/Atmosphere look up tables
-		RenderSkyAtmosphereLookUpTables(GraphBuilder, ExternalAccessQueue);
+		GVRSImageManager.PrepareImageBasedVRS(GraphBuilder, ViewFamily, SceneTextures);
 
-		// Sky env map capture uses the view UB, which contains the LUTs computed above. We need to transition them to readable now.
-		ExternalAccessQueue.Submit(GraphBuilder);
-	}
+		if (!IsForwardShadingEnabled(ShaderPlatform))
+		{
+			// Dynamic shadows are synced later when using the deferred path to make more headroom for tasks.
+			FinishInitDynamicShadows(GraphBuilder, InitViewTaskDatas.DynamicShadows, InstanceCullingManager, ExternalAccessQueue);
+		}
+	
+		// NOTE: The ordering of the lights is used to select sub-sets for different purposes, e.g., those that support clustered deferred.
+		FSortedLightSetSceneInfo& SortedLightSet = *GraphBuilder.AllocObject<FSortedLightSetSceneInfo>();
+		{
+			RDG_CSV_STAT_EXCLUSIVE_SCOPE(GraphBuilder, SortLights);
+			RDG_GPU_STAT_SCOPE(GraphBuilder, SortLights);
+			ComputeLightGridOutput = GatherLightsAndComputeLightGrid(GraphBuilder, bComputeLightGrid, SortedLightSet);
+		}
 
-	// Capture the SkyLight using the SkyAtmosphere and VolumetricCloud component if available.
-	const bool bRealTimeSkyCaptureEnabled = Scene->SkyLight && Scene->SkyLight->bRealTimeCaptureEnabled && Views.Num() > 0 && ViewFamily.EngineShowFlags.SkyLighting;
-	if (bRealTimeSkyCaptureEnabled)
-	{
-		FViewInfo& MainView = Views[0];
-		Scene->AllocateAndCaptureFrameSkyEnvMap(GraphBuilder, *this, MainView, bShouldRenderSkyAtmosphere, bShouldRenderVolumetricCloud, InstanceCullingManager, ExternalAccessQueue);
-	}
+		CSV_CUSTOM_STAT(LightCount, All,  float(SortedLightSet.SortedLights.Num()), ECsvCustomStatOp::Set);
+		CSV_CUSTOM_STAT(LightCount, Batched, float(SortedLightSet.UnbatchedLightStart), ECsvCustomStatOp::Set);
+		CSV_CUSTOM_STAT(LightCount, Unbatched, float(SortedLightSet.SortedLights.Num()) - float(SortedLightSet.UnbatchedLightStart), ECsvCustomStatOp::Set);
 
-	const ECustomDepthPassLocation CustomDepthPassLocation = GetCustomDepthPassLocation(ShaderPlatform);
-	if (CustomDepthPassLocation == ECustomDepthPassLocation::BeforeBasePass)
-	{
-		QUICK_SCOPE_CYCLE_COUNTER(STAT_FDeferredShadingSceneRenderer_CustomDepthPass_BeforeBasePass);
-		if (RenderCustomDepthPass(GraphBuilder, SceneTextures.CustomDepth, SceneTextures.GetSceneTextureShaderParameters(FeatureLevel), NaniteRasterResults, PrimaryNaniteViews))
+		if (SkyAtmospherePassLocation == ESkyAtmospherePassLocation::BeforeOcclusion && bShouldRenderSkyAtmosphere)
 		{
-			SceneTextures.SetupMode |= ESceneTextureSetupMode::CustomDepth;
-			SceneTextures.UniformBuffer = CreateSceneTextureUniformBuffer(GraphBuilder, &SceneTextures, FeatureLevel, SceneTextures.SetupMode);
+			// Generate the Sky/Atmosphere look up tables
+			RenderSkyAtmosphereLookUpTables(GraphBuilder, ExternalAccessQueue);
 		}
-	}
+	
+		// Early occlusion queries
+		const bool bOcclusionBeforeBasePass = ((DepthPass.EarlyZPassMode == EDepthDrawingMode::DDM_AllOccluders) || bIsEarlyDepthComplete);
 
-	// Lumen updates need access to sky atmosphere LUT.
-	ExternalAccessQueue.Submit(GraphBuilder);
+		if (bOcclusionBeforeBasePass)
+		{
+			RenderOcclusionLambda();
+		}
 
-	UpdateLumenScene(GraphBuilder, LumenFrameTemporaries);
+		// End early occlusion queries
 
-	FRDGTextureRef HalfResolutionDepthCheckerboardMinMaxTexture = nullptr;
-	FRDGTextureRef QuarterResolutionDepthMinMaxTexture = nullptr;
-	bool bQuarterResMinMaxDepthRequired = bShouldRenderVolumetricCloud && ShouldVolumetricCloudTraceWithMinMaxDepth(Views);
+		BeginAsyncDistanceFieldShadowProjections(GraphBuilder, SceneTextures, InitViewTaskDatas.DynamicShadows);
 
-	auto GenerateQuarterResDepthMinMaxTexture = [&](auto& GraphBuilder, auto& Views, auto& InputTexture)
-	{
-		if (bQuarterResMinMaxDepthRequired)
+		if (bShouldRenderVolumetricCloudBase)
 		{
-			check(InputTexture != nullptr);	// Must receive a valid texture
-			check(QuarterResolutionDepthMinMaxTexture == nullptr);	// Only generate it once
-			return CreateQuarterResolutionDepthMinAndMax(GraphBuilder, Views, InputTexture);
+			InitVolumetricRenderTargetForViews(GraphBuilder, Views);
+		}
+		else
+		{
+			ResetVolumetricRenderTargetForViews(GraphBuilder, Views);
 		}
-		return FRDGTextureRef(nullptr);
-	};
 
-	// Kick off async compute cloud early if all depth has been written in the prepass
-	if (bShouldRenderVolumetricCloud && bAsyncComputeVolumetricCloud && DepthPass.EarlyZPassMode == DDM_AllOpaque && !bHasRayTracedOverlay)
-	{
-		HalfResolutionDepthCheckerboardMinMaxTexture = CreateHalfResolutionDepthCheckerboardMinMax(GraphBuilder, Views, SceneTextures.Depth.Resolve);
-		QuarterResolutionDepthMinMaxTexture = GenerateQuarterResDepthMinMaxTexture(GraphBuilder, Views, HalfResolutionDepthCheckerboardMinMaxTexture);
+		InitVolumetricCloudsForViews(GraphBuilder, bShouldRenderVolumetricCloudBase, InstanceCullingManager);
 
-		bool bSkipVolumetricRenderTarget = false;
-		bool bSkipPerPixelTracing = true;
-		bAsyncComputeVolumetricCloud = RenderVolumetricCloud(GraphBuilder, SceneTextures, bSkipVolumetricRenderTarget, bSkipPerPixelTracing, 
-			HalfResolutionDepthCheckerboardMinMaxTexture, QuarterResolutionDepthMinMaxTexture, true, InstanceCullingManager);
-	}
-	
-	FRDGTextureRef ForwardScreenSpaceShadowMaskTexture = nullptr;
-	FRDGTextureRef ForwardScreenSpaceShadowMaskHairTexture = nullptr;
-	if (IsForwardShadingEnabled(ShaderPlatform))
-	{
-		// With forward shading we need to render shadow maps early
-		ensureMsgf(!VirtualShadowMapArray.IsEnabled(), TEXT("Virtual shadow maps are not supported in the forward shading path"));
-		RenderShadowDepthMaps(GraphBuilder, InstanceCullingManager, ExternalAccessQueue);
+		// Generate sky LUTs
+		// TODO: Valid shadow maps (for volumetric light shafts) have not yet been generated at this point in the frame. Need to resolve dependency ordering!
+		// This also must happen before the BasePass for Sky material to be able to sample valid LUTs.
+		if (SkyAtmospherePassLocation == ESkyAtmospherePassLocation::BeforeBasePass && bShouldRenderSkyAtmosphere)
+		{
+			// Generate the Sky/Atmosphere look up tables
+			RenderSkyAtmosphereLookUpTables(GraphBuilder, ExternalAccessQueue);
+
+			// Sky env map capture uses the view UB, which contains the LUTs computed above. We need to transition them to readable now.
+			ExternalAccessQueue.Submit(GraphBuilder);
+		}
 
-		if (bHairStrandsEnable && !bHasRayTracedOverlay)
+		// Capture the SkyLight using the SkyAtmosphere and VolumetricCloud component if available.
+		const bool bRealTimeSkyCaptureEnabled = Scene->SkyLight && Scene->SkyLight->bRealTimeCaptureEnabled && Views.Num() > 0 && ViewFamily.EngineShowFlags.SkyLighting;
+		if (bRealTimeSkyCaptureEnabled)
 		{
-			RenderHairPrePass(GraphBuilder, Scene, Views, InstanceCullingManager);
-			RenderHairBasePass(GraphBuilder, Scene, SceneTextures, Views, InstanceCullingManager);
+			FViewInfo& MainView = Views[0];
+			Scene->AllocateAndCaptureFrameSkyEnvMap(GraphBuilder, *this, MainView, bShouldRenderSkyAtmosphere, bShouldRenderVolumetricCloud, InstanceCullingManager, ExternalAccessQueue);
 		}
 
-		RenderForwardShadowProjections(GraphBuilder, SceneTextures, ForwardScreenSpaceShadowMaskTexture, ForwardScreenSpaceShadowMaskHairTexture);
+		const ECustomDepthPassLocation CustomDepthPassLocation = GetCustomDepthPassLocation(ShaderPlatform);
+		if (CustomDepthPassLocation == ECustomDepthPassLocation::BeforeBasePass)
+		{
+			QUICK_SCOPE_CYCLE_COUNTER(STAT_FDeferredShadingSceneRenderer_CustomDepthPass_BeforeBasePass);
+			if (RenderCustomDepthPass(GraphBuilder, SceneTextures.CustomDepth, SceneTextures.GetSceneTextureShaderParameters(FeatureLevel), NaniteRasterResults, PrimaryNaniteViews))
+			{
+				SceneTextures.SetupMode |= ESceneTextureSetupMode::CustomDepth;
+				SceneTextures.UniformBuffer = CreateSceneTextureUniformBuffer(GraphBuilder, &SceneTextures, FeatureLevel, SceneTextures.SetupMode);
+			}
+		}
 
-		// With forward shading we need to render volumetric fog before the base pass
-		ComputeVolumetricFog(GraphBuilder, SceneTextures);
-	}
+		// Lumen updates need access to sky atmosphere LUT.
+		ExternalAccessQueue.Submit(GraphBuilder);
 
-	ExternalAccessQueue.Submit(GraphBuilder);
+		UpdateLumenScene(GraphBuilder, LumenFrameTemporaries);
 
-	FDBufferTextures DBufferTextures = CreateDBufferTextures(GraphBuilder, SceneTextures.Config.Extent, ShaderPlatform);
+		FRDGTextureRef HalfResolutionDepthCheckerboardMinMaxTexture = nullptr;
+		FRDGTextureRef QuarterResolutionDepthMinMaxTexture = nullptr;
+		bool bQuarterResMinMaxDepthRequired = bShouldRenderVolumetricCloud && ShouldVolumetricCloudTraceWithMinMaxDepth(Views);
 
-	{
-		RDG_CSV_STAT_EXCLUSIVE_SCOPE(GraphBuilder, DeferredShadingSceneRenderer_DBuffer);
-		SCOPE_CYCLE_COUNTER(STAT_FDeferredShadingSceneRenderer_DBuffer);
-		CompositionLighting.ProcessBeforeBasePass(GraphBuilder, DBufferTextures);
-	}
-	
-	if (IsForwardShadingEnabled(ShaderPlatform) && bAllowStaticLighting)
-	{
-		RenderIndirectCapsuleShadows(GraphBuilder, SceneTextures);
-	}
+		auto GenerateQuarterResDepthMinMaxTexture = [&](auto& GraphBuilder, auto& Views, auto& InputTexture)
+		{
+			if (bQuarterResMinMaxDepthRequired)
+			{
+				check(InputTexture != nullptr);	// Must receive a valid texture
+				check(QuarterResolutionDepthMinMaxTexture == nullptr);	// Only generate it once
+				return CreateQuarterResolutionDepthMinAndMax(GraphBuilder, Views, InputTexture);
+			}
+			return FRDGTextureRef(nullptr);
+		};
 
-	FTranslucencyLightingVolumeTextures TranslucencyLightingVolumeTextures;
+		// Kick off async compute cloud early if all depth has been written in the prepass
+		if (bShouldRenderVolumetricCloud && bAsyncComputeVolumetricCloud && DepthPass.EarlyZPassMode == DDM_AllOpaque && !bHasRayTracedOverlay)
+		{
+			HalfResolutionDepthCheckerboardMinMaxTexture = CreateHalfResolutionDepthCheckerboardMinMax(GraphBuilder, Views, SceneTextures.Depth.Resolve);
+			QuarterResolutionDepthMinMaxTexture = GenerateQuarterResDepthMinMaxTexture(GraphBuilder, Views, HalfResolutionDepthCheckerboardMinMaxTexture);
 
-	if (bRenderDeferredLighting && GbEnableAsyncComputeTranslucencyLightingVolumeClear && GSupportsEfficientAsyncCompute)
-	{
-		TranslucencyLightingVolumeTextures.Init(GraphBuilder, Views, ERDGPassFlags::AsyncCompute);
-	}
+			bool bSkipVolumetricRenderTarget = false;
+			bool bSkipPerPixelTracing = true;
+			bAsyncComputeVolumetricCloud = RenderVolumetricCloud(GraphBuilder, SceneTextures, bSkipVolumetricRenderTarget, bSkipPerPixelTracing, 
+				HalfResolutionDepthCheckerboardMinMaxTexture, QuarterResolutionDepthMinMaxTexture, true, InstanceCullingManager);
+		}
+		
+		FRDGTextureRef ForwardScreenSpaceShadowMaskTexture = nullptr;
+		FRDGTextureRef ForwardScreenSpaceShadowMaskHairTexture = nullptr;
+		if (IsForwardShadingEnabled(ShaderPlatform))
+		{
+			// With forward shading we need to render shadow maps early
+			ensureMsgf(!VirtualShadowMapArray.IsEnabled(), TEXT("Virtual shadow maps are not supported in the forward shading path"));
+			RenderShadowDepthMaps(GraphBuilder, InstanceCullingManager, ExternalAccessQueue);
 
-	FRDGBufferRef DynamicGeometryScratchBuffer = nullptr;
-#if RHI_RAYTRACING
-	// Async AS builds can potentially overlap with BasePass.
-	bool bNeedToWaitForRayTracingScene = DispatchRayTracingWorldUpdates(GraphBuilder, DynamicGeometryScratchBuffer);
+			if (bHairStrandsEnable && !bHasRayTracedOverlay)
+			{
+				RenderHairPrePass(GraphBuilder, Scene, Views, InstanceCullingManager);
+				RenderHairBasePass(GraphBuilder, Scene, SceneTextures, Views, InstanceCullingManager);
+			}
 
-	/** Should be called somewhere before "WaitForRayTracingScene" */
-	SetupRayTracingLightDataForViews(GraphBuilder);
-#endif
+			RenderForwardShadowProjections(GraphBuilder, SceneTextures, ForwardScreenSpaceShadowMaskTexture, ForwardScreenSpaceShadowMaskHairTexture);
 
-	if (!bHasRayTracedOverlay)
-	{
-#if RHI_RAYTRACING
-		// Lumen scene lighting requires ray tracing scene to be ready if HWRT shadows are desired
-		if (bNeedToWaitForRayTracingScene && Lumen::UseHardwareRayTracedSceneLighting(ViewFamily))
-		{
-			WaitForRayTracingScene(GraphBuilder, DynamicGeometryScratchBuffer);
-			bNeedToWaitForRayTracingScene = false;
+			// With forward shading we need to render volumetric fog before the base pass
+			ComputeVolumetricFog(GraphBuilder, SceneTextures);
 		}
-#endif
 
-		LLM_SCOPE_BYTAG(Lumen);
-		BeginGatheringLumenSurfaceCacheFeedback(GraphBuilder, Views[0], LumenFrameTemporaries);
-		RenderLumenSceneLighting(GraphBuilder, LumenFrameTemporaries, InitViewTaskDatas.LumenDirectLighting);
-	}
+		ExternalAccessQueue.Submit(GraphBuilder);
 
-	{
-		RenderBasePass(GraphBuilder, SceneTextures, DBufferTextures, BasePassDepthStencilAccess, ForwardScreenSpaceShadowMaskTexture, InstanceCullingManager, bNaniteEnabled, NaniteRasterResults);
-		GraphBuilder.AddDispatchHint();
+		FDBufferTextures DBufferTextures = CreateDBufferTextures(GraphBuilder, SceneTextures.Config.Extent, ShaderPlatform);
 
-		if (!bAllowReadOnlyDepthBasePass)
 		{
-			AddResolveSceneDepthPass(GraphBuilder, Views, SceneTextures.Depth);
+			RDG_CSV_STAT_EXCLUSIVE_SCOPE(GraphBuilder, DeferredShadingSceneRenderer_DBuffer);
+			SCOPE_CYCLE_COUNTER(STAT_FDeferredShadingSceneRenderer_DBuffer);
+			CompositionLighting.ProcessBeforeBasePass(GraphBuilder, DBufferTextures);
 		}
-
-		if (bNaniteEnabled)
+		
+		if (IsForwardShadingEnabled(ShaderPlatform) && bAllowStaticLighting)
 		{
-			if (GNaniteShowStats != 0)
-			{
-				for (int32 ViewIndex = 0; ViewIndex < Views.Num(); ViewIndex++)
+			RenderIndirectCapsuleShadows(GraphBuilder, SceneTextures);
+		}
+
+		FTranslucencyLightingVolumeTextures TranslucencyLightingVolumeTextures;
+
+		if (bRenderDeferredLighting && GbEnableAsyncComputeTranslucencyLightingVolumeClear && GSupportsEfficientAsyncCompute)
+		{
+			TranslucencyLightingVolumeTextures.Init(GraphBuilder, Views, ERDGPassFlags::AsyncCompute);
+		}
+
+		FRDGBufferRef DynamicGeometryScratchBuffer = nullptr;
+#if RHI_RAYTRACING
+		// Async AS builds can potentially overlap with BasePass.
+		bool bNeedToWaitForRayTracingScene = DispatchRayTracingWorldUpdates(GraphBuilder, DynamicGeometryScratchBuffer);
+
+		/** Should be called somewhere before "WaitForRayTracingScene" */
+		SetupRayTracingLightDataForViews(GraphBuilder);
+#endif
+
+		if (!bHasRayTracedOverlay)
+		{
+#if RHI_RAYTRACING
+			// Lumen scene lighting requires ray tracing scene to be ready if HWRT shadows are desired
+			if (bNeedToWaitForRayTracingScene && Lumen::UseHardwareRayTracedSceneLighting(ViewFamily))
+			{
+				WaitForRayTracingScene(GraphBuilder, DynamicGeometryScratchBuffer);
+				bNeedToWaitForRayTracingScene = false;
+			}
+#endif
+
+			LLM_SCOPE_BYTAG(Lumen);
+			BeginGatheringLumenSurfaceCacheFeedback(GraphBuilder, Views[0], LumenFrameTemporaries);
+			RenderLumenSceneLighting(GraphBuilder, LumenFrameTemporaries, InitViewTaskDatas.LumenDirectLighting);
+		}
+
+		{
+			RenderBasePass(GraphBuilder, SceneTextures, DBufferTextures, BasePassDepthStencilAccess, ForwardScreenSpaceShadowMaskTexture, InstanceCullingManager, bNaniteEnabled, NaniteRasterResults);
+			GraphBuilder.AddDispatchHint();
+
+			if (!bAllowReadOnlyDepthBasePass)
+			{
+				AddResolveSceneDepthPass(GraphBuilder, Views, SceneTextures.Depth);
+			}
+
+			if (bNaniteEnabled)
+			{
+				if (GNaniteShowStats != 0)
 				{
-					const FViewInfo& View = Views[ViewIndex];
-					if (IStereoRendering::IsAPrimaryView(View))
+					for (int32 ViewIndex = 0; ViewIndex < Views.Num(); ViewIndex++)
 					{
-						Nanite::PrintStats(GraphBuilder, View);
+						const FViewInfo& View = Views[ViewIndex];
+						if (IStereoRendering::IsAPrimaryView(View))
+						{
+							Nanite::PrintStats(GraphBuilder, View);
+						}
 					}
 				}
-			}
 
-			if (bVisualizeNanite)
-			{
-				FNanitePickingFeedback PickingFeedback = { 0 };
+				if (bVisualizeNanite)
+				{
+					FNanitePickingFeedback PickingFeedback = { 0 };
 
-				Nanite::AddVisualizationPasses(
-					GraphBuilder,
-					Scene,
-					SceneTextures,
-					ViewFamily.EngineShowFlags,
-					Views,
-					NaniteRasterResults,
-					PickingFeedback
-				);
+					Nanite::AddVisualizationPasses(
+						GraphBuilder,
+						Scene,
+						SceneTextures,
+						ViewFamily.EngineShowFlags,
+						Views,
+						NaniteRasterResults,
+						PickingFeedback
+					);
 
-				OnGetOnScreenMessages.AddLambda([this, PickingFeedback, RenderFlags = NaniteRasterResults[0].RenderFlags, ScenePtr = Scene](FScreenMessageWriter& ScreenMessageWriter)->void
-				{
-					Nanite::DisplayPicking(ScenePtr, PickingFeedback, RenderFlags, ScreenMessageWriter);
-				});
+					OnGetOnScreenMessages.AddLambda([this, PickingFeedback, RenderFlags = NaniteRasterResults[0].RenderFlags, ScenePtr = Scene](FScreenMessageWriter& ScreenMessageWriter)->void
+					{
+						Nanite::DisplayPicking(ScenePtr, PickingFeedback, RenderFlags, ScreenMessageWriter);
+					});
+				}
 			}
+
+			// VisualizeVirtualShadowMap TODO
 		}
 
-		// VisualizeVirtualShadowMap TODO
-	}
+		// Extract emissive from SceneColor (before lighting is applied) + material diffuse and subsurface colors
+		FRDGTextureRef ExposureIlluminanceSetup = AddSetupExposureIlluminancePass(GraphBuilder, Views, SceneTextures);
 
-	// Extract emissive from SceneColor (before lighting is applied) + material diffuse and subsurface colors
-	FRDGTextureRef ExposureIlluminanceSetup = AddSetupExposureIlluminancePass(GraphBuilder, Views, SceneTextures);
+		if (ViewFamily.EngineShowFlags.VisualizeLightCulling)
+		{
+			FRDGTextureRef VisualizeLightCullingTexture = GraphBuilder.CreateTexture(SceneTextures.Color.Target->Desc, TEXT("SceneColorVisualizeLightCulling"));
+			AddClearRenderTargetPass(GraphBuilder, VisualizeLightCullingTexture, FLinearColor::Transparent);
+			SceneTextures.Color.Target = VisualizeLightCullingTexture;
 
-	if (ViewFamily.EngineShowFlags.VisualizeLightCulling)
-	{
-		FRDGTextureRef VisualizeLightCullingTexture = GraphBuilder.CreateTexture(SceneTextures.Color.Target->Desc, TEXT("SceneColorVisualizeLightCulling"));
-		AddClearRenderTargetPass(GraphBuilder, VisualizeLightCullingTexture, FLinearColor::Transparent);
-		SceneTextures.Color.Target = VisualizeLightCullingTexture;
+			// When not in MSAA, assign to both targets.
+			if (SceneTexturesConfig.NumSamples == 1)
+			{
+				SceneTextures.Color.Resolve = SceneTextures.Color.Target;
+			}
+		}
 
-		// When not in MSAA, assign to both targets.
-		if (SceneTexturesConfig.NumSamples == 1)
+		if (bUseGBuffer)
 		{
-			SceneTextures.Color.Resolve = SceneTextures.Color.Target;
+			// mark GBufferA for saving for next frame if it's needed
+			ExtractNormalsForNextFrameReprojection(GraphBuilder, SceneTextures, Views);
 		}
-	}
-
-	if (bUseGBuffer)
-	{
-		// mark GBufferA for saving for next frame if it's needed
-		ExtractNormalsForNextFrameReprojection(GraphBuilder, SceneTextures, Views);
-	}
 
-	// Rebuild scene textures to include GBuffers.
-	SceneTextures.SetupMode |= ESceneTextureSetupMode::GBuffers;
-	if (bShouldRenderVelocities && (bBasePassCanOutputVelocity || Scene->EarlyZPassMode == DDM_AllOpaqueNoVelocity))
-	{
-		SceneTextures.SetupMode |= ESceneTextureSetupMode::SceneVelocity;
-	}
-	SceneTextures.UniformBuffer = CreateSceneTextureUniformBuffer(GraphBuilder, &SceneTextures, FeatureLevel, SceneTextures.SetupMode);
+		// Rebuild scene textures to include GBuffers.
+		SceneTextures.SetupMode |= ESceneTextureSetupMode::GBuffers;
+		if (bShouldRenderVelocities && (bBasePassCanOutputVelocity || Scene->EarlyZPassMode == DDM_AllOpaqueNoVelocity))
+		{
+			SceneTextures.SetupMode |= ESceneTextureSetupMode::SceneVelocity;
+		}
+		SceneTextures.UniformBuffer = CreateSceneTextureUniformBuffer(GraphBuilder, &SceneTextures, FeatureLevel, SceneTextures.SetupMode);
 
-	if (bRealTimeSkyCaptureEnabled)
-	{
-		Scene->ValidateSkyLightRealTimeCapture(GraphBuilder, Views[0], SceneTextures.Color.Target);
-	}
+		if (bRealTimeSkyCaptureEnabled)
+		{
+			Scene->ValidateSkyLightRealTimeCapture(GraphBuilder, Views[0], SceneTextures.Color.Target);
+		}
 
-	VisualizeVolumetricLightmap(GraphBuilder, SceneTextures);
+		VisualizeVolumetricLightmap(GraphBuilder, SceneTextures);
 
-	// Occlusion after base pass
-	if (!bOcclusionBeforeBasePass)
-	{
-		RenderOcclusionLambda();
-	}
+		// Occlusion after base pass
+		if (!bOcclusionBeforeBasePass)
+		{
+			RenderOcclusionLambda();
+		}
 
-	// End occlusion after base
+		// End occlusion after base
 
-	if (!bUseGBuffer)
-	{
-		AddResolveSceneColorPass(GraphBuilder, Views, SceneTextures.Color);
-	}
+		if (!bUseGBuffer)
+		{
+			AddResolveSceneColorPass(GraphBuilder, Views, SceneTextures.Color);
+		}
 
-	// Render hair
-	if (bHairStrandsEnable && !IsForwardShadingEnabled(ShaderPlatform) && !bHasRayTracedOverlay)
-	{
-		RenderHairPrePass(GraphBuilder, Scene, Views, InstanceCullingManager);
-		RenderHairBasePass(GraphBuilder, Scene, SceneTextures, Views, InstanceCullingManager);
-	}
+		// Render hair
+		if (bHairStrandsEnable && !IsForwardShadingEnabled(ShaderPlatform) && !bHasRayTracedOverlay)
+		{
+			RenderHairPrePass(GraphBuilder, Scene, Views, InstanceCullingManager);
+			RenderHairBasePass(GraphBuilder, Scene, SceneTextures, Views, InstanceCullingManager);
+		}
 
-	// Post base pass for material classification
-	// This needs to run before virtual shadow map, in order to have ready&cleared classified SSS data
-	if (Strata::IsStrataEnabled())
-	{
-		Strata::AddStrataMaterialClassificationPass(GraphBuilder, SceneTextures, DBufferTextures, Views);
-		Strata::AddStrataDBufferPass(GraphBuilder, SceneTextures, DBufferTextures, Views);
-	}
+		// Post base pass for material classification
+		// This needs to run before virtual shadow map, in order to have ready&cleared classified SSS data
+		if (Strata::IsStrataEnabled())
+		{
+			Strata::AddStrataMaterialClassificationPass(GraphBuilder, SceneTextures, DBufferTextures, Views);
+			Strata::AddStrataDBufferPass(GraphBuilder, SceneTextures, DBufferTextures, Views);
+		}
 
-	// Copy lighting channels out of stencil before deferred decals which overwrite those values
-	TArray<FRDGTextureRef, TInlineAllocator<2>> NaniteShadingMask;
-	if (bNaniteEnabled && Views.Num() > 0)
-	{
-		check(Views.Num() == NaniteRasterResults.Num());
-		for (const Nanite::FRasterResults& Results : NaniteRasterResults)
+		// Copy lighting channels out of stencil before deferred decals which overwrite those values
+		TArray<FRDGTextureRef, TInlineAllocator<2>> NaniteShadingMask;
+		if (bNaniteEnabled && Views.Num() > 0)
 		{
-			NaniteShadingMask.Add(Results.ShadingMask);
+			check(Views.Num() == NaniteRasterResults.Num());
+			for (const Nanite::FRasterResults& Results : NaniteRasterResults)
+			{
+				NaniteShadingMask.Add(Results.ShadingMask);
+			}
 		}
-	}
-	FRDGTextureRef LightingChannelsTexture = CopyStencilToLightingChannelTexture(GraphBuilder, SceneTextures.Stencil, NaniteShadingMask);
+		FRDGTextureRef LightingChannelsTexture = CopyStencilToLightingChannelTexture(GraphBuilder, SceneTextures.Stencil, NaniteShadingMask);
 
-	// Single layer water depth prepass. Needs to run before VSM page allocation.
-	const FSingleLayerWaterPrePassResult* SingleLayerWaterPrePassResult = nullptr;
+		// Single layer water depth prepass. Needs to run before VSM page allocation.
+		const FSingleLayerWaterPrePassResult* SingleLayerWaterPrePassResult = nullptr;
 
-	const bool bShouldRenderSingleLayerWaterDepthPrepass = !bHasRayTracedOverlay && ShouldRenderSingleLayerWaterDepthPrepass(Views);
-	if (bShouldRenderSingleLayerWaterDepthPrepass)
-	{
-		SingleLayerWaterPrePassResult = RenderSingleLayerWaterDepthPrepass(GraphBuilder, SceneTextures);
-	}
+		const bool bShouldRenderSingleLayerWaterDepthPrepass = !bHasRayTracedOverlay && ShouldRenderSingleLayerWaterDepthPrepass(Views);
+		if (bShouldRenderSingleLayerWaterDepthPrepass)
+		{
+			SingleLayerWaterPrePassResult = RenderSingleLayerWaterDepthPrepass(GraphBuilder, SceneTextures);
+		}
 
-	FAsyncLumenIndirectLightingOutputs AsyncLumenIndirectLightingOutputs;
-	const bool bHasLumenLights = SortedLightSet.LumenLightStart < SortedLightSet.SortedLights.Num();
+		FAsyncLumenIndirectLightingOutputs AsyncLumenIndirectLightingOutputs;
+		const bool bHasLumenLights = SortedLightSet.LumenLightStart < SortedLightSet.SortedLights.Num();
 
-	GraphBuilder.FlushSetupQueue();
+		GraphBuilder.FlushSetupQueue();
 
-	// Shadows, lumen and fog after base pass
-	if (!bHasRayTracedOverlay)
-	{
-#if RHI_RAYTRACING
-		// When Lumen HWRT is running async we need to wait for ray tracing scene before dispatching the work
-		if (bNeedToWaitForRayTracingScene && Lumen::UseAsyncCompute(ViewFamily) && Lumen::UseHardwareInlineRayTracing(ViewFamily))
+		// Shadows, lumen and fog after base pass
+		if (!bHasRayTracedOverlay)
 		{
-			WaitForRayTracingScene(GraphBuilder, DynamicGeometryScratchBuffer);
-			bNeedToWaitForRayTracingScene = false;
-		}
+#if RHI_RAYTRACING
+			// When Lumen HWRT is running async we need to wait for ray tracing scene before dispatching the work
+			if (bNeedToWaitForRayTracingScene && Lumen::UseAsyncCompute(ViewFamily) && Lumen::UseHardwareInlineRayTracing(ViewFamily))
+			{
+				WaitForRayTracingScene(GraphBuilder, DynamicGeometryScratchBuffer);
+				bNeedToWaitForRayTracingScene = false;
+			}
 #endif // RHI_RAYTRACING
 
-		DispatchAsyncLumenIndirectLightingWork(
-			GraphBuilder,
-			CompositionLighting,
-			SceneTextures,
-			LumenFrameTemporaries,
-			LightingChannelsTexture,
-			bHasLumenLights,
-			AsyncLumenIndirectLightingOutputs);
-
-		// If forward shading is enabled, we rendered shadow maps earlier already
-		if (!IsForwardShadingEnabled(ShaderPlatform))
-		{
-			if (VirtualShadowMapArray.IsEnabled())
+			DispatchAsyncLumenIndirectLightingWork(
+				GraphBuilder,
+				CompositionLighting,
+				SceneTextures,
+				LumenFrameTemporaries,
+				LightingChannelsTexture,
+				bHasLumenLights,
+				AsyncLumenIndirectLightingOutputs);
+
+			// If forward shading is enabled, we rendered shadow maps earlier already
+			if (!IsForwardShadingEnabled(ShaderPlatform))
 			{
-				// TODO: actually move this inside RenderShadowDepthMaps instead of this extra scope to make it 1:1 with profiling captures/traces
-				RDG_GPU_STAT_SCOPE(GraphBuilder, ShadowDepths);
+				if (VirtualShadowMapArray.IsEnabled())
+				{
+					// TODO: actually move this inside RenderShadowDepthMaps instead of this extra scope to make it 1:1 with profiling captures/traces
+					RDG_GPU_STAT_SCOPE(GraphBuilder, ShadowDepths);
 
-				ensureMsgf(AreLightsInLightGrid(), TEXT("Virtual shadow map setup requires local lights to be injected into the light grid (this may be caused by 'r.LightCulling.Quality=0')."));
+					ensureMsgf(AreLightsInLightGrid(), TEXT("Virtual shadow map setup requires local lights to be injected into the light grid (this may be caused by 'r.LightCulling.Quality=0')."));
 
-				FFrontLayerTranslucencyData FrontLayerTranslucencyData = RenderFrontLayerTranslucency(GraphBuilder, Views, SceneTextures, true /*VSM page marking*/);
+					FFrontLayerTranslucencyData FrontLayerTranslucencyData = RenderFrontLayerTranslucency(GraphBuilder, Views, SceneTextures, true /*VSM page marking*/);
 
-				VirtualShadowMapArray.BuildPageAllocations(GraphBuilder, GetActiveSceneTextures(), Views, ViewFamily.EngineShowFlags, SortedLightSet, VisibleLightInfos, SingleLayerWaterPrePassResult, FrontLayerTranslucencyData);
-			}
+					VirtualShadowMapArray.BuildPageAllocations(GraphBuilder, GetActiveSceneTextures(), Views, ViewFamily.EngineShowFlags, SortedLightSet, VisibleLightInfos, SingleLayerWaterPrePassResult, FrontLayerTranslucencyData);
+				}
 
-			RenderShadowDepthMaps(GraphBuilder, InstanceCullingManager, ExternalAccessQueue);
-		}
-		CheckShadowDepthRenderCompleted();
+				RenderShadowDepthMaps(GraphBuilder, InstanceCullingManager, ExternalAccessQueue);
+			}
+			CheckShadowDepthRenderCompleted();
 
 #if RHI_RAYTRACING
-		// Lumen scene lighting requires ray tracing scene to be ready if HWRT shadows are desired
-		if (bNeedToWaitForRayTracingScene && Lumen::UseHardwareRayTracedSceneLighting(ViewFamily))
-		{
-			WaitForRayTracingScene(GraphBuilder, DynamicGeometryScratchBuffer);
-			bNeedToWaitForRayTracingScene = false;
-		}
+			// Lumen scene lighting requires ray tracing scene to be ready if HWRT shadows are desired
+			if (bNeedToWaitForRayTracingScene && Lumen::UseHardwareRayTracedSceneLighting(ViewFamily))
+			{
+				WaitForRayTracingScene(GraphBuilder, DynamicGeometryScratchBuffer);
+				bNeedToWaitForRayTracingScene = false;
+			}
 #endif // RHI_RAYTRACING
-	}
+		}
 
-	ExternalAccessQueue.Submit(GraphBuilder);
+		ExternalAccessQueue.Submit(GraphBuilder);
 
-	// End shadow and fog after base pass
+		// End shadow and fog after base pass
 
-	// Trigger a command submit here, to avoid GPU bubbles
-	AddDispatchToRHIThreadPass(GraphBuilder);
-	
-	if (bNaniteEnabled)
-	{
-		// Needs doing after shadows such that the checks for shadow atlases etc work.
-		Nanite::ListStatFilters(this);
-	}
+		// Trigger a command submit here, to avoid GPU bubbles
+		AddDispatchToRHIThreadPass(GraphBuilder);
+		
+		if (bNaniteEnabled)
+		{
+			// Needs doing after shadows such that the checks for shadow atlases etc work.
+			Nanite::ListStatFilters(this);
+		}
 
-	if (bUpdateNaniteStreaming)
-	{
-		Nanite::GStreamingManager.SubmitFrameStreamingRequests(GraphBuilder);
-	}
+		if (bUpdateNaniteStreaming)
+		{
+			Nanite::GStreamingManager.SubmitFrameStreamingRequests(GraphBuilder);
+		}
 
-	{
-		FVirtualShadowMapArrayCacheManager* CacheManager = VirtualShadowMapArray.CacheManager;
-		if (CacheManager)
 		{
-			// Do this even if VSMs are disabled this frame to clean up any previously extracted data
-			CacheManager->ExtractFrameData(
-				GraphBuilder,
-				VirtualShadowMapArray,
-				*this,
-				ViewFamily.EngineShowFlags.VirtualShadowMapCaching);
+			FVirtualShadowMapArrayCacheManager* CacheManager = VirtualShadowMapArray.CacheManager;
+			if (CacheManager)
+			{
+				// Do this even if VSMs are disabled this frame to clean up any previously extracted data
+				CacheManager->ExtractFrameData(
+					GraphBuilder,
+					VirtualShadowMapArray,
+					*this,
+					ViewFamily.EngineShowFlags.VirtualShadowMapCaching);
+			}
 		}
-	}
 
-	// If not all depth is written during the prepass, kick off async compute cloud after basepass
-	if (bShouldRenderVolumetricCloud && bAsyncComputeVolumetricCloud && DepthPass.EarlyZPassMode != DDM_AllOpaque && !bHasRayTracedOverlay)
-	{
-		HalfResolutionDepthCheckerboardMinMaxTexture = CreateHalfResolutionDepthCheckerboardMinMax(GraphBuilder, Views, SceneTextures.Depth.Resolve);
-		QuarterResolutionDepthMinMaxTexture = GenerateQuarterResDepthMinMaxTexture(GraphBuilder, Views, HalfResolutionDepthCheckerboardMinMaxTexture);
+		// If not all depth is written during the prepass, kick off async compute cloud after basepass
+		if (bShouldRenderVolumetricCloud && bAsyncComputeVolumetricCloud && DepthPass.EarlyZPassMode != DDM_AllOpaque && !bHasRayTracedOverlay)
+		{
+			HalfResolutionDepthCheckerboardMinMaxTexture = CreateHalfResolutionDepthCheckerboardMinMax(GraphBuilder, Views, SceneTextures.Depth.Resolve);
+			QuarterResolutionDepthMinMaxTexture = GenerateQuarterResDepthMinMaxTexture(GraphBuilder, Views, HalfResolutionDepthCheckerboardMinMaxTexture);
 
-		bool bSkipVolumetricRenderTarget = false;
-		bool bSkipPerPixelTracing = true;
-		bAsyncComputeVolumetricCloud = RenderVolumetricCloud(GraphBuilder, SceneTextures, bSkipVolumetricRenderTarget, bSkipPerPixelTracing, 
-			HalfResolutionDepthCheckerboardMinMaxTexture, QuarterResolutionDepthMinMaxTexture, true, InstanceCullingManager);
-	}
+			bool bSkipVolumetricRenderTarget = false;
+			bool bSkipPerPixelTracing = true;
+			bAsyncComputeVolumetricCloud = RenderVolumetricCloud(GraphBuilder, SceneTextures, bSkipVolumetricRenderTarget, bSkipPerPixelTracing, 
+				HalfResolutionDepthCheckerboardMinMaxTexture, QuarterResolutionDepthMinMaxTexture, true, InstanceCullingManager);
+		}
 
-	if (CustomDepthPassLocation == ECustomDepthPassLocation::AfterBasePass)
-	{
-		QUICK_SCOPE_CYCLE_COUNTER(STAT_FDeferredShadingSceneRenderer_CustomDepthPass_AfterBasePass);
-		if (RenderCustomDepthPass(GraphBuilder, SceneTextures.CustomDepth, SceneTextures.GetSceneTextureShaderParameters(FeatureLevel), NaniteRasterResults, PrimaryNaniteViews))
+		if (CustomDepthPassLocation == ECustomDepthPassLocation::AfterBasePass)
 		{
-			SceneTextures.SetupMode |= ESceneTextureSetupMode::CustomDepth;
-			SceneTextures.UniformBuffer = CreateSceneTextureUniformBuffer(GraphBuilder, &SceneTextures, FeatureLevel, SceneTextures.SetupMode);
+			QUICK_SCOPE_CYCLE_COUNTER(STAT_FDeferredShadingSceneRenderer_CustomDepthPass_AfterBasePass);
+			if (RenderCustomDepthPass(GraphBuilder, SceneTextures.CustomDepth, SceneTextures.GetSceneTextureShaderParameters(FeatureLevel), NaniteRasterResults, PrimaryNaniteViews))
+			{
+				SceneTextures.SetupMode |= ESceneTextureSetupMode::CustomDepth;
+				SceneTextures.UniformBuffer = CreateSceneTextureUniformBuffer(GraphBuilder, &SceneTextures, FeatureLevel, SceneTextures.SetupMode);
+			}
 		}
-	}
 
-	// If we are not rendering velocities in depth or base pass then do that here.
-	if (bShouldRenderVelocities && !bBasePassCanOutputVelocity && (Scene->EarlyZPassMode != DDM_AllOpaqueNoVelocity))
-	{
-		GraphBuilder.SetCommandListStat(GET_STATID(STAT_CLM_Velocity));
-		RenderVelocities(GraphBuilder, SceneTextures, EVelocityPass::Opaque, bHairStrandsEnable);
-		GraphBuilder.SetCommandListStat(GET_STATID(STAT_CLM_AfterVelocity));
-	}
-
-	// Pre-lighting composition lighting stage
-	// e.g. deferred decals, SSAO
-	{
-		RDG_CSV_STAT_EXCLUSIVE_SCOPE(GraphBuilder, AfterBasePass);
-		SCOPE_CYCLE_COUNTER(STAT_FDeferredShadingSceneRenderer_AfterBasePass);
-
-		if (!IsForwardShadingEnabled(ShaderPlatform))
+		// If we are not rendering velocities in depth or base pass then do that here.
+		if (bShouldRenderVelocities && !bBasePassCanOutputVelocity && (Scene->EarlyZPassMode != DDM_AllOpaqueNoVelocity))
 		{
-			AddResolveSceneDepthPass(GraphBuilder, Views, SceneTextures.Depth);
+			GraphBuilder.SetCommandListStat(GET_STATID(STAT_CLM_Velocity));
+			RenderVelocities(GraphBuilder, SceneTextures, EVelocityPass::Opaque, bHairStrandsEnable);
+			GraphBuilder.SetCommandListStat(GET_STATID(STAT_CLM_AfterVelocity));
 		}
 
-		const FCompositionLighting::EProcessAfterBasePassMode Mode = AsyncLumenIndirectLightingOutputs.bHasDrawnBeforeLightingDecals ?
-			FCompositionLighting::EProcessAfterBasePassMode::SkipBeforeLightingDecals : FCompositionLighting::EProcessAfterBasePassMode::All;
+		// Pre-lighting composition lighting stage
+		// e.g. deferred decals, SSAO
+		{
+			RDG_CSV_STAT_EXCLUSIVE_SCOPE(GraphBuilder, AfterBasePass);
+			SCOPE_CYCLE_COUNTER(STAT_FDeferredShadingSceneRenderer_AfterBasePass);
 
-		CompositionLighting.ProcessAfterBasePass(GraphBuilder, Mode);
-	}
+			if (!IsForwardShadingEnabled(ShaderPlatform))
+			{
+				AddResolveSceneDepthPass(GraphBuilder, Views, SceneTextures.Depth);
+			}
 
-	// Rebuild scene textures to include velocity, custom depth, and SSAO.
-	SceneTextures.SetupMode |= ESceneTextureSetupMode::All;
-	SceneTextures.UniformBuffer = CreateSceneTextureUniformBuffer(GraphBuilder, &SceneTextures, FeatureLevel, SceneTextures.SetupMode);
+			const FCompositionLighting::EProcessAfterBasePassMode Mode = AsyncLumenIndirectLightingOutputs.bHasDrawnBeforeLightingDecals ?
+				FCompositionLighting::EProcessAfterBasePassMode::SkipBeforeLightingDecals : FCompositionLighting::EProcessAfterBasePassMode::All;
 
-	if (!IsForwardShadingEnabled(ShaderPlatform))
-	{
-		// Clear stencil to 0 now that deferred decals are done using what was setup in the base pass.
-		AddClearStencilPass(GraphBuilder, SceneTextures.Depth.Target);
-	}
+			CompositionLighting.ProcessAfterBasePass(GraphBuilder, Mode);
+		}
 
-#if RHI_RAYTRACING
-	// If Lumen did not force an earlier ray tracing scene sync, we must wait for it here.
-	if (bNeedToWaitForRayTracingScene)
-	{
-		WaitForRayTracingScene(GraphBuilder, DynamicGeometryScratchBuffer);
-		bNeedToWaitForRayTracingScene = false;
-	}
-#endif // RHI_RAYTRACING
+		// Rebuild scene textures to include velocity, custom depth, and SSAO.
+		SceneTextures.SetupMode |= ESceneTextureSetupMode::All;
+		SceneTextures.UniformBuffer = CreateSceneTextureUniformBuffer(GraphBuilder, &SceneTextures, FeatureLevel, SceneTextures.SetupMode);
 
-	GraphBuilder.FlushSetupQueue();
+		if (!IsForwardShadingEnabled(ShaderPlatform))
+		{
+			// Clear stencil to 0 now that deferred decals are done using what was setup in the base pass.
+			AddClearStencilPass(GraphBuilder, SceneTextures.Depth.Target);
+		}
 
-	if (bRenderDeferredLighting)
-	{
-		RDG_GPU_STAT_SCOPE(GraphBuilder, RenderDeferredLighting);
-		RDG_CSV_STAT_EXCLUSIVE_SCOPE(GraphBuilder, RenderLighting);
-		SCOPE_CYCLE_COUNTER(STAT_FDeferredShadingSceneRenderer_Lighting);
-		SCOPED_NAMED_EVENT(RenderLighting, FColor::Emerald);
-
-		FRDGTextureRef DynamicBentNormalAOTexture = nullptr;
-
-		RenderDiffuseIndirectAndAmbientOcclusion(
-			GraphBuilder,
-			SceneTextures,
-			LumenFrameTemporaries,
-			LightingChannelsTexture,
-			bHasLumenLights,
-			/* bCompositeRegularLumenOnly = */ false,
-			/* bIsVisualizePass = */ false,
-			AsyncLumenIndirectLightingOutputs);
-
-		// These modulate the scenecolor output from the basepass, which is assumed to be indirect lighting
-		if (bAllowStaticLighting)
+#if RHI_RAYTRACING
+		// If Lumen did not force an earlier ray tracing scene sync, we must wait for it here.
+		if (bNeedToWaitForRayTracingScene)
 		{
-			RenderIndirectCapsuleShadows(GraphBuilder, SceneTextures);
+			WaitForRayTracingScene(GraphBuilder, DynamicGeometryScratchBuffer);
+			bNeedToWaitForRayTracingScene = false;
 		}
+#endif // RHI_RAYTRACING
 
-		// These modulate the scene color output from the base pass, which is assumed to be indirect lighting
-		RenderDFAOAsIndirectShadowing(GraphBuilder, SceneTextures, DynamicBentNormalAOTexture);
+		GraphBuilder.FlushSetupQueue();
 
-		// Clear the translucent lighting volumes before we accumulate
-		if ((GbEnableAsyncComputeTranslucencyLightingVolumeClear && GSupportsEfficientAsyncCompute) == false)
+		if (bRenderDeferredLighting)
 		{
-			TranslucencyLightingVolumeTextures.Init(GraphBuilder, Views, ERDGPassFlags::Compute);
-		}
+			RDG_GPU_STAT_SCOPE(GraphBuilder, RenderDeferredLighting);
+			RDG_CSV_STAT_EXCLUSIVE_SCOPE(GraphBuilder, RenderLighting);
+			SCOPE_CYCLE_COUNTER(STAT_FDeferredShadingSceneRenderer_Lighting);
+			SCOPED_NAMED_EVENT(RenderLighting, FColor::Emerald);
+
+			FRDGTextureRef DynamicBentNormalAOTexture = nullptr;
+
+			RenderDiffuseIndirectAndAmbientOcclusion(
+				GraphBuilder,
+				SceneTextures,
+				LumenFrameTemporaries,
+				LightingChannelsTexture,
+				bHasLumenLights,
+				/* bCompositeRegularLumenOnly = */ false,
+				/* bIsVisualizePass = */ false,
+				AsyncLumenIndirectLightingOutputs);
+
+			// These modulate the scenecolor output from the basepass, which is assumed to be indirect lighting
+			if (bAllowStaticLighting)
+			{
+				RenderIndirectCapsuleShadows(GraphBuilder, SceneTextures);
+			}
+
+			// These modulate the scene color output from the base pass, which is assumed to be indirect lighting
+			RenderDFAOAsIndirectShadowing(GraphBuilder, SceneTextures, DynamicBentNormalAOTexture);
+
+			// Clear the translucent lighting volumes before we accumulate
+			if ((GbEnableAsyncComputeTranslucencyLightingVolumeClear && GSupportsEfficientAsyncCompute) == false)
+			{
+				TranslucencyLightingVolumeTextures.Init(GraphBuilder, Views, ERDGPassFlags::Compute);
+			}
 
 #if RHI_RAYTRACING
-		if (IsRayTracingEnabled())
-		{
-			RenderDitheredLODFadingOutMask(GraphBuilder, Views[0], SceneTextures.Depth.Target);
-		}
+			if (IsRayTracingEnabled())
+			{
+				RenderDitheredLODFadingOutMask(GraphBuilder, Views[0], SceneTextures.Depth.Target);
+			}
 #endif
 
-		GraphBuilder.SetCommandListStat(GET_STATID(STAT_CLM_Lighting));
-		RenderLights(GraphBuilder, SceneTextures, TranslucencyLightingVolumeTextures, LightingChannelsTexture, SortedLightSet);
-		GraphBuilder.SetCommandListStat(GET_STATID(STAT_CLM_AfterLighting));
+			GraphBuilder.SetCommandListStat(GET_STATID(STAT_CLM_Lighting));
+			RenderLights(GraphBuilder, SceneTextures, TranslucencyLightingVolumeTextures, LightingChannelsTexture, SortedLightSet);
+			GraphBuilder.SetCommandListStat(GET_STATID(STAT_CLM_AfterLighting));
 
-		InjectTranslucencyLightingVolumeAmbientCubemap(GraphBuilder, Views, TranslucencyLightingVolumeTextures);
-		FilterTranslucencyLightingVolume(GraphBuilder, Views, TranslucencyLightingVolumeTextures);
+			InjectTranslucencyLightingVolumeAmbientCubemap(GraphBuilder, Views, TranslucencyLightingVolumeTextures);
+			FilterTranslucencyLightingVolume(GraphBuilder, Views, TranslucencyLightingVolumeTextures);
 
-		// Do DiffuseIndirectComposite after Lights so that async Lumen work can overlap
-		RenderDiffuseIndirectAndAmbientOcclusion(
-			GraphBuilder,
-			SceneTextures,
-			LumenFrameTemporaries,
-			LightingChannelsTexture,
-			bHasLumenLights,
-			/* bCompositeRegularLumenOnly = */ true,
-			/* bIsVisualizePass = */ false,
-			AsyncLumenIndirectLightingOutputs);
+			// Do DiffuseIndirectComposite after Lights so that async Lumen work can overlap
+			RenderDiffuseIndirectAndAmbientOcclusion(
+				GraphBuilder,
+				SceneTextures,
+				LumenFrameTemporaries,
+				LightingChannelsTexture,
+				bHasLumenLights,
+				/* bCompositeRegularLumenOnly = */ true,
+				/* bIsVisualizePass = */ false,
+				AsyncLumenIndirectLightingOutputs);
 
-		// Render diffuse sky lighting and reflections that only operate on opaque pixels
-		RenderDeferredReflectionsAndSkyLighting(GraphBuilder, SceneTextures, LumenFrameTemporaries, DynamicBentNormalAOTexture);
+			// Render diffuse sky lighting and reflections that only operate on opaque pixels
+			RenderDeferredReflectionsAndSkyLighting(GraphBuilder, SceneTextures, LumenFrameTemporaries, DynamicBentNormalAOTexture);
 
 #if !(UE_BUILD_SHIPPING || UE_BUILD_TEST)
-		// Renders debug visualizations for global illumination plugins
-		RenderGlobalIlluminationPluginVisualizations(GraphBuilder, LightingChannelsTexture);
+			// Renders debug visualizations for global illumination plugins
+			RenderGlobalIlluminationPluginVisualizations(GraphBuilder, LightingChannelsTexture);
 #endif
 
-		AddSubsurfacePass(GraphBuilder, SceneTextures, Views);
+			AddSubsurfacePass(GraphBuilder, SceneTextures, Views);
 
-		Strata::AddStrataOpaqueRoughRefractionPasses(GraphBuilder, SceneTextures, Views);
+			Strata::AddStrataOpaqueRoughRefractionPasses(GraphBuilder, SceneTextures, Views);
 
+			{
+				RenderHairStrandsSceneColorScattering(GraphBuilder, SceneTextures.Color.Target, Scene, Views);
+			}
+
+		#if RHI_RAYTRACING
+			if (ShouldRenderRayTracingSkyLight(Scene->SkyLight) 
+				//@todo - integrate RenderRayTracingSkyLight into RenderDiffuseIndirectAndAmbientOcclusion
+				&& GetViewPipelineState(Views[0]).DiffuseIndirectMethod != EDiffuseIndirectMethod::Lumen
+				&& ViewFamily.EngineShowFlags.GlobalIllumination)
+			{
+				FRDGTextureRef SkyLightTexture = nullptr;
+				FRDGTextureRef SkyLightHitDistanceTexture = nullptr;
+				RenderRayTracingSkyLight(GraphBuilder, SceneTextures.Color.Target, SkyLightTexture, SkyLightHitDistanceTexture);
+				CompositeRayTracingSkyLight(GraphBuilder, SceneTextures, SkyLightTexture, SkyLightHitDistanceTexture);
+			}
+		#endif
+
+			if (Strata::IsStrataEnabled())
+			{
+				// Now remove all the Strata tile stencil tags used by deferred tiled light passes. Make later marks such as responssive AA works.
+				AddClearStencilPass(GraphBuilder, SceneTextures.Depth.Target);
+			}
+		}
+		else if (HairStrands::HasViewHairStrandsData(Views) && ViewFamily.EngineShowFlags.Lighting)
 		{
-			RenderHairStrandsSceneColorScattering(GraphBuilder, SceneTextures.Color.Target, Scene, Views);
+			RenderLightsForHair(GraphBuilder, SceneTextures, SortedLightSet, ForwardScreenSpaceShadowMaskHairTexture, LightingChannelsTexture);
+			RenderDeferredReflectionsAndSkyLightingHair(GraphBuilder);
 		}
 
-	#if RHI_RAYTRACING
-		if (ShouldRenderRayTracingSkyLight(Scene->SkyLight) 
-			//@todo - integrate RenderRayTracingSkyLight into RenderDiffuseIndirectAndAmbientOcclusion
-			&& GetViewPipelineState(Views[0]).DiffuseIndirectMethod != EDiffuseIndirectMethod::Lumen
-			&& ViewFamily.EngineShowFlags.GlobalIllumination)
+		// Volumetric fog after Lumen GI and shadow depths
+		if (!IsForwardShadingEnabled(ShaderPlatform) && !bHasRayTracedOverlay)
 		{
-			FRDGTextureRef SkyLightTexture = nullptr;
-			FRDGTextureRef SkyLightHitDistanceTexture = nullptr;
-			RenderRayTracingSkyLight(GraphBuilder, SceneTextures.Color.Target, SkyLightTexture, SkyLightHitDistanceTexture);
-			CompositeRayTracingSkyLight(GraphBuilder, SceneTextures, SkyLightTexture, SkyLightHitDistanceTexture);
+			ComputeVolumetricFog(GraphBuilder, SceneTextures);
 		}
-	#endif
 
-		if (Strata::IsStrataEnabled())
+		if (ShouldRenderHeterogeneousVolumes(Scene) && !bHasRayTracedOverlay)
 		{
-			// Now remove all the Strata tile stencil tags used by deferred tiled light passes. Make later marks such as responssive AA works.
-			AddClearStencilPass(GraphBuilder, SceneTextures.Depth.Target);
+			RenderHeterogeneousVolumes(GraphBuilder, SceneTextures);
 		}
-	}
-	else if (HairStrands::HasViewHairStrandsData(Views) && ViewFamily.EngineShowFlags.Lighting)
-	{
-		RenderLightsForHair(GraphBuilder, SceneTextures, SortedLightSet, ForwardScreenSpaceShadowMaskHairTexture, LightingChannelsTexture);
-		RenderDeferredReflectionsAndSkyLightingHair(GraphBuilder);
-	}
 
-	// Volumetric fog after Lumen GI and shadow depths
-	if (!IsForwardShadingEnabled(ShaderPlatform) && !bHasRayTracedOverlay)
-	{
-		ComputeVolumetricFog(GraphBuilder, SceneTextures);
-	}
-
-	if (ShouldRenderHeterogeneousVolumes(Scene) && !bHasRayTracedOverlay)
-	{
-		RenderHeterogeneousVolumes(GraphBuilder, SceneTextures);
-	}
-
-	GraphBuilder.FlushSetupQueue();
-
-	if (bShouldRenderVolumetricCloud && IsVolumetricRenderTargetEnabled() && HalfResolutionDepthCheckerboardMinMaxTexture==nullptr && !bHasRayTracedOverlay)
-	{
-		HalfResolutionDepthCheckerboardMinMaxTexture = CreateHalfResolutionDepthCheckerboardMinMax(GraphBuilder, Views, SceneTextures.Depth.Resolve);
-		QuarterResolutionDepthMinMaxTexture = GenerateQuarterResDepthMinMaxTexture(GraphBuilder, Views, HalfResolutionDepthCheckerboardMinMaxTexture);
-	}
+		GraphBuilder.FlushSetupQueue();
 
-	if (bShouldRenderVolumetricCloud && !bHasRayTracedOverlay)
-	{
-		if (!bAsyncComputeVolumetricCloud)
+		if (bShouldRenderVolumetricCloud && IsVolumetricRenderTargetEnabled() && HalfResolutionDepthCheckerboardMinMaxTexture==nullptr && !bHasRayTracedOverlay)
 		{
-			// Generate the volumetric cloud render target
-			bool bSkipVolumetricRenderTarget = false;
-			bool bSkipPerPixelTracing = true;
-			RenderVolumetricCloud(GraphBuilder, SceneTextures, bSkipVolumetricRenderTarget, bSkipPerPixelTracing, 
-				HalfResolutionDepthCheckerboardMinMaxTexture, QuarterResolutionDepthMinMaxTexture, false, InstanceCullingManager);
+			HalfResolutionDepthCheckerboardMinMaxTexture = CreateHalfResolutionDepthCheckerboardMinMax(GraphBuilder, Views, SceneTextures.Depth.Resolve);
+			QuarterResolutionDepthMinMaxTexture = GenerateQuarterResDepthMinMaxTexture(GraphBuilder, Views, HalfResolutionDepthCheckerboardMinMaxTexture);
 		}
-		// Reconstruct the volumetric cloud render target to be ready to compose it over the scene
-		ReconstructVolumetricRenderTarget(GraphBuilder, Views, SceneTextures.Depth.Resolve, HalfResolutionDepthCheckerboardMinMaxTexture, bAsyncComputeVolumetricCloud);
-	}
 
-	TArray<FScreenPassTexture, TInlineAllocator<4>> TSRMoireInputTextures;
-	// Extract TSR's moire heuristic luminance before renderering translucency into the scene color.
-	for (int32 ViewIndex = 0; ViewIndex < Views.Num(); ++ViewIndex)
-	{
-		FViewInfo& View = Views[ViewIndex];
-		if (NeedTSRMoireLuma(View))
+		if (bShouldRenderVolumetricCloud && !bHasRayTracedOverlay)
 		{
-			if (TSRMoireInputTextures.Num() == 0)
+			if (!bAsyncComputeVolumetricCloud)
 			{
-				TSRMoireInputTextures.SetNum(Views.Num());
+				// Generate the volumetric cloud render target
+				bool bSkipVolumetricRenderTarget = false;
+				bool bSkipPerPixelTracing = true;
+				RenderVolumetricCloud(GraphBuilder, SceneTextures, bSkipVolumetricRenderTarget, bSkipPerPixelTracing, 
+					HalfResolutionDepthCheckerboardMinMaxTexture, QuarterResolutionDepthMinMaxTexture, false, InstanceCullingManager);
 			}
+			// Reconstruct the volumetric cloud render target to be ready to compose it over the scene
+			ReconstructVolumetricRenderTarget(GraphBuilder, Views, SceneTextures.Depth.Resolve, HalfResolutionDepthCheckerboardMinMaxTexture, bAsyncComputeVolumetricCloud);
+		}
 
-			TSRMoireInputTextures[ViewIndex] = AddTSRComputeMoireLuma(GraphBuilder, View.ShaderMap, FScreenPassTexture(SceneTextures.Color.Target, View.ViewRect));
+		TArray<FScreenPassTexture, TInlineAllocator<4>> TSRMoireInputTextures;
+		// Extract TSR's moire heuristic luminance before renderering translucency into the scene color.
+		for (int32 ViewIndex = 0; ViewIndex < Views.Num(); ++ViewIndex)
+		{
+			FViewInfo& View = Views[ViewIndex];
+			if (NeedTSRMoireLuma(View))
+			{
+				if (TSRMoireInputTextures.Num() == 0)
+				{
+					TSRMoireInputTextures.SetNum(Views.Num());
+				}
+
+				TSRMoireInputTextures[ViewIndex] = AddTSRComputeMoireLuma(GraphBuilder, View.ShaderMap, FScreenPassTexture(SceneTextures.Color.Target, View.ViewRect));
+			}
 		}
-	}
 
-	const bool bShouldRenderTranslucency = !bHasRayTracedOverlay && ShouldRenderTranslucency();
+		const bool bShouldRenderTranslucency = !bHasRayTracedOverlay && ShouldRenderTranslucency();
 
-	// Union of all translucency view render flags.
-	ETranslucencyView TranslucencyViewsToRender = bShouldRenderTranslucency ? GetTranslucencyViews(Views) : ETranslucencyView::None;
+		// Union of all translucency view render flags.
+		ETranslucencyView TranslucencyViewsToRender = bShouldRenderTranslucency ? GetTranslucencyViews(Views) : ETranslucencyView::None;
 
-	FTranslucencyPassResourcesMap TranslucencyResourceMap(Views.Num());
+		FTranslucencyPassResourcesMap TranslucencyResourceMap(Views.Num());
 
-	const bool bShouldRenderSingleLayerWater = !bHasRayTracedOverlay && ShouldRenderSingleLayerWater(Views);
-	FSceneWithoutWaterTextures SceneWithoutWaterTextures;
-	if (bShouldRenderSingleLayerWater)
-	{
-		if (EnumHasAnyFlags(TranslucencyViewsToRender, ETranslucencyView::UnderWater))
+		const bool bShouldRenderSingleLayerWater = !bHasRayTracedOverlay && ShouldRenderSingleLayerWater(Views);
+		FSceneWithoutWaterTextures SceneWithoutWaterTextures;
+		if (bShouldRenderSingleLayerWater)
 		{
-			RDG_CSV_STAT_EXCLUSIVE_SCOPE(GraphBuilder, RenderTranslucency);
-			SCOPED_NAMED_EVENT(RenderTranslucency, FColor::Emerald);
-			SCOPE_CYCLE_COUNTER(STAT_TranslucencyDrawTime);
-			GraphBuilder.SetCommandListStat(GET_STATID(STAT_CLM_Translucency));
-			const bool bStandardTranslucentCanRenderSeparate = false;
-			RenderTranslucency(GraphBuilder, SceneTextures, TranslucencyLightingVolumeTextures, &TranslucencyResourceMap, ETranslucencyView::UnderWater, InstanceCullingManager, bStandardTranslucentCanRenderSeparate);
-			EnumRemoveFlags(TranslucencyViewsToRender, ETranslucencyView::UnderWater);
-		}
+			if (EnumHasAnyFlags(TranslucencyViewsToRender, ETranslucencyView::UnderWater))
+			{
+				RDG_CSV_STAT_EXCLUSIVE_SCOPE(GraphBuilder, RenderTranslucency);
+				SCOPED_NAMED_EVENT(RenderTranslucency, FColor::Emerald);
+				SCOPE_CYCLE_COUNTER(STAT_TranslucencyDrawTime);
+				GraphBuilder.SetCommandListStat(GET_STATID(STAT_CLM_Translucency));
+				const bool bStandardTranslucentCanRenderSeparate = false;
+				RenderTranslucency(GraphBuilder, SceneTextures, TranslucencyLightingVolumeTextures, &TranslucencyResourceMap, ETranslucencyView::UnderWater, InstanceCullingManager, bStandardTranslucentCanRenderSeparate);
+				EnumRemoveFlags(TranslucencyViewsToRender, ETranslucencyView::UnderWater);
+			}
 
-		GraphBuilder.SetCommandListStat(GET_STATID(STAT_CLM_WaterPass));
-		RenderSingleLayerWater(GraphBuilder, SceneTextures, SingleLayerWaterPrePassResult, bShouldRenderVolumetricCloud, SceneWithoutWaterTextures, LumenFrameTemporaries);
+			GraphBuilder.SetCommandListStat(GET_STATID(STAT_CLM_WaterPass));
+			RenderSingleLayerWater(GraphBuilder, SceneTextures, SingleLayerWaterPrePassResult, bShouldRenderVolumetricCloud, SceneWithoutWaterTextures, LumenFrameTemporaries);
 
-		// Replace main depth texture with the output of the SLW depth prepass which contains the scene + water.
-		// Note: Stencil now has all water bits marked with 1. As long as no other passes after this point want to read the depth buffer,
-		// a stencil clear should not be necessary here.
-		if (SingleLayerWaterPrePassResult)
-		{
-			SceneTextures.Depth = SingleLayerWaterPrePassResult->DepthPrepassTexture;
+			// Replace main depth texture with the output of the SLW depth prepass which contains the scene + water.
+			// Note: Stencil now has all water bits marked with 1. As long as no other passes after this point want to read the depth buffer,
+			// a stencil clear should not be necessary here.
+			if (SingleLayerWaterPrePassResult)
+			{
+				SceneTextures.Depth = SingleLayerWaterPrePassResult->DepthPrepassTexture;
+			}
 		}
-	}
 
-	// Rebuild scene textures to include scene color.
-	SceneTextures.UniformBuffer = CreateSceneTextureUniformBuffer(GraphBuilder, &SceneTextures, FeatureLevel, SceneTextures.SetupMode);
+		// Rebuild scene textures to include scene color.
+		SceneTextures.UniformBuffer = CreateSceneTextureUniformBuffer(GraphBuilder, &SceneTextures, FeatureLevel, SceneTextures.SetupMode);
 
-	FRDGTextureRef LightShaftOcclusionTexture = nullptr;
+		FRDGTextureRef LightShaftOcclusionTexture = nullptr;
 
-	// Draw Lightshafts
-	if (!bHasRayTracedOverlay && ViewFamily.EngineShowFlags.LightShafts)
-	{
-		SCOPE_CYCLE_COUNTER(STAT_FDeferredShadingSceneRenderer_RenderLightShaftOcclusion);
-		LightShaftOcclusionTexture = RenderLightShaftOcclusion(GraphBuilder, SceneTextures);
-	}
+		// Draw Lightshafts
+		if (!bHasRayTracedOverlay && ViewFamily.EngineShowFlags.LightShafts)
+		{
+			SCOPE_CYCLE_COUNTER(STAT_FDeferredShadingSceneRenderer_RenderLightShaftOcclusion);
+			LightShaftOcclusionTexture = RenderLightShaftOcclusion(GraphBuilder, SceneTextures);
+		}
 
-	// Draw the sky atmosphere
-	if (!bHasRayTracedOverlay && bShouldRenderSkyAtmosphere && !IsForwardShadingEnabled(ShaderPlatform))
-	{
-		SCOPE_CYCLE_COUNTER(STAT_FDeferredShadingSceneRenderer_RenderSkyAtmosphere);
-		RenderSkyAtmosphere(GraphBuilder, SceneTextures);
-	}
+		// Draw the sky atmosphere
+		if (!bHasRayTracedOverlay && bShouldRenderSkyAtmosphere && !IsForwardShadingEnabled(ShaderPlatform))
+		{
+			SCOPE_CYCLE_COUNTER(STAT_FDeferredShadingSceneRenderer_RenderSkyAtmosphere);
+			RenderSkyAtmosphere(GraphBuilder, SceneTextures);
+		}
 
-	// Draw fog.
-	if (!bHasRayTracedOverlay && ShouldRenderFog(ViewFamily))
-	{
-		RDG_CSV_STAT_EXCLUSIVE_SCOPE(GraphBuilder, RenderFog);
-		SCOPED_NAMED_EVENT(RenderFog, FColor::Emerald);
-		SCOPE_CYCLE_COUNTER(STAT_FDeferredShadingSceneRenderer_RenderFog);
-		RenderFog(GraphBuilder, SceneTextures, LightShaftOcclusionTexture);
-	}
+		// Draw fog.
+		if (!bHasRayTracedOverlay && ShouldRenderFog(ViewFamily))
+		{
+			RDG_CSV_STAT_EXCLUSIVE_SCOPE(GraphBuilder, RenderFog);
+			SCOPED_NAMED_EVENT(RenderFog, FColor::Emerald);
+			SCOPE_CYCLE_COUNTER(STAT_FDeferredShadingSceneRenderer_RenderFog);
+			RenderFog(GraphBuilder, SceneTextures, LightShaftOcclusionTexture);
+		}
 
-	// After the height fog, Draw volumetric clouds (having fog applied on them already) when using per pixel tracing,
-	if (!bHasRayTracedOverlay && bShouldRenderVolumetricCloud)
-	{
-		bool bSkipVolumetricRenderTarget = true;
-		bool bSkipPerPixelTracing = false;
-		RenderVolumetricCloud(GraphBuilder, SceneTextures, bSkipVolumetricRenderTarget, bSkipPerPixelTracing, 
-			HalfResolutionDepthCheckerboardMinMaxTexture, QuarterResolutionDepthMinMaxTexture, false, InstanceCullingManager);
-	}
+		// After the height fog, Draw volumetric clouds (having fog applied on them already) when using per pixel tracing,
+		if (!bHasRayTracedOverlay && bShouldRenderVolumetricCloud)
+		{
+			bool bSkipVolumetricRenderTarget = true;
+			bool bSkipPerPixelTracing = false;
+			RenderVolumetricCloud(GraphBuilder, SceneTextures, bSkipVolumetricRenderTarget, bSkipPerPixelTracing, 
+				HalfResolutionDepthCheckerboardMinMaxTexture, QuarterResolutionDepthMinMaxTexture, false, InstanceCullingManager);
+		}
 
-	// or composite the off screen buffer over the scene.
-	if (bVolumetricRenderTargetRequired)
-	{
-		ComposeVolumetricRenderTargetOverScene(GraphBuilder, Views, SceneTextures.Color.Target, SceneTextures.Depth.Target, bShouldRenderSingleLayerWater, SceneWithoutWaterTextures, SceneTextures);
-	}
+		// or composite the off screen buffer over the scene.
+		if (bVolumetricRenderTargetRequired)
+		{
+			ComposeVolumetricRenderTargetOverScene(GraphBuilder, Views, SceneTextures.Color.Target, SceneTextures.Depth.Target, bShouldRenderSingleLayerWater, SceneWithoutWaterTextures, SceneTextures);
+		}
 
-	if (!bHasRayTracedOverlay && ShouldRenderLocalHeightFog(Scene, ViewFamily))
-	{
-		RDG_CSV_STAT_EXCLUSIVE_SCOPE(GraphBuilder, RenderLocalHeightFog);
-		SCOPED_NAMED_EVENT(RenderLocalHeightFog, FColor::Emerald);
-		SCOPE_CYCLE_COUNTER(STAT_FDeferredShadingSceneRenderer_RenderLocalHeightFog);
-		RenderLocalHeightFog(Scene, Views, GraphBuilder, SceneTextures, LightShaftOcclusionTexture);
-	}
+		if (!bHasRayTracedOverlay && ShouldRenderLocalHeightFog(Scene, ViewFamily))
+		{
+			RDG_CSV_STAT_EXCLUSIVE_SCOPE(GraphBuilder, RenderLocalHeightFog);
+			SCOPED_NAMED_EVENT(RenderLocalHeightFog, FColor::Emerald);
+			SCOPE_CYCLE_COUNTER(STAT_FDeferredShadingSceneRenderer_RenderLocalHeightFog);
+			RenderLocalHeightFog(Scene, Views, GraphBuilder, SceneTextures, LightShaftOcclusionTexture);
+		}
 
-	FRDGTextureRef ExposureIlluminance = AddCalculateExposureIlluminancePass(GraphBuilder, Views, SceneTextures, TranslucencyLightingVolumeTextures, ExposureIlluminanceSetup);
+		FRDGTextureRef ExposureIlluminance = AddCalculateExposureIlluminancePass(GraphBuilder, Views, SceneTextures, TranslucencyLightingVolumeTextures, ExposureIlluminanceSetup);
 
-	RenderOpaqueFX(GraphBuilder, GetSceneViews(), FXSystem, SceneTextures.UniformBuffer);
+		RenderOpaqueFX(GraphBuilder, GetSceneViews(), FXSystem, SceneTextures.UniformBuffer);
 
-	FRendererModule& RendererModule = static_cast<FRendererModule&>(GetRendererModule());
-	RendererModule.RenderPostOpaqueExtensions(GraphBuilder, Views, SceneTextures);
+		FRendererModule& RendererModule = static_cast<FRendererModule&>(GetRendererModule());
+		RendererModule.RenderPostOpaqueExtensions(GraphBuilder, Views, SceneTextures);
 
-	if (Scene->GPUScene.ExecuteDeferredGPUWritePass(GraphBuilder, Views, EGPUSceneGPUWritePass::PostOpaqueRendering))
-	{
-		InstanceCullingManager.BeginDeferredCulling(GraphBuilder, Scene->GPUScene);
-	}
+		if (Scene->GPUScene.ExecuteDeferredGPUWritePass(GraphBuilder, Views, EGPUSceneGPUWritePass::PostOpaqueRendering))
+		{
+			InstanceCullingManager.BeginDeferredCulling(GraphBuilder, Scene->GPUScene);
+		}
 
-	if (GetHairStrandsComposition() == EHairStrandsCompositionType::BeforeTranslucent)
-	{
-		RDG_GPU_STAT_SCOPE(GraphBuilder, HairRendering);
-		RenderHairComposition(GraphBuilder, Views, SceneTextures.Color.Target, SceneTextures.Depth.Target, SceneTextures.Velocity, TranslucencyResourceMap);
-	}
+		if (GetHairStrandsComposition() == EHairStrandsCompositionType::BeforeTranslucent)
+		{
+			RDG_GPU_STAT_SCOPE(GraphBuilder, HairRendering);
+			RenderHairComposition(GraphBuilder, Views, SceneTextures.Color.Target, SceneTextures.Depth.Target, SceneTextures.Velocity, TranslucencyResourceMap);
+		}
 
-	// Experimental voxel test code
-	for (int32 ViewIndex = 0; ViewIndex < Views.Num(); ViewIndex++)
-	{
-		const FViewInfo& View = Views[ViewIndex];
-	
-		Nanite::DrawVisibleBricks( GraphBuilder, *Scene, View, SceneTextures );
-	}
+		// Experimental voxel test code
+		for (int32 ViewIndex = 0; ViewIndex < Views.Num(); ViewIndex++)
+		{
+			const FViewInfo& View = Views[ViewIndex];
+		
+			Nanite::DrawVisibleBricks( GraphBuilder, *Scene, View, SceneTextures );
+		}
 
-	// Draw translucency.
-	if (!bHasRayTracedOverlay && TranslucencyViewsToRender != ETranslucencyView::None)
-	{
-		RDG_CSV_STAT_EXCLUSIVE_SCOPE(GraphBuilder, RenderTranslucency);
-		SCOPED_NAMED_EVENT(RenderTranslucency, FColor::Emerald);
-		SCOPE_CYCLE_COUNTER(STAT_TranslucencyDrawTime);
+		// Draw translucency.
+		if (!bHasRayTracedOverlay && TranslucencyViewsToRender != ETranslucencyView::None)
+		{
+			RDG_CSV_STAT_EXCLUSIVE_SCOPE(GraphBuilder, RenderTranslucency);
+			SCOPED_NAMED_EVENT(RenderTranslucency, FColor::Emerald);
+			SCOPE_CYCLE_COUNTER(STAT_TranslucencyDrawTime);
 
-		RDG_EVENT_SCOPE(GraphBuilder, "Translucency");
+			RDG_EVENT_SCOPE(GraphBuilder, "Translucency");
 
-		// Raytracing doesn't need the distortion effect.
-		const bool bShouldRenderDistortion = TranslucencyViewsToRender != ETranslucencyView::RayTracing && ShouldRenderDistortion();
+			// Raytracing doesn't need the distortion effect.
+			const bool bShouldRenderDistortion = TranslucencyViewsToRender != ETranslucencyView::RayTracing && ShouldRenderDistortion();
 
 #if RHI_RAYTRACING
-		if (EnumHasAnyFlags(TranslucencyViewsToRender, ETranslucencyView::RayTracing))
-		{
-			RenderRayTracingTranslucency(GraphBuilder, SceneTextures.Color);
-			EnumRemoveFlags(TranslucencyViewsToRender, ETranslucencyView::RayTracing);
-		}
+			if (EnumHasAnyFlags(TranslucencyViewsToRender, ETranslucencyView::RayTracing))
+			{
+				RenderRayTracingTranslucency(GraphBuilder, SceneTextures.Color);
+				EnumRemoveFlags(TranslucencyViewsToRender, ETranslucencyView::RayTracing);
+			}
 #endif
 
-		// Lumen/VSM translucent front layer
-		FFrontLayerTranslucencyData FrontLayerTranslucencyData = RenderFrontLayerTranslucency(GraphBuilder, Views, SceneTextures, false /*VSM page marking*/);
-		for (FViewInfo& View : Views)
-		{
-			if (GetViewPipelineState(View).ReflectionsMethod == EReflectionsMethod::Lumen)
+			// Lumen/VSM translucent front layer
+			FFrontLayerTranslucencyData FrontLayerTranslucencyData = RenderFrontLayerTranslucency(GraphBuilder, Views, SceneTextures, false /*VSM page marking*/);
+			for (FViewInfo& View : Views)
 			{
-				RenderLumenFrontLayerTranslucencyReflections(GraphBuilder, View, SceneTextures, LumenFrameTemporaries, FrontLayerTranslucencyData);
+				if (GetViewPipelineState(View).ReflectionsMethod == EReflectionsMethod::Lumen)
+				{
+					RenderLumenFrontLayerTranslucencyReflections(GraphBuilder, View, SceneTextures, LumenFrameTemporaries, FrontLayerTranslucencyData);
+				}
 			}
-		}
 
-		// Sort objects' triangles
-		for (FViewInfo& View : Views)
-		{
-			if (OIT::IsEnabled(EOITSortingType::SortedTriangles, View))
+			// Sort objects' triangles
+			for (FViewInfo& View : Views)
 			{
-				OIT::AddSortTrianglesPass(GraphBuilder, View, Scene->OITSceneData, FTriangleSortingOrder::BackToFront);
+				if (OIT::IsEnabled(EOITSortingType::SortedTriangles, View))
+				{
+					OIT::AddSortTrianglesPass(GraphBuilder, View, Scene->OITSceneData, FTriangleSortingOrder::BackToFront);
+				}
 			}
-		}
 
-		{
-			// Render all remaining translucency views.
-			GraphBuilder.SetCommandListStat(GET_STATID(STAT_CLM_Translucency));
-			const bool bStandardTranslucentCanRenderSeparate = bShouldRenderDistortion; // It is only needed to render standard translucent as separate when there is distortion (non self distortion of transmittance/specular/etc.)
-			RenderTranslucency(GraphBuilder, SceneTextures, TranslucencyLightingVolumeTextures, &TranslucencyResourceMap, TranslucencyViewsToRender, InstanceCullingManager, bStandardTranslucentCanRenderSeparate);
-			TranslucencyViewsToRender = ETranslucencyView::None;
-		}
+			{
+				// Render all remaining translucency views.
+				GraphBuilder.SetCommandListStat(GET_STATID(STAT_CLM_Translucency));
+				const bool bStandardTranslucentCanRenderSeparate = bShouldRenderDistortion; // It is only needed to render standard translucent as separate when there is distortion (non self distortion of transmittance/specular/etc.)
+				RenderTranslucency(GraphBuilder, SceneTextures, TranslucencyLightingVolumeTextures, &TranslucencyResourceMap, TranslucencyViewsToRender, InstanceCullingManager, bStandardTranslucentCanRenderSeparate);
+				TranslucencyViewsToRender = ETranslucencyView::None;
+			}
 
-		// Compose hair before velocity/distortion pass since these pass write depth value, 
-		// and this would make the hair composition fails in this cases.
-		if (GetHairStrandsComposition() == EHairStrandsCompositionType::AfterTranslucent)
+			// Compose hair before velocity/distortion pass since these pass write depth value, 
+			// and this would make the hair composition fails in this cases.
+			if (GetHairStrandsComposition() == EHairStrandsCompositionType::AfterTranslucent)
+			{
+				RDG_GPU_STAT_SCOPE(GraphBuilder, HairRendering);
+				RenderHairComposition(GraphBuilder, Views, SceneTextures.Color.Target, SceneTextures.Depth.Target, SceneTextures.Velocity, TranslucencyResourceMap);
+			}
+
+			if (bShouldRenderDistortion)
+			{
+				GraphBuilder.SetCommandListStat(GET_STATID(STAT_CLM_Distortion));
+				RenderDistortion(GraphBuilder, SceneTextures.Color.Target, SceneTextures.Depth.Target, SceneTextures.Velocity, TranslucencyResourceMap);
+			}
+
+			if (bShouldRenderVelocities && CVarTranslucencyVelocity.GetValueOnRenderThread() != 0)
+			{
+				const bool bRecreateSceneTextures = !HasBeenProduced(SceneTextures.Velocity);
+
+				GraphBuilder.SetCommandListStat(GET_STATID(STAT_CLM_TranslucentVelocity));
+				RenderVelocities(GraphBuilder, SceneTextures, EVelocityPass::Translucent, false);
+
+				if (bRecreateSceneTextures)
+				{
+					// Rebuild scene textures to include newly allocated velocity.
+					SceneTextures.UniformBuffer = CreateSceneTextureUniformBuffer(GraphBuilder, &SceneTextures, FeatureLevel, SceneTextures.SetupMode);
+				}
+			}
+
+			GraphBuilder.SetCommandListStat(GET_STATID(STAT_CLM_AfterTranslucency));
+		}
+		else if (GetHairStrandsComposition() == EHairStrandsCompositionType::AfterTranslucent)
 		{
 			RDG_GPU_STAT_SCOPE(GraphBuilder, HairRendering);
 			RenderHairComposition(GraphBuilder, Views, SceneTextures.Color.Target, SceneTextures.Depth.Target, SceneTextures.Velocity, TranslucencyResourceMap);
 		}
 
-		if (bShouldRenderDistortion)
+#if !UE_BUILD_SHIPPING
+		if (CVarForceBlackVelocityBuffer.GetValueOnRenderThread())
 		{
-			GraphBuilder.SetCommandListStat(GET_STATID(STAT_CLM_Distortion));
-			RenderDistortion(GraphBuilder, SceneTextures.Color.Target, SceneTextures.Depth.Target, SceneTextures.Velocity, TranslucencyResourceMap);
+			SceneTextures.Velocity = SystemTextures.Black;
+
+			// Rebuild the scene texture uniform buffer to include black.
+			SceneTextures.UniformBuffer = CreateSceneTextureUniformBuffer(GraphBuilder, &SceneTextures, FeatureLevel, SceneTextures.SetupMode);
 		}
+#endif
 
-		if (bShouldRenderVelocities && CVarTranslucencyVelocity.GetValueOnRenderThread() != 0)
 		{
-			const bool bRecreateSceneTextures = !HasBeenProduced(SceneTextures.Velocity);
-
-			GraphBuilder.SetCommandListStat(GET_STATID(STAT_CLM_TranslucentVelocity));
-			RenderVelocities(GraphBuilder, SceneTextures, EVelocityPass::Translucent, false);
-
-			if (bRecreateSceneTextures)
+			if (HairStrandsBookmarkParameters.HasInstances())
 			{
-				// Rebuild scene textures to include newly allocated velocity.
-				SceneTextures.UniformBuffer = CreateSceneTextureUniformBuffer(GraphBuilder, &SceneTextures, FeatureLevel, SceneTextures.SetupMode);
+				RenderHairStrandsDebugInfo(GraphBuilder, Scene, Views, HairStrandsBookmarkParameters.HairClusterData, SceneTextures.Color.Target, SceneTextures.Depth.Target);
 			}
 		}
 
-		GraphBuilder.SetCommandListStat(GET_STATID(STAT_CLM_AfterTranslucency));
-	}
-	else if (GetHairStrandsComposition() == EHairStrandsCompositionType::AfterTranslucent)
-	{
-		RDG_GPU_STAT_SCOPE(GraphBuilder, HairRendering);
-		RenderHairComposition(GraphBuilder, Views, SceneTextures.Color.Target, SceneTextures.Depth.Target, SceneTextures.Velocity, TranslucencyResourceMap);
-	}
-
-#if !UE_BUILD_SHIPPING
-	if (CVarForceBlackVelocityBuffer.GetValueOnRenderThread())
-	{
-		SceneTextures.Velocity = SystemTextures.Black;
-
-		// Rebuild the scene texture uniform buffer to include black.
-		SceneTextures.UniformBuffer = CreateSceneTextureUniformBuffer(GraphBuilder, &SceneTextures, FeatureLevel, SceneTextures.SetupMode);
-	}
-#endif
-
-	{
-		if (HairStrandsBookmarkParameters.HasInstances())
+		if (VirtualShadowMapArray.IsEnabled())
 		{
-			RenderHairStrandsDebugInfo(GraphBuilder, Scene, Views, HairStrandsBookmarkParameters.HairClusterData, SceneTextures.Color.Target, SceneTextures.Depth.Target);
+			VirtualShadowMapArray.RenderDebugInfo(GraphBuilder, Views);
 		}
-	}
 
-	if (VirtualShadowMapArray.IsEnabled())
-	{
-		VirtualShadowMapArray.RenderDebugInfo(GraphBuilder, Views);
-	}
-
-	for (FViewInfo& View : Views)
-	{
-		ShadingEnergyConservation::Debug(GraphBuilder, View, SceneTextures);
-	}
+		for (FViewInfo& View : Views)
+		{
+			ShadingEnergyConservation::Debug(GraphBuilder, View, SceneTextures);
+		}
 
-	Scene->ProcessAndRenderIlluminanceMeter(GraphBuilder, Views, SceneTextures.Color.Target);
+		Scene->ProcessAndRenderIlluminanceMeter(GraphBuilder, Views, SceneTextures.Color.Target);
 
-	if (!bHasRayTracedOverlay && ViewFamily.EngineShowFlags.LightShafts)
-	{
-		SCOPE_CYCLE_COUNTER(STAT_FDeferredShadingSceneRenderer_RenderLightShaftBloom);
-		GraphBuilder.SetCommandListStat(GET_STATID(STAT_CLM_LightShaftBloom));
-		RenderLightShaftBloom(GraphBuilder, SceneTextures, /* inout */ TranslucencyResourceMap);
-	}
+		if (!bHasRayTracedOverlay && ViewFamily.EngineShowFlags.LightShafts)
+		{
+			SCOPE_CYCLE_COUNTER(STAT_FDeferredShadingSceneRenderer_RenderLightShaftBloom);
+			GraphBuilder.SetCommandListStat(GET_STATID(STAT_CLM_LightShaftBloom));
+			RenderLightShaftBloom(GraphBuilder, SceneTextures, /* inout */ TranslucencyResourceMap);
+		}
 
-	if (bUseVirtualTexturing)
-	{
-		RDG_GPU_STAT_SCOPE(GraphBuilder, VirtualTextureUpdate);
-		VirtualTextureFeedbackEnd(GraphBuilder);
-	}
+		if (bUseVirtualTexturing)
+		{
+			RDG_GPU_STAT_SCOPE(GraphBuilder, VirtualTextureUpdate);
+			VirtualTextureFeedbackEnd(GraphBuilder);
+		}
 
-	FPathTracingResources PathTracingResources;
+		FPathTracingResources PathTracingResources;
 
 #if RHI_RAYTRACING
-	if (IsRayTracingEnabled())
-	{
-		// Path tracer requires the full ray tracing pipeline support, as well as specialized extra shaders.
-		// Most of the ray tracing debug visualizations also require the full pipeline, but some support inline mode.
-		
-		if (ViewFamily.EngineShowFlags.PathTracing 
-			&& FDataDrivenShaderPlatformInfo::GetSupportsPathTracing(Scene->GetShaderPlatform()))
+		if (IsRayTracingEnabled())
 		{
-			for (const FViewInfo& View : Views)
+			// Path tracer requires the full ray tracing pipeline support, as well as specialized extra shaders.
+			// Most of the ray tracing debug visualizations also require the full pipeline, but some support inline mode.
+			
+			if (ViewFamily.EngineShowFlags.PathTracing 
+				&& FDataDrivenShaderPlatformInfo::GetSupportsPathTracing(Scene->GetShaderPlatform()))
 			{
-				RenderPathTracing(GraphBuilder, View, SceneTextures.UniformBuffer, SceneTextures.Color.Target, SceneTextures.Depth.Target,PathTracingResources);
+				for (const FViewInfo& View : Views)
+				{
+					RenderPathTracing(GraphBuilder, View, SceneTextures.UniformBuffer, SceneTextures.Color.Target, SceneTextures.Depth.Target,PathTracingResources);
+				}
 			}
-		}
-		else if (ViewFamily.EngineShowFlags.RayTracingDebug)
-		{
-			for (const FViewInfo& View : Views)
+			else if (ViewFamily.EngineShowFlags.RayTracingDebug)
 			{
-				FRayTracingPickingFeedback PickingFeedback = {};
-				RenderRayTracingDebug(GraphBuilder, View, SceneTextures.Color.Target, PickingFeedback);
+				for (const FViewInfo& View : Views)
+				{
+					FRayTracingPickingFeedback PickingFeedback = {};
+					RenderRayTracingDebug(GraphBuilder, View, SceneTextures.Color.Target, PickingFeedback);
 
-				OnGetOnScreenMessages.AddLambda([this, PickingFeedback](FScreenMessageWriter& ScreenMessageWriter)->void
-					{
-						RayTracingDisplayPicking(PickingFeedback, ScreenMessageWriter);
-					});
+					OnGetOnScreenMessages.AddLambda([this, PickingFeedback](FScreenMessageWriter& ScreenMessageWriter)->void
+						{
+							RayTracingDisplayPicking(PickingFeedback, ScreenMessageWriter);
+						});
+				}
 			}
 		}
-	}
 #endif
 
-	RendererModule.RenderOverlayExtensions(GraphBuilder, Views, SceneTextures);
-
-	if (ViewFamily.EngineShowFlags.PhysicsField && Scene->PhysicsField)
-	{
-		RenderPhysicsField(GraphBuilder, Views, Scene->PhysicsField, SceneTextures.Color.Target);
-	}
-
-	if (ViewFamily.EngineShowFlags.VisualizeDistanceFieldAO && ShouldRenderDistanceFieldLighting())
-	{
-		GraphBuilder.SetCommandListStat(GET_STATID(STAT_CLM_RenderDistanceFieldLighting));
-
-		// Use the skylight's max distance if there is one, to be consistent with DFAO shadowing on the skylight
-		const float OcclusionMaxDistance = Scene->SkyLight && !Scene->SkyLight->bWantsStaticShadowing ? Scene->SkyLight->OcclusionMaxDistance : Scene->DefaultMaxDistanceFieldOcclusionDistance;
-		FRDGTextureRef DummyOutput = nullptr;
-		RenderDistanceFieldLighting(GraphBuilder, SceneTextures, FDistanceFieldAOParameters(OcclusionMaxDistance), DummyOutput, false, ViewFamily.EngineShowFlags.VisualizeDistanceFieldAO);
-	}
-
-	// Draw visualizations just before use to avoid target contamination
-	if (ViewFamily.EngineShowFlags.VisualizeMeshDistanceFields || ViewFamily.EngineShowFlags.VisualizeGlobalDistanceField)
-	{
-		RenderMeshDistanceFieldVisualization(GraphBuilder, SceneTextures, FDistanceFieldAOParameters(Scene->DefaultMaxDistanceFieldOcclusionDistance));
-	}
-
-	if (bRenderDeferredLighting)
-	{
-		RenderLumenMiscVisualizations(GraphBuilder, SceneTextures, LumenFrameTemporaries);
-		RenderDiffuseIndirectAndAmbientOcclusion(
-			GraphBuilder,
-			SceneTextures,
-			LumenFrameTemporaries,
-			LightingChannelsTexture,
-			/* bHasLumenLights = */ false,
-			/* bCompositeRegularLumenOnly = */ false,
-			/* bIsVisualizePass = */ true,
-			AsyncLumenIndirectLightingOutputs);
-	}
-
-	if (ViewFamily.EngineShowFlags.StationaryLightOverlap)
-	{
-		RenderStationaryLightOverlap(GraphBuilder, SceneTextures, LightingChannelsTexture);
-	}
-
-	if (ShouldRenderHeterogeneousVolumes(Scene) && !bHasRayTracedOverlay)
-	{
-		CompositeHeterogeneousVolumes(GraphBuilder, SceneTextures);
-	}
+		RendererModule.RenderOverlayExtensions(GraphBuilder, Views, SceneTextures);
 
-	if (bShouldVisualizeVolumetricCloud && !bHasRayTracedOverlay)
-	{
-		RenderVolumetricCloud(GraphBuilder, SceneTextures, false, true, HalfResolutionDepthCheckerboardMinMaxTexture, QuarterResolutionDepthMinMaxTexture, false, InstanceCullingManager);
-		ReconstructVolumetricRenderTarget(GraphBuilder, Views, SceneTextures.Depth.Resolve, HalfResolutionDepthCheckerboardMinMaxTexture, false);
-		ComposeVolumetricRenderTargetOverSceneForVisualization(GraphBuilder, Views, SceneTextures.Color.Target, SceneTextures);
-		RenderVolumetricCloud(GraphBuilder, SceneTextures, true, false, HalfResolutionDepthCheckerboardMinMaxTexture, QuarterResolutionDepthMinMaxTexture, false, InstanceCullingManager);
-	}
+		if (ViewFamily.EngineShowFlags.PhysicsField && Scene->PhysicsField)
+		{
+			RenderPhysicsField(GraphBuilder, Views, Scene->PhysicsField, SceneTextures.Color.Target);
+		}
 
-	if (!bHasRayTracedOverlay)
-	{
-		AddSparseVolumeTextureViewerRenderPass(GraphBuilder, *this, SceneTextures);
-	}
+		if (ViewFamily.EngineShowFlags.VisualizeDistanceFieldAO && ShouldRenderDistanceFieldLighting())
+		{
+			GraphBuilder.SetCommandListStat(GET_STATID(STAT_CLM_RenderDistanceFieldLighting));
 
-	// Resolve the scene color for post processing.
-	AddResolveSceneColorPass(GraphBuilder, Views, SceneTextures.Color);
+			// Use the skylight's max distance if there is one, to be consistent with DFAO shadowing on the skylight
+			const float OcclusionMaxDistance = Scene->SkyLight && !Scene->SkyLight->bWantsStaticShadowing ? Scene->SkyLight->OcclusionMaxDistance : Scene->DefaultMaxDistanceFieldOcclusionDistance;
+			FRDGTextureRef DummyOutput = nullptr;
+			RenderDistanceFieldLighting(GraphBuilder, SceneTextures, FDistanceFieldAOParameters(OcclusionMaxDistance), DummyOutput, false, ViewFamily.EngineShowFlags.VisualizeDistanceFieldAO);
+		}
 
-	RendererModule.RenderPostResolvedSceneColorExtension(GraphBuilder, SceneTextures);
+		// Draw visualizations just before use to avoid target contamination
+		if (ViewFamily.EngineShowFlags.VisualizeMeshDistanceFields || ViewFamily.EngineShowFlags.VisualizeGlobalDistanceField)
+		{
+			RenderMeshDistanceFieldVisualization(GraphBuilder, SceneTextures, FDistanceFieldAOParameters(Scene->DefaultMaxDistanceFieldOcclusionDistance));
+		}
 
-	FRDGTextureRef ViewFamilyTexture = TryCreateViewFamilyTexture(GraphBuilder, ViewFamily);
+		if (bRenderDeferredLighting)
+		{
+			RenderLumenMiscVisualizations(GraphBuilder, SceneTextures, LumenFrameTemporaries);
+			RenderDiffuseIndirectAndAmbientOcclusion(
+				GraphBuilder,
+				SceneTextures,
+				LumenFrameTemporaries,
+				LightingChannelsTexture,
+				/* bHasLumenLights = */ false,
+				/* bCompositeRegularLumenOnly = */ false,
+				/* bIsVisualizePass = */ true,
+				AsyncLumenIndirectLightingOutputs);
+		}
 
-	CopySceneCaptureComponentToTarget(GraphBuilder, SceneTextures, ViewFamilyTexture, ViewFamily, Views);
+		if (ViewFamily.EngineShowFlags.StationaryLightOverlap)
+		{
+			RenderStationaryLightOverlap(GraphBuilder, SceneTextures, LightingChannelsTexture);
+		}
 
-	for (int32 ViewIndex = 0; ViewIndex < Views.Num(); ++ViewIndex)
-	{
-		const FViewInfo& View = Views[ViewIndex];
+		if (ShouldRenderHeterogeneousVolumes(Scene) && !bHasRayTracedOverlay)
+		{
+			CompositeHeterogeneousVolumes(GraphBuilder, SceneTextures);
+		}
 
-		if (((View.FinalPostProcessSettings.DynamicGlobalIlluminationMethod == EDynamicGlobalIlluminationMethod::ScreenSpace && ScreenSpaceRayTracing::ShouldKeepBleedFreeSceneColor(View))
-			|| GetViewPipelineState(View).DiffuseIndirectMethod == EDiffuseIndirectMethod::Lumen
-			|| GetViewPipelineState(View).ReflectionsMethod == EReflectionsMethod::Lumen)
-			&& !View.bStatePrevViewInfoIsReadOnly)
+		if (bShouldVisualizeVolumetricCloud && !bHasRayTracedOverlay)
 		{
-			// Keep scene color and depth for next frame screen space ray tracing.
-			FSceneViewState* ViewState = View.ViewState;
-			GraphBuilder.QueueTextureExtraction(SceneTextures.Depth.Resolve, &ViewState->PrevFrameViewInfo.DepthBuffer);
-			GraphBuilder.QueueTextureExtraction(SceneTextures.Color.Resolve, &ViewState->PrevFrameViewInfo.ScreenSpaceRayTracingInput);
+			RenderVolumetricCloud(GraphBuilder, SceneTextures, false, true, HalfResolutionDepthCheckerboardMinMaxTexture, QuarterResolutionDepthMinMaxTexture, false, InstanceCullingManager);
+			ReconstructVolumetricRenderTarget(GraphBuilder, Views, SceneTextures.Depth.Resolve, HalfResolutionDepthCheckerboardMinMaxTexture, false);
+			ComposeVolumetricRenderTargetOverSceneForVisualization(GraphBuilder, Views, SceneTextures.Color.Target, SceneTextures);
+			RenderVolumetricCloud(GraphBuilder, SceneTextures, true, false, HalfResolutionDepthCheckerboardMinMaxTexture, QuarterResolutionDepthMinMaxTexture, false, InstanceCullingManager);
 		}
-	}
 
-	// Finish rendering for each view.
-	if (ViewFamily.bResolveScene && ViewFamilyTexture)
-	{
-		RDG_EVENT_SCOPE(GraphBuilder, "PostProcessing");
-		RDG_GPU_STAT_SCOPE(GraphBuilder, Postprocessing);
-		SCOPED_NAMED_EVENT(PostProcessing, FColor::Emerald);
+		if (!bHasRayTracedOverlay)
+		{
+			AddSparseVolumeTextureViewerRenderPass(GraphBuilder, *this, SceneTextures);
+		}
 
-		GraphBuilder.SetCommandListStat(GET_STATID(STAT_CLM_PostProcessing));
+		// Resolve the scene color for post processing.
+		AddResolveSceneColorPass(GraphBuilder, Views, SceneTextures.Color);
 
-		FPostProcessingInputs PostProcessingInputs;
-		PostProcessingInputs.ViewFamilyTexture = ViewFamilyTexture;
-		PostProcessingInputs.CustomDepthTexture = SceneTextures.CustomDepth.Depth;
-		PostProcessingInputs.ExposureIlluminance = ExposureIlluminance;
-		PostProcessingInputs.SceneTextures = SceneTextures.UniformBuffer;
-		PostProcessingInputs.bSeparateCustomStencil = SceneTextures.CustomDepth.bSeparateStencilBuffer;
-		PostProcessingInputs.PathTracingResources = PathTracingResources;
+		RendererModule.RenderPostResolvedSceneColorExtension(GraphBuilder, SceneTextures);
 
-		GraphBuilder.FlushSetupQueue();
+		CopySceneCaptureComponentToTarget(GraphBuilder, SceneTextures, ViewFamilyTexture, ViewFamily, Views);
 
-		if (ViewFamily.UseDebugViewPS())
+		for (int32 ViewIndex = 0; ViewIndex < Views.Num(); ++ViewIndex)
 		{
-			for (int32 ViewIndex = 0; ViewIndex < Views.Num(); ViewIndex++)
+			const FViewInfo& View = Views[ViewIndex];
+
+			if (((View.FinalPostProcessSettings.DynamicGlobalIlluminationMethod == EDynamicGlobalIlluminationMethod::ScreenSpace && ScreenSpaceRayTracing::ShouldKeepBleedFreeSceneColor(View))
+				|| GetViewPipelineState(View).DiffuseIndirectMethod == EDiffuseIndirectMethod::Lumen
+				|| GetViewPipelineState(View).ReflectionsMethod == EReflectionsMethod::Lumen)
+				&& !View.bStatePrevViewInfoIsReadOnly)
 			{
-				const FViewInfo& View = Views[ViewIndex];
-   				const Nanite::FRasterResults* NaniteResults = bNaniteEnabled ? &NaniteRasterResults[ViewIndex] : nullptr;
-				RDG_GPU_MASK_SCOPE(GraphBuilder, View.GPUMask);
-				RDG_EVENT_SCOPE_CONDITIONAL(GraphBuilder, Views.Num() > 1, "View%d", ViewIndex);
-				PostProcessingInputs.TranslucencyViewResourcesMap = FTranslucencyViewResourcesMap(TranslucencyResourceMap, ViewIndex);
-				AddDebugViewPostProcessingPasses(GraphBuilder, View, GetSceneUniforms(), PostProcessingInputs, NaniteResults);
+				// Keep scene color and depth for next frame screen space ray tracing.
+				FSceneViewState* ViewState = View.ViewState;
+				GraphBuilder.QueueTextureExtraction(SceneTextures.Depth.Resolve, &ViewState->PrevFrameViewInfo.DepthBuffer);
+				GraphBuilder.QueueTextureExtraction(SceneTextures.Color.Resolve, &ViewState->PrevFrameViewInfo.ScreenSpaceRayTracingInput);
 			}
 		}
-		else
+
+		// Finish rendering for each view.
+		if (ViewFamily.bResolveScene && ViewFamilyTexture)
 		{
-			for (int32 ViewExt = 0; ViewExt < ViewFamily.ViewExtensions.Num(); ++ViewExt)
+			RDG_EVENT_SCOPE(GraphBuilder, "PostProcessing");
+			RDG_GPU_STAT_SCOPE(GraphBuilder, Postprocessing);
+			SCOPED_NAMED_EVENT(PostProcessing, FColor::Emerald);
+
+			GraphBuilder.SetCommandListStat(GET_STATID(STAT_CLM_PostProcessing));
+
+			FPostProcessingInputs PostProcessingInputs;
+			PostProcessingInputs.ViewFamilyTexture = ViewFamilyTexture;
+			PostProcessingInputs.CustomDepthTexture = SceneTextures.CustomDepth.Depth;
+			PostProcessingInputs.ExposureIlluminance = ExposureIlluminance;
+			PostProcessingInputs.SceneTextures = SceneTextures.UniformBuffer;
+			PostProcessingInputs.bSeparateCustomStencil = SceneTextures.CustomDepth.bSeparateStencilBuffer;
+			PostProcessingInputs.PathTracingResources = PathTracingResources;
+
+			GraphBuilder.FlushSetupQueue();
+
+			if (ViewFamily.UseDebugViewPS())
 			{
-				for (int32 ViewIndex = 0; ViewIndex < ViewFamily.Views.Num(); ++ViewIndex)
+				for (int32 ViewIndex = 0; ViewIndex < Views.Num(); ViewIndex++)
 				{
-					FViewInfo& View = Views[ViewIndex];
+					const FViewInfo& View = Views[ViewIndex];
+   					const Nanite::FRasterResults* NaniteResults = bNaniteEnabled ? &NaniteRasterResults[ViewIndex] : nullptr;
 					RDG_GPU_MASK_SCOPE(GraphBuilder, View.GPUMask);
+					RDG_EVENT_SCOPE_CONDITIONAL(GraphBuilder, Views.Num() > 1, "View%d", ViewIndex);
 					PostProcessingInputs.TranslucencyViewResourcesMap = FTranslucencyViewResourcesMap(TranslucencyResourceMap, ViewIndex);
-					ViewFamily.ViewExtensions[ViewExt]->PrePostProcessPass_RenderThread(GraphBuilder, View, PostProcessingInputs);
+					AddDebugViewPostProcessingPasses(GraphBuilder, View, GetSceneUniforms(), PostProcessingInputs, NaniteResults);
 				}
 			}
-			for (int32 ViewIndex = 0; ViewIndex < Views.Num(); ViewIndex++)
+			else
 			{
-				const FViewInfo& View = Views[ViewIndex];
-				const Nanite::FRasterResults* NaniteResults = bNaniteEnabled ? &NaniteRasterResults[ViewIndex] : nullptr;
-				RDG_GPU_MASK_SCOPE(GraphBuilder, View.GPUMask);
-				RDG_EVENT_SCOPE_CONDITIONAL(GraphBuilder, Views.Num() > 1, "View%d", ViewIndex);
-
-				PostProcessingInputs.TranslucencyViewResourcesMap = FTranslucencyViewResourcesMap(TranslucencyResourceMap, ViewIndex);
-
-				if (IsPostProcessVisualizeCalibrationMaterialEnabled(View))
+				for (int32 ViewExt = 0; ViewExt < ViewFamily.ViewExtensions.Num(); ++ViewExt)
 				{
-					const UMaterialInterface* DebugMaterialInterface = GetPostProcessVisualizeCalibrationMaterialInterface(View);
-					check(DebugMaterialInterface);
-
-					AddVisualizeCalibrationMaterialPostProcessingPasses(GraphBuilder, View, PostProcessingInputs, DebugMaterialInterface);
+					for (int32 ViewIndex = 0; ViewIndex < ViewFamily.Views.Num(); ++ViewIndex)
+					{
+						FViewInfo& View = Views[ViewIndex];
+						RDG_GPU_MASK_SCOPE(GraphBuilder, View.GPUMask);
+						PostProcessingInputs.TranslucencyViewResourcesMap = FTranslucencyViewResourcesMap(TranslucencyResourceMap, ViewIndex);
+						ViewFamily.ViewExtensions[ViewExt]->PrePostProcessPass_RenderThread(GraphBuilder, View, PostProcessingInputs);
+					}
 				}
-				else
+				for (int32 ViewIndex = 0; ViewIndex < Views.Num(); ViewIndex++)
 				{
-					const FPerViewPipelineState& ViewPipelineState = GetViewPipelineState(View);
-					const bool bAnyLumenActive = ViewPipelineState.DiffuseIndirectMethod == EDiffuseIndirectMethod::Lumen || ViewPipelineState.ReflectionsMethod == EReflectionsMethod::Lumen;
-					const bool bLumenGIEnabled = ViewPipelineState.DiffuseIndirectMethod == EDiffuseIndirectMethod::Lumen;
+					const FViewInfo& View = Views[ViewIndex];
+					const Nanite::FRasterResults* NaniteResults = bNaniteEnabled ? &NaniteRasterResults[ViewIndex] : nullptr;
+					RDG_GPU_MASK_SCOPE(GraphBuilder, View.GPUMask);
+					RDG_EVENT_SCOPE_CONDITIONAL(GraphBuilder, Views.Num() > 1, "View%d", ViewIndex);
 
-					FScreenPassTexture TSRMoireInput;
-					if (ViewIndex < TSRMoireInputTextures.Num())
+					PostProcessingInputs.TranslucencyViewResourcesMap = FTranslucencyViewResourcesMap(TranslucencyResourceMap, ViewIndex);
+
+					if (IsPostProcessVisualizeCalibrationMaterialEnabled(View))
 					{
-						TSRMoireInput = TSRMoireInputTextures[ViewIndex];
+						const UMaterialInterface* DebugMaterialInterface = GetPostProcessVisualizeCalibrationMaterialInterface(View);
+						check(DebugMaterialInterface);
+
+						AddVisualizeCalibrationMaterialPostProcessingPasses(GraphBuilder, View, PostProcessingInputs, DebugMaterialInterface);
 					}
+					else
+					{
+						const FPerViewPipelineState& ViewPipelineState = GetViewPipelineState(View);
+						const bool bAnyLumenActive = ViewPipelineState.DiffuseIndirectMethod == EDiffuseIndirectMethod::Lumen || ViewPipelineState.ReflectionsMethod == EReflectionsMethod::Lumen;
+						const bool bLumenGIEnabled = ViewPipelineState.DiffuseIndirectMethod == EDiffuseIndirectMethod::Lumen;
 
-					AddPostProcessingPasses(
-						GraphBuilder,
-						View, ViewIndex,
-						GetSceneUniforms(),
-						bAnyLumenActive,
-						bLumenGIEnabled,
-						ViewPipelineState.ReflectionsMethod,
-						PostProcessingInputs,
-						NaniteResults,
-						InstanceCullingManager,
-						&VirtualShadowMapArray,
-						LumenFrameTemporaries,
-						SceneWithoutWaterTextures,
-						TSRMoireInput);
+						FScreenPassTexture TSRMoireInput;
+						if (ViewIndex < TSRMoireInputTextures.Num())
+						{
+							TSRMoireInput = TSRMoireInputTextures[ViewIndex];
+						}
+
+						AddPostProcessingPasses(
+							GraphBuilder,
+							View, ViewIndex,
+							GetSceneUniforms(),
+							bAnyLumenActive,
+							bLumenGIEnabled,
+							ViewPipelineState.ReflectionsMethod,
+							PostProcessingInputs,
+							NaniteResults,
+							InstanceCullingManager,
+							&VirtualShadowMapArray,
+							LumenFrameTemporaries,
+							SceneWithoutWaterTextures,
+							TSRMoireInput);
+					}
 				}
 			}
 		}
-	}
 
-	// After AddPostProcessingPasses in case of Lumen Visualizations writing to feedback
-	FinishGatheringLumenSurfaceCacheFeedback(GraphBuilder, Views[0], LumenFrameTemporaries);
+		// After AddPostProcessingPasses in case of Lumen Visualizations writing to feedback
+		FinishGatheringLumenSurfaceCacheFeedback(GraphBuilder, Views[0], LumenFrameTemporaries);
 
-	if (ViewFamily.bResolveScene && ViewFamilyTexture)
-	{
-		GVRSImageManager.DrawDebugPreview(GraphBuilder, ViewFamily, ViewFamilyTexture);
-	}
+		if (ViewFamily.bResolveScene && ViewFamilyTexture)
+		{
+			GVRSImageManager.DrawDebugPreview(GraphBuilder, ViewFamily, ViewFamilyTexture);
+		}
 
-	GEngine->GetPostRenderDelegateEx().Broadcast(GraphBuilder);
+		GEngine->GetPostRenderDelegateEx().Broadcast(GraphBuilder);
 
 #if RHI_RAYTRACING
-	ReleaseRaytracingResources(GraphBuilder, Views, Scene->RayTracingScene, bIsLastSceneRenderer);
+		ReleaseRaytracingResources(GraphBuilder, Views, Scene->RayTracingScene, bIsLastSceneRenderer);
 #endif //  RHI_RAYTRACING
+	}
 
 #if WITH_MGPU
 	DoCrossGPUTransfers(GraphBuilder, ViewFamilyTexture);
diff --git a/Engine/Source/Runtime/Renderer/Private/DeferredShadingRenderer.h b/Engine/Source/Runtime/Renderer/Private/DeferredShadingRenderer.h
--- a/Engine/Source/Runtime/Renderer/Private/DeferredShadingRenderer.h
+++ b/Engine/Source/Runtime/Renderer/Private/DeferredShadingRenderer.h
@@ -319,6 +319,12 @@ public:
 	void InitializeRayTracingFlags_RenderThread();
 #endif
 
+	enum class ERendererOutput
+	{
+		DepthPrepassOnly,	// Only render depth prepass and its related code paths
+		FinalSceneColor		// Render the whole pipeline
+	};
+
 	FDeferredShadingSceneRenderer(const FSceneViewFamily* InViewFamily, FHitProxyConsumer* HitProxyConsumer);
 
 	/** Determine and commit the final state of the pipeline for the view family and views. */
