Moves ITemporalUpscaler to Renderer/Public/

The motivation of this change are:
* Avoid precompiled third party temporal upscaler plugin incompatibility between minor versions of engine like happened between 5.1.0 and 5.1.1;
* A third party temporal upscaler plugin can be fully implemented without needing to include a single private header of the renderer since UBT would like to no longer allow this.

diff --git a/Engine/Source/Runtime/Renderer/Private/PostProcess/PostProcessSubsurface.cpp b/Engine/Source/Runtime/Renderer/Private/PostProcess/PostProcessSubsurface.cpp
--- a/Engine/Source/Runtime/Renderer/Private/PostProcess/PostProcessSubsurface.cpp
+++ b/Engine/Source/Runtime/Renderer/Private/PostProcess/PostProcessSubsurface.cpp
@@ -880,7 +880,7 @@ class FSubsurfaceRecombinePS : public FSubsurfaceShader
 	{
 		const uint32 QualityCVar = GetSSSQuality();
 
-		EMainTAAPassConfig MainTAAConfig = ITemporalUpscaler::GetMainTAAPassConfig(View);
+		EMainTAAPassConfig MainTAAConfig = GetMainTAAPassConfig(View);
 
 		// Low quality is really bad with modern temporal upscalers.
 		bool bAllowLowQuality = MainTAAConfig == EMainTAAPassConfig::Disabled || MainTAAConfig == EMainTAAPassConfig::TAA;
diff --git a/Engine/Source/Runtime/Renderer/Private/PostProcess/PostProcessing.cpp b/Engine/Source/Runtime/Renderer/Private/PostProcess/PostProcessing.cpp
--- a/Engine/Source/Runtime/Renderer/Private/PostProcess/PostProcessing.cpp
+++ b/Engine/Source/Runtime/Renderer/Private/PostProcess/PostProcessing.cpp
@@ -64,6 +64,7 @@
 #include "FXSystem.h"
 #include "SkyAtmosphereRendering.h"
 #include "Strata/Strata.h"
+#include "TemporalUpscaler.h"
 #include "VirtualShadowMaps/VirtualShadowMapArray.h"
 #include "Lumen/LumenVisualize.h"
 #include "RectLightTextureManager.h"
@@ -194,13 +195,77 @@ bool IsPostProcessingWithAlphaChannelSupported()
 FScreenPassTexture AddFinalPostProcessDebugInfoPasses(FRDGBuilder& GraphBuilder, const FViewInfo& View, FScreenPassTexture& ScreenPassSceneColor);
 #endif
 
-void AddTemporalAA2Passes(
+
+FDefaultTemporalUpscaler::FOutputs AddThirdPartyTemporalUpscalerPasses(
 	FRDGBuilder& GraphBuilder,
-	const FSceneTextureParameters& SceneTextures,
 	const FViewInfo& View,
-	FRDGTextureRef InputSceneColorTexture,
-	FRDGTextureRef* OutSceneColorTexture,
-	FIntRect* OutSceneColorViewRect);
+	const FDefaultTemporalUpscaler::FInputs& Inputs)
+{
+	const ITemporalUpscaler* UpscalerToUse = View.Family->GetTemporalUpscalerInterface();
+	check(UpscalerToUse);
+
+	const TCHAR* UpscalerName = UpscalerToUse->GetDebugName();
+
+	// Translate the inputs to the third party temporal upscaler.
+	ITemporalUpscaler::FInputs ThirdPartyInputs;
+	ThirdPartyInputs.OutputViewRect.Min = FIntPoint::ZeroValue;
+	ThirdPartyInputs.OutputViewRect.Max = View.GetSecondaryViewRectSize();
+	ThirdPartyInputs.TemporalJitterPixels = FVector2f(View.TemporalJitterPixels);
+	ThirdPartyInputs.PreExposure = View.PreExposure;
+	ThirdPartyInputs.SceneColor = Inputs.SceneColor;
+	ThirdPartyInputs.SceneDepth = Inputs.SceneDepth;
+	ThirdPartyInputs.SceneVelocity = Inputs.SceneVelocity;
+	ThirdPartyInputs.EyeAdaptationTexture = AddCopyEyeAdaptationDataToTexturePass(GraphBuilder, View);
+	
+	if (View.PrevViewInfo.ThirdPartyTemporalUpscalerHistory && View.PrevViewInfo.ThirdPartyTemporalUpscalerHistory->GetDebugName() == UpscalerName)
+	{
+		ThirdPartyInputs.PrevHistory = View.PrevViewInfo.ThirdPartyTemporalUpscalerHistory;
+	}
+
+	// Standard event scope for temporal upscaler to have all profiling information not matter what,
+	// and with explicit detection of third party.
+	RDG_EVENT_SCOPE(
+		GraphBuilder,
+		"ThirdParty %s %dx%d -> %dx%d",
+		UpscalerToUse->GetDebugName(),
+		View.ViewRect.Width(), View.ViewRect.Height(),
+		ThirdPartyInputs.OutputViewRect.Width(), ThirdPartyInputs.OutputViewRect.Height());
+
+	ITemporalUpscaler::FOutputs ThirdPartyOutputs = UpscalerToUse->AddPasses(
+		GraphBuilder,
+		View,
+		ThirdPartyInputs);
+
+	check(ThirdPartyOutputs.FullRes.ViewRect == ThirdPartyInputs.OutputViewRect);
+	check(ThirdPartyOutputs.FullRes.ViewRect.Max.X <= ThirdPartyOutputs.FullRes.Texture->Desc.Extent.X);
+	check(ThirdPartyOutputs.FullRes.ViewRect.Max.Y <= ThirdPartyOutputs.FullRes.Texture->Desc.Extent.Y);
+
+	check(ThirdPartyOutputs.NewHistory);
+	check(ThirdPartyOutputs.NewHistory->GetDebugName() == UpscalerToUse->GetDebugName());
+
+	// Translate the output.
+	FDefaultTemporalUpscaler::FOutputs Outputs;
+	Outputs.FullRes = ThirdPartyOutputs.FullRes;
+
+	// Saves history for next frame.
+	if (!View.bStatePrevViewInfoIsReadOnly)
+	{
+		View.ViewState->PrevFrameViewInfo.ThirdPartyTemporalUpscalerHistory = ThirdPartyOutputs.NewHistory;
+	}
+
+	// Save output for next frame's SSR
+	if (!View.bStatePrevViewInfoIsReadOnly)
+	{
+		FTemporalAAHistory& OutputHistory = View.ViewState->PrevFrameViewInfo.TemporalAAHistory;
+
+		GraphBuilder.QueueTextureExtraction(Outputs.FullRes.Texture, &OutputHistory.RT[0]);
+
+		OutputHistory.ViewportRect = Outputs.FullRes.ViewRect;
+		OutputHistory.ReferenceBufferSize = Outputs.FullRes.Texture->Desc.Extent;
+	}
+
+	return Outputs;
+}
 
 bool ComposeSeparateTranslucencyInTSR(const FViewInfo& View);
 
@@ -523,7 +588,7 @@ void AddPostProcessingPasses(
 
 
 		// Temporal Anti-aliasing. Also may perform a temporal upsample from primary to secondary view rect.
-		EMainTAAPassConfig TAAConfig = ITemporalUpscaler::GetMainTAAPassConfig(View);
+		EMainTAAPassConfig TAAConfig = GetMainTAAPassConfig(View);
 
 		// Whether separate translucency is composed in TSR.
 		bool bComposeSeparateTranslucencyInTSR = TAAConfig == EMainTAAPassConfig::TSR && ComposeSeparateTranslucencyInTSR(View);
@@ -641,21 +706,10 @@ void AddPostProcessingPasses(
 		FVelocityFlattenTextures VelocityFlattenTextures;
 		if (TAAConfig != EMainTAAPassConfig::Disabled)
 		{
-			const ITemporalUpscaler* UpscalerToUse = (TAAConfig == EMainTAAPassConfig::ThirdParty) ? View.Family->GetTemporalUpscalerInterface() : ITemporalUpscaler::GetDefaultTemporalUpscaler();
-			check(UpscalerToUse);
-			const TCHAR* UpscalerName = UpscalerToUse->GetDebugName();
-
-			// Standard event scope for temporal upscaler to have all profiling information not matter what,
-			// and with explicit detection of third party.
-			RDG_EVENT_SCOPE_CONDITIONAL(
-				GraphBuilder,
-				TAAConfig == EMainTAAPassConfig::ThirdParty,
-				"ThirdParty %s %dx%d -> %dx%d",
-				UpscalerToUse->GetDebugName(),
-				View.ViewRect.Width(), View.ViewRect.Height(),
-				View.GetSecondaryViewRectSize().X, View.GetSecondaryViewRectSize().Y);
-
-			ITemporalUpscaler::FPassInputs UpscalerPassInputs;
+			FDefaultTemporalUpscaler::FInputs UpscalerPassInputs;
+			UpscalerPassInputs.SceneColor = FScreenPassTexture(SceneColor.Texture, View.ViewRect);
+			UpscalerPassInputs.SceneDepth = FScreenPassTexture(SceneDepth.Texture, View.ViewRect);
+			UpscalerPassInputs.SceneVelocity = FScreenPassTexture(Velocity.Texture, View.ViewRect);
 			if (PassSequence.IsEnabled(EPass::MotionBlur))
 			{
 				if (bVisualizeMotionBlur)
@@ -679,16 +733,35 @@ void AddPostProcessingPasses(
 					DownsampleQuality == EDownsampleQuality::Low;
 			}
 			UpscalerPassInputs.DownsampleOverrideFormat = DownsampleOverrideFormat;
-			UpscalerPassInputs.SceneColorTexture = SceneColor.Texture;
-			UpscalerPassInputs.SceneDepthTexture = SceneDepth.Texture;
-			UpscalerPassInputs.SceneVelocityTexture = Velocity.Texture;
 			UpscalerPassInputs.PostDOFTranslucencyResources = PostDOFTranslucencyResources;
 			UpscalerPassInputs.MoireInputTexture = TSRMoireInput;
 
-			ITemporalUpscaler::FOutputs Outputs = UpscalerToUse->AddPasses(
-				GraphBuilder,
-				View,
-				UpscalerPassInputs);
+			FDefaultTemporalUpscaler::FOutputs Outputs;
+			if (TAAConfig == EMainTAAPassConfig::TSR)
+			{
+				Outputs = AddTemporalSuperResolutionPasses(
+					GraphBuilder,
+					View,
+					UpscalerPassInputs);
+			}
+			else if (TAAConfig == EMainTAAPassConfig::TAA)
+			{
+				Outputs = AddGen4MainTemporalAAPasses(
+					GraphBuilder,
+					View,
+					UpscalerPassInputs);
+			}
+			else if (TAAConfig == EMainTAAPassConfig::ThirdParty)
+			{
+				Outputs = AddThirdPartyTemporalUpscalerPasses(
+					GraphBuilder,
+					View,
+					UpscalerPassInputs);
+			}
+			else
+			{
+				unimplemented();
+			}
 
 			SceneColor = Outputs.FullRes;
 			HalfResSceneColor = Outputs.HalfRes;
@@ -698,7 +771,7 @@ void AddPostProcessingPasses(
 			if (PassSequence.IsEnabled(EPass::VisualizeTemporalUpscaler))
 			{
 				VisualizeTemporalUpscalerInputs.TAAConfig = TAAConfig;
-				VisualizeTemporalUpscalerInputs.UpscalerUsed = UpscalerToUse;
+				VisualizeTemporalUpscalerInputs.UpscalerUsed = View.Family->GetTemporalUpscalerInterface();
 				VisualizeTemporalUpscalerInputs.Inputs = UpscalerPassInputs;
 				VisualizeTemporalUpscalerInputs.Outputs = Outputs;
 			}
@@ -2362,32 +2435,33 @@ void AddMobilePostProcessingPasses(FRDGBuilder& GraphBuilder, FScene* Scene, con
 		{
 			PassSequence.AcceptPass(EPass::TAA);
 
-			EMainTAAPassConfig TAAConfig = ITemporalUpscaler::GetMainTAAPassConfig(View);
+			EMainTAAPassConfig TAAConfig = GetMainTAAPassConfig(View);
 			checkSlow(TAAConfig != EMainTAAPassConfig::Disabled);
 
-			const ITemporalUpscaler* UpscalerToUse = (TAAConfig == EMainTAAPassConfig::ThirdParty) ? View.Family->GetTemporalUpscalerInterface() : ITemporalUpscaler::GetDefaultTemporalUpscaler();
-
-			const TCHAR* UpscalerName = UpscalerToUse->GetDebugName();
-
-			// Standard event scope for temporal upscaler to have all profiling information not matter what, and with explicit detection of third party.
-			RDG_EVENT_SCOPE_CONDITIONAL(
-				GraphBuilder,
-				TAAConfig == EMainTAAPassConfig::ThirdParty,
-				"ThirdParty %s %dx%d -> %dx%d",
-				UpscalerToUse->GetDebugName(),
-				View.ViewRect.Width(), View.ViewRect.Height(),
-				View.GetSecondaryViewRectSize().X, View.GetSecondaryViewRectSize().Y);
-
-			ITemporalUpscaler::FPassInputs UpscalerPassInputs;
-			UpscalerPassInputs.SceneColorTexture = SceneColor.Texture;
-			UpscalerPassInputs.SceneDepthTexture = SceneDepth.Texture;
-			UpscalerPassInputs.SceneVelocityTexture = Velocity.Texture;
-
-			ITemporalUpscaler::FOutputs Outputs = UpscalerToUse->AddPasses(
-				GraphBuilder,
-				View,
-				UpscalerPassInputs);
+			FDefaultTemporalUpscaler::FInputs UpscalerPassInputs;
+			UpscalerPassInputs.SceneColor = FScreenPassTexture(SceneColor.Texture, View.ViewRect);
+			UpscalerPassInputs.SceneDepth = FScreenPassTexture(SceneDepth.Texture, View.ViewRect);
+			UpscalerPassInputs.SceneVelocity = FScreenPassTexture(Velocity.Texture, View.ViewRect);
 
+			FDefaultTemporalUpscaler::FOutputs Outputs;
+			if (TAAConfig == EMainTAAPassConfig::TAA)
+			{
+				Outputs = AddGen4MainTemporalAAPasses(
+					GraphBuilder,
+					View,
+					UpscalerPassInputs);
+			}
+			else if (TAAConfig == EMainTAAPassConfig::ThirdParty)
+			{
+				Outputs = AddThirdPartyTemporalUpscalerPasses(
+					GraphBuilder,
+					View,
+					UpscalerPassInputs);
+			}
+			else
+			{
+				unimplemented();
+			}
 			SceneColor = Outputs.FullRes;
 		}
 	}
diff --git a/Engine/Source/Runtime/Renderer/Private/PostProcess/TemporalAA.cpp b/Engine/Source/Runtime/Renderer/Private/PostProcess/TemporalAA.cpp
--- a/Engine/Source/Runtime/Renderer/Private/PostProcess/TemporalAA.cpp
+++ b/Engine/Source/Runtime/Renderer/Private/PostProcess/TemporalAA.cpp
@@ -990,14 +990,10 @@ FTAAOutputs AddTemporalAAPass(
 	return Outputs;
 } // AddTemporalAAPass()
 
-static void AddGen4MainTemporalAAPasses(
+FDefaultTemporalUpscaler::FOutputs AddGen4MainTemporalAAPasses(
 	FRDGBuilder& GraphBuilder,
 	const FViewInfo& View,
-	const ITemporalUpscaler::FPassInputs& PassInputs,
-	FRDGTextureRef* OutSceneColorTexture,
-	FIntRect* OutSceneColorViewRect,
-	FRDGTextureRef* OutSceneColorHalfResTexture,
-	FIntRect* OutSceneColorHalfResViewRect)
+	const FDefaultTemporalUpscaler::FInputs& PassInputs)
 {
 	check(View.ViewState);
 
@@ -1036,9 +1032,9 @@ static void AddGen4MainTemporalAAPasses(
 
 	TAAParameters.bDownsample = (PassInputs.bGenerateSceneColorHalfRes || PassInputs.bGenerateSceneColorQuarterRes) && TAAParameters.Quality != ETAAQuality::High;
 
-	TAAParameters.SceneDepthTexture = PassInputs.SceneDepthTexture;
-	TAAParameters.SceneVelocityTexture = PassInputs.SceneVelocityTexture;
-	TAAParameters.SceneColorInput = PassInputs.SceneColorTexture;
+	TAAParameters.SceneDepthTexture = PassInputs.SceneDepth.Texture;
+	TAAParameters.SceneVelocityTexture = PassInputs.SceneVelocity.Texture;
+	TAAParameters.SceneColorInput = PassInputs.SceneColor.Texture;
 
 	const FTemporalAAHistory& InputHistory = View.PrevViewInfo.TemporalAAHistory;
 
@@ -1064,87 +1060,21 @@ static void AddGen4MainTemporalAAPasses(
 
 		FScreenPassTextureViewport OutputViewport;
 		OutputViewport.Rect = SecondaryViewRect;
-		OutputViewport.Extent.X = FMath::Max(PassInputs.SceneColorTexture->Desc.Extent.X, QuantizedOutputSize.X);
-		OutputViewport.Extent.Y = FMath::Max(PassInputs.SceneColorTexture->Desc.Extent.Y, QuantizedOutputSize.Y);
+		OutputViewport.Extent.X = FMath::Max(PassInputs.SceneColor.Texture->Desc.Extent.X, QuantizedOutputSize.X);
+		OutputViewport.Extent.Y = FMath::Max(PassInputs.SceneColor.Texture->Desc.Extent.Y, QuantizedOutputSize.Y);
 
 		SceneColorTexture = ComputeMitchellNetravaliDownsample(GraphBuilder, View, FScreenPassTexture(SceneColorTexture, InputViewport), OutputViewport);
 	}
 
-	*OutSceneColorTexture = SceneColorTexture;
-	*OutSceneColorViewRect = SecondaryViewRect;
-	*OutSceneColorHalfResTexture = TAAOutputs.DownsampledSceneColor;
-	*OutSceneColorHalfResViewRect = FIntRect::DivideAndRoundUp(SecondaryViewRect, 2);
+	FDefaultTemporalUpscaler::FOutputs Outputs;
+	Outputs.FullRes.Texture = SceneColorTexture;
+	Outputs.FullRes.ViewRect = SecondaryViewRect;
+	Outputs.HalfRes.Texture = TAAOutputs.DownsampledSceneColor;
+	Outputs.HalfRes.ViewRect = FIntRect::DivideAndRoundUp(SecondaryViewRect, 2);
+	return Outputs;
 } // AddGen4MainTemporalAAPasses()
 
-bool DoesPlatformSupportTSR(EShaderPlatform Platform);
-
-ITemporalUpscaler::FOutputs AddTemporalSuperResolutionPasses(
-	FRDGBuilder& GraphBuilder,
-	const FViewInfo& View,
-	const ITemporalUpscaler::FPassInputs& PassInputs);
-
-class FDefaultTemporalUpscaler : public ITemporalUpscaler
-{
-public:
-
-	virtual const TCHAR* GetDebugName() const
-	{
-		return TEXT("FDefaultTemporalUpscaler");
-	}
-
-	virtual FOutputs AddPasses(
-		FRDGBuilder& GraphBuilder,
-		const FViewInfo& View,
-		const FPassInputs& PassInputs) const final
-	{
-		if (ITemporalUpscaler::GetMainTAAPassConfig(View) == EMainTAAPassConfig::TSR)
-		{
-			return AddTemporalSuperResolutionPasses(
-				GraphBuilder,
-				View,
-				PassInputs);
-		}
-		else
-		{
-			FOutputs Outputs;
-			AddGen4MainTemporalAAPasses(
-				GraphBuilder,
-				View,
-				PassInputs,
-				&Outputs.FullRes.Texture,
-				&Outputs.FullRes.ViewRect,
-				&Outputs.HalfRes.Texture,
-				&Outputs.HalfRes.ViewRect);
-
-			return Outputs;
-		}
-	}
-
-	virtual float GetMinUpsampleResolutionFraction() const override
-	{
-		return ISceneViewFamilyScreenPercentage::kMinTAAUpsampleResolutionFraction;
-	}
-	virtual float GetMaxUpsampleResolutionFraction() const override
-	{
-		return ISceneViewFamilyScreenPercentage::kMaxTAAUpsampleResolutionFraction;
-	}
-
-	virtual ITemporalUpscaler* Fork_GameThread(const class FSceneViewFamily& ViewFamily) const override
-	{
-		return nullptr;
-	}
-
-};
-
-// static
-const ITemporalUpscaler* ITemporalUpscaler::GetDefaultTemporalUpscaler()
-{
-	static FDefaultTemporalUpscaler DefaultTemporalUpscaler;
-	return &DefaultTemporalUpscaler;
-}
-
-//static
-EMainTAAPassConfig ITemporalUpscaler::GetMainTAAPassConfig(const FViewInfo& View)
+EMainTAAPassConfig GetMainTAAPassConfig(const FViewInfo& View)
 {
 	if (!IsPostProcessingEnabled(View))
 	{
diff --git a/Engine/Source/Runtime/Renderer/Private/PostProcess/TemporalAA.h b/Engine/Source/Runtime/Renderer/Private/PostProcess/TemporalAA.h
--- a/Engine/Source/Runtime/Renderer/Private/PostProcess/TemporalAA.h
+++ b/Engine/Source/Runtime/Renderer/Private/PostProcess/TemporalAA.h
@@ -4,6 +4,7 @@
 
 #include "ScreenPass.h"
 #include "PostProcess/PostProcessMotionBlur.h"
+#include "TemporalUpscaler.h"
 
 struct FTemporalAAHistory;
 struct FTranslucencyPassResources;
@@ -178,21 +179,22 @@ bool NeedTSRMoireLuma(const FViewInfo& View);
 /** Measure luminance of the scene color for moire anti-flickering. */
 FScreenPassTexture AddTSRComputeMoireLuma(FRDGBuilder& GraphBuilder, FGlobalShaderMap* ShaderMap, FScreenPassTexture SceneColor);
 
-/** Interface for the main temporal upscaling algorithm. */
-class RENDERER_API ITemporalUpscaler : public ISceneViewFamilyExtention
-{
-public:
 
-	struct FPassInputs
+EMainTAAPassConfig GetMainTAAPassConfig(const FViewInfo& View);
+
+/** Interface for the default temporal upscaling algorithm. */
+struct FDefaultTemporalUpscaler
+{
+	struct FInputs
 	{
 		bool bGenerateSceneColorHalfRes = false;
 		bool bGenerateSceneColorQuarterRes = false;
 		bool bGenerateOutputMip1 = false;
 		bool bGenerateVelocityFlattenTextures = false;
 		EPixelFormat DownsampleOverrideFormat;
-		FRDGTextureRef SceneColorTexture = nullptr;
-		FRDGTextureRef SceneDepthTexture = nullptr;
-		FRDGTextureRef SceneVelocityTexture = nullptr;
+		FScreenPassTexture SceneColor;
+		FScreenPassTexture SceneDepth;
+		FScreenPassTexture SceneVelocity;
 		FTranslucencyPassResources PostDOFTranslucencyResources;
 		FScreenPassTexture MoireInputTexture;
 	};
@@ -204,31 +206,17 @@ public:
 		FScreenPassTexture QuarterRes;
 		FVelocityFlattenTextures VelocityFlattenTextures;
 	};
+};
 
-	virtual ~ITemporalUpscaler() {};
-
-	virtual const TCHAR* GetDebugName() const = 0;
-
-	///** Temporal AA helper method which performs filtering on the main pass scene color. Supports upsampled history and,
-	// *  if requested, will attempt to perform the scene color downsample. Returns the filtered scene color, the downsampled
-	// *  scene color (or null if it was not performed), and the secondary view rect.
-	// */
-
-	virtual FOutputs AddPasses(
-		FRDGBuilder& GraphBuilder,
-		const FViewInfo& View,
-		const FPassInputs& PassInputs) const = 0;
-
-
-	virtual float GetMinUpsampleResolutionFraction() const = 0;
-	virtual float GetMaxUpsampleResolutionFraction() const = 0;
-
-	virtual ITemporalUpscaler* Fork_GameThread(const class FSceneViewFamily& ViewFamily) const = 0;
-
-	static const ITemporalUpscaler* GetDefaultTemporalUpscaler();
+FDefaultTemporalUpscaler::FOutputs AddGen4MainTemporalAAPasses(
+	FRDGBuilder& GraphBuilder,
+	const FViewInfo& View,
+	const FDefaultTemporalUpscaler::FInputs& PassInputs);
 
-	static EMainTAAPassConfig GetMainTAAPassConfig(const FViewInfo& View);
-}; 
+FDefaultTemporalUpscaler::FOutputs AddTemporalSuperResolutionPasses(
+	FRDGBuilder& GraphBuilder,
+	const FViewInfo& View,
+	const FDefaultTemporalUpscaler::FInputs& PassInputs);
 
 /** Interface for the VisualizeTSR showflag. */
 struct FVisualizeTemporalUpscalerInputs
@@ -242,8 +230,8 @@ struct FVisualizeTemporalUpscalerInputs
 	// Temporal upscaler used and its inputs and outputs.
 	EMainTAAPassConfig TAAConfig = EMainTAAPassConfig::Disabled;
 	const ITemporalUpscaler* UpscalerUsed = nullptr;
-	ITemporalUpscaler::FPassInputs Inputs;
-	ITemporalUpscaler::FOutputs Outputs;
+	FDefaultTemporalUpscaler::FInputs Inputs;
+	FDefaultTemporalUpscaler::FOutputs Outputs;
 };
 
 FScreenPassTexture AddVisualizeTemporalUpscalerPass(FRDGBuilder& GraphBuilder, const FViewInfo& View, const FVisualizeTemporalUpscalerInputs& Inputs);
diff --git a/Engine/Source/Runtime/Renderer/Private/PostProcess/TemporalSuperResolution.cpp b/Engine/Source/Runtime/Renderer/Private/PostProcess/TemporalSuperResolution.cpp
--- a/Engine/Source/Runtime/Renderer/Private/PostProcess/TemporalSuperResolution.cpp
+++ b/Engine/Source/Runtime/Renderer/Private/PostProcess/TemporalSuperResolution.cpp
@@ -1054,7 +1054,7 @@ static FRDGTextureUAVRef CreateDummyUAV(FRDGBuilder& GraphBuilder, EPixelFormat
 
 bool NeedTSRMoireLuma(const FViewInfo& View)
 {
-	return ITemporalUpscaler::GetMainTAAPassConfig(View) == EMainTAAPassConfig::TSR;
+	return GetMainTAAPassConfig(View) == EMainTAAPassConfig::TSR;
 }
 
 FScreenPassTexture AddTSRComputeMoireLuma(FRDGBuilder& GraphBuilder, FGlobalShaderMap* ShaderMap, FScreenPassTexture SceneColor)
@@ -1092,10 +1092,10 @@ FScreenPassTexture AddTSRComputeMoireLuma(FRDGBuilder& GraphBuilder, FGlobalShad
 	return MoireLuma;
 }
 
-ITemporalUpscaler::FOutputs AddTemporalSuperResolutionPasses(
+FDefaultTemporalUpscaler::FOutputs AddTemporalSuperResolutionPasses(
 	FRDGBuilder& GraphBuilder,
 	const FViewInfo& View,
-	const ITemporalUpscaler::FPassInputs& PassInputs)
+	const FDefaultTemporalUpscaler::FInputs& PassInputs)
 {
 	const FTSRHistory& InputHistory = View.PrevViewInfo.TSRHistory;
 
@@ -1171,7 +1171,7 @@ ITemporalUpscaler::FOutputs AddTemporalSuperResolutionPasses(
 		RejectionAntiAliasingQuality = 0; 
 	}
 
-	FIntPoint InputExtent = PassInputs.SceneColorTexture->Desc.Extent;
+	FIntPoint InputExtent = PassInputs.SceneColor.Texture->Desc.Extent;
 	FIntRect InputRect = View.ViewRect;
 
 	FIntPoint OutputExtent;
@@ -1642,8 +1642,8 @@ ITemporalUpscaler::FOutputs AddTemporalSuperResolutionPasses(
 		PassParameters->bOutputIsMovingTexture = bOutputIsMovingTexture;
 		PassParameters->bIncludeDynamicDepthInSubpixelDetails = SubpixelMethod == ETSRSubpixelMethod::ClosestDepth && CVarTSRSubpixelIncludeMovingDepth.GetValueOnRenderThread();
 
-		PassParameters->SceneDepthTexture = PassInputs.SceneDepthTexture;
-		PassParameters->SceneVelocityTexture = PassInputs.SceneVelocityTexture;
+		PassParameters->SceneDepthTexture = PassInputs.SceneDepth.Texture;
+		PassParameters->SceneVelocityTexture = PassInputs.SceneVelocity.Texture;
 		if (PrevScatteredSubpixelDepthTexture)
 		{
 			PassParameters->PrevScatteredSubpixelDepthTexture = PrevScatteredSubpixelDepthTexture;
@@ -1771,7 +1771,7 @@ ITemporalUpscaler::FOutputs AddTemporalSuperResolutionPasses(
 			PassParameters->WorldDepthToPixelWorldRadius = TanHalfFieldOfView / float(View.ViewRect.Width());
 		}
 
-		PassParameters->InputSceneColorTexture = PassInputs.SceneColorTexture;
+		PassParameters->InputSceneColorTexture = PassInputs.SceneColor.Texture;
 		PassParameters->DilatedVelocityTexture = DilatedVelocityTexture;
 		PassParameters->ClosestDepthTexture = ClosestDepthTexture;
 		PassParameters->PrevUseCountTexture = PrevUseCountTexture;
@@ -1874,7 +1874,7 @@ ITemporalUpscaler::FOutputs AddTemporalSuperResolutionPasses(
 		PassParameters->TileOverscan = FMath::Clamp(CVarTSRShadingTileOverscan.GetValueOnRenderThread(), 2, GroupTileSize / 2 - 1);
 		PassParameters->PerceptionAdd = FMath::Pow(0.5f, CVarTSRShadingExposureOffset.GetValueOnRenderThread());
 
-		PassParameters->InputTexture = PassInputs.SceneColorTexture;
+		PassParameters->InputTexture = PassInputs.SceneColor.Texture;
 		if (PassInputs.MoireInputTexture.IsValid())
 		{
 			ensure(InputRect == PassInputs.MoireInputTexture.ViewRect);
@@ -1951,7 +1951,7 @@ ITemporalUpscaler::FOutputs AddTemporalSuperResolutionPasses(
 		{
 			FTSRSpatialAntiAliasingCS::FParameters* PassParameters = GraphBuilder.AllocParameters<FTSRSpatialAntiAliasingCS::FParameters>();
 			PassParameters->CommonParameters = CommonParameters;
-			PassParameters->InputSceneColorTexture = PassInputs.SceneColorTexture;
+			PassParameters->InputSceneColorTexture = PassInputs.SceneColor.Texture;
 			PassParameters->InputSceneColorLdrLumaTexture = InputSceneColorLdrLumaTexture;
 			PassParameters->AntiAliasingOutput = GraphBuilder.CreateUAV(RawAntiAliasingTexture);
 			PassParameters->NoiseFilteringOutput = GraphBuilder.CreateUAV(NoiseFilteringTexture);
@@ -2002,9 +2002,9 @@ ITemporalUpscaler::FOutputs AddTemporalSuperResolutionPasses(
 
 		FTSRUpdateHistoryCS::FParameters* PassParameters = GraphBuilder.AllocParameters<FTSRUpdateHistoryCS::FParameters>();
 		PassParameters->CommonParameters = CommonParameters;
-		PassParameters->InputSceneColorTexture = PassInputs.SceneColorTexture;
+		PassParameters->InputSceneColorTexture = PassInputs.SceneColor.Texture;
 		PassParameters->InputSceneStencilTexture = GraphBuilder.CreateSRV(
-			FRDGTextureSRVDesc::CreateWithPixelFormat(PassInputs.SceneDepthTexture, PF_X24_G8));
+			FRDGTextureSRVDesc::CreateWithPixelFormat(PassInputs.SceneDepth.Texture, PF_X24_G8));
 		PassParameters->InputSceneTranslucencyTexture = SeparateTranslucencyTexture;
 		PassParameters->HistoryRejectionTexture = HistoryRejectionTexture;
 
@@ -2253,7 +2253,7 @@ ITemporalUpscaler::FOutputs AddTemporalSuperResolutionPasses(
 	}
 
 
-	ITemporalUpscaler::FOutputs Outputs;
+	FDefaultTemporalUpscaler::FOutputs Outputs;
 	Outputs.FullRes.Texture = SceneColorOutputTexture;
 	Outputs.FullRes.ViewRect = OutputRect;
 	if (SceneColorOutputHalfResTexture)
diff --git a/Engine/Source/Runtime/Renderer/Private/PostProcess/VisualizeTemporalUpscaler.cpp b/Engine/Source/Runtime/Renderer/Private/PostProcess/VisualizeTemporalUpscaler.cpp
--- a/Engine/Source/Runtime/Renderer/Private/PostProcess/VisualizeTemporalUpscaler.cpp
+++ b/Engine/Source/Runtime/Renderer/Private/PostProcess/VisualizeTemporalUpscaler.cpp
@@ -31,7 +31,7 @@ FScreenPassTexture AddVisualizeTemporalUpscalerPass(FRDGBuilder& GraphBuilder, c
 	TArray<FVisualizeBufferTile> Tiles;
 	Tiles.SetNum(16);
 
-	if (Inputs.UpscalerUsed)
+	if (Inputs.TAAConfig != EMainTAAPassConfig::Disabled)
 	{
 		auto VisualizeTextureLabel = [](FRDGTextureRef Texture, const TCHAR* Suffix = TEXT(""))
 		{
@@ -46,16 +46,16 @@ FScreenPassTexture AddVisualizeTemporalUpscalerPass(FRDGBuilder& GraphBuilder, c
 		// Depth buffer
 		{
 			FVisualizeBufferTile& Tile = Tiles[4 * 0 + 0];
-			Tile.Input.Texture = FVisualizeTexture::AddVisualizeTexturePass(GraphBuilder, View.ShaderMap, Inputs.Inputs.SceneDepthTexture);
+			Tile.Input.Texture = FVisualizeTexture::AddVisualizeTexturePass(GraphBuilder, View.ShaderMap, Inputs.Inputs.SceneDepth.Texture);
 			Tile.Input.ViewRect = CropViewRectToCenter(View.ViewRect);
-			Tile.Label = VisualizeTextureLabel(Inputs.Inputs.SceneDepthTexture);
+			Tile.Label = VisualizeTextureLabel(Inputs.Inputs.SceneDepth.Texture);
 		}
 
 		// Scene Color
 		{
 			FVisualizeBufferTile& Tile = Tiles[4 * 0 + 1];
-			Tile.Input = FScreenPassTexture(Inputs.Inputs.SceneColorTexture, CropViewRectToCenter(View.ViewRect));
-			Tile.Label = VisualizeTextureLabel(Inputs.Inputs.SceneColorTexture);
+			Tile.Input = FScreenPassTexture(Inputs.Inputs.SceneColor.Texture, CropViewRectToCenter(View.ViewRect));
+			Tile.Label = VisualizeTextureLabel(Inputs.Inputs.SceneColor.Texture);
 		}
 
 		// Scene Color alpha
@@ -63,9 +63,9 @@ FScreenPassTexture AddVisualizeTemporalUpscalerPass(FRDGBuilder& GraphBuilder, c
 		{
 			FVisualizeBufferTile& Tile = Tiles[4 * 3 + 1];
 			//Tile.Input = FScreenPassTexture(Inputs.Inputs.SceneColorTexture, CropViewRectToCenter(View.ViewRect));
-			Tile.Input.Texture = FVisualizeTexture::AddVisualizeTextureAlphaPass(GraphBuilder, View.ShaderMap, Inputs.Inputs.SceneColorTexture);
+			Tile.Input.Texture = FVisualizeTexture::AddVisualizeTextureAlphaPass(GraphBuilder, View.ShaderMap, Inputs.Inputs.SceneColor.Texture);
 			Tile.Input.ViewRect = CropViewRectToCenter(View.ViewRect);
-			Tile.Label = VisualizeTextureLabel(Inputs.Inputs.SceneColorTexture, TEXT(" A"));
+			Tile.Label = VisualizeTextureLabel(Inputs.Inputs.SceneColor.Texture, TEXT(" A"));
 		}
 
 		// Translucency
@@ -104,8 +104,8 @@ FScreenPassTexture AddVisualizeTemporalUpscalerPass(FRDGBuilder& GraphBuilder, c
 		{
 			FMotionBlurInputs PassInputs;
 			PassInputs.SceneColor = Inputs.Outputs.FullRes;
-			PassInputs.SceneDepth = FScreenPassTexture(Inputs.Inputs.SceneDepthTexture, View.ViewRect);
-			PassInputs.SceneVelocity = FScreenPassTexture(Inputs.Inputs.SceneVelocityTexture, View.ViewRect);
+			PassInputs.SceneDepth = Inputs.Inputs.SceneDepth;
+			PassInputs.SceneVelocity = Inputs.Inputs.SceneVelocity;
 
 			FVisualizeBufferTile& Tile = Tiles[4 * 1 + 0];
 			Tile.Input = AddVisualizeMotionBlurPass(GraphBuilder, View, PassInputs);
@@ -117,8 +117,8 @@ FScreenPassTexture AddVisualizeTemporalUpscalerPass(FRDGBuilder& GraphBuilder, c
 		{
 			FVisualizeMotionVectorsInputs PassInputs;
 			PassInputs.SceneColor = Inputs.SceneColor;
-			PassInputs.SceneDepth = FScreenPassTexture(Inputs.Inputs.SceneDepthTexture, View.ViewRect);
-			PassInputs.SceneVelocity = FScreenPassTexture(Inputs.Inputs.SceneVelocityTexture, View.ViewRect);
+			PassInputs.SceneDepth = Inputs.Inputs.SceneDepth;
+			PassInputs.SceneVelocity = Inputs.Inputs.SceneVelocity;
 
 			FVisualizeBufferTile& Tile = Tiles[4 * 2 + 0];
 			Tile.Input = AddVisualizeMotionVectorsPass(GraphBuilder, View, PassInputs);
@@ -172,7 +172,7 @@ FScreenPassTexture AddVisualizeTemporalUpscalerPass(FRDGBuilder& GraphBuilder, c
 
 	// Draw additional text
 	// Early return if not using a temporal upscaler.
-	if (Inputs.UpscalerUsed)
+	if (Inputs.TAAConfig != EMainTAAPassConfig::Disabled)
 	{
 		AddDrawCanvasPass(GraphBuilder, RDG_EVENT_NAME("VisualizeTemporalUpscaler Text"), View, FScreenPassRenderTarget(Output, ERenderTargetLoadAction::ELoad),
 			[&View, &ViewRect = Output.ViewRect, &Inputs, bSupportsAlpha](FCanvas& Canvas)
@@ -209,7 +209,7 @@ FScreenPassTexture AddVisualizeTemporalUpscalerPass(FRDGBuilder& GraphBuilder, c
 
 			// Display the input/output resolutions
 			{
-				QuickDrawSummary(/* Location = */ 1, FString::Printf(TEXT("Input: %dx%d %s"), View.ViewRect.Width(), View.ViewRect.Height(), GPixelFormats[Inputs.Inputs.SceneColorTexture->Desc.Format].Name));
+				QuickDrawSummary(/* Location = */ 1, FString::Printf(TEXT("Input: %dx%d %s"), View.ViewRect.Width(), View.ViewRect.Height(), GPixelFormats[Inputs.Inputs.SceneColor.Texture->Desc.Format].Name));
 				QuickDrawSummary(/* Location = */ 2, FString::Printf(TEXT("Output: %dx%d %s"), Inputs.Outputs.FullRes.ViewRect.Width(), Inputs.Outputs.FullRes.ViewRect.Height(), GPixelFormats[Inputs.Outputs.FullRes.Texture->Desc.Format].Name));
 			}
 
diff --git a/Engine/Source/Runtime/Renderer/Private/RendererScene.cpp b/Engine/Source/Runtime/Renderer/Private/RendererScene.cpp
--- a/Engine/Source/Runtime/Renderer/Private/RendererScene.cpp
+++ b/Engine/Source/Runtime/Renderer/Private/RendererScene.cpp
@@ -924,6 +924,15 @@ uint64 FPreviousViewInfo::GetGPUSizeBytes(bool bLogSizes) const
 		}
 	}
 
+	if (ThirdPartyTemporalUpscalerHistory)
+	{
+		TotalSize += ThirdPartyTemporalUpscalerHistory->GetGPUSizeBytes();
+		if (bLogSizes)
+		{
+			UE_LOG(LogRenderer, Log, TEXT("LogSizes\tThirdPartyTemporalUpscaler\t%s\t%llu"), ThirdPartyTemporalUpscalerHistory->GetDebugName(), ThirdPartyTemporalUpscalerHistory->GetGPUSizeBytes());
+		}
+	}
+
 	return TotalSize;
 }
 
diff --git a/Engine/Source/Runtime/Renderer/Private/SceneRendering.h b/Engine/Source/Runtime/Renderer/Private/SceneRendering.h
--- a/Engine/Source/Runtime/Renderer/Private/SceneRendering.h
+++ b/Engine/Source/Runtime/Renderer/Private/SceneRendering.h
@@ -36,6 +36,7 @@
 #include "Lumen/LumenTranslucencyVolumeLighting.h"
 #include "HairStrands/HairStrandsData.h"
 #include "Strata/Strata.h"
+#include "TemporalUpscaler.h"
 #include "GPUScene.h"
 #include "RenderCore.h"
 #include "SceneTextures.h"
@@ -1089,14 +1090,6 @@ struct FGTAOTAAHistory
 	}
 };
 
-
-// Plugins can derive from this and use it for their own purposes
-class RENDERER_API ICustomTemporalAAHistory : public IRefCountedObject
-{
-public:
-	virtual ~ICustomTemporalAAHistory() {}
-};
-
 // Structure that hold all information related to previous frame.
 struct FPreviousViewInfo
 {
@@ -1141,7 +1134,7 @@ struct FPreviousViewInfo
 	FTSRHistory TSRHistory;
 
 	// Custom Temporal AA result of last frame, used by plugins
-	TRefCountPtr<ICustomTemporalAAHistory> CustomTemporalAAHistory;
+	TRefCountPtr<ITemporalUpscaler::IHistory> ThirdPartyTemporalUpscalerHistory;
 
 	// Half resolution version temporal AA result of last frame
 	TRefCountPtr<IPooledRenderTarget> HalfResTemporalAAHistory;
diff --git a/Engine/Source/Runtime/Renderer/Private/SceneVisibility.cpp b/Engine/Source/Runtime/Renderer/Private/SceneVisibility.cpp
--- a/Engine/Source/Runtime/Renderer/Private/SceneVisibility.cpp
+++ b/Engine/Source/Runtime/Renderer/Private/SceneVisibility.cpp
@@ -4161,7 +4161,7 @@ void FSceneRenderer::PrepareViewStateForVisibility(const FSceneTexturesConfig& S
 		// Subpixel jitter for temporal AA
 		int32 CVarTemporalAASamplesValue = CVarTemporalAASamples.GetValueOnRenderThread();
 
-		EMainTAAPassConfig TAAConfig = ITemporalUpscaler::GetMainTAAPassConfig(View);
+		EMainTAAPassConfig TAAConfig = GetMainTAAPassConfig(View);
 
 		bool bTemporalUpsampling = View.PrimaryScreenPercentageMethod == EPrimaryScreenPercentageMethod::TemporalUpscale;
 		
diff --git a/Engine/Source/Runtime/Renderer/Private/TranslucentRendering.cpp b/Engine/Source/Runtime/Renderer/Private/TranslucentRendering.cpp
--- a/Engine/Source/Runtime/Renderer/Private/TranslucentRendering.cpp
+++ b/Engine/Source/Runtime/Renderer/Private/TranslucentRendering.cpp
@@ -1631,7 +1631,7 @@ void FDeferredShadingSceneRenderer::RenderTranslucency(
 		}
 
 		FTranslucencyPassResources& TranslucencyPassResources = OutTranslucencyResourceMap->Get(ViewIndex, ETranslucencyPass::TPT_TranslucencyAfterDOF);
-		if (SharedUpscaledPostDOFTranslucencyColor && TranslucencyPassResources.IsValid() && TranslucencyPassResources.ViewRect.Size() != View.ViewRect.Size() && ITemporalUpscaler::GetMainTAAPassConfig(View) != EMainTAAPassConfig::TSR)
+		if (SharedUpscaledPostDOFTranslucencyColor && TranslucencyPassResources.IsValid() && TranslucencyPassResources.ViewRect.Size() != View.ViewRect.Size() && GetMainTAAPassConfig(View) != EMainTAAPassConfig::TSR)
 		{
 			FTranslucencyComposition TranslucencyComposition;
 			TranslucencyComposition.Operation = FTranslucencyComposition::EOperation::UpscaleOnly;
diff --git a/Engine/Source/Runtime/Renderer/Public/TemporalUpscaler.h b/Engine/Source/Runtime/Renderer/Public/TemporalUpscaler.h
new file mode 100644
--- /dev/null
+++ b/Engine/Source/Runtime/Renderer/Public/TemporalUpscaler.h
@@ -0,0 +1,81 @@
+// Copyright Epic Games, Inc. All Rights Reserved.
+
+#pragma once
+
+#include "ScreenPass.h"
+
+/** Interface for the main temporal upscaling algorithm. */
+class RENDERER_API ITemporalUpscaler : public ISceneViewFamilyExtention
+{
+public:
+
+	/** Ref counted history to be saved in the history. */
+	class RENDERER_API IHistory : public IRefCountedObject
+	{
+	public:
+		virtual ~IHistory() {}
+
+		/** Debug name of the history. Must exactly point to the same const TCHAR* as ITemporalUpscaler::GetDebugName().
+		 * This is used for debugging GPU memory uses of a viewport, but also to ensure IHistory is fed to a compatible ITemporalUpscaler.
+		 */
+		virtual const TCHAR* GetDebugName() const = 0;
+
+		/** Size of the history on the GPU in bytes. */
+		virtual uint64 GetGPUSizeBytes() const = 0;
+	};
+
+	/** Inputs of the temporal upscaler. */
+	struct FInputs
+	{
+		/** Outputs view rect that must be on FOutputs::FullRes::ViewRect. */
+		FIntRect OutputViewRect;
+
+		/** Pixel jitter offset of the rendering pixel. */
+		FVector2f TemporalJitterPixels;
+
+		/** Pre exposure of the SceneColor. */
+		float PreExposure = 1.0f;
+
+		/** The post-DOF and pre-MotionBlur SceneColor. */
+		FScreenPassTexture SceneColor;
+
+		/** The scene depth. */
+		FScreenPassTexture SceneDepth;
+
+		/** The scene velocity. */
+		FScreenPassTexture SceneVelocity;
+
+		/** Texture that contain eye adaptation on the red channel. */
+		FRDGTextureRef EyeAdaptationTexture = nullptr;
+
+		/** The history of the previous frame set by FOutputs::NewHistory. PrevHistory->GetDebugName() is guarentee to match the ITemporalUpscaler. */
+		TRefCountPtr<IHistory> PrevHistory;
+	};
+
+	/** Outputs of the third party temporal upscaler. */
+	struct FOutputs
+	{
+		/** Output of the temporal upscaler. FullRes.ViewRect must match FInputs::OutputViewRect. */
+		FScreenPassTexture FullRes;
+
+		/** New history to be kept alive for next frame. NewHistory->GetDebugName() must exactly point to the same const TCHAR* as ITemporalUpscaler::GetDebugName(). */
+		TRefCountPtr<IHistory> NewHistory;
+	};
+
+	virtual ~ITemporalUpscaler() {};
+
+	/** Debug name of the history. Must exactly point to the same const TCHAR* as ITemporalUpscaler::IHistory::GetDebugName(). */
+	virtual const TCHAR* GetDebugName() const = 0;
+
+	/** Adds the necessary passes into RDG for temporal upscaling the rendering resolution to desired output res. */
+	virtual FOutputs AddPasses(
+		FRDGBuilder& GraphBuilder,
+		const FSceneView& View,
+		const FInputs& Inputs) const = 0;
+
+
+	virtual float GetMinUpsampleResolutionFraction() const = 0;
+	virtual float GetMaxUpsampleResolutionFraction() const = 0;
+
+	virtual ITemporalUpscaler* Fork_GameThread(const class FSceneViewFamily& ViewFamily) const = 0;
+};
