- Do not report materials supporting dithered LOD transition as masked for non-shadow passes unelss stencil dithering is disabled or the VF doesn't support null pixel shader.
- Add an option to allow ignoring material WPO usage based on a per primitive flag during depth shaders selection. This increases the chance of choosing depth only vertex shaders and showed GPU perf improvement on a platform.
- Enable the option on the beneficiary platform.

#localization none
#tests ReplayRun tested

diff --git a/Engine/Source/Runtime/Engine/Private/Materials/MaterialShared.cpp b/Engine/Source/Runtime/Engine/Private/Materials/MaterialShared.cpp
--- a/Engine/Source/Runtime/Engine/Private/Materials/MaterialShared.cpp
+++ b/Engine/Source/Runtime/Engine/Private/Materials/MaterialShared.cpp
@@ -4093,13 +4093,20 @@ bool FMaterial::GetMaterialExpressionSource( FString& OutSource )
 }
 
 bool FMaterial::WritesEveryPixel(bool bShadowPass) const
+{
+	const bool bVFTypeSupportsNullPixelShader = !IsUsedWithInstancedStaticMeshes();
+	return WritesEveryPixel(bShadowPass, bVFTypeSupportsNullPixelShader);
+}
+
+bool FMaterial::WritesEveryPixel(bool bShadowPass, bool bVFTypeSupportsNullPixelShader) const
 {
 	bool bLocalStencilDitheredLOD = FeatureLevel >= ERHIFeatureLevel::SM5 && bStencilDitheredLOD;
 	return !IsMasked()
 		// Render dithered material as masked if a stencil prepass is not used (UE-50064, UE-49537)
 		&& !((bShadowPass || !bLocalStencilDitheredLOD) && IsDitheredLODTransition())
 		&& !IsWireframe()
-		&& !(bLocalStencilDitheredLOD && IsDitheredLODTransition() && IsUsedWithInstancedStaticMeshes())
+		// If the VF type requires a PS, return false so a PS will be picked
+		&& !(bLocalStencilDitheredLOD && IsDitheredLODTransition() && !bVFTypeSupportsNullPixelShader)
 		&& !IsStencilTestEnabled();
 }
 
diff --git a/Engine/Source/Runtime/Engine/Private/PrimitiveSceneProxy.cpp b/Engine/Source/Runtime/Engine/Private/PrimitiveSceneProxy.cpp
--- a/Engine/Source/Runtime/Engine/Private/PrimitiveSceneProxy.cpp
+++ b/Engine/Source/Runtime/Engine/Private/PrimitiveSceneProxy.cpp
@@ -49,6 +49,18 @@ static TAutoConsoleVariable<bool> CVarOptimizedWPO(
 	ECVF_RenderThreadSafe
 );
 
+static TAutoConsoleVariable<bool> CVarOptimizedWPOAffectNonNaniteShaderSelection(
+	TEXT("r.OptimizedWPO.AffectNonNaniteShaderSelection"),
+	false,
+	TEXT("Whether the per primitive WPO flag should affect shader selection for non-nanite primitives. It increase the chance of selecting the position only depth VS ")
+	TEXT("at the cost of updating cached draw commands whenever the WPO flag changes."),
+	FConsoleVariableDelegate::CreateLambda([](IConsoleVariable* InVariable)
+	{
+		FGlobalComponentRecreateRenderStateContext Context;
+	}),
+	ECVF_RenderThreadSafe
+);
+
 static TAutoConsoleVariable<int32> CVarCacheWPOPrimitives(
 	TEXT("r.Shadow.CacheWPOPrimitives"),
 	0,
@@ -102,6 +114,11 @@ bool IsOptimizedWPO()
 	return CVarOptimizedWPO.GetValueOnAnyThread() != 0;
 }
 
+bool ShouldOptimizedWPOAffectNonNaniteShaderSelection()
+{
+	return CVarOptimizedWPOAffectNonNaniteShaderSelection.GetValueOnAnyThread() != 0;
+}
+
 bool CacheShadowDepthsFromPrimitivesUsingWPO()
 {
 	return CVarCacheWPOPrimitives.GetValueOnAnyThread(true) != 0;
diff --git a/Engine/Source/Runtime/Engine/Private/StaticMeshRender.cpp b/Engine/Source/Runtime/Engine/Private/StaticMeshRender.cpp
--- a/Engine/Source/Runtime/Engine/Private/StaticMeshRender.cpp
+++ b/Engine/Source/Runtime/Engine/Private/StaticMeshRender.cpp
@@ -39,6 +39,7 @@
 #include "NaniteSceneProxy.h"
 #include "RenderCore.h"
 #include "DataDrivenShaderPlatformInfo.h"
+#include "EngineModule.h"
 
 #if WITH_EDITOR
 #include "Rendering/StaticLightingSystemInterface.h"
@@ -1412,6 +1413,15 @@ uint8 FStaticMeshSceneProxy::GetCurrentFirstLODIdx_Internal() const
 	return RenderData->CurrentFirstLODIdx;
 }
 
+void FStaticMeshSceneProxy::OnEvaluateWorldPositionOffsetChanged_RenderThread()
+{
+	if (ShouldOptimizedWPOAffectNonNaniteShaderSelection())
+	{
+		// Re-cache draw commands
+		GetRendererModule().BeginDeferredUpdateOfPrimitiveSceneInfo(GetPrimitiveSceneInfo());
+	}
+}
+
 void FStaticMeshSceneProxy::GetMeshDescription(int32 LODIndex, TArray<FMeshBatch>& OutMeshElements) const
 {
 	const FStaticMeshLODResources& LODModel = RenderData->LODResources[LODIndex];
diff --git a/Engine/Source/Runtime/Engine/Public/MaterialShared.h b/Engine/Source/Runtime/Engine/Public/MaterialShared.h
--- a/Engine/Source/Runtime/Engine/Public/MaterialShared.h
+++ b/Engine/Source/Runtime/Engine/Public/MaterialShared.h
@@ -2207,6 +2207,7 @@ public:
 
 	/* Helper function to look at both IsMasked and IsDitheredLODTransition to determine if it writes every pixel */
 	ENGINE_API bool WritesEveryPixel(bool bShadowPass = false) const;
+	ENGINE_API bool WritesEveryPixel(bool bShadowPass, bool bVFTypeSupportsNullPixelShader) const;
 
 	/** call during shader compilation jobs setup to fill additional settings that may be required by classes who inherit from this */
 	virtual void SetupExtraCompilationSettings(const EShaderPlatform Platform, FExtraShaderCompilerSettings& Settings) const
diff --git a/Engine/Source/Runtime/Engine/Public/PrimitiveSceneProxy.h b/Engine/Source/Runtime/Engine/Public/PrimitiveSceneProxy.h
--- a/Engine/Source/Runtime/Engine/Public/PrimitiveSceneProxy.h
+++ b/Engine/Source/Runtime/Engine/Public/PrimitiveSceneProxy.h
@@ -134,6 +134,7 @@ public:
 };
 
 extern bool IsOptimizedWPO();
+ENGINE_API extern bool ShouldOptimizedWPOAffectNonNaniteShaderSelection();
 extern bool IsAllowingApproximateOcclusionQueries();
 extern bool CacheShadowDepthsFromPrimitivesUsingWPO();
 
diff --git a/Engine/Source/Runtime/Engine/Public/StaticMeshSceneProxy.h b/Engine/Source/Runtime/Engine/Public/StaticMeshSceneProxy.h
--- a/Engine/Source/Runtime/Engine/Public/StaticMeshSceneProxy.h
+++ b/Engine/Source/Runtime/Engine/Public/StaticMeshSceneProxy.h
@@ -101,6 +101,8 @@ protected:
 	/** Only call on render thread timeline */
 	uint8 GetCurrentFirstLODIdx_Internal() const;
 
+	virtual void OnEvaluateWorldPositionOffsetChanged_RenderThread() override;
+
 public:
 	// FPrimitiveSceneProxy interface.
 #if WITH_EDITOR
diff --git a/Engine/Source/Runtime/Renderer/Private/CustomDepthRendering.cpp b/Engine/Source/Runtime/Renderer/Private/CustomDepthRendering.cpp
--- a/Engine/Source/Runtime/Renderer/Private/CustomDepthRendering.cpp
+++ b/Engine/Source/Runtime/Renderer/Private/CustomDepthRendering.cpp
@@ -407,7 +407,7 @@ private:
 		ERasterizerFillMode MeshFillMode,
 		ERasterizerCullMode MeshCullMode);
 
-	bool UseDefaultMaterial(const FMaterial& Material, bool bMaterialModifiesMeshPosition, bool bSupportPositionOnlyStream, bool& bPositionOnly, bool& bIgnoreThisMaterial);
+	bool UseDefaultMaterial(const FMaterial& Material, bool bMaterialModifiesMeshPosition, bool bSupportPositionOnlyStream, bool bVFTypeSupportsNullPixelShader, bool& bPositionOnly, bool& bIgnoreThisMaterial);
 
 	void CollectDefaultMaterialPSOInitializers(
 		const FSceneTexturesConfig& SceneTexturesConfig,
@@ -480,7 +480,7 @@ FRHIDepthStencilState* GetCustomDepthStencilState(bool bWriteCustomStencilValues
 	}
 }
 
-bool FCustomDepthPassMeshProcessor::UseDefaultMaterial(const FMaterial& Material, bool bMaterialModifiesMeshPosition, bool bSupportPositionOnlyStream, bool& bPositionOnly, bool& bIgnoreThisMaterial)
+bool FCustomDepthPassMeshProcessor::UseDefaultMaterial(const FMaterial& Material, bool bMaterialModifiesMeshPosition, bool bSupportPositionOnlyStream, bool bVFTypeSupportsNullPixelShader, bool& bPositionOnly, bool& bIgnoreThisMaterial)
 {
 	bool bUseDefaultMaterial = false;
 	bIgnoreThisMaterial = false;
@@ -490,14 +490,14 @@ bool FCustomDepthPassMeshProcessor::UseDefaultMaterial(const FMaterial& Material
 	if (bIsOpaque
 		&& bSupportPositionOnlyStream
 		&& !bMaterialModifiesMeshPosition
-		&& Material.WritesEveryPixel())
+		&& Material.WritesEveryPixel(false, bVFTypeSupportsNullPixelShader))
 	{
 		bUseDefaultMaterial = true;
 		bPositionOnly = true;
 	}
 	else if (!bIsTranslucent || Material.IsTranslucencyWritingCustomDepth())
 	{
-		const bool bMaterialMasked = !Material.WritesEveryPixel() || Material.IsTranslucencyWritingCustomDepth();
+		const bool bMaterialMasked = !Material.WritesEveryPixel(false, bVFTypeSupportsNullPixelShader) || Material.IsTranslucencyWritingCustomDepth();
 		if (!bMaterialMasked && !bMaterialModifiesMeshPosition)
 		{
 			bUseDefaultMaterial = true;
@@ -533,7 +533,11 @@ bool FCustomDepthPassMeshProcessor::TryAddMeshBatch(
 	// Using default material?
 	bool bIgnoreThisMaterial = false;
 	bool bPositionOnly = false;
-	bool bUseDefaultMaterial = UseDefaultMaterial(Material, Material.MaterialModifiesMeshPosition_RenderThread(), MeshBatch.VertexFactory->SupportsPositionOnlyStream(), bPositionOnly, bIgnoreThisMaterial);
+	const bool bSupportPositionOnlyStream = MeshBatch.VertexFactory->SupportsPositionOnlyStream();
+	const bool bVFTypeSupportsNullPixelShader = MeshBatch.VertexFactory->SupportsNullPixelShader();
+	const bool bEvaluateWPO = Material.MaterialModifiesMeshPosition_RenderThread()
+		&& (!ShouldOptimizedWPOAffectNonNaniteShaderSelection() || PrimitiveSceneProxy->EvaluateWorldPositionOffset());
+	bool bUseDefaultMaterial = UseDefaultMaterial(Material, bEvaluateWPO, bSupportPositionOnlyStream, bVFTypeSupportsNullPixelShader, bPositionOnly, bIgnoreThisMaterial);
 	if (bIgnoreThisMaterial)
 	{
 		return true;
@@ -631,10 +635,11 @@ void FCustomDepthPassMeshProcessor::CollectPSOInitializers(const FSceneTexturesC
 	}
 
 	// assume we can always do this when collecting PSO's for now (vertex factory instance might actually not support it)
-	bool bSupportPositionOnlyStream = VertexFactoryData.VertexFactoryType->SupportsPositionOnly();
+	const bool bSupportPositionOnlyStream = VertexFactoryData.VertexFactoryType->SupportsPositionOnly();
+	const bool bVFTypeSupportsNullPixelShader = VertexFactoryData.VertexFactoryType->SupportsNullPixelShader();
 	bool bIgnoreThisMaterial = false;
 	bool bPositionOnly = false;
-	bool bUseDefaultMaterial = UseDefaultMaterial(Material, Material.MaterialModifiesMeshPosition_GameThread(), bSupportPositionOnlyStream, bPositionOnly, bIgnoreThisMaterial);
+	bool bUseDefaultMaterial = UseDefaultMaterial(Material, Material.MaterialModifiesMeshPosition_GameThread(), bSupportPositionOnlyStream, bVFTypeSupportsNullPixelShader, bPositionOnly, bIgnoreThisMaterial);
 
 	if (!bIgnoreThisMaterial)
 	{
diff --git a/Engine/Source/Runtime/Renderer/Private/DepthRendering.cpp b/Engine/Source/Runtime/Renderer/Private/DepthRendering.cpp
--- a/Engine/Source/Runtime/Renderer/Private/DepthRendering.cpp
+++ b/Engine/Source/Runtime/Renderer/Private/DepthRendering.cpp
@@ -190,7 +190,8 @@ bool GetDepthPassShaders(
 	}
 	else
 	{
-		const bool bNeedsPixelShader = !Material.WritesEveryPixel() || bMaterialUsesPixelDepthOffset || Material.IsTranslucencyWritingCustomDepth();
+		const bool bVFTypeSupportsNullPixelShader = VertexFactoryType->SupportsNullPixelShader();
+		const bool bNeedsPixelShader = !Material.WritesEveryPixel(false, bVFTypeSupportsNullPixelShader) || bMaterialUsesPixelDepthOffset || Material.IsTranslucencyWritingCustomDepth();
 		if (bNeedsPixelShader)
 		{
 			ShaderTypes.AddShaderType<FDepthOnlyPS>();
@@ -902,7 +903,7 @@ void FDepthPassMeshProcessor::CollectPSOInitializersInternal(
 		PSOInitializers);
 }
 
-bool FDepthPassMeshProcessor::UseDefaultMaterial(const FMaterial& Material, bool bMaterialModifiesMeshPosition, bool bSupportPositionOnlyStream, bool& bPositionOnly)
+bool FDepthPassMeshProcessor::UseDefaultMaterial(const FMaterial& Material, bool bMaterialModifiesMeshPosition, bool bSupportPositionOnlyStream, bool bVFTypeSupportsNullPixelShader, bool& bPositionOnly)
 {
 	bool bUseDefaultMaterial = false;
 
@@ -910,7 +911,7 @@ bool FDepthPassMeshProcessor::UseDefaultMaterial(const FMaterial& Material, bool
 		&& EarlyZPassMode != DDM_MaskedOnly
 		&& bSupportPositionOnlyStream
 		&& !bMaterialModifiesMeshPosition
-		&& Material.WritesEveryPixel())
+		&& Material.WritesEveryPixel(false, bVFTypeSupportsNullPixelShader))
 	{
 		bUseDefaultMaterial = true;
 		bPositionOnly = true;
@@ -918,7 +919,7 @@ bool FDepthPassMeshProcessor::UseDefaultMaterial(const FMaterial& Material, bool
 	else
 	{
 		// still possible to use default material
-		const bool bMaterialMasked = !Material.WritesEveryPixel() || Material.IsTranslucencyWritingCustomDepth();
+		const bool bMaterialMasked = !Material.WritesEveryPixel(false, bVFTypeSupportsNullPixelShader) || Material.IsTranslucencyWritingCustomDepth();
 		if ((!bMaterialMasked && EarlyZPassMode != DDM_MaskedOnly) || (bMaterialMasked && EarlyZPassMode != DDM_NonMaskedOnly))
 		{
 			if (!bMaterialMasked && !bMaterialModifiesMeshPosition)
@@ -943,8 +944,12 @@ bool FDepthPassMeshProcessor::TryAddMeshBatch(const FMeshBatch& RESTRICT MeshBat
 		&& ShouldIncludeDomainInMeshPass(Material.GetMaterialDomain())
 		&& ShouldIncludeMaterialInDefaultOpaquePass(Material))
 	{
+		const bool bSupportPositionOnlyStream = MeshBatch.VertexFactory->SupportsPositionOnlyStream();
+		const bool bVFTypeSupportsNullPixelShader = MeshBatch.VertexFactory->SupportsNullPixelShader();
+		const bool bEvaluateWPO = Material.MaterialModifiesMeshPosition_RenderThread()
+			&& (!ShouldOptimizedWPOAffectNonNaniteShaderSelection() || PrimitiveSceneProxy->EvaluateWorldPositionOffset());
 		bool bPositionOnly = false;
-		bool bUseDefaultMaterial = UseDefaultMaterial(Material, Material.MaterialModifiesMeshPosition_RenderThread(), MeshBatch.VertexFactory->SupportsPositionOnlyStream(), bPositionOnly);
+		bool bUseDefaultMaterial = UseDefaultMaterial(Material, bEvaluateWPO, bSupportPositionOnlyStream, bVFTypeSupportsNullPixelShader, bPositionOnly);
 
 		const FMaterialRenderProxy* EffectiveMaterialRenderProxy = &MaterialRenderProxy;
 		const FMaterial* EffectiveMaterial = &Material;
@@ -1087,10 +1092,10 @@ void FDepthPassMeshProcessor::CollectPSOInitializers(const FSceneTexturesConfig&
 	}
 
 	// assume we can always do this when collecting PSO's for now (vertex factory instance might actually not support it)
-	bool bSupportPositionOnlyStream = VertexFactoryData.VertexFactoryType->SupportsPositionOnly();  
-
+	const bool bSupportPositionOnlyStream = VertexFactoryData.VertexFactoryType->SupportsPositionOnly();  
+	const bool bVFTypeSupportsNullPixelShader = VertexFactoryData.VertexFactoryType->SupportsNullPixelShader();
 	bool bPositionOnly = false;
-	bool bUseDefaultMaterial = UseDefaultMaterial(Material, Material.MaterialModifiesMeshPosition_GameThread(), bSupportPositionOnlyStream, bPositionOnly);
+	bool bUseDefaultMaterial = UseDefaultMaterial(Material, Material.MaterialModifiesMeshPosition_GameThread(), bSupportPositionOnlyStream, bVFTypeSupportsNullPixelShader, bPositionOnly);
 
 	const FMaterial* EffectiveMaterial = &Material;
 	if (bUseDefaultMaterial && !bSupportPositionOnlyStream && VertexFactoryData.CustomDefaultVertexDeclaration)
diff --git a/Engine/Source/Runtime/Renderer/Private/DepthRendering.h b/Engine/Source/Runtime/Renderer/Private/DepthRendering.h
--- a/Engine/Source/Runtime/Renderer/Private/DepthRendering.h
+++ b/Engine/Source/Runtime/Renderer/Private/DepthRendering.h
@@ -205,7 +205,7 @@ private:
 		ERasterizerFillMode MeshFillMode,
 		ERasterizerCullMode MeshCullMode);
 
-	bool UseDefaultMaterial(const FMaterial& Material, bool bMaterialModifiesMeshPosition, bool bSupportPositionOnlyStream, bool& bPositionOnly);
+	bool UseDefaultMaterial(const FMaterial& Material, bool bMaterialModifiesMeshPosition, bool bSupportPositionOnlyStream, bool bVFTypeSupportsNullPixelShader, bool& bPositionOnly);
 
 	void CollectDefaultMaterialPSOInitializers(
 		const FSceneTexturesConfig& SceneTexturesConfig, 
diff --git a/Engine/Source/Runtime/Renderer/Private/ShadowDepthRendering.cpp b/Engine/Source/Runtime/Renderer/Private/ShadowDepthRendering.cpp
--- a/Engine/Source/Runtime/Renderer/Private/ShadowDepthRendering.cpp
+++ b/Engine/Source/Runtime/Renderer/Private/ShadowDepthRendering.cpp
@@ -487,10 +487,11 @@ IMPLEMENT_SHADOWDEPTHPASS_PIXELSHADER_TYPE(PixelShadowDepth_VirtualShadowMap);
  */
 static bool UseDefaultMaterialForShadowDepth(
 	const FMaterial& Material,
+	bool bVFTypeSupportsNullPixelShader,
 	bool MaterialModifiesMeshPosition)
 {
 	const bool bShadowPass = true;
-	return (Material.WritesEveryPixel(bShadowPass) &&
+	return (Material.WritesEveryPixel(bShadowPass, bVFTypeSupportsNullPixelShader) &&
 		!MaterialModifiesMeshPosition);
 }
 
@@ -505,9 +506,10 @@ bool GetShadowDepthPassShaders(
 	TShaderRef<FShadowDepthVS>& VertexShader,
 	TShaderRef<FShadowDepthBasePS>& PixelShader)
 {
+	const bool bVFTypeSupportsNullPixelShader = VertexFactoryType->SupportsNullPixelShader();
 	const bool bPositionOnlyVS =
 		bSupportsPositionAndNormalOnlyStream
-		&& Material.WritesEveryPixel(true)
+		&& Material.WritesEveryPixel(true, bVFTypeSupportsNullPixelShader)
 		&& !bMaterialModifiesMeshPosition;
 
 	// Use perspective correct shadow depths for shadow types which typically render low poly meshes into the shadow depth buffer.
@@ -580,7 +582,7 @@ bool GetShadowDepthPassShaders(
 	}
 
 	// Pixel shaders
-	const bool bNullPixelShader = Material.WritesEveryPixel(true) && !bUsePerspectiveCorrectShadowDepths && !bVirtualShadowMap && VertexFactoryType->SupportsNullPixelShader();
+	const bool bNullPixelShader = Material.WritesEveryPixel(true, bVFTypeSupportsNullPixelShader) && !bUsePerspectiveCorrectShadowDepths && !bVirtualShadowMap && bVFTypeSupportsNullPixelShader;
 	if (!bNullPixelShader)
 	{
 		if (bVirtualShadowMap)
@@ -1956,9 +1958,10 @@ void FShadowDepthPassMeshProcessor::CollectPSOInitializersInternal(
 		return;
 	}
 
+	const bool bVFTypeSupportsNullPixelShader = VertexFactoryData.VertexFactoryType->SupportsNullPixelShader();
 	const bool bUsePositionOnlyVS =
 		bSupportsPositionAndNormalOnlyStream
-		&& MaterialResource.WritesEveryPixel(true)
+		&& MaterialResource.WritesEveryPixel(true, bVFTypeSupportsNullPixelShader)
 		&& !MaterialResource.MaterialModifiesMeshPosition_GameThread();
 	
 	FGraphicsPipelineRenderTargetsInfo RenderTargetsInfo;	
@@ -2028,7 +2031,7 @@ bool FShadowDepthPassMeshProcessor::Process(
 
 	const bool bUsePositionOnlyVS =
 		VertexFactory->SupportsPositionAndNormalOnlyStream()
-		&& MaterialResource.WritesEveryPixel(true)
+		&& MaterialResource.WritesEveryPixel(true, VertexFactory->SupportsNullPixelShader())
 		&& !MaterialResource.MaterialModifiesMeshPosition_RenderThread();
 
 	// Need to replicate for cube faces on host if GPU-scene is not available (for this draw).
@@ -2081,8 +2084,11 @@ bool FShadowDepthPassMeshProcessor::TryAddMeshBatch(const FMeshBatch& RESTRICT M
 	{
 		const FMaterialRenderProxy* EffectiveMaterialRenderProxy = &MaterialRenderProxy;
 		const FMaterial* EffectiveMaterial = &Material;
+		const bool bVFTypeSupportsNullPixelShader = MeshBatch.VertexFactory->SupportsNullPixelShader();
+		const bool bEvaluateWPO = Material.MaterialModifiesMeshPosition_RenderThread()
+			&& (!ShouldOptimizedWPOAffectNonNaniteShaderSelection() || PrimitiveSceneProxy->EvaluateWorldPositionOffset());
 
-		if (UseDefaultMaterialForShadowDepth(Material, Material.MaterialModifiesMeshPosition_RenderThread()))
+		if (UseDefaultMaterialForShadowDepth(Material, bVFTypeSupportsNullPixelShader, bEvaluateWPO))
 		{
 			const FMaterialRenderProxy* DefaultProxy = UMaterial::GetDefaultMaterial(MD_Surface)->GetRenderProxy();
 			const FMaterial* DefaultMaterialResource = DefaultProxy->GetMaterialNoFallback(FeatureLevel);
@@ -2144,8 +2150,9 @@ void FShadowDepthPassMeshProcessor::CollectPSOInitializers(const FSceneTexturesC
 	else if (PreCacheParams.bCastShadow)
 	{
 		bool bCollectPSOs = true;
+		const bool bVFTypeSupportsNullPixelShader = VertexFactoryData.VertexFactoryType->SupportsNullPixelShader();
 		const FMaterial* EffectiveMaterial = &Material;
-		if (UseDefaultMaterialForShadowDepth(Material, Material.MaterialModifiesMeshPosition_GameThread()))
+		if (UseDefaultMaterialForShadowDepth(Material, bVFTypeSupportsNullPixelShader, Material.MaterialModifiesMeshPosition_GameThread()))
 		{
 			if (VertexFactoryData.CustomDefaultVertexDeclaration)
 			{
diff --git a/Engine/Source/Runtime/Renderer/Private/SingleLayerWaterRendering.cpp b/Engine/Source/Runtime/Renderer/Private/SingleLayerWaterRendering.cpp
--- a/Engine/Source/Runtime/Renderer/Private/SingleLayerWaterRendering.cpp
+++ b/Engine/Source/Runtime/Renderer/Private/SingleLayerWaterRendering.cpp
@@ -1745,11 +1745,14 @@ bool FSingleLayerWaterDepthPrepassMeshProcessor::TryAddMeshBatch(
 		const FMeshDrawingPolicyOverrideSettings OverrideSettings = ComputeMeshOverrideSettings(MeshBatch);
 		const ERasterizerFillMode MeshFillMode = ComputeMeshFillMode(Material, OverrideSettings);
 		const ERasterizerCullMode MeshCullMode = ComputeMeshCullMode(Material, OverrideSettings);
+		const bool bVFTypeSupportsNullPixelShader = MeshBatch.VertexFactory->SupportsNullPixelShader();
+		const bool bEvaluateWPO = Material.MaterialModifiesMeshPosition_RenderThread()
+			&& (!ShouldOptimizedWPOAffectNonNaniteShaderSelection() || PrimitiveSceneProxy->EvaluateWorldPositionOffset());
 
 		if (IsOpaqueBlendMode(Material)
 			&& MeshBatch.VertexFactory->SupportsPositionOnlyStream()
-			&& !Material.MaterialModifiesMeshPosition_RenderThread()
-			&& Material.WritesEveryPixel())
+			&& !bEvaluateWPO
+			&& Material.WritesEveryPixel(false, bVFTypeSupportsNullPixelShader))
 		{
 			const FMaterialRenderProxy& DefaultProxy = *UMaterial::GetDefaultMaterial(MD_Surface)->GetRenderProxy();
 			const FMaterial& DefaultMaterial = *DefaultProxy.GetMaterialNoFallback(FeatureLevel);
@@ -1757,11 +1760,11 @@ bool FSingleLayerWaterDepthPrepassMeshProcessor::TryAddMeshBatch(
 		}
 		else
 		{
-			const bool bMaterialMasked = !Material.WritesEveryPixel();
+			const bool bMaterialMasked = !Material.WritesEveryPixel(false, bVFTypeSupportsNullPixelShader);
 			const FMaterialRenderProxy* EffectiveMaterialRenderProxy = &MaterialRenderProxy;
 			const FMaterial* EffectiveMaterial = &Material;
 
-			if (!bMaterialMasked && !Material.MaterialModifiesMeshPosition_RenderThread())
+			if (!bMaterialMasked && !bEvaluateWPO)
 			{
 				// Override with the default material for opaque materials that are not two sided
 				EffectiveMaterialRenderProxy = UMaterial::GetDefaultMaterial(MD_Surface)->GetRenderProxy();
@@ -1834,11 +1837,12 @@ void FSingleLayerWaterDepthPrepassMeshProcessor::CollectPSOInitializers(const FS
 		const ERasterizerFillMode MeshFillMode = ComputeMeshFillMode(Material, OverrideSettings);
 		const ERasterizerCullMode MeshCullMode = ComputeMeshCullMode(Material, OverrideSettings);
 		const bool bSupportPositionOnlyStream = VertexFactoryData.VertexFactoryType->SupportsPositionOnly();
+		const bool bVFTypeSupportsNullPixelShader = VertexFactoryData.VertexFactoryType->SupportsNullPixelShader();
 
 		if (IsOpaqueBlendMode(Material)
 			&& bSupportPositionOnlyStream
 			&& !Material.MaterialModifiesMeshPosition_GameThread()
-			&& Material.WritesEveryPixel())
+			&& Material.WritesEveryPixel(false, bVFTypeSupportsNullPixelShader))
 		{
 			const FMaterialRenderProxy& DefaultProxy = *UMaterial::GetDefaultMaterial(MD_Surface)->GetRenderProxy();
 			const FMaterial& DefaultMaterial = *DefaultProxy.GetMaterialNoFallback(FeatureLevel);
@@ -1847,7 +1851,7 @@ void FSingleLayerWaterDepthPrepassMeshProcessor::CollectPSOInitializers(const FS
 		}
 		else
 		{
-			const bool bMaterialMasked = !Material.WritesEveryPixel();
+			const bool bMaterialMasked = !Material.WritesEveryPixel(false, bVFTypeSupportsNullPixelShader);
 			const FMaterial* EffectiveMaterial = &Material;
 
 			if (!bMaterialMasked && !Material.MaterialModifiesMeshPosition_GameThread())
diff --git a/Engine/Source/Runtime/Renderer/Private/VelocityRendering.cpp b/Engine/Source/Runtime/Renderer/Private/VelocityRendering.cpp
--- a/Engine/Source/Runtime/Renderer/Private/VelocityRendering.cpp
+++ b/Engine/Source/Runtime/Renderer/Private/VelocityRendering.cpp
@@ -497,12 +497,12 @@ bool FOpaqueVelocityMeshProcessor::PrimitiveHasVelocityForFrame(const FPrimitive
 	return true;
 }
 
-static bool UseDefaultMaterial(const FMaterial* Material, bool bMaterialModifiesMeshPosition)
+static bool UseDefaultMaterial(const FMaterial* Material, bool bVFTypeSupportsNullPixelShader, bool bMaterialModifiesMeshPosition)
 {
 	// Materials without masking or custom vertex modifications can be swapped out
 	// for the default material, which simplifies the shader. However, the default
 	// material also does not support being two-sided.
-	return Material->WritesEveryPixel() && !Material->IsTwoSided() && !bMaterialModifiesMeshPosition;
+	return Material->WritesEveryPixel(false, bVFTypeSupportsNullPixelShader) && !Material->IsTwoSided() && !bMaterialModifiesMeshPosition;
 }
 
 bool FOpaqueVelocityMeshProcessor::TryAddMeshBatch(
@@ -522,8 +522,11 @@ bool FOpaqueVelocityMeshProcessor::TryAddMeshBatch(
 		const FMeshDrawingPolicyOverrideSettings OverrideSettings = ComputeMeshOverrideSettings(MeshBatch);
 		const ERasterizerFillMode MeshFillMode = ComputeMeshFillMode(*Material, OverrideSettings);
 		const ERasterizerCullMode MeshCullMode = ComputeMeshCullMode(*Material, OverrideSettings);
+		const bool bVFTypeSupportsNullPixelShader = MeshBatch.VertexFactory->SupportsNullPixelShader();
+		const bool bEvaluateWPO = Material->MaterialModifiesMeshPosition_RenderThread()
+			&& (!ShouldOptimizedWPOAffectNonNaniteShaderSelection() || PrimitiveSceneProxy->EvaluateWorldPositionOffset());
 
-		const bool bSwapWithDefaultMaterial = UseDefaultMaterial(Material, Material->MaterialModifiesMeshPosition_RenderThread());
+		const bool bSwapWithDefaultMaterial = UseDefaultMaterial(Material, bVFTypeSupportsNullPixelShader, bEvaluateWPO);
 		if (bSwapWithDefaultMaterial)
 		{
 			MaterialRenderProxy = UMaterial::GetDefaultMaterial(MD_Surface)->GetRenderProxy();
@@ -603,8 +606,9 @@ void FOpaqueVelocityMeshProcessor::CollectPSOInitializers(const FSceneTexturesCo
 		const bool bIsNotTranslucent = IsOpaqueOrMaskedBlendMode(Material);
 
 		if (PreCacheParams.bRenderInMainPass && bIsNotTranslucent && ShouldIncludeMaterialInDefaultOpaquePass(Material))
-		{			
-			const bool bUseDefaultMaterial = UseDefaultMaterial(&Material, Material.MaterialModifiesMeshPosition_GameThread());
+		{
+			const bool bVFTypeSupportsNullPixelShader = VertexFactoryData.VertexFactoryType->SupportsNullPixelShader();
+			const bool bUseDefaultMaterial = UseDefaultMaterial(&Material, bVFTypeSupportsNullPixelShader, Material.MaterialModifiesMeshPosition_GameThread());
 			if (!bUseDefaultMaterial)
 			{
 				bCollectPSOs = true;
