Change engine default of r.VelocityOutputPass back to 0, now that Nanite properly supports WPO materials writing velocity in the base pass. This results in the majority of Nanite materials using significantly fewer VGPRs with the velocity calculations getting removed (and done more efficiently in the Nanite material depth export pass).

diff --git a/Engine/Source/Runtime/Renderer/Private/VelocityRendering.cpp b/Engine/Source/Runtime/Renderer/Private/VelocityRendering.cpp
--- a/Engine/Source/Runtime/Renderer/Private/VelocityRendering.cpp
+++ b/Engine/Source/Runtime/Renderer/Private/VelocityRendering.cpp
@@ -29,7 +29,7 @@
 // Changing this causes a full shader recompile
 static TAutoConsoleVariable<int32> CVarVelocityOutputPass(
 	TEXT("r.VelocityOutputPass"),
-	1,
+	0,
 	TEXT("When to write velocity buffer.\n") \
 	TEXT(" 0: Renders during the depth pass. This splits the depth pass into 2 phases: with and without velocity.\n") \
 	TEXT(" 1: Renders during the regular base pass. This adds an extra GBuffer target during base pass rendering.") \
