Fix for drawing debug primitives in packaged Dev and Debug builds. Will not draw in Shipping/Test. Compile TemporalUpscaleDepth + PopulateDepth for all appropriate platforms to avoid shader crash issues.

diff --git a/Engine/Source/Runtime/Renderer/Private/PostProcess/PostProcessCompositePrimitivesCommon.cpp b/Engine/Source/Runtime/Renderer/Private/PostProcess/PostProcessCompositePrimitivesCommon.cpp
--- a/Engine/Source/Runtime/Renderer/Private/PostProcess/PostProcessCompositePrimitivesCommon.cpp
+++ b/Engine/Source/Runtime/Renderer/Private/PostProcess/PostProcessCompositePrimitivesCommon.cpp
@@ -76,8 +76,11 @@ public:
 
 		static bool ShouldCompilePermutation(const FGlobalShaderPermutationParameters& Parameters)
 	{
-		// Only render if debug viewmodes are allowed.
-		return AllowDebugViewmodes(Parameters.Platform);
+		/** 
+		* Compile for all platforms that support TAA minimum.
+		* Runtime only executes if a temporal upscale method is enabled.
+		*/
+		return SupportsGen4TAA(Parameters.Platform);
 	}
 };
 
@@ -126,8 +129,8 @@ public:
 			return false;
 		}
 
-		// Only render if debug viewmodes are allowed.
-		return AllowDebugViewmodes(Parameters.Platform);
+		// This pass can be used with any AA method so return true for all other circumstances
+		return true;
 	}
 };
 
diff --git a/Engine/Source/Runtime/Renderer/Private/SceneRendering.cpp b/Engine/Source/Runtime/Renderer/Private/SceneRendering.cpp
--- a/Engine/Source/Runtime/Renderer/Private/SceneRendering.cpp
+++ b/Engine/Source/Runtime/Renderer/Private/SceneRendering.cpp
@@ -4056,7 +4056,13 @@ bool FSceneRenderer::ShouldCompositeEditorPrimitives(const FViewInfo& View)
 #if UE_ENABLE_DEBUG_DRAWING
 bool FSceneRenderer::ShouldCompositeDebugPrimitivesInPostProcess(const FViewInfo& View)
 {
-	return AllowDebugViewmodes() && View.DebugSimpleElementCollector.HasAnyPrimitives();
+	if (View.DebugSimpleElementCollector.HasAnyPrimitives())
+	{
+		//If we have primitives to draw, check we aren't forcing debug view modes off.
+		static bool bIsForceDisabled = IConsoleManager::Get().FindTConsoleVariableDataInt(TEXT("r.ForceDebugViewModes"))->GetValueOnAnyThread() == 2;
+		return !bIsForceDisabled;
+	}
+	return false;
 }
 #endif
 
