PR #10333: Fix light bounds for volumetric fog rendering.

The new math is actually correct, tested and it work better indeed for the case described.

diff --git a/Engine/Source/Runtime/Renderer/Private/VolumetricFog.cpp b/Engine/Source/Runtime/Renderer/Private/VolumetricFog.cpp
--- a/Engine/Source/Runtime/Renderer/Private/VolumetricFog.cpp
+++ b/Engine/Source/Runtime/Renderer/Private/VolumetricFog.cpp
@@ -395,8 +395,8 @@ public:
 		VertexBufferRHI = RHICreateBuffer(Size, BUF_Static | BUF_VertexBuffer, 0, ERHIAccess::VertexOrIndexBuffer, CreateInfo);
 		FScreenVertex* DestVertex = (FScreenVertex*)RHILockBuffer(VertexBufferRHI, 0, Size, RLM_WriteOnly);
 
-		const int32 NumRings = NumVertices;
-		const float RadiansPerRingSegment = PI / (float)NumRings;
+		const int32 NumSegments = NumVertices - 1;
+		const float RadiansPerRingSegment = PI / (float)NumSegments;
 
 		// Boost the effective radius so that the edges of the circle approximation lie on the circle, instead of the vertices
 		const float RadiusScale = 1.0f / FMath::Cos(RadiansPerRingSegment);
