Fix a static analysis warning

diff --git a/Engine/Source/Runtime/Renderer/Private/DepthRendering.cpp b/Engine/Source/Runtime/Renderer/Private/DepthRendering.cpp
--- a/Engine/Source/Runtime/Renderer/Private/DepthRendering.cpp
+++ b/Engine/Source/Runtime/Renderer/Private/DepthRendering.cpp
@@ -947,7 +947,7 @@ bool FDepthPassMeshProcessor::TryAddMeshBatch(const FMeshBatch& RESTRICT MeshBat
 		const bool bSupportPositionOnlyStream = MeshBatch.VertexFactory->SupportsPositionOnlyStream();
 		const bool bVFTypeSupportsNullPixelShader = MeshBatch.VertexFactory->SupportsNullPixelShader();
 		const bool bEvaluateWPO = Material.MaterialModifiesMeshPosition_RenderThread()
-			&& (!ShouldOptimizedWPOAffectNonNaniteShaderSelection() || PrimitiveSceneProxy->EvaluateWorldPositionOffset());
+			&& (!ShouldOptimizedWPOAffectNonNaniteShaderSelection() || !PrimitiveSceneProxy || PrimitiveSceneProxy->EvaluateWorldPositionOffset());
 		bool bPositionOnly = false;
 		bool bUseDefaultMaterial = UseDefaultMaterial(Material, bEvaluateWPO, bSupportPositionOnlyStream, bVFTypeSupportsNullPixelShader, bPositionOnly);
 
