Moved GPU-Scene debug rendering over to using the Scene uniform buffer.

diff --git a/Engine/Shaders/Private/GPUSceneDebugRender.usf b/Engine/Shaders/Private/GPUSceneDebugRender.usf
--- a/Engine/Shaders/Private/GPUSceneDebugRender.usf
+++ b/Engine/Shaders/Private/GPUSceneDebugRender.usf
@@ -4,12 +4,9 @@
 #include "SceneData.ush"
 #include "ShaderPrint.ush"
 
-int NumInstances;
-int NumScenePrimitives;
 int bDrawAll;
 int bDrawUpdatedOnly;
 
-uint InstanceDataSOAStride;
 int2 InputCoord;
 float3 PickingRayStart;
 float3 PickingRayEnd;
@@ -88,12 +85,12 @@ void AddFlag(inout FShaderPrintContext Context, bool bEnabled)
 [numthreads(NUM_THREADS_PER_GROUP, 1, 1)]
 void GPUSceneDebugRenderCS(uint InstanceID : SV_DispatchThreadID)
 {
-	if (InstanceID >= NumInstances)
+	if (InstanceID >= Scene.GPUScene.NumInstances)
 	{
 		return;
 	}
 
-	FInstanceSceneData InstanceData = GetInstanceSceneData(InstanceID, InstanceDataSOAStride);
+	FInstanceSceneData InstanceData = GetInstanceSceneData(InstanceID, Scene.GPUScene.InstanceDataSOAStride);
 
 	// Early out on invalid ID
 	if (InstanceData.PrimitiveId == 0xFFFFFFFFu)
@@ -127,7 +124,7 @@ void GPUSceneDebugRenderCS(uint InstanceID : SV_DispatchThreadID)
 	const uint BitsPerWord = 32U;
 	const uint Mask = 1U << (InstanceData.PrimitiveId % BitsPerWord);
 
-	bool bIsSelected = (InstanceData.PrimitiveId < NumScenePrimitives && (SelectedPrimitiveFlags[InstanceData.PrimitiveId / BitsPerWord] & Mask) != 0U);
+	bool bIsSelected = (InstanceData.PrimitiveId < Scene.GPUScene.NumScenePrimitives && (SelectedPrimitiveFlags[InstanceData.PrimitiveId / BitsPerWord] & Mask) != 0U);
 
 	//float4 InstanceBoundColor = InstanceData.PrimitiveId < NumScenePrimitives ? float4(0.3f, 1.0f, 0.5f, 1.0f) : float4(0.3f, 0.3f, 0.3f, 0.5f);
 	float4 InstanceBoundColor = float4(ColorSilver.xyz, 0.5f);
diff --git a/Engine/Source/Runtime/Renderer/Private/DeferredShadingRenderer.cpp b/Engine/Source/Runtime/Renderer/Private/DeferredShadingRenderer.cpp
--- a/Engine/Source/Runtime/Renderer/Private/DeferredShadingRenderer.cpp
+++ b/Engine/Source/Runtime/Renderer/Private/DeferredShadingRenderer.cpp
@@ -2899,7 +2899,7 @@ void FDeferredShadingSceneRenderer::Render(FRDGBuilder& GraphBuilder)
 
 			Scene->GPUScene.UploadDynamicPrimitiveShaderDataForView(GraphBuilder, *Scene, View, ExternalAccessQueue);
 
-			Scene->GPUScene.DebugRender(GraphBuilder, *Scene, View);
+			Scene->GPUScene.DebugRender(GraphBuilder, *Scene, GetSceneUniforms(), View);
 		}
 
 		InstanceCullingManager.BeginDeferredCulling(GraphBuilder, Scene->GPUScene);
diff --git a/Engine/Source/Runtime/Renderer/Private/GPUScene.cpp b/Engine/Source/Runtime/Renderer/Private/GPUScene.cpp
--- a/Engine/Source/Runtime/Renderer/Private/GPUScene.cpp
+++ b/Engine/Source/Runtime/Renderer/Private/GPUScene.cpp
@@ -1744,7 +1744,7 @@ class FGPUSceneDebugRenderCS : public FGlobalShader
 	SHADER_USE_PARAMETER_STRUCT(FGPUSceneDebugRenderCS, FGlobalShader);
 
 	BEGIN_SHADER_PARAMETER_STRUCT(FParameters, )
-		SHADER_PARAMETER_STRUCT_INCLUDE(FGPUSceneResourceParameters, GPUSceneResource)
+		SHADER_PARAMETER_RDG_UNIFORM_BUFFER(FSceneUniformParameters, Scene)
 		SHADER_PARAMETER_STRUCT_INCLUDE(ShaderPrint::FShaderParameters, ShaderPrintUniformBuffer)
 		SHADER_PARAMETER(int32, bDrawAll)
 		SHADER_PARAMETER(int32, bDrawUpdatedOnly)
@@ -1772,7 +1772,6 @@ public:
 		FGlobalShader::ModifyCompilationEnvironment(Parameters, OutEnvironment);
 
 		OutEnvironment.SetDefine(TEXT("VF_SUPPORTS_PRIMITIVE_SCENE_DATA"), 1);
-		OutEnvironment.SetDefine(TEXT("USE_GLOBAL_GPU_SCENE_DATA"), 1);
 		OutEnvironment.SetDefine(TEXT("NUM_THREADS_PER_GROUP"), NumThreadsPerGroup);
 
 		// Skip optimization for avoiding long compilation time due to large UAV writes
@@ -1782,7 +1781,7 @@ public:
 
 IMPLEMENT_GLOBAL_SHADER(FGPUSceneDebugRenderCS, "/Engine/Private/GPUSceneDebugRender.usf", "GPUSceneDebugRenderCS", SF_Compute);
 
-void FGPUScene::DebugRender(FRDGBuilder& GraphBuilder, FScene& Scene, FViewInfo& View)
+void FGPUScene::DebugRender(FRDGBuilder& GraphBuilder, FScene& Scene, FSceneUniformBuffer& SceneUniformBuffer, FViewInfo& View)
 {
 	int32 DebugMode = CVarGPUSceneDebugMode.GetValueOnRenderThread();
 	if (DebugMode > 0)
@@ -1857,7 +1856,7 @@ void FGPUScene::DebugRender(FRDGBuilder& GraphBuilder, FScene& Scene, FViewInfo&
 
 			FGPUSceneDebugRenderCS::FParameters* PassParameters = GraphBuilder.AllocParameters<FGPUSceneDebugRenderCS::FParameters>();
 			ShaderPrint::SetParameters(GraphBuilder, View.ShaderPrintData, PassParameters->ShaderPrintUniformBuffer);
-			PassParameters->GPUSceneResource = ShaderParameters;
+			PassParameters->Scene = SceneUniformBuffer.GetBuffer(GraphBuilder);
 			PassParameters->bDrawUpdatedOnly = DebugMode == 3;
 			PassParameters->bDrawAll = DebugMode != 2;
 			PassParameters->SelectedNameInfoCount = SelectedCount;
diff --git a/Engine/Source/Runtime/Renderer/Private/GPUScene.h b/Engine/Source/Runtime/Renderer/Private/GPUScene.h
--- a/Engine/Source/Runtime/Renderer/Private/GPUScene.h
+++ b/Engine/Source/Runtime/Renderer/Private/GPUScene.h
@@ -307,7 +307,7 @@ public:
 	 * Draw GPU-Scene debug info, such as bounding boxes. Call once per view at some point in the frame after GPU scene has been updated fully.
 	 * What is drawn is controlled by the CVar: r.GPUScene.DebugMode. Enabling this cvar causes ShaderDraw to be being active (if supported). 
 	 */
-	void DebugRender(FRDGBuilder& GraphBuilder, FScene& Scene, FViewInfo& View);
+	void DebugRender(FRDGBuilder& GraphBuilder, FScene& Scene, FSceneUniformBuffer& SceneUniformBuffer, FViewInfo& View);
 
 	/**
 	 * Manually trigger an allocator consolidate (will otherwise be done when an item is allocated).
