Fixing ShaderCompile error and cook warning

diff --git a/Engine/Shaders/Private/HeterogeneousVolumes/HeterogeneousVolumesVoxelGridPipeline.usf b/Engine/Shaders/Private/HeterogeneousVolumes/HeterogeneousVolumesVoxelGridPipeline.usf
--- a/Engine/Shaders/Private/HeterogeneousVolumes/HeterogeneousVolumesVoxelGridPipeline.usf
+++ b/Engine/Shaders/Private/HeterogeneousVolumes/HeterogeneousVolumesVoxelGridPipeline.usf
@@ -526,10 +526,13 @@ void DownsampleMajorantVoxelGridCS(
 	int3 OutputDimensions = int3((InputDimensions + 3) / 4);
 	uint LinearOuputIndex = GetLinearIndex(GroupId, OutputDimensions);
 
-	FMajorantData MajorantData;
-	MajorantData.Minorant = 0;
-	MajorantData.Majorant = GSMajorant[0];
-	SetMajorantData(RWOutputBuffer[LinearOuputIndex], MajorantData);
+	if (LinearThreadIndex == 0)
+	{
+		FMajorantData MajorantData;
+		MajorantData.Minorant = 0;
+		MajorantData.Majorant = GSMajorant[0];
+		SetMajorantData(RWOutputBuffer[LinearOuputIndex], MajorantData);
+	}
 }
 
 [numthreads(1, 1, 1)]
diff --git a/Engine/Source/Runtime/Renderer/Private/HeterogeneousVolumes/HeterogeneousVolumesVoxelGridPipeline.cpp b/Engine/Source/Runtime/Renderer/Private/HeterogeneousVolumes/HeterogeneousVolumesVoxelGridPipeline.cpp
--- a/Engine/Source/Runtime/Renderer/Private/HeterogeneousVolumes/HeterogeneousVolumesVoxelGridPipeline.cpp
+++ b/Engine/Source/Runtime/Renderer/Private/HeterogeneousVolumes/HeterogeneousVolumesVoxelGridPipeline.cpp
@@ -918,7 +918,7 @@ class FDownsampleMajorantVoxelGridCS : public FGlobalShader
 
 	BEGIN_SHADER_PARAMETER_STRUCT(FParameters, )
 		SHADER_PARAMETER(FIntVector, InputDimensions)
-		SHADER_PARAMETER_RDG_BUFFER_SRV(RWStructuredBuffer<FScalarGridData>, InputBuffer)
+		SHADER_PARAMETER_RDG_BUFFER_SRV(StructuredBuffer<FScalarGridData>, InputBuffer)
 		SHADER_PARAMETER_RDG_BUFFER_UAV(RWStructuredBuffer<FScalarGridData>, RWOutputBuffer)
 	END_SHADER_PARAMETER_STRUCT()
 
@@ -953,7 +953,7 @@ class FCopyMaxIntoMajorantVoxelGridCS : public FGlobalShader
 	using FPermutationDomain = TShaderPermutationDomain<>;
 
 	BEGIN_SHADER_PARAMETER_STRUCT(FParameters, )
-		SHADER_PARAMETER_RDG_BUFFER_SRV(RWStructuredBuffer<FScalarGridData>, InputBuffer)
+		SHADER_PARAMETER_RDG_BUFFER_SRV(StructuredBuffer<FScalarGridData>, InputBuffer)
 		SHADER_PARAMETER_RDG_BUFFER_UAV(RWStructuredBuffer<FScalarGridData>, RWMajorantVoxelGridBuffer)
 	END_SHADER_PARAMETER_STRUCT()
 
