Remove unused standalone SSS profile texture reference.

diff --git a/Engine/Plugins/Experimental/GPULightmass/Source/GPULightmass/Private/LightmapRayTracing.h b/Engine/Plugins/Experimental/GPULightmass/Source/GPULightmass/Private/LightmapRayTracing.h
--- a/Engine/Plugins/Experimental/GPULightmass/Source/GPULightmass/Private/LightmapRayTracing.h
+++ b/Engine/Plugins/Experimental/GPULightmass/Source/GPULightmass/Private/LightmapRayTracing.h
@@ -100,9 +100,6 @@ class FLightmapPathTracingRGS : public FGlobalShader
 		SHADER_PARAMETER_STRUCT_REF(FViewUniformShaderParameters, ViewUniformBuffer)
 		SHADER_PARAMETER_RDG_UNIFORM_BUFFER(FSceneUniformParameters, Scene)
 		SHADER_PARAMETER_STRUCT_REF(FIrradianceCachingParameters, IrradianceCachingParameters)
-
-		// Subsurface data
-		SHADER_PARAMETER_TEXTURE(Texture2D, SSProfilesTexture)
 	END_SHADER_PARAMETER_STRUCT()
 };
 
@@ -166,8 +163,6 @@ class FVolumetricLightmapPathTracingRGS : public FGlobalShader
 
 		SHADER_PARAMETER_STRUCT_REF(FViewUniformShaderParameters, ViewUniformBuffer)
 		SHADER_PARAMETER_STRUCT_REF(FIrradianceCachingParameters, IrradianceCachingParameters)
-		// Subsurface data
-		SHADER_PARAMETER_TEXTURE(Texture2D, SSProfilesTexture)
 
 		SHADER_PARAMETER_RDG_BUFFER_SRV(StructuredBuffer<FLightShaderConstants>, LightShaderParametersArray)
 	END_SHADER_PARAMETER_STRUCT()
diff --git a/Engine/Plugins/Experimental/GPULightmass/Source/GPULightmass/Private/LightmapRenderer.cpp b/Engine/Plugins/Experimental/GPULightmass/Source/GPULightmass/Private/LightmapRenderer.cpp
--- a/Engine/Plugins/Experimental/GPULightmass/Source/GPULightmass/Private/LightmapRenderer.cpp
+++ b/Engine/Plugins/Experimental/GPULightmass/Source/GPULightmass/Private/LightmapRenderer.cpp
@@ -2432,9 +2432,6 @@ void FLightmapRenderer::Finalize(FRDGBuilder& GraphBuilder)
 										PassParameters->SkylightMipCount = PreviousPassParameters[GPUIndex]->SkylightMipCount;
 									}
 
-									PassParameters->SSProfilesTexture = GetSubsurfaceProfileTexture();
-
-
 									FLightmapPathTracingRGS::FPermutationDomain PermutationVector;
 									PermutationVector.Set<FLightmapPathTracingRGS::FUseFirstBounceRayGuiding>(Scene->Settings->bUseIrradianceCaching && Scene->Settings->bUseFirstBounceRayGuiding);
 									PermutationVector.Set<FLightmapPathTracingRGS::FUseIrradianceCaching>(Scene->Settings->bUseIrradianceCaching);
diff --git a/Engine/Plugins/Experimental/GPULightmass/Source/GPULightmass/Private/VolumetricLightmap.cpp b/Engine/Plugins/Experimental/GPULightmass/Source/GPULightmass/Private/VolumetricLightmap.cpp
--- a/Engine/Plugins/Experimental/GPULightmass/Source/GPULightmass/Private/VolumetricLightmap.cpp
+++ b/Engine/Plugins/Experimental/GPULightmass/Source/GPULightmass/Private/VolumetricLightmap.cpp
@@ -814,8 +814,6 @@ void FVolumetricLightmapRenderer::BackgroundTick()
 				PassParameters->SkylightMipCount = PreviousPassParameters->SkylightMipCount;
 			}
 
-			PassParameters->SSProfilesTexture = GetSubsurfaceProfileTexture();
-
 			TArray<FLightShaderConstants> OptionalStationaryDirectionalLightShadowing;
 			for (FDirectionalLightRenderState& DirectionalLight : Scene->LightSceneRenderState.DirectionalLights.Elements)
 			{
diff --git a/Engine/Source/Runtime/Renderer/Private/DeferredShadingRenderer.cpp b/Engine/Source/Runtime/Renderer/Private/DeferredShadingRenderer.cpp
--- a/Engine/Source/Runtime/Renderer/Private/DeferredShadingRenderer.cpp
+++ b/Engine/Source/Runtime/Renderer/Private/DeferredShadingRenderer.cpp
@@ -1970,13 +1970,6 @@ bool FDeferredShadingSceneRenderer::SetupRayTracingPipelineStates(FRDGBuilder& G
 	}
 
 	// Initialize common resources used for lighting in ray tracing effects
-
-	ReferenceView.RayTracingSubSurfaceProfileTexture = GetSubsurfaceProfileTextureWithFallback();
-	ReferenceView.RayTracingSubSurfaceProfileSRV = RHICreateShaderResourceView(ReferenceView.RayTracingSubSurfaceProfileTexture, 0);
-
-	ReferenceView.RayTracingSpecularProfileTexture = SpecularProfileAtlas::GetSpecularProfileTextureAtlasWithFallback();
-	ReferenceView.RayTracingSpecularProfileSRV = RHICreateShaderResourceView(ReferenceView.RayTracingSpecularProfileTexture, 0);
-
 	for (int32 ViewIndex = 0; ViewIndex < AllFamilyViews.Num(); ++ViewIndex)
 	{
 		// TODO:  It would make more sense for common ray tracing resources to be in a shared structure, rather than copied into each FViewInfo.
@@ -1987,11 +1980,7 @@ bool FDeferredShadingSceneRenderer::SetupRayTracingPipelineStates(FRDGBuilder& G
 		// Send common ray tracing resources from reference view to all others.
 		if (View->bHasAnyRayTracingPass && View != &ReferenceView)
 		{
-			View->RayTracingSubSurfaceProfileTexture = ReferenceView.RayTracingSubSurfaceProfileTexture;
-			View->RayTracingSubSurfaceProfileSRV = ReferenceView.RayTracingSubSurfaceProfileSRV;
 			View->RayTracingMaterialPipeline = ReferenceView.RayTracingMaterialPipeline;
-			View->RayTracingSpecularProfileTexture = ReferenceView.RayTracingSpecularProfileTexture;
-			View->RayTracingSpecularProfileSRV = ReferenceView.RayTracingSpecularProfileSRV;
 		}
 	}
 
@@ -2228,19 +2217,6 @@ static void ReleaseRaytracingResources(FRDGBuilder& GraphBuilder, TArrayView<FVi
 				RayTracingScene.bUsedThisFrame = true;
 			}
 		}
-
-		// Release resources that were bound to the ray tracing scene to allow them to be immediately recycled.
-		for (int32 ViewIndex = 0; ViewIndex < Views.Num(); ++ViewIndex)
-		{
-			FViewInfo& View = Views[ViewIndex];
-
-			// Release common lighting resources -- these are ref counted, so they won't be released until after the last view
-			// is finished using them (where multiple view families are rendered).
-			View.RayTracingSubSurfaceProfileSRV.SafeRelease();
-			View.RayTracingSubSurfaceProfileTexture.SafeRelease();
-			View.RayTracingSpecularProfileSRV.SafeRelease();
-			View.RayTracingSpecularProfileTexture.SafeRelease();
-		}
 	});
 }
 
diff --git a/Engine/Source/Runtime/Renderer/Private/RayTracing/RayTracingShadows.cpp b/Engine/Source/Runtime/Renderer/Private/RayTracing/RayTracingShadows.cpp
--- a/Engine/Source/Runtime/Renderer/Private/RayTracing/RayTracingShadows.cpp
+++ b/Engine/Source/Runtime/Renderer/Private/RayTracing/RayTracingShadows.cpp
@@ -195,7 +195,6 @@ class FOcclusionRGS : public FGlobalShader
 		SHADER_PARAMETER_STRUCT_INCLUDE(FSceneLightingChannelParameters, SceneLightingChannels)
 
 		SHADER_PARAMETER_RDG_TEXTURE(Texture2D, HairLightChannelMaskTexture)
-		SHADER_PARAMETER_TEXTURE(Texture2D, SSProfilesTexture)
 		SHADER_PARAMETER_RDG_BUFFER_SRV(RaytracingAccelerationStructure, TLAS)
 		SHADER_PARAMETER_RDG_TEXTURE_UAV(RWTexture2D<float4>, RWOcclusionMaskUAV)
 		SHADER_PARAMETER_RDG_TEXTURE_UAV(RWTexture2D<float>, RWRayDistanceUAV)
@@ -402,7 +401,6 @@ void FDeferredShadingSceneRenderer::RenderRayTracingShadows(
 		CommonPassParameters->SceneLightingChannels = GetSceneLightingChannelParameters(GraphBuilder, LightingChannelsTexture);
 		CommonPassParameters->LightScissor = ScissorRect;
 		CommonPassParameters->PixelOffset = PixelOffset;
-		CommonPassParameters->SSProfilesTexture = View.RayTracingSubSurfaceProfileTexture;
 		CommonPassParameters->bTransmissionSamplingDistanceCulling = CVarRayTracingTransmissionSamplingDistanceCulling.GetValueOnRenderThread();
 		CommonPassParameters->TransmissionSamplingTechnique = CVarRayTracingTransmissionSamplingTechnique.GetValueOnRenderThread();
 		CommonPassParameters->TransmissionMeanFreePathType = GetRayTracingTransmissionMeanFreePathType();
diff --git a/Engine/Source/Runtime/Renderer/Private/SceneRendering.h b/Engine/Source/Runtime/Renderer/Private/SceneRendering.h
--- a/Engine/Source/Runtime/Renderer/Private/SceneRendering.h
+++ b/Engine/Source/Runtime/Renderer/Private/SceneRendering.h
@@ -1629,13 +1629,6 @@ public:
 	FGraphEventRef									RayTracingCallableBindingsTask;
 
 	// Common resources used for lighting in ray tracing effects
-	TRefCountPtr<FRHITexture>						RayTracingSubSurfaceProfileTexture;
-	FShaderResourceViewRHIRef						RayTracingSubSurfaceProfileSRV;
-
-	// Common resources used for lighting in ray tracing effects
-	TRefCountPtr<FRHITexture>						RayTracingSpecularProfileTexture;
-	FShaderResourceViewRHIRef						RayTracingSpecularProfileSRV;
-
 	TRDGUniformBufferRef<FRaytracingLightDataPacked>	RayTracingLightDataUniformBuffer;
 	TRDGUniformBufferRef<FRayTracingDecals>				RayTracingDecalUniformBuffer;
 	bool												bHasRayTracingDecals = false;
