Fixed ensures in GL RHI related to skylight captures and texture views

diff --git a/Engine/Source/Runtime/Renderer/Private/MobileReflectionEnvironmentCapture.cpp b/Engine/Source/Runtime/Renderer/Private/MobileReflectionEnvironmentCapture.cpp
--- a/Engine/Source/Runtime/Renderer/Private/MobileReflectionEnvironmentCapture.cpp
+++ b/Engine/Source/Runtime/Renderer/Private/MobileReflectionEnvironmentCapture.cpp
@@ -53,10 +53,16 @@ namespace MobileReflectionEnvironmentCapture
 			const int32 MipSize = 1 << (NumMips - MipIndex - 1);
 			const FIntRect ViewRect(0, 0, MipSize, MipSize);
 
+			FRDGTextureSRVDesc SourceSRVDesc = FRDGTextureSRVDesc::CreateForMipLevel(CubemapTexture, MipIndex - 1);
+			if (GRHISupportsTextureViews == false)
+			{
+				SourceSRVDesc = FRDGTextureSRVDesc::Create(CubemapTexture);
+			}
+
 			for (int32 CubeFace = 0; CubeFace < CubeFace_MAX; CubeFace++)
 			{
 				auto* PassParameters = GraphBuilder.AllocParameters<FMobileCubeDownsamplePS::FParameters>();
-				PassParameters->SourceCubemapTexture = GraphBuilder.CreateSRV(FRDGTextureSRVDesc::CreateForMipLevel(CubemapTexture, MipIndex - 1));
+				PassParameters->SourceCubemapTexture = GraphBuilder.CreateSRV(SourceSRVDesc);
 				PassParameters->SourceCubemapSampler = TStaticSamplerState<SF_Bilinear, AM_Clamp, AM_Clamp, AM_Clamp>::GetRHI();
 				PassParameters->CubeFace = CubeFace;
 				PassParameters->SourceMipIndex = SourceMipIndex;
diff --git a/Engine/Source/Runtime/Renderer/Private/ReflectionEnvironmentDiffuseIrradiance.cpp b/Engine/Source/Runtime/Renderer/Private/ReflectionEnvironmentDiffuseIrradiance.cpp
--- a/Engine/Source/Runtime/Renderer/Private/ReflectionEnvironmentDiffuseIrradiance.cpp
+++ b/Engine/Source/Runtime/Renderer/Private/ReflectionEnvironmentDiffuseIrradiance.cpp
@@ -164,10 +164,16 @@ void ComputeDiffuseIrradiance(FRDGBuilder& GraphBuilder, FGlobalShaderMap* Shade
 				const FVector4f Sample01(-HalfSourceTexelSize, -HalfSourceTexelSize, HalfSourceTexelSize, -HalfSourceTexelSize);
 				const FVector4f Sample23(-HalfSourceTexelSize, HalfSourceTexelSize, HalfSourceTexelSize, HalfSourceTexelSize);
 
+				FRDGTextureSRVDesc SourceSRVDesc = FRDGTextureSRVDesc::CreateForMipLevel(IrradianceCubemapTexture, MipIndex - 1); 
+				if (GRHISupportsTextureViews == false)
+				{
+					SourceSRVDesc = FRDGTextureSRVDesc::Create(IrradianceCubemapTexture);
+				}
+
 				for (int32 CubeFace = 0; CubeFace < CubeFace_MAX; CubeFace++)
 				{
 					auto* PassParameters = GraphBuilder.AllocParameters<FAccumulateDiffuseIrradiancePS::FParameters>();
-					PassParameters->SourceCubemapTexture = GraphBuilder.CreateSRV(FRDGTextureSRVDesc::CreateForMipLevel(IrradianceCubemapTexture, MipIndex - 1));
+					PassParameters->SourceCubemapTexture = GraphBuilder.CreateSRV(SourceSRVDesc);
 					PassParameters->SourceCubemapSampler = TStaticSamplerState<SF_Point, AM_Clamp, AM_Clamp, AM_Clamp>::GetRHI();
 					PassParameters->Sample01 = Sample01;
 					PassParameters->Sample23 = Sample23;
