Unified all strata tile list data buffer into a single one to workaround UAV limitation of dx11.

diff --git a/Engine/Shaders/Private/Strata/StrataDbuffer.usf b/Engine/Shaders/Private/Strata/StrataDbuffer.usf
--- a/Engine/Shaders/Private/Strata/StrataDbuffer.usf
+++ b/Engine/Shaders/Private/Strata/StrataDbuffer.usf
@@ -79,6 +79,7 @@ uint FirstSliceStoringStrataSSSData;
 RWTexture2D<uint2> TopLayerTexture;
 RWTexture2DArray<uint> MaterialTextureArrayUAV;
 
+uint TileListBufferOffset;
 Buffer<uint> TileListBuffer;
 Buffer<uint> TileIndirectBuffer;
 
@@ -96,7 +97,7 @@ void MainCS(
 	if (View.GeneralPurposeTweak == 0)
 		return;
 
-	const uint2 TileCoord = StrataUnpackTile(TileListBuffer[GroupId.x]);
+	const uint2 TileCoord = StrataUnpackTile(TileListBuffer[TileListBufferOffset + GroupId.x]);
 	const uint2 PixelCoord = TileCoord * STRATA_TILE_SIZE + GroupThreadId;
 	const bool bIsValid = all(PixelCoord < uint2(View.ViewSizeAndInvSize.xy));
 
diff --git a/Engine/Shaders/Private/Strata/StrataMaterialClassification.usf b/Engine/Shaders/Private/Strata/StrataMaterialClassification.usf
--- a/Engine/Shaders/Private/Strata/StrataMaterialClassification.usf
+++ b/Engine/Shaders/Private/Strata/StrataMaterialClassification.usf
@@ -22,17 +22,18 @@ Texture2D<uint> TopLayerCmaskTexture;
 #endif
 RWTexture2DArray<uint> MaterialTextureArrayUAV;
 
+uint4 TileListBufferOffsets[STRATA_TILE_TYPE_COUNT];
+uint GetTileListBufferOffsets(uint Type)
+{
+	return TileListBufferOffsets[Type].x;
+}
+
 // Indirect draw data buffer for all tile types
-RWBuffer<uint> TileDrawIndirectDataBuffer;
+RWBuffer<uint> TileDrawIndirectDataBufferUAV;
 
-RWBuffer<uint> SimpleTileListDataBuffer;
-RWBuffer<uint> SingleTileListDataBuffer;
-RWBuffer<uint> ComplexTileListDataBuffer;
+RWBuffer<uint> TileListBufferUAV;
 
 #if PERMUTATION_DECAL
-RWBuffer<uint> DecalSimpleTileListDataBuffer;
-RWBuffer<uint> DecalSingleTileListDataBuffer;
-RWBuffer<uint> DecalComplexTileListDataBuffer;
 
 Texture2D<float4> DBufferATexture;
 Texture2D<float4> DBufferBTexture;
@@ -64,8 +65,6 @@ uint GetDBufferTargetMask(uint2 PixelPos)
 #endif // PERMUTATION_DECAL
 
 #if STRATA_OPAQUE_ROUGH_REFRACTION_ENABLED
-RWBuffer<uint> OpaqueRoughRefractionTileListDataBuffer;
-RWBuffer<uint> OpaqueRoughRefractionSSSWithoutTileListDataBuffer;
 Texture2D<float3> OpaqueRoughRefractionTexture;
 #endif // STRATA_OPAQUE_ROUGH_REFRACTION_ENABLED
 
@@ -84,7 +83,7 @@ void TileMainCS(uint2 DispatchThreadId : SV_DispatchThreadID, uint LinearIndex :
 	{
 		const uint TileType = DispatchThreadId.x;
 		const uint IndexCountPerInstance = bRectPrimitive > 0 ? 4 : 6;
-		TileDrawIndirectDataBuffer[GetStrataTileTypeDrawIndirectArgOffset_DWord(TileType) + 0] = IndexCountPerInstance;
+		TileDrawIndirectDataBufferUAV[GetStrataTileTypeDrawIndirectArgOffset_DWord(TileType) + 0] = IndexCountPerInstance;
 	}
 
 	const uint2 PixelCoord = DispatchThreadId.xy + View.ViewRectMin.xy;
@@ -274,34 +273,34 @@ void TileMainCS(uint2 DispatchThreadId : SV_DispatchThreadID, uint LinearIndex :
 		if (bTileContainsComplex)
 		{
 			uint WriteToIndex;
-			InterlockedAdd(TileDrawIndirectDataBuffer[GetStrataTileTypeDrawIndirectArgOffset_DWord(STRATA_TILE_TYPE_COMPLEX) + 1], 1, WriteToIndex);
-			ComplexTileListDataBuffer[WriteToIndex] = EncodedTile;
+			InterlockedAdd(TileDrawIndirectDataBufferUAV[GetStrataTileTypeDrawIndirectArgOffset_DWord(STRATA_TILE_TYPE_COMPLEX) + 1], 1, WriteToIndex);
+			TileListBufferUAV[GetTileListBufferOffsets(STRATA_TILE_TYPE_COMPLEX) + WriteToIndex] = EncodedTile;
 		}
 		else if (bTileContainsSingle)
 		{
 			uint WriteToIndex;
-			InterlockedAdd(TileDrawIndirectDataBuffer[GetStrataTileTypeDrawIndirectArgOffset_DWord(STRATA_TILE_TYPE_SINGLE) + 1], 1, WriteToIndex);
-			SingleTileListDataBuffer[WriteToIndex] = EncodedTile;
+			InterlockedAdd(TileDrawIndirectDataBufferUAV[GetStrataTileTypeDrawIndirectArgOffset_DWord(STRATA_TILE_TYPE_SINGLE) + 1], 1, WriteToIndex);
+			TileListBufferUAV[GetTileListBufferOffsets(STRATA_TILE_TYPE_SINGLE) + WriteToIndex] = EncodedTile;
 		}
 		else // (bTileContainsSimple)
 		{
 			uint WriteToIndex;
-			InterlockedAdd(TileDrawIndirectDataBuffer[GetStrataTileTypeDrawIndirectArgOffset_DWord(STRATA_TILE_TYPE_SIMPLE) + 1], 1, WriteToIndex);
-			SimpleTileListDataBuffer[WriteToIndex] = EncodedTile;
+			InterlockedAdd(TileDrawIndirectDataBufferUAV[GetStrataTileTypeDrawIndirectArgOffset_DWord(STRATA_TILE_TYPE_SIMPLE) + 1], 1, WriteToIndex);
+			TileListBufferUAV[GetTileListBufferOffsets(STRATA_TILE_TYPE_SIMPLE) + WriteToIndex] = EncodedTile;
 		}
 
 	#if STRATA_OPAQUE_ROUGH_REFRACTION_ENABLED
 		if (bTileContainsOpaqueRoughRefraction)
 		{
 			uint WriteToIndex;
-			InterlockedAdd(TileDrawIndirectDataBuffer[GetStrataTileTypeDrawIndirectArgOffset_DWord(STRATA_TILE_TYPE_ROUGH_REFRACT) + 1], 1, WriteToIndex);
-			OpaqueRoughRefractionTileListDataBuffer[WriteToIndex] = EncodedTile;
+			InterlockedAdd(TileDrawIndirectDataBufferUAV[GetStrataTileTypeDrawIndirectArgOffset_DWord(STRATA_TILE_TYPE_ROUGH_REFRACT) + 1], 1, WriteToIndex);
+			TileListBufferUAV[GetTileListBufferOffsets(STRATA_TILE_TYPE_ROUGH_REFRACT) + WriteToIndex] = EncodedTile;
 		}
 		if(bTileContainsScreenSpaceSubsurfaceScattering && !bTileContainsOpaqueRoughRefraction)
 		{
 			uint WriteToIndex;
-			InterlockedAdd(TileDrawIndirectDataBuffer[GetStrataTileTypeDrawIndirectArgOffset_DWord(STRATA_TILE_TYPE_ROUGH_REFRACT_SSS_WITHOUT) + 1], 1, WriteToIndex);
-			OpaqueRoughRefractionSSSWithoutTileListDataBuffer[WriteToIndex] = EncodedTile;
+			InterlockedAdd(TileDrawIndirectDataBufferUAV[GetStrataTileTypeDrawIndirectArgOffset_DWord(STRATA_TILE_TYPE_ROUGH_REFRACT_SSS_WITHOUT) + 1], 1, WriteToIndex);
+			TileListBufferUAV[GetTileListBufferOffsets(STRATA_TILE_TYPE_ROUGH_REFRACT_SSS_WITHOUT) + WriteToIndex] = EncodedTile;
 		}
 	#endif
 
@@ -311,20 +310,20 @@ void TileMainCS(uint2 DispatchThreadId : SV_DispatchThreadID, uint LinearIndex :
 			if (bTileContainsComplex)
 			{
 				uint WriteToIndex;
-				InterlockedAdd(TileDrawIndirectDataBuffer[GetStrataTileTypeDrawIndirectArgOffset_DWord(STRATA_TILE_TYPE_DECAL_COMPLEX) + 1], 1, WriteToIndex);
-				DecalComplexTileListDataBuffer[WriteToIndex] = EncodedTile;
+				InterlockedAdd(TileDrawIndirectDataBufferUAV[GetStrataTileTypeDrawIndirectArgOffset_DWord(STRATA_TILE_TYPE_DECAL_COMPLEX) + 1], 1, WriteToIndex);
+				TileListBufferUAV[GetTileListBufferOffsets(STRATA_TILE_TYPE_DECAL_COMPLEX) + WriteToIndex] = EncodedTile;
 			}
 			else if (bTileContainsSingle)
 			{
 				uint WriteToIndex;
-				InterlockedAdd(TileDrawIndirectDataBuffer[GetStrataTileTypeDrawIndirectArgOffset_DWord(STRATA_TILE_TYPE_DECAL_SINGLE) + 1], 1, WriteToIndex);
-				DecalSingleTileListDataBuffer[WriteToIndex] = EncodedTile;
+				InterlockedAdd(TileDrawIndirectDataBufferUAV[GetStrataTileTypeDrawIndirectArgOffset_DWord(STRATA_TILE_TYPE_DECAL_SINGLE) + 1], 1, WriteToIndex);
+				TileListBufferUAV[GetTileListBufferOffsets(STRATA_TILE_TYPE_DECAL_SINGLE) + WriteToIndex] = EncodedTile;
 			}
 			else // (bTileContainsSimple)
 			{
 				uint WriteToIndex;
-				InterlockedAdd(TileDrawIndirectDataBuffer[GetStrataTileTypeDrawIndirectArgOffset_DWord(STRATA_TILE_TYPE_DECAL_SIMPLE) + 1], 1, WriteToIndex);
-				DecalSimpleTileListDataBuffer[WriteToIndex] = EncodedTile;
+				InterlockedAdd(TileDrawIndirectDataBufferUAV[GetStrataTileTypeDrawIndirectArgOffset_DWord(STRATA_TILE_TYPE_DECAL_SIMPLE) + 1], 1, WriteToIndex);
+				TileListBufferUAV[GetTileListBufferOffsets(STRATA_TILE_TYPE_DECAL_SIMPLE) + WriteToIndex] = EncodedTile;
 			}
 		}
 	#endif
@@ -410,6 +409,7 @@ Texture2D<uint2> TopLayerTexture;
 Texture2DArray<uint> MaterialTextureArray;
 
 Buffer<uint> TileListBuffer;
+uint TileListBufferOffset;
 
 RWTexture2D<uint> RWBSDFTileTexture;
 RWTexture2D<uint> RWBSDFOffsetTexture;
@@ -430,7 +430,7 @@ WAVESIZE(64) // PERMUTATION_WAVE_OPS is true only when wave>=64 are available
 [numthreads(STRATA_TILE_SIZE, STRATA_TILE_SIZE, 1)]
 void BSDFTileMainCS(uint2 GroupThreadId : SV_GroupThreadID, uint2 GroupId : SV_GroupID, uint LinearIndex : SV_GroupIndex)
 {
-	const uint2 TileCoord = StrataUnpackTile(TileListBuffer[GroupId.x]);
+	const uint2 TileCoord = StrataUnpackTile(TileListBuffer[TileListBufferOffset + GroupId.x]);
 	const uint2 PixelCoord = TileCoord * STRATA_TILE_SIZE + GroupThreadId;
 
 	uint BSDFCount = 0;
diff --git a/Engine/Shaders/Private/Strata/StrataTile.usf b/Engine/Shaders/Private/Strata/StrataTile.usf
--- a/Engine/Shaders/Private/Strata/StrataTile.usf
+++ b/Engine/Shaders/Private/Strata/StrataTile.usf
@@ -15,6 +15,7 @@ float4 OutputBufferSizeAndInvSize;
 float4x4 ViewScreenToTranslatedWorld;
 
 Buffer<uint> TileListBuffer;
+uint		 TileListBufferOffset;
 
 void StrataTilePassVS(
 	in uint InVertexId : SV_VertexID,
@@ -29,7 +30,7 @@ void StrataTilePassVS(
 	out float4 OutPosition : SV_POSITION)	// Use by StencilMainPS
 {
 #if STRATA_ENABLED
-	const uint2 TileCoord = StrataUnpackTile(TileListBuffer[InInstanceId]);
+	const uint2 TileCoord = StrataUnpackTile(TileListBuffer[TileListBufferOffset + InInstanceId]);
 #else
 	// Not used, but this shader is compiled for non-strata path for setup reason, but not used
 	const uint2 TileCoord = 0;
diff --git a/Engine/Source/Runtime/Renderer/Private/Strata/Strata.cpp b/Engine/Source/Runtime/Renderer/Private/Strata/Strata.cpp
--- a/Engine/Source/Runtime/Renderer/Private/Strata/Strata.cpp
+++ b/Engine/Source/Runtime/Renderer/Private/Strata/Strata.cpp
@@ -18,7 +18,7 @@
 #include "SystemTextures.h"
 #include "DBufferTextures.h"
 
-//UE_DISABLE_OPTIMIZATION
+
 
 // The project setting for Strata
 static TAutoConsoleVariable<int32> CVarUseCmaskClear(
@@ -79,12 +79,10 @@ void FStrataViewData::Reset()
 	*this = FStrataViewData();
 	MaxBSDFCount = OldMaxBSDFCount;
 	MaxBytesPerPixel = OldMaxBytesPerPixel;
-	for (uint32 i = 0; i < EStrataTileType::ECount; ++i)
-	{
-		ClassificationTileListBuffer[i] = nullptr;
-		ClassificationTileListBufferUAV[i] = nullptr;
-		ClassificationTileListBufferSRV[i] = nullptr;
-	}
+
+	ClassificationTileListBuffer = nullptr;
+	ClassificationTileListBufferUAV = nullptr;
+	ClassificationTileListBufferSRV = nullptr;
 }
 
 const TCHAR* ToString(EStrataTileType Type)
@@ -231,33 +229,24 @@ static void InitialiseStrataViewData(FRDGBuilder& GraphBuilder, FViewInfo& View,
 			const uint32 DecalTileCount = IsDBufferPassEnabled(View.GetShaderPlatform()) ? TileResolution.X * TileResolution.Y : 4;
 			const uint32 RegularTileCount = TileResolution.X * TileResolution.Y;
 
-			const EPixelFormat ClassificationTileFormat = GetClassificationTileFormat(ViewResolution);
-			for (uint32 i = 0; i < EStrataTileType::ECount; ++i)
-			{
-				const EStrataTileType TileType = EStrataTileType(i);
-				const uint32 FormatBytes = ClassificationTileFormat == PF_R16_UINT ? sizeof(uint16) : sizeof(uint32);
+			Out.ClassificationTileListBufferOffset[EStrataTileType::ESimple]							= RegularTileCount * 0 + RoughTileCount * 0 + DecalTileCount * 0;
+			Out.ClassificationTileListBufferOffset[EStrataTileType::ESingle]							= RegularTileCount * 1 + RoughTileCount * 0 + DecalTileCount * 0;
+			Out.ClassificationTileListBufferOffset[EStrataTileType::EComplex]							= RegularTileCount * 2 + RoughTileCount * 0 + DecalTileCount * 0;
+			Out.ClassificationTileListBufferOffset[EStrataTileType::EOpaqueRoughRefraction]				= RegularTileCount * 3 + RoughTileCount * 0 + DecalTileCount * 0;
+			Out.ClassificationTileListBufferOffset[EStrataTileType::EOpaqueRoughRefractionSSSWithout]	= RegularTileCount * 3 + RoughTileCount * 1 + DecalTileCount * 0;
+			Out.ClassificationTileListBufferOffset[EStrataTileType::EDecalSimple]						= RegularTileCount * 3 + RoughTileCount * 2 + DecalTileCount * 0;
+			Out.ClassificationTileListBufferOffset[EStrataTileType::EDecalSingle]						= RegularTileCount * 3 + RoughTileCount * 2 + DecalTileCount * 1;
+			Out.ClassificationTileListBufferOffset[EStrataTileType::EDecalComplex]						= RegularTileCount * 3 + RoughTileCount * 2 + DecalTileCount * 2;
+			uint32 TotalTileCount										 								= RegularTileCount * 3 + RoughTileCount * 2 + DecalTileCount * 3;
 
-				uint32 TileCount = 0;
-				switch (TileType)
-				{
-				case EStrataTileType::ESimple:
-				case EStrataTileType::ESingle:
-				case EStrataTileType::EComplex:
-					TileCount = RegularTileCount; break;
-				case EStrataTileType::EOpaqueRoughRefraction:
-				case EStrataTileType::EOpaqueRoughRefractionSSSWithout:
-					TileCount = RoughTileCount; break;
-				case EStrataTileType::EDecalSimple:
-				case EStrataTileType::EDecalSingle:
-				case EStrataTileType::EDecalComplex:
-					TileCount = DecalTileCount; break;
-				}
-				check(TileCount > 0);
+			check(TotalTileCount > 0);
 
-				Out.ClassificationTileListBuffer[i] = GraphBuilder.CreateBuffer(FRDGBufferDesc::CreateBufferDesc(FormatBytes, TileCount), StrataTileListBufferNames[i]);
-				Out.ClassificationTileListBufferSRV[i] = GraphBuilder.CreateSRV(Out.ClassificationTileListBuffer[i], ClassificationTileFormat);
-				Out.ClassificationTileListBufferUAV[i] = GraphBuilder.CreateUAV(Out.ClassificationTileListBuffer[i], ClassificationTileFormat);
-			}
+			const EPixelFormat ClassificationTileFormat = GetClassificationTileFormat(ViewResolution);
+			const uint32 FormatBytes = ClassificationTileFormat == PF_R16_UINT ? sizeof(uint16) : sizeof(uint32);
+
+			Out.ClassificationTileListBuffer = GraphBuilder.CreateBuffer(FRDGBufferDesc::CreateBufferDesc(FormatBytes, TotalTileCount), TEXT("Substrate.StrataTileListBuffer"));
+			Out.ClassificationTileListBufferSRV = GraphBuilder.CreateSRV(Out.ClassificationTileListBuffer, ClassificationTileFormat);
+			Out.ClassificationTileListBufferUAV = GraphBuilder.CreateUAV(Out.ClassificationTileListBuffer, ClassificationTileFormat);
 		}
 
 		// BSDF tiles
@@ -695,6 +684,7 @@ class FStrataBSDFTilePassCS : public FGlobalShader
 		SHADER_PARAMETER_RDG_BUFFER_UAV(RWBuffer<uint>, RWBSDFTileCountBuffer)
 
 		SHADER_PARAMETER_RDG_BUFFER_SRV(Buffer<uint>, TileListBuffer)
+		SHADER_PARAMETER(uint32, TileListBufferOffset)
 		RDG_BUFFER_ACCESS(TileIndirectBuffer, ERHIAccess::IndirectArgs)
 	END_SHADER_PARAMETER_STRUCT()
 
@@ -741,17 +731,11 @@ class FStrataMaterialTileClassificationPassCS : public FGlobalShader
 		SHADER_PARAMETER(int32, bRectPrimitive)
 		SHADER_PARAMETER(FIntPoint, ViewResolution)
 		SHADER_PARAMETER(uint32, MaxBytesPerPixel)
+		SHADER_PARAMETER_ARRAY(FUintVector4, TileListBufferOffsets, [STRATA_TILE_TYPE_COUNT])
 		SHADER_PARAMETER_RDG_TEXTURE(Texture2D, TopLayerTexture)
 		SHADER_PARAMETER_RDG_TEXTURE(Texture2D, TopLayerCmaskTexture)
-		SHADER_PARAMETER_RDG_BUFFER_UAV(RWBuffer, TileDrawIndirectDataBuffer)
-		SHADER_PARAMETER_RDG_BUFFER_UAV(RWBuffer, SimpleTileListDataBuffer)
-		SHADER_PARAMETER_RDG_BUFFER_UAV(RWBuffer, SingleTileListDataBuffer)
-		SHADER_PARAMETER_RDG_BUFFER_UAV(RWBuffer, ComplexTileListDataBuffer)
-		SHADER_PARAMETER_RDG_BUFFER_UAV(RWBuffer, OpaqueRoughRefractionTileListDataBuffer)
-		SHADER_PARAMETER_RDG_BUFFER_UAV(RWBuffer, OpaqueRoughRefractionSSSWithoutTileListDataBuffer)
-		SHADER_PARAMETER_RDG_BUFFER_UAV(RWBuffer, DecalSimpleTileListDataBuffer)
-		SHADER_PARAMETER_RDG_BUFFER_UAV(RWBuffer, DecalSingleTileListDataBuffer)
-		SHADER_PARAMETER_RDG_BUFFER_UAV(RWBuffer, DecalComplexTileListDataBuffer)
+		SHADER_PARAMETER_RDG_BUFFER_UAV(RWBuffer, TileDrawIndirectDataBufferUAV)
+		SHADER_PARAMETER_RDG_BUFFER_UAV(RWBuffer, TileListBufferUAV)
 		SHADER_PARAMETER_RDG_TEXTURE_UAV(RWTexture2DArray<uint>, MaterialTextureArrayUAV)
 		SHADER_PARAMETER_RDG_TEXTURE(Texture2D<float3>, OpaqueRoughRefractionTexture)
 		SHADER_PARAMETER_STRUCT_INCLUDE(FDBufferParameters, DBuffer)
@@ -806,6 +790,7 @@ class FStrataDBufferPassCS : public FGlobalShader
 		SHADER_PARAMETER_RDG_TEXTURE_UAV(RWTexture2D<uint>, TopLayerTexture)
 		SHADER_PARAMETER_RDG_TEXTURE_UAV(RWTexture2DArray<uint>, MaterialTextureArrayUAV)
 		SHADER_PARAMETER_RDG_BUFFER_SRV(Buffer<uint>, TileListBuffer)
+		SHADER_PARAMETER(uint32, TileListBufferOffset)
 		SHADER_PARAMETER_RDG_TEXTURE_SRV(Texture2D<float>, SceneStencilTexture)
 		RDG_BUFFER_ACCESS(TileIndirectBuffer, ERHIAccess::IndirectArgs)
 	END_SHADER_PARAMETER_STRUCT()
@@ -941,7 +926,8 @@ static FStrataTileParameter InternalSetTileParameters(FRDGBuilder* GraphBuilder,
 	FStrataTileParameter Out;
 	if (TileType != EStrataTileType::ECount)
 	{
-		Out.TileListBuffer = View.StrataViewData.ClassificationTileListBufferSRV[TileType];
+		Out.TileListBuffer = View.StrataViewData.ClassificationTileListBufferSRV;
+		Out.TileListBufferOffset = View.StrataViewData.ClassificationTileListBufferOffset[TileType];
 		Out.TileIndirectBuffer = View.StrataViewData.ClassificationTileDrawIndirectBuffer;
 	}
 	else if (GraphBuilder)
@@ -949,6 +935,7 @@ static FStrataTileParameter InternalSetTileParameters(FRDGBuilder* GraphBuilder,
 		FRDGBufferRef BufferDummy = GSystemTextures.GetDefaultBuffer(*GraphBuilder, 4, 0u);
 		FRDGBufferSRVRef BufferDummySRV = GraphBuilder->CreateSRV(BufferDummy, PF_R32_UINT);
 		Out.TileListBuffer = BufferDummySRV;
+		Out.TileListBufferOffset = 0;
 		Out.TileIndirectBuffer = BufferDummy;
 	}
 	return Out;
@@ -968,6 +955,7 @@ FStrataTilePassVS::FParameters SetTileParameters(
 	Out.OutputBufferSizeAndInvSize = View.CachedViewUniformShaderParameters->BufferSizeAndInvSize;
 	Out.ViewScreenToTranslatedWorld = View.CachedViewUniformShaderParameters->ScreenToTranslatedWorld;
 	Out.TileListBuffer = Temp.TileListBuffer;
+	Out.TileListBufferOffset = Temp.TileListBufferOffset;
 	Out.TileIndirectBuffer = Temp.TileIndirectBuffer;
 	return Out;
 }
@@ -987,6 +975,7 @@ FStrataTilePassVS::FParameters SetTileParameters(
 	Out.OutputBufferSizeAndInvSize = View.CachedViewUniformShaderParameters->BufferSizeAndInvSize;
 	Out.ViewScreenToTranslatedWorld = View.CachedViewUniformShaderParameters->ScreenToTranslatedWorld;
 	Out.TileListBuffer = Temp.TileListBuffer;
+	Out.TileListBufferOffset = Temp.TileListBufferOffset;
 	Out.TileIndirectBuffer = Temp.TileIndirectBuffer;
 	return Out;
 }
@@ -1280,20 +1269,14 @@ void AddStrataMaterialClassificationPass(FRDGBuilder& GraphBuilder, const FMinim
 			PassParameters->TopLayerCmaskTexture = TopLayerCmaskTexture;
 			PassParameters->MaterialTextureArrayUAV = StrataSceneData->MaterialTextureArrayUAV;
 			PassParameters->OpaqueRoughRefractionTexture = StrataSceneData->OpaqueRoughRefractionTexture;
-			PassParameters->TileDrawIndirectDataBuffer = StrataViewData->ClassificationTileDrawIndirectBufferUAV;
+			PassParameters->TileDrawIndirectDataBufferUAV = StrataViewData->ClassificationTileDrawIndirectBufferUAV;
 			PassParameters->DBuffer = GetDBufferParameters(GraphBuilder, DBufferTextures, Platform);
 			PassParameters->SceneStencilTexture = SceneTextures.Stencil;
-			// Regular tiles
-			PassParameters->SimpleTileListDataBuffer = StrataViewData->ClassificationTileListBufferUAV[EStrataTileType::ESimple];
-			PassParameters->SingleTileListDataBuffer = StrataViewData->ClassificationTileListBufferUAV[EStrataTileType::ESingle];
-			PassParameters->ComplexTileListDataBuffer = StrataViewData->ClassificationTileListBufferUAV[EStrataTileType::EComplex];
-			// Decal tiles
-			PassParameters->DecalSimpleTileListDataBuffer = StrataViewData->ClassificationTileListBufferUAV[EStrataTileType::EDecalSimple];
-			PassParameters->DecalSingleTileListDataBuffer = StrataViewData->ClassificationTileListBufferUAV[EStrataTileType::EDecalSingle];
-			PassParameters->DecalComplexTileListDataBuffer = StrataViewData->ClassificationTileListBufferUAV[EStrataTileType::EDecalComplex];
-			// Rough refraction tiles
-			PassParameters->OpaqueRoughRefractionTileListDataBuffer = StrataViewData->ClassificationTileListBufferUAV[EStrataTileType::EOpaqueRoughRefraction];
-			PassParameters->OpaqueRoughRefractionSSSWithoutTileListDataBuffer = StrataViewData->ClassificationTileListBufferUAV[EStrataTileType::EOpaqueRoughRefractionSSSWithout];
+			PassParameters->TileListBufferUAV = StrataViewData->ClassificationTileListBufferUAV;
+			for (uint32 TileType = 0; TileType < STRATA_TILE_TYPE_COUNT; ++TileType)
+			{
+				PassParameters->TileListBufferOffsets[TileType] = FUintVector4(StrataViewData->ClassificationTileListBufferOffset[TileType], 0, 0, 0);
+			}
 
 			const uint32 GroupSize = 8;
 			FComputeShaderUtils::AddPass(
@@ -1341,7 +1324,8 @@ void AddStrataMaterialClassificationPass(FRDGBuilder& GraphBuilder, const FMinim
 			PassParameters->MaxBytesPerPixel = StrataSceneData->MaxBytesPerPixel;
 			PassParameters->TopLayerTexture = StrataSceneData->TopLayerTexture;
 			PassParameters->MaterialTextureArray = StrataSceneData->MaterialTextureArraySRV;
-			PassParameters->TileListBuffer = StrataViewData->ClassificationTileListBufferSRV[EStrataTileType::EComplex];
+			PassParameters->TileListBuffer = StrataViewData->ClassificationTileListBufferSRV;
+			PassParameters->TileListBufferOffset = StrataViewData->ClassificationTileListBufferOffset[EStrataTileType::EComplex];
 			PassParameters->TileIndirectBuffer = StrataViewData->ClassificationTileDispatchIndirectBuffer;
 
 			PassParameters->RWBSDFOffsetTexture = GraphBuilder.CreateUAV(StrataSceneData->BSDFOffsetTexture);
@@ -1444,7 +1428,8 @@ void AddStrataDBufferPass(FRDGBuilder& GraphBuilder, const FMinimalSceneTextures
 			PassParameters->FirstSliceStoringStrataSSSData = StrataSceneData->FirstSliceStoringStrataSSSData;
 			PassParameters->SceneStencilTexture = SceneTextures.Stencil;
 
-			PassParameters->TileListBuffer = StrataViewData->ClassificationTileListBufferSRV[TileType];
+			PassParameters->TileListBuffer = StrataViewData->ClassificationTileListBufferSRV;
+			PassParameters->TileListBufferOffset = StrataViewData->ClassificationTileListBufferOffset[TileType];
 			PassParameters->TileIndirectBuffer = StrataViewData->ClassificationTileDispatchIndirectBuffer;
 
 			// Dispatch with tile data
diff --git a/Engine/Source/Runtime/Renderer/Private/Strata/Strata.h b/Engine/Source/Runtime/Renderer/Private/Strata/Strata.h
--- a/Engine/Source/Runtime/Renderer/Private/Strata/Strata.h
+++ b/Engine/Source/Runtime/Renderer/Private/Strata/Strata.h
@@ -49,10 +49,11 @@ END_SHADER_PARAMETER_STRUCT()
 
 BEGIN_SHADER_PARAMETER_STRUCT(FStrataTileParameter, )
 	SHADER_PARAMETER_RDG_BUFFER_SRV(Buffer<uint>, TileListBuffer)
+	SHADER_PARAMETER(uint32, TileListBufferOffset)
 	RDG_BUFFER_ACCESS(TileIndirectBuffer, ERHIAccess::IndirectArgs)
 END_SHADER_PARAMETER_STRUCT()
 
-// This paramater struct is declared with RENDERER_API even though it is not public. This is
+// This parameter struct is declared with RENDERER_API even though it is not public. This is
 // to workaround other modules doing 'private include' of the Renderer module
 BEGIN_GLOBAL_SHADER_PARAMETER_STRUCT(FStrataGlobalUniformParameters, RENDERER_API)
 	SHADER_PARAMETER_STRUCT_INCLUDE(FStrataCommonParameters, Common)
@@ -138,9 +139,10 @@ struct FStrataViewData
 	FIntPoint OverflowTileCount = FIntPoint(0, 0);
 	FIntPoint OverflowTileOffset = FIntPoint(0, 0);
 
-	FRDGBufferRef    ClassificationTileListBuffer[STRATA_TILE_TYPE_COUNT];
-	FRDGBufferSRVRef ClassificationTileListBufferSRV[STRATA_TILE_TYPE_COUNT];
-	FRDGBufferUAVRef ClassificationTileListBufferUAV[STRATA_TILE_TYPE_COUNT];
+	FRDGBufferRef    ClassificationTileListBuffer;
+	FRDGBufferSRVRef ClassificationTileListBufferSRV;
+	FRDGBufferUAVRef ClassificationTileListBufferUAV;
+	uint32			 ClassificationTileListBufferOffset[STRATA_TILE_TYPE_COUNT];
 
 	FRDGBufferRef    ClassificationTileDrawIndirectBuffer = nullptr;
 	FRDGBufferUAVRef ClassificationTileDrawIndirectBufferUAV = nullptr;
@@ -211,6 +213,7 @@ class FStrataTilePassVS : public FGlobalShader
 		SHADER_PARAMETER(FVector4f, OutputBufferSizeAndInvSize)
 		SHADER_PARAMETER(FMatrix44f, ViewScreenToTranslatedWorld)
 		SHADER_PARAMETER_RDG_BUFFER_SRV(Buffer<uint>, TileListBuffer)
+		SHADER_PARAMETER(uint32, TileListBufferOffset)
 		RDG_BUFFER_ACCESS(TileIndirectBuffer, ERHIAccess::IndirectArgs)
 	END_SHADER_PARAMETER_STRUCT()
 
