Remove Nanite VF from hit proxy, shadow depth, depth only and velocity shaders

diff --git a/Engine/Source/Runtime/Renderer/Private/DepthRendering.h b/Engine/Source/Runtime/Renderer/Private/DepthRendering.h
--- a/Engine/Source/Runtime/Renderer/Private/DepthRendering.h
+++ b/Engine/Source/Runtime/Renderer/Private/DepthRendering.h
@@ -93,7 +93,8 @@ public:
 		return (
 			Parameters.MaterialParameters.bIsSpecialEngineMaterial ||
 			!Parameters.MaterialParameters.bWritesEveryPixel ||
-			Parameters.MaterialParameters.bMaterialMayModifyMeshPosition);
+			Parameters.MaterialParameters.bMaterialMayModifyMeshPosition)
+			&& !Parameters.VertexFactoryType->SupportsNaniteRendering();
 	}
 
 	void GetShaderBindings(
@@ -126,7 +127,8 @@ public:
 		
 		return
 			// Compile for materials that are masked
-			(!Parameters.MaterialParameters.bWritesEveryPixel || Parameters.MaterialParameters.bHasPixelDepthOffsetConnected);
+			(!Parameters.MaterialParameters.bWritesEveryPixel || Parameters.MaterialParameters.bHasPixelDepthOffsetConnected)
+			&& !Parameters.VertexFactoryType->SupportsNaniteRendering();
 	}
 
 	FDepthOnlyPS(const ShaderMetaType::CompiledShaderInitializerType& Initializer):
diff --git a/Engine/Source/Runtime/Renderer/Private/SceneHitProxyRendering.cpp b/Engine/Source/Runtime/Renderer/Private/SceneHitProxyRendering.cpp
--- a/Engine/Source/Runtime/Renderer/Private/SceneHitProxyRendering.cpp
+++ b/Engine/Source/Runtime/Renderer/Private/SceneHitProxyRendering.cpp
@@ -119,7 +119,8 @@ public:
 			&& (Parameters.MaterialParameters.bIsSpecialEngineMaterial ||
 				!Parameters.MaterialParameters.bWritesEveryPixel ||
 				Parameters.MaterialParameters.bMaterialMayModifyMeshPosition ||
-				Parameters.MaterialParameters.bIsTwoSided);
+				Parameters.MaterialParameters.bIsTwoSided) &&
+				!Parameters.VertexFactoryType->SupportsNaniteRendering();
 	}
 
 	FHitProxyPS(const ShaderMetaType::CompiledShaderInitializerType& Initializer):
diff --git a/Engine/Source/Runtime/Renderer/Private/ShadowDepthRendering.cpp b/Engine/Source/Runtime/Renderer/Private/ShadowDepthRendering.cpp
--- a/Engine/Source/Runtime/Renderer/Private/ShadowDepthRendering.cpp
+++ b/Engine/Source/Runtime/Renderer/Private/ShadowDepthRendering.cpp
@@ -321,11 +321,12 @@ public:
 			// Masked and WPO materials need their shaders but cannot be used with a position only stream.
 			|| ((!Parameters.MaterialParameters.bWritesEveryPixelShadowPass || Parameters.MaterialParameters.bMaterialMayModifyMeshPosition) && !bUsePositionOnlyStream))
 			// Only compile one pass point light shaders for feature levels >= SM5
-				&& (ShaderMode != VertexShadowDepth_OnePassPointLight || IsFeatureLevelSupported(Platform, ERHIFeatureLevel::SM5))
+			&& (ShaderMode != VertexShadowDepth_OnePassPointLight || IsFeatureLevelSupported(Platform, ERHIFeatureLevel::SM5))
 			// Only compile position-only shaders for vertex factories that support it. (Note: this assumes that a vertex factor which supports PositionOnly, supports also PositionAndNormalOnly)
 			&& (!bUsePositionOnlyStream || Parameters.VertexFactoryType->SupportsPositionOnly())
 			// Don't render ShadowDepth for translucent unlit materials
-			&& Parameters.MaterialParameters.bShouldCastDynamicShadows;
+			&& Parameters.MaterialParameters.bShouldCastDynamicShadows
+			&& !Parameters.VertexFactoryType->SupportsNaniteRendering();
 	}
 
 	static void ModifyCompilationEnvironment(const FMaterialShaderPermutationParameters& Parameters, FShaderCompilerEnvironment& OutEnvironment)
@@ -442,7 +443,8 @@ public:
 			// This mode needs a pixel shader and WPO materials can't be overridden with default material.
 			|| (bModeRequiresPS && Parameters.MaterialParameters.bMaterialMayModifyMeshPosition))
 			// Don't render ShadowDepth for translucent unlit materials
-			&& Parameters.MaterialParameters.bShouldCastDynamicShadows;
+			&& Parameters.MaterialParameters.bShouldCastDynamicShadows
+			&& !Parameters.VertexFactoryType->SupportsNaniteRendering();
 	}
 
 	static void ModifyCompilationEnvironment(const FMaterialShaderPermutationParameters& Parameters, FShaderCompilerEnvironment& OutEnvironment)
diff --git a/Engine/Source/Runtime/Renderer/Private/VelocityRendering.cpp b/Engine/Source/Runtime/Renderer/Private/VelocityRendering.cpp
--- a/Engine/Source/Runtime/Renderer/Private/VelocityRendering.cpp
+++ b/Engine/Source/Runtime/Renderer/Private/VelocityRendering.cpp
@@ -121,7 +121,8 @@ public:
 		// The material may explicitly request that it be rendered into the translucent velocity pass.
 		const bool bIsSeparateVelocityPassRequiredByMaterial = Parameters.MaterialParameters.bIsTranslucencyWritingVelocity;
 
-		return bHasPlatformSupport && (bIsDefault || bIsSeparateVelocityPassRequired || bIsSeparateVelocityPassRequiredByMaterial);
+		const bool bIsNaniteFactory = Parameters.VertexFactoryType->SupportsNaniteRendering();
+		return bHasPlatformSupport && !bIsNaniteFactory && (bIsDefault || bIsSeparateVelocityPassRequired || bIsSeparateVelocityPassRequiredByMaterial);
 	}
 
 	FVelocityVS() = default;
