Improve estimate of required buffer size for culled non-nanite instances to not be overly conservative and add cvars to control the size.

diff --git a/Engine/Source/Runtime/Renderer/Private/VirtualShadowMaps/VirtualShadowMapArray.cpp b/Engine/Source/Runtime/Renderer/Private/VirtualShadowMaps/VirtualShadowMapArray.cpp
--- a/Engine/Source/Runtime/Renderer/Private/VirtualShadowMaps/VirtualShadowMapArray.cpp
+++ b/Engine/Source/Runtime/Renderer/Private/VirtualShadowMaps/VirtualShadowMapArray.cpp
@@ -135,6 +135,20 @@ TAutoConsoleVariable<int32> CVarCoarsePagesIncludeNonNanite(
 	ECVF_Scalability | ECVF_RenderThreadSafe
 );
 
+TAutoConsoleVariable<float> CVarNonNaniteCulledInstanceAllocationFactor(
+	TEXT("r.Shadow.Virtual.NonNanite.CulledInstanceAllocationFactor"),
+	0.33f,
+	TEXT("Scale factor multiplied by the conservative size required if all instances are emitted into every clip or mip level, setting to 1.0 is fully conservative.")
+	TEXT("The actual number cannot be known on the CPU as the culling emits an instace for each clip/mip level that is overlapped."),
+	ECVF_RenderThreadSafe
+);
+
+TAutoConsoleVariable<int32> CVarNonNaniteMaxCulledInstanceAllocationSize(
+	TEXT("r.Shadow.Virtual.NonNanite.MaxCulledInstanceAllocationSize"),
+	128 * 1024 * 1024,
+	TEXT("Maximum number of instances that may be output from the culling pass into all VSM mip/clip levels. At 12 byte per instance reference this represents a 1.5GB clamp."),
+	ECVF_RenderThreadSafe
+);
 static TAutoConsoleVariable<int32> CVarShowClipmapStats(
 	TEXT("r.Shadow.Virtual.ShowClipmapStats"),
 	-1,
@@ -2185,6 +2199,7 @@ static FCullingResult AddCullingPasses(FRDGBuilder& GraphBuilder,
 	const TConstArrayView<FVSMCullingBatchInfo> VSMCullingBatchInfos,
 	const TConstArrayView<uint32> BatchInds,
 	uint32 TotalInstances,
+	uint32 TotalViewScaledInstanceCount,
 	uint32 TotalPrimaryViews,
 	FRDGBufferRef VirtualShadowViewsRDG,
 	const FCullPerPageDrawCommandsCs::FHZBShaderParameters &HZBShaderParameters,
@@ -2199,9 +2214,14 @@ static FCullingResult AddCullingPasses(FRDGBuilder& GraphBuilder,
 
 	FRDGBufferRef TmpInstanceIdOffsetBufferRDG = CreateStructuredBuffer(GraphBuilder, TEXT("Shadow.Virtual.TmpInstanceIdOffsetBuffer"), sizeof(uint32), NumIndirectArgs, nullptr, 0);
 
-	// TODO: This is both not right, and also over conservative when running with the atomic path
 	FCullingResult CullingResult;
-	CullingResult.MaxNumInstancesPerPass = FMath::Max(1u, TotalInstances * 64u);
+	// TotalViewScaledInstanceCount is conservative since it is the number of instances needed if each instance was drawn into every possible mip-level.
+	// This is far more than we'd expect in reasonable circumstances, so we use a scale factor to reduce memory pressure from these passes.
+	const uint32 MaxCulledInstanceCount = uint32(CVarNonNaniteMaxCulledInstanceAllocationSize.GetValueOnRenderThread());
+	const uint32 ScaledInstanceCount = static_cast<uint32>(static_cast<double>(TotalViewScaledInstanceCount) * CVarNonNaniteCulledInstanceAllocationFactor.GetValueOnRenderThread());
+	ensureMsgf(ScaledInstanceCount <= MaxCulledInstanceCount, TEXT("Possible non-nanite VSM Instance culling overflow detected (esitmated required size: %d, if visual artifacts appear either increase the r.Shadow.Virtual.NonNanite.MaxCulledInstanceAllocationSize (%d) or reduce r.Shadow.Virtual.NonNanite.CulledInstanceAllocationFactor (%.2f)"), ScaledInstanceCount, MaxCulledInstanceCount, CVarNonNaniteCulledInstanceAllocationFactor.GetValueOnRenderThread());
+	CullingResult.MaxNumInstancesPerPass = FMath::Clamp(ScaledInstanceCount, 1u, MaxCulledInstanceCount);
+
 	FRDGBufferRef VisibleInstancesRdg = CreateStructuredBuffer(GraphBuilder, TEXT("Shadow.Virtual.VisibleInstances"), sizeof(FVSMVisibleInstanceCmd), CullingResult.MaxNumInstancesPerPass, nullptr, 0);
 
 	FRDGBufferRef VisibleInstanceWriteOffsetRDG = CreateStructuredBuffer(GraphBuilder, TEXT("Shadow.Virtual.VisibleInstanceWriteOffset"), sizeof(uint32), 1, nullptr, 0);
@@ -2355,20 +2375,22 @@ static FCullingResult AddCullingPasses(FRDGBuilder& GraphBuilder,
 	return CullingResult;
 }
 
-static uint32 CalcRenderViewCount(const FProjectedShadowInfo* ProjectedShadowInfo, uint32& MaxMipLevels)
+struct FVSMRenderViewCount
+{
+	uint32 NumPrimaryViews = 0u;
+	uint32 NumMipLevels = 0u;
+};
+
+FVSMRenderViewCount GetRenderViewCount(const FProjectedShadowInfo* ProjectedShadowInfo)
 {
-	uint32 NumPrimaryViews = 0;
 	if (ProjectedShadowInfo->VirtualShadowMapClipmap)
 	{
-		NumPrimaryViews += ProjectedShadowInfo->VirtualShadowMapClipmap->GetLevelCount();
-		MaxMipLevels = FMath::Max(MaxMipLevels, 1u);
+		return { uint32(ProjectedShadowInfo->VirtualShadowMapClipmap->GetLevelCount()), 1u };
 	}
 	else
 	{
-		NumPrimaryViews += ProjectedShadowInfo->bOnePassPointLightShadow ? 6u : 1u;
-		MaxMipLevels = FMath::Max(MaxMipLevels, FVirtualShadowMap::MaxMipLevels);
+		return { ProjectedShadowInfo->bOnePassPointLightShadow ? 6u : 1u, FVirtualShadowMap::MaxMipLevels };
 	}
-	return NumPrimaryViews;
 }
 
 static void AddRasterPass(
@@ -2445,22 +2467,21 @@ Nanite::FPackedViewArray* FVirtualShadowMapArray::CreateVirtualShadowMapNaniteVi
 	uint32 TotalPrimaryViews = 0;
 	uint32 MaxNumMips = 0;
 
-		// TODO: replace with persistent and incrementally updated setup
+	// TODO: replace with persistent and incrementally updated setup
 	for (FProjectedShadowInfo* ProjectedShadowInfo : Shadows)
 	{
 		if (ProjectedShadowInfo->bShouldRenderVSM)
 		{
-			uint32 NumMips = 0u;
-			uint32 LightNumPrimaryViews = CalcRenderViewCount(ProjectedShadowInfo, NumMips);
-			MaxNumMips = FMath::Max(MaxNumMips, NumMips);
+			FVSMRenderViewCount VSMRenderViewCount = GetRenderViewCount(ProjectedShadowInfo);
+			MaxNumMips = FMath::Max(MaxNumMips, VSMRenderViewCount.NumMipLevels);
 
 			if (InstanceCullingQuery)
 			{
 				// Add a shadow thingo to be culled, need to know the primary view ranges.
 				// TODO: Move with other connecting-the-dots stuff into shadow scene renderer.
-				InstanceCullingQuery->Add(TotalPrimaryViews, LightNumPrimaryViews, NumMips * LightNumPrimaryViews, GetCullingVolume(ProjectedShadowInfo));
+				InstanceCullingQuery->Add(TotalPrimaryViews, VSMRenderViewCount.NumPrimaryViews, VSMRenderViewCount.NumMipLevels * VSMRenderViewCount.NumPrimaryViews, GetCullingVolume(ProjectedShadowInfo));
 			}
-			TotalPrimaryViews += LightNumPrimaryViews;
+			TotalPrimaryViews += VSMRenderViewCount.NumPrimaryViews;
 		}
 	}
 
@@ -2627,6 +2648,8 @@ void FVirtualShadowMapArray::RenderVirtualShadowMapsNonNanite(FRDGBuilder& Graph
 	// We don't use the registered culling views (this redundancy should probably be addressed at some point), set the number to disable index range checking
 	InstanceCullingMergedContext.NumCullingViews = -1;
 	int32 TotalPreCullInstanceCount = 0;
+	// Instance count multiplied by the number of (VSM) views, gives a safe maximum number of possible output instances from culling.
+	uint32 TotalViewScaledInstanceCount = 0u;
 	for (int32 Index = 0; Index < VirtualSmMeshCommandPasses.Num(); ++Index)
 	{
 		FProjectedShadowInfo* ProjectedShadowInfo = VirtualSmMeshCommandPasses[Index];
@@ -2662,10 +2685,13 @@ void FVirtualShadowMapArray::RenderVirtualShadowMapsNonNanite(FRDGBuilder& Graph
 
 			if (InstanceCullingContext->HasCullingCommands())
 			{
-				uint32 NumPrimaryViews = CalcRenderViewCount(ProjectedShadowInfo, MaxNumMips);
+				FVSMRenderViewCount VSMRenderViewCount = GetRenderViewCount(ProjectedShadowInfo);
+				MaxNumMips = FMath::Max(MaxNumMips, VSMRenderViewCount.NumMipLevels);
+
+				TotalViewScaledInstanceCount += InstanceCullingContext->TotalInstances * VSMRenderViewCount.NumPrimaryViews * VSMRenderViewCount.NumMipLevels;
 
-				VSMCullingBatchInfo.NumPrimaryViews = NumPrimaryViews;
-				TotalPrimaryViews += NumPrimaryViews;
+				VSMCullingBatchInfo.NumPrimaryViews = VSMRenderViewCount.NumPrimaryViews;
+				TotalPrimaryViews += VSMRenderViewCount.NumPrimaryViews;
 				ShadowsToAddRenderViews.Add(ProjectedShadowInfo);
 
 				if (CVarDoNonNaniteBatching.GetValueOnRenderThread() != 0)
@@ -2784,6 +2810,7 @@ void FVirtualShadowMapArray::RenderVirtualShadowMapsNonNanite(FRDGBuilder& Graph
 			VSMCullingBatchInfos,
 			InstanceCullingMergedContext.BatchInds[uint32(EBatchProcessingMode::Generic)],
 			InstanceCullingMergedContext.TotalInstances,
+			TotalViewScaledInstanceCount,
 			TotalPrimaryViews,
 			VirtualShadowViewsRDG,
 			HZBShaderParameters,
@@ -2901,6 +2928,9 @@ void FVirtualShadowMapArray::RenderVirtualShadowMapsNonNanite(FRDGBuilder& Graph
 			FSceneRenderer::GetLightNameForDrawEvent(ProjectedShadowInfo->GetLightSceneInfo().Proxy, LightNameWithLevel);
 			RDG_EVENT_SCOPE(GraphBuilder, "%s", *LightNameWithLevel);
 
+			FVSMRenderViewCount VSMRenderViewCount = GetRenderViewCount(ProjectedShadowInfo);
+			uint32 ViewScaledInstanceCount = VSMRenderViewCount.NumPrimaryViews * VSMRenderViewCount.NumMipLevels * InstanceCullingContext->TotalInstances;
+
 			CullingBatchInfo.DynamicInstanceIdOffset = ShadowDepthView->DynamicPrimitiveCollector.GetInstanceSceneDataOffset();
 			CullingBatchInfo.DynamicInstanceIdMax = CullingBatchInfo.DynamicInstanceIdOffset + ShadowDepthView->DynamicPrimitiveCollector.NumInstances();
 
@@ -2914,6 +2944,7 @@ void FVirtualShadowMapArray::RenderVirtualShadowMapsNonNanite(FRDGBuilder& Graph
 				MakeArrayView(&VSMCullingBatchInfo, 1),
 				MakeArrayView<const uint32>(nullptr, 0),
 				InstanceCullingContext->TotalInstances,
+				ViewScaledInstanceCount,
 				TotalPrimaryViews,
 				VirtualShadowViewsRDG,
 				HZBShaderParameters,
