Substrate -  allocation of tile tile list buffers is now indepenent from the dynamic resolution. It also conservatively adapts to the appearance of ComplexSpecial render path.

diff --git a/Engine/Source/Runtime/Renderer/Private/Strata/Strata.cpp b/Engine/Source/Runtime/Renderer/Private/Strata/Strata.cpp
--- a/Engine/Source/Runtime/Renderer/Private/Strata/Strata.cpp
+++ b/Engine/Source/Runtime/Renderer/Private/Strata/Strata.cpp
@@ -207,10 +207,11 @@ static void InitialiseStrataViewData(FRDGBuilder& GraphBuilder, FViewInfo& View,
 	Out.Reset();
 	Out.SceneData = &SceneData;
 
-	const FIntPoint ViewResolution(View.ViewRect.Width(), View.ViewRect.Height());
+	// Allocate texture using scene render targets size so we do not reallocate every frame when dynamic resolution is used in order to avoid resources allocation hitches.
+	const FIntPoint DynResIndependentViewSize = SceneTexturesConfig.Extent;
 	if (IsStrataEnabled())
 	{
-		const FIntPoint TileResolution(FMath::DivideAndRoundUp(ViewResolution.X, STRATA_TILE_SIZE), FMath::DivideAndRoundUp(ViewResolution.Y, STRATA_TILE_SIZE));
+		const FIntPoint TileResolution(FMath::DivideAndRoundUp(DynResIndependentViewSize.X, STRATA_TILE_SIZE), FMath::DivideAndRoundUp(DynResIndependentViewSize.Y, STRATA_TILE_SIZE));
 
 		// Tile classification buffers
 		{
@@ -229,20 +230,22 @@ static void InitialiseStrataViewData(FRDGBuilder& GraphBuilder, FViewInfo& View,
 			const uint32 DecalTileCount = IsDBufferPassEnabled(View.GetShaderPlatform()) ? TileResolution.X * TileResolution.Y : 4;
 			const uint32 RegularTileCount = TileResolution.X * TileResolution.Y;
 
-			Out.ClassificationTileListBufferOffset[EStrataTileType::ESimple]							= RegularTileCount * 0 + RoughTileCount * 0 + DecalTileCount * 0;
-			Out.ClassificationTileListBufferOffset[EStrataTileType::ESingle]							= RegularTileCount * 1 + RoughTileCount * 0 + DecalTileCount * 0;
-			Out.ClassificationTileListBufferOffset[EStrataTileType::EComplex]							= RegularTileCount * 2 + RoughTileCount * 0 + DecalTileCount * 0;
-			Out.ClassificationTileListBufferOffset[EStrataTileType::EComplexSpecial]					= RegularTileCount * 3 + RoughTileCount * 0 + DecalTileCount * 0;
-			Out.ClassificationTileListBufferOffset[EStrataTileType::EOpaqueRoughRefraction]				= RegularTileCount * 4 + RoughTileCount * 0 + DecalTileCount * 0;
-			Out.ClassificationTileListBufferOffset[EStrataTileType::EOpaqueRoughRefractionSSSWithout]	= RegularTileCount * 4 + RoughTileCount * 1 + DecalTileCount * 0;
-			Out.ClassificationTileListBufferOffset[EStrataTileType::EDecalSimple]						= RegularTileCount * 4 + RoughTileCount * 2 + DecalTileCount * 0;
-			Out.ClassificationTileListBufferOffset[EStrataTileType::EDecalSingle]						= RegularTileCount * 4 + RoughTileCount * 2 + DecalTileCount * 1;
-			Out.ClassificationTileListBufferOffset[EStrataTileType::EDecalComplex]						= RegularTileCount * 4 + RoughTileCount * 2 + DecalTileCount * 2;
-			uint32 TotalTileCount										 								= RegularTileCount * 4 + RoughTileCount * 2 + DecalTileCount * 3;
+			bool bUsesComplexSpecialRenderPath = SceneData.bUsesComplexSpecialRenderPath; // Use the Scene temporally stable bUsesComplexSpecialRenderPath to reduce buffer reallocation.
+
+			Out.ClassificationTileListBufferOffset[EStrataTileType::ESimple]							= 0;
+			Out.ClassificationTileListBufferOffset[EStrataTileType::ESingle]							= Out.ClassificationTileListBufferOffset[EStrataTileType::ESimple]							+ RegularTileCount;
+			Out.ClassificationTileListBufferOffset[EStrataTileType::EComplex]							= Out.ClassificationTileListBufferOffset[EStrataTileType::ESingle]							+ RegularTileCount;
+			Out.ClassificationTileListBufferOffset[EStrataTileType::EComplexSpecial]					= Out.ClassificationTileListBufferOffset[EStrataTileType::EComplex]							+ RegularTileCount;
+			Out.ClassificationTileListBufferOffset[EStrataTileType::EOpaqueRoughRefraction]				= Out.ClassificationTileListBufferOffset[EStrataTileType::EComplexSpecial]					+ (bUsesComplexSpecialRenderPath ? RegularTileCount : 4);
+			Out.ClassificationTileListBufferOffset[EStrataTileType::EOpaqueRoughRefractionSSSWithout]	= Out.ClassificationTileListBufferOffset[EStrataTileType::EOpaqueRoughRefraction]			+ RoughTileCount;
+			Out.ClassificationTileListBufferOffset[EStrataTileType::EDecalSimple]						= Out.ClassificationTileListBufferOffset[EStrataTileType::EOpaqueRoughRefractionSSSWithout]	+ RoughTileCount;
+			Out.ClassificationTileListBufferOffset[EStrataTileType::EDecalSingle]						= Out.ClassificationTileListBufferOffset[EStrataTileType::EDecalSimple]						+ DecalTileCount;
+			Out.ClassificationTileListBufferOffset[EStrataTileType::EDecalComplex]						= Out.ClassificationTileListBufferOffset[EStrataTileType::EDecalSingle]						+ DecalTileCount;
+			uint32 TotalTileCount										 								= Out.ClassificationTileListBufferOffset[EStrataTileType::EDecalComplex]					+ DecalTileCount;
 
 			check(TotalTileCount > 0);
 
-			const EPixelFormat ClassificationTileFormat = GetClassificationTileFormat(ViewResolution);
+			const EPixelFormat ClassificationTileFormat = GetClassificationTileFormat(DynResIndependentViewSize);
 			const uint32 FormatBytes = ClassificationTileFormat == PF_R16_UINT ? sizeof(uint16) : sizeof(uint32);
 
 			Out.ClassificationTileListBuffer = GraphBuilder.CreateBuffer(FRDGBufferDesc::CreateBufferDesc(FormatBytes, TotalTileCount), TEXT("Substrate.StrataTileListBuffer"));
@@ -252,16 +255,15 @@ static void InitialiseStrataViewData(FRDGBuilder& GraphBuilder, FViewInfo& View,
 
 		// BSDF tiles
 		if (bNeedBSDFOffets)
-		{			
-			const FIntPoint BufferSize = SceneTexturesConfig.Extent;
-			const FIntPoint BufferSize_Extended = GetStrataTextureTileResolution(View, SceneTexturesConfig.Extent, EStrataTileSpace::StrataTileSpace_Primary | EStrataTileSpace::StrataTileSpace_Overflow);
+		{
+			const FIntPoint BufferSize_Extended = GetStrataTextureTileResolution(View, DynResIndependentViewSize, EStrataTileSpace::StrataTileSpace_Primary | EStrataTileSpace::StrataTileSpace_Overflow);
 
-			const FIntPoint BaseOverflowTileOffset = FIntPoint(0, FMath::DivideAndRoundUp(BufferSize.Y, STRATA_TILE_SIZE));
+			const FIntPoint BaseOverflowTileOffset = FIntPoint(0, FMath::DivideAndRoundUp(DynResIndependentViewSize.Y, STRATA_TILE_SIZE));
 
-			Out.TileCount	= GetStrataTextureTileResolution(View, ViewResolution, EStrataTileSpace::StrataTileSpace_Primary);
+			Out.TileCount	= GetStrataTextureTileResolution(View, DynResIndependentViewSize, EStrataTileSpace::StrataTileSpace_Primary);
 			Out.TileOffset  = FIntPoint(FMath::DivideAndRoundUp(View.ViewRect.Min.X, STRATA_TILE_SIZE), FMath::DivideAndRoundUp(View.ViewRect.Min.Y, STRATA_TILE_SIZE));
 
-			Out.OverflowTileCount = GetStrataTextureTileResolution(View, ViewResolution, EStrataTileSpace::StrataTileSpace_Overflow);
+			Out.OverflowTileCount = GetStrataTextureTileResolution(View, DynResIndependentViewSize, EStrataTileSpace::StrataTileSpace_Overflow);
 			Out.OverflowTileOffset = GetStrataTextureTileResolution(View, Out.TileOffset * STRATA_TILE_SIZE, EStrataTileSpace::StrataTileSpace_Overflow) + BaseOverflowTileOffset;
 
 			Out.BSDFTileTexture = GraphBuilder.CreateTexture(FRDGTextureDesc::Create2D(BufferSize_Extended, PF_R32_UINT, FClearValueBinding::None, TexCreate_UAV | TexCreate_ShaderResource), TEXT("Substrate.BSDFTiles"));
@@ -273,7 +275,7 @@ static void InitialiseStrataViewData(FRDGBuilder& GraphBuilder, FViewInfo& View,
 		}
 		else
 		{
-			Out.TileCount = GetStrataTextureTileResolution(View, ViewResolution, EStrataTileSpace::StrataTileSpace_Primary);
+			Out.TileCount = GetStrataTextureTileResolution(View, DynResIndependentViewSize, EStrataTileSpace::StrataTileSpace_Primary);
 			Out.TileOffset = FIntPoint(0,0);
 			Out.OverflowTileCount = FIntPoint(0, 0);
 			Out.OverflowTileOffset = FIntPoint(0, 0);
@@ -332,8 +334,10 @@ void InitialiseStrataFrameSceneData(FRDGBuilder& GraphBuilder, FSceneRenderer& S
 	// Reset Substrate scene data
 	{
 		const uint32 MinBytesPerPixel = Out.MinBytesPerPixel;
+		const bool bUsesComplexSpecialRenderPath = Out.bUsesComplexSpecialRenderPath;
 		Out = FStrataSceneData();
 		Out.MinBytesPerPixel = MinBytesPerPixel;
+		Out.bUsesComplexSpecialRenderPath = bUsesComplexSpecialRenderPath;
 	}
 
 	auto UpdateMaterialBufferToTiledResolution = [](FIntPoint InBufferSizeXY, FIntPoint& OutMaterialBufferSizeXY)
@@ -373,6 +377,7 @@ void InitialiseStrataFrameSceneData(FRDGBuilder& GraphBuilder, FSceneRenderer& S
 			if (!View.bIsPlanarReflection && !View.bIsReflectionCapture && !View.bIsSceneCapture)
 			{
 				Out.MinBytesPerPixel = FMath::Max(Out.MinBytesPerPixel, View.StrataViewData.MaxBytesPerPixel);
+				Out.bUsesComplexSpecialRenderPath |= View.StrataViewData.bUsesComplexSpecialRenderPath;
 			}
 		}
 
diff --git a/Engine/Source/Runtime/Renderer/Private/Strata/Strata.h b/Engine/Source/Runtime/Renderer/Private/Strata/Strata.h
--- a/Engine/Source/Runtime/Renderer/Private/Strata/Strata.h
+++ b/Engine/Source/Runtime/Renderer/Private/Strata/Strata.h
@@ -98,6 +98,7 @@ struct FStrataSceneData
 	int32 PeelLayersAboveDepth = -1;
 	bool bRoughDiffuse = false;
 	bool bRoughnessTracking = false;
+	bool bUsesComplexSpecialRenderPath = false;
 
 	int32 SliceStoringDebugStrataTreeDataWithoutMRT = -1;
 	int32 SliceStoringDebugStrataTreeData = -1;
