Do not compute PSO precache hash if validation is disabled.

#rnx

diff --git a/Engine/Source/Runtime/Engine/Private/PSOPrecache.cpp b/Engine/Source/Runtime/Engine/Private/PSOPrecache.cpp
--- a/Engine/Source/Runtime/Engine/Private/PSOPrecache.cpp
+++ b/Engine/Source/Runtime/Engine/Private/PSOPrecache.cpp
@@ -99,7 +99,10 @@ FPSOPrecacheRequestResultArray PrecachePSOs(const TArray<FPSOPrecacheData>& PSOI
 		case FPSOPrecacheData::EType::Graphics:
 		{
 #if PSO_PRECACHING_VALIDATE
-			PSOCollectorStats::GetFullPSOPrecacheStatsCollector().AddStateToCacheByHash(RHIComputePrecachePSOHash(PrecacheData.GraphicsPSOInitializer), PrecacheData.MeshPassType, PrecacheData.VertexFactoryType);
+			if (PSOCollectorStats::IsPrecachingValidationEnabled() && PrecacheData.GraphicsPSOInitializer.bPSOPrecache)
+			{
+				PSOCollectorStats::GetFullPSOPrecacheStatsCollector().AddStateToCacheByHash(RHIComputePrecachePSOHash(PrecacheData.GraphicsPSOInitializer), PrecacheData.MeshPassType, PrecacheData.VertexFactoryType);
+			}
 #endif // PSO_PRECACHING_VALIDATE
 
 			FPSOPrecacheRequestResult PSOPrecacheResult = PipelineStateCache::PrecacheGraphicsPipelineState(PrecacheData.GraphicsPSOInitializer);
diff --git a/Engine/Source/Runtime/Renderer/Private/MeshPassProcessor.cpp b/Engine/Source/Runtime/Renderer/Private/MeshPassProcessor.cpp
--- a/Engine/Source/Runtime/Renderer/Private/MeshPassProcessor.cpp
+++ b/Engine/Source/Runtime/Renderer/Private/MeshPassProcessor.cpp
@@ -1357,14 +1357,17 @@ bool FMeshDrawCommand::SubmitDrawBegin(
 
 		EPSOPrecacheResult PSOPrecacheResult = EPSOPrecacheResult::Unknown;
 #if PSO_PRECACHING_VALIDATE
-		// Retrieve state from cache - can be be cached on the MeshPipelineState to not retrieve it anymore after the final state is known
+		if (PSOCollectorStats::IsPrecachingValidationEnabled() && GraphicsPSOInit.bPSOPrecache)
+		{
+			// Retrieve state from cache - can be be cached on the MeshPipelineState to not retrieve it anymore after the final state is known
 #if MESH_DRAW_COMMAND_DEBUG_DATA
-		PSOPrecacheResult = PSOCollectorStats::GetFullPSOPrecacheStatsCollector().CheckStateInCache(GraphicsPSOInit, RHIComputePrecachePSOHash(GraphicsPSOInit),
-			MeshDrawCommand.DebugData.MeshPassType, MeshDrawCommand.DebugData.VertexFactoryType);
+			PSOPrecacheResult = PSOCollectorStats::GetFullPSOPrecacheStatsCollector().CheckStateInCache(GraphicsPSOInit, RHIComputePrecachePSOHash(GraphicsPSOInit),
+				MeshDrawCommand.DebugData.MeshPassType, MeshDrawCommand.DebugData.VertexFactoryType);
 #else
-		PSOPrecacheResult = PSOCollectorStats::GetFullPSOPrecacheStatsCollector().CheckStateInCache(GraphicsPSOInit, RHIComputePrecachePSOHash(GraphicsPSOInit),
-			EMeshPass::Num, nullptr);
+			PSOPrecacheResult = PSOCollectorStats::GetFullPSOPrecacheStatsCollector().CheckStateInCache(GraphicsPSOInit, RHIComputePrecachePSOHash(GraphicsPSOInit),
+				EMeshPass::Num, nullptr);
 #endif // MESH_DRAW_COMMAND_DEBUG_DATA
+		}
 #endif // PSO_PRECACHING_VALIDATE
 
 		// Check if skip draw is needed when the PSO is still precaching
@@ -2292,4 +2295,4 @@ FGraphicsMinimalPipelineStateInitializer PSOCollectorStats::PatchMinimalPipeline
 	return PatchedInitializer;
 }
 
-#endif // PSO_PRECACHING_VALIDATE
\ No newline at end of file
+#endif // PSO_PRECACHING_VALIDATE
