Fixed typo that caused GPU hang from incorrect buffer aliasing.

diff --git a/Engine/Source/Runtime/Renderer/Private/VT/RuntimeVirtualTextureRender.cpp b/Engine/Source/Runtime/Renderer/Private/VT/RuntimeVirtualTextureRender.cpp
--- a/Engine/Source/Runtime/Renderer/Private/VT/RuntimeVirtualTextureRender.cpp
+++ b/Engine/Source/Runtime/Renderer/Private/VT/RuntimeVirtualTextureRender.cpp
@@ -535,7 +535,7 @@ namespace RuntimeVirtualTexture
 						CompressTexture0 = GraphBuilder.RegisterExternalTexture(Desc.OutputTargets[0], ERDGTextureFlags::None);
 						CompressTextureUAV0_64bit = GraphBuilder.CreateUAV(FRDGTextureUAVDesc(CompressTexture0, 0, Compressed64BitFormat));
 						CompressTexture1 = GraphBuilder.RegisterExternalTexture(Desc.OutputTargets[1], ERDGTextureFlags::None);
-						CompressTextureUAV1_128bit = GraphBuilder.CreateUAV(FRDGTextureUAVDesc(CompressTexture0, 0, Compressed128BitFormat));
+						CompressTextureUAV1_128bit = GraphBuilder.CreateUAV(FRDGTextureUAVDesc(CompressTexture1, 0, Compressed128BitFormat));
 						OutputAlias0 = OutputAlias1 = nullptr;
 					}
 					else
