Implements r.DOF.Gather.ResolutionDivisor used by cinematic post processing scalability

diff --git a/Engine/Config/BaseScalability.ini b/Engine/Config/BaseScalability.ini
--- a/Engine/Config/BaseScalability.ini
+++ b/Engine/Config/BaseScalability.ini
@@ -415,6 +415,7 @@ r.Filter.SizeScale=0.7
 r.Tonemapper.Quality=2
 
 ; DOF settings.
+r.DOF.Gather.ResolutionDivisor=2         ; lower gathering resolution
 r.DOF.Gather.AccumulatorQuality=0        ; lower gathering accumulator quality
 r.DOF.Gather.PostfilterMethod=2          ; Max3x3 postfilering method
 r.DOF.Gather.EnableBokehSettings=0       ; no bokeh simulation when gathering
@@ -447,6 +448,7 @@ r.Filter.SizeScale=0.8
 r.Tonemapper.Quality=5
 
 ; DOF settings.
+r.DOF.Gather.ResolutionDivisor=2         ; lower gathering resolution
 r.DOF.Gather.AccumulatorQuality=0        ; lower gathering accumulator quality
 r.DOF.Gather.PostfilterMethod=2          ; Max3x3 postfilering method
 r.DOF.Gather.EnableBokehSettings=0       ; no bokeh simulation when gathering
@@ -481,6 +483,7 @@ r.Filter.SizeScale=1
 r.Tonemapper.Quality=5
 
 ; DOF settings.
+r.DOF.Gather.ResolutionDivisor=2         ; lower gathering resolution
 r.DOF.Gather.AccumulatorQuality=1        ; higher gathering accumulator quality
 r.DOF.Gather.PostfilterMethod=1          ; Median3x3 postfilering method
 r.DOF.Gather.EnableBokehSettings=0       ; no bokeh simulation when gathering
@@ -518,6 +521,7 @@ r.Filter.SizeScale=1
 r.Tonemapper.Quality=5
 
 ; DOF settings:
+r.DOF.Gather.ResolutionDivisor=1         ; do the gathering at full resolution
 r.DOF.Gather.AccumulatorQuality=1        ; higher gathering accumulator quality
 r.DOF.Gather.PostfilterMethod=1          ; Median3x3 postfilering method
 r.DOF.Gather.EnableBokehSettings=1       ; bokeh simulation when gathering
diff --git a/Engine/Shaders/Private/DiaphragmDOF/DOFCommon.ush b/Engine/Shaders/Private/DiaphragmDOF/DOFCommon.ush
--- a/Engine/Shaders/Private/DiaphragmDOF/DOFCommon.ush
+++ b/Engine/Shaders/Private/DiaphragmDOF/DOFCommon.ush
@@ -7,6 +7,7 @@
 #pragma once
 
 #include "../Common.ush"
+#include "../ScreenPass.ush"
 
 // Piggy back on existing Circle DOF code.
 #include "../CircleDOFCommon.ush"
diff --git a/Engine/Shaders/Private/DiaphragmDOF/DOFRecombine.usf b/Engine/Shaders/Private/DiaphragmDOF/DOFRecombine.usf
--- a/Engine/Shaders/Private/DiaphragmDOF/DOFRecombine.usf
+++ b/Engine/Shaders/Private/DiaphragmDOF/DOFRecombine.usf
@@ -134,11 +134,12 @@
 float4 ViewportSize;
 uint4 ViewportRect;
 
-float2 TemporalJitterPixels;
+FScreenTransform DispatchThreadIdToDOFBufferUV;
 float2 DOFBufferUVMax;
 
 float4 SeparateTranslucencyBilinearUVMinMax;
 uint SeparateTranslucencyUpscaling;
+float EncodedCocRadiusToRecombineCocRadius;
 
 Texture2D BokehLUT;
 
@@ -594,7 +595,7 @@ void FetchAndAccumulateSamplePair(
 		// Fetch full res color and CocRadius.
 		#if CONFIG_FETCH_FULLRES_COC_FROM_ALPHA
 			Sample[k].Color = SceneColorInput.SampleLevel(GlobalPointClampedSampler, BufferUV, 0);
-			Sample[k].CocRadius = Sample[k].Color.a;
+			Sample[k].CocRadius = Sample[k].Color.a * EncodedCocRadiusToRecombineCocRadius;
 		#else
 			Sample[k].Color = SceneColorInput.SampleLevel(GlobalPointClampedSampler, BufferUV, 0);
 			Sample[k].CocRadius = SceneDepthToCocRadius(SampleWorldDepth(BufferUV));
@@ -705,7 +706,7 @@ void RecombineMainCS(
 		}
 
 		// - 0.5 * TemporalJitterPixels because DOF buffer is non temporally jittering, thanks to half res TAA pass.
-		InputParameters.DOFBufferUV = ((DispatchThreadId + 0.5) * 0.5 - 0.5 * TemporalJitterPixels) * InputParameters.DOFBufferSize.zw;
+		InputParameters.DOFBufferUV = ApplyScreenTransform(float2(DispatchThreadId), DispatchThreadIdToDOFBufferUV);
 
 		if (CONFIG_CLAMP_DOF_BUFFER_UV)
 		{
@@ -745,7 +746,7 @@ void RecombineMainCS(
 	CenterSample.Intersection = 1.0;
 
 	#if CONFIG_FETCH_FULLRES_COC_FROM_ALPHA
-		CenterSample.CocRadius = CenterSample.Color.a;
+		CenterSample.CocRadius = CenterSample.Color.a * EncodedCocRadiusToRecombineCocRadius;
 	#else
 		CenterSample.CocRadius = SceneDepthToCocRadius(SampleWorldDepth(InputParameters.SceneBufferUV));
 	#endif
@@ -789,7 +790,7 @@ void RecombineMainCS(
 				}
 
 				#if CONFIG_FETCH_FULLRES_COC_FROM_ALPHA
-					float SampleCocRadius = SceneColorInput.SampleLevel(GlobalPointClampedSampler, SampleUV, 0).a;
+					float SampleCocRadius = SceneColorInput.SampleLevel(GlobalPointClampedSampler, SampleUV, 0).a * EncodedCocRadiusToRecombineCocRadius;
 				#else
 					float SampleCocRadius = SceneDepthToCocRadius(SampleWorldDepth(SampleUV));
 				#endif
diff --git a/Engine/Shaders/Private/DiaphragmDOF/DOFSetup.usf b/Engine/Shaders/Private/DiaphragmDOF/DOFSetup.usf
--- a/Engine/Shaders/Private/DiaphragmDOF/DOFSetup.usf
+++ b/Engine/Shaders/Private/DiaphragmDOF/DOFSetup.usf
@@ -22,9 +22,6 @@
 
 uint4 ViewportRect;
 
-// x: full res basis; y: half res basis.
-float2 CocRadiusBasis;
-
 Texture2D SceneColorTexture;
 Texture2D SceneDepthTexture;
 
@@ -88,11 +85,11 @@ void SetupCS(
 	// Output full resolution.
 	#if CONFIG_DOF_ALPHA
 	{
-		Output0[uint2(InputPixelPosition)] = CocRadiusBasis.x * CocRadius;
+		Output0[uint2(InputPixelPosition)] = CocRadius;
 	}
 	#else
 	{
-		Output0[uint2(InputPixelPosition)] = float4(SceneColor.rgb, CocRadiusBasis.x * CocRadius);
+		Output0[uint2(InputPixelPosition)] = float4(SceneColor.rgb, CocRadius);
 	}
 	#endif
 
@@ -103,11 +100,11 @@ void SetupCS(
 		#if CONFIG_DOF_ALPHA
 		{
 			SharedColor[GroupThreadIndex - SUBTILE_SIZE * SUBTILE_SIZE] = SceneColor;
-			SharedCocRadius[GroupThreadIndex - SUBTILE_SIZE * SUBTILE_SIZE] = CocRadiusBasis.y * CocRadius;
+			SharedCocRadius[GroupThreadIndex - SUBTILE_SIZE * SUBTILE_SIZE] = CocRadius;
 		}
 		#else
 		{
-			SharedColorAndCocRadius[GroupThreadIndex - SUBTILE_SIZE * SUBTILE_SIZE] = float4(SceneColor.rgb, CocRadiusBasis.y * CocRadius);
+			SharedColorAndCocRadius[GroupThreadIndex - SUBTILE_SIZE * SUBTILE_SIZE] = float4(SceneColor.rgb, CocRadius);
 		}
 		#endif
 	}
@@ -123,7 +120,7 @@ void SetupCS(
 		float CocRadii[4];
 
 		SceneColors[0] = SceneColor;
-		CocRadii[0] = CocRadiusBasis.y * CocRadius;
+		CocRadii[0] = CocRadius;
 
 		UNROLL
 		for (uint i = 0; i < 3; i++)
@@ -181,7 +178,7 @@ void SetupCS(uint2 DispatchThreadId : SV_DispatchThreadID)
 	float4 OutColor = SceneColorTexture.SampleLevel(GlobalPointClampedSampler, SceneBufferUV, 0);
 	
 	float DeviceZ = SceneDepthTexture.SampleLevel(GlobalPointClampedSampler, SceneBufferUV, 0).r;
-	float OutCocRadius = CocRadiusBasis.x * SceneDepthToCocRadius(ConvertFromDeviceZ(DeviceZ));
+	float OutCocRadius = SceneDepthToCocRadius(ConvertFromDeviceZ(DeviceZ));
 
 	if (all(DispatchThreadId < ViewportRect.zw))
 	{
diff --git a/Engine/Source/Runtime/Renderer/Private/PostProcess/DiaphragmDOF.cpp b/Engine/Source/Runtime/Renderer/Private/PostProcess/DiaphragmDOF.cpp
--- a/Engine/Source/Runtime/Renderer/Private/PostProcess/DiaphragmDOF.cpp
+++ b/Engine/Source/Runtime/Renderer/Private/PostProcess/DiaphragmDOF.cpp
@@ -21,6 +21,8 @@
 #include "SceneTextureParameters.h"
 #include "TranslucentRendering.h"
 #include "DataDrivenShaderPlatformInfo.h"
+#include "ScreenPass.h"
+
 
 // ---------------------------------------------------- Cvars
 
@@ -29,6 +31,13 @@ namespace
 
 DECLARE_GPU_STAT(DepthOfField)
 
+TAutoConsoleVariable<int32> CVarDOFGatherResDivisor(
+	TEXT("r.DOF.Gather.ResolutionDivisor"), 2,
+	TEXT("Selects the resolution divisor of the gather pass.\n")
+	TEXT(" 1: Do gathering pass at full resolution;\n")
+	TEXT(" 2: Do gathering pass at half resolution (default)."),
+	ECVF_Scalability | ECVF_RenderThreadSafe);
+
 TAutoConsoleVariable<int32> CVarAccumulatorQuality(
 	TEXT("r.DOF.Gather.AccumulatorQuality"),
 	1,
@@ -1272,10 +1281,11 @@ class FDiaphragmDOFRecombineCS : public FDiaphragmDOFShader
 
 		SHADER_PARAMETER(FIntRect, ViewportRect)
 		SHADER_PARAMETER(FVector4f, ViewportSize)
-		SHADER_PARAMETER(FVector2f, TemporalJitterPixels)
+		SHADER_PARAMETER(FScreenTransform, DispatchThreadIdToDOFBufferUV)
 		SHADER_PARAMETER(FVector2f, DOFBufferUVMax)
 		SHADER_PARAMETER(FVector4f, SeparateTranslucencyBilinearUVMinMax)
 		SHADER_PARAMETER(int32, SeparateTranslucencyUpscaling)
+		SHADER_PARAMETER(float, EncodedCocRadiusToRecombineCocRadius)
 		
 		SHADER_PARAMETER_RDG_TEXTURE(Texture2D, BokehLUT)
 
@@ -1383,8 +1393,7 @@ FRDGTextureRef DiaphragmDOF::AddPasses(
 	const int32 RecombineQuality = bSupportsSlightOutOfFocus ? FMath::Clamp(CVarRecombineQuality.GetValueOnRenderThread(), 0, kMaxRecombineQuality) : 0;
 
 	// Resolution divisor.
-	// TODO: Exposes lower resolution divisor?
-	const int32 PrefilteringResolutionDivisor = 2;
+	const int32 PrefilteringResolutionDivisor = CVarDOFGatherResDivisor.GetValueOnRenderThread() >= 2 ? 2 : 1;
 
 	// Minimal absolute Coc radius to spawn a gather pass. Blurring radius under this are considered not great looking.
 	// This is assuming the pass is opacity blending with a ramp from 1 to 2. This can not be exposed as a cvar,
@@ -1456,7 +1465,9 @@ FRDGTextureRef DiaphragmDOF::AddPasses(
 		PreprocessViewSize = FIntPoint::DivideAndRoundUp(TAAParameters.OutputViewRect.Size(), PrefilteringResolutionDivisor);
 	}
 
-	const float PreProcessingToProcessingCocRadiusFactor = float(GatheringViewSize.X) / float(PreprocessViewSize.X);
+	// Basis of the coc encoded in all buffer
+	const float EncodedCocRadiusBasis = PreprocessViewSize.X;
+	const float GatheringCocRadiusBasis = float(GatheringViewSize.X);
 
 	const float MaxBackgroundCocRadius = CocModel.ComputeViewMaxBackgroundCocRadius(GatheringViewSize.X);
 	const float MinForegroundCocRadius = CocModel.ComputeViewMinForegroundCocRadius(GatheringViewSize.X);
@@ -1552,14 +1563,13 @@ FRDGTextureRef DiaphragmDOF::AddPasses(
 		FullResGatherInputTextures = CreateTextures(GraphBuilder, FullResGatherInputDescs, TEXT("DOF.FullResSetup"));
 		HalfResGatherInputTextures = CreateTextures(GraphBuilder, HalfResGatherInputDescs, TEXT("DOF.HalfResSetup"));
 
-		bool bOutputFullResolution = bRecombineDoesSlightOutOfFocus && !bProcessSceneAlpha;
-		bool bOutputHalfResolution = true; // TODO: there is a useless shader permutation.
+		bool bOutputFullResolution = (bRecombineDoesSlightOutOfFocus && !bProcessSceneAlpha) || PrefilteringResolutionDivisor == 1;
+		bool bOutputHalfResolution = PrefilteringResolutionDivisor == 2;
 		
 		FDiaphragmDOFSetupCS::FPermutationDomain PermutationVector;
 	
 		FIntPoint PassViewSize = FullResViewSize;
 		FIntPoint GroupSize(kDefaultGroupSize, kDefaultGroupSize);
-		float CocRadiusBasis = 1.0f;
 		if (bOutputFullResolution && bOutputHalfResolution)
 		{
 			PermutationVector.Set<FDiaphragmDOFSetupCS::FOutputResDivisor>(0);
@@ -1573,7 +1583,6 @@ FRDGTextureRef DiaphragmDOF::AddPasses(
 		{
 			PermutationVector.Set<FDiaphragmDOFSetupCS::FOutputResDivisor>(2);
 			PassViewSize = GatheringViewSize;
-			CocRadiusBasis = PreprocessViewSize.X;
 		}
 		else
 		{
@@ -1583,9 +1592,8 @@ FRDGTextureRef DiaphragmDOF::AddPasses(
 		FDiaphragmDOFSetupCS::FParameters* PassParameters = GraphBuilder.AllocParameters<FDiaphragmDOFSetupCS::FParameters>();
 		{
 			PassParameters->CommonParameters = CommonParameters;
-			SetCocModelParameters(&PassParameters->CocModel, CocModel, CocRadiusBasis);
+			SetCocModelParameters(&PassParameters->CocModel, CocModel, EncodedCocRadiusBasis);
 			PassParameters->ViewportRect = FIntRect(FIntPoint::ZeroValue, PassViewSize);
-			PassParameters->CocRadiusBasis = FVector2f(GatheringViewSize.X, PreprocessViewSize.X);
 			PassParameters->SceneColorTexture = InputSceneColor;
 			PassParameters->SceneDepthTexture = SceneTextures.SceneDepthTexture;
 		
@@ -1629,6 +1637,12 @@ FRDGTextureRef DiaphragmDOF::AddPasses(
 		{
 			FullResGatherInputTextures.SceneColor = InputSceneColor;
 		}
+
+		if (!bOutputHalfResolution)
+		{
+			check(PrefilteringResolutionDivisor == 1);
+			HalfResGatherInputTextures = FullResGatherInputTextures;
+		}
 	}
 	
 	// TAA the setup for the convolution to be temporally stable.
@@ -1742,7 +1756,7 @@ FRDGTextureRef DiaphragmDOF::AddPasses(
 			PassParameters->ViewportRect = FIntRect(0, 0, DilatePassViewSize.X, DilatePassViewSize.Y);
 			PassParameters->SampleOffsetMultipler = SampleOffsetMultipler;
 			PassParameters->fSampleOffsetMultipler = SampleOffsetMultipler;
-			PassParameters->CocRadiusToBucketDistanceUpperBound = PreProcessingToProcessingCocRadiusFactor * BluringRadiusErrorMultiplier;
+			PassParameters->CocRadiusToBucketDistanceUpperBound = (GatheringCocRadiusBasis / EncodedCocRadiusBasis) * BluringRadiusErrorMultiplier;
 			PassParameters->BucketDistanceToCocRadius = 1.0f / PassParameters->CocRadiusToBucketDistanceUpperBound;
 			PassParameters->TileInput = TileInput;
 			PassParameters->DilatedTileMinMax = DilatedTileMinMax;
@@ -1916,7 +1930,7 @@ FRDGTextureRef DiaphragmDOF::AddPasses(
 			PassParameters->MaxBufferUV = FVector2f(
 				(PreprocessViewSize.X - 0.5f) / float(SrcSize.X),
 				(PreprocessViewSize.Y - 0.5f) / float(SrcSize.Y));
-			PassParameters->OutputCocRadiusMultiplier = PreProcessingToProcessingCocRadiusFactor;
+			PassParameters->OutputCocRadiusMultiplier = GatheringCocRadiusBasis / EncodedCocRadiusBasis;
 
 			PassParameters->GatherInputSize = FVector4f(SrcSize.X, SrcSize.Y, 1.0f / SrcSize.X, 1.0f / SrcSize.Y);
 			PassParameters->GatherInput = HalfResGatherInputTextures;
@@ -1965,7 +1979,7 @@ FRDGTextureRef DiaphragmDOF::AddPasses(
 				(PreprocessViewSize.X - 0.5f) / SrcSize.X,
 				(PreprocessViewSize.Y - 0.5f) / SrcSize.Y);
 			PassParameters->MaxScatteringGroupCount = MaxScatteringGroupCount;
-			PassParameters->PreProcessingToProcessingCocRadiusFactor = PreProcessingToProcessingCocRadiusFactor;
+			PassParameters->PreProcessingToProcessingCocRadiusFactor = GatheringCocRadiusBasis / EncodedCocRadiusBasis;
 			PassParameters->MinScatteringCocRadius = MinScatteringCocRadius;
 			PassParameters->NeighborCompareMaxColor = CVarScatterNeighborCompareMaxColor.GetValueOnRenderThread();
 			PassParameters->CocSqueeze = CocModel.Squeeze;
@@ -2301,12 +2315,12 @@ FRDGTextureRef DiaphragmDOF::AddPasses(
 				float(SrcSize.X * GatheringViewSize.X) / float(PreprocessViewSize.X),
 				float(SrcSize.Y * GatheringViewSize.Y) / float(PreprocessViewSize.Y));
 			PassParameters->MipBias = FMath::Log2(float(PreprocessViewSize.X) / float(GatheringViewSize.X));
-			PassParameters->MaxRecombineAbsCocRadius = float(kMaxSlightOutOfFocusCocRadius) / PreProcessingToProcessingCocRadiusFactor;
+			PassParameters->MaxRecombineAbsCocRadius = float(kMaxSlightOutOfFocusCocRadius) / (GatheringCocRadiusBasis / EncodedCocRadiusBasis);
 			PassParameters->CocSqueeze = CocModel.Squeeze;
 			PassParameters->CocInvSqueeze = 1.0f / CocModel.Squeeze;
 
 			PassParameters->TileDecisionParameters.MinGatherRadius = PassParameters->MaxRecombineAbsCocRadius - 1;
-			PassParameters->TileDecisionParameters.SlightOutOfFocusRadiusBoundary = float(kMaxSlightOutOfFocusCocRadius) / PreProcessingToProcessingCocRadiusFactor;
+			PassParameters->TileDecisionParameters.SlightOutOfFocusRadiusBoundary = float(kMaxSlightOutOfFocusCocRadius) / (GatheringCocRadiusBasis / EncodedCocRadiusBasis);
 			PassParameters->CommonParameters = CommonParameters;
 		
 			PassParameters->GatherInputSize = FVector4f(SrcSize.X, SrcSize.Y, 1.0f / SrcSize.X, 1.0f / SrcSize.Y);
@@ -2650,14 +2664,23 @@ FRDGTextureRef DiaphragmDOF::AddPasses(
 
 		FDiaphragmDOFRecombineCS::FParameters* PassParameters = GraphBuilder.AllocParameters<FDiaphragmDOFRecombineCS::FParameters>();
 		PassParameters->CommonParameters = CommonParameters;
-		SetCocModelParameters(&PassParameters->CocModel, CocModel, /* CocRadiusBasis = */ PassViewRect.Width() * 0.5f);
+		SetCocModelParameters(&PassParameters->CocModel, CocModel, /* CocRadiusBasis = */ float(GatheringViewSize.X));
 
 		PassParameters->ViewportRect = PassViewRect;
 		PassParameters->ViewportSize = FVector4f(PassViewRect.Width(), PassViewRect.Height(), 1.0f / PassViewRect.Width(), 1.0f / PassViewRect.Height());
-		PassParameters->TemporalJitterPixels = FVector2f(View.TemporalJitterPixels);	// LWC_TODO: Precision loss
+		PassParameters->DispatchThreadIdToDOFBufferUV = (
+			FScreenTransform::DispatchThreadIdToViewportUV(PassViewRect) *
+			FScreenTransform::ChangeTextureBasisFromTo(
+				RefBufferSize,
+				FIntRect(FIntPoint::ZeroValue, GatheringViewSize),
+				FScreenTransform::ETextureBasis::ViewportUV,
+				FScreenTransform::ETextureBasis::TextureUV)
+			- FVector2f(View.TemporalJitterPixels) * FVector2f(1.0f / float(InputSceneColor->Desc.Extent.X), 1.0f / float(InputSceneColor->Desc.Extent.Y)));
+
 		PassParameters->DOFBufferUVMax = FVector2f(
 			(GatheringViewSize.X - 0.5f) / float(RefBufferSize.X),
 			(GatheringViewSize.Y - 0.5f) / float(RefBufferSize.Y));
+		PassParameters->EncodedCocRadiusToRecombineCocRadius = float(GatheringViewSize.X) / EncodedCocRadiusBasis; 
 
 		PassParameters->SeparateTranslucencyBilinearUVMinMax.X = (SeparateTranslucencyViewport.Rect.Min.X + 0.5f) / float(SeparateTranslucencyViewport.Extent.X);
 		PassParameters->SeparateTranslucencyBilinearUVMinMax.Y = (SeparateTranslucencyViewport.Rect.Min.Y + 0.5f) / float(SeparateTranslucencyViewport.Extent.Y);
