Tonemap: Expose the sharpen parameter in the post process setting volume.

The previous cvar now defaults to -1 which means: inherit from the PPV. Settings the cvar directly will override the PPV for backwards compatibility and ease of testing.

diff --git a/Engine/Source/Runtime/Engine/Classes/Engine/Scene.h b/Engine/Source/Runtime/Engine/Classes/Engine/Scene.h
--- a/Engine/Source/Runtime/Engine/Classes/Engine/Scene.h
+++ b/Engine/Source/Runtime/Engine/Classes/Engine/Scene.h
@@ -940,6 +940,9 @@ struct FPostProcessSettings
 	UPROPERTY(EditAnywhere, BlueprintReadWrite, Category=Overrides, meta=(PinHiddenByDefault, InlineEditConditionToggle))
 	uint8 bOverride_VignetteIntensity:1;
 
+	UPROPERTY(EditAnywhere, BlueprintReadWrite, Category=Overrides, meta=(PinHiddenByDefault, InlineEditConditionToggle))
+	uint8 bOverride_Sharpen:1;
+
 	UPROPERTY()
 	uint8 bOverride_GrainIntensity_DEPRECATED:1;
 
@@ -1901,6 +1904,10 @@ struct FPostProcessSettings
 	UPROPERTY(interp, BlueprintReadWrite, Category="Lens|Image Effects", meta=(UIMin = "0.0", UIMax = "1.0", editcondition = "bOverride_VignetteIntensity"))
 	float VignetteIntensity;
 
+	/** Controls the strength of image sharpening applied during tonemapping. */
+	UPROPERTY(interp, BlueprintReadWrite, Category = "Lens|Image Effects", meta = (ClampMin = "0.0", ClampMax = "10.0", UIMin = "0.0", UIMax = "2.0", editcondition = "bOverride_Sharpen"))
+	float Sharpen;
+
 	UPROPERTY()
 	float GrainJitter_DEPRECATED;
 
diff --git a/Engine/Source/Runtime/Engine/Private/Scene.cpp b/Engine/Source/Runtime/Engine/Private/Scene.cpp
--- a/Engine/Source/Runtime/Engine/Private/Scene.cpp
+++ b/Engine/Source/Runtime/Engine/Private/Scene.cpp
@@ -522,6 +522,7 @@ FPostProcessSettings::FPostProcessSettings()
 	LensFlareBokehSize = 3.0f;
 	LensFlareThreshold = 8.0f;
 	VignetteIntensity = 0.4f;
+	Sharpen = 0.0f;
 	GrainIntensity_DEPRECATED = 0.0f;
 	GrainJitter_DEPRECATED = 0.0f;
 
diff --git a/Engine/Source/Runtime/Engine/Private/SceneView.cpp b/Engine/Source/Runtime/Engine/Private/SceneView.cpp
--- a/Engine/Source/Runtime/Engine/Private/SceneView.cpp
+++ b/Engine/Source/Runtime/Engine/Private/SceneView.cpp
@@ -1545,6 +1545,7 @@ void FSceneView::OverridePostProcessSettings(const FPostProcessSettings& Src, fl
 		LERP_PP(LensFlareBokehSize);
 		LERP_PP(LensFlareThreshold);
 		LERP_PP(VignetteIntensity);
+		LERP_PP(Sharpen);
 		LERP_PP(FilmGrainIntensity);
 		LERP_PP(FilmGrainIntensityShadows);
 		LERP_PP(FilmGrainIntensityMidtones);
diff --git a/Engine/Source/Runtime/Renderer/Private/PostProcess/PostProcessTonemap.cpp b/Engine/Source/Runtime/Renderer/Private/PostProcess/PostProcessTonemap.cpp
--- a/Engine/Source/Runtime/Renderer/Private/PostProcess/PostProcessTonemap.cpp
+++ b/Engine/Source/Runtime/Renderer/Private/PostProcess/PostProcessTonemap.cpp
@@ -33,9 +33,10 @@ namespace
 {
 TAutoConsoleVariable<float> CVarTonemapperSharpen(
 	TEXT("r.Tonemapper.Sharpen"),
-	0,
+	-1,
 	TEXT("Sharpening in the tonemapper (not for mobile), actual implementation is work in progress, clamped at 10\n")
-	TEXT("   0: off(default)\n")
+	TEXT("  <0: inherit from PostProcessVolume settings (default)\n")
+	TEXT("   0: off\n")
 	TEXT(" 0.5: half strength\n")
 	TEXT("   1: full strength"),
 	ECVF_Scalability | ECVF_RenderThreadSafe);
@@ -117,6 +118,13 @@ bool ShouldCompileCommonPermutation(const FGlobalShaderPermutationParameters& Pa
 	return true;
 }
 
+static float GetSharpenSetting(const FPostProcessSettings& Settings)
+{
+	float CVarSharpen = CVarTonemapperSharpen.GetValueOnRenderThread();
+	float Sharpen = CVarSharpen >= 0.0 ? CVarSharpen : Settings.Sharpen;
+	return FMath::Clamp(Sharpen, 0.0f, 10.0f);
+}
+
 // Common conversion of engine settings into.
 FCommonDomain BuildCommonPermutationDomain(const FViewInfo& View, bool bGammaOnly, bool bLocalExposure, bool bMetalMSAAHDRDecode)
 {
@@ -138,7 +146,7 @@ FCommonDomain BuildCommonPermutationDomain(const FViewInfo& View, bool bGammaOnl
 	PermutationVector.Set<FTonemapperBloomDim>(Settings.BloomIntensity > 0.0);
 	PermutationVector.Set<FTonemapperLocalExposureDim>(bLocalExposure);
 	PermutationVector.Set<FTonemapperFilmGrainDim>(View.FilmGrainTexture != nullptr);
-	PermutationVector.Set<FTonemapperSharpenDim>(CVarTonemapperSharpen.GetValueOnRenderThread() > 0.0f);
+	PermutationVector.Set<FTonemapperSharpenDim>(GetSharpenSetting(Settings) > 0.0f);
 	PermutationVector.Set<FTonemapperMsaaDim>(bMetalMSAAHDRDecode);
 	return PermutationVector;
 }
@@ -615,7 +623,7 @@ FScreenPassTexture AddTonemapPass(FRDGBuilder& GraphBuilder, const FViewInfo& Vi
 	FRHISamplerState* BilinearClampSampler = TStaticSamplerState<SF_Bilinear, AM_Clamp, AM_Clamp, AM_Clamp>::GetRHI();
 	FRHISamplerState* PointClampSampler = TStaticSamplerState<SF_Point, AM_Clamp, AM_Clamp, AM_Clamp>::GetRHI();
 
-	const float SharpenDiv6 = FMath::Clamp(CVarTonemapperSharpen.GetValueOnRenderThread(), 0.0f, 10.0f) / 6.0f;
+	const float SharpenDiv6 = TonemapperPermutation::GetSharpenSetting(View.FinalPostProcessSettings) / 6.0f;
 
 	FVector4f ChromaticAberrationParams;
 
diff --git a/Engine/Source/Runtime/Renderer/Private/PostProcess/VisualizeTemporalUpscaler.cpp b/Engine/Source/Runtime/Renderer/Private/PostProcess/VisualizeTemporalUpscaler.cpp
--- a/Engine/Source/Runtime/Renderer/Private/PostProcess/VisualizeTemporalUpscaler.cpp
+++ b/Engine/Source/Runtime/Renderer/Private/PostProcess/VisualizeTemporalUpscaler.cpp
@@ -236,6 +236,16 @@ FScreenPassTexture AddVisualizeTemporalUpscalerPass(FRDGBuilder& GraphBuilder, c
 				}
 				QuickDrawSummary(/* Location = */ 4, TEXT("Support Alpha: ") + Text);
 			}
+
+			// Display if any additional sharpening is happening
+			{
+				static auto CVarSharpen = IConsoleManager::Get().FindConsoleVariable(TEXT("r.Tonemapper.Sharpen"));
+				check(CVarSharpen);
+				float Sharpen = CVarSharpen->GetFloat();
+				Sharpen = (Sharpen < 0) ? View.FinalPostProcessSettings.Sharpen : Sharpen;
+				QuickDrawSummary(/* Location = */ 5, Sharpen > 0 ? FString::Printf(TEXT("Tonemapper Sharpen: %f"), Sharpen) : TEXT("Tonemapper Sharpen: Off"));
+			}
+
 		});
 
 	}
