Fix for decal being rendered a second time in the translucent pass.
This happens when the AutoBeforeDOFTranslucencyBoundary gets set. Then we use some local logic that was not taking into account that we should ignore decal materials.
When the AutoBeforeDOFTranslucencyBoundary is 0 we go through FBasePassMeshProcessor::ShouldDraw() which correctly checked for decal materials.

diff --git a/Engine/Source/Runtime/Renderer/Private/BasePassRendering.cpp b/Engine/Source/Runtime/Renderer/Private/BasePassRendering.cpp
--- a/Engine/Source/Runtime/Renderer/Private/BasePassRendering.cpp
+++ b/Engine/Source/Runtime/Renderer/Private/BasePassRendering.cpp
@@ -1907,7 +1907,7 @@ bool FBasePassMeshProcessor::TryAddMeshBatch(const FMeshBatch& RESTRICT MeshBatc
 
 
 	bool bShouldDraw = false;
-	if (AutoBeforeDOFTranslucencyBoundary > 0.0f && PrimitiveSceneProxy && bIsTranslucent)
+	if (AutoBeforeDOFTranslucencyBoundary > 0.0f && PrimitiveSceneProxy && bIsTranslucent && !Material.IsDeferredDecal())
 	{
 		check(ViewIfDynamicMeshCommand);
 		check(TranslucencyPassType != ETranslucencyPass::TPT_MAX);
