Prevent instance over-allocation from crashing the editor / engine.
 - Clamps the gpu-scene instance buffer size to the max supported instance ID (MAX_INSTANCE_ID, 16M at present)
 - Overflowing ID's will just result in OOB buffer accesses on the GPU which is tolerated and rendering will be undefined.
 - Added logging + on-screen message to surface the issue.

diff --git a/Engine/Shaders/Private/SceneData.ush b/Engine/Shaders/Private/SceneData.ush
--- a/Engine/Shaders/Private/SceneData.ush
+++ b/Engine/Shaders/Private/SceneData.ush
@@ -75,9 +75,6 @@ uint GetDefaultPrimitiveSceneDataFlags()
 		PRIMITIVE_SCENE_DATA_FLAG_VISIBLE_IN_RAY_TRACING;
 }
 
-// GPUCULL_TODO: Eventually we need to remove this workaround
-#define VF_TREAT_INSTANCE_ID_OFFSET_AS_PRIMITIVE_ID_FLAG (1U << 31U)
-
 #define INVALID_INSTANCE_PAYLOAD_OFFSET (0xFFFFFFFFu)
 
 #if PLATFORM_ALLOW_SCENE_DATA_COMPRESSED_TRANSFORMS
diff --git a/Engine/Shaders/Shared/SceneDefinitions.h b/Engine/Shaders/Shared/SceneDefinitions.h
--- a/Engine/Shaders/Shared/SceneDefinitions.h
+++ b/Engine/Shaders/Shared/SceneDefinitions.h
@@ -6,6 +6,25 @@
 
 #pragma once
 
+
+// GPUCULL_TODO: Eventually we need to remove this workaround
+#define VF_TREAT_INSTANCE_ID_OFFSET_AS_PRIMITIVE_ID_FLAG (1U << 31U)
+
+#define PRIMITIVE_ID_NUM_BITS				(20u) // Max of 1,048,576 primitives
+#define INSTANCE_ID_NUM_BITS				(24u) // Max of 16,777,216 instances in the buffer
+#define INSTANCE_SCENE_DATA_FLAGS_NUM_BITS	(12u) // Max of 12 flags
+#define INSTANCE_RELATIVE_ID_NUM_BITS		(24u) // Max of 16,777,216 instances per primitive
+#define INSTANCE_CUSTOM_DATA_COUNT_NUM_BITS	 (8u) // Max of 255 custom data floats per instance
+
+#define MAX_PRIMITIVE_ID					(1U << PRIMITIVE_ID_NUM_BITS)
+#define MAX_INSTANCE_ID						(1U << INSTANCE_ID_NUM_BITS)
+
+#define PRIMITIVE_ID_MASK			((1u << PRIMITIVE_ID_NUM_BITS) - 1u)
+#define INSTANCE_ID_MASK			((1u << INSTANCE_ID_NUM_BITS) - 1u)
+#define INSTANCE_RELATIVE_ID_MASK	((1u << INSTANCE_RELATIVE_ID_NUM_BITS) - 1u)
+
+#define INVALID_PRIMITIVE_ID PRIMITIVE_ID_MASK
+
 // Primitive scene data flags
 #define PRIMITIVE_SCENE_DATA_FLAG_CAST_SHADOWS							0x1
 #define PRIMITIVE_SCENE_DATA_FLAG_USE_SINGLE_SAMPLE_SHADOW_SL			0x2
@@ -70,15 +89,3 @@
 	| INSTANCE_SCENE_DATA_FLAG_HAS_EDITOR_DATA			\
 	| INSTANCE_SCENE_DATA_FLAG_HAS_SPLINE_MESH_PARAMS	\
 )
-
-#define PRIMITIVE_ID_NUM_BITS				(20u) // Max of 1,048,576 primitives
-#define INSTANCE_ID_NUM_BITS				(24u) // Max of 16,777,216 instances in the buffer
-#define INSTANCE_SCENE_DATA_FLAGS_NUM_BITS	(12u) // Max of 12 flags
-#define INSTANCE_RELATIVE_ID_NUM_BITS		(24u) // Max of 16,777,216 instances per primitive
-#define INSTANCE_CUSTOM_DATA_COUNT_NUM_BITS	 (8u) // Max of 255 custom data floats per instance
-
-#define PRIMITIVE_ID_MASK			((1u << PRIMITIVE_ID_NUM_BITS) - 1u)
-#define INSTANCE_ID_MASK			((1u << INSTANCE_ID_NUM_BITS) - 1u)
-#define INSTANCE_RELATIVE_ID_MASK	((1u << INSTANCE_RELATIVE_ID_NUM_BITS) - 1u)
-
-#define INVALID_PRIMITIVE_ID PRIMITIVE_ID_MASK
diff --git a/Engine/Source/Runtime/Renderer/Private/GPUScene.cpp b/Engine/Source/Runtime/Renderer/Private/GPUScene.cpp
--- a/Engine/Source/Runtime/Renderer/Private/GPUScene.cpp
+++ b/Engine/Source/Runtime/Renderer/Private/GPUScene.cpp
@@ -10,9 +10,6 @@
 #include "RHI.h"
 #include "SceneUtils.h"
 #include "ScenePrivate.h"
-#include "SpriteIndexBuffer.h"
-#include "PostProcess/SceneFilterRendering.h"
-#include "ClearQuad.h"
 #include "RendererModule.h"
 #include "Rendering/NaniteResources.h"
 #include "Async/ParallelFor.h"
@@ -29,6 +26,7 @@
 #include "SceneDefinitions.h"
 #include "PrimitiveSceneShaderData.h"
 #include "Components/SplineMeshComponent.h"
+#include "RendererOnScreenNotification.h"
 
 // Defaults to being disabled, enable using the command line argument: -CsvCategory GPUScene
 CSV_DEFINE_CATEGORY(GPUScene, false);
@@ -567,11 +565,34 @@ FGPUScene::FGPUScene()
 	, InstancePayloadDataAllocator(CVarGPUSceneUseGrowOnlyAllocationPolicy.GetValueOnAnyThread() != 0)
 	, LightmapDataAllocator(CVarGPUSceneUseGrowOnlyAllocationPolicy.GetValueOnAnyThread() != 0)
 {
+#if !UE_BUILD_SHIPPING
+	ScreenMessageDelegate = FRendererOnScreenNotification::Get().AddLambda([this](TMultiMap<FCoreDelegates::EOnScreenMessageSeverity, FText >& OutMessages)
+	{
+		if (InstanceSceneDataSOAStride < uint32(InstanceSceneDataAllocator.GetMaxSize()))
+		{
+			OutMessages.Add(FCoreDelegates::EOnScreenMessageSeverity::Warning, FText::FromString(FString::Printf(TEXT("GPU-Scene Instance data overflow detected, reduce instance count to avoid rendering artifacts"))));
+			OutMessages.Add(FCoreDelegates::EOnScreenMessageSeverity::Warning, FText::FromString(FString::Printf(TEXT("  Max allocated ID %d, instance buffer size: %dM"), InstanceSceneDataAllocator.GetMaxSize(), InstanceSceneDataSOAStride >> 20)));
+			if (!bLoggedInstanceOverflow )
+			{
+				UE_LOG(LogRenderer, Warning, TEXT("GPU-Scene Instance data overflow detected, reduce instance count to avoid rendering artifacts.\n")
+					TEXT(" Max allocated ID %d (%0.3fM), instance buffer size: %dM"), InstanceSceneDataAllocator.GetMaxSize(), double(InstanceSceneDataAllocator.GetMaxSize()) / (1024.0 * 1024.0),  InstanceSceneDataSOAStride >> 20);
+				bLoggedInstanceOverflow = true;
+			}
+		}
+		else
+		{
+			bLoggedInstanceOverflow = false;
+		}
+	});
+#endif
 }
 
 
 FGPUScene::~FGPUScene()
 {
+#if !UE_BUILD_SHIPPING
+	FRendererOnScreenNotification::Get().Remove(ScreenMessageDelegate);
+#endif
 }
 
 void FGPUScene::BeginRender(const FScene* Scene, FGPUSceneDynamicContext &GPUSceneDynamicContext)
@@ -774,7 +795,9 @@ void FGPUScene::UpdateBufferState(FRDGBuilder& GraphBuilder, FSceneUniformBuffer
 	const uint32 SizeReserve = FMath::RoundUpToPowerOfTwo(FMath::Max(DynamicPrimitivesOffset, InitialBufferSize));
 	BufferState.PrimitiveBuffer = ResizeStructuredBufferIfNeeded(GraphBuilder, PrimitiveBuffer, SizeReserve * sizeof(FPrimitiveSceneShaderData::Data), TEXT("GPUScene.PrimitiveData"));
 
-	const uint32 InstanceSceneDataSizeReserve = FMath::RoundUpToPowerOfTwo(FMath::Max(InstanceSceneDataAllocator.GetMaxSize(), InitialBufferSize));
+	// Clamp buffer to be smaller than the MAX_INSTANCE_ID.
+	const uint32 InstanceSceneDataSizeReserve = FMath::Min(MAX_INSTANCE_ID, FMath::RoundUpToPowerOfTwo(FMath::Max(InstanceSceneDataAllocator.GetMaxSize(), InitialBufferSize)));
+
 	FResizeResourceSOAParams ResizeParams;
 	ResizeParams.NumBytes = InstanceSceneDataSizeReserve * FInstanceSceneShaderData::GetEffectiveNumBytes();
 	ResizeParams.NumArrays = FInstanceSceneShaderData::GetDataStrideInFloat4s();
diff --git a/Engine/Source/Runtime/Renderer/Private/GPUScene.h b/Engine/Source/Runtime/Renderer/Private/GPUScene.h
--- a/Engine/Source/Runtime/Renderer/Private/GPUScene.h
+++ b/Engine/Source/Runtime/Renderer/Private/GPUScene.h
@@ -422,6 +422,13 @@ private:
 	void AddUpdatePrimitiveIdsPass(FRDGBuilder& GraphBuilder, FInstanceGPULoadBalancer& IdOnlyUpdateItems);
 
 	void AddClearInstancesPass(FRDGBuilder& GraphBuilder);
+
+
+#if !UE_BUILD_SHIPPING
+	FDelegateHandle ScreenMessageDelegate;
+	bool bLoggedInstanceOverflow = false;
+#endif // UE_BUILD_SHIPPING
+
 };
 
 class FGPUSceneScopeBeginEndHelper
diff --git a/Engine/Source/Runtime/Renderer/Private/MeshDrawCommands.cpp b/Engine/Source/Runtime/Renderer/Private/MeshDrawCommands.cpp
--- a/Engine/Source/Runtime/Renderer/Private/MeshDrawCommands.cpp
+++ b/Engine/Source/Runtime/Renderer/Private/MeshDrawCommands.cpp
@@ -11,6 +11,7 @@ MeshDrawCommandSetup.cpp: Mesh draw command setup.
 #include "TranslucentRendering.h"
 #include "InstanceCulling/InstanceCullingManager.h"
 #include "StaticMeshBatch.h"
+#include "SceneDefinitions.h"
 
 TGlobalResource<FPrimitiveIdVertexBufferPool> GPrimitiveIdVertexBufferPool;
 
@@ -446,7 +447,6 @@ FORCEINLINE int32 TranslatePrimitiveId(int32 DrawPrimitiveIdIn, int32 DynamicPri
 
 	// Append flag to mark this as a non-instance data index.
 	// This value is treated as a primitive ID in the SceneData.ush loading.
-	const uint32 VF_TREAT_INSTANCE_ID_OFFSET_AS_PRIMITIVE_ID_FLAG = 1U << 31U;
 	return DrawPrimitiveId |= VF_TREAT_INSTANCE_ID_OFFSET_AS_PRIMITIVE_ID_FLAG;
 }
 
