Change default of r.Shadow.Virtual.NonNanite.CulledInstanceAllocationFactor to 1 as small scenes were producing artifacts.

diff --git a/Engine/Source/Runtime/Renderer/Private/VirtualShadowMaps/VirtualShadowMapArray.cpp b/Engine/Source/Runtime/Renderer/Private/VirtualShadowMaps/VirtualShadowMapArray.cpp
--- a/Engine/Source/Runtime/Renderer/Private/VirtualShadowMaps/VirtualShadowMapArray.cpp
+++ b/Engine/Source/Runtime/Renderer/Private/VirtualShadowMaps/VirtualShadowMapArray.cpp
@@ -137,7 +137,7 @@ TAutoConsoleVariable<int32> CVarCoarsePagesIncludeNonNanite(
 
 TAutoConsoleVariable<float> CVarNonNaniteCulledInstanceAllocationFactor(
 	TEXT("r.Shadow.Virtual.NonNanite.CulledInstanceAllocationFactor"),
-	0.33f,
+	1.0f,
 	TEXT("Scale factor multiplied by the conservative size required if all instances are emitted into every clip or mip level, setting to 1.0 is fully conservative.")
 	TEXT("The actual number cannot be known on the CPU as the culling emits an instace for each clip/mip level that is overlapped."),
 	ECVF_RenderThreadSafe
