Hide details about Substrate public uniform buffer.

Done by [REDACTED_PERSON]

diff --git a/Engine/Source/Runtime/Renderer/Private/Strata/Strata.h b/Engine/Source/Runtime/Renderer/Private/Strata/Strata.h
--- a/Engine/Source/Runtime/Renderer/Private/Strata/Strata.h
+++ b/Engine/Source/Runtime/Renderer/Private/Strata/Strata.h
@@ -73,6 +73,10 @@ BEGIN_GLOBAL_SHADER_PARAMETER_STRUCT(FStrataGlobalUniformParameters, RENDERER_AP
 	SHADER_PARAMETER_RDG_BUFFER_SRV(Buffer<uint>, BSDFTileCountBuffer)
 END_GLOBAL_SHADER_PARAMETER_STRUCT()
 
+BEGIN_GLOBAL_SHADER_PARAMETER_STRUCT(FStrataPublicGlobalUniformParameters, RENDERER_API)
+	SHADER_PARAMETER_RDG_TEXTURE(Texture2D<uint2>, TopLayerTexture)
+END_GLOBAL_SHADER_PARAMETER_STRUCT()
+
 // This must map to the STRATA_TILE_TYPE defines.
 enum EStrataTileType : uint32
 {
diff --git a/Engine/Source/Runtime/Renderer/Public/RendererUtils.h b/Engine/Source/Runtime/Renderer/Public/RendererUtils.h
--- a/Engine/Source/Runtime/Renderer/Public/RendererUtils.h
+++ b/Engine/Source/Runtime/Renderer/Public/RendererUtils.h
@@ -54,9 +54,7 @@ public:
 };
 
 // A minimal uniform struct providing necessary access for external systems to Strata parameters.
-BEGIN_GLOBAL_SHADER_PARAMETER_STRUCT(FStrataPublicGlobalUniformParameters, RENDERER_API)
-	SHADER_PARAMETER_RDG_TEXTURE(Texture2D<uint2>, TopLayerTexture)
-END_GLOBAL_SHADER_PARAMETER_STRUCT()
+DECLARE_UNIFORM_BUFFER_STRUCT(FStrataPublicGlobalUniformParameters, RENDERER_API)
 
 namespace Strata
 {
