Spam log when users set invalid VRS CVar options (clean-up in future)

diff --git a/Engine/Source/Runtime/Renderer/Private/VariableRateShading/VariableRateShadingImageManager.cpp b/Engine/Source/Runtime/Renderer/Private/VariableRateShading/VariableRateShadingImageManager.cpp
--- a/Engine/Source/Runtime/Renderer/Private/VariableRateShading/VariableRateShadingImageManager.cpp
+++ b/Engine/Source/Runtime/Renderer/Private/VariableRateShading/VariableRateShadingImageManager.cpp
@@ -24,6 +24,8 @@
 
 TGlobalResource<FVariableRateShadingImageManager> GVRSImageManager;
 
+DEFINE_LOG_CATEGORY(LogVRS);
+
 /**
  * Basic CVars
  */
@@ -248,6 +250,14 @@ FVariableRateShadingImageManager::FVariableRateShadingImageManager()
 
 FVariableRateShadingImageManager::~FVariableRateShadingImageManager() {}
 
+void FVariableRateShadingImageManager::InitRHI()
+{
+	if (!IsVRSSupportedByRHI())
+	{
+		UE_LOG(LogVRS, Warning, TEXT("VRS is not supported by current RHI"));
+	}
+}
+
 void FVariableRateShadingImageManager::ReleaseDynamicRHI()
 {
 	GRenderTargetPool.FreeUnusedResources();
@@ -437,6 +447,10 @@ void FVariableRateShadingImageManager::DrawDebugPreview(FRDGBuilder& GraphBuilde
 	{
 		PreviewImageType = static_cast<EVRSImageType>(ImageTypeAsInt);
 	}
+	else
+	{
+		UE_LOG(LogVRS, Warning, TEXT("Selected invalid preview mode, disabling preview"));
+	}
 
 	if (PreviewImageType == EVRSImageType::Disabled || !OutputSceneColor)
 	{
@@ -576,7 +590,7 @@ FRDGTextureRef FVariableRateShadingImageManager::GetForceRateImage(FRDGBuilder&
 
 	if (RateIndex >= ValidShadingRates.Num())
 	{
-		// TODO: Log warning
+		UE_LOG(LogVRS, Warning, TEXT("Selected forced shading rate exceeds maximum available, defaulting to 4x4"));
 		RateIndex = ValidShadingRates.Num() - 1; // Default to maximum shading rate if value exceeds valid rates
 	}
 
@@ -617,6 +631,7 @@ FVariableRateShadingImageManager::EVRSImageType FVariableRateShadingImageManager
 	}
 	else
 	{
+		UE_LOG(LogVRS, Warning, TEXT("Selected invalid image type, disabling VRS for pass"));
 		return EVRSImageType::Disabled;
 	}
 }
\ No newline at end of file
diff --git a/Engine/Source/Runtime/Renderer/Public/VariableRateShadingImageManager.h b/Engine/Source/Runtime/Renderer/Public/VariableRateShadingImageManager.h
--- a/Engine/Source/Runtime/Renderer/Public/VariableRateShadingImageManager.h
+++ b/Engine/Source/Runtime/Renderer/Public/VariableRateShadingImageManager.h
@@ -7,6 +7,7 @@
 
 #pragma once
 
+#include "Logging/LogMacros.h"
 #include "RHI.h"
 #include "RendererInterface.h"
 #include "RenderGraphDefinitions.h"
@@ -15,6 +16,8 @@ struct FMinimalSceneTextures; // Forward declaration from SceneTextures.
 
 class IVariableRateShadingImageGenerator;
 
+DECLARE_LOG_CATEGORY_EXTERN(LogVRS, Log, All);
+
 class FVariableRateShadingImageManager : public FRenderResource
 {
 public:
@@ -74,6 +77,7 @@ public:
 	FVariableRateShadingImageManager();
 	virtual ~FVariableRateShadingImageManager();
 
+	virtual void InitRHI() override;
 	virtual void ReleaseDynamicRHI() override;
 
 	/**
