Speculative fixe for GSystemTextures.ZeroUIntDummy that could possibly transitions to Graphics pipe only

diff --git a/Engine/Source/Runtime/Renderer/Private/ScreenSpaceDenoise.cpp b/Engine/Source/Runtime/Renderer/Private/ScreenSpaceDenoise.cpp
--- a/Engine/Source/Runtime/Renderer/Private/ScreenSpaceDenoise.cpp
+++ b/Engine/Source/Runtime/Renderer/Private/ScreenSpaceDenoise.cpp
@@ -1879,16 +1879,18 @@ static void DenoiseSignalAtConstantPixelDensity(
 		bool bGlobalCameraCut = !View.PrevViewInfo.DepthBuffer.IsValid();
 		if (CompressedMetadataLayout == ECompressedMetadataLayout::DepthAndViewNormal)
 		{
-			PassParameters->PrevCompressedMetadata[0] = RegisterExternalTextureWithFallback(
-				GraphBuilder, ViewInfoPooledRenderTargets.PrevCompressedDepthViewNormal, GSystemTextures.ZeroUIntDummy);
+			PassParameters->PrevCompressedMetadata[0] = ViewInfoPooledRenderTargets.PrevCompressedDepthViewNormal
+				? GraphBuilder.RegisterExternalTexture(ViewInfoPooledRenderTargets.PrevCompressedDepthViewNormal)
+				: GSystemTextures.GetZeroUIntDummy(GraphBuilder);
 			bGlobalCameraCut = !View.PrevViewInfo.CompressedDepthViewNormal.IsValid();
 		}
 		else if (CompressedMetadataLayout == ECompressedMetadataLayout::FedDepthAndShadingModelID)
 		{
 			PassParameters->PrevCompressedMetadata[0] = RegisterExternalTextureWithFallback(
 				GraphBuilder, View.PrevViewInfo.CompressedOpaqueDepth, GSystemTextures.BlackDummy);
-			PassParameters->PrevCompressedMetadata[1] = RegisterExternalTextureWithFallback(
-				GraphBuilder, View.PrevViewInfo.CompressedOpaqueShadingModel, GSystemTextures.ZeroUIntDummy);
+			PassParameters->PrevCompressedMetadata[1] = View.PrevViewInfo.CompressedOpaqueShadingModel
+				? GraphBuilder.RegisterExternalTexture(View.PrevViewInfo.CompressedOpaqueShadingModel)
+				: GSystemTextures.GetZeroUIntDummy(GraphBuilder);
 
 			bGlobalCameraCut = !View.PrevViewInfo.CompressedOpaqueDepth.IsValid() || !View.PrevViewInfo.CompressedOpaqueShadingModel.IsValid();
 		}
diff --git a/Engine/Source/Runtime/Renderer/Private/Shadows/ShadowSceneRenderer.cpp b/Engine/Source/Runtime/Renderer/Private/Shadows/ShadowSceneRenderer.cpp
--- a/Engine/Source/Runtime/Renderer/Private/Shadows/ShadowSceneRenderer.cpp
+++ b/Engine/Source/Runtime/Renderer/Private/Shadows/ShadowSceneRenderer.cpp
@@ -446,7 +446,6 @@ void FShadowSceneRenderer::RenderVirtualShadowMapProjectionMaskBits(
 	}
 	else
 	{
-		//FRDGTextureRef Dummy = GraphBuilder.RegisterExternalTexture(GSystemTextures.ZeroUIntDummy);
 		VirtualShadowMapMaskBits = nullptr;//Dummy;
 		VirtualShadowMapMaskBitsHairStrands = nullptr;//Dummy;
 	}
