LocalHeightFogVolume - do not add LHFV if we know they are not going to be visible.

diff --git a/Engine/Source/Runtime/Renderer/Private/LocalHeightFogRendering.cpp b/Engine/Source/Runtime/Renderer/Private/LocalHeightFogRendering.cpp
--- a/Engine/Source/Runtime/Renderer/Private/LocalHeightFogRendering.cpp
+++ b/Engine/Source/Runtime/Renderer/Private/LocalHeightFogRendering.cpp
@@ -142,6 +142,7 @@ void RenderLocalHeightFog(
 	FRDGTextureRef LightShaftOcclusionTexture)
 {
 	uint32 LocalHeightFogInstanceCount = Scene->LocalHeightFogs.Num();
+	uint32 LocalHeightFogInstanceCountFinal = 0;
 	if (LocalHeightFogInstanceCount > 0)
 	{
 		RDG_GPU_STAT_SCOPE(GraphBuilder, LocalHeightFogVolumes);
@@ -152,6 +153,11 @@ void RenderLocalHeightFog(
 		FLocalHeightFogGPUInstanceData* LocalHeightFogGPUInstanceDataIt = LocalHeightFogGPUInstanceData;
 		for (FLocalHeightFogSceneProxy* LHF : Scene->LocalHeightFogs)
 		{
+			if (LHF->FogDensity <= 0.0f)
+			{
+				continue; // this volume will never be visible
+			}
+
 			FTransform TransformScaleOnly;
 			TransformScaleOnly.SetScale3D(LHF->FogTransform.GetScale3D());
 
@@ -171,9 +177,11 @@ void RenderLocalHeightFog(
 			LocalHeightFogGPUInstanceDataIt->Emissive = FVector3f(LHF->FogEmissive);
 
 			LocalHeightFogGPUInstanceDataIt++;
+			LocalHeightFogInstanceCountFinal++;
 		}
+		const uint32 AllLocalHeightFogInstanceBytesFinal = sizeof(FLocalHeightFogGPUInstanceData) * LocalHeightFogInstanceCountFinal;
 		FRDGBufferRef LocalHeightFogGPUInstanceDataBuffer =  CreateStructuredBuffer(GraphBuilder, TEXT("LocalHeightFogGPUInstanceDataBuffer"), 
-			sizeof(FLocalHeightFogGPUInstanceData), LocalHeightFogInstanceCount, LocalHeightFogGPUInstanceData, AllLocalHeightFogInstanceBytes, ERDGInitialDataFlags::NoCopy);
+			sizeof(FLocalHeightFogGPUInstanceData), LocalHeightFogInstanceCountFinal, LocalHeightFogGPUInstanceData, AllLocalHeightFogInstanceBytesFinal, ERDGInitialDataFlags::NoCopy);
 		FRDGBufferSRVRef LocalHeightFogGPUInstanceDataBufferSRV = GraphBuilder.CreateSRV(LocalHeightFogGPUInstanceDataBuffer);
 
 		FRDGTextureRef SceneColorTexture = SceneTextures.Color.Resolve;
