Refactored branch that skips distance field scene entries for degenerate instances and removed some unecessary scopes.
- shouldn't affect runtime behaviour.

diff --git a/Engine/Source/Runtime/Renderer/Private/DistanceFieldObjectManagement.cpp b/Engine/Source/Runtime/Renderer/Private/DistanceFieldObjectManagement.cpp
--- a/Engine/Source/Runtime/Renderer/Private/DistanceFieldObjectManagement.cpp
+++ b/Engine/Source/Runtime/Renderer/Private/DistanceFieldObjectManagement.cpp
@@ -364,71 +364,69 @@ bool ProcessPrimitiveUpdate(
 
 				const FMatrix::FReal MaxScale = LocalToWorld.GetMaximumAxisScale();
 
-				// Skip degenerate primitives
-				if (MaxScale > 0)
+				if (bIsAddOperation && MaxScale <= 0)
 				{
-					uint32 UploadIndex;
+					// Skip degenerate instances
+					PrimitiveSceneInfo->DistanceFieldInstanceIndices[TransformIndex] = -1;
+					continue;
+				}
 
-					if (bIsAddOperation)
-					{
-						UploadIndex = DistanceFieldSceneData.NumObjectsInBuffer;
-						++DistanceFieldSceneData.NumObjectsInBuffer;
-					}
-					else
-					{
-						UploadIndex = PrimitiveSceneInfo->DistanceFieldInstanceIndices[TransformIndex];
-					}
+				uint32 UploadIndex;
+
+				if (bIsAddOperation)
+				{
+					UploadIndex = DistanceFieldSceneData.NumObjectsInBuffer;
+					++DistanceFieldSceneData.NumObjectsInBuffer;
+				}
+				else
+				{
+					UploadIndex = PrimitiveSceneInfo->DistanceFieldInstanceIndices[TransformIndex];
+				}
 
-					IndicesToUpdateInObjectBuffers.Add(UploadIndex);
+				IndicesToUpdateInObjectBuffers.Add(UploadIndex);
 
-					const FBox WorldBounds = DistanceFieldData->LocalSpaceMeshBounds.TransformBy(LocalToWorld);
+				const FBox WorldBounds = DistanceFieldData->LocalSpaceMeshBounds.TransformBy(LocalToWorld);
 
-					if (bIsAddOperation)
-					{
-						const int32 AddIndex = UploadIndex;
-						const int32 MappingIndex = DistanceFieldSceneData.PrimitiveInstanceMapping.Add(FPrimitiveAndInstance(LocalToWorld, WorldBounds, PrimitiveSceneInfo, TransformIndex));
-						PrimitiveSceneInfo->DistanceFieldInstanceIndices[TransformIndex] = AddIndex;
+				if (bIsAddOperation)
+				{
+					const int32 AddIndex = UploadIndex;
+					const int32 MappingIndex = DistanceFieldSceneData.PrimitiveInstanceMapping.Add(FPrimitiveAndInstance(LocalToWorld, WorldBounds, PrimitiveSceneInfo, TransformIndex));
+					PrimitiveSceneInfo->DistanceFieldInstanceIndices[TransformIndex] = AddIndex;
 
-						AddModifiedBounds(Scene, CacheType, WorldBounds);
-						LogDistanceFieldUpdate(PrimitiveSceneInfo, BoundingRadius, bIsAddOperation);
-					}
-					else 
+					AddModifiedBounds(Scene, CacheType, WorldBounds);
+					LogDistanceFieldUpdate(PrimitiveSceneInfo, BoundingRadius, bIsAddOperation);
+				}
+				else 
+				{
+					// InstanceIndex will be -1 with zero scale meshes
+					const int32 InstanceIndex = PrimitiveSceneInfo->DistanceFieldInstanceIndices[TransformIndex];
+					if (InstanceIndex >= 0)
 					{
-						// InstanceIndex will be -1 with zero scale meshes
-						const int32 InstanceIndex = PrimitiveSceneInfo->DistanceFieldInstanceIndices[TransformIndex];
-						if (InstanceIndex >= 0)
-						{
-							FPrimitiveAndInstance& Mapping = DistanceFieldSceneData.PrimitiveInstanceMapping[InstanceIndex];
+						FPrimitiveAndInstance& Mapping = DistanceFieldSceneData.PrimitiveInstanceMapping[InstanceIndex];
 
-							// Filter out global distance field updates which were too small
-							if (!Mapping.WorldBounds.GetExtent().Equals(WorldBounds.GetExtent(), 0.01f)
-								|| !Mapping.LocalToWorld.Equals(LocalToWorld, 0.01f))
+						// Filter out global distance field updates which were too small
+						if (!Mapping.WorldBounds.GetExtent().Equals(WorldBounds.GetExtent(), 0.01f)
+							|| !Mapping.LocalToWorld.Equals(LocalToWorld, 0.01f))
+						{
+							// decide if we want to make a single global distance field update or two updates for large movement (teleport) case
+							const FBox MergedBounds = Mapping.WorldBounds + WorldBounds;
+							const FVector MergedExtentIncrease = MergedBounds.GetExtent() - Mapping.WorldBounds.GetExtent() - WorldBounds.GetExtent();
+							if (MergedExtentIncrease.GetMax() < 100.0f)
 							{
-								// decide if we want to make a single global distance field update or two updates for large movement (teleport) case
-								const FBox MergedBounds = Mapping.WorldBounds + WorldBounds;
-								const FVector MergedExtentIncrease = MergedBounds.GetExtent() - Mapping.WorldBounds.GetExtent() - WorldBounds.GetExtent();
-								if (MergedExtentIncrease.GetMax() < 100.0f)
-								{
-									AddModifiedBounds(Scene, CacheType, MergedBounds);
-								}
-								else
-								{
-									AddModifiedBounds(Scene, CacheType, Mapping.WorldBounds);
-									AddModifiedBounds(Scene, CacheType, WorldBounds);
-								}
-								LogDistanceFieldUpdate(PrimitiveSceneInfo, BoundingRadius, bIsAddOperation);
-
-								Mapping.LocalToWorld = LocalToWorld;
-								Mapping.WorldBounds = WorldBounds;
+								AddModifiedBounds(Scene, CacheType, MergedBounds);
 							}
+							else
+							{
+								AddModifiedBounds(Scene, CacheType, Mapping.WorldBounds);
+								AddModifiedBounds(Scene, CacheType, WorldBounds);
+							}
+							LogDistanceFieldUpdate(PrimitiveSceneInfo, BoundingRadius, bIsAddOperation);
+
+							Mapping.LocalToWorld = LocalToWorld;
+							Mapping.WorldBounds = WorldBounds;
 						}
 					}
 				}
-				else if (bIsAddOperation)
-				{
-					// Set to -1 for zero scale meshes
-					PrimitiveSceneInfo->DistanceFieldInstanceIndices[TransformIndex] = -1;
-				}
 			}
 		}
 		else
@@ -523,169 +521,168 @@ void FDistanceFieldSceneData::UpdateDistanceFieldObjectBuffers(
 			QUICK_SCOPE_CYCLE_COUNTER(UpdateDFObjectBuffers);
 
 			// Upload DF object data and bounds
-			{
-				const uint32 NumDFObjects = NumObjectsInBuffer;
 
-				const uint32 DFObjectDataNumFloat4s = FMath::RoundUpToPowerOfTwo(NumDFObjects * GDistanceFieldObjectDataStride);
-				const uint32 DFObjectDataNumBytes = DFObjectDataNumFloat4s * sizeof(FVector4f);
-				FRDGBuffer* DFObjectDataBuffer = ResizeStructuredBufferIfNeeded(GraphBuilder, ObjectBuffers->Data, DFObjectDataNumBytes, TEXT("DistanceFields.DFObjectData"));
+			const uint32 NumDFObjects = NumObjectsInBuffer;
 
-				const uint32 DFObjectBoundsNumFloat4s = FMath::RoundUpToPowerOfTwo(NumDFObjects * GDistanceFieldObjectBoundsStride);
-				const uint32 DFObjectBoundsNumBytes = DFObjectBoundsNumFloat4s * sizeof(FVector4f);
-				FRDGBuffer* DFObjectBoundsBuffer = ResizeStructuredBufferIfNeeded(GraphBuilder, ObjectBuffers->Bounds, DFObjectBoundsNumBytes, TEXT("DistanceFields.DFObjectBounds"));
+			const uint32 DFObjectDataNumFloat4s = FMath::RoundUpToPowerOfTwo(NumDFObjects * GDistanceFieldObjectDataStride);
+			const uint32 DFObjectDataNumBytes = DFObjectDataNumFloat4s * sizeof(FVector4f);
+			FRDGBuffer* DFObjectDataBuffer = ResizeStructuredBufferIfNeeded(GraphBuilder, ObjectBuffers->Data, DFObjectDataNumBytes, TEXT("DistanceFields.DFObjectData"));
 
-				const int32 NumDFObjectUploads = IndicesToUpdateInObjectBuffers.Num();
+			const uint32 DFObjectBoundsNumFloat4s = FMath::RoundUpToPowerOfTwo(NumDFObjects * GDistanceFieldObjectBoundsStride);
+			const uint32 DFObjectBoundsNumBytes = DFObjectBoundsNumFloat4s * sizeof(FVector4f);
+			FRDGBuffer* DFObjectBoundsBuffer = ResizeStructuredBufferIfNeeded(GraphBuilder, ObjectBuffers->Bounds, DFObjectBoundsNumBytes, TEXT("DistanceFields.DFObjectBounds"));
 
-				static FCriticalSection DFUpdateCS;
+			const int32 NumDFObjectUploads = IndicesToUpdateInObjectBuffers.Num();
 
-				if (NumDFObjectUploads > 0)
-				{
-					UploadDistanceFieldDataBuffer.Init(GraphBuilder, NumDFObjectUploads, GDistanceFieldObjectDataStride * sizeof(FVector4f), true, TEXT("DistanceFields.DFObjectDataUploadBuffer"));
-					UploadDistanceFieldBoundsBuffer.Init(GraphBuilder, NumDFObjectUploads, GDistanceFieldObjectBoundsStride * sizeof(FVector4f), true, TEXT("DistanceFields.DFObjectBoundsUploadBuffer"));
+			static FCriticalSection DFUpdateCS;
+
+			if (NumDFObjectUploads > 0)
+			{
+				UploadDistanceFieldDataBuffer.Init(GraphBuilder, NumDFObjectUploads, GDistanceFieldObjectDataStride * sizeof(FVector4f), true, TEXT("DistanceFields.DFObjectDataUploadBuffer"));
+				UploadDistanceFieldBoundsBuffer.Init(GraphBuilder, NumDFObjectUploads, GDistanceFieldObjectBoundsStride * sizeof(FVector4f), true, TEXT("DistanceFields.DFObjectBoundsUploadBuffer"));
 
-					const TScenePrimitiveArray<FPrimitiveBounds>& PrimitiveBounds = Scene->PrimitiveBounds;
+				const TScenePrimitiveArray<FPrimitiveBounds>& PrimitiveBounds = Scene->PrimitiveBounds;
 
-					FParallelUpdateRangesDFO ParallelRanges;
+				FParallelUpdateRangesDFO ParallelRanges;
 
-					int32 RangeCount = PartitionUpdateRangesDFO(ParallelRanges, IndicesToUpdateInObjectBuffers.Num(), bExecuteInParallel);
+				int32 RangeCount = PartitionUpdateRangesDFO(ParallelRanges, IndicesToUpdateInObjectBuffers.Num(), bExecuteInParallel);
 
-					ParallelFor(RangeCount,
-						[this, &ParallelRanges, &PrimitiveBounds, RangeCount](int32 RangeIndex)
+				ParallelFor(RangeCount,
+					[this, &ParallelRanges, &PrimitiveBounds, RangeCount](int32 RangeIndex)
+					{
+						for (int32 ItemIndex = ParallelRanges.Range[RangeIndex].ItemStart; ItemIndex < ParallelRanges.Range[RangeIndex].ItemStart + ParallelRanges.Range[RangeIndex].ItemCount; ++ItemIndex)
 						{
-							for (int32 ItemIndex = ParallelRanges.Range[RangeIndex].ItemStart; ItemIndex < ParallelRanges.Range[RangeIndex].ItemStart + ParallelRanges.Range[RangeIndex].ItemCount; ++ItemIndex)
+							const int32 Index = IndicesToUpdateInObjectBuffers[ItemIndex];
+							if (Index >= 0 && Index < PrimitiveInstanceMapping.Num())
 							{
-								const int32 Index = IndicesToUpdateInObjectBuffers[ItemIndex];
-								if (Index >= 0 && Index < PrimitiveInstanceMapping.Num())
-								{
-									const FPrimitiveAndInstance& PrimAndInst = PrimitiveInstanceMapping[Index];
-									const FPrimitiveSceneProxy* PrimitiveSceneProxy = PrimAndInst.Primitive->Proxy;
+								const FPrimitiveAndInstance& PrimAndInst = PrimitiveInstanceMapping[Index];
+								const FPrimitiveSceneProxy* PrimitiveSceneProxy = PrimAndInst.Primitive->Proxy;
 
-									if (RangeCount > 1)
-									{
-										DFUpdateCS.Lock();
-									}
+								if (RangeCount > 1)
+								{
+									DFUpdateCS.Lock();
+								}
 
-									FVector4f* UploadObjectData = (FVector4f*)UploadDistanceFieldDataBuffer.Add_GetRef(Index);
-									FVector4f* UploadObjectBounds = (FVector4f*)UploadDistanceFieldBoundsBuffer.Add_GetRef(Index);
+								FVector4f* UploadObjectData = (FVector4f*)UploadDistanceFieldDataBuffer.Add_GetRef(Index);
+								FVector4f* UploadObjectBounds = (FVector4f*)UploadDistanceFieldBoundsBuffer.Add_GetRef(Index);
 
-									if (RangeCount > 1)
-									{
-										DFUpdateCS.Unlock();
-									}
+								if (RangeCount > 1)
+								{
+									DFUpdateCS.Unlock();
+								}
 
-									const FDistanceFieldVolumeData* DistanceFieldData = nullptr;
-									float SelfShadowBias;
-									PrimitiveSceneProxy->GetDistanceFieldAtlasData(DistanceFieldData, SelfShadowBias);
+								const FDistanceFieldVolumeData* DistanceFieldData = nullptr;
+								float SelfShadowBias;
+								PrimitiveSceneProxy->GetDistanceFieldAtlasData(DistanceFieldData, SelfShadowBias);
 
-									const FBox LocalSpaceMeshBounds = DistanceFieldData->LocalSpaceMeshBounds;
+								const FBox LocalSpaceMeshBounds = DistanceFieldData->LocalSpaceMeshBounds;
 			
-									const FMatrix LocalToWorld = PrimAndInst.LocalToWorld;
-
-									{
-										const FBox WorldSpaceMeshBounds = LocalSpaceMeshBounds.TransformBy(LocalToWorld);
-
-										const FLargeWorldRenderPosition AbsoluteWorldPosition(WorldSpaceMeshBounds.GetCenter());
-
-										const FVector4f ObjectBoundingSphere(AbsoluteWorldPosition.GetOffset(), WorldSpaceMeshBounds.GetExtent().Size());
-
-										UploadObjectBounds[0] = AbsoluteWorldPosition.GetTile();
-										UploadObjectBounds[1] = ObjectBoundingSphere;
-
-										const FGlobalDFCacheType CacheType = PrimitiveSceneProxy->IsOftenMoving() ? GDF_Full : GDF_MostlyStatic;
-										const bool bOftenMoving = CacheType == GDF_Full;
-										const bool bCastShadow = PrimitiveSceneProxy->CastsDynamicShadow();
-										const bool bIsNaniteMesh = PrimitiveSceneProxy->IsNaniteMesh();
-										const bool bEmissiveLightSource = PrimitiveSceneProxy->IsEmissiveLightSource();
-										const bool bVisible = PrimitiveSceneProxy->IsDrawnInGame(); // Distance field object can be invisible in main view, but cast shadows
-										const bool bAffectIndirectLightingWhileHidden = PrimitiveSceneProxy->AffectsIndirectLightingWhileHidden();
-
-										uint32 Flags = 0;
-										Flags |= bOftenMoving ? 1u : 0;
-										Flags |= bCastShadow ? 2u : 0;
-										Flags |= bIsNaniteMesh ? 4u : 0;
-										Flags |= bEmissiveLightSource ? 8u : 0;
-										Flags |= bVisible ? 16u : 0;
-										Flags |= bAffectIndirectLightingWhileHidden ? 32u : 0;
-
-										FVector4f ObjectWorldExtentAndFlags((FVector3f)WorldSpaceMeshBounds.GetExtent(), 0.0f);
-										ObjectWorldExtentAndFlags.W = *(const float*)&Flags;
-										UploadObjectBounds[2] = ObjectWorldExtentAndFlags;
-									}
-
-									// Uniformly scale our Volume space to lie within [-1, 1] at the max extent
-									// This is mirrored in the SDF encoding
-									const FBox::FReal LocalToVolumeScale = 1.0f / LocalSpaceMeshBounds.GetExtent().GetMax();
-
-									const FMatrix VolumeToWorld = FScaleMatrix(1.0f / LocalToVolumeScale) * FTranslationMatrix(LocalSpaceMeshBounds.GetCenter()) * LocalToWorld;
-
-									const FLargeWorldRenderPosition WorldPosition(VolumeToWorld.GetOrigin());
-									const FVector TilePositionOffset = WorldPosition.GetTileOffset();
-
-									// Inverse on FMatrix44f can generate NaNs if the source matrix contains large scaling, so do it in double precision.
-									FMatrix LocalToRelativeWorld = FLargeWorldRenderScalar::MakeToRelativeWorldMatrixDouble(TilePositionOffset, VolumeToWorld);
-
-									// TilePosition
-									UploadObjectData[0] = WorldPosition.GetTile();
-
-									const FMatrix44f WorldToVolumeT = FMatrix44f(LocalToRelativeWorld.Inverse().GetTransposed());
-									// WorldToVolumeT
-									UploadObjectData[1] = (*(FVector4f*)&WorldToVolumeT.M[0]);
-									UploadObjectData[2] = (*(FVector4f*)&WorldToVolumeT.M[1]);
-									UploadObjectData[3] = (*(FVector4f*)&WorldToVolumeT.M[2]);
-
-									const FVector VolumePositionExtent = LocalSpaceMeshBounds.GetExtent() * LocalToVolumeScale;
-
-									// Minimal surface bias which increases chance that ray hit will a surface located between two texels
-									float ExpandSurfaceDistance = (GMeshSDFSurfaceBiasExpand * VolumePositionExtent / FVector(DistanceFieldData->Mips[0].IndirectionDimensions * DistanceField::UniqueDataBrickSize)).Size();
-
-									const float WSign = DistanceFieldData->bMostlyTwoSided ? -1 : 1;
-									UploadObjectData[4] = FVector4f((FVector3f)VolumePositionExtent, WSign * ExpandSurfaceDistance);
-
-									const int32 PrimIdx = PrimAndInst.Primitive->GetIndex();
-									const FPrimitiveBounds& PrimBounds = PrimitiveBounds[PrimIdx];
-									float MinDrawDist2 = FMath::Square(PrimBounds.MinDrawDistance);
-									// For IEEE compatible machines, float operations goes to inf if overflow
-									// In this case, it will effectively disable max draw distance culling
-									float MaxDrawDist = FMath::Max(PrimBounds.MaxCullDistance, 0.f) * GetCachedScalabilityCVars().ViewDistanceScale;
-
-									const uint32 GPUSceneInstanceIndex = PrimitiveSceneProxy->SupportsInstanceDataBuffer() ? 
-										PrimAndInst.Primitive->GetInstanceSceneDataOffset() + PrimAndInst.InstanceIndex :
-										PrimAndInst.Primitive->GetInstanceSceneDataOffset();
-
-									// Bypass NaN checks in FVector4f ctor
-									FVector4f Vector4;
-									Vector4.X = MinDrawDist2;
-									Vector4.Y = MaxDrawDist * MaxDrawDist;
-									Vector4.Z = SelfShadowBias;
-									Vector4.W = *(const float*)&GPUSceneInstanceIndex;
-									UploadObjectData[5] = Vector4;
-
-									const FMatrix44f VolumeToWorldT = FMatrix44f(LocalToRelativeWorld.GetTransposed());
-									UploadObjectData[6] = *(FVector4f*)&VolumeToWorldT.M[0];
-									UploadObjectData[7] = *(FVector4f*)&VolumeToWorldT.M[1];
-									UploadObjectData[8] = *(FVector4f*)&VolumeToWorldT.M[2];
-
-									FVector4f FloatVector8(FVector3f(VolumeToWorld.GetScaleVector()), 0.0f);
-
-									// Bypass NaN checks in FVector4f ctor
-									FSetElementId AssetStateSetId = AssetStateArray.FindId(DistanceFieldData);
-									check(AssetStateSetId.IsValidId());
-									const int32 AssetStateInt = AssetStateSetId.AsInteger();
-									FloatVector8.W = *(const float*)&AssetStateInt;
-									
-									UploadObjectData[9] = FloatVector8;
+								const FMatrix LocalToWorld = PrimAndInst.LocalToWorld;
+
+								{
+									const FBox WorldSpaceMeshBounds = LocalSpaceMeshBounds.TransformBy(LocalToWorld);
+
+									const FLargeWorldRenderPosition AbsoluteWorldPosition(WorldSpaceMeshBounds.GetCenter());
+
+									const FVector4f ObjectBoundingSphere(AbsoluteWorldPosition.GetOffset(), WorldSpaceMeshBounds.GetExtent().Size());
+
+									UploadObjectBounds[0] = AbsoluteWorldPosition.GetTile();
+									UploadObjectBounds[1] = ObjectBoundingSphere;
+
+									const FGlobalDFCacheType CacheType = PrimitiveSceneProxy->IsOftenMoving() ? GDF_Full : GDF_MostlyStatic;
+									const bool bOftenMoving = CacheType == GDF_Full;
+									const bool bCastShadow = PrimitiveSceneProxy->CastsDynamicShadow();
+									const bool bIsNaniteMesh = PrimitiveSceneProxy->IsNaniteMesh();
+									const bool bEmissiveLightSource = PrimitiveSceneProxy->IsEmissiveLightSource();
+									const bool bVisible = PrimitiveSceneProxy->IsDrawnInGame(); // Distance field object can be invisible in main view, but cast shadows
+									const bool bAffectIndirectLightingWhileHidden = PrimitiveSceneProxy->AffectsIndirectLightingWhileHidden();
+
+									uint32 Flags = 0;
+									Flags |= bOftenMoving ? 1u : 0;
+									Flags |= bCastShadow ? 2u : 0;
+									Flags |= bIsNaniteMesh ? 4u : 0;
+									Flags |= bEmissiveLightSource ? 8u : 0;
+									Flags |= bVisible ? 16u : 0;
+									Flags |= bAffectIndirectLightingWhileHidden ? 32u : 0;
+
+									FVector4f ObjectWorldExtentAndFlags((FVector3f)WorldSpaceMeshBounds.GetExtent(), 0.0f);
+									ObjectWorldExtentAndFlags.W = *(const float*)&Flags;
+									UploadObjectBounds[2] = ObjectWorldExtentAndFlags;
 								}
+
+								// Uniformly scale our Volume space to lie within [-1, 1] at the max extent
+								// This is mirrored in the SDF encoding
+								const FBox::FReal LocalToVolumeScale = 1.0f / LocalSpaceMeshBounds.GetExtent().GetMax();
+
+								const FMatrix VolumeToWorld = FScaleMatrix(1.0f / LocalToVolumeScale) * FTranslationMatrix(LocalSpaceMeshBounds.GetCenter()) * LocalToWorld;
+
+								const FLargeWorldRenderPosition WorldPosition(VolumeToWorld.GetOrigin());
+								const FVector TilePositionOffset = WorldPosition.GetTileOffset();
+
+								// Inverse on FMatrix44f can generate NaNs if the source matrix contains large scaling, so do it in double precision.
+								FMatrix LocalToRelativeWorld = FLargeWorldRenderScalar::MakeToRelativeWorldMatrixDouble(TilePositionOffset, VolumeToWorld);
+
+								// TilePosition
+								UploadObjectData[0] = WorldPosition.GetTile();
+
+								const FMatrix44f WorldToVolumeT = FMatrix44f(LocalToRelativeWorld.Inverse().GetTransposed());
+								// WorldToVolumeT
+								UploadObjectData[1] = (*(FVector4f*)&WorldToVolumeT.M[0]);
+								UploadObjectData[2] = (*(FVector4f*)&WorldToVolumeT.M[1]);
+								UploadObjectData[3] = (*(FVector4f*)&WorldToVolumeT.M[2]);
+
+								const FVector VolumePositionExtent = LocalSpaceMeshBounds.GetExtent() * LocalToVolumeScale;
+
+								// Minimal surface bias which increases chance that ray hit will a surface located between two texels
+								float ExpandSurfaceDistance = (GMeshSDFSurfaceBiasExpand * VolumePositionExtent / FVector(DistanceFieldData->Mips[0].IndirectionDimensions * DistanceField::UniqueDataBrickSize)).Size();
+
+								const float WSign = DistanceFieldData->bMostlyTwoSided ? -1 : 1;
+								UploadObjectData[4] = FVector4f((FVector3f)VolumePositionExtent, WSign * ExpandSurfaceDistance);
+
+								const int32 PrimIdx = PrimAndInst.Primitive->GetIndex();
+								const FPrimitiveBounds& PrimBounds = PrimitiveBounds[PrimIdx];
+								float MinDrawDist2 = FMath::Square(PrimBounds.MinDrawDistance);
+								// For IEEE compatible machines, float operations goes to inf if overflow
+								// In this case, it will effectively disable max draw distance culling
+								float MaxDrawDist = FMath::Max(PrimBounds.MaxCullDistance, 0.f) * GetCachedScalabilityCVars().ViewDistanceScale;
+
+								const uint32 GPUSceneInstanceIndex = PrimitiveSceneProxy->SupportsInstanceDataBuffer() ? 
+									PrimAndInst.Primitive->GetInstanceSceneDataOffset() + PrimAndInst.InstanceIndex :
+									PrimAndInst.Primitive->GetInstanceSceneDataOffset();
+
+								// Bypass NaN checks in FVector4f ctor
+								FVector4f Vector4;
+								Vector4.X = MinDrawDist2;
+								Vector4.Y = MaxDrawDist * MaxDrawDist;
+								Vector4.Z = SelfShadowBias;
+								Vector4.W = *(const float*)&GPUSceneInstanceIndex;
+								UploadObjectData[5] = Vector4;
+
+								const FMatrix44f VolumeToWorldT = FMatrix44f(LocalToRelativeWorld.GetTransposed());
+								UploadObjectData[6] = *(FVector4f*)&VolumeToWorldT.M[0];
+								UploadObjectData[7] = *(FVector4f*)&VolumeToWorldT.M[1];
+								UploadObjectData[8] = *(FVector4f*)&VolumeToWorldT.M[2];
+
+								FVector4f FloatVector8(FVector3f(VolumeToWorld.GetScaleVector()), 0.0f);
+
+								// Bypass NaN checks in FVector4f ctor
+								FSetElementId AssetStateSetId = AssetStateArray.FindId(DistanceFieldData);
+								check(AssetStateSetId.IsValidId());
+								const int32 AssetStateInt = AssetStateSetId.AsInteger();
+								FloatVector8.W = *(const float*)&AssetStateInt;
+									
+								UploadObjectData[9] = FloatVector8;
 							}
-						},
-						RangeCount == 1
-					);
+						}
+					},
+					RangeCount == 1
+				);
 
-					UploadDistanceFieldDataBuffer.ResourceUploadTo(GraphBuilder, DFObjectDataBuffer);
-					UploadDistanceFieldBoundsBuffer.ResourceUploadTo(GraphBuilder, DFObjectBoundsBuffer);
+				UploadDistanceFieldDataBuffer.ResourceUploadTo(GraphBuilder, DFObjectDataBuffer);
+				UploadDistanceFieldBoundsBuffer.ResourceUploadTo(GraphBuilder, DFObjectBoundsBuffer);
 
-					ExternalAccessQueue.Add(DFObjectDataBuffer, ERHIAccess::SRVMask, ERHIPipeline::All);
-					ExternalAccessQueue.Add(DFObjectBoundsBuffer, ERHIAccess::SRVMask, ERHIPipeline::All);
-				}
+				ExternalAccessQueue.Add(DFObjectDataBuffer, ERHIAccess::SRVMask, ERHIPipeline::All);
+				ExternalAccessQueue.Add(DFObjectBoundsBuffer, ERHIAccess::SRVMask, ERHIPipeline::All);
 			}
 		}
 
