LocalHeightFogVolume
- fixed bug with transform moving the height fog altitude.
- fixed crash when volume count is 0 after culling
- renamed g phase parameter

diff --git a/Engine/Source/Runtime/Engine/Classes/Components/LocalHeightFogComponent.h b/Engine/Source/Runtime/Engine/Classes/Components/LocalHeightFogComponent.h
--- a/Engine/Source/Runtime/Engine/Classes/Components/LocalHeightFogComponent.h
+++ b/Engine/Source/Runtime/Engine/Classes/Components/LocalHeightFogComponent.h
@@ -40,11 +40,11 @@ class ULocalHeightFogComponent : public USceneComponent
 	float FogHeightOffset = 0.0f;
 
 	/** Controls how strong the radial attenuation of this fog volume is. */
-	UPROPERTY(EditAnywhere, BlueprintReadOnly, interp, Category = "Fog Distribution", meta = (UIMin = "0.0", UIMax = "4.0", SliderExponent = 3.0, ClampMin = 0.001))
+	UPROPERTY(EditAnywhere, BlueprintReadOnly, interp, Category = "Fog Distribution", meta = (UIMin = "0.0", UIMax = "4.0", SliderExponent = 3.0, ClampMin = 0.0))
 	float FogRadialAttenuation = 0.0f;
 
-	/** Controls the directionality of the scattering within this fog volume. */
-	UPROPERTY(EditAnywhere, BlueprintReadOnly, interp, Category = "Fog Shading", meta = (UIMin = "0.0", UIMax = "0.999", ClampMin = 0.0, ClampMax = 0.999))
+	/** Controls the phase `G` parameter, describing the directionality of the scattering within this fog volume. */
+	UPROPERTY(EditAnywhere, BlueprintReadOnly, interp, Category = "Fog Shading", meta = (DisplayName = "Scattering Distribution", UIMin = "0.0", UIMax = "0.999", ClampMin = 0.0, ClampMax = 0.999))
 	float FogPhaseG = 0.8f;
 
 	/** Controls the albedo of this fog volume. */
diff --git a/Engine/Source/Runtime/Engine/Private/Components/LocalHeightFogComponent.cpp b/Engine/Source/Runtime/Engine/Private/Components/LocalHeightFogComponent.cpp
--- a/Engine/Source/Runtime/Engine/Private/Components/LocalHeightFogComponent.cpp
+++ b/Engine/Source/Runtime/Engine/Private/Components/LocalHeightFogComponent.cpp
@@ -40,8 +40,8 @@ void ULocalHeightFogComponent::SendRenderTransformCommand()
 		ENQUEUE_RENDER_COMMAND(FUpdateLocalHeightFogSceneProxyTransformCommand)(
 			[SceneProxy, ComponentTransform, HeightOffset](FRHICommandList& RHICmdList)
 			{
+				// Nothing else is needed so that command could actually go.
 				SceneProxy->FogTransform = ComponentTransform;
-				SceneProxy->FogHeightOffset = HeightOffset + ComponentTransform.GetLocation().Z;
 			});
 	}
 }
diff --git a/Engine/Source/Runtime/Engine/Private/LocalHeightFogSceneProxy.cpp b/Engine/Source/Runtime/Engine/Private/LocalHeightFogSceneProxy.cpp
--- a/Engine/Source/Runtime/Engine/Private/LocalHeightFogSceneProxy.cpp
+++ b/Engine/Source/Runtime/Engine/Private/LocalHeightFogSceneProxy.cpp
@@ -13,7 +13,7 @@ FLocalHeightFogSceneProxy::FLocalHeightFogSceneProxy(const ULocalHeightFogCompon
 	: FogTransform(InComponent->GetComponentTransform())
 	, FogDensity(InComponent->FogDensity)
 	, FogHeightFalloff(InComponent->FogHeightFalloff)
-	, FogHeightOffset(InComponent->FogHeightOffset + InComponent->GetComponentTransform().GetLocation().Z)
+	, FogHeightOffset(InComponent->FogHeightOffset)
 	, FogRadialAttenuation(InComponent->FogRadialAttenuation)
 	, FogMode((uint8)InComponent->FogMode)
 	, FogPhaseG(InComponent->FogPhaseG)
diff --git a/Engine/Source/Runtime/Renderer/Private/LocalHeightFogRendering.cpp b/Engine/Source/Runtime/Renderer/Private/LocalHeightFogRendering.cpp
--- a/Engine/Source/Runtime/Renderer/Private/LocalHeightFogRendering.cpp
+++ b/Engine/Source/Runtime/Renderer/Private/LocalHeightFogRendering.cpp
@@ -180,74 +180,79 @@ void RenderLocalHeightFog(
 			LocalHeightFogGPUInstanceDataIt++;
 			LocalHeightFogInstanceCountFinal++;
 		}
-		const uint32 AllLocalHeightFogInstanceBytesFinal = sizeof(FLocalHeightFogGPUInstanceData) * LocalHeightFogInstanceCountFinal;
-		FRDGBufferRef LocalHeightFogGPUInstanceDataBuffer =  CreateStructuredBuffer(GraphBuilder, TEXT("LocalHeightFogGPUInstanceDataBuffer"), 
-			sizeof(FLocalHeightFogGPUInstanceData), LocalHeightFogInstanceCountFinal, LocalHeightFogGPUInstanceData, AllLocalHeightFogInstanceBytesFinal, ERDGInitialDataFlags::NoCopy);
-		FRDGBufferSRVRef LocalHeightFogGPUInstanceDataBufferSRV = GraphBuilder.CreateSRV(LocalHeightFogGPUInstanceDataBuffer);
 
-		FRDGTextureRef SceneColorTexture = SceneTextures.Color.Resolve;
-
-		for (FViewInfo& View : Views)
+		if (LocalHeightFogInstanceCountFinal > 0)
 		{
-			FLocalHeightFogPassParameters* PassParameters = GraphBuilder.AllocParameters<FLocalHeightFogPassParameters>();
+			const uint32 AllLocalHeightFogInstanceBytesFinal = sizeof(FLocalHeightFogGPUInstanceData) * LocalHeightFogInstanceCountFinal;
+			FRDGBufferRef LocalHeightFogGPUInstanceDataBuffer = CreateStructuredBuffer(GraphBuilder, TEXT("LocalHeightFogGPUInstanceDataBuffer"),
+				sizeof(FLocalHeightFogGPUInstanceData), LocalHeightFogInstanceCountFinal, LocalHeightFogGPUInstanceData, AllLocalHeightFogInstanceBytesFinal, ERDGInitialDataFlags::NoCopy);
+
+			FRDGBufferSRVRef LocalHeightFogGPUInstanceDataBufferSRV = GraphBuilder.CreateSRV(LocalHeightFogGPUInstanceDataBuffer);
 
-			PassParameters->VS.View = GetShaderBinding(View.ViewUniformBuffer);
-			PassParameters->VS.LocalHeightFogInstances = LocalHeightFogGPUInstanceDataBufferSRV;
+			FRDGTextureRef SceneColorTexture = SceneTextures.Color.Resolve;
 
-			PassParameters->PS.View = GetShaderBinding(View.ViewUniformBuffer);
-			PassParameters->PS.LocalHeightFogInstances = LocalHeightFogGPUInstanceDataBufferSRV;
+			for (FViewInfo& View : Views)
+			{
+				FLocalHeightFogPassParameters* PassParameters = GraphBuilder.AllocParameters<FLocalHeightFogPassParameters>();
 
-			PassParameters->SceneTextures = SceneTextures.UniformBuffer;
-			PassParameters->RenderTargets[0] = FRenderTargetBinding(SceneColorTexture, ERenderTargetLoadAction::ENoAction);
+				PassParameters->VS.View = GetShaderBinding(View.ViewUniformBuffer);
+				PassParameters->VS.LocalHeightFogInstances = LocalHeightFogGPUInstanceDataBufferSRV;
 
-			FLocalHeightFogVS::FPermutationDomain VSPermutationVector;
-			auto VertexShader = View.ShaderMap->GetShader< FLocalHeightFogVS >(VSPermutationVector);
+				PassParameters->PS.View = GetShaderBinding(View.ViewUniformBuffer);
+				PassParameters->PS.LocalHeightFogInstances = LocalHeightFogGPUInstanceDataBufferSRV;
 
-			FLocalHeightFogPS::FPermutationDomain PsPermutationVector;
-			auto PixelShader = View.ShaderMap->GetShader< FLocalHeightFogPS >(PsPermutationVector);
+				PassParameters->SceneTextures = SceneTextures.UniformBuffer;
+				PassParameters->RenderTargets[0] = FRenderTargetBinding(SceneColorTexture, ERenderTargetLoadAction::ENoAction);
 
-			const FIntRect ViewRect = View.ViewRect;
+				FLocalHeightFogVS::FPermutationDomain VSPermutationVector;
+				auto VertexShader = View.ShaderMap->GetShader< FLocalHeightFogVS >(VSPermutationVector);
 
-			ClearUnusedGraphResources(VertexShader, &PassParameters->VS);
-			ClearUnusedGraphResources(PixelShader, &PassParameters->PS);
+				FLocalHeightFogPS::FPermutationDomain PsPermutationVector;
+				auto PixelShader = View.ShaderMap->GetShader< FLocalHeightFogPS >(PsPermutationVector);
 
-			GraphBuilder.AddPass(
-				RDG_EVENT_NAME("RenderLocalHeightFog %u inst.", LocalHeightFogInstanceCount),
-				PassParameters,
-				ERDGPassFlags::Raster,
-				[VertexShader, PixelShader, PassParameters, LocalHeightFogInstanceCount, ViewRect](FRHICommandList& RHICmdList)
-			{
-				FGraphicsPipelineStateInitializer GraphicsPSOInit;
-				RHICmdList.ApplyCachedRenderTargets(GraphicsPSOInit);
+				const FIntRect ViewRect = View.ViewRect;
 
-				RHICmdList.SetViewport(ViewRect.Min.X, ViewRect.Min.Y, 0.0f, ViewRect.Max.X, ViewRect.Max.Y, 1.0f);
+				ClearUnusedGraphResources(VertexShader, &PassParameters->VS);
+				ClearUnusedGraphResources(PixelShader, &PassParameters->PS);
 
-				// Render back faces only since camera may intersect
-				GraphicsPSOInit.RasterizerState = TStaticRasterizerState<FM_Solid, CM_CW>::GetRHI();
-				GraphicsPSOInit.DepthStencilState = TStaticDepthStencilState<false, CF_Always>::GetRHI();
-				GraphicsPSOInit.BlendState = TStaticBlendState<CW_RGBA, BO_Add, BF_One, BF_SourceAlpha, BO_Add, BF_Zero, BF_SourceAlpha>::GetRHI();
-				GraphicsPSOInit.PrimitiveType = PT_TriangleList;
+				GraphBuilder.AddPass(
+					RDG_EVENT_NAME("RenderLocalHeightFog %u inst.", LocalHeightFogInstanceCountFinal),
+					PassParameters,
+					ERDGPassFlags::Raster,
+					[VertexShader, PixelShader, PassParameters, LocalHeightFogInstanceCountFinal, ViewRect](FRHICommandList& RHICmdList)
+				{
+						FGraphicsPipelineStateInitializer GraphicsPSOInit;
+					RHICmdList.ApplyCachedRenderTargets(GraphicsPSOInit);
 
-				GraphicsPSOInit.BoundShaderState.VertexDeclarationRHI = GetVertexDeclarationFVector4();
-				GraphicsPSOInit.BoundShaderState.VertexShaderRHI = VertexShader.GetVertexShader();
-				GraphicsPSOInit.BoundShaderState.PixelShaderRHI = PixelShader.GetPixelShader();
+					RHICmdList.SetViewport(ViewRect.Min.X, ViewRect.Min.Y, 0.0f, ViewRect.Max.X, ViewRect.Max.Y, 1.0f);
 
-				SetGraphicsPipelineState(RHICmdList, GraphicsPSOInit, 0);
+					// Render back faces only since camera may intersect
+					GraphicsPSOInit.RasterizerState = TStaticRasterizerState<FM_Solid, CM_CW>::GetRHI();
+					GraphicsPSOInit.DepthStencilState = TStaticDepthStencilState<false, CF_Always>::GetRHI();
+					GraphicsPSOInit.BlendState = TStaticBlendState<CW_RGBA, BO_Add, BF_One, BF_SourceAlpha, BO_Add, BF_Zero, BF_SourceAlpha>::GetRHI();
+					GraphicsPSOInit.PrimitiveType = PT_TriangleList;
 
-				SetShaderParameters(RHICmdList, VertexShader, VertexShader.GetVertexShader(), PassParameters->VS);
-				SetShaderParameters(RHICmdList, PixelShader, PixelShader.GetPixelShader(), PassParameters->PS);
+					GraphicsPSOInit.BoundShaderState.VertexDeclarationRHI = GetVertexDeclarationFVector4();
+					GraphicsPSOInit.BoundShaderState.VertexShaderRHI = VertexShader.GetVertexShader();
+					GraphicsPSOInit.BoundShaderState.PixelShaderRHI = PixelShader.GetPixelShader();
 
-				RHICmdList.SetStreamSource(0, GetUnitCubeVertexBuffer(), 0);
+					SetGraphicsPipelineState(RHICmdList, GraphicsPSOInit, 0);
 
-				RHICmdList.DrawIndexedPrimitive( GetUnitCubeIndexBuffer()
-					, 0									//BaseVertexIndex
-					, 0									//FirstInstance
-					, 8									//uint32 NumVertices
-					, 0									//uint32 StartIndex
-					, UE_ARRAY_COUNT(GCubeIndices) / 3	//uint32 NumPrimitives
-					, LocalHeightFogInstanceCount		//uint32 NumInstances
-				);
-			});
+					SetShaderParameters(RHICmdList, VertexShader, VertexShader.GetVertexShader(), PassParameters->VS);
+					SetShaderParameters(RHICmdList, PixelShader, PixelShader.GetPixelShader(), PassParameters->PS);
+
+					RHICmdList.SetStreamSource(0, GetUnitCubeVertexBuffer(), 0);
+
+					RHICmdList.DrawIndexedPrimitive(GetUnitCubeIndexBuffer()
+						, 0									//BaseVertexIndex
+						, 0									//FirstInstance
+						, 8									//uint32 NumVertices
+						, 0									//uint32 StartIndex
+						, UE_ARRAY_COUNT(GCubeIndices) / 3	//uint32 NumPrimitives
+						, LocalHeightFogInstanceCountFinal	//uint32 NumInstances
+					);
+				});
+			}
 		}
 	}
 }
