Local height fog volume - added sorting of each volumes + a way to override the sorting via an integer key.

diff --git a/Engine/Source/Runtime/Engine/Classes/Components/LocalHeightFogComponent.h b/Engine/Source/Runtime/Engine/Classes/Components/LocalHeightFogComponent.h
--- a/Engine/Source/Runtime/Engine/Classes/Components/LocalHeightFogComponent.h
+++ b/Engine/Source/Runtime/Engine/Classes/Components/LocalHeightFogComponent.h
@@ -55,6 +55,10 @@ class ULocalHeightFogComponent : public USceneComponent
 	UPROPERTY(EditAnywhere, BlueprintReadOnly, interp, Category = "Fog Shading", meta = (ClampMin = 0.0, ClampMax = 1.0))
 	FLinearColor FogEmissive = FLinearColor::Black;
 
+	/** The priority can be used as a way to overide the sorting by distance. A higher value means the volume will be considered further away, i.e. behing the one with a lower priority value. */
+	UPROPERTY(EditAnywhere, BlueprintReadOnly, Category = "Sorting", meta = (UIMin = "0", UIMax = "255", ClampMin = 0, ClampMax = 255))
+	uint8 FogSortPriority = 127;
+
 public:
 
 	//~ Begin UObject Interface
diff --git a/Engine/Source/Runtime/Engine/Private/LocalHeightFogSceneProxy.cpp b/Engine/Source/Runtime/Engine/Private/LocalHeightFogSceneProxy.cpp
--- a/Engine/Source/Runtime/Engine/Private/LocalHeightFogSceneProxy.cpp
+++ b/Engine/Source/Runtime/Engine/Private/LocalHeightFogSceneProxy.cpp
@@ -16,6 +16,7 @@ FLocalHeightFogSceneProxy::FLocalHeightFogSceneProxy(const ULocalHeightFogCompon
 	, FogHeightOffset(InComponent->FogHeightOffset)
 	, FogRadialAttenuation(InComponent->FogRadialAttenuation)
 	, FogMode((uint8)InComponent->FogMode)
+	, FogSortPriority(InComponent->FogSortPriority)
 	, FogPhaseG(InComponent->FogPhaseG)
 	, FogAlbedo(InComponent->FogAlbedo)
 	, FogEmissive(InComponent->FogEmissive)
diff --git a/Engine/Source/Runtime/Engine/Public/LocalHeightFogSceneProxy.h b/Engine/Source/Runtime/Engine/Public/LocalHeightFogSceneProxy.h
--- a/Engine/Source/Runtime/Engine/Public/LocalHeightFogSceneProxy.h
+++ b/Engine/Source/Runtime/Engine/Public/LocalHeightFogSceneProxy.h
@@ -24,6 +24,7 @@ public:
 	float FogRadialAttenuation;
 
 	uint8 FogMode;
+	uint8 FogSortPriority;
 
 	float FogPhaseG;
 	FLinearColor FogAlbedo;
diff --git a/Engine/Source/Runtime/Renderer/Private/LocalHeightFogRendering.cpp b/Engine/Source/Runtime/Renderer/Private/LocalHeightFogRendering.cpp
--- a/Engine/Source/Runtime/Renderer/Private/LocalHeightFogRendering.cpp
+++ b/Engine/Source/Runtime/Renderer/Private/LocalHeightFogRendering.cpp
@@ -134,6 +134,32 @@ BEGIN_SHADER_PARAMETER_STRUCT(FLocalHeightFogPassParameters, )
 	RENDER_TARGET_BINDING_SLOTS()
 END_SHADER_PARAMETER_STRUCT()
 
+class RENDERER_API FLocalFogVolumeSortKey
+{
+public:
+	union
+	{
+		uint64 PackedData;
+
+		struct
+		{
+			uint64 Index	: 16; // Index of the volume
+			uint64 Distance : 32; // then by distance
+			uint64 Priority : 16; // First order by priority
+		} FogVolume;
+	};
+
+	FORCEINLINE bool operator!=(FLocalFogVolumeSortKey B) const
+	{
+		return PackedData != B.PackedData;
+	}
+
+	FORCEINLINE bool operator<(FLocalFogVolumeSortKey B) const
+	{
+		return PackedData > B.PackedData;
+	}
+};
+
 void RenderLocalHeightFog(
 	const FScene* Scene,
 	TArray<FViewInfo>& Views,
@@ -148,9 +174,10 @@ void RenderLocalHeightFog(
 		RDG_GPU_STAT_SCOPE(GraphBuilder, LocalHeightFogVolumes);
 		
 		// No culling as of today
-		const uint32 AllLocalHeightFogInstanceBytes = sizeof(FLocalHeightFogGPUInstanceData) * LocalHeightFogInstanceCount;
-		FLocalHeightFogGPUInstanceData* LocalHeightFogGPUInstanceData = (FLocalHeightFogGPUInstanceData*) GraphBuilder.Alloc(AllLocalHeightFogInstanceBytes, 16);
-		FLocalHeightFogGPUInstanceData* LocalHeightFogGPUInstanceDataIt = LocalHeightFogGPUInstanceData;
+		FLocalHeightFogGPUInstanceData* LocalHeightFogGPUInstanceData	= (FLocalHeightFogGPUInstanceData*) GraphBuilder.Alloc(sizeof(FLocalHeightFogGPUInstanceData) * LocalHeightFogInstanceCount, 16);
+		FVector*						LocalHeightFogCenterPos			= (FVector*) GraphBuilder.Alloc(sizeof(FVector3f) * LocalHeightFogInstanceCount, 16);
+		TArray<FLocalFogVolumeSortKey>	LocalHeightFogSortKeys;
+		LocalHeightFogSortKeys.SetNumUninitialized(LocalHeightFogInstanceCount);
 		for (FLocalHeightFogSceneProxy* LHF : Scene->LocalHeightFogs)
 		{
 			if (LHF->FogDensity <= 0.0f)
@@ -161,6 +188,7 @@ void RenderLocalHeightFog(
 			FTransform TransformScaleOnly;
 			TransformScaleOnly.SetScale3D(LHF->FogTransform.GetScale3D());
 
+			FLocalHeightFogGPUInstanceData* LocalHeightFogGPUInstanceDataIt = &LocalHeightFogGPUInstanceData[LocalHeightFogInstanceCountFinal];
 			LocalHeightFogGPUInstanceDataIt->Transform					= FMatrix44f(LHF->FogTransform.ToMatrixWithScale());
 			LocalHeightFogGPUInstanceDataIt->InvTransform				= LocalHeightFogGPUInstanceDataIt->Transform.Inverse();
 			LocalHeightFogGPUInstanceDataIt->InvTranformNoScale			= FMatrix44f(LHF->FogTransform.ToMatrixNoScale()).Inverse();
@@ -177,22 +205,51 @@ void RenderLocalHeightFog(
 			LocalHeightFogGPUInstanceDataIt->PhaseG = LHF->FogPhaseG;
 			LocalHeightFogGPUInstanceDataIt->Emissive = FVector3f(LHF->FogEmissive);
 
-			LocalHeightFogGPUInstanceDataIt++;
+			LocalHeightFogCenterPos[LocalHeightFogInstanceCountFinal] = LHF->FogTransform.GetTranslation();
+
+			FLocalFogVolumeSortKey* LocalHeightFogSortKeysIt = &LocalHeightFogSortKeys[LocalHeightFogInstanceCountFinal];
+			LocalHeightFogSortKeysIt->FogVolume.Index = LocalHeightFogInstanceCountFinal;
+			LocalHeightFogSortKeysIt->FogVolume.Distance = 0;	// Filled up right before sorting according to a view
+			LocalHeightFogSortKeysIt->FogVolume.Priority = LHF->FogSortPriority;
+
 			LocalHeightFogInstanceCountFinal++;
 		}
+		// Shrink the array to only what is needed in order for the sort to correctly work on only what is needed.
+		LocalHeightFogSortKeys.SetNum(LocalHeightFogInstanceCountFinal, false/*bAllowShrinking*/);
 
 		if (LocalHeightFogInstanceCountFinal > 0)
 		{
-			const uint32 AllLocalHeightFogInstanceBytesFinal = sizeof(FLocalHeightFogGPUInstanceData) * LocalHeightFogInstanceCountFinal;
-			FRDGBufferRef LocalHeightFogGPUInstanceDataBuffer = CreateStructuredBuffer(GraphBuilder, TEXT("LocalHeightFogGPUInstanceDataBuffer"),
-				sizeof(FLocalHeightFogGPUInstanceData), LocalHeightFogInstanceCountFinal, LocalHeightFogGPUInstanceData, AllLocalHeightFogInstanceBytesFinal, ERDGInitialDataFlags::NoCopy);
-
-			FRDGBufferSRVRef LocalHeightFogGPUInstanceDataBufferSRV = GraphBuilder.CreateSRV(LocalHeightFogGPUInstanceDataBuffer);
 
 			FRDGTextureRef SceneColorTexture = SceneTextures.Color.Resolve;
 
 			for (FViewInfo& View : Views)
 			{
+				// 1. Sort all the volumes
+				const FVector ViewOrigin = View.ViewMatrices.GetViewOrigin();
+				for (uint32 i = 0; i < LocalHeightFogInstanceCountFinal; ++i)
+				{
+					FVector FogCenterPos = LocalHeightFogCenterPos[LocalHeightFogSortKeys[i].FogVolume.Index];	// Recovered form the original array via index because the sorting of the previous view might have changed the order.
+					float DistancetoView = float((FogCenterPos - ViewOrigin).Size());
+					LocalHeightFogSortKeys[i].FogVolume.Distance = *reinterpret_cast<uint32*>(&DistancetoView);
+				}
+				LocalHeightFogSortKeys.Sort();
+
+				// 2. Create the buffer containing all the fog volume data instance sorted according to their key for the current view.
+				FLocalHeightFogGPUInstanceData* LocalHeightFogGPUSortedInstanceData = (FLocalHeightFogGPUInstanceData*)GraphBuilder.Alloc(sizeof(FLocalHeightFogGPUInstanceData) * LocalHeightFogInstanceCountFinal, 16);
+				for (uint32 i = 0; i < LocalHeightFogInstanceCountFinal; ++i)
+				{
+					// We could also have an indirection buffer on GPU but choosing to go with the sorting + copy on CPU since it is expected to not have many local height fog volumes.
+					LocalHeightFogGPUSortedInstanceData[i] = LocalHeightFogGPUInstanceData[LocalHeightFogSortKeys[i].FogVolume.Index];
+				}
+
+				// 3. Allocate buffer and initialize with sorted data to upload to GPU
+				const uint32 AllLocalHeightFogInstanceBytesFinal = sizeof(FLocalHeightFogGPUInstanceData) * LocalHeightFogInstanceCountFinal;
+				FRDGBufferRef LocalHeightFogGPUInstanceDataBuffer = CreateStructuredBuffer(GraphBuilder, TEXT("LocalHeightFogGPUInstanceDataBuffer"),
+					sizeof(FLocalHeightFogGPUInstanceData), LocalHeightFogInstanceCountFinal, LocalHeightFogGPUSortedInstanceData, AllLocalHeightFogInstanceBytesFinal, ERDGInitialDataFlags::NoCopy);
+
+				FRDGBufferSRVRef LocalHeightFogGPUInstanceDataBufferSRV = GraphBuilder.CreateSRV(LocalHeightFogGPUInstanceDataBuffer);
+
+				// 4. Render
 				FLocalHeightFogPassParameters* PassParameters = GraphBuilder.AllocParameters<FLocalHeightFogPassParameters>();
 
 				PassParameters->VS.View = GetShaderBinding(View.ViewUniformBuffer);
@@ -221,7 +278,7 @@ void RenderLocalHeightFog(
 					ERDGPassFlags::Raster,
 					[VertexShader, PixelShader, PassParameters, LocalHeightFogInstanceCountFinal, ViewRect](FRHICommandList& RHICmdList)
 				{
-						FGraphicsPipelineStateInitializer GraphicsPSOInit;
+					FGraphicsPipelineStateInitializer GraphicsPSOInit;
 					RHICmdList.ApplyCachedRenderTargets(GraphicsPSOInit);
 
 					RHICmdList.SetViewport(ViewRect.Min.X, ViewRect.Min.Y, 0.0f, ViewRect.Max.X, ViewRect.Max.Y, 1.0f);
