Always log RHI support for VRS

diff --git a/Engine/Source/Runtime/Renderer/Private/VariableRateShading/VariableRateShadingImageManager.cpp b/Engine/Source/Runtime/Renderer/Private/VariableRateShading/VariableRateShadingImageManager.cpp
--- a/Engine/Source/Runtime/Renderer/Private/VariableRateShading/VariableRateShadingImageManager.cpp
+++ b/Engine/Source/Runtime/Renderer/Private/VariableRateShading/VariableRateShadingImageManager.cpp
@@ -252,9 +252,13 @@ FVariableRateShadingImageManager::~FVariableRateShadingImageManager() {}
 
 void FVariableRateShadingImageManager::InitRHI()
 {
-	if (!IsVRSSupportedByRHI())
+	if (IsVRSSupportedByRHI())
 	{
-		UE_LOG(LogVRS, Warning, TEXT("VRS is not supported by current RHI"));
+		UE_LOG(LogVRS, Log, TEXT("Current RHI supports Variable Rate Shading"));
+	}
+	else
+	{
+		UE_LOG(LogVRS, Log, TEXT("Current RHI does not support Variable Rate Shading"));
 	}
 }
 
