Optimization to ComputeRelevancy to populate only selected instance HitProxyIDs
* Previous path repeatedly appended a duplicate PrimitiveSceneInfo.NaniteHitProxyIds for each instance and relied on deduplication later to get unique IDs.
* Change fixes this by only adding hit proxy IDs if they are selected to the output array instead of appending all hit proxy IDs.
* Does not change deduplication in FVisibilityTaskData::SetupMeshPasses
* This change fixes a large performance problem when selecting objects in the editor with a lot of instances.  Hillside sample project can OOM and lock editor for minutes to hours.

diff --git a/Engine/Source/Runtime/Renderer/Private/SceneVisibility.cpp b/Engine/Source/Runtime/Renderer/Private/SceneVisibility.cpp
--- a/Engine/Source/Runtime/Renderer/Private/SceneVisibility.cpp
+++ b/Engine/Source/Runtime/Renderer/Private/SceneVisibility.cpp
@@ -1643,7 +1643,7 @@ void FRelevancePacket::ComputeRelevance(FDynamicPrimitiveIndexList& DynamicPrimi
 		auto CollectSelectedNaniteInstanceDraws = [](
 			const FPrimitiveSceneInfo& PrimitiveSceneInfo,
 			TArray<Nanite::FInstanceDraw>& OutInstanceDraws,
-			TArray<uint32>* OutHitProxyIDs,
+			TArray<uint32>* OutSelectedInstanceHitProxyIDs,
 			bool bSelectedInstancesOnly
 		)
 		{
@@ -1678,7 +1678,13 @@ void FRelevancePacket::ComputeRelevance(FDynamicPrimitiveIndexList& DynamicPrimi
 							if (!bSelected)
 							{
 								continue;
-								}
+							}
+							
+							if (OutSelectedInstanceHitProxyIDs != nullptr)
+							{
+								const uint32 HitProxyID = HitProxyColor.ToPackedABGR();
+								OutSelectedInstanceHitProxyIDs->Add(HitProxyID);
+							}
 						}
 					}
 				}
@@ -1689,12 +1695,9 @@ void FRelevancePacket::ComputeRelevance(FDynamicPrimitiveIndexList& DynamicPrimi
 						0u
 					}
 				);
-
-				if (OutHitProxyIDs != nullptr)
-				{
-					OutHitProxyIDs->Append(PrimitiveSceneInfo.NaniteHitProxyIds);
-				}
 			}
+
+
 		};
 
 		if (bEditorVisualizeLevelInstanceRelevance)
