Enable FoveationLevel for emulated stereo and rename foveation-related CVars with "xr" prefix and no "HMD" relation

diff --git a/Engine/Source/Runtime/Engine/Classes/Engine/RendererSettings.h b/Engine/Source/Runtime/Engine/Classes/Engine/RendererSettings.h
--- a/Engine/Source/Runtime/Engine/Classes/Engine/RendererSettings.h
+++ b/Engine/Source/Runtime/Engine/Classes/Engine/RendererSettings.h
@@ -603,14 +603,14 @@ class ENGINE_API URendererSettings : public UDeveloperSettings
 	FVector TranslucentSortAxis;
 
 	UPROPERTY(config, EditAnywhere, Category = VR, meta = (
-		ConsoleVariable = "r.VRS.HMDFoveationLevel", DisplayName = "HMD Foveation Level (Experimental)",
+		ConsoleVariable = "xr.VRS.FoveationLevel", DisplayName = "Stereo Foveation Level (Experimental)",
 		ToolTip = "Set the level of foveation to apply when generating the Variable Rate Shading attachment. This feature is currently experimental.\nThis can yield some fairly significant performance benefits on GPUs that support Tier 2 VRS.\nLower settings will result in almost no discernible artifacting on most HMDs; higher settings will show some artifacts towards the edges of the view."))
-	TEnumAsByte<EFixedFoveationLevels::Type> HMDFoveationLevel;
+	TEnumAsByte<EFixedFoveationLevels::Type> FoveationLevel;
 
 	UPROPERTY(config, EditAnywhere, Category = VR, meta = (
-		ConsoleVariable = "r.VRS.HMDDynamicFoveation", DisplayName = "Dynamic Foveation (Experimental)",
+		ConsoleVariable = "xr.VRS.DynamicFoveation", DisplayName = "Dynamic Foveation (Experimental)",
 		ToolTip = "Allows foveation level to adjust dynamically based on GPU utilization.\nLevel will range between none at the minimum, and the currently selected foveation level at the maximum."))
-	uint32 bHMDDynamicFoveation:1;
+	uint32 bDynamicFoveation:1;
 
 	UPROPERTY(config, EditAnywhere, Category=Postprocessing, meta=(
 		ConsoleVariable="r.CustomDepth",DisplayName="Custom Depth-Stencil Pass",
diff --git a/Engine/Source/Runtime/Renderer/Private/VariableRateShading/FoveatedImageGenerator.cpp b/Engine/Source/Runtime/Renderer/Private/VariableRateShading/FoveatedImageGenerator.cpp
--- a/Engine/Source/Runtime/Renderer/Private/VariableRateShading/FoveatedImageGenerator.cpp
+++ b/Engine/Source/Runtime/Renderer/Private/VariableRateShading/FoveatedImageGenerator.cpp
@@ -16,7 +16,7 @@
 /* CVar values used to control generator behavior */
 
 static TAutoConsoleVariable<int> CVarFoveationLevel(
-	TEXT("r.VRS.HMDFoveationLevel"),
+	TEXT("xr.VRS.FoveationLevel"),
 	0,
 	TEXT("Level of foveated VRS to apply (when Variable Rate Shading is available)\n")
 	TEXT(" 0: Disabled (default);\n")
@@ -27,7 +27,7 @@ static TAutoConsoleVariable<int> CVarFoveationLevel(
 	ECVF_RenderThreadSafe);
 
 static TAutoConsoleVariable<int32> CVarDynamicFoveation(
-	TEXT("r.VRS.HMDDynamicFoveation"),
+	TEXT("xr.VRS.DynamicFoveation"),
 	0,
 	TEXT("Whether foveation level should adjust based on GPU utilization\n")
 	TEXT(" 0: Disabled (default);\n")
@@ -35,7 +35,7 @@ static TAutoConsoleVariable<int32> CVarDynamicFoveation(
 	ECVF_RenderThreadSafe);
 
 static TAutoConsoleVariable<int32> CVarFoveationPreview(
-	TEXT("r.VRS.HMDFoveationPreview"),
+	TEXT("xr.VRS.FoveationPreview"),
 	1,
 	TEXT("Include foveated VRS in the VRS debug overlay.")
 	TEXT(" 0: Disabled;\n")
@@ -43,7 +43,7 @@ static TAutoConsoleVariable<int32> CVarFoveationPreview(
 	ECVF_RenderThreadSafe);
 
 static TAutoConsoleVariable<int32> CVarGazeTrackedFoveation(
-	TEXT("r.VRS.HMDGazeTrackedFoveation"),
+	TEXT("xr.VRS.GazeTrackedFoveation"),
 	0,
 	TEXT("Enable gaze-tracking for foveated VRS\n")
 	TEXT(" 0: Disabled (default);\n")
@@ -212,7 +212,7 @@ void FFoveatedImageGenerator::PrepareImages(FRDGBuilder& GraphBuilder, const FSc
 	EVRSGenerationFlags GenFlags = EVRSGenerationFlags::None;
 
 	// If stereo is enabled, side-by-side is assumed for now, may need to change this if we add support for mobile multi-view
-	if (IStereoRendering::IsStereoEyeView(*ViewFamily.Views[0]) && GEngine->XRSystem.IsValid())
+	if (IStereoRendering::IsStereoEyeView(*ViewFamily.Views[0]))
 	{
 		EnumAddFlags(GenFlags, EVRSGenerationFlags::StereoRendering);
 		EnumAddFlags(GenFlags, EVRSGenerationFlags::SideBySideStereo); 
@@ -243,9 +243,8 @@ void FFoveatedImageGenerator::PrepareImages(FRDGBuilder& GraphBuilder, const FSc
 
 bool FFoveatedImageGenerator::IsEnabledForView(const FSceneView& View) const
 {
-	// Only enabled if we are in XR
-	bool bStereoRendering = IStereoRendering::IsStereoEyeView(View) && GEngine->XRSystem.IsValid();
-	return bStereoRendering && CVarFoveationLevel.GetValueOnRenderThread() > 0;
+	// Enabled for stereo (XR or emulated)
+	return IStereoRendering::IsStereoEyeView(View) && CVarFoveationLevel.GetValueOnRenderThread() > 0;
 }
 
 FVariableRateShadingImageManager::EVRSSourceType FFoveatedImageGenerator::GetType() const
