Initial support for injecting lights with raytraced shadows into volumetric fog:
* Add RGS pass to inject lights and trace shadow rays.
* Injecting lights is done in a common function (InjectShadowedLocalLightCommon) that is used by RGS and PS
* Move shared parameters into FInjectShadowedLocalLightCommonParameters
* Controlled with r.VolumetricFog.InjectRaytracedLights (disabled by default)

diff --git a/Engine/Shaders/Private/VolumetricFog.usf b/Engine/Shaders/Private/VolumetricFog.usf
--- a/Engine/Shaders/Private/VolumetricFog.usf
+++ b/Engine/Shaders/Private/VolumetricFog.usf
@@ -32,6 +32,10 @@
 #define USE_EMISSIVE 1
 #endif
 
+#ifndef USE_RAYTRACED_SHADOWS
+#define USE_RAYTRACED_SHADOWS 0
+#endif
+
 RWTexture3D<float4> RWVBufferA;
 #if USE_EMISSIVE
 RWTexture3D<float4> RWVBufferB;
@@ -274,25 +278,42 @@ float2 PrevConservativeDepthTextureSize;
 Texture2D<float> PrevConservativeDepthTexture;
 uint UseConservativeDepthTexture;
 
-void InjectShadowedLocalLightPS(
-	FWriteToSliceGeometryOutput Input,
-	out float4 OutScattering : SV_Target0
-	)
-{
-	float4 SvPosition = Input.Vertex.Position;
-	OutScattering = 0;
+#if USE_RAYTRACED_SHADOWS
 
-	if (UseConservativeDepthTexture > 0)
+#include "/Engine/Private/RayTracing/RayTracingCommon.ush"
+
+RaytracingAccelerationStructure TLAS;
+
+float ComputeRayTracedShadowFactor(float3 TranslatedWorldPosition, float3 ToLight)
+{
+	uint RayFlags = RAY_FLAG_SKIP_CLOSEST_HIT_SHADER | RAY_FLAG_ACCEPT_FIRST_HIT_AND_END_SEARCH;
+	uint RaytracingMask = RAY_TRACING_MASK_SHADOW | RAY_TRACING_MASK_THIN_SHADOW;
+
+	FRayDesc Ray;
+	Ray.Origin = TranslatedWorldPosition;
+	Ray.Direction = normalize(ToLight);
+	Ray.TMin = 0.0f;
+	Ray.TMax = length(ToLight);
+
+	FMinimalPayload MinimalPayload = TraceVisibilityRay(
+		TLAS,
+		RayFlags,
+		RaytracingMask,
+		Ray);
+
+	if (MinimalPayload.IsHit())
 	{
-		const float SceneDepth = ConservativeDepthTexture.Load(uint3(SvPosition.xy, 0));
-		if (SceneDepth > SvPosition.z)
-		{
-			clip(-1); // If the froxel is behind conservative front depth, skip the light injection.
-			return;
-		}
+		return 0.0f;
 	}
 
-	uint3 GridCoordinate = uint3(Input.Vertex.Position.xy, Input.LayerIndex);
+	return 1.0;
+}
+
+#endif
+
+float4 InjectShadowedLocalLightCommon(uint3 GridCoordinate)
+{
+	float4 OutScattering = 0;
 
 	// Somehow pixels are being rasterized outside of the viewport on a 970 GTX, perhaps due to use of a geometry shader bypassing the viewport scissor.
 	// This triggers the HistoryMissSuperSampleCount path causing significant overhead for shading off-screen pixels.
@@ -360,14 +381,18 @@ void InjectShadowedLocalLightPS(
 	
 #if ENABLE_SHADOW_COMPUTATION
 			if (CombinedAttenuation > 0)
-			{			
+			{
+#if USE_RAYTRACED_SHADOWS
+				ShadowFactor = ComputeRayTracedShadowFactor(TranslatedWorldPosition, ToLight);
+#else
 				bool bUnused = false;
 				ShadowFactor = ComputeVolumeShadowing(TranslatedWorldPosition, IsPointLight, LightData.bSpotLight, bUnused);
 
 #if VIRTUAL_SHADOW_MAP
 				FVirtualShadowMapSampleResult VirtualShadowMapSample = SampleVirtualShadowMapTranslatedWorld(VirtualShadowMapId, TranslatedWorldPosition);
 				ShadowFactor *= VirtualShadowMapSample.ShadowFactor;
-#endif
+#endif // VIRTUAL_SHADOW_MAP
+#endif // USE_RAYTRACED_SHADOWS
 			}
 #endif
 	
@@ -384,6 +409,61 @@ void InjectShadowedLocalLightPS(
 		// We pre-expose the buffer containing shadowed local lights luminance contribution. This helps maintaining details and color at high exposure for very bright scenes.
 		OutScattering.rgb *= View.PreExposure / (float)NumSuperSamples;
 	}
+
+	return OutScattering;
+}
+
+#if USE_RAYTRACED_SHADOWS
+
+RWTexture3D<float4> OutVolumeTexture;
+uint3 OutputOffset;
+int FirstSlice;
+
+int3 GetVoxelIndex(int ThreadIndex, int3 Resolution)
+{
+	int SliceSize = Resolution.x * Resolution.y;
+	int SliceIndex = ThreadIndex / SliceSize;
+	int SliceRemainder = ThreadIndex - SliceIndex * SliceSize;
+	
+	return int3(
+		SliceRemainder % Resolution.x,
+		SliceRemainder / Resolution.x,
+		SliceIndex
+	);
+}
+
+RAY_TRACING_ENTRY_RAYGEN(InjectShadowedLocalLightRGS)
+{
+	int3 Position = GetVoxelIndex(DispatchRaysIndex().x, VolumetricFog.ResourceGridSizeInt);
+	Position.z += FirstSlice;
+
+	float4 OutScattering = InjectShadowedLocalLightCommon(Position);
+
+	float4 DestValue = OutVolumeTexture[Position];
+	OutVolumeTexture[Position] = DestValue + OutScattering;
+}
+
+#endif // USE_RAYTRACED_SHADOWS
+
+void InjectShadowedLocalLightPS(
+	FWriteToSliceGeometryOutput Input,
+	out float4 OutScattering : SV_Target0
+	)
+{
+	float4 SvPosition = Input.Vertex.Position;
+	OutScattering = 0;
+
+	if (UseConservativeDepthTexture > 0)
+	{
+		const float SceneDepth = ConservativeDepthTexture.Load(uint3(SvPosition.xy, 0));
+		if (SceneDepth > SvPosition.z)
+		{
+			clip(-1); // If the froxel is behind conservative front depth, skip the light injection.
+			return;
+		}
+	}
+
+	OutScattering = InjectShadowedLocalLightCommon(uint3(SvPosition.xy, Input.LayerIndex));
 }
 
 Texture3D<float4> VBufferA;
diff --git a/Engine/Source/Runtime/Renderer/Private/DeferredShadingRenderer.cpp b/Engine/Source/Runtime/Renderer/Private/DeferredShadingRenderer.cpp
--- a/Engine/Source/Runtime/Renderer/Private/DeferredShadingRenderer.cpp
+++ b/Engine/Source/Runtime/Renderer/Private/DeferredShadingRenderer.cpp
@@ -1894,6 +1894,7 @@ bool FDeferredShadingSceneRenderer::SetupRayTracingPipelineStates(FRDGBuilder& G
 					PrepareRayTracingGlobalIllumination(View, RayGenShaders);
 					PrepareRayTracingGlobalIlluminationPlugin(View, RayGenShaders);
 					PrepareRayTracingTranslucency(View, RayGenShaders);
+					PrepareRayTracingVolumetricFogShadows(View, *Scene, RayGenShaders);
 
 					if (DoesPlatformSupportLumenGI(ShaderPlatform) && Lumen::UseHardwareRayTracing(ViewFamily))
 					{
diff --git a/Engine/Source/Runtime/Renderer/Private/DeferredShadingRenderer.h b/Engine/Source/Runtime/Renderer/Private/DeferredShadingRenderer.h
--- a/Engine/Source/Runtime/Renderer/Private/DeferredShadingRenderer.h
+++ b/Engine/Source/Runtime/Renderer/Private/DeferredShadingRenderer.h
@@ -1169,6 +1169,7 @@ private:
 	static void PrepareRayTracingGlobalIllumination(const FViewInfo& View, TArray<FRHIRayTracingShader*>& OutRayGenShaders);
 	static void PrepareRayTracingGlobalIlluminationPlugin(const FViewInfo& View, TArray<FRHIRayTracingShader*>& OutRayGenShaders);
 	static void PrepareRayTracingTranslucency(const FViewInfo& View, TArray<FRHIRayTracingShader*>& OutRayGenShaders);
+	static void PrepareRayTracingVolumetricFogShadows(const FViewInfo& View, const FScene& Scene, TArray<FRHIRayTracingShader*>& OutRayGenShaders);
 	static void PrepareRayTracingDebug(const FSceneViewFamily& ViewFamily, TArray<FRHIRayTracingShader*>& OutRayGenShaders);
 	static void PreparePathTracing(const FSceneViewFamily& ViewFamily, TArray<FRHIRayTracingShader*>& OutRayGenShaders);
 	static void PrepareLumenHardwareRayTracingScreenProbeGather(const FViewInfo& View, TArray<FRHIRayTracingShader*>& OutRayGenShaders);
diff --git a/Engine/Source/Runtime/Renderer/Private/VolumetricFog.cpp b/Engine/Source/Runtime/Renderer/Private/VolumetricFog.cpp
--- a/Engine/Source/Runtime/Renderer/Private/VolumetricFog.cpp
+++ b/Engine/Source/Runtime/Renderer/Private/VolumetricFog.cpp
@@ -133,6 +133,14 @@ FAutoConsoleVariableRef CVarVolumetricFogLightFunction(
 	ECVF_Scalability | ECVF_RenderThreadSafe
 );
 
+int GVolumetricFogInjectRaytracedLights = 0;
+FAutoConsoleVariableRef CVarVolumetricInjectRaytracedLights(
+	TEXT("r.VolumetricFog.InjectRaytracedLights"),
+	GVolumetricFogInjectRaytracedLights,
+	TEXT("Whether lights with ray traced shadows are injected into volumetric fog"),
+	ECVF_Scalability | ECVF_RenderThreadSafe
+);
+
 float GLightScatteringSampleJitterMultiplier = 0;
 FAutoConsoleVariableRef CVarLightScatteringSampleJitterMultiplier(
 	TEXT("r.VolumetricFog.LightScatteringSampleJitterMultiplier"),
@@ -267,6 +275,63 @@ class FWriteToBoundingSphereVS : public FGlobalShader
 
 IMPLEMENT_GLOBAL_SHADER(FWriteToBoundingSphereVS, "/Engine/Private/VolumetricFog.usf", "WriteToBoundingSphereVS", SF_Vertex);
 
+BEGIN_SHADER_PARAMETER_STRUCT(FInjectShadowedLocalLightCommonParameters, )
+	SHADER_PARAMETER_STRUCT_REF(FViewUniformShaderParameters, ViewUniformBuffer)
+	SHADER_PARAMETER_RDG_UNIFORM_BUFFER(FDeferredLightUniformStruct, DeferredLight)
+	SHADER_PARAMETER_RDG_TEXTURE(Texture2D, WhiteDummyTexture)
+	SHADER_PARAMETER_RDG_TEXTURE(Texture2D, LightFunctionAtlasTexture)
+	SHADER_PARAMETER_SAMPLER(SamplerState, LightFunctionAtlasSampler)
+	SHADER_PARAMETER_STRUCT_INCLUDE(FVolumetricFogIntegrationParameters, VolumetricFogParameters)
+	SHADER_PARAMETER(FMatrix44f, LocalLightFunctionMatrix)
+	SHADER_PARAMETER(FVector4f, LightFunctionAtlasTileMinMaxUvBound)
+	SHADER_PARAMETER(float, PhaseG)
+	SHADER_PARAMETER(float, InverseSquaredLightDistanceBiasScale)
+END_SHADER_PARAMETER_STRUCT()
+
+static bool SetupInjectShadowedLocalLightCommonParameters(
+	FRDGBuilder& GraphBuilder,
+	FViewInfo& View,
+	const FVolumetricFogIntegrationParameterData& IntegrationData,
+	const FExponentialHeightFogSceneInfo& FogInfo,
+	const FLightSceneInfo* LightSceneInfo,	
+	bool bUsesLightFunction,
+	FInjectShadowedLocalLightCommonParameters& OutCommonParameters)
+{
+	OutCommonParameters.LightFunctionAtlasTexture = View.VolumetricFogResources.TransientLightFunctionTextureAtlas ? View.VolumetricFogResources.TransientLightFunctionTextureAtlas->GetTransientLightFunctionAtlasTexture() : GSystemTextures.GetWhiteDummy(GraphBuilder);;
+	OutCommonParameters.LocalLightFunctionMatrix = FMatrix44f::Identity;
+	OutCommonParameters.LightFunctionAtlasTileMinMaxUvBound = FVector4f(ForceInitToZero);
+	if (bUsesLightFunction)
+	{
+		const TMap<FLightSceneInfo*, FVolumetricFogLocalLightFunctionInfo>& LocalLightFunctionData = View.VolumetricFogResources.LocalLightFunctionData;
+
+		const FVolumetricFogLocalLightFunctionInfo* LightFunctionData = LocalLightFunctionData.Find(LightSceneInfo);
+
+		if (!ensure(LightFunctionData != nullptr))
+		{
+			// The light function data is missing but the light requires it. Skip this light for now.
+			return false;
+		}
+
+		OutCommonParameters.LocalLightFunctionMatrix = LightFunctionData->LightFunctionTranslatedWorldToLightMatrix;
+		OutCommonParameters.LightFunctionAtlasTexture = LightFunctionData->AtlasTile.Texture;
+		OutCommonParameters.LightFunctionAtlasTileMinMaxUvBound = LightFunctionData->AtlasTile.MinMaxUvBound;
+	}
+
+	OutCommonParameters.LightFunctionAtlasSampler = TStaticSamplerState<SF_Bilinear, AM_Clamp, AM_Clamp, AM_Clamp>::GetRHI();
+	// We also bind the default light function texture because when we are out of atlas tile, we fallback to use a white light function so we need the RHI to be created
+	OutCommonParameters.WhiteDummyTexture = View.VolumetricFogResources.TransientLightFunctionTextureAtlas ? View.VolumetricFogResources.TransientLightFunctionTextureAtlas->GetDefaultLightFunctionTexture() : GSystemTextures.GetWhiteDummy(GraphBuilder);
+	SetupVolumetricFogIntegrationParameters(OutCommonParameters.VolumetricFogParameters, View, IntegrationData);
+
+	OutCommonParameters.ViewUniformBuffer = View.ViewUniformBuffer;
+	OutCommonParameters.PhaseG = FogInfo.VolumetricFogScatteringDistribution;
+	OutCommonParameters.InverseSquaredLightDistanceBiasScale = GInverseSquaredLightDistanceBiasScale;
+
+	FDeferredLightUniformStruct* DeferredLightStruct = GraphBuilder.AllocParameters<FDeferredLightUniformStruct>();
+	*DeferredLightStruct = GetDeferredLightParameters(View, *LightSceneInfo);
+	OutCommonParameters.DeferredLight = GraphBuilder.CreateUniformBuffer(DeferredLightStruct);
+
+	return true;
+}
 
 /** Shader that adds direct lighting contribution from the given light to the current volume lighting cascade. */
 class FInjectShadowedLocalLightPS : public FGlobalShader
@@ -275,19 +340,10 @@ class FInjectShadowedLocalLightPS : public FGlobalShader
 	SHADER_USE_PARAMETER_STRUCT(FInjectShadowedLocalLightPS, FGlobalShader);
 
 	BEGIN_SHADER_PARAMETER_STRUCT(FParameters, )
-		SHADER_PARAMETER_STRUCT_REF(FViewUniformShaderParameters, ViewUniformBuffer)
-		SHADER_PARAMETER_RDG_UNIFORM_BUFFER(FDeferredLightUniformStruct, DeferredLight)
+		SHADER_PARAMETER_STRUCT_INCLUDE(FInjectShadowedLocalLightCommonParameters, Common)
 		SHADER_PARAMETER_STRUCT_INCLUDE(FVolumeShadowingShaderParameters, VolumeShadowingShaderParameters)
 		SHADER_PARAMETER_STRUCT_INCLUDE(FVirtualShadowMapSamplingParameters, VirtualShadowMapSamplingParameters)
 		SHADER_PARAMETER_RDG_TEXTURE(Texture2D, ConservativeDepthTexture)
-		SHADER_PARAMETER_RDG_TEXTURE(Texture2D, WhiteDummyTexture)
-		SHADER_PARAMETER_RDG_TEXTURE(Texture2D, LightFunctionAtlasTexture)
-		SHADER_PARAMETER_SAMPLER(SamplerState, LightFunctionAtlasSampler)
-		SHADER_PARAMETER_STRUCT_INCLUDE(FVolumetricFogIntegrationParameters, VolumetricFogParameters)
-		SHADER_PARAMETER(FMatrix44f, LocalLightFunctionMatrix)
-		SHADER_PARAMETER(FVector4f, LightFunctionAtlasTileMinMaxUvBound)
-		SHADER_PARAMETER(float, PhaseG)
-		SHADER_PARAMETER(float, InverseSquaredLightDistanceBiasScale)
 		SHADER_PARAMETER(uint32, UseConservativeDepthTexture)
 		SHADER_PARAMETER(int32, VirtualShadowMapId)
 		RENDER_TARGET_BINDING_SLOTS()
@@ -315,6 +371,74 @@ class FInjectShadowedLocalLightPS : public FGlobalShader
 
 IMPLEMENT_GLOBAL_SHADER(FInjectShadowedLocalLightPS, "/Engine/Private/VolumetricFog.usf", "InjectShadowedLocalLightPS", SF_Pixel);
 
+#if RHI_RAYTRACING
+
+class FInjectShadowedLocalLightRGS : public FGlobalShader
+{
+	DECLARE_GLOBAL_SHADER(FInjectShadowedLocalLightRGS);
+	SHADER_USE_ROOT_PARAMETER_STRUCT(FInjectShadowedLocalLightRGS, FGlobalShader);
+
+	BEGIN_SHADER_PARAMETER_STRUCT(FParameters, )
+		SHADER_PARAMETER_STRUCT_INCLUDE(FInjectShadowedLocalLightCommonParameters, Common)
+
+		SHADER_PARAMETER_RDG_TEXTURE_UAV(RWTexture3D, OutVolumeTexture)
+		SHADER_PARAMETER_RDG_BUFFER_SRV(RaytracingAccelerationStructure, TLAS)
+		SHADER_PARAMETER(int32, FirstSlice)
+	END_SHADER_PARAMETER_STRUCT()
+
+	class FTemporalReprojection : SHADER_PERMUTATION_BOOL("USE_TEMPORAL_REPROJECTION");
+	class FLightFunction : SHADER_PERMUTATION_BOOL("USE_LIGHT_FUNCTION");
+
+	using FPermutationDomain = TShaderPermutationDomain<
+		FTemporalReprojection,
+		FLightFunction >;
+
+	static bool ShouldCompilePermutation(const FGlobalShaderPermutationParameters& Parameters)
+	{
+		return ShouldCompileRayTracingShadersForProject(Parameters.Platform);
+	}
+
+	static ERayTracingPayloadType GetRayTracingPayloadType(const int32 PermutationId)
+	{
+		return ERayTracingPayloadType::RayTracingMaterial;
+	}
+
+	static void ModifyCompilationEnvironment(const FGlobalShaderPermutationParameters& Parameters, FShaderCompilerEnvironment& OutEnvironment)
+	{
+		FGlobalShader::ModifyCompilationEnvironment(Parameters, OutEnvironment);
+		OutEnvironment.SetDefine(TEXT("USE_RAYTRACED_SHADOWS"), TEXT("1"));
+
+		// Only ray traced shadowed lights use this RGS
+		OutEnvironment.SetDefine(TEXT("ENABLE_SHADOW_COMPUTATION"), TEXT("1"));
+	}
+};
+
+IMPLEMENT_GLOBAL_SHADER(FInjectShadowedLocalLightRGS, "/Engine/Private/VolumetricFog.usf", "InjectShadowedLocalLightRGS", SF_RayGen);
+
+void FDeferredShadingSceneRenderer::PrepareRayTracingVolumetricFogShadows(const FViewInfo& View, const FScene& Scene, TArray<FRHIRayTracingShader*>& OutRayGenShaders)
+{
+	const bool bEnabled = Scene.bHasRayTracedLights && ::ShouldRenderVolumetricFog(&Scene, *View.Family) && GVolumetricFogInjectRaytracedLights;
+	if (!bEnabled)
+	{
+		return;
+	}
+
+	for (int32 TemporalReprojection = 0; TemporalReprojection < 2; ++TemporalReprojection)
+	{
+		for (int32 UseLightFunction = 0; UseLightFunction < 2; ++UseLightFunction)
+		{
+			FInjectShadowedLocalLightRGS::FPermutationDomain PermutationVector;
+			PermutationVector.Set<FInjectShadowedLocalLightRGS::FTemporalReprojection>((bool)TemporalReprojection);
+			PermutationVector.Set<FInjectShadowedLocalLightRGS::FLightFunction>((bool)UseLightFunction);
+
+			TShaderMapRef<FInjectShadowedLocalLightRGS> RayGenerationShader(View.ShaderMap, PermutationVector);
+			OutRayGenShaders.Add(RayGenerationShader.GetRayTracingShader());
+		}
+	}
+}
+
+#endif // RHI_RAYTRACING
+
 const FProjectedShadowInfo* GetShadowForInjectionIntoVolumetricFog(const FVisibleLightInfo& VisibleLightInfo)
 {
 	for (int32 ShadowIndex = 0; ShadowIndex < VisibleLightInfo.ShadowsToProject.Num(); ShadowIndex++)
@@ -331,6 +455,11 @@ const FProjectedShadowInfo* GetShadowForInjectionIntoVolumetricFog(const FVisibl
 	return nullptr;
 }
 
+bool LightHasRayTracedShadows(const FLightSceneInfo* LightSceneInfo)
+{
+	return ShouldRenderRayTracingShadowsForLight(*LightSceneInfo->Proxy) && GVolumetricFogInjectRaytracedLights;
+}
+
 bool LightNeedsSeparateInjectionIntoVolumetricFogForOpaqueShadow(const FViewInfo& View, const FLightSceneInfo* LightSceneInfo, const FVisibleLightInfo& VisibleLightInfo)
 {
 	const FLightSceneProxy* LightProxy = LightSceneInfo->Proxy;
@@ -344,8 +473,9 @@ bool LightNeedsSeparateInjectionIntoVolumetricFogForOpaqueShadow(const FViewInfo
 		const FStaticShadowDepthMap* StaticShadowDepthMap = LightProxy->GetStaticShadowDepthMap();
 		const bool bStaticallyShadowed = LightSceneInfo->IsPrecomputedLightingValid() && StaticShadowDepthMap && StaticShadowDepthMap->Data && StaticShadowDepthMap->TextureRHI;
 		const bool bHasVirtualShadowMap = VisibleLightInfo.GetVirtualShadowMapId( &View ) != INDEX_NONE;
+		const bool bHasRayTracedShadows = LightHasRayTracedShadows(LightSceneInfo);
 
-		return GetShadowForInjectionIntoVolumetricFog(VisibleLightInfo) != NULL || bStaticallyShadowed || bHasVirtualShadowMap;
+		return GetShadowForInjectionIntoVolumetricFog(VisibleLightInfo) != NULL || bStaticallyShadowed || bHasVirtualShadowMap || bHasRayTracedShadows;
 	}
 
 	return false;
@@ -464,6 +594,7 @@ void FSceneRenderer::RenderLocalLightsForVolumetricFog(
 {
 	// Gather lights that need to be rendered with shadow from opaque or light functions.
 	TArray<const FLightSceneInfo*, SceneRenderingAllocator> LightsToInject;
+	TArray<const FLightSceneInfo*, SceneRenderingAllocator> RayTracedLightsToInject;
 	for (auto LightIt = Scene->Lights.CreateConstIterator(); LightIt; ++LightIt)
 	{
 		const FLightSceneInfoCompact& LightSceneInfoCompact = *LightIt;
@@ -482,17 +613,23 @@ void FSceneRenderer::RenderLocalLightsForVolumetricFog(
 
 			if ((View.ViewMatrices.GetViewOrigin() - LightBounds.Center).SizeSquared() < (FogInfo.VolumetricFogDistance + LightBounds.W) * (FogInfo.VolumetricFogDistance + LightBounds.W))
 			{
-				LightsToInject.Add(LightSceneInfo);
+				const bool bRayTracedLight = LightHasRayTracedShadows(LightSceneInfo);
+				if (bRayTracedLight)
+				{
+					RayTracedLightsToInject.Add(LightSceneInfo);
+				}
+				else
+				{
+					LightsToInject.Add(LightSceneInfo);
+				}
 			}
 		}
 	}
 
 	// Now voxelise all the light we have just gathered.
+	bool bClearExecuted = false;
 	if (LightsToInject.Num() > 0)
 	{
-		TMap<FLightSceneInfo*, FVolumetricFogLocalLightFunctionInfo>& LocalLightFunctionData = View.VolumetricFogResources.LocalLightFunctionData;
-
-		bool bClearExecuted = false;
 		for (int32 LightIndex = 0; LightIndex < LightsToInject.Num(); LightIndex++)
 		{
 			const FLightSceneInfo* LightSceneInfo = LightsToInject[LightIndex];
@@ -512,48 +649,31 @@ void FSceneRenderer::RenderLocalLightsForVolumetricFog(
 				FInjectShadowedLocalLightPS::FParameters* PassParameters = GraphBuilder.AllocParameters<FInjectShadowedLocalLightPS::FParameters>();
 
 				// Light function parameters
-				PassParameters->LightFunctionAtlasTexture = View.VolumetricFogResources.TransientLightFunctionTextureAtlas ? View.VolumetricFogResources.TransientLightFunctionTextureAtlas->GetTransientLightFunctionAtlasTexture() : GSystemTextures.GetWhiteDummy(GraphBuilder);;
-				PassParameters->LocalLightFunctionMatrix = FMatrix44f::Identity;
-				PassParameters->LightFunctionAtlasTileMinMaxUvBound = FVector4f(ForceInitToZero);
-				if (bUsesLightFunction)
+				bool bValid = SetupInjectShadowedLocalLightCommonParameters(
+					GraphBuilder,
+					View,
+					IntegrationData,
+					FogInfo,
+					LightSceneInfo,
+					bUsesLightFunction,
+					PassParameters->Common
+				);
+
+				if (!bValid)
 				{
-					FVolumetricFogLocalLightFunctionInfo* LightFunctionData = LocalLightFunctionData.Find(LightSceneInfo);
-
-					if (!ensure(LightFunctionData != nullptr))
-					{
-						// The light function data is missing but the light requires it. Skip this light for now.
-						continue;
-					}
-
-					PassParameters->LocalLightFunctionMatrix = LightFunctionData->LightFunctionTranslatedWorldToLightMatrix;
-					PassParameters->LightFunctionAtlasTexture = LightFunctionData->AtlasTile.Texture;
-					PassParameters->LightFunctionAtlasTileMinMaxUvBound = LightFunctionData->AtlasTile.MinMaxUvBound;
+					continue;
 				}
 
-				// We need to delay the real texture creation replacing the dummy on first use after all the culling logic, 
-				// otherwise it will be reported as not written by any pass (in this case we need to keep the default dummy).
 				const bool bHasTextureBeenCreated = bClearExecuted == true;
-				OutLocalShadowedLightScattering = bHasTextureBeenCreated ? OutLocalShadowedLightScattering : GraphBuilder.CreateTexture(VolumeDesc, TEXT("VolumetricFog.LocalShadowedLightScattering"));
+				OutLocalShadowedLightScattering = bHasTextureBeenCreated ? OutLocalShadowedLightScattering : GraphBuilder.CreateTexture(VolumeDesc, TEXT("VolumetricFog.LocalShadowedLightScattering"));				
 
 				PassParameters->RenderTargets[0] = FRenderTargetBinding(OutLocalShadowedLightScattering, bClearExecuted ? ERenderTargetLoadAction::ELoad : ERenderTargetLoadAction::EClear);
 				bClearExecuted = true;
-				PassParameters->VirtualShadowMapSamplingParameters = VirtualShadowMapArray.GetSamplingParameters(GraphBuilder);
+
 				PassParameters->ConservativeDepthTexture = ConservativeDepthTexture;
-				PassParameters->LightFunctionAtlasSampler = TStaticSamplerState<SF_Bilinear, AM_Clamp, AM_Clamp, AM_Clamp>::GetRHI();
-				// We also bind the default light function texture because when we are out of atlas tile, we fallback to use a white light function so we need the RHI to be created
-				PassParameters->WhiteDummyTexture = View.VolumetricFogResources.TransientLightFunctionTextureAtlas ? View.VolumetricFogResources.TransientLightFunctionTextureAtlas->GetDefaultLightFunctionTexture() : GSystemTextures.GetWhiteDummy(GraphBuilder);
-				SetupVolumetricFogIntegrationParameters(PassParameters->VolumetricFogParameters, View, IntegrationData);
-
-				PassParameters->ViewUniformBuffer = View.ViewUniformBuffer;
-				PassParameters->PhaseG = FogInfo.VolumetricFogScatteringDistribution;
-				PassParameters->InverseSquaredLightDistanceBiasScale = GInverseSquaredLightDistanceBiasScale;
 				PassParameters->UseConservativeDepthTexture = GVolumetricFogConservativeDepth > 0 ? 1 : 0;
 				PassParameters->VirtualShadowMapId = VirtualShadowMapId;
 
-				FDeferredLightUniformStruct* DeferredLightStruct = GraphBuilder.AllocParameters<FDeferredLightUniformStruct>();
-				*DeferredLightStruct = GetDeferredLightParameters(View, *LightSceneInfo);
-				PassParameters->DeferredLight = GraphBuilder.CreateUniformBuffer(DeferredLightStruct);
-
 				const FProjectedShadowInfo* ProjectedShadowInfo = GetShadowForInjectionIntoVolumetricFog(VisibleLightInfo);
 				const bool bDynamicallyShadowed = ProjectedShadowInfo != NULL;
 				GetVolumeShadowingShaderParameters(GraphBuilder, View, LightSceneInfo, ProjectedShadowInfo, PassParameters->VolumeShadowingShaderParameters);
@@ -605,7 +725,7 @@ void FSceneRenderer::RenderLocalLightsForVolumetricFog(
 
 						VSPassParameters.ClipRatio = GetVolumetricFogFroxelToScreenSVPosRatio(View);
 
-						VSPassParameters.VolumetricFogParameters = PassParameters->VolumetricFogParameters;
+						VSPassParameters.VolumetricFogParameters = PassParameters->Common.VolumetricFogParameters;
 						SetShaderParameters(RHICmdList, VertexShader, VertexShader.GetVertexShader(), VSPassParameters);
 
 						if (GeometryShader.IsValid())
@@ -624,6 +744,77 @@ void FSceneRenderer::RenderLocalLightsForVolumetricFog(
 			}
 		}
 	}
+
+#if RHI_RAYTRACING
+	if (RayTracedLightsToInject.Num() > 0)
+	{
+		if (!bClearExecuted)
+		{
+			OutLocalShadowedLightScattering = GraphBuilder.CreateTexture(VolumeDesc, TEXT("VolumetricFog.LocalShadowedLightScattering"));
+			AddClearUAVPass(GraphBuilder, GraphBuilder.CreateUAV(OutLocalShadowedLightScattering), 0.0f);
+			bClearExecuted = true;
+		}
+
+		for (int32 LightIndex = 0; LightIndex < RayTracedLightsToInject.Num(); LightIndex++)
+		{
+			const FLightSceneInfo* LightSceneInfo = RayTracedLightsToInject[LightIndex];
+
+			const FSphere LightBounds = LightSceneInfo->Proxy->GetBoundingSphere();
+			const FIntPoint VolumeZBounds = CalculateVolumetricFogBoundsForLight(LightBounds, View, VolumetricFogViewGridSize, GridZParams);
+			if (VolumeZBounds.X < VolumeZBounds.Y)
+			{
+				bool bUsesLightFunction = ViewFamily.EngineShowFlags.LightFunctions
+					&& CheckForLightFunction(LightSceneInfo) && LightNeedsSeparateInjectionIntoVolumetricFogForLightFunction(LightSceneInfo);
+
+				FInjectShadowedLocalLightRGS::FParameters* PassParameters = GraphBuilder.AllocParameters<FInjectShadowedLocalLightRGS::FParameters>();
+				PassParameters->OutVolumeTexture = GraphBuilder.CreateUAV(OutLocalShadowedLightScattering);
+				PassParameters->TLAS = View.GetRayTracingSceneLayerViewChecked(ERayTracingSceneLayer::Base);
+				PassParameters->FirstSlice = VolumeZBounds.X;
+
+				bool bValid = SetupInjectShadowedLocalLightCommonParameters(
+					GraphBuilder,
+					View,
+					IntegrationData,
+					FogInfo,
+					LightSceneInfo,
+					bUsesLightFunction,
+					PassParameters->Common
+				);
+
+				if (!bValid)
+				{
+					continue;
+				}
+
+				FInjectShadowedLocalLightRGS::FPermutationDomain PermutationVector;
+				PermutationVector.Set< FInjectShadowedLocalLightRGS::FTemporalReprojection >(bUseTemporalReprojection);
+				PermutationVector.Set< FInjectShadowedLocalLightRGS::FLightFunction >(bUsesLightFunction);
+
+				TShaderMapRef<FInjectShadowedLocalLightRGS> RayGenerationShader(GetGlobalShaderMap(FeatureLevel), PermutationVector);
+
+				ClearUnusedGraphResources(RayGenerationShader, PassParameters);
+
+				// TODO: better bounds
+				const int32 NumSlices = VolumeZBounds.Y - VolumeZBounds.X;
+				const uint32 DispatchSize = VolumeDesc.Extent.X * VolumeDesc.Extent.Y * NumSlices;
+				GraphBuilder.AddPass(
+					RDG_EVENT_NAME("RayTracedShadowedLights"),
+					PassParameters,
+					ERDGPassFlags::Compute,
+					[this, &View, RayGenerationShader, PassParameters, DispatchSize](FRHIRayTracingCommandList& RHICmdList)
+					{
+						FRayTracingShaderBindingsWriter GlobalResources;
+						SetShaderParameters(GlobalResources, RayGenerationShader, *PassParameters);
+
+						FRHIRayTracingScene* RayTracingSceneRHI = View.GetRayTracingSceneChecked();
+
+						RHICmdList.RayTraceDispatch(View.RayTracingMaterialPipeline, RayGenerationShader.GetRayTracingShader(), RayTracingSceneRHI, GlobalResources, DispatchSize, 1);
+					}
+				);
+			}
+		}
+	}
+#endif
 }
 
 IMPLEMENT_GLOBAL_SHADER_PARAMETER_STRUCT(FLumenTranslucencyLightingUniforms, "LumenGIVolumeStruct");
