Fixed missing VSM parameters after recent changes to volumetric fog.

diff --git a/Engine/Source/Runtime/Renderer/Private/VolumetricFog.cpp b/Engine/Source/Runtime/Renderer/Private/VolumetricFog.cpp
--- a/Engine/Source/Runtime/Renderer/Private/VolumetricFog.cpp
+++ b/Engine/Source/Runtime/Renderer/Private/VolumetricFog.cpp
@@ -670,6 +670,7 @@ void FSceneRenderer::RenderLocalLightsForVolumetricFog(
 				PassParameters->RenderTargets[0] = FRenderTargetBinding(OutLocalShadowedLightScattering, bClearExecuted ? ERenderTargetLoadAction::ELoad : ERenderTargetLoadAction::EClear);
 				bClearExecuted = true;
 
+				PassParameters->VirtualShadowMapSamplingParameters = VirtualShadowMapArray.GetSamplingParameters(GraphBuilder);
 				PassParameters->ConservativeDepthTexture = ConservativeDepthTexture;
 				PassParameters->UseConservativeDepthTexture = GVolumetricFogConservativeDepth > 0 ? 1 : 0;
 				PassParameters->VirtualShadowMapId = VirtualShadowMapId;
