Support for Nanite Raytracing in RayTracingDebug
- Implemented permutation of FRayTracingDebugCHS that calculates geometry normal from Nanite Cluster Data instead of VB/IB used to build BLAS.
- Modified RTPSO to include 2 hitgroups (Regular, Nanite RT) and appropriate hitgroup index is selected when building SBT.
- GetGeometryNormalFromTriangleBaseAttributes(...) use the appropriate logic depending on whether NANITE_RAY_TRACING is defined.

diff --git a/Engine/Shaders/Private/RayTracing/RayTracingHitGroupCommon.ush b/Engine/Shaders/Private/RayTracing/RayTracingHitGroupCommon.ush
--- a/Engine/Shaders/Private/RayTracing/RayTracingHitGroupCommon.ush
+++ b/Engine/Shaders/Private/RayTracing/RayTracingHitGroupCommon.ush
@@ -12,6 +12,15 @@
 #include "/Engine/Private/RayTracing/RayTracingCommon.ush"
 #include "/Engine/Shared/RayTracingBuiltInResources.h"
 
+#if NANITE_RAY_TRACING
+	#define NANITE_USE_RAYTRACING_UNIFORM_BUFFER 1
+	#define NANITE_NUM_TEXCOORDS_TO_DECODE 0
+
+	#include "/Engine/Private/Nanite/NaniteDataDecode.ush"
+	#include "/Engine/Private/Nanite/NaniteAttributeDecode.ush"
+	#include "/Engine/Private/Nanite/NaniteVertexFactory.ush"
+#endif
+
 uint3 LoadIndices16Bit(ByteAddressBuffer SourceBuffer, uint OffsetInBytes)
 {
 	uint3 Result;
@@ -112,6 +121,42 @@ uint GetBaseInstanceIndex()
 
 float3 GetGeometryNormalFromTriangleBaseAttributes(uint PrimitiveIndex)
 {
+#if NANITE_RAY_TRACING
+	const int ScenePrimitiveIndex = GetInstanceUserData();
+	const int SceneInstanceIndex = InstanceIndex() - GetBaseInstanceIndex();
+
+	const uint GPUSceneInstanceIndex = GetGPUSceneInstanceIndex(ScenePrimitiveIndex, SceneInstanceIndex);
+
+	FInstanceSceneData InstanceData = GetInstanceSceneData(GPUSceneInstanceIndex, Scene.GPUScene.InstanceDataSOAStride);
+	FPrimitiveSceneData PrimitiveData = GetPrimitiveData(ScenePrimitiveIndex);
+
+	const uint FirstTriangle = HitGroupSystemRootConstants.FirstPrimitive;
+	const uint PackedTriangleData = RayTracingDataBuffer[PrimitiveData.NaniteRayTracingDataOffset + FirstTriangle + PrimitiveIndex];
+
+	const uint PageIndex = PackedTriangleData & NANITE_MAX_GPU_PAGES_MASK;
+	const uint ClusterIndex = (PackedTriangleData >> NANITE_MAX_GPU_PAGES_BITS) & NANITE_MAX_CLUSTERS_PER_PAGE_MASK;
+	const uint TriIndex = (PackedTriangleData >> (NANITE_MAX_GPU_PAGES_BITS + NANITE_MAX_CLUSTERS_PER_PAGE_BITS)) & NANITE_MAX_CLUSTER_TRIANGLES_MASK;
+
+	FCluster Cluster = GetCluster(PageIndex, ClusterIndex);
+
+	const uint3 TriIndices = DecodeTriangleIndices(Cluster, TriIndex);
+
+	const float3 PointLocalNoWPO[3] =
+	{
+		DecodePosition(TriIndices.x, Cluster),
+		DecodePosition(TriIndices.y, Cluster),
+		DecodePosition(TriIndices.z, Cluster)
+	};
+
+	float3 LocalEdges[2] = {
+		PointLocalNoWPO[1] - PointLocalNoWPO[0],
+		PointLocalNoWPO[2] - PointLocalNoWPO[0]
+	};
+	float3 LocalNormal = cross(LocalEdges[1], LocalEdges[0]);
+
+	float3x3 InverseTranspose3x3 = transpose((float3x3)WorldToObject4x3());
+	float3 WorldNormal = normalize(mul(LocalNormal, InverseTranspose3x3));
+#else
 	FTriangleBaseAttributes Tri = LoadTriangleBaseAttributes(PrimitiveIndex);
 
 	float3 LocalEdges[2] = {
@@ -122,6 +167,7 @@ float3 GetGeometryNormalFromTriangleBaseAttributes(uint PrimitiveIndex)
 
 	float3x3 InverseTranspose3x3 = transpose((float3x3)WorldToObject4x3());
 	float3 WorldNormal = normalize(mul(LocalNormal, InverseTranspose3x3));
+#endif
 	return WorldNormal;
 }
 
diff --git a/Engine/Source/Runtime/Renderer/Private/DeferredShadingRenderer.cpp b/Engine/Source/Runtime/Renderer/Private/DeferredShadingRenderer.cpp
--- a/Engine/Source/Runtime/Renderer/Private/DeferredShadingRenderer.cpp
+++ b/Engine/Source/Runtime/Renderer/Private/DeferredShadingRenderer.cpp
@@ -2255,7 +2255,7 @@ void FDeferredShadingSceneRenderer::WaitForRayTracingScene(FRDGBuilder& GraphBui
 	const bool bIsPathTracing = ViewFamily.EngineShowFlags.PathTracing;
 
 	FBuildAccelerationStructurePassParams* PassParams = GraphBuilder.AllocParameters<FBuildAccelerationStructurePassParams>();
-	PassParams->Scene = GetSceneUniforms().GetBuffer(GraphBuilder);
+	PassParams->Scene = GetSceneUniformBufferRef(GraphBuilder);
 	PassParams->RayTracingSceneScratchBuffer = nullptr;
 	PassParams->DynamicGeometryScratchBuffer = nullptr;
 	PassParams->LightDataPacked = bIsPathTracing ? nullptr : ReferenceView.RayTracingLightDataUniformBuffer; // accessed by FRayTracingLightingMS
@@ -2279,26 +2279,26 @@ void FDeferredShadingSceneRenderer::WaitForRayTracingScene(FRDGBuilder& GraphBui
 	GraphBuilder.AddPass(RDG_EVENT_NAME("WaitForRayTracingScene"), PassParams, ERDGPassFlags::Copy | ERDGPassFlags::Compute | ERDGPassFlags::NeverCull,
 		[this, PassParams, bIsPathTracing, &ReferenceView, bAnyLumenHardwareInlineRayTracingPassEnabled, RayTracingLightFunctionMap](FRHICommandListImmediate& RHICmdList)
 	{
+		if (IsNaniteEnabled())
+		{
+			FNaniteRayTracingUniformParameters NaniteRayTracingUniformParams;
+			NaniteRayTracingUniformParams.PageConstants.X = Scene->GPUScene.InstanceSceneDataSOAStride;
+			NaniteRayTracingUniformParams.PageConstants.Y = Nanite::GStreamingManager.GetMaxStreamingPages();
+			NaniteRayTracingUniformParams.MaxNodes = Nanite::FGlobalResources::GetMaxNodes();
+			NaniteRayTracingUniformParams.MaxVisibleClusters = Nanite::FGlobalResources::GetMaxVisibleClusters();
+			NaniteRayTracingUniformParams.RenderFlags = 0;
+			NaniteRayTracingUniformParams.RayTracingCutError = Nanite::GRayTracingManager.GetCutError();
+			NaniteRayTracingUniformParams.ClusterPageData = PassParams->ClusterPageData->GetRHI();
+			NaniteRayTracingUniformParams.HierarchyBuffer = PassParams->HierarchyBuffer->GetRHI();
+			NaniteRayTracingUniformParams.RayTracingDataBuffer = PassParams->RayTracingDataBuffer->GetRHI();
+
+			Nanite::GRayTracingManager.GetUniformBuffer().UpdateUniformBufferImmediate(NaniteRayTracingUniformParams);
+		}
+
 		check(ReferenceView.RayTracingMaterialPipeline || ReferenceView.RayTracingMaterialBindings.Num() == 0);
 
 		if (ReferenceView.RayTracingMaterialPipeline && (ReferenceView.RayTracingMaterialBindings.Num() || ReferenceView.RayTracingCallableBindings.Num()))
 		{
-			if (IsNaniteEnabled())
-			{
-				FNaniteRayTracingUniformParameters NaniteRayTracingUniformParams;
-				NaniteRayTracingUniformParams.PageConstants.X = Scene->GPUScene.InstanceSceneDataSOAStride;
-				NaniteRayTracingUniformParams.PageConstants.Y = Nanite::GStreamingManager.GetMaxStreamingPages();
-				NaniteRayTracingUniformParams.MaxNodes = Nanite::FGlobalResources::GetMaxNodes();
-				NaniteRayTracingUniformParams.MaxVisibleClusters = Nanite::FGlobalResources::GetMaxVisibleClusters();
-				NaniteRayTracingUniformParams.RenderFlags = 0;
-				NaniteRayTracingUniformParams.RayTracingCutError = Nanite::GRayTracingManager.GetCutError();
-				NaniteRayTracingUniformParams.ClusterPageData = PassParams->ClusterPageData->GetRHI();
-				NaniteRayTracingUniformParams.HierarchyBuffer = PassParams->HierarchyBuffer->GetRHI();
-				NaniteRayTracingUniformParams.RayTracingDataBuffer = PassParams->RayTracingDataBuffer->GetRHI();
-
-				Nanite::GRayTracingManager.GetUniformBuffer().UpdateUniformBufferImmediate(NaniteRayTracingUniformParams);
-			}
-
 			BindRayTracingMaterialPipeline(RHICmdList, ReferenceView, ReferenceView.RayTracingMaterialPipeline);
 
 			if (bIsPathTracing)
diff --git a/Engine/Source/Runtime/Renderer/Private/MeshPassProcessor.cpp b/Engine/Source/Runtime/Renderer/Private/MeshPassProcessor.cpp
--- a/Engine/Source/Runtime/Renderer/Private/MeshPassProcessor.cpp
+++ b/Engine/Source/Runtime/Renderer/Private/MeshPassProcessor.cpp
@@ -1020,6 +1020,11 @@ void FRayTracingMeshCommand::SetShaders(const FMeshProcessorShaders& Shaders)
 	ShaderBindings.Initialize(Shaders);
 }
 
+bool FRayTracingMeshCommand::IsUsingNaniteRayTracing() const
+{
+	return NaniteUniformBufferParameter.IsBound();
+}
+
 void FRayTracingShaderCommand::SetRayTracingShaderBindings(
 	FRayTracingLocalShaderBindingWriter* BindingWriter,
 	const TUniformBufferRef<FViewUniformShaderParameters>& ViewUniformBuffer,
diff --git a/Engine/Source/Runtime/Renderer/Private/RayTracing/RayTracingDebug.cpp b/Engine/Source/Runtime/Renderer/Private/RayTracing/RayTracingDebug.cpp
--- a/Engine/Source/Runtime/Renderer/Private/RayTracing/RayTracingDebug.cpp
+++ b/Engine/Source/Runtime/Renderer/Private/RayTracing/RayTracingDebug.cpp
@@ -17,6 +17,7 @@
 #include "RayTracing/RayTracingLighting.h"
 #include "RayTracing/RaytracingOptions.h"
 #include "RayTracing/RayTracingTraversalStatistics.h"
+#include "Nanite/NaniteRayTracing.h"
 #include "PixelShaderUtils.h"
 #include "SystemTextures.h"
 
@@ -123,6 +124,7 @@ class FRayTracingDebugRGS : public FGlobalShader
 		SHADER_PARAMETER_RDG_BUFFER_SRV(StructuredBuffer, InstancesDebugData)
 		SHADER_PARAMETER_RDG_BUFFER_SRV(StructuredBuffer<FRayTracingPickingFeedback>, PickingBuffer)
 		SHADER_PARAMETER_STRUCT_REF(FViewUniformShaderParameters, ViewUniformBuffer)
+		SHADER_PARAMETER_RDG_UNIFORM_BUFFER(FSceneUniformParameters, SceneUniformBuffer)
 	END_SHADER_PARAMETER_STRUCT()
 
 	static bool ShouldCompilePermutation(const FGlobalShaderPermutationParameters& Parameters)
@@ -149,22 +151,31 @@ class FRayTracingDebugCHS : public FGlobalShader
 {
 	DECLARE_GLOBAL_SHADER(FRayTracingDebugCHS);
 
-	static void ModifyCompilationEnvironment(const FGlobalShaderPermutationParameters& Parameters, FShaderCompilerEnvironment& OutEnvironment)
-	{
-		FGlobalShader::ModifyCompilationEnvironment(Parameters, OutEnvironment);
-	}
-
 public:
 
+	FRayTracingDebugCHS() = default;
+	FRayTracingDebugCHS(const ShaderMetaType::CompiledShaderInitializerType& Initializer)
+		: FGlobalShader(Initializer)
+	{}
+
+	class FNaniteRayTracing : SHADER_PERMUTATION_BOOL("NANITE_RAY_TRACING");
+	using FPermutationDomain = TShaderPermutationDomain<FNaniteRayTracing>;
+
 	static bool ShouldCompilePermutation(const FGlobalShaderPermutationParameters& Parameters)
 	{
 		return ShouldCompileRayTracingShadersForProject(Parameters.Platform);
 	}
 
-	FRayTracingDebugCHS() = default;
-	FRayTracingDebugCHS(const ShaderMetaType::CompiledShaderInitializerType& Initializer)
-		: FGlobalShader(Initializer)
-	{}
+	static void ModifyCompilationEnvironment(const FGlobalShaderPermutationParameters& Parameters, FShaderCompilerEnvironment& OutEnvironment)
+	{
+		FGlobalShader::ModifyCompilationEnvironment(Parameters, OutEnvironment);
+
+		FPermutationDomain PermutationVector(Parameters.PermutationId);
+		if (PermutationVector.Get<FNaniteRayTracing>())
+		{
+			OutEnvironment.SetDefine(TEXT("VF_SUPPORTS_PRIMITIVE_SCENE_DATA"), 1);
+		}
+	}
 
 	static ERayTracingPayloadType GetRayTracingPayloadType(const int32 PermutationId)
 	{
@@ -271,6 +282,7 @@ class FRayTracingPickingRGS : public FGlobalShader
 		SHADER_PARAMETER_RDG_BUFFER_SRV(StructuredBuffer, InstancesDebugData)
 		SHADER_PARAMETER_RDG_BUFFER_SRV(StructuredBuffer, InstanceBuffer)
 		SHADER_PARAMETER_STRUCT_REF(FViewUniformShaderParameters, ViewUniformBuffer)
+		SHADER_PARAMETER_RDG_UNIFORM_BUFFER(FSceneUniformParameters, SceneUniformBuffer)
 	END_SHADER_PARAMETER_STRUCT()
 
 	static bool ShouldCompilePermutation(const FGlobalShaderPermutationParameters& Parameters)
@@ -444,9 +456,9 @@ struct FRayTracingDebugResources : public FRenderResource
 
 TGlobalResource<FRayTracingDebugResources> GRayTracingDebugResources;
 
-void BindRayTracingDebugCHSMaterialBindings(FRHICommandList& RHICmdList, const FViewInfo& View, FRayTracingPipelineState* PipelineState)
+void BindRayTracingDebugCHSMaterialBindings(FRHICommandList& RHICmdList, const FViewInfo& View, FRHIUniformBuffer* SceneUniformBuffer, FRayTracingPipelineState* PipelineState)
 {
-	const int32 NumTotalBindings = View.VisibleRayTracingMeshCommands.Num();
+	const bool bNaniteRayTracingMode = Nanite::GetRayTracingMode() != Nanite::ERayTracingMode::Fallback;
 
 	FSceneRenderingBulkObjectAllocator Allocator;
 
@@ -457,23 +469,73 @@ void BindRayTracingDebugCHSMaterialBindings(FRHICommandList& RHICmdList, const F
 			: RHICmdList.Alloc(Size, Align);
 	};
 
+	const int32 NumTotalBindings = View.VisibleRayTracingMeshCommands.Num();
 	const uint32 MergedBindingsSize = sizeof(FRayTracingLocalShaderBindings) * NumTotalBindings;
 	FRayTracingLocalShaderBindings* Bindings = (FRayTracingLocalShaderBindings*)Alloc(MergedBindingsSize, alignof(FRayTracingLocalShaderBindings));
 
-	const uint32 NumUniformBuffers = 1;
-	FRHIUniformBuffer** UniformBufferArray = (FRHIUniformBuffer**)Alloc(sizeof(FRHIUniformBuffer*) * NumUniformBuffers, alignof(FRHIUniformBuffer*));
-	UniformBufferArray[0] = View.ViewUniformBuffer.GetReference();
+	struct FBinding
+	{
+		int32 ShaderIndexInPipeline;
+		uint32 NumUniformBuffers;
+		FRHIUniformBuffer** UniformBufferArray;
+	};
+
+	auto SetupBinding = [&](FRayTracingDebugCHS::FPermutationDomain PermutationVector)
+	{
+		auto Shader = View.ShaderMap->GetShader<FRayTracingDebugCHS>(PermutationVector);
+		auto HitGroupShader = Shader.GetRayTracingShader();
+
+		FBinding Binding;
+		Binding.ShaderIndexInPipeline = FindRayTracingHitGroupIndex(PipelineState, HitGroupShader, true);
+		Binding.NumUniformBuffers = Shader->ParameterMapInfo.UniformBuffers.Num();
+		Binding.UniformBufferArray = (FRHIUniformBuffer**)Alloc(sizeof(FRHIUniformBuffer*) * Binding.NumUniformBuffers, alignof(FRHIUniformBuffer*));
+
+		const auto& ViewUniformBufferParameter = Shader->GetUniformBufferParameter<FViewUniformShaderParameters>();
+		const auto& SceneUniformBufferParameter = Shader->GetUniformBufferParameter<FSceneUniformParameters>();
+		const auto& NaniteUniformBufferParameter = Shader->GetUniformBufferParameter<FNaniteRayTracingUniformParameters>();
+
+		if (ViewUniformBufferParameter.IsBound())
+		{
+			check(ViewUniformBufferParameter.GetBaseIndex() < Binding.NumUniformBuffers);
+			Binding.UniformBufferArray[ViewUniformBufferParameter.GetBaseIndex()] = View.ViewUniformBuffer.GetReference();
+		}
+
+		if (SceneUniformBufferParameter.IsBound())
+		{
+			check(SceneUniformBufferParameter.GetBaseIndex() < Binding.NumUniformBuffers);
+			Binding.UniformBufferArray[SceneUniformBufferParameter.GetBaseIndex()] = SceneUniformBuffer;
+		}
+
+		if (NaniteUniformBufferParameter.IsBound())
+		{
+			check(NaniteUniformBufferParameter.GetBaseIndex() < Binding.NumUniformBuffers);
+			Binding.UniformBufferArray[NaniteUniformBufferParameter.GetBaseIndex()] = Nanite::GRayTracingManager.GetUniformBuffer().GetReference();
+		}
+
+		return Binding;
+	};
+
+	FRayTracingDebugCHS::FPermutationDomain PermutationVector;
+
+	PermutationVector.Set<FRayTracingDebugCHS::FNaniteRayTracing>(false);
+	FBinding ShaderBinding = SetupBinding(PermutationVector);
+
+	PermutationVector.Set<FRayTracingDebugCHS::FNaniteRayTracing>(true);
+	FBinding ShaderBindingNaniteRT = SetupBinding(PermutationVector);
 
 	uint32 BindingIndex = 0;
 	for (const FVisibleRayTracingMeshCommand VisibleMeshCommand : View.VisibleRayTracingMeshCommands)
 	{
 		const FRayTracingMeshCommand& MeshCommand = *VisibleMeshCommand.RayTracingMeshCommand;
 
+		const FBinding& HelperBinding = MeshCommand.IsUsingNaniteRayTracing() ? ShaderBindingNaniteRT : ShaderBinding;
+
 		FRayTracingLocalShaderBindings Binding = {};
+		Binding.ShaderIndexInPipeline = HelperBinding.ShaderIndexInPipeline;
 		Binding.InstanceIndex = VisibleMeshCommand.InstanceIndex;
 		Binding.SegmentIndex = MeshCommand.GeometrySegmentIndex;
-		Binding.UniformBuffers = UniformBufferArray;
-		Binding.NumUniformBuffers = NumUniformBuffers;
+		Binding.UniformBuffers = HelperBinding.UniformBufferArray;
+		Binding.NumUniformBuffers = HelperBinding.NumUniformBuffers;
 		Binding.UserData = VisibleMeshCommand.InstanceIndex;
 
 		Bindings[BindingIndex] = Binding;
@@ -533,18 +595,27 @@ static FRDGBufferRef RayTracingPerformPicking(FRDGBuilder& GraphBuilder, const F
 	auto RayGenShader = ShaderMap->GetShader<FRayTracingPickingRGS>();
 
 	FRayTracingPipelineStateInitializer Initializer;
+	Initializer.MaxPayloadSizeInBytes = GetRayTracingPayloadTypeMaxSize(ERayTracingPayloadType::RayTracingDebug);
 
 	FRHIRayTracingShader* RayGenShaderTable[] = { RayGenShader.GetRayTracingShader() };
 	Initializer.SetRayGenShaderTable(RayGenShaderTable);
 
-	auto ClosestHitShader = ShaderMap->GetShader<FRayTracingDebugCHS>();
-	auto MissShader = ShaderMap->GetShader<FRayTracingDebugMS>();
-	FRHIRayTracingShader* HitGroupTable[] = { ClosestHitShader.GetRayTracingShader() };
+	FRayTracingDebugCHS::FPermutationDomain PermutationVector;
+
+	PermutationVector.Set<FRayTracingDebugCHS::FNaniteRayTracing>(false);
+	auto HitGroupShader = View.ShaderMap->GetShader<FRayTracingDebugCHS>(PermutationVector);
+
+	PermutationVector.Set<FRayTracingDebugCHS::FNaniteRayTracing>(true);
+	auto HitGroupShaderNaniteRT = View.ShaderMap->GetShader<FRayTracingDebugCHS>(PermutationVector);
+
+	FRHIRayTracingShader* HitGroupTable[] = { HitGroupShader.GetRayTracingShader(), HitGroupShaderNaniteRT.GetRayTracingShader() };
 	Initializer.SetHitGroupTable(HitGroupTable);
 	Initializer.bAllowHitGroupIndexing = true; // Required for stable output using GetBaseInstanceIndex().
+
+	auto MissShader = ShaderMap->GetShader<FRayTracingDebugMS>();
 	FRHIRayTracingShader* MissTable[] = { MissShader.GetRayTracingShader() };
 	Initializer.SetMissShaderTable(MissTable);
-	Initializer.MaxPayloadSizeInBytes = GetRayTracingPayloadTypeMaxSize(ERayTracingPayloadType::RayTracingDebug);
+
 	FRayTracingPipelineState* PickingPipeline = PipelineStateCache::GetAndOrCreateRayTracingPipelineState(GraphBuilder.RHICmdList, Initializer);
 
 	FRDGBufferDesc PickingBufferDesc = FRDGBufferDesc::CreateStructuredDesc(sizeof(FRayTracingPickingFeedback), 1);
@@ -557,6 +628,7 @@ static FRDGBufferRef RayTracingPerformPicking(FRDGBuilder& GraphBuilder, const F
 	RayGenParameters->OpaqueOnly = CVarRayTracingDebugModeOpaqueOnly.GetValueOnRenderThread();
 	RayGenParameters->InstanceBuffer = GraphBuilder.CreateSRV(Scene->RayTracingScene.InstanceBuffer);
 	RayGenParameters->ViewUniformBuffer = View.ViewUniformBuffer;
+	RayGenParameters->SceneUniformBuffer = GetSceneUniformBufferRef(GraphBuilder, View); // TODO: use a separate params structure
 	RayGenParameters->PickingOutput = GraphBuilder.CreateUAV(PickingBuffer);
 
 	GraphBuilder.AddPass(
@@ -568,7 +640,7 @@ static FRDGBufferRef RayTracingPerformPicking(FRDGBuilder& GraphBuilder, const F
 			FRayTracingShaderBindingsWriter GlobalResources;
 			SetShaderParameters(GlobalResources, RayGenShader, *RayGenParameters);
 
-			BindRayTracingDebugCHSMaterialBindings(RHICmdList, View, PickingPipeline);
+			BindRayTracingDebugCHSMaterialBindings(RHICmdList, View, RayGenParameters->SceneUniformBuffer->GetRHI(), PickingPipeline);
 			RHICmdList.SetRayTracingMissShader(View.GetRayTracingSceneChecked(), 0, PickingPipeline, 0 /* ShaderIndexInPipeline */, 0, nullptr, 0);
 
 			RHICmdList.RayTraceDispatch(PickingPipeline, RayGenShader.GetRayTracingShader(), View.GetRayTracingSceneChecked(), GlobalResources, 1, 1);
@@ -911,21 +983,27 @@ void FDeferredShadingSceneRenderer::RenderRayTracingDebug(FRDGBuilder& GraphBuil
 	if (bRequiresDebugCHS)
 	{
 		FRayTracingPipelineStateInitializer Initializer;
+		Initializer.MaxPayloadSizeInBytes = GetRayTracingPayloadTypeMaxSize(ERayTracingPayloadType::RayTracingDebug);
 
 		FRHIRayTracingShader* RayGenShaderTable[] = { RayGenShader.GetRayTracingShader() };
 		Initializer.SetRayGenShaderTable(RayGenShaderTable);
 
-		auto ClosestHitShader = ShaderMap->GetShader<FRayTracingDebugCHS>();
-		auto MissShader = ShaderMap->GetShader<FRayTracingDebugMS>();
-		FRHIRayTracingShader* HitGroupTable[] = { ClosestHitShader.GetRayTracingShader() };
+		FRayTracingDebugCHS::FPermutationDomain PermutationVectorCHS;
+
+		PermutationVectorCHS.Set<FRayTracingDebugCHS::FNaniteRayTracing>(false);
+		auto HitGroupShader = View.ShaderMap->GetShader<FRayTracingDebugCHS>(PermutationVectorCHS);
+
+		PermutationVectorCHS.Set<FRayTracingDebugCHS::FNaniteRayTracing>(true);
+		auto HitGroupShaderNaniteRT = View.ShaderMap->GetShader<FRayTracingDebugCHS>(PermutationVectorCHS);
+
+		FRHIRayTracingShader* HitGroupTable[] = { HitGroupShader.GetRayTracingShader(), HitGroupShaderNaniteRT.GetRayTracingShader() };
 		Initializer.SetHitGroupTable(HitGroupTable);
 		Initializer.bAllowHitGroupIndexing = true; // Required for stable output using GetBaseInstanceIndex().
-		
+
+		auto MissShader = ShaderMap->GetShader<FRayTracingDebugMS>();
 		FRHIRayTracingShader* MissTable[] = { MissShader.GetRayTracingShader() };
 		Initializer.SetMissShaderTable(MissTable);
 
-		Initializer.MaxPayloadSizeInBytes = GetRayTracingPayloadTypeMaxSize(ERayTracingPayloadType::RayTracingDebug);
-
 		Pipeline = PipelineStateCache::GetAndOrCreateRayTracingPipelineState(GraphBuilder.RHICmdList, Initializer);
 		bRequiresBindings = true;
 	}
@@ -968,6 +1046,8 @@ void FDeferredShadingSceneRenderer::RenderRayTracingDebug(FRDGBuilder& GraphBuil
 	RayGenParameters->ViewUniformBuffer = View.ViewUniformBuffer;
 	RayGenParameters->Output = GraphBuilder.CreateUAV(SceneColorTexture);
 
+	RayGenParameters->SceneUniformBuffer = GetSceneUniformBufferRef(GraphBuilder); // TODO: use a separate params structure
+
 	FIntRect ViewRect = View.ViewRect;
 
 	RDG_GPU_STAT_SCOPE(GraphBuilder, RayTracingDebug);
@@ -983,7 +1063,7 @@ void FDeferredShadingSceneRenderer::RenderRayTracingDebug(FRDGBuilder& GraphBuil
 
 		if (bRequiresBindings)
 		{
-			BindRayTracingDebugCHSMaterialBindings(RHICmdList, View, Pipeline);
+			BindRayTracingDebugCHSMaterialBindings(RHICmdList, View, RayGenParameters->SceneUniformBuffer->GetRHI(), Pipeline);
 			RHICmdList.SetRayTracingMissShader(View.GetRayTracingSceneChecked(), 0, Pipeline, 0 /* ShaderIndexInPipeline */, 0, nullptr, 0);
 		}
 
diff --git a/Engine/Source/Runtime/Renderer/Public/MeshPassProcessor.h b/Engine/Source/Runtime/Renderer/Public/MeshPassProcessor.h
--- a/Engine/Source/Runtime/Renderer/Public/MeshPassProcessor.h
+++ b/Engine/Source/Runtime/Renderer/Public/MeshPassProcessor.h
@@ -2187,6 +2187,7 @@ public:
 	/** Sets ray hit group shaders on the mesh command and allocates room for the shader bindings. */
 	RENDERER_API void SetShaders(const FMeshProcessorShaders& Shaders);
 
+	RENDERER_API bool IsUsingNaniteRayTracing() const;
 private:
 	FShaderUniformBufferParameter ViewUniformBufferParameter;
 	FShaderUniformBufferParameter SceneUniformBufferParameter;
