Consistency improvements to scene view extension calls.

* Removes duplicate SetupView calls from the game viewport.
* Adds missing BeginRenderViewFamily to scene captures, and update its viewfamily frame counter.
* Correctly support scene capture cube.

diff --git a/Engine/Source/Runtime/Engine/Private/GameViewportClient.cpp b/Engine/Source/Runtime/Engine/Private/GameViewportClient.cpp
--- a/Engine/Source/Runtime/Engine/Private/GameViewportClient.cpp
+++ b/Engine/Source/Runtime/Engine/Private/GameViewportClient.cpp
@@ -1693,14 +1693,7 @@ void UGameViewportClient::Draw(FViewport* InViewport, FCanvas* SceneCanvas)
 	// Draw the player views.
 	if (!bDisableWorldRendering && PlayerViewMap.Num() > 0 && FSlateApplication::Get().GetPlatformApplication()->IsAllowedToRender()) //-V560
 	{
-		// Views 
-		for (auto ViewExt : ViewFamily.ViewExtensions)
-		{
-			for (FSceneView* View : Views)
-			{
-				ViewExt->SetupView(ViewFamily, *View);
-			}
-		}
+		// Scene view extension SetupView calls already done in LocalPlayer->CalcSceneView above.
 
 		GetRendererModule().BeginRenderingViewFamily(SceneCanvas, &ViewFamily);
 	}
diff --git a/Engine/Source/Runtime/Renderer/Private/SceneCaptureRendering.cpp b/Engine/Source/Runtime/Renderer/Private/SceneCaptureRendering.cpp
--- a/Engine/Source/Runtime/Renderer/Private/SceneCaptureRendering.cpp
+++ b/Engine/Source/Runtime/Renderer/Private/SceneCaptureRendering.cpp
@@ -616,6 +616,7 @@ void SetupViewFamilyForSceneCapture(
 
 	// Initialize frame number
 	ViewFamily.FrameNumber = ViewFamily.Scene->GetFrameNumber();
+	ViewFamily.FrameCounter = GFrameCounter;
 
 	for (int32 ViewIndex = 0; ViewIndex < Views.Num(); ++ViewIndex)
 	{
@@ -936,11 +937,11 @@ void FScene::UpdateSceneCaptureContents(USceneCaptureComponent2D* CaptureCompone
 					--Index;
 				}
 			}
+		}
 
-			for (const TSharedRef<ISceneViewExtension, ESPMode::ThreadSafe>& Extension : SceneRenderer->ViewFamily.ViewExtensions)
-			{
-				Extension->SetupViewFamily(SceneRenderer->ViewFamily);
-			}
+		for (const FSceneViewExtensionRef& Extension : SceneRenderer->ViewFamily.ViewExtensions)
+		{
+			Extension->SetupViewFamily(SceneRenderer->ViewFamily);
 		}
 
 		{
@@ -1019,6 +1020,11 @@ void FScene::UpdateSceneCaptureContents(USceneCaptureComponent2D* CaptureCompone
 		void* CaptureMemorySize = nullptr;		// Dummy value for lambda capture argument list
 #endif
 
+		for (const FSceneViewExtensionRef& Extension : SceneRenderer->ViewFamily.ViewExtensions)
+		{
+			Extension->BeginRenderViewFamily(SceneRenderer->ViewFamily);
+		}
+
 		ENQUEUE_RENDER_COMMAND(CaptureCommand)(
 			[SceneRenderer, TextureRenderTargetResource, TexturePtrNotDeferenced, EventName, TargetName, bGenerateMips, GenerateMipsParams, GameViewportRT, bEnableOrthographicTiling, bIsCompositing, bOrthographicCamera, NumXTiles, NumYTiles, TileID, CaptureMemorySize](FRHICommandListImmediate& RHICmdList)
 			{
@@ -1168,6 +1174,16 @@ void FScene::UpdateSceneCaptureContents(USceneCaptureComponentCube* CaptureCompo
 
 			SceneRenderer->ViewFamily.SceneCaptureSource = CaptureComponent->CaptureSource;
 
+			for (const FSceneViewExtensionRef& Extension : SceneRenderer->ViewFamily.ViewExtensions)
+			{
+				Extension->SetupViewFamily(SceneRenderer->ViewFamily);
+
+				for (FSceneView& View : SceneRenderer->Views)
+				{
+					Extension->SetupView(SceneRenderer->ViewFamily, View);
+				}
+			}
+
 			FTextureRenderTargetCubeResource* TextureRenderTarget = static_cast<FTextureRenderTargetCubeResource*>(TextureTarget->GameThread_GetRenderTargetResource());
 			FString EventName;
 			if (!CaptureComponent->ProfilingEventName.IsEmpty())
@@ -1178,6 +1194,12 @@ void FScene::UpdateSceneCaptureContents(USceneCaptureComponentCube* CaptureCompo
 			{
 				CaptureComponent->GetOwner()->GetFName().ToString(EventName);
 			}
+
+			for (const FSceneViewExtensionRef& Extension : SceneRenderer->ViewFamily.ViewExtensions)
+			{
+				Extension->BeginRenderViewFamily(SceneRenderer->ViewFamily);
+			}
+
 			ENQUEUE_RENDER_COMMAND(CaptureCommand)(
 				[SceneRenderer, TextureRenderTarget, EventName, TargetFace](FRHICommandListImmediate& RHICmdList)
 				{
