Limit number of DF objects uploaded to GPU per frame.
- Remaining uploads are scheduled for the following frames.
- Removed existing PendingThrottledOperations since ProcessPrimitiveUpdate(...) always returned true so it was never used.

diff --git a/Engine/Source/Runtime/Renderer/Private/DistanceFieldObjectManagement.cpp b/Engine/Source/Runtime/Renderer/Private/DistanceFieldObjectManagement.cpp
--- a/Engine/Source/Runtime/Renderer/Private/DistanceFieldObjectManagement.cpp
+++ b/Engine/Source/Runtime/Renderer/Private/DistanceFieldObjectManagement.cpp
@@ -214,7 +214,7 @@ void AddModifiedBounds(FScene* Scene, FGlobalDFCacheType CacheType, const FBox&
 	DistanceFieldData.PrimitiveModifiedBounds[CacheType].Add(Bounds);
 }
 
-void UpdateGlobalDistanceFieldObjectRemoves(FScene* Scene, TArray<FSetElementId>& DistanceFieldAssetRemoves)
+void ProcessDistanceFieldObjectRemoves(FScene* Scene, TArray<FSetElementId>& DistanceFieldAssetRemoves)
 {
 	FDistanceFieldSceneData& DistanceFieldSceneData = Scene->DistanceFieldSceneData;
 
@@ -246,8 +246,11 @@ void UpdateGlobalDistanceFieldObjectRemoves(FScene* Scene, TArray<FSetElementId>
 				// InstanceIndex will be -1 with zero scale meshes
 				if (InstanceIndex >= 0)
 				{
+					// Mark region covered by instance in global distance field as modified
 					FGlobalDFCacheType CacheType = PrimitiveRemoveInfo.bOftenMoving ? GDF_Full : GDF_MostlyStatic;
 					AddModifiedBounds(Scene, CacheType, DistanceFieldSceneData.PrimitiveInstanceMapping[InstanceIndex].WorldBounds);
+
+					// Add individual instances to temporary array for processing in the next pass
 					PendingRemoveOperations.Add(InstanceIndex);
 				}
 			}
@@ -303,7 +306,7 @@ void LogDistanceFieldUpdate(FPrimitiveSceneInfo const* PrimitiveSceneInfo, float
 }
 
 /** Gathers the information needed to represent a single object's distance field and appends it to the upload buffers. */
-bool ProcessPrimitiveUpdate(
+void ProcessPrimitiveUpdate(
 	bool bIsAddOperation,
 	FScene* Scene,
 	FPrimitiveSceneInfo* PrimitiveSceneInfo,
@@ -434,7 +437,6 @@ bool ProcessPrimitiveUpdate(
 			UE_LOG(LogDistanceField,Verbose,TEXT("Primitive %s %s excluded due to huge bounding radius %f"), *Proxy->GetOwnerName().ToString(), *Proxy->GetResourceName().ToString(), BoundingRadius);
 		}
 	}
-	return true;
 }
 
 bool bVerifySceneIntegrity = false;
@@ -451,7 +453,7 @@ void FDistanceFieldSceneData::UpdateDistanceFieldObjectBuffers(
 
 	const bool bExecuteInParallel = GDFParallelUpdate != 0 && FApp::ShouldUseThreadingForPerformance();
 
-	if (HasPendingOperations() || PendingThrottledOperations.Num() > 0)
+	if (HasPendingOperations() || HasPendingUploads())
 	{
 		QUICK_SCOPE_CYCLE_COUNTER(STAT_UpdateDistanceFieldObjectBuffers);
 		RDG_EVENT_SCOPE(GraphBuilder, "UpdateDistanceFieldObjectBuffers");
@@ -461,17 +463,9 @@ void FDistanceFieldSceneData::UpdateDistanceFieldObjectBuffers(
 			ObjectBuffers = new FDistanceFieldObjectBuffers();
 		}
 
-		if (PendingAddOperations.Num() > 0)
-		{
-			PendingThrottledOperations.Reserve(PendingThrottledOperations.Num() + PendingAddOperations.Num());
-		}
-
-		PendingAddOperations.Append(PendingThrottledOperations);
-		PendingThrottledOperations.Reset();
-
 		// Process removes before adds, as the adds will overwrite primitive allocation info
 		// This also prevents re-uploading distance fields on render state recreation
-		UpdateGlobalDistanceFieldObjectRemoves(Scene, DistanceFieldAssetRemoves);
+		ProcessDistanceFieldObjectRemoves(Scene, DistanceFieldAssetRemoves);
 
 		if ((PendingAddOperations.Num() > 0 || PendingUpdateOperations.Num() > 0) && GDFReverseAtlasAllocationOrder == GDFPreviousReverseAtlasAllocationOrder)
 		{
@@ -480,17 +474,14 @@ void FDistanceFieldSceneData::UpdateDistanceFieldObjectBuffers(
 			int32 OriginalNumObjects = NumObjectsInBuffer;
 			for (FPrimitiveSceneInfo* PrimitiveSceneInfo : PendingAddOperations)
 			{
-				if (!ProcessPrimitiveUpdate(
+				ProcessPrimitiveUpdate(
 					true,
 					Scene,
 					PrimitiveSceneInfo,
 					InstanceLocalToPrimitiveTransforms,
 					IndicesToUpdateInObjectBuffers,
 					DistanceFieldAssetAdds,
-					DistanceFieldAssetRemoves))
-				{
-					PendingThrottledOperations.Add(PrimitiveSceneInfo);
-				}
+					DistanceFieldAssetRemoves);
 			}
 
 			for (FPrimitiveSceneInfo* PrimitiveSceneInfo : PendingUpdateOperations)
@@ -507,18 +498,14 @@ void FDistanceFieldSceneData::UpdateDistanceFieldObjectBuffers(
 
 			PendingAddOperations.Reset();
 			PendingUpdateOperations.Reset();
-			if (PendingThrottledOperations.Num() == 0)
-			{
-				PendingThrottledOperations.Reset();
-			}
 		}
 
 		GDFPreviousReverseAtlasAllocationOrder = GDFReverseAtlasAllocationOrder;
 
 		// Upload buffer changes
-		if (IndicesToUpdateInObjectBuffers.Num() > 0)
+		if (HasPendingUploads())
 		{
-			QUICK_SCOPE_CYCLE_COUNTER(UpdateDFObjectBuffers);
+			QUICK_SCOPE_CYCLE_COUNTER(STAT_UploadDistanceFieldObjectDataAndBounds);
 
 			// Upload DF object data and bounds
 
@@ -532,7 +519,12 @@ void FDistanceFieldSceneData::UpdateDistanceFieldObjectBuffers(
 			const uint32 DFObjectBoundsNumBytes = DFObjectBoundsNumFloat4s * sizeof(FVector4f);
 			FRDGBuffer* DFObjectBoundsBuffer = ResizeStructuredBufferIfNeeded(GraphBuilder, ObjectBuffers->Bounds, DFObjectBoundsNumBytes, TEXT("DistanceFields.DFObjectBounds"));
 
-			const int32 NumDFObjectUploads = IndicesToUpdateInObjectBuffers.Num();
+			// Limit number of distance field object uploads per frame to 2M
+			// The bottleneck is GetMaxUploadBufferElements() used by FRDGScatterUploadBuffer
+			// This is not expected to be hit during gameplay.
+			static const int32 MAX_NUM_DISTANCE_FIELD_OBJECT_UPLOADS = (2 << 20);
+
+			const int32 NumDFObjectUploads = FMath::Min(IndicesToUpdateInObjectBuffers.Num(), MAX_NUM_DISTANCE_FIELD_OBJECT_UPLOADS);
 
 			static FCriticalSection DFUpdateCS;
 
@@ -545,7 +537,7 @@ void FDistanceFieldSceneData::UpdateDistanceFieldObjectBuffers(
 
 				FParallelUpdateRangesDFO ParallelRanges;
 
-				int32 RangeCount = PartitionUpdateRangesDFO(ParallelRanges, IndicesToUpdateInObjectBuffers.Num(), bExecuteInParallel);
+				int32 RangeCount = PartitionUpdateRangesDFO(ParallelRanges, NumDFObjectUploads, bExecuteInParallel);
 
 				ParallelFor(RangeCount,
 					[this, &ParallelRanges, &PrimitiveBounds, RangeCount](int32 RangeIndex)
@@ -683,6 +675,18 @@ void FDistanceFieldSceneData::UpdateDistanceFieldObjectBuffers(
 
 				ExternalAccessQueue.Add(DFObjectDataBuffer, ERHIAccess::SRVMask, ERHIPipeline::All);
 				ExternalAccessQueue.Add(DFObjectBoundsBuffer, ERHIAccess::SRVMask, ERHIPipeline::All);
+
+				if (IndicesToUpdateInObjectBuffers.Num() > NumDFObjectUploads)
+				{
+					// this is not expected to happen frequently since we can perform up to MAX_NUM_DISTANCE_FIELD_OBJECT_UPLOADS per frame
+					// RemoveAtSwap would be more efficient but could potentially result in starvation
+					const bool bAllowShrinking = true; // allow array to shrink since getting into this code path means array is very large
+					IndicesToUpdateInObjectBuffers.RemoveAt(0, NumDFObjectUploads, bAllowShrinking);
+				}
+				else
+				{
+					IndicesToUpdateInObjectBuffers.Reset();
+				}
 			}
 		}
 
@@ -694,8 +698,6 @@ void FDistanceFieldSceneData::UpdateDistanceFieldObjectBuffers(
 			VerifyIntegrity();
 		}
 	}
-
-	IndicesToUpdateInObjectBuffers.Reset();
 }
 
 void FSceneRenderer::UpdateGlobalHeightFieldObjectBuffers(FRDGBuilder& GraphBuilder, const TArray<uint32>& IndicesToUpdateInHeightFieldObjectBuffers)
diff --git a/Engine/Source/Runtime/Renderer/Private/RendererScene.cpp b/Engine/Source/Runtime/Renderer/Private/RendererScene.cpp
--- a/Engine/Source/Runtime/Renderer/Private/RendererScene.cpp
+++ b/Engine/Source/Runtime/Renderer/Private/RendererScene.cpp
@@ -643,7 +643,6 @@ void FDistanceFieldSceneData::RemovePrimitive(FPrimitiveSceneInfo* InPrimitive)
 		{
 			PendingAddOperations.Remove(InPrimitive);
 			PendingUpdateOperations.Remove(InPrimitive);
-			PendingThrottledOperations.Remove(InPrimitive);
 
 			if (InPrimitive->DistanceFieldInstanceIndices.Num() > 0)
 			{
@@ -679,8 +678,7 @@ void FDistanceFieldSceneData::RemovePrimitive(FPrimitiveSceneInfo* InPrimitive)
 
 	checkf(!PendingAddOperations.Contains(InPrimitive), TEXT("Primitive is being removed from the scene, but didn't remove from Distance Field Scene properly - a crash will occur when processing PendingAddOperations.  This can happen if the proxy's properties have changed without recreating its render state."));
 	checkf(!PendingUpdateOperations.Contains(InPrimitive), TEXT("Primitive is being removed from the scene, but didn't remove from Distance Field Scene properly - a crash will occur when processing PendingUpdateOperations.  This can happen if the proxy's properties have changed without recreating its render state."));
-	checkf(!PendingThrottledOperations.Contains(InPrimitive), TEXT("Primitive is being removed from the scene, but didn't remove from Distance Field Scene properly - a crash will occur when processing PendingThrottledOperations.  This can happen if the proxy's properties have changed without recreating its render state."));
-	
+
 	checkf(!PendingHeightFieldAddOps.Contains(InPrimitive), TEXT("Primitive is being removed from the scene, but didn't remove from Distance Field Scene properly - a crash will occur when processing PendingHeightFieldAddOps.  This can happen if the proxy's properties have changed without recreating its render state."));
 }
 
diff --git a/Engine/Source/Runtime/Renderer/Private/ScenePrivate.h b/Engine/Source/Runtime/Renderer/Private/ScenePrivate.h
--- a/Engine/Source/Runtime/Renderer/Private/ScenePrivate.h
+++ b/Engine/Source/Runtime/Renderer/Private/ScenePrivate.h
@@ -2002,6 +2002,11 @@ public:
 		TArray<FDistanceFieldAssetMipId>& DistanceFieldAssetAdds,
 		TArray<FSetElementId>& DistanceFieldAssetRemoves);
 
+	bool HasPendingUploads() const
+	{
+		return IndicesToUpdateInObjectBuffers.Num() > 0;
+	}
+
 	bool HasPendingOperations() const
 	{
 		return PendingAddOperations.Num() > 0 || PendingUpdateOperations.Num() > 0 || PendingRemoveOperations.Num() > 0;
@@ -2087,7 +2092,6 @@ public:
 	TArray<FPrimitiveSceneInfo*> HeightfieldPrimitives;
 	/** Pending operations on the object buffers to be processed next frame. */
 	TSet<FPrimitiveSceneInfo*> PendingAddOperations;
-	TSet<FPrimitiveSceneInfo*> PendingThrottledOperations;
 	TSet<FPrimitiveSceneInfo*> PendingUpdateOperations;
 	TArray<FPrimitiveRemoveInfo> PendingRemoveOperations;
 	TArray<FBox> PrimitiveModifiedBounds[GDF_Num];
