Limit number of instances in Distance Field Scene to MAX_INSTANCE_ID
- Also round up NumObjectsInBuffer to next power of two before multiplying by GDistanceFieldObjectDataStride to prevent uint32 overflows.

diff --git a/Engine/Source/Runtime/Renderer/Private/DistanceFieldObjectManagement.cpp b/Engine/Source/Runtime/Renderer/Private/DistanceFieldObjectManagement.cpp
--- a/Engine/Source/Runtime/Renderer/Private/DistanceFieldObjectManagement.cpp
+++ b/Engine/Source/Runtime/Renderer/Private/DistanceFieldObjectManagement.cpp
@@ -363,13 +363,22 @@ void ProcessPrimitiveUpdate(
 
 			for (int32 TransformIndex = 0; TransformIndex < InstanceLocalToPrimitiveTransforms.Num(); TransformIndex++)
 			{
+				const bool bInstanceCountOverflow = bIsAddOperation && (DistanceFieldSceneData.NumObjectsInBuffer + 1 > MAX_INSTANCE_ID);
+
+				static bool bWarnOnce = true;
+				if (bInstanceCountOverflow && bWarnOnce)
+				{
+					bWarnOnce = false;
+					UE_LOG(LogDistanceField, Warning, TEXT("Max instance count in Distance Field Scene reached. New instances might not be represented."));
+				}
+
 				const FMatrix LocalToWorld = InstanceLocalToPrimitiveTransforms[TransformIndex].ToMatrix() * Proxy->GetLocalToWorld();
 
 				const FMatrix::FReal MaxScale = LocalToWorld.GetMaximumAxisScale();
 
-				if (bIsAddOperation && MaxScale <= 0)
+				if (bIsAddOperation && (MaxScale <= 0 || bInstanceCountOverflow))
 				{
-					// Skip degenerate instances
+					// Skip degenerate instances or when instance count limit is reached
 					PrimitiveSceneInfo->DistanceFieldInstanceIndices[TransformIndex] = -1;
 					continue;
 				}
@@ -508,15 +517,14 @@ void FDistanceFieldSceneData::UpdateDistanceFieldObjectBuffers(
 			QUICK_SCOPE_CYCLE_COUNTER(STAT_UploadDistanceFieldObjectDataAndBounds);
 
 			// Upload DF object data and bounds
+			check(NumObjectsInBuffer <= MAX_INSTANCE_ID);
 
-			const uint32 NumDFObjects = NumObjectsInBuffer;
+			const uint32 NumDFObjectsRoundedUp = FMath::RoundUpToPowerOfTwo(NumObjectsInBuffer);
 
-			const uint32 DFObjectDataNumFloat4s = FMath::RoundUpToPowerOfTwo(NumDFObjects * GDistanceFieldObjectDataStride);
-			const uint32 DFObjectDataNumBytes = DFObjectDataNumFloat4s * sizeof(FVector4f);
+			const uint32 DFObjectDataNumBytes = NumDFObjectsRoundedUp * GDistanceFieldObjectDataStride * sizeof(FVector4f);
 			FRDGBuffer* DFObjectDataBuffer = ResizeStructuredBufferIfNeeded(GraphBuilder, ObjectBuffers->Data, DFObjectDataNumBytes, TEXT("DistanceFields.DFObjectData"));
 
-			const uint32 DFObjectBoundsNumFloat4s = FMath::RoundUpToPowerOfTwo(NumDFObjects * GDistanceFieldObjectBoundsStride);
-			const uint32 DFObjectBoundsNumBytes = DFObjectBoundsNumFloat4s * sizeof(FVector4f);
+			const uint32 DFObjectBoundsNumBytes = NumDFObjectsRoundedUp * GDistanceFieldObjectBoundsStride * sizeof(FVector4f);
 			FRDGBuffer* DFObjectBoundsBuffer = ResizeStructuredBufferIfNeeded(GraphBuilder, ObjectBuffers->Bounds, DFObjectBoundsNumBytes, TEXT("DistanceFields.DFObjectBounds"));
 
 			// Limit number of distance field object uploads per frame to 2M
