Broke circular dependency:
MaterialShaderQualitySettings -> PropertyEditor -> UnrealEd -> Engine -> Renderer -> MaterialShaderQualitySettings

diff --git a/Engine/Source/Runtime/Renderer/Renderer.Build.cs b/Engine/Source/Runtime/Renderer/Renderer.Build.cs
--- a/Engine/Source/Runtime/Renderer/Renderer.Build.cs
+++ b/Engine/Source/Runtime/Renderer/Renderer.Build.cs
@@ -15,7 +15,6 @@ public class Renderer : ModuleRules
 
 		PublicDependencyModuleNames.Add("Core");
         PublicDependencyModuleNames.Add("Engine");
-        PublicDependencyModuleNames.Add("MaterialShaderQualitySettings");
 
         if (Target.bBuildEditor == true)
         {
@@ -35,8 +34,9 @@ public class Renderer : ModuleRules
 				"ApplicationCore",
 				"RenderCore", 
 				"ImageWriteQueue",
-				"RHI"
-            }
+				"RHI",
+				"MaterialShaderQualitySettings"
+			}
             );
 
         PrivateIncludePathModuleNames.AddRange(new string[] { "HeadMountedDisplay" });
