Fixed deadlock in visibility task graph when no primitives are in the scene.

diff --git a/Engine/Source/Runtime/Renderer/Private/SceneVisibility.cpp b/Engine/Source/Runtime/Renderer/Private/SceneVisibility.cpp
--- a/Engine/Source/Runtime/Renderer/Private/SceneVisibility.cpp
+++ b/Engine/Source/Runtime/Renderer/Private/SceneVisibility.cpp
@@ -3189,6 +3189,9 @@ FVisibilityViewPacket::FVisibilityViewPacket(FVisibilityTaskData& InTaskData, FS
 			}
 		});
 
+		// Take a reference on the occlusion command pipe that is released when all frustum cull commands are enqueued.
+		OcclusionCull.CommandPipe.AddNumCommands(1);
+
 		// Callback for when a relevance command is queued from occlusion.
 		Relevance.CommandPipe.SetCommandFunction([this](FPrimitiveIndexList&& PrimitiveIndexList)
 		{
@@ -3382,6 +3385,8 @@ void FVisibilityViewPacket::BeginInitVisibility()
 
 			}, PrerequisiteTask, TaskConfig.TaskPriority, ExtendedTaskPriority));
 		}
+
+		OcclusionCull.CommandPipe.ReleaseNumCommands(1);
 	}
 	else
 	{
