Fix the instance overflow warning for the idle scene in PIE

diff --git a/Engine/Source/Runtime/Renderer/Private/GPUScene.cpp b/Engine/Source/Runtime/Renderer/Private/GPUScene.cpp
--- a/Engine/Source/Runtime/Renderer/Private/GPUScene.cpp
+++ b/Engine/Source/Runtime/Renderer/Private/GPUScene.cpp
@@ -568,14 +568,14 @@ FGPUScene::FGPUScene()
 #if !UE_BUILD_SHIPPING
 	ScreenMessageDelegate = FRendererOnScreenNotification::Get().AddLambda([this](TMultiMap<FCoreDelegates::EOnScreenMessageSeverity, FText >& OutMessages)
 	{
-		if (InstanceSceneDataSOAStride < uint32(InstanceSceneDataAllocator.GetMaxSize()))
+		if (InstanceSceneDataSOAStride < MaxInstancesDuringPrevUpdate)
 		{
 			OutMessages.Add(FCoreDelegates::EOnScreenMessageSeverity::Warning, FText::FromString(FString::Printf(TEXT("GPU-Scene Instance data overflow detected, reduce instance count to avoid rendering artifacts"))));
-			OutMessages.Add(FCoreDelegates::EOnScreenMessageSeverity::Warning, FText::FromString(FString::Printf(TEXT("  Max allocated ID %d, instance buffer size: %dM"), InstanceSceneDataAllocator.GetMaxSize(), InstanceSceneDataSOAStride >> 20)));
+			OutMessages.Add(FCoreDelegates::EOnScreenMessageSeverity::Warning, FText::FromString(FString::Printf(TEXT("  Max allocated ID %d, instance buffer size: %dM"), MaxInstancesDuringPrevUpdate, InstanceSceneDataSOAStride >> 20)));
 			if (!bLoggedInstanceOverflow )
 			{
 				UE_LOG(LogRenderer, Warning, TEXT("GPU-Scene Instance data overflow detected, reduce instance count to avoid rendering artifacts.\n")
-					TEXT(" Max allocated ID %d (%0.3fM), instance buffer size: %dM"), InstanceSceneDataAllocator.GetMaxSize(), double(InstanceSceneDataAllocator.GetMaxSize()) / (1024.0 * 1024.0),  InstanceSceneDataSOAStride >> 20);
+					TEXT(" Max allocated ID %d (%0.3fM), instance buffer size: %dM"), MaxInstancesDuringPrevUpdate, double(MaxInstancesDuringPrevUpdate) / (1024.0 * 1024.0),  InstanceSceneDataSOAStride >> 20);
 				bLoggedInstanceOverflow = true;
 			}
 		}
@@ -703,6 +703,10 @@ void FGPUScene::UpdateInternal(FRDGBuilder& GraphBuilder, FSceneUniformBuffer& S
 	// Do this stat separately since it should always be on.
 	CSV_CUSTOM_STAT_GLOBAL(GPUSceneInstanceCount, float(InstanceSceneDataAllocator.GetMaxSize()), ECsvCustomStatOp::Set);
 
+#if !UE_BUILD_SHIPPING
+	MaxInstancesDuringPrevUpdate = uint32(InstanceSceneDataAllocator.GetMaxSize());
+#endif // UE_BUILD_SHIPPING
+
 	LastDeferredGPUWritePass = EGPUSceneGPUWritePass::None;
 
 	if ((CVarGPUSceneUploadEveryFrame.GetValueOnRenderThread() != 0) || bUpdateAllPrimitives)
diff --git a/Engine/Source/Runtime/Renderer/Private/GPUScene.h b/Engine/Source/Runtime/Renderer/Private/GPUScene.h
--- a/Engine/Source/Runtime/Renderer/Private/GPUScene.h
+++ b/Engine/Source/Runtime/Renderer/Private/GPUScene.h
@@ -432,6 +432,7 @@ private:
 #if !UE_BUILD_SHIPPING
 	FDelegateHandle ScreenMessageDelegate;
 	bool bLoggedInstanceOverflow = false;
+	uint32 MaxInstancesDuringPrevUpdate = 0;
 #endif // UE_BUILD_SHIPPING
 
 };
