Mobile Light Prepass
Disabled by default now

diff --git a/Engine/Shaders/Private/MobileBasePassPixelShader.usf b/Engine/Shaders/Private/MobileBasePassPixelShader.usf
--- a/Engine/Shaders/Private/MobileBasePassPixelShader.usf
+++ b/Engine/Shaders/Private/MobileBasePassPixelShader.usf
@@ -811,9 +811,32 @@ void Main(
 #endif
 
 	// Local lights
-#if  !MATERIAL_SHADINGMODEL_SINGLELAYERWATER && ENABLE_CLUSTERED_LIGHTS
+#if  PREPASS_LOCAL_LIGHTS_MOBILE
+	{
+		FDeferredLightData MergedLightData = (FDeferredLightData)0;
+		half3 LocalLightA = Texture2DSample(MobileSceneTextures.LocalLightTextureA, MobileSceneTextures.LocalLightTextureASampler,SvPositionToBufferUV(SvPosition)); 
+		half4 LocalLightB = Texture2DSample(MobileSceneTextures.LocalLightTextureB, MobileSceneTextures.LocalLightTextureBSampler,SvPositionToBufferUV(SvPosition));
+	
+		MergedLightData.SourceRadius	= 0;
+		MergedLightData.SourceLength	= 0;
+		MergedLightData.Tangent			= 0;
+		MergedLightData.SoftSourceRadius= 0;
+		MergedLightData.SpecularScale	= 1.0f;
+		MergedLightData.bSpotLight		= false;
+		MergedLightData.bRectLight		= false;
+		MergedLightData.Color = LocalLightA.rgb;
+		MergedLightData.bRadialLight = false;
+
+		half4 LightAttenuation = half4(1, 1, 1, 1); 
+		float SurfaceShadow = 0;
+
+		MergedLightData.Direction = LocalLightB.rgb * 2 - 1;
+
+		FLightAccumulator NewLighting = AccumulateDynamicLighting(MaterialParameters.WorldPosition_CamRelative, CameraVector, GBuffer, 1, GBuffer.ShadingModelID, MergedLightData, LightAttenuation, 0, uint2(0, 0), SurfaceShadow);
+		DirectLighting = LightAccumulator_Add(DirectLighting, NewLighting);
+	}
+#elif  (!MATERIAL_SHADINGMODEL_SINGLELAYERWATER && ENABLE_CLUSTERED_LIGHTS)
 	{
-
 #if MOBILE_MULTI_VIEW
 		const uint EyeIndex = BasePassInterpolants.MultiViewId;
 #else
diff --git a/Engine/Shaders/Private/MobileLocalLightPrepass.usf b/Engine/Shaders/Private/MobileLocalLightPrepass.usf
new file mode 100644
--- /dev/null
+++ b/Engine/Shaders/Private/MobileLocalLightPrepass.usf
@@ -0,0 +1,74 @@
+// Copyright Epic Games, Inc. All Rights Reserved.
+
+#include "Common.ush"
+#include "SceneTextureParameters.ush"
+#include "LightGridCommon.ush"
+#include "DeferredLightingCommon.ush"
+
+void Main(float4 SvPosition : SV_POSITION, out float3 OutColorA : SV_Target0, out float4 OutColorB : SV_Target1)
+{
+	const uint2 PixelPos = SvPosition.xy;
+	const float2 ScreenPosition = SvPositionToScreenPosition(SvPosition).xy;
+	const uint EyeIndex = 0;
+
+	float2 ScreenUV = ScreenPosition * View.ScreenPositionScaleBias.xy + View.ScreenPositionScaleBias.wz;
+	float SceneDepth = CalcSceneDepth(ScreenUV);
+
+	float3 TranslatedWorldPosition = mul(float4(ScreenPosition * SceneDepth, SceneDepth, 1), PrimaryView.ScreenToTranslatedWorld).xyz;
+
+	uint GridIndex = ComputeLightGridCellIndex(uint2(PixelPos.x, PixelPos.y), SceneDepth);
+	const FCulledLightsGridData CulledLightGridData = GetCulledLightsGrid(GridIndex, EyeIndex);
+
+	const uint NumLocalLights = min(CulledLightGridData.NumLocalLights, GetNumLocalLights(EyeIndex));
+
+	float3 Color = uint3(0.f, 0.f , 0.f);
+	float3 SavedL = float3(0,0,0);
+	float TotalWeight = 0;
+
+	LOOP
+	for (uint LocalLightListIndex = 0; LocalLightListIndex < NumLocalLights; LocalLightListIndex++)
+	{
+		const FLocalLightData LocalLight = GetLocalLightData(CulledLightGridData.DataStartIndex + LocalLightListIndex, EyeIndex);
+
+		// extra-early out since we know light grid is sloppy and all lights in list are radial (have a range)
+		// appears useless
+		float InvLightRadiusSq = LocalLight.LightPositionAndInvRadius.w * LocalLight.LightPositionAndInvRadius.w;
+		float DistLight = length2(TranslatedWorldPosition - LocalLight.LightPositionAndInvRadius.xyz) * InvLightRadiusSq;
+		if (DistLight > 1.0f)
+		{
+			continue;
+		}
+		FDeferredLightData LightData = ConvertToDeferredLight_Mobile(LocalLight);
+
+		float3 L = LightData.Direction;	// Already normalized
+		float3 ToLight = L;
+		float3 MaskedLightColor = LightData.Color;
+		float LightMask = 1;
+		if (LightData.bRadialLight)
+		{
+			LightMask = GetLocalLightAttenuation( TranslatedWorldPosition, LightData, ToLight, L );
+			MaskedLightColor *= LightMask;
+		}
+
+		float WeightLight = 1 - DistLight;
+		
+		if( LightMask > 0)
+		{
+			TotalWeight += WeightLight;
+			Color += MaskedLightColor;
+			SavedL += L *  WeightLight;
+		}
+	}
+
+	if (TotalWeight > 0)
+	{
+		SavedL = (SavedL / TotalWeight);
+		OutColorA = float3(Color.x, Color.y, Color.z);
+		OutColorB = 0.5f * float4(SavedL.x, SavedL.y, SavedL.z, 0)  + float4(0.5f, 0.5f, 0.5f, 0.5f);
+	}
+	else
+	{
+		OutColorA = float3(0, 0, 0);
+		OutColorB = float4(0.5,0.5,0.5, 0.5);
+	}
+}
diff --git a/Engine/Source/Editor/MaterialEditor/Private/MaterialStatsCommon.cpp b/Engine/Source/Editor/MaterialEditor/Private/MaterialStatsCommon.cpp
--- a/Engine/Source/Editor/MaterialEditor/Private/MaterialStatsCommon.cpp
+++ b/Engine/Source/Editor/MaterialEditor/Private/MaterialStatsCommon.cpp
@@ -267,7 +267,7 @@ FSlateColor FMaterialStatsUtils::QualitySettingColor(const EMaterialQualityLevel
 	return  FStyleColors::Foreground;
 }
 
-static void MobileBasePassShaderName(bool bVertexShader, const TCHAR* PolicyName, bool bEnableLocalLights, bool bHDR, bool bSkyLight, FString& OutName)
+static void MobileBasePassShaderName(bool bVertexShader, const TCHAR* PolicyName, const TCHAR* LocalLightSetting, bool bHDR, bool bSkyLight, FString& OutName)
 {
 	OutName.Reset();
 	OutName.Append(bVertexShader ? TEXT("TMobileBasePassVS") : TEXT("TMobileBasePassPS"));
@@ -277,11 +277,10 @@ static void MobileBasePassShaderName(bool bVertexShader, const TCHAR* PolicyName
 	{
 		OutName.Append(TEXT("SkyLight"));
 	}
-	if (!bVertexShader && bEnableLocalLights)
+	if (!bVertexShader)
 	{
-		OutName.Append(TEXT("EnableLocalLights"));
+		OutName.Append(LocalLightSetting);
 	}
-	
 }
 
 void FMaterialStatsUtils::GetRepresentativeShaderTypesAndDescriptions(TMap<FName, TArray<FRepresentativeShaderInfo>>& ShaderTypeNamesAndDescriptions, const FMaterial* TargetMaterial)
@@ -378,15 +377,13 @@ void FMaterialStatsUtils::GetRepresentativeShaderTypesAndDescriptions(TMap<FName
 
 		if (TargetMaterial->GetShadingModels().IsUnlit())
 		{
-			bool bEnableLocalLights = false;
-			
 			//unlit materials are never lightmapped
-			MobileBasePassShaderName(false, TEXT("FNoLightMapPolicy"), bEnableLocalLights, bMobileHDR, false, ShaderNameStr);
+			MobileBasePassShaderName(false, TEXT("FNoLightMapPolicy"), TEXT("LOCAL_LIGHTS_DISABLED"), bMobileHDR, false, ShaderNameStr);
 			const FString Description = FString::Printf(TEXT("Mobile base pass shader without light map%s"), DescSuffix);
 			ShaderTypeNamesAndDescriptions.Add(FLocalVertexFactoryName)
 				.Add(FRepresentativeShaderInfo(ERepresentativeShader::StationarySurface, FName(ShaderNameStr), Description));
 
-			MobileBasePassShaderName(true, TEXT("FNoLightMapPolicy"), bEnableLocalLights, bMobileHDR, false, ShaderNameStr);
+			MobileBasePassShaderName(true, TEXT("FNoLightMapPolicy"), TEXT("LOCAL_LIGHTS_DISABLED"), bMobileHDR, false, ShaderNameStr);
 			ShaderTypeNamesAndDescriptions.FindOrAdd(FLocalVertexFactoryName)
 				.Add(FRepresentativeShaderInfo(ERepresentativeShader::StaticMesh, FName(ShaderNameStr),
 					FString::Printf(TEXT("Mobile base pass vertex shader%s"), DescSuffix)));
@@ -421,8 +418,7 @@ void FMaterialStatsUtils::GetRepresentativeShaderTypesAndDescriptions(TMap<FName
 				{
 					// distance field shadows only shaders
 					{
-						bool bEnableLocalLights = false;
-						MobileBasePassShaderName(false, TEXT("FMobileDistanceFieldShadowsAndLQLightMapPolicy"), bEnableLocalLights, bMobileHDR, bOnlySkyPermutation, ShaderNameStr);
+						MobileBasePassShaderName(false, TEXT("FMobileDistanceFieldShadowsAndLQLightMapPolicy"), TEXT("LOCAL_LIGHTS_DISABLED"), bMobileHDR, bOnlySkyPermutation, ShaderNameStr);
 						const FString Description = FString::Printf(TEXT("Mobile base pass shader with distance field shadows%s"), DescSuffix);
 						ShaderTypeNamesAndDescriptions.FindOrAdd(FLocalVertexFactoryName)
 							.Add(FRepresentativeShaderInfo(ERepresentativeShader::StationarySurface, FName(ShaderNameStr), Description));
@@ -434,16 +430,14 @@ void FMaterialStatsUtils::GetRepresentativeShaderTypesAndDescriptions(TMap<FName
 					{
 						// distance field shadows & CSM shaders
 						{
-							bool bEnableLocalLights = false;
-							MobileBasePassShaderName(false, TEXT("FMobileDistanceFieldShadowsLightMapAndCSMLightingPolicy"), bEnableLocalLights, bMobileHDR, bOnlySkyPermutation, ShaderNameStr);
+							MobileBasePassShaderName(false, TEXT("FMobileDistanceFieldShadowsLightMapAndCSMLightingPolicy"), TEXT("LOCAL_LIGHTS_DISABLED"), bMobileHDR, bOnlySkyPermutation, ShaderNameStr);
 							const FString Description = FString::Printf(TEXT("Mobile base pass shader with distance field shadows and CSM%s"), DescSuffix);
 							ShaderTypeNamesAndDescriptions.FindOrAdd(FLocalVertexFactoryName)
 								.Add(FRepresentativeShaderInfo(ERepresentativeShader::StationarySurfaceCSM, FName(ShaderNameStr), Description));
 						}
 
 						{
-							bool bEnableLocalLights = true;
-							MobileBasePassShaderName(false, TEXT("FMobileDistanceFieldShadowsAndLQLightMapPolicy"), bEnableLocalLights, bMobileHDR, bOnlySkyPermutation, ShaderNameStr);
+							MobileBasePassShaderName(false, TEXT("FMobileDistanceFieldShadowsAndLQLightMapPolicy"), TEXT("LOCAL_LIGHTS_ENABLED"), bMobileHDR, bOnlySkyPermutation, ShaderNameStr);
 							const FString Description = FString::Printf(TEXT("Mobile base pass shader with distance field shadows, CSM and local light(s) %s"), DescSuffix);
 
 							FRepresentativeShaderInfo ShaderInfo = FRepresentativeShaderInfo(ERepresentativeShader::StationarySurfaceNPointLights, FName(ShaderNameStr), Description);
@@ -451,21 +445,36 @@ void FMaterialStatsUtils::GetRepresentativeShaderTypesAndDescriptions(TMap<FName
 							ShaderTypeNamesAndDescriptions.FindOrAdd(FLocalVertexFactoryName).Add(ShaderInfo);
 								
 						}
+
+						{
+							MobileBasePassShaderName(false, TEXT("FMobileDistanceFieldShadowsAndLQLightMapPolicy"), TEXT("LOCAL_LIGHTS_PREPASS_ENABLED"), bMobileHDR, bOnlySkyPermutation, ShaderNameStr);
+							const FString Description = FString::Printf(TEXT("Mobile base pass shader with distance field shadows, CSM and local light(s) %s"), DescSuffix);
+
+							FRepresentativeShaderInfo ShaderInfo = FRepresentativeShaderInfo(ERepresentativeShader::StationarySurfaceNPointLights, FName(ShaderNameStr), Description);
+
+							ShaderTypeNamesAndDescriptions.FindOrAdd(FLocalVertexFactoryName).Add(ShaderInfo);
+
+						}
 					}
 				}
 				else //no shadows & lightmapped
 				{
 					{
-						bool bEnableLocalLights = false;
-						MobileBasePassShaderName(false, TEXT("TLightMapPolicyLQ"), bEnableLocalLights, bMobileHDR, bOnlySkyPermutation, ShaderNameStr);
+						MobileBasePassShaderName(false, TEXT("TLightMapPolicyLQ"), TEXT("LOCAL_LIGHTS_DISABLED"), bMobileHDR, bOnlySkyPermutation, ShaderNameStr);
 						ShaderTypeNamesAndDescriptions.FindOrAdd(FLocalVertexFactoryName)
 							.Add(FRepresentativeShaderInfo(ERepresentativeShader::StationarySurface, FName(ShaderNameStr),
 								FString::Printf(TEXT("Mobile base pass shader with static lighting%s"), DescSuffix)));
 					}
 
 					{
-						bool bEnableLocalLights = true;
-						MobileBasePassShaderName(false, TEXT("TLightMapPolicyLQ"), bEnableLocalLights, bMobileHDR, bOnlySkyPermutation, ShaderNameStr);
+						MobileBasePassShaderName(false, TEXT("TLightMapPolicyLQ"), TEXT("LOCAL_LIGHTS_ENABLED"), bMobileHDR, bOnlySkyPermutation, ShaderNameStr);
+						const FString Description = FString::Printf(TEXT("Mobile base pass shader with static lighting and local light(s) %s"), DescSuffix);
+						FRepresentativeShaderInfo ShaderInfo = FRepresentativeShaderInfo(ERepresentativeShader::StationarySurfaceNPointLights, FName(ShaderNameStr), Description);
+						ShaderTypeNamesAndDescriptions.FindOrAdd(FLocalVertexFactoryName).Add(ShaderInfo);
+					}
+
+					{
+						MobileBasePassShaderName(false, TEXT("TLightMapPolicyLQ"), TEXT("LOCAL_LIGHTS_PREPASS_ENABLED"), bMobileHDR, bOnlySkyPermutation, ShaderNameStr);
 						const FString Description = FString::Printf(TEXT("Mobile base pass shader with static lighting and local light(s) %s"), DescSuffix);
 						FRepresentativeShaderInfo ShaderInfo = FRepresentativeShaderInfo(ERepresentativeShader::StationarySurfaceNPointLights, FName(ShaderNameStr), Description);
 						ShaderTypeNamesAndDescriptions.FindOrAdd(FLocalVertexFactoryName).Add(ShaderInfo);
@@ -476,13 +485,12 @@ void FMaterialStatsUtils::GetRepresentativeShaderTypesAndDescriptions(TMap<FName
 			// only one of these 2 shader types will be displayed
 			
 			// dynamically lit shader NoLightmapPolicy
-			bool bEnableLocalLights = false;
-			MobileBasePassShaderName(false, TEXT("FNoLightmapPolicy"), bEnableLocalLights, bMobileHDR, bOnlySkyPermutation, ShaderNameStr);
+			MobileBasePassShaderName(false, TEXT("FNoLightmapPolicy"), TEXT("LOCAL_LIGHTS_DISABLED"), bMobileHDR, bOnlySkyPermutation, ShaderNameStr);
 			ShaderTypeNamesAndDescriptions.FindOrAdd(FLocalVertexFactoryName)
 				.Add(FRepresentativeShaderInfo(ERepresentativeShader::DynamicallyLitObject, FName(ShaderNameStr),
 				FString::Printf(TEXT("Mobile base pass shader with only dynamic lighting%s"), DescSuffix)));
 
-			MobileBasePassShaderName(true, TEXT("FNoLightMapPolicy"), bEnableLocalLights, bMobileHDR, bOnlySkyPermutation, ShaderNameStr);
+			MobileBasePassShaderName(true, TEXT("FNoLightMapPolicy"), TEXT("LOCAL_LIGHTS_DISABLED"), bMobileHDR, bOnlySkyPermutation, ShaderNameStr);
 			ShaderTypeNamesAndDescriptions.FindOrAdd(FLocalVertexFactoryName)
 				.Add(FRepresentativeShaderInfo(ERepresentativeShader::StaticMesh, FName(ShaderNameStr),
 				FString::Printf(TEXT("Mobile base pass vertex shader%s"), DescSuffix)));
@@ -501,12 +509,12 @@ void FMaterialStatsUtils::GetRepresentativeShaderTypesAndDescriptions(TMap<FName
 			}
 
 			// dynamically lit shader FMobileDirectionalLightAndCSMPolicy
-			MobileBasePassShaderName(false, TEXT("FMobileDirectionalLightAndCSMPolicy"), bEnableLocalLights, bMobileHDR, bOnlySkyPermutation, ShaderNameStr);
+			MobileBasePassShaderName(false, TEXT("FMobileDirectionalLightAndCSMPolicy"), TEXT("LOCAL_LIGHTS_DISABLED"), bMobileHDR, bOnlySkyPermutation, ShaderNameStr);
 			ShaderTypeNamesAndDescriptions.FindOrAdd(FLocalVertexFactoryName)
 				.Add(FRepresentativeShaderInfo(ERepresentativeShader::DynamicallyLitObject, FName(ShaderNameStr),
 				FString::Printf(TEXT("Mobile base pass shader with only dynamic lighting%s"), DescSuffix)));
 
-			MobileBasePassShaderName(true, TEXT("FMobileDirectionalLightAndCSMPolicy"), bEnableLocalLights, bMobileHDR, bOnlySkyPermutation, ShaderNameStr);
+			MobileBasePassShaderName(true, TEXT("FMobileDirectionalLightAndCSMPolicy"), TEXT("LOCAL_LIGHTS_DISABLED"), bMobileHDR, bOnlySkyPermutation, ShaderNameStr);
 			ShaderTypeNamesAndDescriptions.FindOrAdd(FLocalVertexFactoryName)
 				.Add(FRepresentativeShaderInfo(ERepresentativeShader::StaticMesh, FName(ShaderNameStr),
 				FString::Printf(TEXT("Mobile base pass vertex shader%s"), DescSuffix)));
diff --git a/Engine/Source/Runtime/Core/Private/HAL/ConsoleManager.cpp b/Engine/Source/Runtime/Core/Private/HAL/ConsoleManager.cpp
--- a/Engine/Source/Runtime/Core/Private/HAL/ConsoleManager.cpp
+++ b/Engine/Source/Runtime/Core/Private/HAL/ConsoleManager.cpp
@@ -2846,7 +2846,9 @@ static TAutoConsoleVariable<int32> CVarMobileEnabledShadingModelsMask(
 static TAutoConsoleVariable<int32> CVarMobileForwardEnableLocalLights(
 	TEXT("r.Mobile.Forward.EnableLocalLights"),
 	1,
-	TEXT("Enable local lights support on mobile forward. 0 is disabled, 1 is enabled (default)"),
+	TEXT("0: Local Lights Disabled (default)\n"
+		"1: Local Lights Enabled\n"
+		"2: Local Lights Prepass Enabled\n"),
 	ECVF_ReadOnly | ECVF_RenderThreadSafe);
 
 static TAutoConsoleVariable<int32> CVarMobileForwardEnableClusteredReflections(
diff --git a/Engine/Source/Runtime/Engine/Public/SceneTexturesConfig.h b/Engine/Source/Runtime/Engine/Public/SceneTexturesConfig.h
--- a/Engine/Source/Runtime/Engine/Public/SceneTexturesConfig.h
+++ b/Engine/Source/Runtime/Engine/Public/SceneTexturesConfig.h
@@ -55,11 +55,15 @@ BEGIN_GLOBAL_SHADER_PARAMETER_STRUCT(FMobileSceneTextureUniformParameters, ENGIN
 	SHADER_PARAMETER_RDG_TEXTURE(Texture2D, GBufferCTexture)
 	SHADER_PARAMETER_RDG_TEXTURE(Texture2D, GBufferDTexture)
 	SHADER_PARAMETER_RDG_TEXTURE(Texture2D, SceneDepthAuxTexture)
+	SHADER_PARAMETER_RDG_TEXTURE(Texture2D, LocalLightTextureA)
+	SHADER_PARAMETER_RDG_TEXTURE(Texture2D, LocalLightTextureB)
 	SHADER_PARAMETER_SAMPLER(SamplerState, GBufferATextureSampler)
 	SHADER_PARAMETER_SAMPLER(SamplerState, GBufferBTextureSampler)
 	SHADER_PARAMETER_SAMPLER(SamplerState, GBufferCTextureSampler)
 	SHADER_PARAMETER_SAMPLER(SamplerState, GBufferDTextureSampler)
 	SHADER_PARAMETER_SAMPLER(SamplerState, SceneDepthAuxTextureSampler)
+	SHADER_PARAMETER_SAMPLER(SamplerState, LocalLightTextureASampler)
+	SHADER_PARAMETER_SAMPLER(SamplerState, LocalLightTextureBSampler)
 END_GLOBAL_SHADER_PARAMETER_STRUCT()
 
 BEGIN_SHADER_PARAMETER_STRUCT(FSceneTextureShaderParameters, ENGINE_API)
diff --git a/Engine/Source/Runtime/Renderer/Private/DebugViewModeRendering.cpp b/Engine/Source/Runtime/Renderer/Private/DebugViewModeRendering.cpp
--- a/Engine/Source/Runtime/Renderer/Private/DebugViewModeRendering.cpp
+++ b/Engine/Source/Runtime/Renderer/Private/DebugViewModeRendering.cpp
@@ -472,7 +472,7 @@ void FDebugViewModeMeshProcessor::UpdateInstructionCount(FDebugViewModeShaderEle
 		{
 			TShaderRef<TMobileBasePassVSPolicyParamType<FUniformLightMapPolicy>> MobileVS;
 			TShaderRef<TMobileBasePassPSPolicyParamType<FUniformLightMapPolicy>> MobilePS;
-			if (MobileBasePass::GetShaders(LMP_NO_LIGHTMAP, 0, *InBatchMaterial, InVertexFactoryType, false, MobileVS, MobilePS))
+			if (MobileBasePass::GetShaders(LMP_NO_LIGHTMAP, EMobileLocalLightSetting::LOCAL_LIGHTS_DISABLED, *InBatchMaterial, InVertexFactoryType, false, MobileVS, MobilePS))
 			{
 				OutShaderElementData.NumVSInstructions = MobileVS.IsValid() ? MobileVS->GetNumInstructions() : 0;
 				OutShaderElementData.NumPSInstructions = MobilePS.IsValid() ? MobilePS->GetNumInstructions() : 0;
diff --git a/Engine/Source/Runtime/Renderer/Private/EditorPrimitivesRendering.cpp b/Engine/Source/Runtime/Renderer/Private/EditorPrimitivesRendering.cpp
--- a/Engine/Source/Runtime/Renderer/Private/EditorPrimitivesRendering.cpp
+++ b/Engine/Source/Runtime/Renderer/Private/EditorPrimitivesRendering.cpp
@@ -128,7 +128,7 @@ bool FEditorPrimitivesBasePassMeshProcessor::ProcessMobileShadingPath(const FMes
 	typedef FUniformLightMapPolicy LightMapPolicyType;
 
 	const FVertexFactory* VertexFactory = MeshBatch.VertexFactory;
-	const bool bEnableLocalLights = false;
+	const EMobileLocalLightSetting LocalLightSetting = EMobileLocalLightSetting::LOCAL_LIGHTS_DISABLED;
 	const bool bEnableSkyLight = false;
 
 	TMeshProcessorShaders<
@@ -136,7 +136,7 @@ bool FEditorPrimitivesBasePassMeshProcessor::ProcessMobileShadingPath(const FMes
 		TMobileBasePassPSPolicyParamType<FUniformLightMapPolicy>> BasePassShaders;
 	if (!MobileBasePass::GetShaders(
 		NoLightmapPolicy.GetIndirectPolicy(),
-		bEnableLocalLights,
+		LocalLightSetting,
 		Material,
 		VertexFactory->GetType(),
 		bEnableSkyLight,
diff --git a/Engine/Source/Runtime/Renderer/Private/MobileBasePass.cpp b/Engine/Source/Runtime/Renderer/Private/MobileBasePass.cpp
--- a/Engine/Source/Runtime/Renderer/Private/MobileBasePass.cpp
+++ b/Engine/Source/Runtime/Renderer/Private/MobileBasePass.cpp
@@ -13,6 +13,13 @@
 #include "PrimitiveSceneInfo.h"
 #include "MeshPassProcessor.inl"
 #include "Engine/TextureCube.h"
+#include "ShaderPlatformCachedIniValue.h"
+
+bool MobileForwardEnablePrepassLocalLights(const FStaticShaderPlatform Platform)
+{
+	static FShaderPlatformCachedIniValue<int32> MobileForwardEnablePrepassLocalLightsIniValue(TEXT("r.Mobile.Forward.EnableLocalLights"));
+	return MobileForwardEnablePrepassLocalLightsIniValue.Get(Platform) == 2;
+}
 
 uint8 GetMobileShadingModelStencilValue(FMaterialShadingModelField ShadingModel)
 {
@@ -47,7 +54,7 @@ bool MobileUsesNoLightMapPermutation(const FMeshMaterialShaderPermutationParamet
 	return true;
 }
 
-template <ELightMapPolicyType Policy, bool bEnableLocalLights>
+template <ELightMapPolicyType Policy, EMobileLocalLightSetting LocalLightSetting>
 bool GetUniformMobileBasePassShaders(
 	const FMaterial& Material, 
 	const FVertexFactoryType* VertexFactoryType, 
@@ -66,11 +73,11 @@ bool GetUniformMobileBasePassShaders(
 
 		if (bEnableSkyLight)
 		{
-			ShaderTypes.AddShaderType<TMobileBasePassPS<TUniformLightMapPolicy<Policy>, HDR_LINEAR_64, true, bEnableLocalLights>>();
+			ShaderTypes.AddShaderType<TMobileBasePassPS<TUniformLightMapPolicy<Policy>, HDR_LINEAR_64, true, LocalLightSetting>>();
 		}
 		else
 		{
-			ShaderTypes.AddShaderType<TMobileBasePassPS<TUniformLightMapPolicy<Policy>, HDR_LINEAR_64, false, bEnableLocalLights>>();
+			ShaderTypes.AddShaderType<TMobileBasePassPS<TUniformLightMapPolicy<Policy>, HDR_LINEAR_64, false, LocalLightSetting>>();
 		}	
 	}
 	else
@@ -79,11 +86,11 @@ bool GetUniformMobileBasePassShaders(
 
 		if (bEnableSkyLight)
 		{
-			ShaderTypes.AddShaderType<TMobileBasePassPS<TUniformLightMapPolicy<Policy>, LDR_GAMMA_32, true, bEnableLocalLights>>();
+			ShaderTypes.AddShaderType<TMobileBasePassPS<TUniformLightMapPolicy<Policy>, LDR_GAMMA_32, true, LocalLightSetting>>();
 		}
 		else
 		{
-			ShaderTypes.AddShaderType<TMobileBasePassPS<TUniformLightMapPolicy<Policy>, LDR_GAMMA_32, false, bEnableLocalLights>>();
+			ShaderTypes.AddShaderType<TMobileBasePassPS<TUniformLightMapPolicy<Policy>, LDR_GAMMA_32, false, LocalLightSetting>>();
 		}			
 	}
 
@@ -98,7 +105,7 @@ bool GetUniformMobileBasePassShaders(
 	return true;
 }
 
-template <bool bEnableLocalLights>
+template <EMobileLocalLightSetting LocalLightSetting>
 bool GetMobileBasePassShaders(
 	ELightMapPolicyType LightMapPolicyType, 
 	const FMaterial& Material, 
@@ -111,25 +118,25 @@ bool GetMobileBasePassShaders(
 	switch (LightMapPolicyType)
 	{
 	case LMP_NO_LIGHTMAP:
-		return GetUniformMobileBasePassShaders<LMP_NO_LIGHTMAP, bEnableLocalLights>(Material, VertexFactoryType, bEnableSkyLight, VertexShader, PixelShader);
+		return GetUniformMobileBasePassShaders<LMP_NO_LIGHTMAP, LocalLightSetting>(Material, VertexFactoryType, bEnableSkyLight, VertexShader, PixelShader);
 	case LMP_LQ_LIGHTMAP:
-		return GetUniformMobileBasePassShaders<LMP_LQ_LIGHTMAP, bEnableLocalLights>(Material, VertexFactoryType, bEnableSkyLight, VertexShader, PixelShader);
+		return GetUniformMobileBasePassShaders<LMP_LQ_LIGHTMAP, LocalLightSetting>(Material, VertexFactoryType, bEnableSkyLight, VertexShader, PixelShader);
 	case LMP_MOBILE_DISTANCE_FIELD_SHADOWS_AND_LQ_LIGHTMAP:
-		return GetUniformMobileBasePassShaders<LMP_MOBILE_DISTANCE_FIELD_SHADOWS_AND_LQ_LIGHTMAP, bEnableLocalLights>(Material, VertexFactoryType, bEnableSkyLight, VertexShader, PixelShader);
+		return GetUniformMobileBasePassShaders<LMP_MOBILE_DISTANCE_FIELD_SHADOWS_AND_LQ_LIGHTMAP, LocalLightSetting>(Material, VertexFactoryType, bEnableSkyLight, VertexShader, PixelShader);
 	case LMP_MOBILE_DISTANCE_FIELD_SHADOWS_LIGHTMAP_AND_CSM:
-		return GetUniformMobileBasePassShaders<LMP_MOBILE_DISTANCE_FIELD_SHADOWS_LIGHTMAP_AND_CSM, bEnableLocalLights>(Material, VertexFactoryType, bEnableSkyLight, VertexShader, PixelShader);
+		return GetUniformMobileBasePassShaders<LMP_MOBILE_DISTANCE_FIELD_SHADOWS_LIGHTMAP_AND_CSM, LocalLightSetting>(Material, VertexFactoryType, bEnableSkyLight, VertexShader, PixelShader);
 	case LMP_MOBILE_DIRECTIONAL_LIGHT_CSM_AND_LIGHTMAP:
-		return GetUniformMobileBasePassShaders<LMP_MOBILE_DIRECTIONAL_LIGHT_CSM_AND_LIGHTMAP, bEnableLocalLights>(Material, VertexFactoryType, bEnableSkyLight, VertexShader, PixelShader);
+		return GetUniformMobileBasePassShaders<LMP_MOBILE_DIRECTIONAL_LIGHT_CSM_AND_LIGHTMAP, LocalLightSetting>(Material, VertexFactoryType, bEnableSkyLight, VertexShader, PixelShader);
 	case LMP_MOBILE_DIRECTIONAL_LIGHT_AND_SH_INDIRECT:
-		return GetUniformMobileBasePassShaders<LMP_MOBILE_DIRECTIONAL_LIGHT_AND_SH_INDIRECT, bEnableLocalLights>(Material, VertexFactoryType, bEnableSkyLight, VertexShader, PixelShader);
+		return GetUniformMobileBasePassShaders<LMP_MOBILE_DIRECTIONAL_LIGHT_AND_SH_INDIRECT, LocalLightSetting>(Material, VertexFactoryType, bEnableSkyLight, VertexShader, PixelShader);
 	case LMP_MOBILE_DIRECTIONAL_LIGHT_CSM_AND_SH_INDIRECT:
-		return GetUniformMobileBasePassShaders<LMP_MOBILE_DIRECTIONAL_LIGHT_CSM_AND_SH_INDIRECT, bEnableLocalLights>(Material, VertexFactoryType, bEnableSkyLight, VertexShader, PixelShader);
+		return GetUniformMobileBasePassShaders<LMP_MOBILE_DIRECTIONAL_LIGHT_CSM_AND_SH_INDIRECT, LocalLightSetting>(Material, VertexFactoryType, bEnableSkyLight, VertexShader, PixelShader);
 	case LMP_MOBILE_MOVABLE_DIRECTIONAL_LIGHT_WITH_LIGHTMAP:
-		return GetUniformMobileBasePassShaders<LMP_MOBILE_MOVABLE_DIRECTIONAL_LIGHT_WITH_LIGHTMAP, bEnableLocalLights>(Material, VertexFactoryType, bEnableSkyLight, VertexShader, PixelShader);
+		return GetUniformMobileBasePassShaders<LMP_MOBILE_MOVABLE_DIRECTIONAL_LIGHT_WITH_LIGHTMAP, LocalLightSetting>(Material, VertexFactoryType, bEnableSkyLight, VertexShader, PixelShader);
 	case LMP_MOBILE_MOVABLE_DIRECTIONAL_LIGHT_CSM_WITH_LIGHTMAP:
-		return GetUniformMobileBasePassShaders<LMP_MOBILE_MOVABLE_DIRECTIONAL_LIGHT_CSM_WITH_LIGHTMAP, bEnableLocalLights>(Material, VertexFactoryType, bEnableSkyLight, VertexShader, PixelShader);
+		return GetUniformMobileBasePassShaders<LMP_MOBILE_MOVABLE_DIRECTIONAL_LIGHT_CSM_WITH_LIGHTMAP, LocalLightSetting>(Material, VertexFactoryType, bEnableSkyLight, VertexShader, PixelShader);
 	case LMP_MOBILE_DIRECTIONAL_LIGHT_CSM:
-		return GetUniformMobileBasePassShaders<LMP_MOBILE_DIRECTIONAL_LIGHT_CSM, bEnableLocalLights>(Material, VertexFactoryType, bEnableSkyLight, VertexShader, PixelShader);
+		return GetUniformMobileBasePassShaders<LMP_MOBILE_DIRECTIONAL_LIGHT_CSM, LocalLightSetting>(Material, VertexFactoryType, bEnableSkyLight, VertexShader, PixelShader);
 	default:										
 		check(false);
 		return true;
@@ -138,7 +145,7 @@ bool GetMobileBasePassShaders(
 
 bool MobileBasePass::GetShaders(
 	ELightMapPolicyType LightMapPolicyType,
-	bool bEnableLocalLights,
+	EMobileLocalLightSetting LocalLightSetting,
 	const FMaterial& MaterialResource,
 	const FVertexFactoryType* VertexFactoryType,
 	bool bEnableSkyLight, 
@@ -151,27 +158,48 @@ bool MobileBasePass::GetShaders(
 		bEnableSkyLight = !bEnableSkyLight;
 	}
 
-	if (bEnableLocalLights)
-	{
-		return GetMobileBasePassShaders<true>(
-			LightMapPolicyType,
-			MaterialResource,
-			VertexFactoryType,
-			bEnableSkyLight,
-			VertexShader,
-			PixelShader
-			);
-	}
-	else
+	switch (LocalLightSetting)
 	{
-		return GetMobileBasePassShaders<false>(
-			LightMapPolicyType,
-			MaterialResource,
-			VertexFactoryType,
-			bEnableSkyLight,
-			VertexShader,
-			PixelShader
-			);
+		case EMobileLocalLightSetting::LOCAL_LIGHTS_DISABLED:
+		{
+			return GetMobileBasePassShaders<EMobileLocalLightSetting::LOCAL_LIGHTS_DISABLED>(
+				LightMapPolicyType,
+				MaterialResource,
+				VertexFactoryType,
+				bEnableSkyLight,
+				VertexShader,
+				PixelShader
+				);
+			break;
+		}
+		case EMobileLocalLightSetting::LOCAL_LIGHTS_ENABLED:
+		{
+			return GetMobileBasePassShaders<EMobileLocalLightSetting::LOCAL_LIGHTS_ENABLED>(
+				LightMapPolicyType,
+				MaterialResource,
+				VertexFactoryType,
+				bEnableSkyLight,
+				VertexShader,
+				PixelShader
+				);
+			break;
+		}
+		case EMobileLocalLightSetting::LOCAL_LIGHTS_PREPASS_ENABLED:
+		{
+			return GetMobileBasePassShaders<EMobileLocalLightSetting::LOCAL_LIGHTS_PREPASS_ENABLED>(
+				LightMapPolicyType,
+				MaterialResource,
+				VertexFactoryType,
+				bEnableSkyLight,
+				VertexShader,
+				PixelShader
+				);
+			break;
+		}
+
+		default:
+			check(false);
+			return false;
 	}
 }
 
@@ -808,19 +836,23 @@ bool FMobileBasePassMeshProcessor::Process(
 		bEnableSkyLight = ShadingModels.IsLit() && Scene->ShouldRenderSkylightInBasePass(bTranslucentBasePass);
 	}
 
-	bool bEnableLocalLights = false;
+	EMobileLocalLightSetting LocalLightSetting = EMobileLocalLightSetting::LOCAL_LIGHTS_DISABLED;
 	if (Scene && PrimitiveSceneProxy && ShadingModels.IsLit())
 	{
-		// Whether material with a forward shading needs clustered lights
 		if (!bPassUsesDeferredShading && MobileForwardEnableLocalLights(Scene->GetShaderPlatform()))
 		{
-			bEnableLocalLights = PrimitiveSceneProxy->GetPrimitiveSceneInfo()->NumMobileDynamicLocalLights > 0;
+			if ((PrimitiveSceneProxy->GetPrimitiveSceneInfo()->NumMobileDynamicLocalLights > 0))
+			{
+				LocalLightSetting = MobileForwardEnablePrepassLocalLights(Scene->GetShaderPlatform()) ?
+					EMobileLocalLightSetting::LOCAL_LIGHTS_PREPASS_ENABLED :
+					EMobileLocalLightSetting::LOCAL_LIGHTS_ENABLED;
+			}
 		}
 	}
 
 	if (!MobileBasePass::GetShaders(
 		LightMapPolicyType,
-		bEnableLocalLights,
+		LocalLightSetting,
 		MaterialResource,
 		MeshBatch.VertexFactory->GetType(),
 		bEnableSkyLight,
@@ -900,7 +932,7 @@ void FMobileBasePassMeshProcessor::CollectPSOInitializersForLMPolicy(
 	const FGraphicsPipelineRenderTargetsInfo& RESTRICT RenderTargetsInfo,
 	const FMaterial& RESTRICT MaterialResource,
 	const bool bEnableSkyLight,
-	const bool bEnableLocalLights,
+	EMobileLocalLightSetting LocalLightSetting,
 	const ELightMapPolicyType LightMapPolicyType,
 	ERasterizerFillMode MeshFillMode,
 	ERasterizerCullMode MeshCullMode,
@@ -913,7 +945,7 @@ void FMobileBasePassMeshProcessor::CollectPSOInitializersForLMPolicy(
 
 	if (!MobileBasePass::GetShaders(
 		LightMapPolicyType,
-		bEnableLocalLights,
+		LocalLightSetting,
 		MaterialResource,
 		VertexFactoryData.VertexFactoryType,
 		bEnableSkyLight,
@@ -1001,7 +1033,8 @@ void FMobileBasePassMeshProcessor::CollectPSOInitializers(const FSceneTexturesCo
 		MobileBasePass::SetOpaqueRenderState(DrawRenderState, nullptr, Material, ShadingModels, true, bPassUsesDeferredShading);
 	}
 
-	const bool bUseLocalLightPermutation = bLitMaterial && (!bPassUsesDeferredShading && MobileForwardEnableLocalLights(ShaderPlatform));
+	const bool bUseLocalLightPermutation = bLitMaterial && (!bPassUsesDeferredShading && (MobileForwardEnableLocalLights(ShaderPlatform)));
+	EMobileLocalLightSetting LocalLightSetting = MobileForwardEnablePrepassLocalLights(ShaderPlatform) ? EMobileLocalLightSetting::LOCAL_LIGHTS_PREPASS_ENABLED: EMobileLocalLightSetting::LOCAL_LIGHTS_ENABLED;
 	const bool bCanReceiveCSM = ((Flags & FMobileBasePassMeshProcessor::EFlags::CanReceiveCSM) == FMobileBasePassMeshProcessor::EFlags::CanReceiveCSM);
 
 	FMobileLightMapPolicyTypeList UniformLightMapPolicyTypes = GetUniformLightMapPolicyTypeForPSOCollection(bLitMaterial, bTranslucentBasePass, bPassUsesDeferredShading, bCanReceiveCSM, bMovable);
@@ -1012,10 +1045,10 @@ void FMobileBasePassMeshProcessor::CollectPSOInitializers(const FSceneTexturesCo
 		bool bEnableSkyLight = false;
 		if (MobileBasePass::UseSkylightPermutation(bEnableSkyLight, FReadOnlyCVARCache::Get().MobileSkyLightPermutation) || !bLitMaterial)
 		{
-			CollectPSOInitializersForLMPolicy(VertexFactoryData, DrawRenderState, RenderTargetsInfo, Material, bEnableSkyLight, false /*bEnableLocalLights*/, LightMapPolicyType, MeshFillMode, MeshCullMode, (EPrimitiveType)PreCacheParams.PrimitiveType, PSOInitializers);
+			CollectPSOInitializersForLMPolicy(VertexFactoryData, DrawRenderState, RenderTargetsInfo, Material, bEnableSkyLight, EMobileLocalLightSetting::LOCAL_LIGHTS_DISABLED, LightMapPolicyType, MeshFillMode, MeshCullMode, (EPrimitiveType)PreCacheParams.PrimitiveType, PSOInitializers);
 			if (bUseLocalLightPermutation)
 			{
-				CollectPSOInitializersForLMPolicy(VertexFactoryData, DrawRenderState, RenderTargetsInfo, Material, bEnableSkyLight, true /*bEnableLocalLights*/, LightMapPolicyType, MeshFillMode, MeshCullMode, (EPrimitiveType)PreCacheParams.PrimitiveType, PSOInitializers);
+				CollectPSOInitializersForLMPolicy(VertexFactoryData, DrawRenderState, RenderTargetsInfo, Material, bEnableSkyLight, LocalLightSetting, LightMapPolicyType, MeshFillMode, MeshCullMode, (EPrimitiveType)PreCacheParams.PrimitiveType, PSOInitializers);
 			}
 		}
 
@@ -1023,10 +1056,10 @@ void FMobileBasePassMeshProcessor::CollectPSOInitializers(const FSceneTexturesCo
 		bEnableSkyLight = true;
 		if (MobileBasePass::UseSkylightPermutation(bEnableSkyLight, FReadOnlyCVARCache::Get().MobileSkyLightPermutation) && bLitMaterial)
 		{
-			CollectPSOInitializersForLMPolicy(VertexFactoryData, DrawRenderState, RenderTargetsInfo, Material, bEnableSkyLight, false /*bEnableLocalLights*/, LightMapPolicyType, MeshFillMode, MeshCullMode, (EPrimitiveType)PreCacheParams.PrimitiveType, PSOInitializers);
+			CollectPSOInitializersForLMPolicy(VertexFactoryData, DrawRenderState, RenderTargetsInfo, Material, bEnableSkyLight, EMobileLocalLightSetting::LOCAL_LIGHTS_DISABLED, LightMapPolicyType, MeshFillMode, MeshCullMode, (EPrimitiveType)PreCacheParams.PrimitiveType, PSOInitializers);
 			if (bUseLocalLightPermutation)
 			{
-				CollectPSOInitializersForLMPolicy(VertexFactoryData, DrawRenderState, RenderTargetsInfo, Material, bEnableSkyLight, true /*bEnableLocalLights*/, LightMapPolicyType, MeshFillMode, MeshCullMode, (EPrimitiveType)PreCacheParams.PrimitiveType, PSOInitializers);
+				CollectPSOInitializersForLMPolicy(VertexFactoryData, DrawRenderState, RenderTargetsInfo, Material, bEnableSkyLight, LocalLightSetting, LightMapPolicyType, MeshFillMode, MeshCullMode, (EPrimitiveType)PreCacheParams.PrimitiveType, PSOInitializers);
 			}
 		}
 	}
diff --git a/Engine/Source/Runtime/Renderer/Private/MobileBasePassRendering.cpp b/Engine/Source/Runtime/Renderer/Private/MobileBasePassRendering.cpp
--- a/Engine/Source/Runtime/Renderer/Private/MobileBasePassRendering.cpp
+++ b/Engine/Source/Runtime/Renderer/Private/MobileBasePassRendering.cpp
@@ -82,27 +82,21 @@ bool ShouldCacheShaderByPlatformAndOutputFormat(EShaderPlatform Platform, EOutpu
 	IMPLEMENT_MATERIAL_SHADER_TYPE(template<>, TMobileBasePassVS##LightMapPolicyName##LDRGamma32, TEXT("/Engine/Private/MobileBasePassVertexShader.usf"), TEXT("Main"), SF_Vertex); \
 	IMPLEMENT_MATERIAL_SHADER_TYPE(template<>, TMobileBasePassVS##LightMapPolicyName##HDRLinear64, TEXT("/Engine/Private/MobileBasePassVertexShader.usf"), TEXT("Main"), SF_Vertex);
 
-#define IMPLEMENT_MOBILE_SHADING_BASEPASS_LIGHTMAPPED_PIXEL_SHADER_TYPE(LightMapPolicyType,LightMapPolicyName) \
-	typedef TMobileBasePassPS< LightMapPolicyType, LDR_GAMMA_32, false, false > TMobileBasePassPS##LightMapPolicyName##LDRGamma32; \
-	typedef TMobileBasePassPS< LightMapPolicyType, HDR_LINEAR_64, false, false > TMobileBasePassPS##LightMapPolicyName##HDRLinear64; \
-	typedef TMobileBasePassPS< LightMapPolicyType, LDR_GAMMA_32, true, false > TMobileBasePassPS##LightMapPolicyName##LDRGamma32##Skylight; \
-	typedef TMobileBasePassPS< LightMapPolicyType, HDR_LINEAR_64, true, false > TMobileBasePassPS##LightMapPolicyName##HDRLinear64##Skylight; \
-	typedef TMobileBasePassPS< LightMapPolicyType, LDR_GAMMA_32, false, true > TMobileBasePassPS##LightMapPolicyName##LDRGamma32##EnableLocalLights; \
-	typedef TMobileBasePassPS< LightMapPolicyType, HDR_LINEAR_64, false, true > TMobileBasePassPS##LightMapPolicyName##HDRLinear64##EnableLocalLights; \
-	typedef TMobileBasePassPS< LightMapPolicyType, LDR_GAMMA_32, true, true > TMobileBasePassPS##LightMapPolicyName##LDRGamma32##Skylight##EnableLocalLights; \
-	typedef TMobileBasePassPS< LightMapPolicyType, HDR_LINEAR_64, true, true > TMobileBasePassPS##LightMapPolicyName##HDRLinear64##Skylight##EnableLocalLights; \
-	IMPLEMENT_MATERIAL_SHADER_TYPE(template<>, TMobileBasePassPS##LightMapPolicyName##LDRGamma32, TEXT("/Engine/Private/MobileBasePassPixelShader.usf"), TEXT("Main"), SF_Pixel); \
-	IMPLEMENT_MATERIAL_SHADER_TYPE(template<>, TMobileBasePassPS##LightMapPolicyName##HDRLinear64, TEXT("/Engine/Private/MobileBasePassPixelShader.usf"), TEXT("Main"), SF_Pixel); \
-	IMPLEMENT_MATERIAL_SHADER_TYPE(template<>, TMobileBasePassPS##LightMapPolicyName##LDRGamma32##Skylight, TEXT("/Engine/Private/MobileBasePassPixelShader.usf"), TEXT("Main"), SF_Pixel); \
-	IMPLEMENT_MATERIAL_SHADER_TYPE(template<>, TMobileBasePassPS##LightMapPolicyName##HDRLinear64##Skylight, TEXT("/Engine/Private/MobileBasePassPixelShader.usf"), TEXT("Main"), SF_Pixel) \
-	IMPLEMENT_MATERIAL_SHADER_TYPE(template<>, TMobileBasePassPS##LightMapPolicyName##LDRGamma32##EnableLocalLights, TEXT("/Engine/Private/MobileBasePassPixelShader.usf"), TEXT("Main"), SF_Pixel); \
-	IMPLEMENT_MATERIAL_SHADER_TYPE(template<>, TMobileBasePassPS##LightMapPolicyName##HDRLinear64##EnableLocalLights, TEXT("/Engine/Private/MobileBasePassPixelShader.usf"), TEXT("Main"), SF_Pixel); \
-	IMPLEMENT_MATERIAL_SHADER_TYPE(template<>, TMobileBasePassPS##LightMapPolicyName##LDRGamma32##Skylight##EnableLocalLights, TEXT("/Engine/Private/MobileBasePassPixelShader.usf"), TEXT("Main"), SF_Pixel); \
-	IMPLEMENT_MATERIAL_SHADER_TYPE(template<>, TMobileBasePassPS##LightMapPolicyName##HDRLinear64##Skylight##EnableLocalLights, TEXT("/Engine/Private/MobileBasePassPixelShader.usf"), TEXT("Main"), SF_Pixel);
+#define IMPLEMENT_MOBILE_SHADING_BASEPASS_LIGHTMAPPED_PIXEL_SHADER_TYPE(LightMapPolicyType,LightMapPolicyName, LocalLightSetting) \
+	typedef TMobileBasePassPS< LightMapPolicyType, LDR_GAMMA_32, false, LocalLightSetting > TMobileBasePassPS##LightMapPolicyName##LDRGamma32##LocalLightSetting; \
+	typedef TMobileBasePassPS< LightMapPolicyType, HDR_LINEAR_64, false, LocalLightSetting > TMobileBasePassPS##LightMapPolicyName##HDRLinear64##LocalLightSetting; \
+	typedef TMobileBasePassPS< LightMapPolicyType, LDR_GAMMA_32, true, LocalLightSetting > TMobileBasePassPS##LightMapPolicyName##LDRGamma32##Skylight##LocalLightSetting; \
+	typedef TMobileBasePassPS< LightMapPolicyType, HDR_LINEAR_64, true, LocalLightSetting > TMobileBasePassPS##LightMapPolicyName##HDRLinear64##Skylight##LocalLightSetting; \
+	IMPLEMENT_MATERIAL_SHADER_TYPE(template<>, TMobileBasePassPS##LightMapPolicyName##LDRGamma32##LocalLightSetting, TEXT("/Engine/Private/MobileBasePassPixelShader.usf"), TEXT("Main"), SF_Pixel); \
+	IMPLEMENT_MATERIAL_SHADER_TYPE(template<>, TMobileBasePassPS##LightMapPolicyName##HDRLinear64##LocalLightSetting, TEXT("/Engine/Private/MobileBasePassPixelShader.usf"), TEXT("Main"), SF_Pixel); \
+	IMPLEMENT_MATERIAL_SHADER_TYPE(template<>, TMobileBasePassPS##LightMapPolicyName##LDRGamma32##Skylight##LocalLightSetting, TEXT("/Engine/Private/MobileBasePassPixelShader.usf"), TEXT("Main"), SF_Pixel); \
+	IMPLEMENT_MATERIAL_SHADER_TYPE(template<>, TMobileBasePassPS##LightMapPolicyName##HDRLinear64##Skylight##LocalLightSetting, TEXT("/Engine/Private/MobileBasePassPixelShader.usf"), TEXT("Main"), SF_Pixel)
 
 #define IMPLEMENT_MOBILE_SHADING_BASEPASS_LIGHTMAPPED_SHADER_TYPE(LightMapPolicyType,LightMapPolicyName) \
 	IMPLEMENT_MOBILE_SHADING_BASEPASS_LIGHTMAPPED_VERTEX_SHADER_TYPE(LightMapPolicyType, LightMapPolicyName) \
-	IMPLEMENT_MOBILE_SHADING_BASEPASS_LIGHTMAPPED_PIXEL_SHADER_TYPE(LightMapPolicyType, LightMapPolicyName)
+	IMPLEMENT_MOBILE_SHADING_BASEPASS_LIGHTMAPPED_PIXEL_SHADER_TYPE(LightMapPolicyType, LightMapPolicyName, LOCAL_LIGHTS_DISABLED) \
+	IMPLEMENT_MOBILE_SHADING_BASEPASS_LIGHTMAPPED_PIXEL_SHADER_TYPE(LightMapPolicyType, LightMapPolicyName, LOCAL_LIGHTS_ENABLED) \
+	IMPLEMENT_MOBILE_SHADING_BASEPASS_LIGHTMAPPED_PIXEL_SHADER_TYPE(LightMapPolicyType, LightMapPolicyName, LOCAL_LIGHTS_PREPASS_ENABLED)
 
 // Implement shader types per lightmap policy 
 IMPLEMENT_MOBILE_SHADING_BASEPASS_LIGHTMAPPED_SHADER_TYPE(TUniformLightMapPolicy<LMP_NO_LIGHTMAP>, FNoLightMapPolicy);
diff --git a/Engine/Source/Runtime/Renderer/Private/MobileBasePassRendering.h b/Engine/Source/Runtime/Renderer/Private/MobileBasePassRendering.h
--- a/Engine/Source/Runtime/Renderer/Private/MobileBasePassRendering.h
+++ b/Engine/Source/Runtime/Renderer/Private/MobileBasePassRendering.h
@@ -26,6 +26,8 @@
 #include "RenderUtils.h"
 #include "DebugViewModeRendering.h"
 
+bool MobileForwardEnablePrepassLocalLights(const FStaticShaderPlatform Platform);
+
 struct FMobileBasePassTextures
 {
 	FRDGTextureRef ScreenSpaceAO = nullptr;
@@ -95,6 +97,13 @@ enum EOutputFormat
 	HDR_LINEAR_64,
 };
 
+enum EMobileLocalLightSetting
+{
+	LOCAL_LIGHTS_DISABLED,
+	LOCAL_LIGHTS_ENABLED,
+	LOCAL_LIGHTS_PREPASS_ENABLED
+};
+
 bool ShouldCacheShaderByPlatformAndOutputFormat(EShaderPlatform Platform, EOutputFormat OutputFormat);
 // shared defines for mobile base pass VS and PS
 void MobileBasePassModifyCompilationEnvironment(const FMaterialShaderPermutationParameters& Parameters, FShaderCompilerEnvironment& OutEnvironment, EOutputFormat OutputFormat);
@@ -319,7 +328,7 @@ namespace MobileBasePass
 
 	bool GetShaders(
 		ELightMapPolicyType LightMapPolicyType,
-		bool bEnableLocalLights,
+		EMobileLocalLightSetting LocalLightSetting,
 		const FMaterial& MaterialResource,
 		const FVertexFactoryType* VertexFactoryType,
 		bool bEnableSkyLight, 
@@ -346,7 +355,7 @@ namespace MobileBasePass
 	}
 };
 
-template< typename LightMapPolicyType, EOutputFormat OutputFormat, bool bEnableSkyLight, bool bEnableLocalLights>
+template< typename LightMapPolicyType, EOutputFormat OutputFormat, bool bEnableSkyLight, EMobileLocalLightSetting LocalLightSetting>
 class TMobileBasePassPS : public TMobileBasePassPSBaseType<LightMapPolicyType>
 {
 	DECLARE_SHADER_TYPE(TMobileBasePassPS,MeshMaterial);
@@ -375,6 +384,7 @@ public:
 
 		// Deferred shading does not need SkyLight and LocalLight permutations
 		// TODO: skip skylight permutations for deferred	
+		bool bEnableLocalLights = LocalLightSetting != EMobileLocalLightSetting::LOCAL_LIGHTS_DISABLED;
 		const bool bShouldCacheByShading = (bForwardShading || !bEnableLocalLights);
 		const bool bShouldCacheByLocalLights = !bEnableLocalLights || (bIsLit && bEnableLocalLights == bSupportsLocalLights);
 
@@ -398,7 +408,10 @@ public:
 		OutEnvironment.SetDefine(TEXT("ENABLE_AMBIENT_OCCLUSION"), IsMobileAmbientOcclusionEnabled(Parameters.Platform) ? 1u : 0u);
 		
 		FForwardLightingParameters::ModifyCompilationEnvironment(Parameters.Platform, OutEnvironment);
-		OutEnvironment.SetDefine(TEXT("ENABLE_CLUSTERED_LIGHTS"), bEnableLocalLights ? 1u : 0u);
+		OutEnvironment.SetDefine(TEXT("ENABLE_CLUSTERED_LIGHTS"), (LocalLightSetting == EMobileLocalLightSetting::LOCAL_LIGHTS_ENABLED) ? 1u : 0u);
+
+		// Translucent materials don't write to depth so cannot use prepass
+		OutEnvironment.SetDefine(TEXT("PREPASS_LOCAL_LIGHTS_MOBILE"), !bTranslucentMaterial && (LocalLightSetting == EMobileLocalLightSetting::LOCAL_LIGHTS_PREPASS_ENABLED) ? 1u: 0u);
 		OutEnvironment.SetDefine(TEXT("ENABLE_CLUSTERED_REFLECTION"), bEnableClusteredReflections ? 1u : 0u);
 		OutEnvironment.SetDefine(TEXT("USE_SHADOWMASKTEXTURE"), bMobileUsesShadowMaskTexture && !bTranslucentMaterial ? 1u : 0u);
 	}
@@ -453,7 +466,7 @@ public:
 		const FGraphicsPipelineRenderTargetsInfo& RESTRICT RenderTargetsInfo,
 		const FMaterial& RESTRICT MaterialResource,
 		const bool bRenderSkylight,
-		const bool bEnableLocalLights,
+		EMobileLocalLightSetting LocalLightSetting,
 		const ELightMapPolicyType LightMapPolicyType,
 		ERasterizerFillMode MeshFillMode,
 		ERasterizerCullMode MeshCullMode,
diff --git a/Engine/Source/Runtime/Renderer/Private/MobileLocalLightPrepass.cpp b/Engine/Source/Runtime/Renderer/Private/MobileLocalLightPrepass.cpp
new file mode 100644
--- /dev/null
+++ b/Engine/Source/Runtime/Renderer/Private/MobileLocalLightPrepass.cpp
@@ -0,0 +1,68 @@
+// Copyright Epic Games, Inc. All Rights Reserved.
+
+#include "SceneRendering.h"
+#include "DataDrivenShaderPlatformInfo.h"
+#include "BasePassRendering.h"
+#include "PixelShaderUtils.h"
+#include "MobileBasePassRendering.h"
+
+class FLocalLightPrepassPS : public FGlobalShader
+{
+	DECLARE_GLOBAL_SHADER(FLocalLightPrepassPS);
+	SHADER_USE_PARAMETER_STRUCT(FLocalLightPrepassPS, FGlobalShader)
+
+		BEGIN_SHADER_PARAMETER_STRUCT(FParameters, )
+		SHADER_PARAMETER_STRUCT_REF(FViewUniformShaderParameters, View)
+		SHADER_PARAMETER_STRUCT_INCLUDE(FSceneTextureShaderParameters, SceneTextures)
+		SHADER_PARAMETER_RDG_UNIFORM_BUFFER(FForwardLightData, ForwardLightData)
+		RENDER_TARGET_BINDING_SLOTS()
+		END_SHADER_PARAMETER_STRUCT()
+
+		static bool ShouldCompilePermutation(const FGlobalShaderPermutationParameters& Parameters)
+	{
+		return IsMobilePlatform(Parameters.Platform) && MobileForwardEnablePrepassLocalLights(Parameters.Platform);
+	}
+
+	static void ModifyCompilationEnvironment(const FGlobalShaderPermutationParameters& Parameters, FShaderCompilerEnvironment& OutEnvironment)
+	{
+		FForwardLightingParameters::ModifyCompilationEnvironment(Parameters.Platform, OutEnvironment);
+		OutEnvironment.SetRenderTargetOutputFormat(0, PF_FloatR11G11B10);
+		OutEnvironment.SetRenderTargetOutputFormat(1, PF_A2B10G10R10);
+	}
+};
+
+IMPLEMENT_GLOBAL_SHADER(FLocalLightPrepassPS, "/Engine/Private/MobileLocalLightPrepass.usf", "Main", SF_Pixel);
+
+void FMobileSceneRenderer::RenderLocalLightPrepass(FRDGBuilder& GraphBuilder, FSceneTextures& SceneTextures)
+{
+	RDG_EVENT_SCOPE(GraphBuilder, "RenderLocalLightPrepass");
+	QUICK_SCOPE_CYCLE_COUNTER(STAT_RenderLocalLightPrepass);
+
+	for (int32 ViewIndex = 0; ViewIndex < Views.Num(); ViewIndex++)
+	{
+		FViewInfo& View = Views[ViewIndex];
+		bool bHasNoLocalLights = (!View.ForwardLightingResources.ForwardLightData) || (View.ForwardLightingResources.ForwardLightData->NumLocalLights == 0);
+		if (!View.ShouldRenderView() || bHasNoLocalLights)
+		{
+			continue;
+		}
+
+		FLocalLightPrepassPS::FParameters* PassParameters = GraphBuilder.AllocParameters<FLocalLightPrepassPS::FParameters>();
+		PassParameters->View = View.ViewUniformBuffer;
+		PassParameters->SceneTextures = SceneTextures.GetSceneTextureShaderParameters(View.FeatureLevel);
+		PassParameters->RenderTargets[0] = FRenderTargetBinding(SceneTextures.MobileLocalLightTextureA, ERenderTargetLoadAction::EClear);;
+		PassParameters->RenderTargets[1] = FRenderTargetBinding(SceneTextures.MobileLocalLightTextureB, ERenderTargetLoadAction::EClear);;
+		PassParameters->ForwardLightData = View.ForwardLightingResources.ForwardLightUniformBuffer;
+
+		FGlobalShaderMap* GlobalShaderMap = GetGlobalShaderMap(FeatureLevel);
+		TShaderMapRef<FLocalLightPrepassPS> PixelShader(GlobalShaderMap);
+
+		FPixelShaderUtils::AddFullscreenPass(
+			GraphBuilder,
+			GlobalShaderMap,
+			RDG_EVENT_NAME("RenderLocalLightPrepass"),
+			PixelShader,
+			PassParameters,
+			FIntRect(0, 0, SceneTextures.MobileLocalLightTextureA->Desc.Extent.X, SceneTextures.MobileLocalLightTextureA->Desc.Extent.Y));
+	}
+}
\ No newline at end of file
diff --git a/Engine/Source/Runtime/Renderer/Private/MobileShadingRenderer.cpp b/Engine/Source/Runtime/Renderer/Private/MobileShadingRenderer.cpp
--- a/Engine/Source/Runtime/Renderer/Private/MobileShadingRenderer.cpp
+++ b/Engine/Source/Runtime/Renderer/Private/MobileShadingRenderer.cpp
@@ -279,6 +279,7 @@ FMobileSceneRenderer::FMobileSceneRenderer(const FSceneViewFamily* InViewFamily,
 	bIsMaskedOnlyDepthPrepassEnabled = Scene->EarlyZPassMode == DDM_MaskedOnly;
 	bRequiresSceneDepthAux = MobileRequiresSceneDepthAux(ShaderPlatform);
 	bEnableClusteredLocalLights = MobileForwardEnableLocalLights(ShaderPlatform);
+	bEnablePrepassLocalLights = MobileForwardEnablePrepassLocalLights(ShaderPlatform);
 	bEnableClusteredReflections = MobileForwardEnableClusteredReflections(ShaderPlatform);
 	
 	StandardTranslucencyPass = ViewFamily.AllowTranslucencyAfterDOF() ? ETranslucencyPass::TPT_TranslucencyStandard : ETranslucencyPass::TPT_AllTranslucency;
@@ -939,10 +940,10 @@ void FMobileSceneRenderer::Render(FRDGBuilder& GraphBuilder)
 		bool bShadowedLightsInClustered = bRequiresShadowProjections && !bDeferredShading;
 		GatherAndSortLights(SortedLightSet, bShadowedLightsInClustered);
 		int32 NumReflectionCaptures = Views[0].NumBoxReflectionCaptures + Views[0].NumSphereReflectionCaptures;
-		bool bCullLightsToGrid = (((bEnableClusteredReflections || bDeferredShading) && NumReflectionCaptures > 0) || bEnableClusteredLocalLights);
+		bool bCullLightsToGrid = (((bEnableClusteredReflections || bDeferredShading) && NumReflectionCaptures > 0) || bEnableClusteredLocalLights || bEnablePrepassLocalLights);
 		if (bCullLightsToGrid)
 		{
-			ComputeLightGrid(GraphBuilder, bEnableClusteredLocalLights, SortedLightSet);
+			ComputeLightGrid(GraphBuilder, bEnableClusteredLocalLights || bEnablePrepassLocalLights, SortedLightSet);
 		}
 	}
 
@@ -1030,6 +1031,12 @@ void FMobileSceneRenderer::Render(FRDGBuilder& GraphBuilder)
 		{
 			RenderAmbientOcclusion(GraphBuilder, SceneTextures.Depth.Resolve, SceneTextures.ScreenSpaceAO);
 		}
+
+		// Local Light prepass
+		if (bEnablePrepassLocalLights)
+		{
+			RenderLocalLightPrepass(GraphBuilder, SceneTextures);
+		}
 	}
 
 	if (bDeferredShading)
diff --git a/Engine/Source/Runtime/Renderer/Private/SceneRendering.h b/Engine/Source/Runtime/Renderer/Private/SceneRendering.h
--- a/Engine/Source/Runtime/Renderer/Private/SceneRendering.h
+++ b/Engine/Source/Runtime/Renderer/Private/SceneRendering.h
@@ -2686,6 +2686,8 @@ protected:
 	void RenderMaskedPrePass(FRHICommandList& RHICmdList, const FViewInfo& View);
 	void RenderFullDepthPrepass(FRDGBuilder& GraphBuilder, FSceneTextures& SceneTextures);
 
+	void RenderLocalLightPrepass(FRDGBuilder& GraphBuilder, FSceneTextures& SceneTextures);
+
 	/** Renders the opaque base pass for mobile. */
 	void RenderMobileBasePass(FRHICommandList& RHICmdList, const FViewInfo& View);
 
@@ -2764,6 +2766,7 @@ private:
 	bool bIsMaskedOnlyDepthPrepassEnabled;
 	bool bRequiresSceneDepthAux;
 	bool bEnableClusteredLocalLights;
+	bool bEnablePrepassLocalLights;
 	bool bEnableClusteredReflections;
 	bool bRequiresShadowProjections;
 
diff --git a/Engine/Source/Runtime/Renderer/Private/SceneTextures.cpp b/Engine/Source/Runtime/Renderer/Private/SceneTextures.cpp
--- a/Engine/Source/Runtime/Renderer/Private/SceneTextures.cpp
+++ b/Engine/Source/Runtime/Renderer/Private/SceneTextures.cpp
@@ -615,6 +615,16 @@ void FSceneTextures::InitializeViewFamily(FRDGBuilder& GraphBuilder, FViewFamily
 	}
 #endif
 
+	{
+		FRDGTextureDesc Desc = FRDGTextureDesc::Create2D(Config.Extent, PF_FloatR11G11B10, FClearValueBinding::Transparent, TexCreate_RenderTargetable | TexCreate_ShaderResource);
+		SceneTextures.MobileLocalLightTextureA = GraphBuilder.CreateTexture(Desc, TEXT("MobileLocalLightTextureA"));
+	}
+
+	{
+		FRDGTextureDesc Desc = FRDGTextureDesc::Create2D(Config.Extent, PF_A2B10G10R10, FClearValueBinding::Transparent, TexCreate_RenderTargetable | TexCreate_ShaderResource);
+		SceneTextures.MobileLocalLightTextureB = GraphBuilder.CreateTexture(Desc, TEXT("MobileLocalLightTextureB"));
+	}
+
 #if WITH_DEBUG_VIEW_MODES
 	if (AllowDebugViewShaderMode(DVSM_QuadComplexity, Config.ShaderPlatform, Config.FeatureLevel))
 	{
@@ -948,11 +958,15 @@ void SetupMobileSceneTextureUniformParameters(
 	SceneTextureParameters.GBufferDTexture = SystemTextures.Black;
 	// SceneDepthAuxTexture is a color texture on mobile, with DeviceZ values
 	SceneTextureParameters.SceneDepthAuxTexture = SystemTextures.Black;
+	SceneTextureParameters.LocalLightTextureA = SystemTextures.Black;
+	SceneTextureParameters.LocalLightTextureB = SystemTextures.Black;
 	SceneTextureParameters.GBufferATextureSampler = TStaticSamplerState<>::GetRHI();
 	SceneTextureParameters.GBufferBTextureSampler = TStaticSamplerState<>::GetRHI();
 	SceneTextureParameters.GBufferCTextureSampler = TStaticSamplerState<>::GetRHI();
 	SceneTextureParameters.GBufferDTextureSampler = TStaticSamplerState<>::GetRHI();
 	SceneTextureParameters.SceneDepthAuxTextureSampler = TStaticSamplerState<>::GetRHI();
+	SceneTextureParameters.LocalLightTextureASampler = TStaticSamplerState<>::GetRHI();
+	SceneTextureParameters.LocalLightTextureBSampler = TStaticSamplerState<>::GetRHI();
 
 	if (SceneTextures)
 	{
@@ -1006,6 +1020,16 @@ void SetupMobileSceneTextureUniformParameters(
 			}
 		}
 
+		if (HasBeenProduced(SceneTextures->MobileLocalLightTextureA))
+		{
+			SceneTextureParameters.LocalLightTextureA = SceneTextures->MobileLocalLightTextureA;
+		}
+
+		if (HasBeenProduced(SceneTextures->MobileLocalLightTextureB))
+		{
+			SceneTextureParameters.LocalLightTextureB = SceneTextures->MobileLocalLightTextureB;
+		}
+
 		if (EnumHasAnyFlags(SetupMode, EMobileSceneTextureSetupMode::CustomDepth))
 		{
 			const FCustomDepthTextures& CustomDepthTextures = SceneTextures->CustomDepth;
diff --git a/Engine/Source/Runtime/Renderer/Private/SceneTextures.h b/Engine/Source/Runtime/Renderer/Private/SceneTextures.h
--- a/Engine/Source/Runtime/Renderer/Private/SceneTextures.h
+++ b/Engine/Source/Runtime/Renderer/Private/SceneTextures.h
@@ -81,6 +81,10 @@ struct RENDERER_API FSceneTextures : public FMinimalSceneTextures
 
 	// Texture containing dynamic motion vectors. Can be bound by the base pass or its own velocity pass.
 	FRDGTextureRef Velocity{};
+	
+	// (Mobile Local Light Prepass) Textures containing LocalLight Direction and Color
+	FRDGTextureRef MobileLocalLightTextureA {};
+	FRDGTextureRef MobileLocalLightTextureB {};
 
 	// Texture containing the screen space ambient occlusion result.
 	FRDGTextureRef ScreenSpaceAO{};
diff --git a/Engine/Source/Runtime/Renderer/Private/SkyPassRendering.cpp b/Engine/Source/Runtime/Renderer/Private/SkyPassRendering.cpp
--- a/Engine/Source/Runtime/Renderer/Private/SkyPassRendering.cpp
+++ b/Engine/Source/Runtime/Renderer/Private/SkyPassRendering.cpp
@@ -115,7 +115,7 @@ bool FSkyPassMeshProcessor::Process(
 
 		if (!MobileBasePass::GetShaders(
 			LMP_NO_LIGHTMAP,
-			0,
+			EMobileLocalLightSetting::LOCAL_LIGHTS_DISABLED,
 			MaterialResource,
 			VertexFactory->GetType(),
 			false,
