Fixes a bug where Lumen was not using View.PrevViewInfo.ScreenSpaceRayTracingInput with TSR

diff --git a/Engine/Source/Runtime/Renderer/Private/PostProcess/TemporalSuperResolution.cpp b/Engine/Source/Runtime/Renderer/Private/PostProcess/TemporalSuperResolution.cpp
--- a/Engine/Source/Runtime/Renderer/Private/PostProcess/TemporalSuperResolution.cpp
+++ b/Engine/Source/Runtime/Renderer/Private/PostProcess/TemporalSuperResolution.cpp
@@ -2244,11 +2244,11 @@ FDefaultTemporalUpscaler::FOutputs AddTemporalSuperResolutionPasses(
 
 		// Extract the output for next frame SSR so that separate translucency shows up in SSR.
 		{
+			// Output in TemporalAAHistory and not CustomSSR so Lumen can pick up ScreenSpaceRayTracingInput in priority to ensure consistent behavior between TAA and TSR.
 			GraphBuilder.QueueTextureExtraction(
-				SceneColorOutputTexture, &View.ViewState->PrevFrameViewInfo.CustomSSRInput.RT[0]);
-
-			View.ViewState->PrevFrameViewInfo.CustomSSRInput.ViewportRect = OutputRect;
-			View.ViewState->PrevFrameViewInfo.CustomSSRInput.ReferenceBufferSize = OutputExtent;
+				SceneColorOutputTexture, &View.ViewState->PrevFrameViewInfo.TemporalAAHistory.RT[0]);
+			View.ViewState->PrevFrameViewInfo.TemporalAAHistory.ViewportRect = OutputRect;
+			View.ViewState->PrevFrameViewInfo.TemporalAAHistory.ReferenceBufferSize = OutputExtent;
 		}
 	}
 
