Fix assert with Substrate.

diff --git a/Engine/Source/Runtime/Renderer/Private/Strata/Strata.cpp b/Engine/Source/Runtime/Renderer/Private/Strata/Strata.cpp
--- a/Engine/Source/Runtime/Renderer/Private/Strata/Strata.cpp
+++ b/Engine/Source/Runtime/Renderer/Private/Strata/Strata.cpp
@@ -453,7 +453,8 @@ void InitialiseStrataFrameSceneData(FRDGBuilder& GraphBuilder, FSceneRenderer& S
 	const uint32 SliceCountAdvDebug = IsAdvancedVisualizationEnabled() ? 1 : 0;
 	const uint32 SliceCount = FMath::DivideAndRoundUp(Out.MaxBytesPerPixel, 4u) + SliceCountSSS + SliceCountAdvDebug;
 	FRDGTextureDesc MaterialTextureDesc = FRDGTextureDesc::Create2DArray(SceneTextureExtent, PF_R32_UINT, FClearValueBinding::Transparent, TexCreate_TargetArraySlicesIndependently | TexCreate_DisableDCC | TexCreate_NoFastClear | TexCreate_RenderTargetable | TexCreate_ShaderResource | TexCreate_UAV | TexCreate_FastVRAM, SliceCount, 1, 1);
-	MaterialTextureDesc.FastVRAMPercentage = (1.0f / SliceCount) * 0xFF; // Only allocate the first slice into ESRAM
+	// Disable temporarlly partial VRAM allocation as it causes issue with RDG
+	MaterialTextureDesc.FastVRAMPercentage = 0xff;  //(1.0f / SliceCount) * 0xFF; // Only allocate the first slice into ESRAM
 
 	Out.MaterialTextureArray = GraphBuilder.CreateTexture(MaterialTextureDesc, TEXT("Substrate.Material"));
 	Out.MaterialTextureArraySRV = GraphBuilder.CreateSRV(FRDGTextureSRVDesc::Create(Out.MaterialTextureArray));
