Fix a bug of sampling SceneDepthAux in SunMask on IOS mobile deferred.

diff --git a/Engine/Source/Runtime/Renderer/Private/PostProcess/PostProcessing.cpp b/Engine/Source/Runtime/Renderer/Private/PostProcess/PostProcessing.cpp
--- a/Engine/Source/Runtime/Renderer/Private/PostProcess/PostProcessing.cpp
+++ b/Engine/Source/Runtime/Renderer/Private/PostProcess/PostProcessing.cpp
@@ -2139,7 +2139,7 @@ void AddMobilePostProcessingPasses(FRDGBuilder& GraphBuilder, FScene* Scene, con
 		if (PassSequence.IsEnabled(EPass::SunMask))
 		{
 			PassSequence.AcceptPass(EPass::SunMask);
-			bool bUseDepthTexture = !MobileRequiresSceneDepthAux(View.GetShaderPlatform());
+			bool bUseDepthTexture = !MobileRequiresSceneDepthAux(View.GetShaderPlatform()) || IsMobileDeferredShadingEnabled(View.GetShaderPlatform());
 
 			FMobileSunMaskInputs SunMaskInputs;
 			SunMaskInputs.bUseDepthTexture = bUseDepthTexture;
