Fix missing FastVramPercentage data and translation, causing assert when due to mismatching hash/description.

diff --git a/Engine/Source/Runtime/RenderCore/Public/RenderGraphResources.inl b/Engine/Source/Runtime/RenderCore/Public/RenderGraphResources.inl
--- a/Engine/Source/Runtime/RenderCore/Public/RenderGraphResources.inl
+++ b/Engine/Source/Runtime/RenderCore/Public/RenderGraphResources.inl
@@ -70,6 +70,7 @@ inline FPooledRenderTargetDesc Translate(const FRHITextureDesc& InDesc)
 	OutDesc.NumSamples = InDesc.NumSamples;
 	OutDesc.bIsArray = InDesc.IsTextureArray();
 	OutDesc.bIsCubemap = InDesc.IsTextureCube();
+	OutDesc.FastVRAMPercentage = InDesc.FastVRAMPercentage;
 
 	check(OutDesc.IsValid());
 	return OutDesc;
@@ -126,6 +127,7 @@ inline FRDGTextureDesc Translate(const FPooledRenderTargetDesc& InDesc)
 	}
 
 	OutDesc.Flags = InDesc.Flags;
+	OutDesc.FastVRAMPercentage = InDesc.FastVRAMPercentage;
 	check(OutDesc.IsValid());
 
 	return OutDesc;
diff --git a/Engine/Source/Runtime/RenderCore/Public/RendererInterface.h b/Engine/Source/Runtime/RenderCore/Public/RendererInterface.h
--- a/Engine/Source/Runtime/RenderCore/Public/RendererInterface.h
+++ b/Engine/Source/Runtime/RenderCore/Public/RendererInterface.h
@@ -280,7 +280,8 @@ public:
 			&& ArraySize == rhs.ArraySize
 			&& NumMips == rhs.NumMips
 			&& NumSamples == rhs.NumSamples
-			&& PackedBits == rhs.PackedBits;
+			&& PackedBits == rhs.PackedBits
+			&& FastVRAMPercentage == rhs.FastVRAMPercentage;
 	}
 
 	bool IsCubemap() const
@@ -419,7 +420,8 @@ public:
 	uint8 NumMips = 0;
 	/** Number of MSAA samples, default: 1  */
 	uint8 NumSamples = 1;
-
+	/** Resource memory percentage which should be allocated onto fast VRAM (hint-only). (encoding into 8bits, 0..255 -> 0%..100%) */
+	uint8 FastVRAMPercentage = 0xFF;
 	union
 	{
 		struct
diff --git a/Engine/Source/Runtime/Renderer/Private/Strata/Strata.cpp b/Engine/Source/Runtime/Renderer/Private/Strata/Strata.cpp
--- a/Engine/Source/Runtime/Renderer/Private/Strata/Strata.cpp
+++ b/Engine/Source/Runtime/Renderer/Private/Strata/Strata.cpp
@@ -453,8 +453,7 @@ void InitialiseStrataFrameSceneData(FRDGBuilder& GraphBuilder, FSceneRenderer& S
 	const uint32 SliceCountAdvDebug = IsAdvancedVisualizationEnabled() ? 1 : 0;
 	const uint32 SliceCount = FMath::DivideAndRoundUp(Out.MaxBytesPerPixel, 4u) + SliceCountSSS + SliceCountAdvDebug;
 	FRDGTextureDesc MaterialTextureDesc = FRDGTextureDesc::Create2DArray(SceneTextureExtent, PF_R32_UINT, FClearValueBinding::Transparent, TexCreate_TargetArraySlicesIndependently | TexCreate_DisableDCC | TexCreate_NoFastClear | TexCreate_RenderTargetable | TexCreate_ShaderResource | TexCreate_UAV | TexCreate_FastVRAM, SliceCount, 1, 1);
-	// Disable temporarlly partial VRAM allocation as it causes issue with RDG
-	MaterialTextureDesc.FastVRAMPercentage = 0xff;  //(1.0f / SliceCount) * 0xFF; // Only allocate the first slice into ESRAM
+	MaterialTextureDesc.FastVRAMPercentage = (1.0f / SliceCount) * 0xFF; // Only allocate the first slice into ESRAM
 
 	Out.MaterialTextureArray = GraphBuilder.CreateTexture(MaterialTextureDesc, TEXT("Substrate.Material"));
 	Out.MaterialTextureArraySRV = GraphBuilder.CreateSRV(FRDGTextureSRVDesc::Create(Out.MaterialTextureArray));
