Refactor distance field object culling to handle one view at a time.

Since the calling function will be iterating over all views, there's no need for the distance field object culling function to do so.

diff --git a/Engine/Source/Runtime/Renderer/Private/DistanceFieldAmbientOcclusion.cpp b/Engine/Source/Runtime/Renderer/Private/DistanceFieldAmbientOcclusion.cpp
--- a/Engine/Source/Runtime/Renderer/Private/DistanceFieldAmbientOcclusion.cpp
+++ b/Engine/Source/Runtime/Renderer/Private/DistanceFieldAmbientOcclusion.cpp
@@ -828,6 +828,8 @@ void FDeferredShadingSceneRenderer::RenderDistanceFieldLighting(
 		DistanceFieldNormal = GraphBuilder.CreateTexture(Desc, TEXT("DistanceFieldNormal"));
 	}
 
+	ComputeDistanceFieldNormal(GraphBuilder, Views, SceneTextures.UniformBuffer, DistanceFieldNormal, Parameters);
+
 	const FIntPoint TileListGroupSize = GetTileListGroupSizeForView(View);
 	const int32 MaxSceneObjects = FMath::DivideAndRoundUp(Scene->DistanceFieldSceneData.NumObjectsInBuffer, 256) * 256;
 	const bool bAllow16BitIndices = !IsMetalPlatform(GShaderPlatformForFeatureLevel[View.FeatureLevel]);
@@ -849,15 +851,8 @@ void FDeferredShadingSceneRenderer::RenderDistanceFieldLighting(
 			CulledObjectBufferParameters);
 
 		CullObjectsToView(GraphBuilder, Scene, View, Parameters, CulledObjectBufferParameters);
-	}
 
-	ComputeDistanceFieldNormal(GraphBuilder, Views, SceneTextures.UniformBuffer, DistanceFieldNormal, Parameters);
-
-	// Intersect objects with screen tiles, build lists
-	if (UseAOObjectDistanceField())
-	{
-		//@todo - support multiple views - should pass one TileIntersectionParameters per view
-		BuildTileObjectLists(GraphBuilder, Scene, Views, ObjectIndirectArguments, CulledObjectBufferParameters, TileIntersectionParameters, DistanceFieldNormal, Parameters);
+		BuildTileObjectLists(GraphBuilder, Scene, View, SceneTextures.UniformBuffer, ObjectIndirectArguments, CulledObjectBufferParameters, TileIntersectionParameters, DistanceFieldNormal, Parameters);
 	}
 
 	FRDGTextureRef BentNormalOutput = nullptr;
diff --git a/Engine/Source/Runtime/Renderer/Private/DistanceFieldAmbientOcclusion.h b/Engine/Source/Runtime/Renderer/Private/DistanceFieldAmbientOcclusion.h
--- a/Engine/Source/Runtime/Renderer/Private/DistanceFieldAmbientOcclusion.h
+++ b/Engine/Source/Runtime/Renderer/Private/DistanceFieldAmbientOcclusion.h
@@ -107,7 +107,8 @@ extern bool ShouldDoReflectionEnvironment(const FScene* Scene, const FSceneViewF
 extern void CullObjectsToView(FRDGBuilder& GraphBuilder, FScene* Scene, const FViewInfo& View, const FDistanceFieldAOParameters& Parameters, FDistanceFieldCulledObjectBufferParameters& CulledObjectBuffers);
 extern void BuildTileObjectLists(FRDGBuilder& GraphBuilder,
 	FScene* Scene,
-	TArray<FViewInfo>& Views,
+	const FViewInfo& View,
+	TRDGUniformBufferRef<FSceneTextureUniformParameters> SceneTexturesUniformBuffer,
 	FRDGBufferRef ObjectIndirectArguments,
 	const FDistanceFieldCulledObjectBufferParameters& CulledObjectBufferParameters,
 	FTileIntersectionParameters TileIntersectionParameters,
diff --git a/Engine/Source/Runtime/Renderer/Private/DistanceFieldObjectCulling.cpp b/Engine/Source/Runtime/Renderer/Private/DistanceFieldObjectCulling.cpp
--- a/Engine/Source/Runtime/Renderer/Private/DistanceFieldObjectCulling.cpp
+++ b/Engine/Source/Runtime/Renderer/Private/DistanceFieldObjectCulling.cpp
@@ -345,7 +345,8 @@ FIntPoint GetTileListGroupSizeForView(const FViewInfo& View)
 void BuildTileObjectLists(
 	FRDGBuilder& GraphBuilder,
 	FScene* Scene,
-	TArray<FViewInfo>& Views,
+	const FViewInfo& View,
+	TRDGUniformBufferRef<FSceneTextureUniformParameters> SceneTexturesUniformBuffer,
 	FRDGBufferRef ObjectIndirectArguments,
 	const FDistanceFieldCulledObjectBufferParameters& CulledObjectBufferParameters,
 	FTileIntersectionParameters TileIntersectionParameters,
@@ -355,64 +356,57 @@ void BuildTileObjectLists(
 	ensure(GAOScatterTileCulling);
 
 	RDG_EVENT_SCOPE(GraphBuilder, "BuildTileList");
+	RDG_GPU_MASK_SCOPE(GraphBuilder, View.GPUMask);
 
-	TRDGUniformBufferRef<FSceneTextureUniformParameters> SceneTexturesUniformBuffer = GetViewFamilyInfo(Views).GetSceneTextures().UniformBuffer;
+	const FIntPoint TileListGroupSize = GetTileListGroupSizeForView(View);
 
-	for (int32 ViewIndex = 0; ViewIndex < Views.Num(); ViewIndex++)
 	{
-		const FViewInfo& View = Views[ViewIndex];
-		RDG_GPU_MASK_SCOPE(GraphBuilder, View.GPUMask);
+		auto* PassParameters = GraphBuilder.AllocParameters<FBuildTileConesCS::FParameters>();
+		PassParameters->View = View.ViewUniformBuffer;
+		PassParameters->RWTileConeAxisAndCos = TileIntersectionParameters.RWTileConeAxisAndCos;
+		PassParameters->RWTileConeDepthRanges = TileIntersectionParameters.RWTileConeDepthRanges;
+		PassParameters->SceneTextures = SceneTexturesUniformBuffer;
+		PassParameters->DistanceFieldNormalTexture = DistanceFieldNormal;
+		PassParameters->DistanceFieldNormalSampler = TStaticSamplerState<SF_Point, AM_Clamp, AM_Clamp, AM_Clamp>::GetRHI();
+		PassParameters->AOParameters = DistanceField::SetupAOShaderParameters(Parameters);
+		PassParameters->NumGroups = FVector2f(TileListGroupSize.X, TileListGroupSize.Y);
 
-		const FIntPoint TileListGroupSize = GetTileListGroupSizeForView(View);
+		auto ComputeShader = View.ShaderMap->GetShader<FBuildTileConesCS>();
 
-		{
-			auto* PassParameters = GraphBuilder.AllocParameters<FBuildTileConesCS::FParameters>();
-			PassParameters->View = View.ViewUniformBuffer;
-			PassParameters->RWTileConeAxisAndCos = TileIntersectionParameters.RWTileConeAxisAndCos;
-			PassParameters->RWTileConeDepthRanges = TileIntersectionParameters.RWTileConeDepthRanges;
-			PassParameters->SceneTextures = SceneTexturesUniformBuffer;
-			PassParameters->DistanceFieldNormalTexture = DistanceFieldNormal;
-			PassParameters->DistanceFieldNormalSampler = TStaticSamplerState<SF_Point, AM_Clamp, AM_Clamp, AM_Clamp>::GetRHI();
-			PassParameters->AOParameters = DistanceField::SetupAOShaderParameters(Parameters);
-			PassParameters->NumGroups = FVector2f(TileListGroupSize.X, TileListGroupSize.Y);
-
-			auto ComputeShader = View.ShaderMap->GetShader<FBuildTileConesCS>();
-
-			FComputeShaderUtils::AddPass(GraphBuilder, RDG_EVENT_NAME("BuildTileCones"), ComputeShader, PassParameters, FIntVector(TileListGroupSize.X, TileListGroupSize.Y, 1));
-		}
+		FComputeShaderUtils::AddPass(GraphBuilder, RDG_EVENT_NAME("BuildTileCones"), ComputeShader, PassParameters, FIntVector(TileListGroupSize.X, TileListGroupSize.Y, 1));
+	}
 
-		// Start at 0 tiles per object
-		AddClearUAVPass(GraphBuilder, TileIntersectionParameters.RWNumCulledTilesArray, 0);
+	// Start at 0 tiles per object
+	AddClearUAVPass(GraphBuilder, TileIntersectionParameters.RWNumCulledTilesArray, 0);
 
-		// Rasterize object bounding shapes and intersect with screen tiles to compute how many tiles intersect each object
-		ScatterTilesToObjects(GraphBuilder, true, View, Scene->DistanceFieldSceneData, TileListGroupSize, Parameters, ObjectIndirectArguments, CulledObjectBufferParameters, TileIntersectionParameters, SceneTexturesUniformBuffer);
+	// Rasterize object bounding shapes and intersect with screen tiles to compute how many tiles intersect each object
+	ScatterTilesToObjects(GraphBuilder, true, View, Scene->DistanceFieldSceneData, TileListGroupSize, Parameters, ObjectIndirectArguments, CulledObjectBufferParameters, TileIntersectionParameters, SceneTexturesUniformBuffer);
 
-		// Start at 0 threadgroups
-		AddClearUAVPass(GraphBuilder, TileIntersectionParameters.RWObjectTilesIndirectArguments, 0);
+	// Start at 0 threadgroups
+	AddClearUAVPass(GraphBuilder, TileIntersectionParameters.RWObjectTilesIndirectArguments, 0);
 
-		{
-			auto* PassParameters = GraphBuilder.AllocParameters<FComputeCulledTilesStartOffsetCS::FParameters>();
-			PassParameters->View = View.ViewUniformBuffer;
-			PassParameters->TileIntersectionParameters = TileIntersectionParameters;
-			PassParameters->DistanceFieldCulledObjectBuffers = CulledObjectBufferParameters;
-			PassParameters->DistanceFieldAtlas = DistanceField::SetupAtlasParameters(GraphBuilder, Scene->DistanceFieldSceneData);
-			PassParameters->SceneTextures = SceneTexturesUniformBuffer;
-
-			auto ComputeShader = View.ShaderMap->GetShader<FComputeCulledTilesStartOffsetCS>();
-			const FIntVector GroupCount = FComputeShaderUtils::GetGroupCountWrapped(Scene->DistanceFieldSceneData.NumObjectsInBuffer, ComputeStartOffsetGroupSize);
-
-			FComputeShaderUtils::AddPass(
-				GraphBuilder,
-				RDG_EVENT_NAME("ComputeStartOffsets"),
-				ComputeShader,
-				PassParameters,
-				GroupCount);
-		}
+	{
+		auto* PassParameters = GraphBuilder.AllocParameters<FComputeCulledTilesStartOffsetCS::FParameters>();
+		PassParameters->View = View.ViewUniformBuffer;
+		PassParameters->TileIntersectionParameters = TileIntersectionParameters;
+		PassParameters->DistanceFieldCulledObjectBuffers = CulledObjectBufferParameters;
+		PassParameters->DistanceFieldAtlas = DistanceField::SetupAtlasParameters(GraphBuilder, Scene->DistanceFieldSceneData);
+		PassParameters->SceneTextures = SceneTexturesUniformBuffer;
 
-		// Start at 0 tiles per object
-		AddClearUAVPass(GraphBuilder, TileIntersectionParameters.RWNumCulledTilesArray, 0);
+		auto ComputeShader = View.ShaderMap->GetShader<FComputeCulledTilesStartOffsetCS>();
+		const FIntVector GroupCount = FComputeShaderUtils::GetGroupCountWrapped(Scene->DistanceFieldSceneData.NumObjectsInBuffer, ComputeStartOffsetGroupSize);
 
-		// Rasterize object bounding shapes and intersect with screen tiles, and write out intersecting tile indices for the cone tracing pass
-		ScatterTilesToObjects(GraphBuilder, false, View, Scene->DistanceFieldSceneData, TileListGroupSize, Parameters, ObjectIndirectArguments, CulledObjectBufferParameters, TileIntersectionParameters, SceneTexturesUniformBuffer);
+		FComputeShaderUtils::AddPass(
+			GraphBuilder,
+			RDG_EVENT_NAME("ComputeStartOffsets"),
+			ComputeShader,
+			PassParameters,
+			GroupCount);
 	}
+
+	// Start at 0 tiles per object
+	AddClearUAVPass(GraphBuilder, TileIntersectionParameters.RWNumCulledTilesArray, 0);
+
+	// Rasterize object bounding shapes and intersect with screen tiles, and write out intersecting tile indices for the cone tracing pass
+	ScatterTilesToObjects(GraphBuilder, false, View, Scene->DistanceFieldSceneData, TileListGroupSize, Parameters, ObjectIndirectArguments, CulledObjectBufferParameters, TileIntersectionParameters, SceneTexturesUniformBuffer);
 }
