Forward declare class types to resolve compiler errors.

diff --git a/Engine/Source/Runtime/Renderer/Private/DistanceFieldAmbientOcclusion.h b/Engine/Source/Runtime/Renderer/Private/DistanceFieldAmbientOcclusion.h
--- a/Engine/Source/Runtime/Renderer/Private/DistanceFieldAmbientOcclusion.h
+++ b/Engine/Source/Runtime/Renderer/Private/DistanceFieldAmbientOcclusion.h
@@ -14,11 +14,13 @@ class FRDGBuffer;
 class FRDGBuilder;
 class FRDGTexture;
 class FScene;
+class FSceneTextureUniformParameters;
 class FSceneViewFamily;
 class FViewInfo;
 struct FShaderCompilerEnvironment;
 using FRDGBufferRef = FRDGBuffer*;
 using FRDGTextureRef = FRDGTexture*;
+template <typename TUniformStruct> using TRDGUniformBufferRef = TRDGUniformBuffer<TUniformStruct>*;
 
 /** Base downsample factor that all distance field AO operations are done at. */
 const int32 GAODownsampleFactor = 2;
