[SceneCapture] skip EndInitViews in depth prepass capture optimization mode, saves ~0.4ms on RenderThread in replay Insight result

diff --git a/Engine/Source/Runtime/Renderer/Private/DeferredShadingRenderer.cpp b/Engine/Source/Runtime/Renderer/Private/DeferredShadingRenderer.cpp
--- a/Engine/Source/Runtime/Renderer/Private/DeferredShadingRenderer.cpp
+++ b/Engine/Source/Runtime/Renderer/Private/DeferredShadingRenderer.cpp
@@ -2999,6 +2999,7 @@ void FDeferredShadingSceneRenderer::Render(FRDGBuilder& GraphBuilder)
 	}
 
 	{
+		if (RendererOutput == ERendererOutput::FinalSceneColor)
 		{
 			RDG_RHI_GPU_STAT_SCOPE(GraphBuilder, VisibilityCommands);
 			EndInitViews(GraphBuilder, LumenFrameTemporaries, InstanceCullingManager, ExternalAccessQueue, InitViewTaskDatas);
