Override CanRunInEnvironment for System.Renderer.RayTracing.BasicRayTracing test.

diff --git a/Engine/Source/Runtime/Renderer/Private/Tests/RayTracingTestbed.cpp b/Engine/Source/Runtime/Renderer/Private/Tests/RayTracingTestbed.cpp
--- a/Engine/Source/Runtime/Renderer/Private/Tests/RayTracingTestbed.cpp
+++ b/Engine/Source/Runtime/Renderer/Private/Tests/RayTracingTestbed.cpp
@@ -6,7 +6,36 @@
 
 #if WITH_DEV_AUTOMATION_TESTS || WITH_EDITOR
 
-IMPLEMENT_SIMPLE_AUTOMATION_TEST(FRayTracingTestbed, "System.Renderer.RayTracing.BasicRayTracing", EAutomationTestFlags::EditorContext | EAutomationTestFlags::HighPriority | EAutomationTestFlags::EngineFilter)
+class FRayTracingTestbedBase : public FAutomationTestBase
+{
+public:
+	FRayTracingTestbedBase(const FString& InName, const bool bInComplexTask)
+		: FAutomationTestBase(InName, bInComplexTask)
+	{
+	}
+
+	virtual bool CanRunInEnvironment(const FString& TestParams, FString* OutReason, bool* OutWarn) const override
+	{
+		if (!GRHISupportsRayTracing || !GRHISupportsRayTracingShaders)
+		{
+			if (OutReason)
+			{
+				*OutReason = TEXT("RHI does not support Ray Tracing and/or Ray Tracing Shaders.");
+			}
+
+			if (OutWarn)
+			{
+				*OutWarn = false;
+			}
+
+			return false;
+		}
+
+		return true;
+	}
+};
+
+IMPLEMENT_CUSTOM_SIMPLE_AUTOMATION_TEST(FRayTracingTestbed, FRayTracingTestbedBase, "System.Renderer.RayTracing.BasicRayTracing", EAutomationTestFlags::EditorContext | EAutomationTestFlags::HighPriority | EAutomationTestFlags::EngineFilter)
 
 #if RHI_RAYTRACING
 
