Fixed crash with shadows being created but never synced when rendering depth output.

diff --git a/Engine/Source/Runtime/Renderer/Private/SceneVisibility.cpp b/Engine/Source/Runtime/Renderer/Private/SceneVisibility.cpp
--- a/Engine/Source/Runtime/Renderer/Private/SceneVisibility.cpp
+++ b/Engine/Source/Runtime/Renderer/Private/SceneVisibility.cpp
@@ -4845,7 +4845,10 @@ void FDeferredShadingSceneRenderer::BeginInitViews(
 
 	TaskDatas.VisibilityTaskData->ProcessRenderThreadTasks(BasePassDepthStencilAccess, InstanceCullingManager, VirtualTextureUpdater);
 
-	BeginInitDynamicShadows(TaskDatas);
+	if (GetRendererOutput() == ERendererOutput::FinalSceneColor)
+	{
+		BeginInitDynamicShadows(TaskDatas);
+	}
 
 	// This must happen before we start initialising and using views.
 	if (Scene)
