Add an option for rendering shadows early in the frame, allowing for better async compute sheduling on some platforms.
This works well with async SSAO, allowing it to overlap with shadow depths rather than basepass. Note: this is not supported with VSMs.
This improves dynamic res on last gen consoles by 0.4-0.6%

diff --git a/Engine/Source/Runtime/Renderer/Private/DeferredShadingRenderer.cpp b/Engine/Source/Runtime/Renderer/Private/DeferredShadingRenderer.cpp
--- a/Engine/Source/Runtime/Renderer/Private/DeferredShadingRenderer.cpp
+++ b/Engine/Source/Runtime/Renderer/Private/DeferredShadingRenderer.cpp
@@ -309,9 +309,10 @@ static TAutoConsoleVariable<int32> CVarSceneDepthHZBAsyncCompute(
 	TEXT(" 2: Use async compute, start after ComputeLightGrid.CompactLinks pass"),
 	ECVF_RenderThreadSafe);
 
-static TAutoConsoleVariable<int32> CVarShadowsUseSharedExternalAccessQueue(
-	TEXT("r.ShadowsUseSharedExternalAccessQueue"), 1,
-	TEXT("If enabled, shadows will use the shared external access queue, minimizing unnecessary transitions"),
+static TAutoConsoleVariable<int32> CVarShadowMapsRenderEarly(
+	TEXT("r.shadow.ShadowMapsRenderEarly"), 0,
+	TEXT("If enabled, shadows will render earlier in the frame. This can help async compute scheduling on some platforms\n")
+	TEXT("Note: This is not compatible with VSMs\n"),
 	ECVF_RenderThreadSafe);
 
 static TAutoConsoleVariable<int32> CVarTranslucencyVelocity(
@@ -3476,11 +3477,13 @@ void FDeferredShadingSceneRenderer::Render(FRDGBuilder& GraphBuilder)
 		
 		FRDGTextureRef ForwardScreenSpaceShadowMaskTexture = nullptr;
 		FRDGTextureRef ForwardScreenSpaceShadowMaskHairTexture = nullptr;
+		bool bShadowMapsRenderedEarly = false;
 		if (IsForwardShadingEnabled(ShaderPlatform))
 		{
 			// With forward shading we need to render shadow maps early
 			ensureMsgf(!VirtualShadowMapArray.IsEnabled(), TEXT("Virtual shadow maps are not supported in the forward shading path"));
 			RenderShadowDepthMaps(GraphBuilder, InstanceCullingManager, ExternalAccessQueue);
+			bShadowMapsRenderedEarly = true;
 
 			if (bHairStrandsEnable && !bHasRayTracedOverlay)
 			{
@@ -3493,6 +3496,16 @@ void FDeferredShadingSceneRenderer::Render(FRDGBuilder& GraphBuilder)
 			// With forward shading we need to render volumetric fog before the base pass
 			ComputeVolumetricFog(GraphBuilder, SceneTextures);
 		}
+		else if ( CVarShadowMapsRenderEarly.GetValueOnRenderThread() )
+		{
+			// Disable early shadows if VSM is enabled, but warn
+			ensureMsgf(!VirtualShadowMapArray.IsEnabled(), TEXT("Virtual shadow maps are not supported with r.shadow.ShadowMapsRenderEarly. Early shadows will be disabled"));
+			if (!VirtualShadowMapArray.IsEnabled())
+			{
+				RenderShadowDepthMaps(GraphBuilder, InstanceCullingManager, ExternalAccessQueue);
+				bShadowMapsRenderedEarly = true;
+			}
+		}
 
 		ExternalAccessQueue.Submit(GraphBuilder);
 
@@ -3700,8 +3713,8 @@ void FDeferredShadingSceneRenderer::Render(FRDGBuilder& GraphBuilder)
 				bHasLumenLights,
 				AsyncLumenIndirectLightingOutputs);
 
-			// If forward shading is enabled, we rendered shadow maps earlier already
-			if (!IsForwardShadingEnabled(ShaderPlatform))
+			// If we haven't already rendered shadow maps, render them now (due to forward shading or r.shadow.ShadowMapsRenderEarly)
+			if (!bShadowMapsRenderedEarly)
 			{
 				if (VirtualShadowMapArray.IsEnabled())
 				{
