Backout CL 25813358 due to crash - [SceneCapture] skip EndInitViews in depth prepass capture optimization mode, saves ~0.4ms on RenderThread in replay Insight result

diff --git a/Engine/Source/Runtime/Renderer/Private/DeferredShadingRenderer.cpp b/Engine/Source/Runtime/Renderer/Private/DeferredShadingRenderer.cpp
--- a/Engine/Source/Runtime/Renderer/Private/DeferredShadingRenderer.cpp
+++ b/Engine/Source/Runtime/Renderer/Private/DeferredShadingRenderer.cpp
@@ -3000,7 +3000,6 @@ void FDeferredShadingSceneRenderer::Render(FRDGBuilder& GraphBuilder)
 	}
 
 	{
-		if (RendererOutput == ERendererOutput::FinalSceneColor)
 		{
 			RDG_RHI_GPU_STAT_SCOPE(GraphBuilder, VisibilityCommands);
 			EndInitViews(GraphBuilder, LumenFrameTemporaries, InstanceCullingManager, ExternalAccessQueue, InitViewTaskDatas);
