Quick fix for issue where per instance custom data is not present in GPUScene upload.
Need further investigation for root cause.

diff --git a/Engine/Source/Runtime/Renderer/Private/GPUScene.cpp b/Engine/Source/Runtime/Renderer/Private/GPUScene.cpp
--- a/Engine/Source/Runtime/Renderer/Private/GPUScene.cpp
+++ b/Engine/Source/Runtime/Renderer/Private/GPUScene.cpp
@@ -330,7 +330,7 @@ void ValidateInstanceUploadInfo(const FInstanceUploadInfo& UploadInfo, FRDGBuffe
 
 	if (bHasCustomData)
 	{
-		check(UploadInfo.InstanceCustomDataCount > 0);
+		//check(UploadInfo.InstanceCustomDataCount > 0);
 		check(UploadInfo.InstanceCustomDataCount * InstanceCount == UploadInfo.InstanceCustomData.Num());
 	}
 	else
@@ -1243,7 +1243,7 @@ void FGPUScene::UploadGeneral(FRDGBuilder& GraphBuilder, FScene& Scene, FRDGExte
 								PayloadPosition += SPLINE_MESH_PARAMS_FLOAT4_SIZE;
 							}
 
-							if (UploadInfo.InstanceFlags & INSTANCE_SCENE_DATA_FLAG_HAS_CUSTOM_DATA)
+							if (UploadInfo.InstanceFlags & INSTANCE_SCENE_DATA_FLAG_HAS_CUSTOM_DATA && UploadInfo.InstanceCustomDataCount > 0)
 							{
 								check(PayloadPosition + (UploadInfo.InstanceCustomDataCount >> 2u) <= InstancePayloadData.Num());
 								const float* CustomDataGetPtr = &UploadInfo.InstanceCustomData[(InstanceIndex * UploadInfo.InstanceCustomDataCount)];
