Potential fix for crash on Android related to a missing resource in occlusion rendering

diff --git a/Engine/Source/Runtime/Renderer/Private/SceneOcclusion.cpp b/Engine/Source/Runtime/Renderer/Private/SceneOcclusion.cpp
--- a/Engine/Source/Runtime/Renderer/Private/SceneOcclusion.cpp
+++ b/Engine/Source/Runtime/Renderer/Private/SceneOcclusion.cpp
@@ -21,6 +21,7 @@
 #include "ProfilingDebugging/CpuProfilerTrace.h"
 #include "PixelShaderUtils.h"
 #include "RenderCore.h"
+#include "MobileBasePassRendering.h"
 
 
 /*-----------------------------------------------------------------------------
@@ -1850,7 +1851,9 @@ public:
 	class FPixelDiscard : SHADER_PERMUTATION_BOOL("PERMUTATION_PIXEL_DISCARD");
 	using FPermutationDomain = TShaderPermutationDomain<FPixelDiscard>;
 
-	using FParameters = FEmptyShaderParameters;
+	BEGIN_SHADER_PARAMETER_STRUCT(FParameters, )
+		SHADER_PARAMETER_RDG_UNIFORM_BUFFER(FMobileBasePassUniformParameters, MobileBasePass)
+	END_SHADER_PARAMETER_STRUCT()
 };
 
 IMPLEMENT_GLOBAL_SHADER(FOcclusionFeedbackVS,"/Engine/Private/OcclusionFeedbackShaders.usf","MainVS",SF_Vertex);
