VolumetircCloud - added r.VolumetricRenderTarget.MinimumDistanceKmToEnableReprojection. This helps hide reprojection issues due to imperfect approximation of cloud depth as a single front surface, especially visible when flying through the cloud layer. It is not perfect but will help in lots of cases. The problem when using this method: clouds will look noisier when closer to that distance.

diff --git a/Engine/Shaders/Private/VolumetricRenderTarget.usf b/Engine/Shaders/Private/VolumetricRenderTarget.usf
--- a/Engine/Shaders/Private/VolumetricRenderTarget.usf
+++ b/Engine/Shaders/Private/VolumetricRenderTarget.usf
@@ -24,6 +24,7 @@ Texture2D<float4> SecondaryTracingVolumetricTexture;
 Texture2D<float4> TracingVolumetricDepthTexture;
 uint4 TracingVolumetricTextureValidCoordRect;
 float4 TracingVolumetricTextureValidUvRect;
+float MinimumDistanceKmToEnableReprojection;
 
 struct FDepthData
 {
@@ -184,7 +185,10 @@ void ReconstructVolumetricRenderTargetPS(
 
 		float2 PrevScreenPosition = (ScreenPosition - ScreenVelocity);   // NDC in [-1,1]
 		float2 PrevScreenUVs = ScreenPosToViewportUV(PrevScreenPosition);// UV in [0,1]
-		const bool bValidPreviousUVs = all(PrevScreenUVs > 0.0) && all(PrevScreenUVs < 1.0f);
+		const bool bValidPreviousUVs = all(PrevScreenUVs > 0.0) && all(PrevScreenUVs < 1.0f) 
+			&& (TracingVolumetricSampleDepthKm >= MinimumDistanceKmToEnableReprojection);	// This helps hide reprojection issues due to imperfect approximation of cloud depth as a single front surface, 
+																							// especially visible when flying through the cloud layer. It is not perfect but will help in lots of cases. 
+																							// The problem when using this method: clouds will look noisier when closer to that distance. 
 
 		if (VolumetricRenderTargetMode == 2)
 		{
diff --git a/Engine/Source/Runtime/Renderer/Private/VolumetricRenderTarget.cpp b/Engine/Source/Runtime/Renderer/Private/VolumetricRenderTarget.cpp
--- a/Engine/Source/Runtime/Renderer/Private/VolumetricRenderTarget.cpp
+++ b/Engine/Source/Runtime/Renderer/Private/VolumetricRenderTarget.cpp
@@ -44,6 +44,10 @@ static TAutoConsoleVariable<int32> CVarVolumetricRenderTargetReprojectionBoxCons
 	TEXT("Whether reprojected data should be constrained to the new incoming cloud data neighborhod value."),
 	ECVF_RenderThreadSafe | ECVF_Scalability);
 
+static TAutoConsoleVariable<float> CVarVolumetricRenderTargetMinimumDistanceKmToEnableReprojection(
+	TEXT("r.VolumetricRenderTarget.MinimumDistanceKmToEnableReprojection"), 0.0f,
+	TEXT("This is the distance in kilometer at which the `cloud surface` must be before we enable reprojection of the previous frame data. One could start with a value of 4km. This helps hide reprojection issues due to imperfect approximation of cloud depth as a single front surface, especially visible when flying through the cloud layer. It is not perfect but will help in lots of cases. The problem when using this method: clouds will look noisier when closer to that distance."),
+	ECVF_RenderThreadSafe | ECVF_Scalability);
 
 static float GetUvNoiseSampleAcceptanceWeight()
 {
@@ -517,7 +521,7 @@ class FReconstructVolumetricRenderTargetPS : public FGlobalShader
 		SHADER_PARAMETER(FVector4f, TracingVolumetricTextureValidUvRect)
 		SHADER_PARAMETER(FUintVector4, PreviousFrameVolumetricTextureValidCoordRect)
 		SHADER_PARAMETER(FVector4f, PreviousFrameVolumetricTextureValidUvRect)
-		SHADER_PARAMETER(float, TemporalFactor)
+		SHADER_PARAMETER(float, MinimumDistanceKmToEnableReprojection)
 	END_SHADER_PARAMETER_STRUCT()
 
 	static FPermutationDomain RemapPermutation(FPermutationDomain PermutationVector)
@@ -608,6 +612,7 @@ void ReconstructVolumetricRenderTarget(
 		PassParameters->DownSampleFactor = TracingVolumetricCloudRTDownSample;
 		PassParameters->VolumetricRenderTargetMode = VolumetricCloudRT.GetMode();
 		PassParameters->HalfResDepthTexture = (VolumetricCloudRT.GetMode() == 0 || VolumetricCloudRT.GetMode() == 3) ? HalfResolutionDepthCheckerboardMinMaxTexture : SceneDepthTexture;
+		PassParameters->MinimumDistanceKmToEnableReprojection = FMath::Max(0.0f, CVarVolumetricRenderTargetMinimumDistanceKmToEnableReprojection.GetValueOnRenderThread());
 
 		const bool bVisualizeConservativeDensity = ShouldViewVisualizeVolumetricCloudConservativeDensity(ViewInfo, ViewInfo.Family->EngineShowFlags);
 		PassParameters->HalfResDepthTexture = bVisualizeConservativeDensity ?
