Fix a bug of missing simple light on mobile deferred.

diff --git a/Engine/Source/Runtime/Renderer/Private/MobileDeferredShadingPass.cpp b/Engine/Source/Runtime/Renderer/Private/MobileDeferredShadingPass.cpp
--- a/Engine/Source/Runtime/Renderer/Private/MobileDeferredShadingPass.cpp
+++ b/Engine/Source/Runtime/Renderer/Private/MobileDeferredShadingPass.cpp
@@ -915,6 +915,8 @@ static void RenderSimpleLights(
 			uint8 StencilRef = GET_STENCIL_MOBILE_SM_MASK(PassShadingModelStencilValue[PassIndex]);
 			SetGraphicsPipelineState(RHICmdList, GraphicsPSOLight[PassIndex], StencilRef);
 
+			SetShaderParameters(RHICmdList, VertexShader, VertexShader.GetVertexShader(), ParametersVS);
+
 			SetShaderParametersMixedPS(RHICmdList, PassPixelShaders[PassIndex], PassParameters, View, DefaultMaterial.MaterialProxy, *DefaultMaterial.Material);
 
 			// Apply the point or spot light with some approximately bounding geometry,
