Iterate over all views to render distance field lighting.

Not a full fix for multi-view support yet but allows multi-view rendering to occur now. Behavior for single-view should remain the same.

diff --git a/Engine/Source/Runtime/Renderer/Private/DistanceFieldAmbientOcclusion.cpp b/Engine/Source/Runtime/Renderer/Private/DistanceFieldAmbientOcclusion.cpp
--- a/Engine/Source/Runtime/Renderer/Private/DistanceFieldAmbientOcclusion.cpp
+++ b/Engine/Source/Runtime/Renderer/Private/DistanceFieldAmbientOcclusion.cpp
@@ -381,10 +381,13 @@ void ComputeDistanceFieldNormal(
 {
 	if (GAOComputeShaderNormalCalculation)
 	{
+		RDG_EVENT_SCOPE(GraphBuilder, "ComputeDistanceFieldNormalCS");
+
 		for (int32 ViewIndex = 0; ViewIndex < Views.Num(); ViewIndex++)
 		{
 			const FViewInfo& View = Views[ViewIndex];
 			RDG_GPU_MASK_SCOPE(GraphBuilder, View.GPUMask);
+			RDG_EVENT_SCOPE_CONDITIONAL(GraphBuilder, Views.Num() > 1, "View%d", ViewIndex);
 
 			uint32 GroupSizeX = FMath::DivideAndRoundUp(View.ViewRect.Size().X / GAODownsampleFactor, GDistanceFieldAOTileSizeX);
 			uint32 GroupSizeY = FMath::DivideAndRoundUp(View.ViewRect.Size().Y / GAODownsampleFactor, GDistanceFieldAOTileSizeY);
@@ -403,10 +406,13 @@ void ComputeDistanceFieldNormal(
 	}
 	else
 	{
+		RDG_EVENT_SCOPE(GraphBuilder, "ComputeDistanceFieldNormal");
+
 		for (int32 ViewIndex = 0; ViewIndex < Views.Num(); ViewIndex++)
 		{
 			const FViewInfo& View = Views[ViewIndex];
 			RDG_GPU_MASK_SCOPE(GraphBuilder, View.GPUMask);
+			RDG_EVENT_SCOPE_CONDITIONAL(GraphBuilder, Views.Num() > 1, "View%d", ViewIndex);
 
 			auto* PassParameters = GraphBuilder.AllocParameters<FComputeDistanceFieldNormalPS::FParameters>();
 			PassParameters->View = View.ViewUniformBuffer;
@@ -660,10 +666,9 @@ void AllocateTileIntersectionBuffers(
 	OutParameters.TileListGroupSize = TileListGroupSize;
 }
 
-void ListDistanceFieldLightingMemory(const FViewInfo& View, FSceneRenderer& SceneRenderer)
+void ListDistanceFieldLightingMemory(const FScene* Scene, FSceneRenderer& SceneRenderer)
 {
 #if !NO_LOGGING
-	const FScene* Scene = (const FScene*)View.Family->Scene;
 	UE_LOG(LogRenderer, Log, TEXT("Shared GPU memory (excluding render targets)"));
 
 	if (Scene->DistanceFieldSceneData.NumObjectsInBuffer > 0)
@@ -779,13 +784,22 @@ void FDeferredShadingSceneRenderer::RenderDFAOAsIndirectShadowing(
 
 bool FDeferredShadingSceneRenderer::ShouldRenderDistanceFieldLighting() const
 {
-	//@todo - support multiple views
-	const FViewInfo& View = Views[0];
+	//@to-do - check for ISR (may want to exclude certain checks on secondary views)
+	bool bSupportsDistanceFieldAO = true;
+
+	for (int32_t ViewIndex = 0; ViewIndex < Views.Num(); ++ViewIndex)
+	{
+		const FViewInfo& View = Views[ViewIndex];
+
+		if (!SupportsDistanceFieldAO(View.GetFeatureLevel(), View.GetShaderPlatform())
+			|| !View.IsPerspectiveProjection())
+		{
+			bSupportsDistanceFieldAO = false;
+			break;
+		}
+	}
 
-	return SupportsDistanceFieldAO(View.GetFeatureLevel(), View.GetShaderPlatform())
-		&& Views.Num() == 1
-		&& View.IsPerspectiveProjection()
-		&& Scene->DistanceFieldSceneData.NumObjectsInBuffer;
+	return bSupportsDistanceFieldAO && Scene->DistanceFieldSceneData.NumObjectsInBuffer;
 }
 
 void FDeferredShadingSceneRenderer::RenderDistanceFieldLighting(
@@ -799,10 +813,6 @@ void FDeferredShadingSceneRenderer::RenderDistanceFieldLighting(
 	check(ShouldRenderDistanceFieldLighting());
 	check(!Scene->DistanceFieldSceneData.HasPendingOperations());
 
-	//@todo - support multiple views
-	const FViewInfo& View = Views[0];
-
-	RDG_GPU_MASK_SCOPE(GraphBuilder, View.GPUMask);
 	RDG_EVENT_SCOPE(GraphBuilder, "DistanceFieldLighting");
 	QUICK_SCOPE_CYCLE_COUNTER(STAT_RenderDistanceFieldLighting);
 
@@ -811,7 +821,7 @@ void FDeferredShadingSceneRenderer::RenderDistanceFieldLighting(
 	if (bListMemoryNextFrame)
 	{
 		bListMemoryNextFrame = false;
-		ListDistanceFieldLightingMemory(View, *this);
+		ListDistanceFieldLightingMemory(Scene, *this);
 	}
 
 	if (bListMeshDistanceFieldsMemoryNextFrame)
@@ -823,50 +833,57 @@ void FDeferredShadingSceneRenderer::RenderDistanceFieldLighting(
 	FRDGTextureRef DistanceFieldNormal = nullptr;
 
 	{
-		const FIntPoint BufferSize = GetBufferSizeForAO(View);
+		const FIntPoint BufferSize = GetBufferSizeForAO(Views[0]);
 		const FRDGTextureDesc Desc = FRDGTextureDesc::Create2D(BufferSize, PF_FloatRGBA, FClearValueBinding::Transparent, GFastVRamConfig.DistanceFieldNormal | TexCreate_RenderTargetable | TexCreate_UAV | TexCreate_ShaderResource);
 		DistanceFieldNormal = GraphBuilder.CreateTexture(Desc, TEXT("DistanceFieldNormal"));
 	}
 
 	ComputeDistanceFieldNormal(GraphBuilder, Views, SceneTextures.UniformBuffer, DistanceFieldNormal, Parameters);
 
-	const FIntPoint TileListGroupSize = GetTileListGroupSizeForView(View);
-	const int32 MaxSceneObjects = FMath::DivideAndRoundUp(Scene->DistanceFieldSceneData.NumObjectsInBuffer, 256) * 256;
-	const bool bAllow16BitIndices = !IsMetalPlatform(GShaderPlatformForFeatureLevel[View.FeatureLevel]);
+	FRDGTextureRef BentNormalOutput = nullptr;
+
+	for (int32 ViewIndex = 0; ViewIndex < Views.Num(); ++ViewIndex)
+	{
+		const FViewInfo& View = Views[ViewIndex];
+		RDG_GPU_MASK_SCOPE(GraphBuilder, View.GPUMask);
+		RDG_EVENT_SCOPE_CONDITIONAL(GraphBuilder, Views.Num() > 1, "View%d", ViewIndex);
 
-	FRDGBufferRef ObjectTilesIndirectArguments = nullptr;
-	FTileIntersectionParameters TileIntersectionParameters;
+		const FIntPoint TileListGroupSize = GetTileListGroupSizeForView(View);
+		const int32 MaxSceneObjects = FMath::DivideAndRoundUp(Scene->DistanceFieldSceneData.NumObjectsInBuffer, 256) * 256;
+		const bool bAllow16BitIndices = !IsMetalPlatform(GShaderPlatformForFeatureLevel[View.FeatureLevel]);
 
-	AllocateTileIntersectionBuffers(GraphBuilder, TileListGroupSize, MaxSceneObjects, bAllow16BitIndices, ObjectTilesIndirectArguments, TileIntersectionParameters);
+		FRDGBufferRef ObjectTilesIndirectArguments = nullptr;
+		FTileIntersectionParameters TileIntersectionParameters;
 
-	FRDGBufferRef ObjectIndirectArguments = nullptr;
-	FDistanceFieldCulledObjectBufferParameters CulledObjectBufferParameters;
+		AllocateTileIntersectionBuffers(GraphBuilder, TileListGroupSize, MaxSceneObjects, bAllow16BitIndices, ObjectTilesIndirectArguments, TileIntersectionParameters);
 
-	if (UseAOObjectDistanceField())
-	{
-		AllocateDistanceFieldCulledObjectBuffers(
-			GraphBuilder,
-			MaxSceneObjects,
-			ObjectIndirectArguments,
-			CulledObjectBufferParameters);
+		FRDGBufferRef ObjectIndirectArguments = nullptr;
+		FDistanceFieldCulledObjectBufferParameters CulledObjectBufferParameters;
 
-		CullObjectsToView(GraphBuilder, Scene, View, Parameters, CulledObjectBufferParameters);
+		if (UseAOObjectDistanceField())
+		{
+			AllocateDistanceFieldCulledObjectBuffers(
+				GraphBuilder,
+				MaxSceneObjects,
+				ObjectIndirectArguments,
+				CulledObjectBufferParameters);
 
-		BuildTileObjectLists(GraphBuilder, Scene, View, SceneTextures.UniformBuffer, ObjectIndirectArguments, CulledObjectBufferParameters, TileIntersectionParameters, DistanceFieldNormal, Parameters);
-	}
+			CullObjectsToView(GraphBuilder, Scene, View, Parameters, CulledObjectBufferParameters);
 
-	FRDGTextureRef BentNormalOutput = nullptr;
+			BuildTileObjectLists(GraphBuilder, Scene, View, SceneTextures.UniformBuffer, ObjectIndirectArguments, CulledObjectBufferParameters, TileIntersectionParameters, DistanceFieldNormal, Parameters);
+		}
 
-	RenderDistanceFieldAOScreenGrid(
-		GraphBuilder,
-		SceneTextures,
-		View,
-		CulledObjectBufferParameters,
-		ObjectTilesIndirectArguments,
-		TileIntersectionParameters,
-		Parameters,
-		DistanceFieldNormal,
-		BentNormalOutput);
+		RenderDistanceFieldAOScreenGrid(
+			GraphBuilder,
+			SceneTextures,
+			View,
+			CulledObjectBufferParameters,
+			ObjectTilesIndirectArguments,
+			TileIntersectionParameters,
+			Parameters,
+			DistanceFieldNormal,
+			BentNormalOutput);
+	}
 
 	RenderCapsuleShadowsForMovableSkylight(GraphBuilder, SceneTextures.UniformBuffer, BentNormalOutput);
 
