Create per-view-based bent normal textures.

This reverts some of the previously committed changes, because we now no longer want to offset into a texture-atlas for all the views. Instead, we create a normal texture per-view (offsetting into other textures as needed), so that the rest of the pass can work with a single texture at a time.

diff --git a/Engine/Shaders/Private/DistanceFieldScreenGridLighting.usf b/Engine/Shaders/Private/DistanceFieldScreenGridLighting.usf
--- a/Engine/Shaders/Private/DistanceFieldScreenGridLighting.usf
+++ b/Engine/Shaders/Private/DistanceFieldScreenGridLighting.usf
@@ -17,7 +17,7 @@ void ComputeDistanceFieldNormalPS(
 	out float4 OutColor : SV_Target0)
 {
 	// Sample from the center of the top left full resolution texel
-	float2 PixelCoord = floor(SVPos.xy) * DOWNSAMPLE_FACTOR + .5f;
+	float2 PixelCoord = floor(SVPos.xy) * DOWNSAMPLE_FACTOR + View.ViewRectMin.xy + .5f;
 	float2 ScreenUV = float2(PixelCoord * View.BufferSizeAndInvSize.zw);
 
 	float SceneDepth = CalcSceneDepth(ScreenUV);
@@ -50,8 +50,7 @@ void ComputeDistanceFieldNormalCS(
 	float3 WorldNormal = GetGBufferData(ScreenUV).WorldNormal.xyz;
 #endif
 
-	const uint2 DownsampledPixelCoords = DispatchThreadId.xy + uint2(floor(View.ViewRectMin.xy / DOWNSAMPLE_FACTOR));
-	RWDistanceFieldNormal[DownsampledPixelCoords] = EncodeDownsampledGBuffer(WorldNormal, SceneDepth);
+	RWDistanceFieldNormal[DispatchThreadId.xy] = EncodeDownsampledGBuffer(WorldNormal, SceneDepth);
 }
 
 Buffer<float4> TileConeDepthRanges;
diff --git a/Engine/Source/Runtime/Renderer/Private/DistanceFieldAmbientOcclusion.cpp b/Engine/Source/Runtime/Renderer/Private/DistanceFieldAmbientOcclusion.cpp
--- a/Engine/Source/Runtime/Renderer/Private/DistanceFieldAmbientOcclusion.cpp
+++ b/Engine/Source/Runtime/Renderer/Private/DistanceFieldAmbientOcclusion.cpp
@@ -159,7 +159,7 @@ void TileIntersectionModifyCompilationEnvironment(EShaderPlatform Platform, FSha
 
 FIntPoint GetBufferSizeForAO(const FViewInfo& View)
 {
-	return FIntPoint::DivideAndRoundDown(View.GetSceneTexturesConfig().Extent, GAODownsampleFactor);
+	return FIntPoint::DivideAndRoundDown(View.ViewRect.Size(), GAODownsampleFactor);
 }
 
 // Sample set restricted to not self-intersect a surface based on cone angle .475882232
@@ -374,94 +374,72 @@ IMPLEMENT_GLOBAL_SHADER(FComputeDistanceFieldNormalCS, "/Engine/Private/Distance
 
 void ComputeDistanceFieldNormal(
 	FRDGBuilder& GraphBuilder,
-	const TArray<FViewInfo>& Views,
+	const FViewInfo& View,
 	TRDGUniformBufferRef<FSceneTextureUniformParameters> SceneTexturesUniformBuffer,
 	FRDGTextureRef DistanceFieldNormal,
 	const FDistanceFieldAOParameters& Parameters)
 {
 	if (GAOComputeShaderNormalCalculation)
 	{
-		RDG_EVENT_SCOPE(GraphBuilder, "ComputeDistanceFieldNormalCS");
+		uint32 GroupSizeX = FMath::DivideAndRoundUp(View.ViewRect.Size().X / GAODownsampleFactor, GDistanceFieldAOTileSizeX);
+		uint32 GroupSizeY = FMath::DivideAndRoundUp(View.ViewRect.Size().Y / GAODownsampleFactor, GDistanceFieldAOTileSizeY);
 
-		for (int32 ViewIndex = 0; ViewIndex < Views.Num(); ViewIndex++)
-		{
-			const FViewInfo& View = Views[ViewIndex];
-			RDG_GPU_MASK_SCOPE(GraphBuilder, View.GPUMask);
-			RDG_EVENT_SCOPE_CONDITIONAL(GraphBuilder, Views.Num() > 1, "View%d", ViewIndex);
-
-			uint32 GroupSizeX = FMath::DivideAndRoundUp(View.ViewRect.Size().X / GAODownsampleFactor, GDistanceFieldAOTileSizeX);
-			uint32 GroupSizeY = FMath::DivideAndRoundUp(View.ViewRect.Size().Y / GAODownsampleFactor, GDistanceFieldAOTileSizeY);
-
-			auto* PassParameters = GraphBuilder.AllocParameters<FComputeDistanceFieldNormalCS::FParameters>();
-			PassParameters->View = View.ViewUniformBuffer;
-			PassParameters->SceneTextures = SceneTexturesUniformBuffer;
-			PassParameters->Strata = Strata::BindStrataGlobalUniformParameters(View);
-			PassParameters->AOParameters = DistanceField::SetupAOShaderParameters(Parameters);
-			PassParameters->RWDistanceFieldNormal = GraphBuilder.CreateUAV(DistanceFieldNormal);
+		auto* PassParameters = GraphBuilder.AllocParameters<FComputeDistanceFieldNormalCS::FParameters>();
+		PassParameters->View = View.ViewUniformBuffer;
+		PassParameters->SceneTextures = SceneTexturesUniformBuffer;
+		PassParameters->Strata = Strata::BindStrataGlobalUniformParameters(View);
+		PassParameters->AOParameters = DistanceField::SetupAOShaderParameters(Parameters);
+		PassParameters->RWDistanceFieldNormal = GraphBuilder.CreateUAV(DistanceFieldNormal);
 
-			TShaderMapRef<FComputeDistanceFieldNormalCS> ComputeShader(View.ShaderMap);
+		TShaderMapRef<FComputeDistanceFieldNormalCS> ComputeShader(View.ShaderMap);
 
-			FComputeShaderUtils::AddPass(GraphBuilder, RDG_EVENT_NAME("ComputeNormalCS"), ComputeShader, PassParameters, FIntVector(GroupSizeX, GroupSizeY, 1));
-		}
+		FComputeShaderUtils::AddPass(GraphBuilder, RDG_EVENT_NAME("ComputeNormalCS"), ComputeShader, PassParameters, FIntVector(GroupSizeX, GroupSizeY, 1));
 	}
 	else
 	{
-		RDG_EVENT_SCOPE(GraphBuilder, "ComputeDistanceFieldNormal");
-
-		for (int32 ViewIndex = 0; ViewIndex < Views.Num(); ViewIndex++)
+		auto* PassParameters = GraphBuilder.AllocParameters<FComputeDistanceFieldNormalPS::FParameters>();
+		PassParameters->View = View.ViewUniformBuffer;
+		PassParameters->SceneTextures = SceneTexturesUniformBuffer;
+		PassParameters->Strata = Strata::BindStrataGlobalUniformParameters(View);
+		PassParameters->AOParameters = DistanceField::SetupAOShaderParameters(Parameters);
+		PassParameters->RenderTargets[0] = FRenderTargetBinding(DistanceFieldNormal, ERenderTargetLoadAction::EClear);
+
+		GraphBuilder.AddPass(
+			RDG_EVENT_NAME("ComputeNormal"),
+			PassParameters,
+			ERDGPassFlags::Raster,
+			[&View, PassParameters](FRHICommandList& RHICmdList)
 		{
-			const FViewInfo& View = Views[ViewIndex];
-			RDG_GPU_MASK_SCOPE(GraphBuilder, View.GPUMask);
-			RDG_EVENT_SCOPE_CONDITIONAL(GraphBuilder, Views.Num() > 1, "View%d", ViewIndex);
-
-			auto* PassParameters = GraphBuilder.AllocParameters<FComputeDistanceFieldNormalPS::FParameters>();
-			PassParameters->View = View.ViewUniformBuffer;
-			PassParameters->SceneTextures = SceneTexturesUniformBuffer;
-			PassParameters->Strata = Strata::BindStrataGlobalUniformParameters(View);
-			PassParameters->AOParameters = DistanceField::SetupAOShaderParameters(Parameters);
-			PassParameters->RenderTargets[0] = FRenderTargetBinding(DistanceFieldNormal, ViewIndex == 0 ? ERenderTargetLoadAction::EClear : ERenderTargetLoadAction::ELoad);
-
-			GraphBuilder.AddPass(
-				RDG_EVENT_NAME("ComputeNormal"),
-				PassParameters,
-				ERDGPassFlags::Raster,
-				[&View, PassParameters](FRHICommandList& RHICmdList)
-			{
-				const float ViewportMinX = View.ViewRect.Min.X / GAODownsampleFactor;
-				const float ViewportMinY = View.ViewRect.Min.Y / GAODownsampleFactor;
-				const float ViewportMaxX = View.ViewRect.Max.X / GAODownsampleFactor;
-				const float ViewportMaxY = View.ViewRect.Max.Y / GAODownsampleFactor;
-				RHICmdList.SetViewport(ViewportMinX, ViewportMinY, 0.0f, ViewportMaxX, ViewportMaxY, 1.0f);
-
-				TShaderMapRef<FPostProcessVS> VertexShader(View.ShaderMap);
-
-				TShaderMapRef<FComputeDistanceFieldNormalPS> PixelShader(View.ShaderMap);
-
-				FGraphicsPipelineStateInitializer GraphicsPSOInit;
-				RHICmdList.ApplyCachedRenderTargets(GraphicsPSOInit);
-				GraphicsPSOInit.RasterizerState = TStaticRasterizerState<FM_Solid, CM_None>::GetRHI();
-				GraphicsPSOInit.DepthStencilState = TStaticDepthStencilState<false, CF_Always>::GetRHI();
-				GraphicsPSOInit.BlendState = TStaticBlendState<>::GetRHI();
-				GraphicsPSOInit.PrimitiveType = PT_TriangleList;
-				GraphicsPSOInit.BoundShaderState.VertexDeclarationRHI = GFilterVertexDeclaration.VertexDeclarationRHI;
-				GraphicsPSOInit.BoundShaderState.VertexShaderRHI = VertexShader.GetVertexShader();
-				GraphicsPSOInit.BoundShaderState.PixelShaderRHI = PixelShader.GetPixelShader();
-
-				SetGraphicsPipelineState(RHICmdList, GraphicsPSOInit, 0);
-
-				SetShaderParameters(RHICmdList, PixelShader, PixelShader.GetPixelShader(), *PassParameters);
-
-				DrawRectangle(
-					RHICmdList,
-					0, 0,
-					View.ViewRect.Width() / GAODownsampleFactor, View.ViewRect.Height() / GAODownsampleFactor,
-					0, 0,
-					View.ViewRect.Width(), View.ViewRect.Height(),
-					FIntPoint(View.ViewRect.Width() / GAODownsampleFactor, View.ViewRect.Height() / GAODownsampleFactor),
-					View.GetSceneTexturesConfig().Extent,
-					VertexShader);
-			});
-		}
+			RHICmdList.SetViewport(0, 0, 0.0f, View.ViewRect.Width() / GAODownsampleFactor, View.ViewRect.Height() / GAODownsampleFactor, 1.0f);
+
+			TShaderMapRef<FPostProcessVS> VertexShader(View.ShaderMap);
+
+			TShaderMapRef<FComputeDistanceFieldNormalPS> PixelShader(View.ShaderMap);
+
+			FGraphicsPipelineStateInitializer GraphicsPSOInit;
+			RHICmdList.ApplyCachedRenderTargets(GraphicsPSOInit);
+			GraphicsPSOInit.RasterizerState = TStaticRasterizerState<FM_Solid, CM_None>::GetRHI();
+			GraphicsPSOInit.DepthStencilState = TStaticDepthStencilState<false, CF_Always>::GetRHI();
+			GraphicsPSOInit.BlendState = TStaticBlendState<>::GetRHI();
+			GraphicsPSOInit.PrimitiveType = PT_TriangleList;
+			GraphicsPSOInit.BoundShaderState.VertexDeclarationRHI = GFilterVertexDeclaration.VertexDeclarationRHI;
+			GraphicsPSOInit.BoundShaderState.VertexShaderRHI = VertexShader.GetVertexShader();
+			GraphicsPSOInit.BoundShaderState.PixelShaderRHI = PixelShader.GetPixelShader();
+
+			SetGraphicsPipelineState(RHICmdList, GraphicsPSOInit, 0);
+
+			SetShaderParameters(RHICmdList, PixelShader, PixelShader.GetPixelShader(), *PassParameters);
+
+			DrawRectangle(
+				RHICmdList,
+				0, 0,
+				View.ViewRect.Width() / GAODownsampleFactor, View.ViewRect.Height() / GAODownsampleFactor,
+				0, 0,
+				View.ViewRect.Width(), View.ViewRect.Height(),
+				FIntPoint(View.ViewRect.Width() / GAODownsampleFactor, View.ViewRect.Height() / GAODownsampleFactor),
+				View.GetSceneTexturesConfig().Extent,
+				VertexShader);
+		});
 	}
 }
 
@@ -830,16 +808,6 @@ void FDeferredShadingSceneRenderer::RenderDistanceFieldLighting(
 		Scene->DistanceFieldSceneData.ListMeshDistanceFields(true);
 	}
 
-	FRDGTextureRef DistanceFieldNormal = nullptr;
-
-	{
-		const FIntPoint BufferSize = GetBufferSizeForAO(Views[0]);
-		const FRDGTextureDesc Desc = FRDGTextureDesc::Create2D(BufferSize, PF_FloatRGBA, FClearValueBinding::Transparent, GFastVRamConfig.DistanceFieldNormal | TexCreate_RenderTargetable | TexCreate_UAV | TexCreate_ShaderResource);
-		DistanceFieldNormal = GraphBuilder.CreateTexture(Desc, TEXT("DistanceFieldNormal"));
-	}
-
-	ComputeDistanceFieldNormal(GraphBuilder, Views, SceneTextures.UniformBuffer, DistanceFieldNormal, Parameters);
-
 	FRDGTextureRef BentNormalOutput = nullptr;
 
 	for (int32 ViewIndex = 0; ViewIndex < Views.Num(); ++ViewIndex)
@@ -848,6 +816,16 @@ void FDeferredShadingSceneRenderer::RenderDistanceFieldLighting(
 		RDG_GPU_MASK_SCOPE(GraphBuilder, View.GPUMask);
 		RDG_EVENT_SCOPE_CONDITIONAL(GraphBuilder, Views.Num() > 1, "View%d", ViewIndex);
 
+		FRDGTextureRef DistanceFieldNormal = nullptr;
+
+		{
+			const FIntPoint BufferSize = GetBufferSizeForAO(View);
+			const FRDGTextureDesc Desc = FRDGTextureDesc::Create2D(BufferSize, PF_FloatRGBA, FClearValueBinding::Transparent, GFastVRamConfig.DistanceFieldNormal | TexCreate_RenderTargetable | TexCreate_UAV | TexCreate_ShaderResource);
+			DistanceFieldNormal = GraphBuilder.CreateTexture(Desc, TEXT("DistanceFieldNormal"));
+		}
+
+		ComputeDistanceFieldNormal(GraphBuilder, View, SceneTextures.UniformBuffer, DistanceFieldNormal, Parameters);
+
 		const FIntPoint TileListGroupSize = GetTileListGroupSizeForView(View);
 		const int32 MaxSceneObjects = FMath::DivideAndRoundUp(Scene->DistanceFieldSceneData.NumObjectsInBuffer, 256) * 256;
 		const bool bAllow16BitIndices = !IsMetalPlatform(GShaderPlatformForFeatureLevel[View.FeatureLevel]);
diff --git a/Engine/Source/Runtime/Renderer/Private/DistanceFieldScreenGridLighting.cpp b/Engine/Source/Runtime/Renderer/Private/DistanceFieldScreenGridLighting.cpp
--- a/Engine/Source/Runtime/Renderer/Private/DistanceFieldScreenGridLighting.cpp
+++ b/Engine/Source/Runtime/Renderer/Private/DistanceFieldScreenGridLighting.cpp
@@ -347,8 +347,8 @@ void FDeferredShadingSceneRenderer::RenderDistanceFieldAOScreenGrid(
 
 		auto ComputeShader = View.ShaderMap->GetShader<FConeTraceScreenGridGlobalOcclusionCS>(PermutationVector);
 
-		const uint32 GroupSizeX = FMath::DivideAndRoundUp(View.ViewRect.Size().X / GAODownsampleFactor / GConeTraceDownsampleFactor, GConeTraceGlobalDFTileSize);
-		const uint32 GroupSizeY = FMath::DivideAndRoundUp(View.ViewRect.Size().Y / GAODownsampleFactor / GConeTraceDownsampleFactor, GConeTraceGlobalDFTileSize);
+		const uint32 GroupSizeX = FMath::DivideAndRoundUp(ConeTraceBufferSize.X, GConeTraceGlobalDFTileSize);
+		const uint32 GroupSizeY = FMath::DivideAndRoundUp(ConeTraceBufferSize.Y, GConeTraceGlobalDFTileSize);
 
 		FComputeShaderUtils::AddPass(GraphBuilder, RDG_EVENT_NAME("ConeTraceGlobal"), ComputeShader, PassParameters, FIntVector(GroupSizeX, GroupSizeY, 1));
 	}
@@ -406,7 +406,7 @@ void FDeferredShadingSceneRenderer::RenderDistanceFieldAOScreenGrid(
 		PassParameters->ScreenGridParameters = SetupScreenGridParameters(View, DistanceFieldNormal);
 		PassParameters->AOSampleParameters = SetupAOSampleParameters(View.Family->FrameNumber);
 		PassParameters->RWDistanceFieldBentNormal = GraphBuilder.CreateUAV(DownsampledBentNormal);
-		PassParameters->ConeBufferMax = FIntPoint(View.ViewRect.Width() / GAODownsampleFactor / GConeTraceDownsampleFactor - 1, View.ViewRect.Height() / GAODownsampleFactor / GConeTraceDownsampleFactor - 1);
+		PassParameters->ConeBufferMax = FIntPoint(ConeTraceBufferSize.X - 1, ConeTraceBufferSize.Y - 1);
 		PassParameters->DFNormalBufferUVMax = DFNormalBufferUVMaxValue;
 
 		auto ComputeShader = View.ShaderMap->GetShader<FCombineConeVisibilityCS>();
