[Renderer]
* Moved dll export from types to methods/staticvar

diff --git a/Engine/Source/Runtime/Renderer/Private/DeferredShadingRenderer.h b/Engine/Source/Runtime/Renderer/Private/DeferredShadingRenderer.h
--- a/Engine/Source/Runtime/Renderer/Private/DeferredShadingRenderer.h
+++ b/Engine/Source/Runtime/Renderer/Private/DeferredShadingRenderer.h
@@ -237,7 +237,7 @@ private:
 /**
  * Encapsulates the resources and render targets used by global illumination plugins.
  */
-class RENDERER_API FGlobalIlluminationPluginResources : public FRenderResource
+class FGlobalIlluminationPluginResources : public FRenderResource
 
 {
 public:
@@ -252,20 +252,20 @@ public:
 /**
  * Delegate callback used by global illumination plugins
  */
-class RENDERER_API FGlobalIlluminationPluginDelegates
+class FGlobalIlluminationPluginDelegates
 {
 public:
 	DECLARE_MULTICAST_DELEGATE_OneParam(FAnyRayTracingPassEnabled, bool& /*bAnyRayTracingPassEnabled*/);
 	DECLARE_MULTICAST_DELEGATE_TwoParams(FPrepareRayTracing, const FViewInfo& /*View*/, TArray<FRHIRayTracingShader*>& /*OutRayGenShaders*/);
 	DECLARE_MULTICAST_DELEGATE_FourParams(FRenderDiffuseIndirectLight, const FScene& /*Scene*/, const FViewInfo& /*View*/, FRDGBuilder& /*GraphBuilder*/, FGlobalIlluminationPluginResources& /*Resources*/);
 
-	static FAnyRayTracingPassEnabled& AnyRayTracingPassEnabled();
-	static FPrepareRayTracing& PrepareRayTracing();
-	static FRenderDiffuseIndirectLight& RenderDiffuseIndirectLight();
+	static RENDERER_API FAnyRayTracingPassEnabled& AnyRayTracingPassEnabled();
+	static RENDERER_API FPrepareRayTracing& PrepareRayTracing();
+	static RENDERER_API FRenderDiffuseIndirectLight& RenderDiffuseIndirectLight();
 
 #if !(UE_BUILD_SHIPPING || UE_BUILD_TEST)
 	DECLARE_MULTICAST_DELEGATE_FourParams(FRenderDiffuseIndirectVisualizations, const FScene& /*Scene*/, const FViewInfo& /*View*/, FRDGBuilder& /*GraphBuilder*/, FGlobalIlluminationPluginResources& /*Resources*/);
-	static FRenderDiffuseIndirectVisualizations& RenderDiffuseIndirectVisualizations();
+	static RENDERER_API FRenderDiffuseIndirectVisualizations& RenderDiffuseIndirectVisualizations();
 #endif //!(UE_BUILD_SHIPPING || UE_BUILD_TEST)
 };
 
@@ -1210,4 +1210,4 @@ private:
 	bool bAreLightsInLightGrid;
 };
 
-DECLARE_CYCLE_STAT_EXTERN(TEXT("PrePass"), STAT_CLM_PrePass, STATGROUP_CommandListMarkers, );
\ No newline at end of file
+DECLARE_CYCLE_STAT_EXTERN(TEXT("PrePass"), STAT_CLM_PrePass, STATGROUP_CommandListMarkers, );
diff --git a/Engine/Source/Runtime/Renderer/Private/LightMapRendering.h b/Engine/Source/Runtime/Renderer/Private/LightMapRendering.h
--- a/Engine/Source/Runtime/Renderer/Private/LightMapRendering.h
+++ b/Engine/Source/Runtime/Renderer/Private/LightMapRendering.h
@@ -340,7 +340,7 @@ enum ELightMapPolicyType
 	LMP_DUMMY
 };
 
-class RENDERER_API FUniformLightMapPolicyShaderParametersType
+class FUniformLightMapPolicyShaderParametersType
 {
 	DECLARE_TYPE_LAYOUT(FUniformLightMapPolicyShaderParametersType, NonVirtual);
 public:
@@ -363,7 +363,7 @@ public:
 	LAYOUT_FIELD(FShaderUniformBufferParameter, LightmapResourceCluster);
 };
 
-class RENDERER_API FUniformLightMapPolicy
+class FUniformLightMapPolicy
 {
 public:
 
@@ -378,23 +378,23 @@ public:
 
 	FUniformLightMapPolicy(ELightMapPolicyType InIndirectPolicy) : IndirectPolicy(InIndirectPolicy) {}
 
-	static bool ShouldCompilePermutation(ELightMapPolicyType Policy, const FMeshMaterialShaderPermutationParameters& Parameters);
-	static void ModifyCompilationEnvironment(ELightMapPolicyType Policy, const FMaterialShaderPermutationParameters& Parameters, FShaderCompilerEnvironment& OutEnvironment);
+	static RENDERER_API bool ShouldCompilePermutation(ELightMapPolicyType Policy, const FMeshMaterialShaderPermutationParameters& Parameters);
+	static RENDERER_API void ModifyCompilationEnvironment(ELightMapPolicyType Policy, const FMaterialShaderPermutationParameters& Parameters, FShaderCompilerEnvironment& OutEnvironment);
 
-	static void GetVertexShaderBindings(
+	static RENDERER_API void GetVertexShaderBindings(
 		const FPrimitiveSceneProxy* PrimitiveSceneProxy,
 		const ElementDataType& ShaderElementData,
 		const VertexParametersType* VertexShaderParameters,
 		FMeshDrawSingleShaderBindings& ShaderBindings);
 
-	static void GetPixelShaderBindings(
+	static RENDERER_API void GetPixelShaderBindings(
 		const FPrimitiveSceneProxy* PrimitiveSceneProxy,
 		const ElementDataType& ShaderElementData,
 		const PixelParametersType* PixelShaderParameters,
 		FMeshDrawSingleShaderBindings& ShaderBindings);
 
 #if RHI_RAYTRACING
-	void GetRayHitGroupShaderBindings(
+	RENDERER_API void GetRayHitGroupShaderBindings(
 		const FPrimitiveSceneProxy* PrimitiveSceneProxy,
 		const FLightCacheInterface* ElementData,
 		const RayHitGroupParametersType* RayHitGroupShaderParameters,
@@ -402,7 +402,7 @@ public:
 	) const;
 #endif // RHI_RAYTRACING
 
-	static void GetComputeShaderBindings(
+	static RENDERER_API void GetComputeShaderBindings(
 		const FPrimitiveSceneProxy* PrimitiveSceneProxy,
 		const ElementDataType& ShaderElementData,
 		const ComputeParametersType* PixelShaderParameters,
@@ -571,4 +571,4 @@ public:
 		const ElementDataType& ShaderElementData,
 		const ComputeParametersType* ComputeShaderParameters,
 		FMeshDrawSingleShaderBindings& ShaderBindings);
-};
\ No newline at end of file
+};
diff --git a/Engine/Source/Runtime/Renderer/Private/PostProcess/PostProcessUpscale.h b/Engine/Source/Runtime/Renderer/Private/PostProcess/PostProcessUpscale.h
--- a/Engine/Source/Runtime/Renderer/Private/PostProcess/PostProcessUpscale.h
+++ b/Engine/Source/Runtime/Renderer/Private/PostProcess/PostProcessUpscale.h
@@ -61,7 +61,7 @@ enum class EUpscaleStage
 };
 
 /** Interface for custom spatial upscaling algorithm meant to be set on the FSceneViewFamily by ISceneViewExtension::BeginRenderViewFamily(). */
-class RENDERER_API ISpatialUpscaler : public ISceneViewFamilyExtention
+class ISpatialUpscaler : public ISceneViewFamilyExtention
 {
 public:
 	struct FInputs
@@ -86,7 +86,7 @@ public:
 		const FViewInfo& View,
 		const FInputs& PassInputs) const = 0;
 
-	static FScreenPassTexture AddDefaultUpscalePass(
+	static RENDERER_API FScreenPassTexture AddDefaultUpscalePass(
 		FRDGBuilder& GraphBuilder,
 		const FViewInfo& View,
 		const FInputs& PassInputs,
diff --git a/Engine/Source/Runtime/Renderer/Private/RayTracing/RayTracingMaterialHitShaders.h b/Engine/Source/Runtime/Renderer/Private/RayTracing/RayTracingMaterialHitShaders.h
--- a/Engine/Source/Runtime/Renderer/Private/RayTracing/RayTracingMaterialHitShaders.h
+++ b/Engine/Source/Runtime/Renderer/Private/RayTracing/RayTracingMaterialHitShaders.h
@@ -21,13 +21,13 @@ FRHIRayTracingShader* GetRayTracingDefaultMissShader(const FGlobalShaderMap* Sha
 FRHIRayTracingShader* GetRayTracingDefaultOpaqueShader(const FGlobalShaderMap* ShaderMap);
 FRHIRayTracingShader* GetRayTracingDefaultHiddenShader(const FGlobalShaderMap* ShaderMap);
 
-class RENDERER_API FRayTracingMeshProcessor
+class FRayTracingMeshProcessor
 {
 public:
-	FRayTracingMeshProcessor(FRayTracingMeshCommandContext* InCommandContext, const FScene* InScene, const FSceneView* InViewIfDynamicMeshCommand, FMeshPassProcessorRenderState InPassDrawRenderState, ERayTracingMeshCommandsMode InRayTracingMeshCommandsMode);
-	virtual ~FRayTracingMeshProcessor();
+	RENDERER_API FRayTracingMeshProcessor(FRayTracingMeshCommandContext* InCommandContext, const FScene* InScene, const FSceneView* InViewIfDynamicMeshCommand, FMeshPassProcessorRenderState InPassDrawRenderState, ERayTracingMeshCommandsMode InRayTracingMeshCommandsMode);
+	RENDERER_API virtual ~FRayTracingMeshProcessor();
 
-	void AddMeshBatch(const FMeshBatch& RESTRICT MeshBatch, uint64 BatchElementMask, const FPrimitiveSceneProxy* RESTRICT PrimitiveSceneProxy);
+	RENDERER_API void AddMeshBatch(const FMeshBatch& RESTRICT MeshBatch, uint64 BatchElementMask, const FPrimitiveSceneProxy* RESTRICT PrimitiveSceneProxy);
 
 protected:
 	FRayTracingMeshCommandContext* CommandContext;
@@ -37,7 +37,7 @@ protected:
 	FMeshPassProcessorRenderState PassDrawRenderState;
 	ERayTracingMeshCommandsMode RayTracingMeshCommandsMode;
 
-	bool Process(
+	RENDERER_API bool Process(
 		const FMeshBatch& RESTRICT MeshBatch,
 		uint64 BatchElementMask,
 		const FPrimitiveSceneProxy* RESTRICT PrimitiveSceneProxy,
@@ -104,14 +104,14 @@ protected:
 	}
 
 private:
-	bool ProcessPathTracing(
+	RENDERER_API bool ProcessPathTracing(
 		const FMeshBatch& RESTRICT MeshBatch,
 		uint64 BatchElementMask,
 		const FPrimitiveSceneProxy* RESTRICT PrimitiveSceneProxy,
 		const FMaterialRenderProxy& RESTRICT MaterialRenderProxy,
 		const FMaterial& RESTRICT MaterialResource);
 
-	bool TryAddMeshBatch(
+	RENDERER_API bool TryAddMeshBatch(
 		const FMeshBatch& RESTRICT MeshBatch,
 		uint64 BatchElementMask,
 		const FPrimitiveSceneProxy* RESTRICT PrimitiveSceneProxy,
diff --git a/Engine/Source/Runtime/Renderer/Private/SceneRendering.h b/Engine/Source/Runtime/Renderer/Private/SceneRendering.h
--- a/Engine/Source/Runtime/Renderer/Private/SceneRendering.h
+++ b/Engine/Source/Runtime/Renderer/Private/SceneRendering.h
@@ -2132,7 +2132,7 @@ public:
 
 	// FSceneRenderer interface
 
-	RENDERER_API virtual void Render(FRDGBuilder& GraphBuilder) = 0;
+	virtual void Render(FRDGBuilder& GraphBuilder) = 0;
 	virtual void RenderHitProxies(FRDGBuilder& GraphBuilder) {}
 	virtual bool ShouldRenderVelocities() const { return false; }
 	virtual bool ShouldRenderPrePass() const { return false; }
diff --git a/Engine/Source/Runtime/Renderer/Private/SceneTextures.h b/Engine/Source/Runtime/Renderer/Private/SceneTextures.h
--- a/Engine/Source/Runtime/Renderer/Private/SceneTextures.h
+++ b/Engine/Source/Runtime/Renderer/Private/SceneTextures.h
@@ -16,10 +16,10 @@ class FViewFamilyInfo;
 extern RENDERER_API void InitializeSceneTexturesConfig(FSceneTexturesConfig& Config, const FSceneViewFamily& ViewFamily);
 
 /** RDG struct containing the minimal set of scene textures common across all rendering configurations. */
-struct RENDERER_API FMinimalSceneTextures
+struct FMinimalSceneTextures
 {
 	// Initializes the minimal scene textures structure in the FViewFamilyInfo
-	static void InitializeViewFamily(FRDGBuilder& GraphBuilder, FViewFamilyInfo& ViewFamily);
+	static RENDERER_API void InitializeViewFamily(FRDGBuilder& GraphBuilder, FViewFamilyInfo& ViewFamily);
 
 	// Immutable copy of the config used to create scene textures.
 	FSceneTexturesConfig Config;
@@ -47,20 +47,20 @@ struct RENDERER_API FMinimalSceneTextures
 	// Textures containing depth / stencil information from the custom depth pass.
 	FCustomDepthTextures CustomDepth{};
 
-	FSceneTextureShaderParameters GetSceneTextureShaderParameters(ERHIFeatureLevel::Type FeatureLevel) const;
+	RENDERER_API FSceneTextureShaderParameters GetSceneTextureShaderParameters(ERHIFeatureLevel::Type FeatureLevel) const;
 };
 
 /** RDG struct containing the complete set of scene textures for the deferred or mobile renderers. */
-struct RENDERER_API FSceneTextures : public FMinimalSceneTextures
+struct FSceneTextures : public FMinimalSceneTextures
 {
 	// Initializes the scene textures structure in the FViewFamilyInfo
-	static void InitializeViewFamily(FRDGBuilder& GraphBuilder, FViewFamilyInfo& ViewFamily);
+	static RENDERER_API void InitializeViewFamily(FRDGBuilder& GraphBuilder, FViewFamilyInfo& ViewFamily);
 
 	// Configures an array of render targets for the GBuffer pass.
-	uint32 GetGBufferRenderTargets(
+	RENDERER_API uint32 GetGBufferRenderTargets(
 		TArrayView<FTextureRenderTargetBinding> RenderTargets,
 		EGBufferLayout Layout = GBL_Default) const;
-	uint32 GetGBufferRenderTargets(
+	RENDERER_API uint32 GetGBufferRenderTargets(
 		ERenderTargetLoadAction LoadAction,
 		FRenderTargetBindingSlots& RenderTargets,
 		EGBufferLayout Layout = GBL_Default) const;
diff --git a/Engine/Source/Runtime/Renderer/Private/SceneUniformBuffer.h b/Engine/Source/Runtime/Renderer/Private/SceneUniformBuffer.h
--- a/Engine/Source/Runtime/Renderer/Private/SceneUniformBuffer.h
+++ b/Engine/Source/Runtime/Renderer/Private/SceneUniformBuffer.h
@@ -94,10 +94,10 @@ DECLARE_UNIFORM_BUFFER_STRUCT(FSceneUniformParameters, RENDERER_API)
 /**
  * RDG shader parameter struct containing data for FSceneUniformBuffer.
  */
-MS_ALIGN(SHADER_PARAMETER_STRUCT_ALIGNMENT) class RENDERER_API FSceneUniformParameters final
+MS_ALIGN(SHADER_PARAMETER_STRUCT_ALIGNMENT) class FSceneUniformParameters final
 {
 public:
-	struct RENDERER_API FTypeInfo {
+	struct FTypeInfo {
 		static constexpr int32 NumRows = 1;
 		static constexpr int32 NumColumns = 1;
 		static constexpr int32 NumElements = 0;
@@ -106,7 +106,7 @@ public:
 		static constexpr const ANSICHAR* const FileName = UE_LOG_SOURCE_FILE(__FILE__);
 		static constexpr int32 FileLine = __LINE__;
 		using TAlignedType = FSceneUniformParameters;
-		static const FShaderParametersMetadata* GetStructMetadata();
+		static RENDERER_API const FShaderParametersMetadata* GetStructMetadata();
 	};
 
 	const void* GetContents() const
diff --git a/Engine/Source/Runtime/Renderer/Private/ScreenSpaceDenoise.h b/Engine/Source/Runtime/Renderer/Private/ScreenSpaceDenoise.h
--- a/Engine/Source/Runtime/Renderer/Private/ScreenSpaceDenoise.h
+++ b/Engine/Source/Runtime/Renderer/Private/ScreenSpaceDenoise.h
@@ -46,7 +46,7 @@ END_SHADER_PARAMETER_STRUCT()
 
 
 /** Interface for denoiser to have all hook in the renderer. */
-class RENDERER_API IScreenSpaceDenoiser
+class IScreenSpaceDenoiser
 {
 public:
 	/** Maximum number a denoiser might be able to denoise at the same time. */
@@ -216,9 +216,9 @@ public:
 	};
 	
 
-	static FHarmonicTextures CreateHarmonicTextures(FRDGBuilder& GraphBuilder, FIntPoint Extent, const TCHAR* DebugName);
-	static FHarmonicUAVs CreateUAVs(FRDGBuilder& GraphBuilder, const FHarmonicTextures& Textures);
-	static FDiffuseIndirectHarmonicUAVs CreateUAVs(FRDGBuilder& GraphBuilder, const FDiffuseIndirectHarmonic& Textures);
+	static RENDERER_API FHarmonicTextures CreateHarmonicTextures(FRDGBuilder& GraphBuilder, FIntPoint Extent, const TCHAR* DebugName);
+	static RENDERER_API FHarmonicUAVs CreateUAVs(FRDGBuilder& GraphBuilder, const FHarmonicTextures& Textures);
+	static RENDERER_API FDiffuseIndirectHarmonicUAVs CreateUAVs(FRDGBuilder& GraphBuilder, const FDiffuseIndirectHarmonic& Textures);
 
 
 
@@ -326,7 +326,7 @@ public:
 		const FAmbientOcclusionRayTracingConfig Config) const = 0;
 
 	/** Entry point to denoise diffuse indirect probe hierarchy. */
-	static FSSDSignalTextures DenoiseIndirectProbeHierarchy(
+	static RENDERER_API FSSDSignalTextures DenoiseIndirectProbeHierarchy(
 		FRDGBuilder& GraphBuilder,
 		const FViewInfo& View,
 		FPreviousViewInfo* PreviousViewInfos,
@@ -336,14 +336,14 @@ public:
 		FRDGTextureRef CompressedShadingModelTexture);
 
 	/** Returns the interface of the default denoiser of the renderer. */
-	static const IScreenSpaceDenoiser* GetDefaultDenoiser();
+	static RENDERER_API const IScreenSpaceDenoiser* GetDefaultDenoiser();
 
 	/** Returns the denoising mode. */
-	static EMode GetDenoiserMode(const TAutoConsoleVariable<int32>& CVar);
+	static RENDERER_API EMode GetDenoiserMode(const TAutoConsoleVariable<int32>& CVar);
 }; // class IScreenSpaceDenoiser
 
 
 // The interface for the renderer to denoise what it needs, Plugins can come and point this to custom interface.
 extern RENDERER_API const IScreenSpaceDenoiser* GScreenSpaceDenoiser;
 
-extern int GetReflectionsDenoiserMode();
\ No newline at end of file
+extern int GetReflectionsDenoiserMode();
diff --git a/Engine/Source/Runtime/Renderer/Private/VT/RuntimeVirtualTextureProducer.h b/Engine/Source/Runtime/Renderer/Private/VT/RuntimeVirtualTextureProducer.h
--- a/Engine/Source/Runtime/Renderer/Private/VT/RuntimeVirtualTextureProducer.h
+++ b/Engine/Source/Runtime/Renderer/Private/VT/RuntimeVirtualTextureProducer.h
@@ -81,7 +81,7 @@ public:
 		FTransform const& InUVToWorld,
 		FBox const& InWorldBounds);
 	
-	RENDERER_API virtual ~FRuntimeVirtualTextureProducer() {}
+	virtual ~FRuntimeVirtualTextureProducer() {}
 
 	//~ Begin IVirtualTexture Interface.
 	virtual FVTRequestPageResult RequestPageData(
diff --git a/Engine/Source/Runtime/Renderer/Private/VT/VirtualTextureAllocator.h b/Engine/Source/Runtime/Renderer/Private/VT/VirtualTextureAllocator.h
--- a/Engine/Source/Runtime/Renderer/Private/VT/VirtualTextureAllocator.h
+++ b/Engine/Source/Runtime/Renderer/Private/VT/VirtualTextureAllocator.h
@@ -9,16 +9,16 @@
 class FAllocatedVirtualTexture;
 
 /** Allocates virtual memory address space. */
-class RENDERER_API FVirtualTextureAllocator
+class FVirtualTextureAllocator
 {
 public:
-	explicit FVirtualTextureAllocator(uint32 Dimensions);
+	RENDERER_API explicit FVirtualTextureAllocator(uint32 Dimensions);
 	~FVirtualTextureAllocator() {}
 
 	/**
 	 * Initialise the allocator
 	 */
-	void Initialize(uint32 MaxSize);
+	RENDERER_API void Initialize(uint32 MaxSize);
 
 	uint32 GetAllocatedWidth() const { return AllocatedWidth; }
 	uint32 GetAllocatedHeight() const { return AllocatedHeight; }
@@ -27,25 +27,25 @@ public:
 	 * Translate a virtual page address in the address space to a local page address within a virtual texture.
 	 * @return nullptr If there is no virtual texture allocated at this address.
 	 */
-	FAllocatedVirtualTexture* Find(uint32 vAddress, uint32& OutLocal_vAddress) const;
+	RENDERER_API FAllocatedVirtualTexture* Find(uint32 vAddress, uint32& OutLocal_vAddress) const;
 	inline FAllocatedVirtualTexture* Find(uint32 vAddress) const { uint32 UnusedLocal_vAddress = 0u; return Find(vAddress, UnusedLocal_vAddress); }
 
 	/**
 	 * Allocate address space for the virtual texture.
 	 * @return (~0) if no space left, the virtual page address if successfully allocated.
 	 */
-	uint32 Alloc(FAllocatedVirtualTexture* VT);
+	RENDERER_API uint32 Alloc(FAllocatedVirtualTexture* VT);
 
 	/**
 	 * Test if an allocation of the given size will succeed.
 	 * @return false if there isn't enough space left.
 	 */
-	bool TryAlloc(uint32 InSize);
+	RENDERER_API bool TryAlloc(uint32 InSize);
 
 	/**
 	 * Free the virtual texture.
 	 */
-	void Free(FAllocatedVirtualTexture* VT);
+	RENDERER_API void Free(FAllocatedVirtualTexture* VT);
 
 	/** Get current number of allocations. */
 	inline uint32 GetNumAllocations() const { return NumAllocations; }
@@ -54,10 +54,10 @@ public:
 	inline uint32 GetNumAllocatedPages() const { return NumAllocatedPages; }
 
 	/** Output debugging information to the console. */
-	void DumpToConsole(bool verbose);
+	RENDERER_API void DumpToConsole(bool verbose);
 
 #if WITH_EDITOR
-	void SaveDebugImage(const TCHAR* ImageName) const;
+	RENDERER_API void SaveDebugImage(const TCHAR* ImageName) const;
 #endif
 
 private:
@@ -79,26 +79,26 @@ private:
 		uint32 vTileY1;
 	};
 
-	void LinkFreeList(uint16& InOutListHead, EBlockState State, uint16 Index);
-	void UnlinkFreeList(uint16& InOutListHead, EBlockState State, uint16 Index);
+	RENDERER_API void LinkFreeList(uint16& InOutListHead, EBlockState State, uint16 Index);
+	RENDERER_API void UnlinkFreeList(uint16& InOutListHead, EBlockState State, uint16 Index);
 
-	int32 AcquireBlock();
-	void FreeAddressBlock(uint32 Index, bool bTopLevelBlock);
-	uint32 FindAddressBlock(uint32 vAddress) const;
+	RENDERER_API int32 AcquireBlock();
+	RENDERER_API void FreeAddressBlock(uint32 Index, bool bTopLevelBlock);
+	RENDERER_API uint32 FindAddressBlock(uint32 vAddress) const;
 
-	void SubdivideBlock(uint32 ParentIndex);
+	RENDERER_API void SubdivideBlock(uint32 ParentIndex);
 
-	void MarkBlockAllocated(uint32 Index, uint32 vAllocatedTileX, uint32 vAllocatedTileY, FAllocatedVirtualTexture* VT);
+	RENDERER_API void MarkBlockAllocated(uint32 Index, uint32 vAllocatedTileX, uint32 vAllocatedTileY, FAllocatedVirtualTexture* VT);
 
-	bool TestAllocation(uint32 Index, uint32 vTileX0, uint32 vTileY0, uint32 vTileX1, uint32 vTileY1) const;
+	RENDERER_API bool TestAllocation(uint32 Index, uint32 vTileX0, uint32 vTileY0, uint32 vTileX1, uint32 vTileY1) const;
 
-	void RecurseComputeFreeMip(uint16 BlockIndex, uint32 Depth, uint64& IoBlockMap) const;
-	uint32 ComputeFreeMip(uint16 BlockIndex) const;
+	RENDERER_API void RecurseComputeFreeMip(uint16 BlockIndex, uint32 Depth, uint64& IoBlockMap) const;
+	RENDERER_API uint32 ComputeFreeMip(uint16 BlockIndex) const;
 
-	void FreeMipUpdateParents(uint16 ParentIndex);
+	RENDERER_API void FreeMipUpdateParents(uint16 ParentIndex);
 
 #if WITH_EDITOR
-	void FillDebugImage(uint32 Index, uint32* ImageData, TMap<FAllocatedVirtualTexture*, uint32>& ColorMap) const;
+	RENDERER_API void FillDebugImage(uint32 Index, uint32* ImageData, TMap<FAllocatedVirtualTexture*, uint32>& ColorMap) const;
 #endif
 
 	struct FAddressBlock
diff --git a/Engine/Source/Runtime/Renderer/Public/ComputeSystemInterface.h b/Engine/Source/Runtime/Renderer/Public/ComputeSystemInterface.h
--- a/Engine/Source/Runtime/Renderer/Public/ComputeSystemInterface.h
+++ b/Engine/Source/Runtime/Renderer/Public/ComputeSystemInterface.h
@@ -15,7 +15,7 @@ class IComputeTaskWorker;
  * IComputeTaskWorker objects have their work executed by the renderer at specific points in the frame.
  * The compute tasks are scheduled by the IComputeSystem.
  */
-class RENDERER_API IComputeSystem
+class IComputeSystem
 {
 public:
 	virtual ~IComputeSystem() {}
diff --git a/Engine/Source/Runtime/Renderer/Public/FXRenderingUtils.h b/Engine/Source/Runtime/Renderer/Public/FXRenderingUtils.h
--- a/Engine/Source/Runtime/Renderer/Public/FXRenderingUtils.h
+++ b/Engine/Source/Runtime/Renderer/Public/FXRenderingUtils.h
@@ -84,7 +84,7 @@ namespace UE::FXRenderingUtils
 /**
  * This class exposes methods required by FX rendering that must access rendering internals.
  */ 
-class RENDERER_API FFXRenderingUtils
+class FFXRenderingUtils
 {
 public:
 	FFXRenderingUtils() = delete;
diff --git a/Engine/Source/Runtime/Renderer/Public/HairStrandsInterface.h b/Engine/Source/Runtime/Renderer/Public/HairStrandsInterface.h
--- a/Engine/Source/Runtime/Renderer/Public/HairStrandsInterface.h
+++ b/Engine/Source/Runtime/Renderer/Public/HairStrandsInterface.h
@@ -140,16 +140,16 @@ enum class ERDGImportedBufferFlags
 };
 ENUM_CLASS_FLAGS(ERDGImportedBufferFlags);
 
-struct RENDERER_API FRDGExternalBuffer
+struct FRDGExternalBuffer
 {
 	TRefCountPtr<FRDGPooledBuffer> Buffer = nullptr;
 	FShaderResourceViewRHIRef SRV = nullptr;
 	FUnorderedAccessViewRHIRef UAV = nullptr;
 	EPixelFormat Format = PF_Unknown;
-	void Release();
+	RENDERER_API void Release();
 };
 
-struct RENDERER_API FRDGImportedBuffer
+struct FRDGImportedBuffer
 {
 	FRDGBufferRef Buffer = nullptr;
 	FRDGBufferSRVRef SRV = nullptr;
@@ -198,12 +198,12 @@ enum EHairInterpolationType
 ////////////////////////////////////////////////////////////////////////////////////////////////////
 // Public group data 
 
-struct RENDERER_API FHairStrandsInstance
+struct FHairStrandsInstance
 {
 	virtual ~FHairStrandsInstance() = default;
-	uint32 GetRefCount() const;
-	uint32 AddRef() const;
-	uint32 Release() const;
+	RENDERER_API uint32 GetRefCount() const;
+	RENDERER_API uint32 AddRef() const;
+	RENDERER_API uint32 Release() const;
 	int32 RegisteredIndex = -1;
 	virtual const FBoxSphereBounds& GetBounds() const = 0;
 	virtual const FBoxSphereBounds& GetLocalBounds() const = 0;
@@ -214,17 +214,17 @@ protected:
 };
 typedef TArray<FHairStrandsInstance*> FHairStrandsInstances;
 
-class RENDERER_API FHairGroupPublicData : public FRenderResource
+class FHairGroupPublicData : public FRenderResource
 {
 public:
-	FHairGroupPublicData(uint32 InGroupIndex, const FName& OwnerName);
+	RENDERER_API FHairGroupPublicData(uint32 InGroupIndex, const FName& OwnerName);
 	
-	virtual void InitRHI() override;
-	virtual void ReleaseRHI() override;
+	RENDERER_API virtual void InitRHI() override;
+	RENDERER_API virtual void ReleaseRHI() override;
 	virtual FString GetFriendlyName() const override { return TEXT("FHairGroupPublicData"); }
-	void Allocate(FRDGBuilder& GraphBuilder);
-	void Release();
-	uint32 GetResourcesSize() const;
+	RENDERER_API void Allocate(FRDGBuilder& GraphBuilder);
+	RENDERER_API void Release();
+	RENDERER_API uint32 GetResourcesSize() const;
 
 	uint32 GetGroupIndex() const { return GroupIndex; }
 
@@ -304,9 +304,9 @@ public:
 	uint32 GetClusterCount() const { return ClusterCount;  }
 
 	// Return the number of active point/curve for strand geometry
-	uint32 GetActiveStrandsPointCount() const;
-	uint32 GetActiveStrandsCurveCount() const;
-	float  GetActiveStrandsCoverageScale() const;
+	RENDERER_API uint32 GetActiveStrandsPointCount() const;
+	RENDERER_API uint32 GetActiveStrandsCurveCount() const;
+	RENDERER_API float  GetActiveStrandsCoverageScale() const;
 
 	struct FVertexFactoryInput 
 	{
diff --git a/Engine/Source/Runtime/Renderer/Public/MaterialShader.h b/Engine/Source/Runtime/Renderer/Public/MaterialShader.h
--- a/Engine/Source/Runtime/Renderer/Public/MaterialShader.h
+++ b/Engine/Source/Runtime/Renderer/Public/MaterialShader.h
@@ -66,22 +66,22 @@ struct FMaterialShaderPermutationParameters : public FShaderPermutationParameter
 };
 
 /** Base class of all shaders that need material parameters. */
-class RENDERER_API FMaterialShader : public FShader
+class FMaterialShader : public FShader
 {
-	DECLARE_TYPE_LAYOUT(FMaterialShader, NonVirtual);
+	DECLARE_EXPORTED_TYPE_LAYOUT(FMaterialShader, RENDERER_API, NonVirtual);
 public:
 	using FPermutationParameters = FMaterialShaderPermutationParameters;
 	using ShaderMetaType = FMaterialShaderType;
 
-	static FName UniformBufferLayoutName;
+	static RENDERER_API FName UniformBufferLayoutName;
 
-	FMaterialShader();
+	RENDERER_API FMaterialShader();
 
-	FMaterialShader(const FMaterialShaderType::CompiledShaderInitializerType& Initializer);
+	RENDERER_API FMaterialShader(const FMaterialShaderType::CompiledShaderInitializerType& Initializer);
 
-	FRHIUniformBuffer* GetParameterCollectionBuffer(const FGuid& Id, const FSceneInterface* SceneInterface) const;
+	RENDERER_API FRHIUniformBuffer* GetParameterCollectionBuffer(const FGuid& Id, const FSceneInterface* SceneInterface) const;
 
-	void SetViewParameters(FRHIBatchedShaderParameters& BatchedParameters, const FSceneView& View, const TUniformBufferRef<FViewUniformShaderParameters>& ViewUniformBuffer);
+	RENDERER_API void SetViewParameters(FRHIBatchedShaderParameters& BatchedParameters, const FSceneView& View, const TUniformBufferRef<FViewUniformShaderParameters>& ViewUniformBuffer);
 
 	template<typename ShaderRHIParamRef, typename TRHICommandList>
 	FORCEINLINE_DEBUGGABLE void SetViewParameters(TRHICommandList& RHICmdList, const ShaderRHIParamRef ShaderRHI, const FSceneView& View, const TUniformBufferRef<FViewUniformShaderParameters>& ViewUniformBuffer)
@@ -92,7 +92,7 @@ public:
 	}
 
 	/** Sets pixel parameters that are material specific but not FMeshBatch specific. */
-	void SetParameters(
+	RENDERER_API void SetParameters(
 		FRHIBatchedShaderParameters& BatchedParameters,
 		const FMaterialRenderProxy* MaterialRenderProxy, 
 		const FMaterial& Material,
@@ -112,7 +112,7 @@ public:
 		RHICmdList.SetBatchedShaderParameters(ShaderRHI, BatchedParameters);
 	}
 
-	void GetShaderBindings(
+	RENDERER_API void GetShaderBindings(
 		const FScene* Scene,
 		ERHIFeatureLevel::Type FeatureLevel,
 		const FMaterialRenderProxy& MaterialRenderProxy,
@@ -121,12 +121,12 @@ public:
 
 private:
 	/** If true, cached uniform expressions are allowed. */
-	static int32 bAllowCachedUniformExpressions;
+	static RENDERER_API int32 bAllowCachedUniformExpressions;
 	/** Console variable ref to toggle cached uniform expressions. */
-	static FAutoConsoleVariableRef CVarAllowCachedUniformExpressions;
+	static RENDERER_API FAutoConsoleVariableRef CVarAllowCachedUniformExpressions;
 
 #if !(UE_BUILD_TEST || UE_BUILD_SHIPPING || !WITH_EDITOR)
-	void VerifyExpressionAndShaderMaps(const FMaterialRenderProxy* MaterialRenderProxy, const FMaterial& Material, const FUniformExpressionCache* UniformExpressionCache) const;
+	RENDERER_API void VerifyExpressionAndShaderMaps(const FMaterialRenderProxy* MaterialRenderProxy, const FMaterial& Material, const FUniformExpressionCache* UniformExpressionCache) const;
 #endif
 
 	LAYOUT_FIELD(TMemoryImageArray<FShaderUniformBufferParameter>, ParameterCollectionUniformBuffers);
diff --git a/Engine/Source/Runtime/Renderer/Public/MeshMaterialShader.h b/Engine/Source/Runtime/Renderer/Public/MeshMaterialShader.h
--- a/Engine/Source/Runtime/Renderer/Public/MeshMaterialShader.h
+++ b/Engine/Source/Runtime/Renderer/Public/MeshMaterialShader.h
@@ -68,18 +68,18 @@ struct FVertexFactoryShaderPermutationParameters
 };
 
 /** Base class of all shaders that need material and vertex factory parameters. */
-class RENDERER_API FMeshMaterialShader : public FMaterialShader
+class FMeshMaterialShader : public FMaterialShader
 {
-	DECLARE_TYPE_LAYOUT(FMeshMaterialShader, NonVirtual);
+	DECLARE_EXPORTED_TYPE_LAYOUT(FMeshMaterialShader, RENDERER_API, NonVirtual);
 public:
 	using FPermutationParameters = FMeshMaterialShaderPermutationParameters;
 	using ShaderMetaType = FMeshMaterialShaderType;
 
 	FMeshMaterialShader() {}
 
-	FMeshMaterialShader(const FMeshMaterialShaderType::CompiledShaderInitializerType& Initializer);
+	RENDERER_API FMeshMaterialShader(const FMeshMaterialShaderType::CompiledShaderInitializerType& Initializer);
 
-	void GetShaderBindings(
+	RENDERER_API void GetShaderBindings(
 		const FScene* Scene,
 		ERHIFeatureLevel::Type FeatureLevel,
 		const FPrimitiveSceneProxy* PrimitiveSceneProxy,
@@ -89,7 +89,7 @@ public:
 		const FMeshMaterialShaderElementData& ShaderElementData,
 		FMeshDrawSingleShaderBindings& ShaderBindings) const;
 
-	void GetElementShaderBindings(
+	RENDERER_API void GetElementShaderBindings(
 		const FShaderMapPointerTable& PointerTable,
 		const FScene* Scene, 
 		const FSceneView* ViewIfDynamicMeshCommand, 
@@ -121,7 +121,7 @@ public:
 	}
 
 private:
-	void WriteFrozenVertexFactoryParameters(FMemoryImageWriter& Writer, const TMemoryImagePtr<FVertexFactoryShaderParameters>& InVertexFactoryParameters) const;
+	RENDERER_API void WriteFrozenVertexFactoryParameters(FMemoryImageWriter& Writer, const TMemoryImagePtr<FVertexFactoryShaderParameters>& InVertexFactoryParameters) const;
 	LAYOUT_FIELD_WITH_WRITER(TMemoryImagePtr<FVertexFactoryShaderParameters>, VertexFactoryParameters, WriteFrozenVertexFactoryParameters);
 
 protected:
diff --git a/Engine/Source/Runtime/Renderer/Public/MeshPassProcessor.h b/Engine/Source/Runtime/Renderer/Public/MeshPassProcessor.h
--- a/Engine/Source/Runtime/Renderer/Public/MeshPassProcessor.h
+++ b/Engine/Source/Runtime/Renderer/Public/MeshPassProcessor.h
@@ -1284,7 +1284,7 @@ public:
 };
 
 /** FVisibleMeshDrawCommand sort key. */
-class RENDERER_API FMeshDrawCommandSortKey
+class FMeshDrawCommandSortKey
 {
 public:
 	union 
@@ -1322,7 +1322,7 @@ public:
 		return PackedData < B.PackedData;
 	}
 
-	static const FMeshDrawCommandSortKey Default;
+	static RENDERER_API const FMeshDrawCommandSortKey Default;
 };
 
 
diff --git a/Engine/Source/Runtime/Renderer/Public/RayTracingDynamicGeometryCollection.h b/Engine/Source/Runtime/Renderer/Public/RayTracingDynamicGeometryCollection.h
--- a/Engine/Source/Runtime/Renderer/Public/RayTracingDynamicGeometryCollection.h
+++ b/Engine/Source/Runtime/Renderer/Public/RayTracingDynamicGeometryCollection.h
@@ -13,13 +13,13 @@ class FSceneView;
 struct FMeshComputeDispatchCommand;
 struct FRayTracingDynamicGeometryUpdateParams;
 
-class RENDERER_API FRayTracingDynamicGeometryCollection
+class FRayTracingDynamicGeometryCollection
 {
 public:
-	FRayTracingDynamicGeometryCollection();
-	~FRayTracingDynamicGeometryCollection();
+	RENDERER_API FRayTracingDynamicGeometryCollection();
+	RENDERER_API ~FRayTracingDynamicGeometryCollection();
 
-	void AddDynamicMeshBatchForGeometryUpdate(
+	RENDERER_API void AddDynamicMeshBatchForGeometryUpdate(
 		const FScene* Scene, 
 		const FSceneView* View, 
 		const FPrimitiveSceneProxy* PrimitiveSceneProxy, 
@@ -28,15 +28,15 @@ public:
 	);
 
 	// Starts an update batch and returns the current shared buffer generation ID which is used for validation.
-	int64 BeginUpdate();
-	void DispatchUpdates(FRHICommandListImmediate& ParentCmdList, FRHIBuffer* ScratchBuffer);
-	void EndUpdate(FRHICommandListImmediate& RHICmdList);
+	RENDERER_API int64 BeginUpdate();
+	RENDERER_API void DispatchUpdates(FRHICommandListImmediate& ParentCmdList, FRHIBuffer* ScratchBuffer);
+	RENDERER_API void EndUpdate(FRHICommandListImmediate& RHICmdList);
 
 	// Clears the working arrays to not hold any references.
 	// Needs to be called every frame when ray tracing is enabled or once when ray tracing mode has changed.
-	void Clear();
+	RENDERER_API void Clear();
 
-	uint32 ComputeScratchBufferSize();
+	RENDERER_API uint32 ComputeScratchBufferSize();
 
 private:
 
diff --git a/Engine/Source/Runtime/Renderer/Public/RectLightTexture.h b/Engine/Source/Runtime/Renderer/Public/RectLightTexture.h
--- a/Engine/Source/Runtime/Renderer/Public/RectLightTexture.h
+++ b/Engine/Source/Runtime/Renderer/Public/RectLightTexture.h
@@ -11,13 +11,13 @@ namespace RectLightAtlas
 
 // Scope for invalidating a particular texture 
 // This ensures the atlas contains the latest version of the texture and filter it
-struct RENDERER_API FAtlasTextureInvalidationScope
+struct FAtlasTextureInvalidationScope
 {
-	FAtlasTextureInvalidationScope(const UTexture* In);
-	~FAtlasTextureInvalidationScope();
+	RENDERER_API FAtlasTextureInvalidationScope(const UTexture* In);
+	RENDERER_API ~FAtlasTextureInvalidationScope();
 	FAtlasTextureInvalidationScope(const FAtlasTextureInvalidationScope&) = delete;
 	FAtlasTextureInvalidationScope& operator=(const FAtlasTextureInvalidationScope&) = delete;
 	bool bLocked = false;
 };
 
-} 
\ No newline at end of file
+} 
diff --git a/Engine/Source/Runtime/Renderer/Public/RendererUtils.h b/Engine/Source/Runtime/Renderer/Public/RendererUtils.h
--- a/Engine/Source/Runtime/Renderer/Public/RendererUtils.h
+++ b/Engine/Source/Runtime/Renderer/Public/RendererUtils.h
@@ -20,10 +20,10 @@ class FScene;
 struct IPooledRenderTarget;
 template <typename ReferencedType> class TRefCountPtr;
 
-class RENDERER_API FRenderTargetWriteMask
+class FRenderTargetWriteMask
 {
 public:
-	static void Decode(
+	static RENDERER_API void Decode(
 		FRHICommandListImmediate& RHICmdList,
 		FGlobalShaderMap* ShaderMap,
 		TArrayView<IPooledRenderTarget* const> InRenderTargets,
@@ -31,7 +31,7 @@ public:
 		ETextureCreateFlags RTWriteMaskFastVRamConfig,
 		const TCHAR* RTWriteMaskDebugName);
 
-	static void Decode(
+	static RENDERER_API void Decode(
 		FRDGBuilder& GraphBuilder,
 		FGlobalShaderMap* ShaderMap,
 		TArrayView<FRDGTextureRef const> InRenderTargets,
@@ -40,7 +40,7 @@ public:
 		const TCHAR* RTWriteMaskDebugName);
 };
 
-class RENDERER_API FDepthBounds
+class FDepthBounds
 {
 public:
 
@@ -50,7 +50,7 @@ public:
 		float MaxDepth;
 	};
 
-	static FDepthBoundsValues CalculateNearFarDepthExcludingSky();
+	static RENDERER_API FDepthBoundsValues CalculateNearFarDepthExcludingSky();
 };
 
 // A minimal uniform struct providing necessary access for external systems to Strata parameters.
diff --git a/Engine/Source/Runtime/Renderer/Public/SceneRenderTargetParameters.h b/Engine/Source/Runtime/Renderer/Public/SceneRenderTargetParameters.h
--- a/Engine/Source/Runtime/Renderer/Public/SceneRenderTargetParameters.h
+++ b/Engine/Source/Runtime/Renderer/Public/SceneRenderTargetParameters.h
@@ -124,7 +124,7 @@ extern RENDERER_API FSceneTextureShaderParameters GetSceneTextureShaderParameter
  *  left in an SRV read state, so they can safely be used for read without being re-imported into
  *  RDG. Likewise, the uniform buffer is non-RDG and can be used as is.
  */
-class RENDERER_API FSceneTextureExtracts : public FRenderResource
+class FSceneTextureExtracts : public FRenderResource
 {
 public:
 	FRHIUniformBuffer* GetUniformBuffer() const
@@ -147,10 +147,10 @@ public:
 		return Depth ? Depth->GetRHI() : nullptr;
 	}
 
-	void QueueExtractions(FRDGBuilder& GraphBuilder, const FSceneTextures& SceneTextures);
+	RENDERER_API void QueueExtractions(FRDGBuilder& GraphBuilder, const FSceneTextures& SceneTextures);
 
 private:
-	void Release();
+	RENDERER_API void Release();
 	void ReleaseRHI() override { Release(); }
 
 	// Contains the resolved scene depth target.
diff --git a/Engine/Source/Runtime/Renderer/Public/ShaderPrintParameters.h b/Engine/Source/Runtime/Renderer/Public/ShaderPrintParameters.h
--- a/Engine/Source/Runtime/Renderer/Public/ShaderPrintParameters.h
+++ b/Engine/Source/Runtime/Renderer/Public/ShaderPrintParameters.h
@@ -84,14 +84,14 @@ namespace ShaderPrint
 	RENDERER_API void RequestSpaceForTriangles(uint32 MaxElementCount);
 
 	/** Structure containing setup for shader print capturing. */
-	struct RENDERER_API FShaderPrintSetup
+	struct FShaderPrintSetup
 	{
 		/** Construct with shader print disabled setup. */
 		FShaderPrintSetup() = default;
 		/** Construct with view and system defaults. */
-		FShaderPrintSetup(FSceneView const& InView);
+		RENDERER_API FShaderPrintSetup(FSceneView const& InView);
 		/** Construct with view rectangle and system defaults. */
-		FShaderPrintSetup(FIntRect InViewRect);
+		RENDERER_API FShaderPrintSetup(FIntRect InViewRect);
 
 		/** The shader print system's enabled state. This is set in the constructor and should't be overriden. */
 		bool bEnabled = false;
@@ -151,7 +151,7 @@ namespace ShaderPrint
  * (i) FreezeShaderPrintData() to "freeze" the data which exports it from the current RDG builder context.
  * (ii) SubmitShaderPrintData() to submit the frozen data for later thawing and rendering.
  */
-struct RENDERER_API FShaderPrintData
+struct FShaderPrintData
 {
 	ShaderPrint::FShaderPrintSetup Setup;
 	TUniformBufferRef<ShaderPrint::FShaderPrintCommonParameters> UniformBuffer;
@@ -166,7 +166,7 @@ struct RENDERER_API FShaderPrintData
  * (i) Can be thawed by the client for continued gathering of shader print glyphs, or
  * (ii) Can be submitted for later rendering using SubmitShaderPrintData().
  */
-struct RENDERER_API FFrozenShaderPrintData
+struct FFrozenShaderPrintData
 {
 	ShaderPrint::FShaderPrintSetup Setup;
 
diff --git a/Engine/Source/Runtime/Renderer/Public/TemporalUpscaler.h b/Engine/Source/Runtime/Renderer/Public/TemporalUpscaler.h
--- a/Engine/Source/Runtime/Renderer/Public/TemporalUpscaler.h
+++ b/Engine/Source/Runtime/Renderer/Public/TemporalUpscaler.h
@@ -8,12 +8,12 @@ namespace UE::Renderer::Private
 {
 
 /** Interface for implementing third party temporal upscaler. */
-class RENDERER_API ITemporalUpscaler : public ISceneViewFamilyExtention
+class ITemporalUpscaler : public ISceneViewFamilyExtention
 {
 public:
 
 	/** Ref counted history to be saved in the history. */
-	class RENDERER_API IHistory : public IRefCountedObject
+	class IHistory : public IRefCountedObject
 	{
 	public:
 		virtual ~IHistory() {}
diff --git a/Engine/Source/Runtime/Renderer/Public/ViewDebug.h b/Engine/Source/Runtime/Renderer/Public/ViewDebug.h
--- a/Engine/Source/Runtime/Renderer/Public/ViewDebug.h
+++ b/Engine/Source/Runtime/Renderer/Public/ViewDebug.h
@@ -18,15 +18,15 @@ DECLARE_MULTICAST_DELEGATE(FOnUpdateViewDebugInfo);
 /**
  * A collection of debug data associated with the current on screen view.
  */
-struct RENDERER_API FViewDebugInfo
+struct FViewDebugInfo
 {
 	friend class FDrawPrimitiveDebuggerModule;
 
 private:
 
-	static FViewDebugInfo Instance;
+	static RENDERER_API FViewDebugInfo Instance;
 
-	FViewDebugInfo();
+	RENDERER_API FViewDebugInfo();
 	
 public:
 
@@ -80,23 +80,23 @@ private:
 	
 	TArray<FPrimitiveInfo> Primitives;
 
-	void ProcessPrimitive(FPrimitiveSceneInfo* PrimitiveSceneInfo, const FViewInfo& View, FScene* Scene, const UPrimitiveComponent* DebugComponent);
+	RENDERER_API void ProcessPrimitive(FPrimitiveSceneInfo* PrimitiveSceneInfo, const FViewInfo& View, FScene* Scene, const UPrimitiveComponent* DebugComponent);
 
-	void CaptureNextFrame();
+	RENDERER_API void CaptureNextFrame();
 
-	void EnableLiveCapture();
+	RENDERER_API void EnableLiveCapture();
 
-	void DisableLiveCapture();
+	RENDERER_API void DisableLiveCapture();
 
-	static void DumpPrimitives(FScene* Scene, const FViewCommands& ViewCommands);
+	static RENDERER_API void DumpPrimitives(FScene* Scene, const FViewCommands& ViewCommands);
 
 public:
-	void ProcessPrimitives(FScene* Scene, const FViewInfo& View, const FViewCommands& ViewCommands);
+	RENDERER_API void ProcessPrimitives(FScene* Scene, const FViewInfo& View, const FViewCommands& ViewCommands);
 
 	/**
 	 * Writes the currently stored information out to a CSV file.
 	 */
-	void DumpToCSV() const;
+	RENDERER_API void DumpToCSV() const;
 
 	/**
 	 * Performs an operation for each primitive currently tracked.
@@ -121,13 +121,13 @@ public:
 	 * Checks if this debug information has ever been updated.
 	 * @returns True if the information has been updated at least once.
 	 */
-	bool HasEverUpdated() const;
+	RENDERER_API bool HasEverUpdated() const;
 
 	/**
 	 * Checks if current information is from an older frame.
 	 * @returns True if the data in this object is outdated.
 	 */
-	bool IsOutOfDate() const;
+	RENDERER_API bool IsOutOfDate() const;
 
 	template <typename UserClass>
 	FDelegateHandle AddUpdateHandler(UserClass* UserObject, void (UserClass::*Func)())
@@ -145,4 +145,4 @@ public:
 		OnUpdate.Remove(Handle);
 	}
 };
-#endif
\ No newline at end of file
+#endif
