Heterogeneous Volumes: Modifying frustum grid to respect custom field-of-view

diff --git a/Engine/Shaders/Private/HeterogeneousVolumes/HeterogeneousVolumesFrustumVoxelGrid.usf b/Engine/Shaders/Private/HeterogeneousVolumes/HeterogeneousVolumesFrustumVoxelGrid.usf
--- a/Engine/Shaders/Private/HeterogeneousVolumes/HeterogeneousVolumesFrustumVoxelGrid.usf
+++ b/Engine/Shaders/Private/HeterogeneousVolumes/HeterogeneousVolumesFrustumVoxelGrid.usf
@@ -23,7 +23,7 @@ RWStructuredBuffer<FTopLevelGridData> RWTopLevelGridBuffer;
 float4x4 ViewToWorld;
 int3 VoxelDimensions;
 
-float HalfFOV;
+float TanHalfFOV;
 float NearPlaneDepth;
 float FarPlaneDepth;
 
@@ -32,8 +32,8 @@ bool BoxIntersectsVoxel(float3 WorldBoxCenter, float WorldBoxRadius, uint3 Voxel
 	float3 VoxelPosMin = VoxelIndex;
 	float3 VoxelPosMax = (VoxelIndex + 1);
 
-	float3 ViewPosMin = VoxelToView(VoxelPosMin, VoxelDimensions, NearPlaneDepth, FarPlaneDepth);
-	float3 ViewPosMax = VoxelToView(VoxelPosMax, VoxelDimensions, NearPlaneDepth, FarPlaneDepth);
+	float3 ViewPosMin = VoxelToView(VoxelPosMin, VoxelDimensions, NearPlaneDepth, FarPlaneDepth, TanHalfFOV);
+	float3 ViewPosMax = VoxelToView(VoxelPosMax, VoxelDimensions, NearPlaneDepth, FarPlaneDepth, TanHalfFOV);
 
 	float3 WorldPosMin = mul(float4(ViewPosMin, 1), ViewToWorld).xyz;
 	float3 WorldPosMax = mul(float4(ViewPosMax, 1), ViewToWorld).xyz;
diff --git a/Engine/Shaders/Private/HeterogeneousVolumes/HeterogeneousVolumesFrustumVoxelGridUtils.ush b/Engine/Shaders/Private/HeterogeneousVolumes/HeterogeneousVolumesFrustumVoxelGridUtils.ush
--- a/Engine/Shaders/Private/HeterogeneousVolumes/HeterogeneousVolumesFrustumVoxelGridUtils.ush
+++ b/Engine/Shaders/Private/HeterogeneousVolumes/HeterogeneousVolumesFrustumVoxelGridUtils.ush
@@ -102,12 +102,12 @@ float InversePerspectiveRemap(float NormalizedDepth, int3 VoxelDimensions)
 	return ViewDepth;
 }
 
-float3 ViewToVoxel(float3 ViewPos, int3 VoxelDimensions, float NearPlaneDepth, float FarPlaneDepth)
+float3 ViewToVoxel(float3 ViewPos, int3 VoxelDimensions, float NearPlaneDepth, float FarPlaneDepth, float TanHalfFOV)
 {
 	//float4 ClipPos = mul(float4(ViewPos, 1), ViewToClip);
 	//float3 ScreenPos = ClipPos.xyz / ClipPos.w;
 	float2 AspectRatio = VoxelDimensions.xy / float(VoxelDimensions.x);
-	float2 ScreenPos = ViewPos.xy / (AspectRatio * ViewPos.z);
+	float2 ScreenPos = ViewPos.xy / (TanHalfFOV * AspectRatio * ViewPos.z);
 
 #if 0
 	float NormalizedDepth = SquareDistanceRemap(ViewPos.z, NearPlaneDepth, FarPlaneDepth);
@@ -124,7 +124,7 @@ float3 ViewToVoxel(float3 ViewPos, int3 VoxelDimensions, float NearPlaneDepth, f
 	return VoxelPos;
 }
 
-float3 VoxelToView(float3 VoxelPos, int3 VoxelDimensions, float NearPlaneDepth, float FarPlaneDepth)
+float3 VoxelToView(float3 VoxelPos, int3 VoxelDimensions, float NearPlaneDepth, float FarPlaneDepth, float TanHalfFOV)
 {
 	float3 NDC = VoxelPos / VoxelDimensions;
 	float2 ScreenPos = NDC.xy * 2.0 - 1.0;
@@ -140,7 +140,7 @@ float3 VoxelToView(float3 VoxelPos, int3 VoxelDimensions, float NearPlaneDepth,
 
 	// TODO: Guarantee no view-space translation exists..
 	float2 AspectRatio = VoxelDimensions.xy / float(VoxelDimensions.x);
-	float3 ViewPos = float3(ScreenPos * AspectRatio * ViewDepth, ViewDepth);
+	float3 ViewPos = float3(ScreenPos * TanHalfFOV * AspectRatio * ViewDepth, ViewDepth);
 	return ViewPos;
 }
 
diff --git a/Engine/Shaders/Private/HeterogeneousVolumes/HeterogeneousVolumesMaterialBakingPipeline.usf b/Engine/Shaders/Private/HeterogeneousVolumes/HeterogeneousVolumesMaterialBakingPipeline.usf
--- a/Engine/Shaders/Private/HeterogeneousVolumes/HeterogeneousVolumesMaterialBakingPipeline.usf
+++ b/Engine/Shaders/Private/HeterogeneousVolumes/HeterogeneousVolumesMaterialBakingPipeline.usf
@@ -399,7 +399,7 @@ void RasterizeBottomLevelOrthoGridCS(
 int3 VoxelDimensions;
 float4x4 ViewToWorld;
 
-float HalfFOV;
+float TanHalfFOV;
 float NearPlaneDepth;
 float FarPlaneDepth;
 
@@ -440,7 +440,7 @@ void RasterizeBottomLevelFrustumGridCS(
 	float3 LocalVoxelPos = GroupThreadId + Jitter;
 	float3 VoxelPosition = TopLevelVoxelPos + LocalVoxelPos / BottomLevelVoxelResolution;
 
-	float3 ViewPos = VoxelToView(VoxelPosition, VoxelDimensions, NearPlaneDepth, FarPlaneDepth);
+	float3 ViewPos = VoxelToView(VoxelPosition, VoxelDimensions, NearPlaneDepth, FarPlaneDepth, TanHalfFOV);
 	float3 WorldPosition = mul(float4(ViewPos, 1), ViewToWorld).xyz;
 
 	float3 Extinction = 0;
diff --git a/Engine/Shaders/Private/HeterogeneousVolumes/HeterogeneousVolumesVoxelGridRendering.usf b/Engine/Shaders/Private/HeterogeneousVolumes/HeterogeneousVolumesVoxelGridRendering.usf
--- a/Engine/Shaders/Private/HeterogeneousVolumes/HeterogeneousVolumesVoxelGridRendering.usf
+++ b/Engine/Shaders/Private/HeterogeneousVolumes/HeterogeneousVolumesVoxelGridRendering.usf
@@ -770,8 +770,9 @@ void RenderTransmittanceWithVoxelGridCS(
 		int3 VoxelDimensions = FrustumGridUniformBuffer.VoxelDimensions;
 		float NearPlaneDepth = FrustumGridUniformBuffer.NearPlaneDepth;
 		float FarPlaneDepth = FrustumGridUniformBuffer.FarPlaneDepth;
-		float3 VoxelRayOrigin = ViewToVoxel(ViewRayOrigin, VoxelDimensions, NearPlaneDepth, FarPlaneDepth);
-		float3 VoxelRayEnd = ViewToVoxel(ViewRayOrigin + ViewRayDirection * ViewRayLength, VoxelDimensions, NearPlaneDepth, FarPlaneDepth);
+		float TanHalfFOV = FrustumGridUniformBuffer.TanHalfFOV;
+		float3 VoxelRayOrigin = ViewToVoxel(ViewRayOrigin, VoxelDimensions, NearPlaneDepth, FarPlaneDepth, TanHalfFOV);
+		float3 VoxelRayEnd = ViewToVoxel(ViewRayOrigin + ViewRayDirection * ViewRayLength, VoxelDimensions, NearPlaneDepth, FarPlaneDepth, TanHalfFOV);
 
 		float3 VoxelRayDirection = VoxelRayEnd - VoxelRayOrigin;
 		float VoxelRayLength = length(VoxelRayDirection);
diff --git a/Engine/Shaders/Private/HeterogeneousVolumes/HeterogeneousVolumesVoxelGridTraversal.ush b/Engine/Shaders/Private/HeterogeneousVolumes/HeterogeneousVolumesVoxelGridTraversal.ush
--- a/Engine/Shaders/Private/HeterogeneousVolumes/HeterogeneousVolumesVoxelGridTraversal.ush
+++ b/Engine/Shaders/Private/HeterogeneousVolumes/HeterogeneousVolumesVoxelGridTraversal.ush
@@ -488,8 +488,9 @@ float3 TraceFrustumVoxelGrid(float3 Origin, float3 Direction, inout float TMin,
 	int3 VoxelDimensions = FrustumGridUniformBuffer.VoxelDimensions;
 	float NearPlaneDepth = FrustumGridUniformBuffer.NearPlaneDepth;
 	float FarPlaneDepth = FrustumGridUniformBuffer.FarPlaneDepth;
-	float3 VoxelRayOrigin = ViewToVoxel(ViewRayOrigin, VoxelDimensions, NearPlaneDepth, FarPlaneDepth);
-	float3 VoxelRayDirection = ViewToVoxel(ViewRayOrigin + ViewRayDirection, VoxelDimensions, NearPlaneDepth, FarPlaneDepth);
+	float TanHalfFOV = FrustumGridUniformBuffer.TanHalfFOV;
+	float3 VoxelRayOrigin = ViewToVoxel(ViewRayOrigin, VoxelDimensions, NearPlaneDepth, FarPlaneDepth, TanHalfFOV);
+	float3 VoxelRayDirection = ViewToVoxel(ViewRayOrigin + ViewRayDirection, VoxelDimensions, NearPlaneDepth, FarPlaneDepth, TanHalfFOV);
 	VoxelRayDirection -= VoxelRayOrigin;
 	float UnitWorldToVoxelScale = length(VoxelRayDirection);
 	VoxelRayDirection /= UnitWorldToVoxelScale;
diff --git a/Engine/Shaders/Private/PathTracing/Volume/PathTracingHeterogeneousVolumes.ush b/Engine/Shaders/Private/PathTracing/Volume/PathTracingHeterogeneousVolumes.ush
--- a/Engine/Shaders/Private/PathTracing/Volume/PathTracingHeterogeneousVolumes.ush
+++ b/Engine/Shaders/Private/PathTracing/Volume/PathTracingHeterogeneousVolumes.ush
@@ -341,7 +341,8 @@ FVolumeShadedResult GetFrustumVoxelGridDensity(float3 TranslatedWorldPos, inout
 	int3 VoxelDimensions = FrustumGridUniformBuffer.VoxelDimensions;
 	float NearPlaneDepth = FrustumGridUniformBuffer.NearPlaneDepth;
 	float FarPlaneDepth = FrustumGridUniformBuffer.FarPlaneDepth;
-	float3 VoxelPos = ViewToVoxel(ViewPos, VoxelDimensions, NearPlaneDepth, FarPlaneDepth);
+	float TanHalfFOV = FrustumGridUniformBuffer.TanHalfFOV;
+	float3 VoxelPos = ViewToVoxel(ViewPos, VoxelDimensions, NearPlaneDepth, FarPlaneDepth, TanHalfFOV);
 
 	bInFrustum = all(VoxelPos > 0) && all(VoxelPos < FrustumGridUniformBuffer.TopLevelFroxelGridResolution);
 	if (bInFrustum)
diff --git a/Engine/Source/Runtime/Renderer/Private/HeterogeneousVolumes/HeterogeneousVolumes.h b/Engine/Source/Runtime/Renderer/Private/HeterogeneousVolumes/HeterogeneousVolumes.h
--- a/Engine/Source/Runtime/Renderer/Private/HeterogeneousVolumes/HeterogeneousVolumes.h
+++ b/Engine/Source/Runtime/Renderer/Private/HeterogeneousVolumes/HeterogeneousVolumes.h
@@ -182,6 +182,7 @@ BEGIN_UNIFORM_BUFFER_STRUCT(FFrustumVoxelGridUniformBufferParameters, )
 
 	SHADER_PARAMETER(float, NearPlaneDepth)
 	SHADER_PARAMETER(float, FarPlaneDepth)
+	SHADER_PARAMETER(float, TanHalfFOV)
 
 	SHADER_PARAMETER_ARRAY(FVector4f, ViewFrustumPlanes, [6])
 
@@ -227,6 +228,7 @@ struct FAdaptiveFrustumGridParameterCache
 
 	float NearPlaneDepth;
 	float FarPlaneDepth;
+	float TanHalfFOV;
 
 	FVector4f ViewFrustumPlanes[6];
 
diff --git a/Engine/Source/Runtime/Renderer/Private/HeterogeneousVolumes/HeterogeneousVolumesVoxelGridPipeline.cpp b/Engine/Source/Runtime/Renderer/Private/HeterogeneousVolumes/HeterogeneousVolumesVoxelGridPipeline.cpp
--- a/Engine/Source/Runtime/Renderer/Private/HeterogeneousVolumes/HeterogeneousVolumesVoxelGridPipeline.cpp
+++ b/Engine/Source/Runtime/Renderer/Private/HeterogeneousVolumes/HeterogeneousVolumesVoxelGridPipeline.cpp
@@ -304,6 +304,11 @@ namespace HeterogeneousVolumes
 	{
 		return CVarHeterogeneousVolumesEnableMajorantGrid.GetValueOnRenderThread() != 0;
 	}
+
+	float CalcTanHalfFOV(float FOVInDegrees)
+	{
+		return FMath::Tan(FMath::DegreesToRadians(FOVInDegrees * 0.5));
+	}
 }
 
 class FMarkTopLevelGridVoxelsForFrustumGrid : public FGlobalShader
@@ -322,7 +327,7 @@ class FMarkTopLevelGridVoxelsForFrustumGrid : public FGlobalShader
 		SHADER_PARAMETER(FVector3f, PrimitiveWorldBoundsMax)
 
 		SHADER_PARAMETER(FMatrix44f, ViewToWorld)
-		SHADER_PARAMETER(float, HalfFOV)
+		SHADER_PARAMETER(float, TanHalfFOV)
 		SHADER_PARAMETER(float, NearPlaneDepth)
 		SHADER_PARAMETER(float, FarPlaneDepth)
 
@@ -378,7 +383,7 @@ class FRasterizeBottomLevelFrustumGridCS : public FMeshMaterialShader
 		SHADER_PARAMETER(FIntVector, TopLevelGridResolution)
 		SHADER_PARAMETER(FIntVector, VoxelDimensions)
 		SHADER_PARAMETER(FMatrix44f, ViewToWorld)
-		SHADER_PARAMETER(float, HalfFOV)
+		SHADER_PARAMETER(float, TanHalfFOV)
 		SHADER_PARAMETER(float, NearPlaneDepth)
 		SHADER_PARAMETER(float, FarPlaneDepth)
 
@@ -1178,7 +1183,7 @@ void CalcViewBoundsAndMinimumVoxelSize(
 	FVector WorldCameraOrigin = View.ViewMatrices.GetViewOrigin();
 	FBoxSphereBounds WorldCameraBounds(FSphere(WorldCameraOrigin, HeterogeneousVolumes::GetMaxTraceDistance()));
 
-	float TanHalfFOV = FMath::Tan(View.ViewMatrices.ComputeHalfFieldOfViewPerAxis().X);
+	float TanHalfFOV = HeterogeneousVolumes::CalcTanHalfFOV(View.FOV);
 	int32 HalfWidth = View.ViewRect.Width() * 0.5;
 	float PixelWidth = TanHalfFOV / HalfWidth;
 
@@ -1235,7 +1240,7 @@ void CalcGlobalBoundsAndMinimumVoxelSize(
 		FVector WorldCameraOrigin = View.ViewMatrices.GetViewOrigin();
 		FBoxSphereBounds WorldCameraBounds(FSphere(WorldCameraOrigin, HeterogeneousVolumes::GetMaxTraceDistance()));
 
-		float TanHalfFOV = FMath::Tan(View.ViewMatrices.ComputeHalfFieldOfViewPerAxis().X);
+		float TanHalfFOV = HeterogeneousVolumes::CalcTanHalfFOV(View.FOV);
 		int32 HalfWidth = View.ViewRect.Width() * 0.5;
 		float PixelWidth = TanHalfFOV / HalfWidth;
 
@@ -1299,6 +1304,7 @@ void ExtractFrustumVoxelGridUniformBuffer(
 	AdaptiveFrustumGridParameterCache.bUseFrustumGrid = Parameters->bUseFrustumGrid;
 	AdaptiveFrustumGridParameterCache.NearPlaneDepth = Parameters->NearPlaneDepth;
 	AdaptiveFrustumGridParameterCache.FarPlaneDepth = Parameters->FarPlaneDepth;
+	AdaptiveFrustumGridParameterCache.TanHalfFOV = Parameters->TanHalfFOV;
 
 
 	AdaptiveFrustumGridParameterCache.WorldToClip = Parameters->WorldToClip;
@@ -1343,6 +1349,7 @@ void RegisterExternalFrustumVoxelGridUniformBuffer(
 		UniformBufferParameters->bUseFrustumGrid = AdaptiveFrustumGridParameterCache.bUseFrustumGrid;
 		UniformBufferParameters->NearPlaneDepth = AdaptiveFrustumGridParameterCache.NearPlaneDepth;
 		UniformBufferParameters->FarPlaneDepth = AdaptiveFrustumGridParameterCache.FarPlaneDepth;
+		UniformBufferParameters->TanHalfFOV = AdaptiveFrustumGridParameterCache.TanHalfFOV;
 
 		// Frustum assignment
 		for (int i = 0; i < 6; ++i)
@@ -1384,6 +1391,7 @@ void CreateEmptyFrustumVoxelGridUniformBuffer(
 		UniformBufferParameters->bUseFrustumGrid = false;
 		UniformBufferParameters->NearPlaneDepth = 0.0;
 		UniformBufferParameters->FarPlaneDepth = 0.0;
+		UniformBufferParameters->TanHalfFOV = 1.0;
 	}
 	FrustumGridUniformBuffer = GraphBuilder.CreateUniformBuffer(UniformBufferParameters);
 }
@@ -1698,7 +1706,7 @@ void MarkTopLevelGridVoxelsForFrustumGrid(
 			PassParameters->PrimitiveWorldBoundsMax = FVector3f(PrimitiveBounds.Origin + PrimitiveBounds.BoxExtent);
 
 			PassParameters->ViewToWorld = FMatrix44f(ViewToWorld);
-			PassParameters->HalfFOV = FMath::DegreesToRadians(View.FOV * 0.5);
+			PassParameters->TanHalfFOV = HeterogeneousVolumes::CalcTanHalfFOV(View.FOV);
 			PassParameters->NearPlaneDepth = NearPlaneDistance;
 			PassParameters->FarPlaneDepth = FarPlaneDistance;
 
@@ -1831,7 +1839,7 @@ void RasterizeVolumesIntoFrustumVoxelGrid(
 			PassParameters->TopLevelGridResolution = TopLevelGridResolution;
 			PassParameters->VoxelDimensions = TopLevelGridResolution;
 			PassParameters->ViewToWorld = FMatrix44f(ViewToWorld);
-			PassParameters->HalfFOV = FMath::DegreesToRadians(View.FOV * 0.5);
+			PassParameters->TanHalfFOV = HeterogeneousVolumes::CalcTanHalfFOV(View.FOV);
 			PassParameters->NearPlaneDepth = NearPlaneDistance;
 			PassParameters->FarPlaneDepth = FarPlaneDistance;
 
@@ -2029,6 +2037,7 @@ void BuildFrustumVoxelGrid(
 		UniformBufferParameters->bUseFrustumGrid = HeterogeneousVolumes::EnableFrustumVoxelGrid();
 		UniformBufferParameters->NearPlaneDepth = NearPlaneDistance;
 		UniformBufferParameters->FarPlaneDepth = FarPlaneDistance;
+		UniformBufferParameters->TanHalfFOV = HeterogeneousVolumes::CalcTanHalfFOV(View.FOV);
 
 		// Frustum assignment
 		{
