Add Disney Sheen model for Substrate.

This sheen mode is enabled by default for platforms which high quality shading. The models degrades into a Charlie Sheen model for lower platform. This can also be controlled with r.Substrate.SheenQuality.

diff --git a/Engine/Shaders/Private/BRDF.ush b/Engine/Shaders/Private/BRDF.ush
--- a/Engine/Shaders/Private/BRDF.ush
+++ b/Engine/Shaders/Private/BRDF.ush
@@ -500,7 +500,73 @@ float3 F_Fresnel( float3 SpecularColor, float VoH )
 	return 0.5 * Square( (g - VoH) / (g + VoH) ) * ( 1 + Square( ((g+VoH)*VoH - 1) / ((g-VoH)*VoH + 1) ) );
 }
 
+float4 SheenLTC_Cofficients(float NoV, float Roughness, Texture2D InSheenLTCTexture, SamplerState InSheenSampler)
+{
+	const float Alpha = sqrt(Roughness);
+	const float SatNoV = saturate(abs(NoV) + 1e-5);
+	float2 UV = float2(Alpha, SatNoV);
+	UV = UV * (31.0 / 32.0) + (0.5 / 32.0);
+	return InSheenLTCTexture.SampleLevel(InSheenSampler, UV, 0);
+}
+
+float SheenLTC_DirectionalAlbedo(float NoV, float Roughness, Texture2D InSheenLTCTexture, SamplerState InSheenSampler)
+{
+	return SheenLTC_Cofficients(NoV, Roughness, InSheenLTCTexture, InSheenSampler).z;
+}
 
+// [Zeltner et al. 2022, "Practical Multiple-Scattering Sheen Using Linearly Transformed Cosines"]
+float3 SheenLTC_Sample(float3 CosSample, float3 V, float3 L, half3 N, float NoV, float Roughness, Texture2D InSheenLTCTexture, SamplerState InSheenSampler)
+{
+	const float4 LTC = SheenLTC_Cofficients(NoV, Roughness, InSheenLTCTexture, InSheenSampler);
+	const float aInv = LTC.x;
+	const float bInv = LTC.y;
+	const float3 LocalL = normalize(float3(CosSample.x / aInv - CosSample.z * bInv / aInv, CosSample.y / aInv, CosSample.z));
+
+	// Rotate L into tangent space
+	const float3 T1 = normalize(V - N * NoV);
+	const float3 T2 = cross(N, T1);
+	const float3x3 TangentBasis = float3x3(T1, T2, N);
+	const float3 S = mul(LocalL, TangentBasis);
+	return S;
+}
+
+// [Zeltner et al. 2022, "Practical Multiple-Scattering Sheen Using Linearly Transformed Cosines"]
+float SheenLTC_Eval(float3 V, float3 L, half3 N, float NoV, float Roughness, Texture2D InSheenLTCTexture, SamplerState InSheenSampler, inout float OutDirectionalAlbedo, inout float OutPdf)
+{
+	// Rotate L into tangent space
+	const float3 T1 = normalize(V - N * NoV);
+	const float3 T2 = cross(N, T1);
+	const float3x3 TangentBasis = float3x3(T1, T2, N);
+
+	// L' = mul(L, InvT) 
+	// Since TangentBasis is orthnormal Inv(T)=Transpose(T), i.e., mul(L, transpose(T)), which is equivalent to mul(T,L)
+	const float3 LocalL = mul(TangentBasis, L); 
+
+	const float4 LTC = SheenLTC_Cofficients(NoV, Roughness, InSheenLTCTexture, InSheenSampler);
+	const float aInv = LTC.x;
+	const float bInv = LTC.y;
+	OutDirectionalAlbedo = LTC.z;
+
+	float3 WrappedLocalL = float3(
+		aInv * LocalL.x + bInv * LocalL.z,
+		aInv * LocalL.y,
+		LocalL.z);
+
+	const float LenWrappedLocalL = length(WrappedLocalL);
+	// Jabobian is defined as: det(InvLTC) / ||InvLTC . L||
+	const float Jacobian = Pow2(aInv) / Pow3(LenWrappedLocalL);
+
+	// Normalized wrapped direction
+	WrappedLocalL /= LenWrappedLocalL;
+
+	// Evaluation a normalized cos distribution
+	const float CosDistribution = WrappedLocalL.z / PI;
+	const float Out = OutDirectionalAlbedo * CosDistribution * Jacobian; // Fs * CosTheta
+
+	// The function has perfect important sampling PDF=F
+	OutPdf = Out;
+	return Out;
+}
 
 //---------------
 // EnvBRDF
diff --git a/Engine/Shaders/Private/RectLight.ush b/Engine/Shaders/Private/RectLight.ush
--- a/Engine/Shaders/Private/RectLight.ush
+++ b/Engine/Shaders/Private/RectLight.ush
@@ -362,23 +362,25 @@ float3 PolygonIrradiance( float3 Poly[4] )
 	return L;
 }
 
-// [ Heitz et al. 2016, "Real-Time Polygonal-Light Shading with Linearly Transformed Cosines" ]
-float3 RectGGXApproxLTC( float Roughness, float3 SpecularColor, half3 N, float3 V, FRect Rect, FRectTexture RectTexture, inout float3 MeanLightWorldDirection)
+struct FRectLTC
+{
+	float3x3 LTC;
+	float3x3 InvLTC;
+	float3 IrradianceScale;
+};
+
+// Return LTC coefficients for GGX-based BSDF
+FRectLTC GetRectLTC_GGX( float Roughness, float3 SpecularColor, half NoV)
 {
 	// F90 is derived as, anything less than 2% is physically impossible and is instead considered to be shadowing
 	const float3 F0  = SpecularColor;
 	const float3 F90 = saturate(50.0 * SpecularColor);
 
-	// No visibile rect light due to barn door occlusion
-	if (Rect.Extent.x == 0 || Rect.Extent.y == 0) return 0;
-
-	float NoV = saturate( abs( dot(N, V) ) + 1e-5 );
-
 	float2 UV = float2( Roughness, sqrt( 1 - NoV ) );
 	UV = UV * (63.0 / 64.0) + (0.5 / 64.0);
    
-	float4 LTCMat = View.LTCMatTexture.SampleLevel( View.LTCMatSampler, UV, 0 );
-	float4 LTCAmp = View.LTCAmpTexture.SampleLevel( View.LTCAmpSampler, UV, 0 );
+	float4 LTCMat = View.GGXLTCMatTexture.SampleLevel( View.GGXLTCMatSampler, UV, 0 );
+	float4 LTCAmp = View.GGXLTCAmpTexture.SampleLevel( View.GGXLTCAmpSampler, UV, 0 );
 
 	float3x3 LTC = {
 		float3( LTCMat.x, 0, LTCMat.z ),
@@ -395,19 +397,64 @@ float3 RectGGXApproxLTC( float Roughness, float3 SpecularColor, half3 N, float3
 		float3(-InvLTCMat.y, 0, InvLTCMat.x )
 	};
 
+	FRectLTC Out = (FRectLTC)0;
+	Out.LTC = LTC;
+	Out.InvLTC = InvLTC;
+	Out.IrradianceScale = F90 * LTCAmp.y + ( LTCAmp.x - LTCAmp.y ) * F0;
+	return Out;
+}
+
+// Return LTC coefficients for Sheen BSDF
+FRectLTC GetRectLTC_Sheen( float Roughness, half NoV)
+{	
+	const float Alpha = sqrt(Roughness);
+	const float SatNoV = saturate(abs(NoV) + 1e-5);
+	float2 UV = float2(Alpha, SatNoV);
+	UV = UV * (31.0 / 32.0) + (0.5 / 32.0);
+	const float3 SheenLTC = View.SheenLTCTexture.SampleLevel(View.SheenLTCSampler, UV, 0).xyz;
+
+	const float aInv = SheenLTC.x;
+	const float bInv = SheenLTC.y;
+
+	float3x3 LTC = {
+		float3(1/aInv, 0,     -bInv/aInv),    
+		float3(0,      1/aInv, 0),
+		float3(0,      0,      1)
+	};
+
+	float3x3 InvLTC = {
+		float3(aInv, 0,    bInv),
+		float3(0,    aInv, 0),
+		float3(0,    0,    1)
+	};
+
+	FRectLTC Out = (FRectLTC)0;
+	Out.LTC = LTC;
+	Out.InvLTC = InvLTC;
+	Out.IrradianceScale = SheenLTC.z;
+	return Out;
+}
+
+// Integrate a LTC BSDF with a rect light
+// [ Heitz et al. 2016, "Real-Time Polygonal-Light Shading with Linearly Transformed Cosines" ]
+float3 RectApproxLTC(FRectLTC In, half3 N, float3 V, FRect Rect, FRectTexture RectTexture, inout float3 OutMeanLightWorldDirection)
+{	
+	// No visibile rect light due to barn door occlusion
+	if (Rect.Extent.x == 0 || Rect.Extent.y == 0) return 0;
+
 	// Rotate to tangent space
 	float3 T1 = normalize( V - N * dot( N, V ) );
 	float3 T2 = cross( N, T1 );
 	float3x3 TangentBasis = float3x3( T1, T2, N );
 
-	LTC = mul( LTC, TangentBasis );
-	InvLTC = mul( transpose( TangentBasis ), InvLTC );
+	In.LTC = mul( In.LTC, TangentBasis );
+	In.InvLTC = mul( transpose( TangentBasis ), In.InvLTC );
 
 	float3 Poly[4];
-	Poly[0] = mul( LTC, Rect.Origin - Rect.Axis[0] * Rect.Extent.x - Rect.Axis[1] * Rect.Extent.y );
-	Poly[1] = mul( LTC, Rect.Origin + Rect.Axis[0] * Rect.Extent.x - Rect.Axis[1] * Rect.Extent.y );
-	Poly[2] = mul( LTC, Rect.Origin + Rect.Axis[0] * Rect.Extent.x + Rect.Axis[1] * Rect.Extent.y );
-	Poly[3] = mul( LTC, Rect.Origin - Rect.Axis[0] * Rect.Extent.x + Rect.Axis[1] * Rect.Extent.y );
+	Poly[0] = mul( In.LTC, Rect.Origin - Rect.Axis[0] * Rect.Extent.x - Rect.Axis[1] * Rect.Extent.y );
+	Poly[1] = mul( In.LTC, Rect.Origin + Rect.Axis[0] * Rect.Extent.x - Rect.Axis[1] * Rect.Extent.y );
+	Poly[2] = mul( In.LTC, Rect.Origin + Rect.Axis[0] * Rect.Extent.x + Rect.Axis[1] * Rect.Extent.y );
+	Poly[3] = mul( In.LTC, Rect.Origin - Rect.Axis[0] * Rect.Extent.x + Rect.Axis[1] * Rect.Extent.y );
 
 	// Vector irradiance
 	float3 L = PolygonIrradiance( Poly );
@@ -439,22 +486,47 @@ float3 RectGGXApproxLTC( float Roughness, float3 SpecularColor, half3 N, float3
 	Irradiance = Irradiance * ( 1 - a2*a2 ) + a2 * Falloff * NoL;
 #endif
 	
-	float3 IrradianceScale = F90 * LTCAmp.y + ( LTCAmp.x - LTCAmp.y ) * F0;
-	
 	// Transform to world space
-	L = mul( InvLTC, L );
-	MeanLightWorldDirection = L;
+	L = mul( In.InvLTC, L );
+	OutMeanLightWorldDirection = L;
 
 	float3 LightColor = SampleSourceTexture( L, Rect, RectTexture );
 	
-	return LightColor * Irradiance * IrradianceScale;
+	return LightColor * Irradiance * In.IrradianceScale;
+}
+
+// Integrated a GGX-based BSDF with a rect light using LTC
+float3 RectGGXApproxLTC( float Roughness, float3 SpecularColor, half3 N, float3 V, FRect Rect, FRectTexture RectTexture, inout float3 OutMeanLightWorldDirection)
+{
+	// No visibile rect light due to barn door occlusion
+	if (Rect.Extent.x == 0 || Rect.Extent.y == 0) return 0;
+
+	const float NoV = saturate( abs( dot(N, V) ) + 1e-5 );
+
+	const FRectLTC LTC = GetRectLTC_GGX(Roughness, SpecularColor, NoV);
+	return RectApproxLTC(LTC, N, V, Rect, RectTexture, OutMeanLightWorldDirection);
 }
+
 float3 RectGGXApproxLTC(float Roughness, float3 SpecularColor, half3 N, float3 V, FRect Rect, FRectTexture RectTexture)
 {
 	float3 MeanLightWorldDirection = 0.0f;
 	return RectGGXApproxLTC(Roughness, SpecularColor, N, V, Rect, RectTexture, MeanLightWorldDirection);
 }
 
+// Integrated a Sheen-based BSDF with a rect light using LTC
+float3 RectSheenApproxLTC( float Roughness, half3 N, float3 V, FRect Rect, FRectTexture RectTexture, inout float DirectionalAlbedo)
+{
+	// No visibile rect light due to barn door occlusion
+	if (Rect.Extent.x == 0 || Rect.Extent.y == 0) return 0;
+
+	const float NoV = saturate( abs( dot(N, V) ) + 1e-5 );
+
+	const FRectLTC LTC = GetRectLTC_Sheen(Roughness, NoV);
+	DirectionalAlbedo = LTC.IrradianceScale.x;
+
+	float3 MeanLightWorldDirection = 0.0f;
+	return RectApproxLTC(LTC, N, V, Rect, RectTexture, MeanLightWorldDirection);
+}
 
 // Rectangle projected to a sphere
 // [Urena et al. 2013, "An Area-Preserving Parametrization for Spherical Rectangles"]
diff --git a/Engine/Shaders/Private/Strata/StrataEvaluation.ush b/Engine/Shaders/Private/Strata/StrataEvaluation.ush
--- a/Engine/Shaders/Private/Strata/StrataEvaluation.ush
+++ b/Engine/Shaders/Private/Strata/StrataEvaluation.ush
@@ -40,6 +40,9 @@
 #define STRATA_TRANSLUCENT_ENABLED 0
 #endif
 
+#ifndef STRATA_SHEEN_QUALITY
+#define STRATA_SHEEN_QUALITY 0
+#endif
 
 #ifndef STRATA_SSS_TRANSMISSION
 #define STRATA_SSS_TRANSMISSION 0
@@ -696,21 +699,45 @@ FStrataEvaluateResult StrataEvaluateBSDFCommon(FStrataBSDFContext BSDFContext, F
 				const float  FuzzRoughness = MakeRoughnessSafe(SLAB_FUZZ_ROUGHNESS(BSDFContext.BSDF), STRATA_MIN_FUZZ_ROUGHNESS);
 
 				// F0 is ignored, as energy preservation is handle solely based on the white directional albedo and FuzzAmount
-				FBxDFEnergyTermsA EnergyTerms = ComputeClothEnergyTermsA(FuzzRoughness, BSDFContext.Context.NoV);
+				FBxDFEnergyTermsA EnergyTerms = (FBxDFEnergyTermsA)1.f;
+				float3 ClothSpecularPathValueNoL = 0;
+
+				#if STRATA_SHEEN_QUALITY == 1
+				// Disney LTC
+				{
+					float DirectionalAlbedo = 1;
+					BRANCH
+					if (bIsRectLight)
+					{
+						ClothSpecularPathValueNoL = FuzzF0 * RectSheenApproxLTC(FuzzRoughness, BSDFContext.N, BSDFContext.V, AreaLightContext.AreaLight.Rect, AreaLightContext.AreaLight.Texture, DirectionalAlbedo);
+					}
+					else
+					{
+						float SheenPDF = 1;
+						ClothSpecularPathValueNoL = FuzzF0 * SheenLTC_Eval(BSDFContext.V, AreaLightContext.L, BSDFContext.N, BSDFContext.Context.NoV, FuzzRoughness, View.SheenLTCTexture, View.SheenLTCSampler, DirectionalAlbedo, SheenPDF);
+					}
+					EnergyTerms.E = DirectionalAlbedo;
+					EnergyTerms.W = 1.f;
+				}
+				#else
+				// Charlie cloth modle for D and Ashikhmin for visibility (more efficient than Charlie).
+				{
+					// ComputeEnergyConservation(EnergyTerms) is not taken into account because we do not white white fuzz to reflect all incoming energy
+					// The fuzz should have some transmission to the lower lobes. In order to pass a furnace test, a combination of white diffuse+white fuzz is required
+					EnergyTerms = ComputeClothEnergyTermsA(FuzzRoughness, BSDFContext.Context.NoV);
+
+					float ClothD   = D_Charlie(FuzzRoughness, BSDFContext.Context.NoH);
+					float ClothVis = Vis_Ashikhmin(BSDFContext.Context.NoV, AreaLightContext.NoL);
+					float3 ClothF  = F_Schlick(FuzzF0, BSDFContext.Context.VoH);
+					ClothSpecularPathValueNoL = (ClothD * ClothVis) * ClothF * ComputeEnergyConservation(EnergyTerms) * AreaLightContext.NoL;
+				}
+				#endif //STRATA_SHEEN_QUALITY
 
 				// The specular and diffuse components below the fuzz are only attenuated linearly according to the amount of fuzz.
 				const float3 Cloth_DirectionalAlbedo_SpecularTransmission = lerp(1.0, ComputeEnergyPreservation(EnergyTerms), FuzzAmount);
 
-				// Charlie cloth modle for D and Ashikhmin for visibility (more efficient than Charlie).
-				float ClothD   = D_Charlie(FuzzRoughness, BSDFContext.Context.NoH);
-				float ClothVis = Vis_Ashikhmin(BSDFContext.Context.NoV, AreaLightContext.NoL);
-				float3 ClothF  = F_Schlick(FuzzF0, BSDFContext.Context.VoH);
-				// ComputeEnergyConservation(EnergyTerms) is not taken into account because we do not white white fuzz to reflect all incoming energy
-				// The fuzz should have some transmission to the lower lobes. In order to pass a furnace test, a combination of white diffuse+white fuzz is required
-				float3 ClothSpecularPathValue = (ClothD * ClothVis) * ClothF;
-
 				// Area light are not supported by the cloth BRDF
-				float3 ClothIntegratedSpecularValue = (ShadowTerms.SurfaceShadow * AreaLightContext.NoL * AreaLightContext.Falloff * FuzzAmount) * AreaLightContext.AreaLight.FalloffColor * ClothSpecularPathValue;
+				float3 ClothIntegratedSpecularValue = (ShadowTerms.SurfaceShadow * AreaLightContext.Falloff * FuzzAmount) * AreaLightContext.AreaLight.FalloffColor * ClothSpecularPathValueNoL;
 				
 				// Apply specular transmittance to diffuse and specular lob from slab medium and interface sitting below the layer of fuzz.
 				Sample.DiffusePathValue			*= Cloth_DirectionalAlbedo_SpecularTransmission;
@@ -718,7 +745,7 @@ FStrataEvaluateResult StrataEvaluateBSDFCommon(FStrataBSDFContext BSDFContext, F
 				Sample.SpecularPathValue		*= Cloth_DirectionalAlbedo_SpecularTransmission;
 				Sample.IntegratedSpecularValue	*= Cloth_DirectionalAlbedo_SpecularTransmission;
 				
-				Sample.SpecularPathValue		+= ClothSpecularPathValue;
+				Sample.SpecularPathValue		+= ClothSpecularPathValueNoL;
 				Sample.IntegratedSpecularValue  += ClothIntegratedSpecularValue;
 
 				// This is not good. We should really have a separate PDF for cloth it self associated with a probability.
@@ -730,7 +757,7 @@ FStrataEvaluateResult StrataEvaluateBSDFCommon(FStrataBSDFContext BSDFContext, F
 				if (bHaziness)
 				{
 					Sample.SpecularHazePathValue*= Cloth_DirectionalAlbedo_SpecularTransmission;
-					Sample.SpecularHazePathValue+= ClothSpecularPathValue;
+					Sample.SpecularHazePathValue+= ClothSpecularPathValueNoL;
 
 					// STRATA_TODO we should have a separated cloth PDF and Weight
 					//Sample.SpecularHazePDF		= lerp(Sample.SpecularHazePDF,			BSDFContext.SatNoV / PI,	Fuzz);;
@@ -1598,8 +1625,19 @@ FStrataEnvLightResult StrataEvaluateForEnvLight(FStrataBSDFContext BSDFContext,
 				const float3 FuzzF0 = SLAB_FUZZ_COLOR(BSDFContext.BSDF);
 				const float  FuzzRoughness = MakeRoughnessSafe(SLAB_FUZZ_ROUGHNESS(BSDFContext.BSDF), STRATA_MIN_FUZZ_ROUGHNESS);
 
-				// F0 is ignored, as energy preservation is handle solely based on the white directional albedo and FuzzAmount
-				const FBxDFEnergyTermsA ClothEnergyTerms = ComputeClothEnergyTermsA(FuzzRoughness, BSDFContext.Context.NoV);
+				FBxDFEnergyTermsA ClothEnergyTerms = (FBxDFEnergyTermsA)1.0f;
+				#if STRATA_SHEEN_QUALITY == 1
+				{
+					const float DirectionalAlbedo = SheenLTC_DirectionalAlbedo(BSDFContext.Context.NoV, FuzzRoughness, View.SheenLTCTexture, View.SheenLTCSampler);
+					ClothEnergyTerms.E = DirectionalAlbedo;
+					ClothEnergyTerms.W = 1.f;
+				}
+				#else
+				{
+					// F0 is ignored, as energy preservation is handle solely based on the white directional albedo and FuzzAmount
+					ClothEnergyTerms = ComputeClothEnergyTermsA(FuzzRoughness, BSDFContext.Context.NoV);
+				}
+				#endif // STRATA_SHEEN_QUALITY
 
 				// The specular and diffuse components below the fuzz are only attenuated linearly according to the amount of fuzz.
 				const float Cloth_DirectionalAlbedo_SpecularTransmission = lerp(1.0, ComputeEnergyPreservation(ClothEnergyTerms), FuzzAmount);
diff --git a/Engine/Source/Runtime/Engine/Private/SceneManagement.cpp b/Engine/Source/Runtime/Engine/Private/SceneManagement.cpp
--- a/Engine/Source/Runtime/Engine/Private/SceneManagement.cpp
+++ b/Engine/Source/Runtime/Engine/Private/SceneManagement.cpp
@@ -980,11 +980,13 @@ FViewUniformShaderParameters::FViewUniformShaderParameters()
 	HairScatteringLUTTexture = BlackVolume;
 	HairScatteringLUTSampler = TStaticSamplerState<SF_Bilinear>::GetRHI();
 
-	// Rect area light
-	LTCMatTexture = GBlackTextureWithSRV->TextureRHI;
-	LTCMatSampler = TStaticSamplerState<SF_Bilinear, AM_Clamp, AM_Clamp, AM_Clamp>::GetRHI();
-	LTCAmpTexture = GBlackTextureWithSRV->TextureRHI;
-	LTCAmpSampler = TStaticSamplerState<SF_Bilinear, AM_Clamp, AM_Clamp, AM_Clamp>::GetRHI();
+	// GGX/Sheen - Rect area light
+	GGXLTCMatTexture = GBlackTextureWithSRV->TextureRHI;
+	GGXLTCMatSampler = TStaticSamplerState<SF_Bilinear, AM_Clamp, AM_Clamp, AM_Clamp>::GetRHI();
+	GGXLTCAmpTexture = GBlackTextureWithSRV->TextureRHI;
+	GGXLTCAmpSampler = TStaticSamplerState<SF_Bilinear, AM_Clamp, AM_Clamp, AM_Clamp>::GetRHI();
+	SheenLTCTexture = GBlackTextureWithSRV->TextureRHI;
+	SheenLTCSampler = TStaticSamplerState<SF_Bilinear, AM_Clamp, AM_Clamp, AM_Clamp>::GetRHI();
 
 	// Shading energy conservation
 	bShadingEnergyConservation = 0u;
diff --git a/Engine/Source/Runtime/Engine/Private/ShaderCompiler/ShaderCompiler.cpp b/Engine/Source/Runtime/Engine/Private/ShaderCompiler/ShaderCompiler.cpp
--- a/Engine/Source/Runtime/Engine/Private/ShaderCompiler/ShaderCompiler.cpp
+++ b/Engine/Source/Runtime/Engine/Private/ShaderCompiler/ShaderCompiler.cpp
@@ -6669,6 +6669,9 @@ void GlobalBeginCompileShader(
 			const bool bLowQuality = StrataShadingQuality > 1;
 			Input.Environment.SetDefine(TEXT("USE_ACHROMATIC_BXDF_ENERGY"), bLowQuality ? 1u : 0u);
 
+			const uint32 StrataSheenQuality = Strata::GetSheenQuality();
+			Input.Environment.SetDefine(TEXT("STRATA_SHEEN_QUALITY"), bLowQuality ? 2 : StrataSheenQuality);
+
 			const uint32 StrataNormalQuality = Strata::GetNormalQuality();
 			Input.Environment.SetDefine(TEXT("STRATA_NORMAL_QUALITY"), StrataNormalQuality);
 
diff --git a/Engine/Source/Runtime/Engine/Public/SceneView.h b/Engine/Source/Runtime/Engine/Public/SceneView.h
--- a/Engine/Source/Runtime/Engine/Public/SceneView.h
+++ b/Engine/Source/Runtime/Engine/Public/SceneView.h
@@ -996,11 +996,13 @@ BEGIN_GLOBAL_SHADER_PARAMETER_STRUCT_WITH_CONSTRUCTOR(FViewUniformShaderParamete
 	// Hair
 	SHADER_PARAMETER_TEXTURE(Texture3D, HairScatteringLUTTexture)
 	SHADER_PARAMETER_SAMPLER(SamplerState, HairScatteringLUTSampler)
-	// Rect area light	
-	SHADER_PARAMETER_TEXTURE(Texture2D, LTCMatTexture)
-	SHADER_PARAMETER_SAMPLER(SamplerState, LTCMatSampler)
-	SHADER_PARAMETER_TEXTURE(Texture2D, LTCAmpTexture)
-	SHADER_PARAMETER_SAMPLER(SamplerState, LTCAmpSampler)
+	// GGX/Sheen LTC textures
+	SHADER_PARAMETER_TEXTURE(Texture2D, GGXLTCMatTexture)
+	SHADER_PARAMETER_SAMPLER(SamplerState, GGXLTCMatSampler)
+	SHADER_PARAMETER_TEXTURE(Texture2D, GGXLTCAmpTexture)
+	SHADER_PARAMETER_SAMPLER(SamplerState, GGXLTCAmpSampler)
+	SHADER_PARAMETER_TEXTURE(Texture2D, SheenLTCTexture)
+	SHADER_PARAMETER_SAMPLER(SamplerState, SheenLTCSampler)
 	// Energy conservation
 	SHADER_PARAMETER(uint32, bShadingEnergyConservation)
 	SHADER_PARAMETER(uint32, bShadingEnergyPreservation)
diff --git a/Engine/Source/Runtime/RenderCore/Private/RenderUtils.cpp b/Engine/Source/Runtime/RenderCore/Private/RenderUtils.cpp
--- a/Engine/Source/Runtime/RenderCore/Private/RenderUtils.cpp
+++ b/Engine/Source/Runtime/RenderCore/Private/RenderUtils.cpp
@@ -1583,6 +1583,12 @@ static TAutoConsoleVariable<int32> CVarSubstrateOpaqueMaterialRoughRefraction(
 	TEXT("Enable Substrate opaque material rough refractions effect from top layers over layers below."),
 	ECVF_ReadOnly | ECVF_RenderThreadSafe);
 
+static TAutoConsoleVariable<int32> CVarSubstrateSheenQuality(
+	TEXT("r.Substrate.SheenQuality"),
+	1,
+	TEXT("Define Substrate sheen quality (1: Disney-based sheen, 2: Charlie-based sheen, ). r.Substrate.ShadingQuality=2 forces SheenQuality to 2. This variable is read-only."),
+	ECVF_ReadOnly | ECVF_RenderThreadSafe);
+
 static TAutoConsoleVariable<int32> CVarSubstrateShadingQuality(
 	TEXT("r.Substrate.ShadingQuality"),
 	1,
@@ -1745,6 +1751,11 @@ namespace Strata
 		return CVarSubstrateTileCoord8Bits.GetValueOnAnyThread() > 0 ? 1 : 0;
 	}
 
+	uint32 GetSheenQuality()
+	{
+		return CVarSubstrateSheenQuality.GetValueOnAnyThread();
+	}
+
 	uint32 GetShadingQuality()
 	{
 		return CVarSubstrateShadingQuality.GetValueOnAnyThread();
diff --git a/Engine/Source/Runtime/RenderCore/Private/Shader.cpp b/Engine/Source/Runtime/RenderCore/Private/Shader.cpp
--- a/Engine/Source/Runtime/RenderCore/Private/Shader.cpp
+++ b/Engine/Source/Runtime/RenderCore/Private/Shader.cpp
@@ -2015,6 +2015,10 @@ void ShaderMapAppendKeyString(EShaderPlatform Platform, FString& KeyString)
 			KeyString += FString::Printf(TEXT("_STSHQL%u"), Strata::GetShadingQuality(Platform));
 		}
 
+		{
+			KeyString += FString::Printf(TEXT("_SSHEEN%u"), Strata::GetSheenQuality());
+		}
+
 		if (Strata::IsGlintEnabled())
 		{
 			KeyString += FString::Printf(TEXT("_STRTGLT"));
diff --git a/Engine/Source/Runtime/RenderCore/Public/RenderUtils.h b/Engine/Source/Runtime/RenderCore/Public/RenderUtils.h
--- a/Engine/Source/Runtime/RenderCore/Public/RenderUtils.h
+++ b/Engine/Source/Runtime/RenderCore/Public/RenderUtils.h
@@ -484,6 +484,8 @@ namespace Strata
 
 	RENDERCORE_API uint32 GetNormalQuality();
 
+	RENDERCORE_API uint32 GetSheenQuality();
+
 	RENDERCORE_API uint32 GetShadingQuality();
 	RENDERCORE_API uint32 GetShadingQuality(EShaderPlatform InPlatform);
 }
diff --git a/Engine/Source/Runtime/Renderer/Private/LTC.h b/Engine/Source/Runtime/Renderer/Private/LTC.h
--- a/Engine/Source/Runtime/Renderer/Private/LTC.h
+++ b/Engine/Source/Runtime/Renderer/Private/LTC.h
@@ -1,8 +1,11 @@
 #pragma once
 
-static const int LTC_Size = 64;
+///////////////////////////////////////////////////////////////////////////////////////////////////
+// GGX 
 
-static const float LTC_Mat[ 4 * LTC_Size * LTC_Size ] =
+static const int GGX_LTC_Size = 64;
+
+static const float GGX_LTC_Mat[ 4 * GGX_LTC_Size * GGX_LTC_Size] =
 {
 1, 0, 0, 2e-05,
 1, 0, 0, 0.000503905,
@@ -4102,7 +4105,7 @@ static const float LTC_Mat[ 4 * LTC_Size * LTC_Size ] =
 0.996389, -0.0808124, 0.0489007, 1.6577,
 };
 
-static const float LTC_Amp[ 4 * LTC_Size * LTC_Size ] =
+static const float GGX_LTC_Amp[ 4 * GGX_LTC_Size * GGX_LTC_Size] =
 {
 1, 0, 0, 0,
 1, 7.91421e-31, 0, 0,
@@ -8201,3 +8204,851 @@ static const float LTC_Amp[ 4 * LTC_Size * LTC_Size ] =
 0.93222, 0.0484975, 0, 0.984127,
 0.932164, 0.0471899, 0, 1,
 };
+
+
+///////////////////////////////////////////////////////////////////////////////////////////////////
+// Sheen
+// Data are stored as follow:
+//  * [Row      ][Column  ]
+//  * [Roughness][CosTheta]
+
+static const int Sheen_LTC_Size = 32;
+const FVector3f Sheen_LTC_Volume[Sheen_LTC_Size][Sheen_LTC_Size] =
+{
+{
+	FVector3f(0.01415, 0.00060, 0.00001), FVector3f(0.00000, 0.00000, 0.00000), FVector3f(0.00000, 0.00000, 0.00000),
+	FVector3f(0.00000, 0.00000, 0.00000), FVector3f(0.00000, 0.00000, 0.00000), FVector3f(0.00000, 0.00000, 0.00000),
+	FVector3f(0.00000, 0.00000, 0.00000), FVector3f(0.00000, 0.00000, 0.00000), FVector3f(0.00000, 0.00000, 0.00000),
+	FVector3f(0.00000, 0.00000, 0.00000), FVector3f(0.00000, 0.00000, 0.00000), FVector3f(0.00000, 0.00000, 0.00000),
+	FVector3f(0.00000, 0.00000, 0.00000), FVector3f(0.00000, 0.00000, 0.00000), FVector3f(0.00000, 0.00000, 0.00000),
+	FVector3f(0.00000, 0.00000, 0.00000), FVector3f(0.00000, 0.00000, 0.00000), FVector3f(0.00000, 0.00000, 0.00000),
+	FVector3f(0.00000, 0.00000, 0.00000), FVector3f(0.00000, 0.00000, 0.00000), FVector3f(0.00000, 0.00000, 0.00000),
+	FVector3f(0.00000, 0.00000, 0.00000), FVector3f(0.00000, 0.00000, 0.00000), FVector3f(0.00000, 0.00000, 0.00000),
+	FVector3f(0.00000, 0.00000, 0.00000), FVector3f(0.00000, 0.00000, 0.00000), FVector3f(0.00000, 0.00000, 0.00000),
+	FVector3f(0.00000, 0.00000, 0.00000), FVector3f(0.00000, 0.00000, 0.00000), FVector3f(0.00000, 0.00000, 0.00000),
+	FVector3f(0.00000, 0.00000, 0.00000), FVector3f(0.00000, 0.00000, 0.00000)
+},
+{
+	FVector3f(0.01941, -0.00232, 0.05839), FVector3f(0.01741, -0.00581, 0.00071), FVector3f(0.04610, -0.00769, 0.00007),
+	FVector3f(0.10367, -0.00740, 0.00002), FVector3f(0.06244, -0.02445, 0.00000), FVector3f(0.23927, -0.00242, 0.00000),
+	FVector3f(0.00000, 0.00000, 0.00000), FVector3f(0.00000, 0.00000, 0.00000), FVector3f(0.00000, 0.00000, 0.00000),
+	FVector3f(0.00000, 0.00000, 0.00000), FVector3f(0.00000, 0.00000, 0.00000), FVector3f(0.00000, 0.00000, 0.00000),
+	FVector3f(0.00000, 0.00000, 0.00000), FVector3f(0.00000, 0.00000, 0.00000), FVector3f(0.00000, 0.00000, 0.00000),
+	FVector3f(0.00000, 0.00000, 0.00000), FVector3f(0.00000, 0.00000, 0.00000), FVector3f(0.00000, 0.00000, 0.00000),
+	FVector3f(0.00000, 0.00000, 0.00000), FVector3f(0.00000, 0.00000, 0.00000), FVector3f(0.00000, 0.00000, 0.00000),
+	FVector3f(0.00000, 0.00000, 0.00000), FVector3f(0.00000, 0.00000, 0.00000), FVector3f(0.00000, 0.00000, 0.00000),
+	FVector3f(0.00000, 0.00000, 0.00000), FVector3f(0.00000, 0.00000, 0.00000), FVector3f(0.00000, 0.00000, 0.00000),
+	FVector3f(0.00000, 0.00000, 0.00000), FVector3f(0.00000, 0.00000, 0.00000), FVector3f(0.00000, 0.00000, 0.00000),
+	FVector3f(0.00000, 0.00000, 0.00000), FVector3f(0.00000, 0.00000, 0.00000)
+},
+{
+	FVector3f(0.01927, -0.01424, 0.38834), FVector3f(0.01895, -0.00218, 0.09768), FVector3f(0.03002, -0.00194, 0.01072),
+	FVector3f(0.03912, -0.00384, 0.00150), FVector3f(0.04938, -0.00668, 0.00039), FVector3f(0.05239, -0.01107, 0.00012),
+	FVector3f(0.06018, -0.00746, 0.00006), FVector3f(0.06520, -0.01591, 0.00003), FVector3f(0.08253, -0.01052, 0.00002),
+	FVector3f(0.21093, -0.01495, 0.00002), FVector3f(0.12785, -0.01530, 0.00001), FVector3f(0.19030, -0.01428, 0.00001),
+	FVector3f(0.15254, -0.01276, 0.00000), FVector3f(0.16585, -0.02071, 0.00000), FVector3f(0.00000, 0.00000, 0.00000),
+	FVector3f(0.00000, 0.00000, 0.00000), FVector3f(0.00000, 0.00000, 0.00000), FVector3f(0.00000, 0.00000, 0.00000),
+	FVector3f(0.00000, 0.00000, 0.00000), FVector3f(0.00000, 0.00000, 0.00000), FVector3f(0.00000, 0.00000, 0.00000),
+	FVector3f(0.00000, 0.00000, 0.00000), FVector3f(0.00000, 0.00000, 0.00000), FVector3f(0.00000, 0.00000, 0.00000),
+	FVector3f(0.00000, 0.00000, 0.00000), FVector3f(0.00000, 0.00000, 0.00000), FVector3f(0.00000, 0.00000, 0.00000),
+	FVector3f(0.00000, 0.00000, 0.00000), FVector3f(0.00000, 0.00000, 0.00000), FVector3f(0.00000, 0.00000, 0.00000),
+	FVector3f(0.00000, 0.00000, 0.00000), FVector3f(0.00000, 0.00000, 0.00000)
+},
+{
+	FVector3f(0.03084, -0.04909, 0.55348), FVector3f(0.03764, -0.00710, 0.29827), FVector3f(0.03952, -0.00236, 0.11755),
+	FVector3f(0.04092, -0.00201, 0.03677), FVector3f(0.04433, -0.00298, 0.00983), FVector3f(0.05014, -0.00546, 0.00288),
+	FVector3f(0.05570, -0.00834, 0.00101), FVector3f(0.06215, -0.01121, 0.00046), FVector3f(0.06660, -0.01294, 0.00023),
+	FVector3f(0.07902, -0.01692, 0.00014), FVector3f(0.10099, -0.01639, 0.00010), FVector3f(0.10794, -0.01738, 0.00006),
+	FVector3f(0.10632, -0.02032, 0.00004), FVector3f(0.12623, -0.01947, 0.00003), FVector3f(0.13931, -0.02354, 0.00002),
+	FVector3f(0.15353, -0.02910, 0.00002), FVector3f(0.16109, -0.02565, 0.00001), FVector3f(0.14583, -0.02903, 0.00001),
+	FVector3f(0.27891, -0.03066, 0.00001), FVector3f(0.22622, -0.03044, 0.00001), FVector3f(0.18932, -0.04045, 0.00001),
+	FVector3f(0.20219, -0.03226, 0.00000), FVector3f(0.30269, -0.03443, 0.00000), FVector3f(0.38379, -0.03023, 0.00000),
+	FVector3f(0.39038, -0.03610, 0.00000), FVector3f(0.46310, -0.02022, 0.00000), FVector3f(0.44663, -0.02590, 0.00000),
+	FVector3f(0.00000, 0.00000, 0.00000), FVector3f(0.00000, 0.00000, 0.00000), FVector3f(0.00000, 0.00000, 0.00000),
+	FVector3f(0.00000, 0.00000, 0.00000), FVector3f(0.00000, 0.00000, 0.00000)
+},
+{
+	FVector3f(0.04118, -0.10668, 0.63273), FVector3f(0.05152, -0.02772, 0.42999), FVector3f(0.05724, -0.00717, 0.24931),
+	FVector3f(0.05863, -0.00421, 0.13040), FVector3f(0.05952, -0.00351, 0.05910), FVector3f(0.06149, -0.00399, 0.02430),
+	FVector3f(0.06448, -0.00522, 0.00981), FVector3f(0.07004, -0.00676, 0.00433), FVector3f(0.07774, -0.00866, 0.00200),
+	FVector3f(0.08632, -0.01099, 0.00105), FVector3f(0.09629, -0.01332, 0.00062), FVector3f(0.10592, -0.01590, 0.00039),
+	FVector3f(0.10718, -0.01724, 0.00024), FVector3f(0.12207, -0.02041, 0.00018), FVector3f(0.13413, -0.02237, 0.00013),
+	FVector3f(0.13702, -0.02503, 0.00009), FVector3f(0.15294, -0.02664, 0.00008), FVector3f(0.15121, -0.02803, 0.00006),
+	FVector3f(0.17652, -0.03188, 0.00005), FVector3f(0.19532, -0.03147, 0.00004), FVector3f(0.20831, -0.03346, 0.00003),
+	FVector3f(0.19762, -0.03476, 0.00002), FVector3f(0.24202, -0.03464, 0.00002), FVector3f(0.32995, -0.03125, 0.00002),
+	FVector3f(0.30857, -0.03303, 0.00002), FVector3f(0.39596, -0.03009, 0.00002), FVector3f(0.38346, -0.03198, 0.00001),
+	FVector3f(0.42503, -0.02518, 0.00001), FVector3f(0.41592, -0.03195, 0.00001), FVector3f(0.42512, -0.01668, 0.00001),
+	FVector3f(0.36714, -0.02978, 0.00001), FVector3f(0.46502, -0.00394, 0.00001)
+},
+{
+	FVector3f(0.05088, -0.16006, 0.67021), FVector3f(0.06485, -0.05552, 0.50797), FVector3f(0.07697, -0.01274, 0.34475),
+	FVector3f(0.08334, -0.00932, 0.22264), FVector3f(0.08654, -0.00582, 0.13300), FVector3f(0.08950, -0.00507, 0.07375),
+	FVector3f(0.09263, -0.00531, 0.03867), FVector3f(0.09711, -0.00616, 0.02006), FVector3f(0.10113, -0.00792, 0.01032),
+	FVector3f(0.10913, -0.00949, 0.00570), FVector3f(0.11586, -0.01167, 0.00324), FVector3f(0.12425, -0.01421, 0.00197),
+	FVector3f(0.12986, -0.01635, 0.00122), FVector3f(0.13530, -0.01858, 0.00080), FVector3f(0.15037, -0.02237, 0.00059),
+	FVector3f(0.15165, -0.02450, 0.00040), FVector3f(0.15858, -0.02714, 0.00030), FVector3f(0.16660, -0.03145, 0.00022),
+	FVector3f(0.18051, -0.03321, 0.00018), FVector3f(0.18022, -0.03295, 0.00013), FVector3f(0.19896, -0.03492, 0.00012),
+	FVector3f(0.21095, -0.03365, 0.00009), FVector3f(0.21862, -0.03733, 0.00008), FVector3f(0.23861, -0.03930, 0.00007),
+	FVector3f(0.25384, -0.03879, 0.00006), FVector3f(0.27394, -0.03580, 0.00005), FVector3f(0.28563, -0.04089, 0.00004),
+	FVector3f(0.29160, -0.03604, 0.00003), FVector3f(0.29300, -0.03863, 0.00003), FVector3f(0.33458, -0.03575, 0.00002),
+	FVector3f(0.36514, -0.02621, 0.00002), FVector3f(0.38746, 0.00124, 0.00002)
+},
+{
+	FVector3f(0.06054, -0.18133, 0.68765), FVector3f(0.07710, -0.08871, 0.55387), FVector3f(0.09569, -0.02800, 0.41042),
+	FVector3f(0.10917, -0.01456, 0.29613), FVector3f(0.11813, -0.01052, 0.20481), FVector3f(0.12469, -0.00904, 0.13463),
+	FVector3f(0.12958, -0.00864, 0.08432), FVector3f(0.13406, -0.00898, 0.05135), FVector3f(0.13801, -0.01019, 0.03073),
+	FVector3f(0.14324, -0.01110, 0.01861), FVector3f(0.14801, -0.01354, 0.01132), FVector3f(0.15359, -0.01614, 0.00706),
+	FVector3f(0.15945, -0.01922, 0.00454), FVector3f(0.16688, -0.02116, 0.00304), FVector3f(0.17552, -0.02363, 0.00212),
+	FVector3f(0.17956, -0.02606, 0.00145), FVector3f(0.18275, -0.02887, 0.00102), FVector3f(0.19293, -0.03270, 0.00078),
+	FVector3f(0.20081, -0.03492, 0.00059), FVector3f(0.20817, -0.03607, 0.00046), FVector3f(0.21658, -0.03714, 0.00036),
+	FVector3f(0.22866, -0.03827, 0.00030), FVector3f(0.23912, -0.03832, 0.00025), FVector3f(0.24736, -0.03918, 0.00020),
+	FVector3f(0.26573, -0.04071, 0.00017), FVector3f(0.26821, -0.04382, 0.00013), FVector3f(0.28767, -0.03860, 0.00012),
+	FVector3f(0.30592, -0.03690, 0.00010), FVector3f(0.31228, -0.04032, 0.00008), FVector3f(0.35297, -0.03136, 0.00008),
+	FVector3f(0.35570, -0.02365, 0.00006), FVector3f(0.37077, -0.00281, 0.00006)
+},
+{
+	FVector3f(0.07075, -0.18042, 0.69478), FVector3f(0.08974, -0.10806, 0.58263), FVector3f(0.11306, -0.04959, 0.45689),
+	FVector3f(0.13241, -0.02418, 0.35318), FVector3f(0.14704, -0.01632, 0.26617), FVector3f(0.15763, -0.01333, 0.19342),
+	FVector3f(0.16591, -0.01204, 0.13650), FVector3f(0.17258, -0.01177, 0.09391), FVector3f(0.17781, -0.01198, 0.06299),
+	FVector3f(0.18234, -0.01284, 0.04183), FVector3f(0.18700, -0.01408, 0.02782), FVector3f(0.19101, -0.01571, 0.01845),
+	FVector3f(0.19657, -0.01794, 0.01259), FVector3f(0.20165, -0.02041, 0.00864), FVector3f(0.20731, -0.02235, 0.00604),
+	FVector3f(0.21150, -0.02516, 0.00421), FVector3f(0.21692, -0.02822, 0.00303), FVector3f(0.22536, -0.03168, 0.00228),
+	FVector3f(0.23235, -0.03631, 0.00172), FVector3f(0.23720, -0.03840, 0.00129), FVector3f(0.24295, -0.04024, 0.00099),
+	FVector3f(0.25154, -0.04645, 0.00079), FVector3f(0.26200, -0.04435, 0.00066), FVector3f(0.26907, -0.04644, 0.00052),
+	FVector3f(0.28040, -0.04369, 0.00044), FVector3f(0.28922, -0.05007, 0.00034), FVector3f(0.30452, -0.04809, 0.00029),
+	FVector3f(0.31567, -0.04719, 0.00024), FVector3f(0.33294, -0.04179, 0.00021), FVector3f(0.35084, -0.03537, 0.00019),
+	FVector3f(0.37226, -0.02633, 0.00017), FVector3f(0.37956, 0.00196, 0.00013)
+},
+{
+	FVector3f(0.08222, -0.17531, 0.69794), FVector3f(0.10394, -0.13135, 0.60100), FVector3f(0.13034, -0.08092, 0.49085),
+	FVector3f(0.15361, -0.04065, 0.39802), FVector3f(0.17244, -0.02431, 0.31773), FVector3f(0.18620, -0.01864, 0.24683),
+	FVector3f(0.19671, -0.01620, 0.18787), FVector3f(0.20509, -0.01498, 0.13964), FVector3f(0.21199, -0.01460, 0.10220),
+	FVector3f(0.21726, -0.01487, 0.07319), FVector3f(0.22245, -0.01549, 0.05252), FVector3f(0.22702, -0.01665, 0.03747),
+	FVector3f(0.23174, -0.01802, 0.02685), FVector3f(0.23571, -0.01985, 0.01915), FVector3f(0.23966, -0.02185, 0.01375),
+	FVector3f(0.24384, -0.02424, 0.00996), FVector3f(0.24877, -0.02636, 0.00733), FVector3f(0.25548, -0.02871, 0.00552),
+	FVector3f(0.26047, -0.03133, 0.00415), FVector3f(0.26863, -0.03455, 0.00322), FVector3f(0.27404, -0.03705, 0.00247),
+	FVector3f(0.28088, -0.03931, 0.00194), FVector3f(0.29010, -0.04546, 0.00157), FVector3f(0.29704, -0.04797, 0.00126),
+	FVector3f(0.30559, -0.04990, 0.00102), FVector3f(0.31519, -0.04903, 0.00082), FVector3f(0.32721, -0.04842, 0.00069),
+	FVector3f(0.33828, -0.04495, 0.00059), FVector3f(0.35424, -0.04310, 0.00048), FVector3f(0.36868, -0.03925, 0.00042),
+	FVector3f(0.38570, -0.02709, 0.00037), FVector3f(0.39707, 0.00103, 0.00030)
+},
+{
+	FVector3f(0.09597, -0.17397, 0.69862), FVector3f(0.11972, -0.15003, 0.61474), FVector3f(0.14796, -0.10394, 0.51662),
+	FVector3f(0.17350, -0.06278, 0.43321), FVector3f(0.19481, -0.03698, 0.35989), FVector3f(0.21136, -0.02667, 0.29352),
+	FVector3f(0.22369, -0.02186, 0.23483), FVector3f(0.23330, -0.01962, 0.18491), FVector3f(0.24109, -0.01854, 0.14298),
+	FVector3f(0.24787, -0.01810, 0.10960), FVector3f(0.25359, -0.01829, 0.08309), FVector3f(0.25878, -0.01882, 0.06268),
+	FVector3f(0.26320, -0.01990, 0.04689), FVector3f(0.26831, -0.02103, 0.03539), FVector3f(0.27305, -0.02243, 0.02666),
+	FVector3f(0.27810, -0.02425, 0.02015), FVector3f(0.28206, -0.02611, 0.01519), FVector3f(0.28734, -0.02809, 0.01160),
+	FVector3f(0.29228, -0.03045, 0.00888), FVector3f(0.29719, -0.03211, 0.00683), FVector3f(0.30256, -0.03413, 0.00531),
+	FVector3f(0.30857, -0.03666, 0.00417), FVector3f(0.31652, -0.03882, 0.00334), FVector3f(0.32438, -0.04150, 0.00270),
+	FVector3f(0.33291, -0.04323, 0.00220), FVector3f(0.34012, -0.04376, 0.00178), FVector3f(0.35087, -0.04590, 0.00148),
+	FVector3f(0.36243, -0.04591, 0.00123), FVector3f(0.37467, -0.04202, 0.00101), FVector3f(0.38986, -0.03859, 0.00087),
+	FVector3f(0.40394, -0.03046, 0.00073), FVector3f(0.41719, 0.00025, 0.00061)
+},
+{
+	FVector3f(0.11173, -0.17686, 0.69840), FVector3f(0.13694, -0.16229, 0.62456), FVector3f(0.16797, -0.12071, 0.53780),
+	FVector3f(0.19358, -0.08778, 0.46227), FVector3f(0.21571, -0.05746, 0.39486), FVector3f(0.23427, -0.03882, 0.33234),
+	FVector3f(0.24900, -0.03044, 0.27724), FVector3f(0.26011, -0.02620, 0.22722), FVector3f(0.26884, -0.02412, 0.18345),
+	FVector3f(0.27668, -0.02290, 0.14737), FVector3f(0.28326, -0.02247, 0.11687), FVector3f(0.28928, -0.02235, 0.09223),
+	FVector3f(0.29445, -0.02283, 0.07219), FVector3f(0.29932, -0.02349, 0.05639), FVector3f(0.30454, -0.02444, 0.04424),
+	FVector3f(0.30943, -0.02562, 0.03460), FVector3f(0.31431, -0.02725, 0.02709), FVector3f(0.31861, -0.02859, 0.02113),
+	FVector3f(0.32326, -0.03047, 0.01656), FVector3f(0.32881, -0.03199, 0.01309), FVector3f(0.33479, -0.03417, 0.01041),
+	FVector3f(0.34094, -0.03618, 0.00831), FVector3f(0.34705, -0.03771, 0.00665), FVector3f(0.35341, -0.03920, 0.00534),
+	FVector3f(0.36079, -0.04070, 0.00434), FVector3f(0.36863, -0.04138, 0.00355), FVector3f(0.37512, -0.04066, 0.00288),
+	FVector3f(0.38607, -0.04125, 0.00241), FVector3f(0.39611, -0.03916, 0.00200), FVector3f(0.41085, -0.03764, 0.00170),
+	FVector3f(0.42341, -0.02956, 0.00142), FVector3f(0.43959, 0.00020, 0.00123)
+},
+{
+	FVector3f(0.12869, -0.17952, 0.69766), FVector3f(0.15653, -0.16921, 0.63215), FVector3f(0.18847, -0.13634, 0.55391),
+	FVector3f(0.21397, -0.10854, 0.48615), FVector3f(0.23632, -0.08164, 0.42355), FVector3f(0.25652, -0.05663, 0.36707),
+	FVector3f(0.27273, -0.04256, 0.31348), FVector3f(0.28564, -0.03522, 0.26549), FVector3f(0.29568, -0.03124, 0.22185),
+	FVector3f(0.30413, -0.02895, 0.18403), FVector3f(0.31142, -0.02760, 0.15141), FVector3f(0.31801, -0.02688, 0.12380),
+	FVector3f(0.32424, -0.02653, 0.10104), FVector3f(0.32941, -0.02669, 0.08149), FVector3f(0.33479, -0.02714, 0.06598),
+	FVector3f(0.33933, -0.02784, 0.05291), FVector3f(0.34444, -0.02878, 0.04275), FVector3f(0.34892, -0.02988, 0.03428),
+	FVector3f(0.35415, -0.03136, 0.02771), FVector3f(0.35873, -0.03261, 0.02222), FVector3f(0.36457, -0.03429, 0.01807),
+	FVector3f(0.36975, -0.03567, 0.01456), FVector3f(0.37684, -0.03742, 0.01195), FVector3f(0.38258, -0.03792, 0.00969),
+	FVector3f(0.39038, -0.03954, 0.00798), FVector3f(0.39755, -0.04002, 0.00653), FVector3f(0.40428, -0.04014, 0.00534),
+	FVector3f(0.41192, -0.03889, 0.00441), FVector3f(0.42141, -0.03739, 0.00368), FVector3f(0.43074, -0.03386, 0.00307),
+	FVector3f(0.44659, -0.02810, 0.00264), FVector3f(0.46013, 0.00013, 0.00224)
+},
+{
+	FVector3f(0.14693, -0.17953, 0.69641), FVector3f(0.17828, -0.17294, 0.63746), FVector3f(0.20991, -0.14837, 0.56861),
+	FVector3f(0.23513, -0.12458, 0.50592), FVector3f(0.25809, -0.10127, 0.44888), FVector3f(0.27850, -0.07912, 0.39646),
+	FVector3f(0.29576, -0.05884, 0.34628), FVector3f(0.30999, -0.04715, 0.30001), FVector3f(0.32128, -0.04041, 0.25727),
+	FVector3f(0.33053, -0.03637, 0.21872), FVector3f(0.33854, -0.03401, 0.18530), FVector3f(0.34549, -0.03244, 0.15569),
+	FVector3f(0.35212, -0.03165, 0.13073), FVector3f(0.35806, -0.03113, 0.10901), FVector3f(0.36371, -0.03104, 0.09059),
+	FVector3f(0.36901, -0.03134, 0.07509), FVector3f(0.37416, -0.03173, 0.06205), FVector3f(0.37907, -0.03225, 0.05112),
+	FVector3f(0.38378, -0.03304, 0.04199), FVector3f(0.38887, -0.03406, 0.03458), FVector3f(0.39366, -0.03515, 0.02839),
+	FVector3f(0.39953, -0.03594, 0.02356), FVector3f(0.40534, -0.03728, 0.01946), FVector3f(0.41134, -0.03825, 0.01604),
+	FVector3f(0.41832, -0.03910, 0.01337), FVector3f(0.42583, -0.03960, 0.01115), FVector3f(0.43323, -0.03950, 0.00925),
+	FVector3f(0.44084, -0.03877, 0.00766), FVector3f(0.44897, -0.03647, 0.00639), FVector3f(0.45832, -0.03283, 0.00536),
+	FVector3f(0.47095, -0.02659, 0.00456), FVector3f(0.48340, -0.00002, 0.00390)
+},
+{
+	FVector3f(0.16755, -0.17802, 0.69548), FVector3f(0.20139, -0.17499, 0.64297), FVector3f(0.23191, -0.15708, 0.57917),
+	FVector3f(0.25781, -0.13619, 0.52314), FVector3f(0.28103, -0.11601, 0.47074), FVector3f(0.30124, -0.09763, 0.42115),
+	FVector3f(0.31874, -0.07926, 0.37424), FVector3f(0.33379, -0.06262, 0.33015), FVector3f(0.34648, -0.05277, 0.29017),
+	FVector3f(0.35662, -0.04616, 0.25188), FVector3f(0.36536, -0.04217, 0.21825), FVector3f(0.37289, -0.03952, 0.18762),
+	FVector3f(0.37956, -0.03771, 0.16043), FVector3f(0.38585, -0.03664, 0.13705), FVector3f(0.39167, -0.03576, 0.11638),
+	FVector3f(0.39736, -0.03561, 0.09898), FVector3f(0.40272, -0.03531, 0.08360), FVector3f(0.40787, -0.03537, 0.07050),
+	FVector3f(0.41311, -0.03577, 0.05948), FVector3f(0.41811, -0.03606, 0.04991), FVector3f(0.42324, -0.03657, 0.04182),
+	FVector3f(0.42871, -0.03725, 0.03523), FVector3f(0.43459, -0.03790, 0.02974), FVector3f(0.44015, -0.03830, 0.02485),
+	FVector3f(0.44634, -0.03850, 0.02085), FVector3f(0.45322, -0.03897, 0.01764), FVector3f(0.46073, -0.03886, 0.01489),
+	FVector3f(0.46815, -0.03764, 0.01243), FVector3f(0.47699, -0.03576, 0.01052), FVector3f(0.48579, -0.03199, 0.00882),
+	FVector3f(0.49728, -0.02678, 0.00762), FVector3f(0.50776, -0.00003, 0.00638)
+},
+{
+	FVector3f(0.19101, -0.17570, 0.69518), FVector3f(0.22471, -0.17630, 0.64677), FVector3f(0.25559, -0.16233, 0.58896),
+	FVector3f(0.28173, -0.14459, 0.53726), FVector3f(0.30495, -0.12739, 0.48924), FVector3f(0.32541, -0.11162, 0.44375),
+	FVector3f(0.34337, -0.09667, 0.40045), FVector3f(0.35848, -0.08188, 0.35844), FVector3f(0.37152, -0.06821, 0.31901),
+	FVector3f(0.38267, -0.05887, 0.28296), FVector3f(0.39205, -0.05240, 0.24905), FVector3f(0.40013, -0.04793, 0.21813),
+	FVector3f(0.40731, -0.04499, 0.19041), FVector3f(0.41385, -0.04300, 0.16557), FVector3f(0.41996, -0.04151, 0.14360),
+	FVector3f(0.42572, -0.04033, 0.12396), FVector3f(0.43136, -0.03975, 0.10719), FVector3f(0.43675, -0.03947, 0.09240),
+	FVector3f(0.44210, -0.03943, 0.07951), FVector3f(0.44729, -0.03917, 0.06811), FVector3f(0.45254, -0.03889, 0.05809),
+	FVector3f(0.45786, -0.03887, 0.04960), FVector3f(0.46336, -0.03909, 0.04242), FVector3f(0.46904, -0.03903, 0.03607),
+	FVector3f(0.47514, -0.03922, 0.03098), FVector3f(0.48162, -0.03917, 0.02650), FVector3f(0.48829, -0.03827, 0.02249),
+	FVector3f(0.49548, -0.03702, 0.01911), FVector3f(0.50339, -0.03465, 0.01630), FVector3f(0.51184, -0.03118, 0.01385),
+	FVector3f(0.52150, -0.02485, 0.01187), FVector3f(0.53151, -0.00021, 0.01004)
+},
+{
+	FVector3f(0.21608, -0.17314, 0.69404), FVector3f(0.24910, -0.17599, 0.65077), FVector3f(0.28069, -0.16514, 0.59738),
+	FVector3f(0.30724, -0.15021, 0.55038), FVector3f(0.33032, -0.13557, 0.50616), FVector3f(0.35096, -0.12161, 0.46415),
+	FVector3f(0.36879, -0.10876, 0.42260), FVector3f(0.38457, -0.09665, 0.38393), FVector3f(0.39755, -0.08454, 0.34548),
+	FVector3f(0.40917, -0.07412, 0.31090), FVector3f(0.41915, -0.06463, 0.27757), FVector3f(0.42790, -0.05848, 0.24745),
+	FVector3f(0.43555, -0.05371, 0.21932), FVector3f(0.44253, -0.05046, 0.19400), FVector3f(0.44886, -0.04821, 0.17119),
+	FVector3f(0.45491, -0.04588, 0.15001), FVector3f(0.46058, -0.04438, 0.13153), FVector3f(0.46590, -0.04412, 0.11577),
+	FVector3f(0.47137, -0.04304, 0.10099), FVector3f(0.47670, -0.04215, 0.08803), FVector3f(0.48207, -0.04161, 0.07655),
+	FVector3f(0.48730, -0.04144, 0.06667), FVector3f(0.49275, -0.04059, 0.05763), FVector3f(0.49832, -0.04032, 0.05000),
+	FVector3f(0.50411, -0.03974, 0.04326), FVector3f(0.51004, -0.03905, 0.03745), FVector3f(0.51646, -0.03800, 0.03233),
+	FVector3f(0.52320, -0.03639, 0.02793), FVector3f(0.53028, -0.03436, 0.02410), FVector3f(0.53801, -0.03031, 0.02070),
+	FVector3f(0.54612, -0.02493, 0.01799), FVector3f(0.55540, 0.00004, 0.01521)
+},
+{
+	FVector3f(0.24162, -0.17052, 0.69373), FVector3f(0.27454, -0.17439, 0.65370), FVector3f(0.30746, -0.16579, 0.60556),
+	FVector3f(0.33408, -0.15358, 0.56229), FVector3f(0.35697, -0.14112, 0.52099), FVector3f(0.37775, -0.12867, 0.48220),
+	FVector3f(0.39587, -0.11735, 0.44396), FVector3f(0.41165, -0.10700, 0.40686), FVector3f(0.42539, -0.09729, 0.37145),
+	FVector3f(0.43716, -0.08817, 0.33754), FVector3f(0.44735, -0.07945, 0.30538), FVector3f(0.45639, -0.07137, 0.27528),
+	FVector3f(0.46453, -0.06390, 0.24688), FVector3f(0.47188, -0.05902, 0.22137), FVector3f(0.47840, -0.05627, 0.19867),
+	FVector3f(0.48473, -0.05274, 0.17677), FVector3f(0.49060, -0.05014, 0.15718), FVector3f(0.49627, -0.04809, 0.13948),
+	FVector3f(0.50156, -0.04704, 0.12389), FVector3f(0.50672, -0.04617, 0.10988), FVector3f(0.51205, -0.04471, 0.09688),
+	FVector3f(0.51720, -0.04402, 0.08560), FVector3f(0.52258, -0.04284, 0.07528), FVector3f(0.52812, -0.04140, 0.06603),
+	FVector3f(0.53331, -0.04084, 0.05819), FVector3f(0.53939, -0.03922, 0.05082), FVector3f(0.54508, -0.03852, 0.04467),
+	FVector3f(0.55110, -0.03648, 0.03901), FVector3f(0.55773, -0.03362, 0.03404), FVector3f(0.56432, -0.02977, 0.02969),
+	FVector3f(0.57204, -0.02257, 0.02565), FVector3f(0.58012, 0.00013, 0.02208)
+},
+{
+	FVector3f(0.26781, -0.16726, 0.69374), FVector3f(0.30184, -0.17130, 0.65782), FVector3f(0.33546, -0.16494, 0.61300),
+	FVector3f(0.36236, -0.15479, 0.57350), FVector3f(0.38543, -0.14389, 0.53560), FVector3f(0.40587, -0.13313, 0.49892),
+	FVector3f(0.42394, -0.12315, 0.46321), FVector3f(0.43980, -0.11396, 0.42810), FVector3f(0.45355, -0.10555, 0.39357),
+	FVector3f(0.46600, -0.09751, 0.36179), FVector3f(0.47658, -0.08994, 0.33055), FVector3f(0.48605, -0.08303, 0.30167),
+	FVector3f(0.49437, -0.07661, 0.27438), FVector3f(0.50188, -0.07012, 0.24851), FVector3f(0.50886, -0.06435, 0.22461),
+	FVector3f(0.51521, -0.06079, 0.20331), FVector3f(0.52114, -0.05751, 0.18338), FVector3f(0.52686, -0.05490, 0.16512),
+	FVector3f(0.53234, -0.05237, 0.14824), FVector3f(0.53760, -0.05016, 0.13281), FVector3f(0.54305, -0.04763, 0.11850),
+	FVector3f(0.54802, -0.04636, 0.10604), FVector3f(0.55334, -0.04452, 0.09447), FVector3f(0.55850, -0.04311, 0.08417),
+	FVector3f(0.56350, -0.04186, 0.07499), FVector3f(0.56858, -0.04059, 0.06675), FVector3f(0.57426, -0.03828, 0.05906),
+	FVector3f(0.57961, -0.03642, 0.05240), FVector3f(0.58579, -0.03293, 0.04618), FVector3f(0.59207, -0.02837, 0.04061),
+	FVector3f(0.59839, -0.02183, 0.03571), FVector3f(0.60588, 0.00014, 0.03099)
+},
+{
+	FVector3f(0.29552, -0.16307, 0.69478), FVector3f(0.33034, -0.16735, 0.66097), FVector3f(0.36464, -0.16269, 0.62010),
+	FVector3f(0.39162, -0.15439, 0.58318), FVector3f(0.41468, -0.14491, 0.54812), FVector3f(0.43479, -0.13568, 0.51374),
+	FVector3f(0.45273, -0.12677, 0.48033), FVector3f(0.46884, -0.11827, 0.44784), FVector3f(0.48278, -0.11077, 0.41543),
+	FVector3f(0.49535, -0.10366, 0.38456), FVector3f(0.50635, -0.09707, 0.35481), FVector3f(0.51611, -0.09088, 0.32656),
+	FVector3f(0.52482, -0.08505, 0.29983), FVector3f(0.53261, -0.07920, 0.27402), FVector3f(0.53981, -0.07391, 0.25049),
+	FVector3f(0.54631, -0.06982, 0.22916), FVector3f(0.55239, -0.06584, 0.20899), FVector3f(0.55818, -0.06111, 0.18974),
+	FVector3f(0.56348, -0.05885, 0.17283), FVector3f(0.56875, -0.05560, 0.15672), FVector3f(0.57406, -0.05216, 0.14161),
+	FVector3f(0.57898, -0.04974, 0.12804), FVector3f(0.58395, -0.04738, 0.11552), FVector3f(0.58890, -0.04525, 0.10412),
+	FVector3f(0.59356, -0.04361, 0.09392), FVector3f(0.59866, -0.04122, 0.08435), FVector3f(0.60366, -0.03869, 0.07566),
+	FVector3f(0.60857, -0.03626, 0.06785), FVector3f(0.61392, -0.03267, 0.06058), FVector3f(0.61896, -0.02853, 0.05418),
+	FVector3f(0.62446, -0.02146, 0.04812), FVector3f(0.63114, 0.00013, 0.04224)
+},
+{
+	FVector3f(0.32493, -0.15800, 0.69636), FVector3f(0.36044, -0.16239, 0.66479), FVector3f(0.39500, -0.15918, 0.62662),
+	FVector3f(0.42214, -0.15229, 0.59208), FVector3f(0.44504, -0.14422, 0.55948), FVector3f(0.46513, -0.13589, 0.52832),
+	FVector3f(0.48260, -0.12801, 0.49686), FVector3f(0.49819, -0.12058, 0.46577), FVector3f(0.51220, -0.11368, 0.43550),
+	FVector3f(0.52491, -0.10715, 0.40640), FVector3f(0.53607, -0.10121, 0.37791), FVector3f(0.54604, -0.09561, 0.35063),
+	FVector3f(0.55497, -0.09044, 0.32440), FVector3f(0.56304, -0.08553, 0.29963), FVector3f(0.57042, -0.08073, 0.27592),
+	FVector3f(0.57716, -0.07643, 0.25430), FVector3f(0.58338, -0.07244, 0.23411), FVector3f(0.58922, -0.06816, 0.21474),
+	FVector3f(0.59463, -0.06506, 0.19748), FVector3f(0.59983, -0.06156, 0.18085), FVector3f(0.60510, -0.05721, 0.16496),
+	FVector3f(0.60997, -0.05394, 0.15071), FVector3f(0.61473, -0.05071, 0.13745), FVector3f(0.61957, -0.04764, 0.12513),
+	FVector3f(0.62409, -0.04512, 0.11401), FVector3f(0.62852, -0.04243, 0.10365), FVector3f(0.63251, -0.04049, 0.09444),
+	FVector3f(0.63750, -0.03634, 0.08526), FVector3f(0.64231, -0.03217, 0.07697), FVector3f(0.64613, -0.02798, 0.06969),
+	FVector3f(0.65086, -0.02084, 0.06274), FVector3f(0.65651, -0.00002, 0.05589)
+},
+{
+	FVector3f(0.35603, -0.15221, 0.69769), FVector3f(0.39159, -0.15659, 0.66832), FVector3f(0.42682, -0.15414, 0.63306),
+	FVector3f(0.45440, -0.14815, 0.60175), FVector3f(0.47706, -0.14135, 0.57114), FVector3f(0.49652, -0.13430, 0.54130),
+	FVector3f(0.51370, -0.12720, 0.51227), FVector3f(0.52869, -0.12066, 0.48316), FVector3f(0.54234, -0.11418, 0.45520),
+	FVector3f(0.55442, -0.10835, 0.42705), FVector3f(0.56541, -0.10297, 0.39979), FVector3f(0.57520, -0.09810, 0.37306),
+	FVector3f(0.58426, -0.09325, 0.34799), FVector3f(0.59238, -0.08885, 0.32368), FVector3f(0.59985, -0.08462, 0.30051),
+	FVector3f(0.60666, -0.08047, 0.27914), FVector3f(0.61307, -0.07647, 0.25902), FVector3f(0.61890, -0.07278, 0.24017),
+	FVector3f(0.62438, -0.06918, 0.22242), FVector3f(0.62953, -0.06581, 0.20569), FVector3f(0.63458, -0.06234, 0.18950),
+	FVector3f(0.63958, -0.05873, 0.17436), FVector3f(0.64422, -0.05531, 0.16053), FVector3f(0.64882, -0.05182, 0.14758),
+	FVector3f(0.65331, -0.04822, 0.13555), FVector3f(0.65733, -0.04527, 0.12469), FVector3f(0.66150, -0.04160, 0.11431),
+	FVector3f(0.66582, -0.03748, 0.10458), FVector3f(0.66990, -0.03270, 0.09545), FVector3f(0.67332, -0.02788, 0.08731),
+	FVector3f(0.67675, -0.02071, 0.07957), FVector3f(0.68186, 0.00002, 0.07179)
+},
+{
+	FVector3f(0.38874, -0.14547, 0.69990), FVector3f(0.42458, -0.14938, 0.67322), FVector3f(0.46012, -0.14757, 0.63986),
+	FVector3f(0.48761, -0.14257, 0.61072), FVector3f(0.51008, -0.13671, 0.58210), FVector3f(0.52938, -0.13046, 0.55452),
+	FVector3f(0.54611, -0.12415, 0.52757), FVector3f(0.56050, -0.11845, 0.49993), FVector3f(0.57333, -0.11284, 0.47308),
+	FVector3f(0.58470, -0.10773, 0.44623), FVector3f(0.59503, -0.10264, 0.42075), FVector3f(0.60417, -0.09831, 0.39479),
+	FVector3f(0.61278, -0.09386, 0.37078), FVector3f(0.62064, -0.08976, 0.34736), FVector3f(0.62794, -0.08593, 0.32486),
+	FVector3f(0.63473, -0.08203, 0.30393), FVector3f(0.64102, -0.07852, 0.28339), FVector3f(0.64684, -0.07474, 0.26497),
+	FVector3f(0.65232, -0.07134, 0.24680), FVector3f(0.65750, -0.06788, 0.23017), FVector3f(0.66247, -0.06462, 0.21383),
+	FVector3f(0.66730, -0.06134, 0.19836), FVector3f(0.67176, -0.05808, 0.18437), FVector3f(0.67590, -0.05475, 0.17126),
+	FVector3f(0.68040, -0.05106, 0.15847), FVector3f(0.68408, -0.04764, 0.14730), FVector3f(0.68813, -0.04367, 0.13611),
+	FVector3f(0.69178, -0.03955, 0.12594), FVector3f(0.69592, -0.03412, 0.11582), FVector3f(0.69962, -0.02813, 0.10669),
+	FVector3f(0.70287, -0.02016, 0.09818), FVector3f(0.70573, 0.00006, 0.09005)
+},
+{
+	FVector3f(0.42286, -0.13764, 0.70253), FVector3f(0.45893, -0.14091, 0.67808), FVector3f(0.49450, -0.13944, 0.64711),
+	FVector3f(0.52171, -0.13527, 0.61953), FVector3f(0.54404, -0.13016, 0.59306), FVector3f(0.56282, -0.12479, 0.56710),
+	FVector3f(0.57891, -0.11949, 0.54128), FVector3f(0.59290, -0.11426, 0.51588), FVector3f(0.60505, -0.10951, 0.48993),
+	FVector3f(0.61587, -0.10470, 0.46536), FVector3f(0.62546, -0.10029, 0.44072), FVector3f(0.63399, -0.09638, 0.41598),
+	FVector3f(0.64186, -0.09229, 0.39289), FVector3f(0.64911, -0.08880, 0.36972), FVector3f(0.65567, -0.08519, 0.34812),
+	FVector3f(0.66186, -0.08157, 0.32779), FVector3f(0.66775, -0.07846, 0.30734), FVector3f(0.67323, -0.07504, 0.28886),
+	FVector3f(0.67844, -0.07174, 0.27113), FVector3f(0.68333, -0.06833, 0.25452), FVector3f(0.68809, -0.06521, 0.23828),
+	FVector3f(0.69262, -0.06204, 0.22301), FVector3f(0.69697, -0.05878, 0.20869), FVector3f(0.70127, -0.05553, 0.19491),
+	FVector3f(0.70529, -0.05205, 0.18216), FVector3f(0.70896, -0.04825, 0.17048), FVector3f(0.71278, -0.04429, 0.15888),
+	FVector3f(0.71635, -0.03988, 0.14822), FVector3f(0.71997, -0.03495, 0.13785), FVector3f(0.72323, -0.02900, 0.12826),
+	FVector3f(0.72673, -0.02079, 0.11891), FVector3f(0.72944, 0.00004, 0.11016)
+},
+{
+	FVector3f(0.45873, -0.12835, 0.70614), FVector3f(0.49433, -0.13100, 0.68304), FVector3f(0.52949, -0.12981, 0.65400),
+	FVector3f(0.55633, -0.12616, 0.62850), FVector3f(0.57809, -0.12174, 0.60424), FVector3f(0.59634, -0.11714, 0.57985),
+	FVector3f(0.61180, -0.11273, 0.55495), FVector3f(0.62517, -0.10827, 0.53092), FVector3f(0.63669, -0.10407, 0.50671),
+	FVector3f(0.64685, -0.09994, 0.48305), FVector3f(0.65584, -0.09624, 0.45902), FVector3f(0.66376, -0.09244, 0.43645),
+	FVector3f(0.67113, -0.08921, 0.41315), FVector3f(0.67765, -0.08575, 0.39173), FVector3f(0.68370, -0.08256, 0.37068),
+	FVector3f(0.68930, -0.07941, 0.35059), FVector3f(0.69443, -0.07619, 0.33162), FVector3f(0.69947, -0.07326, 0.31286),
+	FVector3f(0.70415, -0.07019, 0.29534), FVector3f(0.70855, -0.06708, 0.27878), FVector3f(0.71287, -0.06409, 0.26263),
+	FVector3f(0.71661, -0.06080, 0.24796), FVector3f(0.72066, -0.05775, 0.23332), FVector3f(0.72451, -0.05456, 0.21962),
+	FVector3f(0.72821, -0.05117, 0.20663), FVector3f(0.73152, -0.04735, 0.19465), FVector3f(0.73501, -0.04351, 0.18279),
+	FVector3f(0.73849, -0.03929, 0.17160), FVector3f(0.74160, -0.03433, 0.16118), FVector3f(0.74481, -0.02845, 0.15104),
+	FVector3f(0.74776, -0.02023, 0.14164), FVector3f(0.75105, 0.00006, 0.13208)
+},
+{
+	FVector3f(0.49554, -0.11754, 0.70978), FVector3f(0.53027, -0.11957, 0.68862), FVector3f(0.56475, -0.11844, 0.66139),
+	FVector3f(0.59075, -0.11545, 0.63717), FVector3f(0.61170, -0.11182, 0.61414), FVector3f(0.62919, -0.10788, 0.59143),
+	FVector3f(0.64398, -0.10410, 0.56846), FVector3f(0.65664, -0.10044, 0.54542), FVector3f(0.66758, -0.09697, 0.52211),
+	FVector3f(0.67700, -0.09338, 0.49989), FVector3f(0.68535, -0.09011, 0.47731), FVector3f(0.69283, -0.08701, 0.45504),
+	FVector3f(0.69943, -0.08398, 0.43352), FVector3f(0.70556, -0.08119, 0.41218), FVector3f(0.71100, -0.07822, 0.39227),
+	FVector3f(0.71624, -0.07555, 0.37245), FVector3f(0.72083, -0.07263, 0.35421), FVector3f(0.72523, -0.06975, 0.33639),
+	FVector3f(0.72925, -0.06691, 0.31941), FVector3f(0.73339, -0.06419, 0.30275), FVector3f(0.73715, -0.06130, 0.28712),
+	FVector3f(0.74070, -0.05852, 0.27199), FVector3f(0.74402, -0.05539, 0.25799), FVector3f(0.74744, -0.05231, 0.24427),
+	FVector3f(0.75049, -0.04891, 0.23147), FVector3f(0.75399, -0.04558, 0.21879), FVector3f(0.75685, -0.04169, 0.20716),
+	FVector3f(0.75973, -0.03736, 0.19611), FVector3f(0.76249, -0.03256, 0.18534), FVector3f(0.76538, -0.02662, 0.17518),
+	FVector3f(0.76785, -0.01865, 0.16553), FVector3f(0.77115, 0.00005, 0.15588)
+},
+{
+	FVector3f(0.53272, -0.10493, 0.71478), FVector3f(0.56610, -0.10643, 0.69461), FVector3f(0.59929, -0.10534, 0.66938),
+	FVector3f(0.62423, -0.10298, 0.64618), FVector3f(0.64430, -0.10006, 0.62412), FVector3f(0.66087, -0.09686, 0.60285),
+	FVector3f(0.67490, -0.09377, 0.58132), FVector3f(0.68693, -0.09079, 0.55910), FVector3f(0.69725, -0.08792, 0.53711),
+	FVector3f(0.70612, -0.08524, 0.51494), FVector3f(0.71373, -0.08232, 0.49425), FVector3f(0.72045, -0.07970, 0.47333),
+	FVector3f(0.72654, -0.07708, 0.45290), FVector3f(0.73215, -0.07461, 0.43265), FVector3f(0.73731, -0.07219, 0.41291),
+	FVector3f(0.74188, -0.06976, 0.39410), FVector3f(0.74597, -0.06728, 0.37624), FVector3f(0.74997, -0.06490, 0.35873),
+	FVector3f(0.75337, -0.06230, 0.34249), FVector3f(0.75709, -0.05978, 0.32632), FVector3f(0.76060, -0.05720, 0.31088),
+	FVector3f(0.76378, -0.05462, 0.29607), FVector3f(0.76659, -0.05180, 0.28230), FVector3f(0.76955, -0.04891, 0.26889),
+	FVector3f(0.77230, -0.04571, 0.25622), FVector3f(0.77514, -0.04252, 0.24377), FVector3f(0.77793, -0.03908, 0.23184),
+	FVector3f(0.78043, -0.03509, 0.22064), FVector3f(0.78282, -0.03041, 0.21005), FVector3f(0.78468, -0.02457, 0.20018),
+	FVector3f(0.78754, -0.01737, 0.19021), FVector3f(0.78995, 0.00006, 0.18089)
+},
+{
+	FVector3f(0.56932, -0.09075, 0.71955), FVector3f(0.60103, -0.09168, 0.70070), FVector3f(0.63267, -0.09087, 0.67652),
+	FVector3f(0.65651, -0.08897, 0.65469), FVector3f(0.67543, -0.08664, 0.63411), FVector3f(0.69104, -0.08420, 0.61379),
+	FVector3f(0.70419, -0.08174, 0.59365), FVector3f(0.71550, -0.07948, 0.57227), FVector3f(0.72506, -0.07710, 0.55184),
+	FVector3f(0.73326, -0.07491, 0.53105), FVector3f(0.74047, -0.07278, 0.51035), FVector3f(0.74675, -0.07074, 0.49002),
+	FVector3f(0.75230, -0.06867, 0.47031), FVector3f(0.75728, -0.06656, 0.45135), FVector3f(0.76209, -0.06453, 0.43250),
+	FVector3f(0.76616, -0.06249, 0.41451), FVector3f(0.76990, -0.06046, 0.39705), FVector3f(0.77334, -0.05835, 0.38038),
+	FVector3f(0.77650, -0.05612, 0.36453), FVector3f(0.77934, -0.05390, 0.34929), FVector3f(0.78214, -0.05159, 0.33447),
+	FVector3f(0.78495, -0.04919, 0.32031), FVector3f(0.78732, -0.04660, 0.30694), FVector3f(0.78980, -0.04409, 0.29376),
+	FVector3f(0.79245, -0.04133, 0.28119), FVector3f(0.79502, -0.03835, 0.26903), FVector3f(0.79706, -0.03500, 0.25765),
+	FVector3f(0.79938, -0.03151, 0.24642), FVector3f(0.80116, -0.02721, 0.23605), FVector3f(0.80364, -0.02228, 0.22580),
+	FVector3f(0.80529, -0.01553, 0.21630), FVector3f(0.80743, 0.00008, 0.20694)
+},
+{
+	FVector3f(0.60446, -0.07481, 0.72541), FVector3f(0.63433, -0.07540, 0.70789), FVector3f(0.66433, -0.07476, 0.68481),
+	FVector3f(0.68668, -0.07333, 0.66418), FVector3f(0.70457, -0.07166, 0.64427), FVector3f(0.71943, -0.06986, 0.62480),
+	FVector3f(0.73164, -0.06808, 0.60530), FVector3f(0.74200, -0.06634, 0.58563), FVector3f(0.75088, -0.06465, 0.56577),
+	FVector3f(0.75865, -0.06301, 0.54574), FVector3f(0.76531, -0.06140, 0.52594), FVector3f(0.77096, -0.05975, 0.50704),
+	FVector3f(0.77604, -0.05817, 0.48812), FVector3f(0.78047, -0.05654, 0.46984), FVector3f(0.78462, -0.05499, 0.45179),
+	FVector3f(0.78846, -0.05337, 0.43440), FVector3f(0.79171, -0.05168, 0.41775), FVector3f(0.79502, -0.04998, 0.40149),
+	FVector3f(0.79763, -0.04815, 0.38649), FVector3f(0.80039, -0.04639, 0.37146), FVector3f(0.80272, -0.04448, 0.35737),
+	FVector3f(0.80522, -0.04252, 0.34354), FVector3f(0.80755, -0.04044, 0.33020), FVector3f(0.80954, -0.03811, 0.31798),
+	FVector3f(0.81129, -0.03571, 0.30606), FVector3f(0.81256, -0.03305, 0.29475), FVector3f(0.81479, -0.03031, 0.28330),
+	FVector3f(0.81667, -0.02709, 0.27265), FVector3f(0.81865, -0.02350, 0.26226), FVector3f(0.82066, -0.01932, 0.25205),
+	FVector3f(0.82209, -0.01368, 0.24262), FVector3f(0.82367, 0.00002, 0.23366)
+},
+{
+	FVector3f(0.63783, -0.05754, 0.73110), FVector3f(0.66580, -0.05784, 0.71450), FVector3f(0.69379, -0.05736, 0.69314),
+	FVector3f(0.71493, -0.05642, 0.67318), FVector3f(0.73184, -0.05529, 0.65409), FVector3f(0.74566, -0.05410, 0.63543),
+	FVector3f(0.75715, -0.05290, 0.61665), FVector3f(0.76684, -0.05171, 0.59776), FVector3f(0.77502, -0.05055, 0.57875),
+	FVector3f(0.78192, -0.04938, 0.56005), FVector3f(0.78802, -0.04824, 0.54142), FVector3f(0.79346, -0.04716, 0.52249),
+	FVector3f(0.79805, -0.04606, 0.50438), FVector3f(0.80183, -0.04480, 0.48754), FVector3f(0.80557, -0.04364, 0.47040),
+	FVector3f(0.80900, -0.04245, 0.45369), FVector3f(0.81218, -0.04128, 0.43725), FVector3f(0.81485, -0.03999, 0.42211),
+	FVector3f(0.81744, -0.03873, 0.40702), FVector3f(0.81962, -0.03727, 0.39305), FVector3f(0.82179, -0.03580, 0.37945),
+	FVector3f(0.82385, -0.03429, 0.36620), FVector3f(0.82598, -0.03268, 0.35338), FVector3f(0.82778, -0.03089, 0.34141),
+	FVector3f(0.82910, -0.02898, 0.33004), FVector3f(0.83055, -0.02688, 0.31907), FVector3f(0.83184, -0.02463, 0.30846),
+	FVector3f(0.83322, -0.02206, 0.29842), FVector3f(0.83509, -0.01925, 0.28818), FVector3f(0.83626, -0.01574, 0.27866),
+	FVector3f(0.83803, -0.01128, 0.26919), FVector3f(0.83936, -0.00001, 0.26038)
+},
+{
+	FVector3f(0.66881, -0.03902, 0.73750), FVector3f(0.69495, -0.03917, 0.72144), FVector3f(0.72127, -0.03890, 0.70116),
+	FVector3f(0.74089, -0.03836, 0.68204), FVector3f(0.75676, -0.03768, 0.66379), FVector3f(0.76971, -0.03701, 0.64562),
+	FVector3f(0.78031, -0.03631, 0.62792), FVector3f(0.78909, -0.03559, 0.61014), FVector3f(0.79688, -0.03488, 0.59193),
+	FVector3f(0.80339, -0.03421, 0.57344), FVector3f(0.80909, -0.03354, 0.55536), FVector3f(0.81375, -0.03284, 0.53793),
+	FVector3f(0.81807, -0.03212, 0.52063), FVector3f(0.82174, -0.03140, 0.50374), FVector3f(0.82535, -0.03066, 0.48732),
+	FVector3f(0.82861, -0.02992, 0.47141), FVector3f(0.83083, -0.02915, 0.45633), FVector3f(0.83320, -0.02828, 0.44181),
+	FVector3f(0.83559, -0.02745, 0.42730), FVector3f(0.83745, -0.02649, 0.41385), FVector3f(0.84005, -0.02553, 0.40059),
+	FVector3f(0.84070, -0.02440, 0.38881), FVector3f(0.84245, -0.02332, 0.37643), FVector3f(0.84401, -0.02210, 0.36506),
+	FVector3f(0.84503, -0.02076, 0.35431), FVector3f(0.84665, -0.01940, 0.34339), FVector3f(0.84777, -0.01782, 0.33324),
+	FVector3f(0.84905, -0.01601, 0.32345), FVector3f(0.84999, -0.01390, 0.31425), FVector3f(0.85072, -0.01143, 0.30523),
+	FVector3f(0.85236, -0.00812, 0.29632), FVector3f(0.85341, 0.00002, 0.28782)
+},
+{
+	FVector3f(0.69740, -0.01972, 0.74339), FVector3f(0.72178, -0.01980, 0.72804), FVector3f(0.74630, -0.01969, 0.70892),
+	FVector3f(0.76475, -0.01946, 0.69082), FVector3f(0.77931, -0.01916, 0.67367), FVector3f(0.79173, -0.01888, 0.65573),
+	FVector3f(0.80171, -0.01860, 0.63826), FVector3f(0.81003, -0.01833, 0.62090), FVector3f(0.81725, -0.01802, 0.60355),
+	FVector3f(0.82283, -0.01768, 0.58677), FVector3f(0.82781, -0.01737, 0.56971), FVector3f(0.83251, -0.01712, 0.55229),
+	FVector3f(0.83628, -0.01681, 0.53567), FVector3f(0.83983, -0.01645, 0.51970), FVector3f(0.84284, -0.01610, 0.50407),
+	FVector3f(0.84559, -0.01577, 0.48879), FVector3f(0.84790, -0.01538, 0.47446), FVector3f(0.85013, -0.01498, 0.46044),
+	FVector3f(0.85214, -0.01455, 0.44710), FVector3f(0.85384, -0.01407, 0.43403), FVector3f(0.85537, -0.01356, 0.42174),
+	FVector3f(0.85665, -0.01304, 0.41005), FVector3f(0.85860, -0.01251, 0.39844), FVector3f(0.85961, -0.01189, 0.38755),
+	FVector3f(0.86039, -0.01119, 0.37756), FVector3f(0.86098, -0.01040, 0.36805), FVector3f(0.86192, -0.00953, 0.35846),
+	FVector3f(0.86278, -0.00857, 0.34911), FVector3f(0.86425, -0.00751, 0.33962), FVector3f(0.86491, -0.00614, 0.33144),
+	FVector3f(0.86581, -0.00433, 0.32300), FVector3f(0.86677, 0.00002, 0.31508)
+},
+{
+	FVector3f(0.72363, 0.00000, 0.74973), FVector3f(0.74617, -0.00002, 0.73527), FVector3f(0.76908, 0.00000, 0.71668),
+	FVector3f(0.78618, -0.00002, 0.69953), FVector3f(0.79981, -0.00003, 0.68291), FVector3f(0.81110, -0.00001, 0.66623),
+	FVector3f(0.82060, -0.00001, 0.64940), FVector3f(0.82853, -0.00004, 0.63212), FVector3f(0.83504, -0.00002, 0.61545),
+	FVector3f(0.84058, -0.00003, 0.59849), FVector3f(0.84521, -0.00005, 0.58226), FVector3f(0.84932, -0.00002, 0.56605),
+	FVector3f(0.85275, -0.00000, 0.55034), FVector3f(0.85605, -0.00001, 0.53461), FVector3f(0.85877, -0.00001, 0.51974),
+	FVector3f(0.86129, -0.00001, 0.50561), FVector3f(0.86320, -0.00005, 0.49191), FVector3f(0.86521, -0.00004, 0.47875),
+	FVector3f(0.86654, -0.00002, 0.46632), FVector3f(0.86846, -0.00002, 0.45401), FVector3f(0.86997, -0.00001, 0.44209),
+	FVector3f(0.87153, 0.00003, 0.43044), FVector3f(0.87288, 0.00002, 0.42005), FVector3f(0.87352, 0.00003, 0.40990),
+	FVector3f(0.87465, 0.00004, 0.39997), FVector3f(0.87549, 0.00003, 0.39069), FVector3f(0.87626, 0.00006, 0.38171),
+	FVector3f(0.87676, 0.00004, 0.37342), FVector3f(0.87714, 0.00006, 0.36523), FVector3f(0.87859, 0.00006, 0.35675),
+	FVector3f(0.87952, 0.00009, 0.34897), FVector3f(0.87958, 0.00003, 0.34187)
+}
+};
+
+const FVector3f Sheen_LTC_Approx[Sheen_LTC_Size][Sheen_LTC_Size] =
+{
+{
+	FVector3f(0.10027, -0.00000, 0.33971), FVector3f(0.10760, -0.00000, 0.35542), FVector3f(0.11991, 0.00001, 0.30888),
+	FVector3f(0.13148, 0.00001, 0.23195), FVector3f(0.14227, 0.00001, 0.15949), FVector3f(0.15231, -0.00000, 0.10356),
+	FVector3f(0.16168, -0.00000, 0.06466), FVector3f(0.17044, 0.00000, 0.03925), FVector3f(0.17867, 0.00001, 0.02334),
+	FVector3f(0.18645, 0.00000, 0.01366), FVector3f(0.19382, -0.00000, 0.00790), FVector3f(0.20084, -0.00001, 0.00452),
+	FVector3f(0.20754, 0.00001, 0.00257), FVector3f(0.21395, 0.00000, 0.00145), FVector3f(0.22011, 0.00000, 0.00081),
+	FVector3f(0.22603, -0.00000, 0.00045), FVector3f(0.23174, -0.00001, 0.00025), FVector3f(0.23726, 0.00000, 0.00014),
+	FVector3f(0.24259, -0.00001, 0.00008), FVector3f(0.24777, -0.00001, 0.00004), FVector3f(0.25279, -0.00001, 0.00002),
+	FVector3f(0.25768, 0.00001, 0.00001), FVector3f(0.26243, 0.00001, 0.00001), FVector3f(0.26707, -0.00000, 0.00000),
+	FVector3f(0.27159, -0.00000, 0.00000), FVector3f(0.27601, -0.00000, 0.00000), FVector3f(0.28033, 0.00000, 0.00000),
+	FVector3f(0.28456, -0.00000, 0.00000), FVector3f(0.28870, -0.00001, 0.00000), FVector3f(0.29276, -0.00000, 0.00000),
+	FVector3f(0.29676, 0.00000, 0.00000), FVector3f(0.30067, -0.00001, 0.00000)
+},
+{
+	FVector3f(0.10068, -0.00013, 0.33844), FVector3f(0.10802, -0.00001, 0.35438), FVector3f(0.12031, -0.00000, 0.30859),
+	FVector3f(0.13190, 0.00001, 0.23230), FVector3f(0.14269, 0.00000, 0.16016), FVector3f(0.15274, 0.00001, 0.10429),
+	FVector3f(0.16211, -0.00000, 0.06529), FVector3f(0.17088, -0.00000, 0.03975), FVector3f(0.17912, -0.00000, 0.02370),
+	FVector3f(0.18691, -0.00001, 0.01391), FVector3f(0.19429, -0.00000, 0.00807), FVector3f(0.20132, -0.00000, 0.00463),
+	FVector3f(0.20803, 0.00001, 0.00264), FVector3f(0.21445, 0.00000, 0.00149), FVector3f(0.22061, 0.00000, 0.00084),
+	FVector3f(0.22655, -0.00000, 0.00047), FVector3f(0.23226, -0.00000, 0.00026), FVector3f(0.23779, -0.00001, 0.00015),
+	FVector3f(0.24314, -0.00002, 0.00008), FVector3f(0.24832, -0.00001, 0.00004), FVector3f(0.25336, -0.00002, 0.00002),
+	FVector3f(0.25824, -0.00003, 0.00001), FVector3f(0.26301, -0.00005, 0.00001), FVector3f(0.26766, -0.00008, 0.00000),
+	FVector3f(0.27220, -0.00012, 0.00000), FVector3f(0.27665, -0.00019, 0.00000), FVector3f(0.28101, -0.00026, 0.00000),
+	FVector3f(0.28532, -0.00041, 0.00000), FVector3f(0.28960, -0.00058, 0.00000), FVector3f(0.29389, -0.00080, 0.00000),
+	FVector3f(0.29830, -0.00096, 0.00000), FVector3f(0.30309, 0.00000, 0.00000)
+},
+{
+	FVector3f(0.09988, -0.00743, 0.33595), FVector3f(0.10928, -0.00078, 0.35135), FVector3f(0.12169, -0.00015, 0.30790),
+	FVector3f(0.13330, -0.00006, 0.23367), FVector3f(0.14412, -0.00003, 0.16253), FVector3f(0.15420, -0.00004, 0.10680),
+	FVector3f(0.16360, -0.00002, 0.06750), FVector3f(0.17240, -0.00004, 0.04148), FVector3f(0.18068, -0.00003, 0.02497),
+	FVector3f(0.18850, -0.00004, 0.01480), FVector3f(0.19591, -0.00005, 0.00867), FVector3f(0.20297, -0.00005, 0.00503),
+	FVector3f(0.20970, -0.00008, 0.00289), FVector3f(0.21616, -0.00012, 0.00165), FVector3f(0.22235, -0.00015, 0.00094),
+	FVector3f(0.22831, -0.00020, 0.00053), FVector3f(0.23407, -0.00028, 0.00030), FVector3f(0.23963, -0.00039, 0.00017),
+	FVector3f(0.24503, -0.00056, 0.00009), FVector3f(0.25028, -0.00084, 0.00005), FVector3f(0.25541, -0.00124, 0.00003),
+	FVector3f(0.26045, -0.00184, 0.00002), FVector3f(0.26545, -0.00274, 0.00001), FVector3f(0.27049, -0.00410, 0.00001),
+	FVector3f(0.27569, -0.00612, 0.00000), FVector3f(0.28127, -0.00908, 0.00000), FVector3f(0.28762, -0.01335, 0.00000),
+	FVector3f(0.29542, -0.01917, 0.00000), FVector3f(0.30593, -0.02649, 0.00000), FVector3f(0.32153, -0.03397, 0.00000),
+	FVector3f(0.34697, -0.03669, 0.00000), FVector3f(0.39704, -0.00000, 0.00000)
+},
+{
+	FVector3f(0.08375, -0.07516, 0.33643), FVector3f(0.10999, -0.00776, 0.34873), FVector3f(0.12392, -0.00160, 0.30771),
+	FVector3f(0.13571, -0.00070, 0.23660), FVector3f(0.14661, -0.00044, 0.16701), FVector3f(0.15676, -0.00034, 0.11147),
+	FVector3f(0.16621, -0.00032, 0.07157), FVector3f(0.17508, -0.00029, 0.04470), FVector3f(0.18341, -0.00033, 0.02736),
+	FVector3f(0.19128, -0.00038, 0.01648), FVector3f(0.19876, -0.00045, 0.00981), FVector3f(0.20587, -0.00057, 0.00579),
+	FVector3f(0.21267, -0.00074, 0.00339), FVector3f(0.21918, -0.00096, 0.00197), FVector3f(0.22544, -0.00129, 0.00114),
+	FVector3f(0.23150, -0.00179, 0.00066), FVector3f(0.23736, -0.00249, 0.00038), FVector3f(0.24308, -0.00352, 0.00022),
+	FVector3f(0.24871, -0.00501, 0.00012), FVector3f(0.25433, -0.00721, 0.00007), FVector3f(0.26003, -0.01042, 0.00004),
+	FVector3f(0.26604, -0.01512, 0.00002), FVector3f(0.27264, -0.02192, 0.00001), FVector3f(0.28039, -0.03158, 0.00001),
+	FVector3f(0.29024, -0.04491, 0.00000), FVector3f(0.30382, -0.06232, 0.00000), FVector3f(0.32395, -0.08307, 0.00000),
+	FVector3f(0.35533, -0.10404, 0.00000), FVector3f(0.40491, -0.11883, 0.00000), FVector3f(0.47816, -0.11902, 0.00000),
+	FVector3f(0.56774, -0.09644, 0.00000), FVector3f(0.66332, -0.00000, 0.00000)
+},
+{
+	FVector3f(0.05655, -0.31167, 0.32420), FVector3f(0.10687, -0.03376, 0.35090), FVector3f(0.12655, -0.00804, 0.31009),
+	FVector3f(0.13914, -0.00363, 0.24233), FVector3f(0.15029, -0.00227, 0.17455), FVector3f(0.16057, -0.00177, 0.11907),
+	FVector3f(0.17014, -0.00156, 0.07820), FVector3f(0.17910, -0.00153, 0.04999), FVector3f(0.18752, -0.00162, 0.03132),
+	FVector3f(0.19549, -0.00182, 0.01932), FVector3f(0.20304, -0.00213, 0.01178), FVector3f(0.21025, -0.00261, 0.00712),
+	FVector3f(0.21716, -0.00326, 0.00427), FVector3f(0.22380, -0.00420, 0.00255), FVector3f(0.23022, -0.00554, 0.00151),
+	FVector3f(0.23648, -0.00740, 0.00089), FVector3f(0.24264, -0.01005, 0.00053), FVector3f(0.24879, -0.01381, 0.00031),
+	FVector3f(0.25510, -0.01911, 0.00018), FVector3f(0.26177, -0.02658, 0.00011), FVector3f(0.26918, -0.03698, 0.00006),
+	FVector3f(0.27795, -0.05122, 0.00004), FVector3f(0.28910, -0.07004, 0.00002), FVector3f(0.30427, -0.09359, 0.00001),
+	FVector3f(0.32606, -0.12066, 0.00001), FVector3f(0.35822, -0.14764, 0.00001), FVector3f(0.40512, -0.16863, 0.00000),
+	FVector3f(0.46849, -0.17766, 0.00000), FVector3f(0.54169, -0.17178, 0.00000), FVector3f(0.61239, -0.15052, 0.00000),
+	FVector3f(0.67350, -0.11117, 0.00000), FVector3f(0.73152, 0.00000, 0.00000)
+},
+{
+	FVector3f(0.05336, -0.34864, 0.38172), FVector3f(0.09920, -0.08509, 0.36009), FVector3f(0.12900, -0.02477, 0.31816),
+	FVector3f(0.14348, -0.01195, 0.25287), FVector3f(0.15525, -0.00763, 0.18668), FVector3f(0.16584, -0.00587, 0.13092),
+	FVector3f(0.17560, -0.00512, 0.08853), FVector3f(0.18472, -0.00492, 0.05832), FVector3f(0.19329, -0.00508, 0.03768),
+	FVector3f(0.20140, -0.00555, 0.02399), FVector3f(0.20910, -0.00634, 0.01510), FVector3f(0.21647, -0.00749, 0.00942),
+	FVector3f(0.22355, -0.00914, 0.00584), FVector3f(0.23039, -0.01140, 0.00360), FVector3f(0.23709, -0.01450, 0.00221),
+	FVector3f(0.24371, -0.01877, 0.00135), FVector3f(0.25039, -0.02458, 0.00083), FVector3f(0.25730, -0.03248, 0.00051),
+	FVector3f(0.26474, -0.04313, 0.00031), FVector3f(0.27313, -0.05721, 0.00019), FVector3f(0.28319, -0.07550, 0.00012),
+	FVector3f(0.29598, -0.09825, 0.00008), FVector3f(0.31310, -0.12490, 0.00005), FVector3f(0.33683, -0.15332, 0.00003),
+	FVector3f(0.36995, -0.17964, 0.00002), FVector3f(0.41497, -0.19882, 0.00002), FVector3f(0.47174, -0.20686, 0.00001),
+	FVector3f(0.53490, -0.20242, 0.00001), FVector3f(0.59635, -0.18603, 0.00001), FVector3f(0.65092, -0.15783, 0.00001),
+	FVector3f(0.69798, -0.11426, 0.00001), FVector3f(0.74494, 0.00000, 0.00001)
+},
+{
+	FVector3f(0.05749, -0.31793, 0.44455), FVector3f(0.09398, -0.14133, 0.37908), FVector3f(0.13152, -0.05344, 0.33487),
+	FVector3f(0.14884, -0.02831, 0.27078), FVector3f(0.16170, -0.01861, 0.20554), FVector3f(0.17282, -0.01431, 0.14892),
+	FVector3f(0.18293, -0.01233, 0.10431), FVector3f(0.19231, -0.01159, 0.07128), FVector3f(0.20110, -0.01165, 0.04782),
+	FVector3f(0.20942, -0.01233, 0.03163), FVector3f(0.21733, -0.01364, 0.02070), FVector3f(0.22491, -0.01559, 0.01344),
+	FVector3f(0.23224, -0.01830, 0.00867), FVector3f(0.23938, -0.02200, 0.00557), FVector3f(0.24643, -0.02692, 0.00357),
+	FVector3f(0.25351, -0.03342, 0.00228), FVector3f(0.26079, -0.04192, 0.00146), FVector3f(0.26852, -0.05298, 0.00094),
+	FVector3f(0.27707, -0.06708, 0.00060), FVector3f(0.28698, -0.08468, 0.00039), FVector3f(0.29907, -0.10591, 0.00026),
+	FVector3f(0.31446, -0.13028, 0.00017), FVector3f(0.33466, -0.15629, 0.00012), FVector3f(0.36155, -0.18122, 0.00008),
+	FVector3f(0.39700, -0.20145, 0.00006), FVector3f(0.44194, -0.21352, 0.00005), FVector3f(0.49484, -0.21535, 0.00004),
+	FVector3f(0.55114, -0.20661, 0.00003), FVector3f(0.60550, -0.18774, 0.00003), FVector3f(0.65477, -0.15830, 0.00002),
+	FVector3f(0.69863, -0.11427, 0.00002), FVector3f(0.74332, 0.00000, 0.00002)
+},
+{
+	FVector3f(0.06502, -0.28106, 0.49493), FVector3f(0.09745, -0.17506, 0.41308), FVector3f(0.13592, -0.08778, 0.36191),
+	FVector3f(0.15585, -0.05167, 0.29839), FVector3f(0.17008, -0.03538, 0.23353), FVector3f(0.18198, -0.02741, 0.17549),
+	FVector3f(0.19260, -0.02338, 0.12792), FVector3f(0.20235, -0.02153, 0.09115), FVector3f(0.21146, -0.02104, 0.06384),
+	FVector3f(0.22006, -0.02158, 0.04414), FVector3f(0.22825, -0.02299, 0.03022), FVector3f(0.23612, -0.02528, 0.02053),
+	FVector3f(0.24374, -0.02852, 0.01387), FVector3f(0.25120, -0.03287, 0.00933), FVector3f(0.25860, -0.03853, 0.00626),
+	FVector3f(0.26608, -0.04579, 0.00420), FVector3f(0.27381, -0.05493, 0.00282), FVector3f(0.28203, -0.06627, 0.00189),
+	FVector3f(0.29109, -0.08016, 0.00128), FVector3f(0.30146, -0.09668, 0.00087), FVector3f(0.31381, -0.11579, 0.00060),
+	FVector3f(0.32900, -0.13678, 0.00042), FVector3f(0.34816, -0.15845, 0.00030), FVector3f(0.37257, -0.17874, 0.00022),
+	FVector3f(0.40358, -0.19512, 0.00016), FVector3f(0.44202, -0.20502, 0.00013), FVector3f(0.48743, -0.20660, 0.00010),
+	FVector3f(0.53748, -0.19897, 0.00008), FVector3f(0.58855, -0.18183, 0.00007), FVector3f(0.63751, -0.15424, 0.00006),
+	FVector3f(0.68300, -0.11197, 0.00005), FVector3f(0.72978, 0.00001, 0.00005)
+},
+{
+	FVector3f(0.07528, -0.24711, 0.53455), FVector3f(0.10863, -0.18975, 0.45764), FVector3f(0.14433, -0.11815, 0.39949),
+	FVector3f(0.16570, -0.07688, 0.33687), FVector3f(0.18121, -0.05511, 0.27245), FVector3f(0.19398, -0.04327, 0.21288),
+	FVector3f(0.20521, -0.03673, 0.16193), FVector3f(0.21545, -0.03319, 0.12068), FVector3f(0.22496, -0.03159, 0.08855),
+	FVector3f(0.23392, -0.03136, 0.06420), FVector3f(0.24244, -0.03223, 0.04612), FVector3f(0.25063, -0.03408, 0.03290),
+	FVector3f(0.25857, -0.03690, 0.02335), FVector3f(0.26634, -0.04074, 0.01650), FVector3f(0.27403, -0.04570, 0.01163),
+	FVector3f(0.28178, -0.05195, 0.00819), FVector3f(0.28972, -0.05961, 0.00577), FVector3f(0.29802, -0.06886, 0.00406),
+	FVector3f(0.30695, -0.07982, 0.00287), FVector3f(0.31682, -0.09255, 0.00204), FVector3f(0.32809, -0.10692, 0.00146),
+	FVector3f(0.34129, -0.12259, 0.00105), FVector3f(0.35714, -0.13882, 0.00077), FVector3f(0.37647, -0.15448, 0.00057),
+	FVector3f(0.40023, -0.16805, 0.00043), FVector3f(0.42931, -0.17770, 0.00033), FVector3f(0.46430, -0.18158, 0.00026),
+	FVector3f(0.50501, -0.17811, 0.00021), FVector3f(0.55015, -0.16592, 0.00018), FVector3f(0.59760, -0.14331, 0.00015),
+	FVector3f(0.64549, -0.10572, 0.00013), FVector3f(0.69674, 0.00000, 0.00012)
+},
+{
+	FVector3f(0.08968, -0.22166, 0.57013), FVector3f(0.12485, -0.19606, 0.50498), FVector3f(0.15755, -0.13952, 0.44492),
+	FVector3f(0.17936, -0.09867, 0.38452), FVector3f(0.19578, -0.07393, 0.32154), FVector3f(0.20935, -0.05908, 0.26116),
+	FVector3f(0.22123, -0.05009, 0.20725), FVector3f(0.23199, -0.04467, 0.16152), FVector3f(0.24197, -0.04157, 0.12414),
+	FVector3f(0.25134, -0.04011, 0.09438), FVector3f(0.26024, -0.03986, 0.07115), FVector3f(0.26879, -0.04064, 0.05329),
+	FVector3f(0.27706, -0.04233, 0.03970), FVector3f(0.28514, -0.04488, 0.02947), FVector3f(0.29311, -0.04830, 0.02181),
+	FVector3f(0.30105, -0.05263, 0.01611), FVector3f(0.30908, -0.05790, 0.01189), FVector3f(0.31731, -0.06417, 0.00877),
+	FVector3f(0.32591, -0.07150, 0.00648), FVector3f(0.33507, -0.07986, 0.00480), FVector3f(0.34502, -0.08921, 0.00356),
+	FVector3f(0.35610, -0.09937, 0.00266), FVector3f(0.36867, -0.11001, 0.00200), FVector3f(0.38323, -0.12062, 0.00151),
+	FVector3f(0.40031, -0.13041, 0.00116), FVector3f(0.42056, -0.13839, 0.00090), FVector3f(0.44465, -0.14322, 0.00070),
+	FVector3f(0.47318, -0.14336, 0.00056), FVector3f(0.50652, -0.13700, 0.00046), FVector3f(0.54470, -0.12166, 0.00038),
+	FVector3f(0.58750, -0.09227, 0.00032), FVector3f(0.63790, 0.00000, 0.00028)
+},
+{
+	FVector3f(0.10928, -0.20982, 0.60377), FVector3f(0.14371, -0.19861, 0.54677), FVector3f(0.17437, -0.15359, 0.48934),
+	FVector3f(0.19613, -0.11610, 0.43261), FVector3f(0.21313, -0.09068, 0.37256), FVector3f(0.22735, -0.07402, 0.31310),
+	FVector3f(0.23983, -0.06310, 0.25792), FVector3f(0.25113, -0.05592, 0.20914), FVector3f(0.26159, -0.05128, 0.16750),
+	FVector3f(0.27141, -0.04843, 0.13286), FVector3f(0.28075, -0.04692, 0.10458), FVector3f(0.28970, -0.04644, 0.08182),
+	FVector3f(0.29835, -0.04682, 0.06371), FVector3f(0.30679, -0.04799, 0.04942), FVector3f(0.31508, -0.04984, 0.03822),
+	FVector3f(0.32329, -0.05234, 0.02949), FVector3f(0.33149, -0.05550, 0.02272), FVector3f(0.33979, -0.05931, 0.01749),
+	FVector3f(0.34827, -0.06373, 0.01346), FVector3f(0.35704, -0.06873, 0.01036), FVector3f(0.36626, -0.07428, 0.00799),
+	FVector3f(0.37608, -0.08028, 0.00617), FVector3f(0.38672, -0.08655, 0.00478), FVector3f(0.39843, -0.09283, 0.00371),
+	FVector3f(0.41148, -0.09873, 0.00290), FVector3f(0.42626, -0.10371, 0.00228), FVector3f(0.44316, -0.10703, 0.00181),
+	FVector3f(0.46265, -0.10766, 0.00145), FVector3f(0.48525, -0.10415, 0.00117), FVector3f(0.51155, -0.09426, 0.00096),
+	FVector3f(0.54238, -0.07326, 0.00080), FVector3f(0.58108, 0.00000, 0.00068)
+},
+{
+	FVector3f(0.13051, -0.20682, 0.62432), FVector3f(0.16330, -0.20073, 0.57323), FVector3f(0.19267, -0.16509, 0.51965),
+	FVector3f(0.21421, -0.13200, 0.46714), FVector3f(0.23154, -0.10730, 0.41130), FVector3f(0.24629, -0.08978, 0.35487),
+	FVector3f(0.25931, -0.07747, 0.30102), FVector3f(0.27115, -0.06883, 0.25188), FVector3f(0.28213, -0.06278, 0.20849),
+	FVector3f(0.29244, -0.05862, 0.17111), FVector3f(0.30225, -0.05585, 0.13949), FVector3f(0.31166, -0.05417, 0.11309),
+	FVector3f(0.32077, -0.05341, 0.09129), FVector3f(0.32964, -0.05335, 0.07343), FVector3f(0.33834, -0.05393, 0.05890),
+	FVector3f(0.34693, -0.05508, 0.04713), FVector3f(0.35546, -0.05671, 0.03765), FVector3f(0.36403, -0.05880, 0.03004),
+	FVector3f(0.37268, -0.06133, 0.02395), FVector3f(0.38149, -0.06421, 0.01909), FVector3f(0.39055, -0.06741, 0.01521),
+	FVector3f(0.39997, -0.07082, 0.01213), FVector3f(0.40987, -0.07434, 0.00969), FVector3f(0.42039, -0.07779, 0.00775),
+	FVector3f(0.43169, -0.08092, 0.00622), FVector3f(0.44398, -0.08341, 0.00500), FVector3f(0.45749, -0.08480, 0.00404),
+	FVector3f(0.47249, -0.08439, 0.00328), FVector3f(0.48934, -0.08117, 0.00268), FVector3f(0.50847, -0.07344, 0.00220),
+	FVector3f(0.53061, -0.05739, 0.00183), FVector3f(0.55814, -0.00000, 0.00155)
+},
+{
+	FVector3f(0.15121, -0.20607, 0.62947), FVector3f(0.18300, -0.20354, 0.58335), FVector3f(0.21155, -0.17590, 0.53312),
+	FVector3f(0.23282, -0.14755, 0.48447), FVector3f(0.25031, -0.12455, 0.43318), FVector3f(0.26543, -0.10703, 0.38096),
+	FVector3f(0.27895, -0.09394, 0.33029), FVector3f(0.29131, -0.08418, 0.28305), FVector3f(0.30281, -0.07693, 0.24033),
+	FVector3f(0.31365, -0.07157, 0.20255), FVector3f(0.32398, -0.06766, 0.16971), FVector3f(0.33391, -0.06487, 0.14151),
+	FVector3f(0.34352, -0.06300, 0.11754), FVector3f(0.35289, -0.06189, 0.09733), FVector3f(0.36206, -0.06138, 0.08038),
+	FVector3f(0.37112, -0.06139, 0.06624), FVector3f(0.38010, -0.06186, 0.05449), FVector3f(0.38905, -0.06272, 0.04477),
+	FVector3f(0.39804, -0.06390, 0.03675), FVector3f(0.40711, -0.06534, 0.03014), FVector3f(0.41635, -0.06699, 0.02471),
+	FVector3f(0.42580, -0.06874, 0.02026), FVector3f(0.43556, -0.07051, 0.01662), FVector3f(0.44571, -0.07214, 0.01365),
+	FVector3f(0.45636, -0.07346, 0.01122), FVector3f(0.46763, -0.07423, 0.00924), FVector3f(0.47968, -0.07408, 0.00762),
+	FVector3f(0.49266, -0.07254, 0.00631), FVector3f(0.50681, -0.06883, 0.00524), FVector3f(0.52242, -0.06160, 0.00437),
+	FVector3f(0.54001, -0.04778, 0.00367), FVector3f(0.56112, 0.00000, 0.00312)
+},
+{
+	FVector3f(0.17197, -0.20561, 0.62587), FVector3f(0.20338, -0.20599, 0.58408), FVector3f(0.23149, -0.18515, 0.53684),
+	FVector3f(0.25251, -0.16154, 0.49150), FVector3f(0.27004, -0.14094, 0.44437), FVector3f(0.28542, -0.12426, 0.39642),
+	FVector3f(0.29931, -0.11107, 0.34949), FVector3f(0.31213, -0.10073, 0.30512), FVector3f(0.32413, -0.09263, 0.26430),
+	FVector3f(0.33550, -0.08631, 0.22751), FVector3f(0.34636, -0.08139, 0.19485), FVector3f(0.35682, -0.07760, 0.16619),
+	FVector3f(0.36695, -0.07471, 0.14127), FVector3f(0.37684, -0.07258, 0.11976), FVector3f(0.38652, -0.07106, 0.10128),
+	FVector3f(0.39607, -0.07006, 0.08549), FVector3f(0.40551, -0.06949, 0.07205), FVector3f(0.41490, -0.06928, 0.06064),
+	FVector3f(0.42427, -0.06933, 0.05099), FVector3f(0.43369, -0.06961, 0.04283), FVector3f(0.44319, -0.07004, 0.03597),
+	FVector3f(0.45283, -0.07053, 0.03019), FVector3f(0.46265, -0.07097, 0.02535), FVector3f(0.47273, -0.07127, 0.02128),
+	FVector3f(0.48314, -0.07125, 0.01788), FVector3f(0.49395, -0.07070, 0.01504), FVector3f(0.50527, -0.06936, 0.01267),
+	FVector3f(0.51721, -0.06683, 0.01069), FVector3f(0.52993, -0.06247, 0.00904), FVector3f(0.54362, -0.05516, 0.00766),
+	FVector3f(0.55866, -0.04229, 0.00653), FVector3f(0.57607, -0.00000, 0.00561)
+},
+{
+	FVector3f(0.19561, -0.20369, 0.62676), FVector3f(0.22742, -0.20521, 0.59059), FVector3f(0.25590, -0.18887, 0.54809),
+	FVector3f(0.27707, -0.16911, 0.50677), FVector3f(0.29477, -0.15102, 0.46369), FVector3f(0.31041, -0.13567, 0.41959),
+	FVector3f(0.32465, -0.12298, 0.37594), FVector3f(0.33788, -0.11258, 0.33408), FVector3f(0.35033, -0.10408, 0.29491),
+	FVector3f(0.36219, -0.09713, 0.25894), FVector3f(0.37355, -0.09144, 0.22638), FVector3f(0.38451, -0.08677, 0.19721),
+	FVector3f(0.39514, -0.08296, 0.17130), FVector3f(0.40551, -0.07987, 0.14843), FVector3f(0.41565, -0.07736, 0.12835),
+	FVector3f(0.42562, -0.07533, 0.11079), FVector3f(0.43546, -0.07371, 0.09548), FVector3f(0.44519, -0.07242, 0.08219),
+	FVector3f(0.45486, -0.07138, 0.07067), FVector3f(0.46449, -0.07053, 0.06071), FVector3f(0.47413, -0.06981, 0.05212),
+	FVector3f(0.48380, -0.06914, 0.04472), FVector3f(0.49355, -0.06841, 0.03835), FVector3f(0.50342, -0.06755, 0.03289),
+	FVector3f(0.51345, -0.06643, 0.02821), FVector3f(0.52370, -0.06487, 0.02421), FVector3f(0.53423, -0.06265, 0.02078),
+	FVector3f(0.54512, -0.05946, 0.01786), FVector3f(0.55646, -0.05481, 0.01537), FVector3f(0.56840, -0.04776, 0.01325),
+	FVector3f(0.58119, -0.03618, 0.01145), FVector3f(0.59545, -0.00000, 0.00995)
+},
+{
+	FVector3f(0.22163, -0.20055, 0.63021), FVector3f(0.25410, -0.20231, 0.59968), FVector3f(0.28326, -0.18906, 0.56225),
+	FVector3f(0.30475, -0.17238, 0.52492), FVector3f(0.32266, -0.15662, 0.48555), FVector3f(0.33852, -0.14279, 0.44484),
+	FVector3f(0.35302, -0.13097, 0.40410), FVector3f(0.36656, -0.12093, 0.36452), FVector3f(0.37937, -0.11245, 0.32693),
+	FVector3f(0.39160, -0.10523, 0.29188), FVector3f(0.40334, -0.09909, 0.25962), FVector3f(0.41470, -0.09386, 0.23022),
+	FVector3f(0.42572, -0.08939, 0.20363), FVector3f(0.43646, -0.08557, 0.17973), FVector3f(0.44696, -0.08227, 0.15834),
+	FVector3f(0.45725, -0.07943, 0.13927), FVector3f(0.46736, -0.07697, 0.12233), FVector3f(0.47733, -0.07481, 0.10732),
+	FVector3f(0.48718, -0.07290, 0.09404), FVector3f(0.49693, -0.07117, 0.08234), FVector3f(0.50661, -0.06954, 0.07203),
+	FVector3f(0.51626, -0.06797, 0.06298), FVector3f(0.52588, -0.06635, 0.05503), FVector3f(0.53552, -0.06465, 0.04807),
+	FVector3f(0.54521, -0.06270, 0.04198), FVector3f(0.55498, -0.06041, 0.03667), FVector3f(0.56487, -0.05757, 0.03203),
+	FVector3f(0.57494, -0.05392, 0.02799), FVector3f(0.58526, -0.04906, 0.02447), FVector3f(0.59592, -0.04224, 0.02142),
+	FVector3f(0.60709, -0.03161, 0.01879), FVector3f(0.61917, -0.00000, 0.01653)
+},
+{
+	FVector3f(0.24798, -0.19716, 0.63007), FVector3f(0.28069, -0.19930, 0.60332), FVector3f(0.31015, -0.18883, 0.56936),
+	FVector3f(0.33172, -0.17504, 0.53486), FVector3f(0.34964, -0.16162, 0.49830), FVector3f(0.36550, -0.14954, 0.46037),
+	FVector3f(0.38006, -0.13890, 0.42220), FVector3f(0.39370, -0.12960, 0.38484), FVector3f(0.40666, -0.12148, 0.34906),
+	FVector3f(0.41907, -0.11438, 0.31535), FVector3f(0.43103, -0.10815, 0.28399), FVector3f(0.44261, -0.10266, 0.25508),
+	FVector3f(0.45385, -0.09783, 0.22861), FVector3f(0.46481, -0.09354, 0.20450), FVector3f(0.47551, -0.08974, 0.18264),
+	FVector3f(0.48599, -0.08633, 0.16289), FVector3f(0.49627, -0.08326, 0.14510), FVector3f(0.50638, -0.08047, 0.12910),
+	FVector3f(0.51633, -0.07791, 0.11476), FVector3f(0.52614, -0.07551, 0.10192), FVector3f(0.53584, -0.07321, 0.09045),
+	FVector3f(0.54546, -0.07096, 0.08022), FVector3f(0.55500, -0.06867, 0.07111), FVector3f(0.56449, -0.06628, 0.06301),
+	FVector3f(0.57397, -0.06367, 0.05582), FVector3f(0.58345, -0.06076, 0.04944), FVector3f(0.59297, -0.05731, 0.04379),
+	FVector3f(0.60257, -0.05316, 0.03879), FVector3f(0.61231, -0.04788, 0.03438), FVector3f(0.62225, -0.04079, 0.03050),
+	FVector3f(0.63252, -0.03022, 0.02708), FVector3f(0.64336, 0.00000, 0.02411)
+},
+{
+	FVector3f(0.27722, -0.19199, 0.63565), FVector3f(0.30981, -0.19408, 0.61130), FVector3f(0.33943, -0.18579, 0.57961),
+	FVector3f(0.36102, -0.17447, 0.54708), FVector3f(0.37889, -0.16322, 0.51263), FVector3f(0.39470, -0.15282, 0.47688),
+	FVector3f(0.40920, -0.14347, 0.44082), FVector3f(0.42282, -0.13509, 0.40537), FVector3f(0.43577, -0.12760, 0.37118),
+	FVector3f(0.44819, -0.12084, 0.33873), FVector3f(0.46018, -0.11477, 0.30827), FVector3f(0.47180, -0.10928, 0.27992),
+	FVector3f(0.48310, -0.10430, 0.25370), FVector3f(0.49411, -0.09979, 0.22957), FVector3f(0.50485, -0.09566, 0.20745),
+	FVector3f(0.51535, -0.09188, 0.18722), FVector3f(0.52564, -0.08837, 0.16879), FVector3f(0.53572, -0.08512, 0.15202),
+	FVector3f(0.54562, -0.08205, 0.13680), FVector3f(0.55536, -0.07912, 0.12300), FVector3f(0.56495, -0.07630, 0.11052),
+	FVector3f(0.57441, -0.07350, 0.09925), FVector3f(0.58375, -0.07067, 0.08908), FVector3f(0.59300, -0.06773, 0.07992),
+	FVector3f(0.60219, -0.06459, 0.07168), FVector3f(0.61132, -0.06115, 0.06428), FVector3f(0.62043, -0.05724, 0.05763),
+	FVector3f(0.62954, -0.05265, 0.05168), FVector3f(0.63871, -0.04703, 0.04636), FVector3f(0.64797, -0.03973, 0.04161),
+	FVector3f(0.65743, -0.02917, 0.03737), FVector3f(0.66725, -0.00000, 0.03363)
+},
+{
+	FVector3f(0.30941, -0.18550, 0.64456), FVector3f(0.34168, -0.18712, 0.62252), FVector3f(0.37132, -0.18024, 0.59322),
+	FVector3f(0.39288, -0.17081, 0.56272), FVector3f(0.41063, -0.16130, 0.53026), FVector3f(0.42627, -0.15243, 0.49646),
+	FVector3f(0.44059, -0.14431, 0.46224), FVector3f(0.45403, -0.13689, 0.42843), FVector3f(0.46680, -0.13010, 0.39561),
+	FVector3f(0.47906, -0.12388, 0.36425), FVector3f(0.49089, -0.11815, 0.33457), FVector3f(0.50236, -0.11285, 0.30671),
+	FVector3f(0.51350, -0.10796, 0.28072), FVector3f(0.52435, -0.10341, 0.25657), FVector3f(0.53493, -0.09917, 0.23422),
+	FVector3f(0.54526, -0.09519, 0.21359), FVector3f(0.55537, -0.09146, 0.19458), FVector3f(0.56525, -0.08792, 0.17712),
+	FVector3f(0.57494, -0.08453, 0.16109), FVector3f(0.58444, -0.08126, 0.14642), FVector3f(0.59377, -0.07806, 0.13299),
+	FVector3f(0.60294, -0.07488, 0.12073), FVector3f(0.61197, -0.07164, 0.10955), FVector3f(0.62087, -0.06830, 0.09937),
+	FVector3f(0.62966, -0.06479, 0.09010), FVector3f(0.63836, -0.06097, 0.08168), FVector3f(0.64699, -0.05673, 0.07404),
+	FVector3f(0.65557, -0.05183, 0.06711), FVector3f(0.66415, -0.04599, 0.06085), FVector3f(0.67274, -0.03858, 0.05518),
+	FVector3f(0.68143, -0.02812, 0.05008), FVector3f(0.69032, -0.00000, 0.04550)
+},
+{
+	FVector3f(0.34235, -0.17826, 0.65206), FVector3f(0.37374, -0.17952, 0.63152), FVector3f(0.40289, -0.17391, 0.60389),
+	FVector3f(0.42408, -0.16614, 0.57496), FVector3f(0.44148, -0.15827, 0.54412), FVector3f(0.45676, -0.15084, 0.51197),
+	FVector3f(0.47071, -0.14393, 0.47937), FVector3f(0.48377, -0.13752, 0.44704), FVector3f(0.49618, -0.13155, 0.41554),
+	FVector3f(0.50808, -0.12598, 0.38525), FVector3f(0.51957, -0.12076, 0.35644), FVector3f(0.53070, -0.11585, 0.32922),
+	FVector3f(0.54151, -0.11120, 0.30365), FVector3f(0.55204, -0.10681, 0.27973), FVector3f(0.56231, -0.10265, 0.25742),
+	FVector3f(0.57233, -0.09867, 0.23667), FVector3f(0.58212, -0.09486, 0.21741), FVector3f(0.59168, -0.09120, 0.19957),
+	FVector3f(0.60106, -0.08765, 0.18306), FVector3f(0.61023, -0.08416, 0.16782), FVector3f(0.61923, -0.08072, 0.15376),
+	FVector3f(0.62806, -0.07727, 0.14081), FVector3f(0.63672, -0.07376, 0.12890), FVector3f(0.64525, -0.07012, 0.11795),
+	FVector3f(0.65365, -0.06628, 0.10790), FVector3f(0.66194, -0.06216, 0.09869), FVector3f(0.67013, -0.05760, 0.09025),
+	FVector3f(0.67825, -0.05241, 0.08254), FVector3f(0.68632, -0.04628, 0.07549), FVector3f(0.69436, -0.03862, 0.06906),
+	FVector3f(0.70244, -0.02799, 0.06321), FVector3f(0.71061, 0.00000, 0.05791)
+},
+{
+	FVector3f(0.37818, -0.17017, 0.65945), FVector3f(0.40866, -0.17100, 0.64009), FVector3f(0.43736, -0.16627, 0.61399),
+	FVector3f(0.45826, -0.15979, 0.58661), FVector3f(0.47535, -0.15320, 0.55744), FVector3f(0.49027, -0.14692, 0.52704),
+	FVector3f(0.50383, -0.14103, 0.49617), FVector3f(0.51647, -0.13549, 0.46545), FVector3f(0.52846, -0.13026, 0.43541),
+	FVector3f(0.53992, -0.12530, 0.40641), FVector3f(0.55095, -0.12057, 0.37864), FVector3f(0.56162, -0.11605, 0.35226),
+	FVector3f(0.57197, -0.11171, 0.32733), FVector3f(0.58204, -0.10753, 0.30384), FVector3f(0.59185, -0.10350, 0.28179),
+	FVector3f(0.60141, -0.09961, 0.26113), FVector3f(0.61074, -0.09583, 0.24181), FVector3f(0.61985, -0.09214, 0.22378),
+	FVector3f(0.62877, -0.08853, 0.20696), FVector3f(0.63747, -0.08494, 0.19131), FVector3f(0.64600, -0.08136, 0.17675),
+	FVector3f(0.65435, -0.07776, 0.16324), FVector3f(0.66254, -0.07407, 0.15070), FVector3f(0.67057, -0.07025, 0.13908),
+	FVector3f(0.67846, -0.06622, 0.12832), FVector3f(0.68622, -0.06190, 0.11837), FVector3f(0.69387, -0.05716, 0.10919),
+	FVector3f(0.70143, -0.05182, 0.10071), FVector3f(0.70890, -0.04557, 0.09289), FVector3f(0.71632, -0.03786, 0.08570),
+	FVector3f(0.72371, -0.02730, 0.07909), FVector3f(0.73113, -0.00000, 0.07304)
+},
+{
+	FVector3f(0.41622, -0.16117, 0.66557), FVector3f(0.44565, -0.16160, 0.64703), FVector3f(0.47380, -0.15756, 0.62216),
+	FVector3f(0.49434, -0.15215, 0.59618), FVector3f(0.51107, -0.14662, 0.56863), FVector3f(0.52562, -0.14132, 0.53998),
+	FVector3f(0.53876, -0.13629, 0.51088), FVector3f(0.55092, -0.13152, 0.48188), FVector3f(0.56240, -0.12695, 0.45343),
+	FVector3f(0.57332, -0.12256, 0.42585), FVector3f(0.58380, -0.11831, 0.39935), FVector3f(0.59391, -0.11419, 0.37404),
+	FVector3f(0.60370, -0.11017, 0.34998), FVector3f(0.61319, -0.10625, 0.32718), FVector3f(0.62242, -0.10243, 0.30564),
+	FVector3f(0.63141, -0.09869, 0.28533), FVector3f(0.64017, -0.09501, 0.26621), FVector3f(0.64871, -0.09139, 0.24823),
+	FVector3f(0.65705, -0.08778, 0.23136), FVector3f(0.66519, -0.08420, 0.21553), FVector3f(0.67315, -0.08059, 0.20071),
+	FVector3f(0.68092, -0.07692, 0.18683), FVector3f(0.68853, -0.07317, 0.17386), FVector3f(0.69599, -0.06927, 0.16174),
+	FVector3f(0.70329, -0.06516, 0.15044), FVector3f(0.71046, -0.06075, 0.13990), FVector3f(0.71751, -0.05595, 0.13009),
+	FVector3f(0.72444, -0.05055, 0.12095), FVector3f(0.73128, -0.04430, 0.11246), FVector3f(0.73804, -0.03667, 0.10458),
+	FVector3f(0.74474, -0.02633, 0.09728), FVector3f(0.75140, 0.00000, 0.09052)
+},
+{
+	FVector3f(0.45566, -0.15078, 0.67181), FVector3f(0.48376, -0.15083, 0.65404), FVector3f(0.51104, -0.14734, 0.63030),
+	FVector3f(0.53106, -0.14281, 0.60566), FVector3f(0.54730, -0.13818, 0.57963), FVector3f(0.56134, -0.13375, 0.55265),
+	FVector3f(0.57395, -0.12949, 0.52523), FVector3f(0.58556, -0.12542, 0.49788), FVector3f(0.59643, -0.12148, 0.47100),
+	FVector3f(0.60673, -0.11762, 0.44485), FVector3f(0.61657, -0.11387, 0.41962), FVector3f(0.62603, -0.11017, 0.39542),
+	FVector3f(0.63516, -0.10653, 0.37231), FVector3f(0.64399, -0.10294, 0.35030), FVector3f(0.65257, -0.09939, 0.32939),
+	FVector3f(0.66090, -0.09589, 0.30956), FVector3f(0.66900, -0.09240, 0.29078), FVector3f(0.67689, -0.08893, 0.27303),
+	FVector3f(0.68459, -0.08546, 0.25624), FVector3f(0.69209, -0.08197, 0.24041), FVector3f(0.69941, -0.07844, 0.22547),
+	FVector3f(0.70656, -0.07483, 0.21139), FVector3f(0.71353, -0.07112, 0.19815), FVector3f(0.72036, -0.06724, 0.18568),
+	FVector3f(0.72704, -0.06317, 0.17397), FVector3f(0.73357, -0.05880, 0.16298), FVector3f(0.73999, -0.05404, 0.15266),
+	FVector3f(0.74629, -0.04872, 0.14298), FVector3f(0.75247, -0.04259, 0.13392), FVector3f(0.75857, -0.03514, 0.12544),
+	FVector3f(0.76459, -0.02515, 0.11752), FVector3f(0.77054, -0.00000, 0.11012)
+},
+{
+	FVector3f(0.49546, -0.13868, 0.67921), FVector3f(0.52189, -0.13841, 0.66231), FVector3f(0.54795, -0.13540, 0.63977),
+	FVector3f(0.56715, -0.13163, 0.61642), FVector3f(0.58273, -0.12780, 0.59185), FVector3f(0.59613, -0.12410, 0.56640),
+	FVector3f(0.60809, -0.12056, 0.54055), FVector3f(0.61905, -0.11715, 0.51472), FVector3f(0.62925, -0.11379, 0.48927),
+	FVector3f(0.63886, -0.11049, 0.46444), FVector3f(0.64800, -0.10724, 0.44041), FVector3f(0.65676, -0.10400, 0.41727),
+	FVector3f(0.66517, -0.10079, 0.39507), FVector3f(0.67330, -0.09757, 0.37385), FVector3f(0.68117, -0.09438, 0.35358),
+	FVector3f(0.68879, -0.09119, 0.33427), FVector3f(0.69620, -0.08798, 0.31589), FVector3f(0.70340, -0.08477, 0.29841),
+	FVector3f(0.71041, -0.08153, 0.28180), FVector3f(0.71724, -0.07825, 0.26604), FVector3f(0.72389, -0.07490, 0.25108),
+	FVector3f(0.73037, -0.07146, 0.23691), FVector3f(0.73669, -0.06790, 0.22349), FVector3f(0.74288, -0.06418, 0.21078),
+	FVector3f(0.74892, -0.06025, 0.19877), FVector3f(0.75481, -0.05603, 0.18742), FVector3f(0.76059, -0.05144, 0.17669),
+	FVector3f(0.76625, -0.04631, 0.16657), FVector3f(0.77180, -0.04042, 0.15703), FVector3f(0.77725, -0.03329, 0.14803),
+	FVector3f(0.78262, -0.02378, 0.13957), FVector3f(0.78790, -0.00000, 0.13161)
+},
+{
+	FVector3f(0.53465, -0.12522, 0.68615), FVector3f(0.55925, -0.12477, 0.67008), FVector3f(0.58386, -0.12219, 0.64863),
+	FVector3f(0.60211, -0.11908, 0.62649), FVector3f(0.61691, -0.11598, 0.60327), FVector3f(0.62960, -0.11298, 0.57927),
+	FVector3f(0.64087, -0.11008, 0.55488), FVector3f(0.65112, -0.10725, 0.53052), FVector3f(0.66062, -0.10446, 0.50646),
+	FVector3f(0.66953, -0.10170, 0.48294), FVector3f(0.67795, -0.09893, 0.46011), FVector3f(0.68598, -0.09617, 0.43805),
+	FVector3f(0.69369, -0.09339, 0.41682), FVector3f(0.70109, -0.09059, 0.39644), FVector3f(0.70824, -0.08778, 0.37690),
+	FVector3f(0.71515, -0.08494, 0.35821), FVector3f(0.72185, -0.08209, 0.34034), FVector3f(0.72836, -0.07918, 0.32326),
+	FVector3f(0.73467, -0.07624, 0.30696), FVector3f(0.74082, -0.07324, 0.29141), FVector3f(0.74679, -0.07016, 0.27658),
+	FVector3f(0.75261, -0.06698, 0.26246), FVector3f(0.75828, -0.06367, 0.24901), FVector3f(0.76381, -0.06018, 0.23621),
+	FVector3f(0.76920, -0.05650, 0.22404), FVector3f(0.77446, -0.05253, 0.21247), FVector3f(0.77961, -0.04820, 0.20148),
+	FVector3f(0.78464, -0.04336, 0.19105), FVector3f(0.78957, -0.03781, 0.18115), FVector3f(0.79440, -0.03111, 0.17176),
+	FVector3f(0.79913, -0.02219, 0.16287), FVector3f(0.80378, -0.00000, 0.15446)
+},
+{
+	FVector3f(0.57256, -0.11055, 0.69171), FVector3f(0.59531, -0.11002, 0.67633), FVector3f(0.61841, -0.10792, 0.65578),
+	FVector3f(0.63563, -0.10544, 0.63470), FVector3f(0.64962, -0.10298, 0.61271), FVector3f(0.66156, -0.10060, 0.59006),
+	FVector3f(0.67214, -0.09828, 0.56710), FVector3f(0.68170, -0.09602, 0.54415), FVector3f(0.69050, -0.09375, 0.52148),
+	FVector3f(0.69871, -0.09149, 0.49928), FVector3f(0.70643, -0.08920, 0.47768), FVector3f(0.71376, -0.08689, 0.45676),
+	FVector3f(0.72074, -0.08455, 0.43658), FVector3f(0.72745, -0.08218, 0.41714), FVector3f(0.73389, -0.07976, 0.39844),
+	FVector3f(0.74011, -0.07732, 0.38048), FVector3f(0.74611, -0.07484, 0.36325), FVector3f(0.75193, -0.07230, 0.34672),
+	FVector3f(0.75756, -0.06969, 0.33088), FVector3f(0.76304, -0.06702, 0.31571), FVector3f(0.76835, -0.06427, 0.30118),
+	FVector3f(0.77352, -0.06140, 0.28727), FVector3f(0.77854, -0.05841, 0.27397), FVector3f(0.78343, -0.05525, 0.26124),
+	FVector3f(0.78818, -0.05188, 0.24909), FVector3f(0.79283, -0.04824, 0.23747), FVector3f(0.79736, -0.04426, 0.22638),
+	FVector3f(0.80178, -0.03982, 0.21580), FVector3f(0.80609, -0.03471, 0.20571), FVector3f(0.81032, -0.02854, 0.19608),
+	FVector3f(0.81445, -0.02034, 0.18691), FVector3f(0.81849, 0.00001, 0.17818)
+},
+{
+	FVector3f(0.60872, -0.09464, 0.69662), FVector3f(0.62966, -0.09415, 0.68177), FVector3f(0.65120, -0.09250, 0.66197),
+	FVector3f(0.66737, -0.09060, 0.64182), FVector3f(0.68051, -0.08873, 0.62094), FVector3f(0.69172, -0.08690, 0.59955),
+	FVector3f(0.70159, -0.08513, 0.57791), FVector3f(0.71047, -0.08335, 0.55632), FVector3f(0.71860, -0.08158, 0.53498),
+	FVector3f(0.72612, -0.07978, 0.51406), FVector3f(0.73318, -0.07795, 0.49369), FVector3f(0.73983, -0.07608, 0.47393),
+	FVector3f(0.74614, -0.07417, 0.45480), FVector3f(0.75217, -0.07223, 0.43634), FVector3f(0.75795, -0.07023, 0.41855),
+	FVector3f(0.76350, -0.06820, 0.40140), FVector3f(0.76885, -0.06611, 0.38489), FVector3f(0.77401, -0.06395, 0.36901),
+	FVector3f(0.77900, -0.06175, 0.35374), FVector3f(0.78383, -0.05946, 0.33905), FVector3f(0.78851, -0.05707, 0.32494),
+	FVector3f(0.79306, -0.05459, 0.31137), FVector3f(0.79745, -0.05197, 0.29835), FVector3f(0.80173, -0.04920, 0.28584),
+	FVector3f(0.80590, -0.04623, 0.27384), FVector3f(0.80993, -0.04302, 0.26232), FVector3f(0.81388, -0.03949, 0.25127),
+	FVector3f(0.81771, -0.03554, 0.24068), FVector3f(0.82144, -0.03097, 0.23053), FVector3f(0.82508, -0.02547, 0.22081),
+	FVector3f(0.82863, -0.01815, 0.21150), FVector3f(0.83210, 0.00000, 0.20258)
+},
+{
+	FVector3f(0.64301, -0.07753, 0.70123), FVector3f(0.66223, -0.07712, 0.68683), FVector3f(0.68223, -0.07593, 0.66763),
+	FVector3f(0.69736, -0.07455, 0.64827), FVector3f(0.70966, -0.07320, 0.62840), FVector3f(0.72015, -0.07187, 0.60815),
+	FVector3f(0.72934, -0.07056, 0.58775), FVector3f(0.73757, -0.06924, 0.56743), FVector3f(0.74506, -0.06791, 0.54737),
+	FVector3f(0.75196, -0.06655, 0.52771), FVector3f(0.75838, -0.06515, 0.50855), FVector3f(0.76441, -0.06371, 0.48992),
+	FVector3f(0.77009, -0.06222, 0.47188), FVector3f(0.77550, -0.06069, 0.45443), FVector3f(0.78065, -0.05912, 0.43757),
+	FVector3f(0.78559, -0.05749, 0.42129), FVector3f(0.79032, -0.05583, 0.40557), FVector3f(0.79488, -0.05409, 0.39042),
+	FVector3f(0.79926, -0.05229, 0.37579), FVector3f(0.80349, -0.05042, 0.36169), FVector3f(0.80758, -0.04846, 0.34810),
+	FVector3f(0.81154, -0.04640, 0.33499), FVector3f(0.81537, -0.04423, 0.32235), FVector3f(0.81906, -0.04191, 0.31018),
+	FVector3f(0.82266, -0.03941, 0.29845), FVector3f(0.82614, -0.03670, 0.28716), FVector3f(0.82952, -0.03372, 0.27629),
+	FVector3f(0.83279, -0.03035, 0.26582), FVector3f(0.83599, -0.02647, 0.25574), FVector3f(0.83908, -0.02178, 0.24605),
+	FVector3f(0.84209, -0.01551, 0.23672), FVector3f(0.84502, 0.00001, 0.22776)
+},
+{
+	FVector3f(0.67520, -0.05931, 0.70649), FVector3f(0.69276, -0.05903, 0.69240), FVector3f(0.71128, -0.05822, 0.67366),
+	FVector3f(0.72537, -0.05730, 0.65499), FVector3f(0.73687, -0.05639, 0.63599), FVector3f(0.74664, -0.05549, 0.61679),
+	FVector3f(0.75518, -0.05459, 0.59754), FVector3f(0.76280, -0.05368, 0.57840), FVector3f(0.76970, -0.05276, 0.55954),
+	FVector3f(0.77600, -0.05180, 0.54106), FVector3f(0.78185, -0.05080, 0.52303), FVector3f(0.78730, -0.04977, 0.50551),
+	FVector3f(0.79243, -0.04869, 0.48852), FVector3f(0.79726, -0.04757, 0.47207), FVector3f(0.80185, -0.04641, 0.45614),
+	FVector3f(0.80622, -0.04521, 0.44074), FVector3f(0.81040, -0.04396, 0.42583), FVector3f(0.81441, -0.04265, 0.41143),
+	FVector3f(0.81826, -0.04130, 0.39750), FVector3f(0.82195, -0.03988, 0.38404), FVector3f(0.82551, -0.03837, 0.37102),
+	FVector3f(0.82892, -0.03679, 0.35844), FVector3f(0.83223, -0.03510, 0.34628), FVector3f(0.83541, -0.03331, 0.33452),
+	FVector3f(0.83849, -0.03135, 0.32316), FVector3f(0.84146, -0.02922, 0.31218), FVector3f(0.84433, -0.02686, 0.30158),
+	FVector3f(0.84711, -0.02421, 0.29134), FVector3f(0.84980, -0.02113, 0.28144), FVector3f(0.85241, -0.01739, 0.27189),
+	FVector3f(0.85493, -0.01240, 0.26266), FVector3f(0.85737, -0.00000, 0.25376)
+},
+{
+	FVector3f(0.70511, -0.04013, 0.71307), FVector3f(0.72112, -0.03997, 0.69919), FVector3f(0.73819, -0.03950, 0.68079),
+	FVector3f(0.75128, -0.03897, 0.66265), FVector3f(0.76198, -0.03845, 0.64443), FVector3f(0.77108, -0.03791, 0.62615),
+	FVector3f(0.77901, -0.03737, 0.60793), FVector3f(0.78606, -0.03682, 0.58989), FVector3f(0.79240, -0.03625, 0.57212),
+	FVector3f(0.79819, -0.03565, 0.55474), FVector3f(0.80350, -0.03502, 0.53779), FVector3f(0.80844, -0.03437, 0.52130),
+	FVector3f(0.81305, -0.03367, 0.50532), FVector3f(0.81738, -0.03295, 0.48981), FVector3f(0.82148, -0.03220, 0.47479),
+	FVector3f(0.82537, -0.03142, 0.46024), FVector3f(0.82906, -0.03059, 0.44615), FVector3f(0.83259, -0.02973, 0.43250),
+	FVector3f(0.83596, -0.02882, 0.41928), FVector3f(0.83918, -0.02787, 0.40648), FVector3f(0.84227, -0.02686, 0.39407),
+	FVector3f(0.84523, -0.02578, 0.38205), FVector3f(0.84808, -0.02463, 0.37040), FVector3f(0.85081, -0.02339, 0.35913),
+	FVector3f(0.85345, -0.02205, 0.34819), FVector3f(0.85597, -0.02057, 0.33761), FVector3f(0.85842, -0.01893, 0.32734),
+	FVector3f(0.86076, -0.01707, 0.31740), FVector3f(0.86302, -0.01492, 0.30777), FVector3f(0.86519, -0.01229, 0.29845),
+	FVector3f(0.86729, -0.00878, 0.28941), FVector3f(0.86932, -0.00000, 0.28066)
+},
+{
+	FVector3f(0.73154, -0.02030, 0.72031), FVector3f(0.74609, -0.02024, 0.70643), FVector3f(0.76177, -0.02006, 0.68812),
+	FVector3f(0.77389, -0.01984, 0.67034), FVector3f(0.78383, -0.01961, 0.65270), FVector3f(0.79228, -0.01939, 0.63519),
+	FVector3f(0.79965, -0.01914, 0.61784), FVector3f(0.80618, -0.01889, 0.60072), FVector3f(0.81203, -0.01864, 0.58392),
+	FVector3f(0.81735, -0.01837, 0.56751), FVector3f(0.82221, -0.01807, 0.55151), FVector3f(0.82672, -0.01776, 0.53597),
+	FVector3f(0.83089, -0.01744, 0.52087), FVector3f(0.83482, -0.01710, 0.50623), FVector3f(0.83850, -0.01673, 0.49203),
+	FVector3f(0.84198, -0.01635, 0.47828), FVector3f(0.84528, -0.01595, 0.46494), FVector3f(0.84842, -0.01552, 0.45201),
+	FVector3f(0.85140, -0.01508, 0.43946), FVector3f(0.85424, -0.01460, 0.42730), FVector3f(0.85696, -0.01408, 0.41549),
+	FVector3f(0.85955, -0.01354, 0.40404), FVector3f(0.86204, -0.01296, 0.39292), FVector3f(0.86442, -0.01233, 0.38212),
+	FVector3f(0.86669, -0.01164, 0.37165), FVector3f(0.86887, -0.01087, 0.36148), FVector3f(0.87097, -0.01002, 0.35161),
+	FVector3f(0.87297, -0.00905, 0.34203), FVector3f(0.87490, -0.00791, 0.33272), FVector3f(0.87673, -0.00653, 0.32369),
+	FVector3f(0.87850, -0.00466, 0.31492), FVector3f(0.88020, 0.00001, 0.30640)
+},
+{
+	FVector3f(0.75486, -0.00000, 0.72806), FVector3f(0.76807, 0.00000, 0.71395), FVector3f(0.78246, -0.00000, 0.69552),
+	FVector3f(0.79366, -0.00000, 0.67790), FVector3f(0.80290, 0.00001, 0.66069), FVector3f(0.81077, 0.00001, 0.64378),
+	FVector3f(0.81763, -0.00000, 0.62716), FVector3f(0.82368, -0.00000, 0.61086), FVector3f(0.82912, -0.00000, 0.59491),
+	FVector3f(0.83404, -0.00000, 0.57936), FVector3f(0.83852, -0.00000, 0.56423), FVector3f(0.84266, 0.00000, 0.54953),
+	FVector3f(0.84649, 0.00000, 0.53526), FVector3f(0.85008, -0.00000, 0.52142), FVector3f(0.85343, 0.00000, 0.50800),
+	FVector3f(0.85660, -0.00000, 0.49498), FVector3f(0.85959, 0.00000, 0.48235), FVector3f(0.86241, -0.00000, 0.47011),
+	FVector3f(0.86510, 0.00001, 0.45821), FVector3f(0.86766, 0.00000, 0.44666), FVector3f(0.87010, -0.00001, 0.43545),
+	FVector3f(0.87242, 0.00000, 0.42456), FVector3f(0.87464, -0.00000, 0.41398), FVector3f(0.87675, 0.00000, 0.40369),
+	FVector3f(0.87877, 0.00000, 0.39370), FVector3f(0.88070, -0.00000, 0.38398), FVector3f(0.88255, -0.00000, 0.37453),
+	FVector3f(0.88431, 0.00000, 0.36535), FVector3f(0.88600, 0.00000, 0.35642), FVector3f(0.88761, -0.00000, 0.34773),
+	FVector3f(0.88915, -0.00000, 0.33929), FVector3f(0.89063, -0.00000, 0.33107)
+}
+};
\ No newline at end of file
diff --git a/Engine/Source/Runtime/Renderer/Private/SceneRendering.cpp b/Engine/Source/Runtime/Renderer/Private/SceneRendering.cpp
--- a/Engine/Source/Runtime/Renderer/Private/SceneRendering.cpp
+++ b/Engine/Source/Runtime/Renderer/Private/SceneRendering.cpp
@@ -1880,17 +1880,22 @@ void FViewInfo::SetupUniformBufferParameters(
 	ViewUniformShaderParameters.RuntimeVirtualTexturePackHeight = FVector2f(ForceInitToZero);
 	ViewUniformShaderParameters.RuntimeVirtualTextureDebugParams = FVector4f(ForceInitToZero);
 
-	// Rect area light
-	if (GSystemTextures.LTCMat.IsValid() && GSystemTextures.LTCAmp.IsValid())
+	// GGX/Sheen LTC (used as BSDF or for rect light integration)
+	if (GSystemTextures.GGXLTCMat.IsValid() && GSystemTextures.GGXLTCAmp.IsValid())
 	{
-		ViewUniformShaderParameters.LTCMatTexture = GSystemTextures.LTCMat->GetRHI();
-		ViewUniformShaderParameters.LTCMatSampler = TStaticSamplerState<SF_Bilinear, AM_Clamp, AM_Clamp, AM_Clamp>::GetRHI();
-		ViewUniformShaderParameters.LTCAmpTexture = GSystemTextures.LTCAmp->GetRHI();
-		ViewUniformShaderParameters.LTCAmpSampler = TStaticSamplerState<SF_Bilinear, AM_Clamp, AM_Clamp, AM_Clamp>::GetRHI();
-
+		ViewUniformShaderParameters.GGXLTCMatTexture = GSystemTextures.GGXLTCMat->GetRHI();
+		ViewUniformShaderParameters.GGXLTCAmpTexture = GSystemTextures.GGXLTCAmp->GetRHI();
+	}
+	if (GSystemTextures.SheenLTC.IsValid())
+	{
+		ViewUniformShaderParameters.SheenLTCTexture = GSystemTextures.SheenLTC->GetRHI();
 	}
-	ViewUniformShaderParameters.LTCMatTexture = OrBlack2DIfNull(ViewUniformShaderParameters.LTCMatTexture);
-	ViewUniformShaderParameters.LTCAmpTexture = OrBlack2DIfNull(ViewUniformShaderParameters.LTCAmpTexture);
+	ViewUniformShaderParameters.GGXLTCMatTexture	= OrBlack2DIfNull(ViewUniformShaderParameters.GGXLTCMatTexture);
+	ViewUniformShaderParameters.GGXLTCAmpTexture	= OrBlack2DIfNull(ViewUniformShaderParameters.GGXLTCAmpTexture);
+	ViewUniformShaderParameters.SheenLTCTexture		= OrBlack2DIfNull(ViewUniformShaderParameters.SheenLTCTexture);
+	ViewUniformShaderParameters.GGXLTCMatSampler	= TStaticSamplerState<SF_Bilinear, AM_Clamp, AM_Clamp, AM_Clamp>::GetRHI();
+	ViewUniformShaderParameters.GGXLTCAmpSampler	= TStaticSamplerState<SF_Bilinear, AM_Clamp, AM_Clamp, AM_Clamp>::GetRHI();
+	ViewUniformShaderParameters.SheenLTCSampler		= TStaticSamplerState<SF_Bilinear, AM_Clamp, AM_Clamp, AM_Clamp>::GetRHI();
 
 	// Rect light. atlas
 	{
diff --git a/Engine/Source/Runtime/Renderer/Private/SystemTextures.cpp b/Engine/Source/Runtime/Renderer/Private/SystemTextures.cpp
--- a/Engine/Source/Runtime/Renderer/Private/SystemTextures.cpp
+++ b/Engine/Source/Runtime/Renderer/Private/SystemTextures.cpp
@@ -673,11 +673,12 @@ void FSystemTextures::InitializeFeatureLevelDependentTextures(FRHICommandListImm
 
 		} // end Create the PerlinNoise3D texture
 
-		// LTC Textures	(used by Rect Lights)
+		// LTC Textures	(used by Rect Lights and/or BSDF evaluation)
 		{
+			// GGX - LTC matrix coefficients (4-coefficients)
 			{
 				const FRHITextureCreateDesc Desc =
-					FRHITextureCreateDesc::Create2D(TEXT("LTCMat"), LTC_Size, LTC_Size, PF_FloatRGBA)
+					FRHITextureCreateDesc::Create2D(TEXT("GGX.LTCMat"), GGX_LTC_Size, GGX_LTC_Size, PF_FloatRGBA)
 					.SetFlags(ETextureCreateFlags::ShaderResource | ETextureCreateFlags::FastVRAM)
 					.SetClassName(SystemTexturesName);
 
@@ -687,26 +688,27 @@ void FSystemTextures::InitializeFeatureLevelDependentTextures(FRHICommandListImm
 				uint32 DestStride;
 				uint8* DestBuffer = (uint8*)RHICmdList.LockTexture2D(Texture, 0, RLM_WriteOnly, DestStride, false);
     
-				for (int32 y = 0; y < LTC_Size; ++y)
+				for (int32 y = 0; y < GGX_LTC_Size; ++y)
 				{
-					for (int32 x = 0; x < LTC_Size; ++x)
+					for (int32 x = 0; x < GGX_LTC_Size; ++x)
 					{
 						uint16* Dest = (uint16*)(DestBuffer + x * 4 * sizeof(uint16) + y * DestStride);
     
 						for (int k = 0; k < 4; k++)
 						{
-							Dest[k] = FFloat16(LTC_Mat[4 * (x + y * LTC_Size) + k]).Encoded;
+							Dest[k] = FFloat16(GGX_LTC_Mat[4 * (x + y * GGX_LTC_Size) + k]).Encoded;
 						}
 					}
 				}
 				RHICmdList.UnlockTexture2D(Texture, 0, false);
 
-				LTCMat = CreateRenderTarget(Texture, Desc.DebugName);
+				GGXLTCMat = CreateRenderTarget(Texture, Desc.DebugName);
 			}
 
+			// GGX - Split-Sum Amplitude coefficients (2-components)
 			{
 				const FRHITextureCreateDesc Desc =
-					FRHITextureCreateDesc::Create2D(TEXT("LTCAmp"), LTC_Size, LTC_Size, PF_G16R16F)
+					FRHITextureCreateDesc::Create2D(TEXT("GGX.LTCAmp"), GGX_LTC_Size, GGX_LTC_Size, PF_G16R16F)
 					.SetFlags(ETextureCreateFlags::ShaderResource | ETextureCreateFlags::FastVRAM)
 					.SetClassName(SystemTexturesName);
 
@@ -716,21 +718,49 @@ void FSystemTextures::InitializeFeatureLevelDependentTextures(FRHICommandListImm
 				uint32 DestStride;
 				uint8* DestBuffer = (uint8*)RHICmdList.LockTexture2D(Texture, 0, RLM_WriteOnly, DestStride, false);
 
-				for (int32 y = 0; y < LTC_Size; ++y)
+				for (int32 y = 0; y < GGX_LTC_Size; ++y)
 				{
-					for (int32 x = 0; x < LTC_Size; ++x)
+					for (int32 x = 0; x < GGX_LTC_Size; ++x)
 					{
 						uint16* Dest = (uint16*)(DestBuffer + x * 2 * sizeof(uint16) + y * DestStride);
 
 						for (int k = 0; k < 2; k++)
 						{
-							Dest[k] = FFloat16(LTC_Amp[4 * (x + y * LTC_Size) + k]).Encoded;
+							Dest[k] = FFloat16(GGX_LTC_Amp[4 * (x + y * GGX_LTC_Size) + k]).Encoded;
 						}
 					}
 				}
 				RHICmdList.UnlockTexture2D(Texture, 0, false);
 
-				LTCAmp = CreateRenderTarget(Texture, Desc.DebugName);
+				GGXLTCAmp = CreateRenderTarget(Texture, Desc.DebugName);
+			}
+
+			// Sheen - Matrix & directional albedo (3-components)
+			{
+				const FRHITextureCreateDesc Desc =
+					FRHITextureCreateDesc::Create2D(TEXT("Sheen.LTC"), Sheen_LTC_Size, Sheen_LTC_Size, PF_FloatRGBA)
+					.SetFlags(ETextureCreateFlags::ShaderResource | ETextureCreateFlags::FastVRAM);
+
+				FTexture2DRHIRef Texture = RHICreateTexture(Desc);
+
+				// Write the contents of the texture.
+				uint32 DestStride;
+				uint8* DestBuffer = (uint8*)RHICmdList.LockTexture2D(Texture, 0, RLM_WriteOnly, DestStride, false);
+
+				for (int32 y = 0; y < Sheen_LTC_Size; ++y)
+				{
+					for (int32 x = 0; x < Sheen_LTC_Size; ++x)
+					{
+						uint16* Dest = (uint16*)(DestBuffer + x * 4 * sizeof(uint16) + y * DestStride);
+						Dest[0] = FFloat16(Sheen_LTC_Volume[x][y].X).Encoded;
+						Dest[1] = FFloat16(Sheen_LTC_Volume[x][y].Y).Encoded;
+						Dest[2] = FFloat16(Sheen_LTC_Volume[x][y].Z).Encoded;
+						Dest[3] = 0;
+					}
+				}
+				RHICmdList.UnlockTexture2D(Texture, 0, false);
+
+				SheenLTC = CreateRenderTarget(Texture, Desc.DebugName);
 			}
 		}
 	}
@@ -917,8 +947,9 @@ void FSystemTextures::ReleaseRHI()
 	HairLUT0.SafeRelease();
 	HairLUT1.SafeRelease();
 	HairLUT2.SafeRelease();
-	LTCMat.SafeRelease();
-	LTCAmp.SafeRelease();
+	GGXLTCMat.SafeRelease();
+	GGXLTCAmp.SafeRelease();
+	SheenLTC.SafeRelease();
 	MaxFP16Depth.SafeRelease();
 	DepthDummy.SafeRelease();
 	GreenDummy.SafeRelease();
diff --git a/Engine/Source/Runtime/Renderer/Public/SystemTextures.h b/Engine/Source/Runtime/Renderer/Public/SystemTextures.h
--- a/Engine/Source/Runtime/Renderer/Public/SystemTextures.h
+++ b/Engine/Source/Runtime/Renderer/Public/SystemTextures.h
@@ -128,9 +128,10 @@ public:
 	TRefCountPtr<IPooledRenderTarget> HairLUT0;
 	TRefCountPtr<IPooledRenderTarget> HairLUT1;
 	TRefCountPtr<IPooledRenderTarget> HairLUT2;
-	/** Linearly Transformed Cosines LUTs */
-	TRefCountPtr<IPooledRenderTarget> LTCMat;
-	TRefCountPtr<IPooledRenderTarget> LTCAmp;
+	/** GGX/Sheen Linearly Transformed Cosines LUTs */
+	TRefCountPtr<IPooledRenderTarget> GGXLTCMat;
+	TRefCountPtr<IPooledRenderTarget> GGXLTCAmp;
+	TRefCountPtr<IPooledRenderTarget> SheenLTC;
 	/** Texture that holds a single value containing the maximum depth that can be stored as FP16. */
 	TRefCountPtr<IPooledRenderTarget> MaxFP16Depth;
 	/** Depth texture that holds a single depth value */
