HeterogeneousVolumes: Fixing improper macro declaration for some global shaders.

diff --git a/Engine/Source/Runtime/Renderer/Private/HeterogeneousVolumes/HeterogeneousVolumesPreshadingPipeline.cpp b/Engine/Source/Runtime/Renderer/Private/HeterogeneousVolumes/HeterogeneousVolumesPreshadingPipeline.cpp
--- a/Engine/Source/Runtime/Renderer/Private/HeterogeneousVolumes/HeterogeneousVolumesPreshadingPipeline.cpp
+++ b/Engine/Source/Runtime/Renderer/Private/HeterogeneousVolumes/HeterogeneousVolumesPreshadingPipeline.cpp
@@ -78,7 +78,7 @@ class FGenerateRayMarchingTiles : public FGlobalShader
 	static int32 GetThreadGroupSize2D() { return 8; }
 };
 
-IMPLEMENT_MATERIAL_SHADER_TYPE(, FGenerateRayMarchingTiles, TEXT("/Engine/Private/HeterogeneousVolumes/HeterogeneousVolumesPreshadingPipeline.usf"), TEXT("GenerateRayMarchingTiles"), SF_Compute);
+IMPLEMENT_GLOBAL_SHADER(FGenerateRayMarchingTiles, "/Engine/Private/HeterogeneousVolumes/HeterogeneousVolumesPreshadingPipeline.usf", "GenerateRayMarchingTiles", SF_Compute);
 
 namespace HeterogeneousVolumes
 {
@@ -280,7 +280,7 @@ class FRenderLightingCacheWithPreshadingCS : public FGlobalShader
 	static int32 GetThreadGroupSize3D() { return 4; }
 };
 
-IMPLEMENT_MATERIAL_SHADER_TYPE(, FRenderLightingCacheWithPreshadingCS, TEXT("/Engine/Private/HeterogeneousVolumes/HeterogeneousVolumesPreshadingPipeline.usf"), TEXT("RenderLightingCacheWithPreshadingCS"), SF_Compute);
+IMPLEMENT_GLOBAL_SHADER(FRenderLightingCacheWithPreshadingCS, "/Engine/Private/HeterogeneousVolumes/HeterogeneousVolumesPreshadingPipeline.usf", "RenderLightingCacheWithPreshadingCS", SF_Compute);
 
 class FRenderSingleScatteringWithPreshadingCS : public FGlobalShader
 {
@@ -376,7 +376,7 @@ class FRenderSingleScatteringWithPreshadingCS : public FGlobalShader
 	static int32 GetThreadGroupSize2D() { return 8; }
 };
 
-IMPLEMENT_MATERIAL_SHADER_TYPE(, FRenderSingleScatteringWithPreshadingCS, TEXT("/Engine/Private/HeterogeneousVolumes/HeterogeneousVolumesPreshadingPipeline.usf"), TEXT("RenderSingleScatteringWithPreshadingCS"), SF_Compute);
+IMPLEMENT_GLOBAL_SHADER(FRenderSingleScatteringWithPreshadingCS, "/Engine/Private/HeterogeneousVolumes/HeterogeneousVolumesPreshadingPipeline.usf", "RenderSingleScatteringWithPreshadingCS", SF_Compute);
 
 void RenderLightingCacheWithPreshadingCompute(
 	FRDGBuilder& GraphBuilder,
diff --git a/Engine/Source/Runtime/Renderer/Private/HeterogeneousVolumes/HeterogeneousVolumesVoxelGridPipeline.cpp b/Engine/Source/Runtime/Renderer/Private/HeterogeneousVolumes/HeterogeneousVolumesVoxelGridPipeline.cpp
--- a/Engine/Source/Runtime/Renderer/Private/HeterogeneousVolumes/HeterogeneousVolumesVoxelGridPipeline.cpp
+++ b/Engine/Source/Runtime/Renderer/Private/HeterogeneousVolumes/HeterogeneousVolumesVoxelGridPipeline.cpp
@@ -357,7 +357,7 @@ class FMarkTopLevelGridVoxelsForFrustumGrid : public FGlobalShader
 	static int32 GetThreadGroupSize3D() { return 4; }
 };
 
-IMPLEMENT_MATERIAL_SHADER_TYPE(, FMarkTopLevelGridVoxelsForFrustumGrid, TEXT("/Engine/Private/HeterogeneousVolumes/HeterogeneousVolumesFrustumVoxelGrid.usf"), TEXT("MarkTopLevelGridVoxelsForFrustumGridCS"), SF_Compute);
+IMPLEMENT_GLOBAL_SHADER(FMarkTopLevelGridVoxelsForFrustumGrid, "/Engine/Private/HeterogeneousVolumes/HeterogeneousVolumesFrustumVoxelGrid.usf", "MarkTopLevelGridVoxelsForFrustumGridCS", SF_Compute);
 
 class FRasterizeBottomLevelFrustumGridCS : public FMeshMaterialShader
 {
@@ -517,7 +517,7 @@ class FTopLevelGridCalculateVoxelSize : public FGlobalShader
 	static int32 GetThreadGroupSize3D() { return 4; }
 };
 
-IMPLEMENT_MATERIAL_SHADER_TYPE(, FTopLevelGridCalculateVoxelSize, TEXT("/Engine/Private/HeterogeneousVolumes/HeterogeneousVolumesVoxelGridPipeline.usf"), TEXT("TopLevelGridCalculateVoxelSize"), SF_Compute);
+IMPLEMENT_GLOBAL_SHADER(FTopLevelGridCalculateVoxelSize, "/Engine/Private/HeterogeneousVolumes/HeterogeneousVolumesVoxelGridPipeline.usf", "TopLevelGridCalculateVoxelSize", SF_Compute);
 
 class FAllocateBottomLevelGrid : public FGlobalShader
 {
@@ -567,7 +567,7 @@ class FAllocateBottomLevelGrid : public FGlobalShader
 	static int32 GetThreadGroupSize3D() { return 4; }
 };
 
-IMPLEMENT_MATERIAL_SHADER_TYPE(, FAllocateBottomLevelGrid, TEXT("/Engine/Private/HeterogeneousVolumes/HeterogeneousVolumesVoxelGridPipeline.usf"), TEXT("AllocateBottomLevelGrid"), SF_Compute);
+IMPLEMENT_GLOBAL_SHADER(FAllocateBottomLevelGrid, "/Engine/Private/HeterogeneousVolumes/HeterogeneousVolumesVoxelGridPipeline.usf", "AllocateBottomLevelGrid", SF_Compute);
 
 class FGenerateRasterTiles : public FGlobalShader
 {
@@ -613,7 +613,7 @@ class FGenerateRasterTiles : public FGlobalShader
 	static int32 GetThreadGroupSize3D() { return 4; }
 };
 
-IMPLEMENT_MATERIAL_SHADER_TYPE(, FGenerateRasterTiles, TEXT("/Engine/Private/HeterogeneousVolumes/HeterogeneousVolumesVoxelGridPipeline.usf"), TEXT("GenerateRasterTiles"), SF_Compute);
+IMPLEMENT_GLOBAL_SHADER(FGenerateRasterTiles, "/Engine/Private/HeterogeneousVolumes/HeterogeneousVolumesVoxelGridPipeline.usf", "GenerateRasterTiles", SF_Compute);
 
 class FSetRasterizeBottomLevelGridIndirectArgs : public FGlobalShader
 {
@@ -650,7 +650,7 @@ class FSetRasterizeBottomLevelGridIndirectArgs : public FGlobalShader
 	static int32 GetThreadGroupSize3D() { return 4; }
 };
 
-IMPLEMENT_MATERIAL_SHADER_TYPE(, FSetRasterizeBottomLevelGridIndirectArgs, TEXT("/Engine/Private/HeterogeneousVolumes/HeterogeneousVolumesVoxelGridPipeline.usf"), TEXT("SetRasterizeBottomLevelGridIndirectArgs"), SF_Compute);
+IMPLEMENT_GLOBAL_SHADER(FSetRasterizeBottomLevelGridIndirectArgs, "/Engine/Private/HeterogeneousVolumes/HeterogeneousVolumesVoxelGridPipeline.usf", "SetRasterizeBottomLevelGridIndirectArgs", SF_Compute);
 
 class FRasterizeBottomLevelOrthoGridCS : public FMeshMaterialShader
 {
@@ -821,7 +821,7 @@ class FTopLevelCreateBitmaskCS : public FGlobalShader
 	static int32 GetThreadGroupSize3D() { return 4; }
 };
 
-IMPLEMENT_MATERIAL_SHADER_TYPE(, FTopLevelCreateBitmaskCS, TEXT("/Engine/Private/HeterogeneousVolumes/HeterogeneousVolumesVoxelGridPipeline.usf"), TEXT("TopLevelCreateBitmaskCS"), SF_Compute);
+IMPLEMENT_GLOBAL_SHADER(FTopLevelCreateBitmaskCS, "/Engine/Private/HeterogeneousVolumes/HeterogeneousVolumesVoxelGridPipeline.usf", "TopLevelCreateBitmaskCS", SF_Compute);
 
 void CreateTopLevelBitmask(
 	FRDGBuilder& GraphBuilder,
@@ -912,7 +912,7 @@ class FBuildMajorantVoxelGridCS : public FGlobalShader
 	static int32 GetThreadGroupSize3D() { return 4; }
 };
 
-IMPLEMENT_MATERIAL_SHADER_TYPE(, FBuildMajorantVoxelGridCS, TEXT("/Engine/Private/HeterogeneousVolumes/HeterogeneousVolumesVoxelGridPipeline.usf"), TEXT("BuildMajorantVoxelGridCS"), SF_Compute);
+IMPLEMENT_GLOBAL_SHADER(FBuildMajorantVoxelGridCS, "/Engine/Private/HeterogeneousVolumes/HeterogeneousVolumesVoxelGridPipeline.usf", "BuildMajorantVoxelGridCS", SF_Compute);
 
 class FDownsampleMajorantVoxelGridCS : public FGlobalShader
 {
@@ -948,7 +948,7 @@ class FDownsampleMajorantVoxelGridCS : public FGlobalShader
 	static int32 GetThreadGroupSize3D() { return 4; }
 };
 
-IMPLEMENT_MATERIAL_SHADER_TYPE(, FDownsampleMajorantVoxelGridCS, TEXT("/Engine/Private/HeterogeneousVolumes/HeterogeneousVolumesVoxelGridPipeline.usf"), TEXT("DownsampleMajorantVoxelGridCS"), SF_Compute);
+IMPLEMENT_GLOBAL_SHADER(FDownsampleMajorantVoxelGridCS, "/Engine/Private/HeterogeneousVolumes/HeterogeneousVolumesVoxelGridPipeline.usf", "DownsampleMajorantVoxelGridCS", SF_Compute);
 
 class FCopyMaxIntoMajorantVoxelGridCS : public FGlobalShader
 {
@@ -983,7 +983,7 @@ class FCopyMaxIntoMajorantVoxelGridCS : public FGlobalShader
 	static int32 GetThreadGroupSize3D() { return 4; }
 };
 
-IMPLEMENT_MATERIAL_SHADER_TYPE(, FCopyMaxIntoMajorantVoxelGridCS, TEXT("/Engine/Private/HeterogeneousVolumes/HeterogeneousVolumesVoxelGridPipeline.usf"), TEXT("CopyMaximumIntoMajorantVoxelGridCS"), SF_Compute);
+IMPLEMENT_GLOBAL_SHADER(FCopyMaxIntoMajorantVoxelGridCS, "/Engine/Private/HeterogeneousVolumes/HeterogeneousVolumesVoxelGridPipeline.usf", "CopyMaximumIntoMajorantVoxelGridCS", SF_Compute);
 
 
 void BuildMajorantVoxelGrid(
@@ -1167,7 +1167,7 @@ class FRenderTransmittanceWithVoxelGridCS : public FGlobalShader
 	static int32 GetThreadGroupSize3D() { return 4; }
 };
 
-IMPLEMENT_MATERIAL_SHADER_TYPE(, FRenderTransmittanceWithVoxelGridCS, TEXT("/Engine/Private/HeterogeneousVolumes/HeterogeneousVolumesVoxelGridRendering.usf"), TEXT("RenderTransmittanceWithVoxelGridCS"), SF_Compute);
+IMPLEMENT_GLOBAL_SHADER(FRenderTransmittanceWithVoxelGridCS, "/Engine/Private/HeterogeneousVolumes/HeterogeneousVolumesVoxelGridRendering.usf", "RenderTransmittanceWithVoxelGridCS", SF_Compute);
 
 void CalcViewBoundsAndMinimumVoxelSize(
 	FRDGBuilder& GraphBuilder,
