Speculative fix for render thread memory stomp - Rollback of DF Streaming refactor (initial CL 22558164)

#rnx
#tests QA hammered the build 0/80 without repro of the memory stomp

diff --git a/Engine/Source/Runtime/Renderer/Private/DistanceFieldStreaming.cpp b/Engine/Source/Runtime/Renderer/Private/DistanceFieldStreaming.cpp
--- a/Engine/Source/Runtime/Renderer/Private/DistanceFieldStreaming.cpp
+++ b/Engine/Source/Runtime/Renderer/Private/DistanceFieldStreaming.cpp
@@ -452,18 +452,17 @@ struct FDistanceFieldReadRequest
 	const uint8* AlwaysLoadedDataPtr = nullptr;
 
 	// Inputs of read request
-	const FBulkData* BulkData = nullptr;
+	const FByteBulkData* BulkData = nullptr;
 	uint32 BulkOffset = 0;
 	uint32 BulkSize = 0;
 
-#if !WITH_EDITOR
 	// Outputs of read request
-	FBulkDataBatchReadRequest RequestHandle;
-	FIoBuffer RequestBuffer;
-#endif
+	uint8* ReadOutputDataPtr = nullptr;
+	FIoRequest Request;
+	IAsyncReadFileHandle* AsyncHandle = nullptr;
+	IAsyncReadRequest* AsyncRequest = nullptr;
 };
 
-
 struct FDistanceFieldAsyncUpdateParameters
 {
 	FDistanceFieldSceneData* DistanceFieldSceneData = nullptr;
@@ -505,13 +504,13 @@ FDistanceFieldSceneData::~FDistanceFieldSceneData()
 class FDistanceFieldStreamingUpdateTask
 {
 public:
-	explicit FDistanceFieldStreamingUpdateTask(FDistanceFieldAsyncUpdateParameters&& InParams) : Parameters(MoveTemp(InParams)) {}
+	explicit FDistanceFieldStreamingUpdateTask(const FDistanceFieldAsyncUpdateParameters& InParams) : Parameters(InParams) {}
 
 	FDistanceFieldAsyncUpdateParameters Parameters;
 
 	void DoTask(ENamedThreads::Type CurrentThread, const FGraphEventRef& MyCompletionGraphEvent)
 	{
-		Parameters.DistanceFieldSceneData->AsyncUpdate(MoveTemp(Parameters));
+		Parameters.DistanceFieldSceneData->AsyncUpdate(Parameters);
 	}
 
 	static ESubsequentsMode::Type	GetSubsequentsMode()	{ return ESubsequentsMode::TrackSubsequents; }
@@ -519,7 +518,7 @@ public:
 	FORCEINLINE TStatId				GetStatId() const		{ return TStatId(); }
 };
 
-void FDistanceFieldSceneData::AsyncUpdate(FDistanceFieldAsyncUpdateParameters&& UpdateParameters)
+void FDistanceFieldSceneData::AsyncUpdate(FDistanceFieldAsyncUpdateParameters UpdateParameters)
 {
 	QUICK_SCOPE_CYCLE_COUNTER(STAT_FDistanceFieldSceneData_AsyncUpdate);
 	TRACE_CPUPROFILER_EVENT_SCOPE(FDistanceFieldSceneData::AsyncUpdate);
@@ -531,7 +530,7 @@ void FDistanceFieldSceneData::AsyncUpdate(FDistanceFieldAsyncUpdateParameters&&
 	int32 BrickUploadIndex = 0;
 	int32 IndirectionUploadIndex = 0;
 
-	for (FDistanceFieldReadRequest& ReadRequest : UpdateParameters.ReadRequestsToUpload)
+	for (FDistanceFieldReadRequest ReadRequest : UpdateParameters.ReadRequestsToUpload)
 	{
 		const FDistanceFieldAssetState& AssetState = AssetStateArray[ReadRequest.AssetSetId];
 		const int32 ReversedMipIndex = ReadRequest.ReversedMipIndex;
@@ -539,18 +538,15 @@ void FDistanceFieldSceneData::AsyncUpdate(FDistanceFieldAsyncUpdateParameters&&
 		const int32 MipIndex = AssetState.BuiltData->Mips.Num() - ReversedMipIndex - 1;
 		const FSparseDistanceFieldMip& MipBuiltData = AssetState.BuiltData->Mips[MipIndex];
 
-		const uint8* BulkDataReadPtr = ReadRequest.AlwaysLoadedDataPtr;
+		const uint8* BulkDataReadPtr = ReadRequest.BulkData ? ReadRequest.ReadOutputDataPtr : ReadRequest.AlwaysLoadedDataPtr;
 
-		if (ReadRequest.BulkData)
-		{
 #if WITH_EDITOR
+		if (ReadRequest.BulkData)
+		{	
 			check((ReadRequest.BulkData->IsBulkDataLoaded() ||ReadRequest.BulkData->CanLoadFromDisk()) && ReadRequest.BulkData->GetBulkDataSize() > 0);
 			BulkDataReadPtr = (const uint8*)ReadRequest.BulkData->LockReadOnly() + ReadRequest.BulkOffset;
-#else
-			BulkDataReadPtr = ReadRequest.RequestBuffer.GetData();
-#endif
 		}
-
+#endif
 		const int32 NumIndirectionEntries = MipBuiltData.IndirectionDimensions.X * MipBuiltData.IndirectionDimensions.Y * MipBuiltData.IndirectionDimensions.Z;
 		const uint32 ExpectedBulkSize = NumIndirectionEntries * sizeof(uint32) + ReadRequest.NumDistanceFieldBricks * BrickSizeBytes;
 
@@ -651,21 +647,53 @@ void FDistanceFieldSceneData::AsyncUpdate(FDistanceFieldAsyncUpdateParameters&&
 	}
 
 #if !WITH_EDITOR
-	if (UpdateParameters.NewReadRequests.Num() > 0)
+
+	for (FDistanceFieldReadRequest ReadRequest : UpdateParameters.ReadRequestsToCleanUp)
 	{
-		FBulkDataBatchRequest::FBatchBuilder Batch = FBulkDataBatchRequest::NewBatch(UpdateParameters.NewReadRequests.Num());
-		for (FDistanceFieldReadRequest& ReadRequest : UpdateParameters.NewReadRequests)
+		if (ReadRequest.AsyncRequest)
 		{
-			check(ReadRequest.BulkSize > 0);
-			Batch.Read(*ReadRequest.BulkData, uint64(ReadRequest.BulkOffset), uint64(ReadRequest.BulkSize), AIOP_Low, ReadRequest.RequestBuffer, ReadRequest.RequestHandle);
-			ReadRequests.Add(MoveTemp(ReadRequest));
+			check(ReadRequest.AsyncRequest->PollCompletion());	
+			delete ReadRequest.AsyncRequest;
+			delete ReadRequest.AsyncHandle;
 		}
-		FBulkDataRequest::EStatus Status = Batch.Issue();
-		UE_CLOG(Status != FBulkDataRequest::EStatus::Ok, LogDistanceField, Error, TEXT("Failed to issue bulk data I/O request"));
+		else
+		{
+			check(ReadRequest.Request.Status().IsCompleted());
+		}
+
+		FMemory::Free(ReadRequest.ReadOutputDataPtr);
 	}
-#else
-	ReadRequests.Append(UpdateParameters.NewReadRequests);
+
+	FIoBatch Batch;
+
+	for (FDistanceFieldReadRequest& ReadRequest : UpdateParameters.NewReadRequests)
+	{
+		check(ReadRequest.BulkSize > 0);
+		ReadRequest.ReadOutputDataPtr = (uint8*)FMemory::Malloc(ReadRequest.BulkSize);
+		const bool bIODispatcher = ReadRequest.BulkData->IsUsingIODispatcher();
+
+		if (bIODispatcher)
+		{
+			// Use IODispatcher when available
+			FIoChunkId ChunkID = ReadRequest.BulkData->CreateChunkId();
+			FIoReadOptions ReadOptions;
+			ReadOptions.SetRange(ReadRequest.BulkData->GetBulkDataOffsetInFile() + ReadRequest.BulkOffset, ReadRequest.BulkSize);
+			ReadOptions.SetTargetVa(ReadRequest.ReadOutputDataPtr);
+			ReadRequest.Request = Batch.Read(ChunkID, ReadOptions, IoDispatcherPriority_Low);
+		}
+		else
+		{
+			// Compatibility path without IODispatcher
+			ReadRequest.AsyncHandle = ReadRequest.BulkData->OpenAsyncReadHandle();
+			ReadRequest.AsyncRequest = ReadRequest.AsyncHandle->ReadRequest(ReadRequest.BulkData->GetBulkDataOffsetInFile() + ReadRequest.BulkOffset, ReadRequest.BulkSize, AIOP_Low, nullptr, ReadRequest.ReadOutputDataPtr);
+		}
+	}
+
+	Batch.Issue();
+
 #endif
+
+	ReadRequests.Append(UpdateParameters.NewReadRequests);
 }
 
 bool AssetHasOutstandingRequest(FSetElementId AssetSetId, const TArray<FDistanceFieldReadRequest>& ReadRequests)
@@ -810,7 +838,8 @@ void FDistanceFieldSceneData::ProcessStreamingRequestsFromGPU(
 void FDistanceFieldSceneData::ProcessReadRequests(
 	TArray<FDistanceFieldAssetMipId>& AssetDataUploads,
 	TArray<FDistanceFieldAssetMipId>& DistanceFieldAssetMipAdds,
-	TArray<FDistanceFieldReadRequest>& ReadRequestsToUpload)
+	TArray<FDistanceFieldReadRequest>& ReadRequestsToUpload,
+	TArray<FDistanceFieldReadRequest>& ReadRequestsToCleanUp)
 {
 	const uint32 BrickSizeBytes = GPixelFormats[DistanceField::DistanceFieldFormat].BlockBytes * DistanceField::BrickSize * DistanceField::BrickSize * DistanceField::BrickSize;
 	const SIZE_T TextureUploadLimitBytes = (SIZE_T)CVarTextureUploadLimitKBytes.GetValueOnRenderThread() * 1024;
@@ -840,33 +869,41 @@ void FDistanceFieldSceneData::ProcessReadRequests(
 
 	for (int32 RequestIndex = 0; RequestIndex < ReadRequests.Num(); RequestIndex++)
 	{
+		const FDistanceFieldReadRequest ReadRequest = ReadRequests[RequestIndex];
+
+		bool bReady = true;
+
 #if !WITH_EDITOR
-		if (ReadRequests[RequestIndex].RequestHandle.IsCompleted() == false)
+		if (ReadRequest.AsyncRequest)
 		{
-			continue;
+			bReady = bReady && ReadRequest.AsyncRequest->PollCompletion();
+		}
+		else
+		{
+			bReady = bReady && ReadRequest.Request.Status().IsCompleted();
 		}
 #endif
 
-		FDistanceFieldReadRequest CompletedRequest = MoveTemp(ReadRequests[RequestIndex]);
-#if !WITH_EDITOR
-		CompletedRequest.RequestHandle = FBulkDataRequest();
-#endif
-
-		ReadRequests.RemoveAtSwap(RequestIndex--);
-
-		if (AssetStateArray.IsValidId(CompletedRequest.AssetSetId) 
-			// Prevent attempting to upload after a different asset has been allocated at the same index
-			&& CompletedRequest.BuiltDataId == AssetStateArray[CompletedRequest.AssetSetId].BuiltData->GetId()
-			// Shader requires sequential reversed mips starting from 0, skip upload if the IO request got out of sync with the streaming feedback requests
-			&& CompletedRequest.ReversedMipIndex == AssetStateArray[CompletedRequest.AssetSetId].ReversedMips.Num())
+		if (bReady)
 		{
-			TextureUploadBytes += CompletedRequest.NumDistanceFieldBricks * BrickSizeBytes;
+			ReadRequests.RemoveAtSwap(RequestIndex);
+			RequestIndex--;
+
+			if (AssetStateArray.IsValidId(ReadRequest.AssetSetId) 
+				// Prevent attempting to upload after a different asset has been allocated at the same index
+				&& ReadRequest.BuiltDataId == AssetStateArray[ReadRequest.AssetSetId].BuiltData->GetId()
+				// Shader requires sequential reversed mips starting from 0, skip upload if the IO request got out of sync with the streaming feedback requests
+				&& ReadRequest.ReversedMipIndex == AssetStateArray[ReadRequest.AssetSetId].ReversedMips.Num())
+			{
+				TextureUploadBytes += ReadRequest.NumDistanceFieldBricks * BrickSizeBytes;
 
-			DistanceFieldAssetMipAdds.Add(FDistanceFieldAssetMipId(CompletedRequest.AssetSetId, CompletedRequest.ReversedMipIndex));
-			// Re-upload mip0 to push the new NumMips to the shader
-			AssetDataUploads.Add(FDistanceFieldAssetMipId(CompletedRequest.AssetSetId, 0));
+				DistanceFieldAssetMipAdds.Add(FDistanceFieldAssetMipId(ReadRequest.AssetSetId, ReadRequest.ReversedMipIndex));
+				// Re-upload mip0 to push the new NumMips to the shader
+				AssetDataUploads.Add(FDistanceFieldAssetMipId(ReadRequest.AssetSetId, 0));
+				ReadRequestsToUpload.Add(ReadRequest);
+			}
 
-			ReadRequestsToUpload.Add(MoveTemp(CompletedRequest));
+			ReadRequestsToCleanUp.Add(ReadRequest);
 		}
 
 		// Stop uploading when we reach the limit
@@ -1366,8 +1403,9 @@ void FDistanceFieldSceneData::UpdateDistanceFieldAtlas(
 	ProcessStreamingRequestsFromGPU(NewReadRequests, AssetDataUploads);
 
 	TArray<FDistanceFieldReadRequest> ReadRequestsToUpload;
+	TArray<FDistanceFieldReadRequest> ReadRequestsToCleanUp;
 	// Build a list of completed read requests that should be uploaded to the GPU this frame
-	ProcessReadRequests(AssetDataUploads, DistanceFieldAssetMipAdds, ReadRequestsToUpload);
+	ProcessReadRequests(AssetDataUploads, DistanceFieldAssetMipAdds, ReadRequestsToUpload, ReadRequestsToCleanUp);
 
 	int32 NumIndirectionTableAdds = 0;
 	int32 NumBrickUploads = 0;
@@ -1491,13 +1529,14 @@ void FDistanceFieldSceneData::UpdateDistanceFieldAtlas(
 			UpdateParameters.BrickUploadCoordinatesPtr = AtlasUpload.BrickUploadCoordinatesPtr;
 		}
 
-		if (NewReadRequests.Num() || ReadRequestsToUpload.Num())
+		if (NewReadRequests.Num() || ReadRequestsToUpload.Num() || ReadRequestsToCleanUp.Num())
 		{
 			UpdateParameters.NewReadRequests = MoveTemp(NewReadRequests);
 			UpdateParameters.ReadRequestsToUpload = MoveTemp(ReadRequestsToUpload);
+			UpdateParameters.ReadRequestsToCleanUp = MoveTemp(ReadRequestsToCleanUp);
 
 			// TODO: We actually run this synchronously now after the RDG conversion, as it would otherwise immediately sync.
-			AsyncUpdate(MoveTemp(UpdateParameters));
+			AsyncUpdate(UpdateParameters);
 		}
 
 		if (NumBrickUploads > 0 || NumIndirectionTableAdds > 0)
@@ -1550,8 +1589,8 @@ void FDistanceFieldSceneData::UpdateDistanceFieldAtlas(
 			for (int32 StartBrickIndex = 0; StartBrickIndex < NumBrickUploads; StartBrickIndex += MaxBrickUploadsPerPass)
 			{
 				const int32 NumBrickUploadsThisPass = FMath::Min(MaxBrickUploadsPerPass, NumBrickUploads - StartBrickIndex);
+				FScatterUploadDistanceFieldAtlasCS::FParameters* PassParameters = GraphBuilder.AllocParameters<FScatterUploadDistanceFieldAtlasCS::FParameters>();
 
-				auto* PassParameters = GraphBuilder.AllocParameters<FScatterUploadDistanceFieldAtlasCS::FParameters>();
 				PassParameters->RWDistanceFieldBrickAtlas = GraphBuilder.CreateUAV(DistanceFieldBrickVolumeTextureRDG);
 				PassParameters->BrickUploadCoordinates = AtlasUpload.BrickUploadCoordinatesBuffer.SRV;
 				PassParameters->BrickUploadData = AtlasUpload.BrickUploadDataBuffer.SRV;
diff --git a/Engine/Source/Runtime/Renderer/Private/ScenePrivate.h b/Engine/Source/Runtime/Renderer/Private/ScenePrivate.h
--- a/Engine/Source/Runtime/Renderer/Private/ScenePrivate.h
+++ b/Engine/Source/Runtime/Renderer/Private/ScenePrivate.h
@@ -2114,7 +2114,8 @@ private:
 	void ProcessReadRequests(
 		TArray<FDistanceFieldAssetMipId>& AssetDataUploads,
 		TArray<FDistanceFieldAssetMipId>& DistanceFieldAssetMipAdds,
-		TArray<FDistanceFieldReadRequest>& ReadRequestsToUpload);
+		TArray<FDistanceFieldReadRequest>& ReadRequestsToUpload,
+		TArray<FDistanceFieldReadRequest>& ReadRequestsToCleanUp);
 
 	FRDGTexture* ResizeBrickAtlasIfNeeded(FRDGBuilder& GraphBuilder, FGlobalShaderMap* GlobalShaderMap);
 
@@ -2126,7 +2127,7 @@ private:
 	
 	void UploadAllAssetData(FRDGBuilder& GraphBuilder, FRDGBuffer* AssetDataBufferRDG);
 
-	void AsyncUpdate(FDistanceFieldAsyncUpdateParameters&& UpdateParameters);
+	void AsyncUpdate(FDistanceFieldAsyncUpdateParameters UpdateParameters);
 
 	void GenerateStreamingRequests(
 		FRDGBuilder& GraphBuilder, 
