Update per-view AO history and copy back to output AO.

AO history may not yet account for the scene velocity textures when they are available (for multiple views).

diff --git a/Engine/Source/Runtime/Renderer/Private/DistanceFieldAmbientOcclusion.cpp b/Engine/Source/Runtime/Renderer/Private/DistanceFieldAmbientOcclusion.cpp
--- a/Engine/Source/Runtime/Renderer/Private/DistanceFieldAmbientOcclusion.cpp
+++ b/Engine/Source/Runtime/Renderer/Private/DistanceFieldAmbientOcclusion.cpp
@@ -810,6 +810,13 @@ void FDeferredShadingSceneRenderer::RenderDistanceFieldLighting(
 
 	FRDGTextureRef BentNormalOutput = nullptr;
 
+	{
+		const FIntPoint BufferSize = GetActiveSceneTexturesConfig().Extent / GAODownsampleFactor;
+		const FRDGTextureDesc Desc = FRDGTextureDesc::Create2D(BufferSize, PF_FloatRGBA, FClearValueBinding::None, GFastVRamConfig.DistanceFieldAOBentNormal | TexCreate_ShaderResource | TexCreate_RenderTargetable | TexCreate_UAV);
+		BentNormalOutput = GraphBuilder.CreateTexture(Desc, TEXT("DistanceFieldBentNormalAO"));
+		AddClearUAVPass(GraphBuilder, GraphBuilder.CreateUAV(BentNormalOutput), FLinearColor::Black);
+	}
+
 	for (int32 ViewIndex = 0; ViewIndex < Views.Num(); ++ViewIndex)
 	{
 		const FViewInfo& View = Views[ViewIndex];
@@ -851,6 +858,10 @@ void FDeferredShadingSceneRenderer::RenderDistanceFieldLighting(
 			BuildTileObjectLists(GraphBuilder, Scene, View, SceneTextures.UniformBuffer, ObjectIndirectArguments, CulledObjectBufferParameters, TileIntersectionParameters, DistanceFieldNormal, Parameters);
 		}
 
+		// Render to a per-view BentNormal first because it also needs to be stored by the view history.
+		// The per-view output can be copied back to the appropriate region in the output BentNormal texture.
+		FRDGTextureRef PerViewBentNormal = nullptr;
+
 		RenderDistanceFieldAOScreenGrid(
 			GraphBuilder,
 			SceneTextures,
@@ -860,7 +871,9 @@ void FDeferredShadingSceneRenderer::RenderDistanceFieldLighting(
 			TileIntersectionParameters,
 			Parameters,
 			DistanceFieldNormal,
-			BentNormalOutput);
+			PerViewBentNormal);
+
+		AddCopyTexturePass(GraphBuilder, PerViewBentNormal, BentNormalOutput, FIntPoint::ZeroValue, View.ViewRect.Min / GAODownsampleFactor, View.ViewRect.Size() / GAODownsampleFactor);
 	}
 
 	RenderCapsuleShadowsForMovableSkylight(GraphBuilder, SceneTextures.UniformBuffer, BentNormalOutput);
diff --git a/Engine/Source/Runtime/Renderer/Private/DistanceFieldLightingPost.cpp b/Engine/Source/Runtime/Renderer/Private/DistanceFieldLightingPost.cpp
--- a/Engine/Source/Runtime/Renderer/Private/DistanceFieldLightingPost.cpp
+++ b/Engine/Source/Runtime/Renderer/Private/DistanceFieldLightingPost.cpp
@@ -320,7 +320,7 @@ void UpdateHistory(
 {
 	const FRDGSystemTextures& SystemTextures = FRDGSystemTextures::Get(GraphBuilder);
 
-	const FIntPoint SceneTextureExtent = View.GetSceneTexturesConfig().Extent;
+	const FIntPoint AOHistoryExtent = View.ViewRect.Size();
 
 	if (BentNormalHistoryState && DistanceFieldAOUseHistory(View))
 	{
@@ -358,7 +358,7 @@ void UpdateHistory(
 				PassParameters->UseHistoryFilter = UseAOHistoryStabilityPass() ? 1.0f : 0.0f;
 
 				{
-					FIntPoint HistoryBufferSize = View.GetSceneTexturesConfig().Extent / FIntPoint(GAODownsampleFactor, GAODownsampleFactor);
+					FIntPoint HistoryBufferSize = AOHistoryExtent / FIntPoint(GAODownsampleFactor, GAODownsampleFactor);
 
 					const float InvBufferSizeX = 1.0f / HistoryBufferSize.X;
 					const float InvBufferSizeY = 1.0f / HistoryBufferSize.Y;
@@ -391,7 +391,7 @@ void UpdateHistory(
 					RDG_EVENT_NAME("UpdateHistory"),
 					PassParameters,
 					ERDGPassFlags::Raster,
-					[PassParameters, VertexShader, PixelShader, &View, SceneTextureExtent]
+					[PassParameters, VertexShader, PixelShader, &View, AOHistoryExtent]
 					(FRHICommandList& RHICmdList)
 				{
 					RHICmdList.SetViewport(0, 0, 0.0f, View.ViewRect.Width() / GAODownsampleFactor, View.ViewRect.Height() / GAODownsampleFactor, 1.0f);
@@ -417,7 +417,7 @@ void UpdateHistory(
 						0, 0,
 						View.ViewRect.Width() / GAODownsampleFactor, View.ViewRect.Height() / GAODownsampleFactor,
 						FIntPoint(View.ViewRect.Width() / GAODownsampleFactor, View.ViewRect.Height() / GAODownsampleFactor),
-						SceneTextureExtent / FIntPoint(GAODownsampleFactor, GAODownsampleFactor),
+						AOHistoryExtent / FIntPoint(GAODownsampleFactor, GAODownsampleFactor),
 						VertexShader);
 				});
 			}
@@ -427,7 +427,7 @@ void UpdateHistory(
 				const FRDGTextureDesc& HistoryDesc = BentNormalHistoryTexture->Desc;
 
 				// Reallocate history if buffer sizes have changed
-				if (HistoryDesc.Extent != SceneTextureExtent / FIntPoint(GAODownsampleFactor, GAODownsampleFactor))
+				if (HistoryDesc.Extent != AOHistoryExtent / FIntPoint(GAODownsampleFactor, GAODownsampleFactor))
 				{
 					GRenderTargetPool.FreeUnusedResource(*BentNormalHistoryState);
 					*BentNormalHistoryState = nullptr;
@@ -435,7 +435,8 @@ void UpdateHistory(
 					AllocateOrReuseAORenderTarget(GraphBuilder, View, BentNormalHistoryTexture, BentNormalHistoryRTName, PF_FloatRGBA);
 				}
 
-				const bool bManuallyClampUV = View.ViewRect.Min != FIntPoint::ZeroValue || View.ViewRect.Max != SceneTextureExtent;
+				//@to-do: this might not be necessary anymore
+				const bool bManuallyClampUV = View.ViewRect.Min != FIntPoint::ZeroValue || View.ViewRect.Max != AOHistoryExtent;
 
 				FFilterHistoryPS::FPermutationDomain PermutationVector;
 				PermutationVector.Set<FFilterHistoryPS::FManuallyClampUV>(bManuallyClampUV);
@@ -443,7 +444,7 @@ void UpdateHistory(
 				auto VertexShader = View.ShaderMap->GetShader<FPostProcessVS>();
 				auto PixelShader = View.ShaderMap->GetShader<FFilterHistoryPS>(PermutationVector);
 
-				const FIntPoint DownsampledBufferSize(View.GetSceneTexturesConfig().Extent / FIntPoint(GAODownsampleFactor, GAODownsampleFactor));
+				const FIntPoint DownsampledBufferSize(AOHistoryExtent / FIntPoint(GAODownsampleFactor, GAODownsampleFactor));
 				FVector2f MaxSampleBufferUV(
 					(View.ViewRect.Width() / GAODownsampleFactor - 0.5f - GAODownsampleFactor) / DownsampledBufferSize.X,
 					(View.ViewRect.Height() / GAODownsampleFactor - 0.5f - GAODownsampleFactor) / DownsampledBufferSize.Y);
@@ -464,7 +465,7 @@ void UpdateHistory(
 					RDG_EVENT_NAME("UpdateHistoryStability"),
 					PassParameters,
 					ERDGPassFlags::Raster,
-					[PassParameters, VertexShader, PixelShader, &View, SceneTextureExtent](FRHICommandList& RHICmdList)
+					[PassParameters, VertexShader, PixelShader, &View, AOHistoryExtent](FRHICommandList& RHICmdList)
 				{
 					RHICmdList.SetViewport(0, 0, 0.0f, View.ViewRect.Width() / GAODownsampleFactor, View.ViewRect.Height() / GAODownsampleFactor, 1.0f);
 
@@ -489,7 +490,7 @@ void UpdateHistory(
 						0, 0,
 						View.ViewRect.Width() / GAODownsampleFactor, View.ViewRect.Height() / GAODownsampleFactor,
 						FIntPoint(View.ViewRect.Width() / GAODownsampleFactor, View.ViewRect.Height() / GAODownsampleFactor),
-						SceneTextureExtent / FIntPoint(GAODownsampleFactor, GAODownsampleFactor),
+						AOHistoryExtent / FIntPoint(GAODownsampleFactor, GAODownsampleFactor),
 						VertexShader);
 				});
 
@@ -507,7 +508,7 @@ void UpdateHistory(
 		{
 			// Use the current frame's upscaled mask for next frame's history
 			FRDGTextureRef DistanceFieldAOBentNormal = nullptr;
-			AllocateOrReuseAORenderTarget(GraphBuilder, View, DistanceFieldAOBentNormal, TEXT("DistanceFieldBentNormalAO"), PF_FloatRGBA, GFastVRamConfig.DistanceFieldAOBentNormal);
+			AllocateOrReuseAORenderTarget(GraphBuilder, View, DistanceFieldAOBentNormal, TEXT("PerViewDistanceFieldBentNormalAO"), PF_FloatRGBA, GFastVRamConfig.DistanceFieldAOHistory);
 
 			GeometryAwareUpsample(GraphBuilder, View, DistanceFieldAOBentNormal, DistanceFieldNormal, BentNormalInterpolation, Parameters);
 
@@ -516,14 +517,13 @@ void UpdateHistory(
 		}
 
 		DistanceFieldAOHistoryViewRect->Min = FIntPoint::ZeroValue;
-		DistanceFieldAOHistoryViewRect->Max.X = View.ViewRect.Size().X / GAODownsampleFactor;
-		DistanceFieldAOHistoryViewRect->Max.Y = View.ViewRect.Size().Y / GAODownsampleFactor;
+		DistanceFieldAOHistoryViewRect->Max = View.ViewRect.Size() / FIntPoint(GAODownsampleFactor, GAODownsampleFactor);
 	}
 	else
 	{
 		// Temporal reprojection is disabled or there is no view state - just upscale
 		FRDGTextureRef DistanceFieldAOBentNormal = nullptr;
-		AllocateOrReuseAORenderTarget(GraphBuilder, View, DistanceFieldAOBentNormal, TEXT("DistanceFieldBentNormalAO"), PF_FloatRGBA, GFastVRamConfig.DistanceFieldAOBentNormal);
+		AllocateOrReuseAORenderTarget(GraphBuilder, View, DistanceFieldAOBentNormal, TEXT("PerViewDistanceFieldBentNormalAO"), PF_FloatRGBA, GFastVRamConfig.DistanceFieldAOHistory);
 
 		GeometryAwareUpsample(GraphBuilder, View, DistanceFieldAOBentNormal, DistanceFieldNormal, BentNormalInterpolation, Parameters);
 
