Add FLocalLightPrepassCS
Modify RenderLocalLightPrepass from fullscreenpass to instance draw of screen tiles
Restrict MobileLocalLightPrepass texture allocation only when LocalLight Prepass is enabled

diff --git a/Engine/Shaders/Private/MobileLocalLightPrepass.usf b/Engine/Shaders/Private/MobileLocalLightPrepass.usf
--- a/Engine/Shaders/Private/MobileLocalLightPrepass.usf
+++ b/Engine/Shaders/Private/MobileLocalLightPrepass.usf
@@ -7,6 +7,72 @@
 #include "LightGridCommon.ush"
 #include "DeferredLightingCommon.ush"
 
+uint2 GroupSize;
+
+RWBuffer<int> RWTileInfo;
+[numthreads(THREADGROUP_SIZEX, 1, 1)]
+void MainCS(
+	uint3 GroupId : SV_GroupID,
+	uint3 DispatchThreadId : SV_DispatchThreadID,
+	uint3 GroupThreadId : SV_GroupThreadID)
+{
+	if(DispatchThreadId.x > (GroupSize.x * GroupSize.y))
+	{
+		return;
+	}
+
+	const uint GridX = DispatchThreadId.x % GroupSize.x;
+	const uint GridY = DispatchThreadId.x / GroupSize.x;
+	const uint EyeIndex = 0;
+	const FLightGridData GridData = GetLightGridData(EyeIndex);
+	uint LocalNumLights = 0;
+	LOOP
+	for (int SliceIt = 0; SliceIt < GridData.CulledGridSize.z; SliceIt++)
+	{
+		const uint GridIndex = ComputeLightGridCellIndex(uint3(GridX, GridY, SliceIt), EyeIndex);
+		const FCulledLightsGridData CulledLightsGrid = GetCulledLightsGrid(GridIndex, EyeIndex);
+		
+		LocalNumLights = CulledLightsGrid.NumLocalLights;
+
+		if(LocalNumLights > 0)
+		{
+			break;
+		}
+	}
+
+	if(LocalNumLights > 0)
+	{
+		RWTileInfo[2 * DispatchThreadId.x] = GridX;
+		RWTileInfo[2 * DispatchThreadId.x + 1] = GridY;
+	}
+	else
+	{
+		RWTileInfo[2 * DispatchThreadId.x] = -1;
+		RWTileInfo[2 * DispatchThreadId.x + 1] = -1;
+	}
+}
+
+int LightGridPixelSize;
+Buffer<int> TileInfo;
+void MainVS(
+	float2 TexCoord : ATTRIBUTE0,
+	uint VertexId : SV_VertexID,
+	uint InstanceId : SV_InstanceID,
+	out float4 OutPosition : SV_POSITION)
+{
+	float YPos = TileInfo[2 * InstanceId + 1] + TexCoord.y;
+	float XPos = TileInfo[2 * InstanceId] + TexCoord.x;
+
+	float2 ScreenUV = float2(XPos, YPos) * LightGridPixelSize * View.BufferSizeAndInvSize.zw;
+	float2 ScreenPosition = (ScreenUV.xy - View.ScreenPositionScaleBias.wz) / View.ScreenPositionScaleBias.xy;
+	OutPosition = float4(ScreenPosition, 0, 1);
+
+	if(TileInfo[2 * InstanceId] < 0)
+	{
+		OutPosition.xy = 0;
+	}
+}
+
 void Main(float4 SvPosition : SV_POSITION, out float3 OutColorA : SV_Target0, out float4 OutColorB : SV_Target1)
 {
 	const uint2 PixelPos = SvPosition.xy;
diff --git a/Engine/Source/Runtime/RenderCore/Private/RenderUtils.cpp b/Engine/Source/Runtime/RenderCore/Private/RenderUtils.cpp
--- a/Engine/Source/Runtime/RenderCore/Private/RenderUtils.cpp
+++ b/Engine/Source/Runtime/RenderCore/Private/RenderUtils.cpp
@@ -586,8 +586,8 @@ RENDERCORE_API bool IsMobileMovableSpotlightShadowsEnabled(const FStaticShaderPl
 
 RENDERCORE_API bool MobileForwardEnableLocalLights(const FStaticShaderPlatform Platform)
 {
-	static FShaderPlatformCachedIniValue<bool> MobileForwardEnableLocalLightsIniValue(TEXT("r.Mobile.Forward.EnableLocalLights"));
-	return MobileForwardEnableLocalLightsIniValue.Get(Platform);
+	static FShaderPlatformCachedIniValue<int32> MobileForwardEnableLocalLightsIniValue(TEXT("r.Mobile.Forward.EnableLocalLights"));
+	return MobileForwardEnableLocalLightsIniValue.Get(Platform) > 0;
 }
 
 RENDERCORE_API bool MobileForwardEnableClusteredReflections(const FStaticShaderPlatform Platform)
diff --git a/Engine/Source/Runtime/Renderer/Private/CapsuleShadowRendering.cpp b/Engine/Source/Runtime/Renderer/Private/CapsuleShadowRendering.cpp
--- a/Engine/Source/Runtime/Renderer/Private/CapsuleShadowRendering.cpp
+++ b/Engine/Source/Runtime/Renderer/Private/CapsuleShadowRendering.cpp
@@ -112,13 +112,6 @@ FAutoConsoleVariableRef CVarCapsuleMinSkyAngle(
 	ECVF_Scalability | ECVF_RenderThreadSafe
 	);
 
-// Nvidia has lower vertex throughput when only processing a few verts per instance
-// Disabled as it hasn't been tested
-static const int32 NumTileQuadsInBuffer = 1;
-
-TGlobalResource<FTileTexCoordVertexBuffer> GTileTexCoordVertexBuffer(NumTileQuadsInBuffer);
-TGlobalResource<FTileIndexBuffer> GTileIndexBuffer(NumTileQuadsInBuffer);
-
 const int32 GComputeLightDirectionFromVolumetricLightmapGroupSize = 64;
 
 class FComputeLightDirectionFromVolumetricLightmapCS : public FGlobalShader
@@ -319,7 +312,7 @@ class FCapsuleShadowingUpsampleVS : public FGlobalShader
 
 	static void ModifyCompilationEnvironment(const FGlobalShaderPermutationParameters& Parameters, FShaderCompilerEnvironment& OutEnvironment)
 	{
-		OutEnvironment.SetDefine(TEXT("TILES_PER_INSTANCE"), NumTileQuadsInBuffer);
+		OutEnvironment.SetDefine(TEXT("TILES_PER_INSTANCE"), 1);
 	}
 };
 
@@ -669,14 +662,14 @@ bool FDeferredShadingSceneRenderer::RenderCapsuleDirectShadows(
 					SetShaderParameters(RHICmdList, VertexShader, VertexShader.GetVertexShader(), PassParameters->VS);
 					SetShaderParameters(RHICmdList, PixelShader, PixelShader.GetPixelShader(), PassParameters->PS);
 
-					RHICmdList.SetStreamSource(0, GTileTexCoordVertexBuffer.VertexBufferRHI, 0);
-					RHICmdList.DrawIndexedPrimitive(GTileIndexBuffer.IndexBufferRHI, 
+					RHICmdList.SetStreamSource(0, GetOneTileQuadVertexBuffer(), 0);
+					RHICmdList.DrawIndexedPrimitive(GetOneTileQuadIndexBuffer(),
 						0, 
 						0, 
 						4,
 						0,
-						2 * NumTileQuadsInBuffer,
-						FMath::DivideAndRoundUp(GroupSize.X * GroupSize.Y, NumTileQuadsInBuffer));
+						2,
+						FMath::DivideAndRoundUp(GroupSize.X * GroupSize.Y, 1));
 				});
 			}
 		}
@@ -1166,14 +1159,14 @@ void FDeferredShadingSceneRenderer::RenderIndirectCapsuleShadows(FRDGBuilder& Gr
 					SetShaderParameters(RHICmdList, VertexShader, VertexShader.GetVertexShader(), PassParameters->VS);
 					SetShaderParameters(RHICmdList, PixelShader, PixelShader.GetPixelShader(), PassParameters->PS);
 
-					RHICmdList.SetStreamSource(0, GTileTexCoordVertexBuffer.VertexBufferRHI, 0);
-					RHICmdList.DrawIndexedPrimitive(GTileIndexBuffer.IndexBufferRHI,
+					RHICmdList.SetStreamSource(0, GetOneTileQuadVertexBuffer(), 0);
+					RHICmdList.DrawIndexedPrimitive(GetOneTileQuadIndexBuffer(),
 						0,
 						0,
 						4,
 						0,
-						2 * NumTileQuadsInBuffer,
-						FMath::DivideAndRoundUp(GroupSize.X * GroupSize.Y, NumTileQuadsInBuffer));
+						2,
+						FMath::DivideAndRoundUp(GroupSize.X * GroupSize.Y, 1));
 				});
 			}
 		}
diff --git a/Engine/Source/Runtime/Renderer/Private/MobileLocalLightPrepass.cpp b/Engine/Source/Runtime/Renderer/Private/MobileLocalLightPrepass.cpp
--- a/Engine/Source/Runtime/Renderer/Private/MobileLocalLightPrepass.cpp
+++ b/Engine/Source/Runtime/Renderer/Private/MobileLocalLightPrepass.cpp
@@ -5,18 +5,64 @@
 #include "BasePassRendering.h"
 #include "PixelShaderUtils.h"
 #include "MobileBasePassRendering.h"
+#include "RendererPrivateUtils.h"
+#include "GlobalRenderResources.h"
+
+const int32 GLocalLightPrepassTileSizeX = 8;
+class FLocalLightPrepassCS : public FGlobalShader
+{
+	DECLARE_GLOBAL_SHADER(FLocalLightPrepassCS);
+public:
+	SHADER_USE_PARAMETER_STRUCT(FLocalLightPrepassCS, FGlobalShader);
+
+	BEGIN_SHADER_PARAMETER_STRUCT(FParameters, )
+		SHADER_PARAMETER_RDG_BUFFER_UAV(RWBuffer<int32>, RWTileInfo)
+		SHADER_PARAMETER_RDG_UNIFORM_BUFFER(FForwardLightData, ForwardLightData)
+		SHADER_PARAMETER(FIntPoint, GroupSize)
+	END_SHADER_PARAMETER_STRUCT()
+ 
+	static bool ShouldCompilePermutation(const FGlobalShaderPermutationParameters& Parameters)
+	{
+		return IsMobilePlatform(Parameters.Platform) && MobileForwardEnablePrepassLocalLights(Parameters.Platform);
+	}
+
+	static void ModifyCompilationEnvironment(const FGlobalShaderPermutationParameters& Parameters, FShaderCompilerEnvironment& OutEnvironment)
+	{
+		FForwardLightingParameters::ModifyCompilationEnvironment(Parameters.Platform, OutEnvironment);
+		OutEnvironment.SetDefine(TEXT("THREADGROUP_SIZEX"), GLocalLightPrepassTileSizeX);
+	}
+};
+
+IMPLEMENT_GLOBAL_SHADER(FLocalLightPrepassCS, "/Engine/Private/MobileLocalLightPrepass.usf", "MainCS", SF_Compute);
+
+class FLocalLightPrepassVS : public FGlobalShader
+{
+	DECLARE_GLOBAL_SHADER(FLocalLightPrepassVS);
+	SHADER_USE_PARAMETER_STRUCT(FLocalLightPrepassVS, FGlobalShader);
+
+	BEGIN_SHADER_PARAMETER_STRUCT(FParameters, )
+		SHADER_PARAMETER_STRUCT_REF(FViewUniformShaderParameters, View)
+		SHADER_PARAMETER_RDG_BUFFER_SRV(Buffer<int32>, TileInfo)
+		SHADER_PARAMETER(int32, LightGridPixelSize)
+	END_SHADER_PARAMETER_STRUCT()
+
+	static bool ShouldCompilePermutation(const FGlobalShaderPermutationParameters& Parameters)
+	{
+		return IsMobilePlatform(Parameters.Platform) && MobileForwardEnablePrepassLocalLights(Parameters.Platform);
+	}
+};
+
+IMPLEMENT_GLOBAL_SHADER(FLocalLightPrepassVS, "/Engine/Private/MobileLocalLightPrepass.usf", "MainVS", SF_Vertex);
 
 class FLocalLightPrepassPS : public FGlobalShader
 {
 	DECLARE_GLOBAL_SHADER(FLocalLightPrepassPS);
 	SHADER_USE_PARAMETER_STRUCT(FLocalLightPrepassPS, FGlobalShader)
 
-		BEGIN_SHADER_PARAMETER_STRUCT(FParameters, )
+	BEGIN_SHADER_PARAMETER_STRUCT(FParameters, )
 		SHADER_PARAMETER_STRUCT_REF(FViewUniformShaderParameters, View)
-		SHADER_PARAMETER_STRUCT_INCLUDE(FSceneTextureShaderParameters, SceneTextures)
 		SHADER_PARAMETER_RDG_UNIFORM_BUFFER(FForwardLightData, ForwardLightData)
-		RENDER_TARGET_BINDING_SLOTS()
-		END_SHADER_PARAMETER_STRUCT()
+	END_SHADER_PARAMETER_STRUCT()
 
 		static bool ShouldCompilePermutation(const FGlobalShaderPermutationParameters& Parameters)
 	{
@@ -33,11 +79,23 @@ class FLocalLightPrepassPS : public FGlobalShader
 
 IMPLEMENT_GLOBAL_SHADER(FLocalLightPrepassPS, "/Engine/Private/MobileLocalLightPrepass.usf", "Main", SF_Pixel);
 
+BEGIN_SHADER_PARAMETER_STRUCT(FLocalLightPrepassParameters, )
+SHADER_PARAMETER_STRUCT_INCLUDE(FLocalLightPrepassVS::FParameters, VS)
+SHADER_PARAMETER_STRUCT_INCLUDE(FLocalLightPrepassPS::FParameters, PS)
+SHADER_PARAMETER_STRUCT_INCLUDE(FSceneTextureShaderParameters, SceneTextures)
+RENDER_TARGET_BINDING_SLOTS()
+END_SHADER_PARAMETER_STRUCT()
+
+
 void FMobileSceneRenderer::RenderLocalLightPrepass(FRDGBuilder& GraphBuilder, FSceneTextures& SceneTextures)
 {
 	RDG_EVENT_SCOPE(GraphBuilder, "RenderLocalLightPrepass");
 	QUICK_SCOPE_CYCLE_COUNTER(STAT_RenderLocalLightPrepass);
 
+	static const auto LightGridPixelSizeCVar = IConsoleManager::Get().FindConsoleVariable(TEXT("r.Forward.LightGridPixelSize"));
+	check(LightGridPixelSizeCVar != nullptr);
+	int32 LightGridPixelSize = LightGridPixelSizeCVar->GetInt();
+
 	for (int32 ViewIndex = 0; ViewIndex < Views.Num(); ViewIndex++)
 	{
 		FViewInfo& View = Views[ViewIndex];
@@ -47,22 +105,79 @@ void FMobileSceneRenderer::RenderLocalLightPrepass(FRDGBuilder& GraphBuilder, FS
 			continue;
 		}
 
-		FLocalLightPrepassPS::FParameters* PassParameters = GraphBuilder.AllocParameters<FLocalLightPrepassPS::FParameters>();
-		PassParameters->View = View.ViewUniformBuffer;
-		PassParameters->SceneTextures = SceneTextures.GetSceneTextureShaderParameters(View.FeatureLevel);
-		PassParameters->RenderTargets[0] = FRenderTargetBinding(SceneTextures.MobileLocalLightTextureA, ERenderTargetLoadAction::EClear);;
-		PassParameters->RenderTargets[1] = FRenderTargetBinding(SceneTextures.MobileLocalLightTextureB, ERenderTargetLoadAction::EClear);;
-		PassParameters->ForwardLightData = View.ForwardLightingResources.ForwardLightUniformBuffer;
-
-		FGlobalShaderMap* GlobalShaderMap = GetGlobalShaderMap(FeatureLevel);
-		TShaderMapRef<FLocalLightPrepassPS> PixelShader(GlobalShaderMap);
-
-		FPixelShaderUtils::AddFullscreenPass(
-			GraphBuilder,
-			GlobalShaderMap,
-			RDG_EVENT_NAME("RenderLocalLightPrepass"),
-			PixelShader,
-			PassParameters,
-			FIntRect(0, 0, SceneTextures.MobileLocalLightTextureA->Desc.Extent.X, SceneTextures.MobileLocalLightTextureA->Desc.Extent.Y));
+		const FIntPoint GroupSize(
+			FMath::DivideAndRoundUp(View.ViewRect.Size().X, LightGridPixelSize),
+			FMath::DivideAndRoundUp(View.ViewRect.Size().Y, LightGridPixelSize));
+
+		FRDGBufferRef TileInfoBuffer = GraphBuilder.CreateBuffer(FRDGBufferDesc::CreateBufferDesc(sizeof(int32), 2 * GroupSize.X * GroupSize.Y), TEXT("TileInfoBuffer"));
+		FRDGBufferUAVRef TileInfoBufferUAV = GraphBuilder.CreateUAV(TileInfoBuffer, PF_R32_SINT);
+		FRDGBufferSRVRef TileInfoBufferSRV = GraphBuilder.CreateSRV(TileInfoBuffer, PF_R32_SINT);
+
+		{
+			auto* PassParameters = GraphBuilder.AllocParameters<FLocalLightPrepassCS::FParameters>();
+			PassParameters->RWTileInfo = TileInfoBufferUAV;
+			PassParameters->ForwardLightData = View.ForwardLightingResources.ForwardLightUniformBuffer;
+			PassParameters->GroupSize = GroupSize;
+			auto ComputeShader = View.ShaderMap->GetShader<FLocalLightPrepassCS>();
+
+			FComputeShaderUtils::AddPass(
+				GraphBuilder,
+				RDG_EVENT_NAME("RenderLocalLightPrepass_TiledInfoCS"),
+				ERDGPassFlags::Compute,
+				ComputeShader,
+				PassParameters,
+				FIntVector(FMath::DivideAndRoundUp<uint32>(GroupSize.Y * GroupSize.X, GLocalLightPrepassTileSizeX), 1, 1));
+		}
+
+
+		{
+			FLocalLightPrepassParameters* PassParameters = GraphBuilder.AllocParameters<FLocalLightPrepassParameters>();
+			PassParameters->RenderTargets[0] = FRenderTargetBinding(SceneTextures.MobileLocalLightTextureA, ERenderTargetLoadAction::EClear);;
+			PassParameters->RenderTargets[1] = FRenderTargetBinding(SceneTextures.MobileLocalLightTextureB, ERenderTargetLoadAction::EClear);;
+			PassParameters->SceneTextures = SceneTextures.GetSceneTextureShaderParameters(View.FeatureLevel);
+
+			PassParameters->VS.View = GetShaderBinding(View.ViewUniformBuffer);
+			PassParameters->VS.TileInfo = TileInfoBufferSRV;
+			PassParameters->VS.LightGridPixelSize = LightGridPixelSize;
+			PassParameters->PS.ForwardLightData = View.ForwardLightingResources.ForwardLightUniformBuffer;
+			PassParameters->PS.View = GetShaderBinding(View.ViewUniformBuffer);
+
+			auto VertexShader = View.ShaderMap->GetShader<FLocalLightPrepassVS>();
+			auto PixelShader = View.ShaderMap->GetShader<FLocalLightPrepassPS>();
+
+			GraphBuilder.AddPass(
+				RDG_EVENT_NAME("RenderLocalLightPrepass"),
+				PassParameters,
+				ERDGPassFlags::Raster,
+				[PassParameters, VertexShader, PixelShader, &View, GroupSize](FRHICommandList& RHICmdList)
+				{
+					FGraphicsPipelineStateInitializer GraphicsPSOInit;
+					RHICmdList.ApplyCachedRenderTargets(GraphicsPSOInit);
+
+					RHICmdList.SetViewport(View.ViewRect.Min.X, View.ViewRect.Min.Y, 0.0f, View.ViewRect.Max.X, View.ViewRect.Max.Y, 1.0f);
+					GraphicsPSOInit.RasterizerState = TStaticRasterizerState<FM_Solid, CM_None>::GetRHI();
+					GraphicsPSOInit.DepthStencilState = TStaticDepthStencilState<false, CF_Always>::GetRHI();
+					GraphicsPSOInit.BlendState = TStaticBlendState<>::GetRHI();
+
+					GraphicsPSOInit.PrimitiveType = PT_TriangleList;
+					GraphicsPSOInit.BoundShaderState.VertexDeclarationRHI = GTileVertexDeclaration.VertexDeclarationRHI;
+					GraphicsPSOInit.BoundShaderState.VertexShaderRHI = VertexShader.GetVertexShader();
+					GraphicsPSOInit.BoundShaderState.PixelShaderRHI = PixelShader.GetPixelShader();
+
+					SetGraphicsPipelineState(RHICmdList, GraphicsPSOInit, 0);
+
+					SetShaderParameters(RHICmdList, VertexShader, VertexShader.GetVertexShader(), PassParameters->VS);
+					SetShaderParameters(RHICmdList, PixelShader, PixelShader.GetPixelShader(), PassParameters->PS);
+
+					RHICmdList.SetStreamSource(0, GetOneTileQuadVertexBuffer(), 0);
+					RHICmdList.DrawIndexedPrimitive(GetOneTileQuadIndexBuffer(),
+						0,
+						0,
+						4,
+						0,
+						2,
+						GroupSize.X * GroupSize.Y);
+				});
+		}
 	}
 }
\ No newline at end of file
diff --git a/Engine/Source/Runtime/Renderer/Private/RendererPrivateUtils.h b/Engine/Source/Runtime/Renderer/Private/RendererPrivateUtils.h
--- a/Engine/Source/Runtime/Renderer/Private/RendererPrivateUtils.h
+++ b/Engine/Source/Runtime/Renderer/Private/RendererPrivateUtils.h
@@ -65,6 +65,11 @@ public:
 	const uint32 NumTileQuadsInBuffer;
 };
 
+/** One Tile Quad Vertex Buffer*/
+RENDERER_API FBufferRHIRef& GetOneTileQuadVertexBuffer();
+/** One Tile Quad Index Buffer*/
+RENDERER_API FBufferRHIRef& GetOneTileQuadIndexBuffer();
+
 template <typename ReadbackProcessingLambdaType>
 void AddBufferLockReadbackPass(FRDGBuilder& GraphBuilder, TRefCountPtr<FRDGPooledBuffer> SourceBuffer, uint32 NumBytes, ReadbackProcessingLambdaType &&ReadbackProcessingLambda)
 {
diff --git a/Engine/Source/Runtime/Renderer/Private/RendererUtils.cpp b/Engine/Source/Runtime/Renderer/Private/RendererUtils.cpp
--- a/Engine/Source/Runtime/Renderer/Private/RendererUtils.cpp
+++ b/Engine/Source/Runtime/Renderer/Private/RendererUtils.cpp
@@ -256,3 +256,16 @@ void FStructuredBufferScatterUploader::UploadTo(FRDGBuilder& GraphBuilder, FRDGB
 	FScatterCopyParams ScatterCopyParams { NumScatters, NumBytesPerElement, NumValuesPerScatter };
 	ScatterCopyResource(GraphBuilder, DestBuffer, GraphBuilder.CreateSRV(ScatterOffsets), GraphBuilder.CreateSRV(Values), ScatterCopyParams);
 }
+
+TGlobalResource<FTileTexCoordVertexBuffer> GOneTileQuadVertexBuffer(1);
+TGlobalResource<FTileIndexBuffer> GOneTileQuadIndexBuffer(1);
+
+RENDERER_API FBufferRHIRef& GetOneTileQuadVertexBuffer()
+{
+	return GOneTileQuadVertexBuffer.VertexBufferRHI;
+}
+
+RENDERER_API FBufferRHIRef& GetOneTileQuadIndexBuffer()
+{
+	return GOneTileQuadIndexBuffer.IndexBufferRHI;
+}
\ No newline at end of file
diff --git a/Engine/Source/Runtime/Renderer/Private/SceneTextures.cpp b/Engine/Source/Runtime/Renderer/Private/SceneTextures.cpp
--- a/Engine/Source/Runtime/Renderer/Private/SceneTextures.cpp
+++ b/Engine/Source/Runtime/Renderer/Private/SceneTextures.cpp
@@ -615,14 +615,14 @@ void FSceneTextures::InitializeViewFamily(FRDGBuilder& GraphBuilder, FViewFamily
 	}
 #endif
 
+	extern bool MobileForwardEnablePrepassLocalLights(const FStaticShaderPlatform Platform);
+	if(MobileForwardEnablePrepassLocalLights(Config.ShaderPlatform))
 	{
-		FRDGTextureDesc Desc = FRDGTextureDesc::Create2D(Config.Extent, PF_FloatR11G11B10, FClearValueBinding::Transparent, TexCreate_RenderTargetable | TexCreate_ShaderResource);
-		SceneTextures.MobileLocalLightTextureA = GraphBuilder.CreateTexture(Desc, TEXT("MobileLocalLightTextureA"));
-	}
+		FRDGTextureDesc MobileLocalLightTextureADesc = FRDGTextureDesc::Create2D(Config.Extent, PF_FloatR11G11B10, FClearValueBinding::Transparent, TexCreate_RenderTargetable | TexCreate_ShaderResource);
+		SceneTextures.MobileLocalLightTextureA = GraphBuilder.CreateTexture(MobileLocalLightTextureADesc, TEXT("MobileLocalLightTextureA"));
 
-	{
-		FRDGTextureDesc Desc = FRDGTextureDesc::Create2D(Config.Extent, PF_A2B10G10R10, FClearValueBinding::Transparent, TexCreate_RenderTargetable | TexCreate_ShaderResource);
-		SceneTextures.MobileLocalLightTextureB = GraphBuilder.CreateTexture(Desc, TEXT("MobileLocalLightTextureB"));
+		FRDGTextureDesc MobileLocalLightTextureBDesc = FRDGTextureDesc::Create2D(Config.Extent, PF_A2B10G10R10, FClearValueBinding::Transparent, TexCreate_RenderTargetable | TexCreate_ShaderResource);
+		SceneTextures.MobileLocalLightTextureB = GraphBuilder.CreateTexture(MobileLocalLightTextureBDesc, TEXT("MobileLocalLightTextureB"));
 	}
 
 #if WITH_DEBUG_VIEW_MODES
