Require all views to support VRS in order to generate CAS, rather than checking if at least one view supports VRS. Fixes a crash in splitscreen where one view is undergoing a camera cut but the other is not.

diff --git a/Engine/Source/Runtime/Renderer/Private/VariableRateShading/ContrastAdaptiveImageGenerator.cpp b/Engine/Source/Runtime/Renderer/Private/VariableRateShading/ContrastAdaptiveImageGenerator.cpp
--- a/Engine/Source/Runtime/Renderer/Private/VariableRateShading/ContrastAdaptiveImageGenerator.cpp
+++ b/Engine/Source/Runtime/Renderer/Private/VariableRateShading/ContrastAdaptiveImageGenerator.cpp
@@ -478,17 +478,18 @@ FRDGTextureRef FContrastAdaptiveImageGenerator::GetImage(FRDGBuilder& GraphBuild
 void FContrastAdaptiveImageGenerator::PrepareImages(FRDGBuilder& GraphBuilder, const FSceneViewFamily& ViewFamily, const FMinimalSceneTextures& SceneTextures)
 {
 	RDG_EVENT_SCOPE(GraphBuilder, "VariableRateShading");
-	bool bIsAnyViewVRSCompatible = false;
+	bool bAreAllViewsVRSCompatible = true;
 	for (const FSceneView* View : ViewFamily.Views)
 	{
 		check(View->bIsViewInfo);
 		const FViewInfo* ViewInfo = static_cast<const FViewInfo*>(View);
-		if (!View->bCameraCut && FVariableRateShadingImageManager::IsVRSCompatibleWithView(*ViewInfo) && ViewInfo->PrevViewInfo.LuminanceHistory)
+		if (View->bCameraCut || !FVariableRateShadingImageManager::IsVRSCompatibleWithView(*ViewInfo) || !ViewInfo->PrevViewInfo.LuminanceHistory)
 		{
-			bIsAnyViewVRSCompatible = true;
+			bAreAllViewsVRSCompatible = false;
+			break;
 		}
 	}
-	bool bPrepareImageBasedVRS = IsContrastAdaptiveShadingEnabled() && bIsAnyViewVRSCompatible;
+	bool bPrepareImageBasedVRS = IsContrastAdaptiveShadingEnabled() && bAreAllViewsVRSCompatible;
 	if (!bPrepareImageBasedVRS)
 	{
 		return;
