Revert DFAO multi-view changes

#rnx

diff --git a/Engine/Shaders/Private/DistanceFieldAOShared.ush b/Engine/Shaders/Private/DistanceFieldAOShared.ush
--- a/Engine/Shaders/Private/DistanceFieldAOShared.ush
+++ b/Engine/Shaders/Private/DistanceFieldAOShared.ush
@@ -155,7 +155,9 @@ FUpsampleDFAOOutput UpsampleDFAO(float2 BufferUV, float SceneDepth)
 {
 	FUpsampleDFAOOutput UpsampleDFAOOutput;
 
-	float2 DistanceFieldUVs = min(BufferUV, AOBufferBilinearUVMax);
+	// Distance field AO was computed at 0,0 regardless of viewrect min
+	float2 DistanceFieldUVs = BufferUV - View.ViewRectMin.xy * View.BufferSizeAndInvSize.zw;
+	DistanceFieldUVs = min(DistanceFieldUVs, AOBufferBilinearUVMax);
 
 #define BILATERAL_UPSAMPLE 1
 #if BILATERAL_UPSAMPLE
diff --git a/Engine/Shaders/Private/DistanceFieldObjectCulling.usf b/Engine/Shaders/Private/DistanceFieldObjectCulling.usf
--- a/Engine/Shaders/Private/DistanceFieldObjectCulling.usf
+++ b/Engine/Shaders/Private/DistanceFieldObjectCulling.usf
@@ -129,10 +129,12 @@ void BuildTileConesMain(
 {
     uint ThreadIndex = GroupThreadId.y * THREADGROUP_SIZEX + GroupThreadId.x;
 
-	// Sampling from the texture based off of the ViewRect size because the texture is created on a per-view basis
-	float2 BaseLevelScreenUV = (DispatchThreadId.xy + float2(.5f, .5f)) * DOWNSAMPLE_FACTOR * View.ViewSizeAndInvSize.zw;
+	float2 BaseLevelScreenUV = (DispatchThreadId.xy + float2(.5f, .5f)) * DOWNSAMPLE_FACTOR * View.BufferSizeAndInvSize.zw;
 	float SceneDepth = GetDownsampledDepth(BaseLevelScreenUV);
 
+	//float2 ScreenUV = (DispatchThreadId.xy * DOWNSAMPLE_FACTOR + View.ViewRectMin.xy + float2(.5f, .5f)) * View.BufferSizeAndInvSize.zw;
+	//float SceneDepth = CalcSceneDepth(ScreenUV);
+
 	// Initialize per-tile variables
     if (ThreadIndex == 0) 
 	{
diff --git a/Engine/Source/Runtime/Renderer/Private/DistanceFieldAmbientOcclusion.cpp b/Engine/Source/Runtime/Renderer/Private/DistanceFieldAmbientOcclusion.cpp
--- a/Engine/Source/Runtime/Renderer/Private/DistanceFieldAmbientOcclusion.cpp
+++ b/Engine/Source/Runtime/Renderer/Private/DistanceFieldAmbientOcclusion.cpp
@@ -159,7 +159,7 @@ void TileIntersectionModifyCompilationEnvironment(EShaderPlatform Platform, FSha
 
 FIntPoint GetBufferSizeForAO(const FViewInfo& View)
 {
-	return FIntPoint::DivideAndRoundDown(View.ViewRect.Size(), GAODownsampleFactor);
+	return FIntPoint::DivideAndRoundDown(View.GetSceneTexturesConfig().Extent, GAODownsampleFactor);
 }
 
 // Sample set restricted to not self-intersect a surface based on cone angle .475882232
@@ -374,72 +374,84 @@ IMPLEMENT_GLOBAL_SHADER(FComputeDistanceFieldNormalCS, "/Engine/Private/Distance
 
 void ComputeDistanceFieldNormal(
 	FRDGBuilder& GraphBuilder,
-	const FViewInfo& View,
+	const TArray<FViewInfo>& Views,
 	TRDGUniformBufferRef<FSceneTextureUniformParameters> SceneTexturesUniformBuffer,
 	FRDGTextureRef DistanceFieldNormal,
 	const FDistanceFieldAOParameters& Parameters)
 {
 	if (GAOComputeShaderNormalCalculation)
 	{
-		uint32 GroupSizeX = FMath::DivideAndRoundUp(View.ViewRect.Size().X / GAODownsampleFactor, GDistanceFieldAOTileSizeX);
-		uint32 GroupSizeY = FMath::DivideAndRoundUp(View.ViewRect.Size().Y / GAODownsampleFactor, GDistanceFieldAOTileSizeY);
+		for (int32 ViewIndex = 0; ViewIndex < Views.Num(); ViewIndex++)
+		{
+			const FViewInfo& View = Views[ViewIndex];
+			RDG_GPU_MASK_SCOPE(GraphBuilder, View.GPUMask);
+
+			uint32 GroupSizeX = FMath::DivideAndRoundUp(View.ViewRect.Size().X / GAODownsampleFactor, GDistanceFieldAOTileSizeX);
+			uint32 GroupSizeY = FMath::DivideAndRoundUp(View.ViewRect.Size().Y / GAODownsampleFactor, GDistanceFieldAOTileSizeY);
 
-		auto* PassParameters = GraphBuilder.AllocParameters<FComputeDistanceFieldNormalCS::FParameters>();
-		PassParameters->View = View.ViewUniformBuffer;
-		PassParameters->SceneTextures = SceneTexturesUniformBuffer;
-		PassParameters->Strata = Strata::BindStrataGlobalUniformParameters(View);
-		PassParameters->AOParameters = DistanceField::SetupAOShaderParameters(Parameters);
-		PassParameters->RWDistanceFieldNormal = GraphBuilder.CreateUAV(DistanceFieldNormal);
+			auto* PassParameters = GraphBuilder.AllocParameters<FComputeDistanceFieldNormalCS::FParameters>();
+			PassParameters->View = View.ViewUniformBuffer;
+			PassParameters->SceneTextures = SceneTexturesUniformBuffer;
+			PassParameters->Strata = Strata::BindStrataGlobalUniformParameters(View);
+			PassParameters->AOParameters = DistanceField::SetupAOShaderParameters(Parameters);
+			PassParameters->RWDistanceFieldNormal = GraphBuilder.CreateUAV(DistanceFieldNormal);
 
-		TShaderMapRef<FComputeDistanceFieldNormalCS> ComputeShader(View.ShaderMap);
+			TShaderMapRef<FComputeDistanceFieldNormalCS> ComputeShader(View.ShaderMap);
 
-		FComputeShaderUtils::AddPass(GraphBuilder, RDG_EVENT_NAME("ComputeNormalCS"), ComputeShader, PassParameters, FIntVector(GroupSizeX, GroupSizeY, 1));
+			FComputeShaderUtils::AddPass(GraphBuilder, RDG_EVENT_NAME("ComputeNormalCS"), ComputeShader, PassParameters, FIntVector(GroupSizeX, GroupSizeY, 1));
+		}
 	}
 	else
 	{
-		auto* PassParameters = GraphBuilder.AllocParameters<FComputeDistanceFieldNormalPS::FParameters>();
-		PassParameters->View = View.ViewUniformBuffer;
-		PassParameters->SceneTextures = SceneTexturesUniformBuffer;
-		PassParameters->Strata = Strata::BindStrataGlobalUniformParameters(View);
-		PassParameters->AOParameters = DistanceField::SetupAOShaderParameters(Parameters);
-		PassParameters->RenderTargets[0] = FRenderTargetBinding(DistanceFieldNormal, ERenderTargetLoadAction::EClear);
-
-		GraphBuilder.AddPass(
-			RDG_EVENT_NAME("ComputeNormal"),
-			PassParameters,
-			ERDGPassFlags::Raster,
-			[&View, PassParameters](FRHICommandList& RHICmdList)
+		for (int32 ViewIndex = 0; ViewIndex < Views.Num(); ViewIndex++)
 		{
-			RHICmdList.SetViewport(0, 0, 0.0f, View.ViewRect.Width() / GAODownsampleFactor, View.ViewRect.Height() / GAODownsampleFactor, 1.0f);
-
-			TShaderMapRef<FPostProcessVS> VertexShader(View.ShaderMap);
-
-			TShaderMapRef<FComputeDistanceFieldNormalPS> PixelShader(View.ShaderMap);
-
-			FGraphicsPipelineStateInitializer GraphicsPSOInit;
-			RHICmdList.ApplyCachedRenderTargets(GraphicsPSOInit);
-			GraphicsPSOInit.RasterizerState = TStaticRasterizerState<FM_Solid, CM_None>::GetRHI();
-			GraphicsPSOInit.DepthStencilState = TStaticDepthStencilState<false, CF_Always>::GetRHI();
-			GraphicsPSOInit.BlendState = TStaticBlendState<>::GetRHI();
-			GraphicsPSOInit.PrimitiveType = PT_TriangleList;
-			GraphicsPSOInit.BoundShaderState.VertexDeclarationRHI = GFilterVertexDeclaration.VertexDeclarationRHI;
-			GraphicsPSOInit.BoundShaderState.VertexShaderRHI = VertexShader.GetVertexShader();
-			GraphicsPSOInit.BoundShaderState.PixelShaderRHI = PixelShader.GetPixelShader();
-
-			SetGraphicsPipelineState(RHICmdList, GraphicsPSOInit, 0);
-
-			SetShaderParameters(RHICmdList, PixelShader, PixelShader.GetPixelShader(), *PassParameters);
-
-			DrawRectangle(
-				RHICmdList,
-				0, 0,
-				View.ViewRect.Width() / GAODownsampleFactor, View.ViewRect.Height() / GAODownsampleFactor,
-				0, 0,
-				View.ViewRect.Width(), View.ViewRect.Height(),
-				FIntPoint(View.ViewRect.Width() / GAODownsampleFactor, View.ViewRect.Height() / GAODownsampleFactor),
-				View.GetSceneTexturesConfig().Extent,
-				VertexShader);
-		});
+			const FViewInfo& View = Views[ViewIndex];
+			RDG_GPU_MASK_SCOPE(GraphBuilder, View.GPUMask);
+
+			auto* PassParameters = GraphBuilder.AllocParameters<FComputeDistanceFieldNormalPS::FParameters>();
+			PassParameters->View = View.ViewUniformBuffer;
+			PassParameters->SceneTextures = SceneTexturesUniformBuffer;
+			PassParameters->Strata = Strata::BindStrataGlobalUniformParameters(View);
+			PassParameters->AOParameters = DistanceField::SetupAOShaderParameters(Parameters);
+			PassParameters->RenderTargets[0] = FRenderTargetBinding(DistanceFieldNormal, ViewIndex == 0 ? ERenderTargetLoadAction::EClear : ERenderTargetLoadAction::ELoad);
+
+			GraphBuilder.AddPass(
+				RDG_EVENT_NAME("ComputeNormal"),
+				PassParameters,
+				ERDGPassFlags::Raster,
+				[&View, PassParameters](FRHICommandList& RHICmdList)
+			{
+				RHICmdList.SetViewport(0, 0, 0.0f, View.ViewRect.Width() / GAODownsampleFactor, View.ViewRect.Height() / GAODownsampleFactor, 1.0f);
+
+				TShaderMapRef<FPostProcessVS> VertexShader(View.ShaderMap);
+
+				TShaderMapRef<FComputeDistanceFieldNormalPS> PixelShader(View.ShaderMap);
+
+				FGraphicsPipelineStateInitializer GraphicsPSOInit;
+				RHICmdList.ApplyCachedRenderTargets(GraphicsPSOInit);
+				GraphicsPSOInit.RasterizerState = TStaticRasterizerState<FM_Solid, CM_None>::GetRHI();
+				GraphicsPSOInit.DepthStencilState = TStaticDepthStencilState<false, CF_Always>::GetRHI();
+				GraphicsPSOInit.BlendState = TStaticBlendState<>::GetRHI();
+				GraphicsPSOInit.PrimitiveType = PT_TriangleList;
+				GraphicsPSOInit.BoundShaderState.VertexDeclarationRHI = GFilterVertexDeclaration.VertexDeclarationRHI;
+				GraphicsPSOInit.BoundShaderState.VertexShaderRHI = VertexShader.GetVertexShader();
+				GraphicsPSOInit.BoundShaderState.PixelShaderRHI = PixelShader.GetPixelShader();
+
+				SetGraphicsPipelineState(RHICmdList, GraphicsPSOInit, 0);
+
+				SetShaderParameters(RHICmdList, PixelShader, PixelShader.GetPixelShader(), *PassParameters);
+
+				DrawRectangle(
+					RHICmdList,
+					0, 0,
+					View.ViewRect.Width() / GAODownsampleFactor, View.ViewRect.Height() / GAODownsampleFactor,
+					0, 0,
+					View.ViewRect.Width(), View.ViewRect.Height(),
+					FIntPoint(View.ViewRect.Width() / GAODownsampleFactor, View.ViewRect.Height() / GAODownsampleFactor),
+					View.GetSceneTexturesConfig().Extent,
+					VertexShader);
+			});
+		}
 	}
 }
 
@@ -644,9 +656,10 @@ void AllocateTileIntersectionBuffers(
 	OutParameters.TileListGroupSize = TileListGroupSize;
 }
 
-void ListDistanceFieldLightingMemory(const FScene* Scene, FSceneRenderer& SceneRenderer)
+void ListDistanceFieldLightingMemory(const FViewInfo& View, FSceneRenderer& SceneRenderer)
 {
 #if !NO_LOGGING
+	const FScene* Scene = (const FScene*)View.Family->Scene;
 	UE_LOG(LogRenderer, Log, TEXT("Shared GPU memory (excluding render targets)"));
 
 	if (Scene->DistanceFieldSceneData.NumObjectsInBuffer > 0)
@@ -762,22 +775,13 @@ void FDeferredShadingSceneRenderer::RenderDFAOAsIndirectShadowing(
 
 bool FDeferredShadingSceneRenderer::ShouldRenderDistanceFieldLighting() const
 {
-	//@to-do - check for ISR (may want to exclude certain checks on secondary views)
-	bool bSupportsDistanceFieldAO = true;
-
-	for (int32_t ViewIndex = 0; ViewIndex < Views.Num(); ++ViewIndex)
-	{
-		const FViewInfo& View = Views[ViewIndex];
+	//@todo - support multiple views
+	const FViewInfo& View = Views[0];
 
-		if (!SupportsDistanceFieldAO(View.GetFeatureLevel(), View.GetShaderPlatform())
-			|| !View.IsPerspectiveProjection())
-		{
-			bSupportsDistanceFieldAO = false;
-			break;
-		}
-	}
-
-	return bSupportsDistanceFieldAO && Scene->DistanceFieldSceneData.NumObjectsInBuffer;
+	return SupportsDistanceFieldAO(View.GetFeatureLevel(), View.GetShaderPlatform())
+		&& Views.Num() == 1
+		&& View.IsPerspectiveProjection()
+		&& Scene->DistanceFieldSceneData.NumObjectsInBuffer;
 }
 
 void FDeferredShadingSceneRenderer::RenderDistanceFieldLighting(
@@ -791,6 +795,10 @@ void FDeferredShadingSceneRenderer::RenderDistanceFieldLighting(
 	check(ShouldRenderDistanceFieldLighting());
 	check(!Scene->DistanceFieldSceneData.HasPendingOperations());
 
+	//@todo - support multiple views
+	const FViewInfo& View = Views[0];
+
+	RDG_GPU_MASK_SCOPE(GraphBuilder, View.GPUMask);
 	RDG_EVENT_SCOPE(GraphBuilder, "DistanceFieldLighting");
 	QUICK_SCOPE_CYCLE_COUNTER(STAT_RenderDistanceFieldLighting);
 
@@ -799,7 +807,7 @@ void FDeferredShadingSceneRenderer::RenderDistanceFieldLighting(
 	if (bListMemoryNextFrame)
 	{
 		bListMemoryNextFrame = false;
-		ListDistanceFieldLightingMemory(Scene, *this);
+		ListDistanceFieldLightingMemory(View, *this);
 	}
 
 	if (bListMeshDistanceFieldsMemoryNextFrame)
@@ -808,73 +816,58 @@ void FDeferredShadingSceneRenderer::RenderDistanceFieldLighting(
 		Scene->DistanceFieldSceneData.ListMeshDistanceFields(true);
 	}
 
-	FRDGTextureRef BentNormalOutput = nullptr;
+	FRDGTextureRef DistanceFieldNormal = nullptr;
 
 	{
-		const FIntPoint BufferSize = GetActiveSceneTexturesConfig().Extent / GAODownsampleFactor;
-		const FRDGTextureDesc Desc = FRDGTextureDesc::Create2D(BufferSize, PF_FloatRGBA, FClearValueBinding::None, GFastVRamConfig.DistanceFieldAOBentNormal | TexCreate_ShaderResource | TexCreate_RenderTargetable | TexCreate_UAV);
-		BentNormalOutput = GraphBuilder.CreateTexture(Desc, TEXT("DistanceFieldBentNormalAO"));
-		AddClearUAVPass(GraphBuilder, GraphBuilder.CreateUAV(BentNormalOutput), FLinearColor::Black);
+		const FIntPoint BufferSize = GetBufferSizeForAO(View);
+		const FRDGTextureDesc Desc = FRDGTextureDesc::Create2D(BufferSize, PF_FloatRGBA, FClearValueBinding::Transparent, GFastVRamConfig.DistanceFieldNormal | TexCreate_RenderTargetable | TexCreate_UAV | TexCreate_ShaderResource);
+		DistanceFieldNormal = GraphBuilder.CreateTexture(Desc, TEXT("DistanceFieldNormal"));
 	}
 
-	for (int32 ViewIndex = 0; ViewIndex < Views.Num(); ++ViewIndex)
-	{
-		const FViewInfo& View = Views[ViewIndex];
-		RDG_GPU_MASK_SCOPE(GraphBuilder, View.GPUMask);
-		RDG_EVENT_SCOPE_CONDITIONAL(GraphBuilder, Views.Num() > 1, "View%d", ViewIndex);
+	const FIntPoint TileListGroupSize = GetTileListGroupSizeForView(View);
+	const int32 MaxSceneObjects = FMath::DivideAndRoundUp(Scene->DistanceFieldSceneData.NumObjectsInBuffer, 256) * 256;
+	const bool bAllow16BitIndices = !IsMetalPlatform(GShaderPlatformForFeatureLevel[View.FeatureLevel]);
 
-		FRDGTextureRef DistanceFieldNormal = nullptr;
+	FRDGBufferRef ObjectTilesIndirectArguments = nullptr;
+	FTileIntersectionParameters TileIntersectionParameters;
 
-		{
-			const FIntPoint BufferSize = GetBufferSizeForAO(View);
-			const FRDGTextureDesc Desc = FRDGTextureDesc::Create2D(BufferSize, PF_FloatRGBA, FClearValueBinding::Transparent, GFastVRamConfig.DistanceFieldNormal | TexCreate_RenderTargetable | TexCreate_UAV | TexCreate_ShaderResource);
-			DistanceFieldNormal = GraphBuilder.CreateTexture(Desc, TEXT("DistanceFieldNormal"));
-		}
-
-		ComputeDistanceFieldNormal(GraphBuilder, View, SceneTextures.UniformBuffer, DistanceFieldNormal, Parameters);
-
-		const FIntPoint TileListGroupSize = GetTileListGroupSizeForView(View);
-		const int32 MaxSceneObjects = FMath::DivideAndRoundUp(Scene->DistanceFieldSceneData.NumObjectsInBuffer, 256) * 256;
-		const bool bAllow16BitIndices = !IsMetalPlatform(GShaderPlatformForFeatureLevel[View.FeatureLevel]);
-
-		FRDGBufferRef ObjectTilesIndirectArguments = nullptr;
-		FTileIntersectionParameters TileIntersectionParameters;
+	AllocateTileIntersectionBuffers(GraphBuilder, TileListGroupSize, MaxSceneObjects, bAllow16BitIndices, ObjectTilesIndirectArguments, TileIntersectionParameters);
 
-		AllocateTileIntersectionBuffers(GraphBuilder, TileListGroupSize, MaxSceneObjects, bAllow16BitIndices, ObjectTilesIndirectArguments, TileIntersectionParameters);
+	FRDGBufferRef ObjectIndirectArguments = nullptr;
+	FDistanceFieldCulledObjectBufferParameters CulledObjectBufferParameters;
 
-		FRDGBufferRef ObjectIndirectArguments = nullptr;
-		FDistanceFieldCulledObjectBufferParameters CulledObjectBufferParameters;
+	if (UseAOObjectDistanceField())
+	{
+		AllocateDistanceFieldCulledObjectBuffers(
+			GraphBuilder,
+			MaxSceneObjects,
+			ObjectIndirectArguments,
+			CulledObjectBufferParameters);
 
-		if (UseAOObjectDistanceField())
-		{
-			AllocateDistanceFieldCulledObjectBuffers(
-				GraphBuilder,
-				MaxSceneObjects,
-				ObjectIndirectArguments,
-				CulledObjectBufferParameters);
+		CullObjectsToView(GraphBuilder, Scene, View, Parameters, CulledObjectBufferParameters);
+	}
 
-			CullObjectsToView(GraphBuilder, Scene, View, Parameters, CulledObjectBufferParameters);
+	ComputeDistanceFieldNormal(GraphBuilder, Views, SceneTextures.UniformBuffer, DistanceFieldNormal, Parameters);
 
-			BuildTileObjectLists(GraphBuilder, Scene, View, SceneTextures.UniformBuffer, ObjectIndirectArguments, CulledObjectBufferParameters, TileIntersectionParameters, DistanceFieldNormal, Parameters);
-		}
+	// Intersect objects with screen tiles, build lists
+	if (UseAOObjectDistanceField())
+	{
+		//@todo - support multiple views - should pass one TileIntersectionParameters per view
+		BuildTileObjectLists(GraphBuilder, Scene, Views, ObjectIndirectArguments, CulledObjectBufferParameters, TileIntersectionParameters, DistanceFieldNormal, Parameters);
+	}
 
-		// Render to a per-view BentNormal first because it also needs to be stored by the view history.
-		// The per-view output can be copied back to the appropriate region in the output BentNormal texture.
-		FRDGTextureRef PerViewBentNormal = nullptr;
+	FRDGTextureRef BentNormalOutput = nullptr;
 
-		RenderDistanceFieldAOScreenGrid(
-			GraphBuilder,
-			SceneTextures,
-			View,
-			CulledObjectBufferParameters,
-			ObjectTilesIndirectArguments,
-			TileIntersectionParameters,
-			Parameters,
-			DistanceFieldNormal,
-			PerViewBentNormal);
-
-		AddCopyTexturePass(GraphBuilder, PerViewBentNormal, BentNormalOutput, FIntPoint::ZeroValue, View.ViewRect.Min / GAODownsampleFactor, View.ViewRect.Size() / GAODownsampleFactor);
-	}
+	RenderDistanceFieldAOScreenGrid(
+		GraphBuilder,
+		SceneTextures,
+		View,
+		CulledObjectBufferParameters,
+		ObjectTilesIndirectArguments,
+		TileIntersectionParameters,
+		Parameters,
+		DistanceFieldNormal,
+		BentNormalOutput);
 
 	RenderCapsuleShadowsForMovableSkylight(GraphBuilder, SceneTextures.UniformBuffer, BentNormalOutput);
 
diff --git a/Engine/Source/Runtime/Renderer/Private/DistanceFieldAmbientOcclusion.h b/Engine/Source/Runtime/Renderer/Private/DistanceFieldAmbientOcclusion.h
--- a/Engine/Source/Runtime/Renderer/Private/DistanceFieldAmbientOcclusion.h
+++ b/Engine/Source/Runtime/Renderer/Private/DistanceFieldAmbientOcclusion.h
@@ -14,13 +14,11 @@ class FRDGBuffer;
 class FRDGBuilder;
 class FRDGTexture;
 class FScene;
-class FSceneTextureUniformParameters;
 class FSceneViewFamily;
 class FViewInfo;
 struct FShaderCompilerEnvironment;
 using FRDGBufferRef = FRDGBuffer*;
 using FRDGTextureRef = FRDGTexture*;
-template <typename TUniformStruct> using TRDGUniformBufferRef = TRDGUniformBuffer<TUniformStruct>*;
 
 /** Base downsample factor that all distance field AO operations are done at. */
 const int32 GAODownsampleFactor = 2;
@@ -109,8 +107,7 @@ extern bool ShouldDoReflectionEnvironment(const FScene* Scene, const FSceneViewF
 extern void CullObjectsToView(FRDGBuilder& GraphBuilder, FScene* Scene, const FViewInfo& View, const FDistanceFieldAOParameters& Parameters, FDistanceFieldCulledObjectBufferParameters& CulledObjectBuffers);
 extern void BuildTileObjectLists(FRDGBuilder& GraphBuilder,
 	FScene* Scene,
-	const FViewInfo& View,
-	TRDGUniformBufferRef<FSceneTextureUniformParameters> SceneTexturesUniformBuffer,
+	TArray<FViewInfo>& Views,
 	FRDGBufferRef ObjectIndirectArguments,
 	const FDistanceFieldCulledObjectBufferParameters& CulledObjectBufferParameters,
 	FTileIntersectionParameters TileIntersectionParameters,
diff --git a/Engine/Source/Runtime/Renderer/Private/DistanceFieldLightingPost.cpp b/Engine/Source/Runtime/Renderer/Private/DistanceFieldLightingPost.cpp
--- a/Engine/Source/Runtime/Renderer/Private/DistanceFieldLightingPost.cpp
+++ b/Engine/Source/Runtime/Renderer/Private/DistanceFieldLightingPost.cpp
@@ -320,7 +320,7 @@ void UpdateHistory(
 {
 	const FRDGSystemTextures& SystemTextures = FRDGSystemTextures::Get(GraphBuilder);
 
-	const FIntPoint AOHistoryExtent = View.ViewRect.Size();
+	const FIntPoint SceneTextureExtent = View.GetSceneTexturesConfig().Extent;
 
 	if (BentNormalHistoryState && DistanceFieldAOUseHistory(View))
 	{
@@ -358,7 +358,7 @@ void UpdateHistory(
 				PassParameters->UseHistoryFilter = UseAOHistoryStabilityPass() ? 1.0f : 0.0f;
 
 				{
-					FIntPoint HistoryBufferSize = AOHistoryExtent / FIntPoint(GAODownsampleFactor, GAODownsampleFactor);
+					FIntPoint HistoryBufferSize = View.GetSceneTexturesConfig().Extent / FIntPoint(GAODownsampleFactor, GAODownsampleFactor);
 
 					const float InvBufferSizeX = 1.0f / HistoryBufferSize.X;
 					const float InvBufferSizeY = 1.0f / HistoryBufferSize.Y;
@@ -391,7 +391,7 @@ void UpdateHistory(
 					RDG_EVENT_NAME("UpdateHistory"),
 					PassParameters,
 					ERDGPassFlags::Raster,
-					[PassParameters, VertexShader, PixelShader, &View, AOHistoryExtent]
+					[PassParameters, VertexShader, PixelShader, &View, SceneTextureExtent]
 					(FRHICommandList& RHICmdList)
 				{
 					RHICmdList.SetViewport(0, 0, 0.0f, View.ViewRect.Width() / GAODownsampleFactor, View.ViewRect.Height() / GAODownsampleFactor, 1.0f);
@@ -417,7 +417,7 @@ void UpdateHistory(
 						0, 0,
 						View.ViewRect.Width() / GAODownsampleFactor, View.ViewRect.Height() / GAODownsampleFactor,
 						FIntPoint(View.ViewRect.Width() / GAODownsampleFactor, View.ViewRect.Height() / GAODownsampleFactor),
-						AOHistoryExtent / FIntPoint(GAODownsampleFactor, GAODownsampleFactor),
+						SceneTextureExtent / FIntPoint(GAODownsampleFactor, GAODownsampleFactor),
 						VertexShader);
 				});
 			}
@@ -427,7 +427,7 @@ void UpdateHistory(
 				const FRDGTextureDesc& HistoryDesc = BentNormalHistoryTexture->Desc;
 
 				// Reallocate history if buffer sizes have changed
-				if (HistoryDesc.Extent != AOHistoryExtent / FIntPoint(GAODownsampleFactor, GAODownsampleFactor))
+				if (HistoryDesc.Extent != SceneTextureExtent / FIntPoint(GAODownsampleFactor, GAODownsampleFactor))
 				{
 					GRenderTargetPool.FreeUnusedResource(*BentNormalHistoryState);
 					*BentNormalHistoryState = nullptr;
@@ -435,8 +435,7 @@ void UpdateHistory(
 					AllocateOrReuseAORenderTarget(GraphBuilder, View, BentNormalHistoryTexture, BentNormalHistoryRTName, PF_FloatRGBA);
 				}
 
-				//@to-do: this might not be necessary anymore
-				const bool bManuallyClampUV = View.ViewRect.Min != FIntPoint::ZeroValue || View.ViewRect.Max != AOHistoryExtent;
+				const bool bManuallyClampUV = View.ViewRect.Min != FIntPoint::ZeroValue || View.ViewRect.Max != SceneTextureExtent;
 
 				FFilterHistoryPS::FPermutationDomain PermutationVector;
 				PermutationVector.Set<FFilterHistoryPS::FManuallyClampUV>(bManuallyClampUV);
@@ -444,7 +443,7 @@ void UpdateHistory(
 				auto VertexShader = View.ShaderMap->GetShader<FPostProcessVS>();
 				auto PixelShader = View.ShaderMap->GetShader<FFilterHistoryPS>(PermutationVector);
 
-				const FIntPoint DownsampledBufferSize(AOHistoryExtent / FIntPoint(GAODownsampleFactor, GAODownsampleFactor));
+				const FIntPoint DownsampledBufferSize(View.GetSceneTexturesConfig().Extent / FIntPoint(GAODownsampleFactor, GAODownsampleFactor));
 				FVector2f MaxSampleBufferUV(
 					(View.ViewRect.Width() / GAODownsampleFactor - 0.5f - GAODownsampleFactor) / DownsampledBufferSize.X,
 					(View.ViewRect.Height() / GAODownsampleFactor - 0.5f - GAODownsampleFactor) / DownsampledBufferSize.Y);
@@ -465,7 +464,7 @@ void UpdateHistory(
 					RDG_EVENT_NAME("UpdateHistoryStability"),
 					PassParameters,
 					ERDGPassFlags::Raster,
-					[PassParameters, VertexShader, PixelShader, &View, AOHistoryExtent](FRHICommandList& RHICmdList)
+					[PassParameters, VertexShader, PixelShader, &View, SceneTextureExtent](FRHICommandList& RHICmdList)
 				{
 					RHICmdList.SetViewport(0, 0, 0.0f, View.ViewRect.Width() / GAODownsampleFactor, View.ViewRect.Height() / GAODownsampleFactor, 1.0f);
 
@@ -490,7 +489,7 @@ void UpdateHistory(
 						0, 0,
 						View.ViewRect.Width() / GAODownsampleFactor, View.ViewRect.Height() / GAODownsampleFactor,
 						FIntPoint(View.ViewRect.Width() / GAODownsampleFactor, View.ViewRect.Height() / GAODownsampleFactor),
-						AOHistoryExtent / FIntPoint(GAODownsampleFactor, GAODownsampleFactor),
+						SceneTextureExtent / FIntPoint(GAODownsampleFactor, GAODownsampleFactor),
 						VertexShader);
 				});
 
@@ -508,7 +507,7 @@ void UpdateHistory(
 		{
 			// Use the current frame's upscaled mask for next frame's history
 			FRDGTextureRef DistanceFieldAOBentNormal = nullptr;
-			AllocateOrReuseAORenderTarget(GraphBuilder, View, DistanceFieldAOBentNormal, TEXT("PerViewDistanceFieldBentNormalAO"), PF_FloatRGBA, GFastVRamConfig.DistanceFieldAOHistory);
+			AllocateOrReuseAORenderTarget(GraphBuilder, View, DistanceFieldAOBentNormal, TEXT("DistanceFieldBentNormalAO"), PF_FloatRGBA, GFastVRamConfig.DistanceFieldAOBentNormal);
 
 			GeometryAwareUpsample(GraphBuilder, View, DistanceFieldAOBentNormal, DistanceFieldNormal, BentNormalInterpolation, Parameters);
 
@@ -517,13 +516,14 @@ void UpdateHistory(
 		}
 
 		DistanceFieldAOHistoryViewRect->Min = FIntPoint::ZeroValue;
-		DistanceFieldAOHistoryViewRect->Max = View.ViewRect.Size() / FIntPoint(GAODownsampleFactor, GAODownsampleFactor);
+		DistanceFieldAOHistoryViewRect->Max.X = View.ViewRect.Size().X / GAODownsampleFactor;
+		DistanceFieldAOHistoryViewRect->Max.Y = View.ViewRect.Size().Y / GAODownsampleFactor;
 	}
 	else
 	{
 		// Temporal reprojection is disabled or there is no view state - just upscale
 		FRDGTextureRef DistanceFieldAOBentNormal = nullptr;
-		AllocateOrReuseAORenderTarget(GraphBuilder, View, DistanceFieldAOBentNormal, TEXT("PerViewDistanceFieldBentNormalAO"), PF_FloatRGBA, GFastVRamConfig.DistanceFieldAOHistory);
+		AllocateOrReuseAORenderTarget(GraphBuilder, View, DistanceFieldAOBentNormal, TEXT("DistanceFieldBentNormalAO"), PF_FloatRGBA, GFastVRamConfig.DistanceFieldAOBentNormal);
 
 		GeometryAwareUpsample(GraphBuilder, View, DistanceFieldAOBentNormal, DistanceFieldNormal, BentNormalInterpolation, Parameters);
 
diff --git a/Engine/Source/Runtime/Renderer/Private/DistanceFieldObjectCulling.cpp b/Engine/Source/Runtime/Renderer/Private/DistanceFieldObjectCulling.cpp
--- a/Engine/Source/Runtime/Renderer/Private/DistanceFieldObjectCulling.cpp
+++ b/Engine/Source/Runtime/Renderer/Private/DistanceFieldObjectCulling.cpp
@@ -337,17 +337,15 @@ void ScatterTilesToObjects(
 
 FIntPoint GetTileListGroupSizeForView(const FViewInfo& View)
 {
-	const FIntPoint AOBufferSize = GetBufferSizeForAO(View);
 	return FIntPoint(
-		FMath::DivideAndRoundUp(FMath::Max(AOBufferSize.X, 1), GDistanceFieldAOTileSizeX),
-		FMath::DivideAndRoundUp(FMath::Max(AOBufferSize.Y, 1), GDistanceFieldAOTileSizeY));
+		FMath::DivideAndRoundUp(FMath::Max(View.ViewRect.Size().X / GAODownsampleFactor, 1), GDistanceFieldAOTileSizeX),
+		FMath::DivideAndRoundUp(FMath::Max(View.ViewRect.Size().Y / GAODownsampleFactor, 1), GDistanceFieldAOTileSizeY));
 }
 
 void BuildTileObjectLists(
 	FRDGBuilder& GraphBuilder,
 	FScene* Scene,
-	const FViewInfo& View,
-	TRDGUniformBufferRef<FSceneTextureUniformParameters> SceneTexturesUniformBuffer,
+	TArray<FViewInfo>& Views,
 	FRDGBufferRef ObjectIndirectArguments,
 	const FDistanceFieldCulledObjectBufferParameters& CulledObjectBufferParameters,
 	FTileIntersectionParameters TileIntersectionParameters,
@@ -357,57 +355,64 @@ void BuildTileObjectLists(
 	ensure(GAOScatterTileCulling);
 
 	RDG_EVENT_SCOPE(GraphBuilder, "BuildTileList");
-	RDG_GPU_MASK_SCOPE(GraphBuilder, View.GPUMask);
 
-	const FIntPoint TileListGroupSize = GetTileListGroupSizeForView(View);
+	TRDGUniformBufferRef<FSceneTextureUniformParameters> SceneTexturesUniformBuffer = GetViewFamilyInfo(Views).GetSceneTextures().UniformBuffer;
 
+	for (int32 ViewIndex = 0; ViewIndex < Views.Num(); ViewIndex++)
 	{
-		auto* PassParameters = GraphBuilder.AllocParameters<FBuildTileConesCS::FParameters>();
-		PassParameters->View = View.ViewUniformBuffer;
-		PassParameters->RWTileConeAxisAndCos = TileIntersectionParameters.RWTileConeAxisAndCos;
-		PassParameters->RWTileConeDepthRanges = TileIntersectionParameters.RWTileConeDepthRanges;
-		PassParameters->SceneTextures = SceneTexturesUniformBuffer;
-		PassParameters->DistanceFieldNormalTexture = DistanceFieldNormal;
-		PassParameters->DistanceFieldNormalSampler = TStaticSamplerState<SF_Point, AM_Clamp, AM_Clamp, AM_Clamp>::GetRHI();
-		PassParameters->AOParameters = DistanceField::SetupAOShaderParameters(Parameters);
-		PassParameters->NumGroups = FVector2f(TileListGroupSize.X, TileListGroupSize.Y);
+		const FViewInfo& View = Views[ViewIndex];
+		RDG_GPU_MASK_SCOPE(GraphBuilder, View.GPUMask);
 
-		auto ComputeShader = View.ShaderMap->GetShader<FBuildTileConesCS>();
+		const FIntPoint TileListGroupSize = GetTileListGroupSizeForView(View);
 
-		FComputeShaderUtils::AddPass(GraphBuilder, RDG_EVENT_NAME("BuildTileCones"), ComputeShader, PassParameters, FIntVector(TileListGroupSize.X, TileListGroupSize.Y, 1));
-	}
+		{
+			auto* PassParameters = GraphBuilder.AllocParameters<FBuildTileConesCS::FParameters>();
+			PassParameters->View = View.ViewUniformBuffer;
+			PassParameters->RWTileConeAxisAndCos = TileIntersectionParameters.RWTileConeAxisAndCos;
+			PassParameters->RWTileConeDepthRanges = TileIntersectionParameters.RWTileConeDepthRanges;
+			PassParameters->SceneTextures = SceneTexturesUniformBuffer;
+			PassParameters->DistanceFieldNormalTexture = DistanceFieldNormal;
+			PassParameters->DistanceFieldNormalSampler = TStaticSamplerState<SF_Point, AM_Clamp, AM_Clamp, AM_Clamp>::GetRHI();
+			PassParameters->AOParameters = DistanceField::SetupAOShaderParameters(Parameters);
+			PassParameters->NumGroups = FVector2f(TileListGroupSize.X, TileListGroupSize.Y);
+
+			auto ComputeShader = View.ShaderMap->GetShader<FBuildTileConesCS>();
+
+			FComputeShaderUtils::AddPass(GraphBuilder, RDG_EVENT_NAME("BuildTileCones"), ComputeShader, PassParameters, FIntVector(TileListGroupSize.X, TileListGroupSize.Y, 1));
+		}
 
-	// Start at 0 tiles per object
-	AddClearUAVPass(GraphBuilder, TileIntersectionParameters.RWNumCulledTilesArray, 0);
+		// Start at 0 tiles per object
+		AddClearUAVPass(GraphBuilder, TileIntersectionParameters.RWNumCulledTilesArray, 0);
 
-	// Rasterize object bounding shapes and intersect with screen tiles to compute how many tiles intersect each object
-	ScatterTilesToObjects(GraphBuilder, true, View, Scene->DistanceFieldSceneData, TileListGroupSize, Parameters, ObjectIndirectArguments, CulledObjectBufferParameters, TileIntersectionParameters, SceneTexturesUniformBuffer);
+		// Rasterize object bounding shapes and intersect with screen tiles to compute how many tiles intersect each object
+		ScatterTilesToObjects(GraphBuilder, true, View, Scene->DistanceFieldSceneData, TileListGroupSize, Parameters, ObjectIndirectArguments, CulledObjectBufferParameters, TileIntersectionParameters, SceneTexturesUniformBuffer);
 
-	// Start at 0 threadgroups
-	AddClearUAVPass(GraphBuilder, TileIntersectionParameters.RWObjectTilesIndirectArguments, 0);
+		// Start at 0 threadgroups
+		AddClearUAVPass(GraphBuilder, TileIntersectionParameters.RWObjectTilesIndirectArguments, 0);
 
-	{
-		auto* PassParameters = GraphBuilder.AllocParameters<FComputeCulledTilesStartOffsetCS::FParameters>();
-		PassParameters->View = View.ViewUniformBuffer;
-		PassParameters->TileIntersectionParameters = TileIntersectionParameters;
-		PassParameters->DistanceFieldCulledObjectBuffers = CulledObjectBufferParameters;
-		PassParameters->DistanceFieldAtlas = DistanceField::SetupAtlasParameters(GraphBuilder, Scene->DistanceFieldSceneData);
-		PassParameters->SceneTextures = SceneTexturesUniformBuffer;
+		{
+			auto* PassParameters = GraphBuilder.AllocParameters<FComputeCulledTilesStartOffsetCS::FParameters>();
+			PassParameters->View = View.ViewUniformBuffer;
+			PassParameters->TileIntersectionParameters = TileIntersectionParameters;
+			PassParameters->DistanceFieldCulledObjectBuffers = CulledObjectBufferParameters;
+			PassParameters->DistanceFieldAtlas = DistanceField::SetupAtlasParameters(GraphBuilder, Scene->DistanceFieldSceneData);
+			PassParameters->SceneTextures = SceneTexturesUniformBuffer;
+
+			auto ComputeShader = View.ShaderMap->GetShader<FComputeCulledTilesStartOffsetCS>();
+			const FIntVector GroupCount = FComputeShaderUtils::GetGroupCountWrapped(Scene->DistanceFieldSceneData.NumObjectsInBuffer, ComputeStartOffsetGroupSize);
+
+			FComputeShaderUtils::AddPass(
+				GraphBuilder,
+				RDG_EVENT_NAME("ComputeStartOffsets"),
+				ComputeShader,
+				PassParameters,
+				GroupCount);
+		}
 
-		auto ComputeShader = View.ShaderMap->GetShader<FComputeCulledTilesStartOffsetCS>();
-		const FIntVector GroupCount = FComputeShaderUtils::GetGroupCountWrapped(Scene->DistanceFieldSceneData.NumObjectsInBuffer, ComputeStartOffsetGroupSize);
+		// Start at 0 tiles per object
+		AddClearUAVPass(GraphBuilder, TileIntersectionParameters.RWNumCulledTilesArray, 0);
 
-		FComputeShaderUtils::AddPass(
-			GraphBuilder,
-			RDG_EVENT_NAME("ComputeStartOffsets"),
-			ComputeShader,
-			PassParameters,
-			GroupCount);
+		// Rasterize object bounding shapes and intersect with screen tiles, and write out intersecting tile indices for the cone tracing pass
+		ScatterTilesToObjects(GraphBuilder, false, View, Scene->DistanceFieldSceneData, TileListGroupSize, Parameters, ObjectIndirectArguments, CulledObjectBufferParameters, TileIntersectionParameters, SceneTexturesUniformBuffer);
 	}
-
-	// Start at 0 tiles per object
-	AddClearUAVPass(GraphBuilder, TileIntersectionParameters.RWNumCulledTilesArray, 0);
-
-	// Rasterize object bounding shapes and intersect with screen tiles, and write out intersecting tile indices for the cone tracing pass
-	ScatterTilesToObjects(GraphBuilder, false, View, Scene->DistanceFieldSceneData, TileListGroupSize, Parameters, ObjectIndirectArguments, CulledObjectBufferParameters, TileIntersectionParameters, SceneTexturesUniformBuffer);
 }
