Fixed race condition with landscape buffers and shadow initialization.

diff --git a/Engine/Source/Runtime/Renderer/Private/DeferredShadingRenderer.cpp b/Engine/Source/Runtime/Renderer/Private/DeferredShadingRenderer.cpp
--- a/Engine/Source/Runtime/Renderer/Private/DeferredShadingRenderer.cpp
+++ b/Engine/Source/Runtime/Renderer/Private/DeferredShadingRenderer.cpp
@@ -2514,8 +2514,6 @@ void FDeferredShadingSceneRenderer::Render(FRDGBuilder& GraphBuilder)
 			const bool bEnableVirtualShadowMaps = UseVirtualShadowMaps(ShaderPlatform, FeatureLevel) && ViewFamily.EngineShowFlags.DynamicShadows;
 			VirtualShadowMapArray.Initialize(GraphBuilder, Scene->GetVirtualShadowMapCache(Views[0]), bEnableVirtualShadowMaps, Views[0].bIsSceneCapture);
 
-			BeginInitDynamicShadows(InitViewTaskDatas);
-
 			if (InitViewTaskDatas.LumenFrameTemporaries)
 			{
 				BeginUpdateLumenSceneTasks(GraphBuilder, *InitViewTaskDatas.LumenFrameTemporaries);
diff --git a/Engine/Source/Runtime/Renderer/Private/SceneVisibility.cpp b/Engine/Source/Runtime/Renderer/Private/SceneVisibility.cpp
--- a/Engine/Source/Runtime/Renderer/Private/SceneVisibility.cpp
+++ b/Engine/Source/Runtime/Renderer/Private/SceneVisibility.cpp
@@ -4818,6 +4818,12 @@ void FDeferredShadingSceneRenderer::BeginInitViews(
 
 	PreVisibilityFrameSetup(GraphBuilder);
 
+	// Attempt to launch dynamic shadow tasks early before finalizing visibility.
+	if (GetRendererOutput() == ERendererOutput::FinalSceneColor)
+	{
+		BeginInitDynamicShadows(TaskDatas);
+	}
+
 	FRHICommandListImmediate& RHICmdList = GraphBuilder.RHICmdList;
 
 	// Create GPU-side representation of the view for instance culling.
@@ -4845,6 +4851,7 @@ void FDeferredShadingSceneRenderer::BeginInitViews(
 
 	TaskDatas.VisibilityTaskData->ProcessRenderThreadTasks(BasePassDepthStencilAccess, InstanceCullingManager, VirtualTextureUpdater);
 
+	// Make a second attempt to launch shadow tasks it wasn't able to the first time due to visibility being deferred.
 	if (GetRendererOutput() == ERendererOutput::FinalSceneColor)
 	{
 		BeginInitDynamicShadows(TaskDatas);
