Fix incorrect dummy resource format in FSceneLightingChannelParameters when there's no light channel texture.

GetSceneLightingChannelParameters() was binding WhiteDummy, which is a RGBA texture, but the shader expects an R32_UINT texture.

#rnx

diff --git a/Engine/Source/Runtime/Renderer/Private/SceneTextureParameters.cpp b/Engine/Source/Runtime/Renderer/Private/SceneTextureParameters.cpp
--- a/Engine/Source/Runtime/Renderer/Private/SceneTextureParameters.cpp
+++ b/Engine/Source/Runtime/Renderer/Private/SceneTextureParameters.cpp
@@ -64,7 +64,7 @@ FSceneLightingChannelParameters GetSceneLightingChannelParameters(FRDGBuilder& G
 	}
 	else
 	{
-		Parameters.SceneLightingChannels = GraphBuilder.RegisterExternalTexture(GSystemTextures.WhiteDummy);
+		Parameters.SceneLightingChannels = GSystemTextures.GetZeroUIntDummy(GraphBuilder);
 		Parameters.bSceneLightingChannelsValid = false;
 	}
 
