Multicolored selection outlines

By placing an FTypedElementViewportColorColumn on an actor, the color of the selection outline in the viewport can be changed. Currently supports 8 colors from a palette.

diff --git a/Engine/Plugins/Experimental/TypedElementsDataStorage/Source/TypedElementsDataStorage/Private/Compatibility/TypedElementActorViewportProcessors.cpp b/Engine/Plugins/Experimental/TypedElementsDataStorage/Source/TypedElementsDataStorage/Private/Compatibility/TypedElementActorViewportProcessors.cpp
new file mode 100644
--- /dev/null
+++ b/Engine/Plugins/Experimental/TypedElementsDataStorage/Source/TypedElementsDataStorage/Private/Compatibility/TypedElementActorViewportProcessors.cpp
@@ -0,0 +1,41 @@
+﻿// Copyright Epic Games, Inc. All Rights Reserved.
+
+#include "Compatibility/TypedElementActorViewportProcessors.h"
+
+#include "Components/PrimitiveComponent.h"
+#include "Elements/Columns/TypedElementViewportColumns.h"
+#include "Elements/Columns/TypedElementMiscColumns.h"
+#include "Elements/Framework/TypedElementQueryBuilder.h"
+#include "MassActorSubsystem.h"
+
+void UTypedElementActorViewportFactory::RegisterQueries(ITypedElementDataStorageInterface& DataStorage) const
+{
+	RegisterSelectionOutlineColorColumnToActor(DataStorage);
+}
+
+void UTypedElementActorViewportFactory::RegisterSelectionOutlineColorColumnToActor(ITypedElementDataStorageInterface& DataStorage) const
+{
+	using namespace TypedElementQueryBuilder;
+	using DSI = ITypedElementDataStorageInterface;
+
+	DataStorage.RegisterQuery(
+		Select(
+			TEXT("Sync selection outline color column to actor"),
+			FProcessor(DSI::EQueryTickPhase::DuringPhysics, DataStorage.GetQueryTickGroupName(DSI::EQueryTickGroups::SyncDataStorageToExternal))
+				.ForceToGameThread(true),
+			[](FMassActorFragment& Actor, const FTypedElementViewportColorColumn& ViewportColor)
+			{
+				if (AActor* ActorInstance = Actor.GetMutable(); ActorInstance != nullptr)
+				{
+					if (UPrimitiveComponent* PrimitiveComponent = ActorInstance->GetComponentByClass<UPrimitiveComponent>())
+					{
+						PrimitiveComponent->SetSelectionOutlineColorIndex(ViewportColor.SelectionOutlineColorIndex);
+					}
+				}
+			}
+		)
+		.Where()
+			.All<FTypedElementSyncBackToWorldTag>()
+		.Compile()
+	);
+}
diff --git a/Engine/Plugins/Experimental/TypedElementsDataStorage/Source/TypedElementsDataStorage/Private/Compatibility/TypedElementActorViewportProcessors.h b/Engine/Plugins/Experimental/TypedElementsDataStorage/Source/TypedElementsDataStorage/Private/Compatibility/TypedElementActorViewportProcessors.h
new file mode 100644
--- /dev/null
+++ b/Engine/Plugins/Experimental/TypedElementsDataStorage/Source/TypedElementsDataStorage/Private/Compatibility/TypedElementActorViewportProcessors.h
@@ -0,0 +1,22 @@
+﻿// Copyright Epic Games, Inc. All Rights Reserved.
+
+#pragma once
+
+#include "Elements/Interfaces/TypedElementDataStorageFactory.h"
+#include "UObject/ObjectMacros.h"
+
+#include "TypedElementActorViewportProcessors.generated.h"
+
+UCLASS()
+class UTypedElementActorViewportFactory : public UTypedElementDataStorageFactory
+{
+	GENERATED_BODY()
+
+public:
+	~UTypedElementActorViewportFactory() override = default;
+
+	void RegisterQueries(ITypedElementDataStorageInterface& DataStorage) const override;
+
+private:
+	void RegisterSelectionOutlineColorColumnToActor(ITypedElementDataStorageInterface& DataStorage) const;
+};
diff --git a/Engine/Shaders/Private/PostProcessSelectionOutline.usf b/Engine/Shaders/Private/PostProcessSelectionOutline.usf
--- a/Engine/Shaders/Private/PostProcessSelectionOutline.usf
+++ b/Engine/Shaders/Private/PostProcessSelectionOutline.usf
@@ -24,16 +24,15 @@ SCREEN_PASS_TEXTURE_VIEWPORT(Color)
 SCREEN_PASS_TEXTURE_VIEWPORT(Depth)
 FScreenTransform ColorToDepth;
 
-float3 OutlineColor;
-float3 SubduedOutlineColor;
+float4 OutlineColors[8];
+int OutlineColorIndexBits;
 float SelectionHighlightIntensity;
 float BSPSelectionIntensity;
 float UILuminanceAndIsSCRGB;
 
 struct FPixelInfo
 {
-	// if there is an selected object
-	bool bObjectMask;
+	float fDepth;
 
 	int ObjectId;
 	// if the current pixel is not occluded by some other object in front of it
@@ -78,9 +77,8 @@ FPixelInfo GetPixelInfo(int2 PixelPos, int SampleID, int OffsetX, int OffsetY, f
 	// this avoids flicking artifacts on the silhouette by limiting the depth reconstruction error
 	SceneDeviceZ = max(SceneDeviceZ, DeviceZMinMax.x);
 	SceneDeviceZ = min(SceneDeviceZ, DeviceZMinMax.y);
-
-	// test against far plane
-	Ret.bObjectMask = DeviceZ != 0.0f;
+	
+	Ret.fDepth = DeviceZ;
 
 	// outline even between multiple selected objects (best usability)
 	Ret.ObjectId = Stencil;
@@ -146,6 +144,14 @@ float4 PerSample(int2 PixelPos, int SampleID, FPixelInfo CenterPixelInfo, float3
 		FLATTEN	if (PixelInfo[3].bVisible) MinMax(VisibleBorderMinMax, PixelInfo[3].ObjectId);
 		FLATTEN	if (PixelInfo[4].bVisible) MinMax(VisibleBorderMinMax, PixelInfo[4].ObjectId);
 	}
+	
+	int ClosestPixelIndex = 0;
+	{
+		FLATTEN if (PixelInfo[1].fDepth > PixelInfo[ClosestPixelIndex].fDepth) ClosestPixelIndex = 1;
+		FLATTEN if (PixelInfo[2].fDepth > PixelInfo[ClosestPixelIndex].fDepth) ClosestPixelIndex = 2;
+		FLATTEN if (PixelInfo[3].fDepth > PixelInfo[ClosestPixelIndex].fDepth) ClosestPixelIndex = 3;
+		FLATTEN if (PixelInfo[4].fDepth > PixelInfo[ClosestPixelIndex].fDepth) ClosestPixelIndex = 4;
+	}
 
 	// this border around the object
 	bool bVisibleBorder = VisibleBorderMinMax.y != 0;
@@ -158,15 +164,12 @@ float4 PerSample(int2 PixelPos, int SampleID, FPixelInfo CenterPixelInfo, float3
 	// the contants express the two opacity values we see when the primitive is hidden
 	float LowContrastPatternMask = lerp(0.2, 1.0f, PatternMask);
 
-	float3  SelectionColor;
-	FLATTEN if (VisibleBorderMinMax.x >= 128 )
-	{
-		SelectionColor = SubduedOutlineColor.rgb;
-	}
-	else
-	{
-		SelectionColor = OutlineColor.rgb;
-	}
+	const int BitsAvailable = 8;
+	const int ColorShiftDistance = (BitsAvailable - OutlineColorIndexBits);
+	const int ColorMask = ((0xFFu >> ColorShiftDistance) << ColorShiftDistance);
+	int SelectionColorId = (PixelInfo[ClosestPixelIndex].ObjectId & ColorMask) >> ColorShiftDistance;
+	float3 SelectionColor = OutlineColors[SelectionColorId].rgb;
+	
 	FLATTEN if (bBorder)
 	{
 		FLATTEN if (bVisibleBorder)
@@ -191,8 +194,8 @@ float4 PerSample(int2 PixelPos, int SampleID, FPixelInfo CenterPixelInfo, float3
 		Color = lerp(Color, float4(SelectionColor * VisibleMask, 1), SelectionHighlightIntensity);
 	}
 
-	// highlight inner part of the object
-	if (PixelInfo[0].bObjectMask)
+	// highlight inner part of the object, test against far plane
+	if (PixelInfo[0].fDepth != 0.0f)
 	{
 		float VisibleMask = lerp(PatternMask, 1.0f, PixelInfo[0].fVisible);
 		float InnerHighlightAmount = PixelInfo[0].ObjectId == 1 ? BSPSelectionIntensity : SelectionHighlightIntensity;
diff --git a/Engine/Source/Editor/UnrealEd/Classes/Settings/EditorStyleSettings.h b/Engine/Source/Editor/UnrealEd/Classes/Settings/EditorStyleSettings.h
--- a/Engine/Source/Editor/UnrealEd/Classes/Settings/EditorStyleSettings.h
+++ b/Engine/Source/Editor/UnrealEd/Classes/Settings/EditorStyleSettings.h
@@ -74,6 +74,10 @@ public:
 	UPROPERTY(EditAnywhere, config, Category=UserInterface, meta=(DisplayName="Viewport Selection Color"))
 	FLinearColor SelectionColor;
 
+	/** Additional colors used for selections with extra meaning */
+	UPROPERTY(EditAnywhere, config, Category=UserInterface, meta=(DisplayName="Additional Viewport Selection Colors"))
+	FLinearColor AdditionalSelectionColors[6];
+
 	UPROPERTY(config)
 	bool bEnableEditorWindowBackgroundColor;
 
diff --git a/Engine/Source/Editor/UnrealEd/Private/EditorViewportClient.cpp b/Engine/Source/Editor/UnrealEd/Private/EditorViewportClient.cpp
--- a/Engine/Source/Editor/UnrealEd/Private/EditorViewportClient.cpp
+++ b/Engine/Source/Editor/UnrealEd/Private/EditorViewportClient.cpp
@@ -69,6 +69,7 @@
 #include "IImageWrapperModule.h"
 #include "HDRHelper.h"
 #include "GlobalRenderResources.h"
+#include "Settings/EditorStyleSettings.h"
 
 #define LOCTEXT_NAMESPACE "EditorViewportClient"
 
@@ -1309,6 +1310,12 @@ FSceneView* FEditorViewportClient::CalcSceneView(FSceneViewFamily* ViewFamily, c
 	View->ViewRotation = ModifiedViewRotation;
 
 	View->SubduedSelectionOutlineColor = GEngine->GetSubduedSelectionOutlineColor();
+	const UEditorStyleSettings* EditorStyle = GetDefault<UEditorStyleSettings>();
+	check(View->AdditionalSelectionOutlineColors.Num() <= UE_ARRAY_COUNT(EditorStyle->AdditionalSelectionColors))
+	for (int OutlineColorIndex = 0; OutlineColorIndex < View->AdditionalSelectionOutlineColors.Num(); ++OutlineColorIndex)
+	{
+		View->AdditionalSelectionOutlineColors[OutlineColorIndex] = EditorStyle->AdditionalSelectionColors[OutlineColorIndex];
+	}
 
 	int32 FamilyIndex = ViewFamily->Views.Add(View);
 	check(FamilyIndex == View->StereoViewIndex || View->StereoViewIndex == INDEX_NONE);
diff --git a/Engine/Source/Editor/UnrealEd/Private/Settings/EditorStyleClasses.cpp b/Engine/Source/Editor/UnrealEd/Private/Settings/EditorStyleClasses.cpp
--- a/Engine/Source/Editor/UnrealEd/Private/Settings/EditorStyleClasses.cpp
+++ b/Engine/Source/Editor/UnrealEd/Private/Settings/EditorStyleClasses.cpp
@@ -23,6 +23,14 @@ class FObjectInitializer;
 
 UEditorStyleSettings::UEditorStyleSettings( const FObjectInitializer& ObjectInitializer )
 	: Super(ObjectInitializer)
+	, AdditionalSelectionColors{
+		FStyleColors::AccentBlue.GetSpecifiedColor(),
+		FStyleColors::AccentPurple.GetSpecifiedColor(),
+		FStyleColors::AccentPink.GetSpecifiedColor(),
+		FStyleColors::AccentRed.GetSpecifiedColor(),
+		FStyleColors::AccentYellow.GetSpecifiedColor(),
+		FStyleColors::AccentGreen.GetSpecifiedColor(),
+	}
 {
 	bEnableUserEditorLayoutManagement = true;
 
diff --git a/Engine/Source/Runtime/Engine/Classes/Components/PrimitiveComponent.h b/Engine/Source/Runtime/Engine/Classes/Components/PrimitiveComponent.h
--- a/Engine/Source/Runtime/Engine/Classes/Components/PrimitiveComponent.h
+++ b/Engine/Source/Runtime/Engine/Classes/Components/PrimitiveComponent.h
@@ -689,6 +689,8 @@ protected:
 #if WITH_EDITOR
 public:
 	uint8 bAlwaysAllowTranslucentSelect : 1;
+
+	uint8 SelectionOutlineColorIndex;
 #endif
 
 public:
@@ -2518,6 +2520,9 @@ public:
 
 	/** Sets whether this component is being moved by the editor so the renderer can render velocities for it, even when Static. */
 	ENGINE_API void SetIsBeingMovedByEditor(bool bNewIsBeingMoved);
+
+
+	ENGINE_API void SetSelectionOutlineColorIndex(uint8 SelectionOutlineColorIndex);
 #endif// WITH_EDITOR
 
 	/**
diff --git a/Engine/Source/Runtime/Engine/Private/Components/PrimitiveComponent.cpp b/Engine/Source/Runtime/Engine/Private/Components/PrimitiveComponent.cpp
--- a/Engine/Source/Runtime/Engine/Private/Components/PrimitiveComponent.cpp
+++ b/Engine/Source/Runtime/Engine/Private/Components/PrimitiveComponent.cpp
@@ -376,6 +376,8 @@ UPrimitiveComponent::UPrimitiveComponent(const FObjectInitializer& ObjectInitial
 
 #if WITH_EDITOR
 	bAlwaysAllowTranslucentSelect = false;
+
+	SelectionOutlineColorIndex = 0;
 #endif
 
 #if WITH_EDITORONLY_DATA
@@ -1882,6 +1884,16 @@ void UPrimitiveComponent::SetIsBeingMovedByEditor(bool bIsBeingMoved)
 		SceneProxy->SetIsBeingMovedByEditor_GameThread(bIsBeingMoved);
 	}
 }
+
+void UPrimitiveComponent::SetSelectionOutlineColorIndex(uint8 InSelectionOutlineColorIndex)
+{
+	SelectionOutlineColorIndex = InSelectionOutlineColorIndex;
+	
+	if (SceneProxy)
+	{
+		SceneProxy->SetSelectionOutlineColorIndex_GameThread(InSelectionOutlineColorIndex);
+	}
+}
 #endif// WITH_EDITOR
 
 void UPrimitiveComponent::PushHoveredToProxy(const bool bInHovered)
diff --git a/Engine/Source/Runtime/Engine/Private/PrimitiveSceneProxy.cpp b/Engine/Source/Runtime/Engine/Private/PrimitiveSceneProxy.cpp
--- a/Engine/Source/Runtime/Engine/Private/PrimitiveSceneProxy.cpp
+++ b/Engine/Source/Runtime/Engine/Private/PrimitiveSceneProxy.cpp
@@ -1343,6 +1343,20 @@ void FPrimitiveSceneProxy::SetIsBeingMovedByEditor_GameThread(bool bIsBeingMoved
 			PrimitiveSceneProxy->GetScene().UpdatePrimitiveVelocityState_RenderThread(PrimitiveSceneProxy->GetPrimitiveSceneInfo(), bIsBeingMoved);
 		});
 }
+
+void FPrimitiveSceneProxy::SetSelectionOutlineColorIndex_GameThread(uint8 ColorIndex)
+{
+	check(IsInGameThread());
+	constexpr int IndexBits = 3;
+	constexpr int MaxIndex = (1 << IndexBits) - 1;
+	check(ColorIndex <= MaxIndex);
+
+	ENQUEUE_RENDER_COMMAND(SetSelectionOutlineColorIndex)(
+		[this, ColorIndex](FRHICommandListImmediate&)
+		{
+			SelectionOutlineColorIndex = ColorIndex;
+		});
+}
 #endif
 
 void FPrimitiveSceneProxy::SetEvaluateWorldPositionOffset_GameThread(bool bEvaluate)
diff --git a/Engine/Source/Runtime/Engine/Private/SceneView.cpp b/Engine/Source/Runtime/Engine/Private/SceneView.cpp
--- a/Engine/Source/Runtime/Engine/Private/SceneView.cpp
+++ b/Engine/Source/Runtime/Engine/Private/SceneView.cpp
@@ -874,6 +874,7 @@ PRAGMA_ENABLE_DEPRECATION_WARNINGS
 	EditorViewBitflag = InitOptions.EditorViewBitflag;
 
 	SelectionOutlineColor = GEngine->GetSelectionOutlineColor();
+	SubduedSelectionOutlineColor = GEngine->GetSubduedSelectionOutlineColor();
 #endif
 
 	// Query instanced stereo and multi-view state
diff --git a/Engine/Source/Runtime/Engine/Public/PrimitiveSceneProxy.h b/Engine/Source/Runtime/Engine/Public/PrimitiveSceneProxy.h
--- a/Engine/Source/Runtime/Engine/Public/PrimitiveSceneProxy.h
+++ b/Engine/Source/Runtime/Engine/Public/PrimitiveSceneProxy.h
@@ -234,6 +234,11 @@ public:
 	 * Enqueue and update for the render thread to notify it that the editor is currently moving the owning component with gizmos.
 	 */
 	void SetIsBeingMovedByEditor_GameThread(bool bIsBeingMoved);
+
+	/**
+	 * Enqueue updated selection outline color for the render thread to use.
+	 */
+	void SetSelectionOutlineColorIndex_GameThread(uint8 ColorIndex);
 #endif	// WITH_EDITOR
 
 	/** Enqueue updated setting for evaluation of World Position Offset. */
@@ -684,6 +689,7 @@ public:
 	inline bool HasPerInstancePayloadExtension() const { return bHasPerInstancePayloadExtension; }
 #if WITH_EDITOR
 	inline bool HasPerInstanceEditorData() const { return bHasPerInstanceEditorData; }
+	inline uint8 GetSelectionOutlineColorIndex() const { return SelectionOutlineColorIndex; }
 #else
 	FORCEINLINE bool HasPerInstanceEditorData() const { return false; }
 #endif // WITH_EDITOR
@@ -1533,6 +1539,8 @@ private:
 	/** A copy of the actor's group membership for handling per-view group hiding */
 	uint64 HiddenEditorViews;
 
+	uint32 SelectionOutlineColorIndex : 8;
+
 	/** Whether this should only draw in any editing mode*/
 	uint32 DrawInAnyEditMode : 1;
 
diff --git a/Engine/Source/Runtime/Engine/Public/SceneView.h b/Engine/Source/Runtime/Engine/Public/SceneView.h
--- a/Engine/Source/Runtime/Engine/Public/SceneView.h
+++ b/Engine/Source/Runtime/Engine/Public/SceneView.h
@@ -1462,6 +1462,8 @@ public:
 	FLinearColor SelectionOutlineColor;
 	/** Selection color for use in the editor with inactive primitives */
 	FLinearColor SubduedSelectionOutlineColor;
+	/** Additional selection colors for the editor (used by post processing) */
+	TStaticArray<FLinearColor, 6> AdditionalSelectionOutlineColors;
 	/** True if any components are selected in isolation (independent of actor selection) */
 	bool bHasSelectedComponents;
 #endif
diff --git a/Engine/Source/Runtime/Renderer/Private/PostProcess/PostProcessSelectionOutline.cpp b/Engine/Source/Runtime/Renderer/Private/PostProcess/PostProcessSelectionOutline.cpp
--- a/Engine/Source/Runtime/Renderer/Private/PostProcess/PostProcessSelectionOutline.cpp
+++ b/Engine/Source/Runtime/Renderer/Private/PostProcess/PostProcessSelectionOutline.cpp
@@ -33,9 +33,9 @@ public:
 		SHADER_PARAMETER_RDG_TEXTURE(Texture2D, EditorPrimitivesDepth)
 		SHADER_PARAMETER_RDG_TEXTURE_SRV(Texture2D, EditorPrimitivesStencil)
 		SHADER_PARAMETER(FScreenTransform, ColorToDepth)
-		SHADER_PARAMETER(FVector3f, OutlineColor)
+		SHADER_PARAMETER_ARRAY(FVector4f, OutlineColors, [8])
+		SHADER_PARAMETER(int, OutlineColorIndexBits)
 		SHADER_PARAMETER(float, SelectionHighlightIntensity)
-		SHADER_PARAMETER(FVector3f, SubduedOutlineColor)
 		SHADER_PARAMETER(float, BSPSelectionIntensity)
 		SHADER_PARAMETER(float, UILuminanceAndIsSCRGB)
 		RENDER_TARGET_BINDING_SLOTS()
@@ -220,9 +220,14 @@ FScreenPassTexture AddSelectionOutlinePass(
 		PassParameters->DepthSampler = PointClampSampler;
 		PassParameters->EditorPrimitivesDepth = DepthStencilTexture;
 		PassParameters->EditorPrimitivesStencil = GraphBuilder.CreateSRV(FRDGTextureSRVDesc::CreateWithPixelFormat(DepthStencilTexture, PF_X24_G8));
-		PassParameters->OutlineColor = FVector3f(View.SelectionOutlineColor);
+		PassParameters->OutlineColors[0] = View.SelectionOutlineColor;
+		PassParameters->OutlineColors[1] = View.SubduedSelectionOutlineColor;
+		for (int OutlineColorIndex = 2; OutlineColorIndex < PassParameters->OutlineColors.Num(); ++OutlineColorIndex)
+		{
+			PassParameters->OutlineColors[OutlineColorIndex] = View.AdditionalSelectionOutlineColors[OutlineColorIndex];
+		}
+		PassParameters->OutlineColorIndexBits = 3;
 		PassParameters->SelectionHighlightIntensity = GEngine->SelectionHighlightIntensity;
-		PassParameters->SubduedOutlineColor = FVector3f(View.SubduedSelectionOutlineColor);
 		PassParameters->BSPSelectionIntensity = GEngine->BSPSelectionHighlightIntensity;
 
 		EDisplayOutputFormat DisplayOutputFormat = View.Family->RenderTarget->GetDisplayOutputFormat();
diff --git a/Engine/Source/Runtime/Renderer/Private/SceneHitProxyRendering.cpp b/Engine/Source/Runtime/Renderer/Private/SceneHitProxyRendering.cpp
--- a/Engine/Source/Runtime/Renderer/Private/SceneHitProxyRendering.cpp
+++ b/Engine/Source/Runtime/Renderer/Private/SceneHitProxyRendering.cpp
@@ -34,6 +34,7 @@
 #include "GPUMessaging.h"
 #include "HairStrands/HairStrandsData.h"
 #include "SimpleMeshDrawCommandPass.h"
+#include "StaticMeshSceneProxy.h"
 
 class FHitProxyShaderElementData : public FMeshMaterialShaderElementData
 {
@@ -1016,11 +1017,45 @@ int32 FEditorSelectionMeshProcessor::GetStencilValue(const FSceneView* View, con
 
 	const int32* ExistingStencilValue = PrimitiveSceneProxy->IsIndividuallySelected() ? ProxyToStencilIndex.Find(PrimitiveSceneProxy) : ActorNameToStencilIndex.Find(PrimitiveSceneProxy->GetOwnerName());
 
-	int32 StencilValue = 0;
+	// Reserved values for the stencil buffer that carry specific meaning
+	enum ESelectionStencilValues : int32
+	{
+		NotSelected = 0,
+		BSP = 1, // The outlines of all BSPs should be merged
+
+		COUNT,
+	};
+
+	static constexpr int BitsAvailable = 8; // Stencil buffer is 8-bit
+	static constexpr int ColorBits = 3; // Can be changed
+	static constexpr int UniqueIdBits = BitsAvailable - ColorBits;
+	static constexpr int MaxColor = (1 << ColorBits);
+	static constexpr int MaxUniqueId = (1 << UniqueIdBits);
+	
+	auto EncodeSelectionStencilValue = [](int32 ColorIndex, int32 UniqueId) -> int32
+	{
+		uint8 Bits = 0;
+		const int ColorShiftDistance = BitsAvailable - ColorBits;
+		const uint8 ColorMask = (0xFFu >> ColorShiftDistance) << ColorShiftDistance;
+		const uint8 UniqueIdMask = 0xFF >> (BitsAvailable - UniqueIdBits);
+		Bits |= ((ColorIndex % MaxColor) << ColorShiftDistance) & ColorMask;
+		// Allow all colors except one to use the full range of unreserved values
+		if (ColorIndex == 0)
+		{
+			Bits |= (UniqueId % (MaxUniqueId - ESelectionStencilValues::COUNT) + ESelectionStencilValues::COUNT) & UniqueIdMask;
+		}
+		else
+		{
+			Bits |= (UniqueId % MaxUniqueId) & UniqueIdMask;
+		}
+		return Bits;
+	};
+	
+	int32 StencilValue = ESelectionStencilValues::NotSelected;
 
 	if (PrimitiveSceneProxy->GetOwnerName() == NAME_BSP)
 	{
-		StencilValue = 1;
+		StencilValue = ESelectionStencilValues::BSP;
 	}
 	else if (ExistingStencilValue != nullptr)
 	{
@@ -1028,14 +1063,21 @@ int32 FEditorSelectionMeshProcessor::GetStencilValue(const FSceneView* View, con
 	}
 	else if (PrimitiveSceneProxy->IsIndividuallySelected())
 	{
-		// Any component that is individually selected should have a stencil value of < 128 so that it can have a unique color.  We offset the value by 2 because 0 means no selection and 1 is for bsp
-		StencilValue = ProxyToStencilIndex.Num() % 126 + 2;
+		const int Color = 0;
+		const int UniqueId = ProxyToStencilIndex.Num();
+		StencilValue = EncodeSelectionStencilValue(Color, UniqueId);
 		ProxyToStencilIndex.Add(PrimitiveSceneProxy, StencilValue);
 	}
 	else
 	{
-		// If we are subduing actor color highlight then use the top level bits to indicate that to the shader.  
-		StencilValue = bActorSelectionColorIsSubdued ? ActorNameToStencilIndex.Num() % 128 + 128 : ActorNameToStencilIndex.Num() % 126 + 2;
+		int Color = PrimitiveSceneProxy->GetSelectionOutlineColorIndex();
+		if (bActorSelectionColorIsSubdued && (Color == 0))
+		{
+			Color = 1;
+		}
+		const int UniqueId = ActorNameToStencilIndex.Num();
+
+		StencilValue = EncodeSelectionStencilValue(Color, UniqueId);
 		ActorNameToStencilIndex.Add(PrimitiveSceneProxy->GetOwnerName(), StencilValue);
 	}
 
diff --git a/Engine/Source/Runtime/TypedElementFramework/Public/Elements/Columns/TypedElementViewportColumns.h b/Engine/Source/Runtime/TypedElementFramework/Public/Elements/Columns/TypedElementViewportColumns.h
new file mode 100644
--- /dev/null
+++ b/Engine/Source/Runtime/TypedElementFramework/Public/Elements/Columns/TypedElementViewportColumns.h
@@ -0,0 +1,20 @@
+﻿// Copyright Epic Games, Inc. All Rights Reserved.
+
+#pragma once
+
+#include "Elements/Interfaces/TypedElementDataStorageInterface.h"
+#include "UObject/ObjectMacros.h"
+
+#include "TypedElementViewportColumns.generated.h"
+
+/**
+ * Column to hold the color that the object is outlined with when selected in the viewport
+ */
+USTRUCT(meta = (DisplayName = "Viewport Color"))
+struct FTypedElementViewportColorColumn final : public FTypedElementDataStorageColumn
+{
+	GENERATED_BODY()
+
+	UPROPERTY(meta = (ClampMin = "0", ClampMax = "7"))
+	uint8 SelectionOutlineColorIndex;
+};
