Change a check() to an ensure() since it doesn't appear to be a fatal condition.

diff --git a/Engine/Source/Runtime/Renderer/Private/DistanceFieldShadowing.cpp b/Engine/Source/Runtime/Renderer/Private/DistanceFieldShadowing.cpp
--- a/Engine/Source/Runtime/Renderer/Private/DistanceFieldShadowing.cpp
+++ b/Engine/Source/Runtime/Renderer/Private/DistanceFieldShadowing.cpp
@@ -861,7 +861,7 @@ FRDGTextureRef FProjectedShadowInfo::RenderRayTracedDistanceFieldProjection(
 
 		if (Scene->DistanceFieldSceneData.NumObjectsInBuffer > 0)
 		{
-			check(!Scene->DistanceFieldSceneData.HasPendingOperations());
+			ensure(!Scene->DistanceFieldSceneData.HasPendingOperations());
 
 			FDistanceFieldObjectBufferParameters ObjectBufferParameters = DistanceField::SetupObjectBufferParameters(GraphBuilder, Scene->DistanceFieldSceneData);
 
