Moved cloud shadow evalaution before the both skyatmosphere lut generations location for cloud shadow to always be valid, when using async or not.

diff --git a/Engine/Source/Runtime/Renderer/Private/DeferredShadingRenderer.cpp b/Engine/Source/Runtime/Renderer/Private/DeferredShadingRenderer.cpp
--- a/Engine/Source/Runtime/Renderer/Private/DeferredShadingRenderer.cpp
+++ b/Engine/Source/Runtime/Renderer/Private/DeferredShadingRenderer.cpp
@@ -3377,6 +3377,9 @@ void FDeferredShadingSceneRenderer::Render(FRDGBuilder& GraphBuilder)
 		CSV_CUSTOM_STAT(LightCount, Batched, float(SortedLightSet.UnbatchedLightStart), ECsvCustomStatOp::Set);
 		CSV_CUSTOM_STAT(LightCount, Unbatched, float(SortedLightSet.SortedLights.Num()) - float(SortedLightSet.UnbatchedLightStart), ECsvCustomStatOp::Set);
 
+		// Run before RenderSkyAtmosphereLookUpTables for cloud shadows to be valid.
+		InitVolumetricCloudsForViews(GraphBuilder, bShouldRenderVolumetricCloudBase, InstanceCullingManager);
+
 		if (SkyAtmospherePassLocation == ESkyAtmospherePassLocation::BeforeOcclusion && bShouldRenderSkyAtmosphere)
 		{
 			// Generate the Sky/Atmosphere look up tables
@@ -3404,8 +3407,6 @@ void FDeferredShadingSceneRenderer::Render(FRDGBuilder& GraphBuilder)
 			ResetVolumetricRenderTargetForViews(GraphBuilder, Views);
 		}
 
-		InitVolumetricCloudsForViews(GraphBuilder, bShouldRenderVolumetricCloudBase, InstanceCullingManager);
-
 		// Generate sky LUTs
 		// TODO: Valid shadow maps (for volumetric light shafts) have not yet been generated at this point in the frame. Need to resolve dependency ordering!
 		// This also must happen before the BasePass for Sky material to be able to sample valid LUTs.
diff --git a/Engine/Source/Runtime/Renderer/Private/SkyAtmosphereRendering.cpp b/Engine/Source/Runtime/Renderer/Private/SkyAtmosphereRendering.cpp
--- a/Engine/Source/Runtime/Renderer/Private/SkyAtmosphereRendering.cpp
+++ b/Engine/Source/Runtime/Renderer/Private/SkyAtmosphereRendering.cpp
@@ -219,7 +219,7 @@ static TAutoConsoleVariable<int32> CVarSkyAtmosphereEditorNotifications(
 
 static TAutoConsoleVariable<bool> CVarSkyAtmosphereASyncCompute(
 	TEXT("r.SkyAtmosphereASyncCompute"), false,
-	TEXT("SkyAtmosphere on async compute (default: false).\n"),
+	TEXT("SkyAtmosphere on async compute (default: false). When running on the async pipe, SkyAtmosphere lut generation will overlap with the occlusion pass.\n"),
 	ECVF_RenderThreadSafe);
 
 
@@ -286,7 +286,7 @@ ESkyAtmospherePassLocation GetSkyAtmospherePassLocation()
 {
 	if (CVarSkyAtmosphereASyncCompute.GetValueOnAnyThread())
 	{
-		// When SkyAtmLUT is running on async compute, try to kick it before occlusion pass to have better wave occupency
+		// When SkyAtmLUT is running on async compute, try to kick it before occlusion pass to have better wave occupancy
 		return ESkyAtmospherePassLocation::BeforeOcclusion;
 	}
 	// When SkyAtmLUT is running on graphics pipe, kickbefore BasePass to have a better overlap when DistanceFieldShadows is running async
