Fix uninitialized property and loop overread

diff --git a/Engine/Source/Runtime/Renderer/Private/PostProcess/PostProcessSelectionOutline.cpp b/Engine/Source/Runtime/Renderer/Private/PostProcess/PostProcessSelectionOutline.cpp
--- a/Engine/Source/Runtime/Renderer/Private/PostProcess/PostProcessSelectionOutline.cpp
+++ b/Engine/Source/Runtime/Renderer/Private/PostProcess/PostProcessSelectionOutline.cpp
@@ -224,7 +224,7 @@ FScreenPassTexture AddSelectionOutlinePass(
 		PassParameters->OutlineColors[1] = View.SubduedSelectionOutlineColor;
 		for (int OutlineColorIndex = 2; OutlineColorIndex < PassParameters->OutlineColors.Num(); ++OutlineColorIndex)
 		{
-			PassParameters->OutlineColors[OutlineColorIndex] = View.AdditionalSelectionOutlineColors[OutlineColorIndex];
+			PassParameters->OutlineColors[OutlineColorIndex] = View.AdditionalSelectionOutlineColors[OutlineColorIndex - 2];
 		}
 		PassParameters->OutlineColorIndexBits = 3;
 		PassParameters->SelectionHighlightIntensity = GEngine->SelectionHighlightIntensity;
diff --git a/Engine/Source/Runtime/TypedElementFramework/Public/Elements/Columns/TypedElementViewportColumns.h b/Engine/Source/Runtime/TypedElementFramework/Public/Elements/Columns/TypedElementViewportColumns.h
--- a/Engine/Source/Runtime/TypedElementFramework/Public/Elements/Columns/TypedElementViewportColumns.h
+++ b/Engine/Source/Runtime/TypedElementFramework/Public/Elements/Columns/TypedElementViewportColumns.h
@@ -16,5 +16,5 @@ struct FTypedElementViewportColorColumn final : public FTypedElementDataStorageC
 	GENERATED_BODY()
 
 	UPROPERTY(meta = (ClampMin = "0", ClampMax = "7"))
-	uint8 SelectionOutlineColorIndex;
+	uint8 SelectionOutlineColorIndex = 0;
 };
