Add SCOPED_NAMED_EVENTs and RDG_EVENT_SCOPESs to the main VRS manager functions (prepare images, get image, and combine images) to allow performance profiling in Unreal Insights and the editor itself. Renamed some existing CAS events for clarity.

diff --git a/Engine/Source/Runtime/Renderer/Private/VariableRateShading/ContrastAdaptiveImageGenerator.cpp b/Engine/Source/Runtime/Renderer/Private/VariableRateShading/ContrastAdaptiveImageGenerator.cpp
--- a/Engine/Source/Runtime/Renderer/Private/VariableRateShading/ContrastAdaptiveImageGenerator.cpp
+++ b/Engine/Source/Runtime/Renderer/Private/VariableRateShading/ContrastAdaptiveImageGenerator.cpp
@@ -447,7 +447,7 @@ void AddPrepareImageBasedVRSPass(
 
 		FComputeShaderUtils::AddPass(
 			GraphBuilder,
-			RDG_EVENT_NAME("PrepareImageBasedVRS"),
+			RDG_EVENT_NAME("ReprojectShadingRateImage"),
 			ERDGPassFlags::AsyncCompute | ERDGPassFlags::NeverCull,
 			RescaleVariableRateShadingCS,
 			PassParameters,
@@ -477,7 +477,7 @@ FRDGTextureRef FContrastAdaptiveImageGenerator::GetImage(FRDGBuilder& GraphBuild
 
 void FContrastAdaptiveImageGenerator::PrepareImages(FRDGBuilder& GraphBuilder, const FSceneViewFamily& ViewFamily, const FMinimalSceneTextures& SceneTextures)
 {
-	RDG_EVENT_SCOPE(GraphBuilder, "VariableRateShading");
+	RDG_EVENT_SCOPE(GraphBuilder, "ContrastAdaptiveShading");
 	bool bAreAllViewsVRSCompatible = true;
 	for (const FSceneView* View : ViewFamily.Views)
 	{
diff --git a/Engine/Source/Runtime/Renderer/Private/VariableRateShading/VariableRateShadingImageManager.cpp b/Engine/Source/Runtime/Renderer/Private/VariableRateShading/VariableRateShadingImageManager.cpp
--- a/Engine/Source/Runtime/Renderer/Private/VariableRateShading/VariableRateShadingImageManager.cpp
+++ b/Engine/Source/Runtime/Renderer/Private/VariableRateShading/VariableRateShadingImageManager.cpp
@@ -323,6 +323,9 @@ FRDGTextureRef FVariableRateShadingImageManager::GetVariableRateShadingImage(FRD
 		return nullptr;
 	}
 
+	RDG_EVENT_SCOPE(GraphBuilder, "GetVariableRateShadingImage");
+	SCOPED_NAMED_EVENT(GetVariableRateShadingImage, FColor::Yellow);
+
 	// Use debug rate if provided
 	if (CVarVRSDebugForceRate->GetInt() >= 0)
 	{
@@ -368,6 +371,9 @@ FRDGTextureRef FVariableRateShadingImageManager::GetVariableRateShadingImage(FRD
 
 void FVariableRateShadingImageManager::PrepareImageBasedVRS(FRDGBuilder& GraphBuilder, const FSceneViewFamily& ViewFamily, const FMinimalSceneTextures& SceneTextures)
 {
+	RDG_EVENT_SCOPE(GraphBuilder, "PrepareImageBasedVRS");
+	SCOPED_NAMED_EVENT(PrepareImageBasedVRS, FColor::Red);
+
 	// If no views support VRS, bail immediately
 	bool bIsAnyViewVRSCompatible = false;
 	for (const FSceneView* View : ViewFamily.Views)
@@ -565,6 +571,8 @@ FRDGTextureRef FVariableRateShadingImageManager::CombineShadingRateImages(FRDGBu
 	}
 	else
 	{
+		SCOPED_NAMED_EVENT(CombineShadingRateImages, FColor::Green);
+
 		// Create texture to hold shading rate image
 		FRDGTextureRef CombinedShadingRateTexture = GraphBuilder.CreateTexture(GetSRIDesc(), TEXT("CombinedShadingRateTexture"));
 
