Moving the IHeterogeneousVolumeInterface from FPrimitiveSceneProxy to FMeshBatchElement UserData.

Moving the interface to custom UserData allows FNiagaraRenderers the to modify parameters based active Niagara bindings and provides scaffolding for particle-driven instancing support.

Encapsulating FPrimitiveSceneProxy in IHeterogeneousVolumeInterface to simplify Heterogeneous Volume API.

Exposing FNiagaraRendererVolumes parameters and connecting to Heterogeneous Volume rendering options.

diff --git a/Engine/Plugins/FX/Niagara/Source/Niagara/Private/NiagaraComponent.cpp b/Engine/Plugins/FX/Niagara/Source/Niagara/Private/NiagaraComponent.cpp
--- a/Engine/Plugins/FX/Niagara/Source/Niagara/Private/NiagaraComponent.cpp
+++ b/Engine/Plugins/FX/Niagara/Source/Niagara/Private/NiagaraComponent.cpp
@@ -197,31 +197,6 @@ FNiagaraSceneProxy::FNiagaraSceneProxy(UNiagaraComponent* InComponent)
 		bAlwaysHasVelocity = RenderData->HasAnyMotionBlurEnabled();
 		bIsHeterogeneousVolume = RenderData->HasAnyHeterogeneousVolumesEnabled();
 
-		// Initialize Heterogeneous Volume properties
-		if (bIsHeterogeneousVolume)
-		{
-			// Temporarily dependening on user variables until a NiagaraVolumeRenderer is created
-			FNiagaraVariable VoxelResolutionVariable(FNiagaraTypeDefinition::GetIntDef(), TEXT("ResolutionMaxAxis"));
-			float ResolutionMaxAxis = InComponent->GetOverrideParameters().GetParameterValueOrDefault(VoxelResolutionVariable, 0);
-			ResolutionMaxAxis = FMath::Max(ResolutionMaxAxis, 0);
-
-			FNiagaraVariable WorldSpaceSizeVariable(FNiagaraTypeDefinition::GetVec3Def(), TEXT("WorldSpaceSize"));
-			FVector3f WorldSpaceSize = InComponent->GetOverrideParameters().GetParameterValueOrDefault(WorldSpaceSizeVariable, FVector3f(0));
-
-			float WorldSpaceSizeMaxInv = WorldSpaceSize.GetMax() > 0.0 ? 1.0 / WorldSpaceSize.GetMax() : 0.0;
-			FVector3f ResolutionFactor = WorldSpaceSize * WorldSpaceSizeMaxInv;
-
-			FVector3f VolumeResolutionV3f = ResolutionFactor * ResolutionMaxAxis;
-			FIntVector VolumeResolutionInt3 = FIntVector(
-				FMath::CeilToInt(VolumeResolutionV3f.X),
-				FMath::CeilToInt(VolumeResolutionV3f.Y),
-				FMath::CeilToInt(VolumeResolutionV3f.Z));
-
-			HeterogeneousVolumeData.VoxelResolution = VolumeResolutionInt3;
-			HeterogeneousVolumeData.LightingDownsampleFactor = 1;
-			HeterogeneousVolumeData.MinimumVoxelSize = 0.1;
-		}
-
 		SystemStatID = InComponent->GetAsset()->GetStatID(false, false);
 	}
 }
diff --git a/Engine/Plugins/FX/Niagara/Source/Niagara/Private/NiagaraMeshRendererProperties.cpp b/Engine/Plugins/FX/Niagara/Source/Niagara/Private/NiagaraMeshRendererProperties.cpp
--- a/Engine/Plugins/FX/Niagara/Source/Niagara/Private/NiagaraMeshRendererProperties.cpp
+++ b/Engine/Plugins/FX/Niagara/Source/Niagara/Private/NiagaraMeshRendererProperties.cpp
@@ -763,6 +763,20 @@ bool UNiagaraMeshRendererProperties::PopulateRequiredBindings(FNiagaraParameterS
 {
 	bool bAnyAdded = Super::PopulateRequiredBindings(InParameterStore);
 
+
+	if (bUseHeterogeneousVolumes)
+	{
+		InParameterStore.AddParameter(
+			FNiagaraVariableBase(FNiagaraTypeDefinition::GetIntDef(), FName("User.ResolutionMaxAxis")),
+			false
+		);
+		InParameterStore.AddParameter(
+			FNiagaraVariableBase(FNiagaraTypeDefinition::GetVec3Def(), FName("User.WorldSpaceSize")),
+			false
+		);
+		bAnyAdded = true;
+	}
+
 	for (const FNiagaraVariableAttributeBinding* Binding : AttributeBindings)
 	{
 		if (Binding && Binding->CanBindToHostParameterMap())
diff --git a/Engine/Plugins/FX/Niagara/Source/Niagara/Private/NiagaraRendererMeshes.cpp b/Engine/Plugins/FX/Niagara/Source/Niagara/Private/NiagaraRendererMeshes.cpp
--- a/Engine/Plugins/FX/Niagara/Source/Niagara/Private/NiagaraRendererMeshes.cpp
+++ b/Engine/Plugins/FX/Niagara/Source/Niagara/Private/NiagaraRendererMeshes.cpp
@@ -101,6 +101,16 @@ FNiagaraRendererMeshes::FNiagaraRendererMeshes(ERHIFeatureLevel::Type FeatureLev
 	bAccurateMotionVectors = Properties->NeedsPreciseMotionVectors();
 	bIsHeterogeneousVolume = Properties->UseHeterogeneousVolumes();
 
+	if (Properties->UseHeterogeneousVolumes())
+	{
+		ResolutionMaxAxisOffset = Emitter->GetRendererBoundVariables().IndexOf(
+			FNiagaraVariableBase(FNiagaraTypeDefinition::GetIntDef(), FName("User.ResolutionMaxAxis"))
+		);
+		WorldSpaceSizeOffset = Emitter->GetRendererBoundVariables().IndexOf(
+			FNiagaraVariableBase(FNiagaraTypeDefinition::GetVec3Def(), FName("User.WorldSpaceSize"))
+		);
+	}
+
 	if (Properties->bEnableCameraDistanceCulling)
 	{
 		DistanceCullRange = FVector2f(Properties->MinCameraDistance, Properties->MaxCameraDistance);
@@ -306,6 +316,15 @@ void FNiagaraRendererMeshes::PrepareParticleMeshRenderData(FParticleMeshRenderDa
 		return;
 	}
 
+	if (ResolutionMaxAxisOffset != INDEX_NONE)
+	{
+		ParticleMeshRenderData.ResolutionMaxAxis = *reinterpret_cast<const int32*>(ParticleMeshRenderData.DynamicDataMesh->ParameterDataBound.GetData() + ResolutionMaxAxisOffset);
+	}
+	if (WorldSpaceSizeOffset != INDEX_NONE)
+	{
+		ParticleMeshRenderData.WorldSpaceSize = *reinterpret_cast<const FVector3f*>(ParticleMeshRenderData.DynamicDataMesh->ParameterDataBound.GetData() + WorldSpaceSizeOffset);
+	}
+
 	// If the visibility tag comes from a parameter map, so we can evaluate it here and just early out if it doesn't match up
 	if (EmitterRendererVisTagOffset != INDEX_NONE && ParticleMeshRenderData.DynamicDataMesh->ParameterDataBound.IsValidIndex(EmitterRendererVisTagOffset))
 	{
@@ -1342,6 +1361,25 @@ void FNiagaraRendererMeshes::GetDynamicMeshElements(const TArray<const FSceneVie
 							bNeedsPrevTransform
 						);
 
+						if (SceneProxy->IsHeterogeneousVolume())
+						{
+							FHeterogeneousVolumeData* HeterogeneousVolumeData = &Collector.AllocateOneFrameResource<FHeterogeneousVolumeData>(SceneProxy);
+
+							FVector3f WorldSpaceSize = ParticleMeshRenderData.WorldSpaceSize;
+							float WorldSpaceSizeMaxInv = WorldSpaceSize.GetMax() > 0.0 ? 1.0 / WorldSpaceSize.GetMax() : 0.0;
+							FVector3f ResolutionFactor = WorldSpaceSize * WorldSpaceSizeMaxInv;
+
+							FVector3f VolumeResolutionV3f = ResolutionFactor * ParticleMeshRenderData.ResolutionMaxAxis;
+
+							FIntVector VolumeResolution = FIntVector(
+								FMath::CeilToInt(VolumeResolutionV3f.X),
+								FMath::CeilToInt(VolumeResolutionV3f.Y),
+								FMath::CeilToInt(VolumeResolutionV3f.Z));
+
+							HeterogeneousVolumeData->VoxelResolution = VolumeResolution;
+							MeshBatch.Elements[0].UserData = HeterogeneousVolumeData;
+						}
+
 						Collector.AddMesh(ViewIndex, MeshBatch);
 					}
 				}
diff --git a/Engine/Plugins/FX/Niagara/Source/Niagara/Private/NiagaraRendererVolumes.cpp b/Engine/Plugins/FX/Niagara/Source/Niagara/Private/NiagaraRendererVolumes.cpp
--- a/Engine/Plugins/FX/Niagara/Source/Niagara/Private/NiagaraRendererVolumes.cpp
+++ b/Engine/Plugins/FX/Niagara/Source/Niagara/Private/NiagaraRendererVolumes.cpp
@@ -45,6 +45,10 @@ namespace NiagaraRendererVolumesLocal
 		int32					DefaultRendererVisibilityTag = 0;
 		int32					DefaultVolumeResolutionMaxAxis = UNiagaraVolumeRendererProperties::GetDefaultVolumeResolutionMaxAxis();
 		FVector3f				DefaultVolumeWorldSpaceSize = UNiagaraVolumeRendererProperties::GetDefaultVolumeWorldSpaceSize();
+
+		int32					VolumeResolutionMaxAxis = UNiagaraVolumeRendererProperties::GetDefaultVolumeResolutionMaxAxis();
+		FVector3f				VolumeWorldSpaceSize = UNiagaraVolumeRendererProperties::GetDefaultVolumeWorldSpaceSize();
+		float					LightingDownsampleFactor = UNiagaraVolumeRendererProperties::GetDefaultLightingDownsampleFactor();
 	};
 
 	struct FHeterogeneousVolumeInstances
@@ -209,6 +213,15 @@ FNiagaraDynamicDataBase* FNiagaraRendererVolumes::GenerateDynamicData(const FNia
 		ProcessMaterialParameterBindings(Properties->MaterialParameters, Emitter, MakeArrayView(BaseMaterials_GT));
 	}
 
+	const UNiagaraVolumeRendererProperties* RendererProperties = CastChecked<const UNiagaraVolumeRendererProperties>(InProperties);
+	const int32 DefaultVolumeResolutionMaxAxis = UNiagaraVolumeRendererProperties::GetDefaultVolumeResolutionMaxAxis();
+	const FVector3f DefaultVolumeWorldSpaceSize = UNiagaraVolumeRendererProperties::GetDefaultVolumeWorldSpaceSize();
+
+	const FNiagaraParameterStore& ParameterStore = Emitter->GetRendererBoundVariables();
+	VolumeDynamicData->VolumeResolutionMaxAxis = ParameterStore.GetParameterValueOrDefault(RendererProperties->VolumeResolutionMaxAxisBinding.GetParamMapBindableVariable(), DefaultVolumeResolutionMaxAxis);
+	VolumeDynamicData->VolumeWorldSpaceSize = ParameterStore.GetParameterValueOrDefault(RendererProperties->VolumeWorldSpaceSizeBinding.GetParamMapBindableVariable(), DefaultVolumeWorldSpaceSize);
+	VolumeDynamicData->LightingDownsampleFactor = RendererProperties->LightingDownsampleFactor;
+
 	return VolumeDynamicData;
 }
 
@@ -323,7 +336,26 @@ void FNiagaraRendererVolumes::GetDynamicMeshElements(const TArray<const FSceneVi
 			BatchElement.MaxVertexIndex		= 3;
 			BatchElement.NumPrimitives		= 2;
 			BatchElement.BaseVertexIndex	= 0;
-			//-TODO:BatchElement.VertexFactoryUserData = VolumeInstances;
+			
+			FHeterogeneousVolumeData* HeterogeneousVolumeData = &Collector.AllocateOneFrameResource<FHeterogeneousVolumeData>(SceneProxy);
+
+			FVector3f WorldSpaceSize = VolumeDynamicData->VolumeWorldSpaceSize;
+			float WorldSpaceSizeMaxInv = WorldSpaceSize.GetMax() > 0.0 ? 1.0 / WorldSpaceSize.GetMax() : 0.0;
+			FVector3f ResolutionFactor = WorldSpaceSize * WorldSpaceSizeMaxInv;
+
+			FVector3f VolumeResolutionV3f = ResolutionFactor * VolumeDynamicData->VolumeResolutionMaxAxis;
+
+			FIntVector VolumeResolution = FIntVector(
+				FMath::CeilToInt(VolumeResolutionV3f.X),
+				FMath::CeilToInt(VolumeResolutionV3f.Y),
+				FMath::CeilToInt(VolumeResolutionV3f.Z));
+
+			HeterogeneousVolumeData->VoxelResolution = VolumeResolution;
+
+
+			HeterogeneousVolumeData->LightingDownsampleFactor = VolumeDynamicData->LightingDownsampleFactor;
+			BatchElement.UserData = HeterogeneousVolumeData;
+
 			Collector.AddMesh(0, Mesh);
 		}
 	}
diff --git a/Engine/Plugins/FX/Niagara/Source/Niagara/Public/NiagaraRendererMeshes.h b/Engine/Plugins/FX/Niagara/Source/Niagara/Public/NiagaraRendererMeshes.h
--- a/Engine/Plugins/FX/Niagara/Source/Niagara/Public/NiagaraRendererMeshes.h
+++ b/Engine/Plugins/FX/Niagara/Source/Niagara/Public/NiagaraRendererMeshes.h
@@ -74,6 +74,9 @@ protected:
 
 		FVector							WorldSpacePivotOffset = FVector::ZeroVector;
 		FSphere							CullingSphere = FSphere(EForceInit::ForceInit);
+
+		int32							ResolutionMaxAxis = 0;
+		FVector3f						WorldSpaceSize = FVector3f::ZeroVector;
 	};
 
 	struct FMeshData
@@ -185,6 +188,9 @@ private:
 	int32 EmitterMeshIndexOffset = INDEX_NONE;
 	uint32 MaterialParamValidMask = 0;
 
+	int32 ResolutionMaxAxisOffset = INDEX_NONE;
+	int32 WorldSpaceSizeOffset = INDEX_NONE;
+
 	int32 VFBoundOffsetsInParamStore[ENiagaraMeshVFLayout::Type::Num_Max];
 	uint32 bSetAnyBoundVars : 1;
 
diff --git a/Engine/Plugins/FX/Niagara/Source/Niagara/Public/NiagaraSceneProxy.h b/Engine/Plugins/FX/Niagara/Source/Niagara/Public/NiagaraSceneProxy.h
--- a/Engine/Plugins/FX/Niagara/Source/Niagara/Public/NiagaraSceneProxy.h
+++ b/Engine/Plugins/FX/Niagara/Source/Niagara/Public/NiagaraSceneProxy.h
@@ -90,8 +90,6 @@ private:
 
 	uint32 GetAllocatedSize() const;
 
-	virtual const IHeterogeneousVolumeInterface* GetHeterogeneousVolumeInterface() const override { return &HeterogeneousVolumeData; }
-
 private:
 	/** Custom Uniform Buffers, allows us to have renderer specific data packed inside such as pre-skinned bounds. */
 	mutable TMap<uint32, TUniformBuffer<FPrimitiveUniformShaderParameters>*> CustomUniformBuffers;
@@ -106,8 +104,5 @@ private:
 	TStatId SystemStatID;
 
 	FDynamicData DynamicData;
-
-	// Consider wrapping in a TUniquePtr
-	FHeterogeneousVolumeData HeterogeneousVolumeData;
 };
 
diff --git a/Engine/Plugins/FX/Niagara/Source/Niagara/Public/NiagaraVolumeRendererProperties.h b/Engine/Plugins/FX/Niagara/Source/Niagara/Public/NiagaraVolumeRendererProperties.h
--- a/Engine/Plugins/FX/Niagara/Source/Niagara/Public/NiagaraVolumeRendererProperties.h
+++ b/Engine/Plugins/FX/Niagara/Source/Niagara/Public/NiagaraVolumeRendererProperties.h
@@ -77,6 +77,7 @@ public:
 	static FVector3f GetDefaultVolumeScale() { return FVector3f(1.0f, 1.0f, 1.0f); }
 	static int32 GetDefaultVolumeResolutionMaxAxis() { return 0; }
 	static FVector3f GetDefaultVolumeWorldSpaceSize() { return FVector3f::ZeroVector; }
+	static float GetDefaultLightingDownsampleFactor() { return 2.0f; }
 
 	/** What material to use for the volume. */
 	UPROPERTY(EditAnywhere, Category = "Volume Rendering")
@@ -100,6 +101,9 @@ public:
 	UPROPERTY(EditAnywhere, Category = "Volume Rendering")
 	int32 RendererVisibility = 0;
 
+	UPROPERTY(EditAnywhere, Category = "Volume Rendering")
+	float LightingDownsampleFactor = 2.0f;
+
 	/** Position binding for the center of the volume. */
 	//-TODO: Unhind when binding is supported
 	//UPROPERTY(EditAnywhere, Category = "Bindings")
@@ -119,14 +123,10 @@ public:
 	UPROPERTY(EditAnywhere, Category = "Bindings")
 	FNiagaraVariableAttributeBinding RendererVisibilityTagBinding;
 
-	/** TODO */
-	//-TODO: Unhind when binding is supported
-	//UPROPERTY(EditAnywhere, Category = "Bindings")
+	UPROPERTY(EditAnywhere, Category = "Bindings")
 	FNiagaraVariableAttributeBinding VolumeResolutionMaxAxisBinding;
 
-	/** TODO */
-	//-TODO: Unhind when binding is supported
-	//UPROPERTY(EditAnywhere, Category = "Bindings")
+	UPROPERTY(EditAnywhere, Category = "Bindings")
 	FNiagaraVariableAttributeBinding VolumeWorldSpaceSizeBinding;
 
 	/** If this array has entries, we will create a MaterialInstanceDynamic per Emitter instance from Material and set the Material parameters using the Niagara simulation variables listed.*/
diff --git a/Engine/Source/Runtime/Engine/Private/Components/HeterogeneousVolumeComponent.cpp b/Engine/Source/Runtime/Engine/Private/Components/HeterogeneousVolumeComponent.cpp
--- a/Engine/Source/Runtime/Engine/Private/Components/HeterogeneousVolumeComponent.cpp
+++ b/Engine/Source/Runtime/Engine/Private/Components/HeterogeneousVolumeComponent.cpp
@@ -48,13 +48,10 @@ public:
 	uint32 GetAllocatedSize(void) const { return FPrimitiveSceneProxy::GetAllocatedSize(); }
 	//~ End FPrimitiveSceneProxy Interface.
 
-	virtual const IHeterogeneousVolumeInterface* GetHeterogeneousVolumeInterface() const override { return &HeterogeneousVolumeData; }
-
 private:
 	UMaterialInterface* MaterialInterface;
 	FLocalVertexFactory VertexFactory;
 	FStaticMeshVertexBuffers StaticMeshVertexBuffers;
-
 	FHeterogeneousVolumeData HeterogeneousVolumeData;
 };
 
@@ -66,6 +63,7 @@ FHeterogeneousVolumeSceneProxy::FHeterogeneousVolumeSceneProxy(UHeterogeneousVol
 	: FPrimitiveSceneProxy(InComponent)
 	, MaterialInterface(InComponent->GetMaterial(0))
 	, VertexFactory(GetScene().GetFeatureLevel(), "FHeterogeneousVolumeSceneProxy")
+	, HeterogeneousVolumeData(this)
 {
 	bIsHeterogeneousVolume = true;
 
@@ -187,6 +185,9 @@ void FHeterogeneousVolumeSceneProxy::GetDynamicMeshElements(
 		BatchElement.NumPrimitives = 2;
 		BatchElement.BaseVertexIndex = 0;
 
+		//FHeterogeneousVolumeData* HeterogeneousVolumeData = &Collector.AllocateOneFrameResource<FHeterogeneousVolumeData>(this);
+		BatchElement.UserData = &HeterogeneousVolumeData;
+
 		Mesh.bCanApplyViewModeOverrides = true;
 		Mesh.bUseWireframeSelectionColoring = IsSelected();
 
diff --git a/Engine/Source/Runtime/Engine/Private/HeterogeneousVolumeInterface.cpp b/Engine/Source/Runtime/Engine/Private/HeterogeneousVolumeInterface.cpp
--- a/Engine/Source/Runtime/Engine/Private/HeterogeneousVolumeInterface.cpp
+++ b/Engine/Source/Runtime/Engine/Private/HeterogeneousVolumeInterface.cpp
@@ -5,18 +5,3 @@
 =============================================================================*/
 
 #include "HeterogeneousVolumeInterface.h"
-
-IHeterogeneousVolumeInterface::~IHeterogeneousVolumeInterface()
-{
-}
-
-FHeterogeneousVolumeData::FHeterogeneousVolumeData()
-	: VoxelResolution(FIntVector::ZeroValue)
-	, MinimumVoxelSize(0.1)
-	, LightingDownsampleFactor(1.0)
-{
-}
-
-FHeterogeneousVolumeData::~FHeterogeneousVolumeData()
-{
-}
diff --git a/Engine/Source/Runtime/Engine/Public/HeterogeneousVolumeInterface.h b/Engine/Source/Runtime/Engine/Public/HeterogeneousVolumeInterface.h
--- a/Engine/Source/Runtime/Engine/Public/HeterogeneousVolumeInterface.h
+++ b/Engine/Source/Runtime/Engine/Public/HeterogeneousVolumeInterface.h
@@ -3,11 +3,23 @@
 #pragma once
 
 #include "CoreMinimal.h"
+#include "SceneManagement.h"
+#include "PrimitiveSceneProxy.h"
 
 class IHeterogeneousVolumeInterface
 {
 public:
-	ENGINE_API virtual ~IHeterogeneousVolumeInterface();
+	ENGINE_API virtual ~IHeterogeneousVolumeInterface() {}
+	virtual const FPrimitiveSceneProxy* GetPrimitiveSceneProxy() const = 0;
+	
+	// Local-Space
+	virtual const FBoxSphereBounds& GetBounds() const = 0;
+	virtual const FBoxSphereBounds& GetLocalBounds() const = 0;
+	virtual const FMatrix& GetLocalToWorld() const = 0;
+
+	// Instance-Space
+	virtual const FMatrix& GetInstanceToLocal() const = 0;
+	virtual const FMatrix GetInstanceToWorld() const = 0;
 
 	// Volume
 	virtual FIntVector GetVoxelResolution() const = 0;
@@ -17,11 +29,28 @@ public:
 	virtual float GetLightingDownsampleFactor() const = 0;
 };
 
-class FHeterogeneousVolumeData : public IHeterogeneousVolumeInterface
+class FPrimitiveSceneProxy;
+
+class FHeterogeneousVolumeData : public IHeterogeneousVolumeInterface, FOneFrameResource
 {
 public:
-	ENGINE_API FHeterogeneousVolumeData();
-	ENGINE_API virtual ~FHeterogeneousVolumeData();
+	ENGINE_API explicit FHeterogeneousVolumeData(const FPrimitiveSceneProxy* SceneProxy)
+		: PrimitiveSceneProxy(SceneProxy)
+		, InstanceToLocal(FMatrix::Identity)
+		, VoxelResolution(FIntVector::ZeroValue)
+		, MinimumVoxelSize(0.1)
+		, LightingDownsampleFactor(1.0)
+	{}
+	ENGINE_API virtual ~FHeterogeneousVolumeData() {}
+
+	virtual const FPrimitiveSceneProxy* GetPrimitiveSceneProxy() const { return PrimitiveSceneProxy; }
+
+	// Local-space
+	virtual const FBoxSphereBounds& GetBounds() const { return PrimitiveSceneProxy->GetBounds(); }
+	virtual const FBoxSphereBounds& GetLocalBounds() const { return PrimitiveSceneProxy->GetLocalBounds(); }
+	virtual const FMatrix& GetLocalToWorld() const { return PrimitiveSceneProxy->GetLocalToWorld(); }
+	virtual const FMatrix& GetInstanceToLocal() const { return InstanceToLocal; }
+	virtual const FMatrix GetInstanceToWorld() const { return InstanceToLocal * PrimitiveSceneProxy->GetLocalToWorld(); }
 
 	// Volume
 	virtual FIntVector GetVoxelResolution() const { return VoxelResolution; }
@@ -30,6 +59,8 @@ public:
 	// Lighting
 	virtual float GetLightingDownsampleFactor() const { return LightingDownsampleFactor; }
 
+	const FPrimitiveSceneProxy* PrimitiveSceneProxy;
+	FMatrix InstanceToLocal;
 	FIntVector VoxelResolution;
 	float MinimumVoxelSize;
 	float LightingDownsampleFactor;
diff --git a/Engine/Source/Runtime/Engine/Public/PrimitiveSceneProxy.h b/Engine/Source/Runtime/Engine/Public/PrimitiveSceneProxy.h
--- a/Engine/Source/Runtime/Engine/Public/PrimitiveSceneProxy.h
+++ b/Engine/Source/Runtime/Engine/Public/PrimitiveSceneProxy.h
@@ -831,11 +831,6 @@ public:
 		return bIsHeterogeneousVolume;
 	}
 
-	virtual const IHeterogeneousVolumeInterface* GetHeterogeneousVolumeInterface() const
-	{
-		return nullptr;
-	}
-
 	/**
 	 * Returns true if all meshes drawn by this proxy support GPU scene. 
 	 */
diff --git a/Engine/Source/Runtime/Renderer/Private/HeterogeneousVolumes/HeterogeneousVolumes.cpp b/Engine/Source/Runtime/Renderer/Private/HeterogeneousVolumes/HeterogeneousVolumes.cpp
--- a/Engine/Source/Runtime/Renderer/Private/HeterogeneousVolumes/HeterogeneousVolumes.cpp
+++ b/Engine/Source/Runtime/Renderer/Private/HeterogeneousVolumes/HeterogeneousVolumes.cpp
@@ -1,8 +1,8 @@
 // Copyright Epic Games, Inc. All Rights Reserved.
 
 #include "HeterogeneousVolumes.h"
-
 #include "HeterogeneousVolumeInterface.h"
+
 #include "PixelShaderUtils.h"
 #include "RayTracingDefinitions.h"
 #include "RayTracingInstance.h"
@@ -343,16 +343,6 @@ namespace HeterogeneousVolumes
 	}
 
 	// Convenience Utils
-
-	const IHeterogeneousVolumeInterface* GetInterface(const FPrimitiveSceneProxy* PrimitiveSceneProxy)
-	{
-		check(PrimitiveSceneProxy->IsHeterogeneousVolume());
-		const IHeterogeneousVolumeInterface* Interface = PrimitiveSceneProxy->GetHeterogeneousVolumeInterface();
-		check(Interface);
-
-		return Interface;
-	}
-
 	int GetVoxelCount(FIntVector VolumeResolution)
 	{
 		return VolumeResolution.X * VolumeResolution.Y * VolumeResolution.Z;
@@ -378,6 +368,9 @@ namespace HeterogeneousVolumes
 		DownsampleFactor = FMath::Max(DownsampleFactor, 1.0);
 
 		FIntVector LightingCacheResolution = GetVolumeResolution(RenderInterface) / DownsampleFactor;
+		LightingCacheResolution.X = FMath::Clamp(LightingCacheResolution.X, 1, 1024);
+		LightingCacheResolution.Y = FMath::Clamp(LightingCacheResolution.Y, 1, 1024);
+		LightingCacheResolution.Z = FMath::Clamp(LightingCacheResolution.Z, 1, 1024);
 		return LightingCacheResolution;
 	}
 }
@@ -439,70 +432,80 @@ void FDeferredShadingSceneRenderer::RenderHeterogeneousVolumes(
 					}
 
 					const FMaterialRenderProxy* MaterialRenderProxy = Mesh->MaterialRenderProxy;
-					const FPrimitiveSceneInfo* PrimitiveSceneInfo = PrimitiveSceneProxy->GetPrimitiveSceneInfo();
-					const int32 PrimitiveId = PrimitiveSceneInfo->GetIndex();
-					const FBoxSphereBounds LocalBoxSphereBounds = PrimitiveSceneProxy->GetLocalBounds();
-
-					RDG_EVENT_SCOPE(GraphBuilder, "%s", *PrimitiveSceneProxy->GetResourceName().ToString());
-
-				// Allocate transmittance volume
-				// TODO: Allow option for scalar transmittance to conserve bandwidth
-				FIntVector LightingCacheResolution = HeterogeneousVolumes::GetLightingCacheResolution(HeterogeneousVolumes::GetInterface(PrimitiveSceneProxy));
-				uint32 NumMips = FMath::Log2(float(FMath::Min(FMath::Min(LightingCacheResolution.X, LightingCacheResolution.Y), LightingCacheResolution.Z))) + 1;
-				FRDGTextureDesc LightingCacheDesc = FRDGTextureDesc::Create3D(
-					LightingCacheResolution,
-					PF_FloatR11G11B10,
-					FClearValueBinding::Black,
-					TexCreate_ShaderResource | TexCreate_UAV | TexCreate_3DTiling,
-					NumMips
-				);
-				FRDGTextureRef LightingCacheTexture = GraphBuilder.CreateTexture(LightingCacheDesc, TEXT("HeterogeneousVolumes.LightingCacheTexture"));
-				AddClearUAVPass(GraphBuilder, GraphBuilder.CreateUAV(LightingCacheTexture), FLinearColor::Black);
-
-					// Material baking executes a pre-shading pipeline
-					if (CVarHeterogeneousVolumesPreshading.GetValueOnRenderThread())
-					{
-						RenderWithPreshading(
-							GraphBuilder,
-							SceneTextures,
-							Scene,
-							ViewFamily,
-							View,
-							// Shadow Data
-							VisibleLightInfos,
-							VirtualShadowMapArray,
-							// Object data
-							PrimitiveSceneProxy,
-							MaterialRenderProxy,
-							PrimitiveId,
-							LocalBoxSphereBounds,
-							// Transmittance accleration
-							LightingCacheTexture,
-							// Output
-							HeterogeneousVolumeRadiance
-						);
-					}
-					// Otherwise execute a live-shading pipeline
-					else
+					for (int32 VolumeIndex = 0; VolumeIndex < Mesh->Elements.Num(); ++VolumeIndex)
 					{
-						RenderWithLiveShading(
-							GraphBuilder,
-							SceneTextures,
-							Scene,
-							View,
-							// Shadow Data
-							VisibleLightInfos,
-							VirtualShadowMapArray,
-							// Object Data
-							PrimitiveSceneProxy,
-							MaterialRenderProxy,
-							PrimitiveId,
-							LocalBoxSphereBounds,
-							// Transmittance accleration
-							LightingCacheTexture,
-							// Output
-							HeterogeneousVolumeRadiance
+						IHeterogeneousVolumeInterface* HeterogeneousVolume = (IHeterogeneousVolumeInterface*) Mesh->Elements[VolumeIndex].UserData;
+						//check(HeterogeneousVolume != nullptr);
+						if (HeterogeneousVolume == nullptr)
+						{
+							continue;
+						}
+
+						const FPrimitiveSceneInfo* PrimitiveSceneInfo = PrimitiveSceneProxy->GetPrimitiveSceneInfo();
+						const int32 PrimitiveId = PrimitiveSceneInfo->GetIndex();
+						const FBoxSphereBounds LocalBoxSphereBounds = HeterogeneousVolume->GetLocalBounds();
+
+						RDG_EVENT_SCOPE(GraphBuilder, "%s [%d]", *PrimitiveSceneProxy->GetResourceName().ToString(), VolumeIndex);
+
+						// Allocate transmittance volume
+						// TODO: Allow option for scalar transmittance to conserve bandwidth
+						FIntVector LightingCacheResolution = HeterogeneousVolumes::GetLightingCacheResolution(HeterogeneousVolume);
+						uint32 NumMips = FMath::Log2(float(FMath::Min(FMath::Min(LightingCacheResolution.X, LightingCacheResolution.Y), LightingCacheResolution.Z))) + 1;
+						FRDGTextureDesc LightingCacheDesc = FRDGTextureDesc::Create3D(
+							LightingCacheResolution,
+							PF_FloatR11G11B10,
+							FClearValueBinding::Black,
+							TexCreate_ShaderResource | TexCreate_UAV | TexCreate_3DTiling,
+							NumMips
 						);
+						FRDGTextureRef LightingCacheTexture = GraphBuilder.CreateTexture(LightingCacheDesc, TEXT("HeterogeneousVolumes.LightingCacheTexture"));
+						AddClearUAVPass(GraphBuilder, GraphBuilder.CreateUAV(LightingCacheTexture), FLinearColor::Black);
+
+						// Material baking executes a pre-shading pipeline
+						if (CVarHeterogeneousVolumesPreshading.GetValueOnRenderThread())
+						{
+							RenderWithPreshading(
+								GraphBuilder,
+								SceneTextures,
+								Scene,
+								ViewFamily,
+								View,
+								// Shadow Data
+								VisibleLightInfos,
+								VirtualShadowMapArray,
+								// Object data
+								HeterogeneousVolume,
+								MaterialRenderProxy,
+								PrimitiveId,
+								LocalBoxSphereBounds,
+								// Transmittance accleration
+								LightingCacheTexture,
+								// Output
+								HeterogeneousVolumeRadiance
+							);
+						}
+						// Otherwise execute a live-shading pipeline
+						else
+						{
+							RenderWithLiveShading(
+								GraphBuilder,
+								SceneTextures,
+								Scene,
+								View,
+								// Shadow Data
+								VisibleLightInfos,
+								VirtualShadowMapArray,
+								// Object Data
+								HeterogeneousVolume,
+								MaterialRenderProxy,
+								PrimitiveId,
+								LocalBoxSphereBounds,
+								// Transmittance accleration
+								LightingCacheTexture,
+								// Output
+								HeterogeneousVolumeRadiance
+							);
+						}
 					}
 				}
 
diff --git a/Engine/Source/Runtime/Renderer/Private/HeterogeneousVolumes/HeterogeneousVolumes.h b/Engine/Source/Runtime/Renderer/Private/HeterogeneousVolumes/HeterogeneousVolumes.h
--- a/Engine/Source/Runtime/Renderer/Private/HeterogeneousVolumes/HeterogeneousVolumes.h
+++ b/Engine/Source/Runtime/Renderer/Private/HeterogeneousVolumes/HeterogeneousVolumes.h
@@ -75,8 +75,6 @@ namespace HeterogeneousVolumes
 	bool EnableLinearInterpolation();
 
 	// Convenience Utils
-	const IHeterogeneousVolumeInterface* GetInterface(const FPrimitiveSceneProxy* PrimitiveSceneProxy);
-
 	int GetVoxelCount(FIntVector VolumeResolution);
 	int GetVoxelCount(const FRDGTextureDesc& TextureDesc);
 	FIntVector GetMipVolumeResolution(FIntVector VolumeResolution, uint32 MipLevel);
@@ -291,7 +289,7 @@ void RenderWithLiveShading(
 	TArray<FVisibleLightInfo, SceneRenderingAllocator>& VisibleLightInfos,
 	const FVirtualShadowMapArray& VirtualShadowMapArray,
 	// Object data
-	const FPrimitiveSceneProxy* PrimitiveSceneProxy,
+	const IHeterogeneousVolumeInterface* HeterogeneousVolumeInterface,
 	const FMaterialRenderProxy* MaterialRenderProxy,
 	const int32 PrimitiveId,
 	const FBoxSphereBounds LocalBoxSphereBounds,
@@ -312,7 +310,7 @@ void RenderWithPreshading(
 	TArray<FVisibleLightInfo, SceneRenderingAllocator>& VisibleLightInfos,
 	const FVirtualShadowMapArray& VirtualShadowMapArray,
 	// Object data
-	const FPrimitiveSceneProxy* PrimitiveSceneProxy,
+	const IHeterogeneousVolumeInterface* HeterogeneousVolumeInterface,
 	const FMaterialRenderProxy* MaterialRenderProxy,
 	const int32 PrimitiveId,
 	const FBoxSphereBounds LocalBoxSphereBounds,
@@ -342,7 +340,7 @@ void ComputeHeterogeneousVolumeBakeMaterial(
 	const FScene* Scene,
 	const FViewInfo& View,
 	// Object data
-	const FPrimitiveSceneProxy* PrimitiveSceneProxy,
+	const IHeterogeneousVolumeInterface* HeterogeneousVolumeInterface,
 	const FMaterialRenderProxy* MaterialRenderProxy,
 	const int32 PrimitiveId,
 	const FBoxSphereBounds LocalBoxSphereBounds,
@@ -382,7 +380,7 @@ void GenerateRayTracingGeometryInstance(
 	const FScene* Scene,
 	const FViewInfo& View,
 	// Object data
-	const FPrimitiveSceneProxy* PrimitiveSceneProxy,
+	const IHeterogeneousVolumeInterface* HeterogeneousVolumeInterface,
 	// Sparse voxel data
 	FRDGBufferRef NumVoxelsBuffer,
 	TRDGUniformBufferRef<FSparseVoxelUniformBufferParameters> SparseVoxelUniformBuffer,
@@ -396,8 +394,6 @@ void GenerateRayTracingScene(
 	// Scene data
 	const FScene* Scene,
 	const FViewInfo& View,
-	// Object data
-	const FPrimitiveSceneProxy* PrimitiveSceneProxy,
 	// Ray tracing data
 	TArray<FRayTracingGeometryRHIRef>& RayTracingGeometries,
 	TArray<FMatrix>& RayTracingTransforms,
@@ -421,7 +417,7 @@ void RenderLightingCacheWithPreshadingHardwareRayTracing(
 	const FVisibleLightInfo* VisibleLightInfo,
 	const FVirtualShadowMapArray& VirtualShadowMapArray,
 	// Object data
-	const FPrimitiveSceneProxy* PrimitiveSceneProxy,
+	const IHeterogeneousVolumeInterface* HeterogeneousVolumeInterface,
 	// Sparse voxel data
 	TRDGUniformBufferRef<FSparseVoxelUniformBufferParameters> SparseVoxelUniformBuffer,
 	// Ray tracing data
@@ -446,7 +442,7 @@ void RenderSingleScatteringWithPreshadingHardwareRayTracing(
 	const FVisibleLightInfo* VisibleLightInfo,
 	const FVirtualShadowMapArray& VirtualShadowMapArray,
 	// Object data
-	const FPrimitiveSceneProxy* PrimitiveSceneProxy,
+	const IHeterogeneousVolumeInterface* HeterogeneousVolumeInterface,
 	// Sparse voxel data
 	TRDGUniformBufferRef<FSparseVoxelUniformBufferParameters> SparseVoxelUniformBuffer,
 	// Ray tracing data
diff --git a/Engine/Source/Runtime/Renderer/Private/HeterogeneousVolumes/HeterogeneousVolumesHardwareRayTracing.cpp b/Engine/Source/Runtime/Renderer/Private/HeterogeneousVolumes/HeterogeneousVolumesHardwareRayTracing.cpp
--- a/Engine/Source/Runtime/Renderer/Private/HeterogeneousVolumes/HeterogeneousVolumesHardwareRayTracing.cpp
+++ b/Engine/Source/Runtime/Renderer/Private/HeterogeneousVolumes/HeterogeneousVolumesHardwareRayTracing.cpp
@@ -1,6 +1,7 @@
 // Copyright Epic Games, Inc. All Rights Reserved.
 
 #include "HeterogeneousVolumes.h"
+#include "HeterogeneousVolumeInterface.h"
 
 #include "LightRendering.h"
 #include "PixelShaderUtils.h"
@@ -113,7 +114,7 @@ void GenerateRayTracingGeometryInstance(
 	const FScene* Scene,
 	const FViewInfo& View,
 	// Object data
-	const FPrimitiveSceneProxy* PrimitiveSceneProxy,
+	const IHeterogeneousVolumeInterface* HeterogeneousVolumeInterface,
 	// Volume data
 	// Sparse voxel data
 	FRDGBufferRef NumVoxelsBuffer,
@@ -157,7 +158,7 @@ void GenerateRayTracingGeometryInstance(
 
 	FRayTracingGeometryInitializer GeometryInitializer;
 	// TODO: REMOVE STRING ALLOCATION
-	GeometryInitializer.DebugName = *PrimitiveSceneProxy->GetResourceName().ToString();// +TEXT(" (HeterogeneousVolume)");
+	//GeometryInitializer.DebugName = *PrimitiveSceneProxy->GetResourceName().ToString();// +TEXT(" (HeterogeneousVolume)");
 	//GeometryInitializer.IndexBuffer = EmptyIndexBuffer;
 	GeometryInitializer.GeometryType = RTGT_Procedural;
 	GeometryInitializer.bFastBuild = false;
@@ -171,7 +172,7 @@ void GenerateRayTracingGeometryInstance(
 	GeometryInitializer.Segments.Add(Segment);
 	GeometryInitializer.TotalPrimitiveCount = Segment.NumPrimitives;
 	RayTracingGeometries.Add(RHICreateRayTracingGeometry(GeometryInitializer));
-	RayTracingTransforms.Add(PrimitiveSceneProxy->GetLocalToWorld());
+	RayTracingTransforms.Add(HeterogeneousVolumeInterface->GetLocalToWorld());
 }
 
 BEGIN_SHADER_PARAMETER_STRUCT(FBuildBLASPassParams, )
@@ -189,8 +190,6 @@ void GenerateRayTracingScene(
 	// Scene data
 	const FScene* Scene,
 	const FViewInfo& View,
-	// Object data
-	const FPrimitiveSceneProxy* PrimitiveSceneProxy,
 	// Ray tracing data
 	TArray<FRayTracingGeometryRHIRef>& RayTracingGeometries,
 	TArray<FMatrix>& RayTracingTransforms,
@@ -580,7 +579,7 @@ void RenderLightingCacheWithPreshadingHardwareRayTracing(
 	const FVisibleLightInfo* VisibleLightInfo,
 	const FVirtualShadowMapArray& VirtualShadowMapArray,
 	// Object data
-	const FPrimitiveSceneProxy* PrimitiveSceneProxy,
+	const IHeterogeneousVolumeInterface* HeterogeneousVolumeInterface,
 	// Sparse voxel data
 	TRDGUniformBufferRef<FSparseVoxelUniformBufferParameters> SparseVoxelUniformBuffer,
 	// Ray tracing data
@@ -589,8 +588,6 @@ void RenderLightingCacheWithPreshadingHardwareRayTracing(
 	FRDGTextureRef& LightingCacheTexture
 )
 {
-	const IHeterogeneousVolumeInterface* Interface = HeterogeneousVolumes::GetInterface(PrimitiveSceneProxy);
-
 	// Note must be done in the same scope as we add the pass otherwise the UB lifetime will not be guaranteed
 	FDeferredLightUniformStruct DeferredLightUniform = GetDeferredLightParameters(View, *LightSceneInfo);
 	TUniformBufferRef<FDeferredLightUniformStruct> DeferredLightUB = CreateUniformBufferImmediate(DeferredLightUniform, UniformBuffer_SingleDraw);
@@ -615,7 +612,7 @@ void RenderLightingCacheWithPreshadingHardwareRayTracing(
 		PassParameters->SparseVoxelUniformBuffer = SparseVoxelUniformBuffer;
 
 		// Transmittance volume
-		PassParameters->LightingCache.LightingCacheResolution = HeterogeneousVolumes::GetLightingCacheResolution(Interface);
+		PassParameters->LightingCache.LightingCacheResolution = HeterogeneousVolumes::GetLightingCacheResolution(HeterogeneousVolumeInterface);
 		//PassParameters->LightingCache.LightingCacheTexture = GraphBuilder.CreateSRV(LightingCacheTexture);
 
 		// Ray data
@@ -668,7 +665,7 @@ void RenderLightingCacheWithPreshadingHardwareRayTracing(
 	FRenderLightingCacheWithPreshadingRGS::FPermutationDomain PermutationVector;
 	PermutationVector.Set<FRenderLightingCacheWithPreshadingRGS::FLightingCacheMode>(HeterogeneousVolumes::GetLightingCacheMode() - 1);
 	TShaderRef<FRenderLightingCacheWithPreshadingRGS> RayGenerationShader = View.ShaderMap->GetShader<FRenderLightingCacheWithPreshadingRGS>(PermutationVector);
-	FIntVector VolumeResolution = HeterogeneousVolumes::GetLightingCacheResolution(Interface);
+	FIntVector VolumeResolution = HeterogeneousVolumes::GetLightingCacheResolution(HeterogeneousVolumeInterface);
 	FIntPoint DispatchResolution = FIntPoint(VolumeResolution.X, VolumeResolution.Y * VolumeResolution.Z);
 
 	GraphBuilder.AddPass(
@@ -723,7 +720,7 @@ void RenderSingleScatteringWithPreshadingHardwareRayTracing(
 	const FVisibleLightInfo* VisibleLightInfo,
 	const FVirtualShadowMapArray& VirtualShadowMapArray,
 	// Object data
-	const FPrimitiveSceneProxy* PrimitiveSceneProxy,
+	const IHeterogeneousVolumeInterface* HeterogeneousVolumeInterface,
 	// Sparse voxel data
 	TRDGUniformBufferRef<FSparseVoxelUniformBufferParameters> SparseVoxelUniformBuffer,
 	// Ray tracing data
@@ -734,8 +731,6 @@ void RenderSingleScatteringWithPreshadingHardwareRayTracing(
 	FRDGTextureRef& HeterogeneousVolumeTexture
 )
 {
-	const IHeterogeneousVolumeInterface* Interface = HeterogeneousVolumes::GetInterface(PrimitiveSceneProxy);
-
 	// Note must be done in the same scope as we add the pass otherwise the UB lifetime will not be guaranteed
 	FDeferredLightUniformStruct DeferredLightUniform;
 	if (bApplyDirectLighting && (LightSceneInfo != nullptr))
@@ -798,7 +793,7 @@ void RenderSingleScatteringWithPreshadingHardwareRayTracing(
 		// Transmittance volume
 		if ((HeterogeneousVolumes::UseLightingCacheForTransmittance() && bApplyShadowTransmittance) || HeterogeneousVolumes::UseLightingCacheForInscattering())
 		{
-			PassParameters->LightingCache.LightingCacheResolution = HeterogeneousVolumes::GetLightingCacheResolution(Interface);
+			PassParameters->LightingCache.LightingCacheResolution = HeterogeneousVolumes::GetLightingCacheResolution(HeterogeneousVolumeInterface);
 			PassParameters->LightingCache.LightingCacheTexture = LightingCacheTexture;
 		}
 
diff --git a/Engine/Source/Runtime/Renderer/Private/HeterogeneousVolumes/HeterogeneousVolumesLiveShadingPipeline.cpp b/Engine/Source/Runtime/Renderer/Private/HeterogeneousVolumes/HeterogeneousVolumesLiveShadingPipeline.cpp
--- a/Engine/Source/Runtime/Renderer/Private/HeterogeneousVolumes/HeterogeneousVolumesLiveShadingPipeline.cpp
+++ b/Engine/Source/Runtime/Renderer/Private/HeterogeneousVolumes/HeterogeneousVolumesLiveShadingPipeline.cpp
@@ -1,6 +1,7 @@
 // Copyright Epic Games, Inc. All Rights Reserved.
 
 #include "HeterogeneousVolumes.h"
+#include "HeterogeneousVolumeInterface.h"
 
 #include "LightRendering.h"
 #include "LocalVertexFactory.h"
@@ -388,7 +389,7 @@ void RenderLightingCacheWithLiveShading(
 	const FVisibleLightInfo* VisibleLightInfo,
 	const FVirtualShadowMapArray& VirtualShadowMapArray,
 	// Object data
-	const FPrimitiveSceneProxy* PrimitiveSceneProxy,
+	const IHeterogeneousVolumeInterface* HeterogeneousVolumeInterface,
 	const FMaterialRenderProxy* DefaultMaterialRenderProxy,
 	const int32 PrimitiveId,
 	const FBoxSphereBounds LocalBoxSphereBounds,
@@ -396,8 +397,6 @@ void RenderLightingCacheWithLiveShading(
 	FRDGTextureRef LightingCacheTexture
 )
 {
-	const IHeterogeneousVolumeInterface* Interface = HeterogeneousVolumes::GetInterface(PrimitiveSceneProxy);
-
 	const FMaterialRenderProxy* MaterialRenderProxy = nullptr;
 	const FMaterial& Material = DefaultMaterialRenderProxy->GetMaterialWithFallback(View.GetFeatureLevel(), MaterialRenderProxy);
 	MaterialRenderProxy = MaterialRenderProxy ? MaterialRenderProxy : DefaultMaterialRenderProxy;
@@ -425,7 +424,7 @@ void RenderLightingCacheWithLiveShading(
 		PassParameters->VolumetricScatteringIntensity = LightSceneInfo->Proxy->GetVolumetricScatteringIntensity();
 
 		// Object data
-		FMatrix44f LocalToWorld = FMatrix44f(PrimitiveSceneProxy->GetLocalToWorld());
+		FMatrix44f LocalToWorld = FMatrix44f(HeterogeneousVolumeInterface->GetLocalToWorld());
 		PassParameters->LocalToWorld = LocalToWorld;
 		PassParameters->WorldToLocal = LocalToWorld.Inverse();
 		PassParameters->LocalBoundsOrigin = FVector3f(LocalBoxSphereBounds.Origin);
@@ -433,7 +432,7 @@ void RenderLightingCacheWithLiveShading(
 		PassParameters->PrimitiveId = PrimitiveId;
 
 		// Transmittance volume
-		PassParameters->LightingCache.LightingCacheResolution = HeterogeneousVolumes::GetLightingCacheResolution(Interface);
+		PassParameters->LightingCache.LightingCacheResolution = HeterogeneousVolumes::GetLightingCacheResolution(HeterogeneousVolumeInterface);
 		//PassParameters->LightingCache.LightingCacheTexture = GraphBuilder.CreateSRV(LightingCacheTexture);
 		PassParameters->LightingCache.LightingCacheTexture = FRDGSystemTextures::Get(GraphBuilder).VolumetricBlack;
 
@@ -486,7 +485,7 @@ void RenderLightingCacheWithLiveShading(
 	}
 #endif // WANTS_DRAW_MESH_EVENTS
 
-	FIntVector GroupCount = HeterogeneousVolumes::GetLightingCacheResolution(Interface);
+	FIntVector GroupCount = HeterogeneousVolumes::GetLightingCacheResolution(HeterogeneousVolumeInterface);
 	GroupCount.X = FMath::DivideAndRoundUp(GroupCount.X, FRenderLightingCacheWithLiveShadingCS::GetThreadGroupSize3D());
 	GroupCount.Y = FMath::DivideAndRoundUp(GroupCount.Y, FRenderLightingCacheWithLiveShadingCS::GetThreadGroupSize3D());
 	GroupCount.Z = FMath::DivideAndRoundUp(GroupCount.Z, FRenderLightingCacheWithLiveShadingCS::GetThreadGroupSize3D());
@@ -516,7 +515,7 @@ void RenderSingleScatteringWithLiveShading(
 	const FVisibleLightInfo* VisibleLightInfo,
 	const FVirtualShadowMapArray& VirtualShadowMapArray,
 	// Object data
-	const FPrimitiveSceneProxy* PrimitiveSceneProxy,
+	const IHeterogeneousVolumeInterface* HeterogeneousVolumeInterface,
 	const FMaterialRenderProxy* DefaultMaterialRenderProxy,
 	const int32 PrimitiveId,
 	const FBoxSphereBounds LocalBoxSphereBounds,
@@ -526,8 +525,6 @@ void RenderSingleScatteringWithLiveShading(
 	FRDGTextureRef& HeterogeneousVolumeTexture
 )
 {
-	const IHeterogeneousVolumeInterface* Interface = HeterogeneousVolumes::GetInterface(PrimitiveSceneProxy);
-
 	const FMaterialRenderProxy* MaterialRenderProxy = nullptr;
 	const FMaterial& Material = DefaultMaterialRenderProxy->GetMaterialWithFallback(View.GetFeatureLevel(), MaterialRenderProxy);
 	MaterialRenderProxy = MaterialRenderProxy ? MaterialRenderProxy : DefaultMaterialRenderProxy;
@@ -572,7 +569,7 @@ void RenderSingleScatteringWithLiveShading(
 		PassParameters->ShadowStepSize = HeterogeneousVolumes::GetShadowStepSize();
 
 		// Object data
-		FMatrix44f LocalToWorld = FMatrix44f(PrimitiveSceneProxy->GetLocalToWorld());
+		FMatrix44f LocalToWorld = FMatrix44f(HeterogeneousVolumeInterface->GetLocalToWorld());
 		PassParameters->LocalToWorld = LocalToWorld;
 		PassParameters->WorldToLocal = LocalToWorld.Inverse();
 		PassParameters->LocalBoundsOrigin = FVector3f(LocalBoxSphereBounds.Origin);
@@ -615,7 +612,7 @@ void RenderSingleScatteringWithLiveShading(
 		// Volume data
 		if ((HeterogeneousVolumes::UseLightingCacheForTransmittance() && bApplyShadowTransmittance) || HeterogeneousVolumes::UseLightingCacheForInscattering())
 		{
-			PassParameters->LightingCache.LightingCacheResolution = HeterogeneousVolumes::GetLightingCacheResolution(Interface);
+			PassParameters->LightingCache.LightingCacheResolution = HeterogeneousVolumes::GetLightingCacheResolution(HeterogeneousVolumeInterface);
 			PassParameters->LightingCache.LightingCacheTexture = LightingCacheTexture;
 		}
 		else
@@ -664,7 +661,7 @@ void RenderWithTransmittanceVolumePipeline(
 	TArray<FVisibleLightInfo, SceneRenderingAllocator>& VisibleLightInfos,
 	const FVirtualShadowMapArray& VirtualShadowMapArray,
 	// Object data
-	const FPrimitiveSceneProxy* PrimitiveSceneProxy,
+	const IHeterogeneousVolumeInterface* HeterogeneousVolumeInterface,
 	const FMaterialRenderProxy* MaterialRenderProxy,
 	const int32 PrimitiveId,
 	const FBoxSphereBounds LocalBoxSphereBounds,
@@ -678,7 +675,7 @@ void RenderWithTransmittanceVolumePipeline(
 	TArray<FLightSceneInfoCompact, TInlineAllocator<64>> LightSceneInfoCompact;
 	for (auto LightIt = Scene->Lights.CreateConstIterator(); LightIt; ++LightIt)
 	{
-		if (LightIt->AffectsPrimitive(PrimitiveSceneProxy->GetBounds(), PrimitiveSceneProxy))
+		if (LightIt->AffectsPrimitive(HeterogeneousVolumeInterface->GetBounds(), HeterogeneousVolumeInterface->GetPrimitiveSceneProxy()))
 		{
 			LightSceneInfoCompact.Add(*LightIt);
 		}
@@ -727,7 +724,7 @@ void RenderWithTransmittanceVolumePipeline(
 				VisibleLightInfo,
 				VirtualShadowMapArray,
 				// Object data
-				PrimitiveSceneProxy,
+				HeterogeneousVolumeInterface,
 				MaterialRenderProxy,
 				PrimitiveId,
 				LocalBoxSphereBounds,
@@ -752,7 +749,7 @@ void RenderWithTransmittanceVolumePipeline(
 			VisibleLightInfo,
 			VirtualShadowMapArray,
 			// Object data
-			PrimitiveSceneProxy,
+			HeterogeneousVolumeInterface,
 			MaterialRenderProxy,
 			PrimitiveId,
 			LocalBoxSphereBounds,
@@ -773,7 +770,7 @@ void RenderWithInscatteringVolumePipeline(
 	TArray<FVisibleLightInfo, SceneRenderingAllocator>& VisibleLightInfos,
 	const FVirtualShadowMapArray& VirtualShadowMapArray,
 	// Object data
-	const FPrimitiveSceneProxy* PrimitiveSceneProxy,
+	const IHeterogeneousVolumeInterface* HeterogeneousVolumeInterface,
 	const FMaterialRenderProxy* MaterialRenderProxy,
 	const int32 PrimitiveId,
 	const FBoxSphereBounds LocalBoxSphereBounds,
@@ -787,7 +784,7 @@ void RenderWithInscatteringVolumePipeline(
 	TArray<FLightSceneInfoCompact, TInlineAllocator<64>> LightSceneInfoCompact;
 	for (auto LightIt = Scene->Lights.CreateConstIterator(); LightIt; ++LightIt)
 	{
-		if (LightIt->AffectsPrimitive(PrimitiveSceneProxy->GetBounds(), PrimitiveSceneProxy))
+		if (LightIt->AffectsPrimitive(HeterogeneousVolumeInterface->GetBounds(), HeterogeneousVolumeInterface->GetPrimitiveSceneProxy()))
 		{
 			LightSceneInfoCompact.Add(*LightIt);
 		}
@@ -834,7 +831,7 @@ void RenderWithInscatteringVolumePipeline(
 			VisibleLightInfo,
 			VirtualShadowMapArray,
 			// Object data
-			PrimitiveSceneProxy,
+			HeterogeneousVolumeInterface,
 			MaterialRenderProxy,
 			PrimitiveId,
 			LocalBoxSphereBounds,
@@ -869,7 +866,7 @@ void RenderWithInscatteringVolumePipeline(
 			VisibleLightInfo,
 			VirtualShadowMapArray,
 			// Object data
-			PrimitiveSceneProxy,
+			HeterogeneousVolumeInterface,
 			MaterialRenderProxy,
 			PrimitiveId,
 			LocalBoxSphereBounds,
@@ -890,7 +887,7 @@ void RenderWithLiveShading(
 	TArray<FVisibleLightInfo, SceneRenderingAllocator>& VisibleLightInfos,
 	const FVirtualShadowMapArray& VirtualShadowMapArray,
 	// Object data
-	const FPrimitiveSceneProxy* PrimitiveSceneProxy,
+	const IHeterogeneousVolumeInterface* HeterogeneousVolumeInterface,
 	const FMaterialRenderProxy* MaterialRenderProxy,
 	const int32 PrimitiveId,
 	const FBoxSphereBounds LocalBoxSphereBounds,
@@ -911,7 +908,7 @@ void RenderWithLiveShading(
 			VisibleLightInfos,
 			VirtualShadowMapArray,
 			// Object data
-			PrimitiveSceneProxy,
+			HeterogeneousVolumeInterface,
 			MaterialRenderProxy,
 			PrimitiveId,
 			LocalBoxSphereBounds,
@@ -932,7 +929,7 @@ void RenderWithLiveShading(
 			VisibleLightInfos,
 			VirtualShadowMapArray,
 			// Object data
-			PrimitiveSceneProxy,
+			HeterogeneousVolumeInterface,
 			MaterialRenderProxy,
 			PrimitiveId,
 			LocalBoxSphereBounds,
diff --git a/Engine/Source/Runtime/Renderer/Private/HeterogeneousVolumes/HeterogeneousVolumesMaterialBakingPipeline.cpp b/Engine/Source/Runtime/Renderer/Private/HeterogeneousVolumes/HeterogeneousVolumesMaterialBakingPipeline.cpp
--- a/Engine/Source/Runtime/Renderer/Private/HeterogeneousVolumes/HeterogeneousVolumesMaterialBakingPipeline.cpp
+++ b/Engine/Source/Runtime/Renderer/Private/HeterogeneousVolumes/HeterogeneousVolumesMaterialBakingPipeline.cpp
@@ -1,6 +1,7 @@
 // Copyright Epic Games, Inc. All Rights Reserved.
 
 #include "HeterogeneousVolumes.h"
+#include "HeterogeneousVolumeInterface.h"
 
 #include "LocalVertexFactory.h"
 #include "PixelShaderUtils.h"
@@ -103,7 +104,7 @@ void ComputeHeterogeneousVolumeBakeMaterial(
 	const FScene* Scene,
 	const FViewInfo& View,
 	// Object data
-	const FPrimitiveSceneProxy* PrimitiveSceneProxy,
+	const IHeterogeneousVolumeInterface* HeterogeneousVolumeInterface,
 	const FMaterialRenderProxy* MaterialRenderProxy,
 	const int32 PrimitiveId,
 	const FBoxSphereBounds LocalBoxSphereBounds,
@@ -136,7 +137,7 @@ void ComputeHeterogeneousVolumeBakeMaterial(
 		PassParameters->Scene = View.GetSceneUniforms().GetBuffer(GraphBuilder);
 
 		// Object data
-		FMatrix44f LocalToWorld = FMatrix44f(PrimitiveSceneProxy->GetLocalToWorld());
+		FMatrix44f LocalToWorld = FMatrix44f(HeterogeneousVolumeInterface->GetLocalToWorld());
 		PassParameters->LocalToWorld = LocalToWorld;
 		PassParameters->WorldToLocal = LocalToWorld.Inverse();
 		PassParameters->LocalBoundsOrigin = FVector3f(LocalBoxSphereBounds.Origin);
diff --git a/Engine/Source/Runtime/Renderer/Private/HeterogeneousVolumes/HeterogeneousVolumesPreshadingPipeline.cpp b/Engine/Source/Runtime/Renderer/Private/HeterogeneousVolumes/HeterogeneousVolumesPreshadingPipeline.cpp
--- a/Engine/Source/Runtime/Renderer/Private/HeterogeneousVolumes/HeterogeneousVolumesPreshadingPipeline.cpp
+++ b/Engine/Source/Runtime/Renderer/Private/HeterogeneousVolumes/HeterogeneousVolumesPreshadingPipeline.cpp
@@ -1,6 +1,7 @@
 // Copyright Epic Games, Inc. All Rights Reserved.
 
 #include "HeterogeneousVolumes.h"
+#include "HeterogeneousVolumeInterface.h"
 
 #include "LightRendering.h"
 #include "PixelShaderUtils.h"
@@ -107,7 +108,7 @@ void GenerateRayMarchingTiles(
 	const FViewInfo& View,
 	const FSceneTextures& SceneTextures,
 	// Object data
-	const FPrimitiveSceneProxy* PrimitiveSceneProxy,
+	const IHeterogeneousVolumeInterface* HeterogeneousVolumeInterface,
 	// Sparse voxel data
 	FRDGBufferRef NumVoxelsBuffer,
 	const TRDGUniformBufferRef<FSparseVoxelUniformBufferParameters>& SparseVoxelUniformBuffer,
@@ -394,7 +395,7 @@ void RenderLightingCacheWithPreshadingCompute(
 	const FVisibleLightInfo* VisibleLightInfo,
 	const FVirtualShadowMapArray& VirtualShadowMapArray,
 	// Object data
-	const FPrimitiveSceneProxy* PrimitiveSceneProxy,
+	const IHeterogeneousVolumeInterface* HeterogeneousVolumeInterface,
 	// Sparse voxel data
 	FRDGBufferRef NumVoxelsBuffer,
 	const TRDGUniformBufferRef<FSparseVoxelUniformBufferParameters>& SparseVoxelUniformBuffer,
@@ -406,8 +407,6 @@ void RenderLightingCacheWithPreshadingCompute(
 	FRDGTextureRef& LightingCacheTexture
 )
 {
-	const IHeterogeneousVolumeInterface* Interface = HeterogeneousVolumes::GetInterface(PrimitiveSceneProxy);
-
 	// Note must be done in the same scope as we add the pass otherwise the UB lifetime will not be guaranteed
 	FDeferredLightUniformStruct DeferredLightUniform = GetDeferredLightParameters(View, *LightSceneInfo);
 	TUniformBufferRef<FDeferredLightUniformStruct> DeferredLightUB = CreateUniformBufferImmediate(DeferredLightUniform, UniformBuffer_SingleDraw);
@@ -431,7 +430,7 @@ void RenderLightingCacheWithPreshadingCompute(
 		PassParameters->SparseVoxelUniformBuffer = SparseVoxelUniformBuffer;
 
 		// Transmittance volume
-		PassParameters->LightingCache.LightingCacheResolution = HeterogeneousVolumes::GetLightingCacheResolution(Interface);
+		PassParameters->LightingCache.LightingCacheResolution = HeterogeneousVolumes::GetLightingCacheResolution(HeterogeneousVolumeInterface);
 		PassParameters->LightingCache.LightingCacheTexture = LightingCacheTexture;
 
 		// Ray data
@@ -485,7 +484,7 @@ void RenderLightingCacheWithPreshadingCompute(
 	PermutationVector.Set<FRenderLightingCacheWithPreshadingCS::FLightingCacheMode>(HeterogeneousVolumes::GetLightingCacheMode() - 1);
 	TShaderRef<FRenderLightingCacheWithPreshadingCS> ComputeShader = View.ShaderMap->GetShader<FRenderLightingCacheWithPreshadingCS>(PermutationVector);
 
-	FIntVector GroupCount = HeterogeneousVolumes::GetLightingCacheResolution(Interface);
+	FIntVector GroupCount = HeterogeneousVolumes::GetLightingCacheResolution(HeterogeneousVolumeInterface);
 	GroupCount.X = FMath::DivideAndRoundUp(GroupCount.X, FRenderLightingCacheWithPreshadingCS::GetThreadGroupSize3D());
 	GroupCount.Y = FMath::DivideAndRoundUp(GroupCount.Y, FRenderLightingCacheWithPreshadingCS::GetThreadGroupSize3D());
 	GroupCount.Z = FMath::DivideAndRoundUp(GroupCount.Z, FRenderLightingCacheWithPreshadingCS::GetThreadGroupSize3D());
@@ -514,7 +513,7 @@ void RenderSingleScatteringWithPreshadingCompute(
 	const FVisibleLightInfo* VisibleLightInfo,
 	const FVirtualShadowMapArray& VirtualShadowMapArray,
 	// Object data
-	const FPrimitiveSceneProxy* PrimitiveSceneProxy,
+	const IHeterogeneousVolumeInterface* HeterogeneousVolumeInterface,
 	// Sparse voxel data
 	FRDGBufferRef NumVoxelsBuffer,
 	const TRDGUniformBufferRef<FSparseVoxelUniformBufferParameters>& SparseVoxelUniformBuffer,
@@ -527,8 +526,6 @@ void RenderSingleScatteringWithPreshadingCompute(
 	FRDGTextureRef& HeterogeneousVolumeTexture
 )
 {
-	const IHeterogeneousVolumeInterface* Interface = HeterogeneousVolumes::GetInterface(PrimitiveSceneProxy);
-
 	FRDGBufferRef VoxelOutputBuffer = GraphBuilder.CreateBuffer(
 		FRDGBufferDesc::CreateStructuredDesc(sizeof(FVoxelDataPacked), HeterogeneousVolumes::GetVoxelCount(SparseVoxelUniformBuffer->GetParameters()->VolumeResolution)),
 		TEXT("HeterogeneousVolumes.VoxelOutputBuffer")
@@ -595,7 +592,7 @@ void RenderSingleScatteringWithPreshadingCompute(
 		// Transmittance volume
 		if ((HeterogeneousVolumes::UseLightingCacheForTransmittance() && bApplyShadowTransmittance) || HeterogeneousVolumes::UseLightingCacheForInscattering())
 		{
-			PassParameters->LightingCache.LightingCacheResolution = HeterogeneousVolumes::GetLightingCacheResolution(Interface);
+			PassParameters->LightingCache.LightingCacheResolution = HeterogeneousVolumes::GetLightingCacheResolution(HeterogeneousVolumeInterface);
 			PassParameters->LightingCache.LightingCacheTexture = LightingCacheTexture;
 		}
 
@@ -652,7 +649,7 @@ void RenderWithInscatteringVolumePipelineWithPreshadingCompute(
 	TArray<FVisibleLightInfo, SceneRenderingAllocator>& VisibleLightInfos,
 	const FVirtualShadowMapArray& VirtualShadowMapArray,
 	// Object data
-	const FPrimitiveSceneProxy* PrimitiveSceneProxy,
+	const IHeterogeneousVolumeInterface* HeterogeneousVolumeInterface,
 	const FMaterialRenderProxy* MaterialRenderProxy,
 	// Sparse voxel data
 	FRDGBufferRef NumVoxelsBuffer,
@@ -672,7 +669,7 @@ void RenderWithInscatteringVolumePipelineWithPreshadingCompute(
 	TArray<FLightSceneInfoCompact, TInlineAllocator<64>> LightSceneInfoCompact;
 	for (auto LightIt = Scene->Lights.CreateConstIterator(); LightIt; ++LightIt)
 	{
-		if (LightIt->AffectsPrimitive(PrimitiveSceneProxy->GetBounds(), PrimitiveSceneProxy))
+		if (LightIt->AffectsPrimitive(HeterogeneousVolumeInterface->GetBounds(), HeterogeneousVolumeInterface->GetPrimitiveSceneProxy()))
 		{
 			LightSceneInfoCompact.Add(*LightIt);
 		}
@@ -719,7 +716,7 @@ void RenderWithInscatteringVolumePipelineWithPreshadingCompute(
 			VisibleLightInfo,
 			VirtualShadowMapArray,
 			// Object data
-			PrimitiveSceneProxy,
+			HeterogeneousVolumeInterface,
 			// Volume data
 			NumVoxelsBuffer,
 			// Sparse voxel data
@@ -759,7 +756,7 @@ void RenderWithInscatteringVolumePipelineWithPreshadingCompute(
 			VisibleLightInfo,
 			VirtualShadowMapArray,
 			// Object
-			PrimitiveSceneProxy,
+			HeterogeneousVolumeInterface,
 			// Volume data
 			NumVoxelsBuffer,
 			// Sparse voxel data
@@ -786,7 +783,7 @@ void RenderWithTransmittanceVolumePipelineWithPreshadingCompute(
 	TArray<FVisibleLightInfo, SceneRenderingAllocator>& VisibleLightInfos,
 	const FVirtualShadowMapArray& VirtualShadowMapArray,
 	// Object data
-	const FPrimitiveSceneProxy* PrimitiveSceneProxy,
+	const IHeterogeneousVolumeInterface* HeterogeneousVolumeInterface,
 	const FMaterialRenderProxy* MaterialRenderProxy,
 	// Sparse voxel data
 	FRDGBufferRef NumVoxelsBuffer,
@@ -806,7 +803,7 @@ void RenderWithTransmittanceVolumePipelineWithPreshadingCompute(
 	TArray<FLightSceneInfoCompact, TInlineAllocator<64>> LightSceneInfoCompact;
 	for (auto LightIt = Scene->Lights.CreateConstIterator(); LightIt; ++LightIt)
 	{
-		if (LightIt->AffectsPrimitive(PrimitiveSceneProxy->GetBounds(), PrimitiveSceneProxy))
+		if (LightIt->AffectsPrimitive(HeterogeneousVolumeInterface->GetBounds(), HeterogeneousVolumeInterface->GetPrimitiveSceneProxy()))
 		{
 			LightSceneInfoCompact.Add(*LightIt);
 		}
@@ -855,7 +852,7 @@ void RenderWithTransmittanceVolumePipelineWithPreshadingCompute(
 				VisibleLightInfo,
 				VirtualShadowMapArray,
 				// Object data
-				PrimitiveSceneProxy,
+				HeterogeneousVolumeInterface,
 				// Volume data
 				NumVoxelsBuffer,
 				// Sparse voxel data
@@ -885,7 +882,7 @@ void RenderWithTransmittanceVolumePipelineWithPreshadingCompute(
 			VisibleLightInfo,
 			VirtualShadowMapArray,
 			// Object
-			PrimitiveSceneProxy,
+			HeterogeneousVolumeInterface,
 			// Volume data
 			NumVoxelsBuffer,
 			// Sparse voxel data
@@ -912,7 +909,7 @@ void RenderWithPreshadingCompute(
 	TArray<FVisibleLightInfo, SceneRenderingAllocator>& VisibleLightInfos,
 	const FVirtualShadowMapArray& VirtualShadowMapArray,
 	// Object data
-	const FPrimitiveSceneProxy* PrimitiveSceneProxy,
+	const IHeterogeneousVolumeInterface* HeterogeneousVolumeInterface,
 	const FMaterialRenderProxy* MaterialRenderProxy,
 	// Sparse voxel data
 	FRDGBufferRef NumVoxelsBuffer,
@@ -936,7 +933,7 @@ void RenderWithPreshadingCompute(
 			View,
 			SceneTextures,
 			// Object
-			PrimitiveSceneProxy,
+			HeterogeneousVolumeInterface,
 			// Volume data
 			NumVoxelsBuffer,
 			// Sparse voxel data
@@ -961,7 +958,7 @@ void RenderWithPreshadingCompute(
 			VisibleLightInfos,
 			VirtualShadowMapArray,
 			// Object data
-			PrimitiveSceneProxy,
+			HeterogeneousVolumeInterface,
 			MaterialRenderProxy,
 			// Sparse voxel data
 			NumVoxelsBuffer,
@@ -988,7 +985,7 @@ void RenderWithPreshadingCompute(
 			VisibleLightInfos,
 			VirtualShadowMapArray,
 			// Object data
-			PrimitiveSceneProxy,
+			HeterogeneousVolumeInterface,
 			MaterialRenderProxy,
 			// Sparse voxel data
 			NumVoxelsBuffer,
@@ -1015,7 +1012,7 @@ void RenderWithInscatteringVolumePipelineWithPreshadingHardwareRayTracing(
 	TArray<FVisibleLightInfo, SceneRenderingAllocator>& VisibleLightInfos,
 	const FVirtualShadowMapArray& VirtualShadowMapArray,
 	// Object data
-	const FPrimitiveSceneProxy* PrimitiveSceneProxy,
+	const IHeterogeneousVolumeInterface* HeterogeneousVolumeInterface,
 	const FMaterialRenderProxy* MaterialRenderProxy,
 	// Sparse voxel data
 	FRDGBufferRef NumVoxelsBuffer,
@@ -1033,7 +1030,7 @@ void RenderWithInscatteringVolumePipelineWithPreshadingHardwareRayTracing(
 	TArray<FLightSceneInfoCompact, TInlineAllocator<64>> LightSceneInfoCompact;
 	for (auto LightIt = Scene->Lights.CreateConstIterator(); LightIt; ++LightIt)
 	{
-		if (LightIt->AffectsPrimitive(PrimitiveSceneProxy->GetBounds(), PrimitiveSceneProxy))
+		if (LightIt->AffectsPrimitive(HeterogeneousVolumeInterface->GetBounds(), HeterogeneousVolumeInterface->GetPrimitiveSceneProxy()))
 		{
 			LightSceneInfoCompact.Add(*LightIt);
 		}
@@ -1080,7 +1077,7 @@ void RenderWithInscatteringVolumePipelineWithPreshadingHardwareRayTracing(
 			VisibleLightInfo,
 			VirtualShadowMapArray,
 			// Object data
-			PrimitiveSceneProxy,
+			HeterogeneousVolumeInterface,
 			// Sparse voxel
 			SparseVoxelUniformBuffer,
 			// Ray tracing data
@@ -1116,7 +1113,7 @@ void RenderWithInscatteringVolumePipelineWithPreshadingHardwareRayTracing(
 			VisibleLightInfo,
 			VirtualShadowMapArray,
 			// Object data
-			PrimitiveSceneProxy,
+			HeterogeneousVolumeInterface,
 			// Sparse voxel
 			SparseVoxelUniformBuffer,
 			// Ray tracing data
@@ -1141,7 +1138,7 @@ void RenderWithTransmittanceVolumePipelineWithPreshadingHardwareRayTracing(
 	TArray<FVisibleLightInfo, SceneRenderingAllocator>& VisibleLightInfos,
 	const FVirtualShadowMapArray& VirtualShadowMapArray,
 	// Object data
-	const FPrimitiveSceneProxy* PrimitiveSceneProxy,
+	const IHeterogeneousVolumeInterface* HeterogeneousVolumeInterface,
 	const FMaterialRenderProxy* MaterialRenderProxy,
 	// Sparse voxel data
 	FRDGBufferRef NumVoxelsBuffer,
@@ -1159,7 +1156,7 @@ void RenderWithTransmittanceVolumePipelineWithPreshadingHardwareRayTracing(
 	TArray<FLightSceneInfoCompact, TInlineAllocator<64>> LightSceneInfoCompact;
 	for (auto LightIt = Scene->Lights.CreateConstIterator(); LightIt; ++LightIt)
 	{
-		if (LightIt->AffectsPrimitive(PrimitiveSceneProxy->GetBounds(), PrimitiveSceneProxy))
+		if (LightIt->AffectsPrimitive(HeterogeneousVolumeInterface->GetBounds(), HeterogeneousVolumeInterface->GetPrimitiveSceneProxy()))
 		{
 			LightSceneInfoCompact.Add(*LightIt);
 		}
@@ -1208,7 +1205,7 @@ void RenderWithTransmittanceVolumePipelineWithPreshadingHardwareRayTracing(
 				VisibleLightInfo,
 				VirtualShadowMapArray,
 				// Object data
-				PrimitiveSceneProxy,
+				HeterogeneousVolumeInterface,
 				// Sparse voxel
 				SparseVoxelUniformBuffer,
 				// Ray tracing data
@@ -1234,7 +1231,7 @@ void RenderWithTransmittanceVolumePipelineWithPreshadingHardwareRayTracing(
 			VisibleLightInfo,
 			VirtualShadowMapArray,
 			// Object data
-			PrimitiveSceneProxy,
+			HeterogeneousVolumeInterface,
 			// Sparse voxel
 			SparseVoxelUniformBuffer,
 			// Ray tracing data
@@ -1259,7 +1256,7 @@ void RenderWithPreshadingHardwareRayTracing(
 	TArray<FVisibleLightInfo, SceneRenderingAllocator>& VisibleLightInfos,
 	const FVirtualShadowMapArray& VirtualShadowMapArray,
 	// Object data
-	const FPrimitiveSceneProxy* PrimitiveSceneProxy,
+	const IHeterogeneousVolumeInterface* HeterogeneousVolumeInterface,
 	const FMaterialRenderProxy* MaterialRenderProxy,
 	// Sparse voxel data
 	FRDGBufferRef NumVoxelsBuffer,
@@ -1285,7 +1282,7 @@ void RenderWithPreshadingHardwareRayTracing(
 			Scene,
 			View,
 			// Object
-			PrimitiveSceneProxy,
+			HeterogeneousVolumeInterface,
 			// Sparse voxel
 			NumVoxelsBuffer,
 			SparseVoxelUniformBuffer,
@@ -1299,8 +1296,6 @@ void RenderWithPreshadingHardwareRayTracing(
 			// Scene
 			Scene,
 			View,
-			// Object
-			PrimitiveSceneProxy,
 			// Ray tracing data
 			RayTracingGeometries,
 			RayTracingTransforms,
@@ -1321,7 +1316,7 @@ void RenderWithPreshadingHardwareRayTracing(
 			VisibleLightInfos,
 			VirtualShadowMapArray,
 			// Object data
-			PrimitiveSceneProxy,
+			HeterogeneousVolumeInterface,
 			MaterialRenderProxy,
 			NumVoxelsBuffer,
 			SparseVoxelUniformBuffer,
@@ -1343,7 +1338,7 @@ void RenderWithPreshadingHardwareRayTracing(
 			VisibleLightInfos,
 			VirtualShadowMapArray,
 			// Object data
-			PrimitiveSceneProxy,
+			HeterogeneousVolumeInterface,
 			MaterialRenderProxy,
 			NumVoxelsBuffer,
 			SparseVoxelUniformBuffer,
@@ -1487,7 +1482,7 @@ void RenderWithPreshading(
 	TArray<FVisibleLightInfo, SceneRenderingAllocator>& VisibleLightInfos,
 	const FVirtualShadowMapArray& VirtualShadowMapArray,
 	// Object data
-	const FPrimitiveSceneProxy* PrimitiveSceneProxy,
+	const IHeterogeneousVolumeInterface* HeterogeneousVolumeInterface,
 	const FMaterialRenderProxy* MaterialRenderProxy,
 	const int32 PrimitiveId,
 	const FBoxSphereBounds LocalBoxSphereBounds,
@@ -1497,10 +1492,8 @@ void RenderWithPreshading(
 	FRDGTextureRef& HeterogeneousVolumeRadiance
 )
 {
-	const IHeterogeneousVolumeInterface* Interface = HeterogeneousVolumes::GetInterface(PrimitiveSceneProxy);
-
 	// Determine baking voxel resolution
-	FIntVector VolumeResolution = HeterogeneousVolumes::GetVolumeResolution(Interface);
+	FIntVector VolumeResolution = HeterogeneousVolumes::GetVolumeResolution(HeterogeneousVolumeInterface);
 
 	// Create baked material grids
 	uint32 NumMips = FMath::Log2(float(FMath::Min(FMath::Min(VolumeResolution.X, VolumeResolution.Y), VolumeResolution.Z))) + 1;
@@ -1529,7 +1522,7 @@ void RenderWithPreshading(
 				Scene,
 				View,
 				// Object data
-				PrimitiveSceneProxy,
+				HeterogeneousVolumeInterface,
 				MaterialRenderProxy,
 				PrimitiveId,
 				LocalBoxSphereBounds,
@@ -1592,7 +1585,7 @@ void RenderWithPreshading(
 	FSparseVoxelUniformBufferParameters* SparseVoxelUniformBufferParameters = GraphBuilder.AllocParameters<FSparseVoxelUniformBufferParameters>();
 	{
 		// Object data
-		FMatrix44f LocalToWorld = FMatrix44f(PrimitiveSceneProxy->GetLocalToWorld());
+		FMatrix44f LocalToWorld = FMatrix44f(HeterogeneousVolumeInterface->GetLocalToWorld());
 		SparseVoxelUniformBufferParameters->LocalToWorld = LocalToWorld;
 		SparseVoxelUniformBufferParameters->WorldToLocal = LocalToWorld.Inverse();
 		SparseVoxelUniformBufferParameters->LocalBoundsOrigin = FVector3f(LocalBoxSphereBounds.Origin);
@@ -1632,7 +1625,7 @@ void RenderWithPreshading(
 			VisibleLightInfos,
 			VirtualShadowMapArray,
 			// Object data
-			PrimitiveSceneProxy,
+			HeterogeneousVolumeInterface,
 			MaterialRenderProxy,
 			// Sparse voxel data
 			NumVoxelsBuffer,
@@ -1657,7 +1650,7 @@ void RenderWithPreshading(
 			VisibleLightInfos,
 			VirtualShadowMapArray,
 			// Object data
-			PrimitiveSceneProxy,
+			HeterogeneousVolumeInterface,
 			MaterialRenderProxy,
 			// Sparse voxel data
 			NumVoxelsBuffer,
diff --git a/Engine/Source/Runtime/Renderer/Private/HeterogeneousVolumes/HeterogeneousVolumesVoxelGridPipeline.cpp b/Engine/Source/Runtime/Renderer/Private/HeterogeneousVolumes/HeterogeneousVolumesVoxelGridPipeline.cpp
--- a/Engine/Source/Runtime/Renderer/Private/HeterogeneousVolumes/HeterogeneousVolumesVoxelGridPipeline.cpp
+++ b/Engine/Source/Runtime/Renderer/Private/HeterogeneousVolumes/HeterogeneousVolumesVoxelGridPipeline.cpp
@@ -1197,23 +1197,31 @@ void CalcViewBoundsAndMinimumVoxelSize(
 			continue;
 		}
 
-		// Only incorporate the primitive if it intersects with the canera bounding sphere where radius=MaxTraceDistance
-		const FBoxSphereBounds& PrimitiveBounds = PrimitiveSceneProxy->GetBounds();
-		if (View.ViewFrustum.IntersectBox(PrimitiveBounds.Origin, PrimitiveBounds.BoxExtent))
+		for (int32 VolumeIndex = 0; VolumeIndex < Mesh->Elements.Num(); ++VolumeIndex)
 		{
-			TopLevelGridBounds = Union(TopLevelGridBounds, PrimitiveBounds);
+			IHeterogeneousVolumeInterface* HeterogeneousVolume = (IHeterogeneousVolumeInterface*)Mesh->Elements[VolumeIndex].UserData;
+			//check(HeterogeneousVolume != nullptr);
+			if (HeterogeneousVolume == nullptr)
+			{
+				continue;
+			}
 
-			if (View.ViewFrustum.IntersectBox(TopLevelGridBounds.Origin, TopLevelGridBounds.BoxExtent))
+			// Only incorporate the primitive if it intersects with the canera bounding sphere where radius=MaxTraceDistance
+			const FBoxSphereBounds& PrimitiveBounds = HeterogeneousVolume->GetBounds();
+			if (View.ViewFrustum.IntersectBox(PrimitiveBounds.Origin, PrimitiveBounds.BoxExtent))
 			{
-				const IHeterogeneousVolumeInterface* HeterogeneousVolumeInterface = HeterogeneousVolumes::GetInterface(PrimitiveSceneProxy);
+				TopLevelGridBounds = Union(TopLevelGridBounds, PrimitiveBounds);
 
-				// Bandlimit minimum voxel size request with projected voxel size, based on shading rate
-				FVector VoxelCenter = PrimitiveBounds.Origin;
-				float Distance = FMath::Max(FVector(PrimitiveBounds.Origin - WorldCameraOrigin).Length() - TopLevelGridBounds.BoxExtent.Length(), 0.0);
-				float VoxelWidth = Distance * PixelWidth * HeterogeneousVolumes::GetShadingRateForFrustumGrid();
+				if (View.ViewFrustum.IntersectBox(TopLevelGridBounds.Origin, TopLevelGridBounds.BoxExtent))
+				{
+					// Bandlimit minimum voxel size request with projected voxel size, based on shading rate
+					FVector VoxelCenter = PrimitiveBounds.Origin;
+					float Distance = FMath::Max(FVector(PrimitiveBounds.Origin - WorldCameraOrigin).Length() - TopLevelGridBounds.BoxExtent.Length(), 0.0);
+					float VoxelWidth = Distance * PixelWidth * HeterogeneousVolumes::GetShadingRateForFrustumGrid();
 
-				float PerVolumeMinimumVoxelSize = FMath::Max(VoxelWidth, HeterogeneousVolumeInterface->GetMinimumVoxelSize());
-				MinimumVoxelSize = FMath::Min(PerVolumeMinimumVoxelSize, MinimumVoxelSize);
+					float PerVolumeMinimumVoxelSize = FMath::Max(VoxelWidth, HeterogeneousVolume->GetMinimumVoxelSize());
+					MinimumVoxelSize = FMath::Min(PerVolumeMinimumVoxelSize, MinimumVoxelSize);
+				}
 			}
 		}
 	}
@@ -1254,27 +1262,34 @@ void CalcGlobalBoundsAndMinimumVoxelSize(
 				continue;
 			}
 
-			// Only incorporate the primitive if it intersects with the canera bounding sphere where radius=MaxTraceDistance
-			const FBoxSphereBounds& PrimitiveBounds = PrimitiveSceneProxy->GetBounds();
-			if (View.ViewFrustum.IntersectBox(PrimitiveBounds.Origin, PrimitiveBounds.BoxExtent))
+			for (int32 VolumeIndex = 0; VolumeIndex < Mesh->Elements.Num(); ++VolumeIndex)
 			{
-				AggregatePrimitiveBounds = Union(AggregatePrimitiveBounds, PrimitiveBounds);
-
-				if (View.ViewFrustum.IntersectBox(AggregatePrimitiveBounds.Origin, AggregatePrimitiveBounds.BoxExtent))
+				IHeterogeneousVolumeInterface* HeterogeneousVolume = (IHeterogeneousVolumeInterface*)Mesh->Elements[VolumeIndex].UserData;
+				//check(HeterogeneousVolume != nullptr);
+				if (HeterogeneousVolume == nullptr)
+				{
+					continue;
+				}
+				// Only incorporate the primitive if it intersects with the canera bounding sphere where radius=MaxTraceDistance
+				const FBoxSphereBounds& PrimitiveBounds = HeterogeneousVolume->GetBounds();
+				if (View.ViewFrustum.IntersectBox(PrimitiveBounds.Origin, PrimitiveBounds.BoxExtent))
 				{
-					const IHeterogeneousVolumeInterface* HeterogeneousVolumeInterface = HeterogeneousVolumes::GetInterface(PrimitiveSceneProxy);
+					AggregatePrimitiveBounds = Union(AggregatePrimitiveBounds, PrimitiveBounds);
 
-					// Bandlimit minimum voxel size request with projected voxel size, based on shading rate
-					FVector VoxelCenter = PrimitiveBounds.Origin;
-					float Distance = FMath::Max(FVector(PrimitiveBounds.Origin - WorldCameraOrigin).Length() - AggregatePrimitiveBounds.BoxExtent.Length(), 0.0);
-					float VoxelWidth = Distance * PixelWidth * HeterogeneousVolumes::GetShadingRateForOrthoGrid();
+					if (View.ViewFrustum.IntersectBox(AggregatePrimitiveBounds.Origin, AggregatePrimitiveBounds.BoxExtent))
+					{
+						// Bandlimit minimum voxel size request with projected voxel size, based on shading rate
+						FVector VoxelCenter = PrimitiveBounds.Origin;
+						float Distance = FMath::Max(FVector(PrimitiveBounds.Origin - WorldCameraOrigin).Length() - AggregatePrimitiveBounds.BoxExtent.Length(), 0.0);
+						float VoxelWidth = Distance * PixelWidth * HeterogeneousVolumes::GetShadingRateForOrthoGrid();
 
-					float PerVolumeMinimumVoxelSize = FMath::Max(VoxelWidth, HeterogeneousVolumeInterface->GetMinimumVoxelSize());
-					ViewMinimumVoxelSize = FMath::Min(PerVolumeMinimumVoxelSize, ViewMinimumVoxelSize);
+						float PerVolumeMinimumVoxelSize = FMath::Max(VoxelWidth, HeterogeneousVolume->GetMinimumVoxelSize());
+						ViewMinimumVoxelSize = FMath::Min(PerVolumeMinimumVoxelSize, ViewMinimumVoxelSize);
+					}
 				}
+				// TODO: Out-of-frustum minimum voxel size per-primitive?
+				// else if (FBoxSphereBounds::BoxesIntersect(WorldCameraBounds, PrimitiveBounds))
 			}
-			// TODO: Out-of-frustum minimum voxel size per-primitive?
-			// else if (FBoxSphereBounds::BoxesIntersect(WorldCameraBounds, PrimitiveBounds))
 		}
 
 		// Clamp per-view minimum voxel-size to in-frustum maximum
@@ -1475,43 +1490,53 @@ void CalculateVoxelSize(
 
 		for (int32 MeshBatchIndex = 0; MeshBatchIndex < View.HeterogeneousVolumesMeshBatches.Num(); ++MeshBatchIndex)
 		{
+			const FMeshBatch* Mesh = View.HeterogeneousVolumesMeshBatches[MeshBatchIndex].Mesh;
 			const FPrimitiveSceneProxy* PrimitiveSceneProxy = View.HeterogeneousVolumesMeshBatches[MeshBatchIndex].Proxy;
-			const FBoxSphereBounds& PrimitiveBounds = PrimitiveSceneProxy->GetBounds();
-			const IHeterogeneousVolumeInterface* HeterogeneousVolumeInterface = HeterogeneousVolumes::GetInterface(PrimitiveSceneProxy);
 
-			FTopLevelGridCalculateVoxelSize::FParameters* PassParameters = GraphBuilder.AllocParameters<FTopLevelGridCalculateVoxelSize::FParameters>();
+			for (int32 VolumeIndex = 0; VolumeIndex < Mesh->Elements.Num(); ++VolumeIndex)
 			{
-				PassParameters->View = View.ViewUniformBuffer;
+				IHeterogeneousVolumeInterface* HeterogeneousVolume = (IHeterogeneousVolumeInterface*)Mesh->Elements[VolumeIndex].UserData;
+				//check(HeterogeneousVolume != nullptr);
+				if (HeterogeneousVolume == nullptr)
+				{
+					continue;
+				}
 
-				PassParameters->TopLevelGridResolution = TopLevelGridResolution;
-				PassParameters->TopLevelGridWorldBoundsMin = FVector3f(TopLevelGridBounds.Origin - TopLevelGridBounds.BoxExtent);
-				PassParameters->TopLevelGridWorldBoundsMax = FVector3f(TopLevelGridBounds.Origin + TopLevelGridBounds.BoxExtent);
+				const FBoxSphereBounds& PrimitiveBounds = HeterogeneousVolume->GetBounds();
+				FTopLevelGridCalculateVoxelSize::FParameters* PassParameters = GraphBuilder.AllocParameters<FTopLevelGridCalculateVoxelSize::FParameters>();
+				{
+					PassParameters->View = View.ViewUniformBuffer;
 
-				PassParameters->PrimitiveWorldBoundsMin = FVector3f(PrimitiveBounds.Origin - PrimitiveBounds.BoxExtent);
-				PassParameters->PrimitiveWorldBoundsMax = FVector3f(PrimitiveBounds.Origin + PrimitiveBounds.BoxExtent);
+					PassParameters->TopLevelGridResolution = TopLevelGridResolution;
+					PassParameters->TopLevelGridWorldBoundsMin = FVector3f(TopLevelGridBounds.Origin - TopLevelGridBounds.BoxExtent);
+					PassParameters->TopLevelGridWorldBoundsMax = FVector3f(TopLevelGridBounds.Origin + TopLevelGridBounds.BoxExtent);
 
-				PassParameters->ShadingRate = HeterogeneousVolumes::GetShadingRateForOrthoGrid();
-				PassParameters->MinVoxelSizeInFrustum = FMath::Max(HeterogeneousVolumeInterface->GetMinimumVoxelSize(), HeterogeneousVolumes::GetMinimumVoxelSizeInFrustum());
-				PassParameters->MinVoxelSizeOutOfFrustum = HeterogeneousVolumes::GetMinimumVoxelSizeOutsideFrustum();
+					PassParameters->PrimitiveWorldBoundsMin = FVector3f(PrimitiveBounds.Origin - PrimitiveBounds.BoxExtent);
+					PassParameters->PrimitiveWorldBoundsMax = FVector3f(PrimitiveBounds.Origin + PrimitiveBounds.BoxExtent);
 
-				PassParameters->RWTopLevelGridBuffer = GraphBuilder.CreateUAV(TopLevelGridBuffer);
-			}
+					PassParameters->ShadingRate = HeterogeneousVolumes::GetShadingRateForOrthoGrid();
+						PassParameters->MinVoxelSizeInFrustum = FMath::Max(HeterogeneousVolume->GetMinimumVoxelSize(), HeterogeneousVolumes::GetMinimumVoxelSizeInFrustum());
+					PassParameters->MinVoxelSizeOutOfFrustum = HeterogeneousVolumes::GetMinimumVoxelSizeOutsideFrustum();
 
-			FIntVector GroupCount;
-			GroupCount.X = FMath::DivideAndRoundUp(TopLevelGridResolution.X, FTopLevelGridCalculateVoxelSize::GetThreadGroupSize3D());
-			GroupCount.Y = FMath::DivideAndRoundUp(TopLevelGridResolution.Y, FTopLevelGridCalculateVoxelSize::GetThreadGroupSize3D());
-			GroupCount.Z = FMath::DivideAndRoundUp(TopLevelGridResolution.Z, FTopLevelGridCalculateVoxelSize::GetThreadGroupSize3D());
+					PassParameters->RWTopLevelGridBuffer = GraphBuilder.CreateUAV(TopLevelGridBuffer);
+				}
 
-			FTopLevelGridCalculateVoxelSize::FPermutationDomain PermutationVector;
-			TShaderRef<FTopLevelGridCalculateVoxelSize> ComputeShader = View.ShaderMap->GetShader<FTopLevelGridCalculateVoxelSize>(PermutationVector);
-			FComputeShaderUtils::AddPass(
-				GraphBuilder,
-				RDG_EVENT_NAME("TopLevelGridCalculateVoxelSize"),
-				ERDGPassFlags::Compute | ERDGPassFlags::NeverCull,
-				ComputeShader,
-				PassParameters,
-				GroupCount
-			);
+				FIntVector GroupCount;
+				GroupCount.X = FMath::DivideAndRoundUp(TopLevelGridResolution.X, FTopLevelGridCalculateVoxelSize::GetThreadGroupSize3D());
+				GroupCount.Y = FMath::DivideAndRoundUp(TopLevelGridResolution.Y, FTopLevelGridCalculateVoxelSize::GetThreadGroupSize3D());
+				GroupCount.Z = FMath::DivideAndRoundUp(TopLevelGridResolution.Z, FTopLevelGridCalculateVoxelSize::GetThreadGroupSize3D());
+
+				FTopLevelGridCalculateVoxelSize::FPermutationDomain PermutationVector;
+				TShaderRef<FTopLevelGridCalculateVoxelSize> ComputeShader = View.ShaderMap->GetShader<FTopLevelGridCalculateVoxelSize>(PermutationVector);
+				FComputeShaderUtils::AddPass(
+					GraphBuilder,
+					RDG_EVENT_NAME("TopLevelGridCalculateVoxelSize"),
+					ERDGPassFlags::Compute | ERDGPassFlags::NeverCull,
+					ComputeShader,
+					PassParameters,
+					GroupCount
+				);
+			}
 		}
 	}
 }
@@ -1694,43 +1719,54 @@ void MarkTopLevelGridVoxelsForFrustumGrid(
 
 	for (int32 MeshBatchIndex = 0; MeshBatchIndex < View.HeterogeneousVolumesMeshBatches.Num(); ++MeshBatchIndex)
 	{
+		const FMeshBatch* Mesh = View.HeterogeneousVolumesMeshBatches[MeshBatchIndex].Mesh;
 		const FPrimitiveSceneProxy* PrimitiveSceneProxy = View.HeterogeneousVolumesMeshBatches[MeshBatchIndex].Proxy;
-		const FBoxSphereBounds& PrimitiveBounds = PrimitiveSceneProxy->GetBounds();
-		const IHeterogeneousVolumeInterface* HeterogeneousVolumeInterface = HeterogeneousVolumes::GetInterface(PrimitiveSceneProxy);
 
-		FMarkTopLevelGridVoxelsForFrustumGrid::FParameters* PassParameters = GraphBuilder.AllocParameters<FMarkTopLevelGridVoxelsForFrustumGrid::FParameters>();
+		for (int32 VolumeIndex = 0; VolumeIndex < Mesh->Elements.Num(); ++VolumeIndex)
 		{
-			PassParameters->View = View.ViewUniformBuffer;
+			IHeterogeneousVolumeInterface* HeterogeneousVolume = (IHeterogeneousVolumeInterface*)Mesh->Elements[VolumeIndex].UserData;
+			//check(HeterogeneousVolume != nullptr);
+			if (HeterogeneousVolume == nullptr)
+			{
+				continue;
+			}
 
-			PassParameters->PrimitiveWorldBoundsMin = FVector3f(PrimitiveBounds.Origin - PrimitiveBounds.BoxExtent);
-			PassParameters->PrimitiveWorldBoundsMax = FVector3f(PrimitiveBounds.Origin + PrimitiveBounds.BoxExtent);
+			const FBoxSphereBounds& PrimitiveBounds = HeterogeneousVolume->GetBounds();
 
-			PassParameters->ViewToWorld = FMatrix44f(ViewToWorld);
-			PassParameters->TanHalfFOV = HeterogeneousVolumes::CalcTanHalfFOV(View.FOV);
-			PassParameters->NearPlaneDepth = NearPlaneDistance;
-			PassParameters->FarPlaneDepth = FarPlaneDistance;
+			FMarkTopLevelGridVoxelsForFrustumGrid::FParameters* PassParameters = GraphBuilder.AllocParameters<FMarkTopLevelGridVoxelsForFrustumGrid::FParameters>();
+			{
+				PassParameters->View = View.ViewUniformBuffer;
 
-			PassParameters->VoxelDimensions = TopLevelGridResolution;
-			PassParameters->TopLevelGridResolution = TopLevelGridResolution;
+				PassParameters->PrimitiveWorldBoundsMin = FVector3f(PrimitiveBounds.Origin - PrimitiveBounds.BoxExtent);
+				PassParameters->PrimitiveWorldBoundsMax = FVector3f(PrimitiveBounds.Origin + PrimitiveBounds.BoxExtent);
 
-			PassParameters->RWTopLevelGridBuffer = GraphBuilder.CreateUAV(TopLevelGridBuffer);
-		}
+				PassParameters->ViewToWorld = FMatrix44f(ViewToWorld);
+				PassParameters->TanHalfFOV = HeterogeneousVolumes::CalcTanHalfFOV(View.FOV);
+				PassParameters->NearPlaneDepth = NearPlaneDistance;
+				PassParameters->FarPlaneDepth = FarPlaneDistance;
 
-		FIntVector GroupCount;
-		GroupCount.X = FMath::DivideAndRoundUp(TopLevelGridResolution.X, FMarkTopLevelGridVoxelsForFrustumGrid::GetThreadGroupSize3D());
-		GroupCount.Y = FMath::DivideAndRoundUp(TopLevelGridResolution.Y, FMarkTopLevelGridVoxelsForFrustumGrid::GetThreadGroupSize3D());
-		GroupCount.Z = FMath::DivideAndRoundUp(TopLevelGridResolution.Z, FMarkTopLevelGridVoxelsForFrustumGrid::GetThreadGroupSize3D());
+				PassParameters->VoxelDimensions = TopLevelGridResolution;
+				PassParameters->TopLevelGridResolution = TopLevelGridResolution;
 
-		FMarkTopLevelGridVoxelsForFrustumGrid::FPermutationDomain PermutationVector;
-		TShaderRef<FMarkTopLevelGridVoxelsForFrustumGrid> ComputeShader = View.ShaderMap->GetShader<FMarkTopLevelGridVoxelsForFrustumGrid>(PermutationVector);
-		FComputeShaderUtils::AddPass(
-			GraphBuilder,
-			RDG_EVENT_NAME("MarkTopLevelGridVoxelsForFrustumGrid"),
-			ERDGPassFlags::Compute | ERDGPassFlags::NeverCull,
-			ComputeShader,
-			PassParameters,
-			GroupCount
-		);
+				PassParameters->RWTopLevelGridBuffer = GraphBuilder.CreateUAV(TopLevelGridBuffer);
+			}
+
+			FIntVector GroupCount;
+			GroupCount.X = FMath::DivideAndRoundUp(TopLevelGridResolution.X, FMarkTopLevelGridVoxelsForFrustumGrid::GetThreadGroupSize3D());
+			GroupCount.Y = FMath::DivideAndRoundUp(TopLevelGridResolution.Y, FMarkTopLevelGridVoxelsForFrustumGrid::GetThreadGroupSize3D());
+			GroupCount.Z = FMath::DivideAndRoundUp(TopLevelGridResolution.Z, FMarkTopLevelGridVoxelsForFrustumGrid::GetThreadGroupSize3D());
+
+			FMarkTopLevelGridVoxelsForFrustumGrid::FPermutationDomain PermutationVector;
+			TShaderRef<FMarkTopLevelGridVoxelsForFrustumGrid> ComputeShader = View.ShaderMap->GetShader<FMarkTopLevelGridVoxelsForFrustumGrid>(PermutationVector);
+			FComputeShaderUtils::AddPass(
+				GraphBuilder,
+				RDG_EVENT_NAME("MarkTopLevelGridVoxelsForFrustumGrid"),
+				ERDGPassFlags::Compute | ERDGPassFlags::NeverCull,
+				ComputeShader,
+				PassParameters,
+				GroupCount
+			);
+		}
 	}
 }
 
