Added const qualifier to Fscene parameter in CullObjectsToView and BuildTileObjectLists.

diff --git a/Engine/Source/Runtime/Renderer/Private/DistanceFieldAmbientOcclusion.cpp b/Engine/Source/Runtime/Renderer/Private/DistanceFieldAmbientOcclusion.cpp
--- a/Engine/Source/Runtime/Renderer/Private/DistanceFieldAmbientOcclusion.cpp
+++ b/Engine/Source/Runtime/Renderer/Private/DistanceFieldAmbientOcclusion.cpp
@@ -844,7 +844,7 @@ void FDeferredShadingSceneRenderer::RenderDistanceFieldLighting(
 			ObjectIndirectArguments,
 			CulledObjectBufferParameters);
 
-		CullObjectsToView(GraphBuilder, Scene, View, Parameters, CulledObjectBufferParameters);
+		CullObjectsToView(GraphBuilder, *Scene, View, Parameters, CulledObjectBufferParameters);
 	}
 
 	ComputeDistanceFieldNormal(GraphBuilder, Views, SceneTextures.UniformBuffer, DistanceFieldNormal, Parameters);
@@ -853,7 +853,7 @@ void FDeferredShadingSceneRenderer::RenderDistanceFieldLighting(
 	if (UseAOObjectDistanceField())
 	{
 		//@todo - support multiple views - should pass one TileIntersectionParameters per view
-		BuildTileObjectLists(GraphBuilder, Scene, Views, ObjectIndirectArguments, CulledObjectBufferParameters, TileIntersectionParameters, DistanceFieldNormal, Parameters);
+		BuildTileObjectLists(GraphBuilder, *Scene, Views, ObjectIndirectArguments, CulledObjectBufferParameters, TileIntersectionParameters, DistanceFieldNormal, Parameters);
 	}
 
 	FRDGTextureRef BentNormalOutput = nullptr;
diff --git a/Engine/Source/Runtime/Renderer/Private/DistanceFieldAmbientOcclusion.h b/Engine/Source/Runtime/Renderer/Private/DistanceFieldAmbientOcclusion.h
--- a/Engine/Source/Runtime/Renderer/Private/DistanceFieldAmbientOcclusion.h
+++ b/Engine/Source/Runtime/Renderer/Private/DistanceFieldAmbientOcclusion.h
@@ -104,9 +104,9 @@ extern void TrackGPUProgress(FRHICommandListImmediate& RHICmdList, uint32 DebugI
 extern bool ShouldRenderDeferredDynamicSkyLight(const FScene* Scene, const FSceneViewFamily& ViewFamily);
 extern bool ShouldDoReflectionEnvironment(const FScene* Scene, const FSceneViewFamily& ViewFamily);
 
-extern void CullObjectsToView(FRDGBuilder& GraphBuilder, FScene* Scene, const FViewInfo& View, const FDistanceFieldAOParameters& Parameters, FDistanceFieldCulledObjectBufferParameters& CulledObjectBuffers);
+extern void CullObjectsToView(FRDGBuilder& GraphBuilder, const FScene& Scene, const FViewInfo& View, const FDistanceFieldAOParameters& Parameters, FDistanceFieldCulledObjectBufferParameters& CulledObjectBuffers);
 extern void BuildTileObjectLists(FRDGBuilder& GraphBuilder,
-	FScene* Scene,
+	const FScene& Scene,
 	TArray<FViewInfo>& Views,
 	FRDGBufferRef ObjectIndirectArguments,
 	const FDistanceFieldCulledObjectBufferParameters& CulledObjectBufferParameters,
diff --git a/Engine/Source/Runtime/Renderer/Private/DistanceFieldObjectCulling.cpp b/Engine/Source/Runtime/Renderer/Private/DistanceFieldObjectCulling.cpp
--- a/Engine/Source/Runtime/Renderer/Private/DistanceFieldObjectCulling.cpp
+++ b/Engine/Source/Runtime/Renderer/Private/DistanceFieldObjectCulling.cpp
@@ -69,17 +69,17 @@ public:
 
 IMPLEMENT_GLOBAL_SHADER(FCullObjectsForViewCS, "/Engine/Private/DistanceFieldObjectCulling.usf", "CullObjectsForViewCS", SF_Compute);
 
-void CullObjectsToView(FRDGBuilder& GraphBuilder, FScene* Scene, const FViewInfo& View, const FDistanceFieldAOParameters& Parameters, FDistanceFieldCulledObjectBufferParameters& CulledObjectBufferParameters)
+void CullObjectsToView(FRDGBuilder& GraphBuilder, const FScene& Scene, const FViewInfo& View, const FDistanceFieldAOParameters& Parameters, FDistanceFieldCulledObjectBufferParameters& CulledObjectBufferParameters)
 {
 	AddClearUAVPass(GraphBuilder, CulledObjectBufferParameters.RWObjectIndirectArguments, 0);
 
 	{
-		const int32 NumObjectsInBuffer = Scene->DistanceFieldSceneData.NumObjectsInBuffer;
+		const int32 NumObjectsInBuffer = Scene.DistanceFieldSceneData.NumObjectsInBuffer;
 
 		auto* PassParameters = GraphBuilder.AllocParameters<FCullObjectsForViewCS::FParameters>();
 
 		PassParameters->CulledObjectBufferParameters = CulledObjectBufferParameters;
-		PassParameters->ObjectBufferParameters = DistanceField::SetupObjectBufferParameters(GraphBuilder, Scene->DistanceFieldSceneData);
+		PassParameters->ObjectBufferParameters = DistanceField::SetupObjectBufferParameters(GraphBuilder, Scene.DistanceFieldSceneData);
 
 		PassParameters->View = View.ViewUniformBuffer;
 		PassParameters->NumConvexHullPlanes = View.ViewFrustum.Planes.Num();
@@ -344,7 +344,7 @@ FIntPoint GetTileListGroupSizeForView(const FViewInfo& View)
 
 void BuildTileObjectLists(
 	FRDGBuilder& GraphBuilder,
-	FScene* Scene,
+	const FScene& Scene,
 	TArray<FViewInfo>& Views,
 	FRDGBufferRef ObjectIndirectArguments,
 	const FDistanceFieldCulledObjectBufferParameters& CulledObjectBufferParameters,
@@ -385,7 +385,7 @@ void BuildTileObjectLists(
 		AddClearUAVPass(GraphBuilder, TileIntersectionParameters.RWNumCulledTilesArray, 0);
 
 		// Rasterize object bounding shapes and intersect with screen tiles to compute how many tiles intersect each object
-		ScatterTilesToObjects(GraphBuilder, true, View, Scene->DistanceFieldSceneData, TileListGroupSize, Parameters, ObjectIndirectArguments, CulledObjectBufferParameters, TileIntersectionParameters, SceneTexturesUniformBuffer);
+		ScatterTilesToObjects(GraphBuilder, true, View, Scene.DistanceFieldSceneData, TileListGroupSize, Parameters, ObjectIndirectArguments, CulledObjectBufferParameters, TileIntersectionParameters, SceneTexturesUniformBuffer);
 
 		// Start at 0 threadgroups
 		AddClearUAVPass(GraphBuilder, TileIntersectionParameters.RWObjectTilesIndirectArguments, 0);
@@ -395,11 +395,11 @@ void BuildTileObjectLists(
 			PassParameters->View = View.ViewUniformBuffer;
 			PassParameters->TileIntersectionParameters = TileIntersectionParameters;
 			PassParameters->DistanceFieldCulledObjectBuffers = CulledObjectBufferParameters;
-			PassParameters->DistanceFieldAtlas = DistanceField::SetupAtlasParameters(GraphBuilder, Scene->DistanceFieldSceneData);
+			PassParameters->DistanceFieldAtlas = DistanceField::SetupAtlasParameters(GraphBuilder, Scene.DistanceFieldSceneData);
 			PassParameters->SceneTextures = SceneTexturesUniformBuffer;
 
 			auto ComputeShader = View.ShaderMap->GetShader<FComputeCulledTilesStartOffsetCS>();
-			const FIntVector GroupCount = FComputeShaderUtils::GetGroupCountWrapped(Scene->DistanceFieldSceneData.NumObjectsInBuffer, ComputeStartOffsetGroupSize);
+			const FIntVector GroupCount = FComputeShaderUtils::GetGroupCountWrapped(Scene.DistanceFieldSceneData.NumObjectsInBuffer, ComputeStartOffsetGroupSize);
 
 			FComputeShaderUtils::AddPass(
 				GraphBuilder,
@@ -413,6 +413,6 @@ void BuildTileObjectLists(
 		AddClearUAVPass(GraphBuilder, TileIntersectionParameters.RWNumCulledTilesArray, 0);
 
 		// Rasterize object bounding shapes and intersect with screen tiles, and write out intersecting tile indices for the cone tracing pass
-		ScatterTilesToObjects(GraphBuilder, false, View, Scene->DistanceFieldSceneData, TileListGroupSize, Parameters, ObjectIndirectArguments, CulledObjectBufferParameters, TileIntersectionParameters, SceneTexturesUniformBuffer);
+		ScatterTilesToObjects(GraphBuilder, false, View, Scene.DistanceFieldSceneData, TileListGroupSize, Parameters, ObjectIndirectArguments, CulledObjectBufferParameters, TileIntersectionParameters, SceneTexturesUniformBuffer);
 	}
 }
diff --git a/Engine/Source/Runtime/Renderer/Private/DistanceFieldVisualization.cpp b/Engine/Source/Runtime/Renderer/Private/DistanceFieldVisualization.cpp
--- a/Engine/Source/Runtime/Renderer/Private/DistanceFieldVisualization.cpp
+++ b/Engine/Source/Runtime/Renderer/Private/DistanceFieldVisualization.cpp
@@ -155,7 +155,7 @@ void FDeferredShadingSceneRenderer::RenderMeshDistanceFieldVisualization(
 				ObjectIndirectArguments,
 				CulledObjectBufferParameters);
 
-			CullObjectsToView(GraphBuilder, Scene, View, Parameters, CulledObjectBufferParameters);
+			CullObjectsToView(GraphBuilder, *Scene, View, Parameters, CulledObjectBufferParameters);
 		}
 
 		uint32 GroupSizeX = FMath::DivideAndRoundUp(View.ViewRect.Size().X / GAODownsampleFactor, GDistanceFieldAOTileSizeX);
