Local Fog Volume - add support for mobile rendering

diff --git a/Engine/Shaders/Private/LocalHeightFog.usf b/Engine/Shaders/Private/LocalHeightFog.usf
--- a/Engine/Shaders/Private/LocalHeightFog.usf
+++ b/Engine/Shaders/Private/LocalHeightFog.usf
@@ -1,6 +1,10 @@
 // Copyright Epic Games, Inc. All Rights Reserved.
-
 #include "Common.ush"
+
+#if SHADING_PATH_MOBILE
+#define MobileSceneTextures MobileBasePass.SceneTextures 
+#endif
+
 #include "SceneTexturesCommon.ush"
 #include "HeightFogCommon.ush"
 #include "ParticipatingMediaCommon.ush"
@@ -22,10 +26,39 @@ struct FLocalHeightFogGPUInstanceData
 	float3 Albedo;
 	float  PhaseG;
 	float3 Emissive;
-	uint   FogMode;
+	float  FogMode;
 };
 
-StructuredBuffer<FLocalHeightFogGPUInstanceData> LocalHeightFogInstances;
+// Structured buffers cannot be sampled from the vertex shader on some mobile platforms.
+// That is why simple Buffers are used, because this tech needs to run on mobile platforms.
+Buffer<float4> LocalHeightFogInstances;
+FLocalHeightFogGPUInstanceData GetLocalHeightFogGPUInstanceData(uint Index)
+{
+	FLocalHeightFogGPUInstanceData Data;
+	uint Offset = Index * 19;
+
+	Data.Transform		= float4x4(LocalHeightFogInstances[Offset + 0], LocalHeightFogInstances[Offset + 1], LocalHeightFogInstances[Offset + 2], LocalHeightFogInstances[Offset + 3]);
+	Data.InvTransform	= float4x4(LocalHeightFogInstances[Offset + 4], LocalHeightFogInstances[Offset + 5], LocalHeightFogInstances[Offset + 6], LocalHeightFogInstances[Offset + 7]);
+
+	Data.InvTranformNoScale = float4x4(LocalHeightFogInstances[Offset + 8], LocalHeightFogInstances[Offset + 9], LocalHeightFogInstances[Offset + 10], LocalHeightFogInstances[Offset + 11]);
+	Data.TransformScaleOnly = float4x4(LocalHeightFogInstances[Offset + 12], LocalHeightFogInstances[Offset + 13], LocalHeightFogInstances[Offset + 14], LocalHeightFogInstances[Offset + 15]);
+
+	float4 Temp0 = LocalHeightFogInstances[Offset + 16];
+	Data.Density			= Temp0.x;
+	Data.HeightFalloff		= Temp0.y;
+	Data.HeightOffset		= Temp0.z;
+	Data.RadialAttenuation	= Temp0.w;
+
+	float4 Temp1 = LocalHeightFogInstances[Offset + 17];
+	Data.Albedo		= Temp1.xyz;
+	Data.PhaseG		= Temp1.w;
+
+	float4 Temp2 = LocalHeightFogInstances[Offset + 18];
+	Data.Emissive	= Temp2.xyz;
+	Data.FogMode	= Temp2.w;
+
+	return Data;
+}
 
 
 struct FLocalHeightFogVertexOutput
@@ -44,7 +77,7 @@ void LocalHeightFogSplatVS(
 
 	float4 LocalPosition = InPosition;
 
-	FLocalHeightFogGPUInstanceData FogInstance = LocalHeightFogInstances[InstanceIndex];
+	FLocalHeightFogGPUInstanceData FogInstance = GetLocalHeightFogGPUInstanceData(InstanceIndex);
 
 	float4x4 FogTransform = FogInstance.Transform;
 
@@ -190,7 +223,7 @@ void LocalHeightFogSplatPS(
 	float3 CamRayTranslatedWorldDir = normalize(SvPositionToTranslatedWorld(float4(SVPos.xy, 0.5, 1.0)));
 	float3 CamRayWorldDir = CamRayTranslatedWorldDir;
 
-	FLocalHeightFogGPUInstanceData FogInstance = LocalHeightFogInstances[InstanceIndex];
+	FLocalHeightFogGPUInstanceData FogInstance = GetLocalHeightFogGPUInstanceData(InstanceIndex);
 
 	OutColor = float4(0, 0, 0, 1);
 
@@ -204,7 +237,14 @@ void LocalHeightFogSplatPS(
 	float2 TsU = RayIntersectSphere(RayPosU, RayDirU, float4(0.0, 0.0, 0.0, 1.0));
 	if (any(TsU > 0.0))
 	{
+#if SHADING_PATH_MOBILE
+		// LookupDeviceZ with uint2 SVPos.xy as input does not work on mobile.
+		float2 UV = SVPos * View.BufferSizeAndInvSize.zw;
+		const float DeviceZ = LookupDeviceZ(UV);
+#else
 		const float DeviceZ = LookupDeviceZ(uint2(SVPos.xy));
+#endif
+
 		float3 DepthBufferTranslatedWorldPos = SvPositionToTranslatedWorld(float4(SVPos.xy, DeviceZ, 1.0));
 		float LengthD = length(DepthBufferTranslatedWorldPos - CamRayTranslatedWorldOrigin);
 
diff --git a/Engine/Source/Runtime/Renderer/Private/LocalHeightFogRendering.cpp b/Engine/Source/Runtime/Renderer/Private/LocalHeightFogRendering.cpp
--- a/Engine/Source/Runtime/Renderer/Private/LocalHeightFogRendering.cpp
+++ b/Engine/Source/Runtime/Renderer/Private/LocalHeightFogRendering.cpp
@@ -5,6 +5,7 @@
 #include "RendererUtils.h"
 #include "ScreenPass.h"
 #include "LocalHeightFogSceneProxy.h"
+#include "MobileBasePassRendering.h"
 
 
 // The runtime ON/OFF toggle
@@ -25,13 +26,10 @@ bool ShouldRenderLocalHeightFog(const FScene* Scene, const FSceneViewFamily& Fam
 
 DECLARE_GPU_STAT(LocalHeightFogVolumes);
 
-
 /*=============================================================================
 	FScene functions
 =============================================================================*/
 
-
-
 void FScene::AddLocalHeightFog(class FLocalHeightFogSceneProxy* FogProxy)
 {
 	check(FogProxy);
@@ -62,30 +60,108 @@ bool FScene::HasAnyLocalHeightFog() const
 	return LocalHeightFogs.Num() > 0;
 }
 
-
 /*=============================================================================
-	Local height fog rendering functions
+	Local height fog rendering common function
 =============================================================================*/
 
-class FLocalHeightFogGPUInstanceData
+void GetLocalFogVolumeSortingData(const FScene* Scene, FRDGBuilder& GraphBuilder, FLocalFogVolumeSortingData& Out)
 {
-public:
-	FMatrix44f Transform;
-	FMatrix44f InvTransform;
-
-	FMatrix44f InvTranformNoScale;
-	FMatrix44f TransformScaleOnly;
-
-	float Density;
-	float HeightFalloff;
-	float HeightOffset;
-	float RadialAttenuation;
-
-	FVector3f Albedo;
-	float PhaseG;
-	FVector3f Emissive;
-	uint32 FogMode;
-};
+	// No culling as of today
+	Out.LocalHeightFogInstanceCount = Scene->LocalHeightFogs.Num();
+	Out.LocalHeightFogInstanceCountFinal = 0;
+	Out.LocalHeightFogGPUInstanceData = (FLocalHeightFogGPUInstanceData*)GraphBuilder.Alloc(sizeof(FLocalHeightFogGPUInstanceData) * Out.LocalHeightFogInstanceCount, 16);
+	Out.LocalHeightFogCenterPos = (FVector*)GraphBuilder.Alloc(sizeof(FVector3f) * Out.LocalHeightFogInstanceCount, 16);
+	Out.LocalHeightFogSortKeys.SetNumUninitialized(Out.LocalHeightFogInstanceCount);
+	for (FLocalHeightFogSceneProxy* LHF : Scene->LocalHeightFogs)
+	{
+		if (LHF->FogDensity <= 0.0f)
+		{
+			continue; // this volume will never be visible
+		}
+
+		FTransform TransformScaleOnly;
+		TransformScaleOnly.SetScale3D(LHF->FogTransform.GetScale3D());
+
+		FLocalHeightFogGPUInstanceData* LocalHeightFogGPUInstanceDataIt = &Out.LocalHeightFogGPUInstanceData[Out.LocalHeightFogInstanceCountFinal];
+		LocalHeightFogGPUInstanceDataIt->Transform = FMatrix44f(LHF->FogTransform.ToMatrixWithScale());
+		LocalHeightFogGPUInstanceDataIt->InvTransform = LocalHeightFogGPUInstanceDataIt->Transform.Inverse();
+		LocalHeightFogGPUInstanceDataIt->InvTranformNoScale = FMatrix44f(LHF->FogTransform.ToMatrixNoScale()).Inverse();
+		LocalHeightFogGPUInstanceDataIt->TransformScaleOnly = FMatrix44f(TransformScaleOnly.ToMatrixWithScale());
+
+		LocalHeightFogGPUInstanceDataIt->Density = LHF->FogDensity;
+		LocalHeightFogGPUInstanceDataIt->HeightFalloff = LHF->FogHeightFalloff * 0.01f;	// This scale is used to have artist author reasonable range.
+		LocalHeightFogGPUInstanceDataIt->HeightOffset = LHF->FogHeightOffset;
+		LocalHeightFogGPUInstanceDataIt->RadialAttenuation = LHF->FogRadialAttenuation;
+
+		LocalHeightFogGPUInstanceDataIt->FogMode = float(LHF->FogMode);
+
+		LocalHeightFogGPUInstanceDataIt->Albedo = FVector3f(LHF->FogAlbedo);
+		LocalHeightFogGPUInstanceDataIt->PhaseG = LHF->FogPhaseG;
+		LocalHeightFogGPUInstanceDataIt->Emissive = FVector3f(LHF->FogEmissive);
+
+		Out.LocalHeightFogCenterPos[Out.LocalHeightFogInstanceCountFinal] = LHF->FogTransform.GetTranslation();
+
+		FLocalFogVolumeSortKey* LocalHeightFogSortKeysIt = &Out.LocalHeightFogSortKeys[Out.LocalHeightFogInstanceCountFinal];
+		LocalHeightFogSortKeysIt->FogVolume.Index = Out.LocalHeightFogInstanceCountFinal;
+		LocalHeightFogSortKeysIt->FogVolume.Distance = 0;	// Filled up right before sorting according to a view
+		LocalHeightFogSortKeysIt->FogVolume.Priority = LHF->FogSortPriority;
+
+		Out.LocalHeightFogInstanceCountFinal++;
+	}
+	// Shrink the array to only what is needed in order for the sort to correctly work on only what is needed.
+	Out.LocalHeightFogSortKeys.SetNum(Out.LocalHeightFogInstanceCountFinal, false/*bAllowShrinking*/);
+}
+
+void CreateViewLocalFogVolumeBufferSRV(FViewInfo& View, FRDGBuilder& GraphBuilder, FLocalFogVolumeSortingData& SortingData)
+{
+	static const uint32 SizeOfFloat4 = sizeof(float) * 4;
+	static const uint32 Float4CountInLocalHeightFogGPUInstanceData = sizeof(FLocalHeightFogGPUInstanceData) / SizeOfFloat4;
+	static_assert(sizeof(FLocalHeightFogGPUInstanceData) == Float4CountInLocalHeightFogGPUInstanceData * SizeOfFloat4); // The size of the structure must be a multiple of FVector4.
+
+	if (SortingData.LocalHeightFogInstanceCountFinal == 0)
+	{
+		View.LocalHeightFogGPUInstanceCount = 0;
+
+		static FLocalHeightFogGPUInstanceData DummyData;
+		View.LocalHeightFogGPUInstanceDataBufferSRV = GraphBuilder.CreateSRV(
+			CreateVertexBuffer(GraphBuilder, TEXT("LocalHeightFogGPUInstanceDataBuffer"), 
+			FRDGBufferDesc::CreateBufferDesc(SizeOfFloat4, Float4CountInLocalHeightFogGPUInstanceData), &DummyData, sizeof(FLocalHeightFogGPUInstanceData) * 1, ERDGInitialDataFlags::NoCopy),
+			PF_A32B32G32R32F);
+		return;
+	}
+
+	// 1. Sort all the volumes
+	const FVector ViewOrigin = View.ViewMatrices.GetViewOrigin();
+	for (uint32 i = 0; i < SortingData.LocalHeightFogInstanceCountFinal; ++i)
+	{
+		FVector FogCenterPos = SortingData.LocalHeightFogCenterPos[SortingData.LocalHeightFogSortKeys[i].FogVolume.Index];	// Recovered form the original array via index because the sorting of the previous view might have changed the order.
+		float DistancetoView = float((FogCenterPos - ViewOrigin).Size());
+		SortingData.LocalHeightFogSortKeys[i].FogVolume.Distance = *reinterpret_cast<uint32*>(&DistancetoView);
+	}
+	SortingData.LocalHeightFogSortKeys.Sort();
+
+	// 2. Create the buffer containing all the fog volume data instance sorted according to their key for the current view.
+	FLocalHeightFogGPUInstanceData* LocalHeightFogGPUSortedInstanceData = (FLocalHeightFogGPUInstanceData*)GraphBuilder.Alloc(sizeof(FLocalHeightFogGPUInstanceData) * SortingData.LocalHeightFogInstanceCountFinal, 16);
+	for (uint32 i = 0; i < SortingData.LocalHeightFogInstanceCountFinal; ++i)
+	{
+		// We could also have an indirection buffer on GPU but choosing to go with the sorting + copy on CPU since it is expected to not have many local height fog volumes.
+		LocalHeightFogGPUSortedInstanceData[i] = SortingData.LocalHeightFogGPUInstanceData[SortingData.LocalHeightFogSortKeys[i].FogVolume.Index];
+	}
+
+	// 3. Allocate buffer and initialize with sorted data to upload to GPU
+	const uint32 AllLocalHeightFogInstanceBytesFinal = sizeof(FLocalHeightFogGPUInstanceData) * SortingData.LocalHeightFogInstanceCountFinal;
+	FRDGBufferRef LocalHeightFogGPUInstanceDataBuffer = CreateVertexBuffer(
+		GraphBuilder, TEXT("LocalHeightFogGPUInstanceDataBuffer"),
+		FRDGBufferDesc::CreateBufferDesc(SizeOfFloat4, SortingData.LocalHeightFogInstanceCountFinal * Float4CountInLocalHeightFogGPUInstanceData), 
+		LocalHeightFogGPUSortedInstanceData, AllLocalHeightFogInstanceBytesFinal, ERDGInitialDataFlags::NoCopy);
+
+	View.LocalHeightFogGPUInstanceCount = SortingData.LocalHeightFogInstanceCountFinal;
+	View.LocalHeightFogGPUInstanceDataBufferSRV = GraphBuilder.CreateSRV(LocalHeightFogGPUInstanceDataBuffer, PF_A32B32G32R32F);
+}
+
+/*=============================================================================
+	Local height fog rendering - non mobile
+=============================================================================*/
 
 class FLocalHeightFogVS : public FGlobalShader
 {
@@ -94,13 +170,17 @@ class FLocalHeightFogVS : public FGlobalShader
 
 	BEGIN_SHADER_PARAMETER_STRUCT(FParameters, )
 		SHADER_PARAMETER_STRUCT_REF(FViewUniformShaderParameters, View)
-		SHADER_PARAMETER_RDG_BUFFER_SRV(StructuredBuffer<FLocalHeightFogGPUInstanceData>, LocalHeightFogInstances)
+		SHADER_PARAMETER_RDG_BUFFER_SRV(Buffer<float4>, LocalHeightFogInstances)
 	END_SHADER_PARAMETER_STRUCT()
 
 	using FPermutationDomain = TShaderPermutationDomain<>;
 
 	static bool ShouldCompilePermutation(const FGlobalShaderPermutationParameters& Parameters)
 	{
+		if (IsMobilePlatform(Parameters.Platform))
+		{
+			return false;
+		}
 		return true;
 	}
 };
@@ -114,13 +194,17 @@ class FLocalHeightFogPS : public FGlobalShader
 
 	BEGIN_SHADER_PARAMETER_STRUCT(FParameters, )
 		SHADER_PARAMETER_STRUCT_REF(FViewUniformShaderParameters, View)
-		SHADER_PARAMETER_RDG_BUFFER_SRV(StructuredBuffer<FLocalHeightFogGPUInstanceData>, LocalHeightFogInstances)
+		SHADER_PARAMETER_RDG_BUFFER_SRV(Buffer<float4>, LocalHeightFogInstances)
 	END_SHADER_PARAMETER_STRUCT()
 
 	using FPermutationDomain = TShaderPermutationDomain<>;
 
 	static bool ShouldCompilePermutation(const FGlobalShaderPermutationParameters& Parameters)
 	{
+		if (IsMobilePlatform(Parameters.Platform))
+		{
+			return false;
+		}
 		return true;
 	}
 };
@@ -133,33 +217,7 @@ BEGIN_SHADER_PARAMETER_STRUCT(FLocalHeightFogPassParameters, )
 	SHADER_PARAMETER_RDG_UNIFORM_BUFFER(FSceneTextureUniformParameters, SceneTextures)
 	RENDER_TARGET_BINDING_SLOTS()
 END_SHADER_PARAMETER_STRUCT()
-
-class RENDERER_API FLocalFogVolumeSortKey
-{
-public:
-	union
-	{
-		uint64 PackedData;
-
-		struct
-		{
-			uint64 Index	: 16; // Index of the volume
-			uint64 Distance : 32; // then by distance
-			uint64 Priority : 16; // First order by priority
-		} FogVolume;
-	};
-
-	FORCEINLINE bool operator!=(FLocalFogVolumeSortKey B) const
-	{
-		return PackedData != B.PackedData;
-	}
-
-	FORCEINLINE bool operator<(FLocalFogVolumeSortKey B) const
-	{
-		return PackedData > B.PackedData;
-	}
-};
-
+ 
 void RenderLocalHeightFog(
 	const FScene* Scene,
 	TArray<FViewInfo>& Views,
@@ -168,95 +226,32 @@ void RenderLocalHeightFog(
 	FRDGTextureRef LightShaftOcclusionTexture)
 {
 	uint32 LocalHeightFogInstanceCount = Scene->LocalHeightFogs.Num();
-	uint32 LocalHeightFogInstanceCountFinal = 0;
 	if (LocalHeightFogInstanceCount > 0)
 	{
 		RDG_GPU_STAT_SCOPE(GraphBuilder, LocalHeightFogVolumes);
-		
-		// No culling as of today
-		FLocalHeightFogGPUInstanceData* LocalHeightFogGPUInstanceData	= (FLocalHeightFogGPUInstanceData*) GraphBuilder.Alloc(sizeof(FLocalHeightFogGPUInstanceData) * LocalHeightFogInstanceCount, 16);
-		FVector*						LocalHeightFogCenterPos			= (FVector*) GraphBuilder.Alloc(sizeof(FVector3f) * LocalHeightFogInstanceCount, 16);
-		TArray<FLocalFogVolumeSortKey>	LocalHeightFogSortKeys;
-		LocalHeightFogSortKeys.SetNumUninitialized(LocalHeightFogInstanceCount);
-		for (FLocalHeightFogSceneProxy* LHF : Scene->LocalHeightFogs)
-		{
-			if (LHF->FogDensity <= 0.0f)
-			{
-				continue; // this volume will never be visible
-			}
-
-			FTransform TransformScaleOnly;
-			TransformScaleOnly.SetScale3D(LHF->FogTransform.GetScale3D());
 
-			FLocalHeightFogGPUInstanceData* LocalHeightFogGPUInstanceDataIt = &LocalHeightFogGPUInstanceData[LocalHeightFogInstanceCountFinal];
-			LocalHeightFogGPUInstanceDataIt->Transform					= FMatrix44f(LHF->FogTransform.ToMatrixWithScale());
-			LocalHeightFogGPUInstanceDataIt->InvTransform				= LocalHeightFogGPUInstanceDataIt->Transform.Inverse();
-			LocalHeightFogGPUInstanceDataIt->InvTranformNoScale			= FMatrix44f(LHF->FogTransform.ToMatrixNoScale()).Inverse();
-			LocalHeightFogGPUInstanceDataIt->TransformScaleOnly			= FMatrix44f(TransformScaleOnly.ToMatrixWithScale());
+		FLocalFogVolumeSortingData SortingData;
+		GetLocalFogVolumeSortingData(Scene, GraphBuilder, SortingData);
 
-			LocalHeightFogGPUInstanceDataIt->Density = LHF->FogDensity;
-			LocalHeightFogGPUInstanceDataIt->HeightFalloff = LHF->FogHeightFalloff * 0.01f;	// This scale is used to have artist author reasonable range.
-			LocalHeightFogGPUInstanceDataIt->HeightOffset = LHF->FogHeightOffset;
-			LocalHeightFogGPUInstanceDataIt->RadialAttenuation = LHF->FogRadialAttenuation;
-
-			LocalHeightFogGPUInstanceDataIt->FogMode = LHF->FogMode;
-
-			LocalHeightFogGPUInstanceDataIt->Albedo = FVector3f(LHF->FogAlbedo);
-			LocalHeightFogGPUInstanceDataIt->PhaseG = LHF->FogPhaseG;
-			LocalHeightFogGPUInstanceDataIt->Emissive = FVector3f(LHF->FogEmissive);
-
-			LocalHeightFogCenterPos[LocalHeightFogInstanceCountFinal] = LHF->FogTransform.GetTranslation();
-
-			FLocalFogVolumeSortKey* LocalHeightFogSortKeysIt = &LocalHeightFogSortKeys[LocalHeightFogInstanceCountFinal];
-			LocalHeightFogSortKeysIt->FogVolume.Index = LocalHeightFogInstanceCountFinal;
-			LocalHeightFogSortKeysIt->FogVolume.Distance = 0;	// Filled up right before sorting according to a view
-			LocalHeightFogSortKeysIt->FogVolume.Priority = LHF->FogSortPriority;
-
-			LocalHeightFogInstanceCountFinal++;
-		}
-		// Shrink the array to only what is needed in order for the sort to correctly work on only what is needed.
-		LocalHeightFogSortKeys.SetNum(LocalHeightFogInstanceCountFinal, false/*bAllowShrinking*/);
-
-		if (LocalHeightFogInstanceCountFinal > 0)
+		if (SortingData.LocalHeightFogInstanceCountFinal > 0)
 		{
-
 			FRDGTextureRef SceneColorTexture = SceneTextures.Color.Resolve;
 
 			for (FViewInfo& View : Views)
 			{
-				// 1. Sort all the volumes
-				const FVector ViewOrigin = View.ViewMatrices.GetViewOrigin();
-				for (uint32 i = 0; i < LocalHeightFogInstanceCountFinal; ++i)
+				CreateViewLocalFogVolumeBufferSRV(View, GraphBuilder, SortingData);
+				if (View.LocalHeightFogGPUInstanceCount == 0)
 				{
-					FVector FogCenterPos = LocalHeightFogCenterPos[LocalHeightFogSortKeys[i].FogVolume.Index];	// Recovered form the original array via index because the sorting of the previous view might have changed the order.
-					float DistancetoView = float((FogCenterPos - ViewOrigin).Size());
-					LocalHeightFogSortKeys[i].FogVolume.Distance = *reinterpret_cast<uint32*>(&DistancetoView);
+					continue;
 				}
-				LocalHeightFogSortKeys.Sort();
 
-				// 2. Create the buffer containing all the fog volume data instance sorted according to their key for the current view.
-				FLocalHeightFogGPUInstanceData* LocalHeightFogGPUSortedInstanceData = (FLocalHeightFogGPUInstanceData*)GraphBuilder.Alloc(sizeof(FLocalHeightFogGPUInstanceData) * LocalHeightFogInstanceCountFinal, 16);
-				for (uint32 i = 0; i < LocalHeightFogInstanceCountFinal; ++i)
-				{
-					// We could also have an indirection buffer on GPU but choosing to go with the sorting + copy on CPU since it is expected to not have many local height fog volumes.
-					LocalHeightFogGPUSortedInstanceData[i] = LocalHeightFogGPUInstanceData[LocalHeightFogSortKeys[i].FogVolume.Index];
-				}
-
-				// 3. Allocate buffer and initialize with sorted data to upload to GPU
-				const uint32 AllLocalHeightFogInstanceBytesFinal = sizeof(FLocalHeightFogGPUInstanceData) * LocalHeightFogInstanceCountFinal;
-				FRDGBufferRef LocalHeightFogGPUInstanceDataBuffer = CreateStructuredBuffer(GraphBuilder, TEXT("LocalHeightFogGPUInstanceDataBuffer"),
-					sizeof(FLocalHeightFogGPUInstanceData), LocalHeightFogInstanceCountFinal, LocalHeightFogGPUSortedInstanceData, AllLocalHeightFogInstanceBytesFinal, ERDGInitialDataFlags::NoCopy);
-
-				FRDGBufferSRVRef LocalHeightFogGPUInstanceDataBufferSRV = GraphBuilder.CreateSRV(LocalHeightFogGPUInstanceDataBuffer);
-
-				// 4. Render
 				FLocalHeightFogPassParameters* PassParameters = GraphBuilder.AllocParameters<FLocalHeightFogPassParameters>();
 
 				PassParameters->VS.View = GetShaderBinding(View.ViewUniformBuffer);
-				PassParameters->VS.LocalHeightFogInstances = LocalHeightFogGPUInstanceDataBufferSRV;
+				PassParameters->VS.LocalHeightFogInstances = View.LocalHeightFogGPUInstanceDataBufferSRV;
 
 				PassParameters->PS.View = GetShaderBinding(View.ViewUniformBuffer);
-				PassParameters->PS.LocalHeightFogInstances = LocalHeightFogGPUInstanceDataBufferSRV;
+				PassParameters->PS.LocalHeightFogInstances = View.LocalHeightFogGPUInstanceDataBufferSRV;
 
 				PassParameters->SceneTextures = SceneTextures.UniformBuffer;
 				PassParameters->RenderTargets[0] = FRenderTargetBinding(SceneColorTexture, ERenderTargetLoadAction::ENoAction);
@@ -272,11 +267,12 @@ void RenderLocalHeightFog(
 				ClearUnusedGraphResources(VertexShader, &PassParameters->VS);
 				ClearUnusedGraphResources(PixelShader, &PassParameters->PS);
 
+				uint32 LocalHeightFogGPUInstanceCount = View.LocalHeightFogGPUInstanceCount;
 				GraphBuilder.AddPass(
-					RDG_EVENT_NAME("RenderLocalHeightFog %u inst.", LocalHeightFogInstanceCountFinal),
+					RDG_EVENT_NAME("RenderLocalHeightFog %u inst.", LocalHeightFogGPUInstanceCount),
 					PassParameters,
 					ERDGPassFlags::Raster,
-					[VertexShader, PixelShader, PassParameters, LocalHeightFogInstanceCountFinal, ViewRect](FRHICommandList& RHICmdList)
+					[VertexShader, PixelShader, PassParameters, LocalHeightFogGPUInstanceCount, ViewRect](FRHICommandList& RHICmdList)
 				{
 					FGraphicsPipelineStateInitializer GraphicsPSOInit;
 					RHICmdList.ApplyCachedRenderTargets(GraphicsPSOInit);
@@ -306,10 +302,138 @@ void RenderLocalHeightFog(
 						, 8									//uint32 NumVertices
 						, 0									//uint32 StartIndex
 						, UE_ARRAY_COUNT(GCubeIndices) / 3	//uint32 NumPrimitives
-						, LocalHeightFogInstanceCountFinal	//uint32 NumInstances
+						, LocalHeightFogGPUInstanceCount	//uint32 NumInstances
 					);
 				});
 			}
 		}
 	}
 }
+
+/*=============================================================================
+	Local height fog rendering - mobile
+=============================================================================*/
+	
+
+class FMobileLocalHeightFogVS : public FGlobalShader
+{
+	DECLARE_GLOBAL_SHADER(FMobileLocalHeightFogVS);
+	SHADER_USE_PARAMETER_STRUCT(FMobileLocalHeightFogVS, FGlobalShader);
+
+	BEGIN_SHADER_PARAMETER_STRUCT(FParameters, )
+		SHADER_PARAMETER_STRUCT_INCLUDE(FViewShaderParameters, View)
+		SHADER_PARAMETER_RDG_BUFFER_SRV(Buffer<float4>, LocalHeightFogInstances)
+	END_SHADER_PARAMETER_STRUCT()
+
+	using FPermutationDomain = TShaderPermutationDomain<>;
+
+	static bool ShouldCompilePermutation(const FGlobalShaderPermutationParameters& Parameters)
+	{
+		if (!IsMobilePlatform(Parameters.Platform))
+		{
+			return false;
+		}
+		return true;
+	}
+
+	static void ModifyCompilationEnvironment(const FGlobalShaderPermutationParameters& Parameters, FShaderCompilerEnvironment& OutEnvironment)
+	{
+		FGlobalShader::ModifyCompilationEnvironment(Parameters, OutEnvironment);
+		const bool bMobileForceDepthRead = MobileUsesFullDepthPrepass(Parameters.Platform);
+		OutEnvironment.SetDefine(TEXT("IS_MOBILE_DEPTHREAD_SUBPASS"), bMobileForceDepthRead ? 0u : 1u);
+	}
+};
+
+IMPLEMENT_GLOBAL_SHADER(FMobileLocalHeightFogVS, "/Engine/Private/LocalHeightFog.usf", "LocalHeightFogSplatVS", SF_Vertex);
+
+class FMobileLocalHeightFogPS : public FGlobalShader
+{
+	DECLARE_GLOBAL_SHADER(FMobileLocalHeightFogPS);
+	SHADER_USE_PARAMETER_STRUCT(FMobileLocalHeightFogPS, FGlobalShader);
+
+	BEGIN_SHADER_PARAMETER_STRUCT(FParameters, )
+		SHADER_PARAMETER_STRUCT_INCLUDE(FViewShaderParameters, View)
+		SHADER_PARAMETER_RDG_UNIFORM_BUFFER(FMobileBasePassUniformParameters, MobileBasePass)
+		SHADER_PARAMETER_RDG_BUFFER_SRV(Buffer<float4>, LocalHeightFogInstances)
+	END_SHADER_PARAMETER_STRUCT()
+
+	using FPermutationDomain = TShaderPermutationDomain<>;
+
+	static bool ShouldCompilePermutation(const FGlobalShaderPermutationParameters& Parameters)
+	{
+		if (!IsMobilePlatform(Parameters.Platform))
+		{
+			return false;
+		}
+		return true;
+	}
+
+	static void ModifyCompilationEnvironment(const FGlobalShaderPermutationParameters& Parameters, FShaderCompilerEnvironment& OutEnvironment)
+	{
+		FGlobalShader::ModifyCompilationEnvironment(Parameters, OutEnvironment);
+		const bool bMobileForceDepthRead = MobileUsesFullDepthPrepass(Parameters.Platform);
+		OutEnvironment.SetDefine(TEXT("IS_MOBILE_DEPTHREAD_SUBPASS"), bMobileForceDepthRead ? 0u : 1u);
+	}
+};
+
+IMPLEMENT_GLOBAL_SHADER(FMobileLocalHeightFogPS, "/Engine/Private/LocalHeightFog.usf", "LocalHeightFogSplatPS", SF_Pixel);
+
+void RenderLocalHeightFogMobile(
+	FRHICommandList& RHICmdList,
+	const FViewInfo& View)
+{
+	if (View.LocalHeightFogGPUInstanceCount == 0)
+	{
+		return;
+	}
+
+	SCOPED_DRAW_EVENT(RHICmdList, LocalHeightFogVolumes);
+
+	FMobileLocalHeightFogVS::FPermutationDomain VSPermutationVector;
+	auto VertexShader = View.ShaderMap->GetShader< FMobileLocalHeightFogVS >(VSPermutationVector);
+
+	FMobileLocalHeightFogPS::FPermutationDomain PsPermutationVector;
+	auto PixelShader = View.ShaderMap->GetShader< FMobileLocalHeightFogPS >(PsPermutationVector);
+
+	const FIntRect ViewRect = View.ViewRect;
+
+	FGraphicsPipelineStateInitializer GraphicsPSOInit;
+	RHICmdList.ApplyCachedRenderTargets(GraphicsPSOInit);
+
+	RHICmdList.SetViewport(ViewRect.Min.X, ViewRect.Min.Y, 0.0f, ViewRect.Max.X, ViewRect.Max.Y, 1.0f);
+
+	// Render back faces only since camera may intersect
+	GraphicsPSOInit.RasterizerState = TStaticRasterizerState<FM_Solid, CM_CW>::GetRHI();
+	GraphicsPSOInit.DepthStencilState = TStaticDepthStencilState<false, CF_Always>::GetRHI();
+	GraphicsPSOInit.BlendState = TStaticBlendState<CW_RGBA, BO_Add, BF_One, BF_SourceAlpha, BO_Add, BF_Zero, BF_SourceAlpha>::GetRHI();
+	GraphicsPSOInit.PrimitiveType = PT_TriangleList;
+
+	GraphicsPSOInit.BoundShaderState.VertexDeclarationRHI = GetVertexDeclarationFVector4();
+	GraphicsPSOInit.BoundShaderState.VertexShaderRHI = VertexShader.GetVertexShader();
+	GraphicsPSOInit.BoundShaderState.PixelShaderRHI = PixelShader.GetPixelShader();
+
+	SetGraphicsPipelineState(RHICmdList, GraphicsPSOInit, 0);
+
+	FMobileLocalHeightFogVS::FParameters VSParameters;
+	VSParameters.View = View.GetShaderParameters();
+	VSParameters.LocalHeightFogInstances = View.LocalHeightFogGPUInstanceDataBufferSRV;
+	SetShaderParameters(RHICmdList, VertexShader, VertexShader.GetVertexShader(), VSParameters);
+
+	FMobileLocalHeightFogPS::FParameters PSParameters;
+	PSParameters.View = View.GetShaderParameters();
+	PSParameters.LocalHeightFogInstances = View.LocalHeightFogGPUInstanceDataBufferSRV;
+	// PSParameters.MobileBasePass filled up by the RDG pass parameters.
+	SetShaderParameters(RHICmdList, PixelShader, PixelShader.GetPixelShader(), PSParameters);
+
+	RHICmdList.SetStreamSource(0, GetUnitCubeVertexBuffer(), 0);
+
+	RHICmdList.DrawIndexedPrimitive(
+		GetUnitCubeIndexBuffer()
+		, 0										//BaseVertexIndex
+		, 0										//FirstInstance
+		, 8										//uint32 NumVertices
+		, 0										//uint32 StartIndex
+		, UE_ARRAY_COUNT(GCubeIndices) / 3		//uint32 NumPrimitives
+		, View.LocalHeightFogGPUInstanceCount	//uint32 NumInstances
+	);
+}
diff --git a/Engine/Source/Runtime/Renderer/Private/LocalHeightFogRendering.h b/Engine/Source/Runtime/Renderer/Private/LocalHeightFogRendering.h
--- a/Engine/Source/Runtime/Renderer/Private/LocalHeightFogRendering.h
+++ b/Engine/Source/Runtime/Renderer/Private/LocalHeightFogRendering.h
@@ -7,13 +7,85 @@
 #include "RenderResource.h"
 #include "RenderGraphResources.h"
 #include "SceneView.h"
+#include "Containers/Array.h"
 
 class FScene;
 class FViewInfo;
 struct FMinimalSceneTextures;
 
+
+
+/*=============================================================================
+	Local height fog rendering common data
+=============================================================================*/
+
+class FLocalHeightFogGPUInstanceData
+{
+public:
+	FMatrix44f Transform;
+	FMatrix44f InvTransform;
+
+	FMatrix44f InvTranformNoScale;
+	FMatrix44f TransformScaleOnly;
+
+	float Density;
+	float HeightFalloff;
+	float HeightOffset;
+	float RadialAttenuation;
+
+	FVector3f Albedo;
+	float PhaseG;
+	FVector3f Emissive;
+	float FogMode;
+};
+
+class RENDERER_API FLocalFogVolumeSortKey
+{
+public:
+	union
+	{
+		uint64 PackedData;
+
+		struct
+		{
+			uint64 Index : 16; // Index of the volume
+			uint64 Distance : 32; // then by distance
+			uint64 Priority : 16; // First order by priority
+		} FogVolume;
+	};
+
+	FORCEINLINE bool operator!=(FLocalFogVolumeSortKey B) const
+	{
+		return PackedData != B.PackedData;
+	}
+
+	FORCEINLINE bool operator<(FLocalFogVolumeSortKey B) const
+	{
+		return PackedData > B.PackedData;
+	}
+};
+
+// This can be generated once for a scene and then shared between all the views to generate the GPU buffer data needed for rendering while accounting for sorting.
+struct FLocalFogVolumeSortingData
+{
+	uint32 LocalHeightFogInstanceCount;
+	uint32 LocalHeightFogInstanceCountFinal;
+	FLocalHeightFogGPUInstanceData* LocalHeightFogGPUInstanceData;
+	FVector* LocalHeightFogCenterPos;
+	TArray<FLocalFogVolumeSortKey>	LocalHeightFogSortKeys;
+};
+
+
+/*=============================================================================
+	Local height fog rendering functions
+=============================================================================*/
+
 bool ShouldRenderLocalHeightFog(const FScene* Scene, const FSceneViewFamily& Family);
 
+void GetLocalFogVolumeSortingData(const FScene* Scene, FRDGBuilder& GraphBuilder, FLocalFogVolumeSortingData& Out);
+
+void CreateViewLocalFogVolumeBufferSRV(FViewInfo& View, FRDGBuilder& GraphBuilder, FLocalFogVolumeSortingData& SortingData);
+
 void RenderLocalHeightFog(
 	const FScene* Scene,
 	TArray<FViewInfo>& Views,
@@ -21,3 +93,7 @@ void RenderLocalHeightFog(
 	const FMinimalSceneTextures& SceneTextures,
 	FRDGTextureRef LightShaftOcclusionTexture);
 
+void RenderLocalHeightFogMobile(
+	FRHICommandList& RHICmdList, 
+	const FViewInfo& View);
+
diff --git a/Engine/Source/Runtime/Renderer/Private/MobileShadingRenderer.cpp b/Engine/Source/Runtime/Renderer/Private/MobileShadingRenderer.cpp
--- a/Engine/Source/Runtime/Renderer/Private/MobileShadingRenderer.cpp
+++ b/Engine/Source/Runtime/Renderer/Private/MobileShadingRenderer.cpp
@@ -66,6 +66,7 @@
 #include "IESTextureManager.h"
 #include "SceneUniformBuffer.h"
 #include "Engine/SpecularProfile.h"
+#include "LocalHeightFogRendering.h"
 
 uint32 GetShadowQuality();
 
@@ -192,6 +193,7 @@ BEGIN_SHADER_PARAMETER_STRUCT(FMobileRenderPassParameters, RENDERER_API)
 	SHADER_PARAMETER_RDG_UNIFORM_BUFFER(FSceneUniformParameters, Scene)
 	SHADER_PARAMETER_RDG_UNIFORM_BUFFER(FMobileBasePassUniformParameters, MobileBasePass)
 	SHADER_PARAMETER_STRUCT_REF(FMobileReflectionCaptureShaderData, ReflectionCapture)
+	SHADER_PARAMETER_RDG_BUFFER_SRV(Buffer<float4>, LocalHeightFogInstances)
 	RDG_BUFFER_ACCESS_ARRAY(DrawIndirectArgsBuffers)
 	RDG_BUFFER_ACCESS_ARRAY(InstanceIdOffsetBuffers)
 	RENDER_TARGET_BINDING_SLOTS()
@@ -1208,6 +1210,9 @@ void FMobileSceneRenderer::RenderForward(FRDGBuilder& GraphBuilder, FRDGTextureR
 
 	const FRDGSystemTextures& SystemTextures = FRDGSystemTextures::Get(GraphBuilder);
 
+	FLocalFogVolumeSortingData SortingData;
+	GetLocalFogVolumeSortingData(Scene, GraphBuilder, SortingData);
+
 	FRenderViewContextArray RenderViews;
 	GetRenderViews(Views, RenderViews);
 
@@ -1233,6 +1238,7 @@ void FMobileSceneRenderer::RenderForward(FRDGBuilder& GraphBuilder, FRDGTextureR
 		View.BeginRenderView();
 
 		UpdateDirectionalLightUniformBuffers(GraphBuilder, View);
+		CreateViewLocalFogVolumeBufferSRV(View, GraphBuilder, SortingData);
 
 		FMobileBasePassTextures MobileBasePassTextures{};
 		MobileBasePassTextures.ScreenSpaceAO = bRequiresAmbientOcclusionPass ? SceneTextures.ScreenSpaceAO : SystemTextures.White;
@@ -1244,6 +1250,7 @@ void FMobileSceneRenderer::RenderForward(FRDGBuilder& GraphBuilder, FRDGTextureR
 		PassParameters->MobileBasePass = CreateMobileBasePassUniformBuffer(GraphBuilder, View, EMobileBasePass::Opaque, SetupMode, MobileBasePassTextures);
 		PassParameters->ReflectionCapture = View.MobileReflectionCaptureUniformBuffer;
 		PassParameters->RenderTargets = BasePassRenderTargets;
+		PassParameters->LocalHeightFogInstances = View.LocalHeightFogGPUInstanceDataBufferSRV;
 	
 		BuildInstanceCullingDrawParams(GraphBuilder, View, PassParameters);
 
@@ -1295,6 +1302,7 @@ void FMobileSceneRenderer::RenderForwardSinglePass(FRDGBuilder& GraphBuilder, FM
 		RenderDecals(RHICmdList, View);
 		RenderModulatedShadowProjections(RHICmdList, ViewContext.ViewIndex, View);
 		RenderFog(RHICmdList, View);
+		RenderLocalHeightFogMobile(RHICmdList, View);
 		// Draw translucency.
 		RenderTranslucency(RHICmdList, View);
 
@@ -1394,6 +1402,7 @@ void FMobileSceneRenderer::RenderForwardMultiPass(FRDGBuilder& GraphBuilder, FMo
 		RenderDecals(RHICmdList, View);
 		RenderModulatedShadowProjections(RHICmdList, ViewContext.ViewIndex, View);
 		RenderFog(RHICmdList, View);
+		RenderLocalHeightFogMobile(RHICmdList, View);
 		// Draw translucency.
 		RenderTranslucency(RHICmdList, View);
 
@@ -1536,6 +1545,9 @@ void FMobileSceneRenderer::RenderDeferred(FRDGBuilder& GraphBuilder, const FSort
 
 	const FRDGSystemTextures& SystemTextures = FRDGSystemTextures::Get(GraphBuilder);
 
+	FLocalFogVolumeSortingData SortingData;
+	GetLocalFogVolumeSortingData(Scene, GraphBuilder, SortingData);
+
 	FRenderViewContextArray RenderViews;
 	GetRenderViews(Views, RenderViews);
 
@@ -1561,6 +1573,7 @@ void FMobileSceneRenderer::RenderDeferred(FRDGBuilder& GraphBuilder, const FSort
 		View.BeginRenderView();
 
 		UpdateDirectionalLightUniformBuffers(GraphBuilder, View);
+		CreateViewLocalFogVolumeBufferSRV(View, GraphBuilder, SortingData);
 
 		FMobileBasePassTextures MobileBasePassTextures{};
 		MobileBasePassTextures.ScreenSpaceAO = bRequiresAmbientOcclusionPass ? SceneTextures.ScreenSpaceAO : SystemTextures.White;
@@ -1572,6 +1585,7 @@ void FMobileSceneRenderer::RenderDeferred(FRDGBuilder& GraphBuilder, const FSort
 		PassParameters->MobileBasePass = CreateMobileBasePassUniformBuffer(GraphBuilder, View, EMobileBasePass::Opaque, SetupMode, MobileBasePassTextures);
 		PassParameters->ReflectionCapture = View.MobileReflectionCaptureUniformBuffer;
 		PassParameters->RenderTargets = BasePassRenderTargets;
+		PassParameters->LocalHeightFogInstances = View.LocalHeightFogGPUInstanceDataBufferSRV;
 		
 		BuildInstanceCullingDrawParams(GraphBuilder, View, PassParameters);
 
@@ -1623,6 +1637,7 @@ void FMobileSceneRenderer::RenderDeferredSinglePass(FRDGBuilder& GraphBuilder, c
 			MobileDeferredCopyBuffer<FMobileDeferredCopyPLSPS>(RHICmdList, View);
 		}
 		RenderFog(RHICmdList, View);
+		RenderLocalHeightFogMobile(RHICmdList, View);
 		// Draw translucency.
 		RenderTranslucency(RHICmdList, View);
 
@@ -1704,6 +1719,7 @@ void FMobileSceneRenderer::RenderDeferredMultiPass(FRDGBuilder& GraphBuilder, cl
 			
 		RHICmdList.SetCurrentStat(GET_STATID(STAT_CLMM_Translucency));
 		MobileDeferredShadingPass(RHICmdList, ViewContext.ViewIndex, Views.Num(), View, *Scene, SortedLightSet, VisibleLightInfos);
+		RenderLocalHeightFogMobile(RHICmdList, View);
 		RenderFog(RHICmdList, View);
 		// Draw translucency.
 		RenderTranslucency(RHICmdList, View);
diff --git a/Engine/Source/Runtime/Renderer/Private/SceneRendering.cpp b/Engine/Source/Runtime/Renderer/Private/SceneRendering.cpp
--- a/Engine/Source/Runtime/Renderer/Private/SceneRendering.cpp
+++ b/Engine/Source/Runtime/Renderer/Private/SceneRendering.cpp
@@ -970,6 +970,9 @@ void FViewInfo::Init()
 	StrataViewData.Reset();
 	HairStrandsViewData = FHairStrandsViewData();
 
+	LocalHeightFogGPUInstanceCount = 0;
+	LocalHeightFogGPUInstanceDataBufferSRV = nullptr;
+
 	GPUSceneViewId = INDEX_NONE;
 
 	// Filled in by FDeferredShadingSceneRenderer::UpdateLumenScene
diff --git a/Engine/Source/Runtime/Renderer/Private/SceneRendering.h b/Engine/Source/Runtime/Renderer/Private/SceneRendering.h
--- a/Engine/Source/Runtime/Renderer/Private/SceneRendering.h
+++ b/Engine/Source/Runtime/Renderer/Private/SceneRendering.h
@@ -1412,6 +1412,9 @@ public:
 
 	FStrataViewData StrataViewData;
 
+	uint32 LocalHeightFogGPUInstanceCount;
+	FRDGBufferSRVRef LocalHeightFogGPUInstanceDataBufferSRV;
+
 	FHairStrandsViewData HairStrandsViewData;
 
 	/** Parameters for exponential height fog. */
