Force DXC for FRayTracingBuildInstanceBufferCS to avoid shader reflection issues in SM5.

diff --git a/Engine/Source/Runtime/Renderer/Private/RayTracing/RayTracingInstanceBufferUtil.cpp b/Engine/Source/Runtime/Renderer/Private/RayTracing/RayTracingInstanceBufferUtil.cpp
--- a/Engine/Source/Runtime/Renderer/Private/RayTracing/RayTracingInstanceBufferUtil.cpp
+++ b/Engine/Source/Runtime/Renderer/Private/RayTracing/RayTracingInstanceBufferUtil.cpp
@@ -15,6 +15,7 @@
 #include "GlobalShader.h"
 #include "PipelineStateCache.h"
 #include "DataDrivenShaderPlatformInfo.h"
+#include "ShaderCompilerCore.h"
 
 #include "Async/ParallelFor.h"
 
@@ -317,6 +318,9 @@ struct FRayTracingBuildInstanceBufferCS : public FGlobalShader
 		OutEnvironment.SetDefine(TEXT("THREADGROUP_SIZE"), ThreadGroupSize);
 		OutEnvironment.SetDefine(TEXT("VF_SUPPORTS_PRIMITIVE_SCENE_DATA"), 1);
 		OutEnvironment.SetDefine(TEXT("USE_GLOBAL_GPU_SCENE_DATA"), 1);
+
+		// Force DXC to avoid shader reflection issues.
+		OutEnvironment.CompilerFlags.Add(CFLAG_ForceDXC);
 	}
 
 	static bool ShouldCompilePermutation(const FGlobalShaderPermutationParameters& Parameters)
