Sort heterogeneous volumes by camera-depth and adjust compositing to match.

diff --git a/Engine/Shaders/Private/HeterogeneousVolumes/HeterogeneousVolumesLiveShadingPipeline.usf b/Engine/Shaders/Private/HeterogeneousVolumes/HeterogeneousVolumesLiveShadingPipeline.usf
--- a/Engine/Shaders/Private/HeterogeneousVolumes/HeterogeneousVolumesLiveShadingPipeline.usf
+++ b/Engine/Shaders/Private/HeterogeneousVolumes/HeterogeneousVolumesLiveShadingPipeline.usf
@@ -257,8 +257,6 @@ void RenderSingleScatteringWithLiveShadingCS(
 	uint2 DispatchThreadId : SV_DispatchThreadID
 )
 {
-	float3 Radiance = 0.0;
-	float3 Transmittance = 1.0;
 	FDebugOutput DebugOutput = (FDebugOutput)0;
 
 	// Create screen ray
@@ -269,6 +267,9 @@ void RenderSingleScatteringWithLiveShadingCS(
 	uint LinearIndex = DispatchThreadId.y * GroupCount.x * THREADGROUP_SIZE_2D + DispatchThreadId.x;
 	uint2 PixelCoord = DispatchThreadId.xy + View.ViewRectMin.xy;
 
+	float3 Radiance = 0.0;
+	float3 Transmittance = 1.0 - RWLightingTexture[PixelCoord].a;
+
 	// Extract depth
 	float DeviceZ = SceneDepthTexture.Load(int3(PixelCoord, 0)).r;
 #if HAS_INVERTED_Z_BUFFER
@@ -336,7 +337,8 @@ void RenderSingleScatteringWithLiveShadingCS(
 
 		// Output..
 		float3 Opacity = bApplyEmissionAndTransmittance ? 1.0 - Transmittance : 0.0;
-		RWLightingTexture[PixelCoord] += float4(Radiance * View.PreExposure, Luminance(Opacity));
+		RWLightingTexture[PixelCoord].rgb += Radiance * View.PreExposure;
+		RWLightingTexture[PixelCoord].a = Luminance(Opacity);
 
 		// Debug output
 		//RWDebugOutputBuffer[LinearIndex] = DebugOutput;
diff --git a/Engine/Source/Runtime/Renderer/Private/HeterogeneousVolumes/HeterogeneousVolumes.cpp b/Engine/Source/Runtime/Renderer/Private/HeterogeneousVolumes/HeterogeneousVolumes.cpp
--- a/Engine/Source/Runtime/Renderer/Private/HeterogeneousVolumes/HeterogeneousVolumes.cpp
+++ b/Engine/Source/Runtime/Renderer/Private/HeterogeneousVolumes/HeterogeneousVolumes.cpp
@@ -169,6 +169,13 @@ static TAutoConsoleVariable<int32> CVarHeterogeneousVolumesLightingCacheDownsamp
 	ECVF_RenderThreadSafe
 );
 
+static TAutoConsoleVariable<int32> CVarHeterogeneousVolumesDepthSort(
+	TEXT("r.HeterogeneousVolumes.DepthSort"),
+	1,
+	TEXT("Iterates over volumes in depth-sorted order, based on its centroid (Default = 1)"),
+	ECVF_RenderThreadSafe
+);
+
 DECLARE_GPU_STAT_NAMED(HeterogeneousVolumesStat, TEXT("HeterogeneousVolumes"));
 
 static bool IsHeterogeneousVolumesEnabled()
@@ -422,6 +429,19 @@ void FDeferredShadingSceneRenderer::RenderHeterogeneousVolumes(
 			}
 			else
 			{
+				// Collect volume interfaces
+				struct VolumeMesh
+				{
+					const IHeterogeneousVolumeInterface* Volume;
+					const FMaterialRenderProxy* MaterialRenderProxy;
+
+					VolumeMesh(const IHeterogeneousVolumeInterface* V, const FMaterialRenderProxy* M) : 
+						Volume(V),
+						MaterialRenderProxy(M)
+					{}
+				};
+
+				TArray<VolumeMesh> VolumeMeshes;
 				for (int32 MeshBatchIndex = 0; MeshBatchIndex < View.HeterogeneousVolumesMeshBatches.Num(); ++MeshBatchIndex)
 				{
 					const FMeshBatch* Mesh = View.HeterogeneousVolumesMeshBatches[MeshBatchIndex].Mesh;
@@ -434,13 +454,47 @@ void FDeferredShadingSceneRenderer::RenderHeterogeneousVolumes(
 					const FMaterialRenderProxy* MaterialRenderProxy = Mesh->MaterialRenderProxy;
 					for (int32 VolumeIndex = 0; VolumeIndex < Mesh->Elements.Num(); ++VolumeIndex)
 					{
-						IHeterogeneousVolumeInterface* HeterogeneousVolume = (IHeterogeneousVolumeInterface*) Mesh->Elements[VolumeIndex].UserData;
+						const IHeterogeneousVolumeInterface* HeterogeneousVolume = (IHeterogeneousVolumeInterface*)Mesh->Elements[VolumeIndex].UserData;
 						//check(HeterogeneousVolume != nullptr);
 						if (HeterogeneousVolume == nullptr)
 						{
 							continue;
 						}
 
+						VolumeMeshes.Add(VolumeMesh(HeterogeneousVolume, MaterialRenderProxy));
+					}
+				}
+
+				// Provide coarse depth-sorting, based on camera-distance to world centroid
+				bool bDepthSort = CVarHeterogeneousVolumesDepthSort.GetValueOnRenderThread() == 1;
+				if (bDepthSort)
+				{
+					struct FDepthCompareHeterogeneousVolumes
+					{
+						FVector WorldCameraOrigin;
+						FDepthCompareHeterogeneousVolumes(FViewInfo& View) : WorldCameraOrigin(View.ViewMatrices.GetViewOrigin()) {}
+
+						FORCEINLINE bool operator()(const VolumeMesh& A, const VolumeMesh& B) const
+						{
+							FVector CameraToA = A.Volume->GetBounds().Origin - WorldCameraOrigin;
+							float SquaredDistanceToA = FVector::DotProduct(CameraToA, CameraToA);
+
+							FVector CameraToB = B.Volume->GetBounds().Origin - WorldCameraOrigin;
+							float SquaredDistanceToB = FVector::DotProduct(CameraToB, CameraToB);
+
+							return SquaredDistanceToA < SquaredDistanceToB;
+						}
+					};
+
+					VolumeMeshes.Sort(FDepthCompareHeterogeneousVolumes(View));
+				}
+
+				for (int32 VolumeIndex = 0; VolumeIndex < VolumeMeshes.Num(); ++VolumeIndex)
+				{
+					{
+						const IHeterogeneousVolumeInterface* HeterogeneousVolume = VolumeMeshes[VolumeIndex].Volume;
+						const FMaterialRenderProxy* MaterialRenderProxy = VolumeMeshes[VolumeIndex].MaterialRenderProxy;
+						const FPrimitiveSceneProxy* PrimitiveSceneProxy = HeterogeneousVolume->GetPrimitiveSceneProxy();
 						const FPrimitiveSceneInfo* PrimitiveSceneInfo = PrimitiveSceneProxy->GetPrimitiveSceneInfo();
 						const int32 PrimitiveId = PrimitiveSceneInfo->GetIndex();
 						const FBoxSphereBounds LocalBoxSphereBounds = HeterogeneousVolume->GetLocalBounds();
