Fix static analysis warning.

#rnx
#prefilght skip

diff --git a/Engine/Source/Runtime/Renderer/Private/LightRendering.cpp b/Engine/Source/Runtime/Renderer/Private/LightRendering.cpp
--- a/Engine/Source/Runtime/Renderer/Private/LightRendering.cpp
+++ b/Engine/Source/Runtime/Renderer/Private/LightRendering.cpp
@@ -1211,6 +1211,11 @@ void FDeferredShadingSceneRenderer::RenderLights(
 	FSortedLightSetSceneInfo& SortedLightSet)
 {
 	const bool bUseHairLighting = HairStrands::HasViewHairStrandsData(Views);
+#if RHI_RAYTRACING
+	const bool bEnableRayTracing = true;
+#else
+	const bool bEnableRayTracing = false;
+#endif // RHI_RAYTRACING
 
 	RDG_EVENT_SCOPE(GraphBuilder, "Lights");
 	RDG_GPU_STAT_SCOPE(GraphBuilder, Lights);
@@ -1384,7 +1389,7 @@ void FDeferredShadingSceneRenderer::RenderLights(
 			// Optimizations: batches all shadow ray tracing denoising. Definitely could be smarter to avoid high VGPR pressure if this entire
 			// function was converted to render graph, and want least intrusive change as possible. So right now it trades render target memory pressure
 			// for denoising perf.
-			if (RHI_RAYTRACING && bDoShadowBatching)
+			if (bEnableRayTracing && bDoShadowBatching)
 			{
 				const uint32 ViewIndex = 0;
 				FViewInfo& View = Views[ViewIndex];
@@ -1407,7 +1412,7 @@ void FDeferredShadingSceneRenderer::RenderLights(
 				{ 
 					PreprocessedShadowMaskSubPixelTextures.SetNum(NumShadowedLights);
 				}
-			} // if (RHI_RAYTRACING)
+			}
 
 			const bool bDirectLighting = ViewFamily.EngineShowFlags.DirectLighting;
 
@@ -1493,7 +1498,7 @@ void FDeferredShadingSceneRenderer::RenderLights(
 							SortedLightInfo.SortKey.Fields.bShadowed;
 
 						// determine if this light doesn't yet have a precomputed shadow and execute a batch to amortize costs if one is needed
-						if (RHI_RAYTRACING &&
+						if (bEnableRayTracing &&
 							bWantsBatchedShadow &&
 							(PreprocessedShadowMaskTextures.Num() == 0 || !PreprocessedShadowMaskTextures[LightIndex - UnbatchedLightStart]))
 						{
@@ -1695,7 +1700,7 @@ void FDeferredShadingSceneRenderer::RenderLights(
 						}
 					} // end inline batched raytraced shadow
 
-					if (RHI_RAYTRACING && PreprocessedShadowMaskTextures.Num() > 0 && PreprocessedShadowMaskTextures[LightIndex - UnbatchedLightStart])
+					if (bEnableRayTracing && PreprocessedShadowMaskTextures.Num() > 0 && PreprocessedShadowMaskTextures[LightIndex - UnbatchedLightStart])
 					{
 						const uint32 ShadowMaskIndex = LightIndex - UnbatchedLightStart;
 						ScreenShadowMaskTexture = PreprocessedShadowMaskTextures[ShadowMaskIndex];
