Allows to pass down a FScreenPassTextureSlice between temporal upscaler and tonemapper

This change is motivated to allow TSR history to be fully within a Texture2DArray without paying the extra memory bandwidth cost of also outputing a Texture2D for next post processing stages.
Currently tonemapper and motion blur allows to take a FScreenPassTextureSlice, but more fixes are deferred to followup CL.

This changes adds a FScreenPassTextureSlice which can be either a Texture2D or a the SRV on a slice of Texture2DArray, along with utils to convert from/to a FScreenPassTexture.
The motivation to introduce this FScreenPassTextureSlice is to explicitly say in code where this can actually be a Texture2DArray in code which should remain mostly rare in code.

Notes:
* Touches FTranslucentBasePassUniformParameters that invalidate some material shaders.

diff --git a/Engine/Source/Runtime/Renderer/Private/BasePassRendering.h b/Engine/Source/Runtime/Renderer/Private/BasePassRendering.h
--- a/Engine/Source/Runtime/Renderer/Private/BasePassRendering.h
+++ b/Engine/Source/Runtime/Renderer/Private/BasePassRendering.h
@@ -107,7 +107,7 @@ BEGIN_GLOBAL_SHADER_PARAMETER_STRUCT(FTranslucentBasePassUniformParameters,)
 	SHADER_PARAMETER(int32, SSRQuality)
 	SHADER_PARAMETER_RDG_TEXTURE(Texture2D, HZBTexture)
 	SHADER_PARAMETER_SAMPLER(SamplerState, HZBSampler)
-	SHADER_PARAMETER_RDG_TEXTURE(Texture2D, PrevSceneColor)
+	SHADER_PARAMETER_RDG_TEXTURE_SRV(Texture2D, PrevSceneColor)
 	SHADER_PARAMETER_SAMPLER(SamplerState, PrevSceneColorSampler)
 	// Volumetric cloud
 	SHADER_PARAMETER_RDG_TEXTURE(Texture2D, VolumetricCloudColor)
diff --git a/Engine/Source/Runtime/Renderer/Private/PostProcess/PostProcessMotionBlur.cpp b/Engine/Source/Runtime/Renderer/Private/PostProcess/PostProcessMotionBlur.cpp
--- a/Engine/Source/Runtime/Renderer/Private/PostProcess/PostProcessMotionBlur.cpp
+++ b/Engine/Source/Runtime/Renderer/Private/PostProcess/PostProcessMotionBlur.cpp
@@ -444,7 +444,7 @@ class FMotionBlurFilterCS : public FMotionBlurShader
 		SHADER_PARAMETER(int32, OutputMip1)
 		SHADER_PARAMETER(int32, OutputMip2)
 
-		SHADER_PARAMETER_RDG_TEXTURE(Texture2D, ColorTexture)
+		SHADER_PARAMETER_RDG_TEXTURE_SRV(Texture2D, ColorTexture)
 		SHADER_PARAMETER_RDG_TEXTURE(Texture2D, VelocityFlatTexture)
 		SHADER_PARAMETER_STRUCT(FVecocityTileTextureSRVs, VelocityTileTextures)
 
@@ -505,7 +505,7 @@ class FMotionBlurVisualizePS : public FMotionBlurShader
 	BEGIN_SHADER_PARAMETER_STRUCT(FParameters, )
 		SHADER_PARAMETER(FMatrix44f, WorldToClipPrev)
 		SHADER_PARAMETER_STRUCT_REF(FViewUniformShaderParameters, View)
-		SHADER_PARAMETER_RDG_TEXTURE(Texture2D, ColorTexture)
+		SHADER_PARAMETER_RDG_TEXTURE_SRV(Texture2D, ColorTexture)
 		SHADER_PARAMETER_RDG_TEXTURE(Texture2D, DepthTexture)
 		SHADER_PARAMETER_RDG_TEXTURE(Texture2D, VelocityTexture)
 
@@ -815,7 +815,7 @@ FMotionBlurOutputs AddMotionBlurFilterPass(
 	const FViewInfo& View,
 	const FMotionBlurInputs& Inputs,
 	const FMotionBlurViewports& Viewports,
-	FRDGTextureRef ColorTexture,
+	FRDGTextureSRVRef ColorTexture,
 	FRDGTextureRef VelocityFlatTexture,
 	FVecocityTileTextureSRVs VelocityTileTextures,
 	FRDGTextureRef PostMotionBlurTranslucency,
@@ -995,21 +995,23 @@ FMotionBlurOutputs AddMotionBlurFilterPass(
 
 		{
 			FRDGTextureDesc Desc = FRDGTextureDesc::Create2D(
-				ColorTexture->Desc.Extent,
+				ColorTexture->Desc.Texture->Desc.Extent,
 				IsPostProcessingWithAlphaChannelSupported() ? PF_FloatRGBA : PF_FloatRGB,
 				FClearValueBinding::None,
 				TexCreate_UAV | TexCreate_ShaderResource | GFastVRamConfig.MotionBlur);
 
-			Output.FullRes.Texture = GraphBuilder.CreateTexture(Desc, TEXT("MotionBlur.SceneColor"));
+			FRDGTextureRef FullResTexture = GraphBuilder.CreateTexture(Desc, TEXT("MotionBlur.SceneColor"));
+
+			Output.FullRes.TextureSRV = GraphBuilder.CreateSRV(FRDGTextureSRVDesc(FullResTexture));
 			Output.FullRes.ViewRect = Viewports.Color.Rect;
-			OriginalPassParameters.SceneColorOutputMip0 = GraphBuilder.CreateUAV(Output.FullRes.Texture, ERDGUnorderedAccessViewFlags::SkipBarrier);
+			OriginalPassParameters.SceneColorOutputMip0 = GraphBuilder.CreateUAV(FullResTexture, ERDGUnorderedAccessViewFlags::SkipBarrier);
 		}
 
 		if (OriginalPassParameters.OutputMip1)
 		{
 			FRDGTextureDesc Desc = FRDGTextureDesc::Create2D(
-				Output.FullRes.Texture->Desc.Extent / 2,
-				Output.FullRes.Texture->Desc.Format,
+				Output.FullRes.TextureSRV->Desc.Texture->Desc.Extent / 2,
+				Output.FullRes.TextureSRV->Desc.Texture->Desc.Format,
 				FClearValueBinding::None,
 				TexCreate_UAV | TexCreate_ShaderResource);
 
@@ -1034,8 +1036,8 @@ FMotionBlurOutputs AddMotionBlurFilterPass(
 		if (OriginalPassParameters.OutputMip2)
 		{
 			FRDGTextureDesc Desc = FRDGTextureDesc::Create2D(
-				Output.FullRes.Texture->Desc.Extent / 4,
-				Output.FullRes.Texture->Desc.Format,
+				Output.FullRes.TextureSRV->Desc.Texture->Desc.Extent / 4,
+				Output.FullRes.TextureSRV->Desc.Texture->Desc.Format,
 				FClearValueBinding::None,
 				TexCreate_UAV | TexCreate_ShaderResource);
 
@@ -1057,7 +1059,7 @@ FMotionBlurOutputs AddMotionBlurFilterPass(
 			GraphBuilder.RemoveUnusedTextureWarning(DummyTexture);
 		}
 
-		OriginalPassParameters.DebugOutput = CreateDebugUAV(GraphBuilder, Output.FullRes.Texture->Desc.Extent, TEXT("Debug.MotionBlur.Filter"));
+		OriginalPassParameters.DebugOutput = CreateDebugUAV(GraphBuilder, Output.FullRes.TextureSRV->Desc.Texture->Desc.Extent, TEXT("Debug.MotionBlur.Filter"));
 
 		RDG_EVENT_SCOPE(GraphBuilder, "MotionBlur FullResFilter(BlurDirections=%d MaxSamples=%d%s%s%s) %dx%d",
 			BlurDirections, OriginalPassParameters.MaxSampleCount,
@@ -1096,7 +1098,7 @@ FMotionBlurOutputs AddMotionBlurFilterPass(
 	return Output;
 }
 
-FScreenPassTexture AddVisualizeMotionBlurPass(FRDGBuilder& GraphBuilder, const FViewInfo& View, const FMotionBlurInputs& Inputs)
+FScreenPassTextureSlice AddVisualizeMotionBlurPass(FRDGBuilder& GraphBuilder, const FViewInfo& View, const FMotionBlurInputs& Inputs)
 {
 	check(Inputs.SceneColor.IsValid());
 	check(Inputs.SceneDepth.IsValid());
@@ -1107,10 +1109,12 @@ FScreenPassTexture AddVisualizeMotionBlurPass(FRDGBuilder& GraphBuilder, const F
 
 	if (!Output.IsValid())
 	{
-		FRDGTextureDesc OutputDesc = Inputs.SceneColor.Texture->Desc;
-		OutputDesc.Reset();
-		OutputDesc.NumMips = 1;
-
+		FRDGTextureDesc OutputDesc = FRDGTextureDesc::Create2D(
+			Inputs.SceneColor.TextureSRV->Desc.Texture->Desc.Extent,
+			Inputs.SceneColor.TextureSRV->Desc.Texture->Desc.Format,
+			FClearValueBinding::None,
+			ETextureCreateFlags::ShaderResource | ETextureCreateFlags::RenderTargetable);
+			
 		Output = FScreenPassRenderTarget(GraphBuilder.CreateTexture(OutputDesc, TEXT("MotionBlur.Visualize")), Inputs.SceneColor.ViewRect, View.GetOverwriteLoadAction());
 	}
 
@@ -1120,7 +1124,7 @@ FScreenPassTexture AddVisualizeMotionBlurPass(FRDGBuilder& GraphBuilder, const F
 	FMotionBlurVisualizePS::FParameters* PassParameters = GraphBuilder.AllocParameters<FMotionBlurVisualizePS::FParameters>();
 	PassParameters->WorldToClipPrev = FMatrix44f(GetPreviousWorldToClipMatrix(View));		// LWC_TODO: Precision loss
 	PassParameters->View = View.ViewUniformBuffer;
-	PassParameters->ColorTexture = Inputs.SceneColor.Texture;
+	PassParameters->ColorTexture = Inputs.SceneColor.TextureSRV;
 	PassParameters->DepthTexture = Inputs.SceneDepth.Texture;
 	PassParameters->VelocityTexture = Inputs.SceneVelocity.Texture;
 	PassParameters->Color = Viewports.ColorParameters;
@@ -1173,7 +1177,7 @@ FScreenPassTexture AddVisualizeMotionBlurPass(FRDGBuilder& GraphBuilder, const F
 		Canvas.DrawShadowedString(X + ColumnWidth, Y, *Line, GetStatsFont(), FLinearColor(1, 1, 0));
 	});
 
-	return MoveTemp(Output);
+	return FScreenPassTextureSlice::CreateFromScreenPassTexture(GraphBuilder, Output);
 }
 
 FMotionBlurOutputs AddMotionBlurPass(FRDGBuilder& GraphBuilder, const FViewInfo& View, const FMotionBlurInputs& Inputs)
@@ -1203,25 +1207,25 @@ FMotionBlurOutputs AddMotionBlurPass(FRDGBuilder& GraphBuilder, const FViewInfo&
 	FMotionBlurOutputs Output;
 	if (Inputs.Filter == EMotionBlurFilter::Separable)
 	{
-		FRDGTextureRef MotionBlurFilterTexture = AddMotionBlurFilterPass(
+		FRDGTextureSRVRef MotionBlurFilterTextureSRV = AddMotionBlurFilterPass(
 			GraphBuilder,
 			View,
 			Inputs,
 			Viewports,
-			Inputs.SceneColor.Texture,
+			Inputs.SceneColor.TextureSRV,
 			VelocityFlatTexture,
 			VelocityTileTextures,
 			nullptr,
 			FIntPoint(0, 0),
 			EMotionBlurFilterPass::Separable0,
-			Inputs.Quality).FullRes.Texture;
+			Inputs.Quality).FullRes.TextureSRV;
 
 		Output = AddMotionBlurFilterPass(
 			GraphBuilder,
 			View,
 			Inputs,
 			Viewports,
-			MotionBlurFilterTexture,
+			MotionBlurFilterTextureSRV,
 			VelocityFlatTexture,
 			VelocityTileTextures,
 			Inputs.PostMotionBlurTranslucency.ColorTexture.Resolve,
@@ -1236,7 +1240,7 @@ FMotionBlurOutputs AddMotionBlurPass(FRDGBuilder& GraphBuilder, const FViewInfo&
 			View,
 			Inputs,
 			Viewports,
-			Inputs.SceneColor.Texture,
+			Inputs.SceneColor.TextureSRV,
 			VelocityFlatTexture,
 			VelocityTileTextures,
 			Inputs.PostMotionBlurTranslucency.ColorTexture.Resolve,
diff --git a/Engine/Source/Runtime/Renderer/Private/PostProcess/PostProcessMotionBlur.h b/Engine/Source/Runtime/Renderer/Private/PostProcess/PostProcessMotionBlur.h
--- a/Engine/Source/Runtime/Renderer/Private/PostProcess/PostProcessMotionBlur.h
+++ b/Engine/Source/Runtime/Renderer/Private/PostProcess/PostProcessMotionBlur.h
@@ -79,7 +79,7 @@ struct FMotionBlurInputs
 	FScreenPassRenderTarget OverrideOutput;
 
 	// [Required] The input scene color and view rect.
-	FScreenPassTexture SceneColor;
+	FScreenPassTextureSlice SceneColor;
 
 	// [Required] The input scene depth and view rect.
 	FScreenPassTexture SceneDepth;
@@ -102,10 +102,10 @@ struct FMotionBlurInputs
 
 struct FMotionBlurOutputs
 {
-	FScreenPassTexture FullRes;
+	FScreenPassTextureSlice FullRes;
 	FScreenPassTexture HalfRes;
 	FScreenPassTexture QuarterRes;
 };
 
 FMotionBlurOutputs AddMotionBlurPass(FRDGBuilder& GraphBuilder, const FViewInfo& View, const FMotionBlurInputs& Inputs);
-FScreenPassTexture AddVisualizeMotionBlurPass(FRDGBuilder& GraphBuilder, const FViewInfo& View, const FMotionBlurInputs& Inputs);
\ No newline at end of file
+FScreenPassTextureSlice AddVisualizeMotionBlurPass(FRDGBuilder& GraphBuilder, const FViewInfo& View, const FMotionBlurInputs& Inputs);
\ No newline at end of file
diff --git a/Engine/Source/Runtime/Renderer/Private/PostProcess/PostProcessTonemap.cpp b/Engine/Source/Runtime/Renderer/Private/PostProcess/PostProcessTonemap.cpp
--- a/Engine/Source/Runtime/Renderer/Private/PostProcess/PostProcessTonemap.cpp
+++ b/Engine/Source/Runtime/Renderer/Private/PostProcess/PostProcessTonemap.cpp
@@ -313,7 +313,7 @@ BEGIN_SHADER_PARAMETER_STRUCT(FTonemapParameters, )
 	SHADER_PARAMETER_STRUCT(FScreenPassTextureViewportParameters, Color)
 	SHADER_PARAMETER_STRUCT(FScreenPassTextureViewportParameters, Output)
 	SHADER_PARAMETER_STRUCT(FEyeAdaptationParameters, EyeAdaptation)
-	SHADER_PARAMETER_RDG_TEXTURE(Texture2D, ColorTexture)
+	SHADER_PARAMETER_RDG_TEXTURE_SRV(Texture2D, ColorTexture)
 
 	// Parameters to apply to the scene color.
 	SHADER_PARAMETER_RDG_BUFFER_SRV(StructuredBuffer<float4>, SceneColorApplyParamaters)
@@ -574,13 +574,12 @@ FScreenPassTexture AddTonemapPass(FRDGBuilder& GraphBuilder, const FViewInfo& Vi
 
 	if (!Output.IsValid())
 	{
-		FRDGTextureDesc OutputDesc = Inputs.SceneColor.Texture->Desc;
-		OutputDesc.Reset();
-		OutputDesc.Flags |= View.bUseComputePasses ? TexCreate_UAV : TexCreate_RenderTargetable;
-		OutputDesc.Flags |= GFastVRamConfig.Tonemap;
-		// RGB is the color in LDR, A is the luminance for PostprocessAA
-		OutputDesc.ClearValue = FClearValueBinding(FLinearColor(0, 0, 0, 0));
-
+		FRDGTextureDesc OutputDesc = FRDGTextureDesc::Create2D(
+			SceneColorViewport.Extent,
+			Inputs.SceneColor.TextureSRV->Desc.Texture->Desc.Format,
+			FClearValueBinding(FLinearColor(0, 0, 0, 0)),
+			GFastVRamConfig.Tonemap | TexCreate_ShaderResource | (View.bUseComputePasses ? TexCreate_UAV : TexCreate_RenderTargetable));;
+		
 		const FTonemapperOutputDeviceParameters OutputDeviceParameters = GetTonemapperOutputDeviceParameters(*View.Family);
 		const EDisplayOutputFormat OutputDevice = static_cast<EDisplayOutputFormat>(OutputDeviceParameters.OutputDevice);
 
@@ -765,7 +764,7 @@ FScreenPassTexture AddTonemapPass(FRDGBuilder& GraphBuilder, const FViewInfo& Vi
 	CommonParameters.OutputDevice = GetTonemapperOutputDeviceParameters(ViewFamily);
 	CommonParameters.Color = GetScreenPassTextureViewportParameters(SceneColorViewport);
 	CommonParameters.Output = GetScreenPassTextureViewportParameters(OutputViewport);
-	CommonParameters.ColorTexture = Inputs.SceneColor.Texture;
+	CommonParameters.ColorTexture = Inputs.SceneColor.TextureSRV;
 	CommonParameters.LumBilateralGrid = Inputs.LocalExposureTexture;
 	CommonParameters.BlurredLogLum = Inputs.BlurredLogLuminanceTexture;
 	CommonParameters.LumBilateralGridSampler = BilinearClampSampler;
diff --git a/Engine/Source/Runtime/Renderer/Private/PostProcess/PostProcessTonemap.h b/Engine/Source/Runtime/Renderer/Private/PostProcess/PostProcessTonemap.h
--- a/Engine/Source/Runtime/Renderer/Private/PostProcess/PostProcessTonemap.h
+++ b/Engine/Source/Runtime/Renderer/Private/PostProcess/PostProcessTonemap.h
@@ -32,7 +32,7 @@ struct FTonemapInputs
 	FScreenPassRenderTarget OverrideOutput;
 
 	// [Required] HDR scene color to tonemap.
-	FScreenPassTexture SceneColor;
+	FScreenPassTextureSlice SceneColor;
 
 	// [Required] Filtered bloom texture to composite with tonemapped scene color. This should be transparent black for no bloom.
 	FScreenPassTexture Bloom;
diff --git a/Engine/Source/Runtime/Renderer/Private/PostProcess/PostProcessing.cpp b/Engine/Source/Runtime/Renderer/Private/PostProcess/PostProcessing.cpp
--- a/Engine/Source/Runtime/Renderer/Private/PostProcess/PostProcessing.cpp
+++ b/Engine/Source/Runtime/Renderer/Private/PostProcess/PostProcessing.cpp
@@ -247,7 +247,7 @@ FDefaultTemporalUpscaler::FOutputs AddThirdPartyTemporalUpscalerPasses(
 
 	// Translate the output.
 	FDefaultTemporalUpscaler::FOutputs Outputs;
-	Outputs.FullRes = ThirdPartyOutputs.FullRes;
+	Outputs.FullRes = FScreenPassTextureSlice::CreateFromScreenPassTexture(GraphBuilder, ThirdPartyOutputs.FullRes);
 
 	// Saves history for next frame.
 	if (!View.bStatePrevViewInfoIsReadOnly)
@@ -260,10 +260,10 @@ FDefaultTemporalUpscaler::FOutputs AddThirdPartyTemporalUpscalerPasses(
 	{
 		FTemporalAAHistory& OutputHistory = View.ViewState->PrevFrameViewInfo.TemporalAAHistory;
 
-		GraphBuilder.QueueTextureExtraction(Outputs.FullRes.Texture, &OutputHistory.RT[0]);
+		GraphBuilder.QueueTextureExtraction(ThirdPartyOutputs.FullRes.Texture, &OutputHistory.RT[0]);
 
-		OutputHistory.ViewportRect = Outputs.FullRes.ViewRect;
-		OutputHistory.ReferenceBufferSize = Outputs.FullRes.Texture->Desc.Extent;
+		OutputHistory.ViewportRect = ThirdPartyOutputs.FullRes.ViewRect;
+		OutputHistory.ReferenceBufferSize = ThirdPartyOutputs.FullRes.Texture->Desc.Extent;
 	}
 
 	return Outputs;
@@ -317,7 +317,7 @@ void AddPostProcessingPasses(
 	FScreenPassTexture SceneColor((*Inputs.SceneTextures)->SceneColorTexture, PrimaryViewRect);
 
 	// Assigned before and after the tonemapper.
-	FScreenPassTexture SceneColorBeforeTonemap;
+	FScreenPassTextureSlice SceneColorBeforeTonemapSlice;
 	FScreenPassTexture SceneColorAfterTonemap;
 
 	// Unprocessed scene color stores the original input.
@@ -518,7 +518,7 @@ void AddPostProcessingPasses(
 		return PostProcessMaterialInputs;
 	};
 
-	const auto AddAfterPass = [&](EPass InPass, FScreenPassTexture InSceneColor) -> FScreenPassTexture
+	const auto AddAfterPassForSceneColorSlice = [&](EPass InPass, const FScreenPassTextureSlice& InSceneColor) -> FScreenPassTextureSlice
 	{
 		// In some cases (e.g. OCIO color conversion) we want View Extensions to be able to add extra custom post processing after the pass.
 
@@ -526,19 +526,28 @@ void AddPostProcessingPasses(
 
 		if (PassCallbacks.Num())
 		{
-			FPostProcessMaterialInputs InOutPostProcessAfterPassInputs = GetPostProcessMaterialInputs(InSceneColor);
+			FScreenPassTexture SceneColor = FScreenPassTexture::CopyFromSlice(GraphBuilder, InSceneColor);
+
+			FPostProcessMaterialInputs InOutPostProcessAfterPassInputs = GetPostProcessMaterialInputs(SceneColor);
 
 			for (int32 AfterPassCallbackIndex = 0; AfterPassCallbackIndex < PassCallbacks.Num(); AfterPassCallbackIndex++)
 			{
-				InOutPostProcessAfterPassInputs.SetInput(EPostProcessMaterialInput::SceneColor, InSceneColor);
+				InOutPostProcessAfterPassInputs.SetInput(EPostProcessMaterialInput::SceneColor, SceneColor);
 
 				FAfterPassCallbackDelegate& AfterPassCallback = PassCallbacks[AfterPassCallbackIndex];
 				PassSequence.AcceptOverrideIfLastPass(InPass, InOutPostProcessAfterPassInputs.OverrideOutput, AfterPassCallbackIndex);
-				InSceneColor = AfterPassCallback.Execute(GraphBuilder, View, InOutPostProcessAfterPassInputs);
+				SceneColor = AfterPassCallback.Execute(GraphBuilder, View, InOutPostProcessAfterPassInputs);
 			}
+
+			return FScreenPassTextureSlice::CreateFromScreenPassTexture(GraphBuilder, SceneColor);
 		}
 
-		return MoveTemp(InSceneColor);
+		return InSceneColor;
+	};
+
+	const auto AddAfterPass = [&](EPass InPass, const FScreenPassTexture InSceneColor) -> FScreenPassTexture
+	{
+		return FScreenPassTexture(AddAfterPassForSceneColorSlice(InPass, FScreenPassTextureSlice::CreateFromScreenPassTexture(GraphBuilder, InSceneColor)));
 	};
 
 	if (bPostProcessingEnabled)
@@ -703,13 +712,17 @@ void AddPostProcessingPasses(
 			}
 		}
 
+		// Allows for the scene color to be the slice of an array between temporal upscaler and tonemaper.
+		FScreenPassTextureSlice SceneColorSlice = FScreenPassTextureSlice::CreateFromScreenPassTexture(GraphBuilder, SceneColor);
+		SceneColor = FScreenPassTexture();
+
 		FScreenPassTexture HalfResSceneColor;
 		FScreenPassTexture QuarterResSceneColor;
 		FVelocityFlattenTextures VelocityFlattenTextures;
 		if (TAAConfig != EMainTAAPassConfig::Disabled)
 		{
 			FDefaultTemporalUpscaler::FInputs UpscalerPassInputs;
-			UpscalerPassInputs.SceneColor = FScreenPassTexture(SceneColor.Texture, View.ViewRect);
+			UpscalerPassInputs.SceneColor = FScreenPassTexture(SceneColorSlice);
 			UpscalerPassInputs.SceneDepth = FScreenPassTexture(SceneDepth.Texture, View.ViewRect);
 			UpscalerPassInputs.SceneVelocity = FScreenPassTexture(Velocity.Texture, View.ViewRect);
 			if (PassSequence.IsEnabled(EPass::MotionBlur))
@@ -734,9 +747,11 @@ void AddPostProcessingPasses(
 					bNeedPostMotionBlurQuarterRes &&
 					DownsampleQuality == EDownsampleQuality::Low;
 			}
+			UpscalerPassInputs.bAllowFullResSlice = PassSequence.IsEnabled(EPass::MotionBlur) || PassSequence.IsEnabled(EPass::Tonemap);
 			UpscalerPassInputs.DownsampleOverrideFormat = DownsampleOverrideFormat;
 			UpscalerPassInputs.PostDOFTranslucencyResources = PostDOFTranslucencyResources;
 			UpscalerPassInputs.MoireInputTexture = TSRMoireInput;
+			check(UpscalerPassInputs.SceneColor.ViewRect == View.ViewRect);
 
 			FDefaultTemporalUpscaler::FOutputs Outputs;
 			if (TAAConfig == EMainTAAPassConfig::TSR)
@@ -765,7 +780,7 @@ void AddPostProcessingPasses(
 				unimplemented();
 			}
 
-			SceneColor = Outputs.FullRes;
+			SceneColorSlice = Outputs.FullRes;
 			HalfResSceneColor = Outputs.HalfRes;
 			QuarterResSceneColor = Outputs.QuarterRes;
 			VelocityFlattenTextures = Outputs.VelocityFlattenTextures;
@@ -786,28 +801,29 @@ void AddPostProcessingPasses(
 			{
 				check(View.ViewState);
 				FTemporalAAHistory& OutputHistory = View.ViewState->PrevFrameViewInfo.TemporalAAHistory;
-				GraphBuilder.QueueTextureExtraction(SceneColor.Texture, &OutputHistory.RT[0]);
+				GraphBuilder.QueueTextureExtraction(SceneColorSlice.TextureSRV->Desc.Texture, &OutputHistory.RT[0]);
 
 				// For SSR, we still fill up the rest of the OutputHistory data using shared math from FTAAPassParameters.
 				FTAAPassParameters TAAInputs(View);
-				TAAInputs.SceneColorInput = SceneColor.Texture;
+				TAAInputs.SceneColorInput = SceneColorSlice.TextureSRV->Desc.Texture;
 				TAAInputs.SetupViewRect(View);
 				OutputHistory.ViewportRect = TAAInputs.OutputViewRect;
 				OutputHistory.ReferenceBufferSize = TAAInputs.GetOutputExtent() * TAAInputs.ResolutionDivisor;
 			}
 		}
 
-		ensure(SceneColor.ViewRect.Size() == PostTAAViewSize);
+		ensure(SceneColorSlice.ViewRect.Size() == PostTAAViewSize);
 
 		// SVE/Post Process Material Chain - SSR Input
 		if (View.ViewState && !View.bStatePrevViewInfoIsReadOnly)
 		{
 			FScreenPassTexture PassOutput;
-			FPostProcessMaterialInputs InOutPassInputs = GetPostProcessMaterialInputs(SceneColor);
+			FPostProcessMaterialInputs InOutPassInputs;
 			const FPostProcessMaterialChain MaterialChain = GetPostProcessMaterialChain(View, BL_SSRInput);
 
 			if (MaterialChain.Num())
 			{
+				InOutPassInputs = GetPostProcessMaterialInputs(FScreenPassTexture::CopyFromSlice(GraphBuilder, SceneColorSlice));
 				PassOutput = AddPostProcessMaterialChain(GraphBuilder, View, InOutPassInputs, MaterialChain);
 			}
 
@@ -817,6 +833,10 @@ void AddPostProcessingPasses(
 				{
 					InOutPassInputs.SetInput(EPostProcessMaterialInput::SceneColor, PassOutput);
 				}
+				else
+				{
+					InOutPassInputs.SetInput(EPostProcessMaterialInput::SceneColor, FScreenPassTexture::CopyFromSlice(GraphBuilder, SceneColorSlice));
+				}
 				PassOutput = PassCallback.Execute(GraphBuilder, View, InOutPassInputs);
 			}
 
@@ -836,7 +856,7 @@ void AddPostProcessingPasses(
 			PassSequence.AcceptOverrideIfLastPass(EPass::MotionBlur, PassInputs.OverrideOutput);
 			PassInputs.bOutputHalfRes = bNeedPostMotionBlurHalfRes && DownsampleQuality == EDownsampleQuality::Low;
 			PassInputs.bOutputQuarterRes = bNeedPostMotionBlurQuarterRes && DownsampleQuality == EDownsampleQuality::Low;
-			PassInputs.SceneColor = SceneColor;
+			PassInputs.SceneColor = SceneColorSlice;
 			PassInputs.SceneDepth = SceneDepth;
 			PassInputs.SceneVelocity = Velocity;
 			PassInputs.PostMotionBlurTranslucency = PostMotionBlurTranslucencyResources;
@@ -847,12 +867,12 @@ void AddPostProcessingPasses(
 			// Motion blur visualization replaces motion blur when enabled.
 			if (bVisualizeMotionBlur)
 			{
-				SceneColor = AddVisualizeMotionBlurPass(GraphBuilder, View, PassInputs);
+				SceneColorSlice = AddVisualizeMotionBlurPass(GraphBuilder, View, PassInputs);
 			}
 			else
 			{
 				FMotionBlurOutputs PassOutputs = AddMotionBlurPass(GraphBuilder, View, PassInputs);
-				SceneColor = PassOutputs.FullRes;
+				SceneColorSlice = PassOutputs.FullRes;
 				HalfResSceneColor = PassOutputs.HalfRes;
 				QuarterResSceneColor = PassOutputs.QuarterRes;
 			}
@@ -862,25 +882,25 @@ void AddPostProcessingPasses(
 			// Compose Post-MotionBlur translucency in a new scene color to ensure it's not writing out in TAA's output that is also the history.
 			FTranslucencyComposition TranslucencyComposition;
 			TranslucencyComposition.Operation = FTranslucencyComposition::EOperation::ComposeToNewSceneColor;
-			TranslucencyComposition.SceneColor = SceneColor;
-			TranslucencyComposition.OutputViewport = FScreenPassTextureViewport(SceneColor);
+			TranslucencyComposition.SceneColor = FScreenPassTexture::CopyFromSlice(GraphBuilder, SceneColorSlice); // TODO
+			TranslucencyComposition.OutputViewport = FScreenPassTextureViewport(SceneColorSlice);
 			TranslucencyComposition.OutputPixelFormat = SceneColorFormat;
 
-			SceneColor = TranslucencyComposition.AddPass(
-				GraphBuilder, View, PostMotionBlurTranslucencyResources);
+			SceneColorSlice = FScreenPassTextureSlice::CreateFromScreenPassTexture(GraphBuilder, TranslucencyComposition.AddPass(
+				GraphBuilder, View, PostMotionBlurTranslucencyResources));
 		}
 
 		{
-			FScreenPassTexture NewSceneColor = AddAfterPass(EPass::MotionBlur, SceneColor);
+			FScreenPassTextureSlice NewSceneColorSlice = AddAfterPassForSceneColorSlice(EPass::MotionBlur, SceneColorSlice);
 
 			// Invalidate half and quarter res.
-			if (NewSceneColor != SceneColor)
+			if (NewSceneColorSlice != SceneColorSlice)
 			{
 				HalfResSceneColor = FScreenPassTexture();
 				QuarterResSceneColor = FScreenPassTexture();
 			}
 
-			SceneColor = NewSceneColor;
+			SceneColorSlice = NewSceneColorSlice;
 		}
 
 		// Generate post motion blur lower res scene color if they have not been generated.
@@ -890,7 +910,7 @@ void AddPostProcessingPasses(
 			{
 				FDownsamplePassInputs PassInputs;
 				PassInputs.Name = TEXT("PostProcessing.SceneColor.HalfRes");
-				PassInputs.SceneColor = SceneColor;
+				PassInputs.SceneColor = FScreenPassTexture(SceneColorSlice);
 				PassInputs.Quality = DownsampleQuality;
 				PassInputs.FormatOverride = DownsampleOverrideFormat;
 				PassInputs.UserSuppliedOutput = View.PrevViewInfo.HalfResTemporalAAHistory;
@@ -1019,7 +1039,7 @@ void AddPostProcessingPasses(
 				}
 				else
 				{
-					InputSceneColor = SceneColor;
+					InputSceneColor = FScreenPassTexture::CopyFromSlice(GraphBuilder, SceneColorSlice); // TODO
 					InputResolutionFraction = 1.0f;
 				}
 
@@ -1079,7 +1099,7 @@ void AddPostProcessingPasses(
 			}
 		}
 
-		SceneColorBeforeTonemap = SceneColor;
+		SceneColorBeforeTonemapSlice = SceneColorSlice;
 
 		if (PassSequence.IsEnabled(EPass::Tonemap))
 		{
@@ -1091,7 +1111,7 @@ void AddPostProcessingPasses(
 
 				FPostProcessMaterialInputs PassInputs;
 				PassSequence.AcceptOverrideIfLastPass(EPass::Tonemap, PassInputs.OverrideOutput);
-				PassInputs.SetInput(EPostProcessMaterialInput::SceneColor, SceneColor);
+				PassInputs.SetInput(EPostProcessMaterialInput::SceneColor, FScreenPassTexture::CopyFromSlice(GraphBuilder, SceneColorSlice));
 				//PassInputs.SetInput(EPostProcessMaterialInput::SeparateTranslucency, SeparateTranslucency);
 				PassInputs.SetInput(EPostProcessMaterialInput::CombinedBloom, Bloom);
 				PassInputs.SceneTextures = GetSceneTextureShaderParameters(Inputs.SceneTextures);
@@ -1121,7 +1141,7 @@ void AddPostProcessingPasses(
 
 				FTonemapInputs PassInputs;
 				PassSequence.AcceptOverrideIfLastPass(EPass::Tonemap, PassInputs.OverrideOutput);
-				PassInputs.SceneColor = SceneColor;
+				PassInputs.SceneColor = SceneColorSlice;
 				PassInputs.Bloom = Bloom;
 				PassInputs.SceneColorApplyParamaters = SceneColorApplyParameters;
 				PassInputs.LocalExposureTexture = LocalExposureTexture;
@@ -1135,6 +1155,10 @@ void AddPostProcessingPasses(
 				SceneColor = AddTonemapPass(GraphBuilder, View, PassInputs);
 			}
 		}
+		else
+		{
+			SceneColor = FScreenPassTexture(SceneColorSlice);
+		}
 		
 		SceneColor = AddAfterPass(EPass::Tonemap, SceneColor);
 
@@ -1157,7 +1181,7 @@ void AddPostProcessingPasses(
 		{
 			FPostProcessMaterialInputs PassInputs = GetPostProcessMaterialInputs(SceneColor);
 			PassSequence.AcceptOverrideIfLastPass(EPass::PostProcessMaterialAfterTonemapping, PassInputs.OverrideOutput);
-			PassInputs.SetInput(EPostProcessMaterialInput::PreTonemapHDRColor, SceneColorBeforeTonemap);
+			PassInputs.SetInput(EPostProcessMaterialInput::PreTonemapHDRColor, FScreenPassTexture::CopyFromSlice(GraphBuilder, SceneColorBeforeTonemapSlice));
 			PassInputs.SetInput(EPostProcessMaterialInput::PostTonemapHDRColor, SceneColorAfterTonemap);
 			PassInputs.SceneTextures = GetSceneTextureShaderParameters(Inputs.SceneTextures);
 
@@ -1223,13 +1247,13 @@ void AddPostProcessingPasses(
 			SceneColor = TranslucencyComposition.SceneColor;
 		}
 
-		SceneColorBeforeTonemap = SceneColor;
+		SceneColorBeforeTonemapSlice = FScreenPassTextureSlice::CreateFromScreenPassTexture(GraphBuilder, SceneColor);
 
 		if (PassSequence.IsEnabled(EPass::Tonemap))
 		{
 			FTonemapInputs PassInputs;
 			PassSequence.AcceptOverrideIfLastPass(EPass::Tonemap, PassInputs.OverrideOutput);
-			PassInputs.SceneColor = SceneColor;
+			PassInputs.SceneColor = FScreenPassTextureSlice::CreateFromScreenPassTexture(GraphBuilder, SceneColor);
 			PassInputs.EyeAdaptationParameters = &EyeAdaptationParameters;
 			PassInputs.EyeAdaptationBuffer = EyeAdaptationBuffer;
 			PassInputs.bOutputInHDR = bViewFamilyOutputInHDR;
@@ -1383,7 +1407,7 @@ void AddPostProcessingPasses(
 		FVisualizeGBufferOverviewInputs PassInputs;
 		PassSequence.AcceptOverrideIfLastPass(EPass::VisualizeGBufferOverview, PassInputs.OverrideOutput);
 		PassInputs.SceneColor = SceneColor;
-		PassInputs.SceneColorBeforeTonemap = SceneColorBeforeTonemap;
+		PassInputs.SceneColorBeforeTonemap = FScreenPassTexture::CopyFromSlice(GraphBuilder, SceneColorBeforeTonemapSlice);
 		PassInputs.SceneColorAfterTonemap = SceneColorAfterTonemap;
 
 		FIntRect ViewRect{ 0, 0, 1, 1 };
@@ -1421,7 +1445,7 @@ void AddPostProcessingPasses(
 		FVisualizeHDRInputs PassInputs;
 		PassSequence.AcceptOverrideIfLastPass(EPass::VisualizeHDR, PassInputs.OverrideOutput);
 		PassInputs.SceneColor = SceneColor;
-		PassInputs.SceneColorBeforeTonemap = SceneColorBeforeTonemap;
+		PassInputs.SceneColorBeforeTonemap = FScreenPassTexture::CopyFromSlice(GraphBuilder, SceneColorBeforeTonemapSlice);
 		PassInputs.HistogramTexture = HistogramTexture;
 		PassInputs.EyeAdaptationBuffer = GetEyeAdaptationBuffer(GraphBuilder, View);
 		PassInputs.EyeAdaptationParameters = &EyeAdaptationParameters;
@@ -1434,7 +1458,7 @@ void AddPostProcessingPasses(
 		FVisualizeLocalExposureInputs PassInputs;
 		PassSequence.AcceptOverrideIfLastPass(EPass::VisualizeLocalExposure, PassInputs.OverrideOutput);
 		PassInputs.SceneColor = SceneColor;
-		PassInputs.HDRSceneColor = SceneColorBeforeTonemap;
+		PassInputs.HDRSceneColor = FScreenPassTexture::CopyFromSlice(GraphBuilder, SceneColorBeforeTonemapSlice);
 		PassInputs.LumBilateralGridTexture = LocalExposureTexture;
 		PassInputs.BlurredLumTexture = LocalExposureBlurredLogLumTexture;
 		PassInputs.EyeAdaptationBuffer = GetEyeAdaptationBuffer(GraphBuilder, View);
@@ -1468,7 +1492,7 @@ void AddPostProcessingPasses(
 		FPixelInspectorInputs PassInputs;
 		PassSequence.AcceptOverrideIfLastPass(EPass::PixelInspector, PassInputs.OverrideOutput);
 		PassInputs.SceneColor = SceneColor;
-		PassInputs.SceneColorBeforeTonemap = SceneColorBeforeTonemap;
+		PassInputs.SceneColorBeforeTonemap = FScreenPassTexture::CopyFromSlice(GraphBuilder, SceneColorBeforeTonemapSlice);
 		PassInputs.OriginalSceneColor = OriginalSceneColor;
 
 		SceneColor = AddPixelInspectorPass(GraphBuilder, View, PassInputs);
@@ -1672,7 +1696,7 @@ void AddDebugViewPostProcessingPasses(FRDGBuilder& GraphBuilder, const FViewInfo
 		const FEyeAdaptationParameters EyeAdaptationParameters = GetEyeAdaptationParameters(View, ERHIFeatureLevel::SM5);
 
 		FTonemapInputs PassInputs;
-		PassInputs.SceneColor = SceneColor;
+		PassInputs.SceneColor = FScreenPassTextureSlice::CreateFromScreenPassTexture(GraphBuilder, SceneColor);
 		PassInputs.bOutputInHDR = bViewFamilyOutputInHDR;
 		PassInputs.bGammaOnly = true;
 		PassInputs.EyeAdaptationParameters = &EyeAdaptationParameters;
@@ -1766,7 +1790,7 @@ void AddDebugViewPostProcessingPasses(FRDGBuilder& GraphBuilder, const FViewInfo
 
 		FTonemapInputs PassInputs;
 		PassSequence.AcceptOverrideIfLastPass(EPass::TonemapAfter, PassInputs.OverrideOutput);
-		PassInputs.SceneColor = SceneColor;
+		PassInputs.SceneColor = FScreenPassTextureSlice::CreateFromScreenPassTexture(GraphBuilder, SceneColor);
 		PassInputs.bOutputInHDR = bViewFamilyOutputInHDR;
 		PassInputs.bGammaOnly = true;
 		// Do eye adaptation in ray tracing debug modes to match raster buffer visualization modes
@@ -2464,7 +2488,7 @@ void AddMobilePostProcessingPasses(FRDGBuilder& GraphBuilder, FScene* Scene, con
 			{
 				unimplemented();
 			}
-			SceneColor = Outputs.FullRes;
+			SceneColor = FScreenPassTexture(Outputs.FullRes);
 		}
 	}
 	else
@@ -2532,7 +2556,7 @@ void AddMobilePostProcessingPasses(FRDGBuilder& GraphBuilder, FScene* Scene, con
 			TonemapperInputs.OverrideOutput.LoadAction = OutputLoadAction;
 		}
 			
-		TonemapperInputs.SceneColor = SceneColor;
+		TonemapperInputs.SceneColor = FScreenPassTextureSlice::CreateFromScreenPassTexture(GraphBuilder, SceneColor);
 		TonemapperInputs.Bloom = BloomOutput;
 		TonemapperInputs.EyeAdaptationParameters = &EyeAdaptationParameters;
 		TonemapperInputs.ColorGradingTexture = ColorGradingTexture;
diff --git a/Engine/Source/Runtime/Renderer/Private/PostProcess/TemporalAA.cpp b/Engine/Source/Runtime/Renderer/Private/PostProcess/TemporalAA.cpp
--- a/Engine/Source/Runtime/Renderer/Private/PostProcess/TemporalAA.cpp
+++ b/Engine/Source/Runtime/Renderer/Private/PostProcess/TemporalAA.cpp
@@ -584,7 +584,7 @@ FTAAOutputs AddTemporalAAPass(
 	const bool bIsMainPass = IsMainTAAConfig(Inputs.Pass);
 
 	// Whether to use camera cut shader permutation or not.
-	const bool bCameraCut = !InputHistory.IsValid() || View.bCameraCut;
+	const bool bCameraCut = !InputHistory.IsValid() || InputHistory.OutputSliceIndex != 0 || View.bCameraCut;
 
 	const FIntPoint OutputExtent = Inputs.GetOutputExtent();
 
@@ -1067,8 +1067,7 @@ FDefaultTemporalUpscaler::FOutputs AddGen4MainTemporalAAPasses(
 	}
 
 	FDefaultTemporalUpscaler::FOutputs Outputs;
-	Outputs.FullRes.Texture = SceneColorTexture;
-	Outputs.FullRes.ViewRect = SecondaryViewRect;
+	Outputs.FullRes = FScreenPassTextureSlice::CreateFromScreenPassTexture(GraphBuilder, FScreenPassTexture(SceneColorTexture, SecondaryViewRect));
 	Outputs.HalfRes.Texture = TAAOutputs.DownsampledSceneColor;
 	Outputs.HalfRes.ViewRect = FIntRect::DivideAndRoundUp(SecondaryViewRect, 2);
 	return Outputs;
diff --git a/Engine/Source/Runtime/Renderer/Private/PostProcess/TemporalAA.h b/Engine/Source/Runtime/Renderer/Private/PostProcess/TemporalAA.h
--- a/Engine/Source/Runtime/Renderer/Private/PostProcess/TemporalAA.h
+++ b/Engine/Source/Runtime/Renderer/Private/PostProcess/TemporalAA.h
@@ -187,6 +187,7 @@ struct FDefaultTemporalUpscaler
 {
 	struct FInputs
 	{
+		bool bAllowFullResSlice = false;
 		bool bGenerateSceneColorHalfRes = false;
 		bool bGenerateSceneColorQuarterRes = false;
 		bool bGenerateOutputMip1 = false;
@@ -201,7 +202,7 @@ struct FDefaultTemporalUpscaler
 
 	struct FOutputs
 	{
-		FScreenPassTexture FullRes;
+		FScreenPassTextureSlice FullRes;
 		FScreenPassTexture HalfRes;
 		FScreenPassTexture QuarterRes;
 		FVelocityFlattenTextures VelocityFlattenTextures;
diff --git a/Engine/Source/Runtime/Renderer/Private/PostProcess/TemporalSuperResolution.cpp b/Engine/Source/Runtime/Renderer/Private/PostProcess/TemporalSuperResolution.cpp
--- a/Engine/Source/Runtime/Renderer/Private/PostProcess/TemporalSuperResolution.cpp
+++ b/Engine/Source/Runtime/Renderer/Private/PostProcess/TemporalSuperResolution.cpp
@@ -2046,7 +2046,7 @@ FDefaultTemporalUpscaler::FOutputs AddTemporalSuperResolutionPasses(
 
 
 	FDefaultTemporalUpscaler::FOutputs Outputs;
-	Outputs.FullRes.Texture = SceneColorOutputTexture;
+	Outputs.FullRes.TextureSRV = GraphBuilder.CreateSRV(FRDGTextureSRVDesc(SceneColorOutputTexture));
 	Outputs.FullRes.ViewRect = OutputRect;
 	if (SceneColorOutputHalfResTexture)
 	{
diff --git a/Engine/Source/Runtime/Renderer/Private/PostProcess/VisualizeTemporalUpscaler.cpp b/Engine/Source/Runtime/Renderer/Private/PostProcess/VisualizeTemporalUpscaler.cpp
--- a/Engine/Source/Runtime/Renderer/Private/PostProcess/VisualizeTemporalUpscaler.cpp
+++ b/Engine/Source/Runtime/Renderer/Private/PostProcess/VisualizeTemporalUpscaler.cpp
@@ -33,6 +33,8 @@ FScreenPassTexture AddVisualizeTemporalUpscalerPass(FRDGBuilder& GraphBuilder, c
 
 	if (Inputs.TAAConfig != EMainTAAPassConfig::Disabled)
 	{
+		FScreenPassTexture OutputTexture = FScreenPassTexture::CopyFromSlice(GraphBuilder, Inputs.Outputs.FullRes);
+
 		auto VisualizeTextureLabel = [](FRDGTextureRef Texture, const TCHAR* Suffix = TEXT(""))
 		{
 			return FString::Printf(TEXT("vis %s%s"), Texture->Name, Suffix);
@@ -108,7 +110,7 @@ FScreenPassTexture AddVisualizeTemporalUpscalerPass(FRDGBuilder& GraphBuilder, c
 			PassInputs.SceneVelocity = Inputs.Inputs.SceneVelocity;
 
 			FVisualizeBufferTile& Tile = Tiles[4 * 1 + 0];
-			Tile.Input = AddVisualizeMotionBlurPass(GraphBuilder, View, PassInputs);
+			Tile.Input = FScreenPassTexture(AddVisualizeMotionBlurPass(GraphBuilder, View, PassInputs));
 			Tile.Input.ViewRect = CropViewRectToCenter(Tile.Input.ViewRect);
 			Tile.Label = TEXT("show VisualizeMotionBlur");
 		}
@@ -138,18 +140,18 @@ FScreenPassTexture AddVisualizeTemporalUpscalerPass(FRDGBuilder& GraphBuilder, c
 		// Output
 		{
 			FVisualizeBufferTile& Tile = Tiles[4 * 3 + 3];
-			Tile.Input = Inputs.Outputs.FullRes;
+			Tile.Input.Texture = OutputTexture.Texture;
 			Tile.Input.ViewRect = CropViewRectToCenter(Tile.Input.ViewRect);
-			Tile.Label = VisualizeTextureLabel(Inputs.Outputs.FullRes.Texture);
+			Tile.Label = VisualizeTextureLabel(OutputTexture.Texture);
 		}
 
 		// Output alpha
 		if (bSupportsAlpha)
 		{
 			FVisualizeBufferTile& Tile = Tiles[4 * 3 + 2];
-			Tile.Input.Texture = FVisualizeTexture::AddVisualizeTextureAlphaPass(GraphBuilder, View.ShaderMap, Inputs.Outputs.FullRes.Texture);
-			Tile.Input.ViewRect = CropViewRectToCenter(Inputs.Outputs.FullRes.ViewRect);
-			Tile.Label = VisualizeTextureLabel(Inputs.Outputs.FullRes.Texture, TEXT(" A"));
+			Tile.Input.Texture = FVisualizeTexture::AddVisualizeTextureAlphaPass(GraphBuilder, View.ShaderMap, OutputTexture.Texture);
+			Tile.Input.ViewRect = CropViewRectToCenter(OutputTexture.ViewRect);
+			Tile.Label = VisualizeTextureLabel(OutputTexture.Texture, TEXT(" A"));
 		}
 
 		// Black bottom left corner
@@ -210,7 +212,7 @@ FScreenPassTexture AddVisualizeTemporalUpscalerPass(FRDGBuilder& GraphBuilder, c
 			// Display the input/output resolutions
 			{
 				QuickDrawSummary(/* Location = */ 1, FString::Printf(TEXT("Input: %dx%d %s"), View.ViewRect.Width(), View.ViewRect.Height(), GPixelFormats[Inputs.Inputs.SceneColor.Texture->Desc.Format].Name));
-				QuickDrawSummary(/* Location = */ 2, FString::Printf(TEXT("Output: %dx%d %s"), Inputs.Outputs.FullRes.ViewRect.Width(), Inputs.Outputs.FullRes.ViewRect.Height(), GPixelFormats[Inputs.Outputs.FullRes.Texture->Desc.Format].Name));
+				QuickDrawSummary(/* Location = */ 2, FString::Printf(TEXT("Output: %dx%d %s"), Inputs.Outputs.FullRes.ViewRect.Width(), Inputs.Outputs.FullRes.ViewRect.Height(), GPixelFormats[Inputs.Outputs.FullRes.TextureSRV->Desc.Texture->Desc.Format].Name));
 			}
 
 			// Display the pre-exposure being used
diff --git a/Engine/Source/Runtime/Renderer/Private/SceneRendering.h b/Engine/Source/Runtime/Renderer/Private/SceneRendering.h
--- a/Engine/Source/Runtime/Renderer/Private/SceneRendering.h
+++ b/Engine/Source/Runtime/Renderer/Private/SceneRendering.h
@@ -973,6 +973,9 @@ struct FTemporalAAHistory
 	// Viewport coordinate of the history in RT according to ReferenceBufferSize.
 	FIntRect ViewportRect;
 
+	// Returns the slice index that contains the output in RT[0].
+	int32 OutputSliceIndex = 0;
+
 
 	void SafeRelease()
 	{
diff --git a/Engine/Source/Runtime/Renderer/Private/ScreenPass.cpp b/Engine/Source/Runtime/Renderer/Private/ScreenPass.cpp
--- a/Engine/Source/Runtime/Renderer/Private/ScreenPass.cpp
+++ b/Engine/Source/Runtime/Renderer/Private/ScreenPass.cpp
@@ -42,6 +42,47 @@ FRDGTextureRef TryCreateViewFamilyTexture(FRDGBuilder& GraphBuilder, const FScen
 	return Texture;
 }
 
+// static
+FScreenPassTexture FScreenPassTexture::CopyFromSlice(FRDGBuilder& GraphBuilder, const FScreenPassTextureSlice& ScreenTextureSlice)
+{
+	if (!ScreenTextureSlice.TextureSRV)
+	{
+		return FScreenPassTexture(nullptr, ScreenTextureSlice.ViewRect);
+	}
+	else if (!ScreenTextureSlice.TextureSRV->Desc.Texture->Desc.IsTextureArray())
+	{
+		return FScreenPassTexture(ScreenTextureSlice.TextureSRV->Desc.Texture, ScreenTextureSlice.ViewRect);
+	}
+
+	FRDGTextureDesc Desc = ScreenTextureSlice.TextureSRV->Desc.Texture->Desc;
+	Desc.Dimension = ETextureDimension::Texture2D;
+
+	FRDGTextureRef NewTexture = GraphBuilder.CreateTexture(Desc, TEXT("CopyToScreenPassTexture2D"));
+
+	FRHICopyTextureInfo CopyInfo;
+	CopyInfo.SourceSliceIndex = ScreenTextureSlice.TextureSRV->Desc.FirstArraySlice;
+
+	AddCopyTexturePass(
+		GraphBuilder,
+		ScreenTextureSlice.TextureSRV->Desc.Texture,
+		NewTexture,
+		CopyInfo);
+
+	return FScreenPassTexture(NewTexture, ScreenTextureSlice.ViewRect);
+}
+
+// static
+FScreenPassTextureSlice FScreenPassTextureSlice::CreateFromScreenPassTexture(FRDGBuilder& GraphBuilder, const FScreenPassTexture& ScreenTexture)
+{
+	if (!ScreenTexture.Texture)
+	{
+		return FScreenPassTextureSlice(nullptr, ScreenTexture.ViewRect);
+	}
+
+	check(!ScreenTexture.Texture->Desc.IsTextureArray());
+	return FScreenPassTextureSlice(GraphBuilder.CreateSRV(FRDGTextureSRVDesc(ScreenTexture.Texture)), ScreenTexture.ViewRect);
+}
+
 FScreenPassRenderTarget FScreenPassRenderTarget::CreateFromInput(
 	FRDGBuilder& GraphBuilder,
 	FScreenPassTexture Input,
diff --git a/Engine/Source/Runtime/Renderer/Private/ScreenSpaceRayTracing.cpp b/Engine/Source/Runtime/Renderer/Private/ScreenSpaceRayTracing.cpp
--- a/Engine/Source/Runtime/Renderer/Private/ScreenSpaceRayTracing.cpp
+++ b/Engine/Source/Runtime/Renderer/Private/ScreenSpaceRayTracing.cpp
@@ -351,7 +351,7 @@ BEGIN_SHADER_PARAMETER_STRUCT(FSSRPassCommonParameters, )
 	SHADER_PARAMETER(FVector4f, PrevScreenPositionScaleBias)
 	SHADER_PARAMETER(float, PrevSceneColorPreExposureCorrection)
 	SHADER_PARAMETER(uint32, ShouldReflectOnlyWater)
-	SHADER_PARAMETER_RDG_TEXTURE(Texture2D, SceneColor)
+	SHADER_PARAMETER_RDG_TEXTURE_SRV(Texture2D, SceneColor)
 	SHADER_PARAMETER_SAMPLER(SamplerState, SceneColorSampler)
 	SHADER_PARAMETER_RDG_TEXTURE(Texture2D, HZB)
 	SHADER_PARAMETER_SAMPLER(SamplerState, HZBSampler)
@@ -396,7 +396,7 @@ class FSSRTPrevFrameReductionCS : public FGlobalShader
 		SHADER_PARAMETER(float, HigherMipDownScaleFactor)
 		SHADER_PARAMETER(float, SkyDistance)
 		
-		SHADER_PARAMETER_RDG_TEXTURE(Texture2D, PrevSceneColor)
+		SHADER_PARAMETER_RDG_TEXTURE_SRV(Texture2D, PrevSceneColor)
 		SHADER_PARAMETER_SAMPLER(SamplerState, PrevSceneColorSampler)
 
 		SHADER_PARAMETER_RDG_TEXTURE(Texture2D, PrevSceneDepth)
@@ -743,7 +743,8 @@ FPrevSceneColorMip ReducePrevSceneColorMip(
 			ViewportOffset = View.PrevViewInfo.ViewRect.Min;
 			ViewportExtent = View.PrevViewInfo.ViewRect.Size();
 
-			PassParameters->PrevSceneColor = GraphBuilder.RegisterExternalTexture(View.PrevViewInfo.ScreenSpaceRayTracingInput);
+			PassParameters->PrevSceneColor = GraphBuilder.CreateSRV(FRDGTextureSRVDesc(
+				GraphBuilder.RegisterExternalTexture(View.PrevViewInfo.ScreenSpaceRayTracingInput)));
 			PassParameters->PrevSceneColorSampler = TStaticSamplerState<SF_Point>::GetRHI();
 
 			PassParameters->PrevSceneDepth = GraphBuilder.RegisterExternalTexture(View.PrevViewInfo.DepthBuffer);
@@ -755,8 +756,11 @@ FPrevSceneColorMip ReducePrevSceneColorMip(
 			ViewportOffset = View.PrevViewInfo.TemporalAAHistory.ViewportRect.Min;
 			ViewportExtent = View.PrevViewInfo.TemporalAAHistory.ViewportRect.Size();
 
-			PassParameters->PrevSceneColor = GraphBuilder.RegisterExternalTexture(View.PrevViewInfo.TemporalAAHistory.RT[0]);
+			FRDGTextureRef TemporalAAHistoryTexture = GraphBuilder.RegisterExternalTexture(View.PrevViewInfo.TemporalAAHistory.RT[0]);
 			PassParameters->PrevSceneColorSampler = TStaticSamplerState<SF_Bilinear>::GetRHI();
+			PassParameters->PrevSceneColor = GraphBuilder.CreateSRV(TemporalAAHistoryTexture->Desc.IsTextureArray()
+				? FRDGTextureSRVDesc::CreateForSlice(TemporalAAHistoryTexture, View.PrevViewInfo.TemporalAAHistory.OutputSliceIndex)
+				: FRDGTextureSRVDesc(TemporalAAHistoryTexture));
 		}
 
 		float InvBufferSizeX = 1.f / float(BufferSize.X);
@@ -941,20 +945,25 @@ void RenderScreenSpaceReflections(
 	bool bSingleLayerWater,
 	FTiledReflection* TiledScreenSpaceReflection)
 {
-	FRDGTextureRef InputColor = CurrentSceneColor;
+	FRDGTextureSRVRef InputColor = GraphBuilder.CreateSRV(FRDGTextureSRVDesc(CurrentSceneColor));
 	if (SSRQuality != ESSRQuality::VisualizeSSR)
 	{
 		if (View.PrevViewInfo.CustomSSRInput.IsValid())
 		{
-			InputColor = GraphBuilder.RegisterExternalTexture(View.PrevViewInfo.CustomSSRInput.RT[0]);
+			InputColor = GraphBuilder.CreateSRV(FRDGTextureSRVDesc(
+				GraphBuilder.RegisterExternalTexture(View.PrevViewInfo.CustomSSRInput.RT[0])));
 		}
 		else if (GSSRHalfResSceneColor && View.PrevViewInfo.HalfResTemporalAAHistory.IsValid())
 		{
-			InputColor = GraphBuilder.RegisterExternalTexture(View.PrevViewInfo.HalfResTemporalAAHistory);
+			InputColor = GraphBuilder.CreateSRV(FRDGTextureSRVDesc(
+				GraphBuilder.RegisterExternalTexture(View.PrevViewInfo.HalfResTemporalAAHistory)));
 		}
 		else if (View.PrevViewInfo.TemporalAAHistory.IsValid())
 		{
-			InputColor = GraphBuilder.RegisterExternalTexture(View.PrevViewInfo.TemporalAAHistory.RT[0]);
+			FRDGTextureRef TemporalAAHistoryTexture = GraphBuilder.RegisterExternalTexture(View.PrevViewInfo.TemporalAAHistory.RT[0]);
+			InputColor = GraphBuilder.CreateSRV(TemporalAAHistoryTexture->Desc.IsTextureArray()
+				? FRDGTextureSRVDesc::CreateForSlice(TemporalAAHistoryTexture, View.PrevViewInfo.TemporalAAHistory.OutputSliceIndex)
+				: FRDGTextureSRVDesc(TemporalAAHistoryTexture));
 		}
 	}
 
@@ -986,7 +995,7 @@ void RenderScreenSpaceReflections(
 	CommonParameters.ViewUniformBuffer = View.ViewUniformBuffer;
 	CommonParameters.SceneTextures = SceneTextures;
 	// Pipe down a mid grey texture when not using TAA's history to avoid wrongly reprojecting current scene color as if previous frame's TAA history.
-	if (InputColor == CurrentSceneColor || !CommonParameters.SceneTextures.GBufferVelocityTexture)
+	if (InputColor->Desc.Texture == CurrentSceneColor || !CommonParameters.SceneTextures.GBufferVelocityTexture)
 	{
 		// Technically should be 32767.0f / 65535.0f to perfectly null out DecodeVelocityFromTexture(), but 0.5f is good enough.
 		CommonParameters.SceneTextures.GBufferVelocityTexture = GraphBuilder.RegisterExternalTexture(GSystemTextures.MidGreyDummy);
@@ -1088,7 +1097,7 @@ void RenderScreenSpaceReflections(
 
 			PassParameters->ScreenSpaceRayTracingDebugOutput = CreateScreenSpaceRayTracingDebugUAV(GraphBuilder, DenoiserInputs->Color->Desc, TEXT("DebugSSR"), true);
 		}
-		PassParameters->PrevSceneColorPreExposureCorrection = InputColor != CurrentSceneColor ? View.PreExposure / View.PrevViewInfo.SceneColorPreExposure : 1.0f;
+		PassParameters->PrevSceneColorPreExposureCorrection = InputColor->Desc.Texture != CurrentSceneColor ? View.PreExposure / View.PrevViewInfo.SceneColorPreExposure : 1.0f;
 		PassParameters->ShouldReflectOnlyWater = bSingleLayerWater ? 1u : 0u;
 		
 		PassParameters->SceneColor = InputColor;
diff --git a/Engine/Source/Runtime/Renderer/Private/TranslucentRendering.cpp b/Engine/Source/Runtime/Renderer/Private/TranslucentRendering.cpp
--- a/Engine/Source/Runtime/Renderer/Private/TranslucentRendering.cpp
+++ b/Engine/Source/Runtime/Renderer/Private/TranslucentRendering.cpp
@@ -841,24 +841,27 @@ TRDGUniformBufferRef<FTranslucentBasePassUniformParameters> CreateTranslucentBas
 			BasePassParameters.HZBTexture = View.HZB;
 			BasePassParameters.HZBSampler = TStaticSamplerState<SF_Point, AM_Clamp, AM_Clamp, AM_Clamp>::GetRHI();
 
-			FRDGTextureRef PrevSceneColorTexture = SystemTextures.Black;
+			FRDGTextureSRVRef PrevSceneColorTexture = GraphBuilder.CreateSRV(FRDGTextureSRVDesc(SystemTextures.Black));
 			FIntRect PrevSceneColorViewRect = FIntRect(0, 0, 1, 1);
 
 			if (View.PrevViewInfo.CustomSSRInput.IsValid())
 			{
-				PrevSceneColorTexture = GetRDG(View.PrevViewInfo.CustomSSRInput.RT[0]);
+				PrevSceneColorTexture = GraphBuilder.CreateSRV(FRDGTextureSRVDesc(GetRDG(View.PrevViewInfo.CustomSSRInput.RT[0])));
 				PrevSceneColorViewRect = View.PrevViewInfo.CustomSSRInput.ViewportRect;
 				PrevSceneColorPreExposureInvValue = 1.0f / View.PrevViewInfo.SceneColorPreExposure;
 			}
 			else if (View.PrevViewInfo.TemporalAAHistory.IsValid())
 			{
-				PrevSceneColorTexture = GetRDG(View.PrevViewInfo.TemporalAAHistory.RT[0]);
+				FRDGTextureRef TemporalAAHistoryTexture = GetRDG(View.PrevViewInfo.TemporalAAHistory.RT[0]);
+				PrevSceneColorTexture = GraphBuilder.CreateSRV(TemporalAAHistoryTexture->Desc.IsTextureArray()
+					? FRDGTextureSRVDesc::CreateForSlice(TemporalAAHistoryTexture, View.PrevViewInfo.TemporalAAHistory.OutputSliceIndex)
+					: FRDGTextureSRVDesc(TemporalAAHistoryTexture));
 				PrevSceneColorViewRect = View.PrevViewInfo.TemporalAAHistory.ViewportRect;
 				PrevSceneColorPreExposureInvValue = 1.0f / View.PrevViewInfo.SceneColorPreExposure;
 			}
 			else if (View.PrevViewInfo.ScreenSpaceRayTracingInput.IsValid())
 			{
-				PrevSceneColorTexture = GetRDG(View.PrevViewInfo.ScreenSpaceRayTracingInput);
+				PrevSceneColorTexture = GraphBuilder.CreateSRV(FRDGTextureSRVDesc(GetRDG(View.PrevViewInfo.ScreenSpaceRayTracingInput)));
 				PrevSceneColorViewRect = View.PrevViewInfo.ViewRect;
 				PrevSceneColorPreExposureInvValue = 1.0f / View.PrevViewInfo.SceneColorPreExposure;
 			}
@@ -866,7 +869,7 @@ TRDGUniformBufferRef<FTranslucentBasePassUniformParameters> CreateTranslucentBas
 			BasePassParameters.PrevSceneColor = PrevSceneColorTexture;
 			BasePassParameters.PrevSceneColorSampler = TStaticSamplerState<SF_Bilinear, AM_Clamp, AM_Clamp, AM_Clamp>::GetRHI();
 
-			FScreenPassTextureViewportParameters PrevSceneColorParameters = GetScreenPassTextureViewportParameters(FScreenPassTextureViewport(PrevSceneColorTexture, PrevSceneColorViewRect));
+			FScreenPassTextureViewportParameters PrevSceneColorParameters = GetScreenPassTextureViewportParameters(FScreenPassTextureViewport(PrevSceneColorTexture->Desc.Texture, PrevSceneColorViewRect));
 			BasePassParameters.PrevSceneColorBilinearUVMin = PrevSceneColorParameters.UVViewportBilinearMin;
 			BasePassParameters.PrevSceneColorBilinearUVMax = PrevSceneColorParameters.UVViewportBilinearMax;
 
@@ -887,7 +890,7 @@ TRDGUniformBufferRef<FTranslucentBasePassUniformParameters> CreateTranslucentBas
 		{
 			BasePassParameters.HZBTexture = SystemTextures.Black;
 			BasePassParameters.HZBSampler = TStaticSamplerState<SF_Point, AM_Clamp, AM_Clamp, AM_Clamp>::GetRHI();
-			BasePassParameters.PrevSceneColor = SystemTextures.Black;
+			BasePassParameters.PrevSceneColor = GraphBuilder.CreateSRV(FRDGTextureSRVDesc(SystemTextures.Black));
 			BasePassParameters.PrevSceneColorSampler = TStaticSamplerState<SF_Bilinear, AM_Clamp, AM_Clamp, AM_Clamp>::GetRHI();
 			BasePassParameters.PrevSceneColorBilinearUVMin = FVector2f(0.0f, 0.0f);
 			BasePassParameters.PrevSceneColorBilinearUVMax = FVector2f(1.0f, 1.0f);
diff --git a/Engine/Source/Runtime/Renderer/Public/ScreenPass.h b/Engine/Source/Runtime/Renderer/Public/ScreenPass.h
--- a/Engine/Source/Runtime/Renderer/Public/ScreenPass.h
+++ b/Engine/Source/Runtime/Renderer/Public/ScreenPass.h
@@ -11,6 +11,9 @@
 #include "RHIStaticStates.h"
 #include "SceneView.h"
 
+struct FScreenPassTexture;
+struct FScreenPassTextureSlice;
+
 // Returns whether a HMD hidden area mask is being used for VR.
 RENDERER_API bool IsHMDHiddenAreaMaskActive();
 
@@ -38,11 +41,11 @@ struct FScreenPassTexture
 	FScreenPassTexture() = default;
 
 	explicit FScreenPassTexture(FRDGTextureRef InTexture);
+	explicit FScreenPassTexture(const FScreenPassTextureSlice& ScreenTexture);
 
-	FScreenPassTexture(FRDGTextureRef InTexture, FIntRect InViewRect)
-		: Texture(InTexture)
-		, ViewRect(InViewRect)
-	{}
+	FScreenPassTexture(FRDGTextureRef InTexture, FIntRect InViewRect);
+
+	RENDERER_API static FScreenPassTexture CopyFromSlice(FRDGBuilder& GraphBuilder, const FScreenPassTextureSlice& ScreenTextureSlice);
 
 	bool IsValid() const;
 
@@ -53,6 +56,24 @@ struct FScreenPassTexture
 	FIntRect ViewRect;
 };
 
+// Describes a Texture2D or a slice within a texture array with a paired viewport rect.
+struct FScreenPassTextureSlice
+{
+	FScreenPassTextureSlice() = default;
+
+	FScreenPassTextureSlice(FRDGTextureSRVRef InTextureSRV, FIntRect InViewRect);
+
+	RENDERER_API static FScreenPassTextureSlice CreateFromScreenPassTexture(FRDGBuilder& GraphBuilder, const FScreenPassTexture& ScreenTexture);
+
+	bool IsValid() const;
+
+	bool operator==(FScreenPassTextureSlice Other) const;
+	bool operator!=(FScreenPassTextureSlice Other) const;
+
+	FRDGTextureSRVRef TextureSRV = nullptr;
+	FIntRect ViewRect;
+};
+
 // Describes a texture with a load action for usage as a render target.
 struct FScreenPassRenderTarget : public FScreenPassTexture
 {
@@ -107,6 +128,7 @@ public:
 	{}
 
 	explicit FScreenPassTextureViewport(FScreenPassTexture InTexture);
+	explicit FScreenPassTextureViewport(FScreenPassTextureSlice InTexture);
 
 	explicit FScreenPassTextureViewport(FRDGTextureRef InTexture)
 		: FScreenPassTextureViewport(FScreenPassTexture(InTexture))
diff --git a/Engine/Source/Runtime/Renderer/Public/ScreenPass.inl b/Engine/Source/Runtime/Renderer/Public/ScreenPass.inl
--- a/Engine/Source/Runtime/Renderer/Public/ScreenPass.inl
+++ b/Engine/Source/Runtime/Renderer/Public/ScreenPass.inl
@@ -72,6 +72,22 @@ inline FScreenPassTexture::FScreenPassTexture(FRDGTextureRef InTexture)
 	}
 }
 
+inline FScreenPassTexture::FScreenPassTexture(FRDGTextureRef InTexture, FIntRect InViewRect)
+	: Texture(InTexture)
+	, ViewRect(InViewRect)
+{ }
+
+inline FScreenPassTexture::FScreenPassTexture(const FScreenPassTextureSlice& ScreenTexture)
+	: Texture(ScreenTexture.TextureSRV->Desc.Texture)
+	, ViewRect(ScreenTexture.ViewRect)
+{
+	if (Texture && Texture->Desc.IsTextureArray())
+	{
+		check(ScreenTexture.TextureSRV->Desc.FirstArraySlice == 0);
+		check(ScreenTexture.TextureSRV->Desc.NumArraySlices == 1);
+	}
+}
+
 inline bool FScreenPassTexture::IsValid() const
 {
 	return Texture != nullptr && !ViewRect.IsEmpty();
@@ -87,6 +103,31 @@ inline bool FScreenPassTexture::operator!=(FScreenPassTexture Other) const
 	return !(*this == Other);
 }
 
+inline FScreenPassTextureSlice::FScreenPassTextureSlice(FRDGTextureSRVRef InTextureSRV, FIntRect InViewRect)
+	: TextureSRV(InTextureSRV)
+	, ViewRect(InViewRect)
+{
+	if (TextureSRV && TextureSRV->Desc.Texture->Desc.IsTextureArray())
+	{
+		check(TextureSRV->Desc.NumArraySlices == 1);
+	}
+}
+
+inline bool FScreenPassTextureSlice::IsValid() const
+{
+	return TextureSRV != nullptr && !ViewRect.IsEmpty();
+}
+
+inline bool FScreenPassTextureSlice::operator==(FScreenPassTextureSlice Other) const
+{
+	return TextureSRV == Other.TextureSRV && ViewRect == Other.ViewRect;
+}
+
+inline bool FScreenPassTextureSlice::operator!=(FScreenPassTextureSlice Other) const
+{
+	return !(*this == Other);
+}
+
 inline bool FScreenPassRenderTarget::operator==(FScreenPassRenderTarget Other) const
 {
 	return FScreenPassTexture::operator==(Other) && LoadAction == Other.LoadAction;
@@ -109,6 +150,13 @@ inline FScreenPassTextureViewport::FScreenPassTextureViewport(FScreenPassTexture
 	Rect = InTexture.ViewRect;
 }
 
+inline FScreenPassTextureViewport::FScreenPassTextureViewport(FScreenPassTextureSlice InTexture)
+{
+	check(InTexture.IsValid());
+	Extent = InTexture.TextureSRV->Desc.Texture->Desc.Extent;
+	Rect = InTexture.ViewRect;
+}
+
 inline bool FScreenPassTextureViewport::operator==(const FScreenPassTextureViewport& Other) const
 {
 	return Extent == Other.Extent && Rect == Other.Rect;
