Fixes a bug where AddAfterPass() could possibly create a SRV on the backbuffer

diff --git a/Engine/Source/Runtime/Renderer/Private/PostProcess/PostProcessing.cpp b/Engine/Source/Runtime/Renderer/Private/PostProcess/PostProcessing.cpp
--- a/Engine/Source/Runtime/Renderer/Private/PostProcess/PostProcessing.cpp
+++ b/Engine/Source/Runtime/Renderer/Private/PostProcess/PostProcessing.cpp
@@ -547,7 +547,14 @@ void AddPostProcessingPasses(
 
 	const auto AddAfterPass = [&](EPass InPass, const FScreenPassTexture InSceneColor) -> FScreenPassTexture
 	{
-		return FScreenPassTexture(AddAfterPassForSceneColorSlice(InPass, FScreenPassTextureSlice::CreateFromScreenPassTexture(GraphBuilder, InSceneColor)));
+		FAfterPassCallbackDelegateArray& PassCallbacks = PassSequence.GetAfterPassCallbacks(InPass);
+
+		if (PassCallbacks.Num())
+		{
+			return FScreenPassTexture(AddAfterPassForSceneColorSlice(InPass, FScreenPassTextureSlice::CreateFromScreenPassTexture(GraphBuilder, InSceneColor)));
+		}
+
+		return InSceneColor;
 	};
 
 	if (bPostProcessingEnabled)
