Change hair culling resources to derive from FGroomResources for supporting progressive allocation.

diff --git a/Engine/Plugins/Runtime/HairStrands/Source/HairStrandsCore/Private/GroomComponent.cpp b/Engine/Plugins/Runtime/HairStrands/Source/HairStrandsCore/Private/GroomComponent.cpp
--- a/Engine/Plugins/Runtime/HairStrands/Source/HairStrandsCore/Private/GroomComponent.cpp
+++ b/Engine/Plugins/Runtime/HairStrands/Source/HairStrandsCore/Private/GroomComponent.cpp
@@ -2774,9 +2774,14 @@ void UGroomComponent::InitResources(bool bIsBindingReloading)
 				check(StrandsLODCount <= LODScreenSizes.Num());
 				check(StrandsLODCount <= LODVisibilities.Num());
 				HairGroupInstance->HairGroupPublicData->ClusterCount = HairGroupInstance->Strands.ClusterCullingResource->BulkData.Header.ClusterCount;
-				BeginInitResource(HairGroupInstance->HairGroupPublicData);
 			}
 
+			HairGroupInstance->Strands.CullingResource = new FHairStrandsCullingResource(
+				HairGroupInstance->HairGroupPublicData->RestPointCount, 
+				HairGroupInstance->HairGroupPublicData->RestCurveCount, 
+				HairGroupInstance->HairGroupPublicData->ClusterCount, 
+				ResourceName, OwnerName);
+			HairGroupInstance->HairGroupPublicData->Culling = &HairGroupInstance->Strands.CullingResource->Resources;
 
 			HairGroupInstance->Strands.HairInterpolationType = HairInterpolationType;
 		}
@@ -2968,6 +2973,7 @@ void UGroomComponent::DeleteDeferredHairGroupInstances()
 			{
 				InternalResourceRelease(LocalInstance->Strands.DeformedRootResource);
 				InternalResourceRelease(LocalInstance->Strands.DeformedResource);
+				InternalResourceRelease(LocalInstance->Strands.CullingResource);
 
 				#if RHI_RAYTRACING
 				if (LocalInstance->Strands.RenRaytracingResourceOwned)
@@ -3025,7 +3031,7 @@ void UGroomComponent::DeleteDeferredHairGroupInstances()
 				}
 			}
 
-			InternalResourceRelease(LocalInstance->HairGroupPublicData);
+			delete LocalInstance->HairGroupPublicData;
 			LocalInstance->Release();
 			delete LocalInstance;
 		});
@@ -4026,7 +4032,7 @@ void DumpLoadedGroomComponent(IConsoleVariable* InCVarPakTesterEnabled)
 						GroupCount,
                         Instance->Cards.LODs.Num(), 
                         InstanceMemory.GetTotalSize()* ToMb,
-                        InstanceMemory.Common * ToMb,
+                        0,
 						InstanceMemory.Guides * ToMb,
 						InstanceMemory.Strands* ToMb,
 						InstanceMemory.Cards  * ToMb,
@@ -4045,7 +4051,7 @@ void DumpLoadedGroomComponent(IConsoleVariable* InCVarPakTesterEnabled)
 		Total_Component,
 		Total_Group,
         TotalMemory.GetTotalSize() * ToMb,
-		TotalMemory.Common* ToMb,
+		0,
 		TotalMemory.Guides * ToMb,
 		TotalMemory.Strands * ToMb,
 		TotalMemory.Cards * ToMb,
@@ -4128,11 +4134,11 @@ FGroomComponentRecreateRenderStateContext::~FGroomComponentRecreateRenderStateCo
 FGroomComponentMemoryStats FGroomComponentMemoryStats::Get(const FHairGroupInstance* In)
 {
 	FGroomComponentMemoryStats Out;
-	Out.Common  += In->HairGroupPublicData ? In->HairGroupPublicData->GetResourcesSize() : 0;
 	Out.Guides  += In->Guides.DeformedResource ? In->Guides.DeformedResource->GetResourcesSize() : 0;
 	Out.Guides  += In->Guides.DeformedRootResource ? In->Guides.DeformedRootResource->GetResourcesSize() : 0;
 	Out.Strands += In->Strands.DeformedResource ? In->Strands.DeformedResource->GetResourcesSize() : 0;
 	Out.Strands += In->Strands.DeformedRootResource ? In->Strands.DeformedRootResource->GetResourcesSize() : 0;
+	Out.Strands += In->Strands.CullingResource ? In->Strands.CullingResource->GetResourcesSize() : 0;
 #if RHI_RAYTRACING
 	Out.Strands += (In->Strands.RenRaytracingResourceOwned && In->Strands.RenRaytracingResource) ? In->Strands.RenRaytracingResource->GetResourcesSize() : 0;
 #endif
@@ -4149,7 +4155,6 @@ FGroomComponentMemoryStats FGroomComponentMemoryStats::Get(const FHairGroupInsta
 
 void FGroomComponentMemoryStats::Accumulate(const FGroomComponentMemoryStats& In)
 {
-	Common += In.Common ;
 	Guides += In.Guides ;
 	Strands+= In.Strands;
 	Cards  += In.Cards  ;
diff --git a/Engine/Plugins/Runtime/HairStrands/Source/HairStrandsCore/Private/GroomDebug.cpp b/Engine/Plugins/Runtime/HairStrands/Source/HairStrandsCore/Private/GroomDebug.cpp
--- a/Engine/Plugins/Runtime/HairStrands/Source/HairStrandsCore/Private/GroomDebug.cpp
+++ b/Engine/Plugins/Runtime/HairStrands/Source/HairStrandsCore/Private/GroomDebug.cpp
@@ -1030,7 +1030,7 @@ static void AddHairDebugPrintInstancePass(
 			const float ContinousLODRadius = Instance->HairGroupPublicData->ContinuousLODScreenSize * MaxRectSizeInPixels * 0.5f; // Diameter->Radius
 
 			FHairDebugPrintInstanceCS::FParameters* Parameters = GraphBuilder.AllocParameters<FHairDebugPrintInstanceCS::FParameters>();
-			Parameters->InstanceAABB = Register(GraphBuilder, Instance->HairGroupPublicData->GroupAABBBuffer, ERDGImportedBufferFlags::CreateSRV).SRV;
+			Parameters->InstanceAABB = Register(GraphBuilder, Instance->HairGroupPublicData->GetGroupAABBBuffer(), ERDGImportedBufferFlags::CreateSRV).SRV;
 			Parameters->InstanceScreenSphereBound = FVector4f(Instance->HairGroupPublicData->ContinuousLODScreenPos.X, Instance->HairGroupPublicData->ContinuousLODScreenPos.Y, 0.f, ContinousLODRadius);
 			ShaderPrint::SetParameters(GraphBuilder, *ShaderPrintData, Parameters->ShaderPrintUniformBuffer);
 			ClearUnusedGraphResources(ComputeShader, Parameters);
@@ -1143,7 +1143,7 @@ static void AddHairDebugPrintMemoryPass(
 				D.Data2.Z = MemoryStats.Cards;
 				D.Data2.W = MemoryStats.Meshes;
 
-				D.Data3.X = MemoryStats.Common;
+				D.Data3.X = 0;
 				D.Data3.Y = 0;
 				D.Data3.Z = 0;
 				D.Data3.W = 0;
diff --git a/Engine/Plugins/Runtime/HairStrands/Source/HairStrandsCore/Private/GroomManager.cpp b/Engine/Plugins/Runtime/HairStrands/Source/HairStrandsCore/Private/GroomManager.cpp
--- a/Engine/Plugins/Runtime/HairStrands/Source/HairStrandsCore/Private/GroomManager.cpp
+++ b/Engine/Plugins/Runtime/HairStrands/Source/HairStrandsCore/Private/GroomManager.cpp
@@ -1068,12 +1068,12 @@ static void RunHairLODSelection(
 			else if (GeometryType == EHairGeometryType::Strands)
 			{
 				check(Instance->HairGroupPublicData);
-				Instance->HairGroupPublicData->Allocate(GraphBuilder);
 
 				if (Instance->Strands.RestRootResource)			{ Instance->Strands.RestRootResource->Allocate(GraphBuilder, EHairResourceLoadingType::Async, ResourceStatus, RequestedCurveCount, RequestedPointCount, MeshLODIndex); }
 				if (Instance->Strands.RestResource)				{ Instance->Strands.RestResource->Allocate(GraphBuilder, EHairResourceLoadingType::Async, ResourceStatus, RequestedCurveCount, RequestedPointCount); }
 				if (Instance->Strands.ClusterCullingResource)	{ Instance->Strands.ClusterCullingResource->Allocate(GraphBuilder, EHairResourceLoadingType::Async, ResourceStatus, RequestedCurveCount, RequestedPointCount); }
 				if (Instance->Strands.InterpolationResource)	{ Instance->Strands.InterpolationResource->Allocate(GraphBuilder, EHairResourceLoadingType::Async, ResourceStatus, RequestedCurveCount, RequestedPointCount); }
+				if (Instance->Strands.CullingResource)			{ Instance->Strands.CullingResource->Allocate(GraphBuilder, EHairResourceLoadingType::Async, ResourceStatus, RequestedCurveCount, RequestedPointCount); }
 
 				if (Instance->Strands.DeformedRootResource)		{ Instance->Strands.DeformedRootResource->Allocate(GraphBuilder, EHairResourceLoadingType::Async, ResourceStatus, RequestedCurveCount, RequestedPointCount, MeshLODIndex); }
 				if (Instance->Strands.DeformedResource)			{ Instance->Strands.DeformedResource->Allocate(GraphBuilder, EHairResourceLoadingType::Async, ResourceStatus, RequestedCurveCount, RequestedPointCount); }
diff --git a/Engine/Plugins/Runtime/HairStrands/Source/HairStrandsCore/Private/GroomResources.cpp b/Engine/Plugins/Runtime/HairStrands/Source/HairStrandsCore/Private/GroomResources.cpp
--- a/Engine/Plugins/Runtime/HairStrands/Source/HairStrandsCore/Private/GroomResources.cpp
+++ b/Engine/Plugins/Runtime/HairStrands/Source/HairStrandsCore/Private/GroomResources.cpp
@@ -57,6 +57,20 @@ bool ValidateHairBulkData()
 	return GHairStrandsBulkData_Validation > 0;
 }
 
+static void ConvertToExternalBufferWithViews(FRDGBuilder& GraphBuilder, FRDGBufferRef& InBuffer, FRDGExternalBuffer& OutBuffer, EPixelFormat Format = PF_Unknown)
+{
+	OutBuffer.Buffer = GraphBuilder.ConvertToExternalBuffer(InBuffer);
+	if (EnumHasAnyFlags(InBuffer->Desc.Usage, BUF_ShaderResource))
+	{
+		OutBuffer.SRV = OutBuffer.Buffer->GetOrCreateSRV(FRDGBufferSRVDesc(InBuffer, Format));
+	}
+	if (EnumHasAnyFlags(InBuffer->Desc.Usage, BUF_UnorderedAccess))
+	{
+		OutBuffer.UAV = OutBuffer.Buffer->GetOrCreateUAV(FRDGBufferUAVDesc(InBuffer, Format));
+	}
+	OutBuffer.Format = Format;
+}
+
 ////////////////////////////////////////////////////////////////////////////////////
 
 EHairResourceLoadingType GetHairResourceLoadingType(EHairGeometryType InGeometryType, int32 InLODIndex)
@@ -383,6 +397,16 @@ void InternalCreateByteAddressBufferRDG_FromHairBulkData(FRDGBuilder& GraphBuild
 ///////////////////////////////////////////////////////////////////////////////////////////////////
 // Regular data loading
 
+static void InternalCreateIndirectBufferRDG(FRDGBuilder& GraphBuilder, FRDGExternalBuffer& Out, const TCHAR* DebugName, const FName& OwnerName)
+{
+	FRDGBufferDesc Desc = FRDGBufferDesc::CreateBufferDesc(4, 4);
+	Desc.Usage |= BUF_DrawIndirect;
+	FRDGBufferRef Buffer = GraphBuilder.CreateBuffer(Desc, DebugName, ERDGBufferFlags::None);
+	Buffer->SetOwnerName(OwnerName);
+	AddClearUAVPass(GraphBuilder, GraphBuilder.CreateUAV(Buffer, PF_R32_UINT), 0u);
+	ConvertToExternalBufferWithViews(GraphBuilder, Buffer, Out, PF_R32_UINT);
+}
+
 static FRDGBufferRef InternalCreateVertexBuffer(
 	FRDGBuilder& GraphBuilder,
 	const TCHAR* Name,
@@ -1244,6 +1268,85 @@ void FHairStrandsClusterCullingResource::InternalResetLoadedSize()
 	BulkData.ResetLoadedSize();
 }
 
+
+/////////////////////////////////////////////////////////////////////////////////////////
+// Cluster culling resources
+
+struct FHairStrandsSintFormat
+{
+	typedef int32 Type;
+	typedef int32 BulkType;
+	static const uint32 ComponentCount = 1;
+	static const uint32 SizeInByte = sizeof(Type);
+	static const EVertexElementType VertexElementType = VET_UInt;
+	static const EPixelFormat Format = PF_R32_SINT;
+};
+
+struct FHairStrandsUintFormat
+{
+	typedef uint32 Type;
+	typedef uint32 BulkType;
+	static const uint32 ComponentCount = 1;
+	static const uint32 SizeInByte = sizeof(Type);
+	static const EVertexElementType VertexElementType = VET_UInt;
+	static const EPixelFormat Format = PF_R32_UINT;
+};
+
+struct FHairStrandsFloatFormat
+{
+	typedef float Type;
+	typedef float BulkType;
+	static const uint32 ComponentCount = 1;
+	static const uint32 SizeInByte = sizeof(Type);
+	static const EVertexElementType VertexElementType = VET_UInt;
+	static const EPixelFormat Format = PF_R32_FLOAT;
+};
+
+FHairStrandsCullingResource::FHairStrandsCullingResource(uint32 InPointCount, uint32 InCurveCount, uint32 InClusterCount, const FHairResourceName& InResourceName, const FName& InOwnerName):
+FHairCommonResource(EHairStrandsAllocationType::Deferred, InResourceName, InOwnerName)
+{
+	MaxPointCount = InPointCount;
+	MaxCurveCount = InCurveCount;
+	ClusterCount = InClusterCount;
+	MaxAvailableCurveCount = 0;
+}
+
+void FHairStrandsCullingResource::InternalAllocate(FRDGBuilder& GraphBuilder, uint32 InCurveCount, uint32 InPointCount, int32 InLODIndex)
+{
+	if (ClusterCount > 0)
+	{	
+		const uint32 PointCount = FMath::Min(InPointCount, MaxPointCount);
+		const uint32 CurveCount = FMath::Min(InCurveCount, MaxCurveCount);
+
+		InternalCreateVertexBufferRDG<FHairStrandsUintFormat>(GraphBuilder, CurveCount, Resources.CulledCurveBuffer, TEXT("Hair.Cluster_CulledCurveBuffer"), GetOwnerName(), EHairResourceUsageType::Dynamic);
+		InternalCreateVertexBufferRDG<FHairStrandsUintFormat>(GraphBuilder, PointCount, Resources.CulledVertexIdBuffer, TEXT("Hair.Cluster_CulledVertexIdBuffer"), GetOwnerName(), EHairResourceUsageType::Dynamic);
+		InternalCreateVertexBufferRDG<FHairStrandsFloatFormat>(GraphBuilder, PointCount, Resources.CulledVertexRadiusScaleBuffer, TEXT("Hair.Cluster_CulledVertexRadiusScaleBuffer"), GetOwnerName(), EHairResourceUsageType::Dynamic);
+		AddClearUAVFloatPass(GraphBuilder, RegisterAsUAV(GraphBuilder, Resources.CulledVertexRadiusScaleBuffer), 0.f);
+	
+		// Only allocated once, as cluster count does not change
+		if (Resources.DrawIndirectBuffer.Buffer == nullptr)
+		{
+			InternalCreateVertexBufferRDG<FHairStrandsSintFormat>(GraphBuilder, ClusterCount * 6, Resources.ClusterAABBBuffer, TEXT("Hair.Cluster_ClusterAABBBuffer"), GetOwnerName(), EHairResourceUsageType::Dynamic);
+			InternalCreateVertexBufferRDG<FHairStrandsSintFormat>(GraphBuilder, 6, Resources.GroupAABBBuffer, TEXT("Hair.Cluster_GroupAABBBuffer"), GetOwnerName(), EHairResourceUsageType::Dynamic);
+			InternalCreateIndirectBufferRDG(GraphBuilder, Resources.DrawIndirectBuffer, TEXT("Hair.Cluster_DrawIndirectBuffer"), GetOwnerName());
+			InternalCreateIndirectBufferRDG(GraphBuilder, Resources.DrawIndirectRasterComputeBuffer, TEXT("Hair.Cluster_DrawIndirectRasterComputeBuffer"), GetOwnerName());
+			GraphBuilder.SetBufferAccessFinal(Register(GraphBuilder, Resources.DrawIndirectBuffer, ERDGImportedBufferFlags::None).Buffer, ERHIAccess::IndirectArgs);
+		}
+	}
+}
+
+void FHairStrandsCullingResource::InternalRelease()
+{
+	Resources.DrawIndirectBuffer.Release();
+	Resources.DrawIndirectRasterComputeBuffer.Release();
+	Resources.ClusterAABBBuffer.Release();
+	Resources.GroupAABBBuffer.Release();
+	Resources.CulledCurveBuffer.Release();
+	Resources.CulledVertexIdBuffer.Release();
+	Resources.CulledVertexRadiusScaleBuffer.Release();
+	MaxAvailableCurveCount = 0;
+}
+
 /////////////////////////////////////////////////////////////////////////////////////////
 
 FHairStrandsRestRootResource::FHairStrandsRestRootResource(FHairStrandsRootBulkData& InBulkData, EHairStrandsResourcesType InCurveType, const FHairResourceName& InResourceName, const FName& InOwnerName) :
diff --git a/Engine/Plugins/Runtime/HairStrands/Source/HairStrandsCore/Private/HairStrandsInterpolation.cpp b/Engine/Plugins/Runtime/HairStrands/Source/HairStrandsCore/Private/HairStrandsInterpolation.cpp
--- a/Engine/Plugins/Runtime/HairStrands/Source/HairStrandsCore/Private/HairStrandsInterpolation.cpp
+++ b/Engine/Plugins/Runtime/HairStrands/Source/HairStrandsCore/Private/HairStrandsInterpolation.cpp
@@ -877,7 +877,7 @@ static void AddHairClusterAABBPass(
 	FTransform InRenLocalToTranslatedWorld = Instance->LocalToWorld;
 	InRenLocalToTranslatedWorld.AddToTranslation(TranslatedWorldOffset);
 	const FBoxSphereBounds TransformedBounds = Bounds.TransformBy(InRenLocalToTranslatedWorld);
-	Instance->HairGroupPublicData->bClusterAABBValid = UpdateType == EHairAABBUpdateType::UpdateClusterAABB;
+	Instance->HairGroupPublicData->SetClusterAABBValid(UpdateType == EHairAABBUpdateType::UpdateClusterAABB);
 
 	FHairClusterAABBCS::FParameters* Parameters = GraphBuilder.AllocParameters<FHairClusterAABBCS::FParameters>();
 	Parameters->LODIndex = ClusterData ? ClusterData->LODIndex : 1;;
@@ -2069,8 +2069,8 @@ void ComputeHairStrandsInterpolation(
 					}
 				}
 
-				Instance->HairGroupPublicData->bClusterAABBValid = false;
-				Instance->HairGroupPublicData->bGroupAABBValid = false;
+				Instance->HairGroupPublicData->SetClusterAABBValid(false);
+				Instance->HairGroupPublicData->SetGroupAABBValid(false);
 
 
 				if (bNeedGPUAABB)
@@ -2090,8 +2090,8 @@ void ComputeHairStrandsInterpolation(
 						Strands_PositionSRV,
 						Strands_CulledVertexCount.SRV);
 
-					Instance->HairGroupPublicData->bClusterAABBValid = UpdateType == EHairAABBUpdateType::UpdateClusterAABB;
-					Instance->HairGroupPublicData->bGroupAABBValid = true;
+					Instance->HairGroupPublicData->SetClusterAABBValid(UpdateType == EHairAABBUpdateType::UpdateClusterAABB);
+					Instance->HairGroupPublicData->SetGroupAABBValid(true);
 				}
 			}
 
diff --git a/Engine/Plugins/Runtime/HairStrands/Source/HairStrandsCore/Public/GroomComponent.h b/Engine/Plugins/Runtime/HairStrands/Source/HairStrandsCore/Public/GroomComponent.h
--- a/Engine/Plugins/Runtime/HairStrands/Source/HairStrandsCore/Public/GroomComponent.h
+++ b/Engine/Plugins/Runtime/HairStrands/Source/HairStrandsCore/Public/GroomComponent.h
@@ -371,7 +371,6 @@ HAIRSTRANDSCORE_API const FLinearColor GetHairGroupDebugColor(int32 GroupIt);
 
 struct FGroomComponentMemoryStats
 {
-	uint32 Common = 0;
 	uint32 Guides = 0;
 	uint32 Strands= 0;
 	uint32 Cards  = 0;
diff --git a/Engine/Plugins/Runtime/HairStrands/Source/HairStrandsCore/Public/GroomInstance.h b/Engine/Plugins/Runtime/HairStrands/Source/HairStrandsCore/Public/GroomInstance.h
--- a/Engine/Plugins/Runtime/HairStrands/Source/HairStrandsCore/Public/GroomInstance.h
+++ b/Engine/Plugins/Runtime/HairStrands/Source/HairStrandsCore/Public/GroomInstance.h
@@ -124,6 +124,7 @@ struct HAIRSTRANDSCORE_API FHairGroupInstance : public FHairStrandsInstance
 	{
 		// Resources - Strands cluster data for culling/voxelization purpose
 		FHairStrandsClusterCullingResource* ClusterCullingResource = nullptr;
+		FHairStrandsCullingResource* CullingResource = nullptr;
 
 		// Resources - Raytracing data when enabling (expensive) raytracing method
 		#if RHI_RAYTRACING
diff --git a/Engine/Plugins/Runtime/HairStrands/Source/HairStrandsCore/Public/GroomResources.h b/Engine/Plugins/Runtime/HairStrands/Source/HairStrandsCore/Public/GroomResources.h
--- a/Engine/Plugins/Runtime/HairStrands/Source/HairStrandsCore/Public/GroomResources.h
+++ b/Engine/Plugins/Runtime/HairStrands/Source/HairStrandsCore/Public/GroomResources.h
@@ -350,6 +350,37 @@ struct FHairStrandsClusterCullingResource : public FHairCommonResource
 	FHairStrandsClusterCullingBulkData& BulkData;
 };
 
+struct FHairStrandsCullingResource : public FHairCommonResource
+{
+	FHairStrandsCullingResource(uint32 InPointCount, uint32 InCurveCount, uint32 InClusterCount, const FHairResourceName& InResourceName, const FName& InOwnerName);
+
+	/* Init/Release buffers */
+	virtual void InternalAllocate(FRDGBuilder& GraphBuilder, uint32 InCurveCount, uint32 InPointCount, int32 InLODIndex) override;
+	virtual void InternalRelease() override;
+
+	/* Get the resource name */
+	virtual FString GetFriendlyName() const override { return TEXT("FHairStrandsCullingResource"); }
+
+	/* Return the memory size for GPU resources */
+	uint32 GetResourcesSize() const
+	{
+		uint32 Total = 0;
+		Total += GetBufferTotalNumBytes(Resources.DrawIndirectBuffer);
+		Total += GetBufferTotalNumBytes(Resources.DrawIndirectRasterComputeBuffer);
+		Total += GetBufferTotalNumBytes(Resources.ClusterAABBBuffer);
+		Total += GetBufferTotalNumBytes(Resources.GroupAABBBuffer);
+		Total += GetBufferTotalNumBytes(Resources.CulledCurveBuffer);
+		Total += GetBufferTotalNumBytes(Resources.CulledVertexIdBuffer);
+		Total += GetBufferTotalNumBytes(Resources.CulledVertexRadiusScaleBuffer);
+		return Total;
+	}
+
+	uint32 ClusterCount = 0;
+	uint32 MaxPointCount = 0;
+	uint32 MaxCurveCount = 0;
+	FHairGroupPublicData::FCulling Resources;
+};
+
 /* Render buffers for root deformation for dynamic meshes */
 struct FHairStrandsRestRootResource : public FHairCommonResource
 {
diff --git a/Engine/Shaders/Private/HairStrands/HairStrandsDebug.usf b/Engine/Shaders/Private/HairStrands/HairStrandsDebug.usf
--- a/Engine/Shaders/Private/HairStrands/HairStrandsDebug.usf
+++ b/Engine/Shaders/Private/HairStrands/HairStrandsDebug.usf
@@ -2544,7 +2544,6 @@ struct FComponent
 	uint RestCurve;
 	uint ActiveCurve;
 
-	uint GPU_Common;
 	uint GPU_Guides;
 	uint GPU_Strands;
 	uint GPU_Cards;
@@ -2565,7 +2564,6 @@ FComponent GetComponent(uint Index)
 	Out.GPU_Strands 		=In.Data1.y;
 	Out.GPU_Cards 			=In.Data1.z;
 	Out.GPU_Meshes 			=In.Data1.w;
-	Out.GPU_Common			=In.Data2.x;
 
 	return Out;
 }
@@ -2722,15 +2720,13 @@ void MainCS(uint GroupIndex : SV_GroupIndex, uint3 DispatchThreadId : SV_Dispatc
 			{
 				const FComponent Info = GetComponent(It);
 				const uint Total
-					= Info.GPU_Common
-					+ Info.GPU_Guides
+					= Info.GPU_Guides
 					+ Info.GPU_Strands
 					+ Info.GPU_Cards;
 
 
 				PrintGroup(Ctx, Info.GroupIndex, Info.GroupCount);
 	
-				PrintLargeNumber(Ctx, Info.GPU_Common,	FontWhite);		Print(Ctx, TEXT(" "), FontWhite);
 				PrintLargeNumber(Ctx, Info.GPU_Guides,	FontWhite);		Print(Ctx, TEXT(" "), FontWhite);
 				PrintLargeNumber(Ctx, Info.GPU_Strands,	FontWhite);		Print(Ctx, TEXT(" "), FontWhite);
 				PrintLargeNumber(Ctx, Info.GPU_Cards,	FontWhite);		Print(Ctx, TEXT(" "), FontWhite);
diff --git a/Engine/Source/Runtime/Renderer/Private/HairStrands/HairStrandsInterface.cpp b/Engine/Source/Runtime/Renderer/Private/HairStrands/HairStrandsInterface.cpp
--- a/Engine/Source/Runtime/Renderer/Private/HairStrands/HairStrandsInterface.cpp
+++ b/Engine/Source/Runtime/Renderer/Private/HairStrands/HairStrandsInterface.cpp
@@ -243,137 +243,12 @@ bool IsHairStrandsSimulationEnable()
 }
 
 ////////////////////////////////////////////////////////////////////////////////////////////////////////////////////////
-void ConvertToExternalBufferWithViews(FRDGBuilder& GraphBuilder, FRDGBufferRef& InBuffer, FRDGExternalBuffer& OutBuffer, EPixelFormat Format)
-{
-	OutBuffer.Buffer = GraphBuilder.ConvertToExternalBuffer(InBuffer);
-	if (EnumHasAnyFlags(InBuffer->Desc.Usage, BUF_ShaderResource))
-	{
-		OutBuffer.SRV = OutBuffer.Buffer->GetOrCreateSRV(FRDGBufferSRVDesc(InBuffer, Format));
-	}
-	if (EnumHasAnyFlags(InBuffer->Desc.Usage, BUF_UnorderedAccess))
-	{
-		OutBuffer.UAV = OutBuffer.Buffer->GetOrCreateUAV(FRDGBufferUAVDesc(InBuffer, Format));
-	}
-	OutBuffer.Format = Format;
-}
-
-void InternalCreateIndirectBufferRDG(FRDGBuilder& GraphBuilder, FRDGExternalBuffer& Out, const TCHAR* DebugName, const FName& OwnerName)
-{
-	FRDGBufferDesc Desc = FRDGBufferDesc::CreateBufferDesc(4, 4);
-	Desc.Usage |= BUF_DrawIndirect;
-	FRDGBufferRef Buffer = GraphBuilder.CreateBuffer(Desc, DebugName, ERDGBufferFlags::None);
-	Buffer->SetOwnerName(OwnerName);
-	AddClearUAVPass(GraphBuilder, GraphBuilder.CreateUAV(Buffer, PF_R32_UINT), 0u);
-	ConvertToExternalBufferWithViews(GraphBuilder, Buffer, Out, PF_R32_UINT);
-}
-
-void InternalCreateVertexBufferRDG(FRDGBuilder& GraphBuilder, uint32 ElementSizeInBytes, uint32 ElementCount, EPixelFormat Format, FRDGExternalBuffer& Out, const TCHAR* DebugName, const FName& OwnerName, bool bClearFloat=false)
-{
-	FRDGBufferRef Buffer = nullptr;
-
-	const uint32 DataCount = ElementCount;
-	const uint32 DataSizeInBytes = ElementSizeInBytes * DataCount;
-	if (DataSizeInBytes == 0)
-	{
-		Out.Buffer = nullptr;
-		return;
-	}
-
-	// #hair_todo: Create this with a create+clear pass instead?
-	const FRDGBufferDesc Desc = FRDGBufferDesc::CreateBufferDesc(ElementSizeInBytes, ElementCount);
-	Buffer = GraphBuilder.CreateBuffer(Desc, DebugName, ERDGBufferFlags::MultiFrame);
-	Buffer->SetOwnerName(OwnerName);
-	if (bClearFloat)
-	{
-		AddClearUAVFloatPass(GraphBuilder, GraphBuilder.CreateUAV(Buffer, Format), 0.f);
-	}
-	else
-	{
-		AddClearUAVPass(GraphBuilder, GraphBuilder.CreateUAV(Buffer, Format), 0);
-	}
-	ConvertToExternalBufferWithViews(GraphBuilder, Buffer, Out, Format);
-}
 
 FHairGroupPublicData::FHairGroupPublicData(uint32 InGroupIndex, const FName& InOwnerName)
 {
-	SetOwnerName(InOwnerName);
 	GroupIndex = InGroupIndex;
 }
 
-void FHairGroupPublicData::InitRHI(FRHICommandListBase& RHICmdList)
-{
-}
-
-void FHairGroupPublicData::Allocate(FRDGBuilder& GraphBuilder)
-{
-	if (bIsInitialized)
-		return;
-
-	if (ClusterCount == 0)
-		return;
-
-	bool bHasStrands = false;
-	for (const EHairGeometryType& Type : LODGeometryTypes)
-	{
-		if (Type == EHairGeometryType::Strands)
-		{
-			bHasStrands = true;
-			break;
-		}
-	}
-	
-	if (GUsingNullRHI || !bHasStrands) { return; }
-
-	InternalCreateIndirectBufferRDG(GraphBuilder, DrawIndirectBuffer, TEXT("Hair.Cluster_DrawIndirectBuffer"), GetOwnerName());
-	InternalCreateIndirectBufferRDG(GraphBuilder, DrawIndirectRasterComputeBuffer, TEXT("Hair.Cluster_DrawIndirectRasterComputeBuffer"), GetOwnerName());
-
-	InternalCreateVertexBufferRDG(GraphBuilder, sizeof(int32), ClusterCount * 6, EPixelFormat::PF_R32_SINT, ClusterAABBBuffer, TEXT("Hair.Cluster_ClusterAABBBuffer"), GetOwnerName());
-	InternalCreateVertexBufferRDG(GraphBuilder, sizeof(int32), 6, EPixelFormat::PF_R32_SINT, GroupAABBBuffer, TEXT("Hair.Cluster_GroupAABBBuffer"), GetOwnerName());
-
-	InternalCreateVertexBufferRDG(GraphBuilder, sizeof(int32), RestCurveCount, EPixelFormat::PF_R32_UINT, CulledCurveBuffer, TEXT("Hair.Cluster_CulledCurveBuffer"), GetOwnerName());
-	InternalCreateVertexBufferRDG(GraphBuilder, sizeof(int32), RestPointCount, EPixelFormat::PF_R32_UINT, CulledVertexIdBuffer, TEXT("Hair.Cluster_CulledVertexIdBuffer"), GetOwnerName());
-	InternalCreateVertexBufferRDG(GraphBuilder, sizeof(float), RestPointCount, EPixelFormat::PF_R32_FLOAT, CulledVertexRadiusScaleBuffer, TEXT("Hair.Cluster_CulledVertexRadiusScaleBuffer"), GetOwnerName(), true);
-
-	GraphBuilder.SetBufferAccessFinal(Register(GraphBuilder, DrawIndirectBuffer, ERDGImportedBufferFlags::None).Buffer, ERHIAccess::IndirectArgs);
-
-	bIsInitialized = true;
-}
-
-void FHairGroupPublicData::ReleaseRHI()
-{
-	//Release();
-}
-
-void FHairGroupPublicData::Release()
-{
-	DrawIndirectBuffer.Release();
-	DrawIndirectRasterComputeBuffer.Release();
-	ClusterAABBBuffer.Release();
-	GroupAABBBuffer.Release();
-	CulledCurveBuffer.Release();
-	CulledVertexIdBuffer.Release();
-	CulledVertexRadiusScaleBuffer.Release();
-	bIsInitialized = false;
-}
-
-uint32 FHairGroupPublicData::GetResourcesSize() const
-{
-	auto ExtractSize = [](const TRefCountPtr<FRDGPooledBuffer>& InBuffer)
-	{
-		return InBuffer ? InBuffer->Desc.BytesPerElement * InBuffer->Desc.NumElements : 0; 
-	};
-
-	uint32 Total = 0;
-	Total += ExtractSize(DrawIndirectBuffer.Buffer);
-	Total += ExtractSize(DrawIndirectRasterComputeBuffer.Buffer);
-	Total += ExtractSize(ClusterAABBBuffer.Buffer);
-	Total += ExtractSize(GroupAABBBuffer.Buffer);
-	Total += ExtractSize(CulledCurveBuffer.Buffer);
-	Total += ExtractSize(CulledVertexIdBuffer.Buffer);
-	Total += ExtractSize(CulledVertexRadiusScaleBuffer.Buffer);
-	return Total;
-}
-
 ///////////////////////////////////////////////////////////////////////////////////////////////////
 
 bool IsHairStrandsNonVisibleShadowCastingEnable()
diff --git a/Engine/Source/Runtime/Renderer/Private/HairStrands/HairStrandsVisibility.cpp b/Engine/Source/Runtime/Renderer/Private/HairStrands/HairStrandsVisibility.cpp
--- a/Engine/Source/Runtime/Renderer/Private/HairStrands/HairStrandsVisibility.cpp
+++ b/Engine/Source/Runtime/Renderer/Private/HairStrands/HairStrandsVisibility.cpp
@@ -2859,10 +2859,10 @@ void AddMeshDrawTransitionPass(
 			ExternalAccessQueue.Add(VFInput.Strands.PositionOffsetBuffer.Buffer);
 			ExternalAccessQueue.Add(VFInput.Strands.PrevPositionOffsetBuffer.Buffer);
 
-			FRDGBufferRef CulledCurveBuffer = Register(GraphBuilder, HairGroupPublicData->CulledCurveBuffer, ERDGImportedBufferFlags::None).Buffer;
-			FRDGBufferRef CulledVertexIdBuffer = Register(GraphBuilder, HairGroupPublicData->CulledVertexIdBuffer, ERDGImportedBufferFlags::None).Buffer;
-			FRDGBufferRef CulledVertexRadiusScaleBuffer = Register(GraphBuilder, HairGroupPublicData->CulledVertexRadiusScaleBuffer, ERDGImportedBufferFlags::None).Buffer;
-			FRDGBufferRef DrawIndirectBuffer = Register(GraphBuilder, HairGroupPublicData->DrawIndirectBuffer, ERDGImportedBufferFlags::None).Buffer;
+			FRDGBufferRef CulledCurveBuffer = Register(GraphBuilder, HairGroupPublicData->GetCulledCurveBuffer(), ERDGImportedBufferFlags::None).Buffer;
+			FRDGBufferRef CulledVertexIdBuffer = Register(GraphBuilder, HairGroupPublicData->GetCulledVertexIdBuffer(), ERDGImportedBufferFlags::None).Buffer;
+			FRDGBufferRef CulledVertexRadiusScaleBuffer = Register(GraphBuilder, HairGroupPublicData->GetCulledVertexRadiusScaleBuffer(), ERDGImportedBufferFlags::None).Buffer;
+			FRDGBufferRef DrawIndirectBuffer = Register(GraphBuilder, HairGroupPublicData->GetDrawIndirectBuffer(), ERDGImportedBufferFlags::None).Buffer;
 			ExternalAccessQueue.Add(CulledCurveBuffer);
 			ExternalAccessQueue.Add(CulledVertexIdBuffer);
 			ExternalAccessQueue.Add(CulledVertexRadiusScaleBuffer);
diff --git a/Engine/Source/Runtime/Renderer/Private/HairStrands/HairStrandsVoxelization.cpp b/Engine/Source/Runtime/Renderer/Private/HairStrands/HairStrandsVoxelization.cpp
--- a/Engine/Source/Runtime/Renderer/Private/HairStrands/HairStrandsVoxelization.cpp
+++ b/Engine/Source/Runtime/Renderer/Private/HairStrands/HairStrandsVoxelization.cpp
@@ -669,7 +669,7 @@ static void AddAllocateVoxelPagesPass(
 				// * Or use its clusters AABBs (for dynamic groom)
 				// * Or use its CPU AABB (for CLOD)
 				// Even if the bGroupAABBValid is invalid, it has been reset to invalid AABB, so the code below won't mark any page as used.
-				const bool bUseClusterAABB = HairGroupData->bClusterAABBValid && bIsGPUDriven;
+				const bool bUseClusterAABB = HairGroupData->GetClusterAABBValid() && bIsGPUDriven;
 				
 				FVoxelMarkValidPageIndex_PrepareCS::FParameters* Parameters = GraphBuilder.AllocParameters<FVoxelMarkValidPageIndex_PrepareCS::FParameters>();
 				Parameters->MaxClusterCount								= HairGroupData->GetClusterCount();
diff --git a/Engine/Source/Runtime/Renderer/Public/HairStrandsInterface.h b/Engine/Source/Runtime/Renderer/Public/HairStrandsInterface.h
--- a/Engine/Source/Runtime/Renderer/Public/HairStrandsInterface.h
+++ b/Engine/Source/Runtime/Renderer/Public/HairStrandsInterface.h
@@ -159,7 +159,6 @@ struct FRDGImportedBuffer
 RENDERER_API FRDGImportedBuffer Register(FRDGBuilder& GraphBuilder, const FRDGExternalBuffer& In, ERDGImportedBufferFlags Flags, ERDGUnorderedAccessViewFlags UAVFlags = ERDGUnorderedAccessViewFlags::None);
 RENDERER_API FRDGBufferSRVRef   RegisterAsSRV(FRDGBuilder& GraphBuilder, const FRDGExternalBuffer& In);
 RENDERER_API FRDGBufferUAVRef   RegisterAsUAV(FRDGBuilder& GraphBuilder, const FRDGExternalBuffer& In, ERDGUnorderedAccessViewFlags Flags = ERDGUnorderedAccessViewFlags::None);
-RENDERER_API void				ConvertToExternalBufferWithViews(FRDGBuilder& GraphBuilder, FRDGBufferRef& InBuffer, FRDGExternalBuffer& OutBuffer, EPixelFormat Format = PF_Unknown);
 ////////////////////////////////////////////////////////////////////////////////////////////////////
 // Misc/Helpers
 
@@ -214,37 +213,35 @@ protected:
 };
 typedef TArray<FHairStrandsInstance*> FHairStrandsInstances;
 
-class FHairGroupPublicData : public FRenderResource
+class FHairGroupPublicData
 {
 public:
 	RENDERER_API FHairGroupPublicData(uint32 InGroupIndex, const FName& OwnerName);
 	
-	RENDERER_API virtual void InitRHI(FRHICommandListBase& RHICmdList) override;
-	RENDERER_API virtual void ReleaseRHI() override;
-	virtual FString GetFriendlyName() const override { return TEXT("FHairGroupPublicData"); }
-	RENDERER_API void Allocate(FRDGBuilder& GraphBuilder);
-	RENDERER_API void Release();
-	RENDERER_API uint32 GetResourcesSize() const;
-
 	uint32 GetGroupIndex() const { return GroupIndex; }
 
-	FRDGExternalBuffer& GetDrawIndirectRasterComputeBuffer() { return DrawIndirectRasterComputeBuffer; }
-	const FRDGExternalBuffer& GetDrawIndirectRasterComputeBuffer() const { return DrawIndirectRasterComputeBuffer; }
-	FRDGExternalBuffer& GetDrawIndirectBuffer() { return DrawIndirectBuffer; }
-	FRDGExternalBuffer& GetClusterAABBBuffer() { return ClusterAABBBuffer; }
-	FRDGExternalBuffer& GetGroupAABBBuffer() { return GroupAABBBuffer; }
-	const FRDGExternalBuffer& GetGroupAABBBuffer() const { return GroupAABBBuffer; }
+	FRDGExternalBuffer& GetDrawIndirectRasterComputeBuffer() { return Culling->DrawIndirectRasterComputeBuffer; }
+	const FRDGExternalBuffer& GetDrawIndirectRasterComputeBuffer() const { return Culling->DrawIndirectRasterComputeBuffer; }
+	FRDGExternalBuffer& GetDrawIndirectBuffer() { return Culling->DrawIndirectBuffer; }
+	FRDGExternalBuffer& GetClusterAABBBuffer() { return Culling->ClusterAABBBuffer; }
+	FRDGExternalBuffer& GetGroupAABBBuffer() { return Culling->GroupAABBBuffer; }
+	const FRDGExternalBuffer& GetGroupAABBBuffer() const { return Culling->GroupAABBBuffer; }
 
-	const FRDGExternalBuffer& GetCulledCurveBuffer() const { return CulledCurveBuffer; }
-	const FRDGExternalBuffer& GetCulledVertexIdBuffer() const { return CulledVertexIdBuffer; }
-	const FRDGExternalBuffer& GetCulledVertexRadiusScaleBuffer() const { return CulledVertexRadiusScaleBuffer; }
+	const FRDGExternalBuffer& GetCulledCurveBuffer() const { return Culling->CulledCurveBuffer; }
+	const FRDGExternalBuffer& GetCulledVertexIdBuffer() const { return Culling->CulledVertexIdBuffer; }
+	const FRDGExternalBuffer& GetCulledVertexRadiusScaleBuffer() const { return Culling->CulledVertexRadiusScaleBuffer; }
 
-	FRDGExternalBuffer& GetCulledCurveBuffer() { return CulledCurveBuffer; }
-	FRDGExternalBuffer& GetCulledVertexIdBuffer() { return CulledVertexIdBuffer; }
-	FRDGExternalBuffer& GetCulledVertexRadiusScaleBuffer() { return CulledVertexRadiusScaleBuffer; }
+	FRDGExternalBuffer& GetCulledCurveBuffer() { return Culling->CulledCurveBuffer; }
+	FRDGExternalBuffer& GetCulledVertexIdBuffer() { return Culling->CulledVertexIdBuffer; }
+	FRDGExternalBuffer& GetCulledVertexRadiusScaleBuffer() { return Culling->CulledVertexRadiusScaleBuffer; }
 
-	bool GetCullingResultAvailable() const { return bCullingResultAvailable; }
-	void SetCullingResultAvailable(bool b) { bCullingResultAvailable = b; }
+	bool GetCullingResultAvailable() const { return Culling->bCullingResultAvailable; }
+	void SetCullingResultAvailable(bool b) { Culling->bCullingResultAvailable = b; }
+	
+	void SetClusterAABBValid(bool In) { Culling->bClusterAABBValid = In;  }
+	bool GetClusterAABBValid() const  { return Culling->bClusterAABBValid;  }
+
+	void SetGroupAABBValid(bool In) { Culling->bGroupAABBValid = In;  }
 
 	void SupportVoxelization(bool InVoxelize) { bSupportVoxelization = InVoxelize; }
 	bool DoesSupportVoxelization() const { return bSupportVoxelization; }
@@ -362,35 +359,39 @@ public:
 		FTransform LocalToWorldTransform;
 	};
 
+	struct FCulling
+	{
+		/* Indirect draw buffer to draw everything or the result of the culling per pass */
+		FRDGExternalBuffer DrawIndirectBuffer;
+		FRDGExternalBuffer DrawIndirectRasterComputeBuffer;
+
+		/* Hair Cluster & Hair Group bounding box buffer */
+		FRDGExternalBuffer ClusterAABBBuffer;
+		FRDGExternalBuffer GroupAABBBuffer;
+		bool bGroupAABBValid = false;
+		bool bClusterAABBValid = false;
+
+		/* Culling & LODing results for a hair group */ // Better to be transient?
+		FRDGExternalBuffer CulledCurveBuffer;
+		FRDGExternalBuffer CulledVertexIdBuffer;
+		FRDGExternalBuffer CulledVertexRadiusScaleBuffer;
+		bool bCullingResultAvailable = false;
+	};
+
 	// Current hook for retriving instance data. 
 	// This needs to be refactor to merge FHairGroupPublicData & HairStrandsInstance
 	FHairStrandsInstance* Instance = nullptr;
+	FCulling* Culling = nullptr;
 
 	FVertexFactoryInput VFInput;
 	uint32 ClusterDataIndex = ~0; // #hair_todo: move this into instance data, or remove FHairStrandClusterData
 
 	uint32 GroupIndex = 0;
-	uint32 ClusterCount = 0;
 	uint32 RestPointCount = 0;
 	uint32 RestCurveCount = 0;
+	uint32 ClusterCount = 0;
 
-	/* Indirect draw buffer to draw everything or the result of the culling per pass */
-	FRDGExternalBuffer DrawIndirectBuffer;
-	FRDGExternalBuffer DrawIndirectRasterComputeBuffer;
-
-	/* Hair Cluster & Hair Group bounding box buffer */
-	FRDGExternalBuffer ClusterAABBBuffer;
-	FRDGExternalBuffer GroupAABBBuffer;
-	bool bGroupAABBValid = false;
-	bool bClusterAABBValid = false;
-
-	/* Culling & LODing results for a hair group */ // Better to be transient?
-	FRDGExternalBuffer CulledCurveBuffer;
-	FRDGExternalBuffer CulledVertexIdBuffer;
-	FRDGExternalBuffer CulledVertexRadiusScaleBuffer;
-	bool bCullingResultAvailable = false;
 	bool bSupportVoxelization = true;
-	bool bIsInitialized = false;
 
 	/* CPU LOD selection. Hair LOD selection can be done by CPU or GPU. If bUseCPULODSelection is true, 
 	   CPU LOD selection is enabled otherwise the GPU selection is used. CPU LOD selection use the CPU 
