Cancel and wait for inflight distance field I/O request(s) when destroying render scene

- this change reverts 25911459 (Rel-25.00) and adds the cancellation and wait

#tests Tested by QA on 25.00
#rnx

diff --git a/Engine/Source/Runtime/Renderer/Private/DistanceFieldStreaming.cpp b/Engine/Source/Runtime/Renderer/Private/DistanceFieldStreaming.cpp
--- a/Engine/Source/Runtime/Renderer/Private/DistanceFieldStreaming.cpp
+++ b/Engine/Source/Runtime/Renderer/Private/DistanceFieldStreaming.cpp
@@ -452,17 +452,18 @@ struct FDistanceFieldReadRequest
 	const uint8* AlwaysLoadedDataPtr = nullptr;
 
 	// Inputs of read request
-	const FByteBulkData* BulkData = nullptr;
+	const FBulkData* BulkData = nullptr;
 	uint32 BulkOffset = 0;
 	uint32 BulkSize = 0;
 
+#if !WITH_EDITOR
 	// Outputs of read request
-	uint8* ReadOutputDataPtr = nullptr;
-	FIoRequest Request;
-	IAsyncReadFileHandle* AsyncHandle = nullptr;
-	IAsyncReadRequest* AsyncRequest = nullptr;
+	FBulkDataBatchReadRequest RequestHandle;
+	FIoBuffer RequestBuffer;
+#endif
 };
 
+
 struct FDistanceFieldAsyncUpdateParameters
 {
 	FDistanceFieldSceneData* DistanceFieldSceneData = nullptr;
@@ -497,6 +498,18 @@ FDistanceFieldSceneData::FDistanceFieldSceneData(EShaderPlatform ShaderPlatform)
 
 FDistanceFieldSceneData::~FDistanceFieldSceneData()
 {
+#if !WITH_EDITOR
+	// Wait for pending I/O request(s)
+	for (FDistanceFieldReadRequest& ReadRequest : ReadRequests)
+	{
+		if (ReadRequest.RequestHandle.IsCompleted() == false)
+		{
+			ReadRequest.RequestHandle.Cancel();
+		}
+		ReadRequest.RequestHandle.Reset();
+	}
+#endif
+
 	delete ObjectBuffers;
 	delete HeightFieldObjectBuffers;
 }
@@ -504,13 +517,13 @@ FDistanceFieldSceneData::~FDistanceFieldSceneData()
 class FDistanceFieldStreamingUpdateTask
 {
 public:
-	explicit FDistanceFieldStreamingUpdateTask(const FDistanceFieldAsyncUpdateParameters& InParams) : Parameters(InParams) {}
+	explicit FDistanceFieldStreamingUpdateTask(FDistanceFieldAsyncUpdateParameters&& InParams) : Parameters(MoveTemp(InParams)) {}
 
 	FDistanceFieldAsyncUpdateParameters Parameters;
 
 	void DoTask(ENamedThreads::Type CurrentThread, const FGraphEventRef& MyCompletionGraphEvent)
 	{
-		Parameters.DistanceFieldSceneData->AsyncUpdate(Parameters);
+		Parameters.DistanceFieldSceneData->AsyncUpdate(MoveTemp(Parameters));
 	}
 
 	static ESubsequentsMode::Type	GetSubsequentsMode()	{ return ESubsequentsMode::TrackSubsequents; }
@@ -518,7 +531,7 @@ public:
 	FORCEINLINE TStatId				GetStatId() const		{ return TStatId(); }
 };
 
-void FDistanceFieldSceneData::AsyncUpdate(FDistanceFieldAsyncUpdateParameters UpdateParameters)
+void FDistanceFieldSceneData::AsyncUpdate(FDistanceFieldAsyncUpdateParameters&& UpdateParameters)
 {
 	QUICK_SCOPE_CYCLE_COUNTER(STAT_FDistanceFieldSceneData_AsyncUpdate);
 	TRACE_CPUPROFILER_EVENT_SCOPE(FDistanceFieldSceneData::AsyncUpdate);
@@ -530,7 +543,7 @@ void FDistanceFieldSceneData::AsyncUpdate(FDistanceFieldAsyncUpdateParameters Up
 	int32 BrickUploadIndex = 0;
 	int32 IndirectionUploadIndex = 0;
 
-	for (FDistanceFieldReadRequest ReadRequest : UpdateParameters.ReadRequestsToUpload)
+	for (const FDistanceFieldReadRequest& ReadRequest : UpdateParameters.ReadRequestsToUpload)
 	{
 		const FDistanceFieldAssetState& AssetState = AssetStateArray[ReadRequest.AssetSetId];
 		const int32 ReversedMipIndex = ReadRequest.ReversedMipIndex;
@@ -538,15 +551,18 @@ void FDistanceFieldSceneData::AsyncUpdate(FDistanceFieldAsyncUpdateParameters Up
 		const int32 MipIndex = AssetState.BuiltData->Mips.Num() - ReversedMipIndex - 1;
 		const FSparseDistanceFieldMip& MipBuiltData = AssetState.BuiltData->Mips[MipIndex];
 
-		const uint8* BulkDataReadPtr = ReadRequest.BulkData ? ReadRequest.ReadOutputDataPtr : ReadRequest.AlwaysLoadedDataPtr;
+		const uint8* BulkDataReadPtr = ReadRequest.AlwaysLoadedDataPtr;
 
-#if WITH_EDITOR
 		if (ReadRequest.BulkData)
-		{	
+		{
+#if WITH_EDITOR
 			check((ReadRequest.BulkData->IsBulkDataLoaded() ||ReadRequest.BulkData->CanLoadFromDisk()) && ReadRequest.BulkData->GetBulkDataSize() > 0);
 			BulkDataReadPtr = (const uint8*)ReadRequest.BulkData->LockReadOnly() + ReadRequest.BulkOffset;
-		}
+#else
+			BulkDataReadPtr = ReadRequest.RequestBuffer.GetData();
 #endif
+		}
+
 		const int32 NumIndirectionEntries = MipBuiltData.IndirectionDimensions.X * MipBuiltData.IndirectionDimensions.Y * MipBuiltData.IndirectionDimensions.Z;
 		const uint32 ExpectedBulkSize = NumIndirectionEntries * sizeof(uint32) + ReadRequest.NumDistanceFieldBricks * BrickSizeBytes;
 
@@ -647,53 +663,21 @@ void FDistanceFieldSceneData::AsyncUpdate(FDistanceFieldAsyncUpdateParameters Up
 	}
 
 #if !WITH_EDITOR
-
-	for (FDistanceFieldReadRequest ReadRequest : UpdateParameters.ReadRequestsToCleanUp)
-	{
-		if (ReadRequest.AsyncRequest)
-		{
-			check(ReadRequest.AsyncRequest->PollCompletion());	
-			delete ReadRequest.AsyncRequest;
-			delete ReadRequest.AsyncHandle;
-		}
-		else
-		{
-			check(ReadRequest.Request.Status().IsCompleted());
-		}
-
-		FMemory::Free(ReadRequest.ReadOutputDataPtr);
-	}
-
-	FIoBatch Batch;
-
-	for (FDistanceFieldReadRequest& ReadRequest : UpdateParameters.NewReadRequests)
+	if (UpdateParameters.NewReadRequests.Num() > 0)
 	{
-		check(ReadRequest.BulkSize > 0);
-		ReadRequest.ReadOutputDataPtr = (uint8*)FMemory::Malloc(ReadRequest.BulkSize);
-		const bool bIODispatcher = ReadRequest.BulkData->IsUsingIODispatcher();
-
-		if (bIODispatcher)
-		{
-			// Use IODispatcher when available
-			FIoChunkId ChunkID = ReadRequest.BulkData->CreateChunkId();
-			FIoReadOptions ReadOptions;
-			ReadOptions.SetRange(ReadRequest.BulkData->GetBulkDataOffsetInFile() + ReadRequest.BulkOffset, ReadRequest.BulkSize);
-			ReadOptions.SetTargetVa(ReadRequest.ReadOutputDataPtr);
-			ReadRequest.Request = Batch.Read(ChunkID, ReadOptions, IoDispatcherPriority_Low);
-		}
-		else
+		FBulkDataBatchRequest::FBatchBuilder Batch = FBulkDataBatchRequest::NewBatch(UpdateParameters.NewReadRequests.Num());
+		for (FDistanceFieldReadRequest& ReadRequest : UpdateParameters.NewReadRequests)
 		{
-			// Compatibility path without IODispatcher
-			ReadRequest.AsyncHandle = ReadRequest.BulkData->OpenAsyncReadHandle();
-			ReadRequest.AsyncRequest = ReadRequest.AsyncHandle->ReadRequest(ReadRequest.BulkData->GetBulkDataOffsetInFile() + ReadRequest.BulkOffset, ReadRequest.BulkSize, AIOP_Low, nullptr, ReadRequest.ReadOutputDataPtr);
+			check(ReadRequest.BulkSize > 0);
+			Batch.Read(*ReadRequest.BulkData, uint64(ReadRequest.BulkOffset), uint64(ReadRequest.BulkSize), AIOP_Low, ReadRequest.RequestBuffer, ReadRequest.RequestHandle);
+			ReadRequests.Add(MoveTemp(ReadRequest));
 		}
+		FBulkDataRequest::EStatus Status = Batch.Issue();
+		UE_CLOG(Status != FBulkDataRequest::EStatus::Ok, LogDistanceField, Error, TEXT("Failed to issue bulk data I/O request"));
 	}
-
-	Batch.Issue();
-
-#endif
-
+#else
 	ReadRequests.Append(UpdateParameters.NewReadRequests);
+#endif
 }
 
 bool AssetHasOutstandingRequest(FSetElementId AssetSetId, const TArray<FDistanceFieldReadRequest>& ReadRequests)
@@ -838,8 +822,7 @@ void FDistanceFieldSceneData::ProcessStreamingRequestsFromGPU(
 void FDistanceFieldSceneData::ProcessReadRequests(
 	TArray<FDistanceFieldAssetMipId>& AssetDataUploads,
 	TArray<FDistanceFieldAssetMipId>& DistanceFieldAssetMipAdds,
-	TArray<FDistanceFieldReadRequest>& ReadRequestsToUpload,
-	TArray<FDistanceFieldReadRequest>& ReadRequestsToCleanUp)
+	TArray<FDistanceFieldReadRequest>& ReadRequestsToUpload)
 {
 	const uint32 BrickSizeBytes = GPixelFormats[DistanceField::DistanceFieldFormat].BlockBytes * DistanceField::BrickSize * DistanceField::BrickSize * DistanceField::BrickSize;
 	const SIZE_T TextureUploadLimitBytes = (SIZE_T)CVarTextureUploadLimitKBytes.GetValueOnRenderThread() * 1024;
@@ -869,41 +852,33 @@ void FDistanceFieldSceneData::ProcessReadRequests(
 
 	for (int32 RequestIndex = 0; RequestIndex < ReadRequests.Num(); RequestIndex++)
 	{
-		const FDistanceFieldReadRequest ReadRequest = ReadRequests[RequestIndex];
-
-		bool bReady = true;
-
 #if !WITH_EDITOR
-		if (ReadRequest.AsyncRequest)
-		{
-			bReady = bReady && ReadRequest.AsyncRequest->PollCompletion();
-		}
-		else
+		if (ReadRequests[RequestIndex].RequestHandle.IsCompleted() == false)
 		{
-			bReady = bReady && ReadRequest.Request.Status().IsCompleted();
+			continue;
 		}
 #endif
 
-		if (bReady)
+		FDistanceFieldReadRequest CompletedRequest = MoveTemp(ReadRequests[RequestIndex]);
+#if !WITH_EDITOR
+		CompletedRequest.RequestHandle = FBulkDataRequest();
+#endif
+
+		ReadRequests.RemoveAtSwap(RequestIndex--);
+
+		if (AssetStateArray.IsValidId(CompletedRequest.AssetSetId) 
+			// Prevent attempting to upload after a different asset has been allocated at the same index
+			&& CompletedRequest.BuiltDataId == AssetStateArray[CompletedRequest.AssetSetId].BuiltData->GetId()
+			// Shader requires sequential reversed mips starting from 0, skip upload if the IO request got out of sync with the streaming feedback requests
+			&& CompletedRequest.ReversedMipIndex == AssetStateArray[CompletedRequest.AssetSetId].ReversedMips.Num())
 		{
-			ReadRequests.RemoveAtSwap(RequestIndex);
-			RequestIndex--;
-
-			if (AssetStateArray.IsValidId(ReadRequest.AssetSetId) 
-				// Prevent attempting to upload after a different asset has been allocated at the same index
-				&& ReadRequest.BuiltDataId == AssetStateArray[ReadRequest.AssetSetId].BuiltData->GetId()
-				// Shader requires sequential reversed mips starting from 0, skip upload if the IO request got out of sync with the streaming feedback requests
-				&& ReadRequest.ReversedMipIndex == AssetStateArray[ReadRequest.AssetSetId].ReversedMips.Num())
-			{
-				TextureUploadBytes += ReadRequest.NumDistanceFieldBricks * BrickSizeBytes;
+			TextureUploadBytes += CompletedRequest.NumDistanceFieldBricks * BrickSizeBytes;
 
-				DistanceFieldAssetMipAdds.Add(FDistanceFieldAssetMipId(ReadRequest.AssetSetId, ReadRequest.ReversedMipIndex));
-				// Re-upload mip0 to push the new NumMips to the shader
-				AssetDataUploads.Add(FDistanceFieldAssetMipId(ReadRequest.AssetSetId, 0));
-				ReadRequestsToUpload.Add(ReadRequest);
-			}
+			DistanceFieldAssetMipAdds.Add(FDistanceFieldAssetMipId(CompletedRequest.AssetSetId, CompletedRequest.ReversedMipIndex));
+			// Re-upload mip0 to push the new NumMips to the shader
+			AssetDataUploads.Add(FDistanceFieldAssetMipId(CompletedRequest.AssetSetId, 0));
 
-			ReadRequestsToCleanUp.Add(ReadRequest);
+			ReadRequestsToUpload.Add(MoveTemp(CompletedRequest));
 		}
 
 		// Stop uploading when we reach the limit
@@ -1403,9 +1378,8 @@ void FDistanceFieldSceneData::UpdateDistanceFieldAtlas(
 	ProcessStreamingRequestsFromGPU(NewReadRequests, AssetDataUploads);
 
 	TArray<FDistanceFieldReadRequest> ReadRequestsToUpload;
-	TArray<FDistanceFieldReadRequest> ReadRequestsToCleanUp;
 	// Build a list of completed read requests that should be uploaded to the GPU this frame
-	ProcessReadRequests(AssetDataUploads, DistanceFieldAssetMipAdds, ReadRequestsToUpload, ReadRequestsToCleanUp);
+	ProcessReadRequests(AssetDataUploads, DistanceFieldAssetMipAdds, ReadRequestsToUpload);
 
 	int32 NumIndirectionTableAdds = 0;
 	int32 NumBrickUploads = 0;
@@ -1529,14 +1503,13 @@ void FDistanceFieldSceneData::UpdateDistanceFieldAtlas(
 			UpdateParameters.BrickUploadCoordinatesPtr = AtlasUpload.BrickUploadCoordinatesPtr;
 		}
 
-		if (NewReadRequests.Num() || ReadRequestsToUpload.Num() || ReadRequestsToCleanUp.Num())
+		if (NewReadRequests.Num() || ReadRequestsToUpload.Num())
 		{
 			UpdateParameters.NewReadRequests = MoveTemp(NewReadRequests);
 			UpdateParameters.ReadRequestsToUpload = MoveTemp(ReadRequestsToUpload);
-			UpdateParameters.ReadRequestsToCleanUp = MoveTemp(ReadRequestsToCleanUp);
 
 			// TODO: We actually run this synchronously now after the RDG conversion, as it would otherwise immediately sync.
-			AsyncUpdate(UpdateParameters);
+			AsyncUpdate(MoveTemp(UpdateParameters));
 		}
 
 		if (NumBrickUploads > 0 || NumIndirectionTableAdds > 0)
diff --git a/Engine/Source/Runtime/Renderer/Private/ScenePrivate.h b/Engine/Source/Runtime/Renderer/Private/ScenePrivate.h
--- a/Engine/Source/Runtime/Renderer/Private/ScenePrivate.h
+++ b/Engine/Source/Runtime/Renderer/Private/ScenePrivate.h
@@ -2114,8 +2114,7 @@ private:
 	void ProcessReadRequests(
 		TArray<FDistanceFieldAssetMipId>& AssetDataUploads,
 		TArray<FDistanceFieldAssetMipId>& DistanceFieldAssetMipAdds,
-		TArray<FDistanceFieldReadRequest>& ReadRequestsToUpload,
-		TArray<FDistanceFieldReadRequest>& ReadRequestsToCleanUp);
+		TArray<FDistanceFieldReadRequest>& ReadRequestsToUpload);
 
 	FRDGTexture* ResizeBrickAtlasIfNeeded(FRDGBuilder& GraphBuilder, FGlobalShaderMap* GlobalShaderMap);
 
@@ -2127,7 +2126,7 @@ private:
 	
 	void UploadAllAssetData(FRDGBuilder& GraphBuilder, FRDGBuffer* AssetDataBufferRDG);
 
-	void AsyncUpdate(FDistanceFieldAsyncUpdateParameters UpdateParameters);
+	void AsyncUpdate(FDistanceFieldAsyncUpdateParameters&& UpdateParameters);
 
 	void GenerateStreamingRequests(
 		FRDGBuilder& GraphBuilder, 
