1. Sort depth pass instead of base pass if full depth pre-pass is enabled.
2. Disable auto instancing for the translucent passes on mobile

diff --git a/Engine/Source/Runtime/Renderer/Private/MeshDrawCommands.cpp b/Engine/Source/Runtime/Renderer/Private/MeshDrawCommands.cpp
--- a/Engine/Source/Runtime/Renderer/Private/MeshDrawCommands.cpp
+++ b/Engine/Source/Runtime/Renderer/Private/MeshDrawCommands.cpp
@@ -325,13 +325,13 @@ static uint64 GetMobileBasePassSortKey_ByState(bool bMasked, bool bBackground, u
 /**
 * Compute mesh sort keys for the mobile base pass
 */
-void UpdateMobileBasePassMeshSortKeys(
+void UpdateMobilePassMeshSortKeys(
 	const FVector& ViewOrigin,
 	const TScenePrimitiveArray<FPrimitiveBounds>& ScenePrimitiveBounds,
 	FMeshCommandOneFrameArray& VisibleMeshCommands
 )
 {
-	QUICK_SCOPE_CYCLE_COUNTER(STAT_UpdateMobileBasePassMeshSortKeys);
+	QUICK_SCOPE_CYCLE_COUNTER(STAT_UpdateMobilePassMeshSortKeys);
 
 	// Object radius past which we treat object as part of 'background'
 	constexpr float MIN_BACKGROUND_OBJECT_RADIUS = 500000;
@@ -346,8 +346,8 @@ void UpdateMobileBasePassMeshSortKeys(
 		{
 			FVisibleMeshDrawCommand& Cmd = VisibleMeshCommands[CmdIdx];
 			// Set in MobileBasePass.cpp - GetBasePassStaticSortKey;
-			bool bMasked = Cmd.SortKey.PackedData & 0x1 ? true : false;
-			bool bBackground = Cmd.SortKey.PackedData & 0x2 ? true : false;
+			bool bMasked = Cmd.SortKey.BasePass.Masked == 1;
+			bool bBackground = Cmd.SortKey.BasePass.Background == 1;
 			float PrimitiveDistance = 0;
 			if (Cmd.PrimitiveIdInfo.ScenePrimitiveId < ScenePrimitiveBounds.Num())
 			{
@@ -405,8 +405,8 @@ void UpdateMobileBasePassMeshSortKeys(
 		{
 			FVisibleMeshDrawCommand& Cmd = VisibleMeshCommands[CmdIdx];
 			// Set in MobileBasePass.cpp - GetBasePassStaticSortKey;
-			bool bMasked = Cmd.SortKey.PackedData & 0x1 ? true : false;
-			bool bBackground = Cmd.SortKey.PackedData & 0x2 ? true : false;
+			bool bMasked = Cmd.SortKey.BasePass.Masked == 1;
+			bool bBackground = Cmd.SortKey.BasePass.Background == 1;
 			float PrimitiveDistance = 0;
 			if (Cmd.PrimitiveIdInfo.ScenePrimitiveId < ScenePrimitiveBounds.Num())
 			{
@@ -892,10 +892,14 @@ public:
 	{
 		FOptionalTaskTagScope Scope(ETaskTag::EParallelRenderingThread);
 		SCOPED_NAMED_EVENT(MeshDrawCommandPassSetupTask, FColor::Magenta);
+		// On SM5 Mobile platform, still want the same sorting
+		const bool bMobile = Context.ShadingPath == EShadingPath::Mobile || IsVulkanMobileSM5Platform(Context.ShaderPlatform);
 		// Mobile base pass is a special case, as final lists is created from two mesh passes based on CSM visibility.
 		const bool bMobileShadingBasePass = Context.ShadingPath == EShadingPath::Mobile && Context.PassType == EMeshPass::BasePass;
 		// On SM5 Mobile platform, still want the same sorting
-		const bool bMobileVulkanSM5BasePass = IsVulkanMobileSM5Platform(Context.ShaderPlatform) && Context.PassType == EMeshPass::BasePass;
+		// Mobile pass expects mesh sorting
+		const bool bNeedsUpdateMobilePassMeshSortKeys = bMobile && (Context.PassType == EMeshPass::DepthPass
+																|| (Context.PassType == EMeshPass::BasePass && Context.Scene->EarlyZPassMode != DDM_AllOpaque));
 
 		if (bMobileShadingBasePass)
 		{
@@ -962,9 +966,9 @@ public:
 			}
 
 			// Update sort keys.
-			if (bMobileShadingBasePass || bMobileVulkanSM5BasePass)
+			if (bNeedsUpdateMobilePassMeshSortKeys)
 			{
-				UpdateMobileBasePassMeshSortKeys(
+				UpdateMobilePassMeshSortKeys(
 					Context.ViewOrigin,
 					*Context.PrimitiveBounds,
 					Context.MeshDrawCommands
diff --git a/Engine/Source/Runtime/Renderer/Private/MobileBasePass.cpp b/Engine/Source/Runtime/Renderer/Private/MobileBasePass.cpp
--- a/Engine/Source/Runtime/Renderer/Private/MobileBasePass.cpp
+++ b/Engine/Source/Runtime/Renderer/Private/MobileBasePass.cpp
@@ -663,11 +663,13 @@ void MobileBasePass::SetTranslucentRenderState(FMeshPassProcessorRenderState& Dr
 	}
 }
 
-static FMeshDrawCommandSortKey GetBasePassStaticSortKey(const bool bIsMasked, bool bBackground)
+static FMeshDrawCommandSortKey GetBasePassStaticSortKey(const bool bIsMasked, bool bBackground, const FMeshMaterialShader* VertexShader, const FMeshMaterialShader* PixelShader)
 {
 	FMeshDrawCommandSortKey SortKey;
-	SortKey.PackedData = (bIsMasked ? 1 : 0);
-	SortKey.PackedData|= (bBackground ? 2 : 0); // background flag in second bit
+	SortKey.BasePass.Masked = (bIsMasked ? 1 : 0);
+	SortKey.BasePass.Background = (bBackground ? 1 : 0); // background flag in second bit
+	SortKey.BasePass.VertexShaderHash = (VertexShader ? VertexShader->GetSortKey() : 0) & 0xFFFF;
+	SortKey.BasePass.PixelShaderHash = PixelShader ? PixelShader->GetSortKey() : 0;
 	return SortKey;
 }
 
@@ -900,7 +902,7 @@ bool FMobileBasePassMeshProcessor::Process(
 		bool bBackground = PrimitiveSceneProxy ? PrimitiveSceneProxy->TreatAsBackgroundForOcclusion() : false;
 		// Default static sort key separates masked and non-masked geometry, generic mesh sorting will also sort by PSO
 		// if platform wants front to back sorting, this key will be recomputed in InitViews
-		SortKey = GetBasePassStaticSortKey(bIsMasked, bBackground);
+		SortKey = GetBasePassStaticSortKey(bIsMasked, bBackground, BasePassShaders.VertexShader.GetShader(), BasePassShaders.PixelShader.GetShader());
 	}
 	
 	const FMeshDrawingPolicyOverrideSettings OverrideSettings = ComputeMeshOverrideSettings(MeshBatch);
@@ -1134,7 +1136,7 @@ FMeshPassProcessor* CreateMobileTranslucencyAllPassProcessor(ERHIFeatureLevel::T
 
 REGISTER_MESHPASSPROCESSOR_AND_PSOCOLLECTOR(MobileBasePass, 			CreateMobileBasePassProcessor, 			EShadingPath::Mobile, EMeshPass::BasePass, 		EMeshPassFlags::CachedMeshCommands | EMeshPassFlags::MainView);
 REGISTER_MESHPASSPROCESSOR_AND_PSOCOLLECTOR(MobileBasePassCSM,			CreateMobileBasePassCSMProcessor,		EShadingPath::Mobile, EMeshPass::MobileBasePassCSM, 	EMeshPassFlags::CachedMeshCommands | EMeshPassFlags::MainView);
-REGISTER_MESHPASSPROCESSOR_AND_PSOCOLLECTOR(MobileTranslucencyAllPass,		CreateMobileTranslucencyAllPassProcessor,	EShadingPath::Mobile, EMeshPass::TranslucencyAll, 	EMeshPassFlags::CachedMeshCommands | EMeshPassFlags::MainView);
-REGISTER_MESHPASSPROCESSOR_AND_PSOCOLLECTOR(MobileTranslucencyStandardPass,	CreateMobileTranslucencyStandardPassProcessor,	EShadingPath::Mobile, EMeshPass::TranslucencyStandard, 	EMeshPassFlags::CachedMeshCommands | EMeshPassFlags::MainView);
-REGISTER_MESHPASSPROCESSOR_AND_PSOCOLLECTOR(MobileTranslucencyAfterDOFPass,	CreateMobileTranslucencyAfterDOFProcessor,	EShadingPath::Mobile, EMeshPass::TranslucencyAfterDOF, 	EMeshPassFlags::CachedMeshCommands | EMeshPassFlags::MainView);
+REGISTER_MESHPASSPROCESSOR_AND_PSOCOLLECTOR(MobileTranslucencyAllPass,		CreateMobileTranslucencyAllPassProcessor,	EShadingPath::Mobile, EMeshPass::TranslucencyAll, 	EMeshPassFlags::MainView);
+REGISTER_MESHPASSPROCESSOR_AND_PSOCOLLECTOR(MobileTranslucencyStandardPass,	CreateMobileTranslucencyStandardPassProcessor,	EShadingPath::Mobile, EMeshPass::TranslucencyStandard, 	EMeshPassFlags::MainView);
+REGISTER_MESHPASSPROCESSOR_AND_PSOCOLLECTOR(MobileTranslucencyAfterDOFPass,	CreateMobileTranslucencyAfterDOFProcessor,	EShadingPath::Mobile, EMeshPass::TranslucencyAfterDOF, 	EMeshPassFlags::MainView);
 // Skipping EMeshPass::TranslucencyAfterDOFModulate because dual blending is not supported on mobile
diff --git a/Engine/Source/Runtime/Renderer/Public/MeshPassProcessor.h b/Engine/Source/Runtime/Renderer/Public/MeshPassProcessor.h
--- a/Engine/Source/Runtime/Renderer/Public/MeshPassProcessor.h
+++ b/Engine/Source/Runtime/Renderer/Public/MeshPassProcessor.h
@@ -1295,7 +1295,8 @@ public:
 		{
 			uint64 VertexShaderHash		: 16; // Order by vertex shader's hash.
 			uint64 PixelShaderHash		: 32; // Order by pixel shader's hash.
-			uint64 Masked				: 16; // First order by masked.
+			uint64 Background			: 1;
+			uint64 Masked				: 15; // First order by masked.
 		} BasePass;
 
 		struct
