Changes FTranslucencyComposition to accept a FScreenPassTextureSlice SceneColor

This change is motivated to avoids a texture copy when compositing post motion blur translucency right after TSR when it uses a Texture2DArray history

diff --git a/Engine/Source/Runtime/Renderer/Private/PostProcess/PostProcessing.cpp b/Engine/Source/Runtime/Renderer/Private/PostProcess/PostProcessing.cpp
--- a/Engine/Source/Runtime/Renderer/Private/PostProcess/PostProcessing.cpp
+++ b/Engine/Source/Runtime/Renderer/Private/PostProcess/PostProcessing.cpp
@@ -694,7 +694,7 @@ void AddPostProcessingPasses(
 				FTranslucencyComposition TranslucencyComposition;
 				TranslucencyComposition.Operation = FTranslucencyComposition::EOperation::ComposeToNewSceneColor;
 				TranslucencyComposition.bApplyModulateOnly = bComposeSeparateTranslucencyInTSR;
-				TranslucencyComposition.SceneColor = SceneColor;
+				TranslucencyComposition.SceneColor = FScreenPassTextureSlice::CreateFromScreenPassTexture(GraphBuilder, SceneColor);
 				TranslucencyComposition.SceneDepth = SceneDepth;
 				TranslucencyComposition.OutputViewport = FScreenPassTextureViewport(SceneColor);
 				TranslucencyComposition.OutputPixelFormat = SceneColorFormat;
@@ -889,7 +889,7 @@ void AddPostProcessingPasses(
 			// Compose Post-MotionBlur translucency in a new scene color to ensure it's not writing out in TAA's output that is also the history.
 			FTranslucencyComposition TranslucencyComposition;
 			TranslucencyComposition.Operation = FTranslucencyComposition::EOperation::ComposeToNewSceneColor;
-			TranslucencyComposition.SceneColor = FScreenPassTexture::CopyFromSlice(GraphBuilder, SceneColorSlice); // TODO
+			TranslucencyComposition.SceneColor = SceneColorSlice;
 			TranslucencyComposition.OutputViewport = FScreenPassTextureViewport(SceneColorSlice);
 			TranslucencyComposition.OutputPixelFormat = SceneColorFormat;
 
@@ -1235,23 +1235,22 @@ void AddPostProcessingPasses(
 		{
 			FTranslucencyComposition TranslucencyComposition;
 			TranslucencyComposition.Operation = FTranslucencyComposition::EOperation::ComposeToNewSceneColor;
-			TranslucencyComposition.SceneColor = SceneColor;
 			TranslucencyComposition.OutputViewport = FScreenPassTextureViewport(SceneColor);
 			TranslucencyComposition.OutputPixelFormat = SceneColorFormat;
 
 			if (PostDOFTranslucencyResources.IsValid())
 			{
-				TranslucencyComposition.SceneColor = TranslucencyComposition.AddPass(
+				TranslucencyComposition.SceneColor = FScreenPassTextureSlice::CreateFromScreenPassTexture(GraphBuilder, SceneColor);
+				SceneColor = TranslucencyComposition.AddPass(
 					GraphBuilder, View, PostDOFTranslucencyResources);
 			}
 
 			if (PostMotionBlurTranslucencyResources.IsValid())
 			{
-				TranslucencyComposition.SceneColor = TranslucencyComposition.AddPass(
+				TranslucencyComposition.SceneColor = FScreenPassTextureSlice::CreateFromScreenPassTexture(GraphBuilder, SceneColor);
+				SceneColor = TranslucencyComposition.AddPass(
 					GraphBuilder, View, PostMotionBlurTranslucencyResources);
 			}
-
-			SceneColor = TranslucencyComposition.SceneColor;
 		}
 
 		SceneColorBeforeTonemapSlice = FScreenPassTextureSlice::CreateFromScreenPassTexture(GraphBuilder, SceneColor);
diff --git a/Engine/Source/Runtime/Renderer/Private/TranslucentRendering.cpp b/Engine/Source/Runtime/Renderer/Private/TranslucentRendering.cpp
--- a/Engine/Source/Runtime/Renderer/Private/TranslucentRendering.cpp
+++ b/Engine/Source/Runtime/Renderer/Private/TranslucentRendering.cpp
@@ -402,7 +402,7 @@ class FComposeSeparateTranslucencyPS : public FGlobalShader
 		SHADER_PARAMETER(FVector2f, SeparateTranslucencyUVMax)
 		SHADER_PARAMETER(FVector2f, SeparateTranslucencyExtentInverse)
 
-		SHADER_PARAMETER_RDG_TEXTURE(Texture2D, SceneColorTexture)
+		SHADER_PARAMETER_RDG_TEXTURE_SRV(Texture2D, SceneColorTexture)
 		SHADER_PARAMETER_SAMPLER(SamplerState,  SceneColorSampler)
 
 		SHADER_PARAMETER_RDG_TEXTURE(Texture2D, SeparateTranslucencyPointTexture)
@@ -464,7 +464,7 @@ FScreenPassTexture FTranslucencyComposition::AddPass(
 	ensure(TranslucencyTextures.IsValid());
 	if (!TranslucencyTextures.IsValid())
 	{
-		return SceneColor;
+		return FScreenPassTexture(SceneColor);
 	}
 
 	ensure(TranslucencyTextures.Pass != ETranslucencyPass::TPT_MAX);
@@ -497,7 +497,7 @@ FScreenPassTexture FTranslucencyComposition::AddPass(
 	{
 		if (!TranslucencyTextures.ColorModulateTexture.IsValid())
 		{
-			return SceneColor;
+			return FScreenPassTexture(SceneColor);
 		}
 
 		SeparateTranslucencyTexture = GraphBuilder.RegisterExternalTexture(GSystemTextures.BlackAlphaOneDummy);
@@ -534,8 +534,8 @@ FScreenPassTexture FTranslucencyComposition::AddPass(
 		OpName = TEXT("ComposeTranslucencyToExistingColor");
 		BlendState = TStaticBlendState<CW_RGB, BO_Add, BF_One, BF_SourceAlpha>::GetRHI();
 
-		ensure(SceneColor.Texture->Desc.Flags & TexCreate_RenderTargetable);
-		NewSceneColor = SceneColor.Texture;
+		ensure(SceneColor.TextureSRV->Desc.Texture->Desc.Flags & TexCreate_RenderTargetable);
+		NewSceneColor = SceneColor.TextureSRV->Desc.Texture;
 	}
 	else if (Operation == EOperation::ComposeToNewSceneColor)
 	{
@@ -545,7 +545,7 @@ FScreenPassTexture FTranslucencyComposition::AddPass(
 
 		FRDGTextureDesc OutputDesc = FRDGTextureDesc::Create2D(
 			OutputViewport.Extent,
-			OutputPixelFormat != PF_Unknown ? OutputPixelFormat : SceneColor.Texture->Desc.Format,
+			OutputPixelFormat != PF_Unknown ? OutputPixelFormat : SceneColor.TextureSRV->Desc.Texture->Desc.Format,
 			FClearValueBinding::Black,
 			TexCreate_RenderTargetable | TexCreate_ShaderResource);
 
@@ -582,8 +582,8 @@ FScreenPassTexture FTranslucencyComposition::AddPass(
 	PassParameters->SeparateTranslucencyExtentInverse = SeparateTranslucencyExtentInv;
 	
 	PassParameters->SceneColorTexture = Operation == EOperation::ComposeToNewSceneColor
-		? SceneColor.Texture
-		: GraphBuilder.RegisterExternalTexture(GSystemTextures.BlackAlphaOneDummy);
+		? SceneColor.TextureSRV
+		: GraphBuilder.CreateSRV(FRDGTextureSRVDesc(GraphBuilder.RegisterExternalTexture(GSystemTextures.BlackAlphaOneDummy)));
 	PassParameters->SceneColorSampler = TStaticSamplerState<SF_Point>::GetRHI();
 
 	PassParameters->SeparateTranslucencyPointTexture = SeparateTranslucencyTexture;
@@ -1415,7 +1415,7 @@ void FDeferredShadingSceneRenderer::RenderTranslucencyInner(
 
 				FTranslucencyComposition TranslucencyComposition;
 				TranslucencyComposition.Operation = FTranslucencyComposition::EOperation::ComposeToExistingSceneColor;
-				TranslucencyComposition.SceneColor = FScreenPassTexture(SceneTextures.Color.Target, View.ViewRect);
+				TranslucencyComposition.SceneColor = FScreenPassTextureSlice::CreateFromScreenPassTexture(GraphBuilder, FScreenPassTexture(SceneTextures.Color.Target, View.ViewRect));
 				TranslucencyComposition.SceneDepth = FScreenPassTexture(SceneTextures.Depth.Resolve, View.ViewRect);
 				TranslucencyComposition.OutputViewport = FScreenPassTextureViewport(SceneTextures.Depth.Resolve, View.ViewRect);
 
diff --git a/Engine/Source/Runtime/Renderer/Private/TranslucentRendering.h b/Engine/Source/Runtime/Renderer/Private/TranslucentRendering.h
--- a/Engine/Source/Runtime/Renderer/Private/TranslucentRendering.h
+++ b/Engine/Source/Runtime/Renderer/Private/TranslucentRendering.h
@@ -167,7 +167,7 @@ struct FTranslucencyComposition
 	EOperation Operation = EOperation::UpscaleOnly;
 	bool bApplyModulateOnly = false;
 
-	FScreenPassTexture SceneColor;
+	FScreenPassTextureSlice SceneColor;
 	FScreenPassTexture SceneDepth;
 
 	FScreenPassTextureViewport OutputViewport;
