Fix static analysis warning.

diff --git a/Engine/Source/Runtime/Renderer/Private/LightRendering.cpp b/Engine/Source/Runtime/Renderer/Private/LightRendering.cpp
--- a/Engine/Source/Runtime/Renderer/Private/LightRendering.cpp
+++ b/Engine/Source/Runtime/Renderer/Private/LightRendering.cpp
@@ -1483,6 +1483,7 @@ void FDeferredShadingSceneRenderer::RenderLights(
 
 					// Inline ray traced shadow batching, launches shadow batches when needed
 					// reduces memory overhead while keeping shadows batched to optimize costs
+#if RHI_RAYTRACING
 					{
 						const uint32 ViewIndex = 0;
 						FViewInfo& View = Views[ViewIndex];
@@ -1498,8 +1499,7 @@ void FDeferredShadingSceneRenderer::RenderLights(
 							SortedLightInfo.SortKey.Fields.bShadowed;
 
 						// determine if this light doesn't yet have a precomputed shadow and execute a batch to amortize costs if one is needed
-						if (bEnableRayTracing &&
-							bWantsBatchedShadow &&
+						if (bWantsBatchedShadow &&
 							(PreprocessedShadowMaskTextures.Num() == 0 || !PreprocessedShadowMaskTextures[LightIndex - UnbatchedLightStart]))
 						{
 							RDG_EVENT_SCOPE(GraphBuilder, "ShadowBatch");
@@ -1700,7 +1700,7 @@ void FDeferredShadingSceneRenderer::RenderLights(
 						}
 					} // end inline batched raytraced shadow
 
-					if (bEnableRayTracing && PreprocessedShadowMaskTextures.Num() > 0 && PreprocessedShadowMaskTextures[LightIndex - UnbatchedLightStart])
+					if (PreprocessedShadowMaskTextures.Num() > 0 && PreprocessedShadowMaskTextures[LightIndex - UnbatchedLightStart])
 					{
 						const uint32 ShadowMaskIndex = LightIndex - UnbatchedLightStart;
 						ScreenShadowMaskTexture = PreprocessedShadowMaskTextures[ShadowMaskIndex];
@@ -1719,7 +1719,9 @@ void FDeferredShadingSceneRenderer::RenderLights(
 							RenderHairStrandsDeepShadowMask(GraphBuilder, Views, &LightSceneInfo, ScreenShadowMaskTexture);
 						}
 					}
-					else if (OcclusionType == FLightOcclusionType::Raytraced)
+					else 
+#endif // RHI_RAYTRACING
+					if (OcclusionType == FLightOcclusionType::Raytraced)
 					{
 						FSceneTextureParameters SceneTextureParameters = GetSceneTextureParameters(GraphBuilder, SceneTextures.UniformBuffer);
 
