Fix overflow when computing rect light texture MIP causing incorrect atlas packing.

diff --git a/Engine/Source/Runtime/Renderer/Private/RectLightTextureManager.cpp b/Engine/Source/Runtime/Renderer/Private/RectLightTextureManager.cpp
--- a/Engine/Source/Runtime/Renderer/Private/RectLightTextureManager.cpp
+++ b/Engine/Source/Runtime/Renderer/Private/RectLightTextureManager.cpp
@@ -225,7 +225,8 @@ static FIntPoint ToMIP(const FIntPoint& InMip, uint32 MipIndex)
 
 static int32 GetSlotMaxMIPLevel(const FAtlasSlot& In)
 {
-	return FMath::FloorLog2(FMath::Min(In.Rect.Resolution.X, In.Rect.Resolution.Y)) - 1;
+	const uint32 MaxMIP = FMath::FloorLog2(FMath::Min(In.Rect.Resolution.X, In.Rect.Resolution.Y));
+	return MaxMIP > 0u ? MaxMIP-1u : 0u;
 }
 
 static bool Traits_IsValid(const FAtlasRect& In)					{ return true; }
