Fixes validation failure in full res FFT bloom when enabling TSR outputs Texture2DArray slice

diff --git a/Engine/Source/Runtime/Renderer/Private/ScreenPass.cpp b/Engine/Source/Runtime/Renderer/Private/ScreenPass.cpp
--- a/Engine/Source/Runtime/Renderer/Private/ScreenPass.cpp
+++ b/Engine/Source/Runtime/Renderer/Private/ScreenPass.cpp
@@ -56,6 +56,7 @@ FScreenPassTexture FScreenPassTexture::CopyFromSlice(FRDGBuilder& GraphBuilder,
 
 	FRDGTextureDesc Desc = ScreenTextureSlice.TextureSRV->Desc.Texture->Desc;
 	Desc.Dimension = ETextureDimension::Texture2D;
+	Desc.ArraySize = 1;
 
 	FRDGTextureRef NewTexture = GraphBuilder.CreateTexture(Desc, TEXT("CopyToScreenPassTexture2D"));
 
