Removes texture copy pass when TSR output Texture2DArray slice but r.Bloom.ScreenPercentage=100

diff --git a/Engine/Source/Runtime/Renderer/Private/GPUFastFourierTransform.cpp b/Engine/Source/Runtime/Renderer/Private/GPUFastFourierTransform.cpp
--- a/Engine/Source/Runtime/Renderer/Private/GPUFastFourierTransform.cpp
+++ b/Engine/Source/Runtime/Renderer/Private/GPUFastFourierTransform.cpp
@@ -178,7 +178,7 @@ class FReorderFFTPassCS : public FFFTShader
 		SHADER_PARAMETER(int32, LogTwoLength)
 		SHADER_PARAMETER(int32, BitCount)
 
-		SHADER_PARAMETER_RDG_TEXTURE(Texture2D, SrcSRV)
+		SHADER_PARAMETER_RDG_TEXTURE_SRV(Texture2D, SrcSRV)
 		SHADER_PARAMETER_RDG_TEXTURE_UAV(RWTexture2D, DstUAV)
 	END_SHADER_PARAMETER_STRUCT()
 
@@ -201,7 +201,7 @@ class FGroupShardSubFFTPassCS : public FFFTShader
 		SHADER_PARAMETER(int32, TransformLength)
 		SHADER_PARAMETER(int32, NumSubRegions)
 
-		SHADER_PARAMETER_RDG_TEXTURE(Texture2D, SrcTexture)
+		SHADER_PARAMETER_RDG_TEXTURE_SRV(Texture2D, SrcTexture)
 		SHADER_PARAMETER_RDG_TEXTURE_UAV(RWTexture2D, DstTexture)
 	END_SHADER_PARAMETER_STRUCT()
 
@@ -229,7 +229,7 @@ class FComplexFFTPassCS : public FFFTShader
 		SHADER_PARAMETER(int32, PowTwoLength)
 		SHADER_PARAMETER(int32, BitCount)
 
-		SHADER_PARAMETER_RDG_TEXTURE(Texture2D, SrcSRV)
+		SHADER_PARAMETER_RDG_TEXTURE_SRV(Texture2D, SrcSRV)
 		SHADER_PARAMETER_RDG_BUFFER_SRV(StructuredBuffer<float4>, DstPostFilterParameters)
 		SHADER_PARAMETER_RDG_TEXTURE_UAV(RWTexture2D, DstUAV)
 	END_SHADER_PARAMETER_STRUCT()
@@ -251,7 +251,7 @@ class FPackTwoForOneFFTPassCS : public FFFTShader
 		SHADER_PARAMETER(FIntRect, DstRect)
 		SHADER_PARAMETER(int32, TransformType)
 
-		SHADER_PARAMETER_RDG_TEXTURE(Texture2D, SrcSRV)
+		SHADER_PARAMETER_RDG_TEXTURE_SRV(Texture2D, SrcSRV)
 		SHADER_PARAMETER_RDG_TEXTURE_UAV(RWTexture2D, DstUAV)
 	END_SHADER_PARAMETER_STRUCT()
 
@@ -272,7 +272,7 @@ class FCopyWindowCS : public FFFTShader
 		SHADER_PARAMETER(FIntRect, DstRect)
 		SHADER_PARAMETER(FVector3f, BrightPixelGain)
 
-		SHADER_PARAMETER_RDG_TEXTURE(Texture2D, SrcSRV)
+		SHADER_PARAMETER_RDG_TEXTURE_SRV(Texture2D, SrcSRV)
 		SHADER_PARAMETER_RDG_TEXTURE_UAV(RWTexture2D, DstUAV)
 	END_SHADER_PARAMETER_STRUCT()
 
@@ -296,7 +296,7 @@ class FComplexMultiplyImagesCS : public FFFTShader
 		SHADER_PARAMETER(FIntRect, SrcRect)
 		SHADER_PARAMETER(int32, DataLayout)
 
-		SHADER_PARAMETER_RDG_TEXTURE(Texture2D, SrcSRV)
+		SHADER_PARAMETER_RDG_TEXTURE_SRV(Texture2D, SrcSRV)
 		SHADER_PARAMETER_RDG_TEXTURE(Texture2D, KnlSRV)
 		SHADER_PARAMETER_RDG_TEXTURE_UAV(RWTexture2D, DstUAV)
 	END_SHADER_PARAMETER_STRUCT()
@@ -323,7 +323,7 @@ class FGSComplexTransformCS : public FFFTShader
 		SHADER_PARAMETER(FVector3f, BrightPixelGain)
 		SHADER_PARAMETER(int32, TransformType)
 
-		SHADER_PARAMETER_RDG_TEXTURE(Texture2D, SrcTexture)
+		SHADER_PARAMETER_RDG_TEXTURE_SRV(Texture2D, SrcTexture)
 		SHADER_PARAMETER_RDG_TEXTURE_UAV(RWTexture2D, DstTexture)
 	END_SHADER_PARAMETER_STRUCT()
 
@@ -361,7 +361,7 @@ class FGSTwoForOneTransformCS : public FFFTShader
 		SHADER_PARAMETER(FVector3f, BrightPixelGain)
 		SHADER_PARAMETER(int32, TransformType)
 
-		SHADER_PARAMETER_RDG_TEXTURE(Texture2D, SrcTexture)
+		SHADER_PARAMETER_RDG_TEXTURE_SRV(Texture2D, SrcTexture)
 		SHADER_PARAMETER_RDG_BUFFER_SRV(StructuredBuffer<float4>, DstPostFilterParameters)
 		SHADER_PARAMETER_RDG_TEXTURE_UAV(RWTexture2D, DstTexture)
 	END_SHADER_PARAMETER_STRUCT()
@@ -398,7 +398,7 @@ class FGSConvolutionWithTextureCS : public FFFTShader
 		SHADER_PARAMETER(FIntPoint, DstExtent)
 		SHADER_PARAMETER(int32, TransformType)
 
-		SHADER_PARAMETER_RDG_TEXTURE(Texture2D, SrcTexture)
+		SHADER_PARAMETER_RDG_TEXTURE_SRV(Texture2D, SrcTexture)
 		SHADER_PARAMETER_RDG_TEXTURE(Texture2D, FilterTexture)
 		SHADER_PARAMETER_RDG_TEXTURE_UAV(RWTexture2D, DstTexture)
 	END_SHADER_PARAMETER_STRUCT()
@@ -452,7 +452,7 @@ void DispatchReorderFFTPassCS(
 	FRDGBuilder& GraphBuilder,
 	const FGlobalShaderMap* ShaderMap,
 	const GPUFFT::FFTDescription& FFTDesc,
-	FRDGTextureRef SrcTexture, const FIntRect& SrcWindow,
+	FRDGTextureSRVRef SrcTexture, const FIntRect& SrcWindow,
 	FRDGTextureRef DstTexture, const FIntRect& DstWindow,
 	const bool bScrubNaNs = false)
 {
@@ -506,7 +506,7 @@ void DispatchGSSubComplexFFTPassCS(
 	FRDGBuilder& GraphBuilder,
 	const FGlobalShaderMap* ShaderMap,
 	const GPUFFT::FFTDescription& FFTDesc,
-	FRDGTextureRef SrcTexture, const FIntRect& SrcWindow,
+	FRDGTextureSRVRef SrcTexture, const FIntRect& SrcWindow,
 	FRDGTextureRef DstTexture)
 {
 	const uint32 TransformLength = FFTDesc.SignalLength;
@@ -563,7 +563,7 @@ void DispatchComplexFFTPassCS(
 	const FGlobalShaderMap* ShaderMap,
 	const GPUFFT::FFTDescription& FFTDesc,
 	const uint32 PassLength,
-	FRDGTextureRef SrcTexture, const FIntRect& SrcWindow,
+	FRDGTextureSRVRef SrcTexture, const FIntRect& SrcWindow,
 	FRDGBufferRef DstPostFilterParameters,
 	FRDGTextureRef DstTexture)
 {
@@ -619,7 +619,7 @@ void DispatchPackTwoForOneFFTPassCS(
 	FRDGBuilder& GraphBuilder,
 	const FGlobalShaderMap* ShaderMap,
 	const GPUFFT::FFTDescription& FFTDesc,
-	FRDGTextureRef SrcTexture,
+	FRDGTextureSRVRef SrcTexture,
 	FRDGTextureRef DstTexture)
 {
 	const uint32 TransformLength = FFTDesc.SignalLength;
@@ -676,7 +676,7 @@ void DispatchCopyWindowCS(
 	FRDGBuilder& GraphBuilder,
 	ERDGPassFlags ComputePassFlags,
 	const FGlobalShaderMap* ShaderMap,
-	FRDGTextureRef SrcTexture, const FIntRect& SrcWindow,
+	FRDGTextureSRVRef SrcTexture, const FIntRect& SrcWindow,
 	FRDGTextureRef DstTexture, const FIntRect& DstWindow,
 	const GPUFFT::FPreFilter& PreFilter = GPUFFT::FPreFilter(TNumericLimits<float>::Max(), TNumericLimits<float>::Lowest(), 0.f))
 {
@@ -731,7 +731,7 @@ void DispatchComplexMultiplyImagesCS(
 	ERDGPassFlags ComputePassFlags,
 	const FGlobalShaderMap* ShaderMap,
 	const bool bHorizontalScanlines,
-	FRDGTextureRef SrcTexture, const FIntRect& SrcWindow,
+	FRDGTextureSRVRef SrcTexture, const FIntRect& SrcWindow,
 	FRDGTextureRef KnlTexure,
 	FRDGTextureRef DstTexture)
 {
@@ -776,7 +776,7 @@ void DispatchGSComplexFFTCS(
 	FRDGBuilder& GraphBuilder,
 	const FGlobalShaderMap* ShaderMap,
 	const GPUFFT::FFTDescription& FFTDesc,
-	FRDGTextureRef SrcTexture, const FIntRect& SrcRect,
+	FRDGTextureSRVRef SrcTexture, const FIntRect& SrcRect,
 	FRDGTextureRef DstTexture)
 {
 	const uint32 TransformLength = FFTDesc.SignalLength;
@@ -839,7 +839,7 @@ void DispatchGSTwoForOneFFTCS(
 	FRDGBuilder& GraphBuilder,
 	const FGlobalShaderMap* ShaderMap,
 	const GPUFFT::FFTDescription& FFTDesc,
-	FRDGTextureRef SrcTexture, const FIntRect& SrcRect,
+	FRDGTextureSRVRef SrcTexture, const FIntRect& SrcRect,
 	FRDGTextureRef DstTexture, const FIntRect& DstRect,
 	const GPUFFT::FPreFilter& PreFilter,
 	FRDGBufferRef DstPostFilterParameters)
@@ -906,7 +906,7 @@ void DispatchGSConvolutionWithTextureCS(
 	FRDGBuilder& GraphBuilder,
 	const FGlobalShaderMap* ShaderMap,
 	const GPUFFT::FFTDescription& FFTDesc,
-	FRDGTextureRef SrcTexture, const FIntRect& SrcRect,
+	FRDGTextureSRVRef SrcTexture, const FIntRect& SrcRect,
 	FRDGTextureRef PreTransformedKernel,
 	FRDGTextureRef DstTexture)
 {
@@ -968,7 +968,7 @@ void GPUFFT::CopyImage2D(
 	FRDGBuilder& GraphBuilder,
 	ERDGPassFlags ComputePassFlags,
 	const FGlobalShaderMap* ShaderMap,
-	FRDGTextureRef SrcTexture, const FIntRect& SrcWindow,
+	FRDGTextureSRVRef SrcTexture, const FIntRect& SrcWindow,
 	FRDGTextureRef DstTexture, const FIntRect& DstWindow,
 	const FPreFilter& PreFilter)
 {
@@ -979,7 +979,7 @@ void GPUFFT::ComplexFFTImage1D::MultiPass(
 	FRDGBuilder& GraphBuilder,
 	const FGlobalShaderMap* ShaderMap,
 	const FFTDescription& FFTDesc,
-	FRDGTextureRef SrcTexture, const FIntRect& SrcWindow,
+	FRDGTextureSRVRef SrcTexture, const FIntRect& SrcWindow,
 	FRDGTextureRef DstTexture,
 	FRDGBufferRef PostFilterParameters,
 	const bool bScrubNaNs)
@@ -1045,7 +1045,7 @@ void GPUFFT::ComplexFFTImage1D::MultiPass(
 	// In total 2 + Log2(TransformLength / SubPassLength() )  passes.   This will be on the order of 3 or 4 passes 
 	// compared with 12 or more ..
 	{
-		FRDGTextureDesc SpectralDesc = FFTDesc.IsForward() ? DstTexture->Desc : SrcTexture->Desc;
+		FRDGTextureDesc SpectralDesc = FFTDesc.IsForward() ? DstTexture->Desc : SrcTexture->Desc.Texture->Desc;
 
 		// Re-order the data so we can do a pass of group-shared transforms
 		FRDGTextureRef SpectralTexture = GraphBuilder.CreateTexture(SpectralDesc, TEXT("FFT.Spectral"));
@@ -1057,7 +1057,7 @@ void GPUFFT::ComplexFFTImage1D::MultiPass(
 		FRDGTextureRef PrevSpectralTexture = GraphBuilder.CreateTexture(SpectralDesc, TEXT("FFT.Spectral"));
 		DispatchGSSubComplexFFTPassCS(
 			GraphBuilder, ShaderMap, FFTDesc,
-			/* SrcTexture = */ SpectralTexture, XFormWindow,
+			/* SrcTexture = */ GraphBuilder.CreateSRV(FRDGTextureSRVDesc(SpectralTexture)), XFormWindow,
 			/* DstTexture = */ PrevSpectralTexture);
 		
 		for (uint32 Ns = FGroupShardSubFFTPassCS::SubPassLength(); Ns < TransformLength; Ns *= 2)
@@ -1080,7 +1080,7 @@ void GPUFFT::ComplexFFTImage1D::MultiPass(
 			DispatchComplexFFTPassCS(
 				GraphBuilder, ShaderMap, FFTDesc,
 				Ns,
-				/* SrcTexture = */ PrevSpectralTexture, XFormWindow,
+				/* SrcTexture = */ GraphBuilder.CreateSRV(FRDGTextureSRVDesc(PrevSpectralTexture)), XFormWindow,
 				/* DstPostFilterParameters = */ DstPostFilterParameters,
 				/* DstTexture = */ OpDstTexture);
 
@@ -1095,7 +1095,7 @@ void GPUFFT::TwoForOneRealFFTImage1D::MultiPass(
 	FRDGBuilder& GraphBuilder,
 	const FGlobalShaderMap* ShaderMap,
 	const FFTDescription& FFTDesc,
-	FRDGTextureRef SrcTexture, const FIntRect& SrcWindow,
+	FRDGTextureSRVRef SrcTexture, const FIntRect& SrcWindow,
 	FRDGTextureRef DstTexture, const FIntRect& DstWindow,
 	const FPreFilter& PreFilter,
 	FRDGBufferRef PostFilterParameters)
@@ -1117,8 +1117,8 @@ void GPUFFT::TwoForOneRealFFTImage1D::MultiPass(
 		if (IsActive(PreFilter))
 		{
 			FRDGTextureDesc Desc = FRDGTextureDesc::Create2D(
-				SrcTexture->Desc.Extent,
-				SrcTexture->Desc.Format,
+				SrcTexture->Desc.Texture->Desc.Extent,
+				SrcTexture->Desc.Texture->Desc.Format,
 				FClearValueBinding::None,
 				TexCreate_ShaderResource | TexCreate_UAV);
 
@@ -1130,7 +1130,7 @@ void GPUFFT::TwoForOneRealFFTImage1D::MultiPass(
 				/* DstTexture = */ PreFilteredTexture, SrcWindow,
 				PreFilter);
 
-			SrcTexture = PreFilteredTexture;
+			SrcTexture = GraphBuilder.CreateSRV(FRDGTextureSRVDesc(PreFilteredTexture));
 		}
 
 		FRDGTextureRef SpectralTexture = GraphBuilder.CreateTexture(DstTexture->Desc, TEXT("FFT.Spectral"));
@@ -1145,7 +1145,7 @@ void GPUFFT::TwoForOneRealFFTImage1D::MultiPass(
 		ensure(DstWindow.Min == FIntPoint::ZeroValue);
 		DispatchPackTwoForOneFFTPassCS(
 			GraphBuilder, ShaderMap, FFTDesc,
-			/* SrcTexture = */ SpectralTexture,
+			/* SrcTexture = */ GraphBuilder.CreateSRV(FRDGTextureSRVDesc(SpectralTexture)),
 			/* DstTexture = */ DstTexture);
 	}
 	else  // Inverse transform.
@@ -1153,7 +1153,7 @@ void GPUFFT::TwoForOneRealFFTImage1D::MultiPass(
 		RDG_EVENT_SCOPE(GraphBuilder, "TwoForOneRealFFTImage1D(Inverse)");
 
 		// Pack the 4 transforms of real data as 2 transforms of complex data 
-		FRDGTextureRef SpectralTexture = GraphBuilder.CreateTexture(SrcTexture->Desc, TEXT("FFT.Spectral"));
+		FRDGTextureRef SpectralTexture = GraphBuilder.CreateTexture(SrcTexture->Desc.Texture->Desc, TEXT("FFT.Spectral"));
 		ensure(SrcWindow.Min == FIntPoint::ZeroValue);
 		DispatchPackTwoForOneFFTPassCS(
 			GraphBuilder, ShaderMap, FFTDesc,
@@ -1164,7 +1164,7 @@ void GPUFFT::TwoForOneRealFFTImage1D::MultiPass(
 		ensure(DstWindow.Min == FIntPoint::ZeroValue);
 		ComplexFFTImage1D::MultiPass(
 			GraphBuilder, ShaderMap, FFTDesc,
-			/* SrcTexture = */ SpectralTexture, SrcWindow,
+			/* SrcTexture = */ GraphBuilder.CreateSRV(FRDGTextureSRVDesc(SpectralTexture)), SrcWindow,
 			/* DstTexture = */ DstTexture,
 			/* PostFilterParameters = */ PostFilterParameters,
 			/* bScrubNaNs = */ false);
@@ -1175,7 +1175,7 @@ void GPUFFT::FFTImage2D(
 	FRDGBuilder& GraphBuilder,
 	const FGlobalShaderMap* ShaderMap,
 	const FIntPoint& FrequencySize, bool bHorizontalFirst,
-	FRDGTextureRef SrcTexture, const FIntRect& ROIRect,
+	FRDGTextureSRVRef SrcTexture, const FIntRect& ROIRect,
 	FRDGTextureRef DstTexture)
 {
 	using GPUFFT::FFTDescription;
@@ -1210,7 +1210,7 @@ void GPUFFT::FFTImage2D(
 	ComplexFFTImage1D::MultiPass(
 		GraphBuilder, ShaderMap,
 		ComplexFFTDesc,
-		TempTexture, TmpRect,
+		GraphBuilder.CreateSRV(FRDGTextureSRVDesc(TempTexture)), TmpRect,
 		DstTexture);
 }
 
@@ -1218,7 +1218,7 @@ void GPUFFT::ConvolutionWithTextureImage1D::MultiPass(
 	FRDGBuilder& GraphBuilder,
 	const FGlobalShaderMap* ShaderMap,
 	const FFTDescription& FFTDesc,
-	FRDGTextureRef SrcTexture, const FIntRect& SrcWindow,
+	FRDGTextureSRVRef SrcTexture, const FIntRect& SrcWindow,
 	FRDGTextureRef TransformedKernel,
 	FRDGTextureRef DstTexture)
 {
@@ -1237,17 +1237,17 @@ void GPUFFT::ConvolutionWithTextureImage1D::MultiPass(
 		const FIntRect  TargetRect(FIntPoint(0, 0), TargetExtent);
 
 		// Forward transform: Results in DstBuffer
-		FRDGTextureRef SpectralTexture0 = GraphBuilder.CreateTexture(SrcTexture->Desc, TEXT("FFT.Spectral"));
+		FRDGTextureRef SpectralTexture0 = GraphBuilder.CreateTexture(SrcTexture->Desc.Texture->Desc, TEXT("FFT.Spectral"));
 		ComplexFFTImage1D::MultiPass(
 			GraphBuilder, ShaderMap, FFTDesc,
 			/* SrcTexture = */ SrcTexture, SrcWindow,
 			/* DstTexture = */ SpectralTexture0);
 
 		// Apply convolution kernel.
-		FRDGTextureRef SpectralTexture1 = GraphBuilder.CreateTexture(SrcTexture->Desc, TEXT("FFT.Spectral"));
+		FRDGTextureRef SpectralTexture1 = GraphBuilder.CreateTexture(SrcTexture->Desc.Texture->Desc, TEXT("FFT.Spectral"));
 		DispatchComplexMultiplyImagesCS(
 			GraphBuilder, FFTDesc.ComputePassFlags, ShaderMap, FFTDesc.IsHorizontal(),
-			/* SrcTexture = */ SpectralTexture0, TargetRect,
+			/* SrcTexture = */ GraphBuilder.CreateSRV(FRDGTextureSRVDesc(SpectralTexture0)), TargetRect,
 			TransformedKernel, 
 			/* DstTexture = */ SpectralTexture1);
 
@@ -1257,7 +1257,7 @@ void GPUFFT::ConvolutionWithTextureImage1D::MultiPass(
 
 		ComplexFFTImage1D::MultiPass(
 			GraphBuilder, ShaderMap, InvFFTDesc,
-			/* SrcTexture = */ SpectralTexture1, TargetRect,
+			/* SrcTexture = */ GraphBuilder.CreateSRV(FRDGTextureSRVDesc(SpectralTexture1)), TargetRect,
 			/* DstTexture = */ DstTexture);
 	}
 }
@@ -1314,7 +1314,7 @@ void GPUFFT::ConvolutionWithTextureImage2D(
 	const FGlobalShaderMap* ShaderMap,
 	const FIntPoint& FrequencySize, bool bHorizontalFirst,
 	FRDGTextureRef TransformedKernel,
-	FRDGTextureRef SrcTexture, const FIntRect& ROIRect,
+	FRDGTextureSRVRef SrcTexture, const FIntRect& ROIRect,
 	FRDGTextureRef DstTexture, const FIntRect& DstRect,
 	const FPreFilter& PreFilter,
 	FRDGBufferRef PostFilterParameters,
@@ -1370,7 +1370,7 @@ void GPUFFT::ConvolutionWithTextureImage2D(
 	ConvolutionWithTextureImage1D::MultiPass(
 		GraphBuilder, ShaderMap,
 		ConvolutionFFTDesc,
-		TwoForOneFFTOutput, SpectralOutputRect,
+		GraphBuilder.CreateSRV(FRDGTextureSRVDesc(TwoForOneFFTOutput)), SpectralOutputRect,
 		TransformedKernel,
 		ConvolutionOutput);
 
@@ -1380,7 +1380,7 @@ void GPUFFT::ConvolutionWithTextureImage2D(
 	TwoForOneRealFFTImage1D::MultiPass(
 		GraphBuilder, ShaderMap,
 		TwoForOneIvnFFTDesc,
-		ConvolutionOutput, SpectralOutputRect,
+		GraphBuilder.CreateSRV(FRDGTextureSRVDesc(ConvolutionOutput)), SpectralOutputRect,
 		DstTexture, DstRect,
 		PreFilter,
 		PostFilterParameters);
diff --git a/Engine/Source/Runtime/Renderer/Private/GPUFastFourierTransform.h b/Engine/Source/Runtime/Renderer/Private/GPUFastFourierTransform.h
--- a/Engine/Source/Runtime/Renderer/Private/GPUFastFourierTransform.h
+++ b/Engine/Source/Runtime/Renderer/Private/GPUFastFourierTransform.h
@@ -129,7 +129,7 @@ namespace GPUFFT
 			FRDGBuilder& GraphBuilder,
 			const FGlobalShaderMap* ShaderMap,
 			const FFTDescription& FFTDesc,
-			FRDGTextureRef SrcTexture, const FIntRect& SrcWindow,
+			FRDGTextureSRVRef SrcTexture, const FIntRect& SrcWindow,
 			FRDGTextureRef DstTexture,
 			FRDGBufferRef PostFilterParameters = nullptr,
 			const bool bScrubNaNs = false);
@@ -179,7 +179,7 @@ namespace GPUFFT
 			FRDGBuilder& GraphBuilder,
 			const FGlobalShaderMap* ShaderMap,
 			const FFTDescription& FFTDesc,
-			FRDGTextureRef SrcTexture, const FIntRect& SrcWindow,
+			FRDGTextureSRVRef SrcTexture, const FIntRect& SrcWindow,
 			FRDGTextureRef DstTexture, const FIntRect& DstWindow,
 			const FPreFilter& PreFilter = FPreFilter(TNumericLimits<float>::Max(), TNumericLimits<float>::Lowest(), 0.f),
 			FRDGBufferRef PostFilterParameters = nullptr);
@@ -226,7 +226,7 @@ namespace GPUFFT
 			FRDGBuilder& GraphBuilder,
 			const FGlobalShaderMap* ShaderMap,
 			const FFTDescription& FFTDesc,
-			FRDGTextureRef SrcTexture, const FIntRect& SrcWindow,
+			FRDGTextureSRVRef SrcTexture, const FIntRect& SrcWindow,
 			FRDGTextureRef TransformedKernel,
 			FRDGTextureRef DstTexture);
 	};
@@ -252,7 +252,7 @@ namespace GPUFFT
 		FRDGBuilder& GraphBuilder,
 		const FGlobalShaderMap* ShaderMap,
 		const FIntPoint& FrequencySize, bool bHorizontalFirst,
-		FRDGTextureRef SrcTexture, const FIntRect& ROIRect,
+		FRDGTextureSRVRef SrcTexture, const FIntRect& ROIRect,
 		FRDGTextureRef ResultTexture);
 
 
@@ -288,7 +288,7 @@ namespace GPUFFT
 		const FIntPoint& FrequencySize,
 		bool bHorizontalFirst,
 		FRDGTextureRef TransformedKernel,
-		FRDGTextureRef SrcTexture, const FIntRect& ROIRect,
+		FRDGTextureSRVRef SrcTexture, const FIntRect& ROIRect,
 		FRDGTextureRef DstTexture, const FIntRect& DstRect,
 		const FPreFilter& PreFilter,
 		FRDGBufferRef PostFilterParameters,
@@ -316,7 +316,7 @@ namespace GPUFFT
 		FRDGBuilder& GraphBuilder,
 		ERDGPassFlags ComputePassFlags,
 		const FGlobalShaderMap* ShaderMap,
-		FRDGTextureRef SrcTexture, const FIntRect& SrcWindow,
+		FRDGTextureSRVRef SrcTexture, const FIntRect& SrcWindow,
 		FRDGTextureRef DstTexture, const FIntRect& DstWindow,
 		const FPreFilter& PreFilter = FPreFilter(TNumericLimits<float>::Max(), TNumericLimits<float>::Lowest(), 0.f));
 };
diff --git a/Engine/Source/Runtime/Renderer/Private/PostProcess/PostProcessDownsample.cpp b/Engine/Source/Runtime/Renderer/Private/PostProcess/PostProcessDownsample.cpp
--- a/Engine/Source/Runtime/Renderer/Private/PostProcess/PostProcessDownsample.cpp
+++ b/Engine/Source/Runtime/Renderer/Private/PostProcess/PostProcessDownsample.cpp
@@ -15,11 +15,11 @@ BEGIN_SHADER_PARAMETER_STRUCT(FDownsampleParameters, )
 	SHADER_PARAMETER_STRUCT_REF(FViewUniformShaderParameters, ViewUniformBuffer)
 	SHADER_PARAMETER_STRUCT(FScreenPassTextureViewportParameters, Input)
 	SHADER_PARAMETER_STRUCT(FScreenPassTextureViewportParameters, Output)
-	SHADER_PARAMETER_RDG_TEXTURE(Texture2D, InputTexture)
+	SHADER_PARAMETER_RDG_TEXTURE_SRV(Texture2D, InputTexture)
 	SHADER_PARAMETER_SAMPLER(SamplerState, InputSampler)
 END_SHADER_PARAMETER_STRUCT()
 
-FDownsampleParameters GetDownsampleParameters(const FViewInfo& View, FScreenPassTexture Output, FScreenPassTexture Input, EDownsampleQuality DownsampleMethod)
+FDownsampleParameters GetDownsampleParameters(const FViewInfo& View, FScreenPassTexture Output, FScreenPassTextureSlice Input, EDownsampleQuality DownsampleMethod)
 {
 	check(Output.IsValid());
 	check(Input.IsValid());
@@ -31,7 +31,7 @@ FDownsampleParameters GetDownsampleParameters(const FViewInfo& View, FScreenPass
 	Parameters.ViewUniformBuffer = View.ViewUniformBuffer;
 	Parameters.Input = InputParameters;
 	Parameters.Output = OutputParameters;
-	Parameters.InputTexture = Input.Texture;
+	Parameters.InputTexture = Input.TextureSRV;
 	Parameters.InputSampler = TStaticSamplerState<SF_Bilinear, AM_Clamp, AM_Clamp, AM_Clamp>::GetRHI();
 	return Parameters;
 }
@@ -145,7 +145,7 @@ FScreenPassTexture AddDownsamplePass(
 		PermutationVector.Set<FDownsampleQualityDimension>(Inputs.Quality);
 
 		FDownsamplePS::FParameters* PassParameters = GraphBuilder.AllocParameters<FDownsamplePS::FParameters>();
-		PassParameters->Common = GetDownsampleParameters(View, Output, Inputs.SceneColor, Inputs.Quality);
+		PassParameters->Common = GetDownsampleParameters(View, Output, FScreenPassTextureSlice::CreateFromScreenPassTexture(GraphBuilder, Inputs.SceneColor), Inputs.Quality);
 		PassParameters->RenderTargets[0] = Output.GetRenderTargetBinding();
 
 		TShaderMapRef<FDownsamplePS> PixelShader(View.ShaderMap, PermutationVector);
@@ -165,7 +165,7 @@ FScreenPassTexture AddDownsamplePass(
 	return MoveTemp(Output);
 }
 
-void AddDownsampleComputePass(FRDGBuilder& GraphBuilder, const FViewInfo& View, FScreenPassTexture Input, FScreenPassTexture Output, EDownsampleQuality Quality, ERDGPassFlags PassFlags)
+void AddDownsampleComputePass(FRDGBuilder& GraphBuilder, const FViewInfo& View, FScreenPassTextureSlice Input, FScreenPassTexture Output, EDownsampleQuality Quality, ERDGPassFlags PassFlags)
 {
 	check(PassFlags == ERDGPassFlags::Compute || PassFlags == ERDGPassFlags::AsyncCompute);
 
@@ -191,6 +191,11 @@ void AddDownsampleComputePass(FRDGBuilder& GraphBuilder, const FViewInfo& View,
 		FComputeShaderUtils::GetGroupCount(Output.ViewRect.Size(), FIntPoint(GDownsampleTileSizeX, GDownsampleTileSizeY)));
 }
 
+void AddDownsampleComputePass(FRDGBuilder& GraphBuilder, const FViewInfo& View, FScreenPassTexture Input, FScreenPassTexture Output, EDownsampleQuality Quality, ERDGPassFlags PassFlags)
+{
+	return AddDownsampleComputePass(GraphBuilder, View, FScreenPassTextureSlice::CreateFromScreenPassTexture(GraphBuilder, Input), Output, Quality, PassFlags);
+}
+
 void FSceneDownsampleChain::Init(
 	FRDGBuilder& GraphBuilder,
 	const FViewInfo& View,
diff --git a/Engine/Source/Runtime/Renderer/Private/PostProcess/PostProcessDownsample.h b/Engine/Source/Runtime/Renderer/Private/PostProcess/PostProcessDownsample.h
--- a/Engine/Source/Runtime/Renderer/Private/PostProcess/PostProcessDownsample.h
+++ b/Engine/Source/Runtime/Renderer/Private/PostProcess/PostProcessDownsample.h
@@ -53,6 +53,8 @@ struct FDownsamplePassInputs
 FScreenPassTexture AddDownsamplePass(FRDGBuilder& GraphBuilder, const FViewInfo& View, const FDownsamplePassInputs& Inputs);
 
 void AddDownsampleComputePass(FRDGBuilder& GraphBuilder, const FViewInfo& View, FScreenPassTexture Input, FScreenPassTexture Output, EDownsampleQuality Quality, ERDGPassFlags PassFlags);
+void AddDownsampleComputePass(FRDGBuilder& GraphBuilder, const FViewInfo& View, FScreenPassTextureSlice Input, FScreenPassTexture Output, EDownsampleQuality Quality, ERDGPassFlags PassFlags);
+
 
 class FSceneDownsampleChain
 {
diff --git a/Engine/Source/Runtime/Renderer/Private/PostProcess/PostProcessFFTBloom.cpp b/Engine/Source/Runtime/Renderer/Private/PostProcess/PostProcessFFTBloom.cpp
--- a/Engine/Source/Runtime/Renderer/Private/PostProcess/PostProcessFFTBloom.cpp
+++ b/Engine/Source/Runtime/Renderer/Private/PostProcess/PostProcessFFTBloom.cpp
@@ -275,7 +275,7 @@ FRDGTextureRef TransformKernelFFT(
 		GraphBuilder,
 		View.ShaderMap,
 		FrequencySize, bDoHorizontalFirst,
-		ResizedKernel, SrcRect,
+		GraphBuilder.CreateSRV(FRDGTextureSRVDesc(ResizedKernel)), SrcRect,
 		SpectralKernel);
 
 	return SpectralKernel;
@@ -708,7 +708,7 @@ void InitDomainAndGetKernel(
 FFFTBloomOutput AddFFTBloomPass(
 	FRDGBuilder& GraphBuilder,
 	const FViewInfo& View,
-	const FScreenPassTexture& InputSceneColor,
+	const FScreenPassTextureSlice& InputSceneColor,
 	float InputResolutionFraction,
 	const FEyeAdaptationParameters& EyeAdaptationParameters,
 	FRDGBufferRef EyeAdaptationBuffer,
@@ -791,7 +791,7 @@ FFFTBloomOutput AddFFTBloomPass(
 	RDG_EVENT_SCOPE(GraphBuilder, "Bloom %dx%d", Intermediates.ImageSize.X, Intermediates.ImageSize.Y);
 
 	// Downscale the input to the final resolution fraction
-	FScreenPassTexture FFTInputSceneColor;
+	FScreenPassTextureSlice FFTInputSceneColor;
 	if (ResolutionFraction != InputResolutionFraction)
 	{
 		FIntPoint Extent;
@@ -799,14 +799,15 @@ FFFTBloomOutput AddFFTBloomPass(
 
 		FRDGTextureDesc Desc = FRDGTextureDesc::Create2D(
 			Extent,
-			InputSceneColor.Texture->Desc.Format,
+			InputSceneColor.TextureSRV->Desc.Texture->Desc.Format,
 			FClearValueBinding::None,
 			TexCreate_ShaderResource | TexCreate_UAV);
 
-		FFTInputSceneColor.Texture = GraphBuilder.CreateTexture(Desc, TEXT("Bloom.FFT.Input"));
+		FFTInputSceneColor.TextureSRV = GraphBuilder.CreateSRV(FRDGTextureSRVDesc(
+			GraphBuilder.CreateTexture(Desc, TEXT("Bloom.FFT.Input"))));
 		FFTInputSceneColor.ViewRect = FIntRect(FIntPoint::ZeroValue, Intermediates.ImageSize);
 
-		AddDownsampleComputePass(GraphBuilder, View, InputSceneColor, FFTInputSceneColor, EDownsampleQuality::Low, Intermediates.ComputePassFlags);
+		AddDownsampleComputePass(GraphBuilder, View, InputSceneColor, FScreenPassTexture(FFTInputSceneColor.TextureSRV->Desc.Texture, FFTInputSceneColor.ViewRect), EDownsampleQuality::Low, Intermediates.ComputePassFlags);
 	}
 	else
 	{
@@ -815,14 +816,27 @@ FFFTBloomOutput AddFFTBloomPass(
 
 	if (LocalExposureTexture != nullptr)
 	{
-		FScreenPassTexture Temp = FFTInputSceneColor;
+		FScreenPassTextureSlice Temp = FFTInputSceneColor;
 
-		FRDGTextureDesc Desc = Temp.Texture->Desc;
-		Desc.Flags = TexCreate_ShaderResource | TexCreate_UAV;
+		FRDGTextureDesc Desc = FRDGTextureDesc::Create2D(
+			Temp.TextureSRV->Desc.Texture->Desc.Extent,
+			Temp.TextureSRV->Desc.Texture->Desc.Format,
+			FClearValueBinding::None,
+			TexCreate_ShaderResource | TexCreate_UAV);
 
-		FFTInputSceneColor.Texture = GraphBuilder.CreateTexture(Desc, TEXT("Bloom.FFT.Input"));
+		FFTInputSceneColor.TextureSRV = GraphBuilder.CreateSRV(FRDGTextureSRVDesc(
+			GraphBuilder.CreateTexture(Desc, TEXT("Bloom.FFT.Input"))));
 
-		AddApplyLocalExposurePass(GraphBuilder, View, EyeAdaptationParameters, EyeAdaptationBuffer, LocalExposureTexture, BlurredLogLuminanceTexture, Temp, FFTInputSceneColor, Intermediates.ComputePassFlags);
+		AddApplyLocalExposurePass(
+			GraphBuilder,
+			View,
+			EyeAdaptationParameters,
+			EyeAdaptationBuffer,
+			LocalExposureTexture,
+			BlurredLogLuminanceTexture,
+			Temp,
+			FFTInputSceneColor,
+			Intermediates.ComputePassFlags);
 	}
 
 	// Init the domain data update the cached kernel if needed.
@@ -864,8 +878,8 @@ FFFTBloomOutput AddFFTBloomPass(
 	}
 
 	FRDGTextureDesc OutputSceneColorDesc = FRDGTextureDesc::Create2D(
-		FFTInputSceneColor.Texture->Desc.Extent,
-		FFTInputSceneColor.Texture->Desc.Format,
+		FFTInputSceneColor.TextureSRV->Desc.Texture->Desc.Extent,
+		FFTInputSceneColor.TextureSRV->Desc.Texture->Desc.Format,
 		FClearValueBinding::None,
 		TexCreate_ShaderResource | TexCreate_UAV | GFastVRamConfig.Bloom);
 		
@@ -879,7 +893,7 @@ FFFTBloomOutput AddFFTBloomPass(
 		Intermediates.FrequencySize,
 		Intermediates.bDoHorizontalFirst,
 		SpectralKernelTexture,
-		FFTInputSceneColor.Texture, FFTInputSceneColor.ViewRect,
+		FFTInputSceneColor.TextureSRV, FFTInputSceneColor.ViewRect,
 		BloomOutput.BloomTexture.Texture, BloomOutput.BloomTexture.ViewRect,
 		Intermediates.PreFilter,
 		FFTMulitplyParameters,
diff --git a/Engine/Source/Runtime/Renderer/Private/PostProcess/PostProcessFFTBloom.h b/Engine/Source/Runtime/Renderer/Private/PostProcess/PostProcessFFTBloom.h
--- a/Engine/Source/Runtime/Renderer/Private/PostProcess/PostProcessFFTBloom.h
+++ b/Engine/Source/Runtime/Renderer/Private/PostProcess/PostProcessFFTBloom.h
@@ -28,7 +28,7 @@ struct FFFTBloomOutput
 FFFTBloomOutput AddFFTBloomPass(
 	FRDGBuilder& GraphBuilder,
 	const FViewInfo& View,
-	const FScreenPassTexture& InputSceneColor,
+	const FScreenPassTextureSlice& InputSceneColor,
 	float InputResolutionFraction,
 	const FEyeAdaptationParameters& EyeAdaptationParameters,
 	FRDGBufferRef EyeAdaptationBuffer,
diff --git a/Engine/Source/Runtime/Renderer/Private/PostProcess/PostProcessLocalExposure.cpp b/Engine/Source/Runtime/Renderer/Private/PostProcess/PostProcessLocalExposure.cpp
--- a/Engine/Source/Runtime/Renderer/Private/PostProcess/PostProcessLocalExposure.cpp
+++ b/Engine/Source/Runtime/Renderer/Private/PostProcess/PostProcessLocalExposure.cpp
@@ -62,7 +62,7 @@ public:
 
 		SHADER_PARAMETER_STRUCT(FScreenPassTextureViewportParameters, Input)
 		SHADER_PARAMETER_STRUCT(FScreenPassTextureViewportParameters, Output)
-		SHADER_PARAMETER_RDG_TEXTURE(Texture2D, InputTexture)
+		SHADER_PARAMETER_RDG_TEXTURE_SRV(Texture2D, InputTexture)
 		SHADER_PARAMETER_RDG_TEXTURE_UAV(RWTexture2D<float4>, OutputFloat4)
 
 		SHADER_PARAMETER_STRUCT(FEyeAdaptationParameters, EyeAdaptation)
@@ -153,8 +153,8 @@ void AddApplyLocalExposurePass(
 	FRDGBufferRef EyeAdaptationBuffer,
 	FRDGTextureRef LocalExposureTexture,
 	FRDGTextureRef BlurredLogLuminanceTexture,
-	FScreenPassTexture Input,
-	FScreenPassTexture Output,
+	FScreenPassTextureSlice Input,
+	FScreenPassTextureSlice Output,
 	ERDGPassFlags PassFlags)
 {
 	check(Input.IsValid() && Output.IsValid());
@@ -167,8 +167,18 @@ void AddApplyLocalExposurePass(
 	PassParameters->Input = GetScreenPassTextureViewportParameters(FScreenPassTextureViewport(Input));
 	PassParameters->Output = GetScreenPassTextureViewportParameters(FScreenPassTextureViewport(Output));
 
-	PassParameters->InputTexture = Input.Texture;
-	PassParameters->OutputFloat4 = GraphBuilder.CreateUAV(Output.Texture);
+	PassParameters->InputTexture = Input.TextureSRV;
+	{
+		FRDGTextureUAVDesc OutputDesc(Output.TextureSRV->Desc.Texture);
+		if (Output.TextureSRV->Desc.Texture->Desc.IsTextureArray())
+		{
+			OutputDesc.DimensionOverride = ETextureDimension::Texture2D;
+			OutputDesc.FirstArraySlice = Output.TextureSRV->Desc.FirstArraySlice;
+			OutputDesc.NumArraySlices = 1;
+		}
+
+		PassParameters->OutputFloat4 = GraphBuilder.CreateUAV(OutputDesc);
+	}
 
 	PassParameters->EyeAdaptation = EyeAdaptationParameters;
 	PassParameters->EyeAdaptationBuffer = GraphBuilder.CreateSRV(EyeAdaptationBuffer);
diff --git a/Engine/Source/Runtime/Renderer/Private/PostProcess/PostProcessLocalExposure.h b/Engine/Source/Runtime/Renderer/Private/PostProcess/PostProcessLocalExposure.h
--- a/Engine/Source/Runtime/Renderer/Private/PostProcess/PostProcessLocalExposure.h
+++ b/Engine/Source/Runtime/Renderer/Private/PostProcess/PostProcessLocalExposure.h
@@ -19,6 +19,6 @@ void AddApplyLocalExposurePass(
 	FRDGBufferRef EyeAdaptationBuffer,
 	FRDGTextureRef LocalExposureTexture,
 	FRDGTextureRef BlurredLogLuminanceTexture,
-	FScreenPassTexture Input,
-	FScreenPassTexture Output,
+	FScreenPassTextureSlice Input,
+	FScreenPassTextureSlice Output,
 	ERDGPassFlags PassFlags);
\ No newline at end of file
diff --git a/Engine/Source/Runtime/Renderer/Private/PostProcess/PostProcessing.cpp b/Engine/Source/Runtime/Renderer/Private/PostProcess/PostProcessing.cpp
--- a/Engine/Source/Runtime/Renderer/Private/PostProcess/PostProcessing.cpp
+++ b/Engine/Source/Runtime/Renderer/Private/PostProcess/PostProcessing.cpp
@@ -1032,21 +1032,21 @@ void AddPostProcessingPasses(
 				LensFlareSceneDownsampleChain = &SceneDownsampleChain;
 
 				float InputResolutionFraction;
-				FScreenPassTexture InputSceneColor;
+				FScreenPassTextureSlice InputSceneColor;
 
 				if (FFTBloomResolutionFraction <= 0.25f)
 				{
-					InputSceneColor = QuarterResSceneColor;
+					InputSceneColor = FScreenPassTextureSlice::CreateFromScreenPassTexture(GraphBuilder, QuarterResSceneColor);
 					InputResolutionFraction = 0.25f;
 				}
 				else if (FFTBloomResolutionFraction <= 0.5f)
 				{
-					InputSceneColor = HalfResSceneColor;
+					InputSceneColor = FScreenPassTextureSlice::CreateFromScreenPassTexture(GraphBuilder, HalfResSceneColor);
 					InputResolutionFraction = 0.5f;
 				}
 				else
 				{
-					InputSceneColor = FScreenPassTexture::CopyFromSlice(GraphBuilder, SceneColorSlice); // TODO
+					InputSceneColor = SceneColorSlice;
 					InputResolutionFraction = 1.0f;
 				}
 
