Only query VRS enabled and debug force rate once per frame.

diff --git a/Engine/Source/Runtime/Renderer/Private/VariableRateShading/VariableRateShadingImageManager.cpp b/Engine/Source/Runtime/Renderer/Private/VariableRateShading/VariableRateShadingImageManager.cpp
--- a/Engine/Source/Runtime/Renderer/Private/VariableRateShading/VariableRateShadingImageManager.cpp
+++ b/Engine/Source/Runtime/Renderer/Private/VariableRateShading/VariableRateShadingImageManager.cpp
@@ -30,24 +30,53 @@ DEFINE_LOG_CATEGORY(LogVRS);
  * Basic CVars
  */
 
+void CVarVRSPreviewCallback(IConsoleVariable* Var)
+{
+	const int32 RequestedPreview = Var->GetInt();
+	if (RequestedPreview < 0 || RequestedPreview > 2)
+	{
+		UE_LOG(LogVRS, Warning, TEXT("Selected invalid preview mode, disabling preview"));
+	}
+}
+
 TAutoConsoleVariable<int32> CVarVRSPreview(
 	TEXT("r.VRS.Preview"),
 	0,
 	TEXT("Show a debug visualization of the VRS shading rate image texture.")
 	TEXT("0 - off, 1 - on, 2 - conservative (affects CAS only)"),
+	FConsoleVariableDelegate::CreateStatic(&CVarVRSPreviewCallback),
 	ECVF_RenderThreadSafe);
 
+void CVarVRSDebugForceRateCallback(IConsoleVariable* Var)
+{
+	const int32 RequestedDebugForceRate = Var->GetInt();
+	if (RequestedDebugForceRate > 6) 
+	{
+		UE_LOG(LogVRS, Warning, TEXT("Selected forced shading rate exceeds maximum available, defaulting to 4x4"));
+	}
+}
+
 int GVRSDebugForceRate = -1;
 FAutoConsoleVariableRef CVarVRSDebugForceRate(
 	TEXT("r.VRS.DebugForceRate"),
 	GVRSDebugForceRate,
 	TEXT("-1 : None, 0 : Force 1x1, 1 : Force 1x2, 2 : Force 2x1, 3: Force 2x2, 4 : Force 2x4, 5 : Force 4x2, 6 : Force 4x4"),
+	FConsoleVariableDelegate::CreateStatic(&CVarVRSDebugForceRateCallback),
 	ECVF_RenderThreadSafe);
 
 /**
  * Pass Settings
  */
 
+static void CVarVRSImagePassTypeCallback(IConsoleVariable* Var)
+{
+	const int32 RequestedImageType = Var->GetInt();
+	if (RequestedImageType < 0 || RequestedImageType >= FVariableRateShadingImageManager::EVRSPassType::Num)
+	{
+		UE_LOG(LogVRS, Warning, TEXT("Selected invalid image type, disabling VRS for pass"));
+	}
+}
+
 TAutoConsoleVariable<int32> CVarVRSBasePass(
 	TEXT("r.VRS.BasePass"),
 	2,
@@ -55,6 +84,7 @@ TAutoConsoleVariable<int32> CVarVRSBasePass(
 	TEXT("0: Disabled")
 	TEXT("1: Full")
 	TEXT("2: Conservative (default)"),
+	FConsoleVariableDelegate::CreateStatic(&CVarVRSImagePassTypeCallback),
 	ECVF_RenderThreadSafe);
 TAutoConsoleVariable<int32> CVarVRSTranslucency(
 	TEXT("r.VRS.Translucency"),
@@ -63,6 +93,7 @@ TAutoConsoleVariable<int32> CVarVRSTranslucency(
 	TEXT("0: Disabled")
 	TEXT("1: Full (default)")
 	TEXT("2: Conservative"),
+	FConsoleVariableDelegate::CreateStatic(&CVarVRSImagePassTypeCallback),
 	ECVF_RenderThreadSafe);
 TAutoConsoleVariable<int32> CVarVRSNaniteEmitGBuffer(
 	TEXT("r.VRS.NaniteEmitGBuffer"),
@@ -71,6 +102,7 @@ TAutoConsoleVariable<int32> CVarVRSNaniteEmitGBuffer(
 	TEXT("0: Disabled")
 	TEXT("1: Full")
 	TEXT("2: Conservative (default)"),
+	FConsoleVariableDelegate::CreateStatic(&CVarVRSImagePassTypeCallback),
 	ECVF_RenderThreadSafe);
 TAutoConsoleVariable<int32> CVarVRS_SSAO(
 	TEXT("r.VRS.SSAO"),
@@ -79,6 +111,7 @@ TAutoConsoleVariable<int32> CVarVRS_SSAO(
 	TEXT("0: Disabled")
 	TEXT("1: Full")
 	TEXT("2: Conservative (default)"),
+	FConsoleVariableDelegate::CreateStatic(&CVarVRSImagePassTypeCallback),
 	ECVF_RenderThreadSafe);
 TAutoConsoleVariable<int32> CVarVRS_SSR(
 	TEXT("r.VRS.SSR"),
@@ -87,6 +120,7 @@ TAutoConsoleVariable<int32> CVarVRS_SSR(
 	TEXT("0: Disabled")
 	TEXT("1: Full")
 	TEXT("2: Conservative (default)"),
+	FConsoleVariableDelegate::CreateStatic(&CVarVRSImagePassTypeCallback),
 	ECVF_RenderThreadSafe);
 TAutoConsoleVariable<int32> CVarVRSReflectionEnvironmentSky(
 	TEXT("r.VRS.ReflectionEnvironmentSky"),
@@ -95,6 +129,7 @@ TAutoConsoleVariable<int32> CVarVRSReflectionEnvironmentSky(
 	TEXT("0: Disabled")
 	TEXT("1: Full")
 	TEXT("2: Conservative (default)"),
+	FConsoleVariableDelegate::CreateStatic(&CVarVRSImagePassTypeCallback),
 	ECVF_RenderThreadSafe);
 TAutoConsoleVariable<int32> CVarVRSLightFunctions(
 	TEXT("r.VRS.LightFunctions"),
@@ -103,6 +138,7 @@ TAutoConsoleVariable<int32> CVarVRSLightFunctions(
 	TEXT("0: Disabled")
 	TEXT("1: Full (default)")
 	TEXT("2: Conservative"),
+	FConsoleVariableDelegate::CreateStatic(&CVarVRSImagePassTypeCallback),
 	ECVF_RenderThreadSafe);
 TAutoConsoleVariable<int32> CVarVRSDecals(
 	TEXT("r.VRS.Decals"),
@@ -111,6 +147,7 @@ TAutoConsoleVariable<int32> CVarVRSDecals(
 	TEXT("0: Disabled")
 	TEXT("1: Full")
 	TEXT("2: Conservative (default)"),
+	FConsoleVariableDelegate::CreateStatic(&CVarVRSImagePassTypeCallback),
 	ECVF_RenderThreadSafe);
 
 /**
@@ -275,7 +312,12 @@ static EDisplayOutputFormat GetDisplayOutputFormat(const FViewInfo& View)
 
 bool FVariableRateShadingImageManager::IsVRSSupportedByRHI()
 {
-	return GRHISupportsAttachmentVariableRateShading && GRHIVariableRateShadingEnabled && GRHIAttachmentVariableRateShadingEnabled && FDataDrivenShaderPlatformInfo::GetSupportsVariableRateShading(GMaxRHIShaderPlatform);
+	return GRHISupportsAttachmentVariableRateShading && FDataDrivenShaderPlatformInfo::GetSupportsVariableRateShading(GMaxRHIShaderPlatform);
+}
+
+bool FVariableRateShadingImageManager::IsVRSEnabled()
+{
+	return GRHIVariableRateShadingEnabled && GRHIAttachmentVariableRateShadingEnabled;
 }
 
 bool FVariableRateShadingImageManager::IsVRSCompatibleWithOutputType(const EDisplayOutputFormat& OutputFormat)
@@ -292,7 +334,7 @@ bool FVariableRateShadingImageManager::IsVRSCompatibleWithView(const FViewInfo&
 	// TODO: Investigate if it's worthwhile getting scene captures working. Things that we'll need to take care of
 	// is to associate shading rate texture image with main scene, and scene capture.  But what if there is
 	// more than 1 scene capture?  Is there a unique identifier that connects two frames of scene capture.
-	return IsVRSSupportedByRHI() && !ViewInfo.bIsSceneCapture && IsVRSCompatibleWithOutputType(GetDisplayOutputFormat(ViewInfo));
+	return !ViewInfo.bIsSceneCapture && IsVRSCompatibleWithOutputType(GetDisplayOutputFormat(ViewInfo));
 }
 
 FIntPoint FVariableRateShadingImageManager::GetSRITileSize()
@@ -318,7 +360,7 @@ FRDGTextureRef FVariableRateShadingImageManager::GetVariableRateShadingImage(FRD
 	const TArray<TRefCountPtr<IPooledRenderTarget>>* ExternalVRSSources, FVariableRateShadingImageManager::EVRSSourceType VRSTypesToExclude)
 {
 	// If the view doesn't support VRS, bail immediately
-	if (!IsVRSCompatibleWithView(ViewInfo))
+	if (!bVRSEnabledForFrame || !IsVRSCompatibleWithView(ViewInfo))
 	{
 		return nullptr;
 	}
@@ -327,9 +369,9 @@ FRDGTextureRef FVariableRateShadingImageManager::GetVariableRateShadingImage(FRD
 	SCOPED_NAMED_EVENT(GetVariableRateShadingImage, FColor::Yellow);
 
 	// Use debug rate if provided
-	if (CVarVRSDebugForceRate->GetInt() >= 0)
+	if (VRSForceRateForFrame >= 0)
 	{
-		return GetForceRateImage(GraphBuilder, CVarVRSDebugForceRate->GetInt(), GetImageTypeFromPassType(PassType));
+		return GetForceRateImage(GraphBuilder, VRSForceRateForFrame, GetImageTypeFromPassType(PassType));
 	}
 
 	// Otherwise collate all internal sources
@@ -371,9 +413,17 @@ FRDGTextureRef FVariableRateShadingImageManager::GetVariableRateShadingImage(FRD
 
 void FVariableRateShadingImageManager::PrepareImageBasedVRS(FRDGBuilder& GraphBuilder, const FSceneViewFamily& ViewFamily, const FMinimalSceneTextures& SceneTextures)
 {
+	bVRSEnabledForFrame = IsVRSSupportedByRHI() && IsVRSEnabled();
+	if (!bVRSEnabledForFrame)
+	{
+		return;
+	}
+
 	RDG_EVENT_SCOPE(GraphBuilder, "PrepareImageBasedVRS");
 	SCOPED_NAMED_EVENT(PrepareImageBasedVRS, FColor::Red);
 
+	VRSForceRateForFrame = CVarVRSDebugForceRate->GetInt();
+
 	// If no views support VRS, bail immediately
 	bool bIsAnyViewVRSCompatible = false;
 	for (const FSceneView* View : ViewFamily.Views)
@@ -457,10 +507,6 @@ void FVariableRateShadingImageManager::DrawDebugPreview(FRDGBuilder& GraphBuilde
 	{
 		PreviewImageType = static_cast<EVRSImageType>(ImageTypeAsInt);
 	}
-	else
-	{
-		UE_LOG(LogVRS, Warning, TEXT("Selected invalid preview mode, disabling preview"));
-	}
 
 	if (PreviewImageType == EVRSImageType::Disabled || !OutputSceneColor)
 	{
@@ -476,9 +522,9 @@ void FVariableRateShadingImageManager::DrawDebugPreview(FRDGBuilder& GraphBuilde
 			FRDGTextureRef PreviewTexture;
 			
 			// Use debug rate if provided
-			if (CVarVRSDebugForceRate->GetInt() >= 0)
+			if (VRSForceRateForFrame >= 0)
 			{
-				PreviewTexture = GetForceRateImage(GraphBuilder, CVarVRSDebugForceRate->GetInt(), PreviewImageType);
+				PreviewTexture = GetForceRateImage(GraphBuilder, VRSForceRateForFrame, PreviewImageType);
 			}
 
 			// Otherwise collate debug images
@@ -602,7 +648,6 @@ FRDGTextureRef FVariableRateShadingImageManager::GetForceRateImage(FRDGBuilder&
 
 	if (RateIndex >= ValidShadingRates.Num())
 	{
-		UE_LOG(LogVRS, Warning, TEXT("Selected forced shading rate exceeds maximum available, defaulting to 4x4"));
 		RateIndex = ValidShadingRates.Num() - 1; // Default to maximum shading rate if value exceeds valid rates
 	}
 
@@ -643,7 +688,6 @@ FVariableRateShadingImageManager::EVRSImageType FVariableRateShadingImageManager
 	}
 	else
 	{
-		UE_LOG(LogVRS, Warning, TEXT("Selected invalid image type, disabling VRS for pass"));
 		return EVRSImageType::Disabled;
 	}
 }
\ No newline at end of file
diff --git a/Engine/Source/Runtime/Renderer/Public/VariableRateShadingImageManager.h b/Engine/Source/Runtime/Renderer/Public/VariableRateShadingImageManager.h
--- a/Engine/Source/Runtime/Renderer/Public/VariableRateShadingImageManager.h
+++ b/Engine/Source/Runtime/Renderer/Public/VariableRateShadingImageManager.h
@@ -106,6 +106,7 @@ public:
 	TRefCountPtr<IPooledRenderTarget> GetMobileVariableRateShadingImage(const FSceneViewFamily& ViewFamily);
 
 	static bool IsVRSSupportedByRHI();
+	static bool IsVRSEnabled();
 	static bool IsVRSCompatibleWithView(const FViewInfo& View);
 	static bool IsVRSCompatibleWithOutputType(const EDisplayOutputFormat& DisplayOutputFormat);
 
@@ -122,6 +123,9 @@ private:
 	FRDGTextureRef GetForceRateImage(FRDGBuilder& GraphBuilder, int RateIndex = 0, EVRSImageType ImageType = EVRSImageType::Full);
 
 	EVRSImageType GetImageTypeFromPassType(EVRSPassType PassType);
+
+	bool bVRSEnabledForFrame = false;
+	int32 VRSForceRateForFrame = -1;
 };
 
 ENUM_CLASS_FLAGS(FVariableRateShadingImageManager::EVRSSourceType);
