Support DF visualization in FMobileSceneRenderer.
- Moved RenderMeshDistanceFieldVisualization(...) to FSceneRenderer

diff --git a/Engine/Source/Runtime/Renderer/Private/DeferredShadingRenderer.cpp b/Engine/Source/Runtime/Renderer/Private/DeferredShadingRenderer.cpp
--- a/Engine/Source/Runtime/Renderer/Private/DeferredShadingRenderer.cpp
+++ b/Engine/Source/Runtime/Renderer/Private/DeferredShadingRenderer.cpp
@@ -4273,7 +4273,7 @@ void FDeferredShadingSceneRenderer::Render(FRDGBuilder& GraphBuilder)
 		// Draw visualizations just before use to avoid target contamination
 		if (ViewFamily.EngineShowFlags.VisualizeMeshDistanceFields || ViewFamily.EngineShowFlags.VisualizeGlobalDistanceField)
 		{
-			RenderMeshDistanceFieldVisualization(GraphBuilder, SceneTextures, FDistanceFieldAOParameters(Scene->DefaultMaxDistanceFieldOcclusionDistance));
+			RenderMeshDistanceFieldVisualization(GraphBuilder, SceneTextures);
 		}
 
 		if (bRenderDeferredLighting)
diff --git a/Engine/Source/Runtime/Renderer/Private/DeferredShadingRenderer.h b/Engine/Source/Runtime/Renderer/Private/DeferredShadingRenderer.h
--- a/Engine/Source/Runtime/Renderer/Private/DeferredShadingRenderer.h
+++ b/Engine/Source/Runtime/Renderer/Private/DeferredShadingRenderer.h
@@ -719,11 +719,6 @@ private:
 		FRDGTextureRef DistanceFieldNormal,
 		FRDGTextureRef& OutDynamicBentNormalAO);
 
-	void RenderMeshDistanceFieldVisualization(
-		FRDGBuilder& GraphBuilder,
-		const FMinimalSceneTextures& SceneTextures,
-		const FDistanceFieldAOParameters& Parameters);
-
 	FSSDSignalTextures RenderLumenFinalGather(
 		FRDGBuilder& GraphBuilder,
 		const FSceneTextures& SceneTextures,
diff --git a/Engine/Source/Runtime/Renderer/Private/DistanceFieldVisualization.cpp b/Engine/Source/Runtime/Renderer/Private/DistanceFieldVisualization.cpp
--- a/Engine/Source/Runtime/Renderer/Private/DistanceFieldVisualization.cpp
+++ b/Engine/Source/Runtime/Renderer/Private/DistanceFieldVisualization.cpp
@@ -40,6 +40,7 @@ public:
 		SHADER_PARAMETER_STRUCT_INCLUDE(FAOParameters, AOParameters)
 		SHADER_PARAMETER_STRUCT_INCLUDE(FGlobalDistanceFieldParameters2, GlobalDistanceFieldParameters)
 		SHADER_PARAMETER_RDG_UNIFORM_BUFFER(FSceneTextureUniformParameters, SceneTextures)
+		SHADER_PARAMETER_RDG_UNIFORM_BUFFER(FMobileSceneTextureUniformParameters, MobileSceneTextures)
 		SHADER_PARAMETER_RDG_TEXTURE_UAV(RWTexture2D<float4>, RWVisualizeMeshDistanceFields)
 		SHADER_PARAMETER(FVector2f, NumGroups)
 	END_SHADER_PARAMETER_STRUCT()
@@ -74,7 +75,6 @@ public:
 	
 	BEGIN_SHADER_PARAMETER_STRUCT(FParameters, )
 		SHADER_PARAMETER_STRUCT_REF(FViewUniformShaderParameters, View)
-		SHADER_PARAMETER_RDG_UNIFORM_BUFFER(FSceneTextureUniformParameters, SceneTextures)	
 		SHADER_PARAMETER_RDG_TEXTURE(Texture2D, VisualizeDistanceFieldTexture)
 		SHADER_PARAMETER_SAMPLER(SamplerState, VisualizeDistanceFieldSampler)
 		RENDER_TARGET_BINDING_SLOTS()
@@ -93,11 +93,10 @@ public:
 
 IMPLEMENT_GLOBAL_SHADER(FVisualizeDistanceFieldUpsamplePS, "/Engine/Private/DistanceFieldVisualization.usf", "VisualizeDistanceFieldUpsamplePS", SF_Pixel);
 
-void FDeferredShadingSceneRenderer::RenderMeshDistanceFieldVisualization(
-	FRDGBuilder& GraphBuilder,
-	const FMinimalSceneTextures& SceneTextures,
-	const FDistanceFieldAOParameters& Parameters)
+void FSceneRenderer::RenderMeshDistanceFieldVisualization(FRDGBuilder& GraphBuilder, const FMinimalSceneTextures& SceneTextures)
 {
+	const FDistanceFieldAOParameters DFAOParameters = FDistanceFieldAOParameters(Scene->DefaultMaxDistanceFieldOcclusionDistance);
+
 	const FViewInfo& FirstView = Views[0];
 
 	const bool bAnyMeshSDFs = Scene->DistanceFieldSceneData.NumObjectsInBuffer > 0;
@@ -105,13 +104,13 @@ void FDeferredShadingSceneRenderer::RenderMeshDistanceFieldVisualization(
 
 	const bool bVisualizeGlobalDistanceField = FirstView.Family->EngineShowFlags.VisualizeGlobalDistanceField 
 		&& ShouldPrepareGlobalDistanceField() 
-		&& UseGlobalDistanceField(Parameters) 
+		&& UseGlobalDistanceField(DFAOParameters)
 		&& (bAnyMeshSDFs || bAnyHeightFields);
 
 	const bool bVisualizeMeshDistanceFields = FirstView.Family->EngineShowFlags.VisualizeMeshDistanceFields && bAnyMeshSDFs;
 
-	if (FeatureLevel < ERHIFeatureLevel::SM5
-		|| !DoesPlatformSupportDistanceFields(FirstView.GetShaderPlatform())
+	if (!DoesPlatformSupportDistanceFields(FirstView.GetShaderPlatform())
+		|| !IsUsingDistanceFields(FirstView.GetShaderPlatform())
 		|| !(bVisualizeGlobalDistanceField || bVisualizeMeshDistanceFields))
 	{
 		return;
@@ -155,7 +154,7 @@ void FDeferredShadingSceneRenderer::RenderMeshDistanceFieldVisualization(
 				ObjectIndirectArguments,
 				CulledObjectBufferParameters);
 
-			CullObjectsToView(GraphBuilder, *Scene, View, Parameters, CulledObjectBufferParameters);
+			CullObjectsToView(GraphBuilder, *Scene, View, DFAOParameters, CulledObjectBufferParameters);
 		}
 
 		uint32 GroupSizeX = FMath::DivideAndRoundUp(View.ViewRect.Size().X / GAODownsampleFactor, GDistanceFieldAOTileSizeX);
@@ -170,10 +169,11 @@ void FDeferredShadingSceneRenderer::RenderMeshDistanceFieldVisualization(
 		PassParameters->DistanceFieldObjectBuffers = DistanceField::SetupObjectBufferParameters(GraphBuilder, Scene->DistanceFieldSceneData);
 		PassParameters->DistanceFieldCulledObjectBuffers = CulledObjectBufferParameters;
 		PassParameters->DistanceFieldAtlas = DistanceField::SetupAtlasParameters(GraphBuilder, Scene->DistanceFieldSceneData);
-		PassParameters->AOParameters = DistanceField::SetupAOShaderParameters(Parameters);
+		PassParameters->AOParameters = DistanceField::SetupAOShaderParameters(DFAOParameters);
 		PassParameters->GlobalDistanceFieldParameters = SetupGlobalDistanceFieldParameters(View.GlobalDistanceFieldInfo.ParameterData);
 		PassParameters->NumGroups = FVector2f(GroupSizeX, GroupSizeY);
 		PassParameters->SceneTextures = SceneTextures.UniformBuffer;
+		PassParameters->MobileSceneTextures = SceneTextures.MobileUniformBuffer;
 		PassParameters->RWVisualizeMeshDistanceFields = GraphBuilder.CreateUAV(VisualizeResultTexture);
 
 		FComputeShaderUtils::AddPass(GraphBuilder, RDG_EVENT_NAME("VisualizeMeshDistanceFieldCS"), ComputeShader, PassParameters, FIntVector(GroupSizeX, GroupSizeY, 1));
@@ -187,7 +187,6 @@ void FDeferredShadingSceneRenderer::RenderMeshDistanceFieldVisualization(
 
 		auto* PassParameters = GraphBuilder.AllocParameters<FVisualizeDistanceFieldUpsamplePS::FParameters>();
 		PassParameters->View = View.ViewUniformBuffer;
-		PassParameters->SceneTextures = SceneTextures.UniformBuffer;
 		PassParameters->VisualizeDistanceFieldTexture = VisualizeResultTexture;
 		PassParameters->VisualizeDistanceFieldSampler = TStaticSamplerState<SF_Bilinear>::GetRHI();
 		PassParameters->RenderTargets[0] = FRenderTargetBinding(SceneTextures.Color.Target, ERenderTargetLoadAction::ELoad);
@@ -198,4 +197,4 @@ void FDeferredShadingSceneRenderer::RenderMeshDistanceFieldVisualization(
 
 		AddDrawScreenPass(GraphBuilder, {}, View, OutputViewport, InputViewport, PixelShader, PassParameters);
 	}
-}
\ No newline at end of file
+}
diff --git a/Engine/Source/Runtime/Renderer/Private/MobileShadingRenderer.cpp b/Engine/Source/Runtime/Renderer/Private/MobileShadingRenderer.cpp
--- a/Engine/Source/Runtime/Renderer/Private/MobileShadingRenderer.cpp
+++ b/Engine/Source/Runtime/Renderer/Private/MobileShadingRenderer.cpp
@@ -817,6 +817,7 @@ void FMobileSceneRenderer::Render(FRDGBuilder& GraphBuilder)
 
 	FRDGExternalAccessQueue ExternalAccessQueue;
 
+	bRequiresDistanceField = true;
 	if (bRequiresDistanceField)
 	{
 		PrepareDistanceFieldScene(GraphBuilder, ExternalAccessQueue);
@@ -1088,6 +1089,11 @@ void FMobileSceneRenderer::Render(FRDGBuilder& GraphBuilder)
 
 		RenderPixelProjectedReflection(GraphBuilder, SceneTextures.Color.Resolve, SceneTextures.Depth.Resolve, SceneTextures.PixelProjectedReflection, PlanarReflectionSceneProxy);
 	}
+
+	if (ViewFamily.EngineShowFlags.VisualizeMeshDistanceFields || ViewFamily.EngineShowFlags.VisualizeGlobalDistanceField)
+	{
+		RenderMeshDistanceFieldVisualization(GraphBuilder, SceneTextures);
+	}
 	
 	if (bUseVirtualTexturing)
 	{
diff --git a/Engine/Source/Runtime/Renderer/Private/SceneRendering.h b/Engine/Source/Runtime/Renderer/Private/SceneRendering.h
--- a/Engine/Source/Runtime/Renderer/Private/SceneRendering.h
+++ b/Engine/Source/Runtime/Renderer/Private/SceneRendering.h
@@ -2547,6 +2547,8 @@ protected:
 		EVelocityPass VelocityPass,
 		bool bForceVelocity);
 
+	void RenderMeshDistanceFieldVisualization(FRDGBuilder& GraphBuilder, const FMinimalSceneTextures& SceneTextures);
+
 protected:
 	FGPUSceneDynamicContext GPUSceneDynamicContext;
 
