Revert hardcoded bRequiresDistanceField in FMobileSceneRenderer::Render(...)

diff --git a/Engine/Source/Runtime/Renderer/Private/MobileShadingRenderer.cpp b/Engine/Source/Runtime/Renderer/Private/MobileShadingRenderer.cpp
--- a/Engine/Source/Runtime/Renderer/Private/MobileShadingRenderer.cpp
+++ b/Engine/Source/Runtime/Renderer/Private/MobileShadingRenderer.cpp
@@ -817,7 +817,6 @@ void FMobileSceneRenderer::Render(FRDGBuilder& GraphBuilder)
 
 	FRDGExternalAccessQueue ExternalAccessQueue;
 
-	bRequiresDistanceField = true;
 	if (bRequiresDistanceField)
 	{
 		PrepareDistanceFieldScene(GraphBuilder, ExternalAccessQueue);
