Fix incorrect ref counting decrement when texture are removed from rect. light texture manager.

diff --git a/Engine/Source/Runtime/Renderer/Private/RectLightTextureManager.cpp b/Engine/Source/Runtime/Renderer/Private/RectLightTextureManager.cpp
--- a/Engine/Source/Runtime/Renderer/Private/RectLightTextureManager.cpp
+++ b/Engine/Source/Runtime/Renderer/Private/RectLightTextureManager.cpp
@@ -1473,7 +1473,7 @@ void RemoveTexture(uint32 InSlotIndex)
 	if (InSlotIndex != InvalidSlotIndex && InSlotIndex < uint32(GRectLightTextureManager.AtlasSlots.Num()))
 	{
 		// If it is the last light referencing this texture, we retires the atlas slot
-		if (GRectLightTextureManager.AtlasSlots[InSlotIndex].RefCount == 1)
+		if (--GRectLightTextureManager.AtlasSlots[InSlotIndex].RefCount == 0)
 		{
 			// Add pending slots to clean-up the layout during the next update call
 			GRectLightTextureManager.DeletedSlots.Add(GRectLightTextureManager.AtlasSlots[InSlotIndex].Rect);
