Prevent VRS and VRS debug preview overlay from being applied to ViewFamilies that aren't updated in real-time, like content thumbnails.

diff --git a/Engine/Source/Runtime/Renderer/Private/VariableRateShading/VariableRateShadingImageManager.cpp b/Engine/Source/Runtime/Renderer/Private/VariableRateShading/VariableRateShadingImageManager.cpp
--- a/Engine/Source/Runtime/Renderer/Private/VariableRateShading/VariableRateShadingImageManager.cpp
+++ b/Engine/Source/Runtime/Renderer/Private/VariableRateShading/VariableRateShadingImageManager.cpp
@@ -322,6 +322,7 @@ bool FVariableRateShadingImageManager::IsVRSEnabled()
 
 bool FVariableRateShadingImageManager::IsVRSCompatibleWithOutputType(const EDisplayOutputFormat& OutputFormat)
 {
+	// VRS texture generation is currently only compatible with SDR and HDR10
 	return OutputFormat == EDisplayOutputFormat::SDR_sRGB
 		|| OutputFormat == EDisplayOutputFormat::HDR_ACES_1000nit_ST2084
 		|| OutputFormat == EDisplayOutputFormat::HDR_ACES_2000nit_ST2084;
@@ -329,12 +330,12 @@ bool FVariableRateShadingImageManager::IsVRSCompatibleWithOutputType(const EDisp
 
 bool FVariableRateShadingImageManager::IsVRSCompatibleWithView(const FViewInfo& ViewInfo)
 {
-	// The VRS texture generation is currently only compatible with SDR and HDR10
-	
 	// TODO: Investigate if it's worthwhile getting scene captures working. Things that we'll need to take care of
 	// is to associate shading rate texture image with main scene, and scene capture.  But what if there is
 	// more than 1 scene capture?  Is there a unique identifier that connects two frames of scene capture.
-	return !ViewInfo.bIsSceneCapture && IsVRSCompatibleWithOutputType(GetDisplayOutputFormat(ViewInfo));
+	return !ViewInfo.bIsSceneCapture
+		&& ViewInfo.Family->bRealtimeUpdate
+		&& IsVRSCompatibleWithOutputType(GetDisplayOutputFormat(ViewInfo));
 }
 
 FIntPoint FVariableRateShadingImageManager::GetSRITileSize()
