Change VRS manager to return nullptr rather than a full 1x1 shading rate image when no generators are active or VRS is disabled for a given pass.

diff --git a/Engine/Source/Runtime/Renderer/Private/VariableRateShading/VariableRateShadingImageManager.cpp b/Engine/Source/Runtime/Renderer/Private/VariableRateShading/VariableRateShadingImageManager.cpp
--- a/Engine/Source/Runtime/Renderer/Private/VariableRateShading/VariableRateShadingImageManager.cpp
+++ b/Engine/Source/Runtime/Renderer/Private/VariableRateShading/VariableRateShadingImageManager.cpp
@@ -360,8 +360,10 @@ FRDGTextureDesc FVariableRateShadingImageManager::GetSRIDesc()
 FRDGTextureRef FVariableRateShadingImageManager::GetVariableRateShadingImage(FRDGBuilder& GraphBuilder, const FViewInfo& ViewInfo, FVariableRateShadingImageManager::EVRSPassType PassType,
 	const TArray<TRefCountPtr<IPooledRenderTarget>>* ExternalVRSSources, FVariableRateShadingImageManager::EVRSSourceType VRSTypesToExclude)
 {
-	// If the view doesn't support VRS, bail immediately
-	if (!bVRSEnabledForFrame || !IsVRSCompatibleWithView(ViewInfo))
+	EVRSImageType ImageType = GetImageTypeFromPassType(PassType);
+
+	// If the view doesn't support VRS or this pass is disabled, bail immediately
+	if (!bVRSEnabledForFrame || !IsVRSCompatibleWithView(ViewInfo) || ImageType == EVRSImageType::Disabled)
 	{
 		return nullptr;
 	}
@@ -372,7 +374,7 @@ FRDGTextureRef FVariableRateShadingImageManager::GetVariableRateShadingImage(FRD
 	// Use debug rate if provided
 	if (VRSForceRateForFrame >= 0)
 	{
-		return GetForceRateImage(GraphBuilder, VRSForceRateForFrame, GetImageTypeFromPassType(PassType));
+		return GetForceRateImage(GraphBuilder, VRSForceRateForFrame, ImageType);
 	}
 
 	// Otherwise collate all internal sources
@@ -383,7 +385,7 @@ FRDGTextureRef FVariableRateShadingImageManager::GetVariableRateShadingImage(FRD
 		FRDGTextureRef Image = nullptr;
 		if (Generator->IsEnabledForView(ViewInfo) && !EnumHasAnyFlags(VRSTypesToExclude, Generator->GetType()))
 		{
-			Image = Generator->GetImage(GraphBuilder, ViewInfo, GetImageTypeFromPassType(PassType));
+			Image = Generator->GetImage(GraphBuilder, ViewInfo, ImageType);
 		}
 
 		if (Image)
@@ -393,6 +395,7 @@ FRDGTextureRef FVariableRateShadingImageManager::GetVariableRateShadingImage(FRD
 	}
 
 	// If we have more than one internal source, combine the first available two
+	// If we have exactly one, the combiner will just return that
 	if (InternalVRSSources.Num())
 	{
 		return CombineShadingRateImages(GraphBuilder, ViewInfo, InternalVRSSources);
@@ -405,10 +408,10 @@ FRDGTextureRef FVariableRateShadingImageManager::GetVariableRateShadingImage(FRD
 		return GraphBuilder.RegisterExternalTexture((*ExternalVRSSources)[0]);
 	}
 
-	// Default to 1x1 shading rate if no sources are available
+	// Default to nullptr if no sources are available
 	else
 	{
-		return GetForceRateImage(GraphBuilder);
+		return nullptr;
 	}
 }
 
@@ -531,7 +534,6 @@ void FVariableRateShadingImageManager::DrawDebugPreview(FRDGBuilder& GraphBuilde
 			// Otherwise collate debug images
 			else
 			{
-				
 				TArray<FRDGTextureRef> InternalVRSSources;
 
 				for (TUniquePtr<IVariableRateShadingImageGenerator>& Generator : ImageGenerators)
@@ -549,8 +551,13 @@ void FVariableRateShadingImageManager::DrawDebugPreview(FRDGBuilder& GraphBuilde
 				}
 
 				PreviewTexture = CombineShadingRateImages(GraphBuilder, *ViewInfo, InternalVRSSources);
+
+				// Generate a dummy 1x1 image if we have no VRS sources
+				if (!PreviewTexture)
+				{
+					PreviewTexture = GetForceRateImage(GraphBuilder);
+				}
 			}
-			
 
 			// If we have an active debug image, render it as a preview overlay
 			auto& RHICmdList = GraphBuilder.RHICmdList;
@@ -610,7 +617,7 @@ FRDGTextureRef FVariableRateShadingImageManager::CombineShadingRateImages(FRDGBu
 	// TODO: Support combining more textures
 	if (Sources.Num() < 1)
 	{
-		return GetForceRateImage(GraphBuilder); // Fall back to uniform 1x1 shading rate if no images provided
+		return nullptr;
 	}
 	else if (Sources.Num() == 1)
 	{
