Cleaned up logic in FSceneRenderer::RenderMeshDistanceFieldVisualization

diff --git a/Engine/Source/Runtime/Renderer/Private/DistanceFieldVisualization.cpp b/Engine/Source/Runtime/Renderer/Private/DistanceFieldVisualization.cpp
--- a/Engine/Source/Runtime/Renderer/Private/DistanceFieldVisualization.cpp
+++ b/Engine/Source/Runtime/Renderer/Private/DistanceFieldVisualization.cpp
@@ -95,22 +95,13 @@ IMPLEMENT_GLOBAL_SHADER(FVisualizeDistanceFieldUpsamplePS, "/Engine/Private/Dist
 
 void FSceneRenderer::RenderMeshDistanceFieldVisualization(FRDGBuilder& GraphBuilder, const FMinimalSceneTextures& SceneTextures)
 {
-	const FDistanceFieldAOParameters DFAOParameters = FDistanceFieldAOParameters(Scene->DefaultMaxDistanceFieldOcclusionDistance);
-
-	const FViewInfo& FirstView = Views[0];
-
-	const bool bAnyMeshSDFs = Scene->DistanceFieldSceneData.NumObjectsInBuffer > 0;
-	const bool bAnyHeightFields = Scene->DistanceFieldSceneData.HeightfieldPrimitives.Num() > 0;
+	const bool bVisualizeGlobalDistanceField = ViewFamily.EngineShowFlags.VisualizeGlobalDistanceField;
+	const bool bVisualizeMeshDistanceFields = ViewFamily.EngineShowFlags.VisualizeMeshDistanceFields
+		&& (Scene->DistanceFieldSceneData.NumObjectsInBuffer > 0)
+		&& !bVisualizeGlobalDistanceField; // only one visualization enabled at a time
 
-	const bool bVisualizeGlobalDistanceField = FirstView.Family->EngineShowFlags.VisualizeGlobalDistanceField 
-		&& ShouldPrepareGlobalDistanceField() 
-		&& UseGlobalDistanceField(DFAOParameters)
-		&& (bAnyMeshSDFs || bAnyHeightFields);
-
-	const bool bVisualizeMeshDistanceFields = FirstView.Family->EngineShowFlags.VisualizeMeshDistanceFields && bAnyMeshSDFs;
-
-	if (!DoesPlatformSupportDistanceFields(FirstView.GetShaderPlatform())
-		|| !IsUsingDistanceFields(FirstView.GetShaderPlatform())
+	if (!DoesPlatformSupportDistanceFields(ShaderPlatform)
+		|| !IsUsingDistanceFields(ShaderPlatform)
 		|| !(bVisualizeGlobalDistanceField || bVisualizeMeshDistanceFields))
 	{
 		return;
@@ -124,29 +115,20 @@ void FSceneRenderer::RenderMeshDistanceFieldVisualization(FRDGBuilder& GraphBuil
 
 	FRDGTextureRef VisualizeResultTexture = nullptr;
 	{
+		const FViewInfo& FirstView = Views[0];
 		const FIntPoint BufferSize = GetBufferSizeForAO(FirstView);
 		const FRDGTextureDesc Desc(FRDGTextureDesc::Create2D(BufferSize, PF_FloatRGBA, FClearValueBinding::None, TexCreate_UAV));
 		VisualizeResultTexture = GraphBuilder.CreateTexture(Desc, TEXT("VisualizeDistanceField"));
 	}
 
-	FVisualizeMeshDistanceFieldCS::FPermutationDomain PermutationVector;
-	PermutationVector.Set<FVisualizeMeshDistanceFieldCS::FUseGlobalDistanceFieldDim>(bVisualizeGlobalDistanceField);
-	PermutationVector.Set<FVisualizeMeshDistanceFieldCS::FCoverageBasedExpand>(IsLumenEnabled(FirstView));
-	PermutationVector.Set<FVisualizeMeshDistanceFieldCS::FSimpleCoverageBasedExpand>(IsLumenEnabled(FirstView) && Lumen::UseGlobalSDFSimpleCoverageBasedExpand());
-	extern int32 GDistanceFieldOffsetDataStructure;
-	PermutationVector.Set<FVisualizeMeshDistanceFieldCS::FOffsetDataStructure>(GDistanceFieldOffsetDataStructure);
-	TShaderMapRef<FVisualizeMeshDistanceFieldCS> ComputeShader(FirstView.ShaderMap, PermutationVector);
+	const FDistanceFieldAOParameters DFAOParameters = FDistanceFieldAOParameters(Scene->DefaultMaxDistanceFieldOcclusionDistance);
 
 	for (const FViewInfo& View : Views)
 	{
 		FRDGBufferRef ObjectIndirectArguments = nullptr;
 		FDistanceFieldCulledObjectBufferParameters CulledObjectBufferParameters;
 
-		if (bVisualizeGlobalDistanceField)
-		{
-			check(View.GlobalDistanceFieldInfo.Clipmaps.Num() > 0);
-		}
-		else
+		if (bVisualizeMeshDistanceFields)
 		{
 			AllocateDistanceFieldCulledObjectBuffers(
 				GraphBuilder,
@@ -176,11 +158,18 @@ void FSceneRenderer::RenderMeshDistanceFieldVisualization(FRDGBuilder& GraphBuil
 		PassParameters->MobileSceneTextures = SceneTextures.MobileUniformBuffer;
 		PassParameters->RWVisualizeMeshDistanceFields = GraphBuilder.CreateUAV(VisualizeResultTexture);
 
+		FVisualizeMeshDistanceFieldCS::FPermutationDomain PermutationVector;
+		PermutationVector.Set<FVisualizeMeshDistanceFieldCS::FUseGlobalDistanceFieldDim>(bVisualizeGlobalDistanceField);
+		PermutationVector.Set<FVisualizeMeshDistanceFieldCS::FCoverageBasedExpand>(IsLumenEnabled(View));
+		PermutationVector.Set<FVisualizeMeshDistanceFieldCS::FSimpleCoverageBasedExpand>(IsLumenEnabled(View) && Lumen::UseGlobalSDFSimpleCoverageBasedExpand());
+		extern int32 GDistanceFieldOffsetDataStructure;
+		PermutationVector.Set<FVisualizeMeshDistanceFieldCS::FOffsetDataStructure>(GDistanceFieldOffsetDataStructure);
+
+		auto ComputeShader = View.ShaderMap->GetShader<FVisualizeMeshDistanceFieldCS>(PermutationVector);
+
 		FComputeShaderUtils::AddPass(GraphBuilder, RDG_EVENT_NAME("VisualizeMeshDistanceFieldCS"), ComputeShader, PassParameters, FIntVector(GroupSizeX, GroupSizeY, 1));
 	}
 
-	TShaderMapRef<FVisualizeDistanceFieldUpsamplePS> PixelShader(FirstView.ShaderMap);
-
 	for (const FViewInfo& View : Views)
 	{
 		RDG_GPU_MASK_SCOPE(GraphBuilder, View.GPUMask);
@@ -195,6 +184,8 @@ void FSceneRenderer::RenderMeshDistanceFieldVisualization(FRDGBuilder& GraphBuil
 		const FScreenPassTextureViewport InputViewport(VisualizeResultTexture, GetDownscaledRect(View.ViewRect, GAODownsampleFactor));
 		const FScreenPassTextureViewport OutputViewport(SceneTextures.Color.Target, View.ViewRect);
 
+		auto PixelShader = View.ShaderMap->GetShader<FVisualizeDistanceFieldUpsamplePS>();
+
 		AddDrawScreenPass(GraphBuilder, {}, View, OutputViewport, InputViewport, PixelShader, PassParameters);
 	}
 }
