Prevent AddAfter pass attempting to create a SRV from the backbuffer.

#rnx

diff --git a/Engine/Source/Runtime/Renderer/Private/PostProcess/PostProcessing.cpp b/Engine/Source/Runtime/Renderer/Private/PostProcess/PostProcessing.cpp
--- a/Engine/Source/Runtime/Renderer/Private/PostProcess/PostProcessing.cpp
+++ b/Engine/Source/Runtime/Renderer/Private/PostProcess/PostProcessing.cpp
@@ -518,7 +518,7 @@ void AddPostProcessingPasses(
 		return PostProcessMaterialInputs;
 	};
 
-	const auto AddAfterPassForSceneColorSlice = [&](EPass InPass, const FScreenPassTextureSlice& InSceneColor) -> FScreenPassTextureSlice
+	const auto AddAfterPass = [&](EPass InPass, FScreenPassTexture InSceneColor) -> FScreenPassTexture
 	{
 		// In some cases (e.g. OCIO color conversion) we want View Extensions to be able to add extra custom post processing after the pass.
 
@@ -526,32 +526,30 @@ void AddPostProcessingPasses(
 
 		if (PassCallbacks.Num())
 		{
-			FScreenPassTexture SceneColor = FScreenPassTexture::CopyFromSlice(GraphBuilder, InSceneColor);
-
-			FPostProcessMaterialInputs InOutPostProcessAfterPassInputs = GetPostProcessMaterialInputs(SceneColor);
+			FPostProcessMaterialInputs InOutPostProcessAfterPassInputs = GetPostProcessMaterialInputs(InSceneColor);
 
 			for (int32 AfterPassCallbackIndex = 0; AfterPassCallbackIndex < PassCallbacks.Num(); AfterPassCallbackIndex++)
 			{
-				InOutPostProcessAfterPassInputs.SetInput(EPostProcessMaterialInput::SceneColor, SceneColor);
+				InOutPostProcessAfterPassInputs.SetInput(EPostProcessMaterialInput::SceneColor, InSceneColor);
 
 				FAfterPassCallbackDelegate& AfterPassCallback = PassCallbacks[AfterPassCallbackIndex];
 				PassSequence.AcceptOverrideIfLastPass(InPass, InOutPostProcessAfterPassInputs.OverrideOutput, AfterPassCallbackIndex);
-				SceneColor = AfterPassCallback.Execute(GraphBuilder, View, InOutPostProcessAfterPassInputs);
+				InSceneColor = AfterPassCallback.Execute(GraphBuilder, View, InOutPostProcessAfterPassInputs);
 			}
-
-			return FScreenPassTextureSlice::CreateFromScreenPassTexture(GraphBuilder, SceneColor);
 		}
 
-		return InSceneColor;
+		return MoveTemp(InSceneColor);
 	};
 
-	const auto AddAfterPass = [&](EPass InPass, const FScreenPassTexture InSceneColor) -> FScreenPassTexture
+	const auto AddAfterPassForSceneColorSlice = [&](EPass InPass, const FScreenPassTextureSlice& InSceneColor) -> FScreenPassTextureSlice
 	{
 		FAfterPassCallbackDelegateArray& PassCallbacks = PassSequence.GetAfterPassCallbacks(InPass);
 
 		if (PassCallbacks.Num())
 		{
-			return FScreenPassTexture(AddAfterPassForSceneColorSlice(InPass, FScreenPassTextureSlice::CreateFromScreenPassTexture(GraphBuilder, InSceneColor)));
+			FScreenPassTexture SceneColor = FScreenPassTexture::CopyFromSlice(GraphBuilder, InSceneColor);
+
+			return FScreenPassTextureSlice::CreateFromScreenPassTexture(GraphBuilder, AddAfterPass(InPass, SceneColor));
 		}
 
 		return InSceneColor;
