Adding MeshPassUtils to replace FComputeShaderUtils for compute based mesh passes to set mesh bindings and pass parameters at the same time.

diff --git a/Engine/Source/Runtime/Renderer/Private/HeterogeneousVolumes/HeterogeneousVolumesLiveShadingPipeline.cpp b/Engine/Source/Runtime/Renderer/Private/HeterogeneousVolumes/HeterogeneousVolumesLiveShadingPipeline.cpp
--- a/Engine/Source/Runtime/Renderer/Private/HeterogeneousVolumes/HeterogeneousVolumesLiveShadingPipeline.cpp
+++ b/Engine/Source/Runtime/Renderer/Private/HeterogeneousVolumes/HeterogeneousVolumesLiveShadingPipeline.cpp
@@ -5,6 +5,7 @@
 
 #include "LightRendering.h"
 #include "LocalVertexFactory.h"
+#include "MeshPassUtils.h"
 #include "PixelShaderUtils.h"
 #include "RayTracingDefinitions.h"
 #include "RayTracingInstance.h"
@@ -365,10 +366,8 @@ void AddComputePass(
 				ComputeShader->ShaderLooseBindings.SetParameters(SingleShaderBindings, PassParameters);
 				ShaderBindings.Finalize(&PassShaders);
 			}
-			SetComputePipelineState(RHICmdList, ComputeShader.GetComputeShader());
-			ShaderBindings.SetOnCommandList(RHICmdList, ComputeShader.GetComputeShader());
 
-			FComputeShaderUtils::Dispatch(RHICmdList, ComputeShader, *PassParameters, GroupCount);
+			UE::MeshPassUtils::Dispatch(RHICmdList, ComputeShader, ShaderBindings, *PassParameters, GroupCount);
 		}
 	);
 }
diff --git a/Engine/Source/Runtime/Renderer/Private/HeterogeneousVolumes/HeterogeneousVolumesMaterialBakingPipeline.cpp b/Engine/Source/Runtime/Renderer/Private/HeterogeneousVolumes/HeterogeneousVolumesMaterialBakingPipeline.cpp
--- a/Engine/Source/Runtime/Renderer/Private/HeterogeneousVolumes/HeterogeneousVolumesMaterialBakingPipeline.cpp
+++ b/Engine/Source/Runtime/Renderer/Private/HeterogeneousVolumes/HeterogeneousVolumesMaterialBakingPipeline.cpp
@@ -4,6 +4,7 @@
 #include "HeterogeneousVolumeInterface.h"
 
 #include "LocalVertexFactory.h"
+#include "MeshPassUtils.h"
 #include "PixelShaderUtils.h"
 #include "RayTracingDefinitions.h"
 #include "RayTracingInstance.h"
@@ -190,11 +191,8 @@ void ComputeHeterogeneousVolumeBakeMaterial(
 
 					ShaderBindings.Finalize(&PassShaders);
 				}
-				SetComputePipelineState(RHICmdList, ComputeShader.GetComputeShader());
-				ShaderBindings.SetOnCommandList(RHICmdList, ComputeShader.GetComputeShader());
 
-				FComputeShaderUtils::Dispatch(RHICmdList, ComputeShader, *PassParameters, GroupCount);
-				//FComputeShaderUtils::DispatchIndirect(RHICmdList, ComputeShader, *PassParameters, IndirectArgs->GetIndirectRHICallBuffer(), IndirectOffset);
+				UE::MeshPassUtils::Dispatch(RHICmdList, ComputeShader, ShaderBindings, *PassParameters, GroupCount);
 			}
 		}
 	);
diff --git a/Engine/Source/Runtime/Renderer/Private/HeterogeneousVolumes/HeterogeneousVolumesVoxelGridPipeline.cpp b/Engine/Source/Runtime/Renderer/Private/HeterogeneousVolumes/HeterogeneousVolumesVoxelGridPipeline.cpp
--- a/Engine/Source/Runtime/Renderer/Private/HeterogeneousVolumes/HeterogeneousVolumesVoxelGridPipeline.cpp
+++ b/Engine/Source/Runtime/Renderer/Private/HeterogeneousVolumes/HeterogeneousVolumesVoxelGridPipeline.cpp
@@ -4,6 +4,7 @@
 
 #include "HeterogeneousVolumeInterface.h"
 #include "LocalVertexFactory.h"
+#include "MeshPassUtils.h"
 #include "PixelShaderUtils.h"
 #include "RayTracingDefinitions.h"
 #include "RayTracingInstance.h"
@@ -1933,9 +1934,8 @@ void RasterizeVolumesIntoFrustumVoxelGrid(
 
 						ShaderBindings.Finalize(&PassShaders);
 					}
-					ShaderBindings.SetOnCommandList(RHICmdList, ComputeShader.GetComputeShader());
 
-					FComputeShaderUtils::DispatchIndirect(RHICmdList, ComputeShader, *PassParameters, PassParameters->IndirectArgs->GetIndirectRHICallBuffer(), 0);
+					UE::MeshPassUtils::DispatchIndirect(RHICmdList, ComputeShader, ShaderBindings, *PassParameters, PassParameters->IndirectArgs->GetIndirectRHICallBuffer(), 0);
 				}
 			}
 		);
@@ -2402,9 +2402,8 @@ void RasterizeVolumesIntoOrthoVoxelGrid(
 
 							ShaderBindings.Finalize(&PassShaders);
 						}
-						ShaderBindings.SetOnCommandList(RHICmdList, ComputeShader.GetComputeShader());
 
-						FComputeShaderUtils::DispatchIndirect(RHICmdList, ComputeShader, *PassParameters, PassParameters->IndirectArgs->GetIndirectRHICallBuffer(), 0);
+						UE::MeshPassUtils::DispatchIndirect(RHICmdList, ComputeShader, ShaderBindings, *PassParameters, PassParameters->IndirectArgs->GetIndirectRHICallBuffer(), 0);
 					}
 				}
 			);
diff --git a/Engine/Source/Runtime/Renderer/Private/VolumetricCloudRendering.cpp b/Engine/Source/Runtime/Renderer/Private/VolumetricCloudRendering.cpp
--- a/Engine/Source/Runtime/Renderer/Private/VolumetricCloudRendering.cpp
+++ b/Engine/Source/Runtime/Renderer/Private/VolumetricCloudRendering.cpp
@@ -8,6 +8,7 @@
 #include "Components/VolumetricCloudComponent.h"
 #include "VolumetricCloudProxy.h"
 #include "DeferredShadingRenderer.h"
+#include "MeshPassUtils.h"
 #include "PixelShaderUtils.h"
 #include "RenderGraphUtils.h"
 #include "ScenePrivate.h"
@@ -2360,12 +2361,7 @@ void FSceneRenderer::RenderVolumetricCloudsInternal(FRDGBuilder& GraphBuilder, F
 
 					ShaderBindings.Finalize(&PassShaders);
 
-					FRHIComputeShader* ComputeShaderRHI = ComputeShader.GetComputeShader();
-					SetComputePipelineState(RHICmdList, ComputeShaderRHI);
-					ShaderBindings.SetOnCommandList(RHICmdList, ComputeShaderRHI);
-					SetShaderParameters(RHICmdList, ComputeShader, ComputeShaderRHI, *PassParameters);
-					RHICmdList.DispatchComputeShader(GroupCount.X, GroupCount.Y, GroupCount.Z);
-					UnsetShaderUAVs(RHICmdList, ComputeShader, ComputeShaderRHI);
+					UE::MeshPassUtils::Dispatch(RHICmdList, ComputeShader, ShaderBindings, *PassParameters, GroupCount);
 				});
 		}
 		
@@ -2458,12 +2454,7 @@ void FSceneRenderer::RenderVolumetricCloudsInternal(FRDGBuilder& GraphBuilder, F
 
 				ShaderBindings.Finalize(&PassShaders);
 
-				FRHIComputeShader* ComputeShaderRHI = ComputeShader.GetComputeShader();
-				SetComputePipelineState(RHICmdList, ComputeShaderRHI);
-				ShaderBindings.SetOnCommandList(RHICmdList, ComputeShaderRHI);
-				SetShaderParameters(RHICmdList, ComputeShader, ComputeShaderRHI, *PassParameters);
-				RHICmdList.DispatchComputeShader(GroupCount.X, GroupCount.Y, GroupCount.Z);
-				UnsetShaderUAVs(RHICmdList, ComputeShader, ComputeShaderRHI);
+				UE::MeshPassUtils::Dispatch(RHICmdList, ComputeShader, ShaderBindings, *PassParameters, GroupCount);
 			});
 	}
 	else
diff --git a/Engine/Source/Runtime/Renderer/Public/MeshPassUtils.h b/Engine/Source/Runtime/Renderer/Public/MeshPassUtils.h
new file mode 100644
--- /dev/null
+++ b/Engine/Source/Runtime/Renderer/Public/MeshPassUtils.h
@@ -0,0 +1,68 @@
+// Copyright Epic Games, Inc. All Rights Reserved.
+
+#pragma once
+
+#include "MeshPassProcessor.h"
+#include "RenderGraphUtils.h"
+
+namespace UE::MeshPassUtils
+{
+	namespace Private
+	{
+		template<typename TShaderClass>
+		inline void PrepareDispatch(
+			FRHIComputeCommandList& RHICmdList,
+			const TShaderRef<TShaderClass>& ComputeShader,
+			const FMeshDrawShaderBindings& ShaderBindings,
+			const typename TShaderClass::FParameters& PassParameters)
+		{
+			FRHIComputeShader* ShaderRHI = ComputeShader.GetComputeShader();
+
+			FRHIBatchedShaderParameters& BatchedParameters = RHICmdList.GetScratchShaderParameters();
+			ShaderBindings.SetParameters(BatchedParameters, ShaderRHI);
+			SetShaderParameters(BatchedParameters, ComputeShader, PassParameters);
+
+			SetComputePipelineState(RHICmdList, ShaderRHI);
+			RHICmdList.SetBatchedShaderParameters(ShaderRHI, BatchedParameters);
+		}
+
+		template<typename TShaderClass>
+		inline void AfterDispatch(FRHIComputeCommandList& RHICmdList, const TShaderRef<TShaderClass>& ComputeShader)
+		{
+			UnsetShaderUAVs(RHICmdList, ComputeShader, ComputeShader.GetComputeShader());
+		}
+	}
+
+	/** Dispatch a compute shader to RHI Command List with its parameters and mesh shader bindings. */
+	template<typename TShaderClass>
+	inline void Dispatch(
+		FRHIComputeCommandList& RHICmdList,
+		const TShaderRef<TShaderClass>& ComputeShader,
+		const FMeshDrawShaderBindings& ShaderBindings,
+		const typename TShaderClass::FParameters& PassParameters,
+		FIntVector GroupCount)
+	{
+		FComputeShaderUtils::ValidateGroupCount(GroupCount);
+
+		Private::PrepareDispatch(RHICmdList, ComputeShader, ShaderBindings, PassParameters);
+		RHICmdList.DispatchComputeShader(GroupCount.X, GroupCount.Y, GroupCount.Z);
+		Private::AfterDispatch(RHICmdList, ComputeShader);
+	}
+
+	/** Indirect dispatch a compute shader to RHI Command List with its parameters and mesh shader bindings. */
+	template<typename TShaderClass>
+	inline void DispatchIndirect(
+		FRHIComputeCommandList& RHICmdList,
+		const TShaderRef<TShaderClass>& ComputeShader,
+		const FMeshDrawShaderBindings& ShaderBindings,
+		const typename TShaderClass::FParameters& PassParameters,
+		FRHIBuffer* IndirectArgsBuffer,
+		uint32 IndirectArgOffset)
+	{
+		FComputeShaderUtils::ValidateIndirectArgsBuffer(IndirectArgsBuffer->GetSize(), IndirectArgOffset);
+
+		Private::PrepareDispatch(RHICmdList, ComputeShader, ShaderBindings, PassParameters);
+		RHICmdList.DispatchIndirectComputeShader(IndirectArgsBuffer, IndirectArgOffset);
+		Private::AfterDispatch(RHICmdList, ComputeShader);
+	}
+}
