Implements experimental TSR history resurrection on parallax disocclusion and shading changes

This allows to resurrect previously discarded fully accumulated details immediatly on parallax disocclusion, whether with velocity vector like on opaque objects or VFX moving arround without velocity vector, as well as shading changes.

This changes most notably deletes the separate translucency accumulation that is slower, only for translucencies, and generate blending artifacts on opaque geometry discontinuities.

For further details, refer to r.TSR.Resurrection's help.

diff --git a/Engine/Shaders/Private/TemporalSuperResolution/TSRCommon.ush b/Engine/Shaders/Private/TemporalSuperResolution/TSRCommon.ush
--- a/Engine/Shaders/Private/TemporalSuperResolution/TSRCommon.ush
+++ b/Engine/Shaders/Private/TemporalSuperResolution/TSRCommon.ush
@@ -286,6 +286,7 @@ float2 PrevHistoryInfo_UVViewportBilinearMin;
 float2 PrevHistoryInfo_UVViewportBilinearMax;
 FScreenTransform ScreenPosToPrevHistoryBufferUV;
 float HistoryPreExposureCorrection;
+float ResurrectionPreExposureCorrection;
 
 uint bCameraCut;
 
diff --git a/Engine/Shaders/Private/TemporalSuperResolution/TSRDecimateHistory.usf b/Engine/Shaders/Private/TemporalSuperResolution/TSRDecimateHistory.usf
--- a/Engine/Shaders/Private/TemporalSuperResolution/TSRDecimateHistory.usf
+++ b/Engine/Shaders/Private/TemporalSuperResolution/TSRDecimateHistory.usf
@@ -29,19 +29,21 @@
 /** Whether to compile the moire reprojection. */
 #define CONFIG_MOIRE_REPROJECTION (DIM_MOIRE_REPROJECTION)
 
+/** Whether to compile the resurrection reprojection. */
+#define CONFIG_RESURRECTION_REPROJECTION (DIM_RESURRECTION_REPROJECTION)
+
 
 //------------------------------------------------------- PARAMETERS
 
 float4x4 RotationalClipToPrevClip;
-float2 CurrentToPrevInputJitter;
 
 Texture2D<float2> DilatedVelocityTexture;
 Texture2D<float> ClosestDepthTexture;
 Texture2D<uint> PrevUseCountTexture;
 Texture2D<uint> PrevClosestDepthTexture;
 
-Texture2D<tsr_halfC> PrevHistoryGuide;
-Texture2D<tsr_half4> PrevHistoryMoire;
+Texture2DArray<tsr_halfC> PrevHistoryGuide;
+Texture2DArray<tsr_half4> PrevHistoryMoire;
 Texture2D<float2> PrevVelocityTexture;
 float2 PrevGuideInfo_Extent;
 float2 PrevGuideInfo_ExtentInverse;
@@ -53,16 +55,28 @@ float2 PrevGuideInfo_UVViewportBilinearMin;
 float2 PrevGuideInfo_UVViewportBilinearMax;
 FScreenTransform InputPixelPosToReprojectScreenPos;
 FScreenTransform ScreenPosToPrevHistoryGuideBufferUV;
+FScreenTransform ScreenPosToResurrectionGuideBufferUV;
+float2 ResurrectionGuideUVViewportBilinearMin;
+float2 ResurrectionGuideUVViewportBilinearMax;
 float3 HistoryGuideQuantizationError;
+float ResurrectionFrameIndex;
+float PrevFrameIndex;
+float4x4 ClipToResurrectionClip;
 
-
-RWTexture2D<tsr_half4> ReprojectedHistoryGuideOutput;
-RWTexture2D<tsr_half4> ReprojectedHistoryMoireOutput;
+RWTexture2DArray<tsr_half4> ReprojectedHistoryGuideOutput;
+RWTexture2DArray<tsr_half4> ReprojectedHistoryMoireOutput;
 
 RWTexture2D<float2> HoleFilledVelocityOutput;
-RWTexture2D<float2> GrandHoleFilledVelocityOutput;
 RWTexture2D<tsr_half2> HoleFilledVelocityMaskOutput;
 RWTexture2D<tsr_half> ParallaxRejectionMaskOutput;
+RWTexture2D<tsr_half> HistoryResurrectionMaskOutput;
+
+
+//------------------------------------------------------- DEBUG
+
+#define DEBUG_ARRAY_SIZE 2
+
+static float4 Debug[DEBUG_ARRAY_SIZE];
 
 
 //------------------------------------------------------- LDS
@@ -315,7 +329,7 @@ void ProcessPrevAtomicTexturesSamples(
 
 // Sample a texture with a catrom kernel
 template<typename Type>
-Type SampleTextureCatrom(Texture2D<Type> InTexture, float2 TextureUV)
+Type SampleTextureCatrom(Texture2DArray<Type> InTexture, float2 TextureUV, float FrameIndex, float2 MinUV, float2 MaxUV)
 {
 	FCatmullRomSamples Samples = GetBicubic2DCatmullRomSamples(TextureUV, PrevGuideInfo_Extent, PrevGuideInfo_ExtentInverse);
 
@@ -323,12 +337,12 @@ Type SampleTextureCatrom(Texture2D<Type> InTexture, float2 TextureUV)
 	for (uint i = 0; i < Samples.Count; i++)
 	{
 		float2 SampleUV = Samples.UV[i];
-		SampleUV = clamp(
+		SampleUV = fastClamp(
 			SampleUV,
-			PrevGuideInfo_UVViewportBilinearMin,
-			PrevGuideInfo_UVViewportBilinearMax);
+			MinUV,
+			MaxUV);
 
-		Color += InTexture.SampleLevel(GlobalBilinearClampedSampler, SampleUV, 0) * tsr_half(Samples.Weight[i]);
+		Color += InTexture.SampleLevel(GlobalBilinearClampedSampler, float3(SampleUV, FrameIndex), 0) * tsr_half(Samples.Weight[i]);
 	}
 	Color *= tsr_half(Samples.FinalMultiplier);
 
@@ -336,7 +350,7 @@ Type SampleTextureCatrom(Texture2D<Type> InTexture, float2 TextureUV)
 }
 
 // Correct pre exposure of the guide color.
-tsr_halfC CorrectGuideColorExposure(tsr_halfC ReprojectedHistoryGuideColor)
+tsr_halfC CorrectGuideColorExposure(tsr_halfC ReprojectedHistoryGuideColor, const float PreExposureCorrection)
 {
 	tsr_halfC LinearReprojectedHistoryGuideColor = ReprojectedHistoryGuideColor * rcp(tsr_half(1.0) - ReprojectedHistoryGuideColor) * tsr_half(HistoryPreExposureCorrection);
 				
@@ -350,12 +364,12 @@ tsr_halfC CorrectGuideColorExposure(tsr_halfC ReprojectedHistoryGuideColor)
 
 // Sample previous moire history
 #if CONFIG_MOIRE_REPROJECTION
-tsr_half4 ReprojectPrevHistoryMoire(tsr_ushort2 InputPixelPos, float2 TextureUV)
+tsr_half4 ReprojectPrevHistoryMoire(tsr_ushort2 InputPixelPos, float2 TextureUV, float FrameIndex)
 {
 	uint2 Random = Rand3DPCG16(int3(InputPixelPos - InputInfo_ViewportMin, View.StateFrameIndexMod8)).xy;
 	float2 E = Hammersley16(0, 1, Random).xy;
 
-	return PrevHistoryMoire.SampleLevel(GlobalPointClampedSampler, TextureUV + (E - 0.5) * PrevGuideInfo_ExtentInverse, 0);
+	return PrevHistoryMoire.SampleLevel(GlobalPointClampedSampler, float3(TextureUV + (E - 0.5) * PrevGuideInfo_ExtentInverse, FrameIndex), 0);
 }
 #endif
 
@@ -367,7 +381,10 @@ void MainCS(
 	uint2 GroupId : SV_GroupID,
 	uint GroupThreadIndex : SV_GroupIndex)
 {
-	float4 Debug = 0;
+	for (uint DebugId = 0; DebugId < DEBUG_ARRAY_SIZE; DebugId++)
+	{
+		Debug[DebugId] = 0.0;
+	}
 
 	// Sample current frame informations.
 	float2 EncodedScreenVelocity;
@@ -392,6 +409,7 @@ void MainCS(
 				EncodedScreenVelocity = DilatedVelocityTexture[InputPixelPos];
 			#else // CONFIG_VELOCITY_SAMPLES != 1
 				LoadVelocityNeighborhood(InputPixelPos, /* out */ EncodedVelocityNeighborhood);
+				EncodedScreenVelocity = EncodedVelocityNeighborhood[0];
 			#endif // CONFIG_VELOCITY_SAMPLES != 1
 		}
 
@@ -428,6 +446,12 @@ void MainCS(
 	#if CONFIG_MOIRE_REPROJECTION
 		tsr_half4 ReprojectedHistoryMoireColor;
 	#endif
+	#if CONFIG_RESURRECTION_REPROJECTION
+		tsr_halfC ResurrectedHistoryGuideColor;
+		#if CONFIG_MOIRE_REPROJECTION
+			tsr_half4 ResurrectedHistoryMoireColor;
+		#endif
+	#endif
 	bool bIsOffScreen;
 
 	// Issues all the history texture fetches
@@ -455,14 +479,44 @@ void MainCS(
 				
 			float2 KernelHistoryBufferUV = ApplyScreenTransform(PrevOutputScreenPos, ScreenPosToPrevHistoryGuideBufferUV);
 
-			ReprojectedHistoryGuideColor = SampleTextureCatrom(PrevHistoryGuide, KernelHistoryBufferUV);
-			ReprojectedHistoryGuideColor = CorrectGuideColorExposure(ReprojectedHistoryGuideColor);
+			ReprojectedHistoryGuideColor = SampleTextureCatrom(
+				PrevHistoryGuide,
+				KernelHistoryBufferUV,
+				PrevFrameIndex,
+				PrevGuideInfo_UVViewportBilinearMin,
+				PrevGuideInfo_UVViewportBilinearMax);
+			ReprojectedHistoryGuideColor = CorrectGuideColorExposure(ReprojectedHistoryGuideColor, HistoryPreExposureCorrection);
+			#if CONFIG_MOIRE_REPROJECTION
+			{
+				ReprojectedHistoryMoireColor = ReprojectPrevHistoryMoire(InputPixelPos, KernelHistoryBufferUV, PrevFrameIndex);
+			}
+			#endif
+		}
+
+		#if CONFIG_RESURRECTION_REPROJECTION
+		{
+			float2 OutputScreenPos = ApplyScreenTransform(float2(InputPixelPos), InputPixelPosToReprojectScreenPos);
+
+			float4 ThisClip = float4(OutputScreenPos, DeviceZ, 1);
+			float4 ResurrectionClip = mul(ThisClip, ClipToResurrectionClip); 
+			float2 ResurrectionScreen = ResurrectionClip.xy / ResurrectionClip.w;
+
+			float2 KernelHistoryBufferUV = ApplyScreenTransform(ResurrectionScreen, ScreenPosToResurrectionGuideBufferUV);
+
+			ResurrectedHistoryGuideColor = SampleTextureCatrom(
+				PrevHistoryGuide,
+				KernelHistoryBufferUV,
+				ResurrectionFrameIndex,
+				ResurrectionGuideUVViewportBilinearMin,
+				ResurrectionGuideUVViewportBilinearMax);
+			ResurrectedHistoryGuideColor = CorrectGuideColorExposure(ResurrectedHistoryGuideColor, ResurrectionPreExposureCorrection);
 			#if CONFIG_MOIRE_REPROJECTION
 			{
-				ReprojectedHistoryMoireColor = ReprojectPrevHistoryMoire(InputPixelPos, KernelHistoryBufferUV);
+				ResurrectedHistoryMoireColor = ReprojectPrevHistoryMoire(InputPixelPos, KernelHistoryBufferUV, ResurrectionFrameIndex);
 			}
 			#endif
 		}
+		#endif
 	}
 	
 	// Process the history samples
@@ -504,6 +558,7 @@ void MainCS(
 		bool bIsOffscreenOrParallaxRejected = ParallaxRejectionMask < tsr_half(PARALLAX_REJECTION_MASK_THRESHOLD);
 	
 		// Hole fill the velocity.
+		bool bAllowResurrection = false;
 		{
 			float2 HoleFillingPixelVelocity;
 			float HoleFillingOutOfRange;
@@ -521,6 +576,13 @@ void MainCS(
 				float VelocityDelta = length(ScreenPixelVelocity - HoleFillingPixelVelocity);
 
 				float2 HoleFillingVelocity = HoleFillingPixelVelocity * InputPixelVelocityToScreenVelocity;
+				
+				#if CONFIG_RESURRECTION_REPROJECTION
+					bAllowResurrection = (
+						ResurrectionFrameIndex != PrevFrameIndex &&
+						length2(ScreenPixelVelocity - HoleFillingPixelVelocity) > pow(3.0, 2.0) &&
+						length2(ScreenPixelVelocity) < pow(100.0, 2.0));
+				#endif
 
 				EncodedScreenVelocity = EncodeVelocityToTexture(float3(HoleFillingVelocity, 0.0)).xy;
 				//HoleFilledVelocityMask = tsr_half(HoleFillingOutOfRange) * tsr_half(saturate(MaxHollFillPixelVelocity - VelocityDelta));
@@ -546,22 +608,38 @@ void MainCS(
 		
 		tsr_short2 OutputPixelPos = InvalidateOutputPixelPos(InputPixelPos, InputInfo_ViewportMax);
 
-		#if CONFIG_SCENE_COLOR_ALPHA
-			ReprojectedHistoryGuideOutput[OutputPixelPos] = ReprojectedHistoryGuideColor;
-		#else
-			ReprojectedHistoryGuideOutput[OutputPixelPos] = tsr_half4(ReprojectedHistoryGuideColor.xyz, tsr_half(1.0));
-		#endif
+		{
+			#if CONFIG_SCENE_COLOR_ALPHA
+				ReprojectedHistoryGuideOutput[tsr_short3(OutputPixelPos, 0)] = ReprojectedHistoryGuideColor;
+			#else
+				ReprojectedHistoryGuideOutput[tsr_short3(OutputPixelPos, 0)] = tsr_half4(ReprojectedHistoryGuideColor.xyz, tsr_half(1.0));
+			#endif
+			#if CONFIG_MOIRE_REPROJECTION
+				ReprojectedHistoryMoireOutput[tsr_short3(OutputPixelPos, 0)] = ReprojectedHistoryMoireColor;
+			#endif
+		}
 
-		#if CONFIG_MOIRE_REPROJECTION
-			ReprojectedHistoryMoireOutput[OutputPixelPos] = ReprojectedHistoryMoireColor;
+		#if CONFIG_RESURRECTION_REPROJECTION
+		{
+			#if CONFIG_SCENE_COLOR_ALPHA
+				ReprojectedHistoryGuideOutput[tsr_short3(OutputPixelPos, 1)] = ResurrectedHistoryGuideColor;
+			#else
+				ReprojectedHistoryGuideOutput[tsr_short3(OutputPixelPos, 1)] = tsr_half4(ResurrectedHistoryGuideColor.xyz, tsr_half(1.0));
+			#endif
+			#if CONFIG_MOIRE_REPROJECTION
+				ReprojectedHistoryMoireOutput[tsr_short3(OutputPixelPos, 1)] = ResurrectedHistoryMoireColor;
+			#endif
+		}
 		#endif
+
 		HoleFilledVelocityOutput[OutputPixelPos] = EncodedScreenVelocity;
 		HoleFilledVelocityMaskOutput[InputPixelPos] = tsr_half2(0.0, RejectHighFrequency);
 		ParallaxRejectionMaskOutput[OutputPixelPos] = ParallaxRejectionMask;
 
 		#if DEBUG_OUTPUT
+		for (uint DebugId = 0; DebugId < DEBUG_ARRAY_SIZE; DebugId++)
 		{
-			DebugOutput[tsr_short3(OutputPixelPos, 0)] = Debug;
+			DebugOutput[tsr_short3(OutputPixelPos, DebugId)] = Debug[DebugId];
 		}
 		#endif
 	}
diff --git a/Engine/Shaders/Private/TemporalSuperResolution/TSRMergeRejection.usf b/Engine/Shaders/Private/TemporalSuperResolution/TSRMergeRejection.usf
new file mode 100644
--- /dev/null
+++ b/Engine/Shaders/Private/TemporalSuperResolution/TSRMergeRejection.usf
@@ -0,0 +1,93 @@
+// Copyright Epic Games, Inc. All Rights Reserved.
+
+#include "TSRCommon.ush"
+
+
+//------------------------------------------------------- CONFIG
+
+#define TILE_SIZE 8
+
+
+
+//------------------------------------------------------- PARAMETERS
+
+float4x4 ClipToResurrectionClip;
+
+Texture2D<float> ClosestDepthTexture;
+
+Texture2D<tsr_half4> PrevGuideTexture;
+Texture2D<tsr_half4> PrevMoireTexture;
+Texture2D<tsr_half2> PrevRejectionTexture;
+
+Texture2D<tsr_half4> ResurrectedGuideTexture;
+Texture2D<tsr_half4> ResurrectedMoireTexture;
+Texture2D<tsr_half2> ResurrectedRejectionTexture;
+
+RWTexture2D<tsr_half4> HistoryGuideOutput;
+RWTexture2D<tsr_half4> HistoryMoireOutput;
+RWTexture2D<tsr_half2> HistoryRejectionOutput;
+RWTexture2D<tsr_half> HistoryResurrectionMaskOutput;
+RWTexture2D<float2> DilatedVelocityOutput;
+
+
+//------------------------------------------------------- ENTRY POINT
+
+[numthreads(TILE_SIZE * TILE_SIZE, 1, 1)]
+void MainCS(
+	uint2 GroupId : SV_GroupID,
+	uint GroupThreadIndex : SV_GroupIndex)
+{
+	uint GroupWaveIndex = GetGroupWaveIndex(GroupThreadIndex, /* GroupSize = */ TILE_SIZE * TILE_SIZE);
+
+	float4 Debug = 0.0;
+	
+	tsr_short2 InputPixelPos = (
+		tsr_short2(InputPixelPosMin) +
+		tsr_short2(GroupId) * tsr_short2(TILE_SIZE, TILE_SIZE) +
+		Map8x8Tile2x2Lane(GroupThreadIndex));
+	
+	tsr_half2 PrevRejection = PrevRejectionTexture[InputPixelPos];
+	tsr_half2 ResurrectedRejection = ResurrectedRejectionTexture[InputPixelPos];
+
+	bool bIsResurrectionBetter = all(ResurrectedRejection.g > PrevRejection.g + 0.5);
+
+	//Debug.r = float(bIsResurrectionBetter);
+	//Debug.rg = saturate(ResurrectedRejection.rg - PrevRejection.rg);
+	Debug.b = bIsResurrectionBetter;
+
+	BRANCH
+	if (bIsResurrectionBetter)
+	{
+		float2 ResurrectedEncodedScreenVelocity;
+		{
+			float DeviceZ = ClosestDepthTexture[InputPixelPos];
+
+			float2 ScreenPos = ApplyScreenTransform(float2(InputPixelPos), InputPixelPosToScreenPos);
+
+			float4 ThisClip = float4(ScreenPos, DeviceZ, 1);
+			float4 ResurrectionClip = mul(ThisClip, ClipToResurrectionClip); 
+			float2 ResurrectionScreen = ResurrectionClip.xy / ResurrectionClip.w;
+			
+			ResurrectedEncodedScreenVelocity = EncodeVelocityToTexture(float3(ScreenPos - ResurrectionScreen, 0.0)).xy;
+		}
+	
+		HistoryGuideOutput[InputPixelPos]     = ResurrectedGuideTexture[InputPixelPos];
+		HistoryMoireOutput[InputPixelPos]     = ResurrectedMoireTexture[InputPixelPos];
+		HistoryRejectionOutput[InputPixelPos] = ResurrectedRejection;
+		DilatedVelocityOutput[InputPixelPos]  = ResurrectedEncodedScreenVelocity;
+	}
+	else
+	{
+		HistoryGuideOutput[InputPixelPos]     = PrevGuideTexture[InputPixelPos];
+		HistoryMoireOutput[InputPixelPos]     = PrevMoireTexture[InputPixelPos];
+		HistoryRejectionOutput[InputPixelPos] = PrevRejection;
+	}
+	
+	HistoryResurrectionMaskOutput[InputPixelPos] = bIsResurrectionBetter ? tsr_half(1.0) : tsr_half(0.0);
+
+	#if DEBUG_OUTPUT
+	{
+		DebugOutput[tsr_short3(InputPixelPos, 0)] = Debug;
+	}
+	#endif
+}
diff --git a/Engine/Shaders/Private/TemporalSuperResolution/TSRRejectShading.usf b/Engine/Shaders/Private/TemporalSuperResolution/TSRRejectShading.usf
--- a/Engine/Shaders/Private/TemporalSuperResolution/TSRRejectShading.usf
+++ b/Engine/Shaders/Private/TemporalSuperResolution/TSRRejectShading.usf
@@ -514,11 +514,10 @@ void MainCS(
 		
 		tsr_simd_half RawRejection      = min3(RawFactor[0], RawFactor[1], RawFactor[2]);
 		tsr_simd_half FilteredRejection = min3(FilteredFactor[0], FilteredFactor[1], FilteredFactor[2]);
-		
-		RejectionBlendFinal = saturate(FilteredRejection + tsr_simd_half::Const(1.0) - bIsValid);
+
+		RejectionBlendFinal = saturate(FilteredRejection);
 		
 		RejectionClampBlend = Median3x3(RawRejection);
-		RejectionClampBlend = saturate(RejectionClampBlend + tsr_simd_half::Const(1.0) - bIsValid);
 		RejectionClampBlend = Min3x3(RejectionClampBlend);
 	}
 	#endif
diff --git a/Engine/Shaders/Private/TemporalSuperResolution/TSRUpdateHistory.usf b/Engine/Shaders/Private/TemporalSuperResolution/TSRUpdateHistory.usf
--- a/Engine/Shaders/Private/TemporalSuperResolution/TSRUpdateHistory.usf
+++ b/Engine/Shaders/Private/TemporalSuperResolution/TSRUpdateHistory.usf
@@ -31,19 +31,9 @@
 	#error Unknown history update quality
 #endif
 
-#if 1
-	#define CONFIG_HIGH_FREQUENCY_ONLY 1
-#endif
-
-#if DIM_SEPARATE_TRANSLUCENCY
-	#define CONFIG_ACCUMULATE_TRANSLUCENCY_SEPARATELY 1
-	#define CONFIG_COMPOSE_TRANSLUCENCY 0
-	#define CONFIG_REJECT_TRANSLUCENCY_ON_RESPONSIVE_AA_ONLY 0
-#else
-	#define CONFIG_ACCUMULATE_TRANSLUCENCY_SEPARATELY 0
-	#define CONFIG_COMPOSE_TRANSLUCENCY 1
-	#define CONFIG_REJECT_TRANSLUCENCY_ON_RESPONSIVE_AA_ONLY 0
-#endif
+#define CONFIG_COMPOSE_TRANSLUCENCY 1
+#define CONFIG_REJECT_TRANSLUCENCY_ON_RESPONSIVE_AA_ONLY 0
+#define CONFIG_HISTORY_RESURRECTION 1
 
 
 //------------------------------------------------------- DEFAULTS
@@ -73,10 +63,10 @@
 	#define CONFIG_INTERPOLATE_VELOCITY 0
 #endif
 
-#if CONFIG_ACCUMULATE_TRANSLUCENCY_SEPARATELY
-	#define CONFIG_OUTPUT_HIGH_FREQUENCY 1
+#if 0
+	#define CONFIG_OUTPUT_MERGED_SCENE_COLOR 1
 #else
-	#define CONFIG_OUTPUT_HIGH_FREQUENCY 0
+	#define CONFIG_OUTPUT_MERGED_SCENE_COLOR 0
 #endif
 
 
@@ -107,6 +97,9 @@ Texture2D<tsr_half2> HistoryRejectionTexture;
 
 Texture2D<float2> DilatedVelocityTexture;
 Texture2D<tsr_half> ParallaxRejectionMaskTexture;
+#if CONFIG_HISTORY_RESURRECTION
+	Texture2D<tsr_half> HistoryResurrectionMaskTexture;
+#endif
 Texture2D<tsr_ushort> AntiAliasingTexture;
 Texture2D<tsr_half> NoiseFilteringTexture;
 Texture2D<tsr_half2> HoleFilledVelocityMaskTexture;
@@ -124,25 +117,22 @@ float WeightClampingRejection;
 float WeightClampingPixelSpeedAmplitude;
 float InvWeightClampingPixelSpeed;
 float InputToHistoryFactor;
+float ResurrectionFrameIndex;
+float PrevFrameIndex;
 uint ResponsiveStencilMask;
 uint bGenerateOutputMip1;
 uint bGenerateOutputMip2;
 uint bHasSeparateTranslucency;
+uint HistoryArrayIndices_HighFrequency;
 
-Texture2D<tsr_halfC> PrevHistory_HighFrequency;
-Texture2D<tsr_halfM> PrevHistory_Metadata;
-Texture2D<tsr_half3> PrevHistory_Translucency;
-Texture2D<tsr_half>  PrevHistory_TranslucencyAlpha;
+Texture2DArray<tsr_halfC> PrevHistoryColorTexture;
+Texture2DArray<tsr_halfM> PrevHistoryMetadataTexture;
 
 RWTexture2D<tsr_halfC> SceneColorOutputMip0;
 RWTexture2D<tsr_halfC> SceneColorOutputMip1;
 
-uint HistoryOutput_ArrayIndices_HighFrequency;
-uint HistoryOutput_ArrayIndices_Translucency;
-
-RWTexture2DArray<tsr_halfC> HistoryOutput_ColorArray;
-RWTexture2D<tsr_halfM> HistoryOutput_Metadata;
-RWTexture2D<tsr_half>  HistoryOutput_TranslucencyAlpha;
+RWTexture2DArray<tsr_halfC> HistoryColorOutput;
+RWTexture2D<tsr_halfM> HistoryMetadataOutput;
 
 
 //------------------------------------------------------- LDS
@@ -219,12 +209,6 @@ tsr_half4x2 WaveBroadcast(const FWaveBroadcastSettings Settings, tsr_half4x2 v)
 
 //------------------------------------------------------- FUNCTIONS
 
-template<typename T>
-void CorrectExposure(inout T X)
-{
-	X *= tsr_half(HistoryPreExposureCorrection);
-}
-
 CALL_SITE_DEBUGLOC 
 tsr_ushort2x2 Map8x8Tile2x2LaneDPV(uint GroupThreadIndex)
 {
@@ -483,7 +467,12 @@ void MainCS(
 	bool2 bIsResponsiveAAPixel = false;
 	bool2 bIsOffScreen = false;
 	bool2 bIsDisoccluded = false;
-	
+	#if CONFIG_HISTORY_RESURRECTION
+		bool2 bIsHistoryResurrection = false;
+	#else
+		const bool2 bIsHistoryResurrection = false;
+	#endif
+
 	#if CONFIG_REJECTION_ANTI_ALIASING
 		tsr_half2x2 SpatialAntiAliasingOffset = dpv_interleave_mono_registers(tsr_half(0.0).xx);
 	#endif
@@ -497,6 +486,9 @@ void MainCS(
 		#endif
 
 		tsr_half RawParallaxRejectionMask[DPV_PIXEL_PER_LANE];
+		#if CONFIG_HISTORY_RESURRECTION
+			tsr_half RawHistoryResurrectionMask[DPV_PIXEL_PER_LANE];
+		#endif
 		tsr_half2 RawLowFrequencyRejection[DPV_PIXEL_PER_LANE];
 		
 		uint RawSceneStencilRef[DPV_PIXEL_PER_LANE];
@@ -521,6 +513,9 @@ void MainCS(
 				#endif
 
 				RawParallaxRejectionMask[PixelId] = ParallaxRejectionMaskTexture[LocalInputPixelPos];
+				#if CONFIG_HISTORY_RESURRECTION
+					RawHistoryResurrectionMask[PixelId] = HistoryResurrectionMaskTexture[LocalInputPixelPos];
+				#endif
 				RawSceneStencilRef[PixelId]       = InputSceneStencilTexture.Load(int3(LocalInputPixelPos, 0)) STENCIL_COMPONENT_SWIZZLE;
 				
 				RawLowFrequencyRejection[PixelId] = HistoryRejectionTexture[LocalInputPixelPos];
@@ -630,6 +625,10 @@ void MainCS(
 				dpv_interleave_registers_array(RawParallaxRejectionMask),
 				/* out */ bIsOffScreen,
 				/* out */ bIsDisoccluded);
+				
+			#if CONFIG_HISTORY_RESURRECTION
+				bIsHistoryResurrection = dpv_interleave_registers_array(RawHistoryResurrectionMask) != tsr_half(0.0);
+			#endif
 
 			// Process input texel, forcing to pack register.
 			LowFrequencyRejection = dpv_force_interleave_registers_array(RawLowFrequencyRejection)[0];
@@ -667,6 +666,9 @@ void MainCS(
 			#if !CONFIG_CLAMP
 			{
 				bIsDisoccluded = false;
+				#if CONFIG_HISTORY_RESURRECTION
+					bIsHistoryResurrection = false;
+				#endif
 				LowFrequencyRejection = tsr_half(1.0).xx;
 				LowFrequencyClamp = tsr_half(1.0).xx;
 			}
@@ -681,14 +683,11 @@ void MainCS(
 	tsr_halfCx2 PrevHistoryMomentMin;
 	tsr_halfCx2 PrevHistoryMomentMax;
 
-	tsr_half4x2 PrevHistoryTranslucency;
-	
 	PLATFORM_SPECIFIC_ISOLATE
 	{
 		// Reproject history
 		tsr_halfCx2 AccumulateHighFrequency = tsr_half(0);
 		tsr_halfMx2 AccumulateMetadata = tsr_half(0);
-		tsr_half4x2 AccumulateTranslucency = tsr_half(0);
 
 		tsr_halfCx2 HighFrequencyMin = INFINITE_FLOAT;
 		tsr_halfCx2 HighFrequencyMax = -INFINITE_FLOAT;
@@ -706,8 +705,6 @@ void MainCS(
 
 			tsr_halfC RawHighFrequency[BICUBIC_CATMULL_ROM_SAMPLES][DPV_PIXEL_PER_LANE];
 			tsr_halfM RawMetadata[BICUBIC_CATMULL_ROM_SAMPLES][DPV_PIXEL_PER_LANE];
-			tsr_half3 RawTranslucencyColor[BICUBIC_CATMULL_ROM_SAMPLES][DPV_PIXEL_PER_LANE];
-			tsr_half  RawTranslucencyAlpha[BICUBIC_CATMULL_ROM_SAMPLES][DPV_PIXEL_PER_LANE];
 
 			tsr_half2 RawKernelWeight[BICUBIC_CATMULL_ROM_SAMPLES];
 
@@ -729,15 +726,14 @@ void MainCS(
 					{
 						float2 SampleUV = PixelId == 0 ? SampleUV0 : SampleUV1;
 						
-
-						RawHighFrequency[i][PixelId]     = BilinearSampleColorHistory(PrevHistory_HighFrequency, SampleUV);
-						RawMetadata[i][PixelId]	         = PrevHistory_Metadata.SampleLevel(GlobalBilinearClampedSampler, SampleUV, 0);
-						
-						#if CONFIG_ACCUMULATE_TRANSLUCENCY_SEPARATELY
-							RawTranslucencyColor[i][PixelId] = BilinearSampleColorHistory(PrevHistory_Translucency, SampleUV);
-							RawTranslucencyAlpha[i][PixelId] = PrevHistory_TranslucencyAlpha.SampleLevel(GlobalBilinearClampedSampler, SampleUV, 0);
+						#if CONFIG_HISTORY_RESURRECTION
+							float FrameIndex = select(bIsHistoryResurrection[PixelId], ResurrectionFrameIndex, PrevFrameIndex);
+						#else
+							float FrameIndex = PrevFrameIndex;
 						#endif
-						
+
+						RawHighFrequency[i][PixelId] = PrevHistoryColorTexture.SampleLevel(GlobalBilinearClampedSampler, float3(SampleUV, FrameIndex), 0);
+						RawMetadata[i][PixelId]      = PrevHistoryMetadataTexture.SampleLevel(GlobalBilinearClampedSampler, float3(SampleUV, FrameIndex), 0);
 					}
 
 				} // for (uint i = 0; i < BICUBIC_CATMULL_ROM_SAMPLES; i++)
@@ -752,9 +748,6 @@ void MainCS(
 					// TODO: Should use dpv_force_interleave_registers_array() but there is a shader compiler bug
 					tsr_halfCx2 SampleHighFrequency = dpv_interleave_registers_array(RawHighFrequency[i]);
 					tsr_halfMx2 SampleMetadata = dpv_interleave_registers_array(RawMetadata[i]);
-					tsr_half4x2 SampleTranslucency = dpv_interleave_registers(
-						tsr_half4(RawTranslucencyColor[i][0], RawTranslucencyAlpha[i][0]),
-						tsr_half4(RawTranslucencyColor[i][1], RawTranslucencyAlpha[i][1]));
 
 					tsr_half2 KernelWeight = RawKernelWeight[i];
 
@@ -763,17 +756,12 @@ void MainCS(
 
 					AccumulateHighFrequency += dpv_scale(SampleHighFrequency, KernelWeight);
 					AccumulateMetadata  += dpv_scale(SampleMetadata, KernelWeight);
-					
-					#if CONFIG_ACCUMULATE_TRANSLUCENCY_SEPARATELY
-						AccumulateTranslucency  += dpv_scale(SampleTranslucency, KernelWeight);
-					#endif
 				}
 				
 				tsr_half2 FinalMultiplier = dpv_force_interleave_registers(tsr_half(Samples0.FinalMultiplier), tsr_half(Samples1.FinalMultiplier));
 
 				AccumulateHighFrequency = dpv_scale(AccumulateHighFrequency, FinalMultiplier);
 				AccumulateMetadata  = dpv_scale(AccumulateMetadata, FinalMultiplier);
-				AccumulateTranslucency  = dpv_scale(AccumulateTranslucency, FinalMultiplier);
 			}
 
 			// Super bight highlights have a tendy to generate negatives
@@ -792,7 +780,6 @@ void MainCS(
 				AccumulateHighFrequency = -dpv_min(-AccumulateHighFrequency, tsr_half(0.0).xxx);
 			#endif
 			AccumulateMetadata      = -dpv_min(-AccumulateMetadata, tsr_half(0.0).xx);
-			AccumulateTranslucency  = -dpv_min(-AccumulateTranslucency, tsr_half(0.0).xxxx);
 		}
 		
 		// Unpack history.
@@ -806,21 +793,28 @@ void MainCS(
 			#else
 				PrevHistoryValidity = AccumulateMetadata[0];
 			#endif
-
-			PrevHistoryTranslucency = AccumulateTranslucency;
 		}
 
 		// Correct history
 		{
-			CorrectExposure(PrevHistoryMomentMin[0]);
-			CorrectExposure(PrevHistoryMomentMin[1]);
-			CorrectExposure(PrevHistoryMomentMin[2]);
-			CorrectExposure(PrevHistoryMomentMax[0]);
-			CorrectExposure(PrevHistoryMomentMax[1]);
-			CorrectExposure(PrevHistoryMomentMax[2]);
-			CorrectExposure(PrevHighFrequencyColor[0]);
-			CorrectExposure(PrevHighFrequencyColor[1]);
-			CorrectExposure(PrevHighFrequencyColor[2]);
+			#if CONFIG_HISTORY_RESURRECTION
+				const tsr_half2 PreExposureCorrection = select(
+					bIsHistoryResurrection,
+					tsr_half(ResurrectionPreExposureCorrection),
+					tsr_half(HistoryPreExposureCorrection));
+			#else
+				const tsr_half2 PreExposureCorrection = tsr_half(HistoryPreExposureCorrection);
+			#endif
+
+			PrevHistoryMomentMin[0] *= PreExposureCorrection;
+			PrevHistoryMomentMin[1] *= PreExposureCorrection;
+			PrevHistoryMomentMin[2] *= PreExposureCorrection;
+			PrevHistoryMomentMax[0] *= PreExposureCorrection;
+			PrevHistoryMomentMax[1] *= PreExposureCorrection;
+			PrevHistoryMomentMax[2] *= PreExposureCorrection;
+			PrevHighFrequencyColor[0] *= PreExposureCorrection;
+			PrevHighFrequencyColor[1] *= PreExposureCorrection;
+			PrevHighFrequencyColor[2] *= PreExposureCorrection;
 
 			// Fixes negative validity that can happen due to negative lobes of Catrom.
 			PrevHistoryValidity = max(PrevHistoryValidity, tsr_half(0.0).xx);
@@ -833,12 +827,6 @@ void MainCS(
 	tsr_halfCx2 FilteredInputColor;
 	tsr_halfCx2 InputMinColor;
 	tsr_halfCx2 InputMaxColor;
-	
-#if CONFIG_ACCUMULATE_TRANSLUCENCY_SEPARATELY
-	tsr_half4x2 FilteredTranslucencyColor;
-	tsr_half4x2 TranslucencyMinColor;
-	tsr_half4x2 TranslucencyMaxColor;
-#endif
 
 	tsr_half2 InputPixelAlignement;
 	tsr_half2 ClosestInputLuma4;
@@ -859,7 +847,7 @@ void MainCS(
 		PLATFORM_SPECIFIC_ISOLATE
 		{
 			// Detect if HistoryBufferUV would be outside of the viewport.
-			tsr_half2 SpatialAntiAliasingLerp = select(or(bIsOffScreen, bIsDisoccluded), tsr_half(1.0).xx, saturate(tsr_half(1.0) - LowFrequencyRejection * tsr_half(4.0)));
+			tsr_half2 SpatialAntiAliasingLerp = select(or(bIsOffScreen, and(bIsDisoccluded, !bIsHistoryResurrection)), tsr_half(1.0).xx, saturate(tsr_half(1.0) - LowFrequencyRejection * tsr_half(4.0)));
 
 			// Pixel coordinate of the center of output pixel O in the input viewport.
 			float2x2 TranslucencyPPCo = ApplyScreenTransform(float2x2(HistoryPixelPos), HistoryPixelPosToTranslucencyPPCo);
@@ -905,13 +893,13 @@ void MainCS(
 				tsr_half2 IdealCurrentContribution = IdealInputPixelAlignement * tsr_half(ActualHistoryHisteresis);
 
 				// Select whether history should be clamped.
-				BlendClamp = select(bIsDisoccluded, tsr_half(0.0), min(HoleFilledVelocityMask[1], LowFrequencyClamp));
+				BlendClamp = select(and(bIsDisoccluded, !bIsHistoryResurrection), tsr_half(0.0), min(HoleFilledVelocityMask[1], LowFrequencyClamp));
 		
 				// Reduces the number of pixels when rejecting history
 				tsr_half2 ClampedPrevHistoryValidity = PrevHistoryValidity;
 				#if 1
 				{
-					tsr_half2 MaxValidity = lerp(tsr_half(WeightClampingRejection), tsr_half(1.0), select(bIsDisoccluded, tsr_half(0.0), LowFrequencyClamp));
+					tsr_half2 MaxValidity = lerp(tsr_half(WeightClampingRejection), tsr_half(1.0), select(and(bIsDisoccluded, !bIsHistoryResurrection), tsr_half(0.0), LowFrequencyClamp));
 					//tsr_half2 MaxValidity = lerp(tsr_half(WeightClampingRejection), tsr_half(1.0), BlendClamp);
 
 					ClampedPrevHistoryValidity = min(ClampedPrevHistoryValidity, MaxValidity);
@@ -944,7 +932,7 @@ void MainCS(
 				// Soften the spatial kernel when there is high noise in the input to make parallax disocclusion of detailed geometry less distracting.
 				tsr_half2 OffScreenInputToHistoryFactor = tsr_half(1.0) - tsr_half(0.5) * NoiseFiltering;
 
-				tsr_half2 KernelInputToHistoryLerp = select(or(bIsOffScreen, bIsDisoccluded), tsr_half(0.0).xx, saturate(LowFrequencyRejection * tsr_half(16.0) - tsr_half(13.0)) * bIsRefining);
+				tsr_half2 KernelInputToHistoryLerp = select(or(bIsOffScreen, and(bIsDisoccluded, !bIsHistoryResurrection)), tsr_half(0.0).xx, saturate(LowFrequencyRejection * tsr_half(16.0) - tsr_half(13.0)) * bIsRefining);
 				KernelInputToHistoryFactor = lerp(OffScreenInputToHistoryFactor, tsr_half(InputToHistoryFactor), KernelInputToHistoryLerp);
 				tsr_half2 KernelInputToHistoryAlignmentFactor = lerp(tsr_half(1.0).xx, tsr_half(InputToHistoryFactor), KernelInputToHistoryLerp);
 			
@@ -987,7 +975,7 @@ void MainCS(
 				for (uint PixelId = 0; PixelId < DPV_PIXEL_PER_LANE; PixelId++)
 				{
 					RawInputColorArray[SampleId][PixelId] = InputSceneColorTexture[dpv_access_pixel(InputSamplePixelPos, PixelId)];
-					#if CONFIG_ACCUMULATE_TRANSLUCENCY_SEPARATELY || CONFIG_COMPOSE_TRANSLUCENCY
+					#if CONFIG_COMPOSE_TRANSLUCENCY
 						RawTranslucencyColorArray[SampleId][PixelId] = InputSceneTranslucencyTexture[dpv_access_pixel(TranslucencySamplePixelPos, PixelId)];
 					#endif
 				}
@@ -1005,11 +993,6 @@ void MainCS(
 			InputMaxColor = dpv_interleave_mono_registers(-LargestSceneColorYCoCg);
 		#endif
 
-		#if CONFIG_ACCUMULATE_TRANSLUCENCY_SEPARATELY
-		TranslucencyMinColor = dpv_interleave_mono_registers(+LargestSceneColorRGBA);
-		TranslucencyMaxColor = dpv_interleave_mono_registers(-LargestSceneColorRGBA);
-		#endif
-
 		UNROLL_N(CONFIG_SAMPLES)
 		for (uint SampleId = 0; SampleId < CONFIG_SAMPLES; SampleId++)
 		{
@@ -1041,27 +1024,8 @@ void MainCS(
 			FilteredInputColor       += dpv_scale((SampleSpatialWeight * ToneWeight), InputColor);
 			FilteredInputColorWeight += (SampleSpatialWeight * ToneWeight);
 
-			#if CONFIG_ACCUMULATE_TRANSLUCENCY_SEPARATELY
-			if (SampleId == 0)
-			{
-				FilteredTranslucencyColor = TranslucencyColor;
-			}
-			#endif
-
 			InputMinColor = min(InputMinColor, InputColor);
 			InputMaxColor = max(InputMaxColor, InputColor);
-
-			#if CONFIG_ACCUMULATE_TRANSLUCENCY_SEPARATELY
-			{
-				TranslucencyMinColor = min(TranslucencyMinColor, TranslucencyColor);
-				TranslucencyMaxColor = max(TranslucencyMaxColor, TranslucencyColor);
-				
-				if (SampleId == 0)
-				{
-					FilteredTranslucencyColor = TranslucencyColor;
-				}
-			}
-			#endif
 		}
 		
 		FilteredInputColor = dpv_scale(FilteredInputColor, rcp(FilteredInputColorWeight));
@@ -1074,9 +1038,6 @@ void MainCS(
 				FilteredInputColor = min(FilteredInputColor, dpv_interleave_mono_registers(LargestSceneColorRGB));
 			#endif
 			ClosestInputLuma4 = min(ClosestInputLuma4, LargestSceneColorYCoCg.x);
-			#if CONFIG_ACCUMULATE_TRANSLUCENCY_SEPARATELY
-			FilteredTranslucencyColor = min(FilteredTranslucencyColor, dpv_interleave_mono_registers(LargestSceneColorRGBA));
-			#endif
 		}
 		#endif
 	}
@@ -1133,94 +1094,11 @@ void MainCS(
 		// Quantize validity for the 8bit encoding to avoid numerical shift between color and validity.
 		FinalHistoryValidity = ceil(tsr_half(255.0) * OutputValidity) * rcp(tsr_half(255.0));
 	}
-	
-	#if CONFIG_ACCUMULATE_TRANSLUCENCY_SEPARATELY
-	// Compute the low frequency history
-	tsr_half4x2 FinalTranslucencyColor;
-	{
-		tsr_half2 TranslucencyLumaMin = Luma4(GetColorChannel(TranslucencyMinColor));
-		tsr_half2 TranslucencyLumaMax = Luma4(GetColorChannel(TranslucencyMaxColor));
-
-		tsr_half2 TranslucencyLumaHistory  = Luma4(GetColorChannel(PrevHistoryTranslucency));
-		tsr_half2 TranslucencyLumaFiltered = Luma4(GetColorChannel(FilteredTranslucencyColor));
-
-		tsr_half4x2 ClampedPrevHistoryTranslucency = fastClamp(PrevHistoryTranslucency, TranslucencyMinColor, TranslucencyMaxColor);
-		
-		tsr_half2 BlendFinal = tsr_half(0.2) * InputPixelAlignement;
-
-		// Antiflicker
-		#if 1
-		{
-			tsr_half2 DistToClamp = min( abs(TranslucencyLumaHistory - TranslucencyLumaMin), abs(TranslucencyLumaHistory - TranslucencyLumaMax) ) / max3( TranslucencyLumaHistory, TranslucencyLumaFiltered, tsr_half(1e-4) );
-			BlendFinal *= tsr_half(0.2) + tsr_half(0.8) * saturate(tsr_half(0.5) * DistToClamp);
-		}
-		#endif
-			
-		// Make sure to have at least some small contribution
-		#if 0
-		{
-			BlendFinal = max( BlendFinal, saturate( tsr_half(0.01) * TranslucencyLumaHistory * rcp(abs(TranslucencyLumaFiltered - TranslucencyLumaHistory))));
-		}
-		#endif
-		
-		#if CONFIG_TAA_RESPONSIVE_BEHAVIOR
-			// Responsive forces 1/4 of new frame.
-			BlendFinal[0] = bIsResponsiveAAPixel[0] ? tsr_half(1.0/4.0) : BlendFinal[0];
-			BlendFinal[1] = bIsResponsiveAAPixel[1] ? tsr_half(1.0/4.0) : BlendFinal[1];
-		#endif
-		
-		FLATTEN
-		if (bCameraCut)
-		{
-			BlendFinal = tsr_half(1.0);
-		}
-		
-		#if 0
-		{
-			tsr_half2 OneMinusBlendFinal = tsr_half(1.0) - BlendFinal;
-		
-			FinalTranslucencyColor = dpv_scale(ClampedPrevHistoryTranslucency, OneMinusBlendFinal) + dpv_scale(FilteredTranslucencyColor, BlendFinal);
-		}
-		#else
-		{
-			tsr_half2 FilterWeight = HdrWeight4(GetColorChannel(FilteredTranslucencyColor));
-			tsr_half2 ClampedHistoryWeight = HdrWeight4(GetColorChannel(ClampedPrevHistoryTranslucency));
-
-			
-			tsr_half2 BlendA = saturate(tsr_half(1.0) - 1.05 * BlendFinal) * ClampedHistoryWeight;
-			tsr_half2 BlendB = BlendFinal * FilterWeight;
-			tsr_half2 RcpBlend = SafeRcp(BlendA + BlendB);
-			BlendA *= RcpBlend;
-			BlendB *= RcpBlend;
-			tsr_half2x2 Weights =  tsr_half2x2(BlendA, BlendB);
-
-			//tsr_half2x2 Weights = WeightedLerpFactors(ClampedHistoryWeight, FilterWeight, BlendFinal);
-
-			FinalTranslucencyColor = dpv_scale(ClampedPrevHistoryTranslucency, Weights[0]) + dpv_scale(FilteredTranslucencyColor, Weights[1]);
-		}
-		#endif
-		
-		FLATTEN
-		if (!bHasSeparateTranslucency)
-		{
-			FinalTranslucencyColor = dpv_interleave_mono_registers(tsr_half4(0.0, 0.0, 0.0, 1.0));
-		}
-	}
-	#endif
 
 	// Compute final output
 	tsr_halfCx2 FinalOutputColor;
 	{
 		FinalOutputColor = FinalHighFrequencyColor;
-
-		#if CONFIG_ACCUMULATE_TRANSLUCENCY_SEPARATELY
-			FinalOutputColor[0] = FinalOutputColor[0] * FinalTranslucencyColor[3] + FinalTranslucencyColor[0];
-			FinalOutputColor[1] = FinalOutputColor[1] * FinalTranslucencyColor[3] + FinalTranslucencyColor[1];
-			FinalOutputColor[2] = FinalOutputColor[2] * FinalTranslucencyColor[3] + FinalTranslucencyColor[2];
-			#if CONFIG_SCENE_COLOR_ALPHA
-				FinalOutputColor[3] = FinalOutputColor[3] * FinalTranslucencyColor[3];
-			#endif
-		#endif
 		
 		// Ensure that alpha values that are expected to be opaque (but are only close to opaque) are forced to be opaque.
 		// (0.995 chosen to accommodate handling of 254/255)
@@ -1287,47 +1165,22 @@ void MainCS(
 		// Protect from NaN and +Inf when writing out the history.
 		{
 			FinalHighFrequencyColor = -dpv_min(-FinalHighFrequencyColor, ColorNull);
-			FinalHighFrequencyColor = dpv_min(FinalHighFrequencyColor, ColorMax10BitsFloat);	
-
-			#if CONFIG_ACCUMULATE_TRANSLUCENCY_SEPARATELY
-				FinalTranslucencyColor = -dpv_min(-FinalTranslucencyColor, tsr_half(0.0).xxxx);
-				FinalTranslucencyColor = dpv_min(FinalTranslucencyColor, tsr_half(Max10BitsFloat).xxxx);
-			#endif
+			FinalHighFrequencyColor = dpv_min(FinalHighFrequencyColor, ColorMax10BitsFloat);
 		}
 
 		// Output full res history
 		{
 			// Output final history lo pixel.
 			{
-				#if CONFIG_OUTPUT_HIGH_FREQUENCY
-					HistoryOutput_ColorArray[tsr_short3(dpv_lo(LocalHistoryPixelPos), HistoryOutput_ArrayIndices_HighFrequency)] = dpv_lo(FinalHighFrequencyColor);
-				#endif
-				HistoryOutput_Metadata[dpv_lo(LocalHistoryPixelPos)]  = dpv_lo(FinalMetadata);
-				#if CONFIG_ACCUMULATE_TRANSLUCENCY_SEPARATELY
-					#if CONFIG_SCENE_COLOR_ALPHA
-						HistoryOutput_ColorArray[tsr_short3(dpv_lo(LocalHistoryPixelPos), HistoryOutput_ArrayIndices_Translucency)] = dpv_lo(FinalTranslucencyColor);
-					#else
-						HistoryOutput_ColorArray[tsr_short3(dpv_lo(LocalHistoryPixelPos), HistoryOutput_ArrayIndices_Translucency)] = dpv_lo(FinalTranslucencyColor).rgb;
-					#endif 
-					HistoryOutput_TranslucencyAlpha[dpv_lo(LocalHistoryPixelPos)] = dpv_lo(FinalTranslucencyColor).a;
-				#endif
+				HistoryColorOutput[tsr_short3(dpv_lo(LocalHistoryPixelPos), HistoryArrayIndices_HighFrequency)] = dpv_lo(FinalHighFrequencyColor);
+				HistoryMetadataOutput[dpv_lo(LocalHistoryPixelPos)]  = dpv_lo(FinalMetadata);
 			}
 		
 			// Output final history hi pixel.
 			#if CONFIG_ENABLE_DUAL_PIXEL_VECTORIZATION
 			{
-				#if CONFIG_OUTPUT_HIGH_FREQUENCY
-					HistoryOutput_ColorArray[tsr_short3(dpv_hi(LocalHistoryPixelPos), HistoryOutput_ArrayIndices_HighFrequency)] = dpv_hi(FinalHighFrequencyColor);
-				#endif
-				HistoryOutput_Metadata[dpv_hi(LocalHistoryPixelPos)] = dpv_hi(FinalMetadata);
-				#if CONFIG_ACCUMULATE_TRANSLUCENCY_SEPARATELY
-					#if CONFIG_SCENE_COLOR_ALPHA
-						HistoryOutput_ColorArray[tsr_short3(dpv_hi(LocalHistoryPixelPos), HistoryOutput_ArrayIndices_Translucency)] = dpv_hi(FinalTranslucencyColor);
-					#else
-						HistoryOutput_ColorArray[tsr_short3(dpv_hi(LocalHistoryPixelPos), HistoryOutput_ArrayIndices_Translucency)] = dpv_hi(FinalTranslucencyColor).rgb;
-					#endif
-					HistoryOutput_TranslucencyAlpha[dpv_hi(LocalHistoryPixelPos)] = dpv_hi(FinalTranslucencyColor).a;
-				#endif
+				HistoryColorOutput[tsr_short3(dpv_hi(LocalHistoryPixelPos), HistoryArrayIndices_HighFrequency)] = dpv_hi(FinalHighFrequencyColor);
+				HistoryMetadataOutput[dpv_hi(LocalHistoryPixelPos)] = dpv_hi(FinalMetadata);
 			}
 			#endif // !CONFIG_ENABLE_DUAL_PIXEL_VECTORIZATION
 		}
@@ -1340,10 +1193,12 @@ void MainCS(
 			OutputColor = -dpv_min(-OutputColor, ColorNull);
 			OutputColor = dpv_min(OutputColor, ColorMax10BitsFloat);
 
-			SceneColorOutputMip0[dpv_lo(LocalHistoryPixelPos)] = dpv_lo(OutputColor);
+			#if CONFIG_OUTPUT_MERGED_SCENE_COLOR
+				SceneColorOutputMip0[dpv_lo(LocalHistoryPixelPos)] = dpv_lo(OutputColor);
 			
-			#if CONFIG_ENABLE_DUAL_PIXEL_VECTORIZATION
-				SceneColorOutputMip0[dpv_hi(LocalHistoryPixelPos)] = dpv_hi(OutputColor);
+				#if CONFIG_ENABLE_DUAL_PIXEL_VECTORIZATION
+					SceneColorOutputMip0[dpv_hi(LocalHistoryPixelPos)] = dpv_hi(OutputColor);
+				#endif
 			#endif
 		}
 		
diff --git a/Engine/Source/Runtime/Renderer/Private/PostProcess/TemporalSuperResolution.cpp b/Engine/Source/Runtime/Renderer/Private/PostProcess/TemporalSuperResolution.cpp
--- a/Engine/Source/Runtime/Renderer/Private/PostProcess/TemporalSuperResolution.cpp
+++ b/Engine/Source/Runtime/Renderer/Private/PostProcess/TemporalSuperResolution.cpp
@@ -17,7 +17,7 @@
 
 namespace
 {
-
+	
 TAutoConsoleVariable<float> CVarTSRHistorySampleCount(
 	TEXT("r.TSR.History.SampleCount"), 16.0f,
 	TEXT("Maximum number sample for each output pixel in the history. Higher values means more stability on highlights on static images, ")
@@ -51,11 +51,6 @@ TAutoConsoleVariable<int32> CVarTSRHistoryUpdateQuality(
 	TEXT("For further details about what each offers, you are invited to look at DIM_UPDATE_QUALITY in TSRUpdateHistory.usf and customise to your need."),
 	ECVF_Scalability | ECVF_RenderThreadSafe);
 
-TAutoConsoleVariable<int32> CVarTSRHistorySeparateTranslucency(
-	TEXT("r.TSR.History.SeparateTranslucency"), 0,
-	TEXT("Whether separate translucency should be accumulated separatly."),
-	ECVF_RenderThreadSafe);
-
 TAutoConsoleVariable<int32> CVarTSRWaveOps(
 	TEXT("r.TSR.WaveOps"), 1,
 	TEXT("Whether to use wave ops in the shading rejection heuristics to speeds up convolutions.\n")
@@ -180,6 +175,45 @@ TAutoConsoleVariable<int32> CVarTSRRejectionAntiAliasingQuality(
 	TEXT("By default, it is only disabled by default in the low anti-aliasing scalability group."),
 	ECVF_Scalability | ECVF_RenderThreadSafe);
 
+TAutoConsoleVariable<int32> CVarTSRResurrectionEnable(
+	TEXT("r.TSR.Resurrection"), 0,
+	TEXT("Allows TSR to resurrect previously discarded details from many frames ago.\n")
+	TEXT("\n")
+	TEXT("When enabled, the entire frames of the TSR are stored in a same unique Texture2DArray including a configurable ")
+	TEXT("number of persistent frame (defined by r.TSR.Resurrection.PersistentFrameCount) that are occasionally recorded ")
+	TEXT("(defined by r.TSR.Resurrection.PersistentFrameInterval).")
+	TEXT("\n")
+	TEXT("Then every frame, TSR will attempt to reproject either previous frame, or the oldest persistent frame available based ")
+	TEXT("which matches best the current frames. The later option will happen when something previously seen by TSR shows up ")
+	TEXT("again (no matter through parallax disocclusion, shading changes, translucent VFX moving) which will have the advantage ")
+	TEXT("bypass the need to newly accumulate a second time by simply resurrected the previously accumulated details.\n")
+	TEXT("\n")
+	TEXT("Command \"vis TSR.HistoryResurrectionMask\" parts of the screen is being resurrected by TSR.\n")
+	TEXT("Command \"vis TSR.History.Color INDEX1\" one of the persistent frame being recorded and later used for history resurrection.\n")
+	TEXT("\n")
+	TEXT("Currently experimental and disabled by default."),
+	ECVF_RenderThreadSafe);
+
+TAutoConsoleVariable<int32> CVarTSRResurrectionPersistentFrameCount(
+	TEXT("r.TSR.Resurrection.PersistentFrameCount"), 2,
+	TEXT("Configures the number of persistent frame to record in history for futur history resurrection. ")
+	TEXT("This will increase the memory footprint of the entire TSR history. ")
+	TEXT("Must be an odd number greater or equal to 2. (default=2)"),
+	ECVF_RenderThreadSafe);
+
+TAutoConsoleVariable<int32> CVarTSRResurrectionPersistentFrameInterval(
+	TEXT("r.TSR.Resurrection.PersistentFrameInterval"), 31,
+	TEXT("Configures in number of frames how often persistent frame should be recorded in history for futur history resurrection. ")
+	TEXT("This has no implication on memory footprint of the TSR history. Must be an even number greater or equal to 1. ")
+	TEXT("(default=31)"),
+	ECVF_RenderThreadSafe);
+
+TAutoConsoleVariable<int32> CVarTSRResurrectionFlickeringHistory(
+	TEXT("r.TSR.Resurrection.Flickering"), 0,
+	TEXT("Configures whether flickering history of r.TSR.ShadingRejection.Flickering should be resurrected. ")
+	TEXT("(Enabled by default)"),
+	ECVF_RenderThreadSafe);
+
 TAutoConsoleVariable<int32> CVarTSRAsyncCompute(
 	TEXT("r.TSR.AsyncCompute"), 2,
 	TEXT("Controls how TSR run on async compute. Some TSR passes can overlap with previous passes.\n")
@@ -272,50 +306,30 @@ END_SHADER_PARAMETER_STRUCT()
 
 BEGIN_SHADER_PARAMETER_STRUCT(FTSRHistoryArrayIndices, )
 	SHADER_PARAMETER(int32, HighFrequency)
-	SHADER_PARAMETER(int32, Translucency)
-
 	SHADER_PARAMETER(int32, Size)
 END_SHADER_PARAMETER_STRUCT()
 
 BEGIN_SHADER_PARAMETER_STRUCT(FTSRHistoryTextures, )
-	SHADER_PARAMETER_STRUCT(FTSRHistoryArrayIndices, ArrayIndices)
-	SHADER_PARAMETER_RDG_TEXTURE(Texture2D, Output)
 	SHADER_PARAMETER_RDG_TEXTURE(Texture2DArray, ColorArray)
-	SHADER_PARAMETER_RDG_TEXTURE(Texture2D, Metadata)
+	SHADER_PARAMETER_RDG_TEXTURE(Texture2DArray, MetadataArray)
+	SHADER_PARAMETER_RDG_TEXTURE(Texture2DArray, GuideArray)
+	SHADER_PARAMETER_RDG_TEXTURE(Texture2DArray, MoireArray)
 	SHADER_PARAMETER_RDG_TEXTURE(Texture2D, SubpixelDepth)
-	SHADER_PARAMETER_RDG_TEXTURE(Texture2D, TranslucencyAlpha)
-	SHADER_PARAMETER_RDG_TEXTURE(Texture2D, Guide)
-	SHADER_PARAMETER_RDG_TEXTURE(Texture2D, Moire)
-END_SHADER_PARAMETER_STRUCT()
-
-BEGIN_SHADER_PARAMETER_STRUCT(FTSRHistorySRVs, )
-	SHADER_PARAMETER_RDG_TEXTURE_SRV(Texture2D, HighFrequency)
-	SHADER_PARAMETER_RDG_TEXTURE_SRV(Texture2D, Metadata)
-
-	SHADER_PARAMETER_RDG_TEXTURE_SRV(Texture2D, Translucency)
-	SHADER_PARAMETER_RDG_TEXTURE_SRV(Texture2D, TranslucencyAlpha)
-	SHADER_PARAMETER_RDG_TEXTURE(Texture2D, Guide)
-	SHADER_PARAMETER_RDG_TEXTURE(Texture2D, Moire)
-END_SHADER_PARAMETER_STRUCT()
-
-BEGIN_SHADER_PARAMETER_STRUCT(FTSRHistoryUAVs, )
-	SHADER_PARAMETER_STRUCT(FTSRHistoryArrayIndices, ArrayIndices)
-	SHADER_PARAMETER_RDG_TEXTURE_UAV(RWTexture2DArray, ColorArray)
-	SHADER_PARAMETER_RDG_TEXTURE_UAV(RWTexture2D, Metadata)
-	SHADER_PARAMETER_RDG_TEXTURE_UAV(RWTexture2D, TranslucencyAlpha)
 END_SHADER_PARAMETER_STRUCT()
 
 BEGIN_SHADER_PARAMETER_STRUCT(FTSRPrevHistoryParameters, )
 	SHADER_PARAMETER_STRUCT(FScreenPassTextureViewportParameters, PrevHistoryInfo)
 	SHADER_PARAMETER(FScreenTransform, ScreenPosToPrevHistoryBufferUV)
+	SHADER_PARAMETER(FScreenTransform, ScreenPosToPrevSubpixelDetails)
+	SHADER_PARAMETER(FVector2f, PrevSubpixelDetailsExtent)
 	SHADER_PARAMETER(float, HistoryPreExposureCorrection)
+	SHADER_PARAMETER(float, ResurrectionPreExposureCorrection)
 END_SHADER_PARAMETER_STRUCT()
 
 enum class ETSRHistoryFormatBits : uint32
 {
 	None = 0,
-	Translucency = 1 << 0,
-	Moire = 1 << 1,
+	Moire = 1 << 0,
 };
 ENUM_CLASS_FLAGS(ETSRHistoryFormatBits);
 
@@ -328,68 +342,17 @@ enum class ETSRSubpixelMethod
 
 bool IsOutputDifferentThanHighFrequency(ETSRHistoryFormatBits HistoryFormatBits)
 {
-	return EnumHasAnyFlags(HistoryFormatBits, ETSRHistoryFormatBits::Translucency);
+	return false;
 }
 
 FTSRHistoryArrayIndices TranslateHistoryFormatBitsToArrayIndices(ETSRHistoryFormatBits HistoryFormatBits)
 {
 	FTSRHistoryArrayIndices ArrayIndices;
-	ArrayIndices.Size = 0;
-	ArrayIndices.HighFrequency = -1;
-	ArrayIndices.Translucency = -1;
-
-	if (IsOutputDifferentThanHighFrequency(HistoryFormatBits))
-	{
-		ArrayIndices.HighFrequency = ArrayIndices.Size++;
-	}
-
-	if (EnumHasAnyFlags(HistoryFormatBits, ETSRHistoryFormatBits::Translucency))
-	{
-		ArrayIndices.Translucency = ArrayIndices.Size++;
-	}
-
+	ArrayIndices.Size = 1;
+	ArrayIndices.HighFrequency = 0;
 	return ArrayIndices;
 }
 
-FTSRHistorySRVs CreateSRVs(FRDGBuilder& GraphBuilder, const FTSRHistoryTextures& Textures)
-{
-	FTSRHistorySRVs SRVs;
-	if (Textures.ArrayIndices.HighFrequency >= 0)
-	{
-		SRVs.HighFrequency = GraphBuilder.CreateSRV(FRDGTextureSRVDesc::CreateForSlice(Textures.ColorArray, Textures.ArrayIndices.HighFrequency));
-	}
-	else
-	{
-		SRVs.HighFrequency = GraphBuilder.CreateSRV(Textures.Output);
-	}
-	SRVs.Metadata = GraphBuilder.CreateSRV(FRDGTextureSRVDesc::Create(Textures.Metadata));
-	if (Textures.ArrayIndices.Translucency >= 0)
-	{
-		SRVs.Translucency = GraphBuilder.CreateSRV(FRDGTextureSRVDesc::CreateForSlice(Textures.ColorArray, Textures.ArrayIndices.Translucency));
-		SRVs.TranslucencyAlpha = GraphBuilder.CreateSRV(FRDGTextureSRVDesc::Create(Textures.TranslucencyAlpha));
-	}
-
-	SRVs.Guide = Textures.Guide;
-	SRVs.Moire = Textures.Moire;
-	return SRVs;
-}
-
-FTSRHistoryUAVs CreateUAVs(FRDGBuilder& GraphBuilder, const FTSRHistoryTextures& Textures)
-{
-	FTSRHistoryUAVs UAVs;
-	UAVs.ArrayIndices = Textures.ArrayIndices;
-	if (Textures.ArrayIndices.Size > 0)
-	{
-		UAVs.ColorArray = GraphBuilder.CreateUAV(Textures.ColorArray);
-	}
-	UAVs.Metadata = GraphBuilder.CreateUAV(Textures.Metadata);
-	if (Textures.TranslucencyAlpha)
-	{
-		UAVs.TranslucencyAlpha = GraphBuilder.CreateUAV(Textures.TranslucencyAlpha);
-	}
-	return UAVs;
-}
-
 class FTSRShader : public FGlobalShader
 {
 public:
@@ -585,12 +548,12 @@ class FTSRDecimateHistoryCS : public FTSRShader
 	SHADER_USE_PARAMETER_STRUCT(FTSRDecimateHistoryCS, FTSRShader);
 
 	class FMoireReprojectionDim : SHADER_PERMUTATION_BOOL("DIM_MOIRE_REPROJECTION");
-	using FPermutationDomain = TShaderPermutationDomain<FMoireReprojectionDim>;
+	class FResurrectionReprojectionDim : SHADER_PERMUTATION_BOOL("DIM_RESURRECTION_REPROJECTION");
+	using FPermutationDomain = TShaderPermutationDomain<FMoireReprojectionDim, FResurrectionReprojectionDim>;
 
 	BEGIN_SHADER_PARAMETER_STRUCT(FParameters, )
 		SHADER_PARAMETER_STRUCT_INCLUDE(FTSRCommonParameters, CommonParameters)
 		SHADER_PARAMETER(FMatrix44f, RotationalClipToPrevClip)
-		SHADER_PARAMETER(FVector2f, CurrentToPrevInputJitter)
 		SHADER_PARAMETER(float, WorldDepthToPixelWorldRadius)
 
 		SHADER_PARAMETER_RDG_TEXTURE(Texture2D, InputSceneColorTexture)
@@ -605,14 +568,21 @@ class FTSRDecimateHistoryCS : public FTSRShader
 		SHADER_PARAMETER_STRUCT(FScreenPassTextureViewportParameters, PrevGuideInfo)
 		SHADER_PARAMETER(FScreenTransform, InputPixelPosToReprojectScreenPos)
 		SHADER_PARAMETER(FScreenTransform, ScreenPosToPrevHistoryGuideBufferUV)
+		SHADER_PARAMETER(FScreenTransform, ScreenPosToResurrectionGuideBufferUV)
+		SHADER_PARAMETER(FVector2f, ResurrectionGuideUVViewportBilinearMin)
+		SHADER_PARAMETER(FVector2f, ResurrectionGuideUVViewportBilinearMax)
 		SHADER_PARAMETER(FVector3f, HistoryGuideQuantizationError)
+		SHADER_PARAMETER(float, ResurrectionFrameIndex)
+		SHADER_PARAMETER(float, PrevFrameIndex)
+		SHADER_PARAMETER(FMatrix44f, ClipToResurrectionClip)
 
-		SHADER_PARAMETER_RDG_TEXTURE_UAV(RWTexture2D, ReprojectedHistoryGuideOutput)
-		SHADER_PARAMETER_RDG_TEXTURE_UAV(RWTexture2D, ReprojectedHistoryMoireOutput)
+		SHADER_PARAMETER_RDG_TEXTURE_UAV(RWTexture2DArray, ReprojectedHistoryGuideOutput)
+		SHADER_PARAMETER_RDG_TEXTURE_UAV(RWTexture2DArray, ReprojectedHistoryMoireOutput)
 		SHADER_PARAMETER_RDG_TEXTURE_UAV(RWTexture2D, HoleFilledVelocityOutput)
 		SHADER_PARAMETER_RDG_TEXTURE_UAV(RWTexture2D, HoleFilledVelocityMaskOutput)
 
 		SHADER_PARAMETER_RDG_TEXTURE_UAV(RWTexture2D, ParallaxRejectionMaskOutput)
+		SHADER_PARAMETER_RDG_TEXTURE_UAV(RWTexture2D, HistoryResurrectionMaskOutput)
 
 		SHADER_PARAMETER_RDG_TEXTURE_UAV(RWTexture2DArray, DebugOutput)
 	END_SHADER_PARAMETER_STRUCT()
@@ -638,12 +608,13 @@ class FTSRRejectShadingCS : public FTSRShader
 		SHADER_PARAMETER(float, TheoricBlendFactor)
 		SHADER_PARAMETER(int32, TileOverscan)
 		SHADER_PARAMETER(float, PerceptionAdd)
+		SHADER_PARAMETER(int32, bOverrideWithResurrectionTexture)
 
 		SHADER_PARAMETER_RDG_TEXTURE(Texture2D, InputTexture)
 		SHADER_PARAMETER_RDG_TEXTURE(Texture2D, InputMoireLumaTexture)
 		SHADER_PARAMETER_RDG_TEXTURE(Texture2D, InputSceneTranslucencyTexture)
-		SHADER_PARAMETER_RDG_TEXTURE(Texture2D, ReprojectedHistoryGuideTexture)
-		SHADER_PARAMETER_RDG_TEXTURE(Texture2D, ReprojectedHistoryMoireTexture)
+		SHADER_PARAMETER_RDG_TEXTURE_SRV(Texture2D, ReprojectedHistoryGuideTexture)
+		SHADER_PARAMETER_RDG_TEXTURE_SRV(Texture2D, ReprojectedHistoryMoireTexture)
 		SHADER_PARAMETER_RDG_TEXTURE(Texture2D, ParallaxRejectionMaskTexture)
 		SHADER_PARAMETER_RDG_TEXTURE_SRV(Texture2D, IsMovingMaskTexture)
 
@@ -726,6 +697,33 @@ class FTSRRejectShadingCS : public FTSRShader
 	}
 }; // class FTSRRejectShadingCS
 
+class FTSRMergeRejectionCS : public FTSRShader
+{
+	DECLARE_GLOBAL_SHADER(FTSRMergeRejectionCS);
+	SHADER_USE_PARAMETER_STRUCT(FTSRMergeRejectionCS, FTSRShader);
+
+	BEGIN_SHADER_PARAMETER_STRUCT(FParameters, )
+		SHADER_PARAMETER_STRUCT_INCLUDE(FTSRCommonParameters, CommonParameters)
+		SHADER_PARAMETER(FMatrix44f, ClipToResurrectionClip)
+
+		SHADER_PARAMETER_RDG_TEXTURE(Texture2D, ClosestDepthTexture)
+		SHADER_PARAMETER_RDG_TEXTURE(Texture2D, PrevGuideTexture)
+		SHADER_PARAMETER_RDG_TEXTURE(Texture2D, PrevMoireTexture)
+		SHADER_PARAMETER_RDG_TEXTURE(Texture2D, PrevRejectionTexture)
+		SHADER_PARAMETER_RDG_TEXTURE(Texture2D, ResurrectedGuideTexture)
+		SHADER_PARAMETER_RDG_TEXTURE(Texture2D, ResurrectedMoireTexture)
+		SHADER_PARAMETER_RDG_TEXTURE(Texture2D, ResurrectedRejectionTexture)
+
+		SHADER_PARAMETER_RDG_TEXTURE_UAV(RWTexture2D, HistoryGuideOutput)
+		SHADER_PARAMETER_RDG_TEXTURE_UAV(RWTexture2D, HistoryMoireOutput)
+		SHADER_PARAMETER_RDG_TEXTURE_UAV(RWTexture2D, HistoryRejectionOutput)
+		SHADER_PARAMETER_RDG_TEXTURE_UAV(RWTexture2D, HistoryResurrectionMaskOutput)
+		SHADER_PARAMETER_RDG_TEXTURE_UAV(RWTexture2D, DilatedVelocityOutput)
+
+		SHADER_PARAMETER_RDG_TEXTURE_UAV(RWTexture2DArray, DebugOutput)
+	END_SHADER_PARAMETER_STRUCT()
+}; // class FTSRMergeRejectionCS
+
 class FTSRSpatialAntiAliasingCS : public FTSRShader
 {
 	DECLARE_GLOBAL_SHADER(FTSRSpatialAntiAliasingCS);
@@ -792,9 +790,8 @@ class FTSRUpdateHistoryCS : public FTSRShader
 	};
 
 	class FQualityDim : SHADER_PERMUTATION_ENUM_CLASS("DIM_UPDATE_QUALITY", EQuality);
-	class FSeparateTranslucencyDim : SHADER_PERMUTATION_BOOL("DIM_SEPARATE_TRANSLUCENCY");
 
-	using FPermutationDomain = TShaderPermutationDomain<FQualityDim, FSeparateTranslucencyDim, FTSRShader::F16BitVALUDim>;
+	using FPermutationDomain = TShaderPermutationDomain<FQualityDim, FTSRShader::F16BitVALUDim>;
 
 	BEGIN_SHADER_PARAMETER_STRUCT(FParameters, )
 		SHADER_PARAMETER_STRUCT_INCLUDE(FTSRCommonParameters, CommonParameters)
@@ -806,6 +803,7 @@ class FTSRUpdateHistoryCS : public FTSRShader
 
 		SHADER_PARAMETER_RDG_TEXTURE(Texture2D, DilatedVelocityTexture)
 		SHADER_PARAMETER_RDG_TEXTURE(Texture2D, ParallaxRejectionMaskTexture)
+		SHADER_PARAMETER_RDG_TEXTURE(Texture2D, HistoryResurrectionMaskTexture)
 		SHADER_PARAMETER_RDG_TEXTURE(Texture2D, AntiAliasingTexture)
 		SHADER_PARAMETER_RDG_TEXTURE(Texture2D, NoiseFilteringTexture)
 		SHADER_PARAMETER_RDG_TEXTURE(Texture2D, HoleFilledVelocityMaskTexture)
@@ -826,17 +824,22 @@ class FTSRUpdateHistoryCS : public FTSRShader
 		SHADER_PARAMETER(float, InvWeightClampingPixelSpeed)
 		SHADER_PARAMETER(float, InputToHistoryFactor)
 		SHADER_PARAMETER(float, InputContributionMultiplier)
+		SHADER_PARAMETER(float, ResurrectionFrameIndex)
+		SHADER_PARAMETER(float, PrevFrameIndex)
 		SHADER_PARAMETER(int32, ResponsiveStencilMask)
 		SHADER_PARAMETER(int32, bGenerateOutputMip1)
 		SHADER_PARAMETER(int32, bGenerateOutputMip2)
 		SHADER_PARAMETER(int32, bHasSeparateTranslucency)
 
+		SHADER_PARAMETER_STRUCT(FTSRHistoryArrayIndices, HistoryArrayIndices)
 		SHADER_PARAMETER_STRUCT_INCLUDE(FTSRPrevHistoryParameters, PrevHistoryParameters)
-		SHADER_PARAMETER_STRUCT(FTSRHistorySRVs, PrevHistory)
+		SHADER_PARAMETER_RDG_TEXTURE_SRV(Texture2DArray, PrevHistoryColorTexture)
+		SHADER_PARAMETER_RDG_TEXTURE_SRV(Texture2DArray, PrevHistoryMetadataTexture)
 
 		SHADER_PARAMETER_RDG_TEXTURE_UAV(RWTexture2D, SceneColorOutputMip0)
 		SHADER_PARAMETER_RDG_TEXTURE_UAV(RWTexture2D, SceneColorOutputMip1)
-		SHADER_PARAMETER_STRUCT(FTSRHistoryUAVs, HistoryOutput)
+		SHADER_PARAMETER_RDG_TEXTURE_UAV(RWTexture2DArray, HistoryColorOutput)
+		SHADER_PARAMETER_RDG_TEXTURE_UAV(RWTexture2DArray, HistoryMetadataOutput)
 		SHADER_PARAMETER_RDG_TEXTURE_UAV(RWTexture2DArray, DebugOutput)
 	END_SHADER_PARAMETER_STRUCT()
 
@@ -882,7 +885,7 @@ class FTSRResolveHistoryCS : public FTSRShader
 		SHADER_PARAMETER(int32, bGenerateOutputMip1)
 		SHADER_PARAMETER(float, HistoryValidityMultiply)
 
-		SHADER_PARAMETER_RDG_TEXTURE(Texture2D, UpdateHistoryOutputTexture)
+		SHADER_PARAMETER_RDG_TEXTURE_SRV(Texture2D, UpdateHistoryOutputTexture)
 
 		SHADER_PARAMETER_RDG_TEXTURE_UAV(RWTexture2D, SceneColorOutputMip0)
 		SHADER_PARAMETER_RDG_TEXTURE_UAV(RWTexture2D, SceneColorOutputMip1)
@@ -948,6 +951,7 @@ IMPLEMENT_GLOBAL_SHADER(FTSRForwardScatterDepthCS,   "/Engine/Private/TemporalSu
 IMPLEMENT_GLOBAL_SHADER(FTSRDilateVelocityCS,        "/Engine/Private/TemporalSuperResolution/TSRDilateVelocity.usf",        "MainCS", SF_Compute);
 IMPLEMENT_GLOBAL_SHADER(FTSRDecimateHistoryCS,       "/Engine/Private/TemporalSuperResolution/TSRDecimateHistory.usf",       "MainCS", SF_Compute);
 IMPLEMENT_GLOBAL_SHADER(FTSRRejectShadingCS,         "/Engine/Private/TemporalSuperResolution/TSRRejectShading.usf",         "MainCS", SF_Compute);
+IMPLEMENT_GLOBAL_SHADER(FTSRMergeRejectionCS,        "/Engine/Private/TemporalSuperResolution/TSRMergeRejection.usf",        "MainCS", SF_Compute);
 IMPLEMENT_GLOBAL_SHADER(FTSRSpatialAntiAliasingCS,   "/Engine/Private/TemporalSuperResolution/TSRSpatialAntiAliasing.usf",   "MainCS", SF_Compute);
 IMPLEMENT_GLOBAL_SHADER(FTSRFilterAntiAliasingCS,    "/Engine/Private/TemporalSuperResolution/TSRFilterAntiAliasing.usf",    "MainCS", SF_Compute);
 IMPLEMENT_GLOBAL_SHADER(FTSRUpdateHistoryCS,         "/Engine/Private/TemporalSuperResolution/TSRUpdateHistory.usf",         "MainCS", SF_Compute);
@@ -978,6 +982,116 @@ static FRDGTextureUAVRef CreateDummyUAV(FRDGBuilder& GraphBuilder, EPixelFormat
 	return GraphBuilder.CreateUAV(DummyTexture);
 };
 
+struct FTSRHistorySliceSequence
+{
+	static constexpr int32 kTransientSliceCount = 2;
+
+	int32 FrameStorageCount = 1;
+	int32 FrameStoragePeriod = 1;
+
+	bool Check() const
+	{
+		check(FrameStorageCount == 1 || ((FrameStorageCount >= 4) && (FrameStorageCount % 2) == 0));
+		check((FrameStoragePeriod % 2) == 1);
+		return true;
+	}
+	
+	/** Returns the total number of rolling indices. */
+	int32 GetRollingIndexCount() const
+	{
+		if (FrameStorageCount == 1)
+		{
+			check(FrameStoragePeriod == 1);
+			return 1;
+		}
+		else if (FrameStoragePeriod == 1)
+		{
+			return FrameStorageCount;
+		}
+
+		const int32 TransientIndexCount = kTransientSliceCount;
+		const int32 PersistentIndexCount = FrameStorageCount - TransientIndexCount;
+
+		return PersistentIndexCount * FrameStoragePeriod;
+	}
+
+	/** Returns a rolling index incremented by one. */
+	int32 IncrementFrameRollingIndex(int32 PrevFrameRollingIndex) const
+	{
+		return (PrevFrameRollingIndex + 1) % GetRollingIndexCount();
+	}
+
+	/** Returns a rolling index incremented by one. */
+	int32 DecrementFrameRollingIndex(int32 CurrentFrameRollingIndex) const
+	{
+		return (CurrentFrameRollingIndex + GetRollingIndexCount() - 1) % GetRollingIndexCount();
+	}
+
+	/** Returns a rolling index incremented by one. */
+	int32 RollingIndexToSliceIndex(int32 FrameRollingIndex) const
+	{
+		if (FrameStorageCount == 1)
+		{
+			check(FrameRollingIndex == 0);
+			check(FrameStoragePeriod == 1);
+			return 0;
+		}
+		else if (FrameStoragePeriod == 1)
+		{
+			return (FrameRollingIndex % 2) * (FrameStorageCount / 2) + (FrameRollingIndex / 2) % (FrameStorageCount / 2);
+		}
+
+		const int32 TransientIndexCount = kTransientSliceCount;
+		const int32 PersistentIndexCount = FrameStorageCount - TransientIndexCount;
+		//const int32 FrameRollingIndexCount = PersistentIndexCount * FrameStoragePeriod;
+		//check(FrameRollingIndex >= 0 && FrameRollingIndex < FrameRollingIndexCount);
+
+		const bool bIsPersistentRollingIndex = (FrameRollingIndex % FrameStoragePeriod) == 0;
+		if (bIsPersistentRollingIndex)
+		{
+			const int32 PersistentIndex = FrameRollingIndex / FrameStoragePeriod;
+
+			return (PersistentIndex % 2)
+				? ((FrameStorageCount / 2) + (PersistentIndex / 2))
+				: ((FrameStorageCount / 2) - (PersistentIndex / 2) - 1);
+		}
+		else
+		{
+			return (FrameRollingIndex % 2) ? (FrameStorageCount - 1) : 0;
+		}
+	}
+
+	int32 GetResurrectionFrameRollingIndex(int32 AccumulatedFrameCount, int32 LastFrameRollingIndex) const
+	{
+		const int32 RollingIndexCount = GetRollingIndexCount();
+
+		if (FrameStorageCount == 1)
+		{
+			check(FrameStoragePeriod == 1);
+			return 0;
+		}
+		else if (FrameStoragePeriod == 1)
+		{
+			return (RollingIndexCount + LastFrameRollingIndex - FMath::DivideAndRoundUp(FMath::Max(AccumulatedFrameCount - 2, 0), 2) * 2) % RollingIndexCount;
+		}
+		
+		if (AccumulatedFrameCount < RollingIndexCount)
+		{
+			return 0;
+		}
+
+		return (FMath::DivideAndRoundUp(LastFrameRollingIndex + FrameStoragePeriod, FrameStoragePeriod) * FrameStoragePeriod) % RollingIndexCount;
+	}
+
+	FRHIRange16 GetSRVSliceRange(int32 CurrentFrameSliceIndex, int32 PrevFrameSliceIndex) const
+	{
+		check(CurrentFrameSliceIndex != PrevFrameSliceIndex);
+		return (PrevFrameSliceIndex > CurrentFrameSliceIndex)
+			? FRHIRange16(CurrentFrameSliceIndex + 1, FrameStorageCount - CurrentFrameSliceIndex - 1)
+			: FRHIRange16(0, CurrentFrameSliceIndex);
+	}
+};
+
 bool NeedTSRMoireLuma(const FViewInfo& View)
 {
 	return GetMainTAAPassConfig(View) == EMainTAAPassConfig::TSR;
@@ -986,7 +1100,7 @@ bool NeedTSRMoireLuma(const FViewInfo& View)
 FScreenPassTexture AddTSRComputeMoireLuma(FRDGBuilder& GraphBuilder, FGlobalShaderMap* ShaderMap, FScreenPassTexture SceneColor)
 {
 	check(SceneColor.Texture)
-;	RDG_GPU_STAT_SCOPE(GraphBuilder, TemporalSuperResolution);
+	RDG_GPU_STAT_SCOPE(GraphBuilder, TemporalSuperResolution);
 
 	FScreenPassTexture MoireLuma;
 	{
@@ -1028,6 +1142,21 @@ FDefaultTemporalUpscaler::FOutputs AddTemporalSuperResolutionPasses(
 	// Maximum depth lifetime encodable in FTSRSubpixelDepth::RemainingLifetime
 	const int32 kMaxDepthLifetime = (1 << 7) - 1;
 
+	// Number of frames stored in the history.
+	FTSRHistorySliceSequence HistorySliceSequence;
+	if (CVarTSRResurrectionEnable.GetValueOnRenderThread())
+	{
+		HistorySliceSequence.FrameStorageCount = FMath::Clamp(
+			FTSRHistorySliceSequence::kTransientSliceCount + FMath::DivideAndRoundUp(CVarTSRResurrectionPersistentFrameCount.GetValueOnRenderThread(), 2) * 2,
+			4,
+			GMaxTextureArrayLayers);
+		HistorySliceSequence.FrameStoragePeriod = FMath::Clamp(CVarTSRResurrectionPersistentFrameInterval.GetValueOnRenderThread() | 0x1, 1, 1024);
+	}
+	check(HistorySliceSequence.Check());
+
+	// Whether the flickering history should be resurrected.
+	const bool bResurrectFlickeringHistory = HistorySliceSequence.FrameStorageCount > 1 && CVarTSRResurrectionFlickeringHistory.GetValueOnRenderThread() != 0;
+
 	// Whether to use wave ops optimizations.
 	const ERHIFeatureSupport WaveOpsSupport = FTSRShader::SupportsWaveOps(View.GetShaderPlatform());
 	const bool bSupportsLDS = FTSRShader::SupportsLDS(View.GetShaderPlatform());
@@ -1043,9 +1172,6 @@ FDefaultTemporalUpscaler::FOutputs AddTemporalSuperResolutionPasses(
 	// Whether alpha channel is supported.
 	const bool bSupportsAlpha = IsPostProcessingWithAlphaChannelSupported();
 
-	// Whether separate translucency should be temporarily accumulated separatly.
-	const bool bAccumulateTranslucencySeparately = CVarTSRHistorySeparateTranslucency.GetValueOnRenderThread() != 0;
-
 	const float RefreshRateToFrameRateCap = (View.Family->Time.GetDeltaRealTimeSeconds() > 0.0f && CVarTSRFlickeringAdjustToFrameRate.GetValueOnRenderThread())
 		? View.Family->Time.GetDeltaRealTimeSeconds() * CVarTSRFlickeringFrameRateCap.GetValueOnRenderThread() : 1.0f;
 
@@ -1060,11 +1186,6 @@ FDefaultTemporalUpscaler::FOutputs AddTemporalSuperResolutionPasses(
 
 	ETSRHistoryFormatBits HistoryFormatBits = ETSRHistoryFormatBits::None;
 	{
-		if (bAccumulateTranslucencySeparately)
-		{
-			HistoryFormatBits |= ETSRHistoryFormatBits::Translucency;
-		}
-
 		if (FlickeringFramePeriod > 0)
 		{
 			HistoryFormatBits |= ETSRHistoryFormatBits::Moire;
@@ -1074,9 +1195,6 @@ FDefaultTemporalUpscaler::FOutputs AddTemporalSuperResolutionPasses(
 
 	const bool bIsOutputDifferentThanHighFrequency = IsOutputDifferentThanHighFrequency(HistoryFormatBits);
 
-	// Whether to use camera cut shader permutation or not.
-	bool bCameraCut = !InputHistory.IsValid() || View.bCameraCut || ETSRHistoryFormatBits(InputHistory.FormatBit) != HistoryFormatBits;
-
 	FTSRUpdateHistoryCS::EQuality UpdateHistoryQuality = FTSRUpdateHistoryCS::EQuality(FMath::Clamp(CVarTSRHistoryUpdateQuality.GetValueOnRenderThread(), 0, int32(FTSRUpdateHistoryCS::EQuality::MAX) - 1));
 
 	bool bIsSeperateTranslucyTexturesValid = PassInputs.PostDOFTranslucencyResources.IsValid();
@@ -1145,6 +1263,13 @@ FDefaultTemporalUpscaler::FOutputs AddTemporalSuperResolutionPasses(
 	float OutputToInputResolutionFraction = float(InputRect.Width()) / float(OutputRect.Width());
 	float OutputToInputResolutionFractionSquare = OutputToInputResolutionFraction * OutputToInputResolutionFraction;
 
+	// Whether to use camera cut shader permutation or not.
+	const bool bCameraCut =
+		!InputHistory.IsValid() ||
+		View.bCameraCut ||
+		ETSRHistoryFormatBits(InputHistory.FormatBit) != HistoryFormatBits ||
+		false;
+
 	RDG_EVENT_SCOPE(GraphBuilder, "TemporalSuperResolution(%s) %dx%d -> %dx%d",
 		bSupportsAlpha ? TEXT("Alpha") : TEXT(""),
 		InputRect.Width(), InputRect.Height(),
@@ -1153,6 +1278,7 @@ FDefaultTemporalUpscaler::FOutputs AddTemporalSuperResolutionPasses(
 
 	FRDGTextureRef BlackUintDummy = GSystemTextures.GetZeroUIntDummy(GraphBuilder);
 	FRDGTextureRef BlackDummy = GSystemTextures.GetBlackDummy(GraphBuilder);
+	FRDGTextureRef BlackArrayDummy = GSystemTextures.GetBlackArrayDummy(GraphBuilder);
 	FRDGTextureRef BlackAlphaOneDummy = GSystemTextures.GetBlackAlphaOneDummy(GraphBuilder);
 	FRDGTextureRef WhiteDummy = GSystemTextures.GetWhiteDummy(GraphBuilder);
 
@@ -1224,146 +1350,222 @@ FDefaultTemporalUpscaler::FOutputs AddTemporalSuperResolutionPasses(
 		return GraphBuilder.CreateUAV(DebugTexture);
 	};
 
-	// Create new history.
-	FRDGTextureRef UpdateHistoryOutputTexture;
-	FRDGTextureRef SceneColorOutputTexture;
+	// Output
+	FRDGTextureRef UpdateHistoryOutputTexture = nullptr;
+	FRDGTextureRef SceneColorOutputTexture = nullptr;
 	FRDGTextureRef SceneColorOutputHalfResTexture = nullptr;
 	FRDGTextureRef SceneColorOutputQuarterResTexture = nullptr;
-	FTSRHistoryTextures History;
 	{
 		check(!(PassInputs.bGenerateOutputMip1 && (PassInputs.bGenerateSceneColorHalfRes || PassInputs.bGenerateSceneColorQuarterRes)));
-		FRDGTextureDesc Desc = FRDGTextureDesc::Create2D(
-			HistoryExtent,
-			bIsOutputDifferentThanHighFrequency ? ColorFormat : HistoryColorFormat,
+		FRDGTextureDesc OutputDesc = FRDGTextureDesc::Create2D(
+			OutputExtent,
+			ColorFormat,
 			FClearValueBinding::None,
 			/* InFlags = */ TexCreate_ShaderResource | TexCreate_UAV,
 			/* NumMips = */ PassInputs.bGenerateOutputMip1 ? 2 : 1);
 
-		UpdateHistoryOutputTexture = GraphBuilder.CreateTexture(Desc, TEXT("TSR.Output"));
-		History.Output = UpdateHistoryOutputTexture;
-
-		Desc.Format = ColorFormat;
-
 		if (OutputRect.Size() != HistorySize)
 		{
-			Desc.Extent = OutputExtent;
-			SceneColorOutputTexture = GraphBuilder.CreateTexture(Desc, TEXT("TSR.Output"));
+			SceneColorOutputTexture = GraphBuilder.CreateTexture(OutputDesc, TEXT("TSR.Output"));
+
+			if (bIsOutputDifferentThanHighFrequency)
+			{
+				FRDGTextureDesc HistoryOutputDesc = OutputDesc;
+				HistoryOutputDesc.Extent = HistoryExtent;
+				HistoryOutputDesc.NumMips = 1;
+				UpdateHistoryOutputTexture = GraphBuilder.CreateTexture(OutputDesc, TEXT("TSR.HistoryUpdateOutput"));
+			}
 		}
 		else
 		{
-			SceneColorOutputTexture = UpdateHistoryOutputTexture;
+			UpdateHistoryOutputTexture = GraphBuilder.CreateTexture(OutputDesc, TEXT("TSR.Output"));
 		}
 
 		// Generate quarter res output if only one needed, otherwise only output half res and let caller downscale to quarter res. 
 		// This is motivated to saves UAV slots on FTSRUpdateHistoryCS
 		if (PassInputs.bGenerateSceneColorQuarterRes && !PassInputs.bGenerateSceneColorHalfRes && OutputRect.Size() == HistorySize)
 		{
-			FRDGTextureDesc QuarterResDesc = Desc;
+			FRDGTextureDesc QuarterResDesc = OutputDesc;
 			QuarterResDesc.Extent = OutputExtent / 4;
 			SceneColorOutputQuarterResTexture = GraphBuilder.CreateTexture(QuarterResDesc, TEXT("TSR.QuarterResOutput"));
 		}
 		else if (PassInputs.bGenerateSceneColorHalfRes || PassInputs.bGenerateSceneColorQuarterRes)
 		{
-			FRDGTextureDesc HalfResDesc = Desc;
+			FRDGTextureDesc HalfResDesc = OutputDesc;
 			HalfResDesc.Extent = OutputExtent / 2;
 			SceneColorOutputHalfResTexture = GraphBuilder.CreateTexture(HalfResDesc, TEXT("TSR.HalfResOutput"));
 		}
 	}
-	{
-		History.ArrayIndices = TranslateHistoryFormatBitsToArrayIndices(HistoryFormatBits);
 
-		if (History.ArrayIndices.Size > 0)
+	// Allocate a new history
+	FTSRHistoryTextures History;
+	{
 		{
+			bool bRequires2Mips = HistorySize == OutputRect.Size() && !bIsOutputDifferentThanHighFrequency && PassInputs.bGenerateOutputMip1;
+
 			FRDGTextureDesc ArrayDesc = FRDGTextureDesc::Create2DArray(
 				HistoryExtent,
 				HistoryColorFormat,
 				FClearValueBinding::None,
 				TexCreate_ShaderResource | TexCreate_UAV,
-				History.ArrayIndices.Size);
-
-			History.ColorArray = GraphBuilder.CreateTexture(ArrayDesc, TEXT("TSR.History.ColorArray"));
+				HistoryArrayIndices.Size * HistorySliceSequence.FrameStorageCount,
+				/* NumMips = */ bRequires2Mips ? 2 : 1);
+			History.ColorArray = GraphBuilder.CreateTexture(ArrayDesc, TEXT("TSR.History.Color"));
 		}
-	}
-	{
-		FRDGTextureDesc Desc = FRDGTextureDesc::Create2D(
-			HistoryExtent,
-			HistoryColorFormat,
-			FClearValueBinding::None,
-			TexCreate_ShaderResource | TexCreate_UAV);
 
-		Desc.Format = PF_R8;
-		History.Metadata = GraphBuilder.CreateTexture(Desc, TEXT("TSR.History.Metadata"));
+		{
+			FRDGTextureDesc Desc = FRDGTextureDesc::Create2DArray(
+				HistoryExtent,
+				PF_R8,
+				FClearValueBinding::None,
+				TexCreate_ShaderResource | TexCreate_UAV,
+				HistorySliceSequence.FrameStorageCount);
+			History.MetadataArray = GraphBuilder.CreateTexture(Desc, TEXT("TSR.History.Metadata"));
+		}
 
-		Desc.Format = PF_R8;
-		if (bAccumulateTranslucencySeparately)
 		{
-			History.TranslucencyAlpha = GraphBuilder.CreateTexture(Desc, TEXT("TSR.History.TranslucencyAlpha"));
+			FRDGTextureDesc Desc = FRDGTextureDesc::Create2DArray(
+				InputExtent,
+				bSupportsAlpha ? PF_FloatRGBA : PF_A2B10G10R10,
+				FClearValueBinding::None,
+				TexCreate_ShaderResource | TexCreate_UAV,
+				HistorySliceSequence.FrameStorageCount);
+
+			History.GuideArray = GraphBuilder.CreateTexture(Desc, TEXT("TSR.History.Guide"));
+			Desc.Format = PF_R8G8B8A8;
+			History.MoireArray = GraphBuilder.CreateTexture(Desc, TEXT("TSR.History.Moire"));
 		}
 	}
 
-	if (SubpixelMethod == ETSRSubpixelMethod::ClosestDepth)
+	// Whether to camera cut the history Resurrection
+	const bool bCameraCutResurrection =
+		bCameraCut ||
+		HistorySliceSequence.GetRollingIndexCount() == 1 ||
+		InputHistory.OutputViewportRect != FIntRect(FIntPoint(0, 0), HistorySize) ||
+		InputHistory.FrameStorageCount != HistorySliceSequence.FrameStorageCount ||
+		InputHistory.FrameStoragePeriod != HistorySliceSequence.FrameStoragePeriod ||
+		History.ColorArray->Desc.Extent != InputHistory.ColorArray->GetDesc().Extent ||
+		History.ColorArray->Desc.Format != InputHistory.ColorArray->GetDesc().Format ||
+		History.ColorArray->Desc.NumMips != InputHistory.ColorArray->GetDesc().NumMips ||
+		History.ColorArray->Desc.ArraySize != InputHistory.ColorArray->GetDesc().ArraySize ||
+		History.GuideArray->Desc.Extent != InputHistory.GuideArray->GetDesc().Extent ||
+		History.GuideArray->Desc.Format != InputHistory.GuideArray->GetDesc().Format ||
+		false;
+
+	// Current and previous frame histories
+	int32 ResurrectionFrameSliceIndex = 0;
+	int32 PrevFrameSliceIndex = 0;
+	int32 CurrentFrameSliceIndex = 0;
+	int32 CurrentFrameRollingIndex = 0;
+	FTSRHistoryTextures PrevHistory;
+	FTSRHistorySliceSequence PrevHistorySliceSequence;
+	if (bCameraCut)
 	{
-		FRDGTextureDesc SubpixelDepthDesc = FRDGTextureDesc::Create2D(
-			InputExtent,
-			PF_R32_UINT,
-			FClearValueBinding::None,
-			TexCreate_ShaderResource | TexCreate_UAV);
-		History.SubpixelDepth = GraphBuilder.CreateTexture(SubpixelDepthDesc, TEXT("TSR.History.SubpixelDepth"));
-	}
+		PrevHistory.ColorArray = BlackArrayDummy;
+		PrevHistory.MetadataArray = BlackArrayDummy;
+		PrevHistory.GuideArray = BlackArrayDummy;
+		PrevHistory.MoireArray = BlackArrayDummy;
+		PrevHistory.SubpixelDepth = BlackUintDummy;
 
-	{
-		FRDGTextureDesc Desc = FRDGTextureDesc::Create2D(
-			InputExtent,
-			bSupportsAlpha ? PF_FloatRGBA : PF_A2B10G10R10,
-			FClearValueBinding::None,
-			TexCreate_ShaderResource | TexCreate_UAV);
+		if (HistorySliceSequence.GetRollingIndexCount() > 1)
+		{
+			check(bCameraCutResurrection);
 
-		History.Guide = GraphBuilder.CreateTexture(Desc, TEXT("TSR.History.Guide"));
-		Desc.Format = PF_R8G8B8A8;
-		History.Moire = GraphBuilder.CreateTexture(Desc, TEXT("TSR.History.Moire"));
-	}
+			int32 PrevFrameRollingIndex = HistorySliceSequence.DecrementFrameRollingIndex(CurrentFrameRollingIndex);
 
-	// Setup a dummy history
-	FTSRHistorySRVs DummyHistorySRVs;
-	{
-		DummyHistorySRVs.HighFrequency = GraphBuilder.CreateSRV(FRDGTextureSRVDesc::Create(BlackDummy));
-		DummyHistorySRVs.Metadata = GraphBuilder.CreateSRV(FRDGTextureSRVDesc::Create(BlackDummy));
-		DummyHistorySRVs.Translucency = GraphBuilder.CreateSRV(FRDGTextureSRVDesc::Create(BlackDummy));
-		DummyHistorySRVs.TranslucencyAlpha = GraphBuilder.CreateSRV(FRDGTextureSRVDesc::Create(WhiteDummy));
-		DummyHistorySRVs.Guide = BlackDummy;
-		DummyHistorySRVs.Moire = BlackDummy;
+			ResurrectionFrameSliceIndex = HistorySliceSequence.RollingIndexToSliceIndex(PrevFrameRollingIndex);
+			PrevFrameSliceIndex        = HistorySliceSequence.RollingIndexToSliceIndex(PrevFrameRollingIndex);
+			CurrentFrameSliceIndex     = HistorySliceSequence.RollingIndexToSliceIndex(CurrentFrameRollingIndex);
+		}
 	}
-
-	// Setup the previous frame history
-	FTSRHistoryTextures PrevHistory;
-	FTSRHistorySRVs PrevHistorySRVs;
-	if (!bCameraCut)
+	else
 	{
-		PrevHistory.ArrayIndices = History.ArrayIndices;
+		PrevHistorySliceSequence.FrameStorageCount = InputHistory.FrameStorageCount;
+		PrevHistorySliceSequence.FrameStoragePeriod = InputHistory.FrameStoragePeriod;
+		check(PrevHistorySliceSequence.Check());
 
 		// Register filterable history
-		PrevHistory.Output = InputHistory.Output.IsValid() ? GraphBuilder.RegisterExternalTexture(InputHistory.Output, TEXT("TSR.PrevHistory.Output")) : nullptr;
-		PrevHistory.ColorArray = InputHistory.ColorArray.IsValid() ? GraphBuilder.RegisterExternalTexture(InputHistory.ColorArray, TEXT("TSR.PrevHistory.ColorArray")) : nullptr;
-		PrevHistory.Metadata = GraphBuilder.RegisterExternalTexture(InputHistory.Metadata, TEXT("TSR.PrevHistory.Metadata"));
-		PrevHistory.TranslucencyAlpha = InputHistory.TranslucencyAlpha.IsValid() ? GraphBuilder.RegisterExternalTexture(InputHistory.TranslucencyAlpha, TEXT("TSR.PrevHistory.TranslucencyAlpha")) : DummyHistorySRVs.TranslucencyAlpha->Desc.Texture;
-		PrevHistory.Guide = GraphBuilder.RegisterExternalTexture(InputHistory.Guide, TEXT("TSR.PrevHistory.Guide"));
-		PrevHistory.Moire = InputHistory.Moire.IsValid() ? GraphBuilder.RegisterExternalTexture(InputHistory.Moire, TEXT("TSR.PrevHistory.Moire")) : DummyHistorySRVs.Moire;
+		PrevHistory.ColorArray = GraphBuilder.RegisterExternalTexture(InputHistory.ColorArray);
+		PrevHistory.MetadataArray = GraphBuilder.RegisterExternalTexture(InputHistory.MetadataArray);
+		PrevHistory.GuideArray = GraphBuilder.RegisterExternalTexture(InputHistory.GuideArray);
+		PrevHistory.MoireArray = InputHistory.MoireArray.IsValid()
+			? GraphBuilder.RegisterExternalTexture(InputHistory.MoireArray)
+			: BlackDummy;
 
 		// Register non-filterable history
-		PrevHistory.SubpixelDepth = InputHistory.SubpixelDepth.IsValid() ? GraphBuilder.RegisterExternalTexture(InputHistory.SubpixelDepth, TEXT("TSR.PrevHistory.SubpixelDepth")) : nullptr;
+		PrevHistory.SubpixelDepth = InputHistory.SubpixelDepth.IsValid()
+			? GraphBuilder.RegisterExternalTexture(InputHistory.SubpixelDepth, TEXT("TSR.PrevHistory.SubpixelDepth"))
+			: nullptr;
+
+		int32 ResurrectionFrameRollingIndex = 0;
+		int32 PrevFrameRollingIndex = 0;
+		if (HistorySliceSequence.GetRollingIndexCount() == 1)
+		{
+			// NOP
+		}
+		else if (bCameraCutResurrection)
+		{
+			ResurrectionFrameRollingIndex = InputHistory.LastFrameRollingIndex;
+			PrevFrameRollingIndex = InputHistory.LastFrameRollingIndex;
+		}
+		else
+		{
+			// Reuse same history so all frames of the history are in the same Texture2DArray for
+			// history resurrection without branching on texture fetches.
+			if (!View.bStatePrevViewInfoIsReadOnly)
+			{
+				History.ColorArray = PrevHistory.ColorArray;
+				History.MetadataArray = PrevHistory.MetadataArray;
+				History.GuideArray = PrevHistory.GuideArray;
+				History.MoireArray = PrevHistory.MoireArray;
+			}
 
-		PrevHistorySRVs = CreateSRVs(GraphBuilder, PrevHistory);
+			ResurrectionFrameRollingIndex = PrevHistorySliceSequence.GetResurrectionFrameRollingIndex(InputHistory.AccumulatedFrameCount, InputHistory.LastFrameRollingIndex);
+			PrevFrameRollingIndex = InputHistory.LastFrameRollingIndex;
+			CurrentFrameRollingIndex = PrevHistorySliceSequence.IncrementFrameRollingIndex(InputHistory.LastFrameRollingIndex);
+		}
+
+		// Translate rolling indices to slice indices to work arround D3D limitation that prevents writing to a Texture2DArray slice when
+		// the array is entirely bound.
+		ResurrectionFrameSliceIndex = PrevHistorySliceSequence.RollingIndexToSliceIndex(ResurrectionFrameRollingIndex);
+		PrevFrameSliceIndex = PrevHistorySliceSequence.RollingIndexToSliceIndex(PrevFrameRollingIndex);
+		CurrentFrameSliceIndex = HistorySliceSequence.RollingIndexToSliceIndex(CurrentFrameRollingIndex);
 	}
-	else
+
+	// Whether history Resurrection is possible at all 
+	const bool bCanResurrectHistory = ResurrectionFrameSliceIndex != PrevFrameSliceIndex;
+
+	FMatrix44f ClipToResurrectionClip = FMatrix44f::Identity;
+	FScreenPassTextureViewport ResurrectionGuideViewport(FIntPoint(1, 1), FIntRect(0, 0, 1, 1));
+	if (bCanResurrectHistory)
 	{
-		PrevHistorySRVs = DummyHistorySRVs;
+		const FViewMatrices& InViewMatrices = View.ViewMatrices;
+		const FViewMatrices& InPrevViewMatrices = InputHistory.ViewMatrices[ResurrectionFrameSliceIndex];
+
+		FVector DeltaTranslation = InPrevViewMatrices.GetPreViewTranslation() - InViewMatrices.GetPreViewTranslation();
+		FMatrix InvViewProj = InViewMatrices.ComputeInvProjectionNoAAMatrix() * InViewMatrices.GetTranslatedViewMatrix().GetTransposed();
+		FMatrix PrevViewProj = FTranslationMatrix(DeltaTranslation) * InPrevViewMatrices.GetTranslatedViewMatrix() * InPrevViewMatrices.ComputeProjectionNoAAMatrix();
+
+		ClipToResurrectionClip = FMatrix44f(InvViewProj * PrevViewProj);
+		ResurrectionGuideViewport = FScreenPassTextureViewport(PrevHistory.GuideArray->Desc.Extent, InputHistory.InputViewportRects[ResurrectionFrameSliceIndex]);
+	}
+
+	// Allocate non-reusuable history resources
+	if (SubpixelMethod == ETSRSubpixelMethod::ClosestDepth)
+	{
+		FRDGTextureDesc SubpixelDepthDesc = FRDGTextureDesc::Create2D(
+			InputExtent,
+			PF_R32_UINT,
+			FClearValueBinding::None,
+			TexCreate_ShaderResource | TexCreate_UAV);
+		History.SubpixelDepth = GraphBuilder.CreateTexture(SubpixelDepthDesc, TEXT("TSR.History.SubpixelDepth"));
 	}
 
 	// Setup the shader parameters for previous frame history
 	FTSRPrevHistoryParameters PrevHistoryParameters;
 	{
 		// Setup prev history parameters.
-		FScreenPassTextureViewport PrevHistoryViewport(PrevHistorySRVs.HighFrequency->Desc.Texture->Desc.Extent, InputHistory.OutputViewportRect);
+		FScreenPassTextureViewport PrevHistoryViewport(PrevHistory.MetadataArray->Desc.Extent, InputHistory.OutputViewportRect);
 		if (bCameraCut)
 		{
 			PrevHistoryViewport.Extent = FIntPoint(1, 1);
@@ -1374,6 +1576,7 @@ FDefaultTemporalUpscaler::FOutputs AddTemporalSuperResolutionPasses(
 		PrevHistoryParameters.ScreenPosToPrevHistoryBufferUV = FScreenTransform::ChangeTextureBasisFromTo(
 			PrevHistoryViewport, FScreenTransform::ETextureBasis::ScreenPosition, FScreenTransform::ETextureBasis::TextureUV);
 		PrevHistoryParameters.HistoryPreExposureCorrection = View.PreExposure / View.PrevViewInfo.SceneColorPreExposure;
+		PrevHistoryParameters.ResurrectionPreExposureCorrection = bCanResurrectHistory ? View.PreExposure / InputHistory.SceneColorPreExposures[ResurrectionFrameSliceIndex] : 1.0f;
 	}
 
 	// Clear atomic scattered texture.
@@ -1589,6 +1792,7 @@ FDefaultTemporalUpscaler::FOutputs AddTemporalSuperResolutionPasses(
 	FRDGTextureRef ReprojectedHistoryGuideTexture = nullptr;
 	FRDGTextureRef ReprojectedHistoryMoireTexture = nullptr;
 	FRDGTextureRef ParallaxRejectionMaskTexture = nullptr;
+	FRDGTextureRef HistoryResurrectionMaskTexture = nullptr;
 	FRDGTextureRef HoleFilledVelocityMaskTexture = nullptr;
 	{
 		FRDGTextureRef HoleFilledVelocityTexture;
@@ -1600,6 +1804,10 @@ FDefaultTemporalUpscaler::FOutputs AddTemporalSuperResolutionPasses(
 				/* InFlags = */ TexCreate_ShaderResource | TexCreate_UAV);
 
 			ParallaxRejectionMaskTexture = GraphBuilder.CreateTexture(Desc, TEXT("TSR.ParallaxRejectionMask"));
+			if (bCanResurrectHistory)
+			{
+				HistoryResurrectionMaskTexture = GraphBuilder.CreateTexture(Desc, TEXT("TSR.HistoryResurrectionMask"));
+			}
 
 			Desc.Format = PF_G16R16;
 			HoleFilledVelocityTexture = GraphBuilder.CreateTexture(Desc, TEXT("TSR.Velocity.HoleFilled"));
@@ -1609,21 +1817,23 @@ FDefaultTemporalUpscaler::FOutputs AddTemporalSuperResolutionPasses(
 		}
 
 		{
-			FRDGTextureDesc Desc = FRDGTextureDesc::Create2D(
+			FRDGTextureDesc Desc = FRDGTextureDesc::Create2DArray(
 				InputExtent,
-				History.Guide->Desc.Format,
+				History.GuideArray->Desc.Format,
 				FClearValueBinding::None,
-				/* InFlags = */ TexCreate_ShaderResource | TexCreate_UAV);
+				/* InFlags = */ TexCreate_ShaderResource | TexCreate_UAV,
+				/* InArraySize = */ bCanResurrectHistory ? 2 : 1);
 
 			ReprojectedHistoryGuideTexture = GraphBuilder.CreateTexture(Desc, TEXT("TSR.ReprojectedHistoryGuide"));
 		}
 
 		{
-			FRDGTextureDesc Desc = FRDGTextureDesc::Create2D(
+			FRDGTextureDesc Desc = FRDGTextureDesc::Create2DArray(
 				InputExtent,
-				History.Moire->Desc.Format,
+				History.MoireArray->Desc.Format,
 				FClearValueBinding::None,
-				/* InFlags = */ TexCreate_ShaderResource | TexCreate_UAV);
+				/* InFlags = */ TexCreate_ShaderResource | TexCreate_UAV,
+				/* InArraySize = */ bCanResurrectHistory ? 2 : 1);
 
 			ReprojectedHistoryMoireTexture = GraphBuilder.CreateTexture(Desc, TEXT("TSR.ReprojectedHistoryMoire"));
 		}
@@ -1631,7 +1841,6 @@ FDefaultTemporalUpscaler::FOutputs AddTemporalSuperResolutionPasses(
 		FTSRDecimateHistoryCS::FParameters* PassParameters = GraphBuilder.AllocParameters<FTSRDecimateHistoryCS::FParameters>();
 		PassParameters->CommonParameters = CommonParameters;
 		PassParameters->RotationalClipToPrevClip = RotationalClipToPrevClip;
-		PassParameters->CurrentToPrevInputJitter = InputHistory.PrevTemporalJitterPixels - CommonParameters.InputJitter;
 		{
 			float TanHalfFieldOfView = View.ViewMatrices.GetInvProjectionMatrix().M[0][0];
 
@@ -1648,33 +1857,49 @@ FDefaultTemporalUpscaler::FOutputs AddTemporalSuperResolutionPasses(
 		PassParameters->PrevHistoryParameters = PrevHistoryParameters;
 
 		{
-			FScreenPassTextureViewport PrevHistoryColorViewport(PrevHistorySRVs.Guide->Desc.Extent, InputHistory.InputViewportRect);
-			PassParameters->PrevHistoryGuide = PrevHistorySRVs.Guide;
-			PassParameters->PrevHistoryMoire = PrevHistorySRVs.Moire;
+			FScreenPassTextureViewport PrevHistoryColorViewport(PrevHistory.GuideArray->Desc.Extent, InputHistory.InputViewportRect);
+			PassParameters->PrevHistoryGuide = PrevHistory.GuideArray;
+			PassParameters->PrevHistoryMoire = PrevHistory.MoireArray;
 			PassParameters->PrevGuideInfo = GetScreenPassTextureViewportParameters(PrevHistoryColorViewport);
 			PassParameters->InputPixelPosToReprojectScreenPos = ((FScreenTransform::Identity - InputRect.Min + 0.5f) / InputRect.Size()) * FScreenTransform::ViewportUVToScreenPos;
 			PassParameters->ScreenPosToPrevHistoryGuideBufferUV = FScreenTransform::ChangeTextureBasisFromTo(
 				PrevHistoryColorViewport,
 				FScreenTransform::ETextureBasis::ScreenPosition,
 				FScreenTransform::ETextureBasis::TextureUV);
-			PassParameters->HistoryGuideQuantizationError = ComputePixelFormatQuantizationError(PrevHistorySRVs.Guide->Desc.Format);
+			PassParameters->ScreenPosToResurrectionGuideBufferUV = FScreenTransform::ChangeTextureBasisFromTo(
+				ResurrectionGuideViewport,
+				FScreenTransform::ETextureBasis::ScreenPosition,
+				FScreenTransform::ETextureBasis::TextureUV);
+			PassParameters->ResurrectionGuideUVViewportBilinearMin = GetScreenPassTextureViewportParameters(ResurrectionGuideViewport).UVViewportBilinearMin;
+			PassParameters->ResurrectionGuideUVViewportBilinearMax = GetScreenPassTextureViewportParameters(ResurrectionGuideViewport).UVViewportBilinearMax;
+			PassParameters->HistoryGuideQuantizationError = ComputePixelFormatQuantizationError(PrevHistory.GuideArray->Desc.Format);
 		}
 
+		PassParameters->ResurrectionFrameIndex = ResurrectionFrameSliceIndex;
+		PassParameters->PrevFrameIndex = PrevFrameSliceIndex;
+		PassParameters->ClipToResurrectionClip = ClipToResurrectionClip;
+
 		PassParameters->ReprojectedHistoryGuideOutput = GraphBuilder.CreateUAV(ReprojectedHistoryGuideTexture);
 		PassParameters->ReprojectedHistoryMoireOutput = GraphBuilder.CreateUAV(ReprojectedHistoryMoireTexture);
 		PassParameters->HoleFilledVelocityOutput = GraphBuilder.CreateUAV(HoleFilledVelocityTexture);
 		PassParameters->HoleFilledVelocityMaskOutput = GraphBuilder.CreateUAV(HoleFilledVelocityMaskTexture);
 		PassParameters->ParallaxRejectionMaskOutput = GraphBuilder.CreateUAV(ParallaxRejectionMaskTexture);
+		if (bCanResurrectHistory)
+		{
+			PassParameters->HistoryResurrectionMaskOutput = GraphBuilder.CreateUAV(HistoryResurrectionMaskTexture);
+		}
 		PassParameters->DebugOutput = CreateDebugUAV(InputExtent, TEXT("Debug.TSR.DecimateHistory"));
 
 		FTSRDecimateHistoryCS::FPermutationDomain PermutationVector;
 		PermutationVector.Set<FTSRDecimateHistoryCS::FMoireReprojectionDim>(FlickeringFramePeriod > 0.0f);
+		PermutationVector.Set<FTSRDecimateHistoryCS::FResurrectionReprojectionDim>(bCanResurrectHistory);
 
 		TShaderMapRef<FTSRDecimateHistoryCS> ComputeShader(View.ShaderMap, PermutationVector);
 		FComputeShaderUtils::AddPass(
 			GraphBuilder,
-			RDG_EVENT_NAME("TSR DecimateHistory(%s) %dx%d",
+			RDG_EVENT_NAME("TSR DecimateHistory(%s%s) %dx%d",
 				PermutationVector.Get<FTSRDecimateHistoryCS::FMoireReprojectionDim>() ? TEXT("ReprojectMoire") : TEXT(""),
+				PermutationVector.Get<FTSRDecimateHistoryCS::FResurrectionReprojectionDim>() ? TEXT(" ReprojectResurrection") : TEXT(""),
 				InputRect.Width(), InputRect.Height()),
 			AsyncComputePasses >= 2 ? ERDGPassFlags::AsyncCompute : ERDGPassFlags::Compute,
 			ComputeShader,
@@ -1684,12 +1909,18 @@ FDefaultTemporalUpscaler::FOutputs AddTemporalSuperResolutionPasses(
 		DilatedVelocityTexture = HoleFilledVelocityTexture;
 	}
 
-	// Reject the history with frequency decomposition.
-	FRDGTextureRef HistoryRejectionTexture;
+	// Perform a history reject the history.
 	FRDGTextureRef InputSceneColorLdrLumaTexture = nullptr;
+	auto RejectReprojectedGuide = [&](
+		const TCHAR* UsageString,
+		int32 HistoryReprojectionIndex,
+		FRDGTextureUAVRef HistoryGuideOutput,
+		FRDGTextureUAVRef HistoryMoireOutput) -> FRDGTextureRef
 	{
 		const int32 GroupTileSize = 16;
 
+		const bool ComputeLdrLuma = InputSceneColorLdrLumaTexture == nullptr;
+		if (ComputeLdrLuma)
 		{
 			FRDGTextureDesc Desc = FRDGTextureDesc::Create2D(
 				InputExtent,
@@ -1700,6 +1931,7 @@ FDefaultTemporalUpscaler::FOutputs AddTemporalSuperResolutionPasses(
 			InputSceneColorLdrLumaTexture = GraphBuilder.CreateTexture(Desc, TEXT("TSR.SceneColorLdrLuma"));
 		}
 
+		FRDGTextureRef HistoryRejectionTexture;
 		{
 			FRDGTextureDesc Desc = FRDGTextureDesc::Create2D(
 				InputExtent,
@@ -1720,8 +1952,8 @@ FDefaultTemporalUpscaler::FOutputs AddTemporalSuperResolutionPasses(
 			FScreenTransform::ChangeTextureBasisFromTo(TranslucencyViewport, FScreenTransform::ETextureBasis::ViewportUV, FScreenTransform::ETextureBasis::TextureUV);
 		PassParameters->TranslucencyTextureUVMin = GetScreenPassTextureViewportParameters(TranslucencyViewport).UVViewportBilinearMin;
 		PassParameters->TranslucencyTextureUVMax = GetScreenPassTextureViewportParameters(TranslucencyViewport).UVViewportBilinearMax;
-		PassParameters->HistoryGuideQuantizationError = ComputePixelFormatQuantizationError(History.Guide->Desc.Format);
-		PassParameters->FlickeringFramePeriod = FlickeringFramePeriod;
+		PassParameters->HistoryGuideQuantizationError = ComputePixelFormatQuantizationError(History.GuideArray->Desc.Format);
+		PassParameters->FlickeringFramePeriod = HistoryMoireOutput ? FlickeringFramePeriod : 0.0f;
 		PassParameters->TheoricBlendFactor = 1.0f / (1.0f + MaxHistorySampleCount / OutputToInputResolutionFractionSquare);
 		PassParameters->TileOverscan = FMath::Clamp(CVarTSRShadingTileOverscan.GetValueOnRenderThread(), 2, GroupTileSize / 2 - 1);
 		PassParameters->PerceptionAdd = FMath::Pow(0.5f, CVarTSRShadingExposureOffset.GetValueOnRenderThread());
@@ -1736,22 +1968,20 @@ FDefaultTemporalUpscaler::FOutputs AddTemporalSuperResolutionPasses(
 		{
 			PassParameters->InputMoireLumaTexture = BlackDummy;
 		}
-		if (bAccumulateTranslucencySeparately)
-		{
-			PassParameters->InputSceneTranslucencyTexture = BlackAlphaOneDummy;
-		}
-		else
-		{
-			PassParameters->InputSceneTranslucencyTexture = SeparateTranslucencyTexture;
-		}
-		PassParameters->ReprojectedHistoryGuideTexture = ReprojectedHistoryGuideTexture;
-		PassParameters->ReprojectedHistoryMoireTexture = ReprojectedHistoryMoireTexture;
+		PassParameters->InputSceneTranslucencyTexture = SeparateTranslucencyTexture;
+		PassParameters->ReprojectedHistoryGuideTexture = GraphBuilder.CreateSRV(FRDGTextureSRVDesc::CreateForSlice(
+			ReprojectedHistoryGuideTexture, HistoryReprojectionIndex));
+		PassParameters->ReprojectedHistoryMoireTexture = ReprojectedHistoryMoireTexture ? GraphBuilder.CreateSRV(FRDGTextureSRVDesc::CreateForSlice(
+			ReprojectedHistoryMoireTexture, HistoryReprojectionIndex)) : nullptr;
 		PassParameters->ParallaxRejectionMaskTexture = ParallaxRejectionMaskTexture;
 		PassParameters->IsMovingMaskTexture = IsMovingMaskTexture;
 
-		PassParameters->HistoryGuideOutput = GraphBuilder.CreateUAV(History.Guide);
-		PassParameters->HistoryMoireOutput = GraphBuilder.CreateUAV(History.Moire);
+		PassParameters->HistoryGuideOutput = HistoryGuideOutput;
+		PassParameters->HistoryMoireOutput = HistoryMoireOutput;
 		PassParameters->HistoryRejectionOutput = GraphBuilder.CreateUAV(HistoryRejectionTexture);
+		//PassParameters->InputSceneColorLdrLumaOutput = ComputeLdrLuma
+		//	? GraphBuilder.CreateUAV(InputSceneColorLdrLumaTexture)
+		//	: CreateDummyUAV(GraphBuilder, PF_R8);
 		PassParameters->InputSceneColorLdrLumaOutput = GraphBuilder.CreateUAV(InputSceneColorLdrLumaTexture);
 		PassParameters->DebugOutput = CreateDebugUAV(InputExtent, TEXT("Debug.TSR.RejectShading"));
 
@@ -1759,25 +1989,136 @@ FDefaultTemporalUpscaler::FOutputs AddTemporalSuperResolutionPasses(
 
 		FTSRRejectShadingCS::FPermutationDomain PermutationVector;
 		PermutationVector.Set<FTSRRejectShadingCS::FWaveSizeOps>(bUseWaveOps && GRHIMinimumWaveSize >= 32 && GRHIMinimumWaveSize <= 64 ? GRHIMinimumWaveSize : 0);
-		PermutationVector.Set<FTSRRejectShadingCS::FFlickeringDetectionDim>(FlickeringFramePeriod > 0.0f);
+		PermutationVector.Set<FTSRRejectShadingCS::FFlickeringDetectionDim>(FlickeringFramePeriod > 0.0f && HistoryMoireOutput);
 		PermutationVector.Set<FTSRShader::F16BitVALUDim>(bUse16BitVALU);
 		PermutationVector = FTSRRejectShadingCS::RemapPermutation(PermutationVector);
 
 		TShaderMapRef<FTSRRejectShadingCS> ComputeShader(View.ShaderMap, PermutationVector);
 		FComputeShaderUtils::AddPass(
 			GraphBuilder,
-			RDG_EVENT_NAME("TSR RejectShading(TileSize=%d PaddingCostMultiplier=%1.1f WaveSize=%d FlickeringFramePeriod=%f %s%s) %dx%d",
+			RDG_EVENT_NAME("TSR RejectShading(%s TileSize=%d PaddingCostMultiplier=%1.1f WaveSize=%d FlickeringFramePeriod=%f %s%s) %dx%d",
+				UsageString,
 				TileSize,
 				FMath::Pow(float(GroupTileSize) / float(TileSize), 2),
 				int32(PermutationVector.Get<FTSRRejectShadingCS::FWaveSizeOps>()),
-				FlickeringFramePeriod,
+				PassParameters->FlickeringFramePeriod,
 				PermutationVector.Get<FTSRShader::F16BitVALUDim>() ? TEXT(" 16bit") : TEXT(""),
-				!bAccumulateTranslucencySeparately ? TEXT(" ComposeTranslucency") : TEXT(""),
+				TEXT(" ComposeTranslucency"),
 				InputRect.Width(), InputRect.Height()),
 			AsyncComputePasses >= 3 ? ERDGPassFlags::AsyncCompute : ERDGPassFlags::Compute,
 			ComputeShader,
 			PassParameters,
 			FComputeShaderUtils::GetGroupCount(InputRect.Size(), TileSize));
+
+		return HistoryRejectionTexture;
+	}; // RejectReprojectedGuide();
+
+	// Rejects frame
+	FRDGTextureRef HistoryRejectionTexture = nullptr;
+	{
+		FRDGTextureUAVRef HistoryGuideOutput;
+		FRDGTextureUAVRef HistoryMoireOutput = nullptr;
+
+		if (View.bStatePrevViewInfoIsReadOnly)
+		{
+			HistoryGuideOutput = CreateDummyUAV(GraphBuilder, History.GuideArray->Desc.Format);
+			HistoryMoireOutput = CreateDummyUAV(GraphBuilder, History.MoireArray->Desc.Format);
+		}
+		else
+		{
+			FRDGTextureUAVDesc GuideUAVDesc(History.GuideArray);
+			GuideUAVDesc.FirstArraySlice = CurrentFrameSliceIndex;
+			GuideUAVDesc.NumArraySlices = 1;
+			GuideUAVDesc.DimensionOverride = ETextureDimension::Texture2D;
+
+			FRDGTextureUAVDesc MoireUAVDesc(History.MoireArray);
+			MoireUAVDesc.FirstArraySlice = CurrentFrameSliceIndex;
+			MoireUAVDesc.NumArraySlices = 1;
+			MoireUAVDesc.DimensionOverride = ETextureDimension::Texture2D;
+
+			HistoryGuideOutput = GraphBuilder.CreateUAV(GuideUAVDesc);
+
+			if (FlickeringFramePeriod > 0.0f)
+			{
+				HistoryMoireOutput = GraphBuilder.CreateUAV(MoireUAVDesc);
+			}
+		}
+
+		if (bCanResurrectHistory)
+		{
+			FRDGTextureDesc GuideDesc = FRDGTextureDesc::Create2D(
+				History.GuideArray->Desc.Extent,
+				History.GuideArray->Desc.Format,
+				FClearValueBinding::None,
+				/* InFlags = */ TexCreate_ShaderResource | TexCreate_UAV);
+
+			FRDGTextureDesc MoireDesc = FRDGTextureDesc::Create2D(
+				History.MoireArray->Desc.Extent,
+				History.MoireArray->Desc.Format,
+				FClearValueBinding::None,
+				/* InFlags = */ TexCreate_ShaderResource | TexCreate_UAV);
+
+			FRDGTextureRef PrevGuideTexture = GraphBuilder.CreateTexture(GuideDesc, TEXT("TSR.Prev.Guide"));
+			FRDGTextureRef PrevMoireTexture = GraphBuilder.CreateTexture(GuideDesc, TEXT("TSR.Prev.Moire"));
+
+			FRDGTextureRef ResurrectedGuideTexture = GraphBuilder.CreateTexture(GuideDesc, TEXT("TSR.Resurrected.Guide"));
+			FRDGTextureRef ResurrectedMoireTexture = GraphBuilder.CreateTexture(GuideDesc, TEXT("TSR.Resurrected.Moire"));
+
+			FRDGTextureRef PrevRejectionTexture = RejectReprojectedGuide(
+				TEXT("PrevFrame"),
+				/* HistoryReprojectionIndex = */ 0,
+				/* out */ GraphBuilder.CreateUAV(PrevGuideTexture),
+				/* out */ bResurrectFlickeringHistory ? GraphBuilder.CreateUAV(PrevMoireTexture) : HistoryMoireOutput);
+
+			FRDGTextureRef ResurrectedRejectionTexture = RejectReprojectedGuide(
+				TEXT("Resurrection"),
+				/* HistoryReprojectionIndex = */ 1,
+				/* out */ GraphBuilder.CreateUAV(ResurrectedGuideTexture),
+				/* out */ bResurrectFlickeringHistory ? GraphBuilder.CreateUAV(ResurrectedMoireTexture) : nullptr);
+
+			{
+				HistoryRejectionTexture = GraphBuilder.CreateTexture(
+					PrevRejectionTexture->Desc, TEXT("TSR.HistoryRejection"));
+
+				FTSRMergeRejectionCS::FParameters* PassParameters = GraphBuilder.AllocParameters<FTSRMergeRejectionCS::FParameters>();
+				PassParameters->CommonParameters = CommonParameters;
+				PassParameters->ClipToResurrectionClip = ClipToResurrectionClip;
+
+				PassParameters->ClosestDepthTexture = ClosestDepthTexture;
+				PassParameters->PrevGuideTexture     = PrevGuideTexture;
+				PassParameters->PrevMoireTexture     = bResurrectFlickeringHistory ? PrevMoireTexture : BlackDummy;
+				PassParameters->PrevRejectionTexture = PrevRejectionTexture;
+
+				PassParameters->ResurrectedGuideTexture     = ResurrectedGuideTexture;
+				PassParameters->ResurrectedMoireTexture     = bResurrectFlickeringHistory ? ResurrectedMoireTexture : BlackDummy;
+				PassParameters->ResurrectedRejectionTexture = ResurrectedRejectionTexture;
+
+				PassParameters->HistoryGuideOutput = HistoryGuideOutput;
+				PassParameters->HistoryMoireOutput = bResurrectFlickeringHistory ? HistoryMoireOutput : CreateDummyUAV(GraphBuilder, History.MoireArray->Desc.Format);
+				PassParameters->HistoryRejectionOutput = GraphBuilder.CreateUAV(HistoryRejectionTexture);
+				PassParameters->HistoryResurrectionMaskOutput = GraphBuilder.CreateUAV(HistoryResurrectionMaskTexture);
+				PassParameters->DilatedVelocityOutput = GraphBuilder.CreateUAV(DilatedVelocityTexture);
+
+				PassParameters->DebugOutput = CreateDebugUAV(InputExtent, TEXT("Debug.TSR.MergeRejection"));
+
+				TShaderMapRef<FTSRMergeRejectionCS> ComputeShader(View.ShaderMap);
+				FComputeShaderUtils::AddPass(
+					GraphBuilder,
+					RDG_EVENT_NAME("TSR MergeRejection %dx%d", InputRect.Width(), InputRect.Height()),
+					AsyncComputePasses >= 3 ? ERDGPassFlags::AsyncCompute : ERDGPassFlags::Compute,
+					ComputeShader,
+					PassParameters,
+					FComputeShaderUtils::GetGroupCount(InputRect.Size(), 8));
+			}
+		}
+		else
+		{
+			HistoryRejectionTexture = RejectReprojectedGuide(
+				TEXT("PrevFrame"),
+				/* HistoryReprojectionIndex = */ 0,
+				/* out */ HistoryGuideOutput,
+				/* out */ HistoryMoireOutput);
+		}
 	}
 
 	// Spatial anti-aliasing when doing history rejection.
@@ -1843,6 +2184,7 @@ FDefaultTemporalUpscaler::FOutputs AddTemporalSuperResolutionPasses(
 	}
 
 	// Update temporal history.
+	FRDGTextureSRVRef UpdateHistoryTextureSRV = nullptr;
 	{
 		static const TCHAR* const kUpdateQualityNames[] = {
 			TEXT("Low"),
@@ -1862,6 +2204,7 @@ FDefaultTemporalUpscaler::FOutputs AddTemporalSuperResolutionPasses(
 
 		PassParameters->DilatedVelocityTexture = DilatedVelocityTexture;
 		PassParameters->ParallaxRejectionMaskTexture = ParallaxRejectionMaskTexture;
+		PassParameters->HistoryResurrectionMaskTexture = bCanResurrectHistory ? HistoryResurrectionMaskTexture : BlackDummy;
 		PassParameters->AntiAliasingTexture = AntiAliasingTexture;
 		PassParameters->NoiseFilteringTexture = NoiseFilteringTexture;
 		PassParameters->HoleFilledVelocityMaskTexture = HoleFilledVelocityMaskTexture;
@@ -1873,7 +2216,7 @@ FDefaultTemporalUpscaler::FOutputs AddTemporalSuperResolutionPasses(
 		PassParameters->HistoryPixelPosToScreenPos = HistoryPixelPosToViewportUV * FScreenTransform::ViewportUVToScreenPos;
 		PassParameters->HistoryPixelPosToInputPPCo = HistoryPixelPosToViewportUV * CommonParameters.InputInfo.ViewportSize + CommonParameters.InputJitter + CommonParameters.InputPixelPosMin;
 		PassParameters->HistoryPixelPosToTranslucencyPPCo = HistoryPixelPosToViewportUV * SeparateTranslucencyRect.Size() + CommonParameters.InputJitter * SeparateTranslucencyRect.Size() / CommonParameters.InputInfo.ViewportSize + SeparateTranslucencyRect.Min;
-		PassParameters->HistoryQuantizationError = ComputePixelFormatQuantizationError((History.ColorArray ? History.ColorArray : History.Output)->Desc.Format);
+		PassParameters->HistoryQuantizationError = ComputePixelFormatQuantizationError(HistoryColorFormat);
 
 		// All parameters to control the sample count in history.
 		PassParameters->HistorySampleCount = MaxHistorySampleCount / OutputToHistoryResolutionFractionSquare;
@@ -1890,16 +2233,88 @@ FDefaultTemporalUpscaler::FOutputs AddTemporalSuperResolutionPasses(
 		PassParameters->bGenerateOutputMip2 = false;
 		PassParameters->bHasSeparateTranslucency = bHasSeparateTranslucency;
 
+		PassParameters->HistoryArrayIndices = HistoryArrayIndices;
 		PassParameters->PrevHistoryParameters = PrevHistoryParameters;
-		PassParameters->PrevHistory = PrevHistorySRVs;
+		if (bCameraCut)
+		{
+			PassParameters->ResurrectionFrameIndex = 0;
+			PassParameters->PrevFrameIndex = 0;
+
+			PassParameters->PrevHistoryColorTexture = GraphBuilder.CreateSRV(BlackArrayDummy);
+			PassParameters->PrevHistoryMetadataTexture = GraphBuilder.CreateSRV(BlackArrayDummy);
+		}
+		else
+		{
+			FRDGTextureSRVDesc PrevColorSRVDesc(PrevHistory.ColorArray);
+			PrevColorSRVDesc.NumMipLevels = 1;
+
+			FRDGTextureSRVDesc PrevMetadataSRVDesc(PrevHistory.MetadataArray);
+			PrevMetadataSRVDesc.NumMipLevels = 1;
+
+			if (bCanResurrectHistory)
+			{
+				FRHIRange16 SliceRange = PrevHistorySliceSequence.GetSRVSliceRange(CurrentFrameSliceIndex, PrevFrameSliceIndex);
+				check(SliceRange.IsInRange(ResurrectionFrameSliceIndex));
+				check(SliceRange.IsInRange(PrevFrameSliceIndex));
+				check(!SliceRange.IsInRange(CurrentFrameSliceIndex));
+
+				PrevColorSRVDesc.FirstArraySlice = SliceRange.First;
+				PrevColorSRVDesc.NumArraySlices = SliceRange.Num;
+
+				PrevMetadataSRVDesc.FirstArraySlice = SliceRange.First;
+				PrevMetadataSRVDesc.NumArraySlices = SliceRange.Num;
+			}
+
+			PassParameters->ResurrectionFrameIndex = ResurrectionFrameSliceIndex - PrevColorSRVDesc.FirstArraySlice;
+			PassParameters->PrevFrameIndex = PrevFrameSliceIndex - PrevColorSRVDesc.FirstArraySlice;
+
+			PassParameters->PrevHistoryColorTexture = GraphBuilder.CreateSRV(PrevColorSRVDesc);
+			PassParameters->PrevHistoryMetadataTexture = GraphBuilder.CreateSRV(PrevMetadataSRVDesc);
+		}
+
+		{
+			FRDGTextureUAVDesc ColorUAVDesc(History.ColorArray);
+			ColorUAVDesc.FirstArraySlice = CurrentFrameSliceIndex;
+			ColorUAVDesc.NumArraySlices = 1;
+
+			FRDGTextureUAVDesc MetadataUAVDesc(History.MetadataArray);
+			MetadataUAVDesc.FirstArraySlice = CurrentFrameSliceIndex;
+			MetadataUAVDesc.NumArraySlices = 1;
+
+			PassParameters->HistoryArrayIndices = HistoryArrayIndices;
+			PassParameters->HistoryColorOutput = GraphBuilder.CreateUAV(ColorUAVDesc);
+			PassParameters->HistoryMetadataOutput = GraphBuilder.CreateUAV(MetadataUAVDesc);
+
+			if (bIsOutputDifferentThanHighFrequency)
+			{
+				PassParameters->SceneColorOutputMip0 = GraphBuilder.CreateUAV(FRDGTextureUAVDesc(UpdateHistoryOutputTexture, /* InMipLevel = */ 0));
+				UpdateHistoryTextureSRV = GraphBuilder.CreateSRV(FRDGTextureSRVDesc(UpdateHistoryOutputTexture));
+			}
+			else
+			{
+				PassParameters->SceneColorOutputMip0 = CreateDummyUAV(GraphBuilder, PF_FloatR11G11B10);
+				UpdateHistoryTextureSRV = GraphBuilder.CreateSRV(FRDGTextureSRVDesc::CreateForSlice(
+					History.ColorArray, ColorUAVDesc.FirstArraySlice + HistoryArrayIndices.HighFrequency));
+			}
+		}
 
-		PassParameters->HistoryOutput = CreateUAVs(GraphBuilder, History);
-			PassParameters->SceneColorOutputMip0 = GraphBuilder.CreateUAV(FRDGTextureUAVDesc(UpdateHistoryOutputTexture, /* InMipLevel = */ 0));
-		
 		if (PassInputs.bGenerateOutputMip1 && HistorySize == OutputRect.Size())
 		{
 			PassParameters->bGenerateOutputMip1 = true;
-			PassParameters->SceneColorOutputMip1 = GraphBuilder.CreateUAV(FRDGTextureUAVDesc(UpdateHistoryOutputTexture, /* InMipLevel = */ 1));
+			if (bIsOutputDifferentThanHighFrequency)
+			{
+				PassParameters->SceneColorOutputMip1 = GraphBuilder.CreateUAV(FRDGTextureUAVDesc(UpdateHistoryOutputTexture, /* InMipLevel = */ 1));
+			}
+			else
+			{
+				FRDGTextureUAVDesc Mip1Desc(History.ColorArray);
+				Mip1Desc.MipLevel = 1;
+				Mip1Desc.FirstArraySlice = UpdateHistoryTextureSRV->Desc.FirstArraySlice;
+				Mip1Desc.NumArraySlices = 1;
+				Mip1Desc.DimensionOverride = ETextureDimension::Texture2D;
+
+				PassParameters->SceneColorOutputMip1 = GraphBuilder.CreateUAV(Mip1Desc);
+			}
 		}
 		else if (SceneColorOutputHalfResTexture && HistorySize == OutputRect.Size())
 		{
@@ -1919,17 +2334,15 @@ FDefaultTemporalUpscaler::FOutputs AddTemporalSuperResolutionPasses(
 
 		FTSRUpdateHistoryCS::FPermutationDomain PermutationVector;
 		PermutationVector.Set<FTSRUpdateHistoryCS::FQualityDim>(UpdateHistoryQuality);
-		PermutationVector.Set<FTSRUpdateHistoryCS::FSeparateTranslucencyDim>(bAccumulateTranslucencySeparately);
 		PermutationVector.Set<FTSRShader::F16BitVALUDim>(bUse16BitVALU);
 
 		TShaderMapRef<FTSRUpdateHistoryCS> ComputeShader(View.ShaderMap, PermutationVector);
 		FComputeShaderUtils::AddPass(
 			GraphBuilder,
-			RDG_EVENT_NAME("TSR UpdateHistory(Quality=%s%s%s%s%s) %dx%d",
+			RDG_EVENT_NAME("TSR UpdateHistory(Quality=%s%s%s%s) %dx%d",
 				kUpdateQualityNames[int32(PermutationVector.Get<FTSRUpdateHistoryCS::FQualityDim>())],
 				PermutationVector.Get<FTSRShader::F16BitVALUDim>() ? TEXT(" 16bit") : TEXT(""),
-				PermutationVector.Get<FTSRUpdateHistoryCS::FSeparateTranslucencyDim>() ? TEXT(" SeparateTranslucency") : TEXT(""),
-				(History.ColorArray ? History.ColorArray : History.Output)->Desc.Format == PF_FloatR11G11B10 ? TEXT(" R11G11B10") : TEXT(""),
+				HistoryColorFormat == PF_FloatR11G11B10 ? TEXT(" R11G11B10") : TEXT(""),
 				PassParameters->bGenerateOutputMip2 ? TEXT(" OutputMip2") : (PassParameters->bGenerateOutputMip1 ? TEXT(" OutputMip1") : TEXT("")),
 				HistorySize.X, HistorySize.Y),
 			AsyncComputePasses >= 3 ? ERDGPassFlags::AsyncCompute : ERDGPassFlags::Compute,
@@ -1939,6 +2352,7 @@ FDefaultTemporalUpscaler::FOutputs AddTemporalSuperResolutionPasses(
 	}
 
 	// If we upscaled the history buffer, downsize back to the secondary screen percentage size.
+	FRDGTextureSRVRef SceneColorOutputTextureSRV = UpdateHistoryTextureSRV;
 	if (HistorySize != OutputRect.Size())
 	{
 		check(!SceneColorOutputQuarterResTexture);
@@ -1957,7 +2371,7 @@ FDefaultTemporalUpscaler::FOutputs AddTemporalSuperResolutionPasses(
 		PassParameters->bGenerateOutputMip1 = false;
 		PassParameters->HistoryValidityMultiply = float(HistorySize.X * HistorySize.Y) / float(OutputRect.Width() * OutputRect.Height());
 
-		PassParameters->UpdateHistoryOutputTexture = UpdateHistoryOutputTexture;
+		PassParameters->UpdateHistoryOutputTexture = UpdateHistoryTextureSRV;
 		
 		PassParameters->SceneColorOutputMip0 = GraphBuilder.CreateUAV(FRDGTextureUAVDesc(SceneColorOutputTexture, /* InMipLevel = */ 0));
 		if (PassInputs.bGenerateOutputMip1)
@@ -1992,6 +2406,8 @@ FDefaultTemporalUpscaler::FOutputs AddTemporalSuperResolutionPasses(
 			ComputeShader,
 			PassParameters,
 			FComputeShaderUtils::GetGroupCount(OutputRect.Size(), PermutationVector.Get<FTSRResolveHistoryCS::FNyquistDim>() ? 6 : 8));
+		
+		SceneColorOutputTextureSRV = GraphBuilder.CreateSRV(FRDGTextureSRVDesc(SceneColorOutputTexture));
 	}
 
 	// Extract all resources for next frame.
@@ -2001,24 +2417,29 @@ FDefaultTemporalUpscaler::FOutputs AddTemporalSuperResolutionPasses(
 		OutputHistory.InputViewportRect = InputRect;
 		OutputHistory.OutputViewportRect = FIntRect(FIntPoint(0, 0), HistorySize);
 		OutputHistory.FormatBit = uint32(HistoryFormatBits);
-		OutputHistory.PrevSceneColorPreExposure = View.PrevViewInfo.SceneColorPreExposure;
-		OutputHistory.PrevTemporalJitterPixels = CommonParameters.InputJitter;
-		OutputHistory.PrevOutputViewportRect = InputHistory.OutputViewportRect;
-
-		// Extract filterable history
-		if (!bIsOutputDifferentThanHighFrequency)
+		OutputHistory.FrameStorageCount     = HistorySliceSequence.FrameStorageCount;
+		OutputHistory.FrameStoragePeriod    = HistorySliceSequence.FrameStoragePeriod;
+		OutputHistory.AccumulatedFrameCount = bCameraCutResurrection ? 1 : FMath::Min(InputHistory.AccumulatedFrameCount + 1, HistorySliceSequence.GetRollingIndexCount());
+		OutputHistory.LastFrameRollingIndex = CurrentFrameRollingIndex;
+		if (bCameraCutResurrection)
 		{
-			GraphBuilder.QueueTextureExtraction(History.Output, &OutputHistory.Output);
+			OutputHistory.ViewMatrices.SetNum(OutputHistory.FrameStorageCount);
+			OutputHistory.SceneColorPreExposures.SetNum(OutputHistory.FrameStorageCount);
+			OutputHistory.InputViewportRects.SetNum(OutputHistory.FrameStorageCount);
 		}
-		if (History.ColorArray)
-		{
-			GraphBuilder.QueueTextureExtraction(History.ColorArray, &OutputHistory.ColorArray);
-		}
-		GraphBuilder.QueueTextureExtraction(History.Metadata, &OutputHistory.Metadata);
-		if (bAccumulateTranslucencySeparately)
+		else
 		{
-			GraphBuilder.QueueTextureExtraction(History.TranslucencyAlpha, &OutputHistory.TranslucencyAlpha);
+			OutputHistory.ViewMatrices = InputHistory.ViewMatrices;
+			OutputHistory.SceneColorPreExposures = InputHistory.SceneColorPreExposures;
+			OutputHistory.InputViewportRects = InputHistory.InputViewportRects;
 		}
+		OutputHistory.ViewMatrices[CurrentFrameSliceIndex] = View.ViewMatrices;
+		OutputHistory.SceneColorPreExposures[CurrentFrameSliceIndex] = View.PreExposure;
+		OutputHistory.InputViewportRects[CurrentFrameSliceIndex] = InputRect;
+
+		// Extract filterable history
+		GraphBuilder.QueueTextureExtraction(History.ColorArray, &OutputHistory.ColorArray);
+		GraphBuilder.QueueTextureExtraction(History.MetadataArray, &OutputHistory.MetadataArray);
 
 		// Extract non-filterable history
 		if (History.SubpixelDepth)
@@ -2027,27 +2448,27 @@ FDefaultTemporalUpscaler::FOutputs AddTemporalSuperResolutionPasses(
 		}
 
 		// Extract history guide
-		GraphBuilder.QueueTextureExtraction(History.Guide, &OutputHistory.Guide);
+		GraphBuilder.QueueTextureExtraction(History.GuideArray, &OutputHistory.GuideArray);
 
 		if (FlickeringFramePeriod > 0.0f)
 		{
-			GraphBuilder.QueueTextureExtraction(History.Moire, &OutputHistory.Moire);
+			GraphBuilder.QueueTextureExtraction(History.MoireArray, &OutputHistory.MoireArray);
 		}
 
 		// Extract the output for next frame SSR so that separate translucency shows up in SSR.
 		{
 			// Output in TemporalAAHistory and not CustomSSR so Lumen can pick up ScreenSpaceRayTracingInput in priority to ensure consistent behavior between TAA and TSR.
 			GraphBuilder.QueueTextureExtraction(
-				SceneColorOutputTexture, &View.ViewState->PrevFrameViewInfo.TemporalAAHistory.RT[0]);
+				SceneColorOutputTextureSRV->Desc.Texture, &View.ViewState->PrevFrameViewInfo.TemporalAAHistory.RT[0]);
 			View.ViewState->PrevFrameViewInfo.TemporalAAHistory.ViewportRect = OutputRect;
 			View.ViewState->PrevFrameViewInfo.TemporalAAHistory.ReferenceBufferSize = OutputExtent;
+			View.ViewState->PrevFrameViewInfo.TemporalAAHistory.OutputSliceIndex = SceneColorOutputTextureSRV->Desc.FirstArraySlice;
 		}
 	}
 
 
 	FDefaultTemporalUpscaler::FOutputs Outputs;
-	Outputs.FullRes.TextureSRV = GraphBuilder.CreateSRV(FRDGTextureSRVDesc(SceneColorOutputTexture));
-	Outputs.FullRes.ViewRect = OutputRect;
+	Outputs.FullRes = FScreenPassTextureSlice(SceneColorOutputTextureSRV, OutputRect);
 	if (SceneColorOutputHalfResTexture)
 	{
 		Outputs.HalfRes.Texture = SceneColorOutputHalfResTexture;
diff --git a/Engine/Source/Runtime/Renderer/Private/RendererScene.cpp b/Engine/Source/Runtime/Renderer/Private/RendererScene.cpp
--- a/Engine/Source/Runtime/Renderer/Private/RendererScene.cpp
+++ b/Engine/Source/Runtime/Renderer/Private/RendererScene.cpp
@@ -849,13 +849,11 @@ uint64 FTemporalAAHistory::GetGPUSizeBytes(bool bLogSizes) const
 uint64 FTSRHistory::GetGPUSizeBytes(bool bLogSizes) const
 {
 	return
-		GetRenderTargetGPUSizeBytes(Output, bLogSizes) +
 		GetRenderTargetGPUSizeBytes(ColorArray, bLogSizes) +
-		GetRenderTargetGPUSizeBytes(Metadata, bLogSizes) +
-		GetRenderTargetGPUSizeBytes(TranslucencyAlpha, bLogSizes) +
+		GetRenderTargetGPUSizeBytes(MetadataArray, bLogSizes) +
 		GetRenderTargetGPUSizeBytes(SubpixelDepth, bLogSizes) +
-		GetRenderTargetGPUSizeBytes(Guide, bLogSizes) +
-		GetRenderTargetGPUSizeBytes(Moire, bLogSizes);
+		GetRenderTargetGPUSizeBytes(GuideArray, bLogSizes) +
+		GetRenderTargetGPUSizeBytes(MoireArray, bLogSizes);
 }
 
 uint64 FScreenSpaceDenoiserHistory::GetGPUSizeBytes(bool bLogSizes) const
diff --git a/Engine/Source/Runtime/Renderer/Private/SceneRendering.h b/Engine/Source/Runtime/Renderer/Private/SceneRendering.h
--- a/Engine/Source/Runtime/Renderer/Private/SceneRendering.h
+++ b/Engine/Source/Runtime/Renderer/Private/SceneRendering.h
@@ -994,28 +994,32 @@ struct FTemporalAAHistory
 struct FTSRHistory
 {
 	// Output resolution.
-	TRefCountPtr<IPooledRenderTarget> Output;
 	TRefCountPtr<IPooledRenderTarget> ColorArray;
-	TRefCountPtr<IPooledRenderTarget> Metadata;
-	TRefCountPtr<IPooledRenderTarget> TranslucencyAlpha;
+	TRefCountPtr<IPooledRenderTarget> MetadataArray;
 
 	// Input resolution representation of the output
 	TRefCountPtr<IPooledRenderTarget> SubpixelDepth;
-	TRefCountPtr<IPooledRenderTarget> Guide;
-	TRefCountPtr<IPooledRenderTarget> Moire;
+	TRefCountPtr<IPooledRenderTarget> GuideArray;
+	TRefCountPtr<IPooledRenderTarget> MoireArray;
 
 	// Frame's input and output resolution.
 	FIntRect InputViewportRect;
 	FIntRect OutputViewportRect;
 
-	// Previous frame's informations.
-	FIntRect PrevOutputViewportRect;
-	FVector2f PrevTemporalJitterPixels;
-	float PrevSceneColorPreExposure = 1.0f;
-
 	// Format of the history for auto camera cut when setting change.
 	uint32 FormatBit = 0;
 
+	// Number of frame in history.
+	int32 FrameStorageCount = 1;
+	int32 FrameStoragePeriod = 1;
+	int32 AccumulatedFrameCount = 1;
+	int32 LastFrameRollingIndex = 0;
+
+	// All the information the previous frames for resurrection.
+	TArray<FViewMatrices> ViewMatrices;
+	TArray<float> SceneColorPreExposures;
+	TArray<FIntRect> InputViewportRects;
+
 
 	void SafeRelease()
 	{
@@ -1024,7 +1028,7 @@ struct FTSRHistory
 
 	bool IsValid() const
 	{
-		return Metadata.IsValid();
+		return MetadataArray.IsValid();
 	}
 
 	uint64 GetGPUSizeBytes(bool bLogSizes) const;
