Change FInstanceCullingContext::CreateDeferredContext() to take InstanceCullingManager by reference instead of pointer since it's required to be non-null

* Fixes static analysis warning

diff --git a/Engine/Source/Runtime/Renderer/Private/InstanceCulling/InstanceCullingContext.cpp b/Engine/Source/Runtime/Renderer/Private/InstanceCulling/InstanceCullingContext.cpp
--- a/Engine/Source/Runtime/Renderer/Private/InstanceCulling/InstanceCullingContext.cpp
+++ b/Engine/Source/Runtime/Renderer/Private/InstanceCulling/InstanceCullingContext.cpp
@@ -883,9 +883,8 @@ FORCEINLINE int32 GetArrayDataSize(const TArray<DataType, AllocatorType>& Array)
 FInstanceCullingDeferredContext *FInstanceCullingContext::CreateDeferredContext(
 	FRDGBuilder& GraphBuilder,
 	FGPUScene& GPUScene,
-	FInstanceCullingManager* InstanceCullingManager)
+	FInstanceCullingManager& InstanceCullingManager)
 {
-	check(InstanceCullingManager != nullptr);
 #define INST_CULL_CALLBACK_MODE(CustomCode) \
 	[PassParameters, DeferredContext, Mode]() \
 	{ \
@@ -918,7 +917,7 @@ FInstanceCullingDeferredContext *FInstanceCullingContext::CreateDeferredContext(
 
 	const ERHIFeatureLevel::Type FeatureLevel = GPUScene.GetFeatureLevel();
 
-	FInstanceCullingDeferredContext* DeferredContext = GraphBuilder.AllocObject<FInstanceCullingDeferredContext>(FeatureLevel, InstanceCullingManager);
+	FInstanceCullingDeferredContext* DeferredContext = GraphBuilder.AllocObject<FInstanceCullingDeferredContext>(FeatureLevel, &InstanceCullingManager);
 
 	const bool bCullInstances = CVarCullInstances.GetValueOnRenderThread() != 0;
 	const bool bAllowWPODisable = true;
@@ -1025,8 +1024,8 @@ FInstanceCullingDeferredContext *FInstanceCullingContext::CreateDeferredContext(
 	PassParametersTmp.InstanceIdOffsetBuffer = GraphBuilder.CreateSRV(InstanceIdOffsetBuffer, PF_R32_UINT);	
 	if (bCullInstances || bAllowWPODisable)
 	{
-		PassParametersTmp.InViews = GraphBuilder.CreateSRV(InstanceCullingManager->CullingIntermediate.CullingViews);
-		PassParametersTmp.NumCullingViews = InstanceCullingManager->CullingIntermediate.NumViews;
+		PassParametersTmp.InViews = GraphBuilder.CreateSRV(InstanceCullingManager.CullingIntermediate.CullingViews);
+		PassParametersTmp.NumCullingViews = InstanceCullingManager.CullingIntermediate.NumViews;
 	}
 
 	// Compaction parameters
@@ -1034,9 +1033,9 @@ FInstanceCullingDeferredContext *FInstanceCullingContext::CreateDeferredContext(
 	PassParametersTmp.CompactInstanceIdsBufferOut = CompactInstanceIdsUAV;
 	PassParametersTmp.CompactionBlockCounts = CompactionBlockCountsUAV;
 
-	if (InstanceCullingManager && InstanceCullingManager->InstanceOcclusionQueryBuffer)
+	if (InstanceCullingManager.InstanceOcclusionQueryBuffer)
 	{
-		PassParametersTmp.InstanceOcclusionQueryBuffer = GraphBuilder.CreateSRV(InstanceCullingManager->InstanceOcclusionQueryBuffer, PF_R32_UINT);
+		PassParametersTmp.InstanceOcclusionQueryBuffer = GraphBuilder.CreateSRV(InstanceCullingManager.InstanceOcclusionQueryBuffer, PF_R32_UINT);
 	}
 	else
 	{
@@ -1045,7 +1044,7 @@ FInstanceCullingDeferredContext *FInstanceCullingContext::CreateDeferredContext(
 	}
 
 	// Record the number of culling views to be able to check that no views referencing out-of bounds views are queued up
-	DeferredContext->NumCullingViews = InstanceCullingManager->CullingIntermediate.NumViews;
+	DeferredContext->NumCullingViews = InstanceCullingManager.CullingIntermediate.NumViews;
 
 	for (uint32 Mode = 0U; Mode < uint32(EBatchProcessingMode::Num); ++Mode)
 	{
diff --git a/Engine/Source/Runtime/Renderer/Private/InstanceCulling/InstanceCullingManager.cpp b/Engine/Source/Runtime/Renderer/Private/InstanceCulling/InstanceCullingManager.cpp
--- a/Engine/Source/Runtime/Renderer/Private/InstanceCulling/InstanceCullingManager.cpp
+++ b/Engine/Source/Runtime/Renderer/Private/InstanceCulling/InstanceCullingManager.cpp
@@ -93,5 +93,5 @@ void FInstanceCullingManager::BeginDeferredCulling(FRDGBuilder& GraphBuilder, FG
 		return;
 	}
 
-	DeferredContext = FInstanceCullingContext::CreateDeferredContext(GraphBuilder, GPUScene, this);
+	DeferredContext = FInstanceCullingContext::CreateDeferredContext(GraphBuilder, GPUScene, *this);
 }
diff --git a/Engine/Source/Runtime/Renderer/Public/InstanceCulling/InstanceCullingContext.h b/Engine/Source/Runtime/Renderer/Public/InstanceCulling/InstanceCullingContext.h
--- a/Engine/Source/Runtime/Renderer/Public/InstanceCulling/InstanceCullingContext.h
+++ b/Engine/Source/Runtime/Renderer/Public/InstanceCulling/InstanceCullingContext.h
@@ -179,10 +179,10 @@ public:
 	 * Batches are collected as we walk through the main render setup and are executed when RDG Execute or Drain is called.
 	 * This implicitly ends the deferred context, so if Drain is used, it should be paired with a new call to BeginDeferredCulling.
 	 */
-	static FInstanceCullingDeferredContext *CreateDeferredContext(
+	static FInstanceCullingDeferredContext* CreateDeferredContext(
 		FRDGBuilder& GraphBuilder,
 		FGPUScene& GPUScene,
-		FInstanceCullingManager* InstanceCullingManager);
+		FInstanceCullingManager& InstanceCullingManager);
 
 	/**
 	 * Helper function to add a pass to zero the instance count in the indirect args.
