Path Tracer: Skydome meshes should be included for raytracing when using the path tracer

This changes ignores the cvar "r.RayTracing.ExcludeSky" when the path tracer is active.

diff --git a/Engine/Source/Runtime/Renderer/Private/DeferredShadingRenderer.cpp b/Engine/Source/Runtime/Renderer/Private/DeferredShadingRenderer.cpp
--- a/Engine/Source/Runtime/Renderer/Private/DeferredShadingRenderer.cpp
+++ b/Engine/Source/Runtime/Renderer/Private/DeferredShadingRenderer.cpp
@@ -240,7 +240,7 @@ static int32 GRayTracingExcludeSky = 1;
 static FAutoConsoleVariableRef CRayTracingExcludeSky(
 	TEXT("r.RayTracing.ExcludeSky"),
 	GRayTracingExcludeSky,
-	TEXT("A toggle that controls inclusion of sky geometry in the ray tracing scene (excluding sky can make ray tracing faster).\n")
+	TEXT("A toggle that controls inclusion of sky geometry in the ray tracing scene (excluding sky can make ray tracing faster). This setting is ignored for the Path Tracer.\n")
 	TEXT(" 0: Sky objects included in the ray tracing scene\n")
 	TEXT(" 1: Sky objects excluded from the ray tracing scene (default)"),
 	ECVF_RenderThreadSafe);
@@ -1429,6 +1429,7 @@ bool FDeferredShadingSceneRenderer::GatherRayTracingWorldInstancesForView(FRDGBu
 		const FScene& Scene;
 		TChunkedArray<FRayTracingRelevantPrimitive>& RelevantStaticPrimitives;
 		const FRayTracingCullingParameters& CullingParameters;
+		const bool bIsPathTracing;
 
 		// Outputs
 
@@ -1439,11 +1440,13 @@ bool FDeferredShadingSceneRenderer::GatherRayTracingWorldInstancesForView(FRDGBu
 		FRayTracingSceneAddInstancesTask(const FScene& InScene,
 											TChunkedArray<FRayTracingRelevantPrimitive>& InRelevantStaticPrimitives,
 											const FRayTracingCullingParameters& InCullingParameters,
+											const bool bInIsPathTracing,
 											FRayTracingScene& InRayTracingScene, TArray<FVisibleRayTracingMeshCommand>& InVisibleRayTracingMeshCommands,
 											TArray<FPrimitiveSceneProxy*>& InProxiesWithDirtyCachedInstance)
 			: Scene(InScene)
 			, RelevantStaticPrimitives(InRelevantStaticPrimitives)
 			, CullingParameters(InCullingParameters)
+			, bIsPathTracing(bInIsPathTracing)
 			, RayTracingScene(InRayTracingScene)
 			, VisibleRayTracingMeshCommands(InVisibleRayTracingMeshCommands)
 			, ProxiesWithDirtyCachedInstance(InProxiesWithDirtyCachedInstance)
@@ -1664,7 +1667,7 @@ bool FDeferredShadingSceneRenderer::GatherRayTracingWorldInstancesForView(FRDGBu
 
 					if ((GRayTracingExcludeDecals && RelevantPrimitive.bAnySegmentsDecal)
 						|| (GRayTracingExcludeTranslucent && RelevantPrimitive.bAllSegmentsTranslucent)
-						|| (GRayTracingExcludeSky && RelevantPrimitive.bIsSky))
+						|| (GRayTracingExcludeSky && RelevantPrimitive.bIsSky && !bIsPathTracing))
 					{
 						continue;
 					}
@@ -1794,7 +1797,7 @@ bool FDeferredShadingSceneRenderer::GatherRayTracingWorldInstancesForView(FRDGBu
 	AddInstancesTaskPrerequisites.Add(RelevantPrimitiveList.StaticPrimitiveLODTask);
 
 	FGraphEventRef AddInstancesTask = TGraphTask<FRayTracingSceneAddInstancesTask>::CreateTask(&AddInstancesTaskPrerequisites).ConstructAndDispatchWhenReady(
-		*Scene, RelevantPrimitiveList.StaticPrimitives, View.RayTracingCullingParameters, // inputs 
+		*Scene, RelevantPrimitiveList.StaticPrimitives, View.RayTracingCullingParameters, bool(View.Family->EngineShowFlags.PathTracing), // inputs 
 		RayTracingScene, View.VisibleRayTracingMeshCommands, View.ProxiesWithDirtyCachedInstance // outputs
 	);
 
