Fixes missing output in the VisualizeTemporalUpscaler show flag

diff --git a/Engine/Source/Runtime/Renderer/Private/PostProcess/VisualizeTemporalUpscaler.cpp b/Engine/Source/Runtime/Renderer/Private/PostProcess/VisualizeTemporalUpscaler.cpp
--- a/Engine/Source/Runtime/Renderer/Private/PostProcess/VisualizeTemporalUpscaler.cpp
+++ b/Engine/Source/Runtime/Renderer/Private/PostProcess/VisualizeTemporalUpscaler.cpp
@@ -141,7 +141,7 @@ FScreenPassTexture AddVisualizeTemporalUpscalerPass(FRDGBuilder& GraphBuilder, c
 		{
 			FVisualizeBufferTile& Tile = Tiles[4 * 3 + 3];
 			Tile.Input.Texture = OutputTexture.Texture;
-			Tile.Input.ViewRect = CropViewRectToCenter(Tile.Input.ViewRect);
+			Tile.Input.ViewRect = CropViewRectToCenter(OutputTexture.ViewRect);
 			Tile.Label = VisualizeTextureLabel(OutputTexture.Texture);
 		}
 
