ShaderPrint - added r.ShaderPrint.DrawOccludedLines (default=1) to disable drawing lines behind geometry on demand. Useful for complex visualizations, where drawing occluded lines makes it too messy and hard to understand.

diff --git a/Engine/Shaders/Private/ShaderPrintDrawPrimitive.usf b/Engine/Shaders/Private/ShaderPrintDrawPrimitive.usf
--- a/Engine/Shaders/Private/ShaderPrintDrawPrimitive.usf
+++ b/Engine/Shaders/Private/ShaderPrintDrawPrimitive.usf
@@ -6,6 +6,7 @@
 
 float3 TranslatedWorldOffsetConversion;
 uint bCheckerboardEnabled;
+uint bDrawOccludedLines;
 
 // The shader used to render the debug element on view
 StructuredBuffer<uint> ShaderDrawDebugPrimitive;
@@ -102,7 +103,7 @@ void ShaderDrawDebugPS(
 
 	const float ColorScale = bIsHidden ? 0.4f : 1;
 	const float CheckerboardScale = ((PixelCoord.x%2) ^ (PixelCoord.y%2)) > 0 ? 1 : 0.0;
-	const float OcclusionScale = (bCheckerboardEnabled && bIsHidden) ? CheckerboardScale : 1;
+	const float OcclusionScale = (bCheckerboardEnabled && bIsHidden) ? (bDrawOccludedLines ? CheckerboardScale : 0.0f) : 1;
 	OutSceneColor = OcclusionScale * ColorAlpha * float4(ColorAlpha.aaa * ColorScale, 1.0); // Pre multiplied alpha with color scale
 }
 #endif
diff --git a/Engine/Source/Runtime/Renderer/Private/ShaderPrint.cpp b/Engine/Source/Runtime/Renderer/Private/ShaderPrint.cpp
--- a/Engine/Source/Runtime/Renderer/Private/ShaderPrint.cpp
+++ b/Engine/Source/Runtime/Renderer/Private/ShaderPrint.cpp
@@ -77,6 +77,11 @@ namespace ShaderPrint
 		TEXT("Lock the line drawing.\n"),
 		ECVF_Cheat | ECVF_RenderThreadSafe);
 
+	static TAutoConsoleVariable<int32> CVarDrawOccludedLines(
+		TEXT("r.ShaderPrint.DrawOccludedLines"),
+		1,
+		TEXT("Whether to draw occluded lines using checkboarding and lower opacity.\n"),
+		ECVF_Cheat | ECVF_RenderThreadSafe);
 
 	static TAutoConsoleVariable<int32> CVarDrawZoomEnable(
 		TEXT("r.ShaderPrint.Zoom"),
@@ -530,6 +535,7 @@ namespace ShaderPrint
 			SHADER_PARAMETER(FVector2f, OriginalViewSize)
 			SHADER_PARAMETER(FVector2f, OriginalBufferInvSize)
 			SHADER_PARAMETER(uint32, bCheckerboardEnabled)
+			SHADER_PARAMETER(uint32, bDrawOccludedLines)
 			SHADER_PARAMETER_RDG_TEXTURE(Texture2D, DepthTexture)
 			SHADER_PARAMETER_SAMPLER(SamplerState, DepthSampler)
 			RENDER_TARGET_BINDING_SLOTS()
@@ -965,6 +971,7 @@ namespace ShaderPrint
 		PassParameters->PS.DepthTexture = DepthTexture;
 		PassParameters->PS.DepthSampler = TStaticSamplerState<SF_Point, AM_Clamp, AM_Clamp, AM_Clamp>::GetRHI();
 		PassParameters->PS.bCheckerboardEnabled = bLines ? 1u : 0u;
+		PassParameters->PS.bDrawOccludedLines = CVarDrawOccludedLines.GetValueOnRenderThread() != 0 ? 1 : 0;
 		PassParameters->VS.ShaderDrawDebugPrimitive = GraphBuilder.CreateSRV(ShaderPrintPrimitiveBuffer);
 		PassParameters->VS.IndirectBuffer = IndirectBuffer;
 		PassParameters->VS.Common = ShaderPrintData.UniformBuffer;
