Adds sg.AntiAliasingQuality in the TSR's main draw events for better discoverability to users that its costs can be controled with that very same cvar.

diff --git a/Engine/Source/Runtime/Renderer/Private/PostProcess/TemporalSuperResolution.cpp b/Engine/Source/Runtime/Renderer/Private/PostProcess/TemporalSuperResolution.cpp
--- a/Engine/Source/Runtime/Renderer/Private/PostProcess/TemporalSuperResolution.cpp
+++ b/Engine/Source/Runtime/Renderer/Private/PostProcess/TemporalSuperResolution.cpp
@@ -1270,8 +1270,12 @@ FDefaultTemporalUpscaler::FOutputs AddTemporalSuperResolutionPasses(
 		ETSRHistoryFormatBits(InputHistory.FormatBit) != HistoryFormatBits ||
 		false;
 
-	RDG_EVENT_SCOPE(GraphBuilder, "TemporalSuperResolution(%s) %dx%d -> %dx%d",
-		bSupportsAlpha ? TEXT("Alpha") : TEXT(""),
+	static auto CVarAntiAliasingQuality = IConsoleManager::Get().FindConsoleVariable(TEXT("sg.AntiAliasingQuality"));
+	check(CVarAntiAliasingQuality);
+
+	RDG_EVENT_SCOPE(GraphBuilder, "TemporalSuperResolution(sg.AntiAliasingQuality=%d%s) %dx%d -> %dx%d",
+		CVarAntiAliasingQuality->GetInt(),
+		bSupportsAlpha ? TEXT(" Alpha") : TEXT(""),
 		InputRect.Width(), InputRect.Height(),
 		OutputRect.Width(), OutputRect.Height());
 	RDG_GPU_STAT_SCOPE(GraphBuilder, TemporalSuperResolution);
