Remove GetMobileVariableRateShadingImage to make VRS behaviour consistent across mobile and desktop.

diff --git a/Engine/Source/Runtime/Renderer/Private/MobileShadingRenderer.cpp b/Engine/Source/Runtime/Renderer/Private/MobileShadingRenderer.cpp
--- a/Engine/Source/Runtime/Renderer/Private/MobileShadingRenderer.cpp
+++ b/Engine/Source/Runtime/Renderer/Private/MobileShadingRenderer.cpp
@@ -1209,7 +1209,8 @@ void FMobileSceneRenderer::RenderForward(FRDGBuilder& GraphBuilder, FRDGTextureR
 		SceneDepth = SceneTextures.Depth.Target;
 	}
 
-	TRefCountPtr<IPooledRenderTarget> ShadingRateTarget = GVRSImageManager.GetMobileVariableRateShadingImage(ViewFamily);
+	GVRSImageManager.PrepareImageBasedVRS(GraphBuilder, ViewFamily, SceneTextures);
+	FRDGTextureRef NewShadingRateTarget = GVRSImageManager.GetVariableRateShadingImage(GraphBuilder, MainView, FVariableRateShadingImageManager::EVRSPassType::BasePass);
 
 	FRenderTargetBindingSlots BasePassRenderTargets;
 	BasePassRenderTargets[0] = FRenderTargetBinding(SceneColor, SceneColorResolve, ERenderTargetLoadAction::EClear);
@@ -1220,7 +1221,7 @@ void FMobileSceneRenderer::RenderForward(FRDGBuilder& GraphBuilder, FRDGTextureR
 	BasePassRenderTargets.DepthStencil = bIsFullDepthPrepassEnabled ? 
 		FDepthStencilBinding(SceneDepth, ERenderTargetLoadAction::ELoad, ERenderTargetLoadAction::ELoad, FExclusiveDepthStencil::DepthRead_StencilWrite) : 
 		FDepthStencilBinding(SceneDepth, ERenderTargetLoadAction::EClear, ERenderTargetLoadAction::EClear, FExclusiveDepthStencil::DepthWrite_StencilWrite);
-	BasePassRenderTargets.ShadingRateTexture = (!MainView.bIsSceneCapture && !MainView.bIsReflectionCapture && ShadingRateTarget.IsValid()) ? RegisterExternalTexture(GraphBuilder, ShadingRateTarget->GetRHI(), TEXT("ShadingRateTexture")) : nullptr;
+	BasePassRenderTargets.ShadingRateTexture = (!MainView.bIsSceneCapture && !MainView.bIsReflectionCapture && (NewShadingRateTarget != nullptr)) ? NewShadingRateTarget : nullptr;
 	BasePassRenderTargets.SubpassHint = ESubpassHint::None;
 	BasePassRenderTargets.NumOcclusionQueries = 0u;
 
diff --git a/Engine/Source/Runtime/Renderer/Private/VariableRateShading/VariableRateShadingImageManager.cpp b/Engine/Source/Runtime/Renderer/Private/VariableRateShading/VariableRateShadingImageManager.cpp
--- a/Engine/Source/Runtime/Renderer/Private/VariableRateShading/VariableRateShadingImageManager.cpp
+++ b/Engine/Source/Runtime/Renderer/Private/VariableRateShading/VariableRateShadingImageManager.cpp
@@ -468,34 +468,6 @@ bool FVariableRateShadingImageManager::IsTypeEnabledForView(const FSceneView& Vi
 	return false;
 }
 
-TRefCountPtr<IPooledRenderTarget> FVariableRateShadingImageManager::GetMobileVariableRateShadingImage(const FSceneViewFamily& ViewFamily)
-{
-	if (!(IStereoRendering::IsStereoEyeView(*ViewFamily.Views[0]) && GEngine->XRSystem.IsValid()))
-	{
-		return TRefCountPtr<IPooledRenderTarget>();
-	}
-
-	FIntPoint Size(ViewFamily.RenderTarget->GetSizeXY());
-
-	const bool bStereo = GEngine->StereoRenderingDevice.IsValid() && GEngine->StereoRenderingDevice->IsStereoEnabled();
-	IStereoRenderTargetManager* const StereoRenderTargetManager = bStereo ? GEngine->StereoRenderingDevice->GetRenderTargetManager() : nullptr;
-
-	FTexture2DRHIRef Texture;
-	FIntPoint TextureSize(0, 0);
-
-	// Allocate variable resolution texture for VR foveation if supported
-	if (StereoRenderTargetManager && StereoRenderTargetManager->NeedReAllocateShadingRateTexture(MobileHMDFixedFoveationOverrideImage))
-	{
-		bool bAllocatedShadingRateTexture = StereoRenderTargetManager->AllocateShadingRateTexture(0, Size.X, Size.Y, GRHIVariableRateShadingImageFormat, 0, TexCreate_None, TexCreate_None, Texture, TextureSize);
-		if (bAllocatedShadingRateTexture)
-		{
-			MobileHMDFixedFoveationOverrideImage = CreateRenderTarget(Texture, TEXT("ShadingRate"));
-		}
-	}
-
-	return MobileHMDFixedFoveationOverrideImage;
-}
-
 void FVariableRateShadingImageManager::DrawDebugPreview(FRDGBuilder& GraphBuilder, const FSceneViewFamily& ViewFamily, FRDGTextureRef OutputSceneColor)
 {
 	uint32 ImageTypeAsInt = CVarVRSPreview.GetValueOnRenderThread();
diff --git a/Engine/Source/Runtime/Renderer/Public/VariableRateShadingImageManager.h b/Engine/Source/Runtime/Renderer/Public/VariableRateShadingImageManager.h
--- a/Engine/Source/Runtime/Renderer/Public/VariableRateShadingImageManager.h
+++ b/Engine/Source/Runtime/Renderer/Public/VariableRateShadingImageManager.h
@@ -98,13 +98,6 @@ public:
 	 */
 	bool IsTypeEnabledForView(const FSceneView& View, EVRSSourceType Type) const;
 
-
-	/** 
-	 *  Special case (currently) for mobile devices; right now this will return either no image, or a static fixed-foveation shading rate image
-	 *  provided by the stereo device runtime.
-	 */
-	TRefCountPtr<IPooledRenderTarget> GetMobileVariableRateShadingImage(const FSceneViewFamily& ViewFamily);
-
 	static bool IsVRSSupportedByRHI();
 	static bool IsVRSEnabled();
 	static bool IsVRSCompatibleWithView(const FViewInfo& View);
