Added async deletion of scene proxies.

diff --git a/Engine/Plugins/FX/Niagara/Source/Niagara/Private/NiagaraComponent.cpp b/Engine/Plugins/FX/Niagara/Source/Niagara/Private/NiagaraComponent.cpp
--- a/Engine/Plugins/FX/Niagara/Source/Niagara/Private/NiagaraComponent.cpp
+++ b/Engine/Plugins/FX/Niagara/Source/Niagara/Private/NiagaraComponent.cpp
@@ -217,8 +217,6 @@ void FNiagaraSceneProxy::DestroyRenderState_Concurrent()
 
 FNiagaraSceneProxy::~FNiagaraSceneProxy()
 {
-	check(IsInRenderingThread());
-
 	if (RenderData)
 	{
 		delete RenderData;
diff --git a/Engine/Plugins/FX/Niagara/Source/Niagara/Private/NiagaraRenderer.cpp b/Engine/Plugins/FX/Niagara/Source/Niagara/Private/NiagaraRenderer.cpp
--- a/Engine/Plugins/FX/Niagara/Source/Niagara/Private/NiagaraRenderer.cpp
+++ b/Engine/Plugins/FX/Niagara/Source/Niagara/Private/NiagaraRenderer.cpp
@@ -423,7 +423,6 @@ FPrimitiveViewRelevance FNiagaraRenderer::GetViewRelevance(const FSceneView* Vie
 
 void FNiagaraRenderer::SetDynamicData_RenderThread(FNiagaraDynamicDataBase* NewDynamicData)
 {
-	check(IsInRenderingThread());
 	if (DynamicDataRender)
 	{
 		delete DynamicDataRender;
diff --git a/Engine/Plugins/FX/Niagara/Source/Niagara/Private/NiagaraSystemRenderData.cpp b/Engine/Plugins/FX/Niagara/Source/Niagara/Private/NiagaraSystemRenderData.cpp
--- a/Engine/Plugins/FX/Niagara/Source/Niagara/Private/NiagaraSystemRenderData.cpp
+++ b/Engine/Plugins/FX/Niagara/Source/Niagara/Private/NiagaraSystemRenderData.cpp
@@ -73,8 +73,6 @@ void FNiagaraSystemRenderData::DestroyRenderState_Concurrent()
 
 void FNiagaraSystemRenderData::Destroy_RenderThread()
 {
-	check(IsInRenderingThread());
-
 	// Delete the renderers
 	for (auto Renderer : EmitterRenderers_RT)
 	{
diff --git a/Engine/Source/Runtime/Engine/Private/PrimitiveSceneProxy.cpp b/Engine/Source/Runtime/Engine/Private/PrimitiveSceneProxy.cpp
--- a/Engine/Source/Runtime/Engine/Private/PrimitiveSceneProxy.cpp
+++ b/Engine/Source/Runtime/Engine/Private/PrimitiveSceneProxy.cpp
@@ -532,7 +532,6 @@ void FPrimitiveSceneProxy::SetUsedMaterialForVerification(const TArray<UMaterial
 
 FPrimitiveSceneProxy::~FPrimitiveSceneProxy()
 {
-	check(IsInRenderingThread());
 }
 
 HHitProxy* FPrimitiveSceneProxy::CreateHitProxies(UPrimitiveComponent* Component,TArray<TRefCountPtr<HHitProxy> >& OutHitProxies)
diff --git a/Engine/Source/Runtime/Renderer/Private/RendererScene.cpp b/Engine/Source/Runtime/Renderer/Private/RendererScene.cpp
--- a/Engine/Source/Runtime/Renderer/Private/RendererScene.cpp
+++ b/Engine/Source/Runtime/Renderer/Private/RendererScene.cpp
@@ -2457,8 +2457,7 @@ void FScene::RemovePrimitiveSceneInfo_RenderThread(FPrimitiveSceneInfo* Primitiv
 			};
 
 			BeginCleanup(new DeferDeleteHitProxies(MoveTemp(PrimitiveSceneInfo->HitProxies)));
-			delete PrimitiveSceneInfo->Proxy;
-			delete PrimitiveSceneInfo;
+			PrimitiveSceneInfosToDelete.Emplace(PrimitiveSceneInfo);
 		}
 	}
 	else
@@ -5492,6 +5491,8 @@ void FScene::UpdateAllPrimitiveSceneInfos(FRDGBuilder& GraphBuilder, EUpdateAllP
 	TArray<int32> RemovedPrimitiveIndices;
 	RemovedPrimitiveIndices.SetNumUninitialized(RemovedLocalPrimitiveSceneInfos.Num());
 
+	bool bNeedPathTracedInvalidation = false;
+
 	GPUScene.ResizeDirtyState(Primitives.Num());
 	{
 		CSV_SCOPED_TIMING_STAT_EXCLUSIVE(RemovePrimitiveSceneInfos);
@@ -5701,6 +5702,7 @@ void FScene::UpdateAllPrimitiveSceneInfos(FRDGBuilder& GraphBuilder, EUpdateAllP
 				}
 #endif
 
+				bNeedPathTracedInvalidation = bNeedPathTracedInvalidation || IsPrimitiveRelevantToPathTracing(PrimitiveSceneInfo);
 				DeletedSceneInfos.Add(PrimitiveSceneInfo);
 
 				const int32 PersistentIndex = PrimitiveSceneInfo->PersistentIndex.Index;
@@ -5713,8 +5715,6 @@ void FScene::UpdateAllPrimitiveSceneInfos(FRDGBuilder& GraphBuilder, EUpdateAllP
 	
 	}
 
-	bool bNeedPathTracedInvalidation = false;
-
 	const int32 SceneInfosContainerReservedSize = AddedPrimitiveSceneInfos.Num() + UpdatedTransforms.Num() + UpdatedInstances.Num();
 
 	TArray<FPrimitiveSceneInfo*, SceneRenderingAllocator>& SceneInfosWithAddToScene = *GraphBuilder.AllocObject<TArray<FPrimitiveSceneInfo*, SceneRenderingAllocator>>();
@@ -6460,12 +6460,12 @@ void FScene::UpdateAllPrimitiveSceneInfos(FRDGBuilder& GraphBuilder, EUpdateAllP
 	// Need to do this here since we delete the proxy next 
 	// TODO: should refactor to add this as a dependency for a proxy-deletion task instead.
 	SceneCullingUpdater.GetAsyncProxyUseTaskHandle().Wait();
+
+	GraphBuilder.AddSetupTask([this, DeletedSceneInfos = MoveTemp(DeletedSceneInfos)]
 	{
 		SCOPED_NAMED_EVENT(FScene_DeletePrimitiveSceneInfo, FColor::Red);
 		for (FPrimitiveSceneInfo* PrimitiveSceneInfo : DeletedSceneInfos)
 		{
-			bNeedPathTracedInvalidation = bNeedPathTracedInvalidation || IsPrimitiveRelevantToPathTracing(PrimitiveSceneInfo);
-
 			// It is possible that the HitProxies list isn't empty if PrimitiveSceneInfo was Added/Removed in same frame
 			// Delete the PrimitiveSceneInfo on the game thread after the rendering thread has processed its removal.
 			// This must be done on the game thread because the hit proxy references (and possibly other members) need to be freed on the game thread.
@@ -6480,7 +6480,16 @@ void FScene::UpdateAllPrimitiveSceneInfos(FRDGBuilder& GraphBuilder, EUpdateAllP
 			delete PrimitiveSceneInfo->Proxy;
 			delete PrimitiveSceneInfo;
 		}
-	}
+
+		// Delete scene primitives that were removed without ever being added.
+		for (FPrimitiveSceneInfo* PrimitiveSceneInfo : PrimitiveSceneInfosToDelete)
+		{
+			delete PrimitiveSceneInfo->Proxy;
+			delete PrimitiveSceneInfo;
+		}
+		PrimitiveSceneInfosToDelete.Empty();
+	});
+
 	if (bNeedPathTracedInvalidation)
 	{
 		InvalidatePathTracedOutput();
diff --git a/Engine/Source/Runtime/Renderer/Private/ScenePrivate.h b/Engine/Source/Runtime/Renderer/Private/ScenePrivate.h
--- a/Engine/Source/Runtime/Renderer/Private/ScenePrivate.h
+++ b/Engine/Source/Runtime/Renderer/Private/ScenePrivate.h
@@ -3741,6 +3741,7 @@ private:
 	Experimental::TRobinHoodHashSet<FPrimitiveSceneInfo*> RemovedPrimitiveSceneInfos;
 	Experimental::TRobinHoodHashSet<FPrimitiveSceneInfo*> DistanceFieldSceneDataUpdates;
 	TArray<FLevelCommand> LevelCommands;
+	TArray<FPrimitiveSceneInfo*> PrimitiveSceneInfosToDelete;
 
 	UE::Tasks::FTask CreateLightPrimitiveInteractionsTask;
 	UE::Tasks::FTask CacheMeshDrawCommandsTask;
