Heterogeneous Volumes: Introduce per-volume step size based on local-space step factors.

Step size is now calculated in local space as a function of voxel width, projected along the ray direction.
Toggle global step-size values to -1 (disabled).
Change default MaxStepCount to 512 steps.

diff --git a/Engine/Shaders/Private/HeterogeneousVolumes/HeterogeneousVolumesLiveShadingPipeline.usf b/Engine/Shaders/Private/HeterogeneousVolumes/HeterogeneousVolumesLiveShadingPipeline.usf
--- a/Engine/Shaders/Private/HeterogeneousVolumes/HeterogeneousVolumesLiveShadingPipeline.usf
+++ b/Engine/Shaders/Private/HeterogeneousVolumes/HeterogeneousVolumesLiveShadingPipeline.usf
@@ -40,6 +40,9 @@ float3 LocalBoundsOrigin;
 float3 LocalBoundsExtent;
 int PrimitiveId;
 
+// Volume data
+int3 VoxelResolution;
+
 // Lighting
 int bApplyEmissionAndTransmittance;
 int bApplyDirectLighting;
@@ -50,11 +53,13 @@ float VolumetricScatteringIntensity;
 // Shadowing
 uint VirtualShadowMapId;
 float ShadowStepSize;
+float ShadowStepFactor;
 
 // Ray data
 float MaxTraceDistance;
 float MaxShadowTraceDistance;
 float StepSize;
+float StepFactor;
 int MaxStepCount;
 int bJitter;
 
@@ -118,6 +123,11 @@ float3 GetLocalBoundsExtent()
 	return LocalBoundsExtent;
 }
 
+int3 GetVolumeResolution()
+{
+	return VoxelResolution;
+}
+
 float GetStepSize()
 {
 	return StepSize;
@@ -128,6 +138,16 @@ float GetShadowStepSize()
 	return ShadowStepSize;
 }
 
+float GetStepFactor()
+{
+	return StepFactor;
+}
+
+float GetShadowStepFactor()
+{
+	return ShadowStepFactor;
+}
+
 float GetMaxTraceDistance()
 {
 	return MaxTraceDistance;
@@ -285,7 +305,7 @@ void RenderSingleScatteringWithLiveShadingCS(
 	// Intersect ray with bounding volume
 	// TODO: LWC integration..
 	float3 WorldRayOrigin = LWCHackToFloat(LWCSubtract(View.TranslatedWorldCameraOrigin, PrimaryView.PreViewTranslation));
-	float3 WorldRayEnd = LWCHackToFloat(SvPositionToWorld(float4(PixelCoord + 0.5, DeviceZ, 1)));
+	float3 WorldRayEnd = LWCHackToFloat(SvPositionToWorld(float4(PixelCoord + View.TemporalAAJitter.xy, DeviceZ, 1)));
 	float3 WorldRayDirection = WorldRayEnd - WorldRayOrigin;
 	float WorldRayLength = length(WorldRayDirection);
 	WorldRayDirection /= WorldRayLength;
@@ -299,7 +319,7 @@ void RenderSingleScatteringWithLiveShadingCS(
 	float3 LocalBoundsMin = LocalBoundsOrigin - LocalBoundsExtent;
 	float3 LocalBoundsMax = LocalBoundsOrigin + LocalBoundsExtent;
 
-	// Test bounding volume
+	// Intersect bounding volume
 	float2 HitT = IntersectAABB(LocalRayOrigin, LocalRayDirection, 0.0, LocalRayLength, LocalBoundsMin, LocalBoundsMax);
 	float HitSpan = HitT.y - HitT.x;
 	if (HitSpan > 0.0)
@@ -316,7 +336,7 @@ void RenderSingleScatteringWithLiveShadingCS(
 			WorldRayDirection,
 			// Ray-step attributes
 			Jitter,
-			StepSize,
+			CalcStepSize(LocalRayDirection),
 			MaxStepCount,
 			bApplyEmissionAndTransmittance,
 			bApplyDirectLighting,
diff --git a/Engine/Shaders/Private/HeterogeneousVolumes/HeterogeneousVolumesPreshadingPipeline.usf b/Engine/Shaders/Private/HeterogeneousVolumes/HeterogeneousVolumesPreshadingPipeline.usf
--- a/Engine/Shaders/Private/HeterogeneousVolumes/HeterogeneousVolumesPreshadingPipeline.usf
+++ b/Engine/Shaders/Private/HeterogeneousVolumes/HeterogeneousVolumesPreshadingPipeline.usf
@@ -507,9 +507,6 @@ void RenderSingleScatteringWithPreshadingCS(
 	uint GroupThreadId : SV_GroupThreadID
 )
 {
-	float3 Radiance = 0.0;
-	float3 Transmittance = 1.0;
-
 	// Convert 1D index to pixel coordinate
 	FRayMarchingTile RayMarchingTile = RayMarchingTilesBuffer[GroupId];
 	// TODO: Morton..
@@ -517,6 +514,9 @@ void RenderSingleScatteringWithPreshadingCS(
 	//float2 PixelCoord = View.ViewRectMin.xy + RayMarchingTile.PixelOffset + TileCoord;
 	float2 PixelCoord = RayMarchingTile.PixelOffset + TileCoord;
 
+	float3 Radiance = 0.0;
+	float3 Transmittance = 1.0 - RWLightingTexture[PixelCoord].a;
+
 	// Create screen ray
 	if (any(PixelCoord.xy >= View.ViewSizeAndInvSize.xy))
 	{
@@ -539,7 +539,7 @@ void RenderSingleScatteringWithPreshadingCS(
 	// Intersect ray with bounding volume
 	// TODO: LWC integration..
 	float3 WorldRayOrigin = LWCHackToFloat(LWCSubtract(View.TranslatedWorldCameraOrigin, PrimaryView.PreViewTranslation));
-	float3 WorldRayEnd = LWCHackToFloat(SvPositionToWorld(float4(PixelCoord + 0.5, DeviceZ, 1)));
+	float3 WorldRayEnd = LWCHackToFloat(SvPositionToWorld(float4(PixelCoord + View.TemporalAAJitter.xy, DeviceZ, 1)));
 	float3 WorldRayDirection = WorldRayEnd - WorldRayOrigin;
 	float WorldRayLength = length(WorldRayDirection);
 	WorldRayDirection /= WorldRayLength;
@@ -588,7 +588,7 @@ void RenderSingleScatteringWithPreshadingCS(
 			WorldRayDirection,
 			// Ray-step attributes
 			Jitter,
-			StepSize,
+			CalcStepSize(LocalRayDirection),
 			MaxStepCount,
 			bApplyEmissionAndTransmittance,
 			bApplyDirectLighting,
@@ -622,7 +622,8 @@ void RenderSingleScatteringWithPreshadingCS(
 
 	// Output..
 	float3 Opacity = bApplyEmissionAndTransmittance ? 1.0 - Transmittance : 0.0;
-	RWLightingTexture[PixelCoord] += float4(Radiance * View.PreExposure, Luminance(Opacity));
+	RWLightingTexture[PixelCoord].rgb += Radiance * View.PreExposure;
+	RWLightingTexture[PixelCoord].a = Luminance(Opacity);
 
 #if DIM_DEBUG
 	//if (!all(DispatchThreadId))
diff --git a/Engine/Shaders/Private/HeterogeneousVolumes/HeterogeneousVolumesRayMarchingUtils.ush b/Engine/Shaders/Private/HeterogeneousVolumes/HeterogeneousVolumesRayMarchingUtils.ush
--- a/Engine/Shaders/Private/HeterogeneousVolumes/HeterogeneousVolumesRayMarchingUtils.ush
+++ b/Engine/Shaders/Private/HeterogeneousVolumes/HeterogeneousVolumesRayMarchingUtils.ush
@@ -93,6 +93,33 @@ FRayMarchingContext CreateRayMarchingContext(
 	return Context;
 }
 
+float CalcStepSize(float3 LocalRayDirection, float LocalStepSize, float LocalStepFactor)
+{
+	if (LocalStepSize <= 0.0f)
+	{
+		// Intersect aligned voxel to determine candidate t-values and use the minimum distance
+		float3 VoxelSize = GetLocalBoundsExtent() / float3(GetVolumeResolution());
+		float3 tValues = abs(VoxelSize / LocalRayDirection);
+		float VoxelStepSize = min(tValues.x, min(tValues.y, tValues.z));
+		LocalStepSize = VoxelStepSize * LocalStepFactor;
+	}
+	return max(LocalStepSize, 0.01f);
+}
+
+float CalcStepSize(float3 LocalRayDirection)
+{
+	float LocalStepSize = GetStepSize();
+	float LocalStepFactor = GetStepFactor();
+	return CalcStepSize(LocalRayDirection, LocalStepSize, LocalStepFactor);
+}
+
+float CalcShadowStepSize(float3 LocalRayDirection)
+{
+	float LocalStepSize = GetShadowStepSize();
+	float LocalStepFactor = GetShadowStepFactor();
+	return CalcStepSize(LocalRayDirection, LocalStepSize, LocalStepFactor);
+}
+
 float RecalculateStepSize(inout FRayMarchingContext Context)
 {
 	// Recalculate step-size if the bounding-box diagonal is greater than the maximum step-count
@@ -239,7 +266,7 @@ float3 ComputeTransmittance(
 	LocalRayDirection /= LocalRayTMax;
 
 	float ShadowBias = 0.5;
-	float ShadowStepSize = GetShadowStepSize();
+	float ShadowStepSize = CalcShadowStepSize(LocalRayDirection);
 	int bApplyEmissionAndTransmittance = 0;
 	int bApplyDirectLighting = 0;
 	int bApplyShadowTransmittance = 0;
diff --git a/Engine/Shaders/Private/HeterogeneousVolumes/HeterogeneousVolumesSparseVoxelUniformBufferUtils.ush b/Engine/Shaders/Private/HeterogeneousVolumes/HeterogeneousVolumesSparseVoxelUniformBufferUtils.ush
--- a/Engine/Shaders/Private/HeterogeneousVolumes/HeterogeneousVolumesSparseVoxelUniformBufferUtils.ush
+++ b/Engine/Shaders/Private/HeterogeneousVolumes/HeterogeneousVolumesSparseVoxelUniformBufferUtils.ush
@@ -35,11 +35,21 @@ float GetStepSize()
 	return SparseVoxelUniformBuffer.StepSize;
 }
 
+float GetStepFactor()
+{
+	return SparseVoxelUniformBuffer.StepFactor;
+}
+
 float GetShadowStepSize()
 {
 	return SparseVoxelUniformBuffer.ShadowStepSize;
 }
 
+float GetShadowStepFactor()
+{
+	return SparseVoxelUniformBuffer.ShadowStepFactor;
+}
+
 uint3 GetVolumeResolution()
 {
 	return SparseVoxelUniformBuffer.VolumeResolution;
diff --git a/Engine/Source/Runtime/Renderer/Private/HeterogeneousVolumes/HeterogeneousVolumes.cpp b/Engine/Source/Runtime/Renderer/Private/HeterogeneousVolumes/HeterogeneousVolumes.cpp
--- a/Engine/Source/Runtime/Renderer/Private/HeterogeneousVolumes/HeterogeneousVolumes.cpp
+++ b/Engine/Source/Runtime/Renderer/Private/HeterogeneousVolumes/HeterogeneousVolumes.cpp
@@ -48,8 +48,8 @@ static TAutoConsoleVariable<int32> CVarHeterogeneousVolumesJitter(
 
 static TAutoConsoleVariable<int32> CVarHeterogeneousVolumesMaxStepCount(
 	TEXT("r.HeterogeneousVolumes.MaxStepCount"),
-	128,
-	TEXT("The maximum ray-marching step count (Default = 128)"),
+	512,
+	TEXT("The maximum ray-marching step count (Default = 512)"),
 	ECVF_RenderThreadSafe
 );
 
@@ -110,8 +110,8 @@ static TAutoConsoleVariable<int32> CVarHeterogeneousVolumesPreshadingVolumeResol
 
 static TAutoConsoleVariable<float> CVarHeterogeneousVolumesShadowStepSize(
 	TEXT("r.HeterogeneousVolumes.ShadowStepSize"),
-	8.0,
-	TEXT("The ray-marching step-size for shadow rays (Default = 8.0)"),
+	-1.0,
+	TEXT("The ray-marching step-size override for shadow rays (Default = -1.0, disabled)"),
 	ECVF_RenderThreadSafe
 );
 
@@ -145,8 +145,8 @@ static TAutoConsoleVariable<int32> CVarHeterogeneousVolumesSparseVoxelRefinement
 
 static TAutoConsoleVariable<float> CVarHeterogeneousVolumesStepSize(
 	TEXT("r.HeterogeneousVolumes.StepSize"),
-	1.0,
-	TEXT("The ray-marching step-size (Default = 1.0)"),
+	-1.0,
+	TEXT("The ray-marching step-size override (Default = -1.0, disabled)"),
 	ECVF_RenderThreadSafe
 );
 
diff --git a/Engine/Source/Runtime/Renderer/Private/HeterogeneousVolumes/HeterogeneousVolumes.h b/Engine/Source/Runtime/Renderer/Private/HeterogeneousVolumes/HeterogeneousVolumes.h
--- a/Engine/Source/Runtime/Renderer/Private/HeterogeneousVolumes/HeterogeneousVolumes.h
+++ b/Engine/Source/Runtime/Renderer/Private/HeterogeneousVolumes/HeterogeneousVolumes.h
@@ -112,7 +112,9 @@ BEGIN_UNIFORM_BUFFER_STRUCT(FSparseVoxelUniformBufferParameters, )
 	SHADER_PARAMETER(float, MaxTraceDistance)
 	SHADER_PARAMETER(float, MaxShadowTraceDistance)
 	SHADER_PARAMETER(float, StepSize)
+	SHADER_PARAMETER(float, StepFactor)
 	SHADER_PARAMETER(float, ShadowStepSize)
+	SHADER_PARAMETER(float, ShadowStepFactor)
 END_UNIFORM_BUFFER_STRUCT()
 
 BEGIN_SHADER_PARAMETER_STRUCT(FLightingCacheParameters, )
diff --git a/Engine/Source/Runtime/Renderer/Private/HeterogeneousVolumes/HeterogeneousVolumesLiveShadingPipeline.cpp b/Engine/Source/Runtime/Renderer/Private/HeterogeneousVolumes/HeterogeneousVolumesLiveShadingPipeline.cpp
--- a/Engine/Source/Runtime/Renderer/Private/HeterogeneousVolumes/HeterogeneousVolumesLiveShadingPipeline.cpp
+++ b/Engine/Source/Runtime/Renderer/Private/HeterogeneousVolumes/HeterogeneousVolumesLiveShadingPipeline.cpp
@@ -113,6 +113,7 @@ class FRenderLightingCacheWithLiveShadingCS : public FMeshMaterialShader
 
 		// Shadow data
 		SHADER_PARAMETER(float, ShadowStepSize)
+		SHADER_PARAMETER(float, ShadowStepFactor)
 		SHADER_PARAMETER_RDG_UNIFORM_BUFFER(FForwardLightData, ForwardLightData)
 		SHADER_PARAMETER_STRUCT_INCLUDE(FVolumeShadowingShaderParameters, VolumeShadowingShaderParameters)
 		SHADER_PARAMETER_STRUCT_INCLUDE(FVirtualShadowMapSamplingParameters, VirtualShadowMapSamplingParameters)
@@ -129,10 +130,12 @@ class FRenderLightingCacheWithLiveShadingCS : public FMeshMaterialShader
 		SHADER_PARAMETER(float, MaxTraceDistance)
 		SHADER_PARAMETER(float, MaxShadowTraceDistance)
 		SHADER_PARAMETER(float, StepSize)
+		SHADER_PARAMETER(float, StepFactor)
 		SHADER_PARAMETER(int, MaxStepCount)
 		SHADER_PARAMETER(int, bJitter)
 
 		// Volume data
+		SHADER_PARAMETER(FIntVector, VoxelResolution)
 		SHADER_PARAMETER_STRUCT_INCLUDE(FLightingCacheParameters, LightingCache)
 
 		// Output
@@ -228,6 +231,7 @@ class FRenderSingleScatteringWithLiveShadingCS : public FMeshMaterialShader
 
 		// Shadow data
 		SHADER_PARAMETER(float, ShadowStepSize)
+		SHADER_PARAMETER(float, ShadowStepFactor)
 		SHADER_PARAMETER_RDG_UNIFORM_BUFFER(FForwardLightData, ForwardLightData)
 		SHADER_PARAMETER_STRUCT_INCLUDE(FVolumeShadowingShaderParameters, VolumeShadowingShaderParameters)
 		SHADER_PARAMETER_STRUCT_INCLUDE(FVirtualShadowMapSamplingParameters, VirtualShadowMapSamplingParameters)
@@ -244,11 +248,13 @@ class FRenderSingleScatteringWithLiveShadingCS : public FMeshMaterialShader
 		SHADER_PARAMETER(int32, PrimitiveId)
 
 		// Volume data
+		SHADER_PARAMETER(FIntVector, VoxelResolution)
 		SHADER_PARAMETER_STRUCT_INCLUDE(FLightingCacheParameters, LightingCache)
 
 		// Ray data
 		SHADER_PARAMETER(float, MaxTraceDistance)
 		SHADER_PARAMETER(float, StepSize)
+		SHADER_PARAMETER(float, StepFactor)
 		SHADER_PARAMETER(int, MaxStepCount)
 		SHADER_PARAMETER(int, bJitter)
 
@@ -431,6 +437,7 @@ void RenderLightingCacheWithLiveShading(
 		PassParameters->PrimitiveId = PrimitiveId;
 
 		// Transmittance volume
+		PassParameters->VoxelResolution = HeterogeneousVolumeInterface->GetVoxelResolution();
 		PassParameters->LightingCache.LightingCacheResolution = HeterogeneousVolumes::GetLightingCacheResolution(HeterogeneousVolumeInterface);
 		//PassParameters->LightingCache.LightingCacheTexture = GraphBuilder.CreateSRV(LightingCacheTexture);
 		PassParameters->LightingCache.LightingCacheTexture = FRDGSystemTextures::Get(GraphBuilder).VolumetricBlack;
@@ -439,7 +446,9 @@ void RenderLightingCacheWithLiveShading(
 		PassParameters->MaxTraceDistance = HeterogeneousVolumes::GetMaxTraceDistance();
 		PassParameters->MaxShadowTraceDistance = HeterogeneousVolumes::GetMaxShadowTraceDistance();
 		PassParameters->StepSize = HeterogeneousVolumes::GetStepSize();
+		PassParameters->StepFactor = HeterogeneousVolumeInterface->GetStepFactor();
 		PassParameters->ShadowStepSize = HeterogeneousVolumes::GetShadowStepSize();
+		PassParameters->ShadowStepFactor = HeterogeneousVolumeInterface->GetShadowStepFactor();
 		PassParameters->MaxStepCount = HeterogeneousVolumes::GetMaxStepCount();
 		PassParameters->bJitter = HeterogeneousVolumes::ShouldJitter();
 
@@ -566,6 +575,7 @@ void RenderSingleScatteringWithLiveShading(
 		PassParameters->DeferredLight = DeferredLightUB;
 		PassParameters->LightType = LightType;
 		PassParameters->ShadowStepSize = HeterogeneousVolumes::GetShadowStepSize();
+		PassParameters->ShadowStepFactor = HeterogeneousVolumeInterface->GetShadowStepFactor();
 
 		// Object data
 		FMatrix44f LocalToWorld = FMatrix44f(HeterogeneousVolumeInterface->GetLocalToWorld());
@@ -575,9 +585,13 @@ void RenderSingleScatteringWithLiveShading(
 		PassParameters->LocalBoundsExtent = FVector3f(LocalBoxSphereBounds.BoxExtent);
 		PassParameters->PrimitiveId = PrimitiveId;
 
+		// Volume data
+		PassParameters->VoxelResolution = HeterogeneousVolumeInterface->GetVoxelResolution();
+
 		// Ray data
 		PassParameters->MaxTraceDistance = HeterogeneousVolumes::GetMaxTraceDistance();
 		PassParameters->StepSize = HeterogeneousVolumes::GetStepSize();
+		PassParameters->StepFactor = HeterogeneousVolumeInterface->GetStepFactor();
 		PassParameters->MaxStepCount = HeterogeneousVolumes::GetMaxStepCount();
 		PassParameters->bJitter = HeterogeneousVolumes::ShouldJitter();
 
diff --git a/Engine/Source/Runtime/Renderer/Private/HeterogeneousVolumes/HeterogeneousVolumesPreshadingPipeline.cpp b/Engine/Source/Runtime/Renderer/Private/HeterogeneousVolumes/HeterogeneousVolumesPreshadingPipeline.cpp
--- a/Engine/Source/Runtime/Renderer/Private/HeterogeneousVolumes/HeterogeneousVolumesPreshadingPipeline.cpp
+++ b/Engine/Source/Runtime/Renderer/Private/HeterogeneousVolumes/HeterogeneousVolumesPreshadingPipeline.cpp
@@ -1607,7 +1607,9 @@ void RenderWithPreshading(
 		SparseVoxelUniformBufferParameters->MaxTraceDistance = HeterogeneousVolumes::GetMaxTraceDistance();
 		SparseVoxelUniformBufferParameters->MaxShadowTraceDistance = HeterogeneousVolumes::GetMaxShadowTraceDistance();
 		SparseVoxelUniformBufferParameters->StepSize = HeterogeneousVolumes::GetStepSize();
+		SparseVoxelUniformBufferParameters->StepFactor = HeterogeneousVolumeInterface->GetStepFactor();
 		SparseVoxelUniformBufferParameters->ShadowStepSize = HeterogeneousVolumes::GetShadowStepSize();
+		SparseVoxelUniformBufferParameters->ShadowStepFactor = HeterogeneousVolumeInterface->GetShadowStepFactor();
 	}
 	TRDGUniformBufferRef<FSparseVoxelUniformBufferParameters> SparseVoxelUniformBuffer = GraphBuilder.CreateUniformBuffer(SparseVoxelUniformBufferParameters);
 
