Restores DFAO multi-view implementation (previously reverted in CL25963932)
Additional fix for indirect capsule shadows detected by devs in editor

Indirect capsule shadows overwrite the result of the screen grid DFAO, which results in two issues. The first is that the indirect capsule shadows changes the size of the BentNormalAOTexture, so the PIE renderer samples the texture at incorrect UV coordinates. The second is that each view iteration will overwrite the BentNormalAOTexture, so the multiview result is the last per-view texture. The solution is to also call this function per-view.

The previous changes include creating a per-view bent normal texture before copying back into an atlased texture. This involves a set of changes for UV sampling (e.g. in UpsampleDFAO) and setting up the culling properly in BuildTileCones.

diff --git a/Engine/Shaders/Private/DistanceFieldAOShared.ush b/Engine/Shaders/Private/DistanceFieldAOShared.ush
--- a/Engine/Shaders/Private/DistanceFieldAOShared.ush
+++ b/Engine/Shaders/Private/DistanceFieldAOShared.ush
@@ -155,9 +155,7 @@ FUpsampleDFAOOutput UpsampleDFAO(float2 BufferUV, float SceneDepth)
 {
 	FUpsampleDFAOOutput UpsampleDFAOOutput;
 
-	// Distance field AO was computed at 0,0 regardless of viewrect min
-	float2 DistanceFieldUVs = BufferUV - View.ViewRectMin.xy * View.BufferSizeAndInvSize.zw;
-	DistanceFieldUVs = min(DistanceFieldUVs, AOBufferBilinearUVMax);
+	float2 DistanceFieldUVs = min(BufferUV, AOBufferBilinearUVMax);
 
 #define BILATERAL_UPSAMPLE 1
 #if BILATERAL_UPSAMPLE
diff --git a/Engine/Shaders/Private/DistanceFieldLightingPost.usf b/Engine/Shaders/Private/DistanceFieldLightingPost.usf
--- a/Engine/Shaders/Private/DistanceFieldLightingPost.usf
+++ b/Engine/Shaders/Private/DistanceFieldLightingPost.usf
@@ -171,19 +171,22 @@ void UpdateHistoryDepthRejectionPS(
 	float4 ThisClip = float4( UVAndScreenPos.zw, ConvertToDeviceZ(SceneDepth), 1 );
 	float4 PrevClip = mul( ThisClip, View.ClipToPrevClip );
 	float2 PrevScreen = PrevClip.xy / PrevClip.w;
-	float2 Screenvelocity = UVAndScreenPos.zw - PrevScreen;
-		
-	float2 FullResTexel = UVAndScreenPos.xy - .5f * View.BufferSizeAndInvSize.zw;
+	float2 ScreenVelocity = UVAndScreenPos.zw - PrevScreen;
+		 
+	// The scene velocity texture is rendered as a texture atlas for all views,
+	// so first convert the vertex shader's UV coordinates to the size of the view rect and offset into the atlas
+	// before converting back to the correct UV coordinates.
+	float2 FullResTexel = (UVAndScreenPos.xy * View.ViewSizeAndInvSize.xy + View.ViewRectMin.xy - .5f) * View.BufferSizeAndInvSize.zw;
 	float4 EncodedVelocity = Texture2DSampleLevel(VelocityTexture, VelocityTextureSampler, FullResTexel, 0);
 		
 	if (EncodedVelocity.x > 0.0)
 	{
-		Screenvelocity = DecodeVelocityFromTexture(EncodedVelocity).xy;
+		ScreenVelocity = DecodeVelocityFromTexture(EncodedVelocity).xy;
 	}
 		
-	float PixelSpeed = 0.5 * length(Screenvelocity);
+	float PixelSpeed = 0.5 * length(ScreenVelocity);
 
-	float2 OldScreenPosition = (UVAndScreenPos.zw - Screenvelocity);
+	float2 OldScreenPosition = (UVAndScreenPos.zw - ScreenVelocity);
 	float2 OldDistanceFieldUVs = OldScreenPosition * HistoryScreenPositionScaleBias.xy + HistoryScreenPositionScaleBias.wz;
 
 	float EffectiveHistoryWeight = HistoryWeight;
@@ -201,8 +204,8 @@ void UpdateHistoryDepthRejectionPS(
 #define MANUAL_HISTORY_RESAMPLE 0
 #if MANUAL_HISTORY_RESAMPLE
 
-		float2 HistoryBufferSize = floor(View.BufferSizeAndInvSize.xy / DOWNSAMPLE_FACTOR);
-		float2 HistoryTexelSize = View.BufferSizeAndInvSize.zw * DOWNSAMPLE_FACTOR;
+		float2 HistoryBufferSize = floor(View.ViewSizeAndInvSize.xy / DOWNSAMPLE_FACTOR);
+		float2 HistoryTexelSize = View.ViewSizeAndInvSize.zw * DOWNSAMPLE_FACTOR;
 		float2 OldDistanceFieldUVsCorner00 = HistoryTexelSize * (floor(OldDistanceFieldUVs * HistoryBufferSize - .5f) + .5f);
 		float2 OldDistanceFieldUVsCorner01 = OldDistanceFieldUVsCorner00 + float2(0, HistoryTexelSize.y);
 		float2 OldDistanceFieldUVsCorner11 = OldDistanceFieldUVsCorner00 + HistoryTexelSize;
@@ -299,10 +302,6 @@ void FilterHistoryPS(
 {
 	float2 BufferUV = UVAndScreenPos.xy;
 
-	#if MANUALLY_CLAMP_UV
-		BufferUV = min(BufferUV, MaxSampleBufferUV);
-	#endif
-
 	float4 HistoryValue = Texture2DSampleLevel(BentNormalAOTexture, BentNormalAOSampler, BufferUV, 0);
 
 	// Only do the spatial search for pixels who discarded their history value
@@ -318,11 +317,6 @@ void FilterHistoryPS(
 			{
 				float2 SampleBufferUV = UVAndScreenPos.xy + BentNormalAOTexelSize * float2(x, y);
 
-				#if MANUALLY_CLAMP_UV
-					// Distance field AO was computed at 0,0 regardless of viewrect min.
-					SampleBufferUV = min(SampleBufferUV, MaxSampleBufferUV);
-				#endif
-
 				float4 TextureValue = Texture2DSampleLevel(BentNormalAOTexture, BentNormalAOSampler, SampleBufferUV, 0);
 
 				float SampleSceneDepth = abs(TextureValue.w);
@@ -365,7 +359,8 @@ void AOUpsamplePS(
 	,out float4 OutSceneColor : SV_Target0
 	)
 {
-	FGBufferData GBuffer = GetGBufferData(UVAndScreenPos.xy);
+	const float2 GBufferUV = UVAndScreenPos.xy * View.BufferBilinearUVMinMax.zw;
+	FGBufferData GBuffer = GetGBufferData(GBufferUV);
 
 	float3 BentNormal = UpsampleDFAO(UVAndScreenPos.xy, GBuffer.Depth, GBuffer.WorldNormal);
 
@@ -374,6 +369,5 @@ void AOUpsamplePS(
 	OutSceneColor = Visibility;
 #else
 	OutSceneColor = float4(length(BentNormal).xxx, 1.0f);
-
 #endif
 }
diff --git a/Engine/Shaders/Private/DistanceFieldObjectCulling.usf b/Engine/Shaders/Private/DistanceFieldObjectCulling.usf
--- a/Engine/Shaders/Private/DistanceFieldObjectCulling.usf
+++ b/Engine/Shaders/Private/DistanceFieldObjectCulling.usf
@@ -129,12 +129,10 @@ void BuildTileConesMain(
 {
     uint ThreadIndex = GroupThreadId.y * THREADGROUP_SIZEX + GroupThreadId.x;
 
-	float2 BaseLevelScreenUV = (DispatchThreadId.xy + float2(.5f, .5f)) * DOWNSAMPLE_FACTOR * View.BufferSizeAndInvSize.zw;
+	// Sampling from the texture based off of the ViewRect size because the texture is created on a per-view basis
+	float2 BaseLevelScreenUV = (DispatchThreadId.xy + float2(.5f, .5f)) * DOWNSAMPLE_FACTOR * View.ViewSizeAndInvSize.zw;
 	float SceneDepth = GetDownsampledDepth(BaseLevelScreenUV);
 
-	//float2 ScreenUV = (DispatchThreadId.xy * DOWNSAMPLE_FACTOR + View.ViewRectMin.xy + float2(.5f, .5f)) * View.BufferSizeAndInvSize.zw;
-	//float SceneDepth = CalcSceneDepth(ScreenUV);
-
 	// Initialize per-tile variables
     if (ThreadIndex == 0) 
 	{
diff --git a/Engine/Source/Runtime/Renderer/Private/CapsuleShadowRendering.cpp b/Engine/Source/Runtime/Renderer/Private/CapsuleShadowRendering.cpp
--- a/Engine/Source/Runtime/Renderer/Private/CapsuleShadowRendering.cpp
+++ b/Engine/Source/Runtime/Renderer/Private/CapsuleShadowRendering.cpp
@@ -1180,6 +1180,7 @@ bool FSceneRenderer::ShouldPrepareForDFInsetIndirectShadow() const
 
 void FDeferredShadingSceneRenderer::RenderCapsuleShadowsForMovableSkylight(
 	FRDGBuilder& GraphBuilder,
+	const FViewInfo& View,
 	TRDGUniformBufferRef<FSceneTextureUniformParameters> SceneTexturesUniformBuffer,
 	FRDGTextureRef& BentNormalOutput) const
 {
@@ -1188,96 +1189,76 @@ void FDeferredShadingSceneRenderer::RenderCapsuleShadowsForMovableSkylight(
 	{
 		QUICK_SCOPE_CYCLE_COUNTER(STAT_RenderCapsuleShadowsSkylight);
 
-		bool bAnyViewsUseCapsuleShadows = false;
-
-		for (int32 ViewIndex = 0; ViewIndex < Views.Num(); ViewIndex++)
+		if (View.IndirectShadowPrimitives.Num() > 0 && View.ViewState)
 		{
-			const FViewInfo& View = Views[ViewIndex];
-
-			if (View.IndirectShadowPrimitives.Num() > 0 && View.ViewState)
-			{
-				bAnyViewsUseCapsuleShadows = true;
-			}
-		}
+			RDG_GPU_MASK_SCOPE(GraphBuilder, View.GPUMask);
+			RDG_EVENT_SCOPE(GraphBuilder, "IndirectCapsuleShadows");
+			RDG_GPU_STAT_SCOPE(GraphBuilder, CapsuleShadows);
 
-		if (bAnyViewsUseCapsuleShadows)
-		{
 			FRDGTextureRef NewBentNormal = nullptr;
-			AllocateOrReuseAORenderTarget(GraphBuilder, Views[0], NewBentNormal, TEXT("CapsuleBentNormal"), PF_FloatRGBA);
+			AllocateOrReuseAORenderTarget(GraphBuilder, View, NewBentNormal, TEXT("CapsuleBentNormal"), PF_FloatRGBA);
 
-			for (int32 ViewIndex = 0; ViewIndex < Views.Num(); ViewIndex++)
-			{
-				const FViewInfo& View = Views[ViewIndex];
+			int32 NumCapsuleShapes = 0;
+			int32 NumMeshesWithCapsules = 0;
+			int32 NumMeshDistanceFieldCasters = 0;
+			IndirectCapsuleShadowsResources Resources = CreateIndirectCapsuleShadowsResources(GraphBuilder, Scene, View, NumCapsuleShapes, NumMeshesWithCapsules, NumMeshDistanceFieldCasters);
 
-				if (View.IndirectShadowPrimitives.Num() > 0 && View.ViewState)
-				{
-					RDG_GPU_MASK_SCOPE(GraphBuilder, View.GPUMask);
-					RDG_EVENT_SCOPE(GraphBuilder, "IndirectCapsuleShadows");
-					RDG_GPU_STAT_SCOPE(GraphBuilder, CapsuleShadows);
+			// Don't render indirect occlusion from mesh distance fields when operating on a movable skylight,
+			// DFAO is responsible for indirect occlusion from meshes with distance fields on a movable skylight.
+			// A single mesh should only provide indirect occlusion for a given lighting component in one way.
+			NumMeshDistanceFieldCasters = 0;
 
-					int32 NumCapsuleShapes = 0;
-					int32 NumMeshesWithCapsules = 0;
-					int32 NumMeshDistanceFieldCasters = 0;
-					IndirectCapsuleShadowsResources Resources = CreateIndirectCapsuleShadowsResources(GraphBuilder, Scene, View, NumCapsuleShapes, NumMeshesWithCapsules, NumMeshDistanceFieldCasters);
+			if (NumCapsuleShapes > 0)
+			{
+				check(Resources.IndirectShadowLightDirectionSRV);
 
-					// Don't render indirect occlusion from mesh distance fields when operating on a movable skylight,
-					// DFAO is responsible for indirect occlusion from meshes with distance fields on a movable skylight.
-					// A single mesh should only provide indirect occlusion for a given lighting component in one way.
-					NumMeshDistanceFieldCasters = 0;
+				FIntRect ScissorRect = View.ViewRect;
 
-					if (NumCapsuleShapes > 0)
-					{
-						check(Resources.IndirectShadowLightDirectionSRV);
-
-						FIntRect ScissorRect = View.ViewRect;
-
-						{
-							uint32 GroupSizeX = FMath::DivideAndRoundUp(ScissorRect.Size().X / GAODownsampleFactor, GShadowShapeTileSize);
-							uint32 GroupSizeY = FMath::DivideAndRoundUp(ScissorRect.Size().Y / GAODownsampleFactor, GShadowShapeTileSize);
-
-							auto* PassParameters = GraphBuilder.AllocParameters<FCapsuleShadowingCS::FParameters>();
-
-							SetupCapsuleShadowingParameters(
-								GraphBuilder,
-								*PassParameters,
-								ECapsuleShadowingType::MovableSkylightTiledCulling,
-								GraphBuilder.CreateUAV(NewBentNormal),
-								FIntPoint(GroupSizeX, GroupSizeY),
-								BentNormalOutput,
-								FVector2D(GroupSizeX, GroupSizeY),
-								nullptr,
-								ScissorRect,
-								GAODownsampleFactor,
-								GCapsuleMaxIndirectOcclusionDistance,
-								Scene,
-								View,
-								SceneTexturesUniformBuffer,
-								NumCapsuleShapes,
-								Resources.IndirectShadowCapsuleShapesSRV,
-								NumMeshDistanceFieldCasters,
-								Resources.IndirectShadowMeshDistanceFieldCasterIndicesSRV,
-								Resources.IndirectShadowLightDirectionSRV,
-								nullptr
-							);
-
-							FCapsuleShadowingCS::FPermutationDomain PermutationVector;
-							PermutationVector.Set<FCapsuleShadowingCS::FShapeShadow>(ECapsuleShadowingType::MovableSkylightTiledCulling);
-							PermutationVector.Set<FCapsuleShadowingCS::FIndirectPrimitiveType>(EIndirectShadowingPrimitiveTypes::CapsuleShapes);
-							auto ComputeShader = View.ShaderMap->GetShader<FCapsuleShadowingCS>(PermutationVector);
-
-							FComputeShaderUtils::AddPass(
-								GraphBuilder,
-								RDG_EVENT_NAME("TiledCapsuleShadowing % u capsules among % u meshes", NumCapsuleShapes, NumMeshesWithCapsules),
-								ERDGPassFlags::Compute,
-								ComputeShader,
-								PassParameters,
-								FIntVector(GroupSizeX, GroupSizeY, 1));
-						}
-
-						// Replace the pipeline output with our output that has capsule shadows applied
-						BentNormalOutput = NewBentNormal;
-					}
+				{
+					uint32 GroupSizeX = FMath::DivideAndRoundUp(ScissorRect.Size().X / GAODownsampleFactor, GShadowShapeTileSize);
+					uint32 GroupSizeY = FMath::DivideAndRoundUp(ScissorRect.Size().Y / GAODownsampleFactor, GShadowShapeTileSize);
+
+					auto* PassParameters = GraphBuilder.AllocParameters<FCapsuleShadowingCS::FParameters>();
+
+					SetupCapsuleShadowingParameters(
+						GraphBuilder,
+						*PassParameters,
+						ECapsuleShadowingType::MovableSkylightTiledCulling,
+						GraphBuilder.CreateUAV(NewBentNormal),
+						FIntPoint(GroupSizeX, GroupSizeY),
+						BentNormalOutput,
+						FVector2D(GroupSizeX, GroupSizeY),
+						nullptr,
+						ScissorRect,
+						GAODownsampleFactor,
+						GCapsuleMaxIndirectOcclusionDistance,
+						Scene,
+						View,
+						SceneTexturesUniformBuffer,
+						NumCapsuleShapes,
+						Resources.IndirectShadowCapsuleShapesSRV,
+						NumMeshDistanceFieldCasters,
+						Resources.IndirectShadowMeshDistanceFieldCasterIndicesSRV,
+						Resources.IndirectShadowLightDirectionSRV,
+						nullptr
+					);
+
+					FCapsuleShadowingCS::FPermutationDomain PermutationVector;
+					PermutationVector.Set<FCapsuleShadowingCS::FShapeShadow>(ECapsuleShadowingType::MovableSkylightTiledCulling);
+					PermutationVector.Set<FCapsuleShadowingCS::FIndirectPrimitiveType>(EIndirectShadowingPrimitiveTypes::CapsuleShapes);
+					auto ComputeShader = View.ShaderMap->GetShader<FCapsuleShadowingCS>(PermutationVector);
+
+					FComputeShaderUtils::AddPass(
+						GraphBuilder,
+						RDG_EVENT_NAME("TiledCapsuleShadowing % u capsules among % u meshes", NumCapsuleShapes, NumMeshesWithCapsules),
+						ERDGPassFlags::Compute,
+						ComputeShader,
+						PassParameters,
+						FIntVector(GroupSizeX, GroupSizeY, 1));
 				}
+
+				// Replace the pipeline output with our output that has capsule shadows applied
+				BentNormalOutput = NewBentNormal;
 			}
 		}
 	}
diff --git a/Engine/Source/Runtime/Renderer/Private/DeferredShadingRenderer.h b/Engine/Source/Runtime/Renderer/Private/DeferredShadingRenderer.h
--- a/Engine/Source/Runtime/Renderer/Private/DeferredShadingRenderer.h
+++ b/Engine/Source/Runtime/Renderer/Private/DeferredShadingRenderer.h
@@ -895,6 +895,7 @@ private:
 	/** Renders capsule shadows for movable skylights, using the cone of visibility (bent normal) from DFAO. */
 	void RenderCapsuleShadowsForMovableSkylight(
 		FRDGBuilder& GraphBuilder,
+		const FViewInfo& View,
 		TRDGUniformBufferRef<FSceneTextureUniformParameters> SceneTexturesUniformBuffer,
 		FRDGTextureRef& BentNormalOutput) const;
 
diff --git a/Engine/Source/Runtime/Renderer/Private/DistanceFieldAmbientOcclusion.cpp b/Engine/Source/Runtime/Renderer/Private/DistanceFieldAmbientOcclusion.cpp
--- a/Engine/Source/Runtime/Renderer/Private/DistanceFieldAmbientOcclusion.cpp
+++ b/Engine/Source/Runtime/Renderer/Private/DistanceFieldAmbientOcclusion.cpp
@@ -159,7 +159,7 @@ void TileIntersectionModifyCompilationEnvironment(EShaderPlatform Platform, FSha
 
 FIntPoint GetBufferSizeForAO(const FViewInfo& View)
 {
-	return FIntPoint::DivideAndRoundDown(View.GetSceneTexturesConfig().Extent, GAODownsampleFactor);
+	return FIntPoint::DivideAndRoundDown(View.ViewRect.Size(), GAODownsampleFactor);
 }
 
 // Sample set restricted to not self-intersect a surface based on cone angle .475882232
@@ -268,10 +268,10 @@ FDFAOUpsampleParameters DistanceField::SetupAOUpsampleParameters(const FViewInfo
 {
 	const float DistanceFadeScaleValue = 1.0f / ((1.0f - GAOViewFadeDistanceScale) * GetMaxAOViewDistance());
 
-	const FIntPoint AOBufferSize = GetBufferSizeForAO(View);
+	const FIntPoint DFAOBufferSize = GetBufferSizeForAO(View);
 	const FVector2f UVMax(
-		(View.ViewRect.Width() / GAODownsampleFactor - 0.51f) / AOBufferSize.X, // 0.51 - so bilateral gather4 won't sample invalid texels
-		(View.ViewRect.Height() / GAODownsampleFactor - 0.51f) / AOBufferSize.Y);
+		(DFAOBufferSize.X - 0.51f) / DFAOBufferSize.X, // 0.51 - so bilateral gather4 won't sample invalid texels
+		(DFAOBufferSize.Y - 0.51f) / DFAOBufferSize.Y);
 
 	FDFAOUpsampleParameters ShaderParameters;
 	ShaderParameters.BentNormalAOTexture = DistanceFieldAOBentNormal;
@@ -374,84 +374,74 @@ IMPLEMENT_GLOBAL_SHADER(FComputeDistanceFieldNormalCS, "/Engine/Private/Distance
 
 void ComputeDistanceFieldNormal(
 	FRDGBuilder& GraphBuilder,
-	const TArray<FViewInfo>& Views,
+	const FViewInfo& View,
 	TRDGUniformBufferRef<FSceneTextureUniformParameters> SceneTexturesUniformBuffer,
 	FRDGTextureRef DistanceFieldNormal,
 	const FDistanceFieldAOParameters& Parameters)
 {
+	const FIntPoint DFAOTextureSize = GetBufferSizeForAO(View);
+
 	if (GAOComputeShaderNormalCalculation)
 	{
-		for (int32 ViewIndex = 0; ViewIndex < Views.Num(); ViewIndex++)
-		{
-			const FViewInfo& View = Views[ViewIndex];
-			RDG_GPU_MASK_SCOPE(GraphBuilder, View.GPUMask);
-
-			uint32 GroupSizeX = FMath::DivideAndRoundUp(View.ViewRect.Size().X / GAODownsampleFactor, GDistanceFieldAOTileSizeX);
-			uint32 GroupSizeY = FMath::DivideAndRoundUp(View.ViewRect.Size().Y / GAODownsampleFactor, GDistanceFieldAOTileSizeY);
+		uint32 GroupSizeX = FMath::DivideAndRoundUp(DFAOTextureSize.X, GDistanceFieldAOTileSizeX);
+		uint32 GroupSizeY = FMath::DivideAndRoundUp(DFAOTextureSize.Y, GDistanceFieldAOTileSizeY);
 
-			auto* PassParameters = GraphBuilder.AllocParameters<FComputeDistanceFieldNormalCS::FParameters>();
-			PassParameters->View = View.ViewUniformBuffer;
-			PassParameters->SceneTextures = SceneTexturesUniformBuffer;
-			PassParameters->Strata = Strata::BindStrataGlobalUniformParameters(View);
-			PassParameters->AOParameters = DistanceField::SetupAOShaderParameters(Parameters);
-			PassParameters->RWDistanceFieldNormal = GraphBuilder.CreateUAV(DistanceFieldNormal);
+		auto* PassParameters = GraphBuilder.AllocParameters<FComputeDistanceFieldNormalCS::FParameters>();
+		PassParameters->View = View.ViewUniformBuffer;
+		PassParameters->SceneTextures = SceneTexturesUniformBuffer;
+		PassParameters->Strata = Strata::BindStrataGlobalUniformParameters(View);
+		PassParameters->AOParameters = DistanceField::SetupAOShaderParameters(Parameters);
+		PassParameters->RWDistanceFieldNormal = GraphBuilder.CreateUAV(DistanceFieldNormal);
 
-			TShaderMapRef<FComputeDistanceFieldNormalCS> ComputeShader(View.ShaderMap);
+		TShaderMapRef<FComputeDistanceFieldNormalCS> ComputeShader(View.ShaderMap);
 
-			FComputeShaderUtils::AddPass(GraphBuilder, RDG_EVENT_NAME("ComputeNormalCS"), ComputeShader, PassParameters, FIntVector(GroupSizeX, GroupSizeY, 1));
-		}
+		FComputeShaderUtils::AddPass(GraphBuilder, RDG_EVENT_NAME("ComputeNormalCS"), ComputeShader, PassParameters, FIntVector(GroupSizeX, GroupSizeY, 1));
 	}
 	else
 	{
-		for (int32 ViewIndex = 0; ViewIndex < Views.Num(); ViewIndex++)
+		auto* PassParameters = GraphBuilder.AllocParameters<FComputeDistanceFieldNormalPS::FParameters>();
+		PassParameters->View = View.ViewUniformBuffer;
+		PassParameters->SceneTextures = SceneTexturesUniformBuffer;
+		PassParameters->Strata = Strata::BindStrataGlobalUniformParameters(View);
+		PassParameters->AOParameters = DistanceField::SetupAOShaderParameters(Parameters);
+		PassParameters->RenderTargets[0] = FRenderTargetBinding(DistanceFieldNormal, ERenderTargetLoadAction::EClear);
+
+		GraphBuilder.AddPass(
+			RDG_EVENT_NAME("ComputeNormal"),
+			PassParameters,
+			ERDGPassFlags::Raster,
+			[&View, PassParameters, DFAOTextureSize](FRHICommandList& RHICmdList)
 		{
-			const FViewInfo& View = Views[ViewIndex];
-			RDG_GPU_MASK_SCOPE(GraphBuilder, View.GPUMask);
-
-			auto* PassParameters = GraphBuilder.AllocParameters<FComputeDistanceFieldNormalPS::FParameters>();
-			PassParameters->View = View.ViewUniformBuffer;
-			PassParameters->SceneTextures = SceneTexturesUniformBuffer;
-			PassParameters->Strata = Strata::BindStrataGlobalUniformParameters(View);
-			PassParameters->AOParameters = DistanceField::SetupAOShaderParameters(Parameters);
-			PassParameters->RenderTargets[0] = FRenderTargetBinding(DistanceFieldNormal, ViewIndex == 0 ? ERenderTargetLoadAction::EClear : ERenderTargetLoadAction::ELoad);
-
-			GraphBuilder.AddPass(
-				RDG_EVENT_NAME("ComputeNormal"),
-				PassParameters,
-				ERDGPassFlags::Raster,
-				[&View, PassParameters](FRHICommandList& RHICmdList)
-			{
-				RHICmdList.SetViewport(0, 0, 0.0f, View.ViewRect.Width() / GAODownsampleFactor, View.ViewRect.Height() / GAODownsampleFactor, 1.0f);
-
-				TShaderMapRef<FPostProcessVS> VertexShader(View.ShaderMap);
-
-				TShaderMapRef<FComputeDistanceFieldNormalPS> PixelShader(View.ShaderMap);
-
-				FGraphicsPipelineStateInitializer GraphicsPSOInit;
-				RHICmdList.ApplyCachedRenderTargets(GraphicsPSOInit);
-				GraphicsPSOInit.RasterizerState = TStaticRasterizerState<FM_Solid, CM_None>::GetRHI();
-				GraphicsPSOInit.DepthStencilState = TStaticDepthStencilState<false, CF_Always>::GetRHI();
-				GraphicsPSOInit.BlendState = TStaticBlendState<>::GetRHI();
-				GraphicsPSOInit.PrimitiveType = PT_TriangleList;
-				GraphicsPSOInit.BoundShaderState.VertexDeclarationRHI = GFilterVertexDeclaration.VertexDeclarationRHI;
-				GraphicsPSOInit.BoundShaderState.VertexShaderRHI = VertexShader.GetVertexShader();
-				GraphicsPSOInit.BoundShaderState.PixelShaderRHI = PixelShader.GetPixelShader();
-
-				SetGraphicsPipelineState(RHICmdList, GraphicsPSOInit, 0);
-
-				SetShaderParameters(RHICmdList, PixelShader, PixelShader.GetPixelShader(), *PassParameters);
-
-				DrawRectangle(
-					RHICmdList,
-					0, 0,
-					View.ViewRect.Width() / GAODownsampleFactor, View.ViewRect.Height() / GAODownsampleFactor,
-					0, 0,
-					View.ViewRect.Width(), View.ViewRect.Height(),
-					FIntPoint(View.ViewRect.Width() / GAODownsampleFactor, View.ViewRect.Height() / GAODownsampleFactor),
-					View.GetSceneTexturesConfig().Extent,
-					VertexShader);
-			});
-		}
+			RHICmdList.SetViewport(0, 0, 0.0f, DFAOTextureSize.X, DFAOTextureSize.Y, 1.0f);
+
+			TShaderMapRef<FPostProcessVS> VertexShader(View.ShaderMap);
+
+			TShaderMapRef<FComputeDistanceFieldNormalPS> PixelShader(View.ShaderMap);
+
+			FGraphicsPipelineStateInitializer GraphicsPSOInit;
+			RHICmdList.ApplyCachedRenderTargets(GraphicsPSOInit);
+			GraphicsPSOInit.RasterizerState = TStaticRasterizerState<FM_Solid, CM_None>::GetRHI();
+			GraphicsPSOInit.DepthStencilState = TStaticDepthStencilState<false, CF_Always>::GetRHI();
+			GraphicsPSOInit.BlendState = TStaticBlendState<>::GetRHI();
+			GraphicsPSOInit.PrimitiveType = PT_TriangleList;
+			GraphicsPSOInit.BoundShaderState.VertexDeclarationRHI = GFilterVertexDeclaration.VertexDeclarationRHI;
+			GraphicsPSOInit.BoundShaderState.VertexShaderRHI = VertexShader.GetVertexShader();
+			GraphicsPSOInit.BoundShaderState.PixelShaderRHI = PixelShader.GetPixelShader();
+
+			SetGraphicsPipelineState(RHICmdList, GraphicsPSOInit, 0);
+
+			SetShaderParameters(RHICmdList, PixelShader, PixelShader.GetPixelShader(), *PassParameters);
+
+			DrawRectangle(
+				RHICmdList,
+				0, 0,
+				DFAOTextureSize.X, DFAOTextureSize.Y,
+				0, 0,
+				View.ViewRect.Width(), View.ViewRect.Height(),
+				DFAOTextureSize,
+				View.GetSceneTexturesConfig().Extent,
+				VertexShader);
+		});
 	}
 }
 
@@ -656,10 +646,9 @@ void AllocateTileIntersectionBuffers(
 	OutParameters.TileListGroupSize = TileListGroupSize;
 }
 
-void ListDistanceFieldLightingMemory(const FViewInfo& View, FSceneRenderer& SceneRenderer)
+void ListDistanceFieldLightingMemory(const FScene* Scene, FSceneRenderer& SceneRenderer)
 {
 #if !NO_LOGGING
-	const FScene* Scene = (const FScene*)View.Family->Scene;
 	UE_LOG(LogRenderer, Log, TEXT("Shared GPU memory (excluding render targets)"));
 
 	if (Scene->DistanceFieldSceneData.NumObjectsInBuffer > 0)
@@ -775,13 +764,21 @@ void FDeferredShadingSceneRenderer::RenderDFAOAsIndirectShadowing(
 
 bool FDeferredShadingSceneRenderer::ShouldRenderDistanceFieldLighting() const
 {
-	//@todo - support multiple views
-	const FViewInfo& View = Views[0];
+	bool bSupportsDistanceFieldAO = true;
+
+	for (int32_t ViewIndex = 0; ViewIndex < Views.Num(); ++ViewIndex)
+	{
+		const FViewInfo& View = Views[ViewIndex];
 
-	return SupportsDistanceFieldAO(View.GetFeatureLevel(), View.GetShaderPlatform())
-		&& Views.Num() == 1
-		&& View.IsPerspectiveProjection()
-		&& Scene->DistanceFieldSceneData.NumObjectsInBuffer;
+		if (!SupportsDistanceFieldAO(View.GetFeatureLevel(), View.GetShaderPlatform())
+			|| !View.IsPerspectiveProjection())
+		{
+			bSupportsDistanceFieldAO = false;
+			break;
+		}
+	}
+
+	return bSupportsDistanceFieldAO && Scene->DistanceFieldSceneData.NumObjectsInBuffer;
 }
 
 void FDeferredShadingSceneRenderer::RenderDistanceFieldLighting(
@@ -795,10 +792,6 @@ void FDeferredShadingSceneRenderer::RenderDistanceFieldLighting(
 	check(ShouldRenderDistanceFieldLighting());
 	check(!Scene->DistanceFieldSceneData.HasPendingOperations());
 
-	//@todo - support multiple views
-	const FViewInfo& View = Views[0];
-
-	RDG_GPU_MASK_SCOPE(GraphBuilder, View.GPUMask);
 	RDG_EVENT_SCOPE(GraphBuilder, "DistanceFieldLighting");
 	QUICK_SCOPE_CYCLE_COUNTER(STAT_RenderDistanceFieldLighting);
 
@@ -807,7 +800,7 @@ void FDeferredShadingSceneRenderer::RenderDistanceFieldLighting(
 	if (bListMemoryNextFrame)
 	{
 		bListMemoryNextFrame = false;
-		ListDistanceFieldLightingMemory(View, *this);
+		ListDistanceFieldLightingMemory(Scene, *this);
 	}
 
 	if (bListMeshDistanceFieldsMemoryNextFrame)
@@ -816,65 +809,80 @@ void FDeferredShadingSceneRenderer::RenderDistanceFieldLighting(
 		Scene->DistanceFieldSceneData.ListMeshDistanceFields(true);
 	}
 
-	FRDGTextureRef DistanceFieldNormal = nullptr;
+	FRDGTextureRef BentNormalOutput = nullptr;
 
 	{
-		const FIntPoint BufferSize = GetBufferSizeForAO(View);
-		const FRDGTextureDesc Desc = FRDGTextureDesc::Create2D(BufferSize, PF_FloatRGBA, FClearValueBinding::Transparent, GFastVRamConfig.DistanceFieldNormal | TexCreate_RenderTargetable | TexCreate_UAV | TexCreate_ShaderResource);
-		DistanceFieldNormal = GraphBuilder.CreateTexture(Desc, TEXT("DistanceFieldNormal"));
+		const FIntPoint BufferSize = GetActiveSceneTexturesConfig().Extent / GAODownsampleFactor;
+		const FRDGTextureDesc Desc = FRDGTextureDesc::Create2D(BufferSize, PF_FloatRGBA, FClearValueBinding::None, GFastVRamConfig.DistanceFieldAOBentNormal | TexCreate_ShaderResource | TexCreate_RenderTargetable | TexCreate_UAV);
+		BentNormalOutput = GraphBuilder.CreateTexture(Desc, TEXT("DistanceFieldBentNormalAO"));
+		AddClearUAVPass(GraphBuilder, GraphBuilder.CreateUAV(BentNormalOutput), FLinearColor::Black);
 	}
 
-	const FIntPoint TileListGroupSize = GetTileListGroupSizeForView(View);
-	const int32 MaxSceneObjects = FMath::DivideAndRoundUp(Scene->DistanceFieldSceneData.NumObjectsInBuffer, 256) * 256;
-	const bool bAllow16BitIndices = !IsMetalPlatform(GShaderPlatformForFeatureLevel[View.FeatureLevel]);
+	for (int32 ViewIndex = 0; ViewIndex < Views.Num(); ++ViewIndex)
+	{
+		const FViewInfo& View = Views[ViewIndex];
+		RDG_GPU_MASK_SCOPE(GraphBuilder, View.GPUMask);
+		RDG_EVENT_SCOPE_CONDITIONAL(GraphBuilder, Views.Num() > 1, "View%d", ViewIndex);
 
-	FRDGBufferRef ObjectTilesIndirectArguments = nullptr;
-	FTileIntersectionParameters TileIntersectionParameters;
+		FRDGTextureRef DistanceFieldNormal = nullptr;
 
-	AllocateTileIntersectionBuffers(GraphBuilder, TileListGroupSize, MaxSceneObjects, bAllow16BitIndices, ObjectTilesIndirectArguments, TileIntersectionParameters);
+		{
+			const FIntPoint BufferSize = GetBufferSizeForAO(View);
+			const FRDGTextureDesc Desc = FRDGTextureDesc::Create2D(BufferSize, PF_FloatRGBA, FClearValueBinding::Transparent, GFastVRamConfig.DistanceFieldNormal | TexCreate_RenderTargetable | TexCreate_UAV | TexCreate_ShaderResource);
+			DistanceFieldNormal = GraphBuilder.CreateTexture(Desc, TEXT("DistanceFieldNormal"));
+		}
 
-	FRDGBufferRef ObjectIndirectArguments = nullptr;
-	FDistanceFieldCulledObjectBufferParameters CulledObjectBufferParameters;
+		ComputeDistanceFieldNormal(GraphBuilder, View, SceneTextures.UniformBuffer, DistanceFieldNormal, Parameters);
 
-	if (UseAOObjectDistanceField())
-	{
-		AllocateDistanceFieldCulledObjectBuffers(
-			GraphBuilder,
-			MaxSceneObjects,
-			ObjectIndirectArguments,
-			CulledObjectBufferParameters);
+		const FIntPoint TileListGroupSize = GetTileListGroupSizeForView(View);
+		const int32 MaxSceneObjects = FMath::DivideAndRoundUp(Scene->DistanceFieldSceneData.NumObjectsInBuffer, 256) * 256;
+		const bool bAllow16BitIndices = !IsMetalPlatform(GShaderPlatformForFeatureLevel[View.FeatureLevel]);
 
-		CullObjectsToView(GraphBuilder, *Scene, View, Parameters, CulledObjectBufferParameters);
-	}
+		FRDGBufferRef ObjectTilesIndirectArguments = nullptr;
+		FTileIntersectionParameters TileIntersectionParameters;
 
-	ComputeDistanceFieldNormal(GraphBuilder, Views, SceneTextures.UniformBuffer, DistanceFieldNormal, Parameters);
+		AllocateTileIntersectionBuffers(GraphBuilder, TileListGroupSize, MaxSceneObjects, bAllow16BitIndices, ObjectTilesIndirectArguments, TileIntersectionParameters);
 
-	// Intersect objects with screen tiles, build lists
-	if (UseAOObjectDistanceField())
-	{
-		//@todo - support multiple views - should pass one TileIntersectionParameters per view
-		BuildTileObjectLists(GraphBuilder, *Scene, Views, ObjectIndirectArguments, CulledObjectBufferParameters, TileIntersectionParameters, DistanceFieldNormal, Parameters);
-	}
+		FRDGBufferRef ObjectIndirectArguments = nullptr;
+		FDistanceFieldCulledObjectBufferParameters CulledObjectBufferParameters;
 
-	FRDGTextureRef BentNormalOutput = nullptr;
+		if (UseAOObjectDistanceField())
+		{
+			AllocateDistanceFieldCulledObjectBuffers(
+				GraphBuilder,
+				MaxSceneObjects,
+				ObjectIndirectArguments,
+				CulledObjectBufferParameters);
 
-	RenderDistanceFieldAOScreenGrid(
-		GraphBuilder,
-		SceneTextures,
-		View,
-		CulledObjectBufferParameters,
-		ObjectTilesIndirectArguments,
-		TileIntersectionParameters,
-		Parameters,
-		DistanceFieldNormal,
-		BentNormalOutput);
-
-	RenderCapsuleShadowsForMovableSkylight(GraphBuilder, SceneTextures.UniformBuffer, BentNormalOutput);
-
-	// Upsample to full resolution, write to output in case of debug AO visualization or scene color modulation (standard upsampling is done later together with sky lighting and reflection environment)
-	if (bModulateToSceneColor || bVisualizeAmbientOcclusion)
-	{
-		UpsampleBentNormalAO(GraphBuilder, Views, SceneTextures.UniformBuffer, SceneTextures.Color.Target, BentNormalOutput, bModulateToSceneColor && !bVisualizeAmbientOcclusion);
+			CullObjectsToView(GraphBuilder, *Scene, View, Parameters, CulledObjectBufferParameters);
+		}
+
+		BuildTileObjectLists(GraphBuilder, *Scene, View, SceneTextures.UniformBuffer, ObjectIndirectArguments, CulledObjectBufferParameters, TileIntersectionParameters, DistanceFieldNormal, Parameters);
+
+		// Render to a per-view BentNormal first because it also needs to be stored by the view history.
+		// The per-view output can be copied back to the appropriate region in the output BentNormal texture.
+		FRDGTextureRef PerViewBentNormal = nullptr;
+
+		RenderDistanceFieldAOScreenGrid(
+			GraphBuilder,
+			SceneTextures,
+			View,
+			CulledObjectBufferParameters,
+			ObjectTilesIndirectArguments,
+			TileIntersectionParameters,
+			Parameters,
+			DistanceFieldNormal,
+			PerViewBentNormal);
+
+		RenderCapsuleShadowsForMovableSkylight(GraphBuilder, View, SceneTextures.UniformBuffer, PerViewBentNormal);
+
+		// Upsample to full resolution, write to output in case of debug AO visualization or scene color modulation (standard upsampling is done later together with sky lighting and reflection environment)
+		if (bModulateToSceneColor || bVisualizeAmbientOcclusion)
+		{
+			UpsampleBentNormalAO(GraphBuilder, View, SceneTextures.UniformBuffer, SceneTextures.Color.Target, PerViewBentNormal, bModulateToSceneColor && !bVisualizeAmbientOcclusion);
+		}
+
+		AddCopyTexturePass(GraphBuilder, PerViewBentNormal, BentNormalOutput, FIntPoint::ZeroValue, View.ViewRect.Min / GAODownsampleFactor, View.ViewRect.Size() / GAODownsampleFactor);
 	}
 
 	OutDynamicBentNormalAO = BentNormalOutput;
diff --git a/Engine/Source/Runtime/Renderer/Private/DistanceFieldAmbientOcclusion.h b/Engine/Source/Runtime/Renderer/Private/DistanceFieldAmbientOcclusion.h
--- a/Engine/Source/Runtime/Renderer/Private/DistanceFieldAmbientOcclusion.h
+++ b/Engine/Source/Runtime/Renderer/Private/DistanceFieldAmbientOcclusion.h
@@ -14,11 +14,13 @@ class FRDGBuffer;
 class FRDGBuilder;
 class FRDGTexture;
 class FScene;
+class FSceneTextureUniformParameters;
 class FSceneViewFamily;
 class FViewInfo;
 struct FShaderCompilerEnvironment;
 using FRDGBufferRef = FRDGBuffer*;
 using FRDGTextureRef = FRDGTexture*;
+template <typename TUniformStruct> using TRDGUniformBufferRef = TRDGUniformBuffer<TUniformStruct>*;
 
 /** Base downsample factor that all distance field AO operations are done at. */
 const int32 GAODownsampleFactor = 2;
@@ -107,7 +109,8 @@ extern bool ShouldDoReflectionEnvironment(const FScene* Scene, const FSceneViewF
 extern void CullObjectsToView(FRDGBuilder& GraphBuilder, const FScene& Scene, const FViewInfo& View, const FDistanceFieldAOParameters& Parameters, FDistanceFieldCulledObjectBufferParameters& CulledObjectBuffers);
 extern void BuildTileObjectLists(FRDGBuilder& GraphBuilder,
 	const FScene& Scene,
-	TArray<FViewInfo>& Views,
+	const FViewInfo& View,
+	TRDGUniformBufferRef<FSceneTextureUniformParameters> SceneTexturesUniformBuffer,
 	FRDGBufferRef ObjectIndirectArguments,
 	const FDistanceFieldCulledObjectBufferParameters& CulledObjectBufferParameters,
 	FTileIntersectionParameters TileIntersectionParameters,
diff --git a/Engine/Source/Runtime/Renderer/Private/DistanceFieldLightingPost.cpp b/Engine/Source/Runtime/Renderer/Private/DistanceFieldLightingPost.cpp
--- a/Engine/Source/Runtime/Renderer/Private/DistanceFieldLightingPost.cpp
+++ b/Engine/Source/Runtime/Renderer/Private/DistanceFieldLightingPost.cpp
@@ -166,14 +166,10 @@ public:
 		SHADER_PARAMETER_RDG_TEXTURE(Texture2D, BentNormalAOTexture)
 		SHADER_PARAMETER_SAMPLER(SamplerState, BentNormalAOSampler)
 		SHADER_PARAMETER(FVector2f, BentNormalAOTexelSize)
-		SHADER_PARAMETER(FVector2f, MaxSampleBufferUV)
 		SHADER_PARAMETER(float, HistoryWeight)
 		RENDER_TARGET_BINDING_SLOTS()
 	END_SHADER_PARAMETER_STRUCT()
 
-	class FManuallyClampUV : SHADER_PERMUTATION_BOOL("MANUALLY_CLAMP_UV");
-	using FPermutationDomain = TShaderPermutationDomain<FManuallyClampUV>;
-
 	static bool ShouldCompilePermutation(const FGlobalShaderPermutationParameters& Parameters)
 	{
 		return ShouldCompileDFLightingPostShaders(Parameters.Platform);
@@ -243,7 +239,9 @@ void GeometryAwareUpsample(FRDGBuilder& GraphBuilder, const FViewInfo& View, FRD
 		ERDGPassFlags::Raster,
 		[PassParameters, VertexShader, PixelShader, &View](FRHICommandList& RHICmdList)
 	{
-		RHICmdList.SetViewport(0, 0, 0.0f, View.ViewRect.Width() / GAODownsampleFactor, View.ViewRect.Height() / GAODownsampleFactor, 1.0f);
+		FIntPoint AOBufferSize = GetBufferSizeForAO(View);
+
+		RHICmdList.SetViewport(0, 0, 0.0f, AOBufferSize.X, AOBufferSize.Y, 1.0f);
 
 		FGraphicsPipelineStateInitializer GraphicsPSOInit;
 		RHICmdList.ApplyCachedRenderTargets(GraphicsPSOInit);
@@ -259,16 +257,14 @@ void GeometryAwareUpsample(FRDGBuilder& GraphBuilder, const FViewInfo& View, FRD
 
 		SetShaderParameters(RHICmdList, PixelShader, PixelShader.GetPixelShader(), *PassParameters);
 
-		FIntPoint BufferSize = GetBufferSizeForAO(View);
-
 		DrawRectangle(
 			RHICmdList,
 			0, 0,
-			View.ViewRect.Width() / GAODownsampleFactor, View.ViewRect.Height() / GAODownsampleFactor,
+			AOBufferSize.X, AOBufferSize.Y,
 			0, 0,
-			View.ViewRect.Width() / GAODownsampleFactor, View.ViewRect.Height() / GAODownsampleFactor,
-			FIntPoint(View.ViewRect.Width() / GAODownsampleFactor, View.ViewRect.Height() / GAODownsampleFactor),
-			BufferSize,
+			AOBufferSize.X, AOBufferSize.Y,
+			AOBufferSize,
+			AOBufferSize,
 			VertexShader);
 	});
 }
@@ -320,18 +316,16 @@ void UpdateHistory(
 {
 	const FRDGSystemTextures& SystemTextures = FRDGSystemTextures::Get(GraphBuilder);
 
-	const FIntPoint SceneTextureExtent = View.GetSceneTexturesConfig().Extent;
+	const FIntPoint AOBufferSize = GetBufferSizeForAO(View);
 
 	if (BentNormalHistoryState && DistanceFieldAOUseHistory(View))
 	{
-		FIntPoint BufferSize = GetBufferSizeForAO(View);
-
 		if (*BentNormalHistoryState 
 			&& !View.bCameraCut 
 			&& !View.bPrevTransformsReset
 			&& !GAOClearHistory
 			// If the scene render targets reallocate, toss the history so we don't read uninitialized data
-			&& (*BentNormalHistoryState)->GetDesc().Extent == BufferSize)
+			&& (*BentNormalHistoryState)->GetDesc().Extent == AOBufferSize)
 		{
 			FRDGTextureRef BentNormalHistoryTexture = GraphBuilder.RegisterExternalTexture(*BentNormalHistoryState);
 
@@ -358,10 +352,8 @@ void UpdateHistory(
 				PassParameters->UseHistoryFilter = UseAOHistoryStabilityPass() ? 1.0f : 0.0f;
 
 				{
-					FIntPoint HistoryBufferSize = View.GetSceneTexturesConfig().Extent / FIntPoint(GAODownsampleFactor, GAODownsampleFactor);
-
-					const float InvBufferSizeX = 1.0f / HistoryBufferSize.X;
-					const float InvBufferSizeY = 1.0f / HistoryBufferSize.Y;
+					const float InvBufferSizeX = 1.0f / AOBufferSize.X;
+					const float InvBufferSizeY = 1.0f / AOBufferSize.Y;
 
 					const FVector4f HistoryScreenPositionScaleBiasValue(
 						PrevHistoryViewRect.Width() * InvBufferSizeX / +2.0f,
@@ -391,10 +383,10 @@ void UpdateHistory(
 					RDG_EVENT_NAME("UpdateHistory"),
 					PassParameters,
 					ERDGPassFlags::Raster,
-					[PassParameters, VertexShader, PixelShader, &View, SceneTextureExtent]
+					[PassParameters, VertexShader, PixelShader, &View, AOBufferSize]
 					(FRHICommandList& RHICmdList)
 				{
-					RHICmdList.SetViewport(0, 0, 0.0f, View.ViewRect.Width() / GAODownsampleFactor, View.ViewRect.Height() / GAODownsampleFactor, 1.0f);
+					RHICmdList.SetViewport(0, 0, 0.0f, AOBufferSize.X, AOBufferSize.Y, 1.0f);
 
 					FGraphicsPipelineStateInitializer GraphicsPSOInit;
 					RHICmdList.ApplyCachedRenderTargets(GraphicsPSOInit);
@@ -413,11 +405,11 @@ void UpdateHistory(
 					DrawRectangle(
 						RHICmdList,
 						0, 0,
-						View.ViewRect.Width() / GAODownsampleFactor, View.ViewRect.Height() / GAODownsampleFactor,
+						AOBufferSize.X, AOBufferSize.Y,
 						0, 0,
-						View.ViewRect.Width() / GAODownsampleFactor, View.ViewRect.Height() / GAODownsampleFactor,
-						FIntPoint(View.ViewRect.Width() / GAODownsampleFactor, View.ViewRect.Height() / GAODownsampleFactor),
-						SceneTextureExtent / FIntPoint(GAODownsampleFactor, GAODownsampleFactor),
+						AOBufferSize.X, AOBufferSize.Y,
+						AOBufferSize,
+						AOBufferSize,
 						VertexShader);
 				});
 			}
@@ -427,7 +419,7 @@ void UpdateHistory(
 				const FRDGTextureDesc& HistoryDesc = BentNormalHistoryTexture->Desc;
 
 				// Reallocate history if buffer sizes have changed
-				if (HistoryDesc.Extent != SceneTextureExtent / FIntPoint(GAODownsampleFactor, GAODownsampleFactor))
+				if (HistoryDesc.Extent != AOBufferSize)
 				{
 					GRenderTargetPool.FreeUnusedResource(*BentNormalHistoryState);
 					*BentNormalHistoryState = nullptr;
@@ -435,26 +427,15 @@ void UpdateHistory(
 					AllocateOrReuseAORenderTarget(GraphBuilder, View, BentNormalHistoryTexture, BentNormalHistoryRTName, PF_FloatRGBA);
 				}
 
-				const bool bManuallyClampUV = View.ViewRect.Min != FIntPoint::ZeroValue || View.ViewRect.Max != SceneTextureExtent;
-
-				FFilterHistoryPS::FPermutationDomain PermutationVector;
-				PermutationVector.Set<FFilterHistoryPS::FManuallyClampUV>(bManuallyClampUV);
-
 				auto VertexShader = View.ShaderMap->GetShader<FPostProcessVS>();
-				auto PixelShader = View.ShaderMap->GetShader<FFilterHistoryPS>(PermutationVector);
-
-				const FIntPoint DownsampledBufferSize(View.GetSceneTexturesConfig().Extent / FIntPoint(GAODownsampleFactor, GAODownsampleFactor));
-				FVector2f MaxSampleBufferUV(
-					(View.ViewRect.Width() / GAODownsampleFactor - 0.5f - GAODownsampleFactor) / DownsampledBufferSize.X,
-					(View.ViewRect.Height() / GAODownsampleFactor - 0.5f - GAODownsampleFactor) / DownsampledBufferSize.Y);
+				auto PixelShader = View.ShaderMap->GetShader<FFilterHistoryPS>();
 
 				auto* PassParameters = GraphBuilder.AllocParameters<FFilterHistoryPS::FParameters>();
 				PassParameters->View = View.ViewUniformBuffer;
 				PassParameters->SceneTextures = SceneTexturesUniformBuffer;
 				PassParameters->BentNormalAOTexture = NewBentNormalHistory;
 				PassParameters->BentNormalAOSampler = TStaticSamplerState<SF_Bilinear, AM_Clamp, AM_Clamp, AM_Clamp>::GetRHI();
-				PassParameters->BentNormalAOTexelSize = FVector2f(1.0f / DownsampledBufferSize.X, 1.0f / DownsampledBufferSize.Y);
-				PassParameters->MaxSampleBufferUV = MaxSampleBufferUV;
+				PassParameters->BentNormalAOTexelSize = FVector2f(1.0f / AOBufferSize.X, 1.0f / AOBufferSize.Y);
 				PassParameters->HistoryWeight = GAOHistoryWeight;
 				PassParameters->RenderTargets[0] = FRenderTargetBinding(BentNormalHistoryTexture, ERenderTargetLoadAction::ELoad);
 
@@ -464,9 +445,9 @@ void UpdateHistory(
 					RDG_EVENT_NAME("UpdateHistoryStability"),
 					PassParameters,
 					ERDGPassFlags::Raster,
-					[PassParameters, VertexShader, PixelShader, &View, SceneTextureExtent](FRHICommandList& RHICmdList)
+					[PassParameters, VertexShader, PixelShader, &View, AOBufferSize](FRHICommandList& RHICmdList)
 				{
-					RHICmdList.SetViewport(0, 0, 0.0f, View.ViewRect.Width() / GAODownsampleFactor, View.ViewRect.Height() / GAODownsampleFactor, 1.0f);
+					RHICmdList.SetViewport(0, 0, 0.0f, AOBufferSize.X, AOBufferSize.Y, 1.0f);
 
 					FGraphicsPipelineStateInitializer GraphicsPSOInit;
 					RHICmdList.ApplyCachedRenderTargets(GraphicsPSOInit);
@@ -485,11 +466,11 @@ void UpdateHistory(
 					DrawRectangle(
 						RHICmdList,
 						0, 0,
-						View.ViewRect.Width() / GAODownsampleFactor, View.ViewRect.Height() / GAODownsampleFactor,
+						AOBufferSize.X, AOBufferSize.Y,
 						0, 0,
-						View.ViewRect.Width() / GAODownsampleFactor, View.ViewRect.Height() / GAODownsampleFactor,
-						FIntPoint(View.ViewRect.Width() / GAODownsampleFactor, View.ViewRect.Height() / GAODownsampleFactor),
-						SceneTextureExtent / FIntPoint(GAODownsampleFactor, GAODownsampleFactor),
+						AOBufferSize.X, AOBufferSize.Y,
+						AOBufferSize ,
+						AOBufferSize,
 						VertexShader);
 				});
 
@@ -507,7 +488,7 @@ void UpdateHistory(
 		{
 			// Use the current frame's upscaled mask for next frame's history
 			FRDGTextureRef DistanceFieldAOBentNormal = nullptr;
-			AllocateOrReuseAORenderTarget(GraphBuilder, View, DistanceFieldAOBentNormal, TEXT("DistanceFieldBentNormalAO"), PF_FloatRGBA, GFastVRamConfig.DistanceFieldAOBentNormal);
+			AllocateOrReuseAORenderTarget(GraphBuilder, View, DistanceFieldAOBentNormal, TEXT("PerViewDistanceFieldBentNormalAO"), PF_FloatRGBA, GFastVRamConfig.DistanceFieldAOHistory);
 
 			GeometryAwareUpsample(GraphBuilder, View, DistanceFieldAOBentNormal, DistanceFieldNormal, BentNormalInterpolation, Parameters);
 
@@ -516,14 +497,13 @@ void UpdateHistory(
 		}
 
 		DistanceFieldAOHistoryViewRect->Min = FIntPoint::ZeroValue;
-		DistanceFieldAOHistoryViewRect->Max.X = View.ViewRect.Size().X / GAODownsampleFactor;
-		DistanceFieldAOHistoryViewRect->Max.Y = View.ViewRect.Size().Y / GAODownsampleFactor;
+		DistanceFieldAOHistoryViewRect->Max = View.ViewRect.Size() / FIntPoint(GAODownsampleFactor, GAODownsampleFactor);
 	}
 	else
 	{
 		// Temporal reprojection is disabled or there is no view state - just upscale
 		FRDGTextureRef DistanceFieldAOBentNormal = nullptr;
-		AllocateOrReuseAORenderTarget(GraphBuilder, View, DistanceFieldAOBentNormal, TEXT("DistanceFieldBentNormalAO"), PF_FloatRGBA, GFastVRamConfig.DistanceFieldAOBentNormal);
+		AllocateOrReuseAORenderTarget(GraphBuilder, View, DistanceFieldAOBentNormal, TEXT("PerViewDistanceFieldBentNormalAO"), PF_FloatRGBA, GFastVRamConfig.DistanceFieldAOHistory);
 
 		GeometryAwareUpsample(GraphBuilder, View, DistanceFieldAOBentNormal, DistanceFieldNormal, BentNormalInterpolation, Parameters);
 
@@ -559,38 +539,35 @@ IMPLEMENT_GLOBAL_SHADER(FDistanceFieldAOUpsamplePS, "/Engine/Private/DistanceFie
 
 void UpsampleBentNormalAO(
 	FRDGBuilder& GraphBuilder,
-	const TArray<FViewInfo>& Views,
+	const FViewInfo& View,
 	TRDGUniformBufferRef<FSceneTextureUniformParameters> SceneTexturesUniformBuffer,
 	FRDGTextureRef SceneColorTexture,
 	FRDGTextureRef DistanceFieldAOBentNormal,
 	bool bModulateSceneColor)
 {
-	for (int32 ViewIndex = 0; ViewIndex < Views.Num(); ViewIndex++)
-	{
-		const FViewInfo& View = Views[ViewIndex];
-		FScene* Scene = (FScene*)View.Family->Scene;
-
-		RDG_GPU_MASK_SCOPE(GraphBuilder, View.GPUMask);
-
-		auto* PassParameters = GraphBuilder.AllocParameters<FDistanceFieldAOUpsamplePS::FParameters>();
-		PassParameters->View = View.ViewUniformBuffer;
-		PassParameters->SceneTextures = SceneTexturesUniformBuffer;
-		PassParameters->DFAOUpsampleParameters = DistanceField::SetupAOUpsampleParameters(View, DistanceFieldAOBentNormal ? DistanceFieldAOBentNormal : GSystemTextures.GetWhiteDummy(GraphBuilder));
-		PassParameters->MinIndirectDiffuseOcclusion = Scene->SkyLight ? Scene->SkyLight->MinOcclusion : 0;
-		PassParameters->RenderTargets[0] = FRenderTargetBinding(SceneColorTexture, ERenderTargetLoadAction::ELoad);
-		PassParameters->DistanceFieldAOBentNormal = DistanceFieldAOBentNormal;
-
-		TShaderMapRef<FPostProcessVS> VertexShader(View.ShaderMap);
-
-		FDistanceFieldAOUpsamplePS::FPermutationDomain PermutationVector;
-		PermutationVector.Set<FDistanceFieldAOUpsamplePS::FModulateToSceneColorDim>(bModulateSceneColor);
-		TShaderMapRef<FDistanceFieldAOUpsamplePS> PixelShader(View.ShaderMap, PermutationVector);
-
-		GraphBuilder.AddPass(
-			RDG_EVENT_NAME("UpsampleAO"),
-			PassParameters,
-			ERDGPassFlags::Raster,
-			[VertexShader, PixelShader, PassParameters, &View, DistanceFieldAOBentNormal, bModulateSceneColor](FRHICommandList& RHICmdList)
+	FScene* Scene = (FScene*)View.Family->Scene;
+
+	RDG_GPU_MASK_SCOPE(GraphBuilder, View.GPUMask);
+
+	auto* PassParameters = GraphBuilder.AllocParameters<FDistanceFieldAOUpsamplePS::FParameters>();
+	PassParameters->View = View.ViewUniformBuffer;
+	PassParameters->SceneTextures = SceneTexturesUniformBuffer;
+	PassParameters->DFAOUpsampleParameters = DistanceField::SetupAOUpsampleParameters(View, DistanceFieldAOBentNormal ? DistanceFieldAOBentNormal : GSystemTextures.GetWhiteDummy(GraphBuilder));
+	PassParameters->MinIndirectDiffuseOcclusion = Scene->SkyLight ? Scene->SkyLight->MinOcclusion : 0;
+	PassParameters->RenderTargets[0] = FRenderTargetBinding(SceneColorTexture, ERenderTargetLoadAction::ELoad);
+	PassParameters->DistanceFieldAOBentNormal = DistanceFieldAOBentNormal;
+
+	TShaderMapRef<FPostProcessVS> VertexShader(View.ShaderMap);
+
+	FDistanceFieldAOUpsamplePS::FPermutationDomain PermutationVector;
+	PermutationVector.Set<FDistanceFieldAOUpsamplePS::FModulateToSceneColorDim>(bModulateSceneColor);
+	TShaderMapRef<FDistanceFieldAOUpsamplePS> PixelShader(View.ShaderMap, PermutationVector);
+
+	GraphBuilder.AddPass(
+		RDG_EVENT_NAME("UpsampleAO"),
+		PassParameters,
+		ERDGPassFlags::Raster,
+		[VertexShader, PixelShader, PassParameters, &View, DistanceFieldAOBentNormal, bModulateSceneColor](FRHICommandList& RHICmdList)
 		{
 			RHICmdList.SetViewport(View.ViewRect.Min.X, View.ViewRect.Min.Y, 0.0f, View.ViewRect.Max.X, View.ViewRect.Max.Y, 1.0f);
 
@@ -625,5 +602,4 @@ void UpsampleBentNormalAO(
 				GetBufferSizeForAO(View),
 				VertexShader);
 		});
-	}
 }
diff --git a/Engine/Source/Runtime/Renderer/Private/DistanceFieldLightingPost.h b/Engine/Source/Runtime/Renderer/Private/DistanceFieldLightingPost.h
--- a/Engine/Source/Runtime/Renderer/Private/DistanceFieldLightingPost.h
+++ b/Engine/Source/Runtime/Renderer/Private/DistanceFieldLightingPost.h
@@ -30,7 +30,7 @@ extern void UpdateHistory(
 
 extern void UpsampleBentNormalAO(
 	FRDGBuilder& GraphBuilder,
-	const TArray<FViewInfo>& Views,
+	const FViewInfo& View,
 	TRDGUniformBufferRef<FSceneTextureUniformParameters> SceneTexturesUniformBuffer,
 	FRDGTextureRef SceneColorTexture,
 	FRDGTextureRef DistanceFieldAOBentNormal,
diff --git a/Engine/Source/Runtime/Renderer/Private/DistanceFieldObjectCulling.cpp b/Engine/Source/Runtime/Renderer/Private/DistanceFieldObjectCulling.cpp
--- a/Engine/Source/Runtime/Renderer/Private/DistanceFieldObjectCulling.cpp
+++ b/Engine/Source/Runtime/Renderer/Private/DistanceFieldObjectCulling.cpp
@@ -337,15 +337,17 @@ void ScatterTilesToObjects(
 
 FIntPoint GetTileListGroupSizeForView(const FViewInfo& View)
 {
+	const FIntPoint AOBufferSize = GetBufferSizeForAO(View);
 	return FIntPoint(
-		FMath::DivideAndRoundUp(FMath::Max(View.ViewRect.Size().X / GAODownsampleFactor, 1), GDistanceFieldAOTileSizeX),
-		FMath::DivideAndRoundUp(FMath::Max(View.ViewRect.Size().Y / GAODownsampleFactor, 1), GDistanceFieldAOTileSizeY));
+		FMath::DivideAndRoundUp(FMath::Max(AOBufferSize.X, 1), GDistanceFieldAOTileSizeX),
+		FMath::DivideAndRoundUp(FMath::Max(AOBufferSize.Y, 1), GDistanceFieldAOTileSizeY));
 }
 
 void BuildTileObjectLists(
 	FRDGBuilder& GraphBuilder,
 	const FScene& Scene,
-	TArray<FViewInfo>& Views,
+	const FViewInfo& View,
+	TRDGUniformBufferRef<FSceneTextureUniformParameters> SceneTexturesUniformBuffer,
 	FRDGBufferRef ObjectIndirectArguments,
 	const FDistanceFieldCulledObjectBufferParameters& CulledObjectBufferParameters,
 	FTileIntersectionParameters TileIntersectionParameters,
@@ -355,64 +357,57 @@ void BuildTileObjectLists(
 	ensure(GAOScatterTileCulling);
 
 	RDG_EVENT_SCOPE(GraphBuilder, "BuildTileList");
+	RDG_GPU_MASK_SCOPE(GraphBuilder, View.GPUMask);
 
-	TRDGUniformBufferRef<FSceneTextureUniformParameters> SceneTexturesUniformBuffer = GetViewFamilyInfo(Views).GetSceneTextures().UniformBuffer;
+	const FIntPoint TileListGroupSize = GetTileListGroupSizeForView(View);
 
-	for (int32 ViewIndex = 0; ViewIndex < Views.Num(); ViewIndex++)
 	{
-		const FViewInfo& View = Views[ViewIndex];
-		RDG_GPU_MASK_SCOPE(GraphBuilder, View.GPUMask);
+		auto* PassParameters = GraphBuilder.AllocParameters<FBuildTileConesCS::FParameters>();
+		PassParameters->View = View.ViewUniformBuffer;
+		PassParameters->RWTileConeAxisAndCos = TileIntersectionParameters.RWTileConeAxisAndCos;
+		PassParameters->RWTileConeDepthRanges = TileIntersectionParameters.RWTileConeDepthRanges;
+		PassParameters->SceneTextures = SceneTexturesUniformBuffer;
+		PassParameters->DistanceFieldNormalTexture = DistanceFieldNormal;
+		PassParameters->DistanceFieldNormalSampler = TStaticSamplerState<SF_Point, AM_Clamp, AM_Clamp, AM_Clamp>::GetRHI();
+		PassParameters->AOParameters = DistanceField::SetupAOShaderParameters(Parameters);
+		PassParameters->NumGroups = FVector2f(TileListGroupSize.X, TileListGroupSize.Y);
 
-		const FIntPoint TileListGroupSize = GetTileListGroupSizeForView(View);
+		auto ComputeShader = View.ShaderMap->GetShader<FBuildTileConesCS>();
 
-		{
-			auto* PassParameters = GraphBuilder.AllocParameters<FBuildTileConesCS::FParameters>();
-			PassParameters->View = View.ViewUniformBuffer;
-			PassParameters->RWTileConeAxisAndCos = TileIntersectionParameters.RWTileConeAxisAndCos;
-			PassParameters->RWTileConeDepthRanges = TileIntersectionParameters.RWTileConeDepthRanges;
-			PassParameters->SceneTextures = SceneTexturesUniformBuffer;
-			PassParameters->DistanceFieldNormalTexture = DistanceFieldNormal;
-			PassParameters->DistanceFieldNormalSampler = TStaticSamplerState<SF_Point, AM_Clamp, AM_Clamp, AM_Clamp>::GetRHI();
-			PassParameters->AOParameters = DistanceField::SetupAOShaderParameters(Parameters);
-			PassParameters->NumGroups = FVector2f(TileListGroupSize.X, TileListGroupSize.Y);
-
-			auto ComputeShader = View.ShaderMap->GetShader<FBuildTileConesCS>();
-
-			FComputeShaderUtils::AddPass(GraphBuilder, RDG_EVENT_NAME("BuildTileCones"), ComputeShader, PassParameters, FIntVector(TileListGroupSize.X, TileListGroupSize.Y, 1));
-		}
+		FComputeShaderUtils::AddPass(GraphBuilder, RDG_EVENT_NAME("BuildTileCones"), ComputeShader, PassParameters, FIntVector(TileListGroupSize.X, TileListGroupSize.Y, 1));
+	}
 
-		// Start at 0 tiles per object
-		AddClearUAVPass(GraphBuilder, TileIntersectionParameters.RWNumCulledTilesArray, 0);
+	// Start at 0 tiles per object
+	AddClearUAVPass(GraphBuilder, TileIntersectionParameters.RWNumCulledTilesArray, 0);
 
-		// Rasterize object bounding shapes and intersect with screen tiles to compute how many tiles intersect each object
-		ScatterTilesToObjects(GraphBuilder, true, View, Scene.DistanceFieldSceneData, TileListGroupSize, Parameters, ObjectIndirectArguments, CulledObjectBufferParameters, TileIntersectionParameters, SceneTexturesUniformBuffer);
+	// Start at 0 threadgroups
+	AddClearUAVPass(GraphBuilder, TileIntersectionParameters.RWObjectTilesIndirectArguments, 0);
 
-		// Start at 0 threadgroups
-		AddClearUAVPass(GraphBuilder, TileIntersectionParameters.RWObjectTilesIndirectArguments, 0);
+	// Rasterize object bounding shapes and intersect with screen tiles to compute how many tiles intersect each object
+	ScatterTilesToObjects(GraphBuilder, true, View, Scene.DistanceFieldSceneData, TileListGroupSize, Parameters, ObjectIndirectArguments, CulledObjectBufferParameters, TileIntersectionParameters, SceneTexturesUniformBuffer);
 
-		{
-			auto* PassParameters = GraphBuilder.AllocParameters<FComputeCulledTilesStartOffsetCS::FParameters>();
-			PassParameters->View = View.ViewUniformBuffer;
-			PassParameters->TileIntersectionParameters = TileIntersectionParameters;
-			PassParameters->DistanceFieldCulledObjectBuffers = CulledObjectBufferParameters;
-			PassParameters->DistanceFieldAtlas = DistanceField::SetupAtlasParameters(GraphBuilder, Scene.DistanceFieldSceneData);
-			PassParameters->SceneTextures = SceneTexturesUniformBuffer;
-
-			auto ComputeShader = View.ShaderMap->GetShader<FComputeCulledTilesStartOffsetCS>();
-			const FIntVector GroupCount = FComputeShaderUtils::GetGroupCountWrapped(Scene.DistanceFieldSceneData.NumObjectsInBuffer, ComputeStartOffsetGroupSize);
-
-			FComputeShaderUtils::AddPass(
-				GraphBuilder,
-				RDG_EVENT_NAME("ComputeStartOffsets"),
-				ComputeShader,
-				PassParameters,
-				GroupCount);
-		}
+	{
+		auto* PassParameters = GraphBuilder.AllocParameters<FComputeCulledTilesStartOffsetCS::FParameters>();
+		PassParameters->View = View.ViewUniformBuffer;
+		PassParameters->TileIntersectionParameters = TileIntersectionParameters;
+		PassParameters->DistanceFieldCulledObjectBuffers = CulledObjectBufferParameters;
+		PassParameters->DistanceFieldAtlas = DistanceField::SetupAtlasParameters(GraphBuilder, Scene.DistanceFieldSceneData);
+		PassParameters->SceneTextures = SceneTexturesUniformBuffer;
 
-		// Start at 0 tiles per object
-		AddClearUAVPass(GraphBuilder, TileIntersectionParameters.RWNumCulledTilesArray, 0);
+		auto ComputeShader = View.ShaderMap->GetShader<FComputeCulledTilesStartOffsetCS>();
+		const FIntVector GroupCount = FComputeShaderUtils::GetGroupCountWrapped(Scene.DistanceFieldSceneData.NumObjectsInBuffer, ComputeStartOffsetGroupSize);
 
-		// Rasterize object bounding shapes and intersect with screen tiles, and write out intersecting tile indices for the cone tracing pass
-		ScatterTilesToObjects(GraphBuilder, false, View, Scene.DistanceFieldSceneData, TileListGroupSize, Parameters, ObjectIndirectArguments, CulledObjectBufferParameters, TileIntersectionParameters, SceneTexturesUniformBuffer);
+		FComputeShaderUtils::AddPass(
+			GraphBuilder,
+			RDG_EVENT_NAME("ComputeStartOffsets"),
+			ComputeShader,
+			PassParameters,
+			GroupCount);
 	}
+
+	// Start at 0 tiles per object
+	AddClearUAVPass(GraphBuilder, TileIntersectionParameters.RWNumCulledTilesArray, 0);
+
+	// Rasterize object bounding shapes and intersect with screen tiles, and write out intersecting tile indices for the cone tracing pass
+	ScatterTilesToObjects(GraphBuilder, false, View, Scene.DistanceFieldSceneData, TileListGroupSize, Parameters, ObjectIndirectArguments, CulledObjectBufferParameters, TileIntersectionParameters, SceneTexturesUniformBuffer);
 }
