ShaderPrint - exposed ShaderPrintData.IsDrawLocked to shaders in order to be able to stop drawing new stuff on top of a locked visualization.

diff --git a/Engine/Source/Runtime/Renderer/Private/ShaderPrint.cpp b/Engine/Source/Runtime/Renderer/Private/ShaderPrint.cpp
--- a/Engine/Source/Runtime/Renderer/Private/ShaderPrint.cpp
+++ b/Engine/Source/Runtime/Renderer/Private/ShaderPrint.cpp
@@ -201,6 +201,7 @@ namespace ShaderPrint
 		OutParameters.MaxStateCount = InSetup.MaxStateCount;
 		OutParameters.MaxLineCount = InSetup.MaxLineCount;
 		OutParameters.MaxTriangleCount = InSetup.MaxTriangleCount;
+		OutParameters.IsDrawLocked = InSetup.bIsDrawLocked ? 1 : 0;
 		OutParameters.TranslatedWorldOffset = FVector3f(InSetup.PreViewTranslation);
 	}
 
@@ -604,6 +605,7 @@ namespace ShaderPrint
 		MaxStateCount = bEnabled ? GetMaxWidgetCount() : 0;
 		MaxLineCount = bEnabled ? GetMaxLineCount() : 0;
 		MaxTriangleCount = bEnabled ? GetMaxTriangleCount(): 0;
+		bIsDrawLocked = false;
 	}
 
 	FShaderPrintSetup::FShaderPrintSetup(FSceneView const& View)
@@ -621,6 +623,7 @@ namespace ShaderPrint
 		MaxStateCount = bEnabled ? GetMaxWidgetCount() : 0;
 		MaxLineCount = bEnabled ? GetMaxLineCount() : 0;
 		MaxTriangleCount = bEnabled ? GetMaxTriangleCount() : 0;
+		bIsDrawLocked = View.State ? ((const FSceneViewState*)View.State)->ShaderPrintStateData.bIsLocked : false;
 	}
 
 	FShaderPrintData CreateShaderPrintData(FRDGBuilder& GraphBuilder, FShaderPrintSetup const& InSetup, FSceneViewState* InViewState)
diff --git a/Engine/Source/Runtime/Renderer/Public/ShaderPrintParameters.h b/Engine/Source/Runtime/Renderer/Public/ShaderPrintParameters.h
--- a/Engine/Source/Runtime/Renderer/Public/ShaderPrintParameters.h
+++ b/Engine/Source/Runtime/Renderer/Public/ShaderPrintParameters.h
@@ -35,6 +35,7 @@ namespace ShaderPrint
 		SHADER_PARAMETER(uint32, MaxStateCount)
 		SHADER_PARAMETER(uint32, MaxLineCount)
 		SHADER_PARAMETER(uint32, MaxTriangleCount)
+		SHADER_PARAMETER(uint32, IsDrawLocked)
 	END_GLOBAL_SHADER_PARAMETER_STRUCT()
 
 	// ShaderPrint parameter struct declaration
@@ -115,6 +116,8 @@ namespace ShaderPrint
 		uint32 MaxLineCount = 0;
 		/** Initial size of triangle buffer. Will also be increased by RequestSpaceForLines(). */
 		uint32 MaxTriangleCount = 0;
+		// Whether current draw is locked or not. Useful to stop rendering new stuff on top of the history. */
+		bool bIsDrawLocked = false;
 	};
 
 	/** Fill the ShaderPrintCommonParameters uniform buffer structure for our setup. */
