Fix engine crash on medium scalability due to uninitialized shader parameter.

diff --git a/Engine/Source/Runtime/Renderer/Private/DistanceFieldAmbientOcclusion.cpp b/Engine/Source/Runtime/Renderer/Private/DistanceFieldAmbientOcclusion.cpp
--- a/Engine/Source/Runtime/Renderer/Private/DistanceFieldAmbientOcclusion.cpp
+++ b/Engine/Source/Runtime/Renderer/Private/DistanceFieldAmbientOcclusion.cpp
@@ -855,9 +855,9 @@ void FDeferredShadingSceneRenderer::RenderDistanceFieldLighting(
 				CulledObjectBufferParameters);
 
 			CullObjectsToView(GraphBuilder, *Scene, View, Parameters, CulledObjectBufferParameters);
-		}
 
-		BuildTileObjectLists(GraphBuilder, *Scene, View, SceneTextures.UniformBuffer, ObjectIndirectArguments, CulledObjectBufferParameters, TileIntersectionParameters, DistanceFieldNormal, Parameters);
+			BuildTileObjectLists(GraphBuilder, *Scene, View, SceneTextures.UniformBuffer, ObjectIndirectArguments, CulledObjectBufferParameters, TileIntersectionParameters, DistanceFieldNormal, Parameters);
+		}
 
 		// Render to a per-view BentNormal first because it also needs to be stored by the view history.
 		// The per-view output can be copied back to the appropriate region in the output BentNormal texture.
