Fixed mismatch between Texture2DArray and Texture2D in FSR

diff --git a/Engine/Source/Runtime/Renderer/Private/PostProcess/TemporalSuperResolution.cpp b/Engine/Source/Runtime/Renderer/Private/PostProcess/TemporalSuperResolution.cpp
--- a/Engine/Source/Runtime/Renderer/Private/PostProcess/TemporalSuperResolution.cpp
+++ b/Engine/Source/Runtime/Renderer/Private/PostProcess/TemporalSuperResolution.cpp
@@ -838,7 +838,7 @@ class FTSRUpdateHistoryCS : public FTSRShader
 		SHADER_PARAMETER_RDG_TEXTURE_UAV(RWTexture2D, SceneColorOutputMip0)
 		SHADER_PARAMETER_RDG_TEXTURE_UAV(RWTexture2D, SceneColorOutputMip1)
 		SHADER_PARAMETER_RDG_TEXTURE_UAV(RWTexture2DArray, HistoryColorOutput)
-		SHADER_PARAMETER_RDG_TEXTURE_UAV(RWTexture2DArray, HistoryMetadataOutput)
+		SHADER_PARAMETER_RDG_TEXTURE_UAV(RWTexture2D, HistoryMetadataOutput)
 		SHADER_PARAMETER_RDG_TEXTURE_UAV(RWTexture2DArray, DebugOutput)
 	END_SHADER_PARAMETER_STRUCT()
 
@@ -2303,7 +2303,8 @@ FDefaultTemporalUpscaler::FOutputs AddTemporalSuperResolutionPasses(
 			FRDGTextureUAVDesc MetadataUAVDesc(History.MetadataArray);
 			MetadataUAVDesc.FirstArraySlice = CurrentFrameSliceIndex;
 			MetadataUAVDesc.NumArraySlices = 1;
-
+            MetadataUAVDesc.DimensionOverride = ETextureDimension::Texture2D;
+            
 			PassParameters->HistoryArrayIndices = HistoryArrayIndices;
 			PassParameters->HistoryColorOutput = GraphBuilder.CreateUAV(ColorUAVDesc);
 			PassParameters->HistoryMetadataOutput = GraphBuilder.CreateUAV(MetadataUAVDesc);
