Invalidate instance culling oncclusion query results when GPUScene instance slots are reallocated

diff --git a/Engine/Source/Runtime/Renderer/Private/GPUScene.cpp b/Engine/Source/Runtime/Renderer/Private/GPUScene.cpp
--- a/Engine/Source/Runtime/Renderer/Private/GPUScene.cpp
+++ b/Engine/Source/Runtime/Renderer/Private/GPUScene.cpp
@@ -26,6 +26,7 @@
 #include "SceneDefinitions.h"
 #include "PrimitiveSceneShaderData.h"
 #include "RendererOnScreenNotification.h"
+#include "InstanceCulling/InstanceCullingOcclusionQuery.h"
 
 // Defaults to being disabled, enable using the command line argument: -CsvCategory GPUScene
 CSV_DEFINE_CATEGORY(GPUScene, false);
@@ -764,7 +765,7 @@ void FGPUScene::UpdateInternal(FRDGBuilder& GraphBuilder, FSceneUniformBuffer& S
 	UpdateBufferState(GraphBuilder, SceneUB, Scene, Adapter, true);
 
 	// Run a pass that clears (Sets ID to invalid) any instances that need it
-	AddClearInstancesPass(GraphBuilder);
+	AddClearInstancesPass(GraphBuilder, Scene.InstanceCullingOcclusionQueryRenderer);
 
 	// Pull out instances needing only primitive ID update, they still have to go to the general update such that the primitive gets updated (as it moved)
 	{
@@ -1592,7 +1593,7 @@ void FGPUScene::UploadDynamicPrimitiveShaderDataForViewInternal(FRDGBuilder& Gra
 		UpdateBufferState(GraphBuilder, View.GetSceneUniforms(), Scene, UploadAdapter, false);
 
 		// Run a pass that clears (Sets ID to invalid) any instances that need it.
-		AddClearInstancesPass(GraphBuilder);
+		AddClearInstancesPass(GraphBuilder, Scene.InstanceCullingOcclusionQueryRenderer);
 
 		UploadGeneral<FUploadDataSourceAdapterDynamicPrimitives>(GraphBuilder, Scene, ExternalAccessQueue, UploadAdapter);
 	}
@@ -1729,6 +1730,21 @@ void FGPUScene::UploadDynamicPrimitiveShaderDataForView(FRDGBuilder& GraphBuilde
 	}
 }
 
+// Grow the last entry in the Output using the given Range if trivial merge is possible (typical case)
+inline void AddOrMergeInstanceRange(TArray<FGPUSceneInstanceRange>& Output, FGPUSceneInstanceRange Range)
+{
+	if (!Output.IsEmpty())
+	{
+		FGPUSceneInstanceRange& Last = Output.Last();
+		if (Range.InstanceSceneDataOffset == Last.InstanceSceneDataOffset + Last.NumInstanceSceneDataEntries)
+		{
+			Last.NumInstanceSceneDataEntries += Range.NumInstanceSceneDataEntries;
+			return;
+		}
+	}
+	Output.Add(Range);
+}
+
 int32 FGPUScene::AllocateInstanceSceneDataSlots(int32 NumInstanceSceneDataEntries)
 {
 	LLM_SCOPE_BYTAG(GPUScene);
@@ -1738,7 +1754,7 @@ int32 FGPUScene::AllocateInstanceSceneDataSlots(int32 NumInstanceSceneDataEntrie
 		if (NumInstanceSceneDataEntries > 0)
 		{
 			int32 InstanceSceneDataOffset = InstanceSceneDataAllocator.Allocate(NumInstanceSceneDataEntries);
-			InstanceRangesToClear.Add(FInstanceRange{ uint32(InstanceSceneDataOffset), uint32(NumInstanceSceneDataEntries) });
+			AddOrMergeInstanceRange(InstanceRangesToClear, FInstanceRange{ uint32(InstanceSceneDataOffset), uint32(NumInstanceSceneDataEntries) });
 #if LOG_INSTANCE_ALLOCATIONS
 			UE_LOG(LogTemp, Warning, TEXT("AllocateInstanceSceneDataSlots: [%6d,%6d)"), InstanceSceneDataOffset, InstanceSceneDataOffset + NumInstanceSceneDataEntries);
 #endif
@@ -1757,7 +1773,7 @@ void FGPUScene::FreeInstanceSceneDataSlots(int32 InstanceSceneDataOffset, int32
 	if (bIsEnabled)
 	{
 		InstanceSceneDataAllocator.Free(InstanceSceneDataOffset, NumInstanceSceneDataEntries);
-		InstanceRangesToClear.Add(FInstanceRange{ uint32(InstanceSceneDataOffset), uint32(NumInstanceSceneDataEntries) });
+		AddOrMergeInstanceRange(InstanceRangesToClear, FInstanceRange{ uint32(InstanceSceneDataOffset), uint32(NumInstanceSceneDataEntries) });
 #if LOG_INSTANCE_ALLOCATIONS
 		UE_LOG(LogTemp, Warning, TEXT("FreeInstanceSceneDataSlots: [%6d,%6d)"), InstanceSceneDataOffset, InstanceSceneDataOffset + NumInstanceSceneDataEntries);
 #endif
@@ -2182,7 +2198,7 @@ void FGPUSceneCompactInstanceData::Init(const FScene* Scene, int32 PrimitiveId)
 	InstanceAuxData			= FVector4f(0);
 }
 
-void FGPUScene::AddClearInstancesPass(FRDGBuilder& GraphBuilder)
+void FGPUScene::AddClearInstancesPass(FRDGBuilder& GraphBuilder, FInstanceCullingOcclusionQueryRenderer* OcclusionQueryRenderer)
 {
 	FInstanceGPULoadBalancer ClearIdData;
 #if LOG_INSTANCE_ALLOCATIONS
@@ -2203,6 +2219,13 @@ void FGPUScene::AddClearInstancesPass(FRDGBuilder& GraphBuilder)
 #endif
 		}
 	}
+
+	if (OcclusionQueryRenderer)
+	{
+		// Invalidate last frame occlusion query slots
+		OcclusionQueryRenderer->MarkInstancesVisible(GraphBuilder, InstanceRangesToClear);
+	}
+
 #if LOG_INSTANCE_ALLOCATIONS
 	UE_LOG(LogTemp, Warning, TEXT("AddClearInstancesPass: \n%s"), *RangesStr);
 #endif
diff --git a/Engine/Source/Runtime/Renderer/Private/GPUScene.h b/Engine/Source/Runtime/Renderer/Private/GPUScene.h
--- a/Engine/Source/Runtime/Renderer/Private/GPUScene.h
+++ b/Engine/Source/Runtime/Renderer/Private/GPUScene.h
@@ -22,6 +22,7 @@ class FLightSceneInfoCompact;
 class FGPUScene;
 class FGPUSceneDynamicContext;
 class FViewUniformShaderParameters;
+class FInstanceCullingOcclusionQueryRenderer;
 
 BEGIN_SHADER_PARAMETER_STRUCT(FGPUSceneResourceParameters, RENDERER_API)
 	SHADER_PARAMETER_RDG_BUFFER_SRV(StructuredBuffer<float4>, GPUSceneInstanceSceneData)
@@ -183,6 +184,12 @@ struct FGPUSceneBufferState
 	FRDGBuffer* LightDataBuffer = nullptr;
 };
 
+struct FGPUSceneInstanceRange
+{
+	uint32 InstanceSceneDataOffset;
+	uint32 NumInstanceSceneDataEntries;
+};
+
 class FGPUScene
 {
 public:
@@ -353,11 +360,7 @@ public:
 	TRefCountPtr<FRDGPooledBuffer> LightmapDataBuffer;
 	FRDGAsyncScatterUploadBuffer   LightmapUploadBuffer;
 
-	struct FInstanceRange
-	{
-		uint32 InstanceSceneDataOffset;
-		uint32 NumInstanceSceneDataEntries;
-	};
+	using FInstanceRange = FGPUSceneInstanceRange;
 
 	TArray<FInstanceRange> DynamicPrimitiveInstancesToInvalidate;
 
@@ -426,7 +429,7 @@ private:
 
 	void AddUpdatePrimitiveIdsPass(FRDGBuilder& GraphBuilder, FInstanceGPULoadBalancer& IdOnlyUpdateItems);
 
-	void AddClearInstancesPass(FRDGBuilder& GraphBuilder);
+	void AddClearInstancesPass(FRDGBuilder& GraphBuilder, FInstanceCullingOcclusionQueryRenderer* OcclusionQueryRenderer = nullptr);
 
 
 #if !UE_BUILD_SHIPPING
diff --git a/Engine/Source/Runtime/Renderer/Private/InstanceCulling/InstanceCullingOcclusionQuery.cpp b/Engine/Source/Runtime/Renderer/Private/InstanceCulling/InstanceCullingOcclusionQuery.cpp
--- a/Engine/Source/Runtime/Renderer/Private/InstanceCulling/InstanceCullingOcclusionQuery.cpp
+++ b/Engine/Source/Runtime/Renderer/Private/InstanceCulling/InstanceCullingOcclusionQuery.cpp
@@ -17,6 +17,7 @@
 #include "ShaderParameterStruct.h"
 #include "SceneRendering.h"
 #include "ScenePrivate.h"
+#include "UnifiedBuffer.h"
 #include "HAL/IConsoleManager.h"
 #include "DataDrivenShaderPlatformInfo.h"
 
@@ -398,6 +399,31 @@ uint32 FInstanceCullingOcclusionQueryRenderer::Render(
 	return ViewMask;
 }
 
+void FInstanceCullingOcclusionQueryRenderer::MarkInstancesVisible(FRDGBuilder& GraphBuilder, TConstArrayView<FGPUSceneInstanceRange> Ranges)
+{
+	if (!InstanceOcclusionQueryBuffer)
+	{
+		// Previous frame buffer does not exist, nothing to clear
+		return;
+	}
+
+	FRDGBufferRef Buffer = GraphBuilder.RegisterExternalBuffer(InstanceOcclusionQueryBuffer);
+
+	// Consecutive uses of the UAV will run in parallel.
+	// Allocating a unique RDG UAV here will still ensure that a barrier is inserted before the first dispatch.
+	FRDGBufferUAVRef UAV = GraphBuilder.CreateUAV(Buffer, PF_R32_UINT, ERDGUnorderedAccessViewFlags::SkipBarrier);
+
+	// NOTE: It is possible to make this more efficient using a specialized GPU scatter shader, if we see many small batches here in practice
+	for (FGPUSceneInstanceRange Range : Ranges)
+	{
+		FMemsetResourceParams MemsetParams;
+		MemsetParams.Value = 0xFFFFFFFF; // Mark instance visible in all views
+		MemsetParams.Count = Range.NumInstanceSceneDataEntries;
+		MemsetParams.DstOffset = Range.InstanceSceneDataOffset;
+		MemsetResource(GraphBuilder, UAV, MemsetParams);
+	}
+}
+
 void FInstanceCullingOcclusionQueryRenderer::EndFrame(FRDGBuilder& GraphBuilder)
 {
 	if (CurrentInstanceOcclusionQueryBuffer)
diff --git a/Engine/Source/Runtime/Renderer/Private/InstanceCulling/InstanceCullingOcclusionQuery.h b/Engine/Source/Runtime/Renderer/Private/InstanceCulling/InstanceCullingOcclusionQuery.h
--- a/Engine/Source/Runtime/Renderer/Private/InstanceCulling/InstanceCullingOcclusionQuery.h
+++ b/Engine/Source/Runtime/Renderer/Private/InstanceCulling/InstanceCullingOcclusionQuery.h
@@ -6,11 +6,13 @@
 #include "RenderGraphFwd.h"
 #include "RenderGraphResources.h"
 #include "Containers/Array.h"
+#include "Containers/ArrayView.h"
 
 class FRDGBuilder;
 class FGPUScene;
 class FViewInfo;
 struct IPooledRenderTarget;
+struct FGPUSceneInstanceRange;
 
 /** 
 * Implements accurate per-instance occlusion culling similar to FOcclusionFeedback, except using GPUScene.
@@ -25,6 +27,13 @@ public:
 	*/	
 	uint32 Render(FRDGBuilder& GraphBuilder, FGPUScene& GPUScene, FViewInfo& View);
 
+	/**
+	* Mark blocks of occlusion query buffer entries as visible in all views.
+	* Required, for example, when GPUScene instances are invalidated betwen frames,
+	* i.e. when when GPUScene instance buffer slot is re-allocated.
+	*/
+	void MarkInstancesVisible(FRDGBuilder& GraphBuilder, TConstArrayView<FGPUSceneInstanceRange> Ranges);
+
 	/**
 	* Perform various book-keeping tasks that must run at the end of each frame
 	* Saves view ID mapping data, extracts result buffer to be available next frame, etc.
