Fix RenderGraphValidation for localight prepass
On mobile there is a separate FMobileDirectionalLightShaderParameters UB which holds all directional light data.

diff --git a/Engine/Source/Runtime/Renderer/Private/LightGridInjection.cpp b/Engine/Source/Runtime/Renderer/Private/LightGridInjection.cpp
--- a/Engine/Source/Runtime/Renderer/Private/LightGridInjection.cpp
+++ b/Engine/Source/Runtime/Renderer/Private/LightGridInjection.cpp
@@ -583,7 +583,8 @@ FComputeLightGridOutput FSceneRenderer::ComputeLightGrid(FRDGBuilder& GraphBuild
 						ViewSpaceDirAndPreprocAngleData.Add(ViewSpaceDirAndPreprocAngle);
 					#endif // ENABLE_LIGHT_CULLING_VIEW_SPACE_BUILD_DATA
 					}
-					else if (SortedLightInfo.SortKey.Fields.LightType == LightType_Directional && ViewFamily.EngineShowFlags.DirectionalLights)
+					// On mobile there is a separate FMobileDirectionalLightShaderParameters UB which holds all directional light data.
+					else if (SortedLightInfo.SortKey.Fields.LightType == LightType_Directional && ViewFamily.EngineShowFlags.DirectionalLights && !IsMobilePlatform(View.GetShaderPlatform()))
 					{
 						// The selected forward directional light is also used for volumetric lighting using ForwardLightData UB.
 						// Also some people noticed that depending on the order a two directional lights are made visible in a level, the selected light for volumetric fog lighting will be different.
