Fix -rhivalidation error in TSR (zero sized dispatch)

diff --git a/Engine/Source/Runtime/Renderer/Private/PostProcess/TemporalSuperResolution.cpp b/Engine/Source/Runtime/Renderer/Private/PostProcess/TemporalSuperResolution.cpp
--- a/Engine/Source/Runtime/Renderer/Private/PostProcess/TemporalSuperResolution.cpp
+++ b/Engine/Source/Runtime/Renderer/Private/PostProcess/TemporalSuperResolution.cpp
@@ -1645,14 +1645,18 @@ FDefaultTemporalUpscaler::FOutputs AddTemporalSuperResolutionPasses(
 		PassParameters->SubpixelDepthOutput = GraphBuilder.CreateUAV(PrevScatteredSubpixelDepthTexture);
 		PassParameters->DebugOutput = CreateDebugUAV(InputExtent, TEXT("Debug.TSR.ForwardScatterDepth"));
 
-		TShaderMapRef<FTSRForwardScatterDepthCS> ComputeShader(View.ShaderMap);
-		FComputeShaderUtils::AddPass(
-			GraphBuilder,
-			RDG_EVENT_NAME("TSR ForwardScatterDepth %dx%d", ScatterDepthViewport.Width(), ScatterDepthViewport.Height()),
-			AsyncComputePasses >= 1 ? ERDGPassFlags::AsyncCompute : ERDGPassFlags::Compute,
-			ComputeShader,
-			PassParameters,
-			FComputeShaderUtils::GetGroupCount(ScatterDepthViewport.Size(), 8));
+		const FIntVector DispatchDim = FComputeShaderUtils::GetGroupCount(ScatterDepthViewport.Size(), 8);
+		if (DispatchDim.X > 0 && DispatchDim.Y > 0)
+		{
+			TShaderMapRef<FTSRForwardScatterDepthCS> ComputeShader(View.ShaderMap);
+			FComputeShaderUtils::AddPass(
+				GraphBuilder,
+				RDG_EVENT_NAME("TSR ForwardScatterDepth %dx%d", ScatterDepthViewport.Width(), ScatterDepthViewport.Height()),
+				AsyncComputePasses >= 1 ? ERDGPassFlags::AsyncCompute : ERDGPassFlags::Compute,
+				ComputeShader,
+				PassParameters,
+				FComputeShaderUtils::GetGroupCount(ScatterDepthViewport.Size(), 8));
+		}
 	}
 
 	// Dilate the velocity texture & scatter reprojection into previous frame
