Add debug visualization for instance culling occlusion query renderer

* Enabled using ShowFlag.VisualizeInstanceOcclusionQueries=1
* Modified FMobileSceneRenderer::ShouldRenderHZB() to take instance occlusion culling into account

diff --git a/Engine/Shaders/Private/InstanceCulling/InstanceCullingOcclusionQuery.usf b/Engine/Shaders/Private/InstanceCulling/InstanceCullingOcclusionQuery.usf
--- a/Engine/Shaders/Private/InstanceCulling/InstanceCullingOcclusionQuery.usf
+++ b/Engine/Shaders/Private/InstanceCulling/InstanceCullingOcclusionQuery.usf
@@ -231,3 +231,74 @@ void MainPS(
 	OutVisibilityMask[InstanceId] = 1;
 	OutColor = float4(1, 0, 0, 1);
 }
+
+// Debug utilities
+
+Buffer<uint> InstanceOcclusionQueryBuffer;
+
+void DebugMainVS(
+	float3 InPosition : ATTRIBUTE0,
+	uint InstanceId : SV_InstanceID,
+	out nointerpolation uint OutDebugVisibility : INSTANCE_DEBUG_VISIBILITY,
+	out float4 OutPosition : SV_POSITION
+)
+{
+	ResolvedView = ResolveView();
+
+	FInstanceSceneData InstanceData;
+	EInstanceVisibilityStatus VisibilityStatus = GetInstanceDataAndVisibility(InstanceId, /*out*/ InstanceData);
+	
+	if (InstanceData.ValidInstance)
+	{
+		float3 LocalPosition = (InPosition * InstanceData.LocalBoundsExtent * InstanceData.DeterminantSign) + InstanceData.LocalBoundsCenter;
+		float4 WorldPosition = TransformLocalToTranslatedWorld(LocalPosition, InstanceData.LocalToWorld);
+
+		OutPosition = mul(WorldPosition, ResolvedView.TranslatedWorldToClip);
+
+		const bool bWasVisibleLastFrame = InstanceOcclusionQueryBuffer[InstanceId] != 0; // YURIY_TODO: multi-view support
+
+		if (VisibilityStatus == IVS_PossiblyVisible)
+		{
+			OutDebugVisibility = bWasVisibleLastFrame ? 2 : 3;
+		}
+		else
+		{
+			OutDebugVisibility = bWasVisibleLastFrame ? 1 : 0;
+		}
+	}
+	else
+	{
+		OutPosition = asfloat(0xFFFFFFFF);
+		OutDebugVisibility = 0;
+	}
+}
+
+void DebugMainPS(
+	in nointerpolation uint DebugVisibility : INSTANCE_DEBUG_VISIBILITY,
+	out float4 OutColor : SV_Target0
+)
+{
+	switch (DebugVisibility)
+	{
+	default:
+	case 0:
+		// Definitely invisible due to HZB test.
+		// No occlusion query performed.
+		OutColor = float4(0.1, 0.0, 0.0, 0.0);
+		break;
+	case 1:
+		// Definitely visible due to near plane intersection or other heuristics.
+		// No occlusion query performed.
+		OutColor = float4(0.1, 0.1, 0.0, 0.0);
+		break;
+	case 2:
+		// Visible due to occlusion query.
+		OutColor = float4(0.0, 0.025, 0.0, 0.0);
+		break;
+	case 3:
+		// Invisible due to occlusion query.
+		OutColor = float4(0.25, 0.0, 0.25, 0.0);
+		break;
+	}
+}
+
diff --git a/Engine/Source/Runtime/Engine/Public/ShowFlags.h b/Engine/Source/Runtime/Engine/Public/ShowFlags.h
--- a/Engine/Source/Runtime/Engine/Public/ShowFlags.h
+++ b/Engine/Source/Runtime/Engine/Public/ShowFlags.h
@@ -478,6 +478,7 @@ private:
 		SetVisualizeCalibrationCustom(false);
 		SetVisualizePostProcessStack(false);
 		SetVirtualTexturePrimitives(false);
+		SetVisualizeInstanceOcclusionQueries(false);
 		SetVisualizeVolumetricCloudConservativeDensity(false);
 		SetVisualizeVolumetricCloudEmptySpaceSkipping(false);
 		SetDrawOnlyVSMInvalidatingGeo(false);
diff --git a/Engine/Source/Runtime/Engine/Public/ShowFlagsValues.inl b/Engine/Source/Runtime/Engine/Public/ShowFlagsValues.inl
--- a/Engine/Source/Runtime/Engine/Public/ShowFlagsValues.inl
+++ b/Engine/Source/Runtime/Engine/Public/ShowFlagsValues.inl
@@ -377,6 +377,8 @@ SHOWFLAG_ALWAYS_ACCESSIBLE(MediaPlanes, SFG_Normal, NSLOCTEXT("UnrealEd", "Media
 SHOWFLAG_FIXED_IN_SHIPPING(0, VREditing, SFG_Hidden, NSLOCTEXT("UnrealEd", "VREditSF", "VR Editing"))
 /** Visualize Occlusion Query bounding meshes */
 SHOWFLAG_FIXED_IN_SHIPPING(0, OcclusionMeshes, SFG_Visualize, NSLOCTEXT("UnrealEd", "VisualizeOcclusionQueries", "Visualize Occlusion Queries"))
+/** Visualize Occlusion Query bounding meshes */
+SHOWFLAG_FIXED_IN_SHIPPING(0, VisualizeInstanceOcclusionQueries, SFG_Visualize, NSLOCTEXT("UnrealEd", "VisualizeInstanceOcclusionQueriesSF", "Visualize Occlusion Queries Used By GPUScene Instance Culling"))
 /** Disable hardware occlusion queries, similar to setting r.AllowOcclusionQueries=0, but just for this scene. */
 SHOWFLAG_FIXED_IN_SHIPPING(0, DisableOcclusionQueries, SFG_Developer, NSLOCTEXT("UnrealEd", "DisableOcclusionQueries", "Disable Hardware Occlusion Queries"))
 
diff --git a/Engine/Source/Runtime/Renderer/Private/DeferredShadingRenderer.cpp b/Engine/Source/Runtime/Renderer/Private/DeferredShadingRenderer.cpp
--- a/Engine/Source/Runtime/Renderer/Private/DeferredShadingRenderer.cpp
+++ b/Engine/Source/Runtime/Renderer/Private/DeferredShadingRenderer.cpp
@@ -4218,6 +4218,15 @@ void FDeferredShadingSceneRenderer::Render(FRDGBuilder& GraphBuilder)
 			ShadingEnergyConservation::Debug(GraphBuilder, View, SceneTextures);
 		}
 
+		if (ViewFamily.EngineShowFlags.VisualizeInstanceOcclusionQueries
+			&& Scene->InstanceCullingOcclusionQueryRenderer)
+		{
+			for (FViewInfo& View : Views)
+			{
+				Scene->InstanceCullingOcclusionQueryRenderer->RenderDebug(GraphBuilder, Scene->GPUScene, View, SceneTextures);
+			}
+		}
+
 		Scene->ProcessAndRenderIlluminanceMeter(GraphBuilder, Views, SceneTextures.Color.Target);
 
 		if (!bHasRayTracedOverlay && ViewFamily.EngineShowFlags.LightShafts)
diff --git a/Engine/Source/Runtime/Renderer/Private/InstanceCulling/InstanceCullingOcclusionQuery.cpp b/Engine/Source/Runtime/Renderer/Private/InstanceCulling/InstanceCullingOcclusionQuery.cpp
--- a/Engine/Source/Runtime/Renderer/Private/InstanceCulling/InstanceCullingOcclusionQuery.cpp
+++ b/Engine/Source/Runtime/Renderer/Private/InstanceCulling/InstanceCullingOcclusionQuery.cpp
@@ -181,12 +181,18 @@ public:
 	{
 		static const uint16 BoxIndexBufferData[] =
 		{
+			// Tri list
 			0, 1, 2, 0, 2, 3,
 			4, 5, 6, 4, 6, 7,
 			1, 4, 7, 1, 7, 2,
 			5, 0, 3, 5, 3, 6,
 			5, 4, 1, 5, 1, 0,
 			3, 2, 7, 3, 7, 6,
+			// Line list
+			0, 1, 0, 3, 0, 5,
+			7, 2, 7, 6, 7, 4,
+			3, 2, 1, 2, 3, 6,
+			5, 6, 5, 4, 1, 4
 		};
 
 		static const FVector3f BoxVertexBufferData[] =
@@ -435,7 +441,7 @@ void FInstanceCullingOcclusionQueryRenderer::EndFrame(FRDGBuilder& GraphBuilder)
 	CurrentRenderedViewIDs.Empty();
 }
 
-uint32 FInstanceCullingOcclusionQueryRenderer::RegisterView(FViewInfo& View)
+uint32 FInstanceCullingOcclusionQueryRenderer::RegisterView(const FViewInfo& View)
 {
 	if (CurrentRenderedViewIDs.Num() < MaxViews)
 	{
@@ -449,10 +455,158 @@ uint32 FInstanceCullingOcclusionQueryRenderer::RegisterView(FViewInfo& View)
 	}
 }
 
-bool FInstanceCullingOcclusionQueryRenderer::IsCompatibleWithView(FViewInfo& View)
+bool FInstanceCullingOcclusionQueryRenderer::IsCompatibleWithView(const FViewInfo& View)
 {
 	return FDataDrivenShaderPlatformInfo::GetSupportsVertexShaderSRVs(View.GetShaderPlatform())
 		&& View.GetSceneTextures().Depth.Target
 		&& View.HZB
 		&& CVarInstanceCullingOcclusionQueries.GetValueOnRenderThread() != 0;
 }
+
+// Debugging utilities
+
+class FInstanceCullingOcclusionQueryDebugVS : public FGlobalShader
+{
+	DECLARE_GLOBAL_SHADER(FInstanceCullingOcclusionQueryDebugVS);
+	SHADER_USE_PARAMETER_STRUCT(FInstanceCullingOcclusionQueryDebugVS, FGlobalShader);
+
+public:
+	static bool ShouldCompilePermutation(const FGlobalShaderPermutationParameters& Parameters)
+	{
+		return FDataDrivenShaderPlatformInfo::GetSupportsVertexShaderSRVs(Parameters.Platform);
+	}
+
+	static void ModifyCompilationEnvironment(const FGlobalShaderPermutationParameters& Parameters, FShaderCompilerEnvironment& OutEnvironment)
+	{
+		FGlobalShader::ModifyCompilationEnvironment(Parameters, OutEnvironment);
+		OutEnvironment.SetDefine(TEXT("USE_GLOBAL_GPU_SCENE_DATA"), 1);
+		OutEnvironment.SetDefine(TEXT("VF_SUPPORTS_PRIMITIVE_SCENE_DATA"), 1);
+	}
+
+	BEGIN_SHADER_PARAMETER_STRUCT(FParameters, )
+		SHADER_PARAMETER(uint32, InstanceSceneDataSOAStride)
+		SHADER_PARAMETER_RDG_BUFFER_SRV(StructuredBuffer<float4>, GPUSceneInstanceSceneData)
+		SHADER_PARAMETER_RDG_BUFFER_SRV(StructuredBuffer<float4>, GPUSceneInstancePayloadData)
+		SHADER_PARAMETER_RDG_BUFFER_SRV(StructuredBuffer<float4>, GPUScenePrimitiveSceneData)
+		SHADER_PARAMETER_STRUCT_REF(FViewUniformShaderParameters, View)
+		SHADER_PARAMETER_RDG_BUFFER_SRV(Buffer<uint>, InstanceOcclusionQueryBuffer)
+		SHADER_PARAMETER_RDG_TEXTURE(Texture2D, HZBTexture)
+		SHADER_PARAMETER_SAMPLER(SamplerState, HZBSampler)
+		SHADER_PARAMETER(FVector2f, HZBSize)
+		SHADER_PARAMETER(FIntVector4, ViewRect)
+	END_SHADER_PARAMETER_STRUCT()
+};
+
+class FInstanceCullingOcclusionQueryDebugPS : public FGlobalShader
+{
+	DECLARE_GLOBAL_SHADER(FInstanceCullingOcclusionQueryDebugPS);
+	SHADER_USE_PARAMETER_STRUCT(FInstanceCullingOcclusionQueryDebugPS, FGlobalShader);
+
+public:
+	static bool ShouldCompilePermutation(const FGlobalShaderPermutationParameters& Parameters)
+	{
+		return FDataDrivenShaderPlatformInfo::GetSupportsVertexShaderSRVs(Parameters.Platform);
+	}
+
+	static void ModifyCompilationEnvironment(const FGlobalShaderPermutationParameters& Parameters, FShaderCompilerEnvironment& OutEnvironment)
+	{
+	}
+
+	BEGIN_SHADER_PARAMETER_STRUCT(FParameters, )
+	END_SHADER_PARAMETER_STRUCT()
+};
+
+IMPLEMENT_GLOBAL_SHADER(FInstanceCullingOcclusionQueryDebugVS, "/Engine/Private/InstanceCulling/InstanceCullingOcclusionQuery.usf", "DebugMainVS", SF_Vertex);
+IMPLEMENT_GLOBAL_SHADER(FInstanceCullingOcclusionQueryDebugPS, "/Engine/Private/InstanceCulling/InstanceCullingOcclusionQuery.usf", "DebugMainPS", SF_Pixel);
+
+BEGIN_SHADER_PARAMETER_STRUCT(FOcclusionInstanceCullingDebugParameters, )
+	SHADER_PARAMETER_STRUCT_INCLUDE(FInstanceCullingOcclusionQueryDebugVS::FParameters, VS)
+	SHADER_PARAMETER_STRUCT_INCLUDE(FInstanceCullingOcclusionQueryDebugPS::FParameters, PS)
+	RENDER_TARGET_BINDING_SLOTS()
+END_SHADER_PARAMETER_STRUCT()
+
+static void RenderInstanceOcclusionCullingDebug(
+	FRHICommandList& RHICmdList,
+	const FViewInfo& View,
+	FOcclusionInstanceCullingDebugParameters* PassParameters,
+	int32 NumInstances)
+{
+	TShaderMapRef<FInstanceCullingOcclusionQueryDebugVS> VertexShader(View.ShaderMap);
+	TShaderMapRef<FInstanceCullingOcclusionQueryDebugPS> PixelShader(View.ShaderMap);
+
+	FGraphicsPipelineStateInitializer GraphicsPSOInit;
+	RHICmdList.ApplyCachedRenderTargets(GraphicsPSOInit);
+
+	FIntVector4 ViewRect = PassParameters->VS.ViewRect;
+	RHICmdList.SetViewport(ViewRect.X, ViewRect.Y, 0.0f, ViewRect.Z, ViewRect.W, 1.0f);
+
+	GraphicsPSOInit.BoundShaderState.VertexDeclarationRHI = GInstanceCullingOcclusionQueryBox.VertexDeclaration;
+	GraphicsPSOInit.BoundShaderState.VertexShaderRHI = VertexShader.GetVertexShader();
+	GraphicsPSOInit.BoundShaderState.PixelShaderRHI = PixelShader.GetPixelShader();
+	GraphicsPSOInit.DepthStencilState = TStaticDepthStencilState<false, CF_Always>::GetRHI(); // No depth test or write
+	GraphicsPSOInit.BlendState = TStaticBlendState<CW_RGBA, BO_Add, BF_One, BF_InverseSourceAlpha, BO_Add, BF_Zero, BF_One>::GetRHI(); // Premultiplied
+	GraphicsPSOInit.RasterizerState = TStaticRasterizerState<FM_Solid, CM_None>::GetRHI();
+	GraphicsPSOInit.PrimitiveType = PT_LineList;
+
+	SetGraphicsPipelineState(RHICmdList, GraphicsPSOInit, 0);
+
+	ClearUnusedGraphResources(VertexShader, &PassParameters->VS);
+	ClearUnusedGraphResources(PixelShader, &PassParameters->PS);
+
+	SetShaderParameters(RHICmdList, VertexShader, VertexShader.GetVertexShader(), PassParameters->VS);
+	SetShaderParameters(RHICmdList, PixelShader, PixelShader.GetPixelShader(), PassParameters->PS);
+
+	RHICmdList.SetStreamSource(0, GInstanceCullingOcclusionQueryBox.VertexBuffer, 0);
+
+	RHICmdList.DrawIndexedPrimitive(GInstanceCullingOcclusionQueryBox.IndexBuffer, 0, 0, 24, 36, 12, NumInstances);
+}
+
+void FInstanceCullingOcclusionQueryRenderer::RenderDebug(FRDGBuilder& GraphBuilder, FGPUScene& GPUScene, const FViewInfo& View, FSceneTextures& SceneTextures)
+{
+	if (!IsCompatibleWithView(View) || !InstanceOcclusionQueryBuffer)
+	{
+		return;
+	}
+
+	FRDGTextureRef SceneColor = View.GetSceneTextures().Color.Target;
+	FRDGTextureRef SceneDepth = View.GetSceneTextures().Depth.Target;
+	FRDGBufferRef InstanceOcclusionQueryBufferRDG = GraphBuilder.RegisterExternalBuffer(InstanceOcclusionQueryBuffer);
+
+	FRDGTextureRef DepthTexture = View.GetSceneTextures().Depth.Target;
+	FRDGTextureRef HZBTexture = View.HZB;
+
+	checkf(DepthTexture && HZBTexture,
+		TEXT("Occlusion query instance culling requires scene depth texture and HZB. See FInstanceCullingOcclusionQueryRenderer::IsCompatibleWithView()"));
+
+	const FIntVector HZBSize = HZBTexture->Desc.GetSize();
+
+	const int32 NumInstances = GPUScene.GetNumInstances();
+	const FGPUSceneResourceParameters GPUSceneParameters = GPUScene.GetShaderParameters();
+
+	FOcclusionInstanceCullingDebugParameters* PassParameters = GraphBuilder.AllocParameters<FOcclusionInstanceCullingDebugParameters>();
+
+	PassParameters->VS.ViewRect = FIntVector4(View.ViewRect.Min.X, View.ViewRect.Min.Y, View.ViewRect.Max.X, View.ViewRect.Max.Y);
+	PassParameters->VS.View = View.ViewUniformBuffer;
+	PassParameters->VS.InstanceSceneDataSOAStride = GPUSceneParameters.InstanceDataSOAStride;
+	PassParameters->VS.GPUSceneInstanceSceneData = GPUSceneParameters.GPUSceneInstanceSceneData;
+	PassParameters->VS.GPUSceneInstancePayloadData = GPUSceneParameters.GPUSceneInstancePayloadData;
+	PassParameters->VS.GPUScenePrimitiveSceneData = GPUSceneParameters.GPUScenePrimitiveSceneData;
+	PassParameters->VS.InstanceOcclusionQueryBuffer = GraphBuilder.CreateSRV(InstanceOcclusionQueryBufferRDG, PF_R32_UINT);
+	PassParameters->VS.HZBTexture = HZBTexture;
+	PassParameters->VS.HZBSampler = TStaticSamplerState<SF_Point, AM_Clamp, AM_Clamp, AM_Clamp>::GetRHI();
+	PassParameters->VS.HZBSize = FVector2f(HZBSize.X, HZBSize.Y);
+	PassParameters->RenderTargets[0] = FRenderTargetBinding(SceneColor, ERenderTargetLoadAction::ELoad);
+	PassParameters->RenderTargets.DepthStencil = FDepthStencilBinding(SceneDepth,
+		ERenderTargetLoadAction::ELoad, ERenderTargetLoadAction::ENoAction,
+		FExclusiveDepthStencil::DepthRead_StencilNop);
+
+	GraphBuilder.AddPass(
+		RDG_EVENT_NAME("InstanceCullingOcclusionQueryRenderer_Draw"),
+		PassParameters, ERDGPassFlags::Raster | ERDGPassFlags::NeverCull,
+		[PassParameters, NumInstances, &View](FRHICommandList& RHICmdList)
+		{
+			RenderInstanceOcclusionCullingDebug(RHICmdList, View, PassParameters, NumInstances);
+		});
+
+}
+
diff --git a/Engine/Source/Runtime/Renderer/Private/InstanceCulling/InstanceCullingOcclusionQuery.h b/Engine/Source/Runtime/Renderer/Private/InstanceCulling/InstanceCullingOcclusionQuery.h
--- a/Engine/Source/Runtime/Renderer/Private/InstanceCulling/InstanceCullingOcclusionQuery.h
+++ b/Engine/Source/Runtime/Renderer/Private/InstanceCulling/InstanceCullingOcclusionQuery.h
@@ -27,6 +27,11 @@ public:
 	*/	
 	uint32 Render(FRDGBuilder& GraphBuilder, FGPUScene& GPUScene, FViewInfo& View);
 
+	/**
+	* Render object bounding boxes, color-coded based on visibility
+	*/
+	void RenderDebug(FRDGBuilder& GraphBuilder, FGPUScene& GPUScene, const FViewInfo& View, FSceneTextures& SceneTextures);
+
 	/**
 	* Mark blocks of occlusion query buffer entries as visible in all views.
 	* Required, for example, when GPUScene instances are invalidated betwen frames,
@@ -50,7 +55,7 @@ public:
 	/*
 	* Returns true if per-instance occlusion queries can be rendered for the view.
 	*/
-	bool IsCompatibleWithView(FViewInfo& View);
+	bool IsCompatibleWithView(const FViewInfo& View);
 
 private:
 
@@ -59,7 +64,7 @@ private:
 	* Returns bit mask with a single set bit that should be used to access the occlusion results.
 	* Returns 0 if maximum number of supported views is reached, falling back to no-occlusion-query code path.
 	*/
-	uint32 RegisterView(FViewInfo& View);
+	uint32 RegisterView(const FViewInfo& View);
 
 	FRDGBufferRef CurrentInstanceOcclusionQueryBuffer = {};
 
diff --git a/Engine/Source/Runtime/Renderer/Private/MobileShadingRenderer.cpp b/Engine/Source/Runtime/Renderer/Private/MobileShadingRenderer.cpp
--- a/Engine/Source/Runtime/Renderer/Private/MobileShadingRenderer.cpp
+++ b/Engine/Source/Runtime/Renderer/Private/MobileShadingRenderer.cpp
@@ -1077,6 +1077,15 @@ void FMobileSceneRenderer::Render(FRDGBuilder& GraphBuilder)
 		RenderMeshDistanceFieldVisualization(GraphBuilder, SceneTextures);
 	}
 	
+	if (ViewFamily.EngineShowFlags.VisualizeInstanceOcclusionQueries
+		&& Scene->InstanceCullingOcclusionQueryRenderer)
+	{
+		for (FViewInfo& View : Views)
+		{
+			Scene->InstanceCullingOcclusionQueryRenderer->RenderDebug(GraphBuilder, Scene->GPUScene, View, SceneTextures);
+		}
+	}
+
 	if (bUseVirtualTexturing)
 	{
 		RDG_GPU_STAT_SCOPE(GraphBuilder, VirtualTextureUpdate);
@@ -1917,6 +1926,12 @@ bool FMobileSceneRenderer::ShouldRenderHZB()
 	// Mobile SSAO requests HZB
 	bool bIsFeatureRequested = bRequiresAmbientOcclusionPass && MobileAmbientOcclusionTechniqueCVar->GetValueOnRenderThread() == 1;
 
+	// Instance occlusion culling requires HZB
+	if (FInstanceCullingContext::IsOcclusionCullingEnabled())
+	{
+		bIsFeatureRequested = true;
+	}
+
 	bool bNeedsHZB = bIsFeatureRequested;
 
 	return bNeedsHZB;
