Fixes RHI validation failure

diff --git a/Engine/Source/Runtime/Renderer/Private/PostProcess/TemporalSuperResolution.cpp b/Engine/Source/Runtime/Renderer/Private/PostProcess/TemporalSuperResolution.cpp
--- a/Engine/Source/Runtime/Renderer/Private/PostProcess/TemporalSuperResolution.cpp
+++ b/Engine/Source/Runtime/Renderer/Private/PostProcess/TemporalSuperResolution.cpp
@@ -1468,7 +1468,7 @@ FDefaultTemporalUpscaler::FOutputs AddTemporalSuperResolutionPasses(
 		PrevHistory.MetadataArray = BlackArrayDummy;
 		PrevHistory.GuideArray = BlackArrayDummy;
 		PrevHistory.MoireArray = BlackArrayDummy;
-		PrevHistory.SubpixelDepth = BlackUintDummy;
+		PrevHistory.SubpixelDepth = nullptr;
 
 		if (HistorySliceSequence.GetRollingIndexCount() > 1)
 		{
@@ -1645,18 +1645,14 @@ FDefaultTemporalUpscaler::FOutputs AddTemporalSuperResolutionPasses(
 		PassParameters->SubpixelDepthOutput = GraphBuilder.CreateUAV(PrevScatteredSubpixelDepthTexture);
 		PassParameters->DebugOutput = CreateDebugUAV(InputExtent, TEXT("Debug.TSR.ForwardScatterDepth"));
 
-		const FIntVector DispatchDim = FComputeShaderUtils::GetGroupCount(ScatterDepthViewport.Size(), 8);
-		if (DispatchDim.X > 0 && DispatchDim.Y > 0)
-		{
-			TShaderMapRef<FTSRForwardScatterDepthCS> ComputeShader(View.ShaderMap);
-			FComputeShaderUtils::AddPass(
-				GraphBuilder,
-				RDG_EVENT_NAME("TSR ForwardScatterDepth %dx%d", ScatterDepthViewport.Width(), ScatterDepthViewport.Height()),
-				AsyncComputePasses >= 1 ? ERDGPassFlags::AsyncCompute : ERDGPassFlags::Compute,
-				ComputeShader,
-				PassParameters,
-				FComputeShaderUtils::GetGroupCount(ScatterDepthViewport.Size(), 8));
-		}
+		TShaderMapRef<FTSRForwardScatterDepthCS> ComputeShader(View.ShaderMap);
+		FComputeShaderUtils::AddPass(
+			GraphBuilder,
+			RDG_EVENT_NAME("TSR ForwardScatterDepth %dx%d", ScatterDepthViewport.Width(), ScatterDepthViewport.Height()),
+			AsyncComputePasses >= 1 ? ERDGPassFlags::AsyncCompute : ERDGPassFlags::Compute,
+			ComputeShader,
+			PassParameters,
+			FComputeShaderUtils::GetGroupCount(ScatterDepthViewport.Size(), 8));
 	}
 
 	// Dilate the velocity texture & scatter reprojection into previous frame
@@ -2307,7 +2303,7 @@ FDefaultTemporalUpscaler::FOutputs AddTemporalSuperResolutionPasses(
 			FRDGTextureUAVDesc MetadataUAVDesc(History.MetadataArray);
 			MetadataUAVDesc.FirstArraySlice = CurrentFrameSliceIndex;
 			MetadataUAVDesc.NumArraySlices = 1;
-            MetadataUAVDesc.DimensionOverride = ETextureDimension::Texture2D;
+			MetadataUAVDesc.DimensionOverride = ETextureDimension::Texture2D;
             
 			PassParameters->HistoryArrayIndices = HistoryArrayIndices;
 			PassParameters->HistoryColorOutput = GraphBuilder.CreateUAV(ColorUAVDesc);
