Implemented RHI-agnostic shader bundle dispatch shader and renderer utils to share as much implementation across RHIs as possible.

diff --git a/Engine/Shaders/Private/ShaderBundleDispatch.usf b/Engine/Shaders/Private/ShaderBundleDispatch.usf
new file mode 100644
--- /dev/null
+++ b/Engine/Shaders/Private/ShaderBundleDispatch.usf
@@ -0,0 +1,31 @@
+// Copyright Epic Games, Inc. All Rights Reserved.
+
+#include "Common.ush"
+#include "WaveOpUtil.ush"
+
+#ifdef OVERRIDE_SHADERBUNDLE_USH
+#include "/Platform/Private/ShaderBundle.ush"
+
+uint RecordCount;
+uint4 PlatformData;
+ByteAddressBuffer RecordArgBuffer;
+ByteAddressBuffer RecordDataBuffer;
+RWByteAddressBuffer RWExecutionBuffer;
+
+#endif
+
+[numthreads(THREADGROUP_SIZEX, 1, 1)]
+void DispatchShaderBundleEntry(uint RecordIndex : SV_DispatchThreadID) 
+{
+#ifdef OVERRIDE_SHADERBUNDLE_USH
+	if (RecordIndex < RecordCount)
+	{
+		ProcessShaderBundleRecord(
+			PlatformData,
+			RecordIndex,
+			RecordArgBuffer,
+			RecordDataBuffer,
+			RWExecutionBuffer);
+	}
+#endif
+}
diff --git a/Engine/Source/Runtime/D3D12RHI/Private/D3D12Commands.cpp b/Engine/Source/Runtime/D3D12RHI/Private/D3D12Commands.cpp
--- a/Engine/Source/Runtime/D3D12RHI/Private/D3D12Commands.cpp
+++ b/Engine/Source/Runtime/D3D12RHI/Private/D3D12Commands.cpp
@@ -1622,7 +1622,7 @@ void FD3D12CommandContext::RHIDispatchShaderBundle(
 {
 	check(ShaderBundle != nullptr && Dispatches.Num() > 0 && ShaderBundle->bEmulated);
 	TRHICommandList_RecursiveHazardous<FD3D12CommandContext> RHICmdList(this);
-	checkNoEntry();
+	UE::RHICore::DispatchShaderBundleEmulation(RHICmdList, ShaderBundle, RecordArgBufferSRV->GetBuffer(), Dispatches);
 }
 
 void FD3D12CommandContext::RHIDrawPrimitive(uint32 BaseVertexIndex, uint32 NumPrimitives, uint32 NumInstances)
diff --git a/Engine/Source/Runtime/Renderer/Private/RendererUtils.cpp b/Engine/Source/Runtime/Renderer/Private/RendererUtils.cpp
--- a/Engine/Source/Runtime/Renderer/Private/RendererUtils.cpp
+++ b/Engine/Source/Runtime/Renderer/Private/RendererUtils.cpp
@@ -268,4 +268,68 @@ RENDERER_API FBufferRHIRef& GetOneTileQuadVertexBuffer()
 RENDERER_API FBufferRHIRef& GetOneTileQuadIndexBuffer()
 {
 	return GOneTileQuadIndexBuffer.IndexBufferRHI;
-}
\ No newline at end of file
+}
+
+class FDispatchShaderBundleCS : public FGlobalShader
+{
+public:
+	DECLARE_GLOBAL_SHADER(FDispatchShaderBundleCS);
+	SHADER_USE_PARAMETER_STRUCT(FDispatchShaderBundleCS, FGlobalShader)
+
+	static const uint32 ThreadGroupSizeX = 64;
+
+	using FPermutationDomain = TShaderPermutationDomain<>;
+
+	BEGIN_SHADER_PARAMETER_STRUCT(FParameters, )
+		SHADER_PARAMETER(uint32, RecordCount)
+		SHADER_PARAMETER(FUintVector4, PlatformData)
+		SHADER_PARAMETER_SRV(ByteAddressBuffer, RecordArgBuffer)
+		SHADER_PARAMETER_SRV(ByteAddressBuffer, RecordDataBuffer)
+		SHADER_PARAMETER_UAV(RWByteAddressBuffer, RWExecutionBuffer)
+	END_SHADER_PARAMETER_STRUCT()
+
+	static bool ShouldCompilePermutation(const FGlobalShaderPermutationParameters& Parameters)
+	{
+		return RHISupportsShaderBundleDispatch(Parameters.Platform);
+	}
+
+	static void ModifyCompilationEnvironment(const FGlobalShaderPermutationParameters& Parameters, FShaderCompilerEnvironment& OutEnvironment)
+	{
+		OutEnvironment.SetDefine(TEXT("THREADGROUP_SIZEX"), ThreadGroupSizeX);
+		FGlobalShader::ModifyCompilationEnvironment(Parameters, OutEnvironment);
+	}
+};
+
+IMPLEMENT_GLOBAL_SHADER(FDispatchShaderBundleCS, "/Engine/Private/ShaderBundleDispatch.usf", "DispatchShaderBundleEntry", SF_Compute);
+
+void FDispatchShaderBundle::Dispatch(
+	FRHIShaderBundle* ShaderBundle,
+	FRHIComputeCommandList& RHICmdList,
+	FRHIShaderResourceView* RecordArgBufferSRV,
+	FRHIShaderResourceView* RecordDataBufferSRV,
+	FRHIUnorderedAccessView* ExecutionBufferUAV
+)
+{
+	check(RHISupportsShaderBundleDispatch(GMaxRHIShaderPlatform) && GRHISupportsDispatchShaderBundle);
+	check(ShaderBundle && ShaderBundle->NumRecords > 0);
+
+	RHICmdList.ClearUAVUint(ExecutionBufferUAV, FUintVector4(0, 0, 0, 0));
+
+	auto ComputeShader = GetGlobalShaderMap(GMaxRHIFeatureLevel)->GetShader<FDispatchShaderBundleCS>();
+
+	FDispatchShaderBundleCS::FParameters Parameters;
+
+	Parameters.RecordCount = ShaderBundle->NumRecords;
+	Parameters.PlatformData = ShaderBundle->GetPlatformData();
+	Parameters.RecordArgBuffer = RecordArgBufferSRV;
+	Parameters.RecordDataBuffer = RecordDataBufferSRV;
+	Parameters.RWExecutionBuffer = ExecutionBufferUAV;
+
+	const FIntVector GroupCount = FComputeShaderUtils::GetGroupCount(ShaderBundle->NumRecords, FDispatchShaderBundleCS::ThreadGroupSizeX);
+
+	FRHIComputeShader* ComputeShaderRHI = ComputeShader.GetComputeShader();
+	SetComputePipelineState(RHICmdList, ComputeShaderRHI);
+	SetShaderParameters(RHICmdList, ComputeShader, ComputeShaderRHI, Parameters);
+	DispatchComputeShader(RHICmdList, ComputeShader, GroupCount.X, GroupCount.Y, GroupCount.Z);
+	UnsetShaderUAVs(RHICmdList, ComputeShader, ComputeShaderRHI);
+}
diff --git a/Engine/Source/Runtime/Renderer/Public/RendererUtils.h b/Engine/Source/Runtime/Renderer/Public/RendererUtils.h
--- a/Engine/Source/Runtime/Renderer/Public/RendererUtils.h
+++ b/Engine/Source/Runtime/Renderer/Public/RendererUtils.h
@@ -2,6 +2,7 @@
 
 #pragma once
 
+#include "CoreMinimal.h"
 #include "Containers/ArrayView.h"
 #include "GlobalShader.h"
 #include "HAL/Platform.h"
@@ -16,6 +17,7 @@ class FScene;
 class FGlobalShaderMap;
 class FRDGBuilder;
 class FRHICommandListImmediate;
+class FRHIShaderBundle;
 class FScene;
 struct IPooledRenderTarget;
 template <typename ReferencedType> class TRefCountPtr;
@@ -40,6 +42,18 @@ public:
 		const TCHAR* RTWriteMaskDebugName);
 };
 
+class FDispatchShaderBundle
+{
+public:
+	static RENDERER_API void Dispatch(
+		FRHIShaderBundle* ShaderBundle,
+		FRHIComputeCommandList& RHICmdList,
+		FRHIShaderResourceView* RecordArgBufferSRV,
+		FRHIShaderResourceView* RecordDataBufferSRV,
+		FRHIUnorderedAccessView* ExecutionBufferUAV
+	);
+};
+
 class FDepthBounds
 {
 public:
