Fixed crash with Dump GPU on Macs due to missing viewport

diff --git a/Engine/Source/Runtime/Renderer/Private/PostProcess/PostProcessSelectionOutline.cpp b/Engine/Source/Runtime/Renderer/Private/PostProcess/PostProcessSelectionOutline.cpp
--- a/Engine/Source/Runtime/Renderer/Private/PostProcess/PostProcessSelectionOutline.cpp
+++ b/Engine/Source/Runtime/Renderer/Private/PostProcess/PostProcessSelectionOutline.cpp
@@ -169,6 +169,8 @@ FScreenPassTexture AddSelectionOutlinePass(
 				ERDGPassFlags::Raster,
 				[SceneColorViewport](FRHICommandListImmediate& RHICmdList)
 				{
+                    RHICmdList.SetViewport(SceneColorViewport.Rect.Min.X, SceneColorViewport.Rect.Min.Y, 0.0f, SceneColorViewport.Rect.Max.X, SceneColorViewport.Rect.Max.Y, 1.0f);
+                    
 					FIntRect InnerRect = SceneColorViewport.Rect;
 
 					// 1 as we have an outline that is that thick
