Fix Tonemap 0.2ms gpu regression on some platform. A change in IsVRSCompatibleWithView code logic was causing the tonemap to output the luminance in a separate RT. Most of the gpu cost was coming for the clear of that RTV (0.15ms).

diff --git a/Engine/Source/Runtime/Renderer/Private/PostProcess/PostProcessTonemap.cpp b/Engine/Source/Runtime/Renderer/Private/PostProcess/PostProcessTonemap.cpp
--- a/Engine/Source/Runtime/Renderer/Private/PostProcess/PostProcessTonemap.cpp
+++ b/Engine/Source/Runtime/Renderer/Private/PostProcess/PostProcessTonemap.cpp
@@ -887,7 +887,7 @@ FScreenPassTexture AddTonemapPass(FRDGBuilder& GraphBuilder, const FViewInfo& Vi
 
 		DesktopPermutationVector.Set<TonemapperPermutation::FTonemapperOutputDeviceDim>(EDisplayOutputFormat(CommonParameters.OutputDevice.OutputDevice));
 
-		DesktopPermutationVector.Set<TonemapperPermutation::FTonemapperOutputLuminance>(FVariableRateShadingImageManager::IsVRSCompatibleWithView(View));
+		DesktopPermutationVector.Set<TonemapperPermutation::FTonemapperOutputLuminance>(FVariableRateShadingImageManager::IsVRSSupportedByRHI() && FVariableRateShadingImageManager::IsVRSCompatibleWithView(View));
 
 		DesktopPermutationVector = TonemapperPermutation::RemapPermutation(DesktopPermutationVector, View.GetFeatureLevel());
 	}
