Fix for foreground debug line rendering being affected by incorrect depth sorting.
Also removed the temporary workaround used to compensate.

diff --git a/Engine/Source/Runtime/Engine/Private/Components/LineBatchComponent.cpp b/Engine/Source/Runtime/Engine/Private/Components/LineBatchComponent.cpp
--- a/Engine/Source/Runtime/Engine/Private/Components/LineBatchComponent.cpp
+++ b/Engine/Source/Runtime/Engine/Private/Components/LineBatchComponent.cpp
@@ -67,10 +67,7 @@ void FLineBatcherSceneProxy::GetDynamicMeshElements(const TArray<const FSceneVie
 
 #if UE_ENABLE_DEBUG_DRAWING
 			//If we are in a debug draw build, this line batcher should write out to the debug only PDI
-			//FPrimitiveDrawInterface* PDI = Collector.GetDebugPDI(ViewIndex);
-
-			// temporary change to fix the problem around depth priority not being accounted for when drawing debug primitive
-			FPrimitiveDrawInterface* PDI = Collector.GetPDI(ViewIndex);
+			FPrimitiveDrawInterface* PDI = Collector.GetDebugPDI(ViewIndex);
 #else
 			FPrimitiveDrawInterface* PDI = Collector.GetPDI(ViewIndex);
 #endif
diff --git a/Engine/Source/Runtime/Renderer/Private/PostProcess/PostProcessCompositeDebugPrimitives.cpp b/Engine/Source/Runtime/Renderer/Private/PostProcess/PostProcessCompositeDebugPrimitives.cpp
--- a/Engine/Source/Runtime/Renderer/Private/PostProcess/PostProcessCompositeDebugPrimitives.cpp
+++ b/Engine/Source/Runtime/Renderer/Private/PostProcess/PostProcessCompositeDebugPrimitives.cpp
@@ -76,9 +76,11 @@ FScreenPassTexture AddDebugPrimitivePass(
 		FMeshPassProcessorRenderState DrawRenderState;
 		DrawRenderState.SetDepthStencilAccess(FExclusiveDepthStencil::DepthWrite_StencilWrite);
 		DrawRenderState.SetBlendState(TStaticBlendState<CW_RGBA, BO_Add, BF_One, BF_One, BO_Add, BF_One, BF_One>::GetRHI());
-		DrawRenderState.SetDepthStencilState(TStaticDepthStencilState<true, CF_DepthNearOrEqual>::GetRHI());
-		
+
+		DrawRenderState.SetDepthStencilState(TStaticDepthStencilState<true, CF_DepthNearOrEqual>::GetRHI());		
 		DebugView->DebugSimpleElementCollector.DrawBatchedElements(RHICmdList, DrawRenderState, *DebugView, EBlendModeFilter::OpaqueAndMasked, SDPG_World);
+
+		DrawRenderState.SetDepthStencilState(TStaticDepthStencilState<true, CF_Always>::GetRHI());
 		DebugView->DebugSimpleElementCollector.DrawBatchedElements(RHICmdList, DrawRenderState, *DebugView, EBlendModeFilter::OpaqueAndMasked, SDPG_Foreground);		
 	});
 
