Adding parenthesis for removing FastBuild error

diff --git a/Engine/Source/Runtime/Renderer/Private/SceneRendering.cpp b/Engine/Source/Runtime/Renderer/Private/SceneRendering.cpp
--- a/Engine/Source/Runtime/Renderer/Private/SceneRendering.cpp
+++ b/Engine/Source/Runtime/Renderer/Private/SceneRendering.cpp
@@ -5680,7 +5680,7 @@ void AddResolveSceneDepthPass(FRDGBuilder& GraphBuilder, const FViewInfo& View,
 		RHICmdList.SetViewport(0.0f, 0.0f, 0.0f, DepthExtent.X, DepthExtent.Y, 1.0f);
 
 		bool bArrayResolve = SourceTextureRHI->GetDesc().IsTextureArray();
-		ensureMsgf(!bArrayResolve || RHISupportsVertexShaderLayer(View.GetShaderPlatform()) && GRHISupportsArrayIndexFromAnyShader,
+		ensureMsgf(!bArrayResolve || (RHISupportsVertexShaderLayer(View.GetShaderPlatform()) && GRHISupportsArrayIndexFromAnyShader),
 			TEXT("Resolving scene depth array requires support for outputting SV_RenderTargetArrayIndex from any shader."));
 
 		/** Chooses one of many ResolvePS variants */
