Add query function for VRS being enabled.

diff --git a/Engine/Source/Runtime/Renderer/Private/PostProcess/PostProcessTonemap.cpp b/Engine/Source/Runtime/Renderer/Private/PostProcess/PostProcessTonemap.cpp
--- a/Engine/Source/Runtime/Renderer/Private/PostProcess/PostProcessTonemap.cpp
+++ b/Engine/Source/Runtime/Renderer/Private/PostProcess/PostProcessTonemap.cpp
@@ -887,7 +887,7 @@ FScreenPassTexture AddTonemapPass(FRDGBuilder& GraphBuilder, const FViewInfo& Vi
 
 		DesktopPermutationVector.Set<TonemapperPermutation::FTonemapperOutputDeviceDim>(EDisplayOutputFormat(CommonParameters.OutputDevice.OutputDevice));
 
-		DesktopPermutationVector.Set<TonemapperPermutation::FTonemapperOutputLuminance>(FVariableRateShadingImageManager::IsVRSSupportedByRHI() && FVariableRateShadingImageManager::IsVRSCompatibleWithView(View));
+		DesktopPermutationVector.Set<TonemapperPermutation::FTonemapperOutputLuminance>(GVRSImageManager.IsVRSEnabledForFrame() && FVariableRateShadingImageManager::IsVRSCompatibleWithView(View));
 
 		DesktopPermutationVector = TonemapperPermutation::RemapPermutation(DesktopPermutationVector, View.GetFeatureLevel());
 	}
diff --git a/Engine/Source/Runtime/Renderer/Private/VariableRateShading/VariableRateShadingImageManager.cpp b/Engine/Source/Runtime/Renderer/Private/VariableRateShading/VariableRateShadingImageManager.cpp
--- a/Engine/Source/Runtime/Renderer/Private/VariableRateShading/VariableRateShadingImageManager.cpp
+++ b/Engine/Source/Runtime/Renderer/Private/VariableRateShading/VariableRateShadingImageManager.cpp
@@ -320,6 +320,11 @@ bool FVariableRateShadingImageManager::IsVRSEnabled()
 	return GRHIVariableRateShadingEnabled && GRHIAttachmentVariableRateShadingEnabled;
 }
 
+bool FVariableRateShadingImageManager::IsVRSEnabledForFrame()
+{
+	return bVRSEnabledForFrame;
+}
+
 bool FVariableRateShadingImageManager::IsVRSCompatibleWithOutputType(const EDisplayOutputFormat& OutputFormat)
 {
 	// VRS texture generation is currently only compatible with SDR and HDR10
diff --git a/Engine/Source/Runtime/Renderer/Public/VariableRateShadingImageManager.h b/Engine/Source/Runtime/Renderer/Public/VariableRateShadingImageManager.h
--- a/Engine/Source/Runtime/Renderer/Public/VariableRateShadingImageManager.h
+++ b/Engine/Source/Runtime/Renderer/Public/VariableRateShadingImageManager.h
@@ -100,6 +100,7 @@ public:
 
 	static bool IsVRSSupportedByRHI();
 	static bool IsVRSEnabled();
+	bool IsVRSEnabledForFrame();
 	static bool IsVRSCompatibleWithView(const FViewInfo& View);
 	static bool IsVRSCompatibleWithOutputType(const EDisplayOutputFormat& DisplayOutputFormat);
 
