Fixed occlusion feedback interaction with sub-queries

diff --git a/Engine/Source/Runtime/Renderer/Private/SceneOcclusion.cpp b/Engine/Source/Runtime/Renderer/Private/SceneOcclusion.cpp
--- a/Engine/Source/Runtime/Renderer/Private/SceneOcclusion.cpp
+++ b/Engine/Source/Runtime/Renderer/Private/SceneOcclusion.cpp
@@ -1788,7 +1788,7 @@ void FOcclusionFeedback::ReadbackResults(FRHICommandList& RHICmdList)
 	}
 }
 
-void FOcclusionFeedback::AddPrimitive(FRHICommandList& RHICmdList, FPrimitiveComponentId PrimitiveId, const FVector& BoundsOrigin, const FVector& BoundsBoxExtent, FGlobalDynamicVertexBuffer& DynamicVertexBuffer)
+void FOcclusionFeedback::AddPrimitive(FRHICommandList& RHICmdList, const FPrimitiveOcclusionHistoryKey& PrimitiveKey, const FVector& BoundsOrigin, const FVector& BoundsBoxExtent, FGlobalDynamicVertexBuffer& DynamicVertexBuffer)
 {
 	constexpr uint32 MaxBatchedPrimitives = 512;
 	constexpr uint32 PrimitiveStride = sizeof(FVector4f) * 2u;
@@ -1817,7 +1817,7 @@ void FOcclusionFeedback::AddPrimitive(FRHICommandList& RHICmdList, FPrimitiveCom
 	OcclusionBatch.NumBatchedPrimitives++;
 
 	FOcclusionBuffer& OcclusionBuffer = OcclusionBuffers[CurrentBufferIndex];
-	OcclusionBuffer.BatchedPrimitives.Add(PrimitiveId);
+	OcclusionBuffer.BatchedPrimitives.Add(PrimitiveKey);
 }
 
 class FOcclusionFeedbackVS : public FGlobalShader
diff --git a/Engine/Source/Runtime/Renderer/Private/ScenePrivate.h b/Engine/Source/Runtime/Renderer/Private/ScenePrivate.h
--- a/Engine/Source/Runtime/Renderer/Private/ScenePrivate.h
+++ b/Engine/Source/Runtime/Renderer/Private/ScenePrivate.h
@@ -303,7 +303,7 @@ struct FPrimitiveOcclusionHistoryKey
 	FPrimitiveComponentId PrimitiveId;
 	int32 CustomIndex;
 
-	FPrimitiveOcclusionHistoryKey(const FPrimitiveOcclusionHistory& Element)
+	explicit FPrimitiveOcclusionHistoryKey(const FPrimitiveOcclusionHistory& Element)
 		: PrimitiveId(Element.PrimitiveId)
 		, CustomIndex(Element.CustomIndex)
 	{
@@ -315,6 +315,16 @@ struct FPrimitiveOcclusionHistoryKey
 	}
 };
 
+inline uint32 GetTypeHash(const FPrimitiveOcclusionHistoryKey& Key)
+{
+	return GetTypeHash(Key.PrimitiveId.PrimIDValue) ^ (GetTypeHash(Key.CustomIndex) >> 20);
+}
+
+inline bool operator==(const FPrimitiveOcclusionHistoryKey& A, const FPrimitiveOcclusionHistoryKey& B)
+{
+	return A.PrimitiveId == B.PrimitiveId && A.CustomIndex == B.CustomIndex;
+}
+
 /** Defines how the hash set indexes the FPrimitiveOcclusionHistory objects. */
 struct FPrimitiveOcclusionHistoryKeyFuncs : BaseKeyFuncs<FPrimitiveOcclusionHistory,FPrimitiveOcclusionHistoryKey>
 {
@@ -327,16 +337,15 @@ struct FPrimitiveOcclusionHistoryKeyFuncs : BaseKeyFuncs<FPrimitiveOcclusionHist
 
 	static bool Matches(KeyInitType A,KeyInitType B)
 	{
-		return A.PrimitiveId == B.PrimitiveId && A.CustomIndex == B.CustomIndex;
+		return A == B;
 	}
 
 	static uint32 GetKeyHash(KeyInitType Key)
 	{
-		return GetTypeHash(Key.PrimitiveId.PrimIDValue) ^ (GetTypeHash(Key.CustomIndex) >> 20);
+		return GetTypeHash(Key);
 	}
 };
 
-
 class FIndividualOcclusionHistory
 {
 	FRHIPooledRenderQuery PendingOcclusionQuery[FOcclusionQueryHelpers::MaxBufferedOcclusionFrames];
@@ -682,6 +691,66 @@ struct FPersistentGlobalDistanceFieldData : public FThreadSafeRefCountedObject
 	uint64 GetGPUSizeBytes(bool bLogSizes) const;
 };
 
+class FOcclusionFeedback : public FRenderResource
+{
+public:
+	FOcclusionFeedback();
+	~FOcclusionFeedback();
+
+	// FRenderResource interface
+	virtual void InitRHI(FRHICommandListBase& RHICmdList) override;
+	virtual void ReleaseRHI() override;
+
+	void AddPrimitive(FRHICommandList& RHICmdList, const FPrimitiveOcclusionHistoryKey& PrimitiveKey, const FVector& BoundsOrigin, const FVector& BoundsBoxExtent, FGlobalDynamicVertexBuffer& DynamicVertexBuffer);
+
+	void BeginOcclusionScope(FRDGBuilder& GraphBuilder);
+	void EndOcclusionScope(FRDGBuilder& GraphBuilder);
+
+	/** Renders the current batch and resets the batch state. */
+	void SubmitOcclusionDraws(FRHICommandList& RHICmdList, FViewInfo& View);
+
+	void ReadbackResults(FRHICommandList& RHICmdList);
+	void AdvanceFrame(uint32 OcclusionFrameCounter);
+
+	inline FRDGBuffer* GetGPUFeedbackBuffer() const
+	{
+		return GPUFeedbackBuffer;
+	}
+
+	inline bool IsOccluded(const FPrimitiveOcclusionHistoryKey& PrimitiveKey) const
+	{
+		return LatestOcclusionResults.Contains(PrimitiveKey);
+	}
+
+private:
+	struct FOcclusionBatch
+	{
+		FGlobalDynamicVertexBuffer::FAllocation VertexAllocation;
+		uint32 NumBatchedPrimitives;
+	};
+
+	/** The pending batches. */
+	TArray<FOcclusionBatch, TInlineAllocator<3>> BatchOcclusionQueries;
+
+	FRDGBuffer* GPUFeedbackBuffer{};
+
+	struct FOcclusionBuffer
+	{
+		TArray<FPrimitiveOcclusionHistoryKey> BatchedPrimitives;
+		FRHIGPUBufferReadback* ReadbackBuffer = nullptr;
+		uint32 OcclusionFrameCounter = 0u;
+	};
+
+	FOcclusionBuffer OcclusionBuffers[3];
+	uint32 CurrentBufferIndex;
+
+	TSet<FPrimitiveOcclusionHistoryKey> LatestOcclusionResults;
+	uint32 ResultsOcclusionFrameCounter;
+
+	//
+	FVertexDeclarationRHIRef OcclusionVertexDeclarationRHI;
+};
+
 /**
  * The scene manager's private implementation of persistent view state.
  * This class is associated with a particular camera across multiple frames by the game thread.
diff --git a/Engine/Source/Runtime/Renderer/Private/SceneRendering.h b/Engine/Source/Runtime/Renderer/Private/SceneRendering.h
--- a/Engine/Source/Runtime/Renderer/Private/SceneRendering.h
+++ b/Engine/Source/Runtime/Renderer/Private/SceneRendering.h
@@ -515,66 +515,6 @@ private:
 	uint32 ValidFrameNumber;
 };
 
-class FOcclusionFeedback : public FRenderResource
-{
-public:
-	FOcclusionFeedback();
-	~FOcclusionFeedback();
-
-	// FRenderResource interface
-	virtual void InitRHI(FRHICommandListBase& RHICmdList) override;
-	virtual void ReleaseRHI() override;
-
-	void AddPrimitive(FRHICommandList& RHICmdList, FPrimitiveComponentId PrimitiveId, const FVector& BoundsOrigin, const FVector& BoundsBoxExtent, FGlobalDynamicVertexBuffer& DynamicVertexBuffer);
-
-	void BeginOcclusionScope(FRDGBuilder& GraphBuilder);
-	void EndOcclusionScope(FRDGBuilder& GraphBuilder);
-
-	/** Renders the current batch and resets the batch state. */
-	void SubmitOcclusionDraws(FRHICommandList& RHICmdList, FViewInfo& View);
-	
-	void ReadbackResults(FRHICommandList& RHICmdList);
-	void AdvanceFrame(uint32 OcclusionFrameCounter);
-
-	inline FRDGBuffer* GetGPUFeedbackBuffer() const
-	{
-		return GPUFeedbackBuffer;
-	}
-
-	inline bool IsOccluded(FPrimitiveComponentId Id) const
-	{
-		return LatestOcclusionResults.Contains(Id);
-	}
-
-private:
-	struct FOcclusionBatch
-	{
-		FGlobalDynamicVertexBuffer::FAllocation VertexAllocation;
-		uint32 NumBatchedPrimitives;
-	};
-
-	/** The pending batches. */
-	TArray<FOcclusionBatch, TInlineAllocator<3>> BatchOcclusionQueries;
-
-	FRDGBuffer* GPUFeedbackBuffer{};
-
-	struct FOcclusionBuffer
-	{
-		TArray<FPrimitiveComponentId> BatchedPrimitives;
-		FRHIGPUBufferReadback* ReadbackBuffer = nullptr;
-		uint32 OcclusionFrameCounter = 0u;
-	};
-
-	FOcclusionBuffer OcclusionBuffers[3];
-	uint32 CurrentBufferIndex;
-	
-	TSet<FPrimitiveComponentId> LatestOcclusionResults;
-	uint32 ResultsOcclusionFrameCounter;
-
-	//
-	FVertexDeclarationRHIRef OcclusionVertexDeclarationRHI;
-};
-
 DECLARE_STATS_GROUP(TEXT("Parallel Command List Markers"), STATGROUP_ParallelCommandListMarkers, STATCAT_Advanced);
 
 /** Helper class to marshal data from your RDG pass into the parallel command list set. */
diff --git a/Engine/Source/Runtime/Renderer/Private/SceneVisibility.cpp b/Engine/Source/Runtime/Renderer/Private/SceneVisibility.cpp
--- a/Engine/Source/Runtime/Renderer/Private/SceneVisibility.cpp
+++ b/Engine/Source/Runtime/Renderer/Private/SceneVisibility.cpp
@@ -2358,7 +2358,7 @@ bool FGPUOcclusionPacket::OcclusionCullPrimitive(VisitorType& Visitor, FOcclusio
 			{
 				if (bUseOcclusionFeedback)
 				{
-					bIsOccluded = OcclusionFeedback.IsOccluded(PrimitiveId);
+					bIsOccluded = OcclusionFeedback.IsOccluded(FPrimitiveOcclusionHistoryKey(PrimitiveId, SubQuery));
 					bOcclusionStateIsDefinite = true;
 				}
 				else if (OcclusionState.bHZBOcclusion)
@@ -2491,7 +2491,7 @@ bool FGPUOcclusionPacket::OcclusionCullPrimitive(VisitorType& Visitor, FOcclusio
 					const FVector BoundOrigin = OcclusionBounds.Origin + View.ViewMatrices.GetPreViewTranslation();
 					const FVector BoundExtent = OcclusionBounds.BoxExtent;
 
-					Visitor.AddOcclusionFeedback(FOcclusionFeedbackEntry(PrimitiveId, BoundOrigin, BoundExtent));
+					Visitor.AddOcclusionFeedback(FOcclusionFeedbackEntry(FPrimitiveOcclusionHistoryKey(PrimitiveId, SubQuery), BoundOrigin, BoundExtent));
 				}
 				else if (OcclusionState.bHZBOcclusion)
 				{
diff --git a/Engine/Source/Runtime/Renderer/Private/SceneVisibilityPrivate.h b/Engine/Source/Runtime/Renderer/Private/SceneVisibilityPrivate.h
--- a/Engine/Source/Runtime/Renderer/Private/SceneVisibilityPrivate.h
+++ b/Engine/Source/Runtime/Renderer/Private/SceneVisibilityPrivate.h
@@ -756,12 +756,12 @@ struct FOcclusionQuery
 
 struct FOcclusionFeedbackEntry
 {
-	FOcclusionFeedbackEntry(FPrimitiveComponentId InPrimitiveId, FVector InBoundsOrigin, FVector InBoundsExtent)
-		: PrimitiveId(InPrimitiveId)
+	FOcclusionFeedbackEntry(const FPrimitiveOcclusionHistoryKey& PrimitiveKey, FVector InBoundsOrigin, FVector InBoundsExtent)
+		: PrimitiveKey(PrimitiveKey)
 		, Bounds(InBoundsOrigin, InBoundsExtent)
 	{}
 
-	const FPrimitiveComponentId PrimitiveId;
+	const FPrimitiveOcclusionHistoryKey PrimitiveKey;
 	const FOcclusionBounds Bounds;
 };
 
@@ -896,7 +896,7 @@ public:
 
 		void AddOcclusionFeedback(const FOcclusionFeedbackEntry& Entry)
 		{
-			Packet.OcclusionFeedback.AddPrimitive(RHICmdList, Entry.PrimitiveId, Entry.Bounds.Origin, Entry.Bounds.Extent, DynamicVertexBuffer);
+			Packet.OcclusionFeedback.AddPrimitive(RHICmdList, Entry.PrimitiveKey, Entry.Bounds.Origin, Entry.Bounds.Extent, DynamicVertexBuffer);
 		}
 
 		void AddVisualizeQuery(const FBox& Box)
