Allow multiple Field targets to be rendered at the same time

diff --git a/Engine/Shaders/Private/PhysicsFieldVisualizer.usf b/Engine/Shaders/Private/PhysicsFieldVisualizer.usf
--- a/Engine/Shaders/Private/PhysicsFieldVisualizer.usf
+++ b/Engine/Shaders/Private/PhysicsFieldVisualizer.usf
@@ -6,6 +6,7 @@
 #include "CommonViewUniformBuffer.ush"
 #include "SceneTextureParameters.ush"
 #include "DeferredShadingCommon.ush"
+#include "ShaderPrintCommon.ush" 
 #include "ShaderPrint.ush" 
 #include "PositionReconstructionCommon.ush" 
  
@@ -21,6 +22,8 @@ Buffer<float4> BoundsMax;
 uint BoundsOffset;
 uint BoundsSize;
 
+uint PrintOffset;
+
 uint LocalTarget;
 uint GlobalTarget;
 
@@ -71,30 +74,53 @@ void SetOutputTexture(inout float RayDensity, in int BoundsIndex, in uint2 Pixel
 // Vector field visualisation
 #if PERMUTATION_FIELD == 0
 
+FFontColor GetVectorTargetColor()
+{
+	if(GlobalTarget == LINEAR_FORCE)
+	{
+		return FontRed;
+	}
+	else if(GlobalTarget == LINEAR_VELOCITY)
+	{
+		return FontGreen;
+	}
+	else if(GlobalTarget == ANGULAR_TORQUE)
+	{
+		return FontBlue;
+	}
+	else if(GlobalTarget == ANGULAR_VELOCITY)
+	{
+		return FontYellow;
+	}
+	return FontWhite;
+}
+
 void PrintVectorField()
 {
 	FShaderPrintContext Context = InitShaderPrintContext(true, uint2(50.0f, 200.0f));
+	Context.Pos.y += Context.Config.FontSpacing.y * PrintOffset;
 
-	Print(Context, TEXT("Vector Field Type - "));
+	const FFontColor VectorColor = GetVectorTargetColor();
+
+	Print(Context, TEXT("Vector Field | Type = "));
 	if(GlobalTarget == LINEAR_FORCE)
 	{
-		Print(Context, TEXT("Linear Force "));
+		Print(Context, TEXT("Linear Force "), VectorColor);
 	}
 	else if(GlobalTarget == LINEAR_VELOCITY)
 	{
-		Print(Context, TEXT("Linear Velocity "));
+		Print(Context, TEXT("Linear Velocity "), VectorColor);
 	}
 	else if(GlobalTarget == ANGULAR_TORQUE)
 	{
-		Print(Context, TEXT("Angular Torque "));
+		Print(Context, TEXT("Angular Torque "), VectorColor);
 	}
 	else if(GlobalTarget == ANGULAR_VELOCITY)
 	{
-		Print(Context, TEXT("Angular Velocity "));
+		Print(Context, TEXT("Angular Velocity "), VectorColor);
 	}
-	Newline(Context);
 
-	Print(Context, TEXT("Vector Field Count - "));
+	Print(Context, TEXT("| Count = "));
 	Print(Context, BoundsSize);
 }
 
@@ -103,6 +129,7 @@ void MainCS(uint3 DispatchThreadId : SV_DispatchThreadID,
     uint3 GroupThreadId : SV_GroupThreadID)
 {
 	GFieldGroupThreadId = GroupThreadId.x * 16 + GroupThreadId.y * 4 + GroupThreadId.z;
+	const FFontColor AABBColor = GetVectorTargetColor();
 	if(DispatchThreadId.x < VOXEL_RESOLUTION * BoundsSize && DispatchThreadId.y < VOXEL_RESOLUTION && DispatchThreadId.z < VOXEL_RESOLUTION)
 	{
 		const uint BoundsIndex = BoundsOffset + DispatchThreadId.x / VOXEL_RESOLUTION;
@@ -140,7 +167,7 @@ void MainCS(uint3 DispatchThreadId : SV_DispatchThreadID,
 		{
 			for(int BoundsIndex = BoundsOffset, BoundsEnd = BoundsOffset + BoundsSize; BoundsIndex < BoundsEnd; ++BoundsIndex)
 			{
-				AddAABBWS(BoundsMin[BoundsIndex].xyz, BoundsMax[BoundsIndex].xyz, float4(1, 0.5f, 0, 1));
+				AddAABBWS(BoundsMin[BoundsIndex].xyz, BoundsMax[BoundsIndex].xyz, float4(AABBColor, 1));
 			}
 			PrintVectorField();
 		}
@@ -151,69 +178,123 @@ void MainCS(uint3 DispatchThreadId : SV_DispatchThreadID,
 
 #if PERMUTATION_FIELD == 1 || PERMUTATION_FIELD == 2
 
+FFontColor GetScalarTargetColor()
+{
+	if(GlobalTarget == EXTERNAL_STRAIN)
+	{
+		return FontCyan;
+	}
+	else if(GlobalTarget == INTERNAL_STRAIN)
+	{
+		return FontMagenta;
+	}
+	else if(GlobalTarget == DISABLED_THRESHOLD)
+	{
+		return FontOrange;
+	}
+	else if(GlobalTarget == SLEEPING_THRESHOLD)
+	{
+		return FontPurple;
+	}
+	else if(GlobalTarget == KILL_PARTICLES)
+	{
+		return FontTurquoise;
+	}
+	return FontWhite;
+}
+
 void PrintScalarField()
 {
 	FShaderPrintContext Context = InitShaderPrintContext(true, uint2(50.0f, 200.0f));
+	Context.Pos.y += Context.Config.FontSpacing.y * PrintOffset;
 
-	Print(Context, TEXT("Scalar Field Type - "));
+	const FFontColor ScalarColor = GetScalarTargetColor();
 
+	Print(Context, TEXT("Scalar Field | Type = "));
 	if(GlobalTarget == EXTERNAL_STRAIN)
 	{
-		Print(Context, TEXT("External Strain "));
+		Print(Context, TEXT("External Strain "), ScalarColor);
 	}
 	else if(GlobalTarget == INTERNAL_STRAIN)
 	{
-		Print(Context, TEXT("Internal Strain "));
+		Print(Context, TEXT("Internal Strain "), ScalarColor);
 	}
 	else if(GlobalTarget == DISABLED_THRESHOLD)
 	{
-		Print(Context, TEXT("Disabled Threshold "));
+		Print(Context, TEXT("Disabled Threshold "), ScalarColor);
 	}
 	else if(GlobalTarget == SLEEPING_THRESHOLD)
 	{
-		Print(Context, TEXT("Sleeping Threshold "));
+		Print(Context, TEXT("Sleeping Threshold "), ScalarColor);
 	}
 	else if(GlobalTarget == KILL_PARTICLES)
 	{
-		Print(Context, TEXT("Kill Particles "));
+		Print(Context, TEXT("Kill Particles "), ScalarColor);
 	}
 
-	Newline(Context);
-
-	Print(Context, TEXT("Scalar Field Count - "));
+	Print(Context, TEXT("| Count = "));
 	Print(Context, BoundsSize);
 }
 
+FFontColor GetIntegerTargetColor()
+{
+	if(GlobalTarget == DYNAMIC_STATE)
+	{
+		return FontSilver;
+	}
+	else if(GlobalTarget == COLLISION_GROUP)
+	{
+		return FontEmerald;
+	}
+	else if(GlobalTarget == ACTIVATE_DISABLED)
+	{
+		return FontWhite;
+	}
+	return FontWhite;
+}
+
 void PrintIntegerField()
 {
 	FShaderPrintContext Context = InitShaderPrintContext(true, uint2(50.0f, 200.0f));
+	Context.Pos.y += Context.Config.FontSpacing.y * PrintOffset;
 
-	Print(Context, TEXT("Integer Field Type - "));
+	const FFontColor IntegerColor = GetIntegerTargetColor();
 
+	Print(Context, TEXT("Integer Field | Type = "));
 	if(GlobalTarget == DYNAMIC_STATE)
 	{
-		Print(Context, TEXT("Dynamic State "));
+		Print(Context, TEXT("Dynamic State "), IntegerColor);
 	}
 	else if(GlobalTarget == COLLISION_GROUP)
 	{
-		Print(Context, TEXT("Collision Group "));
+		Print(Context, TEXT("Collision Group "), IntegerColor);
 	}
 	else if(GlobalTarget == ACTIVATE_DISABLED)
 	{
-		Print(Context, TEXT("Activate Disabled "));
+		Print(Context, TEXT("Activate Disabled "), IntegerColor);
 	}
 
-	Newline(Context);
-	Print(Context, TEXT("Integer Field Count - "));
+	Print(Context, TEXT("| Count = "));
 	Print(Context, BoundsSize);
 }
 
+
+
 // Integer and scalar field visualisation
 [numthreads(8, 8, 1)]
 void MainCS(uint3 DispatchThreadId : SV_DispatchThreadID,
     uint3 GroupThreadId : SV_GroupThreadID)
 {
 	GFieldGroupThreadId = GroupThreadId.x * 8 + GroupThreadId.y + GroupThreadId.z;
+
+	#if PERMUTATION_FIELD == 2
+	const FFontColor AABBColor = GetIntegerTargetColor();
+	#endif
+
+	#if PERMUTATION_FIELD == 1
+	const FFontColor AABBColor = GetScalarTargetColor();
+	#endif
+
 	if(DispatchThreadId.x < OutputResolution.x && DispatchThreadId.y < OutputResolution.y)
 	{
 		float3 RayBegin, RayEnd, RayDirection;
@@ -266,7 +347,7 @@ void MainCS(uint3 DispatchThreadId : SV_DispatchThreadID,
 		{
 			for(int BoundsIndex = BoundsOffset, BoundsEnd = BoundsOffset + BoundsSize; BoundsIndex < BoundsEnd; ++BoundsIndex)
 			{
-				AddAABBWS(BoundsMin[BoundsIndex].xyz, BoundsMax[BoundsIndex].xyz, float4(1, 0.5f, 0, 1));
+				AddAABBWS(BoundsMin[BoundsIndex].xyz, BoundsMax[BoundsIndex].xyz, float4(AABBColor.Color, 1));
 			}
 #if PERMUTATION_FIELD == 2
 			PrintIntegerField();
diff --git a/Engine/Source/Runtime/Engine/Classes/PhysicsField/PhysicsFieldComponent.h b/Engine/Source/Runtime/Engine/Classes/PhysicsField/PhysicsFieldComponent.h
--- a/Engine/Source/Runtime/Engine/Classes/PhysicsField/PhysicsFieldComponent.h
+++ b/Engine/Source/Runtime/Engine/Classes/PhysicsField/PhysicsFieldComponent.h
@@ -185,7 +185,7 @@ public:
 	 * Update the datas based on the new bounds and commands
 	 * @param FieldCommands - Field commands to be sampled
 	 */
-	void UpdateInstance(const float TimeSeconds);
+	void UpdateInstance(const float TimeSeconds, const bool bIsDebugBuffer);
 
 	/** The field system resource. */
 	FPhysicsFieldResource* FieldResource = nullptr;
@@ -243,13 +243,16 @@ public:
 	/** Add the transient field command */
 	ENGINE_API void AddTransientCommand(const FFieldSystemCommand& FieldCommand, const bool bIsGPUField);
 
-	/** Add the persitent field command */
+	/** Add the persistent field command */
 	ENGINE_API void AddPersistentCommand(const FFieldSystemCommand& FieldCommand, const bool bIsGPUField);
 
+	/** Add the construction field command */
+	ENGINE_API void AddConstructionCommand(const FFieldSystemCommand& FieldCommand);
+
 	/** Remove the transient field command */
 	ENGINE_API void RemoveTransientCommand(const FFieldSystemCommand& FieldCommand, const bool bIsGPUField);
 
-	/** Remove the persitent field command */
+	/** Remove the persistent field command */
 	ENGINE_API void RemovePersistentCommand(const FFieldSystemCommand& FieldCommand, const bool bIsGPUField);
 
 	/** Fill the transient commands intersecting the bounding box from the physics field */
@@ -266,9 +269,12 @@ public:
 	/** List of all the field transient commands in the world */
 	TArray<FFieldSystemCommand> TransientCommands[(uint8)(EFieldCommandBuffer::NumFieldBuffers)];
 
-	/** List of all the field persitent commands in the world */
+	/** List of all the field persistent commands in the world */
 	TArray<FFieldSystemCommand> PersistentCommands[(uint8)(EFieldCommandBuffer::NumFieldBuffers)];
 
+	/** List of all the field construction commands in the world */
+	TArray<FFieldSystemCommand> ConstructionCommands[(uint8)(EFieldCommandBuffer::NumFieldBuffers)];
+
 	/** The instance of the GPU field system. */
 	FPhysicsFieldInstance* FieldInstance = nullptr;
 
diff --git a/Engine/Source/Runtime/Engine/Private/PhysicsField/PhysicsFieldComponent.cpp b/Engine/Source/Runtime/Engine/Private/PhysicsField/PhysicsFieldComponent.cpp
--- a/Engine/Source/Runtime/Engine/Private/PhysicsField/PhysicsFieldComponent.cpp
+++ b/Engine/Source/Runtime/Engine/Private/PhysicsField/PhysicsFieldComponent.cpp
@@ -680,7 +680,7 @@ void FPhysicsFieldInstance::ReleaseInstance()
 	BoundsMax.Empty();
 }
 
-void FPhysicsFieldInstance::UpdateInstance(const float TimeSeconds)
+void FPhysicsFieldInstance::UpdateInstance(const float TimeSeconds, const bool bIsDebugBuffer)
 {
 	NodesOffsets.Reset();
 	NodesParams.Reset();
@@ -715,8 +715,17 @@ void FPhysicsFieldInstance::UpdateInstance(const float TimeSeconds)
 				{
 					UPhysicsFieldComponent::BuildCommandBounds(FieldCommand);
 
-					BoundsMin.Add(FVector4(FieldCommand.BoundingBox.Min, 0.0));
-					BoundsMax.Add(FVector4(FieldCommand.BoundingBox.Max, FieldCommand.MaxMagnitude));
+
+					FVector MinBounds = FieldCommand.BoundingBox.Min;
+					FVector MaxBounds = FieldCommand.BoundingBox.Max;
+					if (bIsDebugBuffer && ((MinBounds.X <= -FLT_MAX) || (MaxBounds.X >= FLT_MAX)))
+					{
+						// When no bounds default size is 1m for the box
+						MinBounds = FieldCommand.CenterPosition - FVector(50, 50, 50);
+						MaxBounds = FieldCommand.CenterPosition + FVector(50, 50, 50);
+					}
+					BoundsMin.Add(FVector4(MinBounds, 0.0));
+					BoundsMax.Add(FVector4(MaxBounds, FieldCommand.MaxMagnitude));
 
 					TargetMin = MinVector(TargetMin, FieldCommand.BoundingBox.Min);
 					TargetMax = MaxVector(TargetMax, FieldCommand.BoundingBox.Max);
@@ -847,6 +856,9 @@ void UPhysicsFieldComponent::SendRenderDynamicData_Concurrent()
 
 	Super::SendRenderTransform_Concurrent();
 
+	static IConsoleVariable* TransientLifetimeCVar = IConsoleManager::Get().FindConsoleVariable(TEXT("r.PhysicsField.Rendering.TransientLifetime"));
+	const float TransientLifetime = TransientLifetimeCVar ? TransientLifetimeCVar->GetFloat() : 0.0f;
+
 	TArray<uint8> FieldBuffers = { (uint8)(EFieldCommandBuffer::GPUFieldBuffer) ,(uint8)(EFieldCommandBuffer::GPUDebugBuffer) };
 	for (uint32 FieldIndex = 0; FieldIndex < 2; ++FieldIndex)
 	{
@@ -854,18 +866,35 @@ void UPhysicsFieldComponent::SendRenderDynamicData_Concurrent()
 		if (LocalInstance)
 		{
 			const bool bPreviousUpdate = LocalInstance->FieldCommands.Num() > 0;
+			const bool bIsDebugBuffer = FieldIndex == 1;
+
+			const float TimeSeconds = GetWorld() ? GetWorld()->TimeSeconds : 0.0;
 
 			LocalInstance->FieldCommands.Empty();
 			LocalInstance->FieldCommands.Append(PersistentCommands[FieldBuffers[FieldIndex]]);
 			LocalInstance->FieldCommands.Append(TransientCommands[FieldBuffers[FieldIndex]]);
-			TransientCommands[FieldBuffers[FieldIndex]].Empty();
+			LocalInstance->FieldCommands.Append(ConstructionCommands[FieldBuffers[FieldIndex]]);
+			
+			if (bIsDebugBuffer)
+			{
+				for (int32 CommandIndex = TransientCommands[FieldBuffers[FieldIndex]].Num() - 1; CommandIndex >= 0; --CommandIndex)
+				{
+					if (TimeSeconds - TransientCommands[FieldBuffers[FieldIndex]][CommandIndex].TimeCreation > TransientLifetime)
+					{
+						TransientCommands[FieldBuffers[FieldIndex]].RemoveAt(CommandIndex);
+					}
+				}
+			}
+			else
+			{
+				TransientCommands[FieldBuffers[FieldIndex]].Empty();
+			}
 
 			const bool bCurrentUpdate = LocalInstance->FieldCommands.Num() > 0;
 
 			if (bCurrentUpdate || bPreviousUpdate)
 			{
-				const float TimeSeconds = GetWorld() ? GetWorld()->TimeSeconds : 0.0;
-				LocalInstance->UpdateInstance(TimeSeconds);
+				LocalInstance->UpdateInstance(TimeSeconds, bIsDebugBuffer);
 			}
 		}
 	}
@@ -956,11 +985,12 @@ void UPhysicsFieldComponent::BuildCommandBounds(FFieldSystemCommand& FieldComman
 	FieldCommand.BoundingBox.Min = FVector(-FLT_MAX);
 	FieldCommand.BoundingBox.Max = FVector(FLT_MAX);
 	FieldCommand.MaxMagnitude = 1.0;
+	FieldCommand.CenterPosition = FVector::Zero();
 
 	if (FieldCommand.RootNode.IsValid())
 	{
 		FieldCommand.MaxMagnitude = FieldCommand.RootNode->EvalMaxMagnitude();
-		FieldCommand.RootNode->ComputeFieldBounds(FieldCommand.BoundingBox.Min, FieldCommand.BoundingBox.Max);
+		FieldCommand.RootNode->ComputeFieldBounds(FieldCommand.BoundingBox.Min, FieldCommand.BoundingBox.Max, FieldCommand.CenterPosition);
 	}
 }
 
@@ -990,6 +1020,11 @@ void UPhysicsFieldComponent::AddPersistentCommand(const FFieldSystemCommand& Fie
 	}
 }
 
+void UPhysicsFieldComponent::AddConstructionCommand(const FFieldSystemCommand& FieldCommand)
+{
+	ConstructionCommands[(uint8)(EFieldCommandBuffer::GPUDebugBuffer)].Add(FieldCommand);
+}
+
 void UPhysicsFieldComponent::RemoveTransientCommand(const FFieldSystemCommand& FieldCommand, const bool bIsWorldField)
 {
 	if (bIsWorldField)
diff --git a/Engine/Source/Runtime/Experimental/Chaos/Private/Field/FieldSystemNodes.cpp b/Engine/Source/Runtime/Experimental/Chaos/Private/Field/FieldSystemNodes.cpp
--- a/Engine/Source/Runtime/Experimental/Chaos/Private/Field/FieldSystemNodes.cpp
+++ b/Engine/Source/Runtime/Experimental/Chaos/Private/Field/FieldSystemNodes.cpp
@@ -249,10 +249,11 @@ float FRadialIntMask::EvalMaxMagnitude() const
 	return 1.0f;
 }
 
-void FRadialIntMask::ComputeFieldBounds(FVector& MinBounds, FVector& MaxBounds) const
+void FRadialIntMask::ComputeFieldBounds(FVector& MinBounds, FVector& MaxBounds, FVector& CenterPosition) const
 {
 	MinBounds = (ExteriorValue == 0) ? Position - FVector(Radius) : FVector(-FLT_MAX);
 	MaxBounds = (ExteriorValue == 0) ? Position + FVector(Radius) : FVector(FLT_MAX);
+	CenterPosition = Position;
 }
 
 /**
@@ -423,6 +424,13 @@ float FWaveScalar::EvalMaxMagnitude() const
 	return Magnitude;
 }
 
+void FWaveScalar::ComputeFieldBounds(FVector& MinBounds, FVector& MaxBounds, FVector& CenterPosition) const
+{
+	MinBounds = FVector(-FLT_MAX);
+	MaxBounds = FVector(FLT_MAX);
+	CenterPosition = Position;
+}
+
 /**
 * Function Utils
 */
@@ -574,10 +582,11 @@ float FRadialFalloff::EvalMaxMagnitude() const
 	return Magnitude;
 }
 
-void FRadialFalloff::ComputeFieldBounds(FVector& MinBounds, FVector& MaxBounds) const
+void FRadialFalloff::ComputeFieldBounds(FVector& MinBounds, FVector& MaxBounds, FVector& CenterPosition) const
 {
 	MinBounds = (Default == 0) ? Position - FVector(Radius) : FVector(-FLT_MAX);
 	MaxBounds = (Default == 0) ? Position + FVector(Radius) : FVector(FLT_MAX);
+	CenterPosition = Position;
 }
 
 /**
@@ -690,10 +699,11 @@ float FPlaneFalloff::EvalMaxMagnitude() const
 	return Magnitude;
 }
 
-void FPlaneFalloff::ComputeFieldBounds(FVector& MinBounds, FVector& MaxBounds) const
+void FPlaneFalloff::ComputeFieldBounds(FVector& MinBounds, FVector& MaxBounds, FVector& CenterPosition) const
 {
 	MinBounds = FVector(-FLT_MAX);
 	MaxBounds = FVector(FLT_MAX);
+	CenterPosition = Position;
 }
 
 /**
@@ -808,7 +818,7 @@ float FBoxFalloff::EvalMaxMagnitude() const
 	return Magnitude;
 }
 
-void FBoxFalloff::ComputeFieldBounds(FVector& MinBounds, FVector& MaxBounds) const
+void FBoxFalloff::ComputeFieldBounds(FVector& MinBounds, FVector& MaxBounds, FVector& CenterPosition) const
 {
 	MinBounds = FVector(-FLT_MAX);
 	MaxBounds = FVector(FLT_MAX);
@@ -819,6 +829,7 @@ void FBoxFalloff::ComputeFieldBounds(FVector& MinBounds, FVector& MaxBounds) con
 		MinBounds = BoundingBox.Min;
 		MaxBounds = BoundingBox.Max;
 	}
+	CenterPosition = Transform.GetTranslation();
 }
 
 /**
@@ -1010,6 +1021,13 @@ float FRadialVector::EvalMaxMagnitude() const
 	return Magnitude;
 }
 
+void FRadialVector::ComputeFieldBounds(FVector& MinBounds, FVector& MaxBounds, FVector& CenterPosition) const
+{
+	MinBounds = FVector(-FLT_MAX);
+	MaxBounds = FVector(FLT_MAX);
+	CenterPosition = Position;
+}
+
 
 /**
 * FRandomVector
@@ -1249,17 +1267,31 @@ float FSumScalar::EvalMaxMagnitude() const
 	return Magnitude * MaxMagnitude;
 }
 
-void FSumScalar::ComputeFieldBounds(FVector& MinBounds, FVector& MaxBounds) const
+void FSumScalar::ComputeFieldBounds(FVector& MinBounds, FVector& MaxBounds, FVector& CenterPosition) const
 {
 	FVector MinBoundsA(-FLT_MAX), MaxBoundsA(FLT_MAX), MinBoundsB(-FLT_MAX), MaxBoundsB(FLT_MAX);
+	FVector CenterPositionA = FVector::Zero(), CenterPositionB = FVector::Zero();
 	if (ScalarRight.IsValid())
 	{
-		ScalarRight->ComputeFieldBounds(MinBoundsA, MaxBoundsA);
+		ScalarRight->ComputeFieldBounds(MinBoundsA, MaxBoundsA, CenterPositionA);
 	}
 	if (ScalarLeft.IsValid())
 	{
-		ScalarLeft->ComputeFieldBounds(MinBoundsB, MaxBoundsB);
+		ScalarLeft->ComputeFieldBounds(MinBoundsB, MaxBoundsB, CenterPositionB);
+	}
+	CenterPosition = FVector::Zero();
+	int32 NumCenters = 0;
+	if (CenterPositionA != FVector::Zero())
+	{
+		CenterPosition += CenterPositionA;
+		++NumCenters;
+	}
+	if (CenterPositionB != FVector::Zero())
+	{
+		CenterPosition += CenterPositionB;
+		++NumCenters;
 	}
+	CenterPosition = (NumCenters > 0) ? CenterPosition / NumCenters : FVector::Zero();
 
 	if (Operation == EFieldOperationType::Field_Multiply ||
 		Operation == EFieldOperationType::Field_Divide)
@@ -1479,22 +1511,42 @@ float FSumVector::EvalMaxMagnitude() const
 	return Magnitude * MaxMagnitudeA * MaxMagnitude;
 }
 
-void FSumVector::ComputeFieldBounds(FVector& MinBounds, FVector& MaxBounds) const
+void FSumVector::ComputeFieldBounds(FVector& MinBounds, FVector& MaxBounds, FVector& CenterPosition) const
 {
 	FVector MinBoundsA(-FLT_MAX), MaxBoundsA(FLT_MAX), MinBoundsB(-FLT_MAX), MaxBoundsB(FLT_MAX), MinBoundsC(-FLT_MAX), MaxBoundsC(FLT_MAX);
+	FVector CenterPositionA = FVector::Zero(), CenterPositionB = FVector::Zero(), CenterPositionC = FVector::Zero();
 	if (Scalar.IsValid())
 	{
-		Scalar->ComputeFieldBounds(MinBoundsC, MaxBoundsC);
+		Scalar->ComputeFieldBounds(MinBoundsC, MaxBoundsC, CenterPositionC);
 	}
 	if (VectorRight.IsValid())
 	{
-		VectorRight->ComputeFieldBounds(MinBoundsA, MaxBoundsA);
+		VectorRight->ComputeFieldBounds(MinBoundsA, MaxBoundsA, CenterPositionA);
 	}
 	if (VectorLeft.IsValid())
 	{
-		VectorLeft->ComputeFieldBounds(MinBoundsB, MaxBoundsB);
+		VectorLeft->ComputeFieldBounds(MinBoundsB, MaxBoundsB, CenterPositionB);
 	}
 
+	CenterPosition = FVector::Zero();
+	int32 NumCenters = 0;
+	if (CenterPositionA != FVector::Zero())
+	{
+		CenterPosition += CenterPositionA;
+		++NumCenters;
+	}
+	if (CenterPositionB != FVector::Zero())
+	{
+		CenterPosition += CenterPositionB;
+		++NumCenters;
+	}
+	if (CenterPositionC != FVector::Zero())
+	{
+		CenterPosition += CenterPositionC;
+		++NumCenters;
+	}
+	CenterPosition = (NumCenters > 0) ? CenterPosition / NumCenters : FVector::Zero();
+
 	if (Operation == EFieldOperationType::Field_Multiply ||
 		Operation == EFieldOperationType::Field_Divide)
 	{
@@ -1596,12 +1648,12 @@ float FConversionField<InT, OutT>::EvalMaxMagnitude() const
 }
 
 template<class InT, class OutT>
-void FConversionField<InT, OutT>::ComputeFieldBounds(FVector& MinBounds, FVector& MaxBounds) const
+void FConversionField<InT, OutT>::ComputeFieldBounds(FVector& MinBounds, FVector& MaxBounds, FVector& CenterPosition) const
 {
 	FVector MinBoundsA(-FLT_MAX), MaxBoundsA(FLT_MAX);
 	if (InputField.IsValid())
 	{
-		InputField->ComputeFieldBounds(MinBoundsA, MaxBoundsA);
+		InputField->ComputeFieldBounds(MinBoundsA, MaxBoundsA, CenterPosition);
 	}
 	MinBounds = MinBoundsA;
 	MaxBounds = MaxBoundsA;
@@ -1760,18 +1812,34 @@ float FCullingField<T>::EvalMaxMagnitude() const
 }
 
 template<class T>
-void FCullingField<T>::ComputeFieldBounds(FVector& MinBounds, FVector& MaxBounds) const
+void FCullingField<T>::ComputeFieldBounds(FVector& MinBounds, FVector& MaxBounds, FVector& CenterPosition) const
 {
 	FVector MinBoundsA(-FLT_MAX), MaxBoundsA(FLT_MAX), MinBoundsB(-FLT_MAX), MaxBoundsB(FLT_MAX);
+	FVector CenterPositionA = FVector::Zero(), CenterPositionB = FVector::Zero();
 
 	if (Culling.IsValid())
 	{
-		Culling->ComputeFieldBounds(MinBoundsA, MaxBoundsA);
+		Culling->ComputeFieldBounds(MinBoundsA, MaxBoundsA, CenterPositionA);
 	}
 	if (Input.IsValid())
 	{
-		Input->ComputeFieldBounds(MinBoundsB, MaxBoundsB);
+		Input->ComputeFieldBounds(MinBoundsB, MaxBoundsB, CenterPositionB);
 	}
+
+	CenterPosition = FVector::Zero();
+	int32 NumCenters = 0;
+	if (CenterPositionA != FVector::Zero())
+	{
+		CenterPosition += CenterPositionA;
+		++NumCenters;
+	}
+	if (CenterPositionB != FVector::Zero())
+	{
+		CenterPosition += CenterPositionB;
+		++NumCenters;
+	}
+	CenterPosition = (NumCenters > 0) ? CenterPosition / NumCenters : FVector::Zero();
+
 	if (Operation == EFieldCullingOperationType::Field_Culling_Inside)
 	{
 		MinBounds = MinVector(MinBoundsA, MinBoundsB);
diff --git a/Engine/Source/Runtime/Experimental/Chaos/Public/Field/FieldSystem.h b/Engine/Source/Runtime/Experimental/Chaos/Public/Field/FieldSystem.h
--- a/Engine/Source/Runtime/Experimental/Chaos/Public/Field/FieldSystem.h
+++ b/Engine/Source/Runtime/Experimental/Chaos/Public/Field/FieldSystem.h
@@ -456,10 +456,11 @@ public:
 	virtual float EvalMaxMagnitude() const { return 1.0; }
 
 	/** Compute the min/max spatial bounds of the field */
-	virtual void ComputeFieldBounds(FVector& MinBounds, FVector& MaxBounds) const
+	virtual void ComputeFieldBounds(FVector& MinBounds, FVector& MaxBounds, FVector& CenterPosition) const
 	{
 		MinBounds = FVector(-FLT_MAX);
 		MaxBounds = FVector(FLT_MAX);
+		CenterPosition = FVector::Zero();
 	}
 
 	FName GetName() const { return Name; }
@@ -531,6 +532,7 @@ public:
 		, BoundingBox(FVector(-FLT_MAX), FVector(FLT_MAX))
 		, PhysicsType(EFieldPhysicsType::Field_None)
 		, MaxMagnitude(1.0)
+		, CenterPosition(FVector::Zero())
 	{}
 	FFieldSystemCommand(const FName& TargetAttributeIn, FFieldNodeBase * RootNodeIn)
 		: TargetAttribute(TargetAttributeIn)
@@ -540,6 +542,7 @@ public:
 		, BoundingBox(FVector(-FLT_MAX), FVector(FLT_MAX))
 		, PhysicsType(GetFieldPhysicsType(TargetAttributeIn))
 		, MaxMagnitude(1.0)
+		, CenterPosition(FVector::Zero())
 	{}
 	FFieldSystemCommand(const EFieldPhysicsType PhsyicsTypeIn, FFieldNodeBase* RootNodeIn)
 		: TargetAttribute(GetFieldPhysicsName(PhsyicsTypeIn))
@@ -549,6 +552,7 @@ public:
 		, BoundingBox(FVector(-FLT_MAX), FVector(FLT_MAX))
 		, PhysicsType(PhsyicsTypeIn)
 		, MaxMagnitude(1.0)
+		, CenterPosition(FVector::Zero())
 	{}
 
 	// Commands are copied when moved from the one thread to 
@@ -561,6 +565,7 @@ public:
 		, BoundingBox(Other.BoundingBox)
 		, PhysicsType(Other.RootNode ? Other.PhysicsType : EFieldPhysicsType::Field_None)
 		, MaxMagnitude(Other.MaxMagnitude)
+		, CenterPosition(Other.CenterPosition)
 	{
 		for (const TPair<FFieldSystemMetaData::EMetaType, TUniquePtr<FFieldSystemMetaData>>& Meta : Other.MetaData)
 		{
@@ -614,6 +619,7 @@ public:
 	FBox BoundingBox;
 	EFieldPhysicsType PhysicsType;
 	float MaxMagnitude;
+	FVector CenterPosition;
 
 	TMap<FFieldSystemMetaData::EMetaType, TUniquePtr<FFieldSystemMetaData> > MetaData;
 };
diff --git a/Engine/Source/Runtime/Experimental/Chaos/Public/Field/FieldSystemNodes.h b/Engine/Source/Runtime/Experimental/Chaos/Public/Field/FieldSystemNodes.h
--- a/Engine/Source/Runtime/Experimental/Chaos/Public/Field/FieldSystemNodes.h
+++ b/Engine/Source/Runtime/Experimental/Chaos/Public/Field/FieldSystemNodes.h
@@ -85,7 +85,7 @@ public:
 	CHAOS_API virtual float EvalMaxMagnitude() const override;
 
 	/** Compute the min/max spatial bounds of the field */
-	CHAOS_API virtual void ComputeFieldBounds(FVector& MinBounds, FVector& MaxBounds) const override;
+	CHAOS_API virtual void ComputeFieldBounds(FVector& MinBounds, FVector& MaxBounds, FVector& CenterPosition) const override;
 
 	float Radius;
 	FVector Position;
@@ -172,6 +172,8 @@ public:
 	/** Evaluate the maximum magnitude of the field graph  */
 	CHAOS_API virtual float EvalMaxMagnitude() const override;
 
+	/** Compute the min/max spatial bounds of the field */
+	CHAOS_API virtual void ComputeFieldBounds(FVector& MinBounds, FVector& MaxBounds, FVector& CenterPosition) const override;
 
 	float Magnitude;
 	FVector Position;
@@ -229,7 +231,7 @@ public:
 	CHAOS_API virtual float EvalMaxMagnitude() const override;
 
 	/** Compute the min/max spatial bounds of the field */
-	CHAOS_API virtual void ComputeFieldBounds(FVector& MinBounds, FVector& MaxBounds) const override;
+	CHAOS_API virtual void ComputeFieldBounds(FVector& MinBounds, FVector& MaxBounds, FVector& CenterPosition) const override;
 
 	float Magnitude;
 	float MinRange;
@@ -293,7 +295,7 @@ public:
 	CHAOS_API virtual float EvalMaxMagnitude() const override;
 
 	/** Compute the min/max spatial bounds of the field */
-	CHAOS_API virtual void ComputeFieldBounds(FVector& MinBounds, FVector& MaxBounds) const override;
+	CHAOS_API virtual void ComputeFieldBounds(FVector& MinBounds, FVector& MaxBounds, FVector& CenterPosition) const override;
 
 
 	float Magnitude;
@@ -355,7 +357,7 @@ public:
 	CHAOS_API virtual float EvalMaxMagnitude() const override;
 
 	/** Compute the min/max spatial bounds of the field */
-	CHAOS_API virtual void ComputeFieldBounds(FVector& MinBounds, FVector& MaxBounds) const override;
+	CHAOS_API virtual void ComputeFieldBounds(FVector& MinBounds, FVector& MaxBounds, FVector& CenterPosition) const override;
 
 	float Magnitude;
 	float MinRange;
@@ -483,6 +485,9 @@ public:
 	/** Evaluate the maximum magnitude of the field graph  */
 	CHAOS_API virtual float EvalMaxMagnitude() const override;
 
+	/** Compute the min/max spatial bounds of the field */
+	CHAOS_API virtual void ComputeFieldBounds(FVector& MinBounds, FVector& MaxBounds, FVector& CenterPosition) const override;
+
 	float Magnitude;
 	FVector Position;
 
@@ -570,7 +575,7 @@ public:
 	CHAOS_API virtual float EvalMaxMagnitude() const override;
 
 	/** Compute the min/max spatial bounds of the field */
-	CHAOS_API virtual void ComputeFieldBounds(FVector& MinBounds, FVector& MaxBounds) const override;
+	CHAOS_API virtual void ComputeFieldBounds(FVector& MinBounds, FVector& MaxBounds, FVector& CenterPosition) const override;
 
 	/** Default setup for the sum scalar field */
 	static void FillDefaultSetup(TArray<int32>& NodesOffsets, TArray<float>& NodesParams)
@@ -639,7 +644,7 @@ public:
 	CHAOS_API virtual float EvalMaxMagnitude() const override;
 
 	/** Compute the min/max spatial bounds of the field */
-	CHAOS_API virtual void ComputeFieldBounds(FVector& MinBounds, FVector& MaxBounds) const override;
+	CHAOS_API virtual void ComputeFieldBounds(FVector& MinBounds, FVector& MaxBounds, FVector& CenterPosition) const override;
 
 	/** Default setup for the sum vector field */
 	static void FillDefaultSetup(TArray<int32>& NodesOffsets, TArray<float>& NodesParams)
@@ -699,7 +704,7 @@ public:
 	CHAOS_API virtual float EvalMaxMagnitude() const override;
 
 	/** Compute the min/max spatial bounds of the field */
-	CHAOS_API virtual void ComputeFieldBounds(FVector& MinBounds, FVector& MaxBounds) const override;
+	CHAOS_API virtual void ComputeFieldBounds(FVector& MinBounds, FVector& MaxBounds, FVector& CenterPosition) const override;
 
 	/** Default setup for the conversion field */
 	static void FillDefaultSetup(TArray<int32>& NodesOffsets, TArray<float>& NodesParams)
@@ -760,7 +765,7 @@ public:
 	CHAOS_API virtual float EvalMaxMagnitude() const override;
 
 	/** Compute the min/max spatial bounds of the field */
-	CHAOS_API virtual void ComputeFieldBounds(FVector& MinBounds, FVector& MaxBounds) const override;
+	CHAOS_API virtual void ComputeFieldBounds(FVector& MinBounds, FVector& MaxBounds, FVector& CenterPosition) const override;
 
 	TUniquePtr< FFieldNode<float> > Culling;
 	TUniquePtr< FFieldNode<T> > Input;
diff --git a/Engine/Source/Runtime/Experimental/Chaos/Public/Field/FieldSystemTypes.h b/Engine/Source/Runtime/Experimental/Chaos/Public/Field/FieldSystemTypes.h
--- a/Engine/Source/Runtime/Experimental/Chaos/Public/Field/FieldSystemTypes.h
+++ b/Engine/Source/Runtime/Experimental/Chaos/Public/Field/FieldSystemTypes.h
@@ -339,7 +339,7 @@ inline CHAOS_API const TArray<FName>& GetFieldPhysicsNames()
 	static const TArray<FName> FieldPhysicsNames =
 	{ "Node", "DynamicState", "LinearForce", "ExternalClusterStrain", "Kill", "LinearVelocity", "AngularVelocity",
 	 "AngularTorque", "InternalClusterStrain", "DisableThreshold", "SleepingThreshold", "PositionStatic",
-	 "PositionAnimated", "PositionTarget", "DynamicConstraint", "CollisionGroup", "ActivateDisabled", "InitialLinearVelocity", "InitialAngularVelocity" };
+	 "PositionAnimated", "PositionTarget", "DynamicConstraint", "CollisionGroup", "ActivateDisabled", "InitialLinearVelocity", "InitialAngularVelocity", "LinearImpulse"};
 
 	return FieldPhysicsNames;
 }
@@ -364,7 +364,8 @@ inline CHAOS_API const TArray<EFieldPhysicsType>& GetFieldPhysicsTypes()
 																EFieldPhysicsType::Field_CollisionGroup,
 																EFieldPhysicsType::Field_ActivateDisabled, 
 																EFieldPhysicsType::Field_InitialLinearVelocity, 
-																EFieldPhysicsType::Field_InitialAngularVelocity };
+																EFieldPhysicsType::Field_InitialAngularVelocity,
+																EFieldPhysicsType::Field_LinearImpulse };
 
 	return FieldPhysicsTypes;
 }
diff --git a/Engine/Source/Runtime/Experimental/FieldSystem/Source/FieldSystemEngine/Private/Field/FieldSystemComponent.cpp b/Engine/Source/Runtime/Experimental/FieldSystem/Source/FieldSystemEngine/Private/Field/FieldSystemComponent.cpp
--- a/Engine/Source/Runtime/Experimental/FieldSystem/Source/FieldSystemEngine/Private/Field/FieldSystemComponent.cpp
+++ b/Engine/Source/Runtime/Experimental/FieldSystem/Source/FieldSystemEngine/Private/Field/FieldSystemComponent.cpp
@@ -296,6 +296,17 @@ void UFieldSystemComponent::AddFieldCommand(bool Enabled, EFieldPhysicsType Targ
 			if (Command.RootNode)
 			{
 				SetupConstructionFields.Add(Command);
+
+				UWorld* World = GetWorld();
+				if (World && World->PhysicsField)
+				{
+					const FName Name = GetOwner() ? *GetOwner()->GetName() : TEXT("");
+
+					FFieldSystemCommand LocalCommand = Command;
+					LocalCommand.InitFieldNodes(World->GetTimeSeconds(), Name);
+
+					World->PhysicsField->AddConstructionCommand(LocalCommand);
+				}
 			}
 		}
 
diff --git a/Engine/Source/Runtime/Experimental/GeometryCollectionEngine/Private/GeometryCollection/GeometryCollectionComponent.cpp b/Engine/Source/Runtime/Experimental/GeometryCollectionEngine/Private/GeometryCollection/GeometryCollectionComponent.cpp
--- a/Engine/Source/Runtime/Experimental/GeometryCollectionEngine/Private/GeometryCollection/GeometryCollectionComponent.cpp
+++ b/Engine/Source/Runtime/Experimental/GeometryCollectionEngine/Private/GeometryCollection/GeometryCollectionComponent.cpp
@@ -4581,6 +4581,17 @@ void UGeometryCollectionComponent::GetInitializationCommands(TArray<FFieldSystem
 						if (NewCommand.RootNode)
 						{
 							CombinedCommmands.Emplace(NewCommand);
+
+							UWorld* World = GetWorld();
+							if (World && World->PhysicsField)
+							{
+								const FName Name = GetOwner() ? *GetOwner()->GetName() : TEXT("");
+
+								FFieldSystemCommand LocalCommand = NewCommand;
+								LocalCommand.InitFieldNodes(World->GetTimeSeconds(), Name);
+
+								World->PhysicsField->AddConstructionCommand(LocalCommand);
+							}
 						}
 					}
 				}
diff --git a/Engine/Source/Runtime/Renderer/Private/PhysicsFieldRendering.cpp b/Engine/Source/Runtime/Renderer/Private/PhysicsFieldRendering.cpp
--- a/Engine/Source/Runtime/Renderer/Private/PhysicsFieldRendering.cpp
+++ b/Engine/Source/Runtime/Renderer/Private/PhysicsFieldRendering.cpp
@@ -20,7 +20,7 @@
 #include "ShaderPrint.h"
 
 // Console variables
-int32 GPhysicsFieldTargetType = EFieldPhysicsType::Field_LinearForce;
+int32 GPhysicsFieldTargetType = EFieldPhysicsType::Field_PhysicsType_Max;
 FAutoConsoleVariableRef CVarPhysicsFieldTargetType(
 	TEXT("r.PhysicsField.Rendering.TargetType"),
 	GPhysicsFieldTargetType,
@@ -41,6 +41,13 @@ FAutoConsoleVariableRef CVarPhysicsFieldSystemType(
 	TEXT("Physics field boolean to check if we want to display the CPU(0) or GPU(1) field.\n"),
 	ECVF_RenderThreadSafe);
 
+float GPhysicsFieldTransientLifetime = 3.0f;
+FAutoConsoleVariableRef CVarPhysicsFieldTransientLifetime(
+	TEXT("r.PhysicsField.Rendering.TransientLifetime"),
+	GPhysicsFieldTransientLifetime,
+	TEXT("Physics field transient commands lifetime for rendering.\n"),
+	ECVF_RenderThreadSafe);
+
 /**
  * Physics Field Rendering.
  */
@@ -69,6 +76,7 @@ class FPhysicsFieldRayMarchingCS : public FGlobalShader
 		SHADER_PARAMETER(uint32, BoundsOffset)
 		SHADER_PARAMETER(uint32, BoundsSize)
 		SHADER_PARAMETER(uint32, LocalTarget)
+		SHADER_PARAMETER(uint32, PrintOffset)
 		SHADER_PARAMETER(uint32, GlobalTarget)
 		SHADER_PARAMETER(FVector2f, OutputResolution)
 		SHADER_PARAMETER_RDG_TEXTURE_UAV(RWTexture2D, OutputTexture)
@@ -93,7 +101,7 @@ IMPLEMENT_GLOBAL_SHADER(FPhysicsFieldRayMarchingCS, "/Engine/Private/PhysicsFiel
 
 static FPhysicsFieldRayMarchingCS::FParameters* CreateShaderParameters(FRDGBuilder& GraphBuilder,
 	const FViewInfo& View, const FPhysicsFieldResource* PhysicsFieldResource,
-	FRDGTextureRef& OutputTexture, const int32 TargetIndex)
+	FRDGTextureRef& OutputTexture, const int32 TargetIndex, const int32 TargetType, const uint32 PrintOffset)
 {
 	const FIntPoint OutputResolution(OutputTexture->Desc.Extent);
 	FSceneTextureParameters SceneTextures = GetSceneTextureParameters(GraphBuilder, View);
@@ -108,8 +116,8 @@ static FPhysicsFieldRayMarchingCS::FParameters* CreateShaderParameters(FRDGBuild
 
 	Parameters->BoundsMin = PhysicsFieldResource->BoundsMin.SRV;
 	Parameters->BoundsMax = PhysicsFieldResource->BoundsMax.SRV;
-	Parameters->BoundsOffset = PhysicsFieldResource->FieldInfos.BoundsOffsets[GPhysicsFieldTargetType];
-	Parameters->BoundsSize = PhysicsFieldResource->FieldInfos.BoundsOffsets[GPhysicsFieldTargetType + 1] - PhysicsFieldResource->FieldInfos.BoundsOffsets[GPhysicsFieldTargetType];;
+	Parameters->BoundsOffset = PhysicsFieldResource->FieldInfos.BoundsOffsets[TargetType];
+	Parameters->BoundsSize = PhysicsFieldResource->FieldInfos.BoundsOffsets[TargetType + 1] - PhysicsFieldResource->FieldInfos.BoundsOffsets[TargetType];;
 
 	Parameters->NodesParams = PhysicsFieldResource->NodesParams.SRV;
 	Parameters->NodesOffsets = PhysicsFieldResource->NodesOffsets.SRV;
@@ -117,7 +125,8 @@ static FPhysicsFieldRayMarchingCS::FParameters* CreateShaderParameters(FRDGBuild
 	Parameters->TimeSeconds = PhysicsFieldResource->FieldInfos.TimeSeconds;
 
 	Parameters->LocalTarget = TargetIndex;
-	Parameters->GlobalTarget = GPhysicsFieldTargetType;
+	Parameters->GlobalTarget = TargetType;
+	Parameters->PrintOffset = PrintOffset;
 
 	return Parameters;
 }
@@ -128,47 +137,51 @@ static void AddPhysicsFieldRayMarchingPass(
 	const FPhysicsFieldResource* PhysicsFieldResource,
 	FRDGTextureRef& OutputTexture)
 {
-	FSceneTextureParameters SceneTextures = GetSceneTextureParameters(GraphBuilder, View);
-
 	const FIntPoint OutputResolution(OutputTexture->Desc.Extent);
-	if(GPhysicsFieldTargetType < EFieldPhysicsType::Field_PhysicsType_Max)
-	{
+	auto RenderPhysicsField = [&GraphBuilder, &View, PhysicsFieldResource, &OutputTexture, &OutputResolution](const int32 TargetType, const bool bClearBackground, uint32& PrintOffset) {
+
 		int32 TargetIndex = INDEX_NONE;
 		EFieldOutputType FieldOutputType = EFieldOutputType::Field_Output_Max;
-		GetFieldIndex(GPhysicsFieldTargetType, TargetIndex, FieldOutputType);
+		GetFieldIndex(TargetType, TargetIndex, FieldOutputType);
 
-		if (GPhysicsFieldTargetType < EFieldPhysicsType::Field_PhysicsType_Max)
+		const int32 NumBounds = PhysicsFieldResource->FieldInfos.BoundsOffsets[TargetType + 1] - PhysicsFieldResource->FieldInfos.BoundsOffsets[TargetType];
+		if (NumBounds > 0 && (!bClearBackground || (bClearBackground && (FieldOutputType == EFieldOutputType::Field_Output_Vector))))
 		{
-			// Reset the field background to highlight the vector lines color
-			if (FieldOutputType == EFieldOutputType::Field_Output_Vector)
-			{
-				FPhysicsFieldRayMarchingCS::FParameters* Parameters = CreateShaderParameters(GraphBuilder, View, PhysicsFieldResource, OutputTexture, TargetIndex);
+			const FIntVector BoundResolution(FPhysicsFieldRayMarchingCS::VoxelResolution * NumBounds, FPhysicsFieldRayMarchingCS::VoxelResolution, FPhysicsFieldRayMarchingCS::VoxelResolution);
+			PrintOffset = !bClearBackground ? PrintOffset + 1 : PrintOffset;
 
-				FPhysicsFieldRayMarchingCS::FPermutationDomain PermutationVector;
-				PermutationVector.Set<FPhysicsFieldRayMarchingCS::FFieldType>(3);
-				PermutationVector.Set<FPhysicsFieldRayMarchingCS::FEvalType>(GPhysicsFieldEvalType);
-				TShaderMapRef<FPhysicsFieldRayMarchingCS> ComputeShader(View.ShaderMap, PermutationVector);
+			FPhysicsFieldRayMarchingCS::FParameters* Parameters = CreateShaderParameters(GraphBuilder, View, PhysicsFieldResource, OutputTexture, TargetIndex, TargetType, PrintOffset);
 
-				const FIntVector DispatchCount = FIntVector::DivideAndRoundUp(FIntVector(OutputResolution.X, OutputResolution.Y, 1), FIntVector(8, 8, 1));
-				FComputeShaderUtils::AddPass(GraphBuilder, RDG_EVENT_NAME("PhysicsFieldRayMarching"), ComputeShader, Parameters, DispatchCount);
-			}
+			FPhysicsFieldRayMarchingCS::FPermutationDomain PermutationVector;
+			PermutationVector.Set<FPhysicsFieldRayMarchingCS::FFieldType>(bClearBackground ? 3 : FieldOutputType);
+			PermutationVector.Set<FPhysicsFieldRayMarchingCS::FEvalType>(GPhysicsFieldEvalType);
+			TShaderMapRef<FPhysicsFieldRayMarchingCS> ComputeShader(View.ShaderMap, PermutationVector);
 
-			{
-				FPhysicsFieldRayMarchingCS::FParameters* Parameters = CreateShaderParameters(GraphBuilder, View, PhysicsFieldResource, OutputTexture, TargetIndex);
+			const FIntVector DispatchCount = (!bClearBackground && (FieldOutputType == EFieldOutputType::Field_Output_Vector)) ?
+				FIntVector::DivideAndRoundUp(BoundResolution, FIntVector(4, 4, 4)) :
+				FIntVector::DivideAndRoundUp(FIntVector(OutputResolution.X, OutputResolution.Y, 1), FIntVector(8, 8, 1));
 
-				const int32 NumBounds = PhysicsFieldResource->FieldInfos.BoundsOffsets[GPhysicsFieldTargetType + 1] - PhysicsFieldResource->FieldInfos.BoundsOffsets[GPhysicsFieldTargetType];
-				const FIntVector BoundResolution(FPhysicsFieldRayMarchingCS::VoxelResolution * NumBounds, FPhysicsFieldRayMarchingCS::VoxelResolution, FPhysicsFieldRayMarchingCS::VoxelResolution);
+			FComputeShaderUtils::AddPass(GraphBuilder, RDG_EVENT_NAME("RenderPhysicsField"), ComputeShader, Parameters, DispatchCount);
+		}
+	};
+	FSceneTextureParameters SceneTextures = GetSceneTextureParameters(GraphBuilder, View);
 
-				FPhysicsFieldRayMarchingCS::FPermutationDomain PermutationVector;
-				PermutationVector.Set<FPhysicsFieldRayMarchingCS::FFieldType>(FieldOutputType);
-				PermutationVector.Set<FPhysicsFieldRayMarchingCS::FEvalType>(GPhysicsFieldEvalType);
-				TShaderMapRef<FPhysicsFieldRayMarchingCS> ComputeShader(View.ShaderMap, PermutationVector);
+	uint32 PrintOffset = 0;
+	if(GPhysicsFieldTargetType < EFieldPhysicsType::Field_PhysicsType_Max)
+	{
+		RenderPhysicsField(GPhysicsFieldTargetType, true, PrintOffset);
+		RenderPhysicsField(GPhysicsFieldTargetType, false, PrintOffset);
+	}
+	else
+	{
+		for (int32 TargetType = 0; TargetType < EFieldPhysicsType::Field_PhysicsType_Max; ++TargetType)
+		{
+			RenderPhysicsField(TargetType, true, PrintOffset);
+		}
 
-				const FIntVector DispatchCount = (FieldOutputType == EFieldOutputType::Field_Output_Vector) ?
-					FIntVector::DivideAndRoundUp(BoundResolution, FIntVector(4, 4, 4)) :
-					FIntVector::DivideAndRoundUp(FIntVector(OutputResolution.X, OutputResolution.Y, 1), FIntVector(8, 8, 1));
-				FComputeShaderUtils::AddPass(GraphBuilder, RDG_EVENT_NAME("PhysicsFieldRayMarching"), ComputeShader, Parameters, DispatchCount);
-			}
+		for (int32 TargetType = 0; TargetType < EFieldPhysicsType::Field_PhysicsType_Max; ++TargetType)
+		{
+			RenderPhysicsField(TargetType, false, PrintOffset);
 		}
 	}
 }
@@ -182,10 +195,12 @@ void RenderPhysicsField(
 	if (PhysicsFieldProxy)
 	{
 		FPhysicsFieldResource* PhysicsFieldResource = (GPhysicsFieldSystemType == 0) ? PhysicsFieldProxy->DebugResource : PhysicsFieldProxy->FieldResource;
-
 		if (Views.Num() > 0 && PhysicsFieldResource && ShaderPrint::IsEnabled(Views[0].ShaderPrintData) && UE::PixelFormat::HasCapabilities(PF_FloatRGBA, EPixelFormatCapabilities::TypedUAVLoad))
 		{
-			AddPhysicsFieldRayMarchingPass(GraphBuilder, Views[0], PhysicsFieldResource, SceneColorTexture);
+			if (ShaderPrint::IsSupported(Views[0].GetShaderPlatform()) && IsFeatureLevelSupported(Views[0].GetShaderPlatform(), ERHIFeatureLevel::SM5))
+			{
+				AddPhysicsFieldRayMarchingPass(GraphBuilder, Views[0], PhysicsFieldResource, SceneColorTexture);
+			}
 		}
 	}
 }
