SceneCapture depth prepass optimization for mobile renderer

diff --git a/Engine/Source/Runtime/Renderer/Private/DeferredShadingRenderer.cpp b/Engine/Source/Runtime/Renderer/Private/DeferredShadingRenderer.cpp
--- a/Engine/Source/Runtime/Renderer/Private/DeferredShadingRenderer.cpp
+++ b/Engine/Source/Runtime/Renderer/Private/DeferredShadingRenderer.cpp
@@ -103,15 +103,6 @@ extern int32 GNanitePickingDomain;
 
 extern DynamicRenderScaling::FBudget GDynamicNaniteScalingPrimary;
 
-// This is an experimental optimization switch to render pre-pass depth for scene capture without calling the entire FDeferredShadingSceneRenderer::Render()
-int32 GSceneCaptureDepthPrepassOptimization = 0;
-static FAutoConsoleVariableRef CVarSceneCaptureDepthPrepassOptimization(
-	TEXT("r.SceneCapture.DepthPrepassOptimization"),
-	GSceneCaptureDepthPrepassOptimization,
-	TEXT("Whether to apply optimized render path when capturing depth prepass for scene capture 2D. Experimental!\n")
-	TEXT("Warning: turning it on means rendering after depth pre-pass (e.g. SingleLayerWater) is ignored, hence result is different from when CVar is off.\n"),
-	ECVF_RenderThreadSafe | ECVF_Scalability);
-
 static TAutoConsoleVariable<int32> CVarClearCoatNormal(
 	TEXT("r.ClearCoatNormal"),
 	0,
@@ -2464,12 +2455,6 @@ bool FDeferredShadingSceneRenderer::IsNaniteEnabled() const
 	return UseNanite(ShaderPlatform) && ViewFamily.EngineShowFlags.NaniteMeshes && Nanite::GStreamingManager.HasResourceEntries();
 }
 
-FDeferredShadingSceneRenderer::ERendererOutput FDeferredShadingSceneRenderer::GetRendererOutput() const
-{
-	const bool bSceneCaptureDepthPrepass = Views[0].bIsSceneCapture && (ViewFamily.SceneCaptureSource == ESceneCaptureSource::SCS_SceneDepth || ViewFamily.SceneCaptureSource == ESceneCaptureSource::SCS_DeviceDepth);
-	return bSceneCaptureDepthPrepass && GSceneCaptureDepthPrepassOptimization ? ERendererOutput::DepthPrepassOnly : ERendererOutput::FinalSceneColor;
-}
-
 void FDeferredShadingSceneRenderer::Render(FRDGBuilder& GraphBuilder)
 {
 	if (!ViewFamily.EngineShowFlags.Rendering)
diff --git a/Engine/Source/Runtime/Renderer/Private/DeferredShadingRenderer.h b/Engine/Source/Runtime/Renderer/Private/DeferredShadingRenderer.h
--- a/Engine/Source/Runtime/Renderer/Private/DeferredShadingRenderer.h
+++ b/Engine/Source/Runtime/Renderer/Private/DeferredShadingRenderer.h
@@ -319,12 +319,6 @@ public:
 	void InitializeRayTracingFlags_RenderThread();
 #endif
 
-	enum class ERendererOutput
-	{
-		DepthPrepassOnly,	// Only render depth prepass and its related code paths
-		FinalSceneColor		// Render the whole pipeline
-	};
-
 	FDeferredShadingSceneRenderer(const FSceneViewFamily* InViewFamily, FHitProxyConsumer* HitProxyConsumer);
 
 	/** Determine and commit the final state of the pipeline for the view family and views. */
@@ -1066,8 +1060,6 @@ private:
 		const FSceneTextureParameters& SceneTextures,
 		FRDGTextureRef* OutAmbientOcclusionTexture);
 	
-	ERendererOutput GetRendererOutput() const;
-
 #if RHI_RAYTRACING
 	template <int TextureImportanceSampling>
 	void RenderRayTracingRectLightInternal(
diff --git a/Engine/Source/Runtime/Renderer/Private/MobileShadingRenderer.cpp b/Engine/Source/Runtime/Renderer/Private/MobileShadingRenderer.cpp
--- a/Engine/Source/Runtime/Renderer/Private/MobileShadingRenderer.cpp
+++ b/Engine/Source/Runtime/Renderer/Private/MobileShadingRenderer.cpp
@@ -266,7 +266,7 @@ FMobileSceneRenderer::FMobileSceneRenderer(const FSceneViewFamily* InViewFamily,
 	: FSceneRenderer(InViewFamily, HitProxyConsumer)
 	, bGammaSpace(!IsMobileHDR())
 	, bDeferredShading(IsMobileDeferredShadingEnabled(ShaderPlatform))
-	, bUseVirtualTexturing(UseVirtualTexturing(FeatureLevel))
+	, bUseVirtualTexturing(UseVirtualTexturing(FeatureLevel) && GetRendererOutput() == FSceneRenderer::ERendererOutput::FinalSceneColor)
 {
 	bRenderToSceneColor = false;
 	bRequiresMultiPass = false;
@@ -420,6 +420,8 @@ void FMobileSceneRenderer::InitViews(
 
 	check(Scene);
 
+	const bool bRendererOutputFinalSceneColor = (GetRendererOutput() == ERendererOutput::FinalSceneColor);
+
 	PreVisibilityFrameSetup(GraphBuilder);
 
 	if (InstanceCullingManager.IsEnabled()
@@ -439,7 +441,7 @@ void FMobileSceneRenderer::InitViews(
 
 	const FExclusiveDepthStencil::Type BasePassDepthStencilAccess = FExclusiveDepthStencil::DepthWrite_StencilWrite;
 
-	if (FXSystem && FXSystem->RequiresEarlyViewUniformBuffer() && Views.IsValidIndex(0))
+	if (FXSystem && FXSystem->RequiresEarlyViewUniformBuffer() && Views.IsValidIndex(0) && bRendererOutputFinalSceneColor)
 	{
 		// This is to init the ViewUniformBuffer before rendering for the Niagara compute shader.
 		// This needs to run before ComputeViewVisibility() is called, but the views normally initialize the ViewUniformBuffer after that (at the end of this method).
@@ -449,7 +451,7 @@ void FMobileSceneRenderer::InitViews(
 
 	TaskDatas.VisibilityTaskData->ProcessRenderThreadTasks(BasePassDepthStencilAccess, InstanceCullingManager, VirtualTextureUpdater);
 
-	if(ShouldRenderVolumetricFog())
+	if (ShouldRenderVolumetricFog() && bRendererOutputFinalSceneColor)
 	{
 		SetupVolumetricFog();
 	}
@@ -492,7 +494,8 @@ void FMobileSceneRenderer::InitViews(
 		&& !ViewFamily.EngineShowFlags.HitProxies
 		&& ViewFamily.EngineShowFlags.Lighting
 		&& !ViewFamily.EngineShowFlags.VisualizeLightCulling
-		&& !ViewFamily.UseDebugViewPS();
+		&& !ViewFamily.UseDebugViewPS()
+		&& bRendererOutputFinalSceneColor;
 
 	bRequiresAmbientOcclusionPass = IsUsingMobileAmbientOcclusion(ShaderPlatform)
 		&& Views[0].FinalPostProcessSettings.AmbientOcclusionIntensity > 0
@@ -502,7 +505,8 @@ void FMobileSceneRenderer::InitViews(
 		&& !Views[0].bIsPlanarReflection
 		&& !ViewFamily.EngineShowFlags.HitProxies
 		&& !ViewFamily.EngineShowFlags.VisualizeLightCulling
-		&& !ViewFamily.UseDebugViewPS();
+		&& !ViewFamily.UseDebugViewPS()
+		&& bRendererOutputFinalSceneColor;
 
 	bShouldRenderVelocities = ShouldRenderVelocities();
 
@@ -515,7 +519,8 @@ void FMobileSceneRenderer::InitViews(
 		&& !ViewFamily.EngineShowFlags.HitProxies
 		&& !ViewFamily.EngineShowFlags.VisualizeLightCulling
 		&& !ViewFamily.UseDebugViewPS()
-		&& (CVarDistanceFieldShadowQuality != nullptr && CVarDistanceFieldShadowQuality->GetInt() > 0);
+		&& (CVarDistanceFieldShadowQuality != nullptr && CVarDistanceFieldShadowQuality->GetInt() > 0)
+		&& bRendererOutputFinalSceneColor;
 
 	bRequiresShadowProjections = MobileUsesShadowMaskTexture(ShaderPlatform)
 		&& ViewFamily.EngineShowFlags.Lighting
@@ -523,9 +528,10 @@ void FMobileSceneRenderer::InitViews(
 		&& !Views[0].bIsPlanarReflection
 		&& !ViewFamily.EngineShowFlags.HitProxies
 		&& !ViewFamily.EngineShowFlags.VisualizeLightCulling
-		&& !ViewFamily.UseDebugViewPS();
+		&& !ViewFamily.UseDebugViewPS()
+		&& bRendererOutputFinalSceneColor;
 
-	bShouldRenderHZB = ShouldRenderHZB();
+	bShouldRenderHZB = ShouldRenderHZB() && bRendererOutputFinalSceneColor;
 
 	// Whether we need to store depth for post-processing
 	// On PowerVR we see flickering of shadows and depths not updating correctly if targets are discarded.
@@ -605,19 +611,22 @@ void FMobileSceneRenderer::InitViews(
 	// Finalize and set the scene textures config.
 	FSceneTexturesConfig::Set(SceneTexturesConfig);
 
-	// Initialise Sky/View resources before the view global uniform buffer is built.
-	if (ShouldRenderSkyAtmosphere(Scene, ViewFamily.EngineShowFlags))
+	if (bRendererOutputFinalSceneColor)
 	{
-		InitSkyAtmosphereForViews(RHICmdList);
-	}
+		// Initialise Sky/View resources before the view global uniform buffer is built.
+		if (ShouldRenderSkyAtmosphere(Scene, ViewFamily.EngineShowFlags))
+		{
+			InitSkyAtmosphereForViews(RHICmdList);
+		}
 
-	if (bRequiresShadowProjections)
-	{
-		InitMobileShadowProjectionOutputs(RHICmdList, SceneTexturesConfig.Extent);
-	}
-	else
-	{
-		ReleaseMobileShadowProjectionOutputs();
+		if (bRequiresShadowProjections)
+		{
+			InitMobileShadowProjectionOutputs(RHICmdList, SceneTexturesConfig.Extent);
+		}
+		else
+		{
+			ReleaseMobileShadowProjectionOutputs();
+		}
 	}
 		
 	FRDGExternalAccessQueue ExternalAccessQueue;
@@ -643,26 +652,32 @@ void FMobileSceneRenderer::InitViews(
 		}
 	}
 
-	const bool bDynamicShadows = ViewFamily.EngineShowFlags.DynamicShadows;
-	if (bDynamicShadows)
-	{
-		// Setup dynamic shadows.
-		TaskDatas.DynamicShadows = InitDynamicShadows(GraphBuilder, InstanceCullingManager, ExternalAccessQueue);
-	}
-	else
+	if (bRendererOutputFinalSceneColor)
 	{
-		// TODO: only do this when CSM + static is required.
-		PrepareViewVisibilityLists();
+		const bool bDynamicShadows = ViewFamily.EngineShowFlags.DynamicShadows;
+		if (bDynamicShadows)
+		{
+			// Setup dynamic shadows.
+			TaskDatas.DynamicShadows = InitDynamicShadows(GraphBuilder, InstanceCullingManager, ExternalAccessQueue);
+		}
+		else
+		{
+			// TODO: only do this when CSM + static is required.
+			PrepareViewVisibilityLists();
+		}
 	}
 
 	TaskDatas.VisibilityTaskData->Finish();
 
-	SetupMobileBasePassAfterShadowInit(BasePassDepthStencilAccess, TaskDatas.VisibilityTaskData->GetViewCommandsPerView(), InstanceCullingManager);
-
-	// if we kicked off ILC update via task, wait and finalize.
-	if (ILCTaskData)
+	if (bRendererOutputFinalSceneColor)
 	{
-		Scene->IndirectLightingCache.FinalizeCacheUpdates(Scene, *this, *ILCTaskData);
+		SetupMobileBasePassAfterShadowInit(BasePassDepthStencilAccess, TaskDatas.VisibilityTaskData->GetViewCommandsPerView(), InstanceCullingManager);
+
+		// if we kicked off ILC update via task, wait and finalize.
+		if (ILCTaskData)
+		{
+			Scene->IndirectLightingCache.FinalizeCacheUpdates(Scene, *this, *ILCTaskData);
+		}
 	}
 
 	ExternalAccessQueue.Submit(GraphBuilder);
@@ -682,18 +697,21 @@ void FMobileSceneRenderer::InitViews(
 		}
 	}
 
-	if (bDeferredShading ||
-		bEnableClusteredLocalLights || 
-		bEnableClusteredReflections)
+	if (bRendererOutputFinalSceneColor)
 	{
-		SetupSceneReflectionCaptureBuffer(RHICmdList);
-	}
-	UpdateSkyReflectionUniformBuffer(RHICmdList);
+		if (bDeferredShading ||
+			bEnableClusteredLocalLights || 
+			bEnableClusteredReflections)
+		{
+			SetupSceneReflectionCaptureBuffer(RHICmdList);
+		}
+		UpdateSkyReflectionUniformBuffer(RHICmdList);
 
-	// Now that the indirect lighting cache is updated, we can update the uniform buffers.
-	UpdatePrimitiveIndirectLightingCacheBuffers(RHICmdList);
+		// Now that the indirect lighting cache is updated, we can update the uniform buffers.
+		UpdatePrimitiveIndirectLightingCacheBuffers(RHICmdList);
 
-	UpdateDirectionalLightUniformBuffers(GraphBuilder, Views[0]);
+		UpdateDirectionalLightUniformBuffers(GraphBuilder, Views[0]);
+	}
 	
 	OnStartRender(RHICmdList);
 }
@@ -791,6 +809,9 @@ void FMobileSceneRenderer::Render(FRDGBuilder& GraphBuilder)
 		return;
 	}
 
+	const ERendererOutput RendererOutput = GetRendererOutput();
+	const bool bRendererOutputFinalSceneColor = (RendererOutput == ERendererOutput::FinalSceneColor);
+
 	GraphBuilder.SetCommandListStat(GET_STATID(STAT_CLMM_SceneStart));
 
 	RDG_RHI_EVENT_SCOPE(GraphBuilder, MobileSceneRender);
@@ -800,7 +821,7 @@ void FMobileSceneRenderer::Render(FRDGBuilder& GraphBuilder)
 
 	FRDGExternalAccessQueue ExternalAccessQueue;
 
-	if (bRequiresDistanceField)
+	if (bRequiresDistanceField && bRendererOutputFinalSceneColor)
 	{
 		PrepareDistanceFieldScene(GraphBuilder, ExternalAccessQueue);
 	}
@@ -812,19 +833,22 @@ void FMobileSceneRenderer::Render(FRDGBuilder& GraphBuilder)
 	// Establish scene primitive count (must be done after UpdateAllPrimitiveSceneInfos)
 	FGPUSceneScopeBeginEndHelper GPUSceneScopeBeginEndHelper(Scene->GPUScene, GPUSceneDynamicContext, Scene);
 
-	if (ShouldRenderSkyAtmosphere(Scene, ViewFamily.EngineShowFlags))
+	if (bRendererOutputFinalSceneColor)
 	{
-		for (int32 LightIndex = 0; LightIndex < NUM_ATMOSPHERE_LIGHTS; ++LightIndex)
+		if (ShouldRenderSkyAtmosphere(Scene, ViewFamily.EngineShowFlags))
 		{
-			if (Scene->AtmosphereLights[LightIndex])
+			for (int32 LightIndex = 0; LightIndex < NUM_ATMOSPHERE_LIGHTS; ++LightIndex)
 			{
-				PrepareSunLightProxy(*Scene->GetSkyAtmosphereSceneInfo(), LightIndex, *Scene->AtmosphereLights[LightIndex]);
+				if (Scene->AtmosphereLights[LightIndex])
+				{
+					PrepareSunLightProxy(*Scene->GetSkyAtmosphereSceneInfo(), LightIndex, *Scene->AtmosphereLights[LightIndex]);
+				}
 			}
 		}
-	}
-	else
-	{
-		Scene->ResetAtmosphereLightsProperties();
+		else
+		{
+			Scene->ResetAtmosphereLightsProperties();
+		}
 	}
 
 	CSV_SCOPED_TIMING_STAT_EXCLUSIVE(RenderOther);
@@ -860,28 +884,30 @@ void FMobileSceneRenderer::Render(FRDGBuilder& GraphBuilder)
 	}
 	Strata::InitialiseStrataFrameSceneData(GraphBuilder, *this);
 
-	// Force the subsurface profile texture to be updated.
-	UpdateSubsurfaceProfileTexture(GraphBuilder, ShaderPlatform);
-	SpecularProfileAtlas::UpdateSpecularProfileTextureAtlas(GraphBuilder, ShaderPlatform);
-	
-	if (bDeferredShading)
+	if (bRendererOutputFinalSceneColor)
 	{
-		RectLightAtlas::UpdateAtlasTexture(GraphBuilder, FeatureLevel);
-		IESAtlas::UpdateAtlasTexture(GraphBuilder, FeatureLevel);
-	}
+		// Force the subsurface profile texture to be updated.
+		UpdateSubsurfaceProfileTexture(GraphBuilder, ShaderPlatform);
+		SpecularProfileAtlas::UpdateSpecularProfileTextureAtlas(GraphBuilder, ShaderPlatform);
 
-	FInstanceCullingManager& InstanceCullingManager = *GraphBuilder.AllocObject<FInstanceCullingManager>(GetSceneUniforms(), Scene->GPUScene.IsEnabled(), GraphBuilder);
+		if (bDeferredShading)
+		{
+			RectLightAtlas::UpdateAtlasTexture(GraphBuilder, FeatureLevel);
+			IESAtlas::UpdateAtlasTexture(GraphBuilder, FeatureLevel);
+		}
 
-	// Important that this uses consistent logic throughout the frame, so evaluate once and pass in the flag from here
-	// NOTE: Must be done after  system texture initialization
-	// TODO: This doesn't take into account the potential for split screen views with separate shadow caches
-	VirtualShadowMapArray.Initialize(GraphBuilder, Scene->GetVirtualShadowMapCache(), UseVirtualShadowMaps(ShaderPlatform, FeatureLevel));
+		// Important that this uses consistent logic throughout the frame, so evaluate once and pass in the flag from here
+		// NOTE: Must be done after  system texture initialization
+		// TODO: This doesn't take into account the potential for split screen views with separate shadow caches
+		VirtualShadowMapArray.Initialize(GraphBuilder, Scene->GetVirtualShadowMapCache(), UseVirtualShadowMaps(ShaderPlatform, FeatureLevel));
+	}
 
 	FInitViewTaskDatas InitViewTaskDatas(VisibilityTaskData);
 
 	GraphBuilder.SetCommandListStat(GET_STATID(STAT_CLMM_InitViews));
 
 	// Find the visible primitives and prepare targets and buffers for rendering
+	FInstanceCullingManager& InstanceCullingManager = *GraphBuilder.AllocObject<FInstanceCullingManager>(GetSceneUniforms(), Scene->GPUScene.IsEnabled(), GraphBuilder);
 	InitViews(GraphBuilder, SceneTexturesConfig, InstanceCullingManager, VirtualTextureUpdater.Get(), InitViewTaskDatas);
 
 	GraphBuilder.SetCommandListStat(GET_STATID(STAT_CLMM_AfterInitViews));
@@ -903,78 +929,81 @@ void FMobileSceneRenderer::Render(FRDGBuilder& GraphBuilder)
 	FSceneTextures::InitializeViewFamily(GraphBuilder, ViewFamily);
 	FSceneTextures& SceneTextures = GetActiveSceneTextures();
 
-#if WITH_DEBUG_VIEW_MODES
-	if (ViewFamily.UseDebugViewPS() && ViewFamily.EngineShowFlags.ShaderComplexity && SceneTextures.QuadOverdraw)
+	FSortedLightSetSceneInfo& SortedLightSet = *GraphBuilder.AllocObject<FSortedLightSetSceneInfo>();
+
+	SceneTextures.MobileSetupMode = EMobileSceneTextureSetupMode::None;
+	SceneTextures.MobileUniformBuffer = CreateMobileSceneTextureUniformBuffer(GraphBuilder, &SceneTextures, SceneTextures.MobileSetupMode);
+
+	if (bRendererOutputFinalSceneColor)
 	{
-		AddClearUAVPass(GraphBuilder, GraphBuilder.CreateUAV(SceneTextures.QuadOverdraw), FUintVector4(0, 0, 0, 0));
-	}
+#if WITH_DEBUG_VIEW_MODES
+		if (ViewFamily.UseDebugViewPS() && ViewFamily.EngineShowFlags.ShaderComplexity && SceneTextures.QuadOverdraw)
+		{
+			AddClearUAVPass(GraphBuilder, GraphBuilder.CreateUAV(SceneTextures.QuadOverdraw), FUintVector4(0, 0, 0, 0));
+		}
 #endif
 
-	if (bUseVirtualTexturing)
-	{
-		FVirtualTextureSystem::Get().EndUpdate(GraphBuilder, MoveTemp(VirtualTextureUpdater), FeatureLevel);
-	}
+		if (bUseVirtualTexturing)
+		{
+			FVirtualTextureSystem::Get().EndUpdate(GraphBuilder, MoveTemp(VirtualTextureUpdater), FeatureLevel);
+		}
 
-	FSortedLightSetSceneInfo& SortedLightSet = *GraphBuilder.AllocObject<FSortedLightSetSceneInfo>();
+		if (bDeferredShading ||
+			bEnableClusteredLocalLights ||
+			bEnableClusteredReflections)
+		{
+			RDG_CSV_STAT_EXCLUSIVE_SCOPE(GraphBuilder, SortLights);
+			// Shadows are applied in clustered shading on mobile forward and separately on mobile deferred.
+			bool bShadowedLightsInClustered = bRequiresShadowProjections && !bDeferredShading;
+			GatherAndSortLights(SortedLightSet, bShadowedLightsInClustered);
+			int32 NumReflectionCaptures = Views[0].NumBoxReflectionCaptures + Views[0].NumSphereReflectionCaptures;
+			bool bCullLightsToGrid = (((bEnableClusteredReflections || bDeferredShading) && NumReflectionCaptures > 0) || bEnableClusteredLocalLights || bEnablePrepassLocalLights);
+			if (bCullLightsToGrid)
+			{
+				ComputeLightGrid(GraphBuilder, bEnableClusteredLocalLights || bEnablePrepassLocalLights, SortedLightSet);
+			}
+		}
 
-	if (bDeferredShading ||
-		bEnableClusteredLocalLights || 
-		bEnableClusteredReflections)
-	{
-		RDG_CSV_STAT_EXCLUSIVE_SCOPE(GraphBuilder, SortLights);
-		// Shadows are applied in clustered shading on mobile forward and separately on mobile deferred.
-		bool bShadowedLightsInClustered = bRequiresShadowProjections && !bDeferredShading;
-		GatherAndSortLights(SortedLightSet, bShadowedLightsInClustered);
-		int32 NumReflectionCaptures = Views[0].NumBoxReflectionCaptures + Views[0].NumSphereReflectionCaptures;
-		bool bCullLightsToGrid = (((bEnableClusteredReflections || bDeferredShading) && NumReflectionCaptures > 0) || bEnableClusteredLocalLights || bEnablePrepassLocalLights);
-		if (bCullLightsToGrid)
+		// Notify the FX system that the scene is about to be rendered.
+		if (FXSystem)
 		{
-			ComputeLightGrid(GraphBuilder, bEnableClusteredLocalLights || bEnablePrepassLocalLights, SortedLightSet);
+			FXSystem->PreRender(GraphBuilder, GetSceneViews(), GetSceneUniforms(), true /*bAllowGPUParticleUpdate*/);
+			if (FGPUSortManager* GPUSortManager = FXSystem->GetGPUSortManager())
+			{
+				// if GPUSortManager::OnPostRenderOpaque is called below (from RenderOpaqueFX) we must also call OnPreRender (as it sets up
+				// the internal state of the GPUSortManager).  Any optimization to skip this block needs to take that into consideration.
+				GPUSortManager->OnPreRender(GraphBuilder);
+			}
 		}
-	}
 
-	// Notify the FX system that the scene is about to be rendered.
-	if (FXSystem)
-	{
-		FXSystem->PreRender(GraphBuilder, GetSceneViews(), GetSceneUniforms(), true /*bAllowGPUParticleUpdate*/);
-		if (FGPUSortManager* GPUSortManager = FXSystem->GetGPUSortManager())
+		// Generate the Sky/Atmosphere look up tables
+		const bool bShouldRenderSkyAtmosphere = ShouldRenderSkyAtmosphere(Scene, ViewFamily.EngineShowFlags);
+		if (bShouldRenderSkyAtmosphere)
 		{
-			// if GPUSortManager::OnPostRenderOpaque is called below (from RenderOpaqueFX) we must also call OnPreRender (as it sets up
-			// the internal state of the GPUSortManager).  Any optimization to skip this block needs to take that into consideration.
-			GPUSortManager->OnPreRender(GraphBuilder);
+			RenderSkyAtmosphereLookUpTables(GraphBuilder, ExternalAccessQueue);
 		}
-	}
 
-	// Generate the Sky/Atmosphere look up tables
-	const bool bShouldRenderSkyAtmosphere = ShouldRenderSkyAtmosphere(Scene, ViewFamily.EngineShowFlags);
-	if (bShouldRenderSkyAtmosphere)
-	{
-		RenderSkyAtmosphereLookUpTables(GraphBuilder, ExternalAccessQueue);
-	}
+		GraphBuilder.SetCommandListStat(GET_STATID(STAT_CLMM_Shadows));
+		RenderShadowDepthMaps(GraphBuilder, InstanceCullingManager, ExternalAccessQueue);
 
-	GraphBuilder.SetCommandListStat(GET_STATID(STAT_CLMM_Shadows));
-	RenderShadowDepthMaps(GraphBuilder, InstanceCullingManager, ExternalAccessQueue);
+		if (ShouldRenderVolumetricFog())
+		{
+			ComputeVolumetricFog(GraphBuilder, SceneTextures);
+		}
+		ExternalAccessQueue.Submit(GraphBuilder);
 
-	if(ShouldRenderVolumetricFog())
-	{
-		ComputeVolumetricFog(GraphBuilder, SceneTextures);
-	}
-	ExternalAccessQueue.Submit(GraphBuilder);
-	
-	PollOcclusionQueriesPass(GraphBuilder);
-	GraphBuilder.AddDispatchHint();
+		PollOcclusionQueriesPass(GraphBuilder);
+		GraphBuilder.AddDispatchHint();
 
-	// Custom depth
-	// bShouldRenderCustomDepth has been initialized in InitViews on mobile platform
-	if (bShouldRenderCustomDepth)
-	{
-		SceneTextures.MobileSetupMode = EMobileSceneTextureSetupMode::None;
-		SceneTextures.MobileUniformBuffer = CreateMobileSceneTextureUniformBuffer(GraphBuilder, &SceneTextures, SceneTextures.MobileSetupMode);
+		// Custom depth
+		// bShouldRenderCustomDepth has been initialized in InitViews on mobile platform
+		if (bShouldRenderCustomDepth)
+		{
+			RenderCustomDepthPass(GraphBuilder, SceneTextures.CustomDepth, SceneTextures.GetSceneTextureShaderParameters(FeatureLevel), {}, {});
+		}
 
-		RenderCustomDepthPass(GraphBuilder, SceneTextures.CustomDepth, SceneTextures.GetSceneTextureShaderParameters(FeatureLevel), {}, {});
+		BeginOcclusionScope(GraphBuilder, Views);
 	}
-
-	BeginOcclusionScope(GraphBuilder, Views);
 	
 	// Sort objects' triangles
 	for (FViewInfo& View : Views)
@@ -1001,6 +1030,7 @@ void FMobileSceneRenderer::Render(FRDGBuilder& GraphBuilder)
 		SceneTextures.MobileSetupMode = EMobileSceneTextureSetupMode::SceneDepth;
 		SceneTextures.MobileUniformBuffer = CreateMobileSceneTextureUniformBuffer(GraphBuilder, &SceneTextures, SceneTextures.MobileSetupMode);
 
+		// When renderer is in ERendererOutput::DepthPrepassOnly mode, bRequiresShadowProjections is set to false in InitViews()
 		if (bRequiresShadowProjections)
 		{
 			RDG_CSV_STAT_EXCLUSIVE_SCOPE(GraphBuilder, RenderMobileShadowProjections);
@@ -1008,109 +1038,114 @@ void FMobileSceneRenderer::Render(FRDGBuilder& GraphBuilder)
 			RenderMobileShadowProjections(GraphBuilder);
 		}
 
+		// When renderer is in ERendererOutput::DepthPrepassOnly mode, bShouldRenderHZB is set to false in InitViews()
 		if (bShouldRenderHZB)
 		{
 			RenderHZB(GraphBuilder, SceneTextures.Depth.Resolve);
 		}
 
+		// When renderer is in ERendererOutput::DepthPrepassOnly mode, bRequiresAmbientOcclusionPass is set to false in InitViews()
 		if (bRequiresAmbientOcclusionPass)
 		{
 			RenderAmbientOcclusion(GraphBuilder, SceneTextures.Depth.Resolve, SceneTextures.ScreenSpaceAO);
 		}
 
 		// Local Light prepass
-		if (bEnablePrepassLocalLights)
+		if (bEnablePrepassLocalLights && bRendererOutputFinalSceneColor)
 		{
 			RenderLocalLightPrepass(GraphBuilder, SceneTextures);
 		}
 	}
 
-	if (bDeferredShading)
+	if (bRendererOutputFinalSceneColor)
 	{
-		RenderDeferred(GraphBuilder, SortedLightSet, ViewFamilyTexture, SceneTextures);
-	}
-	else
-	{
-		RenderForward(GraphBuilder, ViewFamilyTexture, SceneTextures);
-	}
-
-	EndOcclusionScope(GraphBuilder, Views);
-
-	if (!bIsFullDepthPrepassEnabled)
-	{
-		FenceOcclusionTests(GraphBuilder);
-	}
-
-	SceneTextures.MobileSetupMode = EMobileSceneTextureSetupMode::All;
-	SceneTextures.MobileSetupMode &= ~EMobileSceneTextureSetupMode::SceneVelocity;
-	SceneTextures.MobileUniformBuffer = CreateMobileSceneTextureUniformBuffer(GraphBuilder, &SceneTextures, SceneTextures.MobileSetupMode);
+		if (bDeferredShading)
+		{
+			RenderDeferred(GraphBuilder, SortedLightSet, ViewFamilyTexture, SceneTextures);
+		}
+		else
+		{
+			RenderForward(GraphBuilder, ViewFamilyTexture, SceneTextures);
+		}
 
-	if (bShouldRenderVelocities)
-	{
-		// Render the velocities of movable objects
-		GraphBuilder.SetCommandListStat(GET_STATID(STAT_CLMM_Velocity));
-		RenderVelocities(GraphBuilder, SceneTextures, EVelocityPass::Opaque, false);
+		EndOcclusionScope(GraphBuilder, Views);
 
-		GraphBuilder.SetCommandListStat(GET_STATID(STAT_CLMM_TranslucentVelocity));
-		RenderVelocities(GraphBuilder, SceneTextures, EVelocityPass::Translucent, false);
+		if (!bIsFullDepthPrepassEnabled)
+		{
+			FenceOcclusionTests(GraphBuilder);
+		}
 
 		SceneTextures.MobileSetupMode = EMobileSceneTextureSetupMode::All;
+		SceneTextures.MobileSetupMode &= ~EMobileSceneTextureSetupMode::SceneVelocity;
 		SceneTextures.MobileUniformBuffer = CreateMobileSceneTextureUniformBuffer(GraphBuilder, &SceneTextures, SceneTextures.MobileSetupMode);
-	}
 
-	GraphBuilder.SetCommandListStat(GET_STATID(STAT_CLMM_Post));
+		if (bShouldRenderVelocities)
+		{
+			// Render the velocities of movable objects
+			GraphBuilder.SetCommandListStat(GET_STATID(STAT_CLMM_Velocity));
+			RenderVelocities(GraphBuilder, SceneTextures, EVelocityPass::Opaque, false);
 
-	FRendererModule& RendererModule = static_cast<FRendererModule&>(GetRendererModule());
-	RendererModule.RenderPostOpaqueExtensions(GraphBuilder, Views, SceneTextures);
+			GraphBuilder.SetCommandListStat(GET_STATID(STAT_CLMM_TranslucentVelocity));
+			RenderVelocities(GraphBuilder, SceneTextures, EVelocityPass::Translucent, false);
 
-	RenderOpaqueFX(GraphBuilder, GetSceneViews(), GetSceneUniforms(), FXSystem, SceneTextures.MobileUniformBuffer);
+			SceneTextures.MobileSetupMode = EMobileSceneTextureSetupMode::All;
+			SceneTextures.MobileUniformBuffer = CreateMobileSceneTextureUniformBuffer(GraphBuilder, &SceneTextures, SceneTextures.MobileSetupMode);
+		}
+	
+		GraphBuilder.SetCommandListStat(GET_STATID(STAT_CLMM_Post));
 
-	if (bRequiresPixelProjectedPlanarRelfectionPass)
-	{
-		const FPlanarReflectionSceneProxy* PlanarReflectionSceneProxy = Scene ? Scene->GetForwardPassGlobalPlanarReflection() : nullptr;
+		FRendererModule& RendererModule = static_cast<FRendererModule&>(GetRendererModule());
+		RendererModule.RenderPostOpaqueExtensions(GraphBuilder, Views, SceneTextures);
 
-		RenderPixelProjectedReflection(GraphBuilder, SceneTextures.Color.Resolve, SceneTextures.Depth.Resolve, SceneTextures.PixelProjectedReflection, PlanarReflectionSceneProxy);
-	}
+		RenderOpaqueFX(GraphBuilder, GetSceneViews(), GetSceneUniforms(), FXSystem, SceneTextures.MobileUniformBuffer);
 
-	if (ViewFamily.EngineShowFlags.VisualizeMeshDistanceFields || ViewFamily.EngineShowFlags.VisualizeGlobalDistanceField)
-	{
-		RenderMeshDistanceFieldVisualization(GraphBuilder, SceneTextures);
-	}
+		if (bRequiresPixelProjectedPlanarRelfectionPass)
+		{
+			const FPlanarReflectionSceneProxy* PlanarReflectionSceneProxy = Scene ? Scene->GetForwardPassGlobalPlanarReflection() : nullptr;
+
+			RenderPixelProjectedReflection(GraphBuilder, SceneTextures.Color.Resolve, SceneTextures.Depth.Resolve, SceneTextures.PixelProjectedReflection, PlanarReflectionSceneProxy);
+		}
+
+		if (ViewFamily.EngineShowFlags.VisualizeMeshDistanceFields || ViewFamily.EngineShowFlags.VisualizeGlobalDistanceField)
+		{
+			RenderMeshDistanceFieldVisualization(GraphBuilder, SceneTextures);
+		}
 	
-	if (ViewFamily.EngineShowFlags.VisualizeInstanceOcclusionQueries
-		&& Scene->InstanceCullingOcclusionQueryRenderer)
-	{
-		for (FViewInfo& View : Views)
+		if (ViewFamily.EngineShowFlags.VisualizeInstanceOcclusionQueries
+			&& Scene->InstanceCullingOcclusionQueryRenderer)
 		{
-			Scene->InstanceCullingOcclusionQueryRenderer->RenderDebug(GraphBuilder, Scene->GPUScene, View, SceneTextures);
+			for (FViewInfo& View : Views)
+			{
+				Scene->InstanceCullingOcclusionQueryRenderer->RenderDebug(GraphBuilder, Scene->GPUScene, View, SceneTextures);
+			}
 		}
-	}
 
-	if (bUseVirtualTexturing)
-	{
-		RDG_GPU_STAT_SCOPE(GraphBuilder, VirtualTextureUpdate);
-		VirtualTextureFeedbackEnd(GraphBuilder);
-	}
+		if (bUseVirtualTexturing)
+		{
+			RDG_GPU_STAT_SCOPE(GraphBuilder, VirtualTextureUpdate);
+			VirtualTextureFeedbackEnd(GraphBuilder);
+		}
 	
-	if (ViewFamily.bResolveScene)
-	{
-		if (bRenderToSceneColor)
+		if (ViewFamily.bResolveScene)
 		{
-			// Finish rendering for each view, or the full stereo buffer if enabled
+			if (bRenderToSceneColor)
 			{
-				RDG_EVENT_SCOPE(GraphBuilder, "PostProcessing");
-				RDG_GPU_STAT_SCOPE(GraphBuilder, Postprocessing);
-				SCOPE_CYCLE_COUNTER(STAT_FinishRenderViewTargetTime);
+				// Finish rendering for each view, or the full stereo buffer if enabled
+				{
+					RDG_EVENT_SCOPE(GraphBuilder, "PostProcessing");
+					RDG_GPU_STAT_SCOPE(GraphBuilder, Postprocessing);
+					SCOPE_CYCLE_COUNTER(STAT_FinishRenderViewTargetTime);
 
-				FMobilePostProcessingInputs PostProcessingInputs;
-				PostProcessingInputs.ViewFamilyTexture = ViewFamilyTexture;
-				PostProcessingInputs.SceneTextures = CreateMobileSceneTextureUniformBuffer(GraphBuilder, &SceneTextures, EMobileSceneTextureSetupMode::All);
-				PostProcessingInputs.bRequiresMultiPass = bRequiresMultiPass;
+					FMobilePostProcessingInputs PostProcessingInputs;
+					PostProcessingInputs.ViewFamilyTexture = ViewFamilyTexture;
+					PostProcessingInputs.SceneTextures = CreateMobileSceneTextureUniformBuffer(GraphBuilder, &SceneTextures, EMobileSceneTextureSetupMode::All);
+					PostProcessingInputs.bRequiresMultiPass = bRequiresMultiPass;
 
-				for (int32 ViewIndex = 0; ViewIndex < Views.Num(); ViewIndex++)
-				{
-					RDG_EVENT_SCOPE_CONDITIONAL(GraphBuilder, Views.Num() > 1, "View%d", ViewIndex);
-					AddMobilePostProcessingPasses(GraphBuilder, Scene, Views[ViewIndex], GetSceneUniforms(), PostProcessingInputs, InstanceCullingManager);
+					for (int32 ViewIndex = 0; ViewIndex < Views.Num(); ViewIndex++)
+					{
+						RDG_EVENT_SCOPE_CONDITIONAL(GraphBuilder, Views.Num() > 1, "View%d", ViewIndex);
+						AddMobilePostProcessingPasses(GraphBuilder, Scene, Views[ViewIndex], GetSceneUniforms(), PostProcessingInputs, InstanceCullingManager);
+					}
 				}
 			}
 		}
@@ -1122,7 +1157,10 @@ void FMobileSceneRenderer::Render(FRDGBuilder& GraphBuilder)
 
 	RenderFinish(GraphBuilder, ViewFamilyTexture);
 
-	PollOcclusionQueriesPass(GraphBuilder);
+	if (bRendererOutputFinalSceneColor)
+	{
+		PollOcclusionQueriesPass(GraphBuilder);
+	}
 
 	QueueSceneTextureExtractions(GraphBuilder, SceneTextures);
 
diff --git a/Engine/Source/Runtime/Renderer/Private/SceneRendering.cpp b/Engine/Source/Runtime/Renderer/Private/SceneRendering.cpp
--- a/Engine/Source/Runtime/Renderer/Private/SceneRendering.cpp
+++ b/Engine/Source/Runtime/Renderer/Private/SceneRendering.cpp
@@ -98,6 +98,15 @@
 
 static TGlobalResource<FVirtualTextureFeedbackBuffer> GVirtualTextureFeedbackBuffer;
 
+// This is an experimental optimization switch to render pre-pass depth for scene capture without calling the entire FDeferredShadingSceneRenderer::Render()
+int32 GSceneCaptureDepthPrepassOptimization = 0;
+static FAutoConsoleVariableRef CVarSceneCaptureDepthPrepassOptimization(
+	TEXT("r.SceneCapture.DepthPrepassOptimization"),
+	GSceneCaptureDepthPrepassOptimization,
+	TEXT("Whether to apply optimized render path when capturing depth prepass for scene capture 2D. Experimental!\n")
+	TEXT("Warning: turning it on means rendering after depth pre-pass (e.g. SingleLayerWater) is ignored, hence result is different from when CVar is off.\n"),
+	ECVF_RenderThreadSafe | ECVF_Scalability);
+
 static int32 GAsyncCreateLightPrimitiveInteractions = 1;
 static FAutoConsoleVariableRef CVarAsyncCreateLightPrimitiveInteractions(
 	TEXT("r.AsyncCreateLightPrimitiveInteractions"),
@@ -2834,6 +2843,12 @@ inline EUpdateAllPrimitiveSceneInfosAsyncOps GetUpdateAllPrimitiveSceneInfosAsyn
 	return AsyncOps;
 }
 
+FSceneRenderer::ERendererOutput FSceneRenderer::GetRendererOutput() const
+{
+	const bool bSceneCaptureDepthPrepass = Views[0].bIsSceneCapture && (ViewFamily.SceneCaptureSource == ESceneCaptureSource::SCS_SceneDepth || ViewFamily.SceneCaptureSource == ESceneCaptureSource::SCS_DeviceDepth);
+	return bSceneCaptureDepthPrepass && GSceneCaptureDepthPrepassOptimization ? ERendererOutput::DepthPrepassOnly : ERendererOutput::FinalSceneColor;
+}
+
 IVisibilityTaskData* FSceneRenderer::UpdateScene(FRDGBuilder& GraphBuilder, FGlobalDynamicBuffers GlobalDynamicBuffers)
 {
 	// Needs to run before UpdateAllPrimitiveSceneInfos, as it may call OnVirtualTextureDestroyedCB, which modifies the uniform
diff --git a/Engine/Source/Runtime/Renderer/Private/SceneRendering.h b/Engine/Source/Runtime/Renderer/Private/SceneRendering.h
--- a/Engine/Source/Runtime/Renderer/Private/SceneRendering.h
+++ b/Engine/Source/Runtime/Renderer/Private/SceneRendering.h
@@ -2290,6 +2290,14 @@ protected:
 
 	// Shared functionality between all scene renderers
 
+	enum class ERendererOutput
+	{
+		DepthPrepassOnly,	// Only render depth prepass and its related code paths
+		FinalSceneColor		// Render the whole pipeline
+	};
+
+	ERendererOutput GetRendererOutput() const;
+
 	IVisibilityTaskData* UpdateScene(FRDGBuilder& GraphBuilder, FGlobalDynamicBuffers GlobalDynamicBuffers);
 
 	FDynamicShadowsTaskData* BeginInitDynamicShadows(bool bRunningEarly, IVisibilityTaskData* VisibilityTaskData);
diff --git a/Engine/Source/Runtime/Renderer/Private/SceneVisibility.cpp b/Engine/Source/Runtime/Renderer/Private/SceneVisibility.cpp
--- a/Engine/Source/Runtime/Renderer/Private/SceneVisibility.cpp
+++ b/Engine/Source/Runtime/Renderer/Private/SceneVisibility.cpp
@@ -4018,79 +4018,82 @@ void FSceneRenderer::PreVisibilityFrameSetup(FRDGBuilder& GraphBuilder)
 	// Notify the RHI we are beginning to render a scene.
 	RHICmdList.BeginScene();
 
-	if (Views.Num() > 0 && !ViewFamily.EngineShowFlags.HitProxies)
+	if (GetRendererOutput() == ERendererOutput::FinalSceneColor)
 	{
-		FHairStrandsBookmarkParameters Parameters = CreateHairStrandsBookmarkParameters(Scene, Views, AllFamilyViews);
-		if (Parameters.HasInstances())
+		if (Views.Num() > 0 && !ViewFamily.EngineShowFlags.HitProxies)
 		{
-			RunHairStrandsBookmark(GraphBuilder, EHairStrandsBookmark::ProcessLODSelection, Parameters);
+			FHairStrandsBookmarkParameters Parameters = CreateHairStrandsBookmarkParameters(Scene, Views, AllFamilyViews);
+			if (Parameters.HasInstances())
+			{
+				RunHairStrandsBookmark(GraphBuilder, EHairStrandsBookmark::ProcessLODSelection, Parameters);
+			}
 		}
-	}
 
-	if (IsHairStrandsEnabled(EHairStrandsShaderType::All, Scene->GetShaderPlatform()) && Views.Num() > 0 && !ViewFamily.EngineShowFlags.HitProxies)
-	{
-		// If we are rendering from scene capture we don't need to run another time the hair bookmarks.
-		if (Views[0].AllowGPUParticleUpdate())
+		if (IsHairStrandsEnabled(EHairStrandsShaderType::All, Scene->GetShaderPlatform()) && Views.Num() > 0 && !ViewFamily.EngineShowFlags.HitProxies)
 		{
-			FHairStrandsBookmarkParameters Parameters = CreateHairStrandsBookmarkParameters(Scene, Views, AllFamilyViews);
-			RunHairStrandsBookmark(GraphBuilder, EHairStrandsBookmark::ProcessGuideInterpolation, Parameters);
+			// If we are rendering from scene capture we don't need to run another time the hair bookmarks.
+			if (Views[0].AllowGPUParticleUpdate())
+			{
+				FHairStrandsBookmarkParameters Parameters = CreateHairStrandsBookmarkParameters(Scene, Views, AllFamilyViews);
+				RunHairStrandsBookmark(GraphBuilder, EHairStrandsBookmark::ProcessGuideInterpolation, Parameters);
+			}
 		}
-	}
 
-	// Notify the FX system that the scene is about to perform visibility checks.
+		// Notify the FX system that the scene is about to perform visibility checks.
 
-	if (FXSystem && Views.IsValidIndex(0))
-	{
-		FXSystem->PreInitViews(GraphBuilder, Views[0].AllowGPUParticleUpdate() && !ViewFamily.EngineShowFlags.HitProxies, AllFamilies, &ViewFamily);
-	}
+		if (FXSystem && Views.IsValidIndex(0))
+		{
+			FXSystem->PreInitViews(GraphBuilder, Views[0].AllowGPUParticleUpdate() && !ViewFamily.EngineShowFlags.HitProxies, AllFamilies, &ViewFamily);
+		}
 
 #if WITH_EDITOR
-	// Draw lines to lights affecting this mesh if its selected.
-	if (ViewFamily.EngineShowFlags.LightInfluences)
-	{
-		Scene->WaitForCreateLightPrimitiveInteractionsTask();
-
-		for (TConstSetBitIterator<> It(Scene->PrimitivesSelected); It; ++It)
+		// Draw lines to lights affecting this mesh if its selected.
+		if (ViewFamily.EngineShowFlags.LightInfluences)
 		{
-			const FPrimitiveSceneInfo* PrimitiveSceneInfo = Scene->Primitives[It.GetIndex()];
-			FLightPrimitiveInteraction *LightList = PrimitiveSceneInfo->LightList;
-			while (LightList)
+			Scene->WaitForCreateLightPrimitiveInteractionsTask();
+
+			for (TConstSetBitIterator<> It(Scene->PrimitivesSelected); It; ++It)
 			{
-				const FLightSceneInfo* LightSceneInfo = LightList->GetLight();
+				const FPrimitiveSceneInfo* PrimitiveSceneInfo = Scene->Primitives[It.GetIndex()];
+				FLightPrimitiveInteraction* LightList = PrimitiveSceneInfo->LightList;
+				while (LightList)
+				{
+					const FLightSceneInfo* LightSceneInfo = LightList->GetLight();
 
-				bool bDynamic = true;
-				bool bRelevant = false;
-				bool bLightMapped = true;
-				bool bShadowMapped = false;
-				PrimitiveSceneInfo->Proxy->GetLightRelevance(LightSceneInfo->Proxy, bDynamic, bRelevant, bLightMapped, bShadowMapped);
+					bool bDynamic = true;
+					bool bRelevant = false;
+					bool bLightMapped = true;
+					bool bShadowMapped = false;
+					PrimitiveSceneInfo->Proxy->GetLightRelevance(LightSceneInfo->Proxy, bDynamic, bRelevant, bLightMapped, bShadowMapped);
 
-				if (bRelevant)
-				{
-					// Draw blue for light-mapped lights and orange for dynamic lights
-					const FColor LineColor = bLightMapped ? FColor(0,140,255) : FColor(255,140,0);
-					for (int32 ViewIndex = 0;ViewIndex < Views.Num();ViewIndex++)
+					if (bRelevant)
 					{
-						FViewInfo& View = Views[ViewIndex];
-						FViewElementPDI LightInfluencesPDI(&View,nullptr,&View.DynamicPrimitiveCollector);
-						LightInfluencesPDI.DrawLine(PrimitiveSceneInfo->Proxy->GetBounds().Origin, LightSceneInfo->Proxy->GetLightToWorld().GetOrigin(), LineColor, SDPG_World);
+						// Draw blue for light-mapped lights and orange for dynamic lights
+						const FColor LineColor = bLightMapped ? FColor(0, 140, 255) : FColor(255, 140, 0);
+						for (int32 ViewIndex = 0; ViewIndex < Views.Num(); ViewIndex++)
+						{
+							FViewInfo& View = Views[ViewIndex];
+							FViewElementPDI LightInfluencesPDI(&View, nullptr, &View.DynamicPrimitiveCollector);
+							LightInfluencesPDI.DrawLine(PrimitiveSceneInfo->Proxy->GetBounds().Origin, LightSceneInfo->Proxy->GetLightToWorld().GetOrigin(), LineColor, SDPG_World);
+						}
 					}
+					LightList = LightList->GetNextLight();
 				}
-				LightList = LightList->GetNextLight();
 			}
 		}
-	}
 #endif
 
 #if RHI_RAYTRACING
-	if (Scene && Views.Num())
-	{
-		const int32 ReferenceViewIndex = 0;
-		const FViewInfo& ReferenceView = Views[ReferenceViewIndex];
+		if (Scene && Views.Num())
+		{
+			const int32 ReferenceViewIndex = 0;
+			const FViewInfo& ReferenceView = Views[ReferenceViewIndex];
 
-		Scene->RayTracingScene.InitPreViewTranslation(ReferenceView.ViewMatrices);
-		Scene->RayTracingScene.bNeedsDebugInstanceGPUSceneIndexBuffer = IsRayTracingInstanceOverlapEnabled(ReferenceView);
-	}
+			Scene->RayTracingScene.InitPreViewTranslation(ReferenceView.ViewMatrices);
+			Scene->RayTracingScene.bNeedsDebugInstanceGPUSceneIndexBuffer = IsRayTracingInstanceOverlapEnabled(ReferenceView);
+		}
 #endif
+	}
 
 	for (const auto& ViewExtension : ViewFamily.ViewExtensions)
 	{
@@ -4781,6 +4784,11 @@ void FSceneRenderer::GatherReflectionCaptureLightMeshElements()
 
 void FSceneRenderer::PostVisibilityFrameSetup(FILCUpdatePrimTaskData*& OutILCTaskData)
 {
+	if (GetRendererOutput() == ERendererOutput::FinalSceneColor)
+	{
+		return;
+	}
+
 	QUICK_SCOPE_CYCLE_COUNTER(STAT_PostVisibilityFrameSetup);
 
 	{
@@ -4841,10 +4849,12 @@ void FDeferredShadingSceneRenderer::BeginInitViews(
 	SCOPE_CYCLE_COUNTER(STAT_InitViewsTime);
 	RDG_CSV_STAT_EXCLUSIVE_SCOPE(GraphBuilder, InitViews_Scene);
 
+	const bool bRendererOutputFinalSceneColor = (GetRendererOutput() == ERendererOutput::FinalSceneColor);
+
 	PreVisibilityFrameSetup(GraphBuilder);
 
 	// Attempt to launch dynamic shadow tasks early before finalizing visibility.
-	if (GetRendererOutput() == ERendererOutput::FinalSceneColor)
+	if (bRendererOutputFinalSceneColor)
 	{
 		BeginInitDynamicShadows(TaskDatas);
 	}
@@ -4867,7 +4877,7 @@ void FDeferredShadingSceneRenderer::BeginInitViews(
 	{
 		// This is to init the ViewUniformBuffer before rendering for the Niagara compute shader.
 		// This needs to run before ComputeViewVisibility() is called, but the views normally initialize the ViewUniformBuffer after that (at the end of this method).
-		if (FXSystem && FXSystem->RequiresEarlyViewUniformBuffer() && Views.IsValidIndex(0))
+		if (FXSystem && FXSystem->RequiresEarlyViewUniformBuffer() && Views.IsValidIndex(0) && bRendererOutputFinalSceneColor)
 		{
 			// during ISR, instanced view RHI resources need to be initialized first.
 			if (FViewInfo* InstancedView = const_cast<FViewInfo*>(Views[0].GetInstancedView()))
@@ -4884,7 +4894,7 @@ void FDeferredShadingSceneRenderer::BeginInitViews(
 	TaskDatas.VisibilityTaskData->ProcessRenderThreadTasks(BasePassDepthStencilAccess, InstanceCullingManager, VirtualTextureUpdater);
 
 	// Make a second attempt to launch shadow tasks it wasn't able to the first time due to visibility being deferred.
-	if (GetRendererOutput() == ERendererOutput::FinalSceneColor)
+	if (bRendererOutputFinalSceneColor)
 	{
 		BeginInitDynamicShadows(TaskDatas);
 	}
