[Backout] - CL26484702
Original CL Desc
-----------------------------------------------------------------
Allow multiple Field targets to be rendered at the same time

diff --git a/Engine/Shaders/Private/PhysicsFieldVisualizer.usf b/Engine/Shaders/Private/PhysicsFieldVisualizer.usf
--- a/Engine/Shaders/Private/PhysicsFieldVisualizer.usf
+++ b/Engine/Shaders/Private/PhysicsFieldVisualizer.usf
@@ -6,7 +6,6 @@
 #include "CommonViewUniformBuffer.ush"
 #include "SceneTextureParameters.ush"
 #include "DeferredShadingCommon.ush"
-#include "ShaderPrintCommon.ush" 
 #include "ShaderPrint.ush" 
 #include "PositionReconstructionCommon.ush" 
  
@@ -22,8 +21,6 @@ Buffer<float4> BoundsMax;
 uint BoundsOffset;
 uint BoundsSize;
 
-uint PrintOffset;
-
 uint LocalTarget;
 uint GlobalTarget;
 
@@ -74,53 +71,30 @@ void SetOutputTexture(inout float RayDensity, in int BoundsIndex, in uint2 Pixel
 // Vector field visualisation
 #if PERMUTATION_FIELD == 0
 
-FFontColor GetVectorTargetColor()
-{
-	if(GlobalTarget == LINEAR_FORCE)
-	{
-		return FontRed;
-	}
-	else if(GlobalTarget == LINEAR_VELOCITY)
-	{
-		return FontGreen;
-	}
-	else if(GlobalTarget == ANGULAR_TORQUE)
-	{
-		return FontBlue;
-	}
-	else if(GlobalTarget == ANGULAR_VELOCITY)
-	{
-		return FontYellow;
-	}
-	return FontWhite;
-}
-
 void PrintVectorField()
 {
 	FShaderPrintContext Context = InitShaderPrintContext(true, uint2(50.0f, 200.0f));
-	Context.Pos.y += Context.Config.FontSpacing.y * PrintOffset;
 
-	const FFontColor VectorColor = GetVectorTargetColor();
-
-	Print(Context, TEXT("Vector Field | Type = "));
+	Print(Context, TEXT("Vector Field Type - "));
 	if(GlobalTarget == LINEAR_FORCE)
 	{
-		Print(Context, TEXT("Linear Force "), VectorColor);
+		Print(Context, TEXT("Linear Force "));
 	}
 	else if(GlobalTarget == LINEAR_VELOCITY)
 	{
-		Print(Context, TEXT("Linear Velocity "), VectorColor);
+		Print(Context, TEXT("Linear Velocity "));
 	}
 	else if(GlobalTarget == ANGULAR_TORQUE)
 	{
-		Print(Context, TEXT("Angular Torque "), VectorColor);
+		Print(Context, TEXT("Angular Torque "));
 	}
 	else if(GlobalTarget == ANGULAR_VELOCITY)
 	{
-		Print(Context, TEXT("Angular Velocity "), VectorColor);
+		Print(Context, TEXT("Angular Velocity "));
 	}
+	Newline(Context);
 
-	Print(Context, TEXT("| Count = "));
+	Print(Context, TEXT("Vector Field Count - "));
 	Print(Context, BoundsSize);
 }
 
@@ -129,7 +103,6 @@ void MainCS(uint3 DispatchThreadId : SV_DispatchThreadID,
     uint3 GroupThreadId : SV_GroupThreadID)
 {
 	GFieldGroupThreadId = GroupThreadId.x * 16 + GroupThreadId.y * 4 + GroupThreadId.z;
-	const FFontColor AABBColor = GetVectorTargetColor();
 	if(DispatchThreadId.x < VOXEL_RESOLUTION * BoundsSize && DispatchThreadId.y < VOXEL_RESOLUTION && DispatchThreadId.z < VOXEL_RESOLUTION)
 	{
 		const uint BoundsIndex = BoundsOffset + DispatchThreadId.x / VOXEL_RESOLUTION;
@@ -167,7 +140,7 @@ void MainCS(uint3 DispatchThreadId : SV_DispatchThreadID,
 		{
 			for(int BoundsIndex = BoundsOffset, BoundsEnd = BoundsOffset + BoundsSize; BoundsIndex < BoundsEnd; ++BoundsIndex)
 			{
-				AddAABBWS(BoundsMin[BoundsIndex].xyz, BoundsMax[BoundsIndex].xyz, float4(AABBColor.Color, 1));
+				AddAABBWS(BoundsMin[BoundsIndex].xyz, BoundsMax[BoundsIndex].xyz, float4(1, 0.5f, 0, 1));
 			}
 			PrintVectorField();
 		}
@@ -178,123 +151,69 @@ void MainCS(uint3 DispatchThreadId : SV_DispatchThreadID,
 
 #if PERMUTATION_FIELD == 1 || PERMUTATION_FIELD == 2
 
-FFontColor GetScalarTargetColor()
-{
-	if(GlobalTarget == EXTERNAL_STRAIN)
-	{
-		return FontCyan;
-	}
-	else if(GlobalTarget == INTERNAL_STRAIN)
-	{
-		return FontMagenta;
-	}
-	else if(GlobalTarget == DISABLED_THRESHOLD)
-	{
-		return FontOrange;
-	}
-	else if(GlobalTarget == SLEEPING_THRESHOLD)
-	{
-		return FontPurple;
-	}
-	else if(GlobalTarget == KILL_PARTICLES)
-	{
-		return FontTurquoise;
-	}
-	return FontWhite;
-}
-
 void PrintScalarField()
 {
 	FShaderPrintContext Context = InitShaderPrintContext(true, uint2(50.0f, 200.0f));
-	Context.Pos.y += Context.Config.FontSpacing.y * PrintOffset;
 
-	const FFontColor ScalarColor = GetScalarTargetColor();
+	Print(Context, TEXT("Scalar Field Type - "));
 
-	Print(Context, TEXT("Scalar Field | Type = "));
 	if(GlobalTarget == EXTERNAL_STRAIN)
 	{
-		Print(Context, TEXT("External Strain "), ScalarColor);
+		Print(Context, TEXT("External Strain "));
 	}
 	else if(GlobalTarget == INTERNAL_STRAIN)
 	{
-		Print(Context, TEXT("Internal Strain "), ScalarColor);
+		Print(Context, TEXT("Internal Strain "));
 	}
 	else if(GlobalTarget == DISABLED_THRESHOLD)
 	{
-		Print(Context, TEXT("Disabled Threshold "), ScalarColor);
+		Print(Context, TEXT("Disabled Threshold "));
 	}
 	else if(GlobalTarget == SLEEPING_THRESHOLD)
 	{
-		Print(Context, TEXT("Sleeping Threshold "), ScalarColor);
+		Print(Context, TEXT("Sleeping Threshold "));
 	}
 	else if(GlobalTarget == KILL_PARTICLES)
 	{
-		Print(Context, TEXT("Kill Particles "), ScalarColor);
+		Print(Context, TEXT("Kill Particles "));
 	}
 
-	Print(Context, TEXT("| Count = "));
-	Print(Context, BoundsSize);
-}
+	Newline(Context);
 
-FFontColor GetIntegerTargetColor()
-{
-	if(GlobalTarget == DYNAMIC_STATE)
-	{
-		return FontSilver;
-	}
-	else if(GlobalTarget == COLLISION_GROUP)
-	{
-		return FontEmerald;
-	}
-	else if(GlobalTarget == ACTIVATE_DISABLED)
-	{
-		return FontWhite;
-	}
-	return FontWhite;
+	Print(Context, TEXT("Scalar Field Count - "));
+	Print(Context, BoundsSize);
 }
 
 void PrintIntegerField()
 {
 	FShaderPrintContext Context = InitShaderPrintContext(true, uint2(50.0f, 200.0f));
-	Context.Pos.y += Context.Config.FontSpacing.y * PrintOffset;
 
-	const FFontColor IntegerColor = GetIntegerTargetColor();
+	Print(Context, TEXT("Integer Field Type - "));
 
-	Print(Context, TEXT("Integer Field | Type = "));
 	if(GlobalTarget == DYNAMIC_STATE)
 	{
-		Print(Context, TEXT("Dynamic State "), IntegerColor);
+		Print(Context, TEXT("Dynamic State "));
 	}
 	else if(GlobalTarget == COLLISION_GROUP)
 	{
-		Print(Context, TEXT("Collision Group "), IntegerColor);
+		Print(Context, TEXT("Collision Group "));
 	}
 	else if(GlobalTarget == ACTIVATE_DISABLED)
 	{
-		Print(Context, TEXT("Activate Disabled "), IntegerColor);
+		Print(Context, TEXT("Activate Disabled "));
 	}
 
-	Print(Context, TEXT("| Count = "));
+	Newline(Context);
+	Print(Context, TEXT("Integer Field Count - "));
 	Print(Context, BoundsSize);
 }
 
-
-
 // Integer and scalar field visualisation
 [numthreads(8, 8, 1)]
 void MainCS(uint3 DispatchThreadId : SV_DispatchThreadID,
     uint3 GroupThreadId : SV_GroupThreadID)
 {
 	GFieldGroupThreadId = GroupThreadId.x * 8 + GroupThreadId.y + GroupThreadId.z;
-
-	#if PERMUTATION_FIELD == 2
-	const FFontColor AABBColor = GetIntegerTargetColor();
-	#endif
-
-	#if PERMUTATION_FIELD == 1
-	const FFontColor AABBColor = GetScalarTargetColor();
-	#endif
-
 	if(DispatchThreadId.x < OutputResolution.x && DispatchThreadId.y < OutputResolution.y)
 	{
 		float3 RayBegin, RayEnd, RayDirection;
@@ -347,7 +266,7 @@ void MainCS(uint3 DispatchThreadId : SV_DispatchThreadID,
 		{
 			for(int BoundsIndex = BoundsOffset, BoundsEnd = BoundsOffset + BoundsSize; BoundsIndex < BoundsEnd; ++BoundsIndex)
 			{
-				AddAABBWS(BoundsMin[BoundsIndex].xyz, BoundsMax[BoundsIndex].xyz, float4(AABBColor.Color, 1));
+				AddAABBWS(BoundsMin[BoundsIndex].xyz, BoundsMax[BoundsIndex].xyz, float4(1, 0.5f, 0, 1));
 			}
 #if PERMUTATION_FIELD == 2
 			PrintIntegerField();
diff --git a/Engine/Source/Runtime/Engine/Classes/PhysicsField/PhysicsFieldComponent.h b/Engine/Source/Runtime/Engine/Classes/PhysicsField/PhysicsFieldComponent.h
--- a/Engine/Source/Runtime/Engine/Classes/PhysicsField/PhysicsFieldComponent.h
+++ b/Engine/Source/Runtime/Engine/Classes/PhysicsField/PhysicsFieldComponent.h
@@ -185,7 +185,7 @@ public:
 	 * Update the datas based on the new bounds and commands
 	 * @param FieldCommands - Field commands to be sampled
 	 */
-	void UpdateInstance(const float TimeSeconds, const bool bIsDebugBuffer);
+	void UpdateInstance(const float TimeSeconds);
 
 	/** The field system resource. */
 	FPhysicsFieldResource* FieldResource = nullptr;
@@ -243,16 +243,13 @@ public:
 	/** Add the transient field command */
 	ENGINE_API void AddTransientCommand(const FFieldSystemCommand& FieldCommand, const bool bIsGPUField);
 
-	/** Add the persistent field command */
+	/** Add the persitent field command */
 	ENGINE_API void AddPersistentCommand(const FFieldSystemCommand& FieldCommand, const bool bIsGPUField);
 
-	/** Add the construction field command */
-	ENGINE_API void AddConstructionCommand(const FFieldSystemCommand& FieldCommand);
-
 	/** Remove the transient field command */
 	ENGINE_API void RemoveTransientCommand(const FFieldSystemCommand& FieldCommand, const bool bIsGPUField);
 
-	/** Remove the persistent field command */
+	/** Remove the persitent field command */
 	ENGINE_API void RemovePersistentCommand(const FFieldSystemCommand& FieldCommand, const bool bIsGPUField);
 
 	/** Fill the transient commands intersecting the bounding box from the physics field */
@@ -269,12 +266,9 @@ public:
 	/** List of all the field transient commands in the world */
 	TArray<FFieldSystemCommand> TransientCommands[(uint8)(EFieldCommandBuffer::NumFieldBuffers)];
 
-	/** List of all the field persistent commands in the world */
+	/** List of all the field persitent commands in the world */
 	TArray<FFieldSystemCommand> PersistentCommands[(uint8)(EFieldCommandBuffer::NumFieldBuffers)];
 
-	/** List of all the field construction commands in the world */
-	TArray<FFieldSystemCommand> ConstructionCommands[(uint8)(EFieldCommandBuffer::NumFieldBuffers)];
-
 	/** The instance of the GPU field system. */
 	FPhysicsFieldInstance* FieldInstance = nullptr;
 
diff --git a/Engine/Source/Runtime/Engine/Private/PhysicsField/PhysicsFieldComponent.cpp b/Engine/Source/Runtime/Engine/Private/PhysicsField/PhysicsFieldComponent.cpp
--- a/Engine/Source/Runtime/Engine/Private/PhysicsField/PhysicsFieldComponent.cpp
+++ b/Engine/Source/Runtime/Engine/Private/PhysicsField/PhysicsFieldComponent.cpp
@@ -680,7 +680,7 @@ void FPhysicsFieldInstance::ReleaseInstance()
 	BoundsMax.Empty();
 }
 
-void FPhysicsFieldInstance::UpdateInstance(const float TimeSeconds, const bool bIsDebugBuffer)
+void FPhysicsFieldInstance::UpdateInstance(const float TimeSeconds)
 {
 	NodesOffsets.Reset();
 	NodesParams.Reset();
@@ -715,17 +715,8 @@ void FPhysicsFieldInstance::UpdateInstance(const float TimeSeconds, const bool b
 				{
 					UPhysicsFieldComponent::BuildCommandBounds(FieldCommand);
 
-
-					FVector MinBounds = FieldCommand.BoundingBox.Min;
-					FVector MaxBounds = FieldCommand.BoundingBox.Max;
-					if (bIsDebugBuffer && ((MinBounds.X <= -FLT_MAX) || (MaxBounds.X >= FLT_MAX)))
-					{
-						// When no bounds default size is 1m for the box
-						MinBounds = FieldCommand.CenterPosition - FVector(50, 50, 50);
-						MaxBounds = FieldCommand.CenterPosition + FVector(50, 50, 50);
-					}
-					BoundsMin.Add(FVector4(MinBounds, 0.0));
-					BoundsMax.Add(FVector4(MaxBounds, FieldCommand.MaxMagnitude));
+					BoundsMin.Add(FVector4(FieldCommand.BoundingBox.Min, 0.0));
+					BoundsMax.Add(FVector4(FieldCommand.BoundingBox.Max, FieldCommand.MaxMagnitude));
 
 					TargetMin = MinVector(TargetMin, FieldCommand.BoundingBox.Min);
 					TargetMax = MaxVector(TargetMax, FieldCommand.BoundingBox.Max);
@@ -856,9 +847,6 @@ void UPhysicsFieldComponent::SendRenderDynamicData_Concurrent()
 
 	Super::SendRenderTransform_Concurrent();
 
-	static IConsoleVariable* TransientLifetimeCVar = IConsoleManager::Get().FindConsoleVariable(TEXT("r.PhysicsField.Rendering.TransientLifetime"));
-	const float TransientLifetime = TransientLifetimeCVar ? TransientLifetimeCVar->GetFloat() : 0.0f;
-
 	TArray<uint8> FieldBuffers = { (uint8)(EFieldCommandBuffer::GPUFieldBuffer) ,(uint8)(EFieldCommandBuffer::GPUDebugBuffer) };
 	for (uint32 FieldIndex = 0; FieldIndex < 2; ++FieldIndex)
 	{
@@ -866,35 +854,18 @@ void UPhysicsFieldComponent::SendRenderDynamicData_Concurrent()
 		if (LocalInstance)
 		{
 			const bool bPreviousUpdate = LocalInstance->FieldCommands.Num() > 0;
-			const bool bIsDebugBuffer = FieldIndex == 1;
-
-			const float TimeSeconds = GetWorld() ? GetWorld()->TimeSeconds : 0.0;
 
 			LocalInstance->FieldCommands.Empty();
 			LocalInstance->FieldCommands.Append(PersistentCommands[FieldBuffers[FieldIndex]]);
 			LocalInstance->FieldCommands.Append(TransientCommands[FieldBuffers[FieldIndex]]);
-			LocalInstance->FieldCommands.Append(ConstructionCommands[FieldBuffers[FieldIndex]]);
-			
-			if (bIsDebugBuffer)
-			{
-				for (int32 CommandIndex = TransientCommands[FieldBuffers[FieldIndex]].Num() - 1; CommandIndex >= 0; --CommandIndex)
-				{
-					if (TimeSeconds - TransientCommands[FieldBuffers[FieldIndex]][CommandIndex].TimeCreation > TransientLifetime)
-					{
-						TransientCommands[FieldBuffers[FieldIndex]].RemoveAt(CommandIndex);
-					}
-				}
-			}
-			else
-			{
-				TransientCommands[FieldBuffers[FieldIndex]].Empty();
-			}
+			TransientCommands[FieldBuffers[FieldIndex]].Empty();
 
 			const bool bCurrentUpdate = LocalInstance->FieldCommands.Num() > 0;
 
 			if (bCurrentUpdate || bPreviousUpdate)
 			{
-				LocalInstance->UpdateInstance(TimeSeconds, bIsDebugBuffer);
+				const float TimeSeconds = GetWorld() ? GetWorld()->TimeSeconds : 0.0;
+				LocalInstance->UpdateInstance(TimeSeconds);
 			}
 		}
 	}
@@ -985,12 +956,11 @@ void UPhysicsFieldComponent::BuildCommandBounds(FFieldSystemCommand& FieldComman
 	FieldCommand.BoundingBox.Min = FVector(-FLT_MAX);
 	FieldCommand.BoundingBox.Max = FVector(FLT_MAX);
 	FieldCommand.MaxMagnitude = 1.0;
-	FieldCommand.CenterPosition = FVector::Zero();
 
 	if (FieldCommand.RootNode.IsValid())
 	{
 		FieldCommand.MaxMagnitude = FieldCommand.RootNode->EvalMaxMagnitude();
-		FieldCommand.RootNode->ComputeFieldBounds(FieldCommand.BoundingBox.Min, FieldCommand.BoundingBox.Max, FieldCommand.CenterPosition);
+		FieldCommand.RootNode->ComputeFieldBounds(FieldCommand.BoundingBox.Min, FieldCommand.BoundingBox.Max);
 	}
 }
 
@@ -1020,11 +990,6 @@ void UPhysicsFieldComponent::AddPersistentCommand(const FFieldSystemCommand& Fie
 	}
 }
 
-void UPhysicsFieldComponent::AddConstructionCommand(const FFieldSystemCommand& FieldCommand)
-{
-	ConstructionCommands[(uint8)(EFieldCommandBuffer::GPUDebugBuffer)].Add(FieldCommand);
-}
-
 void UPhysicsFieldComponent::RemoveTransientCommand(const FFieldSystemCommand& FieldCommand, const bool bIsWorldField)
 {
 	if (bIsWorldField)
diff --git a/Engine/Source/Runtime/Experimental/Chaos/Private/Field/FieldSystemNodes.cpp b/Engine/Source/Runtime/Experimental/Chaos/Private/Field/FieldSystemNodes.cpp
--- a/Engine/Source/Runtime/Experimental/Chaos/Private/Field/FieldSystemNodes.cpp
+++ b/Engine/Source/Runtime/Experimental/Chaos/Private/Field/FieldSystemNodes.cpp
@@ -249,11 +249,10 @@ float FRadialIntMask::EvalMaxMagnitude() const
 	return 1.0f;
 }
 
-void FRadialIntMask::ComputeFieldBounds(FVector& MinBounds, FVector& MaxBounds, FVector& CenterPosition) const
+void FRadialIntMask::ComputeFieldBounds(FVector& MinBounds, FVector& MaxBounds) const
 {
 	MinBounds = (ExteriorValue == 0) ? Position - FVector(Radius) : FVector(-FLT_MAX);
 	MaxBounds = (ExteriorValue == 0) ? Position + FVector(Radius) : FVector(FLT_MAX);
-	CenterPosition = Position;
 }
 
 /**
@@ -424,13 +423,6 @@ float FWaveScalar::EvalMaxMagnitude() const
 	return Magnitude;
 }
 
-void FWaveScalar::ComputeFieldBounds(FVector& MinBounds, FVector& MaxBounds, FVector& CenterPosition) const
-{
-	MinBounds = FVector(-FLT_MAX);
-	MaxBounds = FVector(FLT_MAX);
-	CenterPosition = Position;
-}
-
 /**
 * Function Utils
 */
@@ -582,11 +574,10 @@ float FRadialFalloff::EvalMaxMagnitude() const
 	return Magnitude;
 }
 
-void FRadialFalloff::ComputeFieldBounds(FVector& MinBounds, FVector& MaxBounds, FVector& CenterPosition) const
+void FRadialFalloff::ComputeFieldBounds(FVector& MinBounds, FVector& MaxBounds) const
 {
 	MinBounds = (Default == 0) ? Position - FVector(Radius) : FVector(-FLT_MAX);
 	MaxBounds = (Default == 0) ? Position + FVector(Radius) : FVector(FLT_MAX);
-	CenterPosition = Position;
 }
 
 /**
@@ -699,11 +690,10 @@ float FPlaneFalloff::EvalMaxMagnitude() const
 	return Magnitude;
 }
 
-void FPlaneFalloff::ComputeFieldBounds(FVector& MinBounds, FVector& MaxBounds, FVector& CenterPosition) const
+void FPlaneFalloff::ComputeFieldBounds(FVector& MinBounds, FVector& MaxBounds) const
 {
 	MinBounds = FVector(-FLT_MAX);
 	MaxBounds = FVector(FLT_MAX);
-	CenterPosition = Position;
 }
 
 /**
@@ -818,7 +808,7 @@ float FBoxFalloff::EvalMaxMagnitude() const
 	return Magnitude;
 }
 
-void FBoxFalloff::ComputeFieldBounds(FVector& MinBounds, FVector& MaxBounds, FVector& CenterPosition) const
+void FBoxFalloff::ComputeFieldBounds(FVector& MinBounds, FVector& MaxBounds) const
 {
 	MinBounds = FVector(-FLT_MAX);
 	MaxBounds = FVector(FLT_MAX);
@@ -829,7 +819,6 @@ void FBoxFalloff::ComputeFieldBounds(FVector& MinBounds, FVector& MaxBounds, FVe
 		MinBounds = BoundingBox.Min;
 		MaxBounds = BoundingBox.Max;
 	}
-	CenterPosition = Transform.GetTranslation();
 }
 
 /**
@@ -1021,13 +1010,6 @@ float FRadialVector::EvalMaxMagnitude() const
 	return Magnitude;
 }
 
-void FRadialVector::ComputeFieldBounds(FVector& MinBounds, FVector& MaxBounds, FVector& CenterPosition) const
-{
-	MinBounds = FVector(-FLT_MAX);
-	MaxBounds = FVector(FLT_MAX);
-	CenterPosition = Position;
-}
-
 
 /**
 * FRandomVector
@@ -1267,31 +1249,17 @@ float FSumScalar::EvalMaxMagnitude() const
 	return Magnitude * MaxMagnitude;
 }
 
-void FSumScalar::ComputeFieldBounds(FVector& MinBounds, FVector& MaxBounds, FVector& CenterPosition) const
+void FSumScalar::ComputeFieldBounds(FVector& MinBounds, FVector& MaxBounds) const
 {
 	FVector MinBoundsA(-FLT_MAX), MaxBoundsA(FLT_MAX), MinBoundsB(-FLT_MAX), MaxBoundsB(FLT_MAX);
-	FVector CenterPositionA = FVector::Zero(), CenterPositionB = FVector::Zero();
 	if (ScalarRight.IsValid())
 	{
-		ScalarRight->ComputeFieldBounds(MinBoundsA, MaxBoundsA, CenterPositionA);
+		ScalarRight->ComputeFieldBounds(MinBoundsA, MaxBoundsA);
 	}
 	if (ScalarLeft.IsValid())
 	{
-		ScalarLeft->ComputeFieldBounds(MinBoundsB, MaxBoundsB, CenterPositionB);
-	}
-	CenterPosition = FVector::Zero();
-	int32 NumCenters = 0;
-	if (CenterPositionA != FVector::Zero())
-	{
-		CenterPosition += CenterPositionA;
-		++NumCenters;
-	}
-	if (CenterPositionB != FVector::Zero())
-	{
-		CenterPosition += CenterPositionB;
-		++NumCenters;
+		ScalarLeft->ComputeFieldBounds(MinBoundsB, MaxBoundsB);
 	}
-	CenterPosition = (NumCenters > 0) ? CenterPosition / NumCenters : FVector::Zero();
 
 	if (Operation == EFieldOperationType::Field_Multiply ||
 		Operation == EFieldOperationType::Field_Divide)
@@ -1511,42 +1479,22 @@ float FSumVector::EvalMaxMagnitude() const
 	return Magnitude * MaxMagnitudeA * MaxMagnitude;
 }
 
-void FSumVector::ComputeFieldBounds(FVector& MinBounds, FVector& MaxBounds, FVector& CenterPosition) const
+void FSumVector::ComputeFieldBounds(FVector& MinBounds, FVector& MaxBounds) const
 {
 	FVector MinBoundsA(-FLT_MAX), MaxBoundsA(FLT_MAX), MinBoundsB(-FLT_MAX), MaxBoundsB(FLT_MAX), MinBoundsC(-FLT_MAX), MaxBoundsC(FLT_MAX);
-	FVector CenterPositionA = FVector::Zero(), CenterPositionB = FVector::Zero(), CenterPositionC = FVector::Zero();
 	if (Scalar.IsValid())
 	{
-		Scalar->ComputeFieldBounds(MinBoundsC, MaxBoundsC, CenterPositionC);
+		Scalar->ComputeFieldBounds(MinBoundsC, MaxBoundsC);
 	}
 	if (VectorRight.IsValid())
 	{
-		VectorRight->ComputeFieldBounds(MinBoundsA, MaxBoundsA, CenterPositionA);
+		VectorRight->ComputeFieldBounds(MinBoundsA, MaxBoundsA);
 	}
 	if (VectorLeft.IsValid())
 	{
-		VectorLeft->ComputeFieldBounds(MinBoundsB, MaxBoundsB, CenterPositionB);
+		VectorLeft->ComputeFieldBounds(MinBoundsB, MaxBoundsB);
 	}
 
-	CenterPosition = FVector::Zero();
-	int32 NumCenters = 0;
-	if (CenterPositionA != FVector::Zero())
-	{
-		CenterPosition += CenterPositionA;
-		++NumCenters;
-	}
-	if (CenterPositionB != FVector::Zero())
-	{
-		CenterPosition += CenterPositionB;
-		++NumCenters;
-	}
-	if (CenterPositionC != FVector::Zero())
-	{
-		CenterPosition += CenterPositionC;
-		++NumCenters;
-	}
-	CenterPosition = (NumCenters > 0) ? CenterPosition / NumCenters : FVector::Zero();
-
 	if (Operation == EFieldOperationType::Field_Multiply ||
 		Operation == EFieldOperationType::Field_Divide)
 	{
@@ -1648,12 +1596,12 @@ float FConversionField<InT, OutT>::EvalMaxMagnitude() const
 }
 
 template<class InT, class OutT>
-void FConversionField<InT, OutT>::ComputeFieldBounds(FVector& MinBounds, FVector& MaxBounds, FVector& CenterPosition) const
+void FConversionField<InT, OutT>::ComputeFieldBounds(FVector& MinBounds, FVector& MaxBounds) const
 {
 	FVector MinBoundsA(-FLT_MAX), MaxBoundsA(FLT_MAX);
 	if (InputField.IsValid())
 	{
-		InputField->ComputeFieldBounds(MinBoundsA, MaxBoundsA, CenterPosition);
+		InputField->ComputeFieldBounds(MinBoundsA, MaxBoundsA);
 	}
 	MinBounds = MinBoundsA;
 	MaxBounds = MaxBoundsA;
@@ -1812,34 +1760,18 @@ float FCullingField<T>::EvalMaxMagnitude() const
 }
 
 template<class T>
-void FCullingField<T>::ComputeFieldBounds(FVector& MinBounds, FVector& MaxBounds, FVector& CenterPosition) const
+void FCullingField<T>::ComputeFieldBounds(FVector& MinBounds, FVector& MaxBounds) const
 {
 	FVector MinBoundsA(-FLT_MAX), MaxBoundsA(FLT_MAX), MinBoundsB(-FLT_MAX), MaxBoundsB(FLT_MAX);
-	FVector CenterPositionA = FVector::Zero(), CenterPositionB = FVector::Zero();
 
 	if (Culling.IsValid())
 	{
-		Culling->ComputeFieldBounds(MinBoundsA, MaxBoundsA, CenterPositionA);
+		Culling->ComputeFieldBounds(MinBoundsA, MaxBoundsA);
 	}
 	if (Input.IsValid())
 	{
-		Input->ComputeFieldBounds(MinBoundsB, MaxBoundsB, CenterPositionB);
+		Input->ComputeFieldBounds(MinBoundsB, MaxBoundsB);
 	}
-
-	CenterPosition = FVector::Zero();
-	int32 NumCenters = 0;
-	if (CenterPositionA != FVector::Zero())
-	{
-		CenterPosition += CenterPositionA;
-		++NumCenters;
-	}
-	if (CenterPositionB != FVector::Zero())
-	{
-		CenterPosition += CenterPositionB;
-		++NumCenters;
-	}
-	CenterPosition = (NumCenters > 0) ? CenterPosition / NumCenters : FVector::Zero();
-
 	if (Operation == EFieldCullingOperationType::Field_Culling_Inside)
 	{
 		MinBounds = MinVector(MinBoundsA, MinBoundsB);
diff --git a/Engine/Source/Runtime/Experimental/Chaos/Public/Field/FieldSystem.h b/Engine/Source/Runtime/Experimental/Chaos/Public/Field/FieldSystem.h
--- a/Engine/Source/Runtime/Experimental/Chaos/Public/Field/FieldSystem.h
+++ b/Engine/Source/Runtime/Experimental/Chaos/Public/Field/FieldSystem.h
@@ -456,11 +456,10 @@ public:
 	virtual float EvalMaxMagnitude() const { return 1.0; }
 
 	/** Compute the min/max spatial bounds of the field */
-	virtual void ComputeFieldBounds(FVector& MinBounds, FVector& MaxBounds, FVector& CenterPosition) const
+	virtual void ComputeFieldBounds(FVector& MinBounds, FVector& MaxBounds) const
 	{
 		MinBounds = FVector(-FLT_MAX);
 		MaxBounds = FVector(FLT_MAX);
-		CenterPosition = FVector::Zero();
 	}
 
 	FName GetName() const { return Name; }
@@ -532,7 +531,6 @@ public:
 		, BoundingBox(FVector(-FLT_MAX), FVector(FLT_MAX))
 		, PhysicsType(EFieldPhysicsType::Field_None)
 		, MaxMagnitude(1.0)
-		, CenterPosition(FVector::Zero())
 	{}
 	FFieldSystemCommand(const FName& TargetAttributeIn, FFieldNodeBase * RootNodeIn)
 		: TargetAttribute(TargetAttributeIn)
@@ -542,7 +540,6 @@ public:
 		, BoundingBox(FVector(-FLT_MAX), FVector(FLT_MAX))
 		, PhysicsType(GetFieldPhysicsType(TargetAttributeIn))
 		, MaxMagnitude(1.0)
-		, CenterPosition(FVector::Zero())
 	{}
 	FFieldSystemCommand(const EFieldPhysicsType PhsyicsTypeIn, FFieldNodeBase* RootNodeIn)
 		: TargetAttribute(GetFieldPhysicsName(PhsyicsTypeIn))
@@ -552,7 +549,6 @@ public:
 		, BoundingBox(FVector(-FLT_MAX), FVector(FLT_MAX))
 		, PhysicsType(PhsyicsTypeIn)
 		, MaxMagnitude(1.0)
-		, CenterPosition(FVector::Zero())
 	{}
 
 	// Commands are copied when moved from the one thread to 
@@ -565,7 +561,6 @@ public:
 		, BoundingBox(Other.BoundingBox)
 		, PhysicsType(Other.RootNode ? Other.PhysicsType : EFieldPhysicsType::Field_None)
 		, MaxMagnitude(Other.MaxMagnitude)
-		, CenterPosition(Other.CenterPosition)
 	{
 		for (const TPair<FFieldSystemMetaData::EMetaType, TUniquePtr<FFieldSystemMetaData>>& Meta : Other.MetaData)
 		{
@@ -619,7 +614,6 @@ public:
 	FBox BoundingBox;
 	EFieldPhysicsType PhysicsType;
 	float MaxMagnitude;
-	FVector CenterPosition;
 
 	TMap<FFieldSystemMetaData::EMetaType, TUniquePtr<FFieldSystemMetaData> > MetaData;
 };
diff --git a/Engine/Source/Runtime/Experimental/Chaos/Public/Field/FieldSystemNodes.h b/Engine/Source/Runtime/Experimental/Chaos/Public/Field/FieldSystemNodes.h
--- a/Engine/Source/Runtime/Experimental/Chaos/Public/Field/FieldSystemNodes.h
+++ b/Engine/Source/Runtime/Experimental/Chaos/Public/Field/FieldSystemNodes.h
@@ -85,7 +85,7 @@ public:
 	CHAOS_API virtual float EvalMaxMagnitude() const override;
 
 	/** Compute the min/max spatial bounds of the field */
-	CHAOS_API virtual void ComputeFieldBounds(FVector& MinBounds, FVector& MaxBounds, FVector& CenterPosition) const override;
+	CHAOS_API virtual void ComputeFieldBounds(FVector& MinBounds, FVector& MaxBounds) const override;
 
 	float Radius;
 	FVector Position;
@@ -172,8 +172,6 @@ public:
 	/** Evaluate the maximum magnitude of the field graph  */
 	CHAOS_API virtual float EvalMaxMagnitude() const override;
 
-	/** Compute the min/max spatial bounds of the field */
-	CHAOS_API virtual void ComputeFieldBounds(FVector& MinBounds, FVector& MaxBounds, FVector& CenterPosition) const override;
 
 	float Magnitude;
 	FVector Position;
@@ -231,7 +229,7 @@ public:
 	CHAOS_API virtual float EvalMaxMagnitude() const override;
 
 	/** Compute the min/max spatial bounds of the field */
-	CHAOS_API virtual void ComputeFieldBounds(FVector& MinBounds, FVector& MaxBounds, FVector& CenterPosition) const override;
+	CHAOS_API virtual void ComputeFieldBounds(FVector& MinBounds, FVector& MaxBounds) const override;
 
 	float Magnitude;
 	float MinRange;
@@ -295,7 +293,7 @@ public:
 	CHAOS_API virtual float EvalMaxMagnitude() const override;
 
 	/** Compute the min/max spatial bounds of the field */
-	CHAOS_API virtual void ComputeFieldBounds(FVector& MinBounds, FVector& MaxBounds, FVector& CenterPosition) const override;
+	CHAOS_API virtual void ComputeFieldBounds(FVector& MinBounds, FVector& MaxBounds) const override;
 
 
 	float Magnitude;
@@ -357,7 +355,7 @@ public:
 	CHAOS_API virtual float EvalMaxMagnitude() const override;
 
 	/** Compute the min/max spatial bounds of the field */
-	CHAOS_API virtual void ComputeFieldBounds(FVector& MinBounds, FVector& MaxBounds, FVector& CenterPosition) const override;
+	CHAOS_API virtual void ComputeFieldBounds(FVector& MinBounds, FVector& MaxBounds) const override;
 
 	float Magnitude;
 	float MinRange;
@@ -485,9 +483,6 @@ public:
 	/** Evaluate the maximum magnitude of the field graph  */
 	CHAOS_API virtual float EvalMaxMagnitude() const override;
 
-	/** Compute the min/max spatial bounds of the field */
-	CHAOS_API virtual void ComputeFieldBounds(FVector& MinBounds, FVector& MaxBounds, FVector& CenterPosition) const override;
-
 	float Magnitude;
 	FVector Position;
 
@@ -575,7 +570,7 @@ public:
 	CHAOS_API virtual float EvalMaxMagnitude() const override;
 
 	/** Compute the min/max spatial bounds of the field */
-	CHAOS_API virtual void ComputeFieldBounds(FVector& MinBounds, FVector& MaxBounds, FVector& CenterPosition) const override;
+	CHAOS_API virtual void ComputeFieldBounds(FVector& MinBounds, FVector& MaxBounds) const override;
 
 	/** Default setup for the sum scalar field */
 	static void FillDefaultSetup(TArray<int32>& NodesOffsets, TArray<float>& NodesParams)
@@ -644,7 +639,7 @@ public:
 	CHAOS_API virtual float EvalMaxMagnitude() const override;
 
 	/** Compute the min/max spatial bounds of the field */
-	CHAOS_API virtual void ComputeFieldBounds(FVector& MinBounds, FVector& MaxBounds, FVector& CenterPosition) const override;
+	CHAOS_API virtual void ComputeFieldBounds(FVector& MinBounds, FVector& MaxBounds) const override;
 
 	/** Default setup for the sum vector field */
 	static void FillDefaultSetup(TArray<int32>& NodesOffsets, TArray<float>& NodesParams)
@@ -704,7 +699,7 @@ public:
 	CHAOS_API virtual float EvalMaxMagnitude() const override;
 
 	/** Compute the min/max spatial bounds of the field */
-	CHAOS_API virtual void ComputeFieldBounds(FVector& MinBounds, FVector& MaxBounds, FVector& CenterPosition) const override;
+	CHAOS_API virtual void ComputeFieldBounds(FVector& MinBounds, FVector& MaxBounds) const override;
 
 	/** Default setup for the conversion field */
 	static void FillDefaultSetup(TArray<int32>& NodesOffsets, TArray<float>& NodesParams)
@@ -765,7 +760,7 @@ public:
 	CHAOS_API virtual float EvalMaxMagnitude() const override;
 
 	/** Compute the min/max spatial bounds of the field */
-	CHAOS_API virtual void ComputeFieldBounds(FVector& MinBounds, FVector& MaxBounds, FVector& CenterPosition) const override;
+	CHAOS_API virtual void ComputeFieldBounds(FVector& MinBounds, FVector& MaxBounds) const override;
 
 	TUniquePtr< FFieldNode<float> > Culling;
 	TUniquePtr< FFieldNode<T> > Input;
diff --git a/Engine/Source/Runtime/Experimental/Chaos/Public/Field/FieldSystemTypes.h b/Engine/Source/Runtime/Experimental/Chaos/Public/Field/FieldSystemTypes.h
--- a/Engine/Source/Runtime/Experimental/Chaos/Public/Field/FieldSystemTypes.h
+++ b/Engine/Source/Runtime/Experimental/Chaos/Public/Field/FieldSystemTypes.h
@@ -339,7 +339,7 @@ inline CHAOS_API const TArray<FName>& GetFieldPhysicsNames()
 	static const TArray<FName> FieldPhysicsNames =
 	{ "Node", "DynamicState", "LinearForce", "ExternalClusterStrain", "Kill", "LinearVelocity", "AngularVelocity",
 	 "AngularTorque", "InternalClusterStrain", "DisableThreshold", "SleepingThreshold", "PositionStatic",
-	 "PositionAnimated", "PositionTarget", "DynamicConstraint", "CollisionGroup", "ActivateDisabled", "InitialLinearVelocity", "InitialAngularVelocity", "LinearImpulse"};
+	 "PositionAnimated", "PositionTarget", "DynamicConstraint", "CollisionGroup", "ActivateDisabled", "InitialLinearVelocity", "InitialAngularVelocity" };
 
 	return FieldPhysicsNames;
 }
@@ -364,8 +364,7 @@ inline CHAOS_API const TArray<EFieldPhysicsType>& GetFieldPhysicsTypes()
 																EFieldPhysicsType::Field_CollisionGroup,
 																EFieldPhysicsType::Field_ActivateDisabled, 
 																EFieldPhysicsType::Field_InitialLinearVelocity, 
-																EFieldPhysicsType::Field_InitialAngularVelocity,
-																EFieldPhysicsType::Field_LinearImpulse };
+																EFieldPhysicsType::Field_InitialAngularVelocity };
 
 	return FieldPhysicsTypes;
 }
diff --git a/Engine/Source/Runtime/Experimental/FieldSystem/Source/FieldSystemEngine/Private/Field/FieldSystemComponent.cpp b/Engine/Source/Runtime/Experimental/FieldSystem/Source/FieldSystemEngine/Private/Field/FieldSystemComponent.cpp
--- a/Engine/Source/Runtime/Experimental/FieldSystem/Source/FieldSystemEngine/Private/Field/FieldSystemComponent.cpp
+++ b/Engine/Source/Runtime/Experimental/FieldSystem/Source/FieldSystemEngine/Private/Field/FieldSystemComponent.cpp
@@ -296,17 +296,6 @@ void UFieldSystemComponent::AddFieldCommand(bool Enabled, EFieldPhysicsType Targ
 			if (Command.RootNode)
 			{
 				SetupConstructionFields.Add(Command);
-
-				UWorld* World = GetWorld();
-				if (World && World->PhysicsField)
-				{
-					const FName Name = GetOwner() ? *GetOwner()->GetName() : TEXT("");
-
-					FFieldSystemCommand LocalCommand = Command;
-					LocalCommand.InitFieldNodes(World->GetTimeSeconds(), Name);
-
-					World->PhysicsField->AddConstructionCommand(LocalCommand);
-				}
 			}
 		}
 
diff --git a/Engine/Source/Runtime/Experimental/GeometryCollectionEngine/Private/GeometryCollection/GeometryCollectionComponent.cpp b/Engine/Source/Runtime/Experimental/GeometryCollectionEngine/Private/GeometryCollection/GeometryCollectionComponent.cpp
--- a/Engine/Source/Runtime/Experimental/GeometryCollectionEngine/Private/GeometryCollection/GeometryCollectionComponent.cpp
+++ b/Engine/Source/Runtime/Experimental/GeometryCollectionEngine/Private/GeometryCollection/GeometryCollectionComponent.cpp
@@ -4581,17 +4581,6 @@ void UGeometryCollectionComponent::GetInitializationCommands(TArray<FFieldSystem
 						if (NewCommand.RootNode)
 						{
 							CombinedCommmands.Emplace(NewCommand);
-
-							UWorld* World = GetWorld();
-							if (World && World->PhysicsField)
-							{
-								const FName Name = GetOwner() ? *GetOwner()->GetName() : TEXT("");
-
-								FFieldSystemCommand LocalCommand = NewCommand;
-								LocalCommand.InitFieldNodes(World->GetTimeSeconds(), Name);
-
-								World->PhysicsField->AddConstructionCommand(LocalCommand);
-							}
 						}
 					}
 				}
diff --git a/Engine/Source/Runtime/Renderer/Private/PhysicsFieldRendering.cpp b/Engine/Source/Runtime/Renderer/Private/PhysicsFieldRendering.cpp
--- a/Engine/Source/Runtime/Renderer/Private/PhysicsFieldRendering.cpp
+++ b/Engine/Source/Runtime/Renderer/Private/PhysicsFieldRendering.cpp
@@ -20,7 +20,7 @@
 #include "ShaderPrint.h"
 
 // Console variables
-int32 GPhysicsFieldTargetType = EFieldPhysicsType::Field_PhysicsType_Max;
+int32 GPhysicsFieldTargetType = EFieldPhysicsType::Field_LinearForce;
 FAutoConsoleVariableRef CVarPhysicsFieldTargetType(
 	TEXT("r.PhysicsField.Rendering.TargetType"),
 	GPhysicsFieldTargetType,
@@ -41,13 +41,6 @@ FAutoConsoleVariableRef CVarPhysicsFieldSystemType(
 	TEXT("Physics field boolean to check if we want to display the CPU(0) or GPU(1) field.\n"),
 	ECVF_RenderThreadSafe);
 
-float GPhysicsFieldTransientLifetime = 3.0f;
-FAutoConsoleVariableRef CVarPhysicsFieldTransientLifetime(
-	TEXT("r.PhysicsField.Rendering.TransientLifetime"),
-	GPhysicsFieldTransientLifetime,
-	TEXT("Physics field transient commands lifetime for rendering.\n"),
-	ECVF_RenderThreadSafe);
-
 /**
  * Physics Field Rendering.
  */
@@ -76,7 +69,6 @@ class FPhysicsFieldRayMarchingCS : public FGlobalShader
 		SHADER_PARAMETER(uint32, BoundsOffset)
 		SHADER_PARAMETER(uint32, BoundsSize)
 		SHADER_PARAMETER(uint32, LocalTarget)
-		SHADER_PARAMETER(uint32, PrintOffset)
 		SHADER_PARAMETER(uint32, GlobalTarget)
 		SHADER_PARAMETER(FVector2f, OutputResolution)
 		SHADER_PARAMETER_RDG_TEXTURE_UAV(RWTexture2D, OutputTexture)
@@ -101,7 +93,7 @@ IMPLEMENT_GLOBAL_SHADER(FPhysicsFieldRayMarchingCS, "/Engine/Private/PhysicsFiel
 
 static FPhysicsFieldRayMarchingCS::FParameters* CreateShaderParameters(FRDGBuilder& GraphBuilder,
 	const FViewInfo& View, const FPhysicsFieldResource* PhysicsFieldResource,
-	FRDGTextureRef& OutputTexture, const int32 TargetIndex, const int32 TargetType, const uint32 PrintOffset)
+	FRDGTextureRef& OutputTexture, const int32 TargetIndex)
 {
 	const FIntPoint OutputResolution(OutputTexture->Desc.Extent);
 	FSceneTextureParameters SceneTextures = GetSceneTextureParameters(GraphBuilder, View);
@@ -116,8 +108,8 @@ static FPhysicsFieldRayMarchingCS::FParameters* CreateShaderParameters(FRDGBuild
 
 	Parameters->BoundsMin = PhysicsFieldResource->BoundsMin.SRV;
 	Parameters->BoundsMax = PhysicsFieldResource->BoundsMax.SRV;
-	Parameters->BoundsOffset = PhysicsFieldResource->FieldInfos.BoundsOffsets[TargetType];
-	Parameters->BoundsSize = PhysicsFieldResource->FieldInfos.BoundsOffsets[TargetType + 1] - PhysicsFieldResource->FieldInfos.BoundsOffsets[TargetType];;
+	Parameters->BoundsOffset = PhysicsFieldResource->FieldInfos.BoundsOffsets[GPhysicsFieldTargetType];
+	Parameters->BoundsSize = PhysicsFieldResource->FieldInfos.BoundsOffsets[GPhysicsFieldTargetType + 1] - PhysicsFieldResource->FieldInfos.BoundsOffsets[GPhysicsFieldTargetType];;
 
 	Parameters->NodesParams = PhysicsFieldResource->NodesParams.SRV;
 	Parameters->NodesOffsets = PhysicsFieldResource->NodesOffsets.SRV;
@@ -125,8 +117,7 @@ static FPhysicsFieldRayMarchingCS::FParameters* CreateShaderParameters(FRDGBuild
 	Parameters->TimeSeconds = PhysicsFieldResource->FieldInfos.TimeSeconds;
 
 	Parameters->LocalTarget = TargetIndex;
-	Parameters->GlobalTarget = TargetType;
-	Parameters->PrintOffset = PrintOffset;
+	Parameters->GlobalTarget = GPhysicsFieldTargetType;
 
 	return Parameters;
 }
@@ -137,51 +128,47 @@ static void AddPhysicsFieldRayMarchingPass(
 	const FPhysicsFieldResource* PhysicsFieldResource,
 	FRDGTextureRef& OutputTexture)
 {
-	const FIntPoint OutputResolution(OutputTexture->Desc.Extent);
-	auto RenderPhysicsField = [&GraphBuilder, &View, PhysicsFieldResource, &OutputTexture, &OutputResolution](const int32 TargetType, const bool bClearBackground, uint32& PrintOffset) {
+	FSceneTextureParameters SceneTextures = GetSceneTextureParameters(GraphBuilder, View);
 
+	const FIntPoint OutputResolution(OutputTexture->Desc.Extent);
+	if(GPhysicsFieldTargetType < EFieldPhysicsType::Field_PhysicsType_Max)
+	{
 		int32 TargetIndex = INDEX_NONE;
 		EFieldOutputType FieldOutputType = EFieldOutputType::Field_Output_Max;
-		GetFieldIndex(TargetType, TargetIndex, FieldOutputType);
+		GetFieldIndex(GPhysicsFieldTargetType, TargetIndex, FieldOutputType);
 
-		const int32 NumBounds = PhysicsFieldResource->FieldInfos.BoundsOffsets[TargetType + 1] - PhysicsFieldResource->FieldInfos.BoundsOffsets[TargetType];
-		if (NumBounds > 0 && (!bClearBackground || (bClearBackground && (FieldOutputType == EFieldOutputType::Field_Output_Vector))))
+		if (GPhysicsFieldTargetType < EFieldPhysicsType::Field_PhysicsType_Max)
 		{
-			const FIntVector BoundResolution(FPhysicsFieldRayMarchingCS::VoxelResolution * NumBounds, FPhysicsFieldRayMarchingCS::VoxelResolution, FPhysicsFieldRayMarchingCS::VoxelResolution);
-			PrintOffset = !bClearBackground ? PrintOffset + 1 : PrintOffset;
+			// Reset the field background to highlight the vector lines color
+			if (FieldOutputType == EFieldOutputType::Field_Output_Vector)
+			{
+				FPhysicsFieldRayMarchingCS::FParameters* Parameters = CreateShaderParameters(GraphBuilder, View, PhysicsFieldResource, OutputTexture, TargetIndex);
 
-			FPhysicsFieldRayMarchingCS::FParameters* Parameters = CreateShaderParameters(GraphBuilder, View, PhysicsFieldResource, OutputTexture, TargetIndex, TargetType, PrintOffset);
+				FPhysicsFieldRayMarchingCS::FPermutationDomain PermutationVector;
+				PermutationVector.Set<FPhysicsFieldRayMarchingCS::FFieldType>(3);
+				PermutationVector.Set<FPhysicsFieldRayMarchingCS::FEvalType>(GPhysicsFieldEvalType);
+				TShaderMapRef<FPhysicsFieldRayMarchingCS> ComputeShader(View.ShaderMap, PermutationVector);
 
-			FPhysicsFieldRayMarchingCS::FPermutationDomain PermutationVector;
-			PermutationVector.Set<FPhysicsFieldRayMarchingCS::FFieldType>(bClearBackground ? 3 : FieldOutputType);
-			PermutationVector.Set<FPhysicsFieldRayMarchingCS::FEvalType>(GPhysicsFieldEvalType);
-			TShaderMapRef<FPhysicsFieldRayMarchingCS> ComputeShader(View.ShaderMap, PermutationVector);
+				const FIntVector DispatchCount = FIntVector::DivideAndRoundUp(FIntVector(OutputResolution.X, OutputResolution.Y, 1), FIntVector(8, 8, 1));
+				FComputeShaderUtils::AddPass(GraphBuilder, RDG_EVENT_NAME("PhysicsFieldRayMarching"), ComputeShader, Parameters, DispatchCount);
+			}
 
-			const FIntVector DispatchCount = (!bClearBackground && (FieldOutputType == EFieldOutputType::Field_Output_Vector)) ?
-				FIntVector::DivideAndRoundUp(BoundResolution, FIntVector(4, 4, 4)) :
-				FIntVector::DivideAndRoundUp(FIntVector(OutputResolution.X, OutputResolution.Y, 1), FIntVector(8, 8, 1));
+			{
+				FPhysicsFieldRayMarchingCS::FParameters* Parameters = CreateShaderParameters(GraphBuilder, View, PhysicsFieldResource, OutputTexture, TargetIndex);
 
-			FComputeShaderUtils::AddPass(GraphBuilder, RDG_EVENT_NAME("RenderPhysicsField"), ComputeShader, Parameters, DispatchCount);
-		}
-	};
-	FSceneTextureParameters SceneTextures = GetSceneTextureParameters(GraphBuilder, View);
+				const int32 NumBounds = PhysicsFieldResource->FieldInfos.BoundsOffsets[GPhysicsFieldTargetType + 1] - PhysicsFieldResource->FieldInfos.BoundsOffsets[GPhysicsFieldTargetType];
+				const FIntVector BoundResolution(FPhysicsFieldRayMarchingCS::VoxelResolution * NumBounds, FPhysicsFieldRayMarchingCS::VoxelResolution, FPhysicsFieldRayMarchingCS::VoxelResolution);
 
-	uint32 PrintOffset = 0;
-	if(GPhysicsFieldTargetType < EFieldPhysicsType::Field_PhysicsType_Max)
-	{
-		RenderPhysicsField(GPhysicsFieldTargetType, true, PrintOffset);
-		RenderPhysicsField(GPhysicsFieldTargetType, false, PrintOffset);
-	}
-	else
-	{
-		for (int32 TargetType = 0; TargetType < EFieldPhysicsType::Field_PhysicsType_Max; ++TargetType)
-		{
-			RenderPhysicsField(TargetType, true, PrintOffset);
-		}
+				FPhysicsFieldRayMarchingCS::FPermutationDomain PermutationVector;
+				PermutationVector.Set<FPhysicsFieldRayMarchingCS::FFieldType>(FieldOutputType);
+				PermutationVector.Set<FPhysicsFieldRayMarchingCS::FEvalType>(GPhysicsFieldEvalType);
+				TShaderMapRef<FPhysicsFieldRayMarchingCS> ComputeShader(View.ShaderMap, PermutationVector);
 
-		for (int32 TargetType = 0; TargetType < EFieldPhysicsType::Field_PhysicsType_Max; ++TargetType)
-		{
-			RenderPhysicsField(TargetType, false, PrintOffset);
+				const FIntVector DispatchCount = (FieldOutputType == EFieldOutputType::Field_Output_Vector) ?
+					FIntVector::DivideAndRoundUp(BoundResolution, FIntVector(4, 4, 4)) :
+					FIntVector::DivideAndRoundUp(FIntVector(OutputResolution.X, OutputResolution.Y, 1), FIntVector(8, 8, 1));
+				FComputeShaderUtils::AddPass(GraphBuilder, RDG_EVENT_NAME("PhysicsFieldRayMarching"), ComputeShader, Parameters, DispatchCount);
+			}
 		}
 	}
 }
@@ -195,12 +182,10 @@ void RenderPhysicsField(
 	if (PhysicsFieldProxy)
 	{
 		FPhysicsFieldResource* PhysicsFieldResource = (GPhysicsFieldSystemType == 0) ? PhysicsFieldProxy->DebugResource : PhysicsFieldProxy->FieldResource;
+
 		if (Views.Num() > 0 && PhysicsFieldResource && ShaderPrint::IsEnabled(Views[0].ShaderPrintData) && UE::PixelFormat::HasCapabilities(PF_FloatRGBA, EPixelFormatCapabilities::TypedUAVLoad))
 		{
-			if (ShaderPrint::IsSupported(Views[0].GetShaderPlatform()) && IsFeatureLevelSupported(Views[0].GetShaderPlatform(), ERHIFeatureLevel::SM5))
-			{
-				AddPhysicsFieldRayMarchingPass(GraphBuilder, Views[0], PhysicsFieldResource, SceneColorTexture);
-			}
+			AddPhysicsFieldRayMarchingPass(GraphBuilder, Views[0], PhysicsFieldResource, SceneColorTexture);
 		}
 	}
 }
