Fixed missing light list UB on mobile in some cases.
Made code common.

diff --git a/Engine/Source/Runtime/Renderer/Private/LightGridInjection.cpp b/Engine/Source/Runtime/Renderer/Private/LightGridInjection.cpp
--- a/Engine/Source/Runtime/Renderer/Private/LightGridInjection.cpp
+++ b/Engine/Source/Runtime/Renderer/Private/LightGridInjection.cpp
@@ -147,6 +147,15 @@ TRDGUniformBufferRef<FForwardLightData> CreateDummyForwardLightUniformBuffer(FRD
 	return GraphBuilder.CreateUniformBuffer(ForwardLightData);
 }
 
+void SetDummyForwardLightUniformBufferOnViews(FRDGBuilder& GraphBuilder, EShaderPlatform ShaderPlatform, TArray<FViewInfo>& Views)
+{
+	TRDGUniformBufferRef<FForwardLightData> ForwardLightUniformBuffer = CreateDummyForwardLightUniformBuffer(GraphBuilder, ShaderPlatform);
+	for (auto& View : Views)
+	{
+		View.ForwardLightingResources.SetUniformBuffer(ForwardLightUniformBuffer);
+	}
+}
+
 IMPLEMENT_GLOBAL_SHADER_PARAMETER_STRUCT(FForwardLightData, "ForwardLightData");
 
 FForwardLightData::FForwardLightData()
@@ -922,11 +931,7 @@ FComputeLightGridOutput FDeferredShadingSceneRenderer::GatherLightsAndComputeLig
 	
 	if (!bNeedLightGrid)
 	{
-		TRDGUniformBufferRef<FForwardLightData> ForwardLightUniformBuffer = CreateDummyForwardLightUniformBuffer(GraphBuilder, ShaderPlatform);
-		for (auto& View : Views)
-		{
-			View.ForwardLightingResources.SetUniformBuffer(ForwardLightUniformBuffer);
-		}
+		SetDummyForwardLightUniformBufferOnViews(GraphBuilder, ShaderPlatform, Views);
 		return Result;
 	}
 
diff --git a/Engine/Source/Runtime/Renderer/Private/MobileShadingRenderer.cpp b/Engine/Source/Runtime/Renderer/Private/MobileShadingRenderer.cpp
--- a/Engine/Source/Runtime/Renderer/Private/MobileShadingRenderer.cpp
+++ b/Engine/Source/Runtime/Renderer/Private/MobileShadingRenderer.cpp
@@ -975,6 +975,10 @@ void FMobileSceneRenderer::Render(FRDGBuilder& GraphBuilder)
 				GPUSortManager->OnPreRender(GraphBuilder);
 			}
 		}
+	else
+	{
+		SetDummyForwardLightUniformBufferOnViews(GraphBuilder, ShaderPlatform, Views);
+	}
 
 		// Generate the Sky/Atmosphere look up tables
 		const bool bShouldRenderSkyAtmosphere = ShouldRenderSkyAtmosphere(Scene, ViewFamily.EngineShowFlags);
diff --git a/Engine/Source/Runtime/Renderer/Private/SceneRendering.h b/Engine/Source/Runtime/Renderer/Private/SceneRendering.h
--- a/Engine/Source/Runtime/Renderer/Private/SceneRendering.h
+++ b/Engine/Source/Runtime/Renderer/Private/SceneRendering.h
@@ -759,6 +759,7 @@ BEGIN_GLOBAL_SHADER_PARAMETER_STRUCT_WITH_CONSTRUCTOR(FForwardLightData, )
 END_GLOBAL_SHADER_PARAMETER_STRUCT()
 
 extern TRDGUniformBufferRef<FForwardLightData> CreateDummyForwardLightUniformBuffer(FRDGBuilder& GraphBuilder, EShaderPlatform ShaderPlatform);
+extern void SetDummyForwardLightUniformBufferOnViews(FRDGBuilder& GraphBuilder, EShaderPlatform ShaderPlatform, TArray<FViewInfo>& Views);
 
 class FForwardLightingViewResources
 {
