VulkanRHI: Tweak texture name to be unique.

#rnx

diff --git a/Engine/Source/Runtime/Renderer/Private/LightRendering.cpp b/Engine/Source/Runtime/Renderer/Private/LightRendering.cpp
--- a/Engine/Source/Runtime/Renderer/Private/LightRendering.cpp
+++ b/Engine/Source/Runtime/Renderer/Private/LightRendering.cpp
@@ -1757,7 +1757,7 @@ void FDeferredShadingSceneRenderer::RenderLights(
 								PF_FloatRGBA,
 								FClearValueBinding::Black,
 								TexCreate_ShaderResource | TexCreate_RenderTargetable | TexCreate_UAV);
-							SubPixelRayTracingShadowMaskTexture = GraphBuilder.CreateTexture(Desc, TEXT("RayTracingOcclusion"));
+							SubPixelRayTracingShadowMaskTexture = GraphBuilder.CreateTexture(Desc, TEXT("SubPixelRayTracingOcclusion"));
 							SubPixelRayTracingShadowMaskUAV = GraphBuilder.CreateUAV(FRDGTextureUAVDesc(SubPixelRayTracingShadowMaskTexture));
 						}
 
