Fixed CS compilations issues with mobile GPUScene enabled

diff --git a/Engine/Shaders/Private/InstanceCulling/BuildInstanceDrawCommands.usf b/Engine/Shaders/Private/InstanceCulling/BuildInstanceDrawCommands.usf
--- a/Engine/Shaders/Private/InstanceCulling/BuildInstanceDrawCommands.usf
+++ b/Engine/Shaders/Private/InstanceCulling/BuildInstanceDrawCommands.usf
@@ -17,6 +17,11 @@
 	#define INSTANCE_DATA_STRIDE_ELEMENTS 1u
 #endif
 
+#ifndef VF_SUPPORTS_PRIMITIVE_SCENE_DATA
+	// Enable access to SceneData functionality for these compute shaders
+	#define VF_SUPPORTS_PRIMITIVE_SCENE_DATA 1u
+#endif
+
 // Disable the declaration of any templated types so we can use this shader without HLSL 2021. This is necessary because
 // cross-compilation creates bugs with ENABLE_INSTANCE_COMPACTION.
 #define ALLOW_TEMPLATES 0
diff --git a/Engine/Source/Runtime/Renderer/Private/InstanceCulling/InstanceCullingContext.cpp b/Engine/Source/Runtime/Renderer/Private/InstanceCulling/InstanceCullingContext.cpp
--- a/Engine/Source/Runtime/Renderer/Private/InstanceCulling/InstanceCullingContext.cpp
+++ b/Engine/Source/Runtime/Renderer/Private/InstanceCulling/InstanceCullingContext.cpp
@@ -513,7 +513,6 @@ public:
 		OutEnvironment.CompilerFlags.Add(CFLAG_PrecompileWithDXC);
 
 		OutEnvironment.SetDefine(TEXT("INDIRECT_ARGS_NUM_WORDS"), FInstanceCullingContext::IndirectArgsNumWords);
-		OutEnvironment.SetDefine(TEXT("VF_SUPPORTS_PRIMITIVE_SCENE_DATA"), 1);
 		OutEnvironment.SetDefine(TEXT("USE_GLOBAL_GPU_SCENE_DATA"), 1);
 		OutEnvironment.SetDefine(TEXT("USE_GLOBAL_GPU_LIGHTMAP_DATA"), 1);
 		OutEnvironment.SetDefine(TEXT("NANITE_MULTI_VIEW"), 1);
