Mobile PostProcess locallight

diff --git a/Engine/Shaders/Private/MobileLocalLightPrepass.usf b/Engine/Shaders/Private/MobileLocalLightPrepass.usf
deleted file mode 100644
diff --git a/Engine/Shaders/Private/MobileLocalLightsBuffer.usf b/Engine/Shaders/Private/MobileLocalLightsBuffer.usf
new file mode 100644
--- /dev/null
+++ b/Engine/Shaders/Private/MobileLocalLightsBuffer.usf
@@ -0,0 +1,171 @@
+// Copyright Epic Games, Inc. All Rights Reserved.
+
+#define SUPPORT_CONTACT_SHADOWS 0
+
+#include "LightGridCommon.ush"
+#if !COMPUTE_SHADER
+#include "Common.ush"
+#include "SceneTextureParameters.ush"
+#include "DeferredLightingCommon.ush"
+#endif
+
+#if COMPUTE_SHADER
+
+uint2 GroupSize;
+RWBuffer<int> RWTileInfo;
+[numthreads(THREADGROUP_SIZEX, 1, 1)]
+void MainCS(
+	uint3 GroupId : SV_GroupID,
+	uint3 DispatchThreadId : SV_DispatchThreadID,
+	uint3 GroupThreadId : SV_GroupThreadID)
+{
+	if(DispatchThreadId.x > (GroupSize.x * GroupSize.y))
+	{
+		return;
+	}
+
+	const uint GridX = DispatchThreadId.x % GroupSize.x;
+	const uint GridY = DispatchThreadId.x / GroupSize.x;
+	const uint EyeIndex = 0;
+	const FLightGridData GridData = GetLightGridData(EyeIndex);
+	uint LocalNumLights = 0;
+	LOOP
+	for (int SliceIt = 0; SliceIt < GridData.CulledGridSize.z; SliceIt++)
+	{
+		const uint GridIndex = ComputeLightGridCellIndex(uint3(GridX, GridY, SliceIt), EyeIndex);
+		const FCulledLightsGridData CulledLightsGrid = GetCulledLightsGrid(GridIndex, EyeIndex);
+		
+		LocalNumLights = CulledLightsGrid.NumLocalLights;
+
+		if(LocalNumLights > 0)
+		{
+			break;
+		}
+	}
+
+	if(LocalNumLights > 0)
+	{
+		RWTileInfo[2 * DispatchThreadId.x] = GridX;
+		RWTileInfo[2 * DispatchThreadId.x + 1] = GridY;
+	}
+	else
+	{
+		RWTileInfo[2 * DispatchThreadId.x] = -1;
+		RWTileInfo[2 * DispatchThreadId.x + 1] = -1;
+	}
+}
+
+#else
+
+int LightGridPixelSize;
+Buffer<int> TileInfo;
+void MainVS(
+	float2 TexCoord : ATTRIBUTE0,
+	uint VertexId : SV_VertexID,
+	uint InstanceId : SV_InstanceID,
+	out float4 OutPosition : SV_POSITION)
+{
+	float YPos = TileInfo[2 * InstanceId + 1] + TexCoord.y;
+	float XPos = TileInfo[2 * InstanceId] + TexCoord.x;
+
+	float2 ScreenUV = float2(XPos, YPos) * LightGridPixelSize * View.BufferSizeAndInvSize.zw;
+	float2 ScreenPosition = (ScreenUV.xy - View.ScreenPositionScaleBias.wz) / View.ScreenPositionScaleBias.xy;
+	OutPosition = float4(ScreenPosition, 0, 1);
+
+	if(TileInfo[2 * InstanceId] < 0)
+	{
+		OutPosition.xy = 0;
+	}
+}
+
+void Main(float4 SvPosition : SV_POSITION
+	, out float3 OutColorA : SV_Target0
+#if POST_PROCESS_LOCAL_LIGHTS == 0
+	, out float4 OutColorB : SV_Target1
+#endif
+)
+{
+	const uint2 PixelPos = SvPosition.xy;
+	const float2 ScreenPosition = SvPositionToScreenPosition(SvPosition).xy;
+	const uint EyeIndex = 0;
+
+	float2 ScreenUV = ScreenPosition * View.ScreenPositionScaleBias.xy + View.ScreenPositionScaleBias.wz;
+	float SceneDepth = CalcSceneDepth(ScreenUV);
+
+	float3 TranslatedWorldPosition = mul(float4(ScreenPosition * SceneDepth, SceneDepth, 1), PrimaryView.ScreenToTranslatedWorld).xyz;
+	float3 CameraVector = normalize(TranslatedWorldPosition);
+
+	uint GridIndex = ComputeLightGridCellIndex(uint2(PixelPos.x, PixelPos.y), SceneDepth);
+	const FCulledLightsGridData CulledLightGridData = GetCulledLightsGrid(GridIndex, EyeIndex);
+
+	const uint NumLocalLights = min(CulledLightGridData.NumLocalLights, GetNumLocalLights(EyeIndex));
+
+	float3 Color = uint3(0.f, 0.f , 0.f);
+	float3 SavedL = float3(0,0,0);
+	float TotalWeight = 0;
+
+	LOOP
+	for (uint LocalLightListIndex = 0; LocalLightListIndex < NumLocalLights; LocalLightListIndex++)
+	{
+		const FLocalLightData LocalLight = GetLocalLightData(CulledLightGridData.DataStartIndex + LocalLightListIndex, EyeIndex);
+
+		// extra-early out since we know light grid is sloppy and all lights in list are radial (have a range)
+		// appears useless
+		float InvLightRadiusSq = LocalLight.LightPositionAndInvRadius.w * LocalLight.LightPositionAndInvRadius.w;
+		float DistLight = length2(TranslatedWorldPosition - LocalLight.LightPositionAndInvRadius.xyz) * InvLightRadiusSq;
+		if (DistLight > 1.0f)
+		{
+			continue;
+		}
+		FDeferredLightData LightData = ConvertToDeferredLight_Mobile(LocalLight);
+
+		float3 L = LightData.Direction;	// Already normalized
+		float3 ToLight = L;
+		float3 MaskedLightColor = LightData.Color;
+		float LightMask = 1;
+		if (LightData.bRadialLight)
+		{
+			LightMask = GetLocalLightAttenuation( TranslatedWorldPosition, LightData, ToLight, L );
+			MaskedLightColor *= LightMask;
+		}
+
+		float WeightLight = 1 - DistLight;
+		
+		if( LightMask > 0)
+		{
+			TotalWeight += WeightLight;
+			Color += MaskedLightColor;
+			SavedL += L *  WeightLight;
+		}
+	}
+
+#if POST_PROCESS_LOCAL_LIGHTS == 0
+	if (TotalWeight > 0)
+	{
+		SavedL = (SavedL / TotalWeight);
+		OutColorA = float3(Color.x, Color.y, Color.z);
+		OutColorB = 0.5f * float4(SavedL.x, SavedL.y, SavedL.z, 0)  + float4(0.5f, 0.5f, 0.5f, 0.5f);
+	}
+	else
+	{
+		OutColorA = float3(0, 0, 0);
+		OutColorB = float4(0.5,0.5,0.5, 0.5);
+	}
+#else
+	if (TotalWeight > 0)
+	{
+		SavedL = (SavedL / TotalWeight);
+		float3 N = normalize(cross(ddx(TranslatedWorldPosition.xyz), ddy(TranslatedWorldPosition.xyz)));
+		float3 L = SavedL;	// Already normalized
+		float NoL = max(0, dot(N, L));
+		float3 Diffuse =  (1 / PI) * NoL * Color;
+		OutColorA = 1 + Diffuse;
+	}
+	else
+	{
+		OutColorA = float3(1, 1, 1);
+	}
+#endif
+}
+
+#endif
\ No newline at end of file
diff --git a/Engine/Source/Editor/MaterialEditor/Private/MaterialStatsCommon.cpp b/Engine/Source/Editor/MaterialEditor/Private/MaterialStatsCommon.cpp
--- a/Engine/Source/Editor/MaterialEditor/Private/MaterialStatsCommon.cpp
+++ b/Engine/Source/Editor/MaterialEditor/Private/MaterialStatsCommon.cpp
@@ -447,7 +447,7 @@ void FMaterialStatsUtils::GetRepresentativeShaderTypesAndDescriptions(TMap<FName
 						}
 
 						{
-							MobileBasePassShaderName(false, TEXT("FMobileDistanceFieldShadowsAndLQLightMapPolicy"), TEXT("LOCAL_LIGHTS_PREPASS_ENABLED"), bMobileHDR, bOnlySkyPermutation, ShaderNameStr);
+							MobileBasePassShaderName(false, TEXT("FMobileDistanceFieldShadowsAndLQLightMapPolicy"), TEXT("LOCAL_LIGHTS_BUFFER"), bMobileHDR, bOnlySkyPermutation, ShaderNameStr);
 							const FString Description = FString::Printf(TEXT("Mobile base pass shader with distance field shadows, CSM and local light(s) %s"), DescSuffix);
 
 							FRepresentativeShaderInfo ShaderInfo = FRepresentativeShaderInfo(ERepresentativeShader::StationarySurfaceNPointLights, FName(ShaderNameStr), Description);
@@ -474,7 +474,7 @@ void FMaterialStatsUtils::GetRepresentativeShaderTypesAndDescriptions(TMap<FName
 					}
 
 					{
-						MobileBasePassShaderName(false, TEXT("TLightMapPolicyLQ"), TEXT("LOCAL_LIGHTS_PREPASS_ENABLED"), bMobileHDR, bOnlySkyPermutation, ShaderNameStr);
+						MobileBasePassShaderName(false, TEXT("TLightMapPolicyLQ"), TEXT("LOCAL_LIGHTS_BUFFER"), bMobileHDR, bOnlySkyPermutation, ShaderNameStr);
 						const FString Description = FString::Printf(TEXT("Mobile base pass shader with static lighting and local light(s) %s"), DescSuffix);
 						FRepresentativeShaderInfo ShaderInfo = FRepresentativeShaderInfo(ERepresentativeShader::StationarySurfaceNPointLights, FName(ShaderNameStr), Description);
 						ShaderTypeNamesAndDescriptions.FindOrAdd(FLocalVertexFactoryName).Add(ShaderInfo);
diff --git a/Engine/Source/Runtime/Core/Private/HAL/ConsoleManager.cpp b/Engine/Source/Runtime/Core/Private/HAL/ConsoleManager.cpp
--- a/Engine/Source/Runtime/Core/Private/HAL/ConsoleManager.cpp
+++ b/Engine/Source/Runtime/Core/Private/HAL/ConsoleManager.cpp
@@ -2847,7 +2847,7 @@ static TAutoConsoleVariable<int32> CVarMobileForwardEnableLocalLights(
 	1,
 	TEXT("0: Local Lights Disabled (default)\n"
 		"1: Local Lights Enabled\n"
-		"2: Local Lights Prepass Enabled\n"),
+		"2: Local Lights Buffer Enabled\n"),
 	ECVF_ReadOnly | ECVF_RenderThreadSafe);
 
 static TAutoConsoleVariable<int32> CVarMobileForwardEnableClusteredReflections(
diff --git a/Engine/Source/Runtime/Renderer/Private/MobileBasePass.cpp b/Engine/Source/Runtime/Renderer/Private/MobileBasePass.cpp
--- a/Engine/Source/Runtime/Renderer/Private/MobileBasePass.cpp
+++ b/Engine/Source/Runtime/Renderer/Private/MobileBasePass.cpp
@@ -15,12 +15,45 @@
 #include "Engine/TextureCube.h"
 #include "ShaderPlatformCachedIniValue.h"
 
-bool MobileForwardEnablePrepassLocalLights(const FStaticShaderPlatform Platform)
+bool MobileLocalLightsBufferEnabled(const FStaticShaderPlatform Platform)
 {
 	static FShaderPlatformCachedIniValue<int32> MobileForwardEnablePrepassLocalLightsIniValue(TEXT("r.Mobile.Forward.EnableLocalLights"));
 	return MobileForwardEnablePrepassLocalLightsIniValue.Get(Platform) == 2;
 }
 
+bool MobileLocalLightsBufferPrepassEnabled(const FStaticShaderPlatform Platform)
+{
+	return MobileLocalLightsBufferEnabled(Platform) && MobileUsesFullDepthPrepass(Platform);
+}
+
+bool MobileLocalLightsBufferPostprocessEnabled(const FStaticShaderPlatform Platform)
+{
+	return MobileLocalLightsBufferEnabled(Platform) && !MobileUsesFullDepthPrepass(Platform);
+}
+
+EMobileLocalLightSetting GetMobileForwardLocalLightSetting(EShaderPlatform ShaderPlatform, bool bIsTranslucent)
+{
+	static FShaderPlatformCachedIniValue<int32> MobileForwardLocalLightsIniValueCvar(TEXT("r.Mobile.Forward.EnableLocalLights"));
+	int MobileForwardLocalLightsIniValue = MobileForwardLocalLightsIniValueCvar.Get(ShaderPlatform);
+
+	if (MobileForwardLocalLightsIniValue > 0)
+	{
+		if ((MobileForwardLocalLightsIniValue == 1) || bIsTranslucent)
+		{
+			return EMobileLocalLightSetting::LOCAL_LIGHTS_ENABLED;
+		}
+		else if (MobileForwardLocalLightsIniValue == 2)
+		{
+			if (MobileUsesFullDepthPrepass(ShaderPlatform))
+			{
+				return EMobileLocalLightSetting::LOCAL_LIGHTS_BUFFER;
+			}
+		}
+	}
+
+	return EMobileLocalLightSetting::LOCAL_LIGHTS_DISABLED;
+}
+
 uint8 GetMobileShadingModelStencilValue(FMaterialShadingModelField ShadingModel)
 {
 	if (ShadingModel.HasOnlyShadingModel(MSM_DefaultLit))
@@ -184,9 +217,9 @@ bool MobileBasePass::GetShaders(
 				);
 			break;
 		}
-		case EMobileLocalLightSetting::LOCAL_LIGHTS_PREPASS_ENABLED:
+		case EMobileLocalLightSetting::LOCAL_LIGHTS_BUFFER:
 		{
-			return GetMobileBasePassShaders<EMobileLocalLightSetting::LOCAL_LIGHTS_PREPASS_ENABLED>(
+			return GetMobileBasePassShaders<EMobileLocalLightSetting::LOCAL_LIGHTS_BUFFER>(
 				LightMapPolicyType,
 				MaterialResource,
 				VertexFactoryType,
@@ -841,14 +874,9 @@ bool FMobileBasePassMeshProcessor::Process(
 	EMobileLocalLightSetting LocalLightSetting = EMobileLocalLightSetting::LOCAL_LIGHTS_DISABLED;
 	if (Scene && PrimitiveSceneProxy && ShadingModels.IsLit())
 	{
-		if (!bPassUsesDeferredShading && MobileForwardEnableLocalLights(Scene->GetShaderPlatform()))
+		if (!bPassUsesDeferredShading && (PrimitiveSceneProxy->GetPrimitiveSceneInfo()->NumMobileDynamicLocalLights > 0))
 		{
-			if ((PrimitiveSceneProxy->GetPrimitiveSceneInfo()->NumMobileDynamicLocalLights > 0))
-			{
-				LocalLightSetting = MobileForwardEnablePrepassLocalLights(Scene->GetShaderPlatform()) ?
-					EMobileLocalLightSetting::LOCAL_LIGHTS_PREPASS_ENABLED :
-					EMobileLocalLightSetting::LOCAL_LIGHTS_ENABLED;
-			}
+			LocalLightSetting = GetMobileForwardLocalLightSetting(Scene->GetShaderPlatform(), bIsTranslucent);
 		}
 	}
 
@@ -1035,8 +1063,13 @@ void FMobileBasePassMeshProcessor::CollectPSOInitializers(const FSceneTexturesCo
 		MobileBasePass::SetOpaqueRenderState(DrawRenderState, nullptr, Material, ShadingModels, true, bPassUsesDeferredShading);
 	}
 
-	const bool bUseLocalLightPermutation = bLitMaterial && (!bPassUsesDeferredShading && (MobileForwardEnableLocalLights(ShaderPlatform)));
-	EMobileLocalLightSetting LocalLightSetting = MobileForwardEnablePrepassLocalLights(ShaderPlatform) ? EMobileLocalLightSetting::LOCAL_LIGHTS_PREPASS_ENABLED: EMobileLocalLightSetting::LOCAL_LIGHTS_ENABLED;
+	EMobileLocalLightSetting LocalLightSetting = EMobileLocalLightSetting::LOCAL_LIGHTS_DISABLED;
+	if (bLitMaterial && !bPassUsesDeferredShading)
+	{
+		LocalLightSetting = GetMobileForwardLocalLightSetting(ShaderPlatform, bTranslucentBasePass);
+	}
+	const bool bUseLocalLightPermutation = (LocalLightSetting != EMobileLocalLightSetting::LOCAL_LIGHTS_DISABLED);
+
 	const bool bCanReceiveCSM = ((Flags & FMobileBasePassMeshProcessor::EFlags::CanReceiveCSM) == FMobileBasePassMeshProcessor::EFlags::CanReceiveCSM);
 
 	FMobileLightMapPolicyTypeList UniformLightMapPolicyTypes = GetUniformLightMapPolicyTypeForPSOCollection(bLitMaterial, bTranslucentBasePass, bPassUsesDeferredShading, bCanReceiveCSM, bMovable);
diff --git a/Engine/Source/Runtime/Renderer/Private/MobileBasePassRendering.cpp b/Engine/Source/Runtime/Renderer/Private/MobileBasePassRendering.cpp
--- a/Engine/Source/Runtime/Renderer/Private/MobileBasePassRendering.cpp
+++ b/Engine/Source/Runtime/Renderer/Private/MobileBasePassRendering.cpp
@@ -96,7 +96,7 @@ bool ShouldCacheShaderByPlatformAndOutputFormat(EShaderPlatform Platform, EOutpu
 	IMPLEMENT_MOBILE_SHADING_BASEPASS_LIGHTMAPPED_VERTEX_SHADER_TYPE(LightMapPolicyType, LightMapPolicyName) \
 	IMPLEMENT_MOBILE_SHADING_BASEPASS_LIGHTMAPPED_PIXEL_SHADER_TYPE(LightMapPolicyType, LightMapPolicyName, LOCAL_LIGHTS_DISABLED) \
 	IMPLEMENT_MOBILE_SHADING_BASEPASS_LIGHTMAPPED_PIXEL_SHADER_TYPE(LightMapPolicyType, LightMapPolicyName, LOCAL_LIGHTS_ENABLED) \
-	IMPLEMENT_MOBILE_SHADING_BASEPASS_LIGHTMAPPED_PIXEL_SHADER_TYPE(LightMapPolicyType, LightMapPolicyName, LOCAL_LIGHTS_PREPASS_ENABLED)
+	IMPLEMENT_MOBILE_SHADING_BASEPASS_LIGHTMAPPED_PIXEL_SHADER_TYPE(LightMapPolicyType, LightMapPolicyName, LOCAL_LIGHTS_BUFFER)
 
 // Implement shader types per lightmap policy 
 IMPLEMENT_MOBILE_SHADING_BASEPASS_LIGHTMAPPED_SHADER_TYPE(TUniformLightMapPolicy<LMP_NO_LIGHTMAP>, FNoLightMapPolicy);
diff --git a/Engine/Source/Runtime/Renderer/Private/MobileBasePassRendering.h b/Engine/Source/Runtime/Renderer/Private/MobileBasePassRendering.h
--- a/Engine/Source/Runtime/Renderer/Private/MobileBasePassRendering.h
+++ b/Engine/Source/Runtime/Renderer/Private/MobileBasePassRendering.h
@@ -26,7 +26,9 @@
 #include "RenderUtils.h"
 #include "DebugViewModeRendering.h"
 
-bool MobileForwardEnablePrepassLocalLights(const FStaticShaderPlatform Platform);
+bool MobileLocalLightsBufferEnabled(const FStaticShaderPlatform Platform);
+bool MobileLocalLightsBufferPrepassEnabled(const FStaticShaderPlatform Platform);
+bool MobileLocalLightsBufferPostprocessEnabled(const FStaticShaderPlatform Platform);
 
 struct FMobileBasePassTextures
 {
@@ -101,8 +103,9 @@ enum EMobileLocalLightSetting
 {
 	LOCAL_LIGHTS_DISABLED,
 	LOCAL_LIGHTS_ENABLED,
-	LOCAL_LIGHTS_PREPASS_ENABLED
+	LOCAL_LIGHTS_BUFFER
 };
+EMobileLocalLightSetting GetMobileForwardLocalLightSetting(EShaderPlatform ShaderPlatform, bool bIsTranslucent);
 
 bool ShouldCacheShaderByPlatformAndOutputFormat(EShaderPlatform Platform, EOutputFormat OutputFormat);
 // shared defines for mobile base pass VS and PS
@@ -376,23 +379,25 @@ public:
 		}
 		
 		const bool bDeferredShadingEnabled = IsMobileDeferredShadingEnabled(Parameters.Platform);
-		const bool bMaterialUsesForwardShading = bIsLit && 
-			(IsTranslucentBlendMode(Parameters.MaterialParameters) || Parameters.MaterialParameters.ShadingModels.HasShadingModel(MSM_SingleLayerWater));
+		const bool bIsTranslucent = IsTranslucentBlendMode(Parameters.MaterialParameters) || Parameters.MaterialParameters.ShadingModels.HasShadingModel(MSM_SingleLayerWater);
+		const bool bMaterialUsesForwardShading = bIsLit && bIsTranslucent;
 		// Translucent materials always support clustered shading on mobile deferred
 		const bool bForwardShading = !bDeferredShadingEnabled || bMaterialUsesForwardShading;
-		const bool bSupportsLocalLights = bForwardShading && MobileForwardEnableLocalLights(Parameters.Platform);
 
+		EMobileLocalLightSetting SupportedLocalLightsType = EMobileLocalLightSetting::LOCAL_LIGHTS_DISABLED;
+		if (bForwardShading && bIsLit)
+		{
+			SupportedLocalLightsType = GetMobileForwardLocalLightSetting(Parameters.Platform, bIsTranslucent);
+		}
 		// Deferred shading does not need SkyLight and LocalLight permutations
 		// TODO: skip skylight permutations for deferred	
 		bool bEnableLocalLights = LocalLightSetting != EMobileLocalLightSetting::LOCAL_LIGHTS_DISABLED;
-		const bool bShouldCacheByShading = (bForwardShading || !bEnableLocalLights);
-		const bool bShouldCacheByLocalLights = !bEnableLocalLights || (bIsLit && bEnableLocalLights == bSupportsLocalLights);
+		const bool bShouldCacheByLocalLights = !bEnableLocalLights || (bIsLit && (SupportedLocalLightsType == LocalLightSetting));
 
 		return TMobileBasePassPSBaseType<LightMapPolicyType>::ShouldCompilePermutation(Parameters) && 
 				ShouldCacheShaderByPlatformAndOutputFormat(Parameters.Platform, OutputFormat) && 
 				bShouldCacheBySkylight && 
-				bShouldCacheByLocalLights &&
-				bShouldCacheByShading;
+				bShouldCacheByLocalLights;
 	}
 	
 	static void ModifyCompilationEnvironment(const FMaterialShaderPermutationParameters& Parameters, FShaderCompilerEnvironment& OutEnvironment)
@@ -411,7 +416,7 @@ public:
 		OutEnvironment.SetDefine(TEXT("ENABLE_CLUSTERED_LIGHTS"), (LocalLightSetting == EMobileLocalLightSetting::LOCAL_LIGHTS_ENABLED) ? 1u : 0u);
 
 		// Translucent materials don't write to depth so cannot use prepass
-		OutEnvironment.SetDefine(TEXT("PREPASS_LOCAL_LIGHTS_MOBILE"), !bTranslucentMaterial && (LocalLightSetting == EMobileLocalLightSetting::LOCAL_LIGHTS_PREPASS_ENABLED) ? 1u: 0u);
+		OutEnvironment.SetDefine(TEXT("PREPASS_LOCAL_LIGHTS_MOBILE"), (LocalLightSetting == EMobileLocalLightSetting::LOCAL_LIGHTS_BUFFER) ? 1u: 0u);
 		OutEnvironment.SetDefine(TEXT("ENABLE_CLUSTERED_REFLECTION"), bEnableClusteredReflections ? 1u : 0u);
 		OutEnvironment.SetDefine(TEXT("USE_SHADOWMASKTEXTURE"), bMobileUsesShadowMaskTexture && !bTranslucentMaterial ? 1u : 0u);
 	}
diff --git a/Engine/Source/Runtime/Renderer/Private/MobileLocalLightPrepass.cpp b/Engine/Source/Runtime/Renderer/Private/MobileLocalLightPrepass.cpp
deleted file mode 100644
diff --git a/Engine/Source/Runtime/Renderer/Private/MobileLocalLightsBuffer.cpp b/Engine/Source/Runtime/Renderer/Private/MobileLocalLightsBuffer.cpp
new file mode 100644
--- /dev/null
+++ b/Engine/Source/Runtime/Renderer/Private/MobileLocalLightsBuffer.cpp
@@ -0,0 +1,222 @@
+// Copyright Epic Games, Inc. All Rights Reserved.
+
+#include "SceneRendering.h"
+#include "DataDrivenShaderPlatformInfo.h"
+#include "BasePassRendering.h"
+#include "PixelShaderUtils.h"
+#include "MobileBasePassRendering.h"
+#include "RendererPrivateUtils.h"
+#include "GlobalRenderResources.h"
+#include "ScenePrivate.h"
+
+bool MobileLocalLighsBufferEnabled(const FStaticShaderPlatform Platform)
+{
+	return !IsMobileDeferredShadingEnabled(Platform) && 
+			IsMobilePlatform(Platform) && 
+			MobileLocalLightsBufferEnabled(Platform);
+}
+
+const int32 GLocalLightPrepassTileSizeX = 8;
+class FLocalLightBufferCS : public FGlobalShader
+{
+	DECLARE_GLOBAL_SHADER(FLocalLightBufferCS);
+public:
+	SHADER_USE_PARAMETER_STRUCT(FLocalLightBufferCS, FGlobalShader);
+
+	BEGIN_SHADER_PARAMETER_STRUCT(FParameters, )
+		SHADER_PARAMETER_RDG_BUFFER_UAV(RWBuffer<int32>, RWTileInfo)
+		SHADER_PARAMETER_RDG_UNIFORM_BUFFER(FForwardLightData, ForwardLightData)
+		SHADER_PARAMETER(FIntPoint, GroupSize)
+	END_SHADER_PARAMETER_STRUCT()
+ 
+	static bool ShouldCompilePermutation(const FGlobalShaderPermutationParameters& Parameters)
+	{
+		return MobileLocalLighsBufferEnabled(Parameters.Platform);
+	}
+
+	static void ModifyCompilationEnvironment(const FGlobalShaderPermutationParameters& Parameters, FShaderCompilerEnvironment& OutEnvironment)
+	{
+		FForwardLightingParameters::ModifyCompilationEnvironment(Parameters.Platform, OutEnvironment);
+		OutEnvironment.SetDefine(TEXT("THREADGROUP_SIZEX"), GLocalLightPrepassTileSizeX);
+		OutEnvironment.SetDefine(TEXT("COMPUTE_SHADER"), 1);
+	}
+};
+
+IMPLEMENT_GLOBAL_SHADER(FLocalLightBufferCS, "/Engine/Private/MobileLocalLightsBuffer.usf", "MainCS", SF_Compute);
+
+class FLocalLightBufferVS : public FGlobalShader
+{
+	DECLARE_GLOBAL_SHADER(FLocalLightBufferVS);
+	SHADER_USE_PARAMETER_STRUCT(FLocalLightBufferVS, FGlobalShader);
+
+	BEGIN_SHADER_PARAMETER_STRUCT(FParameters, )
+		SHADER_PARAMETER_STRUCT_REF(FViewUniformShaderParameters, View)
+		SHADER_PARAMETER_RDG_BUFFER_SRV(Buffer<int32>, TileInfo)
+		SHADER_PARAMETER(int32, LightGridPixelSize)
+	END_SHADER_PARAMETER_STRUCT()
+
+	static bool ShouldCompilePermutation(const FGlobalShaderPermutationParameters& Parameters)
+	{
+		return MobileLocalLighsBufferEnabled(Parameters.Platform);
+	}
+};
+
+IMPLEMENT_GLOBAL_SHADER(FLocalLightBufferVS, "/Engine/Private/MobileLocalLightsBuffer.usf", "MainVS", SF_Vertex);
+
+class FLocalLightBufferPS : public FGlobalShader
+{
+	DECLARE_GLOBAL_SHADER(FLocalLightBufferPS);
+	SHADER_USE_PARAMETER_STRUCT(FLocalLightBufferPS, FGlobalShader)
+
+	BEGIN_SHADER_PARAMETER_STRUCT(FParameters, )
+		SHADER_PARAMETER_STRUCT_REF(FViewUniformShaderParameters, View)
+		SHADER_PARAMETER_RDG_UNIFORM_BUFFER(FForwardLightData, ForwardLightData)
+	END_SHADER_PARAMETER_STRUCT()
+
+	static bool ShouldCompilePermutation(const FGlobalShaderPermutationParameters& Parameters)
+	{
+		return MobileLocalLighsBufferEnabled(Parameters.Platform);
+	}
+
+	static void ModifyCompilationEnvironment(const FGlobalShaderPermutationParameters& Parameters, FShaderCompilerEnvironment& OutEnvironment)
+	{
+		FForwardLightingParameters::ModifyCompilationEnvironment(Parameters.Platform, OutEnvironment);
+
+		if (MobileLocalLightsBufferPrepassEnabled(Parameters.Platform))
+		{
+			OutEnvironment.SetRenderTargetOutputFormat(0, PF_FloatR11G11B10);
+			OutEnvironment.SetRenderTargetOutputFormat(1, PF_A2B10G10R10);
+			OutEnvironment.SetDefine(TEXT("POST_PROCESS_LOCAL_LIGHTS"), 0);
+		}
+		else
+		{
+			OutEnvironment.SetRenderTargetOutputFormat(0, PF_FloatR11G11B10);
+			OutEnvironment.SetDefine(TEXT("POST_PROCESS_LOCAL_LIGHTS"), 1);
+		}
+	}
+};
+
+IMPLEMENT_GLOBAL_SHADER(FLocalLightBufferPS, "/Engine/Private/MobileLocalLightsBuffer.usf", "Main", SF_Pixel);
+
+BEGIN_SHADER_PARAMETER_STRUCT(FLocalLightBufferPrepassParameters, )
+SHADER_PARAMETER_STRUCT_INCLUDE(FLocalLightBufferVS::FParameters, VS)
+SHADER_PARAMETER_STRUCT_INCLUDE(FLocalLightBufferPS::FParameters, PS)
+SHADER_PARAMETER_STRUCT_INCLUDE(FSceneTextureShaderParameters, SceneTextures)
+RENDER_TARGET_BINDING_SLOTS()
+END_SHADER_PARAMETER_STRUCT()
+
+void FMobileSceneRenderer::RenderMobileLocalLightsBuffer(FRDGBuilder& GraphBuilder, FSceneTextures& SceneTextures, bool bIsPrepass)
+{
+	if (!MobileLocalLighsBufferEnabled(ShaderPlatform) || 
+		(bIsPrepass != MobileLocalLightsBufferPrepassEnabled(ShaderPlatform)) || 
+		IsMobileDeferredShadingEnabled(ShaderPlatform))
+	{
+		return;
+	}
+
+	RDG_EVENT_SCOPE(GraphBuilder, "RenderMobileLocalLightsBuffer");
+	QUICK_SCOPE_CYCLE_COUNTER(STAT_RenderMobileLocalLightsBuffer);
+
+	static const auto LightGridPixelSizeCVar = IConsoleManager::Get().FindConsoleVariable(TEXT("r.Forward.LightGridPixelSize"));
+	check(LightGridPixelSizeCVar != nullptr);
+	int32 LightGridPixelSize = LightGridPixelSizeCVar->GetInt();
+
+	for (int32 ViewIndex = 0; ViewIndex < Views.Num(); ViewIndex++)
+	{
+		FViewInfo& View = Views[ViewIndex];
+		bool bHasNoLocalLights = (!View.ForwardLightingResources.ForwardLightData) || (View.ForwardLightingResources.ForwardLightData->NumLocalLights == 0);
+		if (!View.ShouldRenderView() || bHasNoLocalLights)
+		{
+			continue;
+		}
+
+		const FIntPoint GroupSize(
+			FMath::DivideAndRoundUp(View.ViewRect.Size().X, LightGridPixelSize),
+			FMath::DivideAndRoundUp(View.ViewRect.Size().Y, LightGridPixelSize));
+
+		FRDGBufferRef TileInfoBuffer = GraphBuilder.CreateBuffer(FRDGBufferDesc::CreateBufferDesc(sizeof(int32), 2 * GroupSize.X * GroupSize.Y), TEXT("TileInfoBuffer"));
+		FRDGBufferUAVRef TileInfoBufferUAV = GraphBuilder.CreateUAV(TileInfoBuffer, PF_R32_SINT);
+		FRDGBufferSRVRef TileInfoBufferSRV = GraphBuilder.CreateSRV(TileInfoBuffer, PF_R32_SINT);
+
+		{
+			auto* PassParameters = GraphBuilder.AllocParameters<FLocalLightBufferCS::FParameters>();
+			PassParameters->RWTileInfo = TileInfoBufferUAV;
+			PassParameters->ForwardLightData = View.ForwardLightingResources.ForwardLightUniformBuffer;
+			PassParameters->GroupSize = GroupSize;
+			auto ComputeShader = View.ShaderMap->GetShader<FLocalLightBufferCS>();
+
+			FComputeShaderUtils::AddPass(
+				GraphBuilder,
+				RDG_EVENT_NAME("RenderMobileLocalLights_TiledInfoCS"),
+				ERDGPassFlags::Compute,
+				ComputeShader,
+				PassParameters,
+				FIntVector(FMath::DivideAndRoundUp<uint32>(GroupSize.Y * GroupSize.X, GLocalLightPrepassTileSizeX), 1, 1));
+		}
+
+
+		{
+			FLocalLightBufferPrepassParameters* PassParameters = GraphBuilder.AllocParameters<FLocalLightBufferPrepassParameters>();
+			if (bIsPrepass)
+			{
+				PassParameters->RenderTargets[0] = FRenderTargetBinding(SceneTextures.MobileLocalLightTextureA, ERenderTargetLoadAction::EClear);
+				PassParameters->RenderTargets[1] = FRenderTargetBinding(SceneTextures.MobileLocalLightTextureB, ERenderTargetLoadAction::EClear);
+			}
+			else
+			{
+				PassParameters->RenderTargets[0] = FRenderTargetBinding(SceneTextures.Color.Resolve, ERenderTargetLoadAction::ELoad);
+			}
+			PassParameters->SceneTextures = SceneTextures.GetSceneTextureShaderParameters(View.FeatureLevel);
+
+			PassParameters->VS.View = GetShaderBinding(View.ViewUniformBuffer);
+			PassParameters->VS.TileInfo = TileInfoBufferSRV;
+			PassParameters->VS.LightGridPixelSize = LightGridPixelSize;
+			PassParameters->PS.ForwardLightData = View.ForwardLightingResources.ForwardLightUniformBuffer;
+			PassParameters->PS.View = GetShaderBinding(View.ViewUniformBuffer);
+
+			auto VertexShader = View.ShaderMap->GetShader<FLocalLightBufferVS>();
+			auto PixelShader = View.ShaderMap->GetShader<FLocalLightBufferPS>();
+
+			GraphBuilder.AddPass(
+				RDG_EVENT_NAME("RenderMobileLocalLightsBuffer %s", bIsPrepass ? TEXT("Prepass") : TEXT("PostProcess")),
+				PassParameters,
+				ERDGPassFlags::Raster,
+				[PassParameters, VertexShader, PixelShader, &View, GroupSize, bIsPrepass](FRHICommandList& RHICmdList)
+				{
+					FGraphicsPipelineStateInitializer GraphicsPSOInit;
+					RHICmdList.ApplyCachedRenderTargets(GraphicsPSOInit);
+
+					RHICmdList.SetViewport(View.ViewRect.Min.X, View.ViewRect.Min.Y, 0.0f, View.ViewRect.Max.X, View.ViewRect.Max.Y, 1.0f);
+					GraphicsPSOInit.RasterizerState = TStaticRasterizerState<FM_Solid, CM_None>::GetRHI();
+					GraphicsPSOInit.DepthStencilState = TStaticDepthStencilState<false, CF_Always>::GetRHI();
+					if (bIsPrepass)
+					{
+						GraphicsPSOInit.BlendState = TStaticBlendState<>::GetRHI();
+					}
+					else
+					{
+						GraphicsPSOInit.BlendState = TStaticBlendState<CW_RGB, BO_Add, BF_DestColor, BF_Zero>::GetRHI();
+					}
+
+					GraphicsPSOInit.PrimitiveType = PT_TriangleList;
+					GraphicsPSOInit.BoundShaderState.VertexDeclarationRHI = GTileVertexDeclaration.VertexDeclarationRHI;
+					GraphicsPSOInit.BoundShaderState.VertexShaderRHI = VertexShader.GetVertexShader();
+					GraphicsPSOInit.BoundShaderState.PixelShaderRHI = PixelShader.GetPixelShader();
+
+					SetGraphicsPipelineState(RHICmdList, GraphicsPSOInit, 0);
+
+					SetShaderParameters(RHICmdList, VertexShader, VertexShader.GetVertexShader(), PassParameters->VS);
+					SetShaderParameters(RHICmdList, PixelShader, PixelShader.GetPixelShader(), PassParameters->PS);
+
+					RHICmdList.SetStreamSource(0, GetOneTileQuadVertexBuffer(), 0);
+					RHICmdList.DrawIndexedPrimitive(GetOneTileQuadIndexBuffer(),
+						0,
+						0,
+						4,
+						0,
+						2,
+						GroupSize.X * GroupSize.Y);
+				});
+		}
+	}
+}
\ No newline at end of file
diff --git a/Engine/Source/Runtime/Renderer/Private/MobileShadingRenderer.cpp b/Engine/Source/Runtime/Renderer/Private/MobileShadingRenderer.cpp
--- a/Engine/Source/Runtime/Renderer/Private/MobileShadingRenderer.cpp
+++ b/Engine/Source/Runtime/Renderer/Private/MobileShadingRenderer.cpp
@@ -280,7 +280,6 @@ FMobileSceneRenderer::FMobileSceneRenderer(const FSceneViewFamily* InViewFamily,
 	bIsMaskedOnlyDepthPrepassEnabled = Scene->EarlyZPassMode == DDM_MaskedOnly;
 	bRequiresSceneDepthAux = MobileRequiresSceneDepthAux(ShaderPlatform);
 	bEnableClusteredLocalLights = MobileForwardEnableLocalLights(ShaderPlatform);
-	bEnablePrepassLocalLights = MobileForwardEnablePrepassLocalLights(ShaderPlatform);
 	bEnableClusteredReflections = MobileForwardEnableClusteredReflections(ShaderPlatform);
 	
 	StandardTranslucencyPass = ViewFamily.AllowTranslucencyAfterDOF() ? ETranslucencyPass::TPT_TranslucencyStandard : ETranslucencyPass::TPT_AllTranslucency;
@@ -590,7 +589,7 @@ void FMobileSceneRenderer::InitViews(
 	}
 
 	// When we capturing scene depth, use a more precise format for SceneDepthAux as it will be used as a source DepthTexture
-	if (bSceneDepthCapture)
+	if (bSceneDepthCapture || MobileLocalLightsBufferPostprocessEnabled(ShaderPlatform))
 	{
 		SceneTexturesConfig.bPreciseDepthAux = true;
 	}
@@ -957,10 +956,10 @@ void FMobileSceneRenderer::Render(FRDGBuilder& GraphBuilder)
 			bool bShadowedLightsInClustered = bRequiresShadowProjections && !bDeferredShading;
 			GatherAndSortLights(SortedLightSet, bShadowedLightsInClustered);
 			int32 NumReflectionCaptures = Views[0].NumBoxReflectionCaptures + Views[0].NumSphereReflectionCaptures;
-			bool bCullLightsToGrid = (((bEnableClusteredReflections || bDeferredShading) && NumReflectionCaptures > 0) || bEnableClusteredLocalLights || bEnablePrepassLocalLights);
+			bool bCullLightsToGrid = (((bEnableClusteredReflections || bDeferredShading) && NumReflectionCaptures > 0) || bEnableClusteredLocalLights);
 			if (bCullLightsToGrid)
 			{
-				ComputeLightGrid(GraphBuilder, bEnableClusteredLocalLights || bEnablePrepassLocalLights, SortedLightSet);
+				ComputeLightGrid(GraphBuilder, bEnableClusteredLocalLights, SortedLightSet);
 			}
 		}
 
@@ -1055,9 +1054,10 @@ void FMobileSceneRenderer::Render(FRDGBuilder& GraphBuilder)
 		}
 
 		// Local Light prepass
-		if (bEnablePrepassLocalLights && bRendererOutputFinalSceneColor)
+
+		if (bRendererOutputFinalSceneColor)
 		{
-			RenderLocalLightPrepass(GraphBuilder, SceneTextures);
+			RenderMobileLocalLightsBuffer(GraphBuilder, SceneTextures, true);
 		}
 	}
 
@@ -1098,6 +1098,8 @@ void FMobileSceneRenderer::Render(FRDGBuilder& GraphBuilder)
 	
 		GraphBuilder.SetCommandListStat(GET_STATID(STAT_CLMM_Post));
 
+		RenderMobileLocalLightsBuffer(GraphBuilder, SceneTextures, false);
+
 		FRendererModule& RendererModule = static_cast<FRendererModule&>(GetRendererModule());
 		RendererModule.RenderPostOpaqueExtensions(GraphBuilder, Views, SceneTextures);
 
diff --git a/Engine/Source/Runtime/Renderer/Private/SceneRendering.h b/Engine/Source/Runtime/Renderer/Private/SceneRendering.h
--- a/Engine/Source/Runtime/Renderer/Private/SceneRendering.h
+++ b/Engine/Source/Runtime/Renderer/Private/SceneRendering.h
@@ -2648,7 +2648,7 @@ protected:
 	void RenderMaskedPrePass(FRHICommandList& RHICmdList, const FViewInfo& View);
 	void RenderFullDepthPrepass(FRDGBuilder& GraphBuilder, FSceneTextures& SceneTextures);
 
-	void RenderLocalLightPrepass(FRDGBuilder& GraphBuilder, FSceneTextures& SceneTextures);
+	void RenderMobileLocalLightsBuffer(FRDGBuilder& GraphBuilder, FSceneTextures& SceneTextures, bool bIsPrepass);
 
 	/** Renders the opaque base pass for mobile. */
 	void RenderMobileBasePass(FRHICommandList& RHICmdList, const FViewInfo& View, const FInstanceCullingDrawParams* InstanceCullingDrawParams);
@@ -2728,7 +2728,6 @@ private:
 	bool bIsMaskedOnlyDepthPrepassEnabled;
 	bool bRequiresSceneDepthAux;
 	bool bEnableClusteredLocalLights;
-	bool bEnablePrepassLocalLights;
 	bool bEnableClusteredReflections;
 	bool bRequiresShadowProjections;
 
diff --git a/Engine/Source/Runtime/Renderer/Private/SceneTextures.cpp b/Engine/Source/Runtime/Renderer/Private/SceneTextures.cpp
--- a/Engine/Source/Runtime/Renderer/Private/SceneTextures.cpp
+++ b/Engine/Source/Runtime/Renderer/Private/SceneTextures.cpp
@@ -615,8 +615,8 @@ void FSceneTextures::InitializeViewFamily(FRDGBuilder& GraphBuilder, FViewFamily
 	}
 #endif
 
-	extern bool MobileForwardEnablePrepassLocalLights(const FStaticShaderPlatform Platform);
-	if(MobileForwardEnablePrepassLocalLights(Config.ShaderPlatform))
+	extern bool MobileLocalLightsBufferPrepassEnabled(const FStaticShaderPlatform Platform);
+	if(MobileLocalLightsBufferPrepassEnabled(Config.ShaderPlatform))
 	{
 		FRDGTextureDesc MobileLocalLightTextureADesc = FRDGTextureDesc::Create2D(Config.Extent, PF_FloatR11G11B10, FClearValueBinding::Transparent, TexCreate_RenderTargetable | TexCreate_ShaderResource);
 		SceneTextures.MobileLocalLightTextureA = GraphBuilder.CreateTexture(MobileLocalLightTextureADesc, TEXT("MobileLocalLightTextureA"));
