Rename/move file(s) from LocalHeightFog to LocalFogVolume

diff --git a/Engine/Shaders/Private/LocalFogVolume.usf b/Engine/Shaders/Private/LocalFogVolume.usf
new file mode 100644
--- /dev/null
+++ b/Engine/Shaders/Private/LocalFogVolume.usf
@@ -0,0 +1,355 @@
+// Copyright Epic Games, Inc. All Rights Reserved.
+#include "Common.ush"
+
+#if SHADING_PATH_MOBILE
+#define MobileSceneTextures MobileBasePass.SceneTextures 
+#endif
+
+#include "SceneTexturesCommon.ush"
+#include "HeightFogCommon.ush"
+#include "ParticipatingMediaCommon.ush"
+
+
+struct FLocalHeightFogGPUInstanceData
+{
+	float4x4 Transform;
+	float4x4 InvTransform;
+
+	float4x4 InvTranformNoScale;
+	float4x4 TransformScaleOnly;
+
+	float Density;
+	float HeightFalloff;
+	float HeightOffset;
+	float RadialAttenuation;
+
+	float3 Albedo;
+	float  PhaseG;
+	float3 Emissive;
+	float  FogMode;
+};
+
+// Structured buffers cannot be sampled from the vertex shader on some mobile platforms.
+// That is why simple Buffers are used, because this tech needs to run on mobile platforms.
+Buffer<float4> LocalHeightFogInstances;
+FLocalHeightFogGPUInstanceData GetLocalHeightFogGPUInstanceData(uint Index)
+{
+	FLocalHeightFogGPUInstanceData Data;
+	uint Offset = Index * 19;
+
+	Data.Transform		= float4x4(LocalHeightFogInstances[Offset + 0], LocalHeightFogInstances[Offset + 1], LocalHeightFogInstances[Offset + 2], LocalHeightFogInstances[Offset + 3]);
+	Data.InvTransform	= float4x4(LocalHeightFogInstances[Offset + 4], LocalHeightFogInstances[Offset + 5], LocalHeightFogInstances[Offset + 6], LocalHeightFogInstances[Offset + 7]);
+
+	Data.InvTranformNoScale = float4x4(LocalHeightFogInstances[Offset + 8], LocalHeightFogInstances[Offset + 9], LocalHeightFogInstances[Offset + 10], LocalHeightFogInstances[Offset + 11]);
+	Data.TransformScaleOnly = float4x4(LocalHeightFogInstances[Offset + 12], LocalHeightFogInstances[Offset + 13], LocalHeightFogInstances[Offset + 14], LocalHeightFogInstances[Offset + 15]);
+
+	float4 Temp0 = LocalHeightFogInstances[Offset + 16];
+	Data.Density			= Temp0.x;
+	Data.HeightFalloff		= Temp0.y;
+	Data.HeightOffset		= Temp0.z;
+	Data.RadialAttenuation	= Temp0.w;
+
+	float4 Temp1 = LocalHeightFogInstances[Offset + 17];
+	Data.Albedo		= Temp1.xyz;
+	Data.PhaseG		= Temp1.w;
+
+	float4 Temp2 = LocalHeightFogInstances[Offset + 18];
+	Data.Emissive	= Temp2.xyz;
+	Data.FogMode	= Temp2.w;
+
+	return Data;
+}
+
+
+struct FLocalHeightFogVertexOutput
+{
+	nointerpolation uint InstanceId : TEXCOORD0;
+};
+
+void LocalHeightFogSplatVS(
+	in float4 InPosition : ATTRIBUTE0,
+	uint InstanceIndex : SV_InstanceID,
+	out float4 OutPosition : SV_POSITION,
+	out FLocalHeightFogVertexOutput OtherVertexOutput
+)
+{
+	ResolvedView = ResolveView();
+
+	float4 LocalPosition = InPosition;
+
+	FLocalHeightFogGPUInstanceData FogInstance = GetLocalHeightFogGPUInstanceData(InstanceIndex);
+
+	float4x4 FogTransform = FogInstance.Transform;
+
+	float4 WorldPosition = mul(float4(LocalPosition.xyz, 1.0f), FogTransform);
+
+	float4 TranslatedWorldPosition = float4(WorldPosition.xyz + LWCHackToFloat(PrimaryView.PreViewTranslation), 1);
+
+	OutPosition = mul(TranslatedWorldPosition, PrimaryView.TranslatedWorldToClip);
+	OtherVertexOutput.InstanceId = InstanceIndex;
+} 
+
+// CameraPos can be behind of RayStart if the camera is out of the volume
+struct FFogData
+{
+	float IntegratedLuminanceFactor;
+	float Coverage;
+};
+
+// Height fog is evaluated according to spatially varying density defined as "Density * Exp[-HeightFalloff * z]"
+FFogData EvaluateHeightFogIntegral(FLocalHeightFogGPUInstanceData FogInstance,  float3 CameraPos, float3 RayStart, float3 RayDir, float RayLength, float RadialAttenuation)
+{
+	FFogData FogData;
+	
+	float3 RayOriU = mul(float4(RayStart, 1.0f),                      FogInstance.InvTransform).xyz;
+	float3 RayIntU = mul(float4(RayStart + RayDir * RayLength, 1.0f), FogInstance.InvTransform).xyz;
+	float3 RayDirU = mul(float4(RayDir  , 0.0f),                      FogInstance.InvTransform).xyz;
+	RayDirU = normalize(RayDirU);
+
+	float RayLengthU = length(RayIntU - RayOriU);
+
+	float StartHeight = RayOriU.z - FogInstance.HeightOffset; // accoutn for the fog offset.
+
+	float OpticalDepth = 0.0f;
+#if 0
+	// Original integral, less artefact
+	OpticalDepth = (FogInstance.Density / FogInstance.HeightFalloff) * exp(-StartHeight * FogInstance.HeightFalloff) * (1.0 - exp(-RayDirU.z * RayLengthU * FogInstance.HeightFalloff)) / RayDirU.z;
+#elif 1
+	// Reworked to avoid large value as input to exp
+	//OpticalDepth = (FogInstance.Density / FogInstance.HeightFalloff) * (exp(-StartHeight * FogInstance.HeightFalloff) - exp(-StartHeight * FogInstance.HeightFalloff - RayDir.z * RayLength * FogInstance.HeightFalloff)) / RayDir.z;
+
+	// But we use an updated formulation to avoid large input to exp leading to inf and nan.
+	float Factor0 = StartHeight * FogInstance.HeightFalloff;
+	float Factor1 = RayDirU.z * RayLengthU * FogInstance.HeightFalloff;
+	Factor0 = max(-80.0f, Factor0); // Simply clamp Factor0 to a valid range in order to not have exp explode.
+
+	OpticalDepth = (FogInstance.Density / (FogInstance.HeightFalloff * RayDirU.z)) * (exp(-Factor0) - exp(-(Factor0 + Factor1)));
+#endif
+
+	// Transmittance is the 
+	float Transmittance = exp(-OpticalDepth);
+	// Assuming extinction is grey scale, we can compute a single coverage value from transmittance
+	FogData.Coverage = (1.0 - Transmittance) * RadialAttenuation;
+	// Assuming extinction==scattering, i.e. albedo=1, we can compute the integrale for each point on a ray while accouting for transmittance back to the view.
+	// It turns out that, under these conditions, the luminance and emissive factor is simply coverage. => see Mathematica details at the bottom of this page.
+	FogData.IntegratedLuminanceFactor = FogData.Coverage;
+
+	return FogData;
+}
+
+
+FFogData EvaluateRadialFogIntegral(FLocalHeightFogGPUInstanceData FogInstance, float3 CameraPos, float3 RayStart, float3 RayDir, float RayLength, float RadialAttenuation)
+{
+	FFogData FogData = (FFogData)0;
+
+	float3 RayOriU = mul(float4(RayStart, 1.0f),                      FogInstance.InvTransform).xyz;
+	float3 RayIntU = mul(float4(RayStart + RayDir * RayLength, 1.0f), FogInstance.InvTransform).xyz;
+	float3 RayDirU = mul(float4(RayDir  , 0.0f),                      FogInstance.InvTransform).xyz;
+	RayDirU = normalize(RayDirU);
+
+	float RayLengthU = length(RayIntU - RayOriU);
+
+	float3 SphereCenter = 0.0f;
+	float3 VolumeCenterToRayO = (RayOriU - SphereCenter);
+
+	float b = dot(RayDirU,            VolumeCenterToRayO);
+	float c = dot(VolumeCenterToRayO, VolumeCenterToRayO) - 1.0f;
+	float h = b * b - c;
+
+	float OpticalDepth = 0.0f;
+	if (h >= 0.0)
+	{
+		h = sqrt(h);
+		float Length0 = -b - h;
+		float Length1 = -b + h;
+
+		Length0 = max(Length0, 0.0);
+		Length1 = max(Length1, 0.0);
+		Length1 = min(Length1, RayLengthU);
+
+		// Integral of a density that descreases accordingf to square distance from the center of the volume.
+		float Integral0 = -(c * Length0 + b * Length0 * Length0 + Length0 * Length0 * Length0 / 3.0f);
+		float Integral1 = -(c * Length1 + b * Length1 * Length1 + Length1 * Length1 * Length1 / 3.0f);
+		OpticalDepth = max(0.0, FogInstance.Density * (Integral1 - Integral0) * (3.0f / 4.0f));
+	}
+
+	// Transmittance is the 
+	float Transmittance = exp(-OpticalDepth);
+	// Assuming extinction is grey scale, we can compute a single coverage value from transmittance
+	FogData.Coverage = (1.0 - Transmittance) * RadialAttenuation;
+	// Assuming extinction==scattering, i.e. albedo=1, we can compute the integrale for each point on a ray while accouting for transmittance back to the view.
+	// It turns out that, under these conditions, the luminance and emissive factor is simply coverage. => see Mathematica details at the bottom of this page.
+	FogData.IntegratedLuminanceFactor = FogData.Coverage;
+
+	return FogData;
+}
+
+
+float3 ComputeFogInscattering(in FLocalHeightFogGPUInstanceData FogInstance, in FFogData FogData, in float3 RayDirWorld)
+{
+	half3 InScattering = 0;
+
+#if PROJECT_SUPPORT_SKY_ATMOSPHERE_AFFECTS_HEIGHFOG
+	InScattering +=  Texture2DSampleLevel(View.DistantSkyLightLutTexture, View.DistantSkyLightLutTextureSampler, float2(0.5f, 0.5f), 0.0f).rgb;
+
+	// No need to test View.AtmosphereLightIlluminanceOnGroundPostTransmittance[0].a because InscatteringLightDirection.w above is doing the same test already.
+	InScattering += View.AtmosphereLightIlluminanceOnGroundPostTransmittance[0].rgb * SchlickPhase(-FogInstance.PhaseG, dot(RayDirWorld, View.AtmosphereLightDirection[0].xyz));
+
+	if (View.AtmosphereLightIlluminanceOnGroundPostTransmittance[1].a > 0.0f) // Skip the second light when disabled.
+	{
+		InScattering += View.AtmosphereLightIlluminanceOnGroundPostTransmittance[0].rgb * SchlickPhase(-FogInstance.PhaseG, dot(RayDirWorld, View.AtmosphereLightDirection[1].xyz));
+	}
+#endif
+
+	InScattering *= FogInstance.Albedo * View.SkyAtmosphereHeightFogContribution.xxx;
+
+	// Now account for emissive luminance after albedo and heighfog contribution has been accounted for.
+	InScattering += FogInstance.Emissive;
+
+	return InScattering * FogData.IntegratedLuminanceFactor;
+}
+
+#define LHFV_FOLLOWS_ROTATIONS 1
+
+void LocalHeightFogSplatPS(
+	in float4 SVPos : SV_POSITION,
+	in FLocalHeightFogVertexOutput OtherVertexOutput,
+	out float4 OutColor : SV_Target0)
+{
+	ResolvedView = ResolveView();
+	uint InstanceIndex = OtherVertexOutput.InstanceId;
+
+	float3 CamRayTranslatedWorldOrigin	= PrimaryView.TranslatedWorldCameraOrigin;
+	float3 CamRayTranslatedWorldDir = normalize(SvPositionToTranslatedWorld(float4(SVPos.xy, 0.5, 1.0)));
+	float3 CamRayWorldDir = CamRayTranslatedWorldDir;
+
+	FLocalHeightFogGPUInstanceData FogInstance = GetLocalHeightFogGPUInstanceData(InstanceIndex);
+
+	OutColor = float4(0, 0, 0, 1);
+
+	float3 CamRayWorldOrigin = CamRayTranslatedWorldOrigin - LWCHackToFloat(PrimaryView.PreViewTranslation);
+
+	// The "U" prefix is for all compuation done in the Unit Sphere space.
+	float3 RayPosU = mul(float4(CamRayWorldOrigin,        1.0f), FogInstance.InvTransform).xyz;
+	float3 RayDirU = mul(float4(CamRayTranslatedWorldDir, 0.0f), FogInstance.InvTransform).xyz; // Try to use InverseTranspose to correctly handle direction transform with scaling.
+	RayDirU.xyz = normalize(RayDirU);
+
+	float2 TsU = RayIntersectSphere(RayPosU, RayDirU, float4(0.0, 0.0, 0.0, 1.0));
+	if (any(TsU > 0.0))
+	{
+#if SHADING_PATH_MOBILE
+		// LookupDeviceZ with uint2 SVPos.xy as input does not work on mobile.
+		float2 UV = SVPos * View.BufferSizeAndInvSize.zw;
+		const float DeviceZ = LookupDeviceZ(UV);
+#else
+		const float DeviceZ = LookupDeviceZ(uint2(SVPos.xy));
+#endif
+
+		float3 DepthBufferTranslatedWorldPos = SvPositionToTranslatedWorld(float4(SVPos.xy, DeviceZ, 1.0));
+		float LengthD = length(DepthBufferTranslatedWorldPos - CamRayTranslatedWorldOrigin);
+
+		float3 P0U = RayPosU + max(0.0, TsU.x) * RayDirU;
+		float3 P1U = RayPosU + max(0.0, TsU.y) * RayDirU;
+#if LHFV_FOLLOWS_ROTATIONS
+		// Computations are all done transformed with Space = local fog volume but ignoring scaling.
+		float3 P0Space = mul(float4(P0U, 1.0), FogInstance.TransformScaleOnly).xyz;
+		float3 P1Space = mul(float4(P1U, 1.0), FogInstance.TransformScaleOnly).xyz;
+		float3 CamRayOriginSpace = mul(float4(CamRayWorldOrigin, 1.0), FogInstance.InvTranformNoScale).xyz;
+		float3 CamRayDirSpace    = mul(float4(CamRayWorldDir,    0.0), FogInstance.InvTranformNoScale).xyz;	// No scale so no need for a special InverseTranspose matrix.
+		CamRayDirSpace = normalize(CamRayDirSpace);
+#else
+		// Computations are all done with Space = world space
+		float3 P0Space = mul(float4(P0U, 1.0), FogInstance.Transform).xyz;
+		float3 P1Space = mul(float4(P1U, 1.0), FogInstance.Transform).xyz;
+		float3 CamRayOriginSpace = CamRayWorldOrigin;
+		float3 CamRayDirSpace    = CamRayWorldDir;
+#endif
+
+		float Length0 = length(P0Space - CamRayOriginSpace);
+		float Length1 = length(P1Space - CamRayOriginSpace);
+
+		// Radial attenuation based on the distance travelled in the unit sphere and ignoring depth buffer intersection. 
+		// This is only to soften sphere edge.
+		float RadialAttenuation = saturate(abs(max(0.0, TsU.y) - max(0.0, TsU.x)) / FogInstance.RadialAttenuation);
+		// The attenuation is linear but it can be made soft using an Hermit interpolation by controling the tangents.
+		{
+			const float t = RadialAttenuation;
+			const float t2 = t * t;
+			const float t3 = t2 * t;
+			const float HP0 = 0.0f;
+			const float HP1 = 1.0f;
+			const float HM0 = 0.0f;	// tangent = 0 for soft transition
+			const float HM1 = 0.0f; // idem
+			RadialAttenuation = saturate((2 * t3 - 3 * t2 + 1) * HP0 + (t3 - 2 * t2 + t) * HM0 + (-2 * t3 + 3 * t2) * HP1 + (t3 - t2) * HM1);
+		}
+
+		// Clamp the traced distance to depth and compute world length.
+		Length0 = min(LengthD, Length0);
+		Length1 = min(LengthD, Length1);
+		float RayTracedLength = max(0.0, abs(Length1 - Length0));
+
+		if (RayTracedLength > 0.0)
+		{
+			FFogData FogData = (FFogData)0;
+
+			if (FogInstance.FogMode == 0)
+			{
+				FogData = EvaluateHeightFogIntegral(FogInstance, CamRayWorldOrigin, CamRayWorldOrigin + CamRayWorldDir * Length0, CamRayWorldDir, RayTracedLength, RadialAttenuation);
+			}
+			else // (FogInstance.FogMode == 1)
+			{
+				// TODO
+				// TODO EvaluateRadialFogIntegral needs to be optimised because it is redoing lots of computations. Maybe move HeightFogIntegral to also be in a local normalised space.
+				// TODO It is like that for now while iterating with artists.
+				// TODO
+				FogData = EvaluateRadialFogIntegral(FogInstance, CamRayWorldOrigin, CamRayWorldOrigin + CamRayWorldDir * Length0, CamRayWorldDir, RayTracedLength, RadialAttenuation);
+			}
+
+			OutColor.xyz	= ComputeFogInscattering(FogInstance, FogData, CamRayWorldDir);	// = Luminance
+			OutColor.a		= 1.0 - FogData.Coverage;										// = Transmittance
+		}
+		else
+		{
+			clip(-1.0);
+		}
+
+		OutColor.xyz *= PrimaryView.PreExposure;
+	}
+	else
+	{
+		clip(-1.0);
+	}
+}
+
+
+/*
+
+Proof of the height fog integral equations using Mathematica
+
+(*Density of matter according to height Y*)
+Density [y_] := A * Exp[-B * y]
+ 
+(*Ray position*)
+Ray[t_] := Oz + t * Rz
+
+(*Integrate Optical Depth for a distance T*)
+Fog[T_] := Integrate[Density[Ray[t]], {t, 0, T}]
+
+(*Optical Depth integration equation*)
+Fog[T]
+
+(*Transmittance from Optical Depth*)
+Exp[-Integrate[Density[Ray[t]], {t, 0, T}]]
+
+(*Transmittance from a position to the origin of tracing*)
+Trans[T_] := Exp[-Integrate[Density[Ray[t]], {t, 0, T}]]
+
+(*Integrate emissive and scattering from height fog, assuming density \
+is extinction and extinction=scattering, thus albedo=1*)
+Integrate[Density[Ray[t]]*EmSc *Trans[t], {t, 0, T}]
+
+=> this leads to EmSc - Transmittance * EmSc
+				 EmSc * (1 - Transmittance)
+				 EmSc * Coverage
+
+*/
\ No newline at end of file
diff --git a/Engine/Shaders/Private/LocalHeightFog.usf b/Engine/Shaders/Private/LocalHeightFog.usf
deleted file mode 100644
diff --git a/Engine/Source/Editor/ComponentVisualizers/Private/ComponentVisualizers.cpp b/Engine/Source/Editor/ComponentVisualizers/Private/ComponentVisualizers.cpp
--- a/Engine/Source/Editor/ComponentVisualizers/Private/ComponentVisualizers.cpp
+++ b/Engine/Source/Editor/ComponentVisualizers/Private/ComponentVisualizers.cpp
@@ -14,7 +14,7 @@
 #include "Components/SpotLightComponent.h"
 #include "Components/StereoLayerComponent.h"
 #include "Components/WorldPartitionStreamingSourceComponent.h"
-#include "Components/LocalHeightFogComponent.h"
+#include "Components/LocalFogVolumeComponent.h"
 #include "ConstraintComponentVisualizer.h"
 #include "DecalComponentVisualizer.h"
 #include "Editor/UnrealEdEngine.h"
@@ -41,7 +41,7 @@
 #include "UObject/Class.h"
 #include "UnrealEdGlobals.h"
 #include "WorldPartitionStreamingSourceComponentVisualizer.h"
-#include "LocalHeightFogComponentVisualizer.h"
+#include "LocalFogVolumeComponentVisualizer.h"
 
 IMPLEMENT_MODULE( FComponentVisualizersModule, ComponentVisualizers );
 
diff --git a/Engine/Source/Editor/ComponentVisualizers/Private/LocalFogVolumeComponentVisualizer.cpp b/Engine/Source/Editor/ComponentVisualizers/Private/LocalFogVolumeComponentVisualizer.cpp
new file mode 100644
--- /dev/null
+++ b/Engine/Source/Editor/ComponentVisualizers/Private/LocalFogVolumeComponentVisualizer.cpp
@@ -0,0 +1,25 @@
+// Copyright Epic Games, Inc. All Rights Reserved.
+
+#include "LocalFogVolumeComponentVisualizer.h"
+
+#include "Components/ActorComponent.h"
+#include "Components/LocalFogVolumeComponent.h"
+#include "SceneView.h"
+#include "SceneManagement.h"
+
+void FLocalHeightFogComponentVisualizer::DrawVisualization( const UActorComponent* Component, const FSceneView* View, FPrimitiveDrawInterface* PDI )
+{
+	if(View->Family->EngineShowFlags.CompositeEditorPrimitives)
+	{
+		const ULocalHeightFogComponent* LocalHeightFog = Cast<const ULocalHeightFogComponent>(Component);
+		if(LocalHeightFog != NULL)
+		{
+			FTransform LocalHeightFogTransform = LocalHeightFog->GetComponentTransform();
+
+			// Draw local fog volume spherical shape radius
+			DrawWireSphereAutoSides(PDI, LocalHeightFogTransform, FColor(200, 255, 255), 1.0f, SDPG_World);
+		}
+	}
+}
+
+
diff --git a/Engine/Source/Editor/ComponentVisualizers/Private/LocalHeightFogComponentVisualizer.cpp b/Engine/Source/Editor/ComponentVisualizers/Private/LocalHeightFogComponentVisualizer.cpp
deleted file mode 100644
diff --git a/Engine/Source/Editor/ComponentVisualizers/Public/LocalFogVolumeComponentVisualizer.h b/Engine/Source/Editor/ComponentVisualizers/Public/LocalFogVolumeComponentVisualizer.h
new file mode 100644
--- /dev/null
+++ b/Engine/Source/Editor/ComponentVisualizers/Public/LocalFogVolumeComponentVisualizer.h
@@ -0,0 +1,18 @@
+// Copyright Epic Games, Inc. All Rights Reserved.
+
+#pragma once
+
+#include "CoreMinimal.h"
+#include "ComponentVisualizer.h"
+
+class FPrimitiveDrawInterface;
+class FSceneView;
+class UActorComponent;
+
+class COMPONENTVISUALIZERS_API FLocalHeightFogComponentVisualizer : public FComponentVisualizer
+{
+public:
+	//~ Begin FComponentVisualizer Interface
+	virtual void DrawVisualization(const UActorComponent* Component, const FSceneView* View, FPrimitiveDrawInterface* PDI) override;
+	//~ End FComponentVisualizer Interface
+};
diff --git a/Engine/Source/Editor/ComponentVisualizers/Public/LocalHeightFogComponentVisualizer.h b/Engine/Source/Editor/ComponentVisualizers/Public/LocalHeightFogComponentVisualizer.h
deleted file mode 100644
diff --git a/Engine/Source/Editor/LevelEditor/Private/LevelEditorOutlinerSettings.cpp b/Engine/Source/Editor/LevelEditor/Private/LevelEditorOutlinerSettings.cpp
--- a/Engine/Source/Editor/LevelEditor/Private/LevelEditorOutlinerSettings.cpp
+++ b/Engine/Source/Editor/LevelEditor/Private/LevelEditorOutlinerSettings.cpp
@@ -19,7 +19,7 @@
 #include "Animation/SkeletalMeshActor.h"
 #include "Components/SkyAtmosphereComponent.h"
 #include "Components/VolumetricCloudComponent.h"
-#include "Engine/LocalHeightFog.h"
+#include "Engine/LocalFogVolume.h"
 #include "Engine/Blueprint.h"
 #include "Engine/Brush.h"
 #include "Engine/ExponentialHeightFog.h"
diff --git a/Engine/Source/Editor/PlacementMode/Private/PlacementModeModule.cpp b/Engine/Source/Editor/PlacementMode/Private/PlacementModeModule.cpp
--- a/Engine/Source/Editor/PlacementMode/Private/PlacementModeModule.cpp
+++ b/Engine/Source/Editor/PlacementMode/Private/PlacementModeModule.cpp
@@ -22,7 +22,7 @@
 #include "ActorFactories/ActorFactoryEmptyActor.h"
 #include "ActorFactories/ActorFactoryPawn.h"
 #include "ActorFactories/ActorFactoryExponentialHeightFog.h"
-#include "ActorFactories/ActorFactoryLocalHeightFog.h"
+#include "ActorFactories/ActorFactoryLocalFogVolume.h"
 #include "ActorFactories/ActorFactorySkyAtmosphere.h"
 #include "ActorFactories/ActorFactoryVolumetricCloud.h"
 #include "ActorFactories/ActorFactoryPlayerStart.h"
diff --git a/Engine/Source/Editor/UnrealEd/Classes/ActorFactories/ActorFactoryLocalFogVolume.h b/Engine/Source/Editor/UnrealEd/Classes/ActorFactories/ActorFactoryLocalFogVolume.h
new file mode 100644
--- /dev/null
+++ b/Engine/Source/Editor/UnrealEd/Classes/ActorFactories/ActorFactoryLocalFogVolume.h
@@ -0,0 +1,15 @@
+// Copyright Epic Games, Inc. All Rights Reserved.
+
+#pragma once
+
+#include "CoreMinimal.h"
+#include "UObject/ObjectMacros.h"
+#include "ActorFactories/ActorFactory.h"
+#include "ActorFactoryLocalFogVolume.generated.h"
+
+UCLASS(MinimalAPI, config = Editor, collapsecategories, hidecategories = Object)
+class UActorFactoryLocalHeightFog : public UActorFactory
+{
+	GENERATED_UCLASS_BODY()
+};
+
diff --git a/Engine/Source/Editor/UnrealEd/Classes/ActorFactories/ActorFactoryLocalHeightFog.h b/Engine/Source/Editor/UnrealEd/Classes/ActorFactories/ActorFactoryLocalHeightFog.h
deleted file mode 100644
diff --git a/Engine/Source/Editor/UnrealEd/Private/Factories/ActorFactory.cpp b/Engine/Source/Editor/UnrealEd/Private/Factories/ActorFactory.cpp
--- a/Engine/Source/Editor/UnrealEd/Private/Factories/ActorFactory.cpp
+++ b/Engine/Source/Editor/UnrealEd/Private/Factories/ActorFactory.cpp
@@ -31,7 +31,7 @@ ActorFactory.cpp:
 #include "ActorFactories/ActorFactoryEmptyActor.h"
 #include "ActorFactories/ActorFactoryPawn.h"
 #include "ActorFactories/ActorFactoryExponentialHeightFog.h"
-#include "ActorFactories/ActorFactoryLocalHeightFog.h"
+#include "ActorFactories/ActorFactoryLocalFogVolume.h"
 #include "ActorFactories/ActorFactoryNote.h"
 #include "ActorFactories/ActorFactoryPhysicsAsset.h"
 #include "ActorFactories/ActorFactoryPlaneReflectionCapture.h"
@@ -80,7 +80,7 @@ ActorFactory.cpp:
 #include "Engine/DecalActor.h"
 #include "Atmosphere/AtmosphericFog.h"
 #include "Engine/ExponentialHeightFog.h"
-#include "Engine/LocalHeightFog.h"
+#include "Engine/LocalFogVolume.h"
 #include "Engine/SkyLight.h"
 #include "Engine/DirectionalLight.h"
 #include "Engine/PointLight.h"
diff --git a/Engine/Source/Runtime/Engine/Classes/Components/LocalFogVolumeComponent.h b/Engine/Source/Runtime/Engine/Classes/Components/LocalFogVolumeComponent.h
new file mode 100644
--- /dev/null
+++ b/Engine/Source/Runtime/Engine/Classes/Components/LocalFogVolumeComponent.h
@@ -0,0 +1,87 @@
+// Copyright Epic Games, Inc. All Rights Reserved.
+
+#pragma once
+
+#include "CoreMinimal.h"
+#include "Components/SceneComponent.h"
+#include "Math/Color.h"
+#include "LocalFogVolumeComponent.generated.h"
+
+class FLocalHeightFogSceneProxy;
+
+UENUM()
+enum class ELocalFogMode : uint8
+{
+	LocalHeightFog = 0,
+	LocalSphereFog = 1,
+};
+
+UCLASS(ClassGroup = Rendering, collapsecategories, hidecategories = (Object, Mobility, Activation, "Components|Activation"), editinlinenew, meta = (BlueprintSpawnableComponent), MinimalAPI)
+class ULocalHeightFogComponent : public USceneComponent
+{
+	GENERATED_UCLASS_BODY()
+
+	~ULocalHeightFogComponent();
+
+	/** Controls the softness of the transition region when the volume is fading out. */
+	UPROPERTY(EditAnywhere, BlueprintReadOnly, Category = "Fog Mode")
+	ELocalFogMode FogMode = ELocalFogMode::LocalHeightFog;
+
+	/** Global density factor for this fog. */
+	UPROPERTY(EditAnywhere, BlueprintReadOnly, interp, Category = "Fog Distribution", meta = (UIMin = "0", UIMax = "10.0", SliderExponent = 2.0, ClampMin = 0.0))
+	float FogDensity = 5.0f;
+
+	/** Controls how the density decreases as height increases. Smaller values make the visible transition larger. */
+	UPROPERTY(EditAnywhere, BlueprintReadOnly, interp, Category = "Fog Distribution", meta = (UIMin = "0.001", UIMax = "5000", SliderExponent = 2.0, ClampMin = 0.001))
+	float FogHeightFalloff = 1000.0f;
+
+	/** Height offset, relative to the actor Z position. */
+	UPROPERTY(EditAnywhere, BlueprintReadOnly, interp, Category = "Fog Distribution", meta = (UIMin = "-2.0", UIMax = "2.0"))
+	float FogHeightOffset = 0.0f;
+
+	/** Controls how strong the radial attenuation of this fog volume is. */
+	UPROPERTY(EditAnywhere, BlueprintReadOnly, interp, Category = "Fog Distribution", meta = (UIMin = "0.0", UIMax = "4.0", SliderExponent = 3.0, ClampMin = 0.0))
+	float FogRadialAttenuation = 0.0f;
+
+	/** Controls the phase `G` parameter, describing the directionality of the scattering within this fog volume. */
+	UPROPERTY(EditAnywhere, BlueprintReadOnly, interp, Category = "Fog Shading", meta = (DisplayName = "Scattering Distribution", UIMin = "0.0", UIMax = "0.999", ClampMin = 0.0, ClampMax = 0.999))
+	float FogPhaseG = 0.8f;
+
+	/** Controls the albedo of this fog volume. */
+	UPROPERTY(EditAnywhere, BlueprintReadOnly, interp, Category = "Fog Shading", meta = (ClampMin = 0.0, ClampMax = 1.0))
+	FLinearColor FogAlbedo = FLinearColor::White;
+
+	/** Controls the emissive color of this fog volume. */
+	UPROPERTY(EditAnywhere, BlueprintReadOnly, interp, Category = "Fog Shading", meta = (ClampMin = 0.0, ClampMax = 1.0))
+	FLinearColor FogEmissive = FLinearColor::Black;
+
+	/** The priority can be used as a way to override the sorting by distance. A lower value means the volume will be considered further away, i.e. it will draw behind the one with a higher priority value. */
+	UPROPERTY(EditAnywhere, BlueprintReadOnly, Category = "Sorting", meta = (UIMin = "-127", UIMax = "127", ClampMin = -127, ClampMax = 127))
+	int32 FogSortPriority = 0;
+
+public:
+
+	//~ Begin UObject Interface
+#if WITH_EDITOR
+	virtual bool CanEditChange(const FProperty* InProperty) const override;
+	virtual void PostEditChangeProperty(FPropertyChangedEvent& PropertyChangedEvent) override;
+#endif // WITH_EDITOR
+	//~ End UObject Interface
+
+	/** Callback to create the rendering thread mirror. */
+	ENGINE_API FLocalHeightFogSceneProxy* CreateSceneProxy();
+
+protected:
+	//~ Begin UActorComponent Interface.
+	virtual void CreateRenderState_Concurrent(FRegisterComponentContext* Context) override;
+	virtual void SendRenderTransform_Concurrent() override;
+	virtual void DestroyRenderState_Concurrent() override;
+	//~ End UActorComponent Interface.
+
+	void SendRenderTransformCommand();
+
+private:
+
+	FLocalHeightFogSceneProxy* LocalHeightFogSceneProxy;
+};
+
diff --git a/Engine/Source/Runtime/Engine/Classes/Components/LocalHeightFogComponent.h b/Engine/Source/Runtime/Engine/Classes/Components/LocalHeightFogComponent.h
deleted file mode 100644
diff --git a/Engine/Source/Runtime/Engine/Classes/Engine/LocalFogVolume.h b/Engine/Source/Runtime/Engine/Classes/Engine/LocalFogVolume.h
new file mode 100644
--- /dev/null
+++ b/Engine/Source/Runtime/Engine/Classes/Engine/LocalFogVolume.h
@@ -0,0 +1,34 @@
+// Copyright Epic Games, Inc. All Rights Reserved.
+
+#pragma once
+
+#include "CoreMinimal.h"
+#include "GameFramework/Info.h"
+#include "LocalFogVolume.generated.h"
+
+class ULocalHeightFogComponent;
+
+/**
+ *	Actor used to position a local fog volume in the scene.
+ *	@see https://docs.unrealengine.com/???
+ */
+UCLASS(showcategories = (Movement, Rendering, Transformation, DataLayers, "Input|MouseInput", "Input|TouchInput"), ClassGroup = Fog, hidecategories = (Info, Object, Input), MinimalAPI)
+class ALocalHeightFog : public AInfo
+{
+	GENERATED_UCLASS_BODY()
+
+private:
+#if WITH_EDITOR
+	virtual bool IsDataLayerTypeSupported(TSubclassOf<UDataLayerInstance> DataLayerType) const override { return true; }
+#endif
+
+	/** Object used to visualize the local fog volume */
+	UPROPERTY(BlueprintReadOnly, VisibleAnywhere, Category = Fog, meta = (AllowPrivateAccess = "true"))
+	TObjectPtr<ULocalHeightFogComponent> LocalHeightFogVolume;
+
+public:
+
+	/** Returns LocalHeightFogVolume subobject **/
+	ULocalHeightFogComponent* GetComponent() const { return LocalHeightFogVolume; }
+};
+
diff --git a/Engine/Source/Runtime/Engine/Classes/Engine/LocalHeightFog.h b/Engine/Source/Runtime/Engine/Classes/Engine/LocalHeightFog.h
deleted file mode 100644
diff --git a/Engine/Source/Runtime/Engine/Private/Components/LocalFogVolumeComponent.cpp b/Engine/Source/Runtime/Engine/Private/Components/LocalFogVolumeComponent.cpp
new file mode 100644
--- /dev/null
+++ b/Engine/Source/Runtime/Engine/Private/Components/LocalFogVolumeComponent.cpp
@@ -0,0 +1,183 @@
+// Copyright Epic Games, Inc. All Rights Reserved.
+
+#include "Components/LocalFogVolumeComponent.h"
+#include "Engine/LocalFogVolume.h"
+#include "LocalFogVolumeSceneProxy.h"
+#include "Components/ArrowComponent.h"
+#include "Components/DrawSphereComponent.h"
+#include "Components/BillboardComponent.h"
+#include "SceneInterface.h"
+#include "SceneManagement.h"
+#include "Engine/World.h"
+#include "UObject/ConstructorHelpers.h"
+#include "Engine/Texture2D.h"
+
+#include UE_INLINE_GENERATED_CPP_BY_NAME(LocalFogVolumeComponent)
+
+
+ULocalHeightFogComponent::ULocalHeightFogComponent(const FObjectInitializer& ObjectInitializer)
+	: Super(ObjectInitializer)
+	, LocalHeightFogSceneProxy(nullptr)
+{
+	Mobility = EComponentMobility::Movable;
+
+	// Set a size by default to be visible when droped in the scene
+	const float Size = 500.0f;
+	SetRelativeScale3D(FVector(Size, Size, Size));
+}
+
+ULocalHeightFogComponent::~ULocalHeightFogComponent()
+{
+}
+
+void ULocalHeightFogComponent::SendRenderTransformCommand()
+{
+	if (LocalHeightFogSceneProxy)
+	{
+		FTransform ComponentTransform = GetComponentTransform();
+		float HeightOffset = FogHeightOffset;
+		FLocalHeightFogSceneProxy* SceneProxy = LocalHeightFogSceneProxy;
+		ENQUEUE_RENDER_COMMAND(FUpdateLocalHeightFogSceneProxyTransformCommand)(
+			[SceneProxy, ComponentTransform, HeightOffset](FRHICommandList& RHICmdList)
+			{
+				// Nothing else is needed so that command could actually go.
+				SceneProxy->FogTransform = ComponentTransform;
+			});
+	}
+}
+
+void ULocalHeightFogComponent::CreateRenderState_Concurrent(FRegisterComponentContext* Context)
+{
+	Super::CreateRenderState_Concurrent(Context);	
+	
+	bool bHidden = false;
+#if WITH_EDITORONLY_DATA
+	bHidden = GetOwner() ? GetOwner()->bHiddenEdLevel : false;
+#endif // WITH_EDITORONLY_DATA
+	if (!ShouldComponentAddToScene())
+	{
+		bHidden = true;
+	}
+
+	if (GetVisibleFlag() && !bHidden && ShouldComponentAddToScene() && ShouldRender() && IsRegistered() && (GetOuter() == NULL || !GetOuter()->HasAnyFlags(RF_ClassDefaultObject)))
+	{
+		LocalHeightFogSceneProxy = CreateSceneProxy();
+		GetWorld()->Scene->AddLocalHeightFog(LocalHeightFogSceneProxy);
+	}
+}
+
+void ULocalHeightFogComponent::SendRenderTransform_Concurrent()
+{
+	Super::SendRenderTransform_Concurrent();
+	SendRenderTransformCommand();
+}
+
+void ULocalHeightFogComponent::DestroyRenderState_Concurrent()
+{
+	Super::DestroyRenderState_Concurrent();
+
+	if (LocalHeightFogSceneProxy)
+	{
+		GetWorld()->Scene->RemoveLocalHeightFog(LocalHeightFogSceneProxy);
+
+		FLocalHeightFogSceneProxy* SceneProxy = LocalHeightFogSceneProxy;
+		ENQUEUE_RENDER_COMMAND(FDestroyLocalHeightFogSceneProxyCommand)(
+			[SceneProxy](FRHICommandList& RHICmdList)
+			{
+				delete SceneProxy;
+			});
+
+		LocalHeightFogSceneProxy = nullptr;
+	}
+}
+
+#if WITH_EDITOR
+
+bool ULocalHeightFogComponent::CanEditChange(const FProperty* InProperty) const
+{
+	if (InProperty)
+	{
+		FString PropertyName = InProperty->GetName();
+
+		if (FogMode != ELocalFogMode::LocalHeightFog)
+		{
+			if (PropertyName == GET_MEMBER_NAME_STRING_CHECKED(ULocalHeightFogComponent, FogHeightFalloff)
+				|| PropertyName == GET_MEMBER_NAME_STRING_CHECKED(ULocalHeightFogComponent, FogHeightOffset))
+			{
+				return false;
+			}
+		}
+	}
+
+	return true;
+}
+
+void ULocalHeightFogComponent::PostEditChangeProperty(FPropertyChangedEvent& PropertyChangedEvent)
+{
+	Super::PostEditChangeProperty(PropertyChangedEvent);
+
+	FProperty* Property = PropertyChangedEvent.Property;
+	const FName PropertyName = Property ? Property->GetFName() : NAME_None;
+
+	if (PropertyName == TEXT("RelativeLocation") ||
+		PropertyName == TEXT("RelativeRotation") ||
+		PropertyName == TEXT("RelativeScale3D"))
+	{
+		SendRenderTransformCommand();
+	}
+	else
+	{
+		MarkRenderStateDirty();
+	}
+}
+
+#endif // WITH_EDITOR
+
+FLocalHeightFogSceneProxy* ULocalHeightFogComponent::CreateSceneProxy()
+{
+	return new FLocalHeightFogSceneProxy(this);
+}
+
+#if WITH_EDITOR
+#include "ObjectEditorUtils.h"
+#endif
+
+ALocalHeightFog::ALocalHeightFog(const FObjectInitializer& ObjectInitializer)
+	: Super(ObjectInitializer.SetDefaultSubobjectClass<ULocalHeightFogComponent>(TEXT("NewLocalHeightFogComponent")))
+{
+	LocalHeightFogVolume = CreateDefaultSubobject<ULocalHeightFogComponent>(TEXT("LocalHeightFogComponent"));
+	RootComponent = LocalHeightFogVolume;
+
+#if WITH_EDITORONLY_DATA
+
+	if (!IsRunningCommandlet())
+	{
+		// Structure to hold one-time initialization
+		struct FConstructorStatics
+		{
+			ConstructorHelpers::FObjectFinderOptional<UTexture2D> LocalHeightFogTextureObject;
+			FName ID_LocalHeightFog;
+			FText NAME_LocalHeightFog;
+			FConstructorStatics()
+				: LocalHeightFogTextureObject(TEXT("/Engine/EditorResources/S_SkyAtmosphere"))	// TODO
+				, ID_LocalHeightFog(TEXT("Fog"))
+				, NAME_LocalHeightFog(NSLOCTEXT("SpriteCategory", "Fog", "Fog"))
+			{
+			}
+		};
+		static FConstructorStatics ConstructorStatics;
+
+		if (GetSpriteComponent())
+		{
+			GetSpriteComponent()->Sprite = ConstructorStatics.LocalHeightFogTextureObject.Get();
+			GetSpriteComponent()->SetRelativeScale3D(FVector(0.5f, 0.5f, 0.5f));
+			GetSpriteComponent()->SpriteInfo.Category = ConstructorStatics.ID_LocalHeightFog;
+			GetSpriteComponent()->SpriteInfo.DisplayName = ConstructorStatics.NAME_LocalHeightFog;
+			GetSpriteComponent()->SetupAttachment(LocalHeightFogVolume);
+		}
+	}
+#endif // WITH_EDITORONLY_DATA
+
+	PrimaryActorTick.bCanEverTick = true;
+	SetHidden(false);
+}
diff --git a/Engine/Source/Runtime/Engine/Private/Components/LocalHeightFogComponent.cpp b/Engine/Source/Runtime/Engine/Private/Components/LocalHeightFogComponent.cpp
deleted file mode 100644
diff --git a/Engine/Source/Runtime/Engine/Private/LocalFogVolumeSceneProxy.cpp b/Engine/Source/Runtime/Engine/Private/LocalFogVolumeSceneProxy.cpp
new file mode 100644
--- /dev/null
+++ b/Engine/Source/Runtime/Engine/Private/LocalFogVolumeSceneProxy.cpp
@@ -0,0 +1,30 @@
+// Copyright Epic Games, Inc. All Rights Reserved.
+
+/*=============================================================================
+	FVolumetricCloudSceneProxy implementation.
+=============================================================================*/
+
+#include "LocalFogVolumeSceneProxy.h"
+#include "Components/LocalFogVolumeComponent.h"
+
+
+
+FLocalHeightFogSceneProxy::FLocalHeightFogSceneProxy(const ULocalHeightFogComponent* InComponent)
+	: FogTransform(InComponent->GetComponentTransform())
+	, FogDensity(InComponent->FogDensity)
+	, FogHeightFalloff(InComponent->FogHeightFalloff)
+	, FogHeightOffset(InComponent->FogHeightOffset)
+	, FogRadialAttenuation(InComponent->FogRadialAttenuation)
+	, FogMode((uint8)InComponent->FogMode)
+	, FogSortPriority(uint8(127 - int8(InComponent->FogSortPriority))) // FogSortPriority on the component is in [-127,127] and needs to be negated to match expected priority behavior.
+	, FogPhaseG(InComponent->FogPhaseG)
+	, FogAlbedo(InComponent->FogAlbedo)
+	, FogEmissive(InComponent->FogEmissive)
+{
+}
+
+FLocalHeightFogSceneProxy::~FLocalHeightFogSceneProxy()
+{
+}
+
+
diff --git a/Engine/Source/Runtime/Engine/Private/LocalHeightFogSceneProxy.cpp b/Engine/Source/Runtime/Engine/Private/LocalHeightFogSceneProxy.cpp
deleted file mode 100644
diff --git a/Engine/Source/Runtime/Engine/Public/LocalFogVolumeSceneProxy.h b/Engine/Source/Runtime/Engine/Public/LocalFogVolumeSceneProxy.h
new file mode 100644
--- /dev/null
+++ b/Engine/Source/Runtime/Engine/Public/LocalFogVolumeSceneProxy.h
@@ -0,0 +1,33 @@
+// Copyright Epic Games, Inc. All Rights Reserved.
+
+#pragma once
+
+#include "CoreMinimal.h"
+#include "Math/Color.h"
+
+class ULocalHeightFogComponent;
+
+/** Represents a UVolumetricCloudComponent to the rendering thread, created game side from the component. */
+class FLocalHeightFogSceneProxy
+{
+public:
+
+	// Initialization constructor.
+	ENGINE_API FLocalHeightFogSceneProxy(const ULocalHeightFogComponent* InComponent);
+	ENGINE_API ~FLocalHeightFogSceneProxy();
+
+	FTransform FogTransform;
+
+	float FogDensity;
+	float FogHeightFalloff;
+	float FogHeightOffset;
+	float FogRadialAttenuation;
+
+	uint8 FogMode;
+	uint8 FogSortPriority;
+
+	float FogPhaseG;
+	FLinearColor FogAlbedo;
+	FLinearColor FogEmissive;
+private:
+};
diff --git a/Engine/Source/Runtime/Engine/Public/LocalHeightFogSceneProxy.h b/Engine/Source/Runtime/Engine/Public/LocalHeightFogSceneProxy.h
deleted file mode 100644
diff --git a/Engine/Source/Runtime/Renderer/Private/DeferredShadingRenderer.cpp b/Engine/Source/Runtime/Renderer/Private/DeferredShadingRenderer.cpp
--- a/Engine/Source/Runtime/Renderer/Private/DeferredShadingRenderer.cpp
+++ b/Engine/Source/Runtime/Renderer/Private/DeferredShadingRenderer.cpp
@@ -89,7 +89,7 @@
 #include "ComponentRecreateRenderStateContext.h"
 #include "RenderCore.h"
 #include "VariableRateShadingImageManager.h"
-#include "LocalHeightFogRendering.h"
+#include "LocalFogVolumeRendering.h"
 #include "Shadows/ShadowScene.h"
 #include "Lumen/LumenHardwareRayTracingCommon.h"
 #include "SparseVolumeTexture/ISparseVolumeTextureStreamingManager.h"
diff --git a/Engine/Source/Runtime/Renderer/Private/LocalFogVolumeRendering.cpp b/Engine/Source/Runtime/Renderer/Private/LocalFogVolumeRendering.cpp
new file mode 100644
--- /dev/null
+++ b/Engine/Source/Runtime/Renderer/Private/LocalFogVolumeRendering.cpp
@@ -0,0 +1,439 @@
+// Copyright Epic Games, Inc. All Rights Reserved.
+
+#include "LocalFogVolumeRendering.h"
+#include "ScenePrivate.h"
+#include "RendererUtils.h"
+#include "ScreenPass.h"
+#include "LocalFogVolumeSceneProxy.h"
+#include "MobileBasePassRendering.h"
+
+
+// The runtime ON/OFF toggle
+static TAutoConsoleVariable<int32> CVarLocalHeightFog(
+	TEXT("r.LocalHeightFog"), 1,
+	TEXT("LocalHeightFog components are rendered when this is not 0, otherwise ignored.\n"),
+	ECVF_RenderThreadSafe);
+
+bool ShouldRenderLocalHeightFog(const FScene* Scene, const FSceneViewFamily& Family)
+{
+	const FEngineShowFlags EngineShowFlags = Family.EngineShowFlags;
+	if (Scene && Scene->HasAnyLocalHeightFog() && EngineShowFlags.Fog && !Family.UseDebugViewPS())
+	{
+		return CVarLocalHeightFog.GetValueOnRenderThread() > 0;
+	}
+	return false;
+}
+
+DECLARE_GPU_STAT(LocalHeightFogVolumes);
+
+/*=============================================================================
+	FScene functions
+=============================================================================*/
+
+void FScene::AddLocalHeightFog(class FLocalHeightFogSceneProxy* FogProxy)
+{
+	check(FogProxy);
+	FScene* Scene = this;
+
+	ENQUEUE_RENDER_COMMAND(FAddLocalHeightFogCommand)(
+		[Scene, FogProxy](FRHICommandListImmediate& RHICmdList)
+		{
+			check(!Scene->LocalHeightFogs.Contains(FogProxy));
+			Scene->LocalHeightFogs.Push(FogProxy);
+		} );
+}
+
+void FScene::RemoveLocalHeightFog(class FLocalHeightFogSceneProxy* FogProxy)
+{
+	check(FogProxy);
+	FScene* Scene = this;
+
+	ENQUEUE_RENDER_COMMAND(FRemoveLocalHeightFogCommand)(
+		[Scene, FogProxy](FRHICommandListImmediate& RHICmdList)
+		{
+			Scene->LocalHeightFogs.RemoveSingle(FogProxy);
+		} );
+}
+
+bool FScene::HasAnyLocalHeightFog() const
+{ 
+	return LocalHeightFogs.Num() > 0;
+}
+
+/*=============================================================================
+	Local height fog rendering common function
+=============================================================================*/
+
+void GetLocalFogVolumeSortingData(const FScene* Scene, FRDGBuilder& GraphBuilder, FLocalFogVolumeSortingData& Out)
+{
+	// No culling as of today
+	Out.LocalHeightFogInstanceCount = Scene->LocalHeightFogs.Num();
+	Out.LocalHeightFogInstanceCountFinal = 0;
+	Out.LocalHeightFogGPUInstanceData = (FLocalHeightFogGPUInstanceData*)GraphBuilder.Alloc(sizeof(FLocalHeightFogGPUInstanceData) * Out.LocalHeightFogInstanceCount, 16);
+	Out.LocalHeightFogCenterPos = (FVector*)GraphBuilder.Alloc(sizeof(FVector3f) * Out.LocalHeightFogInstanceCount, 16);
+	Out.LocalHeightFogSortKeys.SetNumUninitialized(Out.LocalHeightFogInstanceCount);
+	for (FLocalHeightFogSceneProxy* LHF : Scene->LocalHeightFogs)
+	{
+		if (LHF->FogDensity <= 0.0f)
+		{
+			continue; // this volume will never be visible
+		}
+
+		FTransform TransformScaleOnly;
+		TransformScaleOnly.SetScale3D(LHF->FogTransform.GetScale3D());
+
+		FLocalHeightFogGPUInstanceData* LocalHeightFogGPUInstanceDataIt = &Out.LocalHeightFogGPUInstanceData[Out.LocalHeightFogInstanceCountFinal];
+		LocalHeightFogGPUInstanceDataIt->Transform = FMatrix44f(LHF->FogTransform.ToMatrixWithScale());
+		LocalHeightFogGPUInstanceDataIt->InvTransform = LocalHeightFogGPUInstanceDataIt->Transform.Inverse();
+		LocalHeightFogGPUInstanceDataIt->InvTranformNoScale = FMatrix44f(LHF->FogTransform.ToMatrixNoScale()).Inverse();
+		LocalHeightFogGPUInstanceDataIt->TransformScaleOnly = FMatrix44f(TransformScaleOnly.ToMatrixWithScale());
+
+		LocalHeightFogGPUInstanceDataIt->Density = LHF->FogDensity;
+		LocalHeightFogGPUInstanceDataIt->HeightFalloff = LHF->FogHeightFalloff * 0.01f;	// This scale is used to have artist author reasonable range.
+		LocalHeightFogGPUInstanceDataIt->HeightOffset = LHF->FogHeightOffset;
+		LocalHeightFogGPUInstanceDataIt->RadialAttenuation = LHF->FogRadialAttenuation;
+
+		LocalHeightFogGPUInstanceDataIt->FogMode = float(LHF->FogMode);
+
+		LocalHeightFogGPUInstanceDataIt->Albedo = FVector3f(LHF->FogAlbedo);
+		LocalHeightFogGPUInstanceDataIt->PhaseG = LHF->FogPhaseG;
+		LocalHeightFogGPUInstanceDataIt->Emissive = FVector3f(LHF->FogEmissive);
+
+		Out.LocalHeightFogCenterPos[Out.LocalHeightFogInstanceCountFinal] = LHF->FogTransform.GetTranslation();
+
+		FLocalFogVolumeSortKey* LocalHeightFogSortKeysIt = &Out.LocalHeightFogSortKeys[Out.LocalHeightFogInstanceCountFinal];
+		LocalHeightFogSortKeysIt->FogVolume.Index = Out.LocalHeightFogInstanceCountFinal;
+		LocalHeightFogSortKeysIt->FogVolume.Distance = 0;	// Filled up right before sorting according to a view
+		LocalHeightFogSortKeysIt->FogVolume.Priority = LHF->FogSortPriority;
+
+		Out.LocalHeightFogInstanceCountFinal++;
+	}
+	// Shrink the array to only what is needed in order for the sort to correctly work on only what is needed.
+	Out.LocalHeightFogSortKeys.SetNum(Out.LocalHeightFogInstanceCountFinal, false/*bAllowShrinking*/);
+}
+
+void CreateViewLocalFogVolumeBufferSRV(FViewInfo& View, FRDGBuilder& GraphBuilder, FLocalFogVolumeSortingData& SortingData)
+{
+	static const uint32 SizeOfFloat4 = sizeof(float) * 4;
+	static const uint32 Float4CountInLocalHeightFogGPUInstanceData = sizeof(FLocalHeightFogGPUInstanceData) / SizeOfFloat4;
+	static_assert(sizeof(FLocalHeightFogGPUInstanceData) == Float4CountInLocalHeightFogGPUInstanceData * SizeOfFloat4); // The size of the structure must be a multiple of FVector4.
+
+	if (SortingData.LocalHeightFogInstanceCountFinal == 0)
+	{
+		View.LocalHeightFogGPUInstanceCount = 0;
+
+		static FLocalHeightFogGPUInstanceData DummyData;
+		View.LocalHeightFogGPUInstanceDataBufferSRV = GraphBuilder.CreateSRV(
+			CreateVertexBuffer(GraphBuilder, TEXT("LocalHeightFogGPUInstanceDataBuffer"), 
+			FRDGBufferDesc::CreateBufferDesc(SizeOfFloat4, Float4CountInLocalHeightFogGPUInstanceData), &DummyData, sizeof(FLocalHeightFogGPUInstanceData) * 1, ERDGInitialDataFlags::NoCopy),
+			PF_A32B32G32R32F);
+		return;
+	}
+
+	// 1. Sort all the volumes
+	const FVector ViewOrigin = View.ViewMatrices.GetViewOrigin();
+	for (uint32 i = 0; i < SortingData.LocalHeightFogInstanceCountFinal; ++i)
+	{
+		FVector FogCenterPos = SortingData.LocalHeightFogCenterPos[SortingData.LocalHeightFogSortKeys[i].FogVolume.Index];	// Recovered form the original array via index because the sorting of the previous view might have changed the order.
+		float DistancetoView = float((FogCenterPos - ViewOrigin).Size());
+		SortingData.LocalHeightFogSortKeys[i].FogVolume.Distance = *reinterpret_cast<uint32*>(&DistancetoView);
+	}
+	SortingData.LocalHeightFogSortKeys.Sort();
+
+	// 2. Create the buffer containing all the fog volume data instance sorted according to their key for the current view.
+	FLocalHeightFogGPUInstanceData* LocalHeightFogGPUSortedInstanceData = (FLocalHeightFogGPUInstanceData*)GraphBuilder.Alloc(sizeof(FLocalHeightFogGPUInstanceData) * SortingData.LocalHeightFogInstanceCountFinal, 16);
+	for (uint32 i = 0; i < SortingData.LocalHeightFogInstanceCountFinal; ++i)
+	{
+		// We could also have an indirection buffer on GPU but choosing to go with the sorting + copy on CPU since it is expected to not have many local height fog volumes.
+		LocalHeightFogGPUSortedInstanceData[i] = SortingData.LocalHeightFogGPUInstanceData[SortingData.LocalHeightFogSortKeys[i].FogVolume.Index];
+	}
+
+	// 3. Allocate buffer and initialize with sorted data to upload to GPU
+	const uint32 AllLocalHeightFogInstanceBytesFinal = sizeof(FLocalHeightFogGPUInstanceData) * SortingData.LocalHeightFogInstanceCountFinal;
+	FRDGBufferRef LocalHeightFogGPUInstanceDataBuffer = CreateVertexBuffer(
+		GraphBuilder, TEXT("LocalHeightFogGPUInstanceDataBuffer"),
+		FRDGBufferDesc::CreateBufferDesc(SizeOfFloat4, SortingData.LocalHeightFogInstanceCountFinal * Float4CountInLocalHeightFogGPUInstanceData), 
+		LocalHeightFogGPUSortedInstanceData, AllLocalHeightFogInstanceBytesFinal, ERDGInitialDataFlags::NoCopy);
+
+	View.LocalHeightFogGPUInstanceCount = SortingData.LocalHeightFogInstanceCountFinal;
+	View.LocalHeightFogGPUInstanceDataBufferSRV = GraphBuilder.CreateSRV(LocalHeightFogGPUInstanceDataBuffer, PF_A32B32G32R32F);
+}
+
+/*=============================================================================
+	Local height fog rendering - non mobile
+=============================================================================*/
+
+class FLocalHeightFogVS : public FGlobalShader
+{
+	DECLARE_GLOBAL_SHADER(FLocalHeightFogVS);
+	SHADER_USE_PARAMETER_STRUCT(FLocalHeightFogVS, FGlobalShader);
+
+	BEGIN_SHADER_PARAMETER_STRUCT(FParameters, )
+		SHADER_PARAMETER_STRUCT_REF(FViewUniformShaderParameters, View)
+		SHADER_PARAMETER_RDG_BUFFER_SRV(Buffer<float4>, LocalHeightFogInstances)
+	END_SHADER_PARAMETER_STRUCT()
+
+	using FPermutationDomain = TShaderPermutationDomain<>;
+
+	static bool ShouldCompilePermutation(const FGlobalShaderPermutationParameters& Parameters)
+	{
+		if (IsMobilePlatform(Parameters.Platform))
+		{
+			return false;
+		}
+		return true;
+	}
+};
+
+IMPLEMENT_GLOBAL_SHADER(FLocalHeightFogVS, "/Engine/Private/LocalFogVolume.usf", "LocalHeightFogSplatVS", SF_Vertex);
+
+class FLocalHeightFogPS : public FGlobalShader
+{
+	DECLARE_GLOBAL_SHADER(FLocalHeightFogPS);
+	SHADER_USE_PARAMETER_STRUCT(FLocalHeightFogPS, FGlobalShader);
+
+	BEGIN_SHADER_PARAMETER_STRUCT(FParameters, )
+		SHADER_PARAMETER_STRUCT_REF(FViewUniformShaderParameters, View)
+		SHADER_PARAMETER_RDG_BUFFER_SRV(Buffer<float4>, LocalHeightFogInstances)
+	END_SHADER_PARAMETER_STRUCT()
+
+	using FPermutationDomain = TShaderPermutationDomain<>;
+
+	static bool ShouldCompilePermutation(const FGlobalShaderPermutationParameters& Parameters)
+	{
+		if (IsMobilePlatform(Parameters.Platform))
+		{
+			return false;
+		}
+		return true;
+	}
+};
+
+IMPLEMENT_GLOBAL_SHADER(FLocalHeightFogPS, "/Engine/Private/LocalFogVolume.usf", "LocalHeightFogSplatPS", SF_Pixel);
+
+BEGIN_SHADER_PARAMETER_STRUCT(FLocalHeightFogPassParameters, )
+	SHADER_PARAMETER_STRUCT_INCLUDE(FLocalHeightFogVS::FParameters, VS)
+	SHADER_PARAMETER_STRUCT_INCLUDE(FLocalHeightFogPS::FParameters, PS)
+	SHADER_PARAMETER_RDG_UNIFORM_BUFFER(FSceneTextureUniformParameters, SceneTextures)
+	RENDER_TARGET_BINDING_SLOTS()
+END_SHADER_PARAMETER_STRUCT()
+ 
+void RenderLocalHeightFog(
+	const FScene* Scene,
+	TArray<FViewInfo>& Views,
+	FRDGBuilder& GraphBuilder,
+	const FMinimalSceneTextures& SceneTextures,
+	FRDGTextureRef LightShaftOcclusionTexture)
+{
+	uint32 LocalHeightFogInstanceCount = Scene->LocalHeightFogs.Num();
+	if (LocalHeightFogInstanceCount > 0)
+	{
+		RDG_GPU_STAT_SCOPE(GraphBuilder, LocalHeightFogVolumes);
+
+		FLocalFogVolumeSortingData SortingData;
+		GetLocalFogVolumeSortingData(Scene, GraphBuilder, SortingData);
+
+		if (SortingData.LocalHeightFogInstanceCountFinal > 0)
+		{
+			FRDGTextureRef SceneColorTexture = SceneTextures.Color.Resolve;
+
+			for (FViewInfo& View : Views)
+			{
+				CreateViewLocalFogVolumeBufferSRV(View, GraphBuilder, SortingData);
+				if (View.LocalHeightFogGPUInstanceCount == 0)
+				{
+					continue;
+				}
+
+				FLocalHeightFogPassParameters* PassParameters = GraphBuilder.AllocParameters<FLocalHeightFogPassParameters>();
+
+				PassParameters->VS.View = GetShaderBinding(View.ViewUniformBuffer);
+				PassParameters->VS.LocalHeightFogInstances = View.LocalHeightFogGPUInstanceDataBufferSRV;
+
+				PassParameters->PS.View = GetShaderBinding(View.ViewUniformBuffer);
+				PassParameters->PS.LocalHeightFogInstances = View.LocalHeightFogGPUInstanceDataBufferSRV;
+
+				PassParameters->SceneTextures = SceneTextures.UniformBuffer;
+				PassParameters->RenderTargets[0] = FRenderTargetBinding(SceneColorTexture, ERenderTargetLoadAction::ENoAction);
+
+				FLocalHeightFogVS::FPermutationDomain VSPermutationVector;
+				auto VertexShader = View.ShaderMap->GetShader< FLocalHeightFogVS >(VSPermutationVector);
+
+				FLocalHeightFogPS::FPermutationDomain PsPermutationVector;
+				auto PixelShader = View.ShaderMap->GetShader< FLocalHeightFogPS >(PsPermutationVector);
+
+				const FIntRect ViewRect = View.ViewRect;
+
+				ClearUnusedGraphResources(VertexShader, &PassParameters->VS);
+				ClearUnusedGraphResources(PixelShader, &PassParameters->PS);
+
+				uint32 LocalHeightFogGPUInstanceCount = View.LocalHeightFogGPUInstanceCount;
+				GraphBuilder.AddPass(
+					RDG_EVENT_NAME("RenderLocalHeightFog %u inst.", LocalHeightFogGPUInstanceCount),
+					PassParameters,
+					ERDGPassFlags::Raster,
+					[VertexShader, PixelShader, PassParameters, LocalHeightFogGPUInstanceCount, ViewRect](FRHICommandList& RHICmdList)
+				{
+					FGraphicsPipelineStateInitializer GraphicsPSOInit;
+					RHICmdList.ApplyCachedRenderTargets(GraphicsPSOInit);
+
+					RHICmdList.SetViewport(ViewRect.Min.X, ViewRect.Min.Y, 0.0f, ViewRect.Max.X, ViewRect.Max.Y, 1.0f);
+
+					// Render back faces only since camera may intersect
+					GraphicsPSOInit.RasterizerState = TStaticRasterizerState<FM_Solid, CM_CW>::GetRHI();
+					GraphicsPSOInit.DepthStencilState = TStaticDepthStencilState<false, CF_Always>::GetRHI();
+					GraphicsPSOInit.BlendState = TStaticBlendState<CW_RGBA, BO_Add, BF_One, BF_SourceAlpha, BO_Add, BF_Zero, BF_SourceAlpha>::GetRHI();
+					GraphicsPSOInit.PrimitiveType = PT_TriangleList;
+
+					GraphicsPSOInit.BoundShaderState.VertexDeclarationRHI = GetVertexDeclarationFVector4();
+					GraphicsPSOInit.BoundShaderState.VertexShaderRHI = VertexShader.GetVertexShader();
+					GraphicsPSOInit.BoundShaderState.PixelShaderRHI = PixelShader.GetPixelShader();
+
+					SetGraphicsPipelineState(RHICmdList, GraphicsPSOInit, 0);
+
+					SetShaderParameters(RHICmdList, VertexShader, VertexShader.GetVertexShader(), PassParameters->VS);
+					SetShaderParameters(RHICmdList, PixelShader, PixelShader.GetPixelShader(), PassParameters->PS);
+
+					RHICmdList.SetStreamSource(0, GetUnitCubeVertexBuffer(), 0);
+
+					RHICmdList.DrawIndexedPrimitive(GetUnitCubeIndexBuffer()
+						, 0									//BaseVertexIndex
+						, 0									//FirstInstance
+						, 8									//uint32 NumVertices
+						, 0									//uint32 StartIndex
+						, UE_ARRAY_COUNT(GCubeIndices) / 3	//uint32 NumPrimitives
+						, LocalHeightFogGPUInstanceCount	//uint32 NumInstances
+					);
+				});
+			}
+		}
+	}
+}
+
+/*=============================================================================
+	Local height fog rendering - mobile
+=============================================================================*/
+	
+
+class FMobileLocalHeightFogVS : public FGlobalShader
+{
+	DECLARE_GLOBAL_SHADER(FMobileLocalHeightFogVS);
+	SHADER_USE_PARAMETER_STRUCT(FMobileLocalHeightFogVS, FGlobalShader);
+
+	BEGIN_SHADER_PARAMETER_STRUCT(FParameters, )
+		SHADER_PARAMETER_STRUCT_INCLUDE(FViewShaderParameters, View)
+		SHADER_PARAMETER_RDG_BUFFER_SRV(Buffer<float4>, LocalHeightFogInstances)
+	END_SHADER_PARAMETER_STRUCT()
+
+	using FPermutationDomain = TShaderPermutationDomain<>;
+
+	static bool ShouldCompilePermutation(const FGlobalShaderPermutationParameters& Parameters)
+	{
+		if (!IsMobilePlatform(Parameters.Platform))
+		{
+			return false;
+		}
+		return true;
+	}
+
+	static void ModifyCompilationEnvironment(const FGlobalShaderPermutationParameters& Parameters, FShaderCompilerEnvironment& OutEnvironment)
+	{
+		FGlobalShader::ModifyCompilationEnvironment(Parameters, OutEnvironment);
+		const bool bMobileForceDepthRead = MobileUsesFullDepthPrepass(Parameters.Platform);
+		OutEnvironment.SetDefine(TEXT("IS_MOBILE_DEPTHREAD_SUBPASS"), bMobileForceDepthRead ? 0u : 1u);
+	}
+};
+
+IMPLEMENT_GLOBAL_SHADER(FMobileLocalHeightFogVS, "/Engine/Private/LocalFogVolume.usf", "LocalHeightFogSplatVS", SF_Vertex);
+
+class FMobileLocalHeightFogPS : public FGlobalShader
+{
+	DECLARE_GLOBAL_SHADER(FMobileLocalHeightFogPS);
+	SHADER_USE_PARAMETER_STRUCT(FMobileLocalHeightFogPS, FGlobalShader);
+
+	BEGIN_SHADER_PARAMETER_STRUCT(FParameters, )
+		SHADER_PARAMETER_STRUCT_INCLUDE(FViewShaderParameters, View)
+		SHADER_PARAMETER_RDG_UNIFORM_BUFFER(FMobileBasePassUniformParameters, MobileBasePass)
+		SHADER_PARAMETER_RDG_BUFFER_SRV(Buffer<float4>, LocalHeightFogInstances)
+	END_SHADER_PARAMETER_STRUCT()
+
+	using FPermutationDomain = TShaderPermutationDomain<>;
+
+	static bool ShouldCompilePermutation(const FGlobalShaderPermutationParameters& Parameters)
+	{
+		if (!IsMobilePlatform(Parameters.Platform))
+		{
+			return false;
+		}
+		return true;
+	}
+
+	static void ModifyCompilationEnvironment(const FGlobalShaderPermutationParameters& Parameters, FShaderCompilerEnvironment& OutEnvironment)
+	{
+		FGlobalShader::ModifyCompilationEnvironment(Parameters, OutEnvironment);
+		const bool bMobileForceDepthRead = MobileUsesFullDepthPrepass(Parameters.Platform);
+		OutEnvironment.SetDefine(TEXT("IS_MOBILE_DEPTHREAD_SUBPASS"), bMobileForceDepthRead ? 0u : 1u);
+	}
+};
+
+IMPLEMENT_GLOBAL_SHADER(FMobileLocalHeightFogPS, "/Engine/Private/LocalFogVolume.usf", "LocalHeightFogSplatPS", SF_Pixel);
+
+void RenderLocalHeightFogMobile(
+	FRHICommandList& RHICmdList,
+	const FViewInfo& View)
+{
+	if (View.LocalHeightFogGPUInstanceCount == 0)
+	{
+		return;
+	}
+
+	SCOPED_DRAW_EVENT(RHICmdList, LocalHeightFogVolumes);
+
+	FMobileLocalHeightFogVS::FPermutationDomain VSPermutationVector;
+	auto VertexShader = View.ShaderMap->GetShader< FMobileLocalHeightFogVS >(VSPermutationVector);
+
+	FMobileLocalHeightFogPS::FPermutationDomain PsPermutationVector;
+	auto PixelShader = View.ShaderMap->GetShader< FMobileLocalHeightFogPS >(PsPermutationVector);
+
+	const FIntRect ViewRect = View.ViewRect;
+
+	FGraphicsPipelineStateInitializer GraphicsPSOInit;
+	RHICmdList.ApplyCachedRenderTargets(GraphicsPSOInit);
+
+	RHICmdList.SetViewport(ViewRect.Min.X, ViewRect.Min.Y, 0.0f, ViewRect.Max.X, ViewRect.Max.Y, 1.0f);
+
+	// Render back faces only since camera may intersect
+	GraphicsPSOInit.RasterizerState = TStaticRasterizerState<FM_Solid, CM_CW>::GetRHI();
+	GraphicsPSOInit.DepthStencilState = TStaticDepthStencilState<false, CF_Always>::GetRHI();
+	GraphicsPSOInit.BlendState = TStaticBlendState<CW_RGBA, BO_Add, BF_One, BF_SourceAlpha, BO_Add, BF_Zero, BF_SourceAlpha>::GetRHI();
+	GraphicsPSOInit.PrimitiveType = PT_TriangleList;
+
+	GraphicsPSOInit.BoundShaderState.VertexDeclarationRHI = GetVertexDeclarationFVector4();
+	GraphicsPSOInit.BoundShaderState.VertexShaderRHI = VertexShader.GetVertexShader();
+	GraphicsPSOInit.BoundShaderState.PixelShaderRHI = PixelShader.GetPixelShader();
+
+	SetGraphicsPipelineState(RHICmdList, GraphicsPSOInit, 0);
+
+	FMobileLocalHeightFogVS::FParameters VSParameters;
+	VSParameters.View = View.GetShaderParameters();
+	VSParameters.LocalHeightFogInstances = View.LocalHeightFogGPUInstanceDataBufferSRV;
+	SetShaderParameters(RHICmdList, VertexShader, VertexShader.GetVertexShader(), VSParameters);
+
+	FMobileLocalHeightFogPS::FParameters PSParameters;
+	PSParameters.View = View.GetShaderParameters();
+	PSParameters.LocalHeightFogInstances = View.LocalHeightFogGPUInstanceDataBufferSRV;
+	// PSParameters.MobileBasePass filled up by the RDG pass parameters.
+	SetShaderParameters(RHICmdList, PixelShader, PixelShader.GetPixelShader(), PSParameters);
+
+	RHICmdList.SetStreamSource(0, GetUnitCubeVertexBuffer(), 0);
+
+	RHICmdList.DrawIndexedPrimitive(
+		GetUnitCubeIndexBuffer()
+		, 0										//BaseVertexIndex
+		, 0										//FirstInstance
+		, 8										//uint32 NumVertices
+		, 0										//uint32 StartIndex
+		, UE_ARRAY_COUNT(GCubeIndices) / 3		//uint32 NumPrimitives
+		, View.LocalHeightFogGPUInstanceCount	//uint32 NumInstances
+	);
+}
diff --git a/Engine/Source/Runtime/Renderer/Private/LocalFogVolumeRendering.h b/Engine/Source/Runtime/Renderer/Private/LocalFogVolumeRendering.h
new file mode 100644
--- /dev/null
+++ b/Engine/Source/Runtime/Renderer/Private/LocalFogVolumeRendering.h
@@ -0,0 +1,99 @@
+// Copyright Epic Games, Inc. All Rights Reserved.
+
+#pragma once
+
+#include "EngineDefines.h"
+#include "RenderGraph.h"
+#include "RenderResource.h"
+#include "RenderGraphResources.h"
+#include "SceneView.h"
+#include "Containers/Array.h"
+
+class FScene;
+class FViewInfo;
+struct FMinimalSceneTextures;
+
+
+
+/*=============================================================================
+	Local height fog rendering common data
+=============================================================================*/
+
+class FLocalHeightFogGPUInstanceData
+{
+public:
+	FMatrix44f Transform;
+	FMatrix44f InvTransform;
+
+	FMatrix44f InvTranformNoScale;
+	FMatrix44f TransformScaleOnly;
+
+	float Density;
+	float HeightFalloff;
+	float HeightOffset;
+	float RadialAttenuation;
+
+	FVector3f Albedo;
+	float PhaseG;
+	FVector3f Emissive;
+	float FogMode;
+};
+
+class RENDERER_API FLocalFogVolumeSortKey
+{
+public:
+	union
+	{
+		uint64 PackedData;
+
+		struct
+		{
+			uint64 Index : 16; // Index of the volume
+			uint64 Distance : 32; // then by distance
+			uint64 Priority : 16; // First order by priority
+		} FogVolume;
+	};
+
+	FORCEINLINE bool operator!=(FLocalFogVolumeSortKey B) const
+	{
+		return PackedData != B.PackedData;
+	}
+
+	FORCEINLINE bool operator<(FLocalFogVolumeSortKey B) const
+	{
+		return PackedData > B.PackedData;
+	}
+};
+
+// This can be generated once for a scene and then shared between all the views to generate the GPU buffer data needed for rendering while accounting for sorting.
+struct FLocalFogVolumeSortingData
+{
+	uint32 LocalHeightFogInstanceCount;
+	uint32 LocalHeightFogInstanceCountFinal;
+	FLocalHeightFogGPUInstanceData* LocalHeightFogGPUInstanceData;
+	FVector* LocalHeightFogCenterPos;
+	TArray<FLocalFogVolumeSortKey>	LocalHeightFogSortKeys;
+};
+
+
+/*=============================================================================
+	Local height fog rendering functions
+=============================================================================*/
+
+bool ShouldRenderLocalHeightFog(const FScene* Scene, const FSceneViewFamily& Family);
+
+void GetLocalFogVolumeSortingData(const FScene* Scene, FRDGBuilder& GraphBuilder, FLocalFogVolumeSortingData& Out);
+
+void CreateViewLocalFogVolumeBufferSRV(FViewInfo& View, FRDGBuilder& GraphBuilder, FLocalFogVolumeSortingData& SortingData);
+
+void RenderLocalHeightFog(
+	const FScene* Scene,
+	TArray<FViewInfo>& Views,
+	FRDGBuilder& GraphBuilder,
+	const FMinimalSceneTextures& SceneTextures,
+	FRDGTextureRef LightShaftOcclusionTexture);
+
+void RenderLocalHeightFogMobile(
+	FRHICommandList& RHICmdList, 
+	const FViewInfo& View);
+
diff --git a/Engine/Source/Runtime/Renderer/Private/LocalHeightFogRendering.cpp b/Engine/Source/Runtime/Renderer/Private/LocalHeightFogRendering.cpp
deleted file mode 100644
diff --git a/Engine/Source/Runtime/Renderer/Private/LocalHeightFogRendering.h b/Engine/Source/Runtime/Renderer/Private/LocalHeightFogRendering.h
deleted file mode 100644
diff --git a/Engine/Source/Runtime/Renderer/Private/MobileShadingRenderer.cpp b/Engine/Source/Runtime/Renderer/Private/MobileShadingRenderer.cpp
--- a/Engine/Source/Runtime/Renderer/Private/MobileShadingRenderer.cpp
+++ b/Engine/Source/Runtime/Renderer/Private/MobileShadingRenderer.cpp
@@ -67,7 +67,7 @@
 #include "IESTextureManager.h"
 #include "SceneUniformBuffer.h"
 #include "Engine/SpecularProfile.h"
-#include "LocalHeightFogRendering.h"
+#include "LocalFogVolumeRendering.h"
 
 uint32 GetShadowQuality();
 
