Rename from LocalHeightFog to LocalFogVolume in code

diff --git a/Engine/Shaders/Private/LocalFogVolume.usf b/Engine/Shaders/Private/LocalFogVolume.usf
--- a/Engine/Shaders/Private/LocalFogVolume.usf
+++ b/Engine/Shaders/Private/LocalFogVolume.usf
@@ -10,7 +10,7 @@
 #include "ParticipatingMediaCommon.ush"
 
 
-struct FLocalHeightFogGPUInstanceData
+struct FLocalFogVolumeGPUInstanceData
 {
 	float4x4 Transform;
 	float4x4 InvTransform;
@@ -31,29 +31,29 @@ struct FLocalHeightFogGPUInstanceData
 
 // Structured buffers cannot be sampled from the vertex shader on some mobile platforms.
 // That is why simple Buffers are used, because this tech needs to run on mobile platforms.
-Buffer<float4> LocalHeightFogInstances;
-FLocalHeightFogGPUInstanceData GetLocalHeightFogGPUInstanceData(uint Index)
+Buffer<float4> LocalFogVolumeInstances;
+FLocalFogVolumeGPUInstanceData GetLocalFogVolumeGPUInstanceData(uint Index)
 {
-	FLocalHeightFogGPUInstanceData Data;
+	FLocalFogVolumeGPUInstanceData Data;
 	uint Offset = Index * 19;
 
-	Data.Transform		= float4x4(LocalHeightFogInstances[Offset + 0], LocalHeightFogInstances[Offset + 1], LocalHeightFogInstances[Offset + 2], LocalHeightFogInstances[Offset + 3]);
-	Data.InvTransform	= float4x4(LocalHeightFogInstances[Offset + 4], LocalHeightFogInstances[Offset + 5], LocalHeightFogInstances[Offset + 6], LocalHeightFogInstances[Offset + 7]);
+	Data.Transform		= float4x4(LocalFogVolumeInstances[Offset + 0], LocalFogVolumeInstances[Offset + 1], LocalFogVolumeInstances[Offset + 2], LocalFogVolumeInstances[Offset + 3]);
+	Data.InvTransform	= float4x4(LocalFogVolumeInstances[Offset + 4], LocalFogVolumeInstances[Offset + 5], LocalFogVolumeInstances[Offset + 6], LocalFogVolumeInstances[Offset + 7]);
 
-	Data.InvTranformNoScale = float4x4(LocalHeightFogInstances[Offset + 8], LocalHeightFogInstances[Offset + 9], LocalHeightFogInstances[Offset + 10], LocalHeightFogInstances[Offset + 11]);
-	Data.TransformScaleOnly = float4x4(LocalHeightFogInstances[Offset + 12], LocalHeightFogInstances[Offset + 13], LocalHeightFogInstances[Offset + 14], LocalHeightFogInstances[Offset + 15]);
+	Data.InvTranformNoScale = float4x4(LocalFogVolumeInstances[Offset + 8], LocalFogVolumeInstances[Offset + 9], LocalFogVolumeInstances[Offset + 10], LocalFogVolumeInstances[Offset + 11]);
+	Data.TransformScaleOnly = float4x4(LocalFogVolumeInstances[Offset + 12], LocalFogVolumeInstances[Offset + 13], LocalFogVolumeInstances[Offset + 14], LocalFogVolumeInstances[Offset + 15]);
 
-	float4 Temp0 = LocalHeightFogInstances[Offset + 16];
+	float4 Temp0 = LocalFogVolumeInstances[Offset + 16];
 	Data.Density			= Temp0.x;
 	Data.HeightFalloff		= Temp0.y;
 	Data.HeightOffset		= Temp0.z;
 	Data.RadialAttenuation	= Temp0.w;
 
-	float4 Temp1 = LocalHeightFogInstances[Offset + 17];
+	float4 Temp1 = LocalFogVolumeInstances[Offset + 17];
 	Data.Albedo		= Temp1.xyz;
 	Data.PhaseG		= Temp1.w;
 
-	float4 Temp2 = LocalHeightFogInstances[Offset + 18];
+	float4 Temp2 = LocalFogVolumeInstances[Offset + 18];
 	Data.Emissive	= Temp2.xyz;
 	Data.FogMode	= Temp2.w;
 
@@ -61,23 +61,23 @@ FLocalHeightFogGPUInstanceData GetLocalHeightFogGPUInstanceData(uint Index)
 }
 
 
-struct FLocalHeightFogVertexOutput
+struct FLocalFogVolumeVertexOutput
 {
 	nointerpolation uint InstanceId : TEXCOORD0;
 };
 
-void LocalHeightFogSplatVS(
+void LocalFogVolumeSplatVS(
 	in float4 InPosition : ATTRIBUTE0,
 	uint InstanceIndex : SV_InstanceID,
 	out float4 OutPosition : SV_POSITION,
-	out FLocalHeightFogVertexOutput OtherVertexOutput
+	out FLocalFogVolumeVertexOutput OtherVertexOutput
 )
 {
 	ResolvedView = ResolveView();
 
 	float4 LocalPosition = InPosition;
 
-	FLocalHeightFogGPUInstanceData FogInstance = GetLocalHeightFogGPUInstanceData(InstanceIndex);
+	FLocalFogVolumeGPUInstanceData FogInstance = GetLocalFogVolumeGPUInstanceData(InstanceIndex);
 
 	float4x4 FogTransform = FogInstance.Transform;
 
@@ -97,7 +97,7 @@ struct FFogData
 };
 
 // Height fog is evaluated according to spatially varying density defined as "Density * Exp[-HeightFalloff * z]"
-FFogData EvaluateHeightFogIntegral(FLocalHeightFogGPUInstanceData FogInstance,  float3 CameraPos, float3 RayStart, float3 RayDir, float RayLength, float RadialAttenuation)
+FFogData EvaluateHeightFogIntegral(FLocalFogVolumeGPUInstanceData FogInstance,  float3 CameraPos, float3 RayStart, float3 RayDir, float RayLength, float RadialAttenuation)
 {
 	FFogData FogData;
 	
@@ -138,7 +138,7 @@ FFogData EvaluateHeightFogIntegral(FLocalHeightFogGPUInstanceData FogInstance,
 }
 
 
-FFogData EvaluateRadialFogIntegral(FLocalHeightFogGPUInstanceData FogInstance, float3 CameraPos, float3 RayStart, float3 RayDir, float RayLength, float RadialAttenuation)
+FFogData EvaluateRadialFogIntegral(FLocalFogVolumeGPUInstanceData FogInstance, float3 CameraPos, float3 RayStart, float3 RayDir, float RayLength, float RadialAttenuation)
 {
 	FFogData FogData = (FFogData)0;
 
@@ -185,7 +185,7 @@ FFogData EvaluateRadialFogIntegral(FLocalHeightFogGPUInstanceData FogInstance, f
 }
 
 
-float3 ComputeFogInscattering(in FLocalHeightFogGPUInstanceData FogInstance, in FFogData FogData, in float3 RayDirWorld)
+float3 ComputeFogInscattering(in FLocalFogVolumeGPUInstanceData FogInstance, in FFogData FogData, in float3 RayDirWorld)
 {
 	half3 InScattering = 0;
 
@@ -211,9 +211,9 @@ float3 ComputeFogInscattering(in FLocalHeightFogGPUInstanceData FogInstance, in
 
 #define LHFV_FOLLOWS_ROTATIONS 1
 
-void LocalHeightFogSplatPS(
+void LocalFogVolumeSplatPS(
 	in float4 SVPos : SV_POSITION,
-	in FLocalHeightFogVertexOutput OtherVertexOutput,
+	in FLocalFogVolumeVertexOutput OtherVertexOutput,
 	out float4 OutColor : SV_Target0)
 {
 	ResolvedView = ResolveView();
@@ -223,7 +223,7 @@ void LocalHeightFogSplatPS(
 	float3 CamRayTranslatedWorldDir = normalize(SvPositionToTranslatedWorld(float4(SVPos.xy, 0.5, 1.0)));
 	float3 CamRayWorldDir = CamRayTranslatedWorldDir;
 
-	FLocalHeightFogGPUInstanceData FogInstance = GetLocalHeightFogGPUInstanceData(InstanceIndex);
+	FLocalFogVolumeGPUInstanceData FogInstance = GetLocalFogVolumeGPUInstanceData(InstanceIndex);
 
 	OutColor = float4(0, 0, 0, 1);
 
diff --git a/Engine/Source/Editor/ComponentVisualizers/Private/ComponentVisualizers.cpp b/Engine/Source/Editor/ComponentVisualizers/Private/ComponentVisualizers.cpp
--- a/Engine/Source/Editor/ComponentVisualizers/Private/ComponentVisualizers.cpp
+++ b/Engine/Source/Editor/ComponentVisualizers/Private/ComponentVisualizers.cpp
@@ -63,7 +63,7 @@ void FComponentVisualizersModule::StartupModule()
 	RegisterComponentVisualizer(UDecalComponent::StaticClass()->GetFName(), MakeShareable(new FDecalComponentVisualizer));
 	RegisterComponentVisualizer(UStereoLayerComponent::StaticClass()->GetFName(), MakeShareable(new FStereoLayerComponentVisualizer));
 	RegisterComponentVisualizer(UWorldPartitionStreamingSourceComponent::StaticClass()->GetFName(), MakeShareable(new FWorldPartitionStreamingSourceComponentVisualizer));
-	RegisterComponentVisualizer(ULocalHeightFogComponent::StaticClass()->GetFName(), MakeShareable(new FLocalHeightFogComponentVisualizer));
+	RegisterComponentVisualizer(ULocalFogVolumeComponent::StaticClass()->GetFName(), MakeShareable(new FLocalFogVolumeComponentVisualizer));
 }
 
 void FComponentVisualizersModule::ShutdownModule()
diff --git a/Engine/Source/Editor/ComponentVisualizers/Private/LocalFogVolumeComponentVisualizer.cpp b/Engine/Source/Editor/ComponentVisualizers/Private/LocalFogVolumeComponentVisualizer.cpp
--- a/Engine/Source/Editor/ComponentVisualizers/Private/LocalFogVolumeComponentVisualizer.cpp
+++ b/Engine/Source/Editor/ComponentVisualizers/Private/LocalFogVolumeComponentVisualizer.cpp
@@ -7,17 +7,17 @@
 #include "SceneView.h"
 #include "SceneManagement.h"
 
-void FLocalHeightFogComponentVisualizer::DrawVisualization( const UActorComponent* Component, const FSceneView* View, FPrimitiveDrawInterface* PDI )
+void FLocalFogVolumeComponentVisualizer::DrawVisualization( const UActorComponent* Component, const FSceneView* View, FPrimitiveDrawInterface* PDI )
 {
 	if(View->Family->EngineShowFlags.CompositeEditorPrimitives)
 	{
-		const ULocalHeightFogComponent* LocalHeightFog = Cast<const ULocalHeightFogComponent>(Component);
-		if(LocalHeightFog != NULL)
+		const ULocalFogVolumeComponent* LocalFogVolume = Cast<const ULocalFogVolumeComponent>(Component);
+		if(LocalFogVolume != NULL)
 		{
-			FTransform LocalHeightFogTransform = LocalHeightFog->GetComponentTransform();
+			FTransform LocalFogVolumeTransform = LocalFogVolume->GetComponentTransform();
 
 			// Draw local fog volume spherical shape radius
-			DrawWireSphereAutoSides(PDI, LocalHeightFogTransform, FColor(200, 255, 255), 1.0f, SDPG_World);
+			DrawWireSphereAutoSides(PDI, LocalFogVolumeTransform, FColor(200, 255, 255), 1.0f, SDPG_World);
 		}
 	}
 }
diff --git a/Engine/Source/Editor/ComponentVisualizers/Public/LocalFogVolumeComponentVisualizer.h b/Engine/Source/Editor/ComponentVisualizers/Public/LocalFogVolumeComponentVisualizer.h
--- a/Engine/Source/Editor/ComponentVisualizers/Public/LocalFogVolumeComponentVisualizer.h
+++ b/Engine/Source/Editor/ComponentVisualizers/Public/LocalFogVolumeComponentVisualizer.h
@@ -9,7 +9,7 @@ class FPrimitiveDrawInterface;
 class FSceneView;
 class UActorComponent;
 
-class COMPONENTVISUALIZERS_API FLocalHeightFogComponentVisualizer : public FComponentVisualizer
+class COMPONENTVISUALIZERS_API FLocalFogVolumeComponentVisualizer : public FComponentVisualizer
 {
 public:
 	//~ Begin FComponentVisualizer Interface
diff --git a/Engine/Source/Editor/LevelEditor/Private/LevelEditorOutlinerSettings.cpp b/Engine/Source/Editor/LevelEditor/Private/LevelEditorOutlinerSettings.cpp
--- a/Engine/Source/Editor/LevelEditor/Private/LevelEditorOutlinerSettings.cpp
+++ b/Engine/Source/Editor/LevelEditor/Private/LevelEditorOutlinerSettings.cpp
@@ -277,7 +277,7 @@ void FLevelEditorOutlinerSettings::CreateDefaultFilters()
 		FindOrAddClassFilter(AExponentialHeightFog::StaticClass(), *EnvironmentFilterCategory);
 		FindOrAddClassFilter(AInstancedFoliageActor::StaticClass(), *EnvironmentFilterCategory);
 		FindOrAddClassFilter(ASkyAtmosphere::StaticClass(), *EnvironmentFilterCategory);
-		FindOrAddClassFilter(ALocalHeightFog::StaticClass(), *EnvironmentFilterCategory);
+		FindOrAddClassFilter(ALocalFogVolume::StaticClass(), *EnvironmentFilterCategory);
 		FindOrAddClassFilter(AVolumetricCloud::StaticClass(), *EnvironmentFilterCategory);
 	}
 	
diff --git a/Engine/Source/Editor/PlacementMode/Private/PlacementModeModule.cpp b/Engine/Source/Editor/PlacementMode/Private/PlacementModeModule.cpp
--- a/Engine/Source/Editor/PlacementMode/Private/PlacementModeModule.cpp
+++ b/Engine/Source/Editor/PlacementMode/Private/PlacementModeModule.cpp
@@ -234,7 +234,7 @@ void FPlacementModeModule::StartupModule()
 		Category->Items.Add(CreateID(), MakeShareable(new FPlaceableItem(*UActorFactorySkyAtmosphere::StaticClass(), SortOrder += 10)));
 		Category->Items.Add(CreateID(), MakeShareable(new FPlaceableItem(*UActorFactoryVolumetricCloud::StaticClass(), SortOrder += 10)));
 		Category->Items.Add(CreateID(), MakeShareable(new FPlaceableItem(*UActorFactoryExponentialHeightFog::StaticClass(), SortOrder += 10)));
-		Category->Items.Add(CreateID(), MakeShareable(new FPlaceableItem(*UActorFactoryLocalHeightFog::StaticClass(), SortOrder += 10)));
+		Category->Items.Add(CreateID(), MakeShareable(new FPlaceableItem(*UActorFactoryLocalFogVolume::StaticClass(), SortOrder += 10)));
 		Category->Items.Add(CreateID(), MakeShareable(new FPlaceableItem(*UActorFactorySphereReflectionCapture::StaticClass(), SortOrder += 10)));
 		Category->Items.Add(CreateID(), MakeShareable(new FPlaceableItem(*UActorFactoryBoxReflectionCapture::StaticClass(), SortOrder += 10)));
 		Category->Items.Add(CreateID(), MakeShareable(new FPlaceableItem(*UActorFactoryPlanarReflection::StaticClass(), SortOrder += 10)));
diff --git a/Engine/Source/Editor/UnrealEd/Classes/ActorFactories/ActorFactoryLocalFogVolume.h b/Engine/Source/Editor/UnrealEd/Classes/ActorFactories/ActorFactoryLocalFogVolume.h
--- a/Engine/Source/Editor/UnrealEd/Classes/ActorFactories/ActorFactoryLocalFogVolume.h
+++ b/Engine/Source/Editor/UnrealEd/Classes/ActorFactories/ActorFactoryLocalFogVolume.h
@@ -8,7 +8,7 @@
 #include "ActorFactoryLocalFogVolume.generated.h"
 
 UCLASS(MinimalAPI, config = Editor, collapsecategories, hidecategories = Object)
-class UActorFactoryLocalHeightFog : public UActorFactory
+class UActorFactoryLocalFogVolume : public UActorFactory
 {
 	GENERATED_UCLASS_BODY()
 };
diff --git a/Engine/Source/Editor/UnrealEd/Private/Factories/ActorFactory.cpp b/Engine/Source/Editor/UnrealEd/Private/Factories/ActorFactory.cpp
--- a/Engine/Source/Editor/UnrealEd/Private/Factories/ActorFactory.cpp
+++ b/Engine/Source/Editor/UnrealEd/Private/Factories/ActorFactory.cpp
@@ -1877,13 +1877,13 @@ UActorFactoryExponentialHeightFog::UActorFactoryExponentialHeightFog(const FObje
 }
 
 /*-----------------------------------------------------------------------------
-UActorFactoryLocalHeightFog
+UActorFactoryLocalFogVolume
 -----------------------------------------------------------------------------*/
-UActorFactoryLocalHeightFog::UActorFactoryLocalHeightFog(const FObjectInitializer& ObjectInitializer)
+UActorFactoryLocalFogVolume::UActorFactoryLocalFogVolume(const FObjectInitializer& ObjectInitializer)
 	: Super(ObjectInitializer)
 {
-	DisplayName = LOCTEXT("LocalHeightFogDisplayName", "Local Height Fog");
-	NewActorClass = ALocalHeightFog::StaticClass();
+	DisplayName = LOCTEXT("LocalFogVolumeDisplayName", "Local Height Fog");
+	NewActorClass = ALocalFogVolume::StaticClass();
 }
 
 /*-----------------------------------------------------------------------------
diff --git a/Engine/Source/Runtime/Engine/Classes/Components/LocalFogVolumeComponent.h b/Engine/Source/Runtime/Engine/Classes/Components/LocalFogVolumeComponent.h
--- a/Engine/Source/Runtime/Engine/Classes/Components/LocalFogVolumeComponent.h
+++ b/Engine/Source/Runtime/Engine/Classes/Components/LocalFogVolumeComponent.h
@@ -7,25 +7,25 @@
 #include "Math/Color.h"
 #include "LocalFogVolumeComponent.generated.h"
 
-class FLocalHeightFogSceneProxy;
+class FLocalFogVolumeSceneProxy;
 
 UENUM()
 enum class ELocalFogMode : uint8
 {
-	LocalHeightFog = 0,
+	LocalFogVolume = 0,
 	LocalSphereFog = 1,
 };
 
 UCLASS(ClassGroup = Rendering, collapsecategories, hidecategories = (Object, Mobility, Activation, "Components|Activation"), editinlinenew, meta = (BlueprintSpawnableComponent), MinimalAPI)
-class ULocalHeightFogComponent : public USceneComponent
+class ULocalFogVolumeComponent : public USceneComponent
 {
 	GENERATED_UCLASS_BODY()
 
-	~ULocalHeightFogComponent();
+	~ULocalFogVolumeComponent();
 
 	/** Controls the softness of the transition region when the volume is fading out. */
 	UPROPERTY(EditAnywhere, BlueprintReadOnly, Category = "Fog Mode")
-	ELocalFogMode FogMode = ELocalFogMode::LocalHeightFog;
+	ELocalFogMode FogMode = ELocalFogMode::LocalFogVolume;
 
 	/** Global density factor for this fog. */
 	UPROPERTY(EditAnywhere, BlueprintReadOnly, interp, Category = "Fog Distribution", meta = (UIMin = "0", UIMax = "10.0", SliderExponent = 2.0, ClampMin = 0.0))
@@ -69,7 +69,7 @@ public:
 	//~ End UObject Interface
 
 	/** Callback to create the rendering thread mirror. */
-	ENGINE_API FLocalHeightFogSceneProxy* CreateSceneProxy();
+	ENGINE_API FLocalFogVolumeSceneProxy* CreateSceneProxy();
 
 protected:
 	//~ Begin UActorComponent Interface.
@@ -82,6 +82,6 @@ protected:
 
 private:
 
-	FLocalHeightFogSceneProxy* LocalHeightFogSceneProxy;
+	FLocalFogVolumeSceneProxy* LocalFogVolumeSceneProxy;
 };
 
diff --git a/Engine/Source/Runtime/Engine/Classes/Engine/LocalFogVolume.h b/Engine/Source/Runtime/Engine/Classes/Engine/LocalFogVolume.h
--- a/Engine/Source/Runtime/Engine/Classes/Engine/LocalFogVolume.h
+++ b/Engine/Source/Runtime/Engine/Classes/Engine/LocalFogVolume.h
@@ -6,14 +6,14 @@
 #include "GameFramework/Info.h"
 #include "LocalFogVolume.generated.h"
 
-class ULocalHeightFogComponent;
+class ULocalFogVolumeComponent;
 
 /**
  *	Actor used to position a local fog volume in the scene.
  *	@see https://docs.unrealengine.com/???
  */
 UCLASS(showcategories = (Movement, Rendering, Transformation, DataLayers, "Input|MouseInput", "Input|TouchInput"), ClassGroup = Fog, hidecategories = (Info, Object, Input), MinimalAPI)
-class ALocalHeightFog : public AInfo
+class ALocalFogVolume : public AInfo
 {
 	GENERATED_UCLASS_BODY()
 
@@ -24,11 +24,11 @@ private:
 
 	/** Object used to visualize the local fog volume */
 	UPROPERTY(BlueprintReadOnly, VisibleAnywhere, Category = Fog, meta = (AllowPrivateAccess = "true"))
-	TObjectPtr<ULocalHeightFogComponent> LocalHeightFogVolume;
+	TObjectPtr<ULocalFogVolumeComponent> LocalFogVolumeVolume;
 
 public:
 
-	/** Returns LocalHeightFogVolume subobject **/
-	ULocalHeightFogComponent* GetComponent() const { return LocalHeightFogVolume; }
+	/** Returns LocalFogVolumeVolume subobject **/
+	ULocalFogVolumeComponent* GetComponent() const { return LocalFogVolumeVolume; }
 };
 
diff --git a/Engine/Source/Runtime/Engine/Private/Components/LocalFogVolumeComponent.cpp b/Engine/Source/Runtime/Engine/Private/Components/LocalFogVolumeComponent.cpp
--- a/Engine/Source/Runtime/Engine/Private/Components/LocalFogVolumeComponent.cpp
+++ b/Engine/Source/Runtime/Engine/Private/Components/LocalFogVolumeComponent.cpp
@@ -15,9 +15,9 @@
 #include UE_INLINE_GENERATED_CPP_BY_NAME(LocalFogVolumeComponent)
 
 
-ULocalHeightFogComponent::ULocalHeightFogComponent(const FObjectInitializer& ObjectInitializer)
+ULocalFogVolumeComponent::ULocalFogVolumeComponent(const FObjectInitializer& ObjectInitializer)
 	: Super(ObjectInitializer)
-	, LocalHeightFogSceneProxy(nullptr)
+	, LocalFogVolumeSceneProxy(nullptr)
 {
 	Mobility = EComponentMobility::Movable;
 
@@ -26,18 +26,18 @@ ULocalHeightFogComponent::ULocalHeightFogComponent(const FObjectInitializer& Obj
 	SetRelativeScale3D(FVector(Size, Size, Size));
 }
 
-ULocalHeightFogComponent::~ULocalHeightFogComponent()
+ULocalFogVolumeComponent::~ULocalFogVolumeComponent()
 {
 }
 
-void ULocalHeightFogComponent::SendRenderTransformCommand()
+void ULocalFogVolumeComponent::SendRenderTransformCommand()
 {
-	if (LocalHeightFogSceneProxy)
+	if (LocalFogVolumeSceneProxy)
 	{
 		FTransform ComponentTransform = GetComponentTransform();
 		float HeightOffset = FogHeightOffset;
-		FLocalHeightFogSceneProxy* SceneProxy = LocalHeightFogSceneProxy;
-		ENQUEUE_RENDER_COMMAND(FUpdateLocalHeightFogSceneProxyTransformCommand)(
+		FLocalFogVolumeSceneProxy* SceneProxy = LocalFogVolumeSceneProxy;
+		ENQUEUE_RENDER_COMMAND(FUpdateLocalFogVolumeSceneProxyTransformCommand)(
 			[SceneProxy, ComponentTransform, HeightOffset](FRHICommandList& RHICmdList)
 			{
 				// Nothing else is needed so that command could actually go.
@@ -46,7 +46,7 @@ void ULocalHeightFogComponent::SendRenderTransformCommand()
 	}
 }
 
-void ULocalHeightFogComponent::CreateRenderState_Concurrent(FRegisterComponentContext* Context)
+void ULocalFogVolumeComponent::CreateRenderState_Concurrent(FRegisterComponentContext* Context)
 {
 	Super::CreateRenderState_Concurrent(Context);	
 	
@@ -61,48 +61,48 @@ void ULocalHeightFogComponent::CreateRenderState_Concurrent(FRegisterComponentCo
 
 	if (GetVisibleFlag() && !bHidden && ShouldComponentAddToScene() && ShouldRender() && IsRegistered() && (GetOuter() == NULL || !GetOuter()->HasAnyFlags(RF_ClassDefaultObject)))
 	{
-		LocalHeightFogSceneProxy = CreateSceneProxy();
-		GetWorld()->Scene->AddLocalHeightFog(LocalHeightFogSceneProxy);
+		LocalFogVolumeSceneProxy = CreateSceneProxy();
+		GetWorld()->Scene->AddLocalFogVolume(LocalFogVolumeSceneProxy);
 	}
 }
 
-void ULocalHeightFogComponent::SendRenderTransform_Concurrent()
+void ULocalFogVolumeComponent::SendRenderTransform_Concurrent()
 {
 	Super::SendRenderTransform_Concurrent();
 	SendRenderTransformCommand();
 }
 
-void ULocalHeightFogComponent::DestroyRenderState_Concurrent()
+void ULocalFogVolumeComponent::DestroyRenderState_Concurrent()
 {
 	Super::DestroyRenderState_Concurrent();
 
-	if (LocalHeightFogSceneProxy)
+	if (LocalFogVolumeSceneProxy)
 	{
-		GetWorld()->Scene->RemoveLocalHeightFog(LocalHeightFogSceneProxy);
+		GetWorld()->Scene->RemoveLocalFogVolume(LocalFogVolumeSceneProxy);
 
-		FLocalHeightFogSceneProxy* SceneProxy = LocalHeightFogSceneProxy;
-		ENQUEUE_RENDER_COMMAND(FDestroyLocalHeightFogSceneProxyCommand)(
+		FLocalFogVolumeSceneProxy* SceneProxy = LocalFogVolumeSceneProxy;
+		ENQUEUE_RENDER_COMMAND(FDestroyLocalFogVolumeSceneProxyCommand)(
 			[SceneProxy](FRHICommandList& RHICmdList)
 			{
 				delete SceneProxy;
 			});
 
-		LocalHeightFogSceneProxy = nullptr;
+		LocalFogVolumeSceneProxy = nullptr;
 	}
 }
 
 #if WITH_EDITOR
 
-bool ULocalHeightFogComponent::CanEditChange(const FProperty* InProperty) const
+bool ULocalFogVolumeComponent::CanEditChange(const FProperty* InProperty) const
 {
 	if (InProperty)
 	{
 		FString PropertyName = InProperty->GetName();
 
-		if (FogMode != ELocalFogMode::LocalHeightFog)
+		if (FogMode != ELocalFogMode::LocalFogVolume)
 		{
-			if (PropertyName == GET_MEMBER_NAME_STRING_CHECKED(ULocalHeightFogComponent, FogHeightFalloff)
-				|| PropertyName == GET_MEMBER_NAME_STRING_CHECKED(ULocalHeightFogComponent, FogHeightOffset))
+			if (PropertyName == GET_MEMBER_NAME_STRING_CHECKED(ULocalFogVolumeComponent, FogHeightFalloff)
+				|| PropertyName == GET_MEMBER_NAME_STRING_CHECKED(ULocalFogVolumeComponent, FogHeightOffset))
 			{
 				return false;
 			}
@@ -112,7 +112,7 @@ bool ULocalHeightFogComponent::CanEditChange(const FProperty* InProperty) const
 	return true;
 }
 
-void ULocalHeightFogComponent::PostEditChangeProperty(FPropertyChangedEvent& PropertyChangedEvent)
+void ULocalFogVolumeComponent::PostEditChangeProperty(FPropertyChangedEvent& PropertyChangedEvent)
 {
 	Super::PostEditChangeProperty(PropertyChangedEvent);
 
@@ -133,20 +133,20 @@ void ULocalHeightFogComponent::PostEditChangeProperty(FPropertyChangedEvent& Pro
 
 #endif // WITH_EDITOR
 
-FLocalHeightFogSceneProxy* ULocalHeightFogComponent::CreateSceneProxy()
+FLocalFogVolumeSceneProxy* ULocalFogVolumeComponent::CreateSceneProxy()
 {
-	return new FLocalHeightFogSceneProxy(this);
+	return new FLocalFogVolumeSceneProxy(this);
 }
 
 #if WITH_EDITOR
 #include "ObjectEditorUtils.h"
 #endif
 
-ALocalHeightFog::ALocalHeightFog(const FObjectInitializer& ObjectInitializer)
-	: Super(ObjectInitializer.SetDefaultSubobjectClass<ULocalHeightFogComponent>(TEXT("NewLocalHeightFogComponent")))
+ALocalFogVolume::ALocalFogVolume(const FObjectInitializer& ObjectInitializer)
+	: Super(ObjectInitializer.SetDefaultSubobjectClass<ULocalFogVolumeComponent>(TEXT("NewLocalFogVolumeComponent")))
 {
-	LocalHeightFogVolume = CreateDefaultSubobject<ULocalHeightFogComponent>(TEXT("LocalHeightFogComponent"));
-	RootComponent = LocalHeightFogVolume;
+	LocalFogVolumeVolume = CreateDefaultSubobject<ULocalFogVolumeComponent>(TEXT("LocalFogVolumeComponent"));
+	RootComponent = LocalFogVolumeVolume;
 
 #if WITH_EDITORONLY_DATA
 
@@ -155,13 +155,13 @@ ALocalHeightFog::ALocalHeightFog(const FObjectInitializer& ObjectInitializer)
 		// Structure to hold one-time initialization
 		struct FConstructorStatics
 		{
-			ConstructorHelpers::FObjectFinderOptional<UTexture2D> LocalHeightFogTextureObject;
-			FName ID_LocalHeightFog;
-			FText NAME_LocalHeightFog;
+			ConstructorHelpers::FObjectFinderOptional<UTexture2D> LocalFogVolumeTextureObject;
+			FName ID_LocalFogVolume;
+			FText NAME_LocalFogVolume;
 			FConstructorStatics()
-				: LocalHeightFogTextureObject(TEXT("/Engine/EditorResources/S_SkyAtmosphere"))	// TODO
-				, ID_LocalHeightFog(TEXT("Fog"))
-				, NAME_LocalHeightFog(NSLOCTEXT("SpriteCategory", "Fog", "Fog"))
+				: LocalFogVolumeTextureObject(TEXT("/Engine/EditorResources/S_SkyAtmosphere"))	// TODO
+				, ID_LocalFogVolume(TEXT("Fog"))
+				, NAME_LocalFogVolume(NSLOCTEXT("SpriteCategory", "Fog", "Fog"))
 			{
 			}
 		};
@@ -169,11 +169,11 @@ ALocalHeightFog::ALocalHeightFog(const FObjectInitializer& ObjectInitializer)
 
 		if (GetSpriteComponent())
 		{
-			GetSpriteComponent()->Sprite = ConstructorStatics.LocalHeightFogTextureObject.Get();
+			GetSpriteComponent()->Sprite = ConstructorStatics.LocalFogVolumeTextureObject.Get();
 			GetSpriteComponent()->SetRelativeScale3D(FVector(0.5f, 0.5f, 0.5f));
-			GetSpriteComponent()->SpriteInfo.Category = ConstructorStatics.ID_LocalHeightFog;
-			GetSpriteComponent()->SpriteInfo.DisplayName = ConstructorStatics.NAME_LocalHeightFog;
-			GetSpriteComponent()->SetupAttachment(LocalHeightFogVolume);
+			GetSpriteComponent()->SpriteInfo.Category = ConstructorStatics.ID_LocalFogVolume;
+			GetSpriteComponent()->SpriteInfo.DisplayName = ConstructorStatics.NAME_LocalFogVolume;
+			GetSpriteComponent()->SetupAttachment(LocalFogVolumeVolume);
 		}
 	}
 #endif // WITH_EDITORONLY_DATA
diff --git a/Engine/Source/Runtime/Engine/Private/LocalFogVolumeSceneProxy.cpp b/Engine/Source/Runtime/Engine/Private/LocalFogVolumeSceneProxy.cpp
--- a/Engine/Source/Runtime/Engine/Private/LocalFogVolumeSceneProxy.cpp
+++ b/Engine/Source/Runtime/Engine/Private/LocalFogVolumeSceneProxy.cpp
@@ -9,7 +9,7 @@
 
 
 
-FLocalHeightFogSceneProxy::FLocalHeightFogSceneProxy(const ULocalHeightFogComponent* InComponent)
+FLocalFogVolumeSceneProxy::FLocalFogVolumeSceneProxy(const ULocalFogVolumeComponent* InComponent)
 	: FogTransform(InComponent->GetComponentTransform())
 	, FogDensity(InComponent->FogDensity)
 	, FogHeightFalloff(InComponent->FogHeightFalloff)
@@ -23,7 +23,7 @@ FLocalHeightFogSceneProxy::FLocalHeightFogSceneProxy(const ULocalHeightFogCompon
 {
 }
 
-FLocalHeightFogSceneProxy::~FLocalHeightFogSceneProxy()
+FLocalFogVolumeSceneProxy::~FLocalFogVolumeSceneProxy()
 {
 }
 
diff --git a/Engine/Source/Runtime/Engine/Public/LocalFogVolumeSceneProxy.h b/Engine/Source/Runtime/Engine/Public/LocalFogVolumeSceneProxy.h
--- a/Engine/Source/Runtime/Engine/Public/LocalFogVolumeSceneProxy.h
+++ b/Engine/Source/Runtime/Engine/Public/LocalFogVolumeSceneProxy.h
@@ -5,16 +5,16 @@
 #include "CoreMinimal.h"
 #include "Math/Color.h"
 
-class ULocalHeightFogComponent;
+class ULocalFogVolumeComponent;
 
 /** Represents a UVolumetricCloudComponent to the rendering thread, created game side from the component. */
-class FLocalHeightFogSceneProxy
+class FLocalFogVolumeSceneProxy
 {
 public:
 
 	// Initialization constructor.
-	ENGINE_API FLocalHeightFogSceneProxy(const ULocalHeightFogComponent* InComponent);
-	ENGINE_API ~FLocalHeightFogSceneProxy();
+	ENGINE_API FLocalFogVolumeSceneProxy(const ULocalFogVolumeComponent* InComponent);
+	ENGINE_API ~FLocalFogVolumeSceneProxy();
 
 	FTransform FogTransform;
 
diff --git a/Engine/Source/Runtime/Engine/Public/SceneInterface.h b/Engine/Source/Runtime/Engine/Public/SceneInterface.h
--- a/Engine/Source/Runtime/Engine/Public/SceneInterface.h
+++ b/Engine/Source/Runtime/Engine/Public/SceneInterface.h
@@ -351,17 +351,17 @@ public:
 	 *
 	 * @param FogProxy - fog proxy to add
 	 */
-	virtual void AddLocalHeightFog(class FLocalHeightFogSceneProxy* FogProxy) = 0;
+	virtual void AddLocalFogVolume(class FLocalFogVolumeSceneProxy* FogProxy) = 0;
 	/**
 	 * Removes a local height fog component from the scene
 	 *
 	 * @param FogProxy - fog proxy to remove
 	 */
-	virtual void RemoveLocalHeightFog(class FLocalHeightFogSceneProxy* FogProxy) = 0;
+	virtual void RemoveLocalFogVolume(class FLocalFogVolumeSceneProxy* FogProxy) = 0;
 	/**
 	 * @return True if there are any local height fog potentially enabled in the scene
 	 */
-	virtual bool HasAnyLocalHeightFog() const = 0;
+	virtual bool HasAnyLocalFogVolume() const = 0;
 
 	/**
 	 * Adds the unique volumetric cloud component to the scene
diff --git a/Engine/Source/Runtime/Renderer/Private/DeferredShadingRenderer.cpp b/Engine/Source/Runtime/Renderer/Private/DeferredShadingRenderer.cpp
--- a/Engine/Source/Runtime/Renderer/Private/DeferredShadingRenderer.cpp
+++ b/Engine/Source/Runtime/Renderer/Private/DeferredShadingRenderer.cpp
@@ -407,7 +407,7 @@ DECLARE_CYCLE_STAT(TEXT("DeferredShadingSceneRenderer RenderLightShaftOcclusion"
 DECLARE_CYCLE_STAT(TEXT("DeferredShadingSceneRenderer RenderAtmosphere"), STAT_FDeferredShadingSceneRenderer_RenderAtmosphere, STATGROUP_SceneRendering);
 DECLARE_CYCLE_STAT(TEXT("DeferredShadingSceneRenderer RenderSkyAtmosphere"), STAT_FDeferredShadingSceneRenderer_RenderSkyAtmosphere, STATGROUP_SceneRendering);
 DECLARE_CYCLE_STAT(TEXT("DeferredShadingSceneRenderer RenderFog"), STAT_FDeferredShadingSceneRenderer_RenderFog, STATGROUP_SceneRendering);
-DECLARE_CYCLE_STAT(TEXT("DeferredShadingSceneRenderer RenderLocalHeightFog"), STAT_FDeferredShadingSceneRenderer_RenderLocalHeightFog, STATGROUP_SceneRendering);
+DECLARE_CYCLE_STAT(TEXT("DeferredShadingSceneRenderer RenderLocalFogVolume"), STAT_FDeferredShadingSceneRenderer_RenderLocalFogVolume, STATGROUP_SceneRendering);
 DECLARE_CYCLE_STAT(TEXT("DeferredShadingSceneRenderer RenderLightShaftBloom"), STAT_FDeferredShadingSceneRenderer_RenderLightShaftBloom, STATGROUP_SceneRendering);
 DECLARE_CYCLE_STAT(TEXT("DeferredShadingSceneRenderer RenderFinish"), STAT_FDeferredShadingSceneRenderer_RenderFinish, STATGROUP_SceneRendering);
 
@@ -4059,12 +4059,12 @@ void FDeferredShadingSceneRenderer::Render(FRDGBuilder& GraphBuilder)
 			ComposeVolumetricRenderTargetOverScene(GraphBuilder, Views, SceneTextures.Color.Target, SceneTextures.Depth.Target, bShouldRenderSingleLayerWater, SceneWithoutWaterTextures, SceneTextures);
 		}
 
-		if (!bHasRayTracedOverlay && ShouldRenderLocalHeightFog(Scene, ViewFamily))
+		if (!bHasRayTracedOverlay && ShouldRenderLocalFogVolume(Scene, ViewFamily))
 		{
-			RDG_CSV_STAT_EXCLUSIVE_SCOPE(GraphBuilder, RenderLocalHeightFog);
-			SCOPED_NAMED_EVENT(RenderLocalHeightFog, FColor::Emerald);
-			SCOPE_CYCLE_COUNTER(STAT_FDeferredShadingSceneRenderer_RenderLocalHeightFog);
-			RenderLocalHeightFog(Scene, Views, GraphBuilder, SceneTextures, LightShaftOcclusionTexture);
+			RDG_CSV_STAT_EXCLUSIVE_SCOPE(GraphBuilder, RenderLocalFogVolume);
+			SCOPED_NAMED_EVENT(RenderLocalFogVolume, FColor::Emerald);
+			SCOPE_CYCLE_COUNTER(STAT_FDeferredShadingSceneRenderer_RenderLocalFogVolume);
+			RenderLocalFogVolume(Scene, Views, GraphBuilder, SceneTextures, LightShaftOcclusionTexture);
 		}
 
 		FRDGTextureRef ExposureIlluminance = AddCalculateExposureIlluminancePass(GraphBuilder, Views, SceneTextures, TranslucencyLightingVolumeTextures, ExposureIlluminanceSetup);
diff --git a/Engine/Source/Runtime/Renderer/Private/LocalFogVolumeRendering.cpp b/Engine/Source/Runtime/Renderer/Private/LocalFogVolumeRendering.cpp
--- a/Engine/Source/Runtime/Renderer/Private/LocalFogVolumeRendering.cpp
+++ b/Engine/Source/Runtime/Renderer/Private/LocalFogVolumeRendering.cpp
@@ -9,55 +9,55 @@
 
 
 // The runtime ON/OFF toggle
-static TAutoConsoleVariable<int32> CVarLocalHeightFog(
-	TEXT("r.LocalHeightFog"), 1,
-	TEXT("LocalHeightFog components are rendered when this is not 0, otherwise ignored.\n"),
+static TAutoConsoleVariable<int32> CVarLocalFogVolume(
+	TEXT("r.LocalFogVolume"), 1,
+	TEXT("LocalFogVolume components are rendered when this is not 0, otherwise ignored.\n"),
 	ECVF_RenderThreadSafe);
 
-bool ShouldRenderLocalHeightFog(const FScene* Scene, const FSceneViewFamily& Family)
+bool ShouldRenderLocalFogVolume(const FScene* Scene, const FSceneViewFamily& Family)
 {
 	const FEngineShowFlags EngineShowFlags = Family.EngineShowFlags;
-	if (Scene && Scene->HasAnyLocalHeightFog() && EngineShowFlags.Fog && !Family.UseDebugViewPS())
+	if (Scene && Scene->HasAnyLocalFogVolume() && EngineShowFlags.Fog && !Family.UseDebugViewPS())
 	{
-		return CVarLocalHeightFog.GetValueOnRenderThread() > 0;
+		return CVarLocalFogVolume.GetValueOnRenderThread() > 0;
 	}
 	return false;
 }
 
-DECLARE_GPU_STAT(LocalHeightFogVolumes);
+DECLARE_GPU_STAT(LocalFogVolumeVolumes);
 
 /*=============================================================================
 	FScene functions
 =============================================================================*/
 
-void FScene::AddLocalHeightFog(class FLocalHeightFogSceneProxy* FogProxy)
+void FScene::AddLocalFogVolume(class FLocalFogVolumeSceneProxy* FogProxy)
 {
 	check(FogProxy);
 	FScene* Scene = this;
 
-	ENQUEUE_RENDER_COMMAND(FAddLocalHeightFogCommand)(
+	ENQUEUE_RENDER_COMMAND(FAddLocalFogVolumeCommand)(
 		[Scene, FogProxy](FRHICommandListImmediate& RHICmdList)
 		{
-			check(!Scene->LocalHeightFogs.Contains(FogProxy));
-			Scene->LocalHeightFogs.Push(FogProxy);
+			check(!Scene->LocalFogVolumes.Contains(FogProxy));
+			Scene->LocalFogVolumes.Push(FogProxy);
 		} );
 }
 
-void FScene::RemoveLocalHeightFog(class FLocalHeightFogSceneProxy* FogProxy)
+void FScene::RemoveLocalFogVolume(class FLocalFogVolumeSceneProxy* FogProxy)
 {
 	check(FogProxy);
 	FScene* Scene = this;
 
-	ENQUEUE_RENDER_COMMAND(FRemoveLocalHeightFogCommand)(
+	ENQUEUE_RENDER_COMMAND(FRemoveLocalFogVolumeCommand)(
 		[Scene, FogProxy](FRHICommandListImmediate& RHICmdList)
 		{
-			Scene->LocalHeightFogs.RemoveSingle(FogProxy);
+			Scene->LocalFogVolumes.RemoveSingle(FogProxy);
 		} );
 }
 
-bool FScene::HasAnyLocalHeightFog() const
+bool FScene::HasAnyLocalFogVolume() const
 { 
-	return LocalHeightFogs.Num() > 0;
+	return LocalFogVolumes.Num() > 0;
 }
 
 /*=============================================================================
@@ -67,12 +67,12 @@ bool FScene::HasAnyLocalHeightFog() const
 void GetLocalFogVolumeSortingData(const FScene* Scene, FRDGBuilder& GraphBuilder, FLocalFogVolumeSortingData& Out)
 {
 	// No culling as of today
-	Out.LocalHeightFogInstanceCount = Scene->LocalHeightFogs.Num();
-	Out.LocalHeightFogInstanceCountFinal = 0;
-	Out.LocalHeightFogGPUInstanceData = (FLocalHeightFogGPUInstanceData*)GraphBuilder.Alloc(sizeof(FLocalHeightFogGPUInstanceData) * Out.LocalHeightFogInstanceCount, 16);
-	Out.LocalHeightFogCenterPos = (FVector*)GraphBuilder.Alloc(sizeof(FVector3f) * Out.LocalHeightFogInstanceCount, 16);
-	Out.LocalHeightFogSortKeys.SetNumUninitialized(Out.LocalHeightFogInstanceCount);
-	for (FLocalHeightFogSceneProxy* LHF : Scene->LocalHeightFogs)
+	Out.LocalFogVolumeInstanceCount = Scene->LocalFogVolumes.Num();
+	Out.LocalFogVolumeInstanceCountFinal = 0;
+	Out.LocalFogVolumeGPUInstanceData = (FLocalFogVolumeGPUInstanceData*)GraphBuilder.Alloc(sizeof(FLocalFogVolumeGPUInstanceData) * Out.LocalFogVolumeInstanceCount, 16);
+	Out.LocalFogVolumeCenterPos = (FVector*)GraphBuilder.Alloc(sizeof(FVector3f) * Out.LocalFogVolumeInstanceCount, 16);
+	Out.LocalFogVolumeSortKeys.SetNumUninitialized(Out.LocalFogVolumeInstanceCount);
+	for (FLocalFogVolumeSceneProxy* LHF : Scene->LocalFogVolumes)
 	{
 		if (LHF->FogDensity <= 0.0f)
 		{
@@ -82,95 +82,95 @@ void GetLocalFogVolumeSortingData(const FScene* Scene, FRDGBuilder& GraphBuilder
 		FTransform TransformScaleOnly;
 		TransformScaleOnly.SetScale3D(LHF->FogTransform.GetScale3D());
 
-		FLocalHeightFogGPUInstanceData* LocalHeightFogGPUInstanceDataIt = &Out.LocalHeightFogGPUInstanceData[Out.LocalHeightFogInstanceCountFinal];
-		LocalHeightFogGPUInstanceDataIt->Transform = FMatrix44f(LHF->FogTransform.ToMatrixWithScale());
-		LocalHeightFogGPUInstanceDataIt->InvTransform = LocalHeightFogGPUInstanceDataIt->Transform.Inverse();
-		LocalHeightFogGPUInstanceDataIt->InvTranformNoScale = FMatrix44f(LHF->FogTransform.ToMatrixNoScale()).Inverse();
-		LocalHeightFogGPUInstanceDataIt->TransformScaleOnly = FMatrix44f(TransformScaleOnly.ToMatrixWithScale());
+		FLocalFogVolumeGPUInstanceData* LocalFogVolumeGPUInstanceDataIt = &Out.LocalFogVolumeGPUInstanceData[Out.LocalFogVolumeInstanceCountFinal];
+		LocalFogVolumeGPUInstanceDataIt->Transform = FMatrix44f(LHF->FogTransform.ToMatrixWithScale());
+		LocalFogVolumeGPUInstanceDataIt->InvTransform = LocalFogVolumeGPUInstanceDataIt->Transform.Inverse();
+		LocalFogVolumeGPUInstanceDataIt->InvTranformNoScale = FMatrix44f(LHF->FogTransform.ToMatrixNoScale()).Inverse();
+		LocalFogVolumeGPUInstanceDataIt->TransformScaleOnly = FMatrix44f(TransformScaleOnly.ToMatrixWithScale());
 
-		LocalHeightFogGPUInstanceDataIt->Density = LHF->FogDensity;
-		LocalHeightFogGPUInstanceDataIt->HeightFalloff = LHF->FogHeightFalloff * 0.01f;	// This scale is used to have artist author reasonable range.
-		LocalHeightFogGPUInstanceDataIt->HeightOffset = LHF->FogHeightOffset;
-		LocalHeightFogGPUInstanceDataIt->RadialAttenuation = LHF->FogRadialAttenuation;
+		LocalFogVolumeGPUInstanceDataIt->Density = LHF->FogDensity;
+		LocalFogVolumeGPUInstanceDataIt->HeightFalloff = LHF->FogHeightFalloff * 0.01f;	// This scale is used to have artist author reasonable range.
+		LocalFogVolumeGPUInstanceDataIt->HeightOffset = LHF->FogHeightOffset;
+		LocalFogVolumeGPUInstanceDataIt->RadialAttenuation = LHF->FogRadialAttenuation;
 
-		LocalHeightFogGPUInstanceDataIt->FogMode = float(LHF->FogMode);
+		LocalFogVolumeGPUInstanceDataIt->FogMode = float(LHF->FogMode);
 
-		LocalHeightFogGPUInstanceDataIt->Albedo = FVector3f(LHF->FogAlbedo);
-		LocalHeightFogGPUInstanceDataIt->PhaseG = LHF->FogPhaseG;
-		LocalHeightFogGPUInstanceDataIt->Emissive = FVector3f(LHF->FogEmissive);
+		LocalFogVolumeGPUInstanceDataIt->Albedo = FVector3f(LHF->FogAlbedo);
+		LocalFogVolumeGPUInstanceDataIt->PhaseG = LHF->FogPhaseG;
+		LocalFogVolumeGPUInstanceDataIt->Emissive = FVector3f(LHF->FogEmissive);
 
-		Out.LocalHeightFogCenterPos[Out.LocalHeightFogInstanceCountFinal] = LHF->FogTransform.GetTranslation();
+		Out.LocalFogVolumeCenterPos[Out.LocalFogVolumeInstanceCountFinal] = LHF->FogTransform.GetTranslation();
 
-		FLocalFogVolumeSortKey* LocalHeightFogSortKeysIt = &Out.LocalHeightFogSortKeys[Out.LocalHeightFogInstanceCountFinal];
-		LocalHeightFogSortKeysIt->FogVolume.Index = Out.LocalHeightFogInstanceCountFinal;
-		LocalHeightFogSortKeysIt->FogVolume.Distance = 0;	// Filled up right before sorting according to a view
-		LocalHeightFogSortKeysIt->FogVolume.Priority = LHF->FogSortPriority;
+		FLocalFogVolumeSortKey* LocalFogVolumeSortKeysIt = &Out.LocalFogVolumeSortKeys[Out.LocalFogVolumeInstanceCountFinal];
+		LocalFogVolumeSortKeysIt->FogVolume.Index = Out.LocalFogVolumeInstanceCountFinal;
+		LocalFogVolumeSortKeysIt->FogVolume.Distance = 0;	// Filled up right before sorting according to a view
+		LocalFogVolumeSortKeysIt->FogVolume.Priority = LHF->FogSortPriority;
 
-		Out.LocalHeightFogInstanceCountFinal++;
+		Out.LocalFogVolumeInstanceCountFinal++;
 	}
 	// Shrink the array to only what is needed in order for the sort to correctly work on only what is needed.
-	Out.LocalHeightFogSortKeys.SetNum(Out.LocalHeightFogInstanceCountFinal, false/*bAllowShrinking*/);
+	Out.LocalFogVolumeSortKeys.SetNum(Out.LocalFogVolumeInstanceCountFinal, false/*bAllowShrinking*/);
 }
 
 void CreateViewLocalFogVolumeBufferSRV(FViewInfo& View, FRDGBuilder& GraphBuilder, FLocalFogVolumeSortingData& SortingData)
 {
 	static const uint32 SizeOfFloat4 = sizeof(float) * 4;
-	static const uint32 Float4CountInLocalHeightFogGPUInstanceData = sizeof(FLocalHeightFogGPUInstanceData) / SizeOfFloat4;
-	static_assert(sizeof(FLocalHeightFogGPUInstanceData) == Float4CountInLocalHeightFogGPUInstanceData * SizeOfFloat4); // The size of the structure must be a multiple of FVector4.
+	static const uint32 Float4CountInLocalFogVolumeGPUInstanceData = sizeof(FLocalFogVolumeGPUInstanceData) / SizeOfFloat4;
+	static_assert(sizeof(FLocalFogVolumeGPUInstanceData) == Float4CountInLocalFogVolumeGPUInstanceData * SizeOfFloat4); // The size of the structure must be a multiple of FVector4.
 
-	if (SortingData.LocalHeightFogInstanceCountFinal == 0)
+	if (SortingData.LocalFogVolumeInstanceCountFinal == 0)
 	{
-		View.LocalHeightFogGPUInstanceCount = 0;
+		View.LocalFogVolumeGPUInstanceCount = 0;
 
-		static FLocalHeightFogGPUInstanceData DummyData;
-		View.LocalHeightFogGPUInstanceDataBufferSRV = GraphBuilder.CreateSRV(
-			CreateVertexBuffer(GraphBuilder, TEXT("LocalHeightFogGPUInstanceDataBuffer"), 
-			FRDGBufferDesc::CreateBufferDesc(SizeOfFloat4, Float4CountInLocalHeightFogGPUInstanceData), &DummyData, sizeof(FLocalHeightFogGPUInstanceData) * 1, ERDGInitialDataFlags::NoCopy),
+		static FLocalFogVolumeGPUInstanceData DummyData;
+		View.LocalFogVolumeGPUInstanceDataBufferSRV = GraphBuilder.CreateSRV(
+			CreateVertexBuffer(GraphBuilder, TEXT("LocalFogVolumeGPUInstanceDataBuffer"), 
+			FRDGBufferDesc::CreateBufferDesc(SizeOfFloat4, Float4CountInLocalFogVolumeGPUInstanceData), &DummyData, sizeof(FLocalFogVolumeGPUInstanceData) * 1, ERDGInitialDataFlags::NoCopy),
 			PF_A32B32G32R32F);
 		return;
 	}
 
 	// 1. Sort all the volumes
 	const FVector ViewOrigin = View.ViewMatrices.GetViewOrigin();
-	for (uint32 i = 0; i < SortingData.LocalHeightFogInstanceCountFinal; ++i)
+	for (uint32 i = 0; i < SortingData.LocalFogVolumeInstanceCountFinal; ++i)
 	{
-		FVector FogCenterPos = SortingData.LocalHeightFogCenterPos[SortingData.LocalHeightFogSortKeys[i].FogVolume.Index];	// Recovered form the original array via index because the sorting of the previous view might have changed the order.
+		FVector FogCenterPos = SortingData.LocalFogVolumeCenterPos[SortingData.LocalFogVolumeSortKeys[i].FogVolume.Index];	// Recovered form the original array via index because the sorting of the previous view might have changed the order.
 		float DistancetoView = float((FogCenterPos - ViewOrigin).Size());
-		SortingData.LocalHeightFogSortKeys[i].FogVolume.Distance = *reinterpret_cast<uint32*>(&DistancetoView);
+		SortingData.LocalFogVolumeSortKeys[i].FogVolume.Distance = *reinterpret_cast<uint32*>(&DistancetoView);
 	}
-	SortingData.LocalHeightFogSortKeys.Sort();
+	SortingData.LocalFogVolumeSortKeys.Sort();
 
 	// 2. Create the buffer containing all the fog volume data instance sorted according to their key for the current view.
-	FLocalHeightFogGPUInstanceData* LocalHeightFogGPUSortedInstanceData = (FLocalHeightFogGPUInstanceData*)GraphBuilder.Alloc(sizeof(FLocalHeightFogGPUInstanceData) * SortingData.LocalHeightFogInstanceCountFinal, 16);
-	for (uint32 i = 0; i < SortingData.LocalHeightFogInstanceCountFinal; ++i)
+	FLocalFogVolumeGPUInstanceData* LocalFogVolumeGPUSortedInstanceData = (FLocalFogVolumeGPUInstanceData*)GraphBuilder.Alloc(sizeof(FLocalFogVolumeGPUInstanceData) * SortingData.LocalFogVolumeInstanceCountFinal, 16);
+	for (uint32 i = 0; i < SortingData.LocalFogVolumeInstanceCountFinal; ++i)
 	{
 		// We could also have an indirection buffer on GPU but choosing to go with the sorting + copy on CPU since it is expected to not have many local height fog volumes.
-		LocalHeightFogGPUSortedInstanceData[i] = SortingData.LocalHeightFogGPUInstanceData[SortingData.LocalHeightFogSortKeys[i].FogVolume.Index];
+		LocalFogVolumeGPUSortedInstanceData[i] = SortingData.LocalFogVolumeGPUInstanceData[SortingData.LocalFogVolumeSortKeys[i].FogVolume.Index];
 	}
 
 	// 3. Allocate buffer and initialize with sorted data to upload to GPU
-	const uint32 AllLocalHeightFogInstanceBytesFinal = sizeof(FLocalHeightFogGPUInstanceData) * SortingData.LocalHeightFogInstanceCountFinal;
-	FRDGBufferRef LocalHeightFogGPUInstanceDataBuffer = CreateVertexBuffer(
-		GraphBuilder, TEXT("LocalHeightFogGPUInstanceDataBuffer"),
-		FRDGBufferDesc::CreateBufferDesc(SizeOfFloat4, SortingData.LocalHeightFogInstanceCountFinal * Float4CountInLocalHeightFogGPUInstanceData), 
-		LocalHeightFogGPUSortedInstanceData, AllLocalHeightFogInstanceBytesFinal, ERDGInitialDataFlags::NoCopy);
-
-	View.LocalHeightFogGPUInstanceCount = SortingData.LocalHeightFogInstanceCountFinal;
-	View.LocalHeightFogGPUInstanceDataBufferSRV = GraphBuilder.CreateSRV(LocalHeightFogGPUInstanceDataBuffer, PF_A32B32G32R32F);
+	const uint32 AllLocalFogVolumeInstanceBytesFinal = sizeof(FLocalFogVolumeGPUInstanceData) * SortingData.LocalFogVolumeInstanceCountFinal;
+	FRDGBufferRef LocalFogVolumeGPUInstanceDataBuffer = CreateVertexBuffer(
+		GraphBuilder, TEXT("LocalFogVolumeGPUInstanceDataBuffer"),
+		FRDGBufferDesc::CreateBufferDesc(SizeOfFloat4, SortingData.LocalFogVolumeInstanceCountFinal * Float4CountInLocalFogVolumeGPUInstanceData), 
+		LocalFogVolumeGPUSortedInstanceData, AllLocalFogVolumeInstanceBytesFinal, ERDGInitialDataFlags::NoCopy);
+
+	View.LocalFogVolumeGPUInstanceCount = SortingData.LocalFogVolumeInstanceCountFinal;
+	View.LocalFogVolumeGPUInstanceDataBufferSRV = GraphBuilder.CreateSRV(LocalFogVolumeGPUInstanceDataBuffer, PF_A32B32G32R32F);
 }
 
 /*=============================================================================
 	Local height fog rendering - non mobile
 =============================================================================*/
 
-class FLocalHeightFogVS : public FGlobalShader
+class FLocalFogVolumeVS : public FGlobalShader
 {
-	DECLARE_GLOBAL_SHADER(FLocalHeightFogVS);
-	SHADER_USE_PARAMETER_STRUCT(FLocalHeightFogVS, FGlobalShader);
+	DECLARE_GLOBAL_SHADER(FLocalFogVolumeVS);
+	SHADER_USE_PARAMETER_STRUCT(FLocalFogVolumeVS, FGlobalShader);
 
 	BEGIN_SHADER_PARAMETER_STRUCT(FParameters, )
 		SHADER_PARAMETER_STRUCT_REF(FViewUniformShaderParameters, View)
-		SHADER_PARAMETER_RDG_BUFFER_SRV(Buffer<float4>, LocalHeightFogInstances)
+		SHADER_PARAMETER_RDG_BUFFER_SRV(Buffer<float4>, LocalFogVolumeInstances)
 	END_SHADER_PARAMETER_STRUCT()
 
 	using FPermutationDomain = TShaderPermutationDomain<>;
@@ -185,16 +185,16 @@ class FLocalHeightFogVS : public FGlobalShader
 	}
 };
 
-IMPLEMENT_GLOBAL_SHADER(FLocalHeightFogVS, "/Engine/Private/LocalFogVolume.usf", "LocalHeightFogSplatVS", SF_Vertex);
+IMPLEMENT_GLOBAL_SHADER(FLocalFogVolumeVS, "/Engine/Private/LocalFogVolume.usf", "LocalFogVolumeSplatVS", SF_Vertex);
 
-class FLocalHeightFogPS : public FGlobalShader
+class FLocalFogVolumePS : public FGlobalShader
 {
-	DECLARE_GLOBAL_SHADER(FLocalHeightFogPS);
-	SHADER_USE_PARAMETER_STRUCT(FLocalHeightFogPS, FGlobalShader);
+	DECLARE_GLOBAL_SHADER(FLocalFogVolumePS);
+	SHADER_USE_PARAMETER_STRUCT(FLocalFogVolumePS, FGlobalShader);
 
 	BEGIN_SHADER_PARAMETER_STRUCT(FParameters, )
 		SHADER_PARAMETER_STRUCT_REF(FViewUniformShaderParameters, View)
-		SHADER_PARAMETER_RDG_BUFFER_SRV(Buffer<float4>, LocalHeightFogInstances)
+		SHADER_PARAMETER_RDG_BUFFER_SRV(Buffer<float4>, LocalFogVolumeInstances)
 	END_SHADER_PARAMETER_STRUCT()
 
 	using FPermutationDomain = TShaderPermutationDomain<>;
@@ -209,70 +209,70 @@ class FLocalHeightFogPS : public FGlobalShader
 	}
 };
 
-IMPLEMENT_GLOBAL_SHADER(FLocalHeightFogPS, "/Engine/Private/LocalFogVolume.usf", "LocalHeightFogSplatPS", SF_Pixel);
+IMPLEMENT_GLOBAL_SHADER(FLocalFogVolumePS, "/Engine/Private/LocalFogVolume.usf", "LocalFogVolumeSplatPS", SF_Pixel);
 
-BEGIN_SHADER_PARAMETER_STRUCT(FLocalHeightFogPassParameters, )
-	SHADER_PARAMETER_STRUCT_INCLUDE(FLocalHeightFogVS::FParameters, VS)
-	SHADER_PARAMETER_STRUCT_INCLUDE(FLocalHeightFogPS::FParameters, PS)
+BEGIN_SHADER_PARAMETER_STRUCT(FLocalFogVolumePassParameters, )
+	SHADER_PARAMETER_STRUCT_INCLUDE(FLocalFogVolumeVS::FParameters, VS)
+	SHADER_PARAMETER_STRUCT_INCLUDE(FLocalFogVolumePS::FParameters, PS)
 	SHADER_PARAMETER_RDG_UNIFORM_BUFFER(FSceneTextureUniformParameters, SceneTextures)
 	RENDER_TARGET_BINDING_SLOTS()
 END_SHADER_PARAMETER_STRUCT()
  
-void RenderLocalHeightFog(
+void RenderLocalFogVolume(
 	const FScene* Scene,
 	TArray<FViewInfo>& Views,
 	FRDGBuilder& GraphBuilder,
 	const FMinimalSceneTextures& SceneTextures,
 	FRDGTextureRef LightShaftOcclusionTexture)
 {
-	uint32 LocalHeightFogInstanceCount = Scene->LocalHeightFogs.Num();
-	if (LocalHeightFogInstanceCount > 0)
+	uint32 LocalFogVolumeInstanceCount = Scene->LocalFogVolumes.Num();
+	if (LocalFogVolumeInstanceCount > 0)
 	{
-		RDG_GPU_STAT_SCOPE(GraphBuilder, LocalHeightFogVolumes);
+		RDG_GPU_STAT_SCOPE(GraphBuilder, LocalFogVolumeVolumes);
 
 		FLocalFogVolumeSortingData SortingData;
 		GetLocalFogVolumeSortingData(Scene, GraphBuilder, SortingData);
 
-		if (SortingData.LocalHeightFogInstanceCountFinal > 0)
+		if (SortingData.LocalFogVolumeInstanceCountFinal > 0)
 		{
 			FRDGTextureRef SceneColorTexture = SceneTextures.Color.Resolve;
 
 			for (FViewInfo& View : Views)
 			{
 				CreateViewLocalFogVolumeBufferSRV(View, GraphBuilder, SortingData);
-				if (View.LocalHeightFogGPUInstanceCount == 0)
+				if (View.LocalFogVolumeGPUInstanceCount == 0)
 				{
 					continue;
 				}
 
-				FLocalHeightFogPassParameters* PassParameters = GraphBuilder.AllocParameters<FLocalHeightFogPassParameters>();
+				FLocalFogVolumePassParameters* PassParameters = GraphBuilder.AllocParameters<FLocalFogVolumePassParameters>();
 
 				PassParameters->VS.View = GetShaderBinding(View.ViewUniformBuffer);
-				PassParameters->VS.LocalHeightFogInstances = View.LocalHeightFogGPUInstanceDataBufferSRV;
+				PassParameters->VS.LocalFogVolumeInstances = View.LocalFogVolumeGPUInstanceDataBufferSRV;
 
 				PassParameters->PS.View = GetShaderBinding(View.ViewUniformBuffer);
-				PassParameters->PS.LocalHeightFogInstances = View.LocalHeightFogGPUInstanceDataBufferSRV;
+				PassParameters->PS.LocalFogVolumeInstances = View.LocalFogVolumeGPUInstanceDataBufferSRV;
 
 				PassParameters->SceneTextures = SceneTextures.UniformBuffer;
 				PassParameters->RenderTargets[0] = FRenderTargetBinding(SceneColorTexture, ERenderTargetLoadAction::ENoAction);
 
-				FLocalHeightFogVS::FPermutationDomain VSPermutationVector;
-				auto VertexShader = View.ShaderMap->GetShader< FLocalHeightFogVS >(VSPermutationVector);
+				FLocalFogVolumeVS::FPermutationDomain VSPermutationVector;
+				auto VertexShader = View.ShaderMap->GetShader< FLocalFogVolumeVS >(VSPermutationVector);
 
-				FLocalHeightFogPS::FPermutationDomain PsPermutationVector;
-				auto PixelShader = View.ShaderMap->GetShader< FLocalHeightFogPS >(PsPermutationVector);
+				FLocalFogVolumePS::FPermutationDomain PsPermutationVector;
+				auto PixelShader = View.ShaderMap->GetShader< FLocalFogVolumePS >(PsPermutationVector);
 
 				const FIntRect ViewRect = View.ViewRect;
 
 				ClearUnusedGraphResources(VertexShader, &PassParameters->VS);
 				ClearUnusedGraphResources(PixelShader, &PassParameters->PS);
 
-				uint32 LocalHeightFogGPUInstanceCount = View.LocalHeightFogGPUInstanceCount;
+				uint32 LocalFogVolumeGPUInstanceCount = View.LocalFogVolumeGPUInstanceCount;
 				GraphBuilder.AddPass(
-					RDG_EVENT_NAME("RenderLocalHeightFog %u inst.", LocalHeightFogGPUInstanceCount),
+					RDG_EVENT_NAME("RenderLocalFogVolume %u inst.", LocalFogVolumeGPUInstanceCount),
 					PassParameters,
 					ERDGPassFlags::Raster,
-					[VertexShader, PixelShader, PassParameters, LocalHeightFogGPUInstanceCount, ViewRect](FRHICommandList& RHICmdList)
+					[VertexShader, PixelShader, PassParameters, LocalFogVolumeGPUInstanceCount, ViewRect](FRHICommandList& RHICmdList)
 				{
 					FGraphicsPipelineStateInitializer GraphicsPSOInit;
 					RHICmdList.ApplyCachedRenderTargets(GraphicsPSOInit);
@@ -302,7 +302,7 @@ void RenderLocalHeightFog(
 						, 8									//uint32 NumVertices
 						, 0									//uint32 StartIndex
 						, UE_ARRAY_COUNT(GCubeIndices) / 3	//uint32 NumPrimitives
-						, LocalHeightFogGPUInstanceCount	//uint32 NumInstances
+						, LocalFogVolumeGPUInstanceCount	//uint32 NumInstances
 					);
 				});
 			}
@@ -315,14 +315,14 @@ void RenderLocalHeightFog(
 =============================================================================*/
 	
 
-class FMobileLocalHeightFogVS : public FGlobalShader
+class FMobileLocalFogVolumeVS : public FGlobalShader
 {
-	DECLARE_GLOBAL_SHADER(FMobileLocalHeightFogVS);
-	SHADER_USE_PARAMETER_STRUCT(FMobileLocalHeightFogVS, FGlobalShader);
+	DECLARE_GLOBAL_SHADER(FMobileLocalFogVolumeVS);
+	SHADER_USE_PARAMETER_STRUCT(FMobileLocalFogVolumeVS, FGlobalShader);
 
 	BEGIN_SHADER_PARAMETER_STRUCT(FParameters, )
 		SHADER_PARAMETER_STRUCT_INCLUDE(FViewShaderParameters, View)
-		SHADER_PARAMETER_RDG_BUFFER_SRV(Buffer<float4>, LocalHeightFogInstances)
+		SHADER_PARAMETER_RDG_BUFFER_SRV(Buffer<float4>, LocalFogVolumeInstances)
 	END_SHADER_PARAMETER_STRUCT()
 
 	using FPermutationDomain = TShaderPermutationDomain<>;
@@ -344,17 +344,17 @@ class FMobileLocalHeightFogVS : public FGlobalShader
 	}
 };
 
-IMPLEMENT_GLOBAL_SHADER(FMobileLocalHeightFogVS, "/Engine/Private/LocalFogVolume.usf", "LocalHeightFogSplatVS", SF_Vertex);
+IMPLEMENT_GLOBAL_SHADER(FMobileLocalFogVolumeVS, "/Engine/Private/LocalFogVolume.usf", "LocalFogVolumeSplatVS", SF_Vertex);
 
-class FMobileLocalHeightFogPS : public FGlobalShader
+class FMobileLocalFogVolumePS : public FGlobalShader
 {
-	DECLARE_GLOBAL_SHADER(FMobileLocalHeightFogPS);
-	SHADER_USE_PARAMETER_STRUCT(FMobileLocalHeightFogPS, FGlobalShader);
+	DECLARE_GLOBAL_SHADER(FMobileLocalFogVolumePS);
+	SHADER_USE_PARAMETER_STRUCT(FMobileLocalFogVolumePS, FGlobalShader);
 
 	BEGIN_SHADER_PARAMETER_STRUCT(FParameters, )
 		SHADER_PARAMETER_STRUCT_INCLUDE(FViewShaderParameters, View)
 		SHADER_PARAMETER_RDG_UNIFORM_BUFFER(FMobileBasePassUniformParameters, MobileBasePass)
-		SHADER_PARAMETER_RDG_BUFFER_SRV(Buffer<float4>, LocalHeightFogInstances)
+		SHADER_PARAMETER_RDG_BUFFER_SRV(Buffer<float4>, LocalFogVolumeInstances)
 	END_SHADER_PARAMETER_STRUCT()
 
 	using FPermutationDomain = TShaderPermutationDomain<>;
@@ -376,24 +376,24 @@ class FMobileLocalHeightFogPS : public FGlobalShader
 	}
 };
 
-IMPLEMENT_GLOBAL_SHADER(FMobileLocalHeightFogPS, "/Engine/Private/LocalFogVolume.usf", "LocalHeightFogSplatPS", SF_Pixel);
+IMPLEMENT_GLOBAL_SHADER(FMobileLocalFogVolumePS, "/Engine/Private/LocalFogVolume.usf", "LocalFogVolumeSplatPS", SF_Pixel);
 
-void RenderLocalHeightFogMobile(
+void RenderLocalFogVolumeMobile(
 	FRHICommandList& RHICmdList,
 	const FViewInfo& View)
 {
-	if (View.LocalHeightFogGPUInstanceCount == 0)
+	if (View.LocalFogVolumeGPUInstanceCount == 0)
 	{
 		return;
 	}
 
-	SCOPED_DRAW_EVENT(RHICmdList, LocalHeightFogVolumes);
+	SCOPED_DRAW_EVENT(RHICmdList, LocalFogVolumeVolumes);
 
-	FMobileLocalHeightFogVS::FPermutationDomain VSPermutationVector;
-	auto VertexShader = View.ShaderMap->GetShader< FMobileLocalHeightFogVS >(VSPermutationVector);
+	FMobileLocalFogVolumeVS::FPermutationDomain VSPermutationVector;
+	auto VertexShader = View.ShaderMap->GetShader< FMobileLocalFogVolumeVS >(VSPermutationVector);
 
-	FMobileLocalHeightFogPS::FPermutationDomain PsPermutationVector;
-	auto PixelShader = View.ShaderMap->GetShader< FMobileLocalHeightFogPS >(PsPermutationVector);
+	FMobileLocalFogVolumePS::FPermutationDomain PsPermutationVector;
+	auto PixelShader = View.ShaderMap->GetShader< FMobileLocalFogVolumePS >(PsPermutationVector);
 
 	const FIntRect ViewRect = View.ViewRect;
 
@@ -414,14 +414,14 @@ void RenderLocalHeightFogMobile(
 
 	SetGraphicsPipelineState(RHICmdList, GraphicsPSOInit, 0);
 
-	FMobileLocalHeightFogVS::FParameters VSParameters;
+	FMobileLocalFogVolumeVS::FParameters VSParameters;
 	VSParameters.View = View.GetShaderParameters();
-	VSParameters.LocalHeightFogInstances = View.LocalHeightFogGPUInstanceDataBufferSRV;
+	VSParameters.LocalFogVolumeInstances = View.LocalFogVolumeGPUInstanceDataBufferSRV;
 	SetShaderParameters(RHICmdList, VertexShader, VertexShader.GetVertexShader(), VSParameters);
 
-	FMobileLocalHeightFogPS::FParameters PSParameters;
+	FMobileLocalFogVolumePS::FParameters PSParameters;
 	PSParameters.View = View.GetShaderParameters();
-	PSParameters.LocalHeightFogInstances = View.LocalHeightFogGPUInstanceDataBufferSRV;
+	PSParameters.LocalFogVolumeInstances = View.LocalFogVolumeGPUInstanceDataBufferSRV;
 	// PSParameters.MobileBasePass filled up by the RDG pass parameters.
 	SetShaderParameters(RHICmdList, PixelShader, PixelShader.GetPixelShader(), PSParameters);
 
@@ -434,6 +434,6 @@ void RenderLocalHeightFogMobile(
 		, 8										//uint32 NumVertices
 		, 0										//uint32 StartIndex
 		, UE_ARRAY_COUNT(GCubeIndices) / 3		//uint32 NumPrimitives
-		, View.LocalHeightFogGPUInstanceCount	//uint32 NumInstances
+		, View.LocalFogVolumeGPUInstanceCount	//uint32 NumInstances
 	);
 }
diff --git a/Engine/Source/Runtime/Renderer/Private/LocalFogVolumeRendering.h b/Engine/Source/Runtime/Renderer/Private/LocalFogVolumeRendering.h
--- a/Engine/Source/Runtime/Renderer/Private/LocalFogVolumeRendering.h
+++ b/Engine/Source/Runtime/Renderer/Private/LocalFogVolumeRendering.h
@@ -19,7 +19,7 @@ struct FMinimalSceneTextures;
 	Local height fog rendering common data
 =============================================================================*/
 
-class FLocalHeightFogGPUInstanceData
+class FLocalFogVolumeGPUInstanceData
 {
 public:
 	FMatrix44f Transform;
@@ -68,11 +68,11 @@ public:
 // This can be generated once for a scene and then shared between all the views to generate the GPU buffer data needed for rendering while accounting for sorting.
 struct FLocalFogVolumeSortingData
 {
-	uint32 LocalHeightFogInstanceCount;
-	uint32 LocalHeightFogInstanceCountFinal;
-	FLocalHeightFogGPUInstanceData* LocalHeightFogGPUInstanceData;
-	FVector* LocalHeightFogCenterPos;
-	TArray<FLocalFogVolumeSortKey>	LocalHeightFogSortKeys;
+	uint32 LocalFogVolumeInstanceCount;
+	uint32 LocalFogVolumeInstanceCountFinal;
+	FLocalFogVolumeGPUInstanceData* LocalFogVolumeGPUInstanceData;
+	FVector* LocalFogVolumeCenterPos;
+	TArray<FLocalFogVolumeSortKey>	LocalFogVolumeSortKeys;
 };
 
 
@@ -80,20 +80,20 @@ struct FLocalFogVolumeSortingData
 	Local height fog rendering functions
 =============================================================================*/
 
-bool ShouldRenderLocalHeightFog(const FScene* Scene, const FSceneViewFamily& Family);
+bool ShouldRenderLocalFogVolume(const FScene* Scene, const FSceneViewFamily& Family);
 
 void GetLocalFogVolumeSortingData(const FScene* Scene, FRDGBuilder& GraphBuilder, FLocalFogVolumeSortingData& Out);
 
 void CreateViewLocalFogVolumeBufferSRV(FViewInfo& View, FRDGBuilder& GraphBuilder, FLocalFogVolumeSortingData& SortingData);
 
-void RenderLocalHeightFog(
+void RenderLocalFogVolume(
 	const FScene* Scene,
 	TArray<FViewInfo>& Views,
 	FRDGBuilder& GraphBuilder,
 	const FMinimalSceneTextures& SceneTextures,
 	FRDGTextureRef LightShaftOcclusionTexture);
 
-void RenderLocalHeightFogMobile(
+void RenderLocalFogVolumeMobile(
 	FRHICommandList& RHICmdList, 
 	const FViewInfo& View);
 
diff --git a/Engine/Source/Runtime/Renderer/Private/MobileShadingRenderer.cpp b/Engine/Source/Runtime/Renderer/Private/MobileShadingRenderer.cpp
--- a/Engine/Source/Runtime/Renderer/Private/MobileShadingRenderer.cpp
+++ b/Engine/Source/Runtime/Renderer/Private/MobileShadingRenderer.cpp
@@ -194,7 +194,7 @@ BEGIN_SHADER_PARAMETER_STRUCT(FMobileRenderPassParameters, RENDERER_API)
 	SHADER_PARAMETER_STRUCT_INCLUDE(FInstanceCullingDrawParams, InstanceCullingDrawParams)
 	SHADER_PARAMETER_RDG_UNIFORM_BUFFER(FMobileBasePassUniformParameters, MobileBasePass)
 	SHADER_PARAMETER_STRUCT_REF(FMobileReflectionCaptureShaderData, ReflectionCapture)
-	SHADER_PARAMETER_RDG_BUFFER_SRV(Buffer<float4>, LocalHeightFogInstances)
+	SHADER_PARAMETER_RDG_BUFFER_SRV(Buffer<float4>, LocalFogVolumeInstances)
 	RENDER_TARGET_BINDING_SLOTS()
 END_SHADER_PARAMETER_STRUCT()
 
@@ -1287,7 +1287,7 @@ void FMobileSceneRenderer::RenderForward(FRDGBuilder& GraphBuilder, FRDGTextureR
 		PassParameters->MobileBasePass = CreateMobileBasePassUniformBuffer(GraphBuilder, View, EMobileBasePass::Opaque, SetupMode, MobileBasePassTextures);
 		PassParameters->ReflectionCapture = View.MobileReflectionCaptureUniformBuffer;
 		PassParameters->RenderTargets = BasePassRenderTargets;
-		PassParameters->LocalHeightFogInstances = View.LocalHeightFogGPUInstanceDataBufferSRV;
+		PassParameters->LocalFogVolumeInstances = View.LocalFogVolumeGPUInstanceDataBufferSRV;
 	
 		BuildInstanceCullingDrawParams(GraphBuilder, View, PassParameters);
 
@@ -1339,7 +1339,7 @@ void FMobileSceneRenderer::RenderForwardSinglePass(FRDGBuilder& GraphBuilder, FM
 		RenderDecals(RHICmdList, View);
 		RenderModulatedShadowProjections(RHICmdList, ViewContext.ViewIndex, View);
 		RenderFog(RHICmdList, View);
-		RenderLocalHeightFogMobile(RHICmdList, View);
+		RenderLocalFogVolumeMobile(RHICmdList, View);
 		// Draw translucency.
 		RenderTranslucency(RHICmdList, View);
 
@@ -1439,7 +1439,7 @@ void FMobileSceneRenderer::RenderForwardMultiPass(FRDGBuilder& GraphBuilder, FMo
 		RenderDecals(RHICmdList, View);
 		RenderModulatedShadowProjections(RHICmdList, ViewContext.ViewIndex, View);
 		RenderFog(RHICmdList, View);
-		RenderLocalHeightFogMobile(RHICmdList, View);
+		RenderLocalFogVolumeMobile(RHICmdList, View);
 		// Draw translucency.
 		RenderTranslucency(RHICmdList, View);
 
@@ -1621,7 +1621,7 @@ void FMobileSceneRenderer::RenderDeferred(FRDGBuilder& GraphBuilder, const FSort
 		PassParameters->MobileBasePass = CreateMobileBasePassUniformBuffer(GraphBuilder, View, EMobileBasePass::Opaque, SetupMode, MobileBasePassTextures);
 		PassParameters->ReflectionCapture = View.MobileReflectionCaptureUniformBuffer;
 		PassParameters->RenderTargets = BasePassRenderTargets;
-		PassParameters->LocalHeightFogInstances = View.LocalHeightFogGPUInstanceDataBufferSRV;
+		PassParameters->LocalFogVolumeInstances = View.LocalFogVolumeGPUInstanceDataBufferSRV;
 		
 		BuildInstanceCullingDrawParams(GraphBuilder, View, PassParameters);
 
@@ -1673,7 +1673,7 @@ void FMobileSceneRenderer::RenderDeferredSinglePass(FRDGBuilder& GraphBuilder, c
 			MobileDeferredCopyBuffer<FMobileDeferredCopyPLSPS>(RHICmdList, View);
 		}
 		RenderFog(RHICmdList, View);
-		RenderLocalHeightFogMobile(RHICmdList, View);
+		RenderLocalFogVolumeMobile(RHICmdList, View);
 		// Draw translucency.
 		RenderTranslucency(RHICmdList, View);
 
@@ -1755,7 +1755,7 @@ void FMobileSceneRenderer::RenderDeferredMultiPass(FRDGBuilder& GraphBuilder, cl
 			
 		RHICmdList.SetCurrentStat(GET_STATID(STAT_CLMM_Translucency));
 		MobileDeferredShadingPass(RHICmdList, ViewContext.ViewIndex, Views.Num(), View, *Scene, SortedLightSet, VisibleLightInfos);
-		RenderLocalHeightFogMobile(RHICmdList, View);
+		RenderLocalFogVolumeMobile(RHICmdList, View);
 		RenderFog(RHICmdList, View);
 		// Draw translucency.
 		RenderTranslucency(RHICmdList, View);
diff --git a/Engine/Source/Runtime/Renderer/Private/RendererScene.cpp b/Engine/Source/Runtime/Renderer/Private/RendererScene.cpp
--- a/Engine/Source/Runtime/Renderer/Private/RendererScene.cpp
+++ b/Engine/Source/Runtime/Renderer/Private/RendererScene.cpp
@@ -6528,9 +6528,9 @@ public:
 	virtual void RemoveExponentialHeightFog(class UExponentialHeightFogComponent* FogComponent) override {}
 	virtual bool HasAnyExponentialHeightFog() const override { return false; }
 
-	virtual void AddLocalHeightFog(class FLocalHeightFogSceneProxy* FogProxy) override {}
-	virtual void RemoveLocalHeightFog(class FLocalHeightFogSceneProxy* FogProxy) override {}
-	virtual bool HasAnyLocalHeightFog() const override { return false; }
+	virtual void AddLocalFogVolume(class FLocalFogVolumeSceneProxy* FogProxy) override {}
+	virtual void RemoveLocalFogVolume(class FLocalFogVolumeSceneProxy* FogProxy) override {}
+	virtual bool HasAnyLocalFogVolume() const override { return false; }
 
 	virtual void AddSkyAtmosphere(FSkyAtmosphereSceneProxy* SkyAtmosphereSceneProxy, bool bStaticLightingBuilt) override {}
 	virtual void RemoveSkyAtmosphere(FSkyAtmosphereSceneProxy* SkyAtmosphereSceneProxy) override {}
diff --git a/Engine/Source/Runtime/Renderer/Private/ScenePrivate.h b/Engine/Source/Runtime/Renderer/Private/ScenePrivate.h
--- a/Engine/Source/Runtime/Renderer/Private/ScenePrivate.h
+++ b/Engine/Source/Runtime/Renderer/Private/ScenePrivate.h
@@ -3078,7 +3078,7 @@ public:
 	/** Used to track the order that skylights were enabled in. */
 	TArray<FVolumetricCloudSceneProxy*> VolumetricCloudStack;
 
-	TArray<FLocalHeightFogSceneProxy*> LocalHeightFogs;
+	TArray<FLocalFogVolumeSceneProxy*> LocalFogVolumes;
 
 	TArray<FSparseVolumeTextureViewerSceneProxy*> SparseVolumeTextureViewers;
 
@@ -3251,9 +3251,9 @@ public:
 	virtual void GetLightIESAtlasSlot(const FLightSceneProxy* Proxy, FLightRenderParameters* Out) override;
 	virtual void GetRectLightAtlasSlot(const FRectLightSceneProxy* Proxy, FLightRenderParameters* Out) override;
 
-	virtual void AddLocalHeightFog(class FLocalHeightFogSceneProxy* FogProxy) override;
-	virtual void RemoveLocalHeightFog(class FLocalHeightFogSceneProxy* FogProxy) override;
-	virtual bool HasAnyLocalHeightFog() const override;
+	virtual void AddLocalFogVolume(class FLocalFogVolumeSceneProxy* FogProxy) override;
+	virtual void RemoveLocalFogVolume(class FLocalFogVolumeSceneProxy* FogProxy) override;
+	virtual bool HasAnyLocalFogVolume() const override;
 
 	virtual void AddSkyAtmosphere(FSkyAtmosphereSceneProxy* SkyAtmosphereSceneProxy, bool bStaticLightingBuilt) override;
 	virtual void RemoveSkyAtmosphere(FSkyAtmosphereSceneProxy* SkyAtmosphereSceneProxy) override;
diff --git a/Engine/Source/Runtime/Renderer/Private/SceneRendering.cpp b/Engine/Source/Runtime/Renderer/Private/SceneRendering.cpp
--- a/Engine/Source/Runtime/Renderer/Private/SceneRendering.cpp
+++ b/Engine/Source/Runtime/Renderer/Private/SceneRendering.cpp
@@ -979,8 +979,8 @@ void FViewInfo::Init()
 	StrataViewData.Reset();
 	HairStrandsViewData = FHairStrandsViewData();
 
-	LocalHeightFogGPUInstanceCount = 0;
-	LocalHeightFogGPUInstanceDataBufferSRV = nullptr;
+	LocalFogVolumeGPUInstanceCount = 0;
+	LocalFogVolumeGPUInstanceDataBufferSRV = nullptr;
 
 	GPUSceneViewId = INDEX_NONE;
 
diff --git a/Engine/Source/Runtime/Renderer/Private/SceneRendering.h b/Engine/Source/Runtime/Renderer/Private/SceneRendering.h
--- a/Engine/Source/Runtime/Renderer/Private/SceneRendering.h
+++ b/Engine/Source/Runtime/Renderer/Private/SceneRendering.h
@@ -1367,8 +1367,8 @@ public:
 
 	FStrataViewData StrataViewData;
 
-	uint32 LocalHeightFogGPUInstanceCount;
-	FRDGBufferSRVRef LocalHeightFogGPUInstanceDataBufferSRV;
+	uint32 LocalFogVolumeGPUInstanceCount;
+	FRDGBufferSRVRef LocalFogVolumeGPUInstanceDataBufferSRV;
 
 	FHairStrandsViewData HairStrandsViewData;
 
