Fix mismatch shader parameter / hlsl type in Hair LUT code.

diff --git a/Engine/Source/Runtime/Renderer/Private/HairStrands/HairStrandsLUT.cpp b/Engine/Source/Runtime/Renderer/Private/HairStrands/HairStrandsLUT.cpp
--- a/Engine/Source/Runtime/Renderer/Private/HairStrands/HairStrandsLUT.cpp
+++ b/Engine/Source/Runtime/Renderer/Private/HairStrands/HairStrandsLUT.cpp
@@ -41,7 +41,7 @@ class FHairLUTCS : public FGlobalShader
 		SHADER_PARAMETER(uint32, ThetaCount)
 		SHADER_PARAMETER(uint32, SampleCountScale)
 		SHADER_PARAMETER(FIntVector, OutputResolution)
-		SHADER_PARAMETER_RDG_TEXTURE_UAV(RWTexture2D, OutputColor)
+		SHADER_PARAMETER_RDG_TEXTURE_UAV(RWTexture3D, OutputColor)
 	END_SHADER_PARAMETER_STRUCT()
 
 public:
@@ -305,6 +305,7 @@ static FRDGTextureRef AddHairCoverageLUTPass(FRDGBuilder& GraphBuilder, const FV
 
 
 	FRDGTextureDesc OutputDesc;
+	OutputDesc.Dimension = ETextureDimension::Texture2D;
 	OutputDesc.Extent.X = OutputResolution.X;
 	OutputDesc.Extent.Y = OutputResolution.Y;
 	OutputDesc.Format = PF_R32_FLOAT;
