Substrate - removed Strata from GPU resource name

diff --git a/Engine/Source/Runtime/Renderer/Private/Strata/Strata.cpp b/Engine/Source/Runtime/Renderer/Private/Strata/Strata.cpp
--- a/Engine/Source/Runtime/Renderer/Private/Strata/Strata.cpp
+++ b/Engine/Source/Runtime/Renderer/Private/Strata/Strata.cpp
@@ -248,7 +248,7 @@ static void InitialiseStrataViewData(FRDGBuilder& GraphBuilder, FViewInfo& View,
 			const EPixelFormat ClassificationTileFormat = GetClassificationTileFormat(DynResIndependentViewSize);
 			const uint32 FormatBytes = ClassificationTileFormat == PF_R16_UINT ? sizeof(uint16) : sizeof(uint32);
 
-			Out.ClassificationTileListBuffer = GraphBuilder.CreateBuffer(FRDGBufferDesc::CreateBufferDesc(FormatBytes, TotalTileCount), TEXT("Substrate.StrataTileListBuffer"));
+			Out.ClassificationTileListBuffer = GraphBuilder.CreateBuffer(FRDGBufferDesc::CreateBufferDesc(FormatBytes, TotalTileCount), TEXT("Substrate.TileListBuffer"));
 			Out.ClassificationTileListBufferSRV = GraphBuilder.CreateSRV(Out.ClassificationTileListBuffer, ClassificationTileFormat);
 			Out.ClassificationTileListBufferUAV = GraphBuilder.CreateUAV(Out.ClassificationTileListBuffer, ClassificationTileFormat);
 		}
