Fixed incorrect logic in FDistanceFieldShadowingCS::ShouldCompilePermutation(...)
- also use RemapPermutation() to prevent usage of invalid permutations at runtime.

diff --git a/Engine/Source/Runtime/Renderer/Private/DistanceFieldShadowing.cpp b/Engine/Source/Runtime/Renderer/Private/DistanceFieldShadowing.cpp
--- a/Engine/Source/Runtime/Renderer/Private/DistanceFieldShadowing.cpp
+++ b/Engine/Source/Runtime/Renderer/Private/DistanceFieldShadowing.cpp
@@ -243,6 +243,8 @@ class FShadowObjectCullPS : public FGlobalShader
 	}
 };
 
+IMPLEMENT_GLOBAL_SHADER(FShadowObjectCullPS, "/Engine/Private/DistanceFieldShadowing.usf", "ShadowObjectCullPS", SF_Pixel);
+
 BEGIN_SHADER_PARAMETER_STRUCT(FShadowMeshSDFObjectCull, )
 	SHADER_PARAMETER_STRUCT_INCLUDE(FShadowObjectCullVS::FParameters, VS)
 	SHADER_PARAMETER_STRUCT_INCLUDE(FShadowObjectCullPS::FParameters, PS)
@@ -250,8 +252,6 @@ BEGIN_SHADER_PARAMETER_STRUCT(FShadowMeshSDFObjectCull, )
 	RENDER_TARGET_BINDING_SLOTS()
 END_SHADER_PARAMETER_STRUCT()
 
-IMPLEMENT_GLOBAL_SHADER(FShadowObjectCullPS, "/Engine/Private/DistanceFieldShadowing.usf", "ShadowObjectCullPS", SF_Pixel);
-
 enum EDistanceFieldShadowingType
 {
 	DFS_DirectionalLightScatterTileCulling,
@@ -299,12 +299,22 @@ class FDistanceFieldShadowingCS : public FGlobalShader
 	class FCompactCulledObjects : SHADER_PERMUTATION_BOOL("COMPACT_CULLED_SHADOW_OBJECTS");
 	using FPermutationDomain = TShaderPermutationDomain<FCullingType, FShadowQuality, FPrimitiveType, FHasPreviousOutput, FOffsetDataStructure, FCompactCulledObjects>;
 
+	static FPermutationDomain RemapPermutation(FPermutationDomain PermutationVector)
+	{
+		if (PermutationVector.Get<FCullingType>() != DFS_DirectionalLightScatterTileCulling)
+		{
+			// Compacting culled objects is only relevant when we do scattered tile culling
+			PermutationVector.Set<FCompactCulledObjects>(false);
+		}
+
+		return PermutationVector;
+	}
+
 	static bool ShouldCompilePermutation(const FGlobalShaderPermutationParameters& Parameters)
 	{
 		const FPermutationDomain PermutationVector(Parameters.PermutationId);
-		if (!PermutationVector.Get<FCompactCulledObjects>() && PermutationVector.Get<FCullingType>() != 0)
+		if (RemapPermutation(PermutationVector) != PermutationVector)
 		{
-			// Compacting culled objects is only relevant when we do scattered tile culling
 			return false;
 		}
 
@@ -820,6 +830,9 @@ void RayTraceShadows(
 		extern int32 GDistanceFieldOffsetDataStructure;
 		PermutationVector.Set< FDistanceFieldShadowingCS::FOffsetDataStructure >(GDistanceFieldOffsetDataStructure);
 		PermutationVector.Set<FDistanceFieldShadowingCS::FCompactCulledObjects>(GDFShadowCompactCulledObjects != 0);
+
+		PermutationVector = FDistanceFieldShadowingCS::RemapPermutation(PermutationVector);
+
 		auto ComputeShader = View.ShaderMap->GetShader< FDistanceFieldShadowingCS >(PermutationVector);
 
 		uint32 GroupSizeX = FMath::DivideAndRoundUp(ScissorRect.Size().X / GetDFShadowDownsampleFactor(), GDistanceFieldShadowTileSizeX);
