Speeds up TSR history resurrection by only running the rejection heuristic on resurrection when needed

diff --git a/Engine/Shaders/Private/TemporalSuperResolution/TSRCommon.ush b/Engine/Shaders/Private/TemporalSuperResolution/TSRCommon.ush
--- a/Engine/Shaders/Private/TemporalSuperResolution/TSRCommon.ush
+++ b/Engine/Shaders/Private/TemporalSuperResolution/TSRCommon.ush
@@ -178,6 +178,10 @@ static const uint kPlusIndexes3x3[5] = { 4, 1, 3, 7, 5 };
 	#define tsr_ushort2x2 uint16_t2x2
 	#define tsr_ushort3x2 uint16_t3x2
 	#define tsr_ushort4x2 uint16_t4x2
+	
+	#define as_tsr_short(x)  asint16(x)
+	#define as_tsr_ushort(x) asuint16(x)
+	#define as_tsr_half(x)   asfloat16(x)
 
 #else
 	#define tsr_half  float
@@ -206,6 +210,10 @@ static const uint kPlusIndexes3x3[5] = { 4, 1, 3, 7, 5 };
 	#define tsr_ushort2x2 uint2x2
 	#define tsr_ushort3x2 uint3x2
 	#define tsr_ushort4x2 uint4x2
+	
+	#define as_tsr_short(x)  asint(x)
+	#define as_tsr_ushort(x) asuint(x)
+	#define as_tsr_half(x)   asfloat(x)
 
 #endif
 
diff --git a/Engine/Shaders/Private/TemporalSuperResolution/TSRConvolutionNetwork.ush b/Engine/Shaders/Private/TemporalSuperResolution/TSRConvolutionNetwork.ush
--- a/Engine/Shaders/Private/TemporalSuperResolution/TSRConvolutionNetwork.ush
+++ b/Engine/Shaders/Private/TemporalSuperResolution/TSRConvolutionNetwork.ush
@@ -167,46 +167,248 @@ tsr_short2 ComputeElementOutputPixelPos(
 }
 
 
+//------------------------------------------------------- PADDING
+
+CALL_SITE_DEBUGLOC
+template<uint LaneStrideX, uint LaneStrideY>
+TLaneVector2D<bool, 1, LaneStrideX, LaneStrideY> IsValidTilePixel(const uint2 TileOverscan)
+{
+	const uint2 TileSize = uint2(LaneStrideX * WAVE_COUNT_X * LANE_COUNT_X, LaneStrideY * WAVE_COUNT_Y * LANE_COUNT_Y);
+
+	TLaneVector2D<bool, 1, LaneStrideX, LaneStrideY> bIsValidTilePixel;
+
+	UNROLL
+	for (uint ElementIndex = 0; ElementIndex < LaneStrideX * LaneStrideY; ElementIndex++)
+	{
+		tsr_short2 PixelOffsetInTile = GetElementPixelOffsetInTile(
+			LaneStrideX,
+			LaneStrideY,
+			GGroupThreadIndex,
+			ElementIndex);
+			
+		bool2 bIsNotPadding = and(
+			PixelOffsetInTile >= tsr_short2(TileOverscan),
+			PixelOffsetInTile <  tsr_short2(TileSize - TileOverscan));
+
+		bIsValidTilePixel.SetElement(ElementIndex, all(bIsNotPadding));
+	}
+
+	return bIsValidTilePixel;
+}
+
+
 //------------------------------------------------------- GROUP SHARED COMMUNICATION
 
 #if defined(LDS_SIZE)
 
-groupshared tsr_half4 SharedData[LDS_SIZE];
+#if CONFIG_COMPILE_FP16
+
+groupshared uint2 SharedData[LDS_SIZE];
+
+void IncrementAtomic(const uint SharedIndex, uint Value)
+{
+	return InterlockedAdd(/* inout */ SharedData[SharedIndex].x, Value);
+}
+
+void WriteAtomic(const uint SharedIndex, uint X)
+{
+	SharedData[SharedIndex].x = X;
+}
+
+uint ReadAtomic(const uint SharedIndex)
+{
+	return SharedData[SharedIndex].x;
+}
+
+#if COMPILER_SUPPORT_UINT16_BITCAST
+
+void WriteElementToLDS(const uint SharedIndex, tsr_ushort4 X)
+{
+	SharedData[SharedIndex] = uint2(
+		bit_cast_uint(X.xy),
+		bit_cast_uint(X.zw));
+}
+
+void ReadElementFromLDS(const uint SharedIndex, out tsr_ushort4 X)
+{
+	uint2 RX = SharedData[SharedIndex];
+	X.xy = bit_cast_uint16_t2(RX.x);
+	X.zw = bit_cast_uint16_t2(RX.y);
+}
+
+#else
+
+void WriteElementToLDS(const uint SharedIndex, tsr_ushort4 X)
+{
+	SharedData[SharedIndex] = uint2(
+		uint(X.x) | (uint(X.y) << 16),
+		uint(X.z) | (uint(X.w) << 16));
+}
+
+void ReadElementFromLDS(const uint SharedIndex, out tsr_ushort4 X)
+{
+	uint2 RX = SharedData[SharedIndex];
+	X.x = tsr_ushort((RX.x >>  0u) & 0xFFFFu);
+	X.y = tsr_ushort((RX.x >> 16u) & 0xFFFFu);
+	X.z = tsr_ushort((RX.y >>  0u) & 0xFFFFu);
+	X.w = tsr_ushort((RX.y >> 16u) & 0xFFFFu);
+}
+
+#endif
+
+#else
+
+groupshared uint4 SharedData[LDS_SIZE];
+
+void IncrementAtomic(const uint SharedIndex, uint Value)
+{
+	return InterlockedAdd(/* inout */ SharedData[SharedIndex].x, Value);
+}
+
+void WriteAtomic(const uint SharedIndex, uint X)
+{
+	SharedData[SharedIndex].x = X;
+}
+
+uint ReadAtomic(const uint SharedIndex)
+{
+	return SharedData[SharedIndex].x;
+}
+
+void WriteElementToLDS(const uint SharedIndex, tsr_ushort4 X)
+{
+	SharedData[SharedIndex] = uint4(X);
+}
+
+void ReadElementFromLDS(const uint SharedIndex, out tsr_ushort4 X)
+{
+	X = tsr_ushort4(SharedData[SharedIndex].xyzw);
+}
+
+#endif
+
+
+
+void WriteElementToLDS(const uint SharedIndex, tsr_ushort X)
+{
+	WriteElementToLDS(SharedIndex, tsr_ushort4(X, 0.0, 0.0, 0.0));
+}
+
+void WriteElementToLDS(const uint SharedIndex, tsr_ushort3 X)
+{
+	WriteElementToLDS(SharedIndex, tsr_ushort4(X, 0.0));
+}
 
 void WriteElementToLDS(const uint SharedIndex, tsr_half X)
 {
-	SharedData[SharedIndex] = tsr_half4(X, 0.0, 0.0, 0.0);
+	WriteElementToLDS(SharedIndex, as_tsr_ushort(X));
 }
 
 void WriteElementToLDS(const uint SharedIndex, tsr_half3 X)
 {
-	SharedData[SharedIndex] = tsr_half4(X, 0.0);
+	WriteElementToLDS(SharedIndex, tsr_ushort3(as_tsr_ushort(X.x), as_tsr_ushort(X.y), as_tsr_ushort(X.z)));
 }
 
 void WriteElementToLDS(const uint SharedIndex, tsr_half4 X)
 {
-	SharedData[SharedIndex] = tsr_half4(X);
+	WriteElementToLDS(SharedIndex, tsr_ushort4(as_tsr_ushort(X.x), as_tsr_ushort(X.y), as_tsr_ushort(X.z), as_tsr_ushort(X.w)));
+}
+
+void ReadElementFromLDS(const uint SharedIndex, out tsr_ushort X)
+{
+	tsr_ushort4 UX;
+	ReadElementFromLDS(SharedIndex, /* out */ UX);
+	X = UX.x;
+}
+
+void ReadElementFromLDS(const uint SharedIndex, out tsr_ushort3 X)
+{
+	tsr_ushort4 UX;
+	ReadElementFromLDS(SharedIndex, /* out */ UX);
+	X = UX.xyz;
 }
 
 void ReadElementFromLDS(const uint SharedIndex, out tsr_half X)
 {
-	X = SharedData[SharedIndex].x;
+	tsr_ushort4 UX;
+	ReadElementFromLDS(SharedIndex, /* out */ UX);
+	X = as_tsr_half(UX.x);
 }
 
 void ReadElementFromLDS(const uint SharedIndex, out tsr_half3 X)
 {
-	X = SharedData[SharedIndex].xyz;
+	tsr_ushort4 UX;
+	ReadElementFromLDS(SharedIndex, /* out */ UX);
+	X.x = as_tsr_half(UX.x);
+	X.y = as_tsr_half(UX.y);
+	X.z = as_tsr_half(UX.z);
 }
 
 void ReadElementFromLDS(const uint SharedIndex, out tsr_half4 X)
 {
-	X = SharedData[SharedIndex].xyzw;
+	tsr_ushort4 UX;
+	ReadElementFromLDS(SharedIndex, /* out */ UX);
+	X.x = as_tsr_half(UX.x);
+	X.y = as_tsr_half(UX.y);
+	X.z = as_tsr_half(UX.z);
+	X.w = as_tsr_half(UX.w);
 }
 
 #endif // defined(LDS_SIZE)
 
 
-//------------------------------------------------------- FUNCTIONS
+//------------------------------------------------------- VOTE
+
+CALL_SITE_DEBUGLOC
+bool AnyThread(bool X)
+#if !COMPILER_SUPPORTS_WAVE_VOTE
+{
+	GroupMemoryBarrierWithGroupSync();
+	
+	if (GGroupThreadIndex == 0)
+	{
+		WriteAtomic(/* SharedIndex = */ 0, 0);
+	}
+
+	GroupMemoryBarrierWithGroupSync();
+	
+	IncrementAtomic(/* SharedIndex = */ 0, uint(X));
+	
+	GroupMemoryBarrierWithGroupSync();
+	
+	return ReadAtomic(/* SharedIndex = */ 0) > 0;
+}
+#elif WAVE_COUNT_X == 1 && WAVE_COUNT_Y == 1
+{
+	return WaveActiveAnyTrue(X);
+}
+#elif COMPILER_SUPPORTS_WAVE_VOTE
+{
+	const uint kLaneCount = WaveGetLaneCount();
+	const uint kThreadCount = LANE_COUNT_X * LANE_COUNT_Y * WAVE_COUNT_X * WAVE_COUNT_Y;
+	const uint kWaveCount = kThreadCount / kLaneCount;
+
+	bool AnyX = WaveActiveAnyTrue(X);
+
+	GroupMemoryBarrierWithGroupSync();
+
+	if (WaveIsFirstLane())
+	{
+		WriteElementToLDS(GGroupThreadIndex / kLaneCount, tsr_ushort(AnyX));
+	}
+		
+	GroupMemoryBarrierWithGroupSync();
+		
+	tsr_ushort WaveX;
+	ReadElementFromLDS(GGroupThreadIndex % kWaveCount, /* out */ WaveX);
+	//WaveX = select(WaveGetLaneIndex() < kWaveCount, WaveX, 0);
+
+	return WaveActiveAnyTrue(WaveX != 0);
+}
+#endif
+
+
+//------------------------------------------------------- ACCESS NEIGHBOR
 
 #if WAVE_COUNT_X == 1 || WAVE_COUNT_Y == 1
 
@@ -420,6 +622,13 @@ TLaneVector2D<ScalarType, VectorSize, LaneStrideX, LaneStrideY> Blur3x3(TLaneVec
 	return WeightedAvg3x3(Center, /* NormalizedWeights = */ float3(0.25, 0.5, 0.25));
 }
 
+CALL_SITE_DEBUGLOC
+template<typename ScalarType, uint VectorSize, uint LaneStrideX, uint LaneStrideY>
+TLaneVector2D<ScalarType, VectorSize, LaneStrideX, LaneStrideY> BlurEven3x3(TLaneVector2D<ScalarType, VectorSize, LaneStrideX, LaneStrideY> Center)
+{
+	return WeightedAvg3x3(Center, /* NormalizedWeights = */ float3(0.3333, 0.3333, 0.3333));
+}
+
 CALL_SITE_DEBUGLOC
 template<typename ScalarType, uint VectorSize, uint LaneStrideX, uint LaneStrideY>
 TLaneVector2D<ScalarType, VectorSize, LaneStrideX, LaneStrideY> Sharpen3x3(TLaneVector2D<ScalarType, VectorSize, LaneStrideX, LaneStrideY> Center)
@@ -553,6 +762,18 @@ TLaneVector2D<ScalarType, VectorSize, LaneStrideX, LaneStrideY> AnnihilateToGuid
 	return Clamp3x3(ToClamp, Clamp3x3(Guide, ToClamp));
 }
 
+CALL_SITE_DEBUGLOC
+template<typename ScalarType, uint VectorSize, uint LaneStrideX, uint LaneStrideY>
+void AnnihilateMutually3x3(
+	TLaneVector2D<ScalarType, VectorSize, LaneStrideX, LaneStrideY> Input,
+	TLaneVector2D<ScalarType, VectorSize, LaneStrideX, LaneStrideY> History,
+	out TLaneVector2D<ScalarType, VectorSize, LaneStrideX, LaneStrideY> AnnihilatedInput,
+	out TLaneVector2D<ScalarType, VectorSize, LaneStrideX, LaneStrideY> AnnihilatedHistory)
+{
+	AnnihilatedInput   = AnnihilateToGuide3x3(Input  , /* Guide = */ History);
+	AnnihilatedHistory = AnnihilateToGuide3x3(History, /* Guide = */ Input);
+}
+
 
 //------------------------------------------------------- MEDIAN 3x3 CONVOLUTION
 
@@ -632,6 +853,21 @@ TLaneVector2D<ScalarType, VectorSize, LaneStrideX, LaneStrideY> Median3x3(TLaneV
 	return Median(M1[0], M1[1], M1[2]);
 }
 
+CALL_SITE_DEBUGLOC
+template<uint VectorSize, uint LaneStrideX, uint LaneStrideY>
+TLaneVector2D<bool, VectorSize, LaneStrideX, LaneStrideY> MedianBool3x3(TLaneVector2D<bool, VectorSize, LaneStrideX, LaneStrideY> Center)
+{
+	TLaneVector2D<tsr_half, VectorSize, LaneStrideX, LaneStrideY> Zero = TLaneVector2D<tsr_half, VectorSize, LaneStrideX, LaneStrideY>::Const(tsr_half(0.0));
+	TLaneVector2D<tsr_half, VectorSize, LaneStrideX, LaneStrideY> One = TLaneVector2D<tsr_half, VectorSize, LaneStrideX, LaneStrideY>::Const(tsr_half(1.0));
+
+	TLaneVector2D<tsr_half, VectorSize, LaneStrideX, LaneStrideY> CenterHalf = select(
+		Center,
+		One,
+		Zero);
+	TLaneVector2D<tsr_half, VectorSize, LaneStrideX, LaneStrideY> CenterSum = Sum3x3(CenterHalf);
+	return CenterSum > TLaneVector2D<tsr_half, VectorSize, LaneStrideX, LaneStrideY>::Const(4.0);
+}
+
 
 //------------------------------------------------------- DOWNSAMPLE 2x2 CONVOLUTIONS
 
diff --git a/Engine/Shaders/Private/TemporalSuperResolution/TSRDecimateHistory.usf b/Engine/Shaders/Private/TemporalSuperResolution/TSRDecimateHistory.usf
--- a/Engine/Shaders/Private/TemporalSuperResolution/TSRDecimateHistory.usf
+++ b/Engine/Shaders/Private/TemporalSuperResolution/TSRDecimateHistory.usf
@@ -69,7 +69,6 @@ RWTexture2DArray<tsr_half4> ReprojectedHistoryMoireOutput;
 RWTexture2D<float2> HoleFilledVelocityOutput;
 RWTexture2D<tsr_half2> HoleFilledVelocityMaskOutput;
 RWTexture2D<tsr_half> ParallaxRejectionMaskOutput;
-RWTexture2D<tsr_half> HistoryResurrectionMaskOutput;
 
 
 //------------------------------------------------------- DEBUG
diff --git a/Engine/Shaders/Private/TemporalSuperResolution/TSRMergeRejection.usf b/Engine/Shaders/Private/TemporalSuperResolution/TSRMergeRejection.usf
--- a/Engine/Shaders/Private/TemporalSuperResolution/TSRMergeRejection.usf
+++ b/Engine/Shaders/Private/TemporalSuperResolution/TSRMergeRejection.usf
@@ -14,19 +14,8 @@
 float4x4 ClipToResurrectionClip;
 
 Texture2D<float> ClosestDepthTexture;
+Texture2D<tsr_half> HistoryResurrectionMaskTexture;
 
-Texture2D<tsr_half4> PrevGuideTexture;
-Texture2D<tsr_half4> PrevMoireTexture;
-Texture2D<tsr_half2> PrevRejectionTexture;
-
-Texture2D<tsr_half4> ResurrectedGuideTexture;
-Texture2D<tsr_half4> ResurrectedMoireTexture;
-Texture2D<tsr_half2> ResurrectedRejectionTexture;
-
-RWTexture2D<tsr_half4> HistoryGuideOutput;
-RWTexture2D<tsr_half4> HistoryMoireOutput;
-RWTexture2D<tsr_half2> HistoryRejectionOutput;
-RWTexture2D<tsr_half> HistoryResurrectionMaskOutput;
 RWTexture2D<float2> DilatedVelocityOutput;
 
 
@@ -46,14 +35,9 @@ void MainCS(
 		tsr_short2(GroupId) * tsr_short2(TILE_SIZE, TILE_SIZE) +
 		Map8x8Tile2x2Lane(GroupThreadIndex));
 	
-	tsr_half2 PrevRejection = PrevRejectionTexture[InputPixelPos];
-	tsr_half2 ResurrectedRejection = ResurrectedRejectionTexture[InputPixelPos];
+	tsr_half ResurrectionMask = HistoryResurrectionMaskTexture[InputPixelPos];
 
-	bool bIsResurrectionBetter = all(ResurrectedRejection.g > PrevRejection.g + 0.5);
-
-	//Debug.r = float(bIsResurrectionBetter);
-	//Debug.rg = saturate(ResurrectedRejection.rg - PrevRejection.rg);
-	Debug.b = bIsResurrectionBetter;
+	bool bIsResurrectionBetter = ResurrectionMask > tsr_half(0.5);
 
 	BRANCH
 	if (bIsResurrectionBetter)
@@ -71,20 +55,13 @@ void MainCS(
 			ResurrectedEncodedScreenVelocity = EncodeVelocityToTexture(float3(ScreenPos - ResurrectionScreen, 0.0)).xy;
 		}
 	
-		HistoryGuideOutput[InputPixelPos]     = ResurrectedGuideTexture[InputPixelPos];
-		HistoryMoireOutput[InputPixelPos]     = ResurrectedMoireTexture[InputPixelPos];
-		HistoryRejectionOutput[InputPixelPos] = ResurrectedRejection;
 		DilatedVelocityOutput[InputPixelPos]  = ResurrectedEncodedScreenVelocity;
 	}
 	else
 	{
-		HistoryGuideOutput[InputPixelPos]     = PrevGuideTexture[InputPixelPos];
-		HistoryMoireOutput[InputPixelPos]     = PrevMoireTexture[InputPixelPos];
-		HistoryRejectionOutput[InputPixelPos] = PrevRejection;
+		// NOP
 	}
 	
-	HistoryResurrectionMaskOutput[InputPixelPos] = bIsResurrectionBetter ? tsr_half(1.0) : tsr_half(0.0);
-
 	#if DEBUG_OUTPUT
 	{
 		DebugOutput[tsr_short3(InputPixelPos, 0)] = Debug;
diff --git a/Engine/Shaders/Private/TemporalSuperResolution/TSRRejectShading.usf b/Engine/Shaders/Private/TemporalSuperResolution/TSRRejectShading.usf
--- a/Engine/Shaders/Private/TemporalSuperResolution/TSRRejectShading.usf
+++ b/Engine/Shaders/Private/TemporalSuperResolution/TSRRejectShading.usf
@@ -25,6 +25,8 @@
 /** Whether to detect flickering. */
 #define CONFIG_MOIRE_DETECTION (DIM_FLICKERING_DETECTION)
 
+#define CONFIG_HISTORY_RESURRECTION (DIM_HISTORY_RESURRECTION)
+
 #define CONFIG_MANUAL_LDS_SPILL 0
 
 
@@ -94,7 +96,10 @@
 #define tsr_simd_half  TLaneVector2D<tsr_half, 1, LANE_STRIDE_X, LANE_STRIDE_Y>
 #define tsr_simd_half4 TLaneVector2D<tsr_half, 4, LANE_STRIDE_X, LANE_STRIDE_Y>
 #define tsr_simd_halfC TLaneVector2D<tsr_half, CONFIG_CHANNEL_COUNT, LANE_STRIDE_X, LANE_STRIDE_Y>
+#define tsr_simd_ushort  TLaneVector2D<tsr_ushort, 1, LANE_STRIDE_X, LANE_STRIDE_Y>
 #define tsr_simd_bool  TLaneVector2D<bool, 1, LANE_STRIDE_X, LANE_STRIDE_Y>
+#define tsr_simd_bool4   TLaneVector2D<bool, 4, LANE_STRIDE_X, LANE_STRIDE_Y>
+#define tsr_simd_boolC   TLaneVector2D<bool, CONFIG_CHANNEL_COUNT, LANE_STRIDE_X, LANE_STRIDE_Y>
 
 #ifndef PLATFORM_SPECIFIC_ISOLATE
 	#define PLATFORM_SPECIFIC_ISOLATE ISOLATE
@@ -118,12 +123,14 @@ Texture2D<tsr_half>  InputMoireLumaTexture;
 Texture2D<tsr_half4> InputSceneTranslucencyTexture;
 Texture2D<tsr_halfC> ReprojectedHistoryGuideTexture;
 Texture2D<tsr_half4> ReprojectedHistoryMoireTexture;
+Texture2D<tsr_halfC> ResurrectedHistoryGuideTexture;
 Texture2D<tsr_half> ParallaxRejectionMaskTexture;
 Texture2D<tsr_ushort> IsMovingMaskTexture;
 
 RWTexture2D<tsr_half4> HistoryGuideOutput;
 RWTexture2D<tsr_half4> HistoryMoireOutput;
 RWTexture2D<tsr_half2> HistoryRejectionOutput;
+RWTexture2D<tsr_half>  HistoryResurrectionMaskOutput;
 RWTexture2D<tsr_halfC> InputSceneColorOutput;
 RWTexture2D<tsr_half> InputSceneColorLdrLumaOutput;
 
@@ -135,6 +142,18 @@ static tsr_simd_halfC Debug[DEBUG_ARRAY_SIZE];
 
 //------------------------------------------------------- GENERAL OPERATIONS
 
+CALL_SITE_DEBUGLOC
+tsr_simd_halfC ChangePerceptionAdd(tsr_simd_halfC SrcColor, const float SourcePerceptionAdd, const float DestPerceptionAdd)
+{
+	SrcColor = SrcColor * rcp(tsr_simd_halfC::Const(tsr_half(SourcePerceptionAdd)) - SrcColor);
+	SrcColor = SrcColor * rcp(SrcColor + tsr_simd_halfC::Const(tsr_half(DestPerceptionAdd)));
+
+	return SrcColor;
+}
+
+
+//------------------------------------------------------- MAIN ALGORITHMS
+
 void ComputeMoireError(
 	tsr_simd_bool  bIsValid,
 	tsr_simd_half  bIsStatic,
@@ -305,8 +324,8 @@ void ComputeMoireError(
 template<bool ApplyMoireError>
 void MeasureRejection(
 	tsr_simd_bool  bIsValid,
-	tsr_simd_halfC OriginalInput,
-	tsr_simd_halfC OriginalHistory,
+	tsr_simd_halfC ExposedInput,
+	tsr_simd_halfC ExposedHistory,
 	tsr_simd_half  MoireError,
 	out tsr_simd_half OutRejectionBlendFinal,
 	out tsr_simd_half OutRejectionClampBlend)
@@ -314,20 +333,9 @@ void MeasureRejection(
 	const tsr_half FilteringWeight = rcp(tsr_half(1.0 + 4 * 0.5 + 4 * 0.25));
 	tsr_simd_halfC BackbufferQuantizationErrorVector = tsr_simd_halfC::Const(MeasureBackbufferLDRQuantizationError());
 		
-	// Convert input to LDR
-	tsr_simd_halfC ExposedInput = OriginalInput * rcp(OriginalInput + tsr_simd_halfC::Const(tsr_half(PerceptionAdd)));
-	
-	// Change the exposure of history for the shading rejection
-	tsr_simd_halfC ExposedHistory = OriginalHistory;
-	{
-		ExposedHistory = ExposedHistory * rcp(tsr_simd_halfC::Const(1.0) - ExposedHistory);
-		ExposedHistory = ExposedHistory * rcp(ExposedHistory + tsr_simd_halfC::Const(tsr_half(PerceptionAdd)));
-	}
+	tsr_simd_halfC ClampedInput, ClampedHistory;
+	AnnihilateMutually3x3(ExposedInput, ExposedHistory, /* out */ ClampedInput, /* out */ ClampedHistory);
 
-	// Mutual Annihilation
-	tsr_simd_halfC ClampedInput   = AnnihilateToGuide3x3(ExposedInput  , /* Guide = */ ExposedHistory);
-	tsr_simd_halfC ClampedHistory = AnnihilateToGuide3x3(ExposedHistory, /* Guide = */ ExposedInput);
-		
 	// Filter clamped input and history to increase quantization of the rejection detection
 	tsr_simd_halfC FilteredInput   = Blur3x3(ClampedInput);
 	tsr_simd_halfC FilteredHistory = Blur3x3(ClampedHistory);
@@ -445,6 +453,9 @@ void MainCS(
 		tsr_simd_half4 PrevMoireHistory;
 		tsr_simd_half bIsStatic;
 	#endif
+	#if CONFIG_HISTORY_RESURRECTION
+		tsr_simd_halfC OriginalResurrectedHistory;
+	#endif
 	PLATFORM_SPECIFIC_ISOLATE
 	{
 		UNROLL_N(SIMD_SIZE)
@@ -473,6 +484,9 @@ void MainCS(
 				PrevMoireHistory.SetElement(SimdIndex, ReprojectedHistoryMoireTexture[ClampedFetchPixelPos]);
 				bIsStatic.SetElement(SimdIndex, tsr_half(1.0) - tsr_half(IsMovingMaskTexture[ClampedFetchPixelPos]) * rcp(tsr_half(255.0)));
 			#endif
+			#if CONFIG_HISTORY_RESURRECTION
+				OriginalResurrectedHistory.SetElement(SimdIndex, ResurrectedHistoryGuideTexture[ClampedFetchPixelPos]);
+			#endif
 		}
 	}
 	
@@ -515,10 +529,7 @@ void MainCS(
 					InputInfo_ViewportMax,
 					GroupId,
 					GroupThreadIndex,
-					SimdIndex,
-					/* OutputDataOffset = */ int2(0, 0),
-					/* OutputResDivisor = */ 1,
-					/* bDebugDisablePadding = */ DEBUG_DISABLE_PADDING != 0);
+					SimdIndex);
 		
 				InputSceneColorOutput[OutputPixelPos] = CenterFinalSceneColor.GetElement(SimdIndex);
 				InputSceneColorLdrLumaOutput[OutputPixelPos] = CenterLumaLDR.GetElement(SimdIndex);
@@ -581,10 +592,7 @@ void MainCS(
 				InputInfo_ViewportMax,
 				GroupId,
 				GroupThreadIndex,
-				SimdIndex,
-				/* OutputDataOffset = */ int2(0, 0),
-				/* OutputResDivisor = */ 1,
-				/* bDebugDisablePadding = */ DEBUG_DISABLE_PADDING != 0);
+				SimdIndex);
 
 			#if CONFIG_MOIRE_DETECTION
 				HistoryMoireOutput[OutputPixelPos] = MoireHistory.GetElement(SimdIndex);
@@ -595,18 +603,72 @@ void MainCS(
 	// Measure rejection
 	tsr_simd_half RejectionBlendFinal = tsr_simd_half::Const(1.0);
 	tsr_simd_half RejectionClampBlend = tsr_simd_half::Const(1.0);
-	MeasureRejection</* ApplyMoireError = */ CONFIG_MOIRE_DETECTION != 0>(
-		bIsValid,
-		OriginalInput,
-		OriginalHistory,
-		MoireError,
-		/* out */ RejectionBlendFinal,
-		/* out */ RejectionClampBlend);
+	#if CONFIG_HISTORY_RESURRECTION
+		tsr_simd_bool bResurrectHistory = tsr_simd_bool::Const(false);
+	#endif
+	{
+		// Convert input to LDR
+		tsr_simd_halfC ExposedInput = OriginalInput * rcp(OriginalInput + tsr_simd_halfC::Const(tsr_half(PerceptionAdd)));
 	
+		// Change the exposure of history for the shading rejection
+		tsr_simd_halfC ExposedHistory = ChangePerceptionAdd(
+			OriginalHistory,
+			/* SourcePerceptionAdd = */ 1.0,
+			/* DestPerceptionAdd   = */ PerceptionAdd);
+			
+		MeasureRejection</* ApplyMoireError = */ CONFIG_MOIRE_DETECTION != 0>(
+			bIsValid,
+			ExposedInput,
+			ExposedHistory,
+			MoireError,
+			/* out */ RejectionBlendFinal,
+			/* out */ RejectionClampBlend);
+	
+		#if CONFIG_HISTORY_RESURRECTION
+		{
+			tsr_simd_halfC ExposedResurrectedHistory = ChangePerceptionAdd(
+				OriginalResurrectedHistory,
+				/* SourcePerceptionAdd = */ 1.0,
+				/* DestPerceptionAdd   = */ PerceptionAdd);
+				
+			tsr_simd_halfC PrevMatch         = abs(ExposedInput - ExposedHistory);
+			tsr_simd_halfC ResurrectionMatch = abs(ExposedInput - ExposedResurrectedHistory);
+
+			tsr_simd_halfC MatchDiff = PrevMatch - ResurrectionMatch;
+			tsr_simd_half TotalMatchDiff = BlurEven3x3(MatchDiff[0] + MatchDiff[1] + MatchDiff[2]);
+
+			tsr_simd_bool bResurrectionIsCloser = MedianBool3x3(TotalMatchDiff > tsr_simd_half::Const(tsr_half(0.05 * 3.0)));
+			tsr_simd_bool bPrevIsNotGoodEnough = RejectionBlendFinal < tsr_simd_half::Const(0.5);
+			tsr_simd_bool bIsValidTilePixel = IsValidTilePixel<LANE_STRIDE_X, LANE_STRIDE_Y>(TileOverscan);
+			tsr_simd_bool bShouldComputeResurrection = and(and(bResurrectionIsCloser, bPrevIsNotGoodEnough), bIsValidTilePixel);
+
+			bool bComputeResurrection = AnyThread(AnyComponent(bShouldComputeResurrection));
+			BRANCH
+			if (bComputeResurrection)
+			{
+				tsr_simd_half ResurrectionRejectionBlendFinal = tsr_simd_half::Const(1.0);
+				tsr_simd_half ResurrectionRejectionClampBlend = tsr_simd_half::Const(1.0);
+				MeasureRejection</* ApplyMoireError = */ false>(
+					bIsValid,
+					ExposedInput,
+					ExposedResurrectedHistory,
+					/* MoireError = */ tsr_simd_half::Const(0.0),
+					/* out */ ResurrectionRejectionBlendFinal,
+					/* out */ ResurrectionRejectionClampBlend);
+			
+				tsr_simd_bool bResurrectionIsBetter = (ResurrectionRejectionBlendFinal - RejectionBlendFinal) > tsr_simd_half::Const(0.1);
+				bResurrectHistory = and(bResurrectionIsCloser, bResurrectionIsBetter);
+				
+				RejectionBlendFinal = select(bResurrectHistory, ResurrectionRejectionBlendFinal, RejectionBlendFinal);
+				RejectionClampBlend = select(bResurrectHistory, ResurrectionRejectionClampBlend, RejectionClampBlend);
+			}
+		}
+		#endif
+	}
+
 	// Update the guide for next frame.
 	tsr_simd_halfC FinalGuide;
 	{
-		
 		// Convert input to LDR
 		tsr_simd_halfC ExposedInput = OriginalInput * rcp(OriginalInput + tsr_simd_halfC::Const(tsr_half(1.0)));
 	
@@ -627,11 +689,15 @@ void MainCS(
 			BlendFinal.SetElement(SimdIndex, max(RejectedBlendFactor, TotalBlend));
 		}
 		
+		#if CONFIG_HISTORY_RESURRECTION
+			OriginalHistory = select(tsr_simd_boolC::Vectorize(bResurrectHistory), OriginalResurrectedHistory, OriginalHistory);
+		#endif
+
 		FinalGuide = ExposedInput * tsr_simd_halfC::Vectorize(BlendFinal) - OriginalHistory * tsr_simd_halfC::Vectorize(BlendFinal - 1.0);
 	}
 
 	// Mark the top left corner of the tile for debugging purposes.
-	#if DEBUG_OUTPUT
+	#if DEBUG_OUTPUT && DEBUG_DISABLE_PADDING == 2
 	if (GGroupThreadIndex == 0)
 	{
 		UNROLL_N(DEBUG_ARRAY_SIZE)
@@ -656,10 +722,7 @@ void MainCS(
 				InputInfo_ViewportMax,
 				GroupId,
 				GroupThreadIndex,
-				SimdIndex,
-				/* OutputDataOffset = */ int2(0, 0),
-				/* OutputResDivisor = */ 1,
-				/* bDebugDisablePadding = */ DEBUG_DISABLE_PADDING != 0);
+				SimdIndex);
 
 			tsr_half LocalRejectionBlendFinal = RejectionBlendFinal.GetElement(SimdIndex);
 			tsr_half LocalRejectionClampBlend = RejectionClampBlend.GetElement(SimdIndex);
@@ -683,13 +746,33 @@ void MainCS(
 			#endif
 			HistoryRejectionOutput[OutputPixelPos] = tsr_half2(LocalRejectionBlendFinal, LocalRejectionClampBlend);
 			
-			#if DEBUG_OUTPUT
+			#if CONFIG_HISTORY_RESURRECTION
+				HistoryResurrectionMaskOutput[OutputPixelPos] = select(bResurrectHistory.GetElement(SimdIndex), tsr_half(1.0), tsr_half(0.0));
+			#endif
+		}
+		
+		#if DEBUG_OUTPUT
+		UNROLL_N(SIMD_SIZE)
+		for (uint SimdIndex = 0; SimdIndex < SIMD_SIZE; SimdIndex++)
+		{
+			const tsr_short2 OutputPixelPos = ComputeElementOutputPixelPos(
+				/* LaneStride = */ uint2(LANE_STRIDE_X, LANE_STRIDE_Y),
+				TileOverscan,
+				InputInfo_ViewportMin,
+				InputInfo_ViewportMax,
+				GroupId,
+				GroupThreadIndex,
+				SimdIndex,
+				/* OutputDataOffset = */ int2(0, 0),
+				/* OutputResDivisor = */ 1,
+				/* bDebugDisablePadding = */ DEBUG_DISABLE_PADDING != 0);
+
 			UNROLL_N(DEBUG_ARRAY_SIZE)
 			for (uint DebugId = 0; DebugId < DEBUG_ARRAY_SIZE; DebugId++)
 			{
 				DebugOutput[tsr_short3(OutputPixelPos, DebugId)] = float4(Debug[DebugId].GetElement(SimdIndex), 0.0);
 			}
-			#endif
 		}
+		#endif
 	}
 }
diff --git a/Engine/Shaders/Public/LaneVectorization.ush b/Engine/Shaders/Public/LaneVectorization.ush
--- a/Engine/Shaders/Public/LaneVectorization.ush
+++ b/Engine/Shaders/Public/LaneVectorization.ush
@@ -1015,3 +1015,15 @@ TVECTOR_FUNCTION_3PARAMS(max3, ScalarType);
 #if COMPILER_SUPPORTS_MED3
 	TVECTOR_FUNCTION_3PARAMS(med3, ScalarType);
 #endif
+
+
+// ------------- AnyElement & AllElement
+
+CALL_SITE_DEBUGLOC
+template<uint VectorSize, uint SimdSizeX, uint SimdSizeY>
+bool AnyComponent(TLaneVector2D<bool, VectorSize, SimdSizeX, SimdSizeY> A)
+{
+	bool R = any(A.GetElement(0));
+	UNROLL for (uint SimdIndex = 1; SimdIndex < SimdSizeX * SimdSizeY; SimdIndex++) { R = R || any(A.GetElement(SimdIndex)); }
+	return R;
+}
diff --git a/Engine/Source/Runtime/Renderer/Private/PostProcess/TemporalSuperResolution.cpp b/Engine/Source/Runtime/Renderer/Private/PostProcess/TemporalSuperResolution.cpp
--- a/Engine/Source/Runtime/Renderer/Private/PostProcess/TemporalSuperResolution.cpp
+++ b/Engine/Source/Runtime/Renderer/Private/PostProcess/TemporalSuperResolution.cpp
@@ -581,7 +581,6 @@ class FTSRDecimateHistoryCS : public FTSRShader
 		SHADER_PARAMETER_RDG_TEXTURE_UAV(RWTexture2D, HoleFilledVelocityMaskOutput)
 
 		SHADER_PARAMETER_RDG_TEXTURE_UAV(RWTexture2D, ParallaxRejectionMaskOutput)
-		SHADER_PARAMETER_RDG_TEXTURE_UAV(RWTexture2D, HistoryResurrectionMaskOutput)
 
 		SHADER_PARAMETER_RDG_TEXTURE_UAV(RWTexture2DArray, DebugOutput)
 	END_SHADER_PARAMETER_STRUCT()
@@ -594,8 +593,9 @@ class FTSRRejectShadingCS : public FTSRShader
 
 	class FWaveSizeOps : SHADER_PERMUTATION_SPARSE_INT("DIM_WAVE_SIZE", 0, 32, 64);
 	class FFlickeringDetectionDim : SHADER_PERMUTATION_BOOL("DIM_FLICKERING_DETECTION");
+	class FHistoryResurrectionDim : SHADER_PERMUTATION_BOOL("DIM_HISTORY_RESURRECTION");
 
-	using FPermutationDomain = TShaderPermutationDomain<FWaveSizeOps, FFlickeringDetectionDim, FTSRShader::F16BitVALUDim>;
+	using FPermutationDomain = TShaderPermutationDomain<FWaveSizeOps, FFlickeringDetectionDim, FHistoryResurrectionDim, FTSRShader::F16BitVALUDim>;
 
 	BEGIN_SHADER_PARAMETER_STRUCT(FParameters, )
 		SHADER_PARAMETER_STRUCT_INCLUDE(FTSRCommonParameters, CommonParameters)
@@ -614,12 +614,14 @@ class FTSRRejectShadingCS : public FTSRShader
 		SHADER_PARAMETER_RDG_TEXTURE(Texture2D, InputSceneTranslucencyTexture)
 		SHADER_PARAMETER_RDG_TEXTURE_SRV(Texture2D, ReprojectedHistoryGuideTexture)
 		SHADER_PARAMETER_RDG_TEXTURE_SRV(Texture2D, ReprojectedHistoryMoireTexture)
+		SHADER_PARAMETER_RDG_TEXTURE_SRV(Texture2D, ResurrectedHistoryGuideTexture)
 		SHADER_PARAMETER_RDG_TEXTURE(Texture2D, ParallaxRejectionMaskTexture)
 		SHADER_PARAMETER_RDG_TEXTURE_SRV(Texture2D, IsMovingMaskTexture)
 
 		SHADER_PARAMETER_RDG_TEXTURE_UAV(RWTexture2D, HistoryGuideOutput)
 		SHADER_PARAMETER_RDG_TEXTURE_UAV(RWTexture2D, HistoryMoireOutput)
 		SHADER_PARAMETER_RDG_TEXTURE_UAV(RWTexture2D, HistoryRejectionOutput)
+		SHADER_PARAMETER_RDG_TEXTURE_UAV(RWTexture2D, HistoryResurrectionMaskOutput)
 		SHADER_PARAMETER_RDG_TEXTURE_UAV(RWTexture2D, InputSceneColorOutput)
 		SHADER_PARAMETER_RDG_TEXTURE_UAV(RWTexture2D, InputSceneColorLdrLumaOutput)
 		SHADER_PARAMETER_RDG_TEXTURE_UAV(RWTexture2D, DebugOutput)
@@ -707,19 +709,9 @@ class FTSRMergeRejectionCS : public FTSRShader
 		SHADER_PARAMETER(FMatrix44f, ClipToResurrectionClip)
 
 		SHADER_PARAMETER_RDG_TEXTURE(Texture2D, ClosestDepthTexture)
-		SHADER_PARAMETER_RDG_TEXTURE(Texture2D, PrevGuideTexture)
-		SHADER_PARAMETER_RDG_TEXTURE(Texture2D, PrevMoireTexture)
-		SHADER_PARAMETER_RDG_TEXTURE(Texture2D, PrevRejectionTexture)
-		SHADER_PARAMETER_RDG_TEXTURE(Texture2D, ResurrectedGuideTexture)
-		SHADER_PARAMETER_RDG_TEXTURE(Texture2D, ResurrectedMoireTexture)
-		SHADER_PARAMETER_RDG_TEXTURE(Texture2D, ResurrectedRejectionTexture)
+		SHADER_PARAMETER_RDG_TEXTURE(Texture2D, HistoryResurrectionMaskTexture)
 
-		SHADER_PARAMETER_RDG_TEXTURE_UAV(RWTexture2D, HistoryGuideOutput)
-		SHADER_PARAMETER_RDG_TEXTURE_UAV(RWTexture2D, HistoryMoireOutput)
-		SHADER_PARAMETER_RDG_TEXTURE_UAV(RWTexture2D, HistoryRejectionOutput)
-		SHADER_PARAMETER_RDG_TEXTURE_UAV(RWTexture2D, HistoryResurrectionMaskOutput)
 		SHADER_PARAMETER_RDG_TEXTURE_UAV(RWTexture2D, DilatedVelocityOutput)
-
 		SHADER_PARAMETER_RDG_TEXTURE_UAV(RWTexture2DArray, DebugOutput)
 	END_SHADER_PARAMETER_STRUCT()
 }; // class FTSRMergeRejectionCS
@@ -1795,7 +1787,6 @@ FDefaultTemporalUpscaler::FOutputs AddTemporalSuperResolutionPasses(
 	FRDGTextureRef ReprojectedHistoryGuideTexture = nullptr;
 	FRDGTextureRef ReprojectedHistoryMoireTexture = nullptr;
 	FRDGTextureRef ParallaxRejectionMaskTexture = nullptr;
-	FRDGTextureRef HistoryResurrectionMaskTexture = nullptr;
 	FRDGTextureRef HoleFilledVelocityMaskTexture = nullptr;
 	{
 		FRDGTextureRef HoleFilledVelocityTexture;
@@ -1807,10 +1798,6 @@ FDefaultTemporalUpscaler::FOutputs AddTemporalSuperResolutionPasses(
 				/* InFlags = */ TexCreate_ShaderResource | TexCreate_UAV);
 
 			ParallaxRejectionMaskTexture = GraphBuilder.CreateTexture(Desc, TEXT("TSR.ParallaxRejectionMask"));
-			if (bCanResurrectHistory)
-			{
-				HistoryResurrectionMaskTexture = GraphBuilder.CreateTexture(Desc, TEXT("TSR.HistoryResurrectionMask"));
-			}
 
 			Desc.Format = PF_G16R16;
 			HoleFilledVelocityTexture = GraphBuilder.CreateTexture(Desc, TEXT("TSR.Velocity.HoleFilled"));
@@ -1887,10 +1874,6 @@ FDefaultTemporalUpscaler::FOutputs AddTemporalSuperResolutionPasses(
 		PassParameters->HoleFilledVelocityOutput = GraphBuilder.CreateUAV(HoleFilledVelocityTexture);
 		PassParameters->HoleFilledVelocityMaskOutput = GraphBuilder.CreateUAV(HoleFilledVelocityMaskTexture);
 		PassParameters->ParallaxRejectionMaskOutput = GraphBuilder.CreateUAV(ParallaxRejectionMaskTexture);
-		if (bCanResurrectHistory)
-		{
-			PassParameters->HistoryResurrectionMaskOutput = GraphBuilder.CreateUAV(HistoryResurrectionMaskTexture);
-		}
 		PassParameters->DebugOutput = CreateDebugUAV(InputExtent, TEXT("Debug.TSR.DecimateHistory"));
 
 		FTSRDecimateHistoryCS::FPermutationDomain PermutationVector;
@@ -1920,12 +1903,9 @@ FDefaultTemporalUpscaler::FOutputs AddTemporalSuperResolutionPasses(
 	}
 
 	// Perform a history reject the history.
+	FRDGTextureRef HistoryRejectionTexture = nullptr;
+	FRDGTextureRef HistoryResurrectionMaskTexture = nullptr;
 	FRDGTextureRef InputSceneColorLdrLumaTexture = nullptr;
-	auto RejectReprojectedGuide = [&](
-		const TCHAR* UsageString,
-		int32 HistoryReprojectionIndex,
-		FRDGTextureUAVRef HistoryGuideOutput,
-		FRDGTextureUAVRef HistoryMoireOutput) -> FRDGTextureRef
 	{
 		const int32 GroupTileSize = CVarTSRShadingTileSize.GetValueOnRenderThread();
 
@@ -1941,7 +1921,7 @@ FDefaultTemporalUpscaler::FOutputs AddTemporalSuperResolutionPasses(
 			InputSceneColorTexture = GraphBuilder.CreateTexture(Desc, TEXT("TSR.SceneColor"));
 		}
 
-		const bool bComputeLdrLuma = InputSceneColorLdrLumaTexture == nullptr && RejectionAntiAliasingQuality > 0;
+		const bool bComputeLdrLuma = RejectionAntiAliasingQuality > 0;
 		if (bComputeLdrLuma)
 		{
 			FRDGTextureDesc Desc = FRDGTextureDesc::Create2D(
@@ -1953,7 +1933,6 @@ FDefaultTemporalUpscaler::FOutputs AddTemporalSuperResolutionPasses(
 			InputSceneColorLdrLumaTexture = GraphBuilder.CreateTexture(Desc, TEXT("TSR.SceneColorLdrLuma"));
 		}
 
-		FRDGTextureRef HistoryRejectionTexture;
 		{
 			FRDGTextureDesc Desc = FRDGTextureDesc::Create2D(
 				InputExtent,
@@ -1964,6 +1943,17 @@ FDefaultTemporalUpscaler::FOutputs AddTemporalSuperResolutionPasses(
 			HistoryRejectionTexture = GraphBuilder.CreateTexture(Desc, TEXT("TSR.HistoryRejection"));
 		}
 
+		if (bCanResurrectHistory)
+		{
+			FRDGTextureDesc Desc = FRDGTextureDesc::Create2D(
+				InputExtent,
+				PF_R8,
+				FClearValueBinding::None,
+				/* InFlags = */ TexCreate_ShaderResource | TexCreate_UAV);
+
+			HistoryResurrectionMaskTexture = GraphBuilder.CreateTexture(Desc, TEXT("TSR.ResurrectionMask"));
+		}
+
 		FScreenPassTextureViewport TranslucencyViewport(
 			SeparateTranslucencyTexture->Desc.Extent, SeparateTranslucencyRect);
 
@@ -1975,7 +1965,7 @@ FDefaultTemporalUpscaler::FOutputs AddTemporalSuperResolutionPasses(
 		PassParameters->TranslucencyTextureUVMin = GetScreenPassTextureViewportParameters(TranslucencyViewport).UVViewportBilinearMin;
 		PassParameters->TranslucencyTextureUVMax = GetScreenPassTextureViewportParameters(TranslucencyViewport).UVViewportBilinearMax;
 		PassParameters->HistoryGuideQuantizationError = ComputePixelFormatQuantizationError(History.GuideArray->Desc.Format);
-		PassParameters->FlickeringFramePeriod = HistoryMoireOutput ? FlickeringFramePeriod : 0.0f;
+		PassParameters->FlickeringFramePeriod = FlickeringFramePeriod;
 		PassParameters->TheoricBlendFactor = 1.0f / (1.0f + MaxHistorySampleCount / OutputToInputResolutionFractionSquare);
 		PassParameters->TileOverscan = FMath::Clamp(CVarTSRShadingTileOverscan.GetValueOnRenderThread(), 2, GroupTileSize / 2 - 1);
 		PassParameters->PerceptionAdd = FMath::Pow(0.5f, CVarTSRShadingExposureOffset.GetValueOnRenderThread());
@@ -1992,157 +1982,121 @@ FDefaultTemporalUpscaler::FOutputs AddTemporalSuperResolutionPasses(
 		}
 		PassParameters->InputSceneTranslucencyTexture = SeparateTranslucencyTexture;
 		PassParameters->ReprojectedHistoryGuideTexture = GraphBuilder.CreateSRV(FRDGTextureSRVDesc::CreateForSlice(
-			ReprojectedHistoryGuideTexture, HistoryReprojectionIndex));
+			ReprojectedHistoryGuideTexture, /* SliceIndex = */ 0));
 		PassParameters->ReprojectedHistoryMoireTexture = ReprojectedHistoryMoireTexture ? GraphBuilder.CreateSRV(FRDGTextureSRVDesc::CreateForSlice(
-			ReprojectedHistoryMoireTexture, HistoryReprojectionIndex)) : nullptr;
+			ReprojectedHistoryMoireTexture, /* SliceIndex = */ 0)) : nullptr;
+		if (bCanResurrectHistory)
+		{
+			PassParameters->ResurrectedHistoryGuideTexture = GraphBuilder.CreateSRV(FRDGTextureSRVDesc::CreateForSlice(
+				ReprojectedHistoryGuideTexture, /* SliceIndex = */ 1));
+		}
 		PassParameters->ParallaxRejectionMaskTexture = ParallaxRejectionMaskTexture;
 		PassParameters->IsMovingMaskTexture = IsMovingMaskTexture;
 
-		PassParameters->HistoryGuideOutput = HistoryGuideOutput;
-		PassParameters->HistoryMoireOutput = HistoryMoireOutput;
-		PassParameters->HistoryRejectionOutput = GraphBuilder.CreateUAV(HistoryRejectionTexture);
-		PassParameters->InputSceneColorOutput = bComputeInputSceneColorTexture
-			? GraphBuilder.CreateUAV(InputSceneColorTexture)
-			: CreateDummyUAV(GraphBuilder, HistoryColorFormat);
-		PassParameters->InputSceneColorLdrLumaOutput = bComputeLdrLuma
-			? GraphBuilder.CreateUAV(InputSceneColorLdrLumaTexture)
-			: CreateDummyUAV(GraphBuilder, PF_R8);
-		PassParameters->DebugOutput = CreateDebugUAV(InputExtent, TEXT("Debug.TSR.RejectShading"));
+		// Outputs
+		{
+			if (View.bStatePrevViewInfoIsReadOnly)
+			{
+				PassParameters->HistoryGuideOutput = CreateDummyUAV(GraphBuilder, History.GuideArray->Desc.Format);
+			}
+			else
+			{
+				FRDGTextureUAVDesc GuideUAVDesc(History.GuideArray);
+				GuideUAVDesc.FirstArraySlice = CurrentFrameSliceIndex;
+				GuideUAVDesc.NumArraySlices = 1;
+				GuideUAVDesc.DimensionOverride = ETextureDimension::Texture2D;
+
+				PassParameters->HistoryGuideOutput = GraphBuilder.CreateUAV(GuideUAVDesc);
+			}
+
+			// Output history for the anti-flickering heuristic that know how something flicker overtime.
+			if (FlickeringFramePeriod == 0.0f)
+			{
+				// NOP
+			}
+			else if (View.bStatePrevViewInfoIsReadOnly)
+			{
+				PassParameters->HistoryMoireOutput = CreateDummyUAV(GraphBuilder, History.MoireArray->Desc.Format);
+			}
+			else
+			{
+				FRDGTextureUAVDesc MoireUAVDesc(History.MoireArray);
+				MoireUAVDesc.FirstArraySlice = CurrentFrameSliceIndex;
+				MoireUAVDesc.NumArraySlices = 1;
+				MoireUAVDesc.DimensionOverride = ETextureDimension::Texture2D;
+
+				PassParameters->HistoryMoireOutput = GraphBuilder.CreateUAV(MoireUAVDesc);
+			}
+
+			// Output how the history should rejected in the HistoryUpdate
+			PassParameters->HistoryRejectionOutput = GraphBuilder.CreateUAV(HistoryRejectionTexture);
+
+			// Output whether resurrection should be applied.
+			if (bCanResurrectHistory)
+			{
+				PassParameters->HistoryResurrectionMaskOutput = GraphBuilder.CreateUAV(HistoryResurrectionMaskTexture);
+			}
+
+			// Output the composed translucency and opaque scene color to speed up HistoryUpdate
+			PassParameters->InputSceneColorOutput = bComputeInputSceneColorTexture
+				? GraphBuilder.CreateUAV(InputSceneColorTexture)
+				: CreateDummyUAV(GraphBuilder, HistoryColorFormat);
+
+			// Output LDR luminance to speed up spatial anti-aliaser
+			PassParameters->InputSceneColorLdrLumaOutput = bComputeLdrLuma
+				? GraphBuilder.CreateUAV(InputSceneColorLdrLumaTexture)
+				: CreateDummyUAV(GraphBuilder, PF_R8);
+
+			PassParameters->DebugOutput = CreateDebugUAV(InputExtent, TEXT("Debug.TSR.RejectShading"));
+		}
 
 		int32 TileSize = GroupTileSize - 2 * PassParameters->TileOverscan;
 
 		FTSRRejectShadingCS::FPermutationDomain PermutationVector;
 		PermutationVector.Set<FTSRRejectShadingCS::FWaveSizeOps>(bUseWaveOps && GRHIMinimumWaveSize >= 32 && GRHIMinimumWaveSize <= 64 ? GRHIMinimumWaveSize : 0);
-		PermutationVector.Set<FTSRRejectShadingCS::FFlickeringDetectionDim>(FlickeringFramePeriod > 0.0f && HistoryMoireOutput);
+		PermutationVector.Set<FTSRRejectShadingCS::FFlickeringDetectionDim>(FlickeringFramePeriod > 0.0f);
+		PermutationVector.Set<FTSRRejectShadingCS::FHistoryResurrectionDim>(bCanResurrectHistory);
 		PermutationVector.Set<FTSRShader::F16BitVALUDim>(bUse16BitVALU);
 		PermutationVector = FTSRRejectShadingCS::RemapPermutation(PermutationVector);
 
 		TShaderMapRef<FTSRRejectShadingCS> ComputeShader(View.ShaderMap, PermutationVector);
 		FComputeShaderUtils::AddPass(
 			GraphBuilder,
-			RDG_EVENT_NAME("TSR RejectShading(%s TileSize=%d PaddingCostMultiplier=%1.1f WaveSize=%d FlickeringFramePeriod=%f %s%s) %dx%d",
-				UsageString,
+			RDG_EVENT_NAME("TSR RejectShading(TileSize=%d PaddingCostMultiplier=%1.1f WaveSize=%d FlickeringFramePeriod=%f VALU=%s%s) %dx%d",
 				TileSize,
 				FMath::Pow(float(GroupTileSize) / float(TileSize), 2),
 				int32(PermutationVector.Get<FTSRRejectShadingCS::FWaveSizeOps>()),
 				PassParameters->FlickeringFramePeriod,
-				PermutationVector.Get<FTSRShader::F16BitVALUDim>() ? TEXT(" 16bit") : TEXT(""),
-				TEXT(" ComposeTranslucency"),
+				PermutationVector.Get<FTSRShader::F16BitVALUDim>() ? TEXT("16bit") : TEXT("32bit"),
+				PermutationVector.Get<FTSRRejectShadingCS::FHistoryResurrectionDim>() ? TEXT(" Resurrection") : TEXT(""),
 				InputRect.Width(), InputRect.Height()),
 			AsyncComputePasses >= 3 ? ERDGPassFlags::AsyncCompute : ERDGPassFlags::Compute,
 			ComputeShader,
 			PassParameters,
 			FComputeShaderUtils::GetGroupCount(InputRect.Size(), TileSize));
+	}
 
-		return HistoryRejectionTexture;
-	}; // RejectReprojectedGuide();
-
-	// Rejects frame
-	FRDGTextureRef HistoryRejectionTexture = nullptr;
+	// Commit to resurrection if selected by the shading heuristic
+	if (bCanResurrectHistory)
 	{
-		FRDGTextureUAVRef HistoryGuideOutput;
-		FRDGTextureUAVRef HistoryMoireOutput = nullptr;
-
-		if (View.bStatePrevViewInfoIsReadOnly)
-		{
-			HistoryGuideOutput = CreateDummyUAV(GraphBuilder, History.GuideArray->Desc.Format);
-			HistoryMoireOutput = CreateDummyUAV(GraphBuilder, History.MoireArray->Desc.Format);
-		}
-		else
-		{
-			FRDGTextureUAVDesc GuideUAVDesc(History.GuideArray);
-			GuideUAVDesc.FirstArraySlice = CurrentFrameSliceIndex;
-			GuideUAVDesc.NumArraySlices = 1;
-			GuideUAVDesc.DimensionOverride = ETextureDimension::Texture2D;
-
-			FRDGTextureUAVDesc MoireUAVDesc(History.MoireArray);
-			MoireUAVDesc.FirstArraySlice = CurrentFrameSliceIndex;
-			MoireUAVDesc.NumArraySlices = 1;
-			MoireUAVDesc.DimensionOverride = ETextureDimension::Texture2D;
-
-			HistoryGuideOutput = GraphBuilder.CreateUAV(GuideUAVDesc);
-
-			if (FlickeringFramePeriod > 0.0f)
-			{
-				HistoryMoireOutput = GraphBuilder.CreateUAV(MoireUAVDesc);
-			}
-		}
-
-		if (bCanResurrectHistory)
-		{
-			FRDGTextureDesc GuideDesc = FRDGTextureDesc::Create2D(
-				History.GuideArray->Desc.Extent,
-				History.GuideArray->Desc.Format,
-				FClearValueBinding::None,
-				/* InFlags = */ TexCreate_ShaderResource | TexCreate_UAV);
-
-			FRDGTextureDesc MoireDesc = FRDGTextureDesc::Create2D(
-				History.MoireArray->Desc.Extent,
-				History.MoireArray->Desc.Format,
-				FClearValueBinding::None,
-				/* InFlags = */ TexCreate_ShaderResource | TexCreate_UAV);
-
-			FRDGTextureRef PrevGuideTexture = GraphBuilder.CreateTexture(GuideDesc, TEXT("TSR.Prev.Guide"));
-			FRDGTextureRef PrevMoireTexture = GraphBuilder.CreateTexture(GuideDesc, TEXT("TSR.Prev.Moire"));
-
-			FRDGTextureRef ResurrectedGuideTexture = GraphBuilder.CreateTexture(GuideDesc, TEXT("TSR.Resurrected.Guide"));
-			FRDGTextureRef ResurrectedMoireTexture = GraphBuilder.CreateTexture(GuideDesc, TEXT("TSR.Resurrected.Moire"));
-
-			FRDGTextureRef PrevRejectionTexture = RejectReprojectedGuide(
-				TEXT("PrevFrame"),
-				/* HistoryReprojectionIndex = */ 0,
-				/* out */ GraphBuilder.CreateUAV(PrevGuideTexture),
-				/* out */ bResurrectFlickeringHistory ? GraphBuilder.CreateUAV(PrevMoireTexture) : HistoryMoireOutput);
-
-			FRDGTextureRef ResurrectedRejectionTexture = RejectReprojectedGuide(
-				TEXT("Resurrection"),
-				/* HistoryReprojectionIndex = */ 1,
-				/* out */ GraphBuilder.CreateUAV(ResurrectedGuideTexture),
-				/* out */ bResurrectFlickeringHistory ? GraphBuilder.CreateUAV(ResurrectedMoireTexture) : nullptr);
-
-			{
-				HistoryRejectionTexture = GraphBuilder.CreateTexture(
-					PrevRejectionTexture->Desc, TEXT("TSR.HistoryRejection"));
-
-				FTSRMergeRejectionCS::FParameters* PassParameters = GraphBuilder.AllocParameters<FTSRMergeRejectionCS::FParameters>();
-				PassParameters->CommonParameters = CommonParameters;
-				PassParameters->ClipToResurrectionClip = ClipToResurrectionClip;
+		FTSRMergeRejectionCS::FParameters* PassParameters = GraphBuilder.AllocParameters<FTSRMergeRejectionCS::FParameters>();
+		PassParameters->CommonParameters = CommonParameters;
+		PassParameters->ClipToResurrectionClip = ClipToResurrectionClip;
 
-				PassParameters->ClosestDepthTexture = ClosestDepthTexture;
-				PassParameters->PrevGuideTexture     = PrevGuideTexture;
-				PassParameters->PrevMoireTexture     = bResurrectFlickeringHistory ? PrevMoireTexture : BlackDummy;
-				PassParameters->PrevRejectionTexture = PrevRejectionTexture;
+		PassParameters->ClosestDepthTexture = ClosestDepthTexture;
+		PassParameters->HistoryResurrectionMaskTexture = HistoryResurrectionMaskTexture;
 
-				PassParameters->ResurrectedGuideTexture     = ResurrectedGuideTexture;
-				PassParameters->ResurrectedMoireTexture     = bResurrectFlickeringHistory ? ResurrectedMoireTexture : BlackDummy;
-				PassParameters->ResurrectedRejectionTexture = ResurrectedRejectionTexture;
+		PassParameters->DilatedVelocityOutput = GraphBuilder.CreateUAV(DilatedVelocityTexture);
+		PassParameters->DebugOutput = CreateDebugUAV(InputExtent, TEXT("Debug.TSR.MergeRejection"));
 
-				PassParameters->HistoryGuideOutput = HistoryGuideOutput;
-				PassParameters->HistoryMoireOutput = bResurrectFlickeringHistory ? HistoryMoireOutput : CreateDummyUAV(GraphBuilder, History.MoireArray->Desc.Format);
-				PassParameters->HistoryRejectionOutput = GraphBuilder.CreateUAV(HistoryRejectionTexture);
-				PassParameters->HistoryResurrectionMaskOutput = GraphBuilder.CreateUAV(HistoryResurrectionMaskTexture);
-				PassParameters->DilatedVelocityOutput = GraphBuilder.CreateUAV(DilatedVelocityTexture);
-
-				PassParameters->DebugOutput = CreateDebugUAV(InputExtent, TEXT("Debug.TSR.MergeRejection"));
-
-				TShaderMapRef<FTSRMergeRejectionCS> ComputeShader(View.ShaderMap);
-				FComputeShaderUtils::AddPass(
-					GraphBuilder,
-					RDG_EVENT_NAME("TSR MergeRejection %dx%d", InputRect.Width(), InputRect.Height()),
-					AsyncComputePasses >= 3 ? ERDGPassFlags::AsyncCompute : ERDGPassFlags::Compute,
-					ComputeShader,
-					PassParameters,
-					FComputeShaderUtils::GetGroupCount(InputRect.Size(), 8));
-			}
-		}
-		else
-		{
-			HistoryRejectionTexture = RejectReprojectedGuide(
-				TEXT("PrevFrame"),
-				/* HistoryReprojectionIndex = */ 0,
-				/* out */ HistoryGuideOutput,
-				/* out */ HistoryMoireOutput);
-		}
+		TShaderMapRef<FTSRMergeRejectionCS> ComputeShader(View.ShaderMap);
+		FComputeShaderUtils::AddPass(
+			GraphBuilder,
+			RDG_EVENT_NAME("TSR MergeRejection %dx%d", InputRect.Width(), InputRect.Height()),
+			AsyncComputePasses >= 3 ? ERDGPassFlags::AsyncCompute : ERDGPassFlags::Compute,
+			ComputeShader,
+			PassParameters,
+			FComputeShaderUtils::GetGroupCount(InputRect.Size(), 8));
 	}
 
 	// Spatial anti-aliasing when doing history rejection.
