Add CVar to toggle DFAO in multiview.

diff --git a/Engine/Source/Runtime/Renderer/Private/DistanceFieldAmbientOcclusion.cpp b/Engine/Source/Runtime/Renderer/Private/DistanceFieldAmbientOcclusion.cpp
--- a/Engine/Source/Runtime/Renderer/Private/DistanceFieldAmbientOcclusion.cpp
+++ b/Engine/Source/Runtime/Renderer/Private/DistanceFieldAmbientOcclusion.cpp
@@ -33,6 +33,14 @@ FAutoConsoleVariableRef CVarDistanceFieldAO(
 	ECVF_Scalability | ECVF_RenderThreadSafe
 	);
 
+int32 GDistanceFieldAOMultiView = 1;
+FAutoConsoleVariableRef CVarDistanceFieldAOMultiView(
+	TEXT("r.DistanceFieldAO.MultiView"),
+	GDistanceFieldAOMultiView,
+	TEXT("Whether the distance field AO feature is allowed when rendering multiple views."),
+	ECVF_Scalability | ECVF_RenderThreadSafe
+);
+
 int32 GDistanceFieldAOQuality = 2;
 FAutoConsoleVariableRef CVarDistanceFieldAOQuality(
 	TEXT("r.AOQuality"),
@@ -764,6 +772,11 @@ void FDeferredShadingSceneRenderer::RenderDFAOAsIndirectShadowing(
 
 bool FDeferredShadingSceneRenderer::ShouldRenderDistanceFieldLighting() const
 {
+	if (!GDistanceFieldAOMultiView && Views.Num() > 1)
+	{
+		return false;
+	}
+
 	bool bSupportsDistanceFieldAO = true;
 
 	for (int32_t ViewIndex = 0; ViewIndex < Views.Num(); ++ViewIndex)
