Clean unused shader permutation of the BuildInstanceIdBufferAndCommandsFromPrimitiveIdsCs

diff --git a/Engine/Source/Runtime/Renderer/Private/InstanceCulling/InstanceCullingContext.cpp b/Engine/Source/Runtime/Renderer/Private/InstanceCulling/InstanceCullingContext.cpp
--- a/Engine/Source/Runtime/Renderer/Private/InstanceCulling/InstanceCullingContext.cpp
+++ b/Engine/Source/Runtime/Renderer/Private/InstanceCulling/InstanceCullingContext.cpp
@@ -51,6 +51,9 @@ static const TCHAR* BatchProcessingModeStr[] =
 
 static_assert(UE_ARRAY_COUNT(BatchProcessingModeStr) == uint32(EBatchProcessingMode::Num), "BatchProcessingModeStr length does not match EBatchProcessingMode::Num, these must be kept in sync.");
 
+DECLARE_GPU_STAT(BuildRenderingCommandsDeferred);
+DECLARE_GPU_STAT(BuildRenderingCommands);
+
 static bool IsInstanceOrderPreservationAllowed(EShaderPlatform ShaderPlatform)
 {
 	// Instance order preservation is currently not supported on mobile platforms
@@ -455,7 +458,6 @@ public:
 	static constexpr int32 NumThreadsPerGroup = FInstanceProcessingGPULoadBalancer::ThreadGroupSize;
 
 	// GPUCULL_TODO: remove once buffer is somehow unified
-	class FOutputCommandIdDim : SHADER_PERMUTATION_BOOL("OUTPUT_COMMAND_IDS");
 	class FSingleInstanceModeDim : SHADER_PERMUTATION_BOOL("SINGLE_INSTANCE_MODE");
 	class FCullInstancesDim : SHADER_PERMUTATION_BOOL("CULL_INSTANCES");
 	class FAllowWPODisableDim : SHADER_PERMUTATION_BOOL("ALLOW_WPO_DISABLE");
@@ -468,7 +470,7 @@ public:
 	class FBatchedDim : SHADER_PERMUTATION_BOOL("ENABLE_BATCH_MODE");
 	class FInstanceCompactionDim : SHADER_PERMUTATION_BOOL("ENABLE_INSTANCE_COMPACTION");
 
-	using FPermutationDomain = TShaderPermutationDomain<FOutputCommandIdDim, FSingleInstanceModeDim, FCullInstancesDim, FAllowWPODisableDim, FOcclusionCullInstancesDim, FStereoModeDim, FDebugModeDim, FBatchedDim, FInstanceCompactionDim>;
+	using FPermutationDomain = TShaderPermutationDomain<FSingleInstanceModeDim, FCullInstancesDim, FAllowWPODisableDim, FOcclusionCullInstancesDim, FStereoModeDim, FDebugModeDim, FBatchedDim, FInstanceCompactionDim>;
 
 	static bool ShouldCompilePermutation(const FGlobalShaderPermutationParameters& Parameters)
 	{
@@ -710,6 +712,7 @@ void FInstanceCullingContext::BuildRenderingCommandsInternal(
 	const bool bAllowWPODisable = InstanceCullingManager != nullptr;
 
 	RDG_EVENT_SCOPE(GraphBuilder, "BuildRenderingCommands(Culling=%s)", bCullInstances ? TEXT("On") : TEXT("Off"));
+	RDG_GPU_STAT_SCOPE(GraphBuilder, BuildRenderingCommands);
 
 	const bool bOrderPreservationEnabled = IsInstanceOrderPreservationEnabled();
 	const uint32 NumCompactionBlocks = uint32(CompactionBlockDataIndices.Num());
@@ -856,7 +859,6 @@ void FInstanceCullingContext::BuildRenderingCommandsInternal(
 			check(EBatchProcessingMode(Mode) != EBatchProcessingMode::UnCulled || LoadBalancer->HasSingleInstanceItemsOnly());
 
 			FBuildInstanceIdBufferAndCommandsFromPrimitiveIdsCs::FPermutationDomain PermutationVector;
-			PermutationVector.Set<FBuildInstanceIdBufferAndCommandsFromPrimitiveIdsCs::FOutputCommandIdDim>(0);
 			PermutationVector.Set<FBuildInstanceIdBufferAndCommandsFromPrimitiveIdsCs::FSingleInstanceModeDim>(EBatchProcessingMode(Mode) == EBatchProcessingMode::UnCulled);
 			PermutationVector.Set<FBuildInstanceIdBufferAndCommandsFromPrimitiveIdsCs::FCullInstancesDim>(bCullInstances && EBatchProcessingMode(Mode) != EBatchProcessingMode::UnCulled);
 			PermutationVector.Set<FBuildInstanceIdBufferAndCommandsFromPrimitiveIdsCs::FAllowWPODisableDim>(bAllowWPODisable);
@@ -1026,6 +1028,7 @@ FInstanceCullingDeferredContext *FInstanceCullingContext::CreateDeferredContext(
 	const bool bAllowWPODisable = true;
 
 	RDG_EVENT_SCOPE(GraphBuilder, "BuildRenderingCommandsDeferred(Culling=%s)", bCullInstances ? TEXT("On") : TEXT("Off"));
+	RDG_GPU_STAT_SCOPE(GraphBuilder, BuildRenderingCommandsDeferred);
 
 	TStaticArray<FBuildInstanceIdBufferAndCommandsFromPrimitiveIdsCs::FParameters*, static_cast<uint32>(EBatchProcessingMode::Num)> PassParameters;
 	for (uint32 Mode = 0U; Mode < uint32(EBatchProcessingMode::Num); ++Mode)
