Fix missing (removed) cvar in debug logging

diff --git a/Engine/Source/Runtime/Renderer/Private/VirtualShadowMaps/VirtualShadowMapCacheManager.cpp b/Engine/Source/Runtime/Renderer/Private/VirtualShadowMaps/VirtualShadowMapCacheManager.cpp
--- a/Engine/Source/Runtime/Renderer/Private/VirtualShadowMaps/VirtualShadowMapCacheManager.cpp
+++ b/Engine/Source/Runtime/Renderer/Private/VirtualShadowMaps/VirtualShadowMapCacheManager.cpp
@@ -536,17 +536,13 @@ FVirtualShadowMapArrayCacheManager::FVirtualShadowMapArrayCacheManager(FScene* I
 			static const auto* CVarMaxPhysicalPagesPtr = IConsoleManager::Get().FindTConsoleVariableDataInt(TEXT("r.Shadow.Virtual.MaxPhysicalPages"));
 			const int32 MaxPhysicalPages = CVarMaxPhysicalPagesPtr->GetValueOnRenderThread();
 
-			static const auto* CVarMaxPhysicalPagesSceneCapturePtr = IConsoleManager::Get().FindTConsoleVariableDataInt(TEXT("r.Shadow.Virtual.MaxPhysicalPagesSceneCapture"));
-			const int32 MaxPhysicalPagesSceneCapture = CVarMaxPhysicalPagesSceneCapturePtr->GetValueOnRenderThread();
-
 #if !UE_BUILD_SHIPPING
 			if (!bLoggedPageOverflow)
 			{
 				UE_LOG(LogRenderer, Warning, TEXT("Virtual Shadow Map Page Pool overflow (%d page allocations were not served), this will produce visual artifacts (missing shadow), increase the page pool limit or reduce resolution bias to avoid.\n")
-					TEXT(" See r.Shadow.Virtual.MaxPhysicalPages (%d), r.Shadow.Virtual.MaxPhysicalPagesSceneCapture (%d), r.Shadow.Virtual.ResolutionLodBiasLocal (%.2f), and r.Shadow.Virtual.ResolutionLodBiasDirectional (%.2f)"),
+					TEXT(" See r.Shadow.Virtual.MaxPhysicalPages (%d), r.Shadow.Virtual.ResolutionLodBiasLocal (%.2f), and r.Shadow.Virtual.ResolutionLodBiasDirectional (%.2f)"),
 					-NumPagesFree,
 					MaxPhysicalPages,
-					MaxPhysicalPagesSceneCapture,
 					LodBiasLocal,
 					LodBiasDirectional);
 				bLoggedPageOverflow = true;
