Implements VisualizeTSR show flag

diff --git a/Engine/Shaders/Private/TemporalSuperResolution/TSRVisualize.usf b/Engine/Shaders/Private/TemporalSuperResolution/TSRVisualize.usf
new file mode 100644
--- /dev/null
+++ b/Engine/Shaders/Private/TemporalSuperResolution/TSRVisualize.usf
@@ -0,0 +1,200 @@
+// Copyright Epic Games, Inc. All Rights Reserved.
+
+#include "TSRCommon.ush"
+
+
+//------------------------------------------------------- CONFIG
+
+#define TILE_SIZE 8
+
+
+//------------------------------------------------------- CONSTS
+
+static tsr_half3 kWhite  = tsr_half3(0.33, 0.33, 0.33);
+static tsr_half3 kRed    = tsr_half3(1.00, 0.10, 0.10);
+static tsr_half3 kOrange = tsr_half3(1.00, 0.75, 0.10);
+static tsr_half3 kGreen  = tsr_half3(0.10, 1.00, 0.10);
+static tsr_half3 kBlue   = tsr_half3(0.10, 0.10, 1.00);
+
+
+//------------------------------------------------------- PARAMETERS
+
+FScreenTransform OutputPixelPosToScreenPos;
+FScreenTransform ScreenPosToHistoryUV;
+FScreenTransform ScreenPosToInputPixelPos;
+float4x4 ClipToResurrectionClip;
+uint2 OutputViewRectMin;
+uint2 OutputViewRectMax;
+uint VisualizeId;
+uint bCanResurrectHistory;
+float MaxHistorySampleCount;
+float OutputToHistoryResolutionFractionSquare;
+
+Texture2D<tsr_half3> SceneColorTexture;
+Texture2D<float>     ClosestDepthTexture;
+Texture2D<float2>    DilatedVelocityTexture;
+Texture2D<tsr_half>  ParallaxRejectionMaskTexture;
+Texture2D<tsr_half2> HistoryRejectionTexture;
+Texture2D<tsr_half>  HistoryResurrectionMaskTexture;
+Texture2D<tsr_half4> HistoryMetadataTexture;
+Texture2D<tsr_half3> ResurrectedHistoryColorTexture;
+
+RWTexture2D<tsr_half4> Output;
+
+
+//------------------------------------------------------- ENTRY POINT
+
+tsr_half3 GrayScaleColor(tsr_half3 SceneColor)
+{
+	tsr_half3 C = SceneColor;
+	C.rgb = dot(tsr_half(0.3333).xxx, SceneColor.rgb);
+	return C;
+}
+
+struct FContext
+{
+	tsr_short2 OutputPixelPos;
+	float2 ScreenPos;
+	float2 InputUV;
+};
+
+tsr_half3 SampleTSROutputColor(FContext Ctx)
+{
+	return SceneColorTexture[Ctx.OutputPixelPos];
+}
+
+tsr_half3 SampleTSROutputGray(FContext Ctx)
+{
+	return GrayScaleColor(SceneColorTexture[Ctx.OutputPixelPos]);
+}
+
+tsr_half3 DisplayHistorySampleCount(FContext Ctx)
+{
+	float2 HistoryUV =  ApplyScreenTransform(Ctx.ScreenPos, ScreenPosToHistoryUV);
+	HistoryUV = fastClamp(HistoryUV, HistoryInfo_UVViewportBilinearMin, HistoryInfo_UVViewportBilinearMax);
+
+	tsr_half Validity = HistoryMetadataTexture.SampleLevel(GlobalBilinearClampedSampler, HistoryUV, 0).x;
+
+#if 0
+	tsr_half SampleCount = Validity * MaxHistorySampleCount;
+
+	const tsr_half Threshold = tsr_half(1.0) + OutputToHistoryResolutionFractionSquare;
+
+	if (SampleCount > Threshold)
+	{
+		return SampleTSROutputGray(Ctx) * lerp(kOrange, kGreen, saturate((SampleCount - Threshold) / (MaxHistorySampleCount - Threshold)));
+	}
+
+	return SampleTSROutputGray(Ctx) * lerp(kRed, kOrange, saturate((SampleCount - 1.0) / (Threshold - 1.0)));
+#else
+	return SampleTSROutputGray(Ctx) * lerp(kRed, kGreen, saturate(Validity));
+#endif
+}
+
+tsr_half3 DisplayHistoryRejection(FContext Ctx)
+{
+	tsr_half2 RawRejection = HistoryRejectionTexture.SampleLevel(GlobalPointClampedSampler, Ctx.InputUV, 0);
+	return SampleTSROutputGray(Ctx) * lerp(kRed, kWhite, saturate(RawRejection.x));
+}
+
+tsr_half3 DisplayHistoryClamp(FContext Ctx)
+{
+	tsr_half2 RawRejection = HistoryRejectionTexture.SampleLevel(GlobalPointClampedSampler, Ctx.InputUV, 0);
+	return SampleTSROutputGray(Ctx) * lerp(kRed, kWhite, saturate(RawRejection.y));
+}
+
+tsr_half3 DisplayParallaxDisocclusionMask(FContext Ctx)
+{
+	tsr_half ParallaxDisocclusionMask = ParallaxRejectionMaskTexture.SampleLevel(GlobalPointClampedSampler, Ctx.InputUV, 0);
+
+	return SampleTSROutputGray(Ctx) * lerp(kRed, kWhite, saturate(ParallaxDisocclusionMask));
+}
+
+tsr_half3 DisplayResurrectionMask(FContext Ctx)
+{
+	tsr_half ResurrectionMask = HistoryResurrectionMaskTexture.SampleLevel(GlobalPointClampedSampler, Ctx.InputUV, 0);
+	if (!bCanResurrectHistory)
+	{
+		return SampleTSROutputGray(Ctx) * kRed;
+	}
+	return SampleTSROutputGray(Ctx) * lerp(kWhite, kGreen, saturate(ResurrectionMask));
+}
+
+tsr_half3 DisplayResurrectedColor(FContext Ctx)
+{
+	float ClosestDeviceZ = ClosestDepthTexture.SampleLevel(GlobalPointClampedSampler, Ctx.InputUV, 0);
+	tsr_half ResurrectionMask = HistoryResurrectionMaskTexture.SampleLevel(GlobalPointClampedSampler, Ctx.InputUV, 0);
+
+	float4 ThisClip = float4(Ctx.ScreenPos, ClosestDeviceZ, 1);
+	float4 ResurrectionClip = mul(ThisClip, ClipToResurrectionClip); 
+	float2 ResurrectionScreen = ResurrectionClip.xy / ResurrectionClip.w;
+	
+	float2 KernelHistoryBufferUV = ApplyScreenTransform(ResurrectionScreen, ScreenPosToPrevHistoryBufferUV);
+	KernelHistoryBufferUV = fastClamp(KernelHistoryBufferUV, PrevHistoryInfo_UVViewportBilinearMin, PrevHistoryInfo_UVViewportBilinearMax);
+
+	tsr_half3 ResurrectedColor = ResurrectedHistoryColorTexture.SampleLevel(GlobalBilinearClampedSampler, KernelHistoryBufferUV, 0) * tsr_half(ResurrectionPreExposureCorrection);
+
+	return GrayScaleColor(ResurrectedColor) * lerp(kWhite, kGreen, saturate(ResurrectionMask));
+}
+
+
+//------------------------------------------------------- ENTRY POINT
+
+[numthreads(TILE_SIZE * TILE_SIZE, 1, 1)]
+void MainCS(
+	uint2 GroupId : SV_GroupID,
+	uint GroupThreadIndex : SV_GroupIndex)
+{
+	uint GroupWaveIndex = GetGroupWaveIndex(GroupThreadIndex, /* GroupSize = */ TILE_SIZE * TILE_SIZE);
+
+	float4 Debug = 0.0;
+	
+	FContext Ctx;
+	Ctx.OutputPixelPos = (
+		tsr_short2(OutputViewRectMin) +
+		tsr_short2(GroupId) * tsr_short2(TILE_SIZE, TILE_SIZE) +
+		Map8x8Tile2x2Lane(GroupThreadIndex));
+	Ctx.ScreenPos = ApplyScreenTransform(float2(Ctx.OutputPixelPos), OutputPixelPosToScreenPos);
+	Ctx.InputUV = ApplyScreenTransform(Ctx.ScreenPos, ScreenPosToInputPixelPos);
+	Ctx.InputUV = fastClamp(Ctx.InputUV, InputInfo_UVViewportBilinearMin, InputInfo_UVViewportBilinearMax);
+
+	tsr_half3 OutputColor = tsr_half(0.0);
+
+	BRANCH
+	if (VisualizeId == 0)
+	{
+		OutputColor = DisplayHistorySampleCount(Ctx);
+	}
+	else if (VisualizeId == 1)
+	{
+		OutputColor = DisplayParallaxDisocclusionMask(Ctx);
+	}
+	else if (VisualizeId == 2)
+	{
+		OutputColor = DisplayHistoryRejection(Ctx);
+	}
+	else if (VisualizeId == 3)
+	{
+		OutputColor = DisplayHistoryClamp(Ctx);
+	}
+	else if (VisualizeId == 4)
+	{
+		OutputColor = DisplayResurrectionMask(Ctx);
+	}
+	else if (VisualizeId == 5)
+	{
+		OutputColor = DisplayResurrectedColor(Ctx);
+	}
+
+	{
+		tsr_short2 OutputPixelPos = InvalidateOutputPixelPos(Ctx.OutputPixelPos, OutputViewRectMax);
+		
+		Output[OutputPixelPos] = tsr_half4(OutputColor, 0.0);
+		
+		#if DEBUG_OUTPUT
+		{
+			DebugOutput[tsr_short3(OutputPixelPos, 0)] = Debug;
+		}
+		#endif
+	}
+}
diff --git a/Engine/Source/Runtime/Engine/Public/ShowFlags.h b/Engine/Source/Runtime/Engine/Public/ShowFlags.h
--- a/Engine/Source/Runtime/Engine/Public/ShowFlags.h
+++ b/Engine/Source/Runtime/Engine/Public/ShowFlags.h
@@ -392,6 +392,7 @@ private:
 		SetVisualizeMotionVectors(false);
 		SetVisualizeReprojection(false);
 		SetVisualizeTemporalUpscaler(false);
+		SetVisualizeTSR(false);
 		SetEditingLevelInstance(false);
 		SetSelectionOutline(false);
 		SetDebugAI(false);
diff --git a/Engine/Source/Runtime/Engine/Public/ShowFlagsValues.inl b/Engine/Source/Runtime/Engine/Public/ShowFlagsValues.inl
--- a/Engine/Source/Runtime/Engine/Public/ShowFlagsValues.inl
+++ b/Engine/Source/Runtime/Engine/Public/ShowFlagsValues.inl
@@ -122,8 +122,10 @@ SHOWFLAG_FIXED_IN_SHIPPING(0, VisualizeMotionBlur, SFG_Visualize, NSLOCTEXT("Unr
 SHOWFLAG_FIXED_IN_SHIPPING(0, VisualizeMotionVectors, SFG_Hidden, NSLOCTEXT("UnrealEd", "VisualizeMotionVectorsSF", "Motion Vectors"))
 /** Helper to ensure previous frame's temporal history reprojections using motion vectors are working correctly. */
 SHOWFLAG_FIXED_IN_SHIPPING(0, VisualizeReprojection, SFG_Visualize, NSLOCTEXT("UnrealEd", "VisualizeReprojectionSF", "Previous frame's reprojection"))
-/** Helper to diagnose why the temporal upscaler may run into issues. */
-SHOWFLAG_FIXED_IN_SHIPPING(0, VisualizeTemporalUpscaler, SFG_Visualize, NSLOCTEXT("UnrealEd", "VisualizeTemporalUpscalerSF", "Temporal Upscaler (TSR, TAAU or third party plugins)"))
+/** Helper to diagnose temporal upscaler's inputs and outputs. */
+SHOWFLAG_FIXED_IN_SHIPPING(0, VisualizeTemporalUpscaler, SFG_Visualize, NSLOCTEXT("UnrealEd", "VisualizeTemporalUpscalerSF", "Temporal Upscaler I/O (TSR, TAAU or third party plugins)"))
+/** Helper to diagnose the internals of TSR. */
+SHOWFLAG_FIXED_IN_SHIPPING(0, VisualizeTSR, SFG_Visualize, NSLOCTEXT("UnrealEd", "VisualizeTSRSF", "Temporal Super Resolution"))
 /** Whether to display the Reflection Environment feature, which has local reflections from Reflection Capture actors, for now SHOWFLAG_ALWAYS_ACCESSIBLE because it's exposed in SceneCapture */
 SHOWFLAG_ALWAYS_ACCESSIBLE(ReflectionEnvironment, SFG_LightingFeatures, NSLOCTEXT("UnrealEd", "ReflectionEnvironmentSF", "Reflection Environment"))
 /** Visualize pixels that are outside of their object's bounding box (content error). */
diff --git a/Engine/Source/Runtime/Renderer/Private/PostProcess/PostProcessing.cpp b/Engine/Source/Runtime/Renderer/Private/PostProcess/PostProcessing.cpp
--- a/Engine/Source/Runtime/Renderer/Private/PostProcess/PostProcessing.cpp
+++ b/Engine/Source/Runtime/Renderer/Private/PostProcess/PostProcessing.cpp
@@ -557,11 +557,15 @@ void AddPostProcessingPasses(
 
 	if (bPostProcessingEnabled)
 	{
+		// Temporal Anti-aliasing. Also may perform a temporal upsample from primary to secondary view rect.
+		EMainTAAPassConfig TAAConfig = GetMainTAAPassConfig(View);
+
 		const bool bPrimaryView = IStereoRendering::IsAPrimaryView(View);
 		const bool bHasViewState = View.ViewState != nullptr;
 		const bool bDepthOfFieldEnabled = DiaphragmDOF::IsEnabled(View);
 		const bool bVisualizeDepthOfField = bDepthOfFieldEnabled && EngineShowFlags.VisualizeDOF;
 		const bool bVisualizeMotionBlur = IsVisualizeMotionBlurEnabled(View);
+		const bool bVisualizeTSR = TAAConfig == EMainTAAPassConfig::TSR && EngineShowFlags.VisualizeTSR;
 
 		const EAutoExposureMethod AutoExposureMethod = GetAutoExposureMethod(View);
 		const EAntiAliasingMethod AntiAliasingMethod = !bVisualizeDepthOfField ? View.AntiAliasingMethod : AAM_None;
@@ -573,7 +577,7 @@ void AddPostProcessingPasses(
 		const bool bMotionBlurValid = FMotionVectorSimulation::IsEnabled() || (!View.bCameraCut && !View.bPrevTransformsReset);
 
 		// Motion blur gets replaced by the visualization pass.
-		const bool bMotionBlurEnabled = !bVisualizeMotionBlur && IsMotionBlurEnabled(View) && bMotionBlurValid;
+		const bool bMotionBlurEnabled = !bVisualizeMotionBlur && IsMotionBlurEnabled(View) && bMotionBlurValid && !bVisualizeTSR;
 
 		// Skip tonemapping for visualizers which overwrite the HDR scene color.
 		const bool bTonemapEnabled = !bVisualizeMotionBlur;
@@ -600,11 +604,7 @@ void AddPostProcessingPasses(
 			!FMath::IsNearlyEqual(View.FinalPostProcessSettings.LocalExposureShadowContrastScale, 1.0f) ||
 			!FMath::IsNearlyEqual(View.FinalPostProcessSettings.LocalExposureDetailStrength, 1.0f);
 
-		const bool bBloomEnabled = View.FinalPostProcessSettings.BloomIntensity > 0.0f;
-
-
-		// Temporal Anti-aliasing. Also may perform a temporal upsample from primary to secondary view rect.
-		EMainTAAPassConfig TAAConfig = GetMainTAAPassConfig(View);
+		const bool bBloomEnabled = View.FinalPostProcessSettings.BloomIntensity > 0.0f && !bVisualizeTSR;
 
 		// Whether separate translucency is composed in TSR.
 		bool bComposeSeparateTranslucencyInTSR = TAAConfig == EMainTAAPassConfig::TSR && ComposeSeparateTranslucencyInTSR(View);
diff --git a/Engine/Source/Runtime/Renderer/Private/PostProcess/TemporalSuperResolution.cpp b/Engine/Source/Runtime/Renderer/Private/PostProcess/TemporalSuperResolution.cpp
--- a/Engine/Source/Runtime/Renderer/Private/PostProcess/TemporalSuperResolution.cpp
+++ b/Engine/Source/Runtime/Renderer/Private/PostProcess/TemporalSuperResolution.cpp
@@ -12,6 +12,7 @@
 #include "ScenePrivate.h"
 #include "RendererModule.h"
 #include "ShaderPlatformCachedIniValue.h"
+#include "PostProcess/PostProcessVisualizeBuffer.h"
 
 #define COMPILE_TSR_DEBUG_PASSES (!UE_BUILD_SHIPPING)
 
@@ -210,6 +211,7 @@ TAutoConsoleVariable<int32> CVarTSRResurrectionPersistentFrameInterval(
 	TEXT("r.TSR.Resurrection.PersistentFrameInterval"), 31,
 	TEXT("Configures in number of frames how often persistent frame should be recorded in history for futur history resurrection. ")
 	TEXT("This has no implication on memory footprint of the TSR history. Must be an even number greater or equal to 1. ")
+	TEXT("Uses the VisualizeTSR show flag and r.TSR.Visualize=5 to tune this parameter to your content. ")
 	TEXT("(default=31)"),
 	ECVF_RenderThreadSafe);
 
@@ -278,6 +280,22 @@ TAutoConsoleVariable<int32> CVarTSRSubpixelIncludeMovingDepth(
 	TEXT("because it is impossible how a moving object's velocity involves overtime when it's only occasionally drawing its velocity. (disabled by default)."),
 	ECVF_RenderThreadSafe);
 
+#if !UE_BUILD_OPTIMIZED_SHOWFLAGS
+
+TAutoConsoleVariable<int32> CVarTSRVisualize(
+	TEXT("r.TSR.Visualize"), -1,
+	TEXT("Selects what to display with the VisualizeTSR show flag (opened with the `show VisualizeTSR` command at runtime or Show > Visualize > TSR in editor viewports).\n")
+	TEXT(" -1: Display an overview grid (default);\n")
+	TEXT("  0: Number of accumulated samples in the history, particularily interesting to tune r.TSR.ShadingRejection.SampleCount and r.TSR.Velocity.WeightClampingSampleCount;\n")
+	TEXT("  1: Parallax disocclusion based of depth and velocity buffers;\n")
+	TEXT("  2: Mask where the history is rejected;\n")
+	TEXT("  3: Mask where the history is clamped;\n")
+	TEXT("  4: Mask where the history is resurrected (with r.TSR.Resurrection=1);\n")
+	TEXT("  5: Mask where the history is resurrected in the resurrected frame (with r.TSR.Resurrection=1), particularily interesting to tune r.TSR.Resurrection.PersistentFrameInterval;\n"),
+	ECVF_RenderThreadSafe);
+
+#endif
+
 #if COMPILE_TSR_DEBUG_PASSES
 
 TAutoConsoleVariable<int32> CVarTSRDebugArraySize(
@@ -930,6 +948,39 @@ class FTSRResolveHistoryCS : public FTSRShader
 	}
 }; // class FTSRResolveHistoryCS
 
+class FTSRVisualizeCS : public FTSRShader
+{
+	DECLARE_GLOBAL_SHADER(FTSRVisualizeCS);
+	SHADER_USE_PARAMETER_STRUCT(FTSRVisualizeCS, FTSRShader);
+
+	BEGIN_SHADER_PARAMETER_STRUCT(FParameters, )
+		SHADER_PARAMETER_STRUCT_INCLUDE(FTSRCommonParameters, CommonParameters)
+		SHADER_PARAMETER_STRUCT_INCLUDE(FTSRPrevHistoryParameters, PrevHistoryParameters)
+		SHADER_PARAMETER(FScreenTransform, OutputPixelPosToScreenPos)
+		SHADER_PARAMETER(FScreenTransform, ScreenPosToHistoryUV)
+		SHADER_PARAMETER(FScreenTransform, ScreenPosToInputPixelPos)
+		SHADER_PARAMETER(FMatrix44f, ClipToResurrectionClip)
+		SHADER_PARAMETER(FIntPoint, OutputViewRectMin)
+		SHADER_PARAMETER(FIntPoint, OutputViewRectMax)
+		SHADER_PARAMETER(int32, VisualizeId)
+		SHADER_PARAMETER(int32, bCanResurrectHistory)
+		SHADER_PARAMETER(float, MaxHistorySampleCount)
+		SHADER_PARAMETER(float, OutputToHistoryResolutionFractionSquare)
+
+		SHADER_PARAMETER_RDG_TEXTURE_SRV(Texture2D, SceneColorTexture)
+		SHADER_PARAMETER_RDG_TEXTURE(Texture2D, ClosestDepthTexture)
+		SHADER_PARAMETER_RDG_TEXTURE(Texture2D, DilatedVelocityTexture)
+		SHADER_PARAMETER_RDG_TEXTURE(Texture2D, ParallaxRejectionMaskTexture)
+		SHADER_PARAMETER_RDG_TEXTURE(Texture2D, HistoryRejectionTexture)
+		SHADER_PARAMETER_RDG_TEXTURE(Texture2D, HistoryResurrectionMaskTexture)
+		SHADER_PARAMETER_RDG_TEXTURE_SRV(Texture2D, HistoryMetadataTexture)
+		SHADER_PARAMETER_RDG_TEXTURE_SRV(Texture2D, ResurrectedHistoryColorTexture)
+
+		SHADER_PARAMETER_RDG_TEXTURE_UAV(RWTexture2D, Output)
+		SHADER_PARAMETER_RDG_TEXTURE_UAV(RWTexture2DArray, DebugOutput)
+	END_SHADER_PARAMETER_STRUCT()
+}; // class FTSRVisualizeCS
+
 IMPLEMENT_GLOBAL_SHADER(FTSRComputeMoireLumaCS,      "/Engine/Private/TemporalSuperResolution/TSRComputeMoireLuma.usf",      "MainCS", SF_Compute);
 IMPLEMENT_GLOBAL_SHADER(FTSRClearPrevTexturesCS,     "/Engine/Private/TemporalSuperResolution/TSRClearPrevTextures.usf",     "MainCS", SF_Compute);
 IMPLEMENT_GLOBAL_SHADER(FTSRForwardScatterDepthCS,   "/Engine/Private/TemporalSuperResolution/TSRForwardScatterDepth.usf",   "MainCS", SF_Compute);
@@ -941,6 +992,7 @@ IMPLEMENT_GLOBAL_SHADER(FTSRSpatialAntiAliasingCS,   "/Engine/Private/TemporalSu
 IMPLEMENT_GLOBAL_SHADER(FTSRFilterAntiAliasingCS,    "/Engine/Private/TemporalSuperResolution/TSRFilterAntiAliasing.usf",    "MainCS", SF_Compute);
 IMPLEMENT_GLOBAL_SHADER(FTSRUpdateHistoryCS,         "/Engine/Private/TemporalSuperResolution/TSRUpdateHistory.usf",         "MainCS", SF_Compute);
 IMPLEMENT_GLOBAL_SHADER(FTSRResolveHistoryCS,        "/Engine/Private/TemporalSuperResolution/TSRResolveHistory.usf",        "MainCS", SF_Compute);
+IMPLEMENT_GLOBAL_SHADER(FTSRVisualizeCS,             "/Engine/Private/TemporalSuperResolution/TSRVisualize.usf",             "MainCS", SF_Compute);
 
 DECLARE_GPU_STAT(TemporalSuperResolution)
 
@@ -2357,7 +2409,7 @@ FDefaultTemporalUpscaler::FOutputs AddTemporalSuperResolutionPasses(
 		{
 			PassParameters->SceneColorOutputMip1 = CreateDummyUAV(GraphBuilder, PF_FloatR11G11B10);
 		}
-		PassParameters->DebugOutput = CreateDebugUAV(HistoryExtent, TEXT("Debug.TSR.ResolveHistory"));
+		PassParameters->DebugOutput = CreateDebugUAV(OutputExtent, TEXT("Debug.TSR.ResolveHistory"));
 
 		FTSRResolveHistoryCS::FPermutationDomain PermutationVector;
 		PermutationVector.Set<FTSRResolveHistoryCS::FNyquistDim>(bNyquistHistory && bUseWaveOps && GRHIMaximumWaveSize >= 32 && GRHIMinimumWaveSize <= 32);
@@ -2435,6 +2487,139 @@ FDefaultTemporalUpscaler::FOutputs AddTemporalSuperResolutionPasses(
 		}
 	}
 
+#if !UE_BUILD_OPTIMIZED_SHOWFLAGS
+	if (View.Family->EngineShowFlags.VisualizeTSR)
+	{
+		RDG_EVENT_SCOPE(GraphBuilder, "VisualizeTSR %dx%d", OutputRect.Width(), OutputRect.Height());
+
+		enum class EVisualizeId : int32
+		{
+			Overview = -1,
+			HistorySampleCount = 0,
+			ParallaxDisocclusionMask = 1,
+			HistoryRejection = 2,
+			HistoryClamp = 3,
+			ResurrectionMask = 4,
+			ResurrectedColor = 5,
+			MAX,
+		};
+
+		static const TCHAR* kVisualizationName[] = {
+			TEXT("HistorySampleCount"),
+			TEXT("ParallaxDisocclusionMask"),
+			TEXT("HistoryRejection"),
+			TEXT("HistoryClamp"),
+			TEXT("ResurrectionMask"),
+			TEXT("ResurrectedColor"),
+		};
+		static_assert(UE_ARRAY_COUNT(kVisualizationName) == int32(EVisualizeId::MAX), "kVisualizationName doesn't match EVisualizeId");
+
+		const EVisualizeId Visualization = EVisualizeId(FMath::Clamp(CVarTSRVisualize.GetValueOnRenderThread(), -1, int32(EVisualizeId::MAX) - 1));
+		FIntRect VisualizeRect = Visualization == EVisualizeId::Overview ? FIntRect(OutputRect.Min + OutputRect.Size() / 4, OutputRect.Min + (OutputRect.Size() * 3) / 4) : OutputRect;
+
+		auto Visualize = [&](EVisualizeId VisualizeId, FString Label)
+		{
+			check(VisualizeId != EVisualizeId::Overview);
+
+			FRDGTextureDesc OutputDesc = FRDGTextureDesc::Create2D(
+				OutputExtent,
+				ColorFormat,
+				FClearValueBinding::None,
+				/* InFlags = */ TexCreate_ShaderResource | TexCreate_UAV);
+
+			FRDGTextureRef OutputTexture = GraphBuilder.CreateTexture(OutputDesc, TEXT("TSR.Visualize"));
+
+			FTSRVisualizeCS::FParameters* PassParameters = GraphBuilder.AllocParameters<FTSRVisualizeCS::FParameters>();
+			PassParameters->CommonParameters = CommonParameters;
+			PassParameters->PrevHistoryParameters = PrevHistoryParameters;
+			PassParameters->OutputPixelPosToScreenPos = (FScreenTransform::Identity - OutputRect.Min + 0.5f) / OutputRect.Size() * FScreenTransform::ViewportUVToScreenPos;
+			PassParameters->ScreenPosToHistoryUV = FScreenTransform::ChangeTextureBasisFromTo(HistoryExtent, FIntRect(FIntPoint::ZeroValue, HistorySize), FScreenTransform::ETextureBasis::ScreenPosition, FScreenTransform::ETextureBasis::TextureUV);
+			PassParameters->ScreenPosToInputPixelPos = FScreenTransform::ChangeTextureBasisFromTo(InputExtent, InputRect, FScreenTransform::ETextureBasis::ScreenPosition, FScreenTransform::ETextureBasis::TextureUV);
+			PassParameters->ClipToResurrectionClip = ClipToResurrectionClip;
+			PassParameters->OutputViewRectMin = VisualizeRect.Min;
+			PassParameters->OutputViewRectMax = VisualizeRect.Max;
+			PassParameters->VisualizeId = int32(VisualizeId);
+			PassParameters->bCanResurrectHistory = bCanResurrectHistory;
+			PassParameters->MaxHistorySampleCount = MaxHistorySampleCount;
+			PassParameters->OutputToHistoryResolutionFractionSquare = OutputToHistoryResolutionFractionSquare;
+
+			PassParameters->SceneColorTexture = SceneColorOutputTextureSRV;
+			PassParameters->ClosestDepthTexture = ClosestDepthTexture;
+			PassParameters->DilatedVelocityTexture = DilatedVelocityTexture;
+			PassParameters->ParallaxRejectionMaskTexture = ParallaxRejectionMaskTexture;
+			PassParameters->HistoryRejectionTexture = HistoryRejectionTexture;
+			PassParameters->HistoryResurrectionMaskTexture = bCanResurrectHistory ? HistoryResurrectionMaskTexture : BlackDummy;
+			PassParameters->HistoryMetadataTexture = GraphBuilder.CreateSRV(FRDGTextureSRVDesc::CreateForSlice(History.MetadataArray, CurrentFrameSliceIndex));
+			if (PrevHistory.ColorArray == BlackArrayDummy)
+			{
+				PassParameters->ResurrectedHistoryColorTexture = GraphBuilder.CreateSRV(FRDGTextureSRVDesc::CreateForSlice(PrevHistory.ColorArray, 0));
+			}
+			else
+			{
+				PassParameters->ResurrectedHistoryColorTexture = GraphBuilder.CreateSRV(FRDGTextureSRVDesc::CreateForSlice(PrevHistory.ColorArray, bCanResurrectHistory ? ResurrectionFrameSliceIndex : PrevFrameSliceIndex));
+			}
+
+			PassParameters->Output = GraphBuilder.CreateUAV(OutputTexture);
+			PassParameters->DebugOutput = CreateDebugUAV(OutputExtent, TEXT("Debug.TSR.Visualize"));
+
+			TShaderMapRef<FTSRVisualizeCS> ComputeShader(View.ShaderMap);
+			FComputeShaderUtils::AddPass(
+				GraphBuilder,
+				RDG_EVENT_NAME("TSR Visualize(%s) %dx%d", kVisualizationName[int32(VisualizeId)], VisualizeRect.Width(), VisualizeRect.Height()),
+				AsyncComputePasses >= 3 ? ERDGPassFlags::AsyncCompute : ERDGPassFlags::Compute,
+				ComputeShader,
+				PassParameters,
+				FComputeShaderUtils::GetGroupCount(VisualizeRect.Size(), 8));
+
+			FVisualizeBufferTile Tile;
+			Tile.Input = FScreenPassTexture(OutputTexture, VisualizeRect);
+			Tile.Label = FString::Printf(TEXT("%s (r.TSR.Visualize=%d)"), *Label, int32(VisualizeId));
+			return Tile;
+		};
+
+		FRDGTextureRef OutputTexture;
+		if (Visualization == EVisualizeId::Overview)
+		{
+			TArray<FVisualizeBufferTile> Tiles;
+			Tiles.SetNum(16);
+			{
+				Tiles[4 * 0 + 0] = Visualize(EVisualizeId::HistorySampleCount, TEXT("Accumulated Sample Count"));
+				Tiles[4 * 0 + 1] = Visualize(EVisualizeId::ParallaxDisocclusionMask, TEXT("Parallax Disocclusion"));
+				Tiles[4 * 0 + 2] = Visualize(EVisualizeId::HistoryRejection, TEXT("History Rejection"));
+				Tiles[4 * 0 + 3] = Visualize(EVisualizeId::HistoryClamp, TEXT("History Clamp"));
+				Tiles[4 * 1 + 0] = Visualize(EVisualizeId::ResurrectionMask, TEXT("Resurrection Mask"));
+				if (bCanResurrectHistory)
+				{
+					Tiles[4 * 2 + 0] = Visualize(EVisualizeId::ResurrectedColor, TEXT("Resurrected Frame"));
+				}
+			}
+
+			{
+				FRDGTextureDesc OutputDesc = FRDGTextureDesc::Create2D(
+					OutputExtent,
+					ColorFormat,
+					FClearValueBinding::Black,
+					/* InFlags = */ TexCreate_ShaderResource | TexCreate_RenderTargetable);
+
+				OutputTexture = GraphBuilder.CreateTexture(OutputDesc, TEXT("TSR.VisualizeOverview"));
+
+				FVisualizeBufferInputs VisualizeBufferInputs;
+				VisualizeBufferInputs.OverrideOutput = FScreenPassRenderTarget(FScreenPassTexture(OutputTexture, OutputRect), ERenderTargetLoadAction::EClear);
+				VisualizeBufferInputs.SceneColor = FScreenPassTexture::CopyFromSlice(GraphBuilder, FScreenPassTextureSlice(SceneColorOutputTextureSRV, OutputRect));
+				VisualizeBufferInputs.Tiles = Tiles;
+				AddVisualizeBufferPass(GraphBuilder, View, VisualizeBufferInputs);
+			}
+		}
+		else
+		{
+			OutputTexture = Visualize(Visualization, TEXT("")).Input.Texture;
+		}
+
+		FDefaultTemporalUpscaler::FOutputs Outputs;
+		Outputs.FullRes = FScreenPassTextureSlice(GraphBuilder.CreateSRV(FRDGTextureSRVDesc(OutputTexture)), OutputRect);
+		return Outputs;
+	}
+#endif
 
 	FDefaultTemporalUpscaler::FOutputs Outputs;
 	Outputs.FullRes = FScreenPassTextureSlice(SceneColorOutputTextureSRV, OutputRect);
