Fixes a bug where disabling r.TSR.Resurrection would end up reprojecting the wrong frame

diff --git a/Engine/Source/Runtime/Renderer/Private/PostProcess/TemporalSuperResolution.cpp b/Engine/Source/Runtime/Renderer/Private/PostProcess/TemporalSuperResolution.cpp
--- a/Engine/Source/Runtime/Renderer/Private/PostProcess/TemporalSuperResolution.cpp
+++ b/Engine/Source/Runtime/Renderer/Private/PostProcess/TemporalSuperResolution.cpp
@@ -1537,7 +1537,7 @@ FDefaultTemporalUpscaler::FOutputs AddTemporalSuperResolutionPasses(
 
 		int32 ResurrectionFrameRollingIndex = 0;
 		int32 PrevFrameRollingIndex = 0;
-		if (HistorySliceSequence.GetRollingIndexCount() == 1)
+		if (PrevHistorySliceSequence.GetRollingIndexCount() == 1)
 		{
 			// NOP
 		}
