Fixed: Decal material is drawn twice when using mesh decals on mobile

diff --git a/Engine/Source/Runtime/Renderer/Private/MobileBasePass.cpp b/Engine/Source/Runtime/Renderer/Private/MobileBasePass.cpp
--- a/Engine/Source/Runtime/Renderer/Private/MobileBasePass.cpp
+++ b/Engine/Source/Runtime/Renderer/Private/MobileBasePass.cpp
@@ -788,7 +788,7 @@ bool FMobileBasePassMeshProcessor::ShouldDraw(const FMaterial& Material) const
 	if (bTranslucentBasePass)
 	{
 		// Skipping TPT_TranslucencyAfterDOFModulate. That pass is only needed for Dual Blending, which is not supported on Mobile.
-		bool bShouldDraw = bIsTranslucent &&
+		bool bShouldDraw = bIsTranslucent && !Material.IsDeferredDecal() &&
 			(TranslucencyPassType == ETranslucencyPass::TPT_AllTranslucency
 				|| (TranslucencyPassType == ETranslucencyPass::TPT_TranslucencyStandard && !Material.IsMobileSeparateTranslucencyEnabled())
 				|| (TranslucencyPassType == ETranslucencyPass::TPT_TranslucencyAfterDOF && Material.IsMobileSeparateTranslucencyEnabled()));
