Fixed case where LightFunction material use primitive data which was crashing vulkan on a missing UB.
Bind an Identity primitive data to LightFunctions on mobile
Add additional validation that all descriptors are written

diff --git a/Engine/Source/Runtime/Renderer/Private/MobileDeferredShadingPass.cpp b/Engine/Source/Runtime/Renderer/Private/MobileDeferredShadingPass.cpp
--- a/Engine/Source/Runtime/Renderer/Private/MobileDeferredShadingPass.cpp
+++ b/Engine/Source/Runtime/Renderer/Private/MobileDeferredShadingPass.cpp
@@ -154,6 +154,10 @@ class FMobileDirectionalLightFunctionPS : public FMaterialShader
 	void SetParameters(FRHIBatchedShaderParameters& BatchedParameters, const FViewInfo& View, const FMaterialRenderProxy* Proxy, const FMaterial& Material)
 	{
 		FMaterialShader::SetParameters(BatchedParameters, Proxy, Material, View);
+
+		// LightFunctions can use primitive data, set identity so we do not crash on a missing binding
+		auto& PrimitivePS = GetUniformBufferParameter<FPrimitiveUniformShaderParameters>();
+		SetUniformBufferParameter(BatchedParameters, PrimitivePS, GIdentityPrimitiveUniformBuffer);
 	}
 };
 
@@ -233,6 +237,10 @@ public:
 	{
 		FMaterialShader::SetViewParameters(BatchedParameters, View, View.ViewUniformBuffer);
 		FMaterialShader::SetParameters(BatchedParameters, Proxy, Material, View);
+		
+		// LightFunctions can use primitive data, set identity so we do not crash on a missing binding
+		auto& PrimitivePS = GetUniformBufferParameter<FPrimitiveUniformShaderParameters>();
+		SetUniformBufferParameter(BatchedParameters, PrimitivePS, GIdentityPrimitiveUniformBuffer);
 	}
 };
 
diff --git a/Engine/Source/Runtime/VulkanRHI/Private/VulkanDescriptorSets.h b/Engine/Source/Runtime/VulkanRHI/Private/VulkanDescriptorSets.h
--- a/Engine/Source/Runtime/VulkanRHI/Private/VulkanDescriptorSets.h
+++ b/Engine/Source/Runtime/VulkanRHI/Private/VulkanDescriptorSets.h
@@ -1176,6 +1176,7 @@ protected:
 	friend class FVulkanCommonPipelineDescriptorState;
 	friend class FVulkanComputePipelineDescriptorState;
 	friend class FVulkanGraphicsPipelineDescriptorState;
+	friend class FVulkanDescriptorSetCache;
 
 #if VULKAN_VALIDATE_DESCRIPTORS_WRITTEN
 	TArray<uint32, TInlineAllocator<2> > WrittenMask;
diff --git a/Engine/Source/Runtime/VulkanRHI/Private/VulkanPendingState.cpp b/Engine/Source/Runtime/VulkanRHI/Private/VulkanPendingState.cpp
--- a/Engine/Source/Runtime/VulkanRHI/Private/VulkanPendingState.cpp
+++ b/Engine/Source/Runtime/VulkanRHI/Private/VulkanPendingState.cpp
@@ -820,6 +820,7 @@ bool FVulkanDescriptorSetCache::FCachedPool::CreateDescriptorSets(
 		SetCache.Emplace(DSWriter.GetKey().CopyDeep(), NewSetEntry.Sets[Index]);
 
 		DSWriter.SetDescriptorSet(NewSetEntry.Sets[Index]);
+		DSWriter.CheckAllWritten();
 
 		for (int32 i = 0; i < VK_DESCRIPTOR_TYPE_RANGE_SIZE; ++i)
 		{
