Fix FOcclusionRandomStream OOB

#rnx

diff --git a/Engine/Source/Runtime/Renderer/Private/ScenePrivate.h b/Engine/Source/Runtime/Renderer/Private/ScenePrivate.h
--- a/Engine/Source/Runtime/Renderer/Private/ScenePrivate.h
+++ b/Engine/Source/Runtime/Renderer/Private/ScenePrivate.h
@@ -530,11 +530,13 @@ public:
 	/** @return A random number between 0 and 1. */
 	inline float GetFraction()
 	{
-		if (CurrentSample >= NumSamples)
+		if (CurrentSample > NumSamples)
 		{
 			CurrentSample = 0;
 		}
-		return Samples[CurrentSample++];
+		float Fraction = Samples[CurrentSample];
+		CurrentSample++;
+		return Fraction;
 	}
 private:
 
