Fix assert firing when VT system is updated from offscreen renderer passes.

diff --git a/Engine/Source/Runtime/Renderer/Private/VT/VirtualTextureSystem.cpp b/Engine/Source/Runtime/Renderer/Private/VT/VirtualTextureSystem.cpp
--- a/Engine/Source/Runtime/Renderer/Private/VT/VirtualTextureSystem.cpp
+++ b/Engine/Source/Runtime/Renderer/Private/VT/VirtualTextureSystem.cpp
@@ -2821,6 +2821,8 @@ void FVirtualTextureSystem::EndUpdate(FRDGBuilder& GraphBuilder, TUniquePtr<FVir
 
 void FVirtualTextureSystem::Update(FRDGBuilder& GraphBuilder, ERHIFeatureLevel::Type FeatureLevel, FScene* Scene, const FVirtualTextureUpdateSettings& InSettings)
 {
+	CallPendingCallbacks();
+
 	FVirtualTextureUpdateSettings Settings = InSettings;
 	Settings.EnableAsyncTasks(false);
 
