Fixed FDiffuseIndirectCompositePS skipping pixels which need GI due to depth bounds test when SSAO is enabled, causing black indirect lighting past the AmbientOcclusionFadeDistance
 * Lumen GI works with SSAO with: r.Lumen.DiffuseIndirect.SSAO 1 and r.Lumen.ScreenProbeGather.ShortRangeAO 0

diff --git a/Engine/Source/Runtime/Renderer/Private/IndirectLightRendering.cpp b/Engine/Source/Runtime/Renderer/Private/IndirectLightRendering.cpp
--- a/Engine/Source/Runtime/Renderer/Private/IndirectLightRendering.cpp
+++ b/Engine/Source/Runtime/Renderer/Private/IndirectLightRendering.cpp
@@ -1408,7 +1408,9 @@ void FDeferredShadingSceneRenderer::RenderDiffuseIndirectAndAmbientOcclusion(
 				// only use depth bound optimization when diffuse indirect is disable
 				bool bUseDepthBounds = CVarDiffuseIndirectOffUseDepthBoundsAO.GetValueOnRenderThread() &&
 										bool(ERHIZBuffer::IsInverted) && // Inverted depth buffer is assumed when setting depth bounds test for AO.
-										!PermutationVector.Get<FDiffuseIndirectCompositePS::FUpscaleDiffuseIndirectDim>() && AmbientOcclusionMask;
+										!PermutationVector.Get<FDiffuseIndirectCompositePS::FUpscaleDiffuseIndirectDim>() && 
+										AmbientOcclusionMask &&
+										ViewPipelineState.DiffuseIndirectMethod == EDiffuseIndirectMethod::Disabled;
 
 				if (bUseDepthBounds)
 				{
