Optimises texture fetches in TSR UpdateHistory

diff --git a/Engine/Shaders/Private/TemporalSuperResolution/TSRCommon.ush b/Engine/Shaders/Private/TemporalSuperResolution/TSRCommon.ush
--- a/Engine/Shaders/Private/TemporalSuperResolution/TSRCommon.ush
+++ b/Engine/Shaders/Private/TemporalSuperResolution/TSRCommon.ush
@@ -849,11 +849,10 @@ bool IsOffScreen(uint bCameraCut, float2 ScreenPos, tsr_half ParallaxRejectionMa
 }
 
 CALL_SITE_DEBUGLOC
-void IsOffScreenOrDisoccluded(uint bCameraCut, float2x2 ScreenPos, tsr_half2 ParallaxRejectionMask, out bool2 bIsOffScreen, out bool2 bIsDisoccluded)
+void IsOffScreenOrDisoccluded(uint bCameraCut, float2x2 ScreenPos, bool2 bIsParallaxRejected, out bool2 bIsOffScreen, out bool2 bIsDisoccluded)
 {
 	bool bIsCameraCut = bCameraCut != 0;
 	bool2 bIsOutOfBounds = max(abs(ScreenPos[0]), abs(ScreenPos[1])) >= 1.0;
-	bool2 bIsParallaxRejected = ParallaxRejectionMask < tsr_half(PARALLAX_REJECTION_MASK_THRESHOLD);
 
 	bIsOffScreen = or(bIsCameraCut, bIsOutOfBounds);
 	bIsDisoccluded = and(!bIsOffScreen, bIsParallaxRejected);
diff --git a/Engine/Shaders/Private/TemporalSuperResolution/TSRDecimateHistory.usf b/Engine/Shaders/Private/TemporalSuperResolution/TSRDecimateHistory.usf
--- a/Engine/Shaders/Private/TemporalSuperResolution/TSRDecimateHistory.usf
+++ b/Engine/Shaders/Private/TemporalSuperResolution/TSRDecimateHistory.usf
@@ -67,8 +67,7 @@ RWTexture2DArray<tsr_half4> ReprojectedHistoryGuideOutput;
 RWTexture2DArray<tsr_half4> ReprojectedHistoryMoireOutput;
 
 RWTexture2D<float2> HoleFilledVelocityOutput;
-RWTexture2D<tsr_half2> HoleFilledVelocityMaskOutput;
-RWTexture2D<tsr_half> ParallaxRejectionMaskOutput;
+RWTexture2D<tsr_half2> DecimateMaskOutput;
 
 
 //------------------------------------------------------- DEBUG
@@ -128,7 +127,7 @@ void LoadVelocityNeighborhood(
 void ProcessVelocityNeighborhood(
 	float2 EncodedVelocityNeighborhood[CONFIG_VELOCITY_PRE_FETCH_SAMPLES],
 	out float2 EncodedScreenVelocity,
-	out tsr_half RejectHighFrequency)
+	out tsr_half VelocityEdge)
 {
 	// Preprocessed the prefetched samples
 	float2 PixelVelocityNeighborhood[CONFIG_VELOCITY_PRE_FETCH_SAMPLES];
@@ -140,7 +139,7 @@ void ProcessVelocityNeighborhood(
 		}
 	}
 	
-	RejectHighFrequency = tsr_half(0.0);
+	tsr_half RejectHighFrequency = tsr_half(0.0);
 
 	// Sample
 	float2 ScreenPixelVelocity = PixelVelocityNeighborhood[0];
@@ -223,10 +222,7 @@ void ProcessVelocityNeighborhood(
 	}
 
 	// Remove noise from high frequency rejections.
-	RejectHighFrequency = saturate(tsr_half(1.0) - RejectHighFrequency * tsr_half(1.1));
-
-	// Do not reject high frequency additionally where hole filling happen
-	// RejectHighFrequency = max(RejectHighFrequency, HoleFilledVelocityMask); TODO
+	VelocityEdge = saturate(RejectHighFrequency * tsr_half(1.1));
 
 	EncodedScreenVelocity = EncodedVelocityNeighborhood[0];
 } // ProcessVelocityNeighborhood()
@@ -387,7 +383,7 @@ void MainCS(
 	// Sample current frame informations.
 	float2 EncodedScreenVelocity;
 	float DeviceZ;
-	tsr_half RejectHighFrequency = tsr_half(0.0);
+	tsr_half VelocityEdge = tsr_half(1.0);
 	ISOLATE
 	{
 		tsr_ushort2 InputPixelPos = (
@@ -431,7 +427,7 @@ void MainCS(
 			ProcessVelocityNeighborhood(
 				EncodedVelocityNeighborhood,
 				/* out */ EncodedScreenVelocity,
-				/* out */ RejectHighFrequency);
+				/* out */ VelocityEdge);
 		}
 		#endif // CONFIG_VELOCITY_SAMPLES != 1
 	}
@@ -617,8 +613,7 @@ void MainCS(
 		#endif
 
 		HoleFilledVelocityOutput[OutputPixelPos] = EncodedScreenVelocity;
-		HoleFilledVelocityMaskOutput[InputPixelPos] = tsr_half2(0.0, RejectHighFrequency);
-		ParallaxRejectionMaskOutput[OutputPixelPos] = ParallaxRejectionMask;
+		DecimateMaskOutput[InputPixelPos] = tsr_half2(ParallaxRejectionMask, VelocityEdge);
 
 		#if DEBUG_OUTPUT
 		for (uint DebugId = 0; DebugId < DEBUG_ARRAY_SIZE; DebugId++)
diff --git a/Engine/Shaders/Private/TemporalSuperResolution/TSRMergeRejection.usf b/Engine/Shaders/Private/TemporalSuperResolution/TSRMergeRejection.usf
--- a/Engine/Shaders/Private/TemporalSuperResolution/TSRMergeRejection.usf
+++ b/Engine/Shaders/Private/TemporalSuperResolution/TSRMergeRejection.usf
@@ -14,7 +14,7 @@
 float4x4 ClipToResurrectionClip;
 
 Texture2D<float> ClosestDepthTexture;
-Texture2D<tsr_half> HistoryResurrectionMaskTexture;
+Texture2D<tsr_half4> HistoryRejectionTexture;
 
 RWTexture2D<float2> DilatedVelocityOutput;
 
@@ -35,9 +35,9 @@ void MainCS(
 		tsr_short2(GroupId) * tsr_short2(TILE_SIZE, TILE_SIZE) +
 		Map8x8Tile2x2Lane(GroupThreadIndex));
 	
-	tsr_half ResurrectionMask = HistoryResurrectionMaskTexture[InputPixelPos];
+	tsr_half ResurrectionMask = HistoryRejectionTexture[InputPixelPos].a;
 
-	bool bIsResurrectionBetter = ResurrectionMask > tsr_half(0.5);
+	bool bIsResurrectionBetter = ResurrectionMask > tsr_half(1.5 / 255.0);
 
 	BRANCH
 	if (bIsResurrectionBetter)
diff --git a/Engine/Shaders/Private/TemporalSuperResolution/TSRRejectShading.usf b/Engine/Shaders/Private/TemporalSuperResolution/TSRRejectShading.usf
--- a/Engine/Shaders/Private/TemporalSuperResolution/TSRRejectShading.usf
+++ b/Engine/Shaders/Private/TemporalSuperResolution/TSRRejectShading.usf
@@ -124,13 +124,12 @@ Texture2D<tsr_half4> InputSceneTranslucencyTexture;
 Texture2D<tsr_halfC> ReprojectedHistoryGuideTexture;
 Texture2D<tsr_half4> ReprojectedHistoryMoireTexture;
 Texture2D<tsr_halfC> ResurrectedHistoryGuideTexture;
-Texture2D<tsr_half> ParallaxRejectionMaskTexture;
+Texture2D<tsr_half2> DecimateMaskTexture;
 Texture2D<tsr_ushort> IsMovingMaskTexture;
 
 RWTexture2D<tsr_half4> HistoryGuideOutput;
 RWTexture2D<tsr_half4> HistoryMoireOutput;
-RWTexture2D<tsr_half2> HistoryRejectionOutput;
-RWTexture2D<tsr_half>  HistoryResurrectionMaskOutput;
+RWTexture2D<tsr_half4> HistoryRejectionOutput;
 RWTexture2D<tsr_halfC> InputSceneColorOutput;
 RWTexture2D<tsr_half> InputSceneColorLdrLumaOutput;
 
@@ -447,6 +446,7 @@ void MainCS(
 		tsr_simd_half4 OriginalTranslucencyInput;
 	#endif
 	tsr_simd_halfC OriginalHistory;
+	tsr_simd_half VelocityEdge;
 	tsr_simd_bool bIsValid;
 	#if CONFIG_MOIRE_DETECTION
 		tsr_simd_half MoireInput;
@@ -473,12 +473,15 @@ void MainCS(
 			float2 TranslucencyTextureUV = ApplyScreenTransform(float2(ClampedFetchPixelPos), InputPixelPosToTranslucencyTextureUV);
 			TranslucencyTextureUV = fastClamp(TranslucencyTextureUV, TranslucencyTextureUVMin, TranslucencyTextureUVMax);
 
+			tsr_half2 DecimateMask = DecimateMaskTexture[ClampedFetchPixelPos];
+
 			OriginalOpaqueInput.SetElement(SimdIndex, InputTexture[ClampedFetchPixelPos]);
 			#if CONFIG_COMPOSE_TRANSLUCENCY
 				OriginalTranslucencyInput.SetElement(SimdIndex, InputSceneTranslucencyTexture.SampleLevel(GlobalBilinearClampedSampler, TranslucencyTextureUV, 0));
 			#endif
 			OriginalHistory.SetElement(SimdIndex, ReprojectedHistoryGuideTexture[ClampedFetchPixelPos]);
-			bIsValid.SetElement(SimdIndex, ParallaxRejectionMaskTexture[ClampedFetchPixelPos] > tsr_half(PARALLAX_REJECTION_MASK_THRESHOLD));
+			VelocityEdge.SetElement(SimdIndex, DecimateMask.g);
+			bIsValid.SetElement(SimdIndex, DecimateMask.r > tsr_half(PARALLAX_REJECTION_MASK_THRESHOLD));
 			#if CONFIG_MOIRE_DETECTION
 				MoireInput.SetElement(SimdIndex, InputMoireLumaTexture[ClampedFetchPixelPos]);
 				PrevMoireHistory.SetElement(SimdIndex, ReprojectedHistoryMoireTexture[ClampedFetchPixelPos]);
@@ -603,9 +606,7 @@ void MainCS(
 	// Measure rejection
 	tsr_simd_half RejectionBlendFinal = tsr_simd_half::Const(1.0);
 	tsr_simd_half RejectionClampBlend = tsr_simd_half::Const(1.0);
-	#if CONFIG_HISTORY_RESURRECTION
-		tsr_simd_bool bResurrectHistory = tsr_simd_bool::Const(false);
-	#endif
+	tsr_simd_bool bResurrectHistory = tsr_simd_bool::Const(false);
 	{
 		// Convert input to LDR
 		tsr_simd_halfC ExposedInput = OriginalInput * rcp(OriginalInput + tsr_simd_halfC::Const(tsr_half(PerceptionAdd)));
@@ -695,6 +696,16 @@ void MainCS(
 
 		FinalGuide = ExposedInput * tsr_simd_halfC::Vectorize(BlendFinal) - OriginalHistory * tsr_simd_halfC::Vectorize(BlendFinal - 1.0);
 	}
+	
+	// Controls how the history and its validity is clamped.
+	tsr_simd_half DisableHistoryClamp;
+	tsr_simd_half IncreaseValidityMultiplier;
+	{
+		tsr_simd_bool bIsDisoccluded = not(bIsValid);
+
+		DisableHistoryClamp        = select(and(bIsDisoccluded, not(bResurrectHistory)), tsr_simd_half::Const(0.0), min(RejectionClampBlend, VelocityEdge));
+		IncreaseValidityMultiplier = select(and(bIsDisoccluded, not(bResurrectHistory)), tsr_simd_half::Const(0.0), RejectionClampBlend);
+	}
 
 	// Mark the top left corner of the tile for debugging purposes.
 	#if DEBUG_OUTPUT && DEBUG_DISABLE_PADDING == 2
@@ -724,12 +735,12 @@ void MainCS(
 				GroupThreadIndex,
 				SimdIndex);
 
-			tsr_half LocalRejectionBlendFinal = RejectionBlendFinal.GetElement(SimdIndex);
-			tsr_half LocalRejectionClampBlend = RejectionClampBlend.GetElement(SimdIndex);
 			tsr_halfC FinalGuideColor = FinalGuide.GetElement(SimdIndex);
 			
-			LocalRejectionBlendFinal = LocalRejectionBlendFinal - tsr_half(0.5 * rcp(256.0));
-			LocalRejectionClampBlend = LocalRejectionClampBlend - tsr_half(0.5 * rcp(256.0));
+			tsr_half BitMask = (
+				select(bIsValid.GetElement(SimdIndex),          tsr_half(0x1) * rcp(tsr_half(255.0)), tsr_half(0x0)) + 
+				select(bResurrectHistory.GetElement(SimdIndex), tsr_half(0x2) * rcp(tsr_half(255.0)), tsr_half(0x0)));
+
 			#if CONFIG_ENABLE_STOCASTIC_QUANTIZATION
 			{
 				uint2 Random = Rand3DPCG16(int3(LaneSimdPixelOffset, View.StateFrameIndexMod8)).xy;
@@ -744,11 +755,11 @@ void MainCS(
 			#else
 				HistoryGuideOutput[OutputPixelPos] = tsr_half4(FinalGuideColor.xyz, tsr_half(1.0));
 			#endif
-			HistoryRejectionOutput[OutputPixelPos] = tsr_half2(LocalRejectionBlendFinal, LocalRejectionClampBlend);
-			
-			#if CONFIG_HISTORY_RESURRECTION
-				HistoryResurrectionMaskOutput[OutputPixelPos] = select(bResurrectHistory.GetElement(SimdIndex), tsr_half(1.0), tsr_half(0.0));
-			#endif
+			HistoryRejectionOutput[OutputPixelPos] = tsr_half4(
+				RejectionBlendFinal.GetElement(SimdIndex) - tsr_half(0.5 * rcp(256.0)),
+				DisableHistoryClamp.GetElement(SimdIndex) - tsr_half(0.5 * rcp(256.0)),
+				tsr_half(255.5 * rcp(256.0)) - IncreaseValidityMultiplier.GetElement(SimdIndex),
+				BitMask);
 		}
 		
 		#if DEBUG_OUTPUT
diff --git a/Engine/Shaders/Private/TemporalSuperResolution/TSRUpdateHistory.usf b/Engine/Shaders/Private/TemporalSuperResolution/TSRUpdateHistory.usf
--- a/Engine/Shaders/Private/TemporalSuperResolution/TSRUpdateHistory.usf
+++ b/Engine/Shaders/Private/TemporalSuperResolution/TSRUpdateHistory.usf
@@ -45,7 +45,6 @@
 #endif
 
 #define CONFIG_SELECTIVE_REPROJECTION 1
-#define CONFIG_RESPONSIVE_STENCIL 1
 #define CONFIG_TAA_RESPONSIVE_BEHAVIOR 1
 #define CONFIG_LOAD_TRANSLUCENCY_REJECTION 0
 
@@ -91,19 +90,13 @@
 //------------------------------------------------------- PARAMETERS
 
 Texture2D<tsr_halfC> InputSceneColorTexture;
-Texture2D<uint2> InputSceneStencilTexture;
 Texture2D<tsr_half4> InputSceneTranslucencyTexture;
 
-Texture2D<tsr_half2> HistoryRejectionTexture;
+Texture2D<tsr_half4> HistoryRejectionTexture;
 
 Texture2D<float2> DilatedVelocityTexture;
-Texture2D<tsr_half> ParallaxRejectionMaskTexture;
-#if CONFIG_HISTORY_RESURRECTION
-	Texture2D<tsr_half> HistoryResurrectionMaskTexture;
-#endif
 Texture2D<tsr_ushort> AntiAliasingTexture;
 Texture2D<tsr_half> NoiseFilteringTexture;
-Texture2D<tsr_half2> HoleFilledVelocityMaskTexture;
 
 int2   TranslucencyPixelPosMin;
 int2   TranslucencyPixelPosMax;
@@ -120,7 +113,6 @@ float InvWeightClampingPixelSpeed;
 float InputToHistoryFactor;
 float ResurrectionFrameIndex;
 float PrevFrameIndex;
-uint ResponsiveStencilMask;
 uint bGenerateOutputMip1;
 uint bGenerateOutputMip2;
 uint bHasSeparateTranslucency;
@@ -460,12 +452,11 @@ void MainCS(
 
 	// Fetch reprojection-related information.
 	float2x2 PrevScreenPos = ScreenPos;
-	tsr_half2 LowFrequencyRejection = tsr_half(1.0).xx;
-	tsr_half2 LowFrequencyClamp = tsr_half(1.0).xx;
+	tsr_half2 LowFrequencyRejection      = tsr_half(1.0).xx;
+	tsr_half2 DisableHistoryClamp        = tsr_half(1.0).xx;
+	tsr_half2 DecreaseValidityMultiplier = tsr_half(0.0).xx;
 	tsr_half2 OutputPixelVelocity = tsr_half(0.0).xx;
 	tsr_half2 NoiseFiltering = tsr_half(0.0).xx;
-	tsr_half2x2 HoleFilledVelocityMask = dpv_interleave_mono_registers(tsr_half(1.0).xx);
-	bool2 bIsResponsiveAAPixel = false;
 	bool2 bIsOffScreen = false;
 	bool2 bIsDisoccluded = false;
 	#if CONFIG_HISTORY_RESURRECTION
@@ -486,19 +477,11 @@ void MainCS(
 			float2 RawEncodedVelocity[DPV_PIXEL_PER_LANE];
 		#endif
 
-		tsr_half RawParallaxRejectionMask[DPV_PIXEL_PER_LANE];
-		#if CONFIG_HISTORY_RESURRECTION
-			tsr_half RawHistoryResurrectionMask[DPV_PIXEL_PER_LANE];
-		#endif
-		tsr_half2 RawLowFrequencyRejection[DPV_PIXEL_PER_LANE];
-		
-		uint RawSceneStencilRef[DPV_PIXEL_PER_LANE];
-		uint RawSceneStencilRefNeighborhood[CONFIG_SAMPLES - 1][DPV_PIXEL_PER_LANE];
+		tsr_half4 RawHistoryRejection[DPV_PIXEL_PER_LANE];
 		
 		#if CONFIG_REJECTION_ANTI_ALIASING
 			tsr_ushort RawEncodedInputTexelOffset[DPV_PIXEL_PER_LANE];
 			tsr_half RawNoiseFiltering[DPV_PIXEL_PER_LANE]; 
-			tsr_half2 RawHoleFilledVelocityMask[DPV_PIXEL_PER_LANE];
 		#endif
 
 		// Issue overlapped texture fetches
@@ -513,19 +496,12 @@ void MainCS(
 					RawEncodedVelocity[PixelId]   = DilatedVelocityTexture[LocalInputPixelPos];
 				#endif
 
-				RawParallaxRejectionMask[PixelId] = ParallaxRejectionMaskTexture[LocalInputPixelPos];
-				#if CONFIG_HISTORY_RESURRECTION
-					RawHistoryResurrectionMask[PixelId] = HistoryResurrectionMaskTexture[LocalInputPixelPos];
-				#endif
-				RawSceneStencilRef[PixelId]       = InputSceneStencilTexture.Load(int3(LocalInputPixelPos, 0)) STENCIL_COMPONENT_SWIZZLE;
-				
-				RawLowFrequencyRejection[PixelId] = HistoryRejectionTexture[LocalInputPixelPos];
+				RawHistoryRejection[PixelId] = HistoryRejectionTexture[LocalInputPixelPos];
 				
 				#if CONFIG_REJECTION_ANTI_ALIASING
 				{
 					RawEncodedInputTexelOffset[PixelId] = AntiAliasingTexture[LocalInputPixelPos];
 					RawNoiseFiltering[PixelId] = NoiseFilteringTexture[LocalInputPixelPos];
-					RawHoleFilledVelocityMask[PixelId] = HoleFilledVelocityMaskTexture[LocalInputPixelPos];
 				}
 				#endif
 			}
@@ -533,24 +509,6 @@ void MainCS(
 			// Vector in pixel between pixel K -> O.
 			tsr_half2x2 dInputKO = tsr_half2x2(InputPPCo - InputPPCk);
 		
-			UNROLL_N(CONFIG_SAMPLES - 1)
-			for (uint NeighborId = 0; NeighborId < (CONFIG_SAMPLES - 1); NeighborId++)
-			{
-				tsr_short2x2 SampleInputPixelPos;
-				tsr_half2x2 PixelOffset;
-				ComputeInputKernelSamplePosition(
-					InputPixelPos, dInputKO, /* SampleId = */ NeighborId + 1,
-					InputPixelPosMin, InputPixelPosMax,
-					/* out */ SampleInputPixelPos,
-					/* out */ PixelOffset);
-			
-				UNROLL_N(DPV_PIXEL_PER_LANE)
-				for (uint PixelId = 0; PixelId < DPV_PIXEL_PER_LANE; PixelId++)
-				{
-					RawSceneStencilRefNeighborhood[NeighborId][PixelId] = InputSceneStencilTexture.Load(int3(dpv_access_pixel(SampleInputPixelPos, PixelId), 0)) STENCIL_COMPONENT_SWIZZLE;
-				}
-			} // for (uint NeighborId = 0; NeighborId < (CONFIG_SAMPLES - 1); NeighborId++)
-			
 			// Fetch the 2x2 neighborhood of the velocity
 			#if CONFIG_INTERPOLATE_VELOCITY
 			UNROLL_N(DPV_PIXEL_PER_LANE)
@@ -568,8 +526,6 @@ void MainCS(
 			#endif
 		}
 		
-		uint2 SceneStencilRef;
-
 		// Process texture fetches.
 		PLATFORM_SPECIFIC_ISOLATE
 		{
@@ -619,29 +575,24 @@ void MainCS(
 				OutputPixelVelocity = tsr_half2(dpv_length(dpv_mul(ScreenVelocity, HistoryInfo_ViewportSize)));
 			}
 			
+			tsr_ushort2 BitMask = tsr_ushort2(round(dpv_interleave_registers_array(RawHistoryRejection)[3] * tsr_half(255.0)));
+
 			// Detect whether the history reprojection is successful
 			IsOffScreenOrDisoccluded(
 				bCameraCut,
 				PrevScreenPos,
-				dpv_interleave_registers_array(RawParallaxRejectionMask),
+				(BitMask & tsr_ushort(0x1)) == tsr_ushort(0),
 				/* out */ bIsOffScreen,
 				/* out */ bIsDisoccluded);
 				
 			#if CONFIG_HISTORY_RESURRECTION
-				bIsHistoryResurrection = dpv_interleave_registers_array(RawHistoryResurrectionMask) != tsr_half(0.0);
+				bIsHistoryResurrection = (BitMask & tsr_ushort(0x2)) != tsr_ushort(0);
 			#endif
 
 			// Process input texel, forcing to pack register.
-			LowFrequencyRejection = dpv_force_interleave_registers_array(RawLowFrequencyRejection)[0];
-			LowFrequencyClamp     = dpv_force_interleave_registers_array(RawLowFrequencyRejection)[1];
-			SceneStencilRef       = dpv_force_interleave_registers_array(RawSceneStencilRef);
-			
-			// Process neighborhood
-			UNROLL_N(CONFIG_SAMPLES - 1)
-			for (uint NeighborId = 0; NeighborId < (CONFIG_SAMPLES - 1); NeighborId++)
-			{
-				SceneStencilRef |= dpv_interleave_registers_array(RawSceneStencilRefNeighborhood[NeighborId]);
-			}
+			LowFrequencyRejection      = dpv_force_interleave_registers_array(RawHistoryRejection)[0];
+			DisableHistoryClamp        = dpv_force_interleave_registers_array(RawHistoryRejection)[1];
+			DecreaseValidityMultiplier = dpv_force_interleave_registers_array(RawHistoryRejection)[2];
 		}
 
 		#if CONFIG_REJECTION_ANTI_ALIASING
@@ -649,7 +600,6 @@ void MainCS(
 			tsr_ushort2 EncodedInputTexelOffset = dpv_force_interleave_registers_array(RawEncodedInputTexelOffset);
 			
 			NoiseFiltering = dpv_force_interleave_registers_array(RawNoiseFiltering);
-			HoleFilledVelocityMask = dpv_interleave_mono_registers(tsr_half(1.0).xx) - dpv_force_interleave_registers_array(RawHoleFilledVelocityMask);
 
 			SpatialAntiAliasingOffset = DecodeSpatialAntiAliasingOffset(EncodedInputTexelOffset);
 		}
@@ -657,13 +607,6 @@ void MainCS(
 		
 		// Final post processing.
 		{
-			// Fetch whether the pixel is responsive AA or not.
-			#if CONFIG_RESPONSIVE_STENCIL
-			{
-				bIsResponsiveAAPixel = (SceneStencilRef & ResponsiveStencilMask) != 0;
-			}
-			#endif
-
 			#if !CONFIG_CLAMP
 			{
 				bIsDisoccluded = false;
@@ -881,15 +824,11 @@ void MainCS(
 				tsr_half2 CoarseCurrentContribution = CoarseInputPixelAlignement * tsr_half(ActualHistoryHisteresis);
 				tsr_half2 IdealCurrentContribution = IdealInputPixelAlignement * tsr_half(ActualHistoryHisteresis);
 
-				// Select whether history should be clamped.
-				BlendClamp = select(and(bIsDisoccluded, !bIsHistoryResurrection), tsr_half(0.0), min(HoleFilledVelocityMask[1], LowFrequencyClamp));
-		
 				// Reduces the number of pixels when rejecting history
 				tsr_half2 ClampedPrevHistoryValidity = PrevHistoryValidity;
 				#if 1
 				{
-					tsr_half2 MaxValidity = lerp(tsr_half(WeightClampingRejection), tsr_half(1.0), select(and(bIsDisoccluded, !bIsHistoryResurrection), tsr_half(0.0), LowFrequencyClamp));
-					//tsr_half2 MaxValidity = lerp(tsr_half(WeightClampingRejection), tsr_half(1.0), BlendClamp);
+					tsr_half2 MaxValidity = tsr_half(1.0) - tsr_half(WeightClampingRejection) * DecreaseValidityMultiplier;
 
 					ClampedPrevHistoryValidity = min(ClampedPrevHistoryValidity, MaxValidity);
 				}
@@ -1038,11 +977,11 @@ void MainCS(
 	{
 		tsr_halfCx2 ClampedPrevHighFrequencyColor = fastClamp(PrevHighFrequencyColor, InputMinColor, InputMaxColor);
 		tsr_halfCx2 BlendedPrevHighFrequencyColor;
-		BlendedPrevHighFrequencyColor[0] = lerp(ClampedPrevHighFrequencyColor[0], PrevHighFrequencyColor[0], BlendClamp);
-		BlendedPrevHighFrequencyColor[1] = lerp(ClampedPrevHighFrequencyColor[1], PrevHighFrequencyColor[1], BlendClamp);
-		BlendedPrevHighFrequencyColor[2] = lerp(ClampedPrevHighFrequencyColor[2], PrevHighFrequencyColor[2], BlendClamp);
+		BlendedPrevHighFrequencyColor[0] = lerp(ClampedPrevHighFrequencyColor[0], PrevHighFrequencyColor[0], DisableHistoryClamp);
+		BlendedPrevHighFrequencyColor[1] = lerp(ClampedPrevHighFrequencyColor[1], PrevHighFrequencyColor[1], DisableHistoryClamp);
+		BlendedPrevHighFrequencyColor[2] = lerp(ClampedPrevHighFrequencyColor[2], PrevHighFrequencyColor[2], DisableHistoryClamp);
 		#if CONFIG_SCENE_COLOR_ALPHA
-			BlendedPrevHighFrequencyColor[3] = lerp(ClampedPrevHighFrequencyColor[3], PrevHighFrequencyColor[3], BlendClamp);
+			BlendedPrevHighFrequencyColor[3] = lerp(ClampedPrevHighFrequencyColor[3], PrevHighFrequencyColor[3], DisableHistoryClamp);
 		#endif
 		
 		// Clamp the validity due to motion to maintain better sharpness in history reprojection under motion.
diff --git a/Engine/Shaders/Private/TemporalSuperResolution/TSRVisualize.usf b/Engine/Shaders/Private/TemporalSuperResolution/TSRVisualize.usf
--- a/Engine/Shaders/Private/TemporalSuperResolution/TSRVisualize.usf
+++ b/Engine/Shaders/Private/TemporalSuperResolution/TSRVisualize.usf
@@ -33,9 +33,7 @@ float OutputToHistoryResolutionFractionSquare;
 Texture2D<tsr_half3> SceneColorTexture;
 Texture2D<float>     ClosestDepthTexture;
 Texture2D<float2>    DilatedVelocityTexture;
-Texture2D<tsr_half>  ParallaxRejectionMaskTexture;
-Texture2D<tsr_half2> HistoryRejectionTexture;
-Texture2D<tsr_half>  HistoryResurrectionMaskTexture;
+Texture2D<tsr_half4> HistoryRejectionTexture;
 Texture2D<tsr_half4> HistoryMetadataTexture;
 Texture2D<tsr_half3> ResurrectedHistoryColorTexture;
 
@@ -93,37 +91,36 @@ tsr_half3 DisplayHistorySampleCount(FContext Ctx)
 
 tsr_half3 DisplayHistoryRejection(FContext Ctx)
 {
-	tsr_half2 RawRejection = HistoryRejectionTexture.SampleLevel(GlobalPointClampedSampler, Ctx.InputUV, 0);
-	return SampleTSROutputGray(Ctx) * lerp(kRed, kWhite, saturate(RawRejection.x));
+	tsr_half LowFrequencyRejection = HistoryRejectionTexture.SampleLevel(GlobalPointClampedSampler, Ctx.InputUV, 0).r;
+	return SampleTSROutputGray(Ctx) * lerp(kRed, kWhite, saturate(LowFrequencyRejection));
 }
 
 tsr_half3 DisplayHistoryClamp(FContext Ctx)
 {
-	tsr_half2 RawRejection = HistoryRejectionTexture.SampleLevel(GlobalPointClampedSampler, Ctx.InputUV, 0);
-	return SampleTSROutputGray(Ctx) * lerp(kRed, kWhite, saturate(RawRejection.y));
+	tsr_half LowFrequencyClamp = HistoryRejectionTexture.SampleLevel(GlobalPointClampedSampler, Ctx.InputUV, 0).g;
+	return SampleTSROutputGray(Ctx) * lerp(kRed, kWhite, saturate(LowFrequencyClamp));
 }
 
 tsr_half3 DisplayParallaxDisocclusionMask(FContext Ctx)
 {
-	tsr_half ParallaxDisocclusionMask = ParallaxRejectionMaskTexture.SampleLevel(GlobalPointClampedSampler, Ctx.InputUV, 0);
-
-	return SampleTSROutputGray(Ctx) * lerp(kRed, kWhite, saturate(ParallaxDisocclusionMask));
+	bool ParallaxDisocclusionMask = (tsr_ushort(HistoryRejectionTexture.SampleLevel(GlobalPointClampedSampler, Ctx.InputUV, 0).a * tsr_half(255.0)) & tsr_ushort(0x1)) != tsr_ushort(0x0);
+	return SampleTSROutputGray(Ctx) * select(ParallaxDisocclusionMask, kWhite, kRed);
 }
 
 tsr_half3 DisplayResurrectionMask(FContext Ctx)
 {
-	tsr_half ResurrectionMask = HistoryResurrectionMaskTexture.SampleLevel(GlobalPointClampedSampler, Ctx.InputUV, 0);
+	bool ResurrectionMask = (tsr_ushort(HistoryRejectionTexture.SampleLevel(GlobalPointClampedSampler, Ctx.InputUV, 0).a * tsr_half(255.0)) & tsr_ushort(0x2)) != tsr_ushort(0x0);
 	if (!bCanResurrectHistory)
 	{
 		return SampleTSROutputGray(Ctx) * kRed;
 	}
-	return SampleTSROutputGray(Ctx) * lerp(kWhite, kGreen, saturate(ResurrectionMask));
+	return SampleTSROutputGray(Ctx) * select(ResurrectionMask, kGreen, kWhite);
 }
 
 tsr_half3 DisplayResurrectedColor(FContext Ctx)
 {
 	float ClosestDeviceZ = ClosestDepthTexture.SampleLevel(GlobalPointClampedSampler, Ctx.InputUV, 0);
-	tsr_half ResurrectionMask = HistoryResurrectionMaskTexture.SampleLevel(GlobalPointClampedSampler, Ctx.InputUV, 0);
+	bool ResurrectionMask = (tsr_ushort(HistoryRejectionTexture.SampleLevel(GlobalPointClampedSampler, Ctx.InputUV, 0).a * tsr_half(255.0)) & tsr_ushort(0x2)) != tsr_ushort(0x0);
 
 	float4 ThisClip = float4(Ctx.ScreenPos, ClosestDeviceZ, 1);
 	float4 ResurrectionClip = mul(ThisClip, ClipToResurrectionClip); 
@@ -134,7 +131,7 @@ tsr_half3 DisplayResurrectedColor(FContext Ctx)
 
 	tsr_half3 ResurrectedColor = ResurrectedHistoryColorTexture.SampleLevel(GlobalBilinearClampedSampler, KernelHistoryBufferUV, 0) * tsr_half(ResurrectionPreExposureCorrection);
 
-	return GrayScaleColor(ResurrectedColor) * lerp(kWhite, kGreen, saturate(ResurrectionMask));
+	return GrayScaleColor(ResurrectedColor) * select(ResurrectionMask, kGreen, kWhite);
 }
 
 
diff --git a/Engine/Source/Runtime/Renderer/Private/PostProcess/TemporalSuperResolution.cpp b/Engine/Source/Runtime/Renderer/Private/PostProcess/TemporalSuperResolution.cpp
--- a/Engine/Source/Runtime/Renderer/Private/PostProcess/TemporalSuperResolution.cpp
+++ b/Engine/Source/Runtime/Renderer/Private/PostProcess/TemporalSuperResolution.cpp
@@ -224,12 +224,6 @@ TAutoConsoleVariable<int32> CVarTSRAsyncCompute(
 	TEXT(" 3: Run all passes on async compute;"),
 	ECVF_RenderThreadSafe);
 
-// TODO: currently dead code, that is technically relevent for 5.0's r.TSR.History.SeparateTranslucency=1 but that has too many problems on translucency contour around opaque geometry.
-//TAutoConsoleVariable<int32> CVarTSREnableResponiveAA(
-//	TEXT("r.TSR.Translucency.EnableResponiveAA"), 1,
-//	TEXT("Whether the responsive AA should keep history fully clamped."),
-//	ECVF_RenderThreadSafe);
-
 TAutoConsoleVariable<float> CVarTSRWeightClampingSampleCount(
 	TEXT("r.TSR.Velocity.WeightClampingSampleCount"), 4.0f,
 	TEXT("Number of sample to count to in history pixel to clamp history to when output pixel velocity reach r.TSR.Velocity.WeightClampingPixelSpeed. ")
@@ -590,9 +584,7 @@ class FTSRDecimateHistoryCS : public FTSRShader
 		SHADER_PARAMETER_RDG_TEXTURE_UAV(RWTexture2DArray, ReprojectedHistoryGuideOutput)
 		SHADER_PARAMETER_RDG_TEXTURE_UAV(RWTexture2DArray, ReprojectedHistoryMoireOutput)
 		SHADER_PARAMETER_RDG_TEXTURE_UAV(RWTexture2D, HoleFilledVelocityOutput)
-		SHADER_PARAMETER_RDG_TEXTURE_UAV(RWTexture2D, HoleFilledVelocityMaskOutput)
-
-		SHADER_PARAMETER_RDG_TEXTURE_UAV(RWTexture2D, ParallaxRejectionMaskOutput)
+		SHADER_PARAMETER_RDG_TEXTURE_UAV(RWTexture2D, DecimateMaskOutput)
 
 		SHADER_PARAMETER_RDG_TEXTURE_UAV(RWTexture2DArray, DebugOutput)
 	END_SHADER_PARAMETER_STRUCT()
@@ -627,13 +619,12 @@ class FTSRRejectShadingCS : public FTSRShader
 		SHADER_PARAMETER_RDG_TEXTURE_SRV(Texture2D, ReprojectedHistoryGuideTexture)
 		SHADER_PARAMETER_RDG_TEXTURE_SRV(Texture2D, ReprojectedHistoryMoireTexture)
 		SHADER_PARAMETER_RDG_TEXTURE_SRV(Texture2D, ResurrectedHistoryGuideTexture)
-		SHADER_PARAMETER_RDG_TEXTURE(Texture2D, ParallaxRejectionMaskTexture)
+		SHADER_PARAMETER_RDG_TEXTURE(Texture2D, DecimateMaskTexture)
 		SHADER_PARAMETER_RDG_TEXTURE_SRV(Texture2D, IsMovingMaskTexture)
 
 		SHADER_PARAMETER_RDG_TEXTURE_UAV(RWTexture2D, HistoryGuideOutput)
 		SHADER_PARAMETER_RDG_TEXTURE_UAV(RWTexture2D, HistoryMoireOutput)
 		SHADER_PARAMETER_RDG_TEXTURE_UAV(RWTexture2D, HistoryRejectionOutput)
-		SHADER_PARAMETER_RDG_TEXTURE_UAV(RWTexture2D, HistoryResurrectionMaskOutput)
 		SHADER_PARAMETER_RDG_TEXTURE_UAV(RWTexture2D, InputSceneColorOutput)
 		SHADER_PARAMETER_RDG_TEXTURE_UAV(RWTexture2D, InputSceneColorLdrLumaOutput)
 		SHADER_PARAMETER_RDG_TEXTURE_UAV(RWTexture2D, DebugOutput)
@@ -721,7 +712,7 @@ class FTSRMergeRejectionCS : public FTSRShader
 		SHADER_PARAMETER(FMatrix44f, ClipToResurrectionClip)
 
 		SHADER_PARAMETER_RDG_TEXTURE(Texture2D, ClosestDepthTexture)
-		SHADER_PARAMETER_RDG_TEXTURE(Texture2D, HistoryResurrectionMaskTexture)
+		SHADER_PARAMETER_RDG_TEXTURE(Texture2D, HistoryRejectionTexture)
 
 		SHADER_PARAMETER_RDG_TEXTURE_UAV(RWTexture2D, DilatedVelocityOutput)
 		SHADER_PARAMETER_RDG_TEXTURE_UAV(RWTexture2DArray, DebugOutput)
@@ -799,17 +790,12 @@ class FTSRUpdateHistoryCS : public FTSRShader
 	BEGIN_SHADER_PARAMETER_STRUCT(FParameters, )
 		SHADER_PARAMETER_STRUCT_INCLUDE(FTSRCommonParameters, CommonParameters)
 		SHADER_PARAMETER_RDG_TEXTURE(Texture2D, InputSceneColorTexture)
-		SHADER_PARAMETER_RDG_TEXTURE_SRV(Texture2D, InputSceneStencilTexture)
 		SHADER_PARAMETER_RDG_TEXTURE(Texture2D, InputSceneTranslucencyTexture)
 
 		SHADER_PARAMETER_RDG_TEXTURE(Texture2D, HistoryRejectionTexture)
-
 		SHADER_PARAMETER_RDG_TEXTURE(Texture2D, DilatedVelocityTexture)
-		SHADER_PARAMETER_RDG_TEXTURE(Texture2D, ParallaxRejectionMaskTexture)
-		SHADER_PARAMETER_RDG_TEXTURE(Texture2D, HistoryResurrectionMaskTexture)
 		SHADER_PARAMETER_RDG_TEXTURE(Texture2D, AntiAliasingTexture)
 		SHADER_PARAMETER_RDG_TEXTURE(Texture2D, NoiseFilteringTexture)
-		SHADER_PARAMETER_RDG_TEXTURE(Texture2D, HoleFilledVelocityMaskTexture)
 
 		SHADER_PARAMETER_STRUCT(FScreenPassTextureViewportParameters, TranslucencyInfo)
 		SHADER_PARAMETER(FIntPoint, TranslucencyPixelPosMin)
@@ -829,7 +815,6 @@ class FTSRUpdateHistoryCS : public FTSRShader
 		SHADER_PARAMETER(float, InputContributionMultiplier)
 		SHADER_PARAMETER(float, ResurrectionFrameIndex)
 		SHADER_PARAMETER(float, PrevFrameIndex)
-		SHADER_PARAMETER(int32, ResponsiveStencilMask)
 		SHADER_PARAMETER(int32, bGenerateOutputMip1)
 		SHADER_PARAMETER(int32, bGenerateOutputMip2)
 		SHADER_PARAMETER(int32, bHasSeparateTranslucency)
@@ -970,9 +955,7 @@ class FTSRVisualizeCS : public FTSRShader
 		SHADER_PARAMETER_RDG_TEXTURE_SRV(Texture2D, SceneColorTexture)
 		SHADER_PARAMETER_RDG_TEXTURE(Texture2D, ClosestDepthTexture)
 		SHADER_PARAMETER_RDG_TEXTURE(Texture2D, DilatedVelocityTexture)
-		SHADER_PARAMETER_RDG_TEXTURE(Texture2D, ParallaxRejectionMaskTexture)
 		SHADER_PARAMETER_RDG_TEXTURE(Texture2D, HistoryRejectionTexture)
-		SHADER_PARAMETER_RDG_TEXTURE(Texture2D, HistoryResurrectionMaskTexture)
 		SHADER_PARAMETER_RDG_TEXTURE_SRV(Texture2D, HistoryMetadataTexture)
 		SHADER_PARAMETER_RDG_TEXTURE_SRV(Texture2D, ResurrectedHistoryColorTexture)
 
@@ -1829,24 +1812,21 @@ FDefaultTemporalUpscaler::FOutputs AddTemporalSuperResolutionPasses(
 	// Decimate input to flicker at same frequency as input.
 	FRDGTextureRef ReprojectedHistoryGuideTexture = nullptr;
 	FRDGTextureRef ReprojectedHistoryMoireTexture = nullptr;
-	FRDGTextureRef ParallaxRejectionMaskTexture = nullptr;
-	FRDGTextureRef HoleFilledVelocityMaskTexture = nullptr;
+	FRDGTextureRef DecimateMaskTexture = nullptr;
 	{
 		FRDGTextureRef HoleFilledVelocityTexture;
 		{
 			FRDGTextureDesc Desc = FRDGTextureDesc::Create2D(
 				InputExtent,
-				PF_R8,
+				PF_R8G8,
 				FClearValueBinding::None,
 				/* InFlags = */ TexCreate_ShaderResource | TexCreate_UAV);
 
-			ParallaxRejectionMaskTexture = GraphBuilder.CreateTexture(Desc, TEXT("TSR.ParallaxRejectionMask"));
+			Desc.Format = PF_R8G8;
+			DecimateMaskTexture = GraphBuilder.CreateTexture(Desc, TEXT("TSR.DecimateMask"));
 
 			Desc.Format = PF_G16R16;
 			HoleFilledVelocityTexture = GraphBuilder.CreateTexture(Desc, TEXT("TSR.Velocity.HoleFilled"));
-
-			Desc.Format = PF_R8G8;
-			HoleFilledVelocityMaskTexture = GraphBuilder.CreateTexture(Desc, TEXT("TSR.velocity.HoleFillMask"));
 		}
 
 		{
@@ -1915,8 +1895,7 @@ FDefaultTemporalUpscaler::FOutputs AddTemporalSuperResolutionPasses(
 		PassParameters->ReprojectedHistoryGuideOutput = GraphBuilder.CreateUAV(ReprojectedHistoryGuideTexture);
 		PassParameters->ReprojectedHistoryMoireOutput = GraphBuilder.CreateUAV(ReprojectedHistoryMoireTexture);
 		PassParameters->HoleFilledVelocityOutput = GraphBuilder.CreateUAV(HoleFilledVelocityTexture);
-		PassParameters->HoleFilledVelocityMaskOutput = GraphBuilder.CreateUAV(HoleFilledVelocityMaskTexture);
-		PassParameters->ParallaxRejectionMaskOutput = GraphBuilder.CreateUAV(ParallaxRejectionMaskTexture);
+		PassParameters->DecimateMaskOutput = GraphBuilder.CreateUAV(DecimateMaskTexture);
 		PassParameters->DebugOutput = CreateDebugUAV(InputExtent, TEXT("Debug.TSR.DecimateHistory"));
 
 		FTSRDecimateHistoryCS::FPermutationDomain PermutationVector;
@@ -1947,7 +1926,6 @@ FDefaultTemporalUpscaler::FOutputs AddTemporalSuperResolutionPasses(
 
 	// Perform a history reject the history.
 	FRDGTextureRef HistoryRejectionTexture = nullptr;
-	FRDGTextureRef HistoryResurrectionMaskTexture = nullptr;
 	FRDGTextureRef InputSceneColorLdrLumaTexture = nullptr;
 	{
 		const int32 GroupTileSize = CVarTSRShadingTileSize.GetValueOnRenderThread();
@@ -1979,24 +1957,13 @@ FDefaultTemporalUpscaler::FOutputs AddTemporalSuperResolutionPasses(
 		{
 			FRDGTextureDesc Desc = FRDGTextureDesc::Create2D(
 				InputExtent,
-				PF_R8G8,
+				PF_R8G8B8A8,
 				FClearValueBinding::None,
 				/* InFlags = */ TexCreate_ShaderResource | TexCreate_UAV);
 
 			HistoryRejectionTexture = GraphBuilder.CreateTexture(Desc, TEXT("TSR.HistoryRejection"));
 		}
 
-		if (bCanResurrectHistory)
-		{
-			FRDGTextureDesc Desc = FRDGTextureDesc::Create2D(
-				InputExtent,
-				PF_R8,
-				FClearValueBinding::None,
-				/* InFlags = */ TexCreate_ShaderResource | TexCreate_UAV);
-
-			HistoryResurrectionMaskTexture = GraphBuilder.CreateTexture(Desc, TEXT("TSR.ResurrectionMask"));
-		}
-
 		FScreenPassTextureViewport TranslucencyViewport(
 			SeparateTranslucencyTexture->Desc.Extent, SeparateTranslucencyRect);
 
@@ -2033,7 +2000,7 @@ FDefaultTemporalUpscaler::FOutputs AddTemporalSuperResolutionPasses(
 			PassParameters->ResurrectedHistoryGuideTexture = GraphBuilder.CreateSRV(FRDGTextureSRVDesc::CreateForSlice(
 				ReprojectedHistoryGuideTexture, /* SliceIndex = */ 1));
 		}
-		PassParameters->ParallaxRejectionMaskTexture = ParallaxRejectionMaskTexture;
+		PassParameters->DecimateMaskTexture = DecimateMaskTexture;
 		PassParameters->IsMovingMaskTexture = IsMovingMaskTexture;
 
 		// Outputs
@@ -2074,12 +2041,6 @@ FDefaultTemporalUpscaler::FOutputs AddTemporalSuperResolutionPasses(
 			// Output how the history should rejected in the HistoryUpdate
 			PassParameters->HistoryRejectionOutput = GraphBuilder.CreateUAV(HistoryRejectionTexture);
 
-			// Output whether resurrection should be applied.
-			if (bCanResurrectHistory)
-			{
-				PassParameters->HistoryResurrectionMaskOutput = GraphBuilder.CreateUAV(HistoryResurrectionMaskTexture);
-			}
-
 			// Output the composed translucency and opaque scene color to speed up HistoryUpdate
 			PassParameters->InputSceneColorOutput = bComputeInputSceneColorTexture
 				? GraphBuilder.CreateUAV(InputSceneColorTexture)
@@ -2127,7 +2088,7 @@ FDefaultTemporalUpscaler::FOutputs AddTemporalSuperResolutionPasses(
 		PassParameters->ClipToResurrectionClip = ClipToResurrectionClip;
 
 		PassParameters->ClosestDepthTexture = ClosestDepthTexture;
-		PassParameters->HistoryResurrectionMaskTexture = HistoryResurrectionMaskTexture;
+		PassParameters->HistoryRejectionTexture = HistoryRejectionTexture;
 
 		PassParameters->DilatedVelocityOutput = GraphBuilder.CreateUAV(DilatedVelocityTexture);
 		PassParameters->DebugOutput = CreateDebugUAV(InputExtent, TEXT("Debug.TSR.MergeRejection"));
@@ -2217,17 +2178,12 @@ FDefaultTemporalUpscaler::FOutputs AddTemporalSuperResolutionPasses(
 		FTSRUpdateHistoryCS::FParameters* PassParameters = GraphBuilder.AllocParameters<FTSRUpdateHistoryCS::FParameters>();
 		PassParameters->CommonParameters = CommonParameters;
 		PassParameters->InputSceneColorTexture = InputSceneColorTexture;
-		PassParameters->InputSceneStencilTexture = GraphBuilder.CreateSRV(
-			FRDGTextureSRVDesc::CreateWithPixelFormat(PassInputs.SceneDepth.Texture, PF_X24_G8));
 		PassParameters->InputSceneTranslucencyTexture = SeparateTranslucencyTexture;
 		PassParameters->HistoryRejectionTexture = HistoryRejectionTexture;
 
 		PassParameters->DilatedVelocityTexture = DilatedVelocityTexture;
-		PassParameters->ParallaxRejectionMaskTexture = ParallaxRejectionMaskTexture;
-		PassParameters->HistoryResurrectionMaskTexture = bCanResurrectHistory ? HistoryResurrectionMaskTexture : BlackDummy;
 		PassParameters->AntiAliasingTexture = AntiAliasingTexture;
 		PassParameters->NoiseFilteringTexture = NoiseFilteringTexture;
-		PassParameters->HoleFilledVelocityMaskTexture = HoleFilledVelocityMaskTexture;
 
 		PassParameters->TranslucencyPixelPosMin = SeparateTranslucencyRect.Min;
 		PassParameters->TranslucencyPixelPosMax = SeparateTranslucencyRect.Max - 1;
@@ -2241,14 +2197,12 @@ FDefaultTemporalUpscaler::FOutputs AddTemporalSuperResolutionPasses(
 		// All parameters to control the sample count in history.
 		PassParameters->HistorySampleCount = MaxHistorySampleCount / OutputToHistoryResolutionFractionSquare;
 		PassParameters->HistoryHisteresis = 1.0f / PassParameters->HistorySampleCount;
-		PassParameters->WeightClampingRejection = (CVarTSRHistoryRejectionSampleCount.GetValueOnRenderThread() / OutputToHistoryResolutionFractionSquare) * PassParameters->HistoryHisteresis;
+		PassParameters->WeightClampingRejection = 1.0f - (CVarTSRHistoryRejectionSampleCount.GetValueOnRenderThread() / OutputToHistoryResolutionFractionSquare) * PassParameters->HistoryHisteresis;
 		PassParameters->WeightClampingPixelSpeedAmplitude = FMath::Clamp(1.0f - CVarTSRWeightClampingSampleCount.GetValueOnRenderThread() * PassParameters->HistoryHisteresis, 0.0f, 1.0f);
 		PassParameters->InvWeightClampingPixelSpeed = 1.0f / (CVarTSRWeightClampingPixelSpeed.GetValueOnRenderThread() * OutputToHistoryResolutionFraction);
 		
 		PassParameters->InputToHistoryFactor = float(HistorySize.X) / float(InputRect.Width());
 		PassParameters->InputContributionMultiplier = OutputToHistoryResolutionFractionSquare; 
-		//PassParameters->ResponsiveStencilMask = CVarTSREnableResponiveAA.GetValueOnRenderThread() ? (STENCIL_TEMPORAL_RESPONSIVE_AA_MASK) : 0;
-		PassParameters->ResponsiveStencilMask = 0;
 		PassParameters->bGenerateOutputMip1 = false;
 		PassParameters->bGenerateOutputMip2 = false;
 		PassParameters->bHasSeparateTranslucency = bHasSeparateTranslucency;
@@ -2546,9 +2500,7 @@ FDefaultTemporalUpscaler::FOutputs AddTemporalSuperResolutionPasses(
 			PassParameters->SceneColorTexture = SceneColorOutputTextureSRV;
 			PassParameters->ClosestDepthTexture = ClosestDepthTexture;
 			PassParameters->DilatedVelocityTexture = DilatedVelocityTexture;
-			PassParameters->ParallaxRejectionMaskTexture = ParallaxRejectionMaskTexture;
 			PassParameters->HistoryRejectionTexture = HistoryRejectionTexture;
-			PassParameters->HistoryResurrectionMaskTexture = bCanResurrectHistory ? HistoryResurrectionMaskTexture : BlackDummy;
 			PassParameters->HistoryMetadataTexture = GraphBuilder.CreateSRV(FRDGTextureSRVDesc::CreateForSlice(History.MetadataArray, CurrentFrameSliceIndex));
 			if (PrevHistory.ColorArray == BlackArrayDummy)
 			{
