Disable cache mesh draw command task when shaders can't be created off the render thread.

diff --git a/Engine/Source/Runtime/Renderer/Private/RendererScene.cpp b/Engine/Source/Runtime/Renderer/Private/RendererScene.cpp
--- a/Engine/Source/Runtime/Renderer/Private/RendererScene.cpp
+++ b/Engine/Source/Runtime/Renderer/Private/RendererScene.cpp
@@ -6244,7 +6244,7 @@ void FScene::UpdateAllPrimitiveSceneInfos(FRDGBuilder& GraphBuilder, EUpdateAllP
 
 	if (SceneInfosWithStaticDrawListUpdate.Num() > 0)
 	{
-		const bool bLaunchAsyncTask = EnumHasAnyFlags(AsyncOps, EUpdateAllPrimitiveSceneInfosAsyncOps::CacheMeshDrawCommands);
+		const bool bLaunchAsyncTask = EnumHasAnyFlags(AsyncOps, EUpdateAllPrimitiveSceneInfosAsyncOps::CacheMeshDrawCommands) && GRHISupportsMultithreadedShaderCreation;
 
 		CacheMeshDrawCommandsTask = LaunchSceneRenderTask(TEXT("CacheMeshDrawCommandsTask"), [this, &SceneInfosWithStaticDrawListUpdate, bLaunchAsyncTask]()
 		{
