Fixed dependency issue in mobile renderer with distance fields.

diff --git a/Engine/Source/Runtime/Renderer/Private/MobileShadingRenderer.cpp b/Engine/Source/Runtime/Renderer/Private/MobileShadingRenderer.cpp
--- a/Engine/Source/Runtime/Renderer/Private/MobileShadingRenderer.cpp
+++ b/Engine/Source/Runtime/Renderer/Private/MobileShadingRenderer.cpp
@@ -509,18 +509,6 @@ void FMobileSceneRenderer::InitViews(
 
 	bShouldRenderVelocities = ShouldRenderVelocities();
 
-	static auto CVarDistanceFieldShadowQuality = IConsoleManager::Get().FindConsoleVariable(TEXT("r.DFShadowQuality"));
-
-	bRequiresDistanceField = IsMobileDistanceFieldEnabled(ShaderPlatform)
-		&& ViewFamily.EngineShowFlags.Lighting
-		&& !Views[0].bIsReflectionCapture
-		&& !Views[0].bIsPlanarReflection
-		&& !ViewFamily.EngineShowFlags.HitProxies
-		&& !ViewFamily.EngineShowFlags.VisualizeLightCulling
-		&& !ViewFamily.UseDebugViewPS()
-		&& (CVarDistanceFieldShadowQuality != nullptr && CVarDistanceFieldShadowQuality->GetInt() > 0)
-		&& bRendererOutputFinalSceneColor;
-
 	bRequiresShadowProjections = MobileUsesShadowMaskTexture(ShaderPlatform)
 		&& ViewFamily.EngineShowFlags.Lighting
 		&& !Views[0].bIsReflectionCapture
@@ -820,9 +808,20 @@ void FMobileSceneRenderer::Render(FRDGBuilder& GraphBuilder)
 
 	FRDGExternalAccessQueue ExternalAccessQueue;
 
-	if (bRequiresDistanceField && bRendererOutputFinalSceneColor)
 	{
-		PrepareDistanceFieldScene(GraphBuilder, ExternalAccessQueue);
+		static auto CVarDistanceFieldShadowQuality = IConsoleManager::Get().FindConsoleVariable(TEXT("r.DFShadowQuality"));
+
+		if (IsMobileDistanceFieldEnabled(ShaderPlatform)
+			&& ViewFamily.EngineShowFlags.Lighting
+			&& !Views[0].bIsReflectionCapture
+			&& !Views[0].bIsPlanarReflection
+			&& !ViewFamily.EngineShowFlags.VisualizeLightCulling
+			&& !ViewFamily.UseDebugViewPS()
+			&& (CVarDistanceFieldShadowQuality != nullptr && CVarDistanceFieldShadowQuality->GetInt() > 0)
+			&& bRendererOutputFinalSceneColor)
+		{
+			PrepareDistanceFieldScene(GraphBuilder, ExternalAccessQueue);
+		}
 	}
 
 	ExternalAccessQueue.Submit(GraphBuilder);
diff --git a/Engine/Source/Runtime/Renderer/Private/SceneRendering.h b/Engine/Source/Runtime/Renderer/Private/SceneRendering.h
--- a/Engine/Source/Runtime/Renderer/Private/SceneRendering.h
+++ b/Engine/Source/Runtime/Renderer/Private/SceneRendering.h
@@ -2723,7 +2723,6 @@ private:
 	bool bRequiresAmbientOcclusionPass;
 	bool bShouldRenderVelocities;
 	bool bShouldRenderHZB;
-	bool bRequiresDistanceField;
 	bool bIsFullDepthPrepassEnabled;
 	bool bIsMaskedOnlyDepthPrepassEnabled;
 	bool bRequiresSceneDepthAux;
