Fix warning C6385: Reading invalid data from 'this->Samples'.
#horde 339718
#rnx

diff --git a/Engine/Source/Runtime/Renderer/Private/ScenePrivate.h b/Engine/Source/Runtime/Renderer/Private/ScenePrivate.h
--- a/Engine/Source/Runtime/Renderer/Private/ScenePrivate.h
+++ b/Engine/Source/Runtime/Renderer/Private/ScenePrivate.h
@@ -530,7 +530,7 @@ public:
 	/** @return A random number between 0 and 1. */
 	inline float GetFraction()
 	{
-		if (CurrentSample > NumSamples)
+		if (CurrentSample >= NumSamples)
 		{
 			CurrentSample = 0;
 		}
