Remove the (disabled & unused) InstanceBVH code from Scene/GPUScene.

diff --git a/Engine/Source/Runtime/Renderer/Private/GPUScene.cpp b/Engine/Source/Runtime/Renderer/Private/GPUScene.cpp
--- a/Engine/Source/Runtime/Renderer/Private/GPUScene.cpp
+++ b/Engine/Source/Runtime/Renderer/Private/GPUScene.cpp
@@ -71,14 +71,6 @@ static TAutoConsoleVariable<int32> CVarGPUSceneParallelUpdate(
 	ECVF_RenderThreadSafe
 );
 
-int32 GGPUSceneInstanceBVH = 0;
-FAutoConsoleVariableRef CVarGPUSceneInstanceBVH(
-	TEXT("r.GPUScene.InstanceBVH"),
-	GGPUSceneInstanceBVH,
-	TEXT("Add instances to BVH. (WIP)"),
-	ECVF_RenderThreadSafe | ECVF_ReadOnly
-);
-
 static TAutoConsoleVariable<int32> CVarGPUSceneDebugMode(
 	TEXT("r.GPUScene.DebugMode"),
 	0,
@@ -832,9 +824,6 @@ void FGPUScene::UpdateBufferState(FRDGBuilder& GraphBuilder, FSceneUniformBuffer
 	const uint32 InstancePayloadDataSizeReserve = FMath::RoundUpToPowerOfTwo(PayloadFloat4Count * sizeof(FVector4f));
 	BufferState.InstancePayloadDataBuffer = ResizeStructuredBufferIfNeeded(GraphBuilder, InstancePayloadDataBuffer, InstancePayloadDataSizeReserve, TEXT("GPUScene.InstancePayloadData"));
 
-	const uint32 NumNodes = FMath::RoundUpToPowerOfTwo(FMath::Max(Scene.InstanceBVH.GetNumNodes(), InitialBufferSize));
-	BufferState.InstanceBVHBuffer = ResizeStructuredBufferIfNeeded(GraphBuilder, InstanceBVHBuffer, NumNodes * sizeof(FBVHNode), TEXT("InstanceBVH"));
-
 	const bool bNaniteEnabled = DoesPlatformSupportNanite(GMaxRHIShaderPlatform);
 	if (UploadDataSourceAdapter.bUpdateNaniteMaterialTables && bNaniteEnabled)
 	{
@@ -975,7 +964,6 @@ void FGPUScene::UploadGeneral(FRDGBuilder& GraphBuilder, FScene& Scene, FRDGExte
 		FRDGScatterUploader* PrimitiveUploader = nullptr;
 		FRDGScatterUploader* InstancePayloadUploader = nullptr;
 		FRDGScatterUploader* InstanceSceneUploader = nullptr;
-		FRDGScatterUploader* InstanceBVHUploader = nullptr;
 		FRDGScatterUploader* LightmapUploader = nullptr;
 
 		TStaticArray<FNaniteMaterialCommands::FUploader*, ENaniteMeshPass::Num> NaniteMaterialUploaders{ InPlace, nullptr };
@@ -1018,11 +1006,6 @@ void FGPUScene::UploadGeneral(FRDGBuilder& GraphBuilder, FScene& Scene, FRDGExte
 		TaskContext.InstanceSceneUploader = InstanceSceneUploadBuffer.BeginPreSized(GraphBuilder, BufferState.InstanceSceneDataBuffer, TaskContext.NumInstanceSceneDataUploads * FInstanceSceneShaderData::GetDataStrideInFloat4s(), sizeof(FVector4f), TEXT("InstanceSceneUploadBuffer"));
 	}
 
-	if (Scene.InstanceBVH.GetNumDirty() > 0)
-	{
-		TaskContext.InstanceBVHUploader = InstanceBVHUploadBuffer.Begin(GraphBuilder, BufferState.InstanceBVHBuffer, Scene.InstanceBVH.GetNumDirty(), sizeof(FBVHNode), TEXT("InstanceSceneUploadBuffer"));
-	}
-
 	if (TaskContext.NumLightmapDataUploads > 0)
 	{
 		TaskContext.LightmapUploader = LightmapUploadBuffer.Begin(GraphBuilder, BufferState.LightmapDataBuffer, TaskContext.NumLightmapDataUploads, sizeof(FLightmapSceneShaderData::Data), TEXT("LightmapUploadBuffer"));
@@ -1045,7 +1028,6 @@ void FGPUScene::UploadGeneral(FRDGBuilder& GraphBuilder, FScene& Scene, FRDGExte
 		LockIfValid(RHICmdList, TaskContext.PrimitiveUploader);
 		LockIfValid(RHICmdList, TaskContext.InstancePayloadUploader);
 		LockIfValid(RHICmdList, TaskContext.InstanceSceneUploader);
-		LockIfValid(RHICmdList, TaskContext.InstanceBVHUploader);
 		LockIfValid(RHICmdList, TaskContext.LightmapUploader);
 
 		for (FNaniteMaterialCommands::FUploader* Uploader : TaskContext.NaniteMaterialUploaders)
@@ -1281,31 +1263,6 @@ void FGPUScene::UploadGeneral(FRDGBuilder& GraphBuilder, FScene& Scene, FRDGExte
 			}, bExecuteInParallel ? EParallelForFlags::None : EParallelForFlags::ForceSingleThread);
 		}
 
-		if (TaskContext.InstanceBVHUploader)
-		{
-			SCOPED_NAMED_EVENT(InstanceBVH, FColor::Green);
-
-			Scene.InstanceBVH.ForAllDirty(
-				[&](uint32 NodeIndex, const auto& Node)
-			{
-				FBVHNode GPUNode;
-				for (int i = 0; i < 4; i++)
-				{
-					GPUNode.ChildIndexes[i] = Node.ChildIndexes[i];
-
-					GPUNode.ChildMin[0][i] = Node.ChildBounds[i].Min.X;
-					GPUNode.ChildMin[1][i] = Node.ChildBounds[i].Min.Y;
-					GPUNode.ChildMin[2][i] = Node.ChildBounds[i].Min.Z;
-
-					GPUNode.ChildMax[0][i] = Node.ChildBounds[i].Max.X;
-					GPUNode.ChildMax[1][i] = Node.ChildBounds[i].Max.Y;
-					GPUNode.ChildMax[2][i] = Node.ChildBounds[i].Max.Z;
-				}
-
-				TaskContext.InstanceBVHUploader->Add(NodeIndex, &GPUNode);
-			});
-		}
-
 		if (TaskContext.LightmapUploader)
 		{
 			for (int32 ItemIndex = 0; ItemIndex < TaskContext.NumPrimitiveDataUploads; ++ItemIndex)
@@ -1325,7 +1282,6 @@ void FGPUScene::UploadGeneral(FRDGBuilder& GraphBuilder, FScene& Scene, FRDGExte
 		UnlockIfValid(RHICmdList, TaskContext.PrimitiveUploader);
 		UnlockIfValid(RHICmdList, TaskContext.InstancePayloadUploader);
 		UnlockIfValid(RHICmdList, TaskContext.InstanceSceneUploader);
-		UnlockIfValid(RHICmdList, TaskContext.InstanceBVHUploader);
 		UnlockIfValid(RHICmdList, TaskContext.LightmapUploader);
 
 		for (FNaniteMaterialCommands::FUploader* Uploader : TaskContext.NaniteMaterialUploaders)
@@ -1346,11 +1302,6 @@ void FGPUScene::UploadGeneral(FRDGBuilder& GraphBuilder, FScene& Scene, FRDGExte
 		InstanceSceneUploadBuffer.End(GraphBuilder, TaskContext.InstanceSceneUploader);
 	}
 
-	if (TaskContext.InstanceBVHUploader)
-	{
-		InstanceBVHUploadBuffer.End(GraphBuilder, TaskContext.InstanceBVHUploader);
-	}
-
 	if (TaskContext.LightmapUploader)
 	{
 		LightmapUploadBuffer.End(GraphBuilder, TaskContext.LightmapUploader);
@@ -1379,11 +1330,6 @@ void FGPUScene::UploadGeneral(FRDGBuilder& GraphBuilder, FScene& Scene, FRDGExte
 		InstancePayloadUploadBuffer.Release();
 	}
 
-	if (InstanceBVHUploadBuffer.GetNumBytes() > MaxPooledSize)
-	{
-		InstanceBVHUploadBuffer.Release();
-	}
-
 	if (LightmapUploadBuffer.GetNumBytes() > MaxPooledSize)
 	{
 		LightmapUploadBuffer.Release();
diff --git a/Engine/Source/Runtime/Renderer/Private/GPUScene.h b/Engine/Source/Runtime/Renderer/Private/GPUScene.h
--- a/Engine/Source/Runtime/Renderer/Private/GPUScene.h
+++ b/Engine/Source/Runtime/Renderer/Private/GPUScene.h
@@ -177,7 +177,6 @@ struct FGPUSceneBufferState
 	FRDGBuffer* InstanceSceneDataBuffer = nullptr;
 	uint32 InstanceSceneDataSOAStride = 1; // Distance between arrays in float4s
 	FRDGBuffer* InstancePayloadDataBuffer = nullptr;
-	FRDGBuffer* InstanceBVHBuffer = nullptr;
 	FRDGBuffer* LightmapDataBuffer = nullptr;
 	uint32 LightMapDataBufferSize = 0;
 
@@ -352,9 +351,6 @@ public:
 	TRefCountPtr<FRDGPooledBuffer> InstancePayloadDataBuffer;
 	FRDGAsyncScatterUploadBuffer   InstancePayloadUploadBuffer;
 
-	TRefCountPtr<FRDGPooledBuffer> InstanceBVHBuffer;
-	FRDGAsyncScatterUploadBuffer   InstanceBVHUploadBuffer;
-
 	/** GPU light map data */
 	FSpanAllocator                 LightmapDataAllocator;
 	TRefCountPtr<FRDGPooledBuffer> LightmapDataBuffer;
diff --git a/Engine/Source/Runtime/Renderer/Private/PrimitiveSceneInfo.cpp b/Engine/Source/Runtime/Renderer/Private/PrimitiveSceneInfo.cpp
--- a/Engine/Source/Runtime/Renderer/Private/PrimitiveSceneInfo.cpp
+++ b/Engine/Source/Runtime/Renderer/Private/PrimitiveSceneInfo.cpp
@@ -33,7 +33,6 @@
 #include "StaticMeshBatch.h"
 
 extern int32 GGPUSceneInstanceClearList;
-extern int32 GGPUSceneInstanceBVH;
 
 static int32 GMeshDrawCommandsCacheMultithreaded = 1;
 static FAutoConsoleVariableRef CVarDrawCommandsCacheMultithreaded(
@@ -1272,18 +1271,6 @@ void FPrimitiveSceneInfo::AllocateGPUSceneInstances(FScene* Scene, const TArrayV
 						const uint32 TotalFloat4Count = SceneInfo->NumInstanceSceneDataEntries * SceneInfo->InstancePayloadDataStride;
 						SceneInfo->InstancePayloadDataOffset = Scene->GPUScene.AllocateInstancePayloadDataSlots(TotalFloat4Count);
 					}
-
-					if (GGPUSceneInstanceBVH)
-					{
-						// TODO: Replace Instance BVH FBounds with FRenderBounds
-						for (int32 InstanceIndex = 0; InstanceIndex < SceneInfo->NumInstanceSceneDataEntries; ++InstanceIndex)
-						{
-							const FInstanceSceneData& PrimitiveInstance = InstanceSceneData[InstanceIndex];
-							FRenderBounds WorldBounds = SceneInfo->Proxy->GetInstanceLocalBounds(InstanceIndex);
-							WorldBounds.TransformBy(PrimitiveInstance.ComputeLocalToWorld(SceneInfo->Proxy->GetLocalToWorld()));
-							Scene->InstanceBVH.Add(FBounds3f({ WorldBounds.GetMin(), WorldBounds.GetMax() }), SceneInfo->InstanceSceneDataOffset + InstanceIndex);
-						}
-					}
 				}
 			}
 			else
@@ -1339,13 +1326,6 @@ void FPrimitiveSceneInfo::FreeGPUSceneInstances()
 		SCOPE_CYCLE_COUNTER(STAT_UpdateGPUSceneTime);
 
 		check(Proxy->SupportsInstanceDataBuffer() || NumInstanceSceneDataEntries == 1);
-		if (GGPUSceneInstanceBVH)
-		{
-			for (int32 InstanceIndex = 0; InstanceIndex < NumInstanceSceneDataEntries; InstanceIndex++)
-			{
-				Scene->InstanceBVH.Remove(InstanceSceneDataOffset + InstanceIndex);
-			}
-		}
 
 		// Release all instance payload data slots associated with this primitive.
 		if (InstancePayloadDataOffset != INDEX_NONE)
diff --git a/Engine/Source/Runtime/Renderer/Private/ScenePrivate.h b/Engine/Source/Runtime/Renderer/Private/ScenePrivate.h
--- a/Engine/Source/Runtime/Renderer/Private/ScenePrivate.h
+++ b/Engine/Source/Runtime/Renderer/Private/ScenePrivate.h
@@ -3111,8 +3111,6 @@ public:
 	/** An octree containing the primitives in the scene. */
 	FScenePrimitiveOctree PrimitiveOctree;
 
-	FDynamicBVH<4> InstanceBVH;
-
 	/** Indicates whether this scene requires hit proxy rendering. */
 	bool bRequiresHitProxies;
 
