Fixed: SkyAtmosphere appears pixelized on the PowerVR GPU devices

diff --git a/Engine/Source/Runtime/Renderer/Private/SkyAtmosphereRendering.cpp b/Engine/Source/Runtime/Renderer/Private/SkyAtmosphereRendering.cpp
--- a/Engine/Source/Runtime/Renderer/Private/SkyAtmosphereRendering.cpp
+++ b/Engine/Source/Runtime/Renderer/Private/SkyAtmosphereRendering.cpp
@@ -1113,20 +1113,11 @@ static EPixelFormat GetSkyLutTextureFormat(ERHIFeatureLevel::Type FeatureLevel)
 		// TODO: check if need this for Metal, Vulkan
 		TextureLUTFormat = PF_FloatRGBA;
 	}
-	
-	if (CVarSkyAtmosphereLUT32.GetValueOnAnyThread() != 0)
-	{
-		TextureLUTFormat = PF_A32B32G32R32F;
-	}
 
 	return TextureLUTFormat;
 }
 static EPixelFormat GetSkyLutSmallTextureFormat()
 {
-	if (CVarSkyAtmosphereLUT32.GetValueOnAnyThread() != 0)
-	{
-		return PF_A32B32G32R32F;
-	}
 	return PF_R8G8B8A8;
 }
 
