Moved Local Exposure parameters from FEyeAdaptationParameters into a separate FLocalExposureParameters.

diff --git a/Engine/Shaders/Private/PostProcessBloom.usf b/Engine/Shaders/Private/PostProcessBloom.usf
--- a/Engine/Shaders/Private/PostProcessBloom.usf
+++ b/Engine/Shaders/Private/PostProcessBloom.usf
@@ -29,7 +29,7 @@ float2 GetExposureScaleMiddleGreyLumValue()
 	float2 ExposureScaleMiddleGreyLumValue = EyeAdaptationLookupBuffer(EyeAdaptationBuffer).xw;
 
 	// Middle grey lum value adjusted by exposure compensation
-	ExposureScaleMiddleGreyLumValue.y = log2(0.18f * ExposureScaleMiddleGreyLumValue.y * EyeAdaptation_LocalExposureMiddleGreyExposureCompensation);
+	ExposureScaleMiddleGreyLumValue.y = log2(0.18f * ExposureScaleMiddleGreyLumValue.y * LocalExposure_MiddleGreyExposureCompensation);
 
 	return ExposureScaleMiddleGreyLumValue;
 }
@@ -46,8 +46,8 @@ float4 BloomSetupCommon(float2 UV, float2 ViewportUV, float2 ExposureScaleMiddle
 		float LogLuminance = log2(LuminanceVal);
 		float MiddleGreyLumValue = ExposureScaleMiddleGreyLumValue.y;
 
-		float BaseLogLum = CalculateBaseLogLuminance(LogLuminance, EyeAdaptation_LocalExposureBlurredLuminanceBlend, ExposureScale, ViewportUV, LumBilateralGrid, BlurredLogLum, LumBilateralGridSampler, BlurredLogLumSampler);
-		float LocalExposure = CalculateLocalExposure(LogLuminance + log2(ExposureScale), BaseLogLum, MiddleGreyLumValue, EyeAdaptation_LocalExposureHighlightContrastScale, EyeAdaptation_LocalExposureShadowContrastScale, EyeAdaptation_LocalExposureDetailStrength);
+		float BaseLogLum = CalculateBaseLogLuminance(LogLuminance, LocalExposure_BlurredLuminanceBlend, ExposureScale, ViewportUV, LumBilateralGrid, BlurredLogLum, LumBilateralGridSampler, BlurredLogLumSampler);
+		float LocalExposure = CalculateLocalExposure(LogLuminance + log2(ExposureScale), BaseLogLum, MiddleGreyLumValue, LocalExposure_HighlightContrastScale, LocalExposure_ShadowContrastScale, LocalExposure_DetailStrength);
 		
 		SceneColor.rgb *= LocalExposure;
 	}
diff --git a/Engine/Shaders/Private/PostProcessEyeAdaptation.usf b/Engine/Shaders/Private/PostProcessEyeAdaptation.usf
--- a/Engine/Shaders/Private/PostProcessEyeAdaptation.usf
+++ b/Engine/Shaders/Private/PostProcessEyeAdaptation.usf
@@ -191,7 +191,7 @@ float ComputeAverageLocalExposure(Texture2D HistogramTexture, float ExposureScal
 	// Average Local Exposure is calculated using log(LocalExposure) 
 	// since it results in a more consistent value when applied to Pre-Exposure
 
-	const float MiddleGreyLumValue = log2(0.18 * MiddleGreyExposureCompensation * EyeAdaptation_LocalExposureMiddleGreyExposureCompensation);
+	const float MiddleGreyLumValue = log2(0.18 * MiddleGreyExposureCompensation * LocalExposure_MiddleGreyExposureCompensation);
 
 	float2 Sum = 0;
 
@@ -202,7 +202,7 @@ float ComputeAverageLocalExposure(Texture2D HistogramTexture, float ExposureScal
 		float LogLuminanceAtBucket = ComputeLogLuminanceFromHistogramPosition(float(i) / (float)(HISTOGRAM_SIZE - 1));
 
 		float BaseLogLum = LogLuminanceAtBucket + log2(ExposureScale);
-		float LogLocalExposure = CalculateLogLocalExposure(BaseLogLum, BaseLogLum, MiddleGreyLumValue, EyeAdaptation_LocalExposureHighlightContrastScale, EyeAdaptation_LocalExposureShadowContrastScale, 0.0f);
+		float LogLocalExposure = CalculateLogLocalExposure(BaseLogLum, BaseLogLum, MiddleGreyLumValue, LocalExposure_HighlightContrastScale, LocalExposure_ShadowContrastScale, 0.0f);
 
 		Sum += float2(LogLocalExposure, 1) * LocalValue;
 	}
@@ -330,7 +330,7 @@ void BasicEyeAdaptationCS(uint GIndex : SV_GroupIndex, uint2 GTId : SV_GroupThre
 	const float4 OutData0 = EyeAdaptationCommon(AverageSceneLuminance, OldExposureScale);
 
 #if COMPUTE_AVERAGE_LOCAL_EXPOSURE
-	const float MiddleGreyLumValue = log2(0.18 * OutData0.w * EyeAdaptation_LocalExposureMiddleGreyExposureCompensation);
+	const float MiddleGreyLumValue = log2(0.18 * OutData0.w * LocalExposure_MiddleGreyExposureCompensation);
 
 	const float2 LogLumSumWithCount = ComputeTextureAverageAlphaSubRegion(ColorTexture, SubRectMin, SubRectMax, Color_ViewportMin, Color_ViewportMax);
 
@@ -338,7 +338,7 @@ void BasicEyeAdaptationCS(uint GIndex : SV_GroupIndex, uint2 GTId : SV_GroupThre
 	LogLum = (LogLum - EyeAdaptation_HistogramBias) / EyeAdaptation_HistogramScale;
 
 	const float BaseLogLum = LogLum + log2(OutData0.x);
-	const float LogLocalExposure = CalculateLogLocalExposure(BaseLogLum, BaseLogLum, MiddleGreyLumValue, EyeAdaptation_LocalExposureHighlightContrastScale, EyeAdaptation_LocalExposureShadowContrastScale, 0.0f);
+	const float LogLocalExposure = CalculateLogLocalExposure(BaseLogLum, BaseLogLum, MiddleGreyLumValue, LocalExposure_HighlightContrastScale, LocalExposure_ShadowContrastScale, 0.0f);
 
 	const float2 LogLocalExposureSumWithCount = float2(LogLocalExposure, 1.0f) * LogLumSumWithCount.y;
 
diff --git a/Engine/Shaders/Private/PostProcessHistogramCommon.ush b/Engine/Shaders/Private/PostProcessHistogramCommon.ush
--- a/Engine/Shaders/Private/PostProcessHistogramCommon.ush
+++ b/Engine/Shaders/Private/PostProcessHistogramCommon.ush
@@ -16,11 +16,6 @@ float EyeAdaptation_ExposureSpeedDown;
 float EyeAdaptation_HistogramScale;
 float EyeAdaptation_HistogramBias;
 float EyeAdaptation_LuminanceMin;
-float EyeAdaptation_LocalExposureHighlightContrastScale;
-float EyeAdaptation_LocalExposureShadowContrastScale;
-float EyeAdaptation_LocalExposureDetailStrength;
-float EyeAdaptation_LocalExposureBlurredLuminanceBlend;
-float EyeAdaptation_LocalExposureMiddleGreyExposureCompensation;
 float EyeAdaptation_BlackHistogramBucketInfluence;
 float EyeAdaptation_IgnoreMaterialsEvaluationPositionBias;
 float EyeAdaptation_IgnoreMaterialsLuminanceScale;
@@ -36,6 +31,12 @@ int EyeAdaptation_VisualizeDebugType;
 Texture2D EyeAdaptation_MeterMaskTexture;
 SamplerState EyeAdaptation_MeterMaskSampler;
 
+float LocalExposure_HighlightContrastScale;
+float LocalExposure_ShadowContrastScale;
+float LocalExposure_DetailStrength;
+float LocalExposure_BlurredLuminanceBlend;
+float LocalExposure_MiddleGreyExposureCompensation;
+
 float CalculateEyeAdaptationLuminance(float3 Color)
 {
 	return max(dot(Color, float3(1.0f, 1.0f, 1.0f) / 3.0f), EyeAdaptation_LuminanceMin);
diff --git a/Engine/Shaders/Private/PostProcessLocalExposure.usf b/Engine/Shaders/Private/PostProcessLocalExposure.usf
--- a/Engine/Shaders/Private/PostProcessLocalExposure.usf
+++ b/Engine/Shaders/Private/PostProcessLocalExposure.usf
@@ -38,15 +38,15 @@ void ApplyLocalExposureCS(uint2 DispatchThreadId : SV_DispatchThreadID)
 
 		float2 ExposureScaleMiddleGrey = EyeAdaptationBuffer[0].xw;
 		float ExposureScale = ExposureScaleMiddleGrey.x;
-		float MiddleGreyLumValue = log2(0.18 * ExposureScaleMiddleGrey.y * EyeAdaptation_LocalExposureMiddleGreyExposureCompensation);
+		float MiddleGreyLumValue = log2(0.18 * ExposureScaleMiddleGrey.y * LocalExposure_MiddleGreyExposureCompensation);
 
 		float4 SceneColor = InputTexture.Load(uint3(DispatchThreadId + Input_ViewportMin, 0));
 
 		float LuminanceVal = CalculateEyeAdaptationLuminance(SceneColor.rgb * View.OneOverPreExposure);
 		float LogLuminance = log2(LuminanceVal);
 
-		float BaseLogLum = CalculateBaseLogLuminance(LogLuminance, EyeAdaptation_LocalExposureBlurredLuminanceBlend, ExposureScale, UV, LumBilateralGrid, BlurredLogLum, TextureSampler, TextureSampler);
-		float LocalExposure = CalculateLocalExposure(LogLuminance + log2(ExposureScale), BaseLogLum, MiddleGreyLumValue, EyeAdaptation_LocalExposureHighlightContrastScale, EyeAdaptation_LocalExposureShadowContrastScale, EyeAdaptation_LocalExposureDetailStrength);
+		float BaseLogLum = CalculateBaseLogLuminance(LogLuminance, LocalExposure_BlurredLuminanceBlend, ExposureScale, UV, LumBilateralGrid, BlurredLogLum, TextureSampler, TextureSampler);
+		float LocalExposure = CalculateLocalExposure(LogLuminance + log2(ExposureScale), BaseLogLum, MiddleGreyLumValue, LocalExposure_HighlightContrastScale, LocalExposure_ShadowContrastScale, LocalExposure_DetailStrength);
 
 		SceneColor.rgb *= LocalExposure;
 
diff --git a/Engine/Shaders/Private/PostProcessTonemap.usf b/Engine/Shaders/Private/PostProcessTonemap.usf
--- a/Engine/Shaders/Private/PostProcessTonemap.usf
+++ b/Engine/Shaders/Private/PostProcessTonemap.usf
@@ -381,10 +381,10 @@ float4 TonemapCommonPS(
 		float LogLuminance = log2(LuminanceVal);
 
 		// Middle grey lum value adjusted by exposure compensation
-		float MiddleGreyLumValue = log2(0.18 * EyeAdaptationData.y * EyeAdaptation_LocalExposureMiddleGreyExposureCompensation);
+		float MiddleGreyLumValue = log2(0.18 * EyeAdaptationData.y * LocalExposure_MiddleGreyExposureCompensation);
 
-		float BaseLogLum = CalculateBaseLogLuminance(LogLuminance, EyeAdaptation_LocalExposureBlurredLuminanceBlend, GlobalExposure, FullViewUV, LumBilateralGrid, BlurredLogLum, LumBilateralGridSampler, BlurredLogLumSampler);
-		LocalExposure = CalculateLocalExposure(LogLuminance + log2(GlobalExposure), BaseLogLum, MiddleGreyLumValue, EyeAdaptation_LocalExposureHighlightContrastScale, EyeAdaptation_LocalExposureShadowContrastScale, EyeAdaptation_LocalExposureDetailStrength);
+		float BaseLogLum = CalculateBaseLogLuminance(LogLuminance, LocalExposure_BlurredLuminanceBlend, GlobalExposure, FullViewUV, LumBilateralGrid, BlurredLogLum, LumBilateralGridSampler, BlurredLogLumSampler);
+		LocalExposure = CalculateLocalExposure(LogLuminance + log2(GlobalExposure), BaseLogLum, MiddleGreyLumValue, LocalExposure_HighlightContrastScale, LocalExposure_ShadowContrastScale, LocalExposure_DetailStrength);
 	}
 	#else // !USE_LOCAL_EXPOSURE
 	const float LocalExposure = 1.0;
diff --git a/Engine/Shaders/Private/PostProcessVisualizeLocalExposure.usf b/Engine/Shaders/Private/PostProcessVisualizeLocalExposure.usf
--- a/Engine/Shaders/Private/PostProcessVisualizeLocalExposure.usf
+++ b/Engine/Shaders/Private/PostProcessVisualizeLocalExposure.usf
@@ -32,10 +32,10 @@ float4 MainPS(noperspective float4 UVAndScreenPos : TEXCOORD0, float4 SvPosition
 
 	float LuminanceVal = CalculateEyeAdaptationLuminance(SceneColor.rgb);
 	float LogLuminance = log2(LuminanceVal);
-	float MiddleGreyLumValue = log2(0.18 * ExposureScaleMiddleGrey.y * EyeAdaptation_LocalExposureMiddleGreyExposureCompensation);
+	float MiddleGreyLumValue = log2(0.18 * ExposureScaleMiddleGrey.y * LocalExposure_MiddleGreyExposureCompensation);
 
-	float BaseLogLum = CalculateBaseLogLuminance(LogLuminance, EyeAdaptation_LocalExposureBlurredLuminanceBlend, ExposureScaleMiddleGrey.x, ViewportUV, LumBilateralGrid, BlurredLogLum, TextureSampler, TextureSampler);
-	float LocalExposure = CalculateLocalExposure(LogLuminance + log2(ExposureScaleMiddleGrey.x), BaseLogLum, MiddleGreyLumValue, EyeAdaptation_LocalExposureHighlightContrastScale, EyeAdaptation_LocalExposureShadowContrastScale, EyeAdaptation_LocalExposureDetailStrength);
+	float BaseLogLum = CalculateBaseLogLuminance(LogLuminance, LocalExposure_BlurredLuminanceBlend, ExposureScaleMiddleGrey.x, ViewportUV, LumBilateralGrid, BlurredLogLum, TextureSampler, TextureSampler);
+	float LocalExposure = CalculateLocalExposure(LogLuminance + log2(ExposureScaleMiddleGrey.x), BaseLogLum, MiddleGreyLumValue, LocalExposure_HighlightContrastScale, LocalExposure_ShadowContrastScale, LocalExposure_DetailStrength);
 	float DetailLogLum = (LogLuminance + log2(ExposureScaleMiddleGrey.x)) - BaseLogLum;
 
 	float3 OutValue;
diff --git a/Engine/Source/Runtime/Renderer/Private/PostProcess/PostProcessBloomSetup.cpp b/Engine/Source/Runtime/Renderer/Private/PostProcess/PostProcessBloomSetup.cpp
--- a/Engine/Source/Runtime/Renderer/Private/PostProcess/PostProcessBloomSetup.cpp
+++ b/Engine/Source/Runtime/Renderer/Private/PostProcess/PostProcessBloomSetup.cpp
@@ -5,6 +5,7 @@
 #include "PostProcess/PostProcessFFTBloom.h"
 #include "PostProcess/PostProcessWeightedSampleSum.h"
 #include "PostProcess/PostProcessEyeAdaptation.h"
+#include "PostProcess/PostProcessLocalExposure.h"
 #include "PixelShaderUtils.h"
 #include "DataDrivenShaderPlatformInfo.h"
 #include "SceneRendering.h"
@@ -33,6 +34,7 @@ BEGIN_SHADER_PARAMETER_STRUCT(FBloomSetupParameters, )
 	SHADER_PARAMETER_SAMPLER(SamplerState, LumBilateralGridSampler)
 	SHADER_PARAMETER_RDG_TEXTURE(Texture2D, BlurredLogLum)
 	SHADER_PARAMETER_SAMPLER(SamplerState, BlurredLogLumSampler)
+	SHADER_PARAMETER_STRUCT(FLocalExposureParameters, LocalExposure)
 	SHADER_PARAMETER_RDG_BUFFER_SRV(StructuredBuffer<float4>, EyeAdaptationBuffer)
 	SHADER_PARAMETER_STRUCT(FEyeAdaptationParameters, EyeAdaptation)
 	SHADER_PARAMETER(float, BloomThreshold)
@@ -53,6 +55,7 @@ FBloomSetupParameters GetBloomSetupParameters(
 	Parameters.LumBilateralGridSampler = TStaticSamplerState<SF_Bilinear, AM_Clamp, AM_Clamp, AM_Clamp>::GetRHI();
 	Parameters.BlurredLogLum = Inputs.BlurredLogLuminanceTexture;
 	Parameters.BlurredLogLumSampler = TStaticSamplerState<SF_Bilinear, AM_Clamp, AM_Clamp, AM_Clamp>::GetRHI();
+	Parameters.LocalExposure = *Inputs.LocalExposureParameters;
 	Parameters.EyeAdaptationBuffer = GraphBuilder.CreateSRV(Inputs.EyeAdaptationBuffer);
 	Parameters.EyeAdaptation = *Inputs.EyeAdaptationParameters;
 	Parameters.BloomThreshold = Inputs.Threshold;
@@ -119,6 +122,7 @@ FScreenPassTexture AddBloomSetupPass(FRDGBuilder& GraphBuilder, const FViewInfo&
 	check(Inputs.Threshold > -1.0f || Inputs.EyeAdaptationParameters != nullptr);
 
 	const bool bIsComputePass = View.bUseComputePasses;
+	const bool bLocalExposureEnabled = Inputs.LocalExposureTexture != nullptr;
 
 	FRDGTextureDesc OutputDesc = Inputs.SceneColor.Texture->Desc;
 	OutputDesc.Reset();
@@ -134,7 +138,7 @@ FScreenPassTexture AddBloomSetupPass(FRDGBuilder& GraphBuilder, const FViewInfo&
 		PassParameters->RWOutputTexture = GraphBuilder.CreateUAV(Output.Texture);
 
 		FBloomSetupCS::FPermutationDomain PermutationVector;
-		PermutationVector.Set<FBloomSetupCS::FLocalExposureDim>(Inputs.LocalExposureTexture != nullptr);
+		PermutationVector.Set<FBloomSetupCS::FLocalExposureDim>(bLocalExposureEnabled);
 
 		auto ComputeShader = View.ShaderMap->GetShader<FBloomSetupCS>(PermutationVector);
 
@@ -155,7 +159,7 @@ FScreenPassTexture AddBloomSetupPass(FRDGBuilder& GraphBuilder, const FViewInfo&
 		PassParameters->RenderTargets[0] = Output.GetRenderTargetBinding();
 
 		FBloomSetupPS::FPermutationDomain PermutationVector;
-		PermutationVector.Set<FBloomSetupPS::FLocalExposureDim>(Inputs.LocalExposureTexture != nullptr);
+		PermutationVector.Set<FBloomSetupPS::FLocalExposureDim>(bLocalExposureEnabled);
 
 		auto PixelShader = View.ShaderMap->GetShader<FBloomSetupPS>(PermutationVector);
 
diff --git a/Engine/Source/Runtime/Renderer/Private/PostProcess/PostProcessBloomSetup.h b/Engine/Source/Runtime/Renderer/Private/PostProcess/PostProcessBloomSetup.h
--- a/Engine/Source/Runtime/Renderer/Private/PostProcess/PostProcessBloomSetup.h
+++ b/Engine/Source/Runtime/Renderer/Private/PostProcess/PostProcessBloomSetup.h
@@ -6,6 +6,7 @@
 
 class FSceneDownsampleChain;
 class FEyeAdaptationParameters;
+class FLocalExposureParameters;
 
 enum class EBloomQuality : uint32
 {
@@ -36,6 +37,9 @@ struct FBloomSetupInputs
 	// [Optional] Eye adaptation parameters.
 	const FEyeAdaptationParameters* EyeAdaptationParameters = nullptr;
 
+	// [Optional] Local exposure parameters.
+	const FLocalExposureParameters* LocalExposureParameters = nullptr;
+
 	// [Optional] Luminance bilateral grid. If this is null, local exposure is disabled.
 	FRDGTextureRef LocalExposureTexture = nullptr;
 
diff --git a/Engine/Source/Runtime/Renderer/Private/PostProcess/PostProcessEyeAdaptation.cpp b/Engine/Source/Runtime/Renderer/Private/PostProcess/PostProcessEyeAdaptation.cpp
--- a/Engine/Source/Runtime/Renderer/Private/PostProcess/PostProcessEyeAdaptation.cpp
+++ b/Engine/Source/Runtime/Renderer/Private/PostProcess/PostProcessEyeAdaptation.cpp
@@ -17,6 +17,7 @@
 #include "DataDrivenShaderPlatformInfo.h"
 #include "TextureResource.h"
 #include "PostProcessing.h"
+#include "PostProcessLocalExposure.h"
 
 bool IsMobileEyeAdaptationEnabled(const FViewInfo& View);
 
@@ -394,14 +395,6 @@ FEyeAdaptationParameters GetEyeAdaptationParameters(const FViewInfo& View, ERHIF
 	float ExposureCompensationCurve = GetAutoExposureCompensationFromCurve(View);
 	const float BlackHistogramBucketInfluence = CVarEyeAdaptationBlackHistogramBucketInfluence.GetValueOnRenderThread();
 
-	float LocalExposureMiddleGreyExposureCompensation = FMath::Pow(2.0f, View.FinalPostProcessSettings.LocalExposureMiddleGreyBias);
-
-	if (AutoExposureMethod == EAutoExposureMethod::AEM_Manual)
-	{
-		// when using manual exposure cancel exposure compensation setting and curve from middle grey used by local exposure.
-		LocalExposureMiddleGreyExposureCompensation /= (ExposureCompensationSettings * ExposureCompensationCurve);
-	}
-
 	const float kMiddleGrey = 0.18f;
 
 	// AEM_Histogram and AEM_Basic adjust their ExposureCompensation to middle grey (0.18). AEM_Manual ExposureCompensation is already calibrated to 1.0.
@@ -525,11 +518,6 @@ FEyeAdaptationParameters GetEyeAdaptationParameters(const FViewInfo& View, ERHIF
 	Parameters.HistogramScale = HistogramScale;
 	Parameters.HistogramBias = HistogramBias;
 	Parameters.LuminanceMin = LuminanceMin;
-	Parameters.LocalExposureHighlightContrastScale = Settings.LocalExposureHighlightContrastScale;
-	Parameters.LocalExposureShadowContrastScale = Settings.LocalExposureShadowContrastScale;
-	Parameters.LocalExposureDetailStrength = Settings.LocalExposureDetailStrength;
-	Parameters.LocalExposureBlurredLuminanceBlend = Settings.LocalExposureBlurredLuminanceBlend;
-	Parameters.LocalExposureMiddleGreyExposureCompensation = LocalExposureMiddleGreyExposureCompensation;
 	Parameters.BlackHistogramBucketInfluence = BlackHistogramBucketInfluence; // no calibration constant because it is now baked into ExposureCompensation
 	Parameters.GreyMult = GreyMult;
 	Parameters.ExponentialDownM = ExponentialDownM;
@@ -790,6 +778,7 @@ class FEyeAdaptationCS : public FGlobalShader
 
 	BEGIN_SHADER_PARAMETER_STRUCT(FParameters, )
 		SHADER_PARAMETER_STRUCT(FEyeAdaptationParameters, EyeAdaptation)
+		SHADER_PARAMETER_STRUCT(FLocalExposureParameters, LocalExposure)
 		SHADER_PARAMETER_RDG_TEXTURE(Texture2D, HistogramTexture)
 		SHADER_PARAMETER_RDG_BUFFER_UAV(RWStructuredBuffer<float4>, RWEyeAdaptationBuffer)
 	END_SHADER_PARAMETER_STRUCT()
@@ -811,6 +800,7 @@ FRDGBufferRef AddHistogramEyeAdaptationPass(
 	FRDGBuilder& GraphBuilder,
 	const FViewInfo& View,
 	const FEyeAdaptationParameters& EyeAdaptationParameters,
+	const FLocalExposureParameters& LocalExposureParameters,
 	FRDGTextureRef HistogramTexture,
 	bool bComputeAverageLocalExposure)
 {
@@ -820,7 +810,8 @@ FRDGBufferRef AddHistogramEyeAdaptationPass(
 	FRDGBufferRef OutputBuffer = GraphBuilder.RegisterExternalBuffer(View.GetEyeAdaptationBuffer(GraphBuilder), ERDGBufferFlags::MultiFrame);
 
 	FEyeAdaptationCS::FParameters* PassParameters = GraphBuilder.AllocParameters<FEyeAdaptationCS::FParameters>();
-	PassParameters->EyeAdaptation = GetEyeAdaptationParameters(View, ERHIFeatureLevel::SM5);
+	PassParameters->EyeAdaptation = EyeAdaptationParameters;
+	PassParameters->LocalExposure = LocalExposureParameters;
 	PassParameters->HistogramTexture = HistogramTexture;
 	PassParameters->RWEyeAdaptationBuffer = GraphBuilder.CreateUAV(OutputBuffer);
 
@@ -922,6 +913,7 @@ public:
 	BEGIN_SHADER_PARAMETER_STRUCT(FParameters, )
 		SHADER_PARAMETER_STRUCT_REF(FViewUniformShaderParameters, View)
 		SHADER_PARAMETER_STRUCT(FEyeAdaptationParameters, EyeAdaptation)
+		SHADER_PARAMETER_STRUCT(FLocalExposureParameters, LocalExposure)
 		SHADER_PARAMETER_STRUCT(FScreenPassTextureViewportParameters, Color)
 		SHADER_PARAMETER_RDG_TEXTURE(Texture2D, ColorTexture)
 		SHADER_PARAMETER_RDG_BUFFER_SRV(StructuredBuffer<float4>, EyeAdaptationBuffer)
@@ -943,6 +935,7 @@ FRDGBufferRef AddBasicEyeAdaptationPass(
 	FRDGBuilder& GraphBuilder,
 	const FViewInfo& View,
 	const FEyeAdaptationParameters& EyeAdaptationParameters,
+	const FLocalExposureParameters& LocalExposureParameters,
 	FScreenPassTexture SceneColor,
 	FRDGBufferRef EyeAdaptationBuffer,
 	bool bComputeAverageLocalExposure)
@@ -957,6 +950,7 @@ FRDGBufferRef AddBasicEyeAdaptationPass(
 	FBasicEyeAdaptationCS::FParameters* PassParameters = GraphBuilder.AllocParameters<FBasicEyeAdaptationCS::FParameters>();
 	PassParameters->View = View.ViewUniformBuffer;
 	PassParameters->EyeAdaptation = EyeAdaptationParameters;
+	PassParameters->LocalExposure = LocalExposureParameters;
 	PassParameters->Color = GetScreenPassTextureViewportParameters(SceneColorViewport);
 	PassParameters->ColorTexture = SceneColor.Texture;
 	PassParameters->EyeAdaptationBuffer = GraphBuilder.CreateSRV(EyeAdaptationBuffer);
diff --git a/Engine/Source/Runtime/Renderer/Private/PostProcess/PostProcessEyeAdaptation.h b/Engine/Source/Runtime/Renderer/Private/PostProcess/PostProcessEyeAdaptation.h
--- a/Engine/Source/Runtime/Renderer/Private/PostProcess/PostProcessEyeAdaptation.h
+++ b/Engine/Source/Runtime/Renderer/Private/PostProcess/PostProcessEyeAdaptation.h
@@ -4,6 +4,8 @@
 
 #include "ScreenPass.h"
 
+class FLocalExposureParameters;
+
 namespace AutoExposurePermutation
 {
 	class FUsePrecalculatedLuminanceDim : SHADER_PERMUTATION_BOOL("USE_PRECALCULATED_LUMINANCE");
@@ -59,11 +61,6 @@ BEGIN_SHADER_PARAMETER_STRUCT(FEyeAdaptationParameters, )
 	SHADER_PARAMETER(float, HistogramScale)
 	SHADER_PARAMETER(float, HistogramBias)
 	SHADER_PARAMETER(float, LuminanceMin)
-	SHADER_PARAMETER(float, LocalExposureHighlightContrastScale)
-	SHADER_PARAMETER(float, LocalExposureShadowContrastScale)
-	SHADER_PARAMETER(float, LocalExposureDetailStrength)
-	SHADER_PARAMETER(float, LocalExposureBlurredLuminanceBlend)
-	SHADER_PARAMETER(float, LocalExposureMiddleGreyExposureCompensation)
 	SHADER_PARAMETER(float, BlackHistogramBucketInfluence)
 	SHADER_PARAMETER(float, GreyMult)
 	SHADER_PARAMETER(float, ExponentialUpM)
@@ -102,6 +99,7 @@ FRDGBufferRef AddHistogramEyeAdaptationPass(
 	FRDGBuilder& GraphBuilder,
 	const FViewInfo& View,
 	const FEyeAdaptationParameters& EyeAdaptationParameters,
+	const FLocalExposureParameters& LocalExposureParameters,
 	FRDGTextureRef HistogramTexture,
 	bool bComputeAverageLocalExposure);
 
@@ -117,6 +115,7 @@ FRDGBufferRef AddBasicEyeAdaptationPass(
 	FRDGBuilder& GraphBuilder,
 	const FViewInfo& View,
 	const FEyeAdaptationParameters& EyeAdaptationParameters,
+	const FLocalExposureParameters& LocalExposureParameters,
 	FScreenPassTexture SceneColor,
 	FRDGBufferRef EyeAdaptationBuffer,
 	bool bComputeAverageLocalExposure);
diff --git a/Engine/Source/Runtime/Renderer/Private/PostProcess/PostProcessFFTBloom.cpp b/Engine/Source/Runtime/Renderer/Private/PostProcess/PostProcessFFTBloom.cpp
--- a/Engine/Source/Runtime/Renderer/Private/PostProcess/PostProcessFFTBloom.cpp
+++ b/Engine/Source/Runtime/Renderer/Private/PostProcess/PostProcessFFTBloom.cpp
@@ -712,6 +712,7 @@ FFFTBloomOutput AddFFTBloomPass(
 	float InputResolutionFraction,
 	const FEyeAdaptationParameters& EyeAdaptationParameters,
 	FRDGBufferRef EyeAdaptationBuffer,
+	const FLocalExposureParameters& LocalExposureParameters,
 	FRDGTextureRef LocalExposureTexture,
 	FRDGTextureRef BlurredLogLuminanceTexture)
 {
@@ -832,6 +833,7 @@ FFFTBloomOutput AddFFTBloomPass(
 			View,
 			EyeAdaptationParameters,
 			EyeAdaptationBuffer,
+			LocalExposureParameters,
 			LocalExposureTexture,
 			BlurredLogLuminanceTexture,
 			Temp,
diff --git a/Engine/Source/Runtime/Renderer/Private/PostProcess/PostProcessFFTBloom.h b/Engine/Source/Runtime/Renderer/Private/PostProcess/PostProcessFFTBloom.h
--- a/Engine/Source/Runtime/Renderer/Private/PostProcess/PostProcessFFTBloom.h
+++ b/Engine/Source/Runtime/Renderer/Private/PostProcess/PostProcessFFTBloom.h
@@ -6,6 +6,7 @@
 #include "ScreenPass.h"
 
 class FEyeAdaptationParameters;
+class FLocalExposureParameters;
 class FRDGBuffer;
 class FRDGBuilder;
 class FRDGTexture;
@@ -32,5 +33,6 @@ FFFTBloomOutput AddFFTBloomPass(
 	float InputResolutionFraction,
 	const FEyeAdaptationParameters& EyeAdaptationParameters,
 	FRDGBufferRef EyeAdaptationBuffer,
+	const FLocalExposureParameters& LocalExposureParameters,
 	FRDGTextureRef LocalExposureTexture,
 	FRDGTextureRef BlurredLogLuminanceTexture);
diff --git a/Engine/Source/Runtime/Renderer/Private/PostProcess/PostProcessLocalExposure.cpp b/Engine/Source/Runtime/Renderer/Private/PostProcess/PostProcessLocalExposure.cpp
--- a/Engine/Source/Runtime/Renderer/Private/PostProcess/PostProcessLocalExposure.cpp
+++ b/Engine/Source/Runtime/Renderer/Private/PostProcess/PostProcessLocalExposure.cpp
@@ -67,6 +67,8 @@ public:
 
 		SHADER_PARAMETER_STRUCT(FEyeAdaptationParameters, EyeAdaptation)
 		SHADER_PARAMETER_RDG_BUFFER_SRV(StructuredBuffer<float4>, EyeAdaptationBuffer)
+
+		SHADER_PARAMETER_STRUCT(FLocalExposureParameters, LocalExposure)
 		SHADER_PARAMETER_RDG_TEXTURE(Texture3D, LumBilateralGrid)
 		SHADER_PARAMETER_RDG_TEXTURE(Texture2D, BlurredLogLum)
 
@@ -90,6 +92,29 @@ IMPLEMENT_GLOBAL_SHADER(FApplyLocalExposureCS, "/Engine/Private/PostProcessLocal
 
 } //! namespace
 
+FLocalExposureParameters GetLocalExposureParameters(const FViewInfo& View, const FEyeAdaptationParameters& EyeAdaptationParameters)
+{
+	const FPostProcessSettings& Settings = View.FinalPostProcessSettings;
+
+	const EAutoExposureMethod AutoExposureMethod = GetAutoExposureMethod(View);
+
+	float LocalExposureMiddleGreyExposureCompensation = FMath::Pow(2.0f, View.FinalPostProcessSettings.LocalExposureMiddleGreyBias);
+
+	if (AutoExposureMethod == EAutoExposureMethod::AEM_Manual)
+	{
+		// when using manual exposure cancel exposure compensation setting and curve from middle grey used by local exposure.
+		LocalExposureMiddleGreyExposureCompensation /= (EyeAdaptationParameters.ExposureCompensationSettings * EyeAdaptationParameters.ExposureCompensationCurve);
+	}
+
+	FLocalExposureParameters Parameters;
+	Parameters.HighlightContrastScale = Settings.LocalExposureHighlightContrastScale;
+	Parameters.ShadowContrastScale = Settings.LocalExposureShadowContrastScale;
+	Parameters.DetailStrength = Settings.LocalExposureDetailStrength;
+	Parameters.BlurredLuminanceBlend = Settings.LocalExposureBlurredLuminanceBlend;
+	Parameters.MiddleGreyExposureCompensation = LocalExposureMiddleGreyExposureCompensation;
+	return Parameters;
+}
+
 FRDGTextureRef AddLocalExposureBlurredLogLuminancePass(
 	FRDGBuilder& GraphBuilder,
 	const FViewInfo& View,
@@ -151,6 +176,7 @@ void AddApplyLocalExposurePass(
 	const FViewInfo& View,
 	const FEyeAdaptationParameters& EyeAdaptationParameters,
 	FRDGBufferRef EyeAdaptationBuffer,
+	const FLocalExposureParameters& LocalExposureParamaters,
 	FRDGTextureRef LocalExposureTexture,
 	FRDGTextureRef BlurredLogLuminanceTexture,
 	FScreenPassTextureSlice Input,
@@ -182,8 +208,11 @@ void AddApplyLocalExposurePass(
 
 	PassParameters->EyeAdaptation = EyeAdaptationParameters;
 	PassParameters->EyeAdaptationBuffer = GraphBuilder.CreateSRV(EyeAdaptationBuffer);
+
+	PassParameters->LocalExposure = LocalExposureParamaters;
 	PassParameters->LumBilateralGrid = LocalExposureTexture;
 	PassParameters->BlurredLogLum = BlurredLogLuminanceTexture;
+
 	PassParameters->TextureSampler = TStaticSamplerState<SF_Bilinear, AM_Clamp, AM_Clamp, AM_Clamp>::GetRHI();;
 
 	FComputeShaderUtils::AddPass(
diff --git a/Engine/Source/Runtime/Renderer/Private/PostProcess/PostProcessLocalExposure.h b/Engine/Source/Runtime/Renderer/Private/PostProcess/PostProcessLocalExposure.h
--- a/Engine/Source/Runtime/Renderer/Private/PostProcess/PostProcessLocalExposure.h
+++ b/Engine/Source/Runtime/Renderer/Private/PostProcess/PostProcessLocalExposure.h
@@ -6,6 +6,16 @@
 
 class FEyeAdaptationParameters;
 
+BEGIN_SHADER_PARAMETER_STRUCT(FLocalExposureParameters, )
+	SHADER_PARAMETER(float, HighlightContrastScale)
+	SHADER_PARAMETER(float, ShadowContrastScale)
+	SHADER_PARAMETER(float, DetailStrength)
+	SHADER_PARAMETER(float, BlurredLuminanceBlend)
+	SHADER_PARAMETER(float, MiddleGreyExposureCompensation)
+END_SHADER_PARAMETER_STRUCT()
+
+FLocalExposureParameters GetLocalExposureParameters(const FViewInfo& View, const FEyeAdaptationParameters& EyeAdaptationParameters);
+
 FRDGTextureRef AddLocalExposureBlurredLogLuminancePass(
 	FRDGBuilder& GraphBuilder,
 	const FViewInfo& View,
@@ -17,8 +27,9 @@ void AddApplyLocalExposurePass(
 	const FViewInfo& View,
 	const FEyeAdaptationParameters& EyeAdaptationParameters,
 	FRDGBufferRef EyeAdaptationBuffer,
+	const FLocalExposureParameters& LocalExposureParamaters,
 	FRDGTextureRef LocalExposureTexture,
 	FRDGTextureRef BlurredLogLuminanceTexture,
 	FScreenPassTextureSlice Input,
 	FScreenPassTextureSlice Output,
-	ERDGPassFlags PassFlags);
\ No newline at end of file
+	ERDGPassFlags PassFlags);
diff --git a/Engine/Source/Runtime/Renderer/Private/PostProcess/PostProcessTonemap.cpp b/Engine/Source/Runtime/Renderer/Private/PostProcess/PostProcessTonemap.cpp
--- a/Engine/Source/Runtime/Renderer/Private/PostProcess/PostProcessTonemap.cpp
+++ b/Engine/Source/Runtime/Renderer/Private/PostProcess/PostProcessTonemap.cpp
@@ -12,6 +12,7 @@
 #include "PostProcess/PostProcessing.h"
 #include "PostProcess/SceneFilterRendering.h"
 #include "PostProcess/PostProcessCombineLUTs.h"
+#include "PostProcess/PostProcessLocalExposure.h"
 #include "PostProcess/PostProcessMobile.h"
 #include "PostProcess/PostProcessing.h"
 #include "ClearQuad.h"
@@ -313,6 +314,7 @@ BEGIN_SHADER_PARAMETER_STRUCT(FTonemapParameters, )
 	SHADER_PARAMETER_STRUCT(FScreenPassTextureViewportParameters, Color)
 	SHADER_PARAMETER_STRUCT(FScreenPassTextureViewportParameters, Output)
 	SHADER_PARAMETER_STRUCT(FEyeAdaptationParameters, EyeAdaptation)
+	SHADER_PARAMETER_STRUCT(FLocalExposureParameters, LocalExposure)
 	SHADER_PARAMETER_RDG_TEXTURE_SRV(Texture2D, ColorTexture)
 
 	// Parameters to apply to the scene color.
@@ -761,6 +763,14 @@ FScreenPassTexture AddTonemapPass(FRDGBuilder& GraphBuilder, const FViewInfo& Vi
 	
 	const float LUTSize = Inputs.ColorGradingTexture ? (float)Inputs.ColorGradingTexture->Desc.Extent.Y : /* unused (default): */ 32.0f;
 
+	const bool bLocalExposureEnabled = Inputs.LocalExposureTexture != nullptr;
+
+	if (bLocalExposureEnabled)
+	{
+		checkf(Inputs.LocalExposureParameters != nullptr, TEXT("When Local Exposure is enabled, corresponding parameters must be provided"));
+		CommonParameters.LocalExposure = *Inputs.LocalExposureParameters;
+	}
+
 	CommonParameters.OutputDevice = GetTonemapperOutputDeviceParameters(ViewFamily);
 	CommonParameters.Color = GetScreenPassTextureViewportParameters(SceneColorViewport);
 	CommonParameters.Output = GetScreenPassTextureViewportParameters(OutputViewport);
@@ -877,7 +887,7 @@ FScreenPassTexture AddTonemapPass(FRDGBuilder& GraphBuilder, const FViewInfo& Vi
 	TonemapperPermutation::FDesktopDomain DesktopPermutationVector;
 
 	{
-		TonemapperPermutation::FCommonDomain CommonDomain = TonemapperPermutation::BuildCommonPermutationDomain(View, Inputs.bGammaOnly, Inputs.LocalExposureTexture != nullptr, Inputs.bMetalMSAAHDRDecode);
+		TonemapperPermutation::FCommonDomain CommonDomain = TonemapperPermutation::BuildCommonPermutationDomain(View, Inputs.bGammaOnly, bLocalExposureEnabled, Inputs.bMetalMSAAHDRDecode);
 		DesktopPermutationVector.Set<TonemapperPermutation::FCommonDomain>(CommonDomain);
 
 		if (!CommonDomain.Get<TonemapperPermutation::FTonemapperGammaOnlyDim>())
diff --git a/Engine/Source/Runtime/Renderer/Private/PostProcess/PostProcessTonemap.h b/Engine/Source/Runtime/Renderer/Private/PostProcess/PostProcessTonemap.h
--- a/Engine/Source/Runtime/Renderer/Private/PostProcess/PostProcessTonemap.h
+++ b/Engine/Source/Runtime/Renderer/Private/PostProcess/PostProcessTonemap.h
@@ -8,6 +8,7 @@
 #include "OverridePassSequence.h"
 #include "Math/Halton.h"
 
+class FLocalExposureParameters;
 
 bool SupportsFilmGrain(EShaderPlatform Platform);
 
@@ -46,6 +47,9 @@ struct FTonemapInputs
 	// [Optional] Blurred luminance texture used to calculate local exposure.
 	FRDGTextureRef BlurredLogLuminanceTexture = nullptr;
 
+	// [Optional] Local exposure parameters.
+	const FLocalExposureParameters* LocalExposureParameters = nullptr;
+
 	// [Required] Eye adaptation parameters.
 	const FEyeAdaptationParameters* EyeAdaptationParameters = nullptr;
 
diff --git a/Engine/Source/Runtime/Renderer/Private/PostProcess/PostProcessVisualizeLocalExposure.cpp b/Engine/Source/Runtime/Renderer/Private/PostProcess/PostProcessVisualizeLocalExposure.cpp
--- a/Engine/Source/Runtime/Renderer/Private/PostProcess/PostProcessVisualizeLocalExposure.cpp
+++ b/Engine/Source/Runtime/Renderer/Private/PostProcess/PostProcessVisualizeLocalExposure.cpp
@@ -2,6 +2,7 @@
 
 #include "PostProcess/PostProcessVisualizeLocalExposure.h"
 #include "PostProcess/PostProcessTonemap.h"
+#include "PostProcess/PostProcessLocalExposure.h"
 #include "UnrealEngine.h"
 #include "DataDrivenShaderPlatformInfo.h"
 #include "SceneRendering.h"
@@ -24,6 +25,7 @@ public:
 	BEGIN_SHADER_PARAMETER_STRUCT(FParameters, )
 		SHADER_PARAMETER_STRUCT_REF(FViewUniformShaderParameters, View)
 		SHADER_PARAMETER_STRUCT(FEyeAdaptationParameters, EyeAdaptation)
+		SHADER_PARAMETER_STRUCT(FLocalExposureParameters, LocalExposure)
 		SHADER_PARAMETER_STRUCT(FScreenPassTextureViewportParameters, Input)
 		SHADER_PARAMETER_STRUCT(FScreenPassTextureViewportParameters, Output)
 
@@ -74,6 +76,7 @@ FScreenPassTexture AddVisualizeLocalExposurePass(FRDGBuilder& GraphBuilder, cons
 	PassParameters->RenderTargets[0] = Output.GetRenderTargetBinding();
 	PassParameters->View = View.ViewUniformBuffer;
 	PassParameters->EyeAdaptation = *Inputs.EyeAdaptationParameters;
+	PassParameters->LocalExposure = *Inputs.LocalExposureParameters;
 	PassParameters->Input = GetScreenPassTextureViewportParameters(InputViewport);
 	PassParameters->Output = GetScreenPassTextureViewportParameters(OutputViewport);
 	PassParameters->HDRSceneColorTexture = Inputs.HDRSceneColor.Texture;
diff --git a/Engine/Source/Runtime/Renderer/Private/PostProcess/PostProcessVisualizeLocalExposure.h b/Engine/Source/Runtime/Renderer/Private/PostProcess/PostProcessVisualizeLocalExposure.h
--- a/Engine/Source/Runtime/Renderer/Private/PostProcess/PostProcessVisualizeLocalExposure.h
+++ b/Engine/Source/Runtime/Renderer/Private/PostProcess/PostProcessVisualizeLocalExposure.h
@@ -5,6 +5,7 @@
 #include "OverridePassSequence.h"
 
 class FEyeAdaptationParameters;
+class FLocalExposureParameters;
 
 struct FVisualizeLocalExposureInputs
 {
@@ -18,6 +19,7 @@ struct FVisualizeLocalExposureInputs
 
 	FRDGTextureRef LumBilateralGridTexture = nullptr;
 	FRDGTextureRef BlurredLumTexture = nullptr;
+	const FLocalExposureParameters* LocalExposureParameters = nullptr;
 
 	const FEyeAdaptationParameters* EyeAdaptationParameters = nullptr;
 };
diff --git a/Engine/Source/Runtime/Renderer/Private/PostProcess/PostProcessing.cpp b/Engine/Source/Runtime/Renderer/Private/PostProcess/PostProcessing.cpp
--- a/Engine/Source/Runtime/Renderer/Private/PostProcess/PostProcessing.cpp
+++ b/Engine/Source/Runtime/Renderer/Private/PostProcess/PostProcessing.cpp
@@ -328,6 +328,8 @@ void AddPostProcessingPasses(
 	FRDGBufferRef LastEyeAdaptationBuffer = GetEyeAdaptationBuffer(GraphBuilder, View);
 	FRDGBufferRef EyeAdaptationBuffer = LastEyeAdaptationBuffer;
 
+	FLocalExposureParameters LocalExposureParameters;
+
 	// Histogram defaults to black because the histogram eye adaptation pass is used for the manual metering mode.
 	FRDGTextureRef HistogramTexture = BlackDummy.Texture;
 
@@ -951,6 +953,8 @@ void AddPostProcessingPasses(
 				bProcessQuarterResolution ? QuarterResSceneColor : HalfResSceneColor);
 		}
 
+		LocalExposureParameters = GetLocalExposureParameters(View, EyeAdaptationParameters);
+
 		if (bHistogramEnabled)
 		{
 			FScreenPassTexture HistogramSceneColor = bProcessQuarterResolution ? QuarterResSceneColor : HalfResSceneColor;
@@ -1003,6 +1007,7 @@ void AddPostProcessingPasses(
 			EyeAdaptationBuffer = AddBasicEyeAdaptationPass(
 				GraphBuilder, View,
 				EyeAdaptationParameters,
+				LocalExposureParameters,
 				SceneDownsampleChain.GetLastTexture(),
 				LastEyeAdaptationBuffer,
 				bLocalExposureEnabled);
@@ -1013,6 +1018,7 @@ void AddPostProcessingPasses(
 			EyeAdaptationBuffer = AddHistogramEyeAdaptationPass(
 				GraphBuilder, View,
 				EyeAdaptationParameters,
+				LocalExposureParameters,
 				HistogramTexture,
 				bLocalExposureEnabled);
 		}
@@ -1055,6 +1061,7 @@ void AddPostProcessingPasses(
 					InputResolutionFraction,
 					EyeAdaptationParameters,
 					EyeAdaptationBuffer,
+					LocalExposureParameters,
 					CVarBloomApplyLocalExposure.GetValueOnRenderThread() ? LocalExposureTexture : nullptr,
 					LocalExposureBlurredLogLumTexture);
 
@@ -1082,6 +1089,7 @@ void AddPostProcessingPasses(
 						SetupPassInputs.SceneColor = DownsampleInput;
 						SetupPassInputs.EyeAdaptationBuffer = EyeAdaptationBuffer;
 						SetupPassInputs.EyeAdaptationParameters = &EyeAdaptationParameters;
+						SetupPassInputs.LocalExposureParameters = &LocalExposureParameters;
 						SetupPassInputs.LocalExposureTexture = CVarBloomApplyLocalExposure.GetValueOnRenderThread() ? LocalExposureTexture : nullptr;
 						SetupPassInputs.BlurredLogLuminanceTexture = LocalExposureBlurredLogLumTexture;
 						SetupPassInputs.Threshold = BloomThreshold;
@@ -1151,6 +1159,7 @@ void AddPostProcessingPasses(
 				PassInputs.SceneColorApplyParamaters = SceneColorApplyParameters;
 				PassInputs.LocalExposureTexture = LocalExposureTexture;
 				PassInputs.BlurredLogLuminanceTexture = LocalExposureBlurredLogLumTexture;
+				PassInputs.LocalExposureParameters = &LocalExposureParameters;
 				PassInputs.EyeAdaptationParameters = &EyeAdaptationParameters;
 				PassInputs.EyeAdaptationBuffer = EyeAdaptationBuffer;
 				PassInputs.ColorGradingTexture = ColorGradingTexture;
@@ -1465,6 +1474,7 @@ void AddPostProcessingPasses(
 		PassInputs.HDRSceneColor = FScreenPassTexture::CopyFromSlice(GraphBuilder, SceneColorBeforeTonemapSlice);
 		PassInputs.LumBilateralGridTexture = LocalExposureTexture;
 		PassInputs.BlurredLumTexture = LocalExposureBlurredLogLumTexture;
+		PassInputs.LocalExposureParameters = &LocalExposureParameters;
 		PassInputs.EyeAdaptationBuffer = GetEyeAdaptationBuffer(GraphBuilder, View);
 		PassInputs.EyeAdaptationParameters = &EyeAdaptationParameters;
 
@@ -1695,10 +1705,10 @@ void AddDebugViewPostProcessingPasses(FRDGBuilder& GraphBuilder, const FViewInfo
 	PassSequence.SetEnabled(EPass::SecondaryUpscale, View.RequiresSecondaryUpscale() || View.Family->GetSecondarySpatialUpscalerInterface() != nullptr);
 	PassSequence.Finalize();
 
+	const FEyeAdaptationParameters EyeAdaptationParameters = GetEyeAdaptationParameters(View, ERHIFeatureLevel::SM5);
+
 	if (bTonemapBefore)
 	{
-		const FEyeAdaptationParameters EyeAdaptationParameters = GetEyeAdaptationParameters(View, ERHIFeatureLevel::SM5);
-
 		FTonemapInputs PassInputs;
 		PassInputs.SceneColor = FScreenPassTextureSlice::CreateFromScreenPassTexture(GraphBuilder, SceneColor);
 		PassInputs.bOutputInHDR = bViewFamilyOutputInHDR;
@@ -1790,8 +1800,6 @@ void AddDebugViewPostProcessingPasses(FRDGBuilder& GraphBuilder, const FViewInfo
 
 	if (PassSequence.IsEnabled(EPass::TonemapAfter))
 	{
-		const FEyeAdaptationParameters EyeAdaptationParameters = GetEyeAdaptationParameters(View, ERHIFeatureLevel::SM5);
-
 		FTonemapInputs PassInputs;
 		PassSequence.AcceptOverrideIfLastPass(EPass::TonemapAfter, PassInputs.OverrideOutput);
 		PassInputs.SceneColor = FScreenPassTextureSlice::CreateFromScreenPassTexture(GraphBuilder, SceneColor);
