Fixed incorrect sampling of Local Exposure bilateral grid.
- pixels are bucketed into bilateral grid on fixed size tiles (64x64)
- when ViewRect is not a multiple of the tile size, sampling the bilateral grid requires applying a scale to ViewportUV.

diff --git a/Engine/Shaders/Private/PostProcessHistogramCommon.ush b/Engine/Shaders/Private/PostProcessHistogramCommon.ush
--- a/Engine/Shaders/Private/PostProcessHistogramCommon.ush
+++ b/Engine/Shaders/Private/PostProcessHistogramCommon.ush
@@ -36,6 +36,7 @@ float LocalExposure_ShadowContrastScale;
 float LocalExposure_DetailStrength;
 float LocalExposure_BlurredLuminanceBlend;
 float LocalExposure_MiddleGreyExposureCompensation;
+float2 LocalExposure_BilateralGridUVScale;
 
 float CalculateEyeAdaptationLuminance(float3 Color)
 {
@@ -218,8 +219,11 @@ float CalculateBaseLogLuminance(float BilateralLum, float BlurredLum, float Blur
 
 float CalculateBaseLogLuminance(float LogLum, float BlurredLumBlend, float ExposureScale, float2 UV, Texture3D LumBilateralGrid, Texture2D BlurredLogLum, SamplerState LumBilateralGridSampler, SamplerState BlurredLogLumSampler)
 {
-	float LumW = (ComputeHistogramPositionFromLogLuminance(LogLum) * (BILATERAL_GRID_DEPTH - 1) + 0.5f) / BILATERAL_GRID_DEPTH;
-	float2 BilateralGridLum = Texture3DSample(LumBilateralGrid, LumBilateralGridSampler, float3(UV, LumW)).xy;
+	float3 BilateralGridUVW;
+	BilateralGridUVW.xy = UV * LocalExposure_BilateralGridUVScale;
+	BilateralGridUVW.z = (ComputeHistogramPositionFromLogLuminance(LogLum) * (BILATERAL_GRID_DEPTH - 1) + 0.5f) / BILATERAL_GRID_DEPTH;
+	
+	float2 BilateralGridLum = Texture3DSample(LumBilateralGrid, LumBilateralGridSampler, BilateralGridUVW).xy;
 	float BilateralLum = BilateralGridLum.x / BilateralGridLum.y;
 
 	float BlurredLum = Texture2DSample(BlurredLogLum, BlurredLogLumSampler, UV).r;
diff --git a/Engine/Shaders/Private/PostProcessVisualizeLocalExposure.usf b/Engine/Shaders/Private/PostProcessVisualizeLocalExposure.usf
--- a/Engine/Shaders/Private/PostProcessVisualizeLocalExposure.usf
+++ b/Engine/Shaders/Private/PostProcessVisualizeLocalExposure.usf
@@ -54,8 +54,11 @@ float4 MainPS(noperspective float4 UVAndScreenPos : TEXCOORD0, float4 SvPosition
 	}
 	else if (DebugMode == 3)
 	{
-		float LumW = (ComputeHistogramPositionFromLogLuminance(LogLuminance) * (BILATERAL_GRID_DEPTH - 1) + 0.5f) / BILATERAL_GRID_DEPTH;
-		float2 BilateralGridLum = Texture3DSample(LumBilateralGrid, TextureSampler, float3(ViewportUV, LumW)).xy;
+		float3 BilateralGridUVW;
+		BilateralGridUVW.xy = ViewportUV * LocalExposure_BilateralGridUVScale;
+		BilateralGridUVW.z = (ComputeHistogramPositionFromLogLuminance(LogLuminance) * (BILATERAL_GRID_DEPTH - 1) + 0.5f) / BILATERAL_GRID_DEPTH;
+
+		float2 BilateralGridLum = Texture3DSample(LumBilateralGrid, TextureSampler, BilateralGridUVW).xy;
 
 		if (BilateralGridLum.y < 0.001)
 		{
diff --git a/Engine/Source/Runtime/Renderer/Private/PostProcess/PostProcessHistogram.cpp b/Engine/Source/Runtime/Renderer/Private/PostProcess/PostProcessHistogram.cpp
--- a/Engine/Source/Runtime/Renderer/Private/PostProcess/PostProcessHistogram.cpp
+++ b/Engine/Source/Runtime/Renderer/Private/PostProcess/PostProcessHistogram.cpp
@@ -101,10 +101,15 @@ public:
 		return bBilateralGrid ? FIntPoint(8, 8) : FIntPoint(8, 4);
 	}
 
-	static FIntPoint GetThreadGroupCount(FIntPoint InputExtent, bool bBilateralGrid)
+	static FIntPoint GetTexelsPerThreadGroup( bool bBilateralGrid)
 	{
 		// One ThreadGroup ThreadGroupSizeX*ThreadGroupSizeY processes blocks of size LoopCountX*LoopCountY
-		const FIntPoint TexelsPerThreadGroup = GetThreadGroupSize(bBilateralGrid) * FIntPoint(LoopCountX, LoopCountY);
+		return GetThreadGroupSize(bBilateralGrid) * FIntPoint(LoopCountX, LoopCountY);
+	}
+
+	static FIntPoint GetThreadGroupCount(FIntPoint InputExtent, bool bBilateralGrid)
+	{
+		const FIntPoint TexelsPerThreadGroup = GetTexelsPerThreadGroup(bBilateralGrid);
 
 		return FIntPoint::DivideAndRoundUp(InputExtent, TexelsPerThreadGroup);
 	}
@@ -468,6 +473,15 @@ FRDGTextureRef AddHistogramPass(
 	}
 }
 
+FVector2f GetLocalExposureBilateralGridUVScale(const FIntPoint ViewRectSize)
+{
+	const FIntPoint TexelsPerThreadGroup = FHistogramCS::GetTexelsPerThreadGroup(true);
+
+	const FIntPoint ThreadGroupCount = FHistogramCS::GetThreadGroupCount(ViewRectSize, true);
+
+	return FVector2f(float(ViewRectSize.X) / TexelsPerThreadGroup.X / ThreadGroupCount.X, float(ViewRectSize.Y) / TexelsPerThreadGroup.Y / ThreadGroupCount.Y);
+}
+
 FRDGTextureRef AddLocalExposurePass(
 	FRDGBuilder& GraphBuilder,
 	const FViewInfo& View,
diff --git a/Engine/Source/Runtime/Renderer/Private/PostProcess/PostProcessLocalExposure.cpp b/Engine/Source/Runtime/Renderer/Private/PostProcess/PostProcessLocalExposure.cpp
--- a/Engine/Source/Runtime/Renderer/Private/PostProcess/PostProcessLocalExposure.cpp
+++ b/Engine/Source/Runtime/Renderer/Private/PostProcess/PostProcessLocalExposure.cpp
@@ -92,7 +92,9 @@ IMPLEMENT_GLOBAL_SHADER(FApplyLocalExposureCS, "/Engine/Private/PostProcessLocal
 
 } //! namespace
 
-FLocalExposureParameters GetLocalExposureParameters(const FViewInfo& View, const FEyeAdaptationParameters& EyeAdaptationParameters)
+FVector2f GetLocalExposureBilateralGridUVScale(const FIntPoint ViewRectSize);
+
+FLocalExposureParameters GetLocalExposureParameters(const FViewInfo& View, FIntPoint ViewRectSize, const FEyeAdaptationParameters& EyeAdaptationParameters)
 {
 	const FPostProcessSettings& Settings = View.FinalPostProcessSettings;
 
@@ -106,12 +108,15 @@ FLocalExposureParameters GetLocalExposureParameters(const FViewInfo& View, const
 		LocalExposureMiddleGreyExposureCompensation /= (EyeAdaptationParameters.ExposureCompensationSettings * EyeAdaptationParameters.ExposureCompensationCurve);
 	}
 
+	const FVector2f LocalExposureBilateralGridUVScale = GetLocalExposureBilateralGridUVScale(ViewRectSize);
+
 	FLocalExposureParameters Parameters;
 	Parameters.HighlightContrastScale = Settings.LocalExposureHighlightContrastScale;
 	Parameters.ShadowContrastScale = Settings.LocalExposureShadowContrastScale;
 	Parameters.DetailStrength = Settings.LocalExposureDetailStrength;
 	Parameters.BlurredLuminanceBlend = Settings.LocalExposureBlurredLuminanceBlend;
 	Parameters.MiddleGreyExposureCompensation = LocalExposureMiddleGreyExposureCompensation;
+	Parameters.BilateralGridUVScale = LocalExposureBilateralGridUVScale;
 	return Parameters;
 }
 
diff --git a/Engine/Source/Runtime/Renderer/Private/PostProcess/PostProcessLocalExposure.h b/Engine/Source/Runtime/Renderer/Private/PostProcess/PostProcessLocalExposure.h
--- a/Engine/Source/Runtime/Renderer/Private/PostProcess/PostProcessLocalExposure.h
+++ b/Engine/Source/Runtime/Renderer/Private/PostProcess/PostProcessLocalExposure.h
@@ -12,9 +12,10 @@ BEGIN_SHADER_PARAMETER_STRUCT(FLocalExposureParameters, )
 	SHADER_PARAMETER(float, DetailStrength)
 	SHADER_PARAMETER(float, BlurredLuminanceBlend)
 	SHADER_PARAMETER(float, MiddleGreyExposureCompensation)
+	SHADER_PARAMETER(FVector2f, BilateralGridUVScale)
 END_SHADER_PARAMETER_STRUCT()
 
-FLocalExposureParameters GetLocalExposureParameters(const FViewInfo& View, const FEyeAdaptationParameters& EyeAdaptationParameters);
+FLocalExposureParameters GetLocalExposureParameters(const FViewInfo& View, FIntPoint ViewRectSize, const FEyeAdaptationParameters& EyeAdaptationParameters);
 
 FRDGTextureRef AddLocalExposureBlurredLogLuminancePass(
 	FRDGBuilder& GraphBuilder,
diff --git a/Engine/Source/Runtime/Renderer/Private/PostProcess/PostProcessing.cpp b/Engine/Source/Runtime/Renderer/Private/PostProcess/PostProcessing.cpp
--- a/Engine/Source/Runtime/Renderer/Private/PostProcess/PostProcessing.cpp
+++ b/Engine/Source/Runtime/Renderer/Private/PostProcess/PostProcessing.cpp
@@ -945,15 +945,19 @@ void AddPostProcessingPasses(
 			GraphBuilder.QueueTextureExtraction(HalfResSceneColor.Texture, &View.ViewState->PrevFrameViewInfo.HalfResTemporalAAHistory);
 		}
 
-		if (bLocalExposureEnabled)
 		{
-			LocalExposureTexture = AddLocalExposurePass(
-				GraphBuilder, View,
-				EyeAdaptationParameters,
-				bProcessQuarterResolution ? QuarterResSceneColor : HalfResSceneColor);
-		}
+			FScreenPassTexture LocalExposureSceneColor = bProcessQuarterResolution ? QuarterResSceneColor : HalfResSceneColor;
+
+			if (bLocalExposureEnabled)
+			{
+				LocalExposureTexture = AddLocalExposurePass(
+					GraphBuilder, View,
+					EyeAdaptationParameters,
+					LocalExposureSceneColor);
+			}
 
-		LocalExposureParameters = GetLocalExposureParameters(View, EyeAdaptationParameters);
+			LocalExposureParameters = GetLocalExposureParameters(View, LocalExposureSceneColor.ViewRect.Size(), EyeAdaptationParameters);
+		}
 
 		if (bHistogramEnabled)
 		{
