Wrapped VT system in a lock so that it is threadsafe.
 - Moved internal API to be private interface with friendship.
 - Public API is guarded by a single lock.

diff --git a/Engine/Source/Runtime/Engine/Private/VT/RuntimeVirtualTexture.cpp b/Engine/Source/Runtime/Engine/Private/VT/RuntimeVirtualTexture.cpp
--- a/Engine/Source/Runtime/Engine/Private/VT/RuntimeVirtualTexture.cpp
+++ b/Engine/Source/Runtime/Engine/Private/VT/RuntimeVirtualTexture.cpp
@@ -119,7 +119,6 @@ public:
 	/** Getter for the virtual texture producer. */
 	FVirtualTextureProducerHandle GetProducerHandle() const
 	{
-		checkSlow(IsInRenderingThread());
 		return ProducerHandle;
 	}
 
diff --git a/Engine/Source/Runtime/Renderer/Private/VT/AdaptiveVirtualTexture.cpp b/Engine/Source/Runtime/Renderer/Private/VT/AdaptiveVirtualTexture.cpp
--- a/Engine/Source/Runtime/Renderer/Private/VT/AdaptiveVirtualTexture.cpp
+++ b/Engine/Source/Runtime/Renderer/Private/VT/AdaptiveVirtualTexture.cpp
@@ -134,115 +134,110 @@ union FPackedAdaptiveAllocationRequest
 	};
 };
 
-/** Local helper functions. */
-namespace
+/** Allocate a virtual texture for a subset of the full adaptive virtual texture. */
+IAllocatedVirtualTexture* FAdaptiveVirtualTexture::AllocateVirtualTexture(
+	FVirtualTextureSystem* InSystem,
+	FAllocatedVTDescription const& InAllocatedDesc,
+	FIntPoint InGridSize,
+	uint8 InForcedSpaceID,
+	int32 InWidthInTiles,
+	int32 InHeightInTiles,
+	FIntPoint InAddressOffset,
+	int32 InLevelOffset)
 {
-	/** Allocate a virtual texture for a subset of the full adaptive virtual texture. */
-	IAllocatedVirtualTexture* AllocateVirtualTexture(
-		FVirtualTextureSystem* InSystem,
-		FAllocatedVTDescription const& InAllocatedDesc,
-		FIntPoint InGridSize,
-		uint8 InForcedSpaceID,
-		int32 InWidthInTiles,
-		int32 InHeightInTiles,
-		FIntPoint InAddressOffset,
-		int32 InLevelOffset)
-	{
-		FAllocatedVTDescription AllocatedDesc = InAllocatedDesc;
-
-		// We require bPrivateSpace since there can be only one adaptive VT per space.
-		ensure(AllocatedDesc.bPrivateSpace);
-		AllocatedDesc.bPrivateSpace = true;
-		AllocatedDesc.ForceSpaceID = InForcedSpaceID;
-		AllocatedDesc.IndirectionTextureSize = FMath::Max(InGridSize.X, InGridSize.Y);
-		AllocatedDesc.AdaptiveLevelBias = InLevelOffset;
-
-		for (int32 LayerIndex = 0; LayerIndex < InAllocatedDesc.NumTextureLayers; ++LayerIndex)
+	FAllocatedVTDescription AllocatedDesc = InAllocatedDesc;
+
+	// We require bPrivateSpace since there can be only one adaptive VT per space.
+	ensure(AllocatedDesc.bPrivateSpace);
+	AllocatedDesc.bPrivateSpace = true;
+	AllocatedDesc.ForceSpaceID = InForcedSpaceID;
+	AllocatedDesc.IndirectionTextureSize = FMath::Max(InGridSize.X, InGridSize.Y);
+	AllocatedDesc.AdaptiveLevelBias = InLevelOffset;
+
+	for (int32 LayerIndex = 0; LayerIndex < InAllocatedDesc.NumTextureLayers; ++LayerIndex)
+	{
+		// Test if we have already written layer with a new handle.
+		// If we have then we already processed this producer in an ealier layer and have nothing more to do.
+		if (AllocatedDesc.ProducerHandle[LayerIndex] != InAllocatedDesc.ProducerHandle[LayerIndex])
 		{
-			// Test if we have already written layer with a new handle.
-			// If we have then we already processed this producer in an ealier layer and have nothing more to do.
-			if (AllocatedDesc.ProducerHandle[LayerIndex] != InAllocatedDesc.ProducerHandle[LayerIndex])
-			{
-				continue;
-			}
+			continue;
+		}
 
-			FVirtualTextureProducerHandle ProducerHandle = InAllocatedDesc.ProducerHandle[LayerIndex];
-			FVirtualTextureProducer* Producer = InSystem->FindProducer(ProducerHandle);
-			FVTProducerDescription NewProducerDesc = Producer->GetDescription();
-			NewProducerDesc.BlockWidthInTiles = InWidthInTiles;
-			NewProducerDesc.BlockHeightInTiles = InHeightInTiles;
-			NewProducerDesc.MaxLevel = FMath::CeilLogTwo(FMath::Max(InWidthInTiles, InHeightInTiles));
+		FVirtualTextureProducerHandle ProducerHandle = InAllocatedDesc.ProducerHandle[LayerIndex];
+		FVirtualTextureProducer* Producer = InSystem->FindProducer(ProducerHandle);
+		FVTProducerDescription NewProducerDesc = Producer->GetDescription();
+		NewProducerDesc.BlockWidthInTiles = InWidthInTiles;
+		NewProducerDesc.BlockHeightInTiles = InHeightInTiles;
+		NewProducerDesc.MaxLevel = FMath::CeilLogTwo(FMath::Max(InWidthInTiles, InHeightInTiles));
 
-			IVirtualTexture* VirtualTextureProducer = Producer->GetVirtualTexture();
-			IVirtualTexture* NewVirtualTextureProducer = new FVirtualTextureAddressRedirect(VirtualTextureProducer, InAddressOffset, InLevelOffset);
-			FVirtualTextureProducerHandle NewProducerHandle = InSystem->RegisterProducer(NewProducerDesc, NewVirtualTextureProducer);
+		IVirtualTexture* VirtualTextureProducer = Producer->GetVirtualTexture();
+		IVirtualTexture* NewVirtualTextureProducer = new FVirtualTextureAddressRedirect(VirtualTextureProducer, InAddressOffset, InLevelOffset);
+		FVirtualTextureProducerHandle NewProducerHandle = InSystem->RegisterProducer(NewProducerDesc, NewVirtualTextureProducer);
 
-			// Copy new producer to all subsequent layers.
-			for (int32 WriteLayerIndex = LayerIndex; WriteLayerIndex < InAllocatedDesc.NumTextureLayers; ++WriteLayerIndex)
+		// Copy new producer to all subsequent layers.
+		for (int32 WriteLayerIndex = LayerIndex; WriteLayerIndex < InAllocatedDesc.NumTextureLayers; ++WriteLayerIndex)
+		{
+			if (InAllocatedDesc.ProducerHandle[WriteLayerIndex] == ProducerHandle)
 			{
-				if (InAllocatedDesc.ProducerHandle[WriteLayerIndex] == ProducerHandle)
-				{
-					AllocatedDesc.ProducerHandle[WriteLayerIndex] = NewProducerHandle;
-				}
+				AllocatedDesc.ProducerHandle[WriteLayerIndex] = NewProducerHandle;
 			}
 		}
-
-		return InSystem->AllocateVirtualTexture(AllocatedDesc);
 	}
 
-	/** Destroy an allocated virtual texture and release its producers. */
-	static void DestroyVirtualTexture(FVirtualTextureSystem* InSystem, IAllocatedVirtualTexture* InAllocatedVT)
+	return InSystem->AllocateVirtualTexture(AllocatedDesc);
+}
+
+/** Destroy an allocated virtual texture and release its producers. */
+void FAdaptiveVirtualTexture::DestroyVirtualTexture(FVirtualTextureSystem* InSystem, IAllocatedVirtualTexture* InAllocatedVT)
+{
+	FAllocatedVTDescription const& Desc = InAllocatedVT->GetDescription();
+	TArray<FVirtualTextureProducerHandle, TInlineAllocator<8>> ProducersToRelease;
+	for (int32 LayerIndex = 0; LayerIndex < Desc.NumTextureLayers; ++LayerIndex)
 	{
-		FAllocatedVTDescription const& Desc = InAllocatedVT->GetDescription();
-		TArray<FVirtualTextureProducerHandle, TInlineAllocator<8>> ProducersToRelease;
-		for (int32 LayerIndex = 0; LayerIndex < Desc.NumTextureLayers; ++LayerIndex)
-		{
-			ProducersToRelease.AddUnique(Desc.ProducerHandle[LayerIndex]);
-		}
-		InSystem->DestroyVirtualTexture(InAllocatedVT);
-		for (int32 ProducerIndex = 0; ProducerIndex < ProducersToRelease.Num(); ++ProducerIndex)
-		{
-			InSystem->ReleaseProducer(ProducersToRelease[ProducerIndex]);
-		}
+		ProducersToRelease.AddUnique(Desc.ProducerHandle[LayerIndex]);
 	}
-
-	/** Remaps the page mappings from one allocated virtual texture to another. */
-	static void RemapVirtualTexturePages(FVirtualTextureSystem* InSystem, FAllocatedVirtualTexture* OldAllocatedVT, FAllocatedVirtualTexture* NewAllocatedVT, uint32 InFrame)
+	InSystem->DestroyVirtualTexture(InAllocatedVT);
+	for (int32 ProducerIndex = 0; ProducerIndex < ProducersToRelease.Num(); ++ProducerIndex)
 	{
-		const uint32 OldVirtualAddress = OldAllocatedVT->GetVirtualAddress();
-		const uint32 NewVirtualAddress = NewAllocatedVT->GetVirtualAddress();
+		InSystem->ReleaseProducer(ProducersToRelease[ProducerIndex]);
+	}
+}
 
-		for (uint32 ProducerIndex = 0u; ProducerIndex < OldAllocatedVT->GetNumUniqueProducers(); ++ProducerIndex)
-		{
-			check(OldAllocatedVT->GetUniqueProducerMipBias(ProducerIndex) == 0);
-			check(NewAllocatedVT->GetUniqueProducerMipBias(ProducerIndex) == 0);
+/** Remaps the page mappings from one allocated virtual texture to another. */
+void FAdaptiveVirtualTexture::RemapVirtualTexturePages(FVirtualTextureSystem* InSystem, FAllocatedVirtualTexture* OldAllocatedVT, FAllocatedVirtualTexture* NewAllocatedVT, uint32 InFrame)
+{
+	const uint32 OldVirtualAddress = OldAllocatedVT->GetVirtualAddress();
+	const uint32 NewVirtualAddress = NewAllocatedVT->GetVirtualAddress();
 
-			const FVirtualTextureProducerHandle& OldProducerHandle = OldAllocatedVT->GetUniqueProducerHandle(ProducerIndex);
-			const FVirtualTextureProducerHandle& NewProducerHandle = NewAllocatedVT->GetUniqueProducerHandle(ProducerIndex);
+	for (uint32 ProducerIndex = 0u; ProducerIndex < OldAllocatedVT->GetNumUniqueProducers(); ++ProducerIndex)
+	{
+		check(OldAllocatedVT->GetUniqueProducerMipBias(ProducerIndex) == 0);
+		check(NewAllocatedVT->GetUniqueProducerMipBias(ProducerIndex) == 0);
 
-			FVirtualTextureProducer* OldProducer = InSystem->FindProducer(OldProducerHandle);
-			FVirtualTextureProducer* NewProducer = InSystem->FindProducer(NewProducerHandle);
+		const FVirtualTextureProducerHandle& OldProducerHandle = OldAllocatedVT->GetUniqueProducerHandle(ProducerIndex);
+		const FVirtualTextureProducerHandle& NewProducerHandle = NewAllocatedVT->GetUniqueProducerHandle(ProducerIndex);
 
-			if (OldProducer->GetDescription().bPersistentHighestMip)
-			{
-				InSystem->ForceUnlockAllTiles(OldProducerHandle, OldProducer);
-			}
+		FVirtualTextureProducer* OldProducer = InSystem->FindProducer(OldProducerHandle);
+		FVirtualTextureProducer* NewProducer = InSystem->FindProducer(NewProducerHandle);
 
-			const uint32 SpaceID = OldAllocatedVT->GetSpaceID();
-			const int32 vLevelBias = (int32)NewProducer->GetMaxLevel() - (int32)OldProducer->GetMaxLevel();
+		if (OldProducer->GetDescription().bPersistentHighestMip)
+		{
+			InSystem->ForceUnlockAllTiles(OldProducerHandle, OldProducer);
+		}
 
-			for (uint32 PhysicalGroupIndex = 0u; PhysicalGroupIndex < OldProducer->GetNumPhysicalGroups(); ++PhysicalGroupIndex)
-			{
-				FVirtualTexturePhysicalSpace* PhysicalSpace = OldProducer->GetPhysicalSpaceForPhysicalGroup(PhysicalGroupIndex);
-				FTexturePagePool& PagePool = PhysicalSpace->GetPagePool();
+		const uint32 SpaceID = OldAllocatedVT->GetSpaceID();
+		const int32 vLevelBias = (int32)NewProducer->GetMaxLevel() - (int32)OldProducer->GetMaxLevel();
 
-				PagePool.RemapPages(InSystem, SpaceID, PhysicalSpace, OldProducerHandle, OldVirtualAddress, NewProducerHandle, NewVirtualAddress, vLevelBias, InFrame);
-			}
+		for (uint32 PhysicalGroupIndex = 0u; PhysicalGroupIndex < OldProducer->GetNumPhysicalGroups(); ++PhysicalGroupIndex)
+		{
+			FVirtualTexturePhysicalSpace* PhysicalSpace = OldProducer->GetPhysicalSpaceForPhysicalGroup(PhysicalGroupIndex);
+			FTexturePagePool& PagePool = PhysicalSpace->GetPagePool();
+
+			PagePool.RemapPages(InSystem, SpaceID, PhysicalSpace, OldProducerHandle, OldVirtualAddress, NewProducerHandle, NewVirtualAddress, vLevelBias, InFrame);
 		}
 	}
 }
 
-
 FAdaptiveVirtualTexture::FAdaptiveVirtualTexture(
 	FAdaptiveVTDescription const& InAdaptiveDesc,
 	FAllocatedVTDescription const& InAllocatedDesc)
diff --git a/Engine/Source/Runtime/Renderer/Private/VT/AdaptiveVirtualTexture.h b/Engine/Source/Runtime/Renderer/Private/VT/AdaptiveVirtualTexture.h
--- a/Engine/Source/Runtime/Renderer/Private/VT/AdaptiveVirtualTexture.h
+++ b/Engine/Source/Runtime/Renderer/Private/VT/AdaptiveVirtualTexture.h
@@ -70,6 +70,19 @@ private:
 	/** Free or reduce and reallocate the least recently used allocation. */
 	bool FreeLRU(FVirtualTextureSystem* InSystem, uint32 InFrame, uint32 InFrameUnusedThreshold);
 
+	static IAllocatedVirtualTexture* AllocateVirtualTexture(
+		FVirtualTextureSystem* InSystem,
+		FAllocatedVTDescription const& InAllocatedDesc,
+		FIntPoint InGridSize,
+		uint8 InForcedSpaceID,
+		int32 InWidthInTiles,
+		int32 InHeightInTiles,
+		FIntPoint InAddressOffset,
+		int32 InLevelOffset);
+
+	static void DestroyVirtualTexture(FVirtualTextureSystem* InSystem, IAllocatedVirtualTexture* InAllocatedVT);
+	static void RemapVirtualTexturePages(FVirtualTextureSystem* InSystem, FAllocatedVirtualTexture* OldAllocatedVT, FAllocatedVirtualTexture* NewAllocatedVT, uint32 InFrame);
+
 private:
 	/** Adaptive virtual texture description. */
 	FAdaptiveVTDescription AdaptiveDesc;
diff --git a/Engine/Source/Runtime/Renderer/Private/VT/AllocatedVirtualTexture.cpp b/Engine/Source/Runtime/Renderer/Private/VT/AllocatedVirtualTexture.cpp
--- a/Engine/Source/Runtime/Renderer/Private/VT/AllocatedVirtualTexture.cpp
+++ b/Engine/Source/Runtime/Renderer/Private/VT/AllocatedVirtualTexture.cpp
@@ -21,7 +21,6 @@ FAllocatedVirtualTexture::FAllocatedVirtualTexture(FVirtualTextureSystem* InSyst
 	, FrameAllocated(InFrame)
 	, Space(nullptr)
 {
-	check(IsInRenderingThread());
 	FMemory::Memzero(TextureLayers);
 	FMemory::Memzero(FallbackColorPerTextureLayer);
 
@@ -114,7 +113,6 @@ void FAllocatedVirtualTexture::AssignVirtualAddress(uint32 vAddress)
 
 void FAllocatedVirtualTexture::Destroy(FVirtualTextureSystem* System)
 {
-	check(IsInRenderingThread());
 	check(NumRefs == 0);
 
 	// Unlock any locked tiles
diff --git a/Engine/Source/Runtime/Renderer/Private/VT/RuntimeVirtualTextureRender.cpp b/Engine/Source/Runtime/Renderer/Private/VT/RuntimeVirtualTextureRender.cpp
--- a/Engine/Source/Runtime/Renderer/Private/VT/RuntimeVirtualTextureRender.cpp
+++ b/Engine/Source/Runtime/Renderer/Private/VT/RuntimeVirtualTextureRender.cpp
@@ -88,7 +88,6 @@ namespace RuntimeVirtualTexture
 
 	const TUniformBufferRef<FEtcParameters>& GetEtcParametersUniformBufferRef()
 	{
-		check(IsInRenderingThread());
 		static TGlobalResource<FEtcParametersUniformBuffer> EtcParametersUniformBuffer;
 		return EtcParametersUniformBuffer.GetUniformBufferRef();
 	}
diff --git a/Engine/Source/Runtime/Renderer/Private/VT/RuntimeVirtualTextureSceneProxy.cpp b/Engine/Source/Runtime/Renderer/Private/VT/RuntimeVirtualTextureSceneProxy.cpp
--- a/Engine/Source/Runtime/Renderer/Private/VT/RuntimeVirtualTextureSceneProxy.cpp
+++ b/Engine/Source/Runtime/Renderer/Private/VT/RuntimeVirtualTextureSceneProxy.cpp
@@ -127,8 +127,6 @@ FRuntimeVirtualTextureSceneProxy::FRuntimeVirtualTextureSceneProxy(URuntimeVirtu
 
 FRuntimeVirtualTextureSceneProxy::~FRuntimeVirtualTextureSceneProxy()
 {
-	checkSlow(IsInRenderingThread());
-
 #if !UE_BUILD_SHIPPING
 	FRendererOnScreenNotification::Get().Remove(OnScreenWarningDelegateHandle);
 #endif
@@ -136,8 +134,6 @@ FRuntimeVirtualTextureSceneProxy::~FRuntimeVirtualTextureSceneProxy()
 
 void FRuntimeVirtualTextureSceneProxy::Release()
 {
-	checkSlow(!IsInRenderingThread());
-
 	if (VirtualTexture != nullptr)
 	{
 		VirtualTexture->Release();
diff --git a/Engine/Source/Runtime/Renderer/Private/VT/VirtualTextureProducer.cpp b/Engine/Source/Runtime/Renderer/Private/VT/VirtualTextureProducer.cpp
--- a/Engine/Source/Runtime/Renderer/Private/VT/VirtualTextureProducer.cpp
+++ b/Engine/Source/Runtime/Renderer/Private/VT/VirtualTextureProducer.cpp
@@ -55,7 +55,6 @@ FVirtualTextureProducerCollection::FVirtualTextureProducerCollection() : NumPend
 
 FVirtualTextureProducerHandle FVirtualTextureProducerCollection::RegisterProducer(FVirtualTextureSystem* System, const FVTProducerDescription& InDesc, IVirtualTexture* InProducer)
 {
-	check(IsInRenderingThread());
 	const uint32 ProducerWidth = InDesc.BlockWidthInTiles * InDesc.WidthInBlocks * InDesc.TileSize;
 	const uint32 ProducerHeight = InDesc.BlockHeightInTiles * InDesc.HeightInBlocks * InDesc.TileSize;
 	check(ProducerWidth > 0u);
@@ -100,8 +99,6 @@ FVirtualTextureProducerHandle FVirtualTextureProducerCollection::RegisterProduce
 
 void FVirtualTextureProducerCollection::ReleaseProducer(FVirtualTextureSystem* System, const FVirtualTextureProducerHandle& Handle)
 {
-	check(IsInRenderingThread());
-
 	if (FProducerEntry* Entry = GetEntry(Handle))
 	{
 		uint32 CallbackIndex = Callbacks[Entry->DestroyedCallbacksIndex].NextIndex;
@@ -182,7 +179,6 @@ bool FVirtualTextureProducerCollection::HasPendingCallbacks() const
 
 void FVirtualTextureProducerCollection::AddDestroyedCallback(const FVirtualTextureProducerHandle& Handle, FVTProducerDestroyedFunction* Function, void* Baton)
 {
-	check(IsInRenderingThread());
 	check(Function);
 	check(Baton);
 
@@ -202,7 +198,6 @@ void FVirtualTextureProducerCollection::AddDestroyedCallback(const FVirtualTextu
 
 uint32 FVirtualTextureProducerCollection::RemoveAllCallbacks(const void* Baton)
 {
-	check(IsInRenderingThread());
 	check(Baton);
 
 	uint32 NumRemoved = 0u;
diff --git a/Engine/Source/Runtime/Renderer/Private/VT/VirtualTextureProducer.h b/Engine/Source/Runtime/Renderer/Private/VT/VirtualTextureProducer.h
--- a/Engine/Source/Runtime/Renderer/Private/VT/VirtualTextureProducer.h
+++ b/Engine/Source/Runtime/Renderer/Private/VT/VirtualTextureProducer.h
@@ -208,5 +208,6 @@ private:
 	TArray<FProducerEntry> Producers;
 	TArray<FCallbackEntry> Callbacks;
 	TMap<void*, TArray<uint32>> CallbacksMap;
-	uint32 NumPendingCallbacks;
+	UE::FMutex CallbacksMutex;
+	std::atomic_uint32_t NumPendingCallbacks;
 };
diff --git a/Engine/Source/Runtime/Renderer/Private/VT/VirtualTextureSystem.cpp b/Engine/Source/Runtime/Renderer/Private/VT/VirtualTextureSystem.cpp
--- a/Engine/Source/Runtime/Renderer/Private/VT/VirtualTextureSystem.cpp
+++ b/Engine/Source/Runtime/Renderer/Private/VT/VirtualTextureSystem.cpp
@@ -398,7 +398,7 @@ void FVirtualTextureSystem::FlushCache()
 void FVirtualTextureSystem::FlushCache(FVirtualTextureProducerHandle const& ProducerHandle, int32 SpaceID, FIntRect const& TextureRegion, uint32 MaxLevel)
 {
 	check(!bUpdating);
-	checkSlow(IsInRenderingThread());
+	UE::TScopeLock Lock(Mutex);
 
 	SCOPE_CYCLE_COUNTER(STAT_FlushCache);
 	INC_DWORD_STAT_BY(STAT_NumFlushCache, 1);
@@ -600,12 +600,13 @@ IAllocatedVirtualTexture* FVirtualTextureSystem::AllocateVirtualTexture(const FA
 {
 	check(Desc.NumTextureLayers <= VIRTUALTEXTURE_SPACE_MAXLAYERS);
 
+	UE::TScopeLock Lock(Mutex);
+
 	// Check to see if we already have an allocated VT that matches this description
 	// This can happen often as multiple material instances will share the same textures
 	FAllocatedVirtualTexture*& AllocatedVT = AllocatedVTs.FindOrAdd(Desc);
 	if (AllocatedVT)
 	{
-		FScopeLock Lock(&AllocatedVTLock);
 		const int32 PrevNumRefs = AllocatedVT->NumRefs++;
 		check(PrevNumRefs >= 0);
 		if (PrevNumRefs == 0)
@@ -702,7 +703,7 @@ IAllocatedVirtualTexture* FVirtualTextureSystem::AllocateVirtualTexture(const FA
 
 void FVirtualTextureSystem::DestroyVirtualTexture(IAllocatedVirtualTexture* AllocatedVT)
 {
-	FScopeLock Lock(&AllocatedVTLock);
+	UE::TScopeLock Lock(Mutex);
 	const int32 NewNumRefs = --AllocatedVT->NumRefs;
 	check(NewNumRefs >= 0);
 	if (NewNumRefs == 0)
@@ -714,13 +715,10 @@ void FVirtualTextureSystem::DestroyVirtualTexture(IAllocatedVirtualTexture* Allo
 
 void FVirtualTextureSystem::DestroyPendingVirtualTextures(bool bForceDestroyAll)
 {
-	check(IsInRenderingThread());
-
 	TRACE_CPUPROFILER_EVENT_SCOPE(FVirtualTextureSystem::DestroyPendingVirtualTextures);
 
 	TArray<IAllocatedVirtualTexture*> AllocatedVTsToDelete;
 	{
-		FScopeLock Lock(&AllocatedVTLock);
 		if (bForceDestroyAll)
 		{
 			AllocatedVTsToDelete = MoveTemp(PendingDeleteAllocatedVTs);
@@ -777,7 +775,7 @@ void FVirtualTextureSystem::DestroyPendingVirtualTextures(bool bForceDestroyAll)
 IAdaptiveVirtualTexture* FVirtualTextureSystem::AllocateAdaptiveVirtualTexture(const FAdaptiveVTDescription& AdaptiveVTDesc, const FAllocatedVTDescription& AllocatedVTDesc)
 {
 	check(!bUpdating);
-	check(IsInRenderingThread());
+	UE::TScopeLock Lock(Mutex);
 	FAdaptiveVirtualTexture* AdaptiveVT = new FAdaptiveVirtualTexture(AdaptiveVTDesc, AllocatedVTDesc);
 	AdaptiveVT->Init(this);
 	check(AdaptiveVTs[AdaptiveVT->GetSpaceID()] == nullptr);
@@ -788,7 +786,7 @@ IAdaptiveVirtualTexture* FVirtualTextureSystem::AllocateAdaptiveVirtualTexture(c
 void FVirtualTextureSystem::DestroyAdaptiveVirtualTexture(IAdaptiveVirtualTexture* AdaptiveVT)
 {
 	check(!bUpdating);
-	check(IsInRenderingThread());
+	UE::TScopeLock Lock(Mutex);
 	check(AdaptiveVTs[AdaptiveVT->GetSpaceID()] == AdaptiveVT);
 	AdaptiveVTs[AdaptiveVT->GetSpaceID()] = nullptr;
 	AdaptiveVT->Destroy(this);
@@ -797,24 +795,28 @@ void FVirtualTextureSystem::DestroyAdaptiveVirtualTexture(IAdaptiveVirtualTextur
 FVirtualTextureProducerHandle FVirtualTextureSystem::RegisterProducer(const FVTProducerDescription& InDesc, IVirtualTexture* InProducer)
 {
 	check(!bUpdating);
+	UE::TScopeLock Lock(Mutex);
 	return Producers.RegisterProducer(this, InDesc, InProducer);
 }
 
 void FVirtualTextureSystem::ReleaseProducer(const FVirtualTextureProducerHandle& Handle)
 {
 	check(!bUpdating);
+	UE::TScopeLock Lock(Mutex);
 	Producers.ReleaseProducer(this, Handle);
 }
 
 void FVirtualTextureSystem::AddProducerDestroyedCallback(const FVirtualTextureProducerHandle& Handle, FVTProducerDestroyedFunction* Function, void* Baton)
 {
 	check(!bUpdating);
+	UE::TScopeLock Lock(Mutex);
 	Producers.AddDestroyedCallback(Handle, Function, Baton);
 }
 
 uint32 FVirtualTextureSystem::RemoveAllProducerDestroyedCallbacks(const void* Baton)
 {
 	check(!bUpdating);
+	UE::TScopeLock Lock(Mutex);
 	return Producers.RemoveAllCallbacks(Baton);
 }
 
@@ -892,7 +894,6 @@ FVirtualTextureSpace* FVirtualTextureSystem::AcquireSpace(const FVTSpaceDescript
 void FVirtualTextureSystem::ReleaseSpace(FVirtualTextureSpace* Space)
 {
 	check(!bUpdating);
-	check(IsInRenderingThread());
 	const uint32 NumRefs = Space->Release();
 	if (NumRefs == 0u && Space->GetDescription().bPrivateSpace)
 	{
@@ -982,7 +983,6 @@ void GetPoolInitDescription(FVTPhysicalSpaceDescription const& InDesc, FVTPhysic
 /** Cached version of GetPoolInitDescription() to avoid regularly repeating the heavy work in that function. */
 void GetPoolInitDescription_Cached(FVTPhysicalSpaceDescription const& InDesc, FVTPhysicalSpaceInitDescription& OutInitDescription)
 {
-	check(IsInRenderingThread());
 	static TMap<FVTPhysicalSpaceDescription, FVTPhysicalSpaceInitDescription> Map;
 
 	// Invalidate the cache if any relevant CVar settings change.
@@ -1059,7 +1059,6 @@ FVirtualTexturePhysicalSpace* FVirtualTextureSystem::AcquirePhysicalSpace(const
 
 void FVirtualTextureSystem::ReleasePendingSpaces()
 {
-	check(IsInRenderingThread());
 	for (int32 Id = 0; Id < PhysicalSpaces.Num(); ++Id)
 	{
 		// Physical space is released when ref count hits 0
@@ -1083,7 +1082,6 @@ void FVirtualTextureSystem::ReleasePendingSpaces()
 void FVirtualTextureSystem::LockTile(const FVirtualTextureLocalTile& Tile)
 {
 	check(!bUpdating);
-	check(IsInRenderingThread());
 
 	if (TileLocks.Lock(Tile))
 	{
@@ -1109,7 +1107,6 @@ static void UnlockTileInternal(const FVirtualTextureProducerHandle& ProducerHand
 void FVirtualTextureSystem::UnlockTile(const FVirtualTextureLocalTile& Tile, const FVirtualTextureProducer* Producer)
 {
 	check(!bUpdating);
-	check(IsInRenderingThread());
 
 	if (TileLocks.Unlock(Tile))
 	{
@@ -1127,7 +1124,6 @@ void FVirtualTextureSystem::UnlockTile(const FVirtualTextureLocalTile& Tile, con
 void FVirtualTextureSystem::ForceUnlockAllTiles(const FVirtualTextureProducerHandle& ProducerHandle, const FVirtualTextureProducer* Producer)
 {
 	check(!bUpdating);
-	check(IsInRenderingThread());
 
 	TArray<FVirtualTextureLocalTile> TilesToUnlock;
 	TileLocks.ForceUnlockAll(ProducerHandle, TilesToUnlock);
@@ -1157,9 +1153,7 @@ static float ComputeMipLevel(const IAllocatedVirtualTexture* AllocatedVT, const
 void FVirtualTextureSystem::RequestTiles(const FVector2D& InScreenSpaceSize, int32 InMipLevel)
 {
 	check(!bUpdating);
-	check(IsInRenderingThread());
-
-	FScopeLock Lock(&RequestedTilesLock);
+	UE::TScopeLock Lock(Mutex);
 
 	for (const auto& Pair : AllocatedVTs)
 	{
@@ -1170,9 +1164,7 @@ void FVirtualTextureSystem::RequestTiles(const FVector2D& InScreenSpaceSize, int
 void FVirtualTextureSystem::RequestTiles(const FMaterialRenderProxy* InMaterialRenderProxy, const FVector2D& InScreenSpaceSize, ERHIFeatureLevel::Type InFeatureLevel)
 {
 	check(!bUpdating);
-	check(IsInRenderingThread());
-
-	FScopeLock Lock(&RequestedTilesLock);
+	UE::TScopeLock Lock(Mutex);
 
 	for (IAllocatedVirtualTexture* AllocatedVT : InMaterialRenderProxy->UniformExpressionCache[InFeatureLevel].AllocatedVTs)
 	{
@@ -1206,9 +1198,9 @@ void FVirtualTextureSystem::RequestTilesInternal(const IAllocatedVirtualTexture*
 
 void FVirtualTextureSystem::RequestTilesForRegion(IAllocatedVirtualTexture* AllocatedVT, const FVector2D& InScreenSpaceSize, const FVector2D& InViewportPosition, const FVector2D& InViewportSize, const FVector2D& InUV0, const FVector2D& InUV1, int32 InMipLevel)
 {
+	UE::TScopeLock Lock(Mutex);
 	if (InMipLevel >= 0)
 	{
-		FScopeLock Lock(&RequestedTilesLock);
 		RequestTilesForRegionInternal(AllocatedVT, InScreenSpaceSize, InViewportPosition, InViewportSize, InUV0, InUV1, InMipLevel);
 	}
 	else
@@ -1217,7 +1209,6 @@ void FVirtualTextureSystem::RequestTilesForRegion(IAllocatedVirtualTexture* Allo
 		const float vLevel = ComputeMipLevel(AllocatedVT, InScreenSpaceSize);
 		const int32 vMipLevelDown = FMath::Clamp((int32)FMath::FloorToInt(vLevel), 0, (int32)vMaxLevel);
 
-		FScopeLock Lock(&RequestedTilesLock);
 		RequestTilesForRegionInternal(AllocatedVT, InScreenSpaceSize, InViewportPosition, InViewportSize, InUV0, InUV1, vMipLevelDown);
 		if (vMipLevelDown + 1u <= vMaxLevel)
 		{
@@ -1230,12 +1221,11 @@ void FVirtualTextureSystem::RequestTilesForRegion(IAllocatedVirtualTexture* Allo
 void FVirtualTextureSystem::LoadPendingTiles(FRDGBuilder& GraphBuilder, ERHIFeatureLevel::Type FeatureLevel)
 {
 	check(!bUpdating);
-	check(IsInRenderingThread());
+	UE::TScopeLock Lock(Mutex);
 
 	TArray<uint32> PackedTiles;
 	if (RequestedPackedTiles.Num() > 0)
 	{
-		FScopeLock Lock(&RequestedTilesLock);
 		PackedTiles = MoveTemp(RequestedPackedTiles);
 		RequestedPackedTiles.Reset();
 	}
@@ -2597,7 +2587,6 @@ void FVirtualTextureSystem::GatherPackedTileRequests(FConcurrentLinearBulkObject
 	TArray<uint32> PackedTiles;
 	if (RequestedPackedTiles.Num() > 0)
 	{
-		FScopeLock Lock(&RequestedTilesLock);
 		PackedTiles = MoveTemp(RequestedPackedTiles);
 		RequestedPackedTiles.Reset();
 	}
@@ -2633,6 +2622,8 @@ void FVirtualTextureSystem::BeginUpdate(FRDGBuilder& GraphBuilder, FVirtualTextu
 
 		const bool bContinousUpdates = false;
 
+		UE::TScopeLock Lock(Mutex);
+
 		Updater->MergedRequestList = Allocator.Create<FUniqueRequestList>(Allocator);
 		Updater->MergedRequestList->Initialize();
 
@@ -2651,7 +2642,7 @@ void FVirtualTextureSystem::CallPendingCallbacks()
 {
 	TRACE_CPUPROFILER_EVENT_SCOPE(FVirtualTextureSystem::CallPendingCallbacks);
 	SCOPE_CYCLE_COUNTER(STAT_VirtualTextureSystem_Update);
-
+	UE::TScopeLock Lock(Mutex);
 	Producers.CallPendingCallbacks();
 }
 
@@ -2659,6 +2650,9 @@ TUniquePtr<FVirtualTextureUpdater> FVirtualTextureSystem::BeginUpdate(FRDGBuilde
 {
 	check(IsInRenderingThread());
 	check(!bUpdating);
+
+	UE::TScopeLock Lock(Mutex);
+
 	checkf(Producers.HasPendingCallbacks() == false, TEXT("FVirtualTextureSystem::CallPendingCallbacks(), typically called in FSceneRenderer::UpdateScene(), must run before FVirtualTextureSystem::BeginUpdate()"));
 
 	AllocateResources(GraphBuilder);
@@ -2785,6 +2779,8 @@ void FVirtualTextureSystem::EndUpdate(FRDGBuilder& GraphBuilder, TUniquePtr<FVir
 	}
 	bUpdating = false;
 
+	UE::TScopeLock Lock(Mutex);
+
 	if (Updater->FeedbackMapResult.Data)
 	{
 		GVirtualTextureFeedback.Unmap(GraphBuilder.RHICmdList, Updater->FeedbackMapResult.MapHandle);
@@ -2833,12 +2829,15 @@ void FVirtualTextureSystem::Update(FRDGBuilder& GraphBuilder, ERHIFeatureLevel::
 void FVirtualTextureSystem::ReleasePendingResources()
 {
 	check(!bUpdating);
+	UE::TScopeLock Lock(Mutex);
 	DestroyPendingVirtualTextures(true);
 	ReleasePendingSpaces();
 }
 
 float FVirtualTextureSystem::GetGlobalMipBias() const
 {
+	UE::TScopeLock Lock(Mutex);
+
 	float MaxResidencyMipMapBias = 0.f;
 	for (int32 SpaceIndex = 0; SpaceIndex < PhysicalSpaces.Num(); ++SpaceIndex)
 	{
@@ -2852,6 +2851,7 @@ float FVirtualTextureSystem::GetGlobalMipBias() const
 
 bool FVirtualTextureSystem::IsPendingRootPageMap(IAllocatedVirtualTexture* AllocatedVT) const
 {
+	UE::TScopeLock Lock(Mutex);
 	return AllocatedVTsToMap.Find(AllocatedVT) != INDEX_NONE;
 }
 
@@ -2971,7 +2971,7 @@ void FVirtualTextureSystem::DrawResidencyHud(UCanvas* InCanvas, APlayerControlle
 
 void FVirtualTextureSystem::SetVirtualTextureRequestRecordBuffer(uint64 Handle)
 {
-	check(IsInRenderingThread());
+	UE::TScopeLock Lock(Mutex);
 	check(PageRequestRecordHandle == ~0ull && PageRequestRecordBuffer.Num() == 0);
 	
 	PageRequestRecordHandle = Handle;
@@ -3016,7 +3016,7 @@ void FVirtualTextureSystem::RecordPageRequests(FUniquePageList const* UniquePage
 
 uint64 FVirtualTextureSystem::GetVirtualTextureRequestRecordBuffer(TSet<uint64>& OutPageRequests)
 {
-	check(IsInRenderingThread());
+	UE::TScopeLock Lock(Mutex);
 
 	if (PageRequestRecordHandle == ~0ull)
 	{
@@ -3034,7 +3034,7 @@ uint64 FVirtualTextureSystem::GetVirtualTextureRequestRecordBuffer(TSet<uint64>&
 
 void FVirtualTextureSystem::RequestRecordedTiles(TArray<uint64>&& InPageRequests)
 {
-	check(IsInRenderingThread());
+	UE::TScopeLock Lock(Mutex);
 
 	if (PageRequestPlaybackBuffer.Num() == 0)
 	{
diff --git a/Engine/Source/Runtime/Renderer/Private/VT/VirtualTextureSystem.h b/Engine/Source/Runtime/Renderer/Private/VT/VirtualTextureSystem.h
--- a/Engine/Source/Runtime/Renderer/Private/VT/VirtualTextureSystem.h
+++ b/Engine/Source/Runtime/Renderer/Private/VT/VirtualTextureSystem.h
@@ -14,6 +14,7 @@
 #include "VirtualTexturing.h"
 #include "VT/VirtualTextureFeedback.h"
 #include "Tasks/Task.h"
+#include "Async/RecursiveMutex.h"
 
 class FAdaptiveVirtualTexture;
 class FAllocatedVirtualTexture;
@@ -119,23 +120,10 @@ public:
 	void ReleaseProducer(const FVirtualTextureProducerHandle& Handle);
 	void AddProducerDestroyedCallback(const FVirtualTextureProducerHandle& Handle, FVTProducerDestroyedFunction* Function, void* Baton);
 	uint32 RemoveAllProducerDestroyedCallbacks(const void* Baton);
-	FVirtualTextureProducer* FindProducer(const FVirtualTextureProducerHandle& Handle);
 
 	IAdaptiveVirtualTexture* AllocateAdaptiveVirtualTexture(const FAdaptiveVTDescription& AdaptiveVTDesc, const FAllocatedVTDescription& AllocatedVTDesc);
 	void DestroyAdaptiveVirtualTexture(IAdaptiveVirtualTexture* AdaptiveVT);
 
-	FVirtualTextureSpace* AcquireSpace(const FVTSpaceDescription& InDesc, uint8 InForceSpaceID, FAllocatedVirtualTexture* AllocatedVT);
-	void ReleaseSpace(FVirtualTextureSpace* Space);
-
-	FVirtualTexturePhysicalSpace* AcquirePhysicalSpace(const FVTPhysicalSpaceDescription& InDesc);
-
-	FVirtualTextureSpace* GetSpace(uint8 ID) const { check(ID < MaxSpaces); return Spaces[ID].Get(); }
-	FAdaptiveVirtualTexture* GetAdaptiveVirtualTexture(uint8 ID) const { check(ID < MaxSpaces); return AdaptiveVTs[ID]; }
-	FVirtualTexturePhysicalSpace* GetPhysicalSpace(uint16 ID) const { check(PhysicalSpaces[ID]);  return PhysicalSpaces[ID]; }
-
-	void LockTile(const FVirtualTextureLocalTile& Tile);
-	void UnlockTile(const FVirtualTextureLocalTile& Tile, const FVirtualTextureProducer* Producer);
-	void ForceUnlockAllTiles(const FVirtualTextureProducerHandle& ProducerHandle, const FVirtualTextureProducer* Producer);
 	void RequestTiles(const FVector2D& InScreenSpaceSize, int32 InMipLevel = -1);
 	void RequestTiles(const FMaterialRenderProxy* InMaterialRenderProxy, const FVector2D& InScreenSpaceSize, ERHIFeatureLevel::Type InFeatureLevel);
 	void RequestTilesForRegion(IAllocatedVirtualTexture* AllocatedVT, const FVector2D& InScreenSpaceSize, const FVector2D& InViewportPosition, const FVector2D& InViewportSize, const FVector2D& InUV0, const FVector2D& InUV1, int32 InMipLevel = -1);
@@ -158,10 +146,30 @@ private:
 	friend class FFeedbackAnalysisTask;
 	friend class FAddRequestedTilesTask;
 	friend class FGatherRequestsTask;
+	friend class FAllocatedVirtualTexture;
+	friend class FAdaptiveVirtualTexture;
+	friend class FVirtualTextureProducer;
+	friend class FVirtualTextureProducerCollection;
+	friend class FTexturePageMap;
+	friend class FTexturePagePool;
 
 	FVirtualTextureSystem();
 	~FVirtualTextureSystem();
 
+	FVirtualTextureSpace* AcquireSpace(const FVTSpaceDescription& InDesc, uint8 InForceSpaceID, FAllocatedVirtualTexture* AllocatedVT);
+	void ReleaseSpace(FVirtualTextureSpace* Space);
+
+	FVirtualTextureProducer* FindProducer(const FVirtualTextureProducerHandle& Handle);
+	FVirtualTexturePhysicalSpace* AcquirePhysicalSpace(const FVTPhysicalSpaceDescription& InDesc);
+
+	FVirtualTextureSpace* GetSpace(uint8 ID) const { check(ID < MaxSpaces); return Spaces[ID].Get(); }
+	FAdaptiveVirtualTexture* GetAdaptiveVirtualTexture(uint8 ID) const { check(ID < MaxSpaces); return AdaptiveVTs[ID]; }
+	FVirtualTexturePhysicalSpace* GetPhysicalSpace(uint16 ID) const { check(PhysicalSpaces[ID]);  return PhysicalSpaces[ID]; }
+
+	void LockTile(const FVirtualTextureLocalTile& Tile);
+	void UnlockTile(const FVirtualTextureLocalTile& Tile, const FVirtualTextureProducer* Producer);
+	void ForceUnlockAllTiles(const FVirtualTextureProducerHandle& ProducerHandle, const FVirtualTextureProducer* Producer);
+
 	void BeginUpdate(FRDGBuilder& GraphBuilder, FVirtualTextureUpdater* Updater);
 
 	void AllocateResources(FRDGBuilder& GraphBuilder);
@@ -202,6 +210,8 @@ private:
 
 	uint32	Frame;
 
+	mutable UE::FRecursiveMutex Mutex;
+
 	static const uint32 MaxNumTasks = 16;
 	static const uint32 MaxSpaces = 16;
 	uint32 NumAllocatedSpaces = 0;
@@ -209,7 +219,6 @@ private:
 	TArray<FVirtualTexturePhysicalSpace*> PhysicalSpaces;
 	FVirtualTextureProducerCollection Producers;
 
-	FCriticalSection AllocatedVTLock;
 	TArray<IAllocatedVirtualTexture*> PendingDeleteAllocatedVTs;
 
 	TMap<FAllocatedVTDescription, FAllocatedVirtualTexture*> AllocatedVTs;
@@ -237,7 +246,6 @@ private:
 	FAutoConsoleCommand SaveAllocatorImages;
 #endif
 
-	FCriticalSection RequestedTilesLock;
 	TArray<uint32> RequestedPackedTiles;
 
 	TArray<FVirtualTextureLocalTile> TilesToLock;
