Fixed build break in VT system.

diff --git a/Engine/Source/Runtime/Renderer/Private/VT/VirtualTextureProducer.h b/Engine/Source/Runtime/Renderer/Private/VT/VirtualTextureProducer.h
--- a/Engine/Source/Runtime/Renderer/Private/VT/VirtualTextureProducer.h
+++ b/Engine/Source/Runtime/Renderer/Private/VT/VirtualTextureProducer.h
@@ -208,6 +208,5 @@ private:
 	TArray<FProducerEntry> Producers;
 	TArray<FCallbackEntry> Callbacks;
 	TMap<void*, TArray<uint32>> CallbacksMap;
-	UE::FMutex CallbacksMutex;
-	std::atomic_uint32_t NumPendingCallbacks;
+	uint32 NumPendingCallbacks;
 };
