Removes the r.TSR.ShadingRejection.TileSize cvar that was accidentally checked-in

diff --git a/Engine/Source/Runtime/Renderer/Private/PostProcess/TemporalSuperResolution.cpp b/Engine/Source/Runtime/Renderer/Private/PostProcess/TemporalSuperResolution.cpp
--- a/Engine/Source/Runtime/Renderer/Private/PostProcess/TemporalSuperResolution.cpp
+++ b/Engine/Source/Runtime/Renderer/Private/PostProcess/TemporalSuperResolution.cpp
@@ -149,11 +149,6 @@ TAutoConsoleVariable<int32> CVarTSRShadingTileOverscan(
 	TEXT(" anti-aliasing scalability settings."),
 	ECVF_Scalability | ECVF_RenderThreadSafe);
 
-TAutoConsoleVariable<int32> CVarTSRShadingTileSize(
-	TEXT("r.TSR.ShadingRejection.TileSize"), 16,
-	TEXT(""),
-	ECVF_Scalability | ECVF_RenderThreadSafe);
-
 TAutoConsoleVariable<float> CVarTSRShadingExposureOffset(
 	TEXT("r.TSR.ShadingRejection.ExposureOffset"), 3.0,
 	TEXT("The shading rejection needs to have a representative idea how bright a linear color pixel ends up displayed to the user. ")
@@ -1928,7 +1923,7 @@ FDefaultTemporalUpscaler::FOutputs AddTemporalSuperResolutionPasses(
 	FRDGTextureRef HistoryRejectionTexture = nullptr;
 	FRDGTextureRef InputSceneColorLdrLumaTexture = nullptr;
 	{
-		const int32 GroupTileSize = CVarTSRShadingTileSize.GetValueOnRenderThread();
+		const int32 GroupTileSize = 16;
 
 		const bool bComputeInputSceneColorTexture = InputSceneColorLdrLumaTexture == nullptr;
 		if (bComputeInputSceneColorTexture)
