Fix VSM HZB not being freed from thumbnail scene renderers
Rename VSM "caching" -> "persistent data" show flag to better represent what it means/does
Remove unnecessary use of the flag from a couple places

diff --git a/Engine/Plugins/Experimental/Water/Source/Runtime/Private/WaterInfoRendering.cpp b/Engine/Plugins/Experimental/Water/Source/Runtime/Private/WaterInfoRendering.cpp
--- a/Engine/Plugins/Experimental/Water/Source/Runtime/Private/WaterInfoRendering.cpp
+++ b/Engine/Plugins/Experimental/Water/Source/Runtime/Private/WaterInfoRendering.cpp
@@ -671,7 +671,7 @@ static FSceneRenderer* CreateWaterInfoSceneRenderer(const FCreateWaterInfoSceneR
 	ShowFlags.DynamicShadows = 0;
 
 	ShowFlags.SetDisableOcclusionQueries(true);
-	ShowFlags.SetVirtualShadowMapCaching(false);
+	ShowFlags.SetVirtualShadowMapPersistentData(false);
 	
 	FSceneViewFamilyContext ViewFamily(FSceneViewFamily::ConstructionValues(
 		Params.RenderTarget,
diff --git a/Engine/Source/Runtime/Engine/Public/ShowFlags.h b/Engine/Source/Runtime/Engine/Public/ShowFlags.h
--- a/Engine/Source/Runtime/Engine/Public/ShowFlags.h
+++ b/Engine/Source/Runtime/Engine/Public/ShowFlags.h
@@ -224,8 +224,8 @@ struct FEngineShowFlags
 		SetLumenGlobalIllumination(false);
 		SetLumenReflections(false);
 
-		// TODO: Remove when Physical page pool size scales automatically with demand
-		SetVirtualShadowMapCaching(false);
+		// Have VSM drop all persistent data each frame
+		SetVirtualShadowMapPersistentData(false);
 
 		SetShaderPrint(false);
 	}
@@ -255,8 +255,9 @@ struct FEngineShowFlags
 		SetLumenGlobalIllumination(false);
 		SetLumenReflections(false);
 
-		// TODO: Remove when Physical page pool size scales automatically with demand
-		SetVirtualShadowMapCaching(false);
+		// Have VSM drop all persistent data each frame
+		// TODO: Revisit some of the cases that trigger this; if they clean up the scene renderers this is not necessary
+		SetVirtualShadowMapPersistentData(false);
 	}
 
 	bool IsVisualizeCalibrationEnabled() const
diff --git a/Engine/Source/Runtime/Engine/Public/ShowFlagsValues.inl b/Engine/Source/Runtime/Engine/Public/ShowFlagsValues.inl
--- a/Engine/Source/Runtime/Engine/Public/ShowFlagsValues.inl
+++ b/Engine/Source/Runtime/Engine/Public/ShowFlagsValues.inl
@@ -411,7 +411,7 @@ SHOWFLAG_FIXED_IN_SHIPPING(0, VisualizeVolumetricCloudConservativeDensity, SFG_V
 /** Visualize volumetric cloud density for empty space skipping. */
 SHOWFLAG_FIXED_IN_SHIPPING(0, VisualizeVolumetricCloudEmptySpaceSkipping, SFG_Visualize, NSLOCTEXT("UnrealEd", "VisualizeVolumetricCloudEmptySpaceSkippingSF", "Volumetric Cloud Empty Space Skipping Density"))
 
-SHOWFLAG_FIXED_IN_SHIPPING(1, VirtualShadowMapCaching, SFG_Developer, NSLOCTEXT("UnrealEd", "VirtualShadowMapCaching", "Cache Virtual Shadow Maps"))
+SHOWFLAG_FIXED_IN_SHIPPING(1, VirtualShadowMapPersistentData, SFG_Developer, NSLOCTEXT("UnrealEd", "VirtualShadowMapPersistentData", "Virtual Shadow Map Persistent Data"))
 
 SHOWFLAG_FIXED_IN_SHIPPING(0, DrawOnlyVSMInvalidatingGeo, SFG_Visualize, NSLOCTEXT("UnrealEd", "DrawOnlyVSMInvalidatingGeo", "Draw Only Geometry Causing VSM Invalidation"))
 SHOWFLAG_FIXED_IN_SHIPPING(0, DebugDrawDistantVirtualSMLights, SFG_Advanced, NSLOCTEXT("UnrealEd", "DebugDrawDistantVirtualSMLightsSF", "Debug Draw Distant VSM Lights"))
diff --git a/Engine/Source/Runtime/Renderer/Private/DeferredShadingRenderer.cpp b/Engine/Source/Runtime/Renderer/Private/DeferredShadingRenderer.cpp
--- a/Engine/Source/Runtime/Renderer/Private/DeferredShadingRenderer.cpp
+++ b/Engine/Source/Runtime/Renderer/Private/DeferredShadingRenderer.cpp
@@ -3767,7 +3767,7 @@ void FDeferredShadingSceneRenderer::Render(FRDGBuilder& GraphBuilder)
 					GraphBuilder,
 					VirtualShadowMapArray,
 					*this,
-					ViewFamily.EngineShowFlags.VirtualShadowMapCaching);
+					ViewFamily.EngineShowFlags.VirtualShadowMapPersistentData);
 			}
 		}
 
diff --git a/Engine/Source/Runtime/Renderer/Private/VirtualShadowMaps/VirtualShadowMapCacheManager.cpp b/Engine/Source/Runtime/Renderer/Private/VirtualShadowMaps/VirtualShadowMapCacheManager.cpp
--- a/Engine/Source/Runtime/Renderer/Private/VirtualShadowMaps/VirtualShadowMapCacheManager.cpp
+++ b/Engine/Source/Runtime/Renderer/Private/VirtualShadowMaps/VirtualShadowMapCacheManager.cpp
@@ -685,6 +685,7 @@ void FVirtualShadowMapArrayCacheManager::FreePhysicalPool(FRDGBuilder& GraphBuil
 	if (PhysicalPagePool)
 	{
 		PhysicalPagePool = nullptr;
+		PhysicalPageMetaData = nullptr;
 		Invalidate(GraphBuilder);
 	}
 }
@@ -905,12 +906,12 @@ void FVirtualShadowMapArrayCacheManager::ExtractFrameData(
 	FRDGBuilder& GraphBuilder,	
 	FVirtualShadowMapArray &VirtualShadowMapArray,
 	const FSceneRenderer& SceneRenderer,
-	bool bEnableCaching)
+	bool bAllowPersistentData)
 {
 	TrimLoggingInfo();
 
 	const bool bNewShadowData = VirtualShadowMapArray.IsAllocated();
-	const bool bDropAll = !bEnableCaching;
+	const bool bDropAll = !bAllowPersistentData;
 	const bool bDropPrevBuffers = bDropAll || bNewShadowData;
 
 	if (bDropPrevBuffers)
@@ -927,6 +928,7 @@ void FVirtualShadowMapArrayCacheManager::ExtractFrameData(
 		// thumbnail rendering or similar creates multiple FSceneRenderers that never get deleted.
 		// Caching is disabled on these contexts intentionally to avoid these issues.
 		FreePhysicalPool(GraphBuilder);
+		FreeHZBPhysicalPool(GraphBuilder);
 	}
 	else if (bNewShadowData)
 	{
diff --git a/Engine/Source/Runtime/Renderer/Private/VirtualShadowMaps/VirtualShadowMapCacheManager.h b/Engine/Source/Runtime/Renderer/Private/VirtualShadowMaps/VirtualShadowMapCacheManager.h
--- a/Engine/Source/Runtime/Renderer/Private/VirtualShadowMaps/VirtualShadowMapCacheManager.h
+++ b/Engine/Source/Runtime/Renderer/Private/VirtualShadowMaps/VirtualShadowMapCacheManager.h
@@ -220,13 +220,16 @@ public:
 	void UploadProjectionData(FRDGScatterUploadBuffer& Uploader) const;
 
 	/**
-	 * Call at end of frame to extract resouces from the virtual SM array to preserve to next frame.
-	 * If bCachingEnabled is false, all previous frame data is dropped and cache (and HZB!) data will not be available for the next frame.
-	 */ 
+	* Call at end of frame to extract resouces from the virtual SM array to preserve to next frame.
+	* 
+	* If bAllowPersistentData is false, all previous frame data is dropped and cache (and HZB!) data will not be available for the next frame.
+	* This flag is mostly intended for temporary editor resources like thumbnail rendering that will be used infrequently but often not properly destructed.
+	* We need to ensure that the VSM data associated with these renderer instances gets dropped.
+	*/ 
 	void ExtractFrameData(FRDGBuilder& GraphBuilder,
 		FVirtualShadowMapArray &VirtualShadowMapArray,
 		const FSceneRenderer& SceneRenderer,
-		bool bEnableCaching);
+		bool bAllowPersistentData);
 
 	/**
 	 * Finds an existing cache entry and moves to the active set or creates a fresh one.
