Add a separate upload limit r.VT.MaxUploadsPerFrame.Streaming for streaming VT pages. When the cvar is not zero, it enables the separate limit and r.VT.MaxUploadsPerFrame will only constrain non-streaming VT pages. When the cvar is zero (currently default), it disables the separate limit for backward compatibility. The motivation to separate SVT pages from other types of VT pages is mainly three-fold: 1) limiting SVT page uploads also throttles the number of loading requests. On platforms where I/O is slow, we may want to let more SVT page requests to go through so that I/O can be initiated earlier to reduce delay in obtaining page data; 2) SVT pages are usually cheaper to produce than RVT pages because they just upload the loaded data to GPU; 3) long standing SVT page requests can block the requests for other types of pages even if the other pages can be requested in parallel and turned around faster.

#tests a/b tested perf and visual on two console platforms

diff --git a/Engine/Source/Runtime/Engine/Private/VT/RuntimeVirtualTexture.cpp b/Engine/Source/Runtime/Engine/Private/VT/RuntimeVirtualTexture.cpp
--- a/Engine/Source/Runtime/Engine/Private/VT/RuntimeVirtualTexture.cpp
+++ b/Engine/Source/Runtime/Engine/Private/VT/RuntimeVirtualTexture.cpp
@@ -40,6 +40,11 @@ namespace
 		}
 
 		//~ Begin IVirtualTexture Interface.
+		virtual bool IsPageStreamed(uint8 vLevel, uint32 vAddress) const override
+		{
+			return false;
+		}
+
 		virtual FVTRequestPageResult RequestPageData(
 			FRHICommandList& RHICmdList,
 			const FVirtualTextureProducerHandle& ProducerHandle,
diff --git a/Engine/Source/Runtime/Engine/Private/VT/UploadingVirtualTexture.h b/Engine/Source/Runtime/Engine/Private/VT/UploadingVirtualTexture.h
--- a/Engine/Source/Runtime/Engine/Private/VT/UploadingVirtualTexture.h
+++ b/Engine/Source/Runtime/Engine/Private/VT/UploadingVirtualTexture.h
@@ -78,6 +78,7 @@ public:
 
 	// IVirtualTexture interface
 	virtual uint32 GetLocalMipBias(uint8 vLevel, uint32 vAddress) const override;
+	virtual bool IsPageStreamed(uint8 vLevel, uint32 vAddress) const override { return true; }
 	virtual FVTRequestPageResult RequestPageData(FRHICommandList& RHICmdList, const FVirtualTextureProducerHandle& ProducerHandle, uint8 LayerMask, uint8 vLevel, uint64 vAddress, EVTRequestPagePriority Priority) override;
 	virtual IVirtualTextureFinalizer* ProducePageData(FRHICommandList& RHICmdList,
 		ERHIFeatureLevel::Type FeatureLevel,
diff --git a/Engine/Source/Runtime/Engine/Private/VT/VirtualTextureLevelRedirector.cpp b/Engine/Source/Runtime/Engine/Private/VT/VirtualTextureLevelRedirector.cpp
--- a/Engine/Source/Runtime/Engine/Private/VT/VirtualTextureLevelRedirector.cpp
+++ b/Engine/Source/Runtime/Engine/Private/VT/VirtualTextureLevelRedirector.cpp
@@ -14,6 +14,13 @@ FVirtualTextureLevelRedirector::~FVirtualTextureLevelRedirector()
 	delete VirtualTextures[1];
 }
 
+bool FVirtualTextureLevelRedirector::IsPageStreamed(uint8 vLevel, uint32 vAddress) const
+{
+	const int32 VirtualTextureIndex = vLevel < TransitionLevel ? 0 : 1;
+	const int32 vLevelOffset = vLevel < TransitionLevel ? 0 : TransitionLevel;
+	return VirtualTextures[VirtualTextureIndex]->IsPageStreamed(vLevel - vLevelOffset, vAddress);
+}
+
 FVTRequestPageResult FVirtualTextureLevelRedirector::RequestPageData(
 	FRHICommandList& RHICmdList,
 	const FVirtualTextureProducerHandle& ProducerHandle,
diff --git a/Engine/Source/Runtime/Engine/Private/VT/VirtualTextureLevelRedirector.h b/Engine/Source/Runtime/Engine/Private/VT/VirtualTextureLevelRedirector.h
--- a/Engine/Source/Runtime/Engine/Private/VT/VirtualTextureLevelRedirector.h
+++ b/Engine/Source/Runtime/Engine/Private/VT/VirtualTextureLevelRedirector.h
@@ -16,6 +16,8 @@ public:
 	virtual ~FVirtualTextureLevelRedirector();
 
 	//~ Begin IVirtualTexture Interface.
+	virtual bool IsPageStreamed(uint8 vLevel, uint32 vAddress) const override;
+
 	virtual FVTRequestPageResult RequestPageData(
 		FRHICommandList& RHICmdList,
 		const FVirtualTextureProducerHandle& ProducerHandle,
diff --git a/Engine/Source/Runtime/Engine/Private/VT/VirtualTextureScalability.cpp b/Engine/Source/Runtime/Engine/Private/VT/VirtualTextureScalability.cpp
--- a/Engine/Source/Runtime/Engine/Private/VT/VirtualTextureScalability.cpp
+++ b/Engine/Source/Runtime/Engine/Private/VT/VirtualTextureScalability.cpp
@@ -29,6 +29,15 @@ namespace VirtualTextureScalability
 		ECVF_RenderThreadSafe | ECVF_Scalability
 	);
 
+	static TAutoConsoleVariable<int32> CVarVTMaxUploadsPerFrameStreaming(
+		TEXT("r.VT.MaxUploadsPerFrame.Streaming"),
+		0,
+		TEXT("If positive, max number of page uploads per frame in game for streaming VT. Negative means no limit.\n")
+		TEXT("If zero, SVTs won't be budgeted separately. They will be limited by r.VT.MaxUploadsPerFrame along with other types of VTs. This is the old behavior.\n")
+		TEXT("This limit should be high if streaming pages is slow so that I/O requests are not throttled which can cause long delays to acquire page data."),
+		ECVF_RenderThreadSafe | ECVF_Scalability
+	);
+
 	static TAutoConsoleVariable<int32> CVarMaxPagesProducedPerFrame(
 		TEXT("r.VT.MaxTilesProducedPerFrame"),
 		30,
@@ -241,6 +250,22 @@ namespace VirtualTextureScalability
 #endif
 	}
 
+	int32 GetMaxUploadsPerFrameForStreamingVT()
+	{
+		int32 Budget = CVarVTMaxUploadsPerFrameStreaming.GetValueOnAnyThread();
+
+		if (Budget < 0)
+		{
+			Budget = MAX_int32;
+		}
+#if WITH_EDITOR
+		// Don't want this scalability setting to affect editor because we rely on reactive updates while editing.
+		return GIsEditor ? CVarVTMaxUploadsPerFrameInEditor.GetValueOnAnyThread() : Budget;
+#else
+		return Budget;
+#endif
+	}
+
 	int32 GetMaxPagesProducedPerFrame()
 	{
 		return CVarMaxPagesProducedPerFrame.GetValueOnAnyThread();
diff --git a/Engine/Source/Runtime/Engine/Public/VT/VirtualTextureScalability.h b/Engine/Source/Runtime/Engine/Public/VT/VirtualTextureScalability.h
--- a/Engine/Source/Runtime/Engine/Public/VT/VirtualTextureScalability.h
+++ b/Engine/Source/Runtime/Engine/Public/VT/VirtualTextureScalability.h
@@ -8,6 +8,8 @@ namespace VirtualTextureScalability
 {
 	/** Get max upload rate to virtual textures. */
 	ENGINE_API int32 GetMaxUploadsPerFrame();
+	/** Get max upload rate to streaming virtual textures. May return 0 which means SVTs aren't budgeted separately. */
+	ENGINE_API int32 GetMaxUploadsPerFrameForStreamingVT();
 	/** Get max produce rate to virtual textures. */
 	ENGINE_API int32 GetMaxPagesProducedPerFrame();
 	/** Get max update rate of already mapped virtual texture pages. */
diff --git a/Engine/Source/Runtime/RenderCore/Public/VirtualTexturing.h b/Engine/Source/Runtime/RenderCore/Public/VirtualTexturing.h
--- a/Engine/Source/Runtime/RenderCore/Public/VirtualTexturing.h
+++ b/Engine/Source/Runtime/RenderCore/Public/VirtualTexturing.h
@@ -284,6 +284,11 @@ public:
 	 */
 	virtual uint32 GetLocalMipBias(uint8 vLevel, uint32 vAddress) const { return 0u; }
 
+	/**
+	 * Whether data for the given page is streamed (e.g. loading from disk).
+	 */
+	virtual bool IsPageStreamed(uint8 vLevel, uint32 vAddress) const = 0;
+
 	/**
 	* Makes a request for the given page data.
 	* For data sources that can generate data immediately, it's reasonable for this method to do nothing, and simply return 'Available'
diff --git a/Engine/Source/Runtime/Renderer/Private/VT/AdaptiveVirtualTexture.cpp b/Engine/Source/Runtime/Renderer/Private/VT/AdaptiveVirtualTexture.cpp
--- a/Engine/Source/Runtime/Renderer/Private/VT/AdaptiveVirtualTexture.cpp
+++ b/Engine/Source/Runtime/Renderer/Private/VT/AdaptiveVirtualTexture.cpp
@@ -63,6 +63,16 @@ public:
 	{
 	}
 
+	virtual bool IsPageStreamed(uint8 vLevel, uint32 vAddress) const override
+	{
+		uint64 X = FMath::ReverseMortonCode2_64(vAddress) + (AddressOffset.X >> (vLevel + LevelOffset));
+		uint64 Y = FMath::ReverseMortonCode2_64(vAddress >> 1) + (AddressOffset.Y >> (vLevel + LevelOffset));
+		vAddress = FMath::MortonCode2_64(X) | (FMath::MortonCode2_64(Y) << 1);
+		vLevel = (uint8)(FMath::Max((int32)vLevel + LevelOffset, 0));
+
+		return VirtualTexture->IsPageStreamed(vLevel, vAddress);
+	}
+
 	virtual FVTRequestPageResult RequestPageData(
 		FRHICommandList& RHICmdList,
 		const FVirtualTextureProducerHandle& ProducerHandle,
diff --git a/Engine/Source/Runtime/Renderer/Private/VT/RuntimeVirtualTextureProducer.h b/Engine/Source/Runtime/Renderer/Private/VT/RuntimeVirtualTextureProducer.h
--- a/Engine/Source/Runtime/Renderer/Private/VT/RuntimeVirtualTextureProducer.h
+++ b/Engine/Source/Runtime/Renderer/Private/VT/RuntimeVirtualTextureProducer.h
@@ -84,6 +84,11 @@ public:
 	virtual ~FRuntimeVirtualTextureProducer() {}
 
 	//~ Begin IVirtualTexture Interface.
+	virtual bool IsPageStreamed(uint8 vLevel, uint32 vAddress) const override
+	{
+		return false;
+	}
+
 	virtual FVTRequestPageResult RequestPageData(
 		FRHICommandList& RHICmdList,
 		const FVirtualTextureProducerHandle& ProducerHandle,
diff --git a/Engine/Source/Runtime/Renderer/Private/VT/UniqueRequestList.h b/Engine/Source/Runtime/Renderer/Private/VT/UniqueRequestList.h
--- a/Engine/Source/Runtime/Renderer/Private/VT/UniqueRequestList.h
+++ b/Engine/Source/Runtime/Renderer/Private/VT/UniqueRequestList.h
@@ -74,8 +74,10 @@ public:
 		, AdaptiveAllocationsRequests(Allocator.MallocArray<uint32>(LoadRequestCapacity))
 		, LoadRequestCount(Allocator.MallocArray<uint16>(LoadRequestCapacity))
 		, LoadRequestGroupMask(Allocator.MallocArray<uint8>(LoadRequestCapacity))
+		, LoadRequestFlags(Allocator.MallocArray<FLoadRequestFlags>(LoadRequestCapacity))
 		, NumLoadRequests(0u)
 		, NumLockRequests(0u)
+		, NumNonStreamingLoadRequests(0u)
 		, NumMappingRequests(0u)
 		, NumDirectMappingRequests(0u)
 		, NumContinuousUpdateRequests(0u)
@@ -98,6 +100,7 @@ public:
 		DirectMappingRequestHash.Clear();
 		NumLoadRequests = 0;
 		NumLockRequests = 0;
+		NumNonStreamingLoadRequests = 0;
 		NumMappingRequests = 0;
 		NumDirectMappingRequests = 0;
 		NumAdaptiveAllocationRequests = 0;
@@ -110,6 +113,7 @@ public:
 	}
 
 	inline uint32 GetNumLoadRequests() const { return NumLoadRequests; }
+	inline uint32 GetNumNonStreamingLoadRequests() const { return NumNonStreamingLoadRequests; }
 	inline uint32 GetNumMappingRequests() const { return NumMappingRequests; }
 	inline uint32 GetNumDirectMappingRequests() const { return NumDirectMappingRequests; }
 	inline uint32 GetNumContinuousUpdateRequests() const { return NumContinuousUpdateRequests; }
@@ -124,8 +128,8 @@ public:
 	inline uint8 GetGroupMask(uint32 i) const { checkSlow(i < NumLoadRequests); return LoadRequestGroupMask[i]; }
 	inline bool IsLocked(uint32 i) const { checkSlow(i < NumLoadRequests); return i < NumLockRequests; }
 
-	uint16 AddLoadRequest(const FVirtualTextureLocalTile& Tile, uint8 GroupMask, uint16 Count);
-	uint16 LockLoadRequest(const FVirtualTextureLocalTile& Tile, uint8 GroupMask);
+	uint16 AddLoadRequest(const FVirtualTextureLocalTile& Tile, uint8 GroupMask, uint16 Count, bool bStreamingRequest);
+	uint16 LockLoadRequest(const FVirtualTextureLocalTile& Tile, uint8 GroupMask, bool bStreamingRequest);
 
 	void AddMappingRequest(uint16 LoadRequestIndex, uint8 ProducerPhysicalGroupIndex, uint8 SpaceID, uint8 PageTableLayerIndex, uint32 MaxLevel, uint32 vAddress, uint8 vLevel, uint8 Local_vLevel);
 
@@ -138,7 +142,7 @@ public:
 
 	void MergeRequests(const FUniqueRequestList* RESTRICT Other, FConcurrentLinearBulkObjectAllocator& Allocator);
 
-	void SortRequests(FVirtualTextureProducerCollection& Producers, FConcurrentLinearBulkObjectAllocator& Allocator, uint32 MaxNumRequests);
+	void SortRequests(FVirtualTextureProducerCollection& Producers, FConcurrentLinearBulkObjectAllocator& Allocator, uint32 MaxNonStreamingLoadRequests, uint32 MaxStreamingLoadRequests, bool bUseCombinedLimit);
 
 private:
 	static const uint32 LoadRequestCapacity = 4u * 1024;
@@ -161,8 +165,22 @@ private:
 	uint16* LoadRequestCount;
 	uint8* LoadRequestGroupMask;
 
+	struct FLoadRequestFlags
+	{
+		uint8 bLocked : 1;
+		uint8 bStreaming : 1;
+		uint8 Padding : 6;
+
+		FLoadRequestFlags(bool bInLocked, bool bInStreaming)
+			: bLocked(bInLocked)
+			, bStreaming(bInStreaming)
+		{}
+	};
+	FLoadRequestFlags* LoadRequestFlags;
+
 	uint32 NumLoadRequests;
 	uint32 NumLockRequests;
+	uint32 NumNonStreamingLoadRequests;
 	uint32 NumMappingRequests;
 	uint32 NumDirectMappingRequests;
 	uint32 NumContinuousUpdateRequests;
@@ -170,7 +188,7 @@ private:
 };
 
 
-inline uint16 FUniqueRequestList::AddLoadRequest(const FVirtualTextureLocalTile& Tile, uint8 GroupMask, uint16 Count)
+inline uint16 FUniqueRequestList::AddLoadRequest(const FVirtualTextureLocalTile& Tile, uint8 GroupMask, uint16 Count, bool bStreamingRequest)
 {
 	const uint16 Hash = MurmurFinalize64(Tile.PackedValue);
 	check(GroupMask != 0u);
@@ -179,12 +197,9 @@ inline uint16 FUniqueRequestList::AddLoadRequest(const FVirtualTextureLocalTile&
 	{
 		if (Tile == LoadRequests[Index])
 		{
-			const uint32 PrevCount = LoadRequestCount[Index];
-			if (PrevCount != 0xffff)
-			{
-				// Don't adjust count if already locked, don't allow request to transition to lock
-				LoadRequestCount[Index] = FMath::Min<uint32>(PrevCount + Count, 0xfffe);
-			}
+			check(LoadRequestFlags[Index].bStreaming == bStreamingRequest);
+
+			LoadRequestCount[Index] = FMath::Min<uint32>((uint32)LoadRequestCount[Index] + Count, MAX_uint16);
 			LoadRequestGroupMask[Index] |= GroupMask;
 			return Index;
 		}
@@ -195,14 +210,15 @@ inline uint16 FUniqueRequestList::AddLoadRequest(const FVirtualTextureLocalTile&
 		const uint32 Index = NumLoadRequests++;
 		LoadRequestHash.Add(Hash, Index);
 		LoadRequests[Index] = Tile;
-		LoadRequestCount[Index] = FMath::Min<uint32>(Count, 0xfffe);
+		LoadRequestCount[Index] = Count;
 		LoadRequestGroupMask[Index] = GroupMask;
+		LoadRequestFlags[Index] = FLoadRequestFlags(false, bStreamingRequest);
 		return Index;
 	}
 	return 0xffff;
 }
 
-inline uint16 FUniqueRequestList::LockLoadRequest(const FVirtualTextureLocalTile& Tile, uint8 GroupMask)
+inline uint16 FUniqueRequestList::LockLoadRequest(const FVirtualTextureLocalTile& Tile, uint8 GroupMask, bool bStreamingRequest)
 {
 	const uint16 Hash = MurmurFinalize64(Tile.PackedValue);
 	check(GroupMask != 0u);
@@ -211,9 +227,13 @@ inline uint16 FUniqueRequestList::LockLoadRequest(const FVirtualTextureLocalTile
 	{
 		if (Tile == LoadRequests[Index])
 		{
-			if (LoadRequestCount[Index] != 0xffff)
+			FLoadRequestFlags& Flags = LoadRequestFlags[Index];
+			check(bStreamingRequest == Flags.bStreaming);
+
+			if (!Flags.bLocked)
 			{
-				LoadRequestCount[Index] = 0xffff;
+				Flags.bLocked = true;
+				LoadRequestCount[Index] = MAX_uint16;
 				++NumLockRequests;
 			}
 			LoadRequestGroupMask[Index] |= GroupMask;
@@ -226,8 +246,9 @@ inline uint16 FUniqueRequestList::LockLoadRequest(const FVirtualTextureLocalTile
 		const uint32 Index = NumLoadRequests++;
 		LoadRequestHash.Add(Hash, Index);
 		LoadRequests[Index] = Tile;
-		LoadRequestCount[Index] = 0xffff;
+		LoadRequestCount[Index] = MAX_uint16;
 		LoadRequestGroupMask[Index] = GroupMask;
+		LoadRequestFlags[Index] = FLoadRequestFlags(true, bStreamingRequest);
 		++NumLockRequests;
 		return Index;
 	}
@@ -317,11 +338,11 @@ inline void FUniqueRequestList::MergeRequests(const FUniqueRequestList* RESTRICT
 	{
 		if (Other->IsLocked(Index))
 		{
-			LoadRequestIndexRemap[Index] = LockLoadRequest(Other->GetLoadRequest(Index), Other->LoadRequestGroupMask[Index]);
+			LoadRequestIndexRemap[Index] = LockLoadRequest(Other->GetLoadRequest(Index), Other->LoadRequestGroupMask[Index], Other->LoadRequestFlags[Index].bStreaming);
 		}
 		else
 		{
-			LoadRequestIndexRemap[Index] = AddLoadRequest(Other->GetLoadRequest(Index), Other->LoadRequestGroupMask[Index], Other->LoadRequestCount[Index]);
+			LoadRequestIndexRemap[Index] = AddLoadRequest(Other->GetLoadRequest(Index), Other->LoadRequestGroupMask[Index], Other->LoadRequestCount[Index], Other->LoadRequestFlags[Index].bStreaming);
 		}
 	}
 
@@ -352,7 +373,7 @@ inline void FUniqueRequestList::MergeRequests(const FUniqueRequestList* RESTRICT
 	}
 }
 
-inline void FUniqueRequestList::SortRequests(FVirtualTextureProducerCollection& Producers, FConcurrentLinearBulkObjectAllocator& Allocator, uint32 MaxNumRequests)
+inline void FUniqueRequestList::SortRequests(FVirtualTextureProducerCollection& Producers, FConcurrentLinearBulkObjectAllocator& Allocator, uint32 MaxNonStreamingLoadRequests, uint32 MaxStreamingLoadRequests, bool bUseCombinedLimit)
 {
 	struct FPriorityAndIndex
 	{
@@ -360,53 +381,103 @@ inline void FUniqueRequestList::SortRequests(FVirtualTextureProducerCollection&
 		uint16 Index;
 
 		// sort from largest to smallest
-		inline bool operator<(const FPriorityAndIndex& Rhs) const{ return Priroity > Rhs.Priroity; }
+		inline bool operator<(const FPriorityAndIndex& Rhs) const { return Priroity > Rhs.Priroity; }
 	};
 
+	if (bUseCombinedLimit)
+	{
+		MaxNonStreamingLoadRequests += MaxStreamingLoadRequests;
+		MaxStreamingLoadRequests = 0;
+	}
+
 	// Compute priority of each load request
 	uint32 CheckNumLockRequests = 0u;
+	uint32 NumNonStreamingLockRequests = 0u;
+	uint32 NumStreamingNonLockRequests = 0u;
 	FPriorityAndIndex* SortedKeys = Allocator.CreateArray<FPriorityAndIndex>(NumLoadRequests);
+
 	for (uint32 i = 0u; i < NumLoadRequests; ++i)
 	{
 		const uint32 Count = LoadRequestCount[i];
+		const FLoadRequestFlags Flags = LoadRequestFlags[i];
+
+		// Try to load higher mips first
+		const FVirtualTextureLocalTile& TileToLoad = GetLoadRequest(i);
+		const uint32 Priority = (Count * (1u + TileToLoad.Local_vLevel)) & 0x0fffffff;
+
 		SortedKeys[i].Index = i;
-		if (Count == 0xffff)
+		if (Flags.bLocked)
 		{
 			// Lock request, use max priority
-			SortedKeys[i].Priroity = 0xffffffff;
+			SortedKeys[i].Priroity = MAX_uint32;
+			NumNonStreamingLockRequests += (bUseCombinedLimit || !Flags.bStreaming ? 1 : 0);
 			++CheckNumLockRequests;
 		}
+		else if (bUseCombinedLimit || !Flags.bStreaming)
+		{
+			SortedKeys[i].Priroity = 0xd | Priority;
+		}
 		else
 		{
-			// Try to load higher mips first
-			const FVirtualTextureLocalTile& TileToLoad = GetLoadRequest(i);
-			const uint32 Priority = Count * (1u + TileToLoad.Local_vLevel);
-			SortedKeys[i].Priroity = Priority;
+			SortedKeys[i].Priroity = 0xe | Priority;
+			++NumStreamingNonLockRequests;
 		}
 	}
 	checkSlow(CheckNumLockRequests == NumLockRequests);
 
 	// Sort so highest priority requests are at the front of the list
+	// [ Lock requests | streaming non-lock requests | non-streaming non-lock requests ]
 	Algo::Sort(MakeArrayView(SortedKeys, NumLoadRequests));
 
 	// Clamp number of load requests to maximum, but also ensure all lock requests are considered
-	const uint32 NewNumLoadRequests = FMath::Min(NumLoadRequests, FMath::Max(NumLockRequests, MaxNumRequests));
+	const uint32 NumStreamingLockRequests = NumLockRequests - NumNonStreamingLockRequests;
+	const uint32 NumStreamingRequests = NumStreamingNonLockRequests + NumStreamingLockRequests;
+	const uint32 NumNonStreamingRequests = NumLoadRequests - NumStreamingRequests;
+
+	const uint32 NewNumNonStreamingRequests = FMath::Min(NumNonStreamingRequests, FMath::Max(NumNonStreamingLockRequests, MaxNonStreamingLoadRequests));
+	const uint32 NewNumStreamingRequests = FMath::Min(NumStreamingRequests, FMath::Max(NumStreamingLockRequests, MaxStreamingLoadRequests));
+	const uint32 NewNumLoadRequests = NewNumNonStreamingRequests + NewNumStreamingRequests;
 
 	// Re-index load request list, using sorted indices
 	FVirtualTextureLocalTile* SortedLoadRequests = Allocator.CreateArray<FVirtualTextureLocalTile>(NewNumLoadRequests);
 	uint8* SortedGroupMask = Allocator.MallocArray<uint8>(NewNumLoadRequests);
+	FLoadRequestFlags* SortedFlags = Allocator.MallocArray<FLoadRequestFlags>(NewNumLoadRequests);
 	uint16* LoadIndexToSortedLoadIndex = Allocator.MallocArray<uint16>(NumLoadRequests);
 	FMemory::Memset(LoadIndexToSortedLoadIndex, 0xff, NumLoadRequests * sizeof(uint16));
-	for (uint32 i = 0u; i < NewNumLoadRequests; ++i)
+
+	uint32 WriteIndex = 0;
+	auto CopyRequestToSorted = [this, &WriteIndex, SortedKeys, SortedLoadRequests, SortedGroupMask, SortedFlags, LoadIndexToSortedLoadIndex](uint32 SortedIndex)
+	{
+		const uint32 OldIndex = SortedKeys[SortedIndex].Index;
+		checkSlow(OldIndex < NumLoadRequests);
+
+		SortedLoadRequests[WriteIndex] = LoadRequests[OldIndex];
+		SortedGroupMask[WriteIndex] = LoadRequestGroupMask[OldIndex];
+		SortedFlags[WriteIndex] = LoadRequestFlags[OldIndex];
+		LoadIndexToSortedLoadIndex[OldIndex] = WriteIndex;
+
+		++WriteIndex;
+	};
+
+	for (uint32 i = 0u; i < NumLockRequests; ++i)
 	{
-		const uint32 SortedIndex = SortedKeys[i].Index;
-		SortedLoadRequests[i] = LoadRequests[SortedIndex];
-		SortedGroupMask[i] = LoadRequestGroupMask[SortedIndex];
-		checkSlow(SortedIndex < NumLoadRequests);
-		LoadIndexToSortedLoadIndex[SortedIndex] = i;
+		CopyRequestToSorted(i);
 	}
-	FMemory::Memcpy(LoadRequests, SortedLoadRequests, sizeof(FVirtualTextureLocalTile) * NewNumLoadRequests);
-	FMemory::Memcpy(LoadRequestGroupMask, SortedGroupMask, sizeof(uint8) * NewNumLoadRequests);
+
+	for (uint32 i = 0u; i < NewNumStreamingRequests - NumStreamingLockRequests; ++i)
+	{
+		CopyRequestToSorted(NumLockRequests + i);
+	}
+
+	for (uint32 i = 0u; i < NewNumNonStreamingRequests - NumNonStreamingLockRequests; ++i)
+	{
+		CopyRequestToSorted(NumLockRequests + NumStreamingNonLockRequests + i);
+	}
+	
+	check(NewNumLoadRequests == WriteIndex);
+	LoadRequests = SortedLoadRequests;
+	LoadRequestGroupMask = SortedGroupMask;
+	LoadRequestFlags = SortedFlags;
 
 	// Remap LoadRequest indices for all the mapping requests
 	// Can discard any mapping request that refers to a LoadRequest that's no longer being performed this frame
@@ -425,5 +496,7 @@ inline void FUniqueRequestList::SortRequests(FVirtualTextureProducerCollection&
 	}
 
 	NumLoadRequests = NewNumLoadRequests;
+	NumNonStreamingLoadRequests = NewNumNonStreamingRequests;
+	check(!bUseCombinedLimit || NumLoadRequests == NumNonStreamingLoadRequests);
 	NumMappingRequests = NewNumMappingRequests;
 }
diff --git a/Engine/Source/Runtime/Renderer/Private/VT/VirtualTextureSystem.cpp b/Engine/Source/Runtime/Renderer/Private/VT/VirtualTextureSystem.cpp
--- a/Engine/Source/Runtime/Renderer/Private/VT/VirtualTextureSystem.cpp
+++ b/Engine/Source/Runtime/Renderer/Private/VT/VirtualTextureSystem.cpp
@@ -175,7 +175,8 @@ FVirtualTextureUpdateSettings::FVirtualTextureUpdateSettings()
 	NumFeedbackTasks = CVarVTNumFeedbackTasks.GetValueOnRenderThread();
 	NumGatherTasks = CVarVTNumGatherTasks.GetValueOnRenderThread();
 	MaxGatherPagesBeforeFlush = CVarVTPageUpdateFlushCount.GetValueOnRenderThread();
-	MaxPageUploads = VirtualTextureScalability::GetMaxUploadsPerFrame();
+	MaxRVTPageUploads = VirtualTextureScalability::GetMaxUploadsPerFrame();
+	MaxSVTPageUploads = VirtualTextureScalability::GetMaxUploadsPerFrameForStreamingVT();
 	MaxPagesProduced = VirtualTextureScalability::GetMaxPagesProducedPerFrame();
 	MaxContinuousUpdates = VirtualTextureScalability::GetMaxContinuousUpdatesPerFrame();
 }
@@ -1798,7 +1799,8 @@ void FVirtualTextureSystem::GatherRequestsTask(const FGatherRequestsParameters&
 						}
 					}
 
-					const uint16 LoadRequestIndex = RequestList->AddLoadRequest(FVirtualTextureLocalTile(ProducerHandle, PrefetchLocal_vAddress, PrefetchLocal_vLevel), ProducerPhysicalGroupMaskToPrefetch, PageCount);
+					const bool bStreamingRequest = Producer->GetVirtualTexture()->IsPageStreamed(PrefetchLocal_vLevel, PrefetchLocal_vAddress);
+					const uint16 LoadRequestIndex = RequestList->AddLoadRequest(FVirtualTextureLocalTile(ProducerHandle, PrefetchLocal_vAddress, PrefetchLocal_vLevel), ProducerPhysicalGroupMaskToPrefetch, PageCount, bStreamingRequest);
 					if (LoadRequestIndex != 0xffff)
 					{
 						const uint32 PrefetchMapping_vLevel = PrefetchLocal_vLevel + ProducerMipBias;
@@ -1822,7 +1824,8 @@ void FVirtualTextureSystem::GatherRequestsTask(const FGatherRequestsParameters&
 
 			if (GroupMaskToLoad != 0u)
 			{
-				const uint16 LoadRequestIndex = RequestList->AddLoadRequest(FVirtualTextureLocalTile(ProducerHandle, Local_vAddress, Local_vLevel), GroupMaskToLoad, PageCount);
+				const bool bStreamingRequest = Producer->GetVirtualTexture()->IsPageStreamed(Local_vLevel, Local_vAddress);
+				const uint16 LoadRequestIndex = RequestList->AddLoadRequest(FVirtualTextureLocalTile(ProducerHandle, Local_vAddress, Local_vLevel), GroupMaskToLoad, PageCount, bStreamingRequest);
 				if (LoadRequestIndex != 0xffff)
 				{
 					for (uint32 LoadLayerIndex = 0u; LoadLayerIndex < NumPageTableLayersToLoad; ++LoadLayerIndex)
@@ -2025,28 +2028,32 @@ void FVirtualTextureSystem::SubmitThrottledRequests(FRHICommandList& RHICmdList,
 	{
 		SCOPE_CYCLE_COUNTER(STAT_ProcessRequests_Sort);
 
-		// Limit the number of uploads (account for MappedTilesToProduce this frame)
-		// Are all pages equal? Should there be different limits on different types of pages?
-		const int32 MaxNumUploads = Settings.MaxPageUploads;
-		const int32 MaxRequestUploads = FMath::Max(MaxNumUploads - MappedTilesToProduce.Num() - (int32)Updater->NumProcessedLoadRequests, 1);
+		// Limit the number of uploads (account for MappedTilesToProduce this frame but only for RVTs since SVT pages will take a while to load and likely won't be produced this frame).
+		// Use a separate budget for SVTs since we may want more of them to go through so that I/O can be initiated as early as possible
+		const int32 MaxNonStreamingLoadRequests = FMath::Max(Settings.MaxRVTPageUploads - MappedTilesToProduce.Num(), 1);
+		const int32 MaxStreamingLoadRequests = Settings.MaxSVTPageUploads;
+		check(MaxStreamingLoadRequests >= 0);
+		// 0 is a special value that enables the old behavior where all pages are limited by a single budget
+		const bool bUseCombinedLimit = !MaxStreamingLoadRequests;
+		const uint32 OldNumLoadRequests = MergedRequestList->GetNumLoadRequests();
 
-		if (MaxRequestUploads < (int32)MergedRequestList->GetNumLoadRequests())
+		MergedRequestList->SortRequests(Producers, Allocator, MaxNonStreamingLoadRequests, MaxStreamingLoadRequests, bUseCombinedLimit);
+
+		if (MergedRequestList->GetNumLoadRequests() < OldNumLoadRequests)
 		{
 			// Dropping requests is normal but track to log here if we want to tune settings.
 			if (CVarVTVerbose.GetValueOnRenderThread())
 			{
-				UE_LOG(LogConsoleResponse, Display, TEXT("VT dropped %d load requests."), MergedRequestList->GetNumLoadRequests() - MaxRequestUploads);
+				UE_LOG(LogConsoleResponse, Display, TEXT("VT dropped %d load requests."), MergedRequestList->GetNumLoadRequests() - OldNumLoadRequests);
 			}
 		}
-
-		MergedRequestList->SortRequests(Producers, Allocator, MaxRequestUploads);
 	}
 
 	if (bContinousUpdates)
 	{
-		// After sorting and clamping the load requests, if we still have unused upload bandwidth then use it to add some continous updates
-		const int32 MaxNumUploads = Settings.MaxPageUploads;
-		const int32 MaxTilesToProduce = FMath::Max(MaxNumUploads - MappedTilesToProduce.Num() - (int32)MergedRequestList->GetNumLoadRequests() - (int32)Updater->NumProcessedLoadRequests, 0);
+		// After sorting and clamping the load requests, if we still have unused upload bandwidth then use it to add some continuous updates.
+		// Not taking SVT requests into account since they take a while to load and likely won't be produced this frame
+		const int32 MaxTilesToProduce = FMath::Max(Settings.MaxRVTPageUploads - MappedTilesToProduce.Num() - (int32)MergedRequestList->GetNumNonStreamingLoadRequests(), 0);
 		const int32 MaxContinuousUpdates = Settings.MaxContinuousUpdates;
 
 		GetContinuousUpdatesToProduce(MergedRequestList, MaxTilesToProduce, MaxContinuousUpdates);
@@ -2570,7 +2577,8 @@ void FVirtualTextureSystem::GatherLockedTileRequests(FUniqueRequestList* MergedR
 
 			if (ProducerLayerMaskToLoad != 0u)
 			{
-				const uint16 LoadRequestIndex = MergedRequestList->LockLoadRequest(FVirtualTextureLocalTile(Tile.GetProducerHandle(), Tile.Local_vAddress, Tile.Local_vLevel), ProducerLayerMaskToLoad);
+				const bool bStreamingRequest = Producer->GetVirtualTexture()->IsPageStreamed(Tile.Local_vLevel, Tile.Local_vAddress);
+				const uint16 LoadRequestIndex = MergedRequestList->LockLoadRequest(FVirtualTextureLocalTile(Tile.GetProducerHandle(), Tile.Local_vAddress, Tile.Local_vLevel), ProducerLayerMaskToLoad, bStreamingRequest);
 				if (LoadRequestIndex == 0xffff)
 				{
 					// Overflowed the request list...try to lock the tile again next frame
diff --git a/Engine/Source/Runtime/Renderer/Private/VT/VirtualTextureSystem.h b/Engine/Source/Runtime/Renderer/Private/VT/VirtualTextureSystem.h
--- a/Engine/Source/Runtime/Renderer/Private/VT/VirtualTextureSystem.h
+++ b/Engine/Source/Runtime/Renderer/Private/VT/VirtualTextureSystem.h
@@ -44,7 +44,8 @@ struct FVirtualTextureUpdateSettings
 	{
 		if (!bEnable)
 		{
-			MaxPageUploads = 99999;
+			MaxRVTPageUploads = 99999;
+			MaxSVTPageUploads = 99999;
 			MaxPagesProduced = 99999;
 		}
 		return *this;
@@ -73,7 +74,8 @@ struct FVirtualTextureUpdateSettings
 	int32 NumFeedbackTasks;
 	int32 NumGatherTasks;
 	int32 MaxGatherPagesBeforeFlush;
-	int32 MaxPageUploads;
+	int32 MaxRVTPageUploads;
+	int32 MaxSVTPageUploads;
 	int32 MaxPagesProduced;
 	int32 MaxContinuousUpdates;
 };
