Split hair cards instances interpolation between Primary / Shadow to ensure shadow/light data are finished by the time we cull and interpolation instance only visible in shadow.

diff --git a/Engine/Plugins/Runtime/HairStrands/Source/HairStrandsCore/Private/GroomDebug.cpp b/Engine/Plugins/Runtime/HairStrands/Source/HairStrandsCore/Private/GroomDebug.cpp
--- a/Engine/Plugins/Runtime/HairStrands/Source/HairStrandsCore/Private/GroomDebug.cpp
+++ b/Engine/Plugins/Runtime/HairStrands/Source/HairStrandsCore/Private/GroomDebug.cpp
@@ -901,8 +901,8 @@ static void AddHairDebugPrintInstancePass(
 
 				{
 					D.Data2 = FUintVector4(0);
-					D.Data2.X |= InstanceIndex < InstanceCountPerType[HairInstanceCount_StrandsPrimaryView]  ? 0x1u  : 0u;
-					D.Data2.X |= InstanceIndex < InstanceCountPerType[HairInstanceCount_StrandsShadowView]   ? 0x2u  : 0u;
+					D.Data2.X |= InstanceIndex < InstanceCountPerType[uint32(EHairInstanceCount::StrandsPrimaryView)]  ? 0x1u  : 0u;
+					D.Data2.X |= InstanceIndex < InstanceCountPerType[uint32(EHairInstanceCount::StrandsShadowView)]   ? 0x2u  : 0u;
 					D.Data2.X |= Instance->HairGroupPublicData->VFInput.Strands.Common.bScatterSceneLighting ? 0x4u  : 0u;
 					D.Data2.X |= Instance->HairGroupPublicData->VFInput.Strands.Common.bRaytracingGeometry   ? 0x8u  : 0u;
 					D.Data2.X |= Instance->HairGroupPublicData->VFInput.Strands.Common.bStableRasterization  ? 0x10u : 0u;
@@ -941,8 +941,8 @@ static void AddHairDebugPrintInstancePass(
 				D.Data0.W = Instance->Cards.LODs[IntLODIndex].Data->GetNumVertices();
 
 				D.Data2 = FUintVector4(0);
-				D.Data2.X |= InstanceIndex < InstanceCountPerType[HairInstanceCount_CardsOrMeshesPrimaryView] ? 0x1u : 0u;
-				D.Data2.X |= InstanceIndex < InstanceCountPerType[HairInstanceCount_CardsOrMeshesShadowView]  ? 0x2u : 0u;
+				D.Data2.X |= InstanceIndex < InstanceCountPerType[uint32(EHairInstanceCount::CardsOrMeshesPrimaryView)] ? 0x1u : 0u;
+				D.Data2.X |= InstanceIndex < InstanceCountPerType[uint32(EHairInstanceCount::CardsOrMeshesShadowView)]  ? 0x2u : 0u;
 
 				#if RHI_RAYTRACING
 				bHasRaytracing = Instance->Cards.LODs[IntLODIndex].RaytracingResource != nullptr;
@@ -956,8 +956,8 @@ static void AddHairDebugPrintInstancePass(
 				D.Data0.W = Instance->Meshes.LODs[IntLODIndex].Data->GetNumVertices();
 
 				D.Data2 = FUintVector4(0);
-				D.Data2.X |= InstanceIndex < InstanceCountPerType[HairInstanceCount_CardsOrMeshesPrimaryView] ? 0x1u : 0u;
-				D.Data2.X |= InstanceIndex < InstanceCountPerType[HairInstanceCount_CardsOrMeshesShadowView]  ? 0x2u : 0u;
+				D.Data2.X |= InstanceIndex < InstanceCountPerType[uint32(EHairInstanceCount::CardsOrMeshesPrimaryView)] ? 0x1u : 0u;
+				D.Data2.X |= InstanceIndex < InstanceCountPerType[uint32(EHairInstanceCount::CardsOrMeshesShadowView)]  ? 0x2u : 0u;
 
 				#if RHI_RAYTRACING
 				bHasRaytracing = Instance->Meshes.LODs[IntLODIndex].RaytracingResource != nullptr;
@@ -988,10 +988,10 @@ static void AddHairDebugPrintInstancePass(
 	{
 		FHairDebugPrintInstanceCS::FParameters* Parameters = GraphBuilder.AllocParameters<FHairDebugPrintInstanceCS::FParameters>();
 		Parameters->InstanceCount = InstanceCount;
-		Parameters->InstanceCount_StrandsPrimaryView = InstanceCountPerType[HairInstanceCount_StrandsPrimaryView];
-		Parameters->InstanceCount_StrandsShadowView = InstanceCountPerType[HairInstanceCount_StrandsShadowView];
-		Parameters->InstanceCount_CardsOrMeshesPrimaryView = InstanceCountPerType[HairInstanceCount_CardsOrMeshesPrimaryView];
-		Parameters->InstanceCount_CardsOrMeshesShadowView = InstanceCountPerType[HairInstanceCount_CardsOrMeshesShadowView];
+		Parameters->InstanceCount_StrandsPrimaryView = InstanceCountPerType[uint32(EHairInstanceCount::StrandsPrimaryView)];
+		Parameters->InstanceCount_StrandsShadowView = InstanceCountPerType[uint32(EHairInstanceCount::StrandsShadowView)];
+		Parameters->InstanceCount_CardsOrMeshesPrimaryView = InstanceCountPerType[uint32(EHairInstanceCount::CardsOrMeshesPrimaryView)];
+		Parameters->InstanceCount_CardsOrMeshesShadowView = InstanceCountPerType[uint32(EHairInstanceCount::CardsOrMeshesShadowView)];
 		Parameters->InstanceNames = InstanceNames.GetParameters(GraphBuilder);
 		Parameters->AttributeNames = AttributeNames.GetParameters(GraphBuilder);
 		Parameters->Infos = GraphBuilder.CreateSRV(InfoBuffer, PF_R32_UINT);
diff --git a/Engine/Plugins/Runtime/HairStrands/Source/HairStrandsCore/Private/GroomManager.cpp b/Engine/Plugins/Runtime/HairStrands/Source/HairStrandsCore/Private/GroomManager.cpp
--- a/Engine/Plugins/Runtime/HairStrands/Source/HairStrandsCore/Private/GroomManager.cpp
+++ b/Engine/Plugins/Runtime/HairStrands/Source/HairStrandsCore/Private/GroomManager.cpp
@@ -533,6 +533,8 @@ static void RunHairStrandsInterpolation_Strands(
 	const FShaderPrintData* ShaderPrintData,
 	FGlobalShaderMap* ShaderMap)
 {
+	if (Instances.IsEmpty()) { return; }
+
 	check(IsInRenderingThread());
 
 	DECLARE_GPU_STAT(HairStrandsInterpolation);
@@ -589,6 +591,8 @@ static void RunHairStrandsInterpolation_Cards(
 	const FShaderPrintData* ShaderPrintData,
 	FGlobalShaderMap* ShaderMap)
 {
+	if (Instances.IsEmpty()) { return; }
+
 	check(IsInRenderingThread());
 
 	DECLARE_GPU_STAT(HairCardsInterpolation);
@@ -1298,12 +1302,34 @@ void ProcessHairStrandsBookmark(
 {
 	check(Parameters.Instances != nullptr);
 
-	const bool bCulling =
-		IsInstanceFrustumCullingEnable() && (
-			Bookmark == EHairStrandsBookmark::ProcessCardsAndMeshesInterpolation ||
-			Bookmark == EHairStrandsBookmark::ProcessStrandsInterpolation ||
-			Bookmark == EHairStrandsBookmark::ProcessDebug);
-	FHairStrandsInstances& Instances = bCulling ? Parameters.VisibleInstances : *Parameters.Instances;
+	FHairStrandsInstances* Instances = Parameters.Instances;
+
+	// Cards interpolation for ShadowView only needs to run when FrustumCulling is enabled
+	if (Bookmark == EHairStrandsBookmark::ProcessCardsAndMeshesInterpolation_ShadowView && !IsInstanceFrustumCullingEnable())
+	{
+		return;
+	}
+
+	if (IsInstanceFrustumCullingEnable())
+	{
+		Instances = nullptr;
+		if (Bookmark == EHairStrandsBookmark::ProcessCardsAndMeshesInterpolation_PrimaryView)
+		{
+			Instances = &Parameters.VisibleCardsOrMeshes_Primary;
+		}
+		else if (Bookmark == EHairStrandsBookmark::ProcessCardsAndMeshesInterpolation_ShadowView)
+		{
+			Instances = &Parameters.VisibleCardsOrMeshes_Shadow;
+		}
+		else if (Bookmark == EHairStrandsBookmark::ProcessStrandsInterpolation)
+		{
+			Instances = &Parameters.VisibleStrands;
+		}
+		else
+		{
+			Instances = Parameters.Instances;
+		}
+	}
 
 	if (Bookmark == EHairStrandsBookmark::ProcessTasks)
 	{
@@ -1354,11 +1380,11 @@ void ProcessHairStrandsBookmark(
 			Parameters.Scene,
 			Parameters.View,
 			Parameters.ViewUniqueID,
-			Instances,
+			*Instances,
 			Parameters.ShaderPrintData,
 			Parameters.ShaderMap);
 	}
-	else if (Bookmark == EHairStrandsBookmark::ProcessCardsAndMeshesInterpolation)
+	else if (Bookmark == EHairStrandsBookmark::ProcessCardsAndMeshesInterpolation_PrimaryView || Bookmark == EHairStrandsBookmark::ProcessCardsAndMeshesInterpolation_ShadowView)
 	{
 		check(GraphBuilder);
 		RunHairStrandsInterpolation_Cards(
@@ -1367,7 +1393,7 @@ void ProcessHairStrandsBookmark(
 			Parameters.AllViews,
 			Parameters.View,
 			Parameters.ViewUniqueID,
-			Instances,
+			*Instances,
 			Parameters.ShaderPrintData,
 			Parameters.ShaderMap);
 	}
@@ -1380,19 +1406,29 @@ void ProcessHairStrandsBookmark(
 			Parameters.AllViews,
 			Parameters.View,
 			Parameters.ViewUniqueID,
-			Instances,
+			*Instances,
 			Parameters.ShaderPrintData,
 			Parameters.ShaderMap);
 	}
 	else if (Bookmark == EHairStrandsBookmark::ProcessDebug)
 	{
+		// Merge all visible instances
+		FHairStrandsInstances DebugInstance;
+		if (IsInstanceFrustumCullingEnable())
+		{
+			DebugInstance.Append(Parameters.VisibleStrands);
+			DebugInstance.Append(Parameters.VisibleCardsOrMeshes_Primary);
+			DebugInstance.Append(Parameters.VisibleCardsOrMeshes_Shadow);
+			Instances = &DebugInstance;
+		}
+
 		check(GraphBuilder);
 		RunHairStrandsDebug(
 			*GraphBuilder,
 			Parameters.ShaderMap,
 			Parameters.Scene,
 			*Parameters.View,
-			Instances,
+			*Instances,
 			Parameters.InstanceCountPerType,
 			Parameters.ShaderPrintData,
 			Parameters.SceneColorTexture,
diff --git a/Engine/Source/Runtime/Renderer/Private/DeferredShadingRenderer.cpp b/Engine/Source/Runtime/Renderer/Private/DeferredShadingRenderer.cpp
--- a/Engine/Source/Runtime/Renderer/Private/DeferredShadingRenderer.cpp
+++ b/Engine/Source/Runtime/Renderer/Private/DeferredShadingRenderer.cpp
@@ -3011,7 +3011,7 @@ void FDeferredShadingSceneRenderer::Render(FRDGBuilder& GraphBuilder)
 	FHairStrandsBookmarkParameters& HairStrandsBookmarkParameters = *GraphBuilder.AllocObject<FHairStrandsBookmarkParameters>();
 	if (IsHairStrandsEnabled(EHairStrandsShaderType::All, Scene->GetShaderPlatform()) && RendererOutput == ERendererOutput::FinalSceneColor)
 	{
-		HairStrandsBookmarkParameters = CreateHairStrandsBookmarkParameters(Scene, Views, AllFamilyViews);
+		CreateHairStrandsBookmarkParameters(Scene, Views, AllFamilyViews, HairStrandsBookmarkParameters);
 		RunHairStrandsBookmark(GraphBuilder, EHairStrandsBookmark::ProcessTasks, HairStrandsBookmarkParameters);
 
 		// Interpolation needs to happen after the skin cache run as there is a dependency 
@@ -3019,7 +3019,7 @@ void FDeferredShadingSceneRenderer::Render(FRDGBuilder& GraphBuilder)
 		const bool bRunHairStrands = HairStrandsBookmarkParameters.HasInstances() && (Views.Num() > 0);
 		if (bRunHairStrands)
 		{
-			RunHairStrandsBookmark(GraphBuilder, EHairStrandsBookmark::ProcessCardsAndMeshesInterpolation, HairStrandsBookmarkParameters);
+			RunHairStrandsBookmark(GraphBuilder, EHairStrandsBookmark::ProcessCardsAndMeshesInterpolation_PrimaryView, HairStrandsBookmarkParameters);
 		}
 		else
 		{
@@ -3359,6 +3359,19 @@ void FDeferredShadingSceneRenderer::Render(FRDGBuilder& GraphBuilder)
 			// Dynamic shadows are synced later when using the deferred path to make more headroom for tasks.
 			FinishInitDynamicShadows(GraphBuilder, InitViewTaskDatas.DynamicShadows, InstanceCullingManager, ExternalAccessQueue);
 		}
+
+		// Update groom only visible in shadow
+		if (IsHairStrandsEnabled(EHairStrandsShaderType::All, Scene->GetShaderPlatform()) && RendererOutput == ERendererOutput::FinalSceneColor)
+		{
+			UpdateHairStrandsBookmarkParameters(Scene, Views, HairStrandsBookmarkParameters);
+
+			// Interpolation for cards/meshes only visible in shadow needs to happen after the shadow jobs are completed
+			const bool bRunHairStrands = HairStrandsBookmarkParameters.HasInstances() && (Views.Num() > 0);
+			if (bRunHairStrands)
+			{
+				RunHairStrandsBookmark(GraphBuilder, EHairStrandsBookmark::ProcessCardsAndMeshesInterpolation_ShadowView, HairStrandsBookmarkParameters);
+			}
+		}
 	
 		// NOTE: The ordering of the lights is used to select sub-sets for different purposes, e.g., those that support clustered deferred.
 		FSortedLightSetSceneInfo& SortedLightSet = *GraphBuilder.AllocObject<FSortedLightSetSceneInfo>();
diff --git a/Engine/Source/Runtime/Renderer/Private/HairStrands/HairStrandsDebug.cpp b/Engine/Source/Runtime/Renderer/Private/HairStrands/HairStrandsDebug.cpp
--- a/Engine/Source/Runtime/Renderer/Private/HairStrands/HairStrandsDebug.cpp
+++ b/Engine/Source/Runtime/Renderer/Private/HairStrands/HairStrandsDebug.cpp
@@ -1078,7 +1078,8 @@ static void InternalRenderHairStrandsDebugInfo(
 	FRDGTextureRef SceneColorTexture,
 	FRDGTextureRef SceneDepthTexture)
 {
-	FHairStrandsBookmarkParameters Params = CreateHairStrandsBookmarkParameters(Scene, View);
+	FHairStrandsBookmarkParameters Params;
+	CreateHairStrandsBookmarkParameters(Scene, View, Params);
 	Params.SceneColorTexture = SceneColorTexture;
 	Params.SceneDepthTexture = SceneDepthTexture;
 	if (!Params.HasInstances())
diff --git a/Engine/Source/Runtime/Renderer/Private/HairStrands/HairStrandsInterface.cpp b/Engine/Source/Runtime/Renderer/Private/HairStrands/HairStrandsInterface.cpp
--- a/Engine/Source/Runtime/Renderer/Private/HairStrands/HairStrandsInterface.cpp
+++ b/Engine/Source/Runtime/Renderer/Private/HairStrands/HairStrandsInterface.cpp
@@ -400,21 +400,13 @@ void RunHairStrandsBookmark(EHairStrandsBookmark Bookmark, FHairStrandsBookmarkP
 	}
 }
 
-FHairStrandsBookmarkParameters CreateHairStrandsBookmarkParameters(FScene* Scene, FViewInfo& View)
+void CreateHairStrandsBookmarkParameters(FScene* Scene, FViewInfo& View, FHairStrandsBookmarkParameters& Out)
 {
 	const int32 ActiveInstanceCount = Scene->HairStrandsSceneData.RegisteredProxies.Num();
-	TBitArray InstancesVisibility(false, ActiveInstanceCount);
-
-	static bool bDebug = false;
-	if (Scene->HairStrandsSceneData.RegisteredProxies.Num() > 0)
-	{
-		bDebug = true;
-	}
-
-	FHairStrandsBookmarkParameters Out;
-	Out.VisibleInstances.Reserve(View.HairStrandsMeshElements.Num());
+	Out.InstancesVisibility.Init(false, ActiveInstanceCount);
 
 	// 1. Strands - Add all visible strands instances
+	Out.VisibleStrands.Reserve(View.HairStrandsMeshElements.Num());
 	for (const FMeshBatchAndRelevance& MeshBatch : View.HairStrandsMeshElements)
 	{
 		check(MeshBatch.PrimitiveSceneProxy && MeshBatch.PrimitiveSceneProxy->ShouldRenderInMainPass());
@@ -423,31 +415,15 @@ FHairStrandsBookmarkParameters CreateHairStrandsBookmarkParameters(FScene* Scene
 			FHairGroupPublicData* HairData = HairStrands::GetHairData(MeshBatch.Mesh);
 			if (HairData && HairData->Instance)
 			{
-				Out.VisibleInstances.Add(HairData->Instance);
-				InstancesVisibility[HairData->Instance->RegisteredIndex] = true;
-			}
-		}
-	}
-	Out.InstanceCountPerType[HairInstanceCount_StrandsPrimaryView] = Out.VisibleInstances.Num();
-
-	// 2. Strands - Add all instances non-visible primary view(s) but visible in shadow view(s)
-	if (IsHairStrandsNonVisibleShadowCastingEnable())
-	{
-		for (FHairStrandsInstance* Instance : Scene->HairStrandsSceneData.RegisteredProxies)
-		{
-			const bool bStrands = Instance->GetHairGeometry() == EHairGeometryType::Strands;
-			if (Instance->RegisteredIndex >= 0 && Instance->RegisteredIndex < ActiveInstanceCount && !InstancesVisibility[Instance->RegisteredIndex] && bStrands)
-			{
-				if (IsHairStrandsVisibleInShadows(View, *Instance))
-				{
-					Out.VisibleInstances.Add(Instance);
-				}
+				Out.VisibleStrands.Add(HairData->Instance);
+				Out.InstancesVisibility[HairData->Instance->RegisteredIndex] = true;
+				Out.InstanceCountPerType[uint32(EHairInstanceCount::StrandsPrimaryView)]++;
 			}
 		}
 	}
-	Out.InstanceCountPerType[HairInstanceCount_StrandsShadowView] = Out.VisibleInstances.Num() - Out.InstanceCountPerType[HairInstanceCount_StrandsPrimaryView];
 
-	// 3. Cards/Meshes - Add all visible cards instances
+	// 2. Cards/Meshes - Add all visible cards instances
+	Out.VisibleCardsOrMeshes_Primary.Reserve(View.HairStrandsMeshElements.Num());
 	for (const FMeshBatchAndRelevance& MeshBatch : View.HairCardsMeshElements)
 	{
 		check(MeshBatch.PrimitiveSceneProxy && MeshBatch.PrimitiveSceneProxy->ShouldRenderInMainPass());
@@ -456,53 +432,64 @@ FHairStrandsBookmarkParameters CreateHairStrandsBookmarkParameters(FScene* Scene
 			FHairGroupPublicData* HairData = HairStrands::GetHairData(MeshBatch.Mesh);
 			if (HairData && HairData->Instance)
 			{
-				Out.VisibleInstances.Add(HairData->Instance);
-				InstancesVisibility[HairData->Instance->RegisteredIndex] = true;
+				Out.VisibleCardsOrMeshes_Primary.Add(HairData->Instance);
+				Out.InstancesVisibility[HairData->Instance->RegisteredIndex] = true;
+				Out.InstanceCountPerType[uint32(EHairInstanceCount::CardsOrMeshesPrimaryView)]++;
 			}
 		}
 	}
-	Out.InstanceCountPerType[HairInstanceCount_CardsOrMeshesPrimaryView] = Out.VisibleInstances.Num() - Out.InstanceCountPerType[HairInstanceCount_StrandsShadowView];
 
-	// 4. Cards/Meshes - Add all instances non-visible primary view(s) but visible in shadow view(s)
+	Out.ShaderPrintData			= ShaderPrint::IsEnabled(View.ShaderPrintData) ? &View.ShaderPrintData : nullptr;
+	Out.ShaderMap				= View.ShaderMap;
+	Out.Instances				= &Scene->HairStrandsSceneData.RegisteredProxies;
+	Out.View					= &View;
+	Out.ViewRect				= View.ViewRect;
+	Out.ViewUniqueID			= View.ViewState ? View.ViewState->UniqueID : ~0;
+	Out.SceneColorTexture		= nullptr;
+	Out.Scene					= Scene;
+}
+
+void UpdateHairStrandsBookmarkParameters(FScene* Scene, TArray<FViewInfo>& Views, FHairStrandsBookmarkParameters& Out)
+{
+	if (Views.Num() == 0)
+	{
+		return;
+	}
+
 	if (IsHairStrandsNonVisibleShadowCastingEnable())
 	{
+		const int32 ActiveInstanceCount = Scene->HairStrandsSceneData.RegisteredProxies.Num();
 		for (FHairStrandsInstance* Instance : Scene->HairStrandsSceneData.RegisteredProxies)
 		{
+			// 2. Strands - Add all instances non-visible primary view(s) but visible in shadow view(s)
+			const bool bStrands = Instance->GetHairGeometry() == EHairGeometryType::Strands;
 			const bool bCardsOrMeshes = Instance->GetHairGeometry() == EHairGeometryType::Cards || Instance->GetHairGeometry() == EHairGeometryType::Meshes;
-			if (Instance->RegisteredIndex >= 0 && Instance->RegisteredIndex < ActiveInstanceCount && !InstancesVisibility[Instance->RegisteredIndex] && bCardsOrMeshes)
+			const bool bCompatible = bStrands || bCardsOrMeshes;
+
+			if (Instance->RegisteredIndex >= 0 && Instance->RegisteredIndex < ActiveInstanceCount && !Out.InstancesVisibility[Instance->RegisteredIndex] && bCompatible)
 			{
-				if (IsHairStrandsVisibleInShadows(View, *Instance))
+				if (IsHairStrandsVisibleInShadows(Views[0], *Instance))
 				{
-					Out.VisibleInstances.Add(Instance);
+					if (bStrands) 		
+					{ 
+						Out.VisibleStrands.Add(Instance);
+						Out.InstanceCountPerType[uint32(EHairInstanceCount::StrandsShadowView)]++; 
+					}
+					else if (bCardsOrMeshes) 
+					{ 
+						Out.VisibleCardsOrMeshes_Shadow.Add(Instance);
+						Out.InstanceCountPerType[uint32(EHairInstanceCount::CardsOrMeshesShadowView)]++; 
+					}
 				}
 			}
 		}
 	}
-	Out.InstanceCountPerType[HairInstanceCount_CardsOrMeshesShadowView] = Out.VisibleInstances.Num() - Out.InstanceCountPerType[HairInstanceCount_CardsOrMeshesPrimaryView];
-
-	Out.ShaderPrintData			= ShaderPrint::IsEnabled(View.ShaderPrintData) ? &View.ShaderPrintData : nullptr;
-	Out.ShaderMap				= View.ShaderMap;
-	Out.Instances				= &Scene->HairStrandsSceneData.RegisteredProxies;
-	Out.View					= &View;
-	Out.ViewRect				= View.ViewRect;
-	Out.ViewUniqueID			= View.ViewState ? View.ViewState->UniqueID : ~0;
-	Out.SceneColorTexture		= nullptr;
-	Out.bHzbRequest				= false; // Out.HasInstances() && IsHairStrandsEnabled(EHairStrandsShaderType::Strands, View.GetShaderPlatform());
-	Out.Scene					= Scene;
-
-	// Sanity check
-	check(Out.Instances->Num() >= Out.VisibleInstances.Num());
-
-	return Out;
 }
 
-FHairStrandsBookmarkParameters CreateHairStrandsBookmarkParameters(FScene* Scene, TArray<FViewInfo>& Views, TArray<const FSceneView*>& AllFamilyViews)
+void CreateHairStrandsBookmarkParameters(FScene* Scene, TArray<FViewInfo>& Views, TArray<const FSceneView*>& AllFamilyViews, FHairStrandsBookmarkParameters& Out)
 {
-	FHairStrandsBookmarkParameters Out;
-	Out = CreateHairStrandsBookmarkParameters(Scene, Views[0]);
+	CreateHairStrandsBookmarkParameters(Scene, Views[0], Out);
 	Out.AllViews = AllFamilyViews;
-
-	return Out;
 }
 
 namespace HairStrands
diff --git a/Engine/Source/Runtime/Renderer/Private/HairStrands/HairStrandsRendering.h b/Engine/Source/Runtime/Renderer/Private/HairStrands/HairStrandsRendering.h
--- a/Engine/Source/Runtime/Renderer/Private/HairStrands/HairStrandsRendering.h
+++ b/Engine/Source/Runtime/Renderer/Private/HairStrands/HairStrandsRendering.h
@@ -43,5 +43,6 @@ void RunHairStrandsBookmark(
 	EHairStrandsBookmark Bookmark,
 	FHairStrandsBookmarkParameters& Parameters);
 
-FHairStrandsBookmarkParameters CreateHairStrandsBookmarkParameters(FScene* Scene, FViewInfo& View);
-FHairStrandsBookmarkParameters CreateHairStrandsBookmarkParameters(FScene* Scene, TArray<FViewInfo>& Views, TArray<const FSceneView*>& AllFamilyViews);
+void CreateHairStrandsBookmarkParameters(FScene* Scene, FViewInfo& View, FHairStrandsBookmarkParameters& Out);
+void CreateHairStrandsBookmarkParameters(FScene* Scene, TArray<FViewInfo>& Views, TArray<const FSceneView*>& AllFamilyViews, FHairStrandsBookmarkParameters& Out);
+void UpdateHairStrandsBookmarkParameters(FScene* Scene, TArray<FViewInfo>& Views, FHairStrandsBookmarkParameters& Out);
diff --git a/Engine/Source/Runtime/Renderer/Private/SceneRendering.cpp b/Engine/Source/Runtime/Renderer/Private/SceneRendering.cpp
--- a/Engine/Source/Runtime/Renderer/Private/SceneRendering.cpp
+++ b/Engine/Source/Runtime/Renderer/Private/SceneRendering.cpp
@@ -4619,7 +4619,8 @@ static void RenderViewFamilies_RenderThread(FRHICommandListImmediate& RHICmdList
 
 		if (IsHairStrandsEnabled(EHairStrandsShaderType::All, Scene->GetShaderPlatform()) && (SceneRenderers[0]->AllFamilyViews.Num() > 0) && !bAnyShowHitProxies)
 		{
-			FHairStrandsBookmarkParameters Parameters = CreateHairStrandsBookmarkParameters(Scene, SceneRenderers[0]->Views, SceneRenderers[0]->AllFamilyViews);
+			FHairStrandsBookmarkParameters Parameters;
+			CreateHairStrandsBookmarkParameters(Scene, SceneRenderers[0]->Views, SceneRenderers[0]->AllFamilyViews, Parameters);
 			if (Parameters.HasInstances())
 			{
 				RunHairStrandsBookmark(EHairStrandsBookmark::ProcessEndOfFrame, Parameters);
diff --git a/Engine/Source/Runtime/Renderer/Private/SceneVisibility.cpp b/Engine/Source/Runtime/Renderer/Private/SceneVisibility.cpp
--- a/Engine/Source/Runtime/Renderer/Private/SceneVisibility.cpp
+++ b/Engine/Source/Runtime/Renderer/Private/SceneVisibility.cpp
@@ -4022,20 +4022,20 @@ void FSceneRenderer::PreVisibilityFrameSetup(FRDGBuilder& GraphBuilder)
 	{
 		if (Views.Num() > 0 && !ViewFamily.EngineShowFlags.HitProxies)
 		{
-			FHairStrandsBookmarkParameters Parameters = CreateHairStrandsBookmarkParameters(Scene, Views, AllFamilyViews);
+			FHairStrandsBookmarkParameters Parameters; 
+			CreateHairStrandsBookmarkParameters(Scene, Views, AllFamilyViews, Parameters);
+
 			if (Parameters.HasInstances())
 			{
+				// 1. Select appropriate LOD & geometry type
 				RunHairStrandsBookmark(GraphBuilder, EHairStrandsBookmark::ProcessLODSelection, Parameters);
-			}
-		}
 
-		if (IsHairStrandsEnabled(EHairStrandsShaderType::All, Scene->GetShaderPlatform()) && Views.Num() > 0 && !ViewFamily.EngineShowFlags.HitProxies)
-		{
-			// If we are rendering from scene capture we don't need to run another time the hair bookmarks.
-			if (Views[0].AllowGPUParticleUpdate())
-			{
-				FHairStrandsBookmarkParameters Parameters = CreateHairStrandsBookmarkParameters(Scene, Views, AllFamilyViews);
-				RunHairStrandsBookmark(GraphBuilder, EHairStrandsBookmark::ProcessGuideInterpolation, Parameters);
+				// 2. Run guide/simulation update
+				// If we are rendering from scene capture we don't need to run another time the hair bookmarks.
+				if (IsHairStrandsEnabled(EHairStrandsShaderType::All, Scene->GetShaderPlatform()) && Views[0].AllowGPUParticleUpdate())
+				{
+					RunHairStrandsBookmark(GraphBuilder, EHairStrandsBookmark::ProcessGuideInterpolation, Parameters);
+				}
 			}
 		}
 
diff --git a/Engine/Source/Runtime/Renderer/Public/HairStrandsInterface.h b/Engine/Source/Runtime/Renderer/Public/HairStrandsInterface.h
--- a/Engine/Source/Runtime/Renderer/Public/HairStrandsInterface.h
+++ b/Engine/Source/Runtime/Renderer/Public/HairStrandsInterface.h
@@ -9,6 +9,7 @@
 #include "CoreMinimal.h"
 #include "RendererInterface.h"
 #include "Containers/Array.h"
+#include "Containers/BitArray.h"
 #include "Engine/EngineTypes.h"
 #include "Shader.h"
 #include "RenderResource.h"
@@ -477,18 +478,20 @@ enum class EHairStrandsBookmark : uint8
 	ProcessTasks,
 	ProcessLODSelection,
 	ProcessGuideInterpolation,
-	ProcessCardsAndMeshesInterpolation,
+	ProcessCardsAndMeshesInterpolation_PrimaryView,
+	ProcessCardsAndMeshesInterpolation_ShadowView,
 	ProcessStrandsInterpolation,
 	ProcessDebug,
 	ProcessEndOfFrame
 };
 
-enum EHairInstanceCount : uint8
+enum class EHairInstanceCount : uint8
 {
-	HairInstanceCount_StrandsPrimaryView = 0,
-	HairInstanceCount_StrandsShadowView = 1,
-	HairInstanceCount_CardsOrMeshesPrimaryView = 2,
-	HairInstanceCount_CardsOrMeshesShadowView = 3
+	StrandsPrimaryView = 0,
+	StrandsShadowView = 1,
+	CardsOrMeshesPrimaryView = 2,
+	CardsOrMeshesShadowView = 3,
+	Count
 };
 
 struct FHairStrandsBookmarkParameters
@@ -498,8 +501,13 @@ struct FHairStrandsBookmarkParameters
 
 	uint32 ViewUniqueID = ~0; // View 0
 	FIntRect ViewRect; // View 0
-	FHairStrandsInstances VisibleInstances;
+
+	FHairStrandsInstances VisibleStrands; // Primary & Shadow
+	FHairStrandsInstances VisibleCardsOrMeshes_Primary;
+	FHairStrandsInstances VisibleCardsOrMeshes_Shadow;
+
 	FHairStrandsInstances* Instances = nullptr;
+	TBitArray<> InstancesVisibility;
 	const FSceneView* View = nullptr;// // View 0
 	FSceneInterface* Scene = nullptr;
 	TArray<const FSceneView*> AllViews;
@@ -508,9 +516,6 @@ struct FHairStrandsBookmarkParameters
 
 	FUintVector4 InstanceCountPerType = FUintVector4(0);
 
-	bool bHzbRequest = false;
-	uint32 FrameIndex = ~0;
-
 	inline bool HasInstances() const { return Instances != nullptr && Instances->Num() > 0; }
 };
 
