Changes TSR RejectShading pass to a 32x32 tile to better scale on VGPR pressure and r.TSR.ShadingRejection.TileOverscan

diff --git a/Engine/Shaders/Private/TemporalSuperResolution/TSRRejectShading.usf b/Engine/Shaders/Private/TemporalSuperResolution/TSRRejectShading.usf
--- a/Engine/Shaders/Private/TemporalSuperResolution/TSRRejectShading.usf
+++ b/Engine/Shaders/Private/TemporalSuperResolution/TSRRejectShading.usf
@@ -33,7 +33,7 @@
 //------------------------------------------------------- DERIVED
 
 /** Pixel coordinate stride between each lane in the GROUP_TILE_SIZE. */
-#if DIM_WAVE_SIZE == 32 && PLATFORM_GPU_ARCH == PLATFORM_GPU_ARCH_AMD_RDNA_2
+#if DIM_WAVE_SIZE == 32 && PLATFORM_GPU_ARCH == PLATFORM_GPU_ARCH_AMD_RDNA_2 && 0
 	/** Load an entire 16x16 tile in a single wave32 with 4x2 worth of data per lane and just 100% waveops all communications. */
 	#define GROUP_TILE_SIZE 16
 
@@ -50,12 +50,12 @@
 	/** Layouts wave in row major so an entire row of pixels fits within a single wave, and waveops horizontally, but LDS vertically.
 	 * Each lane worth of a 2x2 of data per lane to minimise to trade VGPR pressure for both number of wave ops and LDS.
 	 */
-	#define GROUP_TILE_SIZE 16
+	#define GROUP_TILE_SIZE 32
 
 	#define WAVE_COUNT_X 1
-	#define WAVE_COUNT_Y 8
+	#define WAVE_COUNT_Y 16
 
-	#define LANE_COUNT_X 8
+	#define LANE_COUNT_X 16
 	#define LANE_COUNT_Y 1
 
 	#define LANE_STRIDE_X 2
@@ -65,11 +65,11 @@
 	#define LDS_SIZE (2 * (LANE_STRIDE_X * LANE_COUNT_X) * (WAVE_COUNT_X * WAVE_COUNT_Y))
 
 #elif DIM_WAVE_SIZE == 0
-	/** Load an entire 16x16 with 1x1 worth of data by lane, only using LDS. */
-	#define GROUP_TILE_SIZE 16
+	/** Load an entire 32x32 with 1x1 worth of data by lane, only using LDS. */
+	#define GROUP_TILE_SIZE 32
 
-	#define WAVE_COUNT_X 16
-	#define WAVE_COUNT_Y 16
+	#define WAVE_COUNT_X GROUP_TILE_SIZE
+	#define WAVE_COUNT_Y GROUP_TILE_SIZE
 
 	#define LANE_COUNT_X 1
 	#define LANE_COUNT_Y 1
diff --git a/Engine/Source/Runtime/Renderer/Private/PostProcess/TemporalSuperResolution.cpp b/Engine/Source/Runtime/Renderer/Private/PostProcess/TemporalSuperResolution.cpp
--- a/Engine/Source/Runtime/Renderer/Private/PostProcess/TemporalSuperResolution.cpp
+++ b/Engine/Source/Runtime/Renderer/Private/PostProcess/TemporalSuperResolution.cpp
@@ -1923,8 +1923,6 @@ FDefaultTemporalUpscaler::FOutputs AddTemporalSuperResolutionPasses(
 	FRDGTextureRef HistoryRejectionTexture = nullptr;
 	FRDGTextureRef InputSceneColorLdrLumaTexture = nullptr;
 	{
-		const int32 GroupTileSize = 16;
-
 		const bool bComputeInputSceneColorTexture = InputSceneColorLdrLumaTexture == nullptr;
 		if (bComputeInputSceneColorTexture)
 		{
@@ -1962,6 +1960,17 @@ FDefaultTemporalUpscaler::FOutputs AddTemporalSuperResolutionPasses(
 		FScreenPassTextureViewport TranslucencyViewport(
 			SeparateTranslucencyTexture->Desc.Extent, SeparateTranslucencyRect);
 
+		FTSRRejectShadingCS::FPermutationDomain PermutationVector;
+		PermutationVector.Set<FTSRRejectShadingCS::FWaveSizeOps>(bUseWaveOps&& GRHIMinimumWaveSize >= 32 && GRHIMinimumWaveSize <= 64 ? GRHIMinimumWaveSize : 0);
+		PermutationVector.Set<FTSRRejectShadingCS::FFlickeringDetectionDim>(FlickeringFramePeriod > 0.0f);
+		PermutationVector.Set<FTSRRejectShadingCS::FHistoryResurrectionDim>(bCanResurrectHistory);
+		PermutationVector.Set<FTSRShader::F16BitVALUDim>(bUse16BitVALU);
+		PermutationVector = FTSRRejectShadingCS::RemapPermutation(PermutationVector);
+
+		const int32 GroupTileSize = 32;
+		const int32 TileOverscan = FMath::Clamp(CVarTSRShadingTileOverscan.GetValueOnRenderThread(), 2, GroupTileSize / 2 - 1);
+		const int32 TileSize = GroupTileSize - 2 * TileOverscan;
+
 		FTSRRejectShadingCS::FParameters* PassParameters = GraphBuilder.AllocParameters<FTSRRejectShadingCS::FParameters>();
 		PassParameters->CommonParameters = CommonParameters;
 		PassParameters->InputPixelPosToTranslucencyTextureUV =
@@ -1972,7 +1981,7 @@ FDefaultTemporalUpscaler::FOutputs AddTemporalSuperResolutionPasses(
 		PassParameters->HistoryGuideQuantizationError = ComputePixelFormatQuantizationError(History.GuideArray->Desc.Format);
 		PassParameters->FlickeringFramePeriod = FlickeringFramePeriod;
 		PassParameters->TheoricBlendFactor = 1.0f / (1.0f + MaxHistorySampleCount / OutputToInputResolutionFractionSquare);
-		PassParameters->TileOverscan = FMath::Clamp(CVarTSRShadingTileOverscan.GetValueOnRenderThread(), 2, GroupTileSize / 2 - 1);
+		PassParameters->TileOverscan = TileOverscan;
 		PassParameters->PerceptionAdd = FMath::Pow(0.5f, CVarTSRShadingExposureOffset.GetValueOnRenderThread());
 
 		PassParameters->InputTexture = PassInputs.SceneColor.Texture;
@@ -2049,15 +2058,6 @@ FDefaultTemporalUpscaler::FOutputs AddTemporalSuperResolutionPasses(
 			PassParameters->DebugOutput = CreateDebugUAV(InputExtent, TEXT("Debug.TSR.RejectShading"));
 		}
 
-		int32 TileSize = GroupTileSize - 2 * PassParameters->TileOverscan;
-
-		FTSRRejectShadingCS::FPermutationDomain PermutationVector;
-		PermutationVector.Set<FTSRRejectShadingCS::FWaveSizeOps>(bUseWaveOps && GRHIMinimumWaveSize >= 32 && GRHIMinimumWaveSize <= 64 ? GRHIMinimumWaveSize : 0);
-		PermutationVector.Set<FTSRRejectShadingCS::FFlickeringDetectionDim>(FlickeringFramePeriod > 0.0f);
-		PermutationVector.Set<FTSRRejectShadingCS::FHistoryResurrectionDim>(bCanResurrectHistory);
-		PermutationVector.Set<FTSRShader::F16BitVALUDim>(bUse16BitVALU);
-		PermutationVector = FTSRRejectShadingCS::RemapPermutation(PermutationVector);
-
 		TShaderMapRef<FTSRRejectShadingCS> ComputeShader(View.ShaderMap, PermutationVector);
 		FComputeShaderUtils::AddPass(
 			GraphBuilder,
