Moved UpdateAllPrimitiveSceneInfos to occur prior to landscape (and other pre-render extensions).
 - Fixed hit-proxy rendering to call postrender extensions properly.
 - Moved VT callbacks into UpdateAllPrimitiveSceneInfos to make dependency explicit.
 - Removed unnecessary call to FreezeIdTable; real fix is to wait on clean-up task for last scene renderer deletion.

diff --git a/Engine/Plugins/Experimental/Water/Source/Runtime/Private/WaterInfoRendering.cpp b/Engine/Plugins/Experimental/Water/Source/Runtime/Private/WaterInfoRendering.cpp
--- a/Engine/Plugins/Experimental/Water/Source/Runtime/Private/WaterInfoRendering.cpp
+++ b/Engine/Plugins/Experimental/Water/Source/Runtime/Private/WaterInfoRendering.cpp
@@ -419,7 +419,12 @@ private:
 		for (const TPair<uint32, FLandscapeRenderSystem*>& Pair : *LandscapeRenderSystems)
 		{
 			FLandscapeRenderSystem* LandscapeRenderSystem = Pair.Value;
-			LandscapeRenderSystem->ForcedLODOverride = LandscapeLODOverridesToRestore.FindChecked(LandscapeRenderSystem);
+
+			// Ignore landscape systems that were removed within the scope.
+			if (int8* ForcedLODOverride = LandscapeLODOverridesToRestore.Find(LandscapeRenderSystem))
+			{
+				LandscapeRenderSystem->ForcedLODOverride = *ForcedLODOverride;
+			}
 		}
 	}
 
@@ -462,9 +467,6 @@ static void UpdateWaterInfoRendering_RenderThread(
 
 		FRDGBuilder GraphBuilder(RHICmdList, RDG_EVENT_NAME("WaterInfoDepthRendering"), ERDGBuilderFlags::AllowParallelExecute);
 
-		// We need to execute the pre-render view extensions before we do any view dependent work.
-		FSceneRenderer::ViewExtensionPreRender_RenderThread(GraphBuilder, DepthRenderer);
-		
 		FRDGTextureRef TargetTexture = RegisterExternalTexture(GraphBuilder, RenderTarget->GetRenderTargetTexture(), TEXT("WaterDepthTarget"));
 
 		FRDGTextureDesc DepthTextureDesc(TargetTexture->Desc);
@@ -523,9 +525,6 @@ static void UpdateWaterInfoRendering_RenderThread(
 
 		FRDGBuilder GraphBuilder(RHICmdList, RDG_EVENT_NAME("WaterInfoColorRendering"), ERDGBuilderFlags::AllowParallelExecute);
 
-		// We need to execute the pre-render view extensions before we do any view dependent work.
-		FSceneRenderer::ViewExtensionPreRender_RenderThread(GraphBuilder, ColorRenderer);
-
 		FRDGTextureRef TargetTexture = RegisterExternalTexture(GraphBuilder, RenderTarget->GetRenderTargetTexture(), TEXT("WaterColorTarget"));
 
 		FRDGTextureDesc ColorTextureDesc(TargetTexture->Desc);
@@ -578,10 +577,6 @@ static void UpdateWaterInfoRendering_RenderThread(
 
 		FRDGBuilder GraphBuilder(RHICmdList, RDG_EVENT_NAME("WaterInfoDilationRendering"), ERDGBuilderFlags::AllowParallelExecute);
 
-
-		// We need to execute the pre-render view extensions before we do any view dependent work.
-		FSceneRenderer::ViewExtensionPreRender_RenderThread(GraphBuilder, DilationRenderer);
-
 		FRDGTextureRef TargetTexture = RegisterExternalTexture(GraphBuilder, RenderTarget->GetRenderTargetTexture(), TEXT("WaterDilationTarget"));
 
 		FRDGTextureRef DepthTexture = GraphBuilder.RegisterExternalTexture(ExtractedDepthTexture);
diff --git a/Engine/Source/Runtime/Engine/Private/Materials/MaterialRenderProxy.cpp b/Engine/Source/Runtime/Engine/Private/Materials/MaterialRenderProxy.cpp
--- a/Engine/Source/Runtime/Engine/Private/Materials/MaterialRenderProxy.cpp
+++ b/Engine/Source/Runtime/Engine/Private/Materials/MaterialRenderProxy.cpp
@@ -88,7 +88,6 @@ FUniformExpressionCacheAsyncUpdateScope::FUniformExpressionCacheAsyncUpdateScope
 		[](FRHICommandList&)
 		{
 			GUniformExpressionCacheAsyncUpdateTask.Begin();
-			FMaterialRenderProxy::UpdateDeferredCachedUniformExpressions();
 		});
 }
 
diff --git a/Engine/Source/Runtime/Engine/Public/Materials/MaterialRenderProxy.h b/Engine/Source/Runtime/Engine/Public/Materials/MaterialRenderProxy.h
--- a/Engine/Source/Runtime/Engine/Public/Materials/MaterialRenderProxy.h
+++ b/Engine/Source/Runtime/Engine/Public/Materials/MaterialRenderProxy.h
@@ -73,10 +73,8 @@ private:
 #endif
 };
 
-/** Defines a scope to update deferred uniform expression caches using an async task to fill uniform buffers. The scope
- *  attempts to launch async updates immediately, but further async updates can launch within the scope. Only one async
- *  task is launched at a time though. The async task is synced when the scope destructs. The scope enqueues render commands
- *  so it can be used on the game or render threads.
+/** Defines a scope to update deferred uniform expression caches using an async task to fill uniform buffers. If expression
+ *  caches are updated within the scope, an async task may be launched. Otherwise the update is synchronous.
  */
 class FUniformExpressionCacheAsyncUpdateScope
 {
diff --git a/Engine/Source/Runtime/Renderer/Private/DeferredShadingRenderer.cpp b/Engine/Source/Runtime/Renderer/Private/DeferredShadingRenderer.cpp
--- a/Engine/Source/Runtime/Renderer/Private/DeferredShadingRenderer.cpp
+++ b/Engine/Source/Runtime/Renderer/Private/DeferredShadingRenderer.cpp
@@ -2475,12 +2475,12 @@ void FDeferredShadingSceneRenderer::Render(FRDGBuilder& GraphBuilder)
 
 	GPU_MESSAGE_SCOPE(GraphBuilder);
 
+	FInitViewTaskDatas InitViewTaskDatas = OnRenderBegin(GraphBuilder, FGlobalDynamicBuffers(DynamicIndexBufferForInitViews, DynamicVertexBufferForInitViews, DynamicReadBufferForInitViews));
+
 	FRDGExternalAccessQueue ExternalAccessQueue;
 	TUniquePtr<FVirtualTextureUpdater> VirtualTextureUpdater;
 	FLumenSceneFrameTemporaries LumenFrameTemporaries;
 
-	FInitViewTaskDatas InitViewTaskDatas = UpdateScene(GraphBuilder, FGlobalDynamicBuffers(DynamicIndexBufferForInitViews, DynamicVertexBufferForInitViews, DynamicReadBufferForInitViews));
-
 	FGPUSceneScopeBeginEndHelper GPUSceneScopeBeginEndHelper(Scene->GPUScene, GPUSceneDynamicContext, Scene);
 
 	const bool bUseVirtualTexturing = UseVirtualTexturing(FeatureLevel);
@@ -4475,7 +4475,7 @@ void FDeferredShadingSceneRenderer::Render(FRDGBuilder& GraphBuilder)
 		SCOPE_CYCLE_COUNTER(STAT_FDeferredShadingSceneRenderer_RenderFinish);
 		RDG_GPU_STAT_SCOPE(GraphBuilder, FrameRenderFinish);
 		GraphBuilder.SetCommandListStat(GET_STATID(STAT_CLM_RenderFinish));
-		RenderFinish(GraphBuilder, ViewFamilyTexture);
+		OnRenderFinish(GraphBuilder, ViewFamilyTexture);
 		GraphBuilder.SetCommandListStat(GET_STATID(STAT_CLM_AfterFrame));
 		GraphBuilder.AddDispatchHint();
 		GraphBuilder.FlushSetupQueue();
diff --git a/Engine/Source/Runtime/Renderer/Private/MobileShadingRenderer.cpp b/Engine/Source/Runtime/Renderer/Private/MobileShadingRenderer.cpp
--- a/Engine/Source/Runtime/Renderer/Private/MobileShadingRenderer.cpp
+++ b/Engine/Source/Runtime/Renderer/Private/MobileShadingRenderer.cpp
@@ -700,8 +700,6 @@ void FMobileSceneRenderer::InitViews(
 
 		UpdateDirectionalLightUniformBuffers(GraphBuilder, Views[0]);
 	}
-	
-	OnStartRender(RHICmdList);
 }
 
 static void BeginOcclusionScope(FRDGBuilder& GraphBuilder, TArray<FViewInfo>& Views)
@@ -805,7 +803,7 @@ void FMobileSceneRenderer::Render(FRDGBuilder& GraphBuilder)
 	RDG_RHI_EVENT_SCOPE(GraphBuilder, MobileSceneRender);
 	RDG_RHI_GPU_STAT_SCOPE(GraphBuilder, MobileSceneRender);
 
-	IVisibilityTaskData* VisibilityTaskData = UpdateScene(GraphBuilder, FGlobalDynamicBuffers(DynamicIndexBuffer, DynamicVertexBuffer, DynamicReadBuffer));
+	IVisibilityTaskData* VisibilityTaskData = OnRenderBegin(GraphBuilder, FGlobalDynamicBuffers(DynamicIndexBuffer, DynamicVertexBuffer, DynamicReadBuffer));
 
 	FRDGExternalAccessQueue ExternalAccessQueue;
 
@@ -1161,7 +1159,7 @@ void FMobileSceneRenderer::Render(FRDGBuilder& GraphBuilder)
 
 	GraphBuilder.SetCommandListStat(GET_STATID(STAT_CLMM_SceneEnd));
 
-	RenderFinish(GraphBuilder, ViewFamilyTexture);
+	OnRenderFinish(GraphBuilder, ViewFamilyTexture);
 
 	if (bRendererOutputFinalSceneColor)
 	{
diff --git a/Engine/Source/Runtime/Renderer/Private/PlanarReflectionRendering.cpp b/Engine/Source/Runtime/Renderer/Private/PlanarReflectionRendering.cpp
--- a/Engine/Source/Runtime/Renderer/Private/PlanarReflectionRendering.cpp
+++ b/Engine/Source/Runtime/Renderer/Private/PlanarReflectionRendering.cpp
@@ -353,9 +353,6 @@ static void UpdatePlanarReflectionContents_RenderThread(
 	const ERHIFeatureLevel::Type FeatureLevel = SceneRenderer->FeatureLevel;
 	FRDGBuilder GraphBuilder(RHICmdList, RDG_EVENT_NAME("PlanarReflection"), FSceneRenderer::GetRDGParalelExecuteFlags(FeatureLevel));
 
-	// We need to execute the pre-render view extensions before we do any view dependent work.
-	FSceneRenderer::ViewExtensionPreRender_RenderThread(GraphBuilder, SceneRenderer);
-
 	// Make sure we render to the same set of GPUs as the main scene renderer.
 	if (MainSceneRenderer->ViewFamily.RenderTarget != nullptr)
 	{
diff --git a/Engine/Source/Runtime/Renderer/Private/ReflectionEnvironmentCapture.cpp b/Engine/Source/Runtime/Renderer/Private/ReflectionEnvironmentCapture.cpp
--- a/Engine/Source/Runtime/Renderer/Private/ReflectionEnvironmentCapture.cpp
+++ b/Engine/Source/Runtime/Renderer/Private/ReflectionEnvironmentCapture.cpp
@@ -513,9 +513,6 @@ void CaptureSceneToScratchCubemap(
 
 	FRDGBuilder GraphBuilder(RHICmdList, RDG_EVENT_NAME("CubeMapCapture"), FSceneRenderer::GetRDGParalelExecuteFlags(FeatureLevel));
 
-	// We need to execute the pre-render view extensions before we do any view dependent work.
-	FSceneRenderer::ViewExtensionPreRender_RenderThread(GraphBuilder, SceneRenderer);
-
 	{
 		RDG_EVENT_SCOPE(GraphBuilder, "CubeMapCapture");
 
diff --git a/Engine/Source/Runtime/Renderer/Private/RendererScene.cpp b/Engine/Source/Runtime/Renderer/Private/RendererScene.cpp
--- a/Engine/Source/Runtime/Renderer/Private/RendererScene.cpp
+++ b/Engine/Source/Runtime/Renderer/Private/RendererScene.cpp
@@ -65,6 +65,7 @@
 #include "GPUScene.h"
 #include "HAL/LowLevelMemTracker.h"
 #include "VT/RuntimeVirtualTextureSceneProxy.h"
+#include "VT/VirtualTextureSystem.h"
 #include "HairStrandsInterface.h"
 #include "VelocityRendering.h"
 #include "RectLightSceneProxy.h"
@@ -6144,6 +6145,9 @@ void FScene::UpdateAllPrimitiveSceneInfos(FRDGBuilder& GraphBuilder, EUpdateAllP
 	UpdateCachedShadowState(SceneUpdateChangeSetStorage.GetPreUpdateSet(), SceneUpdateChangeSetStorage.GetPostUpdateSet());
 	ShadowScene->PostSceneUpdate(SceneUpdateChangeSetStorage.GetPreUpdateSet(), SceneUpdateChangeSetStorage.GetPostUpdateSet());
 
+	// Must run prior to static mesh gathering as callbacks can evaluate uniform expression caches.
+	FVirtualTextureSystem::Get().CallPendingCallbacks();
+
 	if (SceneInfosWithStaticDrawListUpdate.Num() > 0)
 	{
 		FPrimitiveSceneInfo::AddStaticMeshes(this, SceneInfosWithStaticDrawListUpdate, false);
diff --git a/Engine/Source/Runtime/Renderer/Private/SceneCaptureRendering.cpp b/Engine/Source/Runtime/Renderer/Private/SceneCaptureRendering.cpp
--- a/Engine/Source/Runtime/Renderer/Private/SceneCaptureRendering.cpp
+++ b/Engine/Source/Runtime/Renderer/Private/SceneCaptureRendering.cpp
@@ -303,9 +303,6 @@ static void UpdateSceneCaptureContentDeferred_RenderThread(
 	FRDGBuilder GraphBuilder(RHICmdList, RDG_EVENT_NAME("SceneCapture"), FSceneRenderer::GetRDGParalelExecuteFlags(FeatureLevel));
 #endif
 
-	// We need to execute the pre-render view extensions before we do any view dependent work.
-	FSceneRenderer::ViewExtensionPreRender_RenderThread(GraphBuilder, SceneRenderer);
-
 	{
 		FRDGTextureRef TargetTexture = RegisterExternalTexture(GraphBuilder, RenderTarget->GetRenderTargetTexture(), TEXT("SceneCaptureTarget"));
 		FRDGTextureRef ShaderResourceTexture = RegisterExternalTexture(GraphBuilder, RenderTargetTexture->TextureRHI, TEXT("SceneCaptureTexture"));
@@ -389,9 +386,6 @@ void UpdateSceneCaptureContentMobile_RenderThread(
 	FRDGBuilder GraphBuilder(RHICmdList, RDG_EVENT_NAME("SceneCaptureMobile"));
 #endif
 
-	// We need to execute the pre-render view extensions before we do any view dependent work.
-	FSceneRenderer::ViewExtensionPreRender_RenderThread(GraphBuilder, SceneRenderer);
-
 	{
 		FViewInfo& View = SceneRenderer->Views[0];
 
diff --git a/Engine/Source/Runtime/Renderer/Private/SceneHitProxyRendering.cpp b/Engine/Source/Runtime/Renderer/Private/SceneHitProxyRendering.cpp
--- a/Engine/Source/Runtime/Renderer/Private/SceneHitProxyRendering.cpp
+++ b/Engine/Source/Runtime/Renderer/Private/SceneHitProxyRendering.cpp
@@ -549,8 +549,6 @@ static void DoRenderHitProxies(
 				false,
 				1.0f
 			);
-
-			RHICmdList.EndScene();
 		});
 	}
 }
@@ -558,7 +556,7 @@ static void DoRenderHitProxies(
 
 void FMobileSceneRenderer::RenderHitProxies(FRDGBuilder& GraphBuilder)
 {
-	IVisibilityTaskData* VisibilityTaskData = UpdateScene(GraphBuilder, FGlobalDynamicBuffers(DynamicIndexBuffer, DynamicVertexBuffer, DynamicReadBuffer));
+	IVisibilityTaskData* VisibilityTaskData = OnRenderBegin(GraphBuilder, FGlobalDynamicBuffers(DynamicIndexBuffer, DynamicVertexBuffer, DynamicReadBuffer));
 
 	GPU_MESSAGE_SCOPE(GraphBuilder);
 
@@ -591,6 +589,8 @@ void FMobileSceneRenderer::RenderHitProxies(FRDGBuilder& GraphBuilder)
 
 	GEngine->GetPostRenderDelegateEx().Broadcast(GraphBuilder);
 #endif
+
+	OnRenderFinish(GraphBuilder, nullptr);
 }
 
 void FDeferredShadingSceneRenderer::RenderHitProxies(FRDGBuilder& GraphBuilder)
@@ -599,7 +599,7 @@ void FDeferredShadingSceneRenderer::RenderHitProxies(FRDGBuilder& GraphBuilder)
 
 	CommitFinalPipelineState();
 
-	IVisibilityTaskData* VisibilityTaskData = UpdateScene(GraphBuilder, FGlobalDynamicBuffers(DynamicIndexBufferForInitViews, DynamicVertexBufferForInitViews, DynamicReadBufferForInitViews));
+	IVisibilityTaskData* VisibilityTaskData = OnRenderBegin(GraphBuilder, FGlobalDynamicBuffers(DynamicIndexBufferForInitViews, DynamicVertexBufferForInitViews, DynamicReadBufferForInitViews));
 
 	GPU_MESSAGE_SCOPE(GraphBuilder);
 
@@ -739,6 +739,8 @@ void FDeferredShadingSceneRenderer::RenderHitProxies(FRDGBuilder& GraphBuilder)
 
 	AddDispatchToRHIThreadPass(GraphBuilder);
 #endif
+
+	OnRenderFinish(GraphBuilder, nullptr);
 }
 
 #if WITH_EDITOR
diff --git a/Engine/Source/Runtime/Renderer/Private/SceneRendering.cpp b/Engine/Source/Runtime/Renderer/Private/SceneRendering.cpp
--- a/Engine/Source/Runtime/Renderer/Private/SceneRendering.cpp
+++ b/Engine/Source/Runtime/Renderer/Private/SceneRendering.cpp
@@ -2828,69 +2828,12 @@ FIntPoint FSceneRenderer::GetDesiredInternalBufferSize(const FSceneViewFamily& V
 	return DesiredBufferSize;
 }
 
-inline EUpdateAllPrimitiveSceneInfosAsyncOps GetUpdateAllPrimitiveSceneInfosAsyncOps()
-{
-	EUpdateAllPrimitiveSceneInfosAsyncOps AsyncOps = EUpdateAllPrimitiveSceneInfosAsyncOps::None;
-
-	if (GAsyncCreateLightPrimitiveInteractions > 0)
-	{
-		AsyncOps |= EUpdateAllPrimitiveSceneInfosAsyncOps::CreateLightPrimitiveInteractions;
-	}
-
-	if (GAsyncCacheMeshDrawCommands > 0)
-	{
-		AsyncOps |= EUpdateAllPrimitiveSceneInfosAsyncOps::CacheMeshDrawCommands;
-	}
-
-	return AsyncOps;
-}
-
 FSceneRenderer::ERendererOutput FSceneRenderer::GetRendererOutput() const
 {
 	const bool bSceneCaptureDepthPrepass = Views[0].bIsSceneCapture && (ViewFamily.SceneCaptureSource == ESceneCaptureSource::SCS_SceneDepth || ViewFamily.SceneCaptureSource == ESceneCaptureSource::SCS_DeviceDepth);
 	return bSceneCaptureDepthPrepass && GSceneCaptureDepthPrepassOptimization ? ERendererOutput::DepthPrepassOnly : ERendererOutput::FinalSceneColor;
 }
 
-IVisibilityTaskData* FSceneRenderer::UpdateScene(FRDGBuilder& GraphBuilder, FGlobalDynamicBuffers GlobalDynamicBuffers)
-{
-	// Needs to run before UpdateAllPrimitiveSceneInfos, as it may call OnVirtualTextureDestroyedCB, which modifies the uniform
-	// expression cache. UpdateAllPrimitiveSceneInfos generates mesh draw commands, which use the uniform expression cache.
-	FVirtualTextureSystem::Get().CallPendingCallbacks();
-
-	/**
-	  * UpdateStaticMeshes removes and re-creates cached FMeshDrawCommands.  If there are multiple scene renderers being run together,
-	  * we need allocated pipeline state IDs not to change, in case async tasks related to prior scene renderers are still in flight
-	  * (FSubmitNaniteMaterialPassCommandsAnyThreadTask or FDrawVisibleMeshCommandsAnyThreadTask).  So we freeze pipeline state IDs,
-	  * preventing them from being de-allocated even if their reference count temporarily goes to zero during calls to
-	  * RemoveCachedMeshDrawCommands followed by CacheMeshDrawCommands (or the Nanite equivalent).
-	  *
-	  * Note that on the first scene renderer, we do want to de-allocate items, so they can be permanently released if no longer in use
-	  * (for example, if there was an impactful change to a render proxy by game logic), but the assumption is that sequential renders
-	  * of the same scene from different views can't make such changes.
-	  */
-	if (!bIsFirstSceneRenderer)
-	{
-		FGraphicsMinimalPipelineStateId::FreezeIdTable(true);
-	}
-
-	Scene->UpdateAllPrimitiveSceneInfos(GraphBuilder, GetUpdateAllPrimitiveSceneInfosAsyncOps());
-
-	if (!bIsFirstSceneRenderer)
-	{
-		FGraphicsMinimalPipelineStateId::FreezeIdTable(false);
-	}
-
-	PrepareViewRectsForRendering(GraphBuilder.RHICmdList);
-
-	InitializeSceneTexturesConfig(ViewFamily.SceneTexturesConfig, ViewFamily);
-	FSceneTexturesConfig& SceneTexturesConfig = GetActiveSceneTexturesConfig();
-	FSceneTexturesConfig::Set(SceneTexturesConfig);
-
-	PrepareViewStateForVisibility(SceneTexturesConfig);
-
-	return LaunchVisibilityTasks(GraphBuilder.RHICmdList, *this, GlobalDynamicBuffers);
-}
-
 void FSceneRenderer::PrepareViewRectsForRendering(FRHICommandListImmediate& RHICmdList)
 {
 	check(IsInRenderingThread());
@@ -3398,14 +3341,65 @@ FSceneRenderer::~FSceneRenderer()
 	SortedShadowsForShadowDepthPass.Release();
 }
 
+IVisibilityTaskData* FSceneRenderer::OnRenderBegin(FRDGBuilder& GraphBuilder, FGlobalDynamicBuffers GlobalDynamicBuffers)
+{
+	check(!FDeferredUpdateResource::IsUpdateNeeded());
+
+	EUpdateAllPrimitiveSceneInfosAsyncOps AsyncOps = EUpdateAllPrimitiveSceneInfosAsyncOps::None;
+
+	if (GAsyncCreateLightPrimitiveInteractions > 0)
+	{
+		AsyncOps |= EUpdateAllPrimitiveSceneInfosAsyncOps::CreateLightPrimitiveInteractions;
+	}
+
+	if (GAsyncCacheMeshDrawCommands > 0)
+	{
+		AsyncOps |= EUpdateAllPrimitiveSceneInfosAsyncOps::CacheMeshDrawCommands;
+	}
+
+	Scene->UpdateAllPrimitiveSceneInfos(GraphBuilder, AsyncOps);
+
+	if (!ViewFamily.ViewExtensions.IsEmpty())
+	{
+		{
+			RDG_CSV_STAT_EXCLUSIVE_SCOPE(GraphBuilder, PreRender);
+			SCOPE_CYCLE_COUNTER(STAT_FDeferredShadingSceneRenderer_ViewExtensionPreRenderView);
+
+			for (int32 ViewExt = 0; ViewExt < ViewFamily.ViewExtensions.Num(); ViewExt++)
+			{
+				ViewFamily.ViewExtensions[ViewExt]->PreRenderViewFamily_RenderThread(GraphBuilder, ViewFamily);
+				for (int32 ViewIndex = 0; ViewIndex < ViewFamily.Views.Num(); ViewIndex++)
+				{
+					ViewFamily.ViewExtensions[ViewExt]->PreRenderView_RenderThread(GraphBuilder, Views[ViewIndex]);
+				}
+			}
+		}
+	}
+
+	PrepareViewRectsForRendering(GraphBuilder.RHICmdList);
+
+	InitializeSceneTexturesConfig(ViewFamily.SceneTexturesConfig, ViewFamily);
+	FSceneTexturesConfig& SceneTexturesConfig = GetActiveSceneTexturesConfig();
+	FSceneTexturesConfig::Set(SceneTexturesConfig);
+
+	PrepareViewStateForVisibility(SceneTexturesConfig);
+
+	FVisualizeTexturePresent::OnStartRender(Views[0]);
+
+	GraphBuilder.RHICmdList.BeginScene();
+
+	return LaunchVisibilityTasks(GraphBuilder.RHICmdList, *this, GlobalDynamicBuffers);
+}
+
 /** 
 * Finishes the view family rendering.
 */
-void FSceneRenderer::RenderFinish(FRDGBuilder& GraphBuilder, FRDGTextureRef ViewFamilyTexture)
+void FSceneRenderer::OnRenderFinish(FRDGBuilder& GraphBuilder, FRDGTextureRef ViewFamilyTexture)
 {
 	RDG_EVENT_SCOPE(GraphBuilder, "RenderFinish");
 
 #if !(UE_BUILD_SHIPPING || UE_BUILD_TEST)
+	if (ViewFamilyTexture)
 	{
 		bool bShowPrecomputedVisibilityWarning = false;
 		static const auto* CVarPrecomputedVisibilityWarning = IConsoleManager::Get().FindTConsoleVariableDataInt(TEXT("r.PrecomputedVisibilityWarning"));
@@ -4091,11 +4085,6 @@ FSceneRenderer* FSceneRenderer::CreateSceneRenderer(const FSceneViewFamily* InVi
 	return SceneRenderers[0];
 }
 
-void FSceneRenderer::OnStartRender(FRHICommandListImmediate& RHICmdList)
-{
-	FVisualizeTexturePresent::OnStartRender(Views[0]);
-}
-
 bool FSceneRenderer::ShouldCompositeEditorPrimitives(const FViewInfo& View)
 {
 	if (View.Family->EngineShowFlags.VisualizeHDR || View.Family->UseDebugViewPS())
@@ -4188,36 +4177,6 @@ void FSceneRenderer::UpdatePrimitiveIndirectLightingCacheBuffers(FRHICommandList
 	FRendererModule
 -----------------------------------------------------------------------------*/
 
-/**
-* Helper function performing actual work in render thread.
-*
-* @param SceneRenderer	Scene renderer to use for rendering.
-*/
-void FSceneRenderer::ViewExtensionPreRender_RenderThread(FRDGBuilder& GraphBuilder, FSceneRenderer* SceneRenderer)
-{
-	if (SceneRenderer->ViewFamily.ViewExtensions.IsEmpty() || !SceneRenderer->ViewFamily.EngineShowFlags.Rendering)
-	{
-		return;
-	}
-
-	{
-		{
-			RDG_CSV_STAT_EXCLUSIVE_SCOPE(GraphBuilder, PreRender);
-			SCOPE_CYCLE_COUNTER(STAT_FDeferredShadingSceneRenderer_ViewExtensionPreRenderView);
-
-			for (int ViewExt = 0; ViewExt < SceneRenderer->ViewFamily.ViewExtensions.Num(); ViewExt++)
-			{
-				SceneRenderer->ViewFamily.ViewExtensions[ViewExt]->PreRenderViewFamily_RenderThread(GraphBuilder, SceneRenderer->ViewFamily);
-				for (int ViewIndex = 0; ViewIndex < SceneRenderer->ViewFamily.Views.Num(); ViewIndex++)
-				{
-					SceneRenderer->ViewFamily.ViewExtensions[ViewExt]->PreRenderView_RenderThread(GraphBuilder, SceneRenderer->Views[ViewIndex]);
-				}
-			}
-		}
-	}
-	check(!FDeferredUpdateResource::IsUpdateNeeded());
-}
-
 static int32 GSceneRenderCleanUpMode = 2;
 static FAutoConsoleVariableRef CVarSceneRenderCleanUpMode(
 	TEXT("r.SceneRenderCleanUpMode"),
@@ -4538,8 +4497,6 @@ static void RenderViewFamilies_RenderThread(FRHICommandListImmediate& RHICmdList
 			FSceneRenderer::GetRDGParalelExecuteFlags(FeatureLevel)
 		);
 
-		// We need to execute the pre-render view extensions before we do any view dependent work.
-		FSceneRenderer::ViewExtensionPreRender_RenderThread(GraphBuilder, SceneRenderer);
 #if WITH_GPUDEBUGCRASH
 		if (GRHIGlobals.TriggerGPUCrash != ERequestedGPUCrash::None) 
 		{
diff --git a/Engine/Source/Runtime/Renderer/Private/SceneRendering.h b/Engine/Source/Runtime/Renderer/Private/SceneRendering.h
--- a/Engine/Source/Runtime/Renderer/Private/SceneRendering.h
+++ b/Engine/Source/Runtime/Renderer/Private/SceneRendering.h
@@ -2139,13 +2139,6 @@ public:
 	static bool ShouldCompositeDebugPrimitivesInPostProcess(const FViewInfo& View);
 #endif
 
-	/**
-	* Helper function performing actual work in render thread.
-	*
-	* @param SceneRenderer	Scene renderer to use for rendering.
-	*/
-	static void RENDERER_API ViewExtensionPreRender_RenderThread(FRDGBuilder& GraphBuilder, FSceneRenderer* SceneRenderer);
-
 	/** Called to release any deallocations that were deferred until the next render. */
 	static void CleanUp(FRHICommandListImmediate& RHICmdList);
 
@@ -2299,8 +2292,6 @@ protected:
 
 	ERendererOutput GetRendererOutput() const;
 
-	IVisibilityTaskData* UpdateScene(FRDGBuilder& GraphBuilder, FGlobalDynamicBuffers GlobalDynamicBuffers);
-
 	FDynamicShadowsTaskData* BeginInitDynamicShadows(bool bRunningEarly, IVisibilityTaskData* VisibilityTaskData);
 	void FinishInitDynamicShadows(FRDGBuilder& GraphBuilder, FDynamicShadowsTaskData* TaskData, FGlobalDynamicIndexBuffer& DynamicIndexBuffer, FGlobalDynamicVertexBuffer& DynamicVertexBuffer, FGlobalDynamicReadBuffer& DynamicReadBuffer, FInstanceCullingManager& InstanceCullingManager, FRDGExternalAccessQueue& ExternalAccessQueue);
 	FDynamicShadowsTaskData* InitDynamicShadows(FRDGBuilder& GraphBuilder, FGlobalDynamicIndexBuffer& DynamicIndexBuffer, FGlobalDynamicVertexBuffer& DynamicVertexBuffer, FGlobalDynamicReadBuffer& DynamicReadBuffer, FInstanceCullingManager& InstanceCullingManager, FRDGExternalAccessQueue& ExternalAccessQueue);
@@ -2434,8 +2425,9 @@ protected:
 	bool ShouldRenderTranslucency() const;
 	bool ShouldRenderTranslucency(ETranslucencyPass::Type TranslucencyPass) const;
 
-	/** Updates state for the end of the frame. */
-	void RenderFinish(FRDGBuilder& GraphBuilder, FRDGTextureRef ViewFamilyTexture);
+	/** Called at the begin / finish of the scene render. */
+	IVisibilityTaskData* OnRenderBegin(FRDGBuilder& GraphBuilder, FGlobalDynamicBuffers GlobalDynamicBuffers);
+	void OnRenderFinish(FRDGBuilder& GraphBuilder, FRDGTextureRef ViewFamilyTexture);
 
 	bool RenderCustomDepthPass(
 		FRDGBuilder& GraphBuilder,
@@ -2444,8 +2436,6 @@ protected:
 		TConstArrayView<Nanite::FRasterResults> PrimaryNaniteRasterResults,
 		TConstArrayView<Nanite::FPackedView> PrimaryNaniteViews);
 
-	void OnStartRender(FRHICommandListImmediate& RHICmdList);
-
 	void UpdatePrimitiveIndirectLightingCacheBuffers(FRHICommandListBase& RHICmdList);
 
 	void RenderPlanarReflection(class FPlanarReflectionSceneProxy* ReflectionSceneProxy);
diff --git a/Engine/Source/Runtime/Renderer/Private/SceneVisibility.cpp b/Engine/Source/Runtime/Renderer/Private/SceneVisibility.cpp
--- a/Engine/Source/Runtime/Renderer/Private/SceneVisibility.cpp
+++ b/Engine/Source/Runtime/Renderer/Private/SceneVisibility.cpp
@@ -4015,9 +4015,6 @@ void FSceneRenderer::PreVisibilityFrameSetup(FRDGBuilder& GraphBuilder)
 {
 	FRHICommandListImmediate& RHICmdList = GraphBuilder.RHICmdList;
 
-	// Notify the RHI we are beginning to render a scene.
-	RHICmdList.BeginScene();
-
 	if (GetRendererOutput() == ERendererOutput::FinalSceneColor)
 	{
 		if (Views.Num() > 0 && !ViewFamily.EngineShowFlags.HitProxies)
@@ -4928,11 +4925,6 @@ void FDeferredShadingSceneRenderer::BeginInitViews(
 		}
 	}
 
-	{
-		QUICK_SCOPE_CYCLE_COUNTER(STAT_InitViews_OnStartRender);
-		OnStartRender(RHICmdList);
-	}
-
 	if (GDynamicRHI->RHIIncludeOptionalFlushes())
 	{
 		RHICmdList.ImmediateFlush(EImmediateFlushType::DispatchToRHIThread);
diff --git a/Engine/Source/Runtime/Renderer/Private/VT/VirtualTextureSystem.cpp b/Engine/Source/Runtime/Renderer/Private/VT/VirtualTextureSystem.cpp
--- a/Engine/Source/Runtime/Renderer/Private/VT/VirtualTextureSystem.cpp
+++ b/Engine/Source/Runtime/Renderer/Private/VT/VirtualTextureSystem.cpp
@@ -2661,7 +2661,7 @@ TUniquePtr<FVirtualTextureUpdater> FVirtualTextureSystem::BeginUpdate(FRDGBuilde
 
 	UE::TScopeLock Lock(Mutex);
 
-	checkf(Producers.HasPendingCallbacks() == false, TEXT("FVirtualTextureSystem::CallPendingCallbacks(), typically called in FSceneRenderer::UpdateScene(), must run before FVirtualTextureSystem::BeginUpdate()"));
+	checkf(Producers.HasPendingCallbacks() == false, TEXT("FVirtualTextureSystem::CallPendingCallbacks(), called in UpdateAllPrimitiveSceneInfos(), must run before FVirtualTextureSystem::BeginUpdate()"));
 
 	AllocateResources(GraphBuilder);
 
