Fix Z-Fighting on mobile platform with partially DXC enabled shader compilation.

#rnx

diff --git a/Engine/Source/Runtime/Renderer/Private/DepthRendering.h b/Engine/Source/Runtime/Renderer/Private/DepthRendering.h
--- a/Engine/Source/Runtime/Renderer/Private/DepthRendering.h
+++ b/Engine/Source/Runtime/Renderer/Private/DepthRendering.h
@@ -97,6 +97,17 @@ public:
 			&& !Parameters.VertexFactoryType->SupportsNaniteRendering();
 	}
 
+	static void ModifyCompilationEnvironment(const FMaterialShaderPermutationParameters& Parameters, FShaderCompilerEnvironment& OutEnvironment)
+	{
+		FMeshMaterialShader::ModifyCompilationEnvironment(Parameters, OutEnvironment);
+
+		// @lh-todo: Same workaround as for the VS of MobileBasePass. See TMobileBasePassVSPolicyParamType::ModifyCompilationEnvironment for details.
+		if (!Strata::IsStrataEnabled())
+		{
+			OutEnvironment.SetCompileArgument(TEXT("WORKAROUND_DISABLE_rShadersForceDXC"), true);
+		}
+	}
+
 	void GetShaderBindings(
 		const FScene* Scene,
 		ERHIFeatureLevel::Type FeatureLevel,
@@ -142,6 +153,12 @@ public:
 		
 		OutEnvironment.SetDefine(TEXT("ALLOW_DEBUG_VIEW_MODES"), AllowDebugViewmodes(Parameters.Platform));
 		OutEnvironment.SetDefine(TEXT("SCENE_TEXTURES_DISABLED"), 1u);
+
+		// @lh-todo: Same workaround as for the VS of MobileBasePass. See TMobileBasePassVSPolicyParamType::ModifyCompilationEnvironment for details.
+		if (!Strata::IsStrataEnabled())
+		{
+			OutEnvironment.SetCompileArgument(TEXT("WORKAROUND_DISABLE_rShadersForceDXC"), true);
+		}
 	}
 
 	FDepthOnlyPS() {}
