Optimises TSR's spatial anti-aliaser costs based on much it is needed on screen

diff --git a/Engine/Shaders/Private/TemporalSuperResolution/TSRCommon.ush b/Engine/Shaders/Private/TemporalSuperResolution/TSRCommon.ush
--- a/Engine/Shaders/Private/TemporalSuperResolution/TSRCommon.ush
+++ b/Engine/Shaders/Private/TemporalSuperResolution/TSRCommon.ush
@@ -91,6 +91,8 @@
 // Minimum weight of a pixel in history.
 #define MIN_IMPLICIT_HISTORY_WEIGHT ceil(tsr_half(255.0) * rcp(tsr_half(MAX_SAMPLE_COUNT)) * tsr_half(MIN_REPROJECTION_SAMPLE_COUNT)) * rcp(tsr_half(255.0))
 
+#define SPATIAL_ANTI_ALIASER_MIN_LUMIMANCE 0.05
+
 
 // K = Center of the nearest input pixel.
 // O = Center of the output pixel.
diff --git a/Engine/Shaders/Private/TemporalSuperResolution/TSRFilterAntiAliasing.usf b/Engine/Shaders/Private/TemporalSuperResolution/TSRFilterAntiAliasing.usf
deleted file mode 100644
diff --git a/Engine/Shaders/Private/TemporalSuperResolution/TSRRejectShading.usf b/Engine/Shaders/Private/TemporalSuperResolution/TSRRejectShading.usf
--- a/Engine/Shaders/Private/TemporalSuperResolution/TSRRejectShading.usf
+++ b/Engine/Shaders/Private/TemporalSuperResolution/TSRRejectShading.usf
@@ -132,6 +132,7 @@ RWTexture2D<tsr_half4> HistoryMoireOutput;
 RWTexture2D<tsr_half4> HistoryRejectionOutput;
 RWTexture2D<tsr_halfC> InputSceneColorOutput;
 RWTexture2D<tsr_half> InputSceneColorLdrLumaOutput;
+RWTexture2D<tsr_ushort> AntiAliasMaskOutput;
 
 
 //------------------------------------------------------- DEBUG
@@ -323,8 +324,8 @@ void ComputeMoireError(
 template<bool ApplyMoireError>
 void MeasureRejection(
 	tsr_simd_bool  bIsValid,
-	tsr_simd_halfC ExposedInput,
-	tsr_simd_halfC ExposedHistory,
+	tsr_simd_halfC ClampedInput,
+	tsr_simd_halfC ClampedHistory,
 	tsr_simd_half  MoireError,
 	out tsr_simd_half OutRejectionBlendFinal,
 	out tsr_simd_half OutRejectionClampBlend)
@@ -332,9 +333,6 @@ void MeasureRejection(
 	const tsr_half FilteringWeight = rcp(tsr_half(1.0 + 4 * 0.5 + 4 * 0.25));
 	tsr_simd_halfC BackbufferQuantizationErrorVector = tsr_simd_halfC::Const(MeasureBackbufferLDRQuantizationError());
 		
-	tsr_simd_halfC ClampedInput, ClampedHistory;
-	AnnihilateMutually3x3(ExposedInput, ExposedHistory, /* out */ ClampedInput, /* out */ ClampedHistory);
-
 	// Filter clamped input and history to increase quantization of the rejection detection
 	tsr_simd_halfC FilteredInput   = Blur3x3(ClampedInput);
 	tsr_simd_halfC FilteredHistory = Blur3x3(ClampedHistory);
@@ -615,6 +613,7 @@ void MainCS(
 	tsr_simd_half RejectionBlendFinal = tsr_simd_half::Const(1.0);
 	tsr_simd_half RejectionClampBlend = tsr_simd_half::Const(1.0);
 	tsr_simd_bool bResurrectHistory = tsr_simd_bool::Const(false);
+	tsr_simd_bool bAliasingIsVisible = tsr_simd_bool::Const(false);
 	{
 		// Convert input to LDR
 		tsr_simd_halfC ExposedInput = OriginalInput * rcp(OriginalInput + tsr_simd_halfC::Const(tsr_half(PerceptionAdd)));
@@ -625,10 +624,27 @@ void MainCS(
 			/* SourcePerceptionAdd = */ 1.0,
 			/* DestPerceptionAdd   = */ PerceptionAdd);
 			
+		tsr_simd_halfC ExposedInputMin;
+		tsr_simd_halfC ExposedInputMax;
+		MinMax3x3(ExposedInput, /* out */ ExposedInputMin, /* out */ ExposedInputMax);
+		
+		{
+			#if CONFIG_SCENE_COLOR_ALPHA
+				tsr_simd_half ExposedInputLuminanceBox = dot(ExposedInputMax - ExposedInputMin, tsr_simd_halfC::Const(tsr_halfC(0.299f, 0.587f, 0.114f, 0.0)));
+			#else
+				tsr_simd_half ExposedInputLuminanceBox = dot(ExposedInputMax - ExposedInputMin, tsr_simd_halfC::Const(tsr_halfC(0.299f, 0.587f, 0.114f)));
+			#endif
+
+			bAliasingIsVisible = ExposedInputLuminanceBox > tsr_simd_half::Const(SPATIAL_ANTI_ALIASER_MIN_LUMIMANCE);
+		}
+
+		tsr_simd_halfC ClampedInput = Clamp3x3(ExposedInput, fastClamp(ExposedHistory, ExposedInputMin, ExposedInputMax));
+		tsr_simd_halfC ClampedHistory = AnnihilateToGuide3x3(ExposedHistory, /* Guide = */ ExposedInput);
+
 		MeasureRejection</* ApplyMoireError = */ CONFIG_MOIRE_DETECTION != 0>(
 			bIsValid,
-			ExposedInput,
-			ExposedHistory,
+			ClampedInput,
+			ClampedHistory,
 			MoireError,
 			/* out */ RejectionBlendFinal,
 			/* out */ RejectionClampBlend);
@@ -655,12 +671,15 @@ void MainCS(
 			BRANCH
 			if (bComputeResurrection)
 			{
+				tsr_simd_halfC ClampedResurrectedInput, ClampedResurrectedHistory;
+				AnnihilateMutually3x3(ExposedInput, ExposedResurrectedHistory, /* out */ ClampedResurrectedInput, /* out */ ClampedResurrectedHistory);
+
 				tsr_simd_half ResurrectionRejectionBlendFinal = tsr_simd_half::Const(1.0);
 				tsr_simd_half ResurrectionRejectionClampBlend = tsr_simd_half::Const(1.0);
 				MeasureRejection</* ApplyMoireError = */ false>(
 					bIsValid,
-					ExposedInput,
-					ExposedResurrectedHistory,
+					ClampedResurrectedInput,
+					ClampedResurrectedHistory,
 					/* MoireError = */ tsr_simd_half::Const(0.0),
 					/* out */ ResurrectionRejectionBlendFinal,
 					/* out */ ResurrectionRejectionClampBlend);
@@ -674,6 +693,8 @@ void MainCS(
 		}
 		#endif
 	}
+	
+	tsr_simd_bool bIsDisoccluded = not(bIsValid);
 
 	// Update the guide for next frame.
 	tsr_simd_halfC FinalGuide;
@@ -688,11 +709,10 @@ void MainCS(
 		for (uint SimdIndex = 0; SimdIndex < SIMD_SIZE; SimdIndex++)
 		{
 			bool bIsOffScreen = false;
-			bool bIsDisoccluded = !bIsValid.GetElement(SimdIndex);
+			bool bIsDisoccludedLocal = bIsDisoccluded.GetElement(SimdIndex);
 			tsr_half LocalRejectionBlendFinal = RejectionBlendFinal.GetElement(SimdIndex);
 
-			//tsr_half TotalBlend = select(or(bIsOffScreen, bIsDisoccluded), tsr_half(1.0), tsr_half(0.0));
-			tsr_half TotalBlend = select(or(bIsOffScreen, bIsDisoccluded), tsr_half(1.0), saturate(tsr_half(1.0) - LocalRejectionBlendFinal * tsr_half(4.0)));
+			tsr_half TotalBlend = select(or(bIsOffScreen, bIsDisoccludedLocal), tsr_half(1.0), saturate(tsr_half(1.0) - LocalRejectionBlendFinal * tsr_half(4.0)));
 			tsr_half RejectedBlendFactor = max(tsr_half(TheoricBlendFactor), tsr_half(1.0) - LocalRejectionBlendFinal);
 
 			BlendFinal.SetElement(SimdIndex, max(RejectedBlendFactor, TotalBlend));
@@ -709,8 +729,6 @@ void MainCS(
 	tsr_simd_half DisableHistoryClamp;
 	tsr_simd_half IncreaseValidityMultiplier;
 	{
-		tsr_simd_bool bIsDisoccluded = not(bIsValid);
-
 		DisableHistoryClamp        = select(and(bIsDisoccluded, not(bResurrectHistory)), tsr_simd_half::Const(0.0), min(RejectionClampBlend, VelocityEdge));
 		IncreaseValidityMultiplier = select(and(bIsDisoccluded, not(bResurrectHistory)), tsr_simd_half::Const(0.0), RejectionClampBlend);
 	}
@@ -727,6 +745,9 @@ void MainCS(
 	}
 	#endif
 
+	// Whether to run spatial anti-aliaser on the frame. Must match SpatialAntiAliasingLerp.
+	tsr_simd_bool bSpatialAntiAlias = and(bAliasingIsVisible, or(RejectionBlendFinal < tsr_simd_half::Const(tsr_half(0.25)), and(bIsDisoccluded, not(bResurrectHistory))));
+
 	// Output data.
 	PLATFORM_SPECIFIC_ISOLATE
 	{
@@ -768,6 +789,7 @@ void MainCS(
 				DisableHistoryClamp.GetElement(SimdIndex) - tsr_half(0.5 * rcp(256.0)),
 				tsr_half(255.5 * rcp(256.0)) - IncreaseValidityMultiplier.GetElement(SimdIndex),
 				BitMask);
+			AntiAliasMaskOutput[OutputPixelPos] = select(bSpatialAntiAlias.GetElement(SimdIndex), tsr_ushort(1), tsr_ushort(0));
 		}
 		
 		#if DEBUG_OUTPUT
diff --git a/Engine/Shaders/Private/TemporalSuperResolution/TSRSpatialAntiAliasing.usf b/Engine/Shaders/Private/TemporalSuperResolution/TSRSpatialAntiAliasing.usf
--- a/Engine/Shaders/Private/TemporalSuperResolution/TSRSpatialAntiAliasing.usf
+++ b/Engine/Shaders/Private/TemporalSuperResolution/TSRSpatialAntiAliasing.usf
@@ -7,8 +7,6 @@
 
 #define TILE_SIZE 8
 
-#define CONFIG_MIN_LUMA_ALIASING 0.05
-
 #define CONFIG_COMPUTE_NOISE_FILTERING 1
 
 #if DIM_QUALITY_PRESET == 1
@@ -34,10 +32,10 @@
 
 //------------------------------------------------------- PARAMETERS
 
-Texture2D<tsr_half> InputSceneColorLdrLumaTexture;
+Texture2D<tsr_ushort> AntiAliasMaskTexture;
+Texture2D<tsr_half>  InputSceneColorLdrLumaTexture;
 
-RWTexture2D<uint> AntiAliasingOutput;
-RWTexture2D<tsr_half> NoiseFilteringOutput;
+RWTexture2D<tsr_ushort2> AntiAliasingOutput;
 
 
 //------------------------------------------------------- FUNCTIONS
@@ -206,8 +204,23 @@ void MainCS(
 	
 	//Debug[0] = 1.0;
 
-	tsr_half2x2 TexelOffset = 0.0;
-	tsr_half2 NoiseFiltering = 0.0;
+	bool2 bAntiAlias;
+	bAntiAlias[0] = AntiAliasMaskTexture[dpv_lo(InputPixelPos)] > tsr_ushort(0);
+	bAntiAlias[1] = AntiAliasMaskTexture[dpv_hi(InputPixelPos)] > tsr_ushort(0);
+	
+	// Whether it is worth to anti-aliasing.
+	bool bSkipAntiAliaser = all(!bAntiAlias);
+
+	tsr_half2x2 TexelOffset;
+	tsr_half2 NoiseFiltering;
+
+	BRANCH
+	if (bSkipAntiAliaser)
+	{
+		TexelOffset = tsr_half(0.0);
+		NoiseFiltering = tsr_half(0.0);
+	}
+	else
 	{
 		#if CONFIG_ENABLE_DUAL_PIXEL_VECTORIZATION && 1
 			tsr_half InputCache[3 * 4];
@@ -269,8 +282,12 @@ void MainCS(
 
 		#if CONFIG_COMPUTE_NOISE_FILTERING
 		{
-			NoiseFiltering = max(NoiseFiltering, tsr_half(2.0) * TotalVarH - max(DiffE, DiffW));
-			NoiseFiltering = max(NoiseFiltering, tsr_half(2.0) * TotalVarV - max(DiffN, DiffS));
+			NoiseFiltering = saturate(tsr_half(2.0) * TotalVarH - max(DiffE, DiffW));
+			NoiseFiltering = max(NoiseFiltering, saturate(tsr_half(2.0) * TotalVarV - max(DiffN, DiffS)));
+		}
+		#else
+		{
+			NoiseFiltering = tsr_half(0.0);
 		}
 		#endif
 
@@ -304,9 +321,9 @@ void MainCS(
 
 		tsr_half2 LumaDelta = abs(EdgeLuma - InputC) * 0.5;
 		
-		bool2 WorthBrowsing = LumaDelta > CONFIG_MIN_LUMA_ALIASING;
+		bool2 WorthBrowsing = LumaDelta > SPATIAL_ANTI_ALIASER_MIN_LUMIMANCE;
 
-		tsr_half2 BrowseDelta = saturate(LumaDelta * rcp(tsr_half(CONFIG_MIN_LUMA_ALIASING)));
+		tsr_half2 BrowseDelta = saturate(LumaDelta * rcp(tsr_half(SPATIAL_ANTI_ALIASER_MIN_LUMIMANCE)));
 		
 		tsr_ushort2 EdgeLengthP = 0;
 		tsr_ushort2 EdgeLengthN = 0;
@@ -500,12 +517,10 @@ void MainCS(
 		
 		OutputPixelPos = InvalidateOutputPixelPos(OutputPixelPos, InputInfo_ViewportMax);
 		
-		AntiAliasingOutput[dpv_lo(OutputPixelPos)] = dpv_lo(EncodedTexelOffset);
-		NoiseFilteringOutput[dpv_lo(OutputPixelPos)] = dpv_lo(NoiseFiltering);
+		AntiAliasingOutput[dpv_lo(OutputPixelPos)] = tsr_ushort2(dpv_lo(EncodedTexelOffset), round(dpv_lo(NoiseFiltering) * tsr_half(255.0)));
 		
 		#if CONFIG_ENABLE_DUAL_PIXEL_VECTORIZATION
-			AntiAliasingOutput[dpv_hi(OutputPixelPos)] = dpv_hi(EncodedTexelOffset);
-			NoiseFilteringOutput[dpv_hi(OutputPixelPos)] = dpv_hi(NoiseFiltering);
+			AntiAliasingOutput[dpv_lo(OutputPixelPos)] = tsr_ushort2(dpv_hi(EncodedTexelOffset), round(dpv_hi(NoiseFiltering) * tsr_half(255.0)));
 		#endif
 
 		#if DEBUG_OUTPUT
diff --git a/Engine/Shaders/Private/TemporalSuperResolution/TSRUpdateHistory.usf b/Engine/Shaders/Private/TemporalSuperResolution/TSRUpdateHistory.usf
--- a/Engine/Shaders/Private/TemporalSuperResolution/TSRUpdateHistory.usf
+++ b/Engine/Shaders/Private/TemporalSuperResolution/TSRUpdateHistory.usf
@@ -95,8 +95,7 @@ Texture2D<tsr_half4> InputSceneTranslucencyTexture;
 Texture2D<tsr_half4> HistoryRejectionTexture;
 
 Texture2D<float2> DilatedVelocityTexture;
-Texture2D<tsr_ushort> AntiAliasingTexture;
-Texture2D<tsr_half> NoiseFilteringTexture;
+Texture2D<tsr_ushort2> AntiAliasingTexture;
 
 int2   TranslucencyPixelPosMin;
 int2   TranslucencyPixelPosMax;
@@ -480,8 +479,7 @@ void MainCS(
 		tsr_half4 RawHistoryRejection[DPV_PIXEL_PER_LANE];
 		
 		#if CONFIG_REJECTION_ANTI_ALIASING
-			tsr_ushort RawEncodedInputTexelOffset[DPV_PIXEL_PER_LANE];
-			tsr_half RawNoiseFiltering[DPV_PIXEL_PER_LANE]; 
+			tsr_ushort2 RawEncodedInputTexelOffset[DPV_PIXEL_PER_LANE];
 		#endif
 
 		// Issue overlapped texture fetches
@@ -501,7 +499,6 @@ void MainCS(
 				#if CONFIG_REJECTION_ANTI_ALIASING
 				{
 					RawEncodedInputTexelOffset[PixelId] = AntiAliasingTexture[LocalInputPixelPos];
-					RawNoiseFiltering[PixelId] = NoiseFilteringTexture[LocalInputPixelPos];
 				}
 				#endif
 			}
@@ -597,9 +594,9 @@ void MainCS(
 
 		#if CONFIG_REJECTION_ANTI_ALIASING
 		{
-			tsr_ushort2 EncodedInputTexelOffset = dpv_force_interleave_registers_array(RawEncodedInputTexelOffset);
+			tsr_ushort2 EncodedInputTexelOffset = dpv_force_interleave_registers_array(RawEncodedInputTexelOffset)[0];
 			
-			NoiseFiltering = dpv_force_interleave_registers_array(RawNoiseFiltering);
+			NoiseFiltering = tsr_half2(dpv_force_interleave_registers_array(RawEncodedInputTexelOffset)[1]) * rcp(tsr_half(255.0));
 
 			SpatialAntiAliasingOffset = DecodeSpatialAntiAliasingOffset(EncodedInputTexelOffset);
 		}
diff --git a/Engine/Shaders/Private/TemporalSuperResolution/TSRVisualize.usf b/Engine/Shaders/Private/TemporalSuperResolution/TSRVisualize.usf
--- a/Engine/Shaders/Private/TemporalSuperResolution/TSRVisualize.usf
+++ b/Engine/Shaders/Private/TemporalSuperResolution/TSRVisualize.usf
@@ -22,11 +22,13 @@ static tsr_half3 kBlue   = tsr_half3(0.10, 0.10, 1.00);
 FScreenTransform OutputPixelPosToScreenPos;
 FScreenTransform ScreenPosToHistoryUV;
 FScreenTransform ScreenPosToInputPixelPos;
+FScreenTransform ScreenPosToInputUV;
 float4x4 ClipToResurrectionClip;
 uint2 OutputViewRectMin;
 uint2 OutputViewRectMax;
 uint VisualizeId;
 uint bCanResurrectHistory;
+uint bCanSpatialAntiAlias;
 float MaxHistorySampleCount;
 float OutputToHistoryResolutionFractionSquare;
 
@@ -34,6 +36,7 @@ Texture2D<tsr_half3> SceneColorTexture;
 Texture2D<float>     ClosestDepthTexture;
 Texture2D<float2>    DilatedVelocityTexture;
 Texture2D<tsr_half4> HistoryRejectionTexture;
+Texture2D<tsr_ushort> AntiAliasMaskTexture;
 Texture2D<tsr_half4> HistoryMetadataTexture;
 Texture2D<tsr_half3> ResurrectedHistoryColorTexture;
 
@@ -54,6 +57,7 @@ struct FContext
 	tsr_short2 OutputPixelPos;
 	float2 ScreenPos;
 	float2 InputUV;
+	tsr_short2 InputPixelPos;
 };
 
 tsr_half3 SampleTSROutputColor(FContext Ctx)
@@ -134,6 +138,16 @@ tsr_half3 DisplayResurrectedColor(FContext Ctx)
 	return GrayScaleColor(ResurrectedColor) * select(ResurrectionMask, kGreen, kWhite);
 }
 
+tsr_half3 DisplaySpatialAntiAliasingMask(FContext Ctx)
+{
+	bool bSpatialAntiAlias = AntiAliasMaskTexture[Ctx.InputPixelPos] != tsr_ushort(0x0);
+	if (!bCanSpatialAntiAlias)
+	{
+		return SampleTSROutputGray(Ctx) * kRed;
+	}
+	return SampleTSROutputGray(Ctx) * select(bSpatialAntiAlias, kGreen, kWhite);
+}
+
 
 //------------------------------------------------------- ENTRY POINT
 
@@ -152,8 +166,10 @@ void MainCS(
 		tsr_short2(GroupId) * tsr_short2(TILE_SIZE, TILE_SIZE) +
 		Map8x8Tile2x2Lane(GroupThreadIndex));
 	Ctx.ScreenPos = ApplyScreenTransform(float2(Ctx.OutputPixelPos), OutputPixelPosToScreenPos);
-	Ctx.InputUV = ApplyScreenTransform(Ctx.ScreenPos, ScreenPosToInputPixelPos);
+	Ctx.InputUV = ApplyScreenTransform(Ctx.ScreenPos, ScreenPosToInputUV);
 	Ctx.InputUV = fastClamp(Ctx.InputUV, InputInfo_UVViewportBilinearMin, InputInfo_UVViewportBilinearMax);
+	Ctx.InputPixelPos = tsr_short2(ApplyScreenTransform(Ctx.ScreenPos, ScreenPosToInputPixelPos));
+	Ctx.InputPixelPos = fastClamp(Ctx.InputPixelPos, tsr_short2(InputInfo_ViewportMin), tsr_short2(InputInfo_ViewportMax - 1));
 
 	tsr_half3 OutputColor = tsr_half(0.0);
 
@@ -182,6 +198,10 @@ void MainCS(
 	{
 		OutputColor = DisplayResurrectedColor(Ctx);
 	}
+	else if (VisualizeId == 6)
+	{
+		OutputColor = DisplaySpatialAntiAliasingMask(Ctx);
+	}
 
 	{
 		tsr_short2 OutputPixelPos = InvalidateOutputPixelPos(Ctx.OutputPixelPos, OutputViewRectMax);
diff --git a/Engine/Source/Runtime/Renderer/Private/PostProcess/TemporalSuperResolution.cpp b/Engine/Source/Runtime/Renderer/Private/PostProcess/TemporalSuperResolution.cpp
--- a/Engine/Source/Runtime/Renderer/Private/PostProcess/TemporalSuperResolution.cpp
+++ b/Engine/Source/Runtime/Renderer/Private/PostProcess/TemporalSuperResolution.cpp
@@ -622,6 +622,7 @@ class FTSRRejectShadingCS : public FTSRShader
 		SHADER_PARAMETER_RDG_TEXTURE_UAV(RWTexture2D, HistoryRejectionOutput)
 		SHADER_PARAMETER_RDG_TEXTURE_UAV(RWTexture2D, InputSceneColorOutput)
 		SHADER_PARAMETER_RDG_TEXTURE_UAV(RWTexture2D, InputSceneColorLdrLumaOutput)
+		SHADER_PARAMETER_RDG_TEXTURE_UAV(RWTexture2D, AntiAliasMaskOutput)
 		SHADER_PARAMETER_RDG_TEXTURE_UAV(RWTexture2D, DebugOutput)
 	END_SHADER_PARAMETER_STRUCT()
 
@@ -725,9 +726,9 @@ class FTSRSpatialAntiAliasingCS : public FTSRShader
 
 	BEGIN_SHADER_PARAMETER_STRUCT(FParameters, )
 		SHADER_PARAMETER_STRUCT_INCLUDE(FTSRCommonParameters, CommonParameters)
+		SHADER_PARAMETER_RDG_TEXTURE(Texture2D, AntiAliasMaskTexture)
 		SHADER_PARAMETER_RDG_TEXTURE(Texture2D, InputSceneColorLdrLumaTexture)
 		SHADER_PARAMETER_RDG_TEXTURE_UAV(RWTexture2D, AntiAliasingOutput)
-		SHADER_PARAMETER_RDG_TEXTURE_UAV(RWTexture2D, NoiseFilteringOutput)
 		SHADER_PARAMETER_RDG_TEXTURE_UAV(RWTexture2DArray, DebugOutput)
 	END_SHADER_PARAMETER_STRUCT()
 
@@ -751,19 +752,6 @@ class FTSRSpatialAntiAliasingCS : public FTSRShader
 	}
 }; // class FTSRSpatialAntiAliasingCS
 
-class FTSRFilterAntiAliasingCS : public FTSRShader
-{
-	DECLARE_GLOBAL_SHADER(FTSRFilterAntiAliasingCS);
-	SHADER_USE_PARAMETER_STRUCT(FTSRFilterAntiAliasingCS, FTSRShader);
-
-	BEGIN_SHADER_PARAMETER_STRUCT(FParameters, )
-		SHADER_PARAMETER_STRUCT_INCLUDE(FTSRCommonParameters, CommonParameters)
-		SHADER_PARAMETER_RDG_TEXTURE(Texture2D, AntiAliasingTexture)
-		SHADER_PARAMETER_RDG_TEXTURE_UAV(RWTexture2D, AntiAliasingOutput)
-		SHADER_PARAMETER_RDG_TEXTURE_UAV(RWTexture2DArray, DebugOutput)
-	END_SHADER_PARAMETER_STRUCT()
-}; // class FTSRFilterAntiAliasingCS
-
 class FTSRUpdateHistoryCS : public FTSRShader
 {
 	DECLARE_GLOBAL_SHADER(FTSRUpdateHistoryCS);
@@ -790,7 +778,6 @@ class FTSRUpdateHistoryCS : public FTSRShader
 		SHADER_PARAMETER_RDG_TEXTURE(Texture2D, HistoryRejectionTexture)
 		SHADER_PARAMETER_RDG_TEXTURE(Texture2D, DilatedVelocityTexture)
 		SHADER_PARAMETER_RDG_TEXTURE(Texture2D, AntiAliasingTexture)
-		SHADER_PARAMETER_RDG_TEXTURE(Texture2D, NoiseFilteringTexture)
 
 		SHADER_PARAMETER_STRUCT(FScreenPassTextureViewportParameters, TranslucencyInfo)
 		SHADER_PARAMETER(FIntPoint, TranslucencyPixelPosMin)
@@ -939,11 +926,13 @@ class FTSRVisualizeCS : public FTSRShader
 		SHADER_PARAMETER(FScreenTransform, OutputPixelPosToScreenPos)
 		SHADER_PARAMETER(FScreenTransform, ScreenPosToHistoryUV)
 		SHADER_PARAMETER(FScreenTransform, ScreenPosToInputPixelPos)
+		SHADER_PARAMETER(FScreenTransform, ScreenPosToInputUV)
 		SHADER_PARAMETER(FMatrix44f, ClipToResurrectionClip)
 		SHADER_PARAMETER(FIntPoint, OutputViewRectMin)
 		SHADER_PARAMETER(FIntPoint, OutputViewRectMax)
 		SHADER_PARAMETER(int32, VisualizeId)
 		SHADER_PARAMETER(int32, bCanResurrectHistory)
+		SHADER_PARAMETER(int32, bCanSpatialAntiAlias)
 		SHADER_PARAMETER(float, MaxHistorySampleCount)
 		SHADER_PARAMETER(float, OutputToHistoryResolutionFractionSquare)
 
@@ -951,6 +940,7 @@ class FTSRVisualizeCS : public FTSRShader
 		SHADER_PARAMETER_RDG_TEXTURE(Texture2D, ClosestDepthTexture)
 		SHADER_PARAMETER_RDG_TEXTURE(Texture2D, DilatedVelocityTexture)
 		SHADER_PARAMETER_RDG_TEXTURE(Texture2D, HistoryRejectionTexture)
+		SHADER_PARAMETER_RDG_TEXTURE(Texture2D, AntiAliasMaskTexture)
 		SHADER_PARAMETER_RDG_TEXTURE_SRV(Texture2D, HistoryMetadataTexture)
 		SHADER_PARAMETER_RDG_TEXTURE_SRV(Texture2D, ResurrectedHistoryColorTexture)
 
@@ -967,7 +957,6 @@ IMPLEMENT_GLOBAL_SHADER(FTSRDecimateHistoryCS,       "/Engine/Private/TemporalSu
 IMPLEMENT_GLOBAL_SHADER(FTSRRejectShadingCS,         "/Engine/Private/TemporalSuperResolution/TSRRejectShading.usf",         "MainCS", SF_Compute);
 IMPLEMENT_GLOBAL_SHADER(FTSRMergeRejectionCS,        "/Engine/Private/TemporalSuperResolution/TSRMergeRejection.usf",        "MainCS", SF_Compute);
 IMPLEMENT_GLOBAL_SHADER(FTSRSpatialAntiAliasingCS,   "/Engine/Private/TemporalSuperResolution/TSRSpatialAntiAliasing.usf",   "MainCS", SF_Compute);
-IMPLEMENT_GLOBAL_SHADER(FTSRFilterAntiAliasingCS,    "/Engine/Private/TemporalSuperResolution/TSRFilterAntiAliasing.usf",    "MainCS", SF_Compute);
 IMPLEMENT_GLOBAL_SHADER(FTSRUpdateHistoryCS,         "/Engine/Private/TemporalSuperResolution/TSRUpdateHistory.usf",         "MainCS", SF_Compute);
 IMPLEMENT_GLOBAL_SHADER(FTSRResolveHistoryCS,        "/Engine/Private/TemporalSuperResolution/TSRResolveHistory.usf",        "MainCS", SF_Compute);
 IMPLEMENT_GLOBAL_SHADER(FTSRVisualizeCS,             "/Engine/Private/TemporalSuperResolution/TSRVisualize.usf",             "MainCS", SF_Compute);
@@ -1922,6 +1911,7 @@ FDefaultTemporalUpscaler::FOutputs AddTemporalSuperResolutionPasses(
 	// Perform a history reject the history.
 	FRDGTextureRef HistoryRejectionTexture = nullptr;
 	FRDGTextureRef InputSceneColorLdrLumaTexture = nullptr;
+	FRDGTextureRef AntiAliasMaskTexture = nullptr;
 	{
 		const bool bComputeInputSceneColorTexture = InputSceneColorLdrLumaTexture == nullptr;
 		if (bComputeInputSceneColorTexture)
@@ -1957,6 +1947,17 @@ FDefaultTemporalUpscaler::FOutputs AddTemporalSuperResolutionPasses(
 			HistoryRejectionTexture = GraphBuilder.CreateTexture(Desc, TEXT("TSR.HistoryRejection"));
 		}
 
+		if (bComputeLdrLuma)
+		{
+			FRDGTextureDesc Desc = FRDGTextureDesc::Create2D(
+				InputExtent,
+				PF_R8_UINT,
+				FClearValueBinding::None,
+				/* InFlags = */ TexCreate_ShaderResource | TexCreate_UAV);
+
+			AntiAliasMaskTexture = GraphBuilder.CreateTexture(Desc, TEXT("TSR.AntiAliasing.Mask"));
+		}
+
 		FScreenPassTextureViewport TranslucencyViewport(
 			SeparateTranslucencyTexture->Desc.Extent, SeparateTranslucencyRect);
 
@@ -2054,6 +2055,9 @@ FDefaultTemporalUpscaler::FOutputs AddTemporalSuperResolutionPasses(
 			PassParameters->InputSceneColorLdrLumaOutput = bComputeLdrLuma
 				? GraphBuilder.CreateUAV(InputSceneColorLdrLumaTexture)
 				: CreateDummyUAV(GraphBuilder, PF_R8);
+			PassParameters->AntiAliasMaskOutput = bComputeLdrLuma
+				? GraphBuilder.CreateUAV(AntiAliasMaskTexture)
+				: CreateDummyUAV(GraphBuilder, PF_R8_UINT);
 
 			PassParameters->DebugOutput = CreateDebugUAV(InputExtent, TEXT("Debug.TSR.RejectShading"));
 		}
@@ -2100,63 +2104,38 @@ FDefaultTemporalUpscaler::FOutputs AddTemporalSuperResolutionPasses(
 
 	// Spatial anti-aliasing when doing history rejection.
 	FRDGTextureRef AntiAliasingTexture = nullptr;
-	FRDGTextureRef NoiseFilteringTexture = nullptr;
 	if (RejectionAntiAliasingQuality > 0)
 	{
-		FRDGTextureRef RawAntiAliasingTexture;
 		{
 			FRDGTextureDesc Desc = FRDGTextureDesc::Create2D(
 				InputExtent,
-				PF_R8_UINT,
+				PF_R8G8_UINT,
 				FClearValueBinding::None,
 				/* InFlags = */ TexCreate_ShaderResource | TexCreate_UAV);
 
-			RawAntiAliasingTexture = GraphBuilder.CreateTexture(Desc, TEXT("TSR.AntiAliasing.Raw"));
-			AntiAliasingTexture = GraphBuilder.CreateTexture(Desc, TEXT("TSR.AntiAliasing.Filtered"));
-
-			Desc.Format = PF_R8;
-			NoiseFilteringTexture = GraphBuilder.CreateTexture(Desc, TEXT("TSR.AntiAliasing.Noise"));
+			AntiAliasingTexture = GraphBuilder.CreateTexture(Desc, TEXT("TSR.AntiAliasing"));
 		}
 
-		{
-			FTSRSpatialAntiAliasingCS::FParameters* PassParameters = GraphBuilder.AllocParameters<FTSRSpatialAntiAliasingCS::FParameters>();
-			PassParameters->CommonParameters = CommonParameters;
-			PassParameters->InputSceneColorLdrLumaTexture = InputSceneColorLdrLumaTexture;
-			PassParameters->AntiAliasingOutput = GraphBuilder.CreateUAV(RawAntiAliasingTexture);
-			PassParameters->NoiseFilteringOutput = GraphBuilder.CreateUAV(NoiseFilteringTexture);
-			PassParameters->DebugOutput = CreateDebugUAV(InputExtent, TEXT("Debug.TSR.SpatialAntiAliasing"));
-
-			FTSRSpatialAntiAliasingCS::FPermutationDomain PermutationVector;
-			PermutationVector.Set<FTSRSpatialAntiAliasingCS::FQualityDim>(RejectionAntiAliasingQuality);
-
-			TShaderMapRef<FTSRSpatialAntiAliasingCS> ComputeShader(View.ShaderMap, PermutationVector);
-			FComputeShaderUtils::AddPass(
-				GraphBuilder,
-				RDG_EVENT_NAME("TSR SpatialAntiAliasing(Quality=%d) %dx%d",
-					RejectionAntiAliasingQuality,
-					InputRect.Width(), InputRect.Height()),
-				AsyncComputePasses >= 3 ? ERDGPassFlags::AsyncCompute : ERDGPassFlags::Compute,
-				ComputeShader,
-				PassParameters,
-				FComputeShaderUtils::GetGroupCount(InputRect.Size(), 8));
-		}
+		FTSRSpatialAntiAliasingCS::FParameters* PassParameters = GraphBuilder.AllocParameters<FTSRSpatialAntiAliasingCS::FParameters>();
+		PassParameters->CommonParameters = CommonParameters;
+		PassParameters->InputSceneColorLdrLumaTexture = InputSceneColorLdrLumaTexture;
+		PassParameters->AntiAliasMaskTexture = AntiAliasMaskTexture;
+		PassParameters->AntiAliasingOutput = GraphBuilder.CreateUAV(AntiAliasingTexture);
+		PassParameters->DebugOutput = CreateDebugUAV(InputExtent, TEXT("Debug.TSR.SpatialAntiAliasing"));
 
-		{
-			FTSRFilterAntiAliasingCS::FParameters* PassParameters = GraphBuilder.AllocParameters<FTSRFilterAntiAliasingCS::FParameters>();
-			PassParameters->CommonParameters = CommonParameters;
-			PassParameters->AntiAliasingTexture = RawAntiAliasingTexture;
-			PassParameters->AntiAliasingOutput = GraphBuilder.CreateUAV(AntiAliasingTexture);
-			PassParameters->DebugOutput = CreateDebugUAV(InputExtent, TEXT("Debug.TSR.FilterAntiAliasing"));
+		FTSRSpatialAntiAliasingCS::FPermutationDomain PermutationVector;
+		PermutationVector.Set<FTSRSpatialAntiAliasingCS::FQualityDim>(RejectionAntiAliasingQuality);
 
-			TShaderMapRef<FTSRFilterAntiAliasingCS> ComputeShader(View.ShaderMap);
-			FComputeShaderUtils::AddPass(
-				GraphBuilder,
-				RDG_EVENT_NAME("TSR FilterAntiAliasing %dx%d", InputRect.Width(), InputRect.Height()),
-				AsyncComputePasses >= 3 ? ERDGPassFlags::AsyncCompute : ERDGPassFlags::Compute,
-				ComputeShader,
-				PassParameters,
-				FComputeShaderUtils::GetGroupCount(InputRect.Size(), 8));
-		}
+		TShaderMapRef<FTSRSpatialAntiAliasingCS> ComputeShader(View.ShaderMap, PermutationVector);
+		FComputeShaderUtils::AddPass(
+			GraphBuilder,
+			RDG_EVENT_NAME("TSR SpatialAntiAliasing(Quality=%d) %dx%d",
+				RejectionAntiAliasingQuality,
+				InputRect.Width(), InputRect.Height()),
+			AsyncComputePasses >= 3 ? ERDGPassFlags::AsyncCompute : ERDGPassFlags::Compute,
+			ComputeShader,
+			PassParameters,
+			FComputeShaderUtils::GetGroupCount(InputRect.Size(), 8));
 	}
 
 	// Update temporal history.
@@ -2178,7 +2157,6 @@ FDefaultTemporalUpscaler::FOutputs AddTemporalSuperResolutionPasses(
 
 		PassParameters->DilatedVelocityTexture = DilatedVelocityTexture;
 		PassParameters->AntiAliasingTexture = AntiAliasingTexture;
-		PassParameters->NoiseFilteringTexture = NoiseFilteringTexture;
 
 		PassParameters->TranslucencyPixelPosMin = SeparateTranslucencyRect.Min;
 		PassParameters->TranslucencyPixelPosMax = SeparateTranslucencyRect.Max - 1;
@@ -2450,6 +2428,7 @@ FDefaultTemporalUpscaler::FOutputs AddTemporalSuperResolutionPasses(
 			HistoryClamp = 3,
 			ResurrectionMask = 4,
 			ResurrectedColor = 5,
+			SpatialAntiAliasingMask = 6,
 			MAX,
 		};
 
@@ -2460,6 +2439,7 @@ FDefaultTemporalUpscaler::FOutputs AddTemporalSuperResolutionPasses(
 			TEXT("HistoryClamp"),
 			TEXT("ResurrectionMask"),
 			TEXT("ResurrectedColor"),
+			TEXT("SpatialAntiAliasingMask"),
 		};
 		static_assert(UE_ARRAY_COUNT(kVisualizationName) == int32(EVisualizeId::MAX), "kVisualizationName doesn't match EVisualizeId");
 
@@ -2483,12 +2463,14 @@ FDefaultTemporalUpscaler::FOutputs AddTemporalSuperResolutionPasses(
 			PassParameters->PrevHistoryParameters = PrevHistoryParameters;
 			PassParameters->OutputPixelPosToScreenPos = (FScreenTransform::Identity - OutputRect.Min + 0.5f) / OutputRect.Size() * FScreenTransform::ViewportUVToScreenPos;
 			PassParameters->ScreenPosToHistoryUV = FScreenTransform::ChangeTextureBasisFromTo(HistoryExtent, FIntRect(FIntPoint::ZeroValue, HistorySize), FScreenTransform::ETextureBasis::ScreenPosition, FScreenTransform::ETextureBasis::TextureUV);
-			PassParameters->ScreenPosToInputPixelPos = FScreenTransform::ChangeTextureBasisFromTo(InputExtent, InputRect, FScreenTransform::ETextureBasis::ScreenPosition, FScreenTransform::ETextureBasis::TextureUV);
+			PassParameters->ScreenPosToInputPixelPos = FScreenTransform::ChangeTextureBasisFromTo(InputExtent, InputRect, FScreenTransform::ETextureBasis::ScreenPosition, FScreenTransform::ETextureBasis::TexelPosition);
+			PassParameters->ScreenPosToInputUV = FScreenTransform::ChangeTextureBasisFromTo(InputExtent, InputRect, FScreenTransform::ETextureBasis::ScreenPosition, FScreenTransform::ETextureBasis::TextureUV);
 			PassParameters->ClipToResurrectionClip = ClipToResurrectionClip;
 			PassParameters->OutputViewRectMin = VisualizeRect.Min;
 			PassParameters->OutputViewRectMax = VisualizeRect.Max;
 			PassParameters->VisualizeId = int32(VisualizeId);
 			PassParameters->bCanResurrectHistory = bCanResurrectHistory;
+			PassParameters->bCanSpatialAntiAlias = RejectionAntiAliasingQuality > 0;
 			PassParameters->MaxHistorySampleCount = MaxHistorySampleCount;
 			PassParameters->OutputToHistoryResolutionFractionSquare = OutputToHistoryResolutionFractionSquare;
 
@@ -2496,6 +2478,7 @@ FDefaultTemporalUpscaler::FOutputs AddTemporalSuperResolutionPasses(
 			PassParameters->ClosestDepthTexture = ClosestDepthTexture;
 			PassParameters->DilatedVelocityTexture = DilatedVelocityTexture;
 			PassParameters->HistoryRejectionTexture = HistoryRejectionTexture;
+			PassParameters->AntiAliasMaskTexture = AntiAliasMaskTexture ? AntiAliasMaskTexture : BlackUintDummy;
 			PassParameters->HistoryMetadataTexture = GraphBuilder.CreateSRV(FRDGTextureSRVDesc::CreateForSlice(History.MetadataArray, CurrentFrameSliceIndex));
 			if (PrevHistory.ColorArray == BlackArrayDummy)
 			{
@@ -2539,6 +2522,7 @@ FDefaultTemporalUpscaler::FOutputs AddTemporalSuperResolutionPasses(
 				{
 					Tiles[4 * 2 + 0] = Visualize(EVisualizeId::ResurrectedColor, TEXT("Resurrected Frame"));
 				}
+				Tiles[4 * 3 + 0] = Visualize(EVisualizeId::SpatialAntiAliasingMask, TEXT("Spatial Anti-Aliasing"));
 			}
 
 			{
