Fix inconsistent lighting evaluation regarding hair transmittance.

* Change hair transmittance to be an actual permutation activated only when hair plugin is used. This does not increase the actual permutation count, as the two variant are disjoint
* Remove a aniso permutation when hair lighting is used

diff --git a/Engine/Source/Runtime/Renderer/Private/LightRendering.cpp b/Engine/Source/Runtime/Renderer/Private/LightRendering.cpp
--- a/Engine/Source/Runtime/Renderer/Private/LightRendering.cpp
+++ b/Engine/Source/Runtime/Renderer/Private/LightRendering.cpp
@@ -727,6 +727,7 @@ class FDeferredLightPS : public FGlobalShader
 	class FLightingChannelsDim	: SHADER_PERMUTATION_BOOL("USE_LIGHTING_CHANNELS");
 	class FTransmissionDim		: SHADER_PERMUTATION_BOOL("USE_TRANSMISSION");
 	class FHairLighting			: SHADER_PERMUTATION_INT("USE_HAIR_LIGHTING", 2);
+	class FHairComplexTransmittance: SHADER_PERMUTATION_BOOL("USE_HAIR_COMPLEX_TRANSMITTANCE");
 	class FAtmosphereTransmittance : SHADER_PERMUTATION_BOOL("USE_ATMOSPHERE_TRANSMITTANCE");
 	class FCloudTransmittance 	: SHADER_PERMUTATION_BOOL("USE_CLOUD_TRANSMITTANCE");
 	class FAnistropicMaterials 	: SHADER_PERMUTATION_BOOL("SUPPORTS_ANISOTROPIC_MATERIALS");
@@ -741,6 +742,7 @@ class FDeferredLightPS : public FGlobalShader
 		FLightingChannelsDim,
 		FTransmissionDim,
 		FHairLighting,
+		FHairComplexTransmittance,
 		FAtmosphereTransmittance,
 		FCloudTransmittance,
 		FAnistropicMaterials,
@@ -783,6 +785,7 @@ class FDeferredLightPS : public FGlobalShader
 			PermutationVector.Get< FIESProfileDim >() ||
 			PermutationVector.Get< FTransmissionDim >() ||
 			PermutationVector.Get< FHairLighting >() ||
+			PermutationVector.Get< FHairComplexTransmittance >() ||
 			PermutationVector.Get< FAtmosphereTransmittance >() ||
 			PermutationVector.Get< FCloudTransmittance >() ||
 			PermutationVector.Get< FAnistropicMaterials >() ||
@@ -830,8 +833,8 @@ class FDeferredLightPS : public FGlobalShader
 				return false;
 			}
 
-			// (Hair Lighting == 2) has its own BxDF and anisotropic BRDF is only for DefaultLit and ClearCoat materials.
-			if (PermutationVector.Get<FHairLighting>() == 2)
+			// (Hair Lighting == 1) has its own BxDF and anisotropic BRDF is only for DefaultLit and ClearCoat materials.
+			if (PermutationVector.Get<FHairLighting>() == 1)
 			{
 				return false;
 			}
@@ -842,6 +845,17 @@ class FDeferredLightPS : public FGlobalShader
 			}
 		}
 
+		// Only compile one version or the other depending if hair strands can be used
+		const bool bNeedComplexTransmittanceSupport = IsHairStrandsSupported(EHairStrandsShaderType::All, Parameters.Platform);
+		if (PermutationVector.Get<FDeferredLightPS::FHairComplexTransmittance>() && !bNeedComplexTransmittanceSupport)
+		{
+			return false;
+		}
+		else if (!PermutationVector.Get<FDeferredLightPS::FHairComplexTransmittance>() && bNeedComplexTransmittanceSupport)
+		{
+			return false;
+		}
+
 		if (!DoesPlatformSupportVirtualShadowMaps(Parameters.Platform) && PermutationVector.Get<FVirtualShadowMapMask>() != 0)
 		{
 			return false;
@@ -877,8 +891,6 @@ class FDeferredLightPS : public FGlobalShader
 		FPermutationDomain PermutationVector(Parameters.PermutationId);
 
 		FGlobalShader::ModifyCompilationEnvironment(Parameters, OutEnvironment);
-		OutEnvironment.SetDefine(TEXT("USE_HAIR_COMPLEX_TRANSMITTANCE"), IsHairStrandsSupported(EHairStrandsShaderType::All, Parameters.Platform) ? 1u : 0u);
-
 		if (PermutationVector.Get< FVirtualShadowMapMask >() != 0)
 		{
 			FVirtualShadowMapArray::SetShaderDefines(OutEnvironment);
@@ -2423,6 +2435,7 @@ static void RenderLight(
 	const bool bIsRadial = LightType != LightType_Directional;
 	const bool bSupportAnisotropyPermutation = ShouldRenderAnisotropyPass(View) && !Strata::IsStrataEnabled(); // Strata managed anisotropy differently than legacy path. No need for special permutation.
 	const bool bUseVirtualShadowMapMask = VirtualShadowMapId != INDEX_NONE && ShadowMaskBits;
+	const bool bNeedComplexTransmittanceSupport = IsHairStrandsSupported(EHairStrandsShaderType::All, View.GetShaderPlatform());
 
 	check(!bUseVirtualShadowMapMask || bIsRadial);		// VSM mask only stores local lights
 
@@ -2487,6 +2500,7 @@ static void RenderLight(
 		PermutationVector.Set< FDeferredLightPS::FVisualizeCullingDim >(View.Family->EngineShowFlags.VisualizeLightCulling);
 		PermutationVector.Set< FDeferredLightPS::FVirtualShadowMapMask >(bUseVirtualShadowMapMask);
 		PermutationVector.Set< FDeferredLightPS::FStrataTileType >(0);
+		PermutationVector.Set< FDeferredLightPS::FHairComplexTransmittance >(bNeedComplexTransmittanceSupport);
 		if (bIsRadial)
 		{
 			PermutationVector.Set< FDeferredLightPS::FSourceShapeDim >(LightProxy->IsRectLight() ? ELightSourceShape::Rect : ELightSourceShape::Capsule);
@@ -2586,6 +2600,7 @@ void FDeferredShadingSceneRenderer::RenderLightForHair(
 
 	const bool bIsDirectional = LightSceneInfo->Proxy->GetLightType() == LightType_Directional;
 	const bool bCloudShadow   = bIsDirectional;
+	const bool bNeedComplexTransmittanceSupport = IsHairStrandsSupported(EHairStrandsShaderType::All, View.GetShaderPlatform());
 
 	FRenderLightForHairParameters* PassParameters = GraphBuilder.AllocParameters<FRenderLightForHairParameters>();
 	// VS - General parameters
@@ -2627,6 +2642,7 @@ void FDeferredShadingSceneRenderer::RenderLightForHair(
 	PermutationVector.Set< FDeferredLightPS::FVisualizeCullingDim >(false);
 	PermutationVector.Set< FDeferredLightPS::FTransmissionDim >(false);
 	PermutationVector.Set< FDeferredLightPS::FHairLighting>(1);
+	PermutationVector.Set< FDeferredLightPS::FHairComplexTransmittance>(bNeedComplexTransmittanceSupport);
 	if (bIsDirectional)
 	{
 		PermutationVector.Set< FDeferredLightPS::FSourceShapeDim >(ELightSourceShape::Directional);
@@ -2803,6 +2819,8 @@ static void InternalRenderSimpleLightsStandardDeferred(
 		SimpleLights.InstanceData[0], // Use a dummy light to create the PassParameter buffer. The light data will be
 		FVector(0, 0, 0));		  // update dynamically with the pass light loop for efficiency purpose
 
+	const bool bNeedComplexTransmittanceSupport = IsHairStrandsSupported(EHairStrandsShaderType::All, View.GetShaderPlatform());
+
 	FDeferredLightPS::FPermutationDomain PermutationVector;
 	PermutationVector.Set< FDeferredLightPS::FSourceShapeDim >(ELightSourceShape::Capsule);
 	PermutationVector.Set< FDeferredLightPS::FIESProfileDim >(false);
@@ -2811,6 +2829,7 @@ static void InternalRenderSimpleLightsStandardDeferred(
 	PermutationVector.Set< FDeferredLightPS::FAnistropicMaterials >(false);
 	PermutationVector.Set< FDeferredLightPS::FTransmissionDim >(false);
 	PermutationVector.Set< FDeferredLightPS::FHairLighting>(0);
+	PermutationVector.Set< FDeferredLightPS::FHairComplexTransmittance>(bNeedComplexTransmittanceSupport);
 	PermutationVector.Set< FDeferredLightPS::FAtmosphereTransmittance >(false);
 	PermutationVector.Set< FDeferredLightPS::FCloudTransmittance >(false);
 	PermutationVector.Set< FDeferredLightPS::FStrataTileType>(TileType != EStrataTileType::ECount ? TileType : 0);
