Add actual permutation for switching between simple/complex transmittance when when HairStrands plugin is enabled.

diff --git a/Engine/Source/Runtime/Renderer/Private/LightRendering.cpp b/Engine/Source/Runtime/Renderer/Private/LightRendering.cpp
--- a/Engine/Source/Runtime/Renderer/Private/LightRendering.cpp
+++ b/Engine/Source/Runtime/Renderer/Private/LightRendering.cpp
@@ -819,6 +819,18 @@ class FDeferredLightPS : public FGlobalShader
 		{
 			return false;
 		}
+		
+		// (Hair Lighting == 1) requires FHairComplexTransmittance 
+		if (PermutationVector.Get<FHairLighting>() == 1 && !PermutationVector.Get<FDeferredLightPS::FHairComplexTransmittance>())
+		{
+			return false;
+		}
+
+		const bool bNeedComplexTransmittanceSupport = IsHairStrandsSupported(EHairStrandsShaderType::All, Parameters.Platform);
+		if (PermutationVector.Get<FHairLighting>() == 0 && PermutationVector.Get<FDeferredLightPS::FHairComplexTransmittance>() && !bNeedComplexTransmittanceSupport)
+		{
+			return false;
+		}
 
 		if (PermutationVector.Get<FDeferredLightPS::FAnistropicMaterials>())
 		{
@@ -845,17 +857,6 @@ class FDeferredLightPS : public FGlobalShader
 			}
 		}
 
-		// Only compile one version or the other depending if hair strands can be used
-		const bool bNeedComplexTransmittanceSupport = IsHairStrandsSupported(EHairStrandsShaderType::All, Parameters.Platform);
-		if (PermutationVector.Get<FDeferredLightPS::FHairComplexTransmittance>() && !bNeedComplexTransmittanceSupport)
-		{
-			return false;
-		}
-		else if (!PermutationVector.Get<FDeferredLightPS::FHairComplexTransmittance>() && bNeedComplexTransmittanceSupport)
-		{
-			return false;
-		}
-
 		if (!DoesPlatformSupportVirtualShadowMaps(Parameters.Platform) && PermutationVector.Get<FVirtualShadowMapMask>() != 0)
 		{
 			return false;
@@ -2435,7 +2436,7 @@ static void RenderLight(
 	const bool bIsRadial = LightType != LightType_Directional;
 	const bool bSupportAnisotropyPermutation = ShouldRenderAnisotropyPass(View) && !Strata::IsStrataEnabled(); // Strata managed anisotropy differently than legacy path. No need for special permutation.
 	const bool bUseVirtualShadowMapMask = VirtualShadowMapId != INDEX_NONE && ShadowMaskBits;
-	const bool bNeedComplexTransmittanceSupport = IsHairStrandsSupported(EHairStrandsShaderType::All, View.GetShaderPlatform());
+	const bool bNeedComplexTransmittanceSupport = View.HairCardsMeshElements.Num() && IsHairStrandsSupported(EHairStrandsShaderType::All, View.GetShaderPlatform());
 
 	check(!bUseVirtualShadowMapMask || bIsRadial);		// VSM mask only stores local lights
 
@@ -2600,7 +2601,6 @@ void FDeferredShadingSceneRenderer::RenderLightForHair(
 
 	const bool bIsDirectional = LightSceneInfo->Proxy->GetLightType() == LightType_Directional;
 	const bool bCloudShadow   = bIsDirectional;
-	const bool bNeedComplexTransmittanceSupport = IsHairStrandsSupported(EHairStrandsShaderType::All, View.GetShaderPlatform());
 
 	FRenderLightForHairParameters* PassParameters = GraphBuilder.AllocParameters<FRenderLightForHairParameters>();
 	// VS - General parameters
@@ -2642,7 +2642,7 @@ void FDeferredShadingSceneRenderer::RenderLightForHair(
 	PermutationVector.Set< FDeferredLightPS::FVisualizeCullingDim >(false);
 	PermutationVector.Set< FDeferredLightPS::FTransmissionDim >(false);
 	PermutationVector.Set< FDeferredLightPS::FHairLighting>(1);
-	PermutationVector.Set< FDeferredLightPS::FHairComplexTransmittance>(bNeedComplexTransmittanceSupport);
+	PermutationVector.Set< FDeferredLightPS::FHairComplexTransmittance>(true);
 	if (bIsDirectional)
 	{
 		PermutationVector.Set< FDeferredLightPS::FSourceShapeDim >(ELightSourceShape::Directional);
@@ -2819,7 +2819,7 @@ static void InternalRenderSimpleLightsStandardDeferred(
 		SimpleLights.InstanceData[0], // Use a dummy light to create the PassParameter buffer. The light data will be
 		FVector(0, 0, 0));		  // update dynamically with the pass light loop for efficiency purpose
 
-	const bool bNeedComplexTransmittanceSupport = IsHairStrandsSupported(EHairStrandsShaderType::All, View.GetShaderPlatform());
+	const bool bNeedComplexTransmittanceSupport = View.HairCardsMeshElements.Num() && IsHairStrandsSupported(EHairStrandsShaderType::All, View.GetShaderPlatform());
 
 	FDeferredLightPS::FPermutationDomain PermutationVector;
 	PermutationVector.Set< FDeferredLightPS::FSourceShapeDim >(ELightSourceShape::Capsule);
