PR #10673: Fix for additive scene capture corruption on Metal

diff --git a/Engine/Source/Runtime/Renderer/Private/SceneCaptureRendering.cpp b/Engine/Source/Runtime/Renderer/Private/SceneCaptureRendering.cpp
--- a/Engine/Source/Runtime/Renderer/Private/SceneCaptureRendering.cpp
+++ b/Engine/Source/Runtime/Renderer/Private/SceneCaptureRendering.cpp
@@ -196,15 +196,18 @@ void CopySceneCaptureComponentToTarget(
 		GraphicsPSOInit.RasterizerState = TStaticRasterizerState<FM_Solid, CM_None>::GetRHI();
 		GraphicsPSOInit.DepthStencilState = TStaticDepthStencilState<false, CF_Always>::GetRHI();
 
+		bool bIsCompositing = false;
 		if (SceneCaptureSource == SCS_SceneColorHDR && ViewFamily.SceneCaptureCompositeMode == SCCM_Composite)
 		{
 			// Blend with existing render target color. Scene capture color is already pre-multiplied by alpha.
 			GraphicsPSOInit.BlendState = TStaticBlendState<CW_RGBA, BO_Add, BF_One, BF_SourceAlpha, BO_Add, BF_Zero, BF_SourceAlpha>::GetRHI();
+			bIsCompositing = true;
 		}
 		else if (SceneCaptureSource == SCS_SceneColorHDR && ViewFamily.SceneCaptureCompositeMode == SCCM_Additive)
 		{
 			// Add to existing render target color. Scene capture color is already pre-multiplied by alpha.
 			GraphicsPSOInit.BlendState = TStaticBlendState<CW_RGBA, BO_Add, BF_One, BF_One, BO_Add, BF_Zero, BF_SourceAlpha>::GetRHI();
+			bIsCompositing = true;
 		}
 		else
 		{
@@ -221,7 +224,7 @@ void CopySceneCaptureComponentToTarget(
 			FSceneCapturePS::FParameters* PassParameters = GraphBuilder.AllocParameters<FSceneCapturePS::FParameters>();
 			PassParameters->View = View.ViewUniformBuffer;
 			PassParameters->SceneTextures = SceneTextures.GetSceneTextureShaderParameters(ViewFamily.GetFeatureLevel());
-			PassParameters->RenderTargets[0] = FRenderTargetBinding(ViewFamilyTexture, ERenderTargetLoadAction::ENoAction);
+			PassParameters->RenderTargets[0] = FRenderTargetBinding(ViewFamilyTexture, bIsCompositing ? ERenderTargetLoadAction::ELoad : ERenderTargetLoadAction::ENoAction);
 
 			TShaderMapRef<FScreenVS> VertexShader(View.ShaderMap);
 			TShaderMapRef<FSceneCapturePS> PixelShader(View.ShaderMap, PixelPermutationVector);
