Fix CustomStencil buffer visualization view mode. Stencil was not being cleared when being visualized when stencil writes were not enabled, so garbage was being displayed. Stencil buffer is now explicitly cleared when being visualized, though technically the buffer will have garbage unless stencil writes are enabled.

diff --git a/Engine/Source/Runtime/Renderer/Private/CustomDepthRendering.cpp b/Engine/Source/Runtime/Renderer/Private/CustomDepthRendering.cpp
--- a/Engine/Source/Runtime/Renderer/Private/CustomDepthRendering.cpp
+++ b/Engine/Source/Runtime/Renderer/Private/CustomDepthRendering.cpp
@@ -239,7 +239,7 @@ bool FSceneRenderer::RenderCustomDepthPass(
 			}
 
 			const ERenderTargetLoadAction DepthLoadAction = GetLoadActionIfProduced(CustomDepthTextures.Depth, CustomDepthTextures.DepthAction);
-			const ERenderTargetLoadAction StencilLoadAction = GetLoadActionIfProduced(CustomDepthTextures.Depth, CustomDepthTextures.StencilAction);
+			const ERenderTargetLoadAction StencilLoadAction = (View.Family->ViewMode == VMI_VisualizeBuffer) ? ERenderTargetLoadAction::EClear : GetLoadActionIfProduced(CustomDepthTextures.Depth, CustomDepthTextures.StencilAction);
 
 			PassParameters->RenderTargets.DepthStencil = FDepthStencilBinding(
 				CustomDepthTextures.Depth,
