Revert to using scene texture extents for each view in DFAO.

Although using the entire scene texture extent is a bit wasteful in terms of splitscreen memory usage, it's more constantly sized than a view rect. This helps avoid render target re-allocation and better supports temporal upsampling during dynamic resolution. As an added bonus, we can also avoid an unnecessary texture copy for the single view case.

diff --git a/Engine/Shaders/Private/DistanceFieldAOShared.ush b/Engine/Shaders/Private/DistanceFieldAOShared.ush
--- a/Engine/Shaders/Private/DistanceFieldAOShared.ush
+++ b/Engine/Shaders/Private/DistanceFieldAOShared.ush
@@ -144,7 +144,7 @@ float3 ComputeBentNormal(float3 RecordWorldNormal, uint RelativeRecordIndex)
 	return UnoccludedDirection;
 }
 
-float2 AOBufferBilinearUVMax;
+float4 AOBufferBilinearUVMinMax;
 
 struct FUpsampleDFAOOutput
 {
@@ -155,7 +155,8 @@ FUpsampleDFAOOutput UpsampleDFAO(float2 BufferUV, float SceneDepth)
 {
 	FUpsampleDFAOOutput UpsampleDFAOOutput;
 
-	float2 DistanceFieldUVs = min(BufferUV, AOBufferBilinearUVMax);
+	// Note: if DFAO is ever stored per-view and not atlased at some point, there should only need to be a max here
+	float2 DistanceFieldUVs = clamp(BufferUV, AOBufferBilinearUVMinMax.xy, AOBufferBilinearUVMinMax.zw);
 
 #define BILATERAL_UPSAMPLE 1
 #if BILATERAL_UPSAMPLE
diff --git a/Engine/Shaders/Private/DistanceFieldLightingPost.usf b/Engine/Shaders/Private/DistanceFieldLightingPost.usf
--- a/Engine/Shaders/Private/DistanceFieldLightingPost.usf
+++ b/Engine/Shaders/Private/DistanceFieldLightingPost.usf
@@ -204,8 +204,8 @@ void UpdateHistoryDepthRejectionPS(
 #define MANUAL_HISTORY_RESAMPLE 0
 #if MANUAL_HISTORY_RESAMPLE
 
-		float2 HistoryBufferSize = floor(View.ViewSizeAndInvSize.xy / DOWNSAMPLE_FACTOR);
-		float2 HistoryTexelSize = View.ViewSizeAndInvSize.zw * DOWNSAMPLE_FACTOR;
+		float2 HistoryBufferSize = floor(View.BufferSizeAndInvSize.xy / DOWNSAMPLE_FACTOR);
+		float2 HistoryTexelSize = View.BufferSizeAndInvSize.zw * DOWNSAMPLE_FACTOR;
 		float2 OldDistanceFieldUVsCorner00 = HistoryTexelSize * (floor(OldDistanceFieldUVs * HistoryBufferSize - .5f) + .5f);
 		float2 OldDistanceFieldUVsCorner01 = OldDistanceFieldUVsCorner00 + float2(0, HistoryTexelSize.y);
 		float2 OldDistanceFieldUVsCorner11 = OldDistanceFieldUVsCorner00 + HistoryTexelSize;
@@ -302,6 +302,10 @@ void FilterHistoryPS(
 {
 	float2 BufferUV = UVAndScreenPos.xy;
 
+	#if MANUALLY_CLAMP_UV
+		BufferUV = min(BufferUV, MaxSampleBufferUV);
+	#endif
+
 	float4 HistoryValue = Texture2DSampleLevel(BentNormalAOTexture, BentNormalAOSampler, BufferUV, 0);
 
 	// Only do the spatial search for pixels who discarded their history value
@@ -317,6 +321,11 @@ void FilterHistoryPS(
 			{
 				float2 SampleBufferUV = UVAndScreenPos.xy + BentNormalAOTexelSize * float2(x, y);
 
+				#if MANUALLY_CLAMP_UV
+					// Distance field AO was computed at 0,0 regardless of viewrect min.
+					SampleBufferUV = min(SampleBufferUV, MaxSampleBufferUV);
+				#endif
+
 				float4 TextureValue = Texture2DSampleLevel(BentNormalAOTexture, BentNormalAOSampler, SampleBufferUV, 0);
 
 				float SampleSceneDepth = abs(TextureValue.w);
@@ -359,8 +368,7 @@ void AOUpsamplePS(
 	,out float4 OutSceneColor : SV_Target0
 	)
 {
-	const float2 GBufferUV = UVAndScreenPos.xy * View.BufferBilinearUVMinMax.zw;
-	FGBufferData GBuffer = GetGBufferData(GBufferUV);
+	FGBufferData GBuffer = GetGBufferData(UVAndScreenPos.xy);
 
 	float3 BentNormal = UpsampleDFAO(UVAndScreenPos.xy, GBuffer.Depth, GBuffer.WorldNormal);
 
diff --git a/Engine/Shaders/Private/DistanceFieldObjectCulling.usf b/Engine/Shaders/Private/DistanceFieldObjectCulling.usf
--- a/Engine/Shaders/Private/DistanceFieldObjectCulling.usf
+++ b/Engine/Shaders/Private/DistanceFieldObjectCulling.usf
@@ -130,7 +130,7 @@ void BuildTileConesMain(
     uint ThreadIndex = GroupThreadId.y * THREADGROUP_SIZEX + GroupThreadId.x;
 
 	// Sampling from the texture based off of the ViewRect size because the texture is created on a per-view basis
-	float2 BaseLevelScreenUV = (DispatchThreadId.xy + float2(.5f, .5f)) * DOWNSAMPLE_FACTOR * View.ViewSizeAndInvSize.zw;
+	float2 BaseLevelScreenUV = (DispatchThreadId.xy + float2(.5f, .5f)) * DOWNSAMPLE_FACTOR * View.BufferSizeAndInvSize.zw;
 	float SceneDepth = GetDownsampledDepth(BaseLevelScreenUV);
 
 	// Initialize per-tile variables
diff --git a/Engine/Source/Runtime/Renderer/Private/DistanceFieldAmbientOcclusion.cpp b/Engine/Source/Runtime/Renderer/Private/DistanceFieldAmbientOcclusion.cpp
--- a/Engine/Source/Runtime/Renderer/Private/DistanceFieldAmbientOcclusion.cpp
+++ b/Engine/Source/Runtime/Renderer/Private/DistanceFieldAmbientOcclusion.cpp
@@ -167,7 +167,7 @@ void TileIntersectionModifyCompilationEnvironment(EShaderPlatform Platform, FSha
 
 FIntPoint GetBufferSizeForAO(const FViewInfo& View)
 {
-	return FIntPoint::DivideAndRoundDown(View.ViewRect.Size(), GAODownsampleFactor);
+	return FIntPoint::DivideAndRoundDown(View.GetSceneTexturesConfig().Extent, GAODownsampleFactor);
 }
 
 // Sample set restricted to not self-intersect a surface based on cone angle .475882232
@@ -276,15 +276,18 @@ FDFAOUpsampleParameters DistanceField::SetupAOUpsampleParameters(const FViewInfo
 {
 	const float DistanceFadeScaleValue = 1.0f / ((1.0f - GAOViewFadeDistanceScale) * GetMaxAOViewDistance());
 
-	const FIntPoint DFAOBufferSize = GetBufferSizeForAO(View);
-	const FVector2f UVMax(
-		(DFAOBufferSize.X - 0.51f) / DFAOBufferSize.X, // 0.51 - so bilateral gather4 won't sample invalid texels
-		(DFAOBufferSize.Y - 0.51f) / DFAOBufferSize.Y);
+	const FIntPoint AOBufferSize = GetBufferSizeForAO(View);
+	const FIntPoint AOViewSize = View.ViewRect.Size() / GAODownsampleFactor;
+	const FVector4f UVMinMax(
+		0.51 / AOBufferSize.X,
+		0.51 / AOBufferSize.Y,
+		(AOViewSize.X - 0.51f) / AOBufferSize.X, // 0.51 - so bilateral gather4 won't sample invalid texels
+		(AOViewSize.Y - 0.51f) / AOBufferSize.Y);
 
 	FDFAOUpsampleParameters ShaderParameters;
 	ShaderParameters.BentNormalAOTexture = DistanceFieldAOBentNormal;
 	ShaderParameters.BentNormalAOSampler = TStaticSamplerState<SF_Bilinear>::GetRHI();
-	ShaderParameters.AOBufferBilinearUVMax = UVMax;
+	ShaderParameters.AOBufferBilinearUVMinMax = UVMinMax;
 	ShaderParameters.DistanceFadeScale = DistanceFadeScaleValue;
 	ShaderParameters.AOMaxViewDistance = GetMaxAOViewDistance();
 
@@ -387,12 +390,12 @@ void ComputeDistanceFieldNormal(
 	FRDGTextureRef DistanceFieldNormal,
 	const FDistanceFieldAOParameters& Parameters)
 {
-	const FIntPoint DFAOTextureSize = GetBufferSizeForAO(View);
+	const FIntPoint DFAOViewSize = View.ViewRect.Size() / GAODownsampleFactor;
 
 	if (GAOComputeShaderNormalCalculation)
 	{
-		uint32 GroupSizeX = FMath::DivideAndRoundUp(DFAOTextureSize.X, GDistanceFieldAOTileSizeX);
-		uint32 GroupSizeY = FMath::DivideAndRoundUp(DFAOTextureSize.Y, GDistanceFieldAOTileSizeY);
+		uint32 GroupSizeX = FMath::DivideAndRoundUp(DFAOViewSize.X, GDistanceFieldAOTileSizeX);
+		uint32 GroupSizeY = FMath::DivideAndRoundUp(DFAOViewSize.Y, GDistanceFieldAOTileSizeY);
 
 		auto* PassParameters = GraphBuilder.AllocParameters<FComputeDistanceFieldNormalCS::FParameters>();
 		PassParameters->View = View.ViewUniformBuffer;
@@ -418,9 +421,9 @@ void ComputeDistanceFieldNormal(
 			RDG_EVENT_NAME("ComputeNormal"),
 			PassParameters,
 			ERDGPassFlags::Raster,
-			[&View, PassParameters, DFAOTextureSize](FRHICommandList& RHICmdList)
+			[&View, PassParameters, DFAOViewSize](FRHICommandList& RHICmdList)
 		{
-			RHICmdList.SetViewport(0, 0, 0.0f, DFAOTextureSize.X, DFAOTextureSize.Y, 1.0f);
+			RHICmdList.SetViewport(0, 0, 0.0f, DFAOViewSize.X, DFAOViewSize.Y, 1.0f);
 
 			TShaderMapRef<FPostProcessVS> VertexShader(View.ShaderMap);
 
@@ -443,10 +446,10 @@ void ComputeDistanceFieldNormal(
 			DrawRectangle(
 				RHICmdList,
 				0, 0,
-				DFAOTextureSize.X, DFAOTextureSize.Y,
+				DFAOViewSize.X, DFAOViewSize.Y,
 				0, 0,
 				View.ViewRect.Width(), View.ViewRect.Height(),
-				DFAOTextureSize,
+				DFAOViewSize,
 				View.GetSceneTexturesConfig().Extent,
 				VertexShader);
 		});
@@ -821,9 +824,10 @@ void FDeferredShadingSceneRenderer::RenderDistanceFieldLighting(
 		bListMeshDistanceFieldsMemoryNextFrame = false;
 		Scene->DistanceFieldSceneData.ListMeshDistanceFields(true);
 	}
-
+	
+	// We only need this texture if we need to copy multiple view outputs to an texture atlas
 	FRDGTextureRef BentNormalOutput = nullptr;
-
+	 if (Views.Num() > 1)
 	{
 		const FIntPoint BufferSize = GetActiveSceneTexturesConfig().Extent / GAODownsampleFactor;
 		const FRDGTextureDesc Desc = FRDGTextureDesc::Create2D(BufferSize, PF_FloatRGBA, FClearValueBinding::None, GFastVRamConfig.DistanceFieldAOBentNormal | TexCreate_ShaderResource | TexCreate_RenderTargetable | TexCreate_UAV);
@@ -895,7 +899,14 @@ void FDeferredShadingSceneRenderer::RenderDistanceFieldLighting(
 			UpsampleBentNormalAO(GraphBuilder, View, SceneTextures.UniformBuffer, SceneTextures.Color.Target, PerViewBentNormal, bModulateToSceneColor && !bVisualizeAmbientOcclusion);
 		}
 
-		AddCopyTexturePass(GraphBuilder, PerViewBentNormal, BentNormalOutput, FIntPoint::ZeroValue, View.ViewRect.Min / GAODownsampleFactor, View.ViewRect.Size() / GAODownsampleFactor);
+		if (Views.Num() == 1) {
+			OutDynamicBentNormalAO = PerViewBentNormal;
+			return;
+		}
+		else 
+		{
+			AddCopyTexturePass(GraphBuilder, PerViewBentNormal, BentNormalOutput, FIntPoint::ZeroValue, View.ViewRect.Min / GAODownsampleFactor, View.ViewRect.Size() / GAODownsampleFactor);
+		}
 	}
 
 	OutDynamicBentNormalAO = BentNormalOutput;
diff --git a/Engine/Source/Runtime/Renderer/Private/DistanceFieldAmbientOcclusion.h b/Engine/Source/Runtime/Renderer/Private/DistanceFieldAmbientOcclusion.h
--- a/Engine/Source/Runtime/Renderer/Private/DistanceFieldAmbientOcclusion.h
+++ b/Engine/Source/Runtime/Renderer/Private/DistanceFieldAmbientOcclusion.h
@@ -90,7 +90,7 @@ END_SHADER_PARAMETER_STRUCT()
 BEGIN_SHADER_PARAMETER_STRUCT(FDFAOUpsampleParameters, )
 	SHADER_PARAMETER_RDG_TEXTURE(Texture2D, BentNormalAOTexture)
 	SHADER_PARAMETER_SAMPLER(SamplerState, BentNormalAOSampler)
-	SHADER_PARAMETER(FVector2f, AOBufferBilinearUVMax)
+	SHADER_PARAMETER(FVector4f, AOBufferBilinearUVMinMax)
 	SHADER_PARAMETER(float, DistanceFadeScale)
 	SHADER_PARAMETER(float, AOMaxViewDistance)
 END_SHADER_PARAMETER_STRUCT()
diff --git a/Engine/Source/Runtime/Renderer/Private/DistanceFieldLightingPost.cpp b/Engine/Source/Runtime/Renderer/Private/DistanceFieldLightingPost.cpp
--- a/Engine/Source/Runtime/Renderer/Private/DistanceFieldLightingPost.cpp
+++ b/Engine/Source/Runtime/Renderer/Private/DistanceFieldLightingPost.cpp
@@ -166,10 +166,14 @@ public:
 		SHADER_PARAMETER_RDG_TEXTURE(Texture2D, BentNormalAOTexture)
 		SHADER_PARAMETER_SAMPLER(SamplerState, BentNormalAOSampler)
 		SHADER_PARAMETER(FVector2f, BentNormalAOTexelSize)
+		SHADER_PARAMETER(FVector2f, MaxSampleBufferUV)
 		SHADER_PARAMETER(float, HistoryWeight)
 		RENDER_TARGET_BINDING_SLOTS()
 	END_SHADER_PARAMETER_STRUCT()
 
+	class FManuallyClampUV : SHADER_PERMUTATION_BOOL("MANUALLY_CLAMP_UV");
+	using FPermutationDomain = TShaderPermutationDomain<FManuallyClampUV>;
+
 	static bool ShouldCompilePermutation(const FGlobalShaderPermutationParameters& Parameters)
 	{
 		return ShouldCompileDFLightingPostShaders(Parameters.Platform);
@@ -239,9 +243,10 @@ void GeometryAwareUpsample(FRDGBuilder& GraphBuilder, const FViewInfo& View, FRD
 		ERDGPassFlags::Raster,
 		[PassParameters, VertexShader, PixelShader, &View](FRHICommandList& RHICmdList)
 	{
-		FIntPoint AOBufferSize = GetBufferSizeForAO(View);
+		const FIntPoint AOBufferSize = GetBufferSizeForAO(View);
+		const FIntPoint AOViewSize = View.ViewRect.Size() / GAODownsampleFactor;
 
-		RHICmdList.SetViewport(0, 0, 0.0f, AOBufferSize.X, AOBufferSize.Y, 1.0f);
+		RHICmdList.SetViewport(0, 0, 0.0f, AOViewSize.X, AOViewSize.Y, 1.0f);
 
 		FGraphicsPipelineStateInitializer GraphicsPSOInit;
 		RHICmdList.ApplyCachedRenderTargets(GraphicsPSOInit);
@@ -260,10 +265,10 @@ void GeometryAwareUpsample(FRDGBuilder& GraphBuilder, const FViewInfo& View, FRD
 		DrawRectangle(
 			RHICmdList,
 			0, 0,
-			AOBufferSize.X, AOBufferSize.Y,
+			AOViewSize.X, AOViewSize.Y,
 			0, 0,
-			AOBufferSize.X, AOBufferSize.Y,
-			AOBufferSize,
+			AOViewSize.X, AOViewSize.Y,
+			AOViewSize,
 			AOBufferSize,
 			VertexShader);
 	});
@@ -317,6 +322,7 @@ void UpdateHistory(
 	const FRDGSystemTextures& SystemTextures = FRDGSystemTextures::Get(GraphBuilder);
 
 	const FIntPoint AOBufferSize = GetBufferSizeForAO(View);
+	const FIntPoint AOViewSize = View.ViewRect.Size() / GAODownsampleFactor;
 
 	if (BentNormalHistoryState && DistanceFieldAOUseHistory(View))
 	{
@@ -383,10 +389,10 @@ void UpdateHistory(
 					RDG_EVENT_NAME("UpdateHistory"),
 					PassParameters,
 					ERDGPassFlags::Raster,
-					[PassParameters, VertexShader, PixelShader, &View, AOBufferSize]
+					[PassParameters, VertexShader, PixelShader, &View, AOBufferSize, AOViewSize]
 					(FRHICommandList& RHICmdList)
 				{
-					RHICmdList.SetViewport(0, 0, 0.0f, AOBufferSize.X, AOBufferSize.Y, 1.0f);
+					RHICmdList.SetViewport(0, 0, 0.0f, AOViewSize.X, AOViewSize.Y, 1.0f);
 
 					FGraphicsPipelineStateInitializer GraphicsPSOInit;
 					RHICmdList.ApplyCachedRenderTargets(GraphicsPSOInit);
@@ -405,10 +411,10 @@ void UpdateHistory(
 					DrawRectangle(
 						RHICmdList,
 						0, 0,
-						AOBufferSize.X, AOBufferSize.Y,
+						AOViewSize.X, AOViewSize.Y,
 						0, 0,
-						AOBufferSize.X, AOBufferSize.Y,
-						AOBufferSize,
+						AOViewSize.X, AOViewSize.Y,
+						AOViewSize,
 						AOBufferSize,
 						VertexShader);
 				});
@@ -427,8 +433,16 @@ void UpdateHistory(
 					AllocateOrReuseAORenderTarget(GraphBuilder, View, BentNormalHistoryTexture, BentNormalHistoryRTName, PF_FloatRGBA);
 				}
 
+				const bool bManuallyClampUV = View.ViewRect.Min != FIntPoint::ZeroValue || View.ViewRect.Max != AOBufferSize;
+				FFilterHistoryPS::FPermutationDomain PermutationVector;
+				PermutationVector.Set<FFilterHistoryPS::FManuallyClampUV>(bManuallyClampUV);
+
 				auto VertexShader = View.ShaderMap->GetShader<FPostProcessVS>();
-				auto PixelShader = View.ShaderMap->GetShader<FFilterHistoryPS>();
+				auto PixelShader = View.ShaderMap->GetShader<FFilterHistoryPS>(PermutationVector);
+
+				FVector2f MaxSampleBufferUV(
+					(AOViewSize.X - 0.5f - GAODownsampleFactor) / AOBufferSize.X,
+					(AOViewSize.Y - 0.5f - GAODownsampleFactor) / AOBufferSize.Y);
 
 				auto* PassParameters = GraphBuilder.AllocParameters<FFilterHistoryPS::FParameters>();
 				PassParameters->View = View.ViewUniformBuffer;
@@ -436,6 +450,7 @@ void UpdateHistory(
 				PassParameters->BentNormalAOTexture = NewBentNormalHistory;
 				PassParameters->BentNormalAOSampler = TStaticSamplerState<SF_Bilinear, AM_Clamp, AM_Clamp, AM_Clamp>::GetRHI();
 				PassParameters->BentNormalAOTexelSize = FVector2f(1.0f / AOBufferSize.X, 1.0f / AOBufferSize.Y);
+				PassParameters->MaxSampleBufferUV = MaxSampleBufferUV;
 				PassParameters->HistoryWeight = GAOHistoryWeight;
 				PassParameters->RenderTargets[0] = FRenderTargetBinding(BentNormalHistoryTexture, ERenderTargetLoadAction::ELoad);
 
@@ -445,9 +460,10 @@ void UpdateHistory(
 					RDG_EVENT_NAME("UpdateHistoryStability"),
 					PassParameters,
 					ERDGPassFlags::Raster,
-					[PassParameters, VertexShader, PixelShader, &View, AOBufferSize](FRHICommandList& RHICmdList)
+					[PassParameters, VertexShader, PixelShader, &View, AOBufferSize, AOViewSize]
+					(FRHICommandList& RHICmdList)
 				{
-					RHICmdList.SetViewport(0, 0, 0.0f, AOBufferSize.X, AOBufferSize.Y, 1.0f);
+					RHICmdList.SetViewport(0, 0, 0.0f, AOViewSize.X, AOViewSize.Y, 1.0f);
 
 					FGraphicsPipelineStateInitializer GraphicsPSOInit;
 					RHICmdList.ApplyCachedRenderTargets(GraphicsPSOInit);
@@ -466,10 +482,10 @@ void UpdateHistory(
 					DrawRectangle(
 						RHICmdList,
 						0, 0,
-						AOBufferSize.X, AOBufferSize.Y,
+						AOViewSize.X, AOViewSize.Y,
 						0, 0,
-						AOBufferSize.X, AOBufferSize.Y,
-						AOBufferSize ,
+						AOViewSize.X, AOViewSize.Y,
+						AOViewSize ,
 						AOBufferSize,
 						VertexShader);
 				});
diff --git a/Engine/Source/Runtime/Renderer/Private/DistanceFieldObjectCulling.cpp b/Engine/Source/Runtime/Renderer/Private/DistanceFieldObjectCulling.cpp
--- a/Engine/Source/Runtime/Renderer/Private/DistanceFieldObjectCulling.cpp
+++ b/Engine/Source/Runtime/Renderer/Private/DistanceFieldObjectCulling.cpp
@@ -337,10 +337,10 @@ void ScatterTilesToObjects(
 
 FIntPoint GetTileListGroupSizeForView(const FViewInfo& View)
 {
-	const FIntPoint AOBufferSize = GetBufferSizeForAO(View);
+	const FIntPoint AOViewSize = View.ViewRect.Size() / GAODownsampleFactor;
 	return FIntPoint(
-		FMath::DivideAndRoundUp(FMath::Max(AOBufferSize.X, 1), GDistanceFieldAOTileSizeX),
-		FMath::DivideAndRoundUp(FMath::Max(AOBufferSize.Y, 1), GDistanceFieldAOTileSizeY));
+		FMath::DivideAndRoundUp(FMath::Max(AOViewSize.X, 1), GDistanceFieldAOTileSizeX),
+		FMath::DivideAndRoundUp(FMath::Max(AOViewSize.Y, 1), GDistanceFieldAOTileSizeY));
 }
 
 void BuildTileObjectLists(
diff --git a/Engine/Source/Runtime/Renderer/Private/IndirectLightRendering.cpp b/Engine/Source/Runtime/Renderer/Private/IndirectLightRendering.cpp
--- a/Engine/Source/Runtime/Renderer/Private/IndirectLightRendering.cpp
+++ b/Engine/Source/Runtime/Renderer/Private/IndirectLightRendering.cpp
@@ -374,7 +374,7 @@ class FReflectionEnvironmentSkyLightingPS : public FGlobalShader
 	BEGIN_SHADER_PARAMETER_STRUCT(FParameters, )
 		// Distance field AO parameters.
 		// TODO. FDFAOUpsampleParameters
-		SHADER_PARAMETER(FVector2f, AOBufferBilinearUVMax)
+		SHADER_PARAMETER(FVector4f, AOBufferBilinearUVMinMax)
 		SHADER_PARAMETER(float, DistanceFadeScale)
 		SHADER_PARAMETER(float, AOMaxViewDistance)
 
@@ -1689,10 +1689,15 @@ static void AddSkyReflectionPass(
 
 		// Setups all shader parameters related to distance field AO
 		{
-			FIntPoint AOBufferSize = GetBufferSizeForAO(View);
-			PassParameters->PS.AOBufferBilinearUVMax = FVector2f(
-				(View.ViewRect.Width() / GAODownsampleFactor - 0.51f) / AOBufferSize.X, // 0.51 - so bilateral gather4 won't sample invalid texels
-				(View.ViewRect.Height() / GAODownsampleFactor - 0.51f) / AOBufferSize.Y);
+			// The view must sample into the texture atlas correctly with the appropriate clamping
+			const FIntPoint AOBufferSize = GetBufferSizeForAO(View);
+			const FIntRect AOViewRect = View.ViewRect / GAODownsampleFactor;
+
+			PassParameters->PS.AOBufferBilinearUVMinMax = FVector4f(
+				(AOViewRect.Min.X + 0.51f) / AOBufferSize.X, // 0.51 - so bilateral gather4 won't sample invalid texels
+				(AOViewRect.Min.Y + 0.51f) / AOBufferSize.Y,
+				(AOViewRect.Max.X - 0.51f) / AOBufferSize.X,
+				(AOViewRect.Max.Y - 0.51f) / AOBufferSize.Y);
 
 			extern float GAOViewFadeDistanceScale;
 			PassParameters->PS.AOMaxViewDistance = GetMaxAOViewDistance();
