Add support for injecting raytraced directional shadows into volumetric fog:
* Rendered in a separate pass into a separate volume texture that is then sampled in LightScatteringCS pass.
* Controlled by r.VolumetricFog.InjectRaytracedLights (and disabled by default)

diff --git a/Engine/Shaders/Private/VolumetricFog.usf b/Engine/Shaders/Private/VolumetricFog.usf
--- a/Engine/Shaders/Private/VolumetricFog.usf
+++ b/Engine/Shaders/Private/VolumetricFog.usf
@@ -284,16 +284,16 @@ uint UseConservativeDepthTexture;
 
 RaytracingAccelerationStructure TLAS;
 
-float ComputeRayTracedShadowFactor(float3 TranslatedWorldPosition, float3 ToLight)
+float ComputeRayTracedShadowFactor(float3 TranslatedWorldPosition, float3 LightDirection, float TMax)
 {
 	uint RayFlags = RAY_FLAG_SKIP_CLOSEST_HIT_SHADER | RAY_FLAG_ACCEPT_FIRST_HIT_AND_END_SEARCH;
 	uint RaytracingMask = RAY_TRACING_MASK_SHADOW | RAY_TRACING_MASK_THIN_SHADOW;
 
 	FRayDesc Ray;
 	Ray.Origin = TranslatedWorldPosition;
-	Ray.Direction = normalize(ToLight);
+	Ray.Direction = LightDirection;
 	Ray.TMin = 0.0f;
-	Ray.TMax = length(ToLight);
+	Ray.TMax = TMax;
 
 	FMinimalPayload MinimalPayload = TraceVisibilityRay(
 		TLAS,
@@ -383,7 +383,7 @@ float4 InjectShadowedLocalLightCommon(uint3 GridCoordinate)
 			if (CombinedAttenuation > 0)
 			{
 #if USE_RAYTRACED_SHADOWS
-				ShadowFactor = ComputeRayTracedShadowFactor(TranslatedWorldPosition, ToLight);
+				ShadowFactor = ComputeRayTracedShadowFactor(TranslatedWorldPosition, normalize(ToLight), length(ToLight));
 #else
 				bool bUnused = false;
 				ShadowFactor = ComputeVolumeShadowing(TranslatedWorldPosition, IsPointLight, LightData.bSpotLight, bUnused);
@@ -443,6 +443,21 @@ RAY_TRACING_ENTRY_RAYGEN(InjectShadowedLocalLightRGS)
 	OutVolumeTexture[Position] = DestValue + OutScattering;
 }
 
+RWTexture3D<float> OutShadowVolumeTexture;
+
+RAY_TRACING_ENTRY_RAYGEN(InjectShadowedDirectionalLightRGS)
+{
+	int3 Position = GetVoxelIndex(DispatchRaysIndex().x, VolumetricFog.ResourceGridSizeInt);
+
+	const int SampleIndex = 0;
+	uint3 Rand32Bits = Rand4DPCG32(int4(Position.xyz, View.StateFrameIndexMod8 + 8 * SampleIndex)).xyz;
+	float3 Rand3D = (float3(Rand32Bits) / float(uint(0xffffffff))) * 2.0f - 1.0f;
+	float3 CellOffset = FrameJitterOffsets[SampleIndex].xyz + LightScatteringSampleJitterMultiplier * Rand3D;
+
+	const float3 TranslatedWorldPosition = ComputeCellTranslatedWorldPosition(Position, CellOffset);
+	OutShadowVolumeTexture[Position] = ComputeRayTracedShadowFactor(TranslatedWorldPosition, ForwardLightData.DirectionalLightDirection, 1.0e27);
+}
+
 #endif // USE_RAYTRACED_SHADOWS
 
 void InjectShadowedLocalLightPS(
@@ -584,6 +599,8 @@ float2 UseHeightFogColors;	// x=override directional light using height fog insc
 float UseDirectionalLightShadowing;
 float StaticLightingScatteringIntensity;
 
+Texture3D<float> RaytracedShadowsVolume;
+
 [numthreads(THREADGROUP_SIZE_X, THREADGROUP_SIZE_Y, THREADGROUP_SIZE_Z)]
 void LightScatteringCS(
 	uint3 GroupId : SV_GroupID,
@@ -668,6 +685,10 @@ void LightScatteringCS(
 				ShadowFactor *= VirtualShadowMapSample.ShadowFactor;
 #endif
 
+#if USE_RAYTRACED_SHADOWS_VOLUME
+				ShadowFactor *= RaytracedShadowsVolume[GridCoordinate];
+#endif
+
 #if USE_CLOUD_TRANSMITTANCE
 				float OutOpticalDepth = 0.0f; 
 				ShadowFactor *= lerp(1.0f, GetCloudVolumetricShadow(TranslatedWorldPosition, CloudShadowmapTranslatedWorldToLightClipMatrix, CloudShadowmapFarDepthKm, CloudShadowmapTexture, CloudShadowmapSampler, OutOpticalDepth), CloudShadowmapStrength);
diff --git a/Engine/Source/Runtime/Renderer/Private/VolumetricFog.cpp b/Engine/Source/Runtime/Renderer/Private/VolumetricFog.cpp
--- a/Engine/Source/Runtime/Renderer/Private/VolumetricFog.cpp
+++ b/Engine/Source/Runtime/Renderer/Private/VolumetricFog.cpp
@@ -415,6 +415,115 @@ class FInjectShadowedLocalLightRGS : public FGlobalShader
 
 IMPLEMENT_GLOBAL_SHADER(FInjectShadowedLocalLightRGS, "/Engine/Private/VolumetricFog.usf", "InjectShadowedLocalLightRGS", SF_RayGen);
 
+class FRayTraceDirectionalLightVolumeShadowMapRGS : public FGlobalShader
+{
+	DECLARE_GLOBAL_SHADER(FRayTraceDirectionalLightVolumeShadowMapRGS);
+	SHADER_USE_ROOT_PARAMETER_STRUCT(FRayTraceDirectionalLightVolumeShadowMapRGS, FGlobalShader);
+
+	BEGIN_SHADER_PARAMETER_STRUCT(FParameters, )
+		SHADER_PARAMETER_STRUCT_REF(FViewUniformShaderParameters, View)
+		SHADER_PARAMETER_RDG_UNIFORM_BUFFER(FForwardLightData, Forward)
+		SHADER_PARAMETER_STRUCT_INCLUDE(FVolumetricFogIntegrationParameters, VolumetricFogParameters)
+		SHADER_PARAMETER_RDG_TEXTURE_UAV(RWTexture3D, OutShadowVolumeTexture)
+		SHADER_PARAMETER_RDG_BUFFER_SRV(RaytracingAccelerationStructure, TLAS)
+		SHADER_PARAMETER(float, LightScatteringSampleJitterMultiplier)
+	END_SHADER_PARAMETER_STRUCT()
+
+	static bool ShouldCompilePermutation(const FGlobalShaderPermutationParameters& Parameters)
+	{
+		return ShouldCompileRayTracingShadersForProject(Parameters.Platform);
+	}
+
+	static ERayTracingPayloadType GetRayTracingPayloadType(const int32 PermutationId)
+	{
+		return ERayTracingPayloadType::RayTracingMaterial;
+	}
+
+	static void ModifyCompilationEnvironment(const FGlobalShaderPermutationParameters& Parameters, FShaderCompilerEnvironment& OutEnvironment)
+	{
+		FGlobalShader::ModifyCompilationEnvironment(Parameters, OutEnvironment);
+		OutEnvironment.SetDefine(TEXT("USE_RAYTRACED_SHADOWS"), TEXT("1"));
+	}
+};
+
+IMPLEMENT_GLOBAL_SHADER(FRayTraceDirectionalLightVolumeShadowMapRGS, "/Engine/Private/VolumetricFog.usf", "InjectShadowedDirectionalLightRGS", SF_RayGen);
+
+bool LightHasRayTracedShadows(const FLightSceneInfo* LightSceneInfo);
+
+static void RenderRaytracedDirectionalShadowVolume(
+	FRDGBuilder& GraphBuilder,
+	FViewInfo& View,
+	const FScene& Scene,
+	const FVolumetricFogIntegrationParameterData& IntegrationData,
+	FRDGTextureRef& OutRaytracedShadowsVolume)
+{	
+	const bool bUseRaytracedShadows = IsRayTracingEnabled(Scene.GetShaderPlatform())
+		&& GRHISupportsRayTracing
+		&& GRHISupportsRayTracingShaders
+		&& GVolumetricFogInjectRaytracedLights;
+
+	if (!bUseRaytracedShadows)
+	{
+		return;
+	}
+
+	// Following how RenderLightFunctionForVolumetricFog is selecting the main directional light, even though we could support all of them.
+	const FLightSceneProxy* SelectedForwardDirectionalLightProxy = View.ForwardLightingResources.SelectedForwardDirectionalLightProxy;
+
+	const FLightSceneInfo* DirectionalLightSceneInfo = nullptr;
+	for (const FLightSceneInfo* LightSceneInfo : Scene.DirectionalLights)
+	{
+		if (LightSceneInfo->ShouldRenderLightViewIndependent()
+			&& LightSceneInfo->ShouldRenderLight(View, true)
+			&& LightHasRayTracedShadows(LightSceneInfo)
+			&& LightSceneInfo->Proxy == SelectedForwardDirectionalLightProxy)
+		{
+			DirectionalLightSceneInfo = LightSceneInfo;
+			break;
+		}
+	}
+
+	if (DirectionalLightSceneInfo)
+	{
+		int32 VolumetricFogGridPixelSize;
+		const FIntVector VolumetricFogResourceGridSize = GetVolumetricFogResourceGridSize(View, VolumetricFogGridPixelSize);
+		FRDGTextureDesc RaytracedShadowsVolumeDesc(FRDGTextureDesc::Create3D(
+			VolumetricFogResourceGridSize,
+			PF_R16F,
+			FClearValueBinding::Black,
+			TexCreate_ShaderResource | TexCreate_UAV | TexCreate_ReduceMemoryWithTilingMode | TexCreate_3DTiling));
+
+		OutRaytracedShadowsVolume = GraphBuilder.CreateTexture(RaytracedShadowsVolumeDesc, TEXT("VolumetricFog.RaytracedShadowVolume"));
+
+		FRayTraceDirectionalLightVolumeShadowMapRGS::FParameters* PassParameters = GraphBuilder.AllocParameters<FRayTraceDirectionalLightVolumeShadowMapRGS::FParameters>();
+		PassParameters->OutShadowVolumeTexture = GraphBuilder.CreateUAV(OutRaytracedShadowsVolume);
+		PassParameters->TLAS = View.GetRayTracingSceneLayerViewChecked(ERayTracingSceneLayer::Base);
+		PassParameters->View = View.ViewUniformBuffer;
+		PassParameters->Forward = View.ForwardLightingResources.ForwardLightUniformBuffer;
+		PassParameters->LightScatteringSampleJitterMultiplier = GVolumetricFogJitter ? GLightScatteringSampleJitterMultiplier : 0;
+		SetupVolumetricFogIntegrationParameters(PassParameters->VolumetricFogParameters, View, IntegrationData);
+
+		TShaderRef<FRayTraceDirectionalLightVolumeShadowMapRGS> RayGenerationShader = View.ShaderMap->GetShader<FRayTraceDirectionalLightVolumeShadowMapRGS>();
+		ClearUnusedGraphResources(RayGenerationShader, PassParameters);
+
+		const uint32 DispatchSize = VolumetricFogResourceGridSize.X * VolumetricFogResourceGridSize.Y * VolumetricFogResourceGridSize.Z;
+		GraphBuilder.AddPass(
+			RDG_EVENT_NAME("RayTracedShadowedDirectionalLight"),
+			PassParameters,
+			ERDGPassFlags::Compute,
+			[&View, RayGenerationShader, PassParameters, DispatchSize](FRHIRayTracingCommandList& RHICmdList)
+			{
+				FRayTracingShaderBindingsWriter GlobalResources;
+				SetShaderParameters(GlobalResources, RayGenerationShader, *PassParameters);
+
+				FRHIRayTracingScene* RayTracingSceneRHI = View.GetRayTracingSceneChecked();
+
+				RHICmdList.RayTraceDispatch(View.RayTracingMaterialPipeline, RayGenerationShader.GetRayTracingShader(), RayTracingSceneRHI, GlobalResources, DispatchSize, 1);
+			}
+		);
+	}
+}
+
 void FDeferredShadingSceneRenderer::PrepareRayTracingVolumetricFogShadows(const FViewInfo& View, const FScene& Scene, TArray<FRHIRayTracingShader*>& OutRayGenShaders)
 {
 	const bool bEnabled = Scene.bHasRayTracedLights && ::ShouldRenderVolumetricFog(&Scene, *View.Family) && GVolumetricFogInjectRaytracedLights;
@@ -435,6 +544,11 @@ void FDeferredShadingSceneRenderer::PrepareRayTracingVolumetricFogShadows(const
 			OutRayGenShaders.Add(RayGenerationShader.GetRayTracingShader());
 		}
 	}
+
+	{
+		TShaderMapRef<FRayTraceDirectionalLightVolumeShadowMapRGS> RayGenerationShader(View.ShaderMap);
+		OutRayGenShaders.Add(RayGenerationShader.GetRayTracingShader());
+	}	
 }
 
 #endif // RHI_RAYTRACING
@@ -831,6 +945,7 @@ class FVolumetricFogLightScatteringCS : public FGlobalShader
 	class FLumenGI						: SHADER_PERMUTATION_BOOL("LUMEN_GI");
 	class FVirtualShadowMap				: SHADER_PERMUTATION_BOOL("VIRTUAL_SHADOW_MAP");
 	class FCloudTransmittance			: SHADER_PERMUTATION_BOOL("USE_CLOUD_TRANSMITTANCE");
+	class FRaytracedShadowsVolume		: SHADER_PERMUTATION_BOOL("USE_RAYTRACED_SHADOWS_VOLUME");
 	
 	using FPermutationDomain = TShaderPermutationDomain<
 		FSuperSampleCount,
@@ -838,7 +953,8 @@ class FVolumetricFogLightScatteringCS : public FGlobalShader
 		FDistanceFieldSkyOcclusion,
 		FLumenGI,
 		FVirtualShadowMap,
-		FCloudTransmittance>;
+		FCloudTransmittance,
+		FRaytracedShadowsVolume>;
 
 	BEGIN_SHADER_PARAMETER_STRUCT(FParameters, )
 		SHADER_PARAMETER_STRUCT_REF(FViewUniformShaderParameters, View)
@@ -856,6 +972,7 @@ class FVolumetricFogLightScatteringCS : public FGlobalShader
 		SHADER_PARAMETER_RDG_TEXTURE(Texture2D, ConservativeDepthTexture)
 		SHADER_PARAMETER_RDG_TEXTURE(Texture2D, PrevConservativeDepthTexture)
 		SHADER_PARAMETER_RDG_TEXTURE(Texture3D, LightScatteringHistory)
+		SHADER_PARAMETER_RDG_TEXTURE_SRV(Texture3D, RaytracedShadowsVolume)
 		SHADER_PARAMETER_SAMPLER(SamplerState, LightScatteringHistorySampler)
 		SHADER_PARAMETER_RDG_UNIFORM_BUFFER(FLumenTranslucencyLightingUniforms, LumenGIVolumeStruct)
 		SHADER_PARAMETER_STRUCT_INCLUDE(FVirtualShadowMapSamplingParameters, VirtualShadowMapSamplingParameters)
@@ -917,6 +1034,12 @@ class FVolumetricFogLightScatteringCS : public FGlobalShader
 		{
 			PermutationVector.Set<FLumenGI>(false);
 		}
+
+		if (!ShouldCompileRayTracingShadersForProject(ShaderPlatform))
+		{
+			PermutationVector.Set<FRaytracedShadowsVolume>(false);
+		}
+
 		return PermutationVector;
 	}
 
@@ -1313,6 +1436,11 @@ void FSceneRenderer::ComputeVolumetricFog(FRDGBuilder& GraphBuilder,
 		FRDGTexture* LocalShadowedLightScattering = GraphBuilder.RegisterExternalTexture(GSystemTextures.VolumetricBlackDummy);
 		RenderLocalLightsForVolumetricFog(GraphBuilder, View, bUseTemporalReprojection, IntegrationData, FogInfo,
 			VolumetricFogViewGridSize, GridZParams, VolumeDescFastVRAM, LocalShadowedLightScattering, ConservativeDepthTexture);
+	
+		FRDGTextureRef RaytracedShadowsVolume = nullptr;
+#if RHI_RAYTRACING		
+		RenderRaytracedDirectionalShadowVolume(GraphBuilder, View, *Scene, IntegrationData, RaytracedShadowsVolume);
+#endif
 
 		{
 			FVolumetricFogMaterialSetupCS::FParameters* PassParameters = GraphBuilder.AllocParameters<FVolumetricFogMaterialSetupCS::FParameters>();
@@ -1475,8 +1603,11 @@ void FSceneRenderer::ComputeVolumetricFog(FRDGBuilder& GraphBuilder,
 			PassParameters->CloudShadowmapStrength = CloudShadowmap_Strength;
 			PassParameters->CloudShadowmapTranslatedWorldToLightClipMatrix = CloudWorldToLightClipShadowMatrix;
 
+			PassParameters->RaytracedShadowsVolume = RaytracedShadowsVolume ? GraphBuilder.CreateSRV(RaytracedShadowsVolume) : nullptr;
+
 			const bool bUseLumenGI = View.GetLumenTranslucencyGIVolume().Texture0 != nullptr && FDataDrivenShaderPlatformInfo::GetSupportsLumenGI(View.GetShaderPlatform());
 			const bool bUseGlobalDistanceField = UseGlobalDistanceField() && Scene->DistanceFieldSceneData.NumObjectsInBuffer > 0;
+			const bool bUseRaytracedShadowsVolume = RaytracedShadowsVolume != nullptr;
 
 			const bool bUseDistanceFieldSkyOcclusion =
 				ViewFamily.EngineShowFlags.AmbientOcclusion
@@ -1500,6 +1631,7 @@ void FSceneRenderer::ComputeVolumetricFog(FRDGBuilder& GraphBuilder,
 			PermutationVector.Set< FVolumetricFogLightScatteringCS::FLumenGI >(bUseLumenGI);
 			PermutationVector.Set< FVolumetricFogLightScatteringCS::FVirtualShadowMap >(VirtualShadowMapArray.IsAllocated() );
 			PermutationVector.Set< FVolumetricFogLightScatteringCS::FCloudTransmittance >(AtmosphericDirectionalLightIndex >= 0);
+			PermutationVector.Set< FVolumetricFogLightScatteringCS::FRaytracedShadowsVolume >(bUseRaytracedShadowsVolume);
 
 			auto ComputeShader = View.ShaderMap->GetShader< FVolumetricFogLightScatteringCS >(PermutationVector);
 			ClearUnusedGraphResources(ComputeShader, PassParameters);
