Add budget settings to MeshDrawCommandStats.
Clarify triangle/primitive terminology.
Change dump command to r.MeshDrawCommand.DumpStats to match other cvar formatting

diff --git a/Engine/Source/Runtime/Engine/Public/MeshDrawCommandStatsSettings.h b/Engine/Source/Runtime/Engine/Public/MeshDrawCommandStatsSettings.h
new file mode 100644
--- /dev/null
+++ b/Engine/Source/Runtime/Engine/Public/MeshDrawCommandStatsSettings.h
@@ -0,0 +1,37 @@
+// Copyright Epic Games, Inc. All Rights Reserved.
+#pragma once
+
+#include "Containers/Array.h"
+#include "Engine/DeveloperSettings.h"
+#include "UObject/NameTypes.h"
+
+#include "MeshDrawCommandStatsSettings.generated.h"
+
+/** Description of a stat category used in the MeshDrawCommandStats system. */
+USTRUCT()
+struct FMeshDrawCommandStatsBudget
+{
+	GENERATED_BODY()
+
+	/** Category name. */
+	UPROPERTY(config, EditAnywhere, Category = Engine)
+	FName CategoryName;
+	/** The category primitive budget. This is the maximum triangles expected, post-culling, summed across all passes. */
+	UPROPERTY(config, EditAnywhere, Category = Engine)
+	int32 PrimitiveBudget = 0;
+};
+
+/** User settings used by the MeshDrawCommandStats system. */
+UCLASS(Config=Engine, defaultconfig, meta = (DisplayName = "Mesh Stats"))
+class ENGINE_API UMeshDrawCommandStatsSettings : public UDeveloperSettings
+{
+	GENERATED_BODY()
+
+public:
+	/** Budgets used by r.MeshDrawCommands.Stats */
+	UPROPERTY(config, EditAnywhere, Category = Engine)
+	TArray<FMeshDrawCommandStatsBudget> Budgets;
+	/** The total primitive budget. This is the maximimum triangles expected, post-culling, summed across all passes. */
+	UPROPERTY(config, EditAnywhere, Category = Engine)
+	int32 TotalPrimitiveBudget = 0;
+};
diff --git a/Engine/Source/Runtime/Renderer/Private/MeshDrawCommandStats.cpp b/Engine/Source/Runtime/Renderer/Private/MeshDrawCommandStats.cpp
--- a/Engine/Source/Runtime/Renderer/Private/MeshDrawCommandStats.cpp
+++ b/Engine/Source/Runtime/Renderer/Private/MeshDrawCommandStats.cpp
@@ -1,13 +1,14 @@
 // Copyright Epic Games, Inc. All Rights Reserved.
 
 #include "MeshDrawCommandStats.h"
+
 #include "InstanceCulling/InstanceCullingContext.h"
+#include "MeshDrawCommandStatsSettings.h"
 #include "RenderGraph.h"
+#include "RendererModule.h"
 #include "RendererOnScreenNotification.h"
-
 #include "RHI.h"
 #include "RHIGPUReadback.h"
-#include "RendererModule.h"
 
 #if MESH_DRAW_COMMAND_STAT_COLLECTION
 
@@ -21,11 +22,11 @@ void FMeshDrawCommandStatsManager::CreateInstance()
 
 DECLARE_STATS_GROUP(TEXT("MeshDrawCommandStats"), STATGROUP_Culling, STATCAT_Advanced);
 
-DECLARE_DWORD_COUNTER_STAT(TEXT("Total Rendered Triangles"), STAT_Culling_TotalNumTriangles, STATGROUP_Culling);
+DECLARE_DWORD_COUNTER_STAT(TEXT("Total Rendered Primitives"), STAT_Culling_TotalNumPrimitives, STATGROUP_Culling);
 DECLARE_DWORD_COUNTER_STAT(TEXT("Total Rendered Instances"), STAT_Culling_TotalNumInstances, STATGROUP_Culling);
-DECLARE_DWORD_COUNTER_STAT(TEXT("InstanceCulling Indirect Rendered Triangles"), STAT_Culling_InstanceCullingIndirectNumTriangles, STATGROUP_Culling);
+DECLARE_DWORD_COUNTER_STAT(TEXT("InstanceCulling Indirect Rendered Primitives"), STAT_Culling_InstanceCullingIndirectNumPrimitives, STATGROUP_Culling);
 DECLARE_DWORD_COUNTER_STAT(TEXT("InstanceCulling Indirect Rendered Instances"), STAT_Culling_InstanceCullingIndirectNumInstances, STATGROUP_Culling);
-DECLARE_DWORD_COUNTER_STAT(TEXT("Custom Indirect Rendered Triangles"), STAT_Culling_CustomIndirectNumTriangles, STATGROUP_Culling);
+DECLARE_DWORD_COUNTER_STAT(TEXT("Custom Indirect Rendered Primitives"), STAT_Culling_CustomIndirectNumPrimitives, STATGROUP_Culling);
 DECLARE_DWORD_COUNTER_STAT(TEXT("Custom Indirect Rendered Instances"), STAT_Culling_CustomIndirectNumInstances, STATGROUP_Culling);
 
 static TAutoConsoleVariable<int32> CVarShowMeshDrawCommandStats(
@@ -136,12 +137,43 @@ FMeshDrawCommandStatsManager::FMeshDrawCommandStatsManager()
 			const bool bShowStats = CVarShowMeshDrawCommandStats->GetInt() == 0 ? false : true;
 			if (bShowStats)
 			{
-				OutMessages.Add(FCoreDelegates::EOnScreenMessageSeverity::Info, FText::FromString(FString::Printf(TEXT("ResourceType Triangle Count:"), Stats.TotalTriangles / 1000)));
-				for (auto Iter = Stats.CategoryStats.CreateConstIterator(); Iter; ++Iter)
+				const UMeshDrawCommandStatsSettings* Settings = GetDefault<UMeshDrawCommandStatsSettings>();
+
+				OutMessages.Add(FCoreDelegates::EOnScreenMessageSeverity::Info, FText::FromString(FString::Printf(TEXT("MeshDrawCommandStats (Triangles / Budget - Category):"), Stats.TotalPrimitives / 1000)));
+				for (FStats::FCategoryStats const& CategoryStat : Stats.CategoryStats)
+				{
+					int32 Budget = 0;
+					if (Settings != nullptr)
+					for (FMeshDrawCommandStatsBudget const& CategoryBudget : Settings->Budgets)
+					{
+						if (CategoryBudget.CategoryName == CategoryStat.Category)
+						{
+							Budget = CategoryBudget.PrimitiveBudget;
+							break;
+						}
+					}
+					
+					if (Budget > 0)
+					{
+						FCoreDelegates::EOnScreenMessageSeverity Severity = Budget < CategoryStat.PrimitiveCount ? FCoreDelegates::EOnScreenMessageSeverity::Warning : FCoreDelegates::EOnScreenMessageSeverity::Info;
+						OutMessages.Add(Severity, FText::FromString(FString::Printf(TEXT("%5dK / %5dK - %s"), CategoryStat.PrimitiveCount / 1000, Budget / 1000, *(CategoryStat.Category.ToString()))));
+					}
+					else
+					{
+						OutMessages.Add(FCoreDelegates::EOnScreenMessageSeverity::Info, FText::FromString(FString::Printf(TEXT("\t%5dK - %s"), CategoryStat.PrimitiveCount / 1000, *(CategoryStat.Category.ToString()))));
+					}
+				}
+
+				const int32 TotalBudget = Settings != 0 ? Settings->TotalPrimitiveBudget : 0;
+				if (TotalBudget > 0)
+				{
+					FCoreDelegates::EOnScreenMessageSeverity Severity = TotalBudget < Stats.TotalPrimitives ? FCoreDelegates::EOnScreenMessageSeverity::Warning : FCoreDelegates::EOnScreenMessageSeverity::Info;
+					OutMessages.Add(Severity, FText::FromString(FString::Printf(TEXT("%5dK / %5dK - TOTAL"), Stats.TotalPrimitives / 1000, TotalBudget / 1000)));
+				}
+				else
 				{
-					OutMessages.Add(FCoreDelegates::EOnScreenMessageSeverity::Info, FText::FromString(FString::Printf(TEXT("\t%5dK - %s"), Iter->PrimitiveCount / 1000, *(Iter->Category.ToString()))));
+					OutMessages.Add(FCoreDelegates::EOnScreenMessageSeverity::Info, FText::FromString(FString::Printf(TEXT("\t%5dK - TOTAL"), Stats.TotalPrimitives / 1000)));
 				}
-				OutMessages.Add(FCoreDelegates::EOnScreenMessageSeverity::Info, FText::FromString(FString::Printf(TEXT("\t%5dK - TOTAL"), Stats.TotalTriangles / 1000)));
 			}
 		});
 }
@@ -266,7 +298,7 @@ void FMeshDrawCommandStatsManager::Update()
 								DrawData.TotalInstanceCount = FMath::Max(DrawData.TotalInstanceCount, DrawData.VisibleInstanceCount);
 
 								Stats.CustomIndirectInstances += DrawData.VisibleInstanceCount;
-								Stats.CustomIndirectTriangles += DrawData.VisibleInstanceCount * DrawData.PrimitiveCount;
+								Stats.CustomIndirectPrimitives += DrawData.VisibleInstanceCount * DrawData.PrimitiveCount;
 							}
 						}
 						else if (DrawData.UseInstantCullingIndirectBuffer > 0 && InstanceCullingReadBackData)
@@ -276,11 +308,11 @@ void FMeshDrawCommandStatsManager::Update()
 							DrawData.VisibleInstanceCount = IndirectArgs.InstanceCount;
 							check(DrawData.VisibleInstanceCount <= DrawData.TotalInstanceCount);
 							Stats.InstanceCullingIndirectInstances += DrawData.VisibleInstanceCount;
-							Stats.InstanceCullingIndirectTriangles += DrawData.VisibleInstanceCount * DrawData.PrimitiveCount;
+							Stats.InstanceCullingIndirectPrimitives += DrawData.VisibleInstanceCount * DrawData.PrimitiveCount;
 						}
 
 						Stats.TotalInstances += DrawData.VisibleInstanceCount;
-						Stats.TotalTriangles += DrawData.VisibleInstanceCount * DrawData.PrimitiveCount;
+						Stats.TotalPrimitives += DrawData.VisibleInstanceCount * DrawData.PrimitiveCount;
 
 						FMeshDrawCommandStatsComponentData ComponentData = ComponentDataManager->GetComponentData(DrawData.StatsData.ComponentDataID);
 						static FName NAME_Unknown("Unknown Category");
@@ -329,11 +361,11 @@ void FMeshDrawCommandStatsManager::Update()
 
 	// We keep and set the value from the previous frame in case there are no readback, this avoids alternating values if for example two queries were consumed in one frame
 	// might be able to do this better perhaps.
-	SET_DWORD_STAT(STAT_Culling_TotalNumTriangles, Stats.TotalTriangles);
+	SET_DWORD_STAT(STAT_Culling_TotalNumPrimitives, Stats.TotalPrimitives);
 	SET_DWORD_STAT(STAT_Culling_TotalNumInstances, Stats.TotalInstances);
-	SET_DWORD_STAT(STAT_Culling_InstanceCullingIndirectNumTriangles, Stats.InstanceCullingIndirectTriangles);
+	SET_DWORD_STAT(STAT_Culling_InstanceCullingIndirectNumPrimitives, Stats.InstanceCullingIndirectPrimitives);
 	SET_DWORD_STAT(STAT_Culling_InstanceCullingIndirectNumInstances, Stats.InstanceCullingIndirectInstances);
-	SET_DWORD_STAT(STAT_Culling_CustomIndirectNumTriangles, Stats.CustomIndirectTriangles);
+	SET_DWORD_STAT(STAT_Culling_CustomIndirectNumPrimitives, Stats.CustomIndirectPrimitives);
 	SET_DWORD_STAT(STAT_Culling_CustomIndirectNumInstances, Stats.CustomIndirectInstances);
 
 	// Collect stats during the next frame (check if STATGROUP_Culling is also visible somehow)
@@ -365,7 +397,7 @@ void FMeshDrawCommandStatsManager::DumpStats(FFrameData* FrameData)
 		FString MaterialName;
 		int32 PrimitiveCount;
 		int32 TotalInstanceCount;
-		int32 TotalTriangleCount;
+		int32 TotalPrimitiveCount;
 	};
 	TArray<FStatEntry> StatEntries;
 
@@ -384,7 +416,7 @@ void FMeshDrawCommandStatsManager::DumpStats(FFrameData* FrameData)
 				StatEntry.SegmentIndex = DrawData.StatsData.SegmentIndex;
 				StatEntry.PrimitiveCount = DrawData.PrimitiveCount;
 				StatEntry.TotalInstanceCount = DrawData.TotalInstanceCount;
-				StatEntry.TotalTriangleCount = DrawData.TotalInstanceCount * DrawData.PrimitiveCount;
+				StatEntry.TotalPrimitiveCount = DrawData.TotalInstanceCount * DrawData.PrimitiveCount;
 				StatEntry.ResourceName = DrawData.ResourceName;
 				StatEntry.MaterialName = DrawData.MaterialName;
 
@@ -407,7 +439,7 @@ void FMeshDrawCommandStatsManager::DumpStats(FFrameData* FrameData)
 			return LHS.VisibilePrimitiveCount > RHS.VisibilePrimitiveCount;
 		});
 
-	const TCHAR* Header = TEXT("Pass,VisiblePrimitiveCount,VisibleInstances,Category,ComponentType,ResourceName,LODIndex,SegmentIndex,MaterialName,PrimitiveCount,TotalInstanceCount,TotalTriangleCount\n");
+	const TCHAR* Header = TEXT("Pass,VisiblePrimitiveCount,VisibleInstances,Category,ComponentType,ResourceName,LODIndex,SegmentIndex,MaterialName,PrimitiveCount,TotalInstanceCount,TotalPrimitiveCount\n");
 	CSVFile->Serialize(TCHAR_TO_ANSI(Header), FPlatformString::Strlen(Header));
 
 	TCHAR PassNameBuffer[FName::StringBufferSize];
@@ -434,7 +466,7 @@ void FMeshDrawCommandStatsManager::DumpStats(FFrameData* FrameData)
 			*StatEntry.MaterialName,
 			StatEntry.PrimitiveCount,
 			StatEntry.TotalInstanceCount,
-			StatEntry.TotalTriangleCount);
+			StatEntry.TotalPrimitiveCount);
 		CSVFile->Serialize(TCHAR_TO_ANSI(*Row), Row.Len());
 	}
 
@@ -442,10 +474,10 @@ void FMeshDrawCommandStatsManager::DumpStats(FFrameData* FrameData)
 	CSVFile = nullptr;
 }
 
-static FAutoConsoleCommandWithWorldArgsAndOutputDevice GDumpMeshDrawCommandFrameStatsCmd(
-	TEXT("DumpMeshDrawCommandFrameStats"),
-	TEXT("Dumps the draw stat of the MeshDrawCommand passes of the last rendered frame\n"),
-	FConsoleCommandWithWorldArgsAndOutputDeviceDelegate::CreateStatic([](const TArray<FString>& Args, UWorld*, FOutputDevice& OutputDevice)
+static FAutoConsoleCommand GDumpMeshDrawCommandStatsCmd(
+	TEXT("r.MeshDrawCommands.DumpStats"),
+	TEXT("Dumps all of the Mesh Draw Command stats for a single frame to a csv file in the saved profile directory.\n"),
+	FConsoleCommandDelegate::CreateStatic([]()
 {
 	if (FMeshDrawCommandStatsManager* Instance = FMeshDrawCommandStatsManager::Get())
 	{
diff --git a/Engine/Source/Runtime/Renderer/Private/MeshDrawCommandStats.h b/Engine/Source/Runtime/Renderer/Private/MeshDrawCommandStats.h
--- a/Engine/Source/Runtime/Renderer/Private/MeshDrawCommandStats.h
+++ b/Engine/Source/Runtime/Renderer/Private/MeshDrawCommandStats.h
@@ -115,25 +115,25 @@ private:
 	{
 		void Reset()
 		{
-			TotalTriangles = 0;
+			TotalPrimitives = 0;
 			TotalInstances = 0;
 
-			InstanceCullingIndirectTriangles = 0;
+			InstanceCullingIndirectPrimitives = 0;
 			InstanceCullingIndirectInstances = 0;
 
-			CustomIndirectTriangles = 0;
+			CustomIndirectPrimitives = 0;
 			CustomIndirectInstances = 0;
 
 			CategoryStats.Empty();
 		}
 
-		int32 TotalTriangles = 0;					
+		int32 TotalPrimitives = 0;					
 		int32 TotalInstances = 0;					
 
-		int32 InstanceCullingIndirectTriangles = 0;
+		int32 InstanceCullingIndirectPrimitives = 0;
 		int32 InstanceCullingIndirectInstances = 0;		
 
-		int32 CustomIndirectTriangles = 0;
+		int32 CustomIndirectPrimitives = 0;
 		int32 CustomIndirectInstances = 0;
 
 		/**
