Use an array of textures for DFAO output.

Maintaining an array of textures allows us to avoid an additional texture copy and maintain the condition that all DFAO views are offset to (0, 0), which makes upsampling easier. Seems to satisfy local tests with static meshes, indirect shadows, and permutations of visualization/dynres/viewrectoffset/emulatestereo.

diff --git a/Engine/Shaders/Private/DistanceFieldAOShared.ush b/Engine/Shaders/Private/DistanceFieldAOShared.ush
--- a/Engine/Shaders/Private/DistanceFieldAOShared.ush
+++ b/Engine/Shaders/Private/DistanceFieldAOShared.ush
@@ -144,7 +144,7 @@ float3 ComputeBentNormal(float3 RecordWorldNormal, uint RelativeRecordIndex)
 	return UnoccludedDirection;
 }
 
-float4 AOBufferBilinearUVMinMax;
+float2 AOBufferBilinearUVMax;
 
 struct FUpsampleDFAOOutput
 {
@@ -155,8 +155,7 @@ FUpsampleDFAOOutput UpsampleDFAO(float2 BufferUV, float SceneDepth)
 {
 	FUpsampleDFAOOutput UpsampleDFAOOutput;
 
-	// Note: if DFAO is ever stored per-view and not atlased at some point, there should only need to be a max here
-	float2 DistanceFieldUVs = clamp(BufferUV, AOBufferBilinearUVMinMax.xy, AOBufferBilinearUVMinMax.zw);
+	float2 DistanceFieldUVs = min(BufferUV, AOBufferBilinearUVMax);
 
 #define BILATERAL_UPSAMPLE 1
 #if BILATERAL_UPSAMPLE
diff --git a/Engine/Shaders/Private/DistanceFieldLightingPost.usf b/Engine/Shaders/Private/DistanceFieldLightingPost.usf
--- a/Engine/Shaders/Private/DistanceFieldLightingPost.usf
+++ b/Engine/Shaders/Private/DistanceFieldLightingPost.usf
@@ -370,7 +370,10 @@ void AOUpsamplePS(
 {
 	FGBufferData GBuffer = GetGBufferData(UVAndScreenPos.xy);
 
-	float3 BentNormal = UpsampleDFAO(UVAndScreenPos.xy, GBuffer.Depth, GBuffer.WorldNormal);
+	// The rectangle is drawn according to GBuffer rect,
+	// while the bent normal is always offset at (0, 0)
+	const float2 BentNormalUV = UVAndScreenPos.xy - View.ViewRectMin * View.BufferSizeAndInvSize.zw;
+	float3 BentNormal = UpsampleDFAO(BentNormalUV, GBuffer.Depth, GBuffer.WorldNormal);
 
 #if MODULATE_SCENE_COLOR
 	float Visibility = lerp(length(BentNormal), 1.0f, MinIndirectDiffuseOcclusion);
diff --git a/Engine/Shaders/Private/ReflectionEnvironmentPixelShader.usf b/Engine/Shaders/Private/ReflectionEnvironmentPixelShader.usf
--- a/Engine/Shaders/Private/ReflectionEnvironmentPixelShader.usf
+++ b/Engine/Shaders/Private/ReflectionEnvironmentPixelShader.usf
@@ -368,7 +368,8 @@ void ReflectionEnvironmentSkyLighting(
 	float3 BentNormal = GBuffer.WorldNormal;
 #if APPLY_SKY_SHADOWING
 	{
-		BentNormal = UpsampleDFAO(BufferUV, GBuffer.Depth, GBuffer.WorldNormal);
+		const float2 BentNormalUV = (PixelPos - View.ViewRectMin.xy) * View.BufferSizeAndInvSize.zw;
+		BentNormal = UpsampleDFAO(BentNormalUV, GBuffer.Depth, GBuffer.WorldNormal);
 	}
 #endif
 
diff --git a/Engine/Source/Runtime/Renderer/Private/DeferredShadingRenderer.cpp b/Engine/Source/Runtime/Renderer/Private/DeferredShadingRenderer.cpp
--- a/Engine/Source/Runtime/Renderer/Private/DeferredShadingRenderer.cpp
+++ b/Engine/Source/Runtime/Renderer/Private/DeferredShadingRenderer.cpp
@@ -3859,7 +3859,7 @@ void FDeferredShadingSceneRenderer::Render(FRDGBuilder& GraphBuilder)
 			SCOPE_CYCLE_COUNTER(STAT_FDeferredShadingSceneRenderer_Lighting);
 			SCOPED_NAMED_EVENT(RenderLighting, FColor::Emerald);
 
-			FRDGTextureRef DynamicBentNormalAOTexture = nullptr;
+			TArray<FRDGTextureRef> DynamicBentNormalAOTextures;
 
 			RenderDiffuseIndirectAndAmbientOcclusion(
 				GraphBuilder,
@@ -3878,7 +3878,7 @@ void FDeferredShadingSceneRenderer::Render(FRDGBuilder& GraphBuilder)
 			}
 
 			// These modulate the scene color output from the base pass, which is assumed to be indirect lighting
-			RenderDFAOAsIndirectShadowing(GraphBuilder, SceneTextures, DynamicBentNormalAOTexture);
+			RenderDFAOAsIndirectShadowing(GraphBuilder, SceneTextures, DynamicBentNormalAOTextures);
 
 			// Clear the translucent lighting volumes before we accumulate
 			if ((GbEnableAsyncComputeTranslucencyLightingVolumeClear && GSupportsEfficientAsyncCompute) == false)
@@ -3912,7 +3912,7 @@ void FDeferredShadingSceneRenderer::Render(FRDGBuilder& GraphBuilder)
 				AsyncLumenIndirectLightingOutputs);
 
 			// Render diffuse sky lighting and reflections that only operate on opaque pixels
-			RenderDeferredReflectionsAndSkyLighting(GraphBuilder, SceneTextures, LumenFrameTemporaries, DynamicBentNormalAOTexture);
+			RenderDeferredReflectionsAndSkyLighting(GraphBuilder, SceneTextures, LumenFrameTemporaries, DynamicBentNormalAOTextures);
 
 #if !(UE_BUILD_SHIPPING || UE_BUILD_TEST)
 			// Renders debug visualizations for global illumination plugins
@@ -4291,7 +4291,7 @@ void FDeferredShadingSceneRenderer::Render(FRDGBuilder& GraphBuilder)
 
 			// Use the skylight's max distance if there is one, to be consistent with DFAO shadowing on the skylight
 			const float OcclusionMaxDistance = Scene->SkyLight && !Scene->SkyLight->bWantsStaticShadowing ? Scene->SkyLight->OcclusionMaxDistance : Scene->DefaultMaxDistanceFieldOcclusionDistance;
-			FRDGTextureRef DummyOutput = nullptr;
+			TArray<FRDGTextureRef> DummyOutput;
 			RenderDistanceFieldLighting(GraphBuilder, SceneTextures, FDistanceFieldAOParameters(OcclusionMaxDistance), DummyOutput, false, ViewFamily.EngineShowFlags.VisualizeDistanceFieldAO);
 		}
 
diff --git a/Engine/Source/Runtime/Renderer/Private/DeferredShadingRenderer.h b/Engine/Source/Runtime/Renderer/Private/DeferredShadingRenderer.h
--- a/Engine/Source/Runtime/Renderer/Private/DeferredShadingRenderer.h
+++ b/Engine/Source/Runtime/Renderer/Private/DeferredShadingRenderer.h
@@ -675,7 +675,7 @@ private:
 		FRDGBuilder& GraphBuilder,
 		const FSceneTextures& SceneTextures,
 		const FLumenSceneFrameTemporaries& LumenFrameTemporaries,
-		FRDGTextureRef DynamicBentNormalAOTexture);
+		TArray<FRDGTextureRef>& DynamicBentNormalAOTexture);
 
 	void RenderDeferredReflectionsAndSkyLightingHair(FRDGBuilder& GraphBuilder);
 
@@ -688,7 +688,7 @@ private:
 	void RenderDFAOAsIndirectShadowing(
 		FRDGBuilder& GraphBuilder,
 		const FSceneTextures& SceneTextures,
-		FRDGTextureRef& DynamicBentNormalAO);
+		TArray<FRDGTextureRef>& DynamicBentNormalAOTextures);
 
 	bool ShouldRenderDistanceFieldLighting() const;
 
@@ -697,7 +697,7 @@ private:
 		FRDGBuilder& GraphBuilder,
 		const FSceneTextures& SceneTextures,
 		const class FDistanceFieldAOParameters& Parameters,
-		FRDGTextureRef& OutDynamicBentNormalAO,
+		TArray<FRDGTextureRef>& OutDynamicBentNormalAOTextures,
 		bool bModulateToSceneColor,
 		bool bVisualizeAmbientOcclusion);
 
diff --git a/Engine/Source/Runtime/Renderer/Private/DistanceFieldAmbientOcclusion.cpp b/Engine/Source/Runtime/Renderer/Private/DistanceFieldAmbientOcclusion.cpp
--- a/Engine/Source/Runtime/Renderer/Private/DistanceFieldAmbientOcclusion.cpp
+++ b/Engine/Source/Runtime/Renderer/Private/DistanceFieldAmbientOcclusion.cpp
@@ -278,16 +278,14 @@ FDFAOUpsampleParameters DistanceField::SetupAOUpsampleParameters(const FViewInfo
 
 	const FIntPoint AOBufferSize = GetBufferSizeForAO(View);
 	const FIntPoint AOViewSize = View.ViewRect.Size() / GAODownsampleFactor;
-	const FVector4f UVMinMax(
-		0.51 / AOBufferSize.X,
-		0.51 / AOBufferSize.Y,
+	const FVector2f UVMax(
 		(AOViewSize.X - 0.51f) / AOBufferSize.X, // 0.51 - so bilateral gather4 won't sample invalid texels
 		(AOViewSize.Y - 0.51f) / AOBufferSize.Y);
 
 	FDFAOUpsampleParameters ShaderParameters;
 	ShaderParameters.BentNormalAOTexture = DistanceFieldAOBentNormal;
 	ShaderParameters.BentNormalAOSampler = TStaticSamplerState<SF_Bilinear>::GetRHI();
-	ShaderParameters.AOBufferBilinearUVMinMax = UVMinMax;
+	ShaderParameters.AOBufferBilinearUVMax = UVMax;
 	ShaderParameters.DistanceFadeScale = DistanceFadeScaleValue;
 	ShaderParameters.AOMaxViewDistance = GetMaxAOViewDistance();
 
@@ -763,13 +761,13 @@ bool FSceneRenderer::ShouldPrepareGlobalDistanceField() const
 void FDeferredShadingSceneRenderer::RenderDFAOAsIndirectShadowing(
 	FRDGBuilder& GraphBuilder,
 	const FSceneTextures& SceneTextures,
-	FRDGTextureRef& DynamicBentNormalAO)
+	TArray<FRDGTextureRef>& DynamicBentNormalAOTextures)
 {
 	if (GDistanceFieldAOApplyToStaticIndirect && ShouldRenderDistanceFieldAO() && ShouldRenderDistanceFieldLighting())
 	{
 		// Use the skylight's max distance if there is one, to be consistent with DFAO shadowing on the skylight
 		const float OcclusionMaxDistance = Scene->SkyLight && !Scene->SkyLight->bWantsStaticShadowing ? Scene->SkyLight->OcclusionMaxDistance : Scene->DefaultMaxDistanceFieldOcclusionDistance;
-		RenderDistanceFieldLighting(GraphBuilder, SceneTextures, FDistanceFieldAOParameters(OcclusionMaxDistance), DynamicBentNormalAO, true, false);
+		RenderDistanceFieldLighting(GraphBuilder, SceneTextures, FDistanceFieldAOParameters(OcclusionMaxDistance), DynamicBentNormalAOTextures, true, false);
 	}
 }
 
@@ -801,7 +799,7 @@ void FDeferredShadingSceneRenderer::RenderDistanceFieldLighting(
 	FRDGBuilder& GraphBuilder,
 	const FSceneTextures& SceneTextures,
 	const FDistanceFieldAOParameters& Parameters,
-	FRDGTextureRef& OutDynamicBentNormalAO,
+	TArray<FRDGTextureRef>& OutDynamicBentNormalAOTextures,
 	bool bModulateToSceneColor,
 	bool bVisualizeAmbientOcclusion)
 {
@@ -826,14 +824,7 @@ void FDeferredShadingSceneRenderer::RenderDistanceFieldLighting(
 	}
 	
 	// We only need this texture if we need to copy multiple view outputs to an texture atlas
-	FRDGTextureRef BentNormalOutput = nullptr;
-	 if (Views.Num() > 1)
-	{
-		const FIntPoint BufferSize = GetActiveSceneTexturesConfig().Extent / GAODownsampleFactor;
-		const FRDGTextureDesc Desc = FRDGTextureDesc::Create2D(BufferSize, PF_FloatRGBA, FClearValueBinding::None, GFastVRamConfig.DistanceFieldAOBentNormal | TexCreate_ShaderResource | TexCreate_RenderTargetable | TexCreate_UAV);
-		BentNormalOutput = GraphBuilder.CreateTexture(Desc, TEXT("DistanceFieldBentNormalAO"));
-		AddClearUAVPass(GraphBuilder, GraphBuilder.CreateUAV(BentNormalOutput), FLinearColor::Black);
-	}
+	OutDynamicBentNormalAOTextures.Reset();
 
 	for (int32 ViewIndex = 0; ViewIndex < Views.Num(); ++ViewIndex)
 	{
@@ -899,17 +890,8 @@ void FDeferredShadingSceneRenderer::RenderDistanceFieldLighting(
 			UpsampleBentNormalAO(GraphBuilder, View, SceneTextures.UniformBuffer, SceneTextures.Color.Target, PerViewBentNormal, bModulateToSceneColor && !bVisualizeAmbientOcclusion);
 		}
 
-		if (Views.Num() == 1) {
-			OutDynamicBentNormalAO = PerViewBentNormal;
-			return;
-		}
-		else 
-		{
-			AddCopyTexturePass(GraphBuilder, PerViewBentNormal, BentNormalOutput, FIntPoint::ZeroValue, View.ViewRect.Min / GAODownsampleFactor, View.ViewRect.Size() / GAODownsampleFactor);
-		}
+		OutDynamicBentNormalAOTextures.Add(PerViewBentNormal);
 	}
-
-	OutDynamicBentNormalAO = BentNormalOutput;
 }
 
 bool FSceneRenderer::ShouldRenderDistanceFieldAO() const
diff --git a/Engine/Source/Runtime/Renderer/Private/DistanceFieldAmbientOcclusion.h b/Engine/Source/Runtime/Renderer/Private/DistanceFieldAmbientOcclusion.h
--- a/Engine/Source/Runtime/Renderer/Private/DistanceFieldAmbientOcclusion.h
+++ b/Engine/Source/Runtime/Renderer/Private/DistanceFieldAmbientOcclusion.h
@@ -90,7 +90,7 @@ END_SHADER_PARAMETER_STRUCT()
 BEGIN_SHADER_PARAMETER_STRUCT(FDFAOUpsampleParameters, )
 	SHADER_PARAMETER_RDG_TEXTURE(Texture2D, BentNormalAOTexture)
 	SHADER_PARAMETER_SAMPLER(SamplerState, BentNormalAOSampler)
-	SHADER_PARAMETER(FVector4f, AOBufferBilinearUVMinMax)
+	SHADER_PARAMETER(FVector2f, AOBufferBilinearUVMax)
 	SHADER_PARAMETER(float, DistanceFadeScale)
 	SHADER_PARAMETER(float, AOMaxViewDistance)
 END_SHADER_PARAMETER_STRUCT()
diff --git a/Engine/Source/Runtime/Renderer/Private/IndirectLightRendering.cpp b/Engine/Source/Runtime/Renderer/Private/IndirectLightRendering.cpp
--- a/Engine/Source/Runtime/Renderer/Private/IndirectLightRendering.cpp
+++ b/Engine/Source/Runtime/Renderer/Private/IndirectLightRendering.cpp
@@ -374,7 +374,7 @@ class FReflectionEnvironmentSkyLightingPS : public FGlobalShader
 	BEGIN_SHADER_PARAMETER_STRUCT(FParameters, )
 		// Distance field AO parameters.
 		// TODO. FDFAOUpsampleParameters
-		SHADER_PARAMETER(FVector4f, AOBufferBilinearUVMinMax)
+		SHADER_PARAMETER(FVector2f, AOBufferBilinearUVMax)
 		SHADER_PARAMETER(float, DistanceFadeScale)
 		SHADER_PARAMETER(float, AOMaxViewDistance)
 
@@ -1691,13 +1691,11 @@ static void AddSkyReflectionPass(
 		{
 			// The view must sample into the texture atlas correctly with the appropriate clamping
 			const FIntPoint AOBufferSize = GetBufferSizeForAO(View);
-			const FIntRect AOViewRect = View.ViewRect / GAODownsampleFactor;
+			const FIntPoint AOViewSize = View.ViewRect.Size() / GAODownsampleFactor;
 
-			PassParameters->PS.AOBufferBilinearUVMinMax = FVector4f(
-				(AOViewRect.Min.X + 0.51f) / AOBufferSize.X, // 0.51 - so bilateral gather4 won't sample invalid texels
-				(AOViewRect.Min.Y + 0.51f) / AOBufferSize.Y,
-				(AOViewRect.Max.X - 0.51f) / AOBufferSize.X,
-				(AOViewRect.Max.Y - 0.51f) / AOBufferSize.Y);
+			PassParameters->PS.AOBufferBilinearUVMax = FVector2f(
+				(AOViewSize.X - 0.51f) / AOBufferSize.X,
+				(AOViewSize.Y - 0.51f) / AOBufferSize.Y);
 
 			extern float GAOViewFadeDistanceScale;
 			PassParameters->PS.AOMaxViewDistance = GetMaxAOViewDistance();
@@ -1865,7 +1863,7 @@ void FDeferredShadingSceneRenderer::RenderDeferredReflectionsAndSkyLighting(
 	FRDGBuilder& GraphBuilder,
 	const FSceneTextures& SceneTextures,
 	const FLumenSceneFrameTemporaries& LumenFrameTemporaries,
-	FRDGTextureRef DynamicBentNormalAOTexture)
+	TArray<FRDGTextureRef>& DynamicBentNormalAOTextures)
 {
 	extern int32 GLumenVisualizeIndirectDiffuse;
 	if (ViewFamily.EngineShowFlags.VisualizeLightCulling 
@@ -1912,7 +1910,7 @@ void FDeferredShadingSceneRenderer::RenderDeferredReflectionsAndSkyLighting(
 		{
 			bApplySkyShadowing = true;
 			FDistanceFieldAOParameters Parameters(Scene->SkyLight->OcclusionMaxDistance, Scene->SkyLight->Contrast);
-			RenderDistanceFieldLighting(GraphBuilder, SceneTextures, Parameters, DynamicBentNormalAOTexture, false, false);
+			RenderDistanceFieldLighting(GraphBuilder, SceneTextures, Parameters, DynamicBentNormalAOTextures, false, false);
 		}
 	}
 
@@ -2108,6 +2106,8 @@ void FDeferredShadingSceneRenderer::RenderDeferredReflectionsAndSkyLighting(
 		{
 			RDG_GPU_STAT_SCOPE(GraphBuilder, ReflectionEnvironment);
 
+			FRDGTextureRef DynamicBentNormalAOTexture = DynamicBentNormalAOTextures.IsEmpty() ? nullptr : DynamicBentNormalAOTextures[CurrentViewIndex];
+
 			if (Strata::IsStrataEnabled())
 			{
 				AddSkyReflectionPass(
