Allows to override VisualizeTSR showflag with r.TSR.Visualize

diff --git a/Engine/Source/Runtime/Renderer/Private/PostProcess/PostProcessing.cpp b/Engine/Source/Runtime/Renderer/Private/PostProcess/PostProcessing.cpp
--- a/Engine/Source/Runtime/Renderer/Private/PostProcess/PostProcessing.cpp
+++ b/Engine/Source/Runtime/Renderer/Private/PostProcess/PostProcessing.cpp
@@ -567,7 +567,7 @@ void AddPostProcessingPasses(
 		const bool bDepthOfFieldEnabled = DiaphragmDOF::IsEnabled(View);
 		const bool bVisualizeDepthOfField = bDepthOfFieldEnabled && EngineShowFlags.VisualizeDOF;
 		const bool bVisualizeMotionBlur = IsVisualizeMotionBlurEnabled(View);
-		const bool bVisualizeTSR = TAAConfig == EMainTAAPassConfig::TSR && EngineShowFlags.VisualizeTSR;
+		const bool bVisualizeTSR = IsVisualizeTSREnabled(View);
 
 		const EAutoExposureMethod AutoExposureMethod = GetAutoExposureMethod(View);
 		const EAntiAliasingMethod AntiAliasingMethod = !bVisualizeDepthOfField ? View.AntiAliasingMethod : AAM_None;
diff --git a/Engine/Source/Runtime/Renderer/Private/PostProcess/TemporalAA.h b/Engine/Source/Runtime/Renderer/Private/PostProcess/TemporalAA.h
--- a/Engine/Source/Runtime/Renderer/Private/PostProcess/TemporalAA.h
+++ b/Engine/Source/Runtime/Renderer/Private/PostProcess/TemporalAA.h
@@ -176,6 +176,9 @@ extern RENDERER_API FTAAOutputs AddTemporalAAPass(
 /** Returns whether a given view need to measure luminance of the scene color for moire anti-flickering. */
 bool NeedTSRMoireLuma(const FViewInfo& View);
 
+/** Returns whether TSR internal visualization is enabled on the view. */
+bool IsVisualizeTSREnabled(const FViewInfo& View);
+
 /** Measure luminance of the scene color for moire anti-flickering. */
 FScreenPassTexture AddTSRComputeMoireLuma(FRDGBuilder& GraphBuilder, FGlobalShaderMap* ShaderMap, FScreenPassTexture SceneColor);
 
diff --git a/Engine/Source/Runtime/Renderer/Private/PostProcess/TemporalSuperResolution.cpp b/Engine/Source/Runtime/Renderer/Private/PostProcess/TemporalSuperResolution.cpp
--- a/Engine/Source/Runtime/Renderer/Private/PostProcess/TemporalSuperResolution.cpp
+++ b/Engine/Source/Runtime/Renderer/Private/PostProcess/TemporalSuperResolution.cpp
@@ -272,9 +272,10 @@ TAutoConsoleVariable<int32> CVarTSRSubpixelIncludeMovingDepth(
 #if !UE_BUILD_OPTIMIZED_SHOWFLAGS
 
 TAutoConsoleVariable<int32> CVarTSRVisualize(
-	TEXT("r.TSR.Visualize"), -1,
-	TEXT("Selects what to display with the VisualizeTSR show flag (opened with the `show VisualizeTSR` command at runtime or Show > Visualize > TSR in editor viewports).\n")
-	TEXT(" -1: Display an overview grid (default);\n")
+	TEXT("r.TSR.Visualize"), -2,
+	TEXT("Selects the TSR internal visualization mode.\n")
+	TEXT(" -2: Display an overview grid based on the VisualizeTSR show flag (default, opened with the `show VisualizeTSR` command at runtime or Show > Visualize > TSR in editor viewports);\n")
+	TEXT(" -1: Display an overview grid based regardless of VisualizeTSR show flag;\n")
 	TEXT("  0: Number of accumulated samples in the history, particularily interesting to tune r.TSR.ShadingRejection.SampleCount and r.TSR.Velocity.WeightClampingSampleCount;\n")
 	TEXT("  1: Parallax disocclusion based of depth and velocity buffers;\n")
 	TEXT("  2: Mask where the history is rejected;\n")
@@ -1103,6 +1104,17 @@ bool NeedTSRMoireLuma(const FViewInfo& View)
 	return GetMainTAAPassConfig(View) == EMainTAAPassConfig::TSR;
 }
 
+bool IsVisualizeTSREnabled(const FViewInfo& View)
+#if UE_BUILD_OPTIMIZED_SHOWFLAGS
+{
+	return false;
+}
+#else
+{
+	return GetMainTAAPassConfig(View) == EMainTAAPassConfig::TSR && (View.Family->EngineShowFlags.VisualizeTSR || CVarTSRVisualize.GetValueOnRenderThread() >= -1);
+}
+#endif
+
 FScreenPassTexture AddTSRComputeMoireLuma(FRDGBuilder& GraphBuilder, FGlobalShaderMap* ShaderMap, FScreenPassTexture SceneColor)
 {
 	check(SceneColor.Texture)
@@ -2433,7 +2445,7 @@ FDefaultTemporalUpscaler::FOutputs AddTemporalSuperResolutionPasses(
 	}
 
 #if !UE_BUILD_OPTIMIZED_SHOWFLAGS
-	if (View.Family->EngineShowFlags.VisualizeTSR)
+	if (IsVisualizeTSREnabled(View))
 	{
 		RDG_EVENT_SCOPE(GraphBuilder, "VisualizeTSR %dx%d", OutputRect.Width(), OutputRect.Height());
 
