Fix module dependency of FDispatchShaderBundleCS and FDispatchShaderBundle utility by moving it from Renderer into RenderCore (suitable to be called from RHI modules in modular builds).

diff --git a/Engine/Source/Runtime/RenderCore/Private/ShaderBundles.cpp b/Engine/Source/Runtime/RenderCore/Private/ShaderBundles.cpp
new file mode 100644
--- /dev/null
+++ b/Engine/Source/Runtime/RenderCore/Private/ShaderBundles.cpp
@@ -0,0 +1,53 @@
+// Copyright Epic Games, Inc. All Rights Reserved.
+
+#include "ShaderBundles.h"
+#include "DataDrivenShaderPlatformInfo.h"
+#include "RenderGraph.h"
+#include "RenderGraphDefinitions.h"
+#include "ShaderParameterMacros.h"
+#include "RenderGraphFwd.h"
+
+IMPLEMENT_GLOBAL_SHADER(FDispatchShaderBundleCS, "/Engine/Private/ShaderBundleDispatch.usf", "DispatchShaderBundleEntry", SF_Compute);
+
+bool FDispatchShaderBundleCS::ShouldCompilePermutation(const FGlobalShaderPermutationParameters& Parameters)
+{
+	return RHISupportsShaderBundleDispatch(Parameters.Platform);
+}
+
+void FDispatchShaderBundleCS::ModifyCompilationEnvironment(const FGlobalShaderPermutationParameters& Parameters, FShaderCompilerEnvironment& OutEnvironment)
+{
+	OutEnvironment.SetDefine(TEXT("THREADGROUP_SIZEX"), ThreadGroupSizeX);
+	FGlobalShader::ModifyCompilationEnvironment(Parameters, OutEnvironment);
+}
+
+void FDispatchShaderBundle::Dispatch(
+	FRHIShaderBundle* ShaderBundle,
+	FRHIComputeCommandList& RHICmdList,
+	FRHIShaderResourceView* RecordArgBufferSRV,
+	FRHIShaderResourceView* RecordDataBufferSRV,
+	FRHIUnorderedAccessView* ExecutionBufferUAV
+)
+{
+	check(RHISupportsShaderBundleDispatch(GMaxRHIShaderPlatform) && GRHISupportsDispatchShaderBundle);
+	check(ShaderBundle && ShaderBundle->NumRecords > 0);
+
+	RHICmdList.ClearUAVUint(ExecutionBufferUAV, FUintVector4(0, 0, 0, 0));
+
+	auto ComputeShader = GetGlobalShaderMap(GMaxRHIFeatureLevel)->GetShader<FDispatchShaderBundleCS>();
+
+	FDispatchShaderBundleCS::FParameters Parameters;
+
+	Parameters.RecordCount = ShaderBundle->NumRecords;
+	Parameters.PlatformData = ShaderBundle->GetPlatformData();
+	Parameters.RecordArgBuffer = RecordArgBufferSRV;
+	Parameters.RecordDataBuffer = RecordDataBufferSRV;
+	Parameters.RWExecutionBuffer = ExecutionBufferUAV;
+
+	const FIntVector GroupCount = FComputeShaderUtils::GetGroupCount(ShaderBundle->NumRecords, FDispatchShaderBundleCS::ThreadGroupSizeX);
+
+	FRHIComputeShader* ComputeShaderRHI = ComputeShader.GetComputeShader();
+	SetComputePipelineState(RHICmdList, ComputeShaderRHI);
+	SetShaderParameters(RHICmdList, ComputeShader, ComputeShaderRHI, Parameters);
+	DispatchComputeShader(RHICmdList, ComputeShader, GroupCount.X, GroupCount.Y, GroupCount.Z);
+	UnsetShaderUAVs(RHICmdList, ComputeShader, ComputeShaderRHI);
+}
\ No newline at end of file
diff --git a/Engine/Source/Runtime/RenderCore/Public/ShaderBundles.h b/Engine/Source/Runtime/RenderCore/Public/ShaderBundles.h
new file mode 100644
--- /dev/null
+++ b/Engine/Source/Runtime/RenderCore/Public/ShaderBundles.h
@@ -0,0 +1,53 @@
+// Copyright Epic Games, Inc. All Rights Reserved.
+
+#pragma once
+
+#include "CoreMinimal.h"
+#include "GlobalShader.h"
+#include "HAL/Platform.h"
+#include "Misc/AssertionMacros.h"
+#include "RHICommandList.h"
+#include "RHIDefinitions.h"
+#include "Serialization/MemoryLayout.h"
+#include "Shader.h"
+#include "ShaderCore.h"
+#include "ShaderParameterMacros.h"
+#include "ShaderParameterStruct.h"
+#include "ShaderParameterUtils.h"
+#include "ShaderParameters.h"
+#include "ShaderPermutation.h"
+
+class FRHIShaderBundle;
+
+class FDispatchShaderBundleCS : public FGlobalShader
+{
+	DECLARE_EXPORTED_GLOBAL_SHADER(FDispatchShaderBundleCS, RENDERCORE_API);
+
+public:
+	SHADER_USE_PARAMETER_STRUCT(FDispatchShaderBundleCS, FGlobalShader)
+
+	static const uint32 ThreadGroupSizeX = 64;
+
+	BEGIN_SHADER_PARAMETER_STRUCT(FParameters, RENDERCORE_API)
+		SHADER_PARAMETER(uint32, RecordCount)
+		SHADER_PARAMETER(FUintVector4, PlatformData)
+		SHADER_PARAMETER_SRV(ByteAddressBuffer, RecordArgBuffer)
+		SHADER_PARAMETER_SRV(ByteAddressBuffer, RecordDataBuffer)
+		SHADER_PARAMETER_UAV(RWByteAddressBuffer, RWExecutionBuffer)
+	END_SHADER_PARAMETER_STRUCT()
+
+	static bool ShouldCompilePermutation(const FGlobalShaderPermutationParameters& Parameters);
+	static void ModifyCompilationEnvironment(const FGlobalShaderPermutationParameters& Parameters, FShaderCompilerEnvironment& OutEnvironment);
+};
+
+class FDispatchShaderBundle
+{
+public:
+	static RENDERCORE_API void Dispatch(
+		FRHIShaderBundle* ShaderBundle,
+		FRHIComputeCommandList& RHICmdList,
+		FRHIShaderResourceView* RecordArgBufferSRV,
+		FRHIShaderResourceView* RecordDataBufferSRV,
+		FRHIUnorderedAccessView* ExecutionBufferUAV
+	);
+};
\ No newline at end of file
diff --git a/Engine/Source/Runtime/Renderer/Private/RendererUtils.cpp b/Engine/Source/Runtime/Renderer/Private/RendererUtils.cpp
--- a/Engine/Source/Runtime/Renderer/Private/RendererUtils.cpp
+++ b/Engine/Source/Runtime/Renderer/Private/RendererUtils.cpp
@@ -269,67 +269,3 @@ RENDERER_API FBufferRHIRef& GetOneTileQuadIndexBuffer()
 {
 	return GOneTileQuadIndexBuffer.IndexBufferRHI;
 }
-
-class FDispatchShaderBundleCS : public FGlobalShader
-{
-public:
-	DECLARE_GLOBAL_SHADER(FDispatchShaderBundleCS);
-	SHADER_USE_PARAMETER_STRUCT(FDispatchShaderBundleCS, FGlobalShader)
-
-	static const uint32 ThreadGroupSizeX = 64;
-
-	using FPermutationDomain = TShaderPermutationDomain<>;
-
-	BEGIN_SHADER_PARAMETER_STRUCT(FParameters, )
-		SHADER_PARAMETER(uint32, RecordCount)
-		SHADER_PARAMETER(FUintVector4, PlatformData)
-		SHADER_PARAMETER_SRV(ByteAddressBuffer, RecordArgBuffer)
-		SHADER_PARAMETER_SRV(ByteAddressBuffer, RecordDataBuffer)
-		SHADER_PARAMETER_UAV(RWByteAddressBuffer, RWExecutionBuffer)
-	END_SHADER_PARAMETER_STRUCT()
-
-	static bool ShouldCompilePermutation(const FGlobalShaderPermutationParameters& Parameters)
-	{
-		return RHISupportsShaderBundleDispatch(Parameters.Platform);
-	}
-
-	static void ModifyCompilationEnvironment(const FGlobalShaderPermutationParameters& Parameters, FShaderCompilerEnvironment& OutEnvironment)
-	{
-		OutEnvironment.SetDefine(TEXT("THREADGROUP_SIZEX"), ThreadGroupSizeX);
-		FGlobalShader::ModifyCompilationEnvironment(Parameters, OutEnvironment);
-	}
-};
-
-IMPLEMENT_GLOBAL_SHADER(FDispatchShaderBundleCS, "/Engine/Private/ShaderBundleDispatch.usf", "DispatchShaderBundleEntry", SF_Compute);
-
-void FDispatchShaderBundle::Dispatch(
-	FRHIShaderBundle* ShaderBundle,
-	FRHIComputeCommandList& RHICmdList,
-	FRHIShaderResourceView* RecordArgBufferSRV,
-	FRHIShaderResourceView* RecordDataBufferSRV,
-	FRHIUnorderedAccessView* ExecutionBufferUAV
-)
-{
-	check(RHISupportsShaderBundleDispatch(GMaxRHIShaderPlatform) && GRHISupportsDispatchShaderBundle);
-	check(ShaderBundle && ShaderBundle->NumRecords > 0);
-
-	RHICmdList.ClearUAVUint(ExecutionBufferUAV, FUintVector4(0, 0, 0, 0));
-
-	auto ComputeShader = GetGlobalShaderMap(GMaxRHIFeatureLevel)->GetShader<FDispatchShaderBundleCS>();
-
-	FDispatchShaderBundleCS::FParameters Parameters;
-
-	Parameters.RecordCount = ShaderBundle->NumRecords;
-	Parameters.PlatformData = ShaderBundle->GetPlatformData();
-	Parameters.RecordArgBuffer = RecordArgBufferSRV;
-	Parameters.RecordDataBuffer = RecordDataBufferSRV;
-	Parameters.RWExecutionBuffer = ExecutionBufferUAV;
-
-	const FIntVector GroupCount = FComputeShaderUtils::GetGroupCount(ShaderBundle->NumRecords, FDispatchShaderBundleCS::ThreadGroupSizeX);
-
-	FRHIComputeShader* ComputeShaderRHI = ComputeShader.GetComputeShader();
-	SetComputePipelineState(RHICmdList, ComputeShaderRHI);
-	SetShaderParameters(RHICmdList, ComputeShader, ComputeShaderRHI, Parameters);
-	DispatchComputeShader(RHICmdList, ComputeShader, GroupCount.X, GroupCount.Y, GroupCount.Z);
-	UnsetShaderUAVs(RHICmdList, ComputeShader, ComputeShaderRHI);
-}
diff --git a/Engine/Source/Runtime/Renderer/Public/RendererUtils.h b/Engine/Source/Runtime/Renderer/Public/RendererUtils.h
--- a/Engine/Source/Runtime/Renderer/Public/RendererUtils.h
+++ b/Engine/Source/Runtime/Renderer/Public/RendererUtils.h
@@ -17,7 +17,6 @@ class FScene;
 class FGlobalShaderMap;
 class FRDGBuilder;
 class FRHICommandListImmediate;
-class FRHIShaderBundle;
 class FScene;
 struct IPooledRenderTarget;
 template <typename ReferencedType> class TRefCountPtr;
@@ -42,18 +41,6 @@ public:
 		const TCHAR* RTWriteMaskDebugName);
 };
 
-class FDispatchShaderBundle
-{
-public:
-	static RENDERER_API void Dispatch(
-		FRHIShaderBundle* ShaderBundle,
-		FRHIComputeCommandList& RHICmdList,
-		FRHIShaderResourceView* RecordArgBufferSRV,
-		FRHIShaderResourceView* RecordDataBufferSRV,
-		FRHIUnorderedAccessView* ExecutionBufferUAV
-	);
-};
-
 class FDepthBounds
 {
 public:
