Mesh draw command stat categories now default to component class name unless explicitly set.
Budget categories can now capture groups of stat categories.
The MDC stat system is now development build only.
Added r.MeshDrawCommands.Stat 2 which displays visibile triangles per draw pass.

diff --git a/Engine/Plugins/Experimental/Water/Source/Runtime/Private/WaterMeshComponent.cpp b/Engine/Plugins/Experimental/Water/Source/Runtime/Private/WaterMeshComponent.cpp
--- a/Engine/Plugins/Experimental/Water/Source/Runtime/Private/WaterMeshComponent.cpp
+++ b/Engine/Plugins/Experimental/Water/Source/Runtime/Private/WaterMeshComponent.cpp
@@ -64,9 +64,6 @@ UWaterMeshComponent::UWaterMeshComponent()
 	bHasPerInstanceHitProxies = true;
 
 	SetMobility(EComponentMobility::Static);
-
-	static FName NAME_StatsCategoryWater("Water");
-	SetMeshDrawCommandStatsCategory(NAME_StatsCategoryWater);
 }
 
 void UWaterMeshComponent::PostLoad()
diff --git a/Engine/Source/Runtime/Engine/Classes/Components/PrimitiveComponent.h b/Engine/Source/Runtime/Engine/Classes/Components/PrimitiveComponent.h
--- a/Engine/Source/Runtime/Engine/Classes/Components/PrimitiveComponent.h
+++ b/Engine/Source/Runtime/Engine/Classes/Components/PrimitiveComponent.h
@@ -30,7 +30,7 @@
 #include "HLOD/HLODLevelExclusion.h"
 #include "Stats/Stats2.h"
 #include "PSOPrecache.h"
-#include "MeshDrawCommandStatsComponentData.h"
+#include "MeshDrawCommandStatsDefines.h"
 #include "PrimitiveComponent.generated.h"
 
 DECLARE_CYCLE_STAT_EXTERN(TEXT("BeginComponentOverlap"), STAT_BeginComponentOverlap, STATGROUP_Game, ENGINE_API);
@@ -691,18 +691,6 @@ public:
 	uint8 SelectionOutlineColorIndex;
 #endif
 
-#if !MESH_DRAW_COMMAND_STAT_COLLECTION
-public:
-	ENGINE_API void SetMeshDrawCommandStatsCategory(FName StatsCategory) {}
-#else
-private:
-	/** Unique id use to retrieve the component data when building the mesh data draw stat information */
-	FMeshDrawCommandStatsComponentDataID MeshDrawCommandStatsComponentDataID = INDEX_NONE;
-public:
-	ENGINE_API void SetMeshDrawCommandStatsCategory(FName StatsCategory);
-	FMeshDrawCommandStatsComponentDataID GetMeshDrawCommandStatsComponentDataID() const { return MeshDrawCommandStatsComponentDataID; }
-#endif //!MESH_DRAW_COMMAND_STAT_COLLECTION
-
 public:
 	/** If true then DoCustomNavigableGeometryExport will be called to collect navigable geometry of this component. */
 	UPROPERTY()
@@ -926,6 +914,11 @@ private:
 
 	float OcclusionBoundsSlack;
 
+#if MESH_DRAW_COMMAND_STATS
+	/** Optional category name for this component in the mesh draw stat collection. */
+	FName MeshDrawCommandStatsCategory;
+#endif
+
 	friend class FPrimitiveSceneInfo;
 
 public:
@@ -945,6 +938,13 @@ public:
 	float GetLastRenderTime() const { return LastRenderTime; }
 	float GetLastRenderTimeOnScreen() const { return LastRenderTimeOnScreen; }
 
+#if MESH_DRAW_COMMAND_STATS
+	ENGINE_API void SetMeshDrawCommandStatsCategory(FName StatsCategory);
+	FName GetMeshDrawCommandStatsCategory() const;
+#else
+	void SetMeshDrawCommandStatsCategory(FName StatsCategory) {}
+#endif
+
 	/**
 	 * Setup the parameter struct used to precache the PSOs used by this component. 
 	 * Precaching uses certain component attributes to derive the shader or state used to render the component such as static lighting, cast shadows, ...
diff --git a/Engine/Source/Runtime/Engine/Private/Components/PrimitiveComponent.cpp b/Engine/Source/Runtime/Engine/Private/Components/PrimitiveComponent.cpp
--- a/Engine/Source/Runtime/Engine/Private/Components/PrimitiveComponent.cpp
+++ b/Engine/Source/Runtime/Engine/Private/Components/PrimitiveComponent.cpp
@@ -4471,19 +4471,22 @@ void UPrimitiveComponent::SetLastRenderTime(float InLastRenderTime)
 	}
 }
 
-#if MESH_DRAW_COMMAND_STAT_COLLECTION
+#if MESH_DRAW_COMMAND_STATS
 void UPrimitiveComponent::SetMeshDrawCommandStatsCategory(FName StatsCategory)
 {
-	if (FMeshDrawCommandStatsComponentDataManager* Manager = FMeshDrawCommandStatsComponentDataManager::Get())
+	if (MeshDrawCommandStatsCategory != StatsCategory)
 	{
-		FMeshDrawCommandStatsComponentData ComponentData;
-		ComponentData.ComponentType = GetClass()->GetFName();
-		ComponentData.StatsCategory = StatsCategory;
-		MeshDrawCommandStatsComponentDataID = Manager->GetID(ComponentData);
+		MeshDrawCommandStatsCategory = StatsCategory;
 		MarkRenderStateDirty();
 	}
 }
-#endif // MESH_DRAW_COMMAND_STAT_COLLECTION
+
+FName UPrimitiveComponent::GetMeshDrawCommandStatsCategory() const
+{
+	// If a stats category isn't set on the component then use the component type.
+	return MeshDrawCommandStatsCategory.IsNone() ? GetClass()->GetFName() : MeshDrawCommandStatsCategory;
+}
+#endif
 
 void UPrimitiveComponent::SetupPrecachePSOParams(FPSOPrecacheParams& Params)
 {
diff --git a/Engine/Source/Runtime/Engine/Private/MeshDrawCommandStatsComponentData.cpp b/Engine/Source/Runtime/Engine/Private/MeshDrawCommandStatsComponentData.cpp
deleted file mode 100644
diff --git a/Engine/Source/Runtime/Engine/Private/PrimitiveSceneProxy.cpp b/Engine/Source/Runtime/Engine/Private/PrimitiveSceneProxy.cpp
--- a/Engine/Source/Runtime/Engine/Private/PrimitiveSceneProxy.cpp
+++ b/Engine/Source/Runtime/Engine/Private/PrimitiveSceneProxy.cpp
@@ -391,9 +391,6 @@ FPrimitiveSceneProxy::FPrimitiveSceneProxy(const UPrimitiveComponent* InComponen
 #endif
 ,	MaxDrawDistance(InComponent->CachedMaxDrawDistance > 0 ? InComponent->CachedMaxDrawDistance : FLT_MAX)
 ,	MinDrawDistance(InComponent->MinDrawDistance)
-#if MESH_DRAW_COMMAND_STAT_COLLECTION
-,	MeshDrawCommandStatsComponentDataID(InComponent->GetMeshDrawCommandStatsComponentDataID())
-#endif
 {
 	check(Scene);
 
@@ -418,6 +415,10 @@ FPrimitiveSceneProxy::FPrimitiveSceneProxy(const UPrimitiveComponent* InComponen
 	}
 #endif
 
+#if MESH_DRAW_COMMAND_STATS
+	MeshDrawCommandStatsCategory = InComponent->GetMeshDrawCommandStatsCategory();
+#endif
+
 	if (bNeedsUnbuiltPreviewLighting && !bHasValidSettingsForStaticLighting)
 	{
 		// Don't use unbuilt preview lighting for static components that have an invalid lightmap UV setup
diff --git a/Engine/Source/Runtime/Engine/Private/UnrealEngine.cpp b/Engine/Source/Runtime/Engine/Private/UnrealEngine.cpp
--- a/Engine/Source/Runtime/Engine/Private/UnrealEngine.cpp
+++ b/Engine/Source/Runtime/Engine/Private/UnrealEngine.cpp
@@ -214,7 +214,6 @@ UnrealEngine.cpp: Implements the UEngine class and helpers.
 #include "IEyeTrackerModule.h"
 #include "Interfaces/IPluginManager.h"
 #include "Animation/SkeletonRemappingRegistry.h"
-#include "MeshDrawCommandStatsComponentData.h"
 
 #if !UE_BUILD_SHIPPING
 #include "GenericPlatform/GenericPlatformCrashContext.h"
@@ -266,10 +265,6 @@ void OnChangeEngineCVarRequiringRecreateRenderState(IConsoleVariable* Var)
 
 void FEngineModule::StartupModule()
 {
-#if MESH_DRAW_COMMAND_STAT_COLLECTION
-	FMeshDrawCommandStatsComponentDataManager::CreateInstance();
-#endif
-
 	// Setup delegate callback for ProfilingHelpers to access current map name
 	extern const FString GetMapNameStatic();
 	GGetMapNameDelegate.BindStatic(&GetMapNameStatic);
diff --git a/Engine/Source/Runtime/Engine/Public/MeshDrawCommandStatsComponentData.h b/Engine/Source/Runtime/Engine/Public/MeshDrawCommandStatsComponentData.h
deleted file mode 100644
diff --git a/Engine/Source/Runtime/Engine/Public/MeshDrawCommandStatsDefines.h b/Engine/Source/Runtime/Engine/Public/MeshDrawCommandStatsDefines.h
new file mode 100644
--- /dev/null
+++ b/Engine/Source/Runtime/Engine/Public/MeshDrawCommandStatsDefines.h
@@ -0,0 +1,6 @@
+// Copyright Epic Games, Inc. All Rights Reserved.
+#pragma once
+
+#include "Misc/Build.h"
+
+#define MESH_DRAW_COMMAND_STATS (!UE_BUILD_SHIPPING && !UE_BUILD_TEST)
diff --git a/Engine/Source/Runtime/Engine/Public/MeshDrawCommandStatsSettings.h b/Engine/Source/Runtime/Engine/Public/MeshDrawCommandStatsSettings.h
--- a/Engine/Source/Runtime/Engine/Public/MeshDrawCommandStatsSettings.h
+++ b/Engine/Source/Runtime/Engine/Public/MeshDrawCommandStatsSettings.h
@@ -16,6 +16,9 @@ struct FMeshDrawCommandStatsBudget
 	/** Category name. */
 	UPROPERTY(config, EditAnywhere, Category = Engine)
 	FName CategoryName;
+	/** Stat names that will match this category name. */
+	UPROPERTY(config, EditAnywhere, Category = Engine)
+	TArray<FName> LinkedStatNames;
 	/** The category primitive budget. This is the maximum triangles expected, post-culling, summed across all passes. */
 	UPROPERTY(config, EditAnywhere, Category = Engine)
 	int32 PrimitiveBudget = 0;
diff --git a/Engine/Source/Runtime/Engine/Public/PrimitiveSceneProxy.h b/Engine/Source/Runtime/Engine/Public/PrimitiveSceneProxy.h
--- a/Engine/Source/Runtime/Engine/Public/PrimitiveSceneProxy.h
+++ b/Engine/Source/Runtime/Engine/Public/PrimitiveSceneProxy.h
@@ -19,7 +19,7 @@
 #include "DrawDebugHelpers.h"
 #include "Math/CapsuleShape.h"
 #include "SceneDefinitions.h"
-#include "MeshDrawCommandStatsComponentData.h"
+#include "MeshDrawCommandStatsDefines.h"
 
 class FLightSceneInfo;
 class FLightSceneProxy;
@@ -618,10 +618,6 @@ public:
 	inline int32 GetVirtualTextureCullMips() const { return VirtualTextureCullMips; }
 	inline int32 GetVirtualTextureMinCoverage() const {	return VirtualTextureMinCoverage; }
 
-#if MESH_DRAW_COMMAND_STAT_COLLECTION
-	inline FMeshDrawCommandStatsComponentDataID GetMeshDrawCommandStatsComponentDataID() const { return MeshDrawCommandStatsComponentDataID; }
-#endif
-
 	inline bool IsMovable() const 
 	{ 
 		// Note: primitives with EComponentMobility::Stationary can still move (as opposed to lights with EComponentMobility::Stationary)
@@ -1080,6 +1076,10 @@ public:
 	ENGINE_API virtual void SetDebugMassData(const TArray<FDebugMassData>& InDebugMassData);
 #endif
 
+#if MESH_DRAW_COMMAND_STATS
+	inline FName GetMeshDrawCommandStatsCategory() const { return MeshDrawCommandStatsCategory; }
+#endif
+
 	/**
 	 * Get the list of LCIs. Used to set the precomputed lighting uniform buffers, which can only be created by the RENDERER_API.
 	 */
@@ -1570,6 +1570,11 @@ private:
 	/** The name of the level the primitive is in. */
 	FName LevelName;
 
+#if MESH_DRAW_COMMAND_STATS
+	/** Category name for this primitive in the mesh draw stat collection. */
+	FName MeshDrawCommandStatsCategory;
+#endif
+
 #if WITH_EDITOR
 	/** A copy of the actor's group membership for handling per-view group hiding */
 	uint64 HiddenEditorViews;
@@ -1628,11 +1633,6 @@ protected:
 	/** The primitive's minimum cull distance. */
 	float MinDrawDistance;
 
-#if MESH_DRAW_COMMAND_STAT_COLLECTION
-	/** Unique id use to retrieve the component data when building the mesh data draw stat information */
-	FMeshDrawCommandStatsComponentDataID MeshDrawCommandStatsComponentDataID = INDEX_NONE;
-#endif
-
 	/**
 	 * Updates the primitive proxy's cached transforms for all instances given a buffer of instance updates.
 	 * @param CmdBuffer - A record of all the add, update and remove instances for the proxy to apply to its internal data.
diff --git a/Engine/Source/Runtime/Landscape/Private/Landscape.cpp b/Engine/Source/Runtime/Landscape/Private/Landscape.cpp
--- a/Engine/Source/Runtime/Landscape/Private/Landscape.cpp
+++ b/Engine/Source/Runtime/Landscape/Private/Landscape.cpp
@@ -1101,9 +1101,6 @@ void ULandscapeComponent::PostLoad()
 		}
 	}
 
-	static FName NAME_StatsCategoryLandscape("Landscape");
-	SetMeshDrawCommandStatsCategory(NAME_StatsCategoryLandscape);
-
 #if WITH_EDITOR
 	ALandscapeProxy* LandscapeProxy = GetLandscapeProxy();
 	if (ensure(LandscapeProxy))
diff --git a/Engine/Source/Runtime/Landscape/Private/LandscapeGrass.cpp b/Engine/Source/Runtime/Landscape/Private/LandscapeGrass.cpp
--- a/Engine/Source/Runtime/Landscape/Private/LandscapeGrass.cpp
+++ b/Engine/Source/Runtime/Landscape/Private/LandscapeGrass.cpp
@@ -3769,9 +3769,6 @@ void ALandscapeProxy::UpdateGrass(const TArray<FVector>& Cameras, int32& InOutNu
 										GrassInstancedStaticMeshComponent->WorldPositionOffsetDisableDistance = GrassVariety.InstanceWorldPositionOffsetDisableDistance;
 										GrassInstancedStaticMeshComponent->ShadowCacheInvalidationBehavior = GrassVariety.ShadowCacheInvalidationBehavior;
 
-										static FName NAME_StatsCategoryGrass("Grass");
-										GrassInstancedStaticMeshComponent->SetMeshDrawCommandStatsCategory(NAME_StatsCategoryGrass);
-										
 										GrassInstancedStaticMeshComponent->PrecachePSOs();
 
 										const FMeshMapBuildData* MeshMapBuildData = Component->GetMeshMapBuildData();
diff --git a/Engine/Source/Runtime/Renderer/Private/InstanceCulling/InstanceCullingContext.cpp b/Engine/Source/Runtime/Renderer/Private/InstanceCulling/InstanceCullingContext.cpp
--- a/Engine/Source/Runtime/Renderer/Private/InstanceCulling/InstanceCullingContext.cpp
+++ b/Engine/Source/Runtime/Renderer/Private/InstanceCulling/InstanceCullingContext.cpp
@@ -175,7 +175,7 @@ FInstanceCullingContext::FInstanceCullingContext(
 	BatchedPrimitiveSlot(GetUniformBufferViewStaticSlot(InShaderPlatform)),
 	bUsesUniformBufferView(PlatformGPUSceneUsesUniformBufferView(InShaderPlatform))
 {
-#if MESH_DRAW_COMMAND_STAT_COLLECTION
+#if MESH_DRAW_COMMAND_STATS
 	if (FMeshDrawCommandStatsManager* Instance = FMeshDrawCommandStatsManager::Get())
 	{
 		MeshDrawCommandPassStats = Instance->CreatePassStats(PassName);
@@ -620,7 +620,7 @@ public:
 
 	void ProcessBatched(TStaticArray<FBuildInstanceIdBufferAndCommandsFromPrimitiveIdsCs::FParameters*, static_cast<uint32>(EBatchProcessingMode::Num)> PassParameters);
 
-#if MESH_DRAW_COMMAND_STAT_COLLECTION
+#if MESH_DRAW_COMMAND_STATS
 	FRHIGPUBufferReadback* MeshDrawCommandStatsIndirectArgsReadbackBuffer = nullptr;
 #endif
 };
@@ -685,7 +685,7 @@ void FInstanceCullingContext::BuildRenderingCommandsInternal(
 	EAsyncProcessingMode AsyncProcessingMode,
 	FInstanceCullingDrawParams* InstanceCullingDrawParams)
 {
-#if MESH_DRAW_COMMAND_STAT_COLLECTION
+#if MESH_DRAW_COMMAND_STATS
 	if (MeshDrawCommandPassStats)
 	{
 		check(!MeshDrawCommandPassStats->bBuildRenderingCommandsCalled);
@@ -968,7 +968,7 @@ void FInstanceCullingContext::BuildRenderingCommandsInternal(
 		InstanceCullingDrawParams->BatchedPrimitive = GraphBuilder.CreateUniformBuffer(BatchedPrimitiveParameters);
 	}
 
-#if MESH_DRAW_COMMAND_STAT_COLLECTION
+#if MESH_DRAW_COMMAND_STATS
 	if (MeshDrawCommandPassStats)
 	{
 		FRHIGPUBufferReadback* GPUBufferReadback = FMeshDrawCommandStatsManager::Get()->QueueDrawRDGIndirectArgsReadback(GraphBuilder, DrawIndirectArgsRDG);
@@ -986,7 +986,7 @@ void FInstanceCullingDeferredContext::ProcessBatched(TStaticArray<FBuildInstance
 
 	MergeBatches();
 
-#if MESH_DRAW_COMMAND_STAT_COLLECTION
+#if MESH_DRAW_COMMAND_STATS
 	// Setup the indirect buffer and correct offset for each pass in the merged buffer
 	if (MeshDrawCommandStatsIndirectArgsReadbackBuffer)
 	{
@@ -999,7 +999,7 @@ void FInstanceCullingDeferredContext::ProcessBatched(TStaticArray<FBuildInstance
 			}
 		}
 	}
-#endif // MESH_DRAW_COMMAND_STAT_COLLECTION
+#endif // MESH_DRAW_COMMAND_STATS
 
 	bProcessed = true;
 
@@ -1327,7 +1327,7 @@ FInstanceCullingDeferredContext *FInstanceCullingContext::CreateDeferredContext(
 #undef INST_CULL_CALLBACK_MODE
 #undef INST_CULL_CREATE_STRUCT_BUFF_ARGS_MODE
 
-#if MESH_DRAW_COMMAND_STAT_COLLECTION
+#if MESH_DRAW_COMMAND_STATS
 	if (FMeshDrawCommandStatsManager* Instance = FMeshDrawCommandStatsManager::Get())
 	{
 		if (Instance->CollectStats())
@@ -1335,7 +1335,7 @@ FInstanceCullingDeferredContext *FInstanceCullingContext::CreateDeferredContext(
 			DeferredContext->MeshDrawCommandStatsIndirectArgsReadbackBuffer = Instance->QueueDrawRDGIndirectArgsReadback(GraphBuilder, DeferredContext->DrawIndirectArgsBuffer);;
 		}
 	}
-#endif // MESH_DRAW_COMMAND_STAT_COLLECTION
+#endif // MESH_DRAW_COMMAND_STATS
 
 	return DeferredContext;
 }
diff --git a/Engine/Source/Runtime/Renderer/Private/MeshDrawCommandStats.cpp b/Engine/Source/Runtime/Renderer/Private/MeshDrawCommandStats.cpp
--- a/Engine/Source/Runtime/Renderer/Private/MeshDrawCommandStats.cpp
+++ b/Engine/Source/Runtime/Renderer/Private/MeshDrawCommandStats.cpp
@@ -10,7 +10,7 @@
 #include "RHI.h"
 #include "RHIGPUReadback.h"
 
-#if MESH_DRAW_COMMAND_STAT_COLLECTION
+#if MESH_DRAW_COMMAND_STATS
 
 FMeshDrawCommandStatsManager* FMeshDrawCommandStatsManager::Instance = nullptr;
 
@@ -32,9 +32,11 @@ DECLARE_DWORD_COUNTER_STAT(TEXT("Custom Indirect Rendered Instances"), STAT_Cull
 static TAutoConsoleVariable<int32> CVarShowMeshDrawCommandStats(
 	TEXT("r.MeshDrawCommands.Stats"),
 	0,
-	TEXT("Show on screen mesh draw command stats per stat category.\n")
-	TEXT("The stats are accumulated across passes and are post-culling.\n")
-	TEXT("Use stat culling to see global culling stats.\n"),
+	TEXT("Show on screen mesh draw command stats.\n")
+	TEXT("The stats for visible triangles are post GPU culling.\n")
+	TEXT(" 1 = Show stats per category. The stats are accumulated across passes.\n")
+	TEXT(" 2 = Show stats per pass.\n")
+	TEXT("You can also use 'stat culling' to see global culling stats.\n"),
 	ECVF_RenderThreadSafe
 );
 
@@ -137,33 +139,54 @@ FMeshDrawCommandStatsManager::FMeshDrawCommandStatsManager()
 			const bool bShowStats = CVarShowMeshDrawCommandStats->GetInt() == 0 ? false : true;
 			if (bShowStats)
 			{
-				const UMeshDrawCommandStatsSettings* Settings = GetDefault<UMeshDrawCommandStatsSettings>();
-
 				OutMessages.Add(FCoreDelegates::EOnScreenMessageSeverity::Info, FText::FromString(FString::Printf(TEXT("MeshDrawCommandStats (Triangles / Budget - Category):"), Stats.TotalPrimitives / 1000)));
-				for (FStats::FCategoryStats const& CategoryStat : Stats.CategoryStats)
+
+				TArray<FStats::FCategoryStats> CategoryStats = Stats.CategoryStats;
+
+				// Show budgeted stats first.
+				const UMeshDrawCommandStatsSettings* Settings = GetDefault<UMeshDrawCommandStatsSettings>();
+				for (FMeshDrawCommandStatsBudget const& CategoryBudget : Settings->Budgets)
 				{
-					int32 Budget = 0;
-					if (Settings != nullptr)
-					for (FMeshDrawCommandStatsBudget const& CategoryBudget : Settings->Budgets)
+					uint64 PrimitiveCount = 0;
+
+					for (FStats::FCategoryStats& CategoryStat : CategoryStats)
 					{
-						if (CategoryBudget.CategoryName == CategoryStat.Category)
+						if (CategoryBudget.CategoryName == CategoryStat.CategoryName)
 						{
-							Budget = CategoryBudget.PrimitiveBudget;
-							break;
+							PrimitiveCount += CategoryStat.PrimitiveCount;
+							CategoryStat.PrimitiveCount = 0;
+						}
+						else
+						{
+							for (FName Name : CategoryBudget.LinkedStatNames)
+							{
+								if (Name == CategoryStat.CategoryName)
+								{
+									PrimitiveCount += CategoryStat.PrimitiveCount;
+									CategoryStat.PrimitiveCount = 0;
+									break;
+								}
+							}
 						}
 					}
-					
-					if (Budget > 0)
+
+					if (PrimitiveCount > 0)
 					{
-						FCoreDelegates::EOnScreenMessageSeverity Severity = Budget < CategoryStat.PrimitiveCount ? FCoreDelegates::EOnScreenMessageSeverity::Warning : FCoreDelegates::EOnScreenMessageSeverity::Info;
-						OutMessages.Add(Severity, FText::FromString(FString::Printf(TEXT("%5dK / %5dK - %s"), CategoryStat.PrimitiveCount / 1000, Budget / 1000, *(CategoryStat.Category.ToString()))));
+						FCoreDelegates::EOnScreenMessageSeverity Severity = CategoryBudget.PrimitiveBudget < PrimitiveCount ? FCoreDelegates::EOnScreenMessageSeverity::Warning : FCoreDelegates::EOnScreenMessageSeverity::Info;
+						OutMessages.Add(Severity, FText::FromString(FString::Printf(TEXT("%5dK / %5dK - %s"), PrimitiveCount / 1000, CategoryBudget.PrimitiveBudget / 1000, *(CategoryBudget.CategoryName.ToString()))));
 					}
-					else
+				}
+
+				// Show remaining (non-zeroed) stats.
+				for (FStats::FCategoryStats const& CategoryStat : CategoryStats)
+				{
+					if (CategoryStat.PrimitiveCount > 0)
 					{
-						OutMessages.Add(FCoreDelegates::EOnScreenMessageSeverity::Info, FText::FromString(FString::Printf(TEXT("\t%5dK - %s"), CategoryStat.PrimitiveCount / 1000, *(CategoryStat.Category.ToString()))));
+						OutMessages.Add(FCoreDelegates::EOnScreenMessageSeverity::Info, FText::FromString(FString::Printf(TEXT("\t%5dK - %s"), CategoryStat.PrimitiveCount / 1000, *(CategoryStat.CategoryName.ToString()))));
 					}
 				}
 
+				// Show total budget.
 				const int32 TotalBudget = Settings != 0 ? Settings->TotalPrimitiveBudget : 0;
 				if (TotalBudget > 0)
 				{
@@ -241,12 +264,12 @@ void FMeshDrawCommandStatsManager::Update()
 
 	++CurrentFrameNumber;
 
+	const bool bShowPassNameStats = CVarShowMeshDrawCommandStats->GetInt() == 2 ? true : false;
+
 	FScopeLock ScopeLock(&FrameDataCS);
 
 	bool bHasProcessedFrame = false;
 
-	FMeshDrawCommandStatsComponentDataManager const* ComponentDataManager = FMeshDrawCommandStatsComponentDataManager::Get();
-
 	// TODO: might be more than one from a given frame. E.g., if it was using a scene capture, need to filter out those, or perhaps record them as a group actually.
 	for (int32 Index = Frames.Num() - 1; Index >= 0; --Index)
 	{
@@ -314,10 +337,8 @@ void FMeshDrawCommandStatsManager::Update()
 						Stats.TotalInstances += DrawData.VisibleInstanceCount;
 						Stats.TotalPrimitives += DrawData.VisibleInstanceCount * DrawData.PrimitiveCount;
 
-						FMeshDrawCommandStatsComponentData ComponentData = ComponentDataManager->GetComponentData(DrawData.StatsData.ComponentDataID);
-						static FName NAME_Unknown("Unknown Category");
-						FName StatsCategory = ComponentData.StatsCategory.IsNone() ? NAME_Unknown : ComponentData.StatsCategory;
-						uint64& TotalCount = CategoryStats.FindOrAdd(StatsCategory);
+						FName StatName = bShowPassNameStats ? PassStats->PassName : DrawData.StatsData.CategoryName;
+						uint64& TotalCount = CategoryStats.FindOrAdd(StatName);
 						TotalCount += DrawData.VisibleInstanceCount * DrawData.PrimitiveCount;
 					}
 
@@ -338,7 +359,7 @@ void FMeshDrawCommandStatsManager::Update()
 				{
 					Stats.CategoryStats.Add(FStats::FCategoryStats(Iter.Key(), Iter.Value()));
 				}
-				Algo::Sort(Stats.CategoryStats, [this](FStats::FCategoryStats& LHS, FStats::FCategoryStats& RHS) { return LHS.Category.ToString() < RHS.Category.ToString(); });
+				Algo::Sort(Stats.CategoryStats, [this](FStats::FCategoryStats& LHS, FStats::FCategoryStats& RHS) { return LHS.CategoryName.ToString() < RHS.CategoryName.ToString(); });
 
 				// Got new stats, so can dump them if requested
 				static bool bDumpStats = false;
@@ -373,8 +394,6 @@ void FMeshDrawCommandStatsManager::Update()
 	bCollectStats = bShowStats || bRequestDumpStats;
 }
 
-#if (!UE_BUILD_SHIPPING && !UE_BUILD_TEST)
-
 void FMeshDrawCommandStatsManager::DumpStats(FFrameData* FrameData)
 {
 	const FString Filename = FString::Printf(TEXT("%sMeshDrawCommandStats-%s.csv"), *FPaths::ProfilingDir(), *FDateTime::Now().ToString());
@@ -389,8 +408,7 @@ void FMeshDrawCommandStatsManager::DumpStats(FFrameData* FrameData)
 		FName PassName;
 		int32 VisibilePrimitiveCount;
 		int32 VisibleInstance;		
-		FName Category;
-		FName ComponentType;
+		FName CategoryName;
 		FName ResourceName;
 		int32 LODIndex;
 		int32 SegmentIndex;
@@ -401,7 +419,6 @@ void FMeshDrawCommandStatsManager::DumpStats(FFrameData* FrameData)
 	};
 	TArray<FStatEntry> StatEntries;
 
-	FMeshDrawCommandStatsComponentDataManager const* ComponentDataManager = FMeshDrawCommandStatsComponentDataManager::Get();
 	for (FMeshDrawCommandPassStats* PassStats : FrameData->PassData)
 	{
 		for (FVisibleMeshDrawCommandStatsData& DrawData : PassStats->DrawData)
@@ -412,17 +429,14 @@ void FMeshDrawCommandStatsManager::DumpStats(FFrameData* FrameData)
 				StatEntry.PassName = PassStats->PassName;
 				StatEntry.VisibilePrimitiveCount = DrawData.VisibleInstanceCount * DrawData.PrimitiveCount;
 				StatEntry.VisibleInstance = DrawData.VisibleInstanceCount;
-				StatEntry.LODIndex = DrawData.StatsData.LODIndex;
-				StatEntry.SegmentIndex = DrawData.StatsData.SegmentIndex;
+				StatEntry.LODIndex = DrawData.LODIndex;
+				StatEntry.SegmentIndex = DrawData.SegmentIndex;
 				StatEntry.PrimitiveCount = DrawData.PrimitiveCount;
 				StatEntry.TotalInstanceCount = DrawData.TotalInstanceCount;
 				StatEntry.TotalPrimitiveCount = DrawData.TotalInstanceCount * DrawData.PrimitiveCount;
 				StatEntry.ResourceName = DrawData.ResourceName;
 				StatEntry.MaterialName = DrawData.MaterialName;
-
-				FMeshDrawCommandStatsComponentData ComponentData = ComponentDataManager->GetComponentData(DrawData.StatsData.ComponentDataID);
-				StatEntry.Category = ComponentData.StatsCategory;
-				StatEntry.ComponentType = ComponentData.ComponentType;
+				StatEntry.CategoryName = DrawData.StatsData.CategoryName;
 			}
 		}
 	}
@@ -439,27 +453,24 @@ void FMeshDrawCommandStatsManager::DumpStats(FFrameData* FrameData)
 			return LHS.VisibilePrimitiveCount > RHS.VisibilePrimitiveCount;
 		});
 
-	const TCHAR* Header = TEXT("Pass,VisiblePrimitiveCount,VisibleInstances,Category,ComponentType,ResourceName,LODIndex,SegmentIndex,MaterialName,PrimitiveCount,TotalInstanceCount,TotalPrimitiveCount\n");
+	const TCHAR* Header = TEXT("Pass,VisiblePrimitiveCount,VisibleInstances,Category,ResourceName,LODIndex,SegmentIndex,MaterialName,PrimitiveCount,TotalInstanceCount,TotalPrimitiveCount\n");
 	CSVFile->Serialize(TCHAR_TO_ANSI(Header), FPlatformString::Strlen(Header));
 
 	TCHAR PassNameBuffer[FName::StringBufferSize];
 	TCHAR ResourceNameBuffer[FName::StringBufferSize];
-	TCHAR ComponentTypeNameBuffer[FName::StringBufferSize];
 	TCHAR CategoryBuffer[FName::StringBufferSize];	
 
 	for (FStatEntry& StatEntry : StatEntries)
 	{
 		StatEntry.PassName.ToString(PassNameBuffer);
-		StatEntry.Category.ToString(CategoryBuffer);
-		StatEntry.ComponentType.ToString(ComponentTypeNameBuffer);
+		StatEntry.CategoryName.ToString(CategoryBuffer);
 		StatEntry.ResourceName.ToString(ResourceNameBuffer);
 
-		FString Row = FString::Printf(TEXT("%s,%d,%d,%s,%s,%s,%d,%d,%s,%d,%d,%d\n"),
+		FString Row = FString::Printf(TEXT("%s,%d,%d,%s,%s,%d,%d,%s,%d,%d,%d\n"),
 			PassNameBuffer,
 			StatEntry.VisibilePrimitiveCount,
 			StatEntry.VisibleInstance,
 			CategoryBuffer,
-			ComponentTypeNameBuffer,
 			ResourceNameBuffer,
 			StatEntry.LODIndex,
 			StatEntry.SegmentIndex,
@@ -485,12 +496,4 @@ static FAutoConsoleCommand GDumpMeshDrawCommandStatsCmd(
 	}
 }));
 
-#else
-
-void FMeshDrawCommandStatsManager::DumpStats(FFrameData* FrameData)
-{
-}
-
-#endif // (!UE_BUILD_SHIPPING && !UE_BUILD_TEST)
-
-#endif  // MESH_DRAW_COMMAND_STAT_COLLECTION
\ No newline at end of file
+#endif  // MESH_DRAW_COMMAND_STATS
\ No newline at end of file
diff --git a/Engine/Source/Runtime/Renderer/Private/MeshDrawCommandStats.h b/Engine/Source/Runtime/Renderer/Private/MeshDrawCommandStats.h
--- a/Engine/Source/Runtime/Renderer/Private/MeshDrawCommandStats.h
+++ b/Engine/Source/Runtime/Renderer/Private/MeshDrawCommandStats.h
@@ -2,12 +2,11 @@
 #pragma once
 
 #include "MeshPassProcessor.h"
-#include "Misc/ScopeRWLock.h"
 
 class FRDGBuilder;
 class FRHIGPUBufferReadback;
 
-#if MESH_DRAW_COMMAND_STAT_COLLECTION
+#if MESH_DRAW_COMMAND_STATS
 
 /**
  * Contains all the required data per mesh draw command which is needed for stat collection - cached locally because lifetime of MDC is unknown by the time indirect args are resolved
@@ -26,6 +25,8 @@ struct FVisibleMeshDrawCommandStatsData
 	FRHIBuffer* CustomIndirectArgsBuffer = nullptr;		//< Optional custom indirect arg buffer which was provided to the MDC at draw time
 
 #if MESH_DRAW_COMMAND_DEBUG_DATA
+	int32 LODIndex = 0;									//< LOD index in draw command.
+	int32 SegmentIndex = 0;								//< Segment index in draw command.
 	FName ResourceName;									//< Minimal resource name 
 	FString MaterialName;								//< Material name used during draw event
 #endif // MESH_DRAW_COMMAND_DEBUG_DATA
@@ -141,8 +142,8 @@ private:
 		 */
 		struct FCategoryStats
 		{
-			FCategoryStats(FName InCategory, uint64 InPrimitiveCount) : Category(InCategory), PrimitiveCount(InPrimitiveCount) {}
-			FName Category;
+			FCategoryStats(FName InCategoryName, uint64 InPrimitiveCount) : CategoryName(InCategoryName), PrimitiveCount(InPrimitiveCount) {}
+			FName CategoryName;
 			uint64 PrimitiveCount;
 		};
 		TArray<FCategoryStats> CategoryStats;
@@ -178,4 +179,4 @@ private:
 	static FMeshDrawCommandStatsManager* Instance;
 };
 
-#endif // MESH_DRAW_COMMAND_STAT_COLLECTION
\ No newline at end of file
+#endif // MESH_DRAW_COMMAND_STATS
\ No newline at end of file
diff --git a/Engine/Source/Runtime/Renderer/Private/MeshDrawCommands.cpp b/Engine/Source/Runtime/Renderer/Private/MeshDrawCommands.cpp
--- a/Engine/Source/Runtime/Renderer/Private/MeshDrawCommands.cpp
+++ b/Engine/Source/Runtime/Renderer/Private/MeshDrawCommands.cpp
@@ -866,7 +866,7 @@ void CollectMeshDrawCommandPassStats(
 	FMeshCommandOneFrameArray& VisibleMeshDrawCommands,
 	FInstanceCullingContext& InstanceCullingContext)
 {	
-#if MESH_DRAW_COMMAND_STAT_COLLECTION
+#if MESH_DRAW_COMMAND_STATS
 	FMeshDrawCommandPassStats* PassStats = InstanceCullingContext.MeshDrawCommandPassStats;
 	if (PassStats == nullptr)
 	{
@@ -941,7 +941,7 @@ void CollectMeshDrawCommandPassStats(
 			DrawData.TotalInstanceCount = MeshDrawCommand->NumInstances * NumInstances;
 		}
 	}
-#endif // MESH_DRAW_COMMAND_STAT_COLLECTION
+#endif // MESH_DRAW_COMMAND_STATS
 }
 
 FAutoConsoleTaskPriority CPrio_FMeshDrawCommandPassSetupTask(
diff --git a/Engine/Source/Runtime/Renderer/Private/MeshPassProcessor.cpp b/Engine/Source/Runtime/Renderer/Private/MeshPassProcessor.cpp
--- a/Engine/Source/Runtime/Renderer/Private/MeshPassProcessor.cpp
+++ b/Engine/Source/Runtime/Renderer/Private/MeshPassProcessor.cpp
@@ -1699,7 +1699,7 @@ uint64 FMeshDrawCommand::GetPipelineStateSortingKey(const FGraphicsPipelineRende
 }
 
 #if MESH_DRAW_COMMAND_DEBUG_DATA
-void FMeshDrawCommand::SetDebugData(const FPrimitiveSceneProxy* PrimitiveSceneProxy, const FMaterial* Material, const FMaterialRenderProxy* MaterialRenderProxy, const FMeshProcessorShaders& UntypedShaders, const FVertexFactory* VertexFactory, uint32 MeshPassType)
+void FMeshDrawCommand::SetDebugData(const FPrimitiveSceneProxy* PrimitiveSceneProxy, const FMaterial* Material, const FMaterialRenderProxy* MaterialRenderProxy, const FMeshProcessorShaders& UntypedShaders, const FVertexFactory* VertexFactory, const FMeshBatch& MeshBatch, uint32 MeshPassType)
 {
 	DebugData.PrimitiveSceneProxyIfNotUsingStateBuckets = PrimitiveSceneProxy;
 	DebugData.MaterialRenderProxy = MaterialRenderProxy;
@@ -1707,18 +1707,18 @@ void FMeshDrawCommand::SetDebugData(const FPrimitiveSceneProxy* PrimitiveScenePr
 	DebugData.PixelShader = UntypedShaders.PixelShader;
 	DebugData.VertexFactory = VertexFactory;
 	DebugData.VertexFactoryType = VertexFactory->GetType();
+	DebugData.LODIndex = MeshBatch.LODIndex;
+	DebugData.SegmentIndex = MeshBatch.SegmentIndex;
 	DebugData.MeshPassType = MeshPassType;
 	DebugData.ResourceName =  PrimitiveSceneProxy ? PrimitiveSceneProxy->GetResourceName() : FName();
 	DebugData.MaterialName = MaterialRenderProxy->GetMaterialName();
 }
 #endif
 
-#if MESH_DRAW_COMMAND_STAT_COLLECTION
-void FMeshDrawCommand::SetStatsData(const FPrimitiveSceneProxy* PrimitiveSceneProxy, const FMeshBatch& MeshBatch, int32 BatchElementIndex)
+#if MESH_DRAW_COMMAND_STATS
+void FMeshDrawCommand::SetStatsData(const FPrimitiveSceneProxy* PrimitiveSceneProxy)
 {
-	StatsData.ComponentDataID = PrimitiveSceneProxy ? PrimitiveSceneProxy->GetMeshDrawCommandStatsComponentDataID() : INDEX_NONE;
-	StatsData.LODIndex = MeshBatch.LODIndex;
-	StatsData.SegmentIndex = MeshBatch.SegmentIndex;
+	StatsData.CategoryName = PrimitiveSceneProxy ? PrimitiveSceneProxy->GetMeshDrawCommandStatsCategory() : FName();
 }
 
 void FMeshDrawCommand::GetStatsData(FVisibleMeshDrawCommandStatsData& OutVisibleStatsData) const
@@ -1726,11 +1726,13 @@ void FMeshDrawCommand::GetStatsData(FVisibleMeshDrawCommandStatsData& OutVisible
 	OutVisibleStatsData.StatsData = StatsData;
 	OutVisibleStatsData.PrimitiveCount = NumPrimitives;
 #if MESH_DRAW_COMMAND_DEBUG_DATA
+	OutVisibleStatsData.LODIndex = DebugData.LODIndex;
+	OutVisibleStatsData.SegmentIndex = DebugData.SegmentIndex;
 	OutVisibleStatsData.ResourceName = DebugData.ResourceName;
 	OutVisibleStatsData.MaterialName = DebugData.MaterialName;
 #endif
 }
-#endif // MESH_DRAW_COMMAND_STAT_COLLECTION
+#endif // MESH_DRAW_COMMAND_STATS
 
 void SubmitMeshDrawCommands(
 	const FMeshCommandOneFrameArray& VisibleMeshDrawCommands,
diff --git a/Engine/Source/Runtime/Renderer/Private/Renderer.cpp b/Engine/Source/Runtime/Renderer/Private/Renderer.cpp
--- a/Engine/Source/Runtime/Renderer/Private/Renderer.cpp
+++ b/Engine/Source/Runtime/Renderer/Private/Renderer.cpp
@@ -63,7 +63,7 @@ FAutoConsoleVariableRef CVarFlushRenderTargetsOnWorldCleanup(TEXT("r.bFlushRende
 
 void FRendererModule::StartupModule()
 {
-#if MESH_DRAW_COMMAND_STAT_COLLECTION
+#if MESH_DRAW_COMMAND_STATS
 	FMeshDrawCommandStatsManager::CreateInstance();
 #endif
 
diff --git a/Engine/Source/Runtime/Renderer/Private/SceneRendering.cpp b/Engine/Source/Runtime/Renderer/Private/SceneRendering.cpp
--- a/Engine/Source/Runtime/Renderer/Private/SceneRendering.cpp
+++ b/Engine/Source/Runtime/Renderer/Private/SceneRendering.cpp
@@ -4556,7 +4556,7 @@ static void RenderViewFamilies_RenderThread(FRHICommandListImmediate& RHICmdList
 
 		GraphBuilder.Execute();
 
-#if MESH_DRAW_COMMAND_STAT_COLLECTION
+#if MESH_DRAW_COMMAND_STATS
 		// Needs to happen after graph builder execute to make sure all FMeshDrawCommandPassSetupTask are finished.
 		// Only during iteration of the visible MDCs we know the custom indirect args used in the frame
 		if (FMeshDrawCommandStatsManager* Instance = FMeshDrawCommandStatsManager::Get())
diff --git a/Engine/Source/Runtime/Renderer/Public/InstanceCulling/InstanceCullingContext.h b/Engine/Source/Runtime/Renderer/Public/InstanceCulling/InstanceCullingContext.h
--- a/Engine/Source/Runtime/Renderer/Public/InstanceCulling/InstanceCullingContext.h
+++ b/Engine/Source/Runtime/Renderer/Public/InstanceCulling/InstanceCullingContext.h
@@ -3,9 +3,9 @@
 #pragma once
 
 #include "CoreMinimal.h"
+#include "MeshPassProcessor.h"
 #include "RHI.h"
 #include "RenderGraphResources.h"
-#include "MeshPassProcessor.h"
 
 struct FInstanceCullingResult;
 class FGPUScene;
@@ -350,11 +350,11 @@ public:
 	// Whether current platform uses Uniform Buffer View path
 	bool bUsesUniformBufferView;
 	
-#if MESH_DRAW_COMMAND_STAT_COLLECTION
+#if MESH_DRAW_COMMAND_STATS
 public:
 	// Optional pass stats
 	FMeshDrawCommandPassStats* MeshDrawCommandPassStats = nullptr;
-#endif // MESH_DRAW_COMMAND_STAT_COLLECTION
+#endif
 };
 
 ENUM_CLASS_FLAGS(FInstanceCullingContext::EInstanceFlags)
diff --git a/Engine/Source/Runtime/Renderer/Public/MeshPassProcessor.h b/Engine/Source/Runtime/Renderer/Public/MeshPassProcessor.h
--- a/Engine/Source/Runtime/Renderer/Public/MeshPassProcessor.h
+++ b/Engine/Source/Runtime/Renderer/Public/MeshPassProcessor.h
@@ -10,8 +10,8 @@
 #include "MeshMaterialShader.h"
 #include "SceneUtils.h"
 #include "MeshBatch.h"
+#include "MeshDrawCommandStatsDefines.h"
 #include "PSOPrecache.h"
-#include "MeshDrawCommandStatsComponentData.h"
 #include "Hash/CityHash.h"
 #include "Experimental/Containers/RobinHoodHashTable.h"
 #include "RHIImmutableSamplerState.h"
@@ -717,23 +717,22 @@ struct FMeshDrawCommandDebugData
 	TShaderRef<FShader> PixelShader;
 	const FVertexFactory* VertexFactory;
 	const FVertexFactoryType* VertexFactoryType;
+	int8 LODIndex;
+	uint8 SegmentIndex;
 	uint32 MeshPassType;
 	FName ResourceName;
 	FString MaterialName;
 #endif
 };
 
-#if MESH_DRAW_COMMAND_STAT_COLLECTION
+/** Data needed to resolve mesh draw command stats. */
 struct FMeshDrawCommandStatsData
 {
+#if MESH_DRAW_COMMAND_STATS
 	/** ID used to retrieve the component data during stat collection. */
-	FMeshDrawCommandStatsComponentDataID ComponentDataID;
-	/** LOD index of the MDC used only during full stat dumping. */
-	int8 LODIndex;
-	/** Segment index of the MDC used only during full stat dumping. */
-	uint8 SegmentIndex;
+	FName CategoryName;
+#endif
 };
-#endif // MESH_DRAW_COMMAND_STAT_COLLECTION
 
 class FMeshDrawCommandStateCache
 {
@@ -1264,9 +1263,9 @@ public:
 		return Other.CachedPipelineId.GetId();
 	}
 #if MESH_DRAW_COMMAND_DEBUG_DATA
-	RENDERER_API void SetDebugData(const FPrimitiveSceneProxy* PrimitiveSceneProxy, const FMaterial* Material, const FMaterialRenderProxy* MaterialRenderProxy, const FMeshProcessorShaders& UntypedShaders, const FVertexFactory* VertexFactory, uint32 MeshPassType);
+	RENDERER_API void SetDebugData(const FPrimitiveSceneProxy* PrimitiveSceneProxy, const FMaterial* Material, const FMaterialRenderProxy* MaterialRenderProxy, const FMeshProcessorShaders& UntypedShaders, const FVertexFactory* VertexFactory, const FMeshBatch& MeshBatch, uint32 MeshPassType);
 #else
-	void SetDebugData(const FPrimitiveSceneProxy* PrimitiveSceneProxy, const FMaterial* Material, const FMaterialRenderProxy* MaterialRenderProxy, const FMeshProcessorShaders& UntypedShaders, const FVertexFactory* VertexFactory, uint32 MeshPassType){}
+	void SetDebugData(const FPrimitiveSceneProxy* PrimitiveSceneProxy, const FMaterial* Material, const FMaterialRenderProxy* MaterialRenderProxy, const FMeshProcessorShaders& UntypedShaders, const FVertexFactory* VertexFactory, const FMeshBatch& MeshBatch, uint32 MeshPassType) {}
 #endif
 
 	SIZE_T GetAllocatedSize() const
@@ -1300,15 +1299,16 @@ public:
 	};
 #endif
 
-#if !MESH_DRAW_COMMAND_STAT_COLLECTION
-	RENDERER_API void SetStatsData(const FPrimitiveSceneProxy* PrimitiveSceneProxy, const FMeshBatch& MeshBatch, int32 BatchElementIndex) {}
-#else
-public:
-	RENDERER_API void SetStatsData(const FPrimitiveSceneProxy* PrimitiveSceneProxy, const FMeshBatch& MeshBatch, int32 BatchElementIndex);
-	void GetStatsData(FVisibleMeshDrawCommandStatsData& OutVisibleStatsData) const;
+#if MESH_DRAW_COMMAND_STATS
 private:
 	FMeshDrawCommandStatsData StatsData;
-#endif //!MESH_DRAW_COMMAND_STAT_COLLECTION
+public:
+	RENDERER_API void SetStatsData(const FPrimitiveSceneProxy* PrimitiveSceneProxy);
+	void GetStatsData(FVisibleMeshDrawCommandStatsData& OutVisibleStatsData) const;
+#else
+public:
+	void SetStatsData(const FPrimitiveSceneProxy* PrimitiveSceneProxy) {}
+#endif
 };
 
 /** FVisibleMeshDrawCommand sort key. */
diff --git a/Engine/Source/Runtime/Renderer/Public/MeshPassProcessor.inl b/Engine/Source/Runtime/Renderer/Public/MeshPassProcessor.inl
--- a/Engine/Source/Runtime/Renderer/Public/MeshPassProcessor.inl
+++ b/Engine/Source/Runtime/Renderer/Public/MeshPassProcessor.inl
@@ -148,7 +148,8 @@ void FMeshPassProcessor::BuildMeshDrawCommands(
 		PassShaders.GeometryShader->GetShaderBindings(Scene, FeatureLevel, PrimitiveSceneProxy, MaterialRenderProxy, MaterialResource, DrawRenderState, ShaderElementData, ShaderBindings);
 	}
 
-	SharedMeshDrawCommand.SetDebugData(PrimitiveSceneProxy, &MaterialResource, &MaterialRenderProxy, PassShaders.GetUntypedShaders(), VertexFactory, (uint32)MeshPassType);
+	SharedMeshDrawCommand.SetDebugData(PrimitiveSceneProxy, &MaterialResource, &MaterialRenderProxy, PassShaders.GetUntypedShaders(), VertexFactory, MeshBatch, (uint32)MeshPassType);
+	SharedMeshDrawCommand.SetStatsData(PrimitiveSceneProxy);
 
 	const int32 NumElements = ShouldSkipMeshDrawCommand(MeshBatch, PrimitiveSceneProxy) ? 0 : MeshBatch.Elements.Num();
 
@@ -159,8 +160,6 @@ void FMeshPassProcessor::BuildMeshDrawCommands(
 			const FMeshBatchElement& BatchElement = MeshBatch.Elements[BatchElementIndex];
 			FMeshDrawCommand& MeshDrawCommand = DrawListContext->AddCommand(SharedMeshDrawCommand, NumElements);
 			
-			MeshDrawCommand.SetStatsData(PrimitiveSceneProxy, MeshBatch, BatchElementIndex);
-			
 			EFVisibleMeshDrawCommandFlags Flags = SharedFlags;
 			if (BatchElement.bForceInstanceCulling)
 			{
