Fix Metal crash when accessing PF_FloatR11G11B10 as RWTexture

diff --git a/Engine/Source/Runtime/Renderer/Private/HeterogeneousVolumes/HeterogeneousVolumes.cpp b/Engine/Source/Runtime/Renderer/Private/HeterogeneousVolumes/HeterogeneousVolumes.cpp
--- a/Engine/Source/Runtime/Renderer/Private/HeterogeneousVolumes/HeterogeneousVolumes.cpp
+++ b/Engine/Source/Runtime/Renderer/Private/HeterogeneousVolumes/HeterogeneousVolumes.cpp
@@ -507,7 +507,7 @@ void FDeferredShadingSceneRenderer::RenderHeterogeneousVolumes(
 						uint32 NumMips = FMath::Log2(float(FMath::Min(FMath::Min(LightingCacheResolution.X, LightingCacheResolution.Y), LightingCacheResolution.Z))) + 1;
 						FRDGTextureDesc LightingCacheDesc = FRDGTextureDesc::Create3D(
 							LightingCacheResolution,
-							PF_FloatR11G11B10,
+                            !IsMetalPlatform(GShaderPlatformForFeatureLevel[View.FeatureLevel]) ? PF_FloatR11G11B10 : PF_FloatRGBA,
 							FClearValueBinding::Black,
 							TexCreate_ShaderResource | TexCreate_UAV | TexCreate_3DTiling,
 							NumMips
