Disable occlusion test and waiting for occlusion results when scene capture depth optimization is on

diff --git a/Engine/Source/Runtime/Renderer/Private/DeferredShadingRenderer.cpp b/Engine/Source/Runtime/Renderer/Private/DeferredShadingRenderer.cpp
--- a/Engine/Source/Runtime/Renderer/Private/DeferredShadingRenderer.cpp
+++ b/Engine/Source/Runtime/Renderer/Private/DeferredShadingRenderer.cpp
@@ -2956,7 +2956,7 @@ void FDeferredShadingSceneRenderer::Render(FRDGBuilder& GraphBuilder)
 	#endif
 
 	
-	const bool bIsOcclusionTesting = DoOcclusionQueries() && !ViewFamily.EngineShowFlags.DisableOcclusionQueries
+	const bool bIsOcclusionTesting = DoOcclusionQueries()
 		&& (!ViewFamily.EngineShowFlags.Wireframe || bIsViewFrozen);
 	const bool bNeedsPrePass = ShouldRenderPrePass();
 
diff --git a/Engine/Source/Runtime/Renderer/Private/SceneCaptureRendering.cpp b/Engine/Source/Runtime/Renderer/Private/SceneCaptureRendering.cpp
--- a/Engine/Source/Runtime/Renderer/Private/SceneCaptureRendering.cpp
+++ b/Engine/Source/Runtime/Renderer/Private/SceneCaptureRendering.cpp
@@ -784,6 +784,9 @@ static FSceneRenderer* CreateSceneRendererForSceneCapture(
 		ViewActor,
 		CubemapFaceIndex);
 
+	// Scene capture source is used to determine whether to disable occlusion queries inside FSceneRenderer constructor
+	ViewFamily.SceneCaptureSource = SceneCaptureComponent->CaptureSource;
+
 	// Screen percentage is still not supported in scene capture.
 	ViewFamily.EngineShowFlags.ScreenPercentage = false;
 	ViewFamily.SetScreenPercentageInterface(new FLegacyScreenPercentageDriver(
@@ -873,7 +876,6 @@ void FScene::UpdateSceneCaptureContents(USceneCaptureComponent2D* CaptureCompone
 
 		SceneRenderer->Views[0].bFogOnlyOnRenderedOpaque = CaptureComponent->bConsiderUnrenderedOpaquePixelAsFullyTranslucent;
 
-		SceneRenderer->ViewFamily.SceneCaptureSource = CaptureComponent->CaptureSource;
 		SceneRenderer->ViewFamily.SceneCaptureCompositeMode = CaptureComponent->CompositeMode;
 
 		// Need view state interface to be allocated for Lumen, as it requires persistent data.  This means
@@ -1150,8 +1152,6 @@ void FScene::UpdateSceneCaptureContents(USceneCaptureComponentCube* CaptureCompo
 				Location, ProjectionMatrix, false, CaptureComponent->MaxViewDistanceOverride,
 				bCaptureSceneColor, &PostProcessSettings, 0, CaptureComponent->GetViewOwner(), faceidx);
 
-			SceneRenderer->ViewFamily.SceneCaptureSource = CaptureComponent->CaptureSource;
-
 			for (const FSceneViewExtensionRef& Extension : SceneRenderer->ViewFamily.ViewExtensions)
 			{
 				Extension->SetupViewFamily(SceneRenderer->ViewFamily);
diff --git a/Engine/Source/Runtime/Renderer/Private/SceneRendering.cpp b/Engine/Source/Runtime/Renderer/Private/SceneRendering.cpp
--- a/Engine/Source/Runtime/Renderer/Private/SceneRendering.cpp
+++ b/Engine/Source/Runtime/Renderer/Private/SceneRendering.cpp
@@ -2648,6 +2648,12 @@ FSceneRenderer::FSceneRenderer(const FSceneViewFamily* InViewFamily, FHitProxyCo
 	ensureMsgf(!(ViewFamily.EngineShowFlags.ScreenPercentage && !ViewFamily.SupportsScreenPercentage()),
 		TEXT("Screen percentage is not supported, but show flag was incorectly set to true."));
 
+	// Disable occlusion queries for scene capture depth optimization mode
+	if (GetRendererOutput() == ERendererOutput::DepthPrepassOnly)
+	{
+		ViewFamily.EngineShowFlags.SetDisableOcclusionQueries(true);
+	}
+
 	// Fork the plugin interfaces of the view family.
 	{
 		{
diff --git a/Engine/Source/Runtime/Renderer/Private/SceneVisibility.cpp b/Engine/Source/Runtime/Renderer/Private/SceneVisibility.cpp
--- a/Engine/Source/Runtime/Renderer/Private/SceneVisibility.cpp
+++ b/Engine/Source/Runtime/Renderer/Private/SceneVisibility.cpp
@@ -310,7 +310,7 @@ static FAutoConsoleVariableRef CVarOcclusionCullEnable(
 
 bool FSceneRenderer::DoOcclusionQueries() const
 {
-	return GOcclusionCullEnabled;
+	return GOcclusionCullEnabled && !ViewFamily.EngineShowFlags.DisableOcclusionQueries;
 }
 
 static int32 GRelevanceNumPrimitivesPerPacket = 0;
