MeshDrawCommandStats: Convert some check() to ensure().
Fix incorrect stat names.

diff --git a/Engine/Source/Runtime/Renderer/Private/MeshDrawCommandStats.cpp b/Engine/Source/Runtime/Renderer/Private/MeshDrawCommandStats.cpp
--- a/Engine/Source/Runtime/Renderer/Private/MeshDrawCommandStats.cpp
+++ b/Engine/Source/Runtime/Renderer/Private/MeshDrawCommandStats.cpp
@@ -91,7 +91,7 @@ void FMeshDrawCommandStatsManager::FFrameData::Validate() const
 
 				if (DrawData.CustomIndirectArgsBuffer)
 				{
-					check(PassStats->CustomIndirectArgsBuffers.Contains(DrawData.CustomIndirectArgsBuffer));
+					ensure(PassStats->CustomIndirectArgsBuffers.Contains(DrawData.CustomIndirectArgsBuffer));
 					bHasIndirectArgs = true;
 				}
 			}
@@ -309,11 +309,10 @@ void FMeshDrawCommandStatsManager::Update()
 						int32 IndirectCommandIndex = DrawData.IndirectArgsOffset / (FInstanceCullingContext::IndirectArgsNumWords * sizeof(uint32));
 						if (DrawData.CustomIndirectArgsBuffer)
 						{
-							check(DrawData.PrimitiveCount == 0);
+							ensure(DrawData.PrimitiveCount == 0);
 
 							FIndirectArgsBufferResult* IndirectArgsBufferResult = FrameData->CustomIndirectArgsBufferResults.Find(DrawData.CustomIndirectArgsBuffer);
-							check(IndirectArgsBufferResult);
-							if (IndirectArgsBufferResult)
+							if (ensure(IndirectArgsBufferResult))
 							{
 								const FRHIDrawIndexedIndirectParameters& IndirectArgs = IndirectArgsBufferResult->DrawIndexedIndirectParameters[IndirectCommandIndex];
 								DrawData.PrimitiveCount = IndirectArgs.IndexCountPerInstance / 3; //< Assume triangles here for now - primitive count is empty so can't be used
@@ -327,9 +326,9 @@ void FMeshDrawCommandStatsManager::Update()
 						else if (DrawData.UseInstantCullingIndirectBuffer > 0 && InstanceCullingReadBackData)
 						{
 							const FRHIDrawIndexedIndirectParameters& IndirectArgs = IndirectArgsPtr[PassStats->IndirectArgParameterOffset + IndirectCommandIndex];
-							check(DrawData.PrimitiveCount == IndirectArgs.IndexCountPerInstance / 3);
+							ensure(DrawData.PrimitiveCount == IndirectArgs.IndexCountPerInstance / 3);
 							DrawData.VisibleInstanceCount = IndirectArgs.InstanceCount;
-							check(DrawData.VisibleInstanceCount <= DrawData.TotalInstanceCount);
+							ensure(DrawData.VisibleInstanceCount <= DrawData.TotalInstanceCount);
 							Stats.InstanceCullingIndirectInstances += DrawData.VisibleInstanceCount;
 							Stats.InstanceCullingIndirectPrimitives += DrawData.VisibleInstanceCount * DrawData.PrimitiveCount;
 						}
