New DF atlas upload buffer shrink strategy. It showed memory savings on most platforms with no noticeable perf impact.

#tests replay a/b tested

diff --git a/Engine/Source/Runtime/Renderer/Private/DistanceFieldStreaming.cpp b/Engine/Source/Runtime/Renderer/Private/DistanceFieldStreaming.cpp
--- a/Engine/Source/Runtime/Renderer/Private/DistanceFieldStreaming.cpp
+++ b/Engine/Source/Runtime/Renderer/Private/DistanceFieldStreaming.cpp
@@ -353,7 +353,7 @@ public:
 		const uint32 BrickDataNumBytesPerElement = GPixelFormats[DistanceField::DistanceFieldFormat].BlockBytes;
 
 		if (BrickUploadDataBuffer.NumBytes < NumBrickDataElements * BrickDataNumBytesPerElement 
-			|| (BrickUploadDataBuffer.NumBytes > NumBrickDataElements * BrickDataNumBytesPerElement && BrickUploadDataBuffer.NumBytes > 32 * 1024 * 1024))
+			|| BrickUploadDataBuffer.NumBytes > 2 * NumBrickDataElements * BrickDataNumBytesPerElement)
 		{
 			BrickUploadDataBuffer.Initialize(RHICmdList, TEXT("DistanceFields.BrickUploadDataBuffer"), BrickDataNumBytesPerElement, NumBrickDataElements, DistanceField::DistanceFieldFormat, BUF_Volatile);
 		}
