DebugPrimitives: don't clear color target for PopulateDepthPass

If we have multiple views atlasted into a target (e.g. stereo), we will clear between the eyes
if the debug draws are the last pass in the postprocessing chain.
We don't actually _need_ to clear because we are blitting the entire source texture into the destination.

diff --git a/Engine/Source/Runtime/Renderer/Private/PostProcess/PostProcessCompositePrimitivesCommon.cpp b/Engine/Source/Runtime/Renderer/Private/PostProcess/PostProcessCompositePrimitivesCommon.cpp
--- a/Engine/Source/Runtime/Renderer/Private/PostProcess/PostProcessCompositePrimitivesCommon.cpp
+++ b/Engine/Source/Runtime/Renderer/Private/PostProcess/PostProcessCompositePrimitivesCommon.cpp
@@ -266,7 +266,7 @@ void PopulateDepthPass(FRDGBuilder& GraphBuilder,
 	PassParameters->DepthTexture = InSceneDepth.Texture;
 	PassParameters->DepthSampler = PointClampSampler;
 	PassParameters->DepthTextureJitter = SceneDepthJitter;
-	PassParameters->RenderTargets[0] = FRenderTargetBinding(OutPopColor, ERenderTargetLoadAction::EClear);
+	PassParameters->RenderTargets[0] = FRenderTargetBinding(OutPopColor, ERenderTargetLoadAction::ELoad);
 	PassParameters->RenderTargets.DepthStencil = FDepthStencilBinding(OutPopDepth, ERenderTargetLoadAction::EClear, ERenderTargetLoadAction::EClear, FExclusiveDepthStencil::DepthWrite_StencilWrite);
 
 	FPopulateCompositeDepthPS::FPermutationDomain PermutationVector;
