Clamp NumLocalLights to MaxCulledLightsPerCell in MobileLocalLightBuffer when Linked List Cull is disabled

diff --git a/Engine/Shaders/Private/MobileLocalLightsBuffer.usf b/Engine/Shaders/Private/MobileLocalLightsBuffer.usf
--- a/Engine/Shaders/Private/MobileLocalLightsBuffer.usf
+++ b/Engine/Shaders/Private/MobileLocalLightsBuffer.usf
@@ -78,6 +78,7 @@ void MainVS(
 	}
 }
 
+int bHasLinkedListCullEnabled;
 void Main(float4 SvPosition : SV_POSITION
 	, out float3 OutColorA : SV_Target0
 #if POST_PROCESS_LOCAL_LIGHTS == 0
@@ -98,7 +99,11 @@ void Main(float4 SvPosition : SV_POSITION
 	uint GridIndex = ComputeLightGridCellIndex(uint2(PixelPos.x, PixelPos.y), SceneDepth);
 	const FCulledLightsGridData CulledLightGridData = GetCulledLightsGrid(GridIndex, EyeIndex);
 
-	const uint NumLocalLights = min(CulledLightGridData.NumLocalLights, GetNumLocalLights(EyeIndex));
+	uint NumLocalLights = min(CulledLightGridData.NumLocalLights, GetNumLocalLights(EyeIndex));
+	if(bHasLinkedListCullEnabled == 0)
+	{
+		NumLocalLights = min(NumLocalLights, ForwardLightData.MaxCulledLightsPerCell - 1);
+	}
 
 	float3 Color = uint3(0.f, 0.f , 0.f);
 	float3 SavedL = float3(0,0,0);
diff --git a/Engine/Source/Runtime/Renderer/Private/MobileLocalLightsBuffer.cpp b/Engine/Source/Runtime/Renderer/Private/MobileLocalLightsBuffer.cpp
--- a/Engine/Source/Runtime/Renderer/Private/MobileLocalLightsBuffer.cpp
+++ b/Engine/Source/Runtime/Renderer/Private/MobileLocalLightsBuffer.cpp
@@ -71,6 +71,7 @@ class FLocalLightBufferPS : public FGlobalShader
 	BEGIN_SHADER_PARAMETER_STRUCT(FParameters, )
 		SHADER_PARAMETER_STRUCT_REF(FViewUniformShaderParameters, View)
 		SHADER_PARAMETER_RDG_UNIFORM_BUFFER(FForwardLightData, ForwardLightData)
+		SHADER_PARAMETER(int32, bHasLinkedListCullEnabled)
 	END_SHADER_PARAMETER_STRUCT()
 
 	static bool ShouldCompilePermutation(const FGlobalShaderPermutationParameters& Parameters)
@@ -174,6 +175,11 @@ void FMobileSceneRenderer::RenderMobileLocalLightsBuffer(FRDGBuilder& GraphBuild
 			PassParameters->PS.ForwardLightData = View.ForwardLightingResources.ForwardLightUniformBuffer;
 			PassParameters->PS.View = GetShaderBinding(View.ViewUniformBuffer);
 
+			static const auto LightLinkedListCullingCVar = IConsoleManager::Get().FindConsoleVariable(TEXT("r.Forward.LightLinkedListCulling"));
+			check(LightLinkedListCullingCVar != nullptr);
+			int32 LightLinkedListCullingEnabled = LightLinkedListCullingCVar->GetInt();
+			PassParameters->PS.bHasLinkedListCullEnabled = LightLinkedListCullingEnabled;
+
 			auto VertexShader = View.ShaderMap->GetShader<FLocalLightBufferVS>();
 			auto PixelShader = View.ShaderMap->GetShader<FLocalLightBufferPS>();
 
