Fixes a crash where the spatial upscaler could be missing FClearValueBinding

diff --git a/Engine/Source/Runtime/Renderer/Private/PostProcess/PostProcessUpscale.cpp b/Engine/Source/Runtime/Renderer/Private/PostProcess/PostProcessUpscale.cpp
--- a/Engine/Source/Runtime/Renderer/Private/PostProcess/PostProcessUpscale.cpp
+++ b/Engine/Source/Runtime/Renderer/Private/PostProcess/PostProcessUpscale.cpp
@@ -184,8 +184,11 @@ FScreenPassTexture ISpatialUpscaler::AddDefaultUpscalePass(
 
 	if (!Output.IsValid())
 	{
-		FRDGTextureDesc OutputDesc = Inputs.SceneColor.Texture->Desc;
-		OutputDesc.Reset();
+		FRDGTextureDesc OutputDesc = FRDGTextureDesc::Create2D(
+			Inputs.SceneColor.Texture->Desc.Extent,
+			Inputs.SceneColor.Texture->Desc.Format,
+			FClearValueBinding::Black,
+			TexCreate_ShaderResource | TexCreate_RenderTargetable | GFastVRamConfig.Upscale);
 
 		if (Inputs.Stage == EUpscaleStage::PrimaryToSecondary)
 		{
@@ -200,8 +203,6 @@ FScreenPassTexture ISpatialUpscaler::AddDefaultUpscalePass(
 			Output.ViewRect = View.UnscaledViewRect;
 		}
 
-		OutputDesc.Flags |= GFastVRamConfig.Upscale;
-
 		Output.Texture = GraphBuilder.CreateTexture(OutputDesc, TEXT("Upscale"));
 		Output.LoadAction = ERenderTargetLoadAction::EClear;
 	}
