Fixes a RHI resource transition validation failure in TSR

diff --git a/Engine/Source/Runtime/Renderer/Private/PostProcess/TemporalSuperResolution.cpp b/Engine/Source/Runtime/Renderer/Private/PostProcess/TemporalSuperResolution.cpp
--- a/Engine/Source/Runtime/Renderer/Private/PostProcess/TemporalSuperResolution.cpp
+++ b/Engine/Source/Runtime/Renderer/Private/PostProcess/TemporalSuperResolution.cpp
@@ -2222,25 +2222,26 @@ FDefaultTemporalUpscaler::FOutputs AddTemporalSuperResolutionPasses(
 		}
 		else
 		{
+			FRHIRange16 SliceRange(uint16(PrevFrameSliceIndex), uint16(1));
+			if (bCanResurrectHistory)
+			{
+				SliceRange = PrevHistorySliceSequence.GetSRVSliceRange(CurrentFrameSliceIndex, PrevFrameSliceIndex);
+			}
+			check(SliceRange.IsInRange(ResurrectionFrameSliceIndex));
+			check(SliceRange.IsInRange(PrevFrameSliceIndex));
+			check(!SliceRange.IsInRange(CurrentFrameSliceIndex) || History.ColorArray != PrevHistory.ColorArray);
+
 			FRDGTextureSRVDesc PrevColorSRVDesc(PrevHistory.ColorArray);
 			PrevColorSRVDesc.NumMipLevels = 1;
 
 			FRDGTextureSRVDesc PrevMetadataSRVDesc(PrevHistory.MetadataArray);
 			PrevMetadataSRVDesc.NumMipLevels = 1;
 
-			if (bCanResurrectHistory)
-			{
-				FRHIRange16 SliceRange = PrevHistorySliceSequence.GetSRVSliceRange(CurrentFrameSliceIndex, PrevFrameSliceIndex);
-				check(SliceRange.IsInRange(ResurrectionFrameSliceIndex));
-				check(SliceRange.IsInRange(PrevFrameSliceIndex));
-				check(!SliceRange.IsInRange(CurrentFrameSliceIndex));
+			PrevColorSRVDesc.FirstArraySlice = SliceRange.First;
+			PrevColorSRVDesc.NumArraySlices = SliceRange.Num;
 
-				PrevColorSRVDesc.FirstArraySlice = SliceRange.First;
-				PrevColorSRVDesc.NumArraySlices = SliceRange.Num;
-
-				PrevMetadataSRVDesc.FirstArraySlice = SliceRange.First;
-				PrevMetadataSRVDesc.NumArraySlices = SliceRange.Num;
-			}
+			PrevMetadataSRVDesc.FirstArraySlice = SliceRange.First;
+			PrevMetadataSRVDesc.NumArraySlices = SliceRange.Num;
 
 			PassParameters->ResurrectionFrameIndex = ResurrectionFrameSliceIndex - PrevColorSRVDesc.FirstArraySlice;
 			PassParameters->PrevFrameIndex = PrevFrameSliceIndex - PrevColorSRVDesc.FirstArraySlice;
