Debug Primitives: match debug targets MSAA with output MSAA

previously we were using the SceneTextures config, which doesn't make sense because we are adding the debug primitives post-resolve.

diff --git a/Engine/Source/Runtime/Renderer/Private/PostProcess/PostProcessCompositeDebugPrimitives.cpp b/Engine/Source/Runtime/Renderer/Private/PostProcess/PostProcessCompositeDebugPrimitives.cpp
--- a/Engine/Source/Runtime/Renderer/Private/PostProcess/PostProcessCompositeDebugPrimitives.cpp
+++ b/Engine/Source/Runtime/Renderer/Private/PostProcess/PostProcessCompositeDebugPrimitives.cpp
@@ -24,7 +24,7 @@ FScreenPassTexture AddDebugPrimitivePass(
 	{
 		Output = FScreenPassRenderTarget::CreateFromInput(GraphBuilder, Inputs.SceneColor, View.GetOverwriteLoadAction(), TEXT("Debug.DrawPrimitivesColor"));
 	}
-	const uint32 NumMSAASamples = View.GetSceneTextures().Config.NumSamples;
+	const uint32 NumMSAASamples = Output.Texture->Desc.NumSamples;
 	//Set sizing values to match the SceneColor size
 	const FIntRect ViewRect = Output.ViewRect;
 	FIntPoint Extent = Output.Texture->Desc.Extent;
