Debug Primitives: adjust RT load action per view for PopulateDepthPass

#rnx

diff --git a/Engine/Source/Runtime/Renderer/Private/PostProcess/PostProcessCompositePrimitivesCommon.cpp b/Engine/Source/Runtime/Renderer/Private/PostProcess/PostProcessCompositePrimitivesCommon.cpp
--- a/Engine/Source/Runtime/Renderer/Private/PostProcess/PostProcessCompositePrimitivesCommon.cpp
+++ b/Engine/Source/Runtime/Renderer/Private/PostProcess/PostProcessCompositePrimitivesCommon.cpp
@@ -266,7 +266,8 @@ void PopulateDepthPass(FRDGBuilder& GraphBuilder,
 	PassParameters->DepthTexture = InSceneDepth.Texture;
 	PassParameters->DepthSampler = PointClampSampler;
 	PassParameters->DepthTextureJitter = SceneDepthJitter;
-	PassParameters->RenderTargets[0] = FRenderTargetBinding(OutPopColor, ERenderTargetLoadAction::ELoad);
+	const ERenderTargetLoadAction ColorLoadAction = View.IsPrimarySceneView() ? ERenderTargetLoadAction::ENoAction : ERenderTargetLoadAction::ELoad;
+	PassParameters->RenderTargets[0] = FRenderTargetBinding(OutPopColor, ColorLoadAction);
 	PassParameters->RenderTargets.DepthStencil = FDepthStencilBinding(OutPopDepth, ERenderTargetLoadAction::EClear, ERenderTargetLoadAction::EClear, FExclusiveDepthStencil::DepthWrite_StencilWrite);
 
 	FPopulateCompositeDepthPS::FPermutationDomain PermutationVector;
