Use RDG to launch setup lambda to enforce waiting on the task at graph builder execution time.

diff --git a/Engine/Source/Runtime/Renderer/Private/RendererScene.cpp b/Engine/Source/Runtime/Renderer/Private/RendererScene.cpp
--- a/Engine/Source/Runtime/Renderer/Private/RendererScene.cpp
+++ b/Engine/Source/Runtime/Renderer/Private/RendererScene.cpp
@@ -6235,7 +6235,7 @@ void FScene::UpdateAllPrimitiveSceneInfos(FRDGBuilder& GraphBuilder, EUpdateAllP
 		});
 	}
 	
-	CreateLightPrimitiveInteractionsTask = LaunchSceneRenderTask(TEXT("CreateLightPrimitiveInteractionsTask"), [this, &SceneInfosWithAddToScene]
+	CreateLightPrimitiveInteractionsTask = GraphBuilder.AddSetupTask([this, &SceneInfosWithAddToScene]
 	{
 		SCOPED_NAMED_EVENT(CreateLightPrimitiveInteractions, FColor::Emerald);
 
@@ -6250,19 +6250,19 @@ void FScene::UpdateAllPrimitiveSceneInfos(FRDGBuilder& GraphBuilder, EUpdateAllP
 	{
 		const bool bLaunchAsyncTask = EnumHasAnyFlags(AsyncOps, EUpdateAllPrimitiveSceneInfosAsyncOps::CacheMeshDrawCommands) && GRHISupportsMultithreadedShaderCreation;
 
-		CacheMeshDrawCommandsTask = LaunchSceneRenderTask(TEXT("CacheMeshDrawCommandsTask"), [this, &SceneInfosWithStaticDrawListUpdate, bLaunchAsyncTask]()
+		CacheMeshDrawCommandsTask = GraphBuilder.AddSetupTask([this, &SceneInfosWithStaticDrawListUpdate, bLaunchAsyncTask]()
 		{
 			FPrimitiveSceneInfo::CacheMeshDrawCommands(this, SceneInfosWithStaticDrawListUpdate);
 
-		}, IsMobilePlatform(GetShaderPlatform()) ? CreateLightPrimitiveInteractionsTask : UE::Tasks::FTask(), bLaunchAsyncTask);
+		}, IsMobilePlatform(GetShaderPlatform()) ? CreateLightPrimitiveInteractionsTask : UE::Tasks::FTask(), UE::Tasks::ETaskPriority::High, bLaunchAsyncTask);
 
-		CacheNaniteDrawCommandsTask = LaunchSceneRenderTask(TEXT("CacheNaniteDrawCommands"), [this, &SceneInfosWithStaticDrawListUpdate, bLaunchAsyncTask]()
+		CacheNaniteDrawCommandsTask = GraphBuilder.AddSetupTask([this, &SceneInfosWithStaticDrawListUpdate, bLaunchAsyncTask]()
 		{
 			FPrimitiveSceneInfo::CacheNaniteDrawCommands(this, SceneInfosWithStaticDrawListUpdate);
 		}, bLaunchAsyncTask);
 
 #if RHI_RAYTRACING
-		CacheRayTracingPrimitivesTask = LaunchSceneRenderTask(TEXT("CacheRayTracingPrimitivesTask"), [this, &SceneInfosWithStaticDrawListUpdate]()
+		CacheRayTracingPrimitivesTask = GraphBuilder.AddSetupTask([this, &SceneInfosWithStaticDrawListUpdate]()
 		{
 			FPrimitiveSceneInfo::CacheRayTracingPrimitives(this, SceneInfosWithStaticDrawListUpdate);
 		}, bLaunchAsyncTask);
