Saves one UAV on the TSR DilateVelocity pass

diff --git a/Engine/Shaders/Private/TemporalSuperResolution/TSRDilateVelocity.usf b/Engine/Shaders/Private/TemporalSuperResolution/TSRDilateVelocity.usf
--- a/Engine/Shaders/Private/TemporalSuperResolution/TSRDilateVelocity.usf
+++ b/Engine/Shaders/Private/TemporalSuperResolution/TSRDilateVelocity.usf
@@ -57,8 +57,7 @@ Texture2D<uint> PrevScatteredSubpixelDepthTexture;
 
 RWTexture2D<float2> DilatedVelocityOutput;
 RWTexture2D<float> ClosestDepthOutput;
-globallycoherent RWTexture2D<uint> PrevUseCountOutput;
-globallycoherent RWTexture2D<uint> PrevClosestDepthOutput;
+globallycoherent RWTexture2DArray<uint> PrevAtomicOutput;
 RWTexture2DArray<uint> R8Output;
 RWTexture2D<uint> SubpixelDepthOutput;
 
@@ -408,8 +407,8 @@ void MainCS(
 					uint2 CompressedUseCountPixelPos = PixelPos;
 				#endif
 
-				InterlockedAdd(PrevUseCountOutput[CompressedUseCountPixelPos], FinalAdd);
-				InterlockedMax(PrevClosestDepthOutput[PixelPos], PrevClosestDepth);
+				InterlockedAdd(PrevAtomicOutput[uint3(CompressedUseCountPixelPos, 0)], FinalAdd);
+				InterlockedMax(PrevAtomicOutput[uint3(PixelPos, 1)], PrevClosestDepth);
 			}
 		}
 		#endif
diff --git a/Engine/Source/Runtime/RenderCore/Public/RenderGraphResources.h b/Engine/Source/Runtime/RenderCore/Public/RenderGraphResources.h
--- a/Engine/Source/Runtime/RenderCore/Public/RenderGraphResources.h
+++ b/Engine/Source/Runtime/RenderCore/Public/RenderGraphResources.h
@@ -816,6 +816,21 @@ public:
 		, Texture(InTexture)
 	{}
 
+	/** Create UAV that access one specific slice. */
+	static FRDGTextureUAVDesc CreateForSlice(FRDGTextureRef Texture, int32 SliceIndex)
+	{
+		check(Texture);
+		check(Texture->Desc.Dimension == ETextureDimension::Texture2DArray);
+		check(SliceIndex >= 0 && SliceIndex < Texture->Desc.ArraySize);
+
+		FRDGTextureUAVDesc Desc = FRDGTextureUAVDesc(Texture);
+		Desc.FirstArraySlice = (uint16)SliceIndex;
+		Desc.NumArraySlices = 1;
+		Desc.DimensionOverride = ETextureDimension::Texture2D;
+
+		return Desc;
+	}
+
 	/** Create UAV with access to a specific meta data plane */
 	static FRDGTextureUAVDesc CreateForMetaData(FRDGTextureRef Texture, ERDGTextureMetaDataAccess MetaData)
 	{
diff --git a/Engine/Source/Runtime/Renderer/Private/PostProcess/TemporalSuperResolution.cpp b/Engine/Source/Runtime/Renderer/Private/PostProcess/TemporalSuperResolution.cpp
--- a/Engine/Source/Runtime/Renderer/Private/PostProcess/TemporalSuperResolution.cpp
+++ b/Engine/Source/Runtime/Renderer/Private/PostProcess/TemporalSuperResolution.cpp
@@ -515,8 +515,7 @@ class FTSRDilateVelocityCS : public FTSRShader
 
 		SHADER_PARAMETER_RDG_TEXTURE_UAV(RWTexture2D, DilatedVelocityOutput)
 		SHADER_PARAMETER_RDG_TEXTURE_UAV(RWTexture2D, ClosestDepthOutput)
-		SHADER_PARAMETER_RDG_TEXTURE_UAV(RWTexture2D, PrevUseCountOutput)
-		SHADER_PARAMETER_RDG_TEXTURE_UAV(RWTexture2D, PrevClosestDepthOutput)
+		SHADER_PARAMETER_RDG_TEXTURE_UAV(RWTexture2DArray, PrevAtomicOutput)
 		SHADER_PARAMETER_RDG_TEXTURE_UAV(RWTexture2DArray, R8Output)
 		SHADER_PARAMETER_RDG_TEXTURE_UAV(RWTexture2DArray, SubpixelDepthOutput)
 
@@ -560,8 +559,8 @@ class FTSRDecimateHistoryCS : public FTSRShader
 		SHADER_PARAMETER_RDG_TEXTURE(Texture2D, InputSceneColorTexture)
 		SHADER_PARAMETER_RDG_TEXTURE(Texture2D, DilatedVelocityTexture)
 		SHADER_PARAMETER_RDG_TEXTURE(Texture2D, ClosestDepthTexture)
-		SHADER_PARAMETER_RDG_TEXTURE(Texture2D, PrevUseCountTexture)
-		SHADER_PARAMETER_RDG_TEXTURE(Texture2D, PrevClosestDepthTexture)
+		SHADER_PARAMETER_RDG_TEXTURE_SRV(Texture2D, PrevUseCountTexture)
+		SHADER_PARAMETER_RDG_TEXTURE_SRV(Texture2D, PrevClosestDepthTexture)
 
 		SHADER_PARAMETER_STRUCT_INCLUDE(FTSRPrevHistoryParameters, PrevHistoryParameters)
 		SHADER_PARAMETER_RDG_TEXTURE(Texture2D, PrevHistoryGuide)
@@ -1607,25 +1606,24 @@ FDefaultTemporalUpscaler::FOutputs AddTemporalSuperResolutionPasses(
 	}
 
 	// Clear atomic scattered texture.
-	FRDGTextureRef PrevUseCountTexture;
-	FRDGTextureRef PrevClosestDepthTexture;
+	FRDGTextureRef PrevAtomicTexture;
 	FRDGTextureRef PrevScatteredSubpixelDepthTexture = nullptr;
 	{
 		{
-			FRDGTextureDesc Desc = FRDGTextureDesc::Create2D(
+			FRDGTextureDesc Desc = FRDGTextureDesc::Create2DArray(
 				InputExtent,
 				PF_R32_UINT,
 				FClearValueBinding::None,
-				/* InFlags = */ TexCreate_ShaderResource | TexCreate_UAV | TexCreate_AtomicCompatible);
+				/* InFlags = */ TexCreate_ShaderResource | TexCreate_UAV | TexCreate_AtomicCompatible,
+				/* ArraySize */ 2);
 
-			PrevUseCountTexture = GraphBuilder.CreateTexture(Desc, TEXT("TSR.PrevUseCountTexture"));
-			PrevClosestDepthTexture = GraphBuilder.CreateTexture(Desc, TEXT("TSR.PrevClosestDepthTexture"));
+			PrevAtomicTexture = GraphBuilder.CreateTexture(Desc, TEXT("TSR.PrevAtomicTexture"));
 		}
 
 		FTSRClearPrevTexturesCS::FParameters* PassParameters = GraphBuilder.AllocParameters<FTSRClearPrevTexturesCS::FParameters>();
 		PassParameters->CommonParameters = CommonParameters;
-		PassParameters->PrevUseCountOutput = GraphBuilder.CreateUAV(PrevUseCountTexture);
-		PassParameters->PrevClosestDepthOutput = GraphBuilder.CreateUAV(PrevClosestDepthTexture);
+		PassParameters->PrevUseCountOutput     = GraphBuilder.CreateUAV(FRDGTextureUAVDesc::CreateForSlice(PrevAtomicTexture, /* SliceIndex = */ 0));
+		PassParameters->PrevClosestDepthOutput = GraphBuilder.CreateUAV(FRDGTextureUAVDesc::CreateForSlice(PrevAtomicTexture, /* SliceIndex = */ 1));
 		if (History.SubpixelDepth && PrevHistory.SubpixelDepth)
 		{
 			FRDGTextureDesc Desc = History.SubpixelDepth->Desc;
@@ -1755,8 +1753,7 @@ FDefaultTemporalUpscaler::FOutputs AddTemporalSuperResolutionPasses(
 
 		PassParameters->DilatedVelocityOutput = GraphBuilder.CreateUAV(DilatedVelocityTexture);
 		PassParameters->ClosestDepthOutput = GraphBuilder.CreateUAV(ClosestDepthTexture);
-		PassParameters->PrevUseCountOutput = GraphBuilder.CreateUAV(PrevUseCountTexture);
-		PassParameters->PrevClosestDepthOutput = GraphBuilder.CreateUAV(PrevClosestDepthTexture);
+		PassParameters->PrevAtomicOutput = GraphBuilder.CreateUAV(PrevAtomicTexture);
 		PassParameters->R8Output = GraphBuilder.CreateUAV(R8OutputTexture);
 		if (SubpixelMethod == ETSRSubpixelMethod::ClosestDepth)
 		{
@@ -1868,8 +1865,8 @@ FDefaultTemporalUpscaler::FOutputs AddTemporalSuperResolutionPasses(
 		PassParameters->InputSceneColorTexture = PassInputs.SceneColor.Texture;
 		PassParameters->DilatedVelocityTexture = DilatedVelocityTexture;
 		PassParameters->ClosestDepthTexture = ClosestDepthTexture;
-		PassParameters->PrevUseCountTexture = PrevUseCountTexture;
-		PassParameters->PrevClosestDepthTexture = PrevClosestDepthTexture;
+		PassParameters->PrevUseCountTexture     = GraphBuilder.CreateSRV(FRDGTextureSRVDesc::CreateForSlice(PrevAtomicTexture, /* SliceIndex = */ 0));
+		PassParameters->PrevClosestDepthTexture = GraphBuilder.CreateSRV(FRDGTextureSRVDesc::CreateForSlice(PrevAtomicTexture, /* SliceIndex = */ 1));
 
 		PassParameters->PrevHistoryParameters = PrevHistoryParameters;
 
