Partially undo CL 26259323 (async deletion of scene proxies)

* Immediately delete Proxy and PrimitiveSceneInfo in RemovePrimitiveSceneInfo_RenderThread, instead of deferring the operation
* Speculative fix for a hard-to-reproduce crash that impacts some users

diff --git a/Engine/Source/Runtime/Renderer/Private/RendererScene.cpp b/Engine/Source/Runtime/Renderer/Private/RendererScene.cpp
--- a/Engine/Source/Runtime/Renderer/Private/RendererScene.cpp
+++ b/Engine/Source/Runtime/Renderer/Private/RendererScene.cpp
@@ -2322,7 +2322,8 @@ void FScene::RemovePrimitiveSceneInfo_RenderThread(FPrimitiveSceneInfo* Primitiv
 			};
 
 			BeginCleanup(new DeferDeleteHitProxies(MoveTemp(PrimitiveSceneInfo->HitProxies)));
-			PrimitiveSceneInfosToDelete.Emplace(PrimitiveSceneInfo);
+			delete PrimitiveSceneInfo->Proxy;
+			delete PrimitiveSceneInfo;
 		}
 	}
 	else
@@ -6353,14 +6354,6 @@ void FScene::UpdateAllPrimitiveSceneInfos(FRDGBuilder& GraphBuilder, EUpdateAllP
 			delete PrimitiveSceneInfo->Proxy;
 			delete PrimitiveSceneInfo;
 		}
-
-		// Delete scene primitives that were removed without ever being added.
-		for (FPrimitiveSceneInfo* PrimitiveSceneInfo : PrimitiveSceneInfosToDelete)
-		{
-			delete PrimitiveSceneInfo->Proxy;
-			delete PrimitiveSceneInfo;
-		}
-		PrimitiveSceneInfosToDelete.Empty();
 	});
 
 	if (bNeedPathTracedInvalidation)
diff --git a/Engine/Source/Runtime/Renderer/Private/ScenePrivate.h b/Engine/Source/Runtime/Renderer/Private/ScenePrivate.h
--- a/Engine/Source/Runtime/Renderer/Private/ScenePrivate.h
+++ b/Engine/Source/Runtime/Renderer/Private/ScenePrivate.h
@@ -3802,7 +3802,6 @@ private:
 	Experimental::TRobinHoodHashSet<FPrimitiveSceneInfo*> RemovedPrimitiveSceneInfos;
 	Experimental::TRobinHoodHashSet<FPrimitiveSceneInfo*> DistanceFieldSceneDataUpdates;
 	TArray<FLevelCommand> LevelCommands;
-	TArray<FPrimitiveSceneInfo*> PrimitiveSceneInfosToDelete;
 
 	UE::Tasks::FTask CreateLightPrimitiveInteractionsTask;
 	UE::Tasks::FTask CacheMeshDrawCommandsTask;
