Fix GOcclusionRandomStream OOB due to CurrentSample access not being thread safe

#rnx

diff --git a/Engine/Source/Runtime/Renderer/Private/ScenePrivate.h b/Engine/Source/Runtime/Renderer/Private/ScenePrivate.h
--- a/Engine/Source/Runtime/Renderer/Private/ScenePrivate.h
+++ b/Engine/Source/Runtime/Renderer/Private/ScenePrivate.h
@@ -530,18 +530,25 @@ public:
 	/** @return A random number between 0 and 1. */
 	inline float GetFraction()
 	{
-		if (CurrentSample >= NumSamples)
+		uint32 Current = CurrentSample.fetch_add(1);
+
+		if (Current >= NumSamples)
 		{
-			CurrentSample = 0;
+			Current++;
+			CurrentSample.compare_exchange_strong(Current, 0);
+			Current = 0;
+			// It is intended here to not check if exchange worked or failed. 
+			// It might be overkill to call recursively GetFraction if exchange failed. 
+			// Another thread might already have reset CurrentSample and it is acceptable 
+			// to have two threads returning the same Fraction at index 0
 		}
-		float Fraction = Samples[CurrentSample];
-		CurrentSample++;
+		float Fraction = Samples[Current];
 		return Fraction;
 	}
 private:
 
 	/** Index of the last sample we produced **/
-	uint32 CurrentSample;
+	std::atomic<uint32> CurrentSample;
 	/** A list of float random samples **/
 	float Samples[NumSamples];
 };
