Validate that ERHIFeatureLevel being passed into FScene is valid for the platform.

diff --git a/Engine/Source/Runtime/Renderer/Private/RendererScene.cpp b/Engine/Source/Runtime/Renderer/Private/RendererScene.cpp
--- a/Engine/Source/Runtime/Renderer/Private/RendererScene.cpp
+++ b/Engine/Source/Runtime/Renderer/Private/RendererScene.cpp
@@ -1653,6 +1653,9 @@ FScene::FScene(UWorld* InWorld, bool bInRequiresHitProxies, bool bInIsEditorScen
 
 	FeatureLevel = World->GetFeatureLevel();
 
+	checkf((uint32)FeatureLevel < (uint32)ERHIFeatureLevel::Num, TEXT("World provided an invalid feature level (%d) to FScene."), FeatureLevel);
+	checkf(GShaderPlatformForFeatureLevel[FeatureLevel] != SP_NumPlatforms, TEXT("Invalid feature level %s for platform (max feature level %s)"), *LexToString(FeatureLevel), *LexToString(GMaxRHIFeatureLevel));
+
 	GPUScene.SetEnabled(FeatureLevel);
 
 	if (GPUScene.IsEnabled())
