Debug Draw: fix usage of input scene color rect + output view rect

Previously, we relied on the fact that the target rect matched into the source color/depth textures.
However, with TSR + Stereo rendering, TSR outputs a monoscopic texture, while DebugPrimitives is
drawing into a double-wide stereo texture. Because of this, we need to generate UVs from the target rect,
and sample from the input texture rects.

It's unclear why this behavior was changed in 25024370. I think the original CL wanted to 'target' the target
view rect, but that shouldn't have affected the source rect.

diff --git a/Engine/Shaders/Private/PostProcessCompositePrimitives.usf b/Engine/Shaders/Private/PostProcessCompositePrimitives.usf
--- a/Engine/Shaders/Private/PostProcessCompositePrimitives.usf
+++ b/Engine/Shaders/Private/PostProcessCompositePrimitives.usf
@@ -248,7 +248,7 @@ void MainPopulateSceneDepthPS(
 	out float4 OutColor : SV_Target0,
 	out float OutDepth : SV_DEPTH)
 {
-	float2 ViewportUV = (SvPosition.xy - Color_ViewportMin) * Color_ViewportSizeInverse;
+	float2 ViewportUV = (SvPosition.xy - View.ViewRectMin.xy) * View.ViewSizeAndInvSize.zw;
 	
 	float2 DepthUV = clamp(
 		(Depth_ViewportMin + ViewportUV * Depth_ViewportSize - DepthTextureJitter) * Depth_ExtentInverse,
diff --git a/Engine/Source/Runtime/Renderer/Private/PostProcess/PostProcessCompositePrimitivesCommon.cpp b/Engine/Source/Runtime/Renderer/Private/PostProcess/PostProcessCompositePrimitivesCommon.cpp
--- a/Engine/Source/Runtime/Renderer/Private/PostProcess/PostProcessCompositePrimitivesCommon.cpp
+++ b/Engine/Source/Runtime/Renderer/Private/PostProcess/PostProcessCompositePrimitivesCommon.cpp
@@ -251,17 +251,12 @@ void PopulateDepthPass(FRDGBuilder& GraphBuilder,
 
 	FPopulateCompositeDepthPS::FParameters* PassParameters = GraphBuilder.AllocParameters<FPopulateCompositeDepthPS::FParameters>();
 	PassParameters->View = View.ViewUniformBuffer;
+	PassParameters->Color = GetScreenPassTextureViewportParameters(FScreenPassTextureViewport(InSceneColor));
+	PassParameters->Depth = GetScreenPassTextureViewportParameters(FScreenPassTextureViewport(InSceneDepth));
 	if (bForceDrawColor)
 	{
 		PassParameters->ColorTexture = InSceneColor.Texture;
 		PassParameters->ColorSampler = PointClampSampler;
-		PassParameters->Color = GetScreenPassTextureViewportParameters(FScreenPassTextureViewport(InSceneColor.Texture->Desc.Extent, View.ViewRect));
-		PassParameters->Depth = GetScreenPassTextureViewportParameters(FScreenPassTextureViewport(InSceneDepth.Texture->Desc.Extent, View.ViewRect));
-	}
-	else
-	{
-		PassParameters->Color = GetScreenPassTextureViewportParameters(FScreenPassTextureViewport(InSceneColor));
-		PassParameters->Depth = GetScreenPassTextureViewportParameters(FScreenPassTextureViewport(InSceneDepth));
 	}
 	PassParameters->DepthTexture = InSceneDepth.Texture;
 	PassParameters->DepthSampler = PointClampSampler;
