Add simple "dynamic resolution" for VSMs that can dynamically decrease the global resolution LOD bias based on page pool pressure.

diff --git a/Engine/Shaders/Private/VirtualShadowMaps/VirtualShadowMapPageMarking.usf b/Engine/Shaders/Private/VirtualShadowMaps/VirtualShadowMapPageMarking.usf
--- a/Engine/Shaders/Private/VirtualShadowMaps/VirtualShadowMapPageMarking.usf
+++ b/Engine/Shaders/Private/VirtualShadowMaps/VirtualShadowMapPageMarking.usf
@@ -35,6 +35,12 @@ uint InputType;
 uint bCullBackfacingPixels;
 uint NumDirectionalLightSmInds;
 
+int GetBiasedClipmapLevel(FVirtualShadowMapProjectionShaderData BaseProjectionData, FLWCVector3 WorldPosition)
+{
+	float BiasedLevel = CalcAbsoluteClipmapLevel(BaseProjectionData, WorldPosition) + BaseProjectionData.ResolutionLodBias + VirtualShadowMap.GlobalResolutionLodBias;
+	return int(floor(BiasedLevel));
+}
+
 uint GetMipLevelLocal(int VirtualShadowMapId, float3 TranslatedWorldPosition, float SceneDepth)
 {
 	FVirtualShadowMapProjectionShaderData ProjectionData = GetVirtualShadowMapProjectionData(VirtualShadowMapId);
@@ -45,7 +51,7 @@ uint GetMipLevelLocal(int VirtualShadowMapId, float3 TranslatedWorldPosition, fl
 
 	float Footprint = VirtualShadowMapCalcPixelFootprintLocal(ProjectionData, ShadowTranslatedWorldPosition, SceneDepth);
 
-	float MipLevelFloat = log2(Footprint) + ProjectionData.ResolutionLodBias;
+	float MipLevelFloat = log2(Footprint) + ProjectionData.ResolutionLodBias + VirtualShadowMap.GlobalResolutionLodBias;
 	uint MipLevel = uint(max(floor(MipLevelFloat), 0.0f));
 	return min(MipLevel, (VSM_MAX_MIP_LEVELS - 1U));
 }
@@ -112,7 +118,7 @@ void MarkPageClipmap(
 	FLWCVector3 WorldPosition,
 	float3 TranslatedWorldPosition)
 {
-	const int ClipmapLevel = CalcClipmapLevel(ProjectionData, WorldPosition);
+	const int ClipmapLevel = GetBiasedClipmapLevel(ProjectionData, WorldPosition);
 	int ClipmapIndex = max(0, ClipmapLevel - ProjectionData.ClipmapLevel);
 	if (ClipmapIndex < ProjectionData.ClipmapLevelCountRemaining)
 	{
diff --git a/Engine/Shaders/Private/VirtualShadowMaps/VirtualShadowMapPhysicalPageManagement.usf b/Engine/Shaders/Private/VirtualShadowMaps/VirtualShadowMapPhysicalPageManagement.usf
--- a/Engine/Shaders/Private/VirtualShadowMaps/VirtualShadowMapPhysicalPageManagement.usf
+++ b/Engine/Shaders/Private/VirtualShadowMaps/VirtualShadowMapPhysicalPageManagement.usf
@@ -1111,11 +1111,14 @@ StructuredBuffer<int> PhysicalPageLists;
 void FeedbackStatusCS()
 {
 	{
-		FGPUMessageWriter Mw = GPUMessageBegin(StatusMessageId, 1U);
+		FGPUMessageWriter Mw = GPUMessageBegin(StatusMessageId, 2U);
 
 		// Write out how many pages are still available
-		int CountIndex = GetPhysicalPageListStart(PHYSICAL_PAGE_LIST_AVAILABLE) + VirtualShadowMap.MaxPhysicalPages;		
+		int CountIndex = GetPhysicalPageListStart(PHYSICAL_PAGE_LIST_AVAILABLE) + VirtualShadowMap.MaxPhysicalPages;
 		GPUMessageWriteItem(Mw, PhysicalPageLists[CountIndex]);
+
+		// Write out the resolution lod bias from this frame
+		GPUMessageWriteItem(Mw, VirtualShadowMap.GlobalResolutionLodBias);
 	}
 }
 
diff --git a/Engine/Shaders/Private/VirtualShadowMaps/VirtualShadowMapPrintStats.usf b/Engine/Shaders/Private/VirtualShadowMaps/VirtualShadowMapPrintStats.usf
--- a/Engine/Shaders/Private/VirtualShadowMaps/VirtualShadowMapPrintStats.usf
+++ b/Engine/Shaders/Private/VirtualShadowMaps/VirtualShadowMapPrintStats.usf
@@ -185,6 +185,11 @@ void LogVirtualSmStatsCS()
 			PrintLeftAlign(Context, InStatsBuffer[VSM_STAT_STATIC_INVALIDATED_PAGES]);
 			Newline(Context);
 		}
+
+		Context.Pos.x = HeadlineX;
+		Print(Context, TEXT("Global Resolution Bias   "), FontWhite);
+		Print(Context, VirtualShadowMap.GlobalResolutionLodBias);
+		Newline(Context);
 	}
 
 	if (bShowNonNaniteInstances)
diff --git a/Engine/Shaders/Private/VirtualShadowMaps/VirtualShadowMapProjectionCommon.ush b/Engine/Shaders/Private/VirtualShadowMaps/VirtualShadowMapProjectionCommon.ush
--- a/Engine/Shaders/Private/VirtualShadowMaps/VirtualShadowMapProjectionCommon.ush
+++ b/Engine/Shaders/Private/VirtualShadowMaps/VirtualShadowMapProjectionCommon.ush
@@ -23,12 +23,6 @@ float CalcAbsoluteClipmapLevel(FVirtualShadowMapProjectionShaderData BaseProject
 	return log2(DistanceToClipmapOrigin);
 }
 
-int CalcClipmapLevel(FVirtualShadowMapProjectionShaderData BaseProjectionData, FLWCVector3 WorldPosition)
-{
-	float BiasedLevel = CalcAbsoluteClipmapLevel(BaseProjectionData, WorldPosition) + BaseProjectionData.ResolutionLodBias;
-	return int(floor(BiasedLevel));
-}
-
 // Calculate the projected pixel footprint onto a local light, scaled in shadow map texels
 // ShadowTranslatedWorldPosition should already be translated by both the primary view and shadow ProjectionData.PreViewTranslation
 float VirtualShadowMapCalcPixelFootprintLocal(
@@ -360,9 +354,8 @@ FVirtualShadowMapSampleResult SampleVirtualShadowMapInner(int VirtualShadowMapId
 
 	if (BaseProjectionData.LightType == LIGHT_TYPE_DIRECTIONAL)
 	{
-		const int FirstClipmapLevel = BaseProjectionData.ClipmapLevel;
-		const int ClipmapLevel = CalcClipmapLevel(BaseProjectionData, WorldPosition);
-		int ClipmapIndex = max(0, ClipmapLevel - FirstClipmapLevel);
+		const int ClipmapLevel = int(floor(CalcAbsoluteClipmapLevel(BaseProjectionData, WorldPosition)));
+		int ClipmapIndex = max(0, ClipmapLevel - BaseProjectionData.ClipmapLevel);
 
 		// Check if sample is within the clipmap range (from camera)
 		if (ClipmapIndex < BaseProjectionData.ClipmapLevelCountRemaining)
@@ -412,16 +405,17 @@ FVirtualShadowMapSampleResult SampleVirtualShadowMapInner(int VirtualShadowMapId
 	}
 	else
 	{
-		float3 ShadowTranslatedWorldPosition = LWCToFloat(LWCAdd(WorldPosition, BaseProjectionData.PreViewTranslation));
+		FVirtualShadowMapProjectionShaderData ProjectionData = GetVirtualShadowMapProjectionData(VirtualShadowMapId);
+		float3 ShadowTranslatedWorldPosition = LWCToFloat(LWCAdd(WorldPosition, ProjectionData.PreViewTranslation));
 
-		if (BaseProjectionData.LightType != LIGHT_TYPE_SPOT)
+		if (ProjectionData.LightType != LIGHT_TYPE_SPOT)
 		{
 			VirtualShadowMapId += VirtualShadowMapGetCubeFace(ShadowTranslatedWorldPosition);
-			BaseProjectionData = GetVirtualShadowMapProjectionData(VirtualShadowMapId);
+			ProjectionData = GetVirtualShadowMapProjectionData(VirtualShadowMapId);
 		}
 
 		// TODO cubemap math directly instead of vector fetching matrix
-		float4 ShadowUVz = mul(float4(ShadowTranslatedWorldPosition, 1.0f), BaseProjectionData.TranslatedWorldToShadowUVMatrix);
+		float4 ShadowUVz = mul(float4(ShadowTranslatedWorldPosition, 1.0f), ProjectionData.TranslatedWorldToShadowUVMatrix);
 		ShadowUVz.xyz /= ShadowUVz.w;
 
 		FVirtualShadowMapSample SmSample = SampleVirtualShadowMap(VirtualShadowMapId, ShadowUVz.xy);
@@ -440,7 +434,7 @@ FVirtualShadowMapSampleResult SampleVirtualShadowMapInner(int VirtualShadowMapId
 				? ComputeVirtualShadowMapOptimalSlopeBias(VirtualShadowMapId, SmSample, ShadowTranslatedWorldPosition, EstimatedGeoWorldNormal)
 				: 0.0f;
 
-			float RayStartBias = -RayStartDistance * BaseProjectionData.ShadowViewToClipMatrix._33 / ShadowUVz.w;
+			float RayStartBias = -RayStartDistance * ProjectionData.ShadowViewToClipMatrix._33 / ShadowUVz.w;
 			float BiasedDepth = SmSample.Depth - OptimalSlopeBias - RayStartBias;
 
 			if (BiasedDepth > ShadowUVz.z)
@@ -449,7 +443,7 @@ FVirtualShadowMapSampleResult SampleVirtualShadowMapInner(int VirtualShadowMapId
 				// Shadow view matrix is rotation
 				float ReceiverDistance = length(ShadowTranslatedWorldPosition);
 				Result.OccluderDistance = RayStartDistance + ComputeOccluderDistancePerspective(
-					BaseProjectionData.ShadowViewToClipMatrix,
+					ProjectionData.ShadowViewToClipMatrix,
 					SmSample.Depth,
 					ShadowUVz.z,
 					ReceiverDistance);
diff --git a/Engine/Source/Runtime/Renderer/Private/VirtualShadowMaps/VirtualShadowMapArray.cpp b/Engine/Source/Runtime/Renderer/Private/VirtualShadowMaps/VirtualShadowMapArray.cpp
--- a/Engine/Source/Runtime/Renderer/Private/VirtualShadowMaps/VirtualShadowMapArray.cpp
+++ b/Engine/Source/Runtime/Renderer/Private/VirtualShadowMaps/VirtualShadowMapArray.cpp
@@ -448,6 +448,8 @@ void FVirtualShadowMapArray::Initialize(FRDGBuilder& GraphBuilder, FVirtualShado
 		UniformParameters.PhysicalPoolSize = FIntPoint( PhysicalX, PhysicalY );
 		UniformParameters.PhysicalPoolSizePages = FIntPoint( PhysicalPagesX, PhysicalPagesY );
 
+		UniformParameters.GlobalResolutionLodBias = CacheManager->GetGlobalResolutionLodBias();
+
 		// TODO: Parameterize this in a useful way; potentially modify it automatically
 		// when there are fewer lights in the scene and/or clustered shading settings differ.
 		UniformParameters.PackedShadowMaskMaxLightCount = FMath::Min(CVarVirtualShadowOnePassProjectionMaxLights.GetValueOnRenderThread(), 32);
@@ -678,9 +680,11 @@ class FGeneratePageFlagsFromPixelsCS : public FVirtualPageManagementShader
 		SHADER_PARAMETER_RDG_BUFFER_SRV(StructuredBuffer< uint >, DirectionalLightIds)
 		SHADER_PARAMETER_RDG_BUFFER_SRV(StructuredBuffer< uint >, PrunedLightGridData)
 		SHADER_PARAMETER_RDG_BUFFER_SRV(StructuredBuffer< uint >, PrunedNumCulledLightsGrid)
+		// PERMUTATION_WATER_DEPTH
 		SHADER_PARAMETER_RDG_TEXTURE(Texture2D<float>, SingleLayerWaterDepthTexture)
 		SHADER_PARAMETER_RDG_BUFFER_SRV(StructuredBuffer< uint >, SingleLayerWaterTileMask)
 		SHADER_PARAMETER(FIntPoint, SingleLayerWaterTileViewRes)
+		// PERMUTATION_TRANSLUCENCY_DEPTH
 		SHADER_PARAMETER(uint32, FrontLayerMode)
 		SHADER_PARAMETER(FVector4f, FrontLayerHistoryUVMinMax)
 		SHADER_PARAMETER(FVector4f, FrontLayerHistoryScreenPositionScaleBias)
@@ -692,7 +696,6 @@ class FGeneratePageFlagsFromPixelsCS : public FVirtualPageManagementShader
 		SHADER_PARAMETER(uint32, InputType)
 		SHADER_PARAMETER(uint32, NumDirectionalLightSmInds)
 		SHADER_PARAMETER(uint32, bPostBasePass)
-		SHADER_PARAMETER(float, ResolutionLodBiasLocal)
 		SHADER_PARAMETER(float, PageDilationBorderSizeDirectional)
 		SHADER_PARAMETER(float, PageDilationBorderSizeLocal)
 		SHADER_PARAMETER(uint32, bCullBackfacingPixels)
diff --git a/Engine/Source/Runtime/Renderer/Private/VirtualShadowMaps/VirtualShadowMapArray.h b/Engine/Source/Runtime/Renderer/Private/VirtualShadowMaps/VirtualShadowMapArray.h
--- a/Engine/Source/Runtime/Renderer/Private/VirtualShadowMaps/VirtualShadowMapArray.h
+++ b/Engine/Source/Runtime/Renderer/Private/VirtualShadowMaps/VirtualShadowMapArray.h
@@ -134,7 +134,7 @@ BEGIN_GLOBAL_SHADER_PARAMETER_STRUCT(FVirtualShadowMapUniformParameters, )
 	SHADER_PARAMETER(uint32, PackedShadowMaskMaxLightCount)
 	SHADER_PARAMETER(FVector4f, RecPhysicalPoolSize)
 	SHADER_PARAMETER(FIntPoint, PhysicalPoolSize)
-	SHADER_PARAMETER(FIntPoint, PhysicalPoolSizePages)	
+	SHADER_PARAMETER(FIntPoint, PhysicalPoolSizePages)
 
 	// Set to 1 if r.Shadow.Virtual.NonNanite.IncludeInCoarsePages is set to 0 in order to signal that we want to use the legacy path for just excluding non-nanite
 	SHADER_PARAMETER(uint32, bExcludeNonNaniteFromCoarsePages)
@@ -145,6 +145,8 @@ BEGIN_GLOBAL_SHADER_PARAMETER_STRUCT(FVirtualShadowMapUniformParameters, )
 	// For shadow page age calculations
 	SHADER_PARAMETER(uint32, SceneFrameNumber)
 
+	SHADER_PARAMETER(float, GlobalResolutionLodBias)
+
 	SHADER_PARAMETER_RDG_BUFFER_SRV(ByteAddressBuffer, ProjectionData)
 	SHADER_PARAMETER_RDG_BUFFER_SRV(StructuredBuffer<uint>, PageTable)
 	SHADER_PARAMETER_RDG_BUFFER_SRV(StructuredBuffer<uint>, PageFlags)
diff --git a/Engine/Source/Runtime/Renderer/Private/VirtualShadowMaps/VirtualShadowMapCacheManager.cpp b/Engine/Source/Runtime/Renderer/Private/VirtualShadowMaps/VirtualShadowMapCacheManager.cpp
--- a/Engine/Source/Runtime/Renderer/Private/VirtualShadowMaps/VirtualShadowMapCacheManager.cpp
+++ b/Engine/Source/Runtime/Renderer/Private/VirtualShadowMaps/VirtualShadowMapCacheManager.cpp
@@ -86,6 +86,20 @@ static TAutoConsoleVariable<int32> CVarVSMReservedResource(
 	ECVF_RenderThreadSafe
 );
 
+static TAutoConsoleVariable<float> CVarVSMDynamicResolutionMaxLodBias(
+	TEXT("r.Shadow.Virtual.DynamicRes.MaxResolutionLodBias"),
+	2.0f,
+	TEXT("Maximum LOD bias to clamp to for global dynamic resolution reduction. 0 = disabled"),
+	ECVF_RenderThreadSafe
+);
+
+static TAutoConsoleVariable<float> CVarVSMDynamicResolutionMaxPagePoolAllocation(
+	TEXT("r.Shadow.Virtual.DynamicRes.MaxPagePoolAllocation"),
+	0.85f,
+	TEXT("If allocation exceeds this factor of total page pool capacity, shadow resolution will be biased downwards. 0 = disabled"),
+	ECVF_RenderThreadSafe
+);
+
 void FVirtualShadowMapCacheEntry::UpdateClipmapLevel(
 	FVirtualShadowMapArray& VirtualShadowMapArray,
 	const FVirtualShadowMapPerLightCacheEntry& PerLightEntry,
@@ -521,30 +535,64 @@ FVirtualShadowMapArrayCacheManager::FVirtualShadowMapArrayCacheManager(FScene* I
 	StatusFeedbackSocket = GPUMessage::RegisterHandler(TEXT("Shadow.Virtual.StatusFeedback"), [this](GPUMessage::FReader Message)
 	{
 		// Goes negative on underflow
-		int32 NumPagesFree = Message.Read<int32>(0);
-			
-		CSV_CUSTOM_STAT(VSM, FreePages, NumPagesFree, ECsvCustomStatOp::Set);
+		int32 LastFreePhysicalPages = Message.Read<int32>(0);
+		const float LastGlobalResolutionLodBias = FMath::AsFloat(Message.Read<uint32>(0U));
+		
+		CSV_CUSTOM_STAT(VSM, FreePages, LastFreePhysicalPages, ECsvCustomStatOp::Set);
 
-		if (NumPagesFree < 0)
+		// Dynamic resolution
 		{
-			static const auto* CVarResolutionLodBiasLocalPtr = IConsoleManager::Get().FindTConsoleVariableDataFloat(TEXT("r.Shadow.Virtual.ResolutionLodBiasLocal"));
-			const float LodBiasLocal = CVarResolutionLodBiasLocalPtr->GetValueOnRenderThread();
+			// Could be cvars if needed, but not clearly something that needs to be tweaked currently
+			// NOTE: Should react more quickly when reducing resolution than when increasing again
+			// TODO: Possibly something smarter/PID-like rather than simple exponential decay
+			const float ResolutionDownExpLerpFactor = 0.5f;
+			const float ResolutionUpExpLerpFactor = 0.1f;
+			const uint32 FramesBeforeResolutionUp = 10;
+
+			const float MaxPageAllocation = CVarVSMDynamicResolutionMaxPagePoolAllocation.GetValueOnRenderThread();
+			const float MaxLodBias = CVarVSMDynamicResolutionMaxLodBias.GetValueOnRenderThread();
+			
+			if (MaxPageAllocation > 0.0f)
+			{
+				const uint32 SceneFrameNumber = Scene->GetFrameNumberRenderThread();
+
+				// Dynamically bias shadow resolution when we get too near the maximum pool capacity
+				// NB: In a perfect world each +1 of resolution bias will drop the allocation in half
+				float CurrentAllocation = 1.0f - (LastFreePhysicalPages / static_cast<float>(MaxPhysicalPages));
+				float AllocationRatio = CurrentAllocation / MaxPageAllocation;
+				float TargetLodBias = FMath::Max(0.0f, LastGlobalResolutionLodBias + FMath::Log2(AllocationRatio));
 
-			static const auto* CVarResolutionLodBiasDirectionalPtr = IConsoleManager::Get().FindTConsoleVariableDataFloat(TEXT("r.Shadow.Virtual.ResolutionLodBiasDirectional"));
-			const float LodBiasDirectional = CVarResolutionLodBiasDirectionalPtr->GetValueOnRenderThread();
+				if (CurrentAllocation <= MaxPageAllocation &&
+					(SceneFrameNumber - LastFrameOverPageAllocationBudget) > FramesBeforeResolutionUp)
+				{
+					GlobalResolutionLodBias = FMath::Lerp(GlobalResolutionLodBias, TargetLodBias, ResolutionUpExpLerpFactor);
+				}
+				else if (CurrentAllocation > MaxPageAllocation)
+				{
+					LastFrameOverPageAllocationBudget = SceneFrameNumber;
+					GlobalResolutionLodBias = FMath::Lerp(GlobalResolutionLodBias, TargetLodBias, ResolutionDownExpLerpFactor);
+				}
+			}
 
-			static const auto* CVarMaxPhysicalPagesPtr = IConsoleManager::Get().FindTConsoleVariableDataInt(TEXT("r.Shadow.Virtual.MaxPhysicalPages"));
-			const int32 MaxPhysicalPages = CVarMaxPhysicalPagesPtr->GetValueOnRenderThread();
+			GlobalResolutionLodBias = FMath::Clamp(GlobalResolutionLodBias, 0.0f, MaxLodBias);
+		}
 
+		if (LastFreePhysicalPages < 0)
+		{
 #if !UE_BUILD_SHIPPING
 			if (!bLoggedPageOverflow)
 			{
+				static const auto* CVarResolutionLodBiasLocalPtr = IConsoleManager::Get().FindTConsoleVariableDataFloat(TEXT("r.Shadow.Virtual.ResolutionLodBiasLocal"));
+				static const auto* CVarResolutionLodBiasDirectionalPtr = IConsoleManager::Get().FindTConsoleVariableDataFloat(TEXT("r.Shadow.Virtual.ResolutionLodBiasDirectional"));
+
 				UE_LOG(LogRenderer, Warning, TEXT("Virtual Shadow Map Page Pool overflow (%d page allocations were not served), this will produce visual artifacts (missing shadow), increase the page pool limit or reduce resolution bias to avoid.\n")
-					TEXT(" See r.Shadow.Virtual.MaxPhysicalPages (%d), r.Shadow.Virtual.ResolutionLodBiasLocal (%.2f), and r.Shadow.Virtual.ResolutionLodBiasDirectional (%.2f)"),
-					-NumPagesFree,
+					TEXT(" See r.Shadow.Virtual.MaxPhysicalPages (%d), r.Shadow.Virtual.ResolutionLodBiasLocal (%.2f), r.Shadow.Virtual.ResolutionLodBiasDirectional (%.2f), Global Resolution Lod Bias (%.2f)"),
+					-LastFreePhysicalPages,
 					MaxPhysicalPages,
-					LodBiasLocal,
-					LodBiasDirectional);
+					CVarResolutionLodBiasLocalPtr->GetValueOnRenderThread(),
+					CVarResolutionLodBiasDirectionalPtr->GetValueOnRenderThread(),
+					GlobalResolutionLodBias);
+
 				bLoggedPageOverflow = true;
 			}
 			LastOverflowTime = float(FGameTime::GetTimeSinceAppStart().GetRealTimeSeconds());
@@ -600,7 +648,6 @@ FVirtualShadowMapArrayCacheManager::FVirtualShadowMapArrayCacheManager(FScene* I
 		float RealTimeSeconds = float(FGameTime::GetTimeSinceAppStart().GetRealTimeSeconds());
 
 		// Show for ~5s after last overflow
-		int32 CurrentFrameNumber = Scene->GetFrameNumber();
 		if (LastOverflowTime >= 0.0f && RealTimeSeconds - LastOverflowTime < 5.0f)
 		{
 			OutMessages.Add(FCoreDelegates::EOnScreenMessageSeverity::Warning, FText::FromString(FString::Printf(TEXT("Virtual Shadow Map Page Pool overflow detected (%0.0f seconds ago)"), RealTimeSeconds - LastOverflowTime)));
@@ -628,7 +675,7 @@ FVirtualShadowMapArrayCacheManager::~FVirtualShadowMapArrayCacheManager()
 }
 
 
-void FVirtualShadowMapArrayCacheManager::SetPhysicalPoolSize(FRDGBuilder& GraphBuilder, FIntPoint RequestedSize, int RequestedArraySize, uint32 MaxPhysicalPages)
+void FVirtualShadowMapArrayCacheManager::SetPhysicalPoolSize(FRDGBuilder& GraphBuilder, FIntPoint RequestedSize, int RequestedArraySize, uint32 RequestedMaxPhysicalPages)
 {
 	// Using ReservedResource|ImmediateCommit flags hint to the RHI that the resource can be allocated using N small physical memory allocations,
 	// instead of a single large contighous allocation. This helps Windows video memory manager page allocations in and out of local memory more efficiently.
@@ -639,6 +686,7 @@ void FVirtualShadowMapArrayCacheManager::SetPhysicalPoolSize(FRDGBuilder& GraphB
 	if (!PhysicalPagePool 
 		|| PhysicalPagePool->GetDesc().Extent != RequestedSize 
 		|| PhysicalPagePool->GetDesc().ArraySize != RequestedArraySize
+		|| RequestedMaxPhysicalPages != MaxPhysicalPages
 		|| PhysicalPagePoolCreateFlags != RequestedCreateFlags)
 	{
 		if (PhysicalPagePool)
@@ -669,6 +717,8 @@ void FVirtualShadowMapArrayCacheManager::SetPhysicalPoolSize(FRDGBuilder& GraphB
 		);
 		GRenderTargetPool.FindFreeElement(GraphBuilder.RHICmdList, Desc2D, PhysicalPagePool, TEXT("Shadow.Virtual.PhysicalPagePool"));
 
+		MaxPhysicalPages = RequestedMaxPhysicalPages;
+
 		// Allocate page metadata alongside
 		FRDGBufferRef PhysicalPageMetaDataRDG = GraphBuilder.CreateBuffer(
 			FRDGBufferDesc::CreateStructuredDesc(sizeof(FPhysicalPageMetaData), MaxPhysicalPages),
diff --git a/Engine/Source/Runtime/Renderer/Private/VirtualShadowMaps/VirtualShadowMapCacheManager.h b/Engine/Source/Runtime/Renderer/Private/VirtualShadowMaps/VirtualShadowMapCacheManager.h
--- a/Engine/Source/Runtime/Renderer/Private/VirtualShadowMaps/VirtualShadowMapCacheManager.h
+++ b/Engine/Source/Runtime/Renderer/Private/VirtualShadowMaps/VirtualShadowMapCacheManager.h
@@ -321,6 +321,8 @@ public:
 	uint32 GetStatsFeedbackMessageId() const { return StatsFeedbackSocket.GetMessageId().IsValid() ? StatsFeedbackSocket.GetMessageId().GetIndex() : INDEX_NONE; }
 #endif
 
+	float GetGlobalResolutionLodBias() const { return GlobalResolutionLodBias; }
+
 private:
 	// Invalidate the cache for all shadows, causing any pages to be rerendered
 	void Invalidate(FRDGBuilder& GraphBuilder);
@@ -356,6 +358,7 @@ private:
 	TRefCountPtr<IPooledRenderTarget> HZBPhysicalPagePool;
 	ETextureCreateFlags PhysicalPagePoolCreateFlags = TexCreate_None;
 	TRefCountPtr<FRDGPooledBuffer> PhysicalPageMetaData;
+	uint32 MaxPhysicalPages = 0;
 
 	// Index the Cache entries by the light ID
 	TMap< uint64, TSharedPtr<FVirtualShadowMapPerLightCacheEntry> > CacheEntries;
@@ -377,6 +380,12 @@ private:
 	FVirtualShadowMapFeedback StaticGPUInvalidationsFeedback;
 	GPUMessage::FSocket StatusFeedbackSocket;
 
+	// Current global resolution bias (when enabled) based on feedback from page pressure, etc.
+	float GlobalResolutionLodBias = 0.0f;
+	uint32 LastFrameOverPageAllocationBudget = 0;
+	
+	
+
 	// Debug stuff
 #if !UE_BUILD_SHIPPING
 	FDelegateHandle ScreenMessageDelegate;
