Fixes a bug where, when using instanced stereo, meshes occluded in the primary view would not appear in the secondary view.
When ISR is enabled, primary view parallel occlusion culling now passes all meshes through to the relevance pass, but updates the primitive visibility map accordingly. Before the relevance pass starts, visbility is synchronized with the secondary view, and we then filter based on the visbility map in the relevance pass.

diff --git a/Engine/Source/Runtime/Renderer/Private/SceneVisibility.cpp b/Engine/Source/Runtime/Renderer/Private/SceneVisibility.cpp
--- a/Engine/Source/Runtime/Renderer/Private/SceneVisibility.cpp
+++ b/Engine/Source/Runtime/Renderer/Private/SceneVisibility.cpp
@@ -1871,7 +1871,11 @@ FComputeAndMarkRelevance::FComputeAndMarkRelevance(FVisibilityTaskData& InTaskDa
 
 void FComputeAndMarkRelevance::AddPrimitives(FPrimitiveIndexList&& PrimitiveIndexList)
 {
-	if (PrimitiveIndexList.Num() == NumPrimitivesPerPacket)
+	// In ISR only, primitives that have been occlusion culled in the primary view are still piped in because they may be visible in secondary views.
+	// They can be filtered out now using the visibility map, which is merged with that of the secondary view once occlusion tasks are complete.
+	const bool bIsPrimaryISRView = View.bIsMultiViewportEnabled && View.StereoPass == EStereoscopicPass::eSSP_PRIMARY;
+
+	if (!bIsPrimaryISRView && PrimitiveIndexList.Num() == NumPrimitivesPerPacket)
 	{
 		// Create a one-off packet that will take all the primitives.
 		FRelevancePacket* Packet = CreateRelevancePacket();
@@ -1889,7 +1893,10 @@ void FComputeAndMarkRelevance::AddPrimitives(FPrimitiveIndexList&& PrimitiveInde
 	{
 		for (int32 Index : PrimitiveIndexList)
 		{
-			AddPrimitive(Index);
+			if (!bIsPrimaryISRView || View.PrimitiveVisibilityMap[Index])
+			{
+				AddPrimitive(Index);
+			}
 		}
 	}
 }
@@ -2877,9 +2884,14 @@ FOcclusionCullResult FGPUOcclusionParallelPacket::OcclusionCullTask(FPrimitiveIn
 
 	PrimitiveIndexList.Reserve(Input.Num());
 
+	const bool bIsPrimaryISRView = View.bIsMultiViewportEnabled && View.StereoPass == EStereoscopicPass::eSSP_PRIMARY;
+
 	for (int32 Index : Input)
 	{
-		if (OcclusionCullPrimitive(RecordVisitor, Result, Index))
+		// When in ISR, primary stereo views can have additional visible primitives derived from secondary views.
+		// In that case, forward the primary view primitives down the pipe even if they're occluded.
+		// We also need to update the visibility map so we can later filter primitives culled in both views.
+		if (OcclusionCullPrimitive(RecordVisitor, Result, Index) || bIsPrimaryISRView)
 		{
 			PrimitiveIndexList.Emplace(Index);
 		}
@@ -3471,7 +3483,7 @@ void FVisibilityTaskData::LaunchVisibilityTasks()
 
 	if (TaskConfig.Schedule == EVisibilityTaskSchedule::Parallel)
 	{
-		if (Views.Num() == 1)
+		if (Views.Num() == 1 && !Views[0].bIsMultiViewportEnabled)
 		{
 			// When using a single view, dynamic mesh elements are pushed into a pipe that is executed on the render thread which allows for some overlap with compute relevance work.
 			DynamicMeshElements.CommandPipe = Allocator.Create<TCommandPipe<FDynamicPrimitiveIndexList>>(TEXT("GatherDynamicMeshElements"), ENamedThreads::GetRenderThread_Local());
@@ -3493,25 +3505,22 @@ void FVisibilityTaskData::LaunchVisibilityTasks()
 		}
 		else
 		{
-			TArray<UE::Tasks::FTask, SceneRenderingAllocator> SecondaryOcclusionCullTasks;
+			TArray<UE::Tasks::FTask, SceneRenderingAllocator> OcclusionCullTasks;
 
 			// Instanced multi-view scenarios have secondary views which feed visibility data into primary views. In this
-			// case we make secondary views finish occlusion culling first before launching a task that merges visibility
-			// into the primary view, which then becomes a prerequisite task for processing occlusion pipe commands in the
+			// case we make all views finish occlusion culling first before launching a task that merges visibility
+			// into the primary view, which then becomes a prerequisite task for processing relevance pipe commands in the
 			// primary view.
 
 			for (FVisibilityViewPacket& ViewPacket : ViewPackets)
 			{
-				if (ViewPacket.View.StereoPass == EStereoscopicPass::eSSP_SECONDARY)
-				{
-					SecondaryOcclusionCullTasks.Emplace(ViewPacket.Tasks.OcclusionCull);
-				}
+				OcclusionCullTasks.Emplace(ViewPacket.Tasks.OcclusionCull);
 			}
 
-			if (!SecondaryOcclusionCullTasks.IsEmpty())
+			if (!OcclusionCullTasks.IsEmpty())
 			{
 				// Wait for frustum culling to complete to avoid race conditions with writing to the visibility bits.
-				SecondaryOcclusionCullTasks.Emplace(Tasks.FrustumCull);
+				OcclusionCullTasks.Emplace(Tasks.FrustumCull);
 
 				FGraphEventRef MergeSecondaryViewsTask = FGraphEvent::CreateGraphEvent();
 
@@ -3520,7 +3529,7 @@ void FVisibilityTaskData::LaunchVisibilityTasks()
 					MergeSecondaryViewVisibility();
 					MergeSecondaryViewsTask->DispatchSubsequents();
 
-				}, SecondaryOcclusionCullTasks, TaskConfig.OcclusionCull.TaskPriority);
+				}, OcclusionCullTasks, TaskConfig.OcclusionCull.TaskPriority);
 
 				for (FVisibilityViewPacket& ViewPacket : ViewPackets)
 				{
@@ -3528,8 +3537,8 @@ void FVisibilityTaskData::LaunchVisibilityTasks()
 
 					if (View.StereoPass == EStereoscopicPass::eSSP_PRIMARY)
 					{
-						// Force the primary view occlusion cull pipe to wait until the merge task completes.
-						ViewPacket.OcclusionCull.CommandPipe.SetPrerequisiteTask(MergeSecondaryViewsTask);
+						// Force the primary view relevance pipe to wait until the merge task completes.
+						ViewPacket.Relevance.CommandPipe.SetPrerequisiteTask(MergeSecondaryViewsTask);
 					}
 				}
 			}
