Expand Heterogeneous Volumes MaxTraceDistance and MaxShadowTraceDistance to 30k units.
Allow lighting down-sample factor scale below 1.

diff --git a/Engine/Source/Runtime/Renderer/Private/HeterogeneousVolumes/HeterogeneousVolumes.cpp b/Engine/Source/Runtime/Renderer/Private/HeterogeneousVolumes/HeterogeneousVolumes.cpp
--- a/Engine/Source/Runtime/Renderer/Private/HeterogeneousVolumes/HeterogeneousVolumes.cpp
+++ b/Engine/Source/Runtime/Renderer/Private/HeterogeneousVolumes/HeterogeneousVolumes.cpp
@@ -55,15 +55,15 @@ static TAutoConsoleVariable<int32> CVarHeterogeneousVolumesMaxStepCount(
 
 static TAutoConsoleVariable<float> CVarHeterogeneousVolumesMaxTraceDistance(
 	TEXT("r.HeterogeneousVolumes.MaxTraceDistance"),
-	10000.0,
-	TEXT("The maximum trace view-distance for direct volume rendering (Default = 10000)"),
+	30000.0,
+	TEXT("The maximum trace view-distance for direct volume rendering (Default = 30000)"),
 	ECVF_RenderThreadSafe
 );
 
 static TAutoConsoleVariable<float> CVarHeterogeneousVolumesMaxShadowTraceDistance(
 	TEXT("r.HeterogeneousVolumes.MaxShadowTraceDistance"),
-	10000.0,
-	TEXT("The maximum shadow-trace distance (Default = 10000)"),
+	30000.0,
+	TEXT("The maximum shadow-trace distance (Default = 30000)"),
 	ECVF_RenderThreadSafe
 );
 
@@ -372,12 +372,13 @@ namespace HeterogeneousVolumes
 	{
 		float OverrideDownsampleFactor = CVarHeterogeneousVolumesLightingCacheDownsampleFactor.GetValueOnRenderThread();
 		float DownsampleFactor = OverrideDownsampleFactor > 0.0 ? OverrideDownsampleFactor : RenderInterface->GetLightingDownsampleFactor();
-		DownsampleFactor = FMath::Max(DownsampleFactor, 1.0);
+		DownsampleFactor = FMath::Max(DownsampleFactor, 0.125);
 
-		FIntVector LightingCacheResolution = GetVolumeResolution(RenderInterface) / DownsampleFactor;
+		FVector VolumeResolution = FVector(GetVolumeResolution(RenderInterface));
+		FIntVector LightingCacheResolution = FIntVector(VolumeResolution / DownsampleFactor);
 		LightingCacheResolution.X = FMath::Clamp(LightingCacheResolution.X, 1, 1024);
 		LightingCacheResolution.Y = FMath::Clamp(LightingCacheResolution.Y, 1, 1024);
-		LightingCacheResolution.Z = FMath::Clamp(LightingCacheResolution.Z, 1, 1024);
+		LightingCacheResolution.Z = FMath::Clamp(LightingCacheResolution.Z, 1, 512);
 		return LightingCacheResolution;
 	}
 }
