Fix single layer water color banding issue on some platforms

diff --git a/Engine/Source/Runtime/Renderer/Private/SingleLayerWaterRendering.cpp b/Engine/Source/Runtime/Renderer/Private/SingleLayerWaterRendering.cpp
--- a/Engine/Source/Runtime/Renderer/Private/SingleLayerWaterRendering.cpp
+++ b/Engine/Source/Runtime/Renderer/Private/SingleLayerWaterRendering.cpp
@@ -425,6 +425,7 @@ class FWaterRefractionCopyPS : public FGlobalShader
 	static void ModifyCompilationEnvironment(const FGlobalShaderPermutationParameters& Parameters, FShaderCompilerEnvironment& OutEnvironment)
 	{
 		FGlobalShader::ModifyCompilationEnvironment(Parameters, OutEnvironment);
+		OutEnvironment.SetRenderTargetOutputFormat(0, PF_R32_FLOAT);
 	}
 };
 
@@ -721,6 +722,7 @@ static FSceneWithoutWaterTextures AddCopySceneWithoutWaterPass(
 		SceneColorWithoutSingleLayerWaterTexture = GraphBuilder.CreateTexture(ColorDesc, TEXT("SLW.SceneColorWithout"));
 	}
 
+	// Note: if changing format, also update FWaterRefractionCopyPS::ModifyCompilationEnvironment accordingly
 	const FRDGTextureDesc DepthDesc(FRDGTextureDesc::Create2D(RefractionResolution, PF_R32_FLOAT, SceneDepthDesc.ClearValue, TexCreate_ShaderResource | TexCreate_RenderTargetable));
 	FRDGTextureRef SceneDepthWithoutSingleLayerWaterTexture = GraphBuilder.CreateTexture(DepthDesc, TEXT("SLW.SceneDepthWithout"));
 
