Fixed: GPU crash on mobile platforms related to incorrect autop-merging of ISM drawcalls

diff --git a/Engine/Source/Runtime/Renderer/Private/InstanceCulling/InstanceCullingContext.cpp b/Engine/Source/Runtime/Renderer/Private/InstanceCulling/InstanceCullingContext.cpp
--- a/Engine/Source/Runtime/Renderer/Private/InstanceCulling/InstanceCullingContext.cpp
+++ b/Engine/Source/Runtime/Renderer/Private/InstanceCulling/InstanceCullingContext.cpp
@@ -1424,7 +1424,7 @@ void FInstanceCullingContext::AddClearIndirectArgInstanceCountPass(FRDGBuilder&
  */
 void FInstanceCullingContext::SetupDrawCommands(
 	FMeshCommandOneFrameArray& VisibleMeshDrawCommandsInOut,
-	bool bCompactIdenticalCommands,
+	bool bInCompactIdenticalCommands,
 	// Stats
 	int32& MaxInstances,
 	int32& VisibleMeshDrawCommandsNum,
@@ -1490,8 +1490,10 @@ void FInstanceCullingContext::SetupDrawCommands(
 		const bool bForceInstanceCulling = EnumHasAnyFlags(VisibleMeshDrawCommand.Flags, EFVisibleMeshDrawCommandFlags::ForceInstanceCulling);
 		const bool bPreserveInstanceOrder = bOrderPreservationEnabled && EnumHasAnyFlags(VisibleMeshDrawCommand.Flags, EFVisibleMeshDrawCommandFlags::PreserveInstanceOrder);
 		const bool bUseIndirectDraw = bAlwaysUseIndirectDraws || bForceInstanceCulling || (VisibleMeshDrawCommand.NumRuns > 0 || MeshDrawCommand->NumInstances > 1);
+		// UniformBufferView path does not support merging ISM draws atm
+		const bool bCompactIdenticalCommands = bInCompactIdenticalCommands && (bUsesUniformBufferView ? (CurrentAutoInstanceCount < MaxPrimitiveBatchSize && !bUseIndirectDraw) : true);
 
-		if (bCompactIdenticalCommands && CurrentStateBucketId != -1 && VisibleMeshDrawCommand.StateBucketId == CurrentStateBucketId && CurrentAutoInstanceCount < MaxPrimitiveBatchSize)
+		if (bCompactIdenticalCommands && CurrentStateBucketId != -1 && VisibleMeshDrawCommand.StateBucketId == CurrentStateBucketId)
 		{
 			// Drop since previous covers for this
 
@@ -1614,7 +1616,7 @@ void FInstanceCullingContext::SetupDrawCommands(
 			}
 		}
 	}
-	check(bCompactIdenticalCommands || NumDrawCommandsIn == NumDrawCommandsOut);
+	check(bInCompactIdenticalCommands || NumDrawCommandsIn == NumDrawCommandsOut);
 	checkf(NumDrawCommandsOut == MeshDrawCommandInfos.Num(), TEXT("There must be a 1:1 mapping between MeshDrawCommandInfos and mesh draw commands, as this assumption is made in SubmitDrawCommands."));
 
 	// Setup instancing stats for logging.
