PR #10728: Apply pre-exposure correction in volumetric RT reconstruction

diff --git a/Engine/Shaders/Private/VolumetricRenderTarget.usf b/Engine/Shaders/Private/VolumetricRenderTarget.usf
--- a/Engine/Shaders/Private/VolumetricRenderTarget.usf
+++ b/Engine/Shaders/Private/VolumetricRenderTarget.usf
@@ -74,13 +74,17 @@ float4 PreviousVolumetricTextureSizeAndInvSize;
 uint4 PreviousFrameVolumetricTextureValidCoordRect;
 float4 PreviousFrameVolumetricTextureValidUvRect;
 
+float HistoryPreExposureCorrection;
+
 float4 SafeLoadPreviousFrameVolumetricTexture(uint2 Coord)
 {
-	return PreviousFrameVolumetricTexture.Load(uint3(clamp(Coord, PreviousFrameVolumetricTextureValidCoordRect.xy, PreviousFrameVolumetricTextureValidCoordRect.zw), 0));
+	float4 Result = PreviousFrameVolumetricTexture.Load(uint3(clamp(Coord, PreviousFrameVolumetricTextureValidCoordRect.xy, PreviousFrameVolumetricTextureValidCoordRect.zw), 0));
+	return float4(Result.rgb * HistoryPreExposureCorrection, Result.a);
 }
 float4 SafeSamplePreviousFrameVolumetricTexture(float2 UV)
 {
-	return PreviousFrameVolumetricTexture.SampleLevel(LinearTextureSampler, clamp(UV, PreviousFrameVolumetricTextureValidUvRect.xy, PreviousFrameVolumetricTextureValidUvRect.zw), 0);
+	float4 Result = PreviousFrameVolumetricTexture.SampleLevel(LinearTextureSampler, clamp(UV, PreviousFrameVolumetricTextureValidUvRect.xy, PreviousFrameVolumetricTextureValidUvRect.zw), 0);
+	return float4(Result.rgb * HistoryPreExposureCorrection, Result.a);
 }
 FDepthData SafeLoadPreviousFrameVolumetricDepthTexture(uint2 Coord)
 {
diff --git a/Engine/Source/Runtime/Renderer/Private/VolumetricRenderTarget.cpp b/Engine/Source/Runtime/Renderer/Private/VolumetricRenderTarget.cpp
--- a/Engine/Source/Runtime/Renderer/Private/VolumetricRenderTarget.cpp
+++ b/Engine/Source/Runtime/Renderer/Private/VolumetricRenderTarget.cpp
@@ -522,6 +522,7 @@ class FReconstructVolumetricRenderTargetPS : public FGlobalShader
 		SHADER_PARAMETER(FUintVector4, PreviousFrameVolumetricTextureValidCoordRect)
 		SHADER_PARAMETER(FVector4f, PreviousFrameVolumetricTextureValidUvRect)
 		SHADER_PARAMETER(float, MinimumDistanceKmToEnableReprojection)
+		SHADER_PARAMETER(float, HistoryPreExposureCorrection)
 	END_SHADER_PARAMETER_STRUCT()
 
 	static FPermutationDomain RemapPermutation(FPermutationDomain PermutationVector)
@@ -613,6 +614,7 @@ void ReconstructVolumetricRenderTarget(
 		PassParameters->VolumetricRenderTargetMode = VolumetricCloudRT.GetMode();
 		PassParameters->HalfResDepthTexture = (VolumetricCloudRT.GetMode() == 0 || VolumetricCloudRT.GetMode() == 3) ? HalfResolutionDepthCheckerboardMinMaxTexture : SceneDepthTexture;
 		PassParameters->MinimumDistanceKmToEnableReprojection = FMath::Max(0.0f, CVarVolumetricRenderTargetMinimumDistanceKmToEnableReprojection.GetValueOnRenderThread());
+		PassParameters->HistoryPreExposureCorrection = ViewInfo.PreExposure / ViewInfo.PrevViewInfo.SceneColorPreExposure;
 
 		const bool bVisualizeConservativeDensity = ShouldViewVisualizeVolumetricCloudConservativeDensity(ViewInfo, ViewInfo.Family->EngineShowFlags);
 		PassParameters->HalfResDepthTexture = bVisualizeConservativeDensity ?
