Rename cvar from last commit.

diff --git a/Engine/Source/Runtime/Renderer/Private/VirtualShadowMaps/VirtualShadowMapCacheManager.cpp b/Engine/Source/Runtime/Renderer/Private/VirtualShadowMaps/VirtualShadowMapCacheManager.cpp
--- a/Engine/Source/Runtime/Renderer/Private/VirtualShadowMaps/VirtualShadowMapCacheManager.cpp
+++ b/Engine/Source/Runtime/Renderer/Private/VirtualShadowMaps/VirtualShadowMapCacheManager.cpp
@@ -93,8 +93,8 @@ static TAutoConsoleVariable<float> CVarVSMDynamicResolutionMaxLodBias(
 	ECVF_RenderThreadSafe
 );
 
-static TAutoConsoleVariable<float> CVarVSMDynamicResolutionMaxPagePoolAllocation(
-	TEXT("r.Shadow.Virtual.DynamicRes.MaxPagePoolAllocation"),
+static TAutoConsoleVariable<float> CVarVSMDynamicResolutionMaxPagePoolLoadFactor(
+	TEXT("r.Shadow.Virtual.DynamicRes.MaxPagePoolLoadFactor"),
 	0.85f,
 	TEXT("If allocation exceeds this factor of total page pool capacity, shadow resolution will be biased downwards. 0 = disabled"),
 	ECVF_RenderThreadSafe
@@ -549,7 +549,7 @@ FVirtualShadowMapArrayCacheManager::FVirtualShadowMapArrayCacheManager(FScene* I
 			const float ResolutionUpExpLerpFactor = 0.1f;
 			const uint32 FramesBeforeResolutionUp = 10;
 
-			const float MaxPageAllocation = CVarVSMDynamicResolutionMaxPagePoolAllocation.GetValueOnRenderThread();
+			const float MaxPageAllocation = CVarVSMDynamicResolutionMaxPagePoolLoadFactor.GetValueOnRenderThread();
 			const float MaxLodBias = CVarVSMDynamicResolutionMaxLodBias.GetValueOnRenderThread();
 			
 			if (MaxPageAllocation > 0.0f)
