Extend instance culling occlusion query bounding boxes by OCCLUSION_SLOP parameter, similar to what's done for standard occlusion queries

diff --git a/Engine/Shaders/Private/InstanceCulling/InstanceCullingOcclusionQuery.usf b/Engine/Shaders/Private/InstanceCulling/InstanceCullingOcclusionQuery.usf
--- a/Engine/Shaders/Private/InstanceCulling/InstanceCullingOcclusionQuery.usf
+++ b/Engine/Shaders/Private/InstanceCulling/InstanceCullingOcclusionQuery.usf
@@ -29,6 +29,7 @@ RWBuffer<uint> OutVisibilityMask;
 uint NumInstances;
 uint InstanceSceneDataSOAStride;
 int4 ViewRect;
+float OcclusionSlop;
 
 #define EInstanceVisibilityStatus uint // Not all platforms support native enums
 #define IVS_Hidden          0 // Instance is hidden or invalid, no occlusion query needed
@@ -60,7 +61,7 @@ EInstanceVisibilityStatus GetInstanceDataAndVisibility(uint InstanceId, out FIns
 	}
 
 	float3 BoxOrigin = InstanceData.LocalBoundsCenter;
-	float3 BoxExtent = InstanceData.LocalBoundsExtent;
+	float3 BoxExtent = InstanceData.LocalBoundsExtent + OcclusionSlop;
 
 	const bool bIsOrtho = IsOrthoProjection(ResolvedView.ViewToClip);
 
@@ -229,7 +230,7 @@ void MainVS(
 	else // Status == IVS_PossiblyVisible
 #endif // USE_COMPUTE_INSTANCE_CULLING_SETUP
 	{
-		float3 LocalPosition = (InPosition * InstanceData.LocalBoundsExtent * InstanceData.DeterminantSign) + InstanceData.LocalBoundsCenter;
+		float3 LocalPosition = (InPosition * (InstanceData.LocalBoundsExtent + OcclusionSlop) * InstanceData.DeterminantSign) + InstanceData.LocalBoundsCenter;
 		float4 WorldPosition = TransformLocalToTranslatedWorld(LocalPosition, InstanceData.LocalToWorld);
 		OutPosition = mul(WorldPosition, ResolvedView.TranslatedWorldToClip);
 	}
@@ -268,7 +269,7 @@ void DebugMainVS(
 	
 	if (InstanceData.ValidInstance && VisibilityStatus != IVS_Incompatible)
 	{
-		float3 LocalPosition = (InPosition * InstanceData.LocalBoundsExtent * InstanceData.DeterminantSign) + InstanceData.LocalBoundsCenter;
+		float3 LocalPosition = (InPosition * (InstanceData.LocalBoundsExtent + OcclusionSlop) * InstanceData.DeterminantSign) + InstanceData.LocalBoundsCenter;
 		float4 WorldPosition = TransformLocalToTranslatedWorld(LocalPosition, InstanceData.LocalToWorld);
 
 		OutPosition = mul(WorldPosition, ResolvedView.TranslatedWorldToClip);
diff --git a/Engine/Source/Runtime/Renderer/Private/InstanceCulling/InstanceCullingOcclusionQuery.cpp b/Engine/Source/Runtime/Renderer/Private/InstanceCulling/InstanceCullingOcclusionQuery.cpp
--- a/Engine/Source/Runtime/Renderer/Private/InstanceCulling/InstanceCullingOcclusionQuery.cpp
+++ b/Engine/Source/Runtime/Renderer/Private/InstanceCulling/InstanceCullingOcclusionQuery.cpp
@@ -99,6 +99,7 @@ public:
 		SHADER_PARAMETER_SAMPLER(SamplerState, HZBSampler)
 		SHADER_PARAMETER(FVector2f, HZBSize)
 		SHADER_PARAMETER(FIntVector4, ViewRect)
+		SHADER_PARAMETER(float, OcclusionSlop)
 		SHADER_PARAMETER(int32, NumInstances)
 	END_SHADER_PARAMETER_STRUCT()
 };
@@ -135,6 +136,7 @@ public:
 		SHADER_PARAMETER_SAMPLER(SamplerState, HZBSampler)
 		SHADER_PARAMETER(FVector2f, HZBSize)
 		SHADER_PARAMETER(FIntVector4, ViewRect)
+		SHADER_PARAMETER(float, OcclusionSlop)
 	END_SHADER_PARAMETER_STRUCT()
 };
 
@@ -354,6 +356,7 @@ uint32 FInstanceCullingOcclusionQueryRenderer::Render(
 		PassParameters->HZBSampler = TStaticSamplerState<SF_Point, AM_Clamp, AM_Clamp, AM_Clamp>::GetRHI();
 		PassParameters->HZBSize = FVector2f(HZBSize.X, HZBSize.Y);
 		PassParameters->ViewRect = FIntVector4(View.ViewRect.Min.X, View.ViewRect.Min.Y, View.ViewRect.Max.X, View.ViewRect.Max.Y);
+		PassParameters->OcclusionSlop = OCCLUSION_SLOP;
 		PassParameters->InstanceSceneDataSOAStride = GPUSceneParameters.InstanceDataSOAStride;
 		PassParameters->GPUSceneInstanceSceneData = GPUSceneParameters.GPUSceneInstanceSceneData;
 		PassParameters->GPUSceneInstancePayloadData = GPUSceneParameters.GPUSceneInstancePayloadData;
@@ -383,6 +386,7 @@ uint32 FInstanceCullingOcclusionQueryRenderer::Render(
 		PassParameters->VS.HZBSampler = TStaticSamplerState<SF_Point, AM_Clamp, AM_Clamp, AM_Clamp>::GetRHI();
 		PassParameters->VS.HZBSize = FVector2f(HZBSize.X, HZBSize.Y);
 		PassParameters->VS.ViewRect = FIntVector4(View.ViewRect.Min.X, View.ViewRect.Min.Y, View.ViewRect.Max.X, View.ViewRect.Max.Y);
+		PassParameters->VS.OcclusionSlop = OCCLUSION_SLOP;
 		PassParameters->VS.InstanceSceneDataSOAStride = GPUSceneParameters.InstanceDataSOAStride;
 		PassParameters->VS.GPUSceneInstanceSceneData = GPUSceneParameters.GPUSceneInstanceSceneData;
 		PassParameters->VS.GPUSceneInstancePayloadData = GPUSceneParameters.GPUSceneInstancePayloadData;
@@ -494,6 +498,7 @@ public:
 		SHADER_PARAMETER_SAMPLER(SamplerState, HZBSampler)
 		SHADER_PARAMETER(FVector2f, HZBSize)
 		SHADER_PARAMETER(FIntVector4, ViewRect)
+		SHADER_PARAMETER(float, OcclusionSlop)
 	END_SHADER_PARAMETER_STRUCT()
 };
 
@@ -585,6 +590,7 @@ void FInstanceCullingOcclusionQueryRenderer::RenderDebug(FRDGBuilder& GraphBuild
 
 	FOcclusionInstanceCullingDebugParameters* PassParameters = GraphBuilder.AllocParameters<FOcclusionInstanceCullingDebugParameters>();
 
+	PassParameters->VS.OcclusionSlop = OCCLUSION_SLOP;
 	PassParameters->VS.ViewRect = FIntVector4(View.ViewRect.Min.X, View.ViewRect.Min.Y, View.ViewRect.Max.X, View.ViewRect.Max.Y);
 	PassParameters->VS.View = View.ViewUniformBuffer;
 	PassParameters->VS.InstanceSceneDataSOAStride = GPUSceneParameters.InstanceDataSOAStride;
