Fix bad include causing CIS errors.

diff --git a/Engine/Source/Runtime/Renderer/Private/SplineMeshSceneResources.h b/Engine/Source/Runtime/Renderer/Private/SplineMeshSceneResources.h
--- a/Engine/Source/Runtime/Renderer/Private/SplineMeshSceneResources.h
+++ b/Engine/Source/Runtime/Renderer/Private/SplineMeshSceneResources.h
@@ -7,7 +7,7 @@
 #include "SceneUniformBuffer.h"
 #include "RHIShaderPlatform.h"
 #include "RendererInterface.h"
-#include "Containers\Map.h"
+#include "Containers/Map.h"
 
 class FRDGBuilder;
 class FScene;
