Fix CIS warnings

diff --git a/Engine/Source/Runtime/Renderer/Private/SplineMeshSceneResources.cpp b/Engine/Source/Runtime/Renderer/Private/SplineMeshSceneResources.cpp
--- a/Engine/Source/Runtime/Renderer/Private/SplineMeshSceneResources.cpp
+++ b/Engine/Source/Runtime/Renderer/Private/SplineMeshSceneResources.cpp
@@ -98,7 +98,7 @@ class FSplineMeshTextureFillCS : public FGlobalShader
 
 	BEGIN_SHADER_PARAMETER_STRUCT(FParameters, )
 		SHADER_PARAMETER_RDG_UNIFORM_BUFFER(FSceneUniformParameters, Scene)
-		SHADER_PARAMETER_RDG_TEXTURE_UAV(RWTexture<float3>, SplineTextureOut)
+		SHADER_PARAMETER_RDG_TEXTURE_UAV(RWTexture<float4>, SplineTextureOut)
 		SHADER_PARAMETER_RDG_BUFFER_SRV(StructuredBuffer<uint>, InstanceIdLookup)
 		SHADER_PARAMETER_RDG_BUFFER_SRV(StructuredBuffer<uint>, UpdateRequests)
 		SHADER_PARAMETER(uint32, NumUpdateRequests)
