Fix for the scenedepth crash seen when using the mobile renderer.  The scenedepth texture was being used before it was rendered to.

diff --git a/Engine/Source/Runtime/Renderer/Private/MobileShadingRenderer.cpp b/Engine/Source/Runtime/Renderer/Private/MobileShadingRenderer.cpp
--- a/Engine/Source/Runtime/Renderer/Private/MobileShadingRenderer.cpp
+++ b/Engine/Source/Runtime/Renderer/Private/MobileShadingRenderer.cpp
@@ -2002,44 +2002,46 @@ void FMobileSceneRenderer::RenderHZB(FRDGBuilder& GraphBuilder, FRDGTextureRef S
 	for (int32 ViewIndex = 0; ViewIndex < Views.Num(); ViewIndex++)
 	{
 		FViewInfo& View = Views[ViewIndex];
-
-		RDG_GPU_MASK_SCOPE(GraphBuilder, View.GPUMask);
+		if (View.ShouldRenderView())
 		{
-			RDG_EVENT_SCOPE(GraphBuilder, "BuildHZB(ViewId=%d)", ViewIndex);
-			
-			FRDGTextureRef FurthestHZBTexture = nullptr;
-
-			BuildHZBFurthest(
-				GraphBuilder,
-				SceneDepthTexture,
-				/* VisBufferTexture = */ nullptr,
-				View.ViewRect,
-				View.GetFeatureLevel(),
-				View.GetShaderPlatform(),
-				TEXT("MobileHZBFurthest"),
-				&FurthestHZBTexture);
-
-			View.HZBMipmap0Size = FurthestHZBTexture->Desc.Extent;
-			View.HZB = FurthestHZBTexture;
-
-			if (View.ViewState)
+			RDG_GPU_MASK_SCOPE(GraphBuilder, View.GPUMask);
 			{
-				if (FInstanceCullingContext::IsOcclusionCullingEnabled())
-				{
-					GraphBuilder.QueueTextureExtraction(FurthestHZBTexture, &View.ViewState->PrevFrameViewInfo.HZB);
-				}
-				else
+				RDG_EVENT_SCOPE(GraphBuilder, "BuildHZB(ViewId=%d)", ViewIndex);
+
+				FRDGTextureRef FurthestHZBTexture = nullptr;
+
+				BuildHZBFurthest(
+					GraphBuilder,
+					SceneDepthTexture,
+					/* VisBufferTexture = */ nullptr,
+					View.ViewRect,
+					View.GetFeatureLevel(),
+					View.GetShaderPlatform(),
+					TEXT("MobileHZBFurthest"),
+					&FurthestHZBTexture);
+
+				View.HZBMipmap0Size = FurthestHZBTexture->Desc.Extent;
+				View.HZB = FurthestHZBTexture;
+
+				if (View.ViewState)
 				{
-					View.ViewState->PrevFrameViewInfo.HZB = nullptr;
+					if (FInstanceCullingContext::IsOcclusionCullingEnabled())
+					{
+						GraphBuilder.QueueTextureExtraction(FurthestHZBTexture, &View.ViewState->PrevFrameViewInfo.HZB);
+					}
+					else
+					{
+						View.ViewState->PrevFrameViewInfo.HZB = nullptr;
+					}
 				}
 			}
-		}
 
-		if (Scene->InstanceCullingOcclusionQueryRenderer && View.ViewState)
-		{
-			// Render per-instance occlusion queries and save the mask to interpret results on the next frame
-			const uint32 OcclusionQueryMaskForThisView = Scene->InstanceCullingOcclusionQueryRenderer->Render(GraphBuilder, Scene->GPUScene, View);
-			View.ViewState->PrevFrameViewInfo.InstanceOcclusionQueryMask = OcclusionQueryMaskForThisView;
+			if (Scene->InstanceCullingOcclusionQueryRenderer && View.ViewState)
+			{
+				// Render per-instance occlusion queries and save the mask to interpret results on the next frame
+				const uint32 OcclusionQueryMaskForThisView = Scene->InstanceCullingOcclusionQueryRenderer->Render(GraphBuilder, Scene->GPUScene, View);
+				View.ViewState->PrevFrameViewInfo.InstanceOcclusionQueryMask = OcclusionQueryMaskForThisView;
+			}
 		}
 	}
 }
