Added logging for cases where all mesh commands are forced to re-cache

diff --git a/Engine/Source/Runtime/Renderer/Private/RendererScene.cpp b/Engine/Source/Runtime/Renderer/Private/RendererScene.cpp
--- a/Engine/Source/Runtime/Renderer/Private/RendererScene.cpp
+++ b/Engine/Source/Runtime/Renderer/Private/RendererScene.cpp
@@ -2542,6 +2542,7 @@ void FScene::AddLightSceneInfo_RenderThread(FLightSceneInfo* LightSceneInfo)
 			    if (!LightSceneInfo->Proxy->HasStaticShadowing() || bUseCSMForDynamicObjects)
 				{
 		    		bScenesPrimitivesNeedStaticMeshElementUpdate = true;
+					UE_CLOG(!GIsEditor, LogRenderer, Log, TEXT("Forcing update for all mesh draw commands: Add directional light"));
 				}
 		    }
 		}
@@ -2666,6 +2667,7 @@ void FScene::SetSkyLight(FSkyLightSceneProxy* LightProxy)
 				// Mark the scene as needing static draw lists to be recreated if needed
 				// The base pass chooses shaders based on whether there's a skylight in the scene, and that is cached in static draw lists
 				Scene->bScenesPrimitivesNeedStaticMeshElementUpdate = true;
+				UE_CLOG(!GIsEditor, LogRenderer, Log, TEXT("Forcing update for all mesh draw commands: Enable SkyLight"));
 			}
 			Scene->InvalidatePathTracedOutput();
 		});
@@ -2701,6 +2703,7 @@ void FScene::DisableSkyLight(FSkyLightSceneProxy* LightProxy)
 		if (bOriginalHadSkylight != bNewHasSkylight)
 		{
 			Scene->bScenesPrimitivesNeedStaticMeshElementUpdate = true;
+			UE_CLOG(!GIsEditor, LogRenderer, Log, TEXT("Forcing update for all mesh draw commands: Disable SkyLight"));
 		}
 		Scene->InvalidatePathTracedOutput();
 	});
@@ -3803,6 +3806,7 @@ void FScene::RemoveLightSceneInfo_RenderThread(FLightSceneInfo* LightSceneInfo)
 				if (!LightSceneInfo->Proxy->HasStaticShadowing() || bUseCSMForDynamicObjects)
 				{
 					bScenesPrimitivesNeedStaticMeshElementUpdate = true;
+					UE_CLOG(!GIsEditor, LogRenderer, Log, TEXT("Forcing update for all mesh draw commands: Remove directional light"));
 				}
 				break;
 			}
@@ -5031,10 +5035,13 @@ FLightSceneChangeSet FScene::UpdateAllLightSceneInfos(FRDGBuilder& GraphBuilder)
 			// Mobile renderer:
 			// a light with no color/intensity can cause the light to be ignored when rendering.
 			// thus, lights that change state in this way must update the draw lists.
-			bScenesPrimitivesNeedStaticMeshElementUpdate =
-				bScenesPrimitivesNeedStaticMeshElementUpdate ||
-				(GetShadingPath() == EShadingPath::Mobile
-					&& NewParameters.NewColor.IsAlmostBlack() != LightSceneInfo->Proxy->GetColor().IsAlmostBlack());
+			if (GetShadingPath() == EShadingPath::Mobile 
+				&& LightSceneInfo->Proxy->GetLightType() == LightType_Directional 
+				&& NewParameters.NewColor.IsAlmostBlack() != LightSceneInfo->Proxy->GetColor().IsAlmostBlack())
+			{
+				bScenesPrimitivesNeedStaticMeshElementUpdate = true;
+				UE_CLOG(!GIsEditor, LogRenderer, Log, TEXT("Forcing update for all mesh draw commands: Toggle directional light"));
+			}
 
 			// Path Tracing: something about the light has changed, restart path traced accumulation
 			InvalidatePathTracedOutput();
