Print number of updated pages when using r.GlobalDistanceField.Debug.ShowStats.
- also renamed some global distance field debug cvars.

diff --git a/Engine/Shaders/Private/DistanceField/GlobalDistanceFieldDebug.usf b/Engine/Shaders/Private/DistanceField/GlobalDistanceFieldDebug.usf
--- a/Engine/Shaders/Private/DistanceField/GlobalDistanceFieldDebug.usf
+++ b/Engine/Shaders/Private/DistanceField/GlobalDistanceFieldDebug.usf
@@ -3,6 +3,7 @@
 #include "../Common.ush"
 #include "../ShaderPrint.ush"
 #include "GlobalDistanceFieldShared.ush"
+#include "GlobalDistanceFieldUpdate.ush"
 
 StructuredBuffer<uint> GlobalDistanceFieldPageFreeListAllocatorBuffer;
 uint GlobalDistanceFieldMaxPageNum;
@@ -16,6 +17,20 @@ Texture3D<uint> PageTableCombinedTexture;
 Texture3D PageAtlasTexture;
 SamplerState DistanceFieldSampler;
 
+StructuredBuffer<uint> PageComposeIndirectArgBuffer;
+
+[numthreads(1, 1, 1)]
+void GlobalDistanceFieldAccumulateUpdatedPagesCS(
+	uint3 GroupId : SV_GroupID,
+	uint3 DispatchThreadId : SV_DispatchThreadID,
+	uint3 GroupThreadId : SV_GroupThreadID)
+{
+	if (all(DispatchThreadId.xyz == 0))
+	{
+		InterlockedAdd(RWPageStatsBuffer[2], PageComposeIndirectArgBuffer[0] / GLOBAL_DISTANCE_FIELD_PAGE_GROUP_NUM);
+	}
+}
+
 [numthreads(THREADGROUP_SIZE, THREADGROUP_SIZE, THREADGROUP_SIZE)]
 void GlobalDistanceFieldPageStatsCS(
 	uint3 GroupId : SV_GroupID,
@@ -102,5 +117,10 @@ void GlobalDistanceFieldDebugCS(
 
 		Print(Context, TEXT(" Empty "));
 		Print(Context, PageStatsBuffer[1]);
+		Newline(Context);
+
+		Print(Context, TEXT(" Updated  "), FontOrange);
+		Print(Context, PageStatsBuffer[2], FontOrange);
+		Newline(Context);
 	}
 }
diff --git a/Engine/Source/Runtime/Renderer/Private/DistanceFieldObjectManagement.cpp b/Engine/Source/Runtime/Renderer/Private/DistanceFieldObjectManagement.cpp
--- a/Engine/Source/Runtime/Renderer/Private/DistanceFieldObjectManagement.cpp
+++ b/Engine/Source/Runtime/Renderer/Private/DistanceFieldObjectManagement.cpp
@@ -290,10 +290,10 @@ void ProcessDistanceFieldObjectRemoves(FScene* Scene, TArray<FSetElementId>& Dis
 
 void LogDistanceFieldUpdate(FPrimitiveSceneInfo const* PrimitiveSceneInfo, float BoundingRadius, bool bIsAddOperation)
 {
-	extern int32 GAOLogGlobalDistanceFieldModifiedPrimitives;
+	extern int32 GGlobalDistanceFieldDebugLogModifiedPrimitives;
 
-	if (GAOLogGlobalDistanceFieldModifiedPrimitives == 1
-		|| (GAOLogGlobalDistanceFieldModifiedPrimitives == 2 && !PrimitiveSceneInfo->Proxy->IsOftenMoving()))
+	if (GGlobalDistanceFieldDebugLogModifiedPrimitives == 1
+		|| (GGlobalDistanceFieldDebugLogModifiedPrimitives == 2 && !PrimitiveSceneInfo->Proxy->IsOftenMoving()))
 	{
 		UE_LOG(LogDistanceField, Log,
 			TEXT("Global Distance Field %s primitive %s %s %s bounding radius %.1f"),
diff --git a/Engine/Source/Runtime/Renderer/Private/GlobalDistanceField.cpp b/Engine/Source/Runtime/Renderer/Private/GlobalDistanceField.cpp
--- a/Engine/Source/Runtime/Renderer/Private/GlobalDistanceField.cpp
+++ b/Engine/Source/Runtime/Renderer/Private/GlobalDistanceField.cpp
@@ -107,33 +107,6 @@ FAutoConsoleVariableRef CVarAOGlobalDistanceFieldForceUpdateOnce(
 	ECVF_Scalability | ECVF_RenderThreadSafe
 );
 
-int32 GAOGlobalDistanceFieldForceMovementUpdate = 0;
-FAutoConsoleVariableRef CVarAOGlobalDistanceFieldForceMovementUpdate(
-	TEXT("r.AOGlobalDistanceFieldForceMovementUpdate"),
-	GAOGlobalDistanceFieldForceMovementUpdate,
-	TEXT("Whether to force N texel border on X, Y and Z update each frame."),
-	ECVF_Scalability | ECVF_RenderThreadSafe
-);
-
-int32 GAOLogGlobalDistanceFieldModifiedPrimitives = 0;
-FAutoConsoleVariableRef CVarAOLogGlobalDistanceFieldModifiedPrimitives(
-	TEXT("r.AOGlobalDistanceFieldLogModifiedPrimitives"),
-	GAOLogGlobalDistanceFieldModifiedPrimitives,
-	TEXT("Whether to log primitive modifications (add, remove, updatetransform) that caused an update of the global distance field.\n")
-	TEXT("This can be useful for tracking down why updating the global distance field is always costing a lot, since it should be mostly cached.\n")
-	TEXT("Pass 2 to log only non movable object updates."),
-	ECVF_Scalability | ECVF_RenderThreadSafe
-	);
-
-int32 GAODrawGlobalDistanceFieldModifiedPrimitives = 0;
-FAutoConsoleVariableRef CVarAODrawGlobalDistanceFieldModifiedPrimitives(
-	TEXT("r.AOGlobalDistanceFieldDrawModifiedPrimitives"),
-	GAODrawGlobalDistanceFieldModifiedPrimitives,
-	TEXT("Whether to draw primitive modifications (add, remove, updatetransform) that caused an update of the global distance field.\n")
-	TEXT("This can be useful for tracking down why updating the global distance field is always costing a lot, since it should be mostly cached."),
-	ECVF_Scalability | ECVF_RenderThreadSafe
-	);
-
 float GAOGlobalDFClipmapDistanceExponent = 2;
 FAutoConsoleVariableRef CVarAOGlobalDFClipmapDistanceExponent(
 	TEXT("r.AOGlobalDFClipmapDistanceExponent"),
@@ -222,6 +195,22 @@ FAutoConsoleVariableRef CVarAOGlobalDistanceFieldMipFactor(
 	ECVF_Scalability | ECVF_RenderThreadSafe
 );
 
+int32 GAOGlobalDistanceFieldRecacheClipmapsWithPendingStreaming = 1;
+FAutoConsoleVariableRef CVarAOGlobalDistanceFieldRecacheClipmapsWithPendingStreaming(
+	TEXT("r.AOGlobalDistanceField.RecacheClipmapsWithPendingStreaming"),
+	GAOGlobalDistanceFieldRecacheClipmapsWithPendingStreaming,
+	TEXT("Whether to readback clipmaps cached with incomplete Mesh SDFs due to streaming and recache them on subsequent frames. Fixes innaccurate Global SDF around the camera after teleporting or loading a new level."),
+	ECVF_Scalability | ECVF_RenderThreadSafe
+);
+
+int32 GAOGlobalDistanceFieldForceRecacheForStreaming = 0;
+FAutoConsoleVariableRef CVarAOGlobalDistanceFieldForceRecacheForStreaming(
+	TEXT("r.AOGlobalDistanceField.ForceRecacheForStreaming"),
+	GAOGlobalDistanceFieldForceRecacheForStreaming,
+	TEXT("Useful for debugging or profiling full clipmap updates that happen when a clipmap is detected to have pending streaming."),
+	ECVF_Scalability | ECVF_RenderThreadSafe
+);
+
 float GLumenSceneGlobalSDFCoveredExpandSurfaceScale = 1.0f;
 FAutoConsoleVariableRef CVarLumenSceneGlobalSDFCoveredExpandSurfaceScale(
 	TEXT("r.LumenScene.GlobalSDF.CoveredExpandSurfaceScale"),
@@ -269,26 +258,37 @@ FAutoConsoleVariableRef CVarLumenSceneGlobalSDFDitheredTransparencyTraceThreshol
 	ECVF_Scalability | ECVF_RenderThreadSafe
 );
 
-TAutoConsoleVariable<int32> CVarGlobalDistanceFieldDebug(
-	TEXT("r.GlobalDistanceField.Debug"),
+TAutoConsoleVariable<bool> CVarGlobalDistanceFieldDebugShowStats(
+	TEXT("r.GlobalDistanceField.Debug.ShowStats"),
 	0,
 	TEXT("Debug drawing for the Global Distance Field."),
 	ECVF_Scalability | ECVF_RenderThreadSafe
 );
 
-int32 GAOGlobalDistanceFieldRecacheClipmapsWithPendingStreaming = 1;
-FAutoConsoleVariableRef CVarAOGlobalDistanceFieldRecacheClipmapsWithPendingStreaming(
-	TEXT("r.AOGlobalDistanceField.RecacheClipmapsWithPendingStreaming"),
-	GAOGlobalDistanceFieldRecacheClipmapsWithPendingStreaming,
-	TEXT("Whether to readback clipmaps cached with incomplete Mesh SDFs due to streaming and reache them on subsequent frames.  Fixes innaccurate Global SDF around the camera after teleporting or loading a new level."),
+int32 GGlobalDistanceFieldDebugForceMovementUpdate = 0;
+FAutoConsoleVariableRef CVarGlobalDistanceFieldDebugForceMovementUpdate(
+	TEXT("r.GlobalDistanceField.Debug.ForceMovementUpdate"),
+	GGlobalDistanceFieldDebugForceMovementUpdate,
+	TEXT("Whether to force N texel border on X, Y and Z update each frame."),
 	ECVF_Scalability | ECVF_RenderThreadSafe
 );
 
-int32 GAOGlobalDistanceFieldForceRecacheForStreaming = 0;
-FAutoConsoleVariableRef CVarAOGlobalDistanceFieldForceRecacheForStreaming(
-	TEXT("r.AOGlobalDistanceField.ForceRecacheForStreaming"),
-	GAOGlobalDistanceFieldForceRecacheForStreaming,
-	TEXT("Useful for debugging or profiling full clipmap updates that happen when a clipmap is detected to have pending streaming."),
+int32 GGlobalDistanceFieldDebugLogModifiedPrimitives = 0;
+FAutoConsoleVariableRef CVarGlobalDistanceFieldDebugLogModifiedPrimitives(
+	TEXT("r.GlobalDistanceField.Debug.LogModifiedPrimitives"),
+	GGlobalDistanceFieldDebugLogModifiedPrimitives,
+	TEXT("Whether to log primitive modifications (add, remove, updatetransform) that caused an update of the global distance field.\n")
+	TEXT("This can be useful for tracking down why updating the global distance field is always costing a lot, since it should be mostly cached.\n")
+	TEXT("Pass 2 to log only non movable object updates."),
+	ECVF_Scalability | ECVF_RenderThreadSafe
+);
+
+int32 GGlobalDistanceFieldDebugDrawModifiedPrimitives = 0;
+FAutoConsoleVariableRef CVarGlobalDistanceFieldDebugDrawModifiedPrimitives(
+	TEXT("r.GlobalDistanceField.Debug.DrawModifiedPrimitives"),
+	GGlobalDistanceFieldDebugDrawModifiedPrimitives,
+	TEXT("Whether to draw primitive modifications (add, remove, updatetransform) that caused an update of the global distance field.\n")
+	TEXT("This can be useful for tracking down why updating the global distance field is always costing a lot, since it should be mostly cached."),
 	ECVF_Scalability | ECVF_RenderThreadSafe
 );
 
@@ -1161,7 +1161,7 @@ static void ComputeUpdateRegionsAndUpdateViewState(
 
 							Clipmap.UpdateBounds.Add(FClipmapUpdateBounds(ModifiedBounds.GetCenter(), ModifiedBounds.GetExtent(), true));
 							
-							if (GAODrawGlobalDistanceFieldModifiedPrimitives)
+							if (GGlobalDistanceFieldDebugDrawModifiedPrimitives)
 							{
 								const uint8 MarkerHue = ((ClipmapIndex * 10 + BoundsIndex) * 10) & 0xFF;
 								const uint8 MarkerSaturation = 0xFF;
@@ -1179,9 +1179,9 @@ static void ComputeUpdateRegionsAndUpdateViewState(
 					{
 						FInt64Vector MovementInPages = PageGridCenter - ClipmapViewState.LastPartialUpdateOriginInPages;
 
-						if (GAOGlobalDistanceFieldForceMovementUpdate != 0)
+						if (GGlobalDistanceFieldDebugForceMovementUpdate != 0)
 						{
-							MovementInPages = FInt64Vector(GAOGlobalDistanceFieldForceMovementUpdate, GAOGlobalDistanceFieldForceMovementUpdate, GAOGlobalDistanceFieldForceMovementUpdate);
+							MovementInPages = FInt64Vector(GGlobalDistanceFieldDebugForceMovementUpdate, GGlobalDistanceFieldDebugForceMovementUpdate, GGlobalDistanceFieldDebugForceMovementUpdate);
 						}
 
 						if (CacheType == GDF_MostlyStatic || !GAOGlobalDistanceFieldCacheMostlyStaticSeparately)
@@ -1883,6 +1883,24 @@ class FPropagateMipDistanceCS : public FGlobalShader
 
 IMPLEMENT_GLOBAL_SHADER(FPropagateMipDistanceCS, "/Engine/Private/DistanceField/GlobalDistanceFieldMip.usf", "PropagateMipDistanceCS", SF_Compute);
 
+class FGlobalDistanceFieldAccumulateUpdatedPagesCS : public FGlobalShader
+{
+	DECLARE_GLOBAL_SHADER(FGlobalDistanceFieldAccumulateUpdatedPagesCS)
+	SHADER_USE_PARAMETER_STRUCT(FGlobalDistanceFieldAccumulateUpdatedPagesCS, FGlobalShader)
+
+	BEGIN_SHADER_PARAMETER_STRUCT(FParameters, )
+		SHADER_PARAMETER_RDG_BUFFER_UAV(RWStructuredBuffer<uint>, RWPageStatsBuffer)
+		SHADER_PARAMETER_RDG_BUFFER_SRV(StructuredBuffer<uint>, PageComposeIndirectArgBuffer)
+	END_SHADER_PARAMETER_STRUCT()
+
+	static bool ShouldCompilePermutation(const FGlobalShaderPermutationParameters& Parameters)
+	{
+		return ShouldCompileDistanceFieldShaders(Parameters.Platform);
+	}
+};
+
+IMPLEMENT_GLOBAL_SHADER(FGlobalDistanceFieldAccumulateUpdatedPagesCS, "/Engine/Private/DistanceField/GlobalDistanceFieldDebug.usf", "GlobalDistanceFieldAccumulateUpdatedPagesCS", SF_Compute);
+
 class FGlobalDistanceFieldPageStatsCS : public FGlobalShader
 {
 	DECLARE_GLOBAL_SHADER(FGlobalDistanceFieldPageStatsCS)
@@ -2103,6 +2121,9 @@ void UpdateGlobalDistanceFieldVolume(
 		// Recreate the view uniform buffer now that we have updated GlobalDistanceFieldInfo
 		View.SetupGlobalDistanceFieldUniformBufferParameters(*View.CachedViewUniformShaderParameters);
 
+		FRDGBufferRef PageStatsBuffer = GraphBuilder.CreateBuffer(FRDGBufferDesc::CreateStructuredDesc(sizeof(uint32), 64), TEXT("GlobalDistanceField.PageStatsBuffer"));
+		AddClearUAVPass(GraphBuilder, GraphBuilder.CreateUAV(PageStatsBuffer), 0);
+
 		bool bHasUpdateBounds = false;
 
 		for (int32 ClipmapIndex = 0; ClipmapIndex < GlobalDistanceFieldInfo.Clipmaps.Num(); ClipmapIndex++)
@@ -2751,6 +2772,22 @@ void UpdateGlobalDistanceFieldVolume(
 							GlobalDistanceFieldUpdateParameters.RelativePreViewTranslation = RelativePreViewTranslation;
 						}
 
+						if(CVarGlobalDistanceFieldDebugShowStats.GetValueOnRenderThread())
+						{
+							FGlobalDistanceFieldAccumulateUpdatedPagesCS::FParameters* PassParameters = GraphBuilder.AllocParameters<FGlobalDistanceFieldAccumulateUpdatedPagesCS::FParameters>();
+							PassParameters->PageComposeIndirectArgBuffer = GraphBuilder.CreateSRV(PageComposeIndirectArgBuffer);
+							PassParameters->RWPageStatsBuffer = GraphBuilder.CreateUAV(PageStatsBuffer);
+
+							auto ComputeShader = View.ShaderMap->GetShader<FGlobalDistanceFieldAccumulateUpdatedPagesCS>();
+
+							FComputeShaderUtils::AddPass(
+								GraphBuilder,
+								RDG_EVENT_NAME("AccumulateUpdatedPages (Debug)"),
+								ComputeShader,
+								PassParameters,
+								FIntVector(1, 1, 1));
+						}
+
 						// Mesh distance fields
 						if(Scene->DistanceFieldSceneData.NumObjectsInBuffer > 0 || UpdateRegionHeightfield.ComponentDescriptions.Num() > 0)
 						{
@@ -2980,14 +3017,13 @@ void UpdateGlobalDistanceFieldVolume(
 			FinalizeGlobalDistanceFieldExternalResourceAccess(GraphBuilder, ExternalAccessQueue, GlobalDistanceFieldInfo, StartCacheType, GlobalDistanceFieldInfoRDG);
 		}
 
-		if (CVarGlobalDistanceFieldDebug.GetValueOnRenderThread() != 0 
+		if (CVarGlobalDistanceFieldDebugShowStats.GetValueOnRenderThread()
 			&& GlobalDistanceFieldInfo.PageFreeListAllocatorBuffer 
 			&& GlobalDistanceFieldInfo.PageAtlasTexture)
 		{
-			ShaderPrint::SetEnabled(true);
+			RDG_EVENT_SCOPE(GraphBuilder, "GlobalDistanceFieldDebug");
 
-			FRDGBufferRef PageStatsBuffer = GraphBuilder.CreateBuffer(FRDGBufferDesc::CreateStructuredDesc(sizeof(uint32), 64), TEXT("GlobalDistanceField.PageStatsBuffer"));
-			AddClearUAVPass(GraphBuilder, GraphBuilder.CreateUAV(PageStatsBuffer), 0);
+			ShaderPrint::SetEnabled(true);
 
 			FRDGTextureRef PageTableCombinedTexture = GraphBuilder.RegisterExternalTexture(GAOGlobalDistanceFieldCacheMostlyStaticSeparately ? GlobalDistanceFieldInfo.PageTableCombinedTexture : GlobalDistanceFieldInfo.PageTableLayerTextures[0], TEXT("GlobalDistanceField.PageTableCombined"));
 			FRDGTextureRef PageAtlasTexture = GraphBuilder.RegisterExternalTexture(GlobalDistanceFieldInfo.PageAtlasTexture, TEXT("GlobalDistanceField.PageAtlasTexture"));
@@ -3029,7 +3065,7 @@ void UpdateGlobalDistanceFieldVolume(
 
 				FComputeShaderUtils::AddPass(
 					GraphBuilder,
-					RDG_EVENT_NAME("GlobalDistanceFieldDebug"),
+					RDG_EVENT_NAME("Print Stats"),
 					ComputeShader,
 					PassParameters,
 					FIntVector(1, 1, 1));
