MovieRenderQueue rendered linear .exr images are not linear wrt. incident illumination for Subsurface Profile. Add comments to not disable checkerboard when the scene color format does not have an alpha. channel.

diff --git a/Engine/Source/Runtime/Renderer/Private/PostProcess/PostProcessSubsurface.cpp b/Engine/Source/Runtime/Renderer/Private/PostProcess/PostProcessSubsurface.cpp
--- a/Engine/Source/Runtime/Renderer/Private/PostProcess/PostProcessSubsurface.cpp
+++ b/Engine/Source/Runtime/Renderer/Private/PostProcess/PostProcessSubsurface.cpp
@@ -92,7 +92,7 @@ namespace
 		2,
 		TEXT("Enables or disables checkerboard rendering for subsurface profile rendering.\n")
 		TEXT("This is necessary if SceneColor does not include a floating point alpha channel (e.g 32-bit formats)\n")
-		TEXT(" 0: Disabled (high quality) \n")
+		TEXT(" 0: Disabled (high quality). If the rendertarget format does not have an alpha channel (e.g., PF_FloatR11G11B10), it leads to over-washed SSS. \n")
 		TEXT(" 1: Enabled (low quality). Surface lighting will be at reduced resolution.\n")
 		TEXT(" 2: Automatic. Non-checkerboard lighting will be applied if we have a suitable rendertarget format\n"),
 		ECVF_RenderThreadSafe);
