Fixed type mismatch when rendering cubemaps from TextureCube to Texture2DArray

diff --git a/Engine/Source/Runtime/Renderer/Private/ReflectionEnvironmentRealTimeCapture.cpp b/Engine/Source/Runtime/Renderer/Private/ReflectionEnvironmentRealTimeCapture.cpp
--- a/Engine/Source/Runtime/Renderer/Private/ReflectionEnvironmentRealTimeCapture.cpp
+++ b/Engine/Source/Runtime/Renderer/Private/ReflectionEnvironmentRealTimeCapture.cpp
@@ -761,7 +761,10 @@ void FScene::AllocateAndCaptureFrameSkyEnvMap(
 				FApplyLowerHemisphereColor::FParameters* PassParameters = GraphBuilder.AllocParameters<FApplyLowerHemisphereColor::FParameters>();
 				PassParameters->ValidDispatchCoord = FIntPoint(Mip0Resolution, Mip0Resolution);
 				PassParameters->LowerHemisphereSolidColor = SkyLight->LowerHemisphereColor;
-				PassParameters->OutTextureMipColor = GraphBuilder.CreateUAV(FRDGTextureUAVDesc(SkyCubeTexture, MipIndex));
+
+				FRDGTextureUAVDesc OutTextureMipColorDesc(SkyCubeTexture, MipIndex);
+				OutTextureMipColorDesc.DimensionOverride = ETextureDimension::Texture2DArray;
+				PassParameters->OutTextureMipColor = GraphBuilder.CreateUAV(OutTextureMipColorDesc);
 
 				FIntVector NumGroups = FIntVector::DivideAndRoundUp(FIntVector(Mip0Resolution, Mip0Resolution, 1), FIntVector(FApplyLowerHemisphereColor::ThreadGroupSize, FApplyLowerHemisphereColor::ThreadGroupSize, 1));
 
@@ -805,7 +808,9 @@ void FScene::AllocateAndCaptureFrameSkyEnvMap(
 			PassParameters->SourceCubemapSampler = TStaticSamplerState<SF_Point>::GetRHI();
 
 			PassParameters->SourceCubemapTexture = SkyCubeTextureSRV;
-			PassParameters->OutTextureMipColor = GraphBuilder.CreateUAV(FRDGTextureUAVDesc(SkyCubeTexture, MipIndex));
+			FRDGTextureUAVDesc OutTextureMipColorDesc(SkyCubeTexture, MipIndex);
+			OutTextureMipColorDesc.DimensionOverride = ETextureDimension::Texture2DArray;
+			PassParameters->OutTextureMipColor = GraphBuilder.CreateUAV(OutTextureMipColorDesc);
 
 			FIntVector NumGroups = FIntVector::DivideAndRoundUp(FIntVector(MipResolution, MipResolution, 1), FIntVector(FDownsampleCubeFaceCS::ThreadGroupSize, FDownsampleCubeFaceCS::ThreadGroupSize, 1));
 
@@ -851,7 +856,9 @@ void FScene::AllocateAndCaptureFrameSkyEnvMap(
 			PassParameters->SourceCubemapSampler = TStaticSamplerState<SF_Point>::GetRHI();
 
 			PassParameters->SourceCubemapTexture = RDGSrcRenderTargetSRV;
-			PassParameters->OutTextureMipColor = GraphBuilder.CreateUAV(FRDGTextureUAVDesc(RDGDstRenderTarget, MipIndex));
+			FRDGTextureUAVDesc OutTextureMipColorDesc(RDGDstRenderTarget, MipIndex);
+			OutTextureMipColorDesc.DimensionOverride = ETextureDimension::Texture2DArray;
+			PassParameters->OutTextureMipColor = GraphBuilder.CreateUAV(OutTextureMipColorDesc);
 
 			FIntVector NumGroups = FIntVector::DivideAndRoundUp(FIntVector(MipResolution, MipResolution, 1), FIntVector(FConvolveSpecularFaceCS::ThreadGroupSize, FConvolveSpecularFaceCS::ThreadGroupSize, 1));
 
