Fix ensure when accessing the default material during water PSO collection from the game thread (following CL 24117541)

diff --git a/Engine/Source/Runtime/Renderer/Private/SingleLayerWaterRendering.cpp b/Engine/Source/Runtime/Renderer/Private/SingleLayerWaterRendering.cpp
--- a/Engine/Source/Runtime/Renderer/Private/SingleLayerWaterRendering.cpp
+++ b/Engine/Source/Runtime/Renderer/Private/SingleLayerWaterRendering.cpp
@@ -1846,10 +1846,11 @@ void FSingleLayerWaterDepthPrepassMeshProcessor::CollectPSOInitializers(const FS
 			&& !Material.MaterialModifiesMeshPosition_GameThread()
 			&& Material.WritesEveryPixel(false, bVFTypeSupportsNullPixelShader))
 		{
-			const FMaterialRenderProxy& DefaultProxy = *UMaterial::GetDefaultMaterial(MD_Surface)->GetRenderProxy();
-			const FMaterial& DefaultMaterial = *DefaultProxy.GetMaterialNoFallback(FeatureLevel);
+			EMaterialQualityLevel::Type ActiveQualityLevel = GetCachedScalabilityCVars().MaterialQualityLevel;
+			const FMaterial* DefaultMaterial = UMaterial::GetDefaultMaterial(MD_Surface)->GetMaterialResource(FeatureLevel, ActiveQualityLevel);
+			check(DefaultMaterial);
 
-			CollectPSOInitializersInternal<true>(SceneTexturesConfig, VertexFactoryData, DefaultMaterial, MeshFillMode, MeshCullMode, PreCacheParams, PSOInitializers);
+			CollectPSOInitializersInternal<true>(SceneTexturesConfig, VertexFactoryData, *DefaultMaterial, MeshFillMode, MeshCullMode, PreCacheParams, PSOInitializers);
 		}
 		else
 		{
