Fixes to post-process materials with stencil test:
- Fix operand order of manual stencil test used when Nanite renders custom depth (reference value should be on the left hand side)
- Force compositing of a post-process material's output with the input when manual stencil test is enabled to fix corruption in discarded regions.

diff --git a/Engine/Shaders/Private/PostProcessMaterialShaders.usf b/Engine/Shaders/Private/PostProcessMaterialShaders.usf
--- a/Engine/Shaders/Private/PostProcessMaterialShaders.usf
+++ b/Engine/Shaders/Private/PostProcessMaterialShaders.usf
@@ -189,11 +189,10 @@ void MainVS_VideoOverlay(
 void ManualStencilTest(uint2 PixelPos)
 {
 	const uint Stencil = CalcSceneCustomStencil(PixelPos);
-	const int Diff = ManualStencilReferenceValue - Stencil;
 	const uint Test =
-		(Stencil == ManualStencilReferenceValue	? STENCIL_TEST_EQUAL	: 0) |
-		(Stencil <  ManualStencilReferenceValue	? STENCIL_TEST_LESS		: 0) |
-		(Stencil >  ManualStencilReferenceValue	? STENCIL_TEST_GREATER	: 0);
+		(ManualStencilReferenceValue == Stencil	? STENCIL_TEST_EQUAL	: 0) |
+		(ManualStencilReferenceValue <  Stencil	? STENCIL_TEST_LESS		: 0) |
+		(ManualStencilReferenceValue >  Stencil	? STENCIL_TEST_GREATER	: 0);
 
 	if ((Test & ManualStencilTestMask) == 0)
 	{
diff --git a/Engine/Source/Runtime/Renderer/Private/PostProcess/PostProcessMaterial.cpp b/Engine/Source/Runtime/Renderer/Private/PostProcess/PostProcessMaterial.cpp
--- a/Engine/Source/Runtime/Renderer/Private/PostProcess/PostProcessMaterial.cpp
+++ b/Engine/Source/Runtime/Renderer/Private/PostProcess/PostProcessMaterial.cpp
@@ -463,8 +463,13 @@ FScreenPassTexture AddPostProcessMaterialPass(
 		BlendState = GetMaterialBlendState(Material);
 	}
 
+	// Determine if the pixel shader may discard, requiring us to initialize the output texture
+	const bool bMayDiscard = CustomStencilPolicy == EMaterialCustomDepthPolicy::Enabled && Inputs.bManualStencilTest;
+
 	// Blend / Depth Stencil usage requires that the render target have primed color data.
-	const bool bCompositeWithInput = DepthStencilState != DefaultDepthStencilState || BlendState != DefaultBlendState;
+	const bool bCompositeWithInput = DepthStencilState != DefaultDepthStencilState ||
+									 BlendState != DefaultBlendState ||
+									 bMayDiscard;
 
 	// We only prime color on the output texture if we are using fixed function Blend / Depth-Stencil,
 	// or we need to retain previously rendered views.
