Don't force bScenesPrimitivesNeedStaticMeshElementUpdate when mobile base pass is not relying on the directional light directly for light map policy selection

diff --git a/Engine/Source/Runtime/Renderer/Private/MobileBasePass.cpp b/Engine/Source/Runtime/Renderer/Private/MobileBasePass.cpp
--- a/Engine/Source/Runtime/Renderer/Private/MobileBasePass.cpp
+++ b/Engine/Source/Runtime/Renderer/Private/MobileBasePass.cpp
@@ -269,11 +269,16 @@ bool MobileBasePass::StaticCanReceiveCSM(const FLightSceneInfo* LightSceneInfo,
 	return false; 
 }
 
+bool MobileBasePass::IsUsingDirectionalLightForLighmapPolicySelection(const FScene* Scene)
+{
+	const FReadOnlyCVARCache& ReadOnlyCVARCache = FReadOnlyCVARCache::Get();
+	return !(!ReadOnlyCVARCache.bAllowStaticLighting || (ReadOnlyCVARCache.bMobileEnableNoPrecomputedLightingCSMShader && Scene && Scene->GetForceNoPrecomputedLighting()));
+}
+
 ELightMapPolicyType MobileBasePass::SelectMeshLightmapPolicy(
 	const FScene* Scene, 
 	const FMeshBatch& Mesh, 
-	const FPrimitiveSceneProxy* PrimitiveSceneProxy,
-	const FLightSceneInfo* MobileDirectionalLight,
+	const FPrimitiveSceneProxy* PrimitiveSceneProxy,	
 	bool bPrimReceivesCSM,
 	bool bUsesDeferredShading,
 	bool bIsLitMaterial,
@@ -287,9 +292,9 @@ ELightMapPolicyType MobileBasePass::SelectMeshLightmapPolicy(
 		constexpr ERHIFeatureLevel::Type FeatureLevel = ERHIFeatureLevel::ES3_1;
 		const FReadOnlyCVARCache& ReadOnlyCVARCache = FReadOnlyCVARCache::Get();
 
-		if (!ReadOnlyCVARCache.bAllowStaticLighting || (ReadOnlyCVARCache.bMobileEnableNoPrecomputedLightingCSMShader && Scene && Scene->GetForceNoPrecomputedLighting()))
+		if (!IsUsingDirectionalLightForLighmapPolicySelection(Scene))
 		{
-			if (!bIsTranslucent)
+ 			if (!bIsTranslucent)
 			{
 				// Whether to use a single CSM permutation with a branch in the shader
 				bPrimReceivesCSM |= MobileUseCSMShaderBranch();
@@ -311,6 +316,8 @@ ELightMapPolicyType MobileBasePass::SelectMeshLightmapPolicy(
 			const FLightMapInteraction LightMapInteraction = (Mesh.LCI != nullptr)
 				? Mesh.LCI->GetLightMapInteraction(FeatureLevel)
 				: FLightMapInteraction();
+
+			const FLightSceneInfo* MobileDirectionalLight = MobileBasePass::GetDirectionalLightInfo(Scene, PrimitiveSceneProxy);
 		
 			const bool bUseMovableLight = MobileDirectionalLight && !MobileDirectionalLight->Proxy->HasStaticShadowing() && ReadOnlyCVARCache.bMobileAllowMovableDirectionalLights;
 			const bool bUseStaticAndCSM = MobileDirectionalLight && MobileDirectionalLight->Proxy->UseCSMForDynamicObjects()
@@ -812,9 +819,8 @@ bool FMobileBasePassMeshProcessor::TryAddMeshBatch(const FMeshBatch& RESTRICT Me
 		const EBlendMode BlendMode = Material.GetBlendMode();
 		const bool bIsLitMaterial = ShadingModels.IsLit();
 		const bool bIsTranslucent = IsTranslucentBlendMode(BlendMode) || ShadingModels.HasShadingModel(MSM_SingleLayerWater); // Water goes into the translucent pass;
-		const bool bIsMasked = IsMaskedBlendMode(Material);
-		const FLightSceneInfo* MobileDirectionalLight = MobileBasePass::GetDirectionalLightInfo(Scene, PrimitiveSceneProxy);
-		ELightMapPolicyType LightmapPolicyType = MobileBasePass::SelectMeshLightmapPolicy(Scene, MeshBatch, PrimitiveSceneProxy, MobileDirectionalLight, bCanReceiveCSM, bPassUsesDeferredShading, bIsLitMaterial, bIsTranslucent);
+		const bool bIsMasked = IsMaskedBlendMode(Material);		
+		ELightMapPolicyType LightmapPolicyType = MobileBasePass::SelectMeshLightmapPolicy(Scene, MeshBatch, PrimitiveSceneProxy, bCanReceiveCSM, bPassUsesDeferredShading, bIsLitMaterial, bIsTranslucent);
 		return Process(MeshBatch, BatchElementMask, StaticMeshId, PrimitiveSceneProxy, MaterialRenderProxy, Material, bIsMasked, bIsTranslucent, ShadingModels, LightmapPolicyType, bCanReceiveCSM, MeshBatch.LCI);
 	}
 	return true;
diff --git a/Engine/Source/Runtime/Renderer/Private/MobileBasePassRendering.h b/Engine/Source/Runtime/Renderer/Private/MobileBasePassRendering.h
--- a/Engine/Source/Runtime/Renderer/Private/MobileBasePassRendering.h
+++ b/Engine/Source/Runtime/Renderer/Private/MobileBasePassRendering.h
@@ -319,11 +319,12 @@ public:
 
 namespace MobileBasePass
 {
+	bool IsUsingDirectionalLightForLighmapPolicySelection(const FScene* Scene);
+
 	ELightMapPolicyType SelectMeshLightmapPolicy(
 		const FScene* Scene, 
 		const FMeshBatch& MeshBatch, 
 		const FPrimitiveSceneProxy* PrimitiveSceneProxy, 
-		const FLightSceneInfo* MobileDirectionalLight, 
 		bool bPrimReceivesCSM, 
 		bool bUsedDeferredShading,
 		bool bIsLitMaterial,
diff --git a/Engine/Source/Runtime/Renderer/Private/RendererScene.cpp b/Engine/Source/Runtime/Renderer/Private/RendererScene.cpp
--- a/Engine/Source/Runtime/Renderer/Private/RendererScene.cpp
+++ b/Engine/Source/Runtime/Renderer/Private/RendererScene.cpp
@@ -2674,7 +2674,7 @@ void FScene::AddLightSceneInfo_RenderThread(FLightSceneInfo* LightSceneInfo)
 			    MobileDirectionalLights[FirstLightingChannel] = LightSceneInfo;
     
 			    // if this light is a dynamic shadowcast then we need to update the static draw lists to pick a new lighting policy:
-			    if (!LightSceneInfo->Proxy->HasStaticShadowing() || bUseCSMForDynamicObjects)
+			    if (MobileBasePass::IsUsingDirectionalLightForLighmapPolicySelection(this) && (!LightSceneInfo->Proxy->HasStaticShadowing() || bUseCSMForDynamicObjects))
 				{
 		    		bScenesPrimitivesNeedStaticMeshElementUpdate = true;
 					UE_CLOG(!GIsEditor, LogRenderer, Log, TEXT("Forcing update for all mesh draw commands: Add directional light"));
@@ -3938,7 +3938,7 @@ void FScene::RemoveLightSceneInfo_RenderThread(FLightSceneInfo* LightSceneInfo)
 				}
 
 				// if this light is a dynamic shadowcast then we need to update the static draw lists to pick a new lightingpolicy
-				if (!LightSceneInfo->Proxy->HasStaticShadowing() || bUseCSMForDynamicObjects)
+				if (MobileBasePass::IsUsingDirectionalLightForLighmapPolicySelection(this) && (!LightSceneInfo->Proxy->HasStaticShadowing() || bUseCSMForDynamicObjects))
 				{
 					bScenesPrimitivesNeedStaticMeshElementUpdate = true;
 					UE_CLOG(!GIsEditor, LogRenderer, Log, TEXT("Forcing update for all mesh draw commands: Remove directional light"));
