Don't collect data from HitProxy pass for MeshDrawCommands

diff --git a/Engine/Source/Runtime/Renderer/Private/InstanceCulling/InstanceCullingContext.cpp b/Engine/Source/Runtime/Renderer/Private/InstanceCulling/InstanceCullingContext.cpp
--- a/Engine/Source/Runtime/Renderer/Private/InstanceCulling/InstanceCullingContext.cpp
+++ b/Engine/Source/Runtime/Renderer/Private/InstanceCulling/InstanceCullingContext.cpp
@@ -178,7 +178,10 @@ FInstanceCullingContext::FInstanceCullingContext(
 #if MESH_DRAW_COMMAND_STATS
 	if (FMeshDrawCommandStatsManager* Instance = FMeshDrawCommandStatsManager::Get())
 	{
-		MeshDrawCommandPassStats = Instance->CreatePassStats(PassName);
+		if (FCString::Strcmp(*(PassName.ToString()), TEXT("HitProxy")) != 0)
+		{
+			MeshDrawCommandPassStats = Instance->CreatePassStats(PassName);
+		}
 	}
 #endif
 }
