VulkanRHI: Fix resource mismatch issue in ray tracing tests for Vulkan (RWBuffer <-> RWStructuredBuffer).  Also move to newer buffer and view creation methods.

#rnx

diff --git a/Engine/Source/Runtime/Renderer/Private/Tests/RayTracingTestbed.cpp b/Engine/Source/Runtime/Renderer/Private/Tests/RayTracingTestbed.cpp
--- a/Engine/Source/Runtime/Renderer/Private/Tests/RayTracingTestbed.cpp
+++ b/Engine/Source/Runtime/Renderer/Private/Tests/RayTracingTestbed.cpp
@@ -107,8 +107,18 @@ bool RunRayTracingTestbed_RenderThread(const FString& Parameters)
 		FRHIResourceCreateInfo CreateInfo(TEXT("RayBuffer"));
 		CreateInfo.ResourceArray = &RayData;
 
-		RayBuffer = RHICmdList.CreateStructuredBuffer(sizeof(FBasicRayTracingRay), RayData.GetResourceDataSize(), BUF_Static | BUF_ShaderResource, CreateInfo);
-		RayBufferView = RHICmdList.CreateShaderResourceView(RayBuffer);
+		RayBuffer = RHICmdList.CreateBuffer(RayData.GetResourceDataSize(), 
+			BUF_Static | BUF_ShaderResource | BUF_StructuredBuffer, 
+			sizeof(FBasicRayTracingRay),
+			ERHIAccess::SRVMask,
+			CreateInfo
+		);
+		RayBufferView = RHICmdList.CreateShaderResourceView(RayBuffer, 
+			FRHIViewDesc::CreateBufferSRV()
+			.SetType(FRHIViewDesc::EBufferType::Structured)
+			.SetStride(sizeof(FBasicRayTracingRay))
+			.SetNumElements(NumRays)
+		);
 	}
 
 	FBufferRHIRef OcclusionResultBuffer;
@@ -116,8 +126,18 @@ bool RunRayTracingTestbed_RenderThread(const FString& Parameters)
 
 	{
 		FRHIResourceCreateInfo CreateInfo(TEXT("OcclusionResultBuffer"));
-		OcclusionResultBuffer = RHICmdList.CreateVertexBuffer(sizeof(uint32)*NumRays, BUF_Static | BUF_UnorderedAccess, CreateInfo);
-		OcclusionResultBufferView = RHICmdList.CreateUnorderedAccessView(OcclusionResultBuffer, PF_R32_UINT);
+		OcclusionResultBuffer = RHICmdList.CreateBuffer(sizeof(uint32) * NumRays, 
+			BUF_Static | BUF_UnorderedAccess | BUF_StructuredBuffer, 
+			sizeof(uint32), 
+			ERHIAccess::UAVMask, 
+			CreateInfo
+		);
+		OcclusionResultBufferView = RHICmdList.CreateUnorderedAccessView(OcclusionResultBuffer, 
+			FRHIViewDesc::CreateBufferUAV()
+			.SetType(FRHIViewDesc::EBufferType::Structured)
+			.SetStride(sizeof(uint32))
+			.SetNumElements(NumRays)
+		);
 	}
 
 	FBufferRHIRef IntersectionResultBuffer;
@@ -125,8 +145,18 @@ bool RunRayTracingTestbed_RenderThread(const FString& Parameters)
 
 	{
 		FRHIResourceCreateInfo CreateInfo(TEXT("IntersectionResultBuffer"));
-		IntersectionResultBuffer = RHICmdList.CreateVertexBuffer(sizeof(FBasicRayTracingIntersectionResult)*NumRays, BUF_Static | BUF_UnorderedAccess, CreateInfo);
-		IntersectionResultBufferView = RHICmdList.CreateUnorderedAccessView(IntersectionResultBuffer, PF_R32_UINT);
+		IntersectionResultBuffer = RHICmdList.CreateBuffer(sizeof(FBasicRayTracingIntersectionResult) * NumRays, 
+			BUF_Static | BUF_UnorderedAccess | BUF_StructuredBuffer, 
+			sizeof(FBasicRayTracingIntersectionResult), 
+			ERHIAccess::UAVMask, 
+			CreateInfo
+		);
+		IntersectionResultBufferView = RHICmdList.CreateUnorderedAccessView(IntersectionResultBuffer, 
+			FRHIViewDesc::CreateBufferUAV()
+			.SetType(FRHIViewDesc::EBufferType::Structured)
+			.SetStride(sizeof(FBasicRayTracingIntersectionResult))
+			.SetNumElements(NumRays)
+		);
 	}
 
 	FRayTracingGeometryInitializer GeometryInitializer;
