Make sure we correctly setup single primitive data for "uniform view" path

diff --git a/Engine/Shaders/Private/SceneData.ush b/Engine/Shaders/Private/SceneData.ush
--- a/Engine/Shaders/Private/SceneData.ush
+++ b/Engine/Shaders/Private/SceneData.ush
@@ -41,10 +41,6 @@
 
 #ifndef PLATFORM_USES_PRIMITIVE_UBO
 	#define PLATFORM_USES_PRIMITIVE_UBO (PLATFORM_SUPPORTS_UNIFORM_BUFFER_OBJECTS && FEATURE_LEVEL == FEATURE_LEVEL_ES3_1) 
-	#define BATCHED_PRIMITIVE_DATA_STRIDE 512u
-	#define BATCHED_INSTANCE_DATA_STRIDE 256u
-	#define BATCHED_PRIMITIVE_DATA_STRIDE_FLOAT4 (BATCHED_PRIMITIVE_DATA_STRIDE/16u)
-	#define BATCHED_INSTANCE_DATA_STRIDE_FLOAT4 (BATCHED_INSTANCE_DATA_STRIDE/16u)
 #endif
 
 #ifndef SHADER_USES_PRIMITIVE_UBO
diff --git a/Engine/Shaders/Private/SceneDataMobileWriter.ush b/Engine/Shaders/Private/SceneDataMobileWriter.ush
--- a/Engine/Shaders/Private/SceneDataMobileWriter.ush
+++ b/Engine/Shaders/Private/SceneDataMobileWriter.ush
@@ -13,6 +13,9 @@ void WriteDataElementUBO(RWStructuredBuffer<float4> Output, uint Offset, uint Da
 #endif
 }
 
+//
+// Note: layout must match FBatchedPrimitiveShaderData in PrimitiveUniformShaderParameters.cpp
+//
 void WritePrimitiveDataUBO(RWStructuredBuffer<float4> Output, uint Offset, FPrimitiveSceneData PrimitiveData, uint MeshLODIndex)
 {
 	uint DataIdx = 0;
diff --git a/Engine/Shaders/Shared/SceneDefinitions.h b/Engine/Shaders/Shared/SceneDefinitions.h
--- a/Engine/Shaders/Shared/SceneDefinitions.h
+++ b/Engine/Shaders/Shared/SceneDefinitions.h
@@ -98,3 +98,9 @@
 	| INSTANCE_SCENE_DATA_FLAG_HAS_EDITOR_DATA			\
 	| INSTANCE_SCENE_DATA_FLAG_HAS_PAYLOAD_EXTENSION	\
 )
+
+// Specific to "UniformView" rendering path
+#define BATCHED_PRIMITIVE_DATA_STRIDE 512u
+#define BATCHED_INSTANCE_DATA_STRIDE 256u
+#define BATCHED_PRIMITIVE_DATA_STRIDE_FLOAT4 (BATCHED_PRIMITIVE_DATA_STRIDE/16u)
+#define BATCHED_INSTANCE_DATA_STRIDE_FLOAT4 (BATCHED_INSTANCE_DATA_STRIDE/16u)
\ No newline at end of file
diff --git a/Engine/Source/Runtime/Engine/Private/PrimitiveUniformShaderParameters.cpp b/Engine/Source/Runtime/Engine/Private/PrimitiveUniformShaderParameters.cpp
--- a/Engine/Source/Runtime/Engine/Private/PrimitiveUniformShaderParameters.cpp
+++ b/Engine/Source/Runtime/Engine/Private/PrimitiveUniformShaderParameters.cpp
@@ -281,4 +281,123 @@ FPrimitiveUniformShaderParameters GetIdentityPrimitiveParameters()
 			.WorldBounds(FBoxSphereBounds(EForceInit::ForceInit))
 			.LocalBounds(FBoxSphereBounds(EForceInit::ForceInit))
 		.Build();
+}
+
+
+FBatchedPrimitiveShaderData::FBatchedPrimitiveShaderData(const FPrimitiveSceneProxy* RESTRICT Proxy)
+	: Data(InPlace, NoInit)
+{
+	FPrimitiveUniformShaderParametersBuilder Builder = FPrimitiveUniformShaderParametersBuilder{};
+	Proxy->BuildUniformShaderParameters(Builder);
+	Setup(Builder.Build());
+}
+
+void FBatchedPrimitiveShaderData::Setup(const FPrimitiveUniformShaderParameters& PrimitiveUniformShaderParameters)
+{
+	static_assert(NUM_LIGHTING_CHANNELS == 3, "The FPrimitiveSceneShaderData packing currently assumes a maximum of 3 lighting channels.");
+	
+	// Note: layout must match LoadPrimitiveDataUBO in SceneDataMobileLoader.ush
+	int32 i = 0;
+
+	const bool bAllowStaticLighting = FReadOnlyCVARCache::Get().bAllowStaticLighting;
+	if (bAllowStaticLighting)
+	{
+		FVector4f LightMapUVScaleBias = FVector4f(1, 1, 0, 0);
+		FVector4f ShadowMapUVScaleBias = FVector4f(1, 1, 0, 0);
+		uint32 LightMapDataIdx = 0;
+
+		// FIXME: valid lightmap data
+		Data[i+0] = LightMapUVScaleBias;
+		Data[i+1] = ShadowMapUVScaleBias;
+		Data[i+2] = FVector4f(FMath::AsFloat(LightMapDataIdx), 0.f, 0.f, 0.f);
+		i+=3;
+	}
+
+	// TilePosition, Flags
+	{
+		Data[i].X = PrimitiveUniformShaderParameters.TilePosition.X;
+		Data[i].Y = PrimitiveUniformShaderParameters.TilePosition.Y;
+		Data[i].Z = PrimitiveUniformShaderParameters.TilePosition.Z;
+		Data[i].W = FMath::AsFloat(PrimitiveUniformShaderParameters.Flags);
+		i+=1;
+	}
+	
+	// LocalToWorld
+	{
+		FMatrix44f LocalToRelativeWorldTranspose = PrimitiveUniformShaderParameters.LocalToRelativeWorld.GetTransposed();
+		Data[i+0] = *(const FVector4f*)&LocalToRelativeWorldTranspose.M[0][0];
+		Data[i+1] = *(const FVector4f*)&LocalToRelativeWorldTranspose.M[1][0];
+		Data[i+2] = *(const FVector4f*)&LocalToRelativeWorldTranspose.M[2][0];
+		i+=3;
+	}
+
+	// InvNonUniformScale, TODO .w
+	{
+		Data[i+0] = FVector4f(PrimitiveUniformShaderParameters.InvNonUniformScale, 0.0f);
+		i+=1;
+	}
+
+	// ObjectWorldPosition, Radius
+	{
+		Data[i+0] = PrimitiveUniformShaderParameters.ObjectRelativeWorldPositionAndRadius;
+		i+=1;
+	}
+
+	// ActorWorldPosition, TODO .w
+	{
+		Data[i+0] = FVector4f(PrimitiveUniformShaderParameters.ActorRelativeWorldPosition, 0.0f);
+		i+=1;
+	}
+
+	// ObjectOrientation, ObjectBoundsX
+	{
+		Data[i+0] = FVector4f(PrimitiveUniformShaderParameters.ObjectOrientation, PrimitiveUniformShaderParameters.ObjectBoundsX);
+		i+=1;
+	}
+
+	// LocalObjectBoundsMin, ObjectBoundsY
+	{
+		Data[i+0] = FVector4f(PrimitiveUniformShaderParameters.LocalObjectBoundsMin, PrimitiveUniformShaderParameters.ObjectBoundsY);
+		i+=1;
+	}
+
+	// LocalObjectBoundsMax, ObjectBoundsZ
+	{
+		Data[i+0] = FVector4f(PrimitiveUniformShaderParameters.LocalObjectBoundsMax, PrimitiveUniformShaderParameters.ObjectBoundsZ);
+		i+=1;
+	}
+
+	// WorldToLocal
+	{
+		FMatrix44f RelativeWorldToLocalTranspose = PrimitiveUniformShaderParameters.RelativeWorldToLocal.GetTransposed();
+		Data[i+0] = *(const FVector4f*)&RelativeWorldToLocalTranspose.M[0][0];
+		Data[i+1] = *(const FVector4f*)&RelativeWorldToLocalTranspose.M[1][0];
+		Data[i+2] = *(const FVector4f*)&RelativeWorldToLocalTranspose.M[2][0];
+		i+=3;
+	}
+
+	// PreviousLocalToWorld
+	{
+		FMatrix44f PreviousLocalToRelativeWorldTranspose = PrimitiveUniformShaderParameters.PreviousLocalToRelativeWorld.GetTransposed();
+		Data[i+0] = *(const FVector4f*)&PreviousLocalToRelativeWorldTranspose.M[0][0];
+		Data[i+1] = *(const FVector4f*)&PreviousLocalToRelativeWorldTranspose.M[1][0];
+		Data[i+2] = *(const FVector4f*)&PreviousLocalToRelativeWorldTranspose.M[2][0];
+		i+=3;
+	}
+
+	// PreviousWorldToLocal
+	{
+		FMatrix44f PreviousRelativeWorldToLocalTranspose = PrimitiveUniformShaderParameters.PreviousRelativeWorldToLocal.GetTransposed();
+		Data[i+0] = *(const FVector4f*)&PreviousRelativeWorldToLocalTranspose.M[0][0];
+		Data[i+1] = *(const FVector4f*)&PreviousRelativeWorldToLocalTranspose.M[1][0];
+		Data[i+2] = *(const FVector4f*)&PreviousRelativeWorldToLocalTranspose.M[2][0];
+		i+=3;
+	}
+	
+	// Set all the custom primitive data float4. This matches the loop in SceneData.ush
+	int32 NumCustomData = FMath::Min<int32>(FCustomPrimitiveData::NumCustomPrimitiveDataFloat4s, DataStrideInFloat4s - i);
+	for (int32 DataIndex = 0; DataIndex < FCustomPrimitiveData::NumCustomPrimitiveDataFloat4s; ++DataIndex)
+	{
+		Data[i + DataIndex] = PrimitiveUniformShaderParameters.CustomPrimitiveData[DataIndex];
+	}
 }
\ No newline at end of file
diff --git a/Engine/Source/Runtime/Engine/Public/PrimitiveSceneShaderData.h b/Engine/Source/Runtime/Engine/Public/PrimitiveSceneShaderData.h
--- a/Engine/Source/Runtime/Engine/Public/PrimitiveSceneShaderData.h
+++ b/Engine/Source/Runtime/Engine/Public/PrimitiveSceneShaderData.h
@@ -122,3 +122,27 @@ public:
 extern ENGINE_API TGlobalResource<FSinglePrimitiveStructured> GIdentityPrimitiveBuffer;
 UE_DEPRECATED(5.3, "Use IRendererModule::CreateSinglePrimitiveSceneUniformBuffer instead")
 extern ENGINE_API TGlobalResource<FSinglePrimitiveStructured> GTilePrimitiveBuffer;
+
+
+struct FBatchedPrimitiveShaderData
+{
+	static const uint32 DataStrideInFloat4s = BATCHED_PRIMITIVE_DATA_STRIDE_FLOAT4;
+
+	TStaticArray<FVector4f, DataStrideInFloat4s> Data;
+
+	FBatchedPrimitiveShaderData()
+		: Data(InPlace, NoInit)
+	{
+		Setup(GetIdentityPrimitiveParameters());
+	}
+
+	explicit FBatchedPrimitiveShaderData(const FPrimitiveUniformShaderParameters& PrimitiveUniformShaderParameters)
+		: Data(InPlace, NoInit)
+	{
+		Setup(PrimitiveUniformShaderParameters);
+	}
+
+	ENGINE_API FBatchedPrimitiveShaderData(const class FPrimitiveSceneProxy* RESTRICT Proxy);
+
+	ENGINE_API void Setup(const FPrimitiveUniformShaderParameters& PrimitiveUniformShaderParameters);
+};
diff --git a/Engine/Source/Runtime/RenderCore/Public/RendererInterface.h b/Engine/Source/Runtime/RenderCore/Public/RendererInterface.h
--- a/Engine/Source/Runtime/RenderCore/Public/RendererInterface.h
+++ b/Engine/Source/Runtime/RenderCore/Public/RendererInterface.h
@@ -39,6 +39,7 @@ class FGPUScenePrimitiveCollector;
 class FViewInfo;
 template<typename ShaderType, typename PointerTableType> class TShaderRefBase;
 class FSceneUniformBuffer;
+class FBatchedPrimitiveParameters;
 
 namespace Nanite
 {
@@ -716,6 +717,8 @@ public:
 
 	/** Create a Scene Uniform Buffer containing only the scene representation for a single primitive */
 	virtual FSceneUniformBuffer* CreateSinglePrimitiveSceneUniformBuffer(FRDGBuilder& GraphBuilder, const FViewInfo& SceneView, FMeshBatch& Mesh) = 0;
+	/** Create a Uniform Buffer containing representation for a single primitive. (For platforms that use "UniformView" path) */
+	virtual TRDGUniformBufferRef<FBatchedPrimitiveParameters> CreateSinglePrimitiveUniformView(FRDGBuilder& GraphBuilder, const FViewInfo& SceneView, FMeshBatch& Mesh) = 0;
 
 	/** Draws a tile mesh element with the specified view. */
 	virtual void DrawTileMesh(FCanvasRenderContext& RenderContext, struct FMeshPassProcessorRenderState& DrawRenderState, const FSceneView& View, FMeshBatch& Mesh, bool bIsHitTesting, const class FHitProxyId& HitProxyId, bool bUse128bitRT = false) = 0;
diff --git a/Engine/Source/Runtime/Renderer/Private/Renderer.cpp b/Engine/Source/Runtime/Renderer/Private/Renderer.cpp
--- a/Engine/Source/Runtime/Renderer/Private/Renderer.cpp
+++ b/Engine/Source/Runtime/Renderer/Private/Renderer.cpp
@@ -129,8 +129,7 @@ void FRendererModule::InitializeSystemTextures(FRHICommandListImmediate& RHICmdL
 
 BEGIN_SHADER_PARAMETER_STRUCT(FDrawTileMeshPassParameters, )
 	SHADER_PARAMETER_STRUCT_INCLUDE(FViewShaderParameters, View)
-	SHADER_PARAMETER_RDG_UNIFORM_BUFFER(FSceneUniformParameters, Scene)
-	SHADER_PARAMETER_RDG_UNIFORM_BUFFER(FInstanceCullingGlobalUniforms, InstanceCulling)
+	SHADER_PARAMETER_STRUCT_INCLUDE(FInstanceCullingDrawParams, InstanceCullingDrawParams)
 	SHADER_PARAMETER_STRUCT_REF(FReflectionCaptureShaderData, ReflectionCapture)
 	SHADER_PARAMETER_RDG_UNIFORM_BUFFER(FDebugViewModePassUniformParameters, DebugViewMode)
 	SHADER_PARAMETER_RDG_UNIFORM_BUFFER(FTranslucentBasePassUniformParameters, TranslucentBasePass)
@@ -213,6 +212,48 @@ FSceneUniformBuffer* FRendererModule::CreateSinglePrimitiveSceneUniformBuffer(FR
 	return &SceneUniforms;
 }
 
+TRDGUniformBufferRef<FBatchedPrimitiveParameters> FRendererModule::CreateSinglePrimitiveUniformView(FRDGBuilder& GraphBuilder, const FViewInfo& SceneView, FMeshBatch& Mesh)
+{
+	check(PlatformGPUSceneUsesUniformBufferView(SceneView.GetShaderPlatform()));
+
+	FBatchedPrimitiveParameters* BatchedPrimitiveParameters = GraphBuilder.AllocParameters<FBatchedPrimitiveParameters>();
+
+	FRDGBufferDesc PrimitiveDataBufferDesc = FRDGBufferDesc::CreateStructuredDesc(16u, (PLATFORM_MAX_UNIFORM_BUFFER_RANGE / 16u));
+	PrimitiveDataBufferDesc.Usage |= EBufferUsageFlags::UniformBuffer;
+	FRDGBufferRef PrimitiveDataBuffer = nullptr;
+	
+	ERHIFeatureLevel::Type FeatureLevel = SceneView.GetFeatureLevel();
+	if (Mesh.VertexFactory->GetPrimitiveIdStreamIndex(FeatureLevel, EVertexInputStreamType::PositionOnly) >= 0)
+	{
+		FMeshBatchElement& MeshElement = Mesh.Elements[0];
+		checkf(Mesh.Elements.Num() == 1, TEXT("Only 1 batch element currently supported by CreateSinglePrimitiveSceneUniformBuffer"));
+		checkf(MeshElement.PrimitiveUniformBuffer == nullptr, TEXT("CreateSinglePrimitiveUniformView does not currently support an explicit primitive uniform buffer on vertex factories which manually fetch primitive data.  Use PrimitiveUniformBufferResource instead."));
+
+		if (MeshElement.PrimitiveUniformBufferResource)
+		{
+			checkf(MeshElement.NumInstances == 1, TEXT("CreateSinglePrimitiveUniformView does not currently support instancing"));
+			// Force PrimitiveId to be 0 in the shader
+			MeshElement.PrimitiveIdMode = PrimID_ForceZero;
+			FPrimitiveUniformShaderParameters PrimitiveParams = *(const FPrimitiveUniformShaderParameters*)MeshElement.PrimitiveUniformBufferResource->GetContents();
+			// Now we just need to fill out the first entry of a batched primitive data in a buffer
+			FBatchedPrimitiveShaderData ShaderData(PrimitiveParams);
+			PrimitiveDataBuffer = GraphBuilder.CreateBuffer(PrimitiveDataBufferDesc, TEXT("SinglePrimitiveUniformView"));
+			GraphBuilder.QueueBufferUpload(PrimitiveDataBuffer, ShaderData.Data.GetData(), ShaderData.Data.Num() * sizeof(FVector4f));
+		}
+	}
+
+	if (PrimitiveDataBuffer == nullptr)
+	{
+		// Upload Identity parameters
+		FBatchedPrimitiveShaderData ShaderData{};
+		PrimitiveDataBuffer = GraphBuilder.CreateBuffer(PrimitiveDataBufferDesc, TEXT("SinglePrimitiveUniformView"));
+		GraphBuilder.QueueBufferUpload(PrimitiveDataBuffer, ShaderData.Data.GetData(), ShaderData.Data.Num() * sizeof(FVector4f));
+	}
+
+	BatchedPrimitiveParameters->Data = GraphBuilder.CreateSRV(PrimitiveDataBuffer);
+	return GraphBuilder.CreateUniformBuffer(BatchedPrimitiveParameters);
+}
+
 void FRendererModule::DrawTileMesh(FCanvasRenderContext& RenderContext, FMeshPassProcessorRenderState& DrawRenderState, const FSceneView& SceneView, FMeshBatch& Mesh, bool bIsHitTesting, const FHitProxyId& HitProxyId, bool bUse128bitRT)
 {
 	if (!GUsingNullRHI)
@@ -279,9 +320,13 @@ void FRendererModule::DrawTileMesh(FCanvasRenderContext& RenderContext, FMeshPas
 		auto* PassParameters = GraphBuilder.AllocParameters<FDrawTileMeshPassParameters>();
 		PassParameters->RenderTargets[0] = FRenderTargetBinding(RenderContext.GetRenderTarget(), ERenderTargetLoadAction::ELoad);
 		PassParameters->View = View.GetShaderParameters();
-		PassParameters->Scene = SceneUniforms.GetBuffer(GraphBuilder);
+		PassParameters->InstanceCullingDrawParams.Scene = SceneUniforms.GetBuffer(GraphBuilder);
+		PassParameters->InstanceCullingDrawParams.InstanceCulling = FInstanceCullingContext::CreateDummyInstanceCullingUniformBuffer(GraphBuilder);
+		if (UseGPUScene(View.GetShaderPlatform(), FeatureLevel) && PlatformGPUSceneUsesUniformBufferView(View.GetShaderPlatform()))
+		{
+			PassParameters->InstanceCullingDrawParams.BatchedPrimitive = CreateSinglePrimitiveUniformView(GraphBuilder, View, Mesh);
+		}
 		PassParameters->ReflectionCapture = EmptyReflectionCaptureUniformBuffer;
-		PassParameters->InstanceCulling = FInstanceCullingContext::CreateDummyInstanceCullingUniformBuffer(GraphBuilder);
 
 		// handle translucent material blend modes, not relevant in MaterialTexCoordScalesAnalysis since it outputs the scales.
 		if (ViewFamily->GetDebugViewShaderMode() == DVSM_OutputMaterialTextureScales)
diff --git a/Engine/Source/Runtime/Renderer/Private/RendererModule.h b/Engine/Source/Runtime/Renderer/Private/RendererModule.h
--- a/Engine/Source/Runtime/Renderer/Private/RendererModule.h
+++ b/Engine/Source/Runtime/Renderer/Private/RendererModule.h
@@ -50,6 +50,7 @@ public:
 	virtual void OnWorldCleanup(UWorld* World, bool bSessionEnded, bool bCleanupResources, bool bWorldChanged) override;
 	virtual void InitializeSystemTextures(FRHICommandListImmediate& RHICmdList);
 	virtual FSceneUniformBuffer* CreateSinglePrimitiveSceneUniformBuffer(FRDGBuilder& GraphBuilder, const FViewInfo& SceneView, FMeshBatch& Mesh) override;
+	virtual TRDGUniformBufferRef<FBatchedPrimitiveParameters> CreateSinglePrimitiveUniformView(FRDGBuilder& GraphBuilder, const FViewInfo& SceneView, FMeshBatch& Mesh) override;
 	virtual void DrawTileMesh(FCanvasRenderContext& RenderContext, FMeshPassProcessorRenderState& DrawRenderState, const FSceneView& View, FMeshBatch& Mesh, bool bIsHitTesting, const FHitProxyId& HitProxyId, bool bUse128bitRT = false) override;
 	virtual void DebugLogOnCrash() override;
 	virtual void GPUBenchmark(FSynthBenchmarkResults& InOut, float WorkScale) override;
diff --git a/Engine/Source/Runtime/Renderer/Public/InstanceCulling/InstanceCullingContext.h b/Engine/Source/Runtime/Renderer/Public/InstanceCulling/InstanceCullingContext.h
--- a/Engine/Source/Runtime/Renderer/Public/InstanceCulling/InstanceCullingContext.h
+++ b/Engine/Source/Runtime/Renderer/Public/InstanceCulling/InstanceCullingContext.h
@@ -82,7 +82,12 @@ public:
 		PreserveInstanceOrder		= 1 << 2
 	};
 
-	static constexpr uint32 UniformViewInstanceStride[2] = { 256u, 512u }; // One for each BatchProcessingMode
+	static constexpr uint32 UniformViewInstanceStride[2] = 
+	{	// One for each BatchProcessingMode
+		BATCHED_INSTANCE_DATA_STRIDE, 
+		BATCHED_PRIMITIVE_DATA_STRIDE 
+	};
+
 	static constexpr uint32 IndirectArgsNumWords = 5;
 	static constexpr uint32 CompactionBlockNumInstances = 64;
 	RENDERER_API static uint32 GetInstanceIdBufferStride(EShaderPlatform ShaderPlatform);
