Fix ensure when the OpenImageDenoise plug-in tries to download the image in a multi-GPU configuration.
Fix GPU crash for path tracing when multiple GPU is enabled.

#rnx

diff --git a/Engine/Plugins/Experimental/OpenImageDenoise/Source/OpenImageDenoise/Private/Plugin.cpp b/Engine/Plugins/Experimental/OpenImageDenoise/Source/OpenImageDenoise/Private/Plugin.cpp
--- a/Engine/Plugins/Experimental/OpenImageDenoise/Source/OpenImageDenoise/Private/Plugin.cpp
+++ b/Engine/Plugins/Experimental/OpenImageDenoise/Source/OpenImageDenoise/Private/Plugin.cpp
@@ -8,6 +8,7 @@
 #include "RHICommandList.h"
 #include "RHIResources.h"
 #include "RHITypes.h"
+#include "RHIGPUReadback.h"
 
 #include "OpenImageDenoise/oidn.hpp"
 
@@ -125,16 +126,20 @@ static OIDNState DenoiserState;
 template <typename PixelType>
 static void CopyTextureFromGPUToCPU(FRHICommandListImmediate& RHICmdList, FRHITexture* SrcTexture, FIntPoint Size, TArray<PixelType>& DstArray)
 {
-	uint32_t SrcStride = 0;
-	const PixelType* SrcBuffer = static_cast<PixelType*>(RHICmdList.LockTexture2D(SrcTexture, 0, RLM_ReadOnly, SrcStride, false));
-	SrcStride /= sizeof(PixelType);
+	FRHIGPUTextureReadback Readback(TEXT("DenoiserReadback"));
+	Readback.EnqueueCopy(RHICmdList, SrcTexture, FIntVector::ZeroValue, 0, FIntVector(Size.X, Size.Y, 1));
+	RHICmdList.BlockUntilGPUIdle();
+
+	int32_t SrcStride = 0;
+	const PixelType* SrcBuffer = static_cast<PixelType*>(Readback.Lock(SrcStride, nullptr));
+
 	PixelType* DstBuffer = DstArray.GetData();
 	for (int Y = 0; Y < Size.Y; Y++, DstBuffer += Size.X, SrcBuffer += SrcStride)
 	{
 		FPlatformMemory::Memcpy(DstBuffer, SrcBuffer, Size.X * sizeof(PixelType));
 
 	}
-	RHICmdList.UnlockTexture2D(SrcTexture, 0, false);
+	Readback.Unlock();
 }
 
 template <typename PixelType>
diff --git a/Engine/Source/Runtime/Renderer/Private/PathTracing.cpp b/Engine/Source/Runtime/Renderer/Private/PathTracing.cpp
--- a/Engine/Source/Runtime/Renderer/Private/PathTracing.cpp
+++ b/Engine/Source/Runtime/Renderer/Private/PathTracing.cpp
@@ -2899,12 +2899,12 @@ void FDeferredShadingSceneRenderer::RenderPathTracing(
 			FClearValueBinding::None,
 			TexCreate_ShaderResource | TexCreate_UAV);
 		FRDGTexture* NewRadianceTexture = GraphBuilder.CreateTexture(Desc, TEXT("PathTracer.RadianceUnswizzled"));
+		FRDGTexture* NewNormalTexture = GraphBuilder.CreateTexture(Desc, TEXT("PathTracer.NormalUnswizzled"));
 		FRDGTexture* NewAlbedoTexture = NeedsDenoise ? GraphBuilder.CreateTexture(Desc, TEXT("PathTracer.AlbedoUnswizzled")) : nullptr;
-		FRDGTexture* NewNormalTexture = NeedsDenoise ? GraphBuilder.CreateTexture(Desc, TEXT("PathTracer.NormalUnswizzled")) : nullptr;
 
-		FRDGTexture* InputTextures[3] = { RadianceTexture, AlbedoTexture, NormalTexture };
-		FRDGTexture* OutputTextures[3] = { NewRadianceTexture, NewAlbedoTexture, NewNormalTexture };
-		for (int Index = 0, Num = NeedsDenoise ? 3 : 1; Index < Num; Index++)
+		FRDGTexture* InputTextures[3] = { RadianceTexture, NormalTexture, AlbedoTexture};
+		FRDGTexture* OutputTextures[3] = { NewRadianceTexture, NewNormalTexture, NewAlbedoTexture};
+		for (int Index = 0, Num = NeedsDenoise ? 3 : 2; Index < Num; Index++)
 		{
 			FPathTracingSwizzleScanlinesCS::FParameters* PassParameters = GraphBuilder.AllocParameters<FPathTracingSwizzleScanlinesCS::FParameters>();
 			PassParameters->DispatchDim.X = DispatchResX;
@@ -2924,8 +2924,8 @@ void FDeferredShadingSceneRenderer::RenderPathTracing(
 
 		// let the remaining code operate on the unswizzled textures
 		RadianceTexture = NewRadianceTexture;
-		AlbedoTexture = NewAlbedoTexture;
 		NormalTexture = NewNormalTexture;
+		AlbedoTexture = NewAlbedoTexture;
 	}
 #endif
 
