Moved distant light invalidation logic before init dynamic shadows such that we can skip mesh setup for fully cached lights. The new logic can be disabled by turning off 'r.Shadow.Virtual.EnableNewDistantInvalidationLogic' (can be toggled at runtime).

diff --git a/Engine/Source/Runtime/Renderer/Private/DeferredShadingRenderer.cpp b/Engine/Source/Runtime/Renderer/Private/DeferredShadingRenderer.cpp
--- a/Engine/Source/Runtime/Renderer/Private/DeferredShadingRenderer.cpp
+++ b/Engine/Source/Runtime/Renderer/Private/DeferredShadingRenderer.cpp
@@ -2567,7 +2567,7 @@ void FDeferredShadingSceneRenderer::Render(FRDGBuilder& GraphBuilder)
 		}
 
 		// kick off dependent scene updates 
-		Scene->ShadowScene->UpdateForRenderedFrame(GraphBuilder);
+		ShadowSceneRenderer->BeginRender(GraphBuilder);
 
 #if RHI_RAYTRACING
 		// Initialize ray tracing flags, in case they weren't initialized in the CreateSceneRenderers code path
diff --git a/Engine/Source/Runtime/Renderer/Private/ShadowSetup.cpp b/Engine/Source/Runtime/Renderer/Private/ShadowSetup.cpp
--- a/Engine/Source/Runtime/Renderer/Private/ShadowSetup.cpp
+++ b/Engine/Source/Runtime/Renderer/Private/ShadowSetup.cpp
@@ -4082,12 +4082,9 @@ void FSceneRenderer::CreateWholeSceneProjectedShadow(
 					ProjectedShadowInfo->MeshSelectionMask = EShadowMeshSelection::VSM;
 
 					bool bContainsNaniteSubjects = false;
-					
-					// TODO: Perform no rendering work (command setup etc) if if fully cached, this optimiztion would be nice to get back, however, right now this is tricky
-					//       as the logic to select those that get to re-render must happen when the full set of rendered local lights is known.
-					//       At that point we would need to re-do this logic, which is currently rather tied to this part of the code. 
-					//       After new/persistent shadow setup is done, this should be far easier.
-					//if (!bHasFullyCachedVSM)
+
+					// Skip mesh setup if it won't be rendered anyway
+					if (!ShadowSceneRenderer->IsUsingNewDistantInvalidationLogic() || ProjectedShadowInfo->bShouldRenderVSM)
 					{
 						// Skip convex hull tests for VSM since this causes artifacts due to caching (potentially check if caching is enabled but may lead to race condition).
 						// The interaction setup has already tested the light bounds (by calling FLightSceneProxy::AffectsBounds)
@@ -4943,6 +4940,12 @@ void FSceneRenderer::BeginGatherShadowPrimitives(FDynamicShadowsTaskData* TaskDa
 		ShadowSetupPrerequisites.AddPrerequisites(VisibilityTaskData->GetComputeRelevanceTask());
 		ShadowSetupPrerequisites.AddPrerequisites(VisibilityTaskData->GetLightVisibilityTask());
 	}
+
+	if (ShadowSceneRenderer)
+	{
+		ShadowSetupPrerequisites.AddPrerequisites(ShadowSceneRenderer->GetRendererSetupTask());
+	}
+
 	ShadowSetupPrerequisites.AddPrerequisites(Scene->GetCreateLightPrimitiveInteractionsTask());
 	ShadowSetupPrerequisites.AddPrerequisites(Scene->GetCacheMeshDrawCommandsTask());
 	ShadowSetupPrerequisites.Trigger();
diff --git a/Engine/Source/Runtime/Renderer/Private/Shadows/ShadowSceneRenderer.cpp b/Engine/Source/Runtime/Renderer/Private/Shadows/ShadowSceneRenderer.cpp
--- a/Engine/Source/Runtime/Renderer/Private/Shadows/ShadowSceneRenderer.cpp
+++ b/Engine/Source/Runtime/Renderer/Private/Shadows/ShadowSceneRenderer.cpp
@@ -74,6 +74,13 @@ static TAutoConsoleVariable<float> CVarResolutionLodBiasLocalMoving(
 	ECVF_Scalability | ECVF_RenderThreadSafe
 );
 
+TAutoConsoleVariable<int32> CVarEnableNewDistantInvalidationLogic(
+	TEXT("r.Shadow.Virtual.EnableNewDistantInvalidationLogic"),
+	1,
+	TEXT("Only intended to allow hotfixing the next release of a certain project."),
+	ECVF_RenderThreadSafe
+);
+
 DECLARE_DWORD_COUNTER_STAT(TEXT("VSM Total Raster Bins"), STAT_VSMNaniteBasePassTotalRasterBins, STATGROUP_ShadowRendering);
 DECLARE_DWORD_COUNTER_STAT(TEXT("VSM Total Shading Draws"), STAT_VSMNaniteBasePassTotalShadingDraws, STATGROUP_ShadowRendering);
 
@@ -92,6 +99,7 @@ FShadowSceneRenderer::FShadowSceneRenderer(FDeferredShadingSceneRenderer& InScen
 	, Scene(*InSceneRenderer.Scene)
 	, ShadowScene(*Scene.ShadowScene)
 	, VirtualShadowMapArray(InSceneRenderer.VirtualShadowMapArray)
+	, bShouldUseNewDistantInvalidationLogic(CVarEnableNewDistantInvalidationLogic.GetValueOnAnyThread() != 0)
 {
 }
 
@@ -100,6 +108,48 @@ float FShadowSceneRenderer::ComputeNaniteShadowsLODScaleFactor()
 	return FMath::Pow(2.0f, -CVarNaniteShadowsLODBias.GetValueOnRenderThread());
 }
 
+void FShadowSceneRenderer::BeginRender(FRDGBuilder& GraphBuilder)
+{
+	// Kick off shadow scene updates.
+	ShadowScene.UpdateForRenderedFrame(GraphBuilder);
+
+	if (bShouldUseNewDistantInvalidationLogic && VirtualShadowMapArray.IsEnabled() && VirtualShadowMapArray.CacheManager->IsCacheEnabled())
+	{
+		RendererSetupTask = GraphBuilder.AddSetupTask([this]()
+		{
+			FVirtualShadowMapArrayCacheManager& CacheManager = *VirtualShadowMapArray.CacheManager;
+			// Priority queue of distant lights to update.
+			int32 SceneFrameNumber = int32(Scene.GetFrameNumber());
+
+			// Build age-based priority queue of fully cached entries.
+			for (auto It = Scene.Lights.CreateConstIterator(); It; ++It)
+			{
+				int32 LightId = It.GetIndex();
+				TSharedPtr<FVirtualShadowMapPerLightCacheEntry> PerLightCacheEntry = CacheManager.FindLightCacheEntry(LightId, 0);
+				if (PerLightCacheEntry && PerLightCacheEntry->IsFullyCached())
+				{
+					int32 FramesSinceLastRender = SceneFrameNumber - int32(PerLightCacheEntry->GetLastScheduledFrameNumber());
+					DistantLightUpdateQueue.Add(-FramesSinceLastRender, LightId);
+				}
+			}
+
+			int32 UpdateBudgetNumLights = CVarMaxDistantLightsPerFrame.GetValueOnRenderThread() < 0 ? int32(DistantLightUpdateQueue.Num()) : FMath::Min(int32(DistantLightUpdateQueue.Num()), CVarMaxDistantLightsPerFrame.GetValueOnRenderThread());
+			for (int32 Index = 0; Index < UpdateBudgetNumLights; ++Index)
+			{
+				const int32 LightId = DistantLightUpdateQueue.Top();
+				//const int32 Age = DistantLightUpdateQueue.GetKey(LocalLightShadowIndex);
+				// UE_LOG(LogTemp, Log, TEXT("Index: %d Age: %d"), LocalLightShadowIndex, Age);
+				DistantLightUpdateQueue.Pop();
+
+				TSharedPtr<FVirtualShadowMapPerLightCacheEntry> PerLightCacheEntry = CacheManager.FindLightCacheEntry(LightId, 0);
+
+				// Mark frame it was scheduled, this is picked up later in AddLocalLightShadow to trigger invalidation 
+				PerLightCacheEntry->Current.ScheduledFrameNumber = SceneFrameNumber;
+			}
+		});
+	}
+}
+
 static float GetResolutionLODBiasLocal(float LightMobilityFactor)
 {
 	return FVirtualShadowMapArray::InterpolateResolutionBias(
@@ -170,6 +220,13 @@ TSharedPtr<FVirtualShadowMapPerLightCacheEntry> FShadowSceneRenderer::AddLocalLi
 	LocalLightShadowFrameSetup.PerLightCacheEntry = PerLightCacheEntry;
 	bool bIsCached = PerLightCacheEntry->UpdateLocal(ProjectedShadowInitializer, bIsDistantLight, CacheManager->IsCacheEnabled(), !bShouldForceTimeSliceDistantUpdate);
 
+	if (bShouldUseNewDistantInvalidationLogic)
+	{
+		if (bIsCached && bIsDistantLight && PerLightCacheEntry->Prev.ScheduledFrameNumber == Scene.GetFrameNumber())
+		{
+			PerLightCacheEntry->Invalidate();
+		}
+	}
 	// Update info on the ProjectionShadowInfo; eventually this should all move into local data structures here
 	const int32 VirtualShadowMapId = VirtualShadowMapArray.Allocate(bIsDistantLight, NumMaps);
 	ProjectedShadowInfo->VirtualShadowMapId = VirtualShadowMapId;
@@ -185,15 +242,17 @@ TSharedPtr<FVirtualShadowMapPerLightCacheEntry> FShadowSceneRenderer::AddLocalLi
 		VirtualSmCacheEntry.ProjectionData = GetLocalLightProjectionShaderData(ResolutionLODBiasLocal, ProjectedShadowInfo, Index);
 	}
 
-	// Only round-robin those that were not invalidated.
-	if (bIsDistantLight && bIsCached)
+	if (!bShouldUseNewDistantInvalidationLogic)
 	{
-		// This priority could be calculated based also on whether the light has actually been invalidated or not (currently not tracked on host).
-		// E.g., all things being equal update those with an animated mesh in, for example. Plus don't update those the don't need it at all.
-		int32 FramesSinceLastRender = int32(Scene.GetFrameNumber()) - int32(PerLightCacheEntry->GetLastScheduledFrameNumber());
-		DistantLightUpdateQueue.Add(-FramesSinceLastRender, LocalLightShadowIndex);
+		// Only round-robin those that were not invalidated.
+		if (bIsDistantLight && bIsCached)
+		{
+			// This priority could be calculated based also on whether the light has actually been invalidated or not (currently not tracked on host).
+			// E.g., all things being equal update those with an animated mesh in, for example. Plus don't update those the don't need it at all.
+			int32 FramesSinceLastRender = int32(Scene.GetFrameNumber()) - int32(PerLightCacheEntry->GetLastScheduledFrameNumber());
+			DistantLightUpdateQueue.Add(-FramesSinceLastRender, LocalLightShadowIndex);
+		}
 	}
-
 	return PerLightCacheEntry;
 }
 
@@ -205,8 +264,10 @@ void FShadowSceneRenderer::AddDirectionalLightShadow(FProjectedShadowInfo* Proje
 
 void FShadowSceneRenderer::PostInitDynamicShadowsSetup()
 {
-	UpdateDistantLightPriorityRender();
-
+	if (!bShouldUseNewDistantInvalidationLogic)
+	{
+		UpdateDistantLightPriorityRender();
+	}
 	// Dispatch async Nanite culling job if appropriate
 	if (CVarVSMMaterialVisibility.GetValueOnRenderThread() != 0)
 	{
diff --git a/Engine/Source/Runtime/Renderer/Private/Shadows/ShadowSceneRenderer.h b/Engine/Source/Runtime/Renderer/Private/Shadows/ShadowSceneRenderer.h
--- a/Engine/Source/Runtime/Renderer/Private/Shadows/ShadowSceneRenderer.h
+++ b/Engine/Source/Runtime/Renderer/Private/Shadows/ShadowSceneRenderer.h
@@ -6,6 +6,7 @@
 #include "Containers/ArrayView.h"
 #include "Containers/BinaryHeap.h"
 #include "Containers/Array.h"
+#include "Tasks/Task.h"
 
 #include "VirtualShadowMaps/VirtualShadowMapArray.h"
 
@@ -30,6 +31,11 @@ public:
 	 */
 	static float ComputeNaniteShadowsLODScaleFactor();
 
+	/**
+	 * Call at the start of the frame to kick off work that can be done early.
+	 */
+	void BeginRender(FRDGBuilder& GraphBuilder);
+
 	/**
 	 * Add a cube/spot light for processing this frame.
 	 * TODO: Don't use legacy FProjectedShadowInfo or other params, instead info should flow from persistent setup & update.
@@ -85,8 +91,16 @@ public:
 	{
 		return VirtualShadowMapMaskBits != nullptr;
 	}
+	
+	UE::Tasks::FTask GetRendererSetupTask() const
+	{
+		return RendererSetupTask;
+	}
 
+	bool IsUsingNewDistantInvalidationLogic() const { return bShouldUseNewDistantInvalidationLogic; }
 private:
+	UE::Tasks::FTask RendererSetupTask;
+
 	FVirtualShadowMapProjectionShaderData GetLocalLightProjectionShaderData(float ResolutionLODBiasLocal, const FProjectedShadowInfo* ProjectedShadowInfo, int32 MapIndex) const;
 
 	/**
@@ -127,4 +141,6 @@ private:
 	FNaniteVisibilityQuery* NaniteVisibilityQuery = nullptr;
 	Nanite::FPackedViewArray* VirtualShadowMapViews = nullptr;
 	FSceneInstanceCullingQuery *SceneInstanceCullingQuery = nullptr;
+
+	bool bShouldUseNewDistantInvalidationLogic = true;
 };
diff --git a/Engine/Source/Runtime/Renderer/Private/VirtualShadowMaps/VirtualShadowMapCacheManager.cpp b/Engine/Source/Runtime/Renderer/Private/VirtualShadowMaps/VirtualShadowMapCacheManager.cpp
--- a/Engine/Source/Runtime/Renderer/Private/VirtualShadowMaps/VirtualShadowMapCacheManager.cpp
+++ b/Engine/Source/Runtime/Renderer/Private/VirtualShadowMaps/VirtualShadowMapCacheManager.cpp
@@ -843,6 +843,17 @@ TSharedPtr<FVirtualShadowMapPerLightCacheEntry> FVirtualShadowMapArrayCacheManag
 	return LightEntry;
 }
 
+TSharedPtr<FVirtualShadowMapPerLightCacheEntry> FVirtualShadowMapArrayCacheManager::FindLightCacheEntry(int32 LightSceneId, uint32 ViewUniqueID)
+{
+	const uint64 CacheKey = (uint64(ViewUniqueID) << 32U) | uint64(LightSceneId);
+
+	if (TSharedPtr<FVirtualShadowMapPerLightCacheEntry> *Entry = CacheEntries.Find(CacheKey))
+	{
+		return *Entry;
+	}
+	return TSharedPtr<FVirtualShadowMapPerLightCacheEntry>();
+}
+
 void FVirtualShadowMapPerLightCacheEntry::OnPrimitiveRendered(const FPrimitiveSceneInfo* PrimitiveSceneInfo)
 {
 	// Mark as (potentially present in a cached page somehwere, so we'd need to invalidate if it is removed/moved)
diff --git a/Engine/Source/Runtime/Renderer/Private/VirtualShadowMaps/VirtualShadowMapCacheManager.h b/Engine/Source/Runtime/Renderer/Private/VirtualShadowMaps/VirtualShadowMapCacheManager.h
--- a/Engine/Source/Runtime/Renderer/Private/VirtualShadowMaps/VirtualShadowMapCacheManager.h
+++ b/Engine/Source/Runtime/Renderer/Private/VirtualShadowMaps/VirtualShadowMapCacheManager.h
@@ -236,6 +236,11 @@ public:
 	 */
 	TSharedPtr<FVirtualShadowMapPerLightCacheEntry> FindCreateLightCacheEntry(int32 LightSceneId, uint32 ViewUniqueID, uint32 NumShadowMaps);
 
+	/**
+	 * Finds an existing cache entry, returns null if none exists.
+	 */
+	TSharedPtr<FVirtualShadowMapPerLightCacheEntry> FindLightCacheEntry(int32 LightSceneId, uint32 ViewUniqueID);
+
 	bool IsCacheEnabled();
 	bool IsCacheDataAvailable();
 	bool IsHZBDataAvailable();
