Make r.InstanceCulling.OcclusionCull have ECVF_Preview flag

diff --git a/Engine/Source/Runtime/Renderer/Private/InstanceCulling/InstanceCullingContext.cpp b/Engine/Source/Runtime/Renderer/Private/InstanceCulling/InstanceCullingContext.cpp
--- a/Engine/Source/Runtime/Renderer/Private/InstanceCulling/InstanceCullingContext.cpp
+++ b/Engine/Source/Runtime/Renderer/Private/InstanceCulling/InstanceCullingContext.cpp
@@ -29,7 +29,7 @@ static TAutoConsoleVariable<int32> CVarOcclusionCullInstances(
 	TEXT("r.InstanceCulling.OcclusionCull"),
 	0,
 	TEXT("Whether to do per instance occlusion culling for GPU instance culling."),
-	ECVF_RenderThreadSafe);
+	ECVF_RenderThreadSafe | ECVF_Preview);
 
 static int32 GInstanceCullingAllowOrderPreservation = 1;
 static FAutoConsoleVariableRef CVarInstanceCullingAllowOrderPreservation(
