Modified VT system to take a command list to support async creation of virtual textures.

diff --git a/Engine/Source/Runtime/Engine/Private/Texture2D.cpp b/Engine/Source/Runtime/Engine/Private/Texture2D.cpp
--- a/Engine/Source/Runtime/Engine/Private/Texture2D.cpp
+++ b/Engine/Source/Runtime/Engine/Private/Texture2D.cpp
@@ -1489,7 +1489,7 @@ FVirtualTexture2DResource::~FVirtualTexture2DResource()
 {
 }
 
-void FVirtualTexture2DResource::InitRHI(FRHICommandListBase&)
+void FVirtualTexture2DResource::InitRHI(FRHICommandListBase& RHICmdList)
 {
 	TRACE_CPUPROFILER_EVENT_SCOPE(FVirtualTexture2DResource::InitRHI);
 	LLM_SCOPE_DYNAMIC_STAT_OBJECTPATH_FNAME(PackageName, ELLMTagSet::Assets);
@@ -1555,7 +1555,7 @@ void FVirtualTexture2DResource::InitRHI(FRHICommandListBase&)
 	}
 
 	FUploadingVirtualTexture* VirtualTexture = new FUploadingVirtualTexture(ProducerDesc.Name, VTData, FirstMipToUse);
-	ProducerHandle = GetRendererModule().RegisterVirtualTextureProducer(ProducerDesc, VirtualTexture);
+	ProducerHandle = GetRendererModule().RegisterVirtualTextureProducer(RHICmdList, ProducerDesc, VirtualTexture);
 
 	// Only create the miptails mini-texture in-editor.
 #if WITH_EDITOR
diff --git a/Engine/Source/Runtime/Engine/Private/VT/RuntimeVirtualTexture.cpp b/Engine/Source/Runtime/Engine/Private/VT/RuntimeVirtualTexture.cpp
--- a/Engine/Source/Runtime/Engine/Private/VT/RuntimeVirtualTexture.cpp
+++ b/Engine/Source/Runtime/Engine/Private/VT/RuntimeVirtualTexture.cpp
@@ -160,8 +160,8 @@ public:
 			const bool bAdaptive = FillVTDescriptions(Resource->ProducerHandle, InProducerDesc, InInitDesc, AllocatedVTDesc, AdaptiveVTDesc);
 
 			// Only one or none of these should be allocated...
-			Resource->AllocatedVirtualTexture = !bAdaptive ? AllocateVirtualTexture(AllocatedVTDesc) : nullptr;
-			Resource->AdaptiveVirtualTexture = bAdaptive ? AllocateAdaptiveVirtualTexture(AllocatedVTDesc, AdaptiveVTDesc) : nullptr;
+			Resource->AllocatedVirtualTexture = !bAdaptive ? AllocateVirtualTexture(RHICmdList, AllocatedVTDesc) : nullptr;
+			Resource->AdaptiveVirtualTexture = bAdaptive ? AllocateAdaptiveVirtualTexture(RHICmdList, AllocatedVTDesc, AdaptiveVTDesc) : nullptr;
 
 			// Release old producer after new one is created so that any destroy callbacks can access the new producer
 			GetRendererModule().ReleaseVirtualTextureProducer(OldProducerHandle);
@@ -223,13 +223,13 @@ protected:
 	}
 
 	/** Allocate in the virtual texture system. */
-	static IAllocatedVirtualTexture* AllocateVirtualTexture(FAllocatedVTDescription const& InAllocatedVTDesc)
+	static IAllocatedVirtualTexture* AllocateVirtualTexture(FRHICommandListBase& RHICmdList, FAllocatedVTDescription const& InAllocatedVTDesc)
 	{
 		// Check for NumLayers avoids allocating for the null producer
 		IAllocatedVirtualTexture* OutAllocatedVirtualTexture = nullptr;
 		if (InAllocatedVTDesc.NumTextureLayers > 0)
 		{ 
-			OutAllocatedVirtualTexture = GetRendererModule().AllocateVirtualTexture(InAllocatedVTDesc);
+			OutAllocatedVirtualTexture = GetRendererModule().AllocateVirtualTexture(RHICmdList, InAllocatedVTDesc);
 		}
 		return OutAllocatedVirtualTexture;
 	}
@@ -244,13 +244,13 @@ protected:
 	}
 
 	/** Allocate an adaptive virtual texture in the virtual texture system. */
-	static IAdaptiveVirtualTexture* AllocateAdaptiveVirtualTexture(FAllocatedVTDescription const& InAllocatedVTDesc, FAdaptiveVTDescription const& InAdaptiveVTDesc)
+	static IAdaptiveVirtualTexture* AllocateAdaptiveVirtualTexture(FRHICommandListBase& RHICmdList, FAllocatedVTDescription const& InAllocatedVTDesc, FAdaptiveVTDescription const& InAdaptiveVTDesc)
 	{
 		// Check for NumLayers avoids allocating for the null producer
 		IAdaptiveVirtualTexture* OutAdaptiveVirtualTexture = nullptr;
 		if (InAllocatedVTDesc.NumTextureLayers > 0)
 		{
-			OutAdaptiveVirtualTexture = GetRendererModule().AllocateAdaptiveVirtualTexture(InAdaptiveVTDesc, InAllocatedVTDesc);
+			OutAdaptiveVirtualTexture = GetRendererModule().AllocateAdaptiveVirtualTexture(RHICmdList, InAdaptiveVTDesc, InAllocatedVTDesc);
 		}
 		return OutAdaptiveVirtualTexture;
 	}
diff --git a/Engine/Source/Runtime/RenderCore/Private/RenderCore.cpp b/Engine/Source/Runtime/RenderCore/Private/RenderCore.cpp
--- a/Engine/Source/Runtime/RenderCore/Private/RenderCore.cpp
+++ b/Engine/Source/Runtime/RenderCore/Private/RenderCore.cpp
@@ -14,6 +14,7 @@
 #include "RenderTimer.h"
 #include "RenderCounters.h"
 #include "RenderingThread.h"
+#include "RendererInterface.h"
 
 void UpdateShaderDevelopmentMode();
 
@@ -1026,3 +1027,18 @@ RENDERCORE_API void ConvertPixelDataToSCRGB(TArray<FLinearColor>& InOutRawPixels
 		}
 	}
 }
+
+FVirtualTextureProducerHandle IRendererModule::RegisterVirtualTextureProducer(const FVTProducerDescription& Desc, IVirtualTexture* Producer)
+{
+	return RegisterVirtualTextureProducer(FRHICommandListImmediate::Get(), Desc, Producer);
+}
+
+IAllocatedVirtualTexture* IRendererModule::AllocateVirtualTexture(const FAllocatedVTDescription& Desc)
+{
+	return AllocateVirtualTexture(FRHICommandListImmediate::Get(), Desc);
+}
+
+IAdaptiveVirtualTexture* IRendererModule::AllocateAdaptiveVirtualTexture(const FAdaptiveVTDescription& AdaptiveVTDesc, const FAllocatedVTDescription& AllocatedVTDesc)
+{
+	return AllocateAdaptiveVirtualTexture(FRHICommandListImmediate::Get(), AdaptiveVTDesc, AllocatedVTDesc);
+}
\ No newline at end of file
diff --git a/Engine/Source/Runtime/RenderCore/Public/RendererInterface.h b/Engine/Source/Runtime/RenderCore/Public/RendererInterface.h
--- a/Engine/Source/Runtime/RenderCore/Public/RendererInterface.h
+++ b/Engine/Source/Runtime/RenderCore/Public/RendererInterface.h
@@ -786,13 +786,18 @@ public:
 	/** Performs necessary per-frame cleanup. Only use when rendering through scene renderer (i.e. BeginRenderingViewFamily) is skipped */
 	virtual void PerFrameCleanupIfSkipRenderer() = 0;
 
-	virtual IAllocatedVirtualTexture* AllocateVirtualTexture(const FAllocatedVTDescription& Desc) = 0;
+	virtual IAllocatedVirtualTexture* AllocateVirtualTexture(FRHICommandListBase& RHICmdList, const FAllocatedVTDescription& Desc) = 0;
+	RENDERCORE_API IAllocatedVirtualTexture* AllocateVirtualTexture(const FAllocatedVTDescription& Desc);
+
 	virtual void DestroyVirtualTexture(IAllocatedVirtualTexture* AllocatedVT) = 0;
 
-	virtual IAdaptiveVirtualTexture* AllocateAdaptiveVirtualTexture(const FAdaptiveVTDescription& AdaptiveVTDesc, const FAllocatedVTDescription& AllocatedVTDesc) = 0;
+	virtual IAdaptiveVirtualTexture* AllocateAdaptiveVirtualTexture(FRHICommandListBase& RHICmdList, const FAdaptiveVTDescription& AdaptiveVTDesc, const FAllocatedVTDescription& AllocatedVTDesc) = 0;
+	RENDERCORE_API IAdaptiveVirtualTexture* AllocateAdaptiveVirtualTexture(const FAdaptiveVTDescription& AdaptiveVTDesc, const FAllocatedVTDescription& AllocatedVTDesc);
 	virtual void DestroyAdaptiveVirtualTexture(IAdaptiveVirtualTexture* AdaptiveVT) = 0;
 
-	virtual FVirtualTextureProducerHandle RegisterVirtualTextureProducer(const FVTProducerDescription& Desc, IVirtualTexture* Producer) = 0;
+	virtual FVirtualTextureProducerHandle RegisterVirtualTextureProducer(FRHICommandListBase& RHICmdList, const FVTProducerDescription& Desc, IVirtualTexture* Producer) = 0;
+	RENDERCORE_API FVirtualTextureProducerHandle RegisterVirtualTextureProducer(const FVTProducerDescription& Desc, IVirtualTexture* Producer);
+
 	virtual void ReleaseVirtualTextureProducer(const FVirtualTextureProducerHandle& Handle) = 0;
 	virtual void AddVirtualTextureProducerDestroyedCallback(const FVirtualTextureProducerHandle& Handle, FVTProducerDestroyedFunction* Function, void* Baton) = 0;
 	virtual uint32 RemoveAllVirtualTextureProducerDestroyedCallbacks(const void* Baton) = 0;
diff --git a/Engine/Source/Runtime/Renderer/Private/RendererModule.h b/Engine/Source/Runtime/Renderer/Private/RendererModule.h
--- a/Engine/Source/Runtime/Renderer/Private/RendererModule.h
+++ b/Engine/Source/Runtime/Renderer/Private/RendererModule.h
@@ -94,11 +94,11 @@ public:
 
 	virtual void PerFrameCleanupIfSkipRenderer() override;
 
-	virtual IAllocatedVirtualTexture* AllocateVirtualTexture(const FAllocatedVTDescription& Desc) override;
+	virtual IAllocatedVirtualTexture* AllocateVirtualTexture(FRHICommandListBase& RHICmdList, const FAllocatedVTDescription& Desc) override;
 	virtual void DestroyVirtualTexture(IAllocatedVirtualTexture* AllocatedVT) override;
-	virtual IAdaptiveVirtualTexture* AllocateAdaptiveVirtualTexture(const FAdaptiveVTDescription& AdaptiveVTDesc, const FAllocatedVTDescription& AllocatedVTDesc) override;
+	virtual IAdaptiveVirtualTexture* AllocateAdaptiveVirtualTexture(FRHICommandListBase& RHICmdList, const FAdaptiveVTDescription& AdaptiveVTDesc, const FAllocatedVTDescription& AllocatedVTDesc) override;
 	virtual void DestroyAdaptiveVirtualTexture(IAdaptiveVirtualTexture* AdaptiveVT) override;
-	virtual FVirtualTextureProducerHandle RegisterVirtualTextureProducer(const FVTProducerDescription& Desc, IVirtualTexture* Producer) override;
+	virtual FVirtualTextureProducerHandle RegisterVirtualTextureProducer(FRHICommandListBase& RHICmdList, const FVTProducerDescription& Desc, IVirtualTexture* Producer) override;
 	virtual void ReleaseVirtualTextureProducer(const FVirtualTextureProducerHandle& Handle) override;
 	virtual void AddVirtualTextureProducerDestroyedCallback(const FVirtualTextureProducerHandle& Handle, FVTProducerDestroyedFunction* Function, void* Baton) override;
 	virtual uint32 RemoveAllVirtualTextureProducerDestroyedCallbacks(const void* Baton) override;
diff --git a/Engine/Source/Runtime/Renderer/Private/SceneRendering.cpp b/Engine/Source/Runtime/Renderer/Private/SceneRendering.cpp
--- a/Engine/Source/Runtime/Renderer/Private/SceneRendering.cpp
+++ b/Engine/Source/Runtime/Renderer/Private/SceneRendering.cpp
@@ -5053,9 +5053,9 @@ IScenePrimitiveRenderingContext* FRendererModule::BeginScenePrimitiveRendering(F
 }
 
 
-IAllocatedVirtualTexture* FRendererModule::AllocateVirtualTexture(const FAllocatedVTDescription& Desc)
+IAllocatedVirtualTexture* FRendererModule::AllocateVirtualTexture(FRHICommandListBase& RHICmdList, const FAllocatedVTDescription& Desc)
 {
-	return FVirtualTextureSystem::Get().AllocateVirtualTexture(Desc);
+	return FVirtualTextureSystem::Get().AllocateVirtualTexture(RHICmdList, Desc);
 }
 
 void FRendererModule::DestroyVirtualTexture(IAllocatedVirtualTexture* AllocatedVT)
@@ -5063,9 +5063,9 @@ void FRendererModule::DestroyVirtualTexture(IAllocatedVirtualTexture* AllocatedV
 	FVirtualTextureSystem::Get().DestroyVirtualTexture(AllocatedVT);
 }
 
-IAdaptiveVirtualTexture* FRendererModule::AllocateAdaptiveVirtualTexture(const FAdaptiveVTDescription& AdaptiveVTDesc, const FAllocatedVTDescription& AllocatedVTDesc)
+IAdaptiveVirtualTexture* FRendererModule::AllocateAdaptiveVirtualTexture(FRHICommandListBase& RHICmdList, const FAdaptiveVTDescription& AdaptiveVTDesc, const FAllocatedVTDescription& AllocatedVTDesc)
 {
-	return FVirtualTextureSystem::Get().AllocateAdaptiveVirtualTexture(AdaptiveVTDesc, AllocatedVTDesc);
+	return FVirtualTextureSystem::Get().AllocateAdaptiveVirtualTexture(RHICmdList, AdaptiveVTDesc, AllocatedVTDesc);
 }
 
 void FRendererModule::DestroyAdaptiveVirtualTexture(IAdaptiveVirtualTexture* AdaptiveVT)
@@ -5073,9 +5073,9 @@ void FRendererModule::DestroyAdaptiveVirtualTexture(IAdaptiveVirtualTexture* Ada
 	FVirtualTextureSystem::Get().DestroyAdaptiveVirtualTexture(AdaptiveVT);
 }
 
-FVirtualTextureProducerHandle FRendererModule::RegisterVirtualTextureProducer(const FVTProducerDescription& Desc, IVirtualTexture* Producer)
+FVirtualTextureProducerHandle FRendererModule::RegisterVirtualTextureProducer(FRHICommandListBase& RHICmdList, const FVTProducerDescription& Desc, IVirtualTexture* Producer)
 {
-	return FVirtualTextureSystem::Get().RegisterProducer(Desc, Producer);
+	return FVirtualTextureSystem::Get().RegisterProducer(RHICmdList, Desc, Producer);
 }
 
 void FRendererModule::ReleaseVirtualTextureProducer(const FVirtualTextureProducerHandle& Handle)
diff --git a/Engine/Source/Runtime/Renderer/Private/VT/AdaptiveVirtualTexture.cpp b/Engine/Source/Runtime/Renderer/Private/VT/AdaptiveVirtualTexture.cpp
--- a/Engine/Source/Runtime/Renderer/Private/VT/AdaptiveVirtualTexture.cpp
+++ b/Engine/Source/Runtime/Renderer/Private/VT/AdaptiveVirtualTexture.cpp
@@ -146,6 +146,7 @@ union FPackedAdaptiveAllocationRequest
 
 /** Allocate a virtual texture for a subset of the full adaptive virtual texture. */
 IAllocatedVirtualTexture* FAdaptiveVirtualTexture::AllocateVirtualTexture(
+	FRHICommandListBase& RHICmdList,
 	FVirtualTextureSystem* InSystem,
 	FAllocatedVTDescription const& InAllocatedDesc,
 	FIntPoint InGridSize,
@@ -182,7 +183,7 @@ IAllocatedVirtualTexture* FAdaptiveVirtualTexture::AllocateVirtualTexture(
 
 		IVirtualTexture* VirtualTextureProducer = Producer->GetVirtualTexture();
 		IVirtualTexture* NewVirtualTextureProducer = new FVirtualTextureAddressRedirect(VirtualTextureProducer, InAddressOffset, InLevelOffset);
-		FVirtualTextureProducerHandle NewProducerHandle = InSystem->RegisterProducer(NewProducerDesc, NewVirtualTextureProducer);
+		FVirtualTextureProducerHandle NewProducerHandle = InSystem->RegisterProducer(RHICmdList, NewProducerDesc, NewVirtualTextureProducer);
 
 		// Copy new producer to all subsequent layers.
 		for (int32 WriteLayerIndex = LayerIndex; WriteLayerIndex < InAllocatedDesc.NumTextureLayers; ++WriteLayerIndex)
@@ -194,7 +195,7 @@ IAllocatedVirtualTexture* FAdaptiveVirtualTexture::AllocateVirtualTexture(
 		}
 	}
 
-	return InSystem->AllocateVirtualTexture(AllocatedDesc);
+	return InSystem->AllocateVirtualTexture(RHICmdList, AllocatedDesc);
 }
 
 /** Destroy an allocated virtual texture and release its producers. */
@@ -265,11 +266,11 @@ FAdaptiveVirtualTexture::FAdaptiveVirtualTexture(
 	GridSize = FIntPoint(1 << FMath::Max(AdaptiveGridLevelsX, 0), 1 << FMath::Max(AdaptiveGridLevelsY, 0));
 }
 
-void FAdaptiveVirtualTexture::Init(FVirtualTextureSystem* InSystem)
+void FAdaptiveVirtualTexture::Init(FRHICommandListBase& RHICmdList, FVirtualTextureSystem* InSystem)
 {
 	// Allocate a low mips virtual texture.
 	const int32 LevelOffset = AdaptiveDesc.MaxAdaptiveLevel;
-	AllocatedVirtualTextureLowMips = (FAllocatedVirtualTexture*)AllocateVirtualTexture(InSystem, AllocatedDesc, GridSize, 0xff, GridSize.X, GridSize.Y, FIntPoint::ZeroValue, LevelOffset);
+	AllocatedVirtualTextureLowMips = (FAllocatedVirtualTexture*)AllocateVirtualTexture(RHICmdList, InSystem, AllocatedDesc, GridSize, 0xff, GridSize.X, GridSize.Y, FIntPoint::ZeroValue, LevelOffset);
 }
 
 void FAdaptiveVirtualTexture::Destroy(FVirtualTextureSystem* InSystem)
@@ -517,7 +518,7 @@ void FAdaptiveVirtualTexture::QueuePackedAllocationRequests(uint32 const* InRequ
 	}
 }
 
-void FAdaptiveVirtualTexture::Allocate(FVirtualTextureSystem* InSystem, uint32 InPackedRequest, uint32 InFrame)
+void FAdaptiveVirtualTexture::Allocate(FRHICommandListBase& RHICmdList, FVirtualTextureSystem* InSystem, uint32 InPackedRequest, uint32 InFrame)
 {
 	FPackedAdaptiveAllocationRequest Request;
 	Request.PackedValue = InPackedRequest;
@@ -539,10 +540,10 @@ void FAdaptiveVirtualTexture::Allocate(FVirtualTextureSystem* InSystem, uint32 I
 		return;
 	}
 
-	Allocate(InSystem, GridIndex, AllocationIndex, NewLevel, InFrame);
+	Allocate(RHICmdList, InSystem, GridIndex, AllocationIndex, NewLevel, InFrame);
 }
 
-void FAdaptiveVirtualTexture::Allocate(FVirtualTextureSystem* InSystem, uint32 InGridIndex, uint32 InAllocationIndex, uint32 InNewLevel, uint32 InFrame)
+void FAdaptiveVirtualTexture::Allocate(FRHICommandListBase& RHICmdList, FVirtualTextureSystem* InSystem, uint32 InGridIndex, uint32 InAllocationIndex, uint32 InNewLevel, uint32 InFrame)
 {
 	const uint32 X = InGridIndex % GridSize.X;
 	const uint32 Y = InGridIndex / GridSize.X;
@@ -550,7 +551,7 @@ void FAdaptiveVirtualTexture::Allocate(FVirtualTextureSystem* InSystem, uint32 I
 	const int32 LevelOffset = (int32)AdaptiveDesc.MaxAdaptiveLevel - (int32)InNewLevel;
 
 	FAllocatedVirtualTexture* OldAllocatedVT = (InAllocationIndex != INDEX_NONE) ? AllocationSlots[InAllocationIndex].AllocatedVT : nullptr;
-	FAllocatedVirtualTexture* NewAllocatedVT = (FAllocatedVirtualTexture*)AllocateVirtualTexture(InSystem, AllocatedDesc, GridSize, GetSpaceID(), 1 << InNewLevel, 1 << InNewLevel, PageOffset, LevelOffset);
+	FAllocatedVirtualTexture* NewAllocatedVT = (FAllocatedVirtualTexture*)AllocateVirtualTexture(RHICmdList, InSystem, AllocatedDesc, GridSize, GetSpaceID(), 1 << InNewLevel, 1 << InNewLevel, PageOffset, LevelOffset);
 
 	if (OldAllocatedVT != nullptr)
 	{
@@ -627,7 +628,7 @@ void FAdaptiveVirtualTexture::Free(FVirtualTextureSystem* InSystem, uint32 InAll
 	TextureUpdates.Add(FIndirectionTextureUpdate{ GridIndex % GridSize.X, GridIndex / GridSize.X, 0 });
 }
 
-bool FAdaptiveVirtualTexture::FreeLRU(FVirtualTextureSystem* InSystem, uint32 InFrame, uint32 InFrameAgeToFree)
+bool FAdaptiveVirtualTexture::FreeLRU(FRHICommandListBase& RHICmdList, FVirtualTextureSystem* InSystem, uint32 InFrame, uint32 InFrameAgeToFree)
 {
 	// Check if top is ready for eviction.
 	const uint32 AllocationIndex = LRUHeap.Top();
@@ -665,7 +666,7 @@ bool FAdaptiveVirtualTexture::FreeLRU(FVirtualTextureSystem* InSystem, uint32 In
 	{
 		// Reallocate to the selected level.
 		const uint32 GridIndex = AllocationSlots[AllocationIndex].GridIndex;
-		Allocate(InSystem, GridIndex, AllocationIndex, NewLevel, InFrame);
+		Allocate(RHICmdList, InSystem, GridIndex, AllocationIndex, NewLevel, InFrame);
 	}
 
 	return true;
@@ -682,7 +683,7 @@ void FAdaptiveVirtualTexture::UpdateAllocations(FVirtualTextureSystem* InSystem,
 		bool bFreeSuccess = true;
 		for (int32 FreeCount = 0; bFreeSuccess && FreeCount < NumToFree; FreeCount++)
 		{
-			bFreeSuccess = FreeLRU(InSystem, InFrame, FrameAgeToFree);
+			bFreeSuccess = FreeLRU(RHICmdList, InSystem, InFrame, FrameAgeToFree);
 		}
 	}
 	else
@@ -698,7 +699,7 @@ void FAdaptiveVirtualTexture::UpdateAllocations(FVirtualTextureSystem* InSystem,
 		for (int32 FreeCount = 0; bFreeSuccess && FreeCount < NumToFree && Space->GetAllocator().GetNumAllocatedPages() > TargetPages; FreeCount++)
 		{
 			const uint32 FrameAgeToFree = 15; // Hardcoded threshold. Don't release anything used more recently then this.
-			bFreeSuccess = FreeLRU(InSystem, InFrame, FrameAgeToFree);
+			bFreeSuccess = FreeLRU(RHICmdList, InSystem, InFrame, FrameAgeToFree);
 		}
 
 		// Process allocation requests.
@@ -709,7 +710,7 @@ void FAdaptiveVirtualTexture::UpdateAllocations(FVirtualTextureSystem* InSystem,
 			// Randomize request order to prevent feedback from top of the view being prioritized.
 			int32 RequestIndex = FMath::Rand() % RequestsToMap.Num();
 			uint32 PackedRequest = RequestsToMap[RequestIndex];
-			Allocate(InSystem, PackedRequest, InFrame);
+			Allocate(RHICmdList, InSystem, PackedRequest, InFrame);
 			RequestsToMap.RemoveAtSwap(RequestIndex, 1, false);
 		}
 	}
diff --git a/Engine/Source/Runtime/Renderer/Private/VT/AdaptiveVirtualTexture.h b/Engine/Source/Runtime/Renderer/Private/VT/AdaptiveVirtualTexture.h
--- a/Engine/Source/Runtime/Renderer/Private/VT/AdaptiveVirtualTexture.h
+++ b/Engine/Source/Runtime/Renderer/Private/VT/AdaptiveVirtualTexture.h
@@ -10,6 +10,7 @@
 class FAllocatedVirtualTexture;
 class FRHICommandListImmediate;
 class FVirtualTextureSystem;
+class FRHICommandListBase;
 
 /**
  * Concrete implementation of an adaptive virtual texture.
@@ -24,7 +25,7 @@ public:
 	FAdaptiveVirtualTexture(FAdaptiveVTDescription const& InAdaptiveDesc, FAllocatedVTDescription const& InAllocatedDesc);
 
 	/** Initialize the object. This creates the persistent low mips allocated VT. */
-	void Init(FVirtualTextureSystem* InSystem);
+	void Init(FRHICommandListBase& RHICmdList, FVirtualTextureSystem* InSystem);
 
 	/** Get a packed allocation key based on a virtual texture feedback request. The virtual texture system collects these opaque keys before queuing them for processing. */
 	uint32 GetPackedAllocationRequest(uint32 vAddress, uint32 vLevelPlusOne, uint32 Frame) const;
@@ -62,15 +63,16 @@ private:
 	uint32 GetAllocationIndex(FAllocatedVirtualTexture* InAllocatedVT) const;
 
 	/** Allocate a packed request. */
-	void Allocate(FVirtualTextureSystem* InSystem, uint32 InRequest, uint32 InFrame);
+	void Allocate(FRHICommandListBase& RHICmdList, FVirtualTextureSystem* InSystem, uint32 InRequest, uint32 InFrame);
 	/** Allocate or reallocate the allocated virtual texture at InGridIndex. */
-	void Allocate(FVirtualTextureSystem* InSystem, uint32 InGridIndex, uint32 InAllocationIndex, uint32 InNewLevel, uint32 InFrame);
+	void Allocate(FRHICommandListBase& RHICmdList, FVirtualTextureSystem* InSystem, uint32 InGridIndex, uint32 InAllocationIndex, uint32 InNewLevel, uint32 InFrame);
 	/** Free an allocated virtual texture. */
 	void Free(FVirtualTextureSystem* InSystem, uint32 InAllocationIndex, uint32 InFrame);
 	/** Free or reduce and reallocate the least recently used allocation. */
-	bool FreeLRU(FVirtualTextureSystem* InSystem, uint32 InFrame, uint32 InFrameUnusedThreshold);
+	bool FreeLRU(FRHICommandListBase& RHICmdList, FVirtualTextureSystem* InSystem, uint32 InFrame, uint32 InFrameUnusedThreshold);
 
 	static IAllocatedVirtualTexture* AllocateVirtualTexture(
+		FRHICommandListBase& RHICmdList,
 		FVirtualTextureSystem* InSystem,
 		FAllocatedVTDescription const& InAllocatedDesc,
 		FIntPoint InGridSize,
diff --git a/Engine/Source/Runtime/Renderer/Private/VT/AllocatedVirtualTexture.cpp b/Engine/Source/Runtime/Renderer/Private/VT/AllocatedVirtualTexture.cpp
--- a/Engine/Source/Runtime/Renderer/Private/VT/AllocatedVirtualTexture.cpp
+++ b/Engine/Source/Runtime/Renderer/Private/VT/AllocatedVirtualTexture.cpp
@@ -8,7 +8,8 @@
 #include "VT/VirtualTexturePhysicalSpace.h"
 #include "Misc/StringBuilder.h"
 
-FAllocatedVirtualTexture::FAllocatedVirtualTexture(FVirtualTextureSystem* InSystem,
+FAllocatedVirtualTexture::FAllocatedVirtualTexture(FRHICommandListBase& RHICmdList,
+	FVirtualTextureSystem* InSystem,
 	uint32 InFrame,
 	const FAllocatedVTDescription& InDesc,
 	FVirtualTextureProducer* const* InProducers,
@@ -93,7 +94,7 @@ FAllocatedVirtualTexture::FAllocatedVirtualTexture(FVirtualTextureSystem* InSyst
 	SpaceDesc.IndirectionTextureSize = InDesc.IndirectionTextureSize;
 	SpaceDesc.PageTableFormat = bSupport16BitPageTable ? EVTPageTableFormat::UInt16 : EVTPageTableFormat::UInt32;
 
-	Space = InSystem->AcquireSpace(SpaceDesc, InDesc.ForceSpaceID, this);
+	Space = InSystem->AcquireSpace(RHICmdList, SpaceDesc, InDesc.ForceSpaceID, this);
 	SpaceID = Space->GetID();
 	PageTableFormat = Space->GetPageTableFormat();
 }
diff --git a/Engine/Source/Runtime/Renderer/Private/VT/AllocatedVirtualTexture.h b/Engine/Source/Runtime/Renderer/Private/VT/AllocatedVirtualTexture.h
--- a/Engine/Source/Runtime/Renderer/Private/VT/AllocatedVirtualTexture.h
+++ b/Engine/Source/Runtime/Renderer/Private/VT/AllocatedVirtualTexture.h
@@ -18,6 +18,7 @@ class FAllocatedVirtualTexture final : public IAllocatedVirtualTexture
 {
 public:
 	FAllocatedVirtualTexture(
+		FRHICommandListBase& RHICmdList,
 		FVirtualTextureSystem* InSystem,
 		uint32 InFrame,
 		const FAllocatedVTDescription& InDesc,
diff --git a/Engine/Source/Runtime/Renderer/Private/VT/VirtualTexturePhysicalSpace.cpp b/Engine/Source/Runtime/Renderer/Private/VT/VirtualTexturePhysicalSpace.cpp
--- a/Engine/Source/Runtime/Renderer/Private/VT/VirtualTexturePhysicalSpace.cpp
+++ b/Engine/Source/Runtime/Renderer/Private/VT/VirtualTexturePhysicalSpace.cpp
@@ -120,10 +120,8 @@ EPixelFormat RemapVirtualTexturePhysicalSpaceFormat(EPixelFormat InFormat)
 	return InFormat;
 }
 
-void FVirtualTexturePhysicalSpace::InitRHI(FRHICommandListBase&)
+void FVirtualTexturePhysicalSpace::InitRHI(FRHICommandListBase& RHICmdList)
 {
-	FRHICommandListImmediate& RHICmdList = FRHICommandListExecutor::GetImmediateCommandList();
-
 	for (int32 Layer = 0; Layer < Description.NumLayers; ++Layer)
 	{
 		const EPixelFormat FormatSRV = RemapVirtualTexturePhysicalSpaceFormat(Description.Format[Layer]);
diff --git a/Engine/Source/Runtime/Renderer/Private/VT/VirtualTextureProducer.cpp b/Engine/Source/Runtime/Renderer/Private/VT/VirtualTextureProducer.cpp
--- a/Engine/Source/Runtime/Renderer/Private/VT/VirtualTextureProducer.cpp
+++ b/Engine/Source/Runtime/Renderer/Private/VT/VirtualTextureProducer.cpp
@@ -53,7 +53,7 @@ FVirtualTextureProducerCollection::FVirtualTextureProducerCollection() : NumPend
 	}
 }
 
-FVirtualTextureProducerHandle FVirtualTextureProducerCollection::RegisterProducer(FVirtualTextureSystem* System, const FVTProducerDescription& InDesc, IVirtualTexture* InProducer)
+FVirtualTextureProducerHandle FVirtualTextureProducerCollection::RegisterProducer(FRHICommandListBase& RHICmdList, FVirtualTextureSystem* System, const FVTProducerDescription& InDesc, IVirtualTexture* InProducer)
 {
 	const uint32 ProducerWidth = InDesc.BlockWidthInTiles * InDesc.WidthInBlocks * InDesc.TileSize;
 	const uint32 ProducerHeight = InDesc.BlockHeightInTiles * InDesc.HeightInBlocks * InDesc.TileSize;
@@ -90,7 +90,7 @@ FVirtualTextureProducerHandle FVirtualTextureProducerCollection::RegisterProduce
 			}
 		}
 		check(PhysicalSpaceDesc.NumLayers > 0);
-		Entry.Producer.PhysicalGroups[PhysicalGroupIndex].PhysicalSpace = System->AcquirePhysicalSpace(PhysicalSpaceDesc);
+		Entry.Producer.PhysicalGroups[PhysicalGroupIndex].PhysicalSpace = System->AcquirePhysicalSpace(RHICmdList, PhysicalSpaceDesc);
 	}
 
 	const FVirtualTextureProducerHandle Handle(Index, Entry.Magic);
diff --git a/Engine/Source/Runtime/Renderer/Private/VT/VirtualTextureProducer.h b/Engine/Source/Runtime/Renderer/Private/VT/VirtualTextureProducer.h
--- a/Engine/Source/Runtime/Renderer/Private/VT/VirtualTextureProducer.h
+++ b/Engine/Source/Runtime/Renderer/Private/VT/VirtualTextureProducer.h
@@ -56,7 +56,7 @@ class FVirtualTextureProducerCollection
 {
 public:
 	FVirtualTextureProducerCollection();
-	FVirtualTextureProducerHandle RegisterProducer(FVirtualTextureSystem* System, const FVTProducerDescription& InDesc, IVirtualTexture* InProducer);
+	FVirtualTextureProducerHandle RegisterProducer(FRHICommandListBase& RHICmdList, FVirtualTextureSystem* System, const FVTProducerDescription& InDesc, IVirtualTexture* InProducer);
 	void ReleaseProducer(FVirtualTextureSystem* System, const FVirtualTextureProducerHandle& Handle);
 	
 	void AddDestroyedCallback(const FVirtualTextureProducerHandle& Handle, FVTProducerDestroyedFunction* Function, void* Baton);
diff --git a/Engine/Source/Runtime/Renderer/Private/VT/VirtualTextureSpace.cpp b/Engine/Source/Runtime/Renderer/Private/VT/VirtualTextureSpace.cpp
--- a/Engine/Source/Runtime/Renderer/Private/VT/VirtualTextureSpace.cpp
+++ b/Engine/Source/Runtime/Renderer/Private/VT/VirtualTextureSpace.cpp
@@ -107,15 +107,15 @@ void FVirtualTextureSpace::FreeVirtualTexture(FAllocatedVirtualTexture* VirtualT
 	Allocator.Free(VirtualTexture);
 }
 
-void FVirtualTextureSpace::InitRHI(FRHICommandListBase&)
+void FVirtualTextureSpace::InitRHI(FRHICommandListBase& RHICmdList)
 {
 	for (uint32 TextureIndex = 0u; TextureIndex < GetNumPageTableTextures(); ++TextureIndex)
 	{
 		FTextureEntry& TextureEntry = PageTable[TextureIndex];
-		TextureEntry.TextureReferenceRHI = RHICreateTextureReference();
+		TextureEntry.TextureReferenceRHI = RHICmdList.CreateTextureReference();
 	}
-	PageTableIndirection.TextureReferenceRHI = RHICreateTextureReference();
-	RHIUpdateTextureReference(PageTableIndirection.TextureReferenceRHI, GBlackUintTexture->TextureRHI);
+	PageTableIndirection.TextureReferenceRHI = RHICmdList.CreateTextureReference();
+	RHICmdList.UpdateTextureReference(PageTableIndirection.TextureReferenceRHI, GBlackUintTexture->TextureRHI);
 }
 
 void FVirtualTextureSpace::ReleaseRHI()
diff --git a/Engine/Source/Runtime/Renderer/Private/VT/VirtualTextureSystem.cpp b/Engine/Source/Runtime/Renderer/Private/VT/VirtualTextureSystem.cpp
--- a/Engine/Source/Runtime/Renderer/Private/VT/VirtualTextureSystem.cpp
+++ b/Engine/Source/Runtime/Renderer/Private/VT/VirtualTextureSystem.cpp
@@ -597,7 +597,7 @@ void FVirtualTextureSystem::SaveAllocatorImagesFromConsole()
 }
 #endif // WITH_EDITOR
 
-IAllocatedVirtualTexture* FVirtualTextureSystem::AllocateVirtualTexture(const FAllocatedVTDescription& Desc)
+IAllocatedVirtualTexture* FVirtualTextureSystem::AllocateVirtualTexture(FRHICommandListBase& RHICmdList, const FAllocatedVTDescription& Desc)
 {
 	check(Desc.NumTextureLayers <= VIRTUALTEXTURE_SPACE_MAXLAYERS);
 
@@ -687,7 +687,7 @@ IAllocatedVirtualTexture* FVirtualTextureSystem::AllocateVirtualTexture(const FA
 		DestroyPendingVirtualTextures(true);
 	}
 
-	AllocatedVT = new FAllocatedVirtualTexture(this, Frame, Desc, ProducerForLayer, BlockWidthInTiles, BlockHeightInTiles, WidthInBlocks, HeightInBlocks, DepthInTiles);
+	AllocatedVT = new FAllocatedVirtualTexture(RHICmdList, this, Frame, Desc, ProducerForLayer, BlockWidthInTiles, BlockHeightInTiles, WidthInBlocks, HeightInBlocks, DepthInTiles);
 	AllocatedVT->NumRefs = 1;
 	if (bAnyLayerProducerWantsPersistentHighestMip)
 	{
@@ -773,12 +773,12 @@ void FVirtualTextureSystem::DestroyPendingVirtualTextures(bool bForceDestroyAll)
 	}
 }
 
-IAdaptiveVirtualTexture* FVirtualTextureSystem::AllocateAdaptiveVirtualTexture(const FAdaptiveVTDescription& AdaptiveVTDesc, const FAllocatedVTDescription& AllocatedVTDesc)
+IAdaptiveVirtualTexture* FVirtualTextureSystem::AllocateAdaptiveVirtualTexture(FRHICommandListBase& RHICmdList, const FAdaptiveVTDescription& AdaptiveVTDesc, const FAllocatedVTDescription& AllocatedVTDesc)
 {
 	check(!bUpdating);
 	UE::TScopeLock Lock(Mutex);
 	FAdaptiveVirtualTexture* AdaptiveVT = new FAdaptiveVirtualTexture(AdaptiveVTDesc, AllocatedVTDesc);
-	AdaptiveVT->Init(this);
+	AdaptiveVT->Init(RHICmdList, this);
 	check(AdaptiveVTs[AdaptiveVT->GetSpaceID()] == nullptr);
 	AdaptiveVTs[AdaptiveVT->GetSpaceID()] = AdaptiveVT;
 	return AdaptiveVT;
@@ -793,11 +793,11 @@ void FVirtualTextureSystem::DestroyAdaptiveVirtualTexture(IAdaptiveVirtualTextur
 	AdaptiveVT->Destroy(this);
 }
 
-FVirtualTextureProducerHandle FVirtualTextureSystem::RegisterProducer(const FVTProducerDescription& InDesc, IVirtualTexture* InProducer)
+FVirtualTextureProducerHandle FVirtualTextureSystem::RegisterProducer(FRHICommandListBase& RHICmdList, const FVTProducerDescription& InDesc, IVirtualTexture* InProducer)
 {
 	check(!bUpdating);
 	UE::TScopeLock Lock(Mutex);
-	return Producers.RegisterProducer(this, InDesc, InProducer);
+	return Producers.RegisterProducer(RHICmdList, this, InDesc, InProducer);
 }
 
 void FVirtualTextureSystem::ReleaseProducer(const FVirtualTextureProducerHandle& Handle)
@@ -826,7 +826,7 @@ FVirtualTextureProducer* FVirtualTextureSystem::FindProducer(const FVirtualTextu
 	return Producers.FindProducer(Handle);
 }
 
-FVirtualTextureSpace* FVirtualTextureSystem::AcquireSpace(const FVTSpaceDescription& InDesc, uint8 InForceSpaceID, FAllocatedVirtualTexture* AllocatedVT)
+FVirtualTextureSpace* FVirtualTextureSystem::AcquireSpace(FRHICommandListBase& RHICmdList, const FVTSpaceDescription& InDesc, uint8 InForceSpaceID, FAllocatedVirtualTexture* AllocatedVT)
 {
 	check(!bUpdating);
 	LLM_SCOPE(ELLMTag::VirtualTextureSystem);
@@ -875,7 +875,7 @@ FVirtualTextureSpace* FVirtualTextureSystem::AcquireSpace(const FVTSpaceDescript
 				Spaces[SpaceIndex].Reset(Space);
 				NumAllocatedSpaces++;
 				INC_MEMORY_STAT_BY(STAT_TotalPagetableMemory, Space->GetSizeInBytes());
-				BeginInitResource(Space);
+				Space->InitResource(RHICmdList);
 
 				const uint32 vAddress = Space->AllocateVirtualTexture(AllocatedVT);
 				AllocatedVT->AssignVirtualAddress(vAddress);
@@ -1005,7 +1005,7 @@ void GetPoolInitDescription_Cached(FVTPhysicalSpaceDescription const& InDesc, FV
 	}
 }
 
-FVirtualTexturePhysicalSpace* FVirtualTextureSystem::AcquirePhysicalSpace(const FVTPhysicalSpaceDescription& InDesc)
+FVirtualTexturePhysicalSpace* FVirtualTextureSystem::AcquirePhysicalSpace(FRHICommandListBase& RHICmdList, const FVTPhysicalSpaceDescription& InDesc)
 {
 	LLM_SCOPE(ELLMTag::VirtualTextureSystem);
 
@@ -1053,7 +1053,7 @@ FVirtualTexturePhysicalSpace* FVirtualTextureSystem::AcquirePhysicalSpace(const
 	PhysicalSpaces[ID] = PhysicalSpace;
 
 	INC_MEMORY_STAT_BY(STAT_TotalPhysicalMemory, PhysicalSpace->GetSizeInBytes());
-	BeginInitResource(PhysicalSpace);
+	PhysicalSpace->InitResource(RHICmdList);
 
 	return PhysicalSpace;
 }
diff --git a/Engine/Source/Runtime/Renderer/Private/VT/VirtualTextureSystem.h b/Engine/Source/Runtime/Renderer/Private/VT/VirtualTextureSystem.h
--- a/Engine/Source/Runtime/Renderer/Private/VT/VirtualTextureSystem.h
+++ b/Engine/Source/Runtime/Renderer/Private/VT/VirtualTextureSystem.h
@@ -115,15 +115,15 @@ public:
 
 	void ReleasePendingResources();
 
-	IAllocatedVirtualTexture* AllocateVirtualTexture(const FAllocatedVTDescription& Desc);
+	IAllocatedVirtualTexture* AllocateVirtualTexture(FRHICommandListBase& RHICmdList, const FAllocatedVTDescription& Desc);
 	void DestroyVirtualTexture(IAllocatedVirtualTexture* AllocatedVT);
 
-	FVirtualTextureProducerHandle RegisterProducer(const FVTProducerDescription& InDesc, IVirtualTexture* InProducer);
+	FVirtualTextureProducerHandle RegisterProducer(FRHICommandListBase& RHICmdList, const FVTProducerDescription& InDesc, IVirtualTexture* InProducer);
 	void ReleaseProducer(const FVirtualTextureProducerHandle& Handle);
 	void AddProducerDestroyedCallback(const FVirtualTextureProducerHandle& Handle, FVTProducerDestroyedFunction* Function, void* Baton);
 	uint32 RemoveAllProducerDestroyedCallbacks(const void* Baton);
 
-	IAdaptiveVirtualTexture* AllocateAdaptiveVirtualTexture(const FAdaptiveVTDescription& AdaptiveVTDesc, const FAllocatedVTDescription& AllocatedVTDesc);
+	IAdaptiveVirtualTexture* AllocateAdaptiveVirtualTexture(FRHICommandListBase& RHICmdList, const FAdaptiveVTDescription& AdaptiveVTDesc, const FAllocatedVTDescription& AllocatedVTDesc);
 	void DestroyAdaptiveVirtualTexture(IAdaptiveVirtualTexture* AdaptiveVT);
 
 	void RequestTiles(const FVector2D& InScreenSpaceSize, int32 InMipLevel = -1);
@@ -158,11 +158,11 @@ private:
 	FVirtualTextureSystem();
 	~FVirtualTextureSystem();
 
-	FVirtualTextureSpace* AcquireSpace(const FVTSpaceDescription& InDesc, uint8 InForceSpaceID, FAllocatedVirtualTexture* AllocatedVT);
+	FVirtualTextureSpace* AcquireSpace(FRHICommandListBase& RHICmdList, const FVTSpaceDescription& InDesc, uint8 InForceSpaceID, FAllocatedVirtualTexture* AllocatedVT);
 	void ReleaseSpace(FVirtualTextureSpace* Space);
 
 	FVirtualTextureProducer* FindProducer(const FVirtualTextureProducerHandle& Handle);
-	FVirtualTexturePhysicalSpace* AcquirePhysicalSpace(const FVTPhysicalSpaceDescription& InDesc);
+	FVirtualTexturePhysicalSpace* AcquirePhysicalSpace(FRHICommandListBase& RHICmdList, const FVTPhysicalSpaceDescription& InDesc);
 
 	FVirtualTextureSpace* GetSpace(uint8 ID) const { check(ID < MaxSpaces); return Spaces[ID].Get(); }
 	FAdaptiveVirtualTexture* GetAdaptiveVirtualTexture(uint8 ID) const { check(ID < MaxSpaces); return AdaptiveVTs[ID]; }
