Converted TSR and SSD relevant passes to utilise the Orthographic helper functions.

Moved the WorldDepthToPixelRadius calculation to view buffer.
SSD changes are mostly a refactor as they were already compensating for the ortho branching.
NOTE: This is missing the required change for the depth plane, so artifacts will be present with the current default Ortho camera far plane (UE_OLD_WORLD_MAX). If you reduce the far plane to 100,000 it should resolve it, 10,000 seems to be a sweet spot. This will be resolved in a future CL.

diff --git a/Engine/Shaders/Private/Common.ush b/Engine/Shaders/Private/Common.ush
--- a/Engine/Shaders/Private/Common.ush
+++ b/Engine/Shaders/Private/Common.ush
@@ -1247,6 +1247,11 @@ float2 GetScreenPositionForProjectionType(float2 ScreenPosition, float SceneDept
 	return select(IsOrthoProjection(), ScreenPosition, ScreenPosition * SceneDepth);
 }
 
+float GetDepthPixelRadiusForProjectionType(float SceneDepth)
+{
+	return View.WorldDepthToPixelWorldRadius * SceneDepth + View.PixelWorldRadiusOffset;
+}
+
 float2 SvPositionToBufferUV(float4 SvPosition)
 {
 	return SvPosition.xy * View.BufferSizeAndInvSize.zw;
diff --git a/Engine/Shaders/Private/RayTracing/RayTracingAmbientOcclusionRGS.usf b/Engine/Shaders/Private/RayTracing/RayTracingAmbientOcclusionRGS.usf
--- a/Engine/Shaders/Private/RayTracing/RayTracingAmbientOcclusionRGS.usf
+++ b/Engine/Shaders/Private/RayTracing/RayTracingAmbientOcclusionRGS.usf
@@ -26,10 +26,8 @@ RWTexture2D<float> RWAmbientOcclusionHitDistanceUAV;
 /** Returns the radius of a pixel in world space. */
 float ComputeWorldBluringRadiusCausedByPixelSize(float WorldDepth)
 {
-	float WorldDepthToPixelWorldRadius = GetTanHalfFieldOfView().x * View.ViewSizeAndInvSize.z;
-
 	// Should be multiplied 0.5* for the diameter to radius, and by 2.0 because GetTanHalfFieldOfView() cover only half of the pixels.
-	return WorldDepthToPixelWorldRadius * WorldDepth;
+	return GetDepthPixelRadiusForProjectionType(WorldDepth);
 }
 
 void GenerateCosineNormalRay(
diff --git a/Engine/Shaders/Private/SSRT/SSRTPrevFrameReduction.usf b/Engine/Shaders/Private/SSRT/SSRTPrevFrameReduction.usf
--- a/Engine/Shaders/Private/SSRT/SSRTPrevFrameReduction.usf
+++ b/Engine/Shaders/Private/SSRT/SSRTPrevFrameReduction.usf
@@ -104,9 +104,7 @@ void MainCS(
 				WorldDepth = ConvertFromDeviceZ(HZBDeviceZ);
 			}
 
-			float WorldDepthToPixelWorldRadius = GetTanHalfFieldOfView().x * View.ViewSizeAndInvSize.z * 100;
-
-			float WorldBluringRadius = WorldDepthToPixelWorldRadius * WorldDepth;
+			float WorldBluringRadius = GetDepthPixelRadiusForProjectionType(WorldDepth) * 100.0f;
 			float InvSquareWorldBluringRadius = rcp(WorldBluringRadius * WorldBluringRadius);
 
 			PrevColor = 0.0;
@@ -167,9 +165,7 @@ void MainCS(
 
 			float NoV = dot(View.TranslatedWorldCameraOrigin - normalize(RefWorldPosition), GetGBufferDataFromSceneTextures(SceneBufferUV).WorldNormal);
 				
-			float WorldDepthToPixelWorldRadius = GetTanHalfFieldOfView().x * View.ViewSizeAndInvSize.z * 100; // / max(NoV, 0.01);
-
-			float WorldBluringRadius = WorldDepthToPixelWorldRadius * WorldDepth;
+			float WorldBluringRadius = GetDepthPixelRadiusForProjectionType(WorldDepth) * 100.0f;
 			float InvSquareWorldBluringRadius = rcp(WorldBluringRadius * WorldBluringRadius);
 
 			//PrevColor = 0.0;
@@ -288,9 +284,7 @@ void MainCS(
 					}
 				}
 
-				float WorldDepthToPixelWorldRadius = GetTanHalfFieldOfView().x * View.ViewSizeAndInvSize.z * 100;
-
-				float WorldBluringRadius = WorldDepthToPixelWorldRadius * FurthestDepth;
+				float WorldBluringRadius = GetDepthPixelRadiusForProjectionType(FurthestDepth) * 100.0f;
 				float InvSquareWorldBluringRadius = rcp(WorldBluringRadius * WorldBluringRadius);
 
 				#endif
diff --git a/Engine/Shaders/Private/ScreenSpaceDenoise/SSDCommon.ush b/Engine/Shaders/Private/ScreenSpaceDenoise/SSDCommon.ush
--- a/Engine/Shaders/Private/ScreenSpaceDenoise/SSDCommon.ush
+++ b/Engine/Shaders/Private/ScreenSpaceDenoise/SSDCommon.ush
@@ -91,8 +91,6 @@ float2 BufferUVBilinearCorrection;
 
 float4 PrevSceneBufferUVToScreenPosition;
 
-float WorldDepthToPixelWorldRadius;
-
 Texture2D PrevDepthBuffer;
 Texture2D PrevGBufferA;
 Texture2D PrevGBufferB;
@@ -197,7 +195,7 @@ float ComputeAnisotropyInvFactor(FSSDSampleSceneInfos SceneMetadata)
 float ComputeWorldBluringRadiusCausedByPixelSize(FSSDSampleSceneInfos SceneMetadata)
 {
 	// Should be multiplied 0.5* for the diameter to radius, and by 2.0 because GetTanHalfFieldOfView() cover only half of the pixels.
-	return WorldDepthToPixelWorldRadius * GetWorldDepth(SceneMetadata);
+	return GetDepthPixelRadiusForProjectionType(GetWorldDepth(SceneMetadata));
 }
 
 /** Convert to bilateral world distance. */
diff --git a/Engine/Shaders/Private/ScreenSpaceDenoise/SSDCompressMetadata.usf b/Engine/Shaders/Private/ScreenSpaceDenoise/SSDCompressMetadata.usf
--- a/Engine/Shaders/Private/ScreenSpaceDenoise/SSDCompressMetadata.usf
+++ b/Engine/Shaders/Private/ScreenSpaceDenoise/SSDCompressMetadata.usf
@@ -37,7 +37,7 @@ FSSDSampleSceneInfos FetchCurrentSceneInfosFromGBuffer(float2 ScreenPosition, fl
 	
 	// Compute translated world position.
 	{
-		float2 ClipPosition = ScreenPosition * (View.ViewToClip[3][3] < 1.0f ? Infos.WorldDepth : 1.0f);
+		float2 ClipPosition = GetScreenPositionForProjectionType(ScreenPosition, Infos.WorldDepth);
 		Infos.TranslatedWorldPosition = mul(float4(ClipPosition, Infos.WorldDepth, 1), View.ScreenToTranslatedWorld).xyz;
 	}
 	
diff --git a/Engine/Shaders/Private/ScreenSpaceDenoise/SSDMetadata.ush b/Engine/Shaders/Private/ScreenSpaceDenoise/SSDMetadata.ush
--- a/Engine/Shaders/Private/ScreenSpaceDenoise/SSDMetadata.ush
+++ b/Engine/Shaders/Private/ScreenSpaceDenoise/SSDMetadata.ush
@@ -195,7 +195,7 @@ FSSDSampleSceneInfos UncompressSampleSceneInfo(
 
 	// Compute translated world position.
 	{
-		float4 ClipPosition = float4(ScreenPosition * (View.ViewToClip[3][3] < 1.0f ? Infos.WorldDepth : 1.0f), Infos.WorldDepth, 1);
+		float4 ClipPosition = float4(GetScreenPositionForProjectionType(ScreenPosition, Infos.WorldDepth), Infos.WorldDepth, 1);
 
 		if (bIsPrevFrame)
 		{
diff --git a/Engine/Shaders/Private/ScreenSpaceDenoise/SSDSignalFramework.ush b/Engine/Shaders/Private/ScreenSpaceDenoise/SSDSignalFramework.ush
--- a/Engine/Shaders/Private/ScreenSpaceDenoise/SSDSignalFramework.ush
+++ b/Engine/Shaders/Private/ScreenSpaceDenoise/SSDSignalFramework.ush
@@ -473,7 +473,7 @@ float ComputeBilateralWeight(
 	else if ((BilateralSettings & 0xF) == 2)
 	{
 		// TODO(Denoiser): save 2 VGPR by using SceneDepth instead of RefSceneMetadata.TranslatedWorldPosition (not sure it can work in math yet).
-		// float2 ClipPosition = ScreenPosition * (View.ViewToClip[3][3] < 1.0f ? WorldDepth : 1.0f);
+		// float2 ClipPosition = GetScreenPositionForProjectionType(ScreenPosition, WorldDepth);
 		// Infos.TranslatedWorldPosition = mul(float4(ClipPosition, WorldDepth, 1), ScreenToTranslatedWorld).xyz;
 
 		float DistFromRefPlane = abs(dot(RefNormal, NeighborToRefVector));
diff --git a/Engine/Shaders/Private/ScreenSpaceDenoise/SSDSpatialKernel.ush b/Engine/Shaders/Private/ScreenSpaceDenoise/SSDSpatialKernel.ush
--- a/Engine/Shaders/Private/ScreenSpaceDenoise/SSDSpatialKernel.ush
+++ b/Engine/Shaders/Private/ScreenSpaceDenoise/SSDSpatialKernel.ush
@@ -481,10 +481,10 @@ float3 ComputeVectorFromNeighborToRef(
 	{
 		// Recompute the the screen position of the reference, from the most minimal VGPR footprint.
 		float2 RefScreenPos = RefSceneMetadata.ScreenPosition;
-		float3 RefClipPosition = float3(RefScreenPos * (View.ViewToClip[3][3] < 1.0f ? RefWorldDepth : 1.0f), RefWorldDepth);
+		float3 RefClipPosition = float3(GetScreenPositionForProjectionType(RefScreenPos, RefWorldDepth), RefWorldDepth);
 
 		float2 NeighborScreenPos = NeighborSceneMetadata.ScreenPosition;
-		float3 NeighborClipPosition = float3(NeighborScreenPos * (View.ViewToClip[3][3] < 1.0f ? NeighborWorldDepth : 1.0f), NeighborWorldDepth);
+		float3 NeighborClipPosition = float3(GetScreenPositionForProjectionType(NeighborScreenPos, NeighborWorldDepth), NeighborWorldDepth);
 		
 		#if CONFIG_USE_VIEW_SPACE
 			float3 NeighborToRefVector = mul(float4(RefClipPosition - NeighborClipPosition, 0), GetScreenToViewDistanceMatrix()).xyz;
diff --git a/Engine/Shaders/Private/ScreenSpaceDenoise/SSDTemporalAccumulation.usf b/Engine/Shaders/Private/ScreenSpaceDenoise/SSDTemporalAccumulation.usf
--- a/Engine/Shaders/Private/ScreenSpaceDenoise/SSDTemporalAccumulation.usf
+++ b/Engine/Shaders/Private/ScreenSpaceDenoise/SSDTemporalAccumulation.usf
@@ -625,7 +625,7 @@ void TemporallyAccumulate(
 				PrevRefInfo.DeviceZ = HistoryScreenPosition.z;
 				PrevRefInfo.WorldDepth = ConvertFromDeviceZ(HistoryScreenPosition.z);
 				
-				float4 ClipPosition = float4(HistoryScreenPosition.xy * (View.ViewToClip[3][3] < 1.0f ? PrevRefInfo.WorldDepth : 1.0f), PrevRefInfo.WorldDepth, 1);
+				float4 ClipPosition = float4(GetScreenPositionForProjectionType(HistoryScreenPosition.xy, PrevRefInfo.WorldDepth), PrevRefInfo.WorldDepth, 1);
 
 				float3 PreViewTranslationOffset = LWCToFloat(LWCSubtract(PrimaryView.PreViewTranslation, PrimaryView.PrevPreViewTranslation));
 				PrevRefInfo.TranslatedWorldPosition = mul(ClipPosition, View.PrevScreenToTranslatedWorld).xyz + PreViewTranslationOffset;
diff --git a/Engine/Shaders/Private/TemporalSuperResolution/TSRDecimateHistory.usf b/Engine/Shaders/Private/TemporalSuperResolution/TSRDecimateHistory.usf
--- a/Engine/Shaders/Private/TemporalSuperResolution/TSRDecimateHistory.usf
+++ b/Engine/Shaders/Private/TemporalSuperResolution/TSRDecimateHistory.usf
@@ -310,11 +310,9 @@ void ProcessPrevAtomicTexturesSamples(
 
 		tsr_half DepthRejection;
 		{
-			const float WorldDepthToPixelWorldRadius = GetTanHalfFieldOfView().x * InputInfo_ViewportSizeInverse.x;
 			const float PixelDepthError = 3.0;
-			const float RadiusToDiamater = 2.0;
-			
-			float WorldDepthEpsilon = (WorldDepthToPixelWorldRadius * RadiusToDiamater * PixelDepthError) * HistoryClosestWorldDepth; 
+			const float RadiusToDiameter = 2.0;
+			float WorldDepthEpsilon = GetDepthPixelRadiusForProjectionType(HistoryClosestDeviceZ) * PixelDepthError * RadiusToDiameter;
 			float DeltaDepth = abs(HistoryClosestWorldDepth - WorldDepth);
 
 			DepthRejection = saturate(tsr_half(2.0) - tsr_half(DeltaDepth / WorldDepthEpsilon));
diff --git a/Engine/Shaders/Private/TemporalSuperResolution/TSRDilateVelocity.usf b/Engine/Shaders/Private/TemporalSuperResolution/TSRDilateVelocity.usf
--- a/Engine/Shaders/Private/TemporalSuperResolution/TSRDilateVelocity.usf
+++ b/Engine/Shaders/Private/TemporalSuperResolution/TSRDilateVelocity.usf
@@ -44,7 +44,6 @@ static const int2 kOffsets[4] = {
 float4x4 RotationalClipToPrevClip;
 float2 PrevOutputBufferUVMin;
 float2 PrevOutputBufferUVMax;
-float WorldDepthToDepthError;
 float VelocityExtrapolationMultiplier;
 float InvFlickeringMaxParralaxVelocity;
 uint bOutputIsMovingTexture;
@@ -178,7 +177,7 @@ void MainCS(
 		float DeviceZError;
 		{
 			float Depth = ConvertFromDeviceZ(CenterSubpixelDeviceZ);
-			float ErrorCorrectedDeviceZ = ConvertToDeviceZ(Depth + Depth * WorldDepthToDepthError);
+			float ErrorCorrectedDeviceZ = ConvertToDeviceZ(GetDepthPixelRadiusForProjectionType(Depth) * 2.0f + Depth);
 			DeviceZError = abs(ErrorCorrectedDeviceZ - CenterSubpixelDeviceZ);
 		}
 		#endif
@@ -335,15 +334,14 @@ void MainCS(
 			
 		tsr_half PixelParallaxVelocity = tsr_half(0.5) * length(tsr_half2(RotationalPrevScreen - PrevScreen) * tsr_half2(InputInfo_ViewportSize));
 			
-		float4 ClipPosition = float4(ScreenPos * (View.ViewToClip[3][3] < 1.0f ? Depth : 1.0f), Depth, 1);
-		float4 PrevClipPosition = float4(PrevScreenPos * (View.ViewToClip[3][3] < 1.0f ? PrevDepth : 1.0f), PrevDepth, 1);
+		float4 ClipPosition = float4(GetScreenPositionForProjectionType(ScreenPos,Depth), Depth, 1);
+		float4 PrevClipPosition = float4(GetScreenPositionForProjectionType(PrevScreenPos, PrevDepth), PrevDepth, 1);
 		
 		float3 PreViewTranslationOffset = LWCToFloat(LWCSubtract(PrimaryView.PreViewTranslation, PrimaryView.PrevPreViewTranslation));
 		float3 TranslatedWorldPosition = mul(ClipPosition, View.ScreenToTranslatedWorld).xyz;
 		float3 PrevTranslatedWorldPosition = mul(PrevClipPosition, View.PrevScreenToTranslatedWorld).xyz + PreViewTranslationOffset;
 		
-		float WorldDepthToPixelWorldRadius = 2.0 * GetTanHalfFieldOfView().x * View.ViewSizeAndInvSize.z;
-		float StaticWorldRadius = Depth * WorldDepthToPixelWorldRadius;
+		float StaticWorldRadius = GetDepthPixelRadiusForProjectionType(Depth) * 2.0f;
 
 		IsMovingMask = max(
 			tsr_half(saturate(length(TranslatedWorldPosition - PrevTranslatedWorldPosition) * rcp(StaticWorldRadius) - 1.0)),
diff --git a/Engine/Source/Runtime/Engine/Private/SceneView.cpp b/Engine/Source/Runtime/Engine/Private/SceneView.cpp
--- a/Engine/Source/Runtime/Engine/Private/SceneView.cpp
+++ b/Engine/Source/Runtime/Engine/Private/SceneView.cpp
@@ -333,7 +333,7 @@ static TAutoConsoleVariable<int32> CVarEnableTemporalUpsample(
 	ECVF_Default);
 
 float GOrthographicDepthThicknessScale = 0.01;
-static FAutoConsoleVariableRef CVarOrthographicDepthThicknessScal(
+static FAutoConsoleVariableRef CVarOrthographicDepthThicknessScale(
 	TEXT("r.OrthographicDepthThicknessScale"),
 	GOrthographicDepthThicknessScale,
 	TEXT("Orthographic scene depth scales proportionally lower than perspective, typically on a scale of 1/100")
@@ -2403,6 +2403,17 @@ void FSceneView::SetupViewRectUniformBufferParameters(FViewUniformShaderParamete
 
 		ViewUniformShaderParameters.TanAndInvTanHalfFOV = InViewMatrices.GetTanAndInvTanHalfFOV();
 		ViewUniformShaderParameters.PrevTanAndInvTanHalfFOV = InPrevViewMatrices.GetTanAndInvTanHalfFOV();
+
+		if (IsPerspectiveProjection())
+		{
+			ViewUniformShaderParameters.WorldDepthToPixelWorldRadius = InViewMatrices.GetTanAndInvTanHalfFOV().X / float(EffectiveViewRect.Width());
+			ViewUniformShaderParameters.PixelWorldRadiusOffset = 0;
+		}
+		else
+		{
+			ViewUniformShaderParameters.WorldDepthToPixelWorldRadius = 0;
+			ViewUniformShaderParameters.PixelWorldRadiusOffset = InViewMatrices.GetTanAndInvTanHalfFOV().X / float(EffectiveViewRect.Width());
+		}
 	}
 	
 	float FovFixX = ViewUniformShaderParameters.TanAndInvTanHalfFOV.X;
diff --git a/Engine/Source/Runtime/Engine/Private/ShowFlags.cpp b/Engine/Source/Runtime/Engine/Private/ShowFlags.cpp
--- a/Engine/Source/Runtime/Engine/Private/ShowFlags.cpp
+++ b/Engine/Source/Runtime/Engine/Private/ShowFlags.cpp
@@ -694,7 +694,6 @@ void EngineShowFlagOrthographicOverride(bool bIsPerspective, FEngineShowFlags& E
 	// Disable post processing that doesn't work in ortho viewports.
 	if (!bIsPerspective)
 	{
-		EngineShowFlags.SetTemporalAA(false);
 		EngineShowFlags.SetMotionBlur(false);
 	}
 }
diff --git a/Engine/Source/Runtime/Engine/Public/SceneView.h b/Engine/Source/Runtime/Engine/Public/SceneView.h
--- a/Engine/Source/Runtime/Engine/Public/SceneView.h
+++ b/Engine/Source/Runtime/Engine/Public/SceneView.h
@@ -518,8 +518,16 @@ public:
 
 		TemporalAAProjectionJitter = InTemporalAAProjectionJitter;
 
-		ProjectionMatrix.M[2][0] += TemporalAAProjectionJitter.X;
-		ProjectionMatrix.M[2][1] += TemporalAAProjectionJitter.Y;
+		if (IsPerspectiveProjection())
+		{
+			ProjectionMatrix.M[2][0] += TemporalAAProjectionJitter.X;
+			ProjectionMatrix.M[2][1] += TemporalAAProjectionJitter.Y;
+		}
+		else
+		{
+			ProjectionMatrix.M[3][0] += TemporalAAProjectionJitter.X;
+			ProjectionMatrix.M[3][1] += TemporalAAProjectionJitter.Y;
+		}
 		InvProjectionMatrix = InvertProjectionMatrix(ProjectionMatrix);
 
 		RecomputeDerivedMatrices();
@@ -527,8 +535,16 @@ public:
 
 	void HackRemoveTemporalAAProjectionJitter()
 	{
-		ProjectionMatrix.M[2][0] -= TemporalAAProjectionJitter.X;
-		ProjectionMatrix.M[2][1] -= TemporalAAProjectionJitter.Y;
+		if (IsPerspectiveProjection())
+		{
+			ProjectionMatrix.M[2][0] -= TemporalAAProjectionJitter.X;
+			ProjectionMatrix.M[2][1] -= TemporalAAProjectionJitter.Y;
+		}
+		else
+		{
+			ProjectionMatrix.M[3][0] -= TemporalAAProjectionJitter.X;
+			ProjectionMatrix.M[3][1] -= TemporalAAProjectionJitter.Y;
+		}
 		InvProjectionMatrix = InvertProjectionMatrix(ProjectionMatrix);
 
 		TemporalAAProjectionJitter = FVector2D::ZeroVector;
@@ -539,8 +555,16 @@ public:
 	{
 		FMatrix ProjNoAAMatrix = ProjectionMatrix;
 
-		ProjNoAAMatrix.M[2][0] -= TemporalAAProjectionJitter.X;
-		ProjNoAAMatrix.M[2][1] -= TemporalAAProjectionJitter.Y;
+		if (IsPerspectiveProjection())
+		{
+			ProjNoAAMatrix.M[2][0] -= TemporalAAProjectionJitter.X;
+			ProjNoAAMatrix.M[2][1] -= TemporalAAProjectionJitter.Y;
+		}
+		else
+		{
+			ProjNoAAMatrix.M[3][0] -= TemporalAAProjectionJitter.X;
+			ProjNoAAMatrix.M[3][1] -= TemporalAAProjectionJitter.Y;
+		}
 
 		return ProjNoAAMatrix;
 	}
@@ -901,6 +925,8 @@ enum ETranslucencyVolumeCascade
 	VIEW_UNIFORM_BUFFER_MEMBER(float, SubSurfaceColorAsTransmittanceAtDistanceInMeters) \
 	VIEW_UNIFORM_BUFFER_MEMBER_PER_VIEW(FVector4f, TanAndInvTanHalfFOV) \
 	VIEW_UNIFORM_BUFFER_MEMBER_PER_VIEW(FVector4f, PrevTanAndInvTanHalfFOV) \
+	VIEW_UNIFORM_BUFFER_MEMBER_PER_VIEW(float, WorldDepthToPixelWorldRadius) \
+	VIEW_UNIFORM_BUFFER_MEMBER_PER_VIEW(float, PixelWorldRadiusOffset) \
 	VIEW_UNIFORM_BUFFER_MEMBER_PER_VIEW(FVector4f, GlintLUTParameters0) \
 	VIEW_UNIFORM_BUFFER_MEMBER_PER_VIEW(FVector4f, GlintLUTParameters1) \
 
diff --git a/Engine/Source/Runtime/Renderer/Private/PostProcess/TemporalSuperResolution.cpp b/Engine/Source/Runtime/Renderer/Private/PostProcess/TemporalSuperResolution.cpp
--- a/Engine/Source/Runtime/Renderer/Private/PostProcess/TemporalSuperResolution.cpp
+++ b/Engine/Source/Runtime/Renderer/Private/PostProcess/TemporalSuperResolution.cpp
@@ -502,7 +502,6 @@ class FTSRDilateVelocityCS : public FTSRShader
 		SHADER_PARAMETER(FMatrix44f, RotationalClipToPrevClip)
 		SHADER_PARAMETER(FVector2f, PrevOutputBufferUVMin)
 		SHADER_PARAMETER(FVector2f, PrevOutputBufferUVMax)
-		SHADER_PARAMETER(float, WorldDepthToDepthError)
 		SHADER_PARAMETER(float, VelocityExtrapolationMultiplier)
 		SHADER_PARAMETER(float, InvFlickeringMaxParralaxVelocity)
 		SHADER_PARAMETER(int32, bOutputIsMovingTexture)
@@ -555,7 +554,6 @@ class FTSRDecimateHistoryCS : public FTSRShader
 	BEGIN_SHADER_PARAMETER_STRUCT(FParameters, )
 		SHADER_PARAMETER_STRUCT_INCLUDE(FTSRCommonParameters, CommonParameters)
 		SHADER_PARAMETER(FMatrix44f, RotationalClipToPrevClip)
-		SHADER_PARAMETER(float, WorldDepthToPixelWorldRadius)
 
 		SHADER_PARAMETER_RDG_TEXTURE(Texture2D, InputSceneColorTexture)
 		SHADER_PARAMETER_RDG_TEXTURE(Texture2D, DilatedVelocityTexture)
@@ -1733,14 +1731,6 @@ FDefaultTemporalUpscaler::FOutputs AddTemporalSuperResolutionPasses(
 		PassParameters->RotationalClipToPrevClip = RotationalClipToPrevClip;
 		PassParameters->PrevOutputBufferUVMin = CommonParameters.InputInfo.UVViewportBilinearMin - CommonParameters.InputInfo.ExtentInverse;
 		PassParameters->PrevOutputBufferUVMax = CommonParameters.InputInfo.UVViewportBilinearMax + CommonParameters.InputInfo.ExtentInverse;
-		{
-			float TanHalfFieldOfView = View.ViewMatrices.GetInvProjectionMatrix().M[0][0];
-
-			// Should be multiplied 0.5* for the diameter to radius, and by 2.0 because GetTanHalfFieldOfView() cover only half of the pixels.
-			float WorldDepthToPixelWorldRadius = TanHalfFieldOfView / float(View.ViewRect.Width());
-
-			PassParameters->WorldDepthToDepthError = WorldDepthToPixelWorldRadius * 2.0f;
-		}
 		PassParameters->VelocityExtrapolationMultiplier = 0.0; //FMath::Clamp(CVarTSRVelocityExtrapolation.GetValueOnRenderThread(), 0.0f, 1.0f);
 		{
 			float FlickeringMaxParralaxVelocity = RefreshRateToFrameRateCap * CVarTSRFlickeringMaxParralaxVelocity.GetValueOnRenderThread() * float(View.ViewRect.Width()) / 1920.0f;
@@ -1866,12 +1856,6 @@ FDefaultTemporalUpscaler::FOutputs AddTemporalSuperResolutionPasses(
 		FTSRDecimateHistoryCS::FParameters* PassParameters = GraphBuilder.AllocParameters<FTSRDecimateHistoryCS::FParameters>();
 		PassParameters->CommonParameters = CommonParameters;
 		PassParameters->RotationalClipToPrevClip = RotationalClipToPrevClip;
-		{
-			float TanHalfFieldOfView = View.ViewMatrices.GetInvProjectionMatrix().M[0][0];
-
-			// Should be multiplied 0.5* for the diameter to radius, and by 2.0 because GetTanHalfFieldOfView() cover only half of the pixels.
-			PassParameters->WorldDepthToPixelWorldRadius = TanHalfFieldOfView / float(View.ViewRect.Width());
-		}
 
 		PassParameters->InputSceneColorTexture = PassInputs.SceneColor.Texture;
 		PassParameters->DilatedVelocityTexture = DilatedVelocityTexture;
diff --git a/Engine/Source/Runtime/Renderer/Private/ScreenSpaceDenoise.cpp b/Engine/Source/Runtime/Renderer/Private/ScreenSpaceDenoise.cpp
--- a/Engine/Source/Runtime/Renderer/Private/ScreenSpaceDenoise.cpp
+++ b/Engine/Source/Runtime/Renderer/Private/ScreenSpaceDenoise.cpp
@@ -1312,11 +1312,6 @@ void Denoiser::SetupCommonShaderParameters(
 		float(DenoiserViewport.Min.Y + 0.5f) / float(DenoiserBufferExtent.Y),
 		float(DenoiserViewport.Max.X - 0.5f) / float(DenoiserBufferExtent.X),
 		float(DenoiserViewport.Max.Y - 0.5f) / float(DenoiserBufferExtent.Y));
-
-	float TanHalfFieldOfView = View.ViewMatrices.GetInvProjectionMatrix().M[0][0];
-
-	// Should be multiplied 0.5* for the diameter to radius, and by 2.0 because GetTanHalfFieldOfView() cover only half of the pixels.
-	OutPublicCommonParameters->WorldDepthToPixelWorldRadius = TanHalfFieldOfView / float(View.ViewRect.Width());
 }
 
 
diff --git a/Engine/Source/Runtime/Renderer/Private/ScreenSpaceDenoise.h b/Engine/Source/Runtime/Renderer/Private/ScreenSpaceDenoise.h
--- a/Engine/Source/Runtime/Renderer/Private/ScreenSpaceDenoise.h
+++ b/Engine/Source/Runtime/Renderer/Private/ScreenSpaceDenoise.h
@@ -26,7 +26,6 @@ BEGIN_SHADER_PARAMETER_STRUCT(FCommonShaderParameters, )
 	SHADER_PARAMETER(FVector4f, DenoiserBufferSizeAndInvSize)
 	SHADER_PARAMETER(FVector4f, DenoiserBufferBilinearUVMinMax)
 	SHADER_PARAMETER(FVector4f, SceneBufferUVToScreenPosition) // TODO: move to view uniform buffer
-	SHADER_PARAMETER(float, WorldDepthToPixelWorldRadius)
 END_SHADER_PARAMETER_STRUCT()
 
 void SetupCommonShaderParameters(
