Temporarily disables an ensure in FTSRDilateVelocityCS::ShouldCompilePermutation()

diff --git a/Engine/Source/Runtime/Renderer/Private/PostProcess/TemporalSuperResolution.cpp b/Engine/Source/Runtime/Renderer/Private/PostProcess/TemporalSuperResolution.cpp
--- a/Engine/Source/Runtime/Renderer/Private/PostProcess/TemporalSuperResolution.cpp
+++ b/Engine/Source/Runtime/Renderer/Private/PostProcess/TemporalSuperResolution.cpp
@@ -542,7 +542,6 @@ class FTSRDilateVelocityCS : public FTSRShader
 		{
 			return false;
 		}
-		ensureMsgf(VelocityEncodeHasPixelAnimation(Parameters.Platform), TEXT("Platform did not return true from VelocityEncodeHasPixelAnimation().  TSR requires encoding UMaterial::bHasPixelAnimation in the velocity buffer."));
 		return true;
 	}
 }; // class FTSRDilateVelocityCS
