Iinstance culling occlusion query optimizations

* Only ever process instances that pass the CPU-side culling
* Use R8_UINT format for the instance visibility mask

diff --git a/Engine/Shaders/Private/InstanceCulling/InstanceCullingOcclusionQuery.usf b/Engine/Shaders/Private/InstanceCulling/InstanceCullingOcclusionQuery.usf
--- a/Engine/Shaders/Private/InstanceCulling/InstanceCullingOcclusionQuery.usf
+++ b/Engine/Shaders/Private/InstanceCulling/InstanceCullingOcclusionQuery.usf
@@ -46,7 +46,7 @@ EInstanceVisibilityStatus GetInstanceDataAndVisibility(uint InstanceId, out FIns
 	const FInstanceSceneData InstanceData = GetInstanceSceneData(InstanceId, InstanceSceneDataSOAStride);
 
 	OutInstanceData = InstanceData;
-	
+
 	if (!InstanceData.ValidInstance)
 	{
 		return IVS_Hidden;
@@ -126,6 +126,9 @@ groupshared uint GroupOutputBaseIndex;
 RWBuffer<uint> OutIndirectArgsBuffer;
 RWBuffer<uint> OutInstanceIdBuffer;
 
+// Remap linear visible index to a visible GPUScene instance index
+Buffer<uint> InstanceIdBuffer;
+
 [numthreads(NUM_THREADS_PER_GROUP, 1, 1)]
 void MainCS(uint3 GroupId : SV_GroupID, uint GroupThreadIndex : SV_GroupIndex, uint DispatchThreadId : SV_DispatchThreadID)
 {
@@ -138,7 +141,8 @@ void MainCS(uint3 GroupId : SV_GroupID, uint GroupThreadIndex : SV_GroupIndex, u
 
 	const bool bValidInvocation = DispatchThreadId < NumInstances;
 
-	uint InstanceId = min(DispatchThreadId, NumInstances);
+	uint ClampedThreadId = min(DispatchThreadId, NumInstances-1);
+	uint InstanceId = InstanceIdBuffer[ClampedThreadId];
 
 	GroupMemoryBarrierWithGroupSync();
 
@@ -196,9 +200,6 @@ void MainCS(uint3 GroupId : SV_GroupID, uint GroupThreadIndex : SV_GroupIndex, u
 	}
 }
 
-// Remap linear visible index to a visible GPUScene instance index
-Buffer<uint> InstanceIdBuffer;
-
 void MainVS(
 	float3 InPosition : ATTRIBUTE0,
 	uint InstanceId : SV_InstanceID,
diff --git a/Engine/Source/Runtime/Renderer/Private/InstanceCulling/InstanceCullingContext.cpp b/Engine/Source/Runtime/Renderer/Private/InstanceCulling/InstanceCullingContext.cpp
--- a/Engine/Source/Runtime/Renderer/Private/InstanceCulling/InstanceCullingContext.cpp
+++ b/Engine/Source/Runtime/Renderer/Private/InstanceCulling/InstanceCullingContext.cpp
@@ -863,7 +863,9 @@ void FInstanceCullingContext::BuildRenderingCommandsInternal(
 
 	if (InstanceCullingManager && InstanceCullingManager->InstanceOcclusionQueryBuffer)
 	{
-		PassParametersTmp.InstanceOcclusionQueryBuffer = GraphBuilder.CreateSRV(InstanceCullingManager->InstanceOcclusionQueryBuffer, PF_R32_UINT);
+		PassParametersTmp.InstanceOcclusionQueryBuffer = GraphBuilder.CreateSRV(
+			InstanceCullingManager->InstanceOcclusionQueryBuffer, 
+			InstanceCullingManager->InstanceOcclusionQueryBufferFormat);
 	}
 	else
 	{
@@ -1186,7 +1188,9 @@ FInstanceCullingDeferredContext *FInstanceCullingContext::CreateDeferredContext(
 
 	if (InstanceCullingManager.InstanceOcclusionQueryBuffer)
 	{
-		PassParametersTmp.InstanceOcclusionQueryBuffer = GraphBuilder.CreateSRV(InstanceCullingManager.InstanceOcclusionQueryBuffer, PF_R32_UINT);
+		PassParametersTmp.InstanceOcclusionQueryBuffer = GraphBuilder.CreateSRV(
+			InstanceCullingManager.InstanceOcclusionQueryBuffer,
+			InstanceCullingManager.InstanceOcclusionQueryBufferFormat);
 	}
 	else
 	{
diff --git a/Engine/Source/Runtime/Renderer/Private/InstanceCulling/InstanceCullingManager.h b/Engine/Source/Runtime/Renderer/Private/InstanceCulling/InstanceCullingManager.h
--- a/Engine/Source/Runtime/Renderer/Private/InstanceCulling/InstanceCullingManager.h
+++ b/Engine/Source/Runtime/Renderer/Private/InstanceCulling/InstanceCullingManager.h
@@ -98,6 +98,7 @@ public:
 
 	// Reference to a buffer owned by FInstanceCullingOcclusionQueryRenderer
 	FRDGBufferRef InstanceOcclusionQueryBuffer = {};
+	EPixelFormat InstanceOcclusionQueryBufferFormat = PF_Unknown;
 
 private:
 
diff --git a/Engine/Source/Runtime/Renderer/Private/InstanceCulling/InstanceCullingOcclusionQuery.cpp b/Engine/Source/Runtime/Renderer/Private/InstanceCulling/InstanceCullingOcclusionQuery.cpp
--- a/Engine/Source/Runtime/Renderer/Private/InstanceCulling/InstanceCullingOcclusionQuery.cpp
+++ b/Engine/Source/Runtime/Renderer/Private/InstanceCulling/InstanceCullingOcclusionQuery.cpp
@@ -27,6 +27,8 @@ static TAutoConsoleVariable<int32> CVarInstanceCullingOcclusionQueries(
 	TEXT("EXPERIMENTAL: Use per-instance software occlusion queries to perform less conservative visibility test than what's possible with HZB alone"),
 	ECVF_RenderThreadSafe);
 
+struct FInstanceCullingOcclusionQueryDeferredContext;
+
 namespace
 {
 
@@ -60,6 +62,27 @@ static FBufferRHIRef CreateBufferWithData(FRHICommandListBase& RHICmdList, EBuff
 	return RHICmdList.CreateBuffer(DataView.SizeInBytes, UsageFlags, Data.GetTypeSize(), ResourceState, CreateInfo);
 }
 
+
+static EPixelFormat GetPreferredVisibilityMaskFormat()
+{
+	EPixelFormat PossibleFormats[] =
+	{
+		PF_R8_UINT,  // may be available if typed UAV load/store is supported on current hardware
+		PF_R32_UINT, // guaranteed to be supported
+	};
+
+	for (EPixelFormat Format : PossibleFormats)
+	{
+		EPixelFormatCapabilities Capabilities = GPixelFormats[Format].Capabilities;
+		if (EnumHasAllFlags(Capabilities, EPixelFormatCapabilities::TypedUAVLoad | EPixelFormatCapabilities::TypedUAVStore))
+		{
+			return Format;
+		}
+	}
+
+	return PF_Unknown;
+}
+
 }
 
 /*
@@ -94,6 +117,7 @@ public:
 		SHADER_PARAMETER_RDG_BUFFER_SRV(StructuredBuffer<float4>, GPUSceneInstanceSceneData)
 		SHADER_PARAMETER_RDG_BUFFER_SRV(StructuredBuffer<float4>, GPUSceneInstancePayloadData)
 		SHADER_PARAMETER_RDG_BUFFER_SRV(StructuredBuffer<float4>, GPUScenePrimitiveSceneData)
+		SHADER_PARAMETER_RDG_BUFFER_SRV(Buffer<uint32>, InstanceIdBuffer)
 		SHADER_PARAMETER_STRUCT_REF(FViewUniformShaderParameters, View)
 		SHADER_PARAMETER_RDG_TEXTURE(Texture2D, HZBTexture)
 		SHADER_PARAMETER_SAMPLER(SamplerState, HZBSampler)
@@ -235,8 +259,7 @@ TGlobalResource<FInstanceCullingOcclusionQueryBox> GInstanceCullingOcclusionQuer
 static void RenderInstanceOcclusionCulling(
 	FRHICommandList& RHICmdList,
 	FViewInfo& View,
-	FOcclusionInstanceCullingParameters* PassParameters,
-	int32 NumInstances)
+	FOcclusionInstanceCullingParameters* PassParameters)
 {
 	TShaderMapRef<FInstanceCullingOcclusionQueryVS> VertexShader(View.ShaderMap);
 	TShaderMapRef<FInstanceCullingOcclusionQueryPS> PixelShader(View.ShaderMap);
@@ -271,6 +294,158 @@ static void RenderInstanceOcclusionCulling(
 	RHICmdList.DrawIndexedPrimitiveIndirect(GInstanceCullingOcclusionQueryBox.IndexBuffer, IndirectArgsBuffer->GetRHI(), 0);
 }
 
+
+/*
+* Structure to compute data that's not available on the rendering thread during RDG setup.
+* In particular, we want to wait for visible mesh draw commands as late as possible.
+*/
+struct FInstanceCullingOcclusionQueryDeferredContext
+{
+	FInstanceCullingOcclusionQueryDeferredContext(const FViewInfo* InView, int32 InNumGPUSceneInstances, EMeshPass::Type InMeshPass)
+		: View(InView)
+		, NumGPUSceneInstances(InNumGPUSceneInstances)
+		, MeshPass(InMeshPass)
+	{
+	}
+
+	static bool IsRelevantCommand(const FVisibleMeshDrawCommand& VisibleCommand)
+	{
+		// There may be multiple visible mesh draw commands that refer to the same instance when GPU-based LOD selection is used.
+		// This filter is designed to remove the duplicates, keeping only the "authoritative" instance.
+		// TODO: a less implicit mechanism would be welcome here, such as a dedicated flag.
+		EMeshDrawCommandCullingPayloadFlags Flags = VisibleCommand.CullingPayloadFlags;
+		bool bCompatibleFlags = Flags == EMeshDrawCommandCullingPayloadFlags::Default
+							 || Flags == EMeshDrawCommandCullingPayloadFlags::MinScreenSizeCull;
+
+		// NumPrimitives is 0 if mesh draw command uses IndirectArgs
+		// This path is currently not implemented/supported by oclcusion query culling.
+		return bCompatibleFlags
+			&& VisibleCommand.MeshDrawCommand->NumPrimitives != 0;
+	};
+
+	void Execute()
+	{
+		if (bExecuted)
+		{
+			return;
+		}
+
+		bExecuted = true;
+
+		const FParallelMeshDrawCommandPass& MeshDrawCommandPass = View->ParallelMeshDrawCommandPasses[MeshPass];
+
+		// Execute() is expected to run late enough to not stall here.
+		// If it does happen, then we may have to move the render pass to later point in the frame.
+		MeshDrawCommandPass.WaitForSetupTask(); 
+
+		const FMeshCommandOneFrameArray& VisibleMeshDrawCommands = MeshDrawCommandPass.GetMeshDrawCommands();
+
+		uint32 NumPotentiallyVisibleInstances = 0;
+
+		for (const FVisibleMeshDrawCommand& VisibleCommand : VisibleMeshDrawCommands)
+		{
+			if (!IsRelevantCommand(VisibleCommand))
+			{
+				continue;
+			}
+			NumPotentiallyVisibleInstances += VisibleCommand.MeshDrawCommand->NumInstances;
+		}
+
+		NumInstances = NumPotentiallyVisibleInstances;
+
+		NumThreadGroups = FComputeShaderUtils::GetGroupCount(NumInstances, FInstanceCullingOcclusionQueryCS::NumThreadsPerGroup);
+
+		const int32 MaxSupportedInstances = GRHIGlobals.MaxDispatchThreadGroupsPerDimension.X * FInstanceCullingOcclusionQueryCS::NumThreadsPerGroup;
+		if (!ensureMsgf(NumThreadGroups.X * FInstanceCullingOcclusionQueryCS::NumThreadsPerGroup <= MaxSupportedInstances,
+			TEXT("Number of instances (%d) is greater than currently supported by FInstanceCullingOcclusionQueryRenderer (%d). ")
+			TEXT("Per-instance occlusion queries will be disabled. ")
+			TEXT("Increase FInstanceCullingOcclusionQueryCS::NumThreadsPerGroup or implement wrapped group count support."),
+			NumInstances, MaxSupportedInstances))
+		{
+			return;
+		}
+
+		// Align buffer sizes to ensure each thread in the thread group has a valid slot to write without introducing bounds checks
+		AlignedNumInstances = NumThreadGroups.X * FInstanceCullingOcclusionQueryCS::NumThreadsPerGroup;
+
+		if (AlignedNumInstances == 0)
+		{
+			return;
+		}
+
+		VisibleInstanceIds.Reserve(AlignedNumInstances);
+
+		for (const FVisibleMeshDrawCommand& VisibleCommand : VisibleMeshDrawCommands)
+		{
+			if (!IsRelevantCommand(VisibleCommand))
+			{
+				continue;
+			}
+
+			uint32 InstanceBaseIndex = VisibleCommand.PrimitiveIdInfo.InstanceSceneDataOffset;
+			uint32 CommandNumInstances = VisibleCommand.MeshDrawCommand->NumInstances;
+
+			check(InstanceBaseIndex + CommandNumInstances <= uint32(NumGPUSceneInstances));
+
+			for (uint32 i = 0; i < CommandNumInstances; ++i)
+			{
+				VisibleInstanceIds.Add(InstanceBaseIndex + i);
+			}
+		}
+
+		for (int32 i = NumInstances; i < AlignedNumInstances; ++i)
+		{
+			VisibleInstanceIds.Add(0);
+		}
+
+		check(VisibleInstanceIds.Num() == AlignedNumInstances);
+
+		bValid = true;
+	}
+
+	FRDGBufferNumElementsCallback DeferredAlignedNumInstances()
+	{
+		return [Context = this]() -> uint32
+			{
+				Context->Execute();
+				return Context->AlignedNumInstances;
+			};
+	}
+
+	FRDGBufferInitialDataCallback DeferredInstanceIdData()
+	{
+		return [Context = this]() -> const void*
+			{
+				Context->Execute();
+				return Context->VisibleInstanceIds.GetData();
+			};
+	}
+
+	FRDGBufferInitialDataSizeCallback DeferredInstanceIdDataSize()
+	{
+		return [Context = this]() -> uint64
+			{
+				Context->Execute();
+				return Context->VisibleInstanceIds.Num() * Context->VisibleInstanceIds.GetTypeSize();
+			};
+	}
+
+	// Execute function may be called multiple times, but we only want to run computations once
+	bool bExecuted = false;
+
+	// If this is false, then some late validation have failed and rendering should be skipped
+	bool bValid = false;
+
+	const FViewInfo* View = nullptr;
+	int32 NumGPUSceneInstances = 0;
+	EMeshPass::Type MeshPass = EMeshPass::Num;
+	int32 NumInstances = 0;
+	int32 AlignedNumInstances = 0;
+	FIntVector NumThreadGroups = FIntVector::ZeroValue;
+
+	TArray<uint32> VisibleInstanceIds;
+};
+
 uint32 FInstanceCullingOcclusionQueryRenderer::Render(
 	FRDGBuilder& GraphBuilder,
 	FGPUScene& GPUScene,
@@ -278,7 +453,7 @@ uint32 FInstanceCullingOcclusionQueryRenderer::Render(
 {
 	if (!IsCompatibleWithView(View))
 	{
-		return 0 ;
+		return 0;
 	}
 
 	const uint32 ViewMask = RegisterView(View);
@@ -289,6 +464,10 @@ uint32 FInstanceCullingOcclusionQueryRenderer::Render(
 		return 0;
 	}
 
+	const int32 NumGPUSceneInstances = GPUScene.GetNumInstances();
+
+	FInstanceCullingOcclusionQueryDeferredContext* DeferredContext = GraphBuilder.AllocObject<FInstanceCullingOcclusionQueryDeferredContext>(&View, NumGPUSceneInstances, EMeshPass::BasePass);
+
 	FRDGTextureRef DepthTexture = View.GetSceneTextures().Depth.Target;
 	FRDGTextureRef HZBTexture = View.HZB;
 
@@ -299,46 +478,59 @@ uint32 FInstanceCullingOcclusionQueryRenderer::Render(
 
 	const FGPUSceneResourceParameters GPUSceneParameters = GPUScene.GetShaderParameters();
 
-	const int32 NumInstances = GPUScene.GetNumInstances();
 	const FIntPoint ViewRectSize = View.ViewRect.Size();
 
-	const FIntVector NumThreadGroups = FComputeShaderUtils::GetGroupCount(NumInstances, FInstanceCullingOcclusionQueryCS::NumThreadsPerGroup);
-	const int32 MaxSupportedInstances = GRHIGlobals.MaxDispatchThreadGroupsPerDimension.X * FInstanceCullingOcclusionQueryCS::NumThreadsPerGroup;
-	if (!ensureMsgf(NumThreadGroups.X * FInstanceCullingOcclusionQueryCS::NumThreadsPerGroup <= MaxSupportedInstances,
-		TEXT("Number of instances (%d) is greater than currently supported by FInstanceCullingOcclusionQueryRenderer (%d). ")
-		TEXT("Per-instance occlusion queries will be disabled. ")
-		TEXT("Increase FInstanceCullingOcclusionQueryCS::NumThreadsPerGroup or implement wrapped group count support."),
-		NumInstances, MaxSupportedInstances))
-	{
-		return 0;
-	}
-
-	// Align buffer sizes to ensure each thread in the thread group has a valid slot to write without introducing bounds checks
-	const int32 AlignedNumInstances = NumThreadGroups.X * FInstanceCullingOcclusionQueryCS::NumThreadsPerGroup;
+	EPixelFormat VisibilityMaskFormat = GetPreferredVisibilityMaskFormat();
+	int32 VisibilityMaskStride = GPixelFormats[VisibilityMaskFormat].BlockBytes;
 
 	// Create the result buffer on demand
 	if (!CurrentInstanceOcclusionQueryBuffer)
 	{
+		const int32 AlignedNumGPUSceneInstances =
+			FMath::DivideAndRoundUp(NumGPUSceneInstances, FInstanceCullingOcclusionQueryCS::NumThreadsPerGroup)
+			* FInstanceCullingOcclusionQueryCS::NumThreadsPerGroup;
+
 		CurrentInstanceOcclusionQueryBuffer = GraphBuilder.CreateBuffer(
-			FRDGBufferDesc::CreateBufferDesc(sizeof(uint32), AlignedNumInstances),
+			FRDGBufferDesc::CreateBufferDesc(VisibilityMaskStride, AlignedNumGPUSceneInstances),
 			TEXT("FInstanceCullingOcclusionQueryRenderer_VisibleInstanceMask"));
 
-		AllocatedNumInstances = NumInstances;
+		InstanceOcclusionQueryBufferFormat = VisibilityMaskFormat;
+
+		AllocatedNumInstances = NumGPUSceneInstances;
+
+		// Create a wide-format alias for the underlying resource for a more efficient clear
+		FRDGBufferUAVRef UAV = GraphBuilder.CreateUAV(CurrentInstanceOcclusionQueryBuffer, PF_R32G32B32A32_UINT);
+		AddClearUAVPass(GraphBuilder, UAV, 0xFFFFFFFF);
 	}
 
-	checkf(uint32(NumInstances) == AllocatedNumInstances, TEXT("Number of instances in GPUScene is not expected change to during the frame"));
+	checkf(uint32(NumGPUSceneInstances) == AllocatedNumInstances, TEXT("Number of instances in GPUScene is not expected change to during the frame"));
 
 	FRDGBufferRef VisibleInstanceMaskBuffer = CurrentInstanceOcclusionQueryBuffer;
-	FRDGBufferUAVRef VisibilityMaskUAV = GraphBuilder.CreateUAV(VisibleInstanceMaskBuffer, PF_R32_UINT);
+	FRDGBufferUAVRef VisibilityMaskUAV = GraphBuilder.CreateUAV(VisibleInstanceMaskBuffer, VisibilityMaskFormat);
 
 	FRDGBufferRef IndirectArgsBuffer = GraphBuilder.CreateBuffer(
 		FRDGBufferDesc::CreateIndirectDesc<FRHIDrawIndexedIndirectParameters>(1),
 		TEXT("FInstanceCullingOcclusionQueryRenderer_IndirectArgsBuffer"));
 	FRDGBufferUAVRef IndirectArgsUAV = GraphBuilder.CreateUAV(IndirectArgsBuffer, PF_R32_UINT);
 
+	// Buffer of GPUScene instance indices to run occlusion queries for (input for setup CS)
+	FRDGBufferRef SetupInstanceIdBuffer = GraphBuilder.CreateBuffer(
+		FRDGBufferDesc::CreateBufferDesc(sizeof(uint32), 1 /*real size is provided via callback later*/),
+		TEXT("FInstanceCullingOcclusionQueryRenderer_SetupInstanceIdBuffer"), 
+		DeferredContext->DeferredAlignedNumInstances());
+
+	GraphBuilder.QueueBufferUpload(SetupInstanceIdBuffer,
+		DeferredContext->DeferredInstanceIdData(),
+		DeferredContext->DeferredInstanceIdDataSize());
+
+	FRDGBufferSRVRef SetupInstanceIdBufferSRV = GraphBuilder.CreateSRV(SetupInstanceIdBuffer, PF_R32_UINT);
+
+	// Buffer of GPUScene instance indices that passed the filtering in the setup CS pass and should be rendered in the subsequent graphics pass
 	FRDGBufferRef InstanceIdBuffer = GraphBuilder.CreateBuffer(
-		FRDGBufferDesc::CreateBufferDesc(sizeof(uint32), AlignedNumInstances),
-		TEXT("FInstanceCullingOcclusionQueryRenderer_InstanceIdBuffer"));
+		FRDGBufferDesc::CreateBufferDesc(sizeof(uint32), 1 /*real size is provided via callback later*/),
+		TEXT("FInstanceCullingOcclusionQueryRenderer_InstanceIdBuffer"),
+		DeferredContext->DeferredAlignedNumInstances());
+
 	FRDGBufferUAVRef InstanceIdUAV = GraphBuilder.CreateUAV(InstanceIdBuffer, PF_R32_UINT);
 	FRDGBufferSRVRef InstanceIdSRV = GraphBuilder.CreateSRV(InstanceIdBuffer, PF_R32_UINT);
 
@@ -361,19 +553,32 @@ uint32 FInstanceCullingOcclusionQueryRenderer::Render(
 		PassParameters->GPUSceneInstanceSceneData = GPUSceneParameters.GPUSceneInstanceSceneData;
 		PassParameters->GPUSceneInstancePayloadData = GPUSceneParameters.GPUSceneInstancePayloadData;
 		PassParameters->GPUScenePrimitiveSceneData = GPUSceneParameters.GPUScenePrimitiveSceneData;
-		PassParameters->NumInstances = NumInstances;
+		PassParameters->NumInstances = 0; // filled from DeferredContext later
+		PassParameters->InstanceIdBuffer = SetupInstanceIdBufferSRV;
 
 		TShaderMapRef<FInstanceCullingOcclusionQueryCS> ComputeShader(View.ShaderMap);
 
 		ClearUnusedGraphResources(ComputeShader, PassParameters);
 
-		FComputeShaderUtils::AddPass(
-			GraphBuilder, 
+		GraphBuilder.AddPass(
 			RDG_EVENT_NAME("InstanceCullingOcclusionQueryRenderer_Setup"),
-			ComputeShader,
 			PassParameters,
-			NumThreadGroups
-		);
+			ERDGPassFlags::Compute,
+			[PassParameters, DeferredContext, ComputeShader](FRHIRayTracingCommandList& RHICmdList)
+		{
+			if (!DeferredContext->bValid)
+			{
+				return;
+			}
+
+			PassParameters->NumInstances = DeferredContext->NumInstances;
+
+			FComputeShaderUtils::Dispatch(
+				RHICmdList,
+				ComputeShader,
+				*PassParameters,
+				DeferredContext->NumThreadGroups);
+		});
 	}
 
 	// Perform per-instance per-pixel occlusion tests by drawing bounding boxes that write into VisibleInstanceMaskBuffer slots for visible instances
@@ -400,9 +605,14 @@ uint32 FInstanceCullingOcclusionQueryRenderer::Render(
 		GraphBuilder.AddPass(
 			RDG_EVENT_NAME("InstanceCullingOcclusionQueryRenderer_Draw"),
 			PassParameters, ERDGPassFlags::Raster | ERDGPassFlags::NeverCull,
-			[PassParameters, NumInstances, &View](FRHICommandList& RHICmdList)
+			[PassParameters, DeferredContext, &View](FRHICommandList& RHICmdList)
 			{
-				RenderInstanceOcclusionCulling(RHICmdList, View, PassParameters, NumInstances);
+				if (!DeferredContext->bValid)
+				{
+					return;
+				}
+
+				RenderInstanceOcclusionCulling(RHICmdList, View, PassParameters);
 			});
 	}
 
@@ -417,11 +627,13 @@ void FInstanceCullingOcclusionQueryRenderer::MarkInstancesVisible(FRDGBuilder& G
 		return;
 	}
 
+	EPixelFormat VisibilityMaskFormat = GetPreferredVisibilityMaskFormat();
+
 	FRDGBufferRef Buffer = GraphBuilder.RegisterExternalBuffer(InstanceOcclusionQueryBuffer);
 
 	// Consecutive uses of the UAV will run in parallel.
 	// Allocating a unique RDG UAV here will still ensure that a barrier is inserted before the first dispatch.
-	FRDGBufferUAVRef UAV = GraphBuilder.CreateUAV(Buffer, PF_R32_UINT, ERDGUnorderedAccessViewFlags::SkipBarrier);
+	FRDGBufferUAVRef UAV = GraphBuilder.CreateUAV(Buffer, VisibilityMaskFormat, ERDGUnorderedAccessViewFlags::SkipBarrier);
 
 	// NOTE: It is possible to make this more efficient using a specialized GPU scatter shader, if we see many small batches here in practice
 	for (FGPUSceneInstanceRange Range : Ranges)
@@ -461,9 +673,11 @@ uint32 FInstanceCullingOcclusionQueryRenderer::RegisterView(const FViewInfo& Vie
 
 bool FInstanceCullingOcclusionQueryRenderer::IsCompatibleWithView(const FViewInfo& View)
 {
+	EPixelFormat VisibilityMaskFormat = GetPreferredVisibilityMaskFormat();
 	return FDataDrivenShaderPlatformInfo::GetSupportsVertexShaderSRVs(View.GetShaderPlatform())
 		&& View.GetSceneTextures().Depth.Target
 		&& View.HZB
+		&& VisibilityMaskFormat != PF_Unknown
 		&& CVarInstanceCullingOcclusionQueries.GetValueOnRenderThread() != 0;
 }
 
@@ -597,7 +811,7 @@ void FInstanceCullingOcclusionQueryRenderer::RenderDebug(FRDGBuilder& GraphBuild
 	PassParameters->VS.GPUSceneInstanceSceneData = GPUSceneParameters.GPUSceneInstanceSceneData;
 	PassParameters->VS.GPUSceneInstancePayloadData = GPUSceneParameters.GPUSceneInstancePayloadData;
 	PassParameters->VS.GPUScenePrimitiveSceneData = GPUSceneParameters.GPUScenePrimitiveSceneData;
-	PassParameters->VS.InstanceOcclusionQueryBuffer = GraphBuilder.CreateSRV(InstanceOcclusionQueryBufferRDG, PF_R32_UINT);
+	PassParameters->VS.InstanceOcclusionQueryBuffer = GraphBuilder.CreateSRV(InstanceOcclusionQueryBufferRDG, InstanceOcclusionQueryBufferFormat);
 	PassParameters->VS.HZBTexture = HZBTexture;
 	PassParameters->VS.HZBSampler = TStaticSamplerState<SF_Point, AM_Clamp, AM_Clamp, AM_Clamp>::GetRHI();
 	PassParameters->VS.HZBSize = FVector2f(HZBSize.X, HZBSize.Y);
diff --git a/Engine/Source/Runtime/Renderer/Private/InstanceCulling/InstanceCullingOcclusionQuery.h b/Engine/Source/Runtime/Renderer/Private/InstanceCulling/InstanceCullingOcclusionQuery.h
--- a/Engine/Source/Runtime/Renderer/Private/InstanceCulling/InstanceCullingOcclusionQuery.h
+++ b/Engine/Source/Runtime/Renderer/Private/InstanceCulling/InstanceCullingOcclusionQuery.h
@@ -46,12 +46,18 @@ public:
 	void EndFrame(FRDGBuilder& GraphBuilder);
 
 	/**
-	* One uint32 per instance in GPUScene, with 1 bit per view (see GetOcclusionQueryMaskForView)
+	* One uint8 per instance in GPUScene, with 1 bit per view
 	* Available when when instance culling is enabled and r.InstanceCulling.OcclusionQueries=1.
 	* Contains data for *previous* frame. Assumes that GPUScene instance indices are consistent.
 	*/
 	TRefCountPtr<FRDGPooledBuffer> InstanceOcclusionQueryBuffer;
 
+	/*
+	* Format that should be used to create views for InstanceOcclusionQueryBuffer.
+	* May be R8_UINT or R32_UINT, basd on current hardware capability.
+	*/
+	EPixelFormat InstanceOcclusionQueryBufferFormat = PF_Unknown;
+
 	/*
 	* Returns true if per-instance occlusion queries can be rendered for the view.
 	*/
@@ -68,7 +74,7 @@ private:
 
 	FRDGBufferRef CurrentInstanceOcclusionQueryBuffer = {};
 
-	static constexpr uint32 MaxViews = 1; // 32; -- YURIY_TODO: we could theoretically support up to 32 views
+	static constexpr uint32 MaxViews = 1; // -- YURIY_TODO: we could theoretically support up to 8 views for uint8 or 32 for uint32 mask
 	TArray<uint32> CurrentRenderedViewIDs;
 
 	uint32 AllocatedNumInstances = 0;
diff --git a/Engine/Source/Runtime/Renderer/Private/MeshDrawCommands.cpp b/Engine/Source/Runtime/Renderer/Private/MeshDrawCommands.cpp
--- a/Engine/Source/Runtime/Renderer/Private/MeshDrawCommands.cpp
+++ b/Engine/Source/Runtime/Renderer/Private/MeshDrawCommands.cpp
@@ -1608,7 +1608,7 @@ void FParallelMeshDrawCommandPass::BuildRenderingCommands(
 
 
 
-void FParallelMeshDrawCommandPass::WaitForSetupTask()
+void FParallelMeshDrawCommandPass::WaitForSetupTask() const
 {
 	WaitForMeshPassSetupTask();
 }
diff --git a/Engine/Source/Runtime/Renderer/Private/MeshDrawCommands.h b/Engine/Source/Runtime/Renderer/Private/MeshDrawCommands.h
--- a/Engine/Source/Runtime/Renderer/Private/MeshDrawCommands.h
+++ b/Engine/Source/Runtime/Renderer/Private/MeshDrawCommands.h
@@ -164,7 +164,7 @@ public:
 	/**
 	 * Sync with setup task.
 	 */
-	void WaitForSetupTask();
+	void WaitForSetupTask() const;
 
 	/**
 	 * Dispatch visible mesh draw command draw task.
@@ -183,6 +183,13 @@ public:
 	FInstanceCullingContext* GetInstanceCullingContext() { return &TaskContext.InstanceCullingContext; }
 	const FGraphEventRef& GetTaskEvent() const { return TaskEventRef; }
 
+	// NOTE: It is only safe to access mesh draw commands after the setup task is complete (use WaitForSetupTask). 
+	// Only access the data late in the frame to allow as much time as possible for async tasks to complete.
+	const FMeshCommandOneFrameArray& GetMeshDrawCommands() const
+	{
+		return TaskContext.MeshDrawCommands;
+	}
+
 private:
 	FMeshDrawCommandPassSetupTaskContext TaskContext;
 	FGraphEventRef TaskEventRef;
diff --git a/Engine/Source/Runtime/Renderer/Private/MobileShadingRenderer.cpp b/Engine/Source/Runtime/Renderer/Private/MobileShadingRenderer.cpp
--- a/Engine/Source/Runtime/Renderer/Private/MobileShadingRenderer.cpp
+++ b/Engine/Source/Runtime/Renderer/Private/MobileShadingRenderer.cpp
@@ -429,6 +429,7 @@ void FMobileSceneRenderer::InitViews(
 		&& Scene->InstanceCullingOcclusionQueryRenderer->InstanceOcclusionQueryBuffer)
 	{
 		InstanceCullingManager.InstanceOcclusionQueryBuffer = GraphBuilder.RegisterExternalBuffer(Scene->InstanceCullingOcclusionQueryRenderer->InstanceOcclusionQueryBuffer);
+		InstanceCullingManager.InstanceOcclusionQueryBufferFormat = Scene->InstanceCullingOcclusionQueryRenderer->InstanceOcclusionQueryBufferFormat;
 	}
 
 	// Create GPU-side representation of the view for instance culling.
diff --git a/Engine/Source/Runtime/Renderer/Private/SceneVisibility.cpp b/Engine/Source/Runtime/Renderer/Private/SceneVisibility.cpp
--- a/Engine/Source/Runtime/Renderer/Private/SceneVisibility.cpp
+++ b/Engine/Source/Runtime/Renderer/Private/SceneVisibility.cpp
@@ -4872,6 +4872,7 @@ void FDeferredShadingSceneRenderer::BeginInitViews(
 		&& Scene->InstanceCullingOcclusionQueryRenderer->InstanceOcclusionQueryBuffer)
 	{
 		InstanceCullingManager.InstanceOcclusionQueryBuffer = GraphBuilder.RegisterExternalBuffer(Scene->InstanceCullingOcclusionQueryRenderer->InstanceOcclusionQueryBuffer);
+		InstanceCullingManager.InstanceOcclusionQueryBufferFormat = Scene->InstanceCullingOcclusionQueryRenderer->InstanceOcclusionQueryBufferFormat;
 	}
 
 	// Create GPU-side representation of the view for instance culling.
