Mark the Tonemap scene texture RenderTargetable even if compute postprocessing is used because PP materials can still write to the texture using the HW raster pipeline.

diff --git a/Engine/Source/Runtime/Renderer/Private/PostProcess/PostProcessTonemap.cpp b/Engine/Source/Runtime/Renderer/Private/PostProcess/PostProcessTonemap.cpp
--- a/Engine/Source/Runtime/Renderer/Private/PostProcess/PostProcessTonemap.cpp
+++ b/Engine/Source/Runtime/Renderer/Private/PostProcess/PostProcessTonemap.cpp
@@ -580,7 +580,7 @@ FScreenPassTexture AddTonemapPass(FRDGBuilder& GraphBuilder, const FViewInfo& Vi
 			SceneColorViewport.Extent,
 			Inputs.SceneColor.TextureSRV->Desc.Texture->Desc.Format,
 			FClearValueBinding(FLinearColor(0, 0, 0, 0)),
-			GFastVRamConfig.Tonemap | TexCreate_ShaderResource | (View.bUseComputePasses ? TexCreate_UAV : TexCreate_RenderTargetable));;
+			GFastVRamConfig.Tonemap | TexCreate_ShaderResource | TexCreate_RenderTargetable | (View.bUseComputePasses ? TexCreate_UAV : TexCreate_None));;
 		
 		const FTonemapperOutputDeviceParameters OutputDeviceParameters = GetTonemapperOutputDeviceParameters(*View.Family);
 		const EDisplayOutputFormat OutputDevice = static_cast<EDisplayOutputFormat>(OutputDeviceParameters.OutputDevice);
