Defer deletion of scene proxies to scene update. This avoids race conditions for render command pipe commands that touch proxies, and allows deletion to happen asynchronously.

diff --git a/Engine/Source/Runtime/Renderer/Private/RendererScene.cpp b/Engine/Source/Runtime/Renderer/Private/RendererScene.cpp
--- a/Engine/Source/Runtime/Renderer/Private/RendererScene.cpp
+++ b/Engine/Source/Runtime/Renderer/Private/RendererScene.cpp
@@ -2416,21 +2416,7 @@ void FScene::RemovePrimitiveSceneInfo_RenderThread(FPrimitiveSceneInfo* Primitiv
 		UpdatedOcclusionBoundsSlacks.Remove(PrimitiveSceneInfo->Proxy);
 		DistanceFieldSceneDataUpdates.Remove(PrimitiveSceneInfo);
 		UpdatedAttachmentRoots.Remove(PrimitiveSceneInfo);
-
-		{
-			SCOPED_NAMED_EVENT(FScene_DeletePrimitiveSceneInfo, FColor::Red);
-			// Delete the PrimitiveSceneInfo on the game thread after the rendering thread has processed its removal.
-			// This must be done on the game thread because the hit proxy references (and possibly other members) need to be freed on the game thread.
-			struct DeferDeleteHitProxies : FDeferredCleanupInterface
-			{
-				DeferDeleteHitProxies(TArray<TRefCountPtr<HHitProxy>>&& InHitProxies) : HitProxies(MoveTemp(InHitProxies)) {}
-				TArray<TRefCountPtr<HHitProxy>> HitProxies;
-			};
-
-			BeginCleanup(new DeferDeleteHitProxies(MoveTemp(PrimitiveSceneInfo->HitProxies)));
-			delete PrimitiveSceneInfo->Proxy;
-			delete PrimitiveSceneInfo;
-		}
+		DeletedPrimitiveSceneInfos.Emplace(PrimitiveSceneInfo);
 	}
 	else
 	{
@@ -5519,8 +5505,7 @@ void FScene::UpdateAllPrimitiveSceneInfos(FRDGBuilder& GraphBuilder, EUpdateAllP
 
 	AddedLocalPrimitiveSceneInfos.Sort(FPrimitiveArraySortKey());
 
-	TSet<FPrimitiveSceneInfo*> DeletedSceneInfos;
-	DeletedSceneInfos.Reserve(RemovedLocalPrimitiveSceneInfos.Num());
+	DeletedPrimitiveSceneInfos.Reserve(RemovedLocalPrimitiveSceneInfos.Num());
 
 	TArray<int32> RemovedPrimitiveIndices;
 	RemovedPrimitiveIndices.SetNumUninitialized(RemovedLocalPrimitiveSceneInfos.Num());
@@ -5737,7 +5722,8 @@ void FScene::UpdateAllPrimitiveSceneInfos(FRDGBuilder& GraphBuilder, EUpdateAllP
 #endif
 
 				bNeedPathTracedInvalidation = bNeedPathTracedInvalidation || IsPrimitiveRelevantToPathTracing(PrimitiveSceneInfo);
-				DeletedSceneInfos.Add(PrimitiveSceneInfo);
+
+				DeletedPrimitiveSceneInfos.Emplace(PrimitiveSceneInfo);
 
 				const int32 PersistentIndex = PrimitiveSceneInfo->PersistentIndex.Index;
 				PersistentPrimitiveIdAllocator.Free(PersistentIndex);
@@ -5811,7 +5797,7 @@ void FScene::UpdateAllPrimitiveSceneInfos(FRDGBuilder& GraphBuilder, EUpdateAllP
 			FPrimitiveSceneInfo* PrimitiveSceneInfo = PrimitiveSceneProxy->GetPrimitiveSceneInfo();
 			
 			// being added or deleted, skip update logic
-			if (DeletedSceneInfos.Contains(PrimitiveSceneInfo))
+			if (DeletedPrimitiveSceneInfos.Contains(PrimitiveSceneInfo))
 			{
 				continue;
 			}
@@ -6076,7 +6062,7 @@ void FScene::UpdateAllPrimitiveSceneInfos(FRDGBuilder& GraphBuilder, EUpdateAllP
 		for (const auto& Transform : UpdatedTransforms)
 		{
 			FPrimitiveSceneProxy* PrimitiveSceneProxy = Transform.Key;
-			if (DeletedSceneInfos.Contains(PrimitiveSceneProxy->GetPrimitiveSceneInfo()))
+			if (DeletedPrimitiveSceneInfos.Contains(PrimitiveSceneProxy->GetPrimitiveSceneInfo()))
 			{
 				continue;
 			}
@@ -6168,7 +6154,7 @@ void FScene::UpdateAllPrimitiveSceneInfos(FRDGBuilder& GraphBuilder, EUpdateAllP
 			FPrimitiveSceneInfo* PrimitiveSceneInfo = PrimitiveSceneProxy->GetPrimitiveSceneInfo();
 			
 			// being added or deleted, skip update logic
-			if (DeletedSceneInfos.Contains(PrimitiveSceneInfo))
+			if (DeletedPrimitiveSceneInfos.Contains(PrimitiveSceneInfo))
 			{
 				continue;
 			}
@@ -6362,7 +6348,7 @@ void FScene::UpdateAllPrimitiveSceneInfos(FRDGBuilder& GraphBuilder, EUpdateAllP
 	for (const auto& CustomParams : UpdatedCustomPrimitiveParams)
 	{
 		FPrimitiveSceneProxy* PrimitiveSceneProxy = CustomParams.Key;
-		if (DeletedSceneInfos.Contains(PrimitiveSceneProxy->GetPrimitiveSceneInfo()))
+		if (DeletedPrimitiveSceneInfos.Contains(PrimitiveSceneProxy->GetPrimitiveSceneInfo()))
 		{
 			continue;
 		}
@@ -6450,7 +6436,7 @@ void FScene::UpdateAllPrimitiveSceneInfos(FRDGBuilder& GraphBuilder, EUpdateAllP
 	for (const auto& Attachments : UpdatedAttachmentRoots)
 	{
 		FPrimitiveSceneInfo* PrimitiveSceneInfo = Attachments.Key;
-		if (DeletedSceneInfos.Contains(PrimitiveSceneInfo))
+		if (DeletedPrimitiveSceneInfos.Contains(PrimitiveSceneInfo))
 		{
 			continue;
 		}
@@ -6462,7 +6448,7 @@ void FScene::UpdateAllPrimitiveSceneInfos(FRDGBuilder& GraphBuilder, EUpdateAllP
 
 	for (FPrimitiveSceneInfo* PrimitiveSceneInfo : DistanceFieldSceneDataUpdates)
 	{
-		if (DeletedSceneInfos.Contains(PrimitiveSceneInfo))
+		if (DeletedPrimitiveSceneInfos.Contains(PrimitiveSceneInfo))
 		{
 			continue;
 		}
@@ -6475,7 +6461,7 @@ void FScene::UpdateAllPrimitiveSceneInfos(FRDGBuilder& GraphBuilder, EUpdateAllP
 		const FPrimitiveSceneProxy* SceneProxy = OccSlackDelta.Key;
 		const FPrimitiveSceneInfo* SceneInfo = SceneProxy->GetPrimitiveSceneInfo();
 
-		if (DeletedSceneInfos.Contains(SceneInfo))
+		if (DeletedPrimitiveSceneInfos.Contains(SceneInfo))
 		{
 			continue;
 		}
@@ -6497,10 +6483,10 @@ void FScene::UpdateAllPrimitiveSceneInfos(FRDGBuilder& GraphBuilder, EUpdateAllP
 	// TODO: should refactor to add this as a dependency for a proxy-deletion task instead.
 	SceneCullingUpdater.GetAsyncProxyUseTaskHandle().Wait();
 
-	GraphBuilder.AddSetupTask([this, DeletedSceneInfos = MoveTemp(DeletedSceneInfos)]
+	GraphBuilder.AddSetupTask([DeletedPrimitiveSceneInfos = MoveTemp(DeletedPrimitiveSceneInfos)]
 	{
 		SCOPED_NAMED_EVENT(FScene_DeletePrimitiveSceneInfo, FColor::Red);
-		for (FPrimitiveSceneInfo* PrimitiveSceneInfo : DeletedSceneInfos)
+		for (FPrimitiveSceneInfo* PrimitiveSceneInfo : DeletedPrimitiveSceneInfos)
 		{
 			// It is possible that the HitProxies list isn't empty if PrimitiveSceneInfo was Added/Removed in same frame
 			// Delete the PrimitiveSceneInfo on the game thread after the rendering thread has processed its removal.
diff --git a/Engine/Source/Runtime/Renderer/Private/ScenePrivate.h b/Engine/Source/Runtime/Renderer/Private/ScenePrivate.h
--- a/Engine/Source/Runtime/Renderer/Private/ScenePrivate.h
+++ b/Engine/Source/Runtime/Renderer/Private/ScenePrivate.h
@@ -3839,6 +3839,7 @@ private:
 	Experimental::TRobinHoodHashSet<FPrimitiveSceneInfo*> RemovedPrimitiveSceneInfos;
 	Experimental::TRobinHoodHashSet<FPrimitiveSceneInfo*> DistanceFieldSceneDataUpdates;
 	TArray<FLevelCommand> LevelCommands;
+	TSet<FPrimitiveSceneInfo*> DeletedPrimitiveSceneInfos;
 
 	UE::Tasks::FTask CreateLightPrimitiveInteractionsTask;
 	UE::Tasks::FTask CacheMeshDrawCommandsTask;
