Debug Visualization: Anyhitshader count, count the number of triangles hit per ray, and generate a heatmap for them

diff --git a/Engine/Shaders/Private/RayTracing/RayTracingDebug.usf b/Engine/Shaders/Private/RayTracing/RayTracingDebug.usf
--- a/Engine/Shaders/Private/RayTracing/RayTracingDebug.usf
+++ b/Engine/Shaders/Private/RayTracing/RayTracingDebug.usf
@@ -29,6 +29,8 @@ float MaxTraceDistance;
 float FarFieldMaxTraceDistance;
 float3 FarFieldReferencePos;
 
+float TriangleHitCountMaxThreshold;
+
 StructuredBuffer<FRayTracingInstanceDebugData> InstancesDebugData;
 StructuredBuffer<FRayTracingPickingFeedback> PickingBuffer;
 
@@ -168,9 +170,16 @@ float4 DebugRayTracingInstance(FRayDesc Ray, uint2 PixelCoord, out float OutDept
 	FRayTracingDebugPayload Payload = (FRayTracingDebugPayload)0;
 	Payload.SetMiss();
 
-	const uint RayFlags = RAY_FLAG_CULL_BACK_FACING_TRIANGLES;
+	uint RayFlags = RAY_FLAG_CULL_BACK_FACING_TRIANGLES;
 	const uint InstanceInclusionMask = (OpaqueOnly ? RAY_TRACING_MASK_OPAQUE : RAY_TRACING_MASK_ALL);
-
+	if (VisualizationMode == RAY_TRACING_DEBUG_VIZ_TRIANGLE_HITCOUNT)
+	{
+		RayFlags |= RAY_FLAG_FORCE_NON_OPAQUE;
+	}
+	else
+	{
+		RayFlags |= RAY_FLAG_FORCE_OPAQUE;
+	}
 	TraceRay(
 		TLAS,
 		RayFlags  /*RayFlags*/,
@@ -192,6 +201,9 @@ float4 DebugRayTracingInstance(FRayDesc Ray, uint2 PixelCoord, out float OutDept
 		case RAY_TRACING_DEBUG_VIZ_INSTANCES:
 			Result.rgb = InstanceColor;
 			break;
+		case RAY_TRACING_DEBUG_VIZ_TRIANGLE_HITCOUNT:
+			Result.rgb = ColorMapTurbo(Payload.TriangleHitCountPerRay / TriangleHitCountMaxThreshold);
+			break;
 		case RAY_TRACING_DEBUG_VIZ_INSTANCE_OVERLAP:
 			Result.rgb = 1.0; // We only use fake lighting as instance heatmap is blended in
 			break;
@@ -234,8 +246,12 @@ float4 DebugRayTracingInstance(FRayDesc Ray, uint2 PixelCoord, out float OutDept
 		const float FakeLighting = dot(Payload.WorldNormal, normalize(float3(1.0f, 1.0f, 1.0f)));
 		const float FakeLightingScale = (VisualizationMode == RAY_TRACING_DEBUG_VIZ_INSTANCE_OVERLAP) ? 0.5f : 0.125f;
 		const float FakeLightingBias = (VisualizationMode == RAY_TRACING_DEBUG_VIZ_INSTANCE_OVERLAP) ? 0.5f : 0.875f;
+		const int NoFakeLighting = (VisualizationMode == RAY_TRACING_DEBUG_VIZ_TRIANGLE_HITCOUNT);
 
-		Result.rgb *= FakeLighting * FakeLightingScale + FakeLightingBias;
+		if (!NoFakeLighting)
+		{
+			Result.rgb *= FakeLighting * FakeLightingScale + FakeLightingBias;
+		}
 
 		float3 TranslatedWorldPosition = Ray.Origin + Payload.HitT * Ray.Direction;
 		float4 ClipPosition = mul(float4(TranslatedWorldPosition, 1.0), View.TranslatedWorldToClip);
diff --git a/Engine/Shaders/Private/RayTracing/RayTracingDebugCHS.usf b/Engine/Shaders/Private/RayTracing/RayTracingDebugCHS.usf
--- a/Engine/Shaders/Private/RayTracing/RayTracingDebugCHS.usf
+++ b/Engine/Shaders/Private/RayTracing/RayTracingDebugCHS.usf
@@ -24,3 +24,10 @@ RAY_TRACING_ENTRY_CLOSEST_HIT(RayTracingDebugMainCHS,
 	Payload.InstanceIndex = InstanceIndex();
 	Payload.GeometryIndex = 0; // GeometryIndex();
 }
+
+RAY_TRACING_ENTRY_ANY_HIT(RayTracingDebugAHS,
+	FRayTracingDebugPayload, Payload,
+	FRayTracingIntersectionAttributes, Attributes)
+{
+	Payload.TriangleHitCountPerRay += 1;
+}
diff --git a/Engine/Shaders/Private/RayTracing/RayTracingDebugUtils.ush b/Engine/Shaders/Private/RayTracing/RayTracingDebugUtils.ush
--- a/Engine/Shaders/Private/RayTracing/RayTracingDebugUtils.ush
+++ b/Engine/Shaders/Private/RayTracing/RayTracingDebugUtils.ush
@@ -15,6 +15,8 @@ struct FRayTracingDebugPayload : FMinimalPayload
 
 	float3 WorldNormal;
 	uint GeometryIndex;
+	
+	uint TriangleHitCountPerRay;
 };
 CHECK_RT_PAYLOAD_SIZE(FRayTracingDebugPayload)
 #endif
diff --git a/Engine/Shaders/Shared/RaytracingDebugDefinitions.h b/Engine/Shaders/Shared/RaytracingDebugDefinitions.h
--- a/Engine/Shaders/Shared/RaytracingDebugDefinitions.h
+++ b/Engine/Shaders/Shared/RaytracingDebugDefinitions.h
@@ -42,6 +42,7 @@
 #define RAY_TRACING_DEBUG_VIZ_PROXY_TYPE				32
 #define RAY_TRACING_DEBUG_VIZ_PICKER					33
 #define RAY_TRACING_DEBUG_VIZ_INSTANCE_OVERLAP			34
+#define RAY_TRACING_DEBUG_VIZ_TRIANGLE_HITCOUNT			35
 
 #define RAY_TRACING_DEBUG_PICKER_DOMAIN_TRIANGLE		0
 #define RAY_TRACING_DEBUG_PICKER_DOMAIN_INSTANCE		1
diff --git a/Engine/Source/Editor/UnrealEd/Private/RayTracingDebugVisualizationMenuCommands.cpp b/Engine/Source/Editor/UnrealEd/Private/RayTracingDebugVisualizationMenuCommands.cpp
--- a/Engine/Source/Editor/UnrealEd/Private/RayTracingDebugVisualizationMenuCommands.cpp
+++ b/Engine/Source/Editor/UnrealEd/Private/RayTracingDebugVisualizationMenuCommands.cpp
@@ -55,7 +55,7 @@ void FRayTracingDebugVisualizationMenuCommands::CreateRayTracingDebugVisualizati
 	RayTracingDebugModeNames.Add(LOCTEXT("Anisotropy", "Anisotropy"));
 	RayTracingDebugModeNames.Add(LOCTEXT("Instances", "Instances"));
 	RayTracingDebugModeNames.Add(LOCTEXT("Instance Overlap", "Instance Overlap"));
-
+	RayTracingDebugModeNames.Add(LOCTEXT("Triangle Hit Count", "Triangle Hit Count"));
 	if (GRHISupportsShaderTimestamp)
 	{
 		RayTracingDebugModeNames.Add(LOCTEXT("Performance", "Performance"));
diff --git a/Engine/Source/Runtime/Renderer/Private/RayTracing/RayTracingDebug.cpp b/Engine/Source/Runtime/Renderer/Private/RayTracing/RayTracingDebug.cpp
--- a/Engine/Source/Runtime/Renderer/Private/RayTracing/RayTracingDebug.cpp
+++ b/Engine/Source/Runtime/Renderer/Private/RayTracing/RayTracingDebug.cpp
@@ -89,6 +89,12 @@ static TAutoConsoleVariable<float> CVarRayTracingDebugTraversalTriangleScale(
 	TEXT("Scaling factor for triangle traversal heat map visualization. (default = 30)\n")
 );
 
+static TAutoConsoleVariable<int32> CVarRayTracingDebugHitCountMaxThreshold(
+	TEXT("r.RayTracing.DebugTriangleHitCount.MaxThreshold"),
+	6,
+	TEXT("Maximum hit count threshold for debug ray tracing triangle hit count heat map visualization. (default = 6)\n")
+);
+
 static int32 GVisualizeProceduralPrimitives = 0;
 static FAutoConsoleVariableRef CVarVisualizeProceduralPrimitives(
 	TEXT("r.RayTracing.DebugVisualizationMode.ProceduralPrimitives"),
@@ -98,7 +104,7 @@ static FAutoConsoleVariableRef CVarVisualizeProceduralPrimitives(
 	ECVF_RenderThreadSafe
 );
 
-IMPLEMENT_RT_PAYLOAD_TYPE(ERayTracingPayloadType::RayTracingDebug, 36);
+IMPLEMENT_RT_PAYLOAD_TYPE(ERayTracingPayloadType::RayTracingDebug, 40);
 
 class FRayTracingDebugRGS : public FGlobalShader
 {
@@ -118,6 +124,7 @@ class FRayTracingDebugRGS : public FGlobalShader
 		SHADER_PARAMETER(float, FarFieldMaxTraceDistance)
 		SHADER_PARAMETER(FVector3f, FarFieldReferencePos)
 		SHADER_PARAMETER(int32, OpaqueOnly)
+		SHADER_PARAMETER(float, TriangleHitCountMaxThreshold)
 		SHADER_PARAMETER_RDG_BUFFER_SRV(RaytracingAccelerationStructure, TLAS)
 		SHADER_PARAMETER_RDG_TEXTURE_UAV(RWTexture2D<float4>, Output)
 		SHADER_PARAMETER_RDG_TEXTURE_UAV(RWTexture2D<float>, OutputDepth)
@@ -182,7 +189,7 @@ public:
 		return ERayTracingPayloadType::RayTracingDebug;
 	}
 };
-IMPLEMENT_GLOBAL_SHADER(FRayTracingDebugCHS, "/Engine/Private/RayTracing/RayTracingDebugCHS.usf", "RayTracingDebugMainCHS", SF_RayHitGroup);
+IMPLEMENT_GLOBAL_SHADER(FRayTracingDebugCHS, "/Engine/Private/RayTracing/RayTracingDebugCHS.usf", "closesthit=RayTracingDebugMainCHS anyhit=RayTracingDebugAHS", SF_RayHitGroup);
 
 class FRayTracingDebugMS : public FGlobalShader
 {
@@ -212,7 +219,6 @@ public:
 };
 IMPLEMENT_GLOBAL_SHADER(FRayTracingDebugMS, "/Engine/Private/RayTracing/RayTracingDebugMS.usf", "RayTracingDebugMS", SF_RayMiss);
 
-
 class FRayTracingDebugTraversalCS : public FGlobalShader
 {
 	DECLARE_GLOBAL_SHADER(FRayTracingDebugTraversalCS)
@@ -449,7 +455,6 @@ struct FRayTracingDebugResources : public FRenderResource
 				PickingBuffers[BufferIndex] = nullptr;
 			}
 		}
-
 		PickingBuffers.Reset();
 	}
 };
@@ -550,12 +555,13 @@ void BindRayTracingDebugCHSMaterialBindings(FRHICommandList& RHICmdList, const F
 
 static bool RequiresRayTracingDebugCHS(uint32 DebugVisualizationMode)
 {
-	return DebugVisualizationMode == RAY_TRACING_DEBUG_VIZ_INSTANCES || 
+	return DebugVisualizationMode == RAY_TRACING_DEBUG_VIZ_INSTANCES ||
 		DebugVisualizationMode == RAY_TRACING_DEBUG_VIZ_TRIANGLES ||
 		DebugVisualizationMode == RAY_TRACING_DEBUG_VIZ_DYNAMIC_INSTANCES ||
 		DebugVisualizationMode == RAY_TRACING_DEBUG_VIZ_PROXY_TYPE ||
 		DebugVisualizationMode == RAY_TRACING_DEBUG_VIZ_PICKER ||
-		DebugVisualizationMode == RAY_TRACING_DEBUG_VIZ_INSTANCE_OVERLAP;
+		DebugVisualizationMode == RAY_TRACING_DEBUG_VIZ_INSTANCE_OVERLAP ||
+		DebugVisualizationMode == RAY_TRACING_DEBUG_VIZ_TRIANGLE_HITCOUNT;
 }
 
 static bool IsRayTracingDebugTraversalMode(uint32 DebugVisualizationMode)
@@ -857,6 +863,7 @@ void FDeferredShadingSceneRenderer::RenderRayTracingDebug(FRDGBuilder& GraphBuil
 		RayTracingDebugVisualizationModes.Emplace(FName(*LOCTEXT("Proxy Type", "Proxy Type").ToString()),										RAY_TRACING_DEBUG_VIZ_PROXY_TYPE);
 		RayTracingDebugVisualizationModes.Emplace(FName(*LOCTEXT("Picker", "Picker").ToString()),												RAY_TRACING_DEBUG_VIZ_PICKER);
 		RayTracingDebugVisualizationModes.Emplace(FName(*LOCTEXT("Instance Overlap", "Instance Overlap").ToString()),							RAY_TRACING_DEBUG_VIZ_INSTANCE_OVERLAP);
+		RayTracingDebugVisualizationModes.Emplace(FName(*LOCTEXT("Triangle Hit Count", "Triangle Hit Count").ToString()),						RAY_TRACING_DEBUG_VIZ_TRIANGLE_HITCOUNT);
 	}
 
 	uint32 DebugVisualizationMode;
@@ -994,7 +1001,7 @@ void FDeferredShadingSceneRenderer::RenderRayTracingDebug(FRDGBuilder& GraphBuil
 		PermutationVectorCHS.Set<FRayTracingDebugCHS::FNaniteRayTracing>(true);
 		auto HitGroupShaderNaniteRT = View.ShaderMap->GetShader<FRayTracingDebugCHS>(PermutationVectorCHS);
 
-		FRHIRayTracingShader* HitGroupTable[] = { HitGroupShader.GetRayTracingShader(), HitGroupShaderNaniteRT.GetRayTracingShader() };
+		FRHIRayTracingShader* HitGroupTable[] = { HitGroupShader.GetRayTracingShader(), HitGroupShaderNaniteRT.GetRayTracingShader()};
 		Initializer.SetHitGroupTable(HitGroupTable);
 		Initializer.bAllowHitGroupIndexing = true; // Required for stable output using GetBaseInstanceIndex().
 
@@ -1013,6 +1020,8 @@ void FDeferredShadingSceneRenderer::RenderRayTracingDebug(FRDGBuilder& GraphBuil
 	RayGenParameters->ShouldUsePreExposure = View.Family->EngineShowFlags.Tonemapper;
 	RayGenParameters->TimingScale = CVarRayTracingDebugTimingScale.GetValueOnAnyThread() / 25000.0f;
 	RayGenParameters->OpaqueOnly = CVarRayTracingDebugModeOpaqueOnly.GetValueOnRenderThread();
+	RayGenParameters->TriangleHitCountMaxThreshold = FMath::Clamp((float)CVarRayTracingDebugHitCountMaxThreshold.GetValueOnRenderThread(), 1, 100000);
+
 	
 	// If we don't output depth, create dummy 1x1 texture
 	const bool bOutputDepth = DebugVisualizationMode == RAY_TRACING_DEBUG_VIZ_INSTANCE_OVERLAP;
