Path tracing updates with Heterogeneous Volumes

1) Move Heterogeneous Volume grid builder into frame loop, preventing grids from being built once the simulation completes.
2) Remove r.PathTracing.HeterogeneousVolumes variable, inheriting system default behavior.

diff --git a/Engine/Source/Runtime/Renderer/Private/HeterogeneousVolumes/HeterogeneousVolumes.cpp b/Engine/Source/Runtime/Renderer/Private/HeterogeneousVolumes/HeterogeneousVolumes.cpp
--- a/Engine/Source/Runtime/Renderer/Private/HeterogeneousVolumes/HeterogeneousVolumes.cpp
+++ b/Engine/Source/Runtime/Renderer/Private/HeterogeneousVolumes/HeterogeneousVolumes.cpp
@@ -192,6 +192,18 @@ bool ShouldRenderHeterogeneousVolumes(
 		&& DoesPlatformSupportHeterogeneousVolumes(Scene->GetShaderPlatform());
 }
 
+bool ShouldRenderHeterogeneousVolumesForAnyView(
+	const TArray<FViewInfo>& Views
+)
+{
+	bool Result = false;
+	for (int32 ViewIndex = 0; ViewIndex < Views.Num(); ++ViewIndex)
+	{
+		Result |= ShouldRenderHeterogeneousVolumesForView(Views[ViewIndex]);
+	}
+	return Result;
+}
+
 bool ShouldRenderHeterogeneousVolumesForView(
 	const FViewInfo& View
 )
diff --git a/Engine/Source/Runtime/Renderer/Private/HeterogeneousVolumes/HeterogeneousVolumes.h b/Engine/Source/Runtime/Renderer/Private/HeterogeneousVolumes/HeterogeneousVolumes.h
--- a/Engine/Source/Runtime/Renderer/Private/HeterogeneousVolumes/HeterogeneousVolumes.h
+++ b/Engine/Source/Runtime/Renderer/Private/HeterogeneousVolumes/HeterogeneousVolumes.h
@@ -27,6 +27,7 @@ struct FSceneTextures;
 //
 
 bool ShouldRenderHeterogeneousVolumes(const FScene* Scene);
+bool ShouldRenderHeterogeneousVolumesForAnyView(const TArray<FViewInfo>& Views);
 bool ShouldRenderHeterogeneousVolumesForView(const FViewInfo& View);
 bool DoesPlatformSupportHeterogeneousVolumes(EShaderPlatform Platform);
 bool DoesMaterialShaderSupportHeterogeneousVolumes(const FMaterialShaderParameters& Parameters);
diff --git a/Engine/Source/Runtime/Renderer/Private/PathTracing.cpp b/Engine/Source/Runtime/Renderer/Private/PathTracing.cpp
--- a/Engine/Source/Runtime/Renderer/Private/PathTracing.cpp
+++ b/Engine/Source/Runtime/Renderer/Private/PathTracing.cpp
@@ -320,17 +320,6 @@ TAutoConsoleVariable<int32> CVarPathTracingLightFunctionColor(
 	ECVF_RenderThreadSafe
 );
 
-TAutoConsoleVariable<int32> CVarPathTracingHeterogeneousVolumes(
-	TEXT("r.PathTracing.HeterogeneousVolumes"),
-	0,
-	TEXT("Enables heterogeneous volumes (default = 0)\n")
-	TEXT("0: off (default)\n")
-	TEXT("1: on (combined frustum + ortho grids)\n")
-	TEXT("2: frustum-only grid\n")
-	TEXT("3: ortho-only grid\n"),
-	ECVF_RenderThreadSafe
-);
-
 TAutoConsoleVariable<int32> CVarPathTracingHeterogeneousVolumesRebuildEveryFrame(
 	TEXT("r.PathTracing.HeterogeneousVolumes.RebuildEveryFrame"),
 	1,
@@ -637,7 +626,7 @@ static void PreparePathTracingData(const FScene* Scene, const FViewInfo& View, F
 		&& (Scene->ExponentialFogs[0].FogData[0].Density > 0 ||
 			Scene->ExponentialFogs[0].FogData[1].Density > 0);
 
-	PathTracingData.EnableHeterogeneousVolumes = CVarPathTracingHeterogeneousVolumes.GetValueOnRenderThread();
+	PathTracingData.EnableHeterogeneousVolumes = ShouldRenderHeterogeneousVolumesForView(View);
 	PathTracingData.UseAnalyticTransmittance = EvalUseAnalyticTransmittance(View);
 	PathTracingData.EnableDBuffer = CVarPathTracingUseDBuffer.GetValueOnRenderThread();
 
@@ -2470,30 +2459,10 @@ void FDeferredShadingSceneRenderer::RenderPathTracing(
 		View.ViewState->PathTracingInvalidate();
 	}
 
-
+	// Declare heterogeneous volume buffers
 	TRDGUniformBufferRef<FOrthoVoxelGridUniformBufferParameters> OrthoGridUniformBuffer;
 	TRDGUniformBufferRef<FFrustumVoxelGridUniformBufferParameters> FrustumGridUniformBuffer;
-	bool bForceRebuild = CVarPathTracingHeterogeneousVolumesRebuildEveryFrame.GetValueOnRenderThread() != 0;
-	bool bCreateVolumeGrids = bForceRebuild ||
-		!PathTracingState->AdaptiveFrustumGridParameterCache.TopLevelGridBuffer ||
-		!PathTracingState->AdaptiveOrthoGridParameterCache.TopLevelGridBuffer;
-	if (bCreateVolumeGrids)
-	{
-		BuildOrthoVoxelGrid(GraphBuilder, Scene, Views, OrthoGridUniformBuffer);
-		BuildFrustumVoxelGrid(GraphBuilder, Scene, Views[0], FrustumGridUniformBuffer);
-	}
-	else
-	{
-		RegisterExternalOrthoVoxelGridUniformBuffer(GraphBuilder,
-			PathTracingState->AdaptiveOrthoGridParameterCache,
-			OrthoGridUniformBuffer
-		);
-
-		RegisterExternalFrustumVoxelGridUniformBuffer(GraphBuilder,
-			PathTracingState->AdaptiveFrustumGridParameterCache,
-			FrustumGridUniformBuffer
-		);
-	}
+	bool bCreateVolumeGrids = false;
 
 	// Prepare radiance buffer (will be shared with display pass)
 	FRDGTexture* RadianceTexture = nullptr;
@@ -2557,6 +2526,28 @@ void FDeferredShadingSceneRenderer::RenderPathTracing(
 
 		if (bNeedsMoreRays)
 		{
+			bool bForceRebuild = CVarPathTracingHeterogeneousVolumesRebuildEveryFrame.GetValueOnRenderThread() != 0;
+			bCreateVolumeGrids = bForceRebuild ||
+				!PathTracingState->AdaptiveFrustumGridParameterCache.TopLevelGridBuffer ||
+				!PathTracingState->AdaptiveOrthoGridParameterCache.TopLevelGridBuffer;
+			if (bCreateVolumeGrids)
+			{
+				BuildOrthoVoxelGrid(GraphBuilder, Scene, Views, OrthoGridUniformBuffer);
+				BuildFrustumVoxelGrid(GraphBuilder, Scene, Views[0], FrustumGridUniformBuffer);
+			}
+			else
+			{
+				RegisterExternalOrthoVoxelGridUniformBuffer(GraphBuilder,
+					PathTracingState->AdaptiveOrthoGridParameterCache,
+					OrthoGridUniformBuffer
+				);
+
+				RegisterExternalFrustumVoxelGridUniformBuffer(GraphBuilder,
+					PathTracingState->AdaptiveFrustumGridParameterCache,
+					FrustumGridUniformBuffer
+				);
+			}
+
 			// We are writing to the texture, we'll need to extract it...
 			bNeedsTextureExtract = true;
 
