CPU optimizations for FInstanceCullingOcclusionQueryDeferredContext::Execute

* Force inline IsRelevantCommand
* Write instance indices directly into TArray memory, instead of going through TArray::Add
* Add few CPU timing scopes

diff --git a/Engine/Source/Runtime/Renderer/Private/InstanceCulling/InstanceCullingOcclusionQuery.cpp b/Engine/Source/Runtime/Renderer/Private/InstanceCulling/InstanceCullingOcclusionQuery.cpp
--- a/Engine/Source/Runtime/Renderer/Private/InstanceCulling/InstanceCullingOcclusionQuery.cpp
+++ b/Engine/Source/Runtime/Renderer/Private/InstanceCulling/InstanceCullingOcclusionQuery.cpp
@@ -20,6 +20,7 @@
 #include "UnifiedBuffer.h"
 #include "HAL/IConsoleManager.h"
 #include "DataDrivenShaderPlatformInfo.h"
+#include "ProfilingDebugging/CpuProfilerTrace.h"
 
 static TAutoConsoleVariable<int32> CVarInstanceCullingOcclusionQueries(
 	TEXT("r.InstanceCulling.OcclusionQueries"),
@@ -308,7 +309,7 @@ struct FInstanceCullingOcclusionQueryDeferredContext
 	{
 	}
 
-	static bool IsRelevantCommand(const FVisibleMeshDrawCommand& VisibleCommand)
+	static FORCEINLINE bool IsRelevantCommand(const FVisibleMeshDrawCommand& VisibleCommand)
 	{
 		// There may be multiple visible mesh draw commands that refer to the same instance when GPU-based LOD selection is used.
 		// This filter is designed to remove the duplicates, keeping only the "authoritative" instance.
@@ -330,6 +331,8 @@ struct FInstanceCullingOcclusionQueryDeferredContext
 			return;
 		}
 
+		TRACE_CPUPROFILER_EVENT_SCOPE(FInstanceCullingOcclusionQueryDeferredContext::Execute);
+
 		bExecuted = true;
 
 		const FParallelMeshDrawCommandPass& MeshDrawCommandPass = View->ParallelMeshDrawCommandPasses[MeshPass];
@@ -340,18 +343,7 @@ struct FInstanceCullingOcclusionQueryDeferredContext
 
 		const FMeshCommandOneFrameArray& VisibleMeshDrawCommands = MeshDrawCommandPass.GetMeshDrawCommands();
 
-		uint32 NumPotentiallyVisibleInstances = 0;
-
-		for (const FVisibleMeshDrawCommand& VisibleCommand : VisibleMeshDrawCommands)
-		{
-			if (!IsRelevantCommand(VisibleCommand))
-			{
-				continue;
-			}
-			NumPotentiallyVisibleInstances += VisibleCommand.MeshDrawCommand->NumInstances;
-		}
-
-		NumInstances = NumPotentiallyVisibleInstances;
+		NumInstances = CountVisibleInstances(VisibleMeshDrawCommands);
 
 		NumThreadGroups = FComputeShaderUtils::GetGroupCount(NumInstances, FInstanceCullingOcclusionQueryCS::NumThreadsPerGroup);
 
@@ -373,8 +365,40 @@ struct FInstanceCullingOcclusionQueryDeferredContext
 			return;
 		}
 
-		VisibleInstanceIds.Reserve(AlignedNumInstances);
+		FillVisibleInstanceIds(VisibleMeshDrawCommands);
+
+		bValid = true;
+	}
+
+	uint32 CountVisibleInstances(const FMeshCommandOneFrameArray& VisibleMeshDrawCommands) const 
+	{
+		TRACE_CPUPROFILER_EVENT_SCOPE(FInstanceCullingOcclusionQueryDeferredContext::CountVisibleInstances);
+
+		uint32 Result = 0;
+
+		for (const FVisibleMeshDrawCommand& VisibleCommand : VisibleMeshDrawCommands)
+		{
+			if (!IsRelevantCommand(VisibleCommand))
+			{
+				continue;
+			}
+			Result += VisibleCommand.MeshDrawCommand->NumInstances;
+		}
+
+		return Result;
+	}
+
+	void FillVisibleInstanceIds(const FMeshCommandOneFrameArray& VisibleMeshDrawCommands)
+	{
+		TRACE_CPUPROFILER_EVENT_SCOPE(FInstanceCullingOcclusionQueryDeferredContext::FillVisibleInstanceIds);
+
+		check(AlignedNumInstances != 0);
 
+		// Write output data directly, bypassing TArray::Add overhead (resize branch, etc.)
+		VisibleInstanceIds.SetNumUninitialized(AlignedNumInstances);
+		uint32* ResultData = VisibleInstanceIds.GetData();
+		uint32* ResultCursor = ResultData;
+		
 		for (const FVisibleMeshDrawCommand& VisibleCommand : VisibleMeshDrawCommands)
 		{
 			if (!IsRelevantCommand(VisibleCommand))
@@ -389,18 +413,18 @@ struct FInstanceCullingOcclusionQueryDeferredContext
 
 			for (uint32 i = 0; i < CommandNumInstances; ++i)
 			{
-				VisibleInstanceIds.Add(InstanceBaseIndex + i);
+				*ResultCursor = InstanceBaseIndex + i;
+				++ResultCursor;
 			}
 		}
 
 		for (int32 i = NumInstances; i < AlignedNumInstances; ++i)
 		{
-			VisibleInstanceIds.Add(0);
+			*ResultCursor = 0;
+			++ResultCursor;
 		}
 
-		check(VisibleInstanceIds.Num() == AlignedNumInstances);
-
-		bValid = true;
+		check(ResultCursor == ResultData + AlignedNumInstances);
 	}
 
 	FRDGBufferNumElementsCallback DeferredAlignedNumInstances()
@@ -464,6 +488,8 @@ uint32 FInstanceCullingOcclusionQueryRenderer::Render(
 		return 0;
 	}
 
+	TRACE_CPUPROFILER_EVENT_SCOPE(FInstanceCullingOcclusionQueryRenderer::Render);
+
 	const int32 NumGPUSceneInstances = GPUScene.GetNumInstances();
 
 	FInstanceCullingOcclusionQueryDeferredContext* DeferredContext = GraphBuilder.AllocObject<FInstanceCullingOcclusionQueryDeferredContext>(&View, NumGPUSceneInstances, EMeshPass::BasePass);
