Fix incorrect shadow culling of ISM LODs when there is more than one section in a mesh.

diff --git a/Engine/Source/Runtime/Renderer/Private/ShadowSetup.cpp b/Engine/Source/Runtime/Renderer/Private/ShadowSetup.cpp
--- a/Engine/Source/Runtime/Renderer/Private/ShadowSetup.cpp
+++ b/Engine/Source/Runtime/Renderer/Private/ShadowSetup.cpp
@@ -1355,7 +1355,9 @@ struct FAddSubjectPrimitiveResult
 			uint32 bFadingIn : 1;
 			uint32 bAddOnRenderThread : 1;
 			uint32 bShouldRecordShadowSubjectsForMobile : 1;
-			uint32 bIsLODRange : 1;
+			uint32 bIsLodRange : 1;
+			uint32 LodRangeMin : 7;
+			uint32 LodRangeMax : 7;
 
 			union
 			{
@@ -1460,6 +1462,13 @@ struct FAddSubjectPrimitiveResult
 		return NumMeshes;
 	}
 
+	void SetLodRange(FLODMask InLodMask)
+	{
+		bIsLodRange = InLodMask.IsLODRange();
+		LodRangeMin = bIsLodRange ? InLodMask.LODIndex0 : 0;
+		LodRangeMax = bIsLodRange ? InLodMask.LODIndex1 : 0;
+	}
+
 private:
 	void HandleOverflow(FAddSubjectPrimitiveOverflowedIndices& OverflowBuffer)
 	{
@@ -1728,7 +1737,7 @@ bool FProjectedShadowInfo::ShouldDrawStaticMeshes_AnyThread(
 		const int32 ForcedLOD = CurrentView.Family->EngineShowFlags.LOD ? (GetCVarForceLODShadow_AnyThread() != -1 ? GetCVarForceLODShadow_AnyThread() : GetCVarForceLOD_AnyThread()) : -1;
 
 		FLODMask ShadowLODToRender = CalcAndUpdateLODToRender(CurrentView, FBoxSphereBounds(PrimitiveSceneInfoCompact.Bounds), PrimitiveSceneInfo, ForcedLOD);
-		const bool bIsLODRange = ShadowLODToRender.IsLODRange();
+		OutResult.SetLodRange(ShadowLODToRender);
 
 		if (WholeSceneDirectionalShadow)
 		{
@@ -1752,8 +1761,6 @@ bool FProjectedShadowInfo::ShouldDrawStaticMeshes_AnyThread(
 						++OutStats.NumMDCBuildRequests;
 						OutResult.AcceptMesh(NumAcceptedStaticMeshes++, MeshIndex, OverflowBuffer);
 					}
-
-					OutResult.bIsLODRange = bIsLODRange;
 				}
 			}
 		}
@@ -1771,8 +1778,6 @@ bool FProjectedShadowInfo::ShouldDrawStaticMeshes_AnyThread(
 					check(MeshIndex < MAX_uint16);
 					++OutStats.NumMDCBuildRequests;
 					OutResult.AcceptMesh(NumAcceptedStaticMeshes++, MeshIndex, OverflowBuffer);
-
-					OutResult.bIsLODRange = bIsLODRange;
 				}
 			}
 		}
@@ -2224,8 +2229,7 @@ void FProjectedShadowInfo::FinalizeAddSubjectPrimitive(
 				&Scene->CachedDrawLists[MeshPassTargetType].MeshDrawCommands[CmdInfo.CommandIndex];
 			const int32 PrimIdx = PrimitiveSceneInfo->GetIndex();
 			const int32 InstanceSceneDataOffset = PrimitiveSceneInfo->GetInstanceSceneDataOffset();
-			// Assumption: We have one command per LOD in order of increasing LOD index.
-			const EMeshDrawCommandCullingPayloadFlags CullingPayloadFlags = GetCullingPayloadFlags(Result.bIsLODRange, Idx == 0, Idx == NumMDCs - 1);
+			const EMeshDrawCommandCullingPayloadFlags CullingPayloadFlags = GetCullingPayloadFlags(Result.bIsLodRange, CmdInfo.CullingPayload.LodIndex == Result.LodRangeMin, CmdInfo.CullingPayload.LodIndex == Result.LodRangeMax);
 
 			FVisibleMeshDrawCommand& VisibleCmd = ShadowDepthPassVisibleCommands[ShadowDepthPassVisibleCommands.AddUninitialized()];
 			VisibleCmd.Setup(
@@ -2256,8 +2260,7 @@ void FProjectedShadowInfo::FinalizeAddSubjectPrimitive(
 
 			NumSubjectMeshCommandBuildRequestElements += MeshRelevance.NumElements;
 			SubjectMeshCommandBuildRequests.Add(&MeshBatch);
-			// Assumption: We have one mesh batch per LOD in order of increasing LOD index.
-			SubjectMeshCommandBuildFlags.Add(GetCullingPayloadFlags(Result.bIsLODRange, Idx == 0, Idx == NumMeshes - 1));
+			SubjectMeshCommandBuildFlags.Add(GetCullingPayloadFlags(Result.bIsLodRange, MeshRelevance.LODIndex == Result.LodRangeMin, MeshRelevance.LODIndex == Result.LodRangeMax));
 		}
 	}
 
