HeterogeneousVolumes VoxelGrid fixes

BuildMajorantVoxelGrid dispatched with unneccessarily large GridCount.

diff --git a/Engine/Source/Runtime/Renderer/Private/HeterogeneousVolumes/HeterogeneousVolumesVoxelGridPipeline.cpp b/Engine/Source/Runtime/Renderer/Private/HeterogeneousVolumes/HeterogeneousVolumesVoxelGridPipeline.cpp
--- a/Engine/Source/Runtime/Renderer/Private/HeterogeneousVolumes/HeterogeneousVolumesVoxelGridPipeline.cpp
+++ b/Engine/Source/Runtime/Renderer/Private/HeterogeneousVolumes/HeterogeneousVolumesVoxelGridPipeline.cpp
@@ -1026,13 +1026,18 @@ void BuildMajorantVoxelGrid(
 			PermutationVector.Set<FBuildMajorantVoxelGridCS::FEnableIndirectionGrid>(HeterogeneousVolumes::EnableIndirectionGrid());
 			TShaderRef<FBuildMajorantVoxelGridCS> ComputeShader = GlobalShaderMap->GetShader<FBuildMajorantVoxelGridCS>(PermutationVector);
 
+			FIntVector GroupCount = FIntVector(
+				FMath::DivideAndRoundUp(TopLevelGridResolution.X, FBuildMajorantVoxelGridCS::GetThreadGroupSize3D()),
+				FMath::DivideAndRoundUp(TopLevelGridResolution.Y, FBuildMajorantVoxelGridCS::GetThreadGroupSize3D()),
+				FMath::DivideAndRoundUp(TopLevelGridResolution.Z, FBuildMajorantVoxelGridCS::GetThreadGroupSize3D())
+			);
 			FComputeShaderUtils::AddPass(
 				GraphBuilder,
-				RDG_EVENT_NAME("RenderTransmittanceTopLevelGridCS"),
-				ERDGPassFlags::Compute | ERDGPassFlags::NeverCull,
+				RDG_EVENT_NAME("BuildMajorantVoxelGridCS"),
+				ERDGPassFlags::Compute,
 				ComputeShader,
 				PassParameters,
-				TopLevelGridResolution
+				GroupCount
 			);
 		}
 
@@ -2418,7 +2423,7 @@ void BuildOrthoVoxelGrid(
 	TRDGUniformBufferRef<FOrthoVoxelGridUniformBufferParameters>& OrthoGridUniformBuffer
 )
 {
-	if (!ShouldRenderHeterogeneousVolumes(Scene) || !HeterogeneousVolumes::EnableOrthoVoxelGrid())
+	if (!ShouldRenderHeterogeneousVolumes(Scene) || !ShouldRenderHeterogeneousVolumesForAnyView(Views) || !HeterogeneousVolumes::EnableOrthoVoxelGrid())
 	{
 		CreateEmptyOrthoVoxelGridUniformBuffer(GraphBuilder, OrthoGridUniformBuffer);
 		return;
