Add ECVF_Preview flag to r.InstanceCulling.OcclusionQueries CVar to allow it to be automatically set in Edito platform preview mode

diff --git a/Engine/Source/Runtime/Renderer/Private/InstanceCulling/InstanceCullingOcclusionQuery.cpp b/Engine/Source/Runtime/Renderer/Private/InstanceCulling/InstanceCullingOcclusionQuery.cpp
--- a/Engine/Source/Runtime/Renderer/Private/InstanceCulling/InstanceCullingOcclusionQuery.cpp
+++ b/Engine/Source/Runtime/Renderer/Private/InstanceCulling/InstanceCullingOcclusionQuery.cpp
@@ -26,7 +26,7 @@ static TAutoConsoleVariable<int32> CVarInstanceCullingOcclusionQueries(
 	TEXT("r.InstanceCulling.OcclusionQueries"),
 	0,
 	TEXT("EXPERIMENTAL: Use per-instance software occlusion queries to perform less conservative visibility test than what's possible with HZB alone"),
-	ECVF_RenderThreadSafe);
+	ECVF_RenderThreadSafe | ECVF_Preview);
 
 struct FInstanceCullingOcclusionQueryDeferredContext;
 
