Prevent ~7 mb global distance fields related allocations when global DFs are never used

diff --git a/Engine/Source/Runtime/Renderer/Private/DistanceFieldAmbientOcclusion.cpp b/Engine/Source/Runtime/Renderer/Private/DistanceFieldAmbientOcclusion.cpp
--- a/Engine/Source/Runtime/Renderer/Private/DistanceFieldAmbientOcclusion.cpp
+++ b/Engine/Source/Runtime/Renderer/Private/DistanceFieldAmbientOcclusion.cpp
@@ -125,6 +125,16 @@ FAutoConsoleVariableRef CVarAOObjectDistanceField(
 	ECVF_RenderThreadSafe
 );
 
+// TODO: Find a proper way to check bUsesGlobalDistanceField material relevance and remove this cvar. Currently,
+// PrepareDistanceFieldScene is called before InitViews where we check material relevances. Maybe use last frame result?
+int32 GAOGlobalDistanceFieldDetailedNecessityCheck = 1;
+FAutoConsoleVariableRef CVarAOGlobalDistanceFieldDetailedCheck(
+	TEXT("r.AOGlobalDistanceField.DetailedNecessityCheck"),
+	GAOGlobalDistanceFieldDetailedNecessityCheck,
+	TEXT("Whether to perform detailed necessity check in FSceneRenderer::ShouldPrepareGlobalDistanceField()."),
+	ECVF_RenderThreadSafe
+);
+
 bool UseDistanceFieldAO()
 {
 	return GDistanceFieldAO && GDistanceFieldAOQuality >= 1;
@@ -751,7 +761,10 @@ bool FSceneRenderer::ShouldPrepareGlobalDistanceField() const
 		return false;
 	}
 
-	bool bShouldPrepareForAO = SupportsDistanceFieldAO(Scene->GetFeatureLevel(), Scene->GetShaderPlatform());
+	bool bShouldPrepareForAO = SupportsDistanceFieldAO(Scene->GetFeatureLevel(), Scene->GetShaderPlatform())
+		&& (!GAOGlobalDistanceFieldDetailedNecessityCheck
+			|| ShouldPrepareForDistanceFieldAO()
+			|| ((FXSystem != nullptr) && FXSystem->UsesGlobalDistanceField()));
 
 	bShouldPrepareForAO = bShouldPrepareForAO || (IsLumenEnabled(Views[0]) && Lumen::UseGlobalSDFObjectGrid(*Views[0].Family));
 
