Fixed: tile rendering on platforms that use "uniform view" GPUScene path

diff --git a/Engine/Source/Runtime/Renderer/Private/Renderer.cpp b/Engine/Source/Runtime/Renderer/Private/Renderer.cpp
--- a/Engine/Source/Runtime/Renderer/Private/Renderer.cpp
+++ b/Engine/Source/Runtime/Renderer/Private/Renderer.cpp
@@ -218,7 +218,7 @@ TRDGUniformBufferRef<FBatchedPrimitiveParameters> FRendererModule::CreateSingleP
 
 	FBatchedPrimitiveParameters* BatchedPrimitiveParameters = GraphBuilder.AllocParameters<FBatchedPrimitiveParameters>();
 
-	FRDGBufferDesc PrimitiveDataBufferDesc = FRDGBufferDesc::CreateStructuredDesc(16u, (PLATFORM_MAX_UNIFORM_BUFFER_RANGE / 16u));
+	FRDGBufferDesc PrimitiveDataBufferDesc = FRDGBufferDesc::CreateStructuredUploadDesc(16u, (PLATFORM_MAX_UNIFORM_BUFFER_RANGE / 16u));
 	PrimitiveDataBufferDesc.Usage |= EBufferUsageFlags::UniformBuffer;
 	FRDGBufferRef PrimitiveDataBuffer = nullptr;
 	
