Don't need to move ownership of Segments array to RHI thread since FRHIComputeCommandList::BuildAccelerationStructures(...) copies all the data to inline memory.

diff --git a/Engine/Source/Runtime/Renderer/Private/RayTracing/RayTracingDynamicGeometry.cpp b/Engine/Source/Runtime/Renderer/Private/RayTracing/RayTracingDynamicGeometry.cpp
--- a/Engine/Source/Runtime/Renderer/Private/RayTracing/RayTracingDynamicGeometry.cpp
+++ b/Engine/Source/Runtime/Renderer/Private/RayTracing/RayTracingDynamicGeometry.cpp
@@ -372,9 +372,9 @@ void FRayTracingDynamicGeometryCollection::AddDynamicMeshBatchForGeometryUpdate(
 
 	if (bUseSharedVertexBuffer)
 	{
-		// Make render thread side temporary copy and move to rhi side allocation when command list is known
-		// Cache the count of segments so final views can be made when all segments are collected (Segments array could still be reallocated)
 		Segments.Append(Geometry.Initializer.Segments);
+
+		// Cache the count of segments so final views can be made when all segments are collected (Segments array could still be reallocated)
 		Params.Segments = MakeArrayView((FRayTracingGeometrySegment*)nullptr, Geometry.Initializer.Segments.Num());
 	}
 
@@ -581,9 +581,6 @@ void FRayTracingDynamicGeometryCollection::DispatchUpdates(FRHICommandListImmedi
 void FRayTracingDynamicGeometryCollection::EndUpdate(FRHICommandListImmediate& RHICmdList)
 {
 	ReferencedUniformBuffers.Empty(ReferencedUniformBuffers.Max());
-
-	// Move ownership to RHI thread for another frame
-	RHICmdList.EnqueueLambda([ArrayOwnedByRHIThread = MoveTemp(Segments)](FRHICommandListImmediate&){});
 }
 
 uint32 FRayTracingDynamicGeometryCollection::ComputeScratchBufferSize()
