Prefer TArrayView to TArray.

diff --git a/Engine/Source/Runtime/Renderer/Private/HeterogeneousVolumes/HeterogeneousVolumes.cpp b/Engine/Source/Runtime/Renderer/Private/HeterogeneousVolumes/HeterogeneousVolumes.cpp
--- a/Engine/Source/Runtime/Renderer/Private/HeterogeneousVolumes/HeterogeneousVolumes.cpp
+++ b/Engine/Source/Runtime/Renderer/Private/HeterogeneousVolumes/HeterogeneousVolumes.cpp
@@ -193,13 +193,13 @@ bool ShouldRenderHeterogeneousVolumes(
 }
 
 bool ShouldRenderHeterogeneousVolumesForAnyView(
-	const TArray<FViewInfo>& Views
+	const TArrayView<FViewInfo>& Views
 )
 {
 	bool Result = false;
-	for (int32 ViewIndex = 0; ViewIndex < Views.Num(); ++ViewIndex)
+	for (FViewInfo& View : Views)
 	{
-		Result |= ShouldRenderHeterogeneousVolumesForView(Views[ViewIndex]);
+		Result |= ShouldRenderHeterogeneousVolumesForView(View);
 	}
 	return Result;
 }
diff --git a/Engine/Source/Runtime/Renderer/Private/HeterogeneousVolumes/HeterogeneousVolumes.h b/Engine/Source/Runtime/Renderer/Private/HeterogeneousVolumes/HeterogeneousVolumes.h
--- a/Engine/Source/Runtime/Renderer/Private/HeterogeneousVolumes/HeterogeneousVolumes.h
+++ b/Engine/Source/Runtime/Renderer/Private/HeterogeneousVolumes/HeterogeneousVolumes.h
@@ -27,7 +27,7 @@ struct FSceneTextures;
 //
 
 bool ShouldRenderHeterogeneousVolumes(const FScene* Scene);
-bool ShouldRenderHeterogeneousVolumesForAnyView(const TArray<FViewInfo>& Views);
+bool ShouldRenderHeterogeneousVolumesForAnyView(const TArrayView<FViewInfo>& Views);
 bool ShouldRenderHeterogeneousVolumesForView(const FViewInfo& View);
 bool DoesPlatformSupportHeterogeneousVolumes(EShaderPlatform Platform);
 bool DoesMaterialShaderSupportHeterogeneousVolumes(const FMaterialShaderParameters& Parameters);
