Fix a bug that primitives use a light channel will lose sky lighting on mobile deferred.

diff --git a/Engine/Source/Runtime/Renderer/Private/MobileDeferredShadingPass.cpp b/Engine/Source/Runtime/Renderer/Private/MobileDeferredShadingPass.cpp
--- a/Engine/Source/Runtime/Renderer/Private/MobileDeferredShadingPass.cpp
+++ b/Engine/Source/Runtime/Renderer/Private/MobileDeferredShadingPass.cpp
@@ -576,8 +576,8 @@ static void RenderDirectionalLights(FRHICommandList& RHICmdList, const FScene& S
 	{
 		NumLights += (Scene.MobileDirectionalLights[ChannelIdx] ? 1 : 0);
 	}
-	// We can merge reflection and skylight pass with a sole directional light pass
-	const bool bInlineReflectionAndSky = (NumLights == 1);
+	// We can merge reflection and skylight pass with a sole directional light pass and if all primitives and the directional light use the default lighting channel
+	const bool bInlineReflectionAndSky = (NumLights == 1) && !View.bUsesLightingChannels && (Scene.MobileDirectionalLights[0] != nullptr);
 
 	for (uint32 ChannelIdx = 0; ChannelIdx < UE_ARRAY_COUNT(Scene.MobileDirectionalLights); ChannelIdx++)
 	{
