Fix a crash bug of missing MobileSceneTexture uniform buffer when rendering custom depth with TAA enabled on mobile.

diff --git a/Engine/Source/Runtime/Renderer/Private/VelocityRendering.cpp b/Engine/Source/Runtime/Renderer/Private/VelocityRendering.cpp
--- a/Engine/Source/Runtime/Renderer/Private/VelocityRendering.cpp
+++ b/Engine/Source/Runtime/Renderer/Private/VelocityRendering.cpp
@@ -239,7 +239,7 @@ bool FMobileSceneRenderer::ShouldRenderVelocities() const
 
 BEGIN_SHADER_PARAMETER_STRUCT(FVelocityPassParameters, )
 	SHADER_PARAMETER_STRUCT_INCLUDE(FViewShaderParameters, View)
-	SHADER_PARAMETER_RDG_UNIFORM_BUFFER(FSceneTextureUniformParameters, SceneTextures)
+	SHADER_PARAMETER_STRUCT_INCLUDE(FSceneTextureShaderParameters, SceneTextures)
 	SHADER_PARAMETER_STRUCT_INCLUDE(FInstanceCullingDrawParams, InstanceCullingDrawParams)
 	RENDER_TARGET_BINDING_SLOTS()
 END_SHADER_PARAMETER_STRUCT()
@@ -313,7 +313,7 @@ void FSceneRenderer::RenderVelocities(
 			FVelocityPassParameters* PassParameters = GraphBuilder.AllocParameters<FVelocityPassParameters>();
 			PassParameters->View = View.GetShaderParameters();
 			ParallelMeshPass.BuildRenderingCommands(GraphBuilder, Scene->GPUScene, PassParameters->InstanceCullingDrawParams);
-			PassParameters->SceneTextures = SceneTextures.UniformBuffer;
+			PassParameters->SceneTextures = SceneTextures.GetSceneTextureShaderParameters(View.FeatureLevel);
 			PassParameters->RenderTargets.DepthStencil = FDepthStencilBinding(
 				SceneTextures.Depth.Resolve,
 				ERenderTargetLoadAction::ELoad,
