Allow scene capture to be added as additional views of the main renderer and render depth capture as additional render prepass before the main depth prepass

*) Added "Render in Main Renderer" toggle to SceneCaptureComponent2D, only enabled on scene depth & device depth capture modes
*) Added "r.SceneCapture.AllowRenderInMainRenderer" cvar to allow scene capture to be rendered in the main renderer if the above toggle is on

diff --git a/Engine/Source/Runtime/Engine/Classes/Components/SceneCaptureComponent2D.h b/Engine/Source/Runtime/Engine/Classes/Components/SceneCaptureComponent2D.h
--- a/Engine/Source/Runtime/Engine/Classes/Components/SceneCaptureComponent2D.h
+++ b/Engine/Source/Runtime/Engine/Classes/Components/SceneCaptureComponent2D.h
@@ -98,6 +98,10 @@ public:
 	UPROPERTY(EditAnywhere, BlueprintReadWrite, AdvancedDisplay, Category=SceneCapture)
 	FVector ClipPlaneNormal;
 	
+	/** Render scene capture as additional render passes of the main renderer rather than as an independent renderer. Can only apply to scene depth and device depth modes. */
+	UPROPERTY(EditAnywhere, BlueprintReadWrite, AdvancedDisplay, Category = SceneCapture, meta = (EditCondition = "CaptureSource == ESceneCaptureSource::SCS_SceneDepth || CaptureSource == ESceneCaptureSource::SCS_DeviceDepth"))
+	bool bRenderInMainRenderer = false;
+
 	/** 
 	 * True if we did a camera cut this frame. Automatically reset to false at every capture.
 	 * This flag affects various things in the renderer (such as whether to use the occlusion queries from last frame, and motion blur).
diff --git a/Engine/Source/Runtime/Engine/Private/SceneView.cpp b/Engine/Source/Runtime/Engine/Private/SceneView.cpp
--- a/Engine/Source/Runtime/Engine/Private/SceneView.cpp
+++ b/Engine/Source/Runtime/Engine/Private/SceneView.cpp
@@ -787,6 +787,8 @@ FSceneView::FSceneView(const FSceneViewInitOptions& InitOptions)
 	, CursorPos(InitOptions.CursorPos)
 	, bIsGameView(false)
 	, bIsViewInfo(false)
+	, SceneCaptureRenderTarget(InitOptions.SceneCaptureRenderTarget)
+	, SceneCaptureSource(InitOptions.SceneCaptureSource)
 	, bIsSceneCapture(InitOptions.bIsSceneCapture)
 	, bIsSceneCaptureCube(InitOptions.bIsSceneCaptureCube)
 	, bSceneCaptureUsesRayTracing(InitOptions.bSceneCaptureUsesRayTracing)
diff --git a/Engine/Source/Runtime/Engine/Public/SceneInterface.h b/Engine/Source/Runtime/Engine/Public/SceneInterface.h
--- a/Engine/Source/Runtime/Engine/Public/SceneInterface.h
+++ b/Engine/Source/Runtime/Engine/Public/SceneInterface.h
@@ -7,6 +7,8 @@
 #include "Math/MathFwd.h"
 #include "Misc/EnumClassFlags.h"
 #include "RenderGraphFwd.h"
+#include "Engine/EngineTypes.h"
+#include "PrimitiveComponentId.h"
 
 #if UE_ENABLE_INCLUDE_ORDER_DEPRECATED_IN_5_2
 #include "CoreMinimal.h"
@@ -26,7 +28,6 @@ class FMaterial;
 class FMaterialShaderMap;
 class FName;
 class FOutputDevice;
-class FPrimitiveComponentId;
 class FPrimitiveSceneInfo;
 class FRDGBuilder;
 class FRDGExternalAccessQueue;
@@ -681,6 +682,20 @@ public:
 
 	virtual void RefreshNaniteRasterBins(FPrimitiveSceneInfo& PrimitiveSceneInfo) { }
 
+	/** Scene capture info is used to collect scene capture settings and passing to scene views during the renderer construction. */
+	struct FSceneCaptureInfo
+	{
+		FRenderTarget* RenderTarget = nullptr;
+		FVector ViewLocation;
+		FMatrix ViewRotationMatrix;
+		FMatrix ProjectionMatrix;
+		TSet<FPrimitiveComponentId> HiddenPrimitives;
+		TOptional<TSet<FPrimitiveComponentId>> ShowOnlyPrimitives;
+		ESceneCaptureSource SceneCaptureSource;
+		const AActor* ViewActor = nullptr;
+	};
+	TArray<FSceneCaptureInfo> SceneCaptureInfos;
+
 protected:
 	virtual ~FSceneInterface() {}
 
diff --git a/Engine/Source/Runtime/Engine/Public/SceneView.h b/Engine/Source/Runtime/Engine/Public/SceneView.h
--- a/Engine/Source/Runtime/Engine/Public/SceneView.h
+++ b/Engine/Source/Runtime/Engine/Public/SceneView.h
@@ -189,6 +189,12 @@ struct FSceneViewInitOptions : public FSceneViewProjectionData
 	/** In case of ortho, generate a fake view position that has a non-zero W component. The view position will be derived based on the view matrix. */
 	bool bUseFauxOrthoViewPos;
 
+	/** The render target of scene capture associated with this view. If used it takes priority over view family RT. */
+	FRenderTarget* SceneCaptureRenderTarget = nullptr;
+
+	/** Which component of the scene rendering this view should output for scene capture. If used it takes priority over view family SceneCaptureSource. */
+	ESceneCaptureSource SceneCaptureSource;
+
 	/** Whether this view is being used to render a scene capture. */
 	bool bIsSceneCapture;
 
@@ -1378,6 +1384,12 @@ public:
 	/** For sanity checking casts that are assumed to be safe. */
 	bool bIsViewInfo;
 
+	/** The render target of scene capture associated with this view. If used it takes priority over view family RT. */
+	FRenderTarget* SceneCaptureRenderTarget = nullptr;
+
+	/** Which component of the scene rendering this view should output for scene capture. If used it takes priority over view family SceneCaptureSource. */
+	ESceneCaptureSource SceneCaptureSource;
+
 	/** Whether this view is being used to render a scene capture. */
 	bool bIsSceneCapture;
 
@@ -1980,6 +1992,9 @@ public:
 	/** The views which make up the family. */
 	TArray<const FSceneView*> Views;
 
+	/** All views include main camera views and scene capture views. */
+	TArray<const FSceneView*> AllViews;
+
 	/** View mode of the family. */
 	EViewModeIndex ViewMode;
 
diff --git a/Engine/Source/Runtime/Landscape/Private/LandscapeRender.cpp b/Engine/Source/Runtime/Landscape/Private/LandscapeRender.cpp
--- a/Engine/Source/Runtime/Landscape/Private/LandscapeRender.cpp
+++ b/Engine/Source/Runtime/Landscape/Private/LandscapeRender.cpp
@@ -908,7 +908,7 @@ void FLandscapeSceneViewExtension::PreRenderView_RenderThread(FRDGBuilder& Graph
 #endif	// RHI_RAYTRACING
 
 	// Kick the job once all views have been collected.
-	if (!LandscapeRenderSystems.IsEmpty() && LandscapeViews.Num() == InView.Family->Views.Num())
+	if (!LandscapeRenderSystems.IsEmpty() && LandscapeViews.Num() == InView.Family->AllViews.Num())
 	{
 		LandscapeSetupTask = GraphBuilder.AddCommandListSetupTask([this] (FRHICommandListBase& RHICmdList)
 		{
diff --git a/Engine/Source/Runtime/Renderer/Private/DeferredShadingRenderer.cpp b/Engine/Source/Runtime/Renderer/Private/DeferredShadingRenderer.cpp
--- a/Engine/Source/Runtime/Renderer/Private/DeferredShadingRenderer.cpp
+++ b/Engine/Source/Runtime/Renderer/Private/DeferredShadingRenderer.cpp
@@ -510,7 +510,7 @@ FDeferredShadingSceneRenderer::FDeferredShadingSceneRenderer(const FSceneViewFam
 	, SceneCullingRenderer(*Scene->SceneCulling)
 	, bAreLightsInLightGrid(false)
 {
-	ViewPipelineStates.SetNum(Views.Num());
+	ViewPipelineStates.SetNum(AllViews.Num());
 
 	ShadowSceneRenderer = MakeUnique<FShadowSceneRenderer>(*this);
 }
@@ -2455,6 +2455,222 @@ bool FDeferredShadingSceneRenderer::IsNaniteEnabled() const
 	return UseNanite(ShaderPlatform) && ViewFamily.EngineShowFlags.NaniteMeshes && Nanite::GStreamingManager.HasResourceEntries();
 }
 
+void FDeferredShadingSceneRenderer::RenderNanite(FRDGBuilder& GraphBuilder, const TArray<FViewInfo>& InViews, FSceneTextures& SceneTextures, bool bIsEarlyDepthComplete,
+	FNaniteBasePassVisibility& InNaniteBasePassVisibility,
+	TArray<Nanite::FRasterResults, TInlineAllocator<2>>& NaniteRasterResults,
+	TArray<Nanite::FPackedView, SceneRenderingAllocator>& PrimaryNaniteViews)
+{
+	LLM_SCOPE_BYTAG(Nanite);
+	TRACE_CPUPROFILER_EVENT_SCOPE(InitNaniteRaster);
+
+	NaniteRasterResults.AddDefaulted(InViews.Num());
+	if (InNaniteBasePassVisibility.Query != nullptr)
+	{
+		InNaniteBasePassVisibility.Visibility->FinishVisibilityQuery(InNaniteBasePassVisibility.Query, NaniteRasterResults[0].VisibilityResults);
+
+		// For now we'll share the same visibility results across all views
+		for (int32 ViewIndex = 1; ViewIndex < NaniteRasterResults.Num(); ++ViewIndex)
+		{
+			NaniteRasterResults[ViewIndex].VisibilityResults = NaniteRasterResults[0].VisibilityResults;
+		}
+
+		uint32 TotalRasterBins = 0;
+		uint32 VisibleRasterBins = 0;
+		NaniteRasterResults[0].VisibilityResults.GetRasterBinStats(VisibleRasterBins, TotalRasterBins);
+
+		uint32 TotalShadingDraws = 0;
+		uint32 VisibleShadingDraws = 0;
+		NaniteRasterResults[0].VisibilityResults.GetShadingDrawStats(VisibleShadingDraws, TotalShadingDraws);
+
+		SET_DWORD_STAT(STAT_NaniteBasePassTotalRasterBins, TotalRasterBins);
+		SET_DWORD_STAT(STAT_NaniteBasePassTotalShadingDraws, TotalShadingDraws);
+
+		SET_DWORD_STAT(STAT_NaniteBasePassVisibleRasterBins, VisibleRasterBins);
+		SET_DWORD_STAT(STAT_NaniteBasePassVisibleShadingDraws, VisibleShadingDraws);
+	}
+
+	const FIntPoint RasterTextureSize = SceneTextures.Depth.Target->Desc.Extent;
+
+	// Primary raster view
+	{
+		Nanite::FSharedContext SharedContext{};
+		SharedContext.FeatureLevel = Scene->GetFeatureLevel();
+		SharedContext.ShaderMap = GetGlobalShaderMap(SharedContext.FeatureLevel);
+		SharedContext.Pipeline = Nanite::EPipeline::Primary;
+
+		FIntRect RasterTextureRect(0, 0, RasterTextureSize.X, RasterTextureSize.Y);
+		if (InViews.Num() == 1)
+		{
+			const FViewInfo& View = InViews[0];
+			if (View.ViewRect.Min.X == 0 && View.ViewRect.Min.Y == 0)
+			{
+				RasterTextureRect = View.ViewRect;
+			}
+		}
+
+		Nanite::FRasterContext RasterContext;
+
+		// Nanite::VisBuffer (Visibility Buffer Clear)
+		{
+			RDG_GPU_STAT_SCOPE(GraphBuilder, NaniteVisBuffer);
+			RasterContext = Nanite::InitRasterContext(
+				GraphBuilder,
+				SharedContext,
+				ViewFamily,
+				RasterTextureSize,
+				RasterTextureRect,
+				ViewFamily.EngineShowFlags.VisualizeNanite
+			);
+		}
+
+		Nanite::FConfiguration CullingConfig = { 0 };
+		CullingConfig.bTwoPassOcclusion = true;
+		CullingConfig.bUpdateStreaming = true;
+		CullingConfig.bPrimaryContext = true;
+
+		static FString EmptyFilterName = TEXT(""); // Empty filter represents primary view.
+		CullingConfig.bExtractStats = Nanite::IsStatFilterActive(EmptyFilterName);
+
+		const bool bDrawSceneViewsInOneNanitePass = InViews.Num() > 1 && Nanite::ShouldDrawSceneViewsInOneNanitePass(InViews[0]);
+
+		// creates one or more Nanite views (normally one per view unless drawing multiple views together - e.g. Stereo ISR views)
+		auto CreateNaniteViews = [bDrawSceneViewsInOneNanitePass, &InViews, &PrimaryNaniteViews, &GraphBuilder](const FViewInfo& View, int32 ViewIndex, const FIntPoint& RasterTextureSize, float MaxPixelsPerEdgeMultipler) -> Nanite::FPackedViewArray*
+		{
+			Nanite::FPackedViewArray::ArrayType OutViews;
+
+			// always add the primary view. In case of bDrawSceneViewsInOneNanitePass HZB is built from all views so using viewrects
+			// to account for a rare case when the primary view doesn't start from 0, 0 (maybe can happen in splitscreen?)
+			FIntRect HZBTestRect = bDrawSceneViewsInOneNanitePass ?
+				View.PrevViewInfo.ViewRect :
+				FIntRect(0, 0, View.PrevViewInfo.ViewRect.Width(), View.PrevViewInfo.ViewRect.Height());
+
+			Nanite::FPackedView PackedView = Nanite::CreatePackedViewFromViewInfo(
+				View,
+				RasterTextureSize,
+				NANITE_VIEW_FLAG_HZBTEST | NANITE_VIEW_FLAG_NEAR_CLIP,
+				/* StreamingPriorityCategory = */ 3,
+				/* MinBoundsRadius = */ 0.0f,
+				MaxPixelsPerEdgeMultipler,
+				&HZBTestRect
+			);
+			OutViews.Add(PackedView);
+			PrimaryNaniteViews.Add(PackedView);
+
+			if (bDrawSceneViewsInOneNanitePass)
+			{
+				// All other views in the family will need to be rendered in one go, to cover both ISR and (later) split-screen
+				for (int32 ViewIdx = 1, NumViews = InViews.Num(); ViewIdx < NumViews; ++ViewIdx)
+				{
+					const FViewInfo& SecondaryViewInfo = InViews[ViewIdx];
+
+					/* viewport rect in HZB space. For instanced stereo passes HZB is built for all atlased views */
+					FIntRect SecondaryHZBTestRect = SecondaryViewInfo.PrevViewInfo.ViewRect;
+					Nanite::FPackedView SecondaryPackedView = Nanite::CreatePackedViewFromViewInfo(
+						SecondaryViewInfo,
+						RasterTextureSize,
+						NANITE_VIEW_FLAG_HZBTEST | NANITE_VIEW_FLAG_NEAR_CLIP,
+						/* StreamingPriorityCategory = */ 3,
+						/* MinBoundsRadius = */ 0.0f,
+						MaxPixelsPerEdgeMultipler,
+						&SecondaryHZBTestRect
+					);
+
+					OutViews.Add(SecondaryPackedView);
+				}
+			}
+
+			return Nanite::FPackedViewArray::Create(GraphBuilder, OutViews.Num(), 1, MoveTemp(OutViews));
+		};
+
+		// in case of bDrawSceneViewsInOneNanitePass we only need one iteration
+		uint32 ViewsToRender = (bDrawSceneViewsInOneNanitePass ? 1u : (uint32)InViews.Num());
+		for (uint32 ViewIndex = 0; ViewIndex < ViewsToRender; ++ViewIndex)
+		{
+			Nanite::FRasterResults& RasterResults = NaniteRasterResults[ViewIndex];
+			const FViewInfo& View = InViews[ViewIndex];
+			// We don't check View.ShouldRenderView() since this is already taken care of by bDrawSceneViewsInOneNanitePass.
+			// If bDrawSceneViewsInOneNanitePass is false, we need to render the secondary view even if ShouldRenderView() is false
+
+			RDG_EVENT_SCOPE_CONDITIONAL(GraphBuilder, InViews.Num() > 1 && !bDrawSceneViewsInOneNanitePass, "View%u", ViewIndex);
+			RDG_EVENT_SCOPE_CONDITIONAL(GraphBuilder, InViews.Num() > 1 && bDrawSceneViewsInOneNanitePass, "View%u (together with %d more)", ViewIndex, InViews.Num() - 1);
+
+			FIntRect ViewRect = bDrawSceneViewsInOneNanitePass ? FIntRect(0, 0, FamilySize.X, FamilySize.Y) : View.ViewRect;
+			CullingConfig.SetViewFlags(View);
+
+			float LODScaleFactor = 1.0f;
+			if (View.PrimaryScreenPercentageMethod == EPrimaryScreenPercentageMethod::TemporalUpscale &&
+				CVarNaniteViewMeshLODBiasEnable.GetValueOnRenderThread() != 0)
+			{
+				float TemporalUpscaleFactor = float(View.GetSecondaryViewRectSize().X) / float(ViewRect.Width());
+
+				LODScaleFactor = TemporalUpscaleFactor * FMath::Exp2(-CVarNaniteViewMeshLODBiasOffset.GetValueOnRenderThread());
+				LODScaleFactor = FMath::Min(LODScaleFactor, FMath::Exp2(-CVarNaniteViewMeshLODBiasMin.GetValueOnRenderThread()));
+			}
+
+			float MaxPixelsPerEdgeMultipler = 1.0f / LODScaleFactor;
+			if (GDynamicNaniteScalingPrimary.GetSettings().IsEnabled())
+			{
+				MaxPixelsPerEdgeMultipler *= 1.0f / DynamicResolutionFractions[GDynamicNaniteScalingPrimary];
+			}
+
+			Nanite::FPackedViewArray* NaniteViewsToRender = CreateNaniteViews(View, ViewIndex, RasterTextureSize, MaxPixelsPerEdgeMultipler);
+
+			TUniquePtr< Nanite::IRenderer > NaniteRenderer;
+
+			// Nanite::VisBuffer (Culling and Rasterization)
+			{
+				DynamicRenderScaling::FRDGScope DynamicScalingScope(GraphBuilder, GDynamicNaniteScalingPrimary);
+
+				RDG_GPU_STAT_SCOPE(GraphBuilder, NaniteVisBuffer);
+				RDG_EVENT_SCOPE(GraphBuilder, "Nanite::VisBuffer");
+
+				NaniteRenderer = Nanite::IRenderer::Create(
+					GraphBuilder,
+					*Scene,
+					View,
+					GetSceneUniforms(),
+					SharedContext,
+					RasterContext,
+					CullingConfig,
+					ViewRect,
+					!bIsEarlyDepthComplete ? View.PrevViewInfo.NaniteHZB : View.PrevViewInfo.HZB
+				);
+
+				FSceneInstanceCullingQuery *SceneInstanceCullQuery = SceneCullingRenderer.CullInstances(GraphBuilder, View.ViewFrustum);
+				NaniteRenderer->DrawGeometry(
+					Scene->NaniteRasterPipelines[ENaniteMeshPass::BasePass],
+					RasterResults.VisibilityResults,
+					*NaniteViewsToRender,
+					SceneInstanceCullQuery
+				);
+
+				NaniteRenderer->ExtractResults( RasterResults );
+			}
+
+			// Nanite::BasePass (Depth Pre-Pass and HZB Build)
+			{
+				RDG_GPU_STAT_SCOPE(GraphBuilder, NaniteBasePass);
+
+				// Emit velocity with depth if not writing it in base pass.
+				FRDGTexture* VelocityBuffer = !IsUsingBasePassVelocity(ShaderPlatform) ? SceneTextures.Velocity : nullptr;
+
+				Nanite::EmitDepthTargets(
+					GraphBuilder,
+					*Scene,
+					InViews[ViewIndex],
+					bDrawSceneViewsInOneNanitePass,
+					RasterResults,
+					SceneTextures.Depth.Target,
+					VelocityBuffer
+				);
+				
+				// Sanity check (always force Z prepass)
+				check(bIsEarlyDepthComplete);
+			}
+		}
+	}
+}
+
 void FDeferredShadingSceneRenderer::Render(FRDGBuilder& GraphBuilder)
 {
 	if (!ViewFamily.EngineShowFlags.Rendering)
@@ -2873,9 +3089,9 @@ void FDeferredShadingSceneRenderer::Render(FRDGBuilder& GraphBuilder)
 
 		Scene->GPUScene.Update(GraphBuilder, GetSceneUniforms(), *Scene, ExternalAccessQueue);
 
-		for (int32 ViewIndex = 0; ViewIndex < Views.Num(); ViewIndex++)
+		for (int32 ViewIndex = 0; ViewIndex < AllViews.Num(); ViewIndex++)
 		{
-			FViewInfo& View = Views[ViewIndex];
+			FViewInfo& View = *AllViews[ViewIndex];
 			RDG_GPU_MASK_SCOPE(GraphBuilder, View.GPUMask);
 
 			Scene->GPUScene.UploadDynamicPrimitiveShaderDataForView(GraphBuilder, *Scene, View, ExternalAccessQueue);
@@ -3037,274 +3253,83 @@ void FDeferredShadingSceneRenderer::Render(FRDGBuilder& GraphBuilder)
 	const bool bBasePassCanOutputVelocity = FVelocityRendering::BasePassCanOutputVelocity(Platform);
 	const bool bHairStrandsEnable = HairStrandsBookmarkParameters.HasInstances() && Views.Num() > 0 && IsHairStrandsEnabled(EHairStrandsShaderType::Strands, Platform);
 
-	FRDGTextureRef FirstStageDepthBuffer = nullptr;
+	auto RenderPrepassAndVelocity = [&](auto& InViews, auto& InNaniteBasePassVisibility, auto& NaniteRasterResults, auto& PrimaryNaniteViews)
 	{
-		GraphBuilder.SetCommandListStat(GET_STATID(STAT_CLM_PrePass));
-
-		// Both compute approaches run earlier, so skip clearing stencil here, just load existing.
-		const ERenderTargetLoadAction StencilLoadAction = DepthPass.IsComputeStencilDitherEnabled()
-			? ERenderTargetLoadAction::ELoad
-			: ERenderTargetLoadAction::EClear;
-
-		const ERenderTargetLoadAction DepthLoadAction = ERenderTargetLoadAction::EClear;
-		AddClearDepthStencilPass(GraphBuilder, SceneTextures.Depth.Target, DepthLoadAction, StencilLoadAction);
-
-		// Draw the scene pre-pass / early z pass, populating the scene depth buffer and HiZ
-		if (bNeedsPrePass)
-		{
-			RenderPrePass(GraphBuilder, SceneTextures.Depth.Target, InstanceCullingManager, &FirstStageDepthBuffer);
-		}
-		else
-		{
-			// We didn't do the prepass, but we still want the HMD mask if there is one
-			RenderPrePassHMD(GraphBuilder, SceneTextures.Depth.Target);
-		}
-
-		GraphBuilder.SetCommandListStat(GET_STATID(STAT_CLM_AfterPrePass));
-
-		// special pass for DDM_AllOpaqueNoVelocity, which uses the velocity pass to finish the early depth pass write
-		if (bShouldRenderVelocities && Scene->EarlyZPassMode == DDM_AllOpaqueNoVelocity && RendererOutput == ERendererOutput::FinalSceneColor)
+		FRDGTextureRef FirstStageDepthBuffer = nullptr;
 		{
-			// Render the velocities of movable objects
-			GraphBuilder.SetCommandListStat(GET_STATID(STAT_CLM_Velocity));
-			RenderVelocities(GraphBuilder, SceneTextures, EVelocityPass::Opaque, bHairStrandsEnable);
-			GraphBuilder.SetCommandListStat(GET_STATID(STAT_CLM_AfterVelocity));
-		}
-	}
+			GraphBuilder.SetCommandListStat(GET_STATID(STAT_CLM_PrePass));
 
-	TArray<Nanite::FRasterResults, TInlineAllocator<2>> NaniteRasterResults;
-	TArray<Nanite::FPackedView, SceneRenderingAllocator> PrimaryNaniteViews;
-	{
-		Scene->WaitForCacheNaniteDrawCommandsTask();
+			// Both compute approaches run earlier, so skip clearing stencil here, just load existing.
+			const ERenderTargetLoadAction StencilLoadAction = DepthPass.IsComputeStencilDitherEnabled()
+				? ERenderTargetLoadAction::ELoad
+				: ERenderTargetLoadAction::EClear;
 
-		if (bNaniteEnabled && Views.Num() > 0)
-		{
-			LLM_SCOPE_BYTAG(Nanite);
-			TRACE_CPUPROFILER_EVENT_SCOPE(InitNaniteRaster);
+			const ERenderTargetLoadAction DepthLoadAction = ERenderTargetLoadAction::EClear;
+			AddClearDepthStencilPass(GraphBuilder, SceneTextures.Depth.Target, DepthLoadAction, StencilLoadAction);
 
-			NaniteRasterResults.AddDefaulted(Views.Num());
-			if (NaniteBasePassVisibility.Query != nullptr)
+			// Draw the scene pre-pass / early z pass, populating the scene depth buffer and HiZ
+			if (bNeedsPrePass)
 			{
-				NaniteBasePassVisibility.Visibility->FinishVisibilityQuery(NaniteBasePassVisibility.Query, NaniteRasterResults[0].VisibilityResults);
-
-				// For now we'll share the same visibility results across all views
-				for (int32 ViewIndex = 1; ViewIndex < NaniteRasterResults.Num(); ++ViewIndex)
-				{
-					NaniteRasterResults[ViewIndex].VisibilityResults = NaniteRasterResults[0].VisibilityResults;
-				}
-
-				uint32 TotalRasterBins = 0;
-				uint32 VisibleRasterBins = 0;
-				NaniteRasterResults[0].VisibilityResults.GetRasterBinStats(VisibleRasterBins, TotalRasterBins);
-
-				uint32 TotalShadingDraws = 0;
-				uint32 VisibleShadingDraws = 0;
-				NaniteRasterResults[0].VisibilityResults.GetShadingDrawStats(VisibleShadingDraws, TotalShadingDraws);
-
-				SET_DWORD_STAT(STAT_NaniteBasePassTotalRasterBins, TotalRasterBins);
-				SET_DWORD_STAT(STAT_NaniteBasePassTotalShadingDraws, TotalShadingDraws);
-
-				SET_DWORD_STAT(STAT_NaniteBasePassVisibleRasterBins, VisibleRasterBins);
-				SET_DWORD_STAT(STAT_NaniteBasePassVisibleShadingDraws, VisibleShadingDraws);
+				RenderPrePass(GraphBuilder, InViews, SceneTextures.Depth.Target, InstanceCullingManager, &FirstStageDepthBuffer);
 			}
-
-			const FIntPoint RasterTextureSize = SceneTextures.Depth.Target->Desc.Extent;
-
-			// Primary raster view
+			else
 			{
-				Nanite::FSharedContext SharedContext{};
-				SharedContext.FeatureLevel = Scene->GetFeatureLevel();
-				SharedContext.ShaderMap = GetGlobalShaderMap(SharedContext.FeatureLevel);
-				SharedContext.Pipeline = Nanite::EPipeline::Primary;
-
-				FIntRect RasterTextureRect(0, 0, RasterTextureSize.X, RasterTextureSize.Y);
-				if (Views.Num() == 1)
-				{
-					const FViewInfo& View = Views[0];
-					if (View.ViewRect.Min.X == 0 && View.ViewRect.Min.Y == 0)
-					{
-						RasterTextureRect = View.ViewRect;
-					}
-				}
-
-				Nanite::FRasterContext RasterContext;
-
-				// Nanite::VisBuffer (Visibility Buffer Clear)
-				{
-					RDG_GPU_STAT_SCOPE(GraphBuilder, NaniteVisBuffer);
-					RasterContext = Nanite::InitRasterContext(
-						GraphBuilder,
-						SharedContext,
-						ViewFamily,
-						RasterTextureSize,
-						RasterTextureRect,
-						ViewFamily.EngineShowFlags.VisualizeNanite
-					);
-				}
-
-				Nanite::FConfiguration CullingConfig = { 0 };
-				CullingConfig.bTwoPassOcclusion = true;
-				CullingConfig.bUpdateStreaming = true;
-				CullingConfig.bPrimaryContext = true;
-
-				static FString EmptyFilterName = TEXT(""); // Empty filter represents primary view.
-				CullingConfig.bExtractStats = Nanite::IsStatFilterActive(EmptyFilterName);
-
-				const bool bDrawSceneViewsInOneNanitePass = Views.Num() > 1 && Nanite::ShouldDrawSceneViewsInOneNanitePass(Views[0]);
-
-				// creates one or more Nanite views (normally one per view unless drawing multiple views together - e.g. Stereo ISR views)
-				auto CreateNaniteViews = [bDrawSceneViewsInOneNanitePass, &PrimaryNaniteViews, &GraphBuilder](const FViewInfo& View, int32 ViewIndex, const FIntPoint& RasterTextureSize, float MaxPixelsPerEdgeMultipler) -> Nanite::FPackedViewArray*
-				{
-					Nanite::FPackedViewArray::ArrayType OutViews;
-
-					// always add the primary view. In case of bDrawSceneViewsInOneNanitePass HZB is built from all views so using viewrects
-					// to account for a rare case when the primary view doesn't start from 0, 0 (maybe can happen in splitscreen?)
-					FIntRect HZBTestRect = bDrawSceneViewsInOneNanitePass ?
-						View.PrevViewInfo.ViewRect :
-						FIntRect(0, 0, View.PrevViewInfo.ViewRect.Width(), View.PrevViewInfo.ViewRect.Height());
-
-					Nanite::FPackedView PackedView = Nanite::CreatePackedViewFromViewInfo(
-						View,
-						RasterTextureSize,
-						NANITE_VIEW_FLAG_HZBTEST | NANITE_VIEW_FLAG_NEAR_CLIP,
-						/* StreamingPriorityCategory = */ 3,
-						/* MinBoundsRadius = */ 0.0f,
-						MaxPixelsPerEdgeMultipler,
-						&HZBTestRect
-					);
-					OutViews.Add(PackedView);
-					PrimaryNaniteViews.Add(PackedView);
-
-					if (bDrawSceneViewsInOneNanitePass)
-					{
-						// All other views in the family will need to be rendered in one go, to cover both ISR and (later) split-screen
-						for (int32 ViewIdx = 1, NumViews = View.Family->Views.Num(); ViewIdx < NumViews; ++ViewIdx)
-						{
-							const FViewInfo& SecondaryViewInfo = static_cast<const FViewInfo&>(*View.Family->Views[ViewIdx]);
-
-							/* viewport rect in HZB space. For instanced stereo passes HZB is built for all atlased views */
-							FIntRect SecondaryHZBTestRect = SecondaryViewInfo.PrevViewInfo.ViewRect;
-							Nanite::FPackedView SecondaryPackedView = Nanite::CreatePackedViewFromViewInfo(
-								SecondaryViewInfo,
-								RasterTextureSize,
-								NANITE_VIEW_FLAG_HZBTEST | NANITE_VIEW_FLAG_NEAR_CLIP,
-								/* StreamingPriorityCategory = */ 3,
-								/* MinBoundsRadius = */ 0.0f,
-								MaxPixelsPerEdgeMultipler,
-								&SecondaryHZBTestRect
-							);
-
-							OutViews.Add(SecondaryPackedView);
-						}
-					}
-
-					return Nanite::FPackedViewArray::Create(GraphBuilder, OutViews.Num(), 1, MoveTemp(OutViews));
-				};
-
-				// in case of bDrawSceneViewsInOneNanitePass we only need one iteration
-				uint32 ViewsToRender = (bDrawSceneViewsInOneNanitePass ? 1u : (uint32)Views.Num());
-				for (uint32 ViewIndex = 0; ViewIndex < ViewsToRender; ++ViewIndex)
-				{
-					Nanite::FRasterResults& RasterResults = NaniteRasterResults[ViewIndex];
-					const FViewInfo& View = Views[ViewIndex];
-					// We don't check View.ShouldRenderView() since this is already taken care of by bDrawSceneViewsInOneNanitePass.
-					// If bDrawSceneViewsInOneNanitePass is false, we need to render the secondary view even if ShouldRenderView() is false
-
-					RDG_EVENT_SCOPE_CONDITIONAL(GraphBuilder, Views.Num() > 1 && !bDrawSceneViewsInOneNanitePass, "View%u", ViewIndex);
-					RDG_EVENT_SCOPE_CONDITIONAL(GraphBuilder, Views.Num() > 1 && bDrawSceneViewsInOneNanitePass, "View%u (together with %d more)", ViewIndex, Views.Num() - 1);
-
-					FIntRect ViewRect = bDrawSceneViewsInOneNanitePass ? FIntRect(0, 0, FamilySize.X, FamilySize.Y) : View.ViewRect;
-					CullingConfig.SetViewFlags(View);
-
-					float LODScaleFactor = 1.0f;
-					if (View.PrimaryScreenPercentageMethod == EPrimaryScreenPercentageMethod::TemporalUpscale &&
-						CVarNaniteViewMeshLODBiasEnable.GetValueOnRenderThread() != 0)
-					{
-						float TemporalUpscaleFactor = float(View.GetSecondaryViewRectSize().X) / float(ViewRect.Width());
-
-						LODScaleFactor = TemporalUpscaleFactor * FMath::Exp2(-CVarNaniteViewMeshLODBiasOffset.GetValueOnRenderThread());
-						LODScaleFactor = FMath::Min(LODScaleFactor, FMath::Exp2(-CVarNaniteViewMeshLODBiasMin.GetValueOnRenderThread()));
-					}
-
-					float MaxPixelsPerEdgeMultipler = 1.0f / LODScaleFactor;
-					if (GDynamicNaniteScalingPrimary.GetSettings().IsEnabled())
-					{
-						MaxPixelsPerEdgeMultipler *= 1.0f / DynamicResolutionFractions[GDynamicNaniteScalingPrimary];
-					}
+				// We didn't do the prepass, but we still want the HMD mask if there is one
+				RenderPrePassHMD(GraphBuilder, InViews, SceneTextures.Depth.Target);
+			}
 
-					Nanite::FPackedViewArray* NaniteViewsToRender = CreateNaniteViews(View, ViewIndex, RasterTextureSize, MaxPixelsPerEdgeMultipler);
+			GraphBuilder.SetCommandListStat(GET_STATID(STAT_CLM_AfterPrePass));
 
-					TUniquePtr< Nanite::IRenderer > NaniteRenderer;
+			// special pass for DDM_AllOpaqueNoVelocity, which uses the velocity pass to finish the early depth pass write
+			if (bShouldRenderVelocities && Scene->EarlyZPassMode == DDM_AllOpaqueNoVelocity && RendererOutput == ERendererOutput::FinalSceneColor)
+			{
+				// Render the velocities of movable objects
+				GraphBuilder.SetCommandListStat(GET_STATID(STAT_CLM_Velocity));
+				RenderVelocities(GraphBuilder, InViews, SceneTextures, EVelocityPass::Opaque, bHairStrandsEnable);
+				GraphBuilder.SetCommandListStat(GET_STATID(STAT_CLM_AfterVelocity));
+			}
+		}
 
-					// Nanite::VisBuffer (Culling and Rasterization)
-					{
-						DynamicRenderScaling::FRDGScope DynamicScalingScope(GraphBuilder, GDynamicNaniteScalingPrimary);
+		{
+			Scene->WaitForCacheNaniteDrawCommandsTask();
 
-						RDG_GPU_STAT_SCOPE(GraphBuilder, NaniteVisBuffer);
-						RDG_EVENT_SCOPE(GraphBuilder, "Nanite::VisBuffer");
+			if (bNaniteEnabled && InViews.Num() > 0)
+			{
+				RenderNanite(GraphBuilder, InViews, SceneTextures, bIsEarlyDepthComplete, InNaniteBasePassVisibility, NaniteRasterResults, PrimaryNaniteViews);
+			}
+		}
 
-						NaniteRenderer = Nanite::IRenderer::Create(
-							GraphBuilder,
-							*Scene,
-							View,
-							GetSceneUniforms(),
-							SharedContext,
-							RasterContext,
-							CullingConfig,
-							ViewRect,
-							!bIsEarlyDepthComplete ? View.PrevViewInfo.NaniteHZB : View.PrevViewInfo.HZB
-						);
-
-						FSceneInstanceCullingQuery *SceneInstanceCullQuery = SceneCullingRenderer.CullInstances(GraphBuilder, View.ViewFrustum);
-						NaniteRenderer->DrawGeometry(
-							Scene->NaniteRasterPipelines[ENaniteMeshPass::BasePass],
-							RasterResults.VisibilityResults,
-							*NaniteViewsToRender,
-							SceneInstanceCullQuery
-						);
-
-						NaniteRenderer->ExtractResults( RasterResults );
-					}
+		if (FirstStageDepthBuffer)
+		{
+			SceneTextures.PartialDepth = FirstStageDepthBuffer;
+			AddResolveSceneDepthPass(GraphBuilder, InViews, SceneTextures.PartialDepth);
+		}
+		else
+		{
+			// Setup default partial depth to be scene depth so that it also works on transparent emitter when partial depth has not been generated.
+			SceneTextures.PartialDepth = SceneTextures.Depth;
+		}
+		SceneTextures.SetupMode = ESceneTextureSetupMode::SceneDepth;
+		SceneTextures.UniformBuffer = CreateSceneTextureUniformBuffer(GraphBuilder, &SceneTextures, FeatureLevel, SceneTextures.SetupMode);
 
-					// Nanite::BasePass (Depth Pre-Pass and HZB Build)
-					{
-						RDG_GPU_STAT_SCOPE(GraphBuilder, NaniteBasePass);
+		AddResolveSceneDepthPass(GraphBuilder, InViews, SceneTextures.Depth);
+	};
 
-						// Emit velocity with depth if not writing it in base pass.
-						FRDGTexture* VelocityBuffer = !IsUsingBasePassVelocity(ShaderPlatform) ? SceneTextures.Velocity : nullptr;
+	for (int32 i = 0; i < SceneCaptureRenderPassInfos.Num(); ++i)
+	{
+		FSceneCaptureRenderPassInfo& PassInfo = SceneCaptureRenderPassInfos[i];
 
-						Nanite::EmitDepthTargets(
-							GraphBuilder,
-							*Scene,
-							Views[ViewIndex],
-							bDrawSceneViewsInOneNanitePass,
-							RasterResults,
-							SceneTextures.Depth.Target,
-							VelocityBuffer
-						);
-						
-						// Sanity check (always force Z prepass)
-						check(bIsEarlyDepthComplete);
-					}
-				}
-			}
-		}
-	}
+		TArray<Nanite::FRasterResults, TInlineAllocator<2>> NaniteRasterResults;
+		TArray<Nanite::FPackedView, SceneRenderingAllocator> PrimaryNaniteViews;
+		FNaniteBasePassVisibility DummyNaniteBasePassVisibility;
+		RenderPrepassAndVelocity(PassInfo.Views, DummyNaniteBasePassVisibility, NaniteRasterResults, PrimaryNaniteViews);
 
-	if (FirstStageDepthBuffer)
-	{
-		SceneTextures.PartialDepth = FirstStageDepthBuffer;
-		AddResolveSceneDepthPass(GraphBuilder, Views, SceneTextures.PartialDepth);
+		CopySceneCaptureComponentToTarget(GraphBuilder, SceneTextures, nullptr, ViewFamily, PassInfo.Views);
 	}
-	else
-	{
-		// Setup default partial depth to be scene depth so that it also works on transparent emitter when partial depth has not been generated.
-		SceneTextures.PartialDepth = SceneTextures.Depth;
-	}
-	SceneTextures.SetupMode = ESceneTextureSetupMode::SceneDepth;
-	SceneTextures.UniformBuffer = CreateSceneTextureUniformBuffer(GraphBuilder, &SceneTextures, FeatureLevel, SceneTextures.SetupMode);
 
-	AddResolveSceneDepthPass(GraphBuilder, Views, SceneTextures.Depth);
+	TArray<Nanite::FRasterResults, TInlineAllocator<2>> NaniteRasterResults;
+	TArray<Nanite::FPackedView, SceneRenderingAllocator> PrimaryNaniteViews;
+	RenderPrepassAndVelocity(Views, NaniteBasePassVisibility, NaniteRasterResults, PrimaryNaniteViews);
 
 	FComputeLightGridOutput ComputeLightGridOutput = {};
 
@@ -3810,7 +3835,7 @@ void FDeferredShadingSceneRenderer::Render(FRDGBuilder& GraphBuilder)
 		if (bShouldRenderVelocities && !bBasePassCanOutputVelocity && (Scene->EarlyZPassMode != DDM_AllOpaqueNoVelocity))
 		{
 			GraphBuilder.SetCommandListStat(GET_STATID(STAT_CLM_Velocity));
-			RenderVelocities(GraphBuilder, SceneTextures, EVelocityPass::Opaque, bHairStrandsEnable);
+			RenderVelocities(GraphBuilder, Views, SceneTextures, EVelocityPass::Opaque, bHairStrandsEnable);
 			GraphBuilder.SetCommandListStat(GET_STATID(STAT_CLM_AfterVelocity));
 		}
 
@@ -4178,7 +4203,7 @@ void FDeferredShadingSceneRenderer::Render(FRDGBuilder& GraphBuilder)
 				const bool bRecreateSceneTextures = !HasBeenProduced(SceneTextures.Velocity);
 
 				GraphBuilder.SetCommandListStat(GET_STATID(STAT_CLM_TranslucentVelocity));
-				RenderVelocities(GraphBuilder, SceneTextures, EVelocityPass::Translucent, false);
+				RenderVelocities(GraphBuilder, Views, SceneTextures, EVelocityPass::Translucent, false);
 
 				if (bRecreateSceneTextures)
 				{
diff --git a/Engine/Source/Runtime/Renderer/Private/DeferredShadingRenderer.h b/Engine/Source/Runtime/Renderer/Private/DeferredShadingRenderer.h
--- a/Engine/Source/Runtime/Renderer/Private/DeferredShadingRenderer.h
+++ b/Engine/Source/Runtime/Renderer/Private/DeferredShadingRenderer.h
@@ -624,8 +624,8 @@ private:
 
 	void CreateIndirectCapsuleShadows();
 
-	void RenderPrePass(FRDGBuilder& GraphBuilder, FRDGTextureRef SceneDepthTexture, FInstanceCullingManager& InstanceCullingManager, FRDGTextureRef* FirstStageDepthBuffer);
-	void RenderPrePassHMD(FRDGBuilder& GraphBuilder, FRDGTextureRef SceneDepthTexture);
+	void RenderPrePass(FRDGBuilder& GraphBuilder, TArrayView<FViewInfo> InViews, FRDGTextureRef SceneDepthTexture, FInstanceCullingManager& InstanceCullingManager, FRDGTextureRef* FirstStageDepthBuffer);
+	void RenderPrePassHMD(FRDGBuilder& GraphBuilder, TArrayView<FViewInfo> InViews, FRDGTextureRef SceneDepthTexture);
 
 	void RenderFog(
 		FRDGBuilder& GraphBuilder,
@@ -1194,6 +1194,11 @@ private:
 
 	} NaniteBasePassVisibility;
 
+	void RenderNanite(FRDGBuilder& GraphBuilder, const TArray<FViewInfo>& InViews, FSceneTextures& SceneTextures, bool bIsEarlyDepthComplete,
+		FNaniteBasePassVisibility& InNaniteBasePassVisibility,
+		TArray<Nanite::FRasterResults, TInlineAllocator<2>>& NaniteRasterResults,
+		TArray<Nanite::FPackedView, SceneRenderingAllocator>& PrimaryNaniteViews);
+
 	/** Set to true if lights were injected into the light grid (this controlled by somewhat complex logic, this flag is used to cross-check). */
 	bool bAreLightsInLightGrid;
 };
diff --git a/Engine/Source/Runtime/Renderer/Private/DepthRendering.cpp b/Engine/Source/Runtime/Renderer/Private/DepthRendering.cpp
--- a/Engine/Source/Runtime/Renderer/Private/DepthRendering.cpp
+++ b/Engine/Source/Runtime/Renderer/Private/DepthRendering.cpp
@@ -483,7 +483,7 @@ void SetupDepthPassState(FMeshPassProcessorRenderState& DrawRenderState)
 
 extern const TCHAR* GetDepthPassReason(bool bDitheredLODTransitionsUseStencil, EShaderPlatform ShaderPlatform);
 
-void FDeferredShadingSceneRenderer::RenderPrePass(FRDGBuilder& GraphBuilder, FRDGTextureRef SceneDepthTexture, FInstanceCullingManager& InstanceCullingManager, FRDGTextureRef* FirstStageDepthBuffer)
+void FDeferredShadingSceneRenderer::RenderPrePass(FRDGBuilder& GraphBuilder, TArrayView<FViewInfo> InViews, FRDGTextureRef SceneDepthTexture, FInstanceCullingManager& InstanceCullingManager, FRDGTextureRef* FirstStageDepthBuffer)
 {
 	RDG_EVENT_SCOPE(GraphBuilder, "PrePass %s %s", GetDepthDrawingModeString(DepthPass.EarlyZPassMode), GetDepthPassReason(DepthPass.bDitheredLODTransitionsUseStencil, ShaderPlatform));
 	RDG_CSV_STAT_EXCLUSIVE_SCOPE(GraphBuilder, RenderPrePass);
@@ -494,11 +494,11 @@ void FDeferredShadingSceneRenderer::RenderPrePass(FRDGBuilder& GraphBuilder, FRD
 
 	const bool bParallelDepthPass = GRHICommandList.UseParallelAlgorithms() && CVarParallelPrePass.GetValueOnRenderThread();
 
-	RenderPrePassHMD(GraphBuilder, SceneDepthTexture);
+	RenderPrePassHMD(GraphBuilder, InViews, SceneDepthTexture);
 
 	if (DepthPass.IsRasterStencilDitherEnabled())
 	{
-		AddDitheredStencilFillPass(GraphBuilder, Views, SceneDepthTexture, DepthPass);
+		AddDitheredStencilFillPass(GraphBuilder, InViews, SceneDepthTexture, DepthPass);
 	}
 
 	auto RenderDepthPass = [&](uint8 DepthMeshPass)
@@ -510,11 +510,11 @@ void FDeferredShadingSceneRenderer::RenderPrePass(FRDGBuilder& GraphBuilder, FRD
 		{
 			RDG_WAIT_FOR_TASKS_CONDITIONAL(GraphBuilder, IsDepthPassWaitForTasksEnabled());
 
-			for (int32 ViewIndex = 0; ViewIndex < Views.Num(); ++ViewIndex)
+			for (int32 ViewIndex = 0; ViewIndex < InViews.Num(); ++ViewIndex)
 			{
-				FViewInfo& View = Views[ViewIndex];
+				FViewInfo& View = InViews[ViewIndex];
 				RDG_GPU_MASK_SCOPE(GraphBuilder, View.GPUMask);
-				RDG_EVENT_SCOPE_CONDITIONAL(GraphBuilder, Views.Num() > 1, "View%d", ViewIndex);
+				RDG_EVENT_SCOPE_CONDITIONAL(GraphBuilder, InViews.Num() > 1, "View%d", ViewIndex);
 
 				FMeshPassProcessorRenderState DrawRenderState;
 				SetupDepthPassState(DrawRenderState);
@@ -544,11 +544,11 @@ void FDeferredShadingSceneRenderer::RenderPrePass(FRDGBuilder& GraphBuilder, FRD
 		}
 		else
 		{
-			for (int32 ViewIndex = 0; ViewIndex < Views.Num(); ++ViewIndex)
+			for (int32 ViewIndex = 0; ViewIndex < InViews.Num(); ++ViewIndex)
 			{
-				FViewInfo& View = Views[ViewIndex];
+				FViewInfo& View = InViews[ViewIndex];
 				RDG_GPU_MASK_SCOPE(GraphBuilder, View.GPUMask);
-				RDG_EVENT_SCOPE_CONDITIONAL(GraphBuilder, Views.Num() > 1, "View%d", ViewIndex);
+				RDG_EVENT_SCOPE_CONDITIONAL(GraphBuilder, InViews.Num() > 1, "View%d", ViewIndex);
 
 				FMeshPassProcessorRenderState DrawRenderState;
 				SetupDepthPassState(DrawRenderState);
@@ -585,9 +585,9 @@ void FDeferredShadingSceneRenderer::RenderPrePass(FRDGBuilder& GraphBuilder, FRD
 
 		// Evaluate if any second stage depth buffer processing is required
 		bool bUsesSecondStageDepthPass = false;
-		for (int32 ViewIndex = 0; ViewIndex < Views.Num(); ++ViewIndex)
+		for (int32 ViewIndex = 0; ViewIndex < InViews.Num(); ++ViewIndex)
 		{
-			FViewInfo& View = Views[ViewIndex];
+			FViewInfo& View = InViews[ViewIndex];
 			bUsesSecondStageDepthPass |= View.bUsesSecondStageDepthPass;
 		}
 
@@ -600,9 +600,9 @@ void FDeferredShadingSceneRenderer::RenderPrePass(FRDGBuilder& GraphBuilder, FRD
 			// TODO: A proper solution would be an update to AddCopyTexturePass that also maintain the source texture HTile data.
 			FRDGTextureDesc FirstStageDepthBufferDesc = FRDGTextureDesc::Create2D(SceneDepthTexture->Desc.Extent, PF_R32_FLOAT, FClearValueBinding::Black, TexCreate_ShaderResource | TexCreate_UAV);
 			*FirstStageDepthBuffer = GraphBuilder.CreateTexture(FirstStageDepthBufferDesc, TEXT("FirstStageDepthBuffer"));
-			for (int32 ViewIndex = 0; ViewIndex < Views.Num(); ++ViewIndex)
+			for (int32 ViewIndex = 0; ViewIndex < InViews.Num(); ++ViewIndex)
 			{
-				FViewInfo& View = Views[ViewIndex];
+				FViewInfo& View = InViews[ViewIndex];
 				if (View.ViewState && !View.bStatePrevViewInfoIsReadOnly && View.bUsesSecondStageDepthPass)
 				{
 					AddViewDepthCopyCSPass(GraphBuilder, View, SceneDepthTexture, *FirstStageDepthBuffer);
@@ -624,14 +624,14 @@ void FDeferredShadingSceneRenderer::RenderPrePass(FRDGBuilder& GraphBuilder, FRD
 			RDG_EVENT_NAME("DitherStencilClear"),
 			PassParameters,
 			ERDGPassFlags::Raster,
-			[this](FRHICommandList& RHICmdList)
+			[this, InViews](FRHICommandList& RHICmdList)
 		{
-			if (Views.Num() > 1)
+			if (InViews.Num() > 1)
 			{
-				FIntRect FullViewRect = Views[0].ViewRect;
-				for (int32 ViewIndex = 1; ViewIndex < Views.Num(); ++ViewIndex)
+				FIntRect FullViewRect = InViews[0].ViewRect;
+				for (int32 ViewIndex = 1; ViewIndex < InViews.Num(); ++ViewIndex)
 				{
-					FullViewRect.Union(Views[ViewIndex].ViewRect);
+					FullViewRect.Union(InViews[ViewIndex].ViewRect);
 				}
 				RHICmdList.SetViewport(FullViewRect.Min.X, FullViewRect.Min.Y, 0, FullViewRect.Max.X, FullViewRect.Max.Y, 1);
 			}
@@ -643,7 +643,7 @@ void FDeferredShadingSceneRenderer::RenderPrePass(FRDGBuilder& GraphBuilder, FRD
 	const bool bForwardShadingEnabled = IsForwardShadingEnabled(ShaderPlatform);
 	if (!bForwardShadingEnabled)
 	{
-		StampDeferredDebugProbeDepthPS(GraphBuilder, Views, SceneDepthTexture);
+		StampDeferredDebugProbeDepthPS(GraphBuilder, InViews, SceneDepthTexture);
 	}
 #endif
 }
@@ -656,11 +656,6 @@ bool FMobileSceneRenderer::ShouldRenderPrePass() const
 
 void FMobileSceneRenderer::RenderPrePass(FRHICommandList& RHICmdList, const FViewInfo& View, const FInstanceCullingDrawParams* InstanceCullingDrawParams)
 {
-	if (!ShouldRenderPrePass())
-	{
-		return;
-	}
-	
 	checkSlow(RHICmdList.IsInsideRenderPass());
 
 	SCOPED_NAMED_EVENT(FMobileSceneRenderer_RenderPrePass, FColor::Emerald);
@@ -674,7 +669,7 @@ void FMobileSceneRenderer::RenderPrePass(FRHICommandList& RHICmdList, const FVie
 	View.ParallelMeshDrawCommandPasses[EMeshPass::DepthPass].DispatchDraw(nullptr, RHICmdList, InstanceCullingDrawParams);
 }
 
-void FDeferredShadingSceneRenderer::RenderPrePassHMD(FRDGBuilder& GraphBuilder, FRDGTextureRef DepthTexture)
+void FDeferredShadingSceneRenderer::RenderPrePassHMD(FRDGBuilder& GraphBuilder, TArrayView<FViewInfo> InViews, FRDGTextureRef DepthTexture)
 {
 	// Early out before we change any state if there's not a mask to render
 	if (!IsHMDHiddenAreaMaskActive())
@@ -688,7 +683,7 @@ void FDeferredShadingSceneRenderer::RenderPrePassHMD(FRDGBuilder& GraphBuilder,
 		return;
 	}
 
-	for (const FViewInfo& View : Views)
+	for (const FViewInfo& View : InViews)
 	{
 		// Don't draw the hidden area mesh in scene captures as they are not displayed
 		// through the HMD lenses.
@@ -978,12 +973,6 @@ void FDepthPassMeshProcessor::AddMeshBatch(const FMeshBatch& RESTRICT MeshBatch,
 {
 	bool bDraw = MeshBatch.bUseForDepthPass;
 	
-	// Mobile rendering does not use depth prepass by default
-	if (FeatureLevel == ERHIFeatureLevel::ES3_1 && EarlyZPassMode == DDM_None)
-	{
-		bDraw = false;
-	}
-
 	// Filter by occluder flags and settings if required.
 	if (bDraw && bRespectUseAsOccluderFlag && !MeshBatch.bUseAsOccluder && EarlyZPassMode < DDM_AllOpaque)
 	{
@@ -1070,12 +1059,6 @@ void FDepthPassMeshProcessor::CollectPSOInitializers(const FSceneTexturesConfig&
 		return;
 	}
 
-	// Mobile rendering does not use depth prepass by default
-	if (FeatureLevel == ERHIFeatureLevel::ES3_1 && EarlyZPassMode == DDM_None)
-	{
-		return;
-	}
-
 	const bool bIsTranslucent = IsTranslucentBlendMode(Material);
 
 	// Early out if translucent or material shouldn't be used during this pass
diff --git a/Engine/Source/Runtime/Renderer/Private/MobileShadingRenderer.cpp b/Engine/Source/Runtime/Renderer/Private/MobileShadingRenderer.cpp
--- a/Engine/Source/Runtime/Renderer/Private/MobileShadingRenderer.cpp
+++ b/Engine/Source/Runtime/Renderer/Private/MobileShadingRenderer.cpp
@@ -68,6 +68,7 @@
 #include "SceneUniformBuffer.h"
 #include "Engine/SpecularProfile.h"
 #include "LocalFogVolumeRendering.h"
+#include "SceneCaptureRendering.h"
 
 uint32 GetShadowQuality();
 
@@ -239,7 +240,7 @@ struct FRenderViewContext
 };
 using FRenderViewContextArray = TArray<FRenderViewContext, TInlineAllocator<2, SceneRenderingAllocator>>;
 
-static void GetRenderViews(TArray<FViewInfo>& InViews, FRenderViewContextArray& RenderViews)
+static void GetRenderViews(TArrayView<FViewInfo> InViews, FRenderViewContextArray& RenderViews)
 {
 	for (int32 ViewIndex = 0; ViewIndex < InViews.Num(); ViewIndex++)
 	{
@@ -631,13 +632,21 @@ void FMobileSceneRenderer::InitViews(
 		View.InitRHIResources();
 	}
 
+	for (int32 i = 0; i < SceneCaptureRenderPassInfos.Num(); ++i)
+	{
+		for (FViewInfo& View : SceneCaptureRenderPassInfos[i].Views)
+		{
+			View.InitRHIResources();
+		}		
+	}
+
 	{
 		RDG_CSV_STAT_EXCLUSIVE_SCOPE(GraphBuilder, UpdateGPUScene);
 
 		Scene->GPUScene.Update(GraphBuilder, GetSceneUniforms(), *Scene, ExternalAccessQueue);
-		for (int32 ViewIndex = 0; ViewIndex < Views.Num(); ViewIndex++)
+		for (int32 ViewIndex = 0; ViewIndex < AllViews.Num(); ViewIndex++)
 		{
-			Scene->GPUScene.UploadDynamicPrimitiveShaderDataForView(GraphBuilder, *Scene, Views[ViewIndex], ExternalAccessQueue);
+			Scene->GPUScene.UploadDynamicPrimitiveShaderDataForView(GraphBuilder, *Scene, *AllViews[ViewIndex], ExternalAccessQueue);
 		}
 	}
 
@@ -730,14 +739,15 @@ static void EndOcclusionScope(FRDGBuilder& GraphBuilder, TArray<FViewInfo>& View
 /*
 * Renders the Full Depth Prepass
 */
-void FMobileSceneRenderer::RenderFullDepthPrepass(FRDGBuilder& GraphBuilder, FSceneTextures& SceneTextures)
+void FMobileSceneRenderer::RenderFullDepthPrepass(FRDGBuilder& GraphBuilder, TArrayView<FViewInfo> InViews, FSceneTextures& SceneTextures, bool bIsSceneCaptureRenderPass)
 {
 	FRenderTargetBindingSlots BasePassRenderTargets;
 	BasePassRenderTargets.DepthStencil = FDepthStencilBinding(SceneTextures.Depth.Target, ERenderTargetLoadAction::EClear, ERenderTargetLoadAction::EClear, FExclusiveDepthStencil::DepthWrite_StencilWrite);
-	BasePassRenderTargets.NumOcclusionQueries = ComputeNumOcclusionQueriesToBatch();
+	// If this is scene capture render pass, don't render occlusion
+	BasePassRenderTargets.NumOcclusionQueries = bIsSceneCaptureRenderPass ? 0 : ComputeNumOcclusionQueriesToBatch();
 
 	FRenderViewContextArray RenderViews;
-	GetRenderViews(Views, RenderViews);
+	GetRenderViews(InViews, RenderViews);
 
 	for (FRenderViewContext& ViewContext : RenderViews)
 	{
@@ -760,7 +770,8 @@ void FMobileSceneRenderer::RenderFullDepthPrepass(FRDGBuilder& GraphBuilder, FSc
 		View.ParallelMeshDrawCommandPasses[EMeshPass::DepthPass].BuildRenderingCommands(GraphBuilder, Scene->GPUScene, PassParameters->InstanceCullingDrawParams);
 
 		// Render occlusion at the last view pass only, as they already loop through all views
-		bool bDoOcclusionQueries = (ViewContext.bIsLastView && DoOcclusionQueries());
+		// If this is scene capture render pass, don't render occlusion.
+		bool bDoOcclusionQueries = (ViewContext.bIsLastView && DoOcclusionQueries() && !bIsSceneCaptureRenderPass);
 
 		GraphBuilder.AddPass(
 			RDG_EVENT_NAME("FullDepthPrepass"),
@@ -1020,9 +1031,20 @@ void FMobileSceneRenderer::Render(FRDGBuilder& GraphBuilder)
 
 	const FRDGSystemTextures& SystemTextures = FRDGSystemTextures::Get(GraphBuilder);
 	
+	for (int32 i = 0; i < SceneCaptureRenderPassInfos.Num(); ++i)
+	{
+		FSceneCaptureRenderPassInfo& PassInfo = SceneCaptureRenderPassInfos[i];
+		RenderFullDepthPrepass(GraphBuilder, PassInfo.Views, SceneTextures, true);
+		AddResolveSceneDepthPass(GraphBuilder, PassInfo.Views, SceneTextures.Depth);
+
+		SceneTextures.MobileSetupMode = EMobileSceneTextureSetupMode::SceneDepth;
+		SceneTextures.MobileUniformBuffer = CreateMobileSceneTextureUniformBuffer(GraphBuilder, &SceneTextures, SceneTextures.MobileSetupMode);
+		CopySceneCaptureComponentToTarget(GraphBuilder, SceneTextures, nullptr, ViewFamily, PassInfo.Views);
+	}
+
 	if (bIsFullDepthPrepassEnabled)
 	{
-		RenderFullDepthPrepass(GraphBuilder, SceneTextures);
+		RenderFullDepthPrepass(GraphBuilder, Views, SceneTextures);
 
 		if (!bRequiresSceneDepthAux)
 		{
@@ -1086,10 +1108,10 @@ void FMobileSceneRenderer::Render(FRDGBuilder& GraphBuilder)
 		{
 			// Render the velocities of movable objects
 			GraphBuilder.SetCommandListStat(GET_STATID(STAT_CLMM_Velocity));
-			RenderVelocities(GraphBuilder, SceneTextures, EVelocityPass::Opaque, false);
+			RenderVelocities(GraphBuilder, Views, SceneTextures, EVelocityPass::Opaque, false);
 
 			GraphBuilder.SetCommandListStat(GET_STATID(STAT_CLMM_TranslucentVelocity));
-			RenderVelocities(GraphBuilder, SceneTextures, EVelocityPass::Translucent, false);
+			RenderVelocities(GraphBuilder, Views, SceneTextures, EVelocityPass::Translucent, false);
 
 			SceneTextures.MobileSetupMode = EMobileSceneTextureSetupMode::All;
 			SceneTextures.MobileUniformBuffer = CreateMobileSceneTextureUniformBuffer(GraphBuilder, &SceneTextures, SceneTextures.MobileSetupMode);
diff --git a/Engine/Source/Runtime/Renderer/Private/SceneCaptureRendering.cpp b/Engine/Source/Runtime/Renderer/Private/SceneCaptureRendering.cpp
--- a/Engine/Source/Runtime/Renderer/Private/SceneCaptureRendering.cpp
+++ b/Engine/Source/Runtime/Renderer/Private/SceneCaptureRendering.cpp
@@ -43,6 +43,15 @@
 #include "RectLightTexture.h"
 #include "Materials/MaterialRenderProxy.h"
 
+bool GSceneCaptureAllowRenderInMainRenderer = true;
+static FAutoConsoleVariableRef CVarSceneCaptureAllowRenderInMainRenderer(
+	TEXT("r.SceneCapture.AllowRenderInMainRenderer"),
+	GSceneCaptureAllowRenderInMainRenderer,
+	TEXT("Whether to allow SceneDepth & DeviceDepth scene capture to render in the main renderer as an optimization.\n")
+	TEXT("0: render as an independent renderer.\n")
+	TEXT("1: render as part of the main renderer if Render in Main Renderer is enabled on scene capture component.\n"),
+	ECVF_Scalability);
+
 #if WITH_EDITOR
 // All scene captures on the given render thread frame will be dumped
 uint32 GDumpSceneCaptureMemoryFrame = INDEX_NONE;
@@ -181,20 +190,49 @@ void CopySceneCaptureComponentToTarget(
 	const FSceneViewFamily& ViewFamily,
 	TConstArrayView<FViewInfo> Views)
 {
-	ESceneCaptureSource SceneCaptureSource = ViewFamily.SceneCaptureSource;
-
-	if (IsForwardShadingEnabled(ViewFamily.GetShaderPlatform()) && (SceneCaptureSource == SCS_Normal || SceneCaptureSource == SCS_BaseColor))
+	TArray<const FViewInfo*> ViewPtrArray;
+	for (const FViewInfo& View : Views)
 	{
-		SceneCaptureSource = SCS_SceneColorHDR;
+		ViewPtrArray.Add(&View);
 	}
+	CopySceneCaptureComponentToTarget(GraphBuilder, SceneTextures, ViewFamilyTexture, ViewFamily, ViewPtrArray);
+}
 
-	if (CaptureNeedsSceneColor(SceneCaptureSource))
+void CopySceneCaptureComponentToTarget(
+	FRDGBuilder& GraphBuilder,
+	const FMinimalSceneTextures& SceneTextures,
+	FRDGTextureRef ViewFamilyTexture,
+	const FSceneViewFamily& ViewFamily,
+	const TArray<const FViewInfo*>& Views)
+{
+	FGraphicsPipelineStateInitializer GraphicsPSOInit;
+	GraphicsPSOInit.RasterizerState = TStaticRasterizerState<FM_Solid, CM_None>::GetRHI();
+	GraphicsPSOInit.DepthStencilState = TStaticDepthStencilState<false, CF_Always>::GetRHI();
+
+	const bool bForwardShadingEnabled = IsForwardShadingEnabled(ViewFamily.GetShaderPlatform());
+	for (int32 ViewIndex = 0; ViewIndex < Views.Num(); ViewIndex++)
 	{
-		RDG_EVENT_SCOPE(GraphBuilder, "CaptureSceneComponent[%d]", SceneCaptureSource);
+		const FViewInfo& View = *Views[ViewIndex];
+
+		// If view has its own scene capture RT, it takes priority over view family RT
+		FRDGTextureRef RenderTarget = View.SceneCaptureRenderTarget ? View.SceneCaptureRenderTarget->GetRenderTargetTexture(GraphBuilder) : ViewFamilyTexture;
+		if (!RenderTarget)
+		{
+			continue;
+		}
+
+		// If view has its own scene capture setting, use it over view family setting
+		ESceneCaptureSource SceneCaptureSource = View.SceneCaptureRenderTarget ? View.SceneCaptureSource : ViewFamily.SceneCaptureSource;
+		if (bForwardShadingEnabled && (SceneCaptureSource == SCS_Normal || SceneCaptureSource == SCS_BaseColor))
+		{
+			SceneCaptureSource = SCS_SceneColorHDR;
+		}
+		if (!CaptureNeedsSceneColor(SceneCaptureSource))
+		{
+			continue;
+		}
 
-		FGraphicsPipelineStateInitializer GraphicsPSOInit;
-		GraphicsPSOInit.RasterizerState = TStaticRasterizerState<FM_Solid, CM_None>::GetRHI();
-		GraphicsPSOInit.DepthStencilState = TStaticDepthStencilState<false, CF_Always>::GetRHI();
+		RDG_EVENT_SCOPE(GraphBuilder, "CaptureSceneComponent_View[%d]", SceneCaptureSource);
 
 		bool bIsCompositing = false;
 		if (SceneCaptureSource == SCS_SceneColorHDR && ViewFamily.SceneCaptureCompositeMode == SCCM_Composite)
@@ -214,51 +252,46 @@ void CopySceneCaptureComponentToTarget(
 			GraphicsPSOInit.BlendState = TStaticBlendState<>::GetRHI();
 		}
 
-		const bool bUse128BitRT = PlatformRequires128bitRT(ViewFamilyTexture->Desc.Format);
+		const bool bUse128BitRT = PlatformRequires128bitRT(RenderTarget->Desc.Format);
 		const FSceneCapturePS::FPermutationDomain PixelPermutationVector = FSceneCapturePS::GetPermutationVector(SceneCaptureSource, bUse128BitRT, IsMobilePlatform(ViewFamily.GetShaderPlatform()));
 
-		for (int32 ViewIndex = 0; ViewIndex < Views.Num(); ViewIndex++)
+		FSceneCapturePS::FParameters* PassParameters = GraphBuilder.AllocParameters<FSceneCapturePS::FParameters>();
+		PassParameters->View = View.ViewUniformBuffer;
+		PassParameters->SceneTextures = SceneTextures.GetSceneTextureShaderParameters(ViewFamily.GetFeatureLevel());
+		PassParameters->RenderTargets[0] = FRenderTargetBinding(RenderTarget, bIsCompositing ? ERenderTargetLoadAction::ELoad : ERenderTargetLoadAction::ENoAction);
+
+		TShaderMapRef<FScreenVS> VertexShader(View.ShaderMap);
+		TShaderMapRef<FSceneCapturePS> PixelShader(View.ShaderMap, PixelPermutationVector);
+
+		GraphicsPSOInit.BoundShaderState.VertexDeclarationRHI = GFilterVertexDeclaration.VertexDeclarationRHI;
+		GraphicsPSOInit.BoundShaderState.VertexShaderRHI = VertexShader.GetVertexShader();
+		GraphicsPSOInit.BoundShaderState.PixelShaderRHI = PixelShader.GetPixelShader();
+		GraphicsPSOInit.PrimitiveType = PT_TriangleList;
+
+		GraphBuilder.AddPass(
+			RDG_EVENT_NAME("View(%d)", ViewIndex),
+			PassParameters,
+			ERDGPassFlags::Raster,
+			[PassParameters, GraphicsPSOInit, VertexShader, PixelShader, &View] (FRHICommandList& RHICmdList)
 		{
-			const FViewInfo& View = Views[ViewIndex];
-
-			FSceneCapturePS::FParameters* PassParameters = GraphBuilder.AllocParameters<FSceneCapturePS::FParameters>();
-			PassParameters->View = View.ViewUniformBuffer;
-			PassParameters->SceneTextures = SceneTextures.GetSceneTextureShaderParameters(ViewFamily.GetFeatureLevel());
-			PassParameters->RenderTargets[0] = FRenderTargetBinding(ViewFamilyTexture, bIsCompositing ? ERenderTargetLoadAction::ELoad : ERenderTargetLoadAction::ENoAction);
-
-			TShaderMapRef<FScreenVS> VertexShader(View.ShaderMap);
-			TShaderMapRef<FSceneCapturePS> PixelShader(View.ShaderMap, PixelPermutationVector);
-
-			GraphicsPSOInit.BoundShaderState.VertexDeclarationRHI = GFilterVertexDeclaration.VertexDeclarationRHI;
-			GraphicsPSOInit.BoundShaderState.VertexShaderRHI = VertexShader.GetVertexShader();
-			GraphicsPSOInit.BoundShaderState.PixelShaderRHI = PixelShader.GetPixelShader();
-			GraphicsPSOInit.PrimitiveType = PT_TriangleList;
-
-			GraphBuilder.AddPass(
-				RDG_EVENT_NAME("View(%d)", ViewIndex),
-				PassParameters,
-				ERDGPassFlags::Raster,
-				[PassParameters, GraphicsPSOInit, VertexShader, PixelShader, &View] (FRHICommandList& RHICmdList)
-			{
-				FGraphicsPipelineStateInitializer LocalGraphicsPSOInit = GraphicsPSOInit;
-				RHICmdList.ApplyCachedRenderTargets(LocalGraphicsPSOInit);
-				SetGraphicsPipelineState(RHICmdList, LocalGraphicsPSOInit, 0);
-				SetShaderParameters(RHICmdList, PixelShader, PixelShader.GetPixelShader(), *PassParameters);
-				
-				CopyCaptureToTargetSetViewportFn(RHICmdList);
-
-				DrawRectangle(
-					RHICmdList,
-					View.ViewRect.Min.X, View.ViewRect.Min.Y,
-					View.ViewRect.Width(), View.ViewRect.Height(),
-					View.ViewRect.Min.X, View.ViewRect.Min.Y,
-					View.ViewRect.Width(), View.ViewRect.Height(),
-					View.UnconstrainedViewRect.Size(),
-					View.GetSceneTexturesConfig().Extent,
-					VertexShader,
-					EDRF_UseTriangleOptimization);
-			});
-		}
+			FGraphicsPipelineStateInitializer LocalGraphicsPSOInit = GraphicsPSOInit;
+			RHICmdList.ApplyCachedRenderTargets(LocalGraphicsPSOInit);
+			SetGraphicsPipelineState(RHICmdList, LocalGraphicsPSOInit, 0);
+			SetShaderParameters(RHICmdList, PixelShader, PixelShader.GetPixelShader(), *PassParameters);
+			
+			CopyCaptureToTargetSetViewportFn(RHICmdList);
+
+			DrawRectangle(
+				RHICmdList,
+				View.ViewRect.Min.X, View.ViewRect.Min.Y,
+				View.ViewRect.Width(), View.ViewRect.Height(),
+				View.ViewRect.Min.X, View.ViewRect.Min.Y,
+				View.ViewRect.Width(), View.ViewRect.Height(),
+				View.UnconstrainedViewRect.Size(),
+				View.GetSceneTexturesConfig().Extent,
+				VertexShader,
+				EDRF_UseTriangleOptimization);
+		});
 	}
 }
 
@@ -593,6 +626,77 @@ void BuildProjectionMatrix(FIntPoint InRenderTargetSize, float InFOV, float InNe
 	}
 }
 
+void GetShowOnlyAndHiddenComponents(USceneCaptureComponent* SceneCaptureComponent, TSet<FPrimitiveComponentId>& HiddenPrimitives, TOptional<TSet<FPrimitiveComponentId>>& ShowOnlyPrimitives)
+{
+	check(SceneCaptureComponent);
+	for (auto It = SceneCaptureComponent->HiddenComponents.CreateConstIterator(); It; ++It)
+	{
+		// If the primitive component was destroyed, the weak pointer will return NULL.
+		UPrimitiveComponent* PrimitiveComponent = It->Get();
+		if (PrimitiveComponent)
+		{
+			HiddenPrimitives.Add(PrimitiveComponent->GetPrimitiveSceneId());
+		}
+	}
+
+	for (auto It = SceneCaptureComponent->HiddenActors.CreateConstIterator(); It; ++It)
+	{
+		AActor* Actor = *It;
+
+		if (Actor)
+		{
+			for (UActorComponent* Component : Actor->GetComponents())
+			{
+				if (UPrimitiveComponent* PrimComp = Cast<UPrimitiveComponent>(Component))
+				{
+					HiddenPrimitives.Add(PrimComp->GetPrimitiveSceneId());
+				}
+			}
+		}
+	}
+
+	if (SceneCaptureComponent->PrimitiveRenderMode == ESceneCapturePrimitiveRenderMode::PRM_UseShowOnlyList)
+	{
+		ShowOnlyPrimitives.Emplace();
+
+		for (auto It = SceneCaptureComponent->ShowOnlyComponents.CreateConstIterator(); It; ++It)
+		{
+			// If the primitive component was destroyed, the weak pointer will return NULL.
+			UPrimitiveComponent* PrimitiveComponent = It->Get();
+			if (PrimitiveComponent)
+			{
+				ShowOnlyPrimitives->Add(PrimitiveComponent->GetPrimitiveSceneId());
+			}
+		}
+
+		for (auto It = SceneCaptureComponent->ShowOnlyActors.CreateConstIterator(); It; ++It)
+		{
+			AActor* Actor = *It;
+
+			if (Actor)
+			{
+				for (UActorComponent* Component : Actor->GetComponents())
+				{
+					if (UPrimitiveComponent* PrimComp = Cast<UPrimitiveComponent>(Component))
+					{
+						ShowOnlyPrimitives->Add(PrimComp->GetPrimitiveSceneId());
+					}
+				}
+			}
+		}
+	}
+	else if (SceneCaptureComponent->ShowOnlyComponents.Num() > 0 || SceneCaptureComponent->ShowOnlyActors.Num() > 0)
+	{
+		static bool bWarned = false;
+
+		if (!bWarned)
+		{
+			UE_LOG(LogRenderer, Log, TEXT("Scene Capture has ShowOnlyComponents or ShowOnlyActors ignored by the PrimitiveRenderMode setting! %s"), *SceneCaptureComponent->GetPathName());
+			bWarned = true;
+		}
+	}
+}
+
 void SetupViewFamilyForSceneCapture(
 	FSceneViewFamily& ViewFamily,
 	USceneCaptureComponent* SceneCaptureComponent,
@@ -653,73 +757,7 @@ void SetupViewFamilyForSceneCapture(
 
 		FSceneView* View = new FSceneView(ViewInitOptions);
 
-		check(SceneCaptureComponent);
-		for (auto It = SceneCaptureComponent->HiddenComponents.CreateConstIterator(); It; ++It)
-		{
-			// If the primitive component was destroyed, the weak pointer will return NULL.
-			UPrimitiveComponent* PrimitiveComponent = It->Get();
-			if (PrimitiveComponent)
-			{
-				View->HiddenPrimitives.Add(PrimitiveComponent->GetPrimitiveSceneId());
-			}
-		}
-
-		for (auto It = SceneCaptureComponent->HiddenActors.CreateConstIterator(); It; ++It)
-		{
-			AActor* Actor = *It;
-
-			if (Actor)
-			{
-				for (UActorComponent* Component : Actor->GetComponents())
-				{
-					if (UPrimitiveComponent* PrimComp = Cast<UPrimitiveComponent>(Component))
-					{
-						View->HiddenPrimitives.Add(PrimComp->GetPrimitiveSceneId());
-					}
-				}
-			}
-		}
-
-		if (SceneCaptureComponent->PrimitiveRenderMode == ESceneCapturePrimitiveRenderMode::PRM_UseShowOnlyList)
-		{
-				View->ShowOnlyPrimitives.Emplace();
-
-			for (auto It = SceneCaptureComponent->ShowOnlyComponents.CreateConstIterator(); It; ++It)
-			{
-				// If the primitive component was destroyed, the weak pointer will return NULL.
-				UPrimitiveComponent* PrimitiveComponent = It->Get();
-				if (PrimitiveComponent)
-				{
-					View->ShowOnlyPrimitives->Add(PrimitiveComponent->GetPrimitiveSceneId());
-				}
-			}
-
-			for (auto It = SceneCaptureComponent->ShowOnlyActors.CreateConstIterator(); It; ++It)
-			{
-				AActor* Actor = *It;
-
-				if (Actor)
-				{
-					for (UActorComponent* Component : Actor->GetComponents())
-					{
-						if (UPrimitiveComponent* PrimComp = Cast<UPrimitiveComponent>(Component))
-						{
-							View->ShowOnlyPrimitives->Add(PrimComp->GetPrimitiveSceneId());
-						}
-					}
-				}
-			}
-		}
-		else if (SceneCaptureComponent->ShowOnlyComponents.Num() > 0 || SceneCaptureComponent->ShowOnlyActors.Num() > 0)
-		{
-			static bool bWarned = false;
-
-			if (!bWarned)
-			{
-				UE_LOG(LogRenderer, Log, TEXT("Scene Capture has ShowOnlyComponents or ShowOnlyActors ignored by the PrimitiveRenderMode setting! %s"), *SceneCaptureComponent->GetPathName());
-				bWarned = true;
-			}
-		}
+		GetShowOnlyAndHiddenComponents(SceneCaptureComponent, View->HiddenPrimitives, View->ShowOnlyPrimitives);
 
 		ViewFamily.Views.Add(View);
 
@@ -857,6 +895,24 @@ void FScene::UpdateSceneCaptureContents(USceneCaptureComponent2D* CaptureCompone
 			}
 		}
 
+		// As optimization for depth capture modes, render scene capture as additional render passes inside the main renderer.
+		if (GSceneCaptureAllowRenderInMainRenderer && CaptureComponent->bRenderInMainRenderer && (CaptureComponent->CaptureSource == ESceneCaptureSource::SCS_SceneDepth || CaptureComponent->CaptureSource == ESceneCaptureSource::SCS_DeviceDepth))
+		{
+			FSceneCaptureInfo CaptureInfo;
+			CaptureInfo.ViewLocation = ViewLocation;
+			CaptureInfo.ViewRotationMatrix = ViewRotationMatrix;
+			CaptureInfo.ProjectionMatrix = ProjectionMatrix;
+			CaptureInfo.RenderTarget = TextureRenderTarget->GameThread_GetRenderTargetResource();
+			CaptureInfo.SceneCaptureSource = CaptureComponent->CaptureSource;
+			CaptureInfo.ViewActor = CaptureComponent->GetViewOwner();
+
+			GetShowOnlyAndHiddenComponents(CaptureComponent, CaptureInfo.HiddenPrimitives, CaptureInfo.ShowOnlyPrimitives);
+
+			// Caching scene capture info to be passed to the scene renderer.
+			SceneCaptureInfos.Add(CaptureInfo);
+			return;
+		}
+
 		FSceneRenderer* SceneRenderer = CreateSceneRendererForSceneCapture(
 			this, 
 			CaptureComponent, 
diff --git a/Engine/Source/Runtime/Renderer/Private/SceneRendering.cpp b/Engine/Source/Runtime/Renderer/Private/SceneRendering.cpp
--- a/Engine/Source/Runtime/Renderer/Private/SceneRendering.cpp
+++ b/Engine/Source/Runtime/Renderer/Private/SceneRendering.cpp
@@ -2745,6 +2745,56 @@ FSceneRenderer::FSceneRenderer(const FSceneViewFamily* InViewFamily, FHitProxyCo
 		}
 	}
 
+	SceneCaptureRenderPassInfos.Empty(Scene->SceneCaptureInfos.Num());
+	SceneCaptureRenderPassInfos.AddDefaulted(Scene->SceneCaptureInfos.Num());
+
+	int32 NumSceneCaptureViews = 0;
+	for (int32 i = 0; i < Scene->SceneCaptureInfos.Num(); i++)
+	{
+		const FScene::FSceneCaptureInfo& CaptureInfo = Scene->SceneCaptureInfos[i];
+
+		FSceneViewInitOptions ViewInitOptions;
+		FIntPoint RenderTargetSize = CaptureInfo.RenderTarget->GetSizeXY();
+		ViewInitOptions.SetViewRectangle(FIntRect(0, 0, RenderTargetSize.X, RenderTargetSize.Y));
+		ViewInitOptions.ViewOrigin = CaptureInfo.ViewLocation;
+		ViewInitOptions.ViewRotationMatrix = CaptureInfo.ViewRotationMatrix;
+		ViewInitOptions.ProjectionMatrix = CaptureInfo.ProjectionMatrix;
+		ViewInitOptions.SceneCaptureRenderTarget = CaptureInfo.RenderTarget;
+		ViewInitOptions.bIsSceneCapture = true;
+		ViewInitOptions.SceneCaptureSource = CaptureInfo.SceneCaptureSource;
+		ViewInitOptions.ViewFamily = &ViewFamily;
+		ViewInitOptions.ViewActor = CaptureInfo.ViewActor;
+		ViewInitOptions.ShowOnlyPrimitives = CaptureInfo.ShowOnlyPrimitives;
+		ViewInitOptions.HiddenPrimitives = CaptureInfo.HiddenPrimitives;
+
+		FSceneView NewView(ViewInitOptions);
+		FViewInfo* ViewInfo = &SceneCaptureRenderPassInfos[i].Views.Emplace_GetRef(&NewView);
+		// Must initialize to have a GPUScene connected to be able to collect dynamic primitives.
+		ViewInfo->DynamicPrimitiveCollector = FGPUScenePrimitiveCollector(&GPUSceneDynamicContext);
+		ViewInfo->bDisableQuerySubmissions = true;
+		ViewInfo->bIgnoreExistingQueries = true;
+
+		NumSceneCaptureViews++;
+	}
+
+	AllViews.Empty(Views.Num() + NumSceneCaptureViews);
+	for (int32 i = 0; i < Views.Num(); ++i)
+	{
+		AllViews.Add(&Views[i]);
+	}
+	for (FSceneCaptureRenderPassInfo& PassInfo : SceneCaptureRenderPassInfos)
+	{
+		for (FViewInfo& View : PassInfo.Views)
+		{
+			AllViews.Add(&View);
+		}
+	}
+
+	check(!ViewFamily.AllViews.Num());
+	ViewFamily.AllViews.Append(AllViews);
+
+	Scene->SceneCaptureInfos.Reset();
+
 	FeatureLevel = Scene->GetFeatureLevel();
 	ShaderPlatform = Scene->GetShaderPlatform();
 
@@ -3037,6 +3087,14 @@ void FSceneRenderer::PrepareViewRectsForRendering(FRHICommandListImmediate& RHIC
 			GEngine->StereoRenderingDevice->SetFinalViewRect(RHICmdList, View.StereoViewIndex, OutputViewRect);
 		}
 	}
+
+	for (FSceneCaptureRenderPassInfo& PassInfo : SceneCaptureRenderPassInfos)
+	{
+		for (FViewInfo& View : PassInfo.Views)
+		{
+			View.ViewRect = View.UnscaledViewRect;
+		}	
+	}
 }
 
 #if WITH_MGPU
@@ -3374,9 +3432,9 @@ IVisibilityTaskData* FSceneRenderer::OnRenderBegin(FRDGBuilder& GraphBuilder, FG
 			for (int32 ViewExt = 0; ViewExt < ViewFamily.ViewExtensions.Num(); ViewExt++)
 			{
 				ViewFamily.ViewExtensions[ViewExt]->PreRenderViewFamily_RenderThread(GraphBuilder, ViewFamily);
-				for (int32 ViewIndex = 0; ViewIndex < ViewFamily.Views.Num(); ViewIndex++)
+				for (int32 ViewIndex = 0; ViewIndex < AllViews.Num(); ViewIndex++)
 				{
-					ViewFamily.ViewExtensions[ViewExt]->PreRenderView_RenderThread(GraphBuilder, Views[ViewIndex]);
+					ViewFamily.ViewExtensions[ViewExt]->PreRenderView_RenderThread(GraphBuilder, *AllViews[ViewIndex]);
 				}
 			}
 		}
@@ -4238,9 +4296,9 @@ static void DeleteSceneRenderers(const TArray<FSceneRenderer*>& SceneRenderers,
 			SceneRenderer->DispatchedShadowDepthPasses[PassIndex]->WaitForTasksAndEmpty(WaitThread);
 		}
 
-		for (FViewInfo& View : SceneRenderer->Views)
+		for (FViewInfo* View : SceneRenderer->AllViews)
 		{
-			View.WaitForTasks(WaitThread);
+			View->WaitForTasks(WaitThread);
 		}
 	}
 
@@ -4801,8 +4859,6 @@ void FRendererModule::BeginRenderingViewFamilies(FCanvas* Canvas, TArrayView<FSc
 			ViewFamiliesConst.Add(ViewFamily);
 		}
 
-		FSceneRenderer::CreateSceneRenderers(ViewFamiliesConst, Canvas->GetHitProxyConsumer(), SceneRenderers);
-
 		bool bShowHitProxies = false;
 		for (FSceneRenderer* SceneRenderer : SceneRenderers)
 		{
@@ -4816,7 +4872,12 @@ void FRendererModule::BeginRenderingViewFamilies(FCanvas* Canvas, TArrayView<FSc
 		if (!bShowHitProxies)
 		{
 			USceneCaptureComponent::UpdateDeferredCaptures(Scene);
+		}
+
+		FSceneRenderer::CreateSceneRenderers(ViewFamiliesConst, Canvas->GetHitProxyConsumer(), SceneRenderers);
 
+		if (!bShowHitProxies)
+		{
 			for (int32 ReflectionIndex = 0; ReflectionIndex < Scene->PlanarReflections_GameThread.Num(); ReflectionIndex++)
 			{
 				UPlanarReflectionComponent* ReflectionComponent = Scene->PlanarReflections_GameThread[ReflectionIndex];
diff --git a/Engine/Source/Runtime/Renderer/Private/SceneRendering.h b/Engine/Source/Runtime/Renderer/Private/SceneRendering.h
--- a/Engine/Source/Runtime/Renderer/Private/SceneRendering.h
+++ b/Engine/Source/Runtime/Renderer/Private/SceneRendering.h
@@ -2020,6 +2020,16 @@ public:
 	/** The views being rendered. */
 	TArray<FViewInfo> Views;
 
+	/** Information of a scene capture that renders as part of the main renderer. */
+	struct FSceneCaptureRenderPassInfo
+	{
+		TArray<FViewInfo> Views;
+	};
+	TArray<FSceneCaptureRenderPassInfo> SceneCaptureRenderPassInfos;
+
+	/** All views include main camera views and scene capture views. */
+	TArray<FViewInfo*> AllViews;
+
 	/** Views across all view families (may contain additional views if rendering multiple families together). */
 	TArray<const FSceneView*> AllFamilyViews;
 	TArray<const FSceneViewFamily*> AllFamilies;
@@ -2493,6 +2503,7 @@ protected:
 
 	void RenderVelocities(
 		FRDGBuilder& GraphBuilder,
+		TArrayView<FViewInfo> InViews,
 		const FSceneTextures& SceneTextures,
 		EVelocityPass VelocityPass,
 		bool bForceVelocity);
@@ -2636,7 +2647,7 @@ protected:
 
 	void RenderPrePass(FRHICommandList& RHICmdList, const FViewInfo& View, const FInstanceCullingDrawParams* InstanceCullingDrawParams);
 	void RenderMaskedPrePass(FRHICommandList& RHICmdList, const FViewInfo& View);
-	void RenderFullDepthPrepass(FRDGBuilder& GraphBuilder, FSceneTextures& SceneTextures);
+	void RenderFullDepthPrepass(FRDGBuilder& GraphBuilder, TArrayView<FViewInfo> InViews, FSceneTextures& SceneTextures, bool bIsSceneCaptureRenderPass=false);
 
 	void RenderMobileLocalLightsBuffer(FRDGBuilder& GraphBuilder, FSceneTextures& SceneTextures, bool bIsPrepass);
 
diff --git a/Engine/Source/Runtime/Renderer/Private/SceneVisibility.cpp b/Engine/Source/Runtime/Renderer/Private/SceneVisibility.cpp
--- a/Engine/Source/Runtime/Renderer/Private/SceneVisibility.cpp
+++ b/Engine/Source/Runtime/Renderer/Private/SceneVisibility.cpp
@@ -3066,7 +3066,7 @@ static int32 PrecomputedOcclusionCull(FVisibilityViewPacket& ViewPacket, FPrimit
 
 ///////////////////////////////////////////////////////////////////////////////
 
-FVisibilityTaskConfig::FVisibilityTaskConfig(const FScene& Scene, TConstArrayView<FViewInfo> Views)
+FVisibilityTaskConfig::FVisibilityTaskConfig(const FScene& Scene, TConstArrayView<FViewInfo*> Views)
 {
 	Schedule = GVisibilityTaskSchedule != 0 ? EVisibilityTaskSchedule::Parallel : EVisibilityTaskSchedule::RenderThread;
 
@@ -3121,7 +3121,7 @@ FVisibilityTaskConfig::FVisibilityTaskConfig(const FScene& Scene, TConstArrayVie
 			{
 				MaxQueriesPerTask = GOcclusionCullMaxQueriesPerTask;
 			}
-			else if (FSceneViewState* ViewState = static_cast<FSceneViewState*>(Views[ViewIndex].State))
+			else if (FSceneViewState* ViewState = static_cast<FSceneViewState*>(Views[ViewIndex]->State))
 			{
 				MaxQueriesPerTask = FMath::DivideAndRoundUp(ViewState->Occlusion.NumRequestedQueries, FMath::Max(1u, NumWorkerThreads) * NumOcclusionCullTasksPerThread);
 			}
@@ -3450,7 +3450,7 @@ FVisibilityTaskData::FVisibilityTaskData(FRHICommandListImmediate& InRHICmdList,
 	: RHICmdList(InRHICmdList)
 	, SceneRenderer(InSceneRenderer)
 	, Scene(*SceneRenderer.Scene)
-	, Views(SceneRenderer.Views)
+	, Views(SceneRenderer.AllViews)
 	, ViewFamily(SceneRenderer.ViewFamily)
 	, ShadingPath(Scene.GetShadingPath())
 	, GlobalDynamicBuffers(InGlobalDynamicBuffers)
@@ -3470,7 +3470,7 @@ void FVisibilityTaskData::LaunchVisibilityTasks()
 	for (int32 ViewIndex = 0; ViewIndex < Views.Num(); ++ViewIndex)
 	{
 		// Each view gets its own visibility task packet which contains all the state to manage the task graph for a view.
-		FVisibilityViewPacket& ViewPacket = ViewPackets.Emplace_GetRef(*this, Scene, Views[ViewIndex], ViewIndex);
+		FVisibilityViewPacket& ViewPacket = ViewPackets.Emplace_GetRef(*this, Scene, *Views[ViewIndex], ViewIndex);
 
 		// Each respective stage of visibility for each view is connected to its corresponding task event in order to track when all views complete each stage.
 		Tasks.LightVisibility.AddPrerequisites(ViewPacket.Tasks.LightVisibility);
@@ -3483,7 +3483,7 @@ void FVisibilityTaskData::LaunchVisibilityTasks()
 
 	if (TaskConfig.Schedule == EVisibilityTaskSchedule::Parallel)
 	{
-		if (Views.Num() == 1 && !Views[0].bIsMultiViewportEnabled)
+		if (Views.Num() == 1 && !Views[0]->bIsMultiViewportEnabled)
 		{
 			// When using a single view, dynamic mesh elements are pushed into a pipe that is executed on the render thread which allows for some overlap with compute relevance work.
 			DynamicMeshElements.CommandPipe = Allocator.Create<TCommandPipe<FDynamicPrimitiveIndexList>>(TEXT("GatherDynamicMeshElements"), ENamedThreads::GetRenderThread_Local());
@@ -3568,16 +3568,16 @@ void FVisibilityTaskData::LaunchVisibilityTasks()
 	for (int32 ViewIndex = 0; ViewIndex < Views.Num(); ViewIndex++)
 	{
 		MeshCollector.AddViewMeshArrays(
-			&Views[ViewIndex],
-			&Views[ViewIndex].DynamicMeshElements,
-			&Views[ViewIndex].SimpleElementCollector,
-			&Views[ViewIndex].DynamicPrimitiveCollector,
+			Views[ViewIndex],
+			&Views[ViewIndex]->DynamicMeshElements,
+			&Views[ViewIndex]->SimpleElementCollector,
+			&Views[ViewIndex]->DynamicPrimitiveCollector,
 			Scene.GetFeatureLevel(),
 			GlobalDynamicBuffers.Index,
 			GlobalDynamicBuffers.Vertex,
 			GlobalDynamicBuffers.Read
 #if UE_ENABLE_DEBUG_DRAWING
-			, &Views[ViewIndex].DebugSimpleElementCollector
+			, &Views[ViewIndex]->DebugSimpleElementCollector
 #endif
 		);
 	}
@@ -3589,16 +3589,16 @@ void FVisibilityTaskData::LaunchVisibilityTasks()
 		for (int32 ViewIndex = 0; ViewIndex < Views.Num(); ViewIndex++)
 		{
 			EditorMeshCollector.AddViewMeshArrays(
-				&Views[ViewIndex],
-				&Views[ViewIndex].DynamicEditorMeshElements,
-				&Views[ViewIndex].EditorSimpleElementCollector,
-				&Views[ViewIndex].DynamicPrimitiveCollector,
+				Views[ViewIndex],
+				&Views[ViewIndex]->DynamicEditorMeshElements,
+				&Views[ViewIndex]->EditorSimpleElementCollector,
+				&Views[ViewIndex]->DynamicPrimitiveCollector,
 				Scene.GetFeatureLevel(),
 				GlobalDynamicBuffers.Index,
 				GlobalDynamicBuffers.Vertex,
 				GlobalDynamicBuffers.Read
 #if UE_ENABLE_DEBUG_DRAWING
-				, &Views[ViewIndex].DebugSimpleElementCollector
+				, &Views[ViewIndex]->DebugSimpleElementCollector
 #endif
 			);
 		}
@@ -3658,12 +3658,12 @@ void FVisibilityTaskData::MergeSecondaryViewVisibility()
 	// When using instanced multi-view, any primitive that is marked visible in ANY view must be marked visible in its primary view.
 	for (int32 ViewIndex = 0; ViewIndex < Views.Num(); ++ViewIndex)
 	{
-		FViewInfo& View = Views[ViewIndex];
+		FViewInfo& View = *Views[ViewIndex];
 
 		if (View.StereoPass == EStereoscopicPass::eSSP_SECONDARY)
 		{
 			const FSceneBitArray& SecondaryVisibilityMap = View.PrimitiveVisibilityMap;
-			FSceneBitArray& PrimaryVisibilityMap = Views[View.PrimaryViewIndex].PrimitiveVisibilityMap;
+			FSceneBitArray& PrimaryVisibilityMap = Views[View.PrimaryViewIndex]->PrimitiveVisibilityMap;
 			check(PrimaryVisibilityMap.Num() == SecondaryVisibilityMap.Num());
 
 			const uint32 NumWords = FMath::DivideAndRoundUp(PrimaryVisibilityMap.Num(), NumBitsPerDWORD);
@@ -3706,11 +3706,11 @@ void FVisibilityTaskData::GatherDynamicMeshElementsForPrimitive(int32 PrimitiveI
 	FPrimitiveSceneInfo* PrimitiveSceneInfo = Scene.Primitives[PrimitiveIndex];
 	MeshCollector.SetPrimitive(PrimitiveSceneInfo->Proxy, PrimitiveSceneInfo->DefaultDynamicHitProxyId);
 
-	PrimitiveSceneInfo->Proxy->GetDynamicMeshElements(ViewFamily.Views, ViewFamily, ViewMask, MeshCollector);
+	PrimitiveSceneInfo->Proxy->GetDynamicMeshElements(ViewFamily.AllViews, ViewFamily, ViewMask, MeshCollector);
 
 	for (int32 ViewIndex = 0; ViewIndex < Views.Num(); ViewIndex++)
 	{
-		FViewInfo& View = Views[ViewIndex];
+		FViewInfo& View = *Views[ViewIndex];
 		if (ViewMask & (1 << ViewIndex))
 		{
 			FDynamicPrimitive& DynamicPrimitive = DynamicMeshElements.DynamicPrimitives.Emplace_GetRef();
@@ -3731,7 +3731,7 @@ void FVisibilityTaskData::GatherDynamicMeshElementsForEditorPrimitive(int32 Prim
 	FPrimitiveSceneInfo* PrimitiveSceneInfo = Scene.Primitives[PrimitiveIndex];
 	EditorMeshCollector.SetPrimitive(PrimitiveSceneInfo->Proxy, PrimitiveSceneInfo->DefaultDynamicHitProxyId);
 
-	PrimitiveSceneInfo->Proxy->GetDynamicMeshElements(ViewFamily.Views, ViewFamily, ViewMask, EditorMeshCollector);
+	PrimitiveSceneInfo->Proxy->GetDynamicMeshElements(ViewFamily.AllViews, ViewFamily, ViewMask, EditorMeshCollector);
 }
 
 void FVisibilityTaskData::GatherDynamicMeshElements(const FDynamicPrimitiveIndexList& DynamicPrimitiveIndexList)
@@ -3762,11 +3762,11 @@ void FVisibilityTaskData::GatherDynamicMeshElements(const FDynamicPrimitiveViewM
 
 		for (int32 ViewIndex = 0; ViewIndex < Views.Num(); ViewIndex++)
 		{
-			const FViewInfo& View = Views[ViewIndex];
+			const FViewInfo& View = *Views[ViewIndex];
 
 			if (ViewMask & (1 << ViewIndex) && IStereoRendering::IsASecondaryView(View))
 			{
-				ViewMaskFinal |= 1 << Views[ViewIndex].PrimaryViewIndex;
+				ViewMaskFinal |= 1 << Views[ViewIndex]->PrimaryViewIndex;
 			}
 		}
 
@@ -3806,16 +3806,16 @@ void FVisibilityTaskData::SetupMeshPasses(FExclusiveDepthStencil::Type BasePassD
 	{
 		SCOPED_NAMED_EVENT(DynamicRelevance, FColor::Magenta);
 
-		for (FViewInfo& View : Views)
+		for (FViewInfo* View : Views)
 		{
-			View.DynamicMeshElementsPassRelevance.SetNum(View.DynamicMeshElements.Num());
+			View->DynamicMeshElementsPassRelevance.SetNum(View->DynamicMeshElements.Num());
 		}
 
 		for (FDynamicPrimitive DynamicPrimitive : DynamicMeshElements.DynamicPrimitives)
 		{
 			FPrimitiveSceneInfo* PrimitiveSceneInfo = Scene.Primitives[DynamicPrimitive.PrimitiveIndex];
 			const FPrimitiveBounds& Bounds = Scene.PrimitiveBounds[DynamicPrimitive.PrimitiveIndex];
-			FViewInfo& View = Views[DynamicPrimitive.ViewIndex];
+			FViewInfo& View = *Views[DynamicPrimitive.ViewIndex];
 			const FPrimitiveViewRelevance& ViewRelevance = View.PrimitiveViewRelevanceMap[DynamicPrimitive.PrimitiveIndex];
 
 			for (int32 ElementIndex = DynamicPrimitive.StartElementIndex; ElementIndex < DynamicPrimitive.EndElementIndex; ++ElementIndex)
@@ -3830,7 +3830,7 @@ void FVisibilityTaskData::SetupMeshPasses(FExclusiveDepthStencil::Type BasePassD
 
 	for (int32 ViewIndex = 0; ViewIndex < Views.Num(); ViewIndex++)
 	{
-		FViewInfo& View = Views[ViewIndex];
+		FViewInfo& View = *Views[ViewIndex];
 		View.MeshDecalBatches.Sort();
 
 	#if WITH_EDITOR
@@ -4132,9 +4132,9 @@ void FSceneRenderer::PrepareViewStateForVisibility(const FSceneTexturesConfig& S
 	uint32 CurrentPathTracingInvalidationCounter = Scene->PathTracingInvalidationCounter.Load();
 
 	// Setup motion blur parameters (also check for camera movement thresholds)
-	for(int32 ViewIndex = 0;ViewIndex < Views.Num();ViewIndex++)
+	for(int32 ViewIndex = 0;ViewIndex < AllViews.Num();ViewIndex++)
 	{
-		FViewInfo& View = Views[ViewIndex];
+		FViewInfo& View = *AllViews[ViewIndex];
 		FSceneViewState* ViewState = View.ViewState;
 
 		check(View.VerifyMembersChecks());
@@ -4933,6 +4933,14 @@ void FDeferredShadingSceneRenderer::BeginInitViews(
 			UpdateHairResources(GraphBuilder, View);
 			View.InitRHIResources();
 		}
+
+		for (FSceneCaptureRenderPassInfo& PassInfo : SceneCaptureRenderPassInfos)
+		{
+			for (FViewInfo& View : PassInfo.Views)
+			{
+				View.InitRHIResources();
+			}
+		}
 	}
 
 	if (GDynamicRHI->RHIIncludeOptionalFlushes())
diff --git a/Engine/Source/Runtime/Renderer/Private/SceneVisibilityPrivate.h b/Engine/Source/Runtime/Renderer/Private/SceneVisibilityPrivate.h
--- a/Engine/Source/Runtime/Renderer/Private/SceneVisibilityPrivate.h
+++ b/Engine/Source/Runtime/Renderer/Private/SceneVisibilityPrivate.h
@@ -248,7 +248,7 @@ enum class EVisibilityTaskSchedule
 class FVisibilityTaskConfig
 {
 public:
-	FVisibilityTaskConfig(const FScene& Scene, TConstArrayView<FViewInfo> Views);
+	FVisibilityTaskConfig(const FScene& Scene, TConstArrayView<FViewInfo*> Views);
 
 	EVisibilityTaskSchedule Schedule;
 
@@ -451,7 +451,7 @@ private:
 	FRHICommandListImmediate& RHICmdList;
 	FSceneRenderer& SceneRenderer;
 	FScene& Scene;
-	TArrayView<FViewInfo> Views;
+	TArrayView<FViewInfo*> Views;
 	FViewFamilyInfo& ViewFamily;
 	EShadingPath ShadingPath;
 	FGlobalDynamicBuffers GlobalDynamicBuffers;
diff --git a/Engine/Source/Runtime/Renderer/Private/VelocityRendering.cpp b/Engine/Source/Runtime/Renderer/Private/VelocityRendering.cpp
--- a/Engine/Source/Runtime/Renderer/Private/VelocityRendering.cpp
+++ b/Engine/Source/Runtime/Renderer/Private/VelocityRendering.cpp
@@ -246,6 +246,7 @@ END_SHADER_PARAMETER_STRUCT()
 
 void FSceneRenderer::RenderVelocities(
 	FRDGBuilder& GraphBuilder,
+	TArrayView<FViewInfo> InViews,
 	const FSceneTextures& SceneTextures,
 	EVelocityPass VelocityPass,
 	bool bForceVelocity)
@@ -271,9 +272,9 @@ void FSceneRenderer::RenderVelocities(
 														? FExclusiveDepthStencil::DepthRead_StencilWrite
 														: FExclusiveDepthStencil::DepthWrite_StencilWrite;
 
-	for (int32 ViewIndex = 0; ViewIndex < Views.Num(); ViewIndex++)
+	for (int32 ViewIndex = 0; ViewIndex < InViews.Num(); ViewIndex++)
 	{
-		FViewInfo& View = Views[ViewIndex];
+		FViewInfo& View = InViews[ViewIndex];
 
 		if (View.ShouldRenderView())
 		{
@@ -361,7 +362,7 @@ void FSceneRenderer::RenderVelocities(
 			ERenderTargetLoadAction::ELoad,
 			ExclusiveDepthStencil);
 
-		StampDeferredDebugProbeVelocityPS(GraphBuilder, Views, VelocityRenderTargets);
+		StampDeferredDebugProbeVelocityPS(GraphBuilder, InViews, VelocityRenderTargets);
 	}
 #endif
 
diff --git a/Engine/Source/Runtime/Renderer/Public/SceneCaptureRendering.h b/Engine/Source/Runtime/Renderer/Public/SceneCaptureRendering.h
--- a/Engine/Source/Runtime/Renderer/Public/SceneCaptureRendering.h
+++ b/Engine/Source/Runtime/Renderer/Public/SceneCaptureRendering.h
@@ -14,6 +14,13 @@ class FSceneViewFamily;
 class FViewInfo;
 struct FMinimalSceneTextures;
 
+RENDERER_API void CopySceneCaptureComponentToTarget(
+	FRDGBuilder& GraphBuilder,
+	const FMinimalSceneTextures& SceneTextures,
+	FRDGTextureRef ViewFamilyTexture,
+	const FSceneViewFamily& ViewFamily,
+	const TArray<const FViewInfo*>& Views);
+
 RENDERER_API void CopySceneCaptureComponentToTarget(
 	FRDGBuilder& GraphBuilder,
 	const FMinimalSceneTextures& SceneTextures,
