Fix missing velocity texure for distortion.

diff --git a/Engine/Source/Runtime/Renderer/Private/DeferredShadingRenderer.cpp b/Engine/Source/Runtime/Renderer/Private/DeferredShadingRenderer.cpp
--- a/Engine/Source/Runtime/Renderer/Private/DeferredShadingRenderer.cpp
+++ b/Engine/Source/Runtime/Renderer/Private/DeferredShadingRenderer.cpp
@@ -3252,6 +3252,7 @@ void FDeferredShadingSceneRenderer::Render(FRDGBuilder& GraphBuilder)
 	const EShaderPlatform Platform = GetViewFamilyInfo(Views).GetShaderPlatform();
 	const bool bBasePassCanOutputVelocity = FVelocityRendering::BasePassCanOutputVelocity(Platform);
 	const bool bHairStrandsEnable = HairStrandsBookmarkParameters.HasInstances() && Views.Num() > 0 && IsHairStrandsEnabled(EHairStrandsShaderType::Strands, Platform);
+	const bool bForceVelocityOutput = bHairStrandsEnable || ShouldRenderDistortion();
 
 	auto RenderPrepassAndVelocity = [&](auto& InViews, auto& InNaniteBasePassVisibility, auto& NaniteRasterResults, auto& PrimaryNaniteViews)
 	{
@@ -3285,7 +3286,7 @@ void FDeferredShadingSceneRenderer::Render(FRDGBuilder& GraphBuilder)
 			{
 				// Render the velocities of movable objects
 				GraphBuilder.SetCommandListStat(GET_STATID(STAT_CLM_Velocity));
-				RenderVelocities(GraphBuilder, InViews, SceneTextures, EVelocityPass::Opaque, bHairStrandsEnable);
+				RenderVelocities(GraphBuilder, InViews, SceneTextures, EVelocityPass::Opaque, bForceVelocityOutput);
 				GraphBuilder.SetCommandListStat(GET_STATID(STAT_CLM_AfterVelocity));
 			}
 		}
diff --git a/Engine/Source/Runtime/Renderer/Private/VelocityRendering.cpp b/Engine/Source/Runtime/Renderer/Private/VelocityRendering.cpp
--- a/Engine/Source/Runtime/Renderer/Private/VelocityRendering.cpp
+++ b/Engine/Source/Runtime/Renderer/Private/VelocityRendering.cpp
@@ -211,7 +211,9 @@ bool FDeferredShadingSceneRenderer::ShouldRenderVelocities() const
 		bool bSSGI = ViewPipelineState.DiffuseIndirectMethod == EDiffuseIndirectMethod::SSGI;
 		bool bLumen = ViewPipelineState.DiffuseIndirectMethod == EDiffuseIndirectMethod::Lumen || ViewPipelineState.ReflectionsMethod == EReflectionsMethod::Lumen;
 		
-		bNeedsVelocity |= bVisualizeMotionblur || bMotionBlur || bTemporalAA || bDistanceFieldAO || bSSRTemporal || bDenoise || bSSGI || bLumen;
+		bool bDistortion = ShouldRenderDistortion();
+
+		bNeedsVelocity |= bVisualizeMotionblur || bMotionBlur || bTemporalAA || bDistanceFieldAO || bSSRTemporal || bDenoise || bSSGI || bLumen || bDistortion;
 	}
 
 	return bNeedsVelocity;
