Custom WaterInfoTexture rendering. This adds a new path for rendering the WaterInfoTexture without SceneCaptures, massively improving rendering thread timings. Set r.Water.WaterInfo.RenderMethod to 1 to activate or to 0 to revert to the legacy scene capture path. This implementation is currently using cached draw commands + new EMeshPass entries, but it might be worth it to use DrawDynamicMeshPass() in the future.

#rnx

diff --git a/Engine/Plugins/Experimental/Water/Source/Runtime/Private/WaterBodyActor.cpp b/Engine/Plugins/Experimental/Water/Source/Runtime/Private/WaterBodyActor.cpp
--- a/Engine/Plugins/Experimental/Water/Source/Runtime/Private/WaterBodyActor.cpp
+++ b/Engine/Plugins/Experimental/Water/Source/Runtime/Private/WaterBodyActor.cpp
@@ -61,10 +61,12 @@ AWaterBody::AWaterBody(const FObjectInitializer& ObjectInitializer)
 	WaterInfoMeshComponent = CreateDefaultSubobject<UWaterBodyInfoMeshComponent>(TEXT("WaterInfoMeshComponent"));
 	WaterInfoMeshComponent->SetMobility(EComponentMobility::Static);
 	WaterInfoMeshComponent->SetupAttachment(RootComponent);
+	WaterInfoMeshComponent->bIsDilatedMesh = false;
 
 	DilatedWaterInfoMeshComponent = CreateDefaultSubobject<UWaterBodyInfoMeshComponent>(TEXT("DilatedWaterInfoMeshComponent"));
 	DilatedWaterInfoMeshComponent->SetMobility(EComponentMobility::Static);
 	DilatedWaterInfoMeshComponent->SetupAttachment(RootComponent);
+	DilatedWaterInfoMeshComponent->bIsDilatedMesh = true;
 
 #if WITH_EDITORONLY_DATA
 	bAffectsLandscape_DEPRECATED = true;
diff --git a/Engine/Plugins/Experimental/Water/Source/Runtime/Private/WaterBodyInfoMeshComponent.cpp b/Engine/Plugins/Experimental/Water/Source/Runtime/Private/WaterBodyInfoMeshComponent.cpp
--- a/Engine/Plugins/Experimental/Water/Source/Runtime/Private/WaterBodyInfoMeshComponent.cpp
+++ b/Engine/Plugins/Experimental/Water/Source/Runtime/Private/WaterBodyInfoMeshComponent.cpp
@@ -38,11 +38,12 @@ FPrimitiveSceneProxy* UWaterBodyInfoMeshComponent::CreateSceneProxy()
 	{
 		return nullptr;
 	}
-	return new FWaterBodyInfoMeshSceneProxy(this);
+	return new FWaterBodyInfoMeshSceneProxy(this, bIsDilatedMesh);
 }
 
-FWaterBodyInfoMeshSceneProxy::FWaterBodyInfoMeshSceneProxy(UWaterBodyInfoMeshComponent* Component)
-	: FStaticMeshSceneProxy(Component, true)
+FWaterBodyInfoMeshSceneProxy::FWaterBodyInfoMeshSceneProxy(UWaterBodyInfoMeshComponent* Component, bool InbIsDilatedMesh)
+	: FStaticMeshSceneProxy(Component, true),
+	bIsDilatedMesh(InbIsDilatedMesh)
 {
 	// Disable Notify on WorldAddRemove. This prevents the component from being unhidden in FPrimitiveSceneProxy::OnLevelAddedToWorld_RenderThread if it was part of a streamed level.
 	// WaterInfo proxies should only be unhidden during WaterInfo passes.
@@ -51,6 +52,16 @@ FWaterBodyInfoMeshSceneProxy::FWaterBodyInfoMeshSceneProxy(UWaterBodyInfoMeshCom
 	SetEnabled(false);
 }
 
+bool FWaterBodyInfoMeshSceneProxy::GetMeshElement(int32 LODIndex, int32 BatchIndex, int32 ElementIndex, uint8 InDepthPriorityGroup, bool bUseSelectionOutline, bool bAllowPreCulledIndices, FMeshBatch& OutMeshBatch) const
+{
+	bool bResult = FStaticMeshSceneProxy::GetMeshElement(LODIndex, BatchIndex, ElementIndex, InDepthPriorityGroup, bUseSelectionOutline, bAllowPreCulledIndices, OutMeshBatch);
+	if (bResult)
+	{
+		OutMeshBatch.bUseForWaterInfoTextureDepth = bIsDilatedMesh; // The dilated mesh is drawn in the water info texture depth-only pass
+	}
+	return bResult;
+}
+
 void FWaterBodyInfoMeshSceneProxy::SetEnabled(bool bInEnabled)
 {
 	SetForceHidden(!(bInEnabled || CVarShowWaterInfoSceneProxies.GetValueOnAnyThread()));
diff --git a/Engine/Plugins/Experimental/Water/Source/Runtime/Private/WaterBodyInfoMeshComponent.h b/Engine/Plugins/Experimental/Water/Source/Runtime/Private/WaterBodyInfoMeshComponent.h
--- a/Engine/Plugins/Experimental/Water/Source/Runtime/Private/WaterBodyInfoMeshComponent.h
+++ b/Engine/Plugins/Experimental/Water/Source/Runtime/Private/WaterBodyInfoMeshComponent.h
@@ -23,12 +23,26 @@ public:
 	virtual FPrimitiveSceneProxy* CreateSceneProxy() override;
 	
 	virtual bool IsHLODRelevant() const { return false; }
+
+	// Whether the water body mesh has been enlarged/dilated past the shoreline of the actual water body.
+	// For every water body, there is a dilated and non-dilated mesh. Both are needed for generating the water info texture.
+	UPROPERTY()
+	bool bIsDilatedMesh = false;
 };
 
 struct FWaterBodyInfoMeshSceneProxy : public FStaticMeshSceneProxy
 {
 public:
-	FWaterBodyInfoMeshSceneProxy(UWaterBodyInfoMeshComponent* Component);
+	FWaterBodyInfoMeshSceneProxy(UWaterBodyInfoMeshComponent* Component, bool InbIsDilatedMesh);
+
+	virtual bool GetMeshElement(
+		int32 LODIndex,
+		int32 BatchIndex,
+		int32 ElementIndex,
+		uint8 InDepthPriorityGroup,
+		bool bUseSelectionOutline,
+		bool bAllowPreCulledIndices,
+		FMeshBatch& OutMeshBatch) const override;
 
 	SIZE_T GetTypeHash() const
 	{
@@ -37,5 +51,8 @@ public:
 	}
 
 	void SetEnabled(bool bInEnabled);
+
+private:
+	bool bIsDilatedMesh = false;
 };
 
diff --git a/Engine/Plugins/Experimental/Water/Source/Runtime/Private/WaterInfoRendering.cpp b/Engine/Plugins/Experimental/Water/Source/Runtime/Private/WaterInfoRendering.cpp
--- a/Engine/Plugins/Experimental/Water/Source/Runtime/Private/WaterInfoRendering.cpp
+++ b/Engine/Plugins/Experimental/Water/Source/Runtime/Private/WaterInfoRendering.cpp
@@ -40,16 +40,6 @@ static FAutoConsoleVariableRef CVarWaterInfoRenderLandscapeMinimumMipLevel(
 	WaterInfoRenderLandscapeMinimumMipLevel,
 	TEXT("Clamps the minimum allowed mip level for the landscape when rendering the water info texture. Used on the lowest end platforms which cannot support rendering all the landscape vertices at the highest LOD."));
 
-static TAutoConsoleVariable<float> CVarWaterInfoUndergroundDilationDepthOffset(
-		TEXT("r.Water.WaterInfo.UndergroundDilationDepthOffset"),
-		64.f,
-		TEXT("The minimum distance below the ground when we allow dilation to write on top of water"));
-
-static TAutoConsoleVariable<float> CVarWaterInfoDilationOverwriteMinimumDistance(
-		TEXT("r.Water.WaterInfo.DilationOverwriteMinimumDistance"),
-		128.f,
-		TEXT("The minimum distance below the ground when we allow dilation to write on top of water"));
-
 namespace UE::WaterInfo
 {
 
@@ -129,6 +119,9 @@ static void MergeWaterInfoAndDepth(
 	GraphicsPSOInit.BlendState = TStaticBlendState<>::GetRHI();
 
 	{
+		static auto* CVarDilationOverwriteMinimumDistance = IConsoleManager::Get().FindTConsoleVariableDataFloat(TEXT("r.Water.WaterInfo.DilationOverwriteMinimumDistance"));
+		static auto* CVarUndergroundDilationDepthOffset = IConsoleManager::Get().FindTConsoleVariableDataFloat(TEXT("r.Water.WaterInfo.UndergroundDilationDepthOffset"));
+
 		FWaterInfoMergePS::FParameters* PassParameters = GraphBuilder.AllocParameters<FWaterInfoMergePS::FParameters>();
 		PassParameters->View = View.ViewUniformBuffer;
 		PassParameters->RenderTargets[0] = FRenderTargetBinding(OutputTexture, ERenderTargetLoadAction::ENoAction);
@@ -142,8 +135,8 @@ static void MergeWaterInfoAndDepth(
 		PassParameters->CaptureZ = Params.CaptureZ;
 		PassParameters->WaterHeightExtents = Params.WaterHeightExtents;
 		PassParameters->GroundZMin = Params.GroundZMin;
-		PassParameters->DilationOverwriteMinimumDistance = CVarWaterInfoDilationOverwriteMinimumDistance.GetValueOnRenderThread();
-		PassParameters->UndergroundDilationDepthOffset = CVarWaterInfoUndergroundDilationDepthOffset.GetValueOnRenderThread();
+		PassParameters->DilationOverwriteMinimumDistance = CVarDilationOverwriteMinimumDistance ? CVarDilationOverwriteMinimumDistance->GetValueOnRenderThread() : 128.0f;
+		PassParameters->UndergroundDilationDepthOffset = CVarUndergroundDilationDepthOffset ? CVarUndergroundDilationDepthOffset->GetValueOnRenderThread() : 64.0f;
 
 		TShaderMapRef<FScreenVS> VertexShader(ShaderMap);
 		TShaderMapRef<FWaterInfoMergePS> PixelShader(ShaderMap);
@@ -829,4 +822,79 @@ void UpdateWaterInfoRendering(
 		});
 }
 
+void UpdateWaterInfoRendering2(FSceneView& InView, const TMap<AWaterZone*, UE::WaterInfo::FRenderingContext>& WaterInfoContexts)
+{
+	TRACE_CPUPROFILER_EVENT_SCOPE(WaterInfo::UpdateWaterInfoRendering2);
+
+	InView.WaterInfoTextureRenderingParams.Reset();
+	for (const TPair<AWaterZone*, UE::WaterInfo::FRenderingContext>& Pair : WaterInfoContexts)
+	{
+		const UE::WaterInfo::FRenderingContext& Context(Pair.Value);
+
+		RenderCaptureInterface::FScopedCapture RenderCapture((RenderCaptureNextWaterInfoDraws != 0), TEXT("RenderWaterInfo"));
+		RenderCaptureNextWaterInfoDraws = FMath::Max(0, RenderCaptureNextWaterInfoDraws - 1);
+
+		if (!IsValid(Context.TextureRenderTarget))
+		{
+			continue;
+		}
+		const FVector ZoneExtent = Context.ZoneToRender->GetDynamicWaterInfoExtent();
+
+		FVector ViewLocation = Context.ZoneToRender->GetDynamicWaterInfoCenter();
+		ViewLocation.Z = Context.CaptureZ;
+
+		const FBox2D CaptureBounds(FVector2D(ViewLocation - ZoneExtent), FVector2D(ViewLocation + ZoneExtent));
+
+		// Zone rendering always happens facing towards negative z.
+		const FVector LookAt = ViewLocation - FVector(0.f, 0.f, 1.f);
+
+		FSceneView::FWaterInfoTextureRenderingParams RenderingParams;
+		RenderingParams.RenderTarget = Context.TextureRenderTarget->GameThread_GetRenderTargetResource();
+		RenderingParams.ViewLocation = ViewLocation;
+		RenderingParams.ViewRotationMatrix = FLookAtMatrix(ViewLocation, LookAt, FVector(0.f, -1.f, 0.f));
+		RenderingParams.ViewRotationMatrix = RenderingParams.ViewRotationMatrix.RemoveTranslation();
+		RenderingParams.ViewRotationMatrix.RemoveScaling();
+		RenderingParams.ProjectionMatrix = BuildOrthoMatrix(ZoneExtent.X, ZoneExtent.Y);
+		RenderingParams.CaptureZ = ViewLocation.Z;
+		RenderingParams.WaterHeightExtents = Context.ZoneToRender->GetWaterHeightExtents();
+		RenderingParams.GroundZMin = Context.ZoneToRender->GetGroundZMin();
+		RenderingParams.VelocityBlurRadius = Context.ZoneToRender->GetVelocityBlurRadius();
+		RenderingParams.WaterZoneExtents = ZoneExtent;
+
+		if (Context.GroundPrimitiveComponents.Num() > 0)
+		{
+			RenderingParams.TerrainComponentIds.Reserve(Context.GroundPrimitiveComponents.Num());
+			for (TWeakObjectPtr<UPrimitiveComponent> GroundPrimComp : Context.GroundPrimitiveComponents)
+			{
+				if (GroundPrimComp.IsValid())
+				{
+					RenderingParams.TerrainComponentIds.Add(GroundPrimComp.Get()->GetPrimitiveSceneId());
+				}
+			}
+		}
+		if (Context.WaterBodies.Num() > 0)
+		{
+			RenderingParams.WaterBodyComponentIds.Reserve(Context.WaterBodies.Num());
+			RenderingParams.DilatedWaterBodyComponentIds.Reserve(Context.WaterBodies.Num());
+			for (const UWaterBodyComponent* WaterBodyToRender : Context.WaterBodies)
+			{
+				if (!IsValid(WaterBodyToRender))
+				{
+					continue;
+				}
+
+				// Perform our own simple culling based on the known Capture bounds:
+				const FBox WaterBodyBounds = WaterBodyToRender->Bounds.GetBox();
+				if (CaptureBounds.Intersect(FBox2D(FVector2D(WaterBodyBounds.Min), FVector2D(WaterBodyBounds.Max))))
+				{
+					RenderingParams.WaterBodyComponentIds.Add(WaterBodyToRender->GetWaterInfoMeshComponent()->GetPrimitiveSceneId());
+					RenderingParams.DilatedWaterBodyComponentIds.Add(WaterBodyToRender->GetDilatedWaterInfoMeshComponent()->GetPrimitiveSceneId());
+				}
+			}
+		}
+
+		InView.WaterInfoTextureRenderingParams.Add(MoveTemp(RenderingParams));
+	}
+}
+
 } // namespace WaterInfo
diff --git a/Engine/Plugins/Experimental/Water/Source/Runtime/Private/WaterInfoRendering.h b/Engine/Plugins/Experimental/Water/Source/Runtime/Private/WaterInfoRendering.h
--- a/Engine/Plugins/Experimental/Water/Source/Runtime/Private/WaterInfoRendering.h
+++ b/Engine/Plugins/Experimental/Water/Source/Runtime/Private/WaterInfoRendering.h
@@ -9,6 +9,7 @@ class UWaterBodyComponent;
 class FSceneInterface;
 class UTextureRenderTarget2D;
 class UPrimitiveComponent;
+class FSceneView;
 
 namespace UE::WaterInfo
 {
@@ -24,4 +25,8 @@ namespace UE::WaterInfo
 void UpdateWaterInfoRendering(
 	FSceneInterface* Scene,
 	const FRenderingContext& Context);
+
+void UpdateWaterInfoRendering2(
+	FSceneView& InView, 
+	const TMap<AWaterZone*, UE::WaterInfo::FRenderingContext>& WaterInfoContexts);
 }
diff --git a/Engine/Plugins/Experimental/Water/Source/Runtime/Private/WaterViewExtension.cpp b/Engine/Plugins/Experimental/Water/Source/Runtime/Private/WaterViewExtension.cpp
--- a/Engine/Plugins/Experimental/Water/Source/Runtime/Private/WaterViewExtension.cpp
+++ b/Engine/Plugins/Experimental/Water/Source/Runtime/Private/WaterViewExtension.cpp
@@ -25,6 +25,12 @@ static TAutoConsoleVariable<float> CVarLocalTessellationUpdateMargin(
 	TEXT("If the view is less than UpdateMargin units away from the edge, it moves the sliding window forward."),
 	ECVF_Default);
 
+static TAutoConsoleVariable<int32> CVarWaterInfoRenderMethod(
+	TEXT("r.Water.WaterInfo.RenderMethod"),
+	1,
+	TEXT("0: SceneCaptures, 1: Custom"),
+	ECVF_Default);
+
 // ----------------------------------------------------------------------------------
 
 FWaterViewExtension::FWaterViewExtension(const FAutoRegister& AutoReg, UWorld* InWorld)
@@ -300,6 +306,8 @@ void FWaterViewExtension::SetupView(FSceneViewFamily& InViewFamily, FSceneView&
 	bUpdatingWaterInfo = true;
 	ON_SCOPE_EXIT { bUpdatingWaterInfo = false; };
 
+	const int32 WaterInfoRenderMethod = CVarWaterInfoRenderMethod.GetValueOnGameThread();
+
 	for (const TPair<AWaterZone*, UE::WaterInfo::FRenderingContext>& Pair : WaterInfoContextsToRender)
 	{
 		AWaterZone* WaterZone = Pair.Key;
@@ -342,8 +350,18 @@ void FWaterViewExtension::SetupView(FSceneViewFamily& InViewFamily, FSceneView&
 			MarkGPUDataDirty();
 		}
 
-		const UE::WaterInfo::FRenderingContext& Context(Pair.Value);
-		UE::WaterInfo::UpdateWaterInfoRendering(WorldPtr.Get()->Scene, Context);
+		// Old method of rendering the water info texture; uses scene captures
+		if (WaterInfoRenderMethod == 0)
+		{
+			const UE::WaterInfo::FRenderingContext& Context(Pair.Value);
+			UE::WaterInfo::UpdateWaterInfoRendering(WorldPtr.Get()->Scene, Context);
+		}
+	}
+
+	// New method of rendering the water info texture; rendering is done in a separate pass when rendering the main view
+	if (WaterInfoRenderMethod != 0)
+	{
+		UE::WaterInfo::UpdateWaterInfoRendering2(InView, WaterInfoContextsToRender);
 	}
 
 	WaterInfoContextsToRender.Empty();
diff --git a/Engine/Shaders/Private/WaterInfoTextureBlur.usf b/Engine/Shaders/Private/WaterInfoTextureBlur.usf
new file mode 100644
--- /dev/null
+++ b/Engine/Shaders/Private/WaterInfoTextureBlur.usf
@@ -0,0 +1,75 @@
+// Copyright Epic Games, Inc. All Rights Reserved.
+
+#include "Common.ush"
+
+Texture2D<float4> WaterInfoTexture;
+
+float WaterZMin;
+float WaterZMax;
+float GroundZMin;
+float CaptureZ;
+
+int BlurRadius;
+
+struct FWaterInfo
+{
+	float2 Velocity;
+	float NormalizedWaterZ;
+	float NormalizedGroundZ;
+};
+
+FWaterInfo MakeWaterInfo(float4 InSampleData)
+{
+	FWaterInfo Result;
+	Result.Velocity = InSampleData.rg;
+	Result.NormalizedWaterZ = InSampleData.b;
+	Result.NormalizedGroundZ = InSampleData.a;
+	return Result;
+}
+
+float2 VelocityBlur(float Radius, uint2 TextureCoordinates)
+{
+	float2 Accum = 0;
+
+	float Neighbors = 0;
+
+	int2 CurrentTextureCoordinates = int2(TextureCoordinates);
+
+	for (int X = -Radius; X <= Radius; X++)
+	{
+		for (int Y = -Radius; Y <= Radius; Y++)
+		{
+			CurrentTextureCoordinates = TextureCoordinates + int2(X, Y);
+
+			const FWaterInfo WaterInfo = MakeWaterInfo(WaterInfoTexture.Load(int3(CurrentTextureCoordinates, 0)));
+
+			const float WaterZ = WaterInfo.NormalizedWaterZ * (WaterZMax - WaterZMin) + WaterZMin;
+			const float GroundZ = WaterInfo.NormalizedGroundZ * (WaterZMax - GroundZMin) + GroundZMin;
+
+			if (WaterZ > GroundZ)
+			{
+				Accum.xy += WaterInfo.Velocity;
+				Neighbors += 1;
+			}
+		}
+	}
+
+	if (Neighbors > 0)
+	{
+		Accum.xy /= Neighbors;
+	}
+
+	return float2(Accum);
+}
+
+void Main(
+	in float4 Position : SV_Position,
+	out float4 OutColor : SV_Target0)
+{
+	uint2 TextureCoordinates = floor(Position.xy);
+
+	const FWaterInfo WaterInfo = MakeWaterInfo(WaterInfoTexture.Load(int3(TextureCoordinates, 0)));
+	const float2 BlurredVelocity = VelocityBlur(BlurRadius, TextureCoordinates);
+
+	OutColor = float4(BlurredVelocity.xy, WaterInfo.NormalizedWaterZ, WaterInfo.NormalizedGroundZ);
+}
diff --git a/Engine/Shaders/Private/WaterInfoTextureMerge.usf b/Engine/Shaders/Private/WaterInfoTextureMerge.usf
new file mode 100644
--- /dev/null
+++ b/Engine/Shaders/Private/WaterInfoTextureMerge.usf
@@ -0,0 +1,100 @@
+// Copyright Epic Games, Inc. All Rights Reserved.
+
+#include "Common.ush"
+
+Texture2D<float4> GroundDepthTexture;
+Texture2D<float4> WaterBodyTexture;
+Texture2D<float4> WaterBodyDepthTexture;
+Texture2D<float4> DilatedWaterBodyDepthTexture;
+
+float2 WaterHeightExtents;
+float GroundZMin;
+float CaptureZ;
+float UndergroundDilationDepthOffset;
+float DilationOverwriteMinimumDistance;
+
+float NormalizeHeight(float Height, float Min, float Max)
+{
+	float NormalizedHeight = 1.f;
+	// It's possible that the min and max are the same value in which case we might end up generating a div by zero error
+	// In this case it's okay to use the Normalized value of 1.0 which will just map to the Min value.
+	if (Min != Max)
+	{
+		NormalizedHeight = saturate((Height - Min) / (Max - Min));
+	}
+	return NormalizedHeight;
+}
+
+void Main(
+	in float4 Position : SV_Position,
+	out float4 OutColor : SV_Target0
+)
+{
+	const float2 WaterBodyVelocity = WaterBodyTexture.Load(int3(Position.xy, 0)).xy;
+	const float WaterBodyDeviceDepth = WaterBodyDepthTexture.Load(int3(Position.xy, 0)).x;
+	const float DilationDeviceDepth = DilatedWaterBodyDepthTexture.Load(int3(Position.xy, 0)).x;
+	const float GroundDeviceDepth = GroundDepthTexture.Load(int3(Position.xy, 0)).x;
+
+	const float GroundSceneDepth = ConvertFromDeviceZ(GroundDeviceDepth);
+	float GroundZ = CaptureZ - GroundSceneDepth;
+
+	const float WaterSceneDepth = ConvertFromDeviceZ(WaterBodyDeviceDepth);
+	const float WaterZ = CaptureZ - WaterSceneDepth;
+
+	const float DilationSceneDepth = ConvertFromDeviceZ(DilationDeviceDepth);
+	const bool bPixelHasDilationData = DilationSceneDepth != 0.0;
+
+	float DilatedZ = CaptureZ - DilationSceneDepth;
+
+
+	// Always initialize output to avoid bugs in old shader compilers 
+	OutColor = 0.f;
+
+	// If the GroundDepth is further than the far plane of the projection, set the GroundZ to some really far away value
+	if (GroundDeviceDepth == 1.0)
+	{
+		GroundZ = GroundZMin;
+	}
+
+	const float WaterZMin = WaterHeightExtents.x;
+	const float WaterZMax = WaterHeightExtents.y;
+
+	// if the measured WaterZ is below WaterZMin it means that there is no water at this pixel because it is not possible
+	// for water to be further than ZMin.
+	bool bPixelHasValidWaterData = WaterZ >= WaterZMin;
+
+	// Additionally, if we have dilation data here and the water data would be placed below the landscape, instead draw the dilation data
+	if (bPixelHasDilationData && (WaterZ + 128.f < GroundZ))
+	{
+		bPixelHasValidWaterData = false;
+	}
+
+	if (bPixelHasValidWaterData)
+	{
+		const float NormalizedWaterZ = NormalizeHeight(WaterZ, WaterZMin, WaterZMax);
+		const float NormalizedGroundZ = NormalizeHeight(GroundZ, GroundZMin, WaterZMax);
+
+		OutColor = float4(WaterBodyVelocity, NormalizedWaterZ, NormalizedGroundZ);
+	}
+	else if (bPixelHasDilationData)
+	{
+
+		// If the dilated water would render on top of terrain, squash it down below
+		if (DilatedZ > GroundZ)
+		{
+			DilatedZ = GroundZ - UndergroundDilationDepthOffset;
+		}
+
+		const float NormalizedWaterZ = NormalizeHeight(DilatedZ, WaterZMin, WaterZMax);
+		const float NormalizedGroundZ = NormalizeHeight(GroundZ, GroundZMin, WaterZMax);
+
+		OutColor = float4(0.0, 0.0, NormalizedWaterZ, NormalizedGroundZ);
+	}
+	else
+	{
+		// Pixel may not have any dilation data nor water data. In this case just write 0s to the texture
+
+		const float NormalizedGroundZ = NormalizeHeight(GroundZ, GroundZMin, WaterZMax);
+		OutColor = float4(0.0, 0.0, 0.0, NormalizedGroundZ);
+	}
+}
diff --git a/Engine/Source/Runtime/Engine/Public/MeshBatch.h b/Engine/Source/Runtime/Engine/Public/MeshBatch.h
--- a/Engine/Source/Runtime/Engine/Public/MeshBatch.h
+++ b/Engine/Source/Runtime/Engine/Public/MeshBatch.h
@@ -412,6 +412,9 @@ struct FMeshBatch
 	 */
 	uint32 bViewDependentArguments : 1;
 
+	/** Whether the mesh batch should be used in the depth-only passes of rendering the water info texture for the water plugin */
+	uint32 bUseForWaterInfoTextureDepth : 1;
+
 #if UE_ENABLE_DEBUG_DRAWING
 	/** Conceptual HLOD index used for the HLOD Coloration visualization. */
 	int8 VisualizeHLODIndex;
@@ -491,6 +494,7 @@ struct FMeshBatch
 	,	CastRayTracedShadow(true)
 #endif
 	,	bViewDependentArguments(false)
+	,	bUseForWaterInfoTextureDepth(false)
 #if (!(UE_BUILD_SHIPPING || UE_BUILD_TEST) || WITH_EDITOR)
 	,	VisualizeHLODIndex(INDEX_NONE)
 #endif
diff --git a/Engine/Source/Runtime/Engine/Public/PSOPrecache.h b/Engine/Source/Runtime/Engine/Public/PSOPrecache.h
--- a/Engine/Source/Runtime/Engine/Public/PSOPrecache.h
+++ b/Engine/Source/Runtime/Engine/Public/PSOPrecache.h
@@ -276,7 +276,7 @@ class FPSOCollectorCreateManager
 {
 public:
 
-	constexpr static uint32 MaxPSOCollectorCount = 34;
+	constexpr static uint32 MaxPSOCollectorCount = 36;
 
 	static PSOCollectorCreateFunction GetCreateFunction(EShadingPath ShadingPath, uint32 Index)
 	{
diff --git a/Engine/Source/Runtime/Engine/Public/SceneView.h b/Engine/Source/Runtime/Engine/Public/SceneView.h
--- a/Engine/Source/Runtime/Engine/Public/SceneView.h
+++ b/Engine/Source/Runtime/Engine/Public/SceneView.h
@@ -1531,6 +1531,23 @@ public:
 	FShaderResourceViewRHIRef WaterIndirectionBuffer;
 	FShaderResourceViewRHIRef WaterDataBuffer;
 
+	struct FWaterInfoTextureRenderingParams
+	{
+		FRenderTarget* RenderTarget = nullptr;
+		FVector ViewLocation = FVector::Zero();
+		FMatrix ViewRotationMatrix = FMatrix(EForceInit::ForceInit);
+		FMatrix ProjectionMatrix = FMatrix(EForceInit::ForceInit);
+		TSet<FPrimitiveComponentId> TerrainComponentIds;
+		TSet<FPrimitiveComponentId> WaterBodyComponentIds;
+		TSet<FPrimitiveComponentId> DilatedWaterBodyComponentIds;
+		FVector WaterZoneExtents = FVector::Zero();
+		FVector2f WaterHeightExtents = FVector2f::ZeroVector;
+		float GroundZMin = 0.0f;
+		float CaptureZ = 0.0f;
+		int32 VelocityBlurRadius = 0;
+	};
+	TArray<FWaterInfoTextureRenderingParams> WaterInfoTextureRenderingParams;
+
 	/** Landscape rendering related data */
 	FShaderResourceViewRHIRef LandscapeIndirectionBuffer;
 	FShaderResourceViewRHIRef LandscapePerComponentDataBuffer;
diff --git a/Engine/Source/Runtime/Landscape/Private/Landscape.cpp b/Engine/Source/Runtime/Landscape/Private/Landscape.cpp
--- a/Engine/Source/Runtime/Landscape/Private/Landscape.cpp
+++ b/Engine/Source/Runtime/Landscape/Private/Landscape.cpp
@@ -1079,6 +1079,9 @@ void ULandscapeComponent::PostLoad()
 		// we need the fixed grid vertex factory for both virtual texturing and grass
 		ERHIFeatureLevel::Type FeatureLevel = GMaxRHIFeatureLevel;
 		bool bNeedsFixedGridVertexFactory = UseVirtualTexturing(FeatureLevel);
+		// This cvar is defined in the water plugin and searching for it should return nullptr if the plugin is not loaded
+		const bool bWaterPluginLoaded = IConsoleManager::Get().FindConsoleVariable(TEXT("r.Water.WaterInfo.RenderMethod")) != nullptr;
+		bNeedsFixedGridVertexFactory |= bWaterPluginLoaded;
 		if (bNeedsFixedGridVertexFactory)
 		{
 			VertexFactoryDataList.Add(FPSOPrecacheVertexFactoryData(&FLandscapeFixedGridVertexFactory::StaticType));
diff --git a/Engine/Source/Runtime/Landscape/Private/LandscapeRender.cpp b/Engine/Source/Runtime/Landscape/Private/LandscapeRender.cpp
--- a/Engine/Source/Runtime/Landscape/Private/LandscapeRender.cpp
+++ b/Engine/Source/Runtime/Landscape/Private/LandscapeRender.cpp
@@ -1376,8 +1376,11 @@ void FLandscapeComponentSceneProxy::CreateRenderThreadResources(FRHICommandListB
 			SharedBuffers->VertexFactory = LandscapeXYOffsetVertexFactory;
 		}
 
-		// we need the fixed grid vertex factory for both virtual texturing and grass : 
+		// We need the fixed grid vertex factory for virtual texturing, grass and for rendering the water info texture : 
 		bool bNeedsFixedGridVertexFactory = UseVirtualTexturing(FeatureLevel);
+		// This cvar is defined in the water plugin and searching for it should return nullptr if the plugin is not loaded
+		const bool bWaterPluginLoaded = IConsoleManager::Get().FindConsoleVariable(TEXT("r.Water.WaterInfo.RenderMethod")) != nullptr;
+		bNeedsFixedGridVertexFactory |= bWaterPluginLoaded;
 
 #if WITH_EDITOR
 		bNeedsFixedGridVertexFactory |= (SharedBuffers->GrassIndexBuffer != nullptr);
@@ -1998,8 +2001,8 @@ void FLandscapeComponentSceneProxy::OnTransformChanged(FRHICommandListBase& RHIC
 	UpdateDefaultInstanceSceneData();
 }
 
-/** Creates a mesh batch for virtual texture rendering. Will render a simple fixed grid with combined subsections. */
-bool FLandscapeComponentSceneProxy::GetMeshElementForVirtualTexture(int32 InLodIndex, ERuntimeVirtualTextureMaterialType MaterialType, FMaterialRenderProxy* InMaterialInterface, FMeshBatch& OutMeshBatch, TArray<FLandscapeBatchElementParams>& OutStaticBatchParamArray) const
+/** Creates a mesh batch for virtual texture or water info texture rendering. The caller is responsible for setting the required flags (bRenderToVirtualTexture, bUseForWaterInfoTextureDepth). Will render a simple fixed grid with combined subsections. */
+bool FLandscapeComponentSceneProxy::GetMeshElementForFixedGrid(int32 InLodIndex, FMaterialRenderProxy* InMaterialInterface, FMeshBatch& OutMeshBatch, TArray<FLandscapeBatchElementParams>& OutStaticBatchParamArray) const
 {
 	if (InMaterialInterface == nullptr)
 	{
@@ -2017,8 +2020,6 @@ bool FLandscapeComponentSceneProxy::GetMeshElementForVirtualTexture(int32 InLodI
 	OutMeshBatch.DepthPriorityGroup = SDPG_World;
 	OutMeshBatch.LODIndex = static_cast<int8>(InLodIndex);
 	OutMeshBatch.bDitheredLODTransition = false;
-	OutMeshBatch.bRenderToVirtualTexture = true;
-	OutMeshBatch.RuntimeVirtualTextureMaterialType = (uint32)MaterialType;
 
 	OutMeshBatch.Elements.Empty(1);
 
@@ -2133,6 +2134,7 @@ void FLandscapeComponentSceneProxy::DrawStaticElements(FStaticPrimitiveDrawInter
 
 	int32 TotalBatchCount = 1 + LastLOD - FirstLOD;
 	TotalBatchCount += (1 + LastVirtualTextureLOD - FirstVirtualTextureLOD) * RuntimeVirtualTextureMaterialTypes.Num();
+	TotalBatchCount += 1; // TODO: Currently we always add a single LOD0 fixed grid landscape mesh batch for rendering the water info texture. Higher LODs might be better and we might not always need to do this.
 
 	StaticBatchParamArray.Empty(TotalBatchCount);
 	PDI->ReserveMemoryForMeshes(TotalBatchCount);
@@ -2145,13 +2147,28 @@ void FLandscapeComponentSceneProxy::DrawStaticElements(FStaticPrimitiveDrawInter
 		for (int32 LODIndex = FirstVirtualTextureLOD; LODIndex <= LastVirtualTextureLOD; ++LODIndex)
 		{
 			FMeshBatch RuntimeVirtualTextureMeshBatch;
-			if (GetMeshElementForVirtualTexture(LODIndex, MaterialType, AvailableMaterials[MaterialIndex], RuntimeVirtualTextureMeshBatch, StaticBatchParamArray))
+			if (GetMeshElementForFixedGrid(LODIndex, AvailableMaterials[MaterialIndex], RuntimeVirtualTextureMeshBatch, StaticBatchParamArray))
 			{
+				RuntimeVirtualTextureMeshBatch.bRenderToVirtualTexture = true;
+				RuntimeVirtualTextureMeshBatch.RuntimeVirtualTextureMaterialType = (uint32)MaterialType;
 				PDI->DrawMesh(RuntimeVirtualTextureMeshBatch, FLT_MAX);
 			}
 		}
 	}
 
+	// Add fixed grid mesh batch for rendering the water info texture
+	{
+		int32 LODIndex = 0;
+		int32 MaterialIndex = LODIndexToMaterialIndex[LODIndex];
+
+		FMeshBatch MeshBatch;
+		if (GetMeshElementForFixedGrid(LODIndex, AvailableMaterials[MaterialIndex], MeshBatch, StaticBatchParamArray))
+		{
+			MeshBatch.bUseForWaterInfoTextureDepth = true;
+			PDI->DrawMesh(MeshBatch, FLT_MAX);
+		}
+	}
+
 	for (int32 LODIndex = FirstLOD; LODIndex <= LastLOD; LODIndex++)
 	{
 		FMeshBatch MeshBatch;
diff --git a/Engine/Source/Runtime/Landscape/Public/LandscapeRender.h b/Engine/Source/Runtime/Landscape/Public/LandscapeRender.h
--- a/Engine/Source/Runtime/Landscape/Public/LandscapeRender.h
+++ b/Engine/Source/Runtime/Landscape/Public/LandscapeRender.h
@@ -843,7 +843,7 @@ protected:
 	
 	LANDSCAPE_API int8 GetLODFromScreenSize(float InScreenSizeSquared, float InViewLODScale) const;
 
-	LANDSCAPE_API bool GetMeshElementForVirtualTexture(int32 InLodIndex, ERuntimeVirtualTextureMaterialType MaterialType, FMaterialRenderProxy* InMaterialInterface, FMeshBatch& OutMeshBatch, TArray<FLandscapeBatchElementParams>& OutStaticBatchParamArray) const;
+	LANDSCAPE_API bool GetMeshElementForFixedGrid(int32 InLodIndex, FMaterialRenderProxy* InMaterialInterface, FMeshBatch& OutMeshBatch, TArray<FLandscapeBatchElementParams>& OutStaticBatchParamArray) const;
 	template<class ArrayType> bool GetStaticMeshElement(int32 LODIndex, bool bForToolMesh, FMeshBatch& MeshBatch, ArrayType& OutStaticBatchParamArray) const;
 
 public:
diff --git a/Engine/Source/Runtime/Renderer/Private/DeferredShadingRenderer.cpp b/Engine/Source/Runtime/Renderer/Private/DeferredShadingRenderer.cpp
--- a/Engine/Source/Runtime/Renderer/Private/DeferredShadingRenderer.cpp
+++ b/Engine/Source/Runtime/Renderer/Private/DeferredShadingRenderer.cpp
@@ -93,6 +93,7 @@
 #include "Shadows/ShadowScene.h"
 #include "Lumen/LumenHardwareRayTracingCommon.h"
 #include "SparseVolumeTexture/ISparseVolumeTextureStreamingManager.h"
+#include "WaterInfoTextureRendering.h"
 
 #if !UE_BUILD_SHIPPING
 #include "RenderCaptureInterface.h"
@@ -3248,6 +3249,8 @@ void FDeferredShadingSceneRenderer::Render(FRDGBuilder& GraphBuilder)
 
 	ExternalAccessQueue.Submit(GraphBuilder);
 
+	RenderWaterInfoTexture(GraphBuilder, *this, Scene);
+
 	const bool bShouldRenderVelocities = ShouldRenderVelocities();
 	const EShaderPlatform Platform = GetViewFamilyInfo(Views).GetShaderPlatform();
 	const bool bBasePassCanOutputVelocity = FVelocityRendering::BasePassCanOutputVelocity(Platform);
diff --git a/Engine/Source/Runtime/Renderer/Private/MobileShadingRenderer.cpp b/Engine/Source/Runtime/Renderer/Private/MobileShadingRenderer.cpp
--- a/Engine/Source/Runtime/Renderer/Private/MobileShadingRenderer.cpp
+++ b/Engine/Source/Runtime/Renderer/Private/MobileShadingRenderer.cpp
@@ -69,6 +69,7 @@
 #include "Engine/SpecularProfile.h"
 #include "LocalFogVolumeRendering.h"
 #include "SceneCaptureRendering.h"
+#include "WaterInfoTextureRendering.h"
 
 uint32 GetShadowQuality();
 
@@ -1030,6 +1031,8 @@ void FMobileSceneRenderer::Render(FRDGBuilder& GraphBuilder)
 	FRDGTextureRef ViewFamilyTexture = TryCreateViewFamilyTexture(GraphBuilder, ViewFamily);
 
 	const FRDGSystemTextures& SystemTextures = FRDGSystemTextures::Get(GraphBuilder);
+
+	RenderWaterInfoTexture(GraphBuilder, *this, Scene);
 	
 	for (int32 i = 0; i < SceneCaptureRenderPassInfos.Num(); ++i)
 	{
diff --git a/Engine/Source/Runtime/Renderer/Private/PrimitiveSceneInfo.cpp b/Engine/Source/Runtime/Renderer/Private/PrimitiveSceneInfo.cpp
--- a/Engine/Source/Runtime/Renderer/Private/PrimitiveSceneInfo.cpp
+++ b/Engine/Source/Runtime/Renderer/Private/PrimitiveSceneInfo.cpp
@@ -122,6 +122,7 @@ public:
 			bool bUseAnisotropy = Material.GetShadingModels().HasAnyShadingModel({MSM_DefaultLit, MSM_ClearCoat}) && Material.MaterialUsesAnisotropy_RenderThread();
 			bool bSupportsNaniteRendering = SupportsNaniteRendering(StaticMesh->VertexFactory, PrimitiveSceneProxy, Mesh.MaterialRenderProxy, FeatureLevel);
 			bool bSupportsGPUScene = StaticMesh->VertexFactory->SupportsGPUScene(FeatureLevel);
+			bool bUseForWaterInfoTextureDepth = Mesh.bUseForWaterInfoTextureDepth;
 
 			FStaticMeshBatchRelevance* StaticMeshRelevance = new(PrimitiveSceneInfo->StaticMeshRelevances) FStaticMeshBatchRelevance(
 				*StaticMesh, 
@@ -132,6 +133,7 @@ public:
 				bUseAnisotropy,
 				bSupportsNaniteRendering,
 				bSupportsGPUScene,
+				bUseForWaterInfoTextureDepth,
 				FeatureLevel
 				);
 		}
diff --git a/Engine/Source/Runtime/Renderer/Private/StaticMeshBatch.cpp b/Engine/Source/Runtime/Renderer/Private/StaticMeshBatch.cpp
--- a/Engine/Source/Runtime/Renderer/Private/StaticMeshBatch.cpp
+++ b/Engine/Source/Runtime/Renderer/Private/StaticMeshBatch.cpp
@@ -27,7 +27,7 @@ FStaticMeshBatch::~FStaticMeshBatch()
 	Scene->StaticMeshes.RemoveAt(Id);
 }
 
-FStaticMeshBatchRelevance::FStaticMeshBatchRelevance(const FStaticMeshBatch& StaticMesh, float InScreenSize, bool InbSupportsCachingMeshDrawCommands, bool InbUseSkyMaterial, bool bInUseSingleLayerWaterMaterial, bool bInUseAnisotropy, bool bInSupportsNaniteRendering, bool bInSupportsGPUScene, ERHIFeatureLevel::Type FeatureLevel)
+FStaticMeshBatchRelevance::FStaticMeshBatchRelevance(const FStaticMeshBatch& StaticMesh, float InScreenSize, bool InbSupportsCachingMeshDrawCommands, bool InbUseSkyMaterial, bool bInUseSingleLayerWaterMaterial, bool bInUseAnisotropy, bool bInSupportsNaniteRendering, bool bInSupportsGPUScene, bool bInUseForWaterInfoTextureDepth, ERHIFeatureLevel::Type FeatureLevel)
 	: Id(StaticMesh.Id)
 	, ScreenSize(InScreenSize)
 	, NumElements(StaticMesh.Elements.Num())
@@ -49,6 +49,7 @@ FStaticMeshBatchRelevance::FStaticMeshBatchRelevance(const FStaticMeshBatch& Sta
 	, bSupportsCachingMeshDrawCommands(InbSupportsCachingMeshDrawCommands)
 	, bSupportsNaniteRendering(bInSupportsNaniteRendering)
 	, bSupportsGPUScene(bInSupportsGPUScene)
+	, bUseForWaterInfoTextureDepth(bInUseForWaterInfoTextureDepth)
 {
 }
 
diff --git a/Engine/Source/Runtime/Renderer/Private/WaterInfoTextureRendering.cpp b/Engine/Source/Runtime/Renderer/Private/WaterInfoTextureRendering.cpp
new file mode 100644
--- /dev/null
+++ b/Engine/Source/Runtime/Renderer/Private/WaterInfoTextureRendering.cpp
@@ -0,0 +1,1002 @@
+// Copyright Epic Games, Inc. All Rights Reserved.
+
+#include "WaterInfoTextureRendering.h"
+#include "UnrealClient.h"
+#include "RenderGraph.h"
+#include "DeferredShadingRenderer.h"
+#include "SceneUtils.h"
+#include "ScenePrivate.h"
+#include "MeshPassProcessor.inl"
+#include "PixelShaderUtils.h"
+#include "BasePassRendering.h"
+#include "MobileBasePassRendering.h"
+
+DECLARE_GPU_DRAWCALL_STAT(WaterInfoTexture);
+
+static TAutoConsoleVariable<float> CVarWaterInfoUndergroundDilationDepthOffset(
+	TEXT("r.Water.WaterInfo.UndergroundDilationDepthOffset"),
+	64.f,
+	TEXT("The minimum distance below the ground when we allow dilation to write on top of water"),
+	ECVF_RenderThreadSafe);
+
+static TAutoConsoleVariable<float> CVarWaterInfoDilationOverwriteMinimumDistance(
+	TEXT("r.Water.WaterInfo.DilationOverwriteMinimumDistance"),
+	128.f,
+	TEXT("The minimum distance below the ground when we allow dilation to write on top of water"),
+	ECVF_RenderThreadSafe);
+
+BEGIN_SHADER_PARAMETER_STRUCT(FWaterInfoTexturePassParameters, )
+	SHADER_PARAMETER_RDG_UNIFORM_BUFFER(FViewUniformShaderParameters, View)
+	SHADER_PARAMETER_STRUCT_INCLUDE(FInstanceCullingDrawParams, InstanceCullingDrawParams)
+	RENDER_TARGET_BINDING_SLOTS()
+END_SHADER_PARAMETER_STRUCT()
+
+BEGIN_SHADER_PARAMETER_STRUCT(FWaterInfoTextureDepthPassParameters, )
+	SHADER_PARAMETER_RDG_UNIFORM_BUFFER(FViewUniformShaderParameters, View)
+	SHADER_PARAMETER_STRUCT_INCLUDE(FInstanceCullingDrawParams, InstanceCullingDrawParams)
+	RENDER_TARGET_BINDING_SLOTS()
+END_SHADER_PARAMETER_STRUCT()
+
+class FWaterInfoTextureMergePS : public FGlobalShader
+{
+public:
+	DECLARE_GLOBAL_SHADER(FWaterInfoTextureMergePS);
+	SHADER_USE_PARAMETER_STRUCT(FWaterInfoTextureMergePS, FGlobalShader);
+
+	BEGIN_SHADER_PARAMETER_STRUCT(FParameters, )
+		SHADER_PARAMETER_RDG_UNIFORM_BUFFER(FViewUniformShaderParameters, View)
+		SHADER_PARAMETER_STRUCT_INCLUDE(FSceneTextureShaderParameters, SceneTextures)
+		SHADER_PARAMETER_RDG_TEXTURE(Texture2D<float4>, GroundDepthTexture)
+		SHADER_PARAMETER_RDG_TEXTURE(Texture2D<float4>, WaterBodyTexture)
+		SHADER_PARAMETER_RDG_TEXTURE(Texture2D<float4>, WaterBodyDepthTexture)
+		SHADER_PARAMETER_RDG_TEXTURE(Texture2D<float4>, DilatedWaterBodyDepthTexture)
+		SHADER_PARAMETER(FVector2f, WaterHeightExtents)
+		SHADER_PARAMETER(float, GroundZMin)
+		SHADER_PARAMETER(float, CaptureZ)
+		SHADER_PARAMETER(float, UndergroundDilationDepthOffset)
+		SHADER_PARAMETER(float, DilationOverwriteMinimumDistance)
+		RENDER_TARGET_BINDING_SLOTS()
+	END_SHADER_PARAMETER_STRUCT()
+
+	class FEnable128BitRT : SHADER_PERMUTATION_BOOL("ENABLE_128_BIT");
+	using FPermutationDomain = TShaderPermutationDomain<FEnable128BitRT>;
+
+	static void ModifyCompilationEnvironment(const FGlobalShaderPermutationParameters& Parameters, FShaderCompilerEnvironment& OutEnvironment)
+	{
+		FGlobalShader::ModifyCompilationEnvironment(Parameters, OutEnvironment);
+		const FPermutationDomain PermutationVector(Parameters.PermutationId);
+		OutEnvironment.SetRenderTargetOutputFormat(0, PermutationVector.Get<FEnable128BitRT>() ? PF_A32B32G32R32F : PF_FloatRGBA);
+	}
+};
+
+IMPLEMENT_GLOBAL_SHADER(FWaterInfoTextureMergePS, "/Engine/Private/WaterInfoTextureMerge.usf", "Main", SF_Pixel);
+
+class FWaterInfoTextureBlurPS : public FGlobalShader
+{
+public:
+	DECLARE_GLOBAL_SHADER(FWaterInfoTextureBlurPS);
+	SHADER_USE_PARAMETER_STRUCT(FWaterInfoTextureBlurPS, FGlobalShader);
+
+	BEGIN_SHADER_PARAMETER_STRUCT(FParameters, )
+		SHADER_PARAMETER_RDG_UNIFORM_BUFFER(FViewUniformShaderParameters, View)
+		SHADER_PARAMETER_STRUCT_INCLUDE(FSceneTextureShaderParameters, SceneTextures)
+		SHADER_PARAMETER_RDG_TEXTURE(Texture2D, WaterInfoTexture)
+		SHADER_PARAMETER(float, WaterZMin)
+		SHADER_PARAMETER(float, WaterZMax)
+		SHADER_PARAMETER(float, GroundZMin)
+		SHADER_PARAMETER(float, CaptureZ)
+		SHADER_PARAMETER(int, BlurRadius)
+		RENDER_TARGET_BINDING_SLOTS()
+	END_SHADER_PARAMETER_STRUCT()
+
+	class FEnable128BitRT : SHADER_PERMUTATION_BOOL("ENABLE_128_BIT");
+	using FPermutationDomain = TShaderPermutationDomain<FEnable128BitRT>;
+
+	static void ModifyCompilationEnvironment(const FGlobalShaderPermutationParameters& Parameters, FShaderCompilerEnvironment& OutEnvironment)
+	{
+		FGlobalShader::ModifyCompilationEnvironment(Parameters, OutEnvironment);
+		const FPermutationDomain PermutationVector(Parameters.PermutationId);
+		OutEnvironment.SetRenderTargetOutputFormat(0, PermutationVector.Get<FEnable128BitRT>() ? PF_A32B32G32R32F : PF_FloatRGBA);
+	}
+};
+
+IMPLEMENT_GLOBAL_SHADER(FWaterInfoTextureBlurPS, "/Engine/Private/WaterInfoTextureBlur.usf", "Main", SF_Pixel);
+
+/**
+* MeshPassProcessor for the "color" pass required for generating the water info texture. The associated pass draws water body meshes with an unlit material
+* in order to write water surface depth, river velocity and possibly other data too.
+*/
+class FWaterInfoTexturePassMeshProcessor : public FSceneRenderingAllocatorObject<FWaterInfoTexturePassMeshProcessor>, public FMeshPassProcessor
+{
+public:
+	FWaterInfoTexturePassMeshProcessor(const FScene* Scene, ERHIFeatureLevel::Type FeatureLevel, const FSceneView* InViewIfDynamicMeshCommand, FMeshPassDrawListContext* InDrawListContext, bool bIsMobile);
+
+	virtual void AddMeshBatch(const FMeshBatch& RESTRICT MeshBatch, uint64 BatchElementMask, const FPrimitiveSceneProxy* RESTRICT PrimitiveSceneProxy, int32 StaticMeshId = -1) override final;
+	virtual void CollectPSOInitializers(const FSceneTexturesConfig& SceneTexturesConfig, const FMaterial& Material, const FPSOPrecacheVertexFactoryData& VertexFactoryData, const FPSOPrecacheParams& PreCacheParams, TArray<FPSOPrecacheData>& PSOInitializers) override final;
+
+private:
+	bool TryAddMeshBatch(
+		const FMeshBatch& RESTRICT MeshBatch,
+		uint64 BatchElementMask,
+		const FPrimitiveSceneProxy* RESTRICT PrimitiveSceneProxy,
+		int32 StaticMeshId,
+		const FMaterialRenderProxy& MaterialRenderProxy,
+		const FMaterial& Material);
+
+	FMeshPassProcessorRenderState PassDrawRenderState;
+	bool bIsMobile;
+};
+
+FWaterInfoTexturePassMeshProcessor::FWaterInfoTexturePassMeshProcessor(const FScene* Scene, ERHIFeatureLevel::Type FeatureLevel, const FSceneView* InViewIfDynamicMeshCommand, FMeshPassDrawListContext* InDrawListContext, bool bIsMobile)
+	: FMeshPassProcessor(EMeshPass::WaterInfoTexturePass, Scene, FeatureLevel, InViewIfDynamicMeshCommand, InDrawListContext),
+	bIsMobile(bIsMobile)
+{
+	PassDrawRenderState.SetBlendState(TStaticBlendState<>::GetRHI());
+	PassDrawRenderState.SetDepthStencilAccess(FExclusiveDepthStencil::DepthWrite_StencilNop);
+	PassDrawRenderState.SetDepthStencilState(TStaticDepthStencilState<true, CF_DepthNearOrEqual>::GetRHI());
+}
+
+void FWaterInfoTexturePassMeshProcessor::AddMeshBatch(const FMeshBatch& RESTRICT MeshBatch, uint64 BatchElementMask, const FPrimitiveSceneProxy* RESTRICT PrimitiveSceneProxy, int32 StaticMeshId)
+{
+	if (!MeshBatch.bUseForMaterial)
+	{
+		return;
+	}
+
+	const FMaterialRenderProxy* MaterialRenderProxy = MeshBatch.MaterialRenderProxy;
+	while (MaterialRenderProxy)
+	{
+		const FMaterial* Material = MaterialRenderProxy->GetMaterialNoFallback(FeatureLevel);
+		if (Material)
+		{
+			if (TryAddMeshBatch(MeshBatch, BatchElementMask, PrimitiveSceneProxy, StaticMeshId, *MaterialRenderProxy, *Material))
+			{
+				break;
+			}
+		}
+
+		MaterialRenderProxy = MaterialRenderProxy->GetFallback(FeatureLevel);
+	}
+}
+
+void FWaterInfoTexturePassMeshProcessor::CollectPSOInitializers(const FSceneTexturesConfig& SceneTexturesConfig, const FMaterial& Material, const FPSOPrecacheVertexFactoryData& VertexFactoryData, const FPSOPrecacheParams& PreCacheParams, TArray<FPSOPrecacheData>& PSOInitializers)
+{
+	// Try to reduce the number of materials considered for this mesh pass:
+	// Materials for drawing the water info texture are supposed to be unlit (they write velocity and possibly other data into Emissive).
+	// They also need to be applied to meshes (MD_Surface) and we can safely exclude sky materials which also use the unlit shading model.
+	if (Material.GetShadingModels().IsUnlit() && Material.GetMaterialDomain() == MD_Surface && !Material.IsSky())
+	{
+		// Determine the mesh's material and blend mode.
+		const FMeshDrawingPolicyOverrideSettings OverrideSettings = ComputeMeshOverrideSettings(PreCacheParams);
+		const ERasterizerFillMode MeshFillMode = ComputeMeshFillMode(Material, OverrideSettings);
+		const ERasterizerCullMode MeshCullMode = ComputeMeshCullMode(Material, OverrideSettings);
+
+		// Setup the render target info
+		FGraphicsPipelineRenderTargetsInfo RenderTargetsInfo;
+		RenderTargetsInfo.NumSamples = 1;
+		AddRenderTargetInfo(PF_FloatRGBA, TexCreate_RenderTargetable | TexCreate_ShaderResource, RenderTargetsInfo);
+		SetupDepthStencilInfo(PF_DepthStencil, TexCreate_DepthStencilTargetable | TexCreate_ShaderResource, ERenderTargetLoadAction::EClear, ERenderTargetLoadAction::ENoAction, FExclusiveDepthStencil::DepthWrite_StencilNop, RenderTargetsInfo);
+
+		if (!bIsMobile)
+		{
+			// Get the shaders if possible for given vertex factory
+			TMeshProcessorShaders<
+				TBasePassVertexShaderPolicyParamType<FUniformLightMapPolicy>,
+				TBasePassPixelShaderPolicyParamType<FUniformLightMapPolicy>> PassShaders;
+			if (!GetBasePassShaders<FUniformLightMapPolicy>(
+				Material,
+				VertexFactoryData.VertexFactoryType,
+				FUniformLightMapPolicy(LMP_NO_LIGHTMAP),
+				FeatureLevel,
+				false, // bRenderSkylight
+				false, // 128bit
+				GBL_Default, // Currently only Nanite supports non-default layout
+				&PassShaders.VertexShader,
+				&PassShaders.PixelShader
+			))
+			{
+				return;
+			}
+
+			AddGraphicsPipelineStateInitializer(
+				VertexFactoryData,
+				Material,
+				PassDrawRenderState,
+				RenderTargetsInfo,
+				PassShaders,
+				MeshFillMode,
+				MeshCullMode,
+				(EPrimitiveType)PreCacheParams.PrimitiveType,
+				EMeshPassFeatures::Default,
+				true /*bRequired*/,
+				PSOInitializers);
+		}
+		else
+		{
+			TMeshProcessorShaders<
+				TMobileBasePassVSPolicyParamType<FUniformLightMapPolicy>,
+				TMobileBasePassPSPolicyParamType<FUniformLightMapPolicy>> PassShaders;
+
+			FMaterialShaderTypes ShaderTypes;
+			ShaderTypes.AddShaderType<TMobileBasePassVS<TUniformLightMapPolicy<LMP_NO_LIGHTMAP>, HDR_LINEAR_64>>();
+			ShaderTypes.AddShaderType<TMobileBasePassPS<TUniformLightMapPolicy<LMP_NO_LIGHTMAP>, HDR_LINEAR_64, false, LOCAL_LIGHTS_DISABLED>>();
+			FMaterialShaders Shaders;
+			if (!Material.TryGetShaders(ShaderTypes, VertexFactoryData.VertexFactoryType, Shaders))
+			{
+				return;
+			}
+
+			Shaders.TryGetVertexShader(PassShaders.VertexShader);
+			Shaders.TryGetPixelShader(PassShaders.PixelShader);
+
+			AddGraphicsPipelineStateInitializer(
+				VertexFactoryData,
+				Material,
+				PassDrawRenderState,
+				RenderTargetsInfo,
+				PassShaders,
+				MeshFillMode,
+				MeshCullMode,
+				(EPrimitiveType)PreCacheParams.PrimitiveType,
+				EMeshPassFeatures::Default,
+				true /*bRequired*/,
+				PSOInitializers);
+		}
+		
+	}
+}
+
+bool FWaterInfoTexturePassMeshProcessor::TryAddMeshBatch(const FMeshBatch& RESTRICT MeshBatch, uint64 BatchElementMask, const FPrimitiveSceneProxy* RESTRICT PrimitiveSceneProxy, int32 StaticMeshId, const FMaterialRenderProxy& MaterialRenderProxy, const FMaterial& Material)
+{
+	// Try to reduce the number of materials considered for this mesh pass:
+	// Materials for drawing the water info texture are supposed to be unlit (they write velocity and possibly other data into Emissive).
+	// They also need to be applied to meshes (MD_Surface) and we can safely exclude sky materials which also use the unlit shading model.
+	if (Material.GetShadingModels().IsUnlit() && Material.GetMaterialDomain() == MD_Surface && !Material.IsSky())
+	{
+		const FMeshDrawingPolicyOverrideSettings OverrideSettings = ComputeMeshOverrideSettings(MeshBatch);
+		const ERasterizerFillMode MeshFillMode = ComputeMeshFillMode(Material, OverrideSettings);
+		const ERasterizerCullMode MeshCullMode = ComputeMeshCullMode(Material, OverrideSettings);
+
+		const FVertexFactory* VertexFactory = MeshBatch.VertexFactory;
+
+		if (!bIsMobile)
+		{
+			TMeshProcessorShaders<
+				TBasePassVertexShaderPolicyParamType<FUniformLightMapPolicy>,
+				TBasePassPixelShaderPolicyParamType<FUniformLightMapPolicy>> PassShaders;
+
+			if (!GetBasePassShaders<FUniformLightMapPolicy>(
+				Material,
+				VertexFactory->GetType(),
+				FUniformLightMapPolicy(LMP_NO_LIGHTMAP),
+				FeatureLevel,
+				false, // bRenderSkylight
+				false, // 128bit
+				GBL_Default, // Currently only Nanite uses non-default layout
+				&PassShaders.VertexShader,
+				&PassShaders.PixelShader
+			))
+			{
+				return false;
+			}
+
+			FMeshDrawCommandSortKey SortKey = CalculateMeshStaticSortKey(PassShaders.VertexShader, PassShaders.PixelShader);
+
+			TBasePassShaderElementData<FUniformLightMapPolicy> ShaderElementData(MeshBatch.LCI);
+			ShaderElementData.InitializeMeshMaterialData(ViewIfDynamicMeshCommand, PrimitiveSceneProxy, MeshBatch, StaticMeshId, true);
+
+			BuildMeshDrawCommands(
+				MeshBatch,
+				BatchElementMask,
+				PrimitiveSceneProxy,
+				MaterialRenderProxy,
+				Material,
+				PassDrawRenderState,
+				PassShaders,
+				MeshFillMode,
+				MeshCullMode,
+				SortKey,
+				EMeshPassFeatures::Default,
+				ShaderElementData);
+		}
+		else
+		{
+			TMeshProcessorShaders<
+				TMobileBasePassVSPolicyParamType<FUniformLightMapPolicy>,
+				TMobileBasePassPSPolicyParamType<FUniformLightMapPolicy>> PassShaders;
+
+			FMaterialShaderTypes ShaderTypes;
+			ShaderTypes.AddShaderType<TMobileBasePassVS<TUniformLightMapPolicy<LMP_NO_LIGHTMAP>, HDR_LINEAR_64>>();
+			ShaderTypes.AddShaderType<TMobileBasePassPS<TUniformLightMapPolicy<LMP_NO_LIGHTMAP>, HDR_LINEAR_64, false, LOCAL_LIGHTS_DISABLED>>();
+			FMaterialShaders Shaders;
+			if (!Material.TryGetShaders(ShaderTypes, VertexFactory->GetType(), Shaders))
+			{
+				return false;
+			}
+
+			Shaders.TryGetVertexShader(PassShaders.VertexShader);
+			Shaders.TryGetPixelShader(PassShaders.PixelShader);
+
+			FMeshDrawCommandSortKey SortKey = CalculateMeshStaticSortKey(PassShaders.VertexShader, PassShaders.PixelShader);
+
+			TMobileBasePassShaderElementData<FUniformLightMapPolicy> ShaderElementData(MeshBatch.LCI, false);
+			ShaderElementData.InitializeMeshMaterialData(ViewIfDynamicMeshCommand, PrimitiveSceneProxy, MeshBatch, StaticMeshId, false);
+
+			BuildMeshDrawCommands(
+				MeshBatch,
+				BatchElementMask,
+				PrimitiveSceneProxy,
+				MaterialRenderProxy,
+				Material,
+				PassDrawRenderState,
+				PassShaders,
+				MeshFillMode,
+				MeshCullMode,
+				SortKey,
+				EMeshPassFeatures::Default,
+				ShaderElementData);
+		}
+		
+
+		return true;
+	}
+
+	return false;
+}
+
+FMeshPassProcessor* CreateWaterInfoTexturePassMeshProcessor(ERHIFeatureLevel::Type FeatureLevel, const FScene* Scene, const FSceneView* InViewIfDynamicMeshCommand, FMeshPassDrawListContext* InDrawListContext)
+{
+	return new FWaterInfoTexturePassMeshProcessor(Scene, FeatureLevel, InViewIfDynamicMeshCommand, InDrawListContext, false);
+}
+
+FMeshPassProcessor* CreateWaterInfoTexturePassMeshProcessorMobile(ERHIFeatureLevel::Type FeatureLevel, const FScene* Scene, const FSceneView* InViewIfDynamicMeshCommand, FMeshPassDrawListContext* InDrawListContext)
+{
+	return new FWaterInfoTexturePassMeshProcessor(Scene, FeatureLevel, InViewIfDynamicMeshCommand, InDrawListContext, true);
+}
+
+REGISTER_MESHPASSPROCESSOR_AND_PSOCOLLECTOR(WaterInfoTexturePass, CreateWaterInfoTexturePassMeshProcessor, EShadingPath::Deferred, EMeshPass::WaterInfoTexturePass, EMeshPassFlags::CachedMeshCommands);
+REGISTER_MESHPASSPROCESSOR_AND_PSOCOLLECTOR(MobileWaterInfoTexturePass, CreateWaterInfoTexturePassMeshProcessorMobile, EShadingPath::Mobile, EMeshPass::WaterInfoTexturePass, EMeshPassFlags::CachedMeshCommands);
+
+
+/**
+* MeshPassProcessor for the two depth-only passes involved in generating the water info texture. The two passes can use the same processor. The first (terrain) pass renders terrain
+* with a fixed grid vertex factory and possibly other static meshes too. The second pass renders dilated water body meshes. All these meshes should only be rendered into this view/texture, 
+  which is why MeshBatches need to have the bUseForWaterInfoTextureDepth bool set.
+*/
+class FWaterInfoTextureDepthPassMeshProcessor : public FSceneRenderingAllocatorObject<FWaterInfoTextureDepthPassMeshProcessor>, public FMeshPassProcessor
+{
+public:
+	FWaterInfoTextureDepthPassMeshProcessor(const FScene* Scene, ERHIFeatureLevel::Type FeatureLevel, const FSceneView* InViewIfDynamicMeshCommand, FMeshPassDrawListContext* InDrawListContext);
+
+	virtual void AddMeshBatch(const FMeshBatch& RESTRICT MeshBatch, uint64 BatchElementMask, const FPrimitiveSceneProxy* RESTRICT PrimitiveSceneProxy, int32 StaticMeshId = -1) override final;
+	virtual void CollectPSOInitializers(const FSceneTexturesConfig& SceneTexturesConfig, const FMaterial& Material, const FPSOPrecacheVertexFactoryData& VertexFactoryData, const FPSOPrecacheParams& PreCacheParams, TArray<FPSOPrecacheData>& PSOInitializers) override final;
+
+private:
+	bool TryAddMeshBatch(
+		const FMeshBatch& RESTRICT MeshBatch,
+		uint64 BatchElementMask,
+		const FPrimitiveSceneProxy* RESTRICT PrimitiveSceneProxy,
+		int32 StaticMeshId,
+		const FMaterialRenderProxy& MaterialRenderProxy,
+		const FMaterial& Material);
+
+	template<bool bPositionOnly>
+	bool Process(
+		const FMeshBatch& RESTRICT MeshBatch,
+		uint64 BatchElementMask,
+		const FPrimitiveSceneProxy* RESTRICT PrimitiveSceneProxy,
+		int32 StaticMeshId,
+		const FMaterialRenderProxy& RESTRICT MaterialRenderProxy,
+		const FMaterial& RESTRICT MaterialResource,
+		ERasterizerFillMode MeshFillMode,
+		ERasterizerCullMode MeshCullMode);
+
+	template<bool bPositionOnly>
+	void CollectPSOInitializersInternal(
+		const FSceneTexturesConfig& SceneTexturesConfig,
+		const FPSOPrecacheVertexFactoryData& VertexFactoryData,
+		const FMaterial& RESTRICT MaterialResource,
+		ERasterizerFillMode MeshFillMode,
+		ERasterizerCullMode MeshCullMode,
+		const FPSOPrecacheParams& PreCacheParams,
+		TArray<FPSOPrecacheData>& PSOInitializers);
+
+	FMeshPassProcessorRenderState PassDrawRenderState;
+};
+
+FWaterInfoTextureDepthPassMeshProcessor::FWaterInfoTextureDepthPassMeshProcessor(const FScene* Scene, ERHIFeatureLevel::Type FeatureLevel, const FSceneView* InViewIfDynamicMeshCommand, FMeshPassDrawListContext* InDrawListContext)
+	: FMeshPassProcessor(EMeshPass::WaterInfoTextureDepthPass, Scene, FeatureLevel, InViewIfDynamicMeshCommand, InDrawListContext)
+{
+	PassDrawRenderState.SetBlendState(TStaticBlendState<CW_NONE>::GetRHI());
+	PassDrawRenderState.SetDepthStencilAccess(FExclusiveDepthStencil::DepthWrite_StencilNop);
+	PassDrawRenderState.SetDepthStencilState(TStaticDepthStencilState<true, CF_DepthNearOrEqual>::GetRHI());
+}
+
+void FWaterInfoTextureDepthPassMeshProcessor::AddMeshBatch(const FMeshBatch& RESTRICT MeshBatch, uint64 BatchElementMask, const FPrimitiveSceneProxy* RESTRICT PrimitiveSceneProxy, int32 StaticMeshId)
+{
+	const FMaterialRenderProxy* MaterialRenderProxy = MeshBatch.MaterialRenderProxy;
+	while (MaterialRenderProxy)
+	{
+		const FMaterial* Material = MaterialRenderProxy->GetMaterialNoFallback(FeatureLevel);
+		if (Material)
+		{
+			if (TryAddMeshBatch(MeshBatch, BatchElementMask, PrimitiveSceneProxy, StaticMeshId, *MaterialRenderProxy, *Material))
+			{
+				break;
+			}
+		}
+
+		MaterialRenderProxy = MaterialRenderProxy->GetFallback(FeatureLevel);
+	}
+}
+
+bool FWaterInfoTextureDepthPassMeshProcessor::TryAddMeshBatch(
+	const FMeshBatch& RESTRICT MeshBatch,
+	uint64 BatchElementMask,
+	const FPrimitiveSceneProxy* RESTRICT PrimitiveSceneProxy,
+	int32 StaticMeshId,
+	const FMaterialRenderProxy& MaterialRenderProxy,
+	const FMaterial& Material)
+{
+	if (MeshBatch.bUseForWaterInfoTextureDepth && Material.GetMaterialDomain() == MD_Surface && !IsTranslucentBlendMode(Material))
+	{
+		// Determine the mesh's material and blend mode.
+		const FMeshDrawingPolicyOverrideSettings OverrideSettings = ComputeMeshOverrideSettings(MeshBatch);
+		const ERasterizerFillMode MeshFillMode = ComputeMeshFillMode(Material, OverrideSettings);
+		const ERasterizerCullMode MeshCullMode = ComputeMeshCullMode(Material, OverrideSettings);
+		const bool bVFTypeSupportsNullPixelShader = MeshBatch.VertexFactory->SupportsNullPixelShader();
+		const bool bEvaluateWPO = Material.MaterialModifiesMeshPosition_RenderThread() && (!ShouldOptimizedWPOAffectNonNaniteShaderSelection() || PrimitiveSceneProxy->EvaluateWorldPositionOffset());
+
+		if (IsOpaqueBlendMode(Material)
+			&& MeshBatch.VertexFactory->SupportsPositionOnlyStream()
+			&& !bEvaluateWPO
+			&& Material.WritesEveryPixel(false, bVFTypeSupportsNullPixelShader))
+		{
+			const FMaterialRenderProxy& DefaultProxy = *UMaterial::GetDefaultMaterial(MD_Surface)->GetRenderProxy();
+			const FMaterial& DefaultMaterial = *DefaultProxy.GetMaterialNoFallback(FeatureLevel);
+			return Process<true>(MeshBatch, BatchElementMask, PrimitiveSceneProxy, StaticMeshId, DefaultProxy, DefaultMaterial, MeshFillMode, MeshCullMode);
+		}
+		else
+		{
+			const bool bMaterialMasked = !Material.WritesEveryPixel(false, bVFTypeSupportsNullPixelShader);
+			const FMaterialRenderProxy* EffectiveMaterialRenderProxy = &MaterialRenderProxy;
+			const FMaterial* EffectiveMaterial = &Material;
+
+			if (!bMaterialMasked && !bEvaluateWPO)
+			{
+				// Override with the default material for opaque materials that are not two sided
+				EffectiveMaterialRenderProxy = UMaterial::GetDefaultMaterial(MD_Surface)->GetRenderProxy();
+				EffectiveMaterial = EffectiveMaterialRenderProxy->GetMaterialNoFallback(FeatureLevel);
+				check(EffectiveMaterial);
+			}
+
+			return Process<false>(MeshBatch, BatchElementMask, PrimitiveSceneProxy, StaticMeshId, *EffectiveMaterialRenderProxy, *EffectiveMaterial, MeshFillMode, MeshCullMode);
+		}
+	}
+
+	return true;
+}
+
+template<bool bPositionOnly>
+bool FWaterInfoTextureDepthPassMeshProcessor::Process(
+	const FMeshBatch& RESTRICT MeshBatch,
+	uint64 BatchElementMask,
+	const FPrimitiveSceneProxy* RESTRICT PrimitiveSceneProxy,
+	int32 StaticMeshId,
+	const FMaterialRenderProxy& RESTRICT MaterialRenderProxy,
+	const FMaterial& RESTRICT MaterialResource,
+	ERasterizerFillMode MeshFillMode,
+	ERasterizerCullMode MeshCullMode)
+{
+	TMeshProcessorShaders<TDepthOnlyVS<bPositionOnly>, FDepthOnlyPS> DepthPassShaders;
+	FShaderPipelineRef ShaderPipeline;
+
+	if (!GetDepthPassShaders<bPositionOnly>(
+		MaterialResource,
+		MeshBatch.VertexFactory->GetType(),
+		FeatureLevel,
+		MaterialResource.MaterialUsesPixelDepthOffset_GameThread(),
+		DepthPassShaders.VertexShader,
+		DepthPassShaders.PixelShader,
+		ShaderPipeline))
+	{
+		return false;
+	}
+
+	FMeshMaterialShaderElementData ShaderElementData;
+	ShaderElementData.InitializeMeshMaterialData(ViewIfDynamicMeshCommand, PrimitiveSceneProxy, MeshBatch, StaticMeshId, true);
+
+	const bool bIsMasked = IsMaskedBlendMode(MaterialResource);
+	FMeshDrawCommandSortKey SortKey = CalculateDepthPassMeshStaticSortKey(bIsMasked, DepthPassShaders.VertexShader.GetShader(), DepthPassShaders.PixelShader.GetShader());
+
+	BuildMeshDrawCommands(
+		MeshBatch,
+		BatchElementMask,
+		PrimitiveSceneProxy,
+		MaterialRenderProxy,
+		MaterialResource,
+		PassDrawRenderState,
+		DepthPassShaders,
+		MeshFillMode,
+		MeshCullMode,
+		SortKey,
+		bPositionOnly ? EMeshPassFeatures::PositionOnly : EMeshPassFeatures::Default,
+		ShaderElementData);
+
+	return true;
+}
+
+void FWaterInfoTextureDepthPassMeshProcessor::CollectPSOInitializers(const FSceneTexturesConfig& SceneTexturesConfig, const FMaterial& Material, const FPSOPrecacheVertexFactoryData& VertexFactoryData, const FPSOPrecacheParams& PreCacheParams, TArray<FPSOPrecacheData>& PSOInitializers)
+{
+	// We need to support all materials that could possibly be rendered in a depth-only pass. Unfortunately there doesn't seem to be a way to filter for bUseForWaterInfoTextureDepth at this point.
+	if (Material.GetMaterialDomain() == MD_Surface && !IsTranslucentBlendMode(Material))
+	{
+		// Determine the mesh's material and blend mode.
+		const FMeshDrawingPolicyOverrideSettings OverrideSettings = ComputeMeshOverrideSettings(PreCacheParams);
+		const ERasterizerFillMode MeshFillMode = ComputeMeshFillMode(Material, OverrideSettings);
+		const ERasterizerCullMode MeshCullMode = ComputeMeshCullMode(Material, OverrideSettings);
+		const bool bSupportPositionOnlyStream = VertexFactoryData.VertexFactoryType->SupportsPositionOnly();
+		const bool bVFTypeSupportsNullPixelShader = VertexFactoryData.VertexFactoryType->SupportsNullPixelShader();
+
+		if (IsOpaqueBlendMode(Material)
+			&& bSupportPositionOnlyStream
+			&& !Material.MaterialModifiesMeshPosition_GameThread()
+			&& Material.WritesEveryPixel(false, bVFTypeSupportsNullPixelShader))
+		{
+			const FMaterialRenderProxy& DefaultProxy = *UMaterial::GetDefaultMaterial(MD_Surface)->GetRenderProxy();
+			const FMaterial& DefaultMaterial = *DefaultProxy.GetMaterialNoFallback(FeatureLevel);
+
+			CollectPSOInitializersInternal<true>(SceneTexturesConfig, VertexFactoryData, DefaultMaterial, MeshFillMode, MeshCullMode, PreCacheParams, PSOInitializers);
+		}
+		else
+		{
+			const bool bMaterialMasked = !Material.WritesEveryPixel(false, bVFTypeSupportsNullPixelShader);
+			const FMaterial* EffectiveMaterial = &Material;
+
+			if (!bMaterialMasked && !Material.MaterialModifiesMeshPosition_GameThread())
+			{
+				// Override with the default material for opaque materials that are not two sided
+				EMaterialQualityLevel::Type ActiveQualityLevel = GetCachedScalabilityCVars().MaterialQualityLevel;
+				EffectiveMaterial = UMaterial::GetDefaultMaterial(MD_Surface)->GetMaterialResource(FeatureLevel, ActiveQualityLevel);
+				check(EffectiveMaterial);
+			}
+
+			CollectPSOInitializersInternal<false>(SceneTexturesConfig, VertexFactoryData, *EffectiveMaterial, MeshFillMode, MeshCullMode, PreCacheParams, PSOInitializers);
+		}
+	}
+}
+
+template<bool bPositionOnly>
+void FWaterInfoTextureDepthPassMeshProcessor::CollectPSOInitializersInternal(
+	const FSceneTexturesConfig& SceneTexturesConfig,
+	const FPSOPrecacheVertexFactoryData& VertexFactoryData,
+	const FMaterial& RESTRICT MaterialResource,
+	ERasterizerFillMode MeshFillMode,
+	ERasterizerCullMode MeshCullMode,
+	const FPSOPrecacheParams& PreCacheParams,
+	TArray<FPSOPrecacheData>& PSOInitializers)
+{
+	TMeshProcessorShaders<TDepthOnlyVS<bPositionOnly>, FDepthOnlyPS> DepthPassShaders;
+	FShaderPipelineRef ShaderPipeline;
+
+	if (!GetDepthPassShaders<bPositionOnly>(
+		MaterialResource,
+		VertexFactoryData.VertexFactoryType,
+		FeatureLevel,
+		MaterialResource.MaterialUsesPixelDepthOffset_GameThread(),
+		DepthPassShaders.VertexShader,
+		DepthPassShaders.PixelShader,
+		ShaderPipeline))
+	{
+		return;
+	}
+
+	FGraphicsPipelineRenderTargetsInfo RenderTargetsInfo;
+	RenderTargetsInfo.NumSamples = 1;
+	SetupDepthStencilInfo(PF_DepthStencil, TexCreate_DepthStencilTargetable | TexCreate_ShaderResource, ERenderTargetLoadAction::EClear,
+		ERenderTargetLoadAction::ENoAction, FExclusiveDepthStencil::DepthWrite_StencilNop, RenderTargetsInfo);
+
+	AddGraphicsPipelineStateInitializer(
+		VertexFactoryData,
+		MaterialResource,
+		PassDrawRenderState,
+		RenderTargetsInfo,
+		DepthPassShaders,
+		MeshFillMode,
+		MeshCullMode,
+		(EPrimitiveType)PreCacheParams.PrimitiveType,
+		bPositionOnly ? EMeshPassFeatures::PositionOnly : EMeshPassFeatures::Default,
+		true /*bRequired*/,
+		PSOInitializers);
+}
+
+FMeshPassProcessor* CreateWaterInfoTextureDepthPassMeshProcessor(ERHIFeatureLevel::Type FeatureLevel, const FScene* Scene, const FSceneView* InViewIfDynamicMeshCommand, FMeshPassDrawListContext* InDrawListContext)
+{
+	return new FWaterInfoTextureDepthPassMeshProcessor(Scene, FeatureLevel, InViewIfDynamicMeshCommand, InDrawListContext);
+}
+
+REGISTER_MESHPASSPROCESSOR_AND_PSOCOLLECTOR(WaterInfoTextureDepthPass, CreateWaterInfoTextureDepthPassMeshProcessor, EShadingPath::Deferred, EMeshPass::WaterInfoTextureDepthPass, EMeshPassFlags::CachedMeshCommands);
+REGISTER_MESHPASSPROCESSOR_AND_PSOCOLLECTOR(MobileWaterInfoTextureDepthPass, CreateWaterInfoTextureDepthPassMeshProcessor, EShadingPath::Mobile, EMeshPass::WaterInfoTextureDepthPass, EMeshPassFlags::CachedMeshCommands);
+
+
+
+enum class EWaterInfoTextureMeshPass : int32
+{
+	// Depth-only pass for actors which are considered the terrain for water rendering
+	TerrainDepth,
+
+	// Depth-only pass for water body dilated sections which get composited back into the water info texture at a lower priority than regular water body data
+	DilatedWaterBodyDepth,
+	
+	// Depth and color pass for non-dilated water body meshes, rendering out water surface depth, river velocity/flow and possibly other data in the future.
+	WaterBody,
+	
+	Num
+};
+
+struct FWaterInfoTextureDraws
+{
+	EMeshPass::Type MeshPass = EMeshPass::WaterInfoTextureDepthPass;
+	const TSet<FPrimitiveComponentId>* ComponentIDsToDraw = nullptr;
+	uint32 PrimitiveInstanceIndex = 0;
+	FMeshCommandOneFrameArray VisibleMeshCommands;
+	TArray<int32, SceneRenderingAllocator> PrimitiveIds;
+	TArray<uint32, SceneRenderingAllocator> InstanceRuns;
+	FInstanceCullingContext* InstanceCullingContext = nullptr;
+	FInstanceCullingResult InstanceCullingResult;
+
+	FWaterInfoTextureDraws() = default;
+	FWaterInfoTextureDraws(EMeshPass::Type InMeshPass, const TSet<FPrimitiveComponentId>* InComponentIDsToDraw, uint32 InPrimitiveInstanceIndex)
+		: MeshPass(InMeshPass), ComponentIDsToDraw(InComponentIDsToDraw), PrimitiveInstanceIndex(InPrimitiveInstanceIndex)
+	{}
+};
+
+static void AddWaterInfoTextureDraws(const FScene* Scene, const TArrayView<FWaterInfoTextureDraws>& PassDrawLists)
+{
+	const int32 NumDrawLists = PassDrawLists.Num();
+	TArray<uint32, TInlineAllocator<3>> MaxVisibleMeshDrawCommands;
+	MaxVisibleMeshDrawCommands.SetNum(NumDrawLists);
+
+	// Compute upper bounds for mesh draw commands
+	for (const FPrimitiveSceneInfo* PrimitiveSceneInfo : Scene->Primitives)
+	{
+		if (PrimitiveSceneInfo && !PrimitiveSceneInfo->Proxy->IsNaniteMesh())
+		{
+			for (int32 DrawListIdx = 0; DrawListIdx < NumDrawLists; ++DrawListIdx)
+			{
+				if (PassDrawLists[DrawListIdx].ComponentIDsToDraw->Contains(PrimitiveSceneInfo->PrimitiveComponentId))
+				{
+					MaxVisibleMeshDrawCommands[DrawListIdx] += PrimitiveSceneInfo->StaticMeshRelevances.Num();
+				}
+			}
+			
+		}
+	}
+	for (int32 DrawListIdx = 0; DrawListIdx < NumDrawLists; ++DrawListIdx)
+	{
+		PassDrawLists[DrawListIdx].InstanceRuns.Reserve(2 * MaxVisibleMeshDrawCommands[DrawListIdx]);
+	}
+
+	// Collect mesh draw commands
+	for (const FPrimitiveSceneInfo* PrimitiveSceneInfo : Scene->Primitives)
+	{
+		if (PrimitiveSceneInfo && !PrimitiveSceneInfo->Proxy->IsNaniteMesh())
+		{
+			for (int32 DrawListIdx = 0; DrawListIdx < NumDrawLists; ++DrawListIdx)
+			{
+				const EWaterInfoTextureMeshPass PassType = static_cast<EWaterInfoTextureMeshPass>(DrawListIdx);
+				FWaterInfoTextureDraws& Draws = PassDrawLists[DrawListIdx];
+				if (Draws.ComponentIDsToDraw->Contains(PrimitiveSceneInfo->PrimitiveComponentId))
+				{
+					FMeshDrawCommandPrimitiveIdInfo IdInfo(PrimitiveSceneInfo->GetIndex(), PrimitiveSceneInfo->GetInstanceSceneDataOffset());
+
+					for (int32 MeshIndex = 0; MeshIndex < PrimitiveSceneInfo->StaticMeshRelevances.Num(); MeshIndex++)
+					{
+						const FStaticMeshBatchRelevance& StaticMeshRelevance = PrimitiveSceneInfo->StaticMeshRelevances[MeshIndex];
+						const FStaticMeshBatch& StaticMesh = PrimitiveSceneInfo->StaticMeshes[MeshIndex];
+
+						const bool bUseInDepthPass = (PassType == EWaterInfoTextureMeshPass::TerrainDepth || PassType == EWaterInfoTextureMeshPass::DilatedWaterBodyDepth) && StaticMeshRelevance.bUseForWaterInfoTextureDepth;
+						const bool bUseInColorPass = (PassType == EWaterInfoTextureMeshPass::WaterBody) && StaticMeshRelevance.bUseForMaterial;
+
+						if ((bUseInDepthPass || bUseInColorPass) && StaticMeshRelevance.LODIndex == 0)
+						{
+							const int32 StaticMeshCommandInfoIndex = StaticMeshRelevance.GetStaticMeshCommandInfoIndex(Draws.MeshPass);
+							if (StaticMeshCommandInfoIndex >= 0)
+							{
+								const FCachedMeshDrawCommandInfo& CachedMeshDrawCommand = PrimitiveSceneInfo->StaticMeshCommandInfos[StaticMeshCommandInfoIndex];
+								const FCachedPassMeshDrawList& SceneDrawList = Scene->CachedDrawLists[Draws.MeshPass];
+
+								const FMeshDrawCommand* MeshDrawCommand = nullptr;
+								if (CachedMeshDrawCommand.StateBucketId >= 0)
+								{
+									MeshDrawCommand = &Scene->CachedMeshDrawCommandStateBuckets[Draws.MeshPass].GetByElementId(CachedMeshDrawCommand.StateBucketId).Key;
+								}
+								else
+								{
+									MeshDrawCommand = &SceneDrawList.MeshDrawCommands[CachedMeshDrawCommand.CommandIndex];
+								}
+
+								const uint32* InstanceRunArray = nullptr;
+								uint32 NumInstanceRuns = 0;
+
+								if (MeshDrawCommand->NumInstances > 1 && Draws.PrimitiveInstanceIndex >= 0)
+								{
+									// Render only a single specified instance, by specifying an inclusive [x;x] range
+
+									ensure(Draws.InstanceRuns.Num() + 2 <= Draws.InstanceRuns.Max());
+									InstanceRunArray = Draws.InstanceRuns.GetData() + Draws.InstanceRuns.Num();
+									NumInstanceRuns = 1;
+
+									Draws.InstanceRuns.Add(Draws.PrimitiveInstanceIndex);
+									Draws.InstanceRuns.Add(Draws.PrimitiveInstanceIndex);
+								}
+
+								FVisibleMeshDrawCommand NewVisibleMeshDrawCommand;
+
+								NewVisibleMeshDrawCommand.Setup(
+									MeshDrawCommand,
+									IdInfo,
+									CachedMeshDrawCommand.StateBucketId,
+									CachedMeshDrawCommand.MeshFillMode,
+									CachedMeshDrawCommand.MeshCullMode,
+									CachedMeshDrawCommand.Flags,
+									CachedMeshDrawCommand.SortKey,
+									CachedMeshDrawCommand.CullingPayload,
+									EMeshDrawCommandCullingPayloadFlags::NoScreenSizeCull,
+									InstanceRunArray,
+									NumInstanceRuns);
+
+								Draws.VisibleMeshCommands.Add(NewVisibleMeshDrawCommand);
+								Draws.PrimitiveIds.Add(PrimitiveSceneInfo->GetIndex());
+							}
+						}
+					}
+				}
+			}
+		}
+	}
+}
+
+void RenderWaterInfoTexture(
+	FRDGBuilder& GraphBuilder,
+	FSceneRenderer& SceneRenderer,
+	const FScene* Scene
+)
+{
+	const FViewInfo& MainView = SceneRenderer.Views[0];
+	if (MainView.WaterInfoTextureRenderingParams.IsEmpty())
+	{
+		return;
+	}
+
+	TRACE_CPUPROFILER_EVENT_SCOPE(WaterInfo::RenderWaterInfoTexture);
+	RDG_EVENT_SCOPE(GraphBuilder, "WaterInfoTexture");
+	RDG_GPU_STAT_SCOPE(GraphBuilder, WaterInfoTexture);
+	
+	for (const FViewInfo::FWaterInfoTextureRenderingParams& RenderingParams : MainView.WaterInfoTextureRenderingParams)
+	{
+		const FIntPoint TextureSize = RenderingParams.RenderTarget->GetSizeXY();
+		const FIntRect Viewport(0, 0, TextureSize.X, TextureSize.Y);
+
+		FViewInfo& WaterView = *MainView.CreateSnapshot();
+		{
+			WaterView.ViewState = nullptr;
+			WaterView.DynamicPrimitiveCollector = FGPUScenePrimitiveCollector(&SceneRenderer.GetGPUSceneDynamicContext());
+			WaterView.StereoPass = EStereoscopicPass::eSSP_FULL;
+			WaterView.DrawDynamicFlags = EDrawDynamicFlags::ForceLowestLOD;
+			WaterView.MaterialTextureMipBias = 0;
+			WaterView.PreExposure = 1.0f;
+			WaterView.ViewRect = Viewport;
+
+			WaterView.CachedViewUniformShaderParameters = MakeUnique<FViewUniformShaderParameters>();
+			WaterView.CachedViewUniformShaderParameters->OverrideLandscapeLOD = 0.0f;
+			FBox VolumeBounds[TVC_MAX];
+			WaterView.SetupUniformBufferParameters(VolumeBounds, TVC_MAX, *WaterView.CachedViewUniformShaderParameters);
+
+			WaterView.UpdateProjectionMatrix(RenderingParams.ProjectionMatrix);
+
+			FViewMatrices::FMinimalInitializer Initializer;
+			Initializer.ViewRotationMatrix = RenderingParams.ViewRotationMatrix;
+			Initializer.ViewOrigin = RenderingParams.ViewLocation;
+			Initializer.ProjectionMatrix = RenderingParams.ProjectionMatrix;
+			Initializer.ConstrainedViewRect = Viewport;
+			WaterView.ViewMatrices = FViewMatrices(Initializer);
+
+			TRefCountPtr<IPooledRenderTarget> NullRef;
+			FPlatformMemory::Memcpy(&WaterView.PrevViewInfo.HZB, &NullRef, sizeof(WaterView.PrevViewInfo.HZB));
+
+			WaterView.SetupCommonViewUniformBufferParameters(
+				*WaterView.CachedViewUniformShaderParameters, 
+				TextureSize, 
+				1, 
+				Initializer.ConstrainedViewRect, 
+				WaterView.ViewMatrices, 
+				WaterView.ViewMatrices);
+
+			WaterView.CreateViewUniformBuffers(*WaterView.CachedViewUniformShaderParameters);
+		}
+
+		const FTextureRHIRef OutputRenderTarget = RenderingParams.RenderTarget->GetRenderTargetTexture();
+		FRDGTexture* OutputTexture = GraphBuilder.RegisterExternalTexture(CreateRenderTarget(OutputRenderTarget, TEXT("WaterInfo.OutputTexture")));
+
+		FRDGTextureDesc TerrainDepthBufferDesc = FRDGTextureDesc::Create2D(TextureSize, PF_DepthStencil, FClearValueBinding::DepthFar, TexCreate_DepthStencilTargetable | TexCreate_ShaderResource);
+		FRDGTexture* TerrainDepthBuffer = GraphBuilder.CreateTexture(TerrainDepthBufferDesc, TEXT("WaterInfo.TerrainDepthTexture"));
+
+		FRDGTextureDesc WaterInfoTextureDesc = FRDGTextureDesc::Create2D(TextureSize, PF_FloatRGBA, FClearValueBinding::Black, TexCreate_RenderTargetable | TexCreate_ShaderResource);
+		FRDGTexture* WaterInfoColorTexture = GraphBuilder.CreateTexture(WaterInfoTextureDesc, TEXT("WaterInfo.ColorTexture"));
+
+		FRDGTextureDesc WaterInfoDepthBufferDesc = FRDGTextureDesc::Create2D(TextureSize, PF_DepthStencil, FClearValueBinding::DepthFar, TexCreate_DepthStencilTargetable | TexCreate_ShaderResource);
+		FRDGTexture* WaterInfoDepthBuffer = GraphBuilder.CreateTexture(WaterInfoDepthBufferDesc, TEXT("WaterInfo.DepthTexture"));
+
+		FRDGTextureDesc DilatedDepthBufferDesc = FRDGTextureDesc::Create2D(TextureSize, PF_DepthStencil, FClearValueBinding::DepthFar, TexCreate_DepthStencilTargetable | TexCreate_ShaderResource);
+		FRDGTexture* DilatedDepthBuffer = GraphBuilder.CreateTexture(DilatedDepthBufferDesc, TEXT("WaterInfo.DilatedDepthTexture"));
+
+		FRDGTextureDesc MergedTextureDesc = FRDGTextureDesc::Create2D(TextureSize, OutputTexture->Desc.Format, FClearValueBinding::Black, TexCreate_RenderTargetable | TexCreate_ShaderResource);
+		FRDGTexture* MergedTexture = GraphBuilder.CreateTexture(MergedTextureDesc, TEXT("WaterInfo.MergedTexture"));
+
+		FGlobalShaderMap* ShaderMap = GetGlobalShaderMap(WaterView.GetFeatureLevel());
+
+		TStaticArray<FWaterInfoTextureDraws, static_cast<int32>(EWaterInfoTextureMeshPass::Num)> PassDraws;
+		PassDraws[static_cast<int32>(EWaterInfoTextureMeshPass::TerrainDepth)] =			FWaterInfoTextureDraws(EMeshPass::WaterInfoTextureDepthPass, &RenderingParams.TerrainComponentIds, 0);
+		PassDraws[static_cast<int32>(EWaterInfoTextureMeshPass::DilatedWaterBodyDepth)] =	FWaterInfoTextureDraws(EMeshPass::WaterInfoTextureDepthPass, &RenderingParams.DilatedWaterBodyComponentIds, 0);
+		PassDraws[static_cast<int32>(EWaterInfoTextureMeshPass::WaterBody)] =				FWaterInfoTextureDraws(EMeshPass::WaterInfoTexturePass, &RenderingParams.WaterBodyComponentIds, 0);
+		
+		AddWaterInfoTextureDraws(Scene, PassDraws);
+
+		// Build rendering commands or prepare primitive ID buffer for the mesh draw commands
+		for (int32 PassIdx = 0; PassIdx < 3; ++PassIdx)
+		{
+			if (Scene->GPUScene.IsEnabled())
+			{
+				int32 MaxInstances = 0;
+				int32 VisibleMeshDrawCommandsNum = 0;
+				int32 NewPassVisibleMeshDrawCommandsNum = 0;
+
+				static FName NAME_WaterInfoTexturePass("WaterInfoTexture");
+				PassDraws[PassIdx].InstanceCullingContext = GraphBuilder.AllocObject<FInstanceCullingContext>(NAME_WaterInfoTexturePass, WaterView.GetShaderPlatform(), nullptr, TArrayView<const int32>(&WaterView.GPUSceneViewId, 1), nullptr);
+
+				PassDraws[PassIdx].InstanceCullingContext->SetupDrawCommands(PassDraws[PassIdx].VisibleMeshCommands, false, MaxInstances, VisibleMeshDrawCommandsNum, NewPassVisibleMeshDrawCommandsNum);
+				// Not supposed to do any compaction here.
+				ensure(VisibleMeshDrawCommandsNum == PassDraws[PassIdx].VisibleMeshCommands.Num());
+
+				PassDraws[PassIdx].InstanceCullingContext->BuildRenderingCommands(GraphBuilder, Scene->GPUScene, WaterView.DynamicPrimitiveCollector.GetInstanceSceneDataOffset(), WaterView.DynamicPrimitiveCollector.NumInstances(), PassDraws[PassIdx].InstanceCullingResult);
+			}
+
+			PassDraws[PassIdx].InstanceCullingResult.Parameters.Scene = SceneRenderer.GetSceneUniforms().GetBuffer(GraphBuilder);
+		}
+
+		TRDGUniformBufferRef<FViewUniformShaderParameters> ViewUB = GraphBuilder.CreateUniformBuffer(GraphBuilder.AllocParameters(WaterView.CachedViewUniformShaderParameters.Get()));
+
+		// Execute the three mesh passes involved in generating the water info texture: terrain depth-only, dilated water body depth-only and finally regular water body depth + river velocity
+		for (uint32 PassIndex = 0; PassIndex < static_cast<int32>(EWaterInfoTextureMeshPass::Num); ++PassIndex)
+		{
+			const EWaterInfoTextureMeshPass PassType = static_cast<EWaterInfoTextureMeshPass>(PassIndex);
+
+			// Terrain depth and dilated water body depth share the same setup, they just render a different set of meshes into different depth buffers
+			if (PassType == EWaterInfoTextureMeshPass::TerrainDepth || PassType == EWaterInfoTextureMeshPass::DilatedWaterBodyDepth)
+			{
+				const bool bIsTerrainDepthPass = PassType == EWaterInfoTextureMeshPass::TerrainDepth;
+				FWaterInfoTextureDepthPassParameters* PassParameters = GraphBuilder.AllocParameters<FWaterInfoTextureDepthPassParameters>();
+				PassParameters->View = ViewUB;
+				PassParameters->RenderTargets.DepthStencil = FDepthStencilBinding(bIsTerrainDepthPass ? TerrainDepthBuffer : DilatedDepthBuffer, ERenderTargetLoadAction::EClear, ERenderTargetLoadAction::ENoAction, FExclusiveDepthStencil::DepthWrite_StencilNop);
+				PassDraws[PassIndex].InstanceCullingResult.GetDrawParameters(PassParameters->InstanceCullingDrawParams);
+
+				GraphBuilder.AddPass(
+					bIsTerrainDepthPass ? RDG_EVENT_NAME("WaterInfoTexture(Terrain)") : RDG_EVENT_NAME("WaterInfoTexture(DilatedWaterBodies)"),
+					PassParameters,
+					ERDGPassFlags::Raster,
+					[MeshDrawCmds = MoveTemp(PassDraws[PassIndex].VisibleMeshCommands), Scene = Scene, Viewport,
+					PassParameters, InstanceCullingContext = PassDraws[PassIndex].InstanceCullingContext](FRHICommandList& RHICmdList)
+					{
+						QUICK_SCOPE_CYCLE_COUNTER(MeshPass);
+
+						RHICmdList.SetViewport(Viewport.Min.X, Viewport.Min.Y, 0.0f, Viewport.Max.X, Viewport.Max.Y, 1.0f);
+						FGraphicsMinimalPipelineStateSet PipelineStateSet;
+						if (Scene->GPUScene.IsEnabled())
+						{
+							InstanceCullingContext->SubmitDrawCommands(MeshDrawCmds, PipelineStateSet, GetMeshDrawCommandOverrideArgs(PassParameters->InstanceCullingDrawParams), 0, MeshDrawCmds.Num(), 1, RHICmdList);
+						}
+						else
+						{
+							FMeshDrawCommandSceneArgs SceneArgs;
+							SubmitMeshDrawCommandsRange(MeshDrawCmds, PipelineStateSet, SceneArgs, FInstanceCullingContext::GetInstanceIdBufferStride(Scene->GetShaderPlatform()), false, 0, MeshDrawCmds.Num(), 1, RHICmdList);
+						}
+					});
+			}
+			else if (PassType == EWaterInfoTextureMeshPass::WaterBody)
+			{
+				FWaterInfoTexturePassParameters* PassParameters = GraphBuilder.AllocParameters<FWaterInfoTexturePassParameters>();
+				PassParameters->View = ViewUB;
+				PassParameters->RenderTargets = GetRenderTargetBindings(ERenderTargetLoadAction::EClear, MakeArrayView(&WaterInfoColorTexture, 1));
+				PassParameters->RenderTargets.DepthStencil = FDepthStencilBinding(WaterInfoDepthBuffer, ERenderTargetLoadAction::EClear, ERenderTargetLoadAction::ENoAction, FExclusiveDepthStencil::DepthWrite_StencilNop);
+				PassDraws[PassIndex].InstanceCullingResult.GetDrawParameters(PassParameters->InstanceCullingDrawParams);
+
+				GraphBuilder.AddPass(
+					RDG_EVENT_NAME("WaterInfoTexture(WaterBodies)"),
+					PassParameters,
+					ERDGPassFlags::Raster,
+					[MeshDrawCmds = MoveTemp(PassDraws[PassIndex].VisibleMeshCommands), Scene = Scene, Viewport,
+					PassParameters, InstanceCullingContext = PassDraws[PassIndex].InstanceCullingContext](FRHICommandList& RHICmdList)
+					{
+						QUICK_SCOPE_CYCLE_COUNTER(MeshPass);
+
+						RHICmdList.SetViewport(Viewport.Min.X, Viewport.Min.Y, 0.0f, Viewport.Max.X, Viewport.Max.Y, 1.0f);
+						FGraphicsMinimalPipelineStateSet PipelineStateSet;
+						if (Scene->GPUScene.IsEnabled())
+						{
+							InstanceCullingContext->SubmitDrawCommands(MeshDrawCmds, PipelineStateSet, GetMeshDrawCommandOverrideArgs(PassParameters->InstanceCullingDrawParams), 0, MeshDrawCmds.Num(), 1, RHICmdList);
+						}
+						else
+						{
+							FMeshDrawCommandSceneArgs SceneArgs;
+							SubmitMeshDrawCommandsRange(MeshDrawCmds, PipelineStateSet, SceneArgs, FInstanceCullingContext::GetInstanceIdBufferStride(Scene->GetShaderPlatform()), false, 0, MeshDrawCmds.Num(), 1, RHICmdList);
+						}
+					});
+			}
+			
+		}
+
+		// Merge terrain depth, dilated water body depth, water body depth and velocity into a single texture
+		{
+			RDG_EVENT_SCOPE(GraphBuilder, "WaterInfoTextureMerge");
+
+			FWaterInfoTextureMergePS::FPermutationDomain PixelPermutationVector;
+			// The output format depends on the user configurable format of the passed in render target. Since we store depth data in the texture,
+			// it is sometimes desirable to have full 32bit float precision.
+			PixelPermutationVector.Set<FWaterInfoTextureMergePS::FEnable128BitRT>(OutputTexture->Desc.Format == PF_A32B32G32R32F);
+			TShaderMapRef<FWaterInfoTextureMergePS> PixelShader(ShaderMap, PixelPermutationVector);
+
+			FWaterInfoTextureMergePS::FParameters* PassParameters = GraphBuilder.AllocParameters<FWaterInfoTextureMergePS::FParameters>();
+			PassParameters->View = ViewUB;
+			PassParameters->RenderTargets[0] = FRenderTargetBinding(MergedTexture, ERenderTargetLoadAction::ENoAction);
+			PassParameters->SceneTextures = GetSceneTextureShaderParameters(WaterView);
+			PassParameters->GroundDepthTexture = TerrainDepthBuffer;
+			PassParameters->WaterBodyTexture = WaterInfoColorTexture;
+			PassParameters->WaterBodyDepthTexture = WaterInfoDepthBuffer;
+			PassParameters->DilatedWaterBodyDepthTexture = DilatedDepthBuffer;
+			PassParameters->CaptureZ = RenderingParams.CaptureZ;
+			PassParameters->WaterHeightExtents = RenderingParams.WaterHeightExtents;
+			PassParameters->GroundZMin = RenderingParams.GroundZMin;
+			PassParameters->DilationOverwriteMinimumDistance = CVarWaterInfoDilationOverwriteMinimumDistance.GetValueOnRenderThread();
+			PassParameters->UndergroundDilationDepthOffset = CVarWaterInfoUndergroundDilationDepthOffset.GetValueOnRenderThread();
+
+			FPixelShaderUtils::AddFullscreenPass(
+				GraphBuilder,
+				ShaderMap,
+				RDG_EVENT_NAME("WaterInfoTextureMerge"),
+				PixelShader,
+				PassParameters,
+				Viewport);
+		}
+
+		// Blur the velocity component in the water info texture to get a smooth gradient between water body transitions
+		{
+			RDG_EVENT_SCOPE(GraphBuilder, "WaterInfoTextureBlur");
+
+			FWaterInfoTextureBlurPS::FPermutationDomain PixelPermutationVector;
+			// The output format depends on the user configurable format of the passed in render target. Since we store depth data in the texture,
+			// it is sometimes desirable to have full 32bit float precision.
+			PixelPermutationVector.Set<FWaterInfoTextureBlurPS::FEnable128BitRT>(OutputTexture->Desc.Format == PF_A32B32G32R32F);
+			TShaderMapRef<FWaterInfoTextureBlurPS> PixelShader(ShaderMap, PixelPermutationVector);
+
+			FWaterInfoTextureBlurPS::FParameters* PassParameters = GraphBuilder.AllocParameters<FWaterInfoTextureBlurPS::FParameters>();
+			PassParameters->View = ViewUB;
+			PassParameters->RenderTargets[0] = FRenderTargetBinding(OutputTexture, ERenderTargetLoadAction::ENoAction);
+			PassParameters->SceneTextures = GetSceneTextureShaderParameters(WaterView);
+			PassParameters->WaterInfoTexture = MergedTexture;
+			PassParameters->WaterZMin = RenderingParams.WaterHeightExtents.X;
+			PassParameters->WaterZMax = RenderingParams.WaterHeightExtents.Y;
+			PassParameters->GroundZMin = RenderingParams.GroundZMin;
+			PassParameters->CaptureZ = RenderingParams.CaptureZ;
+			PassParameters->BlurRadius = RenderingParams.VelocityBlurRadius;
+
+			FPixelShaderUtils::AddFullscreenPass(
+				GraphBuilder,
+				ShaderMap,
+				RDG_EVENT_NAME("WaterInfoTextureBlur"),
+				PixelShader,
+				PassParameters,
+				Viewport);
+		}
+	}
+}
diff --git a/Engine/Source/Runtime/Renderer/Private/WaterInfoTextureRendering.h b/Engine/Source/Runtime/Renderer/Private/WaterInfoTextureRendering.h
new file mode 100644
--- /dev/null
+++ b/Engine/Source/Runtime/Renderer/Private/WaterInfoTextureRendering.h
@@ -0,0 +1,10 @@
+// Copyright Epic Games, Inc. All Rights Reserved.
+
+#pragma once
+
+class FRDGBuilder;
+class FSceneRenderer;
+class FScene;
+
+
+void RenderWaterInfoTexture(FRDGBuilder& GraphBuilder, FSceneRenderer& SceneRenderer, const FScene* Scene);
\ No newline at end of file
diff --git a/Engine/Source/Runtime/Renderer/Public/MeshPassProcessor.h b/Engine/Source/Runtime/Renderer/Public/MeshPassProcessor.h
--- a/Engine/Source/Runtime/Renderer/Public/MeshPassProcessor.h
+++ b/Engine/Source/Runtime/Renderer/Public/MeshPassProcessor.h
@@ -63,6 +63,8 @@ namespace EMeshPass
 		DitheredLODFadingOutMaskPass, /** A mini depth pass used to mark pixels with dithered LOD fading out. Currently only used by ray tracing shadows. */
 		NaniteMeshPass,
 		MeshDecal,
+		WaterInfoTextureDepthPass,
+		WaterInfoTexturePass,
 
 #if WITH_EDITOR
 		HitProxy,
@@ -112,6 +114,8 @@ inline const TCHAR* GetMeshPassName(EMeshPass::Type MeshPass)
 	case EMeshPass::DitheredLODFadingOutMaskPass: return TEXT("DitheredLODFadingOutMaskPass");
 	case EMeshPass::NaniteMeshPass: return TEXT("NaniteMeshPass");
 	case EMeshPass::MeshDecal: return TEXT("MeshDecal");
+	case EMeshPass::WaterInfoTextureDepthPass: return TEXT("WaterInfoTextureDepthPass");
+	case EMeshPass::WaterInfoTexturePass: return TEXT("WaterInfoTexturePass");
 #if WITH_EDITOR
 	case EMeshPass::HitProxy: return TEXT("HitProxy");
 	case EMeshPass::HitProxyOpaqueOnly: return TEXT("HitProxyOpaqueOnly");
@@ -121,9 +125,9 @@ inline const TCHAR* GetMeshPassName(EMeshPass::Type MeshPass)
 	}
 
 #if WITH_EDITOR
-	static_assert(EMeshPass::Num == 30 + 4, "Need to update switch(MeshPass) after changing EMeshPass"); // GUID to prevent incorrect auto-resolves, please change when changing the expression: {A6E82589-44B3-4DAD-AC57-8AF6BD50DF43}
+	static_assert(EMeshPass::Num == 32 + 4, "Need to update switch(MeshPass) after changing EMeshPass"); // GUID to prevent incorrect auto-resolves, please change when changing the expression: {674D7D62-CFD8-4971-9A8D-CD91E5612CD8}
 #else
-	static_assert(EMeshPass::Num == 30, "Need to update switch(MeshPass) after changing EMeshPass"); // GUID to prevent incorrect auto-resolves, please change when changing the expression: {A6E82589-44B3-4DAD-AC57-8AF6BD50DF43}
+	static_assert(EMeshPass::Num == 32, "Need to update switch(MeshPass) after changing EMeshPass"); // GUID to prevent incorrect auto-resolves, please change when changing the expression: {674D7D62-CFD8-4971-9A8D-CD91E5612CD8}
 #endif
 
 	checkf(0, TEXT("Missing case for EMeshPass %u"), (uint32)MeshPass);
diff --git a/Engine/Source/Runtime/Renderer/Public/StaticMeshBatch.h b/Engine/Source/Runtime/Renderer/Public/StaticMeshBatch.h
--- a/Engine/Source/Runtime/Renderer/Public/StaticMeshBatch.h
+++ b/Engine/Source/Runtime/Renderer/Public/StaticMeshBatch.h
@@ -37,7 +37,7 @@ private:
 class FStaticMeshBatchRelevance
 {
 public:
-	FStaticMeshBatchRelevance(const FStaticMeshBatch& StaticMesh, float InScreenSize, bool InbSupportsCachingMeshDrawCommands, bool InbUseSkyMaterial, bool bInUseSingleLayerWaterMaterial, bool bInUseAnisotropy, bool bInSupportsNaniteRendering, bool bInSupportsGPUScene, ERHIFeatureLevel::Type FeatureLevel);
+	FStaticMeshBatchRelevance(const FStaticMeshBatch& StaticMesh, float InScreenSize, bool InbSupportsCachingMeshDrawCommands, bool InbUseSkyMaterial, bool bInUseSingleLayerWaterMaterial, bool bInUseAnisotropy, bool bInSupportsNaniteRendering, bool bInSupportsGPUScene, bool bInUseForWaterInfoTextureDepth, ERHIFeatureLevel::Type FeatureLevel);
 
 	/** Starting offset into continuous array of command infos for this mesh in FPrimitiveSceneInfo::CachedMeshDrawCommandInfos. */
 	FMeshPassMask CommandInfosMask;
@@ -87,6 +87,9 @@ public:
 	/** Cached from vertex factory to avoid dereferencing VF in shadow depth rendering. */
 	uint8 bSupportsGPUScene : 1;
 
+	/** Whether the mesh batch should be used in the depth-only passes of rendering the water info texture for the water plugin */
+	uint8 bUseForWaterInfoTextureDepth : 1;
+
 	/** Computes index of cached mesh draw command in FPrimitiveSceneInfo::CachedMeshDrawCommandInfos, for a given mesh pass. */
 	int32 GetStaticMeshCommandInfoIndex(EMeshPass::Type MeshPass) const;
 };
