Removed locks from VT update to fix deadlock due to busywaiting. This is technically safe as new VT allocations aren't allowed during updates anyway.

diff --git a/Engine/Source/Runtime/Renderer/Private/VT/VirtualTextureSystem.cpp b/Engine/Source/Runtime/Renderer/Private/VT/VirtualTextureSystem.cpp
--- a/Engine/Source/Runtime/Renderer/Private/VT/VirtualTextureSystem.cpp
+++ b/Engine/Source/Runtime/Renderer/Private/VT/VirtualTextureSystem.cpp
@@ -2630,8 +2630,6 @@ void FVirtualTextureSystem::BeginUpdate(FRDGBuilder& GraphBuilder, FVirtualTextu
 
 		const bool bContinousUpdates = false;
 
-		UE::TScopeLock Lock(Mutex);
-
 		Updater->MergedRequestList = Allocator.Create<FUniqueRequestList>(Allocator);
 		Updater->MergedRequestList->Initialize();
 
@@ -2659,8 +2657,6 @@ TUniquePtr<FVirtualTextureUpdater> FVirtualTextureSystem::BeginUpdate(FRDGBuilde
 	check(IsInRenderingThread());
 	check(!bUpdating);
 
-	UE::TScopeLock Lock(Mutex);
-
 	checkf(Producers.HasPendingCallbacks() == false, TEXT("FVirtualTextureSystem::CallPendingCallbacks(), called in UpdateAllPrimitiveSceneInfos(), must run before FVirtualTextureSystem::BeginUpdate()"));
 
 	AllocateResources(GraphBuilder);
@@ -2787,8 +2783,6 @@ void FVirtualTextureSystem::EndUpdate(FRDGBuilder& GraphBuilder, TUniquePtr<FVir
 	}
 	bUpdating = false;
 
-	UE::TScopeLock Lock(Mutex);
-
 	if (Updater->FeedbackMapResult.Data)
 	{
 		GVirtualTextureFeedback.Unmap(GraphBuilder.RHICmdList, Updater->FeedbackMapResult.MapHandle);
