Fix artifacts when using LocalExposure in stereo, which only consistently became apparent at high screen percentages.
Set GaussianBlur to use the extent of its GaussianLum texture rather than using the current view's ViewRect, since the GaussianLum texture will always be single-wide.

diff --git a/Engine/Source/Runtime/Renderer/Private/PostProcess/PostProcessLocalExposure.cpp b/Engine/Source/Runtime/Renderer/Private/PostProcess/PostProcessLocalExposure.cpp
--- a/Engine/Source/Runtime/Renderer/Private/PostProcess/PostProcessLocalExposure.cpp
+++ b/Engine/Source/Runtime/Renderer/Private/PostProcess/PostProcessLocalExposure.cpp
@@ -164,7 +164,7 @@ FRDGTextureRef AddLocalExposureBlurredLogLuminancePass(
 		FGaussianBlurInputs GaussianBlurInputs;
 		GaussianBlurInputs.NameX = TEXT("LocalExposureGaussianX");
 		GaussianBlurInputs.NameY = TEXT("LocalExposureGaussianY");
-		GaussianBlurInputs.Filter = FScreenPassTexture(GaussianLumSetupTexture, InputTexture.ViewRect);
+		GaussianBlurInputs.Filter = FScreenPassTexture(GaussianLumSetupTexture);
 		GaussianBlurInputs.TintColor = FLinearColor::White;
 		GaussianBlurInputs.CrossCenterWeight = FVector2f::ZeroVector;
 		GaussianBlurInputs.KernelSizePercent = View.FinalPostProcessSettings.LocalExposureBlurredLuminanceKernelSizePercent;
