Cleaned up initial move of distant light scheduling (CL 27373452) and removed the old path & cvar that was preserved for hotfixing.

diff --git a/Engine/Source/Runtime/Renderer/Private/SceneRendering.h b/Engine/Source/Runtime/Renderer/Private/SceneRendering.h
--- a/Engine/Source/Runtime/Renderer/Private/SceneRendering.h
+++ b/Engine/Source/Runtime/Renderer/Private/SceneRendering.h
@@ -2061,9 +2061,10 @@ public:
 	/** True if precomputed visibility was used when rendering the scene. */
 	bool bUsedPrecomputedVisibility;
 
+#if !(UE_BUILD_SHIPPING || UE_BUILD_TEST)
 	/** Lights added if wholescenepointlight shadow would have been rendered (ignoring r.SupportPointLightWholeSceneShadows). Used for warning about unsupported features. */
 	TArray<FString, SceneRenderingAllocator> UsedWholeScenePointLightNames;
-
+#endif
 	/** Feature level being rendered */
 	ERHIFeatureLevel::Type FeatureLevel;
 	EShaderPlatform ShaderPlatform;
diff --git a/Engine/Source/Runtime/Renderer/Private/ShadowSetup.cpp b/Engine/Source/Runtime/Renderer/Private/ShadowSetup.cpp
--- a/Engine/Source/Runtime/Renderer/Private/ShadowSetup.cpp
+++ b/Engine/Source/Runtime/Renderer/Private/ShadowSetup.cpp
@@ -4087,7 +4087,7 @@ void FSceneRenderer::CreateWholeSceneProjectedShadow(
 					bool bContainsNaniteSubjects = false;
 
 					// Skip mesh setup if it won't be rendered anyway
-					if (!ShadowSceneRenderer->IsUsingNewDistantInvalidationLogic() || ProjectedShadowInfo->bShouldRenderVSM)
+					if (ProjectedShadowInfo->bShouldRenderVSM)
 					{
 						// Skip convex hull tests for VSM since this causes artifacts due to caching (potentially check if caching is enabled but may lead to race condition).
 						// The interaction setup has already tested the light bounds (by calling FLightSceneProxy::AffectsBounds)
@@ -5828,11 +5828,13 @@ void FSceneRenderer::CreateDynamicShadows(FDynamicShadowsTaskData& TaskData)
 							// Stationary point light and spot light shadow are unsupported on mobile
 							&& (!bMobile || bDirectionalLightShadow);
 
+#if !(UE_BUILD_SHIPPING || UE_BUILD_TEST)
 						const bool bPointLightWholeSceneShadow = (bShouldCreateShadowForMovableLight || bShouldCreateShadowForOverflowStaticShadowing || bShouldCreateShadowToPreviewStaticLight) && bPointLightShadow;
 						if (bPointLightWholeSceneShadow)
 						{
 							UsedWholeScenePointLightNames.Add(LightSceneInfoCompact.LightSceneInfo->Proxy->GetOwnerNameOrLabel());
 						}
+#endif
 
 						if (bCreateShadowForMovableLight || bCreateShadowToPreviewStaticLight || bCreateShadowForOverflowStaticShadowing)
 						{
diff --git a/Engine/Source/Runtime/Renderer/Private/Shadows/ShadowSceneRenderer.cpp b/Engine/Source/Runtime/Renderer/Private/Shadows/ShadowSceneRenderer.cpp
--- a/Engine/Source/Runtime/Renderer/Private/Shadows/ShadowSceneRenderer.cpp
+++ b/Engine/Source/Runtime/Renderer/Private/Shadows/ShadowSceneRenderer.cpp
@@ -74,13 +74,6 @@ static TAutoConsoleVariable<float> CVarResolutionLodBiasLocalMoving(
 	ECVF_Scalability | ECVF_RenderThreadSafe
 );
 
-TAutoConsoleVariable<int32> CVarEnableNewDistantInvalidationLogic(
-	TEXT("r.Shadow.Virtual.EnableNewDistantInvalidationLogic"),
-	1,
-	TEXT("Only intended to allow hotfixing the next release of a certain project."),
-	ECVF_RenderThreadSafe
-);
-
 DECLARE_DWORD_COUNTER_STAT(TEXT("VSM Total Raster Bins"), STAT_VSMNaniteBasePassTotalRasterBins, STATGROUP_ShadowRendering);
 DECLARE_DWORD_COUNTER_STAT(TEXT("VSM Total Shading Draws"), STAT_VSMNaniteBasePassTotalShadingDraws, STATGROUP_ShadowRendering);
 
@@ -99,7 +92,6 @@ FShadowSceneRenderer::FShadowSceneRenderer(FDeferredShadingSceneRenderer& InScen
 	, Scene(*InSceneRenderer.Scene)
 	, ShadowScene(*Scene.ShadowScene)
 	, VirtualShadowMapArray(InSceneRenderer.VirtualShadowMapArray)
-	, bShouldUseNewDistantInvalidationLogic(CVarEnableNewDistantInvalidationLogic.GetValueOnAnyThread() != 0)
 {
 }
 
@@ -108,46 +100,68 @@ float FShadowSceneRenderer::ComputeNaniteShadowsLODScaleFactor()
 	return FMath::Pow(2.0f, -CVarNaniteShadowsLODBias.GetValueOnRenderThread());
 }
 
+namespace
+{
+
+struct FHeapPair
+{
+	int32 Age;
+	TSharedPtr<FVirtualShadowMapPerLightCacheEntry> CacheEntry;
+
+	// Order for a min-heap, we always want to replace the least-old item
+	bool operator <(const FHeapPair& Other) const { return Age < Other.Age; }
+};
+
+}
+
 void FShadowSceneRenderer::BeginRender(FRDGBuilder& GraphBuilder)
 {
 	// Kick off shadow scene updates.
 	ShadowScene.UpdateForRenderedFrame(GraphBuilder);
 
-	if (bShouldUseNewDistantInvalidationLogic && VirtualShadowMapArray.IsEnabled() && VirtualShadowMapArray.CacheManager->IsCacheEnabled())
+	// Priority queue of distant lights to update.
+	const int32 MaxToUpdate = CVarMaxDistantLightsPerFrame.GetValueOnRenderThread() < 0 ? INT32_MAX : CVarMaxDistantLightsPerFrame.GetValueOnRenderThread();
+
+	if (MaxToUpdate == 0 || !VirtualShadowMapArray.IsEnabled() || !VirtualShadowMapArray.CacheManager->IsCacheEnabled())
 	{
-		RendererSetupTask = GraphBuilder.AddSetupTask([this]()
-		{
-			FVirtualShadowMapArrayCacheManager& CacheManager = *VirtualShadowMapArray.CacheManager;
-			// Priority queue of distant lights to update.
-			int32 SceneFrameNumber = int32(Scene.GetFrameNumber());
+		return;
+	}	
+	
+	RendererSetupTask = GraphBuilder.AddSetupTask([this, MaxToUpdate]()
+	{
+		FVirtualShadowMapArrayCacheManager& CacheManager = *VirtualShadowMapArray.CacheManager;
 
-			// Build age-based priority queue of fully cached entries.
-			for (auto It = Scene.Lights.CreateConstIterator(); It; ++It)
+		TArray<FHeapPair> DistantLightUpdateQueue;
+		int32 SceneFrameNumber = int32(Scene.GetFrameNumber());
+		for (auto It = CacheManager.CreateConstEntryIterator(); It; ++It)
+		{
+			TSharedPtr<FVirtualShadowMapPerLightCacheEntry> PerLightCacheEntry = It.Value();
+			if (PerLightCacheEntry->IsFullyCached())
 			{
-				int32 LightId = It.GetIndex();
-				TSharedPtr<FVirtualShadowMapPerLightCacheEntry> PerLightCacheEntry = CacheManager.FindLightCacheEntry(LightId, 0);
-				if (PerLightCacheEntry && PerLightCacheEntry->IsFullyCached())
+				int32 Age = SceneFrameNumber - int32(PerLightCacheEntry->GetLastScheduledFrameNumber());
+				if (DistantLightUpdateQueue.Num() < MaxToUpdate)
 				{
-					int32 FramesSinceLastRender = SceneFrameNumber - int32(PerLightCacheEntry->GetLastScheduledFrameNumber());
-					DistantLightUpdateQueue.Add(-FramesSinceLastRender, LightId);
+					DistantLightUpdateQueue.HeapPush(FHeapPair{ Age, PerLightCacheEntry});
+				}
+				else
+				{
+					// Queue is full, but we found an older item
+					if (DistantLightUpdateQueue.HeapTop().Age < Age)
+					{
+						// Replace heap top and restore heap property.
+						DistantLightUpdateQueue[0] = FHeapPair{ Age, PerLightCacheEntry };
+						AlgoImpl::HeapSiftDown(DistantLightUpdateQueue.GetData(), 0, DistantLightUpdateQueue.Num(), FIdentityFunctor(), TLess<FHeapPair>());
+					}
 				}
 			}
+		}
 
-			int32 UpdateBudgetNumLights = CVarMaxDistantLightsPerFrame.GetValueOnRenderThread() < 0 ? int32(DistantLightUpdateQueue.Num()) : FMath::Min(int32(DistantLightUpdateQueue.Num()), CVarMaxDistantLightsPerFrame.GetValueOnRenderThread());
-			for (int32 Index = 0; Index < UpdateBudgetNumLights; ++Index)
-			{
-				const int32 LightId = DistantLightUpdateQueue.Top();
-				//const int32 Age = DistantLightUpdateQueue.GetKey(LocalLightShadowIndex);
-				// UE_LOG(LogTemp, Log, TEXT("Index: %d Age: %d"), LocalLightShadowIndex, Age);
-				DistantLightUpdateQueue.Pop();
-
-				TSharedPtr<FVirtualShadowMapPerLightCacheEntry> PerLightCacheEntry = CacheManager.FindLightCacheEntry(LightId, 0);
-
-				// Mark frame it was scheduled, this is picked up later in AddLocalLightShadow to trigger invalidation 
-				PerLightCacheEntry->Current.ScheduledFrameNumber = SceneFrameNumber;
-			}
-		});
-	}
+		for (const FHeapPair &HeapPair : DistantLightUpdateQueue)
+		{
+			// Mark frame it was scheduled, this is picked up later in AddLocalLightShadow to trigger invalidation 
+			HeapPair.CacheEntry->Current.ScheduledFrameNumber = SceneFrameNumber;
+		}
+	});
 }
 
 static float GetResolutionLODBiasLocal(float LightMobilityFactor)
@@ -220,13 +234,11 @@ TSharedPtr<FVirtualShadowMapPerLightCacheEntry> FShadowSceneRenderer::AddLocalLi
 	LocalLightShadowFrameSetup.PerLightCacheEntry = PerLightCacheEntry;
 	bool bIsCached = PerLightCacheEntry->UpdateLocal(ProjectedShadowInitializer, bIsDistantLight, CacheManager->IsCacheEnabled(), !bShouldForceTimeSliceDistantUpdate);
 
-	if (bShouldUseNewDistantInvalidationLogic)
+	if (bIsCached && bIsDistantLight && PerLightCacheEntry->Prev.ScheduledFrameNumber == Scene.GetFrameNumber())
 	{
-		if (bIsCached && bIsDistantLight && PerLightCacheEntry->Prev.ScheduledFrameNumber == Scene.GetFrameNumber())
-		{
-			PerLightCacheEntry->Invalidate();
-		}
+		PerLightCacheEntry->Invalidate();
 	}
+
 	// Update info on the ProjectionShadowInfo; eventually this should all move into local data structures here
 	const int32 VirtualShadowMapId = VirtualShadowMapArray.Allocate(bIsDistantLight, NumMaps);
 	ProjectedShadowInfo->VirtualShadowMapId = VirtualShadowMapId;
@@ -242,17 +254,6 @@ TSharedPtr<FVirtualShadowMapPerLightCacheEntry> FShadowSceneRenderer::AddLocalLi
 		VirtualSmCacheEntry.ProjectionData = GetLocalLightProjectionShaderData(ResolutionLODBiasLocal, ProjectedShadowInfo, Index);
 	}
 
-	if (!bShouldUseNewDistantInvalidationLogic)
-	{
-		// Only round-robin those that were not invalidated.
-		if (bIsDistantLight && bIsCached)
-		{
-			// This priority could be calculated based also on whether the light has actually been invalidated or not (currently not tracked on host).
-			// E.g., all things being equal update those with an animated mesh in, for example. Plus don't update those the don't need it at all.
-			int32 FramesSinceLastRender = int32(Scene.GetFrameNumber()) - int32(PerLightCacheEntry->GetLastScheduledFrameNumber());
-			DistantLightUpdateQueue.Add(-FramesSinceLastRender, LocalLightShadowIndex);
-		}
-	}
 	return PerLightCacheEntry;
 }
 
@@ -264,10 +265,6 @@ void FShadowSceneRenderer::AddDirectionalLightShadow(FProjectedShadowInfo* Proje
 
 void FShadowSceneRenderer::PostInitDynamicShadowsSetup()
 {
-	if (!bShouldUseNewDistantInvalidationLogic)
-	{
-		UpdateDistantLightPriorityRender();
-	}
 	// Dispatch async Nanite culling job if appropriate
 	if (CVarVSMMaterialVisibility.GetValueOnRenderThread() != 0)
 	{
@@ -371,27 +368,6 @@ void FShadowSceneRenderer::RenderVirtualShadowMaps(FRDGBuilder& GraphBuilder, bo
 	VirtualShadowMapArray.MergeStaticPhysicalPages(GraphBuilder);
 }
 
-void FShadowSceneRenderer::UpdateDistantLightPriorityRender()
-{
-	int32 UpdateBudgetNumLights = CVarMaxDistantLightsPerFrame.GetValueOnRenderThread() < 0 ? int32(DistantLightUpdateQueue.Num()) : FMath::Min(int32(DistantLightUpdateQueue.Num()), CVarMaxDistantLightsPerFrame.GetValueOnRenderThread());
-	for (int32 Index = 0; Index < UpdateBudgetNumLights; ++Index)
-	{
-		const int32 LocalLightShadowIndex = DistantLightUpdateQueue.Top();
-		const int32 Age = DistantLightUpdateQueue.GetKey(LocalLightShadowIndex);
-		// UE_LOG(LogTemp, Log, TEXT("Index: %d Age: %d"), LocalLightShadowIndex, Age);
-		DistantLightUpdateQueue.Pop();
-
-		FLocalLightShadowFrameSetup& LocalLightShadowFrameSetup = LocalLights[LocalLightShadowIndex];
-		
-		// Force fully cached to be off.
-		LocalLightShadowFrameSetup.ProjectedShadowInfo->bShouldRenderVSM = true;
-		LocalLightShadowFrameSetup.PerLightCacheEntry->Current.ScheduledFrameNumber = Scene.GetFrameNumber();
-		// Should trigger invalidations also.
-		LocalLightShadowFrameSetup.PerLightCacheEntry->Invalidate();
-	}
-}
-
-
 void FShadowSceneRenderer::DispatchVirtualShadowMapViewAndCullingSetup(FRDGBuilder& GraphBuilder, TConstArrayView<FProjectedShadowInfo*> VirtualShadowMapShadows)
 {
 	// Don't want to run this more than once in a given frame.
diff --git a/Engine/Source/Runtime/Renderer/Private/Shadows/ShadowSceneRenderer.h b/Engine/Source/Runtime/Renderer/Private/Shadows/ShadowSceneRenderer.h
--- a/Engine/Source/Runtime/Renderer/Private/Shadows/ShadowSceneRenderer.h
+++ b/Engine/Source/Runtime/Renderer/Private/Shadows/ShadowSceneRenderer.h
@@ -4,7 +4,6 @@
 #include "CoreMinimal.h"
 #include "Containers/SparseArray.h"
 #include "Containers/ArrayView.h"
-#include "Containers/BinaryHeap.h"
 #include "Containers/Array.h"
 #include "Tasks/Task.h"
 
@@ -97,17 +96,11 @@ public:
 		return RendererSetupTask;
 	}
 
-	bool IsUsingNewDistantInvalidationLogic() const { return bShouldUseNewDistantInvalidationLogic; }
 private:
 	UE::Tasks::FTask RendererSetupTask;
 
 	FVirtualShadowMapProjectionShaderData GetLocalLightProjectionShaderData(float ResolutionLODBiasLocal, const FProjectedShadowInfo* ProjectedShadowInfo, int32 MapIndex) const;
 
-	/**
-	 * Select the budgeted set of distant lights to update this frame.
-	 */
-	void UpdateDistantLightPriorityRender();
-
 	struct FLocalLightShadowFrameSetup
 	{
 		TSharedPtr<FVirtualShadowMapPerLightCacheEntry> PerLightCacheEntry;
@@ -126,9 +119,6 @@ private:
 	};
 	TArray<FDirectionalLightShadowFrameSetup, SceneRenderingAllocator> DirectionalLights;
 
-	// Priority queue of distant lights to update.
-	FBinaryHeap<int32, uint32> DistantLightUpdateQueue;
-
 	// One pass projection stuff. Set up in RenderVitualShadowMapProjectionMaskBits
 	bool bShouldUseVirtualShadowMapOnePassProjection = false;
 
@@ -141,6 +131,4 @@ private:
 	FNaniteVisibilityQuery* NaniteVisibilityQuery = nullptr;
 	Nanite::FPackedViewArray* VirtualShadowMapViews = nullptr;
 	FSceneInstanceCullingQuery *SceneInstanceCullingQuery = nullptr;
-
-	bool bShouldUseNewDistantInvalidationLogic = true;
 };
diff --git a/Engine/Source/Runtime/Renderer/Private/VirtualShadowMaps/VirtualShadowMapCacheManager.cpp b/Engine/Source/Runtime/Renderer/Private/VirtualShadowMaps/VirtualShadowMapCacheManager.cpp
--- a/Engine/Source/Runtime/Renderer/Private/VirtualShadowMaps/VirtualShadowMapCacheManager.cpp
+++ b/Engine/Source/Runtime/Renderer/Private/VirtualShadowMaps/VirtualShadowMapCacheManager.cpp
@@ -811,7 +811,7 @@ bool FVirtualShadowMapArrayCacheManager::IsHZBDataAvailable()
 TSharedPtr<FVirtualShadowMapPerLightCacheEntry> FVirtualShadowMapArrayCacheManager::FindCreateLightCacheEntry(
 	int32 LightSceneId, uint32 ViewUniqueID, uint32 NumShadowMaps)
 {
-	const uint64 CacheKey = (uint64(ViewUniqueID) << 32U) | uint64(LightSceneId);
+	const FVirtualShadowMapCacheKey CacheKey = { ViewUniqueID, LightSceneId };
 
 	TSharedPtr<FVirtualShadowMapPerLightCacheEntry> *LightEntryKey = CacheEntries.Find(CacheKey);
 
@@ -843,17 +843,6 @@ TSharedPtr<FVirtualShadowMapPerLightCacheEntry> FVirtualShadowMapArrayCacheManag
 	return LightEntry;
 }
 
-TSharedPtr<FVirtualShadowMapPerLightCacheEntry> FVirtualShadowMapArrayCacheManager::FindLightCacheEntry(int32 LightSceneId, uint32 ViewUniqueID)
-{
-	const uint64 CacheKey = (uint64(ViewUniqueID) << 32U) | uint64(LightSceneId);
-
-	if (TSharedPtr<FVirtualShadowMapPerLightCacheEntry> *Entry = CacheEntries.Find(CacheKey))
-	{
-		return *Entry;
-	}
-	return TSharedPtr<FVirtualShadowMapPerLightCacheEntry>();
-}
-
 void FVirtualShadowMapPerLightCacheEntry::OnPrimitiveRendered(const FPrimitiveSceneInfo* PrimitiveSceneInfo)
 {
 	// Mark as (potentially present in a cached page somehwere, so we'd need to invalidate if it is removed/moved)
@@ -878,10 +867,9 @@ void FVirtualShadowMapArrayCacheManager::UpdateUnreferencedCacheEntries(
 {
 	const uint32 SceneFrameNumber = Scene->GetFrameNumberRenderThread();
 
-	TArray<uint64, SceneRenderingAllocator> EntriesToRemove;
-	for (auto& LightEntry : CacheEntries)
+	for (FEntryMap::TIterator It = CacheEntries.CreateIterator(); It; ++It)
 	{
-		TSharedPtr<FVirtualShadowMapPerLightCacheEntry> CacheEntry = LightEntry.Value;
+		TSharedPtr<FVirtualShadowMapPerLightCacheEntry> CacheEntry = (*It).Value;
 		// For this test we care if it is active *this render*, not just this scene frame number (which can include multiple renders)
 		if (CacheEntry->bReferencedThisRender)
 		{
@@ -908,17 +896,9 @@ void FVirtualShadowMapArrayCacheManager::UpdateUnreferencedCacheEntries(
 		}
 		else
 		{
-			// Enqueue for remove
-			EntriesToRemove.Add(LightEntry.Key);
-			//UE_LOG(LogRenderer, Display, TEXT("Removed VSM light cache entry (%d, Age %d)"), LightEntry.Key, Age);
+			It.RemoveCurrent();
 		}
 	}
-
-	for (uint64 Entry : EntriesToRemove)
-	{
-		int32 NumRemoved = CacheEntries.Remove(Entry);
-		check(NumRemoved > 0);
-	}
 }
 
 void FVirtualShadowMapArrayCacheManager::UploadProjectionData(FRDGScatterUploadBuffer& Uploader) const
@@ -1223,7 +1203,8 @@ void FVirtualShadowMapArrayCacheManager::OnSceneChange()
 
 void FVirtualShadowMapArrayCacheManager::OnLightRemoved(int32 LightId)
 {
-	CacheEntries.Remove(LightId);
+	const FVirtualShadowMapCacheKey CacheKey = { /* TODO: this is broken for directional lights! ViewUniqueID */0, LightId };
+	CacheEntries.Remove(CacheKey);
 }
 
 /**
diff --git a/Engine/Source/Runtime/Renderer/Private/VirtualShadowMaps/VirtualShadowMapCacheManager.h b/Engine/Source/Runtime/Renderer/Private/VirtualShadowMaps/VirtualShadowMapCacheManager.h
--- a/Engine/Source/Runtime/Renderer/Private/VirtualShadowMaps/VirtualShadowMapCacheManager.h
+++ b/Engine/Source/Runtime/Renderer/Private/VirtualShadowMaps/VirtualShadowMapCacheManager.h
@@ -188,9 +188,24 @@ struct FPhysicalPageMetaData
 	FUintPoint PageAddress;
 };
 
+struct FVirtualShadowMapCacheKey
+{
+	uint32 ViewUniqueID;
+	uint32 LightSceneId;
+
+	inline bool operator==(const FVirtualShadowMapCacheKey& Other) const { return ViewUniqueID == Other.ViewUniqueID && Other.LightSceneId == LightSceneId; }
+};
+
+inline uint32 GetTypeHash(FVirtualShadowMapCacheKey Key)
+{
+	return GetTypeHash(Key.LightSceneId) ^ GetTypeHash(Key.ViewUniqueID);
+}
+
 class FVirtualShadowMapArrayCacheManager
 {
 public:
+	using FEntryMap = TMap< FVirtualShadowMapCacheKey, TSharedPtr<FVirtualShadowMapPerLightCacheEntry> >;
+
 	FVirtualShadowMapArrayCacheManager(FScene *InScene);
 	~FVirtualShadowMapArrayCacheManager();
 
@@ -236,11 +251,6 @@ public:
 	 */
 	TSharedPtr<FVirtualShadowMapPerLightCacheEntry> FindCreateLightCacheEntry(int32 LightSceneId, uint32 ViewUniqueID, uint32 NumShadowMaps);
 
-	/**
-	 * Finds an existing cache entry, returns null if none exists.
-	 */
-	TSharedPtr<FVirtualShadowMapPerLightCacheEntry> FindLightCacheEntry(int32 LightSceneId, uint32 ViewUniqueID);
-
 	bool IsCacheEnabled();
 	bool IsCacheDataAvailable();
 	bool IsHZBDataAvailable();
@@ -328,6 +338,16 @@ public:
 
 	float GetGlobalResolutionLodBias() const { return GlobalResolutionLodBias; }
 
+	inline FEntryMap::TIterator CreateEntryIterator()
+	{
+		return CacheEntries.CreateIterator();
+	}
+
+	inline FEntryMap::TConstIterator CreateConstEntryIterator() const
+	{
+		return CacheEntries.CreateConstIterator();
+	}
+
 private:
 	// Invalidate the cache for all shadows, causing any pages to be rerendered
 	void Invalidate(FRDGBuilder& GraphBuilder);
@@ -366,7 +386,7 @@ private:
 	uint32 MaxPhysicalPages = 0;
 
 	// Index the Cache entries by the light ID
-	TMap< uint64, TSharedPtr<FVirtualShadowMapPerLightCacheEntry> > CacheEntries;
+	FEntryMap CacheEntries;
 
 	// Tracks primitives (by persistent primitive index) that have been removed recently
 	// This allows us to ignore feedback from previous frames in the case of persistent primitive indices being
