Skip shadow maps that don't need nanite rendering when VSM is enabled. Got lost in previous refactor.

diff --git a/Engine/Source/Runtime/Renderer/Private/ShadowDepthRendering.cpp b/Engine/Source/Runtime/Renderer/Private/ShadowDepthRendering.cpp
--- a/Engine/Source/Runtime/Renderer/Private/ShadowDepthRendering.cpp
+++ b/Engine/Source/Runtime/Renderer/Private/ShadowDepthRendering.cpp
@@ -1374,6 +1374,11 @@ static void RenderShadowDepthAtlasNanite(
 	{
 		FProjectedShadowInfo* ProjectedShadowInfo = ShadowMapAtlas.Shadows[ShadowIndex];
 
+		if (!ProjectedShadowInfo->bNaniteGeometry)
+		{
+			continue;
+		}
+
 		if (ProjectedShadowInfo->CacheMode == SDCM_StaticPrimitivesOnly)
 		{
 			StaticShadowsToEmit.Add(ProjectedShadowInfo);
