Always mark IES Atlas with UAV flag in Editor to avoid issues with feature level

diff --git a/Engine/Source/Runtime/Renderer/Private/IESTextureManager.cpp b/Engine/Source/Runtime/Renderer/Private/IESTextureManager.cpp
--- a/Engine/Source/Runtime/Renderer/Private/IESTextureManager.cpp
+++ b/Engine/Source/Runtime/Renderer/Private/IESTextureManager.cpp
@@ -427,6 +427,11 @@ void RemoveTexture(uint32 InSlotIndex)
 static FRDGTextureRef CreateAtlasTexture(FRDGBuilder& GraphBuilder, ERHIFeatureLevel::Type FeatureLevel, const FIntPoint& Resolution, uint32 SliceCount)
 {
 	ETextureCreateFlags CreateAddFlags = UseComputeForAtlasUpdate(FeatureLevel) ? ETextureCreateFlags::UAV : ETextureCreateFlags::TargetArraySlicesIndependently;
+	if (GIsEditor)
+	{
+		// Make sure UAV flag is always present in Editor environment as target can be used by both compute and raster
+		CreateAddFlags |= ETextureCreateFlags::UAV;
+	}
 	
 	return GraphBuilder.CreateTexture(FRDGTextureDesc::Create2DArray(
 		Resolution,
