Support for multiple GPU transfer when SceneCapture is rendered as additional pass of the main renderer

diff --git a/Engine/Source/Runtime/Renderer/Private/DeferredShadingRenderer.cpp b/Engine/Source/Runtime/Renderer/Private/DeferredShadingRenderer.cpp
--- a/Engine/Source/Runtime/Renderer/Private/DeferredShadingRenderer.cpp
+++ b/Engine/Source/Runtime/Renderer/Private/DeferredShadingRenderer.cpp
@@ -3329,6 +3329,14 @@ void FDeferredShadingSceneRenderer::Render(FRDGBuilder& GraphBuilder)
 		RenderPrepassAndVelocity(PassInfo.Views, DummyNaniteBasePassVisibility, NaniteRasterResults, PrimaryNaniteViews);
 
 		CopySceneCaptureComponentToTarget(GraphBuilder, SceneTextures, nullptr, ViewFamily, PassInfo.Views);
+
+#if WITH_MGPU
+		const FRenderTarget* RenderTarget = PassInfo.Views.Num() > 0 ? PassInfo.Views[0].SceneCaptureRenderTarget : nullptr;
+		if (RenderTarget)
+		{
+			DoCrossGPUTransfers(GraphBuilder, RenderTarget->GetRenderTargetTexture(GraphBuilder), PassInfo.Views, false, FRHIGPUMask::All());
+		}
+#endif
 	}
 
 	TArray<Nanite::FRasterResults, TInlineAllocator<2>> NaniteRasterResults;
@@ -4497,7 +4505,10 @@ void FDeferredShadingSceneRenderer::Render(FRDGBuilder& GraphBuilder)
 	}
 
 #if WITH_MGPU
-	DoCrossGPUTransfers(GraphBuilder, ViewFamilyTexture);
+	if (ViewFamily.bMultiGPUForkAndJoin)
+	{
+		DoCrossGPUTransfers(GraphBuilder, ViewFamilyTexture, Views, CrossGPUTransferFencesDefer.Num() > 0, RenderTargetGPUMask);
+	}
 #endif
 
 	{
diff --git a/Engine/Source/Runtime/Renderer/Private/SceneRendering.cpp b/Engine/Source/Runtime/Renderer/Private/SceneRendering.cpp
--- a/Engine/Source/Runtime/Renderer/Private/SceneRendering.cpp
+++ b/Engine/Source/Runtime/Renderer/Private/SceneRendering.cpp
@@ -3188,22 +3188,20 @@ struct FCrossGPUTransfer
 	}
 };
 
-static void GetCrossGPUTransfers(FSceneRenderer* SceneRenderer, TArray<FCrossGPUTransfer>& OutTransfers)
+static void GetCrossGPUTransfers(FSceneRenderer* SceneRenderer, TArray<FCrossGPUTransfer>& OutTransfers, TArrayView<FViewInfo> InViews, const FIntPoint RenderTargetSize, FRHIGPUMask RenderTargetGPUMask)
 {
 	check(SceneRenderer->bGPUMasksComputed);
 
-	const FIntPoint Extent = SceneRenderer->ViewFamily.RenderTarget->GetSizeXY();
-
-	for (int32 ViewIndex = 0; ViewIndex < SceneRenderer->Views.Num(); ++ViewIndex)
+	for (int32 ViewIndex = 0; ViewIndex < InViews.Num(); ++ViewIndex)
 	{
-		const FViewInfo& ViewInfo = SceneRenderer->Views[ViewIndex];
-		if (ViewInfo.bAllowCrossGPUTransfer && ViewInfo.GPUMask != SceneRenderer->RenderTargetGPUMask)
+		const FViewInfo& ViewInfo = InViews[ViewIndex];
+		if (ViewInfo.bAllowCrossGPUTransfer && ViewInfo.GPUMask != RenderTargetGPUMask)
 		{
 			// Clamp the view rect by the rendertarget rect to prevent issues when resizing the viewport.
-			const FIntRect TransferRect(ViewInfo.UnscaledViewRect.Min.ComponentMin(Extent), ViewInfo.UnscaledViewRect.Max.ComponentMin(Extent));
+			const FIntRect TransferRect(ViewInfo.UnscaledViewRect.Min.ComponentMin(RenderTargetSize), ViewInfo.UnscaledViewRect.Max.ComponentMin(RenderTargetSize));
 			if (TransferRect.Width() > 0 && TransferRect.Height() > 0)
 			{
-				for (uint32 RenderTargetGPUIndex : SceneRenderer->RenderTargetGPUMask)
+				for (uint32 RenderTargetGPUIndex : RenderTargetGPUMask)
 				{
 					if (!ViewInfo.GPUMask.Contains(RenderTargetGPUIndex))
 					{
@@ -3241,7 +3239,7 @@ void FSceneRenderer::PreallocateCrossGPUFences(const TArray<FSceneRenderer*>& Sc
 				if (SceneRenderer->AllViewsGPUMask.HasSingleIndex())
 				{
 					TArray<FCrossGPUTransfer> Transfers;
-					GetCrossGPUTransfers(SceneRenderer, Transfers);
+					GetCrossGPUTransfers(SceneRenderer, Transfers, SceneRenderer->Views, SceneRenderer->ViewFamily.RenderTarget->GetSizeXY(), SceneRenderer->RenderTargetGPUMask);
 
 					SceneRenderer->CrossGPUTransferFencesDefer.SetNumUninitialized(Transfers.Num());
 
@@ -3259,59 +3257,56 @@ void FSceneRenderer::PreallocateCrossGPUFences(const TArray<FSceneRenderer*>& Sc
 #endif
 }
 
-void FSceneRenderer::DoCrossGPUTransfers(FRDGBuilder& GraphBuilder, FRDGTextureRef ViewFamilyTexture)
+void FSceneRenderer::DoCrossGPUTransfers(FRDGBuilder& GraphBuilder, FRDGTextureRef RenderTargetTexture, TArrayView<FViewInfo> InViews, bool bCrossGPUTransferFencesDefer, FRHIGPUMask InRenderTargetGPUMask)
 {
 #if WITH_MGPU
-	if (ViewFamily.bMultiGPUForkAndJoin)
-	{
-		// Must be all GPUs because context redirector only supports single or all GPUs
-		RDG_GPU_MASK_SCOPE(GraphBuilder, FRHIGPUMask::All());
-		RDG_GPU_STAT_SCOPE(GraphBuilder, CrossGPUTransfers);
+	// Must be all GPUs because context redirector only supports single or all GPUs
+	RDG_GPU_MASK_SCOPE(GraphBuilder, FRHIGPUMask::All());
+	RDG_GPU_STAT_SCOPE(GraphBuilder, CrossGPUTransfers);
 
-		// Need to use this structure as an intermediate, because the RHI texture reference isn't available yet,
-		// and must be fetched inside the pass.
-		TArray<FCrossGPUTransfer> Transfers;
-		GetCrossGPUTransfers(this, Transfers);
+	// Need to use this structure as an intermediate, because the RHI texture reference isn't available yet,
+	// and must be fetched inside the pass.
+	TArray<FCrossGPUTransfer> Transfers;
+	GetCrossGPUTransfers(this, Transfers, InViews, RenderTargetTexture->Desc.Extent, InRenderTargetGPUMask);
 
-		if (Transfers.Num() > 0)
+	if (Transfers.Num() > 0)
+	{
+		// Check if we can go through the optimized code path, with delay for the cross GPU transfer fence wait
+		if (bCrossGPUTransferFencesDefer)
 		{
-			// Check if we can go through the optimized code path, with delay for the cross GPU transfer fence wait
-			if (CrossGPUTransferFencesDefer.Num())
+			// A readback pass is the closest analog to what this is doing. There isn't a way to express cross-GPU transfers via the RHI barrier API.
+			AddReadbackTexturePass(GraphBuilder, RDG_EVENT_NAME("CrossGPUTransfers"), RenderTargetTexture,
+				[this, RenderTargetTexture, LocalTransfers = MoveTemp(Transfers), PostTransferFences = MoveTemp(CrossGPUTransferFencesDefer)](FRHICommandListImmediate& RHICmdList)
 			{
-				// A readback pass is the closest analog to what this is doing. There isn't a way to express cross-GPU transfers via the RHI barrier API.
-				AddReadbackTexturePass(GraphBuilder, RDG_EVENT_NAME("CrossGPUTransfers"), ViewFamilyTexture,
-					[this, ViewFamilyTexture, LocalTransfers = MoveTemp(Transfers), PostTransferFences = MoveTemp(CrossGPUTransferFencesDefer)](FRHICommandListImmediate& RHICmdList)
+				TArray<FTransferResourceParams> TransferParams;
+				for (const FCrossGPUTransfer& Transfer : LocalTransfers)
 				{
-					TArray<FTransferResourceParams> TransferParams;
-					for (const FCrossGPUTransfer& Transfer : LocalTransfers)
-					{
-						TransferParams.Add(FTransferResourceParams(ViewFamilyTexture->GetRHI(), Transfer.SrcGPUIndex, Transfer.DestGPUIndex, false, false));
-					}
+					TransferParams.Add(FTransferResourceParams(RenderTargetTexture->GetRHI(), Transfer.SrcGPUIndex, Transfer.DestGPUIndex, false, false));
+				}
 
-					// Transition resources on destination GPU and signal when transition has finished
-					TArray<FCrossGPUTransferFence*> PreTransferFences;
-					RHIGenerateCrossGPUPreTransferFences(TransferParams, PreTransferFences);
-					RHICmdList.CrossGPUTransferSignal(TransferParams, PreTransferFences);
+				// Transition resources on destination GPU and signal when transition has finished
+				TArray<FCrossGPUTransferFence*> PreTransferFences;
+				RHIGenerateCrossGPUPreTransferFences(TransferParams, PreTransferFences);
+				RHICmdList.CrossGPUTransferSignal(TransferParams, PreTransferFences);
 
-					// Then do the actual transfer
-					RHICmdList.CrossGPUTransfer(TransferParams, PreTransferFences, PostTransferFences);
-				});
-			}
-			else
+				// Then do the actual transfer
+				RHICmdList.CrossGPUTransfer(TransferParams, PreTransferFences, PostTransferFences);
+			});
+		}
+		else
+		{
+			// A readback pass is the closest analog to what this is doing. There isn't a way to express cross-GPU transfers via the RHI barrier API.
+			AddReadbackTexturePass(GraphBuilder, RDG_EVENT_NAME("CrossGPUTransfers"), RenderTargetTexture,
+				[this, RenderTargetTexture, LocalTransfers = MoveTemp(Transfers)](FRHICommandListImmediate& RHICmdList)
 			{
-				// A readback pass is the closest analog to what this is doing. There isn't a way to express cross-GPU transfers via the RHI barrier API.
-				AddReadbackTexturePass(GraphBuilder, RDG_EVENT_NAME("CrossGPUTransfers"), ViewFamilyTexture,
-					[this, ViewFamilyTexture, LocalTransfers = MoveTemp(Transfers)](FRHICommandListImmediate& RHICmdList)
+				TArray<FTransferResourceParams> TransferParams;
+				for (const FCrossGPUTransfer& Transfer : LocalTransfers)
 				{
-					TArray<FTransferResourceParams> TransferParams;
-					for (const FCrossGPUTransfer& Transfer : LocalTransfers)
-					{
-						TransferParams.Add(FTransferResourceParams(ViewFamilyTexture->GetRHI(), Transfer.SrcGPUIndex, Transfer.DestGPUIndex, false, false));
-					}
+					TransferParams.Add(FTransferResourceParams(RenderTargetTexture->GetRHI(), Transfer.SrcGPUIndex, Transfer.DestGPUIndex, false, false));
+				}
 
-					RHICmdList.TransferResources(TransferParams);
-				});
-			}
+				RHICmdList.TransferResources(TransferParams);
+			});
 		}
 	}
 #endif // WITH_MGPU
diff --git a/Engine/Source/Runtime/Renderer/Private/SceneRendering.h b/Engine/Source/Runtime/Renderer/Private/SceneRendering.h
--- a/Engine/Source/Runtime/Renderer/Private/SceneRendering.h
+++ b/Engine/Source/Runtime/Renderer/Private/SceneRendering.h
@@ -2121,7 +2121,7 @@ public:
 
 	/** Logic to update render targets across all GPUs */
 	static void PreallocateCrossGPUFences(const TArray<FSceneRenderer*>& SceneRenderers);
-	void DoCrossGPUTransfers(FRDGBuilder& GraphBuilder, FRDGTextureRef ViewFamilyTexture);
+	void DoCrossGPUTransfers(FRDGBuilder& GraphBuilder, FRDGTextureRef ViewFamilyTexture, TArrayView<FViewInfo> InViews, bool bCrossGPUTransferFencesDefer, FRHIGPUMask RenderTargetGPUMask);
 	void FlushCrossGPUFences(FRDGBuilder& GraphBuilder);
 
 	bool DoOcclusionQueries() const;
