Fixed translucency rendering when view is under water. In this case now, translucent meshes are simply all splatted on the water scene color without any separate mechanism.
We simply  respect the composition order of the different passes.

This CL thus render the translucent meshes at scene color resolution and ignores separate translucency render target which can be a different resolution.
If we notice that there is a performance issue in some case, we can then consider using the 27554020 shelve which respect the separate translucent render targets and downsampling, but does require composition of them on full resolution scene color.

Split screen tested.

diff --git a/Engine/Source/Runtime/Renderer/Private/TranslucentRendering.cpp b/Engine/Source/Runtime/Renderer/Private/TranslucentRendering.cpp
--- a/Engine/Source/Runtime/Renderer/Private/TranslucentRendering.cpp
+++ b/Engine/Source/Runtime/Renderer/Private/TranslucentRendering.cpp
@@ -1329,20 +1329,24 @@ void FDeferredShadingSceneRenderer::RenderTranslucencyInner(
 	ESceneTextureSetupMode SceneTextureSetupMode = ESceneTextureSetupMode::All;
 	EnumRemoveFlags(SceneTextureSetupMode, ESceneTextureSetupMode::SceneColor);
 
-	if (bRenderInSeparateTranslucency)
+	// Create resources shared by each view (each view data is tiled into each of the render target resources)
+	FRDGTextureMSAA SharedColorTexture = CreatePostDOFTranslucentTexture(GraphBuilder, TranslucencyPass, SeparateTranslucencyDimensions, bIsModulate, ShaderPlatform);
+
+	for (int32 ViewIndex = 0, NumProcessedViews = 0; ViewIndex < Views.Num(); ++ViewIndex)
 	{
-		// Create resources shared by each view (each view data is tiled into each of the render target resources)
-		FRDGTextureMSAA SharedColorTexture = CreatePostDOFTranslucentTexture(GraphBuilder, TranslucencyPass, SeparateTranslucencyDimensions, bIsModulate, ShaderPlatform);
+		FViewInfo& View = Views[ViewIndex];
+		const ETranslucencyView TranslucencyView = GetTranslucencyView(View);
 
-		for (int32 ViewIndex = 0, NumProcessedViews = 0; ViewIndex < Views.Num(); ++ViewIndex)
+		if (!EnumHasAnyFlags(TranslucencyView, ViewsToRender))
 		{
-			FViewInfo& View = Views[ViewIndex];
-			const ETranslucencyView TranslucencyView = GetTranslucencyView(View);
+			continue;
+		}
 
-			if (!EnumHasAnyFlags(TranslucencyView, ViewsToRender))
-			{
-				continue;
-			}
+		// We run separate and composited translucent only when the view is NOT under water.
+		// When under water, we render each translucency pass in forward on the water buffer itself.
+		const bool bViewIsUnderWater = EnumHasAnyFlags(TranslucencyView, ETranslucencyView::UnderWater);
+		if (bRenderInSeparateTranslucency && !bViewIsUnderWater)
+		{
 
 			RDG_GPU_MASK_SCOPE(GraphBuilder, View.GPUMask);
 			RDG_EVENT_SCOPE_CONDITIONAL(GraphBuilder, Views.Num() > 1, "View%d", ViewIndex);
@@ -1350,7 +1354,7 @@ void FDeferredShadingSceneRenderer::RenderTranslucencyInner(
 			FIntRect ScaledViewRect = GetScaledRect(View.ViewRect, SeparateTranslucencyDimensions.Scale);
 
 			const FScreenPassTextureViewport SeparateTranslucencyViewport = SeparateTranslucencyDimensions.GetInstancedStereoViewport(View);
-			const bool bCompositeBackToSceneColor = (IsMainTranslucencyPass(TranslucencyPass) && !bIsStandardSeparatedTranslucency) || EnumHasAnyFlags(TranslucencyView, ETranslucencyView::UnderWater);
+			const bool bCompositeBackToSceneColor = (IsMainTranslucencyPass(TranslucencyPass) && !bIsStandardSeparatedTranslucency);
 			const bool bLumenGIEnabled = GetViewPipelineState(View).DiffuseIndirectMethod == EDiffuseIndirectMethod::Lumen;
 
 			/** Separate translucency color is either composited immediately or later during post processing. If done immediately, it's because the view doesn't support
@@ -1453,22 +1457,17 @@ void FDeferredShadingSceneRenderer::RenderTranslucencyInner(
 
 			++NumProcessedViews;
 		}
-	}
-	else
-	{
-		for (int32 ViewIndex = 0; ViewIndex < Views.Num(); ++ViewIndex)
+		else
 		{
-			FViewInfo& View = Views[ViewIndex];
-			const ETranslucencyView TranslucencyView = GetTranslucencyView(View);
-
-			if (!EnumHasAnyFlags(TranslucencyView, ViewsToRender))
+			// When rendering translucent meshes under water, we skip modulate passes which are only required when compositing separate translucency passes from render target.
+			const bool bSkipPass = bViewIsUnderWater && bIsModulate; 
+			RDG_GPU_MASK_SCOPE(GraphBuilder, View.GPUMask);
+			RDG_EVENT_SCOPE_CONDITIONAL(GraphBuilder, Views.Num() > 1 && !bSkipPass, "View%d", ViewIndex);
+			if (bSkipPass)
 			{
-				continue;
+				return;
 			}
 
-			RDG_GPU_MASK_SCOPE(GraphBuilder, View.GPUMask);
-			RDG_EVENT_SCOPE_CONDITIONAL(GraphBuilder, Views.Num() > 1, "View%d", ViewIndex);
-
 			const ERenderTargetLoadAction SceneColorLoadAction = ERenderTargetLoadAction::ELoad;
 			const FScreenPassTextureViewport Viewport(SceneTextures.Color.Target, View.ViewRect);
 			const float ViewportScale = 1.0f;
