Add cvar to mark only even/odd mips from local lights.
This can be used to cut down on instance duplication, which can be a major performance cost especially for higher poly/larger non-nanite meshes.

r.Shadow.Virtual.MarkPixelPagesMipModeLocal
Even mips mode (1) cuts down on instances but can increase the resolution vs. the default mode for some pages. i.e. fewer instances/vertices, but more pixels. Generally a performance win for non-Nanite geometry, but a potential performance loss for Nanite geometry.
Odd mips mode (2) cuts down on instances *and* resolution so has a higher performance benefit, but can be a noticeable loss in resolution.

From Pari at The Coalition

diff --git a/Engine/Shaders/Private/VirtualShadowMaps/VirtualShadowMapPageMarking.usf b/Engine/Shaders/Private/VirtualShadowMaps/VirtualShadowMapPageMarking.usf
--- a/Engine/Shaders/Private/VirtualShadowMaps/VirtualShadowMapPageMarking.usf
+++ b/Engine/Shaders/Private/VirtualShadowMaps/VirtualShadowMapPageMarking.usf
@@ -34,6 +34,7 @@ float PageDilationBorderSizeLocal;
 uint InputType;
 uint bCullBackfacingPixels;
 uint NumDirectionalLightSmInds;
+uint MipModeLocal;
 
 int GetBiasedClipmapLevel(FVirtualShadowMapProjectionShaderData BaseProjectionData, FLWCVector3 WorldPosition)
 {
@@ -53,7 +54,20 @@ uint GetMipLevelLocal(int VirtualShadowMapId, float3 TranslatedWorldPosition, fl
 
 	float MipLevelFloat = log2(Footprint) + ProjectionData.ResolutionLodBias + VirtualShadowMap.GlobalResolutionLodBias;
 	uint MipLevel = uint(max(floor(MipLevelFloat), 0.0f));
-	return min(MipLevel, (VSM_MAX_MIP_LEVELS - 1U));
+	MipLevel = min(MipLevel, (VSM_MAX_MIP_LEVELS - 1U));
+
+	if (MipModeLocal == 1)
+	{
+		// Even mips
+		MipLevel &= ~(1U);
+	}
+	else if (MipModeLocal == 2)
+	{
+		// Odd mips
+		MipLevel |= 1U;
+	}
+
+	return MipLevel;
 }
 
 void MarkPageAddress(uint PageOffset, uint Flags)
diff --git a/Engine/Source/Runtime/Renderer/Private/VirtualShadowMaps/VirtualShadowMapArray.cpp b/Engine/Source/Runtime/Renderer/Private/VirtualShadowMaps/VirtualShadowMapArray.cpp
--- a/Engine/Source/Runtime/Renderer/Private/VirtualShadowMaps/VirtualShadowMapArray.cpp
+++ b/Engine/Source/Runtime/Renderer/Private/VirtualShadowMaps/VirtualShadowMapArray.cpp
@@ -105,6 +105,16 @@ TAutoConsoleVariable<int32> CVarMarkPixelPages(
 	ECVF_RenderThreadSafe
 );
 
+static TAutoConsoleVariable<int32> CVarMarkPixelPagesMipModeLocal(
+	TEXT("r.Shadow.Virtual.MarkPixelPagesMipModeLocal"),
+	0,
+	TEXT("When enabled, this uses a subset of mips to reduce instance duplication in VSMs. Will result in better performance but a harsher falloff on mip transitions.\n")
+	TEXT(" 0 - Disabled: Use all 8 mips\n")
+	TEXT(" 1 - Quality Mode: Use 4 higher res mips (16k, 4k, 1k, 256)\n")
+	TEXT(" 2 - Performance Mode: Use 4 lower res mips (8k, 2k, 512, 128)\n"),
+	ECVF_RenderThreadSafe | ECVF_Scalability
+);
+
 TAutoConsoleVariable<int32> CVarMarkCoarsePagesDirectional(
 	TEXT("r.Shadow.Virtual.MarkCoarsePagesDirectional"),
 	1,
@@ -693,12 +703,12 @@ class FGeneratePageFlagsFromPixelsCS : public FVirtualPageManagementShader
 		SHADER_PARAMETER_RDG_TEXTURE(Texture2D<float>, FrontLayerTranslucencyNormalTexture)
 		SHADER_PARAMETER_RDG_UNIFORM_BUFFER(FSubstrateGlobalUniformParameters, Substrate)
 		RDG_BUFFER_ACCESS(IndirectBufferArgs, ERHIAccess::IndirectArgs)
-		SHADER_PARAMETER(uint32, InputType)
-		SHADER_PARAMETER(uint32, NumDirectionalLightSmInds)
-		SHADER_PARAMETER(uint32, bPostBasePass)
 		SHADER_PARAMETER(float, PageDilationBorderSizeDirectional)
 		SHADER_PARAMETER(float, PageDilationBorderSizeLocal)
+		SHADER_PARAMETER(uint32, InputType)
 		SHADER_PARAMETER(uint32, bCullBackfacingPixels)
+		SHADER_PARAMETER(uint32, NumDirectionalLightSmInds)
+		SHADER_PARAMETER(uint32, MipModeLocal)
 		SHADER_PARAMETER(FIntPoint, PixelStride)
 	END_SHADER_PARAMETER_STRUCT()
 };
@@ -1521,6 +1531,7 @@ void FVirtualShadowMapArray::BuildPageAllocations(
 					PassParameters->bCullBackfacingPixels = ShouldCullBackfacingPixels() ? 1 : 0;
 					PassParameters->Substrate = Substrate::BindSubstrateGlobalUniformParameters(View);
 					PassParameters->PixelStride = PixelStride;
+					PassParameters->MipModeLocal = CVarMarkPixelPagesMipModeLocal.GetValueOnRenderThread();
 					
 					const FIntPoint StridedPixelSize = FIntPoint::DivideAndRoundUp(View.ViewRect.Size(), PixelStride);
 					// Note: we use the tile size defined by the water as the group-size - this is needed because the tile mask testing code relies on the size being the same to scalarize efficiently.
