PR #10819: Fix FRelevancePacket::NumVisibleDynamicEditorPrimitives statistics.

#10819

diff --git a/Engine/Source/Runtime/Renderer/Private/SceneVisibility.cpp b/Engine/Source/Runtime/Renderer/Private/SceneVisibility.cpp
--- a/Engine/Source/Runtime/Renderer/Private/SceneVisibility.cpp
+++ b/Engine/Source/Runtime/Renderer/Private/SceneVisibility.cpp
@@ -1251,7 +1251,7 @@ void FRelevancePacket::ComputeRelevance(FDynamicPrimitiveIndexList& DynamicPrimi
 #if WITH_EDITOR
 		if (GIsEditor)
 		{
-			++NumVisibleDynamicPrimitives;
+			++NumVisibleDynamicEditorPrimitives;
 
 			if (DynamicPrimitiveViewMasks)
 			{
