Fix for DebugDraws being misaligned while retaining previous fixes for stereo views.

diff --git a/Engine/Source/Runtime/Renderer/Private/PostProcess/PostProcessCompositePrimitivesCommon.cpp b/Engine/Source/Runtime/Renderer/Private/PostProcess/PostProcessCompositePrimitivesCommon.cpp
--- a/Engine/Source/Runtime/Renderer/Private/PostProcess/PostProcessCompositePrimitivesCommon.cpp
+++ b/Engine/Source/Runtime/Renderer/Private/PostProcess/PostProcessCompositePrimitivesCommon.cpp
@@ -251,12 +251,19 @@ void PopulateDepthPass(FRDGBuilder& GraphBuilder,
 
 	FPopulateCompositeDepthPS::FParameters* PassParameters = GraphBuilder.AllocParameters<FPopulateCompositeDepthPS::FParameters>();
 	PassParameters->View = View.ViewUniformBuffer;
-	PassParameters->Color = GetScreenPassTextureViewportParameters(FScreenPassTextureViewport(InSceneColor));
-	PassParameters->Depth = GetScreenPassTextureViewportParameters(FScreenPassTextureViewport(InSceneDepth));
+
 	if (bForceDrawColor)
 	{
 		PassParameters->ColorTexture = InSceneColor.Texture;
 		PassParameters->ColorSampler = PointClampSampler;
+
+		PassParameters->Color = GetScreenPassTextureViewportParameters(View.StereoPass == EStereoscopicPass::eSSP_FULL ? FScreenPassTextureViewport(InSceneColor.Texture->Desc.Extent, View.ViewRect) : FScreenPassTextureViewport(InSceneColor));
+		PassParameters->Depth = GetScreenPassTextureViewportParameters(FScreenPassTextureViewport(InSceneDepth.Texture->Desc.Extent, View.ViewRect));
+	}
+	else
+	{
+		PassParameters->Color = GetScreenPassTextureViewportParameters(FScreenPassTextureViewport(InSceneColor));
+		PassParameters->Depth = GetScreenPassTextureViewportParameters(FScreenPassTextureViewport(InSceneDepth));
 	}
 	PassParameters->DepthTexture = InSceneDepth.Texture;
 	PassParameters->DepthSampler = PointClampSampler;
