Fix FSceneRenderer::PostVisibilityFrameSetup early out mistake

diff --git a/Engine/Source/Runtime/Renderer/Private/SceneVisibility.cpp b/Engine/Source/Runtime/Renderer/Private/SceneVisibility.cpp
--- a/Engine/Source/Runtime/Renderer/Private/SceneVisibility.cpp
+++ b/Engine/Source/Runtime/Renderer/Private/SceneVisibility.cpp
@@ -4790,7 +4790,7 @@ void FSceneRenderer::GatherReflectionCaptureLightMeshElements()
 
 void FSceneRenderer::PostVisibilityFrameSetup(FILCUpdatePrimTaskData*& OutILCTaskData)
 {
-	if (GetRendererOutput() == ERendererOutput::FinalSceneColor)
+	if (GetRendererOutput() != ERendererOutput::FinalSceneColor)
 	{
 		return;
 	}
