Explicit cast to avoid asserts on overflow

diff --git a/Engine/Source/Runtime/Renderer/Private/VirtualShadowMaps/VirtualShadowMapClipmap.cpp b/Engine/Source/Runtime/Renderer/Private/VirtualShadowMaps/VirtualShadowMapClipmap.cpp
--- a/Engine/Source/Runtime/Renderer/Private/VirtualShadowMaps/VirtualShadowMapClipmap.cpp
+++ b/Engine/Source/Runtime/Renderer/Private/VirtualShadowMaps/VirtualShadowMapClipmap.cpp
@@ -197,7 +197,9 @@ FVirtualShadowMapClipmap::FVirtualShadowMapClipmap(
 		const FInt64Point SnappedPageOriginLi(-SnappedViewCenter.X, SnappedViewCenter.Y);
 		const FInt64Point SnappedPageOriginLn(-SnappedOriginLn.X, SnappedOriginLn.Y);
 		const FInt64Point RelativeCornerOffset = SnappedPageOriginLi - SnappedPageOriginLn + ((RadiiPerLevel / 2) * (int64_t)SnapSize);
-		Level.RelativeCornerOffset = FIntPoint(RelativeCornerOffset / SnapSize);
+		Level.RelativeCornerOffset = FIntPoint(
+			static_cast<int32>(RelativeCornerOffset.X / SnapSize),
+			static_cast<int32>(RelativeCornerOffset.Y / SnapSize));
 
 		// Check if we have a cache entry for this level
 		// If we do and it covers our required depth range, we can use cached pages. Otherwise we need to invalidate.
