Fixed separate translucency not working with scale different from 100 or 50 percent.

diff --git a/Engine/Shaders/Private/DownsampleDepthPixelShader.usf b/Engine/Shaders/Private/DownsampleDepthPixelShader.usf
--- a/Engine/Shaders/Private/DownsampleDepthPixelShader.usf
+++ b/Engine/Shaders/Private/DownsampleDepthPixelShader.usf
@@ -16,7 +16,7 @@
 
 Texture2D<uint> NaniteShadingMask;
 Texture2D<float> DepthTexture;
-float2 DestinationTexelSize;
+float2 DstToSrcPixelScale;
 float2 SourceMaxUV;
 float2 DestinationResolution;
 uint DownsampleDepthFilter;
@@ -43,7 +43,7 @@ void Main(
 	// Pixel coordinate for the output 
 	const int2 DstPixelPos = SVPos.xy;
 	// Lower left corner from the input 2x2 quad
-	const int2 PixelPos = DstPixelPos * 2;
+	const int2 PixelPos = SVPos.xy * DstToSrcPixelScale;
 
 	float DeviceZ0 = DepthTexture.Load(uint3(clamp(PixelPos + int2(0, 0), SrcPixelCoordMinAndMax.xy, SrcPixelCoordMinAndMax.zw), 0));
 	float DeviceZ1 = DepthTexture.Load(uint3(clamp(PixelPos + int2(0, 1), SrcPixelCoordMinAndMax.xy, SrcPixelCoordMinAndMax.zw), 0));
diff --git a/Engine/Source/Runtime/Renderer/Private/ScreenPass.cpp b/Engine/Source/Runtime/Renderer/Private/ScreenPass.cpp
--- a/Engine/Source/Runtime/Renderer/Private/ScreenPass.cpp
+++ b/Engine/Source/Runtime/Renderer/Private/ScreenPass.cpp
@@ -299,7 +299,7 @@ public:
 	BEGIN_SHADER_PARAMETER_STRUCT(FParameters, )
 		SHADER_PARAMETER_STRUCT_REF(FViewUniformShaderParameters, View)
 		SHADER_PARAMETER_RDG_TEXTURE(Texture2D, DepthTexture)
-		SHADER_PARAMETER(FVector2f, DestinationTexelSize)
+		SHADER_PARAMETER(FVector2f, DstToSrcPixelScale)
 		SHADER_PARAMETER(FVector2f, SourceMaxUV)
 		SHADER_PARAMETER(FVector2f, DestinationResolution)
 		SHADER_PARAMETER(uint32, DownsampleDepthFilter)
@@ -339,7 +339,7 @@ void AddDownsampleDepthPass(
 	FDownsampleDepthPS::FParameters* PassParameters = GraphBuilder.AllocParameters<FDownsampleDepthPS::FParameters>();
 	PassParameters->View = View.ViewUniformBuffer;
 	PassParameters->DepthTexture = Input.Texture;
-	PassParameters->DestinationTexelSize = FVector2f(1.0f / OutputViewport.Extent.X, 1.0f / OutputViewport.Extent.X);
+	PassParameters->DstToSrcPixelScale = FVector2f(float(InputViewport.Extent.X) / float(OutputViewport.Extent.X), float(InputViewport.Extent.Y) / float(OutputViewport.Extent.Y));
 	PassParameters->SourceMaxUV = FVector2f((float(View.ViewRect.Max.X) -1.0f - 0.51f) / InputViewport.Extent.X, (float(View.ViewRect.Max.Y) - 1.0f - 0.51f) / InputViewport.Extent.Y);
 	PassParameters->DownsampleDepthFilter = (uint32)DownsampleDepthFilter;
 
