Path Tracing: Implement Adaptive sampling

Implement a basic adaptive sampling strategy. This is not yet exposed in the UI, it must be enabled by the cvar: "r.PathTracing.AdaptiveSampling".

When enabled, the path tracer keeps track of a variance buffer. This variance buffer can then be used to predict the number of samples requires to reach a prescribed error threshold. This is used to de-active pixels that are deemed to be converged already, speeding up rendering of mostly converged areas.

The error estimate operates on multiple scales (mip maps) so that we have a more robust estimate of the variance even with very low sample counts. Currently the error estimate runs after each pass, and is used by the next sample to decide which pixels are still active.

A proper integration with MRQ (in particular the reference motion blur mode) is not yet done, as this will require a different way of stepping through temporal samples.This will be investigated as part of the MRQ 2.0 interface.

This feature (along with a few other path tracer permutations) is now gated behind r.PathTracing.Experimental (off by default) to avoid bloating startup times.

The adaptive metric is relative to the current exposure level. This means that the combination of adaptive sampling and auto-exposure has a feedback loop which could lead to unpredictable results. Therefore manual exposure is recommended when using adaptive sampling. In particular, when using the adaptive sampling visualization modes with auto-exposure, a feedback effect can be observed as the visualization influences the exposure level. A proper fix for this would be to move the visualization of adaptive sampling after tonemapping. This is left as future work for now to minimize the "spread" of this feature for now.

diff --git a/Engine/Shaders/Private/PathTracing/PathTracing.usf b/Engine/Shaders/Private/PathTracing/PathTracing.usf
--- a/Engine/Shaders/Private/PathTracing/PathTracing.usf
+++ b/Engine/Shaders/Private/PathTracing/PathTracing.usf
@@ -6,8 +6,7 @@ PathTracing.usf: Reference path tracing
 
 #define PATH_TRACING  1
 
-// Bump to force-compile this shader
-#pragma message("UESHADERMETADATA_VERSION 8535C2EF-E5D9-4ABB-A0AF-F8FA45BBB08B")
+#define PATH_TRACING_SHADER_VERSION 0x51478b68 // Bump to force-compile this shader
 
 // Ensure specular profile and glints are supported
 #if SUBSTRATE_ENABLED && PATH_TRACER_USE_SUBSTRATE_SPECIAL_COMPLEX_MATERIAL
@@ -104,7 +103,7 @@ void StorePathStateData(FPathState PathState, uint Index)
 
 RAY_TRACING_ENTRY_RAYGEN(PathTracingMainRG)
 {
-	const uint2 DispatchIdx = DispatchRaysIndex().xy;
+	uint2 DispatchIdx = DispatchRaysIndex().xy;
 	const uint2 DispatchDim = DispatchRaysDimensions().xy;
 
 	int LinearPixelIndex = DispatchIdx.x + DispatchDim.x * DispatchIdx.y;
@@ -117,10 +116,19 @@ RAY_TRACING_ENTRY_RAYGEN(PathTracingMainRG)
 	FPathState PathState;
 	if (Bounce == 0)
 	{
+#if PATH_TRACER_USE_ADAPTIVE_SAMPLING
+		if (Iteration > 0)
+		{
+			LinearPixelIndex = ActivePaths[LinearPixelIndex];
+			if (LinearPixelIndex < 0)
+			{
+				return;
+			}
+			DispatchIdx = uint2(LinearPixelIndex % DispatchDim.x, LinearPixelIndex / DispatchDim.x);
+		}
+#endif
 		PathState = CreatePathState(ComputePixelIndex(DispatchIdx), 
-									ComputeTextureIndex(uint2(LinearPixelIndex % DispatchDim.x,
-															  LinearPixelIndex / DispatchDim.x))
-								   );
+									ComputeTextureIndex(DispatchIdx));
 	}
 	else
 	{
@@ -152,13 +160,28 @@ RAY_TRACING_ENTRY_RAYGEN(PathTracingMainRG)
 
 #else // PATH_TRACER_USE_COMPACTION == 0
 
+#if PATH_TRACER_USE_ADAPTIVE_SAMPLING
+Buffer<int> ActivePaths;
+#endif
+
 RAY_TRACING_ENTRY_RAYGEN(PathTracingMainRG)
 {
-	const uint2 DispatchIdx = DispatchRaysIndex().xy;
+	uint2 DispatchIdx = DispatchRaysIndex().xy;
 	const uint2 DispatchDim = DispatchRaysDimensions().xy;
-	int2 TextureIndex = ComputeTextureIndex(DispatchIdx);
 
-	FPathState PathState = CreatePathState(ComputePixelIndex(DispatchIdx), TextureIndex);
+#if PATH_TRACER_USE_ADAPTIVE_SAMPLING
+	if (Iteration > 0)
+	{
+		uint LinearPixelIndex = DispatchIdx.x + DispatchDim.x * DispatchIdx.y;
+		LinearPixelIndex = ActivePaths[LinearPixelIndex];
+		if (LinearPixelIndex < 0)
+		{
+			return;
+		}
+		DispatchIdx = uint2(LinearPixelIndex % DispatchDim.x, LinearPixelIndex / DispatchDim.x);
+	}
+#endif
+	FPathState PathState = CreatePathState(ComputePixelIndex(DispatchIdx), ComputeTextureIndex(DispatchIdx));
 
 	for (int Bounce = 0; Bounce <= MaxBounces; Bounce++)
 	{
diff --git a/Engine/Shaders/Private/PathTracing/PathTracingAdaptiveStart.usf b/Engine/Shaders/Private/PathTracing/PathTracingAdaptiveStart.usf
new file mode 100644
--- /dev/null
+++ b/Engine/Shaders/Private/PathTracing/PathTracingAdaptiveStart.usf
@@ -0,0 +1,39 @@
+// Copyright Epic Games, Inc. All Rights Reserved.
+
+#include "../Common.ush"
+#include "PathTracingAdaptiveUtils.ush"
+
+
+RWBuffer<int> NextActivePaths;
+RWBuffer<uint> NumPathStates;
+int2 TileTextureOffset;
+uint2 DispatchDim;
+float ViewPreExposure;
+
+
+[numthreads(THREADGROUPSIZE_X, THREADGROUPSIZE_Y, 1)]
+void PathTracingAdaptiveStartCS(uint2 DispatchThreadId : SV_DispatchThreadID)
+{
+	uint2 DispatchIdx = DispatchThreadId;
+	if (any(DispatchIdx >= DispatchDim))
+	{
+		return;
+	}
+
+	int LinearPixelIndex = DispatchIdx.x + DispatchDim.x * DispatchIdx.y;
+	if (LinearPixelIndex == 0)
+	{
+		// write other dimensions for indirect dispatch
+		NumPathStates[1] = 1;
+		NumPathStates[2] = 1;
+	}
+
+	float SampleScaleFactor = GetAdaptiveSampleFactor(DispatchIdx + TileTextureOffset, ViewPreExposure);
+	if (SampleScaleFactor > 1.0)
+	{
+		// This pixel needs more samples, allocate a path state for it
+		uint PathStateIndex;
+		InterlockedAdd(NumPathStates[0], 1, PathStateIndex);
+		NextActivePaths[PathStateIndex] = LinearPixelIndex;
+	}
+}
diff --git a/Engine/Shaders/Private/PathTracing/PathTracingAdaptiveUtils.ush b/Engine/Shaders/Private/PathTracing/PathTracingAdaptiveUtils.ush
new file mode 100644
--- /dev/null
+++ b/Engine/Shaders/Private/PathTracing/PathTracingAdaptiveUtils.ush
@@ -0,0 +1,95 @@
+// Copyright Epic Games, Inc. All Rights Reserved.
+
+#pragma once
+
+Texture2D<float4>   VarianceTexture;
+SamplerState 		VarianceSampler;
+uint3				VarianceTextureDims;
+float 				AdaptiveSamplingErrorThreshold;
+
+
+// Variance of the tonemapped color can be predicted from the variance of the scene linear color:
+// https://en.wikipedia.org/wiki/Taylor_expansions_for_the_moments_of_functions_of_random_variables
+float PerceptualError(float Mean, float StdErr, float Exposure)
+{
+	// The UE tone mapping curve is quite a bit more involved, so getting its derivative would be messy and likely overkill
+	// instead just use a roughly S-curved shape curve that roughly does something filmic without costing too much
+	float x = Mean * Exposure;
+#if 0
+	// https://knarkowicz.wordpress.com/2016/01/06/aces-filmic-tone-mapping-curve/
+	// saturate((x*(a*x+b))/(x*(c*x+d)+e))
+	float a = 2.51f;
+	float b = 0.03f;
+	float c = 2.43f;
+	float d = 0.59f;
+	float e = 0.14f;
+	return StdErr * Exposure * (a * x * (2 * e + d * x) + b * (e - c * x * x)) / Square(x * (c * x + d) + e);
+#elif 0
+	// Cheap fit to curve above
+	// tonemap(x) = (6*x/(6*x+1))^2
+	return StdErr * Exposure * 72.0 * x / Pow3(6.0 * x + 1.0);
+#elif 1
+	// Similar curve, but with faster highlight rolloff
+	// tonemap(x) = (1 - exp(-4*x))^2
+	float E = exp(-4.0 * x);
+	return StdErr * Exposure * 8.0 * (E - E * E);
+#elif 0
+	return StdErr * Exposure * Square(rcp(x + 1));
+#elif 0
+	return StdErr * Exposure * 0.5 * rsqrt(x + 0.01);
+#elif 0
+	return StdErr * Exposure * 4.0 * exp(-4.0 * x);
+#endif
+}
+
+float4 SampleVarianceTexture(float2 VarianceUV, int Mip)
+{
+#if 0
+	// cheap, but has artifacts 
+	return VarianceTexture.SampleLevel(VarianceSampler, VarianceUV, Mip);
+#else
+	// Efficient GPU-Based Texture Interpolation using Uniform B-Splines
+	// Daniel Ruijters, Bart M. ter Haar Romeny, Paul Suetens
+	// Journal of Graphics GPU and Game Tools · January 2008
+	const float2 TextureSize = float2(VarianceTextureDims.xy >> Mip);
+	const float2 UV = VarianceUV * TextureSize - 0.5;
+	const float2 BaseUV = floor(UV);
+	const float2 F = UV - BaseUV;
+	const float2 O = 1.0 - F; 
+	const float2 F2 = F * F;
+	const float2 O2 = O * O;
+	const float2 W0 = (1.0f / 6.0f) * O2 * O;;
+	const float2 W1 = (2.0f / 3.0f) - F2 * (1.0f - 0.5f * F);
+	const float2 W2 = (2.0f / 3.0f) - O2 * (1.0f - 0.5f * O);
+	const float2 W3 = (1.0f / 6.0f) * F2 * F;
+	const float2 G0 = W0 + W1;
+	const float2 G1 = W2 + W3;
+	const float2 H0 = (W1 / G0 + BaseUV - 0.5) / TextureSize;
+	const float2 H1 = (W3 / G1 + BaseUV + 1.5) / TextureSize;
+	const float4 Tex00 = VarianceTexture.SampleLevel(VarianceSampler, float2(H0.x, H0.y), Mip);
+	const float4 Tex10 = VarianceTexture.SampleLevel(VarianceSampler, float2(H1.x, H0.y), Mip);
+	const float4 Tex01 = VarianceTexture.SampleLevel(VarianceSampler, float2(H0.x, H1.y), Mip);
+	const float4 Tex11 = VarianceTexture.SampleLevel(VarianceSampler, float2(H1.x, H1.y), Mip);
+	return lerp(lerp(Tex11, Tex10, G0.y),
+				lerp(Tex01, Tex00, G0.y),
+				G0.x);
+#endif
+}
+
+float GetAdaptiveSampleFactor(uint2 TexCoord, float Exposure)
+{
+	float2 VarianceUV = float2(TexCoord + 0.5) / float2(VarianceTextureDims.xy);
+	float SampleScaleFactor = 0.0;
+	for (int Mip = 1, NumMips = VarianceTextureDims.z; Mip < NumMips; Mip++)
+	{
+		float4 Variance = SampleVarianceTexture(VarianceUV, Mip);
+		float NumSamples = Variance.z;
+		float OutVariance = Variance.y - Variance.x * Variance.x;
+		float Lum = Variance.x;
+		float StdErr = sqrt(OutVariance / NumSamples);
+		float RelErr = PerceptualError(Lum, StdErr, Exposure);
+		// project the multiple of how many more samples we will need
+		SampleScaleFactor = max(SampleScaleFactor, (1 + Square(RelErr / AdaptiveSamplingErrorThreshold)) / NumSamples);
+	}
+	return SampleScaleFactor;
+}
\ No newline at end of file
diff --git a/Engine/Shaders/Private/PathTracing/PathTracingBuildAdaptiveError.usf b/Engine/Shaders/Private/PathTracing/PathTracingBuildAdaptiveError.usf
new file mode 100644
--- /dev/null
+++ b/Engine/Shaders/Private/PathTracing/PathTracingBuildAdaptiveError.usf
@@ -0,0 +1,42 @@
+// Copyright Epic Games, Inc. All Rights Reserved.
+
+#include "../Common.ush"
+
+int2 InputResolution;
+int2 OutputResolution;
+SamplerState InputMipSampler;
+Texture2D<float4>   InputMip;
+RWTexture2D<float4> OutputMip;
+
+[numthreads(THREADGROUPSIZE_X, THREADGROUPSIZE_Y, 1)]
+void PathTracingBuildAdaptiveErrorTextureCS(uint2 DispatchThreadId : SV_DispatchThreadID)
+{
+	int2 OutPixelCoord = int2(DispatchThreadId);
+	if (any(OutPixelCoord >= OutputResolution))
+	{
+		return;
+	}
+
+	float2 UV = (OutPixelCoord + 0.5) / float2(OutputResolution);
+	float2 InvInputResolution = rcp(float2(InputResolution));
+
+#if 0
+	// single tap mipmap
+	float4 Result = InputMip.SampleLevel(InputMipSampler, UV, 0);
+#elif 1
+	float4 Result = 0.0;
+	float FilterSum = 0.0;
+	const int R = 2;
+	for (int dy = -R; dy <= R; dy++)
+	for (int dx = -R; dx <= R; dx++)
+	{
+		float2 Offset = float2(dx, dy);
+		float Filter = exp(-0.5 * length2(Offset));
+		float4 Pixel = InputMip.SampleLevel(InputMipSampler, UV + Offset * InvInputResolution, 0);
+		FilterSum += Filter;
+		Result = lerp(Result, Pixel, Filter / FilterSum);
+	}
+#endif
+	//Result.z *= 4.0; // we are downsampling by 2x in each direction, so the overal normalization in the number of samples should be 4x
+	OutputMip[OutPixelCoord] = Result;
+}
\ No newline at end of file
diff --git a/Engine/Shaders/Private/PathTracing/PathTracingCompositingPixelShader.usf b/Engine/Shaders/Private/PathTracing/PathTracingCompositingPixelShader.usf
--- a/Engine/Shaders/Private/PathTracing/PathTracingCompositingPixelShader.usf
+++ b/Engine/Shaders/Private/PathTracing/PathTracingCompositingPixelShader.usf
@@ -5,6 +5,8 @@ ReferencePathTracinPixelShader.ush: Reference path tracing  compositing pixel sh
 ===============================================================================================*/
 
 #include "../Common.ush"
+#include "../ColorMap.ush"
+#include "PathTracingAdaptiveUtils.ush"
 
 uint Iteration;
 uint MaxSamples;
@@ -12,6 +14,8 @@ int ProgressDisplayEnabled;
 Texture2D<float4> RadianceTexture;
 Texture2D<float4> NormalDepthTexture;
 
+int AdaptiveSamplingVisualize;
+
 void CompositeMain(
 	in noperspective float2 UV : TEXCOORD0,
 	out float4 OutColor : SV_Target0,
@@ -21,10 +25,44 @@ void CompositeMain(
 	float2 BufferSize = View.BufferSizeAndInvSize.xy;
 	int3 TexCoord = int3(UV * BufferSize - View.ViewRectMin.xy, 0);
 
-	float4 Radiance = RadianceTexture.Load(TexCoord);
+	float4 Radiance = RadianceTexture.Load(TexCoord) * View.PreExposure;
 
-	OutColor = float4(Radiance.rgb * View.PreExposure, Radiance.a);
+	OutColor = float4(Radiance.rgb, Radiance.a);
+		
+	if (AdaptiveSamplingVisualize == 1)
+	{
+		// Visualize which pixels are still being worked and color them according to how many more samples we think they need
+		float2 VarianceUV = float2(TexCoord.xy + 0.5) * View.BufferSizeAndInvSize.zw;
+		float SampleScaleFactor = GetAdaptiveSampleFactor(TexCoord.xy, View.PreExposure);
+		if (SampleScaleFactor > 1.0)
+		{
+			float4 Variance = VarianceTexture.Load(TexCoord);
+			float Ramp = Variance.z * SampleScaleFactor / MaxSamples;
+			OutColor.xyz = ColorMapTurbo(Ramp);
+		}
+	}
+	if (AdaptiveSamplingVisualize == 2)
+	{
+		// Visualize sample count relative to maximum
+		float4 Variance = VarianceTexture.Load(TexCoord);
+		float Ramp = Variance.z / min(Iteration + 1, MaxSamples);
+		OutColor.xyz = ColorMapTurbo(Ramp);
+	}
 
+	if (AdaptiveSamplingVisualize >= 3 && AdaptiveSamplingVisualize < 3 + int(VarianceTextureDims.z))
+	{
+		//float4 Variance = VarianceTexture.Load(TexCoord);
+		float2 VarianceUV = float2(TexCoord.xy + 0.5) / float2(VarianceTextureDims.xy);
+		int Mip = AdaptiveSamplingVisualize - 3;
+		float4 Variance = SampleVarianceTexture(VarianceUV, Mip);
+		float NumSamples = Variance.z;
+		float OutVariance = Variance.y - Variance.x * Variance.x;
+		float Lum = Variance.x;
+		float StdErr = sqrt(OutVariance / NumSamples);
+		float RelErr = PerceptualError(Lum, StdErr, View.PreExposure);
+		OutColor.xyz = RelErr;
+	}
+	
 	if (ProgressDisplayEnabled && (Iteration + 1) < MaxSamples)
 	{
 		// Draw a progress meter for how far along we are to the max number of passes
@@ -74,13 +112,5 @@ void CompositeMain(
 #endif
 	}
 
-#if 0
-    // split-screen debugging aids
-	if (UV.x < 0.5)
-	{
-		OutColor = 1 - Radiance.a; // visualize alpha channel
-	}
-#endif
-	
 	OutDepth = NormalDepthTexture.Load(TexCoord).w;
 }
diff --git a/Engine/Shaders/Private/PathTracing/PathTracingCore.ush b/Engine/Shaders/Private/PathTracing/PathTracingCore.ush
--- a/Engine/Shaders/Private/PathTracing/PathTracingCore.ush
+++ b/Engine/Shaders/Private/PathTracing/PathTracingCore.ush
@@ -66,6 +66,9 @@ float CameraFocusDistance;
 float2 CameraLensRadius;
 
 RWTexture2D<float4> RadianceTexture;
+#if PATH_TRACER_USE_ADAPTIVE_SAMPLING
+RWTexture2D<float4> VarianceTexture;
+#endif
 RWTexture2D<float4> AlbedoTexture;
 RWTexture2D<float4> NormalTexture;
 RaytracingAccelerationStructure TLAS;
@@ -170,15 +173,27 @@ struct FPathState
 
 	void WritePixel(uint2 TextureIndex)
 	{
+#if PATH_TRACER_USE_ADAPTIVE_SAMPLING
+		float4 OldVariance = Iteration > 0 ? VarianceTexture[TextureIndex] : 0;
+		const float Blend = 1.0 / (OldVariance.z + 1);
+		const float PixelLum = Luminance(Radiance);
+		const float2 VarianceValue = float2(PixelLum, PixelLum * PixelLum);
+		VarianceTexture[TextureIndex] = float4(lerp(OldVariance.xy, VarianceValue, Blend), OldVariance.z + 1.0, 0.0);
+#else
+		// non-adaptive sampling, blend factor is constant for all pixels
+		const float Blend = BlendFactor;
+#endif
+
 		// Avoid reading the old pixel on the first sample on the off-chance there is a NaN/Inf pixel ...
 		float4 OldPixel  = Iteration > 0 ? RadianceTexture[TextureIndex] : 0;
 		float4 OldAlbedo = Iteration > 0 ? AlbedoTexture[TextureIndex]   : 0;
 		float4 OldNormal = Iteration > 0 ? NormalTexture[TextureIndex]   : 0;
 
 		const float4 PixelValue = float4(Radiance, saturate(BackgroundVisibility));
-		RadianceTexture[TextureIndex] = lerp(OldPixel, PixelValue, BlendFactor);
-		AlbedoTexture[TextureIndex] = lerp(OldAlbedo, float4(Albedo, 0), BlendFactor);
-		NormalTexture[TextureIndex].xyz = lerp(OldNormal.xyz,Normal.xyz, BlendFactor);
+
+		RadianceTexture[TextureIndex] = lerp(OldPixel, PixelValue, Blend);
+		AlbedoTexture[TextureIndex]   = lerp(OldAlbedo, float4(Albedo, 0), Blend);
+		NormalTexture[TextureIndex].xyz = lerp(OldNormal.xyz,Normal.xyz, Blend);
 	}
 
 	void WritePixel()
@@ -551,6 +566,7 @@ FPathTracingPayload TraceTransparentRay(inout FPathState PathState, int Bounce,
 				float StartVolumeTMin = Interval.VolumeTMin;
 				// Ray marching loop
 				float3 PathThroughput = PathState.PathThroughput;
+				float3 VolumeRadiance = 0;
 				// Limit number of steps to prevent timeouts // FIXME: This biases the result!
 				for (int Step = 0; Step < MaxRaymarchSteps; Step++)
 				{
@@ -590,7 +606,7 @@ FPathTracingPayload TraceTransparentRay(inout FPathState PathState, int Bounce,
 							}
 						}
 
-						PathState.Radiance += Result.Emission * PathThroughput;
+						VolumeRadiance += Result.Emission * PathThroughput;
 						// keep tracing through the volume
 						PathThroughput *= SigmaN;
 					}
@@ -602,6 +618,9 @@ FPathTracingPayload TraceTransparentRay(inout FPathState PathState, int Bounce,
 						break;
 					}
 				}
+				// accumulate all at once so clamping applies to entire volume contribution
+				AccumulateRadiance(PathState.Radiance, VolumeRadiance, bIsCameraRay);
+
 				// restore interval start now that we are finished with ray marching
 				Interval.VolumeTMin = StartVolumeTMin;
 				if (UseNullTrackingTransmittance)
@@ -1328,6 +1347,12 @@ FPathState CreatePathState(int2 PixelIndex, int2 PrimaryRayTextureIndex)
 
 	uint2 LaunchIndex = PixelIndex + View.ViewRectMin.xy;
 
+#if PATH_TRACER_USE_ADAPTIVE_SAMPLING
+	// initialize progressive sobol sequence based on current index for this pixel
+	// NOTE: Error diffusion sampler cannot be used in this mode
+	int SampleIndex = int(VarianceTexture[PrimaryRayTextureIndex].z);
+	RandomSequence_Initialize(Output.RandSequence, LaunchIndex.x + LaunchIndex.y * 65536, TemporalSeed - Iteration + SampleIndex);
+#else
 	// Initialize random sequence 
 	if (SamplerType == PATHTRACER_SAMPLER_ERROR_DIFFUSION)
 	{
@@ -1339,7 +1364,7 @@ FPathState CreatePathState(int2 PixelIndex, int2 PrimaryRayTextureIndex)
 		// random sobol init
 		RandomSequence_Initialize(Output.RandSequence, LaunchIndex.x + LaunchIndex.y * 65536, TemporalSeed);
 	}
-
+#endif
 
 	// Initialize ray and payload
 	{
diff --git a/Engine/Source/Runtime/Renderer/Private/PathTracing.cpp b/Engine/Source/Runtime/Renderer/Private/PathTracing.cpp
--- a/Engine/Source/Runtime/Renderer/Private/PathTracing.cpp
+++ b/Engine/Source/Runtime/Renderer/Private/PathTracing.cpp
@@ -35,6 +35,14 @@ TAutoConsoleVariable<int32> CVarPathTracing(
 #include <limits>
 #include "PathTracingSpatialTemporalDenoising.h"
 
+TAutoConsoleVariable<int32> CVarPathTracingExperimental(
+	TEXT("r.PathTracing.Experimental"),
+	0,
+	TEXT("Enables some experimental features of the path tracing renderer that require compiling additional permutations of the path tracer."),
+	ECVF_RenderThreadSafe | ECVF_ReadOnly
+);
+
+
 TAutoConsoleVariable<int32> CVarPathTracingCompaction(
 	TEXT("r.PathTracing.Compaction"),
 	1,
@@ -382,6 +390,33 @@ TAutoConsoleVariable<int32> CVarPathTracingUseAnalyticTransmittance(
 	ECVF_RenderThreadSafe
 );
 
+TAutoConsoleVariable<int32> CVarPathTracingAdaptiveSampling(
+	TEXT("r.PathTracing.AdaptiveSampling"),
+	0,
+	TEXT("Determines if adaptive sampling is enabled. When non-zero, the path tracer will try to skip calculation of pixels below the specified error threshold.\n")
+	TEXT("0: off (uniform sampling - default)\n")
+	TEXT("1: on (adaptive sampling)\n"),
+	ECVF_RenderThreadSafe
+);
+
+TAutoConsoleVariable<float> CVarPathTracingAdaptiveSamplingErrorThreshold(
+	TEXT("r.PathTracing.AdaptiveSampling.ErrorThreshold"),
+	0.001f,
+	TEXT("This is the target perceptual error threshold. Once a pixel's error falls below this value, it will not be sampled again (default: 0.001)\n"),
+	ECVF_RenderThreadSafe
+);
+
+TAutoConsoleVariable<int32> CVarPathTracingAdaptiveSamplingVisualize(
+	TEXT("r.PathTracing.AdaptiveSampling.Visualize"),
+	0,
+	TEXT("Select a visualization mode to help understand how adaptive sampling is working.\n")
+	TEXT("0: off\n")
+	TEXT("1: Visualize active pixels with heatmap (converged pixels are displayed as is)\n")
+	TEXT("2: Visualize sample count heatmap (against current max samples)\n")
+	TEXT("3-7: Visualize variance mip levels\n"),
+	ECVF_RenderThreadSafe
+);
+
 BEGIN_SHADER_PARAMETER_STRUCT(FPathTracingData, )
 	SHADER_PARAMETER(float, BlendFactor)
 	SHADER_PARAMETER(uint32, Iteration)
@@ -429,8 +464,10 @@ struct FPathTracingConfig
 	bool UseMISCompensation;
 	bool LockedSamplingPattern;
 	bool UseCameraMediumTracking;
+	bool UseAdaptiveSampling;
 	bool UseMultiGPU; // NOTE: Requires invalidation because the buffer layout changes
 	int DenoiserMode; // NOTE: does not require path tracing invalidation
+	float AdaptiveSamplingThreshold;
 
 	bool IsDifferent(const FPathTracingConfig& Other) const
 	{
@@ -469,6 +506,8 @@ struct FPathTracingConfig
 			UseMISCompensation != Other.UseMISCompensation ||
 			LockedSamplingPattern != Other.LockedSamplingPattern ||
 			UseCameraMediumTracking != Other.UseCameraMediumTracking ||
+			UseAdaptiveSampling != Other.UseAdaptiveSampling ||
+			AdaptiveSamplingThreshold != Other.AdaptiveSamplingThreshold ||
 			UseMultiGPU != Other.UseMultiGPU;
 	}
 
@@ -522,6 +561,7 @@ struct FPathTracingState {
 	FPathTracingConfig LastConfig;
 	// Textures holding onto the accumulated frame data
 	TRefCountPtr<IPooledRenderTarget> RadianceRT;
+	TRefCountPtr<IPooledRenderTarget> VarianceRT;
 	TRefCountPtr<IPooledRenderTarget> AlbedoRT;
 	TRefCountPtr<IPooledRenderTarget> NormalRT;
 	TRefCountPtr<FRDGPooledBuffer> VarianceBuffer;
@@ -856,12 +896,37 @@ class FPathTracingRG : public FGlobalShader
 	DECLARE_GLOBAL_SHADER(FPathTracingRG)
 	SHADER_USE_ROOT_PARAMETER_STRUCT(FPathTracingRG, FGlobalShader)
 
-	class FCompactionType : SHADER_PERMUTATION_INT("PATH_TRACER_USE_COMPACTION", 2);
+	class FCompactionType : SHADER_PERMUTATION_BOOL("PATH_TRACER_USE_COMPACTION");
+	class FAdaptiveSampling : SHADER_PERMUTATION_BOOL("PATH_TRACER_USE_ADAPTIVE_SAMPLING");
 	class FSubstrateComplexSpecialMaterial : SHADER_PERMUTATION_BOOL("PATH_TRACER_USE_SUBSTRATE_SPECIAL_COMPLEX_MATERIAL");
-	using FPermutationDomain = TShaderPermutationDomain<FCompactionType, FSubstrateComplexSpecialMaterial>;
+	using FPermutationDomain = TShaderPermutationDomain<FCompactionType, FAdaptiveSampling, FSubstrateComplexSpecialMaterial>;
 
 	static bool ShouldCompilePermutation(const FGlobalShaderPermutationParameters& Parameters)
 	{
+		const bool bUseExperimental = CVarPathTracingExperimental.GetValueOnAnyThread() != 0;
+		FPermutationDomain PermutationVector(Parameters.PermutationId);
+		if (bUseExperimental == false)
+		{
+			
+			if (PermutationVector.Get<FCompactionType>() == 0)
+			{
+				// non-compaction tracing is considered experimental
+				return false;
+			}
+			if (PermutationVector.Get<FAdaptiveSampling>())
+			{
+				// adaptive sampling is experimental
+				return false;
+			}
+		}
+		if (!Substrate::IsSubstrateEnabled())
+		{
+			// If we aren't using strata, no need to compile the complex material path
+			if (PermutationVector.Get<FSubstrateComplexSpecialMaterial>())
+			{
+				return false;
+			}
+		}
 		return ShouldCompilePathTracingShadersForProject(Parameters.Platform);
 	}
 
@@ -879,6 +944,7 @@ class FPathTracingRG : public FGlobalShader
 
 	BEGIN_SHADER_PARAMETER_STRUCT(FParameters, )
 		SHADER_PARAMETER_RDG_TEXTURE_UAV(RWTexture2D<float4>, RadianceTexture)
+		SHADER_PARAMETER_RDG_TEXTURE_UAV(RWTexture2D<float2>, VarianceTexture)
 		SHADER_PARAMETER_RDG_TEXTURE_UAV(RWTexture2D<float4>, AlbedoTexture)
 		SHADER_PARAMETER_RDG_TEXTURE_UAV(RWTexture2D<float4>, NormalTexture)
 		SHADER_PARAMETER_RDG_BUFFER_SRV(RaytracingAccelerationStructure, TLAS)
@@ -997,7 +1063,7 @@ class FPathTracingBuildAtmosphereOpticalDepthLUTCS : public FGlobalShader
 
 	static void ModifyCompilationEnvironment(const FGlobalShaderPermutationParameters& Parameters, FShaderCompilerEnvironment& OutEnvironment)
 	{
-		//OutEnvironment.CompilerFlags.Add(CFLAG_WarningsAsErrors);
+		OutEnvironment.CompilerFlags.Add(CFLAG_WarningsAsErrors);
 		OutEnvironment.SetDefine(TEXT("THREADGROUPSIZE_X"), FComputeShaderUtils::kGolden2DGroupSize);
 		OutEnvironment.SetDefine(TEXT("THREADGROUPSIZE_Y"), FComputeShaderUtils::kGolden2DGroupSize);
 	}
@@ -1011,6 +1077,68 @@ class FPathTracingBuildAtmosphereOpticalDepthLUTCS : public FGlobalShader
 };
 IMPLEMENT_SHADER_TYPE(, FPathTracingBuildAtmosphereOpticalDepthLUTCS, TEXT("/Engine/Private/PathTracing/PathTracingBuildAtmosphereLUT.usf"), TEXT("PathTracingBuildAtmosphereOpticalDepthLUTCS"), SF_Compute);
 
+
+class FPathTracingBuildAdaptiveErrorTextureCS : public FGlobalShader
+{
+	DECLARE_GLOBAL_SHADER(FPathTracingBuildAdaptiveErrorTextureCS)
+	SHADER_USE_PARAMETER_STRUCT(FPathTracingBuildAdaptiveErrorTextureCS, FGlobalShader)
+
+	static bool ShouldCompilePermutation(const FGlobalShaderPermutationParameters& Parameters)
+	{
+		return ShouldCompilePathTracingShadersForProject(Parameters.Platform);
+	}
+
+	static void ModifyCompilationEnvironment(const FGlobalShaderPermutationParameters& Parameters, FShaderCompilerEnvironment& OutEnvironment)
+	{
+		OutEnvironment.CompilerFlags.Add(CFLAG_WarningsAsErrors);
+		OutEnvironment.SetDefine(TEXT("THREADGROUPSIZE_X"), FComputeShaderUtils::kGolden2DGroupSize);
+		OutEnvironment.SetDefine(TEXT("THREADGROUPSIZE_Y"), FComputeShaderUtils::kGolden2DGroupSize);
+	}
+
+	BEGIN_SHADER_PARAMETER_STRUCT(FParameters, )
+		SHADER_PARAMETER(FIntPoint, InputResolution)
+		SHADER_PARAMETER(FIntPoint, OutputResolution)
+		SHADER_PARAMETER_SAMPLER(SamplerState, InputMipSampler)
+		SHADER_PARAMETER_RDG_TEXTURE_SRV(Texture2D, InputMip)
+		SHADER_PARAMETER_RDG_TEXTURE_UAV(RWTexture2D, OutputMip)
+
+	END_SHADER_PARAMETER_STRUCT()
+};
+IMPLEMENT_SHADER_TYPE(, FPathTracingBuildAdaptiveErrorTextureCS, TEXT("/Engine/Private/PathTracing/PathTracingBuildAdaptiveError.usf"), TEXT("PathTracingBuildAdaptiveErrorTextureCS"), SF_Compute);
+
+class FPathTracingAdaptiveStartCS : public FGlobalShader
+{
+	DECLARE_GLOBAL_SHADER(FPathTracingAdaptiveStartCS)
+	SHADER_USE_PARAMETER_STRUCT(FPathTracingAdaptiveStartCS, FGlobalShader)
+
+	static bool ShouldCompilePermutation(const FGlobalShaderPermutationParameters& Parameters)
+	{
+		return ShouldCompilePathTracingShadersForProject(Parameters.Platform);
+	}
+
+	static void ModifyCompilationEnvironment(const FGlobalShaderPermutationParameters& Parameters, FShaderCompilerEnvironment& OutEnvironment)
+	{
+		OutEnvironment.CompilerFlags.Add(CFLAG_WarningsAsErrors);
+		OutEnvironment.SetDefine(TEXT("THREADGROUPSIZE_X"), FComputeShaderUtils::kGolden2DGroupSize);
+		OutEnvironment.SetDefine(TEXT("THREADGROUPSIZE_Y"), FComputeShaderUtils::kGolden2DGroupSize);
+	}
+
+	BEGIN_SHADER_PARAMETER_STRUCT(FParameters, )
+		SHADER_PARAMETER_RDG_TEXTURE_SRV(Texture2D, VarianceTexture)
+		SHADER_PARAMETER_SAMPLER(SamplerState, VarianceSampler)
+		SHADER_PARAMETER(FIntVector, VarianceTextureDims)
+		SHADER_PARAMETER(float, AdaptiveSamplingErrorThreshold)
+		SHADER_PARAMETER(FIntPoint, TileTextureOffset)
+		SHADER_PARAMETER(FIntPoint, DispatchDim)
+		SHADER_PARAMETER(float, ViewPreExposure)
+		SHADER_PARAMETER_RDG_BUFFER_UAV(RWBuffer<int>, NextActivePaths)
+		SHADER_PARAMETER_RDG_BUFFER_UAV(RWBuffer<int>, NumPathStates)
+	END_SHADER_PARAMETER_STRUCT()
+};
+IMPLEMENT_SHADER_TYPE(, FPathTracingAdaptiveStartCS, TEXT("/Engine/Private/PathTracing/PathTracingAdaptiveStart.usf"), TEXT("PathTracingAdaptiveStartCS"), SF_Compute);
+
+
+
 // Default miss shader (using the path tracing payload)
 template <bool IsGPULightmass>
 class TPathTracingDefaultMS : public FGlobalShader
@@ -2137,13 +2265,23 @@ class FPathTracingCompositorPS : public FGlobalShader
 		return ShouldCompileRayTracingShadersForProject(Parameters.Platform);
 	}
 
+	static void ModifyCompilationEnvironment(const FGlobalShaderPermutationParameters& Parameters, FShaderCompilerEnvironment& OutEnvironment)
+	{
+		OutEnvironment.CompilerFlags.Add(CFLAG_WarningsAsErrors);
+	}
+
 	BEGIN_SHADER_PARAMETER_STRUCT(FParameters, )
+		SHADER_PARAMETER_SAMPLER(SamplerState, VarianceSampler)
 		SHADER_PARAMETER_RDG_TEXTURE_SRV(Texture2D<float4>, RadianceTexture)
+		SHADER_PARAMETER_RDG_TEXTURE_SRV(Texture2D<float2>, VarianceTexture)
 		SHADER_PARAMETER_RDG_TEXTURE_SRV(Texture2D<float4>, NormalDepthTexture)
 		SHADER_PARAMETER_STRUCT_REF(FViewUniformShaderParameters, ViewUniformBuffer)
 		SHADER_PARAMETER(uint32, Iteration)
 		SHADER_PARAMETER(uint32, MaxSamples)
 		SHADER_PARAMETER(int, ProgressDisplayEnabled)
+		SHADER_PARAMETER(float, AdaptiveSamplingErrorThreshold)
+		SHADER_PARAMETER(int, AdaptiveSamplingVisualize)
+		SHADER_PARAMETER(FIntVector, VarianceTextureDims)
 
 		RENDER_TARGET_BINDING_SLOTS()
 	END_SHADER_PARAMETER_STRUCT()
@@ -2152,11 +2290,14 @@ IMPLEMENT_SHADER_TYPE(, FPathTracingCompositorPS, TEXT("/Engine/Private/PathTrac
 
 static FPathTracingRG::FPermutationDomain GetPathTracingRGPermutation(const FScene& Scene)
 {
-	const int CompactionType = CVarPathTracingCompaction.GetValueOnRenderThread();
+	const bool bUseExperimental = CVarPathTracingExperimental.GetValueOnRenderThread() != 0;
+	const bool bUseCompaction = bUseExperimental && CVarPathTracingCompaction.GetValueOnRenderThread() != 0;
+	const bool bUseAdaptiveSampling = bUseExperimental && CVarPathTracingAdaptiveSampling.GetValueOnRenderThread() != 0;
 	const bool bHasComplexSpecialRenderPath = Substrate::IsSubstrateEnabled() && Scene.SubstrateSceneData.bUsesComplexSpecialRenderPath;
 
 	FPathTracingRG::FPermutationDomain Out;
-	Out.Set<FPathTracingRG::FCompactionType>(CompactionType);
+	Out.Set<FPathTracingRG::FCompactionType>(bUseCompaction);
+	Out.Set<FPathTracingRG::FAdaptiveSampling>(bUseAdaptiveSampling);
 	Out.Set<FPathTracingRG::FSubstrateComplexSpecialMaterial>(bHasComplexSpecialRenderPath);
 	return Out;
 }
@@ -2167,7 +2308,6 @@ void FDeferredShadingSceneRenderer::PreparePathTracing(const FSceneViewFamily& V
 		&& ShouldCompilePathTracingShadersForProject(ViewFamily.GetShaderPlatform()))
 	{
 		// Declare all RayGen shaders that require material closest hit shaders to be bound
-		const int CompactionType = CVarPathTracingCompaction.GetValueOnRenderThread();
 		FPathTracingRG::FPermutationDomain PermutationVector = GetPathTracingRGPermutation(Scene);
 		{
 			auto RayGenShader = GetGlobalShaderMap(ViewFamily.GetShaderPlatform())->GetShader<FPathTracingRG>(PermutationVector);
@@ -2197,6 +2337,7 @@ void FSceneViewState::PathTracingInvalidate(bool InvalidateAnimationStates)
 		}
 
 		State->RadianceRT.SafeRelease();
+		State->VarianceRT.SafeRelease();
 		State->AlbedoRT.SafeRelease();
 		State->NormalRT.SafeRelease();
 		State->VarianceBuffer.SafeRelease();
@@ -2275,8 +2416,13 @@ void FDeferredShadingSceneRenderer::RenderPathTracing(
 	int32 SamplesPerPixelCVar = View.bIsOfflineRender ? -1 : CVarPathTracingSamplesPerPixel.GetValueOnRenderThread();
 	uint32 MaxSPP = SamplesPerPixelCVar > -1 ? SamplesPerPixelCVar : View.FinalPostProcessSettings.PathTracingSamplesPerPixel;
 	MaxSPP = FMath::Max(MaxSPP, 1u);
+
+	const bool bUseExperimental = CVarPathTracingExperimental.GetValueOnRenderThread() != 0;
+
 	Config.LockedSamplingPattern = CVarPathTracingFrameIndependentTemporalSeed.GetValueOnRenderThread() == 0;
 	Config.UseCameraMediumTracking = CVarPathTracingCameraMediumTracking.GetValueOnRenderThread() != 0;
+	Config.UseAdaptiveSampling = bUseExperimental && CVarPathTracingAdaptiveSampling.GetValueOnAnyThread() != 0;
+	Config.AdaptiveSamplingThreshold = CVarPathTracingAdaptiveSamplingErrorThreshold.GetValueOnRenderThread();
 
 	// compute an integer code of what show flags and booleans related to lights are currently enabled so we can detect changes
 	Config.LightShowFlags = 0;
@@ -2441,6 +2587,8 @@ void FDeferredShadingSceneRenderer::RenderPathTracing(
 	}
 #if WITH_MGPU
 	Config.UseMultiGPU = CVarPathTracingMultiGPU.GetValueOnRenderThread() != 0;
+	// TODO: Figure out how to support adaptive sampling in multi-gpu cases (this is complicated due to the swizzled layout of the variance texture)
+	Config.UseMultiGPU &= !Config.UseAdaptiveSampling;
 #else
 	Config.UseMultiGPU = false;
 #endif
@@ -2466,8 +2614,10 @@ void FDeferredShadingSceneRenderer::RenderPathTracing(
 
 	// Prepare radiance buffer (will be shared with display pass)
 	FRDGTexture* RadianceTexture = nullptr;
+	FRDGTexture* VarianceTexture = nullptr;
 	FRDGTexture* AlbedoTexture = nullptr;
 	FRDGTexture* NormalTexture = nullptr;
+	const int NumVarianceMips = FMath::Min(5u, 1 + FMath::FloorLog2(uint32(View.ViewRect.Size().GetMin())));
 	if (PathTracingState->RadianceRT)
 	{
 		// we already have a valid radiance texture, re-use it
@@ -2487,6 +2637,30 @@ void FDeferredShadingSceneRenderer::RenderPathTracing(
 		AlbedoTexture   = GraphBuilder.CreateTexture(Desc, TEXT("PathTracer.Albedo")  , ERDGTextureFlags::MultiFrame);
 		NormalTexture   = GraphBuilder.CreateTexture(Desc, TEXT("PathTracer.Normal")  , ERDGTextureFlags::MultiFrame);
 	}
+	if (Config.UseAdaptiveSampling)
+	{
+		if (PathTracingState->VarianceRT)
+		{
+			VarianceTexture = GraphBuilder.RegisterExternalTexture(PathTracingState->VarianceRT, TEXT("PathTracer.Variance"));
+		}
+		else
+		{
+			// format stores Luminance,Luminance^2,NumSamples which can be used for error estimation
+			FRDGTextureDesc Desc = FRDGTextureDesc::Create2D(
+				View.ViewRect.Size(),
+				PF_A32B32G32R32F,
+				FClearValueBinding::None,
+				TexCreate_ShaderResource | TexCreate_UAV);
+			Desc.NumMips = NumVarianceMips;
+			VarianceTexture = GraphBuilder.CreateTexture(Desc, TEXT("PathTracer.Variance"), ERDGTextureFlags::MultiFrame);
+		}
+	}
+	else
+	{
+		// If we are not using adaptive, make sure the old variance buffer doesn't stick around
+		PathTracingState->VarianceRT.SafeRelease();
+	}
+
 
 	// should we use multiple GPUs to render the image?
 	const FRHIGPUMask GPUMask = Config.UseMultiGPU ? FRHIGPUMask::All() : View.GPUMask;
@@ -2526,6 +2700,7 @@ void FDeferredShadingSceneRenderer::RenderPathTracing(
 
 		if (bNeedsMoreRays)
 		{
+			RDG_EVENT_SCOPE(GraphBuilder, "Path Tracing Compute (%d x %d)", DispatchResX, DispatchResY);
 			bool bForceRebuild = CVarPathTracingHeterogeneousVolumesRebuildEveryFrame.GetValueOnRenderThread() != 0;
 			bCreateVolumeGrids = bForceRebuild ||
 				!PathTracingState->AdaptiveFrustumGridParameterCache.TopLevelGridBuffer ||
@@ -2552,14 +2727,14 @@ void FDeferredShadingSceneRenderer::RenderPathTracing(
 			bNeedsTextureExtract = true;
 
 			// should we use path compaction?
-			const int CompactionType = CVarPathTracingCompaction.GetValueOnRenderThread();
+			const bool bUseCompaction = bUseExperimental && CVarPathTracingCompaction.GetValueOnRenderThread() != 0;
 			const bool bUseIndirectDispatch = GRHISupportsRayTracingDispatchIndirect && CVarPathTracingIndirectDispatch.GetValueOnRenderThread() != 0;
 			const int FlushRenderingCommands = CVarPathTracingFlushDispatch.GetValueOnRenderThread();
 
 			FRDGBuffer* ActivePaths[2] = {};
 			FRDGBuffer* NumActivePaths[2] = {};
 			FRDGBuffer* PathStateData = nullptr;
-			if (CompactionType == 1)
+			if (bUseCompaction || Config.UseAdaptiveSampling)
 			{
 				const int32 NumPaths = FMath::Min(
 					DispatchSize * FMath::DivideAndRoundUp(DispatchSize, NumGPUs),
@@ -2603,13 +2778,46 @@ void FDeferredShadingSceneRenderer::RenderPathTracing(
 
 						// Compute the dispatch size for just this set of scanlines
 						const int32 DispatchSizeYLocal = FMath::Min(DispatchSizeYSplit, DispatchSizeY - CurrentGPU * DispatchSizeYSplit);
-						if (CompactionType == 1)
+
+						RDG_EVENT_SCOPE_CONDITIONAL(GraphBuilder, DispatchResX > DispatchSize || DispatchResY > DispatchSize, "Tile=(%d,%d - %dx%d)", TileX, TileY, DispatchSizeX, DispatchSizeYLocal);
+
+						if (Config.UseAdaptiveSampling && Config.PathTracingData.Iteration > 0)
+						{
+							// If we are using adaptive sampling, build a smaller list of active paths after the first iteration
+							TShaderMapRef<FPathTracingAdaptiveStartCS> ComputeShader(GetGlobalShaderMap(View.FeatureLevel));
+
+							FPathTracingAdaptiveStartCS::FParameters* PassParameters = GraphBuilder.AllocParameters<FPathTracingAdaptiveStartCS::FParameters>();
+
+							PassParameters->VarianceTexture = GraphBuilder.CreateSRV(VarianceTexture);
+							PassParameters->VarianceSampler = TStaticSamplerState<SF_Bilinear>::GetRHI();
+							PassParameters->VarianceTextureDims = FIntVector(DispatchResX, DispatchResY, NumVarianceMips);
+							PassParameters->AdaptiveSamplingErrorThreshold = Config.AdaptiveSamplingThreshold;
+							PassParameters->ViewPreExposure = View.PreExposure;
+
+							PassParameters->NextActivePaths = GraphBuilder.CreateUAV(ActivePaths[0], PF_R32_SINT);
+							PassParameters->NumPathStates = GraphBuilder.CreateUAV(NumActivePaths[0], PF_R32_UINT);
+							AddClearUAVPass(GraphBuilder, PassParameters->NextActivePaths, -1); // make sure everything is initialized to -1 since paths that go inactive don't write anything
+							AddClearUAVPass(GraphBuilder, PassParameters->NumPathStates, 0);
+
+							PassParameters->TileTextureOffset.X = TileX;
+							PassParameters->TileTextureOffset.Y = TileY + CurrentGPU * DispatchSizeYSplit;
+							PassParameters->DispatchDim = FIntPoint(DispatchSizeX, DispatchSizeYLocal);
+
+							FComputeShaderUtils::AddPass(
+								GraphBuilder,
+								RDG_EVENT_NAME("Prepare Adaptive Sampling Mask"),
+								ComputeShader,
+								PassParameters,
+								FComputeShaderUtils::GetGroupCount(PassParameters->DispatchDim, FComputeShaderUtils::kGolden2DGroupSize));
+						}
+						else if (bUseCompaction)
 						{
 							AddClearUAVPass(GraphBuilder, GraphBuilder.CreateUAV(ActivePaths[0], PF_R32_UINT), 0);
 						}
+
 						// When using path compaction, we need to run the path tracer once per bounce
 						// otherwise, the path tracer is the one doing the bounces
-						for (int Bounce = 0, MaxBounces = CompactionType == 1 ? Config.PathTracingData.MaxBounces : 0; Bounce <= MaxBounces; Bounce++)
+						for (int Bounce = 0, MaxBounces = bUseCompaction ? Config.PathTracingData.MaxBounces : 0; Bounce <= MaxBounces; Bounce++)
 						{
 							FPathTracingRG::FParameters* PassParameters = GraphBuilder.AllocParameters<FPathTracingRG::FParameters>();
 							PassParameters->TLAS = Scene->RayTracingScene.GetLayerView(ERayTracingSceneLayer::Base);
@@ -2638,6 +2846,17 @@ void FDeferredShadingSceneRenderer::RenderPathTracing(
 							PassParameters->AlbedoTexture = GraphBuilder.CreateUAV(AlbedoTexture);
 							PassParameters->NormalTexture = GraphBuilder.CreateUAV(NormalTexture);
 
+							if (Config.UseAdaptiveSampling)
+							{
+								PassParameters->VarianceTexture = GraphBuilder.CreateUAV(VarianceTexture);
+							}
+							else
+							{
+								// this texture is not used in this case
+								PassParameters->VarianceTexture = nullptr;
+							}
+							
+
 							if (PreviousPassParameters != nullptr)
 							{
 								PassParameters->Atmosphere = PreviousPassParameters->Atmosphere;
@@ -2683,7 +2902,7 @@ void FDeferredShadingSceneRenderer::RenderPathTracing(
 							PassParameters->ScanlineWidth = DispatchSizeX;
 
 							PassParameters->Bounce = Bounce;
-							if (CompactionType == 1)
+							if (bUseCompaction)
 							{
 								PassParameters->ActivePaths = GraphBuilder.CreateSRV(ActivePaths[Bounce & 1], PF_R32_SINT);
 								PassParameters->NextActivePaths = GraphBuilder.CreateUAV(ActivePaths[(Bounce & 1) ^ 1], PF_R32_SINT);
@@ -2703,9 +2922,9 @@ void FDeferredShadingSceneRenderer::RenderPathTracing(
 							ClearUnusedGraphResources(RayGenShader, PassParameters);
 							const bool bFlushRenderingCommands = FlushRenderingCommands == 1 || (FlushRenderingCommands == 2 && Bounce == MaxBounces);
 							GraphBuilder.AddPass(
-								CompactionType == 1
-								? RDG_EVENT_NAME("Path Tracer Compute (%d x %d) Tile=(%d,%d - %dx%d) Sample=%d/%d NumLights=%d (Bounce=%d%s)", DispatchResX, DispatchResY, TileX, TileY, DispatchSizeX, DispatchSizeYLocal, PathTracingState->SampleIndex, MaxSPP, PassParameters->SceneLightCount, Bounce, bUseIndirectDispatch && Bounce > 0 ? TEXT(" indirect") : TEXT(""))
-								: RDG_EVENT_NAME("Path Tracer Compute (%d x %d) Tile=(%d,%d - %dx%d) Sample=%d/%d NumLights=%d", DispatchResX, DispatchResY, TileX, TileY, DispatchSizeX, DispatchSizeYLocal, PathTracingState->SampleIndex, MaxSPP, PassParameters->SceneLightCount),
+								bUseCompaction
+								? RDG_EVENT_NAME("Path Tracer Sample=%d/%d NumLights=%d (Bounce=%d%s)", PathTracingState->SampleIndex, MaxSPP, PassParameters->SceneLightCount, Bounce, bUseIndirectDispatch && Bounce > 0 ? TEXT(" indirect") : TEXT(""))
+								: RDG_EVENT_NAME("Path Tracer Sample=%d/%d NumLights=%d"              , PathTracingState->SampleIndex, MaxSPP, PassParameters->SceneLightCount),
 								PassParameters,
 								ERDGPassFlags::Compute,
 								[PassParameters, RayGenShader, DispatchSizeX, DispatchSizeYLocal, bUseIndirectDispatch, bFlushRenderingCommands, GPUIndex, &View](FRHIRayTracingCommandList& RHICmdList)
@@ -2856,6 +3075,11 @@ void FDeferredShadingSceneRenderer::RenderPathTracing(
 		GraphBuilder.QueueTextureExtraction(RadianceTexture, &PathTracingState->RadianceRT);
 		GraphBuilder.QueueTextureExtraction(AlbedoTexture, &PathTracingState->AlbedoRT);
 		GraphBuilder.QueueTextureExtraction(NormalTexture, &PathTracingState->NormalRT);
+		if (Config.UseAdaptiveSampling)
+		{
+			check(VarianceTexture != nullptr);
+			GraphBuilder.QueueTextureExtraction(VarianceTexture, &PathTracingState->VarianceRT);
+		}
 	}
 
 	if (bCreateVolumeGrids)
@@ -2920,6 +3144,34 @@ void FDeferredShadingSceneRenderer::RenderPathTracing(
 	}
 #endif
 
+	// build adaptive sampling error map if we traced some rays
+	if (Config.UseAdaptiveSampling && bNeedsMoreRays)
+	{
+		RDG_EVENT_SCOPE(GraphBuilder, "Adaptive Sampling");
+		FIntPoint BufferSize = View.ViewRect.Size();
+		TShaderMapRef<FPathTracingBuildAdaptiveErrorTextureCS> ComputeShader(GetGlobalShaderMap(View.FeatureLevel));
+		for (int MipLevel = 0; MipLevel < NumVarianceMips - 1; MipLevel++)
+		{
+			FPathTracingBuildAdaptiveErrorTextureCS::FParameters* PassParameters = GraphBuilder.AllocParameters<FPathTracingBuildAdaptiveErrorTextureCS::FParameters>();
+
+			PassParameters->InputMipSampler = TStaticSamplerState<ESamplerFilter::SF_Bilinear>::CreateRHI();
+			PassParameters->InputMip = GraphBuilder.CreateSRV(FRDGTextureSRVDesc::CreateForMipLevel(VarianceTexture, MipLevel));
+			PassParameters->OutputMip = GraphBuilder.CreateUAV(FRDGTextureUAVDesc(VarianceTexture, MipLevel + 1));
+			PassParameters->InputResolution = FIntPoint(
+				FMath::Max(BufferSize.X >> MipLevel, 1),
+				FMath::Max(BufferSize.Y >> MipLevel, 1));
+			PassParameters->OutputResolution = FIntPoint(
+				FMath::Max(BufferSize.X >> (MipLevel + 1), 1),
+				FMath::Max(BufferSize.Y >> (MipLevel + 1), 1));
+			FComputeShaderUtils::AddPass(
+				GraphBuilder,
+				RDG_EVENT_NAME("Build Error Estimation Mips (%d)", MipLevel),
+				ComputeShader,
+				PassParameters,
+				FComputeShaderUtils::GetGroupCount(PassParameters->OutputResolution, FComputeShaderUtils::kGolden2DGroupSize));
+		}
+	}
+
 	FPathTracingSpatialTemporalDenoisingContext DenoisingContext;
 	const bool EnablePathTracingDenoiserRealtimeDebug = ShouldEnablePathTracingDenoiserRealtimeDebug();
 
@@ -3005,11 +3257,16 @@ void FDeferredShadingSceneRenderer::RenderPathTracing(
 	DisplayParameters->Iteration = Config.PathTracingData.Iteration;
 	DisplayParameters->MaxSamples = MaxSPP;
 	DisplayParameters->ProgressDisplayEnabled = CVarPathTracingProgressDisplay.GetValueOnRenderThread();
+	DisplayParameters->AdaptiveSamplingErrorThreshold = Config.AdaptiveSamplingThreshold;
+	DisplayParameters->AdaptiveSamplingVisualize = CVarPathTracingAdaptiveSamplingVisualize.GetValueOnRenderThread();
+	DisplayParameters->VarianceTextureDims = FIntVector(DispatchResX, DispatchResY, NumVarianceMips);
 	DisplayParameters->ViewUniformBuffer = View.ViewUniformBuffer;
 	DisplayParameters->RadianceTexture = GraphBuilder.CreateSRV(FRDGTextureSRVDesc::Create(DenoisedRadianceTexture ? DenoisedRadianceTexture : RadianceTexture));
+	DisplayParameters->VarianceTexture = GraphBuilder.CreateSRV(FRDGTextureSRVDesc::Create(VarianceTexture ? VarianceTexture : GSystemTextures.GetBlackDummy(GraphBuilder)));
 	DisplayParameters->NormalDepthTexture = GraphBuilder.CreateSRV(FRDGTextureSRVDesc::Create(NormalTexture));
 	DisplayParameters->RenderTargets[0] = FRenderTargetBinding(SceneColorOutputTexture, ERenderTargetLoadAction::ELoad);
 	DisplayParameters->RenderTargets.DepthStencil = FDepthStencilBinding(SceneDepthOutputTexture,  ERenderTargetLoadAction::ELoad, ERenderTargetLoadAction::ENoAction, FExclusiveDepthStencil::DepthWrite_StencilNop);
+	DisplayParameters->VarianceSampler = TStaticSamplerState<SF_Bilinear>::GetRHI();
 
 	FScreenPassTextureViewport Viewport(SceneColorOutputTexture, View.ViewRect);
 
