Local fog volume UB preparing for application on translucent.

#prelfight https://horde.devtools.epicgames.com/job/64f854bb9541d1f3824c04d7

diff --git a/Engine/Shaders/Private/LocalFogVolumeCommon.ush b/Engine/Shaders/Private/LocalFogVolumeCommon.ush
--- a/Engine/Shaders/Private/LocalFogVolumeCommon.ush
+++ b/Engine/Shaders/Private/LocalFogVolumeCommon.ush
@@ -33,7 +33,10 @@ struct FLocalFogVolumeGPUInstanceData
 
 // Structured buffers cannot be sampled from the vertex shader on some mobile platforms.
 // That is why simple Buffers are used, because this tech needs to run on mobile platforms.
-Buffer<float4> LocalFogVolumeInstances;
+
+uint LFV_LocalFogVolumeInstanceCount;
+Buffer<float4> LFV_LocalFogVolumeInstances;
+#define LFVInsts LFV_LocalFogVolumeInstances
 
 
 
@@ -42,23 +45,23 @@ FLocalFogVolumeGPUInstanceData GetLocalFogVolumeGPUInstanceData(uint Index)
 	FLocalFogVolumeGPUInstanceData Data;
 	uint Offset = Index * 19;
 
-	Data.Transform = float4x4(LocalFogVolumeInstances[Offset + 0], LocalFogVolumeInstances[Offset + 1], LocalFogVolumeInstances[Offset + 2], LocalFogVolumeInstances[Offset + 3]);
-	Data.InvTransform = float4x4(LocalFogVolumeInstances[Offset + 4], LocalFogVolumeInstances[Offset + 5], LocalFogVolumeInstances[Offset + 6], LocalFogVolumeInstances[Offset + 7]);
+	Data.Transform			= float4x4(LFVInsts[Offset + 0 ], LFVInsts[Offset + 1 ], LFVInsts[Offset + 2 ], LFVInsts[Offset + 3 ]);
+	Data.InvTransform		= float4x4(LFVInsts[Offset + 4 ], LFVInsts[Offset + 5 ], LFVInsts[Offset + 6 ], LFVInsts[Offset + 7 ]);
 
-	Data.InvTranformNoScale = float4x4(LocalFogVolumeInstances[Offset + 8], LocalFogVolumeInstances[Offset + 9], LocalFogVolumeInstances[Offset + 10], LocalFogVolumeInstances[Offset + 11]);
-	Data.TransformScaleOnly = float4x4(LocalFogVolumeInstances[Offset + 12], LocalFogVolumeInstances[Offset + 13], LocalFogVolumeInstances[Offset + 14], LocalFogVolumeInstances[Offset + 15]);
+	Data.InvTranformNoScale = float4x4(LFVInsts[Offset + 8 ], LFVInsts[Offset + 9 ], LFVInsts[Offset + 10], LFVInsts[Offset + 11]);
+	Data.TransformScaleOnly = float4x4(LFVInsts[Offset + 12], LFVInsts[Offset + 13], LFVInsts[Offset + 14], LFVInsts[Offset + 15]);
 
-	float4 Temp0 = LocalFogVolumeInstances[Offset + 16];
+	float4 Temp0 = LFVInsts[Offset + 16];
 	Data.Density = Temp0.x;
 	Data.HeightFalloff = Temp0.y;
 	Data.HeightOffset = Temp0.z;
 	Data.RadialAttenuation = Temp0.w;
 
-	float4 Temp1 = LocalFogVolumeInstances[Offset + 17];
+	float4 Temp1 = LFVInsts[Offset + 17];
 	Data.Albedo = Temp1.xyz;
 	Data.PhaseG = Temp1.w;
 
-	float4 Temp2 = LocalFogVolumeInstances[Offset + 18];
+	float4 Temp2 = LFVInsts[Offset + 18];
 	Data.Emissive = Temp2.xyz;
 	Data.FogMode = Temp2.w;
 
diff --git a/Engine/Source/Runtime/Renderer/Private/DeferredShadingRenderer.cpp b/Engine/Source/Runtime/Renderer/Private/DeferredShadingRenderer.cpp
--- a/Engine/Source/Runtime/Renderer/Private/DeferredShadingRenderer.cpp
+++ b/Engine/Source/Runtime/Renderer/Private/DeferredShadingRenderer.cpp
@@ -3431,6 +3431,9 @@ void FDeferredShadingSceneRenderer::Render(FRDGBuilder& GraphBuilder)
 		// Run before RenderSkyAtmosphereLookUpTables for cloud shadows to be valid.
 		InitVolumetricCloudsForViews(GraphBuilder, bShouldRenderVolumetricCloudBase, InstanceCullingManager);
 
+		// Run local fog volume initialisation before basepass for when data is needed in forward
+		InitLocalFogVolumesForViews(Scene, Views, GraphBuilder);
+
 		if (SkyAtmospherePassLocation == ESkyAtmospherePassLocation::BeforeOcclusion && bShouldRenderSkyAtmosphere)
 		{
 			// Generate the Sky/Atmosphere look up tables
diff --git a/Engine/Source/Runtime/Renderer/Private/LocalFogVolumeRendering.cpp b/Engine/Source/Runtime/Renderer/Private/LocalFogVolumeRendering.cpp
--- a/Engine/Source/Runtime/Renderer/Private/LocalFogVolumeRendering.cpp
+++ b/Engine/Source/Runtime/Renderer/Private/LocalFogVolumeRendering.cpp
@@ -26,6 +26,23 @@ bool ShouldRenderLocalFogVolume(const FScene* Scene, const FSceneViewFamily& Fam
 
 DECLARE_GPU_STAT(LocalFogVolumeVolumes);
 
+static const uint32 SizeOfUintVec4 = sizeof(FUintVector4);
+static const uint32 UintVec4CountInLocalFogVolumeGPUInstanceData = sizeof(FLocalFogVolumeGPUInstanceData) / SizeOfUintVec4;
+
+static void SetDummyLocalFogVolumeForView(FRDGBuilder& GraphBuilder, FViewInfo& View)
+{
+	// The size of the structure must be a multiple of FUintVector4.
+	static_assert(sizeof(FLocalFogVolumeGPUInstanceData) == UintVec4CountInLocalFogVolumeGPUInstanceData * SizeOfUintVec4);
+
+	View.LocalFogVolumeGPUInstanceCount = 0;
+
+	static FLocalFogVolumeGPUInstanceData DummyData;
+	View.LocalFogVolumeGPUInstanceDataBuffer = CreateVertexBuffer(GraphBuilder, TEXT("DUMMYLocalFogVolumeGPUInstanceDataBuffer"),
+		FRDGBufferDesc::CreateBufferDesc(SizeOfUintVec4, UintVec4CountInLocalFogVolumeGPUInstanceData), &DummyData, sizeof(FLocalFogVolumeGPUInstanceData) * 1, ERDGInitialDataFlags::NoCopy);
+
+	View.LocalFogVolumeGPUInstanceDataBufferSRV = GraphBuilder.CreateSRV(View.LocalFogVolumeGPUInstanceDataBuffer, PF_A32B32G32R32F);
+};
+
 /*=============================================================================
 	FScene functions
 =============================================================================*/
@@ -114,20 +131,9 @@ void GetLocalFogVolumeSortingData(const FScene* Scene, FRDGBuilder& GraphBuilder
 
 void CreateViewLocalFogVolumeBufferSRV(FViewInfo& View, FRDGBuilder& GraphBuilder, FLocalFogVolumeSortingData& SortingData)
 {
-	static const uint32 SizeOfFloat4 = sizeof(float) * 4;
-	static const uint32 Float4CountInLocalFogVolumeGPUInstanceData = sizeof(FLocalFogVolumeGPUInstanceData) / SizeOfFloat4;
-	static_assert(sizeof(FLocalFogVolumeGPUInstanceData) == Float4CountInLocalFogVolumeGPUInstanceData * SizeOfFloat4); // The size of the structure must be a multiple of FVector4.
-
 	if (SortingData.LocalFogVolumeInstanceCountFinal == 0)
 	{
-		View.LocalFogVolumeGPUInstanceCount = 0;
-
-		static FLocalFogVolumeGPUInstanceData DummyData;
-		View.LocalFogVolumeGPUInstanceDataBufferSRV = GraphBuilder.CreateSRV(
-			CreateVertexBuffer(GraphBuilder, TEXT("LocalFogVolumeGPUInstanceDataBuffer"), 
-			FRDGBufferDesc::CreateBufferDesc(SizeOfFloat4, Float4CountInLocalFogVolumeGPUInstanceData), &DummyData, sizeof(FLocalFogVolumeGPUInstanceData) * 1, ERDGInitialDataFlags::NoCopy),
-			PF_A32B32G32R32F);
-		return;
+		SetDummyLocalFogVolumeForView(GraphBuilder, View);
 	}
 
 	// 1. Sort all the volumes
@@ -150,13 +156,45 @@ void CreateViewLocalFogVolumeBufferSRV(FViewInfo& View, FRDGBuilder& GraphBuilde
 
 	// 3. Allocate buffer and initialize with sorted data to upload to GPU
 	const uint32 AllLocalFogVolumeInstanceBytesFinal = sizeof(FLocalFogVolumeGPUInstanceData) * SortingData.LocalFogVolumeInstanceCountFinal;
-	FRDGBufferRef LocalFogVolumeGPUInstanceDataBuffer = CreateVertexBuffer(
+
+	View.LocalFogVolumeGPUInstanceCount = SortingData.LocalFogVolumeInstanceCountFinal;
+	View.LocalFogVolumeGPUInstanceDataBuffer = CreateVertexBuffer(
 		GraphBuilder, TEXT("LocalFogVolumeGPUInstanceDataBuffer"),
-		FRDGBufferDesc::CreateBufferDesc(SizeOfFloat4, SortingData.LocalFogVolumeInstanceCountFinal * Float4CountInLocalFogVolumeGPUInstanceData), 
+		FRDGBufferDesc::CreateBufferDesc(SizeOfUintVec4, SortingData.LocalFogVolumeInstanceCountFinal * UintVec4CountInLocalFogVolumeGPUInstanceData),
 		LocalFogVolumeGPUSortedInstanceData, AllLocalFogVolumeInstanceBytesFinal, ERDGInitialDataFlags::NoCopy);
+	View.LocalFogVolumeGPUInstanceDataBufferSRV = GraphBuilder.CreateSRV(View.LocalFogVolumeGPUInstanceDataBuffer, PF_A32B32G32R32F);
 
-	View.LocalFogVolumeGPUInstanceCount = SortingData.LocalFogVolumeInstanceCountFinal;
-	View.LocalFogVolumeGPUInstanceDataBufferSRV = GraphBuilder.CreateSRV(LocalFogVolumeGPUInstanceDataBuffer, PF_A32B32G32R32F);
+	View.LocalFogVolumeUniformParametersStruct.LocalFogVolumeInstanceCount = View.LocalFogVolumeGPUInstanceCount;
+	View.LocalFogVolumeUniformParametersStruct.LocalFogVolumeInstances = View.LocalFogVolumeGPUInstanceDataBufferSRV;
+	View.LocalFogVolumeUniformBuffer = GraphBuilder.CreateUniformBuffer(&View.LocalFogVolumeUniformParametersStruct);
+}
+
+void InitLocalFogVolumesForViews(
+	const FScene* Scene,
+	TArray<FViewInfo>& Views,
+	FRDGBuilder& GraphBuilder)
+{
+	const uint32 LocalFogVolumeInstanceCount = Scene->LocalFogVolumes.Num();
+	if (LocalFogVolumeInstanceCount > 0)
+	{
+		RDG_GPU_STAT_SCOPE(GraphBuilder, LocalFogVolumeVolumes);
+
+		FLocalFogVolumeSortingData SortingData;
+		GetLocalFogVolumeSortingData(Scene, GraphBuilder, SortingData);
+
+		for (FViewInfo& View : Views)
+		{
+			// LFV_TODO View.LocalFogVolumeGPUInstanceDataBufferSRV should be on the scene since it is common for all view.
+			CreateViewLocalFogVolumeBufferSRV(View, GraphBuilder, SortingData);
+		}
+	}
+	else
+	{
+		for (FViewInfo& View : Views)
+		{
+			SetDummyLocalFogVolumeForView(GraphBuilder, View);
+		}
+	}
 }
 
 /*=============================================================================
@@ -170,7 +208,7 @@ class FLocalFogVolumeVS : public FGlobalShader
 
 	BEGIN_SHADER_PARAMETER_STRUCT(FParameters, )
 		SHADER_PARAMETER_STRUCT_REF(FViewUniformShaderParameters, View)
-		SHADER_PARAMETER_RDG_BUFFER_SRV(Buffer<float4>, LocalFogVolumeInstances)
+		SHADER_PARAMETER_STRUCT(FLocalFogVolumeUniformParameters, LFV)
 	END_SHADER_PARAMETER_STRUCT()
 
 	using FPermutationDomain = TShaderPermutationDomain<>;
@@ -194,7 +232,7 @@ class FLocalFogVolumePS : public FGlobalShader
 
 	BEGIN_SHADER_PARAMETER_STRUCT(FParameters, )
 		SHADER_PARAMETER_STRUCT_REF(FViewUniformShaderParameters, View)
-		SHADER_PARAMETER_RDG_BUFFER_SRV(Buffer<float4>, LocalFogVolumeInstances)
+		SHADER_PARAMETER_STRUCT(FLocalFogVolumeUniformParameters, LFV)
 	END_SHADER_PARAMETER_STRUCT()
 
 	using FPermutationDomain = TShaderPermutationDomain<>;
@@ -230,82 +268,75 @@ void RenderLocalFogVolume(
 	{
 		RDG_GPU_STAT_SCOPE(GraphBuilder, LocalFogVolumeVolumes);
 
-		FLocalFogVolumeSortingData SortingData;
-		GetLocalFogVolumeSortingData(Scene, GraphBuilder, SortingData);
+		FRDGTextureRef SceneColorTexture = SceneTextures.Color.Resolve;
 
-		if (SortingData.LocalFogVolumeInstanceCountFinal > 0)
+		for (FViewInfo& View : Views)
 		{
-			FRDGTextureRef SceneColorTexture = SceneTextures.Color.Resolve;
-
-			for (FViewInfo& View : Views)
+			if (View.LocalFogVolumeGPUInstanceCount == 0)
 			{
-				CreateViewLocalFogVolumeBufferSRV(View, GraphBuilder, SortingData);
-				if (View.LocalFogVolumeGPUInstanceCount == 0)
-				{
-					continue;
-				}
+				continue;
+			}
 
-				FLocalFogVolumePassParameters* PassParameters = GraphBuilder.AllocParameters<FLocalFogVolumePassParameters>();
+			FLocalFogVolumePassParameters* PassParameters = GraphBuilder.AllocParameters<FLocalFogVolumePassParameters>();
 
-				PassParameters->VS.View = GetShaderBinding(View.ViewUniformBuffer);
-				PassParameters->VS.LocalFogVolumeInstances = View.LocalFogVolumeGPUInstanceDataBufferSRV;
+			PassParameters->VS.View = GetShaderBinding(View.ViewUniformBuffer);
+			PassParameters->VS.LFV = View.LocalFogVolumeUniformParametersStruct;
 
-				PassParameters->PS.View = GetShaderBinding(View.ViewUniformBuffer);
-				PassParameters->PS.LocalFogVolumeInstances = View.LocalFogVolumeGPUInstanceDataBufferSRV;
+			PassParameters->PS.View = GetShaderBinding(View.ViewUniformBuffer);
+			PassParameters->PS.LFV = View.LocalFogVolumeUniformParametersStruct;
 
-				PassParameters->SceneTextures = SceneTextures.UniformBuffer;
-				PassParameters->RenderTargets[0] = FRenderTargetBinding(SceneColorTexture, ERenderTargetLoadAction::ENoAction);
+			PassParameters->SceneTextures = SceneTextures.UniformBuffer;
+			PassParameters->RenderTargets[0] = FRenderTargetBinding(SceneColorTexture, ERenderTargetLoadAction::ENoAction);
 
-				FLocalFogVolumeVS::FPermutationDomain VSPermutationVector;
-				auto VertexShader = View.ShaderMap->GetShader< FLocalFogVolumeVS >(VSPermutationVector);
+			FLocalFogVolumeVS::FPermutationDomain VSPermutationVector;
+			auto VertexShader = View.ShaderMap->GetShader< FLocalFogVolumeVS >(VSPermutationVector);
 
-				FLocalFogVolumePS::FPermutationDomain PsPermutationVector;
-				auto PixelShader = View.ShaderMap->GetShader< FLocalFogVolumePS >(PsPermutationVector);
+			FLocalFogVolumePS::FPermutationDomain PsPermutationVector;
+			auto PixelShader = View.ShaderMap->GetShader< FLocalFogVolumePS >(PsPermutationVector);
 
-				const FIntRect ViewRect = View.ViewRect;
+			const FIntRect ViewRect = View.ViewRect;
 
-				ClearUnusedGraphResources(VertexShader, &PassParameters->VS);
-				ClearUnusedGraphResources(PixelShader, &PassParameters->PS);
+			ClearUnusedGraphResources(VertexShader, &PassParameters->VS);
+			ClearUnusedGraphResources(PixelShader, &PassParameters->PS);
 
-				uint32 LocalFogVolumeGPUInstanceCount = View.LocalFogVolumeGPUInstanceCount;
-				GraphBuilder.AddPass(
-					RDG_EVENT_NAME("RenderLocalFogVolume %u inst.", LocalFogVolumeGPUInstanceCount),
-					PassParameters,
-					ERDGPassFlags::Raster,
-					[VertexShader, PixelShader, PassParameters, LocalFogVolumeGPUInstanceCount, ViewRect](FRHICommandList& RHICmdList)
-				{
-					FGraphicsPipelineStateInitializer GraphicsPSOInit;
-					RHICmdList.ApplyCachedRenderTargets(GraphicsPSOInit);
+			uint32 LocalFogVolumeGPUInstanceCount = View.LocalFogVolumeGPUInstanceCount;
+			GraphBuilder.AddPass(
+				RDG_EVENT_NAME("RenderLocalFogVolume %u inst.", LocalFogVolumeGPUInstanceCount),
+				PassParameters,
+				ERDGPassFlags::Raster,
+				[VertexShader, PixelShader, PassParameters, LocalFogVolumeGPUInstanceCount, ViewRect](FRHICommandList& RHICmdList)
+			{
+				FGraphicsPipelineStateInitializer GraphicsPSOInit;
+				RHICmdList.ApplyCachedRenderTargets(GraphicsPSOInit);
 
-					RHICmdList.SetViewport(ViewRect.Min.X, ViewRect.Min.Y, 0.0f, ViewRect.Max.X, ViewRect.Max.Y, 1.0f);
+				RHICmdList.SetViewport(ViewRect.Min.X, ViewRect.Min.Y, 0.0f, ViewRect.Max.X, ViewRect.Max.Y, 1.0f);
 
-					// Render back faces only since camera may intersect
-					GraphicsPSOInit.RasterizerState = TStaticRasterizerState<FM_Solid, CM_CW>::GetRHI();
-					GraphicsPSOInit.DepthStencilState = TStaticDepthStencilState<false, CF_Always>::GetRHI();
-					GraphicsPSOInit.BlendState = TStaticBlendState<CW_RGBA, BO_Add, BF_One, BF_SourceAlpha, BO_Add, BF_Zero, BF_SourceAlpha>::GetRHI();
-					GraphicsPSOInit.PrimitiveType = PT_TriangleList;
+				// Render back faces only since camera may intersect
+				GraphicsPSOInit.RasterizerState = TStaticRasterizerState<FM_Solid, CM_CW>::GetRHI();
+				GraphicsPSOInit.DepthStencilState = TStaticDepthStencilState<false, CF_Always>::GetRHI();
+				GraphicsPSOInit.BlendState = TStaticBlendState<CW_RGBA, BO_Add, BF_One, BF_SourceAlpha, BO_Add, BF_Zero, BF_SourceAlpha>::GetRHI();
+				GraphicsPSOInit.PrimitiveType = PT_TriangleList;
 
-					GraphicsPSOInit.BoundShaderState.VertexDeclarationRHI = GetVertexDeclarationFVector4();
-					GraphicsPSOInit.BoundShaderState.VertexShaderRHI = VertexShader.GetVertexShader();
-					GraphicsPSOInit.BoundShaderState.PixelShaderRHI = PixelShader.GetPixelShader();
+				GraphicsPSOInit.BoundShaderState.VertexDeclarationRHI = GetVertexDeclarationFVector4();
+				GraphicsPSOInit.BoundShaderState.VertexShaderRHI = VertexShader.GetVertexShader();
+				GraphicsPSOInit.BoundShaderState.PixelShaderRHI = PixelShader.GetPixelShader();
 
-					SetGraphicsPipelineState(RHICmdList, GraphicsPSOInit, 0);
+				SetGraphicsPipelineState(RHICmdList, GraphicsPSOInit, 0);
 
-					SetShaderParameters(RHICmdList, VertexShader, VertexShader.GetVertexShader(), PassParameters->VS);
-					SetShaderParameters(RHICmdList, PixelShader, PixelShader.GetPixelShader(), PassParameters->PS);
+				SetShaderParameters(RHICmdList, VertexShader, VertexShader.GetVertexShader(), PassParameters->VS);
+				SetShaderParameters(RHICmdList, PixelShader, PixelShader.GetPixelShader(), PassParameters->PS);
 
-					RHICmdList.SetStreamSource(0, GetUnitCubeVertexBuffer(), 0);
+				RHICmdList.SetStreamSource(0, GetUnitCubeVertexBuffer(), 0);
 
-					RHICmdList.DrawIndexedPrimitive(GetUnitCubeIndexBuffer()
-						, 0									//BaseVertexIndex
-						, 0									//FirstInstance
-						, 8									//uint32 NumVertices
-						, 0									//uint32 StartIndex
-						, UE_ARRAY_COUNT(GCubeIndices) / 3	//uint32 NumPrimitives
-						, LocalFogVolumeGPUInstanceCount	//uint32 NumInstances
-					);
-				});
-			}
+				RHICmdList.DrawIndexedPrimitive(GetUnitCubeIndexBuffer()
+					, 0									//BaseVertexIndex
+					, 0									//FirstInstance
+					, 8									//uint32 NumVertices
+					, 0									//uint32 StartIndex
+					, UE_ARRAY_COUNT(GCubeIndices) / 3	//uint32 NumPrimitives
+					, LocalFogVolumeGPUInstanceCount	//uint32 NumInstances
+				);
+			});
 		}
 	}
 }
@@ -322,7 +353,7 @@ class FMobileLocalFogVolumeVS : public FGlobalShader
 
 	BEGIN_SHADER_PARAMETER_STRUCT(FParameters, )
 		SHADER_PARAMETER_STRUCT_INCLUDE(FViewShaderParameters, View)
-		SHADER_PARAMETER_RDG_BUFFER_SRV(Buffer<float4>, LocalFogVolumeInstances)
+		SHADER_PARAMETER_STRUCT(FLocalFogVolumeUniformParameters, LFV)
 	END_SHADER_PARAMETER_STRUCT()
 
 	using FPermutationDomain = TShaderPermutationDomain<>;
@@ -354,7 +385,7 @@ class FMobileLocalFogVolumePS : public FGlobalShader
 	BEGIN_SHADER_PARAMETER_STRUCT(FParameters, )
 		SHADER_PARAMETER_STRUCT_INCLUDE(FViewShaderParameters, View)
 		SHADER_PARAMETER_RDG_UNIFORM_BUFFER(FMobileBasePassUniformParameters, MobileBasePass)
-		SHADER_PARAMETER_RDG_BUFFER_SRV(Buffer<float4>, LocalFogVolumeInstances)
+		SHADER_PARAMETER_STRUCT(FLocalFogVolumeUniformParameters, LFV)
 	END_SHADER_PARAMETER_STRUCT()
 
 	using FPermutationDomain = TShaderPermutationDomain<>;
@@ -416,12 +447,12 @@ void RenderLocalFogVolumeMobile(
 
 	FMobileLocalFogVolumeVS::FParameters VSParameters;
 	VSParameters.View = View.GetShaderParameters();
-	VSParameters.LocalFogVolumeInstances = View.LocalFogVolumeGPUInstanceDataBufferSRV;
+	VSParameters.LFV = View.LocalFogVolumeUniformParametersStruct;
 	SetShaderParameters(RHICmdList, VertexShader, VertexShader.GetVertexShader(), VSParameters);
 
 	FMobileLocalFogVolumePS::FParameters PSParameters;
 	PSParameters.View = View.GetShaderParameters();
-	PSParameters.LocalFogVolumeInstances = View.LocalFogVolumeGPUInstanceDataBufferSRV;
+	PSParameters.LFV = View.LocalFogVolumeUniformParametersStruct;
 	// PSParameters.MobileBasePass filled up by the RDG pass parameters.
 	SetShaderParameters(RHICmdList, PixelShader, PixelShader.GetPixelShader(), PSParameters);
 
diff --git a/Engine/Source/Runtime/Renderer/Private/LocalFogVolumeRendering.h b/Engine/Source/Runtime/Renderer/Private/LocalFogVolumeRendering.h
--- a/Engine/Source/Runtime/Renderer/Private/LocalFogVolumeRendering.h
+++ b/Engine/Source/Runtime/Renderer/Private/LocalFogVolumeRendering.h
@@ -15,6 +15,15 @@ struct FMinimalSceneTextures;
 
 
 
+/*=============================================================================
+	Local height fog rendering GPU data
+=============================================================================*/
+
+BEGIN_SHADER_PARAMETER_STRUCT(FLocalFogVolumeUniformParameters, )
+	SHADER_PARAMETER_RDG_BUFFER_SRV(Buffer<float4>, LocalFogVolumeInstances)
+	SHADER_PARAMETER(uint32, LocalFogVolumeInstanceCount)
+END_SHADER_PARAMETER_STRUCT()
+
 /*=============================================================================
 	Local height fog rendering common data
 =============================================================================*/
@@ -86,6 +95,11 @@ void GetLocalFogVolumeSortingData(const FScene* Scene, FRDGBuilder& GraphBuilder
 
 void CreateViewLocalFogVolumeBufferSRV(FViewInfo& View, FRDGBuilder& GraphBuilder, FLocalFogVolumeSortingData& SortingData);
 
+void InitLocalFogVolumesForViews(
+	const FScene* Scene,
+	TArray<FViewInfo>& Views,
+	FRDGBuilder& GraphBuilder);
+
 void RenderLocalFogVolume(
 	const FScene* Scene,
 	TArray<FViewInfo>& Views,
diff --git a/Engine/Source/Runtime/Renderer/Private/SceneRendering.cpp b/Engine/Source/Runtime/Renderer/Private/SceneRendering.cpp
--- a/Engine/Source/Runtime/Renderer/Private/SceneRendering.cpp
+++ b/Engine/Source/Runtime/Renderer/Private/SceneRendering.cpp
@@ -976,7 +976,9 @@ void FViewInfo::Init()
 	HairStrandsViewData = FHairStrandsViewData();
 
 	LocalFogVolumeGPUInstanceCount = 0;
+	LocalFogVolumeGPUInstanceDataBuffer = nullptr;
 	LocalFogVolumeGPUInstanceDataBufferSRV = nullptr;
+	LocalFogVolumeUniformBuffer = nullptr;
 
 	GPUSceneViewId = INDEX_NONE;
 
diff --git a/Engine/Source/Runtime/Renderer/Private/SceneRendering.h b/Engine/Source/Runtime/Renderer/Private/SceneRendering.h
--- a/Engine/Source/Runtime/Renderer/Private/SceneRendering.h
+++ b/Engine/Source/Runtime/Renderer/Private/SceneRendering.h
@@ -46,6 +46,7 @@
 #include "TextureFallbacks.h"
 #include "SceneInterface.h"
 #include "Async/Mutex.h"
+#include "LocalFogVolumeRendering.h"
 
 #if RHI_RAYTRACING
 #include "RayTracingInstanceBufferUtil.h"
@@ -1368,7 +1369,10 @@ public:
 	FSubstrateViewData SubstrateViewData;
 
 	uint32 LocalFogVolumeGPUInstanceCount;
+	FRDGBufferRef LocalFogVolumeGPUInstanceDataBuffer;
 	FRDGBufferSRVRef LocalFogVolumeGPUInstanceDataBufferSRV;
+	FLocalFogVolumeUniformParameters LocalFogVolumeUniformParametersStruct;
+	TRDGUniformBufferRef<FLocalFogVolumeUniformParameters> LocalFogVolumeUniformBuffer;
 
 	FHairStrandsViewData HairStrandsViewData;
 
