Fixed potential div by 0 crash in volumetric fog.

diff --git a/Engine/Source/Runtime/Renderer/Private/VolumetricFog.cpp b/Engine/Source/Runtime/Renderer/Private/VolumetricFog.cpp
--- a/Engine/Source/Runtime/Renderer/Private/VolumetricFog.cpp
+++ b/Engine/Source/Runtime/Renderer/Private/VolumetricFog.cpp
@@ -162,7 +162,10 @@ static int32 GetVolumetricFogGridSizeZ()
 static FIntPoint GetVolumetricFogTextureResourceRes(const FViewInfo& View)
 {
 	// Allocate texture using scene render targets size so we do not reallocate every frame when dynamic resolution is used in order to avoid resources allocation hitches.
-	const FIntPoint BufferSize = View.GetSceneTexturesConfig().Extent;
+	FIntPoint BufferSize = View.GetSceneTexturesConfig().Extent;
+	// Make sure the buffer size has some minimum resolution to make sure everything is always valid.
+	BufferSize.X = FMath::Max(1, BufferSize.X);
+	BufferSize.Y = FMath::Max(1, BufferSize.Y);
 	return BufferSize;
 }
 
