LFV - fixed potential crash

#prelfight shader

diff --git a/Engine/Source/Runtime/Renderer/Private/LocalFogVolumeRendering.cpp b/Engine/Source/Runtime/Renderer/Private/LocalFogVolumeRendering.cpp
--- a/Engine/Source/Runtime/Renderer/Private/LocalFogVolumeRendering.cpp
+++ b/Engine/Source/Runtime/Renderer/Private/LocalFogVolumeRendering.cpp
@@ -41,6 +41,10 @@ static void SetDummyLocalFogVolumeForView(FRDGBuilder& GraphBuilder, FViewInfo&
 		FRDGBufferDesc::CreateBufferDesc(SizeOfUintVec4, UintVec4CountInLocalFogVolumeGPUInstanceData), &DummyData, sizeof(FLocalFogVolumeGPUInstanceData) * 1, ERDGInitialDataFlags::NoCopy);
 
 	View.LocalFogVolumeGPUInstanceDataBufferSRV = GraphBuilder.CreateSRV(View.LocalFogVolumeGPUInstanceDataBuffer, PF_A32B32G32R32F);
+
+	View.LocalFogVolumeUniformParametersStruct.LocalFogVolumeInstanceCount = View.LocalFogVolumeGPUInstanceCount;
+	View.LocalFogVolumeUniformParametersStruct.LocalFogVolumeInstances = View.LocalFogVolumeGPUInstanceDataBufferSRV;
+	View.LocalFogVolumeUniformBuffer = GraphBuilder.CreateUniformBuffer(&View.LocalFogVolumeUniformParametersStruct);
 };
 
 /*=============================================================================
