LFV - fixed bug when trying to allocate a buffer with a size of 0.

#prelfight shader

diff --git a/Engine/Source/Runtime/Renderer/Private/LocalFogVolumeRendering.cpp b/Engine/Source/Runtime/Renderer/Private/LocalFogVolumeRendering.cpp
--- a/Engine/Source/Runtime/Renderer/Private/LocalFogVolumeRendering.cpp
+++ b/Engine/Source/Runtime/Renderer/Private/LocalFogVolumeRendering.cpp
@@ -138,6 +138,7 @@ void CreateViewLocalFogVolumeBufferSRV(FViewInfo& View, FRDGBuilder& GraphBuilde
 	if (SortingData.LocalFogVolumeInstanceCountFinal == 0)
 	{
 		SetDummyLocalFogVolumeForView(GraphBuilder, View);
+		return;
 	}
 
 	// 1. Sort all the volumes
