LFV - sampling on straightforward on translucent.

#prelfight https://horde.devtools.epicgames.com/job/64f9c66e190d96dd29c652fa

diff --git a/Engine/Shaders/Private/BasePassCommon.ush b/Engine/Shaders/Private/BasePassCommon.ush
--- a/Engine/Shaders/Private/BasePassCommon.ush
+++ b/Engine/Shaders/Private/BasePassCommon.ush
@@ -10,11 +10,13 @@
 	#define ReflectionStruct TranslucentBasePass.Shared.Reflection
 	#define PlanarReflectionStruct TranslucentBasePass.Shared.PlanarReflection
 	#define FogStruct TranslucentBasePass.Shared.Fog
+	#define LFVStruct TranslucentBasePass.Shared
 #else
 	#define ForwardLightData OpaqueBasePass.Shared.Forward
 	#define ReflectionStruct OpaqueBasePass.Shared.Reflection
 	#define PlanarReflectionStruct OpaqueBasePass.Shared.PlanarReflection
 	#define FogStruct OpaqueBasePass.Shared.Fog
+	#define LFVStruct OpaqueBasePass.Shared
 #endif 
 
 #undef NEEDS_LIGHTMAP_COORDINATE
@@ -64,10 +66,14 @@
 #define NEEDS_BASEPASS_CLOUD_SHADOW_INTERPOLATOR (SUPPORT_CLOUD_SHADOW_ON_FORWARD_LIT_TRANSLUCENT && IS_MATERIAL_TRANSLUCENT_AND_LIT && !TRANSLUCENCY_PERVERTEX_FORWARD_SHADING && !TRANSLUCENCY_PERVERTEX_LIGHTING_VOLUME)
 #endif
 
+// Local fog volumes are always evaluated per vertex for now, even with forward shading.
+#define LOCAL_FOG_VOLUME_ON_TRANSLUCENT (PROJECT_LOCALFOGVOLUME_APPLYONTRANSLUCENT && (NEEDS_BASEPASS_VERTEX_FOGGING || NEEDS_BASEPASS_PIXEL_FOGGING))
+
 struct FSharedBasePassInterpolants
 {
-	//for texture-lightmapped translucency we can pass the vertex fog in its own interpolator
-#if NEEDS_BASEPASS_VERTEX_FOGGING
+	// This is enabled for vertex fogging or for local fog volume that are better evaluated per vertex for now.
+	// For texture-lightmapped translucency we can pass the vertex fog in its own interpolator
+#if NEEDS_BASEPASS_VERTEX_FOGGING || LOCAL_FOG_VOLUME_ON_TRANSLUCENT
 	float4 VertexFog		: TEXCOORD7;
 #endif
 
diff --git a/Engine/Shaders/Private/BasePassPixelShader.usf b/Engine/Shaders/Private/BasePassPixelShader.usf
--- a/Engine/Shaders/Private/BasePassPixelShader.usf
+++ b/Engine/Shaders/Private/BasePassPixelShader.usf
@@ -1293,6 +1293,10 @@ void FPixelShaderInOut_MainPS(
 		float4 HeightFogging = BasePassInterpolants.VertexFog;
 	#elif NEEDS_BASEPASS_PIXEL_FOGGING
 		float4 HeightFogging = CalculateHeightFog(MaterialParameters.WorldPosition_CamRelative);
+		#if LOCAL_FOG_VOLUME_ON_TRANSLUCENT
+		const float4 LocalFogVolumeContrib = BasePassInterpolants.VertexFog;
+		HeightFogging = float4(LocalFogVolumeContrib.rgb + HeightFogging.rgb * LocalFogVolumeContrib.a, LocalFogVolumeContrib.a * HeightFogging.a);
+		#endif // LOCAL_FOG_VOLUME_ON_TRANSLUCENT
 	#else
 		float4 HeightFogging = float4(0,0,0,1);
 	#endif
diff --git a/Engine/Shaders/Private/BasePassVertexShader.usf b/Engine/Shaders/Private/BasePassVertexShader.usf
--- a/Engine/Shaders/Private/BasePassVertexShader.usf
+++ b/Engine/Shaders/Private/BasePassVertexShader.usf
@@ -24,6 +24,10 @@
 #include "VolumetricCloudCommon.ush"
 #endif
 
+#if LOCAL_FOG_VOLUME_ON_TRANSLUCENT
+#include "LocalFogVolumeCommon.ush"
+#endif
+
 /** Entry point for the base pass vertex shader. */
 void Main(
 	FVertexFactoryInput Input,
@@ -100,6 +104,35 @@ void Main(
 			OneOverPreExposure, Output.BasePassInterpolants.VertexFog);
 	}
 #endif
+#endif // NEEDS_BASEPASS_VERTEX_FOGGING
+
+	// Apply local fog volume to if enabled. It is always enabled per vertex as of today, that is why it is out ofr the NEEDS_BASEPASS_VERTEX_FOGGING defines.
+#if LOCAL_FOG_VOLUME_ON_TRANSLUCENT
+	float4 LFVContribution = float4(0.0f, 0.0f, 0.0f, 1.0f);
+	for (uint LVFIndex = 0; LVFIndex < LFVInstcount; ++LVFIndex)
+	{
+		float3 CamRayTranslatedWorldOrigin	= ResolvedView.TranslatedWorldCameraOrigin;
+		float3 CamRayWorldOrigin			= CamRayTranslatedWorldOrigin - LWCHackToFloat(ResolvedView.PreViewTranslation);
+		float3 CamRayTranslatedWorldDir		= normalize(WorldPosition - CamRayTranslatedWorldOrigin);	// WorldPosition is expressed in TranslatedWorld space
+		float3 CamRayWorldDir				= CamRayTranslatedWorldDir;
+
+		float3 DepthBufferTranslatedWorldPos = WorldPosition.xyz;
+
+		const float4 LFVInstanceContribution = GetLFVInstanceContribution(
+			LVFIndex,
+			CamRayWorldOrigin,
+			CamRayWorldDir,
+			CamRayTranslatedWorldOrigin,
+			CamRayTranslatedWorldDir,
+			DepthBufferTranslatedWorldPos);
+
+		LFVContribution = float4(LFVContribution.rgb * LFVInstanceContribution.aaa + LFVInstanceContribution.rgb, LFVContribution.a * LFVInstanceContribution.a);
+	}
+
+	Output.BasePassInterpolants.VertexFog = float4(LFVContribution.rgb + Output.BasePassInterpolants.VertexFog.rgb * LFVContribution.a, LFVContribution.a * Output.BasePassInterpolants.VertexFog.a);
+#endif // LOCAL_FOG_VOLUME_ON_TRANSLUCENT
+
+#if NEEDS_BASEPASS_VERTEX_FOGGING
 
 #if MATERIAL_ENABLE_TRANSLUCENCY_CLOUD_FOGGING
 	if (TranslucentBasePass.ApplyVolumetricCloudOnTransparent > 0.0f)
diff --git a/Engine/Shaders/Private/LocalFogVolumeCommon.ush b/Engine/Shaders/Private/LocalFogVolumeCommon.ush
--- a/Engine/Shaders/Private/LocalFogVolumeCommon.ush
+++ b/Engine/Shaders/Private/LocalFogVolumeCommon.ush
@@ -32,11 +32,22 @@ struct FLocalFogVolumeGPUInstanceData
 
 
 // Structured buffers cannot be sampled from the vertex shader on some mobile platforms.
-// That is why simple Buffers are used, because this tech needs to run on mobile platforms.
+// That is why simple Buffers are used for LocalFogVolumeInstances, because this tech needs to run on mobile platforms.
+
+#if defined(LFVStruct)
+
+#define LFVInsts LFVStruct.LFV.LocalFogVolumeInstances
+#define LFVInstcount LFVStruct.LFV.LocalFogVolumeInstanceCount
+
+#else
 
 uint LFV_LocalFogVolumeInstanceCount;
 Buffer<float4> LFV_LocalFogVolumeInstances;
+
 #define LFVInsts LFV_LocalFogVolumeInstances
+#define LFVInstcount LFV_LocalFogVolumeInstanceCount
+
+#endif
 
 
 
diff --git a/Engine/Shaders/Private/MobileBasePassCommon.ush b/Engine/Shaders/Private/MobileBasePassCommon.ush
--- a/Engine/Shaders/Private/MobileBasePassCommon.ush
+++ b/Engine/Shaders/Private/MobileBasePassCommon.ush
@@ -7,6 +7,7 @@
 #include "BasePassCommon.ush"
 #define FogStruct MobileBasePass.Fog
 #define PlanarReflectionStruct MobileBasePass.PlanarReflection
+#define LFVStruct MobileBasePass
 
 #undef NEEDS_LIGHTMAP_COORDINATE
 
@@ -21,6 +22,9 @@
 // And always use vertex fogging on sky for perf
 #define USE_VERTEX_FOG	(!PROJECT_MOBILE_DISABLE_VERTEX_FOG || MATERIAL_IS_SKY || (MATERIAL_ENABLE_TRANSLUCENCY_FOGGING && (MATERIALBLENDING_ANY_TRANSLUCENT || MATERIAL_SHADINGMODEL_SINGLELAYERWATER)))
 
+// Local fog volumes are always evaluated per vertex for now, on mobile, it simply follow VertexFog.
+#define LOCAL_FOG_VOLUME_ON_TRANSLUCENT (PROJECT_LOCALFOGVOLUME_APPLYONTRANSLUCENT && USE_VERTEX_FOG)
+
 #define PACK_INTERPOLANTS (USE_VERTEX_FOG && NUM_VF_PACKED_INTERPOLANTS > 0 && (ES3_1_PROFILE))
 #define LANDSCAPE_BUG_WORKAROUND (IOS && IS_MOBILE_BASEPASS_VERTEX_SHADER && PACK_INTERPOLANTS)
 
diff --git a/Engine/Shaders/Private/MobileBasePassVertexShader.usf b/Engine/Shaders/Private/MobileBasePassVertexShader.usf
--- a/Engine/Shaders/Private/MobileBasePassVertexShader.usf
+++ b/Engine/Shaders/Private/MobileBasePassVertexShader.usf
@@ -20,6 +20,10 @@
 #include "SkyAtmosphereCommon.ush"
 #endif
 
+#if LOCAL_FOG_VOLUME_ON_TRANSLUCENT
+#include "LocalFogVolumeCommon.ush"
+#endif
+
 struct FMobileShadingBasePassVSToPS
 {
 	FVertexFactoryInterpolantsVSToPS FactoryInterpolants;
@@ -105,12 +109,38 @@ void Main(
 	}
 #endif
 
+	// Apply local fog volume to if enabled
+#if LOCAL_FOG_VOLUME_ON_TRANSLUCENT
+	float4 LFVContribution = float4(0.0f, 0.0f, 0.0f, 1.0f);
+	for (uint LVFIndex = 0; LVFIndex < LFVInstcount; ++LVFIndex)
+	{
+		float3 CamRayTranslatedWorldOrigin = ResolvedView.TranslatedWorldCameraOrigin;
+		float3 CamRayWorldOrigin = CamRayTranslatedWorldOrigin - LWCHackToFloat(ResolvedView.PreViewTranslation);
+		float3 CamRayTranslatedWorldDir = normalize(WorldPosition - CamRayTranslatedWorldOrigin);	// WorldPosition is expressed in TranslatedWorld space
+		float3 CamRayWorldDir = CamRayTranslatedWorldDir;
+
+		float3 DepthBufferTranslatedWorldPos = WorldPosition.xyz;
+
+		const float4 LFVInstanceContribution = GetLFVInstanceContribution(
+			LVFIndex,
+			CamRayWorldOrigin,
+			CamRayWorldDir,
+			CamRayTranslatedWorldOrigin,
+			CamRayTranslatedWorldDir,
+			DepthBufferTranslatedWorldPos);
+
+		LFVContribution = float4(LFVContribution.rgb * LFVInstanceContribution.aaa + LFVInstanceContribution.rgb, LFVContribution.a * LFVInstanceContribution.a);
+	}
+
+	VertexFog = float4(LFVContribution.rgb + VertexFog.rgb * LFVContribution.a, LFVContribution.a * VertexFog.a);
+#endif // LOCAL_FOG_VOLUME_ON_TRANSLUCENT
+
 #if PACK_INTERPOLANTS
 	PackedInterps[0] = VertexFog;
 #else
 	Output.BasePassInterpolants.VertexFog = VertexFog;
 #endif // PACK_INTERPOLANTS
-#endif
+#endif // USE_VERTEX_FOG
 
 #if LANDSCAPE_BUG_WORKAROUND
 	Output.BasePassInterpolants.DummyInterp = 0;
diff --git a/Engine/Shaders/Public/ShaderVersion.ush b/Engine/Shaders/Public/ShaderVersion.ush
--- a/Engine/Shaders/Public/ShaderVersion.ush
+++ b/Engine/Shaders/Public/ShaderVersion.ush
@@ -2,4 +2,4 @@
 // Each time the console command is executed it generates a new GUID. As this file is included
 // in Platform.ush (which should be included in any shader) it allows to invalidate the shader DDC.
 // 
-#pragma message("UESHADERMETADATA_VERSION 6C9D3414-5F37-4FF8-8160-6542407A4DCE")
\ No newline at end of file
+#pragma message("UESHADERMETADATA_VERSION E7D109A8-D147-450F-A96D-9A3E24EB1B05")
diff --git a/Engine/Source/Runtime/Engine/Classes/Engine/RendererSettings.h b/Engine/Source/Runtime/Engine/Classes/Engine/RendererSettings.h
--- a/Engine/Source/Runtime/Engine/Classes/Engine/RendererSettings.h
+++ b/Engine/Source/Runtime/Engine/Classes/Engine/RendererSettings.h
@@ -603,6 +603,11 @@ class URendererSettings : public UDeveloperSettings
 		ToolTip="The axis that sorting will occur along when Translucent Sort Policy is set to SortAlongAxis."))
 	FVector TranslucentSortAxis;
 
+	UPROPERTY(config, EditAnywhere, Category=Translucency, meta=(
+		ConsoleVariable="r.LocalFogVolume.ApplyOnTransclucent",
+		ToolTip="Allow translucency to be rendered to a separate render targeted and composited after depth of field. Prevents translucency from appearing out of focus."))
+	uint32 bLocalFogVolumeApplyOnTranslucent:1;
+
 	UPROPERTY(config, EditAnywhere, Category = VR, meta = (
 		ConsoleVariable = "xr.VRS.FoveationLevel", DisplayName = "Stereo Foveation Level (Experimental)",
 		ToolTip = "Set the level of foveation to apply when generating the Variable Rate Shading attachment. This feature is currently experimental.\nThis can yield some fairly significant performance benefits on GPUs that support Tier 2 VRS.\nLower settings will result in almost no discernible artifacting on most HMDs; higher settings will show some artifacts towards the edges of the view."))
diff --git a/Engine/Source/Runtime/Engine/Private/ShaderCompiler/ShaderCompiler.cpp b/Engine/Source/Runtime/Engine/Private/ShaderCompiler/ShaderCompiler.cpp
--- a/Engine/Source/Runtime/Engine/Private/ShaderCompiler/ShaderCompiler.cpp
+++ b/Engine/Source/Runtime/Engine/Private/ShaderCompiler/ShaderCompiler.cpp
@@ -7908,6 +7908,12 @@ void GlobalBeginCompileShader(
 		SET_SHADER_DEFINE(Input.Environment, PROJECT_MOBILE_DISABLE_VERTEX_FOG, CVar ? (CVar->GetInt() != 0) : 0);
 	}
 
+	{
+		static IConsoleVariable* CVar = IConsoleManager::Get().FindConsoleVariable(TEXT("r.LocalFogVolume.ApplyOnTransclucent"));
+		const bool bLocalFogVolumesApplyOnTranclucent = CVar && CVar->GetInt() > 0;
+		SET_SHADER_DEFINE(Input.Environment, PROJECT_LOCALFOGVOLUME_APPLYONTRANSLUCENT, (bLocalFogVolumesApplyOnTranclucent ? 1 : 0));
+	}
+
 	bool bSupportSkyAtmosphere = false;
 	{
 		static IConsoleVariable* CVar = IConsoleManager::Get().FindConsoleVariable(TEXT("r.SupportSkyAtmosphere"));
diff --git a/Engine/Source/Runtime/RenderCore/Private/Shader.cpp b/Engine/Source/Runtime/RenderCore/Private/Shader.cpp
--- a/Engine/Source/Runtime/RenderCore/Private/Shader.cpp
+++ b/Engine/Source/Runtime/RenderCore/Private/Shader.cpp
@@ -1935,6 +1935,15 @@ void ShaderMapAppendKeyString(EShaderPlatform Platform, FString& KeyString)
 		}
 	}
 
+	{
+		static IConsoleVariable* CVar = IConsoleManager::Get().FindConsoleVariable(TEXT("r.LocalFogVolume.ApplyOnTransclucent"));
+		const bool bLocalFogVolumesApplyOnTranclucent = CVar && CVar->GetInt() > 0;
+		if (bLocalFogVolumesApplyOnTranclucent)
+		{
+			KeyString += TEXT("_LFVTRA");
+		}
+	}
+
 	{
 		static const auto CVar = IConsoleManager::Get().FindTConsoleVariableDataInt(TEXT("r.SupportSkyAtmosphere"));
 		if (CVar && CVar->GetValueOnAnyThread() > 0)
diff --git a/Engine/Source/Runtime/Renderer/Private/BasePassRendering.cpp b/Engine/Source/Runtime/Renderer/Private/BasePassRendering.cpp
--- a/Engine/Source/Runtime/Renderer/Private/BasePassRendering.cpp
+++ b/Engine/Source/Runtime/Renderer/Private/BasePassRendering.cpp
@@ -719,6 +719,8 @@ void SetupSharedBasePassParameters(
 		SharedParameters.FogISR = SharedParameters.Fog;
 	}
 
+	SharedParameters.LFV = View.LocalFogVolumeUniformParametersStruct;
+
 	const FScene* Scene = View.Family->Scene ? View.Family->Scene->GetRenderScene() : nullptr;
 	const FPlanarReflectionSceneProxy* ReflectionSceneProxy = Scene ? Scene->GetForwardPassGlobalPlanarReflection() : nullptr;
 
diff --git a/Engine/Source/Runtime/Renderer/Private/BasePassRendering.h b/Engine/Source/Runtime/Renderer/Private/BasePassRendering.h
--- a/Engine/Source/Runtime/Renderer/Private/BasePassRendering.h
+++ b/Engine/Source/Runtime/Renderer/Private/BasePassRendering.h
@@ -33,6 +33,7 @@
 #include "VolumetricCloudRendering.h"
 #include "Nanite/NaniteMaterials.h"
 #include "BlueNoise.h"
+#include "LocalFogVolumeRendering.h"
 
 class FScene;
 
@@ -72,6 +73,7 @@ BEGIN_GLOBAL_SHADER_PARAMETER_STRUCT(FSharedBasePassUniformParameters,)
 	SHADER_PARAMETER_STRUCT(FPlanarReflectionUniformParameters, PlanarReflection) // Single global planar reflection for the forward pass.
 	SHADER_PARAMETER_STRUCT(FFogUniformParameters, Fog)
 	SHADER_PARAMETER_STRUCT(FFogUniformParameters, FogISR)
+	SHADER_PARAMETER_STRUCT(FLocalFogVolumeUniformParameters, LFV)
 	SHADER_PARAMETER(uint32, UseBasePassSkylight)
 END_GLOBAL_SHADER_PARAMETER_STRUCT()
 
diff --git a/Engine/Source/Runtime/Renderer/Private/DeferredShadingRenderer.cpp b/Engine/Source/Runtime/Renderer/Private/DeferredShadingRenderer.cpp
--- a/Engine/Source/Runtime/Renderer/Private/DeferredShadingRenderer.cpp
+++ b/Engine/Source/Runtime/Renderer/Private/DeferredShadingRenderer.cpp
@@ -3432,7 +3432,7 @@ void FDeferredShadingSceneRenderer::Render(FRDGBuilder& GraphBuilder)
 		InitVolumetricCloudsForViews(GraphBuilder, bShouldRenderVolumetricCloudBase, InstanceCullingManager);
 
 		// Run local fog volume initialisation before basepass for when data is needed in forward
-		InitLocalFogVolumesForViews(Scene, Views, GraphBuilder);
+		InitLocalFogVolumesForViews(Scene, Views, ViewFamily, GraphBuilder);
 
 		if (SkyAtmospherePassLocation == ESkyAtmospherePassLocation::BeforeOcclusion && bShouldRenderSkyAtmosphere)
 		{
@@ -4129,7 +4129,7 @@ void FDeferredShadingSceneRenderer::Render(FRDGBuilder& GraphBuilder)
 			RDG_CSV_STAT_EXCLUSIVE_SCOPE(GraphBuilder, RenderLocalFogVolume);
 			SCOPED_NAMED_EVENT(RenderLocalFogVolume, FColor::Emerald);
 			SCOPE_CYCLE_COUNTER(STAT_FDeferredShadingSceneRenderer_RenderLocalFogVolume);
-			RenderLocalFogVolume(Scene, Views, GraphBuilder, SceneTextures, LightShaftOcclusionTexture);
+			RenderLocalFogVolume(Scene, Views, ViewFamily, GraphBuilder, SceneTextures, LightShaftOcclusionTexture);
 		}
 
 		FRDGTextureRef ExposureIlluminance = AddCalculateExposureIlluminancePass(GraphBuilder, Views, SceneTextures, TranslucencyLightingVolumeTextures, ExposureIlluminanceSetup);
diff --git a/Engine/Source/Runtime/Renderer/Private/LocalFogVolumeRendering.cpp b/Engine/Source/Runtime/Renderer/Private/LocalFogVolumeRendering.cpp
--- a/Engine/Source/Runtime/Renderer/Private/LocalFogVolumeRendering.cpp
+++ b/Engine/Source/Runtime/Renderer/Private/LocalFogVolumeRendering.cpp
@@ -14,6 +14,11 @@ static TAutoConsoleVariable<int32> CVarLocalFogVolume(
 	TEXT("LocalFogVolume components are rendered when this is not 0, otherwise ignored.\n"),
 	ECVF_RenderThreadSafe);
 
+static TAutoConsoleVariable<int32> CVarLocalFogVolumeApplyOnTransclucent(
+	TEXT("r.LocalFogVolume.ApplyOnTransclucent"), 0,
+	TEXT("Project settings enabling the sampling of local fog volumes on translucent elements.\n"),
+	ECVF_ReadOnly | ECVF_RenderThreadSafe);
+
 bool ShouldRenderLocalFogVolume(const FScene* Scene, const FSceneViewFamily& Family)
 {
 	const FEngineShowFlags EngineShowFlags = Family.EngineShowFlags;
@@ -29,7 +34,7 @@ DECLARE_GPU_STAT(LocalFogVolumeVolumes);
 static const uint32 SizeOfUintVec4 = sizeof(FUintVector4);
 static const uint32 UintVec4CountInLocalFogVolumeGPUInstanceData = sizeof(FLocalFogVolumeGPUInstanceData) / SizeOfUintVec4;
 
-static void SetDummyLocalFogVolumeForView(FRDGBuilder& GraphBuilder, FViewInfo& View)
+void SetDummyLocalFogVolumeForView(FRDGBuilder& GraphBuilder, FViewInfo& View)
 {
 	// The size of the structure must be a multiple of FUintVector4.
 	static_assert(sizeof(FLocalFogVolumeGPUInstanceData) == UintVec4CountInLocalFogVolumeGPUInstanceData * SizeOfUintVec4);
@@ -177,10 +182,11 @@ void CreateViewLocalFogVolumeBufferSRV(FViewInfo& View, FRDGBuilder& GraphBuilde
 void InitLocalFogVolumesForViews(
 	const FScene* Scene,
 	TArray<FViewInfo>& Views,
+	const FSceneViewFamily& Family,
 	FRDGBuilder& GraphBuilder)
 {
 	const uint32 LocalFogVolumeInstanceCount = Scene->LocalFogVolumes.Num();
-	if (LocalFogVolumeInstanceCount > 0)
+	if (LocalFogVolumeInstanceCount > 0 && ShouldRenderLocalFogVolume(Scene, Family))
 	{
 		RDG_GPU_STAT_SCOPE(GraphBuilder, LocalFogVolumeVolumes);
 
@@ -264,12 +270,13 @@ END_SHADER_PARAMETER_STRUCT()
 void RenderLocalFogVolume(
 	const FScene* Scene,
 	TArray<FViewInfo>& Views,
+	const FSceneViewFamily& Family,
 	FRDGBuilder& GraphBuilder,
 	const FMinimalSceneTextures& SceneTextures,
 	FRDGTextureRef LightShaftOcclusionTexture)
 {
 	uint32 LocalFogVolumeInstanceCount = Scene->LocalFogVolumes.Num();
-	if (LocalFogVolumeInstanceCount > 0)
+	if (LocalFogVolumeInstanceCount > 0 && ShouldRenderLocalFogVolume(Scene, Family))
 	{
 		RDG_GPU_STAT_SCOPE(GraphBuilder, LocalFogVolumeVolumes);
 
diff --git a/Engine/Source/Runtime/Renderer/Private/LocalFogVolumeRendering.h b/Engine/Source/Runtime/Renderer/Private/LocalFogVolumeRendering.h
--- a/Engine/Source/Runtime/Renderer/Private/LocalFogVolumeRendering.h
+++ b/Engine/Source/Runtime/Renderer/Private/LocalFogVolumeRendering.h
@@ -94,15 +94,18 @@ bool ShouldRenderLocalFogVolume(const FScene* Scene, const FSceneViewFamily& Fam
 void GetLocalFogVolumeSortingData(const FScene* Scene, FRDGBuilder& GraphBuilder, FLocalFogVolumeSortingData& Out);
 
 void CreateViewLocalFogVolumeBufferSRV(FViewInfo& View, FRDGBuilder& GraphBuilder, FLocalFogVolumeSortingData& SortingData);
+void SetDummyLocalFogVolumeForView(FRDGBuilder& GraphBuilder, FViewInfo& View);
 
 void InitLocalFogVolumesForViews(
 	const FScene* Scene,
 	TArray<FViewInfo>& Views,
+	const FSceneViewFamily& Family,
 	FRDGBuilder& GraphBuilder);
 
 void RenderLocalFogVolume(
 	const FScene* Scene,
 	TArray<FViewInfo>& Views,
+	const FSceneViewFamily& Family,
 	FRDGBuilder& GraphBuilder,
 	const FMinimalSceneTextures& SceneTextures,
 	FRDGTextureRef LightShaftOcclusionTexture);
diff --git a/Engine/Source/Runtime/Renderer/Private/MobileBasePassRendering.cpp b/Engine/Source/Runtime/Renderer/Private/MobileBasePassRendering.cpp
--- a/Engine/Source/Runtime/Renderer/Private/MobileBasePassRendering.cpp
+++ b/Engine/Source/Runtime/Renderer/Private/MobileBasePassRendering.cpp
@@ -247,6 +247,8 @@ void SetupMobileBasePassUniformParameters(
 	// Substrate
 	Substrate::BindSubstrateMobileForwardPasslUniformParameters(GraphBuilder, View, BasePassParameters.Substrate);
 
+	BasePassParameters.LFV = View.LocalFogVolumeUniformParametersStruct;
+
 	SetupReflectionUniformParameters(GraphBuilder, View, BasePassParameters.ReflectionsParameters);
 }
 
diff --git a/Engine/Source/Runtime/Renderer/Private/MobileBasePassRendering.h b/Engine/Source/Runtime/Renderer/Private/MobileBasePassRendering.h
--- a/Engine/Source/Runtime/Renderer/Private/MobileBasePassRendering.h
+++ b/Engine/Source/Runtime/Renderer/Private/MobileBasePassRendering.h
@@ -25,6 +25,7 @@
 #include "SkyAtmosphereRendering.h"
 #include "RenderUtils.h"
 #include "DebugViewModeRendering.h"
+#include "LocalFogVolumeRendering.h"
 
 bool MobileLocalLightsBufferEnabled(const FStaticShaderPlatform Platform);
 bool MobileLocalLightsBufferPrepassEnabled(const FStaticShaderPlatform Platform);
@@ -38,6 +39,7 @@ struct FMobileBasePassTextures
 BEGIN_GLOBAL_SHADER_PARAMETER_STRUCT(FMobileBasePassUniformParameters, )
 	SHADER_PARAMETER(float, AmbientOcclusionStaticFraction)
 	SHADER_PARAMETER_STRUCT(FFogUniformParameters, Fog)
+	SHADER_PARAMETER_STRUCT(FLocalFogVolumeUniformParameters, LFV)
 	SHADER_PARAMETER_STRUCT(FForwardLightData, Forward)
 	SHADER_PARAMETER_STRUCT(FForwardLightData, ForwardMMV)
 	SHADER_PARAMETER_STRUCT(FPlanarReflectionUniformParameters, PlanarReflection) // Single global planar reflection for the forward pass.
diff --git a/Engine/Source/Runtime/Renderer/Private/Renderer.cpp b/Engine/Source/Runtime/Renderer/Private/Renderer.cpp
--- a/Engine/Source/Runtime/Renderer/Private/Renderer.cpp
+++ b/Engine/Source/Runtime/Renderer/Private/Renderer.cpp
@@ -44,6 +44,7 @@
 #include "RendererInterface.h"
 #include "PrimitiveSceneShaderData.h"
 #include "MeshDrawCommandStats.h"
+#include "LocalFogVolumeRendering.h"
 
 DEFINE_LOG_CATEGORY(LogRenderer);
 
@@ -307,6 +308,7 @@ void FRendererModule::DrawTileMesh(FCanvasRenderContext& RenderContext, FMeshPas
 
 		View.InitRHIResources();
 		View.ForwardLightingResources.SetUniformBuffer(CreateDummyForwardLightUniformBuffer(GraphBuilder, View.GetShaderPlatform()));
+		SetDummyLocalFogVolumeForView(GraphBuilder, View);
 
 		TUniformBufferRef<FReflectionCaptureShaderData> EmptyReflectionCaptureUniformBuffer;
 
