Quick fix for render thread check()
Needs further fix to reenable PSO collection.

diff --git a/Engine/Source/Runtime/Renderer/Private/WaterInfoTextureRendering.cpp b/Engine/Source/Runtime/Renderer/Private/WaterInfoTextureRendering.cpp
--- a/Engine/Source/Runtime/Renderer/Private/WaterInfoTextureRendering.cpp
+++ b/Engine/Source/Runtime/Renderer/Private/WaterInfoTextureRendering.cpp
@@ -543,10 +543,11 @@ void FWaterInfoTextureDepthPassMeshProcessor::CollectPSOInitializers(const FScen
 			&& !Material.MaterialModifiesMeshPosition_GameThread()
 			&& Material.WritesEveryPixel(false, bVFTypeSupportsNullPixelShader))
 		{
-			const FMaterialRenderProxy& DefaultProxy = *UMaterial::GetDefaultMaterial(MD_Surface)->GetRenderProxy();
-			const FMaterial& DefaultMaterial = *DefaultProxy.GetMaterialNoFallback(FeatureLevel);
+			// todo: collect here without calling render thread only code in GetMaterialNoFallback()
+			//const FMaterialRenderProxy& DefaultProxy = *UMaterial::GetDefaultMaterial(MD_Surface)->GetRenderProxy();
+			//const FMaterial& DefaultMaterial = *DefaultProxy.GetMaterialNoFallback(FeatureLevel);
 
-			CollectPSOInitializersInternal<true>(SceneTexturesConfig, VertexFactoryData, DefaultMaterial, MeshFillMode, MeshCullMode, PreCacheParams, PSOInitializers);
+			//CollectPSOInitializersInternal<true>(SceneTexturesConfig, VertexFactoryData, DefaultMaterial, MeshFillMode, MeshCullMode, PreCacheParams, PSOInitializers);
 		}
 		else
 		{
