Added new column for Perf Reports which contains the output primitive count from indirect draw which are sent to the vertex shader.

diff --git a/Engine/Source/Runtime/Engine/Public/MeshDrawCommandStatsDefines.h b/Engine/Source/Runtime/Engine/Public/MeshDrawCommandStatsDefines.h
--- a/Engine/Source/Runtime/Engine/Public/MeshDrawCommandStatsDefines.h
+++ b/Engine/Source/Runtime/Engine/Public/MeshDrawCommandStatsDefines.h
@@ -3,4 +3,4 @@
 
 #include "Misc/Build.h"
 
-#define MESH_DRAW_COMMAND_STATS (!UE_BUILD_SHIPPING && !UE_BUILD_TEST)
+#define MESH_DRAW_COMMAND_STATS (!UE_BUILD_SHIPPING)
diff --git a/Engine/Source/Runtime/Engine/Public/MeshDrawCommandStatsSettings.h b/Engine/Source/Runtime/Engine/Public/MeshDrawCommandStatsSettings.h
--- a/Engine/Source/Runtime/Engine/Public/MeshDrawCommandStatsSettings.h
+++ b/Engine/Source/Runtime/Engine/Public/MeshDrawCommandStatsSettings.h
@@ -22,6 +22,9 @@ struct FMeshDrawCommandStatsBudget
 	/** The category primitive budget. This is the maximum triangles expected, post-culling, summed across all passes. */
 	UPROPERTY(config, EditAnywhere, Category = Engine)
 	int32 PrimitiveBudget = 0;
+	/** If true, exports the category's stats to the CSVProfiler. */
+	UPROPERTY(config, EditAnywhere, Category = Engine)
+	bool bExportToCsv = false;
 };
 
 /** User settings used by the MeshDrawCommandStats system. */
diff --git a/Engine/Source/Runtime/Renderer/Private/MeshDrawCommandStats.cpp b/Engine/Source/Runtime/Renderer/Private/MeshDrawCommandStats.cpp
--- a/Engine/Source/Runtime/Renderer/Private/MeshDrawCommandStats.cpp
+++ b/Engine/Source/Runtime/Renderer/Private/MeshDrawCommandStats.cpp
@@ -4,6 +4,7 @@
 
 #include "InstanceCulling/InstanceCullingContext.h"
 #include "MeshDrawCommandStatsSettings.h"
+#include "ProfilingDebugging/CsvProfiler.h"
 #include "RenderGraph.h"
 #include "RendererModule.h"
 #include "RendererOnScreenNotification.h"
@@ -29,6 +30,8 @@ DECLARE_DWORD_COUNTER_STAT(TEXT("InstanceCulling Indirect Rendered Instances"),
 DECLARE_DWORD_COUNTER_STAT(TEXT("Custom Indirect Rendered Primitives"), STAT_Culling_CustomIndirectNumPrimitives, STATGROUP_Culling);
 DECLARE_DWORD_COUNTER_STAT(TEXT("Custom Indirect Rendered Instances"), STAT_Culling_CustomIndirectNumInstances, STATGROUP_Culling);
 
+CSV_DEFINE_CATEGORY(MeshDrawCommandStats, true);
+
 static TAutoConsoleVariable<int32> CVarShowMeshDrawCommandStats(
 	TEXT("r.MeshDrawCommands.Stats"),
 	0,
@@ -40,6 +43,13 @@ static TAutoConsoleVariable<int32> CVarShowMeshDrawCommandStats(
 	ECVF_RenderThreadSafe
 );
 
+static TAutoConsoleVariable<int32> CVarExportMeshDrawCommandStats(
+	TEXT("r.MeshDrawCommands.Export"),
+	false,
+	TEXT("Export mesh draw command stats to the CSVProfiler.\n"),
+	ECVF_RenderThreadSafe
+);
+
 FMeshDrawCommandStatsManager::FFrameData::~FFrameData()
 {
 	// Collect set of unique readback buffers for deletion (can be shared between MDCs and passes)
@@ -141,48 +151,27 @@ FMeshDrawCommandStatsManager::FMeshDrawCommandStatsManager()
 			{
 				OutMessages.Add(FCoreDelegates::EOnScreenMessageSeverity::Info, FText::FromString(FString::Printf(TEXT("MeshDrawCommandStats (Triangles / Budget - Category):"), Stats.TotalPrimitives / 1000)));
 
-				TArray<FStats::FCategoryStats> CategoryStats = Stats.CategoryStats;
-
 				// Show budgeted stats first.
 				const UMeshDrawCommandStatsSettings* Settings = GetDefault<UMeshDrawCommandStatsSettings>();
 				for (FMeshDrawCommandStatsBudget const& CategoryBudget : Settings->Budgets)
 				{
-					uint64 PrimitiveCount = 0;
-
-					for (FStats::FCategoryStats& CategoryStat : CategoryStats)
+					uint64* PrimitiveCount = BudgetedPrimitives.Find(CategoryBudget.CategoryName);
+					if (PrimitiveCount && *PrimitiveCount > 0)
 					{
-						if (CategoryBudget.CategoryName == CategoryStat.CategoryName)
-						{
-							PrimitiveCount += CategoryStat.PrimitiveCount;
-							CategoryStat.PrimitiveCount = 0;
-						}
-						else
-						{
-							for (FName Name : CategoryBudget.LinkedStatNames)
-							{
-								if (Name == CategoryStat.CategoryName)
-								{
-									PrimitiveCount += CategoryStat.PrimitiveCount;
-									CategoryStat.PrimitiveCount = 0;
-									break;
-								}
-							}
-						}
-					}
-
-					if (PrimitiveCount > 0)
-					{
-						FCoreDelegates::EOnScreenMessageSeverity Severity = CategoryBudget.PrimitiveBudget < PrimitiveCount ? FCoreDelegates::EOnScreenMessageSeverity::Warning : FCoreDelegates::EOnScreenMessageSeverity::Info;
-						OutMessages.Add(Severity, FText::FromString(FString::Printf(TEXT("%5dK / %5dK - %s"), PrimitiveCount / 1000, CategoryBudget.PrimitiveBudget / 1000, *(CategoryBudget.CategoryName.ToString()))));
+						FCoreDelegates::EOnScreenMessageSeverity Severity = CategoryBudget.PrimitiveBudget < *PrimitiveCount ? FCoreDelegates::EOnScreenMessageSeverity::Warning : FCoreDelegates::EOnScreenMessageSeverity::Info;
+						OutMessages.Add(Severity, FText::FromString(FString::Printf(TEXT("%5dK / %5dK - %s"), *PrimitiveCount / 1000, CategoryBudget.PrimitiveBudget / 1000, *(CategoryBudget.CategoryName.ToString()))));
 					}
 				}
 
-				// Show remaining (non-zeroed) stats.
-				for (FStats::FCategoryStats const& CategoryStat : CategoryStats)
+				// Show remaining (non-zeroed) stats not coverted by Budgets
+				for (const TPair<FName, uint64>& Pair : UntrackedPrimitives)
 				{
-					if (CategoryStat.PrimitiveCount > 0)
+					const FName& Name = Pair.Key;
+					uint64 PrimitiveCount = Pair.Value;
+
+					if (PrimitiveCount > 0)
 					{
-						OutMessages.Add(FCoreDelegates::EOnScreenMessageSeverity::Info, FText::FromString(FString::Printf(TEXT("\t%5dK - %s"), CategoryStat.PrimitiveCount / 1000, *(CategoryStat.CategoryName.ToString()))));
+						OutMessages.Add(FCoreDelegates::EOnScreenMessageSeverity::Info, FText::FromString(FString::Printf(TEXT("\t%5dK - %s"), PrimitiveCount / 1000, *(Name.ToString()))));
 					}
 				}
 
@@ -391,6 +380,65 @@ void FMeshDrawCommandStatsManager::Update()
 	// Collect stats during the next frame (check if STATGROUP_Culling is also visible somehow)
 	const bool bShowStats = CVarShowMeshDrawCommandStats->GetInt() == 0 ? false : true;
 	bCollectStats = bShowStats || bRequestDumpStats;
+
+#if CSV_PROFILER
+	bCollectStats |= (FCsvProfiler::Get()->IsCapturing_Renderthread() && CVarExportMeshDrawCommandStats->GetBool());
+#endif
+
+	if (bCollectStats)
+	{
+		// First time - Build associative map for quick Stat -> Budget lookup
+		if (Budgets.IsEmpty())
+		{
+			const UMeshDrawCommandStatsSettings* Settings = GetDefault<UMeshDrawCommandStatsSettings>();	
+			for (const FMeshDrawCommandStatsBudget& CategoryBudget : Settings->Budgets)
+			{
+				Budgets.Add(CategoryBudget.CategoryName, CategoryBudget.CategoryName);
+
+				for (FName Name : CategoryBudget.LinkedStatNames)
+				{
+					Budgets.Add(Name, CategoryBudget.CategoryName);
+				}
+			}
+		}
+
+		// Total up Primitive across stats to their respective Budgets
+		BudgetedPrimitives.Reset();
+		UntrackedPrimitives.Reset();
+
+		for (const FStats::FCategoryStats& CategoryStat : Stats.CategoryStats)
+		{ 
+			FName* BudgetName = Budgets.Find(CategoryStat.CategoryName);
+			TMap<FName, uint64>* Map = BudgetName ? &BudgetedPrimitives : &UntrackedPrimitives;
+
+			uint64& Count = Map->FindOrAdd(BudgetName ? *BudgetName : CategoryStat.CategoryName);
+			Count += CategoryStat.PrimitiveCount;
+		}
+	}
+
+#if CSV_PROFILER
+	if (FCsvProfiler::Get()->IsCapturing_Renderthread() && CVarExportMeshDrawCommandStats->GetBool())
+	{
+		// Output Budget totals
+		for (const TPair<FName, uint64>& Pair : BudgetedPrimitives)
+		{
+			TRACE_CSV_PROFILER_INLINE_STAT(TCHAR_TO_ANSI(*Pair.Key.ToString()), CSV_CATEGORY_INDEX(MeshDrawCommandStats));
+			FCsvProfiler::RecordCustomStat(Pair.Key, CSV_CATEGORY_INDEX(MeshDrawCommandStats), IntCastChecked<int32>(Pair.Value), ECsvCustomStatOp::Set);
+		}
+
+		// Output Untracked totals as a single bucket
+		uint64 TotalUntracked = 0;
+
+		for (const TPair<FName, uint64>& Pair : UntrackedPrimitives)
+		{
+			TotalUntracked += Pair.Value;
+		}
+
+		const char* Name = "Untracked";
+		TRACE_CSV_PROFILER_INLINE_STAT(Name, CSV_CATEGORY_INDEX(MeshDrawCommandStats));
+		FCsvProfiler::RecordCustomStat(Name, CSV_CATEGORY_INDEX(MeshDrawCommandStats), IntCastChecked<int32>(TotalUntracked), ECsvCustomStatOp::Set);
+	}
+#endif
 }
 
 void FMeshDrawCommandStatsManager::DumpStats(FFrameData* FrameData)
@@ -405,16 +453,16 @@ void FMeshDrawCommandStatsManager::DumpStats(FFrameData* FrameData)
 	struct FStatEntry
 	{
 		FName PassName;
-		int32 VisibilePrimitiveCount;
-		int32 VisibleInstance;		
+		int32 VisibilePrimitiveCount = 0;
+		int32 VisibleInstance = 0;
 		FName CategoryName;
 		FName ResourceName;
-		int32 LODIndex;
-		int32 SegmentIndex;
+		int32 LODIndex = 0;
+		int32 SegmentIndex = 0;
 		FString MaterialName;
-		int32 PrimitiveCount;
-		int32 TotalInstanceCount;
-		int32 TotalPrimitiveCount;
+		int32 PrimitiveCount = 0;
+		int32 TotalInstanceCount = 0;
+		int32 TotalPrimitiveCount = 0;
 	};
 	TArray<FStatEntry> StatEntries;
 
@@ -428,14 +476,16 @@ void FMeshDrawCommandStatsManager::DumpStats(FFrameData* FrameData)
 				StatEntry.PassName = PassStats->PassName;
 				StatEntry.VisibilePrimitiveCount = DrawData.VisibleInstanceCount * DrawData.PrimitiveCount;
 				StatEntry.VisibleInstance = DrawData.VisibleInstanceCount;
-				StatEntry.LODIndex = DrawData.LODIndex;
-				StatEntry.SegmentIndex = DrawData.SegmentIndex;
 				StatEntry.PrimitiveCount = DrawData.PrimitiveCount;
 				StatEntry.TotalInstanceCount = DrawData.TotalInstanceCount;
 				StatEntry.TotalPrimitiveCount = DrawData.TotalInstanceCount * DrawData.PrimitiveCount;
+				StatEntry.CategoryName = DrawData.StatsData.CategoryName;
+			#if MESH_DRAW_COMMAND_DEBUG_DATA
+				StatEntry.LODIndex = DrawData.LODIndex;
+				StatEntry.SegmentIndex = DrawData.SegmentIndex;
 				StatEntry.ResourceName = DrawData.ResourceName;
 				StatEntry.MaterialName = DrawData.MaterialName;
-				StatEntry.CategoryName = DrawData.StatsData.CategoryName;
+			#endif
 			}
 		}
 	}
@@ -495,4 +545,4 @@ static FAutoConsoleCommand GDumpMeshDrawCommandStatsCmd(
 	}
 }));
 
-#endif  // MESH_DRAW_COMMAND_STATS
\ No newline at end of file
+#endif  // MESH_DRAW_COMMAND_STATS
diff --git a/Engine/Source/Runtime/Renderer/Private/MeshDrawCommandStats.h b/Engine/Source/Runtime/Renderer/Private/MeshDrawCommandStats.h
--- a/Engine/Source/Runtime/Renderer/Private/MeshDrawCommandStats.h
+++ b/Engine/Source/Runtime/Renderer/Private/MeshDrawCommandStats.h
@@ -173,6 +173,9 @@ private:
 	FCriticalSection FrameDataCS;			
 	TArray<FFrameData*> Frames;				//< All active frames (contains the frame for which we are collecting stats now and all frames waiting for GPU readback)
 	FStats Stats;							//< Last updated frame stats
+	TMap<FName, FName> Budgets;				//< LinkedStatName to Budget CategoryName
+	TMap<FName, uint64> BudgetedPrimitives;	//< Budget CategoryName to Total Primitive Count
+	TMap<FName, uint64> UntrackedPrimitives;//< Primitives which aren't tracked by any Budgets
 
 	FDelegateHandle ScreenMessageDelegate;	//< Delegate used to render optional screen stats
 
