Fix for crash in FWaterInfoTextureDepthPassMeshProcessor::CollectPSOInitializers when accessing the default material during PSO collection from the game thread.

#rnx

diff --git a/Engine/Source/Runtime/Renderer/Private/WaterInfoTextureRendering.cpp b/Engine/Source/Runtime/Renderer/Private/WaterInfoTextureRendering.cpp
--- a/Engine/Source/Runtime/Renderer/Private/WaterInfoTextureRendering.cpp
+++ b/Engine/Source/Runtime/Renderer/Private/WaterInfoTextureRendering.cpp
@@ -543,11 +543,11 @@ void FWaterInfoTextureDepthPassMeshProcessor::CollectPSOInitializers(const FScen
 			&& !Material.MaterialModifiesMeshPosition_GameThread()
 			&& Material.WritesEveryPixel(false, bVFTypeSupportsNullPixelShader))
 		{
-			// todo: collect here without calling render thread only code in GetMaterialNoFallback()
-			//const FMaterialRenderProxy& DefaultProxy = *UMaterial::GetDefaultMaterial(MD_Surface)->GetRenderProxy();
-			//const FMaterial& DefaultMaterial = *DefaultProxy.GetMaterialNoFallback(FeatureLevel);
+			EMaterialQualityLevel::Type ActiveQualityLevel = GetCachedScalabilityCVars().MaterialQualityLevel;
+			const FMaterial* DefaultMaterial = UMaterial::GetDefaultMaterial(MD_Surface)->GetMaterialResource(FeatureLevel, ActiveQualityLevel);
+			check(DefaultMaterial);
 
-			//CollectPSOInitializersInternal<true>(SceneTexturesConfig, VertexFactoryData, DefaultMaterial, MeshFillMode, MeshCullMode, PreCacheParams, PSOInitializers);
+			CollectPSOInitializersInternal<true>(SceneTexturesConfig, VertexFactoryData, *DefaultMaterial, MeshFillMode, MeshCullMode, PreCacheParams, PSOInitializers);
 		}
 		else
 		{
