Fix complex special tile not being rendered with local light because stencil value was cleared during the stencil based shadow mask computation.

diff --git a/Engine/Shaders/Private/DeferredLightPixelShaders.usf b/Engine/Shaders/Private/DeferredLightPixelShaders.usf
--- a/Engine/Shaders/Private/DeferredLightPixelShaders.usf
+++ b/Engine/Shaders/Private/DeferredLightPixelShaders.usf
@@ -15,7 +15,7 @@
   #elif SUBSTRATE_TILETYPE == 2
 	// COMPLEX PATH
   #elif SUBSTRATE_TILETYPE == 3
-	// COMPLEX PATH
+	// COMPLEX SPECIAL PATH
 	#define SUBSTRATE_COMPLEXSPECIALPATH 1
   #else
 	#error Substrate tile type non-implemented
diff --git a/Engine/Source/Runtime/Renderer/Private/LightRendering.cpp b/Engine/Source/Runtime/Renderer/Private/LightRendering.cpp
--- a/Engine/Source/Runtime/Renderer/Private/LightRendering.cpp
+++ b/Engine/Source/Runtime/Renderer/Private/LightRendering.cpp
@@ -730,7 +730,7 @@ class FDeferredLightPS : public FGlobalShader
 	class FAtmosphereTransmittance : SHADER_PERMUTATION_BOOL("USE_ATMOSPHERE_TRANSMITTANCE");
 	class FCloudTransmittance 	: SHADER_PERMUTATION_BOOL("USE_CLOUD_TRANSMITTANCE");
 	class FAnistropicMaterials 	: SHADER_PERMUTATION_BOOL("SUPPORTS_ANISOTROPIC_MATERIALS");
-	class FSubstrateTileType		: SHADER_PERMUTATION_INT("SUBSTRATE_TILETYPE", 4);
+	class FSubstrateTileType	: SHADER_PERMUTATION_INT("SUBSTRATE_TILETYPE", 4);
 	class FVirtualShadowMapMask : SHADER_PERMUTATION_BOOL("USE_VIRTUAL_SHADOW_MAP_MASK");
 
 	using FPermutationDomain = TShaderPermutationDomain<
@@ -2092,8 +2092,8 @@ static uint32 InternalSetBoundingGeometryDepthState(FGraphicsPipelineStateInitia
 		{
 		case ESubstrateTileType::ESimple :			StencilRef = Substrate::StencilBit_Fast;			GraphicsPSOInit.DepthStencilState = TStaticDepthStencilState<false, CompareFunction, true, CF_Equal, SO_Keep, SO_Keep, SO_Keep, true, CF_Equal, SO_Keep, SO_Keep, SO_Keep, Substrate::StencilBit_Fast, 0x0>::GetRHI(); break;
 		case ESubstrateTileType::ESingle :			StencilRef = Substrate::StencilBit_Single;			GraphicsPSOInit.DepthStencilState = TStaticDepthStencilState<false, CompareFunction, true, CF_Equal, SO_Keep, SO_Keep, SO_Keep, true, CF_Equal, SO_Keep, SO_Keep, SO_Keep, Substrate::StencilBit_Single, 0x0>::GetRHI(); break;
-		case ESubstrateTileType::EComplex:			StencilRef = Substrate::StencilBit_Complex;		GraphicsPSOInit.DepthStencilState = TStaticDepthStencilState<false, CompareFunction, true, CF_Equal, SO_Keep, SO_Keep, SO_Keep, true, CF_Equal, SO_Keep, SO_Keep, SO_Keep, Substrate::StencilBit_Complex, 0x0>::GetRHI(); break;
-		case ESubstrateTileType::EComplexSpecial:	StencilRef = Substrate::StencilBit_ComplexSpecial; GraphicsPSOInit.DepthStencilState = TStaticDepthStencilState<false, CompareFunction, true, CF_Equal, SO_Keep, SO_Keep, SO_Keep, true, CF_Equal, SO_Keep, SO_Keep, SO_Keep, Substrate::StencilBit_ComplexSpecial, 0x0>::GetRHI(); break;
+		case ESubstrateTileType::EComplex:			StencilRef = Substrate::StencilBit_Complex;			GraphicsPSOInit.DepthStencilState = TStaticDepthStencilState<false, CompareFunction, true, CF_Equal, SO_Keep, SO_Keep, SO_Keep, true, CF_Equal, SO_Keep, SO_Keep, SO_Keep, Substrate::StencilBit_Complex, 0x0>::GetRHI(); break;
+		case ESubstrateTileType::EComplexSpecial:	StencilRef = Substrate::StencilBit_ComplexSpecial; 	GraphicsPSOInit.DepthStencilState = TStaticDepthStencilState<false, CompareFunction, true, CF_Equal, SO_Keep, SO_Keep, SO_Keep, true, CF_Equal, SO_Keep, SO_Keep, SO_Keep, Substrate::StencilBit_ComplexSpecial, 0x0>::GetRHI(); break;
 		default: check(false);
 		}
 	}
@@ -2395,6 +2395,7 @@ BEGIN_SHADER_PARAMETER_STRUCT(FRenderLightParameters, )
 	SHADER_PARAMETER_STRUCT_INCLUDE(FSubstrateTileParameter, SubstrateTileSimple)
 	SHADER_PARAMETER_STRUCT_INCLUDE(FSubstrateTileParameter, SubstrateTileSingle)
 	SHADER_PARAMETER_STRUCT_INCLUDE(FSubstrateTileParameter, SubstrateTileComplex)
+	SHADER_PARAMETER_STRUCT_INCLUDE(FSubstrateTileParameter, SubstrateTileSpectialComplex)
 END_SHADER_PARAMETER_STRUCT()
 
 /** Shader parameters for Standard Deferred Light pass. */
@@ -2491,6 +2492,7 @@ static void RenderLight(
 			PassParameters->SubstrateTileSimple = Substrate::SetTileParameters(GraphBuilder, View, ESubstrateTileType::ESingle);
 			PassParameters->SubstrateTileSingle = Substrate::SetTileParameters(GraphBuilder, View, ESubstrateTileType::ESimple);
 			PassParameters->SubstrateTileComplex = Substrate::SetTileParameters(GraphBuilder, View, ESubstrateTileType::EComplex);
+			PassParameters->SubstrateTileSpectialComplex = Substrate::SetTileParameters(GraphBuilder, View, ESubstrateTileType::EComplexSpecial);
 		}
 
 		FDeferredLightPS::FPermutationDomain PermutationVector;
diff --git a/Engine/Source/Runtime/Renderer/Private/PostProcess/SceneRenderTargets.h b/Engine/Source/Runtime/Renderer/Private/PostProcess/SceneRenderTargets.h
--- a/Engine/Source/Runtime/Renderer/Private/PostProcess/SceneRenderTargets.h
+++ b/Engine/Source/Runtime/Renderer/Private/PostProcess/SceneRenderTargets.h
@@ -28,12 +28,12 @@
 #define STENCIL_LIGHTING_CHANNELS_BIT_ID				4
 #define STENCIL_RECEIVE_DECAL_BIT_ID					7
 // Used only during the lighting pass - alias/reuse light channels (which copied from stencil to a texture prior to lighting pass)
-#define STENCIL_SUBSTRATE_FASTPATH							4 
+#define STENCIL_SUBSTRATE_FASTPATH							4
 #define STENCIL_SUBSTRATE_SINGLEPATH						5
 #define STENCIL_SUBSTRATE_COMPLEX							6
-#define STENCIL_SUBSTRATE_COMPLEX_SPECIAL					1
+#define STENCIL_SUBSTRATE_COMPLEX_SPECIAL					7
 // Used only by Substrate during the base pass when bUseDBufferPass is enabled (to mark material SUBSTRATE_DBUFFER_RESPONSE_xxx Normal/BaseColor/Roughness)
-#define STENCIL_SUBSTRATE_RECEIVE_DBUFFER_NORMAL_BIT_ID	1
+#define STENCIL_SUBSTRATE_RECEIVE_DBUFFER_NORMAL_BIT_ID		1
 #define STENCIL_SUBSTRATE_RECEIVE_DBUFFER_DIFFUSE_BIT_ID	3
 #define STENCIL_SUBSTRATE_RECEIVE_DBUFFER_ROUGHNESS_BIT_ID	7
 
diff --git a/Engine/Source/Runtime/Renderer/Private/Substrate/Substrate.cpp b/Engine/Source/Runtime/Renderer/Private/Substrate/Substrate.cpp
--- a/Engine/Source/Runtime/Renderer/Private/Substrate/Substrate.cpp
+++ b/Engine/Source/Runtime/Renderer/Private/Substrate/Substrate.cpp
@@ -99,7 +99,7 @@ const TCHAR* ToString(ESubstrateTileType Type)
 	case ESubstrateTileType::EOpaqueRoughRefractionSSSWithout:	return TEXT("Opaque/RoughRefraction/SSSWithout");
 	case ESubstrateTileType::EDecalSimple:						return TEXT("Decal/Simple");
 	case ESubstrateTileType::EDecalSingle:						return TEXT("Decal/Single");
-	case ESubstrateTileType::EDecalComplex:					return TEXT("Decal/Complex");
+	case ESubstrateTileType::EDecalComplex:						return TEXT("Decal/Complex");
 
 	}
 	return TEXT("Unknown");
@@ -234,14 +234,14 @@ static void InitialiseSubstrateViewData(FRDGBuilder& GraphBuilder, FViewInfo& Vi
 
 			Out.ClassificationTileListBufferOffset[ESubstrateTileType::ESimple]							= 0;
 			Out.ClassificationTileListBufferOffset[ESubstrateTileType::ESingle]							= Out.ClassificationTileListBufferOffset[ESubstrateTileType::ESimple]							+ RegularTileCount;
-			Out.ClassificationTileListBufferOffset[ESubstrateTileType::EComplex]							= Out.ClassificationTileListBufferOffset[ESubstrateTileType::ESingle]							+ RegularTileCount;
+			Out.ClassificationTileListBufferOffset[ESubstrateTileType::EComplex]						= Out.ClassificationTileListBufferOffset[ESubstrateTileType::ESingle]							+ RegularTileCount;
 			Out.ClassificationTileListBufferOffset[ESubstrateTileType::EComplexSpecial]					= Out.ClassificationTileListBufferOffset[ESubstrateTileType::EComplex]							+ RegularTileCount;
-			Out.ClassificationTileListBufferOffset[ESubstrateTileType::EOpaqueRoughRefraction]				= Out.ClassificationTileListBufferOffset[ESubstrateTileType::EComplexSpecial]					+ (bUsesComplexSpecialRenderPath ? RegularTileCount : 4);
-			Out.ClassificationTileListBufferOffset[ESubstrateTileType::EOpaqueRoughRefractionSSSWithout]	= Out.ClassificationTileListBufferOffset[ESubstrateTileType::EOpaqueRoughRefraction]			+ RoughTileCount;
-			Out.ClassificationTileListBufferOffset[ESubstrateTileType::EDecalSimple]						= Out.ClassificationTileListBufferOffset[ESubstrateTileType::EOpaqueRoughRefractionSSSWithout]	+ RoughTileCount;
-			Out.ClassificationTileListBufferOffset[ESubstrateTileType::EDecalSingle]						= Out.ClassificationTileListBufferOffset[ESubstrateTileType::EDecalSimple]						+ DecalTileCount;
-			Out.ClassificationTileListBufferOffset[ESubstrateTileType::EDecalComplex]						= Out.ClassificationTileListBufferOffset[ESubstrateTileType::EDecalSingle]						+ DecalTileCount;
-			uint32 TotalTileCount										 								= Out.ClassificationTileListBufferOffset[ESubstrateTileType::EDecalComplex]					+ DecalTileCount;
+			Out.ClassificationTileListBufferOffset[ESubstrateTileType::EOpaqueRoughRefraction]			= Out.ClassificationTileListBufferOffset[ESubstrateTileType::EComplexSpecial]					+ (bUsesComplexSpecialRenderPath ? RegularTileCount : 4);
+			Out.ClassificationTileListBufferOffset[ESubstrateTileType::EOpaqueRoughRefractionSSSWithout]= Out.ClassificationTileListBufferOffset[ESubstrateTileType::EOpaqueRoughRefraction]			+ RoughTileCount;
+			Out.ClassificationTileListBufferOffset[ESubstrateTileType::EDecalSimple]					= Out.ClassificationTileListBufferOffset[ESubstrateTileType::EOpaqueRoughRefractionSSSWithout]	+ RoughTileCount;
+			Out.ClassificationTileListBufferOffset[ESubstrateTileType::EDecalSingle]					= Out.ClassificationTileListBufferOffset[ESubstrateTileType::EDecalSimple]						+ DecalTileCount;
+			Out.ClassificationTileListBufferOffset[ESubstrateTileType::EDecalComplex]					= Out.ClassificationTileListBufferOffset[ESubstrateTileType::EDecalSingle]						+ DecalTileCount;
+			uint32 TotalTileCount										 								= Out.ClassificationTileListBufferOffset[ESubstrateTileType::EDecalComplex]						+ DecalTileCount;
 
 			check(TotalTileCount > 0);
 
@@ -1005,10 +1005,10 @@ uint32 TileTypeDispatchIndirectArgOffset(const ESubstrateTileType Type)
 
 // Add additionnaly bits for filling/clearing stencil to ensure that the 'Substrate' bits are not corrupted by the stencil shadows 
 // when generating shadow mask. Withouth these 'trailing' bits, the incr./decr. operation would change/corrupt the 'Substrate' bits
-constexpr uint32 StencilBit_Fast_1	  = StencilBit_Fast;
-constexpr uint32 StencilBit_Single_1  = StencilBit_Single;
-constexpr uint32 StencilBit_Complex_1 = StencilBit_Complex; 
-constexpr uint32 StencilBit_ComplexSpecial_1 = StencilBit_ComplexSpecial; 
+constexpr uint32 StencilBit_Fast_1			= StencilBit_Fast;
+constexpr uint32 StencilBit_Single_1		= StencilBit_Single;
+constexpr uint32 StencilBit_Complex_1		= StencilBit_Complex; 
+constexpr uint32 StencilBit_ComplexSpecial_1= StencilBit_ComplexSpecial; 
 
 void AddSubstrateInternalClassificationTilePass(
 	FRDGBuilder& GraphBuilder,
@@ -1048,7 +1048,7 @@ void AddSubstrateInternalClassificationTilePass(
 		case ESubstrateTileType::ESimple:							ParametersPS->DebugTileColor = FVector4f(0.0f, 1.0f, 0.0f, 1.0); break;
 		case ESubstrateTileType::ESingle:							ParametersPS->DebugTileColor = FVector4f(1.0f, 1.0f, 0.0f, 1.0); break;
 		case ESubstrateTileType::EComplex:							ParametersPS->DebugTileColor = FVector4f(1.0f, 0.0f, 0.0f, 1.0); break;
-		case ESubstrateTileType::EComplexSpecial:					ParametersPS->DebugTileColor = FVector4f(0.3f, 0.0f, 0.0f, 1.0); break;
+		case ESubstrateTileType::EComplexSpecial:					ParametersPS->DebugTileColor = FVector4f(0.3f, 0.0f, 0.3f, 1.0); break;
 
 		case ESubstrateTileType::EOpaqueRoughRefraction:			ParametersPS->DebugTileColor = FVector4f(0.0f, 1.0f, 1.0f, 1.0); break;
 		case ESubstrateTileType::EOpaqueRoughRefractionSSSWithout:	ParametersPS->DebugTileColor = FVector4f(0.0f, 0.0f, 1.0f, 1.0); break;
diff --git a/Engine/Source/Runtime/Renderer/Private/Substrate/Substrate.h b/Engine/Source/Runtime/Renderer/Private/Substrate/Substrate.h
--- a/Engine/Source/Runtime/Renderer/Private/Substrate/Substrate.h
+++ b/Engine/Source/Runtime/Renderer/Private/Substrate/Substrate.h
@@ -173,10 +173,10 @@ struct FSubstrateViewData
 
 namespace Substrate
 {
-constexpr uint32 StencilBit_Fast			= 0x08; // In sync with SceneRenderTargets.h - GET_STENCIL_BIT_MASK(STENCIL_SUBSTRATE_FASTPATH)
-constexpr uint32 StencilBit_Single			= 0x10; // In sync with SceneRenderTargets.h - GET_STENCIL_BIT_MASK(STENCIL_SUBSTRATE_SINGLEPATH)
-constexpr uint32 StencilBit_Complex			= 0x20; // In sync with SceneRenderTargets.h - GET_STENCIL_BIT_MASK(STENCIL_SUBSTRATE_COMPLEX)
-constexpr uint32 StencilBit_ComplexSpecial	= 0x02; // In sync with SceneRenderTargets.h - GET_STENCIL_BIT_MASK(STENCIL_SUBSTRATE_COMPLEX_SPECIAL)	
+constexpr uint32 StencilBit_Fast			= 0x10; // In sync with SceneRenderTargets.h - GET_STENCIL_BIT_MASK(STENCIL_SUBSTRATE_FASTPATH)
+constexpr uint32 StencilBit_Single			= 0x20; // In sync with SceneRenderTargets.h - GET_STENCIL_BIT_MASK(STENCIL_SUBSTRATE_SINGLEPATH)
+constexpr uint32 StencilBit_Complex			= 0x40; // In sync with SceneRenderTargets.h - GET_STENCIL_BIT_MASK(STENCIL_SUBSTRATE_COMPLEX)
+constexpr uint32 StencilBit_ComplexSpecial	= 0x80; // In sync with SceneRenderTargets.h - GET_STENCIL_BIT_MASK(STENCIL_SUBSTRATE_COMPLEX_SPECIAL)	
 
 FIntPoint GetSubstrateTextureResolution(const FViewInfo& View, const FIntPoint& InResolution);
 bool GetSubstrateUsesComplexSpecialPath(const FViewInfo& View);
