LFV - put view data into its own structure.

diff --git a/Engine/Source/Runtime/Renderer/Private/BasePassRendering.cpp b/Engine/Source/Runtime/Renderer/Private/BasePassRendering.cpp
--- a/Engine/Source/Runtime/Renderer/Private/BasePassRendering.cpp
+++ b/Engine/Source/Runtime/Renderer/Private/BasePassRendering.cpp
@@ -719,7 +719,7 @@ void SetupSharedBasePassParameters(
 		SharedParameters.FogISR = SharedParameters.Fog;
 	}
 
-	SharedParameters.LFV = View.LocalFogVolumeUniformParametersStruct;
+	SharedParameters.LFV = View.LocalFogVolumeViewData.UniformParametersStruct;
 
 	const FScene* Scene = View.Family->Scene ? View.Family->Scene->GetRenderScene() : nullptr;
 	const FPlanarReflectionSceneProxy* ReflectionSceneProxy = Scene ? Scene->GetForwardPassGlobalPlanarReflection() : nullptr;
diff --git a/Engine/Source/Runtime/Renderer/Private/LocalFogVolumeRendering.cpp b/Engine/Source/Runtime/Renderer/Private/LocalFogVolumeRendering.cpp
--- a/Engine/Source/Runtime/Renderer/Private/LocalFogVolumeRendering.cpp
+++ b/Engine/Source/Runtime/Renderer/Private/LocalFogVolumeRendering.cpp
@@ -39,19 +39,20 @@ void SetDummyLocalFogVolumeForView(FRDGBuilder& GraphBuilder, FViewInfo& View)
 	// The size of the structure must be a multiple of FUintVector4.
 	static_assert(sizeof(FLocalFogVolumeGPUInstanceData) == UintVec4CountInLocalFogVolumeGPUInstanceData * SizeOfUintVec4);
 
-	View.LocalFogVolumeGPUInstanceCount = 0;
+	View.LocalFogVolumeViewData.GPUInstanceCount = 0;
 
 	static FLocalFogVolumeGPUInstanceData DummyData;
-	View.LocalFogVolumeGPUInstanceDataBuffer = CreateVertexBuffer(GraphBuilder, TEXT("DUMMYLocalFogVolumeGPUInstanceDataBuffer"),
+	View.LocalFogVolumeViewData.GPUInstanceDataBuffer = CreateVertexBuffer(GraphBuilder, TEXT("DUMMYLocalFogVolumeGPUInstanceDataBuffer"),
 		FRDGBufferDesc::CreateBufferDesc(SizeOfUintVec4, UintVec4CountInLocalFogVolumeGPUInstanceData), &DummyData, sizeof(FLocalFogVolumeGPUInstanceData) * 1, ERDGInitialDataFlags::NoCopy);
 
-	View.LocalFogVolumeGPUInstanceDataBufferSRV = GraphBuilder.CreateSRV(View.LocalFogVolumeGPUInstanceDataBuffer, PF_A32B32G32R32F);
+	View.LocalFogVolumeViewData.GPUInstanceDataBufferSRV = GraphBuilder.CreateSRV(View.LocalFogVolumeViewData.GPUInstanceDataBuffer, PF_A32B32G32R32F);
 
-	View.LocalFogVolumeUniformParametersStruct.LocalFogVolumeInstanceCount = View.LocalFogVolumeGPUInstanceCount;
-	View.LocalFogVolumeUniformParametersStruct.LocalFogVolumeInstances = View.LocalFogVolumeGPUInstanceDataBufferSRV;
-	View.LocalFogVolumeUniformBuffer = GraphBuilder.CreateUniformBuffer(&View.LocalFogVolumeUniformParametersStruct);
+	View.LocalFogVolumeViewData.UniformParametersStruct.LocalFogVolumeInstanceCount = View.LocalFogVolumeViewData.GPUInstanceCount;
+	View.LocalFogVolumeViewData.UniformParametersStruct.LocalFogVolumeInstances = View.LocalFogVolumeViewData.GPUInstanceDataBufferSRV;
+	View.LocalFogVolumeViewData.UniformBuffer = GraphBuilder.CreateUniformBuffer(&View.LocalFogVolumeViewData.UniformParametersStruct);
 };
 
+
 /*=============================================================================
 	FScene functions
 =============================================================================*/
@@ -167,16 +168,16 @@ void CreateViewLocalFogVolumeBufferSRV(FViewInfo& View, FRDGBuilder& GraphBuilde
 	// 3. Allocate buffer and initialize with sorted data to upload to GPU
 	const uint32 AllLocalFogVolumeInstanceBytesFinal = sizeof(FLocalFogVolumeGPUInstanceData) * SortingData.LocalFogVolumeInstanceCountFinal;
 
-	View.LocalFogVolumeGPUInstanceCount = SortingData.LocalFogVolumeInstanceCountFinal;
-	View.LocalFogVolumeGPUInstanceDataBuffer = CreateVertexBuffer(
+	View.LocalFogVolumeViewData.GPUInstanceCount = SortingData.LocalFogVolumeInstanceCountFinal;
+	View.LocalFogVolumeViewData.GPUInstanceDataBuffer = CreateVertexBuffer(
 		GraphBuilder, TEXT("LocalFogVolumeGPUInstanceDataBuffer"),
 		FRDGBufferDesc::CreateBufferDesc(SizeOfUintVec4, SortingData.LocalFogVolumeInstanceCountFinal * UintVec4CountInLocalFogVolumeGPUInstanceData),
 		LocalFogVolumeGPUSortedInstanceData, AllLocalFogVolumeInstanceBytesFinal, ERDGInitialDataFlags::NoCopy);
-	View.LocalFogVolumeGPUInstanceDataBufferSRV = GraphBuilder.CreateSRV(View.LocalFogVolumeGPUInstanceDataBuffer, PF_A32B32G32R32F);
+	View.LocalFogVolumeViewData.GPUInstanceDataBufferSRV = GraphBuilder.CreateSRV(View.LocalFogVolumeViewData.GPUInstanceDataBuffer, PF_A32B32G32R32F);
 
-	View.LocalFogVolumeUniformParametersStruct.LocalFogVolumeInstanceCount = View.LocalFogVolumeGPUInstanceCount;
-	View.LocalFogVolumeUniformParametersStruct.LocalFogVolumeInstances = View.LocalFogVolumeGPUInstanceDataBufferSRV;
-	View.LocalFogVolumeUniformBuffer = GraphBuilder.CreateUniformBuffer(&View.LocalFogVolumeUniformParametersStruct);
+	View.LocalFogVolumeViewData.UniformParametersStruct.LocalFogVolumeInstanceCount = View.LocalFogVolumeViewData.GPUInstanceCount;
+	View.LocalFogVolumeViewData.UniformParametersStruct.LocalFogVolumeInstances = View.LocalFogVolumeViewData.GPUInstanceDataBufferSRV;
+	View.LocalFogVolumeViewData.UniformBuffer = GraphBuilder.CreateUniformBuffer(&View.LocalFogVolumeViewData.UniformParametersStruct);
 }
 
 void InitLocalFogVolumesForViews(
@@ -284,7 +285,7 @@ void RenderLocalFogVolume(
 
 		for (FViewInfo& View : Views)
 		{
-			if (View.LocalFogVolumeGPUInstanceCount == 0)
+			if (View.LocalFogVolumeViewData.GPUInstanceCount == 0)
 			{
 				continue;
 			}
@@ -292,10 +293,10 @@ void RenderLocalFogVolume(
 			FLocalFogVolumePassParameters* PassParameters = GraphBuilder.AllocParameters<FLocalFogVolumePassParameters>();
 
 			PassParameters->VS.View = GetShaderBinding(View.ViewUniformBuffer);
-			PassParameters->VS.LFV = View.LocalFogVolumeUniformParametersStruct;
+			PassParameters->VS.LFV = View.LocalFogVolumeViewData.UniformParametersStruct;
 
 			PassParameters->PS.View = GetShaderBinding(View.ViewUniformBuffer);
-			PassParameters->PS.LFV = View.LocalFogVolumeUniformParametersStruct;
+			PassParameters->PS.LFV = View.LocalFogVolumeViewData.UniformParametersStruct;
 
 			PassParameters->SceneTextures = SceneTextures.UniformBuffer;
 			PassParameters->RenderTargets[0] = FRenderTargetBinding(SceneColorTexture, ERenderTargetLoadAction::ENoAction);
@@ -311,7 +312,7 @@ void RenderLocalFogVolume(
 			ClearUnusedGraphResources(VertexShader, &PassParameters->VS);
 			ClearUnusedGraphResources(PixelShader, &PassParameters->PS);
 
-			uint32 LocalFogVolumeGPUInstanceCount = View.LocalFogVolumeGPUInstanceCount;
+			uint32 LocalFogVolumeGPUInstanceCount = View.LocalFogVolumeViewData.GPUInstanceCount;
 			GraphBuilder.AddPass(
 				RDG_EVENT_NAME("RenderLocalFogVolume %u inst.", LocalFogVolumeGPUInstanceCount),
 				PassParameters,
@@ -425,7 +426,7 @@ void RenderLocalFogVolumeMobile(
 	FRHICommandList& RHICmdList,
 	const FViewInfo& View)
 {
-	if (View.LocalFogVolumeGPUInstanceCount == 0)
+	if (View.LocalFogVolumeViewData.GPUInstanceCount == 0)
 	{
 		return;
 	}
@@ -459,12 +460,12 @@ void RenderLocalFogVolumeMobile(
 
 	FMobileLocalFogVolumeVS::FParameters VSParameters;
 	VSParameters.View = View.GetShaderParameters();
-	VSParameters.LFV = View.LocalFogVolumeUniformParametersStruct;
+	VSParameters.LFV = View.LocalFogVolumeViewData.UniformParametersStruct;
 	SetShaderParameters(RHICmdList, VertexShader, VertexShader.GetVertexShader(), VSParameters);
 
 	FMobileLocalFogVolumePS::FParameters PSParameters;
 	PSParameters.View = View.GetShaderParameters();
-	PSParameters.LFV = View.LocalFogVolumeUniformParametersStruct;
+	PSParameters.LFV = View.LocalFogVolumeViewData.UniformParametersStruct;
 	// PSParameters.MobileBasePass filled up by the RDG pass parameters.
 	SetShaderParameters(RHICmdList, PixelShader, PixelShader.GetPixelShader(), PSParameters);
 
@@ -472,11 +473,11 @@ void RenderLocalFogVolumeMobile(
 
 	RHICmdList.DrawIndexedPrimitive(
 		GetUnitCubeIndexBuffer()
-		, 0										//BaseVertexIndex
-		, 0										//FirstInstance
-		, 8										//uint32 NumVertices
-		, 0										//uint32 StartIndex
-		, UE_ARRAY_COUNT(GCubeIndices) / 3		//uint32 NumPrimitives
-		, View.LocalFogVolumeGPUInstanceCount	//uint32 NumInstances
+		, 0												//BaseVertexIndex
+		, 0												//FirstInstance
+		, 8												//uint32 NumVertices
+		, 0												//uint32 StartIndex
+		, UE_ARRAY_COUNT(GCubeIndices) / 3				//uint32 NumPrimitives
+		, View.LocalFogVolumeViewData.GPUInstanceCount	//uint32 NumInstances
 	);
 }
diff --git a/Engine/Source/Runtime/Renderer/Private/LocalFogVolumeRendering.h b/Engine/Source/Runtime/Renderer/Private/LocalFogVolumeRendering.h
--- a/Engine/Source/Runtime/Renderer/Private/LocalFogVolumeRendering.h
+++ b/Engine/Source/Runtime/Renderer/Private/LocalFogVolumeRendering.h
@@ -84,6 +84,16 @@ struct FLocalFogVolumeSortingData
 	TArray<FLocalFogVolumeSortKey>	LocalFogVolumeSortKeys;
 };
 
+struct FLocalFogVolumeViewData
+{
+	uint32 GPUInstanceCount = 0;
+	FRDGBufferRef GPUInstanceDataBuffer = nullptr;
+	FRDGBufferSRVRef GPUInstanceDataBufferSRV = nullptr;
+	TRDGUniformBufferRef<FLocalFogVolumeUniformParameters> UniformBuffer = nullptr;
+
+	FLocalFogVolumeUniformParameters UniformParametersStruct;
+};
+
 
 /*=============================================================================
 	Local height fog rendering functions
diff --git a/Engine/Source/Runtime/Renderer/Private/MobileBasePassRendering.cpp b/Engine/Source/Runtime/Renderer/Private/MobileBasePassRendering.cpp
--- a/Engine/Source/Runtime/Renderer/Private/MobileBasePassRendering.cpp
+++ b/Engine/Source/Runtime/Renderer/Private/MobileBasePassRendering.cpp
@@ -247,7 +247,7 @@ void SetupMobileBasePassUniformParameters(
 	// Substrate
 	Substrate::BindSubstrateMobileForwardPasslUniformParameters(GraphBuilder, View, BasePassParameters.Substrate);
 
-	BasePassParameters.LFV = View.LocalFogVolumeUniformParametersStruct;
+	BasePassParameters.LFV = View.LocalFogVolumeViewData.UniformParametersStruct;
 
 	SetupReflectionUniformParameters(GraphBuilder, View, BasePassParameters.ReflectionsParameters);
 }
diff --git a/Engine/Source/Runtime/Renderer/Private/MobileShadingRenderer.cpp b/Engine/Source/Runtime/Renderer/Private/MobileShadingRenderer.cpp
--- a/Engine/Source/Runtime/Renderer/Private/MobileShadingRenderer.cpp
+++ b/Engine/Source/Runtime/Renderer/Private/MobileShadingRenderer.cpp
@@ -1311,7 +1311,7 @@ void FMobileSceneRenderer::RenderForward(FRDGBuilder& GraphBuilder, FRDGTextureR
 		PassParameters->MobileBasePass = CreateMobileBasePassUniformBuffer(GraphBuilder, View, EMobileBasePass::Opaque, SetupMode, MobileBasePassTextures);
 		PassParameters->ReflectionCapture = View.MobileReflectionCaptureUniformBuffer;
 		PassParameters->RenderTargets = BasePassRenderTargets;
-		PassParameters->LocalFogVolumeInstances = View.LocalFogVolumeGPUInstanceDataBufferSRV;
+		PassParameters->LocalFogVolumeInstances = View.LocalFogVolumeViewData.GPUInstanceDataBufferSRV;
 	
 		BuildInstanceCullingDrawParams(GraphBuilder, View, PassParameters);
 
@@ -1648,7 +1648,7 @@ void FMobileSceneRenderer::RenderDeferred(FRDGBuilder& GraphBuilder, const FSort
 		PassParameters->MobileBasePass = CreateMobileBasePassUniformBuffer(GraphBuilder, View, EMobileBasePass::Opaque, SetupMode, MobileBasePassTextures);
 		PassParameters->ReflectionCapture = View.MobileReflectionCaptureUniformBuffer;
 		PassParameters->RenderTargets = BasePassRenderTargets;
-		PassParameters->LocalFogVolumeInstances = View.LocalFogVolumeGPUInstanceDataBufferSRV;
+		PassParameters->LocalFogVolumeInstances = View.LocalFogVolumeViewData.GPUInstanceDataBufferSRV;
 		
 		BuildInstanceCullingDrawParams(GraphBuilder, View, PassParameters);
 
diff --git a/Engine/Source/Runtime/Renderer/Private/SceneRendering.cpp b/Engine/Source/Runtime/Renderer/Private/SceneRendering.cpp
--- a/Engine/Source/Runtime/Renderer/Private/SceneRendering.cpp
+++ b/Engine/Source/Runtime/Renderer/Private/SceneRendering.cpp
@@ -975,10 +975,7 @@ void FViewInfo::Init()
 	SubstrateViewData.Reset();
 	HairStrandsViewData = FHairStrandsViewData();
 
-	LocalFogVolumeGPUInstanceCount = 0;
-	LocalFogVolumeGPUInstanceDataBuffer = nullptr;
-	LocalFogVolumeGPUInstanceDataBufferSRV = nullptr;
-	LocalFogVolumeUniformBuffer = nullptr;
+	LocalFogVolumeViewData = FLocalFogVolumeViewData();
 
 	GPUSceneViewId = INDEX_NONE;
 
diff --git a/Engine/Source/Runtime/Renderer/Private/SceneRendering.h b/Engine/Source/Runtime/Renderer/Private/SceneRendering.h
--- a/Engine/Source/Runtime/Renderer/Private/SceneRendering.h
+++ b/Engine/Source/Runtime/Renderer/Private/SceneRendering.h
@@ -1368,11 +1368,7 @@ public:
 
 	FSubstrateViewData SubstrateViewData;
 
-	uint32 LocalFogVolumeGPUInstanceCount;
-	FRDGBufferRef LocalFogVolumeGPUInstanceDataBuffer;
-	FRDGBufferSRVRef LocalFogVolumeGPUInstanceDataBufferSRV;
-	FLocalFogVolumeUniformParameters LocalFogVolumeUniformParametersStruct;
-	TRDGUniformBufferRef<FLocalFogVolumeUniformParameters> LocalFogVolumeUniformBuffer;
+	FLocalFogVolumeViewData LocalFogVolumeViewData;
 
 	FHairStrandsViewData HairStrandsViewData;
 
