LFV - preparing for tile culling.
Added a maximum volume count and volume are clamp to that count after view culling.

diff --git a/Engine/Source/Runtime/Renderer/Private/LocalFogVolumeRendering.cpp b/Engine/Source/Runtime/Renderer/Private/LocalFogVolumeRendering.cpp
--- a/Engine/Source/Runtime/Renderer/Private/LocalFogVolumeRendering.cpp
+++ b/Engine/Source/Runtime/Renderer/Private/LocalFogVolumeRendering.cpp
@@ -19,6 +19,31 @@ static TAutoConsoleVariable<int32> CVarLocalFogVolumeApplyOnTransclucent(
 	TEXT("Project settings enabling the sampling of local fog volumes on translucent elements.\n"),
 	ECVF_ReadOnly | ECVF_RenderThreadSafe);
 
+static TAutoConsoleVariable<int32> CVarLocalFogVolumeTilePixelSize(
+	TEXT("r.LocalFogVolume.TilePixelSize"), 128,
+	TEXT("Tile size on screen in pixel at which we cull the local fog volumes.\n"),
+	ECVF_RenderThreadSafe);
+
+static TAutoConsoleVariable<int32> CVarLocalFogVolumeTileMaxInstanceCount(
+	TEXT("r.LocalFogVolume.TileMaxInstanceCount"), 32,
+	TEXT("Maximum number of local fog volume to account for per view (and per tile or consistency).\n"),
+	ECVF_RenderThreadSafe);
+
+// Example of tile setup
+//  - 1920x1080 => 15x9 tiles
+//  - Allowing max 32 volumes at once => culling list buffer = 15 * 9 * 32 * 1 byte = 4320 bytes = 4.3KB
+
+static uint32 GetLocalFogVolumeTilePixelSize()
+{
+	return FMath::Max(8u, FMath::Min(512u, (uint32)CVarLocalFogVolumeTilePixelSize.GetValueOnRenderThread()));
+}
+
+static uint32 GetLocalFogVolumeTileMaxInstanceCount()
+{
+	// We do not allow more than 256 instances since culled indices might be stored a u8 values.
+	return FMath::Max(1u, FMath::Min(256u, (uint32)CVarLocalFogVolumeTileMaxInstanceCount.GetValueOnRenderThread()));
+}
+
 bool ShouldRenderLocalFogVolume(const FScene* Scene, const FSceneViewFamily& Family)
 {
 	const FEngineShowFlags EngineShowFlags = Family.EngineShowFlags;
@@ -157,12 +182,17 @@ void CreateViewLocalFogVolumeBufferSRV(FViewInfo& View, FRDGBuilder& GraphBuilde
 	}
 	SortingData.LocalFogVolumeSortKeys.Sort();
 
+	// We limit the instance count to the maximum of instance we can have per tile
+	const uint32 LocalFogVolumeTileMaxInstanceCount = GetLocalFogVolumeTileMaxInstanceCount();
+	const uint32 DiscardedOffset = (uint32)FMath::Max(0, int32(SortingData.LocalFogVolumeInstanceCountFinal) - int32(LocalFogVolumeTileMaxInstanceCount));
+	SortingData.LocalFogVolumeInstanceCountFinal = FMath::Min(SortingData.LocalFogVolumeInstanceCountFinal, LocalFogVolumeTileMaxInstanceCount);
+
 	// 2. Create the buffer containing all the fog volume data instance sorted according to their key for the current view.
 	FLocalFogVolumeGPUInstanceData* LocalFogVolumeGPUSortedInstanceData = (FLocalFogVolumeGPUInstanceData*)GraphBuilder.Alloc(sizeof(FLocalFogVolumeGPUInstanceData) * SortingData.LocalFogVolumeInstanceCountFinal, 16);
-	for (uint32 i = 0; i < SortingData.LocalFogVolumeInstanceCountFinal; ++i)
+	for (uint32 i = 0; i < SortingData.LocalFogVolumeInstanceCountFinal; i++)
 	{
 		// We could also have an indirection buffer on GPU but choosing to go with the sorting + copy on CPU since it is expected to not have many local height fog volumes.
-		LocalFogVolumeGPUSortedInstanceData[i] = SortingData.LocalFogVolumeGPUInstanceData[SortingData.LocalFogVolumeSortKeys[i].FogVolume.Index];
+		LocalFogVolumeGPUSortedInstanceData[i] = SortingData.LocalFogVolumeGPUInstanceData[SortingData.LocalFogVolumeSortKeys[i + DiscardedOffset].FogVolume.Index];
 	}
 
 	// 3. Allocate buffer and initialize with sorted data to upload to GPU
diff --git a/Engine/Source/Runtime/Renderer/Private/LocalFogVolumeRendering.h b/Engine/Source/Runtime/Renderer/Private/LocalFogVolumeRendering.h
--- a/Engine/Source/Runtime/Renderer/Private/LocalFogVolumeRendering.h
+++ b/Engine/Source/Runtime/Renderer/Private/LocalFogVolumeRendering.h
@@ -92,6 +92,8 @@ struct FLocalFogVolumeViewData
 	TRDGUniformBufferRef<FLocalFogVolumeUniformParameters> UniformBuffer = nullptr;
 
 	FLocalFogVolumeUniformParameters UniformParametersStruct;
+
+	FRDGTextureRef TileInstanceIndexTextureArray; // First slice is the instance count, later slices are instance indices.
 };
 
 
