Path Tracer: Build RTPSO for the path tracer during startup to avoid lengthy compilation times the first time the path tracer is selected in the UI

It would be better if we could make this a block compilation so that we guarantee the RTPSO is complete before the UI appears, but this seems to cause a driver crash at the moment.

diff --git a/Engine/Source/Runtime/D3D12RHI/Private/D3D12RayTracing.cpp b/Engine/Source/Runtime/D3D12RHI/Private/D3D12RayTracing.cpp
--- a/Engine/Source/Runtime/D3D12RHI/Private/D3D12RayTracing.cpp
+++ b/Engine/Source/Runtime/D3D12RHI/Private/D3D12RayTracing.cpp
@@ -994,6 +994,7 @@ public:
 		void DoTask(ENamedThreads::Type CurrentThread, const FGraphEventRef& MyCompletionGraphEvent)
 		{
 			SCOPE_CYCLE_COUNTER(STAT_RTPSO_CompileShader);
+			TRACE_CPUPROFILER_EVENT_SCOPE(ShaderCompileTask);
 
 			uint64 CompileTimeCycles = 0;
 			CompileTimeCycles -= FPlatformTime::Cycles64();
@@ -1095,6 +1096,12 @@ public:
 			CompileTimeCycles += FPlatformTime::Cycles64();
 
 			Entry.CompileTimeMS = float(FPlatformTime::ToMilliseconds64(CompileTimeCycles));
+
+			if (Entry.CompileTimeMS >= 1000.0f)
+			{
+				// Log compilations of individual shaders that took more than 1 second
+				UE_LOG(LogD3D12RHI, Log, TEXT("Compiled %s for RTPSO in %.2f ms."), OriginalEntryPoints[0], Entry.CompileTimeMS);
+			}
 		}
 
 		FORCEINLINE TStatId GetStatId() const
@@ -1932,6 +1939,7 @@ public:
 	FD3D12RayTracingPipelineState(FD3D12Device* Device, const FRayTracingPipelineStateInitializer& Initializer)
 	{
 		SCOPE_CYCLE_COUNTER(STAT_RTPSO_CreatePipeline);
+		TRACE_CPUPROFILER_EVENT_SCOPE(RTPSO_CreatePipeline);
 
 		checkf(Initializer.GetRayGenTable().Num() > 0 || Initializer.bPartial, TEXT("Ray tracing pipelines must have at leat one ray generation shader."));
 
@@ -2281,7 +2289,7 @@ public:
 		{
 			const double CompileTimeMS = 1000.0 * FPlatformTime::ToSeconds64(CompileTime);
 			const double LinkTimeMS = 1000.0 * FPlatformTime::ToSeconds64(LinkTime);
-			uint32 NumUniqueShaders = UniqueShaderCollections.Num();
+			const uint32 NumUniqueShaders = UniqueShaderCollections.Num();
 			UE_LOG(LogD3D12RHI, Log,
 				TEXT("Creating RTPSO with %d shaders (%d cached, %d new) took %.2f ms. Compile time %.2f ms, link time %.2f ms."),
 				PipelineShaderHashes.Num(), NumCacheHits, NumUniqueShaders - NumCacheHits, (float)TotalCreationTimeMS, (float)CompileTimeMS, (float)LinkTimeMS);
diff --git a/Engine/Source/Runtime/Renderer/Private/PathTracing.cpp b/Engine/Source/Runtime/Renderer/Private/PathTracing.cpp
--- a/Engine/Source/Runtime/Renderer/Private/PathTracing.cpp
+++ b/Engine/Source/Runtime/Renderer/Private/PathTracing.cpp
@@ -2291,7 +2291,7 @@ IMPLEMENT_SHADER_TYPE(, FPathTracingCompositorPS, TEXT("/Engine/Private/PathTrac
 static FPathTracingRG::FPermutationDomain GetPathTracingRGPermutation(const FScene& Scene)
 {
 	const bool bUseExperimental = CVarPathTracingExperimental.GetValueOnRenderThread() != 0;
-	const bool bUseCompaction = bUseExperimental && CVarPathTracingCompaction.GetValueOnRenderThread() != 0;
+	const bool bUseCompaction = (bUseExperimental == false) || CVarPathTracingCompaction.GetValueOnRenderThread() != 0;
 	const bool bUseAdaptiveSampling = bUseExperimental && CVarPathTracingAdaptiveSampling.GetValueOnRenderThread() != 0;
 	const bool bHasComplexSpecialRenderPath = Substrate::IsSubstrateEnabled() && Scene.SubstrateSceneData.bUsesComplexSpecialRenderPath;
 
@@ -2321,6 +2321,47 @@ void FDeferredShadingSceneRenderer::PreparePathTracing(const FSceneViewFamily& V
 	}
 }
 
+void PreparePathTracingRTPSO()
+{
+	ENQUEUE_RENDER_COMMAND(PreparePathTracingRTPSO)([](FRHICommandListImmediate& RHICmdList)
+		{
+			int NumValidPermutations = 0;
+			for (int PermutationId = 0; PermutationId < FPathTracingRG::FPermutationDomain::PermutationCount; PermutationId++)
+			{
+				FGlobalShaderPermutationParameters Parameters(FPathTracingRG::GetStaticType().GetFName(), GMaxRHIShaderPlatform, PermutationId);
+				if (!FPathTracingRG::ShouldCompilePermutation(Parameters))
+				{
+					// Permutation is not enabled, nothing to pre-compile
+					continue;
+				}
+				FGlobalShaderMap* ShaderMap = GetGlobalShaderMap(GMaxRHIFeatureLevel);
+
+				FPathTracingRG::FPermutationDomain PermutationVector(PermutationId);
+				
+				FRHIRayTracingShader* RayGenShaderTable[] = {
+					ShaderMap->GetShader<FPathTracingRG>(PermutationVector).GetRayTracingShader(),
+					ShaderMap->GetShader<FPathTracingInitExtinctionCoefficientRG>().GetRayTracingShader(),
+				};
+				FRHIRayTracingShader* MissShaderTable[] = {
+					GetPathTracingDefaultMissShader(ShaderMap),
+				};
+				FRHIRayTracingShader* HitGroupTable[] = {
+					GetPathTracingDefaultOpaqueHitShader(ShaderMap),
+				};
+				FRayTracingPipelineStateInitializer Initializer;
+				Initializer.bPartial = true; // TODO: getting a crash in nvidia driver when false
+				Initializer.SetRayGenShaderTable(RayGenShaderTable);
+				Initializer.SetMissShaderTable(MissShaderTable);
+				Initializer.SetHitGroupTable(HitGroupTable);
+				Initializer.MaxPayloadSizeInBytes = RayGenShaderTable[0]->RayTracingPayloadSize;
+				FRayTracingPipelineState* PipelineState = PipelineStateCache::GetAndOrCreateRayTracingPipelineState(RHICmdList, Initializer);
+				NumValidPermutations++;
+			}
+			UE_LOG(LogRenderer, Log, TEXT("Requested compilation of Path Tracing RTPSOs (%d permutations)."), NumValidPermutations);
+		}
+	);
+}
+
 void FSceneViewState::PathTracingInvalidate(bool InvalidateAnimationStates)
 {
 	FPathTracingState* State = PathTracingState.Get();
@@ -2727,7 +2768,7 @@ void FDeferredShadingSceneRenderer::RenderPathTracing(
 			bNeedsTextureExtract = true;
 
 			// should we use path compaction?
-			const bool bUseCompaction = bUseExperimental && CVarPathTracingCompaction.GetValueOnRenderThread() != 0;
+			const bool bUseCompaction = (bUseExperimental == false) || CVarPathTracingCompaction.GetValueOnRenderThread() != 0;
 			const bool bUseIndirectDispatch = GRHISupportsRayTracingDispatchIndirect && CVarPathTracingIndirectDispatch.GetValueOnRenderThread() != 0;
 			const int FlushRenderingCommands = CVarPathTracingFlushDispatch.GetValueOnRenderThread();
 
diff --git a/Engine/Source/Runtime/Renderer/Private/PathTracing.h b/Engine/Source/Runtime/Renderer/Private/PathTracing.h
--- a/Engine/Source/Runtime/Renderer/Private/PathTracing.h
+++ b/Engine/Source/Runtime/Renderer/Private/PathTracing.h
@@ -29,6 +29,8 @@ RENDERER_API FRHIRayTracingShader* GetGPULightmassDefaultMissShader(const FGloba
 RENDERER_API FRHIRayTracingShader* GetGPULightmassDefaultOpaqueHitShader(const FGlobalShaderMap* ShaderMap);
 RENDERER_API FRHIRayTracingShader* GetGPULightmassDefaultHiddenHitShader(const FGlobalShaderMap* ShaderMap);
 
+void PreparePathTracingRTPSO();
+
 class FRDGTexture;
 
 struct FPathTracingResources
diff --git a/Engine/Source/Runtime/Renderer/Private/Renderer.cpp b/Engine/Source/Runtime/Renderer/Private/Renderer.cpp
--- a/Engine/Source/Runtime/Renderer/Private/Renderer.cpp
+++ b/Engine/Source/Runtime/Renderer/Private/Renderer.cpp
@@ -45,6 +45,7 @@
 #include "PrimitiveSceneShaderData.h"
 #include "MeshDrawCommandStats.h"
 #include "LocalFogVolumeRendering.h"
+#include "PathTracing.h"
 
 DEFINE_LOG_CATEGORY(LogRenderer);
 
@@ -86,6 +87,19 @@ void FRendererModule::StartupModule()
 	GIdentityPrimitiveUniformBuffer.InitContents();
 	GDistanceCullFadedInUniformBuffer.InitContents();
 	GDitherFadedInUniformBuffer.InitContents();
+
+#if RHI_RAYTRACING && WITH_EDITOR
+	if (FApp::CanEverRender() && !FApp::IsUnattended())
+	{
+		FCoreDelegates::OnPostEngineInit.AddLambda([]() {
+			// We add this step via the PostEngineInit delegate so that it can run after PostInitRHI has run,
+			// and the rendering thread has been started so that we are able to create RTPSOs.
+			// For now, we only attempt to create the PathTracer RTPSO as it is the most expensive to compile by far.
+			// See UE-190955 for example timings.
+			PreparePathTracingRTPSO();
+		});
+	}
+#endif
 }
 
 void FRendererModule::ShutdownModule()
