Path Tracer: Fix crash on startup when using an RHI that doesn't support ray tracing (such as DX11)

diff --git a/Engine/Source/Runtime/Renderer/Private/PathTracing.cpp b/Engine/Source/Runtime/Renderer/Private/PathTracing.cpp
--- a/Engine/Source/Runtime/Renderer/Private/PathTracing.cpp
+++ b/Engine/Source/Runtime/Renderer/Private/PathTracing.cpp
@@ -2323,6 +2323,10 @@ void FDeferredShadingSceneRenderer::PreparePathTracing(const FSceneViewFamily& V
 
 void PreparePathTracingRTPSO()
 {
+	if (!IsRayTracingEnabled())
+	{
+		return;
+	}
 	ENQUEUE_RENDER_COMMAND(PreparePathTracingRTPSO)([](FRHICommandListImmediate& RHICmdList)
 		{
 			int NumValidPermutations = 0;
