Fixed VT uniform buffer layout crash. Moved CallPendingCallbacks to happen only during a scene render and not with every scene update to avoid ordering issues.

diff --git a/Engine/Source/Runtime/Renderer/Private/RendererScene.cpp b/Engine/Source/Runtime/Renderer/Private/RendererScene.cpp
--- a/Engine/Source/Runtime/Renderer/Private/RendererScene.cpp
+++ b/Engine/Source/Runtime/Renderer/Private/RendererScene.cpp
@@ -6314,9 +6314,6 @@ void FScene::UpdateAllPrimitiveSceneInfos(FRDGBuilder& GraphBuilder, EUpdateAllP
 
 	GPUSkinCacheTask.Wait();
 
-	// Must run prior to static mesh gathering as callbacks can evaluate uniform expression caches.
-	FVirtualTextureSystem::Get().CallPendingCallbacks();
-
 	if (SceneInfosWithStaticDrawListUpdate.Num() > 0)
 	{
 		FPrimitiveSceneInfo::AddStaticMeshes(this, SceneInfosWithStaticDrawListUpdate, false);
diff --git a/Engine/Source/Runtime/Renderer/Private/SceneRendering.cpp b/Engine/Source/Runtime/Renderer/Private/SceneRendering.cpp
--- a/Engine/Source/Runtime/Renderer/Private/SceneRendering.cpp
+++ b/Engine/Source/Runtime/Renderer/Private/SceneRendering.cpp
@@ -3397,6 +3397,9 @@ IVisibilityTaskData* FSceneRenderer::OnRenderBegin(FRDGBuilder& GraphBuilder, FG
 {
 	check(!FDeferredUpdateResource::IsUpdateNeeded());
 
+	// This is called prior to scene update to avoid a race condition with the MDC caching task.
+	FVirtualTextureSystem::Get().CallPendingCallbacks();
+
 	EUpdateAllPrimitiveSceneInfosAsyncOps AsyncOps = EUpdateAllPrimitiveSceneInfosAsyncOps::None;
 
 	if (GAsyncCreateLightPrimitiveInteractions > 0)
