Allow capsule indirect shadows in projects that don't support static lighting.  However, the usefulness of capsule indirect shadows with Lumen is limited, since the shadow direction will be extracted only from the skylight, and will not match the indirect shadow directions of the surrounding environment.
Fix texture creation flags on the capsule shadow render target so it can be visualized with 'vis'

diff --git a/Engine/Source/Runtime/Renderer/Private/CapsuleShadowRendering.cpp b/Engine/Source/Runtime/Renderer/Private/CapsuleShadowRendering.cpp
--- a/Engine/Source/Runtime/Renderer/Private/CapsuleShadowRendering.cpp
+++ b/Engine/Source/Runtime/Renderer/Private/CapsuleShadowRendering.cpp
@@ -496,8 +496,8 @@ bool FDeferredShadingSceneRenderer::RenderCapsuleDirectShadows(
 
 	{
 		const FIntPoint BufferSize = GetBufferSizeForCapsuleShadows(Views[0]);
-		const FRDGTextureDesc Desc(FRDGTextureDesc::Create2D(BufferSize, PF_G16R16F, FClearValueBinding::None, TexCreate_RenderTargetable | TexCreate_UAV));
-		RayTracedShadowsRT = GraphBuilder.CreateTexture(Desc, TEXT("RayTracedShadows"));
+		const FRDGTextureDesc Desc(FRDGTextureDesc::Create2D(BufferSize, PF_G16R16F, FClearValueBinding::None, TexCreate_ShaderResource | TexCreate_UAV));
+		RayTracedShadowsRT = GraphBuilder.CreateTexture(Desc, TEXT("CapsuleShadows.ShadowFactors"));
 	}
 
 	for (int32 ViewIndex = 0; ViewIndex < Views.Num(); ViewIndex++)
@@ -982,9 +982,8 @@ void FDeferredShadingSceneRenderer::RenderIndirectCapsuleShadows(FRDGBuilder& Gr
 
 	{
 		const FIntPoint BufferSize = GetBufferSizeForCapsuleShadows(Views[0]);
-		const FRDGTextureDesc Desc(FRDGTextureDesc::Create2D(BufferSize, PF_G16R16F, FClearValueBinding::None, TexCreate_RenderTargetable | TexCreate_UAV));
-		// Reuse temporary target from RTDF shadows
-		RayTracedShadowsRT = GraphBuilder.CreateTexture(Desc, TEXT("RayTracedShadows"));
+		const FRDGTextureDesc Desc(FRDGTextureDesc::Create2D(BufferSize, PF_G16R16F, FClearValueBinding::None, TexCreate_ShaderResource | TexCreate_UAV));
+		RayTracedShadowsRT = GraphBuilder.CreateTexture(Desc, TEXT("CapsuleShadows.ShadowFactors"));
 	}
 
 	TArray<FRDGTextureRef, TInlineAllocator<2>> RenderTargets;
diff --git a/Engine/Source/Runtime/Renderer/Private/DeferredShadingRenderer.cpp b/Engine/Source/Runtime/Renderer/Private/DeferredShadingRenderer.cpp
--- a/Engine/Source/Runtime/Renderer/Private/DeferredShadingRenderer.cpp
+++ b/Engine/Source/Runtime/Renderer/Private/DeferredShadingRenderer.cpp
@@ -3570,7 +3570,7 @@ void FDeferredShadingSceneRenderer::Render(FRDGBuilder& GraphBuilder)
 			CompositionLighting.ProcessBeforeBasePass(GraphBuilder, DBufferTextures);
 		}
 		
-		if (IsForwardShadingEnabled(ShaderPlatform) && bAllowStaticLighting)
+		if (IsForwardShadingEnabled(ShaderPlatform))
 		{
 			RenderIndirectCapsuleShadows(GraphBuilder, SceneTextures);
 		}
@@ -3918,10 +3918,7 @@ void FDeferredShadingSceneRenderer::Render(FRDGBuilder& GraphBuilder)
 				AsyncLumenIndirectLightingOutputs);
 
 			// These modulate the scenecolor output from the basepass, which is assumed to be indirect lighting
-			if (bAllowStaticLighting)
-			{
-				RenderIndirectCapsuleShadows(GraphBuilder, SceneTextures);
-			}
+			RenderIndirectCapsuleShadows(GraphBuilder, SceneTextures);
 
 			// These modulate the scene color output from the base pass, which is assumed to be indirect lighting
 			RenderDFAOAsIndirectShadowing(GraphBuilder, SceneTextures, DynamicBentNormalAOTextures);
