Add support for light functions to mobile foward renderer.

Light function materials are rendered as additional draws with additive blending to the prepass/postprocess targets.
For the prepass codepath (that requires a depth prepass) the directional render target is updated in the main pass, only the color target is updated per light function.

Light functions of non-directional lights are currently not supported.

diff --git a/Engine/Shaders/Private/LightData.ush b/Engine/Shaders/Private/LightData.ush
--- a/Engine/Shaders/Private/LightData.ush
+++ b/Engine/Shaders/Private/LightData.ush
@@ -223,6 +223,12 @@ bool UnpackCastShadow(uint InShadowMapChannelMaskPacked)
 	return ((InShadowMapChannelMaskPacked >> 18) & 0x1) != 0;
 }
 
+bool UnpackHasLightFunction(FLocalLightData In)
+{
+	const uint LightTypeAndPackedShadowMapChannelMask = asuint(In.LightDirectionAndShadowMask.w);
+	return ((LightTypeAndPackedShadowMapChannelMask >> 19) & 0x1) != 0;
+}
+
 float UnpackVolumetricScatteringIntensity(FLocalLightData In)
 {
 	return f16tof32(asuint(In.SpotAnglesAndSourceRadiusPacked.w) >> 16);
diff --git a/Engine/Shaders/Private/MobileLocalLightsBuffer.usf b/Engine/Shaders/Private/MobileLocalLightsBuffer.usf
--- a/Engine/Shaders/Private/MobileLocalLightsBuffer.usf
+++ b/Engine/Shaders/Private/MobileLocalLightsBuffer.usf
@@ -2,13 +2,20 @@
 
 #define SUPPORT_CONTACT_SHADOWS 0
 
-#include "LightGridCommon.ush"
 #if !COMPUTE_SHADER
 #include "Common.ush"
 #include "SceneTextureParameters.ush"
 #include "DeferredLightingCommon.ush"
 #endif
 
+#if LIGHT_FUNCTION
+#include "/Engine/Generated/Material.ush"
+#include "LightDataUniforms.ush"
+#include "LightFunctionCommon.ush"
+#else
+#include "LightGridCommon.ush"
+#endif
+
 #if COMPUTE_SHADER
 
 uint2 GroupSize;
@@ -55,7 +62,9 @@ void MainCS(
 	}
 }
 
-#else
+#else // COMPUTE_SHADER
+
+#if !LIGHT_FUNCTION
 
 int LightGridPixelSize;
 Buffer<int> TileInfo;
@@ -113,6 +122,7 @@ void Main(float4 SvPosition : SV_POSITION
 		// appears useless
 		float InvLightRadiusSq = LocalLight.LightPositionAndInvRadius.w * LocalLight.LightPositionAndInvRadius.w;
 		float DistLight = length2(TranslatedWorldPosition - LocalLight.LightPositionAndInvRadius.xyz) * InvLightRadiusSq;
+		
 		if (DistLight > 1.0f)
 		{
 			continue;
@@ -134,7 +144,13 @@ void Main(float4 SvPosition : SV_POSITION
 		if( LightMask > 0)
 		{
 			TotalWeight += WeightLight;
-			Color += MaskedLightColor;
+			// For lights with light functions we only only use the compute the direction (OutColorB)
+			// The Color is added in additional draws using MainLightFunction()
+			if (!UnpackHasLightFunction(LocalLight))
+			{
+				Color += MaskedLightColor;
+			}
+
 			SavedL += L *  WeightLight;
 		}
 	}
@@ -168,4 +184,57 @@ void Main(float4 SvPosition : SV_POSITION
 #endif
 }
 
-#endif
\ No newline at end of file
+
+#else // !LIGHT_FUNCTION
+
+
+float4x4 SvPositionToLight;
+/** Fade distance in x, disabled brightness in y */
+float2 LightFunctionParameters2;
+
+void MainLightFunction(
+	float4 InScreenPosition : TEXCOORD0,
+	float4 SvPosition : SV_POSITION, 
+	out float3 OutColorA : SV_Target0
+)
+{
+	float2 ScreenUV = SvPositionToBufferUV(SvPosition);
+	SvPosition.z = LookupDeviceZ(ScreenUV);
+	float3 TranslatedWorldPosition = SvPositionToTranslatedWorld(SvPosition);
+	float ViewDistance = length(GetPrimaryView().TranslatedWorldCameraOrigin - TranslatedWorldPosition);
+	float GreyScale;
+	{
+		float4 Hom = mul(float4(SvPosition.xyz, 1), SvPositionToLight);
+		float3 LightVector = Hom.xyz / Hom.w;
+		float3 Color = GetLightFunctionColor(LightVector, TranslatedWorldPosition);
+		GreyScale = dot(Color, .3333f);
+	}
+	float DistanceFadeAlpha = saturate((LightFunctionParameters2.x - ViewDistance) / (LightFunctionParameters2.x * .2f));
+	GreyScale = lerp(LightFunctionParameters2.y, GreyScale, DistanceFadeAlpha);
+	GreyScale = lerp(LightFunctionParameters2.y, GreyScale, LightFunctionParameters.y); 
+	float EncodedLightAttenuation = EncodeLightAttenuation(GreyScale);
+
+	FDeferredLightData LightData = InitDeferredLightFromUniforms();
+	float3 L = LightData.Direction;	// Already normalized
+	float3 ToLight = LightData.TranslatedWorldPosition - TranslatedWorldPosition;
+	float LightMask = 1;
+	// At the moment directional lights are not supported
+	// if (LightData.bRadialLight)
+	{
+		LightMask = GetLocalLightAttenuation( TranslatedWorldPosition, LightData, ToLight, L );
+	}
+
+	float3 Color = LightData.Color * LightMask * EncodedLightAttenuation;
+
+#if POST_PROCESS_LOCAL_LIGHTS == 0
+	OutColorA = Color;
+#else
+	float3 N = normalize(cross(ddx(TranslatedWorldPosition.xyz), ddy(TranslatedWorldPosition.xyz)));
+	float NoL = max(0, dot(N, L));
+	float3 Diffuse =  (1 / PI) * NoL * Color;
+	OutColorA = Diffuse;
+#endif
+}
+#endif // !LIGHT_FUNCTION
+
+#endif // COMPUTE_SHADER
diff --git a/Engine/Source/Runtime/Renderer/Private/LightFunctionRendering.cpp b/Engine/Source/Runtime/Renderer/Private/LightFunctionRendering.cpp
--- a/Engine/Source/Runtime/Renderer/Private/LightFunctionRendering.cpp
+++ b/Engine/Source/Runtime/Renderer/Private/LightFunctionRendering.cpp
@@ -4,6 +4,7 @@
 	LightFunctionRendering.cpp: Implementation for rendering light functions.
 =============================================================================*/
 
+#include "LightFunctionRendering.h"
 #include "CoreMinimal.h"
 #include "Engine/Engine.h"
 #include "HAL/IConsoleManager.h"
@@ -77,6 +78,35 @@ private:
 
 IMPLEMENT_MATERIAL_SHADER_TYPE(,FLightFunctionVS,TEXT("/Engine/Private/LightFunctionVertexShader.usf"),TEXT("Main"),SF_Vertex);
 
+void LightFunctionSvPositionToLightTransform(FMatrix44f& OutMatrix, const FViewInfo& View, const FLightSceneInfo& LightSceneInfo)
+{
+	const FVector Scale = LightSceneInfo.Proxy->GetLightFunctionScale();
+	// Switch x and z so that z of the user specified scale affects the distance along the light direction
+	const FVector InverseScale = FVector(1.0 / Scale.Z, 1.0 / Scale.Y, 1.0 / Scale.X);
+	const FMatrix WorldToLight = LightSceneInfo.Proxy->GetWorldToLight() * FScaleMatrix(InverseScale);
+	const FVector2D InvViewSize = FVector2D(1.0 / View.ViewRect.Width(), 1.0 / View.ViewRect.Height());
+
+	// setup a matrix to transform float4(SvPosition.xyz,1) directly to Light (quality, performance as we don't need to convert or use interpolator)
+	//	new_xy = (xy - ViewRectMin.xy) * ViewSizeAndInvSize.zw * float2(2,-2) + float2(-1, 1);
+	//  transformed into one MAD:  new_xy = xy * ViewSizeAndInvSize.zw * float2(2,-2)      +       (-ViewRectMin.xy) * ViewSizeAndInvSize.zw * float2(2,-2) + float2(-1, 1);
+
+	const double Mx = 2.0 * InvViewSize.X;
+	const double My = -2.0 * InvViewSize.Y;
+	const double Ax = -1.0 - 2.0 * View.ViewRect.Min.X * InvViewSize.X;
+	const double Ay = 1.0 + 2.0 * View.ViewRect.Min.Y * InvViewSize.Y;
+
+	// todo: we could use InvTranslatedViewProjectionMatrix and TranslatedWorldToLight for better quality
+	const FMatrix SvPositionToLightValue =
+		FMatrix(
+			FPlane(Mx,  0,   0,  0),
+			FPlane(0, My,   0,  0),
+			FPlane(0,  0,   1,  0),
+			FPlane(Ax, Ay,   0,  1)
+		) * View.ViewMatrices.GetInvViewProjectionMatrix() * WorldToLight;
+
+	OutMatrix = FMatrix44f(SvPositionToLightValue);
+}
+
 /**
  * A pixel shader for projecting a light function onto the scene.
  */
@@ -119,31 +149,9 @@ public:
 		// Set the transform from screen space to light space.
 		if ( SvPositionToLight.IsBound() )
 		{
-			const FVector Scale = LightSceneInfo->Proxy->GetLightFunctionScale();
-			// Switch x and z so that z of the user specified scale affects the distance along the light direction
-			const FVector InverseScale = FVector( 1.0 / Scale.Z, 1.0 / Scale.Y, 1.0 / Scale.X );
-			const FMatrix WorldToLight = LightSceneInfo->Proxy->GetWorldToLight() * FScaleMatrix(InverseScale);	
-			const FVector2D InvViewSize = FVector2D(1.0 / View.ViewRect.Width(), 1.0 / View.ViewRect.Height());
-
-			// setup a matrix to transform float4(SvPosition.xyz,1) directly to Light (quality, performance as we don't need to convert or use interpolator)
-			//	new_xy = (xy - ViewRectMin.xy) * ViewSizeAndInvSize.zw * float2(2,-2) + float2(-1, 1);
-			//  transformed into one MAD:  new_xy = xy * ViewSizeAndInvSize.zw * float2(2,-2)      +       (-ViewRectMin.xy) * ViewSizeAndInvSize.zw * float2(2,-2) + float2(-1, 1);
-
-			const double Mx = 2.0 * InvViewSize.X;
-			const double My = -2.0 * InvViewSize.Y;
-			const double Ax = -1.0 - 2.0 * View.ViewRect.Min.X * InvViewSize.X;
-			const double Ay = 1.0 + 2.0 * View.ViewRect.Min.Y * InvViewSize.Y;
-
-			// todo: we could use InvTranslatedViewProjectionMatrix and TranslatedWorldToLight for better quality
-			const FMatrix SvPositionToLightValue = 
-				FMatrix(
-					FPlane(Mx,  0,   0,  0),
-					FPlane( 0, My,   0,  0),
-					FPlane( 0,  0,   1,  0),
-					FPlane(Ax, Ay,   0,  1)
-				) * View.ViewMatrices.GetInvViewProjectionMatrix() * WorldToLight;
-
-			SetShaderValue(BatchedParameters, SvPositionToLight, FMatrix44f(SvPositionToLightValue) );
+			FMatrix44f SvPositionToLightValue;
+			LightFunctionSvPositionToLightTransform(SvPositionToLightValue, View, *LightSceneInfo);
+			SetShaderValue(BatchedParameters, SvPositionToLight, SvPositionToLightValue);
 		}
 
 		LightFunctionParameters.Set(BatchedParameters, LightSceneInfo, ShadowFadeFraction);
@@ -175,8 +183,7 @@ private:
 
 IMPLEMENT_MATERIAL_SHADER_TYPE(,FLightFunctionPS,TEXT("/Engine/Private/LightFunctionPixelShader.usf"),TEXT("Main"),SF_Pixel);
 
-/** Returns a fade fraction for a light function and a given view based on the appropriate fade settings. */
-static float GetLightFunctionFadeFraction(const FViewInfo& View, FSphere LightBounds)
+float GetLightFunctionFadeFraction(const FViewInfo& View, FSphere LightBounds)
 {
 	extern float CalculateShadowFadeAlpha(const float MaxUnclampedResolution, const uint32 ShadowFadeResolution, const uint32 MinShadowResolution);
 
diff --git a/Engine/Source/Runtime/Renderer/Private/LightFunctionRendering.h b/Engine/Source/Runtime/Renderer/Private/LightFunctionRendering.h
new file mode 100644
--- /dev/null
+++ b/Engine/Source/Runtime/Renderer/Private/LightFunctionRendering.h
@@ -0,0 +1,14 @@
+// Copyright Epic Games, Inc. All Rights Reserved.
+
+#pragma once
+
+#include "Math/MathFwd.h"
+
+class FViewInfo;
+class FLightSceneInfo;
+
+// Computes a matrix to transform float4(SvPosition.xyz,1) directly to coordinate system of the given light
+void LightFunctionSvPositionToLightTransform(FMatrix44f& OutMatrix, const FViewInfo& View, const FLightSceneInfo& LightSceneInfo);
+
+// Returns a fade fraction for a light function and a given view based on the appropriate fade settings.
+float GetLightFunctionFadeFraction(const FViewInfo& View, FSphere LightBounds);
diff --git a/Engine/Source/Runtime/Renderer/Private/LightGridInjection.cpp b/Engine/Source/Runtime/Renderer/Private/LightGridInjection.cpp
--- a/Engine/Source/Runtime/Renderer/Private/LightGridInjection.cpp
+++ b/Engine/Source/Runtime/Renderer/Private/LightGridInjection.cpp
@@ -539,7 +539,7 @@ FComputeLightGridOutput FSceneRenderer::ComputeLightGrid(FRDGBuilder& GraphBuild
 						LightParameters.Color *= LightProxy->GetIndirectLightingScale();
 					}
 
-					uint32 LightTypeAndShadowMapChannelMaskPacked = LightSceneInfo->PackLightTypeAndShadowMapChannelMask(bAllowStaticLighting);
+					uint32 LightTypeAndShadowMapChannelMaskPacked = LightSceneInfo->PackLightTypeAndShadowMapChannelMask(bAllowStaticLighting, SortedLightInfo.SortKey.Fields.bLightFunction);
 
 					const bool bDynamicShadows = ViewFamily.EngineShowFlags.DynamicShadows && VisibleLightInfos.IsValidIndex(LightSceneInfo->Id);
 					const int32 VirtualShadowMapId = bDynamicShadows ? VisibleLightInfos[LightSceneInfo->Id].GetVirtualShadowMapId( &View ) : INDEX_NONE;
diff --git a/Engine/Source/Runtime/Renderer/Private/LightRendering.cpp b/Engine/Source/Runtime/Renderer/Private/LightRendering.cpp
--- a/Engine/Source/Runtime/Renderer/Private/LightRendering.cpp
+++ b/Engine/Source/Runtime/Renderer/Private/LightRendering.cpp
@@ -28,6 +28,7 @@
 #include "Shadows/ShadowSceneRenderer.h"
 #include "RenderCore.h"
 #include "BasePassRendering.h"
+#include "MobileBasePassRendering.h"
 #include "TranslucentLighting.h"
 
 // ENABLE_DEBUG_DISCARD_PROP is used to test the lighting code by allowing to discard lights to see how performance scales
@@ -382,7 +383,7 @@ bool FDeferredLightVS::ShouldCompilePermutation(const FGlobalShaderPermutationPa
 	FPermutationDomain PermutationVector(Parameters.PermutationId);
 	if (PermutationVector.Get<FRadialLight>())
 	{
-		return IsFeatureLevelSupported(Parameters.Platform, ERHIFeatureLevel::SM5) || IsMobileDeferredShadingEnabled(Parameters.Platform);
+		return IsFeatureLevelSupported(Parameters.Platform, ERHIFeatureLevel::SM5) || IsMobileDeferredShadingEnabled(Parameters.Platform) || MobileLocalLightsBufferEnabled(Parameters.Platform);
 	}
 	// used with FPrefilterPlanarReflectionPS on mobile
 	return true;
diff --git a/Engine/Source/Runtime/Renderer/Private/LightSceneInfo.cpp b/Engine/Source/Runtime/Renderer/Private/LightSceneInfo.cpp
--- a/Engine/Source/Runtime/Renderer/Private/LightSceneInfo.cpp
+++ b/Engine/Source/Runtime/Renderer/Private/LightSceneInfo.cpp
@@ -354,7 +354,7 @@ bool FLightSceneInfo::ShouldRecordShadowSubjectsForMobile() const
 	return bShouldRecordShadowSubjectsForMobile;
 }
 
-uint32 FLightSceneInfo::PackLightTypeAndShadowMapChannelMask(bool bAllowStaticLighting) const
+uint32 FLightSceneInfo::PackLightTypeAndShadowMapChannelMask(bool bAllowStaticLighting, bool bLightFunction) const
 {
 	uint32 Result = 0;
 
@@ -383,6 +383,8 @@ uint32 FLightSceneInfo::PackLightTypeAndShadowMapChannelMask(bool bAllowStaticLi
 	Result |= ((uint32)Proxy->GetLightType()) << 16;
 	const uint32 CastShadows = Proxy->CastsDynamicShadow() ? 1 : 0;
 	Result |= CastShadows << (16 + LightType_NumBits);
+	uint32 HasLightFunction = bLightFunction ? 1 : 0;;
+	Result |= HasLightFunction << (16 + LightType_NumBits + 1);
 
 	return Result;
 }
diff --git a/Engine/Source/Runtime/Renderer/Private/LightSceneInfo.h b/Engine/Source/Runtime/Renderer/Private/LightSceneInfo.h
--- a/Engine/Source/Runtime/Renderer/Private/LightSceneInfo.h
+++ b/Engine/Source/Runtime/Renderer/Private/LightSceneInfo.h
@@ -313,7 +313,7 @@ public:
 
 	bool ShouldRecordShadowSubjectsForMobile() const;
 
-	uint32 PackLightTypeAndShadowMapChannelMask(bool bAllowStaticLighting) const;
+	uint32 PackLightTypeAndShadowMapChannelMask(bool bAllowStaticLighting, bool bLightFunction = false) const;
 };
 
 /** Defines how the light is stored in the scene's light octree. */
diff --git a/Engine/Source/Runtime/Renderer/Private/MobileLocalLightsBuffer.cpp b/Engine/Source/Runtime/Renderer/Private/MobileLocalLightsBuffer.cpp
--- a/Engine/Source/Runtime/Renderer/Private/MobileLocalLightsBuffer.cpp
+++ b/Engine/Source/Runtime/Renderer/Private/MobileLocalLightsBuffer.cpp
@@ -8,8 +8,11 @@
 #include "RendererPrivateUtils.h"
 #include "GlobalRenderResources.h"
 #include "ScenePrivate.h"
+#include "LightRendering.h"
+#include "LightFunctionRendering.h"
+#include "Materials/MaterialRenderProxy.h"
 
-bool MobileLocalLighsBufferEnabled(const FStaticShaderPlatform Platform)
+static bool CompileShaderPermutationsForMobileLocalLightsBuffer(const FStaticShaderPlatform Platform)
 {
 	return !IsMobileDeferredShadingEnabled(Platform) && 
 			IsMobilePlatform(Platform) && 
@@ -31,7 +34,7 @@ public:
  
 	static bool ShouldCompilePermutation(const FGlobalShaderPermutationParameters& Parameters)
 	{
-		return MobileLocalLighsBufferEnabled(Parameters.Platform);
+		return CompileShaderPermutationsForMobileLocalLightsBuffer(Parameters.Platform);
 	}
 
 	static void ModifyCompilationEnvironment(const FGlobalShaderPermutationParameters& Parameters, FShaderCompilerEnvironment& OutEnvironment)
@@ -39,6 +42,7 @@ public:
 		FForwardLightingParameters::ModifyCompilationEnvironment(Parameters.Platform, OutEnvironment);
 		OutEnvironment.SetDefine(TEXT("THREADGROUP_SIZEX"), GLocalLightPrepassTileSizeX);
 		OutEnvironment.SetDefine(TEXT("COMPUTE_SHADER"), 1);
+		OutEnvironment.SetDefine(TEXT("LIGHT_FUNCTION"), 0);
 	}
 };
 
@@ -57,7 +61,7 @@ class FLocalLightBufferVS : public FGlobalShader
 
 	static bool ShouldCompilePermutation(const FGlobalShaderPermutationParameters& Parameters)
 	{
-		return MobileLocalLighsBufferEnabled(Parameters.Platform);
+		return CompileShaderPermutationsForMobileLocalLightsBuffer(Parameters.Platform);
 	}
 };
 
@@ -75,12 +79,13 @@ class FLocalLightBufferPS : public FGlobalShader
 
 	static bool ShouldCompilePermutation(const FGlobalShaderPermutationParameters& Parameters)
 	{
-		return MobileLocalLighsBufferEnabled(Parameters.Platform);
+		return CompileShaderPermutationsForMobileLocalLightsBuffer(Parameters.Platform);
 	}
 
 	static void ModifyCompilationEnvironment(const FGlobalShaderPermutationParameters& Parameters, FShaderCompilerEnvironment& OutEnvironment)
 	{
 		FForwardLightingParameters::ModifyCompilationEnvironment(Parameters.Platform, OutEnvironment);
+		OutEnvironment.SetDefine(TEXT("LIGHT_FUNCTION"), 0);
 
 		if (MobileLocalLightsBufferPrepassEnabled(Parameters.Platform))
 		{
@@ -105,9 +110,100 @@ SHADER_PARAMETER_STRUCT_INCLUDE(FSceneTextureShaderParameters, SceneTextures)
 RENDER_TARGET_BINDING_SLOTS()
 END_SHADER_PARAMETER_STRUCT()
 
-void FMobileSceneRenderer::RenderMobileLocalLightsBuffer(FRDGBuilder& GraphBuilder, FSceneTextures& SceneTextures, bool bIsPrepass)
+/**
+ * A pixel shader for projecting a light function onto the scene and blending with the color from the previously calculated lights in the prepass. 
+ */
+class FMobileLocalLightFunctionPS : public FMaterialShader
 {
-	if (!MobileLocalLighsBufferEnabled(ShaderPlatform) || 
+	DECLARE_SHADER_TYPE(FMobileLocalLightFunctionPS, Material);
+	SHADER_USE_PARAMETER_STRUCT_WITH_LEGACY_BASE(FMobileLocalLightFunctionPS, FMaterialShader)
+
+	BEGIN_SHADER_PARAMETER_STRUCT(FParameters, )
+		SHADER_PARAMETER(FMatrix44f, SvPositionToLight)
+		SHADER_PARAMETER(FVector4f, LightFunctionParameters)
+		SHADER_PARAMETER(FVector2f, LightFunctionParameters2)
+		SHADER_PARAMETER_STRUCT_REF(FDeferredLightUniformStruct, DeferredLightUniforms)
+	END_SHADER_PARAMETER_STRUCT()
+
+public:
+	/**
+	  * Makes sure only shaders for materials that are explicitly flagged
+	  * as 'UsedAsLightFunction' in the Material Editor gets compiled into
+	  * the shader cache.
+	**/  
+	static bool ShouldCompilePermutation(const FMaterialShaderPermutationParameters& Parameters)
+	{
+		return Parameters.MaterialParameters.MaterialDomain == MD_LightFunction && CompileShaderPermutationsForMobileLocalLightsBuffer(Parameters.Platform);
+	}
+
+	static void ModifyCompilationEnvironment(const FMaterialShaderPermutationParameters& Parameters, FShaderCompilerEnvironment& OutEnvironment)
+	{
+		FMaterialShader::ModifyCompilationEnvironment(Parameters, OutEnvironment);
+		OutEnvironment.SetDefine(TEXT("LIGHT_FUNCTION"), 1);
+		OutEnvironment.SetDefine(TEXT("COMPUTE_SHADER"), 0);
+		
+		if (MobileLocalLightsBufferPrepassEnabled(Parameters.Platform))
+		{
+			OutEnvironment.SetRenderTargetOutputFormat(0, PF_FloatR11G11B10);
+			OutEnvironment.SetRenderTargetOutputFormat(1, PF_A2B10G10R10);
+			OutEnvironment.SetDefine(TEXT("POST_PROCESS_LOCAL_LIGHTS"), 0);
+		}
+		else
+		{
+			OutEnvironment.SetRenderTargetOutputFormat(0, PF_FloatR11G11B10);
+			OutEnvironment.SetDefine(TEXT("POST_PROCESS_LOCAL_LIGHTS"), 1);
+		}
+	}
+
+	void SetParameters(FRHIBatchedShaderParameters& BatchedParameters, const FViewInfo& View, const FMaterialRenderProxy* MaterialProxy, const FMaterial& Material)
+	{
+		FMaterialShader::SetViewParameters(BatchedParameters, View, View.ViewUniformBuffer);
+		FMaterialShader::SetParameters(BatchedParameters, MaterialProxy, Material, View);
+	}
+
+	FParameters GetParameters(const FViewInfo& View, const FLightSceneInfo* LightSceneInfo, const float FadeAlpha)
+	{
+		FParameters PS;
+		
+		LightFunctionSvPositionToLightTransform(PS.SvPositionToLight, View, *LightSceneInfo);
+		PS.LightFunctionParameters = FLightFunctionSharedParameters::GetLightFunctionSharedParameters(LightSceneInfo, FadeAlpha);
+		PS.LightFunctionParameters2 = FVector2f(LightSceneInfo->Proxy->GetLightFunctionFadeDistance(), LightSceneInfo->Proxy->GetLightFunctionDisabledBrightness());
+		PS.DeferredLightUniforms = TUniformBufferRef<FDeferredLightUniformStruct>::CreateUniformBufferImmediate(GetDeferredLightParameters(View, *LightSceneInfo), EUniformBufferUsage::UniformBuffer_SingleFrame);
+
+		return PS;
+	}
+};
+
+IMPLEMENT_MATERIAL_SHADER_TYPE(, FMobileLocalLightFunctionPS, TEXT("/Engine/Private/MobileLocalLightsBuffer.usf"), TEXT("MainLightFunction"), SF_Pixel);
+
+BEGIN_SHADER_PARAMETER_STRUCT(FLocalLightFunctionParameters, )
+	SHADER_PARAMETER_STRUCT_INCLUDE(FSceneTextureShaderParameters, SceneTextures)
+	RENDER_TARGET_BINDING_SLOTS()
+END_SHADER_PARAMETER_STRUCT()
+
+static bool TryGetLightFunctionShaders(FMaterialRenderProxy const*& OutMaterialProxy, FMaterial const*& OutMaterial, FMaterialShaders& OutShaders)
+{
+	while (OutMaterialProxy)
+	{
+		OutMaterial = OutMaterialProxy->GetMaterialNoFallback(ERHIFeatureLevel::ES3_1);
+		if (OutMaterial && OutMaterial->IsLightFunction())
+		{
+			FMaterialShaderTypes ShaderTypes;
+			ShaderTypes.AddShaderType<FMobileLocalLightFunctionPS>();
+			if (OutMaterial->TryGetShaders(ShaderTypes, nullptr, OutShaders))
+			{
+				return true;
+			}
+		}
+		OutMaterialProxy = OutMaterialProxy->GetFallback(ERHIFeatureLevel::ES3_1);
+	}
+	return false;
+}
+
+
+void FMobileSceneRenderer::RenderMobileLocalLightsBuffer(FRDGBuilder& GraphBuilder, FSceneTextures& SceneTextures, bool bIsPrepass, const FSortedLightSetSceneInfo& SortedLights)
+{
+	if (!CompileShaderPermutationsForMobileLocalLightsBuffer(ShaderPlatform) ||
 		(bIsPrepass != MobileLocalLightsBufferPrepassEnabled(ShaderPlatform)) || 
 		IsMobileDeferredShadingEnabled(ShaderPlatform))
 	{
@@ -181,7 +277,7 @@ void FMobileSceneRenderer::RenderMobileLocalLightsBuffer(FRDGBuilder& GraphBuild
 				RDG_EVENT_NAME("RenderMobileLocalLightsBuffer %s", bIsPrepass ? TEXT("Prepass") : TEXT("PostProcess")),
 				PassParameters,
 				ERDGPassFlags::Raster,
-				[PassParameters, VertexShader, PixelShader, &View, GroupSize, bIsPrepass](FRHICommandList& RHICmdList)
+				[PassParameters, VertexShader, PixelShader, &View, GroupSize, bIsPrepass, &SortedLights](FRHICommandList& RHICmdList)
 				{
 					FGraphicsPipelineStateInitializer GraphicsPSOInit;
 					RHICmdList.ApplyCachedRenderTargets(GraphicsPSOInit);
@@ -216,7 +312,84 @@ void FMobileSceneRenderer::RenderMobileLocalLightsBuffer(FRDGBuilder& GraphBuild
 						0,
 						2,
 						GroupSize.X * GroupSize.Y);
+
+
+					// Draws a pass for each visible light with a light function and blends on top of the color texture 
+					// generated by last pass
+					for (const FSortedLightSceneInfo& SortedLightSceneInfo : SortedLights.SortedLights)
+					{
+						 // Directional lights are currently not supported
+						if (SortedLightSceneInfo.SortKey.Fields.LightType == LightType_Directional)
+						{
+							continue;
+						}
+
+						if (!SortedLightSceneInfo.SortKey.Fields.bLightFunction)
+						{
+							continue;
+						}
+
+						const FLightSceneInfo* const LightSceneInfo = SortedLightSceneInfo.LightSceneInfo;
+						const FSphere LightBounds = LightSceneInfo->Proxy->GetBoundingSphere();
+
+						const float FadeAlpha = GetLightFunctionFadeFraction(View, LightBounds);
+						// Don't draw the light function if it has completely faded out
+						if (FadeAlpha < 1.0f / 256.0f)
+						{
+							continue;
+						}
+
+						FMaterialShaders MaterialShaders;
+						const FMaterialRenderProxy* MaterialProxyForRendering = LightSceneInfo->Proxy->GetLightFunctionMaterial();
+						const FMaterial* MaterialForRendering = nullptr;
+						if (!TryGetLightFunctionShaders(MaterialProxyForRendering, MaterialForRendering, MaterialShaders))
+						{
+							UE_LOG(LogTemp, Error, TEXT("Light function shader for light %d not found."), LightSceneInfo->Id);
+							continue;
+						}
+
+						FDeferredLightVS::FPermutationDomain PermutationVectorVS;
+						PermutationVectorVS.Set<FDeferredLightVS::FRadialLight>(true);
+						TShaderMapRef<FDeferredLightVS> LightFunctionVertexShader(View.ShaderMap, PermutationVectorVS);
+						FDeferredLightVS::FParameters ParametersVS = FDeferredLightVS::GetParameters(View, LightSceneInfo, false);
+						ParametersVS.View = TUniformBufferBinding(View.ViewUniformBuffer, EUniformBufferBindingFlags::Shader);
+						TShaderRef<FMobileLocalLightFunctionPS> LightFunctionPixelShader;
+						MaterialShaders.TryGetPixelShader(LightFunctionPixelShader);
+						FMobileLocalLightFunctionPS::FParameters ParametersPS = LightFunctionPixelShader->GetParameters(View, LightSceneInfo, FadeAlpha);
+
+						// Directional RT was already generated in main pass
+						GraphicsPSOInit.BlendState = TStaticBlendState<CW_RGB, BO_Add, BF_One, BF_One, BO_Add, BF_One, BF_Zero, CW_NONE>::GetRHI(); 
+
+						GraphicsPSOInit.PrimitiveType = PT_TriangleList;
+						GraphicsPSOInit.BoundShaderState.VertexDeclarationRHI = GetVertexDeclarationFVector4();
+						GraphicsPSOInit.BoundShaderState.VertexShaderRHI = LightFunctionVertexShader.GetVertexShader();
+						GraphicsPSOInit.BoundShaderState.PixelShaderRHI = LightFunctionPixelShader.GetPixelShader();
+
+						if (((FVector)View.ViewMatrices.GetViewOrigin() - LightBounds.Center).SizeSquared() < FMath::Square(LightBounds.W * 1.05f + View.NearClippingDistance * 2.0f))
+						{
+							// Render backfaces with depth tests disabled since the camera is inside (or close to inside) the light function geometry
+							GraphicsPSOInit.RasterizerState = View.bReverseCulling ? TStaticRasterizerState<FM_Solid, CM_CW>::GetRHI() : TStaticRasterizerState<FM_Solid, CM_CCW>::GetRHI();
+						}
+
+						// Set the light's scissor rectangle.
+						LightSceneInfo->Proxy->SetScissorRect(RHICmdList, View, View.ViewRect);
+
+						SetGraphicsPipelineState(RHICmdList, GraphicsPSOInit, 0);
+
+						SetShaderParameters(RHICmdList, LightFunctionVertexShader, LightFunctionVertexShader.GetVertexShader(), ParametersVS);
+						SetShaderParametersMixedPS(RHICmdList, LightFunctionPixelShader, ParametersPS, View, MaterialProxyForRendering, *MaterialForRendering);
+
+						// Project the light function using a sphere around the light
+						if (SortedLightSceneInfo.SortKey.Fields.LightType == LightType_Spot)
+						{
+							StencilingGeometry::DrawCone(RHICmdList);
+						}
+						else
+						{
+							StencilingGeometry::DrawSphere(RHICmdList);
+						}
+					}
 				});
 		}
 	}
-}
\ No newline at end of file
+}
diff --git a/Engine/Source/Runtime/Renderer/Private/MobileShadingRenderer.cpp b/Engine/Source/Runtime/Renderer/Private/MobileShadingRenderer.cpp
--- a/Engine/Source/Runtime/Renderer/Private/MobileShadingRenderer.cpp
+++ b/Engine/Source/Runtime/Renderer/Private/MobileShadingRenderer.cpp
@@ -1081,7 +1081,7 @@ void FMobileSceneRenderer::Render(FRDGBuilder& GraphBuilder)
 
 		if (bRendererOutputFinalSceneColor)
 		{
-			RenderMobileLocalLightsBuffer(GraphBuilder, SceneTextures, true);
+			RenderMobileLocalLightsBuffer(GraphBuilder, SceneTextures, true, SortedLightSet);
 		}
 	}
 
@@ -1122,7 +1122,7 @@ void FMobileSceneRenderer::Render(FRDGBuilder& GraphBuilder)
 	
 		GraphBuilder.SetCommandListStat(GET_STATID(STAT_CLMM_Post));
 
-		RenderMobileLocalLightsBuffer(GraphBuilder, SceneTextures, false);
+		RenderMobileLocalLightsBuffer(GraphBuilder, SceneTextures, false, SortedLightSet);
 
 		FRendererModule& RendererModule = static_cast<FRendererModule&>(GetRendererModule());
 		RendererModule.RenderPostOpaqueExtensions(GraphBuilder, Views, SceneTextures);
diff --git a/Engine/Source/Runtime/Renderer/Private/SceneRendering.h b/Engine/Source/Runtime/Renderer/Private/SceneRendering.h
--- a/Engine/Source/Runtime/Renderer/Private/SceneRendering.h
+++ b/Engine/Source/Runtime/Renderer/Private/SceneRendering.h
@@ -2650,7 +2650,7 @@ protected:
 	void RenderMaskedPrePass(FRHICommandList& RHICmdList, const FViewInfo& View);
 	void RenderFullDepthPrepass(FRDGBuilder& GraphBuilder, TArrayView<FViewInfo> InViews, FSceneTextures& SceneTextures, bool bIsSceneCaptureRenderPass=false);
 
-	void RenderMobileLocalLightsBuffer(FRDGBuilder& GraphBuilder, FSceneTextures& SceneTextures, bool bIsPrepass);
+	void RenderMobileLocalLightsBuffer(FRDGBuilder& GraphBuilder, FSceneTextures& SceneTextures, bool bIsPrepass, const FSortedLightSetSceneInfo& SortedLights);
 
 	/** Renders the opaque base pass for mobile. */
 	void RenderMobileBasePass(FRHICommandList& RHICmdList, const FViewInfo& View, const FInstanceCullingDrawParams* InstanceCullingDrawParams);
