LFV - Allocation and binding of the culled instance buffer

#prelfight shader

diff --git a/Engine/Shaders/Private/LocalFogVolumes/LocalFogVolumeCommon.ush b/Engine/Shaders/Private/LocalFogVolumes/LocalFogVolumeCommon.ush
--- a/Engine/Shaders/Private/LocalFogVolumes/LocalFogVolumeCommon.ush
+++ b/Engine/Shaders/Private/LocalFogVolumes/LocalFogVolumeCommon.ush
@@ -37,15 +37,23 @@ struct FLocalFogVolumeGPUInstanceData
 #if defined(LFVStruct)
 
 #define LFVInsts LFVStruct.LFV.LocalFogVolumeInstances
+#define LFVCullDataTex LFVStruct.LFV.LocalFogVolumeCullDataTexture
+
 #define LFVInstcount LFVStruct.LFV.LocalFogVolumeInstanceCount
+#define LFVTilePixelSize LFVStruct.LFV.LocalFogVolumeTilePixelSize
 
 #else
 
 uint LFV_LocalFogVolumeInstanceCount;
+uint LFV_LocalFogVolumeTilePixelSize;
 Buffer<float4> LFV_LocalFogVolumeInstances;
+Texture2DArray<uint> LFV_LocalFogVolumeCullDataTexture;
 
 #define LFVInsts LFV_LocalFogVolumeInstances
+#define LFVCullDataTex LFV_LocalFogVolumeCullDataTexture
+
 #define LFVInstcount LFV_LocalFogVolumeInstanceCount
+#define LFVTilePixelSize LFV_LocalFogVolumeTilePixelSize
 
 #endif
 
diff --git a/Engine/Shaders/Public/ShaderVersion.ush b/Engine/Shaders/Public/ShaderVersion.ush
--- a/Engine/Shaders/Public/ShaderVersion.ush
+++ b/Engine/Shaders/Public/ShaderVersion.ush
@@ -2,4 +2,4 @@
 // Each time the console command is executed it generates a new GUID. As this file is included
 // in Platform.ush (which should be included in any shader) it allows to invalidate the shader DDC.
 // 
-#pragma message("UESHADERMETADATA_VERSION E7D109A8-D147-450F-A96D-9A3E24EB1B05")
+#pragma message("UESHADERMETADATA_VERSION 1CBD3DC4-205B-47DA-9533-FC9B7440B383")
diff --git a/Engine/Source/Runtime/Renderer/Private/LocalFogVolumeRendering.cpp b/Engine/Source/Runtime/Renderer/Private/LocalFogVolumeRendering.cpp
--- a/Engine/Source/Runtime/Renderer/Private/LocalFogVolumeRendering.cpp
+++ b/Engine/Source/Runtime/Renderer/Private/LocalFogVolumeRendering.cpp
@@ -64,18 +64,24 @@ void SetDummyLocalFogVolumeForView(FRDGBuilder& GraphBuilder, FViewInfo& View)
 	// The size of the structure must be a multiple of FUintVector4.
 	static_assert(sizeof(FLocalFogVolumeGPUInstanceData) == UintVec4CountInLocalFogVolumeGPUInstanceData * SizeOfUintVec4);
 
-	View.LocalFogVolumeViewData.GPUInstanceCount = 0;
+	View.LocalFogVolumeViewData.GPUInstanceCount			= 0;
 
 	static FLocalFogVolumeGPUInstanceData DummyData;
-	View.LocalFogVolumeViewData.GPUInstanceDataBuffer = CreateVertexBuffer(GraphBuilder, TEXT("DUMMYLocalFogVolumeGPUInstanceDataBuffer"),
+	View.LocalFogVolumeViewData.GPUInstanceDataBuffer		= CreateVertexBuffer(GraphBuilder, TEXT("DUMMYLocalFogVolumeGPUInstanceDataBuffer"),
 		FRDGBufferDesc::CreateBufferDesc(SizeOfUintVec4, UintVec4CountInLocalFogVolumeGPUInstanceData), &DummyData, sizeof(FLocalFogVolumeGPUInstanceData) * 1, ERDGInitialDataFlags::NoCopy);
 
-	View.LocalFogVolumeViewData.GPUInstanceDataBufferSRV = GraphBuilder.CreateSRV(View.LocalFogVolumeViewData.GPUInstanceDataBuffer, PF_A32B32G32R32F);
+	View.LocalFogVolumeViewData.GPUInstanceDataBufferSRV	= GraphBuilder.CreateSRV(View.LocalFogVolumeViewData.GPUInstanceDataBuffer, PF_A32B32G32R32F);
 
-	View.LocalFogVolumeViewData.UniformParametersStruct.LocalFogVolumeInstanceCount = View.LocalFogVolumeViewData.GPUInstanceCount;
-	View.LocalFogVolumeViewData.UniformParametersStruct.LocalFogVolumeInstances = View.LocalFogVolumeViewData.GPUInstanceDataBufferSRV;
-	View.LocalFogVolumeViewData.UniformBuffer = GraphBuilder.CreateUniformBuffer(&View.LocalFogVolumeViewData.UniformParametersStruct);
-};
+	View.LocalFogVolumeViewData.CullDataTextureArray		= GSystemTextures.GetZeroUIntArrayDummy(GraphBuilder);
+	View.LocalFogVolumeViewData.CullDataTextureArraySRV		= GraphBuilder.CreateSRV(View.LocalFogVolumeViewData.CullDataTextureArray);
+	View.LocalFogVolumeViewData.CullDataTextureArrayUAV		= nullptr;	// Should never be written by culling passes if there are no instances.
+
+	View.LocalFogVolumeViewData.UniformParametersStruct.LocalFogVolumeInstanceCount		= View.LocalFogVolumeViewData.GPUInstanceCount;
+	View.LocalFogVolumeViewData.UniformParametersStruct.LocalFogVolumeTilePixelSize		= GetLocalFogVolumeTilePixelSize();
+	View.LocalFogVolumeViewData.UniformParametersStruct.LocalFogVolumeInstances			= View.LocalFogVolumeViewData.GPUInstanceDataBufferSRV;
+	View.LocalFogVolumeViewData.UniformParametersStruct.LocalFogVolumeCullDataTexture	= View.LocalFogVolumeViewData.CullDataTextureArraySRV;
+	View.LocalFogVolumeViewData.UniformBuffer											= GraphBuilder.CreateUniformBuffer(&View.LocalFogVolumeViewData.UniformParametersStruct);
+}
 
 
 /*=============================================================================
@@ -200,15 +206,35 @@ void CreateViewLocalFogVolumeBufferSRV(FViewInfo& View, FRDGBuilder& GraphBuilde
 	// 3. Allocate buffer and initialize with sorted data to upload to GPU
 	const uint32 AllLocalFogVolumeInstanceBytesFinal = sizeof(FLocalFogVolumeGPUInstanceData) * SortingData.LocalFogVolumeInstanceCountFinal;
 
-	View.LocalFogVolumeViewData.GPUInstanceCount = SortingData.LocalFogVolumeInstanceCountFinal;
-	View.LocalFogVolumeViewData.GPUInstanceDataBuffer = CreateVertexBuffer(
-		GraphBuilder, TEXT("LocalFogVolumeGPUInstanceDataBuffer"),
+	View.LocalFogVolumeViewData.GPUInstanceCount			= SortingData.LocalFogVolumeInstanceCountFinal;
+	View.LocalFogVolumeViewData.GPUInstanceDataBuffer		= CreateVertexBuffer(
+		GraphBuilder, TEXT("LocalFogVolume.GPUInstanceDataBuffer"),
 		FRDGBufferDesc::CreateBufferDesc(SizeOfUintVec4, SortingData.LocalFogVolumeInstanceCountFinal * UintVec4CountInLocalFogVolumeGPUInstanceData),
 		LocalFogVolumeGPUSortedInstanceData, AllLocalFogVolumeInstanceBytesFinal, ERDGInitialDataFlags::NoCopy);
-	View.LocalFogVolumeViewData.GPUInstanceDataBufferSRV = GraphBuilder.CreateSRV(View.LocalFogVolumeViewData.GPUInstanceDataBuffer, PF_A32B32G32R32F);
-
-	View.LocalFogVolumeViewData.UniformParametersStruct.LocalFogVolumeInstanceCount = View.LocalFogVolumeViewData.GPUInstanceCount;
-	View.LocalFogVolumeViewData.UniformParametersStruct.LocalFogVolumeInstances = View.LocalFogVolumeViewData.GPUInstanceDataBufferSRV;
+	View.LocalFogVolumeViewData.GPUInstanceDataBufferSRV	= GraphBuilder.CreateSRV(View.LocalFogVolumeViewData.GPUInstanceDataBuffer, PF_A32B32G32R32F);
+
+	// Create the texture that will contain the tiled culled result: count in the first slice and indices in the remaining slices
+	const uint32 LocalFogVolumeTilePixelSize				= GetLocalFogVolumeTilePixelSize();
+#if 1
+	// Allocate dummy resource while this is not yet used without the tile culling pass.
+	View.LocalFogVolumeViewData.CullDataTextureArray = GSystemTextures.GetZeroUIntArrayDummy(GraphBuilder);
+	View.LocalFogVolumeViewData.CullDataTextureArraySRV = GraphBuilder.CreateSRV(View.LocalFogVolumeViewData.CullDataTextureArray);
+	View.LocalFogVolumeViewData.CullDataTextureArrayUAV = nullptr;	// Should never be written by culling passes if there are no instances.
+#else
+	const FIntPoint CullDataTextureResolution				= FIntPoint::DivideAndRoundUp(View.ViewRect.Size(), FIntPoint(LocalFogVolumeTilePixelSize, LocalFogVolumeTilePixelSize));
+	const uint32 CullDataTextureSliceCount					= LocalFogVolumeTileMaxInstanceCount + 1; // +1 because the first slice is the culled instance count
+	FRDGTextureDesc Texture2DArrayDesc(FRDGTextureDesc::Create2DArray(CullDataTextureResolution, PF_R8_UINT, FClearValueBinding(EClearBinding::ENoneBound), TexCreate_ShaderResource | TexCreate_UAV | TexCreate_ReduceMemoryWithTilingMode | TexCreate_3DTiling, CullDataTextureSliceCount));
+	// LFV TODO, consider FFastVramConfig onto Texture2DArrayDesc
+
+	View.LocalFogVolumeViewData.CullDataTextureArray		= GraphBuilder.CreateTexture(Texture2DArrayDesc, TEXT("LocalFogVolume.CullingDataTexture"));
+	View.LocalFogVolumeViewData.CullDataTextureArraySRV		= GraphBuilder.CreateSRV(FRDGTextureSRVDesc::Create(View.LocalFogVolumeViewData.CullDataTextureArray));
+	View.LocalFogVolumeViewData.CullDataTextureArrayUAV		= GraphBuilder.CreateUAV(FRDGTextureUAVDesc(View.LocalFogVolumeViewData.CullDataTextureArray));
+#endif
+
+	View.LocalFogVolumeViewData.UniformParametersStruct.LocalFogVolumeInstanceCount		= View.LocalFogVolumeViewData.GPUInstanceCount;
+	View.LocalFogVolumeViewData.UniformParametersStruct.LocalFogVolumeTilePixelSize		= LocalFogVolumeTilePixelSize;
+	View.LocalFogVolumeViewData.UniformParametersStruct.LocalFogVolumeInstances			= View.LocalFogVolumeViewData.GPUInstanceDataBufferSRV;
+	View.LocalFogVolumeViewData.UniformParametersStruct.LocalFogVolumeCullDataTexture	= View.LocalFogVolumeViewData.CullDataTextureArraySRV;
 	View.LocalFogVolumeViewData.UniformBuffer = GraphBuilder.CreateUniformBuffer(&View.LocalFogVolumeViewData.UniformParametersStruct);
 }
 
diff --git a/Engine/Source/Runtime/Renderer/Private/LocalFogVolumeRendering.h b/Engine/Source/Runtime/Renderer/Private/LocalFogVolumeRendering.h
--- a/Engine/Source/Runtime/Renderer/Private/LocalFogVolumeRendering.h
+++ b/Engine/Source/Runtime/Renderer/Private/LocalFogVolumeRendering.h
@@ -21,7 +21,9 @@ struct FMinimalSceneTextures;
 
 BEGIN_SHADER_PARAMETER_STRUCT(FLocalFogVolumeUniformParameters, )
 	SHADER_PARAMETER_RDG_BUFFER_SRV(Buffer<float4>, LocalFogVolumeInstances)
+	SHADER_PARAMETER_RDG_TEXTURE_SRV(Texture2DArray<uint>, LocalFogVolumeCullDataTexture)
 	SHADER_PARAMETER(uint32, LocalFogVolumeInstanceCount)
+	SHADER_PARAMETER(uint32, LocalFogVolumeTilePixelSize)
 END_SHADER_PARAMETER_STRUCT()
 
 /*=============================================================================
@@ -93,7 +95,9 @@ struct FLocalFogVolumeViewData
 
 	FLocalFogVolumeUniformParameters UniformParametersStruct;
 
-	FRDGTextureRef TileInstanceIndexTextureArray; // First slice is the instance count, later slices are instance indices.
+	FRDGTextureRef CullDataTextureArray;		// First slice is the instance count, later slices are instance indices.
+	FRDGTextureSRVRef CullDataTextureArraySRV;
+	FRDGTextureUAVRef CullDataTextureArrayUAV;
 };
 
 
