Added scene render command pipe which moves all scene commands to run asynchronously.
 - Create / Destroy render thread resources is now deferred until scene update and runs async.
 - Added bConstrainToRenderThread option to FPrimitiveSceneProxy, which controls whether the scene pipe is used for proxy virtual function calls.

diff --git a/Engine/Plugins/FX/Niagara/Source/Niagara/Private/NiagaraSystemRenderData.cpp b/Engine/Plugins/FX/Niagara/Source/Niagara/Private/NiagaraSystemRenderData.cpp
--- a/Engine/Plugins/FX/Niagara/Source/Niagara/Private/NiagaraSystemRenderData.cpp
+++ b/Engine/Plugins/FX/Niagara/Source/Niagara/Private/NiagaraSystemRenderData.cpp
@@ -7,6 +7,7 @@
 #include "NiagaraSystem.h"
 #include "PrimitiveViewRelevance.h"
 #include "SceneView.h"
+#include "Rendering/RenderCommandPipes.h"
 
 DECLARE_RENDER_COMMAND_PIPE(NiagaraDynamicData, );
 
@@ -310,8 +311,8 @@ void FNiagaraSystemRenderData::RecacheRenderers(const FNiagaraSystemInstance& Sy
 
 	// NOTE: Since this object and its render thread resources may be concurrently accessed on the render thread, we have to create the renderers and pass them off to
 	// replace the current ones on the render thread's time line
-	ENQUEUE_RENDER_COMMAND(NiagaraRecacheRenderers)(
-		[this, EmitterRenderers_Copy=EmitterRenderers_GT, TaskEvent = MoveTemp(TaskEvent)](FRHICommandListImmediate& RHICmdList) mutable
+	ENQUEUE_RENDER_COMMAND(NiagaraRecacheRenderers)(UE::RenderCommandPipe::Scene,
+		[this, EmitterRenderers_Copy=EmitterRenderers_GT, TaskEvent = MoveTemp(TaskEvent)](FRHICommandListBase& RHICmdList) mutable
 		{
 			TaskEvent.Wait();
 
diff --git a/Engine/Source/Runtime/Engine/Classes/Engine/InstancedStaticMesh.h b/Engine/Source/Runtime/Engine/Classes/Engine/InstancedStaticMesh.h
--- a/Engine/Source/Runtime/Engine/Classes/Engine/InstancedStaticMesh.h
+++ b/Engine/Source/Runtime/Engine/Classes/Engine/InstancedStaticMesh.h
@@ -435,6 +435,7 @@ private:
 	/**
 	 * Called to update the PerInstanceBounds/PerInstanceTransforms arrays whenever the instance array is modified
 	 */
+	void UpdateBoundsTransforms_RenderThread();
 	void UpdateBoundsTransforms_Concurrent();
 	void UpdateBoundsTransforms();
 	void EnsureInstanceDataUpdated(bool bForceUpdate = false);
diff --git a/Engine/Source/Runtime/Engine/Classes/PhysicsField/PhysicsFieldComponent.h b/Engine/Source/Runtime/Engine/Classes/PhysicsField/PhysicsFieldComponent.h
--- a/Engine/Source/Runtime/Engine/Classes/PhysicsField/PhysicsFieldComponent.h
+++ b/Engine/Source/Runtime/Engine/Classes/PhysicsField/PhysicsFieldComponent.h
@@ -144,7 +144,7 @@ public:
 	virtual void InitRHI(FRHICommandListBase& RHICmdList) override;
 
 	/** Update RHI resources. */
-	void UpdateResource(FRHICommandListImmediate& RHICmdList,
+	void UpdateResource(FRHICommandList& RHICmdList,
 		const TStaticArray<int32, EFieldPhysicsType::Field_PhysicsType_Max + 1>& TargetsOffsetsDatas, const TArray<int32>& NodesOffsetsDatas, const TArray<float>& NodesParamsDatas,
 		const TArray<FVector>& TargetsMinDatas, const TArray<FVector>& TargetsMaxDatas, const float TimeSeconds, 
 		const TArray<FVector4>& BoundsMinDatas, const TArray<FVector4>& BoundsMaxDatas, const TStaticArray<int32, EFieldPhysicsType::Field_PhysicsType_Max + 1>& BoundsOffsetsDatas);
diff --git a/Engine/Source/Runtime/Engine/Private/InstancedStaticMesh.cpp b/Engine/Source/Runtime/Engine/Private/InstancedStaticMesh.cpp
--- a/Engine/Source/Runtime/Engine/Private/InstancedStaticMesh.cpp
+++ b/Engine/Source/Runtime/Engine/Private/InstancedStaticMesh.cpp
@@ -44,6 +44,7 @@
 #include "Elements/Interfaces/TypedElementWorldInterface.h"
 #include "UObject/UObjectIterator.h"
 #include "GenericPlatform/ICursor.h"
+#include "Rendering/RenderCommandPipes.h"
 
 #include "InstancedStaticMeshSceneProxyDesc.h"
 
@@ -572,8 +573,8 @@ void FStaticMeshInstanceBuffer::UpdateFromCommandBuffer_Concurrent(FInstanceUpda
 	// Any query of number of render instances on game thread should use this instead of InstanceData->GetNumInstances();
 	CmdBuffer.NumEditInstances = NewCmdBuffer->NumAdds + InstanceData->GetNumInstances();
 		
-	ENQUEUE_RENDER_COMMAND(InstanceBuffer_UpdateFromPreallocatedData)(
-		[InstanceBuffer, NewCmdBuffer](FRHICommandListImmediate& RHICmdList)
+	ENQUEUE_RENDER_COMMAND(InstanceBuffer_UpdateFromPreallocatedData)(UE::RenderCommandPipe::Scene,
+		[InstanceBuffer, NewCmdBuffer](FRHICommandListBase& RHICmdList)
 		{
 			InstanceBuffer->UpdateFromCommandBuffer_RenderThread(RHICmdList, *NewCmdBuffer);
 			delete NewCmdBuffer;
@@ -961,8 +962,8 @@ void FInstancedStaticMeshVertexFactory::Copy(const FInstancedStaticMeshVertexFac
 	FInstancedStaticMeshVertexFactory* VertexFactory = this;
 	const FLocalVertexFactory::FDataType* DataCopy = &Other.Data;
 	const FInstancedStaticMeshDataType* InstanceDataCopy = &Other.InstanceData;
-	ENQUEUE_RENDER_COMMAND(FInstancedStaticMeshVertexFactoryCopyData)(
-	[VertexFactory, DataCopy, InstanceDataCopy](FRHICommandListImmediate& RHICmdList)
+	ENQUEUE_RENDER_COMMAND(FInstancedStaticMeshVertexFactoryCopyData)(UE::RenderCommandPipe::Scene,
+	[VertexFactory, DataCopy, InstanceDataCopy]
 	{
 		VertexFactory->Data = *DataCopy;
 		VertexFactory->InstanceData = *InstanceDataCopy;
@@ -1254,12 +1255,6 @@ void FInstancedStaticMeshRenderData::InitVertexFactories()
 	{
 		VertexFactories.Add(new FInstancedStaticMeshVertexFactory(FeatureLevel));
 	}
-
-	ENQUEUE_RENDER_COMMAND(InstancedStaticMeshRenderData_InitVertexFactories)(
-		[this](FRHICommandListBase& RHICmdList)
-		{
-			BindBuffersToVertexFactories(RHICmdList);
-		});
 }
 
 void FInstancedStaticMeshRenderData::RegisterSpeedTreeWind(const FInstancedStaticMeshSceneProxyDesc* InProxyDesc)
@@ -1314,8 +1309,8 @@ void FPerInstanceRenderData::UpdateFromPreallocatedData(FStaticMeshInstanceData&
 
 	FStaticMeshInstanceDataPtr InInstanceBufferDataPtr = InstanceBuffer_GameThread;
 	FStaticMeshInstanceBuffer* InInstanceBuffer = &InstanceBuffer;
-	ENQUEUE_RENDER_COMMAND(FInstanceBuffer_UpdateFromPreallocatedData)(
-		[InInstanceBufferDataPtr, InInstanceBuffer, this](FRHICommandListImmediate& RHICmdList)
+	ENQUEUE_RENDER_COMMAND(FInstanceBuffer_UpdateFromPreallocatedData)(UE::RenderCommandPipe::Scene,
+		[InInstanceBufferDataPtr, InInstanceBuffer, this](FRHICommandListBase& RHICmdList)
 		{
 			// The assignment to InstanceData shared pointer kills the old data
 			// If UpdateBoundsTask is in-flight it will crash
@@ -1325,43 +1320,44 @@ void FPerInstanceRenderData::UpdateFromPreallocatedData(FStaticMeshInstanceData&
 			{
 				InInstanceBuffer->UpdateRHI(RHICmdList);
 			}
-			UpdateBoundsTransforms_Concurrent();
+			UpdateBoundsTransforms_RenderThread();
 		}
 	);
 }
 
 void FPerInstanceRenderData::UpdateBoundsTransforms_Concurrent()
 {
-	// Enqueue a render command to create a task to update the buffer data.
-	// Yes double-wrapping a lambda looks a little silly, but the only safe way to update the render data
-	// is to issue this task from the rendering thread.
-	ENQUEUE_RENDER_COMMAND(FInstanceBuffer_UpdateBoundsTransforms)(
-		[this](FRHICommandListImmediate& RHICmdList)
-		{
-			bBoundsTransformsDirty = true;
-			if (!IsRayTracingEnabled() || !CVarRayTracingRenderInstances.GetValueOnRenderThread())
-			{
-				return;
-			}
+	ENQUEUE_RENDER_COMMAND(FInstanceBuffer_UpdateBoundsTransforms)(UE::RenderCommandPipe::Scene,
+		[this](FRHICommandListBase& RHICmdList)
+	{
+		UpdateBoundsTransforms_RenderThread();
+	});
+}
 
-			FGraphEventArray Prerequisites{};
-			if (UpdateBoundsTask.IsValid())
-			{
-				// There's already a task either in flight or unconsumed, but the instance data has now changed so its result might be incorrect.
-				// This new task should run after the first one completes, so make the old one a prerequisite of the new one.
-				Prerequisites = FGraphEventArray{ UpdateBoundsTask };
-				UE_LOG(LogStaticMesh, Warning, TEXT("Unconsumed ISM bounds/transforms update task, we did more work than necessary"));
-			}
+void FPerInstanceRenderData::UpdateBoundsTransforms_RenderThread()
+{
+	bBoundsTransformsDirty = true;
+	if (!IsRayTracingEnabled() || !CVarRayTracingRenderInstances.GetValueOnRenderThread())
+	{
+		return;
+	}
 
-			UpdateBoundsTask = FFunctionGraphTask::CreateAndDispatchWhenReady(
-			[this]()
-			{
-				UpdateBoundsTransforms();
-			},
-			TStatId(),
-			&Prerequisites
-			);
-		}
+	FGraphEventArray Prerequisites{};
+	if (UpdateBoundsTask.IsValid())
+	{
+		// There's already a task either in flight or unconsumed, but the instance data has now changed so its result might be incorrect.
+		// This new task should run after the first one completes, so make the old one a prerequisite of the new one.
+		Prerequisites = FGraphEventArray{ UpdateBoundsTask };
+		UE_LOG(LogStaticMesh, Warning, TEXT("Unconsumed ISM bounds/transforms update task, we did more work than necessary"));
+	}
+
+	UpdateBoundsTask = FFunctionGraphTask::CreateAndDispatchWhenReady(
+	[this]()
+	{
+		UpdateBoundsTransforms();
+	},
+	TStatId(),
+	&Prerequisites
 	);
 }
 
@@ -1409,12 +1405,12 @@ void FPerInstanceRenderData::UpdateBoundsTransforms()
 
 void FPerInstanceRenderData::EnsureInstanceDataUpdated(bool bForceUpdate)
 {
-	check(IsInRenderingThread());
+	check(IsInParallelRenderingThread());
 
 	// wait for bounds/transforms update to complete
 	if (UpdateBoundsTask.IsValid())
 	{
-		UpdateBoundsTask->Wait(ENamedThreads::GetRenderThread_Local());
+		UpdateBoundsTask->Wait();
 		UpdateBoundsTask.SafeRelease();
 		bBoundsTransformsDirty = false;
 	}
@@ -1453,7 +1449,7 @@ void FPerInstanceRenderData::UpdateFromCommandBuffer(FInstanceUpdateCmdBuffer& C
 {
 	// UpdateFromCommandBuffer reallocates InstanceData in InstanceBuffer
 	// If UpdateBoundsTask is in-flight it will crash
-	ENQUEUE_RENDER_COMMAND(EnsureInstanceDataUpdatedCmd)(
+	ENQUEUE_RENDER_COMMAND(EnsureInstanceDataUpdatedCmd)(UE::RenderCommandPipe::Scene,
 		[this](FRHICommandList&) {
 		EnsureInstanceDataUpdated();
 	});
@@ -1975,6 +1971,8 @@ void FInstancedStaticMeshSceneProxy::CreateRenderThreadResources(FRHICommandList
 {
 	FStaticMeshSceneProxy::CreateRenderThreadResources(RHICmdList);
 
+	InstancedRenderData.BindBuffersToVertexFactories(RHICmdList);
+
 	const bool bCanUseGPUScene = UseGPUScene(GetScene().GetShaderPlatform(), GetScene().GetFeatureLevel());
 	
 	// Flush upload of GPU data for ISM/HISM
@@ -3396,8 +3394,8 @@ void UInstancedStaticMeshComponent::ReleasePerInstanceRenderData()
 		PerInstanceRenderData.Reset();
 
 		FPerInstanceRenderDataPtr* InCleanupRenderDataPtr = CleanupRenderDataPtr;
-		ENQUEUE_RENDER_COMMAND(FReleasePerInstanceRenderData)(
-			[InCleanupRenderDataPtr](FRHICommandList& RHICmdList)
+		ENQUEUE_RENDER_COMMAND(FReleasePerInstanceRenderData)(UE::RenderCommandPipe::Scene,
+			[InCleanupRenderDataPtr]
 			{
 				// Destroy the shared pointer object we allocated on the heap.
 				// Resource will either be released here or by scene proxy on the render thread, whoever gets executed last
diff --git a/Engine/Source/Runtime/Engine/Private/PhysicsField/PhysicsFieldComponent.cpp b/Engine/Source/Runtime/Engine/Private/PhysicsField/PhysicsFieldComponent.cpp
--- a/Engine/Source/Runtime/Engine/Private/PhysicsField/PhysicsFieldComponent.cpp
+++ b/Engine/Source/Runtime/Engine/Private/PhysicsField/PhysicsFieldComponent.cpp
@@ -15,6 +15,7 @@
 #include "RHIContext.h"
 #include "RenderingThread.h"
 #include "SceneInterface.h"
+#include "Rendering/RenderCommandPipes.h"
 
 #include UE_INLINE_GENERATED_CPP_BY_NAME(PhysicsFieldComponent)
 
@@ -527,7 +528,7 @@ void FPhysicsFieldResource::UpdateBounds(const TArray<FVector>& TargetsMin, cons
 	}
 }
 
-void FPhysicsFieldResource::UpdateResource(FRHICommandListImmediate& RHICmdList, 
+void FPhysicsFieldResource::UpdateResource(FRHICommandList& RHICmdList, 
 	const TStaticArray<int32, EFieldPhysicsType::Field_PhysicsType_Max + 1>& TargetsOffsetsDatas, const TArray<int32>& NodesOffsetsDatas, const TArray<float>& NodesParamsDatas,
 	const TArray<FVector>& TargetsMinDatas, const TArray<FVector>& TargetsMaxDatas, const float TimeSeconds,
 	const TArray<FVector4>& BoundsMinDatas, const TArray<FVector4>& BoundsMaxDatas, const TStaticArray<int32, EFieldPhysicsType::Field_PhysicsType_Max + 1>& BoundsOffsetsDatas)
@@ -650,7 +651,7 @@ void FPhysicsFieldInstance::InitInstance( const TArray<EFieldPhysicsType>& Targe
 		FieldResource = new FPhysicsFieldResource(TargetCount, TargetTypes, VectorTargets, ScalarTargets, IntegerTargets, PhysicsTargets, PhysicsBounds, bBuildClipmap);
 
 		FPhysicsFieldResource* LocalFieldResource = FieldResource;
-		ENQUEUE_RENDER_COMMAND(FInitPhysicsFieldResourceCommand)(
+		ENQUEUE_RENDER_COMMAND(FInitPhysicsFieldResourceCommand)(UE::RenderCommandPipe::Scene,
 			[LocalFieldResource](FRHICommandList& RHICmdList)
 			{
 				LocalFieldResource->InitResource(RHICmdList);
@@ -663,12 +664,12 @@ void FPhysicsFieldInstance::ReleaseInstance()
 	if (FieldResource)
 	{
 		FPhysicsFieldResource* LocalFieldResource = FieldResource;
-		ENQUEUE_RENDER_COMMAND(FDestroyPhysicsFieldResourceCommand)(
-			[LocalFieldResource](FRHICommandList& RHICmdList)
-			{
-				LocalFieldResource->ReleaseResource();
-				delete LocalFieldResource;
-			});
+		ENQUEUE_RENDER_COMMAND(FDestroyPhysicsFieldResourceCommand)(UE::RenderCommandPipe::Scene,
+			[LocalFieldResource]
+		{
+			LocalFieldResource->ReleaseResource();
+			delete LocalFieldResource;
+		});
 		FieldResource = nullptr;
 	}
 
@@ -793,12 +794,12 @@ void FPhysicsFieldInstance::UpdateInstance(const float TimeSeconds, const bool b
 			const float LocalTimeSeconds = TimeSeconds;
 
 			FPhysicsFieldResource* LocalFieldResource = FieldResource;
-			ENQUEUE_RENDER_COMMAND(FUpdateFieldInstanceCommand)(
-				[LocalFieldResource, LocalNodesParams, LocalNodesOffsets, LocalTargetsOffsets, LocalTargetsMin, LocalTargetsMax, LocalTimeSeconds, LocalBoundsOffsets, LocalBoundsMin, LocalBoundsMax](FRHICommandListImmediate& RHICmdList)
-				{
-					LocalFieldResource->UpdateResource(RHICmdList, 
-						LocalTargetsOffsets, LocalNodesOffsets, LocalNodesParams, LocalTargetsMin, LocalTargetsMax, LocalTimeSeconds, LocalBoundsMin, LocalBoundsMax, LocalBoundsOffsets);
-				});
+			ENQUEUE_RENDER_COMMAND(FUpdateFieldInstanceCommand)(UE::RenderCommandPipe::Scene,
+				[LocalFieldResource, LocalNodesParams, LocalNodesOffsets, LocalTargetsOffsets, LocalTargetsMin, LocalTargetsMax, LocalTimeSeconds, LocalBoundsOffsets, LocalBoundsMin, LocalBoundsMax](FRHICommandList& RHICmdList)
+			{
+				LocalFieldResource->UpdateResource(RHICmdList, 
+				LocalTargetsOffsets, LocalNodesOffsets, LocalNodesParams, LocalTargetsMin, LocalTargetsMax, LocalTimeSeconds, LocalBoundsMin, LocalBoundsMax, LocalBoundsOffsets);
+			});
 		}
 	}
 }
@@ -840,11 +841,11 @@ void UPhysicsFieldComponent::DestroyRenderState_Concurrent()
 		GetWorld()->Scene->ResetPhysicsField();
 
 		FPhysicsFieldSceneProxy* SceneProxy = FieldProxy;
-		ENQUEUE_RENDER_COMMAND(FDestroySkyLightCommand)(
+		ENQUEUE_RENDER_COMMAND(FDestroySkyLightCommand)(UE::RenderCommandPipe::Scene,
 			[SceneProxy](FRHICommandList& RHICmdList)
-			{
-				delete SceneProxy;
-			});
+		{
+			delete SceneProxy;
+		});
 
 		FieldProxy = nullptr;
 	}
@@ -949,7 +950,7 @@ void UPhysicsFieldComponent::OnUnregister()
 			LocalInstance->ReleaseInstance();
 
 			FPhysicsFieldInstance* LocalFieldInstance = LocalInstance;
-			ENQUEUE_RENDER_COMMAND(FDestroyPhysicsFieldInstanceCommand)(
+			ENQUEUE_RENDER_COMMAND(FDestroyPhysicsFieldInstanceCommand)(UE::RenderCommandPipe::Scene,
 				[LocalFieldInstance](FRHICommandList& RHICmdList)
 				{
 					delete LocalFieldInstance;
diff --git a/Engine/Source/Runtime/Engine/Private/PrimitiveSceneProxy.cpp b/Engine/Source/Runtime/Engine/Private/PrimitiveSceneProxy.cpp
--- a/Engine/Source/Runtime/Engine/Private/PrimitiveSceneProxy.cpp
+++ b/Engine/Source/Runtime/Engine/Private/PrimitiveSceneProxy.cpp
@@ -433,6 +433,7 @@ FPrimitiveSceneProxy::FPrimitiveSceneProxy(const FPrimitiveSceneProxyDesc& InPro
 ,	bForceHidden(false)
 ,	bCollisionEnabled(InProxyDesc.IsCollisionEnabled())
 ,	bTreatAsBackgroundForOcclusion(InProxyDesc.bTreatAsBackgroundForOcclusion)
+,	bConstrainToRenderThread(false)
 ,	bVisibleInLumenScene(false)
 ,	bCanSkipRedundantTransformUpdates(true)
 ,	bGoodCandidateForCachedShadowmap(true)
diff --git a/Engine/Source/Runtime/Engine/Private/Rendering/RenderCommandPipes.cpp b/Engine/Source/Runtime/Engine/Private/Rendering/RenderCommandPipes.cpp
--- a/Engine/Source/Runtime/Engine/Private/Rendering/RenderCommandPipes.cpp
+++ b/Engine/Source/Runtime/Engine/Private/Rendering/RenderCommandPipes.cpp
@@ -2,4 +2,5 @@
 
 #include "Rendering/RenderCommandPipes.h"
 
-DEFINE_RENDER_COMMAND_PIPE(SkeletalMesh, ERenderCommandPipeFlags::None);
\ No newline at end of file
+DEFINE_RENDER_COMMAND_PIPE(SkeletalMesh, ERenderCommandPipeFlags::None);
+DEFINE_RENDER_COMMAND_PIPE(Scene, ERenderCommandPipeFlags::None);
\ No newline at end of file
diff --git a/Engine/Source/Runtime/Engine/Private/StaticMeshRender.cpp b/Engine/Source/Runtime/Engine/Private/StaticMeshRender.cpp
--- a/Engine/Source/Runtime/Engine/Private/StaticMeshRender.cpp
+++ b/Engine/Source/Runtime/Engine/Private/StaticMeshRender.cpp
@@ -520,7 +520,6 @@ void FStaticMeshSceneProxy::AddSpeedTreeWind()
 
 void FStaticMeshSceneProxy::RemoveSpeedTreeWind()
 {
-	check(IsInRenderingThread());
 	if (StaticMesh && RenderData && StaticMesh->SpeedTreeWind.IsValid())
 	{
 		for (int32 LODIndex = 0; LODIndex < RenderData->LODVertexFactories.Num(); ++LODIndex)
diff --git a/Engine/Source/Runtime/Engine/Public/GeometryCollectionRendering.h b/Engine/Source/Runtime/Engine/Public/GeometryCollectionRendering.h
--- a/Engine/Source/Runtime/Engine/Public/GeometryCollectionRendering.h
+++ b/Engine/Source/Runtime/Engine/Public/GeometryCollectionRendering.h
@@ -91,10 +91,10 @@ public:
 	//
 	// Set the data on the vertex factory
 	//
-	void SetData(const FDataType& InData)
+	void SetData(FRHICommandListBase& RHICmdList, const FDataType& InData)
 	{
 		Data = InData;
-		UpdateRHI(FRHICommandListImmediate::Get());
+		UpdateRHI(RHICmdList);
 	}
 
 	//
diff --git a/Engine/Source/Runtime/Engine/Public/PrimitiveSceneProxy.h b/Engine/Source/Runtime/Engine/Public/PrimitiveSceneProxy.h
--- a/Engine/Source/Runtime/Engine/Public/PrimitiveSceneProxy.h
+++ b/Engine/Source/Runtime/Engine/Public/PrimitiveSceneProxy.h
@@ -683,6 +683,7 @@ public:
 	inline bool IsVisibleInLumenScene() const { return bVisibleInLumenScene; }
 	inline bool ShouldRenderInMainPass() const { return bRenderInMainPass; }
 	inline bool ShouldRenderInDepthPass() const { return bRenderInMainPass || bRenderInDepthPass; }
+	inline bool ShouldConstrainToRenderThread() const { return bConstrainToRenderThread; }
 	inline bool IsCollisionEnabled() const { return bCollisionEnabled; }
 	inline bool IsHovered() const { return bHovered; }
 	inline bool IsOwnedBy(const AActor* Actor) const { return Owners.Find(Actor) != INDEX_NONE; }
@@ -1289,6 +1290,9 @@ private:
 
 protected:
 
+	/** Whether to force legacy render thread commands for OnTransformChanged / CreateRenderThreadResources / DestroyRenderThreadResources. By default, commands can run asynchronously. */
+	uint8 bConstrainToRenderThread : 1;
+
 	/** Whether this component should be tracked by Lumen Scene. Turning this off will remove it from Lumen Scene and Lumen won't generate surface cache for it. */
 	uint8 bVisibleInLumenScene : 1;
 
diff --git a/Engine/Source/Runtime/Engine/Public/Rendering/RenderCommandPipes.h b/Engine/Source/Runtime/Engine/Public/Rendering/RenderCommandPipes.h
--- a/Engine/Source/Runtime/Engine/Public/Rendering/RenderCommandPipes.h
+++ b/Engine/Source/Runtime/Engine/Public/Rendering/RenderCommandPipes.h
@@ -4,4 +4,5 @@
 
 #include "RenderingThread.h"
 
-DECLARE_RENDER_COMMAND_PIPE(SkeletalMesh, ENGINE_API);
\ No newline at end of file
+DECLARE_RENDER_COMMAND_PIPE(SkeletalMesh, ENGINE_API);
+DECLARE_RENDER_COMMAND_PIPE(Scene, ENGINE_API);
\ No newline at end of file
diff --git a/Engine/Source/Runtime/Experimental/GeometryCollectionEngine/Private/GeometryCollection/GeometryCollectionSceneProxy.cpp b/Engine/Source/Runtime/Experimental/GeometryCollectionEngine/Private/GeometryCollection/GeometryCollectionSceneProxy.cpp
--- a/Engine/Source/Runtime/Experimental/GeometryCollectionEngine/Private/GeometryCollection/GeometryCollectionSceneProxy.cpp
+++ b/Engine/Source/Runtime/Experimental/GeometryCollectionEngine/Private/GeometryCollection/GeometryCollectionSceneProxy.cpp
@@ -285,7 +285,7 @@ void FGeometryCollectionSceneProxy::SetupVertexFactory(FRHICommandListBase& RHIC
 		Data.BonePrevTransformSRV = GNullColorVertexBuffer.VertexBufferSRV;
 	}
 
-	GeometryCollectionVertexFactory.SetData(Data);
+	GeometryCollectionVertexFactory.SetData(RHICmdList, Data);
 
 	if (!GeometryCollectionVertexFactory.IsInitialized())
 	{
diff --git a/Engine/Source/Runtime/Landscape/Private/LandscapeRender.cpp b/Engine/Source/Runtime/Landscape/Private/LandscapeRender.cpp
--- a/Engine/Source/Runtime/Landscape/Private/LandscapeRender.cpp
+++ b/Engine/Source/Runtime/Landscape/Private/LandscapeRender.cpp
@@ -1336,6 +1336,9 @@ FLandscapeComponentSceneProxy::FLandscapeComponentSceneProxy(ULandscapeComponent
 #endif
 
 	UpdateVisibleInLumenScene();
+
+	// Force OnTransformChanged and {Create, Destroy}RenderThreadResources onto the render thread.
+	bConstrainToRenderThread = true;
 }
 
 void FLandscapeComponentSceneProxy::CreateRenderThreadResources(FRHICommandListBase& RHICmdList)
@@ -4230,6 +4233,9 @@ FLandscapeMeshProxySceneProxy::FLandscapeMeshProxySceneProxy(UStaticMeshComponen
 	{
 		ProxySectionsInfos.Emplace(MakeUnique<FLandscapeProxySectionInfo>(InComponent->GetWorld(), InLandscapeGuid, SectionBase, InProxyLOD, InLODGroupKey));
 	}
+
+	// Force OnTransformChanged and {Create, Destroy}RenderThreadResources onto the render thread.
+	bConstrainToRenderThread = true;
 }
 
 void FLandscapeMeshProxySceneProxy::RegisterSections()
diff --git a/Engine/Source/Runtime/Renderer/Private/RendererScene.cpp b/Engine/Source/Runtime/Renderer/Private/RendererScene.cpp
--- a/Engine/Source/Runtime/Renderer/Private/RendererScene.cpp
+++ b/Engine/Source/Runtime/Renderer/Private/RendererScene.cpp
@@ -78,6 +78,7 @@
 #include "ProfilingDebugging/CountersTrace.h"
 #include "SceneCulling/SceneCulling.h"
 #include "InstanceCulling/InstanceCullingOcclusionQuery.h"
+#include "Rendering/RenderCommandPipes.h"
 #include "ComputeWorkerInterface.h"
 
 #if RHI_RAYTRACING
@@ -698,8 +699,8 @@ void FScene::UpdateSceneSettings(AWorldSettings* WorldSettings)
 	float InDefaultMaxDistanceFieldOcclusionDistance = WorldSettings->DefaultMaxDistanceFieldOcclusionDistance;
 	float InGlobalDistanceFieldViewDistance = WorldSettings->GlobalDistanceFieldViewDistance;
 	float InDynamicIndirectShadowsSelfShadowingIntensity = FMath::Clamp(WorldSettings->DynamicIndirectShadowsSelfShadowingIntensity, 0.0f, 1.0f);
-	ENQUEUE_RENDER_COMMAND(UpdateSceneSettings)(
-		[Scene, InDefaultMaxDistanceFieldOcclusionDistance, InGlobalDistanceFieldViewDistance, InDynamicIndirectShadowsSelfShadowingIntensity](FRHICommandList& RHICmdList)
+	ENQUEUE_RENDER_COMMAND(UpdateSceneSettings)(UE::RenderCommandPipe::Scene,
+		[Scene, InDefaultMaxDistanceFieldOcclusionDistance, InGlobalDistanceFieldViewDistance, InDynamicIndirectShadowsSelfShadowingIntensity]
 	{
 		Scene->DefaultMaxDistanceFieldOcclusionDistance = InDefaultMaxDistanceFieldOcclusionDistance;
 		Scene->GlobalDistanceFieldViewDistance = InGlobalDistanceFieldViewDistance;
@@ -1158,8 +1159,8 @@ void FSceneViewState::AddLumenSceneData(FSceneInterface* InScene, float InSurfac
 		Scene = (FScene*)InScene;
 
 		// Modification of scene structure needs to happen on render thread
-		ENQUEUE_RENDER_COMMAND(SceneViewStateAdd)(
-			[RenderScene = Scene, RenderViewState = this](FRHICommandList&)
+		ENQUEUE_RENDER_COMMAND(SceneViewStateAdd)(UE::RenderCommandPipe::Scene,
+			[RenderScene = Scene, RenderViewState = this]
 			{
 				RenderScene->ViewStates.Add(RenderViewState);
 			});
@@ -1178,8 +1179,8 @@ void FSceneViewState::AddLumenSceneData(FSceneInterface* InScene, float InSurfac
 			SceneData->SurfaceCacheResolution = FMath::Clamp(InSurfaceCacheResolution, 0.5f, 1.0f);
 
 			// Need to add reference to Lumen scene data in render thread
-			ENQUEUE_RENDER_COMMAND(LinkLumenSceneData)(
-				[this, SceneData](FRHICommandListImmediate& RHICmdList)
+			ENQUEUE_RENDER_COMMAND(LinkLumenSceneData)(UE::RenderCommandPipe::Scene,
+				[this, SceneData]
 				{
 					SceneData->CopyInitialData(*Scene->DefaultLumenSceneData);
 
@@ -1194,8 +1195,8 @@ void FSceneViewState::AddLumenSceneData(FSceneInterface* InScene, float InSurfac
 		{
 			LumenSurfaceCacheResolution = InSurfaceCacheResolution;
 
-			ENQUEUE_RENDER_COMMAND(ChangeLumenSceneDataQuality)(
-				[this, InSurfaceCacheResolution](FRHICommandListImmediate& RHICmdList)
+			ENQUEUE_RENDER_COMMAND(ChangeLumenSceneDataQuality)(UE::RenderCommandPipe::Scene,
+				[this, InSurfaceCacheResolution]
 				{
 					FLumenSceneDataKey ByViewKey = { GetViewKey(), (uint32)INDEX_NONE };
 					FLumenSceneData** SceneData = Scene->PerViewOrGPULumenSceneData.Find(ByViewKey);
@@ -1215,8 +1216,8 @@ void FSceneViewState::RemoveLumenSceneData(FSceneInterface* InScene)
 	{
 		bLumenSceneDataAdded = false;
 
-		ENQUEUE_RENDER_COMMAND(RemoveLumenSceneData)(
-			[this](FRHICommandListImmediate& RHICmdList)
+		ENQUEUE_RENDER_COMMAND(RemoveLumenSceneData)(UE::RenderCommandPipe::Scene,
+			[this]
 			{
 				FLumenSceneDataKey ByViewKey = { GetViewKey(), (uint32)INDEX_NONE };
 				FLumenSceneData** SceneData = Scene->PerViewOrGPULumenSceneData.Find(ByViewKey);
@@ -1258,8 +1259,8 @@ FLumenSceneData* FScene::FindLumenSceneData(uint32 ViewKey, uint32 GPUIndex) con
 
 void FScene::UpdateParameterCollections(const TArray<FMaterialParameterCollectionInstanceResource*>& InParameterCollections)
 {
-	ENQUEUE_RENDER_COMMAND(UpdateParameterCollectionsCommand)(
-		[this, InParameterCollections](FRHICommandList&)
+	ENQUEUE_RENDER_COMMAND(UpdateParameterCollectionsCommand)(UE::RenderCommandPipe::Scene,
+		[this, InParameterCollections]
 	{
 		// Empty the scene's map so any unused uniform buffers will be released
 		ParameterCollections.Empty();
@@ -1539,7 +1540,6 @@ static void TBitArraySwapElements(TBitArray<>& Array, int32 i1, int32 i2)
 
 void FScene::AddPrimitiveSceneInfo_RenderThread(FPrimitiveSceneInfo* PrimitiveSceneInfo, const TOptional<FTransform>& PreviousTransform)
 {
-	check(IsInRenderingThread());
 	check(PrimitiveSceneInfo->PackedIndex == INDEX_NONE);
 	check(AddedPrimitiveSceneInfos.Find(PrimitiveSceneInfo) == nullptr);
 	AddedPrimitiveSceneInfos.FindOrAdd(PrimitiveSceneInfo);
@@ -1707,8 +1707,8 @@ FScene::FScene(UWorld* InWorld, bool bInRequiresHitProxies, bool bInIsEditorScen
 	World->UpdateParameterCollectionInstances(false, false);
 
 	FPersistentUniformBuffers* PersistentUniformBuffers = &UniformBuffers;
-	ENQUEUE_RENDER_COMMAND(InitializeUniformBuffers)(
-		[PersistentUniformBuffers](FRHICommandListImmediate& RHICmdList)
+	ENQUEUE_RENDER_COMMAND(InitializeUniformBuffers)(UE::RenderCommandPipe::Scene,
+		[PersistentUniformBuffers]
 	{
 		PersistentUniformBuffers->Initialize();
 	});
@@ -1866,30 +1866,45 @@ void FScene::BatchAddPrimitivesInternal(TArrayView<T*> InPrimitives)
 	SCOPE_CYCLE_COUNTER(STAT_AddScenePrimitiveGT);
 	SCOPED_NAMED_EVENT(FScene_AddPrimitive, FColor::Green);
 
-	struct FCreateRenderThreadParameters
+	struct FCreateCommand
 	{
-		FCreateRenderThreadParameters(
-			FPrimitiveSceneInfo* InPrimitiveSceneInfo,
+		FCreateCommand(
 			FPrimitiveSceneProxy* InPrimitiveSceneProxy,
 			FMatrix InRenderMatrix,
 			FBoxSphereBounds InWorldBounds,
 			FVector InAttachmentRootPosition,
-			FBoxSphereBounds InLocalBounds,
-			TOptional<FTransform> InPreviousTransform)
-			: PrimitiveSceneInfo(InPrimitiveSceneInfo), PrimitiveSceneProxy(InPrimitiveSceneProxy), RenderMatrix(InRenderMatrix), WorldBounds(InWorldBounds),
-				AttachmentRootPosition(InAttachmentRootPosition), LocalBounds(InLocalBounds), PreviousTransform(InPreviousTransform)
+			FBoxSphereBounds InLocalBounds)
+			: PrimitiveSceneProxy(InPrimitiveSceneProxy)
+			, RenderMatrix(InRenderMatrix)
+			, WorldBounds(InWorldBounds)
+			, AttachmentRootPosition(InAttachmentRootPosition)
+			, LocalBounds(InLocalBounds)
 		{}
 
-		FPrimitiveSceneInfo* PrimitiveSceneInfo;
 		FPrimitiveSceneProxy* PrimitiveSceneProxy;
 		FMatrix RenderMatrix;
 		FBoxSphereBounds WorldBounds;
 		FVector AttachmentRootPosition;
 		FBoxSphereBounds LocalBounds;
+	};
+
+	struct FAttachCommand
+	{
+		FAttachCommand(
+			FPrimitiveSceneInfo* InPrimitiveSceneInfo,
+			TOptional<FTransform> InPreviousTransform)
+			: PrimitiveSceneInfo(InPrimitiveSceneInfo)
+			, PreviousTransform(InPreviousTransform)
+		{}
+
+		FPrimitiveSceneInfo* PrimitiveSceneInfo;
 		TOptional<FTransform> PreviousTransform;
 	};
-	TArray<FCreateRenderThreadParameters, TInlineAllocator<1>> ParamsList;
-	ParamsList.Reserve(InPrimitives.Num());
+
+	TArray<FCreateCommand, TInlineAllocator<1>> CreateCommandsScenePipe;
+	TArray<FCreateCommand, TInlineAllocator<1>> CreateCommandsRenderThread;
+	TArray<FAttachCommand, TInlineAllocator<1>> AttachCommands;
+	AttachCommands.Reserve(InPrimitives.Num());
 	
 	for (T* Primitive : InPrimitives)
 	{
@@ -1939,13 +1954,20 @@ void FScene::BatchAddPrimitivesInternal(TArrayView<T*> InPrimitives)
 		FMatrix RenderMatrix = Primitive->GetRenderMatrix();
 		FVector AttachmentRootPosition = Primitive->GetActorPositionForRenderer();
 
-		ParamsList.Emplace(
-			PrimitiveSceneInfo,
+		TArray<FCreateCommand, TInlineAllocator<1>>& CreateCommands = PrimitiveSceneProxy->ShouldConstrainToRenderThread()
+			? CreateCommandsRenderThread
+			: CreateCommandsScenePipe;
+
+		CreateCommands.Emplace(
 			PrimitiveSceneProxy,
 			RenderMatrix,
 			Primitive->Bounds,
 			AttachmentRootPosition,
-			Primitive->GetLocalBounds(),
+			Primitive->GetLocalBounds()
+		);
+
+		AttachCommands.Emplace(
+			PrimitiveSceneInfo,
 
 			// If this primitive has a simulated previous transform, ensure that the velocity data for the scene representation is correct.
 			FMotionVectorSimulation::Get().GetPreviousTransform(ToUObject(Primitive))
@@ -1964,25 +1986,37 @@ void FScene::BatchAddPrimitivesInternal(TArrayView<T*> InPrimitives)
 		SceneData.AttachmentCounter.Increment();
 	}
 
-	// Create any RenderThreadResources required and send a command to the rendering thread to add the primitive to the scene.
-	FScene* Scene = this;
-
-	ENQUEUE_RENDER_COMMAND(AddPrimitiveCommand)(
-		[ParamsList = MoveTemp(ParamsList), Scene](FRHICommandListImmediate& RHICmdList)
+	if (!AttachCommands.IsEmpty())
+	{
+		ENQUEUE_RENDER_COMMAND(AddPrimitiveCommand)(UE::RenderCommandPipe::Scene,
+			[this, AttachCommands = MoveTemp(AttachCommands), CreateCommands = MoveTemp(CreateCommandsScenePipe)](FRHICommandListBase& RHICmdList)
 		{
-			for (const FCreateRenderThreadParameters& Params : ParamsList)
+			for (const FCreateCommand& Command : CreateCommands)
 			{
-				FPrimitiveSceneProxy* SceneProxy = Params.PrimitiveSceneProxy;
-				FScopeCycleCounter Context(SceneProxy->GetStatId());
-				SceneProxy->SetTransform(RHICmdList, Params.RenderMatrix, Params.WorldBounds, Params.LocalBounds, Params.AttachmentRootPosition);
-
-				// Create any RenderThreadResources required.
-				SceneProxy->CreateRenderThreadResources(RHICmdList);
+				FScopeCycleCounter Context(Command.PrimitiveSceneProxy->GetStatId());
+				Command.PrimitiveSceneProxy->SetTransform(RHICmdList, Command.RenderMatrix, Command.WorldBounds, Command.LocalBounds, Command.AttachmentRootPosition);
+			}
 
-				Scene->AddPrimitiveSceneInfo_RenderThread(Params.PrimitiveSceneInfo, Params.PreviousTransform);
+			for (const FAttachCommand& Command : AttachCommands)
+			{
+				AddPrimitiveSceneInfo_RenderThread(Command.PrimitiveSceneInfo, Command.PreviousTransform);
 			}
 		});
+	}
 
+	if (!CreateCommandsRenderThread.IsEmpty())
+	{
+		ENQUEUE_RENDER_COMMAND(CreatePrimitiveCommand)(
+			[this, CreateCommands = MoveTemp(CreateCommandsRenderThread)](FRHICommandListBase& RHICmdList)
+		{
+			for (const FCreateCommand& Command : CreateCommands)
+			{
+				FScopeCycleCounter Context(Command.PrimitiveSceneProxy->GetStatId());
+				Command.PrimitiveSceneProxy->SetTransform(RHICmdList, Command.RenderMatrix, Command.WorldBounds, Command.LocalBounds, Command.AttachmentRootPosition);
+				Command.PrimitiveSceneProxy->CreateRenderThreadResources(RHICmdList);
+			}
+		});
+	}
 }
 
 void FScene::BatchAddPrimitives(TArrayView<UPrimitiveComponent*> InPrimitives)
@@ -2013,8 +2047,6 @@ static FAutoConsoleVariableRef CVarSkipRedundantTransformUpdate(
 
 void FScene::UpdatePrimitiveTransform_RenderThread(FPrimitiveSceneProxy* PrimitiveSceneProxy, const FBoxSphereBounds& WorldBounds, const FBoxSphereBounds& LocalBounds, const FMatrix& LocalToWorld, const FVector& AttachmentRootPosition, const TOptional<FTransform>& PreviousTransform)
 {
-	check(IsInRenderingThread());
-
 #if VALIDATE_PRIMITIVE_PACKED_INDEX
 	if (AddedPrimitiveSceneInfos.Find(PrimitiveSceneProxy->GetPrimitiveSceneInfo()) != nullptr)
 	{
@@ -2038,8 +2070,6 @@ void FScene::UpdatePrimitiveTransform_RenderThread(FPrimitiveSceneProxy* Primiti
 
 void FScene::UpdatePrimitiveOcclusionBoundsSlack_RenderThread(const FPrimitiveSceneProxy* PrimitiveSceneProxy, float NewSlack)
 {
-	check(IsInRenderingThread());
-
 #if VALIDATE_PRIMITIVE_PACKED_INDEX
 	if (AddedPrimitiveSceneInfos.Find(PrimitiveSceneProxy->GetPrimitiveSceneInfo()) != nullptr)
 	{
@@ -2166,8 +2196,8 @@ void FScene::UpdatePrimitiveTransformInternal(T* Primitive)
 
 			if (bPerformUpdate)
 			{
-				ENQUEUE_RENDER_COMMAND(UpdateTransformCommand)(
-					[UpdateParams](FRHICommandListImmediate& RHICmdList)
+				ENQUEUE_RENDER_COMMAND(UpdateTransformCommand)(UE::RenderCommandPipe::Scene,
+					[UpdateParams]
 					{
 						FScopeCycleCounter Context(UpdateParams.PrimitiveSceneProxy->GetStatId());
 						UpdateParams.Scene->UpdatePrimitiveTransform_RenderThread(
@@ -2194,8 +2224,8 @@ void FScene::UpdatePrimitiveOcclusionBoundsSlack(UPrimitiveComponent* Primitive,
 {
 	if (const FPrimitiveSceneProxy* SceneProxy = Primitive->GetSceneProxy())
 	{
-		ENQUEUE_RENDER_COMMAND(UpdateOcclusionBoundsSlackCmd)(
-			[this, SceneProxy, NewSlack](FRHICommandListImmediate&)
+		ENQUEUE_RENDER_COMMAND(UpdateOcclusionBoundsSlackCmd)(UE::RenderCommandPipe::Scene,
+			[this, SceneProxy, NewSlack]
 			{
 				UpdatePrimitiveOcclusionBoundsSlack_RenderThread(SceneProxy, NewSlack);
 			});
@@ -2252,8 +2282,8 @@ void FScene::UpdatePrimitiveInstances(FInstancedStaticMeshSceneDesc* Primitive)
 
 void FScene::UpdatePrimitiveInstances(FUpdateInstanceCommand& UpdateParams)
 {	
-	ENQUEUE_RENDER_COMMAND(UpdateInstanceCommand)(
-		[this, UpdateParams = MoveTemp(UpdateParams)](FRHICommandListImmediate& RHICmdList)
+		ENQUEUE_RENDER_COMMAND(UpdateInstanceCommand)(UE::RenderCommandPipe::Scene,
+			[this, UpdateParams = MoveTemp(UpdateParams)]
 		{
 #if VALIDATE_PRIMITIVE_PACKED_INDEX
 			if (AddedPrimitiveSceneInfos.Find(UpdateParams.PrimitiveSceneProxy->GetPrimitiveSceneInfo()) != nullptr)
@@ -2303,8 +2333,8 @@ void FScene::UpdatePrimitiveLightingAttachmentRoot(UPrimitiveComponent* Primitiv
 	{
 		FPrimitiveSceneProxy* Proxy = Primitive->SceneProxy;
 		FScene* Scene = this;
-		ENQUEUE_RENDER_COMMAND(UpdatePrimitiveAttachment)(
-			[Scene, Proxy, NewComponentId](FRHICommandList&)
+		ENQUEUE_RENDER_COMMAND(UpdatePrimitiveAttachment)(UE::RenderCommandPipe::Scene,
+			[Scene, Proxy, NewComponentId]
 			{
 				FPrimitiveSceneInfo* PrimitiveInfo = Proxy->GetPrimitiveSceneInfo();
 				Scene->UpdatedAttachmentRoots.Update(PrimitiveInfo, NewComponentId);
@@ -2365,8 +2395,8 @@ void FScene::UpdateCustomPrimitiveData(FPrimitiveSceneProxy* SceneProxy, const F
 		UpdateParams.PrimitiveSceneProxy = SceneProxy;
 		UpdateParams.CustomPrimitiveData = CustomPrimitiveData; 
 
-		ENQUEUE_RENDER_COMMAND(UpdateCustomPrimitiveDataCommand)(
-			[UpdateParams](FRHICommandListImmediate& RHICmdList)
+		ENQUEUE_RENDER_COMMAND(UpdateCustomPrimitiveDataCommand)(UE::RenderCommandPipe::Scene,
+			[UpdateParams]
 			{
 				UpdateParams.Scene->UpdatedCustomPrimitiveParams.Update(UpdateParams.PrimitiveSceneProxy, UpdateParams.CustomPrimitiveData);
 			});
@@ -2381,8 +2411,8 @@ void FScene::UpdatePrimitiveDistanceFieldSceneData_GameThread(UPrimitiveComponen
 	{
 		Primitive->GetSceneData().LastSubmitTime = GetWorld()->GetTimeSeconds();
 
-		ENQUEUE_RENDER_COMMAND(UpdatePrimDFSceneDataCmd)(
-			[this, PrimitiveSceneProxy = Primitive->SceneProxy](FRHICommandList&)
+		ENQUEUE_RENDER_COMMAND(UpdatePrimDFSceneDataCmd)(UE::RenderCommandPipe::Scene,
+			[this, PrimitiveSceneProxy = Primitive->SceneProxy]
 			{
 				if (PrimitiveSceneProxy && PrimitiveSceneProxy->GetPrimitiveSceneInfo())
 				{
@@ -2415,8 +2445,6 @@ FPrimitiveSceneInfo* FScene::GetPrimitiveSceneInfo(const FPersistentPrimitiveInd
 
 void FScene::RemovePrimitiveSceneInfo_RenderThread(FPrimitiveSceneInfo* PrimitiveSceneInfo)
 {
-	check(IsInRenderingThread());
-
 	if (AddedPrimitiveSceneInfos.Remove(PrimitiveSceneInfo))
 	{
 		check(PrimitiveSceneInfo->PackedIndex == INDEX_NONE);
@@ -2463,14 +2491,15 @@ void FScene::BatchRemovePrimitivesInternal(TArrayView<T*> InPrimitives)
 	SCOPE_CYCLE_COUNTER(STAT_RemoveScenePrimitiveGT);
 	SCOPED_NAMED_EVENT(FScene_RemovePrimitive, FColor::Yellow);
 
-	struct FPrimitiveRemoveInfo
+	struct FDetachCommand
 	{
 		FPrimitiveSceneInfo* PrimitiveSceneInfo;
 		FThreadSafeCounter* AttachmentCounter;
 	};
 
-	TArray<FPrimitiveRemoveInfo, TInlineAllocator<1>> RemoveInfos;
-	RemoveInfos.Reserve(InPrimitives.Num());
+	TArray<FPrimitiveSceneProxy*, TInlineAllocator<1>> DestroyProxies;
+	TArray<FDetachCommand, TInlineAllocator<1>> DetachCommands;
+	DetachCommands.Reserve(InPrimitives.Num());
 
 	for (T* Primitive : InPrimitives)
 	{
@@ -2480,26 +2509,40 @@ void FScene::BatchRemovePrimitivesInternal(TArrayView<T*> InPrimitives)
 		{
 			FPrimitiveSceneInfo* PrimitiveSceneInfo = PrimitiveSceneProxy->GetPrimitiveSceneInfo();
 
+			if (PrimitiveSceneProxy->ShouldConstrainToRenderThread())
+			{
+				DestroyProxies.Emplace(PrimitiveSceneProxy);
+			}
+
 			// Disassociate the primitive's scene proxy.
 			Primitive->ReleaseSceneProxy();
-			RemoveInfos.Add({ PrimitiveSceneInfo, &Primitive->GetSceneData().AttachmentCounter });
+			DetachCommands.Add({ PrimitiveSceneInfo, &Primitive->GetSceneData().AttachmentCounter });
 		}
 	}
 
-	if (RemoveInfos.Num())
+	if (!DetachCommands.IsEmpty())
 	{
-		// Send a command to the rendering thread to remove the primitives from the scene.
-		FScene* Scene = this;
-		ENQUEUE_RENDER_COMMAND(FRemovePrimitiveCommand)(
-			[Scene, RemoveInfos = MoveTemp(RemoveInfos)](FRHICommandList&)
+		ENQUEUE_RENDER_COMMAND(FRemovePrimitiveCommand)(UE::RenderCommandPipe::Scene,
+			[this, DetachCommands = MoveTemp(DetachCommands)]
+		{
+			for (const FDetachCommand& Command : DetachCommands)
 			{
-				for (const FPrimitiveRemoveInfo& RemoveInfo : RemoveInfos)
-				{
-					RemoveInfo.PrimitiveSceneInfo->Proxy->DestroyRenderThreadResources();
-					Scene->RemovePrimitiveSceneInfo_RenderThread(RemoveInfo.PrimitiveSceneInfo);
-					RemoveInfo.AttachmentCounter->Decrement();
-				}
-			});
+				RemovePrimitiveSceneInfo_RenderThread(Command.PrimitiveSceneInfo);
+				Command.AttachmentCounter->Decrement();
+			}
+		});
+	}
+
+	if (!DestroyProxies.IsEmpty())
+	{
+		ENQUEUE_RENDER_COMMAND(FDestroyPrimitiveCommand)(
+			[this, DestroyProxies = MoveTemp(DestroyProxies)] (FRHICommandListBase&)
+		{
+			for (FPrimitiveSceneProxy* SceneProxy : DestroyProxies)
+			{
+				SceneProxy->DestroyRenderThreadResources();
+			}
+		});
 	}
 }
 
@@ -2546,8 +2589,8 @@ void FScene::BatchReleasePrimitivesInternal(TArrayView<T*> InPrimitives)
 		ReleaseComponentIds[ComponentIndex] = InPrimitives[ComponentIndex]->GetPrimitiveSceneId();
 	}
 
-	ENQUEUE_RENDER_COMMAND(FReleasePrimitiveCommand)(
-		[Scene, ReleaseComponentIds = MoveTemp(ReleaseComponentIds)](FRHICommandList&)
+	ENQUEUE_RENDER_COMMAND(FReleasePrimitiveCommand)(UE::RenderCommandPipe::Scene,
+		[Scene, ReleaseComponentIds = MoveTemp(ReleaseComponentIds)]
 		{
 			for (FPrimitiveComponentId PrimitiveComponentId : ReleaseComponentIds)
 			{
@@ -2731,8 +2774,8 @@ void FScene::AddLight(ULightComponent* Light)
 
 		// Send a command to the rendering thread to add the light to the scene.
 		FLightSceneInfo* LightSceneInfo = Proxy->LightSceneInfo;
-		ENQUEUE_RENDER_COMMAND(FAddLightCommand)(
-			[this, LightSceneInfo](FRHICommandListImmediate& RHICmdList)
+		ENQUEUE_RENDER_COMMAND(FAddLightCommand)(UE::RenderCommandPipe::Scene,
+			[this, LightSceneInfo]
 			{
 				CSV_SCOPED_TIMING_STAT_EXCLUSIVE(Scene_AddLight);
 				FScopeCycleCounter Context(LightSceneInfo->Proxy->GetStatId());
@@ -2765,8 +2808,8 @@ void FScene::AddInvisibleLight(ULightComponent* Light)
 		// Send a command to the rendering thread to add the light to the scene.
 		FScene* Scene = this;
 		FLightSceneInfo* LightSceneInfo = Proxy->LightSceneInfo;
-		ENQUEUE_RENDER_COMMAND(FAddLightCommand)(
-			[Scene, LightSceneInfo](FRHICommandListImmediate& RHICmdList)
+		ENQUEUE_RENDER_COMMAND(FAddLightCommand)(UE::RenderCommandPipe::Scene,
+			[Scene, LightSceneInfo]
 			{
 				FScopeCycleCounter Context(LightSceneInfo->Proxy->GetStatId());
 				LightSceneInfo->Id = Scene->InvisibleLights.Add(FLightSceneInfoCompact(LightSceneInfo));
@@ -2781,8 +2824,8 @@ void FScene::SetSkyLight(FSkyLightSceneProxy* LightProxy)
 
 	FScene* Scene = this;
 
-	ENQUEUE_RENDER_COMMAND(FSetSkyLightCommand)
-		([Scene, LightProxy](FRHICommandListImmediate& RHICmdList)
+	ENQUEUE_RENDER_COMMAND(FSetSkyLightCommand)(UE::RenderCommandPipe::Scene,
+		[Scene, LightProxy]
 		{
 			check(!Scene->SkyLightStack.Contains(LightProxy));
 			Scene->SkyLightStack.Push(LightProxy);
@@ -2811,8 +2854,8 @@ void FScene::DisableSkyLight(FSkyLightSceneProxy* LightProxy)
 
 	FScene* Scene = this;
 
-	ENQUEUE_RENDER_COMMAND(FDisableSkyLightCommand)
-		([Scene, LightProxy](FRHICommandListImmediate& RHICmdList)
+	ENQUEUE_RENDER_COMMAND(FDisableSkyLightCommand)(UE::RenderCommandPipe::Scene,
+		[Scene, LightProxy]
 	{
 		const bool bOriginalHadSkylight = Scene->ShouldRenderSkylightInBasePass(false);
 
@@ -2885,8 +2928,8 @@ void FScene::SetPhysicsField(FPhysicsFieldSceneProxy* PhysicsFieldSceneProxy)
 	check(PhysicsFieldSceneProxy);
 	FScene* Scene = this;
 
-	ENQUEUE_RENDER_COMMAND(FSetPhysicsFieldCommand)(
-		[Scene, PhysicsFieldSceneProxy](FRHICommandListImmediate& RHICmdList)
+	ENQUEUE_RENDER_COMMAND(FSetPhysicsFieldCommand)(UE::RenderCommandPipe::Scene,
+		[Scene, PhysicsFieldSceneProxy]
 		{
 			Scene->PhysicsField = PhysicsFieldSceneProxy;
 		});
@@ -2908,8 +2951,8 @@ void FScene::ResetPhysicsField()
 {
 	FScene* Scene = this;
 
-	ENQUEUE_RENDER_COMMAND(FResetPhysicsFieldCommand)(
-		[Scene](FRHICommandListImmediate& RHICmdList)
+	ENQUEUE_RENDER_COMMAND(FResetPhysicsFieldCommand)(UE::RenderCommandPipe::Scene,
+		[Scene]
 		{
 			Scene->PhysicsField = nullptr;
 		});
@@ -2939,8 +2982,8 @@ void FScene::AddDecal(UDecalComponent* Component)
 		// Send a command to the rendering thread to add the light to the scene.
 		FScene* Scene = this;
 		FDeferredDecalProxy* Proxy = Component->SceneProxy;
-		ENQUEUE_RENDER_COMMAND(FAddDecalCommand)(
-			[Scene, Proxy](FRHICommandListImmediate& RHICmdList)
+		ENQUEUE_RENDER_COMMAND(FAddDecalCommand)(UE::RenderCommandPipe::Scene,
+			[Scene, Proxy]
 			{
 				Scene->AddOrRemoveDecal_RenderThread(Proxy, true);
 			});
@@ -2956,8 +2999,8 @@ void FScene::RemoveDecal(UDecalComponent* Component)
 		// Send a command to the rendering thread to remove the light from the scene.
 		FScene* Scene = this;
 		FDeferredDecalProxy* Proxy = Component->SceneProxy;
-		ENQUEUE_RENDER_COMMAND(FRemoveDecalCommand)(
-			[Scene, Proxy](FRHICommandListImmediate& RHICmdList)
+		ENQUEUE_RENDER_COMMAND(FRemoveDecalCommand)(UE::RenderCommandPipe::Scene,
+			[Scene, Proxy]
 			{
 				Scene->AddOrRemoveDecal_RenderThread(Proxy, false);
 			});
@@ -2976,8 +3019,8 @@ void FScene::UpdateDecalTransform(UDecalComponent* Decal)
 		FDeferredDecalProxy* DecalSceneProxy = Decal->SceneProxy;
 		FTransform ComponentToWorldIncludingDecalSize = Decal->GetTransformIncludingDecalSize();
 		FBoxSphereBounds Bounds = Decal->CalcBounds(Decal->GetComponentTransform());
-		ENQUEUE_RENDER_COMMAND(UpdateTransformCommand)(
-			[DecalSceneProxy, ComponentToWorldIncludingDecalSize, Bounds, Scene](FRHICommandListImmediate& RHICmdList)
+		ENQUEUE_RENDER_COMMAND(UpdateTransformCommand)(UE::RenderCommandPipe::Scene,
+			[DecalSceneProxy, ComponentToWorldIncludingDecalSize, Bounds, Scene]
 			{
 				// Invalidate the path tracer only if the decal was sufficiently moved
 				if (!ComponentToWorldIncludingDecalSize.Equals(DecalSceneProxy->ComponentTrans, SMALL_NUMBER))
@@ -2999,8 +3042,8 @@ void FScene::UpdateDecalFadeOutTime(UDecalComponent* Decal)
 		float DecalFadeStartDelay = Decal->FadeStartDelay;
 		float DecalFadeDuration = Decal->FadeDuration;
 
-		ENQUEUE_RENDER_COMMAND(FUpdateDecalFadeInTimeCommand)(
-			[Proxy, CurrentTime, DecalFadeStartDelay, DecalFadeDuration](FRHICommandListImmediate& RHICmdList)
+		ENQUEUE_RENDER_COMMAND(FUpdateDecalFadeInTimeCommand)(UE::RenderCommandPipe::Scene,
+			[Proxy, CurrentTime, DecalFadeStartDelay, DecalFadeDuration]
 		{
 			if (DecalFadeDuration > 0.0f)
 			{
@@ -3025,8 +3068,8 @@ void FScene::UpdateDecalFadeInTime(UDecalComponent* Decal)
 		float DecalFadeStartDelay = Decal->FadeInStartDelay;
 		float DecalFadeDuration = Decal->FadeInDuration;
 
-		ENQUEUE_RENDER_COMMAND(FUpdateDecalFadeInTimeCommand)(
-			[Proxy, CurrentTime, DecalFadeStartDelay, DecalFadeDuration](FRHICommandListImmediate& RHICmdList)
+		ENQUEUE_RENDER_COMMAND(FUpdateDecalFadeInTimeCommand)(UE::RenderCommandPipe::Scene,
+			[Proxy, CurrentTime, DecalFadeStartDelay, DecalFadeDuration]
 		{
 			if (DecalFadeDuration > 0.0f)
 			{
@@ -3044,8 +3087,8 @@ void FScene::UpdateDecalFadeInTime(UDecalComponent* Decal)
 
 void FScene::BatchUpdateDecals(TArray<FDeferredDecalUpdateParams>&& UpdateParams)
 {
-	ENQUEUE_RENDER_COMMAND(FBatchUpdateDecalsCommand)(
-		[Scene=this, UpdateParams_RT=MoveTemp(UpdateParams)](FRHICommandListImmediate& RHICmdList)
+	ENQUEUE_RENDER_COMMAND(FBatchUpdateDecalsCommand)(UE::RenderCommandPipe::Scene,
+		[Scene=this, UpdateParams_RT=MoveTemp(UpdateParams)]
 		{
 			for (const FDeferredDecalUpdateParams& DecalUpdate : UpdateParams_RT )
 			{
@@ -3074,7 +3117,6 @@ void FScene::AddHairStrands(FHairStrandsInstance* Proxy)
 {
 	if (Proxy)
 	{
-		check(IsInRenderingThread());
 		const int32 PackedIndex = HairStrandsSceneData.RegisteredProxies.Add(Proxy);
 		Proxy->RegisteredIndex = PackedIndex;
 	}
@@ -3084,7 +3126,6 @@ void FScene::RemoveHairStrands(FHairStrandsInstance* Proxy)
 {
 	if (Proxy)
 	{
-		check(IsInRenderingThread());
 		int32 ProxyIndex = Proxy->RegisteredIndex;
 		if (HairStrandsSceneData.RegisteredProxies.IsValidIndex(ProxyIndex))
 		{
@@ -3131,8 +3172,8 @@ void FScene::AddReflectionCapture(UReflectionCaptureComponent* Component)
 		FReflectionCaptureProxy* Proxy = Component->SceneProxy;
 		const FVector Position = Component->GetComponentLocation();
 
-		ENQUEUE_RENDER_COMMAND(FAddCaptureCommand)
-			([Scene, Proxy, Position](FRHICommandListImmediate& RHICmdList)
+		ENQUEUE_RENDER_COMMAND(FAddCaptureCommand)(UE::RenderCommandPipe::Scene,
+			[Scene, Proxy, Position](FRHICommandListBase& RHICmdList)
 		{
 			if (Proxy->bUsingPreviewCaptureData)
 			{
@@ -3162,8 +3203,8 @@ void FScene::RemoveReflectionCapture(UReflectionCaptureComponent* Component)
 		FScene* Scene = this;
 		FReflectionCaptureProxy* Proxy = Component->SceneProxy;
 
-		ENQUEUE_RENDER_COMMAND(FRemoveCaptureCommand)
-			([Scene, Proxy](FRHICommandListImmediate& RHICmdList)
+		ENQUEUE_RENDER_COMMAND(FRemoveCaptureCommand)(UE::RenderCommandPipe::Scene,
+			[Scene, Proxy]
 		{
 			if (Proxy->bUsingPreviewCaptureData)
 			{
@@ -3209,8 +3250,8 @@ void FScene::UpdateReflectionCaptureTransform(UReflectionCaptureComponent* Compo
 		FReflectionCaptureProxy* Proxy = Component->SceneProxy;
 		FMatrix Transform = Component->GetComponentTransform().ToMatrixWithScale();
 
-		ENQUEUE_RENDER_COMMAND(FUpdateTransformCommand)
-			([Scene, Proxy, Transform, bUsingPreviewCaptureData](FRHICommandListImmediate& RHICmdList)
+		ENQUEUE_RENDER_COMMAND(FUpdateTransformCommand)(UE::RenderCommandPipe::Scene,
+			[Scene, Proxy, Transform, bUsingPreviewCaptureData](FRHICommandListBase& RHICmdList)
 		{
 			if (Proxy->bUsingPreviewCaptureData)
 			{
@@ -3253,8 +3294,8 @@ void FScene::ReleaseReflectionCubemap(UReflectionCaptureComponent* CaptureCompon
 	if (bRemoved)
 	{
 		FScene* Scene = this;
-		ENQUEUE_RENDER_COMMAND(RemoveCaptureCommand)(
-			[CaptureComponent, Scene](FRHICommandListImmediate& RHICmdList)
+		ENQUEUE_RENDER_COMMAND(RemoveCaptureCommand)(UE::RenderCommandPipe::Scene,
+			[CaptureComponent, Scene]
 			{
 				int32 IndexToFree = -1;
 
@@ -3430,8 +3471,8 @@ void FScene::AddPrecomputedLightVolume(const FPrecomputedLightVolume* Volume)
 {
 	FScene* Scene = this;
 
-	ENQUEUE_RENDER_COMMAND(AddVolumeCommand)
-		([Scene, Volume](FRHICommandListImmediate& RHICmdList) 
+	ENQUEUE_RENDER_COMMAND(AddVolumeCommand)(UE::RenderCommandPipe::Scene,
+		[Scene, Volume]
 		{
 			Scene->PrecomputedLightVolumes.Add(Volume);
 			Scene->IndirectLightingCache.SetLightingCacheDirty(Scene, Volume);
@@ -3442,8 +3483,8 @@ void FScene::RemovePrecomputedLightVolume(const FPrecomputedLightVolume* Volume)
 {
 	FScene* Scene = this; 
 
-	ENQUEUE_RENDER_COMMAND(RemoveVolumeCommand)
-		([Scene, Volume](FRHICommandListImmediate& RHICmdList) 
+	ENQUEUE_RENDER_COMMAND(RemoveVolumeCommand)(UE::RenderCommandPipe::Scene,
+		[Scene, Volume]
 		{
 			Scene->PrecomputedLightVolumes.Remove(Volume);
 			Scene->IndirectLightingCache.SetLightingCacheDirty(Scene, Volume);
@@ -3530,7 +3571,7 @@ void FScene::AddPrecomputedVolumetricLightmap(const FPrecomputedVolumetricLightm
 	FScene* Scene = this;
 
 	ENQUEUE_RENDER_COMMAND(AddVolumeCommand)
-	([Scene, Volume, bIsPersistentLevel](FRHICommandListImmediate& RHICmdList)
+		([Scene, Volume, bIsPersistentLevel] (FRHICommandListBase&)
 	{
 		Scene->VolumetricLightmapSceneData.AddLevelVolume(Volume, Scene->GetShadingPath(), bIsPersistentLevel);
 	});
@@ -3541,7 +3582,7 @@ void FScene::RemovePrecomputedVolumetricLightmap(const FPrecomputedVolumetricLig
 	FScene* Scene = this; 
 
 	ENQUEUE_RENDER_COMMAND(RemoveVolumeCommand)
-	([Scene, Volume](FRHICommandListImmediate& RHICmdList) 
+		([Scene, Volume] (FRHICommandListBase&)
 	{
 		Scene->VolumetricLightmapSceneData.RemoveLevelVolume(Volume);
 	});
@@ -3556,8 +3597,8 @@ void FScene::AddRuntimeVirtualTexture(class URuntimeVirtualTextureComponent* Com
 		FScene* Scene = this;
 		FRuntimeVirtualTextureSceneProxy* SceneProxy = Component->SceneProxy;
 
-		ENQUEUE_RENDER_COMMAND(AddRuntimeVirtualTextureCommand)(
-			[Scene, SceneProxy](FRHICommandListImmediate& RHICmdList)
+		ENQUEUE_RENDER_COMMAND(AddRuntimeVirtualTextureCommand)(UE::RenderCommandPipe::Scene,
+			[Scene, SceneProxy]
 		{
 			Scene->AddRuntimeVirtualTexture_RenderThread(SceneProxy);
 			Scene->UpdateRuntimeVirtualTextureForAllPrimitives_RenderThread();
@@ -3574,8 +3615,8 @@ void FScene::AddRuntimeVirtualTexture(class URuntimeVirtualTextureComponent* Com
 		FScene* Scene = this;
 		FRuntimeVirtualTextureSceneProxy* SceneProxy = Component->SceneProxy;
 
-		ENQUEUE_RENDER_COMMAND(AddRuntimeVirtualTextureCommand)(
-			[Scene, SceneProxy, SceneProxyToReplace](FRHICommandListImmediate& RHICmdList)
+		ENQUEUE_RENDER_COMMAND(AddRuntimeVirtualTextureCommand)(UE::RenderCommandPipe::Scene,
+			[Scene, SceneProxy, SceneProxyToReplace]
 		{
 			const bool bUpdatePrimitives = SceneProxy->VirtualTexture != SceneProxyToReplace->VirtualTexture;
 			Scene->UpdateRuntimeVirtualTexture_RenderThread(SceneProxy, SceneProxyToReplace);
@@ -3597,8 +3638,8 @@ void FScene::RemoveRuntimeVirtualTexture(class URuntimeVirtualTextureComponent*
 		Component->SceneProxy = nullptr;
 
 		FScene* Scene = this;
-		ENQUEUE_RENDER_COMMAND(RemoveRuntimeVirtualTextureCommand)(
-			[Scene, SceneProxy](FRHICommandListImmediate& RHICmdList)
+		ENQUEUE_RENDER_COMMAND(RemoveRuntimeVirtualTextureCommand)(UE::RenderCommandPipe::Scene,
+			[Scene, SceneProxy]
 		{
 			Scene->RemoveRuntimeVirtualTexture_RenderThread(SceneProxy);
 			Scene->UpdateRuntimeVirtualTextureForAllPrimitives_RenderThread();
@@ -3687,8 +3728,8 @@ void FScene::InvalidateRuntimeVirtualTexture(class URuntimeVirtualTextureCompone
 	if (Component->SceneProxy != nullptr)
 	{
 		FRuntimeVirtualTextureSceneProxy* SceneProxy = Component->SceneProxy;
-		ENQUEUE_RENDER_COMMAND(RuntimeVirtualTextureComponent_SetDirty)(
-			[SceneProxy, WorldBounds](FRHICommandList& RHICmdList)
+		ENQUEUE_RENDER_COMMAND(RuntimeVirtualTextureComponent_SetDirty)(UE::RenderCommandPipe::Scene,
+			[SceneProxy, WorldBounds]
 		{
 			SceneProxy->Dirty(WorldBounds);
 		});
@@ -3707,8 +3748,8 @@ void FScene::InvalidateLumenSurfaceCache_GameThread(UPrimitiveComponent* Compone
 
 	if (Component->SceneProxy)
 	{
-		ENQUEUE_RENDER_COMMAND(InvalidateLumenSurfaceCacheCmd)(
-			[this, PrimitiveSceneProxy = Component->SceneProxy](FRHICommandList&)
+		ENQUEUE_RENDER_COMMAND(InvalidateLumenSurfaceCacheCmd)(UE::RenderCommandPipe::Scene,
+			[this, PrimitiveSceneProxy = Component->SceneProxy]
 		{
 			if (PrimitiveSceneProxy && PrimitiveSceneProxy->GetPrimitiveSceneInfo())
 			{
@@ -3834,8 +3875,8 @@ void FScene::UpdateLightTransform(ULightComponent* Light)
 		FLightSceneInfo* LightSceneInfo = Light->SceneProxy->GetLightSceneInfo();
 		if (LightSceneInfo->bVisible)
 		{
-			ENQUEUE_RENDER_COMMAND(UpdateLightTransform)(
-				[this, LightSceneInfo, Parameters](FRHICommandListImmediate& RHICmdList)
+			ENQUEUE_RENDER_COMMAND(UpdateLightTransform)(UE::RenderCommandPipe::Scene,
+				[this, LightSceneInfo, Parameters]
 				{
 					FScopeCycleCounter Context(LightSceneInfo->Proxy->GetStatId());
 					SceneLightInfoUpdates->Enqueue(Parameters, LightSceneInfo);
@@ -3863,8 +3904,8 @@ void FScene::UpdateLightColorAndBrightness(ULightComponent* Light)
 		FLightSceneInfo* LightSceneInfo = Light->SceneProxy->GetLightSceneInfo();
 		if (LightSceneInfo->bVisible)
 		{
-			ENQUEUE_RENDER_COMMAND(UpdateLightColorAndBrightness)(
-				[this, LightSceneInfo, NewParameters](FRHICommandListImmediate& RHICmdList)
+			ENQUEUE_RENDER_COMMAND(UpdateLightColorAndBrightness)(UE::RenderCommandPipe::Scene,
+				[this, LightSceneInfo, NewParameters]
 				{
 					SceneLightInfoUpdates->Enqueue(NewParameters, LightSceneInfo);
 					++SceneLightInfoUpdates->NumUpdates;
@@ -3995,8 +4036,8 @@ void FScene::RemoveLight(ULightComponent* Light)
 		Light->SceneProxy = nullptr;
 
 		// Send a command to the rendering thread to queue the light for removal from the scene.
-		ENQUEUE_RENDER_COMMAND(FQueueRemoveLightCommand)(
-			[this, LightSceneInfo](FRHICommandListImmediate& RHICmdList)
+		ENQUEUE_RENDER_COMMAND(FQueueRemoveLightCommand)(UE::RenderCommandPipe::Scene,
+			[this, LightSceneInfo]
 			{
 				FScopeCycleCounter Context(LightSceneInfo->Proxy->GetStatId());
 
@@ -4036,8 +4077,8 @@ void FScene::AddExponentialHeightFog(UExponentialHeightFogComponent* FogComponen
 {
 	FScene* Scene = this;
 	FExponentialHeightFogSceneInfo HeightFogSceneInfo = FExponentialHeightFogSceneInfo(FogComponent);
-	ENQUEUE_RENDER_COMMAND(FAddFogCommand)(
-		[Scene, HeightFogSceneInfo](FRHICommandListImmediate& RHICmdList)
+	ENQUEUE_RENDER_COMMAND(FAddFogCommand)(UE::RenderCommandPipe::Scene,
+		[Scene, HeightFogSceneInfo]
 		{
 			// Create a FExponentialHeightFogSceneInfo for the component in the scene's fog array.
 			new(Scene->ExponentialFogs) FExponentialHeightFogSceneInfo(HeightFogSceneInfo);
@@ -4048,8 +4089,8 @@ void FScene::AddExponentialHeightFog(UExponentialHeightFogComponent* FogComponen
 void FScene::RemoveExponentialHeightFog(UExponentialHeightFogComponent* FogComponent)
 {
 	FScene* Scene = this;
-	ENQUEUE_RENDER_COMMAND(FRemoveFogCommand)(
-		[Scene, FogComponent](FRHICommandListImmediate& RHICmdList)
+	ENQUEUE_RENDER_COMMAND(FRemoveFogCommand)(UE::RenderCommandPipe::Scene,
+		[Scene, FogComponent]
 		{
 			// Remove the given component's FExponentialHeightFogSceneInfo from the scene's fog array.
 			for(int32 FogIndex = 0;FogIndex < Scene->ExponentialFogs.Num();FogIndex++)
@@ -4083,8 +4124,8 @@ void FScene::AddWindSource(UWindDirectionalSourceComponent* WindComponent)
 	WindComponent->SceneProxy = SceneProxy;
 
 	FScene* Scene = this;
-	ENQUEUE_RENDER_COMMAND(FAddWindSourceCommand)(
-		[Scene, SceneProxy](FRHICommandListImmediate& RHICmdList)
+	ENQUEUE_RENDER_COMMAND(FAddWindSourceCommand)(UE::RenderCommandPipe::Scene,
+		[Scene, SceneProxy]
 		{
 			Scene->WindSources.Add(SceneProxy);
 		});
@@ -4101,8 +4142,8 @@ void FScene::RemoveWindSource(UWindDirectionalSourceComponent* WindComponent)
 	if(SceneProxy)
 	{
 		FScene* Scene = this;
-		ENQUEUE_RENDER_COMMAND(FRemoveWindSourceCommand)(
-			[Scene, SceneProxy](FRHICommandListImmediate& RHICmdList)
+		ENQUEUE_RENDER_COMMAND(FRemoveWindSourceCommand)(UE::RenderCommandPipe::Scene,
+			[Scene, SceneProxy]
 			{
 				Scene->WindSources.Remove(SceneProxy);
 
@@ -4121,8 +4162,8 @@ void FScene::UpdateWindSource(UWindDirectionalSourceComponent* WindComponent)
 	{
 		WindComponent->SceneProxy = nullptr;
 
-		ENQUEUE_RENDER_COMMAND(FRemoveWindSourceCommand)(
-			[Scene = this, OldSceneProxy](FRHICommandListImmediate& RHICmdList)
+		ENQUEUE_RENDER_COMMAND(FRemoveWindSourceCommand)(UE::RenderCommandPipe::Scene,
+			[Scene = this, OldSceneProxy]
 			{
 				Scene->WindSources.Remove(OldSceneProxy);
 
@@ -4135,8 +4176,8 @@ void FScene::UpdateWindSource(UWindDirectionalSourceComponent* WindComponent)
 		FWindSourceSceneProxy* const NewSceneProxy = WindComponent->CreateSceneProxy();
 		WindComponent->SceneProxy = NewSceneProxy;
 
-		ENQUEUE_RENDER_COMMAND(FAddWindSourceCommand)(
-			[Scene = this, NewSceneProxy](FRHICommandListImmediate& RHICmdList)
+		ENQUEUE_RENDER_COMMAND(FAddWindSourceCommand)(UE::RenderCommandPipe::Scene,
+			[Scene = this, NewSceneProxy]
 			{
 				Scene->WindSources.Add(NewSceneProxy);
 			});
@@ -4259,8 +4300,8 @@ void FScene::AddSpeedTreeWind(FVertexFactory* VertexFactory, const UStaticMesh*
 	if (StaticMesh != NULL && StaticMesh->SpeedTreeWind.IsValid() && StaticMesh->GetRenderData())
 	{
 		FScene* Scene = this;
-		ENQUEUE_RENDER_COMMAND(FAddSpeedTreeWindCommand)(
-			[Scene, StaticMesh, VertexFactory](FRHICommandListImmediate& RHICmdList)
+		ENQUEUE_RENDER_COMMAND(FAddSpeedTreeWindCommand)(UE::RenderCommandPipe::Scene,
+			[Scene, StaticMesh, VertexFactory]
 			{
 				Scene->SpeedTreeVertexFactoryMap.Add(VertexFactory, StaticMesh);
 
@@ -4313,8 +4354,8 @@ void FScene::UpdateSpeedTreeWind(double CurrentTime)
 	UniformParameters.Prev##name = *(FVector4f*)(WindShaderValues + FSpeedTreeWind::offset + FSpeedTreeWind::NUM_SHADER_VALUES);
 
 	FScene* Scene = this;
-	ENQUEUE_RENDER_COMMAND(FUpdateSpeedTreeWindCommand)(
-		[Scene, CurrentTime](FRHICommandListImmediate& RHICmdList)
+	ENQUEUE_RENDER_COMMAND(FUpdateSpeedTreeWindCommand)(UE::RenderCommandPipe::Scene,
+		[Scene, CurrentTime] (FRHICommandListBase& RHICmdList)
 		{
 			FVector WindDirection;
 			float WindSpeed;
@@ -4429,8 +4470,8 @@ void FScene::GetRelevantLights( UPrimitiveComponent* Primitive, TArray<const ULi
 	{
 		// Add interacting lights to the array.
 		const FScene* Scene = this;
-		ENQUEUE_RENDER_COMMAND(FGetRelevantLightsCommand)(
-			[Scene, Primitive, RelevantLights](FRHICommandListImmediate& RHICmdList)
+		ENQUEUE_RENDER_COMMAND(FGetRelevantLightsCommand)(UE::RenderCommandPipe::Scene,
+			[Scene, Primitive, RelevantLights]
 			{
 				Scene->GetRelevantLights_RenderThread( Primitive, RelevantLights );
 			});
@@ -4444,8 +4485,8 @@ void FScene::GetRelevantLights( UPrimitiveComponent* Primitive, TArray<const ULi
 void FScene::SetPrecomputedVisibility(const FPrecomputedVisibilityHandler* NewPrecomputedVisibilityHandler)
 {
 	FScene* Scene = this;
-	ENQUEUE_RENDER_COMMAND(UpdatePrecomputedVisibility)(
-		[Scene, NewPrecomputedVisibilityHandler](FRHICommandListImmediate& RHICmdList)
+	ENQUEUE_RENDER_COMMAND(UpdatePrecomputedVisibility)(UE::RenderCommandPipe::Scene,
+		[Scene, NewPrecomputedVisibilityHandler]
 		{
 			Scene->PrecomputedVisibilityHandler = NewPrecomputedVisibilityHandler;
 		});
@@ -4482,8 +4523,6 @@ void FScene::UpdateStaticDrawLists()
 
 void FScene::UpdateCachedRenderStates(FPrimitiveSceneProxy* SceneProxy)
 {
-	check(IsInRenderingThread());
-
 	if (SceneProxy->GetPrimitiveSceneInfo())
 	{
 		SceneProxy->GetPrimitiveSceneInfo()->RequestStaticMeshUpdate();
@@ -4494,8 +4533,6 @@ void FScene::UpdateCachedRenderStates(FPrimitiveSceneProxy* SceneProxy)
 
 void FScene::UpdateCachedRayTracingState(FPrimitiveSceneProxy* SceneProxy)
 {
-	check(IsInRenderingThread());
-
 	if (SceneProxy->GetPrimitiveSceneInfo())
 	{
 		SceneProxy->GetPrimitiveSceneInfo()->bCachedRaytracingDataDirty = true;
@@ -4836,8 +4873,8 @@ void FScene::OnLevelAddedToWorld(const FName& InLevelAddedName, UWorld* InWorld,
 
 	FScene* Scene = this;
 	FName LevelAddedName = InLevelAddedName;
-	ENQUEUE_RENDER_COMMAND(FLevelAddedToWorld)(
-		[Scene, LevelAddedName](FRHICommandListImmediate& RHICmdList)
+	ENQUEUE_RENDER_COMMAND(FLevelAddedToWorld)(UE::RenderCommandPipe::Scene,
+		[Scene, LevelAddedName]
 		{
 			FLevelCommand Cmd;
 			Cmd.Name = LevelAddedName;
@@ -4855,8 +4892,8 @@ void FScene::OnLevelRemovedFromWorld(const FName& InLevelRemovedName, UWorld* In
 
 	FScene* Scene = this;
 	FName LevelRemovedName = InLevelRemovedName;
-	ENQUEUE_RENDER_COMMAND(FLevelRemovedFromWorld)(
-		[Scene, LevelRemovedName](FRHICommandListImmediate& RHICmdList)
+	ENQUEUE_RENDER_COMMAND(FLevelRemovedFromWorld)(UE::RenderCommandPipe::Scene,
+		[Scene, LevelRemovedName]
 		{
 			FLevelCommand Cmd;
 			Cmd.Name = LevelRemovedName;
@@ -5395,13 +5432,6 @@ void FScene::UpdateAllPrimitiveSceneInfos(FRDGBuilder& GraphBuilder, EUpdateAllP
 
 	RDG_EVENT_SCOPE(GraphBuilder, "UpdateAllPrimitiveSceneInfos");
 
-	UpdateAllLightSceneInfos(GraphBuilder);
-
-#if RHI_RAYTRACING
-	UpdateRayTracingGroupBounds_RemovePrimitives(RemovedPrimitiveSceneInfos);
-	UpdateRayTracingGroupBounds_AddPrimitives(AddedPrimitiveSceneInfos);
-#endif
-
 	// Allocated with render graph lifetime, safe to reference from RDG tasks.
 	FSceneUpdateChangeSetStorage &SceneUpdateChangeSetStorage = *GraphBuilder.AllocObject<FSceneUpdateChangeSetStorage>();
 	SceneUpdateChangeSetStorage.RemovedPrimitiveIds.Reserve(RemovedPrimitiveSceneInfos.Num());
@@ -5417,6 +5447,52 @@ void FScene::UpdateAllPrimitiveSceneInfos(FRDGBuilder& GraphBuilder, EUpdateAllP
 		SceneUpdateChangeSetStorage.RemovedPrimitiveSceneInfos.Add(PrimitiveSceneInfo);
 	}
 
+	TArray<FPrimitiveSceneInfo*> AddedLocalPrimitiveSceneInfos;
+	AddedLocalPrimitiveSceneInfos.Reserve(AddedPrimitiveSceneInfos.Num());
+	for (FPrimitiveSceneInfo* SceneInfo : AddedPrimitiveSceneInfos)
+	{
+		AddedLocalPrimitiveSceneInfos.Add(SceneInfo);
+	}
+
+	UE::Tasks::FTaskEvent ProcessPrimitiveResourcesTask{ UE_SOURCE_LOCATION };
+
+	if (!AddedLocalPrimitiveSceneInfos.IsEmpty() || !RemovedLocalPrimitiveSceneInfos.IsEmpty())
+	{
+		ProcessPrimitiveResourcesTask.AddPrerequisites(GraphBuilder.AddCommandListSetupTask([AddedPrimitives = AddedLocalPrimitiveSceneInfos, RemovedPrimitives = RemovedLocalPrimitiveSceneInfos](FRHICommandListBase& RHICmdList)
+		{
+			{
+				TRACE_CPUPROFILER_EVENT_SCOPE(Scene::DestroyPrimitiveResources);
+				for (FPrimitiveSceneInfo* Primitive : RemovedPrimitives)
+				{
+					if (!Primitive->Proxy->ShouldConstrainToRenderThread())
+					{
+						Primitive->Proxy->DestroyRenderThreadResources();
+					}
+				}
+			}
+
+			{
+				TRACE_CPUPROFILER_EVENT_SCOPE(Scene::CreatePrimitiveResources);
+				for (FPrimitiveSceneInfo* Primitive : AddedPrimitives)
+				{
+					if (!Primitive->Proxy->ShouldConstrainToRenderThread())
+					{
+						Primitive->Proxy->CreateRenderThreadResources(RHICmdList);
+					}
+				}
+			}
+		}));
+	}
+
+	ProcessPrimitiveResourcesTask.Trigger();
+
+	UpdateAllLightSceneInfos(GraphBuilder);
+
+#if RHI_RAYTRACING
+	UpdateRayTracingGroupBounds_RemovePrimitives(RemovedPrimitiveSceneInfos);
+	UpdateRayTracingGroupBounds_AddPrimitives(AddedPrimitiveSceneInfos);
+#endif
+
 	{
 		SceneUpdateChangeSetStorage.UpdatedPrimitiveIds.Reserve(UpdatedInstances.Num() + UpdatedTransforms.Num());
 		SceneUpdateChangeSetStorage.UpdatedPrimitiveSceneInfos.Reserve(UpdatedInstances.Num() + UpdatedTransforms.Num());
@@ -5511,13 +5587,6 @@ void FScene::UpdateAllPrimitiveSceneInfos(FRDGBuilder& GraphBuilder, EUpdateAllP
 		CacheManager->ProcessInvalidations(GraphBuilder, SceneUB, InvalidatingPrimitiveCollector);
 	}
 
-	TArray<FPrimitiveSceneInfo*> AddedLocalPrimitiveSceneInfos;
-	AddedLocalPrimitiveSceneInfos.Reserve(AddedPrimitiveSceneInfos.Num());
-	for (FPrimitiveSceneInfo* SceneInfo : AddedPrimitiveSceneInfos)
-	{
-		AddedLocalPrimitiveSceneInfos.Add(SceneInfo);
-	}
-
 	AddedLocalPrimitiveSceneInfos.Sort(FPrimitiveArraySortKey());
 
 	DeletedPrimitiveSceneInfos.Reserve(RemovedLocalPrimitiveSceneInfos.Num());
@@ -6156,6 +6225,8 @@ void FScene::UpdateAllPrimitiveSceneInfos(FRDGBuilder& GraphBuilder, EUpdateAllP
 	bool bUpdateCachedRayTracingInstances = false;
 #endif
 
+	ProcessPrimitiveResourcesTask.Wait();
+
 	// Updating step of instance update. Note that instance IDs are not allocated at this point.
 	{
 		CSV_SCOPED_TIMING_STAT_EXCLUSIVE(UpdatePrimitiveInstances);
