Refactored mesh collector and dynamic vertex / index buffers in preparation for moving GDME off the render thread.
 - Added GetRHICommandList() method to FMeshElementCollector.
 - Removed global dynamic vertex / index buffers in favor of local variants.
 - Added some locks to shared systems.

diff --git a/Engine/Plugins/FX/Niagara/Source/Niagara/Private/NiagaraRenderer.cpp b/Engine/Plugins/FX/Niagara/Source/Niagara/Private/NiagaraRenderer.cpp
--- a/Engine/Plugins/FX/Niagara/Source/Niagara/Private/NiagaraRenderer.cpp
+++ b/Engine/Plugins/FX/Niagara/Source/Niagara/Private/NiagaraRenderer.cpp
@@ -217,9 +217,9 @@ FParticleRenderData FNiagaraRenderer::TransferDataToGPU(FRHICommandListBase& RHI
 	const int32 TotalIntSize = IntComponents.Num() * NumInstances;
 
 	FParticleRenderData Allocation;
-	Allocation.FloatData = TotalFloatSize ? DynamicReadBuffer.AllocateFloat(RHICmdList, TotalFloatSize) : FGlobalDynamicReadBuffer::FAllocation();
-	Allocation.HalfData = TotalHalfSize ? DynamicReadBuffer.AllocateHalf(RHICmdList, TotalHalfSize) : FGlobalDynamicReadBuffer::FAllocation();
-	Allocation.IntData = TotalIntSize ? DynamicReadBuffer.AllocateInt32(RHICmdList, TotalIntSize) : FGlobalDynamicReadBuffer::FAllocation();
+	Allocation.FloatData = TotalFloatSize ? DynamicReadBuffer.AllocateFloat(TotalFloatSize) : FGlobalDynamicReadBuffer::FAllocation();
+	Allocation.HalfData = TotalHalfSize ? DynamicReadBuffer.AllocateHalf(TotalHalfSize) : FGlobalDynamicReadBuffer::FAllocation();
+	Allocation.IntData = TotalIntSize ? DynamicReadBuffer.AllocateInt32(TotalIntSize) : FGlobalDynamicReadBuffer::FAllocation();
 
 	Allocation.FloatStride = TotalFloatSize ? NumInstances * sizeof(float) : 0;
 	Allocation.HalfStride = TotalHalfSize ? NumInstances * sizeof(FFloat16) : 0;
diff --git a/Engine/Plugins/FX/Niagara/Source/Niagara/Private/NiagaraRendererMeshes.cpp b/Engine/Plugins/FX/Niagara/Source/Niagara/Private/NiagaraRendererMeshes.cpp
--- a/Engine/Plugins/FX/Niagara/Source/Niagara/Private/NiagaraRendererMeshes.cpp
+++ b/Engine/Plugins/FX/Niagara/Source/Niagara/Private/NiagaraRendererMeshes.cpp
@@ -721,7 +721,7 @@ uint32 FNiagaraRendererMeshes::PerformSortAndCull(FRHICommandListBase& RHICmdLis
 		else
 		{
 			FGlobalDynamicReadBuffer::FAllocation SortedIndices;
-			SortedIndices = ReadBuffer.AllocateUInt32(RHICmdList, NumInstances);
+			SortedIndices = ReadBuffer.AllocateUInt32(NumInstances);
 			NumInstances = SortAndCullIndices(SortInfo, *ParticleMeshRenderData.SourceParticleData, SortedIndices);
 			ParticleMeshRenderData.ParticleSortedIndicesSRV = SortedIndices.SRV;
 			ParticleMeshRenderData.ParticleSortedIndicesOffset = 0;
diff --git a/Engine/Plugins/FX/Niagara/Source/Niagara/Private/NiagaraRendererSprites.cpp b/Engine/Plugins/FX/Niagara/Source/Niagara/Private/NiagaraRendererSprites.cpp
--- a/Engine/Plugins/FX/Niagara/Source/Niagara/Private/NiagaraRendererSprites.cpp
+++ b/Engine/Plugins/FX/Niagara/Source/Niagara/Private/NiagaraRendererSprites.cpp
@@ -1034,7 +1034,7 @@ void FNiagaraRendererSprites::GetDynamicMeshElements(const TArray<const FSceneVi
 				else
 				{
 					FGlobalDynamicReadBuffer::FAllocation SortedIndices;
-					SortedIndices = Collector.GetDynamicReadBuffer().AllocateUInt32(RHICmdList, NumInstances);
+					SortedIndices = Collector.GetDynamicReadBuffer().AllocateUInt32(NumInstances);
 					NumInstances = SortAndCullIndices(SortInfo, *ParticleSpriteRenderData.SourceParticleData, SortedIndices);
 					VertexFactory.SetSortedIndices(SortedIndices.SRV, 0);
 				}
@@ -1118,7 +1118,7 @@ void FNiagaraRendererSprites::GetDynamicRayTracingInstances(FRayTracingMaterialG
 		else
 		{
 			FGlobalDynamicReadBuffer::FAllocation SortedIndices;
-			SortedIndices = DynamicReadBuffer.AllocateUInt32(RHICmdList, NumInstances);
+			SortedIndices = DynamicReadBuffer.AllocateUInt32(NumInstances);
 			NumInstances = SortAndCullIndices(SortInfo, *ParticleSpriteRenderData.SourceParticleData, SortedIndices);
 			VertexFactory.SetSortedIndices(SortedIndices.SRV, 0);
 		}
diff --git a/Engine/Source/Developer/MaterialBaking/Private/MaterialBakingModule.cpp b/Engine/Source/Developer/MaterialBaking/Private/MaterialBakingModule.cpp
--- a/Engine/Source/Developer/MaterialBaking/Private/MaterialBakingModule.cpp
+++ b/Engine/Source/Developer/MaterialBaking/Private/MaterialBakingModule.cpp
@@ -114,7 +114,7 @@ namespace
 			return Ref;
 		}
 
-		virtual FBufferRHIRef AllocIndexBuffer(uint32 NumElements) override
+		virtual FBufferRHIRef AllocIndexBuffer(FRHICommandListBase& RHICmdList, uint32 NumElements) override
 		{
 			uint32 BufferSize = GetIndexBufferSize(NumElements);
 			if (BufferSize > SmallestPooledBufferSize)
@@ -126,7 +126,7 @@ namespace
 				}
 			}
 
-			return FDynamicMeshBufferAllocator::AllocIndexBuffer(NumElements);
+			return FDynamicMeshBufferAllocator::AllocIndexBuffer(RHICmdList, NumElements);
 		}
 
 		virtual void ReleaseIndexBuffer(FBufferRHIRef& IndexBufferRHI) override
@@ -139,7 +139,7 @@ namespace
 			IndexBufferRHI = nullptr;
 		}
 
-		virtual FBufferRHIRef AllocVertexBuffer(uint32 Stride, uint32 NumElements) override
+		virtual FBufferRHIRef AllocVertexBuffer(FRHICommandListBase& RHICmdList, uint32 Stride, uint32 NumElements) override
 		{
 			uint32 BufferSize = GetVertexBufferSize(Stride, NumElements);
 			if (BufferSize > SmallestPooledBufferSize)
@@ -151,7 +151,7 @@ namespace
 				}
 			}
 
-			return FDynamicMeshBufferAllocator::AllocVertexBuffer(Stride, NumElements);
+			return FDynamicMeshBufferAllocator::AllocVertexBuffer(RHICmdList, Stride, NumElements);
 		}
 
 		virtual void ReleaseVertexBuffer(FBufferRHIRef& VertexBufferRHI) override
diff --git a/Engine/Source/Runtime/Engine/Private/DynamicMeshBuilder.cpp b/Engine/Source/Runtime/Engine/Private/DynamicMeshBuilder.cpp
--- a/Engine/Source/Runtime/Engine/Private/DynamicMeshBuilder.cpp
+++ b/Engine/Source/Runtime/Engine/Private/DynamicMeshBuilder.cpp
@@ -103,7 +103,7 @@ public:
 		// The use of BUF_Static is deliberate - on OS X the buffer backing-store orphaning & reallocation will dominate execution time
 		// so to avoid this we don't reuse a buffer for several frames, thereby avoiding the pipeline stall and the reallocation cost.
 		FRHIResourceCreateInfo CreateInfo(TEXT("FGlobalDynamicMeshIndexPolicy"));
-		FBufferRHIRef VertexBuffer = FRHICommandListImmediate::Get().CreateIndexBuffer(sizeof(DynamicMeshIndexType), BufferSize, BUF_Static, CreateInfo);
+		FBufferRHIRef VertexBuffer = RHICmdList.CreateIndexBuffer(sizeof(DynamicMeshIndexType), BufferSize, BUF_Static, CreateInfo);
 		return VertexBuffer;
 	}
 	
@@ -159,7 +159,7 @@ public:
 	{
 		FGlobalDynamicMeshPoolPolicy::CreationArguments BufferSize = GetPoolBucketSize(GetPoolBucketIndex(Args));
 		FRHIResourceCreateInfo CreateInfo(TEXT("FGlobalDynamicMeshVertexPolicy"));
-		FBufferRHIRef VertexBuffer = FRHICommandListImmediate::Get().CreateVertexBuffer(BufferSize, BUF_Volatile | BUF_ShaderResource, CreateInfo);
+		FBufferRHIRef VertexBuffer = RHICmdList.CreateVertexBuffer(BufferSize, BUF_Volatile | BUF_ShaderResource, CreateInfo);
 		return VertexBuffer;
 	}
 	
@@ -234,11 +234,16 @@ int32 FDynamicMeshBufferAllocator::GetVertexBufferSize(uint32 Stride, uint32 Num
 }
 
 FBufferRHIRef FDynamicMeshBufferAllocator::AllocIndexBuffer(uint32 NumElements)
+{
+	return AllocIndexBuffer(FRHICommandListImmediate::Get(), NumElements);
+}
+
+FBufferRHIRef FDynamicMeshBufferAllocator::AllocIndexBuffer(FRHICommandListBase& RHICmdList, uint32 NumElements)
 {
 	uint32 SizeInBytes = GetIndexBufferSize(NumElements);
 
 	FRHIResourceCreateInfo CreateInfo(TEXT("FDynamicMeshBufferAllocator"));
-	return FRHICommandListImmediate::Get().CreateIndexBuffer(sizeof(DynamicMeshIndexType), SizeInBytes, BUF_Volatile, CreateInfo);
+	return RHICmdList.CreateIndexBuffer(sizeof(DynamicMeshIndexType), SizeInBytes, BUF_Volatile, CreateInfo);
 }
 
 void FDynamicMeshBufferAllocator::ReleaseIndexBuffer(FBufferRHIRef& IndexBufferRHI)
@@ -247,11 +252,16 @@ void FDynamicMeshBufferAllocator::ReleaseIndexBuffer(FBufferRHIRef& IndexBufferR
 }
 
 FBufferRHIRef FDynamicMeshBufferAllocator::AllocVertexBuffer(uint32 Stride, uint32 NumElements)
+{
+	return AllocVertexBuffer(FRHICommandListImmediate::Get(), Stride, NumElements);
+}
+
+FBufferRHIRef FDynamicMeshBufferAllocator::AllocVertexBuffer(FRHICommandListBase& RHICmdList, uint32 Stride, uint32 NumElements)
 {
 	uint32 SizeInBytes = GetVertexBufferSize(Stride, NumElements);
 
 	FRHIResourceCreateInfo CreateInfo(TEXT("FDynamicMeshBufferAllocator"));
-	return FRHICommandListImmediate::Get().CreateVertexBuffer(SizeInBytes, BUF_Volatile | BUF_ShaderResource, CreateInfo);
+	return RHICmdList.CreateVertexBuffer(SizeInBytes, BUF_Volatile | BUF_ShaderResource, CreateInfo);
 }
 
 void FDynamicMeshBufferAllocator::ReleaseVertexBuffer(FBufferRHIRef& VertexBufferRHI)
@@ -262,15 +272,15 @@ void FDynamicMeshBufferAllocator::ReleaseVertexBuffer(FBufferRHIRef& VertexBuffe
 /** This is our default implementation using GDynamicMeshIndexPool. */
 class FPooledDynamicMeshBufferAllocator : public FDynamicMeshBufferAllocator
 {
-	virtual FBufferRHIRef AllocIndexBuffer(uint32 NumElements) override
+	virtual FBufferRHIRef AllocIndexBuffer(FRHICommandListBase& RHICmdList, uint32 NumElements) override
 	{
 		uint32 SizeInBytes = NumElements * sizeof(DynamicMeshIndexType);
 		if (SizeInBytes <= FGlobalDynamicMeshIndexPolicy().GetPoolBucketSize(FGlobalDynamicMeshIndexPolicy::NumPoolBuckets - 1))
 		{
-			return GDynamicMeshIndexPool.CreatePooledResource(SizeInBytes);
+			return GDynamicMeshIndexPool.CreatePooledResource(RHICmdList, SizeInBytes);
 		}
 
-		return FDynamicMeshBufferAllocator::AllocIndexBuffer(NumElements);
+		return FDynamicMeshBufferAllocator::AllocIndexBuffer(RHICmdList, NumElements);
 	}
 
 	virtual void ReleaseIndexBuffer(FBufferRHIRef& IndexBufferRHI)
@@ -286,15 +296,15 @@ class FPooledDynamicMeshBufferAllocator : public FDynamicMeshBufferAllocator
 		}
 	}
 
-	virtual FBufferRHIRef AllocVertexBuffer(uint32 Stride, uint32 NumElements)
+	virtual FBufferRHIRef AllocVertexBuffer(FRHICommandListBase& RHICmdList, uint32 Stride, uint32 NumElements)
 	{
 		uint32 SizeInBytes = NumElements * Stride;
 		if (SizeInBytes <= FGlobalDynamicMeshVertexPolicy().GetPoolBucketSize(FGlobalDynamicMeshVertexPolicy::NumPoolBuckets - 1))
 		{
-			return GDynamicMeshVertexPool.CreatePooledResource(SizeInBytes);
+			return GDynamicMeshVertexPool.CreatePooledResource(RHICmdList, SizeInBytes);
 		}
 
-		return FDynamicMeshBufferAllocator::AllocVertexBuffer(Stride, NumElements);
+		return FDynamicMeshBufferAllocator::AllocVertexBuffer(RHICmdList, Stride, NumElements);
 	}
 
 	virtual void ReleaseVertexBuffer(FBufferRHIRef& VertexBufferRHI)
@@ -329,7 +339,7 @@ public:
 	{
 		TRACE_CPUPROFILER_EVENT_SCOPE(FPooledDynamicMeshIndexBuffer::InitRHI)
 
-		IndexBufferRHI = DynamicMeshBufferAllocator.AllocIndexBuffer(Indices.Num());
+		IndexBufferRHI = DynamicMeshBufferAllocator.AllocIndexBuffer(RHICmdList, Indices.Num());
 
 		// Write the indices to the index buffer.
 		void* Buffer;
@@ -411,10 +421,10 @@ public:
 			TextureFormat = PF_G16R16F;
 		}
 
-		PositionBuffer.VertexBufferRHI = DynamicMeshBufferAllocator.AllocVertexBuffer(sizeof(FVector3f), Vertices.Num());
-		TangentBuffer.VertexBufferRHI  = DynamicMeshBufferAllocator.AllocVertexBuffer(sizeof(FPackedNormal), 2 * Vertices.Num());
-		TexCoordBuffer.VertexBufferRHI = DynamicMeshBufferAllocator.AllocVertexBuffer(TextureStride, NumTexCoords * Vertices.Num());
-		ColorBuffer.VertexBufferRHI    = DynamicMeshBufferAllocator.AllocVertexBuffer(sizeof(FColor), Vertices.Num());
+		PositionBuffer.VertexBufferRHI = DynamicMeshBufferAllocator.AllocVertexBuffer(RHICmdList, sizeof(FVector3f), Vertices.Num());
+		TangentBuffer.VertexBufferRHI  = DynamicMeshBufferAllocator.AllocVertexBuffer(RHICmdList, sizeof(FPackedNormal), 2 * Vertices.Num());
+		TexCoordBuffer.VertexBufferRHI = DynamicMeshBufferAllocator.AllocVertexBuffer(RHICmdList, TextureStride, NumTexCoords * Vertices.Num());
+		ColorBuffer.VertexBufferRHI    = DynamicMeshBufferAllocator.AllocVertexBuffer(RHICmdList, sizeof(FColor), Vertices.Num());
 
 		void* TexCoordBufferData = RHICmdList.LockBuffer(TexCoordBuffer.VertexBufferRHI, 0, NumTexCoords * TextureStride * Vertices.Num(), RLM_WriteOnly);
 		FVector2f* TexCoordBufferData32 = !Use16bitTexCoord ? static_cast<FVector2f*>(TexCoordBufferData) : nullptr;
@@ -538,77 +548,74 @@ public:
 	{
 		FLocalVertexFactory* VertexFactory = this;
 		const FPooledDynamicMeshVertexBuffer* PooledVertexBuffer = VertexBuffer;
-		ENQUEUE_RENDER_COMMAND(InitDynamicMeshVertexFactory)(
-			[VertexFactory, PooledVertexBuffer](FRHICommandListImmediate& RHICmdList)
+
+		FDataType LocalData;
+		LocalData.PositionComponent = FVertexStreamComponent(
+			&PooledVertexBuffer->PositionBuffer,
+			0,
+			sizeof(FVector3f),
+			VET_Float3
+		);
+
+		LocalData.NumTexCoords = PooledVertexBuffer->GetNumTexCoords();
+		{
+			LocalData.LightMapCoordinateIndex = PooledVertexBuffer->GetLightmapCoordinateIndex();
+			LocalData.TangentsSRV = PooledVertexBuffer->TangentBufferSRV;
+			LocalData.TextureCoordinatesSRV = PooledVertexBuffer->TexCoordBufferSRV;
+			LocalData.ColorComponentsSRV = PooledVertexBuffer->ColorBufferSRV;
+			LocalData.PositionComponentSRV = PooledVertexBuffer->PositionBufferSRV;
+		}
+
 		{
-			FDataType Data;
-			Data.PositionComponent = FVertexStreamComponent(
-				&PooledVertexBuffer->PositionBuffer,
-				0,
-				sizeof(FVector3f),
-				VET_Float3
-			);
-
-			Data.NumTexCoords = PooledVertexBuffer->GetNumTexCoords();
+			EVertexElementType UVDoubleWideVertexElementType = VET_None;
+			EVertexElementType UVVertexElementType = VET_None;
+			uint32 UVSizeInBytes = 0;
+			if (PooledVertexBuffer->GetUse16bitTexCoords())
 			{
-				Data.LightMapCoordinateIndex = PooledVertexBuffer->GetLightmapCoordinateIndex();
-				Data.TangentsSRV = PooledVertexBuffer->TangentBufferSRV;
-				Data.TextureCoordinatesSRV = PooledVertexBuffer->TexCoordBufferSRV;
-				Data.ColorComponentsSRV = PooledVertexBuffer->ColorBufferSRV;
-				Data.PositionComponentSRV = PooledVertexBuffer->PositionBufferSRV;
+				UVSizeInBytes = sizeof(FVector2DHalf);
+				UVDoubleWideVertexElementType = VET_Half4;
+				UVVertexElementType = VET_Half2;
 			}
-
+			else
 			{
-				EVertexElementType UVDoubleWideVertexElementType = VET_None;
-				EVertexElementType UVVertexElementType = VET_None;
-				uint32 UVSizeInBytes = 0;
-				if (PooledVertexBuffer->GetUse16bitTexCoords())
-				{
-					UVSizeInBytes = sizeof(FVector2DHalf);
-					UVDoubleWideVertexElementType = VET_Half4;
-					UVVertexElementType = VET_Half2;
-				}
-				else
-				{
-					UVSizeInBytes = sizeof(FVector2f);
-					UVDoubleWideVertexElementType = VET_Float4;
-					UVVertexElementType = VET_Float2;
-				}
-
-				int32 UVIndex;
-				uint32 UvStride = UVSizeInBytes * PooledVertexBuffer->GetNumTexCoords();
-				for (UVIndex = 0; UVIndex < (int32)PooledVertexBuffer->GetNumTexCoords() - 1; UVIndex += 2)
-				{
-					Data.TextureCoordinates.Add
-					(
-						FVertexStreamComponent(
-							&PooledVertexBuffer->TexCoordBuffer, 
-							UVSizeInBytes * UVIndex, 
-							UvStride,
-							UVDoubleWideVertexElementType,
-							EVertexStreamUsage::ManualFetch
-						)
-					);
-				}
+				UVSizeInBytes = sizeof(FVector2f);
+				UVDoubleWideVertexElementType = VET_Float4;
+				UVVertexElementType = VET_Float2;
+			}
 
-				// possible last UV channel if we have an odd number
-				if (UVIndex < (int32)PooledVertexBuffer->GetNumTexCoords())
-				{
-					Data.TextureCoordinates.Add(FVertexStreamComponent(
-						&PooledVertexBuffer->TexCoordBuffer,
-						UVSizeInBytes * UVIndex,
+			int32 UVIndex;
+			uint32 UvStride = UVSizeInBytes * PooledVertexBuffer->GetNumTexCoords();
+			for (UVIndex = 0; UVIndex < (int32)PooledVertexBuffer->GetNumTexCoords() - 1; UVIndex += 2)
+			{
+				LocalData.TextureCoordinates.Add
+				(
+					FVertexStreamComponent(
+						&PooledVertexBuffer->TexCoordBuffer, 
+						UVSizeInBytes * UVIndex, 
 						UvStride,
-						UVVertexElementType, 
+						UVDoubleWideVertexElementType,
 						EVertexStreamUsage::ManualFetch
-					));
-				}
+					)
+				);
+			}
 
-				Data.TangentBasisComponents[0] = FVertexStreamComponent(&PooledVertexBuffer->TangentBuffer, 0, 2 * sizeof(FPackedNormal), VET_PackedNormal, EVertexStreamUsage::ManualFetch);
-				Data.TangentBasisComponents[1] = FVertexStreamComponent(&PooledVertexBuffer->TangentBuffer, sizeof(FPackedNormal), 2 * sizeof(FPackedNormal), VET_PackedNormal, EVertexStreamUsage::ManualFetch);
-				Data.ColorComponent = FVertexStreamComponent(&PooledVertexBuffer->ColorBuffer, 0, sizeof(FColor), VET_Color, EVertexStreamUsage::ManualFetch);
+			// possible last UV channel if we have an odd number
+			if (UVIndex < (int32)PooledVertexBuffer->GetNumTexCoords())
+			{
+				LocalData.TextureCoordinates.Add(FVertexStreamComponent(
+					&PooledVertexBuffer->TexCoordBuffer,
+					UVSizeInBytes * UVIndex,
+					UvStride,
+					UVVertexElementType, 
+					EVertexStreamUsage::ManualFetch
+				));
 			}
-			VertexFactory->SetData(RHICmdList, Data);
-		});
+
+			LocalData.TangentBasisComponents[0] = FVertexStreamComponent(&PooledVertexBuffer->TangentBuffer, 0, 2 * sizeof(FPackedNormal), VET_PackedNormal, EVertexStreamUsage::ManualFetch);
+			LocalData.TangentBasisComponents[1] = FVertexStreamComponent(&PooledVertexBuffer->TangentBuffer, sizeof(FPackedNormal), 2 * sizeof(FPackedNormal), VET_PackedNormal, EVertexStreamUsage::ManualFetch);
+			LocalData.ColorComponent = FVertexStreamComponent(&PooledVertexBuffer->ColorBuffer, 0, sizeof(FColor), VET_Color, EVertexStreamUsage::ManualFetch);
+		}
+		VertexFactory->SetData(RHICmdList, LocalData);
 
 		FLocalVertexFactory::InitResource(RHICmdList);
 	}
@@ -820,7 +827,7 @@ void FDynamicMeshBuilder::GetMesh(
 {
 	TRACE_CPUPROFILER_EVENT_SCOPE(FDynamicMeshBuilder::GetMesh)
 
-	FRHICommandListBase& RHICmdList = FRHICommandListImmediate::Get();
+	FRHICommandListBase& RHICmdList = Collector.GetRHICommandList();
 
 	// Only draw non-empty meshes.
 	if((VertexBuffer && VertexBuffer->Vertices.Num() > 0) || (DrawOffset != nullptr))
@@ -858,14 +865,8 @@ void FDynamicMeshBuilder::GetMesh(
 					.OutputVelocity(true)
 				.Build();
 
-			if (IsInGameThread())
-			{
-				BeginSetUniformBufferContents(*OneFrameResources->PrimitiveUniformBuffer, PrimitiveParams);
-			}
-			else
-			{
-				OneFrameResources->PrimitiveUniformBuffer->SetContents(RHICmdList, PrimitiveParams);
-			}
+			
+			OneFrameResources->PrimitiveUniformBuffer->SetContents(RHICmdList, PrimitiveParams);
 
 			OneFrameResources->PrimitiveUniformBuffer->InitResource(RHICmdList);
 
@@ -949,14 +950,7 @@ void FDynamicMeshBuilder::GetMeshElement(const FPrimitiveUniformShaderParameters
 			// Create the primitive uniform buffer.
 			OneFrameResource.PrimitiveUniformBuffer = new FDynamicMeshPrimitiveUniformBuffer();
 
-			if (IsInGameThread())
-			{
-				BeginSetUniformBufferContents(*OneFrameResource.PrimitiveUniformBuffer, PrimitiveParams);
-			}
-			else
-			{
-				OneFrameResource.PrimitiveUniformBuffer->SetContents(RHICmdList, PrimitiveParams);
-			}
+			OneFrameResource.PrimitiveUniformBuffer->SetContents(RHICmdList, PrimitiveParams);
 
 			OneFrameResource.PrimitiveUniformBuffer->InitResource(RHICmdList);
 
diff --git a/Engine/Source/Runtime/Engine/Private/SceneManagement.cpp b/Engine/Source/Runtime/Engine/Private/SceneManagement.cpp
--- a/Engine/Source/Runtime/Engine/Private/SceneManagement.cpp
+++ b/Engine/Source/Runtime/Engine/Private/SceneManagement.cpp
@@ -22,6 +22,7 @@
 #include "StaticMeshBatch.h"
 #include "PrimitiveUniformShaderParametersBuilder.h"
 #include "PrimitiveSceneShaderData.h"
+#include "RenderGraphBuilder.h"
 
 static TAutoConsoleVariable<float> CVarLODTemporalLag(
 	TEXT("lod.TemporalLag"),
@@ -320,29 +321,50 @@ FMeshBatchAndRelevance::FMeshBatchAndRelevance(const FMeshBatch& InMesh, const F
 	bRenderInMainPass = PrimitiveSceneProxy->ShouldRenderInMainPass();
 }
 
+#if RHI_RAYTRACING
+
+FRayTracingMaterialGatheringContext::FRayTracingMaterialGatheringContext(
+	const FScene* InScene,
+	const FSceneView* InReferenceView,
+	const FSceneViewFamily& InReferenceViewFamily,
+	FRDGBuilder& InGraphBuilder,
+	FRayTracingMeshResourceCollector& InRayTracingMeshResourceCollector,
+	FGlobalDynamicReadBuffer& DynamicReadBuffer)
+	: Scene(InScene)
+	, ReferenceView(InReferenceView)
+	, ReferenceViewFamily(InReferenceViewFamily)
+	, GraphBuilder(InGraphBuilder)
+	, RHICmdList(GraphBuilder.RHICmdList)
+	, RayTracingMeshResourceCollector(InRayTracingMeshResourceCollector)
+	, DynamicVertexBuffer(GraphBuilder.RHICmdList)
+	, DynamicIndexBuffer(GraphBuilder.RHICmdList)
+{
+	RayTracingMeshResourceCollector.Start(RHICmdList, DynamicVertexBuffer, DynamicIndexBuffer, DynamicReadBuffer);
+}
+
+FRayTracingMaterialGatheringContext::~FRayTracingMaterialGatheringContext()
+{
+	RayTracingMeshResourceCollector.Finish();
+}
+
+#endif
+
 FMeshElementCollector::FMeshElementCollector(ERHIFeatureLevel::Type InFeatureLevel, FSceneRenderingBulkObjectAllocator& InBulkAllocator) :
 	OneFrameResources(InBulkAllocator),
 	PrimitiveSceneProxy(NULL),
-	DynamicIndexBuffer(nullptr),
-	DynamicVertexBuffer(nullptr),
 	DynamicReadBuffer(nullptr),
-	FeatureLevel(InFeatureLevel)
-{	
-}
-
-FMeshElementCollector::~FMeshElementCollector()
+	FeatureLevel(InFeatureLevel),
+	bUseGPUScene(UseGPUScene(GMaxRHIShaderPlatform, FeatureLevel))
 {
-	DeleteTemporaryProxies();
 }
 
-void FMeshElementCollector::DeleteTemporaryProxies()
+FMeshElementCollector::~FMeshElementCollector()
 {
-	for (int32 ProxyIndex = 0; ProxyIndex < TemporaryProxies.Num(); ProxyIndex++)
+	for (FMaterialRenderProxy* Proxy : MaterialProxiesToDelete)
 	{
-		delete TemporaryProxies[ProxyIndex];
+		delete Proxy;
 	}
-
-	TemporaryProxies.Empty();
+	MaterialProxiesToDelete.Empty();
 }
 
 void FMeshElementCollector::SetPrimitive(const FPrimitiveSceneProxy* InPrimitiveSceneProxy, FHitProxyId DefaultHitProxyId)
@@ -369,57 +391,97 @@ void FMeshElementCollector::SetPrimitive(const FPrimitiveSceneProxy* InPrimitive
 #endif
 }
 
-void FMeshElementCollector::ClearViewMeshArrays()
+void FMeshElementCollector::Start(
+	FRHICommandList& InRHICmdList,
+	FGlobalDynamicVertexBuffer& InDynamicVertexBuffer,
+	FGlobalDynamicIndexBuffer& InDynamicIndexBuffer,
+	FGlobalDynamicReadBuffer& InDynamicReadBuffer)
 {
-	Views.Empty();
-	MeshBatches.Empty();
-	SimpleElementCollectors.Empty();
-#if UE_ENABLE_DEBUG_DRAWING
-	DebugSimpleElementCollectors.Empty();
-#endif
-	MeshIdInPrimitivePerView.Empty();
-	DynamicPrimitiveCollectorPerView.Empty();
-	NumMeshBatchElementsPerView.Empty();
-	DynamicIndexBuffer = nullptr;
-	DynamicVertexBuffer = nullptr;
-	DynamicReadBuffer = nullptr;
+	check(!RHICmdList);
+	RHICmdList = &InRHICmdList;
+	DynamicVertexBuffer = &InDynamicVertexBuffer;
+	DynamicIndexBuffer = &InDynamicIndexBuffer;
+	DynamicReadBuffer = &InDynamicReadBuffer;
 }
 
 void FMeshElementCollector::AddViewMeshArrays(
 	FSceneView* InView,
 	TArray<FMeshBatchAndRelevance, SceneRenderingAllocator>* ViewMeshes,
 	FSimpleElementCollector* ViewSimpleElementCollector,
-	FGPUScenePrimitiveCollector* InDynamicPrimitiveCollector,
-	ERHIFeatureLevel::Type InFeatureLevel,
-	FGlobalDynamicIndexBuffer* InDynamicIndexBuffer,
-	FGlobalDynamicVertexBuffer* InDynamicVertexBuffer,
-	FGlobalDynamicReadBuffer* InDynamicReadBuffer
+	FGPUScenePrimitiveCollector* DynamicPrimitiveCollector
 #if UE_ENABLE_DEBUG_DRAWING
-	,FSimpleElementCollector* InDebugSimpleElementCollector
+	, FSimpleElementCollector* DebugSimpleElementCollector
 #endif
 )
 {
+	check(RHICmdList);
+
 	Views.Add(InView);
 	MeshIdInPrimitivePerView.Add(0);
 	MeshBatches.Add(ViewMeshes);
 	NumMeshBatchElementsPerView.Add(0);
 	SimpleElementCollectors.Add(ViewSimpleElementCollector);
-	DynamicPrimitiveCollectorPerView.Add(InDynamicPrimitiveCollector);
-
-	check(InDynamicIndexBuffer && InDynamicVertexBuffer && InDynamicReadBuffer);
-	DynamicIndexBuffer = InDynamicIndexBuffer;
-	DynamicVertexBuffer = InDynamicVertexBuffer;
-	DynamicReadBuffer = InDynamicReadBuffer;
+	DynamicPrimitiveCollectorPerView.Add(DynamicPrimitiveCollector);
 
 #if UE_ENABLE_DEBUG_DRAWING
 	//Assign the debug draw only simple element collector per view	
-	if (InDebugSimpleElementCollector)
+	if (DebugSimpleElementCollector)
 	{
-		DebugSimpleElementCollectors.Add(InDebugSimpleElementCollector);
+		DebugSimpleElementCollectors.Add(DebugSimpleElementCollector);
 	}
 #endif
 }
 
+void FMeshElementCollector::ClearViewMeshArrays()
+{
+	Views.Reset();
+	MeshIdInPrimitivePerView.Reset();
+	MeshBatches.Reset();
+	NumMeshBatchElementsPerView.Reset();
+	SimpleElementCollectors.Reset();
+	DynamicPrimitiveCollectorPerView.Reset();
+#if UE_ENABLE_DEBUG_DRAWING
+	DebugSimpleElementCollectors.Reset();
+#endif
+}
+
+void FMeshElementCollector::Commit()
+{
+	TRACE_CPUPROFILER_EVENT_SCOPE(FMeshElementCollector::Commit);
+	check(RHICmdList);
+
+	for (TPair<FGPUScenePrimitiveCollector*, FMeshBatch*> Pair : MeshBatchesForGPUScene)
+	{
+		GetRendererModule().AddMeshBatchToGPUScene(Pair.Key, *Pair.Value);
+	}
+
+	for (TPair<FMaterialRenderProxy*, bool> Parameters : MaterialProxiesToInvalidate)
+	{
+		Parameters.Key->InvalidateUniformExpressionCache(Parameters.Value);
+	}
+
+	for (const FMaterialRenderProxy* Proxy : MaterialProxiesToUpdate)
+	{
+		Proxy->UpdateUniformExpressionCacheIfNeeded(*RHICmdList, FeatureLevel);
+	}
+
+	MeshBatchesForGPUScene.Empty();
+	MaterialProxiesToInvalidate.Empty();
+	MaterialProxiesToUpdate.Empty();
+}
+
+void FMeshElementCollector::Finish()
+{
+	SCOPED_NAMED_EVENT(FMeshElementCollector_Finish, FColor::Magenta);
+
+	Commit();
+	ClearViewMeshArrays();
+	DynamicIndexBuffer = nullptr;
+	DynamicVertexBuffer = nullptr;
+	DynamicReadBuffer = nullptr;
+	RHICmdList = nullptr;
+}
+
 void FMeshElementCollector::AddMesh(int32 ViewIndex, FMeshBatch& MeshBatch)
 {
 	DEFINE_LOG_CATEGORY_STATIC(FMeshElementCollector_AddMesh, Warning, All);
@@ -445,21 +507,19 @@ void FMeshElementCollector::AddMesh(int32 ViewIndex, FMeshBatch& MeshBatch)
 
 	MeshBatch.PreparePrimitiveUniformBuffer(PrimitiveSceneProxy, FeatureLevel);
 
-	// If we are maintaining primitive scene data on the GPU, copy the primitive uniform buffer data to a unified array so it can be uploaded later
-	if (UseGPUScene(GMaxRHIShaderPlatform, FeatureLevel) && MeshBatch.VertexFactory->GetPrimitiveIdStreamIndex(FeatureLevel, EVertexInputStreamType::Default) >= 0)
+	if (bUseGPUScene && MeshBatch.VertexFactory->GetPrimitiveIdStreamIndex(FeatureLevel, EVertexInputStreamType::Default) >= 0)
 	{
-		GetRendererModule().AddMeshBatchToGPUScene(DynamicPrimitiveCollectorPerView[ViewIndex], MeshBatch);
+		MeshBatchesForGPUScene.Emplace(DynamicPrimitiveCollectorPerView[ViewIndex], &MeshBatch);
 	}
 
-	MeshBatch.MaterialRenderProxy->UpdateUniformExpressionCacheIfNeeded(Views[ViewIndex]->GetFeatureLevel());
+	MaterialProxiesToUpdate.Emplace(MeshBatch.MaterialRenderProxy);
 
 	MeshBatch.MeshIdInPrimitive = MeshIdInPrimitivePerView[ViewIndex];
 	++MeshIdInPrimitivePerView[ViewIndex];
 
 	NumMeshBatchElementsPerView[ViewIndex] += MeshBatch.Elements.Num();
 
-	TArray<FMeshBatchAndRelevance,SceneRenderingAllocator>& ViewMeshBatches = *MeshBatches[ViewIndex];
-	new (ViewMeshBatches) FMeshBatchAndRelevance(MeshBatch, PrimitiveSceneProxy, FeatureLevel);	
+	MeshBatches[ViewIndex]->Emplace(MeshBatch, PrimitiveSceneProxy, FeatureLevel);
 }
 
 FDynamicPrimitiveUniformBuffer::FDynamicPrimitiveUniformBuffer() = default;
@@ -469,6 +529,7 @@ FDynamicPrimitiveUniformBuffer::~FDynamicPrimitiveUniformBuffer()
 }
 
 void FDynamicPrimitiveUniformBuffer::Set(
+	FRHICommandListBase& RHICmdList,
 	const FMatrix& LocalToWorld,
 	const FMatrix& PreviousLocalToWorld,
 	const FVector& ActorPositionWS,
@@ -480,7 +541,7 @@ void FDynamicPrimitiveUniformBuffer::Set(
 	bool bOutputVelocity,
 	const FCustomPrimitiveData* CustomPrimitiveData)
 {
-	FRHICommandListBase& RHICmdList = FRHICommandListImmediate::Get();
+	UniformBuffer.BufferUsage = UniformBuffer_SingleFrame;
 	UniformBuffer.SetContents(
 		RHICmdList,
 		FPrimitiveUniformShaderParametersBuilder{}
@@ -500,9 +561,52 @@ void FDynamicPrimitiveUniformBuffer::Set(
 	UniformBuffer.InitResource(RHICmdList);
 }
 
+void FDynamicPrimitiveUniformBuffer::Set(
+	FRHICommandListBase& RHICmdList,
+	const FMatrix& LocalToWorld,
+	const FMatrix& PreviousLocalToWorld,
+	const FBoxSphereBounds& WorldBounds,
+	const FBoxSphereBounds& LocalBounds,
+	const FBoxSphereBounds& PreSkinnedLocalBounds,
+	bool bReceivesDecals,
+	bool bHasPrecomputedVolumetricLightmap,
+	bool bOutputVelocity,
+	const FCustomPrimitiveData* CustomPrimitiveData)
+{
+	Set(RHICmdList, LocalToWorld, PreviousLocalToWorld, WorldBounds.Origin, WorldBounds, LocalBounds, PreSkinnedLocalBounds, bReceivesDecals, bHasPrecomputedVolumetricLightmap, bOutputVelocity, CustomPrimitiveData);
+}
+
+void FDynamicPrimitiveUniformBuffer::Set(
+	FRHICommandListBase& RHICmdList,
+	const FMatrix& LocalToWorld,
+	const FMatrix& PreviousLocalToWorld,
+	const FBoxSphereBounds& WorldBounds,
+	const FBoxSphereBounds& LocalBounds,
+	const FBoxSphereBounds& PreSkinnedLocalBounds,
+	bool bReceivesDecals,
+	bool bHasPrecomputedVolumetricLightmap,
+	bool bOutputVelocity)
+{
+	Set(RHICmdList, LocalToWorld, PreviousLocalToWorld, WorldBounds, LocalBounds, PreSkinnedLocalBounds, bReceivesDecals, bHasPrecomputedVolumetricLightmap, bOutputVelocity, nullptr);
+}
+
+void FDynamicPrimitiveUniformBuffer::Set(
+	FRHICommandListBase& RHICmdList,
+	const FMatrix& LocalToWorld,
+	const FMatrix& PreviousLocalToWorld,
+	const FBoxSphereBounds& WorldBounds,
+	const FBoxSphereBounds& LocalBounds,
+	bool bReceivesDecals,
+	bool bHasPrecomputedVolumetricLightmap,
+	bool bOutputVelocity)
+{
+	Set(RHICmdList, LocalToWorld, PreviousLocalToWorld, WorldBounds, LocalBounds, LocalBounds, bReceivesDecals, bHasPrecomputedVolumetricLightmap, bOutputVelocity, nullptr);
+}
+
 void FDynamicPrimitiveUniformBuffer::Set(
 	const FMatrix& LocalToWorld,
 	const FMatrix& PreviousLocalToWorld,
+	const FVector& ActorPositionWS,
 	const FBoxSphereBounds& WorldBounds,
 	const FBoxSphereBounds& LocalBounds,
 	const FBoxSphereBounds& PreSkinnedLocalBounds,
@@ -511,7 +615,21 @@ void FDynamicPrimitiveUniformBuffer::Set(
 	bool bOutputVelocity,
 	const FCustomPrimitiveData* CustomPrimitiveData)
 {
-	Set(LocalToWorld, PreviousLocalToWorld, WorldBounds.Origin, WorldBounds, LocalBounds, PreSkinnedLocalBounds, bReceivesDecals, bHasPrecomputedVolumetricLightmap, bOutputVelocity, CustomPrimitiveData);
+	Set(FRHICommandListImmediate::Get(), LocalToWorld, PreviousLocalToWorld, ActorPositionWS, WorldBounds, LocalBounds, PreSkinnedLocalBounds, bReceivesDecals, bHasPrecomputedVolumetricLightmap, bOutputVelocity, CustomPrimitiveData);
+}
+
+void FDynamicPrimitiveUniformBuffer::Set(
+	const FMatrix& LocalToWorld,
+	const FMatrix& PreviousLocalToWorld,
+	const FBoxSphereBounds& WorldBounds,
+	const FBoxSphereBounds& LocalBounds,
+	const FBoxSphereBounds& PreSkinnedLocalBounds,
+	bool bReceivesDecals,
+	bool bHasPrecomputedVolumetricLightmap,
+	bool bOutputVelocity,
+	const FCustomPrimitiveData* CustomPrimitiveData)
+{
+	Set(FRHICommandListImmediate::Get(), LocalToWorld, PreviousLocalToWorld, WorldBounds.Origin, WorldBounds, LocalBounds, PreSkinnedLocalBounds, bReceivesDecals, bHasPrecomputedVolumetricLightmap, bOutputVelocity, CustomPrimitiveData);
 }
 
 void FDynamicPrimitiveUniformBuffer::Set(
@@ -524,7 +642,7 @@ void FDynamicPrimitiveUniformBuffer::Set(
 	bool bHasPrecomputedVolumetricLightmap,
 	bool bOutputVelocity)
 {
-	Set(LocalToWorld, PreviousLocalToWorld, WorldBounds, LocalBounds, PreSkinnedLocalBounds, bReceivesDecals, bHasPrecomputedVolumetricLightmap, bOutputVelocity, nullptr);
+	Set(FRHICommandListImmediate::Get(), LocalToWorld, PreviousLocalToWorld, WorldBounds, LocalBounds, PreSkinnedLocalBounds, bReceivesDecals, bHasPrecomputedVolumetricLightmap, bOutputVelocity, nullptr);
 }
 
 void FDynamicPrimitiveUniformBuffer::Set(
@@ -536,7 +654,7 @@ void FDynamicPrimitiveUniformBuffer::Set(
 	bool bHasPrecomputedVolumetricLightmap,
 	bool bOutputVelocity)
 {
-	Set(LocalToWorld, PreviousLocalToWorld, WorldBounds, LocalBounds, LocalBounds, bReceivesDecals, bHasPrecomputedVolumetricLightmap, bOutputVelocity, nullptr);
+	Set(FRHICommandListImmediate::Get(), LocalToWorld, PreviousLocalToWorld, WorldBounds, LocalBounds, LocalBounds, bReceivesDecals, bHasPrecomputedVolumetricLightmap, bOutputVelocity, nullptr);
 }
 
 FLightMapInteraction FLightMapInteraction::Texture(
@@ -1417,7 +1535,7 @@ void FMeshBatch::PreparePrimitiveUniformBuffer(const FPrimitiveSceneProxy* Primi
 
 #if USE_MESH_BATCH_VALIDATION
 bool FMeshBatch::Validate(const FPrimitiveSceneProxy* PrimitiveSceneProxy, ERHIFeatureLevel::Type FeatureLevel) const
-		{
+{
 	check(PrimitiveSceneProxy);
 
 	const auto LogMeshError = [&](const FString& Error) -> bool
diff --git a/Engine/Source/Runtime/Engine/Public/DynamicMeshBuilder.h b/Engine/Source/Runtime/Engine/Public/DynamicMeshBuilder.h
--- a/Engine/Source/Runtime/Engine/Public/DynamicMeshBuilder.h
+++ b/Engine/Source/Runtime/Engine/Public/DynamicMeshBuilder.h
@@ -137,10 +137,16 @@ public:
 	ENGINE_API int32 GetVertexBufferSize(uint32 Stride, uint32 NumElements) const;
 	ENGINE_API virtual ~FDynamicMeshBufferAllocator();
 
-	ENGINE_API virtual FBufferRHIRef AllocIndexBuffer(uint32 NumElements);
+	ENGINE_API virtual FBufferRHIRef AllocIndexBuffer(FRHICommandListBase& RHICmdList, uint32 NumElements);
 	ENGINE_API virtual void ReleaseIndexBuffer(FBufferRHIRef& IndexBufferRHI);
-	ENGINE_API virtual FBufferRHIRef AllocVertexBuffer(uint32 Stride, uint32 NumElements);
+	ENGINE_API virtual FBufferRHIRef AllocVertexBuffer(FRHICommandListBase& RHICmdList, uint32 Stride, uint32 NumElements);
 	ENGINE_API virtual void ReleaseVertexBuffer(FBufferRHIRef& VertexBufferRHI);
+
+	UE_DEPRECATED(5.4, "AllocIndexBuffer requires a command list.")
+	ENGINE_API virtual FBufferRHIRef AllocIndexBuffer(uint32 NumElements) final;
+
+	UE_DEPRECATED(5.4, "AllocVertexBuffer requires a command list.")
+	ENGINE_API virtual FBufferRHIRef AllocVertexBuffer(uint32 Stride, uint32 NumElements) final;
 };
 
 /**
diff --git a/Engine/Source/Runtime/Engine/Public/ResourcePool.h b/Engine/Source/Runtime/Engine/Public/ResourcePool.h
--- a/Engine/Source/Runtime/Engine/Public/ResourcePool.h
+++ b/Engine/Source/Runtime/Engine/Public/ResourcePool.h
@@ -219,6 +219,7 @@ public:
 		}
 	}
 
+	UE_DEPRECATED(5.4, "CreatePooledResource requires an RHI command list.")
 	ResourceType CreatePooledResource(ResourceCreationArguments Args)
 	{
 		if (IsInitialized())
diff --git a/Engine/Source/Runtime/Engine/Public/SceneManagement.h b/Engine/Source/Runtime/Engine/Public/SceneManagement.h
--- a/Engine/Source/Runtime/Engine/Public/SceneManagement.h
+++ b/Engine/Source/Runtime/Engine/Public/SceneManagement.h
@@ -29,6 +29,7 @@
 #include "DynamicBufferAllocator.h"
 #include "Rendering/SkyAtmosphereCommonData.h"
 #include "Math/SHMath.h"
+#include "GlobalRenderResources.h"
 
 #if UE_ENABLE_INCLUDE_ORDER_DEPRECATED_IN_5_2
 #include "Engine/TextureLightProfile.h"
@@ -1898,6 +1899,13 @@ public:
 		return *DynamicReadBuffer;
 	}
 
+	/** Return the current RHI command list used to initialize resources. */
+	FRHICommandList& GetRHICommandList()
+	{
+		check(RHICmdList);
+		return *RHICmdList;
+	}
+
 	// @return number of MeshBatches collected (so far) for a given view
 	uint32 GetMeshBatchCount(uint32 ViewIndex) const
 	{
@@ -1918,7 +1926,15 @@ public:
 	/** Add a material render proxy that will be cleaned up automatically */
 	void RegisterOneFrameMaterialProxy(FMaterialRenderProxy* Proxy)
 	{
-		TemporaryProxies.Add(Proxy);
+		check(Proxy);
+		MaterialProxiesToDelete.Add(Proxy);
+	}
+
+	/** Adds a request to force caching of uniform expressions for a material render proxy. */
+	void CacheUniformExpressions(FMaterialRenderProxy* Proxy, bool bRecreateUniformBuffer)
+	{
+		check(Proxy);
+		MaterialProxiesToInvalidate.Emplace(Proxy, bRecreateUniformBuffer);
 	}
 
 	/** Allocates a temporary resource that is safe to be referenced by an FMeshBatch added to the collector. */
@@ -1954,25 +1970,29 @@ protected:
 
 	ENGINE_API ~FMeshElementCollector();
 
-	ENGINE_API void DeleteTemporaryProxies();
-
 	ENGINE_API void SetPrimitive(const FPrimitiveSceneProxy* InPrimitiveSceneProxy, FHitProxyId DefaultHitProxyId);
 
-	ENGINE_API void ClearViewMeshArrays();
+	ENGINE_API void Start(
+		FRHICommandList& RHICmdList,
+		FGlobalDynamicVertexBuffer& DynamicVertexBuffer,
+		FGlobalDynamicIndexBuffer& DynamicIndexBuffer,
+		FGlobalDynamicReadBuffer& DynamicReadBuffer);
 
 	ENGINE_API void AddViewMeshArrays(
 		FSceneView* InView,
 		TArray<FMeshBatchAndRelevance, SceneRenderingAllocator>* ViewMeshes,
 		FSimpleElementCollector* ViewSimpleElementCollector,
-		FGPUScenePrimitiveCollector* InDynamicPrimitiveCollector,
-		ERHIFeatureLevel::Type InFeatureLevel,
-		FGlobalDynamicIndexBuffer* InDynamicIndexBuffer,
-		FGlobalDynamicVertexBuffer* InDynamicVertexBuffer,
-		FGlobalDynamicReadBuffer* InDynamicReadBuffer
+		FGPUScenePrimitiveCollector* DynamicPrimitiveCollector
 #if UE_ENABLE_DEBUG_DRAWING
-		,FSimpleElementCollector* InDebugSimpleElementCollector = nullptr
+		, FSimpleElementCollector* DebugSimpleElementCollector = nullptr
 #endif
-	);
+		);
+
+	ENGINE_API void ClearViewMeshArrays();
+
+	ENGINE_API void Commit();
+
+	ENGINE_API void Finish();
 
 	/** 
 	 * Using TChunkedArray which will never realloc as new elements are added
@@ -2000,7 +2020,16 @@ protected:
 	TArray<uint16, TInlineAllocator<2, SceneRenderingAllocator>> MeshIdInPrimitivePerView;
 
 	/** Material proxies that will be deleted at the end of the frame. */
-	TArray<FMaterialRenderProxy*, SceneRenderingAllocator> TemporaryProxies;
+	TArray<FMaterialRenderProxy*, SceneRenderingAllocator> MaterialProxiesToDelete;
+
+	/** Material proxies to force uniform expression evaluation. */
+	TArray<TPair<FMaterialRenderProxy*, bool>, SceneRenderingAllocator> MaterialProxiesToInvalidate;
+
+	/** Material proxies to force uniform expression evaluation. */
+	TArray<const FMaterialRenderProxy*, SceneRenderingAllocator> MaterialProxiesToUpdate;
+
+	/** List of mesh batches that require GPU scene updates. */
+	TArray<TPair<FGPUScenePrimitiveCollector*, FMeshBatch*>, SceneRenderingAllocator> MeshBatchesForGPUScene;
 
 	/** Resources that will be deleted at the end of the frame. */
 	FSceneRenderingBulkObjectAllocator& OneFrameResources;
@@ -2009,11 +2038,14 @@ protected:
 	const FPrimitiveSceneProxy* PrimitiveSceneProxy;
 
 	/** Dynamic buffer pools. */
-	FGlobalDynamicIndexBuffer* DynamicIndexBuffer;
-	FGlobalDynamicVertexBuffer* DynamicVertexBuffer;
-	FGlobalDynamicReadBuffer* DynamicReadBuffer;
+	FGlobalDynamicIndexBuffer* DynamicIndexBuffer = nullptr;
+	FGlobalDynamicVertexBuffer* DynamicVertexBuffer = nullptr;
+	FGlobalDynamicReadBuffer* DynamicReadBuffer = nullptr;
 
-	ERHIFeatureLevel::Type FeatureLevel;
+	FRHICommandList* RHICmdList = nullptr;
+
+	const ERHIFeatureLevel::Type FeatureLevel;
+	const bool bUseGPUScene;
 
 	/** Tracks dynamic primitive data for upload to GPU Scene for every view, when enabled. */
 	TArray<FGPUScenePrimitiveCollector*, TInlineAllocator<2, SceneRenderingAllocator>> DynamicPrimitiveCollectorPerView;
@@ -2024,6 +2056,8 @@ protected:
 	friend class FProjectedShadowInfo;
 	friend class FCardPageRenderData;
 	friend class FViewFamilyInfo;
+	friend class FShadowMeshCollector;
+	friend struct FRayTracingMaterialGatheringContext;
 };
 
 #if RHI_RAYTRACING
@@ -2042,16 +2076,9 @@ public:
 
 	FRayTracingMeshResourceCollector(
 		ERHIFeatureLevel::Type InFeatureLevel,
-		FSceneRenderingBulkObjectAllocator& InBulkAllocator,
-		FGlobalDynamicIndexBuffer* InDynamicIndexBuffer,
-		FGlobalDynamicVertexBuffer* InDynamicVertexBuffer,
-		FGlobalDynamicReadBuffer* InDynamicReadBuffer)
+		FSceneRenderingBulkObjectAllocator& InBulkAllocator)
 		: FMeshElementCollector(InFeatureLevel, InBulkAllocator)
-	{
-		DynamicIndexBuffer = InDynamicIndexBuffer;
-		DynamicVertexBuffer = InDynamicVertexBuffer;
-		DynamicReadBuffer = InDynamicReadBuffer;
-	}
+	{}
 };
 
 struct FRayTracingDynamicGeometryUpdateParams
@@ -2084,17 +2111,22 @@ struct FRayTracingMaterialGatheringContext
 	const FSceneViewFamily& ReferenceViewFamily;
 
 	FRDGBuilder& GraphBuilder;
+	FRHICommandList& RHICmdList;
 	FRayTracingMeshResourceCollector& RayTracingMeshResourceCollector;
 	TArray<FRayTracingDynamicGeometryUpdateParams> DynamicRayTracingGeometriesToUpdate;
+	FGlobalDynamicVertexBuffer DynamicVertexBuffer;
+	FGlobalDynamicIndexBuffer DynamicIndexBuffer;
+
+	ENGINE_API FRayTracingMaterialGatheringContext(
+		const FScene* InScene,
+		const FSceneView* InReferenceView,
+		const FSceneViewFamily& InReferenceViewFamily,
+		FRDGBuilder& InGraphBuilder,
+		FRayTracingMeshResourceCollector& InRayTracingMeshResourceCollector,
+		FGlobalDynamicReadBuffer& InGlobalDynamicReadBuffer);
 
-	FRayTracingMaterialGatheringContext(const FScene* InScene, const FSceneView* InReferenceView, const FSceneViewFamily& InReferenceViewFamily, FRDGBuilder& InGraphBuilder, FRayTracingMeshResourceCollector& InRayTracingMeshResourceCollector)
-	:Scene(InScene),
-	ReferenceView(InReferenceView),
-	ReferenceViewFamily(InReferenceViewFamily),
-	GraphBuilder(InGraphBuilder),
-	RayTracingMeshResourceCollector(InRayTracingMeshResourceCollector){}
+	ENGINE_API virtual ~FRayTracingMaterialGatheringContext();
 
-	virtual ~FRayTracingMaterialGatheringContext() {}
 	virtual FRayTracingMaskAndFlags BuildInstanceMaskAndFlags(const FRayTracingInstance& Instance, const FPrimitiveSceneProxy& ScenePrimitive) = 0;
 };
 #endif
@@ -2110,6 +2142,7 @@ public:
 	TUniformBuffer<FPrimitiveUniformShaderParameters> UniformBuffer;
 
 	ENGINE_API void Set(
+		FRHICommandListBase& RHICmdList,
 		const FMatrix& LocalToWorld,
 		const FMatrix& PreviousLocalToWorld,
 		const FVector& ActorPositionWS, 
@@ -2122,6 +2155,7 @@ public:
 		const FCustomPrimitiveData* CustomPrimitiveData);
 
 	ENGINE_API void Set(
+		FRHICommandListBase& RHICmdList,
 		const FMatrix& LocalToWorld,
 		const FMatrix& PreviousLocalToWorld,
 		const FBoxSphereBounds& WorldBounds,
@@ -2133,6 +2167,7 @@ public:
 		const FCustomPrimitiveData* CustomPrimitiveData);
 
 	ENGINE_API void Set(
+		FRHICommandListBase& RHICmdList,
 		const FMatrix& LocalToWorld,
 		const FMatrix& PreviousLocalToWorld,
 		const FBoxSphereBounds& WorldBounds,
@@ -2144,6 +2179,7 @@ public:
 
 	/** Pass-through implementation which calls the overloaded Set function with LocalBounds for PreSkinnedLocalBounds. */
 	ENGINE_API void Set(
+		FRHICommandListBase& RHICmdList,
 		const FMatrix& LocalToWorld,
 		const FMatrix& PreviousLocalToWorld,
 		const FBoxSphereBounds& WorldBounds,
@@ -2152,29 +2188,52 @@ public:
 		bool bHasPrecomputedVolumetricLightmap,
 		bool bOutputVelocity);
 
-	UE_DEPRECATED(5.1, "Use version without bDrawsVelocity instead.")
-	void Set(
-		const FMatrix& LocalToWorld, const FMatrix& PreviousLocalToWorld, const FBoxSphereBounds& WorldBounds, const FBoxSphereBounds& LocalBounds, const FBoxSphereBounds& PreSkinnedLocalBounds,
-		bool bReceivesDecals, bool bHasPrecomputedVolumetricLightmap, bool bDrawsVelocity, bool bOutputVelocity, const FCustomPrimitiveData* CustomPrimitiveData)
-	{
-		Set(LocalToWorld, PreviousLocalToWorld, WorldBounds, LocalBounds, PreSkinnedLocalBounds, bReceivesDecals, bHasPrecomputedVolumetricLightmap, bOutputVelocity, CustomPrimitiveData);
-	}
+	//UE_DEPRECATED(5.4, "Set requires a command list")
+	ENGINE_API void Set(
+		const FMatrix& LocalToWorld,
+		const FMatrix& PreviousLocalToWorld,
+		const FVector& ActorPositionWS, 
+		const FBoxSphereBounds& WorldBounds,
+		const FBoxSphereBounds& LocalBounds,
+		const FBoxSphereBounds& PreSkinnedLocalBounds,
+		bool bReceivesDecals,
+		bool bHasPrecomputedVolumetricLightmap,
+		bool bOutputVelocity,
+		const FCustomPrimitiveData* CustomPrimitiveData);
 
-	UE_DEPRECATED(5.1, "Use version without bDrawsVelocity instead.")
-	void Set(
-		const FMatrix& LocalToWorld, const FMatrix& PreviousLocalToWorld, const FBoxSphereBounds& WorldBounds, const FBoxSphereBounds& LocalBounds, const FBoxSphereBounds& PreSkinnedLocalBounds,
-		bool bReceivesDecals, bool bHasPrecomputedVolumetricLightmap, bool bDrawsVelocity, bool bOutputVelocity)
-	{
-		Set(LocalToWorld, PreviousLocalToWorld, WorldBounds, LocalBounds, PreSkinnedLocalBounds, bReceivesDecals, bHasPrecomputedVolumetricLightmap, bOutputVelocity);
-	}
+	//UE_DEPRECATED(5.4, "Set requires a command list")
+	ENGINE_API void Set(
+		const FMatrix& LocalToWorld,
+		const FMatrix& PreviousLocalToWorld,
+		const FBoxSphereBounds& WorldBounds,
+		const FBoxSphereBounds& LocalBounds,
+		const FBoxSphereBounds& PreSkinnedLocalBounds,
+		bool bReceivesDecals,
+		bool bHasPrecomputedVolumetricLightmap,
+		bool bOutputVelocity,
+		const FCustomPrimitiveData* CustomPrimitiveData);
 
-	UE_DEPRECATED(5.1, "Use version without bDrawsVelocity instead.")
-	void Set(
-		const FMatrix& LocalToWorld, const FMatrix& PreviousLocalToWorld, const FBoxSphereBounds& WorldBounds, const FBoxSphereBounds& LocalBounds, 
-		bool bReceivesDecals, bool bHasPrecomputedVolumetricLightmap, bool bDrawsVelocity, bool bOutputVelocity)
-	{
-		Set(LocalToWorld, PreviousLocalToWorld, WorldBounds, LocalBounds, bReceivesDecals, bHasPrecomputedVolumetricLightmap, bOutputVelocity);
-	}
+	//UE_DEPRECATED(5.4, "Set requires a command list")
+	ENGINE_API void Set(
+		const FMatrix& LocalToWorld,
+		const FMatrix& PreviousLocalToWorld,
+		const FBoxSphereBounds& WorldBounds,
+		const FBoxSphereBounds& LocalBounds,
+		const FBoxSphereBounds& PreSkinnedLocalBounds,
+		bool bReceivesDecals,
+		bool bHasPrecomputedVolumetricLightmap,
+		bool bOutputVelocity);
+
+	/** Pass-through implementation which calls the overloaded Set function with LocalBounds for PreSkinnedLocalBounds. */
+	//UE_DEPRECATED(5.4, "Set requires a command list")
+	ENGINE_API void Set(
+		const FMatrix& LocalToWorld,
+		const FMatrix& PreviousLocalToWorld,
+		const FBoxSphereBounds& WorldBounds,
+		const FBoxSphereBounds& LocalBounds,
+		bool bReceivesDecals,
+		bool bHasPrecomputedVolumetricLightmap,
+		bool bOutputVelocity);
 };
 
 //
diff --git a/Engine/Source/Runtime/RenderCore/Private/DynamicBufferAllocator.cpp b/Engine/Source/Runtime/RenderCore/Private/DynamicBufferAllocator.cpp
--- a/Engine/Source/Runtime/RenderCore/Private/DynamicBufferAllocator.cpp
+++ b/Engine/Source/Runtime/RenderCore/Private/DynamicBufferAllocator.cpp
@@ -54,8 +54,6 @@ struct FDynamicReadBufferPool
 			Buffers[BufferIndex].Release();
 		}
 	}
-
-	FCriticalSection CriticalSection;
 };
 
 
@@ -100,7 +98,7 @@ void FGlobalDynamicReadBuffer::Cleanup()
 		HalfBufferPool = nullptr;
 	}
 }
-void FGlobalDynamicReadBuffer::InitRHI(FRHICommandListBase& RHICmdList)
+void FGlobalDynamicReadBuffer::InitRHI(FRHICommandListBase&)
 {
 }
 
@@ -110,11 +108,17 @@ void FGlobalDynamicReadBuffer::ReleaseRHI()
 }
 
 template<EPixelFormat Format, typename Type>
-FGlobalDynamicReadBuffer::FAllocation FGlobalDynamicReadBuffer_AllocateInternal(FRHICommandListBase& RHICmdList, FDynamicReadBufferPool* BufferPool, uint32 Num)
+FGlobalDynamicReadBuffer::FAllocation FGlobalDynamicReadBuffer::AllocateInternal(FDynamicReadBufferPool* BufferPool, uint32 Num)
 {
-	FScopeLock ScopeLock(&BufferPool->CriticalSection);
+	UE::TScopeLock ScopeLock(Mutex);
 	FGlobalDynamicReadBuffer::FAllocation Allocation;
 
+	if (!RHICmdList)
+	{
+		RHICmdList = new FRHICommandList(FRHIGPUMask::All());
+		RHICmdList->SwitchPipeline(ERHIPipeline::Graphics);
+	}
+
 	uint32 SizeInBytes = sizeof(Type) * Num;
 	FDynamicAllocReadBuffer* Buffer = BufferPool->CurrentBuffer;
 
@@ -143,13 +147,13 @@ FGlobalDynamicReadBuffer::FAllocation FGlobalDynamicReadBuffer_AllocateInternal(
 			const uint32 NewBufferSize = FMath::Max(AlignedNum, (uint32)GMinReadBufferRenderingBufferSize);
 			Buffer = new FDynamicAllocReadBuffer();
 			BufferPool->Buffers.Add(Buffer);
-			Buffer->Initialize(RHICmdList, TEXT("FGlobalDynamicReadBuffer_AllocateInternal"), sizeof(Type), NewBufferSize, Format, BUF_Dynamic);
+			Buffer->Initialize(*RHICmdList, TEXT("FGlobalDynamicReadBuffer_AllocateInternal"), sizeof(Type), NewBufferSize, Format, BUF_Dynamic);
 		}
 
 		// Lock the buffer if needed.
 		if (Buffer->MappedBuffer == nullptr)
 		{
-			Buffer->Lock(RHICmdList);
+			Buffer->Lock(*RHICmdList);
 		}
 
 		// Remember this buffer, we'll try to allocate out of it in the future.
@@ -161,7 +165,7 @@ FGlobalDynamicReadBuffer::FAllocation FGlobalDynamicReadBuffer_AllocateInternal(
 	checkf(Buffer->AllocatedByteCount + SizeInBytes <= Buffer->NumBytes, TEXT("Global dynamic read buffer allocation failed: BufferSize=%d AllocatedByteCount=%d SizeInBytes=%d"), Buffer->NumBytes, Buffer->AllocatedByteCount, SizeInBytes);
 	Allocation.Buffer = Buffer->MappedBuffer + Buffer->AllocatedByteCount;
 	Allocation.ReadBuffer = Buffer;
-	Buffer->SubAllocations.Emplace(RHICmdList.CreateShaderResourceView(FShaderResourceViewInitializer(Buffer->Buffer, Format, Buffer->AllocatedByteCount, Num)));
+	Buffer->SubAllocations.Emplace(RHICmdList->CreateShaderResourceView(FShaderResourceViewInitializer(Buffer->Buffer, Format, Buffer->AllocatedByteCount, Num)));
 	Allocation.SRV = Buffer->SubAllocations.Last();
 	Buffer->AllocatedByteCount += SizeInBytes;
 
@@ -177,28 +181,28 @@ void FGlobalDynamicReadBuffer::IncrementTotalAllocations(uint32 Num)
 	}
 }
 
-FGlobalDynamicReadBuffer::FAllocation FGlobalDynamicReadBuffer::AllocateFloat(FRHICommandListBase& RHICmdList, uint32 Num)
+FGlobalDynamicReadBuffer::FAllocation FGlobalDynamicReadBuffer::AllocateFloat(uint32 Num)
 {
 	IncrementTotalAllocations(Num);
-	return FGlobalDynamicReadBuffer_AllocateInternal<PF_R32_FLOAT, float>(RHICmdList, FloatBufferPool, Num);
+	return AllocateInternal<PF_R32_FLOAT, float>(FloatBufferPool, Num);
 }
 
-FGlobalDynamicReadBuffer::FAllocation FGlobalDynamicReadBuffer::AllocateHalf(FRHICommandListBase& RHICmdList, uint32 Num)
+FGlobalDynamicReadBuffer::FAllocation FGlobalDynamicReadBuffer::AllocateHalf(uint32 Num)
 {
 	IncrementTotalAllocations(Num);
-	return FGlobalDynamicReadBuffer_AllocateInternal<PF_R16F, FFloat16>(RHICmdList, HalfBufferPool, Num);
+	return AllocateInternal<PF_R16F, FFloat16>(HalfBufferPool, Num);
 }
 
-FGlobalDynamicReadBuffer::FAllocation FGlobalDynamicReadBuffer::AllocateInt32(FRHICommandListBase& RHICmdList, uint32 Num)
+FGlobalDynamicReadBuffer::FAllocation FGlobalDynamicReadBuffer::AllocateInt32(uint32 Num)
 {
 	IncrementTotalAllocations(Num);
-	return FGlobalDynamicReadBuffer_AllocateInternal<PF_R32_SINT, int32>(RHICmdList, Int32BufferPool, Num);
+	return AllocateInternal<PF_R32_SINT, int32>(Int32BufferPool, Num);
 }
 
-FGlobalDynamicReadBuffer::FAllocation FGlobalDynamicReadBuffer::AllocateUInt32(FRHICommandListBase& RHICmdList, uint32 Num)
+FGlobalDynamicReadBuffer::FAllocation FGlobalDynamicReadBuffer::AllocateUInt32(uint32 Num)
 {
 	IncrementTotalAllocations(Num);
-	return FGlobalDynamicReadBuffer_AllocateInternal<PF_R32_UINT, uint32>(RHICmdList, UInt32BufferPool, Num);
+	return AllocateInternal<PF_R32_UINT, uint32>(UInt32BufferPool, Num);
 }
 
 bool FGlobalDynamicReadBuffer::IsRenderAlarmLoggingEnabled() const
@@ -206,7 +210,7 @@ bool FGlobalDynamicReadBuffer::IsRenderAlarmLoggingEnabled() const
 	return GMaxReadBufferRenderingBytesAllocatedPerFrame > 0 && TotalAllocatedSinceLastCommit >= (size_t)GMaxReadBufferRenderingBytesAllocatedPerFrame;
 }
 
-static void RemoveUnusedBuffers(FRHICommandListBase& RHICmdList, FDynamicReadBufferPool* BufferPool)
+static void RemoveUnusedBuffers(FRHICommandListBase* RHICmdList, FDynamicReadBufferPool* BufferPool)
 {
 	extern int32 GGlobalBufferNumFramesUnusedThreshold;
 
@@ -215,7 +219,8 @@ static void RemoveUnusedBuffers(FRHICommandListBase& RHICmdList, FDynamicReadBuf
 		FDynamicAllocReadBuffer& Buffer = BufferPool->Buffers[BufferIndex];
 		if (Buffer.MappedBuffer != nullptr)
 		{
-			Buffer.Unlock(RHICmdList);
+			check(RHICmdList);
+			Buffer.Unlock(*RHICmdList);
 		}
 		else if (GGlobalBufferNumFramesUnusedThreshold && !Buffer.AllocatedByteCount)
 		{
@@ -232,8 +237,10 @@ static void RemoveUnusedBuffers(FRHICommandListBase& RHICmdList, FDynamicReadBuf
 	}
 }
 
-void FGlobalDynamicReadBuffer::Commit(FRHICommandListBase& RHICmdList)
+void FGlobalDynamicReadBuffer::Commit(FRHICommandListImmediate& RHICmdListImmediate)
 {
+	UE::TScopeLock ScopeLock(Mutex);
+
 	RemoveUnusedBuffers(RHICmdList, FloatBufferPool);
 	FloatBufferPool->CurrentBuffer = nullptr;
 
@@ -246,5 +253,12 @@ void FGlobalDynamicReadBuffer::Commit(FRHICommandListBase& RHICmdList)
 	RemoveUnusedBuffers(RHICmdList, HalfBufferPool);
 	HalfBufferPool->CurrentBuffer = nullptr;
 
+	if (RHICmdList)
+	{
+		RHICmdList->FinishRecording();
+		RHICmdListImmediate.QueueAsyncCommandListSubmit(RHICmdList);
+		RHICmdList = nullptr;
+	}
+
 	TotalAllocatedSinceLastCommit = 0;
 }
diff --git a/Engine/Source/Runtime/RenderCore/Private/GlobalRenderResources.cpp b/Engine/Source/Runtime/RenderCore/Private/GlobalRenderResources.cpp
--- a/Engine/Source/Runtime/RenderCore/Private/GlobalRenderResources.cpp
+++ b/Engine/Source/Runtime/RenderCore/Private/GlobalRenderResources.cpp
@@ -854,16 +854,13 @@ static TGlobalResource<TDynamicBufferPool<FDynamicVertexBuffer>, FRenderResource
 
 FGlobalDynamicVertexBuffer::FAllocation FGlobalDynamicVertexBuffer::Allocate(uint32 SizeInBytes)
 {
-	return Allocate(FRHICommandListImmediate::Get(), SizeInBytes);
-}
+	checkf(RHICmdList, TEXT("FGlobalDynamicVertexBuffer was not initialized prior to calling Allocate."));
 
-FGlobalDynamicVertexBuffer::FAllocation FGlobalDynamicVertexBuffer::Allocate(FRHICommandListBase& RHICmdList, uint32 SizeInBytes)
-{
 	FAllocation Allocation;
 
 	if (VertexBuffers.IsEmpty() || VertexBuffers.Last()->AllocatedByteCount + SizeInBytes > VertexBuffers.Last()->BufferSize)
 	{
-		VertexBuffers.Emplace(GDynamicVertexBufferPool.Acquire(RHICmdList, SizeInBytes, 0));
+		VertexBuffers.Emplace(GDynamicVertexBufferPool.Acquire(*RHICmdList, SizeInBytes, 0));
 	}
 
 	FDynamicVertexBuffer* VertexBuffer = VertexBuffers.Last();
@@ -881,26 +878,21 @@ bool FGlobalDynamicVertexBuffer::IsRenderAlarmLoggingEnabled() const
 	return GDynamicVertexBufferPool.IsRenderAlarmLoggingEnabled();
 }
 
-void FGlobalDynamicVertexBuffer::Commit(FRHICommandListBase& RHICmdList)
-{
-	GDynamicVertexBufferPool.Forfeit(RHICmdList, VertexBuffers);
-	VertexBuffers.Reset();
-}
-
 void FGlobalDynamicVertexBuffer::Commit()
 {
-	Commit(FRHICommandListImmediate::Get());
-}
-
-FGlobalDynamicIndexBuffer::FAllocation FGlobalDynamicIndexBuffer::Allocate(uint32 NumIndices, uint32 IndexStride)
-{
-	return Allocate(FRHICommandListImmediate::Get(), NumIndices, IndexStride);
+	if (RHICmdList)
+	{
+		GDynamicVertexBufferPool.Forfeit(*RHICmdList, VertexBuffers);
+		VertexBuffers.Reset();
+	}
 }
 
 static TGlobalResource<TDynamicBufferPool<FDynamicIndexBuffer>, FRenderResource::EInitPhase::Pre> GDynamicIndexBufferPool;
 
-FGlobalDynamicIndexBuffer::FAllocation FGlobalDynamicIndexBuffer::Allocate(FRHICommandListBase& RHICmdList, uint32 NumIndices, uint32 IndexStride)
+FGlobalDynamicIndexBuffer::FAllocation FGlobalDynamicIndexBuffer::Allocate(uint32 NumIndices, uint32 IndexStride)
 {
+	checkf(RHICmdList, TEXT("FGlobalDynamicIndexBuffer was not initialized prior to calling Allocate."));
+
 	FAllocation Allocation;
 
 	if (IndexStride != 2 && IndexStride != 4)
@@ -916,7 +908,7 @@ FGlobalDynamicIndexBuffer::FAllocation FGlobalDynamicIndexBuffer::Allocate(FRHIC
 
 	if (IndexBuffers.IsEmpty() || IndexBuffers.Last()->AllocatedByteCount + SizeInBytes > IndexBuffers.Last()->BufferSize)
 	{
-		IndexBuffers.Emplace(GDynamicIndexBufferPool.Acquire(RHICmdList, SizeInBytes, IndexStride));
+		IndexBuffers.Emplace(GDynamicIndexBufferPool.Acquire(*RHICmdList, SizeInBytes, IndexStride));
 	}
 
 	FDynamicIndexBuffer* IndexBuffer = IndexBuffers.Last();
@@ -929,17 +921,15 @@ FGlobalDynamicIndexBuffer::FAllocation FGlobalDynamicIndexBuffer::Allocate(FRHIC
 	return Allocation;
 }
 
-void FGlobalDynamicIndexBuffer::Commit(FRHICommandListBase& RHICmdList)
-{
-	GDynamicIndexBufferPool.Forfeit(RHICmdList, IndexBuffers16);
-	GDynamicIndexBufferPool.Forfeit(RHICmdList, IndexBuffers32);
-	IndexBuffers16.Reset();
-	IndexBuffers32.Reset();
-}
-
 void FGlobalDynamicIndexBuffer::Commit()
 {
-	Commit(FRHICommandListImmediate::Get());
+	if (RHICmdList)
+	{
+		GDynamicIndexBufferPool.Forfeit(*RHICmdList, IndexBuffers16);
+		GDynamicIndexBufferPool.Forfeit(*RHICmdList, IndexBuffers32);
+		IndexBuffers16.Reset();
+		IndexBuffers32.Reset();
+	}
 }
 
 namespace GlobalDynamicBuffer
diff --git a/Engine/Source/Runtime/RenderCore/Public/DynamicBufferAllocator.h b/Engine/Source/Runtime/RenderCore/Public/DynamicBufferAllocator.h
--- a/Engine/Source/Runtime/RenderCore/Public/DynamicBufferAllocator.h
+++ b/Engine/Source/Runtime/RenderCore/Public/DynamicBufferAllocator.h
@@ -11,6 +11,7 @@ DynamicBufferAllocator.h: Classes for allocating transient rendering data.
 #include "RHI.h"
 #include "RHIUtilities.h"
 #include "RenderResource.h"
+#include "Async/Mutex.h"
 
 struct FDynamicReadBufferPool;
 
@@ -83,26 +84,17 @@ public:
 	RENDERCORE_API FGlobalDynamicReadBuffer();
 	RENDERCORE_API ~FGlobalDynamicReadBuffer();
 	
-	RENDERCORE_API FAllocation AllocateHalf(FRHICommandListBase& RHICmdList, uint32 Num);
-	RENDERCORE_API FAllocation AllocateFloat(FRHICommandListBase& RHICmdList, uint32 Num);
-	RENDERCORE_API FAllocation AllocateInt32(FRHICommandListBase& RHICmdList, uint32 Num);
-	RENDERCORE_API FAllocation AllocateUInt32(FRHICommandListBase& RHICmdList, uint32 Num);
-
-	UE_DEPRECATED(5.3, "AllocateHalf now requires a command list.")
-	FAllocation AllocateHalf(uint32 Num) { return AllocateHalf(FRHICommandListImmediate::Get(), Num); }
-	UE_DEPRECATED(5.3, "AllocateFloat now requires a command list.")
-	FAllocation AllocateFloat(uint32 Num) { return AllocateFloat(FRHICommandListImmediate::Get(), Num); }
-	UE_DEPRECATED(5.3, "AllocateInt32 now requires a command list.")
-	FAllocation AllocateInt32(uint32 Num) { return AllocateInt32(FRHICommandListImmediate::Get(), Num); }
-	UE_DEPRECATED(5.3, "AllocateUInt32 now requires a command list.")
-	FAllocation AllocateUInt32(uint32 Num) { return AllocateUInt32(FRHICommandListImmediate::Get(), Num); }
+	RENDERCORE_API FAllocation AllocateHalf(uint32 Num);
+	RENDERCORE_API FAllocation AllocateFloat(uint32 Num);
+	RENDERCORE_API FAllocation AllocateInt32(uint32 Num);
+	RENDERCORE_API FAllocation AllocateUInt32(uint32 Num);
 
 	/**
 	* Commits allocated memory to the GPU.
 	*		WARNING: Once this buffer has been committed to the GPU, allocations
 	*		remain valid only until the next call to Allocate!
 	*/
-	RENDERCORE_API void Commit(FRHICommandListBase& RHICmdList);
+	RENDERCORE_API void Commit(FRHICommandListImmediate& RHICmdList);
 
 	UE_DEPRECATED(5.3, "Commit now requires a command list.")
 	void Commit() { Commit(FRHICommandListImmediate::Get()); }
@@ -116,6 +108,12 @@ protected:
 	RENDERCORE_API void Cleanup();
 	RENDERCORE_API void IncrementTotalAllocations(uint32 Num);
 
+	template<EPixelFormat Format, typename Type>
+	FAllocation AllocateInternal(FDynamicReadBufferPool* BufferPool, uint32 Num);
+
+	UE::FMutex Mutex;
+	FRHICommandListBase* RHICmdList = nullptr;
+
 	/** The pools of read buffers from which allocations are made. */
 	FDynamicReadBufferPool* HalfBufferPool;
 	FDynamicReadBufferPool* FloatBufferPool;
diff --git a/Engine/Source/Runtime/RenderCore/Public/GlobalRenderResources.h b/Engine/Source/Runtime/RenderCore/Public/GlobalRenderResources.h
--- a/Engine/Source/Runtime/RenderCore/Public/GlobalRenderResources.h
+++ b/Engine/Source/Runtime/RenderCore/Public/GlobalRenderResources.h
@@ -188,12 +188,28 @@ class FGlobalDynamicVertexBuffer
 public:
 	using FAllocation = FGlobalDynamicVertexBufferAllocation;
 
+	FGlobalDynamicVertexBuffer() = default;
+
+	FGlobalDynamicVertexBuffer(FRHICommandListBase& InRHICmdList)
+		: RHICmdList(&InRHICmdList)
+	{}
+
+	~FGlobalDynamicVertexBuffer()
+	{
+		Commit();
+	}
+
+	void Init(FRHICommandListBase& InRHICmdList)
+	{
+		check(VertexBuffers.IsEmpty());
+		RHICmdList = &InRHICmdList;
+	}
+
 	/**
 	 * Allocates space in the global vertex buffer.
 	 * @param SizeInBytes - The amount of memory to allocate in bytes.
 	 * @returns An FAllocation with information regarding the allocated memory.
 	 */
-	RENDERCORE_API FAllocation Allocate(FRHICommandListBase& RHICmdList, uint32 SizeInBytes);
 	RENDERCORE_API FAllocation Allocate(uint32 SizeInBytes);
 
 	/**
@@ -201,9 +217,6 @@ public:
 	 *		WARNING: Once this buffer has been committed to the GPU, allocations
 	 *		remain valid only until the next call to Allocate!
 	 */
-	RENDERCORE_API void Commit(FRHICommandListBase& RHICmdList);
-
-	UE_DEPRECATED(5.4, "Commit requires a command list.")
 	RENDERCORE_API void Commit();
 
 	UE_DEPRECATED(5.4, "Use GlobalDynamicBuffer::GarbageCollect instead.")
@@ -213,6 +226,7 @@ public:
 	RENDERCORE_API bool IsRenderAlarmLoggingEnabled() const;
 
 private:
+	FRHICommandListBase* RHICmdList = nullptr;
 	TArray<FDynamicVertexBuffer*> VertexBuffers;
 };
 
@@ -259,13 +273,29 @@ public:
 	using FAllocation = FGlobalDynamicIndexBufferAllocation;
 	using FAllocationEx = FGlobalDynamicIndexBufferAllocationEx;
 
+	FGlobalDynamicIndexBuffer() = default;
+
+	FGlobalDynamicIndexBuffer(FRHICommandListBase& InRHICmdList)
+		: RHICmdList(&InRHICmdList)
+	{}
+
+	~FGlobalDynamicIndexBuffer()
+	{
+		Commit();
+	}
+
+	void Init(FRHICommandListBase& InRHICmdList)
+	{
+		check(IndexBuffers16.IsEmpty() && IndexBuffers32.IsEmpty());
+		RHICmdList = &InRHICmdList;
+	}
+
 	/**
 	 * Allocates space in the global index buffer.
 	 * @param NumIndices - The number of indices to allocate.
 	 * @param IndexStride - The size of an index (2 or 4 bytes).
 	 * @returns An FAllocation with information regarding the allocated memory.
 	 */
-	RENDERCORE_API FAllocation Allocate(FRHICommandListBase& RHICmdList, uint32 NumIndices, uint32 IndexStride);
 	RENDERCORE_API FAllocation Allocate(uint32 NumIndices, uint32 IndexStride);
 
 	/**
@@ -279,23 +309,15 @@ public:
 		return FAllocationEx(Allocate(NumIndices, sizeof(IndexType)), NumIndices, sizeof(IndexType));
 	}
 
-	template <typename IndexType>
-	FORCEINLINE FAllocationEx Allocate(FRHICommandListBase& RHICmdList, uint32 NumIndices)
-	{
-		return FAllocationEx(Allocate(RHICmdList, NumIndices, sizeof(IndexType)), NumIndices, sizeof(IndexType));
-	}
-
 	/**
 	 * Commits allocated memory to the GPU.
 	 *		WARNING: Once this buffer has been committed to the GPU, allocations
 	 *		remain valid only until the next call to Allocate!
 	 */
-	RENDERCORE_API void Commit(FRHICommandListBase& RHICmdList);
-
-	UE_DEPRECATED(5.4, "Commit requires a command list.")
 	RENDERCORE_API void Commit();
 
 private:
+	FRHICommandListBase* RHICmdList = nullptr;
 	TArray<FDynamicIndexBuffer*> IndexBuffers16;
 	TArray<FDynamicIndexBuffer*> IndexBuffers32;
 };
diff --git a/Engine/Source/Runtime/Renderer/Private/DeferredShadingRenderer.cpp b/Engine/Source/Runtime/Renderer/Private/DeferredShadingRenderer.cpp
--- a/Engine/Source/Runtime/Renderer/Private/DeferredShadingRenderer.cpp
+++ b/Engine/Source/Runtime/Renderer/Private/DeferredShadingRenderer.cpp
@@ -540,12 +540,8 @@ DECLARE_CYCLE_STAT(TEXT("RenderFinish"), STAT_CLM_RenderFinish, STATGROUP_Comman
 DECLARE_CYCLE_STAT(TEXT("AfterFrame"), STAT_CLM_AfterFrame, STATGROUP_CommandListMarkers);
 DECLARE_CYCLE_STAT(TEXT("Wait RayTracing Add Mesh Batch"), STAT_WaitRayTracingAddMesh, STATGROUP_SceneRendering);
 
-FGlobalDynamicIndexBuffer FDeferredShadingSceneRenderer::DynamicIndexBufferForInitViews;
 FGlobalDynamicIndexBuffer FDeferredShadingSceneRenderer::DynamicIndexBufferForInitShadows;
-FGlobalDynamicVertexBuffer FDeferredShadingSceneRenderer::DynamicVertexBufferForInitViews;
 FGlobalDynamicVertexBuffer FDeferredShadingSceneRenderer::DynamicVertexBufferForInitShadows;
-TGlobalResource<FGlobalDynamicReadBuffer> FDeferredShadingSceneRenderer::DynamicReadBufferForInitShadows;
-TGlobalResource<FGlobalDynamicReadBuffer> FDeferredShadingSceneRenderer::DynamicReadBufferForInitViews;
 
 /**
  * Returns true if the depth Prepass needs to run
@@ -1016,8 +1012,9 @@ struct FDeferredShadingRayTracingMaterialGatheringContext : public FRayTracingMa
 		const FSceneView* InReferenceView,
 		const FSceneViewFamily& InReferenceViewFamily,
 		FRDGBuilder& InGraphBuilder,
-		FRayTracingMeshResourceCollector& InRayTracingMeshResourceCollector)
-		:FRayTracingMaterialGatheringContext(InScene, InReferenceView, InReferenceViewFamily, InGraphBuilder, InRayTracingMeshResourceCollector){}
+		FRayTracingMeshResourceCollector& InRayTracingMeshResourceCollector,
+		FGlobalDynamicReadBuffer& InDynamicReadBuffer)
+		:FRayTracingMaterialGatheringContext(InScene, InReferenceView, InReferenceViewFamily, InGraphBuilder, InRayTracingMeshResourceCollector, InDynamicReadBuffer){}
 
 	virtual FRayTracingMaskAndFlags BuildInstanceMaskAndFlags(const FRayTracingInstance& Instance, const FPrimitiveSceneProxy& ScenePrimitive) override
 	{
@@ -1061,21 +1058,8 @@ bool FDeferredShadingSceneRenderer::GatherRayTracingWorldInstancesForView(FRDGBu
 	// is done on task threads, therefore all work must be done here up-front as UpdateUniformExpressionCacheIfNeeded is not free-threaded.
 	check(!FMaterialRenderProxy::HasDeferredUniformExpressionCacheRequests());
 
-	RayTracingCollector.ClearViewMeshArrays();
-
 	FGPUScenePrimitiveCollector DummyDynamicPrimitiveCollector;
 
-	RayTracingCollector.AddViewMeshArrays(
-		&View,
-		&View.RayTracedDynamicMeshElements,
-		&View.SimpleElementCollector,
-		&DummyDynamicPrimitiveCollector,
-		ViewFamily.GetFeatureLevel(),
-		&DynamicIndexBufferForInitViews,
-		&DynamicVertexBufferForInitViews,
-		&DynamicReadBufferForInitViews
-	);
-
 	View.DynamicRayTracingMeshCommandStorage.Reserve(Scene->Primitives.Num());
 	View.VisibleRayTracingMeshCommands.Reserve(Scene->Primitives.Num());
 
@@ -1086,12 +1070,7 @@ bool FDeferredShadingSceneRenderer::GatherRayTracingWorldInstancesForView(FRDGBu
 		Extension->BeginRenderView(&View);
 	}
 
-	View.RayTracingMeshResourceCollector = MakeUnique<FRayTracingMeshResourceCollector>(
-		Scene->GetFeatureLevel(),
-		Allocator,
-		&DynamicIndexBufferForInitViews,
-		&DynamicVertexBufferForInitViews,
-		&DynamicReadBufferForInitViews);
+	View.RayTracingMeshResourceCollector = MakeUnique<FRayTracingMeshResourceCollector>(Scene->GetFeatureLevel(), Allocator);
 
 	View.RayTracingCullingParameters.Init(View);
 
@@ -1101,7 +1080,8 @@ bool FDeferredShadingSceneRenderer::GatherRayTracingWorldInstancesForView(FRDGBu
 		&View,
 		ViewFamily,
 		GraphBuilder,
-		*View.RayTracingMeshResourceCollector
+		*View.RayTracingMeshResourceCollector,
+		DynamicReadBufferForInitViews
 	);
 
 	const float CurrentWorldTime = View.Family->Time.GetWorldTimeSeconds();
@@ -2384,17 +2364,15 @@ void FDeferredShadingSceneRenderer::FinishInitDynamicShadows(FRDGBuilder& GraphB
 		// Setup dynamic shadows.
 		if (TaskData)
 		{
-			FSceneRenderer::FinishInitDynamicShadows(GraphBuilder, TaskData, DynamicIndexBufferForInitShadows, DynamicVertexBufferForInitShadows, DynamicReadBufferForInitShadows, InstanceCullingManager, ExternalAccessQueue);
+			FSceneRenderer::FinishInitDynamicShadows(GraphBuilder, TaskData, InstanceCullingManager, ExternalAccessQueue);
 		}
 		else
 		{
-			TaskData = InitDynamicShadows(GraphBuilder, DynamicIndexBufferForInitShadows, DynamicVertexBufferForInitShadows, DynamicReadBufferForInitShadows, InstanceCullingManager, ExternalAccessQueue);
+			TaskData = InitDynamicShadows(GraphBuilder, InstanceCullingManager, ExternalAccessQueue);
 		}
 
 		SCOPE_CYCLE_COUNTER(STAT_FDeferredShadingSceneRenderer_FGlobalDynamicVertexBuffer_Commit);
-		DynamicVertexBufferForInitShadows.Commit(GraphBuilder.RHICmdList);
-		DynamicIndexBufferForInitShadows.Commit(GraphBuilder.RHICmdList);
-		DynamicReadBufferForInitShadows.Commit(GraphBuilder.RHICmdList);
+		DynamicReadBufferForShadows.Commit(GraphBuilder.RHICmdList);
 	}
 }
 
@@ -2715,7 +2693,7 @@ void FDeferredShadingSceneRenderer::Render(FRDGBuilder& GraphBuilder)
 
 	GPU_MESSAGE_SCOPE(GraphBuilder);
 
-	FInitViewTaskDatas InitViewTaskDatas = OnRenderBegin(GraphBuilder, FGlobalDynamicBuffers(DynamicIndexBufferForInitViews, DynamicVertexBufferForInitViews, DynamicReadBufferForInitViews));
+	FInitViewTaskDatas InitViewTaskDatas = OnRenderBegin(GraphBuilder);
 
 	FRDGExternalAccessQueue ExternalAccessQueue;
 	TUniquePtr<FVirtualTextureUpdater> VirtualTextureUpdater;
@@ -3237,18 +3215,11 @@ void FDeferredShadingSceneRenderer::Render(FRDGBuilder& GraphBuilder)
 	}
 
 	{
-		{
-			RDG_RHI_GPU_STAT_SCOPE(GraphBuilder, VisibilityCommands);
-			EndInitViews(GraphBuilder, LumenFrameTemporaries, InstanceCullingManager, ExternalAccessQueue, InitViewTaskDatas);
-		}
+		RDG_RHI_GPU_STAT_SCOPE(GraphBuilder, VisibilityCommands);
+		EndInitViews(GraphBuilder, LumenFrameTemporaries, InstanceCullingManager, ExternalAccessQueue, InitViewTaskDatas);
 
-		// Dynamic vertex and index buffers need to be committed before rendering.
-		{
-			SCOPE_CYCLE_COUNTER(STAT_FDeferredShadingSceneRenderer_FGlobalDynamicVertexBuffer_Commit);
-			DynamicIndexBufferForInitViews.Commit(GraphBuilder.RHICmdList);
-			DynamicVertexBufferForInitViews.Commit(GraphBuilder.RHICmdList);
-			DynamicReadBufferForInitViews.Commit(GraphBuilder.RHICmdList);
-		}
+		SCOPE_CYCLE_COUNTER(STAT_FDeferredShadingSceneRenderer_FGlobalDynamicVertexBuffer_Commit);
+		DynamicReadBufferForInitViews.Commit(GraphBuilder.RHICmdList);
 	}
 
 	UE::SVT::GetStreamingManager().EndAsyncUpdate(GraphBuilder);
diff --git a/Engine/Source/Runtime/Renderer/Private/DeferredShadingRenderer.h b/Engine/Source/Runtime/Renderer/Private/DeferredShadingRenderer.h
--- a/Engine/Source/Runtime/Renderer/Private/DeferredShadingRenderer.h
+++ b/Engine/Source/Runtime/Renderer/Private/DeferredShadingRenderer.h
@@ -521,7 +521,6 @@ private:
 	static FGlobalDynamicIndexBuffer DynamicIndexBufferForInitShadows;
 	static FGlobalDynamicVertexBuffer DynamicVertexBufferForInitViews;
 	static FGlobalDynamicVertexBuffer DynamicVertexBufferForInitShadows;
-	static TGlobalResource<FGlobalDynamicReadBuffer> DynamicReadBufferForInitViews;
 	static TGlobalResource<FGlobalDynamicReadBuffer> DynamicReadBufferForInitShadows;
 
 	FSeparateTranslucencyDimensions SeparateTranslucencyDimensions;
diff --git a/Engine/Source/Runtime/Renderer/Private/GPUScene.cpp b/Engine/Source/Runtime/Renderer/Private/GPUScene.cpp
--- a/Engine/Source/Runtime/Renderer/Private/GPUScene.cpp
+++ b/Engine/Source/Runtime/Renderer/Private/GPUScene.cpp
@@ -2034,6 +2034,7 @@ FGPUScenePrimitiveCollector::FUploadData* FGPUSceneDynamicContext::AllocateDynam
 	LLM_SCOPE_BYTAG(GPUScene);
 
 	FGPUScenePrimitiveCollector::FUploadData* UploadData = new FGPUScenePrimitiveCollector::FUploadData;
+	UE::TScopeLock ScopeLock(DymamicPrimitiveUploadDataMutex);
 	DymamicPrimitiveUploadData.Add(UploadData);
 	return UploadData;
 }
diff --git a/Engine/Source/Runtime/Renderer/Private/GPUScene.h b/Engine/Source/Runtime/Renderer/Private/GPUScene.h
--- a/Engine/Source/Runtime/Renderer/Private/GPUScene.h
+++ b/Engine/Source/Runtime/Renderer/Private/GPUScene.h
@@ -161,6 +161,7 @@ private:
 	friend class FGPUScenePrimitiveCollector;
 
 	FGPUScenePrimitiveCollector::FUploadData* AllocateDynamicPrimitiveData();
+	UE::FMutex DymamicPrimitiveUploadDataMutex;
 	TArray<FGPUScenePrimitiveCollector::FUploadData*, TInlineAllocator<128, SceneRenderingAllocator> > DymamicPrimitiveUploadData;
 	FGPUScene& GPUScene;
 };
diff --git a/Engine/Source/Runtime/Renderer/Private/MobileShadingRenderer.cpp b/Engine/Source/Runtime/Renderer/Private/MobileShadingRenderer.cpp
--- a/Engine/Source/Runtime/Renderer/Private/MobileShadingRenderer.cpp
+++ b/Engine/Source/Runtime/Renderer/Private/MobileShadingRenderer.cpp
@@ -111,10 +111,6 @@ DECLARE_CYCLE_STAT(TEXT("PrePass"), STAT_CLM_MobilePrePass, STATGROUP_CommandLis
 DECLARE_CYCLE_STAT(TEXT("Velocity"), STAT_CLMM_Velocity, STATGROUP_CommandListMarkers);
 DECLARE_CYCLE_STAT(TEXT("TranslucentVelocity"), STAT_CLMM_TranslucentVelocity, STATGROUP_CommandListMarkers);
 
-FGlobalDynamicIndexBuffer FMobileSceneRenderer::DynamicIndexBuffer;
-FGlobalDynamicVertexBuffer FMobileSceneRenderer::DynamicVertexBuffer;
-TGlobalResource<FGlobalDynamicReadBuffer> FMobileSceneRenderer::DynamicReadBuffer;
-
 extern bool IsMobileEyeAdaptationEnabled(const FViewInfo& View);
 
 struct FMobileCustomDepthStencilUsage
@@ -817,7 +813,7 @@ void FMobileSceneRenderer::Render(FRDGBuilder& GraphBuilder)
 	RDG_RHI_EVENT_SCOPE(GraphBuilder, MobileSceneRender);
 	RDG_RHI_GPU_STAT_SCOPE(GraphBuilder, MobileSceneRender);
 
-	IVisibilityTaskData* VisibilityTaskData = OnRenderBegin(GraphBuilder, FGlobalDynamicBuffers(DynamicIndexBuffer, DynamicVertexBuffer, DynamicReadBuffer));
+	IVisibilityTaskData* VisibilityTaskData = OnRenderBegin(GraphBuilder);
 
 	FRDGExternalAccessQueue ExternalAccessQueue;
 
@@ -930,9 +926,8 @@ void FMobileSceneRenderer::Render(FRDGBuilder& GraphBuilder)
 	GEngine->GetPreRenderDelegateEx().Broadcast(GraphBuilder);
 
 	// Global dynamic buffers need to be committed before rendering.
-	DynamicIndexBuffer.Commit(GraphBuilder.RHICmdList);
-	DynamicVertexBuffer.Commit(GraphBuilder.RHICmdList);
-	DynamicReadBuffer.Commit(GraphBuilder.RHICmdList);
+	DynamicReadBufferForInitViews.Commit(GraphBuilder.RHICmdList);
+	DynamicReadBufferForShadows.Commit(GraphBuilder.RHICmdList);
 	
 	GraphBuilder.SetCommandListStat(GET_STATID(STAT_CLMM_SceneSim));
 
diff --git a/Engine/Source/Runtime/Renderer/Private/SceneHitProxyRendering.cpp b/Engine/Source/Runtime/Renderer/Private/SceneHitProxyRendering.cpp
--- a/Engine/Source/Runtime/Renderer/Private/SceneHitProxyRendering.cpp
+++ b/Engine/Source/Runtime/Renderer/Private/SceneHitProxyRendering.cpp
@@ -557,7 +557,7 @@ static void DoRenderHitProxies(
 
 void FMobileSceneRenderer::RenderHitProxies(FRDGBuilder& GraphBuilder)
 {
-	IVisibilityTaskData* VisibilityTaskData = OnRenderBegin(GraphBuilder, FGlobalDynamicBuffers(DynamicIndexBuffer, DynamicVertexBuffer, DynamicReadBuffer));
+	IVisibilityTaskData* VisibilityTaskData = OnRenderBegin(GraphBuilder);
 
 	GPU_MESSAGE_SCOPE(GraphBuilder);
 
@@ -579,9 +579,8 @@ void FMobileSceneRenderer::RenderHitProxies(FRDGBuilder& GraphBuilder)
 	GEngine->GetPreRenderDelegateEx().Broadcast(GraphBuilder);
 
 	// Global dynamic buffers need to be committed before rendering.
-	DynamicIndexBuffer.Commit(GraphBuilder.RHICmdList);
-	DynamicVertexBuffer.Commit(GraphBuilder.RHICmdList);
-	DynamicReadBuffer.Commit(GraphBuilder.RHICmdList);
+	DynamicReadBufferForInitViews.Commit(GraphBuilder.RHICmdList);
+	DynamicReadBufferForShadows.Commit(GraphBuilder.RHICmdList);
 
 	InstanceCullingManager.FlushRegisteredViews(GraphBuilder);
 
@@ -600,7 +599,7 @@ void FDeferredShadingSceneRenderer::RenderHitProxies(FRDGBuilder& GraphBuilder)
 
 	CommitFinalPipelineState();
 
-	IVisibilityTaskData* VisibilityTaskData = OnRenderBegin(GraphBuilder, FGlobalDynamicBuffers(DynamicIndexBufferForInitViews, DynamicVertexBufferForInitViews, DynamicReadBufferForInitViews));
+	IVisibilityTaskData* VisibilityTaskData = OnRenderBegin(GraphBuilder);
 
 	GPU_MESSAGE_SCOPE(GraphBuilder);
 
@@ -667,9 +666,8 @@ void FDeferredShadingSceneRenderer::RenderHitProxies(FRDGBuilder& GraphBuilder)
 	}
 
 	// Global dynamic buffers need to be committed before rendering.
-	DynamicIndexBufferForInitViews.Commit(GraphBuilder.RHICmdList);
-	DynamicVertexBufferForInitViews.Commit(GraphBuilder.RHICmdList);
 	DynamicReadBufferForInitViews.Commit(GraphBuilder.RHICmdList);
+	DynamicReadBufferForShadows.Commit(GraphBuilder.RHICmdList);
 
 	// Notify the FX system that the scene is about to be rendered.
 	if (FXSystem && Views.IsValidIndex(0))
diff --git a/Engine/Source/Runtime/Renderer/Private/SceneOcclusion.cpp b/Engine/Source/Runtime/Renderer/Private/SceneOcclusion.cpp
--- a/Engine/Source/Runtime/Renderer/Private/SceneOcclusion.cpp
+++ b/Engine/Source/Runtime/Renderer/Private/SceneOcclusion.cpp
@@ -485,7 +485,7 @@ void FOcclusionQueryBatcher::Flush(FRHICommandList& RHICmdList)
 	}
 }
 
-FRHIRenderQuery* FOcclusionQueryBatcher::BatchPrimitive(FRHICommandList& RHICmdList, const FVector& BoundsOrigin,const FVector& BoundsBoxExtent, FGlobalDynamicVertexBuffer& DynamicVertexBuffer)
+FRHIRenderQuery* FOcclusionQueryBatcher::BatchPrimitive(const FVector& BoundsOrigin,const FVector& BoundsBoxExtent, FGlobalDynamicVertexBuffer& DynamicVertexBuffer)
 {
 	// Check if the current batch is full.
 	if(CurrentBatchOcclusionQuery == NULL || NumBatchedPrimitives >= MaxBatchedPrimitives)
@@ -493,7 +493,7 @@ FRHIRenderQuery* FOcclusionQueryBatcher::BatchPrimitive(FRHICommandList& RHICmdL
 		check(OcclusionQueryPool);
 		CurrentBatchOcclusionQuery = new(BatchOcclusionQueries) FOcclusionBatch;
 		CurrentBatchOcclusionQuery->Query = OcclusionQueryPool->AllocateQuery();
-		CurrentBatchOcclusionQuery->VertexAllocation = DynamicVertexBuffer.Allocate(RHICmdList, MaxBatchedPrimitives * 8 * sizeof(FVector3f));
+		CurrentBatchOcclusionQuery->VertexAllocation = DynamicVertexBuffer.Allocate(MaxBatchedPrimitives * 8 * sizeof(FVector3f));
 		check(CurrentBatchOcclusionQuery->VertexAllocation.IsValid());
 		NumBatchedPrimitives = 0;
 	}
@@ -1788,7 +1788,7 @@ void FOcclusionFeedback::ReadbackResults(FRHICommandList& RHICmdList)
 	}
 }
 
-void FOcclusionFeedback::AddPrimitive(FRHICommandList& RHICmdList, const FPrimitiveOcclusionHistoryKey& PrimitiveKey, const FVector& BoundsOrigin, const FVector& BoundsBoxExtent, FGlobalDynamicVertexBuffer& DynamicVertexBuffer)
+void FOcclusionFeedback::AddPrimitive(const FPrimitiveOcclusionHistoryKey& PrimitiveKey, const FVector& BoundsOrigin, const FVector& BoundsBoxExtent, FGlobalDynamicVertexBuffer& DynamicVertexBuffer)
 {
 	constexpr uint32 MaxBatchedPrimitives = 512;
 	constexpr uint32 PrimitiveStride = sizeof(FVector4f) * 2u;
@@ -1798,7 +1798,7 @@ void FOcclusionFeedback::AddPrimitive(FRHICommandList& RHICmdList, const FPrimit
 	{
 		FOcclusionBatch OcclusionBatch;
 		OcclusionBatch.NumBatchedPrimitives = 0u;
-		OcclusionBatch.VertexAllocation = DynamicVertexBuffer.Allocate(RHICmdList, MaxBatchedPrimitives * PrimitiveStride);
+		OcclusionBatch.VertexAllocation = DynamicVertexBuffer.Allocate(MaxBatchedPrimitives * PrimitiveStride);
 		check(OcclusionBatch.VertexAllocation.IsValid());
 		BatchOcclusionQueries.Add(OcclusionBatch);
 	}
diff --git a/Engine/Source/Runtime/Renderer/Private/ScenePrivate.h b/Engine/Source/Runtime/Renderer/Private/ScenePrivate.h
--- a/Engine/Source/Runtime/Renderer/Private/ScenePrivate.h
+++ b/Engine/Source/Runtime/Renderer/Private/ScenePrivate.h
@@ -711,7 +711,7 @@ public:
 	virtual void InitRHI(FRHICommandListBase& RHICmdList) override;
 	virtual void ReleaseRHI() override;
 
-	void AddPrimitive(FRHICommandList& RHICmdList, const FPrimitiveOcclusionHistoryKey& PrimitiveKey, const FVector& BoundsOrigin, const FVector& BoundsBoxExtent, FGlobalDynamicVertexBuffer& DynamicVertexBuffer);
+	void AddPrimitive(const FPrimitiveOcclusionHistoryKey& PrimitiveKey, const FVector& BoundsOrigin, const FVector& BoundsBoxExtent, FGlobalDynamicVertexBuffer& DynamicVertexBuffer);
 
 	void BeginOcclusionScope(FRDGBuilder& GraphBuilder);
 	void EndOcclusionScope(FRDGBuilder& GraphBuilder);
diff --git a/Engine/Source/Runtime/Renderer/Private/SceneRendering.cpp b/Engine/Source/Runtime/Renderer/Private/SceneRendering.cpp
--- a/Engine/Source/Runtime/Renderer/Private/SceneRendering.cpp
+++ b/Engine/Source/Runtime/Renderer/Private/SceneRendering.cpp
@@ -2486,6 +2486,9 @@ FViewFamilyInfo::~FViewFamilyInfo()
 {
 }
 
+TGlobalResource<FGlobalDynamicReadBuffer> FSceneRenderer::DynamicReadBufferForInitViews;
+TGlobalResource<FGlobalDynamicReadBuffer> FSceneRenderer::DynamicReadBufferForShadows;
+
 /*-----------------------------------------------------------------------------
 	FSceneRenderer
 -----------------------------------------------------------------------------*/
@@ -2494,7 +2497,6 @@ FSceneRenderer::FSceneRenderer(const FSceneViewFamily* InViewFamily, FHitProxyCo
 ,	ViewFamily(*CheckPointer(InViewFamily))
 ,	MeshCollector(InViewFamily->GetFeatureLevel(), Allocator)
 ,	EditorMeshCollector(InViewFamily->GetFeatureLevel(), Allocator)
-,	RayTracingCollector(InViewFamily->GetFeatureLevel(), Allocator)
 ,	VirtualShadowMapArray(*CheckPointer(Scene))
 ,	bHasRequestedToggleFreeze(false)
 ,	bUsedPrecomputedVisibility(false)
@@ -3393,7 +3395,7 @@ FSceneRenderer::~FSceneRenderer()
 	SortedShadowsForShadowDepthPass.Release();
 }
 
-IVisibilityTaskData* FSceneRenderer::OnRenderBegin(FRDGBuilder& GraphBuilder, FGlobalDynamicBuffers GlobalDynamicBuffers)
+IVisibilityTaskData* FSceneRenderer::OnRenderBegin(FRDGBuilder& GraphBuilder)
 {
 	check(!FDeferredUpdateResource::IsUpdateNeeded());
 
@@ -3443,7 +3445,7 @@ IVisibilityTaskData* FSceneRenderer::OnRenderBegin(FRDGBuilder& GraphBuilder, FG
 
 	GraphBuilder.RHICmdList.BeginScene();
 
-	return LaunchVisibilityTasks(GraphBuilder.RHICmdList, *this, GlobalDynamicBuffers);
+	return LaunchVisibilityTasks(GraphBuilder.RHICmdList, *this);
 }
 
 /** 
@@ -4432,16 +4434,6 @@ void FSceneRenderer::RenderThreadEnd(FRHICommandListImmediate& RHICmdList, const
 				}
 			}
 
-			// The temporary proxies would normally be deleted by the scene renderer destructors, but this needs to happen on the render thread, so
-			// we will do it now, before destroying the renderers on a background thread. It's important to do this after the above wait, since the
-			// setup tasks use these temporary proxies.
-			for (FSceneRenderer* SceneRenderer : SceneRenderers)
-			{
-				SceneRenderer->MeshCollector.DeleteTemporaryProxies();
-				SceneRenderer->EditorMeshCollector.DeleteTemporaryProxies();
-				SceneRenderer->RayTracingCollector.DeleteTemporaryProxies();
-			}
-
 			FGraphEventArray CommandListTasks = MoveTemp(RHICmdList.GetRenderThreadTaskArray());
 
 			GSceneRenderCleanUpState.Task = FFunctionGraphTask::CreateAndDispatchWhenReady([LocalSceneRenderers = CopyTemp(SceneRenderers)]
diff --git a/Engine/Source/Runtime/Renderer/Private/SceneRendering.h b/Engine/Source/Runtime/Renderer/Private/SceneRendering.h
--- a/Engine/Source/Runtime/Renderer/Private/SceneRendering.h
+++ b/Engine/Source/Runtime/Renderer/Private/SceneRendering.h
@@ -441,7 +441,7 @@ public:
 	 * Batches a primitive's occlusion query for rendering.
 	 * @param Bounds - The primitive's bounds.
 	 */
-	FRHIRenderQuery* BatchPrimitive(FRHICommandList& RHICmdList, const FVector& BoundsOrigin, const FVector& BoundsBoxExtent, FGlobalDynamicVertexBuffer& DynamicVertexBuffer);
+	FRHIRenderQuery* BatchPrimitive(const FVector& BoundsOrigin, const FVector& BoundsBoxExtent, FGlobalDynamicVertexBuffer& DynamicVertexBuffer);
 	inline int32 GetNumBatchOcclusionQueries() const
 	{
 		return BatchOcclusionQueries.Num();
@@ -2040,7 +2040,6 @@ public:
 
 	FMeshElementCollector MeshCollector;
 	FMeshElementCollector EditorMeshCollector;
-	FMeshElementCollector RayTracingCollector;
 
 	/** Information about the visible lights. */
 	TArray<FVisibleLightInfo, SceneRenderingAllocator> VisibleLightInfos;
@@ -2267,6 +2266,9 @@ public:
 
 	inline TConstStridedView<FSceneView> GetSceneViews() const { return MakeStridedViewOfBase<const FSceneView>(MakeArrayView(Views)); }
 
+	static TGlobalResource<FGlobalDynamicReadBuffer> DynamicReadBufferForInitViews;
+	static TGlobalResource<FGlobalDynamicReadBuffer> DynamicReadBufferForShadows;
+
 protected:
 
 	/** Size of the family. */
@@ -2303,9 +2305,9 @@ protected:
 
 	ERendererOutput GetRendererOutput() const;
 
-	FDynamicShadowsTaskData* BeginInitDynamicShadows(bool bRunningEarly, IVisibilityTaskData* VisibilityTaskData);
-	void FinishInitDynamicShadows(FRDGBuilder& GraphBuilder, FDynamicShadowsTaskData* TaskData, FGlobalDynamicIndexBuffer& DynamicIndexBuffer, FGlobalDynamicVertexBuffer& DynamicVertexBuffer, FGlobalDynamicReadBuffer& DynamicReadBuffer, FInstanceCullingManager& InstanceCullingManager, FRDGExternalAccessQueue& ExternalAccessQueue);
-	FDynamicShadowsTaskData* InitDynamicShadows(FRDGBuilder& GraphBuilder, FGlobalDynamicIndexBuffer& DynamicIndexBuffer, FGlobalDynamicVertexBuffer& DynamicVertexBuffer, FGlobalDynamicReadBuffer& DynamicReadBuffer, FInstanceCullingManager& InstanceCullingManager, FRDGExternalAccessQueue& ExternalAccessQueue);
+	FDynamicShadowsTaskData* BeginInitDynamicShadows( bool bRunningEarly, IVisibilityTaskData* VisibilityTaskData);
+	void FinishInitDynamicShadows(FRDGBuilder& GraphBuilder, FDynamicShadowsTaskData* TaskData, FInstanceCullingManager& InstanceCullingManager, FRDGExternalAccessQueue& ExternalAccessQueue);
+	FDynamicShadowsTaskData* InitDynamicShadows(FRDGBuilder& GraphBuilder, FInstanceCullingManager& InstanceCullingManager, FRDGExternalAccessQueue& ExternalAccessQueue);
 
 	void CreateDynamicShadows(FDynamicShadowsTaskData& TaskData);
 	void FilterDynamicShadows(FDynamicShadowsTaskData& TaskData);
@@ -2419,7 +2421,7 @@ protected:
 	void DrawDebugShadowFrustum(FViewInfo& View, FProjectedShadowInfo& ProjectedShadowInfo);
 
 	/** Gathers dynamic mesh elements for all shadows. */
-	void GatherShadowDynamicMeshElements(FGlobalDynamicIndexBuffer& DynamicIndexBuffer, FGlobalDynamicVertexBuffer& DynamicVertexBuffer, FGlobalDynamicReadBuffer& DynamicReadBuffer, FInstanceCullingManager& InstanceCullingManager);
+	void GatherShadowDynamicMeshElements(FRHICommandList& RHICmdList, FInstanceCullingManager& InstanceCullingManager);
 
 	/** Performs once per frame setup prior to visibility determination. */
 	void PreVisibilityFrameSetup(FRDGBuilder& GraphBuilder);
@@ -2437,7 +2439,7 @@ protected:
 	bool ShouldRenderTranslucency(ETranslucencyPass::Type TranslucencyPass) const;
 
 	/** Called at the begin / finish of the scene render. */
-	IVisibilityTaskData* OnRenderBegin(FRDGBuilder& GraphBuilder, FGlobalDynamicBuffers GlobalDynamicBuffers);
+	IVisibilityTaskData* OnRenderBegin(FRDGBuilder& GraphBuilder);
 	void OnRenderFinish(FRDGBuilder& GraphBuilder, FRDGTextureRef ViewFamilyTexture);
 
 	bool RenderCustomDepthPass(
@@ -2742,10 +2744,6 @@ private:
 	FInstanceCullingDrawParams DebugViewModeInstanceCullingDrawParams;
 	FInstanceCullingDrawParams TranslucencyInstanceCullingDrawParams;
 
-	static FGlobalDynamicIndexBuffer DynamicIndexBuffer;
-	static FGlobalDynamicVertexBuffer DynamicVertexBuffer;
-	static TGlobalResource<FGlobalDynamicReadBuffer> DynamicReadBuffer;
-
 	const FViewInfo* CachedView = nullptr;
 };
 
diff --git a/Engine/Source/Runtime/Renderer/Private/SceneVisibility.cpp b/Engine/Source/Runtime/Renderer/Private/SceneVisibility.cpp
--- a/Engine/Source/Runtime/Renderer/Private/SceneVisibility.cpp
+++ b/Engine/Source/Runtime/Renderer/Private/SceneVisibility.cpp
@@ -368,9 +368,9 @@ TRACE_DECLARE_INT_COUNTER(Scene_Visibility_Relevance_NumPrimitivesProcessed, TEX
 
 ///////////////////////////////////////////////////////////////////////////////
 
-IVisibilityTaskData* LaunchVisibilityTasks(FRHICommandListImmediate& RHICmdList, FSceneRenderer& SceneRenderer, FGlobalDynamicBuffers GlobalDynamicBuffers)
+IVisibilityTaskData* LaunchVisibilityTasks(FRHICommandListImmediate& RHICmdList, FSceneRenderer& SceneRenderer)
 {
-	FVisibilityTaskData* TaskData = SceneRenderer.Allocator.Create<FVisibilityTaskData>(RHICmdList, SceneRenderer, GlobalDynamicBuffers);
+	FVisibilityTaskData* TaskData = SceneRenderer.Allocator.Create<FVisibilityTaskData>(RHICmdList, SceneRenderer);
 	TaskData->LaunchVisibilityTasks();
 	return TaskData;
 }
@@ -2625,11 +2625,11 @@ void FGPUOcclusionPacket::FProcessVisitor::AddOcclusionQuery(const FOcclusionQue
 
 	if (Query.bGroupedQuery)
 	{
-		RenderQuery = Packet.View.GroupedOcclusionQueries.BatchPrimitive(RHICmdList, Query.Bounds.Origin, Query.Bounds.Extent, DynamicVertexBuffer);
+		RenderQuery = Packet.View.GroupedOcclusionQueries.BatchPrimitive(Query.Bounds.Origin, Query.Bounds.Extent, DynamicVertexBuffer);
 	}
 	else
 	{
-		RenderQuery = Packet.View.IndividualOcclusionQueries.BatchPrimitive(RHICmdList, Query.Bounds.Origin, Query.Bounds.Extent, DynamicVertexBuffer);
+		RenderQuery = Packet.View.IndividualOcclusionQueries.BatchPrimitive(Query.Bounds.Origin, Query.Bounds.Extent, DynamicVertexBuffer);
 	}
 
 	Packet.ViewState.Occlusion.LastOcclusionQuery = RenderQuery;
@@ -2690,7 +2690,7 @@ void FGPUOcclusionPacket::FProcessVisitor::SubmitThrottledOcclusionQueries()
 		FThrottledOcclusionQuery* Query = SortedQueries[Index];
 		FPrimitiveOcclusionHistory* PrimitiveOcclusionHistory = PrimitiveOcclusionHistorySet.Find(Query->PrimitiveOcclusionHistoryKey);
 
-		FRHIRenderQuery* RenderQuery = Packet.View.IndividualOcclusionQueries.BatchPrimitive(RHICmdList, Query->Bounds.Origin, Query->Bounds.Extent, DynamicVertexBuffer);
+		FRHIRenderQuery* RenderQuery = Packet.View.IndividualOcclusionQueries.BatchPrimitive(Query->Bounds.Origin, Query->Bounds.Extent, DynamicVertexBuffer);
 
 		Packet.ViewState.Occlusion.LastOcclusionQuery = RenderQuery;
 		Packet.ViewState.Occlusion.NumRequestedQueries++;
@@ -2982,7 +2982,7 @@ void FGPUOcclusionParallel::Finalize()
 	}
 
 	Packets.Empty();
-	DynamicVertexBuffer.Commit(*RHICmdList);
+	DynamicVertexBuffer.Commit();
 	RHICmdList->FinishRecording();
 	FinalizeTask.Trigger();
 }
@@ -2992,6 +2992,7 @@ void FGPUOcclusionParallel::Map(FRHICommandListImmediate& RHICmdListImmediate)
 	FGPUOcclusion::Map(RHICmdListImmediate);
 	RHICmdList = new FRHICommandList(RHICmdListImmediate.GetGPUMask());
 	RHICmdList->SwitchPipeline(ERHIPipeline::Graphics);
+	DynamicVertexBuffer.Init(*RHICmdList);
 }
 
 void FGPUOcclusionParallel::Unmap(FRHICommandListImmediate& RHICmdListImmediate)
@@ -3018,8 +3019,15 @@ void FGPUOcclusionSerial::AddPrimitives(FPrimitiveRange PrimitiveRange)
 	}
 }
 
+void FGPUOcclusionSerial::Map(FRHICommandListImmediate& RHICmdListImmediate)
+{
+	FGPUOcclusion::Map(RHICmdListImmediate);
+	DynamicVertexBuffer.Init(RHICmdListImmediate);
+}
+
 void FGPUOcclusionSerial::Unmap(FRHICommandListImmediate& RHICmdListImmediate)
 {
+	DynamicVertexBuffer.Commit();
 	ProcessVisitor.SubmitThrottledOcclusionQueries();
 	Packet.RecordOcclusionCullResult(OcclusionCullResult);
 	FGPUOcclusion::Unmap(RHICmdListImmediate);
@@ -3446,14 +3454,13 @@ void FVisibilityViewPacket::BeginInitVisibility()
 
 ///////////////////////////////////////////////////////////////////////////////
 
-FVisibilityTaskData::FVisibilityTaskData(FRHICommandListImmediate& InRHICmdList, FSceneRenderer& InSceneRenderer, FGlobalDynamicBuffers InGlobalDynamicBuffers)
+FVisibilityTaskData::FVisibilityTaskData(FRHICommandListImmediate& InRHICmdList, FSceneRenderer& InSceneRenderer)
 	: RHICmdList(InRHICmdList)
 	, SceneRenderer(InSceneRenderer)
 	, Scene(*SceneRenderer.Scene)
 	, Views(SceneRenderer.AllViews)
 	, ViewFamily(SceneRenderer.ViewFamily)
 	, ShadingPath(Scene.GetShadingPath())
-	, GlobalDynamicBuffers(InGlobalDynamicBuffers)
 	, MeshCollector(SceneRenderer.MeshCollector)
 	, EditorMeshCollector(SceneRenderer.EditorMeshCollector)
 	, TaskConfig(Scene, Views)
@@ -3563,7 +3570,15 @@ void FVisibilityTaskData::LaunchVisibilityTasks()
 	DynamicMeshElements.ViewCommandsPerView.SetNum(Views.Num());
 	DynamicMeshElements.LastElementPerView.SetNum(Views.Num());
 
-	MeshCollector.ClearViewMeshArrays();
+	DynamicMeshElements.DynamicVertexBuffer.Init(RHICmdList);
+	DynamicMeshElements.DynamicIndexBuffer.Init(RHICmdList);
+
+	MeshCollector.Start(
+		RHICmdList,
+		DynamicMeshElements.DynamicVertexBuffer,
+		DynamicMeshElements.DynamicIndexBuffer,
+		SceneRenderer.DynamicReadBufferForInitViews
+	);
 
 	for (int32 ViewIndex = 0; ViewIndex < Views.Num(); ViewIndex++)
 	{
@@ -3571,11 +3586,7 @@ void FVisibilityTaskData::LaunchVisibilityTasks()
 			Views[ViewIndex],
 			&Views[ViewIndex]->DynamicMeshElements,
 			&Views[ViewIndex]->SimpleElementCollector,
-			&Views[ViewIndex]->DynamicPrimitiveCollector,
-			Scene.GetFeatureLevel(),
-			GlobalDynamicBuffers.Index,
-			GlobalDynamicBuffers.Vertex,
-			GlobalDynamicBuffers.Read
+			&Views[ViewIndex]->DynamicPrimitiveCollector
 #if UE_ENABLE_DEBUG_DRAWING
 			, &Views[ViewIndex]->DebugSimpleElementCollector
 #endif
@@ -3584,7 +3595,12 @@ void FVisibilityTaskData::LaunchVisibilityTasks()
 
 	if (GIsEditor)
 	{
-		EditorMeshCollector.ClearViewMeshArrays();
+		EditorMeshCollector.Start(
+			RHICmdList,
+			DynamicMeshElements.DynamicVertexBuffer,
+			DynamicMeshElements.DynamicIndexBuffer,
+			SceneRenderer.DynamicReadBufferForInitViews
+		);
 
 		for (int32 ViewIndex = 0; ViewIndex < Views.Num(); ViewIndex++)
 		{
@@ -3592,11 +3608,7 @@ void FVisibilityTaskData::LaunchVisibilityTasks()
 				Views[ViewIndex],
 				&Views[ViewIndex]->DynamicEditorMeshElements,
 				&Views[ViewIndex]->EditorSimpleElementCollector,
-				&Views[ViewIndex]->DynamicPrimitiveCollector,
-				Scene.GetFeatureLevel(),
-				GlobalDynamicBuffers.Index,
-				GlobalDynamicBuffers.Vertex,
-				GlobalDynamicBuffers.Read
+				&Views[ViewIndex]->DynamicPrimitiveCollector
 #if UE_ENABLE_DEBUG_DRAWING
 				, &Views[ViewIndex]->DebugSimpleElementCollector
 #endif
@@ -3676,7 +3688,7 @@ void FVisibilityTaskData::MergeSecondaryViewVisibility()
 	}
 }
 
-void FVisibilityTaskData::GatherDynamicMeshElements()
+void FVisibilityTaskData::FinishGatherDynamicMeshElements(FVirtualTextureUpdater* VirtualTextureUpdater)
 {
 	check(IsInRenderingThread());
 
@@ -3692,10 +3704,27 @@ void FVisibilityTaskData::GatherDynamicMeshElements()
 		check(DynamicMeshElements.PrimitiveViewMasks);
 		GatherDynamicMeshElements(*DynamicMeshElements.PrimitiveViewMasks);
 	}
+
+	// Sync the virtual texture update task before finishing the mesh collectors. Render proxies can register new
+	// materials which require evaluating uniform expression caches, which can contain virtual textures. Newly allocated
+	// virtual textures are processed later.
+	FVirtualTextureSystem::Get().WaitForTasks(VirtualTextureUpdater);
+
+	MeshCollector.Finish();
+
+	if (GIsEditor)
+	{
+		EditorMeshCollector.Finish();
+	}
+
+	DynamicMeshElements.DynamicVertexBuffer.Commit();
+	DynamicMeshElements.DynamicIndexBuffer.Commit();
 }
 
 void FVisibilityTaskData::GatherDynamicMeshElementsForPrimitive(int32 PrimitiveIndex, uint8 ViewMask)
 {
+	SCOPED_NAMED_EVENT(DynamicPrimitive, FColor::Magenta);
+
 	TArray<int32, TInlineAllocator<4>> MeshBatchCountBefore;
 	MeshBatchCountBefore.SetNumUninitialized(Views.Num());
 	for (int32 ViewIndex = 0; ViewIndex < Views.Num(); ViewIndex++)
@@ -3777,11 +3806,11 @@ void FVisibilityTaskData::GatherDynamicMeshElements(const FDynamicPrimitiveViewM
 
 	for (int32 PrimitiveIndex = 0; PrimitiveIndex < NumPrimitives; ++PrimitiveIndex)
 	{
-		const uint8 ViewMask = GetPrimaryViewMask(DynamicPrimitiveViewMasks.Primitives[PrimitiveIndex]);
+		const uint8 ViewMask = DynamicPrimitiveViewMasks.Primitives[PrimitiveIndex];
 
 		if (ViewMask != 0)
 		{
-			GatherDynamicMeshElementsForPrimitive(PrimitiveIndex, ViewMask);
+			GatherDynamicMeshElementsForPrimitive(PrimitiveIndex, GetPrimaryViewMask(ViewMask));
 		}
 	}
 
@@ -3790,11 +3819,11 @@ void FVisibilityTaskData::GatherDynamicMeshElements(const FDynamicPrimitiveViewM
 	{
 		for (int32 PrimitiveIndex = 0; PrimitiveIndex < NumPrimitives; ++PrimitiveIndex)
 		{
-			const uint8 ViewMask = GetPrimaryViewMask(DynamicPrimitiveViewMasks.EditorPrimitives[PrimitiveIndex]);
+			const uint8 ViewMask = DynamicPrimitiveViewMasks.EditorPrimitives[PrimitiveIndex];
 
 			if (ViewMask != 0)
 			{
-				GatherDynamicMeshElementsForEditorPrimitive(PrimitiveIndex, ViewMask);
+				GatherDynamicMeshElementsForEditorPrimitive(PrimitiveIndex, GetPrimaryViewMask(ViewMask));
 			}
 		}
 	}
@@ -3927,15 +3956,7 @@ void FVisibilityTaskData::ProcessRenderThreadTasks(FExclusiveDepthStencil::Type
 		Tasks.bWaitingAllowed = true;
 	}
 
-	if (VirtualTextureUpdater)
-	{
-		// Sync the virtual texture update task before gathering dynamic mesh elements. Render proxies can register new
-		// materials which require evaluating uniform expression caches, which can contain virtual textures. Newly allocated
-		// virtual textures are processed later.
-		FVirtualTextureSystem::Get().WaitForTasks(VirtualTextureUpdater);
-	}
-
-	GatherDynamicMeshElements();
+	FinishGatherDynamicMeshElements(VirtualTextureUpdater);
 
 	if (TaskConfig.Schedule == EVisibilityTaskSchedule::Parallel)
 	{
diff --git a/Engine/Source/Runtime/Renderer/Private/SceneVisibility.h b/Engine/Source/Runtime/Renderer/Private/SceneVisibility.h
--- a/Engine/Source/Runtime/Renderer/Private/SceneVisibility.h
+++ b/Engine/Source/Runtime/Renderer/Private/SceneVisibility.h
@@ -11,25 +11,6 @@ class FVisibilityTaskData;
 class FSceneRenderer;
 class FInstanceCullingManager;
 class FVirtualTextureUpdater;
-class FGlobalDynamicIndexBuffer;
-class FGlobalDynamicVertexBuffer;
-class FGlobalDynamicReadBuffer;
-
-struct FGlobalDynamicBuffers
-{
-	FGlobalDynamicBuffers(
-		FGlobalDynamicIndexBuffer& InDynamicIndexBuffer,
-		FGlobalDynamicVertexBuffer& InDynamicVertexBuffer,
-		FGlobalDynamicReadBuffer& InDynamicReadBuffer)
-		: Index(&InDynamicIndexBuffer)
-		, Vertex(&InDynamicVertexBuffer)
-		, Read(&InDynamicReadBuffer)
-	{}
-
-	FGlobalDynamicIndexBuffer* Index;
-	FGlobalDynamicVertexBuffer* Vertex;
-	FGlobalDynamicReadBuffer* Read;
-};
 
 class FViewCommands
 {
@@ -80,4 +61,4 @@ public:
 	//////////////////////////////////////////////////////////////////////////////
 };
 
-extern IVisibilityTaskData* LaunchVisibilityTasks(FRHICommandListImmediate& RHICmdList, FSceneRenderer& SceneRenderer, FGlobalDynamicBuffers GlobalDynamicBuffers);
\ No newline at end of file
+extern IVisibilityTaskData* LaunchVisibilityTasks(FRHICommandListImmediate& RHICmdList, FSceneRenderer& SceneRenderer);
\ No newline at end of file
diff --git a/Engine/Source/Runtime/Renderer/Private/SceneVisibilityPrivate.h b/Engine/Source/Runtime/Renderer/Private/SceneVisibilityPrivate.h
--- a/Engine/Source/Runtime/Renderer/Private/SceneVisibilityPrivate.h
+++ b/Engine/Source/Runtime/Renderer/Private/SceneVisibilityPrivate.h
@@ -18,15 +18,7 @@ class FGPUOcclusionParallel;
 class FGPUOcclusionParallelPacket;
 class FRelevancePacket;
 class FVisibilityTaskData;
-
-enum class ECommandPipeFlushMode : uint8
-{
-	// Launches tasks onto the pipe immediately.
-	Automatic,
-
-	// Requires an explicit call to Flush to process elements.
-	Manual
-};
+class FVirtualTextureUpdater;
 
 /** An async MPSC queue that can schedule serialized tasks onto the render thread or a task thread. When a new command is enqueued into an empty queue, a
  *  task is launched to process all pending elements in the queue as soon as possible. Each command must be reserved with AddNumCommands,
@@ -394,7 +386,7 @@ class FVisibilityTaskData : public IVisibilityTaskData
 	friend FRelevancePacket;
 	friend FComputeAndMarkRelevance;
 public:
-	FVisibilityTaskData(FRHICommandListImmediate& RHICmdList, FSceneRenderer& SceneRenderer, FGlobalDynamicBuffers GlobalDynamicBuffers);
+	FVisibilityTaskData(FRHICommandListImmediate& RHICmdList, FSceneRenderer& SceneRenderer);
 
 	~FVisibilityTaskData() override
 	{
@@ -430,15 +422,10 @@ public:
 		return Tasks.bWaitingAllowed;
 	}
 
-	inline FGlobalDynamicVertexBuffer* GetDynamicVertexBuffer()
-	{
-		return GlobalDynamicBuffers.Vertex;
-	}
-
 private:
 	void MergeSecondaryViewVisibility();
 
-	void GatherDynamicMeshElements();
+	void FinishGatherDynamicMeshElements(FVirtualTextureUpdater* VirtualTextureUpdater);
 
 	void GatherDynamicMeshElements(const FDynamicPrimitiveIndexList& Primitives);
 	void GatherDynamicMeshElements(const FDynamicPrimitiveViewMasks& Primitives);
@@ -454,7 +441,6 @@ private:
 	TArrayView<FViewInfo*> Views;
 	FViewFamilyInfo& ViewFamily;
 	EShadingPath ShadingPath;
-	FGlobalDynamicBuffers GlobalDynamicBuffers;
 	FMeshElementCollector& MeshCollector;
 	FMeshElementCollector& EditorMeshCollector;
 	FSceneRenderingBulkObjectAllocator Allocator;
@@ -469,6 +455,9 @@ private:
 		// Primitive view masks are only non-null when in multi-view mode.
 		FDynamicPrimitiveViewMasks* PrimitiveViewMasks = nullptr;
 
+		FGlobalDynamicVertexBuffer DynamicVertexBuffer;
+		FGlobalDynamicIndexBuffer DynamicIndexBuffer;
+
 		TArray<FViewCommands, TInlineAllocator<4>> ViewCommandsPerView;
 		TArray<uint32, TInlineAllocator<4, SceneRenderingAllocator>> LastElementPerView;
 		TArray<FDynamicPrimitive, SceneRenderingAllocator> DynamicPrimitives;
@@ -896,7 +885,7 @@ public:
 
 		void AddOcclusionFeedback(const FOcclusionFeedbackEntry& Entry)
 		{
-			Packet.OcclusionFeedback.AddPrimitive(RHICmdList, Entry.PrimitiveKey, Entry.Bounds.Origin, Entry.Bounds.Extent, DynamicVertexBuffer);
+			Packet.OcclusionFeedback.AddPrimitive(Entry.PrimitiveKey, Entry.Bounds.Origin, Entry.Bounds.Extent, DynamicVertexBuffer);
 		}
 
 		void AddVisualizeQuery(const FBox& Box)
@@ -984,6 +973,7 @@ public:
 protected:
 	void WaitForLastOcclusionQuery();
 
+	FGlobalDynamicVertexBuffer DynamicVertexBuffer;
 	FVisibilityViewPacket& ViewPacket;
 	const FScene& Scene;
 	FViewInfo& View;
@@ -1059,7 +1049,6 @@ private:
 	const uint32 MaxNonOccludedPrimitives;
 	TArray<FGPUOcclusionParallelPacket*, SceneRenderingAllocator> Packets;
 	FPrimitiveIndexList NonOccludedPrimitives;
-	FGlobalDynamicVertexBuffer DynamicVertexBuffer;
 	FRHICommandList* RHICmdList = nullptr;
 	UE::Tasks::FTaskEvent FinalizeTask{ UE_SOURCE_LOCATION };
 	bool bFinished = false;
@@ -1072,10 +1061,11 @@ public:
 	FGPUOcclusionSerial(FVisibilityViewPacket& InViewPacket)
 		: FGPUOcclusion(InViewPacket)
 		, Packet(InViewPacket, State)
-		, ProcessVisitor(Packet, FRHICommandListExecutor::GetImmediateCommandList(), *InViewPacket.TaskData.GetDynamicVertexBuffer())
+		, ProcessVisitor(Packet, FRHICommandListExecutor::GetImmediateCommandList(), DynamicVertexBuffer)
 	{}
 
 	void AddPrimitives(FPrimitiveRange PrimitiveRange) override;
+	void Map(FRHICommandListImmediate& RHICmdListImmediate) override;
 	void Unmap(FRHICommandListImmediate& RHICmdListImmediate) override;
 
 private:
diff --git a/Engine/Source/Runtime/Renderer/Private/ShadowRendering.h b/Engine/Source/Runtime/Renderer/Private/ShadowRendering.h
--- a/Engine/Source/Runtime/Renderer/Private/ShadowRendering.h
+++ b/Engine/Source/Runtime/Renderer/Private/ShadowRendering.h
@@ -688,8 +688,7 @@ public:
 	void AddReceiverPrimitive(FPrimitiveSceneInfo* PrimitiveSceneInfo);
 
 	/** Gathers dynamic mesh elements for all the shadow's primitives arrays. */
-	void GatherDynamicMeshElements(FSceneRenderer& Renderer, class FVisibleLightInfo& VisibleLightInfo, TArray<const FSceneView*>& ReusedViewsArray, 
-		FGlobalDynamicIndexBuffer& DynamicIndexBuffer, FGlobalDynamicVertexBuffer& DynamicVertexBuffer, FGlobalDynamicReadBuffer& DynamicReadBuffer, FInstanceCullingManager& InstanceCullingManager);
+	void GatherDynamicMeshElements(FMeshElementCollector& MeshCollector, FSceneRenderer& Renderer, class FVisibleLightInfo& VisibleLightInfo, TArray<const FSceneView*>& ReusedViewsArray, FInstanceCullingManager& InstanceCullingManager);
 
 	void SetupMeshDrawCommandsForShadowDepth(FSceneRenderer& Renderer, FInstanceCullingManager& InstanceCullingManager);
 
@@ -925,12 +924,10 @@ private:
 
 	/** Gathers dynamic mesh elements for the given primitive array. */
 	void GatherDynamicMeshElementsArray(
-		FSceneRenderer& Renderer, 
-		FGlobalDynamicIndexBuffer& DynamicIndexBuffer,
-		FGlobalDynamicVertexBuffer& DynamicVertexBuffer,
-		FGlobalDynamicReadBuffer& DynamicReadBuffer,
+		FMeshElementCollector& Collector,
 		const PrimitiveArrayType& PrimitiveArray, 
-		const TArray<const FSceneView*>& ReusedViewsArray,
+		const TArray<const FSceneView*>& Views,
+		const FSceneViewFamily& ViewFamily,
 		TArray<FMeshBatchAndRelevance,SceneRenderingAllocator>& OutDynamicMeshElements,
 		int32& OutNumDynamicSubjectMeshElements);
 
diff --git a/Engine/Source/Runtime/Renderer/Private/ShadowSetup.cpp b/Engine/Source/Runtime/Renderer/Private/ShadowSetup.cpp
--- a/Engine/Source/Runtime/Renderer/Private/ShadowSetup.cpp
+++ b/Engine/Source/Runtime/Renderer/Private/ShadowSetup.cpp
@@ -1245,6 +1245,45 @@ struct FFilteredShadowArrays
 #endif
 };
 
+class FShadowMeshCollector
+{
+public:
+	static FShadowMeshCollector* Create(FRHICommandList& RHICmdList, FSceneRenderer& SceneRenderer)
+	{
+		return SceneRenderer.Allocator.Create<FShadowMeshCollector>(RHICmdList, SceneRenderer);
+	}
+
+	FMeshElementCollector& GetCollector()
+	{
+		return Collector;
+	}
+
+	void Finish()
+	{
+		Collector.Finish();
+
+		DynamicVertexBuffer.Commit();
+		DynamicIndexBuffer.Commit();
+	}
+
+private:
+	FShadowMeshCollector(FRHICommandList& RHICmdList, FSceneRenderer& SceneRenderer)
+		: Collector(SceneRenderer.FeatureLevel, SceneRenderer.Allocator)
+	{
+		DynamicVertexBuffer.Init(RHICmdList);
+		DynamicIndexBuffer.Init(RHICmdList);
+
+		Collector.Start(RHICmdList, DynamicVertexBuffer, DynamicIndexBuffer, SceneRenderer.DynamicReadBufferForShadows);
+	}
+
+	FMeshElementCollector Collector;
+	FGlobalDynamicVertexBuffer DynamicVertexBuffer;
+	FGlobalDynamicIndexBuffer DynamicIndexBuffer;
+
+	template <typename T>
+	friend class TConcurrentLinearBulkObjectAllocator;
+};
+
 // Common setup and working data for all GatherShadowPrimitives Tasks
 struct FDynamicShadowsTaskData
 {
@@ -2516,12 +2555,11 @@ void FProjectedShadowInfo::SetupMeshDrawCommandsForProjectionStenciling(FSceneRe
 	}
 }
 
-void FProjectedShadowInfo::GatherDynamicMeshElements(FSceneRenderer& Renderer, FVisibleLightInfo& VisibleLightInfo, TArray<const FSceneView*>& ReusedViewsArray,
-	FGlobalDynamicIndexBuffer& DynamicIndexBuffer, FGlobalDynamicVertexBuffer& DynamicVertexBuffer, FGlobalDynamicReadBuffer& DynamicReadBuffer, FInstanceCullingManager &InstanceCullingManager)
+void FProjectedShadowInfo::GatherDynamicMeshElements(FMeshElementCollector& MeshCollector, FSceneRenderer& Renderer, FVisibleLightInfo& VisibleLightInfo, TArray<const FSceneView*>& ReusedViewsArray, FInstanceCullingManager& InstanceCullingManager)
 {
 	QUICK_SCOPE_CYCLE_COUNTER(STAT_Shadow_GatherDynamicMeshElements);
 
-	check(ShadowDepthView && IsInRenderingThread());
+	check(ShadowDepthView);
 
 	if (DynamicSubjectPrimitives.Num() > 0 || ReceiverPrimitives.Num() > 0 || SubjectTranslucentPrimitives.Num() > 0)
 	{
@@ -2547,60 +2585,52 @@ void FProjectedShadowInfo::GatherDynamicMeshElements(FSceneRenderer& Renderer, F
 		{
 			ShadowDepthView->SetPreShadowTranslation(FVector(0, 0, 0));
 			ShadowDepthView->SetDynamicMeshElementsShadowCullFrustum(&CascadeSettings.ShadowBoundsAccurate);
-			GatherDynamicMeshElementsArray(Renderer, DynamicIndexBuffer, DynamicVertexBuffer, DynamicReadBuffer, 
-				DynamicSubjectPrimitives, ReusedViewsArray, DynamicSubjectMeshElements, NumDynamicSubjectMeshElements);
+			GatherDynamicMeshElementsArray(MeshCollector, DynamicSubjectPrimitives, ReusedViewsArray, Renderer.ViewFamily, DynamicSubjectMeshElements, NumDynamicSubjectMeshElements);
 			ShadowDepthView->SetPreShadowTranslation(PreShadowTranslation);
 		}
 		else
 		{
 			ShadowDepthView->SetPreShadowTranslation(PreShadowTranslation);
 			ShadowDepthView->SetDynamicMeshElementsShadowCullFrustum(&CasterOuterFrustum);
-			GatherDynamicMeshElementsArray(Renderer, DynamicIndexBuffer, DynamicVertexBuffer, DynamicReadBuffer, 
-				DynamicSubjectPrimitives, ReusedViewsArray, DynamicSubjectMeshElements, NumDynamicSubjectMeshElements);
+			GatherDynamicMeshElementsArray(MeshCollector, DynamicSubjectPrimitives, ReusedViewsArray, Renderer.ViewFamily, DynamicSubjectMeshElements, NumDynamicSubjectMeshElements);
 		}
 
 		ShadowDepthView->DrawDynamicFlags = EDrawDynamicFlags::None;
 
 		int32 NumDynamicSubjectTranslucentMeshElements = 0;
 		ShadowDepthView->SetDynamicMeshElementsShadowCullFrustum(&CasterOuterFrustum);
-		GatherDynamicMeshElementsArray(Renderer, DynamicIndexBuffer, DynamicVertexBuffer, DynamicReadBuffer, 
-			SubjectTranslucentPrimitives, ReusedViewsArray, DynamicSubjectTranslucentMeshElements, NumDynamicSubjectTranslucentMeshElements);
+		GatherDynamicMeshElementsArray(MeshCollector, SubjectTranslucentPrimitives, ReusedViewsArray, Renderer.ViewFamily, DynamicSubjectTranslucentMeshElements, NumDynamicSubjectTranslucentMeshElements);
 	}
 
+	MeshCollector.Commit();
+
 	SetupMeshDrawCommandsForProjectionStenciling(Renderer, InstanceCullingManager);
 	SetupMeshDrawCommandsForShadowDepth(Renderer, InstanceCullingManager);
 }
 
 void FProjectedShadowInfo::GatherDynamicMeshElementsArray(
-	FSceneRenderer& Renderer, 
-	FGlobalDynamicIndexBuffer& DynamicIndexBuffer,
-	FGlobalDynamicVertexBuffer& DynamicVertexBuffer,
-	FGlobalDynamicReadBuffer& DynamicReadBuffer,
-	const PrimitiveArrayType& PrimitiveArray, 
-	const TArray<const FSceneView*>& ReusedViewsArray,
+	FMeshElementCollector& MeshCollector,
+	const PrimitiveArrayType& Primitives,
+	const TArray<const FSceneView*>& Views,
+	const FSceneViewFamily& ViewFamily,
 	TArray<FMeshBatchAndRelevance,SceneRenderingAllocator>& OutDynamicMeshElements,
 	int32& OutNumDynamicSubjectMeshElements)
 {
 	// Simple elements not supported in shadow passes
 	FSimpleElementCollector DynamicSubjectSimpleElements;
 
-	Renderer.MeshCollector.ClearViewMeshArrays();
-	Renderer.MeshCollector.AddViewMeshArrays(
+	MeshCollector.AddViewMeshArrays(
 		ShadowDepthView,
 		&OutDynamicMeshElements, 
 		&DynamicSubjectSimpleElements, 
-		&ShadowDepthView->DynamicPrimitiveCollector,
-		Renderer.ViewFamily.GetFeatureLevel(),
-		&DynamicIndexBuffer,
-		&DynamicVertexBuffer,
-		&DynamicReadBuffer
-		);
+		&ShadowDepthView->DynamicPrimitiveCollector);
 
-	const uint32 PrimitiveCount = PrimitiveArray.Num();
+	const uint32 PrimitiveCount = Primitives.Num();
 
 	for (uint32 PrimitiveIndex = 0; PrimitiveIndex < PrimitiveCount; ++PrimitiveIndex)
 	{
-		const FPrimitiveSceneInfo* PrimitiveSceneInfo = PrimitiveArray[PrimitiveIndex];
+		TRACE_CPUPROFILER_EVENT_SCOPE(DynamicPrimitive);
+		const FPrimitiveSceneInfo* PrimitiveSceneInfo = Primitives[PrimitiveIndex];
 		const FPrimitiveSceneProxy* PrimitiveSceneProxy = PrimitiveSceneInfo->Proxy;
 
 		// Lookup the primitive's cached view relevance
@@ -2609,19 +2639,20 @@ void FProjectedShadowInfo::GatherDynamicMeshElementsArray(
 		if (!ViewRelevance.bInitializedThisFrame)
 		{
 			// Compute the subject primitive's view relevance since it wasn't cached
-			ViewRelevance = PrimitiveSceneInfo->Proxy->GetViewRelevance(ShadowDepthView);
+			ViewRelevance = PrimitiveSceneProxy->GetViewRelevance(ShadowDepthView);
 		}
 
 		// Only draw if the subject primitive is shadow relevant.
 		if (ViewRelevance.bShadowRelevance && ViewRelevance.bDynamicRelevance)
 		{
-			Renderer.MeshCollector.SetPrimitive(PrimitiveSceneInfo->Proxy, PrimitiveSceneInfo->DefaultDynamicHitProxyId);
+			MeshCollector.SetPrimitive(PrimitiveSceneProxy, PrimitiveSceneInfo->DefaultDynamicHitProxyId);
 
-			PrimitiveSceneInfo->Proxy->GetDynamicMeshElements(ReusedViewsArray, Renderer.ViewFamily, 0x1, Renderer.MeshCollector);
+			PrimitiveSceneProxy->GetDynamicMeshElements(Views, ViewFamily, 0x1, MeshCollector);
 		}
 	}
 
-	OutNumDynamicSubjectMeshElements = Renderer.MeshCollector.GetMeshElementCount(0);
+	OutNumDynamicSubjectMeshElements = MeshCollector.GetMeshElementCount(0);
+	MeshCollector.ClearViewMeshArrays();
 }
 
 /** 
@@ -4368,8 +4399,10 @@ void FSceneRenderer::DrawDebugShadowFrustum(FViewInfo& View, FProjectedShadowInf
 	}
 }
 
-void FSceneRenderer::GatherShadowDynamicMeshElements(FGlobalDynamicIndexBuffer& DynamicIndexBuffer, FGlobalDynamicVertexBuffer& DynamicVertexBuffer, FGlobalDynamicReadBuffer& DynamicReadBuffer, FInstanceCullingManager& InstanceCullingManager)
+void FSceneRenderer::GatherShadowDynamicMeshElements(FRHICommandList& RHICmdList, FInstanceCullingManager& InstanceCullingManager)
 {
+	FShadowMeshCollector* ShadowMeshCollector = FShadowMeshCollector::Create(RHICmdList, *this);
+
 	TArray<const FSceneView*> ReusedViewsArray;
 	ReusedViewsArray.AddZeroed(1);
 
@@ -4381,7 +4414,7 @@ void FSceneRenderer::GatherShadowDynamicMeshElements(FGlobalDynamicIndexBuffer&
 		{
 			FProjectedShadowInfo* ProjectedShadowInfo = Atlas.Shadows[ShadowIndex];
 			FVisibleLightInfo& VisibleLightInfo = VisibleLightInfos[ProjectedShadowInfo->GetLightSceneInfo().Id];
-			ProjectedShadowInfo->GatherDynamicMeshElements(*this, VisibleLightInfo, ReusedViewsArray, DynamicIndexBuffer, DynamicVertexBuffer, DynamicReadBuffer, InstanceCullingManager);
+			ProjectedShadowInfo->GatherDynamicMeshElements(ShadowMeshCollector->GetCollector(), *this, VisibleLightInfo, ReusedViewsArray, InstanceCullingManager);
 		}
 	}
 
@@ -4393,7 +4426,7 @@ void FSceneRenderer::GatherShadowDynamicMeshElements(FGlobalDynamicIndexBuffer&
 		{
 			FProjectedShadowInfo* ProjectedShadowInfo = Atlas.Shadows[ShadowIndex];
 			FVisibleLightInfo& VisibleLightInfo = VisibleLightInfos[ProjectedShadowInfo->GetLightSceneInfo().Id];
-			ProjectedShadowInfo->GatherDynamicMeshElements(*this, VisibleLightInfo, ReusedViewsArray, DynamicIndexBuffer, DynamicVertexBuffer, DynamicReadBuffer, InstanceCullingManager);
+			ProjectedShadowInfo->GatherDynamicMeshElements(ShadowMeshCollector->GetCollector(), *this, VisibleLightInfo, ReusedViewsArray, InstanceCullingManager);
 		}
 	}
 
@@ -4401,7 +4434,7 @@ void FSceneRenderer::GatherShadowDynamicMeshElements(FGlobalDynamicIndexBuffer&
 	{
 		FProjectedShadowInfo* ProjectedShadowInfo = SortedShadowsForShadowDepthPass.PreshadowCache.Shadows[ShadowIndex];
 		FVisibleLightInfo& VisibleLightInfo = VisibleLightInfos[ProjectedShadowInfo->GetLightSceneInfo().Id];
-		ProjectedShadowInfo->GatherDynamicMeshElements(*this, VisibleLightInfo, ReusedViewsArray, DynamicIndexBuffer, DynamicVertexBuffer, DynamicReadBuffer, InstanceCullingManager);
+		ProjectedShadowInfo->GatherDynamicMeshElements(ShadowMeshCollector->GetCollector(), *this, VisibleLightInfo, ReusedViewsArray, InstanceCullingManager);
 	}
 
 	for (int32 AtlasIndex = 0; AtlasIndex < SortedShadowsForShadowDepthPass.TranslucencyShadowMapAtlases.Num(); AtlasIndex++)
@@ -4412,7 +4445,7 @@ void FSceneRenderer::GatherShadowDynamicMeshElements(FGlobalDynamicIndexBuffer&
 		{
 			FProjectedShadowInfo* ProjectedShadowInfo = Atlas.Shadows[ShadowIndex];
 			FVisibleLightInfo& VisibleLightInfo = VisibleLightInfos[ProjectedShadowInfo->GetLightSceneInfo().Id];
-			ProjectedShadowInfo->GatherDynamicMeshElements(*this, VisibleLightInfo, ReusedViewsArray, DynamicIndexBuffer, DynamicVertexBuffer, DynamicReadBuffer, InstanceCullingManager);
+			ProjectedShadowInfo->GatherDynamicMeshElements(ShadowMeshCollector->GetCollector(), *this, VisibleLightInfo, ReusedViewsArray, InstanceCullingManager);
 		}
 	}
 
@@ -4424,10 +4457,12 @@ void FSceneRenderer::GatherShadowDynamicMeshElements(FGlobalDynamicIndexBuffer&
 			if (ProjectedShadowInfo->bShouldRenderVSM)
 			{
 				FVisibleLightInfo& VisibleLightInfo = VisibleLightInfos[ProjectedShadowInfo->GetLightSceneInfo().Id];
-				ProjectedShadowInfo->GatherDynamicMeshElements(*this, VisibleLightInfo, ReusedViewsArray, DynamicIndexBuffer, DynamicVertexBuffer, DynamicReadBuffer, InstanceCullingManager);
+				ProjectedShadowInfo->GatherDynamicMeshElements(ShadowMeshCollector->GetCollector(), *this, VisibleLightInfo, ReusedViewsArray, InstanceCullingManager);
 			}
 		}
 	}
+
+	ShadowMeshCollector->Finish();
 }
 
 FAutoConsoleTaskPriority CPrio_GatherShadowPrimitives(
@@ -6206,7 +6241,7 @@ FDynamicShadowsTaskData* FSceneRenderer::BeginInitDynamicShadows(bool bRunningEa
 }
 
 
-void FSceneRenderer::FinishInitDynamicShadows(FRDGBuilder& GraphBuilder, FDynamicShadowsTaskData* TaskData, FGlobalDynamicIndexBuffer& DynamicIndexBuffer, FGlobalDynamicVertexBuffer& DynamicVertexBuffer, FGlobalDynamicReadBuffer& DynamicReadBuffer, FInstanceCullingManager& InstanceCullingManager, FRDGExternalAccessQueue& ExternalAccessQueue)
+void FSceneRenderer::FinishInitDynamicShadows(FRDGBuilder& GraphBuilder, FDynamicShadowsTaskData* TaskData, FInstanceCullingManager& InstanceCullingManager, FRDGExternalAccessQueue& ExternalAccessQueue)
 {
 	SCOPED_NAMED_EVENT_TEXT("FSceneRenderer::FinishInitDynamicShadows", FColor::Magenta);
 	SCOPE_CYCLE_COUNTER(STAT_InitDynamicShadowsTime);
@@ -6226,7 +6261,7 @@ void FSceneRenderer::FinishInitDynamicShadows(FRDGBuilder& GraphBuilder, FDynami
 	AllocateShadowDepthTargets(GraphBuilder.RHICmdList, *TaskData);
 
 	// Generate mesh element arrays from shadow primitive arrays
-	GatherShadowDynamicMeshElements(DynamicIndexBuffer, DynamicVertexBuffer, DynamicReadBuffer, InstanceCullingManager);
+	GatherShadowDynamicMeshElements(GraphBuilder.RHICmdList, InstanceCullingManager);
 
 	// Ensure all shadow view dynamic primitives are uploaded before shadow-culling batching pass.
 	// TODO: automate this such that:
@@ -6267,10 +6302,10 @@ void FSceneRenderer::FinishInitDynamicShadows(FRDGBuilder& GraphBuilder, FDynami
 	}
 }
 
-FDynamicShadowsTaskData* FSceneRenderer::InitDynamicShadows(FRDGBuilder& GraphBuilder, FGlobalDynamicIndexBuffer& DynamicIndexBuffer, FGlobalDynamicVertexBuffer& DynamicVertexBuffer, FGlobalDynamicReadBuffer& DynamicReadBuffer, FInstanceCullingManager& InstanceCullingManager, FRDGExternalAccessQueue& ExternalAccessQueue)
+FDynamicShadowsTaskData* FSceneRenderer::InitDynamicShadows(FRDGBuilder& GraphBuilder, FInstanceCullingManager& InstanceCullingManager, FRDGExternalAccessQueue& ExternalAccessQueue)
 {
 	FDynamicShadowsTaskData* TaskData = BeginInitDynamicShadows(false, nullptr);
-	FinishInitDynamicShadows(GraphBuilder, TaskData, DynamicIndexBuffer, DynamicVertexBuffer, DynamicReadBuffer, InstanceCullingManager, ExternalAccessQueue);
+	FinishInitDynamicShadows(GraphBuilder, TaskData, InstanceCullingManager, ExternalAccessQueue);
 	return TaskData;
 }
 
diff --git a/Engine/Source/Runtime/Renderer/Private/ShadowSetupMobile.cpp b/Engine/Source/Runtime/Renderer/Private/ShadowSetupMobile.cpp
--- a/Engine/Source/Runtime/Renderer/Private/ShadowSetupMobile.cpp
+++ b/Engine/Source/Runtime/Renderer/Private/ShadowSetupMobile.cpp
@@ -308,7 +308,7 @@ FDynamicShadowsTaskData* FMobileSceneRenderer::InitDynamicShadows(FRDGBuilder& G
 		}
 	}
 
-	FDynamicShadowsTaskData* TaskData = FSceneRenderer::InitDynamicShadows(GraphBuilder, DynamicIndexBuffer, DynamicVertexBuffer, DynamicReadBuffer, InstanceCullingManager, ExternalAccessQueue);
+	FDynamicShadowsTaskData* TaskData = FSceneRenderer::InitDynamicShadows(GraphBuilder, InstanceCullingManager, ExternalAccessQueue);
 
 	bool bAlwaysUseCSM = false;
 	const bool bSkipCSMShaderCulling = MobileBasePassAlwaysUsesCSM(Scene->GetShaderPlatform());
