LFV - Added tiled culling pass, unfinished.

#prelfight shader

diff --git a/Engine/Shaders/Private/LocalFogVolumes/LocalFogVolumeTiledCulling.usf b/Engine/Shaders/Private/LocalFogVolumes/LocalFogVolumeTiledCulling.usf
new file mode 100644
--- /dev/null
+++ b/Engine/Shaders/Private/LocalFogVolumes/LocalFogVolumeTiledCulling.usf
@@ -0,0 +1,29 @@
+// Copyright Epic Games, Inc. All Rights Reserved.
+#include "../Common.ush"
+
+#include "LocalFogVolumeCommon.ush"
+
+
+
+uint2 CullDataTextureResolution;
+
+RWTexture2DArray<uint> LocalFogVolumeCullDataTextureUAV;
+
+#ifndef THREADGROUP_SIZE
+#error THREADGROUP_SIZE must be defined
+#endif
+
+[numthreads(THREADGROUP_SIZE, THREADGROUP_SIZE, 1)]
+void LocalFogVolumeTiledCullingCS(uint3 ThreadId : SV_DispatchThreadID)
+{
+	uint2 TilePos = uint2(ThreadId.xy);
+
+	if (any(TilePos >= CullDataTextureResolution))
+	{
+		return; // Skip out of bound texture tile data
+	}
+
+	LocalFogVolumeCullDataTextureUAV[uint3(TilePos, 0)] = TilePos.x; // DEBUG
+}
+
+
diff --git a/Engine/Source/Runtime/Renderer/Private/LocalFogVolumeRendering.cpp b/Engine/Source/Runtime/Renderer/Private/LocalFogVolumeRendering.cpp
--- a/Engine/Source/Runtime/Renderer/Private/LocalFogVolumeRendering.cpp
+++ b/Engine/Source/Runtime/Renderer/Private/LocalFogVolumeRendering.cpp
@@ -76,11 +76,11 @@ void SetDummyLocalFogVolumeForView(FRDGBuilder& GraphBuilder, FViewInfo& View)
 	View.LocalFogVolumeViewData.CullDataTextureArraySRV		= GraphBuilder.CreateSRV(View.LocalFogVolumeViewData.CullDataTextureArray);
 	View.LocalFogVolumeViewData.CullDataTextureArrayUAV		= nullptr;	// Should never be written by culling passes if there are no instances.
 
-	View.LocalFogVolumeViewData.UniformParametersStruct.LocalFogVolumeInstanceCount		= View.LocalFogVolumeViewData.GPUInstanceCount;
-	View.LocalFogVolumeViewData.UniformParametersStruct.LocalFogVolumeTilePixelSize		= GetLocalFogVolumeTilePixelSize();
-	View.LocalFogVolumeViewData.UniformParametersStruct.LocalFogVolumeInstances			= View.LocalFogVolumeViewData.GPUInstanceDataBufferSRV;
-	View.LocalFogVolumeViewData.UniformParametersStruct.LocalFogVolumeCullDataTexture	= View.LocalFogVolumeViewData.CullDataTextureArraySRV;
-	View.LocalFogVolumeViewData.UniformBuffer											= GraphBuilder.CreateUniformBuffer(&View.LocalFogVolumeViewData.UniformParametersStruct);
+	View.LocalFogVolumeViewData.UniformParametersStruct.LocalFogVolumeCommon.LocalFogVolumeInstanceCount		= View.LocalFogVolumeViewData.GPUInstanceCount;
+	View.LocalFogVolumeViewData.UniformParametersStruct.LocalFogVolumeCommon.LocalFogVolumeTilePixelSize		= GetLocalFogVolumeTilePixelSize();
+	View.LocalFogVolumeViewData.UniformParametersStruct.LocalFogVolumeCommon.LocalFogVolumeInstances			= View.LocalFogVolumeViewData.GPUInstanceDataBufferSRV;
+	View.LocalFogVolumeViewData.UniformParametersStruct.LocalFogVolumeCullDataTexture							= View.LocalFogVolumeViewData.CullDataTextureArraySRV;
+	View.LocalFogVolumeViewData.UniformBuffer																	= GraphBuilder.CreateUniformBuffer(&View.LocalFogVolumeViewData.UniformParametersStruct);
 }
 
 
@@ -118,6 +118,56 @@ bool FScene::HasAnyLocalFogVolume() const
 	return LocalFogVolumes.Num() > 0;
 }
 
+/*=============================================================================
+	Local height fog tiled culling
+=============================================================================*/
+
+
+class FLocalFogVolumeTiledCullingCS : public FGlobalShader
+{
+	DECLARE_GLOBAL_SHADER(FLocalFogVolumeTiledCullingCS);
+	SHADER_USE_PARAMETER_STRUCT(FLocalFogVolumeTiledCullingCS, FGlobalShader);
+
+	using FPermutationDomain = TShaderPermutationDomain<>;
+
+public:
+	const static uint32 GroupSize = 8;
+
+	BEGIN_SHADER_PARAMETER_STRUCT(FParameters, )
+		SHADER_PARAMETER_STRUCT_REF(FViewUniformShaderParameters, View)
+		SHADER_PARAMETER_STRUCT(FLocalFogVolumeCommonParameters, LocalFogVolumeCommon)
+		SHADER_PARAMETER(FUintVector2, CullDataTextureResolution)
+		SHADER_PARAMETER_RDG_TEXTURE_UAV(RWTexture2DArray<uint>, LocalFogVolumeCullDataTextureUAV)
+	END_SHADER_PARAMETER_STRUCT()
+
+	static void ModifyCompilationEnvironment(const FGlobalShaderPermutationParameters& Parameters, FShaderCompilerEnvironment& OutEnvironment)
+	{
+		FGlobalShader::ModifyCompilationEnvironment(Parameters, OutEnvironment);
+		OutEnvironment.SetDefine(TEXT("THREADGROUP_SIZE"), GroupSize);
+	}
+};
+IMPLEMENT_GLOBAL_SHADER(FLocalFogVolumeTiledCullingCS, "/Engine/Private/LocalFogVolumes/LocalFogVolumeTiledCulling.usf", "LocalFogVolumeTiledCullingCS", SF_Compute);
+
+static void LocalFogVolumeViewTiledCullingPass(FViewInfo& View, FRDGBuilder& GraphBuilder)
+{
+	FIntVector TileDataTextureSize = View.LocalFogVolumeViewData.CullDataTextureArray->Desc.GetSize();
+
+	FLocalFogVolumeTiledCullingCS::FParameters* PassParameters = GraphBuilder.AllocParameters<FLocalFogVolumeTiledCullingCS::FParameters>();
+	PassParameters->View = View.ViewUniformBuffer;
+	PassParameters->LocalFogVolumeCommon = View.LocalFogVolumeViewData.UniformParametersStruct.LocalFogVolumeCommon;
+	PassParameters->LocalFogVolumeCullDataTextureUAV = View.LocalFogVolumeViewData.CullDataTextureArrayUAV;
+	PassParameters->CullDataTextureResolution = FUintVector2(TileDataTextureSize.X, TileDataTextureSize.Y);
+
+	ERDGPassFlags PassFlag = ERDGPassFlags::Compute; // LFV_TODO try ERDGPassFlags::AsyncCompute later
+
+	TileDataTextureSize.Z = 1;
+	const FIntVector NumGroups = FIntVector::DivideAndRoundUp(TileDataTextureSize, FLocalFogVolumeTiledCullingCS::GroupSize);
+
+	TShaderMapRef<FLocalFogVolumeTiledCullingCS> ComputeShader(View.ShaderMap);
+	FComputeShaderUtils::AddPass(GraphBuilder, RDG_EVENT_NAME("LocalFogVolume.TiledCulling"), PassFlag, ComputeShader, PassParameters, NumGroups);
+}
+
+
 /*=============================================================================
 	Local height fog rendering common function
 =============================================================================*/
@@ -215,12 +265,7 @@ void CreateViewLocalFogVolumeBufferSRV(FViewInfo& View, FRDGBuilder& GraphBuilde
 
 	// Create the texture that will contain the tiled culled result: count in the first slice and indices in the remaining slices
 	const uint32 LocalFogVolumeTilePixelSize				= GetLocalFogVolumeTilePixelSize();
-#if 1
-	// Allocate dummy resource while this is not yet used without the tile culling pass.
-	View.LocalFogVolumeViewData.CullDataTextureArray = GSystemTextures.GetZeroUIntArrayDummy(GraphBuilder);
-	View.LocalFogVolumeViewData.CullDataTextureArraySRV = GraphBuilder.CreateSRV(View.LocalFogVolumeViewData.CullDataTextureArray);
-	View.LocalFogVolumeViewData.CullDataTextureArrayUAV = nullptr;	// Should never be written by culling passes if there are no instances.
-#else
+
 	const FIntPoint CullDataTextureResolution				= FIntPoint::DivideAndRoundUp(View.ViewRect.Size(), FIntPoint(LocalFogVolumeTilePixelSize, LocalFogVolumeTilePixelSize));
 	const uint32 CullDataTextureSliceCount					= LocalFogVolumeTileMaxInstanceCount + 1; // +1 because the first slice is the culled instance count
 	FRDGTextureDesc Texture2DArrayDesc(FRDGTextureDesc::Create2DArray(CullDataTextureResolution, PF_R8_UINT, FClearValueBinding(EClearBinding::ENoneBound), TexCreate_ShaderResource | TexCreate_UAV | TexCreate_ReduceMemoryWithTilingMode | TexCreate_3DTiling, CullDataTextureSliceCount));
@@ -229,13 +274,12 @@ void CreateViewLocalFogVolumeBufferSRV(FViewInfo& View, FRDGBuilder& GraphBuilde
 	View.LocalFogVolumeViewData.CullDataTextureArray		= GraphBuilder.CreateTexture(Texture2DArrayDesc, TEXT("LocalFogVolume.CullingDataTexture"));
 	View.LocalFogVolumeViewData.CullDataTextureArraySRV		= GraphBuilder.CreateSRV(FRDGTextureSRVDesc::Create(View.LocalFogVolumeViewData.CullDataTextureArray));
 	View.LocalFogVolumeViewData.CullDataTextureArrayUAV		= GraphBuilder.CreateUAV(FRDGTextureUAVDesc(View.LocalFogVolumeViewData.CullDataTextureArray));
-#endif
 
-	View.LocalFogVolumeViewData.UniformParametersStruct.LocalFogVolumeInstanceCount		= View.LocalFogVolumeViewData.GPUInstanceCount;
-	View.LocalFogVolumeViewData.UniformParametersStruct.LocalFogVolumeTilePixelSize		= LocalFogVolumeTilePixelSize;
-	View.LocalFogVolumeViewData.UniformParametersStruct.LocalFogVolumeInstances			= View.LocalFogVolumeViewData.GPUInstanceDataBufferSRV;
-	View.LocalFogVolumeViewData.UniformParametersStruct.LocalFogVolumeCullDataTexture	= View.LocalFogVolumeViewData.CullDataTextureArraySRV;
-	View.LocalFogVolumeViewData.UniformBuffer = GraphBuilder.CreateUniformBuffer(&View.LocalFogVolumeViewData.UniformParametersStruct);
+	View.LocalFogVolumeViewData.UniformParametersStruct.LocalFogVolumeCommon.LocalFogVolumeInstanceCount		= View.LocalFogVolumeViewData.GPUInstanceCount;
+	View.LocalFogVolumeViewData.UniformParametersStruct.LocalFogVolumeCommon.LocalFogVolumeTilePixelSize		= LocalFogVolumeTilePixelSize;
+	View.LocalFogVolumeViewData.UniformParametersStruct.LocalFogVolumeCommon.LocalFogVolumeInstances			= View.LocalFogVolumeViewData.GPUInstanceDataBufferSRV;
+	View.LocalFogVolumeViewData.UniformParametersStruct.LocalFogVolumeCullDataTexture							= View.LocalFogVolumeViewData.CullDataTextureArraySRV;
+	View.LocalFogVolumeViewData.UniformBuffer																	= GraphBuilder.CreateUniformBuffer(&View.LocalFogVolumeViewData.UniformParametersStruct);
 }
 
 void InitLocalFogVolumesForViews(
@@ -254,8 +298,9 @@ void InitLocalFogVolumesForViews(
 
 		for (FViewInfo& View : Views)
 		{
-			// LFV_TODO View.LocalFogVolumeGPUInstanceDataBufferSRV should be on the scene since it is common for all view.
 			CreateViewLocalFogVolumeBufferSRV(View, GraphBuilder, SortingData);
+
+			LocalFogVolumeViewTiledCullingPass(View, GraphBuilder);
 		}
 	}
 	else
@@ -372,7 +417,7 @@ void RenderLocalFogVolume(
 
 			uint32 LocalFogVolumeGPUInstanceCount = View.LocalFogVolumeViewData.GPUInstanceCount;
 			GraphBuilder.AddPass(
-				RDG_EVENT_NAME("RenderLocalFogVolume %u inst.", LocalFogVolumeGPUInstanceCount),
+				RDG_EVENT_NAME("LocalFogVolume.Splat (%u inst.)", LocalFogVolumeGPUInstanceCount),
 				PassParameters,
 				ERDGPassFlags::Raster,
 				[VertexShader, PixelShader, PassParameters, LocalFogVolumeGPUInstanceCount, ViewRect](FRHICommandList& RHICmdList)
diff --git a/Engine/Source/Runtime/Renderer/Private/LocalFogVolumeRendering.h b/Engine/Source/Runtime/Renderer/Private/LocalFogVolumeRendering.h
--- a/Engine/Source/Runtime/Renderer/Private/LocalFogVolumeRendering.h
+++ b/Engine/Source/Runtime/Renderer/Private/LocalFogVolumeRendering.h
@@ -19,13 +19,17 @@ struct FMinimalSceneTextures;
 	Local height fog rendering GPU data
 =============================================================================*/
 
-BEGIN_SHADER_PARAMETER_STRUCT(FLocalFogVolumeUniformParameters, )
+BEGIN_SHADER_PARAMETER_STRUCT(FLocalFogVolumeCommonParameters, )
 	SHADER_PARAMETER_RDG_BUFFER_SRV(Buffer<float4>, LocalFogVolumeInstances)
-	SHADER_PARAMETER_RDG_TEXTURE_SRV(Texture2DArray<uint>, LocalFogVolumeCullDataTexture)
 	SHADER_PARAMETER(uint32, LocalFogVolumeInstanceCount)
 	SHADER_PARAMETER(uint32, LocalFogVolumeTilePixelSize)
 END_SHADER_PARAMETER_STRUCT()
 
+BEGIN_SHADER_PARAMETER_STRUCT(FLocalFogVolumeUniformParameters, )
+	SHADER_PARAMETER_STRUCT_INCLUDE(FLocalFogVolumeCommonParameters, LocalFogVolumeCommon)
+	SHADER_PARAMETER_RDG_TEXTURE_SRV(Texture2DArray<uint>, LocalFogVolumeCullDataTexture)
+END_SHADER_PARAMETER_STRUCT()
+
 /*=============================================================================
 	Local height fog rendering common data
 =============================================================================*/
