Reenabled new rendering path for the water info texture. Fixed water info texture not being rendered if the update is issued during a loading screen. Fixed by delaying update until the first frame where the world is actually rendered. Since rendering itself is very fast, this should not perceivably impact performance. Fixed renderdoc capture not working with the new path.

#rnx

diff --git a/Engine/Plugins/Experimental/Water/Source/Runtime/Private/WaterInfoRendering.cpp b/Engine/Plugins/Experimental/Water/Source/Runtime/Private/WaterInfoRendering.cpp
--- a/Engine/Plugins/Experimental/Water/Source/Runtime/Private/WaterInfoRendering.cpp
+++ b/Engine/Plugins/Experimental/Water/Source/Runtime/Private/WaterInfoRendering.cpp
@@ -28,12 +28,6 @@
 #include "ScenePrivate.h"
 #include "SceneRendering.h"
 
-static int32 RenderCaptureNextWaterInfoDraws = 0;
-static FAutoConsoleVariableRef CVarRenderCaptureNextWaterInfoDraws(
-	TEXT("r.Water.WaterInfo.RenderCaptureNextWaterInfoDraws"),
-	RenderCaptureNextWaterInfoDraws,
-	TEXT("Enable capturing of the water info texture for the next N draws"));
-
 static int32 WaterInfoRenderLandscapeMinimumMipLevel = 0;
 static FAutoConsoleVariableRef CVarWaterInfoRenderLandscapeMinimumMipLevel(
 	TEXT("r.Water.WaterInfo.LandscapeMinimumMipLevel"),
@@ -719,9 +713,15 @@ void UpdateWaterInfoRendering(
 	const WaterInfo::FRenderingContext& Context)
 {
 	TRACE_CPUPROFILER_EVENT_SCOPE(WaterInfo::UpdateWaterInfoRendering);
-	
+
+	static auto* CVarRenderCaptureNextWaterInfoDraws = IConsoleManager::Get().FindConsoleVariable(TEXT("r.Water.WaterInfo.RenderCaptureNextWaterInfoDraws"));
+	int32 RenderCaptureNextWaterInfoDraws = CVarRenderCaptureNextWaterInfoDraws ? CVarRenderCaptureNextWaterInfoDraws->AsVariableInt()->GetValueOnGameThread() : 0;
 	RenderCaptureInterface::FScopedCapture RenderCapture((RenderCaptureNextWaterInfoDraws != 0), TEXT("RenderWaterInfo"));
-	RenderCaptureNextWaterInfoDraws = FMath::Max(0, RenderCaptureNextWaterInfoDraws - 1);
+	if (RenderCaptureNextWaterInfoDraws != 0)
+	{
+		RenderCaptureNextWaterInfoDraws = FMath::Max(0, RenderCaptureNextWaterInfoDraws - 1);
+		CVarRenderCaptureNextWaterInfoDraws->Set(RenderCaptureNextWaterInfoDraws);
+	}
 
 	if (!IsValid(Context.TextureRenderTarget) || Scene == nullptr)
 	{
@@ -831,9 +831,6 @@ void UpdateWaterInfoRendering2(FSceneView& InView, const TMap<AWaterZone*, UE::W
 	{
 		const UE::WaterInfo::FRenderingContext& Context(Pair.Value);
 
-		RenderCaptureInterface::FScopedCapture RenderCapture((RenderCaptureNextWaterInfoDraws != 0), TEXT("RenderWaterInfo"));
-		RenderCaptureNextWaterInfoDraws = FMath::Max(0, RenderCaptureNextWaterInfoDraws - 1);
-
 		if (!IsValid(Context.TextureRenderTarget))
 		{
 			continue;
diff --git a/Engine/Plugins/Experimental/Water/Source/Runtime/Private/WaterViewExtension.cpp b/Engine/Plugins/Experimental/Water/Source/Runtime/Private/WaterViewExtension.cpp
--- a/Engine/Plugins/Experimental/Water/Source/Runtime/Private/WaterViewExtension.cpp
+++ b/Engine/Plugins/Experimental/Water/Source/Runtime/Private/WaterViewExtension.cpp
@@ -27,7 +27,7 @@ static TAutoConsoleVariable<float> CVarLocalTessellationUpdateMargin(
 
 static TAutoConsoleVariable<int32> CVarWaterInfoRenderMethod(
 	TEXT("r.Water.WaterInfo.RenderMethod"),
-	0,
+	1,
 	TEXT("0: SceneCaptures, 1: Custom"),
 	ECVF_Default);
 
diff --git a/Engine/Plugins/Experimental/Water/Source/Runtime/Private/WaterZoneActor.cpp b/Engine/Plugins/Experimental/Water/Source/Runtime/Private/WaterZoneActor.cpp
--- a/Engine/Plugins/Experimental/Water/Source/Runtime/Private/WaterZoneActor.cpp
+++ b/Engine/Plugins/Experimental/Water/Source/Runtime/Private/WaterZoneActor.cpp
@@ -13,6 +13,7 @@
 #include "Materials/MaterialInstanceDynamic.h"
 #include "WaterViewExtension.h"
 #include "Algo/AnyOf.h"
+#include "Engine/GameViewportClient.h"
 #include "WaterBodyInfoMeshComponent.h"
 
 #include UE_INLINE_GENERATED_CPP_BY_NAME(WaterZoneActor)
@@ -561,6 +562,14 @@ bool AWaterZone::UpdateWaterInfoTexture()
 		}
 #endif // WITH_EDITOR
 
+		// The render path for rendering the water info texture without scene captures is executed within the scene renderer.
+		// If world rendering is disabled like in a loading screen, the scene renderer is not called and the water info texture will not be drawn.
+		UGameViewportClient* GameViewport = World->GetGameViewport();
+		if (GameViewport && GameViewport->bDisableWorldRendering)
+		{
+			return false;
+		}
+
 		const ETextureRenderTargetFormat Format = bHalfPrecisionTexture ? ETextureRenderTargetFormat::RTF_RGBA16f : RTF_RGBA32f;
 		UTextureRenderTarget2D* OldTexture = WaterInfoTexture;
 		WaterInfoTexture = FWaterUtils::GetOrCreateTransientRenderTarget2D(OldTexture, TEXT("WaterInfoTexture"), RenderTargetResolution, Format);
diff --git a/Engine/Source/Runtime/Renderer/Private/WaterInfoTextureRendering.cpp b/Engine/Source/Runtime/Renderer/Private/WaterInfoTextureRendering.cpp
--- a/Engine/Source/Runtime/Renderer/Private/WaterInfoTextureRendering.cpp
+++ b/Engine/Source/Runtime/Renderer/Private/WaterInfoTextureRendering.cpp
@@ -10,6 +10,7 @@
 #include "PixelShaderUtils.h"
 #include "BasePassRendering.h"
 #include "MobileBasePassRendering.h"
+#include "RenderCaptureInterface.h"
 
 DECLARE_GPU_DRAWCALL_STAT(WaterInfoTexture);
 
@@ -25,6 +26,12 @@ static TAutoConsoleVariable<float> CVarWaterInfoDilationOverwriteMinimumDistance
 	TEXT("The minimum distance below the ground when we allow dilation to write on top of water"),
 	ECVF_RenderThreadSafe);
 
+static TAutoConsoleVariable<int32> CVarRenderCaptureNextWaterInfoDraws(
+	TEXT("r.Water.WaterInfo.RenderCaptureNextWaterInfoDraws"),
+	0,
+	TEXT("Enable capturing of the water info texture for the next N draws"),
+	ECVF_RenderThreadSafe);
+
 BEGIN_SHADER_PARAMETER_STRUCT(FWaterInfoTexturePassParameters, )
 	SHADER_PARAMETER_RDG_UNIFORM_BUFFER(FViewUniformShaderParameters, View)
 	SHADER_PARAMETER_STRUCT_INCLUDE(FInstanceCullingDrawParams, InstanceCullingDrawParams)
@@ -770,6 +777,15 @@ void RenderWaterInfoTexture(
 		return;
 	}
 
+	int32 RenderCaptureNextWaterInfoDraws = CVarRenderCaptureNextWaterInfoDraws.GetValueOnRenderThread();
+	RenderCaptureInterface::FScopedCapture RenderCapture((RenderCaptureNextWaterInfoDraws != 0), GraphBuilder, TEXT("RenderWaterInfo"));
+	if (RenderCaptureNextWaterInfoDraws != 0)
+	{
+		RenderCaptureNextWaterInfoDraws = FMath::Max(0, RenderCaptureNextWaterInfoDraws - 1);
+		CVarRenderCaptureNextWaterInfoDraws->Set(RenderCaptureNextWaterInfoDraws);
+	}
+	
+
 	TRACE_CPUPROFILER_EVENT_SCOPE(WaterInfo::RenderWaterInfoTexture);
 	RDG_EVENT_SCOPE(GraphBuilder, "WaterInfoTexture");
 	RDG_GPU_STAT_SCOPE(GraphBuilder, WaterInfoTexture);
@@ -939,8 +955,6 @@ void RenderWaterInfoTexture(
 
 		// Merge terrain depth, dilated water body depth, water body depth and velocity into a single texture
 		{
-			RDG_EVENT_SCOPE(GraphBuilder, "WaterInfoTextureMerge");
-
 			FWaterInfoTextureMergePS::FPermutationDomain PixelPermutationVector;
 			// The output format depends on the user configurable format of the passed in render target. Since we store depth data in the texture,
 			// it is sometimes desirable to have full 32bit float precision.
@@ -972,8 +986,6 @@ void RenderWaterInfoTexture(
 
 		// Blur the velocity component in the water info texture to get a smooth gradient between water body transitions
 		{
-			RDG_EVENT_SCOPE(GraphBuilder, "WaterInfoTextureBlur");
-
 			FWaterInfoTextureBlurPS::FPermutationDomain PixelPermutationVector;
 			// The output format depends on the user configurable format of the passed in render target. Since we store depth data in the texture,
 			// it is sometimes desirable to have full 32bit float precision.
