Destroy render thread resources for primitives that are being deleted.

diff --git a/Engine/Source/Runtime/Renderer/Private/RendererScene.cpp b/Engine/Source/Runtime/Renderer/Private/RendererScene.cpp
--- a/Engine/Source/Runtime/Renderer/Private/RendererScene.cpp
+++ b/Engine/Source/Runtime/Renderer/Private/RendererScene.cpp
@@ -6579,6 +6579,11 @@ void FScene::UpdateAllPrimitiveSceneInfos(FRDGBuilder& GraphBuilder, EUpdateAllP
 		SCOPED_NAMED_EVENT(FScene_DeletePrimitiveSceneInfo, FColor::Red);
 		for (FPrimitiveSceneInfo* PrimitiveSceneInfo : DeletedPrimitiveSceneInfos)
 		{
+			if (!PrimitiveSceneInfo->Proxy->ShouldConstrainToRenderThread())
+			{
+				PrimitiveSceneInfo->Proxy->DestroyRenderThreadResources();
+			}
+
 			// It is possible that the HitProxies list isn't empty if PrimitiveSceneInfo was Added/Removed in same frame
 			// Delete the PrimitiveSceneInfo on the game thread after the rendering thread has processed its removal.
 			// This must be done on the game thread because the hit proxy references (and possibly other members) need to be freed on the game thread.
