Implemented r.Shadow.LODDistanceFactor and r.Shadow.FarShadow.LODDistanceFactor to adjust which LOD is rendered into shadow maps.

diff --git a/Engine/Source/Runtime/Renderer/Private/ShadowDepthRendering.cpp b/Engine/Source/Runtime/Renderer/Private/ShadowDepthRendering.cpp
--- a/Engine/Source/Runtime/Renderer/Private/ShadowDepthRendering.cpp
+++ b/Engine/Source/Runtime/Renderer/Private/ShadowDepthRendering.cpp
@@ -51,6 +51,20 @@
 #include "ShaderPlatformCachedIniValue.h"
 #include "UnrealEngine.h"
 
+static float GShadowLODDistanceFactor = 1;
+static FAutoConsoleVariableRef CVarShadowScreenMultiple(
+	TEXT("r.Shadow.LODDistanceFactor"),
+	GShadowLODDistanceFactor,
+	TEXT("Multiplier for LOD selection distance when rendering regular shadows"),
+	ECVF_RenderThreadSafe);
+
+static float GFarShadowLODDistanceFactor = 1;
+static FAutoConsoleVariableRef CVarFarShadowScreenMultiple(
+	TEXT("r.Shadow.FarShadow.LODDistanceFactor"),
+	GFarShadowLODDistanceFactor,
+	TEXT("Multiplier for LOD selection distance when rendering far shadows"),
+	ECVF_RenderThreadSafe);
+
 DEFINE_GPU_DRAWCALL_STAT(ShadowDepths);
 
 IMPLEMENT_STATIC_UNIFORM_BUFFER_STRUCT(FShadowDepthPassUniformParameters, "ShadowDepthPass", SceneTextures);
@@ -1188,6 +1202,11 @@ void FProjectedShadowInfo::RenderDepth(
 	}
 }
 
+float FProjectedShadowInfo::GetLODDistanceFactor() const
+{
+	return CascadeSettings.bFarShadowCascade ? GFarShadowLODDistanceFactor : GShadowLODDistanceFactor;
+}
+
 void FProjectedShadowInfo::ModifyViewForShadow(FViewInfo* FoundView) const
 {
 	FIntRect OriginalViewRect = FoundView->ViewRect;
@@ -1200,6 +1219,9 @@ void FProjectedShadowInfo::ModifyViewForShadow(FViewInfo* FoundView) const
 		(int32&)FoundView->DrawDynamicFlags |= (int32)EDrawDynamicFlags::FarShadowCascade;
 	}
 
+	// must match logic in FProjectedShadowInfo::CalcAndUpdateLODToRender(...)
+	FoundView->LODDistanceFactor *= GetLODDistanceFactor();
+
 	// Don't do material texture mip biasing in shadow maps.
 	FoundView->MaterialTextureMipBias = 0;
 
diff --git a/Engine/Source/Runtime/Renderer/Private/ShadowRendering.h b/Engine/Source/Runtime/Renderer/Private/ShadowRendering.h
--- a/Engine/Source/Runtime/Renderer/Private/ShadowRendering.h
+++ b/Engine/Source/Runtime/Renderer/Private/ShadowRendering.h
@@ -856,6 +856,8 @@ private:
 		const FRenderTargetBindingSlots& RenderTargets,
 		const FMeshPassProcessorRenderState& DrawRenderState);
 
+	float GetLODDistanceFactor() const;
+
 	/**
 	* Modifies the passed in view for this shadow
 	*/
diff --git a/Engine/Source/Runtime/Renderer/Private/ShadowSetup.cpp b/Engine/Source/Runtime/Renderer/Private/ShadowSetup.cpp
--- a/Engine/Source/Runtime/Renderer/Private/ShadowSetup.cpp
+++ b/Engine/Source/Runtime/Renderer/Private/ShadowSetup.cpp
@@ -1619,19 +1619,27 @@ void FProjectedShadowInfo::AddCachedMeshDrawCommands_AnyThread(
 
 FLODMask FProjectedShadowInfo::CalcAndUpdateLODToRender(FViewInfo& CurrentView, const FBoxSphereBounds& Bounds, const FPrimitiveSceneInfo* PrimitiveSceneInfo, int32 ForcedLOD) const
 {
-	int32 PrimitiveId = PrimitiveSceneInfo->GetIndex();
+	// must match logic in FProjectedShadowInfo::ModifyViewForShadow(...)
+	const float ShadowLODDistanceFactor = GetLODDistanceFactor();
+	const float bShadowLODDistanceFactorEnabled = ShadowLODDistanceFactor != 1.0f;
+
+	const int32 PrimitiveId = PrimitiveSceneInfo->GetIndex();
 
 	FLODMask ShadowLODToRender = CurrentView.PrimitivesLODMask[PrimitiveId];
 	// calculate it it's not set OR if LOD is overridden
-	if (ForcedLOD > -1 || !ShadowLODToRender.IsValid())
+	if (ForcedLOD > -1 || !ShadowLODToRender.IsValid() || bShadowLODDistanceFactorEnabled)
 	{
 		float MeshScreenSizeSquared = 0;
 		const int8 CurFirstLODIdx = PrimitiveSceneInfo->Proxy->GetCurrentFirstLODIdx_RenderThread();
 
-		const float LODScale = CurrentView.LODDistanceFactor * GetCachedScalabilityCVars().StaticMeshLODDistanceScale;
+		const float LODScale = ShadowLODDistanceFactor * CurrentView.LODDistanceFactor * GetCachedScalabilityCVars().StaticMeshLODDistanceScale;
 		ShadowLODToRender = ComputeLODForMeshes(PrimitiveSceneInfo->StaticMeshRelevances, CurrentView, Bounds.Origin, Bounds.SphereRadius, PrimitiveSceneInfo->GpuLodInstanceRadius, ForcedLOD, MeshScreenSizeSquared, CurFirstLODIdx, LODScale);
 
-		CurrentView.PrimitivesLODMask[PrimitiveId] = ShadowLODToRender;
+		// TODO: support caching when ShadowLODDistanceFactorEnabled (cascades of the same type (regular/far) could reuse results)
+		if (!bShadowLODDistanceFactorEnabled)
+		{
+			CurrentView.PrimitivesLODMask[PrimitiveId] = ShadowLODToRender;
+		}
 	}
 
 	// Use lowest LOD for PreShadow
