LFV - added tiled pass to render all local fog volume in one go after culling. With some renaming and clarification.

#prelfight shader

diff --git a/Engine/Shaders/Private/LocalFogVolumes/LocalFogVolumeCommon.ush b/Engine/Shaders/Private/LocalFogVolumes/LocalFogVolumeCommon.ush
--- a/Engine/Shaders/Private/LocalFogVolumes/LocalFogVolumeCommon.ush
+++ b/Engine/Shaders/Private/LocalFogVolumes/LocalFogVolumeCommon.ush
@@ -31,29 +31,33 @@ struct FLocalFogVolumeGPUInstanceData
 
 
 
-// Structured buffers cannot be sampled from the vertex shader on some mobile platforms.
+// Structured buffers cannot be sampled from the vertex shader on some mobile platforms. 
 // That is why simple Buffers are used for LocalFogVolumeInstances, because this tech needs to run on mobile platforms.
 
 #if defined(LFVStruct)
 
-#define LFVInsts LFVStruct.LFV.LocalFogVolumeInstances
-#define LFVCullDataTex LFVStruct.LFV.LocalFogVolumeCullDataTexture
+#define LFVInsts				LFVStruct.LFV.LocalFogVolumeInstances
+#define LFVTileDataTex			LFVStruct.LFV.LocalFogVolumeTileDataTexture
 
-#define LFVInstcount LFVStruct.LFV.LocalFogVolumeInstanceCount
-#define LFVTilePixelSize LFVStruct.LFV.LocalFogVolumeTilePixelSize
+#define LFVTileDataResolution	LFVStruct.LFV.LocalFogVolumeTileDataTextureResolution
+#define LFVInstcount			LFVStruct.LFV.LocalFogVolumeInstanceCount
+#define LFVTilePixelSize		LFVStruct.LFV.LocalFogVolumeTilePixelSize
 
 #else
 
-uint LFV_LocalFogVolumeInstanceCount;
-uint LFV_LocalFogVolumeTilePixelSize;
-Buffer<float4> LFV_LocalFogVolumeInstances;
-Texture2DArray<uint> LFV_LocalFogVolumeCullDataTexture;
+uint2							LFV_LocalFogVolumeTileDataTextureResolution;
+uint							LFV_LocalFogVolumeInstanceCount;
+uint							LFV_LocalFogVolumeTilePixelSize;
 
-#define LFVInsts LFV_LocalFogVolumeInstances
-#define LFVCullDataTex LFV_LocalFogVolumeCullDataTexture
+Buffer<float4>					LFV_LocalFogVolumeInstances;
+Texture2DArray<uint>			LFV_LocalFogVolumeTileDataTexture;
 
-#define LFVInstcount LFV_LocalFogVolumeInstanceCount
-#define LFVTilePixelSize LFV_LocalFogVolumeTilePixelSize
+#define LFVInsts				LFV_LocalFogVolumeInstances
+#define LFVTileDataTex			LFV_LocalFogVolumeTileDataTexture
+
+#define LFVTileDataResolution	LFV_LocalFogVolumeTileDataTextureResolution
+#define LFVInstcount			LFV_LocalFogVolumeInstanceCount
+#define LFVTilePixelSize		LFV_LocalFogVolumeTilePixelSize
 
 #endif
 
diff --git a/Engine/Shaders/Private/LocalFogVolumes/LocalFogVolumeSplat.usf b/Engine/Shaders/Private/LocalFogVolumes/LocalFogVolumeSplat.usf
--- a/Engine/Shaders/Private/LocalFogVolumes/LocalFogVolumeSplat.usf
+++ b/Engine/Shaders/Private/LocalFogVolumes/LocalFogVolumeSplat.usf
@@ -7,16 +7,44 @@
 
 #include "../SceneTexturesCommon.ush"
 
+
+
+#ifndef LFV_SPLAT_VS
+#define LFV_SPLAT_VS 0
+#endif
+#ifndef LFV_SPLAT_PS
+#define LFV_SPLAT_PS 0
+#endif
+#ifndef LFV_TILED_VS
+#define LFV_TILED_VS 0
+#endif
+#ifndef LFV_TILED_PS
+#define LFV_TILED_PS 0
+#endif
+
+
+
+#if LFV_SPLAT_PS
 #define LFV_CAN_CLIP 1
+#else
+#define LFV_CAN_CLIP 0
+#endif
 #include "LocalFogVolumeCommon.ush"
 
 
+
 struct FLocalFogVolumeVertexOutput
 {
+#if LFV_SPLAT_VS || LFV_SPLAT_PS
 	nointerpolation uint InstanceId : TEXCOORD0;
-};
+#endif
 
+#if LFV_TILED_VS || LFV_TILED_PS
+	nointerpolation uint TileData : TEXCOORD0;
+#endif
+};
 
+#if LFV_SPLAT_VS
 void LocalFogVolumeSplatVS(
 	in float4 InPosition : ATTRIBUTE0,
 	uint InstanceIndex : SV_InstanceID,
@@ -39,8 +67,10 @@ void LocalFogVolumeSplatVS(
 	OutPosition = mul(TranslatedWorldPosition, PrimaryView.TranslatedWorldToClip);
 	OtherVertexOutput.InstanceId = InstanceIndex;
 }
+#endif // LFV_SPLAT_VS
 
 
+#if LFV_SPLAT_PS
 void LocalFogVolumeSplatPS(
 	in float4 SVPos : SV_POSITION,
 	in FLocalFogVolumeVertexOutput OtherVertexOutput,
@@ -76,36 +106,51 @@ void LocalFogVolumeSplatPS(
 		LuminanceColorTransmittance.rgb * PrimaryView.PreExposure,
 		LuminanceColorTransmittance.a);
 }
+#endif // LFV_SPLAT_PS
 
 
-/*
 
-Proof of the height fog integral equations using Mathematica
 
-(*Density of matter according to height Y*)
-Density [y_] := A * Exp[-B * y]
- 
-(*Ray position*)
-Ray[t_] := Oz + t * Rz
 
-(*Integrate Optical Depth for a distance T*)
-Fog[T_] := Integrate[Density[Ray[t]], {t, 0, T}]
+#if LFV_TILED_VS
+void LocalFogVolumeTiledVS(
+	in uint VertexIndex : SV_VertexID,
+	in uint InstanceIndex : SV_InstanceID,
+	out float4 OutPosition : SV_POSITION,
+	out FLocalFogVolumeVertexOutput OtherVertexOutput
+)
+{
+	ResolvedView = ResolveView();
+
+	uint2 TileCoord;
+	TileCoord.y = InstanceIndex / LFVTileDataResolution.x;
+	TileCoord.x = InstanceIndex - TileCoord.y * LFVTileDataResolution.x;
+
+	uint2 TileVertex = TileCoord * LFVTilePixelSize;
+	TileVertex.x += VertexIndex == 1 || VertexIndex == 2 || VertexIndex == 4 ? LFVTilePixelSize : 0;
+	TileVertex.y += VertexIndex == 2 || VertexIndex == 4 || VertexIndex == 5 ? LFVTilePixelSize : 0;
+
+	OutPosition = float4(float2(TileVertex) * ResolvedView.ViewSizeAndInvSize.zw * float2(2.0f, -2.0f) + float2(-1.0, 1.0f), 0.5f, 1.0f);
+
+	OtherVertexOutput.TileData = (TileCoord.x << 16) | TileCoord.y;
+}
+#endif // LFV_TILED_VS
 
-(*Optical Depth integration equation*)
-Fog[T]
 
-(*Transmittance from Optical Depth*)
-Exp[-Integrate[Density[Ray[t]], {t, 0, T}]]
+#if LFV_TILED_PS
+void LocalFogVolumeTiledPS(
+	in float4 SVPos : SV_POSITION,
+	in FLocalFogVolumeVertexOutput OtherVertexOutput,
+	out float4 OutColor : SV_Target0)
+{
+	ResolvedView = ResolveView();
 
-(*Transmittance from a position to the origin of tracing*)
-Trans[T_] := Exp[-Integrate[Density[Ray[t]], {t, 0, T}]]
+	uint2 TilePos = uint2(OtherVertexOutput.TileData >> 16, OtherVertexOutput.TileData & 0xFFFF);
 
-(*Integrate emissive and scattering from height fog, assuming density \
-is extinction and extinction=scattering, thus albedo=1*)
-Integrate[Density[Ray[t]]*EmSc *Trans[t], {t, 0, T}]
+	uint LFVCount = LFVTileDataTex[uint3(TilePos,0)];
 
-=> this leads to EmSc - Transmittance * EmSc
-				 EmSc * (1 - Transmittance)
-				 EmSc * Coverage
+//	OutColor = float4(float(TilePos.x) * 0.1, float(TilePos.y) * 0.1, 0.0, 0.5);
+	OutColor = float4(float(LFVCount) * 0.1, 0.0, 0.0, 0.5);
+}
+#endif // LFV_TILED_PS
 
-*/
\ No newline at end of file
diff --git a/Engine/Shaders/Private/LocalFogVolumes/LocalFogVolumeTiledCulling.usf b/Engine/Shaders/Private/LocalFogVolumes/LocalFogVolumeTiledCulling.usf
--- a/Engine/Shaders/Private/LocalFogVolumes/LocalFogVolumeTiledCulling.usf
+++ b/Engine/Shaders/Private/LocalFogVolumes/LocalFogVolumeTiledCulling.usf
@@ -4,10 +4,7 @@
 #include "LocalFogVolumeCommon.ush"
 
 
-
-uint2 CullDataTextureResolution;
-
-RWTexture2DArray<uint> LocalFogVolumeCullDataTextureUAV;
+RWTexture2DArray<uint> LocalFogVolumeTileDataTextureUAV;
 
 #ifndef THREADGROUP_SIZE
 #error THREADGROUP_SIZE must be defined
@@ -18,12 +15,12 @@ void LocalFogVolumeTiledCullingCS(uint3 ThreadId : SV_DispatchThreadID)
 {
 	uint2 TilePos = uint2(ThreadId.xy);
 
-	if (any(TilePos >= CullDataTextureResolution))
+	if (any(TilePos >= LFVTileDataResolution))
 	{
-		return; // Skip out of bound texture tile data
+		return; // Skip out of bound texture tile data 
 	}
 
-	LocalFogVolumeCullDataTextureUAV[uint3(TilePos, 0)] = TilePos.x; // DEBUG
+	LocalFogVolumeTileDataTextureUAV[uint3(TilePos, 0)] = TilePos.x + TilePos.y; // DEBUG
 }
 
 
diff --git a/Engine/Source/Runtime/Renderer/Private/LocalFogVolumeRendering.cpp b/Engine/Source/Runtime/Renderer/Private/LocalFogVolumeRendering.cpp
--- a/Engine/Source/Runtime/Renderer/Private/LocalFogVolumeRendering.cpp
+++ b/Engine/Source/Runtime/Renderer/Private/LocalFogVolumeRendering.cpp
@@ -19,6 +19,11 @@ static TAutoConsoleVariable<int32> CVarLocalFogVolumeApplyOnTransclucent(
 	TEXT("Project settings enabling the sampling of local fog volumes on translucent elements.\n"),
 	ECVF_ReadOnly | ECVF_RenderThreadSafe);
 
+static TAutoConsoleVariable<int32> CVarLocalFogVolumeTiledRendering(
+	TEXT("r.LocalFogVolume.TiledRendering"), 0,
+	TEXT("Use a tiled approach to render the local fog volume on screen.\n"),
+	ECVF_RenderThreadSafe);
+
 static TAutoConsoleVariable<int32> CVarLocalFogVolumeTilePixelSize(
 	TEXT("r.LocalFogVolume.TilePixelSize"), 128,
 	TEXT("Tile size on screen in pixel at which we cull the local fog volumes.\n"),
@@ -33,6 +38,11 @@ static TAutoConsoleVariable<int32> CVarLocalFogVolumeTileMaxInstanceCount(
 //  - 1920x1080 => 15x9 tiles
 //  - Allowing max 32 volumes at once => culling list buffer = 15 * 9 * 32 * 1 byte = 4320 bytes = 4.3KB
 
+static bool GetLocalFogVolumeTiledRenderingEnable()
+{
+	return CVarLocalFogVolumeTiledRendering.GetValueOnRenderThread() > 0;
+}
+
 static uint32 GetLocalFogVolumeTilePixelSize()
 {
 	return FMath::Max(8u, FMath::Min(512u, (uint32)CVarLocalFogVolumeTilePixelSize.GetValueOnRenderThread()));
@@ -72,15 +82,16 @@ void SetDummyLocalFogVolumeForView(FRDGBuilder& GraphBuilder, FViewInfo& View)
 
 	View.LocalFogVolumeViewData.GPUInstanceDataBufferSRV	= GraphBuilder.CreateSRV(View.LocalFogVolumeViewData.GPUInstanceDataBuffer, PF_A32B32G32R32F);
 
-	View.LocalFogVolumeViewData.CullDataTextureArray		= GSystemTextures.GetZeroUIntArrayDummy(GraphBuilder);
-	View.LocalFogVolumeViewData.CullDataTextureArraySRV		= GraphBuilder.CreateSRV(View.LocalFogVolumeViewData.CullDataTextureArray);
-	View.LocalFogVolumeViewData.CullDataTextureArrayUAV		= nullptr;	// Should never be written by culling passes if there are no instances.
+	View.LocalFogVolumeViewData.TileDataTextureArray		= GSystemTextures.GetZeroUIntArrayDummy(GraphBuilder);
+	View.LocalFogVolumeViewData.TileDataTextureArraySRV		= GraphBuilder.CreateSRV(View.LocalFogVolumeViewData.TileDataTextureArray);
+	View.LocalFogVolumeViewData.TileDataTextureArrayUAV		= nullptr;	// Should never be written by culling passes if there are no instances.
 
-	View.LocalFogVolumeViewData.UniformParametersStruct.LocalFogVolumeCommon.LocalFogVolumeInstanceCount		= View.LocalFogVolumeViewData.GPUInstanceCount;
-	View.LocalFogVolumeViewData.UniformParametersStruct.LocalFogVolumeCommon.LocalFogVolumeTilePixelSize		= GetLocalFogVolumeTilePixelSize();
-	View.LocalFogVolumeViewData.UniformParametersStruct.LocalFogVolumeCommon.LocalFogVolumeInstances			= View.LocalFogVolumeViewData.GPUInstanceDataBufferSRV;
-	View.LocalFogVolumeViewData.UniformParametersStruct.LocalFogVolumeCullDataTexture							= View.LocalFogVolumeViewData.CullDataTextureArraySRV;
-	View.LocalFogVolumeViewData.UniformBuffer																	= GraphBuilder.CreateUniformBuffer(&View.LocalFogVolumeViewData.UniformParametersStruct);
+	View.LocalFogVolumeViewData.UniformParametersStruct.LocalFogVolumeCommon.LocalFogVolumeTileDataTextureResolution= FUintVector2(1, 1);
+	View.LocalFogVolumeViewData.UniformParametersStruct.LocalFogVolumeCommon.LocalFogVolumeInstanceCount			= View.LocalFogVolumeViewData.GPUInstanceCount;
+	View.LocalFogVolumeViewData.UniformParametersStruct.LocalFogVolumeCommon.LocalFogVolumeTilePixelSize			= GetLocalFogVolumeTilePixelSize();
+	View.LocalFogVolumeViewData.UniformParametersStruct.LocalFogVolumeCommon.LocalFogVolumeInstances				= View.LocalFogVolumeViewData.GPUInstanceDataBufferSRV;
+	View.LocalFogVolumeViewData.UniformParametersStruct.LocalFogVolumeTileDataTexture								= View.LocalFogVolumeViewData.TileDataTextureArraySRV;
+	View.LocalFogVolumeViewData.UniformBuffer																		= GraphBuilder.CreateUniformBuffer(&View.LocalFogVolumeViewData.UniformParametersStruct);
 }
 
 
@@ -135,9 +146,8 @@ public:
 
 	BEGIN_SHADER_PARAMETER_STRUCT(FParameters, )
 		SHADER_PARAMETER_STRUCT_REF(FViewUniformShaderParameters, View)
-		SHADER_PARAMETER_STRUCT(FLocalFogVolumeCommonParameters, LocalFogVolumeCommon)
-		SHADER_PARAMETER(FUintVector2, CullDataTextureResolution)
-		SHADER_PARAMETER_RDG_TEXTURE_UAV(RWTexture2DArray<uint>, LocalFogVolumeCullDataTextureUAV)
+		SHADER_PARAMETER_STRUCT(FLocalFogVolumeCommonParameters, LFV)
+		SHADER_PARAMETER_RDG_TEXTURE_UAV(RWTexture2DArray<uint>, LocalFogVolumeTileDataTextureUAV)
 	END_SHADER_PARAMETER_STRUCT()
 
 	static void ModifyCompilationEnvironment(const FGlobalShaderPermutationParameters& Parameters, FShaderCompilerEnvironment& OutEnvironment)
@@ -150,13 +160,12 @@ IMPLEMENT_GLOBAL_SHADER(FLocalFogVolumeTiledCullingCS, "/Engine/Private/LocalFog
 
 static void LocalFogVolumeViewTiledCullingPass(FViewInfo& View, FRDGBuilder& GraphBuilder)
 {
-	FIntVector TileDataTextureSize = View.LocalFogVolumeViewData.CullDataTextureArray->Desc.GetSize();
+	FIntVector TileDataTextureSize = View.LocalFogVolumeViewData.TileDataTextureArray->Desc.GetSize();
 
 	FLocalFogVolumeTiledCullingCS::FParameters* PassParameters = GraphBuilder.AllocParameters<FLocalFogVolumeTiledCullingCS::FParameters>();
 	PassParameters->View = View.ViewUniformBuffer;
-	PassParameters->LocalFogVolumeCommon = View.LocalFogVolumeViewData.UniformParametersStruct.LocalFogVolumeCommon;
-	PassParameters->LocalFogVolumeCullDataTextureUAV = View.LocalFogVolumeViewData.CullDataTextureArrayUAV;
-	PassParameters->CullDataTextureResolution = FUintVector2(TileDataTextureSize.X, TileDataTextureSize.Y);
+	PassParameters->LFV = View.LocalFogVolumeViewData.UniformParametersStruct.LocalFogVolumeCommon;
+	PassParameters->LocalFogVolumeTileDataTextureUAV = View.LocalFogVolumeViewData.TileDataTextureArrayUAV;
 
 	ERDGPassFlags PassFlag = ERDGPassFlags::Compute; // LFV_TODO try ERDGPassFlags::AsyncCompute later
 
@@ -266,20 +275,21 @@ void CreateViewLocalFogVolumeBufferSRV(FViewInfo& View, FRDGBuilder& GraphBuilde
 	// Create the texture that will contain the tiled culled result: count in the first slice and indices in the remaining slices
 	const uint32 LocalFogVolumeTilePixelSize				= GetLocalFogVolumeTilePixelSize();
 
-	const FIntPoint CullDataTextureResolution				= FIntPoint::DivideAndRoundUp(View.ViewRect.Size(), FIntPoint(LocalFogVolumeTilePixelSize, LocalFogVolumeTilePixelSize));
-	const uint32 CullDataTextureSliceCount					= LocalFogVolumeTileMaxInstanceCount + 1; // +1 because the first slice is the culled instance count
-	FRDGTextureDesc Texture2DArrayDesc(FRDGTextureDesc::Create2DArray(CullDataTextureResolution, PF_R8_UINT, FClearValueBinding(EClearBinding::ENoneBound), TexCreate_ShaderResource | TexCreate_UAV | TexCreate_ReduceMemoryWithTilingMode | TexCreate_3DTiling, CullDataTextureSliceCount));
+	const FIntPoint TileDataTextureResolution				= FIntPoint::DivideAndRoundUp(View.ViewRect.Size(), FIntPoint(LocalFogVolumeTilePixelSize, LocalFogVolumeTilePixelSize));
+	const uint32 TileDataTextureSliceCount					= LocalFogVolumeTileMaxInstanceCount + 1; // +1 because the first slice is the culled instance count
+	FRDGTextureDesc Texture2DArrayDesc(FRDGTextureDesc::Create2DArray(TileDataTextureResolution, PF_R8_UINT, FClearValueBinding(EClearBinding::ENoneBound), TexCreate_ShaderResource | TexCreate_UAV | TexCreate_ReduceMemoryWithTilingMode | TexCreate_3DTiling, TileDataTextureSliceCount));
 	// LFV TODO, consider FFastVramConfig onto Texture2DArrayDesc
 
-	View.LocalFogVolumeViewData.CullDataTextureArray		= GraphBuilder.CreateTexture(Texture2DArrayDesc, TEXT("LocalFogVolume.CullingDataTexture"));
-	View.LocalFogVolumeViewData.CullDataTextureArraySRV		= GraphBuilder.CreateSRV(FRDGTextureSRVDesc::Create(View.LocalFogVolumeViewData.CullDataTextureArray));
-	View.LocalFogVolumeViewData.CullDataTextureArrayUAV		= GraphBuilder.CreateUAV(FRDGTextureUAVDesc(View.LocalFogVolumeViewData.CullDataTextureArray));
+	View.LocalFogVolumeViewData.TileDataTextureArray		= GraphBuilder.CreateTexture(Texture2DArrayDesc, TEXT("LocalFogVolume.CullingDataTexture"));
+	View.LocalFogVolumeViewData.TileDataTextureArraySRV		= GraphBuilder.CreateSRV(FRDGTextureSRVDesc::Create(View.LocalFogVolumeViewData.TileDataTextureArray));
+	View.LocalFogVolumeViewData.TileDataTextureArrayUAV		= GraphBuilder.CreateUAV(FRDGTextureUAVDesc(View.LocalFogVolumeViewData.TileDataTextureArray));
 
-	View.LocalFogVolumeViewData.UniformParametersStruct.LocalFogVolumeCommon.LocalFogVolumeInstanceCount		= View.LocalFogVolumeViewData.GPUInstanceCount;
-	View.LocalFogVolumeViewData.UniformParametersStruct.LocalFogVolumeCommon.LocalFogVolumeTilePixelSize		= LocalFogVolumeTilePixelSize;
-	View.LocalFogVolumeViewData.UniformParametersStruct.LocalFogVolumeCommon.LocalFogVolumeInstances			= View.LocalFogVolumeViewData.GPUInstanceDataBufferSRV;
-	View.LocalFogVolumeViewData.UniformParametersStruct.LocalFogVolumeCullDataTexture							= View.LocalFogVolumeViewData.CullDataTextureArraySRV;
-	View.LocalFogVolumeViewData.UniformBuffer																	= GraphBuilder.CreateUniformBuffer(&View.LocalFogVolumeViewData.UniformParametersStruct);
+	View.LocalFogVolumeViewData.UniformParametersStruct.LocalFogVolumeCommon.LocalFogVolumeTileDataTextureResolution= FUintVector2(TileDataTextureResolution.X, TileDataTextureResolution.Y);
+	View.LocalFogVolumeViewData.UniformParametersStruct.LocalFogVolumeCommon.LocalFogVolumeInstanceCount			= View.LocalFogVolumeViewData.GPUInstanceCount;
+	View.LocalFogVolumeViewData.UniformParametersStruct.LocalFogVolumeCommon.LocalFogVolumeTilePixelSize			= LocalFogVolumeTilePixelSize;
+	View.LocalFogVolumeViewData.UniformParametersStruct.LocalFogVolumeCommon.LocalFogVolumeInstances				= View.LocalFogVolumeViewData.GPUInstanceDataBufferSRV;
+	View.LocalFogVolumeViewData.UniformParametersStruct.LocalFogVolumeTileDataTexture								= View.LocalFogVolumeViewData.TileDataTextureArraySRV;
+	View.LocalFogVolumeViewData.UniformBuffer																		= GraphBuilder.CreateUniformBuffer(&View.LocalFogVolumeViewData.UniformParametersStruct);
 }
 
 void InitLocalFogVolumesForViews(
@@ -336,6 +346,12 @@ class FLocalFogVolumeSplatVS : public FGlobalShader
 		}
 		return true;
 	}
+
+	static void ModifyCompilationEnvironment(const FGlobalShaderPermutationParameters& Parameters, FShaderCompilerEnvironment& OutEnvironment)
+	{
+		FGlobalShader::ModifyCompilationEnvironment(Parameters, OutEnvironment);
+		OutEnvironment.SetDefine(TEXT("LFV_SPLAT_VS"), 1);
+	}
 };
 
 IMPLEMENT_GLOBAL_SHADER(FLocalFogVolumeSplatVS, "/Engine/Private/LocalFogVolumes/LocalFogVolumeSplat.usf", "LocalFogVolumeSplatVS", SF_Vertex);
@@ -360,16 +376,90 @@ class FLocalFogVolumeSplatPS : public FGlobalShader
 		}
 		return true;
 	}
+
+	static void ModifyCompilationEnvironment(const FGlobalShaderPermutationParameters& Parameters, FShaderCompilerEnvironment& OutEnvironment)
+	{
+		FGlobalShader::ModifyCompilationEnvironment(Parameters, OutEnvironment);
+		OutEnvironment.SetDefine(TEXT("LFV_SPLAT_PS"), 1);
+	}
 };
 
 IMPLEMENT_GLOBAL_SHADER(FLocalFogVolumeSplatPS, "/Engine/Private/LocalFogVolumes/LocalFogVolumeSplat.usf", "LocalFogVolumeSplatPS", SF_Pixel);
 
+
+class FLocalFogVolumeTiledRenderVS : public FGlobalShader
+{
+	DECLARE_GLOBAL_SHADER(FLocalFogVolumeTiledRenderVS);
+	SHADER_USE_PARAMETER_STRUCT(FLocalFogVolumeTiledRenderVS, FGlobalShader);
+
+	BEGIN_SHADER_PARAMETER_STRUCT(FParameters, )
+		SHADER_PARAMETER_STRUCT_REF(FViewUniformShaderParameters, View)
+		SHADER_PARAMETER_STRUCT(FLocalFogVolumeUniformParameters, LFV)
+	END_SHADER_PARAMETER_STRUCT()
+
+	using FPermutationDomain = TShaderPermutationDomain<>;
+
+	static bool ShouldCompilePermutation(const FGlobalShaderPermutationParameters& Parameters)
+	{
+		if (IsMobilePlatform(Parameters.Platform))
+		{
+			return false;
+		}
+		return true;
+	}
+
+	static void ModifyCompilationEnvironment(const FGlobalShaderPermutationParameters& Parameters, FShaderCompilerEnvironment& OutEnvironment)
+	{
+		FGlobalShader::ModifyCompilationEnvironment(Parameters, OutEnvironment);
+		OutEnvironment.SetDefine(TEXT("LFV_TILED_VS"), 1);
+	}
+};
+
+IMPLEMENT_GLOBAL_SHADER(FLocalFogVolumeTiledRenderVS, "/Engine/Private/LocalFogVolumes/LocalFogVolumeSplat.usf", "LocalFogVolumeTiledVS", SF_Vertex);
+
+class FLocalFogVolumeTiledRenderPS : public FGlobalShader
+{
+	DECLARE_GLOBAL_SHADER(FLocalFogVolumeTiledRenderPS);
+	SHADER_USE_PARAMETER_STRUCT(FLocalFogVolumeTiledRenderPS, FGlobalShader);
+
+	BEGIN_SHADER_PARAMETER_STRUCT(FParameters, )
+		SHADER_PARAMETER_STRUCT_REF(FViewUniformShaderParameters, View)
+		SHADER_PARAMETER_STRUCT(FLocalFogVolumeUniformParameters, LFV)
+	END_SHADER_PARAMETER_STRUCT()
+
+	using FPermutationDomain = TShaderPermutationDomain<>;
+
+	static bool ShouldCompilePermutation(const FGlobalShaderPermutationParameters& Parameters)
+	{
+		if (IsMobilePlatform(Parameters.Platform))
+		{
+			return false;
+		}
+		return true;
+	}
+
+	static void ModifyCompilationEnvironment(const FGlobalShaderPermutationParameters& Parameters, FShaderCompilerEnvironment& OutEnvironment)
+	{
+		FGlobalShader::ModifyCompilationEnvironment(Parameters, OutEnvironment);
+		OutEnvironment.SetDefine(TEXT("LFV_TILED_PS"), 1);
+	}
+};
+
+IMPLEMENT_GLOBAL_SHADER(FLocalFogVolumeTiledRenderPS, "/Engine/Private/LocalFogVolumes/LocalFogVolumeSplat.usf", "LocalFogVolumeTiledPS", SF_Pixel);
+
 BEGIN_SHADER_PARAMETER_STRUCT(FLocalFogVolumePassParameters, )
 	SHADER_PARAMETER_STRUCT_INCLUDE(FLocalFogVolumeSplatVS::FParameters, VS)
 	SHADER_PARAMETER_STRUCT_INCLUDE(FLocalFogVolumeSplatPS::FParameters, PS)
 	SHADER_PARAMETER_RDG_UNIFORM_BUFFER(FSceneTextureUniformParameters, SceneTextures)
 	RENDER_TARGET_BINDING_SLOTS()
 END_SHADER_PARAMETER_STRUCT()
+
+BEGIN_SHADER_PARAMETER_STRUCT(FLocalFogVolumeTiledPassParameters, )
+	SHADER_PARAMETER_STRUCT_INCLUDE(FLocalFogVolumeTiledRenderVS::FParameters, VS)
+	SHADER_PARAMETER_STRUCT_INCLUDE(FLocalFogVolumeTiledRenderPS::FParameters, PS)
+	SHADER_PARAMETER_RDG_UNIFORM_BUFFER(FSceneTextureUniformParameters, SceneTextures)
+	RENDER_TARGET_BINDING_SLOTS()
+END_SHADER_PARAMETER_STRUCT()
  
 void RenderLocalFogVolume(
 	const FScene* Scene,
@@ -393,66 +483,131 @@ void RenderLocalFogVolume(
 				continue;
 			}
 
-			FLocalFogVolumePassParameters* PassParameters = GraphBuilder.AllocParameters<FLocalFogVolumePassParameters>();
 
-			PassParameters->VS.View = GetShaderBinding(View.ViewUniformBuffer);
-			PassParameters->VS.LFV = View.LocalFogVolumeViewData.UniformParametersStruct;
+			if (GetLocalFogVolumeTiledRenderingEnable())
+			{
+				FLocalFogVolumeTiledPassParameters* PassParameters = GraphBuilder.AllocParameters<FLocalFogVolumeTiledPassParameters>();
+
+				PassParameters->VS.View = GetShaderBinding(View.ViewUniformBuffer);
+				PassParameters->VS.LFV = View.LocalFogVolumeViewData.UniformParametersStruct;
+
+				PassParameters->PS.View = GetShaderBinding(View.ViewUniformBuffer);
+				PassParameters->PS.LFV = View.LocalFogVolumeViewData.UniformParametersStruct;
+
+				PassParameters->SceneTextures = SceneTextures.UniformBuffer;
+				PassParameters->RenderTargets[0] = FRenderTargetBinding(SceneColorTexture, ERenderTargetLoadAction::ENoAction);
+
+				FLocalFogVolumeTiledRenderVS::FPermutationDomain VSPermutationVector;
+				auto VertexShader = View.ShaderMap->GetShader< FLocalFogVolumeTiledRenderVS >(VSPermutationVector);
+
+				FLocalFogVolumeTiledRenderPS::FPermutationDomain PsPermutationVector;
+				auto PixelShader = View.ShaderMap->GetShader< FLocalFogVolumeTiledRenderPS >(PsPermutationVector);
+
+				const FIntRect ViewRect = View.ViewRect;
+
+				ClearUnusedGraphResources(VertexShader, &PassParameters->VS);
+				ClearUnusedGraphResources(PixelShader, &PassParameters->PS);
 
-			PassParameters->PS.View = GetShaderBinding(View.ViewUniformBuffer);
-			PassParameters->PS.LFV = View.LocalFogVolumeViewData.UniformParametersStruct;
+				FUintVector2& LocalFogVolumeTileDataTextureResolution = View.LocalFogVolumeViewData.UniformParametersStruct.LocalFogVolumeCommon.LocalFogVolumeTileDataTextureResolution;
+				const uint32 TileInstanceCount = LocalFogVolumeTileDataTextureResolution.X * LocalFogVolumeTileDataTextureResolution.Y;
 
-			PassParameters->SceneTextures = SceneTextures.UniformBuffer;
-			PassParameters->RenderTargets[0] = FRenderTargetBinding(SceneColorTexture, ERenderTargetLoadAction::ENoAction);
+				GraphBuilder.AddPass(
+					RDG_EVENT_NAME("LocalFogVolume.Tiled (%u X %u)", LocalFogVolumeTileDataTextureResolution.X, LocalFogVolumeTileDataTextureResolution.Y),
+					PassParameters,
+					ERDGPassFlags::Raster,
+					[VertexShader, PixelShader, PassParameters, ViewRect, TileInstanceCount](FRHICommandList& RHICmdList)
+				{
+					FGraphicsPipelineStateInitializer GraphicsPSOInit;
+					RHICmdList.ApplyCachedRenderTargets(GraphicsPSOInit);
 
-			FLocalFogVolumeSplatVS::FPermutationDomain VSPermutationVector;
-			auto VertexShader = View.ShaderMap->GetShader< FLocalFogVolumeSplatVS >(VSPermutationVector);
+					RHICmdList.SetViewport(ViewRect.Min.X, ViewRect.Min.Y, 0.0f, ViewRect.Max.X, ViewRect.Max.Y, 1.0f);
 
-			FLocalFogVolumeSplatPS::FPermutationDomain PsPermutationVector;
-			auto PixelShader = View.ShaderMap->GetShader< FLocalFogVolumeSplatPS >(PsPermutationVector);
+					// Render back faces only since camera may intersect
+					GraphicsPSOInit.RasterizerState = TStaticRasterizerState<>::GetRHI();
+					GraphicsPSOInit.DepthStencilState = TStaticDepthStencilState<false, CF_Always>::GetRHI();
+					GraphicsPSOInit.BlendState = TStaticBlendState<CW_RGBA, BO_Add, BF_One, BF_SourceAlpha, BO_Add, BF_Zero, BF_SourceAlpha>::GetRHI();
+					GraphicsPSOInit.PrimitiveType = PT_TriangleList; // LFV_TODO check if rects are supported and use them if so
 
-			const FIntRect ViewRect = View.ViewRect;
+					GraphicsPSOInit.BoundShaderState.VertexDeclarationRHI = GFilterVertexDeclaration.VertexDeclarationRHI; //		GraphicsPSOInit.BoundShaderState.VertexDeclarationRHI = GetVertexDeclarationFVector4();
+					GraphicsPSOInit.BoundShaderState.VertexShaderRHI = VertexShader.GetVertexShader();
+					GraphicsPSOInit.BoundShaderState.PixelShaderRHI = PixelShader.GetPixelShader();
 
-			ClearUnusedGraphResources(VertexShader, &PassParameters->VS);
-			ClearUnusedGraphResources(PixelShader, &PassParameters->PS);
+					SetGraphicsPipelineState(RHICmdList, GraphicsPSOInit, 0);
 
-			uint32 LocalFogVolumeGPUInstanceCount = View.LocalFogVolumeViewData.GPUInstanceCount;
-			GraphBuilder.AddPass(
-				RDG_EVENT_NAME("LocalFogVolume.Splat (%u inst.)", LocalFogVolumeGPUInstanceCount),
-				PassParameters,
-				ERDGPassFlags::Raster,
-				[VertexShader, PixelShader, PassParameters, LocalFogVolumeGPUInstanceCount, ViewRect](FRHICommandList& RHICmdList)
+					SetShaderParameters(RHICmdList, VertexShader, VertexShader.GetVertexShader(), PassParameters->VS);
+					SetShaderParameters(RHICmdList, PixelShader, PixelShader.GetPixelShader(), PassParameters->PS);
+
+					RHICmdList.SetStreamSource(0, nullptr, 0);
+					RHICmdList.DrawPrimitive(
+						0,					// BaseVertexIndex
+						6,					// NumPrimitives
+						TileInstanceCount	// NumInstances
+						);
+				});
+			}
+			else
 			{
-				FGraphicsPipelineStateInitializer GraphicsPSOInit;
-				RHICmdList.ApplyCachedRenderTargets(GraphicsPSOInit);
+				FLocalFogVolumePassParameters* PassParameters = GraphBuilder.AllocParameters<FLocalFogVolumePassParameters>();
+
+				PassParameters->VS.View = GetShaderBinding(View.ViewUniformBuffer);
+				PassParameters->VS.LFV = View.LocalFogVolumeViewData.UniformParametersStruct;
+
+				PassParameters->PS.View = GetShaderBinding(View.ViewUniformBuffer);
+				PassParameters->PS.LFV = View.LocalFogVolumeViewData.UniformParametersStruct;
+
+				PassParameters->SceneTextures = SceneTextures.UniformBuffer;
+				PassParameters->RenderTargets[0] = FRenderTargetBinding(SceneColorTexture, ERenderTargetLoadAction::ENoAction);
+
+				FLocalFogVolumeSplatVS::FPermutationDomain VSPermutationVector;
+				auto VertexShader = View.ShaderMap->GetShader< FLocalFogVolumeSplatVS >(VSPermutationVector);
 
-				RHICmdList.SetViewport(ViewRect.Min.X, ViewRect.Min.Y, 0.0f, ViewRect.Max.X, ViewRect.Max.Y, 1.0f);
+				FLocalFogVolumeSplatPS::FPermutationDomain PsPermutationVector;
+				auto PixelShader = View.ShaderMap->GetShader< FLocalFogVolumeSplatPS >(PsPermutationVector);
 
-				// Render back faces only since camera may intersect
-				GraphicsPSOInit.RasterizerState = TStaticRasterizerState<FM_Solid, CM_CW>::GetRHI();
-				GraphicsPSOInit.DepthStencilState = TStaticDepthStencilState<false, CF_Always>::GetRHI();
-				GraphicsPSOInit.BlendState = TStaticBlendState<CW_RGBA, BO_Add, BF_One, BF_SourceAlpha, BO_Add, BF_Zero, BF_SourceAlpha>::GetRHI();
-				GraphicsPSOInit.PrimitiveType = PT_TriangleList;
+				const FIntRect ViewRect = View.ViewRect;
 
-				GraphicsPSOInit.BoundShaderState.VertexDeclarationRHI = GetVertexDeclarationFVector4();
-				GraphicsPSOInit.BoundShaderState.VertexShaderRHI = VertexShader.GetVertexShader();
-				GraphicsPSOInit.BoundShaderState.PixelShaderRHI = PixelShader.GetPixelShader();
+				ClearUnusedGraphResources(VertexShader, &PassParameters->VS);
+				ClearUnusedGraphResources(PixelShader, &PassParameters->PS);
 
-				SetGraphicsPipelineState(RHICmdList, GraphicsPSOInit, 0);
+				uint32 LocalFogVolumeGPUInstanceCount = View.LocalFogVolumeViewData.GPUInstanceCount;
+				GraphBuilder.AddPass(
+					RDG_EVENT_NAME("LocalFogVolume.Splat (%u inst.)", LocalFogVolumeGPUInstanceCount),
+					PassParameters,
+					ERDGPassFlags::Raster,
+					[VertexShader, PixelShader, PassParameters, LocalFogVolumeGPUInstanceCount, ViewRect](FRHICommandList& RHICmdList)
+				{
+					FGraphicsPipelineStateInitializer GraphicsPSOInit;
+					RHICmdList.ApplyCachedRenderTargets(GraphicsPSOInit);
 
-				SetShaderParameters(RHICmdList, VertexShader, VertexShader.GetVertexShader(), PassParameters->VS);
-				SetShaderParameters(RHICmdList, PixelShader, PixelShader.GetPixelShader(), PassParameters->PS);
+					RHICmdList.SetViewport(ViewRect.Min.X, ViewRect.Min.Y, 0.0f, ViewRect.Max.X, ViewRect.Max.Y, 1.0f);
 
-				RHICmdList.SetStreamSource(0, GetUnitCubeVertexBuffer(), 0);
+					// Render back faces only since camera may intersect
+					GraphicsPSOInit.RasterizerState = TStaticRasterizerState<FM_Solid, CM_CW>::GetRHI();
+					GraphicsPSOInit.DepthStencilState = TStaticDepthStencilState<false, CF_Always>::GetRHI();
+					GraphicsPSOInit.BlendState = TStaticBlendState<CW_RGBA, BO_Add, BF_One, BF_SourceAlpha, BO_Add, BF_Zero, BF_SourceAlpha>::GetRHI();
+					GraphicsPSOInit.PrimitiveType = PT_TriangleList;
 
-				RHICmdList.DrawIndexedPrimitive(GetUnitCubeIndexBuffer()
-					, 0									//BaseVertexIndex
-					, 0									//FirstInstance
-					, 8									//uint32 NumVertices
-					, 0									//uint32 StartIndex
-					, UE_ARRAY_COUNT(GCubeIndices) / 3	//uint32 NumPrimitives
-					, LocalFogVolumeGPUInstanceCount	//uint32 NumInstances
-				);
-			});
+					GraphicsPSOInit.BoundShaderState.VertexDeclarationRHI = GetVertexDeclarationFVector4();
+					GraphicsPSOInit.BoundShaderState.VertexShaderRHI = VertexShader.GetVertexShader();
+					GraphicsPSOInit.BoundShaderState.PixelShaderRHI = PixelShader.GetPixelShader();
+
+					SetGraphicsPipelineState(RHICmdList, GraphicsPSOInit, 0);
+
+					SetShaderParameters(RHICmdList, VertexShader, VertexShader.GetVertexShader(), PassParameters->VS);
+					SetShaderParameters(RHICmdList, PixelShader, PixelShader.GetPixelShader(), PassParameters->PS);
+
+					RHICmdList.SetStreamSource(0, GetUnitCubeVertexBuffer(), 0);
+
+					RHICmdList.DrawIndexedPrimitive(GetUnitCubeIndexBuffer()
+						, 0									//BaseVertexIndex
+						, 0									//FirstInstance
+						, 8									//uint32 NumVertices
+						, 0									//uint32 StartIndex
+						, UE_ARRAY_COUNT(GCubeIndices) / 3	//uint32 NumPrimitives
+						, LocalFogVolumeGPUInstanceCount	//uint32 NumInstances
+					);
+				});
+			}
 		}
 	}
 }
@@ -487,6 +642,7 @@ class FMobileLocalFogVolumeSplatVS : public FGlobalShader
 	{
 		FGlobalShader::ModifyCompilationEnvironment(Parameters, OutEnvironment);
 		const bool bMobileForceDepthRead = MobileUsesFullDepthPrepass(Parameters.Platform);
+		OutEnvironment.SetDefine(TEXT("LFV_SPLAT_VS"), 1);
 		OutEnvironment.SetDefine(TEXT("IS_MOBILE_DEPTHREAD_SUBPASS"), bMobileForceDepthRead ? 0u : 1u);
 	}
 };
@@ -519,6 +675,7 @@ class FMobileLocalFogVolumeSplatPS : public FGlobalShader
 	{
 		FGlobalShader::ModifyCompilationEnvironment(Parameters, OutEnvironment);
 		const bool bMobileForceDepthRead = MobileUsesFullDepthPrepass(Parameters.Platform);
+		OutEnvironment.SetDefine(TEXT("LFV_SPLAT_PS"), 1);
 		OutEnvironment.SetDefine(TEXT("IS_MOBILE_DEPTHREAD_SUBPASS"), bMobileForceDepthRead ? 0u : 1u);
 	}
 };
diff --git a/Engine/Source/Runtime/Renderer/Private/LocalFogVolumeRendering.h b/Engine/Source/Runtime/Renderer/Private/LocalFogVolumeRendering.h
--- a/Engine/Source/Runtime/Renderer/Private/LocalFogVolumeRendering.h
+++ b/Engine/Source/Runtime/Renderer/Private/LocalFogVolumeRendering.h
@@ -21,13 +21,14 @@ struct FMinimalSceneTextures;
 
 BEGIN_SHADER_PARAMETER_STRUCT(FLocalFogVolumeCommonParameters, )
 	SHADER_PARAMETER_RDG_BUFFER_SRV(Buffer<float4>, LocalFogVolumeInstances)
+	SHADER_PARAMETER(FUintVector2, LocalFogVolumeTileDataTextureResolution)
 	SHADER_PARAMETER(uint32, LocalFogVolumeInstanceCount)
 	SHADER_PARAMETER(uint32, LocalFogVolumeTilePixelSize)
 END_SHADER_PARAMETER_STRUCT()
 
 BEGIN_SHADER_PARAMETER_STRUCT(FLocalFogVolumeUniformParameters, )
 	SHADER_PARAMETER_STRUCT_INCLUDE(FLocalFogVolumeCommonParameters, LocalFogVolumeCommon)
-	SHADER_PARAMETER_RDG_TEXTURE_SRV(Texture2DArray<uint>, LocalFogVolumeCullDataTexture)
+	SHADER_PARAMETER_RDG_TEXTURE_SRV(Texture2DArray<uint>, LocalFogVolumeTileDataTexture)
 END_SHADER_PARAMETER_STRUCT()
 
 /*=============================================================================
@@ -99,9 +100,9 @@ struct FLocalFogVolumeViewData
 
 	FLocalFogVolumeUniformParameters UniformParametersStruct;
 
-	FRDGTextureRef CullDataTextureArray;		// First slice is the instance count, later slices are instance indices.
-	FRDGTextureSRVRef CullDataTextureArraySRV;
-	FRDGTextureUAVRef CullDataTextureArrayUAV;
+	FRDGTextureRef TileDataTextureArray;		// First slice is the instance count, later slices are instance indices.
+	FRDGTextureSRVRef TileDataTextureArraySRV;
+	FRDGTextureUAVRef TileDataTextureArrayUAV;
 };
 
 
