Path Tracer: Implement support for virtual texture feedback

This allows skipping the basepass when the path tracer is active, reducing overhead.

VT feedback logic is identical to the rasterizer for now, and will not restart the path tracer accumulation. This is not noticeable in the editor as the right VT pages are loaded fairly quickly. In the offline case, this could potentially lead to some textures being blurrier for some samples, but with proper warmup this should not be noticeable either.

Note that the VT feedback logic is only executed for camera rays for now.

diff --git a/Engine/Shaders/Private/PathTracing/PathTracingMaterialHitShader.usf b/Engine/Shaders/Private/PathTracing/PathTracingMaterialHitShader.usf
--- a/Engine/Shaders/Private/PathTracing/PathTracingMaterialHitShader.usf
+++ b/Engine/Shaders/Private/PathTracing/PathTracingMaterialHitShader.usf
@@ -1301,6 +1301,20 @@ RAY_TRACING_ENTRY_CLOSEST_HIT(PathTracingMaterialCHS,
 	}
 #endif
 
+#if MATERIAL_VIRTUALTEXTURE_FEEDBACK
+	if (PackedPayload.IsCameraRay())
+	{
+		// Virtual texturing feedback logic (camera rays only for now)
+		FinalizeVirtualTextureFeedback(
+			MaterialParameters.VirtualTextureFeedback,
+			MaterialParameters.SvPosition,
+			Payload.BSDFOpacity,
+			View.FrameNumber,
+			View.VTFeedbackBuffer
+		);
+	}
+#endif
+
 	// Output depth for solid, masked 
 	// or when `output depth and velocity' is set and Opacity >= Opacity mask clip value to match the behavior of the rasterizer path. 
 #if (MATERIALBLENDING_SOLID || MATERIALBLENDING_MASKED)
diff --git a/Engine/Shaders/Private/RayTracing/RayTracingDecalMaterialShader.usf b/Engine/Shaders/Private/RayTracing/RayTracingDecalMaterialShader.usf
--- a/Engine/Shaders/Private/RayTracing/RayTracingDecalMaterialShader.usf
+++ b/Engine/Shaders/Private/RayTracing/RayTracingDecalMaterialShader.usf
@@ -142,7 +142,23 @@ RAY_TRACING_ENTRY_CALLABLE(RayTracingDecalMaterialShader, FDecalShaderPayload, P
 	const float DecalFading = saturate(4 - 4 * abs(SwizzlePos.z * 2 - 1)) * DecalParams.x;
 
 	FDecalData Data = GetDecalData(PixelMaterialInputs, MaterialParameters);
+
+#if MATERIAL_VIRTUALTEXTURE_FEEDBACK
+	{
+		// Virtual texturing feedback logic
+		// NOTE: decals are only evaluated for camera rays (and sharp reflections)
+		FinalizeVirtualTextureFeedback(
+			MaterialParameters.VirtualTextureFeedback,
+			MaterialParameters.SvPosition,
+			Data.Opacity * DecalFading,
+			View.FrameNumber,
+			View.VTFeedbackBuffer
+		);
+	}
+#endif
+	
 	Data.WriteToPayload(Payload, DecalWriteFlags, DecalFading);
+
 }
 #endif // RAYCALLABLESHADER
 
@@ -200,8 +216,22 @@ RAY_TRACING_ENTRY_CLOSEST_HIT(RayTracingDecalMaterialCHS,
 	}
 
 	FDecalData Data = GetDecalData(PixelMaterialInputs, MaterialParameters);
-	Data.WriteToPayload(Payload, DECAL_PAYLOAD_FLAGS, 1.0);
+	
+#if MATERIAL_VIRTUALTEXTURE_FEEDBACK
+	{
+		// Virtual texturing feedback logic
+		// NOTE: decals are only evaluated for camera rays (and sharp reflections)
+		FinalizeVirtualTextureFeedback(
+			MaterialParameters.VirtualTextureFeedback,
+			MaterialParameters.SvPosition,
+			Data.Opacity,
+			View.FrameNumber,
+			View.VTFeedbackBuffer
+		);
+	}
+#endif
 
+	Data.WriteToPayload(Payload, DECAL_PAYLOAD_FLAGS, 1.0);
 }
 
 #if USE_MATERIAL_ANY_HIT_SHADER
diff --git a/Engine/Source/Runtime/Renderer/Private/DeferredShadingRenderer.cpp b/Engine/Source/Runtime/Renderer/Private/DeferredShadingRenderer.cpp
--- a/Engine/Source/Runtime/Renderer/Private/DeferredShadingRenderer.cpp
+++ b/Engine/Source/Runtime/Renderer/Private/DeferredShadingRenderer.cpp
@@ -3601,8 +3601,11 @@ void FDeferredShadingSceneRenderer::Render(FRDGBuilder& GraphBuilder)
 		}
 
 		{
-			RenderBasePass(GraphBuilder, SceneTextures, DBufferTextures, BasePassDepthStencilAccess, ForwardScreenSpaceShadowMaskTexture, InstanceCullingManager, bNaniteEnabled, NaniteRasterResults);
-			GraphBuilder.AddDispatchHint();
+			if (!bHasRayTracedOverlay)
+			{
+				RenderBasePass(GraphBuilder, SceneTextures, DBufferTextures, BasePassDepthStencilAccess, ForwardScreenSpaceShadowMaskTexture, InstanceCullingManager, bNaniteEnabled, NaniteRasterResults);
+				GraphBuilder.AddDispatchHint();
+			}
 
 			if (!bAllowReadOnlyDepthBasePass)
 			{
@@ -4124,7 +4127,10 @@ void FDeferredShadingSceneRenderer::Render(FRDGBuilder& GraphBuilder)
 
 		FRDGTextureRef ExposureIlluminance = AddCalculateExposureIlluminancePass(GraphBuilder, Views, SceneTextures, TranslucencyLightingVolumeTextures, ExposureIlluminanceSetup);
 
-		RenderOpaqueFX(GraphBuilder, GetSceneViews(), GetSceneUniforms(), FXSystem, SceneTextures.UniformBuffer);
+		if (!bHasRayTracedOverlay)
+		{
+			RenderOpaqueFX(GraphBuilder, GetSceneViews(), GetSceneUniforms(), FXSystem, SceneTextures.UniformBuffer);
+		}
 
 		FRendererModule& RendererModule = static_cast<FRendererModule&>(GetRendererModule());
 		RendererModule.RenderPostOpaqueExtensions(GraphBuilder, Views, SceneTextures);
@@ -4276,12 +4282,6 @@ void FDeferredShadingSceneRenderer::Render(FRDGBuilder& GraphBuilder)
 			RenderLightShaftBloom(GraphBuilder, SceneTextures, /* inout */ TranslucencyResourceMap);
 		}
 
-		if (bUseVirtualTexturing)
-		{
-			RDG_GPU_STAT_SCOPE(GraphBuilder, VirtualTextureUpdate);
-			VirtualTextureFeedbackEnd(GraphBuilder);
-		}
-
 		FPathTracingResources PathTracingResources;
 
 #if RHI_RAYTRACING
@@ -4313,6 +4313,11 @@ void FDeferredShadingSceneRenderer::Render(FRDGBuilder& GraphBuilder)
 			}
 		}
 #endif
+		if (bUseVirtualTexturing)
+		{
+			RDG_GPU_STAT_SCOPE(GraphBuilder, VirtualTextureUpdate);
+			VirtualTextureFeedbackEnd(GraphBuilder);
+		}
 
 		RendererModule.RenderOverlayExtensions(GraphBuilder, Views, SceneTextures);
 
