LFV - Started adding tile culling

#prelfight shader

diff --git a/Engine/Shaders/Private/LocalFogVolumes/LocalFogVolumeTiledCulling.usf b/Engine/Shaders/Private/LocalFogVolumes/LocalFogVolumeTiledCulling.usf
--- a/Engine/Shaders/Private/LocalFogVolumes/LocalFogVolumeTiledCulling.usf
+++ b/Engine/Shaders/Private/LocalFogVolumes/LocalFogVolumeTiledCulling.usf
@@ -2,7 +2,13 @@
 #include "../Common.ush"
 
 #include "LocalFogVolumeCommon.ush"
+#include "../Nanite/NaniteHZBCull.ush"
 
+float4 LeftPlane;
+float4 RightPlane;
+float4 TopPlane;
+float4 BottomPlane;
+float4 NearPlane;
 
 RWTexture2DArray<uint> LocalFogVolumeTileDataTextureUAV;
 
@@ -20,7 +26,49 @@ void LocalFogVolumeTiledCullingCS(uint3 ThreadId : SV_DispatchThreadID)
 		return; // Skip out of bound texture tile data 
 	}
 
-	LocalFogVolumeTileDataTextureUAV[uint3(TilePos, 0)] = TilePos.x + TilePos.y; // DEBUG
+
+	float2 MinTileLerp =  float2(TilePos)       / float2(LFVTileDataResolution);
+	float2 MaxTileLerp = (float2(TilePos)+1.0f) / float2(LFVTileDataResolution);
+
+	float4 TilePlanes[5];
+	TilePlanes[0] = lerp(LeftPlane, RightPlane, MinTileLerp.xxxx);		// Left tile plane
+	TilePlanes[1] = lerp(LeftPlane, RightPlane, MaxTileLerp.xxxx);		// Right tile plane
+	TilePlanes[2] = lerp(-TopPlane, -BottomPlane, MinTileLerp.yyyy);	// Bottom tile plane
+	TilePlanes[3] = lerp(-TopPlane, -BottomPlane, MaxTileLerp.yyyy);	// Top plane
+//	TilePlanes[4] = NearPlane;		LFV_TODO							// Near plane
+	// Swap right and top plane to go from lerpable planes to planes we can use for culling
+	TilePlanes[1] *= -1.0f;
+	TilePlanes[3] *= -1.0f;
+
+#if 0
+	// Optional plane normalise
+	UNROLL
+	for (uint i = 0; i < 4; ++i)
+	{
+		float SquareSum = TilePlanes[i].x * TilePlanes[i].x + TilePlanes[i].y * TilePlanes[i].y + TilePlanes[i].z * TilePlanes[i].z;
+		float Scale = 1.0 / sqrt(SquareSum);
+		TilePlanes[i] *= Scale;
+	}
+#endif
+
+	// Single sphere culling proof of concept
+	float4 Sphere = float4(0.0, 0.0, 0.0, 500.0);
+
+	uint LFVCount = 0;
+	bool bCompletelyOutside = false;
+	UNROLL
+	for (uint i = 0; i < 4; ++i)
+	{
+		if (dot(TilePlanes[i].xyz, Sphere.xyz) - TilePlanes[i].w > Sphere.w)
+		{
+			bCompletelyOutside = true;
+			break;
+		}
+	}
+
+	LFVCount += bCompletelyOutside ? 0 : 1;
+
+	LocalFogVolumeTileDataTextureUAV[uint3(TilePos, 0)] = LFVCount * 20; // DEBUG
 }
 
 
diff --git a/Engine/Source/Runtime/Renderer/Private/LocalFogVolumeRendering.cpp b/Engine/Source/Runtime/Renderer/Private/LocalFogVolumeRendering.cpp
--- a/Engine/Source/Runtime/Renderer/Private/LocalFogVolumeRendering.cpp
+++ b/Engine/Source/Runtime/Renderer/Private/LocalFogVolumeRendering.cpp
@@ -148,6 +148,11 @@ public:
 		SHADER_PARAMETER_STRUCT_REF(FViewUniformShaderParameters, View)
 		SHADER_PARAMETER_STRUCT(FLocalFogVolumeCommonParameters, LFV)
 		SHADER_PARAMETER_RDG_TEXTURE_UAV(RWTexture2DArray<uint>, LocalFogVolumeTileDataTextureUAV)
+		SHADER_PARAMETER(FVector4f, LeftPlane)
+		SHADER_PARAMETER(FVector4f, RightPlane)
+		SHADER_PARAMETER(FVector4f, TopPlane)
+		SHADER_PARAMETER(FVector4f, BottomPlane)
+		SHADER_PARAMETER(FVector4f, NearPlane)
 	END_SHADER_PARAMETER_STRUCT()
 
 	static void ModifyCompilationEnvironment(const FGlobalShaderPermutationParameters& Parameters, FShaderCompilerEnvironment& OutEnvironment)
@@ -167,6 +172,28 @@ static void LocalFogVolumeViewTiledCullingPass(FViewInfo& View, FRDGBuilder& Gra
 	PassParameters->LFV = View.LocalFogVolumeViewData.UniformParametersStruct.LocalFogVolumeCommon;
 	PassParameters->LocalFogVolumeTileDataTextureUAV = View.LocalFogVolumeViewData.TileDataTextureArrayUAV;
 
+	auto ConvertPlanToVector4f = [&](FVector4f& OutVec4f, auto& Plane, bool bFlipPlane)
+	{
+		OutVec4f.X = Plane.X;
+		OutVec4f.Y = Plane.Y;
+		OutVec4f.Z = Plane.Z;
+		OutVec4f.W = Plane.W;
+		if (bFlipPlane)
+		{
+			// We swap some of the planes normal so that they are lerpable while avoiding potential null normal and precision issue at the middle of the frustum.
+			OutVec4f.X *= -1.0f;
+			OutVec4f.Y *= -1.0f;
+			OutVec4f.Z *= -1.0f;
+			OutVec4f.W *= -1.0f;
+		}
+	};
+	// Using world space plane for now. LFV_TODO: do computation in view space.
+	ConvertPlanToVector4f(PassParameters->LeftPlane,	View.CullingFrustum.Planes[0], false);
+	ConvertPlanToVector4f(PassParameters->RightPlane,	View.CullingFrustum.Planes[1], true);
+	ConvertPlanToVector4f(PassParameters->TopPlane,		View.CullingFrustum.Planes[2], true);
+	ConvertPlanToVector4f(PassParameters->BottomPlane,	View.CullingFrustum.Planes[3], false);
+//	ConvertPlanToVector4f(PassParameters->NearPlane,	View.CullingFrustum.Planes[4], false); // LFV TODO View.ViewMatrices.GetViewProjectionMatrix().GetNearPlane
+
 	ERDGPassFlags PassFlag = ERDGPassFlags::Compute; // LFV_TODO try ERDGPassFlags::AsyncCompute later
 
 	TileDataTextureSize.Z = 1;
