LFV - fixed out of bound access causing crash due to bad aloloation size.

diff --git a/Engine/Source/Runtime/Renderer/Private/LocalFogVolumeRendering.cpp b/Engine/Source/Runtime/Renderer/Private/LocalFogVolumeRendering.cpp
--- a/Engine/Source/Runtime/Renderer/Private/LocalFogVolumeRendering.cpp
+++ b/Engine/Source/Runtime/Renderer/Private/LocalFogVolumeRendering.cpp
@@ -216,7 +216,7 @@ void GetLocalFogVolumeSortingData(const FScene* Scene, FRDGBuilder& GraphBuilder
 	Out.LocalFogVolumeInstanceCount = Scene->LocalFogVolumes.Num();
 	Out.LocalFogVolumeInstanceCountFinal = 0;
 	Out.LocalFogVolumeGPUInstanceData = (FLocalFogVolumeGPUInstanceData*)GraphBuilder.Alloc(sizeof(FLocalFogVolumeGPUInstanceData) * Out.LocalFogVolumeInstanceCount, 16);
-	Out.LocalFogVolumeCenterPos = (FVector*)GraphBuilder.Alloc(sizeof(FVector3f) * Out.LocalFogVolumeInstanceCount, 16);
+	Out.LocalFogVolumeCenterPos = (FVector*)GraphBuilder.Alloc(sizeof(FVector) * Out.LocalFogVolumeInstanceCount, 16);
 	Out.LocalFogVolumeSortKeys.SetNumUninitialized(Out.LocalFogVolumeInstanceCount);
 	for (FLocalFogVolumeSceneProxy* LHF : Scene->LocalFogVolumes)
 	{
