Option to scale LOD distance factor in shadows with cascade index
- new cvar r.Shadow.LODDistanceFactor.CascadeScale

diff --git a/Engine/Source/Runtime/Renderer/Private/ShadowDepthRendering.cpp b/Engine/Source/Runtime/Renderer/Private/ShadowDepthRendering.cpp
--- a/Engine/Source/Runtime/Renderer/Private/ShadowDepthRendering.cpp
+++ b/Engine/Source/Runtime/Renderer/Private/ShadowDepthRendering.cpp
@@ -55,7 +55,14 @@ static float GShadowLODDistanceFactor = 1;
 static FAutoConsoleVariableRef CVarShadowScreenMultiple(
 	TEXT("r.Shadow.LODDistanceFactor"),
 	GShadowLODDistanceFactor,
-	TEXT("Multiplier for LOD selection distance when rendering regular shadows"),
+	TEXT("Multiplier for LOD selection distance when rendering regular shadows (global)"),
+	ECVF_RenderThreadSafe);
+
+static float GShadowLODDistanceFactorCascadeScale = 0.0f;
+static FAutoConsoleVariableRef CVarShadowScreenMultipleCascadeScale(
+	TEXT("r.Shadow.LODDistanceFactor.CascadeScale"),
+	GShadowLODDistanceFactorCascadeScale,
+	TEXT("Multiplier for LOD selection distance when rendering regular shadows (scales with cascade index)"),
 	ECVF_RenderThreadSafe);
 
 static float GFarShadowLODDistanceFactor = 1;
@@ -1204,7 +1211,7 @@ void FProjectedShadowInfo::RenderDepth(
 
 float FProjectedShadowInfo::GetLODDistanceFactor() const
 {
-	return CascadeSettings.bFarShadowCascade ? GFarShadowLODDistanceFactor : GShadowLODDistanceFactor;
+	return CascadeSettings.bFarShadowCascade ? GFarShadowLODDistanceFactor : GShadowLODDistanceFactor + CascadeSettings.ShadowSplitIndex * GShadowLODDistanceFactorCascadeScale;
 }
 
 void FProjectedShadowInfo::ModifyViewForShadow(FViewInfo* FoundView) const
