Take dynamic primitive GPU Scene instance data offset into account when computing instance IDs for occlusion query rendering. Ignore mesh draw commands that use instance runs mechanism instead of a single contiguous instance range.

diff --git a/Engine/Source/Runtime/Renderer/Private/InstanceCulling/InstanceCullingOcclusionQuery.cpp b/Engine/Source/Runtime/Renderer/Private/InstanceCulling/InstanceCullingOcclusionQuery.cpp
--- a/Engine/Source/Runtime/Renderer/Private/InstanceCulling/InstanceCullingOcclusionQuery.cpp
+++ b/Engine/Source/Runtime/Renderer/Private/InstanceCulling/InstanceCullingOcclusionQuery.cpp
@@ -320,8 +320,10 @@ struct FInstanceCullingOcclusionQueryDeferredContext
 
 		// NumPrimitives is 0 if mesh draw command uses IndirectArgs
 		// This path is currently not implemented/supported by oclcusion query culling.
+		// Commands that use instance runs are currently not supported.
 		return bCompatibleFlags
-			&& VisibleCommand.MeshDrawCommand->NumPrimitives != 0;
+			&& VisibleCommand.MeshDrawCommand->NumPrimitives != 0
+			&& VisibleCommand.NumRuns == 0;
 	};
 
 	void Execute()
@@ -365,7 +367,9 @@ struct FInstanceCullingOcclusionQueryDeferredContext
 			return;
 		}
 
-		FillVisibleInstanceIds(VisibleMeshDrawCommands);
+		const uint32 DynamicPrimitiveInstanceOffset = View->DynamicPrimitiveCollector.GetInstanceSceneDataOffset();
+
+		FillVisibleInstanceIds(VisibleMeshDrawCommands, DynamicPrimitiveInstanceOffset);
 
 		bValid = true;
 	}
@@ -388,7 +392,7 @@ struct FInstanceCullingOcclusionQueryDeferredContext
 		return Result;
 	}
 
-	void FillVisibleInstanceIds(const FMeshCommandOneFrameArray& VisibleMeshDrawCommands)
+	void FillVisibleInstanceIds(const FMeshCommandOneFrameArray& VisibleMeshDrawCommands, const uint32 DynamicPrimitiveInstanceOffset)
 	{
 		TRACE_CPUPROFILER_EVENT_SCOPE(FInstanceCullingOcclusionQueryDeferredContext::FillVisibleInstanceIds);
 
@@ -407,6 +411,11 @@ struct FInstanceCullingOcclusionQueryDeferredContext
 			}
 
 			uint32 InstanceBaseIndex = VisibleCommand.PrimitiveIdInfo.InstanceSceneDataOffset;
+			if (VisibleCommand.PrimitiveIdInfo.bIsDynamicPrimitive)
+			{
+				InstanceBaseIndex += DynamicPrimitiveInstanceOffset;
+			}
+
 			uint32 CommandNumInstances = VisibleCommand.MeshDrawCommand->NumInstances;
 
 			check(InstanceBaseIndex + CommandNumInstances <= uint32(NumGPUSceneInstances));
