Fixed crash in hair system due to destroying primitive render resources outside of the scene update.

diff --git a/Engine/Source/Runtime/Renderer/Private/RendererScene.cpp b/Engine/Source/Runtime/Renderer/Private/RendererScene.cpp
--- a/Engine/Source/Runtime/Renderer/Private/RendererScene.cpp
+++ b/Engine/Source/Runtime/Renderer/Private/RendererScene.cpp
@@ -5439,7 +5439,7 @@ void FScene::UpdateAllPrimitiveSceneInfos(FRDGBuilder& GraphBuilder, EUpdateAllP
 	SceneUpdateChangeSetStorage.RemovedPrimitiveIds.Reserve(RemovedPrimitiveSceneInfos.Num());
 	SceneUpdateChangeSetStorage.RemovedPrimitiveSceneInfos.Reserve(RemovedPrimitiveSceneInfos.Num());
 
-	TArray<FPrimitiveSceneInfo*> RemovedLocalPrimitiveSceneInfos;
+	TArray<FPrimitiveSceneInfo*, SceneRenderingAllocator> RemovedLocalPrimitiveSceneInfos;
 	RemovedLocalPrimitiveSceneInfos.Reserve(RemovedPrimitiveSceneInfos.Num());
 	for (FPrimitiveSceneInfo* PrimitiveSceneInfo : RemovedPrimitiveSceneInfos)
 	{
@@ -5449,7 +5449,16 @@ void FScene::UpdateAllPrimitiveSceneInfos(FRDGBuilder& GraphBuilder, EUpdateAllP
 		SceneUpdateChangeSetStorage.RemovedPrimitiveSceneInfos.Add(PrimitiveSceneInfo);
 	}
 
-	TArray<FPrimitiveSceneInfo*> AddedLocalPrimitiveSceneInfos;
+	TArray<FPrimitiveSceneInfo*, SceneRenderingAllocator> DestroyResourcesPrimitives;
+	DestroyResourcesPrimitives.Reserve(RemovedLocalPrimitiveSceneInfos.Num() + DeletedPrimitiveSceneInfos.Num());
+	DestroyResourcesPrimitives = RemovedLocalPrimitiveSceneInfos;
+
+	for (FPrimitiveSceneInfo* Primitive : DeletedPrimitiveSceneInfos)
+	{
+		DestroyResourcesPrimitives.Add(Primitive);
+	}
+
+	TArray<FPrimitiveSceneInfo*, SceneRenderingAllocator> AddedLocalPrimitiveSceneInfos;
 	AddedLocalPrimitiveSceneInfos.Reserve(AddedPrimitiveSceneInfos.Num());
 	for (FPrimitiveSceneInfo* SceneInfo : AddedPrimitiveSceneInfos)
 	{
@@ -5460,11 +5469,11 @@ void FScene::UpdateAllPrimitiveSceneInfos(FRDGBuilder& GraphBuilder, EUpdateAllP
 
 	if (!AddedLocalPrimitiveSceneInfos.IsEmpty() || !RemovedLocalPrimitiveSceneInfos.IsEmpty())
 	{
-		ProcessPrimitiveResourcesTask.AddPrerequisites(GraphBuilder.AddCommandListSetupTask([AddedPrimitives = AddedLocalPrimitiveSceneInfos, RemovedPrimitives = RemovedLocalPrimitiveSceneInfos](FRHICommandListBase& RHICmdList)
+		ProcessPrimitiveResourcesTask.AddPrerequisites(GraphBuilder.AddCommandListSetupTask([CreateResourcesPrimitives = AddedLocalPrimitiveSceneInfos, DestroyResourcesPrimitives = MoveTemp(DestroyResourcesPrimitives)](FRHICommandListBase& RHICmdList)
 		{
 			{
 				TRACE_CPUPROFILER_EVENT_SCOPE(Scene::DestroyPrimitiveResources);
-				for (FPrimitiveSceneInfo* Primitive : RemovedPrimitives)
+				for (FPrimitiveSceneInfo* Primitive : DestroyResourcesPrimitives)
 				{
 					if (!Primitive->Proxy->ShouldConstrainToRenderThread())
 					{
@@ -5475,7 +5484,7 @@ void FScene::UpdateAllPrimitiveSceneInfos(FRDGBuilder& GraphBuilder, EUpdateAllP
 
 			{
 				TRACE_CPUPROFILER_EVENT_SCOPE(Scene::CreatePrimitiveResources);
-				for (FPrimitiveSceneInfo* Primitive : AddedPrimitives)
+				for (FPrimitiveSceneInfo* Primitive : CreateResourcesPrimitives)
 				{
 					if (!Primitive->Proxy->ShouldConstrainToRenderThread())
 					{
@@ -6579,11 +6588,6 @@ void FScene::UpdateAllPrimitiveSceneInfos(FRDGBuilder& GraphBuilder, EUpdateAllP
 		SCOPED_NAMED_EVENT(FScene_DeletePrimitiveSceneInfo, FColor::Red);
 		for (FPrimitiveSceneInfo* PrimitiveSceneInfo : DeletedPrimitiveSceneInfos)
 		{
-			if (!PrimitiveSceneInfo->Proxy->ShouldConstrainToRenderThread())
-			{
-				PrimitiveSceneInfo->Proxy->DestroyRenderThreadResources();
-			}
-
 			// It is possible that the HitProxies list isn't empty if PrimitiveSceneInfo was Added/Removed in same frame
 			// Delete the PrimitiveSceneInfo on the game thread after the rendering thread has processed its removal.
 			// This must be done on the game thread because the hit proxy references (and possibly other members) need to be freed on the game thread.
