LFV - tiled culling and rendering of local fog volume on translucent.
Also contains a fix for out ot bound memory access/write.

diff --git a/Engine/Shaders/Private/LocalFogVolumes/LocalFogVolumeSplat.usf b/Engine/Shaders/Private/LocalFogVolumes/LocalFogVolumeSplat.usf
--- a/Engine/Shaders/Private/LocalFogVolumes/LocalFogVolumeSplat.usf
+++ b/Engine/Shaders/Private/LocalFogVolumes/LocalFogVolumeSplat.usf
@@ -145,12 +145,45 @@ void LocalFogVolumeTiledPS(
 {
 	ResolvedView = ResolveView();
 
-	uint2 TilePos = uint2(OtherVertexOutput.TileData >> 16, OtherVertexOutput.TileData & 0xFFFF);
+	float3 CamRayTranslatedWorldOrigin = PrimaryView.TranslatedWorldCameraOrigin;
+	float3 CamRayTranslatedWorldDir = normalize(SvPositionToTranslatedWorld(float4(SVPos.xy, 0.5, 1.0)));
+	float3 CamRayWorldDir = CamRayTranslatedWorldDir;
+	float3 CamRayWorldOrigin = CamRayTranslatedWorldOrigin - LWCHackToFloat(PrimaryView.PreViewTranslation);
+
+#if SHADING_PATH_MOBILE
+	// LookupDeviceZ with uint2 SVPos.xy as input does not work on mobile.
+	float2 UV = SVPos * View.BufferSizeAndInvSize.zw;
+	const float DeviceZ = LookupDeviceZ(UV);
+#else
+	const float DeviceZ = LookupDeviceZ(uint2(SVPos.xy));
+#endif
+
+	float3 DepthBufferTranslatedWorldPos = SvPositionToTranslatedWorld(float4(SVPos.xy, DeviceZ, 1.0));
 
+	uint2 TilePos = uint2(OtherVertexOutput.TileData >> 16, OtherVertexOutput.TileData & 0xFFFF);
 	uint LFVCount = LFVTileDataTex[uint3(TilePos,0)];
+	float4 LuminanceColorTransmittance = float4(0.0, 0.0, 0.0, 1.0);
+	for (int LFVIndex = 0; LFVIndex < LFVCount; ++LFVIndex)
+	{
+		uint InstanceIndex = LFVTileDataTex[uint3(TilePos, 1 + LFVIndex)];
+
+		float4 LFVLuminanceColorTransmittance = GetLFVInstanceContribution(
+			InstanceIndex,
+			CamRayWorldOrigin,
+			CamRayWorldDir,
+			CamRayTranslatedWorldOrigin,
+			CamRayTranslatedWorldDir,
+			DepthBufferTranslatedWorldPos);
+
+		LuminanceColorTransmittance = float4(LuminanceColorTransmittance.rgb * LFVLuminanceColorTransmittance.a + LFVLuminanceColorTransmittance.rgb, LuminanceColorTransmittance.a * LFVLuminanceColorTransmittance.a);
+	}
+	
+	OutColor = float4(
+		LuminanceColorTransmittance.rgb * PrimaryView.PreExposure,
+		LuminanceColorTransmittance.a);
 
-//	OutColor = float4(float(TilePos.x) * 0.1, float(TilePos.y) * 0.1, 0.0, 0.5);
-	OutColor = float4(float(LFVCount) * 0.1, 0.0, 0.0, 0.5);
+	// Debug
+	OutColor += float4(float(LFVCount) * 0.1 * PrimaryView.PreExposure, 0.0, 0.0, 0.0);
 }
 #endif // LFV_TILED_PS
 
diff --git a/Engine/Shaders/Private/LocalFogVolumes/LocalFogVolumeTiledCulling.usf b/Engine/Shaders/Private/LocalFogVolumes/LocalFogVolumeTiledCulling.usf
--- a/Engine/Shaders/Private/LocalFogVolumes/LocalFogVolumeTiledCulling.usf
+++ b/Engine/Shaders/Private/LocalFogVolumes/LocalFogVolumeTiledCulling.usf
@@ -10,6 +10,8 @@ float4 TopPlane;
 float4 BottomPlane;
 float4 NearPlane;
 
+Buffer<float4> LocalFogVolumeCullingDataBuffer;
+
 RWTexture2DArray<uint> LocalFogVolumeTileDataTextureUAV;
 
 #ifndef THREADGROUP_SIZE
@@ -51,24 +53,29 @@ void LocalFogVolumeTiledCullingCS(uint3 ThreadId : SV_DispatchThreadID)
 	}
 #endif
 
-	// Single sphere culling proof of concept
-	float4 Sphere = float4(0.0, 0.0, 0.0, 500.0);
-
 	uint LFVCount = 0;
-	bool bCompletelyOutside = false;
-	UNROLL
-	for (uint i = 0; i < 4; ++i)
+	for (int LFVIndex = 0; LFVIndex < LFVInstcount; ++LFVIndex)
 	{
-		if (dot(TilePlanes[i].xyz, Sphere.xyz) - TilePlanes[i].w > Sphere.w)
+		float4 Sphere = LocalFogVolumeCullingDataBuffer[LFVIndex];
+
+		bool bCompletelyOutside = false;
+		UNROLL
+		for (uint i = 0; i < 4; ++i)
+		{
+			if (dot(TilePlanes[i].xyz, Sphere.xyz) - TilePlanes[i].w > Sphere.w)
+			{
+				bCompletelyOutside = true;
+				break;
+			}
+		}
+		if (!bCompletelyOutside)
 		{
-			bCompletelyOutside = true;
-			break;
+			LocalFogVolumeTileDataTextureUAV[uint3(TilePos, 1 + LFVCount)] = LFVIndex;
+			LFVCount++;
 		}
 	}
 
-	LFVCount += bCompletelyOutside ? 0 : 1;
-
-	LocalFogVolumeTileDataTextureUAV[uint3(TilePos, 0)] = LFVCount * 20; // DEBUG
+	LocalFogVolumeTileDataTextureUAV[uint3(TilePos, 0)] = LFVCount;
 }
 
 
diff --git a/Engine/Source/Runtime/Renderer/Private/LocalFogVolumeRendering.cpp b/Engine/Source/Runtime/Renderer/Private/LocalFogVolumeRendering.cpp
--- a/Engine/Source/Runtime/Renderer/Private/LocalFogVolumeRendering.cpp
+++ b/Engine/Source/Runtime/Renderer/Private/LocalFogVolumeRendering.cpp
@@ -76,12 +76,16 @@ void SetDummyLocalFogVolumeForView(FRDGBuilder& GraphBuilder, FViewInfo& View)
 
 	View.LocalFogVolumeViewData.GPUInstanceCount			= 0;
 
-	static FLocalFogVolumeGPUInstanceData DummyData;
+	static FLocalFogVolumeGPUInstanceData DummyData;	// Static data so ERDGInitialDataFlags::NoCopy can be used. 
 	View.LocalFogVolumeViewData.GPUInstanceDataBuffer		= CreateVertexBuffer(GraphBuilder, TEXT("DUMMYLocalFogVolumeGPUInstanceDataBuffer"),
 		FRDGBufferDesc::CreateBufferDesc(SizeOfUintVec4, UintVec4CountInLocalFogVolumeGPUInstanceData), &DummyData, sizeof(FLocalFogVolumeGPUInstanceData) * 1, ERDGInitialDataFlags::NoCopy);
-
 	View.LocalFogVolumeViewData.GPUInstanceDataBufferSRV	= GraphBuilder.CreateSRV(View.LocalFogVolumeViewData.GPUInstanceDataBuffer, PF_A32B32G32R32F);
 
+	static FVector4f DummyCullingData(EForceInit::ForceInitToZero);
+	View.LocalFogVolumeViewData.GPUInstanceCullingDataBuffer	= CreateVertexBuffer(GraphBuilder, TEXT("DUMMYLocalFogVolumeGPUInstanceCullingDataBuffer"),
+		FRDGBufferDesc::CreateBufferDesc(sizeof(FVector4f), 1), &DummyCullingData, sizeof(FVector4f) * 1, ERDGInitialDataFlags::NoCopy);
+	View.LocalFogVolumeViewData.GPUInstanceCullingDataBufferSRV = GraphBuilder.CreateSRV(View.LocalFogVolumeViewData.GPUInstanceCullingDataBuffer, PF_A32B32G32R32F);
+
 	View.LocalFogVolumeViewData.TileDataTextureArray		= GSystemTextures.GetZeroUIntArrayDummy(GraphBuilder);
 	View.LocalFogVolumeViewData.TileDataTextureArraySRV		= GraphBuilder.CreateSRV(View.LocalFogVolumeViewData.TileDataTextureArray);
 	View.LocalFogVolumeViewData.TileDataTextureArrayUAV		= nullptr;	// Should never be written by culling passes if there are no instances.
@@ -148,6 +152,7 @@ public:
 		SHADER_PARAMETER_STRUCT_REF(FViewUniformShaderParameters, View)
 		SHADER_PARAMETER_STRUCT(FLocalFogVolumeCommonParameters, LFV)
 		SHADER_PARAMETER_RDG_TEXTURE_UAV(RWTexture2DArray<uint>, LocalFogVolumeTileDataTextureUAV)
+		SHADER_PARAMETER_RDG_BUFFER_SRV(Buffer<float4>, LocalFogVolumeCullingDataBuffer)
 		SHADER_PARAMETER(FVector4f, LeftPlane)
 		SHADER_PARAMETER(FVector4f, RightPlane)
 		SHADER_PARAMETER(FVector4f, TopPlane)
@@ -170,6 +175,7 @@ static void LocalFogVolumeViewTiledCullingPass(FViewInfo& View, FRDGBuilder& Gra
 	FLocalFogVolumeTiledCullingCS::FParameters* PassParameters = GraphBuilder.AllocParameters<FLocalFogVolumeTiledCullingCS::FParameters>();
 	PassParameters->View = View.ViewUniformBuffer;
 	PassParameters->LFV = View.LocalFogVolumeViewData.UniformParametersStruct.LocalFogVolumeCommon;
+	PassParameters->LocalFogVolumeCullingDataBuffer = View.LocalFogVolumeViewData.GPUInstanceCullingDataBufferSRV;
 	PassParameters->LocalFogVolumeTileDataTextureUAV = View.LocalFogVolumeViewData.TileDataTextureArrayUAV;
 
 	auto ConvertPlanToVector4f = [&](FVector4f& OutVec4f, auto& Plane, bool bFlipPlane)
@@ -283,10 +289,17 @@ void CreateViewLocalFogVolumeBufferSRV(FViewInfo& View, FRDGBuilder& GraphBuilde
 
 	// 2. Create the buffer containing all the fog volume data instance sorted according to their key for the current view.
 	FLocalFogVolumeGPUInstanceData* LocalFogVolumeGPUSortedInstanceData = (FLocalFogVolumeGPUInstanceData*)GraphBuilder.Alloc(sizeof(FLocalFogVolumeGPUInstanceData) * SortingData.LocalFogVolumeInstanceCountFinal, 16);
+	FVector4f* LocalFogVolumeGPUSortedInstanceCullingData = (FVector4f*)GraphBuilder.Alloc(sizeof(FVector4f) * SortingData.LocalFogVolumeInstanceCountFinal, 16);
 	for (uint32 i = 0; i < SortingData.LocalFogVolumeInstanceCountFinal; i++)
 	{
+		FLocalFogVolumeSortKey LFVKey = SortingData.LocalFogVolumeSortKeys[i + DiscardedOffset];
+
 		// We could also have an indirection buffer on GPU but choosing to go with the sorting + copy on CPU since it is expected to not have many local height fog volumes.
-		LocalFogVolumeGPUSortedInstanceData[i] = SortingData.LocalFogVolumeGPUInstanceData[SortingData.LocalFogVolumeSortKeys[i + DiscardedOffset].FogVolume.Index];
+		LocalFogVolumeGPUSortedInstanceData[i] = SortingData.LocalFogVolumeGPUInstanceData[LFVKey.FogVolume.Index];
+
+		const float LFVMaximumAxisScale = LocalFogVolumeGPUSortedInstanceData[i].Transform.GetMaximumAxisScale();
+		FVector& LFVPosition = SortingData.LocalFogVolumeCenterPos[LFVKey.FogVolume.Index];
+		LocalFogVolumeGPUSortedInstanceCullingData[i] = FVector4f(LFVPosition.X, LFVPosition.Y, LFVPosition.Z, LFVMaximumAxisScale);
 	}
 
 	// 3. Allocate buffer and initialize with sorted data to upload to GPU
@@ -299,6 +312,12 @@ void CreateViewLocalFogVolumeBufferSRV(FViewInfo& View, FRDGBuilder& GraphBuilde
 		LocalFogVolumeGPUSortedInstanceData, AllLocalFogVolumeInstanceBytesFinal, ERDGInitialDataFlags::NoCopy);
 	View.LocalFogVolumeViewData.GPUInstanceDataBufferSRV	= GraphBuilder.CreateSRV(View.LocalFogVolumeViewData.GPUInstanceDataBuffer, PF_A32B32G32R32F);
 
+	View.LocalFogVolumeViewData.GPUInstanceCullingDataBuffer = CreateVertexBuffer(
+		GraphBuilder, TEXT("LocalFogVolume.GPUInstanceCullingDataBuffer"),
+		FRDGBufferDesc::CreateBufferDesc(sizeof(FVector4f), SortingData.LocalFogVolumeInstanceCountFinal * sizeof(FVector4f)),
+		LocalFogVolumeGPUSortedInstanceCullingData, SortingData.LocalFogVolumeInstanceCountFinal * sizeof(FVector4f), ERDGInitialDataFlags::NoCopy);
+	View.LocalFogVolumeViewData.GPUInstanceCullingDataBufferSRV = GraphBuilder.CreateSRV(View.LocalFogVolumeViewData.GPUInstanceCullingDataBuffer, PF_A32B32G32R32F); // LFV_TODO use byte buffer to leverage scalar pipe in the culling compute shader.
+
 	// Create the texture that will contain the tiled culled result: count in the first slice and indices in the remaining slices
 	const uint32 LocalFogVolumeTilePixelSize				= GetLocalFogVolumeTilePixelSize();
 
diff --git a/Engine/Source/Runtime/Renderer/Private/LocalFogVolumeRendering.h b/Engine/Source/Runtime/Renderer/Private/LocalFogVolumeRendering.h
--- a/Engine/Source/Runtime/Renderer/Private/LocalFogVolumeRendering.h
+++ b/Engine/Source/Runtime/Renderer/Private/LocalFogVolumeRendering.h
@@ -96,6 +96,8 @@ struct FLocalFogVolumeViewData
 	uint32 GPUInstanceCount = 0;
 	FRDGBufferRef GPUInstanceDataBuffer = nullptr;
 	FRDGBufferSRVRef GPUInstanceDataBufferSRV = nullptr;
+	FRDGBufferRef GPUInstanceCullingDataBuffer = nullptr;
+	FRDGBufferSRVRef GPUInstanceCullingDataBufferSRV = nullptr;
 	TRDGUniformBufferRef<FLocalFogVolumeUniformParameters> UniformBuffer = nullptr;
 
 	FLocalFogVolumeUniformParameters UniformParametersStruct;
