Fix for cvar potential access from non rendering thread when creating the passprocessor.

diff --git a/Engine/Source/Runtime/Renderer/Private/DistortionRendering.cpp b/Engine/Source/Runtime/Renderer/Private/DistortionRendering.cpp
--- a/Engine/Source/Runtime/Renderer/Private/DistortionRendering.cpp
+++ b/Engine/Source/Runtime/Renderer/Private/DistortionRendering.cpp
@@ -113,7 +113,7 @@ TRDGUniformBufferRef<FDistortionPassUniformParameters> CreateDistortionPassUnifo
 
 static bool GetUseRoughRefraction()
 {
-	return Substrate::IsSubstrateEnabled() && CVarRefractionBlur.GetValueOnRenderThread() > 0;
+	return Substrate::IsSubstrateEnabled() && CVarRefractionBlur.GetValueOnAnyThread() > 0; // Any thread since it can be called when creating the PassProcessor
 }
 
 class FDistortionScreenPS : public FGlobalShader
