LFV - added tile debug view for.

diff --git a/Engine/Shaders/Private/LocalFogVolumes/LocalFogVolumeSplat.usf b/Engine/Shaders/Private/LocalFogVolumes/LocalFogVolumeSplat.usf
--- a/Engine/Shaders/Private/LocalFogVolumes/LocalFogVolumeSplat.usf
+++ b/Engine/Shaders/Private/LocalFogVolumes/LocalFogVolumeSplat.usf
@@ -138,6 +138,12 @@ void LocalFogVolumeTiledVS(
 
 
 #if LFV_TILED_PS
+int LocalFogVolumeTileDebug;
+
+#if USE_DEVELOPMENT_SHADERS
+#include "../ColorMap.ush"
+#endif
+
 void LocalFogVolumeTiledPS(
 	in float4 SVPos : SV_POSITION,
 	in FLocalFogVolumeVertexOutput OtherVertexOutput,
@@ -183,13 +189,26 @@ void LocalFogVolumeTiledPS(
 		LuminanceColorTransmittance.rgb * PrimaryView.PreExposure,
 		LuminanceColorTransmittance.a);
 
-	if (all(DefaultColorTransmittance == OutColor))
+	bool bDoClip = true;
+#if USE_DEVELOPMENT_SHADERS
+	bDoClip = LocalFogVolumeTileDebug == 2;
+#endif
+
+	if (all(DefaultColorTransmittance == OutColor) && bDoClip)
 	{
 		clip(-1.0f);
 	}
 
-	// Debug
-	OutColor += float4(float(LFVCount) * 0.1 * PrimaryView.PreExposure, 0.0, 0.0, 0.0);
+#if USE_DEVELOPMENT_SHADERS
+	if (LocalFogVolumeTileDebug > 0)
+	{
+		float DebugTransmittance = 0.5;
+		float3 DebugColor = GetHSVDebugColor(float(LFVCount)  / 4.0f);
+		DebugColor = DebugColor * (1.0f - DebugTransmittance) * PrimaryView.PreExposure;
+
+		OutColor = float4(OutColor.rgb * DebugTransmittance + DebugColor, OutColor.a * DebugTransmittance);
+	}
+#endif
 }
 #endif // LFV_TILED_PS
 
diff --git a/Engine/Source/Runtime/Renderer/Private/LocalFogVolumeRendering.cpp b/Engine/Source/Runtime/Renderer/Private/LocalFogVolumeRendering.cpp
--- a/Engine/Source/Runtime/Renderer/Private/LocalFogVolumeRendering.cpp
+++ b/Engine/Source/Runtime/Renderer/Private/LocalFogVolumeRendering.cpp
@@ -34,6 +34,11 @@ static TAutoConsoleVariable<int32> CVarLocalFogVolumeTileMaxInstanceCount(
 	TEXT("Maximum number of local fog volume to account for per view (and per tile or consistency).\n"),
 	ECVF_RenderThreadSafe);
 
+static TAutoConsoleVariable<int32> CVarLocalFogVolumeTileDebug(
+	TEXT("r.LocalFogVolume.TileDebug"), 0,
+	TEXT("Debug the tiled rendering data complexity.\n"),
+	ECVF_RenderThreadSafe);
+
 // Example of tile setup
 //  - 1920x1080 => 15x9 tiles
 //  - Allowing max 32 volumes at once => culling list buffer = 15 * 9 * 32 * 1 byte = 4320 bytes = 4.3KB
@@ -471,6 +476,7 @@ class FLocalFogVolumeTiledRenderPS : public FGlobalShader
 	BEGIN_SHADER_PARAMETER_STRUCT(FParameters, )
 		SHADER_PARAMETER_STRUCT_REF(FViewUniformShaderParameters, View)
 		SHADER_PARAMETER_STRUCT(FLocalFogVolumeUniformParameters, LFV)
+		SHADER_PARAMETER(int32, LocalFogVolumeTileDebug)
 	END_SHADER_PARAMETER_STRUCT()
 
 	using FPermutationDomain = TShaderPermutationDomain<>;
@@ -539,6 +545,7 @@ void RenderLocalFogVolume(
 
 				PassParameters->PS.View = GetShaderBinding(View.ViewUniformBuffer);
 				PassParameters->PS.LFV = View.LocalFogVolumeViewData.UniformParametersStruct;
+				PassParameters->PS.LocalFogVolumeTileDebug = FMath::Clamp(CVarLocalFogVolumeTileDebug.GetValueOnRenderThread(), 0, 2);
 
 				PassParameters->SceneTextures = SceneTextures.UniformBuffer;
 				PassParameters->RenderTargets[0] = FRenderTargetBinding(SceneColorTexture, ERenderTargetLoadAction::ENoAction);
