LFV - tiling is now achieved using  indirect draw.

#prelfight shader

diff --git a/Engine/Shaders/Private/LocalFogVolumes/LocalFogVolumeCommon.ush b/Engine/Shaders/Private/LocalFogVolumes/LocalFogVolumeCommon.ush
--- a/Engine/Shaders/Private/LocalFogVolumes/LocalFogVolumeCommon.ush
+++ b/Engine/Shaders/Private/LocalFogVolumes/LocalFogVolumeCommon.ush
@@ -8,8 +8,6 @@
 ///////////////////////////////////////////////////////////////////////////////
 // Local Fog Volume Data
 
-
-
 struct FLocalFogVolumeGPUInstanceData
 {
 	float4x4 Transform;
@@ -30,7 +28,6 @@ struct FLocalFogVolumeGPUInstanceData
 };
 
 
-
 // Structured buffers cannot be sampled from the vertex shader on some mobile platforms. 
 // That is why simple Buffers are used for LocalFogVolumeInstances, because this tech needs to run on mobile platforms.
 
@@ -62,7 +59,6 @@ Texture2DArray<uint>			LFV_LocalFogVolumeTileDataTexture;
 #endif
 
 
-
 FLocalFogVolumeGPUInstanceData GetLocalFogVolumeGPUInstanceData(uint Index)
 {
 	FLocalFogVolumeGPUInstanceData Data;
@@ -94,10 +90,27 @@ FLocalFogVolumeGPUInstanceData GetLocalFogVolumeGPUInstanceData(uint Index)
 
 
 ///////////////////////////////////////////////////////////////////////////////
-// Local Fog Volume Integrals
+// Local Fog Volume Tile Data
+
+uint2 LFVUnpackTile(uint In)
+{
+	return uint2(In & 0xFFFF, In >> 16);
+}
 
+uint LFVPackTile(uint2 TileCoord)
+{
+#if COMPILER_SUPPORTS_PACK_INTRINSICS
+	return PackUInt2ToUInt(TileCoord.x, TileCoord.y);
+#else
+	return TileCoord.x | (TileCoord.y << 16); // assumes 16bit is enough to represent a tiled resolution up to 65,535 :)
+#endif
+}
 
 
+
+///////////////////////////////////////////////////////////////////////////////
+// Local Fog Volume Integrals
+
 // CameraPos can be behind of RayStart if the camera is out of the volume
 struct FFogData
 {
diff --git a/Engine/Shaders/Private/LocalFogVolumes/LocalFogVolumeSplat.usf b/Engine/Shaders/Private/LocalFogVolumes/LocalFogVolumeSplat.usf
--- a/Engine/Shaders/Private/LocalFogVolumes/LocalFogVolumeSplat.usf
+++ b/Engine/Shaders/Private/LocalFogVolumes/LocalFogVolumeSplat.usf
@@ -113,6 +113,9 @@ void LocalFogVolumeSplatPS(
 
 
 #if LFV_TILED_VS
+
+Buffer<uint> TileDataBuffer;
+
 void LocalFogVolumeTiledVS(
 	in uint VertexIndex : SV_VertexID,
 	in uint InstanceIndex : SV_InstanceID,
@@ -122,9 +125,8 @@ void LocalFogVolumeTiledVS(
 {
 	ResolvedView = ResolveView();
 
-	uint2 TileCoord;
-	TileCoord.y = InstanceIndex / LFVTileDataResolution.x;
-	TileCoord.x = InstanceIndex - TileCoord.y * LFVTileDataResolution.x;
+	uint PackedTileCoord = TileDataBuffer[InstanceIndex];
+	uint2 TileCoord = LFVUnpackTile(PackedTileCoord);
 
 	uint2 TileVertex = TileCoord * LFVTilePixelSize;
 	TileVertex.x += VertexIndex == 1 || VertexIndex == 2 || VertexIndex == 4 ? LFVTilePixelSize : 0;
@@ -132,7 +134,7 @@ void LocalFogVolumeTiledVS(
 
 	OutPosition = float4(float2(TileVertex) * ResolvedView.ViewSizeAndInvSize.zw * float2(2.0f, -2.0f) + float2(-1.0, 1.0f), 0.5f, 1.0f);
 
-	OtherVertexOutput.TileData = (TileCoord.x << 16) | TileCoord.y;
+	OtherVertexOutput.TileData = PackedTileCoord;
 }
 #endif // LFV_TILED_VS
 
@@ -166,7 +168,7 @@ void LocalFogVolumeTiledPS(
 
 	float3 DepthBufferTranslatedWorldPos = SvPositionToTranslatedWorld(float4(SVPos.xy, DeviceZ, 1.0));
 
-	uint2 TilePos = uint2(OtherVertexOutput.TileData >> 16, OtherVertexOutput.TileData & 0xFFFF);
+	uint2 TilePos = LFVUnpackTile(OtherVertexOutput.TileData);
 	uint LFVCount = LFVTileDataTex[uint3(TilePos,0)];
 	const float4 DefaultColorTransmittance = float4(0.0, 0.0, 0.0, 1.0);
 	float4 LuminanceColorTransmittance = DefaultColorTransmittance;
diff --git a/Engine/Shaders/Private/LocalFogVolumes/LocalFogVolumeTiledCulling.usf b/Engine/Shaders/Private/LocalFogVolumes/LocalFogVolumeTiledCulling.usf
--- a/Engine/Shaders/Private/LocalFogVolumes/LocalFogVolumeTiledCulling.usf
+++ b/Engine/Shaders/Private/LocalFogVolumes/LocalFogVolumeTiledCulling.usf
@@ -14,23 +14,32 @@ Buffer<float4> LocalFogVolumeCullingDataBuffer;
 
 RWTexture2DArray<uint> LocalFogVolumeTileDataTextureUAV;
 
+RWBuffer<uint> TileDataBufferUAV;
+RWBuffer<uint> TileDrawIndirectBufferUAV;
+
 #ifndef THREADGROUP_SIZE
 #error THREADGROUP_SIZE must be defined
 #endif
 
 [numthreads(THREADGROUP_SIZE, THREADGROUP_SIZE, 1)]
-void LocalFogVolumeTiledCullingCS(uint3 ThreadId : SV_DispatchThreadID)
+void LocalFogVolumeTiledCullingCS(uint3 DispatchThreadId : SV_DispatchThreadID)
 {
-	uint2 TilePos = uint2(ThreadId.xy);
+	uint2 TileCoord = uint2(DispatchThreadId.xy);
 
-	if (any(TilePos >= LFVTileDataResolution))
+	if (any(TileCoord >= LFVTileDataResolution))
 	{
 		return; // Skip out of bound texture tile data 
 	}
 
+	// Init primitive index
+	if (DispatchThreadId.x == 0 && DispatchThreadId.y == 0)
+	{
+		const uint IndexCountPerInstance = 6; // LFV_TODO handle rect primitives with 4
+		TileDrawIndirectBufferUAV[0] = IndexCountPerInstance;
+	}
 
-	float2 MinTileLerp =  float2(TilePos)       / float2(LFVTileDataResolution);
-	float2 MaxTileLerp = (float2(TilePos)+1.0f) / float2(LFVTileDataResolution);
+	float2 MinTileLerp =  float2(TileCoord)       / float2(LFVTileDataResolution);
+	float2 MaxTileLerp = (float2(TileCoord)+1.0f) / float2(LFVTileDataResolution);
 
 	float4 TilePlanes[5];
 	TilePlanes[0] = lerp(LeftPlane, RightPlane, MinTileLerp.xxxx);		// Left tile plane
@@ -70,12 +79,20 @@ void LocalFogVolumeTiledCullingCS(uint3 ThreadId : SV_DispatchThreadID)
 		}
 		if (!bCompletelyOutside)
 		{
-			LocalFogVolumeTileDataTextureUAV[uint3(TilePos, 1 + LFVCount)] = LFVIndex;
+			LocalFogVolumeTileDataTextureUAV[uint3(TileCoord, 1 + LFVCount)] = LFVIndex;
 			LFVCount++;
 		}
 	}
 
-	LocalFogVolumeTileDataTextureUAV[uint3(TilePos, 0)] = LFVCount;
+	LocalFogVolumeTileDataTextureUAV[uint3(TileCoord, 0)] = LFVCount;
+
+	if (LFVCount > 0)
+	{
+		uint WriteToIndex;
+		InterlockedAdd(TileDrawIndirectBufferUAV[1], 1, WriteToIndex);
+
+		TileDataBufferUAV[WriteToIndex] = LFVPackTile(TileCoord);
+	}
 }
 
 
diff --git a/Engine/Source/Runtime/Renderer/Private/LocalFogVolumeRendering.cpp b/Engine/Source/Runtime/Renderer/Private/LocalFogVolumeRendering.cpp
--- a/Engine/Source/Runtime/Renderer/Private/LocalFogVolumeRendering.cpp
+++ b/Engine/Source/Runtime/Renderer/Private/LocalFogVolumeRendering.cpp
@@ -91,9 +91,9 @@ void SetDummyLocalFogVolumeForView(FRDGBuilder& GraphBuilder, FViewInfo& View)
 		FRDGBufferDesc::CreateBufferDesc(sizeof(FVector4f), 1), &DummyCullingData, sizeof(FVector4f) * 1, ERDGInitialDataFlags::NoCopy);
 	View.LocalFogVolumeViewData.GPUInstanceCullingDataBufferSRV = GraphBuilder.CreateSRV(View.LocalFogVolumeViewData.GPUInstanceCullingDataBuffer, PF_A32B32G32R32F);
 
-	View.LocalFogVolumeViewData.TileDataTextureArray		= GSystemTextures.GetZeroUIntArrayDummy(GraphBuilder);
-	View.LocalFogVolumeViewData.TileDataTextureArraySRV		= GraphBuilder.CreateSRV(View.LocalFogVolumeViewData.TileDataTextureArray);
-	View.LocalFogVolumeViewData.TileDataTextureArrayUAV		= nullptr;	// Should never be written by culling passes if there are no instances.
+	View.LocalFogVolumeViewData.TileDataTextureArray			= GSystemTextures.GetZeroUIntArrayDummy(GraphBuilder);
+	View.LocalFogVolumeViewData.TileDataTextureArraySRV			= GraphBuilder.CreateSRV(View.LocalFogVolumeViewData.TileDataTextureArray);
+	View.LocalFogVolumeViewData.TileDataTextureArrayUAV			= nullptr;	// Should never be written by culling passes if there are no instances.
 
 	View.LocalFogVolumeViewData.UniformParametersStruct.LocalFogVolumeCommon.LocalFogVolumeTileDataTextureResolution= FUintVector2(1, 1);
 	View.LocalFogVolumeViewData.UniformParametersStruct.LocalFogVolumeCommon.LocalFogVolumeInstanceCount			= View.LocalFogVolumeViewData.GPUInstanceCount;
@@ -101,6 +101,15 @@ void SetDummyLocalFogVolumeForView(FRDGBuilder& GraphBuilder, FViewInfo& View)
 	View.LocalFogVolumeViewData.UniformParametersStruct.LocalFogVolumeCommon.LocalFogVolumeInstances				= View.LocalFogVolumeViewData.GPUInstanceDataBufferSRV;
 	View.LocalFogVolumeViewData.UniformParametersStruct.LocalFogVolumeTileDataTexture								= View.LocalFogVolumeViewData.TileDataTextureArraySRV;
 	View.LocalFogVolumeViewData.UniformBuffer																		= GraphBuilder.CreateUniformBuffer(&View.LocalFogVolumeViewData.UniformParametersStruct);
+
+	// This buffer must remain a basic vertex buffer for mobile to be able to read it from vertex shader
+	View.LocalFogVolumeViewData.GPUTileDataBuffer		= CreateVertexBuffer(
+		GraphBuilder, TEXT("LocalFogVolume.GPUTileDataBuffer"), FRDGBufferDesc::CreateBufferDesc(sizeof(uint32), 1 * 1 * sizeof(uint32)), nullptr, 0);
+	View.LocalFogVolumeViewData.GPUTileDataBufferSRV	= GraphBuilder.CreateSRV(View.LocalFogVolumeViewData.GPUTileDataBuffer);
+	View.LocalFogVolumeViewData.GPUTileDataBufferUAV	= nullptr;	// Should never be written by culling passes if there are no instances.
+
+	View.LocalFogVolumeViewData.GPUTileDrawIndirectBuffer	= GraphBuilder.CreateBuffer(FRDGBufferDesc::CreateIndirectDesc<FRHIDrawIndirectParameters>(), TEXT("LocalFogVolume.DispatchIndirectBuffer"));
+	View.LocalFogVolumeViewData.GPUTileDrawIndirectBufferUAV= nullptr;
 }
 
 
@@ -158,6 +167,8 @@ public:
 		SHADER_PARAMETER_STRUCT(FLocalFogVolumeCommonParameters, LFV)
 		SHADER_PARAMETER_RDG_TEXTURE_UAV(RWTexture2DArray<uint>, LocalFogVolumeTileDataTextureUAV)
 		SHADER_PARAMETER_RDG_BUFFER_SRV(Buffer<float4>, LocalFogVolumeCullingDataBuffer)
+		SHADER_PARAMETER_RDG_BUFFER_UAV(RWBuffer<uint>, TileDataBufferUAV)
+		SHADER_PARAMETER_RDG_BUFFER_UAV(RWBuffer<uint>, TileDrawIndirectBufferUAV)
 		SHADER_PARAMETER(FVector4f, LeftPlane)
 		SHADER_PARAMETER(FVector4f, RightPlane)
 		SHADER_PARAMETER(FVector4f, TopPlane)
@@ -178,10 +189,12 @@ static void LocalFogVolumeViewTiledCullingPass(FViewInfo& View, FRDGBuilder& Gra
 	FIntVector TileDataTextureSize = View.LocalFogVolumeViewData.TileDataTextureArray->Desc.GetSize();
 
 	FLocalFogVolumeTiledCullingCS::FParameters* PassParameters = GraphBuilder.AllocParameters<FLocalFogVolumeTiledCullingCS::FParameters>();
-	PassParameters->View = View.ViewUniformBuffer;
-	PassParameters->LFV = View.LocalFogVolumeViewData.UniformParametersStruct.LocalFogVolumeCommon;
-	PassParameters->LocalFogVolumeCullingDataBuffer = View.LocalFogVolumeViewData.GPUInstanceCullingDataBufferSRV;
-	PassParameters->LocalFogVolumeTileDataTextureUAV = View.LocalFogVolumeViewData.TileDataTextureArrayUAV;
+	PassParameters->View								= View.ViewUniformBuffer;
+	PassParameters->LFV									= View.LocalFogVolumeViewData.UniformParametersStruct.LocalFogVolumeCommon;
+	PassParameters->LocalFogVolumeCullingDataBuffer		= View.LocalFogVolumeViewData.GPUInstanceCullingDataBufferSRV;
+	PassParameters->LocalFogVolumeTileDataTextureUAV	= View.LocalFogVolumeViewData.TileDataTextureArrayUAV;
+	PassParameters->TileDataBufferUAV					= View.LocalFogVolumeViewData.GPUTileDataBufferUAV;
+	PassParameters->TileDrawIndirectBufferUAV			= View.LocalFogVolumeViewData.GPUTileDrawIndirectBufferUAV;
 
 	auto ConvertPlanToVector4f = [&](FVector4f& OutVec4f, auto& Plane, bool bFlipPlane)
 	{
@@ -341,6 +354,17 @@ void CreateViewLocalFogVolumeBufferSRV(FViewInfo& View, FRDGBuilder& GraphBuilde
 	View.LocalFogVolumeViewData.UniformParametersStruct.LocalFogVolumeCommon.LocalFogVolumeInstances				= View.LocalFogVolumeViewData.GPUInstanceDataBufferSRV;
 	View.LocalFogVolumeViewData.UniformParametersStruct.LocalFogVolumeTileDataTexture								= View.LocalFogVolumeViewData.TileDataTextureArraySRV;
 	View.LocalFogVolumeViewData.UniformBuffer																		= GraphBuilder.CreateUniformBuffer(&View.LocalFogVolumeViewData.UniformParametersStruct);
+
+	// This buffer must remain a basic vertex buffer for mobile to be able to read it from vertex shader
+	View.LocalFogVolumeViewData.GPUTileDataBuffer = CreateVertexBuffer(
+		GraphBuilder, TEXT("LocalFogVolume.GPUTileDataBuffer"),
+		FRDGBufferDesc::CreateBufferDesc(sizeof(uint32), TileDataTextureResolution.X * TileDataTextureResolution.Y * sizeof(uint32)), nullptr, 0);
+	View.LocalFogVolumeViewData.GPUTileDataBufferSRV = GraphBuilder.CreateSRV(View.LocalFogVolumeViewData.GPUTileDataBuffer, PF_R32_UINT);
+	View.LocalFogVolumeViewData.GPUTileDataBufferUAV = GraphBuilder.CreateUAV(View.LocalFogVolumeViewData.GPUTileDataBuffer, PF_R32_UINT);
+
+	View.LocalFogVolumeViewData.GPUTileDrawIndirectBuffer = GraphBuilder.CreateBuffer(FRDGBufferDesc::CreateIndirectDesc<FRHIDrawIndirectParameters>(), TEXT("LocalFogVolume.DispatchIndirectBuffer"));
+	View.LocalFogVolumeViewData.GPUTileDrawIndirectBufferUAV = GraphBuilder.CreateUAV(View.LocalFogVolumeViewData.GPUTileDrawIndirectBuffer, PF_R32_UINT);
+	AddClearUAVPass(GraphBuilder, View.LocalFogVolumeViewData.GPUTileDrawIndirectBufferUAV, 0);
 }
 
 void InitLocalFogVolumesForViews(
@@ -446,6 +470,7 @@ class FLocalFogVolumeTiledRenderVS : public FGlobalShader
 	BEGIN_SHADER_PARAMETER_STRUCT(FParameters, )
 		SHADER_PARAMETER_STRUCT_REF(FViewUniformShaderParameters, View)
 		SHADER_PARAMETER_STRUCT(FLocalFogVolumeUniformParameters, LFV)
+		SHADER_PARAMETER_RDG_BUFFER_SRV(Buffer<uint>, TileDataBuffer)
 	END_SHADER_PARAMETER_STRUCT()
 
 	using FPermutationDomain = TShaderPermutationDomain<>;
@@ -510,6 +535,7 @@ BEGIN_SHADER_PARAMETER_STRUCT(FLocalFogVolumeTiledPassParameters, )
 	SHADER_PARAMETER_STRUCT_INCLUDE(FLocalFogVolumeTiledRenderVS::FParameters, VS)
 	SHADER_PARAMETER_STRUCT_INCLUDE(FLocalFogVolumeTiledRenderPS::FParameters, PS)
 	SHADER_PARAMETER_RDG_UNIFORM_BUFFER(FSceneTextureUniformParameters, SceneTextures)
+	RDG_BUFFER_ACCESS(TileDrawIndirectBuffer, ERHIAccess::IndirectArgs)
 	RENDER_TARGET_BINDING_SLOTS()
 END_SHADER_PARAMETER_STRUCT()
  
@@ -535,19 +561,20 @@ void RenderLocalFogVolume(
 				continue;
 			}
 
-
 			if (GetLocalFogVolumeTiledRenderingEnable())
 			{
 				FLocalFogVolumeTiledPassParameters* PassParameters = GraphBuilder.AllocParameters<FLocalFogVolumeTiledPassParameters>();
 
 				PassParameters->VS.View = GetShaderBinding(View.ViewUniformBuffer);
 				PassParameters->VS.LFV = View.LocalFogVolumeViewData.UniformParametersStruct;
+				PassParameters->VS.TileDataBuffer = View.LocalFogVolumeViewData.GPUTileDataBufferSRV;
 
 				PassParameters->PS.View = GetShaderBinding(View.ViewUniformBuffer);
 				PassParameters->PS.LFV = View.LocalFogVolumeViewData.UniformParametersStruct;
 				PassParameters->PS.LocalFogVolumeTileDebug = FMath::Clamp(CVarLocalFogVolumeTileDebug.GetValueOnRenderThread(), 0, 2);
 
 				PassParameters->SceneTextures = SceneTextures.UniformBuffer;
+				PassParameters->TileDrawIndirectBuffer = View.LocalFogVolumeViewData.GPUTileDrawIndirectBuffer;
 				PassParameters->RenderTargets[0] = FRenderTargetBinding(SceneColorTexture, ERenderTargetLoadAction::ENoAction);
 
 				FLocalFogVolumeTiledRenderVS::FPermutationDomain VSPermutationVector;
@@ -581,7 +608,7 @@ void RenderLocalFogVolume(
 					GraphicsPSOInit.BlendState = TStaticBlendState<CW_RGBA, BO_Add, BF_One, BF_SourceAlpha, BO_Add, BF_Zero, BF_SourceAlpha>::GetRHI();
 					GraphicsPSOInit.PrimitiveType = PT_TriangleList; // LFV_TODO check if rects are supported and use them if so
 
-					GraphicsPSOInit.BoundShaderState.VertexDeclarationRHI = GFilterVertexDeclaration.VertexDeclarationRHI; //		GraphicsPSOInit.BoundShaderState.VertexDeclarationRHI = GetVertexDeclarationFVector4();
+					GraphicsPSOInit.BoundShaderState.VertexDeclarationRHI = GFilterVertexDeclaration.VertexDeclarationRHI;
 					GraphicsPSOInit.BoundShaderState.VertexShaderRHI = VertexShader.GetVertexShader();
 					GraphicsPSOInit.BoundShaderState.PixelShaderRHI = PixelShader.GetPixelShader();
 
@@ -591,11 +618,7 @@ void RenderLocalFogVolume(
 					SetShaderParameters(RHICmdList, PixelShader, PixelShader.GetPixelShader(), PassParameters->PS);
 
 					RHICmdList.SetStreamSource(0, nullptr, 0);
-					RHICmdList.DrawPrimitive(
-						0,					// BaseVertexIndex
-						6,					// NumPrimitives
-						TileInstanceCount	// NumInstances
-						);
+					RHICmdList.DrawPrimitiveIndirect(PassParameters->TileDrawIndirectBuffer->GetIndirectRHICallBuffer(), 0);
 				});
 			}
 			else
diff --git a/Engine/Source/Runtime/Renderer/Private/LocalFogVolumeRendering.h b/Engine/Source/Runtime/Renderer/Private/LocalFogVolumeRendering.h
--- a/Engine/Source/Runtime/Renderer/Private/LocalFogVolumeRendering.h
+++ b/Engine/Source/Runtime/Renderer/Private/LocalFogVolumeRendering.h
@@ -96,15 +96,23 @@ struct FLocalFogVolumeViewData
 	uint32 GPUInstanceCount = 0;
 	FRDGBufferRef GPUInstanceDataBuffer = nullptr;
 	FRDGBufferSRVRef GPUInstanceDataBufferSRV = nullptr;
+
 	FRDGBufferRef GPUInstanceCullingDataBuffer = nullptr;
 	FRDGBufferSRVRef GPUInstanceCullingDataBufferSRV = nullptr;
+
+	FRDGBufferRef	 GPUTileDataBuffer = nullptr;
+	FRDGBufferSRVRef GPUTileDataBufferSRV = nullptr;
+	FRDGBufferUAVRef GPUTileDataBufferUAV = nullptr;
+	FRDGBufferRef    GPUTileDrawIndirectBuffer = nullptr;
+	FRDGBufferUAVRef GPUTileDrawIndirectBufferUAV = nullptr;
+
 	TRDGUniformBufferRef<FLocalFogVolumeUniformParameters> UniformBuffer = nullptr;
 
 	FLocalFogVolumeUniformParameters UniformParametersStruct;
 
-	FRDGTextureRef TileDataTextureArray;		// First slice is the instance count, later slices are instance indices.
-	FRDGTextureSRVRef TileDataTextureArraySRV;
-	FRDGTextureUAVRef TileDataTextureArrayUAV;
+	FRDGTextureRef		TileDataTextureArray;		// First slice is the instance count, later slices are instance indices.
+	FRDGTextureSRVRef	TileDataTextureArraySRV;
+	FRDGTextureUAVRef	TileDataTextureArrayUAV;
 };
 
 
