LFV - tiled fog volume rendering on mobile.

#prelfight shader

diff --git a/Engine/Source/Runtime/Renderer/Private/LocalFogVolumeRendering.cpp b/Engine/Source/Runtime/Renderer/Private/LocalFogVolumeRendering.cpp
--- a/Engine/Source/Runtime/Renderer/Private/LocalFogVolumeRendering.cpp
+++ b/Engine/Source/Runtime/Renderer/Private/LocalFogVolumeRendering.cpp
@@ -20,7 +20,7 @@ static TAutoConsoleVariable<int32> CVarLocalFogVolumeApplyOnTransclucent(
 	ECVF_ReadOnly | ECVF_RenderThreadSafe);
 
 static TAutoConsoleVariable<int32> CVarLocalFogVolumeTiledRendering(
-	TEXT("r.LocalFogVolume.TiledRendering"), 0,
+	TEXT("r.LocalFogVolume.TiledRendering"), 1,
 	TEXT("Use a tiled approach to render the local fog volume on screen.\n"),
 	ECVF_RenderThreadSafe);
 
@@ -105,7 +105,7 @@ void SetDummyLocalFogVolumeForView(FRDGBuilder& GraphBuilder, FViewInfo& View)
 	// This buffer must remain a basic vertex buffer for mobile to be able to read it from vertex shader
 	View.LocalFogVolumeViewData.GPUTileDataBuffer		= CreateVertexBuffer(
 		GraphBuilder, TEXT("LocalFogVolume.GPUTileDataBuffer"), FRDGBufferDesc::CreateBufferDesc(sizeof(uint32), 1 * 1 * sizeof(uint32)), nullptr, 0);
-	View.LocalFogVolumeViewData.GPUTileDataBufferSRV	= GraphBuilder.CreateSRV(View.LocalFogVolumeViewData.GPUTileDataBuffer);
+	View.LocalFogVolumeViewData.GPUTileDataBufferSRV	= GraphBuilder.CreateSRV(View.LocalFogVolumeViewData.GPUTileDataBuffer, PF_R32_UINT);
 	View.LocalFogVolumeViewData.GPUTileDataBufferUAV	= nullptr;	// Should never be written by culling passes if there are no instances.
 
 	View.LocalFogVolumeViewData.GPUTileDrawIndirectBuffer	= GraphBuilder.CreateBuffer(FRDGBufferDesc::CreateIndirectDesc<FRHIDrawIndirectParameters>(), TEXT("LocalFogVolume.DispatchIndirectBuffer"));
@@ -589,13 +589,12 @@ void RenderLocalFogVolume(
 				ClearUnusedGraphResources(PixelShader, &PassParameters->PS);
 
 				FUintVector2& LocalFogVolumeTileDataTextureResolution = View.LocalFogVolumeViewData.UniformParametersStruct.LocalFogVolumeCommon.LocalFogVolumeTileDataTextureResolution;
-				const uint32 TileInstanceCount = LocalFogVolumeTileDataTextureResolution.X * LocalFogVolumeTileDataTextureResolution.Y;
 
 				GraphBuilder.AddPass(
 					RDG_EVENT_NAME("LocalFogVolume.Tiled (%u X %u)", LocalFogVolumeTileDataTextureResolution.X, LocalFogVolumeTileDataTextureResolution.Y),
 					PassParameters,
 					ERDGPassFlags::Raster,
-					[VertexShader, PixelShader, PassParameters, ViewRect, TileInstanceCount](FRHICommandList& RHICmdList)
+					[VertexShader, PixelShader, PassParameters, ViewRect](FRHICommandList& RHICmdList)
 				{
 					FGraphicsPipelineStateInitializer GraphicsPSOInit;
 					RHICmdList.ApplyCachedRenderTargets(GraphicsPSOInit);
@@ -758,6 +757,75 @@ class FMobileLocalFogVolumeSplatPS : public FGlobalShader
 
 IMPLEMENT_GLOBAL_SHADER(FMobileLocalFogVolumeSplatPS, "/Engine/Private/LocalFogVolumes/LocalFogVolumeSplat.usf", "LocalFogVolumeSplatPS", SF_Pixel);
 
+
+class FMobileLocalFogVolumeTiledRenderVS : public FGlobalShader
+{
+	DECLARE_GLOBAL_SHADER(FMobileLocalFogVolumeTiledRenderVS);
+	SHADER_USE_PARAMETER_STRUCT(FMobileLocalFogVolumeTiledRenderVS, FGlobalShader);
+
+	BEGIN_SHADER_PARAMETER_STRUCT(FParameters, )
+		SHADER_PARAMETER_STRUCT_INCLUDE(FViewShaderParameters, View)
+		SHADER_PARAMETER_STRUCT(FLocalFogVolumeUniformParameters, LFV)
+		SHADER_PARAMETER_RDG_BUFFER_SRV(Buffer<uint>, TileDataBuffer)
+		RDG_BUFFER_ACCESS(TileDrawIndirectBuffer, ERHIAccess::IndirectArgs)
+	END_SHADER_PARAMETER_STRUCT()
+
+	using FPermutationDomain = TShaderPermutationDomain<>;
+
+	static bool ShouldCompilePermutation(const FGlobalShaderPermutationParameters& Parameters)
+	{
+		if (!IsMobilePlatform(Parameters.Platform))
+		{
+			return false;
+		}
+		return true;
+	}
+
+	static void ModifyCompilationEnvironment(const FGlobalShaderPermutationParameters& Parameters, FShaderCompilerEnvironment& OutEnvironment)
+	{
+		FGlobalShader::ModifyCompilationEnvironment(Parameters, OutEnvironment);
+		const bool bMobileForceDepthRead = MobileUsesFullDepthPrepass(Parameters.Platform);
+		OutEnvironment.SetDefine(TEXT("LFV_TILED_VS"), 1);
+		OutEnvironment.SetDefine(TEXT("IS_MOBILE_DEPTHREAD_SUBPASS"), bMobileForceDepthRead ? 0u : 1u);
+	}
+};
+
+IMPLEMENT_GLOBAL_SHADER(FMobileLocalFogVolumeTiledRenderVS, "/Engine/Private/LocalFogVolumes/LocalFogVolumeSplat.usf", "LocalFogVolumeTiledVS", SF_Vertex);
+
+class FMobileLocalFogVolumeTiledRenderPS : public FGlobalShader
+{
+	DECLARE_GLOBAL_SHADER(FMobileLocalFogVolumeTiledRenderPS);
+	SHADER_USE_PARAMETER_STRUCT(FMobileLocalFogVolumeTiledRenderPS, FGlobalShader);
+
+	BEGIN_SHADER_PARAMETER_STRUCT(FParameters, )
+		SHADER_PARAMETER_STRUCT_INCLUDE(FViewShaderParameters, View)
+		SHADER_PARAMETER_RDG_UNIFORM_BUFFER(FMobileBasePassUniformParameters, MobileBasePass)
+		SHADER_PARAMETER_STRUCT(FLocalFogVolumeUniformParameters, LFV)
+		SHADER_PARAMETER(int32, LocalFogVolumeTileDebug)
+	END_SHADER_PARAMETER_STRUCT()
+
+	using FPermutationDomain = TShaderPermutationDomain<>;
+
+	static bool ShouldCompilePermutation(const FGlobalShaderPermutationParameters& Parameters)
+	{
+		if (!IsMobilePlatform(Parameters.Platform))
+		{
+			return false;
+		}
+		return true;
+	}
+
+	static void ModifyCompilationEnvironment(const FGlobalShaderPermutationParameters& Parameters, FShaderCompilerEnvironment& OutEnvironment)
+	{
+		FGlobalShader::ModifyCompilationEnvironment(Parameters, OutEnvironment);
+		const bool bMobileForceDepthRead = MobileUsesFullDepthPrepass(Parameters.Platform);
+		OutEnvironment.SetDefine(TEXT("LFV_TILED_PS"), 1);
+		OutEnvironment.SetDefine(TEXT("IS_MOBILE_DEPTHREAD_SUBPASS"), bMobileForceDepthRead ? 0u : 1u);
+	}
+};
+
+IMPLEMENT_GLOBAL_SHADER(FMobileLocalFogVolumeTiledRenderPS, "/Engine/Private/LocalFogVolumes/LocalFogVolumeSplat.usf", "LocalFogVolumeTiledPS", SF_Pixel);
+
 void RenderLocalFogVolumeMobile(
 	FRHICommandList& RHICmdList,
 	const FViewInfo& View)
@@ -769,51 +837,99 @@ void RenderLocalFogVolumeMobile(
 
 	SCOPED_DRAW_EVENT(RHICmdList, LocalFogVolumeVolumes);
 
-	FMobileLocalFogVolumeSplatVS::FPermutationDomain VSPermutationVector;
-	auto VertexShader = View.ShaderMap->GetShader< FMobileLocalFogVolumeSplatVS >(VSPermutationVector);
 
-	FMobileLocalFogVolumeSplatPS::FPermutationDomain PsPermutationVector;
-	auto PixelShader = View.ShaderMap->GetShader< FMobileLocalFogVolumeSplatPS >(PsPermutationVector);
+	if (GetLocalFogVolumeTiledRenderingEnable())
+	{
+		FMobileLocalFogVolumeTiledRenderVS::FPermutationDomain VSPermutationVector;
+		auto VertexShader = View.ShaderMap->GetShader< FMobileLocalFogVolumeTiledRenderVS >(VSPermutationVector);
+
+		FMobileLocalFogVolumeTiledRenderPS::FPermutationDomain PsPermutationVector;
+		auto PixelShader = View.ShaderMap->GetShader< FMobileLocalFogVolumeTiledRenderPS >(PsPermutationVector);
+
+		FGraphicsPipelineStateInitializer GraphicsPSOInit;
+		RHICmdList.ApplyCachedRenderTargets(GraphicsPSOInit);
+
+		const FIntRect ViewRect = View.ViewRect;
+		RHICmdList.SetViewport(ViewRect.Min.X, ViewRect.Min.Y, 0.0f, ViewRect.Max.X, ViewRect.Max.Y, 1.0f);
 
-	const FIntRect ViewRect = View.ViewRect;
+		// Render back faces only since camera may intersect
+		GraphicsPSOInit.RasterizerState = TStaticRasterizerState<>::GetRHI();
+		GraphicsPSOInit.DepthStencilState = TStaticDepthStencilState<false, CF_Always>::GetRHI();
+		GraphicsPSOInit.BlendState = TStaticBlendState<CW_RGBA, BO_Add, BF_One, BF_SourceAlpha, BO_Add, BF_Zero, BF_SourceAlpha>::GetRHI();
+		GraphicsPSOInit.PrimitiveType = PT_TriangleList; // LFV_TODO check if rects are supported and use them if so
 
-	FGraphicsPipelineStateInitializer GraphicsPSOInit;
-	RHICmdList.ApplyCachedRenderTargets(GraphicsPSOInit);
+		GraphicsPSOInit.BoundShaderState.VertexDeclarationRHI = GFilterVertexDeclaration.VertexDeclarationRHI;
+		GraphicsPSOInit.BoundShaderState.VertexShaderRHI = VertexShader.GetVertexShader();
+		GraphicsPSOInit.BoundShaderState.PixelShaderRHI = PixelShader.GetPixelShader();
 
-	RHICmdList.SetViewport(ViewRect.Min.X, ViewRect.Min.Y, 0.0f, ViewRect.Max.X, ViewRect.Max.Y, 1.0f);
+		SetGraphicsPipelineState(RHICmdList, GraphicsPSOInit, 0);
 
-	// Render back faces only since camera may intersect
-	GraphicsPSOInit.RasterizerState = TStaticRasterizerState<FM_Solid, CM_CW>::GetRHI();
-	GraphicsPSOInit.DepthStencilState = TStaticDepthStencilState<false, CF_Always>::GetRHI();
-	GraphicsPSOInit.BlendState = TStaticBlendState<CW_RGBA, BO_Add, BF_One, BF_SourceAlpha, BO_Add, BF_Zero, BF_SourceAlpha>::GetRHI();
-	GraphicsPSOInit.PrimitiveType = PT_TriangleList;
 
-	GraphicsPSOInit.BoundShaderState.VertexDeclarationRHI = GetVertexDeclarationFVector4();
-	GraphicsPSOInit.BoundShaderState.VertexShaderRHI = VertexShader.GetVertexShader();
-	GraphicsPSOInit.BoundShaderState.PixelShaderRHI = PixelShader.GetPixelShader();
+		FMobileLocalFogVolumeTiledRenderVS::FParameters VSParameters;
+		VSParameters.View = View.GetShaderParameters();
+		VSParameters.LFV = View.LocalFogVolumeViewData.UniformParametersStruct;
+		VSParameters.TileDataBuffer = View.LocalFogVolumeViewData.GPUTileDataBufferSRV;
+		VSParameters.TileDrawIndirectBuffer = View.LocalFogVolumeViewData.GPUTileDrawIndirectBuffer;
+		SetShaderParameters(RHICmdList, VertexShader, VertexShader.GetVertexShader(), VSParameters);
 
-	SetGraphicsPipelineState(RHICmdList, GraphicsPSOInit, 0);
+		FMobileLocalFogVolumeTiledRenderPS::FParameters PSParameters;
+		PSParameters.View = View.GetShaderParameters();
+		PSParameters.LFV = View.LocalFogVolumeViewData.UniformParametersStruct;
+		PSParameters.LocalFogVolumeTileDebug = FMath::Clamp(CVarLocalFogVolumeTileDebug.GetValueOnRenderThread(), 0, 2);
+		// PSParameters.MobileBasePass filled up by the RDG pass parameters.
+		SetShaderParameters(RHICmdList, PixelShader, PixelShader.GetPixelShader(), PSParameters);
 
-	FMobileLocalFogVolumeSplatVS::FParameters VSParameters;
-	VSParameters.View = View.GetShaderParameters();
-	VSParameters.LFV = View.LocalFogVolumeViewData.UniformParametersStruct;
-	SetShaderParameters(RHICmdList, VertexShader, VertexShader.GetVertexShader(), VSParameters);
+		RHICmdList.SetStreamSource(0, nullptr, 0);
+		RHICmdList.DrawPrimitiveIndirect(View.LocalFogVolumeViewData.GPUTileDrawIndirectBuffer->GetIndirectRHICallBuffer(), 0);
+	}
+	else
+	{
+		FMobileLocalFogVolumeSplatVS::FPermutationDomain VSPermutationVector;
+		auto VertexShader = View.ShaderMap->GetShader< FMobileLocalFogVolumeSplatVS >(VSPermutationVector);
 
-	FMobileLocalFogVolumeSplatPS::FParameters PSParameters;
-	PSParameters.View = View.GetShaderParameters();
-	PSParameters.LFV = View.LocalFogVolumeViewData.UniformParametersStruct;
-	// PSParameters.MobileBasePass filled up by the RDG pass parameters.
-	SetShaderParameters(RHICmdList, PixelShader, PixelShader.GetPixelShader(), PSParameters);
+		FMobileLocalFogVolumeSplatPS::FPermutationDomain PsPermutationVector;
+		auto PixelShader = View.ShaderMap->GetShader< FMobileLocalFogVolumeSplatPS >(PsPermutationVector);
 
-	RHICmdList.SetStreamSource(0, GetUnitCubeVertexBuffer(), 0);
+		const FIntRect ViewRect = View.ViewRect;
 
-	RHICmdList.DrawIndexedPrimitive(
-		GetUnitCubeIndexBuffer()
-		, 0												//BaseVertexIndex
-		, 0												//FirstInstance
-		, 8												//uint32 NumVertices
-		, 0												//uint32 StartIndex
-		, UE_ARRAY_COUNT(GCubeIndices) / 3				//uint32 NumPrimitives
-		, View.LocalFogVolumeViewData.GPUInstanceCount	//uint32 NumInstances
-	);
+		FGraphicsPipelineStateInitializer GraphicsPSOInit;
+		RHICmdList.ApplyCachedRenderTargets(GraphicsPSOInit);
+
+		RHICmdList.SetViewport(ViewRect.Min.X, ViewRect.Min.Y, 0.0f, ViewRect.Max.X, ViewRect.Max.Y, 1.0f);
+
+		// Render back faces only since camera may intersect
+		GraphicsPSOInit.RasterizerState = TStaticRasterizerState<FM_Solid, CM_CW>::GetRHI();
+		GraphicsPSOInit.DepthStencilState = TStaticDepthStencilState<false, CF_Always>::GetRHI();
+		GraphicsPSOInit.BlendState = TStaticBlendState<CW_RGBA, BO_Add, BF_One, BF_SourceAlpha, BO_Add, BF_Zero, BF_SourceAlpha>::GetRHI();
+		GraphicsPSOInit.PrimitiveType = PT_TriangleList;
+
+		GraphicsPSOInit.BoundShaderState.VertexDeclarationRHI = GetVertexDeclarationFVector4();
+		GraphicsPSOInit.BoundShaderState.VertexShaderRHI = VertexShader.GetVertexShader();
+		GraphicsPSOInit.BoundShaderState.PixelShaderRHI = PixelShader.GetPixelShader();
+
+		SetGraphicsPipelineState(RHICmdList, GraphicsPSOInit, 0);
+
+		FMobileLocalFogVolumeSplatVS::FParameters VSParameters;
+		VSParameters.View = View.GetShaderParameters();
+		VSParameters.LFV = View.LocalFogVolumeViewData.UniformParametersStruct;
+		SetShaderParameters(RHICmdList, VertexShader, VertexShader.GetVertexShader(), VSParameters);
+
+		FMobileLocalFogVolumeSplatPS::FParameters PSParameters;
+		PSParameters.View = View.GetShaderParameters();
+		PSParameters.LFV = View.LocalFogVolumeViewData.UniformParametersStruct;
+		// PSParameters.MobileBasePass filled up by the RDG pass parameters.
+		SetShaderParameters(RHICmdList, PixelShader, PixelShader.GetPixelShader(), PSParameters);
+
+		RHICmdList.SetStreamSource(0, GetUnitCubeVertexBuffer(), 0);
+
+		RHICmdList.DrawIndexedPrimitive(
+			GetUnitCubeIndexBuffer()
+			, 0												//BaseVertexIndex
+			, 0												//FirstInstance
+			, 8												//uint32 NumVertices
+			, 0												//uint32 StartIndex
+			, UE_ARRAY_COUNT(GCubeIndices) / 3				//uint32 NumPrimitives
+			, View.LocalFogVolumeViewData.GPUInstanceCount	//uint32 NumInstances
+		);
+	}
 }
diff --git a/Engine/Source/Runtime/Renderer/Private/MobileShadingRenderer.cpp b/Engine/Source/Runtime/Renderer/Private/MobileShadingRenderer.cpp
--- a/Engine/Source/Runtime/Renderer/Private/MobileShadingRenderer.cpp
+++ b/Engine/Source/Runtime/Renderer/Private/MobileShadingRenderer.cpp
@@ -193,6 +193,9 @@ BEGIN_SHADER_PARAMETER_STRUCT(FMobileRenderPassParameters, RENDERER_API)
 	SHADER_PARAMETER_RDG_UNIFORM_BUFFER(FMobileBasePassUniformParameters, MobileBasePass)
 	SHADER_PARAMETER_STRUCT_REF(FMobileReflectionCaptureShaderData, ReflectionCapture)
 	SHADER_PARAMETER_RDG_BUFFER_SRV(Buffer<float4>, LocalFogVolumeInstances)
+	RDG_BUFFER_ACCESS(TileDrawIndirectBuffer, ERHIAccess::IndirectArgs)
+	SHADER_PARAMETER_RDG_TEXTURE_SRV(Texture2DArray<uint>, LocalFogVolumeTileDataTexture)
+	SHADER_PARAMETER_RDG_BUFFER_SRV(Buffer<uint>, TileDataBuffer)
 	RENDER_TARGET_BINDING_SLOTS()
 END_SHADER_PARAMETER_STRUCT()
 
@@ -1306,6 +1309,9 @@ void FMobileSceneRenderer::RenderForward(FRDGBuilder& GraphBuilder, FRDGTextureR
 		PassParameters->ReflectionCapture = View.MobileReflectionCaptureUniformBuffer;
 		PassParameters->RenderTargets = BasePassRenderTargets;
 		PassParameters->LocalFogVolumeInstances = View.LocalFogVolumeViewData.GPUInstanceDataBufferSRV;
+		PassParameters->TileDrawIndirectBuffer = View.LocalFogVolumeViewData.GPUTileDrawIndirectBuffer;
+		PassParameters->LocalFogVolumeTileDataTexture = View.LocalFogVolumeViewData.TileDataTextureArraySRV;
+		PassParameters->TileDataBuffer = View.LocalFogVolumeViewData.GPUTileDataBufferSRV;
 	
 		BuildInstanceCullingDrawParams(GraphBuilder, View, PassParameters);
 
@@ -1642,6 +1648,9 @@ void FMobileSceneRenderer::RenderDeferred(FRDGBuilder& GraphBuilder, const FSort
 		PassParameters->ReflectionCapture = View.MobileReflectionCaptureUniformBuffer;
 		PassParameters->RenderTargets = BasePassRenderTargets;
 		PassParameters->LocalFogVolumeInstances = View.LocalFogVolumeViewData.GPUInstanceDataBufferSRV;
+		PassParameters->TileDrawIndirectBuffer = View.LocalFogVolumeViewData.GPUTileDrawIndirectBuffer;
+		PassParameters->LocalFogVolumeTileDataTexture = View.LocalFogVolumeViewData.TileDataTextureArraySRV;
+		PassParameters->TileDataBuffer = View.LocalFogVolumeViewData.GPUTileDataBufferSRV;
 		
 		BuildInstanceCullingDrawParams(GraphBuilder, View, PassParameters);
 
