LFV - updated cvar documentation.

#prelfight shader

diff --git a/Engine/Source/Runtime/Renderer/Private/LocalFogVolumeRendering.cpp b/Engine/Source/Runtime/Renderer/Private/LocalFogVolumeRendering.cpp
--- a/Engine/Source/Runtime/Renderer/Private/LocalFogVolumeRendering.cpp
+++ b/Engine/Source/Runtime/Renderer/Private/LocalFogVolumeRendering.cpp
@@ -36,7 +36,7 @@ static TAutoConsoleVariable<int32> CVarLocalFogVolumeTileMaxInstanceCount(
 
 static TAutoConsoleVariable<int32> CVarLocalFogVolumeTileDebug(
 	TEXT("r.LocalFogVolume.TileDebug"), 0,
-	TEXT("Debug the tiled rendering data complexity.\n"),
+	TEXT("Debug the tiled rendering data complexity. 1: show per tile LFV count as color ; 2: same as one but also show the effect of pixel discard/clipping.\n"),
 	ECVF_RenderThreadSafe);
 
 // Example of tile setup
