Update the Foveated and Debug Force Rate VRS image generators to clamp shading rates when additional rates are not supported.

diff --git a/Engine/Source/Runtime/Renderer/Private/VariableRateShading/FoveatedImageGenerator.cpp b/Engine/Source/Runtime/Renderer/Private/VariableRateShading/FoveatedImageGenerator.cpp
--- a/Engine/Source/Runtime/Renderer/Private/VariableRateShading/FoveatedImageGenerator.cpp
+++ b/Engine/Source/Runtime/Renderer/Private/VariableRateShading/FoveatedImageGenerator.cpp
@@ -204,7 +204,7 @@ void FFoveatedImageGenerator::PrepareImages(FRDGBuilder& GraphBuilder, const FSc
 	PassParameters->HMDFieldOfView = HMDFieldOfView;
 
 	PassParameters->FoveationFullRateCutoffSquared = FoveationFullRateCutoff * FoveationFullRateCutoff;
-	PassParameters->FoveationHalfRateCutoffSquared = FoveationHalfRateCutoff * FoveationHalfRateCutoff;
+	PassParameters->FoveationHalfRateCutoffSquared = GRHISupportsLargerVariableRateShadingSizes ? (FoveationHalfRateCutoff * FoveationHalfRateCutoff) : 2.0f; // Set cutoff to outside screen edge if quarter-rate is unsupported
 
 	PassParameters->LeftEyeCenterPixelXY = FVector2f(Desc.Extent.X * FoveationCenterX, Desc.Extent.Y * FoveationCenterY);
 	PassParameters->RightEyeCenterPixelXY = PassParameters->LeftEyeCenterPixelXY;
diff --git a/Engine/Source/Runtime/Renderer/Private/VariableRateShading/VariableRateShadingImageManager.cpp b/Engine/Source/Runtime/Renderer/Private/VariableRateShading/VariableRateShadingImageManager.cpp
--- a/Engine/Source/Runtime/Renderer/Private/VariableRateShading/VariableRateShadingImageManager.cpp
+++ b/Engine/Source/Runtime/Renderer/Private/VariableRateShading/VariableRateShadingImageManager.cpp
@@ -50,9 +50,10 @@ TAutoConsoleVariable<int32> CVarVRSPreview(
 void CVarVRSDebugForceRateCallback(IConsoleVariable* Var)
 {
 	const int32 RequestedDebugForceRate = Var->GetInt();
-	if (RequestedDebugForceRate > 6) 
+	const int32 NumberOfAvailableRates = FVariableRateShadingImageManager::GetNumberOfSupportedRates();
+	if (RequestedDebugForceRate >= NumberOfAvailableRates) 
 	{
-		UE_LOG(LogVRS, Warning, TEXT("Selected forced shading rate exceeds maximum available, defaulting to 4x4"));
+		UE_LOG(LogVRS, Warning, TEXT("Selected forced shading rate exceeds maximum available, defaulting to %s"), GRHISupportsLargerVariableRateShadingSizes ? TEXT("4x4") : TEXT("2x2"));
 	}
 }
 
@@ -366,6 +367,16 @@ FRDGTextureDesc FVariableRateShadingImageManager::GetSRIDesc()
 
 }
 
+int32 FVariableRateShadingImageManager::GetNumberOfSupportedRates()
+{
+	// We will always support the 4 rates 1x1, 1x2, 2x1, and 2x2.
+	// If the RHI supports larger rates, we can also use 2x4, 4x2, and 4x4, for 7 total.
+	static const int32 NumBaseRates = 4;
+	static const int32 NumExpandedRates = 7;
+
+	return GRHISupportsLargerVariableRateShadingSizes ? NumExpandedRates : NumBaseRates;
+}
+
 FRDGTextureRef FVariableRateShadingImageManager::GetVariableRateShadingImage(FRDGBuilder& GraphBuilder, const FViewInfo& ViewInfo, FVariableRateShadingImageManager::EVRSPassType PassType,
 	FVariableRateShadingImageManager::EVRSSourceType VRSTypesToExclude)
 {
@@ -659,14 +670,16 @@ FRDGTextureRef FVariableRateShadingImageManager::GetForceRateImage(FRDGBuilder&
 {
 	static const TArray<uint32> ValidShadingRates = { VRSSR_1x1, VRSSR_1x2, VRSSR_2x1, VRSSR_2x2, VRSSR_2x4, VRSSR_4x2, VRSSR_4x4 };
 
-	if (RateIndex >= ValidShadingRates.Num())
+	const int32 NumberOfAvailableRates = GetNumberOfSupportedRates();
+
+	if (RateIndex >= NumberOfAvailableRates)
 	{
-		RateIndex = ValidShadingRates.Num() - 1; // Default to maximum shading rate if value exceeds valid rates
+		RateIndex = NumberOfAvailableRates - 1; // Default to maximum shading rate if value exceeds valid rates
 	}
 
 	if (ImageType == EVRSImageType::Disabled)
 	{
-		RateIndex = ValidShadingRates[0]; // Force to minimum shading rate if VRS is disabled for this pass
+		RateIndex = 0; // Force to minimum shading rate if VRS is disabled for this pass
 	}
 
 	FRDGTextureRef ForceShadingRateTexture = GraphBuilder.CreateTexture(GetSRIDesc(), TEXT("ForceShadingRateTexture"));
diff --git a/Engine/Source/Runtime/Renderer/Public/VariableRateShadingImageManager.h b/Engine/Source/Runtime/Renderer/Public/VariableRateShadingImageManager.h
--- a/Engine/Source/Runtime/Renderer/Public/VariableRateShadingImageManager.h
+++ b/Engine/Source/Runtime/Renderer/Public/VariableRateShadingImageManager.h
@@ -109,6 +109,7 @@ public:
 
 	static FIntPoint GetSRITileSize();
 	static FRDGTextureDesc GetSRIDesc();
+	static int32 GetNumberOfSupportedRates();
 
 	void DrawDebugPreview(FRDGBuilder& GraphBuilder, const FSceneViewFamily& ViewFamily, FRDGTextureRef OutputSceneColor);
 
