Fixed ensure with SSAO texture being produced. Capsule shadows can write to the SSAO texture even when Lumen is enabled. Only clear the texture if we are rendering to it.

diff --git a/Engine/Source/Runtime/Renderer/Private/CapsuleShadowRendering.cpp b/Engine/Source/Runtime/Renderer/Private/CapsuleShadowRendering.cpp
--- a/Engine/Source/Runtime/Renderer/Private/CapsuleShadowRendering.cpp
+++ b/Engine/Source/Runtime/Renderer/Private/CapsuleShadowRendering.cpp
@@ -986,19 +986,15 @@ void FDeferredShadingSceneRenderer::RenderIndirectCapsuleShadows(FRDGBuilder& Gr
 		RayTracedShadowsRT = GraphBuilder.CreateTexture(Desc, TEXT("CapsuleShadows.ShadowFactors"));
 	}
 
-	TArray<FRDGTextureRef, TInlineAllocator<2>> RenderTargets;
+	TArray<FRenderTargetBinding, TInlineAllocator<2>> RenderTargets;
 
 	if (SceneTextures.Color.IsValid())
 	{
-		RenderTargets.Add(SceneTextures.Color.Target);
+		RenderTargets.Emplace(SceneTextures.Color.Target, ERenderTargetLoadAction::ELoad);
 	}
 
 	check(SceneTextures.ScreenSpaceAO);
-	if (!SceneTextures.ScreenSpaceAO->HasBeenProduced())
-	{
-		AddClearRenderTargetPass(GraphBuilder, SceneTextures.ScreenSpaceAO);
-	}
-	RenderTargets.Add(SceneTextures.ScreenSpaceAO);
+	RenderTargets.Emplace(SceneTextures.ScreenSpaceAO, SceneTextures.ScreenSpaceAO->HasBeenProduced() ? ERenderTargetLoadAction::ELoad : ERenderTargetLoadAction::EClear);
 
 	for (int32 ViewIndex = 0; ViewIndex < Views.Num(); ViewIndex++)
 	{
@@ -1102,7 +1098,10 @@ void FDeferredShadingSceneRenderer::RenderIndirectCapsuleShadows(FRDGBuilder& Gr
 				FUpsampleCapsuleShadowParameters* PassParameters = GraphBuilder.AllocParameters<FUpsampleCapsuleShadowParameters>();
 				for (int32 Index = 0; Index < RenderTargetCount; ++Index)
 				{
-					PassParameters->RenderTargets[Index] = FRenderTargetBinding(RenderTargets[Index], ERenderTargetLoadAction::ELoad);
+					PassParameters->RenderTargets[Index] = RenderTargets[Index];
+
+					// Only allow clears for the first use of the render target.
+					RenderTargets[Index].SetLoadAction(ERenderTargetLoadAction::ELoad);
 				}
 				PassParameters->SceneTextures = SceneTextures.UniformBuffer;
 
diff --git a/Engine/Source/Runtime/Renderer/Private/IndirectLightRendering.cpp b/Engine/Source/Runtime/Renderer/Private/IndirectLightRendering.cpp
--- a/Engine/Source/Runtime/Renderer/Private/IndirectLightRendering.cpp
+++ b/Engine/Source/Runtime/Renderer/Private/IndirectLightRendering.cpp
@@ -1166,7 +1166,6 @@ void FDeferredShadingSceneRenderer::RenderDiffuseIndirectAndAmbientOcclusion(
 		bool bWritableAmbientOcclusionMask = true;
 		if (ViewPipelineState.AmbientOcclusionMethod == EAmbientOcclusionMethod::Disabled)
 		{
-			ensure(!HasBeenProduced(SceneTextures.ScreenSpaceAO));
 			AmbientOcclusionMask = nullptr;
 			bWritableAmbientOcclusionMask = false;
 		}
