Sync render commands prior to a GC to avoid race conditions with render resource fences. Also moved DestroyRenderThreadResources back to the scene pipe as it can touch UObjects.

diff --git a/Engine/Source/Runtime/Engine/Private/StaticMeshRender.cpp b/Engine/Source/Runtime/Engine/Private/StaticMeshRender.cpp
--- a/Engine/Source/Runtime/Engine/Private/StaticMeshRender.cpp
+++ b/Engine/Source/Runtime/Engine/Private/StaticMeshRender.cpp
@@ -442,8 +442,6 @@ FStaticMeshSceneProxy::FStaticMeshSceneProxy(const FStaticMeshSceneProxyDesc& In
 	{
 		UpdateVisibleInLumenScene();
 	}
-
-	bConstrainToRenderThread = true;
 }
 
 void FStaticMeshSceneProxy::SetEvaluateWorldPositionOffsetInRayTracing(FRHICommandListBase& RHICmdList, bool NewValue)
diff --git a/Engine/Source/Runtime/Engine/Private/UnrealEngine.cpp b/Engine/Source/Runtime/Engine/Private/UnrealEngine.cpp
--- a/Engine/Source/Runtime/Engine/Private/UnrealEngine.cpp
+++ b/Engine/Source/Runtime/Engine/Private/UnrealEngine.cpp
@@ -1652,6 +1652,9 @@ void UEngine::PreGarbageCollect()
 	TRACE_CPUPROFILER_EVENT_SCOPE(UEngine::PreGarbageCollect);
 
 	SendWorldEndOfFrameUpdates();
+
+	// Sync recording render command pipes prior to calling BeginDestroy on resources so that render thread fences are serialized with other render commands.
+	UE::RenderCommandPipe::StopRecording();
 }
 
 float GetLowMemoryGCTimer(const float DefaultTimeBetweenGC)
diff --git a/Engine/Source/Runtime/Renderer/Private/RendererScene.cpp b/Engine/Source/Runtime/Renderer/Private/RendererScene.cpp
--- a/Engine/Source/Runtime/Renderer/Private/RendererScene.cpp
+++ b/Engine/Source/Runtime/Renderer/Private/RendererScene.cpp
@@ -1954,7 +1954,7 @@ void FScene::BatchAddPrimitivesInternal(TArrayView<T*> InPrimitives)
 		FMatrix RenderMatrix = Primitive->GetRenderMatrix();
 		FVector AttachmentRootPosition = Primitive->GetActorPositionForRenderer();
 
-		TArray<FCreateCommand, TInlineAllocator<1>>& CreateCommands = PrimitiveSceneProxy->ShouldConstrainToRenderThread()
+		TArray<FCreateCommand, TInlineAllocator<1>>& CreateCommands = GRenderCommandPipeMode != ERenderCommandPipeMode::All || PrimitiveSceneProxy->ShouldConstrainToRenderThread()
 			? CreateCommandsRenderThread
 			: CreateCommandsScenePipe;
 
@@ -2443,9 +2443,11 @@ FPrimitiveSceneInfo* FScene::GetPrimitiveSceneInfo(const FPersistentPrimitiveInd
 	return GetPrimitiveSceneInfo(PrimitiveIndex);
 }
 
-void FScene::RemovePrimitiveSceneInfo_RenderThread(FPrimitiveSceneInfo* PrimitiveSceneInfo)
+bool FScene::RemovePrimitiveSceneInfo_RenderThread(FPrimitiveSceneInfo* PrimitiveSceneInfo)
 {
-	if (AddedPrimitiveSceneInfos.Remove(PrimitiveSceneInfo))
+	const bool bRemovePendingAdd = AddedPrimitiveSceneInfos.Remove(PrimitiveSceneInfo);
+
+	if (bRemovePendingAdd)
 	{
 		check(PrimitiveSceneInfo->PackedIndex == INDEX_NONE);
 		UpdatedTransforms.Remove(PrimitiveSceneInfo->Proxy);
@@ -2462,6 +2464,8 @@ void FScene::RemovePrimitiveSceneInfo_RenderThread(FPrimitiveSceneInfo* Primitiv
 		check(RemovedPrimitiveSceneInfos.Find(PrimitiveSceneInfo) == nullptr);
 		RemovedPrimitiveSceneInfos.FindOrAdd(PrimitiveSceneInfo);
 	}
+	
+	return !bRemovePendingAdd;
 }
 
 void FScene::RemovePrimitive(UPrimitiveComponent* Primitive)
@@ -2494,6 +2498,7 @@ void FScene::BatchRemovePrimitivesInternal(TArrayView<T*> InPrimitives)
 	struct FDetachCommand
 	{
 		FPrimitiveSceneInfo* PrimitiveSceneInfo;
+		FPrimitiveSceneProxy* PrimitiveSceneProxyIfDestroy;
 		FThreadSafeCounter* AttachmentCounter;
 	};
 
@@ -2509,14 +2514,16 @@ void FScene::BatchRemovePrimitivesInternal(TArrayView<T*> InPrimitives)
 		{
 			FPrimitiveSceneInfo* PrimitiveSceneInfo = PrimitiveSceneProxy->GetPrimitiveSceneInfo();
 
-			if (PrimitiveSceneProxy->ShouldConstrainToRenderThread())
+			const bool bConstrainToRenderThread = GRenderCommandPipeMode != ERenderCommandPipeMode::All || PrimitiveSceneProxy->ShouldConstrainToRenderThread();
+
+			if (bConstrainToRenderThread)
 			{
 				DestroyProxies.Emplace(PrimitiveSceneProxy);
 			}
 
 			// Disassociate the primitive's scene proxy.
 			Primitive->ReleaseSceneProxy();
-			DetachCommands.Add({ PrimitiveSceneInfo, &Primitive->GetSceneData().AttachmentCounter });
+			DetachCommands.Add({ PrimitiveSceneInfo, !bConstrainToRenderThread ? PrimitiveSceneProxy : nullptr, &Primitive->GetSceneData().AttachmentCounter });
 		}
 	}
 
@@ -2527,7 +2534,10 @@ void FScene::BatchRemovePrimitivesInternal(TArrayView<T*> InPrimitives)
 		{
 			for (const FDetachCommand& Command : DetachCommands)
 			{
-				RemovePrimitiveSceneInfo_RenderThread(Command.PrimitiveSceneInfo);
+				if (RemovePrimitiveSceneInfo_RenderThread(Command.PrimitiveSceneInfo) && Command.PrimitiveSceneProxyIfDestroy)
+				{
+					Command.PrimitiveSceneProxyIfDestroy->DestroyRenderThreadResources();
+				}
 				Command.AttachmentCounter->Decrement();
 			}
 		});
@@ -5449,15 +5459,6 @@ void FScene::UpdateAllPrimitiveSceneInfos(FRDGBuilder& GraphBuilder, EUpdateAllP
 		SceneUpdateChangeSetStorage.RemovedPrimitiveSceneInfos.Add(PrimitiveSceneInfo);
 	}
 
-	TArray<FPrimitiveSceneInfo*, SceneRenderingAllocator> DestroyResourcesPrimitives;
-	DestroyResourcesPrimitives.Reserve(RemovedLocalPrimitiveSceneInfos.Num() + DeletedPrimitiveSceneInfos.Num());
-	DestroyResourcesPrimitives = RemovedLocalPrimitiveSceneInfos;
-
-	for (FPrimitiveSceneInfo* Primitive : DeletedPrimitiveSceneInfos)
-	{
-		DestroyResourcesPrimitives.Add(Primitive);
-	}
-
 	TArray<FPrimitiveSceneInfo*, SceneRenderingAllocator> AddedLocalPrimitiveSceneInfos;
 	AddedLocalPrimitiveSceneInfos.Reserve(AddedPrimitiveSceneInfos.Num());
 	for (FPrimitiveSceneInfo* SceneInfo : AddedPrimitiveSceneInfos)
@@ -5467,29 +5468,16 @@ void FScene::UpdateAllPrimitiveSceneInfos(FRDGBuilder& GraphBuilder, EUpdateAllP
 
 	UE::Tasks::FTaskEvent ProcessPrimitiveResourcesTask{ UE_SOURCE_LOCATION };
 
-	if (!AddedLocalPrimitiveSceneInfos.IsEmpty() || !RemovedLocalPrimitiveSceneInfos.IsEmpty())
+	if (!AddedLocalPrimitiveSceneInfos.IsEmpty() && GRenderCommandPipeMode == ERenderCommandPipeMode::All)
 	{
-		ProcessPrimitiveResourcesTask.AddPrerequisites(GraphBuilder.AddCommandListSetupTask([CreateResourcesPrimitives = AddedLocalPrimitiveSceneInfos, DestroyResourcesPrimitives = MoveTemp(DestroyResourcesPrimitives)](FRHICommandListBase& RHICmdList)
+		ProcessPrimitiveResourcesTask.AddPrerequisites(GraphBuilder.AddCommandListSetupTask([CreateResourcesPrimitives = AddedLocalPrimitiveSceneInfos](FRHICommandListBase& RHICmdList)
 		{
+			TRACE_CPUPROFILER_EVENT_SCOPE(Scene::CreatePrimitiveResources);
+			for (FPrimitiveSceneInfo* Primitive : CreateResourcesPrimitives)
 			{
-				TRACE_CPUPROFILER_EVENT_SCOPE(Scene::DestroyPrimitiveResources);
-				for (FPrimitiveSceneInfo* Primitive : DestroyResourcesPrimitives)
+				if (!Primitive->Proxy->ShouldConstrainToRenderThread())
 				{
-					if (!Primitive->Proxy->ShouldConstrainToRenderThread())
-					{
-						Primitive->Proxy->DestroyRenderThreadResources();
-					}
-				}
-			}
-
-			{
-				TRACE_CPUPROFILER_EVENT_SCOPE(Scene::CreatePrimitiveResources);
-				for (FPrimitiveSceneInfo* Primitive : CreateResourcesPrimitives)
-				{
-					if (!Primitive->Proxy->ShouldConstrainToRenderThread())
-					{
-						Primitive->Proxy->CreateRenderThreadResources(RHICmdList);
-					}
+					Primitive->Proxy->CreateRenderThreadResources(RHICmdList);
 				}
 			}
 		}));
diff --git a/Engine/Source/Runtime/Renderer/Private/ScenePrivate.h b/Engine/Source/Runtime/Renderer/Private/ScenePrivate.h
--- a/Engine/Source/Runtime/Renderer/Private/ScenePrivate.h
+++ b/Engine/Source/Runtime/Renderer/Private/ScenePrivate.h
@@ -3715,7 +3715,7 @@ private:
 	 * Removes a primitive from the scene.  Called in the rendering thread by RemovePrimitive.
 	 * @param PrimitiveSceneInfo - The primitive being removed.
 	 */
-	void RemovePrimitiveSceneInfo_RenderThread(FPrimitiveSceneInfo* PrimitiveSceneInfo);
+	bool RemovePrimitiveSceneInfo_RenderThread(FPrimitiveSceneInfo* PrimitiveSceneInfo);
 
 	/** Updates a primitive's transform, called on the rendering thread. */
 	void UpdatePrimitiveTransform_RenderThread(FPrimitiveSceneProxy* PrimitiveSceneProxy, const FBoxSphereBounds& WorldBounds, const FBoxSphereBounds& LocalBounds, const FMatrix& LocalToWorld, const FVector& OwnerPosition, const TOptional<FTransform>& PreviousTransform);
