LFV - Fixing array out of bounds when culling frustum planes aren't populated

diff --git a/Engine/Source/Runtime/Renderer/Private/LocalFogVolumeRendering.cpp b/Engine/Source/Runtime/Renderer/Private/LocalFogVolumeRendering.cpp
--- a/Engine/Source/Runtime/Renderer/Private/LocalFogVolumeRendering.cpp
+++ b/Engine/Source/Runtime/Renderer/Private/LocalFogVolumeRendering.cpp
@@ -211,11 +211,16 @@ static void LocalFogVolumeViewTiledCullingPass(FViewInfo& View, FRDGBuilder& Gra
 			OutVec4f.W *= -1.0f;
 		}
 	};
+	
 	// Using world space plane for now. LFV_TODO: do computation in view space.
-	ConvertPlanToVector4f(PassParameters->LeftPlane,	View.CullingFrustum.Planes[0],	false);
-	ConvertPlanToVector4f(PassParameters->RightPlane,	View.CullingFrustum.Planes[1],	true);
-	ConvertPlanToVector4f(PassParameters->TopPlane,		View.CullingFrustum.Planes[2],	true);
-	ConvertPlanToVector4f(PassParameters->BottomPlane,	View.CullingFrustum.Planes[3],	false);
+	if (View.CullingFrustum.Planes.Num() >= 4)
+	{
+		ConvertPlanToVector4f(PassParameters->LeftPlane,	View.CullingFrustum.Planes[0],	false);
+		ConvertPlanToVector4f(PassParameters->RightPlane,	View.CullingFrustum.Planes[1],	true);
+		ConvertPlanToVector4f(PassParameters->TopPlane,		View.CullingFrustum.Planes[2],	true);
+		ConvertPlanToVector4f(PassParameters->BottomPlane,	View.CullingFrustum.Planes[3],	false);
+	}
+	
 	ConvertPlanToVector4f(PassParameters->NearPlane,	View.NearClippingPlane,			false);
 
 	ERDGPassFlags PassFlag = ERDGPassFlags::Compute; // LFV_TODO try ERDGPassFlags::AsyncCompute later
