Fixes validation failure in TSR binding a dummy Texture2D on a RWTexture2DArray

diff --git a/Engine/Source/Runtime/Renderer/Private/PostProcess/TemporalSuperResolution.cpp b/Engine/Source/Runtime/Renderer/Private/PostProcess/TemporalSuperResolution.cpp
--- a/Engine/Source/Runtime/Renderer/Private/PostProcess/TemporalSuperResolution.cpp
+++ b/Engine/Source/Runtime/Renderer/Private/PostProcess/TemporalSuperResolution.cpp
@@ -991,6 +991,21 @@ static FRDGTextureUAVRef CreateDummyUAV(FRDGBuilder& GraphBuilder, EPixelFormat
 	return GraphBuilder.CreateUAV(DummyTexture);
 };
 
+static FRDGTextureUAVRef CreateDummyUAVArray(FRDGBuilder& GraphBuilder, EPixelFormat PixelFormat)
+{
+	FRDGTextureDesc Desc = FRDGTextureDesc::Create2DArray(
+		FIntPoint(1, 1),
+		PixelFormat,
+		FClearValueBinding::None,
+		/* InFlags = */ TexCreate_ShaderResource | TexCreate_UAV,
+		/* ArraySize = */ 1);
+
+	FRDGTextureRef DummyTexture = GraphBuilder.CreateTexture(Desc, TEXT("TSR.DummyOutput"));
+	GraphBuilder.RemoveUnusedTextureWarning(DummyTexture);
+
+	return GraphBuilder.CreateUAV(DummyTexture);
+};
+
 struct FTSRHistorySliceSequence
 {
 	static constexpr int32 kTransientSliceCount = 2;
@@ -2039,7 +2054,7 @@ FDefaultTemporalUpscaler::FOutputs AddTemporalSuperResolutionPasses(
 		{
 			if (View.bStatePrevViewInfoIsReadOnly)
 			{
-				PassParameters->HistoryGuideOutput = CreateDummyUAV(GraphBuilder, History.GuideArray->Desc.Format);
+				PassParameters->HistoryGuideOutput = CreateDummyUAVArray(GraphBuilder, History.GuideArray->Desc.Format);
 			}
 			else
 			{
