Mobile forward light function bug fix and performance improvements

In the original changelist that added light function support to the mobile forward renderer the proxy geometry had face culling disabled in case the camera was outside the light function proxy geometry. This caused incorrect rendering (noticeable mainly when moving in and out of the light function geometry area) and also a performance penalty. This is fixed with this change.

This change also creates a stencil mask in the initial (non-lightfunction) pass to reduce the number of pixels shaded in the lightfunction draws.
This also enables depth test for lightfunction draws to further reduce the number of pixel shader executions.

diff --git a/Engine/Shaders/Private/MobileLocalLightsBuffer.usf b/Engine/Shaders/Private/MobileLocalLightsBuffer.usf
--- a/Engine/Shaders/Private/MobileLocalLightsBuffer.usf
+++ b/Engine/Shaders/Private/MobileLocalLightsBuffer.usf
@@ -91,6 +91,8 @@ void Main(float4 SvPosition : SV_POSITION
 	, out float3 OutColorA : SV_Target0
 #if POST_PROCESS_LOCAL_LIGHTS == 0
 	, out float4 OutColorB : SV_Target1
+#elif GENERATE_LIGHT_FUNCTION_DEPTH_STENCIL && USE_SCENE_DEPTH_AUX
+	, out float OutDepth : SV_Depth
 #endif
 )
 {
@@ -99,7 +101,8 @@ void Main(float4 SvPosition : SV_POSITION
 	const uint EyeIndex = 0;
 
 	float2 ScreenUV = ScreenPosition * View.ScreenPositionScaleBias.xy + View.ScreenPositionScaleBias.wz;
-	float SceneDepth = CalcSceneDepth(ScreenUV);
+	float DeviceZ = LookupDeviceZ(ScreenUV);
+	float SceneDepth = ConvertFromDeviceZ(DeviceZ);
 
 	float3 TranslatedWorldPosition = mul(float4(ScreenPosition * SceneDepth, SceneDepth, 1), PrimaryView.ScreenToTranslatedWorld).xyz;
 	float3 CameraVector = normalize(TranslatedWorldPosition);
@@ -166,6 +169,10 @@ void Main(float4 SvPosition : SV_POSITION
 	{
 		OutColorA = float3(0, 0, 0);
 		OutColorB = float4(0.5,0.5,0.5, 0.5);
+		#if GENERATE_LIGHT_FUNCTION_DEPTH_STENCIL
+			// Clip to avoid writing to stencil
+			clip(-1);
+		#endif
 	}
 #else
 	if (TotalWeight > 0)
@@ -176,18 +183,24 @@ void Main(float4 SvPosition : SV_POSITION
 		float NoL = max(0, dot(N, L));
 		float3 Diffuse =  (1 / PI) * NoL * Color;
 		OutColorA = 1 + Diffuse;
+		#if GENERATE_LIGHT_FUNCTION_DEPTH_STENCIL && USE_SCENE_DEPTH_AUX
+			// Copy depth value from scene depth aux to the real depth buffer
+			OutDepth = DeviceZ;
+		#endif
 	}
 	else
 	{
 		OutColorA = float3(1, 1, 1);
+		#if GENERATE_LIGHT_FUNCTION_DEPTH_STENCIL
+			// Clip to avoid writing to stencil
+			clip(-1);
+		#endif
 	}
 #endif
 }
 
-
 #else // !LIGHT_FUNCTION
 
-
 float4x4 SvPositionToLight;
 /** Fade distance in x, disabled brightness in y */
 float2 LightFunctionParameters2;
@@ -201,18 +214,6 @@ void MainLightFunction(
 	float2 ScreenUV = SvPositionToBufferUV(SvPosition);
 	SvPosition.z = LookupDeviceZ(ScreenUV);
 	float3 TranslatedWorldPosition = SvPositionToTranslatedWorld(SvPosition);
-	float ViewDistance = length(GetPrimaryView().TranslatedWorldCameraOrigin - TranslatedWorldPosition);
-	float GreyScale;
-	{
-		float4 Hom = mul(float4(SvPosition.xyz, 1), SvPositionToLight);
-		float3 LightVector = Hom.xyz / Hom.w;
-		float3 Color = GetLightFunctionColor(LightVector, TranslatedWorldPosition);
-		GreyScale = dot(Color, .3333f);
-	}
-	float DistanceFadeAlpha = saturate((LightFunctionParameters2.x - ViewDistance) / (LightFunctionParameters2.x * .2f));
-	GreyScale = lerp(LightFunctionParameters2.y, GreyScale, DistanceFadeAlpha);
-	GreyScale = lerp(LightFunctionParameters2.y, GreyScale, LightFunctionParameters.y); 
-	float EncodedLightAttenuation = EncodeLightAttenuation(GreyScale);
 
 	FDeferredLightData LightData = InitDeferredLightFromUniforms();
 	float3 L = LightData.Direction;	// Already normalized
@@ -224,6 +225,26 @@ void MainLightFunction(
 		LightMask = GetLocalLightAttenuation( TranslatedWorldPosition, LightData, ToLight, L );
 	}
 
+	if (LightMask <= 0)
+	{
+		OutColorA = float3(0, 0, 0);
+		return;
+	}
+
+	float EncodedLightAttenuation;
+	{
+		float ViewDistance = length(GetPrimaryView().TranslatedWorldCameraOrigin - TranslatedWorldPosition);
+		float4 Hom = mul(float4(SvPosition.xyz, 1), SvPositionToLight);
+		float3 LightVector = Hom.xyz / Hom.w;
+		float3 Color = GetLightFunctionColor(LightVector, TranslatedWorldPosition);
+		float GreyScale = dot(Color, .3333f);
+
+		float DistanceFadeAlpha = saturate((LightFunctionParameters2.x - ViewDistance) / (LightFunctionParameters2.x * .2f));
+		GreyScale = lerp(LightFunctionParameters2.y, GreyScale, DistanceFadeAlpha);
+		GreyScale = lerp(LightFunctionParameters2.y, GreyScale, LightFunctionParameters.y); 
+		EncodedLightAttenuation = EncodeLightAttenuation(GreyScale);
+	}
+
 	float3 Color = LightData.Color * LightMask * EncodedLightAttenuation;
 
 #if POST_PROCESS_LOCAL_LIGHTS == 0
diff --git a/Engine/Source/Runtime/Renderer/Private/MobileLocalLightsBuffer.cpp b/Engine/Source/Runtime/Renderer/Private/MobileLocalLightsBuffer.cpp
--- a/Engine/Source/Runtime/Renderer/Private/MobileLocalLightsBuffer.cpp
+++ b/Engine/Source/Runtime/Renderer/Private/MobileLocalLightsBuffer.cpp
@@ -72,6 +72,9 @@ class FLocalLightBufferPS : public FGlobalShader
 	DECLARE_GLOBAL_SHADER(FLocalLightBufferPS);
 	SHADER_USE_PARAMETER_STRUCT(FLocalLightBufferPS, FGlobalShader)
 
+	class FGenerateLightFunctionDepthStencil : SHADER_PERMUTATION_BOOL("GENERATE_LIGHT_FUNCTION_DEPTH_STENCIL");
+	using FPermutationDomain = TShaderPermutationDomain<FGenerateLightFunctionDepthStencil>;
+
 	BEGIN_SHADER_PARAMETER_STRUCT(FParameters, )
 		SHADER_PARAMETER_STRUCT_REF(FViewUniformShaderParameters, View)
 		SHADER_PARAMETER_RDG_UNIFORM_BUFFER(FForwardLightData, ForwardLightData)
@@ -200,6 +203,19 @@ static bool TryGetLightFunctionShaders(FMaterialRenderProxy const*& OutMaterialP
 	return false;
 }
 
+template <bool DepthWrite>
+using LightFunctionMainPassDepthStencilState = TStaticDepthStencilState<
+	DepthWrite, CF_Always,
+	true, CF_Always, SO_Keep, SO_Keep, SO_Replace,
+	true, CF_Always, SO_Keep, SO_Keep, SO_Replace,
+	0, STENCIL_MOBILE_LIGHTFUNCTION_MASK>;
+
+template <ECompareFunction CompareFunction>
+using LightFunctionMaterialPassDepthStencilState = TStaticDepthStencilState<
+	false, CompareFunction,
+	true, CF_Equal, SO_Keep, SO_Keep, SO_Keep,
+	true, CF_Equal, SO_Keep, SO_Keep, SO_Keep,
+	STENCIL_MOBILE_LIGHTFUNCTION_MASK, 0>;
 
 void FMobileSceneRenderer::RenderMobileLocalLightsBuffer(FRDGBuilder& GraphBuilder, FSceneTextures& SceneTextures, bool bIsPrepass, const FSortedLightSetSceneInfo& SortedLights)
 {
@@ -250,6 +266,26 @@ void FMobileSceneRenderer::RenderMobileLocalLightsBuffer(FRDGBuilder& GraphBuild
 				FIntVector(FMath::DivideAndRoundUp<uint32>(GroupSize.Y * GroupSize.X, GLocalLightPrepassTileSizeX), 1, 1));
 		}
 
+		bool bRenderLightFunctions = false;
+		{
+			// Check if there are any light functions
+			for (const FSortedLightSceneInfo& SortedLightSceneInfo : SortedLights.SortedLights)
+			{
+				// Directional lights are currently not supported
+				if (SortedLightSceneInfo.SortKey.Fields.LightType == LightType_Directional)
+				{
+					continue;
+				}
+
+				if (!SortedLightSceneInfo.SortKey.Fields.bLightFunction)
+				{
+					continue;
+				}
+
+				bRenderLightFunctions = true;
+				break;
+			}
+		}
 
 		{
 			FLocalLightBufferPrepassParameters* PassParameters = GraphBuilder.AllocParameters<FLocalLightBufferPrepassParameters>();
@@ -262,6 +298,24 @@ void FMobileSceneRenderer::RenderMobileLocalLightsBuffer(FRDGBuilder& GraphBuild
 			{
 				PassParameters->RenderTargets[0] = FRenderTargetBinding(SceneTextures.Color.Resolve, ERenderTargetLoadAction::ELoad);
 			}
+
+			bool bRestoreDepthBuffer = false;
+			if (bRenderLightFunctions)
+			{
+				if (MobileRequiresSceneDepthAux(ShaderPlatform))
+				{
+					// In this case the depth buffer is typically memoryless and not preserved, so we restore it from SceneDepthAux
+					PassParameters->RenderTargets.DepthStencil = FDepthStencilBinding(SceneTextures.Depth.Resolve, ERenderTargetLoadAction::EClear, ERenderTargetLoadAction::EClear, FExclusiveDepthStencil::DepthWrite_StencilWrite);
+					bRestoreDepthBuffer = true;
+				}
+				else
+				{
+					// If the main depth buffer is not memoryless we rely on STENCIL_MOBILE_LIGHTFUNCTION_MASK being cleared to 0 already
+					const ERenderTargetLoadAction LoadAction = EnumHasAnyFlags(SceneTextures.Depth.Resolve->Desc.Flags, TexCreate_Memoryless) ? ERenderTargetLoadAction::EClear : ERenderTargetLoadAction::ELoad;
+					PassParameters->RenderTargets.DepthStencil = FDepthStencilBinding(SceneTextures.Depth.Resolve, LoadAction, LoadAction, FExclusiveDepthStencil::DepthRead_StencilWrite);
+				}
+			}
+		
 			PassParameters->SceneTextures = SceneTextures.GetSceneTextureShaderParameters(View.FeatureLevel);
 
 			PassParameters->VS.View = GetShaderBinding(View.ViewUniformBuffer);
@@ -271,20 +325,22 @@ void FMobileSceneRenderer::RenderMobileLocalLightsBuffer(FRDGBuilder& GraphBuild
 			PassParameters->PS.View = GetShaderBinding(View.ViewUniformBuffer);
 
 			auto VertexShader = View.ShaderMap->GetShader<FLocalLightBufferVS>();
-			auto PixelShader = View.ShaderMap->GetShader<FLocalLightBufferPS>();
+
+			FLocalLightBufferPS::FPermutationDomain PermutationVectorPS;
+			PermutationVectorPS.Set<FLocalLightBufferPS::FGenerateLightFunctionDepthStencil>(bRenderLightFunctions);
+			auto PixelShader = View.ShaderMap->GetShader<FLocalLightBufferPS>(PermutationVectorPS);
 
 			GraphBuilder.AddPass(
 				RDG_EVENT_NAME("RenderMobileLocalLightsBuffer %s", bIsPrepass ? TEXT("Prepass") : TEXT("PostProcess")),
 				PassParameters,
 				ERDGPassFlags::Raster,
-				[PassParameters, VertexShader, PixelShader, &View, GroupSize, bIsPrepass, &SortedLights](FRHICommandList& RHICmdList)
+				[PassParameters, VertexShader, PixelShader, &View, GroupSize, bIsPrepass, &SortedLights, bRenderLightFunctions, bRestoreDepthBuffer](FRHICommandList& RHICmdList)
 				{
 					FGraphicsPipelineStateInitializer GraphicsPSOInit;
 					RHICmdList.ApplyCachedRenderTargets(GraphicsPSOInit);
 
 					RHICmdList.SetViewport(View.ViewRect.Min.X, View.ViewRect.Min.Y, 0.0f, View.ViewRect.Max.X, View.ViewRect.Max.Y, 1.0f);
 					GraphicsPSOInit.RasterizerState = TStaticRasterizerState<FM_Solid, CM_None>::GetRHI();
-					GraphicsPSOInit.DepthStencilState = TStaticDepthStencilState<false, CF_Always>::GetRHI();
 					if (bIsPrepass)
 					{
 						GraphicsPSOInit.BlendState = TStaticBlendState<>::GetRHI();
@@ -293,13 +349,33 @@ void FMobileSceneRenderer::RenderMobileLocalLightsBuffer(FRDGBuilder& GraphBuild
 					{
 						GraphicsPSOInit.BlendState = TStaticBlendState<CW_RGB, BO_Add, BF_DestColor, BF_Zero>::GetRHI();
 					}
+					
+					// The main non-lightfunction pass creates a stencil mask of the lit area
+					// Later lightfunction passes use it for stencil test to avoid redundant PS execution
+					uint32 StencilRef = 0;
+					if (bRenderLightFunctions)
+					{
+						StencilRef = STENCIL_MOBILE_LIGHTFUNCTION_MASK;
+						if (bRestoreDepthBuffer)
+						{
+							GraphicsPSOInit.DepthStencilState = LightFunctionMainPassDepthStencilState<true>::GetRHI();
+						}
+						else
+						{
+							GraphicsPSOInit.DepthStencilState = LightFunctionMainPassDepthStencilState<false>::GetRHI();
+						}
+					}
+					else
+					{
+						GraphicsPSOInit.DepthStencilState = TStaticDepthStencilState<false, CF_Always>::GetRHI();
+					}
 
 					GraphicsPSOInit.PrimitiveType = PT_TriangleList;
 					GraphicsPSOInit.BoundShaderState.VertexDeclarationRHI = GTileVertexDeclaration.VertexDeclarationRHI;
 					GraphicsPSOInit.BoundShaderState.VertexShaderRHI = VertexShader.GetVertexShader();
 					GraphicsPSOInit.BoundShaderState.PixelShaderRHI = PixelShader.GetPixelShader();
 
-					SetGraphicsPipelineState(RHICmdList, GraphicsPSOInit, 0);
+					SetGraphicsPipelineState(RHICmdList, GraphicsPSOInit, StencilRef);
 
 					SetShaderParameters(RHICmdList, VertexShader, VertexShader.GetVertexShader(), PassParameters->VS);
 					SetShaderParameters(RHICmdList, PixelShader, PixelShader.GetPixelShader(), PassParameters->PS);
@@ -313,6 +389,10 @@ void FMobileSceneRenderer::RenderMobileLocalLightsBuffer(FRDGBuilder& GraphBuild
 						2,
 						GroupSize.X * GroupSize.Y);
 
+					if (!bRenderLightFunctions)
+					{
+						return;
+					}
 
 					// Draws a pass for each visible light with a light function and blends on top of the color texture 
 					// generated by last pass
@@ -369,12 +449,19 @@ void FMobileSceneRenderer::RenderMobileLocalLightsBuffer(FRDGBuilder& GraphBuild
 						{
 							// Render backfaces with depth tests disabled since the camera is inside (or close to inside) the light function geometry
 							GraphicsPSOInit.RasterizerState = View.bReverseCulling ? TStaticRasterizerState<FM_Solid, CM_CW>::GetRHI() : TStaticRasterizerState<FM_Solid, CM_CCW>::GetRHI();
+							GraphicsPSOInit.DepthStencilState = LightFunctionMaterialPassDepthStencilState<CF_Always>::GetRHI();
+						}
+						else
+						{
+							// Render frontfaces with depth test so that less pixel are shaded
+							GraphicsPSOInit.RasterizerState = View.bReverseCulling ? TStaticRasterizerState<FM_Solid, CM_CCW>::GetRHI() : TStaticRasterizerState<FM_Solid, CM_CW>::GetRHI();
+							GraphicsPSOInit.DepthStencilState = LightFunctionMaterialPassDepthStencilState<CF_DepthNearOrEqual>::GetRHI();
 						}
 
 						// Set the light's scissor rectangle.
 						LightSceneInfo->Proxy->SetScissorRect(RHICmdList, View, View.ViewRect);
 
-						SetGraphicsPipelineState(RHICmdList, GraphicsPSOInit, 0);
+						SetGraphicsPipelineState(RHICmdList, GraphicsPSOInit, StencilRef);
 
 						SetShaderParameters(RHICmdList, LightFunctionVertexShader, LightFunctionVertexShader.GetVertexShader(), ParametersVS);
 						SetShaderParametersMixedPS(RHICmdList, LightFunctionPixelShader, ParametersPS, View, MaterialProxyForRendering, *MaterialForRendering);
diff --git a/Engine/Source/Runtime/Renderer/Private/PostProcess/SceneRenderTargets.h b/Engine/Source/Runtime/Renderer/Private/PostProcess/SceneRenderTargets.h
--- a/Engine/Source/Runtime/Renderer/Private/PostProcess/SceneRenderTargets.h
+++ b/Engine/Source/Runtime/Renderer/Private/PostProcess/SceneRenderTargets.h
@@ -50,7 +50,9 @@
 #define STENCIL_LIGHTING_CHANNELS_MASK(Value) uint8(((Value) & 0x7) << STENCIL_LIGHTING_CHANNELS_BIT_ID)
 
 // Mobile specific
-// Store shading model into stencil [1-2] bits
+// Store shading model into stencil [1-2] bits (only used in mobile deferred)
 #define GET_STENCIL_MOBILE_SM_MASK(Value) uint8(((Value) & 0x3) << 1)
+// Forward local light buffer mask for light function - bit 2
+#define STENCIL_MOBILE_LIGHTFUNCTION_MASK uint8(1 << 2)
 // Sky material mask - bit 3
 #define STENCIL_MOBILE_SKY_MASK uint8(1 << 3)
diff --git a/Engine/Source/Runtime/Renderer/Private/SceneTextures.cpp b/Engine/Source/Runtime/Renderer/Private/SceneTextures.cpp
--- a/Engine/Source/Runtime/Renderer/Private/SceneTextures.cpp
+++ b/Engine/Source/Runtime/Renderer/Private/SceneTextures.cpp
@@ -624,7 +624,7 @@ void FSceneTextures::InitializeViewFamily(FRDGBuilder& GraphBuilder, FViewFamily
 		FRDGTextureDesc MobileLocalLightTextureADesc = FRDGTextureDesc::Create2D(Config.Extent, PF_FloatR11G11B10, FClearValueBinding::Transparent, TexCreate_RenderTargetable | TexCreate_ShaderResource);
 		SceneTextures.MobileLocalLightTextureA = GraphBuilder.CreateTexture(MobileLocalLightTextureADesc, TEXT("MobileLocalLightTextureA"));
 
-		FRDGTextureDesc MobileLocalLightTextureBDesc = FRDGTextureDesc::Create2D(Config.Extent, PF_A2B10G10R10, FClearValueBinding::Transparent, TexCreate_RenderTargetable | TexCreate_ShaderResource);
+		FRDGTextureDesc MobileLocalLightTextureBDesc = FRDGTextureDesc::Create2D(Config.Extent, PF_A2B10G10R10, FClearValueBinding::DefaultNormal8Bit, TexCreate_RenderTargetable | TexCreate_ShaderResource);
 		SceneTextures.MobileLocalLightTextureB = GraphBuilder.CreateTexture(MobileLocalLightTextureBDesc, TEXT("MobileLocalLightTextureB"));
 	}
 
