LFV - fixed cvar name.

#prelfight shader

diff --git a/Engine/Source/Runtime/Engine/Classes/Engine/RendererSettings.h b/Engine/Source/Runtime/Engine/Classes/Engine/RendererSettings.h
--- a/Engine/Source/Runtime/Engine/Classes/Engine/RendererSettings.h
+++ b/Engine/Source/Runtime/Engine/Classes/Engine/RendererSettings.h
@@ -604,7 +604,7 @@ class URendererSettings : public UDeveloperSettings
 	FVector TranslucentSortAxis;
 
 	UPROPERTY(config, EditAnywhere, Category=Translucency, meta=(
-		ConsoleVariable="r.LocalFogVolume.ApplyOnTransclucent",
+		ConsoleVariable="r.LocalFogVolume.ApplyOnTranslucent",
 		ToolTip="Allow translucency to be rendered to a separate render targeted and composited after depth of field. Prevents translucency from appearing out of focus."))
 	uint32 bLocalFogVolumeApplyOnTranslucent:1;
 
diff --git a/Engine/Source/Runtime/Engine/Private/ShaderCompiler/ShaderCompiler.cpp b/Engine/Source/Runtime/Engine/Private/ShaderCompiler/ShaderCompiler.cpp
--- a/Engine/Source/Runtime/Engine/Private/ShaderCompiler/ShaderCompiler.cpp
+++ b/Engine/Source/Runtime/Engine/Private/ShaderCompiler/ShaderCompiler.cpp
@@ -7934,7 +7934,7 @@ void GlobalBeginCompileShader(
 	}
 
 	{
-		static IConsoleVariable* CVar = IConsoleManager::Get().FindConsoleVariable(TEXT("r.LocalFogVolume.ApplyOnTransclucent"));
+		static IConsoleVariable* CVar = IConsoleManager::Get().FindConsoleVariable(TEXT("r.LocalFogVolume.ApplyOnTranslucent"));
 		const bool bLocalFogVolumesApplyOnTranclucent = CVar && CVar->GetInt() > 0;
 		SET_SHADER_DEFINE(Input.Environment, PROJECT_LOCALFOGVOLUME_APPLYONTRANSLUCENT, (bLocalFogVolumesApplyOnTranclucent ? 1 : 0));
 	}
diff --git a/Engine/Source/Runtime/RenderCore/Private/Shader.cpp b/Engine/Source/Runtime/RenderCore/Private/Shader.cpp
--- a/Engine/Source/Runtime/RenderCore/Private/Shader.cpp
+++ b/Engine/Source/Runtime/RenderCore/Private/Shader.cpp
@@ -1936,7 +1936,7 @@ void ShaderMapAppendKeyString(EShaderPlatform Platform, FString& KeyString)
 	}
 
 	{
-		static IConsoleVariable* CVar = IConsoleManager::Get().FindConsoleVariable(TEXT("r.LocalFogVolume.ApplyOnTransclucent"));
+		static IConsoleVariable* CVar = IConsoleManager::Get().FindConsoleVariable(TEXT("r.LocalFogVolume.ApplyOnTranslucent"));
 		const bool bLocalFogVolumesApplyOnTranclucent = CVar && CVar->GetInt() > 0;
 		if (bLocalFogVolumesApplyOnTranclucent)
 		{
diff --git a/Engine/Source/Runtime/Renderer/Private/LocalFogVolumeRendering.cpp b/Engine/Source/Runtime/Renderer/Private/LocalFogVolumeRendering.cpp
--- a/Engine/Source/Runtime/Renderer/Private/LocalFogVolumeRendering.cpp
+++ b/Engine/Source/Runtime/Renderer/Private/LocalFogVolumeRendering.cpp
@@ -14,8 +14,8 @@ static TAutoConsoleVariable<int32> CVarLocalFogVolume(
 	TEXT("LocalFogVolume components are rendered when this is not 0, otherwise ignored.\n"),
 	ECVF_RenderThreadSafe);
 
-static TAutoConsoleVariable<int32> CVarLocalFogVolumeApplyOnTransclucent(
-	TEXT("r.LocalFogVolume.ApplyOnTransclucent"), 0,
+static TAutoConsoleVariable<int32> CVarLocalFogVolumeApplyOnTranslucent(
+	TEXT("r.LocalFogVolume.ApplyOnTranslucent"), 0,
 	TEXT("Project settings enabling the sampling of local fog volumes on translucent elements.\n"),
 	ECVF_ReadOnly | ECVF_RenderThreadSafe);
 
