LFV - Make sure wealways  initialise cull plane eq to zero to be valid and ignore culling when no possible to get clip planes.

#prelfight shader

diff --git a/Engine/Source/Runtime/Renderer/Private/LocalFogVolumeRendering.cpp b/Engine/Source/Runtime/Renderer/Private/LocalFogVolumeRendering.cpp
--- a/Engine/Source/Runtime/Renderer/Private/LocalFogVolumeRendering.cpp
+++ b/Engine/Source/Runtime/Renderer/Private/LocalFogVolumeRendering.cpp
@@ -213,12 +213,21 @@ static void LocalFogVolumeViewTiledCullingPass(FViewInfo& View, FRDGBuilder& Gra
 	};
 	
 	// Using world space plane for now. LFV_TODO: do computation in view space.
-	if (View.CullingFrustum.Planes.Num() >= 4)
+	if (View.ViewFrustum.Planes.Num() >= 4)
 	{
-		ConvertPlanToVector4f(PassParameters->LeftPlane,	View.CullingFrustum.Planes[0],	false);
-		ConvertPlanToVector4f(PassParameters->RightPlane,	View.CullingFrustum.Planes[1],	true);
-		ConvertPlanToVector4f(PassParameters->TopPlane,		View.CullingFrustum.Planes[2],	true);
-		ConvertPlanToVector4f(PassParameters->BottomPlane,	View.CullingFrustum.Planes[3],	false);
+		// We use the view frustum witch matches the rendering frustum even when in stereo mode (not monoscopic).
+		ConvertPlanToVector4f(PassParameters->LeftPlane,	View.ViewFrustum.Planes[0],	false);
+		ConvertPlanToVector4f(PassParameters->RightPlane,	View.ViewFrustum.Planes[1],	true);
+		ConvertPlanToVector4f(PassParameters->TopPlane,		View.ViewFrustum.Planes[2],	true);
+		ConvertPlanToVector4f(PassParameters->BottomPlane,	View.ViewFrustum.Planes[3],	false);
+	}
+	else
+	{
+		// Disable culling and make each volume visible.
+		PassParameters->LeftPlane	= FVector4f::Zero();
+		PassParameters->RightPlane	= FVector4f::Zero();
+		PassParameters->TopPlane	= FVector4f::Zero();
+		PassParameters->BottomPlane	= FVector4f::Zero();
 	}
 	
 	ConvertPlanToVector4f(PassParameters->NearPlane,	View.NearClippingPlane,			false);
