Fixes misleading VisualizeTemporalUpscaler show flag that still showed post DOF translucency with TAA even thaw it's not technically a direct input.

diff --git a/Engine/Source/Runtime/Renderer/Private/PostProcess/VisualizeTemporalUpscaler.cpp b/Engine/Source/Runtime/Renderer/Private/PostProcess/VisualizeTemporalUpscaler.cpp
--- a/Engine/Source/Runtime/Renderer/Private/PostProcess/VisualizeTemporalUpscaler.cpp
+++ b/Engine/Source/Runtime/Renderer/Private/PostProcess/VisualizeTemporalUpscaler.cpp
@@ -71,6 +71,7 @@ FScreenPassTexture AddVisualizeTemporalUpscalerPass(FRDGBuilder& GraphBuilder, c
 		}
 
 		// Translucency
+		if (Inputs.TAAConfig == EMainTAAPassConfig::TSR)
 		{
 			FVisualizeBufferTile& Tile = Tiles[4 * 0 + 2];
 			if (Inputs.Inputs.PostDOFTranslucencyResources.IsValid())
@@ -87,6 +88,7 @@ FScreenPassTexture AddVisualizeTemporalUpscalerPass(FRDGBuilder& GraphBuilder, c
 		}
 
 		// Translucency alpha
+		if (Inputs.TAAConfig == EMainTAAPassConfig::TSR)
 		{
 			FVisualizeBufferTile& Tile = Tiles[4 * 0 + 3];
 			if (Inputs.Inputs.PostDOFTranslucencyResources.IsValid())
