Remove dependency on INVALID_PRIMITIVE_ID to hide instances of force-hidden primitives, instead using INSTANCE_SCENE_DATA_FLAG_HIDDEN.

diff --git a/Engine/Source/Runtime/Engine/Private/InstanceUniformShaderParameters.cpp b/Engine/Source/Runtime/Engine/Private/InstanceUniformShaderParameters.cpp
--- a/Engine/Source/Runtime/Engine/Private/InstanceUniformShaderParameters.cpp
+++ b/Engine/Source/Runtime/Engine/Private/InstanceUniformShaderParameters.cpp
@@ -23,10 +23,11 @@ void FInstanceSceneShaderData::Build
 	uint32 InstanceFlags,
 	uint32 LastUpdateFrame,
 	uint32 CustomDataCount,
-	float RandomID
+	float RandomID,
+	bool bIsVisible
 )
 {
-	BuildInternal(PrimitiveId, RelativeId, InstanceFlags, LastUpdateFrame, CustomDataCount, RandomID, FRenderTransform::Identity);
+	BuildInternal(PrimitiveId, RelativeId, InstanceFlags, LastUpdateFrame, CustomDataCount, RandomID, FRenderTransform::Identity, bIsVisible);
 }
 
 void FInstanceSceneShaderData::Build
@@ -38,7 +39,8 @@ void FInstanceSceneShaderData::Build
 	uint32 CustomDataCount,
 	float RandomID,
 	const FRenderTransform& LocalToPrimitive,
-	const FRenderTransform& PrimitiveToWorld
+	const FRenderTransform& PrimitiveToWorld,
+	bool bIsVisible
 )
 {
 	FRenderTransform LocalToWorld = LocalToPrimitive * PrimitiveToWorld;
@@ -46,7 +48,7 @@ void FInstanceSceneShaderData::Build
 	// Remove shear
 	LocalToWorld.Orthogonalize();
 
-	BuildInternal(PrimitiveId, RelativeId, InstanceFlags, LastUpdateFrame, CustomDataCount, RandomID, LocalToWorld);
+	BuildInternal(PrimitiveId, RelativeId, InstanceFlags, LastUpdateFrame, CustomDataCount, RandomID, LocalToWorld, bIsVisible);
 }
 
 void FInstanceSceneShaderData::BuildInternal
@@ -57,7 +59,8 @@ void FInstanceSceneShaderData::BuildInternal
 	uint32 LastUpdateFrame,
 	uint32 CustomDataCount,
 	float RandomID,
-	const FRenderTransform& LocalToWorld // Assumes shear has been removed already
+	const FRenderTransform& LocalToWorld, // Assumes shear has been removed already
+	bool bIsVisible
 )
 {
 	// Note: layout must match GetInstanceData in SceneData.ush and InitializeInstanceSceneData in GPUSceneWriter.ush
@@ -71,6 +74,11 @@ void FInstanceSceneShaderData::BuildInternal
 		InstanceFlags &= ~INSTANCE_SCENE_DATA_FLAG_DETERMINANT_SIGN;
 	}
 
+	if (!bIsVisible)
+	{
+		InstanceFlags |= INSTANCE_SCENE_DATA_FLAG_HIDDEN;
+	}
+
 	checkSlow((PrimitiveId		& 0x000FFFFF) == PrimitiveId);
 	checkSlow((InstanceFlags	& 0x00000FFF) == InstanceFlags);
 	checkSlow((RelativeId		& 0x00FFFFFF) == RelativeId);
diff --git a/Engine/Source/Runtime/Engine/Public/InstanceUniformShaderParameters.h b/Engine/Source/Runtime/Engine/Public/InstanceUniformShaderParameters.h
--- a/Engine/Source/Runtime/Engine/Public/InstanceUniformShaderParameters.h
+++ b/Engine/Source/Runtime/Engine/Public/InstanceUniformShaderParameters.h
@@ -93,7 +93,8 @@ public:
 		uint32 InstanceFlags,
 		uint32 LastUpdateFrame,
 		uint32 CustomDataCount,
-		float RandomID
+		float RandomID,
+		bool bIsVisible = true
 	);
 
 	ENGINE_API void Build
@@ -105,7 +106,8 @@ public:
 		uint32 CustomDataCount,
 		float RandomID,
 		const FRenderTransform& LocalToPrimitive,
-		const FRenderTransform& PrimitiveToWorld
+		const FRenderTransform& PrimitiveToWorld,
+		bool bIsVisible = true
 	);
 
 	ENGINE_API void BuildInternal
@@ -116,7 +118,8 @@ public:
 		uint32 LastUpdateFrame,
 		uint32 CustomDataCount,
 		float RandomID,
-		const FRenderTransform& LocalToWorld
+		const FRenderTransform& LocalToWorld,
+		bool bIsVisible = true
 	);
 
 	TStaticArray<FVector4f, UnCompressedTransformDataStrideInFloat4s> Data;
diff --git a/Engine/Source/Runtime/Renderer/Private/GPUScene.cpp b/Engine/Source/Runtime/Renderer/Private/GPUScene.cpp
--- a/Engine/Source/Runtime/Renderer/Private/GPUScene.cpp
+++ b/Engine/Source/Runtime/Renderer/Private/GPUScene.cpp
@@ -296,6 +296,7 @@ struct FInstanceUploadInfo
 	FRenderTransform PrevPrimitiveToWorld;
 	int32 PrimitiveID = INDEX_NONE;
 	uint32 LastUpdateSceneFrameNumber = ~uint32(0);
+	bool bIsPrimitiveForceHidden = false;
 };
 
 void ValidateInstanceUploadInfo(const FInstanceUploadInfo& UploadInfo, FRDGBuffer* InstancePayloadDataBuffer)
@@ -365,13 +366,6 @@ struct FLightMapUploadInfo
 	int32 LightmapDataOffset = 0;
 };
 
-// TODO: Temporary hack : For FPrimitiveSceneProxy::IsForceHidden() to work with Nanite proxies, return an invalid primitive ID if IsForceHidden() returns true.
-static FORCEINLINE int32 GetPrimitiveID(const FScene& InScene, const int32 InPrimitiveID)
-{
-	const FPrimitiveSceneProxy* PrimitiveSceneProxy = InScene.PrimitiveSceneProxies[InPrimitiveID];
-	return (PrimitiveSceneProxy->IsNaniteMesh() && PrimitiveSceneProxy->IsForceHidden()) ? INVALID_PRIMITIVE_ID : InPrimitiveID;
-}
-
 /**
  * Implements a thin data abstraction such that the UploadGeneral function can upload primitive data from
  * both scene primitives and dynamic primitives (which are not stored in the same way). 
@@ -467,8 +461,9 @@ struct FUploadDataSourceAdapterScenePrimitives
 		InstanceUploadInfo.InstancePayloadDataStride = PrimitiveSceneInfo->GetInstancePayloadDataStride();
 
 		InstanceUploadInfo.LastUpdateSceneFrameNumber = SceneFrameNumber;
-		InstanceUploadInfo.PrimitiveID = GetPrimitiveID(Scene, PrimitiveID);
+		InstanceUploadInfo.PrimitiveID = PrimitiveID;
 		InstanceUploadInfo.PrimitiveToWorld = FLargeWorldRenderScalar::MakeToRelativeWorldMatrix(AbsoluteOrigin.GetTileOffset(), LocalToWorld);
+		InstanceUploadInfo.bIsPrimitiveForceHidden = PrimitiveSceneProxy->IsForceHidden();
 
 		{
 			bool bHasPrecomputedVolumetricLightmap{};
@@ -693,10 +688,10 @@ void FGPUScene::UpdateInternal(FRDGBuilder& GraphBuilder, FSceneUniformBuffer& S
 {
 	LLM_SCOPE_BYTAG(GPUScene);
 
-	ensure(bInBeginEndBlock);
-	ensure(bIsEnabled == UseGPUScene(GMaxRHIShaderPlatform, Scene.GetFeatureLevel()));
-	ensure(NumScenePrimitives == Scene.Primitives.Num());
-	ensure(DynamicPrimitivesOffset >= Scene.Primitives.Num());
+	check(bInBeginEndBlock);
+	check(bIsEnabled == UseGPUScene(GMaxRHIShaderPlatform, Scene.GetFeatureLevel()));
+	check(NumScenePrimitives == Scene.Primitives.Num());
+	check(DynamicPrimitivesOffset >= Scene.Primitives.Num());
 
 	CSV_CUSTOM_STAT(GPUScene, InstanceAllocMaxSize, InstanceSceneDataAllocator.GetMaxSize(), ECsvCustomStatOp::Set);
 	CSV_CUSTOM_STAT(GPUScene, InstanceAllocUsedSize, InstanceSceneDataAllocator.GetSparselyAllocatedSize(), ECsvCustomStatOp::Set);
@@ -771,7 +766,7 @@ void FGPUScene::UpdateInternal(FRDGBuilder& GraphBuilder, FSceneUniformBuffer& S
 			{
 				const FPrimitiveSceneInfo* PrimitiveSceneInfo = Scene.Primitives[PrimitiveId];
 				check(PrimitiveSceneInfo->GetInstanceSceneDataOffset() >= 0 || PrimitiveSceneInfo->GetNumInstanceSceneDataEntries() == 0);
-				IdOnlyUpdateData.Add(PrimitiveSceneInfo->GetInstanceSceneDataOffset(), PrimitiveSceneInfo->GetNumInstanceSceneDataEntries(), GetPrimitiveID(Scene, PrimitiveId));
+				IdOnlyUpdateData.Add(PrimitiveSceneInfo->GetInstanceSceneDataOffset(), PrimitiveSceneInfo->GetNumInstanceSceneDataEntries(), PrimitiveId);
 			}
 		}
 		AddUpdatePrimitiveIdsPass(GraphBuilder, IdOnlyUpdateData);
@@ -797,9 +792,9 @@ void FGPUScene::UpdateBufferState(FRDGBuilder& GraphBuilder, FSceneUniformBuffer
 {
 	LLM_SCOPE_BYTAG(GPUScene);
 
-	ensure(bInBeginEndBlock);
-	ensure(bIsEnabled == UseGPUScene(GMaxRHIShaderPlatform, Scene.GetFeatureLevel()));
-	ensure(NumScenePrimitives == Scene.Primitives.Num());
+	check(bInBeginEndBlock);
+	check(bIsEnabled == UseGPUScene(GMaxRHIShaderPlatform, Scene.GetFeatureLevel()));
+	check(NumScenePrimitives == Scene.Primitives.Num());
 
 	RDG_GPU_MASK_SCOPE(GraphBuilder, FRHIGPUMask::All());
 
@@ -1157,7 +1152,8 @@ void FGPUScene::UploadGeneral(FRDGBuilder& GraphBuilder, FScene& Scene, FRDGExte
 							UploadInfo.InstanceCustomDataCount,
 							RandomID,
 							SceneData.LocalToPrimitive,
-							UploadInfo.PrimitiveToWorld
+							UploadInfo.PrimitiveToWorld,
+							!UploadInfo.bIsPrimitiveForceHidden
 						);
 
 						// RefIndex* BufferState.InstanceSceneDataSOAStride + UploadInfo.InstanceSceneDataOffset + InstanceIndex
@@ -1446,6 +1442,7 @@ struct FUploadDataSourceAdapterDynamicPrimitives
 			InstanceUploadInfo.InstanceEditorData			= TConstArrayView<uint32>();
 #endif
 
+			InstanceUploadInfo.bIsPrimitiveForceHidden = false;
 			// upload dummies where applicable
 			if (InstanceUploadInfo.PrimitiveInstances.Num() == 0)
 			{
@@ -2053,9 +2050,6 @@ void FGPUScene::GetWriteParameters(FRDGBuilder& GraphBuilder, FGPUSceneWriterPar
 	GPUSceneWriterParametersOut.GPUScenePrimitiveSceneDataRW = GraphBuilder.CreateUAV(BufferState.PrimitiveBuffer, ERDGUnorderedAccessViewFlags::SkipBarrier);
 }
 
-/**
- * Compute shader to project and invalidate the rectangles of given instances.
- */
 class FGPUSceneSetInstancePrimitiveIdCS : public FGlobalShader
 {
 	DECLARE_GLOBAL_SHADER(FGPUSceneSetInstancePrimitiveIdCS);
