Disabled fully async material uniform expression caching on mobile platforms as OpenGL doesn't support async uniform buffer creation.

diff --git a/Engine/Source/Runtime/Renderer/Private/SceneRendering.cpp b/Engine/Source/Runtime/Renderer/Private/SceneRendering.cpp
--- a/Engine/Source/Runtime/Renderer/Private/SceneRendering.cpp
+++ b/Engine/Source/Runtime/Renderer/Private/SceneRendering.cpp
@@ -3420,7 +3420,7 @@ IVisibilityTaskData* FSceneRenderer::OnRenderBegin(FRDGBuilder& GraphBuilder)
 		AsyncOps |= EUpdateAllPrimitiveSceneInfosAsyncOps::CacheMeshDrawCommands;
 	}
 
-	if (GAsyncCacheMaterialUniformExpressions > 0)
+	if (GAsyncCacheMaterialUniformExpressions > 0 && !IsMobilePlatform(ShaderPlatform))
 	{
 		AsyncOps |= EUpdateAllPrimitiveSceneInfosAsyncOps::CacheMaterialUniformExpressions;
 	}
