Debug Visualization: count the number of rays hits per scene primitive
generate heatmap and hit stats info of the top K primitives.

diff --git a/Engine/Shaders/Private/RayTracing/RayTracingDebug.usf b/Engine/Shaders/Private/RayTracing/RayTracingDebug.usf
--- a/Engine/Shaders/Private/RayTracing/RayTracingDebug.usf
+++ b/Engine/Shaders/Private/RayTracing/RayTracingDebug.usf
@@ -30,6 +30,9 @@ float FarFieldMaxTraceDistance;
 float3 FarFieldReferencePos;
 
 float TriangleHitCountMaxThreshold;
+float TriangleHitCountPerInstanceMaxThreshold;
+uint TopKMostHitInstances;
+uint NumTotalInstances;
 
 StructuredBuffer<FRayTracingInstanceDebugData> InstancesDebugData;
 StructuredBuffer<FRayTracingPickingFeedback> PickingBuffer;
@@ -204,6 +207,9 @@ float4 DebugRayTracingInstance(FRayDesc Ray, uint2 PixelCoord, out float OutDept
 		case RAY_TRACING_DEBUG_VIZ_TRIANGLE_HITCOUNT:
 			Result.rgb = ColorMapTurbo(Payload.TriangleHitCountPerRay / TriangleHitCountMaxThreshold);
 			break;
+		case RAY_TRACING_DEBUG_VIZ_HITCOUNT_PER_INSTANCE:
+			Result.rgb = ColorMapTurbo(RayTracingDebugHitStatsUniformBuffer.HitStatsOutput[Payload.ScenePrimitiveIndex].y / (float)TriangleHitCountPerInstanceMaxThreshold);
+			break;
 		case RAY_TRACING_DEBUG_VIZ_INSTANCE_OVERLAP:
 			Result.rgb = 1.0; // We only use fake lighting as instance heatmap is blended in
 			break;
@@ -246,7 +252,7 @@ float4 DebugRayTracingInstance(FRayDesc Ray, uint2 PixelCoord, out float OutDept
 		const float FakeLighting = dot(Payload.WorldNormal, normalize(float3(1.0f, 1.0f, 1.0f)));
 		const float FakeLightingScale = (VisualizationMode == RAY_TRACING_DEBUG_VIZ_INSTANCE_OVERLAP) ? 0.5f : 0.125f;
 		const float FakeLightingBias = (VisualizationMode == RAY_TRACING_DEBUG_VIZ_INSTANCE_OVERLAP) ? 0.5f : 0.875f;
-		const int NoFakeLighting = (VisualizationMode == RAY_TRACING_DEBUG_VIZ_TRIANGLE_HITCOUNT);
+		const int NoFakeLighting = (VisualizationMode == RAY_TRACING_DEBUG_VIZ_TRIANGLE_HITCOUNT) || (VisualizationMode == RAY_TRACING_DEBUG_VIZ_HITCOUNT_PER_INSTANCE);
 
 		if (!NoFakeLighting)
 		{
diff --git a/Engine/Shaders/Private/RayTracing/RayTracingDebugCHS.usf b/Engine/Shaders/Private/RayTracing/RayTracingDebugCHS.usf
--- a/Engine/Shaders/Private/RayTracing/RayTracingDebugCHS.usf
+++ b/Engine/Shaders/Private/RayTracing/RayTracingDebugCHS.usf
@@ -23,6 +23,7 @@ RAY_TRACING_ENTRY_CLOSEST_HIT(RayTracingDebugMainCHS,
 	Payload.GeometryInstanceIndex = GetHitGroupUserData();
 	Payload.InstanceIndex = InstanceIndex();
 	Payload.GeometryIndex = 0; // GeometryIndex();
+	Payload.ScenePrimitiveIndex = ScenePrimitiveIndex;
 }
 
 RAY_TRACING_ENTRY_ANY_HIT(RayTracingDebugAHS,
diff --git a/Engine/Shaders/Private/RayTracing/RayTracingDebugHitStats.usf b/Engine/Shaders/Private/RayTracing/RayTracingDebugHitStats.usf
new file mode 100644
--- /dev/null
+++ b/Engine/Shaders/Private/RayTracing/RayTracingDebugHitStats.usf
@@ -0,0 +1,38 @@
+// Copyright Epic Games, Inc. All Rights Reserved.
+
+#include "../Common.ush"
+#include "RayTracingDebugUtils.ush"
+
+#include "/Engine/Shared/RayTracingDebugDefinitions.h"
+#include "/Engine/Shared/RayTracingDebugTypes.h"
+
+uint OpaqueOnly;
+
+uint VisualizationMode;
+uint TriangleHitCountForceNonOpaque;
+RaytracingAccelerationStructure TLAS;
+
+RAY_TRACING_ENTRY_RAYGEN(RayTracingDebugHitStatsRGS)
+{
+	const uint2 PixelPos = DispatchRaysIndex().xy + View.ViewRectMin.xy;
+	float2 RenderTargetUV = (float2(PixelPos) + .5f) * View.BufferSizeAndInvSize.zw;
+
+	FRayDesc Ray = CreatePrimaryRay(RenderTargetUV);
+
+	const uint InstanceInclusionMask = (OpaqueOnly ? RAY_TRACING_MASK_OPAQUE : RAY_TRACING_MASK_ALL);
+	
+	uint RayFlags = RAY_FLAG_CULL_BACK_FACING_TRIANGLES | RAY_FLAG_FORCE_NON_OPAQUE;
+
+	FRayTracingDebugPayload Payload = (FRayTracingDebugPayload)0;
+	Payload.SetMiss();
+
+	TraceRay(
+		TLAS,
+		RayFlags  /*RayFlags*/,
+		InstanceInclusionMask,
+		0 /*RayContributionToHitGroupIndex*/,
+		RAY_TRACING_NUM_SHADER_SLOTS /*MultiplierForGeometryContributionToShaderIndex*/,
+		0 /*MissShaderIndex*/,
+		Ray.GetNativeDesc(),
+		Payload);
+}
diff --git a/Engine/Shaders/Private/RayTracing/RayTracingDebugHitStatsCHS.usf b/Engine/Shaders/Private/RayTracing/RayTracingDebugHitStatsCHS.usf
new file mode 100644
--- /dev/null
+++ b/Engine/Shaders/Private/RayTracing/RayTracingDebugHitStatsCHS.usf
@@ -0,0 +1,25 @@
+// Copyright Epic Games, Inc. All Rights Reserved.
+
+#include "../Common.ush"
+#include "../Hash.ush"
+
+#include "RayTracingCommon.ush"
+#include "RayTracingDebugUtils.ush"
+#include "RayTracingHitGroupCommon.ush"
+#include "/Engine/Shared/RayTracingDebugTypes.h"
+
+RAY_TRACING_ENTRY_ANY_HIT(RayTracingDebugHitStatsAHS,
+	FRayTracingDebugPayload, Payload,
+	FRayTracingIntersectionAttributes, Attributes)
+{
+	uint Index = GetInstanceUserData();
+	uint OldValue;
+	RayTracingDebugHitStatsUniformBuffer.HitStatsOutput[Index].x = Index;
+	InterlockedAdd(RayTracingDebugHitStatsUniformBuffer.HitStatsOutput[Index].y, 1, OldValue);
+}
+
+RAY_TRACING_ENTRY_CLOSEST_HIT(RayTracingDebugHitStatsCHS,
+	FRayTracingDebugPayload, Payload,
+	FRayTracingIntersectionAttributes, Attributes)
+{
+}
\ No newline at end of file
diff --git a/Engine/Shaders/Private/RayTracing/RayTracingDebugHitStatsUtils.usf b/Engine/Shaders/Private/RayTracing/RayTracingDebugHitStatsUtils.usf
new file mode 100644
--- /dev/null
+++ b/Engine/Shaders/Private/RayTracing/RayTracingDebugHitStatsUtils.usf
@@ -0,0 +1,68 @@
+// Copyright Epic Games, Inc. All Rights Reserved.
+#include "../Common.ush"
+#include "../SceneData.ush"
+#include "../ShaderPrint.ush"
+
+#define INVALID_PRIMITIVEID ~0
+
+uint					SelectedNameInfoCount;
+uint					SelectedNameCharacterCount;
+StructuredBuffer<uint3> SelectedPrimitiveNameInfos;
+Buffer<uint>			SelectedPrimitiveNames;
+
+struct FRayTracingSceneDebugHitStatsNameInfo
+{
+	uint PrimitiveID;
+	uint Count;
+	uint Length;
+	uint Offset;
+	uint Pad0;
+	uint Pad1;
+};
+
+FRayTracingSceneDebugHitStatsNameInfo UnpackDebugNameInfo(uint3 In)
+{
+	FRayTracingSceneDebugHitStatsNameInfo Out = (FRayTracingSceneDebugHitStatsNameInfo)0;
+	Out.PrimitiveID = In.x;
+	Out.Count = In.y;
+	Out.Offset = (In.z) & 0xFFFF;
+	Out.Length = (In.z >> 16) & 0xFF;
+
+	return Out;
+}
+
+// Simple local bin sort based on shared atomics
+[numthreads(NUM_THREADS_PER_GROUP, 1, 1)]
+void RayTracingSceneDebugHitStatsRenderCS(uint3 DispatchThread : SV_DispatchThreadID)
+{
+	uint EntryID = DispatchThread.x;
+	if (EntryID >= SelectedNameInfoCount)
+	{
+		return;
+	}
+
+	uint StartOffsetX = 10;
+	uint StartOffsetY = 30;
+	uint SlotHeight = 8;
+	if (EntryID == 0 && SelectedNameInfoCount > 0)
+	{
+		FShaderPrintContext Context = InitShaderPrintContext(true, uint2(StartOffsetX, StartOffsetY + SlotHeight));
+		Print(Context, TEXT("Rank        Prim.ID     Hit Count    Name"), FontWhite);
+	}
+
+	if (SelectedNameInfoCount > 0)
+	{
+		FRayTracingSceneDebugHitStatsNameInfo NameInfo = UnpackDebugNameInfo(SelectedPrimitiveNameInfos[EntryID]);
+		FShaderPrintContext Context = InitShaderPrintContext(true, uint2(StartOffsetX, StartOffsetY + 2 * SlotHeight + EntryID * SlotHeight));
+
+		Print(Context, EntryID, FontYellow);
+		Print(Context, NameInfo.PrimitiveID, FontWhite);
+		Print(Context, NameInfo.Count, FontWhite);
+
+		for (uint It = 0; It < NameInfo.Length; ++It)
+		{
+			const uint Char = SelectedPrimitiveNames[It + NameInfo.Offset];
+			PrintSymbol(Context, Char, FontWhite);
+		}
+	}
+}
diff --git a/Engine/Shaders/Private/RayTracing/RayTracingDebugUtils.ush b/Engine/Shaders/Private/RayTracing/RayTracingDebugUtils.ush
--- a/Engine/Shaders/Private/RayTracing/RayTracingDebugUtils.ush
+++ b/Engine/Shaders/Private/RayTracing/RayTracingDebugUtils.ush
@@ -17,6 +17,7 @@ struct FRayTracingDebugPayload : FMinimalPayload
 	uint GeometryIndex;
 	
 	uint TriangleHitCountPerRay;
+	uint ScenePrimitiveIndex;
 };
 CHECK_RT_PAYLOAD_SIZE(FRayTracingDebugPayload)
 #endif
diff --git a/Engine/Shaders/Shared/RaytracingDebugDefinitions.h b/Engine/Shaders/Shared/RaytracingDebugDefinitions.h
--- a/Engine/Shaders/Shared/RaytracingDebugDefinitions.h
+++ b/Engine/Shaders/Shared/RaytracingDebugDefinitions.h
@@ -43,6 +43,7 @@
 #define RAY_TRACING_DEBUG_VIZ_PICKER					33
 #define RAY_TRACING_DEBUG_VIZ_INSTANCE_OVERLAP			34
 #define RAY_TRACING_DEBUG_VIZ_TRIANGLE_HITCOUNT			35
+#define RAY_TRACING_DEBUG_VIZ_HITCOUNT_PER_INSTANCE		36
 
 #define RAY_TRACING_DEBUG_PICKER_DOMAIN_TRIANGLE		0
 #define RAY_TRACING_DEBUG_PICKER_DOMAIN_INSTANCE		1
diff --git a/Engine/Source/Editor/UnrealEd/Private/RayTracingDebugVisualizationMenuCommands.cpp b/Engine/Source/Editor/UnrealEd/Private/RayTracingDebugVisualizationMenuCommands.cpp
--- a/Engine/Source/Editor/UnrealEd/Private/RayTracingDebugVisualizationMenuCommands.cpp
+++ b/Engine/Source/Editor/UnrealEd/Private/RayTracingDebugVisualizationMenuCommands.cpp
@@ -56,6 +56,8 @@ void FRayTracingDebugVisualizationMenuCommands::CreateRayTracingDebugVisualizati
 	RayTracingDebugModeNames.Add(LOCTEXT("Instances", "Instances"));
 	RayTracingDebugModeNames.Add(LOCTEXT("Instance Overlap", "Instance Overlap"));
 	RayTracingDebugModeNames.Add(LOCTEXT("Triangle Hit Count", "Triangle Hit Count"));
+	RayTracingDebugModeNames.Add(LOCTEXT("Hit Count Per Instance", "Hit Count Per Instance"));
+
 	if (GRHISupportsShaderTimestamp)
 	{
 		RayTracingDebugModeNames.Add(LOCTEXT("Performance", "Performance"));
diff --git a/Engine/Source/Runtime/Renderer/Private/RayTracing/RayTracingDebug.cpp b/Engine/Source/Runtime/Renderer/Private/RayTracing/RayTracingDebug.cpp
--- a/Engine/Source/Runtime/Renderer/Private/RayTracing/RayTracingDebug.cpp
+++ b/Engine/Source/Runtime/Renderer/Private/RayTracing/RayTracingDebug.cpp
@@ -95,6 +95,18 @@ static TAutoConsoleVariable<int32> CVarRayTracingDebugHitCountMaxThreshold(
 	TEXT("Maximum hit count threshold for debug ray tracing triangle hit count heat map visualization. (default = 6)\n")
 );
 
+static TAutoConsoleVariable<int32> CVarRayTracingDebugHitCountPerInstanceMaxThreshold(
+	TEXT("r.RayTracing.DebugTriangleHitCountPerInstance.MaxThreshold"),
+	100000,
+	TEXT("Maximum hit count threshold for debug ray tracing hit count per instance heat map visualization. (default = 100000)\n")
+);
+
+static TAutoConsoleVariable<int32> CVarRayTracingDebugHitCountTopKHits(
+	TEXT("r.RayTracing.DebugTriangleHitCount.TopKMostHits"),
+	10,
+	TEXT("Highlight top k most hit instances in the view. (default = 10)\n")
+);
+
 static int32 GVisualizeProceduralPrimitives = 0;
 static FAutoConsoleVariableRef CVarVisualizeProceduralPrimitives(
 	TEXT("r.RayTracing.DebugVisualizationMode.ProceduralPrimitives"),
@@ -104,7 +116,14 @@ static FAutoConsoleVariableRef CVarVisualizeProceduralPrimitives(
 	ECVF_RenderThreadSafe
 );
 
-IMPLEMENT_RT_PAYLOAD_TYPE(ERayTracingPayloadType::RayTracingDebug, 40);
+IMPLEMENT_RT_PAYLOAD_TYPE(ERayTracingPayloadType::RayTracingDebug, 44);
+
+BEGIN_UNIFORM_BUFFER_STRUCT(FRayTracingDebugHitStatsUniformBufferParameters, )
+	SHADER_PARAMETER_RDG_BUFFER_UAV(RWStructuredBuffer<uint2>, HitStatsOutput)
+END_UNIFORM_BUFFER_STRUCT()
+
+IMPLEMENT_UNIFORM_BUFFER_STRUCT(FRayTracingDebugHitStatsUniformBufferParameters, "RayTracingDebugHitStatsUniformBuffer");
+
 
 class FRayTracingDebugRGS : public FGlobalShader
 {
@@ -125,13 +144,18 @@ class FRayTracingDebugRGS : public FGlobalShader
 		SHADER_PARAMETER(FVector3f, FarFieldReferencePos)
 		SHADER_PARAMETER(int32, OpaqueOnly)
 		SHADER_PARAMETER(float, TriangleHitCountMaxThreshold)
+		SHADER_PARAMETER(float, TriangleHitCountPerInstanceMaxThreshold)
+		SHADER_PARAMETER(uint32, TopKMostHitInstances)
+		SHADER_PARAMETER(uint32, NumTotalInstances)
 		SHADER_PARAMETER_RDG_BUFFER_SRV(RaytracingAccelerationStructure, TLAS)
 		SHADER_PARAMETER_RDG_TEXTURE_UAV(RWTexture2D<float4>, Output)
 		SHADER_PARAMETER_RDG_TEXTURE_UAV(RWTexture2D<float>, OutputDepth)
 		SHADER_PARAMETER_RDG_BUFFER_SRV(StructuredBuffer, InstancesDebugData)
 		SHADER_PARAMETER_RDG_BUFFER_SRV(StructuredBuffer<FRayTracingPickingFeedback>, PickingBuffer)
+		SHADER_PARAMETER_RDG_BUFFER_SRV(StructuredBuffer<uint2>, TopKHitStats)
 		SHADER_PARAMETER_STRUCT_REF(FViewUniformShaderParameters, ViewUniformBuffer)
 		SHADER_PARAMETER_RDG_UNIFORM_BUFFER(FSceneUniformParameters, SceneUniformBuffer)
+		SHADER_PARAMETER_RDG_UNIFORM_BUFFER(FRayTracingDebugHitStatsUniformBufferParameters, RayTracingDebugHitStatsUniformBuffer)
 	END_SHADER_PARAMETER_STRUCT()
 
 	static bool ShouldCompilePermutation(const FGlobalShaderPermutationParameters& Parameters)
@@ -219,6 +243,71 @@ public:
 };
 IMPLEMENT_GLOBAL_SHADER(FRayTracingDebugMS, "/Engine/Private/RayTracing/RayTracingDebugMS.usf", "RayTracingDebugMS", SF_RayMiss);
 
+class FRayTracingDebugHitStatsRGS : public FGlobalShader
+{
+	DECLARE_GLOBAL_SHADER(FRayTracingDebugHitStatsRGS);
+	SHADER_USE_ROOT_PARAMETER_STRUCT(FRayTracingDebugHitStatsRGS, FGlobalShader)
+
+	BEGIN_SHADER_PARAMETER_STRUCT(FParameters, )
+		SHADER_PARAMETER_RDG_BUFFER_SRV(RaytracingAccelerationStructure, TLAS)
+		SHADER_PARAMETER(int32, OpaqueOnly)
+		SHADER_PARAMETER(uint32, VisualizationMode)
+		SHADER_PARAMETER(uint32, TriangleHitCountForceNonOpaque)
+		SHADER_PARAMETER_STRUCT_REF(FViewUniformShaderParameters, ViewUniformBuffer)
+		SHADER_PARAMETER_RDG_UNIFORM_BUFFER(FSceneUniformParameters, SceneUniformBuffer)
+		SHADER_PARAMETER_RDG_UNIFORM_BUFFER(FRayTracingDebugHitStatsUniformBufferParameters, RayTracingDebugHitStatsUniformBuffer)
+	END_SHADER_PARAMETER_STRUCT()
+public:
+
+	static bool ShouldCompilePermutation(const FGlobalShaderPermutationParameters& Parameters)
+	{
+		return ShouldCompileRayTracingShadersForProject(Parameters.Platform);
+	}
+
+	static ERayTracingPayloadType GetRayTracingPayloadType(const int32 PermutationId)
+	{
+		return ERayTracingPayloadType::RayTracingDebug;
+	}
+};
+IMPLEMENT_GLOBAL_SHADER(FRayTracingDebugHitStatsRGS, "/Engine/Private/RayTracing/RayTracingDebugHitStats.usf", "RayTracingDebugHitStatsRGS", SF_RayGen);
+
+
+class FRayTracingDebugHitStatsCHS : public FGlobalShader
+{
+	DECLARE_GLOBAL_SHADER(FRayTracingDebugHitStatsCHS);
+	SHADER_USE_ROOT_PARAMETER_STRUCT(FRayTracingDebugHitStatsCHS, FGlobalShader)
+
+	BEGIN_SHADER_PARAMETER_STRUCT(FParameters, )
+		SHADER_PARAMETER_RDG_UNIFORM_BUFFER(FRayTracingDebugHitStatsUniformBufferParameters, RayTracingDebugHitStatsUniformBuffer)
+	END_SHADER_PARAMETER_STRUCT()
+public:
+
+	class FNaniteRayTracing : SHADER_PERMUTATION_BOOL("NANITE_RAY_TRACING");
+	using FPermutationDomain = TShaderPermutationDomain<FNaniteRayTracing>;
+
+	static bool ShouldCompilePermutation(const FGlobalShaderPermutationParameters& Parameters)
+	{
+		return ShouldCompileRayTracingShadersForProject(Parameters.Platform);
+	}
+
+	static void ModifyCompilationEnvironment(const FGlobalShaderPermutationParameters& Parameters, FShaderCompilerEnvironment& OutEnvironment)
+	{
+		FGlobalShader::ModifyCompilationEnvironment(Parameters, OutEnvironment);
+
+		FPermutationDomain PermutationVector(Parameters.PermutationId);
+		if (PermutationVector.Get<FNaniteRayTracing>())
+		{
+			OutEnvironment.SetDefine(TEXT("VF_SUPPORTS_PRIMITIVE_SCENE_DATA"), 1);
+		}
+	}
+
+	static ERayTracingPayloadType GetRayTracingPayloadType(const int32 PermutationId)
+	{
+		return ERayTracingPayloadType::RayTracingDebug;
+	}
+};
+IMPLEMENT_GLOBAL_SHADER(FRayTracingDebugHitStatsCHS, "/Engine/Private/RayTracing/RayTracingDebugHitStatsCHS.usf", "closesthit=RayTracingDebugHitStatsCHS anyhit=RayTracingDebugHitStatsAHS", SF_RayHitGroup);
+
 class FRayTracingDebugTraversalCS : public FGlobalShader
 {
 	DECLARE_GLOBAL_SHADER(FRayTracingDebugTraversalCS)
@@ -440,9 +529,15 @@ struct FRayTracingDebugResources : public FRenderResource
 	int32 PickingBufferNumPending = 0;
 	TArray<FRHIGPUBufferReadback*> PickingBuffers;
 
+	const int32 MaxHitStatsBuffers = 4;
+	int32 HitStatsBufferWriteIndex = 0;
+	int32 HitStatsBufferNumPending = 0;
+	TArray<FRHIGPUBufferReadback*> HitStatsBuffers;
+
 	virtual void InitRHI(FRHICommandListBase& RHICmdList) override
 	{
 		PickingBuffers.AddZeroed(MaxPickingBuffers);
+		HitStatsBuffers.AddZeroed(MaxHitStatsBuffers);
 	}
 
 	virtual void ReleaseRHI() override
@@ -455,12 +550,128 @@ struct FRayTracingDebugResources : public FRenderResource
 				PickingBuffers[BufferIndex] = nullptr;
 			}
 		}
+		for (int32 BufferIndex = 0; BufferIndex < HitStatsBuffers.Num(); ++BufferIndex)
+		{
+			if (HitStatsBuffers[BufferIndex])
+			{
+				delete HitStatsBuffers[BufferIndex];
+				HitStatsBuffers[BufferIndex] = nullptr;
+			}
+		}
+
 		PickingBuffers.Reset();
+		HitStatsBuffers.Reset();
 	}
 };
 
 TGlobalResource<FRayTracingDebugResources> GRayTracingDebugResources;
 
+class FRayTracingHitStatsEntry {
+public:
+	uint32 PrimitiveID;
+	uint32 Count;
+};
+
+void BindRayTracingDebugHitStatsCHSMaterialBindings(FRHICommandList& RHICmdList, const FViewInfo& View, FRHIUniformBuffer* SceneUniformBuffer, FRHIUniformBuffer* HitStatsUniformBuffer, FRayTracingPipelineState* PipelineState)
+{
+	FSceneRenderingBulkObjectAllocator Allocator;
+
+	auto Alloc = [&](uint32 Size, uint32 Align)
+	{
+		return RHICmdList.Bypass()
+			? Allocator.Malloc(Size, Align)
+			: RHICmdList.Alloc(Size, Align);
+	};
+
+	const int32 NumTotalBindings = View.VisibleRayTracingMeshCommands.Num();
+	const uint32 MergedBindingsSize = sizeof(FRayTracingLocalShaderBindings) * NumTotalBindings;
+	FRayTracingLocalShaderBindings* Bindings = (FRayTracingLocalShaderBindings*)Alloc(MergedBindingsSize, alignof(FRayTracingLocalShaderBindings));
+
+	struct FBinding
+	{
+		int32 ShaderIndexInPipeline;
+		uint32 NumUniformBuffers;
+		FRHIUniformBuffer** UniformBufferArray;
+	};
+
+	auto SetupBinding = [&](FRayTracingDebugHitStatsCHS::FPermutationDomain PermutationVector)
+	{
+		auto Shader = View.ShaderMap->GetShader<FRayTracingDebugHitStatsCHS>(PermutationVector);
+		auto HitGroupShader = Shader.GetRayTracingShader();
+
+		FBinding Binding;
+		Binding.ShaderIndexInPipeline = FindRayTracingHitGroupIndex(PipelineState, HitGroupShader, true);
+		Binding.NumUniformBuffers = Shader->ParameterMapInfo.UniformBuffers.Num();
+		Binding.UniformBufferArray = (FRHIUniformBuffer**)Alloc(sizeof(FRHIUniformBuffer*) * Binding.NumUniformBuffers, alignof(FRHIUniformBuffer*));
+
+		const auto& HitStatsUniformBufferParameter = Shader->GetUniformBufferParameter<FRayTracingDebugHitStatsUniformBufferParameters>();
+		const auto& ViewUniformBufferParameter = Shader->GetUniformBufferParameter<FViewUniformShaderParameters>();
+		const auto& SceneUniformBufferParameter = Shader->GetUniformBufferParameter<FSceneUniformParameters>();
+		const auto& NaniteUniformBufferParameter = Shader->GetUniformBufferParameter<FNaniteRayTracingUniformParameters>();
+
+		if (HitStatsUniformBufferParameter.IsBound())
+		{
+			check(HitStatsUniformBufferParameter.GetBaseIndex() < Binding.NumUniformBuffers);
+			Binding.UniformBufferArray[HitStatsUniformBufferParameter.GetBaseIndex()] = HitStatsUniformBuffer;
+		}
+
+		if (ViewUniformBufferParameter.IsBound())
+		{
+			check(ViewUniformBufferParameter.GetBaseIndex() < Binding.NumUniformBuffers);
+			Binding.UniformBufferArray[ViewUniformBufferParameter.GetBaseIndex()]
+				= View.ViewUniformBuffer.GetReference();
+		}
+
+		if (SceneUniformBufferParameter.IsBound())
+		{
+			check(SceneUniformBufferParameter.GetBaseIndex() < Binding.NumUniformBuffers);
+			Binding.UniformBufferArray[SceneUniformBufferParameter.GetBaseIndex()] = SceneUniformBuffer;
+		}
+
+		if (NaniteUniformBufferParameter.IsBound())
+		{
+			check(NaniteUniformBufferParameter.GetBaseIndex() < Binding.NumUniformBuffers);
+			Binding.UniformBufferArray[NaniteUniformBufferParameter.GetBaseIndex()] = Nanite::GRayTracingManager.GetUniformBuffer().GetReference();
+		}
+
+		return Binding;
+	};
+
+	FRayTracingDebugHitStatsCHS::FPermutationDomain PermutationVector;
+
+	PermutationVector.Set<FRayTracingDebugHitStatsCHS::FNaniteRayTracing>(false);
+	FBinding ShaderBinding = SetupBinding(PermutationVector);
+
+	PermutationVector.Set<FRayTracingDebugHitStatsCHS::FNaniteRayTracing>(true);
+	FBinding ShaderBindingNaniteRT = SetupBinding(PermutationVector);
+
+	uint32 BindingIndex = 0;
+	for (const FVisibleRayTracingMeshCommand VisibleMeshCommand : View.VisibleRayTracingMeshCommands)
+	{
+		const FRayTracingMeshCommand& MeshCommand = *VisibleMeshCommand.RayTracingMeshCommand;
+
+		const FBinding& HelperBinding = MeshCommand.IsUsingNaniteRayTracing() ? ShaderBindingNaniteRT : ShaderBinding;
+
+		FRayTracingLocalShaderBindings Binding = {};
+		Binding.ShaderIndexInPipeline = HelperBinding.ShaderIndexInPipeline;
+		Binding.InstanceIndex = VisibleMeshCommand.InstanceIndex;
+		Binding.SegmentIndex = MeshCommand.GeometrySegmentIndex;
+		Binding.UniformBuffers = HelperBinding.UniformBufferArray;
+		Binding.NumUniformBuffers = HelperBinding.NumUniformBuffers;
+		Binding.UserData = VisibleMeshCommand.InstanceIndex;
+
+		Bindings[BindingIndex] = Binding;
+		BindingIndex++;
+	}
+
+	const bool bCopyDataToInlineStorage = false; // Storage is already allocated from RHICmdList, no extra copy necessary
+	RHICmdList.SetRayTracingHitGroups(
+		View.GetRayTracingSceneChecked(),
+		PipelineState,
+		NumTotalBindings, Bindings,
+		bCopyDataToInlineStorage);
+}
+
 void BindRayTracingDebugCHSMaterialBindings(FRHICommandList& RHICmdList, const FViewInfo& View, FRHIUniformBuffer* SceneUniformBuffer, FRayTracingPipelineState* PipelineState)
 {
 	FSceneRenderingBulkObjectAllocator Allocator;
@@ -500,7 +711,8 @@ void BindRayTracingDebugCHSMaterialBindings(FRHICommandList& RHICmdList, const F
 		if (ViewUniformBufferParameter.IsBound())
 		{
 			check(ViewUniformBufferParameter.GetBaseIndex() < Binding.NumUniformBuffers);
-			Binding.UniformBufferArray[ViewUniformBufferParameter.GetBaseIndex()] = View.ViewUniformBuffer.GetReference();
+			Binding.UniformBufferArray[ViewUniformBufferParameter.GetBaseIndex()]
+				= View.ViewUniformBuffer.GetReference();
 		}
 
 		if (SceneUniformBufferParameter.IsBound())
@@ -561,7 +773,8 @@ static bool RequiresRayTracingDebugCHS(uint32 DebugVisualizationMode)
 		DebugVisualizationMode == RAY_TRACING_DEBUG_VIZ_PROXY_TYPE ||
 		DebugVisualizationMode == RAY_TRACING_DEBUG_VIZ_PICKER ||
 		DebugVisualizationMode == RAY_TRACING_DEBUG_VIZ_INSTANCE_OVERLAP ||
-		DebugVisualizationMode == RAY_TRACING_DEBUG_VIZ_TRIANGLE_HITCOUNT;
+		DebugVisualizationMode == RAY_TRACING_DEBUG_VIZ_TRIANGLE_HITCOUNT ||
+		DebugVisualizationMode == RAY_TRACING_DEBUG_VIZ_HITCOUNT_PER_INSTANCE;
 }
 
 static bool IsRayTracingDebugTraversalMode(uint32 DebugVisualizationMode)
@@ -703,6 +916,273 @@ static FRDGBufferRef RayTracingPerformPicking(FRDGBuilder& GraphBuilder, const F
 	return PickingBuffer;
 }
 
+static TRDGUniformBufferRef<FRayTracingDebugHitStatsUniformBufferParameters> DebugHitStatsUniformBuffer;
+
+struct FRayTracingSceneDebugHitStatsNameInfo
+{
+	uint32 PrimitiveID;
+	uint32 Count;
+	uint16 Offset;
+	uint8  Length;
+	uint8  Pad0;
+};
+
+class FRayTracingSceneHitStatsDebugRenderCS : public FGlobalShader
+{
+	DECLARE_GLOBAL_SHADER(FRayTracingSceneHitStatsDebugRenderCS);
+	SHADER_USE_PARAMETER_STRUCT(FRayTracingSceneHitStatsDebugRenderCS, FGlobalShader);
+
+	BEGIN_SHADER_PARAMETER_STRUCT(FParameters, )
+		SHADER_PARAMETER(int32, SelectedNameInfoCount)
+		SHADER_PARAMETER(int32, SelectedNameCharacterCount)
+		SHADER_PARAMETER_RDG_UNIFORM_BUFFER(FSceneUniformParameters, Scene)
+		SHADER_PARAMETER_STRUCT_INCLUDE(ShaderPrint::FShaderParameters, ShaderPrintUniformBuffer)
+		SHADER_PARAMETER_RDG_BUFFER_SRV(Buffer<uint8>, SelectedPrimitiveNames)
+		SHADER_PARAMETER_RDG_BUFFER_SRV(StructuredBuffer<uint3>, SelectedPrimitiveNameInfos)
+	END_SHADER_PARAMETER_STRUCT()
+
+public:
+	static constexpr uint32 NumThreadsPerGroup = 32U;
+
+	static bool ShouldCompilePermutation(const FGlobalShaderPermutationParameters& Parameters)
+	{
+		return ShaderPrint::IsSupported(Parameters.Platform) && IsRayTracingEnabledForProject(Parameters.Platform);
+	}
+
+	static void ModifyCompilationEnvironment(const FGlobalShaderPermutationParameters& Parameters, FShaderCompilerEnvironment& OutEnvironment)
+	{
+		FGlobalShader::ModifyCompilationEnvironment(Parameters, OutEnvironment);
+
+		OutEnvironment.SetDefine(TEXT("VF_SUPPORTS_PRIMITIVE_SCENE_DATA"), 1);
+		OutEnvironment.SetDefine(TEXT("NUM_THREADS_PER_GROUP"), NumThreadsPerGroup);
+
+		// Skip optimization for avoiding long compilation time due to large UAV writes
+		OutEnvironment.CompilerFlags.Add(CFLAG_Debug);
+	}
+};
+IMPLEMENT_GLOBAL_SHADER(FRayTracingSceneHitStatsDebugRenderCS, "/Engine/Private/RayTracing/RayTracingDebugHitStatsUtils.usf", "RayTracingSceneDebugHitStatsRenderCS", SF_Compute);
+
+static void PrintTopKMostHitMessage(FRDGBuilder& GraphBuilder, const FScene* Scene, const FViewInfo& View, const TArray<FRayTracingHitStatsEntry>& HitStatsArray)
+{
+	// Force ShaderPrint on.
+	ShaderPrint::SetEnabled(true);
+
+	int32 NumPrimitives = CVarRayTracingDebugHitCountTopKHits.GetValueOnRenderThread();
+	if (ShaderPrint::IsEnabled(View.ShaderPrintData) && NumPrimitives > 0)
+	{
+		// This lags by one frame, so may miss some in one frame, also overallocates since we will cull a lot.
+		ShaderPrint::RequestSpaceForLines(NumPrimitives * 12);
+
+		const uint32 MaxPrimitiveNameCount = 128u;
+		check(sizeof(FRayTracingSceneDebugHitStatsNameInfo) == 12);
+		TArray<FRayTracingSceneDebugHitStatsNameInfo> SelectedNameInfos;
+		TArray<uint8> SelectedNames;
+		SelectedNames.Reserve(NumPrimitives * 30u);
+
+		uint32 SelectedCount = 0;
+		const int32 BitsPerWord = (sizeof(uint32) * 8U);
+		for (int32 HitStatsID = 0; HitStatsID < NumPrimitives; ++HitStatsID)
+		{
+			const uint32 PrimitiveID = HitStatsArray[HitStatsID].PrimitiveID;
+			if (PrimitiveID >= (uint32)Scene->Primitives.Num())
+			{
+				continue;
+			}
+			const FString OwnerName = Scene->Primitives[PrimitiveID]->GetFullnameForDebuggingOnly();
+			const uint32 NameOffset = SelectedNames.Num();
+			const uint32 NameLength = OwnerName.Len();
+			for (TCHAR C : OwnerName)
+			{
+				SelectedNames.Add(uint8(C));
+			}
+
+			FRayTracingSceneDebugHitStatsNameInfo& NameInfo = SelectedNameInfos.AddDefaulted_GetRef();
+			NameInfo.PrimitiveID = PrimitiveID;
+			NameInfo.Count = HitStatsArray[HitStatsID].Count;
+			NameInfo.Length = NameLength;
+			NameInfo.Offset = NameOffset;
+			++SelectedCount;
+		}
+
+		if (SelectedNameInfos.IsEmpty())
+		{
+			FRayTracingSceneDebugHitStatsNameInfo& NameInfo = SelectedNameInfos.AddDefaulted_GetRef();
+			NameInfo.PrimitiveID = ~0;
+			NameInfo.Count = -1;
+			NameInfo.Length = 4;
+			NameInfo.Offset = 0;
+			SelectedNames.Add(uint8('N'));
+			SelectedNames.Add(uint8('o'));
+			SelectedNames.Add(uint8('n'));
+			SelectedNames.Add(uint8('e'));
+		}
+
+		// Request more characters for printing if needed
+		ShaderPrint::RequestSpaceForCharacters(SelectedNames.Num() + SelectedCount * 48u);
+
+		FRDGBufferRef SelectedPrimitiveNames = CreateVertexBuffer(GraphBuilder, TEXT("RayTracingDebug.HitStats.SelectedPrimitiveNames"), 
+			FRDGBufferDesc::CreateBufferDesc(1, SelectedNames.Num()), SelectedNames.GetData(), SelectedNames.Num());
+		FRDGBufferRef SelectedPrimitiveNameInfos = CreateStructuredBuffer(GraphBuilder, TEXT("RayTracingDebug.HitStats.SelectedPrimitiveNameInfos"), SelectedNameInfos);
+
+		FRayTracingSceneHitStatsDebugRenderCS::FParameters* PassParameters = GraphBuilder.AllocParameters<FRayTracingSceneHitStatsDebugRenderCS::FParameters>();
+		ShaderPrint::SetParameters(GraphBuilder, View.ShaderPrintData, PassParameters->ShaderPrintUniformBuffer);
+		PassParameters->Scene = GetSceneUniformBufferRef(GraphBuilder, View);
+		PassParameters->SelectedNameInfoCount = SelectedCount;
+		PassParameters->SelectedNameCharacterCount = SelectedCount > 0 ? SelectedNames.Num() : 0;
+		PassParameters->SelectedPrimitiveNameInfos = GraphBuilder.CreateSRV(SelectedPrimitiveNameInfos);
+		PassParameters->SelectedPrimitiveNames = GraphBuilder.CreateSRV(SelectedPrimitiveNames, PF_R8_UINT);
+
+		auto ComputeShader = View.ShaderMap->GetShader<FRayTracingSceneHitStatsDebugRenderCS>();
+
+		FComputeShaderUtils::AddPass(
+			GraphBuilder,
+			RDG_EVENT_NAME("RayTracingDebug::TopKHitStatsInfo"),
+			ComputeShader,
+			PassParameters,
+			FComputeShaderUtils::GetGroupCount(NumPrimitives, FRayTracingSceneHitStatsDebugRenderCS::NumThreadsPerGroup)
+		);
+	}
+}
+
+static FRDGBufferRef RayTracingPerformHitStats(FRDGBuilder& GraphBuilder, const FScene* Scene, const FViewInfo& View)
+{
+	FGlobalShaderMap* ShaderMap = GetGlobalShaderMap(GMaxRHIFeatureLevel);
+	auto RayGenShader = ShaderMap->GetShader<FRayTracingDebugHitStatsRGS>();
+
+	FRayTracingPipelineStateInitializer Initializer;
+	Initializer.MaxPayloadSizeInBytes = GetRayTracingPayloadTypeMaxSize(ERayTracingPayloadType::RayTracingDebug);
+
+	FRHIRayTracingShader* RayGenShaderTable[] = { RayGenShader.GetRayTracingShader() };
+	Initializer.SetRayGenShaderTable(RayGenShaderTable);
+
+	FRayTracingDebugHitStatsCHS::FPermutationDomain PermutationVector;
+
+	PermutationVector.Set<FRayTracingDebugHitStatsCHS::FNaniteRayTracing>(false);
+	auto HitGroupShader = View.ShaderMap->GetShader<FRayTracingDebugHitStatsCHS>(PermutationVector);
+
+	PermutationVector.Set<FRayTracingDebugHitStatsCHS::FNaniteRayTracing>(true);
+	auto HitGroupShaderNaniteRT = View.ShaderMap->GetShader<FRayTracingDebugHitStatsCHS>(PermutationVector);
+
+	FRHIRayTracingShader* HitGroupTable[] = { HitGroupShader.GetRayTracingShader(), HitGroupShaderNaniteRT.GetRayTracingShader() };
+	Initializer.SetHitGroupTable(HitGroupTable);
+	Initializer.bAllowHitGroupIndexing = true; // Required for stable output using GetBaseInstanceIndex().
+
+	auto MissShader = ShaderMap->GetShader<FRayTracingDebugMS>();
+	FRHIRayTracingShader* MissTable[] = { MissShader.GetRayTracingShader() };
+	Initializer.SetMissShaderTable(MissTable);
+
+	FRayTracingPipelineState* PickingPipeline = PipelineStateCache::GetAndOrCreateRayTracingPipelineState(GraphBuilder.RHICmdList, Initializer);
+
+	
+	const uint32 NumInstances = FMath::Max<uint32>(View.RayTracingSceneInitData.NumNativeCPUInstances + View.RayTracingSceneInitData.NumNativeGPUInstances + View.RayTracingSceneInitData.NumNativeGPUSceneInstances,
+		(uint32)CVarRayTracingDebugHitCountTopKHits.GetValueOnRenderThread());
+	const uint32 NumPrimitives = FMath::Max<uint32>(Scene->Primitives.Num(),
+		(uint32)CVarRayTracingDebugHitCountTopKHits.GetValueOnRenderThread());
+	FRDGBufferDesc PickingBufferDesc = FRDGBufferDesc::CreateStructuredDesc(sizeof(uint32) * 2, NumPrimitives);
+	PickingBufferDesc.Usage = EBufferUsageFlags(PickingBufferDesc.Usage | BUF_SourceCopy);
+	FRDGBufferRef HitStatsBuffer = GraphBuilder.CreateBuffer(PickingBufferDesc, TEXT("RayTracingDebug.HitStatsBuffer"));
+	
+
+	FRayTracingDebugHitStatsUniformBufferParameters* DebugHitStatsUniformBufferParameters = GraphBuilder.AllocParameters<FRayTracingDebugHitStatsUniformBufferParameters>();
+	DebugHitStatsUniformBufferParameters->HitStatsOutput = GraphBuilder.CreateUAV(HitStatsBuffer);
+	DebugHitStatsUniformBuffer = GraphBuilder.CreateUniformBuffer(DebugHitStatsUniformBufferParameters);
+
+	FRayTracingDebugHitStatsRGS::FParameters* RayGenParameters = GraphBuilder.AllocParameters<FRayTracingDebugHitStatsRGS::FParameters>();
+	RayGenParameters->TLAS = Scene->RayTracingScene.GetLayerView(ERayTracingSceneLayer::Base);
+	RayGenParameters->OpaqueOnly = CVarRayTracingDebugModeOpaqueOnly.GetValueOnRenderThread();
+	RayGenParameters->ViewUniformBuffer = View.ViewUniformBuffer;
+	RayGenParameters->SceneUniformBuffer = GetSceneUniformBufferRef(GraphBuilder, View); // TODO: use a separate params structure
+	RayGenParameters->RayTracingDebugHitStatsUniformBuffer = DebugHitStatsUniformBuffer;
+
+	AddClearUAVPass(GraphBuilder, DebugHitStatsUniformBufferParameters->HitStatsOutput, 0, ERDGPassFlags::Compute);
+
+	FIntRect ViewRect = View.ViewRect;
+	GraphBuilder.AddPass(
+		RDG_EVENT_NAME("RayTracingHitStats"),
+		RayGenParameters,
+		ERDGPassFlags::Compute,
+		[RayGenParameters, RayGenShader, &View, PickingPipeline, ViewRect](FRHIRayTracingCommandList& RHICmdList)
+		{
+			FRayTracingShaderBindingsWriter GlobalResources;
+			SetShaderParameters(GlobalResources, RayGenShader, *RayGenParameters);
+
+			BindRayTracingDebugHitStatsCHSMaterialBindings(RHICmdList, View, RayGenParameters->SceneUniformBuffer->GetRHI(), DebugHitStatsUniformBuffer->GetRHI(), PickingPipeline);
+			RHICmdList.SetRayTracingMissShader(View.GetRayTracingSceneChecked(), 0, PickingPipeline, 0 /* ShaderIndexInPipeline */, 0, nullptr, 0);
+
+			RHICmdList.RayTraceDispatch(PickingPipeline, RayGenShader.GetRayTracingShader(), View.GetRayTracingSceneChecked(), GlobalResources, ViewRect.Size().X, ViewRect.Size().Y);
+		});
+
+	TArray<FRayTracingHitStatsEntry> HitStatsKeyValuePairs;	
+	FRHIGPUBufferReadback* CPUHitStatsBuffer = nullptr;
+	const int32 MaxHitStatsBuffers = GRayTracingDebugResources.MaxHitStatsBuffers;
+
+	int32& HitStatsBufferWriteIndex = GRayTracingDebugResources.HitStatsBufferWriteIndex;
+	int32& HitStatsBufferNumPending = GRayTracingDebugResources.HitStatsBufferNumPending;
+
+	TArray<FRHIGPUBufferReadback*>& HitStatsBuffers = GRayTracingDebugResources.HitStatsBuffers;
+
+	{
+		FRHIGPUBufferReadback* LatestHitStatsBuffer = nullptr;
+
+		// Find latest buffer that is ready
+		while (HitStatsBufferNumPending > 0)
+		{
+			uint32 Index = (HitStatsBufferWriteIndex + MaxHitStatsBuffers - HitStatsBufferNumPending) % MaxHitStatsBuffers;
+			if (HitStatsBuffers[Index]->IsReady())
+			{
+				--HitStatsBufferNumPending;
+				LatestHitStatsBuffer = HitStatsBuffers[Index];
+			}
+			else
+			{
+				break;
+			}
+		}
+
+		if (LatestHitStatsBuffer != nullptr)
+		{
+			TRACE_CPUPROFILER_EVENT_SCOPE(LockBuffer);
+			const uint32 BufferSize = LatestHitStatsBuffer->GetGPUSizeBytes();
+			const uint32* BufferAddr = (const uint32*)LatestHitStatsBuffer->Lock(BufferSize);
+			for (uint32 Index = 0; Index < BufferSize / sizeof(uint32) / 2; Index++)
+			{
+				HitStatsKeyValuePairs.Add(FRayTracingHitStatsEntry{ BufferAddr[Index * 2], BufferAddr[Index * 2 + 1] });
+			}
+			LatestHitStatsBuffer->Unlock();
+		}
+	}
+
+	// Skip when queue is full. It is NOT safe to EnqueueCopy on a buffer that already has a pending copy
+	if (HitStatsBufferNumPending != MaxHitStatsBuffers)
+	{
+		if (HitStatsBuffers[HitStatsBufferWriteIndex] == nullptr)
+		{
+			FRHIGPUBufferReadback* GPUBufferReadback = new FRHIGPUBufferReadback(TEXT("RayTracingDebug.HitStatsFeedback"));
+			HitStatsBuffers[HitStatsBufferWriteIndex] = GPUBufferReadback;
+		}
+
+		FRHIGPUBufferReadback* HitStatsReadback = HitStatsBuffers[HitStatsBufferWriteIndex];
+
+		AddEnqueueCopyPass(GraphBuilder, HitStatsReadback, HitStatsBuffer, 0u);
+		 
+		HitStatsBufferWriteIndex = (HitStatsBufferWriteIndex + 1) % MaxHitStatsBuffers;
+		HitStatsBufferNumPending = FMath::Min(HitStatsBufferNumPending + 1, MaxHitStatsBuffers);
+	}
+
+
+	if (HitStatsKeyValuePairs.Num() > 0) 
+	{
+		HitStatsKeyValuePairs.Sort([](const FRayTracingHitStatsEntry& A, const FRayTracingHitStatsEntry& B) {
+			return A.Count > B.Count;
+			});
+		PrintTopKMostHitMessage(GraphBuilder, Scene, View, HitStatsKeyValuePairs);
+		return HitStatsBuffer;
+	}
+
+	return HitStatsBuffer;
+}
+
+
 static void RayTracingDrawInstances(FRDGBuilder& GraphBuilder, const FViewInfo& View, FRDGTextureRef OutputTexture, FRDGTextureRef SceneDepthTexture, FRDGBufferRef InstanceGPUSceneIndexBuffer, bool bWireframe)
 {
 	TShaderMapRef<FRayTracingDebugInstanceOverlapVS> VertexShader(View.ShaderMap);
@@ -864,6 +1344,7 @@ void FDeferredShadingSceneRenderer::RenderRayTracingDebug(FRDGBuilder& GraphBuil
 		RayTracingDebugVisualizationModes.Emplace(FName(*LOCTEXT("Picker", "Picker").ToString()),												RAY_TRACING_DEBUG_VIZ_PICKER);
 		RayTracingDebugVisualizationModes.Emplace(FName(*LOCTEXT("Instance Overlap", "Instance Overlap").ToString()),							RAY_TRACING_DEBUG_VIZ_INSTANCE_OVERLAP);
 		RayTracingDebugVisualizationModes.Emplace(FName(*LOCTEXT("Triangle Hit Count", "Triangle Hit Count").ToString()),						RAY_TRACING_DEBUG_VIZ_TRIANGLE_HITCOUNT);
+		RayTracingDebugVisualizationModes.Emplace(FName(*LOCTEXT("Hit Count Per Instance", "Hit Count Per Instance").ToString()),				RAY_TRACING_DEBUG_VIZ_HITCOUNT_PER_INSTANCE);
 	}
 
 	uint32 DebugVisualizationMode;
@@ -959,6 +1440,7 @@ void FDeferredShadingSceneRenderer::RenderRayTracingDebug(FRDGBuilder& GraphBuil
 	}
 
 	FRDGBufferRef PickingBuffer = nullptr;
+	FRDGBufferRef StatsBuffer = nullptr;
 	if (IsRayTracingPickingEnabled(DebugVisualizationMode) && Scene->RayTracingScene.InstanceDebugBuffer != nullptr)
 	{
 		PickingBuffer = RayTracingPerformPicking(GraphBuilder, Scene, View, PickingFeedback);
@@ -968,6 +1450,22 @@ void FDeferredShadingSceneRenderer::RenderRayTracingDebug(FRDGBuilder& GraphBuil
 		PickingBuffer = GSystemTextures.GetDefaultStructuredBuffer(GraphBuilder, sizeof(FRayTracingPickingFeedback));
 	}
 
+	if (DebugVisualizationMode == RAY_TRACING_DEBUG_VIZ_HITCOUNT_PER_INSTANCE)
+	{
+		StatsBuffer = RayTracingPerformHitStats(GraphBuilder, Scene, View);
+	}
+	else 
+	{
+		FRDGBufferDesc StatsBufferDesc = FRDGBufferDesc::CreateStructuredDesc(sizeof(uint32) * 2, 1);
+		StatsBufferDesc.Usage = EBufferUsageFlags(StatsBufferDesc.Usage);
+		StatsBuffer = GraphBuilder.CreateBuffer(StatsBufferDesc, TEXT("RayTracingDebug.HitStatsBuffer"));
+
+		FRayTracingDebugHitStatsUniformBufferParameters* DebugHitStatsUniformBufferParameters = GraphBuilder.AllocParameters<FRayTracingDebugHitStatsUniformBufferParameters>();
+		DebugHitStatsUniformBufferParameters->HitStatsOutput = GraphBuilder.CreateUAV(StatsBuffer);
+		DebugHitStatsUniformBuffer = GraphBuilder.CreateUniformBuffer(DebugHitStatsUniformBufferParameters);
+		AddClearUAVPass(GraphBuilder, DebugHitStatsUniformBufferParameters->HitStatsOutput, 0);
+	}
+
 	FRDGBufferRef InstanceDebugBuffer = Scene->RayTracingScene.InstanceDebugBuffer;
 	if (InstanceDebugBuffer == nullptr)
 	{
@@ -1001,6 +1499,8 @@ void FDeferredShadingSceneRenderer::RenderRayTracingDebug(FRDGBuilder& GraphBuil
 		PermutationVectorCHS.Set<FRayTracingDebugCHS::FNaniteRayTracing>(true);
 		auto HitGroupShaderNaniteRT = View.ShaderMap->GetShader<FRayTracingDebugCHS>(PermutationVectorCHS);
 
+		// auto AnyHitShader = ShaderMap->GetShader<FRayTracingDebugAHS>();
+
 		FRHIRayTracingShader* HitGroupTable[] = { HitGroupShader.GetRayTracingShader(), HitGroupShaderNaniteRT.GetRayTracingShader()};
 		Initializer.SetHitGroupTable(HitGroupTable);
 		Initializer.bAllowHitGroupIndexing = true; // Required for stable output using GetBaseInstanceIndex().
@@ -1013,6 +1513,8 @@ void FDeferredShadingSceneRenderer::RenderRayTracingDebug(FRDGBuilder& GraphBuil
 		bRequiresBindings = true;
 	}
 
+	const uint32 NumInstances = View.RayTracingSceneInitData.NumNativeCPUInstances + View.RayTracingSceneInitData.NumNativeGPUInstances + View.RayTracingSceneInitData.NumNativeGPUSceneInstances;
+
 	FRayTracingDebugRGS::FParameters* RayGenParameters = GraphBuilder.AllocParameters<FRayTracingDebugRGS::FParameters>();
 
 	RayGenParameters->VisualizationMode = DebugVisualizationMode;
@@ -1021,6 +1523,10 @@ void FDeferredShadingSceneRenderer::RenderRayTracingDebug(FRDGBuilder& GraphBuil
 	RayGenParameters->TimingScale = CVarRayTracingDebugTimingScale.GetValueOnAnyThread() / 25000.0f;
 	RayGenParameters->OpaqueOnly = CVarRayTracingDebugModeOpaqueOnly.GetValueOnRenderThread();
 	RayGenParameters->TriangleHitCountMaxThreshold = FMath::Clamp((float)CVarRayTracingDebugHitCountMaxThreshold.GetValueOnRenderThread(), 1, 100000);
+	RayGenParameters->TriangleHitCountPerInstanceMaxThreshold = FMath::Max(1, CVarRayTracingDebugHitCountPerInstanceMaxThreshold.GetValueOnRenderThread());
+	RayGenParameters->RayTracingDebugHitStatsUniformBuffer = DebugHitStatsUniformBuffer;
+	RayGenParameters->TopKMostHitInstances = CVarRayTracingDebugHitCountTopKHits.GetValueOnRenderThread();
+	RayGenParameters->NumTotalInstances = NumInstances;
 
 	
 	// If we don't output depth, create dummy 1x1 texture
@@ -1052,6 +1558,7 @@ void FDeferredShadingSceneRenderer::RenderRayTracingDebug(FRDGBuilder& GraphBuil
 	RayGenParameters->TLAS = Scene->RayTracingScene.GetLayerView(ERayTracingSceneLayer::Base);
 	RayGenParameters->ViewUniformBuffer = View.ViewUniformBuffer;
 	RayGenParameters->Output = GraphBuilder.CreateUAV(SceneColorTexture);
+	RayGenParameters->TopKHitStats = GraphBuilder.CreateSRV(StatsBuffer);
 
 	RayGenParameters->SceneUniformBuffer = GetSceneUniformBufferRef(GraphBuilder); // TODO: use a separate params structure
 
