Make sure we early out of DrawDynamicMeshPass for UniformView path only when GPUScene enabled

diff --git a/Engine/Source/Runtime/Renderer/Private/MeshPassProcessor.cpp b/Engine/Source/Runtime/Renderer/Private/MeshPassProcessor.cpp
--- a/Engine/Source/Runtime/Renderer/Private/MeshPassProcessor.cpp
+++ b/Engine/Source/Runtime/Renderer/Private/MeshPassProcessor.cpp
@@ -1786,7 +1786,10 @@ void DrawDynamicMeshPassPrivate(
 {
 	if (VisibleMeshDrawCommands.Num() > 0)
 	{
-		if (PlatformGPUSceneUsesUniformBufferView(View.GetShaderPlatform()))
+		const ERHIFeatureLevel::Type FeatureLevel = View.GetFeatureLevel();
+		const bool bUseGPUScene = UseGPUScene(GMaxRHIShaderPlatform, FeatureLevel);
+
+		if (bUseGPUScene && PlatformGPUSceneUsesUniformBufferView(View.GetShaderPlatform()))
 		{
 			// FIXME: UniformBufferView path
 			return;
