LFV - added injection into volumetric fog of local fog volumes.
Order of composition is SkyAtmosphereAP, height fog, local fog volumes, volumetric fog.

Follow-up: apply when under water, apply on clouds, measure VGPRs impact.

#prelfight https://horde.devtools.epicgames.com/job/65099aa2c6396f37340323a7

diff --git a/Engine/Shaders/Private/HeightFogPixelShader.usf b/Engine/Shaders/Private/HeightFogPixelShader.usf
--- a/Engine/Shaders/Private/HeightFogPixelShader.usf
+++ b/Engine/Shaders/Private/HeightFogPixelShader.usf
@@ -10,6 +10,10 @@
 #include "HeightFogCommon.ush"
 #include "Random.ush"
 
+#if PERMUTATION_SUPPORT_LOCAL_FOG_VOLUME
+#include "LocalFogVolumes/LocalFogVolumeCommon.ush"
+#endif
+
 float UpsampleJitterMultiplier;
 
 float bOnlyOnRenderedOpaque;
@@ -71,6 +75,12 @@ void ExponentialPixelMain(
 
 	float4 HeightFogInscatteringAndOpacity = CalculateHeightFog(WorldPositionRelativeToCamera);
 
+#if PERMUTATION_SUPPORT_LOCAL_FOG_VOLUME
+	uint2 TilePos = uint2(SVPos.xy - View.ViewRectMin.xy) / LFVTilePixelSize.xx;
+	float4 LFVContribution = GetLFVContribution(PrimaryView, TilePos, WorldPositionRelativeToCamera);
+	HeightFogInscatteringAndOpacity = float4(LFVContribution.rgb + HeightFogInscatteringAndOpacity.rgb * LFVContribution.a, LFVContribution.a * HeightFogInscatteringAndOpacity.a);
+#endif
+
 	float4 FogInscatteringAndOpacity = CombineVolumetricFog(HeightFogInscatteringAndOpacity, VolumeUV, 0, SceneDepth);
 	float2 OcclusionTexCoord = clamp(TexCoord, OcclusionTextureMinMaxUV.xy, OcclusionTextureMinMaxUV.zw);
 	float LightShaftMask = Texture2DSample(OcclusionTexture, OcclusionSampler, OcclusionTexCoord).x;
diff --git a/Engine/Shaders/Private/LocalFogVolumes/LocalFogVolumeCommon.ush b/Engine/Shaders/Private/LocalFogVolumes/LocalFogVolumeCommon.ush
--- a/Engine/Shaders/Private/LocalFogVolumes/LocalFogVolumeCommon.ush
+++ b/Engine/Shaders/Private/LocalFogVolumes/LocalFogVolumeCommon.ush
@@ -16,6 +16,7 @@ struct FLocalFogVolumeGPUInstanceData
 	float Density;
 	float HeightFalloff;
 	float HeightOffset;
+	float UniformScale;
 
 	float3 Albedo;
 	float  PhaseG;
@@ -29,28 +30,31 @@ struct FLocalFogVolumeGPUInstanceData
 
 #if defined(LFVStruct)
 
-#define LFVInsts				LFVStruct.LFV.LocalFogVolumeInstances
-#define LFVTileDataTex			LFVStruct.LFV.LocalFogVolumeTileDataTexture
+#define LFVInsts					LFVStruct.LFV.LocalFogVolumeInstances
+#define LFVTileDataTex				LFVStruct.LFV.LocalFogVolumeTileDataTexture
 
-#define LFVTileDataResolution	LFVStruct.LFV.LocalFogVolumeTileDataTextureResolution
-#define LFVInstcount			LFVStruct.LFV.LocalFogVolumeInstanceCount
-#define LFVTilePixelSize		LFVStruct.LFV.LocalFogVolumeTilePixelSize
+#define LFVTileDataResolution		LFVStruct.LFV.LocalFogVolumeTileDataTextureResolution
+#define LFVInstcount				LFVStruct.LFV.LocalFogVolumeInstanceCount
+#define LFVTilePixelSize			LFVStruct.LFV.LocalFogVolumeTilePixelSize
+#define LFVRenderInVolumetricFog	LFVStruct.LFV.ShouldRenderLocalFogVolumeInVolumetricFog
 
 #else
 
-uint2							LFV_LocalFogVolumeTileDataTextureResolution;
-uint							LFV_LocalFogVolumeInstanceCount;
-uint							LFV_LocalFogVolumeTilePixelSize;
+uint2								LFV_LocalFogVolumeTileDataTextureResolution;
+uint								LFV_LocalFogVolumeInstanceCount;
+uint								LFV_LocalFogVolumeTilePixelSize;
+uint								LFV_ShouldRenderLocalFogVolumeInVolumetricFog;
 
-Buffer<float4>					LFV_LocalFogVolumeInstances;
-Texture2DArray<uint>			LFV_LocalFogVolumeTileDataTexture;
+Buffer<float4>						LFV_LocalFogVolumeInstances;
+Texture2DArray<uint>				LFV_LocalFogVolumeTileDataTexture;
 
-#define LFVInsts				LFV_LocalFogVolumeInstances
-#define LFVTileDataTex			LFV_LocalFogVolumeTileDataTexture
+#define LFVInsts					LFV_LocalFogVolumeInstances
+#define LFVTileDataTex				LFV_LocalFogVolumeTileDataTexture
 
-#define LFVTileDataResolution	LFV_LocalFogVolumeTileDataTextureResolution
-#define LFVInstcount			LFV_LocalFogVolumeInstanceCount
-#define LFVTilePixelSize		LFV_LocalFogVolumeTilePixelSize
+#define LFVTileDataResolution		LFV_LocalFogVolumeTileDataTextureResolution
+#define LFVInstcount				LFV_LocalFogVolumeInstanceCount
+#define LFVTilePixelSize			LFV_LocalFogVolumeTilePixelSize
+#define LFVRenderInVolumetricFog	LFV_ShouldRenderLocalFogVolumeInVolumetricFog
 
 #endif
 
@@ -67,6 +71,7 @@ FLocalFogVolumeGPUInstanceData GetLocalFogVolumeGPUInstanceData(uint Index)
 	Data.Density = Temp0.x;
 	Data.HeightFalloff = Temp0.y;
 	Data.HeightOffset = Temp0.z;
+	Data.UniformScale = Temp0.w;
 
 	float4 Temp1 = LFVInsts[Offset + 9];
 	Data.Albedo = Temp1.xyz;
@@ -295,6 +300,16 @@ float4 GetLFVInstanceContribution(
 		float Length0 = length(P0Space - CamRayOriginSpace);
 		float Length1 = length(P1Space - CamRayOriginSpace);
 
+		if (LFVRenderInVolumetricFog)
+		{
+			// Volumetric fog covers up to MaxDistance along ViewZ, exclude analytical fog from this range
+			float CosAngle = dot(CamRayWorldDir, View.ViewForward);
+			float InvCosAngle = (CosAngle > 0.0001) ? rcp(CosAngle) : 0;
+			float ExcludeDistance = View.VolumetricFogMaxDistance * InvCosAngle;
+			Length0 = max(Length0, ExcludeDistance);
+			Length1 = max(Length1, ExcludeDistance);
+		}
+
 		// Clamp the traced distance to depth and compute world length.
 		Length0 = min(LengthD, Length0);
 		Length1 = min(LengthD, Length1);
diff --git a/Engine/Shaders/Private/MobileFog.usf b/Engine/Shaders/Private/MobileFog.usf
--- a/Engine/Shaders/Private/MobileFog.usf
+++ b/Engine/Shaders/Private/MobileFog.usf
@@ -9,6 +9,10 @@
 #include "HeightFogCommon.ush"
 #include "SkyAtmosphereCommon.ush"
 
+#if PERMUTATION_SUPPORT_LOCAL_FOG_VOLUME
+#include "LocalFogVolumes/LocalFogVolumeCommon.ush"
+#endif
+
 float StartDepthZ;
 
 void MobileFogVS(
@@ -46,13 +50,19 @@ void MobileFogPS(
 	FinalFog = CalculateHeightFog(WorldPositionRelativeToCamera);
 #endif
 
+#if PERMUTATION_SUPPORT_LOCAL_FOG_VOLUME
+	uint2 TilePos = uint2(SVPos.xy - ResolvedView.ViewRectMin.xy) / LFVTilePixelSize.xx;
+	float4 LFVContribution = GetLFVContribution(ResolvedView, TilePos, WorldPositionRelativeToCamera);
+	FinalFog = float4(LFVContribution.rgb + FinalFog.rgb * LFVContribution.a, LFVContribution.a * FinalFog.a);
+#endif
+
 #if PERMUTATION_SUPPORT_VOLUMETRIC_FOG
 	if (FogStruct.ApplyVolumetricFog > 0) 
-		{
-			const uint EyeIndex = 0;
-			float3 VolumeUV = ComputeVolumeUV(LWCPromote(WorldPosition), ResolvedView.WorldToClip);
-			FinalFog = CombineVolumetricFog(FinalFog, VolumeUV, EyeIndex, SceneDepth);
-		}
+	{
+		const uint EyeIndex = 0;
+		float3 VolumeUV = ComputeVolumeUV(LWCPromote(WorldPosition), ResolvedView.WorldToClip);
+		FinalFog = CombineVolumetricFog(FinalFog, VolumeUV, EyeIndex, SceneDepth);
+	}
 #endif
 
 #if PERMUTATION_SUPPORT_AERIAL_PERSPECTIVE
diff --git a/Engine/Shaders/Private/VolumetricFog.usf b/Engine/Shaders/Private/VolumetricFog.usf
--- a/Engine/Shaders/Private/VolumetricFog.usf
+++ b/Engine/Shaders/Private/VolumetricFog.usf
@@ -28,6 +28,10 @@
 #include "VirtualShadowMaps/VirtualShadowMapProjectionCommon.ush"
 #endif
 
+#if USE_LOCAL_FOG_VOLUMES
+#include "LocalFogVolumes/LocalFogVolumeCommon.ush"
+#endif
+
 #ifndef USE_EMISSIVE
 #define USE_EMISSIVE 1
 #endif
@@ -149,11 +153,68 @@ void MaterialSetupCS(
 	float3 Scattering = Albedo * Extinction;
 	float Absorption = max(Extinction - Luminance(Scattering), 0.0f);
 
+	float3 AdditionalEmissive = 0.0f;
+
+#if USE_LOCAL_FOG_VOLUMES
+	float3 CamRayTranslatedWorldOrigin = PrimaryView.TranslatedWorldCameraOrigin;
+	float3 PreViewTranslation = LWCHackToFloat(PrimaryView.PreViewTranslation);
+	//float3 TranslatedWorldPosition = WorldPosition + PreViewTranslation;
+
+	float CellOffset = .5f;
+	float SceneDepth;
+	float3 TranslatedWorldPosition = ComputeCellTranslatedWorldPosition(GridCoordinate, CellOffset, SceneDepth);
+
+	float4 VertexClipSpacePosition = mul(float4(TranslatedWorldPosition.xyz, 1), PrimaryView.TranslatedWorldToClip);
+	float2 SvPosition = (VertexClipSpacePosition.xy / VertexClipSpacePosition.w * float2(.5f, -.5f) + .5f) * PrimaryView.ViewSizeAndInvSize.xy;
+	uint2 TilePos = clamp(uint2(SvPosition.xy / float(LFVTilePixelSize)), uint2(0, 0), LFVTileDataResolution - 1);
+
+	uint LFVCount = LFVTileDataTex[uint3(TilePos, 0)];
+
+	//AdditionalEmissive += float3(TilePos.x * 0.00001, TilePos.y * 0.00001, 0.0);
+	//AdditionalEmissive += LFVCount * 0.00001;
+
+	for (int LFVIndex = 0; LFVIndex < LFVCount; ++LFVIndex)
+	{
+		uint InstanceIndex = LFVTileDataTex[uint3(TilePos, 1 + LFVIndex)];
+
+		float3 CamRayWorldOrigin = CamRayTranslatedWorldOrigin - PreViewTranslation;
+		float3 CamRayTranslatedWorldDir = normalize(TranslatedWorldPosition - CamRayTranslatedWorldOrigin);
+		float3 CamRayWorldDir = CamRayTranslatedWorldDir;
+
+		// LFV_TODO skip if out of sphere of influence
+		FLocalFogVolumeGPUInstanceData FogInstance = GetLocalFogVolumeGPUInstanceData(InstanceIndex);
+
+		float3 UnitSpacePos = mul(float4(WorldPosition, 1.0), FogInstance.InvTransform).xyz;
+
+		float SphereFade = saturate(1.0 - length(UnitSpacePos - float3(0, 0, 0)));
+
+		float LFVExtinction = 0.0;
+		if (FogInstance.FogMode == 0)
+		{
+			LFVExtinction = SphereFade <= 0.0f ? 0.0f : FogInstance.Density * exp(-FogInstance.HeightFalloff * UnitSpacePos.z) / FogInstance.UniformScale;
+		}
+		else // (FogInstance.FogMode == 1)
+		{
+			LFVExtinction = FogInstance.Density * SphereFade / FogInstance.UniformScale;
+		}
+
+		float3 LFVScattering = LFVExtinction * FogInstance.Albedo;
+
+		Scattering += LFVScattering;
+		Absorption += max(LFVExtinction - Luminance(LFVScattering), 0.0f);
+
+#if USE_EMISSIVE
+		// Emissive follow the density of the medium
+		AdditionalEmissive += LFVExtinction * FogInstance.Emissive;
+#endif
+	}
+#endif // USE_LOCAL_FOG_VOLUMES
+
 	if (all((int3)GridCoordinate < VolumetricFog.ResourceGridSizeInt))
 	{
 		RWVBufferA[GridCoordinate] = float4(Scattering, Absorption);
 #if USE_EMISSIVE
-		RWVBufferB[GridCoordinate] = float4(GlobalEmissive.rgb, 0); 
+		RWVBufferB[GridCoordinate] = float4(GlobalEmissive.rgb + AdditionalEmissive, 0);
 #endif
 	}
 }
diff --git a/Engine/Source/Runtime/Engine/Private/LocalFogVolumeSceneProxy.cpp b/Engine/Source/Runtime/Engine/Private/LocalFogVolumeSceneProxy.cpp
--- a/Engine/Source/Runtime/Engine/Private/LocalFogVolumeSceneProxy.cpp
+++ b/Engine/Source/Runtime/Engine/Private/LocalFogVolumeSceneProxy.cpp
@@ -13,6 +13,7 @@ FLocalFogVolumeSceneProxy::FLocalFogVolumeSceneProxy(const ULocalFogVolumeCompon
 	: FogDensity(InComponent->FogDensity)
 	, FogHeightFalloff(InComponent->FogHeightFalloff)
 	, FogHeightOffset(InComponent->FogHeightOffset)
+	, FogUniformScale(1.0f)
 	, FogMode((uint8)InComponent->FogMode)
 	, FogSortPriority(uint8(127 - int8(InComponent->FogSortPriority))) // FogSortPriority on the component is in [-127,127] and needs to be negated to match expected priority behavior.
 	, FogPhaseG(InComponent->FogPhaseG)
@@ -31,5 +32,7 @@ void FLocalFogVolumeSceneProxy::UpdateComponentTransform(const FTransform& Trans
 	FogTransform = Transform;
 	const float MaximumAxisScale = FogTransform.GetMaximumAxisScale();
 	FogTransform.SetScale3D(FVector(MaximumAxisScale, MaximumAxisScale, MaximumAxisScale));
+
+	FogUniformScale = MaximumAxisScale;
 }
 
diff --git a/Engine/Source/Runtime/Engine/Public/LocalFogVolumeSceneProxy.h b/Engine/Source/Runtime/Engine/Public/LocalFogVolumeSceneProxy.h
--- a/Engine/Source/Runtime/Engine/Public/LocalFogVolumeSceneProxy.h
+++ b/Engine/Source/Runtime/Engine/Public/LocalFogVolumeSceneProxy.h
@@ -23,6 +23,7 @@ public:
 	float FogDensity;
 	float FogHeightFalloff;
 	float FogHeightOffset;
+	float FogUniformScale;
 
 	uint8 FogMode;
 	uint8 FogSortPriority;
diff --git a/Engine/Source/Runtime/Renderer/Private/DeferredShadingRenderer.cpp b/Engine/Source/Runtime/Renderer/Private/DeferredShadingRenderer.cpp
--- a/Engine/Source/Runtime/Renderer/Private/DeferredShadingRenderer.cpp
+++ b/Engine/Source/Runtime/Renderer/Private/DeferredShadingRenderer.cpp
@@ -3114,6 +3114,9 @@ void FDeferredShadingSceneRenderer::Render(FRDGBuilder& GraphBuilder)
 	FSceneTextures& SceneTextures = GetActiveSceneTextures();
 
 	const bool bUseGBuffer = IsUsingGBuffers(ShaderPlatform);
+	const bool bShouldRenderVolumetricFog = ShouldRenderVolumetricFog();
+	const bool bShouldRenderLocalFogVolume = ShouldRenderLocalFogVolume(Scene, ViewFamily);
+	const bool bShouldRenderLocalFogVolumeInVolumetricFog = ShouldRenderLocalFogVolumeInVolumetricFog(Scene, ViewFamily, bShouldRenderLocalFogVolume);
 
 	const bool bRenderDeferredLighting = ViewFamily.EngineShowFlags.Lighting
 		&& FeatureLevel >= ERHIFeatureLevel::SM5
@@ -3158,7 +3161,7 @@ void FDeferredShadingSceneRenderer::Render(FRDGBuilder& GraphBuilder)
 			}
 
 			bComputeLightGrid |= (
-				ShouldRenderVolumetricFog() ||
+				bShouldRenderVolumetricFog ||
 				VolumetricCloudWantsToSampleLocalLights(Scene, ViewFamily.EngineShowFlags) ||
 				ViewFamily.ViewMode != VMI_Lit ||
 				bAnyLumenEnabled ||
@@ -3428,7 +3431,7 @@ void FDeferredShadingSceneRenderer::Render(FRDGBuilder& GraphBuilder)
 		InitVolumetricCloudsForViews(GraphBuilder, bShouldRenderVolumetricCloudBase, InstanceCullingManager);
 
 		// Run local fog volume initialization before base pass for when data is needed in forward
-		InitLocalFogVolumesForViews(Scene, Views, ViewFamily, GraphBuilder);
+		InitLocalFogVolumesForViews(Scene, Views, ViewFamily, GraphBuilder, bShouldRenderVolumetricFog);
 
 		if (SkyAtmospherePassLocation == ESkyAtmospherePassLocation::BeforeOcclusion && bShouldRenderSkyAtmosphere)
 		{
@@ -4097,12 +4100,26 @@ void FDeferredShadingSceneRenderer::Render(FRDGBuilder& GraphBuilder)
 		}
 
 		// Draw fog.
+		bool bHeightFogHasComposedLocalFogVolume = false;
 		if (!bHasRayTracedOverlay && ShouldRenderFog(ViewFamily))
 		{
 			RDG_CSV_STAT_EXCLUSIVE_SCOPE(GraphBuilder, RenderFog);
 			SCOPED_NAMED_EVENT(RenderFog, FColor::Emerald);
 			SCOPE_CYCLE_COUNTER(STAT_FDeferredShadingSceneRenderer_RenderFog);
-			RenderFog(GraphBuilder, SceneTextures, LightShaftOcclusionTexture);
+			const bool bFogComposeLocalFogVolumes = bShouldRenderLocalFogVolumeInVolumetricFog && bShouldRenderVolumetricFog;
+			RenderFog(GraphBuilder, SceneTextures, LightShaftOcclusionTexture, bFogComposeLocalFogVolumes);
+			bHeightFogHasComposedLocalFogVolume = bFogComposeLocalFogVolumes;
+		}
+
+		// Local Fog Volumes (LFV) rendering order is first HeightFog, then LFV, then volumetric fog on top.
+		// LFVs are rendered as part of the regular height fog + volumetric fog pass when volumetric fog is enabled and it is requested to voxelise LFVs into volumetric fog.
+		// Otherwise, they are rendered in an independent pass (this for instance make it independent of the near clip plane optimization).
+		if (!bHasRayTracedOverlay && !bHeightFogHasComposedLocalFogVolume)
+		{
+			RDG_CSV_STAT_EXCLUSIVE_SCOPE(GraphBuilder, RenderLocalFogVolume);
+			SCOPED_NAMED_EVENT(RenderLocalFogVolume, FColor::Emerald);
+			SCOPE_CYCLE_COUNTER(STAT_FDeferredShadingSceneRenderer_RenderLocalFogVolume);
+			RenderLocalFogVolume(Scene, Views, ViewFamily, GraphBuilder, SceneTextures, LightShaftOcclusionTexture);
 		}
 
 		// After the height fog, Draw volumetric clouds (having fog applied on them already) when using per pixel tracing,
@@ -4120,14 +4137,6 @@ void FDeferredShadingSceneRenderer::Render(FRDGBuilder& GraphBuilder)
 			ComposeVolumetricRenderTargetOverScene(GraphBuilder, Views, SceneTextures.Color.Target, SceneTextures.Depth.Target, bShouldRenderSingleLayerWater, SceneWithoutWaterTextures, SceneTextures);
 		}
 
-		if (!bHasRayTracedOverlay && ShouldRenderLocalFogVolume(Scene, ViewFamily))
-		{
-			RDG_CSV_STAT_EXCLUSIVE_SCOPE(GraphBuilder, RenderLocalFogVolume);
-			SCOPED_NAMED_EVENT(RenderLocalFogVolume, FColor::Emerald);
-			SCOPE_CYCLE_COUNTER(STAT_FDeferredShadingSceneRenderer_RenderLocalFogVolume);
-			RenderLocalFogVolume(Scene, Views, ViewFamily, GraphBuilder, SceneTextures, LightShaftOcclusionTexture);
-		}
-
 		FRDGTextureRef ExposureIlluminance = AddCalculateExposureIlluminancePass(GraphBuilder, Views, SceneTextures, TranslucencyLightingVolumeTextures, ExposureIlluminanceSetup);
 
 		if (!bHasRayTracedOverlay)
diff --git a/Engine/Source/Runtime/Renderer/Private/DeferredShadingRenderer.h b/Engine/Source/Runtime/Renderer/Private/DeferredShadingRenderer.h
--- a/Engine/Source/Runtime/Renderer/Private/DeferredShadingRenderer.h
+++ b/Engine/Source/Runtime/Renderer/Private/DeferredShadingRenderer.h
@@ -629,7 +629,8 @@ private:
 	void RenderFog(
 		FRDGBuilder& GraphBuilder,
 		const FMinimalSceneTextures& SceneTextures,
-		FRDGTextureRef LightShaftOcclusionTexture);
+		FRDGTextureRef LightShaftOcclusionTexture,
+		bool bFogComposeLocalFogVolumes);
 
 	void RenderUnderWaterFog(
 		FRDGBuilder& GraphBuilder,
diff --git a/Engine/Source/Runtime/Renderer/Private/FogRendering.cpp b/Engine/Source/Runtime/Renderer/Private/FogRendering.cpp
--- a/Engine/Source/Runtime/Renderer/Private/FogRendering.cpp
+++ b/Engine/Source/Runtime/Renderer/Private/FogRendering.cpp
@@ -121,11 +121,13 @@ class FExponentialHeightFogPS : public FGlobalShader
 	class FSupportFogInScatteringTexture : SHADER_PERMUTATION_BOOL("PERMUTATION_SUPPORT_FOG_INSCATTERING_TEXTURE");
 	class FSupportFogDirectionalLightInScattering : SHADER_PERMUTATION_BOOL("PERMUTATION_SUPPORT_FOG_DIRECTIONAL_LIGHT_INSCATTERING");
 	class FSupportVolumetricFog : SHADER_PERMUTATION_BOOL("PERMUTATION_SUPPORT_VOLUMETRIC_FOG");
-	using FPermutationDomain = TShaderPermutationDomain<FSupportFogInScatteringTexture, FSupportFogDirectionalLightInScattering, FSupportVolumetricFog>;
+	class FSupportLocalFogVolume : SHADER_PERMUTATION_BOOL("PERMUTATION_SUPPORT_LOCAL_FOG_VOLUME");
+	using FPermutationDomain = TShaderPermutationDomain<FSupportFogInScatteringTexture, FSupportFogDirectionalLightInScattering, FSupportVolumetricFog, FSupportLocalFogVolume>;
 
 	BEGIN_SHADER_PARAMETER_STRUCT(FParameters, )
 		SHADER_PARAMETER_STRUCT_REF(FViewUniformShaderParameters, ViewUniformBuffer)
 		SHADER_PARAMETER_RDG_UNIFORM_BUFFER(FFogUniformParameters, FogUniformBuffer)
+		SHADER_PARAMETER_STRUCT(FLocalFogVolumeUniformParameters, LFV)
 		SHADER_PARAMETER_RDG_TEXTURE(Texture2D, OcclusionTexture)
 		SHADER_PARAMETER_SAMPLER(SamplerState, OcclusionSampler)
 		SHADER_PARAMETER_RDG_TEXTURE(Texture2D, WaterDepthTexture)
@@ -267,6 +269,7 @@ static FFogPassParameters* CreateDefaultFogPassParameters(
 	PassParameters->VS.ViewUniformBuffer = GetShaderBinding(View.ViewUniformBuffer);
 	PassParameters->PS.ViewUniformBuffer = GetShaderBinding(View.ViewUniformBuffer);
 	PassParameters->PS.FogUniformBuffer = FogUniformBuffer;
+	PassParameters->PS.LFV = View.LocalFogVolumeViewData.UniformParametersStruct;
 	PassParameters->PS.OcclusionTexture = LightShaftOcclusionTexture != nullptr ? LightShaftOcclusionTexture : GSystemTextures.GetWhiteDummy(GraphBuilder);
 	PassParameters->PS.OcclusionSampler = TStaticSamplerState<SF_Bilinear, AM_Clamp, AM_Clamp, AM_Clamp>::GetRHI();
 	PassParameters->PS.WaterDepthTexture = GSystemTextures.GetDepthDummy(GraphBuilder);
@@ -284,7 +287,8 @@ static void RenderViewFog(
 	const FViewInfo& View, 
 	FIntRect ViewRect, 
 	FFogPassParameters* PassParameters, 
-	bool bShouldRenderVolumetricFog)
+	bool bShouldRenderVolumetricFog,
+	bool bFogComposeLocalFogVolumes)
 {
 	FGraphicsPipelineStateInitializer GraphicsPSOInit;
 	RHICmdList.ApplyCachedRenderTargets(GraphicsPSOInit);
@@ -301,11 +305,15 @@ static void RenderViewFog(
 	GraphicsPSOInit.BoundShaderState.VertexDeclarationRHI = GFogVertexDeclaration.VertexDeclarationRHI;
 	GraphicsPSOInit.BoundShaderState.VertexShaderRHI = VertexShader.GetVertexShader();
 
+	// We should only render local fog volume here when volumetric fog is enabled, making sure that the depth bound test is not used.
+	check(!bFogComposeLocalFogVolumes || (bFogComposeLocalFogVolumes && bShouldRenderVolumetricFog));
+
 	const bool bUseFogInscatteringColorCubemap = View.FogInscatteringColorCubemap != nullptr;
 	FExponentialHeightFogPS::FPermutationDomain PsPermutationVector;
 	PsPermutationVector.Set<FExponentialHeightFogPS::FSupportFogInScatteringTexture>(bUseFogInscatteringColorCubemap);
 	PsPermutationVector.Set<FExponentialHeightFogPS::FSupportFogDirectionalLightInScattering>(!bUseFogInscatteringColorCubemap && View.bUseDirectionalInscattering);
 	PsPermutationVector.Set<FExponentialHeightFogPS::FSupportVolumetricFog>(bShouldRenderVolumetricFog);
+	PsPermutationVector.Set<FExponentialHeightFogPS::FSupportLocalFogVolume>(bFogComposeLocalFogVolumes);
 	TShaderMapRef<FExponentialHeightFogPS> PixelShader(View.ShaderMap, PsPermutationVector);
 	GraphicsPSOInit.BoundShaderState.PixelShaderRHI = PixelShader.GetPixelShader();
 
@@ -363,7 +371,8 @@ static void RenderViewFog(
 void FDeferredShadingSceneRenderer::RenderFog(
 	FRDGBuilder& GraphBuilder,
 	const FMinimalSceneTextures& SceneTextures,
-	FRDGTextureRef LightShaftOcclusionTexture)
+	FRDGTextureRef LightShaftOcclusionTexture,
+	bool bFogComposeLocalFogVolumes)
 {
 	if (Scene->ExponentialFogs.Num() > 0 
 		// Fog must be done in the base pass for MSAA to work
@@ -393,9 +402,9 @@ void FDeferredShadingSceneRenderer::RenderFog(
 				PassParameters->RenderTargets.DepthStencil = FDepthStencilBinding(SceneTextures.Depth.Target, ERenderTargetLoadAction::ELoad, ERenderTargetLoadAction::ELoad, FExclusiveDepthStencil::DepthRead_StencilWrite);
 
 				GraphBuilder.AddPass(RDG_EVENT_NAME("Fog"), PassParameters, ERDGPassFlags::Raster, 
-					[this, &View, PassParameters, bShouldRenderVolumetricFog](FRHICommandList& RHICmdList)
+					[this, &View, PassParameters, bShouldRenderVolumetricFog, bFogComposeLocalFogVolumes](FRHICommandList& RHICmdList)
 				{
-					RenderViewFog(RHICmdList, View, View.ViewRect, PassParameters, bShouldRenderVolumetricFog);
+					RenderViewFog(RHICmdList, View, View.ViewRect, PassParameters, bShouldRenderVolumetricFog, bFogComposeLocalFogVolumes);
 				});
 			}
 		}
@@ -443,7 +452,8 @@ void FDeferredShadingSceneRenderer::RenderUnderWaterFog(
 
 				GraphBuilder.AddPass(RDG_EVENT_NAME("FogBehindWater"), PassParameters, ERDGPassFlags::Raster, [this, &View, SceneWithoutWaterView, PassParameters, bShouldRenderVolumetricFog](FRHICommandList& RHICmdList)
 				{
-					RenderViewFog(RHICmdList, View, SceneWithoutWaterView.ViewRect, PassParameters, bShouldRenderVolumetricFog);
+					const bool bFogComposeLocalFogVolumes = false; // LFV_TODO implement for under water
+					RenderViewFog(RHICmdList, View, SceneWithoutWaterView.ViewRect, PassParameters, bShouldRenderVolumetricFog, bFogComposeLocalFogVolumes);
 				});
 			}
 		}
diff --git a/Engine/Source/Runtime/Renderer/Private/LocalFogVolumeRendering.cpp b/Engine/Source/Runtime/Renderer/Private/LocalFogVolumeRendering.cpp
--- a/Engine/Source/Runtime/Renderer/Private/LocalFogVolumeRendering.cpp
+++ b/Engine/Source/Runtime/Renderer/Private/LocalFogVolumeRendering.cpp
@@ -14,6 +14,11 @@ static TAutoConsoleVariable<int32> CVarLocalFogVolume(
 	TEXT("LocalFogVolume components are rendered when this is not 0, otherwise ignored.\n"),
 	ECVF_RenderThreadSafe);
 
+static TAutoConsoleVariable<int32> CVarLocalFogVolumeRenderIntoVolumetricFog(
+	TEXT("r.LocalFogVolume.RenderIntoVolumetricFog"), 0,
+	TEXT("LocalFogVolume are going to be voxelised into the volumetric fog when this is not 0, otherwise it will remain isolated.\n"),
+	ECVF_RenderThreadSafe);
+
 static TAutoConsoleVariable<int32> CVarLocalFogVolumeApplyOnTranslucent(
 	TEXT("r.LocalFogVolume.ApplyOnTranslucent"), 0,
 	TEXT("Project settings enabling the sampling of local fog volumes on translucent elements.\n"),
@@ -59,16 +64,26 @@ static uint32 GetLocalFogVolumeTileMaxInstanceCount()
 	return FMath::Max(1u, FMath::Min(256u, (uint32)CVarLocalFogVolumeTileMaxInstanceCount.GetValueOnRenderThread()));
 }
 
-bool ShouldRenderLocalFogVolume(const FScene* Scene, const FSceneViewFamily& Family)
+bool ShouldRenderLocalFogVolume(const FScene* Scene, const FSceneViewFamily& SceneViewFamily)
 {
-	const FEngineShowFlags EngineShowFlags = Family.EngineShowFlags;
-	if (Scene && Scene->HasAnyLocalFogVolume() && EngineShowFlags.Fog && !Family.UseDebugViewPS())
+	const FEngineShowFlags EngineShowFlags = SceneViewFamily.EngineShowFlags;
+	if (Scene && Scene->HasAnyLocalFogVolume() && EngineShowFlags.Fog && !SceneViewFamily.UseDebugViewPS())
 	{
 		return CVarLocalFogVolume.GetValueOnRenderThread() > 0;
 	}
 	return false;
 }
 
+bool ShouldRenderLocalFogVolumeInVolumetricFog(const FScene* Scene, const FSceneViewFamily& SceneViewFamily, bool bShouldRenderVolumetricFog)
+{
+	if (ShouldRenderLocalFogVolume(Scene, SceneViewFamily) && bShouldRenderVolumetricFog)
+	{
+		return CVarLocalFogVolumeRenderIntoVolumetricFog.GetValueOnRenderThread() > 0;
+	}
+	return false;
+}
+
+
 DECLARE_GPU_STAT(LocalFogVolumeVolumes);
 
 static const uint32 SizeOfUintVec4 = sizeof(FUintVector4);
@@ -263,9 +278,6 @@ void GetLocalFogVolumeSortingData(const FScene* Scene, FRDGBuilder& GraphBuilder
 			continue; // this volume will never be visible
 		}
 
-		FTransform TransformScaleOnly;
-		TransformScaleOnly.SetScale3D(LHF->FogTransform.GetScale3D());
-
 		FLocalFogVolumeGPUInstanceData* LocalFogVolumeGPUInstanceDataIt = &Out.LocalFogVolumeGPUInstanceData[Out.LocalFogVolumeInstanceCountFinal];
 		LocalFogVolumeGPUInstanceDataIt->Transform = FMatrix44f(LHF->FogTransform.ToMatrixWithScale());
 		LocalFogVolumeGPUInstanceDataIt->InvTransform = LocalFogVolumeGPUInstanceDataIt->Transform.Inverse();
@@ -273,6 +285,7 @@ void GetLocalFogVolumeSortingData(const FScene* Scene, FRDGBuilder& GraphBuilder
 		LocalFogVolumeGPUInstanceDataIt->Density = LHF->FogDensity;
 		LocalFogVolumeGPUInstanceDataIt->HeightFalloff = LHF->FogHeightFalloff * 0.01f;	// This scale is used to have artist author reasonable range.
 		LocalFogVolumeGPUInstanceDataIt->HeightOffset = LHF->FogHeightOffset;
+		LocalFogVolumeGPUInstanceDataIt->UniformScale = LHF->FogUniformScale;
 
 		LocalFogVolumeGPUInstanceDataIt->FogMode = float(LHF->FogMode);
 
@@ -293,7 +306,7 @@ void GetLocalFogVolumeSortingData(const FScene* Scene, FRDGBuilder& GraphBuilder
 	Out.LocalFogVolumeSortKeys.SetNum(Out.LocalFogVolumeInstanceCountFinal, false/*bAllowShrinking*/);
 }
 
-void CreateViewLocalFogVolumeBufferSRV(FViewInfo& View, FRDGBuilder& GraphBuilder, FLocalFogVolumeSortingData& SortingData)
+void CreateViewLocalFogVolumeBufferSRV(FViewInfo& View, FRDGBuilder& GraphBuilder, FLocalFogVolumeSortingData& SortingData, bool bShouldRenderLocalFogVolumeInVolumetricFog)
 {
 	if (SortingData.LocalFogVolumeInstanceCountFinal == 0)
 	{
@@ -326,9 +339,8 @@ void CreateViewLocalFogVolumeBufferSRV(FViewInfo& View, FRDGBuilder& GraphBuilde
 		// We could also have an indirection buffer on GPU but choosing to go with the sorting + copy on CPU since it is expected to not have many local height fog volumes.
 		LocalFogVolumeGPUSortedInstanceData[i] = SortingData.LocalFogVolumeGPUInstanceData[LFVKey.FogVolume.Index];
 
-		const float LFVMaximumAxisScale = LocalFogVolumeGPUSortedInstanceData[i].Transform.GetMaximumAxisScale();
 		FVector& LFVPosition = SortingData.LocalFogVolumeCenterPos[LFVKey.FogVolume.Index];
-		LocalFogVolumeGPUSortedInstanceCullingData[i] = FVector4f(LFVPosition.X, LFVPosition.Y, LFVPosition.Z, LFVMaximumAxisScale);
+		LocalFogVolumeGPUSortedInstanceCullingData[i] = FVector4f(LFVPosition.X, LFVPosition.Y, LFVPosition.Z, LocalFogVolumeGPUSortedInstanceData[i].UniformScale);
 	}
 
 	// 3. Allocate buffer and initialize with sorted data to upload to GPU
@@ -359,12 +371,13 @@ void CreateViewLocalFogVolumeBufferSRV(FViewInfo& View, FRDGBuilder& GraphBuilde
 	View.LocalFogVolumeViewData.TileDataTextureArraySRV		= GraphBuilder.CreateSRV(FRDGTextureSRVDesc::Create(View.LocalFogVolumeViewData.TileDataTextureArray));
 	View.LocalFogVolumeViewData.TileDataTextureArrayUAV		= GraphBuilder.CreateUAV(FRDGTextureUAVDesc(View.LocalFogVolumeViewData.TileDataTextureArray));
 
-	View.LocalFogVolumeViewData.UniformParametersStruct.LocalFogVolumeCommon.LocalFogVolumeTileDataTextureResolution= FUintVector2(TileDataTextureResolution.X, TileDataTextureResolution.Y);
-	View.LocalFogVolumeViewData.UniformParametersStruct.LocalFogVolumeCommon.LocalFogVolumeInstanceCount			= View.LocalFogVolumeViewData.GPUInstanceCount;
-	View.LocalFogVolumeViewData.UniformParametersStruct.LocalFogVolumeCommon.LocalFogVolumeTilePixelSize			= LocalFogVolumeTilePixelSize;
-	View.LocalFogVolumeViewData.UniformParametersStruct.LocalFogVolumeCommon.LocalFogVolumeInstances				= View.LocalFogVolumeViewData.GPUInstanceDataBufferSRV;
-	View.LocalFogVolumeViewData.UniformParametersStruct.LocalFogVolumeTileDataTexture								= View.LocalFogVolumeViewData.TileDataTextureArraySRV;
-	View.LocalFogVolumeViewData.UniformBuffer																		= GraphBuilder.CreateUniformBuffer(&View.LocalFogVolumeViewData.UniformParametersStruct);
+	View.LocalFogVolumeViewData.UniformParametersStruct.LocalFogVolumeCommon.LocalFogVolumeTileDataTextureResolution	= FUintVector2(TileDataTextureResolution.X, TileDataTextureResolution.Y);
+	View.LocalFogVolumeViewData.UniformParametersStruct.LocalFogVolumeCommon.LocalFogVolumeInstanceCount				= View.LocalFogVolumeViewData.GPUInstanceCount;
+	View.LocalFogVolumeViewData.UniformParametersStruct.LocalFogVolumeCommon.LocalFogVolumeTilePixelSize				= LocalFogVolumeTilePixelSize;
+	View.LocalFogVolumeViewData.UniformParametersStruct.LocalFogVolumeCommon.ShouldRenderLocalFogVolumeInVolumetricFog	= bShouldRenderLocalFogVolumeInVolumetricFog ? 1 : 0;
+	View.LocalFogVolumeViewData.UniformParametersStruct.LocalFogVolumeCommon.LocalFogVolumeInstances					= View.LocalFogVolumeViewData.GPUInstanceDataBufferSRV;
+	View.LocalFogVolumeViewData.UniformParametersStruct.LocalFogVolumeTileDataTexture									= View.LocalFogVolumeViewData.TileDataTextureArraySRV;
+	View.LocalFogVolumeViewData.UniformBuffer																			= GraphBuilder.CreateUniformBuffer(&View.LocalFogVolumeViewData.UniformParametersStruct);
 
 	// This buffer must remain a basic vertex buffer for mobile to be able to read it from vertex shader
 	View.LocalFogVolumeViewData.GPUTileDataBuffer = CreateVertexBuffer(
@@ -381,11 +394,13 @@ void CreateViewLocalFogVolumeBufferSRV(FViewInfo& View, FRDGBuilder& GraphBuilde
 void InitLocalFogVolumesForViews(
 	const FScene* Scene,
 	TArray<FViewInfo>& Views,
-	const FSceneViewFamily& Family,
-	FRDGBuilder& GraphBuilder)
+	const FSceneViewFamily& SceneViewFamily,
+	FRDGBuilder& GraphBuilder,
+	bool bShouldRenderVolumetricFog)
 {
 	const uint32 LocalFogVolumeInstanceCount = Scene->LocalFogVolumes.Num();
-	if (LocalFogVolumeInstanceCount > 0 && ShouldRenderLocalFogVolume(Scene, Family))
+	const bool bShouldRenderLocalFogVolume = ShouldRenderLocalFogVolume(Scene, SceneViewFamily);
+	if (LocalFogVolumeInstanceCount > 0 && bShouldRenderLocalFogVolume)
 	{
 		RDG_GPU_STAT_SCOPE(GraphBuilder, LocalFogVolumeVolumes);
 
@@ -394,7 +409,7 @@ void InitLocalFogVolumesForViews(
 
 		for (FViewInfo& View : Views)
 		{
-			CreateViewLocalFogVolumeBufferSRV(View, GraphBuilder, SortingData);
+			CreateViewLocalFogVolumeBufferSRV(View, GraphBuilder, SortingData, ShouldRenderLocalFogVolumeInVolumetricFog(Scene, SceneViewFamily, bShouldRenderVolumetricFog));
 
 			LocalFogVolumeViewTiledCullingPass(View, GraphBuilder);
 		}
@@ -553,13 +568,13 @@ END_SHADER_PARAMETER_STRUCT()
 void RenderLocalFogVolume(
 	const FScene* Scene,
 	TArray<FViewInfo>& Views,
-	const FSceneViewFamily& Family,
+	const FSceneViewFamily& SceneViewFamily,
 	FRDGBuilder& GraphBuilder,
 	const FMinimalSceneTextures& SceneTextures,
 	FRDGTextureRef LightShaftOcclusionTexture)
 {
 	uint32 LocalFogVolumeInstanceCount = Scene->LocalFogVolumes.Num();
-	if (LocalFogVolumeInstanceCount > 0 && ShouldRenderLocalFogVolume(Scene, Family))
+	if (LocalFogVolumeInstanceCount > 0 && ShouldRenderLocalFogVolume(Scene, SceneViewFamily))
 	{
 		RDG_GPU_STAT_SCOPE(GraphBuilder, LocalFogVolumeVolumes);
 
diff --git a/Engine/Source/Runtime/Renderer/Private/LocalFogVolumeRendering.h b/Engine/Source/Runtime/Renderer/Private/LocalFogVolumeRendering.h
--- a/Engine/Source/Runtime/Renderer/Private/LocalFogVolumeRendering.h
+++ b/Engine/Source/Runtime/Renderer/Private/LocalFogVolumeRendering.h
@@ -24,6 +24,7 @@ BEGIN_SHADER_PARAMETER_STRUCT(FLocalFogVolumeCommonParameters, )
 	SHADER_PARAMETER(FUintVector2, LocalFogVolumeTileDataTextureResolution)
 	SHADER_PARAMETER(uint32, LocalFogVolumeInstanceCount)
 	SHADER_PARAMETER(uint32, LocalFogVolumeTilePixelSize)
+	SHADER_PARAMETER(uint32, ShouldRenderLocalFogVolumeInVolumetricFog)
 END_SHADER_PARAMETER_STRUCT()
 
 BEGIN_SHADER_PARAMETER_STRUCT(FLocalFogVolumeUniformParameters, )
@@ -44,7 +45,7 @@ public:
 	float Density;
 	float HeightFalloff;
 	float HeightOffset;
-	float Unused;
+	float UniformScale;
 
 	FVector3f Albedo;
 	float PhaseG;
@@ -117,23 +118,24 @@ struct FLocalFogVolumeViewData
 	Local height fog rendering functions
 =============================================================================*/
 
-bool ShouldRenderLocalFogVolume(const FScene* Scene, const FSceneViewFamily& Family);
+bool ShouldRenderLocalFogVolume(const FScene* Scene, const FSceneViewFamily& SceneViewFamily);
+bool ShouldRenderLocalFogVolumeInVolumetricFog(const FScene* Scene, const FSceneViewFamily& SceneViewFamily, bool bShouldRenderVolumetricFog);
 
 void GetLocalFogVolumeSortingData(const FScene* Scene, FRDGBuilder& GraphBuilder, FLocalFogVolumeSortingData& Out);
 
-void CreateViewLocalFogVolumeBufferSRV(FViewInfo& View, FRDGBuilder& GraphBuilder, FLocalFogVolumeSortingData& SortingData);
 void SetDummyLocalFogVolumeForView(FRDGBuilder& GraphBuilder, FViewInfo& View);
 
 void InitLocalFogVolumesForViews(
 	const FScene* Scene,
 	TArray<FViewInfo>& Views,
-	const FSceneViewFamily& Family,
-	FRDGBuilder& GraphBuilder);
+	const FSceneViewFamily& SceneViewFamily,
+	FRDGBuilder& GraphBuilder,
+	bool bShouldRenderVolumetricFog);
 
 void RenderLocalFogVolume(
 	const FScene* Scene,
 	TArray<FViewInfo>& Views,
-	const FSceneViewFamily& Family,
+	const FSceneViewFamily& SceneViewFamily,
 	FRDGBuilder& GraphBuilder,
 	const FMinimalSceneTextures& SceneTextures,
 	FRDGTextureRef LightShaftOcclusionTexture);
diff --git a/Engine/Source/Runtime/Renderer/Private/MobileFogRendering.cpp b/Engine/Source/Runtime/Renderer/Private/MobileFogRendering.cpp
--- a/Engine/Source/Runtime/Renderer/Private/MobileFogRendering.cpp
+++ b/Engine/Source/Runtime/Renderer/Private/MobileFogRendering.cpp
@@ -18,6 +18,7 @@
 #include "ProfilingDebugging/RealtimeGPUProfiler.h"
 #include "MobileBasePassRendering.h"
 #include "SkyAtmosphereRendering.h"
+#include "LocalFogVolumeRendering.h"
 
 static TAutoConsoleVariable<int32> CVarPixelFogQuality(
 	TEXT("r.Mobile.PixelFogQuality"),
@@ -63,6 +64,7 @@ class FMobileFogPS : public FGlobalShader
 	class FSupportFogDirectionalLightInScattering	: SHADER_PERMUTATION_BOOL("PERMUTATION_SUPPORT_FOG_DIRECTIONAL_LIGHT_INSCATTERING");
 	class FSupportAerialPerspective					: SHADER_PERMUTATION_BOOL("PERMUTATION_SUPPORT_AERIAL_PERSPECTIVE");
 	class FSupportVolumetricFog						: SHADER_PERMUTATION_BOOL("PERMUTATION_SUPPORT_VOLUMETRIC_FOG");
+	class FSupportLocalFogVolume					: SHADER_PERMUTATION_BOOL("PERMUTATION_SUPPORT_LOCAL_FOG_VOLUME");
 	
 	using FPermutationDomain = TShaderPermutationDomain< 
 		FSupportHeightFog, 
@@ -71,7 +73,8 @@ class FMobileFogPS : public FGlobalShader
 		FSupportFogSecondTerm,
 		FSupportFogDirectionalLightInScattering,
 		FSupportAerialPerspective,
-		FSupportVolumetricFog
+		FSupportVolumetricFog,
+		FSupportLocalFogVolume
 	>;
 
 	static FPermutationDomain RemapPermutation(FPermutationDomain PermutationVector)
@@ -111,14 +114,24 @@ class FMobileFogPS : public FGlobalShader
 
 	BEGIN_SHADER_PARAMETER_STRUCT(FParameters, )
 		SHADER_PARAMETER_STRUCT_INCLUDE(FViewShaderParameters, View)
-		SHADER_PARAMETER_RDG_UNIFORM_BUFFER(FMobileBasePassUniformParameters, MobileBasePass)
+		SHADER_PARAMETER_STRUCT(FLocalFogVolumeUniformParameters, LFV)
 	END_SHADER_PARAMETER_STRUCT()
 };
 
 IMPLEMENT_SHADER_TYPE(, FMobileFogPS, TEXT("/Engine/Private/MobileFog.usf"), TEXT("MobileFogPS"), SF_Pixel);
 
 void FMobileSceneRenderer::RenderFog(FRHICommandList& RHICmdList, const FViewInfo& View)
-{	
+{
+	// RenderFog has some extra logic to skip the rendering of fog. So we account for that inside this function using a lambda.
+	bool bFogHasComposedLocalFogVolumes = false;
+	auto RenderLocalFogVolumeMobileLambda = [&]()
+	{
+		if (!bFogHasComposedLocalFogVolumes)
+		{
+			RenderLocalFogVolumeMobile(RHICmdList, View);
+		}
+	};
+
 	static const auto* CVarDisableVertexFog = IConsoleManager::Get().FindTConsoleVariableDataInt(TEXT("r.Mobile.DisableVertexFog"));
 	if (CVarDisableVertexFog && CVarDisableVertexFog->GetValueOnRenderThread() == 0)
 	{
@@ -132,6 +145,7 @@ void FMobileSceneRenderer::RenderFog(FRHICommandList& RHICmdList, const FViewInf
 
 	if (!bUseAerialPerspective && !bUseHeightFog)
 	{
+		RenderLocalFogVolumeMobileLambda();
 		return;
 	}
 
@@ -175,6 +189,9 @@ void FMobileSceneRenderer::RenderFog(FRHICommandList& RHICmdList, const FViewInf
 	const bool bUseFogSecondTerm = PixelFogQuality > 0 && (View.ExponentialFogParameters2.X > 0);
 	const bool bUseFogDirectionalInscatering = PixelFogQuality > 0 && !bUseFogInscatteringColorCubemap && (View.DirectionalInscatteringColor.GetLuminance() > 0 || View.bUseDirectionalInscattering);
 	
+	const bool bShouldRenderVolumetricFog = ShouldRenderVolumetricFog();
+	bFogHasComposedLocalFogVolumes = bShouldRenderVolumetricFog && ShouldRenderLocalFogVolume(Scene, *View.Family) && ShouldRenderLocalFogVolumeInVolumetricFog(Scene, *View.Family, bShouldRenderVolumetricFog);
+
 	FMobileFogPS::FPermutationDomain PsPermutationVector;
 	PsPermutationVector.Set<FMobileFogPS::FSupportHeightFog>(bUseHeightFog);
 	PsPermutationVector.Set<FMobileFogPS::FSupportFogStartDistance>(bUseFogStartDistance);
@@ -182,7 +199,8 @@ void FMobileSceneRenderer::RenderFog(FRHICommandList& RHICmdList, const FViewInf
 	PsPermutationVector.Set<FMobileFogPS::FSupportFogSecondTerm>(bUseFogSecondTerm);
 	PsPermutationVector.Set<FMobileFogPS::FSupportAerialPerspective>(bUseAerialPerspective);
 	PsPermutationVector.Set<FMobileFogPS::FSupportFogDirectionalLightInScattering>(bUseFogDirectionalInscatering);
-	PsPermutationVector.Set<FMobileFogPS::FSupportVolumetricFog>(ShouldRenderVolumetricFog());
+	PsPermutationVector.Set<FMobileFogPS::FSupportVolumetricFog>(bShouldRenderVolumetricFog);
+	PsPermutationVector.Set<FMobileFogPS::FSupportLocalFogVolume>(bFogHasComposedLocalFogVolumes);
 	
 	TShaderMapRef<FMobileFogPS> PixelShader(View.ShaderMap, PsPermutationVector);
 		
@@ -218,7 +236,14 @@ void FMobileSceneRenderer::RenderFog(FRHICommandList& RHICmdList, const FViewInf
 	VSParameters.StartDepthZ = StartDepthZ;
 	SetShaderParameters(RHICmdList, VertexShader, VertexShader.GetVertexShader(), VSParameters);
 
+	FMobileFogPS::FParameters PSParameters;
+	PSParameters.View = View.GetShaderParameters();
+	PSParameters.LFV = View.LocalFogVolumeViewData.UniformParametersStruct;
+	SetShaderParameters(RHICmdList, PixelShader, PixelShader.GetVertexShader(), PSParameters);
+
 	// Draw a quad covering the view.
 	RHICmdList.SetStreamSource(0, GScreenSpaceVertexBuffer.VertexBufferRHI, 0);
 	RHICmdList.DrawIndexedPrimitive(GTwoTrianglesIndexBuffer.IndexBufferRHI, 0, 0, 4, 0, 2, 1);
+
+	RenderLocalFogVolumeMobileLambda();
 }
diff --git a/Engine/Source/Runtime/Renderer/Private/MobileShadingRenderer.cpp b/Engine/Source/Runtime/Renderer/Private/MobileShadingRenderer.cpp
--- a/Engine/Source/Runtime/Renderer/Private/MobileShadingRenderer.cpp
+++ b/Engine/Source/Runtime/Renderer/Private/MobileShadingRenderer.cpp
@@ -998,6 +998,9 @@ void FMobileSceneRenderer::Render(FRDGBuilder& GraphBuilder)
 		GraphBuilder.SetCommandListStat(GET_STATID(STAT_CLMM_Shadows));
 		RenderShadowDepthMaps(GraphBuilder, InstanceCullingManager, ExternalAccessQueue);
 
+		// Run local fog volume initialization before base pass and volumetric fog for all the culled instance instance data to be ready.
+		InitLocalFogVolumesForViews(Scene, Views, ViewFamily, GraphBuilder, ShouldRenderVolumetricFog());
+
 		if (ShouldRenderVolumetricFog())
 		{
 			ComputeVolumetricFog(GraphBuilder, SceneTextures);
@@ -1270,9 +1273,6 @@ void FMobileSceneRenderer::RenderForward(FRDGBuilder& GraphBuilder, FRDGTextureR
 
 	const FRDGSystemTextures& SystemTextures = FRDGSystemTextures::Get(GraphBuilder);
 
-	// Run local fog volume initialization before base pass for when data is needed in forward
-	InitLocalFogVolumesForViews(Scene, Views, ViewFamily, GraphBuilder);
-
 	FRenderViewContextArray RenderViews;
 	GetRenderViews(Views, RenderViews);
 
@@ -1366,7 +1366,6 @@ void FMobileSceneRenderer::RenderForwardSinglePass(FRDGBuilder& GraphBuilder, FM
 		{
 			RenderFog(RHICmdList, View);
 		}
-		RenderLocalFogVolumeMobile(RHICmdList, View);
 		// Draw translucency.
 		RenderTranslucency(RHICmdList, View);
 
@@ -1466,7 +1465,6 @@ void FMobileSceneRenderer::RenderForwardMultiPass(FRDGBuilder& GraphBuilder, FMo
 		RenderDecals(RHICmdList, View);
 		RenderModulatedShadowProjections(RHICmdList, ViewContext.ViewIndex, View);
 		RenderFog(RHICmdList, View);
-		RenderLocalFogVolumeMobile(RHICmdList, View);
 		// Draw translucency.
 		RenderTranslucency(RHICmdList, View);
 
@@ -1609,9 +1607,6 @@ void FMobileSceneRenderer::RenderDeferred(FRDGBuilder& GraphBuilder, const FSort
 
 	const FRDGSystemTextures& SystemTextures = FRDGSystemTextures::Get(GraphBuilder);
 
-	// Run local fog volume initialization before base pass for when data is needed in forward
-	InitLocalFogVolumesForViews(Scene, Views, ViewFamily, GraphBuilder);
-
 	FRenderViewContextArray RenderViews;
 	GetRenderViews(Views, RenderViews);
 
@@ -1702,7 +1697,6 @@ void FMobileSceneRenderer::RenderDeferredSinglePass(FRDGBuilder& GraphBuilder, c
 			MobileDeferredCopyBuffer<FMobileDeferredCopyPLSPS>(RHICmdList, View);
 		}
 		RenderFog(RHICmdList, View);
-		RenderLocalFogVolumeMobile(RHICmdList, View);
 		// Draw translucency.
 		RenderTranslucency(RHICmdList, View);
 
@@ -1784,8 +1778,8 @@ void FMobileSceneRenderer::RenderDeferredMultiPass(FRDGBuilder& GraphBuilder, cl
 			
 		RHICmdList.SetCurrentStat(GET_STATID(STAT_CLMM_Translucency));
 		MobileDeferredShadingPass(RHICmdList, ViewContext.ViewIndex, Views.Num(), View, *Scene, SortedLightSet, VisibleLightInfos);
-		RenderLocalFogVolumeMobile(RHICmdList, View);
 		RenderFog(RHICmdList, View);
+
 		// Draw translucency.
 		RenderTranslucency(RHICmdList, View);
 
diff --git a/Engine/Source/Runtime/Renderer/Private/VolumetricFog.cpp b/Engine/Source/Runtime/Renderer/Private/VolumetricFog.cpp
--- a/Engine/Source/Runtime/Renderer/Private/VolumetricFog.cpp
+++ b/Engine/Source/Runtime/Renderer/Private/VolumetricFog.cpp
@@ -221,6 +221,7 @@ static const uint32 VolumetricFogGridInjectionGroupSize = 4;
 namespace
 {
 class FPermutationUseEmissive : SHADER_PERMUTATION_BOOL("USE_EMISSIVE");
+class FPermutationLocalFogVolume : SHADER_PERMUTATION_BOOL("USE_LOCAL_FOG_VOLUMES");
 }
 
 class FVolumetricFogMaterialSetupCS : public FGlobalShader
@@ -228,7 +229,7 @@ class FVolumetricFogMaterialSetupCS : public FGlobalShader
 	DECLARE_GLOBAL_SHADER(FVolumetricFogMaterialSetupCS);
 	SHADER_USE_PARAMETER_STRUCT(FVolumetricFogMaterialSetupCS, FGlobalShader);
 
-	using FPermutationDomain = TShaderPermutationDomain<FPermutationUseEmissive>;
+	using FPermutationDomain = TShaderPermutationDomain<FPermutationUseEmissive, FPermutationLocalFogVolume>;
 
 	BEGIN_SHADER_PARAMETER_STRUCT(FParameters, )
 		SHADER_PARAMETER(FLinearColor, GlobalAlbedo)
@@ -240,6 +241,8 @@ class FVolumetricFogMaterialSetupCS : public FGlobalShader
 
 		SHADER_PARAMETER_STRUCT_INCLUDE(FVolumetricFogIntegrationParameters, VolumetricFogParameters)
 
+		SHADER_PARAMETER_STRUCT(FLocalFogVolumeUniformParameters, LFV)
+
 		SHADER_PARAMETER_RDG_TEXTURE_UAV(RWTexture3D<float4>, RWVBufferA)
 		SHADER_PARAMETER_RDG_TEXTURE_UAV(RWTexture3D<float4>, RWVBufferB)
 	END_SHADER_PARAMETER_STRUCT()
@@ -1454,12 +1457,15 @@ void FSceneRenderer::ComputeVolumetricFog(FRDGBuilder& GraphBuilder,
 			PassParameters->RWVBufferA = IntegrationData.VBufferA_UAV;
 			PassParameters->RWVBufferB = IntegrationData.VBufferB_UAV; // FVolumetricFogMaterialSetupCS uses a permutation to not reference that UAV when bUseEmissive is false.
 
+			PassParameters->LFV = View.LocalFogVolumeViewData.UniformParametersStruct;
+
 			PassParameters->Fog = FogUniformBuffer; 
 			PassParameters->View = View.ViewUniformBuffer;
 			SetupVolumetricFogIntegrationParameters(PassParameters->VolumetricFogParameters, View, IntegrationData);
 
 			FVolumetricFogMaterialSetupCS::FPermutationDomain PermutationVector;
 			PermutationVector.Set< FPermutationUseEmissive >(bUseEmissive);
+			PermutationVector.Set< FPermutationLocalFogVolume >(ShouldRenderLocalFogVolumeInVolumetricFog(Scene, ViewFamily, ShouldRenderLocalFogVolume(Scene, ViewFamily)));
 			auto ComputeShader = View.ShaderMap->GetShader< FVolumetricFogMaterialSetupCS >(PermutationVector);
 			ClearUnusedGraphResources(ComputeShader, PassParameters);
 
