LFV - renamed some buffers for the mobile renderer to make it clear where they comes from.

#prelfight none

diff --git a/Engine/Source/Runtime/Renderer/Private/MobileShadingRenderer.cpp b/Engine/Source/Runtime/Renderer/Private/MobileShadingRenderer.cpp
--- a/Engine/Source/Runtime/Renderer/Private/MobileShadingRenderer.cpp
+++ b/Engine/Source/Runtime/Renderer/Private/MobileShadingRenderer.cpp
@@ -193,9 +193,9 @@ BEGIN_SHADER_PARAMETER_STRUCT(FMobileRenderPassParameters, RENDERER_API)
 	SHADER_PARAMETER_RDG_UNIFORM_BUFFER(FMobileBasePassUniformParameters, MobileBasePass)
 	SHADER_PARAMETER_STRUCT_REF(FMobileReflectionCaptureShaderData, ReflectionCapture)
 	SHADER_PARAMETER_RDG_BUFFER_SRV(Buffer<float4>, LocalFogVolumeInstances)
-	RDG_BUFFER_ACCESS(TileDrawIndirectBuffer, ERHIAccess::IndirectArgs)
+	RDG_BUFFER_ACCESS(LocalFogVolumeTileDrawIndirectBuffer, ERHIAccess::IndirectArgs)
 	SHADER_PARAMETER_RDG_TEXTURE_SRV(Texture2DArray<uint>, LocalFogVolumeTileDataTexture)
-	SHADER_PARAMETER_RDG_BUFFER_SRV(Buffer<uint>, TileDataBuffer)
+	SHADER_PARAMETER_RDG_BUFFER_SRV(Buffer<uint>, LocalFogVolumeTileDataBuffer)
 	RENDER_TARGET_BINDING_SLOTS()
 END_SHADER_PARAMETER_STRUCT()
 
@@ -1309,9 +1309,9 @@ void FMobileSceneRenderer::RenderForward(FRDGBuilder& GraphBuilder, FRDGTextureR
 		PassParameters->ReflectionCapture = View.MobileReflectionCaptureUniformBuffer;
 		PassParameters->RenderTargets = BasePassRenderTargets;
 		PassParameters->LocalFogVolumeInstances = View.LocalFogVolumeViewData.GPUInstanceDataBufferSRV;
-		PassParameters->TileDrawIndirectBuffer = View.LocalFogVolumeViewData.GPUTileDrawIndirectBuffer;
+		PassParameters->LocalFogVolumeTileDrawIndirectBuffer = View.LocalFogVolumeViewData.GPUTileDrawIndirectBuffer;
 		PassParameters->LocalFogVolumeTileDataTexture = View.LocalFogVolumeViewData.TileDataTextureArraySRV;
-		PassParameters->TileDataBuffer = View.LocalFogVolumeViewData.GPUTileDataBufferSRV;
+		PassParameters->LocalFogVolumeTileDataBuffer = View.LocalFogVolumeViewData.GPUTileDataBufferSRV;
 	
 		BuildInstanceCullingDrawParams(GraphBuilder, View, PassParameters);
 
@@ -1643,9 +1643,9 @@ void FMobileSceneRenderer::RenderDeferred(FRDGBuilder& GraphBuilder, const FSort
 		PassParameters->ReflectionCapture = View.MobileReflectionCaptureUniformBuffer;
 		PassParameters->RenderTargets = BasePassRenderTargets;
 		PassParameters->LocalFogVolumeInstances = View.LocalFogVolumeViewData.GPUInstanceDataBufferSRV;
-		PassParameters->TileDrawIndirectBuffer = View.LocalFogVolumeViewData.GPUTileDrawIndirectBuffer;
+		PassParameters->LocalFogVolumeTileDrawIndirectBuffer = View.LocalFogVolumeViewData.GPUTileDrawIndirectBuffer;
 		PassParameters->LocalFogVolumeTileDataTexture = View.LocalFogVolumeViewData.TileDataTextureArraySRV;
-		PassParameters->TileDataBuffer = View.LocalFogVolumeViewData.GPUTileDataBufferSRV;
+		PassParameters->LocalFogVolumeTileDataBuffer = View.LocalFogVolumeViewData.GPUTileDataBufferSRV;
 		
 		BuildInstanceCullingDrawParams(GraphBuilder, View, PassParameters);
 
