Fix incorrect Substrate settings causing rough distortion to be incorrect/uninitialized.

The distort accumuation pass didn't bound the substrate settings. However even after correcting this, the parameters were not correct. This CL force roughness tracking to be enabled for distortion so that rough refraction looks correct.

diff --git a/Engine/Shaders/Private/DistortAccumulatePS.usf b/Engine/Shaders/Private/DistortAccumulatePS.usf
--- a/Engine/Shaders/Private/DistortAccumulatePS.usf
+++ b/Engine/Shaders/Private/DistortAccumulatePS.usf
@@ -130,7 +130,7 @@ void Main(
 		const float3 V = MaterialParameters.CameraVector;
 
 		// Update tree (coverage/transmittance/luminace weights)
-		const FSubstrateIntegrationSettings Settings = InitSubstrateIntegrationSettings(false /*bForceFullyRough*/, Substrate.bRoughDiffuse /*bRoughDiffuseEnabled*/, Substrate.PeelLayersAboveDepth, Substrate.bRoughnessTracking);
+		const FSubstrateIntegrationSettings Settings = InitSubstrateIntegrationSettings(false /*bForceFullyRough*/, false /*bRoughDiffuseEnabled*/, 0, true/*bRoughnessTracking*/);
 		SubstratePixelHeader.SubstrateUpdateTree(V, Settings);
 
 		RefractionWorldNormal = SubstratePixelHeader.SubstrateTree.Operators[SubstrateData.OperatorIndex].TopDownRefractionWorldNormal;
diff --git a/Engine/Source/Runtime/Renderer/Private/DistortionRendering.cpp b/Engine/Source/Runtime/Renderer/Private/DistortionRendering.cpp
--- a/Engine/Source/Runtime/Renderer/Private/DistortionRendering.cpp
+++ b/Engine/Source/Runtime/Renderer/Private/DistortionRendering.cpp
@@ -257,6 +257,7 @@ BEGIN_SHADER_PARAMETER_STRUCT(FDistortionPassParameters, RENDERER_API)
 	SHADER_PARAMETER_STRUCT_INCLUDE(FViewShaderParameters, View)
 	SHADER_PARAMETER_RDG_UNIFORM_BUFFER(FDistortionPassUniformParameters, Pass)
 	SHADER_PARAMETER_STRUCT_INCLUDE(FInstanceCullingDrawParams, InstanceCullingDrawParams)
+	SHADER_PARAMETER_RDG_UNIFORM_BUFFER(FSubstrateGlobalUniformParameters, Substrate)
 	RENDER_TARGET_BINDING_SLOTS()
 END_SHADER_PARAMETER_STRUCT()
 
@@ -742,6 +743,7 @@ void FDeferredShadingSceneRenderer::RenderDistortion(
 			auto* PassParameters = GraphBuilder.AllocParameters<FDistortionPassParameters>();
 			PassParameters->View = View.GetShaderParameters();
 			PassParameters->Pass = CreateDistortionPassUniformBuffer(GraphBuilder, View);
+			PassParameters->Substrate = Substrate::BindSubstrateGlobalUniformParameters(View);
 			PassParameters->RenderTargets[0] = FRenderTargetBinding(DistortionTexture, LoadAction);
 			if (bUseRoughRefraction)
 			{
