Path Tracer: Avoid cooking errors in decal shaders on platforms which do not support the path tracer

diff --git a/Engine/Source/Runtime/Renderer/Private/RayTracing/RayTracingDecals.cpp b/Engine/Source/Runtime/Renderer/Private/RayTracing/RayTracingDecals.cpp
--- a/Engine/Source/Runtime/Renderer/Private/RayTracing/RayTracingDecals.cpp
+++ b/Engine/Source/Runtime/Renderer/Private/RayTracing/RayTracingDecals.cpp
@@ -156,6 +156,11 @@ public:
 			// is raytracing enabled at all?
 			return false;
 		}
+		if (!GetSupportsPathTracing(Parameters.Platform))
+		{
+			// this shader is currently only used by the path tracer
+			return false;
+		}
 		if (Parameters.MaterialParameters.MaterialDomain != MD_DeferredDecal)
 		{
 			// only compile callable shader permutation of deferred decal material
@@ -248,6 +253,11 @@ public:
 			// is raytracing enabled at all?
 			return false;
 		}
+		if (!GetSupportsPathTracing(Parameters.Platform))
+		{
+			// this shader is currently only used by the path tracer
+			return false;
+		}
 		if (!Parameters.VertexFactoryType->SupportsRayTracing())
 		{
 			// does the VF support ray tracing at all?
