Path Tracing: Fix compilation error due to missing class name

diff --git a/Engine/Source/Runtime/Renderer/Private/RayTracing/RayTracingDecals.cpp b/Engine/Source/Runtime/Renderer/Private/RayTracing/RayTracingDecals.cpp
--- a/Engine/Source/Runtime/Renderer/Private/RayTracing/RayTracingDecals.cpp
+++ b/Engine/Source/Runtime/Renderer/Private/RayTracing/RayTracingDecals.cpp
@@ -156,7 +156,7 @@ public:
 			// is raytracing enabled at all?
 			return false;
 		}
-		if (!GetSupportsPathTracing(Parameters.Platform))
+		if (!FDataDrivenShaderPlatformInfo::GetSupportsPathTracing(Parameters.Platform))
 		{
 			// this shader is currently only used by the path tracer
 			return false;
@@ -253,7 +253,7 @@ public:
 			// is raytracing enabled at all?
 			return false;
 		}
-		if (!GetSupportsPathTracing(Parameters.Platform))
+		if (!FDataDrivenShaderPlatformInfo::GetSupportsPathTracing(Parameters.Platform))
 		{
 			// this shader is currently only used by the path tracer
 			return false;
