Allow to align GPUScene InstanceId buffer size

#rnx

diff --git a/Engine/Source/Runtime/Renderer/Private/InstanceCulling/InstanceCullingContext.cpp b/Engine/Source/Runtime/Renderer/Private/InstanceCulling/InstanceCullingContext.cpp
--- a/Engine/Source/Runtime/Renderer/Private/InstanceCulling/InstanceCullingContext.cpp
+++ b/Engine/Source/Runtime/Renderer/Private/InstanceCulling/InstanceCullingContext.cpp
@@ -623,12 +623,38 @@ public:
 #endif
 };
 
+static int32 GGPUSceneUniformBufferMinSizeKB = 0;
+static FAutoConsoleVariableRef CVarGPUSceneUniformBufferMinSizeKB(
+	TEXT("r.GPUSceneUniformBufferMinSizeKB"),
+	GGPUSceneUniformBufferMinSizeKB,
+	TEXT("Experimental: Minimum size used by GPU Scene Instance Buffer with Uniform Buffer View"),
+	ECVF_RenderThreadSafe);
+
+static int32 GGPUSceneUniformBufferAlignSizeKB = 0;
+static FAutoConsoleVariableRef CVarGPUSceneUniformBufferAlignSizeKB(
+	TEXT("r.GPUSceneUniformBufferAlignSizeKB"),
+	GGPUSceneUniformBufferAlignSizeKB,
+	TEXT("Experimental: Alignment on size used by GPU Scene Instance Buffer with Uniform Buffer View"),
+	ECVF_RenderThreadSafe);
+
+constexpr uint32 GInstanceIdBufferUBOStride = 16u;
 static uint32 GetInstanceIdBufferSize(EShaderPlatform ShaderPlatform, uint32 NumInstanceElements)
 {
 	if (PlatformGPUSceneUsesUniformBufferView(ShaderPlatform))
 	{
 		// Add an additional max range slack to a buffer size, so when binding last element we still have a full UBO range
-		NumInstanceElements += (PLATFORM_MAX_UNIFORM_BUFFER_RANGE / 16u);
+		NumInstanceElements += (PLATFORM_MAX_UNIFORM_BUFFER_RANGE / GInstanceIdBufferUBOStride);
+		
+		if (GGPUSceneUniformBufferAlignSizeKB > 0)
+		{
+			NumInstanceElements = AlignArbitrary<uint32>(NumInstanceElements, GGPUSceneUniformBufferAlignSizeKB * 1024 / GInstanceIdBufferUBOStride);
+		}
+
+		if (GGPUSceneUniformBufferMinSizeKB > 0)
+		{
+			NumInstanceElements = FMath::Max<uint32>(NumInstanceElements, GGPUSceneUniformBufferMinSizeKB * 1024 / GInstanceIdBufferUBOStride);
+		}
+
 		return NumInstanceElements;
 	}
 	else
@@ -643,7 +669,7 @@ static FRDGBufferDesc CreateInstanceIdBufferDesc(EShaderPlatform ShaderPlatform,
 	if (PlatformGPUSceneUsesUniformBufferView(ShaderPlatform))
 	{
 		// float4
-		FRDGBufferDesc Desc = FRDGBufferDesc::CreateStructuredDesc(16u, NumInstanceElements);
+		FRDGBufferDesc Desc = FRDGBufferDesc::CreateStructuredDesc(GInstanceIdBufferUBOStride, NumInstanceElements);
 		Desc.Usage |= EBufferUsageFlags::UniformBuffer;
 		return Desc;
 	}
