Virtual Texture: Fixed FVirtualTextureSystem::RequestTilesForRegionInternal() potentially issuing many invalid tile requests if fetching for a VT outside of the viewport
* Solves degenerate case where the WP minimap would be zoomed out and dragged far outside the viewport, issuing millions of tile requests each frame
* Properly clamp the requested VT region to what's visible in the viewport
* Deprecated previous method which required the viewport position to be negated to work properly

diff --git a/Engine/Source/Editor/TextureEditor/Private/Models/TextureEditorViewportClient.cpp b/Engine/Source/Editor/TextureEditor/Private/Models/TextureEditorViewportClient.cpp
--- a/Engine/Source/Editor/TextureEditor/Private/Models/TextureEditorViewportClient.cpp
+++ b/Engine/Source/Editor/TextureEditor/Private/Models/TextureEditorViewportClient.cpp
@@ -249,7 +249,7 @@ void FTextureEditorViewportClient::Draw(FViewport* Viewport, FCanvas* Canvas)
 		{
 			FVirtualTexture2DResource* VTResource = static_cast<FVirtualTexture2DResource*>(Texture->GetResource());
 			const FVector2D ScreenSpaceSize((float)Width, (float)Height);
-			const FVector2D ViewportPositon(-(float)XPos, -(float)YPos);
+			const FVector2D ViewportPositon((float)XPos, (float)YPos);
 			const FVector2D UV0 = TileItem.UV0;
 			const FVector2D UV1 = TileItem.UV1;
 
@@ -263,7 +263,7 @@ void FTextureEditorViewportClient::Draw(FViewport* Viewport, FCanvas* Canvas)
 				IAllocatedVirtualTexture* AllocatedVT = VTResource->AcquireAllocatedVT();
 
 				IRendererModule& RenderModule = GetRendererModule();
-				RenderModule.RequestVirtualTextureTilesForRegion(AllocatedVT, ScreenSpaceSize, ViewportPositon, ViewportSize, UV0, UV1, (int32)MipLevel);
+				RenderModule.RequestVirtualTextureTiles(AllocatedVT, ScreenSpaceSize, ViewportPositon, ViewportSize, UV0, UV1, (int32)MipLevel);
 				RenderModule.LoadPendingVirtualTextureTiles(RHICmdList, InFeatureLevel);
 			});
 		}
diff --git a/Engine/Source/Editor/WorldPartitionEditor/Private/WorldPartition/SWorldPartitionEditorGrid2D.cpp b/Engine/Source/Editor/WorldPartitionEditor/Private/WorldPartition/SWorldPartitionEditorGrid2D.cpp
--- a/Engine/Source/Editor/WorldPartitionEditor/Private/WorldPartition/SWorldPartitionEditorGrid2D.cpp
+++ b/Engine/Source/Editor/WorldPartitionEditor/Private/WorldPartition/SWorldPartitionEditorGrid2D.cpp
@@ -1977,7 +1977,7 @@ int32 SWorldPartitionEditorGrid2D::PaintMinimap(const FGeometry& AllottedGeometr
 			FVirtualTexture2DResource* VTResource = static_cast<FVirtualTexture2DResource*>(Texture2D->GetResource());
 			const FVector2D ViewportSize = AllottedGeometry.GetLocalSize();
 			const FVector2D ScreenSpaceSize = WorldImageGeometry.GetLocalSize();
-			const FVector2D ViewportPositon = -WorldImageGeometry.GetAccumulatedRenderTransform().GetTranslation() + AllottedGeometry.GetAbsolutePosition();
+			const FVector2D ViewportPositon = WorldImageGeometry.GetAccumulatedRenderTransform().GetTranslation() + AllottedGeometry.GetAbsolutePosition();
 
 			FBox2D UVRegion = WorldMiniMapBrush.GetUVRegion();
 			const FVector2D UV0 = UVRegion.Min;
diff --git a/Engine/Source/Runtime/RenderCore/Public/RendererInterface.h b/Engine/Source/Runtime/RenderCore/Public/RendererInterface.h
--- a/Engine/Source/Runtime/RenderCore/Public/RendererInterface.h
+++ b/Engine/Source/Runtime/RenderCore/Public/RendererInterface.h
@@ -805,9 +805,24 @@ public:
 
 	virtual void RequestVirtualTextureTiles(const FVector2D& InScreenSpaceSize, int32 InMipLevel) = 0;
 	virtual void RequestVirtualTextureTiles(const FMaterialRenderProxy* InMaterialRenderProxy, const FVector2D& InScreenSpaceSize, ERHIFeatureLevel::Type InFeatureLevel) = 0;
+
+	/**
+	 * Helper function to request loading of tiles for a virtual texture that will be displayed in the UI. 
+	 * It will request only the tiles that will be visible after clipping to the provided viewport.
+	 * @param AllocatedVT			The virtual texture.
+	 * @param InScreenSpaceSize		Size on screen at which the texture is to be displayed.
+	 * @param InViewportPosition	Position in the viewport where the texture will be displayed.
+	 * @param InViewportSize		Size of the viewport.
+	 * @param InUV0					UV coordinate to use for the top left corner of the texture.
+	 * @param InUV1					UV coordinate to use for the bottom right corner of the texture.
+	 * @param InMipLevel [optional] Specific mip level to fetch tiles for.
+	 */
+	virtual void RequestVirtualTextureTiles(IAllocatedVirtualTexture* AllocatedVT, const FVector2D& InScreenSpaceSize, const FVector2D& InViewportPosition, const FVector2D& InViewportSize, const FVector2D& InUV0, const FVector2D& InUV1, int32 InMipLevel) = 0;
+
+	UE_DEPRECATED(5.4, "Use RequestVirtualTextureTiles() overloads that takes similar parameters. Make sure not to negate the InViewportPosition.")
 	virtual void RequestVirtualTextureTilesForRegion(IAllocatedVirtualTexture* AllocatedVT, const FVector2D& InScreenSpaceSize, const FVector2D& InViewportPosition, const FVector2D& InViewportSize, const FVector2D& InUV0, const FVector2D& InUV1, int32 InMipLevel) = 0;
 
-	/** Ensure that any tiles requested by 'RequestVirtualTextureTilesForRegion' are loaded, must be called from render thread */
+	/** Ensure that any tiles requested by 'RequestVirtualTextureTiles' are loaded, must be called from render thread */
 	virtual void LoadPendingVirtualTextureTiles(FRHICommandListImmediate& RHICmdList, ERHIFeatureLevel::Type FeatureLevel) = 0;
 
 	/** Allocate a buffer and record all virtual texture page requests until the next call to either SetVirtualTextureRequestRecordBuffer or GetVirtualTextureRequestRecordBuffer. */
diff --git a/Engine/Source/Runtime/Renderer/Private/RendererModule.h b/Engine/Source/Runtime/Renderer/Private/RendererModule.h
--- a/Engine/Source/Runtime/Renderer/Private/RendererModule.h
+++ b/Engine/Source/Runtime/Renderer/Private/RendererModule.h
@@ -106,7 +106,11 @@ public:
 	virtual void RequestVirtualTextureTiles(TArray<uint64>&& InPageRequests) override;
 	virtual void RequestVirtualTextureTiles(const FVector2D& InScreenSpaceSize, int32 InMipLevel) override;
 	virtual void RequestVirtualTextureTiles(const FMaterialRenderProxy* InMaterialRenderProxy, const FVector2D& InScreenSpaceSize, ERHIFeatureLevel::Type InFeatureLevel) override;
+	virtual void RequestVirtualTextureTiles(IAllocatedVirtualTexture* AllocatedVT, const FVector2D& InScreenSpaceSize, const FVector2D& InViewportPosition, const FVector2D& InViewportSize, const FVector2D& InUV0, const FVector2D& InUV1, int32 InMipLevel) override;
+	
+	UE_DEPRECATED(5.4, "Use RequestVirtualTextureTiles() overloads that takes similar parameters. Make sure not to negate the InViewportPosition.")
 	virtual void RequestVirtualTextureTilesForRegion(IAllocatedVirtualTexture* AllocatedVT, const FVector2D& InScreenSpaceSize, const FVector2D& InViewportPosition, const FVector2D& InViewportSize, const FVector2D& InUV0, const FVector2D& InUV1, int32 InMipLevel) override;
+	
 	virtual void LoadPendingVirtualTextureTiles(FRHICommandListImmediate& RHICmdList, ERHIFeatureLevel::Type FeatureLevel) override;
 	virtual void SetVirtualTextureRequestRecordBuffer(uint64 Handle) override;
 	virtual uint64 GetVirtualTextureRequestRecordBuffer(TSet<uint64>& OutPageRequests) override;
diff --git a/Engine/Source/Runtime/Renderer/Private/SceneRendering.cpp b/Engine/Source/Runtime/Renderer/Private/SceneRendering.cpp
--- a/Engine/Source/Runtime/Renderer/Private/SceneRendering.cpp
+++ b/Engine/Source/Runtime/Renderer/Private/SceneRendering.cpp
@@ -5159,9 +5159,16 @@ void FRendererModule::RequestVirtualTextureTiles(const FMaterialRenderProxy* InM
 	FVirtualTextureSystem::Get().RequestTiles(InMaterialRenderProxy, InScreenSpaceSize, InFeatureLevel);
 }
 
+void FRendererModule::RequestVirtualTextureTiles(IAllocatedVirtualTexture* AllocatedVT, const FVector2D& InScreenSpaceSize, const FVector2D& InViewportPosition, const FVector2D& InViewportSize, const FVector2D& InUV0, const FVector2D& InUV1, int32 InMipLevel)
+{
+	FVirtualTextureSystem::Get().RequestTiles(AllocatedVT, InScreenSpaceSize, InViewportPosition, InViewportSize, InUV0, InUV1, InMipLevel);
+}
+
 void FRendererModule::RequestVirtualTextureTilesForRegion(IAllocatedVirtualTexture* AllocatedVT, const FVector2D& InScreenSpaceSize, const FVector2D& InViewportPosition, const FVector2D& InViewportSize, const FVector2D& InUV0, const FVector2D& InUV1, int32 InMipLevel)
 {
+PRAGMA_DISABLE_DEPRECATION_WARNINGS
 	FVirtualTextureSystem::Get().RequestTilesForRegion(AllocatedVT, InScreenSpaceSize, InViewportPosition, InViewportSize, InUV0, InUV1, InMipLevel);
+PRAGMA_ENABLE_DEPRECATION_WARNINGS
 }
 
 void FRendererModule::LoadPendingVirtualTextureTiles(FRHICommandListImmediate& RHICmdList, ERHIFeatureLevel::Type FeatureLevel)
diff --git a/Engine/Source/Runtime/Renderer/Private/VT/VirtualTextureSystem.cpp b/Engine/Source/Runtime/Renderer/Private/VT/VirtualTextureSystem.cpp
--- a/Engine/Source/Runtime/Renderer/Private/VT/VirtualTextureSystem.cpp
+++ b/Engine/Source/Runtime/Renderer/Private/VT/VirtualTextureSystem.cpp
@@ -1197,7 +1197,7 @@ void FVirtualTextureSystem::RequestTilesInternal(const IAllocatedVirtualTexture*
 	}
 }
 
-void FVirtualTextureSystem::RequestTilesForRegion(IAllocatedVirtualTexture* AllocatedVT, const FVector2D& InScreenSpaceSize, const FVector2D& InViewportPosition, const FVector2D& InViewportSize, const FVector2D& InUV0, const FVector2D& InUV1, int32 InMipLevel)
+void FVirtualTextureSystem::RequestTiles(IAllocatedVirtualTexture* AllocatedVT, const FVector2D& InScreenSpaceSize, const FVector2D& InViewportPosition, const FVector2D& InViewportSize, const FVector2D& InUV0, const FVector2D& InUV1, int32 InMipLevel)
 {
 	UE::TScopeLock Lock(Mutex);
 	if (InMipLevel >= 0)
@@ -1207,7 +1207,7 @@ void FVirtualTextureSystem::RequestTilesForRegion(IAllocatedVirtualTexture* Allo
 	else
 	{
 		const uint32 vMaxLevel = AllocatedVT->GetMaxLevel();
-		const float vLevel = ComputeMipLevel(AllocatedVT, InScreenSpaceSize);
+		const float vLevel = ComputeMipLevel(AllocatedVT, InScreenSpaceSize); // TODO: ComputeMipLevel() is incorrect if not using the whole UV range
 		const int32 vMipLevelDown = FMath::Clamp((int32)FMath::FloorToInt(vLevel), 0, (int32)vMaxLevel);
 
 		RequestTilesForRegionInternal(AllocatedVT, InScreenSpaceSize, InViewportPosition, InViewportSize, InUV0, InUV1, vMipLevelDown);
@@ -1219,6 +1219,12 @@ void FVirtualTextureSystem::RequestTilesForRegion(IAllocatedVirtualTexture* Allo
 	}
 }
 
+void FVirtualTextureSystem::RequestTilesForRegion(IAllocatedVirtualTexture* AllocatedVT, const FVector2D& InScreenSpaceSize, const FVector2D& InViewportPosition, const FVector2D& InViewportSize, const FVector2D& InUV0, const FVector2D& InUV1, int32 InMipLevel)
+{
+	// RequestTilesForRegion() used to require that the viewport position was negated which felt wrong. Correct this when calling the new implementation.
+	RequestTiles(AllocatedVT, InScreenSpaceSize, -InViewportPosition, InViewportSize, InUV0, InUV1, InMipLevel);
+}
+
 void FVirtualTextureSystem::LoadPendingTiles(FRDGBuilder& GraphBuilder, ERHIFeatureLevel::Type FeatureLevel)
 {
 	check(!bUpdating);
@@ -1280,23 +1286,30 @@ static int32 WrapTilePosition(int32 Position, int32 Size)
 
 void FVirtualTextureSystem::RequestTilesForRegionInternal(const IAllocatedVirtualTexture* AllocatedVT, const FVector2D& InScreenSpaceSize, const FVector2D& InViewportPosition, const FVector2D& InViewportSize, const FVector2D& InUV0, const FVector2D& InUV1, int32 InMipLevel)
 {
-	const float ScreenSpaceSizeX = FMath::Max(InScreenSpaceSize.X, 1.0f);
-	const float ScreenSpaceSizeY = FMath::Max(InScreenSpaceSize.Y, 1.0f);
+	// Screen size must be a least a pixel
+	FVector2D ScreenSize = FVector2D::Max(InScreenSpaceSize, FVector2D::One());
+
+	// TopLeft vs BottomRight - In viewport space
+	FVector2D TextureTopLeftViewportSpace = InViewportPosition;
+	FVector2D TextureBottomRightViewportSpace = InViewportPosition + ScreenSize;
 
-	// Determine the screen-space coordinates of the texture we're trying to display
-	const float SrcPositionX0 = FMath::Max(InViewportPosition.X, 0.0f);
-	const float SrcPositionY0 = FMath::Max(InViewportPosition.Y, 0.0f);
-	const float SrcPositionX1 = FMath::Max(FMath::Min(InViewportPosition.X + InViewportSize.X, ScreenSpaceSizeX), 0.0f);
-	const float SrcPositionY1 = FMath::Max(FMath::Min(InViewportPosition.Y + InViewportSize.Y, ScreenSpaceSizeY), 0.0f);
+	// TopLeft vs BottomRight - Clamped to viewport
+	FVector2D TextureTopLeftViewportSpaceClamped = FVector2D::Clamp(TextureTopLeftViewportSpace, FVector2D::Zero(), InViewportSize);
+	FVector2D TextureBottomRightViewportSpaceClamped = FVector2D::Clamp(TextureBottomRightViewportSpace, FVector2D::Zero(), InViewportSize);
 
+	// Range of initial screen size for TopLeft & BottomRight
+	// For example, if 10% of the image is outside each side of the viewport, LerpTopLeft & LerpBottomRight would be (0.1, 0.1) & (0.9, 0.9), respectively
+	FVector2D LerpTopLeft = FVector2D::Clamp((TextureTopLeftViewportSpaceClamped - TextureTopLeftViewportSpace) / ScreenSize, FVector2D::Zero(), FVector2D::One());
+	FVector2D LerpBottomRight = FVector2D::Clamp((TextureBottomRightViewportSpaceClamped - TextureTopLeftViewportSpace) / ScreenSize, FVector2D::Zero(), FVector2D::One());
+	
 	const int32 WidthInBlocks = AllocatedVT->GetWidthInBlocks();
 	const int32 HeightInBlocks = AllocatedVT->GetHeightInBlocks();
 
 	// Map coordinates to UV space
-	const float PositionU0 = FMath::Lerp(InUV0.X, InUV1.X, SrcPositionX0 / ScreenSpaceSizeX) / WidthInBlocks;
-	const float PositionV0 = FMath::Lerp(InUV0.Y, InUV1.Y, SrcPositionY0 / ScreenSpaceSizeY) / HeightInBlocks;
-	const float PositionU1 = FMath::Lerp(InUV0.X, InUV1.X, SrcPositionX1 / ScreenSpaceSizeX) / WidthInBlocks;
-	const float PositionV1 = FMath::Lerp(InUV0.Y, InUV1.Y, SrcPositionY1 / ScreenSpaceSizeY) / HeightInBlocks;
+	const float PositionU0 = FMath::Lerp(InUV0.X, InUV1.X, LerpTopLeft.X) / WidthInBlocks;
+	const float PositionV0 = FMath::Lerp(InUV0.Y, InUV1.Y, LerpTopLeft.Y) / HeightInBlocks;
+	const float PositionU1 = FMath::Lerp(InUV0.X, InUV1.X, LerpBottomRight.X) / WidthInBlocks;
+	const float PositionV1 = FMath::Lerp(InUV0.Y, InUV1.Y, LerpBottomRight.Y) / HeightInBlocks;
 
 	// Map UVs to tile coordinates
 	const int32 MipWidthInTiles = FMath::Max<int32>(AllocatedVT->GetWidthInTiles() >> InMipLevel, 1);
diff --git a/Engine/Source/Runtime/Renderer/Private/VT/VirtualTextureSystem.h b/Engine/Source/Runtime/Renderer/Private/VT/VirtualTextureSystem.h
--- a/Engine/Source/Runtime/Renderer/Private/VT/VirtualTextureSystem.h
+++ b/Engine/Source/Runtime/Renderer/Private/VT/VirtualTextureSystem.h
@@ -128,7 +128,23 @@ public:
 
 	void RequestTiles(const FVector2D& InScreenSpaceSize, int32 InMipLevel = -1);
 	void RequestTiles(const FMaterialRenderProxy* InMaterialRenderProxy, const FVector2D& InScreenSpaceSize, ERHIFeatureLevel::Type InFeatureLevel);
+
+	/**
+	 * Helper function to request loading of tiles for a virtual texture that will be displayed in the UI. 
+	 * It will request only the tiles that will be visible after clipping to the provided viewport.
+	 * @param AllocatedVT			The virtual texture.
+	 * @param InScreenSpaceSize		Size on screen at which the texture is to be displayed.
+	 * @param InViewportPosition	Position in the viewport where the texture will be displayed.
+	 * @param InViewportSize		Size of the viewport.
+	 * @param InUV0					UV coordinate to use for the top left corner of the texture.
+	 * @param InUV1					UV coordinate to use for the bottom right corner of the texture.
+	 * @param InMipLevel [optional] Specific mip level to fetch tiles for.
+	 */
+	void RequestTiles(IAllocatedVirtualTexture* AllocatedVT, const FVector2D& InScreenSpaceSize, const FVector2D& InViewportPosition, const FVector2D& InViewportSize, const FVector2D& InUV0, const FVector2D& InUV1, int32 InMipLevel = -1);
+
+	UE_DEPRECATED(5.4, "Use RequestTiles() overloads that takes similar parameters. Make sure not to negate the InViewportPosition.")
 	void RequestTilesForRegion(IAllocatedVirtualTexture* AllocatedVT, const FVector2D& InScreenSpaceSize, const FVector2D& InViewportPosition, const FVector2D& InViewportSize, const FVector2D& InUV0, const FVector2D& InUV1, int32 InMipLevel = -1);
+
 	void LoadPendingTiles(FRDGBuilder& GraphBuilder, ERHIFeatureLevel::Type FeatureLevel);
 	
 #if WITH_EDITOR
