LFV - fix for wrong culling now accounting for the tile to screen extent ratio when lerping planes. (tile can span outside the view pixel)

#prelfight none

diff --git a/Engine/Shaders/Private/LocalFogVolumes/LocalFogVolumeTiledCulling.usf b/Engine/Shaders/Private/LocalFogVolumes/LocalFogVolumeTiledCulling.usf
--- a/Engine/Shaders/Private/LocalFogVolumes/LocalFogVolumeTiledCulling.usf
+++ b/Engine/Shaders/Private/LocalFogVolumes/LocalFogVolumeTiledCulling.usf
@@ -9,6 +9,7 @@ float4 RightPlane;
 float4 TopPlane;
 float4 BottomPlane;
 float4 NearPlane;
+float2 ViewToTileSpaceRatio;
 
 Buffer<float4> LocalFogVolumeCullingDataBuffer;
 
@@ -41,6 +42,12 @@ void LocalFogVolumeTiledCullingCS(uint3 DispatchThreadId : SV_DispatchThreadID)
 	float2 MinTileLerp =  float2(TileCoord)       / float2(LFVTileDataResolution);
 	float2 MaxTileLerp = (float2(TileCoord)+1.0f) / float2(LFVTileDataResolution);
 
+	// Apply the view to tile space ratio to account ofr the fact that tiles can span outside the view pixel due to their discretised nature.
+	MinTileLerp *= ViewToTileSpaceRatio;
+	MaxTileLerp *= ViewToTileSpaceRatio;
+	// Still clamp MaxTileLerp to not include not visible volumes to the right and up part of the screen.
+	MaxTileLerp = min(MaxTileLerp, 1.0f.xx);
+
 	float4 TilePlanes[5];
 	TilePlanes[0] = lerp(LeftPlane, RightPlane, MinTileLerp.xxxx);		// Left tile plane
 	TilePlanes[1] = lerp(LeftPlane, RightPlane, MaxTileLerp.xxxx);		// Right tile plane
diff --git a/Engine/Source/Runtime/Renderer/Private/LocalFogVolumeRendering.cpp b/Engine/Source/Runtime/Renderer/Private/LocalFogVolumeRendering.cpp
--- a/Engine/Source/Runtime/Renderer/Private/LocalFogVolumeRendering.cpp
+++ b/Engine/Source/Runtime/Renderer/Private/LocalFogVolumeRendering.cpp
@@ -189,6 +189,7 @@ public:
 		SHADER_PARAMETER(FVector4f, TopPlane)
 		SHADER_PARAMETER(FVector4f, BottomPlane)
 		SHADER_PARAMETER(FVector4f, NearPlane)
+		SHADER_PARAMETER(FVector2f, ViewToTileSpaceRatio)
 	END_SHADER_PARAMETER_STRUCT()
 
 	static void ModifyCompilationEnvironment(const FGlobalShaderPermutationParameters& Parameters, FShaderCompilerEnvironment& OutEnvironment)
@@ -244,9 +245,12 @@ static void LocalFogVolumeViewTiledCullingPass(FViewInfo& View, FRDGBuilder& Gra
 		PassParameters->TopPlane	= FVector4f::Zero();
 		PassParameters->BottomPlane	= FVector4f::Zero();
 	}
-	
 	ConvertPlanToVector4f(PassParameters->NearPlane,	View.NearClippingPlane,			false);
 
+	float TileCoveredResolutionX = View.LocalFogVolumeViewData.UniformParametersStruct.LocalFogVolumeCommon.LocalFogVolumeTilePixelSize * View.LocalFogVolumeViewData.UniformParametersStruct.LocalFogVolumeCommon.LocalFogVolumeTileDataTextureResolution.X;
+	float TileCoveredResolutionY = View.LocalFogVolumeViewData.UniformParametersStruct.LocalFogVolumeCommon.LocalFogVolumeTilePixelSize * View.LocalFogVolumeViewData.UniformParametersStruct.LocalFogVolumeCommon.LocalFogVolumeTileDataTextureResolution.Y;
+	PassParameters->ViewToTileSpaceRatio = FVector2f(TileCoveredResolutionX * View.CachedViewUniformShaderParameters->ViewSizeAndInvSize.Z, TileCoveredResolutionY * View.CachedViewUniformShaderParameters->ViewSizeAndInvSize.W);
+
 	ERDGPassFlags PassFlag = ERDGPassFlags::Compute; // LFV_TODO try ERDGPassFlags::AsyncCompute later
 
 	TileDataTextureSize.Z = 1;
