Fix strata -> substrate comment.

diff --git a/Engine/Source/Runtime/Renderer/Private/PathTracing.cpp b/Engine/Source/Runtime/Renderer/Private/PathTracing.cpp
--- a/Engine/Source/Runtime/Renderer/Private/PathTracing.cpp
+++ b/Engine/Source/Runtime/Renderer/Private/PathTracing.cpp
@@ -921,7 +921,7 @@ class FPathTracingRG : public FGlobalShader
 		}
 		if (!Substrate::IsSubstrateEnabled())
 		{
-			// If we aren't using strata, no need to compile the complex material path
+			// If we aren't using Substrate, no need to compile the complex material path
 			if (PermutationVector.Get<FSubstrateComplexSpecialMaterial>())
 			{
 				return false;
