Only include draw commands with HasPrimitiveIdStreamIndex flag in occlusion query instance culling list

diff --git a/Engine/Source/Runtime/Renderer/Private/InstanceCulling/InstanceCullingOcclusionQuery.cpp b/Engine/Source/Runtime/Renderer/Private/InstanceCulling/InstanceCullingOcclusionQuery.cpp
--- a/Engine/Source/Runtime/Renderer/Private/InstanceCulling/InstanceCullingOcclusionQuery.cpp
+++ b/Engine/Source/Runtime/Renderer/Private/InstanceCulling/InstanceCullingOcclusionQuery.cpp
@@ -314,14 +314,18 @@ struct FInstanceCullingOcclusionQueryDeferredContext
 		// There may be multiple visible mesh draw commands that refer to the same instance when GPU-based LOD selection is used.
 		// This filter is designed to remove the duplicates, keeping only the "authoritative" instance.
 		// TODO: a less implicit mechanism would be welcome here, such as a dedicated flag.
-		EMeshDrawCommandCullingPayloadFlags Flags = VisibleCommand.CullingPayloadFlags;
-		bool bCompatibleFlags = Flags == EMeshDrawCommandCullingPayloadFlags::Default
-							 || Flags == EMeshDrawCommandCullingPayloadFlags::MinScreenSizeCull;
+		const EMeshDrawCommandCullingPayloadFlags Flags = VisibleCommand.CullingPayloadFlags;
+		const bool bCompatibleFlags = Flags == EMeshDrawCommandCullingPayloadFlags::Default
+			|| Flags == EMeshDrawCommandCullingPayloadFlags::MinScreenSizeCull;
+
+		// Only commands with HasPrimitiveIdStreamIndex are compatible with GPU Instance Culling
+		const bool bSupportsGPUSceneInstancing = EnumHasAnyFlags(VisibleCommand.Flags, EFVisibleMeshDrawCommandFlags::HasPrimitiveIdStreamIndex);
 
 		// NumPrimitives is 0 if mesh draw command uses IndirectArgs
 		// This path is currently not implemented/supported by oclcusion query culling.
 		// Commands that use instance runs are currently not supported.
 		return bCompatibleFlags
+			&& bSupportsGPUSceneInstancing
 			&& VisibleCommand.MeshDrawCommand->NumPrimitives != 0
 			&& VisibleCommand.NumRuns == 0;
 	};
