Path Tracing: Revert change that skipped RenderOpaqueFX and use a placeholder for the SceneColor texture as it is not valid in the path tracing context

diff --git a/Engine/Source/Runtime/Renderer/Private/DeferredShadingRenderer.cpp b/Engine/Source/Runtime/Renderer/Private/DeferredShadingRenderer.cpp
--- a/Engine/Source/Runtime/Renderer/Private/DeferredShadingRenderer.cpp
+++ b/Engine/Source/Runtime/Renderer/Private/DeferredShadingRenderer.cpp
@@ -638,6 +638,7 @@ static void RenderOpaqueFX(
 	TConstStridedView<FSceneView> Views,
 	FSceneUniformBuffer &SceneUniformBuffer,
 	FFXSystemInterface* FXSystem,
+	ERHIFeatureLevel::Type FeatureLevel,
 	TRDGUniformBufferRef<FSceneTextureUniformParameters> SceneTexturesUniformBuffer)
 {
 	// Notify the FX system that opaque primitives have been rendered and we now have a valid depth buffer.
@@ -648,6 +649,12 @@ static void RenderOpaqueFX(
 
 		const ERDGPassFlags UBPassFlags = ERDGPassFlags::Compute | ERDGPassFlags::Raster | ERDGPassFlags::SkipRenderPass | ERDGPassFlags::NeverCull;
 
+		if (HasRayTracedOverlay(*Views[0].Family))
+		{
+			// In the case of Path Tracing/RT Debug -- we have not yet written to the SceneColor buffer, so make a dummy set of textures instead
+			SceneTexturesUniformBuffer = CreateSceneTextureUniformBuffer(GraphBuilder, nullptr, FeatureLevel, ESceneTextureSetupMode::SceneVelocity);
+		}
+
 		// Add a pass which extracts the RHI handle from the scene textures UB and sends it to the FX system.
 		FRenderOpaqueFXPassParameters* ExtractUBPassParameters = GraphBuilder.AllocParameters<FRenderOpaqueFXPassParameters>();
 		ExtractUBPassParameters->SceneTextures = SceneTexturesUniformBuffer;
@@ -4139,10 +4146,7 @@ void FDeferredShadingSceneRenderer::Render(FRDGBuilder& GraphBuilder)
 
 		FRDGTextureRef ExposureIlluminance = AddCalculateExposureIlluminancePass(GraphBuilder, Views, SceneTextures, TranslucencyLightingVolumeTextures, ExposureIlluminanceSetup);
 
-		if (!bHasRayTracedOverlay)
-		{
-			RenderOpaqueFX(GraphBuilder, GetSceneViews(), GetSceneUniforms(), FXSystem, SceneTextures.UniformBuffer);
-		}
+		RenderOpaqueFX(GraphBuilder, GetSceneViews(), GetSceneUniforms(), FXSystem, FeatureLevel, SceneTextures.UniformBuffer);
 
 		FRendererModule& RendererModule = static_cast<FRendererModule&>(GetRendererModule());
 		RendererModule.RenderPostOpaqueExtensions(GraphBuilder, Views, SceneTextures);
