Disable force hiding of non-Nanite primitives for now to match prior behavior until issues the change inadvertently caused are worked around.

diff --git a/Engine/Source/Runtime/Renderer/Private/GPUScene.cpp b/Engine/Source/Runtime/Renderer/Private/GPUScene.cpp
--- a/Engine/Source/Runtime/Renderer/Private/GPUScene.cpp
+++ b/Engine/Source/Runtime/Renderer/Private/GPUScene.cpp
@@ -463,7 +463,10 @@ struct FUploadDataSourceAdapterScenePrimitives
 		InstanceUploadInfo.LastUpdateSceneFrameNumber = SceneFrameNumber;
 		InstanceUploadInfo.PrimitiveID = PrimitiveID;
 		InstanceUploadInfo.PrimitiveToWorld = FLargeWorldRenderScalar::MakeToRelativeWorldMatrix(AbsoluteOrigin.GetTileOffset(), LocalToWorld);
-		InstanceUploadInfo.bIsPrimitiveForceHidden = PrimitiveSceneProxy->IsForceHidden();
+		
+		// HACK: ignoring IsForceHidden for non-Nanite due to issues that cropped up with water rendering.
+		// TODO: Remove the IsNaniteMesh() check
+		InstanceUploadInfo.bIsPrimitiveForceHidden = PrimitiveSceneProxy->IsNaniteMesh() && PrimitiveSceneProxy->IsForceHidden();
 
 		{
 			bool bHasPrecomputedVolumetricLightmap{};
