Modifies VRS manager and generators to minimize unnecessary work when r.VRS.Enable is set to 1, but no generators are enabled.
We now create a list of active generators at the start of each frame, which speeds up iterating over them later and removes the need to wrap every loop in a mutex.

diff --git a/Engine/Plugins/Runtime/OpenXR/Source/OpenXRHMD/Private/FBFoveationImageGenerator.cpp b/Engine/Plugins/Runtime/OpenXR/Source/OpenXRHMD/Private/FBFoveationImageGenerator.cpp
--- a/Engine/Plugins/Runtime/OpenXR/Source/OpenXRHMD/Private/FBFoveationImageGenerator.cpp
+++ b/Engine/Plugins/Runtime/OpenXR/Source/OpenXRHMD/Private/FBFoveationImageGenerator.cpp
@@ -80,15 +80,21 @@ void FFBFoveationImageGenerator::PrepareImages(FRDGBuilder& GraphBuilder, const
 	return; //Currently not implemented as the images are prepared only when the color swapchain is reallocated.
 }
 
-bool FFBFoveationImageGenerator::IsEnabledForView(const FSceneView& View) const
+bool FFBFoveationImageGenerator::IsEnabled() const
 {
 	if (OpenXRHMD)
 	{
-		return OpenXRHMD->IsStereoEnabled() && bFoveationExtensionSupported && View.StereoPass != EStereoscopicPass::eSSP_FULL;
+		return OpenXRHMD->IsStereoEnabled() && bFoveationExtensionSupported;
 	}
 	return false;
 }
 
+bool FFBFoveationImageGenerator::IsSupportedByView(const FSceneView& View) const
+{
+	// Only used for XR views
+	return IStereoRendering::IsStereoEyeView(View);
+}
+
 // This is currently not implemented.
 FRDGTextureRef FFBFoveationImageGenerator::GetDebugImage(FRDGBuilder& GraphBuilder, const FViewInfo& ViewInfo, FVariableRateShadingImageManager::EVRSImageType ImageType)
 {
diff --git a/Engine/Plugins/Runtime/OpenXR/Source/OpenXRHMD/Private/FBFoveationImageGenerator.h b/Engine/Plugins/Runtime/OpenXR/Source/OpenXRHMD/Private/FBFoveationImageGenerator.h
--- a/Engine/Plugins/Runtime/OpenXR/Source/OpenXRHMD/Private/FBFoveationImageGenerator.h
+++ b/Engine/Plugins/Runtime/OpenXR/Source/OpenXRHMD/Private/FBFoveationImageGenerator.h
@@ -23,7 +23,8 @@ public:
 	/** IVariableRateShadingImageGenerator interface */
 	virtual FRDGTextureRef GetImage(FRDGBuilder& GraphBuilder, const FViewInfo& ViewInfo, FVariableRateShadingImageManager::EVRSImageType ImageType) override;
 	virtual void PrepareImages(FRDGBuilder& GraphBuilder, const FSceneViewFamily& ViewFamily, const FMinimalSceneTextures& SceneTextures) override;
-	virtual bool IsEnabledForView(const FSceneView& View) const override;
+	virtual bool IsEnabled() const override;
+	virtual bool IsSupportedByView(const FSceneView& View) const override;
 	virtual FRDGTextureRef GetDebugImage(FRDGBuilder& GraphBuilder, const FViewInfo& ViewInfo, FVariableRateShadingImageManager::EVRSImageType ImageType) override;
 	virtual FVariableRateShadingImageManager::EVRSSourceType GetType() const override
 	{
diff --git a/Engine/Source/Runtime/Renderer/Private/VariableRateShading/ContrastAdaptiveImageGenerator.cpp b/Engine/Source/Runtime/Renderer/Private/VariableRateShading/ContrastAdaptiveImageGenerator.cpp
--- a/Engine/Source/Runtime/Renderer/Private/VariableRateShading/ContrastAdaptiveImageGenerator.cpp
+++ b/Engine/Source/Runtime/Renderer/Private/VariableRateShading/ContrastAdaptiveImageGenerator.cpp
@@ -320,11 +320,6 @@ static bool IsHDR10(const EDisplayOutputFormat& OutputFormat)
 		OutputFormat == EDisplayOutputFormat::HDR_ACES_2000nit_ST2084;
 }
 
-static bool IsContrastAdaptiveShadingEnabled()
-{
-	return GRHISupportsAttachmentVariableRateShading && GRHIAttachmentVariableRateShadingEnabled && (CVarCASContrastAdaptiveShading.GetValueOnRenderThread() != 0);
-}
-
 static FIntRect GetPostProcessOutputRect(const FViewInfo& ViewInfo)
 {
 	// If TAA/TSR is enabled, upscaling is done at the start of post-processing so the final output will match UnscaledViewRect. Otherwise use the dynamically rescale view rect since
@@ -336,17 +331,14 @@ bool AddCreateShadingRateImagePass(
 	FRDGBuilder& GraphBuilder,
 	const FViewInfo& View)
 {
-	//------------------------------------------------------------------------------------------------
-	// Do some sanity checks for early out
-	if (!IsContrastAdaptiveShadingEnabled() || !FVariableRateShadingImageManager::IsVRSCompatibleWithView(View) || !View.PrevViewInfo.LuminanceHistory)
+	if (!View.PrevViewInfo.LuminanceHistory)
 	{
 		// Shading Rate Image unsupported
 		return false;
 	}
+
 	FRDGTextureRef Luminance = GraphBuilder.RegisterExternalTexture(View.PrevViewInfo.LuminanceHistory);
 	const FVRSTextures& VRSTextures = FVRSTextures::Get(GraphBuilder);
-	// Complete early out sanity checks
-	//------------------------------------------------------------------------------------------------
 
 	{
 		FCalculateShadingRateImageCS::FPermutationDomain PermutationVector;
@@ -478,22 +470,17 @@ FRDGTextureRef FContrastAdaptiveImageGenerator::GetImage(FRDGBuilder& GraphBuild
 void FContrastAdaptiveImageGenerator::PrepareImages(FRDGBuilder& GraphBuilder, const FSceneViewFamily& ViewFamily, const FMinimalSceneTextures& SceneTextures)
 {
 	RDG_EVENT_SCOPE(GraphBuilder, "ContrastAdaptiveShading");
-	bool bAreAllViewsVRSCompatible = true;
+
 	for (const FSceneView* View : ViewFamily.Views)
 	{
 		check(View->bIsViewInfo);
 		const FViewInfo* ViewInfo = static_cast<const FViewInfo*>(View);
 		if (View->bCameraCut || !FVariableRateShadingImageManager::IsVRSCompatibleWithView(*ViewInfo) || !ViewInfo->PrevViewInfo.LuminanceHistory)
 		{
-			bAreAllViewsVRSCompatible = false;
-			break;
+			// CAS is not supported unless all views are set up to support it
+			return;
 		}
 	}
-	bool bPrepareImageBasedVRS = IsContrastAdaptiveShadingEnabled() && bAreAllViewsVRSCompatible;
-	if (!bPrepareImageBasedVRS)
-	{
-		return;
-	}
 
 	FVRSTextures& VRSTextures = GraphBuilder.Blackboard.Create<FVRSTextures>();
 	VRSTextures.Create(GraphBuilder, ViewFamily);
@@ -509,12 +496,16 @@ void FContrastAdaptiveImageGenerator::PrepareImages(FRDGBuilder& GraphBuilder, c
 	AddPrepareImageBasedVRSPass(GraphBuilder, SceneTextures, ViewFamily);
 }
 
-bool FContrastAdaptiveImageGenerator::IsEnabledForView(const FSceneView& View) const
+bool FContrastAdaptiveImageGenerator::IsEnabled() const
 {
-	EDisplayOutputFormat DisplayOutputFormat = GetDisplayOutputFormat(View);
-	bool bCompatibleWithOutputType = (DisplayOutputFormat == EDisplayOutputFormat::SDR_sRGB) || IsHDR10(DisplayOutputFormat);
+	return GRHISupportsAttachmentVariableRateShading && GRHIAttachmentVariableRateShadingEnabled && (CVarCASContrastAdaptiveShading.GetValueOnRenderThread() != 0);
+}
 
-	return IsContrastAdaptiveShadingEnabled() && !View.bIsSceneCapture && bCompatibleWithOutputType;
+bool FContrastAdaptiveImageGenerator::IsSupportedByView(const FSceneView& View) const
+{
+	EDisplayOutputFormat DisplayOutputFormat = GetDisplayOutputFormat(View);
+	const bool bCompatibleWithOutputType = (DisplayOutputFormat == EDisplayOutputFormat::SDR_sRGB) || IsHDR10(DisplayOutputFormat);
+	return !View.bIsSceneCapture && bCompatibleWithOutputType;
 }
 
 FRDGTextureRef FContrastAdaptiveImageGenerator::GetDebugImage(FRDGBuilder& GraphBuilder, const FViewInfo& ViewInfo, FVariableRateShadingImageManager::EVRSImageType ImageType)
diff --git a/Engine/Source/Runtime/Renderer/Private/VariableRateShading/ContrastAdaptiveImageGenerator.h b/Engine/Source/Runtime/Renderer/Private/VariableRateShading/ContrastAdaptiveImageGenerator.h
--- a/Engine/Source/Runtime/Renderer/Private/VariableRateShading/ContrastAdaptiveImageGenerator.h
+++ b/Engine/Source/Runtime/Renderer/Private/VariableRateShading/ContrastAdaptiveImageGenerator.h
@@ -15,7 +15,8 @@ public:
 	virtual ~FContrastAdaptiveImageGenerator() override {};
 	virtual FRDGTextureRef GetImage(FRDGBuilder& GraphBuilder, const FViewInfo& ViewInfo, FVariableRateShadingImageManager::EVRSImageType ImageType) override;
 	virtual void PrepareImages(FRDGBuilder& GraphBuilder, const FSceneViewFamily& ViewFamily, const FMinimalSceneTextures& SceneTextures) override;
-	virtual bool IsEnabledForView(const FSceneView& View) const override;
+	virtual bool IsEnabled() const override;
+	virtual bool IsSupportedByView(const FSceneView& View) const override;
 	virtual FVariableRateShadingImageManager::EVRSSourceType GetType() const override
 	{
 		return FVariableRateShadingImageManager::EVRSSourceType::ContrastAdaptiveShading;
diff --git a/Engine/Source/Runtime/Renderer/Private/VariableRateShading/FoveatedImageGenerator.cpp b/Engine/Source/Runtime/Renderer/Private/VariableRateShading/FoveatedImageGenerator.cpp
--- a/Engine/Source/Runtime/Renderer/Private/VariableRateShading/FoveatedImageGenerator.cpp
+++ b/Engine/Source/Runtime/Renderer/Private/VariableRateShading/FoveatedImageGenerator.cpp
@@ -241,10 +241,15 @@ void FFoveatedImageGenerator::PrepareImages(FRDGBuilder& GraphBuilder, const FSc
 	CachedImage = ShadingRateTexture;
 }
 
-bool FFoveatedImageGenerator::IsEnabledForView(const FSceneView& View) const
+bool FFoveatedImageGenerator::IsEnabled() const
 {
-	// Enabled for stereo (XR or emulated)
-	return IStereoRendering::IsStereoEyeView(View) && CVarFoveationLevel.GetValueOnRenderThread() > 0;
+	return CVarFoveationLevel.GetValueOnRenderThread() > 0;
+}
+
+bool FFoveatedImageGenerator::IsSupportedByView(const FSceneView& View) const
+{
+	// Only used for XR views
+	return true; // IStereoRendering::IsStereoEyeView(View);
 }
 
 FVariableRateShadingImageManager::EVRSSourceType FFoveatedImageGenerator::GetType() const
diff --git a/Engine/Source/Runtime/Renderer/Private/VariableRateShading/FoveatedImageGenerator.h b/Engine/Source/Runtime/Renderer/Private/VariableRateShading/FoveatedImageGenerator.h
--- a/Engine/Source/Runtime/Renderer/Private/VariableRateShading/FoveatedImageGenerator.h
+++ b/Engine/Source/Runtime/Renderer/Private/VariableRateShading/FoveatedImageGenerator.h
@@ -15,7 +15,8 @@ public:
 	virtual ~FFoveatedImageGenerator() override {};
 	virtual FRDGTextureRef GetImage(FRDGBuilder& GraphBuilder, const FViewInfo& ViewInfo, FVariableRateShadingImageManager::EVRSImageType ImageType) override;
 	virtual void PrepareImages(FRDGBuilder& GraphBuilder, const FSceneViewFamily& ViewFamily, const FMinimalSceneTextures& SceneTextures) override;
-	virtual bool IsEnabledForView(const FSceneView& View) const override;
+	virtual bool IsEnabled() const override;
+	virtual bool IsSupportedByView(const FSceneView& View) const override;
 	virtual FVariableRateShadingImageManager::EVRSSourceType GetType() const override;
 	virtual FRDGTextureRef GetDebugImage(FRDGBuilder& GraphBuilder, const FViewInfo& ViewInfo, FVariableRateShadingImageManager::EVRSImageType ImageType) override;
 private:
diff --git a/Engine/Source/Runtime/Renderer/Private/VariableRateShading/VariableRateShadingImageManager.cpp b/Engine/Source/Runtime/Renderer/Private/VariableRateShading/VariableRateShadingImageManager.cpp
--- a/Engine/Source/Runtime/Renderer/Private/VariableRateShading/VariableRateShadingImageManager.cpp
+++ b/Engine/Source/Runtime/Renderer/Private/VariableRateShading/VariableRateShadingImageManager.cpp
@@ -380,10 +380,14 @@ int32 FVariableRateShadingImageManager::GetNumberOfSupportedRates()
 FRDGTextureRef FVariableRateShadingImageManager::GetVariableRateShadingImage(FRDGBuilder& GraphBuilder, const FViewInfo& ViewInfo, FVariableRateShadingImageManager::EVRSPassType PassType,
 	FVariableRateShadingImageManager::EVRSSourceType VRSTypesToExclude)
 {
-	EVRSImageType ImageType = GetImageTypeFromPassType(PassType);
-
 	// If the view doesn't support VRS or this pass is disabled, bail immediately
-	if (!bVRSEnabledForFrame || !IsVRSCompatibleWithView(ViewInfo) || ImageType == EVRSImageType::Disabled)
+	if (!bVRSEnabledForFrame || ActiveGenerators.IsEmpty())
+	{
+		return nullptr;
+	}
+
+	EVRSImageType ImageType = GetImageTypeFromPassType(PassType);
+	if (ImageType == EVRSImageType::Disabled || !IsVRSCompatibleWithView(ViewInfo))
 	{
 		return nullptr;
 	}
@@ -399,21 +403,17 @@ FRDGTextureRef FVariableRateShadingImageManager::GetVariableRateShadingImage(FRD
 
 	// Otherwise collate all internal sources
 	TArray<FRDGTextureRef> InternalVRSSources;
-
+	for (IVariableRateShadingImageGenerator* const Generator : ActiveGenerators)
 	{
-		FReadScopeLock GeneratorsLock(GeneratorsMutex);	
-		for (IVariableRateShadingImageGenerator* const Generator : ImageGenerators)
+		FRDGTextureRef Image = nullptr;
+		if (Generator && Generator->IsSupportedByView(ViewInfo) && !EnumHasAnyFlags(VRSTypesToExclude, Generator->GetType()))
 		{
-			FRDGTextureRef Image = nullptr;
-			if (Generator->IsEnabledForView(ViewInfo) && !EnumHasAnyFlags(VRSTypesToExclude, Generator->GetType()))
-			{
-				Image = Generator->GetImage(GraphBuilder, ViewInfo, ImageType);
-			}
+			Image = Generator->GetImage(GraphBuilder, ViewInfo, ImageType);
+		}
 
-			if (Image)
-			{
-				InternalVRSSources.Add(Image);
-			}
+		if (Image)
+		{
+			InternalVRSSources.Add(Image);
 		}
 	}
 
@@ -439,12 +439,18 @@ void FVariableRateShadingImageManager::PrepareImageBasedVRS(FRDGBuilder& GraphBu
 		return;
 	}
 
-	RDG_EVENT_SCOPE(GraphBuilder, "PrepareImageBasedVRS");
-	SCOPED_NAMED_EVENT(PrepareImageBasedVRS, FColor::Red);
-
-	VRSForceRateForFrame = CVarVRSDebugForceRate->GetInt();
+	// If no generators are active, bail
+	{
+		FReadScopeLock GeneratorsLock(GeneratorsMutex);
+		ActiveGenerators = ImageGenerators.FilterByPredicate([](IVariableRateShadingImageGenerator* const InGenerator) { return InGenerator && InGenerator->IsEnabled(); });
+	}
 
-	// If no views support VRS, bail immediately
+	if (ActiveGenerators.IsEmpty())
+	{
+		return;
+	}
+	
+	// If no views support VRS, bail
 	bool bIsAnyViewVRSCompatible = false;
 	for (const FSceneView* View : ViewFamily.Views)
 	{
@@ -469,25 +475,26 @@ void FVariableRateShadingImageManager::PrepareImageBasedVRS(FRDGBuilder& GraphBu
 		return;
 	}
 
-	// Invoke image generators
+	RDG_EVENT_SCOPE(GraphBuilder, "PrepareImageBasedVRS");
+	SCOPED_NAMED_EVENT(PrepareImageBasedVRS, FColor::Red);
+
+	VRSForceRateForFrame = CVarVRSDebugForceRate->GetInt();
+
+	// Invoke active image generators
+	for (IVariableRateShadingImageGenerator* const Generator : ActiveGenerators)
 	{
-		FReadScopeLock GeneratorsLock(GeneratorsMutex);
-		for (IVariableRateShadingImageGenerator* const Generator : ImageGenerators)
+		if (Generator && Generator->IsSupportedByView(*ViewFamily.Views[0]))
 		{
-			if (Generator->IsEnabledForView(*ViewFamily.Views[0]))
-			{
-				Generator->PrepareImages(GraphBuilder, ViewFamily, SceneTextures);
-			}
+			Generator->PrepareImages(GraphBuilder, ViewFamily, SceneTextures);
 		}
 	}
 }
 
 bool FVariableRateShadingImageManager::IsTypeEnabledForView(const FSceneView& View, FVariableRateShadingImageManager::EVRSSourceType Type)
 {
-	FReadScopeLock GeneratorsLock(GeneratorsMutex);
-	for (IVariableRateShadingImageGenerator* const Generator : ImageGenerators)
+	for (IVariableRateShadingImageGenerator* const Generator : ActiveGenerators)
 	{
-		if (EnumHasAnyFlags(Type, Generator->GetType()) && Generator->IsEnabledForView(View))
+		if (Generator && EnumHasAnyFlags(Type, Generator->GetType()) && Generator->IsSupportedByView(View))
 		{
 			return true;
 		}
@@ -497,6 +504,11 @@ bool FVariableRateShadingImageManager::IsTypeEnabledForView(const FSceneView& Vi
 
 void FVariableRateShadingImageManager::DrawDebugPreview(FRDGBuilder& GraphBuilder, const FSceneViewFamily& ViewFamily, FRDGTextureRef OutputSceneColor)
 {
+	if (!bVRSEnabledForFrame)
+	{
+		return;
+	}
+
 	uint32 ImageTypeAsInt = CVarVRSPreview.GetValueOnRenderThread();
 	EVRSImageType PreviewImageType = EVRSImageType::Disabled;
 	if (ImageTypeAsInt >= 0 && ImageTypeAsInt <= EVRSImageType::Conservative)
@@ -528,20 +540,17 @@ void FVariableRateShadingImageManager::DrawDebugPreview(FRDGBuilder& GraphBuilde
 			{
 				TArray<FRDGTextureRef> InternalVRSSources;
 
+				for (IVariableRateShadingImageGenerator* const Generator : ActiveGenerators)
 				{
-					FReadScopeLock GeneratorsLock(GeneratorsMutex);
-					for (IVariableRateShadingImageGenerator* const Generator : ImageGenerators)
+					FRDGTextureRef Image = nullptr;
+					if (Generator && Generator->IsSupportedByView(*View))
 					{
-						FRDGTextureRef Image = nullptr;
-						if (Generator->IsEnabledForView(*View))
-						{
-							Image = Generator->GetDebugImage(GraphBuilder, *ViewInfo, PreviewImageType);
-						}
-
-						if (Image)
-						{
-							InternalVRSSources.Add(Image);
-						}
+						Image = Generator->GetDebugImage(GraphBuilder, *ViewInfo, PreviewImageType);
+					}
+
+					if (Image)
+					{
+						InternalVRSSources.Add(Image);
 					}
 				}
 
@@ -584,6 +593,8 @@ void FVariableRateShadingImageManager::DrawDebugPreview(FRDGBuilder& GraphBuilde
 			const FScreenPassTextureViewport InputViewport = FScreenPassTextureViewport(PreviewTexture, ScaledSrcRect);
 			const FScreenPassTextureViewport OutputViewport(OutputSceneColor, DestViewRect);
 
+			FPlatformMisc::LowLevelOutputDebugStringf(TEXT("Preview for frame %d\n"), ViewFamily.FrameNumber);
+
 			AddDrawScreenPass(
 				GraphBuilder,
 				RDG_EVENT_NAME("Display VRS Debug Preview"),
@@ -653,6 +664,8 @@ FRDGTextureRef FVariableRateShadingImageManager::CombineShadingRateImages(FRDGBu
 
 		TShaderMapRef<FCombineShadingRateTexturesCS> ComputeShader(GetGlobalShaderMap(GMaxRHIFeatureLevel));
 
+		FPlatformMisc::LowLevelOutputDebugStringf(TEXT("Combiner for frame %d\n"), ViewInfo.Family->FrameNumber);
+
 		FComputeShaderUtils::AddPass(
 			GraphBuilder,
 			RDG_EVENT_NAME("CombineShadingRateImages"),
diff --git a/Engine/Source/Runtime/Renderer/Public/VariableRateShadingImageManager.h b/Engine/Source/Runtime/Renderer/Public/VariableRateShadingImageManager.h
--- a/Engine/Source/Runtime/Renderer/Public/VariableRateShadingImageManager.h
+++ b/Engine/Source/Runtime/Renderer/Public/VariableRateShadingImageManager.h
@@ -117,6 +117,7 @@ private:
 	TRefCountPtr<IPooledRenderTarget> MobileHMDFixedFoveationOverrideImage;
 	// Used to access the generators functionalities.
 	TArray<IVariableRateShadingImageGenerator*> ImageGenerators;
+	TArray<IVariableRateShadingImageGenerator*> ActiveGenerators;
 	// This is used only to own the memory of generators created in the FVariableRateShadingImageManager constructor.
 	TArray<TUniquePtr<IVariableRateShadingImageGenerator>> InternalGenerators;
 	// Guards ImageGenerators because most of the calls of the manager are on the render thread but the 
@@ -148,7 +149,10 @@ public:
 	virtual void PrepareImages(FRDGBuilder& GraphBuilder, const FSceneViewFamily& ViewFamily, const FMinimalSceneTextures& SceneTextures) = 0;
 
 	// Returns whether or not generator is enabled - can change at runtime
-	virtual bool IsEnabledForView(const FSceneView& View) const { return false; };
+	virtual bool IsEnabled() const { return false; };
+
+	// Returns whether or not the given view supports this generator
+	virtual bool IsSupportedByView(const FSceneView& View) const { return false; };
 
 	// Return bitmask of generator type
 	virtual FVariableRateShadingImageManager::EVRSSourceType GetType() const { return FVariableRateShadingImageManager::EVRSSourceType::None; };
