Path Tracer: Skip Planar reflections when path tracing or RT debug is enabled

diff --git a/Engine/Source/Runtime/Renderer/Private/SceneRendering.cpp b/Engine/Source/Runtime/Renderer/Private/SceneRendering.cpp
--- a/Engine/Source/Runtime/Renderer/Private/SceneRendering.cpp
+++ b/Engine/Source/Runtime/Renderer/Private/SceneRendering.cpp
@@ -4873,6 +4873,10 @@ void FRendererModule::BeginRenderingViewFamilies(FCanvas* Canvas, TArrayView<FSc
 				UPlanarReflectionComponent* ReflectionComponent = Scene->PlanarReflections_GameThread[ReflectionIndex];
 				for (FSceneRenderer* SceneRenderer : SceneRenderers)
 				{
+					if (HasRayTracedOverlay(SceneRenderer->ViewFamily))
+					{
+						continue;
+					}
 					Scene->UpdatePlanarReflectionContents(ReflectionComponent, *SceneRenderer);
 				}
 			}
