Remove unintended debug print statements from VRS generator.

diff --git a/Engine/Source/Runtime/Renderer/Private/VariableRateShading/VariableRateShadingImageManager.cpp b/Engine/Source/Runtime/Renderer/Private/VariableRateShading/VariableRateShadingImageManager.cpp
--- a/Engine/Source/Runtime/Renderer/Private/VariableRateShading/VariableRateShadingImageManager.cpp
+++ b/Engine/Source/Runtime/Renderer/Private/VariableRateShading/VariableRateShadingImageManager.cpp
@@ -593,8 +593,6 @@ void FVariableRateShadingImageManager::DrawDebugPreview(FRDGBuilder& GraphBuilde
 			const FScreenPassTextureViewport InputViewport = FScreenPassTextureViewport(PreviewTexture, ScaledSrcRect);
 			const FScreenPassTextureViewport OutputViewport(OutputSceneColor, DestViewRect);
 
-			FPlatformMisc::LowLevelOutputDebugStringf(TEXT("Preview for frame %d\n"), ViewFamily.FrameNumber);
-
 			AddDrawScreenPass(
 				GraphBuilder,
 				RDG_EVENT_NAME("Display VRS Debug Preview"),
@@ -664,8 +662,6 @@ FRDGTextureRef FVariableRateShadingImageManager::CombineShadingRateImages(FRDGBu
 
 		TShaderMapRef<FCombineShadingRateTexturesCS> ComputeShader(GetGlobalShaderMap(GMaxRHIFeatureLevel));
 
-		FPlatformMisc::LowLevelOutputDebugStringf(TEXT("Combiner for frame %d\n"), ViewInfo.Family->FrameNumber);
-
 		FComputeShaderUtils::AddPass(
 			GraphBuilder,
 			RDG_EVENT_NAME("CombineShadingRateImages"),
