Path Tracing: Disable auto-exposure mode which depends on the presence of a GBuffer

This mode is not easy to support in the path tracer now that we skip the basepass. If it is really needed in the future, we can potentially try to popular some of gbuffer textures from the path tracer, but this would slow down the more common case and this particular option is still experimental for now.

diff --git a/Engine/Source/Runtime/Renderer/Private/PostProcess/PostProcessEyeAdaptation.cpp b/Engine/Source/Runtime/Renderer/Private/PostProcess/PostProcessEyeAdaptation.cpp
--- a/Engine/Source/Runtime/Renderer/Private/PostProcess/PostProcessEyeAdaptation.cpp
+++ b/Engine/Source/Runtime/Renderer/Private/PostProcess/PostProcessEyeAdaptation.cpp
@@ -183,7 +183,8 @@ bool IsExtendLuminanceRangeEnabled()
 
 bool IsAutoExposureUsingIlluminanceEnabled(const FViewInfo& View)
 {
-	return CVarAutoExposureIgnoreMaterials.GetValueOnRenderThread() && GetAutoExposureMethod(View) == AEM_Histogram;
+	// NOTE: This method cannot be supported with PathTracing because it requires a GBuffer which is not available in the path tracing case
+	return CVarAutoExposureIgnoreMaterials.GetValueOnRenderThread() && GetAutoExposureMethod(View) == AEM_Histogram && !View.Family->EngineShowFlags.PathTracing;
 }
 
 int32 GetAutoExposureIlluminanceDownscaleFactor()
