LFV - height fog volume can have exponentially increasing density and this causes leaking issue with VFog temporal reprojection.
A max density has been added to get that visual artefact under control.

#prelfight none

diff --git a/Engine/Shaders/Private/LocalFogVolumes/LocalFogVolumeCommon.ush b/Engine/Shaders/Private/LocalFogVolumes/LocalFogVolumeCommon.ush
--- a/Engine/Shaders/Private/LocalFogVolumes/LocalFogVolumeCommon.ush
+++ b/Engine/Shaders/Private/LocalFogVolumes/LocalFogVolumeCommon.ush
@@ -30,31 +30,34 @@ struct FLocalFogVolumeGPUInstanceData
 
 #if defined(LFVStruct)
 
-#define LFVInsts					LFVStruct.LFV.LocalFogVolumeInstances
-#define LFVTileDataTex				LFVStruct.LFV.LocalFogVolumeTileDataTexture
+#define LFVInsts						LFVStruct.LFV.LocalFogVolumeInstances
+#define LFVTileDataTex					LFVStruct.LFV.LocalFogVolumeTileDataTexture
 
-#define LFVTileDataResolution		LFVStruct.LFV.LocalFogVolumeTileDataTextureResolution
-#define LFVInstcount				LFVStruct.LFV.LocalFogVolumeInstanceCount
-#define LFVTilePixelSize			LFVStruct.LFV.LocalFogVolumeTilePixelSize
-#define LFVRenderInVolumetricFog	LFVStruct.LFV.ShouldRenderLocalFogVolumeInVolumetricFog
+#define LFVTileDataResolution			LFVStruct.LFV.LocalFogVolumeTileDataTextureResolution
+#define LFVInstcount					LFVStruct.LFV.LocalFogVolumeInstanceCount
+#define LFVTilePixelSize				LFVStruct.LFV.LocalFogVolumeTilePixelSize
+#define LFVMaxDensityIntoVolumetricFog	LFVStruct.LFV.LocalFogVolumeMaxDensityIntoVolumetricFog
+#define LFVRenderInVolumetricFog		LFVStruct.LFV.ShouldRenderLocalFogVolumeInVolumetricFog
 
 #else
 
-uint2								LFV_LocalFogVolumeTileDataTextureResolution;
-uint								LFV_LocalFogVolumeInstanceCount;
-uint								LFV_LocalFogVolumeTilePixelSize;
-uint								LFV_ShouldRenderLocalFogVolumeInVolumetricFog;
+uint2									LFV_LocalFogVolumeTileDataTextureResolution;
+uint									LFV_LocalFogVolumeInstanceCount;
+uint									LFV_LocalFogVolumeTilePixelSize;
+float									LFV_LocalFogVolumeMaxDensityIntoVolumetricFog;
+uint									LFV_ShouldRenderLocalFogVolumeInVolumetricFog;
 
-Buffer<float4>						LFV_LocalFogVolumeInstances;
-Texture2DArray<uint>				LFV_LocalFogVolumeTileDataTexture;
+Buffer<float4>							LFV_LocalFogVolumeInstances;
+Texture2DArray<uint>					LFV_LocalFogVolumeTileDataTexture;
 
-#define LFVInsts					LFV_LocalFogVolumeInstances
-#define LFVTileDataTex				LFV_LocalFogVolumeTileDataTexture
+#define LFVInsts						LFV_LocalFogVolumeInstances
+#define LFVTileDataTex					LFV_LocalFogVolumeTileDataTexture
 
-#define LFVTileDataResolution		LFV_LocalFogVolumeTileDataTextureResolution
-#define LFVInstcount				LFV_LocalFogVolumeInstanceCount
-#define LFVTilePixelSize			LFV_LocalFogVolumeTilePixelSize
-#define LFVRenderInVolumetricFog	LFV_ShouldRenderLocalFogVolumeInVolumetricFog
+#define LFVTileDataResolution			LFV_LocalFogVolumeTileDataTextureResolution
+#define LFVInstcount					LFV_LocalFogVolumeInstanceCount
+#define LFVTilePixelSize				LFV_LocalFogVolumeTilePixelSize
+#define LFVMaxDensityIntoVolumetricFog	LFV_LocalFogVolumeMaxDensityIntoVolumetricFog
+#define LFVRenderInVolumetricFog		LFV_ShouldRenderLocalFogVolumeInVolumetricFog
 
 #endif
 
diff --git a/Engine/Shaders/Private/VolumetricFog.usf b/Engine/Shaders/Private/VolumetricFog.usf
--- a/Engine/Shaders/Private/VolumetricFog.usf
+++ b/Engine/Shaders/Private/VolumetricFog.usf
@@ -197,6 +197,7 @@ void MaterialSetupCS(
 		{
 			LFVExtinction = FogInstance.Density * SphereFade / FogInstance.UniformScale;
 		}
+		LFVExtinction = min(LFVMaxDensityIntoVolumetricFog, LFVExtinction);
 
 		float3 LFVScattering = LFVExtinction * FogInstance.Albedo;
 
diff --git a/Engine/Source/Runtime/Renderer/Private/LocalFogVolumeRendering.cpp b/Engine/Source/Runtime/Renderer/Private/LocalFogVolumeRendering.cpp
--- a/Engine/Source/Runtime/Renderer/Private/LocalFogVolumeRendering.cpp
+++ b/Engine/Source/Runtime/Renderer/Private/LocalFogVolumeRendering.cpp
@@ -19,6 +19,11 @@ static TAutoConsoleVariable<int32> CVarLocalFogVolumeRenderIntoVolumetricFog(
 	TEXT("LocalFogVolume are going to be voxelised into the volumetric fog when this is not 0, otherwise it will remain isolated.\n"),
 	ECVF_RenderThreadSafe);
 
+static TAutoConsoleVariable<float> CVarLocalFogVolumeMaxDensityIntoVolumetricFog(
+	TEXT("r.LocalFogVolume.MaxDensityIntoVolumetricFog"), 0.01f,
+	TEXT("LocalFogVolume height fog mode can become exponentially dense in the bottom part. VolumetricFog temporal reprojection then can leak du to high density. Clamping density is a way to get that visual artefact under control.\n"),
+	ECVF_RenderThreadSafe);
+
 static TAutoConsoleVariable<int32> CVarLocalFogVolumeApplyOnTranslucent(
 	TEXT("r.LocalFogVolume.ApplyOnTranslucent"), 0,
 	TEXT("Project settings enabling the sampling of local fog volumes on translucent elements.\n"),
@@ -58,6 +63,11 @@ static uint32 GetLocalFogVolumeTilePixelSize()
 	return FMath::Max(8u, FMath::Min(512u, (uint32)CVarLocalFogVolumeTilePixelSize.GetValueOnRenderThread()));
 }
 
+static float GetLocalFogVolumeMaxDensityIntoVolumetricFog()
+{
+	return FMath::Max(0.0f, CVarLocalFogVolumeMaxDensityIntoVolumetricFog.GetValueOnRenderThread());
+}
+
 static uint32 GetLocalFogVolumeTileMaxInstanceCount()
 {
 	// We do not allow more than 256 instances since culled indices might be stored a u8 values.
@@ -110,12 +120,13 @@ void SetDummyLocalFogVolumeForView(FRDGBuilder& GraphBuilder, FViewInfo& View)
 	View.LocalFogVolumeViewData.TileDataTextureArraySRV			= GraphBuilder.CreateSRV(View.LocalFogVolumeViewData.TileDataTextureArray);
 	View.LocalFogVolumeViewData.TileDataTextureArrayUAV			= nullptr;	// Should never be written by culling passes if there are no instances.
 
-	View.LocalFogVolumeViewData.UniformParametersStruct.LocalFogVolumeCommon.LocalFogVolumeTileDataTextureResolution= FUintVector2(1, 1);
-	View.LocalFogVolumeViewData.UniformParametersStruct.LocalFogVolumeCommon.LocalFogVolumeInstanceCount			= View.LocalFogVolumeViewData.GPUInstanceCount;
-	View.LocalFogVolumeViewData.UniformParametersStruct.LocalFogVolumeCommon.LocalFogVolumeTilePixelSize			= GetLocalFogVolumeTilePixelSize();
-	View.LocalFogVolumeViewData.UniformParametersStruct.LocalFogVolumeCommon.LocalFogVolumeInstances				= View.LocalFogVolumeViewData.GPUInstanceDataBufferSRV;
-	View.LocalFogVolumeViewData.UniformParametersStruct.LocalFogVolumeTileDataTexture								= View.LocalFogVolumeViewData.TileDataTextureArraySRV;
-	View.LocalFogVolumeViewData.UniformBuffer																		= GraphBuilder.CreateUniformBuffer(&View.LocalFogVolumeViewData.UniformParametersStruct);
+	View.LocalFogVolumeViewData.UniformParametersStruct.LocalFogVolumeCommon.LocalFogVolumeTileDataTextureResolution	= FUintVector2(1, 1);
+	View.LocalFogVolumeViewData.UniformParametersStruct.LocalFogVolumeCommon.LocalFogVolumeInstanceCount				= View.LocalFogVolumeViewData.GPUInstanceCount;
+	View.LocalFogVolumeViewData.UniformParametersStruct.LocalFogVolumeCommon.LocalFogVolumeTilePixelSize				= GetLocalFogVolumeTilePixelSize();
+	View.LocalFogVolumeViewData.UniformParametersStruct.LocalFogVolumeCommon.LocalFogVolumeMaxDensityIntoVolumetricFog	= GetLocalFogVolumeMaxDensityIntoVolumetricFog();
+	View.LocalFogVolumeViewData.UniformParametersStruct.LocalFogVolumeCommon.LocalFogVolumeInstances					= View.LocalFogVolumeViewData.GPUInstanceDataBufferSRV;
+	View.LocalFogVolumeViewData.UniformParametersStruct.LocalFogVolumeTileDataTexture									= View.LocalFogVolumeViewData.TileDataTextureArraySRV;
+	View.LocalFogVolumeViewData.UniformBuffer																			= GraphBuilder.CreateUniformBuffer(&View.LocalFogVolumeViewData.UniformParametersStruct);
 
 	// This buffer must remain a basic vertex buffer for mobile to be able to read it from vertex shader
 	View.LocalFogVolumeViewData.GPUTileDataBuffer		= CreateVertexBuffer(
@@ -378,6 +389,7 @@ void CreateViewLocalFogVolumeBufferSRV(FViewInfo& View, FRDGBuilder& GraphBuilde
 	View.LocalFogVolumeViewData.UniformParametersStruct.LocalFogVolumeCommon.LocalFogVolumeTileDataTextureResolution	= FUintVector2(TileDataTextureResolution.X, TileDataTextureResolution.Y);
 	View.LocalFogVolumeViewData.UniformParametersStruct.LocalFogVolumeCommon.LocalFogVolumeInstanceCount				= View.LocalFogVolumeViewData.GPUInstanceCount;
 	View.LocalFogVolumeViewData.UniformParametersStruct.LocalFogVolumeCommon.LocalFogVolumeTilePixelSize				= LocalFogVolumeTilePixelSize;
+	View.LocalFogVolumeViewData.UniformParametersStruct.LocalFogVolumeCommon.LocalFogVolumeMaxDensityIntoVolumetricFog	= GetLocalFogVolumeMaxDensityIntoVolumetricFog();
 	View.LocalFogVolumeViewData.UniformParametersStruct.LocalFogVolumeCommon.ShouldRenderLocalFogVolumeInVolumetricFog	= bShouldRenderLocalFogVolumeInVolumetricFog ? 1 : 0;
 	View.LocalFogVolumeViewData.UniformParametersStruct.LocalFogVolumeCommon.LocalFogVolumeInstances					= View.LocalFogVolumeViewData.GPUInstanceDataBufferSRV;
 	View.LocalFogVolumeViewData.UniformParametersStruct.LocalFogVolumeTileDataTexture									= View.LocalFogVolumeViewData.TileDataTextureArraySRV;
diff --git a/Engine/Source/Runtime/Renderer/Private/LocalFogVolumeRendering.h b/Engine/Source/Runtime/Renderer/Private/LocalFogVolumeRendering.h
--- a/Engine/Source/Runtime/Renderer/Private/LocalFogVolumeRendering.h
+++ b/Engine/Source/Runtime/Renderer/Private/LocalFogVolumeRendering.h
@@ -24,6 +24,7 @@ BEGIN_SHADER_PARAMETER_STRUCT(FLocalFogVolumeCommonParameters, )
 	SHADER_PARAMETER(FUintVector2, LocalFogVolumeTileDataTextureResolution)
 	SHADER_PARAMETER(uint32, LocalFogVolumeInstanceCount)
 	SHADER_PARAMETER(uint32, LocalFogVolumeTilePixelSize)
+	SHADER_PARAMETER(float,  LocalFogVolumeMaxDensityIntoVolumetricFog)
 	SHADER_PARAMETER(uint32, ShouldRenderLocalFogVolumeInVolumetricFog)
 END_SHADER_PARAMETER_STRUCT()
 
