Path Tracing: Skip Distance field heightfield work when the path tracer is active

diff --git a/Engine/Source/Runtime/Renderer/Private/DistanceFieldShadowing.cpp b/Engine/Source/Runtime/Renderer/Private/DistanceFieldShadowing.cpp
--- a/Engine/Source/Runtime/Renderer/Private/DistanceFieldShadowing.cpp
+++ b/Engine/Source/Runtime/Renderer/Private/DistanceFieldShadowing.cpp
@@ -724,6 +724,7 @@ bool FSceneRenderer::ShouldPrepareHeightFieldScene() const
 {
 	return Scene
 		&& ViewFamily.EngineShowFlags.DynamicShadows
+		&& !ViewFamily.EngineShowFlags.PathTracing
 		&& SupportsHeightFieldShadows(Scene->GetFeatureLevel(), Scene->GetShaderPlatform());
 }
 
