Re-enable spline mesh scene textures and bump the relevant shaders to invalidate their DDC cache.
NOTE: The issue that caused me to be disabled them in the first place was fixed by a different GPU Scene change (see 27994320).

diff --git a/Engine/Source/Runtime/Renderer/Private/SplineMeshSceneResources.cpp b/Engine/Source/Runtime/Renderer/Private/SplineMeshSceneResources.cpp
--- a/Engine/Source/Runtime/Renderer/Private/SplineMeshSceneResources.cpp
+++ b/Engine/Source/Runtime/Renderer/Private/SplineMeshSceneResources.cpp
@@ -13,9 +13,34 @@
 #include "RHIStaticStates.h"
 #include "RenderCaptureInterface.h"
 
+DEFINE_LOG_CATEGORY(LogSplineMesh);
+
+#define SPLINE_MESH_DEBUG_LOGGING (!UE_BUILD_SHIPPING && !UE_BUILD_TEST)
+
+#if SPLINE_MESH_DEBUG_LOGGING
+	#define SPLINE_MESH_LOG(...) UE_LOG(LogSplineMesh, __VA_ARGS__)
+#else
+	#define SPLINE_MESH_LOG(...) ((void)0)
+#endif
+
+FString GetIntListStringDebug(TConstArrayView<uint32> IntList)
+{
+#if SPLINE_MESH_DEBUG_LOGGING
+	FString String;
+	for (uint32 Int : IntList)
+	{
+		String.Appendf(TEXT("%u "), Int);
+	}
+#else
+	FString String = TEXT("");
+#endif
+
+	return String;
+}
+
 static TAutoConsoleVariable<int32> CVarSplineMeshSceneTextures(
 	TEXT("r.SplineMesh.SceneTextures"),
-	0,
+	1,
 	TEXT("Whether to cache all spline mesh splines in the scene to textures (performance optimization)."),
 	ECVF_ReadOnly
 );
@@ -174,6 +199,8 @@ void FSplineMeshSceneResources::PostSceneUpdate(FRDGBuilder& GraphBuilder, const
 		const uint32 DefraggedSize = SlotAllocator.GetSparselyAllocatedSize();
 		if (SplineMesh::CalcTextureSize(DefraggedSize) < SplineMesh::CalcTextureSize(CurSize))
 		{
+			SPLINE_MESH_LOG(Display, TEXT("Defragmenting spline mesh scene textures. (%u -> %u)"), CurSize, DefraggedSize);
+
 			DefragTexture();
 		}
 	}
@@ -183,7 +210,11 @@ void FSplineMeshSceneResources::PostSceneUpdate(FRDGBuilder& GraphBuilder, const
 		SavedPosTexture = nullptr;
 		SavedRotTexture = nullptr;
 		SavedIdLookup = nullptr;
-		UpdateRequests.Reset();
+		if (UpdateRequests.Num() > 0)
+		{
+			SPLINE_MESH_LOG(Verbose, TEXT("Dropping update requests: [%s]"), *GetIntListStringDebug(UpdateRequests))
+			UpdateRequests.Reset();
+		}
 	}
 }
 
@@ -204,7 +235,7 @@ void FSplineMeshSceneResources::Update(FRDGBuilder& GraphBuilder, FSceneUniformB
 		NeededSize = SPLINE_MESH_TEXTURE_MAX_DIMENSION;
 		if (!bOverflowError)
 		{
-			UE_LOG(LogRenderer, Error,
+			UE_LOG(LogSplineMesh, Error,
 				TEXT("Too many spline meshes have been registered with the scene. The spline mesh texture has grown ")
 				TEXT("to its max size (%dx%d - see r.SplineMesh.BakeToTexture.MaxDimension) and has ran out of space. ")
 				TEXT("Expect some spline meshes to render incorrectly."),
@@ -222,6 +253,9 @@ void FSplineMeshSceneResources::Update(FRDGBuilder& GraphBuilder, FSceneUniformB
 	FRDGTextureRef RotTexture = nullptr;
 	if (NeededSize != CurSize)
 	{
+		SPLINE_MESH_LOG(Display, TEXT("Resizing spline mesh scene textures (%u -> %u). Full update: %s"),
+						CurSize, NeededSize, bFullUpdate ? TEXT("Yes") : TEXT("No"));
+
 		PosTexture = GraphBuilder.CreateTexture(
 			FRDGTextureDesc::Create2D(
 				FIntPoint(NeededSize, NeededSize),
@@ -325,6 +359,21 @@ void FSplineMeshSceneResources::AddUpdatePass(
 		--GSplineMeshSceneTexturesCaptureNextUpdate;
 	}
 
+	const uint32 NumUpdateRequests = bFullUpdate ? 0 : UpdateRequests.Num();
+
+	if (!bForceUpdate)
+	{
+		if (bFullUpdate)
+		{
+			SPLINE_MESH_LOG(Verbose, TEXT("Performing full update of spline mesh scene textures."));
+		}
+		else
+		{
+			SPLINE_MESH_LOG(Verbose, TEXT("Performing partial update (%u) of spline mesh scene textures: [%s]."),
+							NumUpdateRequests, *GetIntListStringDebug(UpdateRequests));
+		}
+	}
+
 	FRDGTextureUAVRef PosTextureUAV = GraphBuilder.CreateUAV(PosTexture);
 	FRDGTextureUAVRef RotTextureUAV = GraphBuilder.CreateUAV(RotTexture);
 	if (bForceUpdate)
@@ -335,7 +384,6 @@ void FSplineMeshSceneResources::AddUpdatePass(
 	}
 
 	FRDGBufferRef UpdateRequestBuffer = nullptr;
-	const uint32 NumUpdateRequests = bFullUpdate ? 0 : UpdateRequests.Num();
 	if (NumUpdateRequests > 0)
 	{
 		// Update only select instances
@@ -384,6 +432,12 @@ void FSplineMeshSceneResources::Register(const FPrimitiveSceneInfo& PrimitiveSce
 
 	// Alloc space for the new splines and ensure they are included in the next update
 	AllocTextureSpace(PrimitiveSceneInfo, GetNumSplines(PrimitiveSceneInfo), Slot);
+
+
+	SPLINE_MESH_LOG(Verbose, TEXT("Registered spline mesh %u (Prim %u): %s"),
+					Slot.FirstSplineIndex, PrimitiveSceneInfo.GetIndex(),
+					*PrimitiveSceneInfo.GetFullnameForDebuggingOnly());
+
 	RequestUpdate(Slot);
 }
 
@@ -395,6 +449,10 @@ void FSplineMeshSceneResources::Unregister(const FPrimitiveSceneInfo& PrimitiveS
 		return;
 	}
 
+	SPLINE_MESH_LOG(Verbose, TEXT("Unregistered Spline Mesh %u (Prim %u): %s"),
+					Slot->FirstSplineIndex, PrimitiveSceneInfo.GetIndex(),
+					*PrimitiveSceneInfo.GetFullnameForDebuggingOnly());
+
 	SlotAllocator.Free(Slot->FirstSplineIndex, Slot->NumSplines);
 
 	// Clear the instance IDs of the newly freed range
@@ -474,7 +532,13 @@ void FSplineMeshSceneResources::AssignCoordinates(TSplineMeshSceneProxy* ScenePr
 {
 	for (uint32 i = 0; i < Slot.NumSplines; ++i)
 	{
-		SceneProxy->SetSplineTextureCoord_RenderThread(i, SplineMesh::CalcTextureCoord(Slot.FirstSplineIndex));
+		const uint32 SplineIndex = Slot.FirstSplineIndex + i;
+		const FUintVector2 Coord = SplineMesh::CalcTextureCoord(SplineIndex);
+
+		SPLINE_MESH_LOG(VeryVerbose, TEXT("Assigned Tex Coord <%u, %u> to spline mesh %u (Prim %u)"),
+						Coord.X, Coord.Y, SplineIndex, SceneProxy->GetPrimitiveSceneInfo()->GetIndex());
+
+		SceneProxy->SetSplineTextureCoord_RenderThread(i, Coord);
 	}
 }
 
@@ -485,6 +549,8 @@ void FSplineMeshSceneResources::RequestUpdate(const FPrimitiveSlot& Slot)
 		const uint32 SplineIndex = Slot.FirstSplineIndex + i;
 		check(RegisteredInstanceIds[SplineIndex] != INDEX_NONE); // Sanity check
 		UpdateRequests.AddUnique(SplineIndex);
+
+		SPLINE_MESH_LOG(VeryVerbose, TEXT("Requested update for spline mesh %u"), SplineIndex);
 	}
 }
 
