Fix a typo in mobile local fog rendering.

diff --git a/Engine/Source/Runtime/Renderer/Private/MobileFogRendering.cpp b/Engine/Source/Runtime/Renderer/Private/MobileFogRendering.cpp
--- a/Engine/Source/Runtime/Renderer/Private/MobileFogRendering.cpp
+++ b/Engine/Source/Runtime/Renderer/Private/MobileFogRendering.cpp
@@ -239,7 +239,7 @@ void FMobileSceneRenderer::RenderFog(FRHICommandList& RHICmdList, const FViewInf
 	FMobileFogPS::FParameters PSParameters;
 	PSParameters.View = View.GetShaderParameters();
 	PSParameters.LFV = View.LocalFogVolumeViewData.UniformParametersStruct;
-	SetShaderParameters(RHICmdList, PixelShader, PixelShader.GetVertexShader(), PSParameters);
+	SetShaderParameters(RHICmdList, PixelShader, PixelShader.GetPixelShader(), PSParameters);
 
 	// Draw a quad covering the view.
 	RHICmdList.SetStreamSource(0, GScreenSpaceVertexBuffer.VertexBufferRHI, 0);
