Implements TSR's anti-flickering visualizer.

#tests ps5, win64, xboxsx

diff --git a/Engine/Shaders/Private/TemporalSuperResolution/TSRVisualize.usf b/Engine/Shaders/Private/TemporalSuperResolution/TSRVisualize.usf
--- a/Engine/Shaders/Private/TemporalSuperResolution/TSRVisualize.usf
+++ b/Engine/Shaders/Private/TemporalSuperResolution/TSRVisualize.usf
@@ -11,10 +11,12 @@
 //------------------------------------------------------- CONSTS
 
 static tsr_half3 kWhite  = tsr_half3(0.33, 0.33, 0.33);
+static tsr_half3 kYellow = tsr_half3(1.00, 1.00, 0.10);
 static tsr_half3 kRed    = tsr_half3(1.00, 0.10, 0.10);
 static tsr_half3 kOrange = tsr_half3(1.00, 0.75, 0.10);
 static tsr_half3 kGreen  = tsr_half3(0.10, 1.00, 0.10);
 static tsr_half3 kBlue   = tsr_half3(0.10, 0.10, 1.00);
+static tsr_half3 kPink   = tsr_half3(1.00, 0.10, 0.50);
 
 
 //------------------------------------------------------- PARAMETERS
@@ -31,11 +33,13 @@ uint bCanResurrectHistory;
 uint bCanSpatialAntiAlias;
 float MaxHistorySampleCount;
 float OutputToHistoryResolutionFractionSquare;
+float FlickeringFramePeriod;
 
 Texture2D<tsr_half3> SceneColorTexture;
 Texture2D<float2>    ClosestDepthTexture;
 Texture2D<float2>    DilatedVelocityTexture;
 Texture2D<tsr_half4> HistoryRejectionTexture;
+Texture2D<tsr_half4> MoireHistoryTexture;
 Texture2D<tsr_ushort> AntiAliasMaskTexture;
 Texture2D<tsr_half4> HistoryMetadataTexture;
 Texture2D<tsr_half3> ResurrectedHistoryColorTexture;
@@ -117,13 +121,17 @@ tsr_half3 DisplayResurrectionMask(FContext Ctx)
 	bool ResurrectionMask = (tsr_ushort(HistoryRejectionTexture.SampleLevel(GlobalPointClampedSampler, Ctx.InputUV, 0).a * tsr_half(255.0)) & tsr_ushort(0x2)) != tsr_ushort(0x0);
 	if (!bCanResurrectHistory)
 	{
-		return SampleTSROutputGray(Ctx) * kRed;
+		return SampleTSROutputGray(Ctx) * kPink;
 	}
 	return SampleTSROutputGray(Ctx) * select(ResurrectionMask, lerp(kRed, kGreen, LowFrequencyRejection), kWhite);
 }
 
 tsr_half3 DisplayResurrectedColor(FContext Ctx)
 {
+	if (!bCanResurrectHistory)
+	{
+		return SampleTSROutputGray(Ctx) * kPink;
+	}
 	tsr_half LowFrequencyRejection = HistoryRejectionTexture.SampleLevel(GlobalPointClampedSampler, Ctx.InputUV, 0).r;
 	float ClosestDeviceZ = ClosestDepthTexture.SampleLevel(GlobalPointClampedSampler, Ctx.InputUV, 0).g;
 	bool ResurrectionMask = (tsr_ushort(HistoryRejectionTexture.SampleLevel(GlobalPointClampedSampler, Ctx.InputUV, 0).a * tsr_half(255.0)) & tsr_ushort(0x2)) != tsr_ushort(0x0);
@@ -145,11 +153,50 @@ tsr_half3 DisplaySpatialAntiAliasingMask(FContext Ctx)
 	bool bSpatialAntiAlias = AntiAliasMaskTexture[Ctx.InputPixelPos] != tsr_ushort(0x0);
 	if (!bCanSpatialAntiAlias)
 	{
-		return SampleTSROutputGray(Ctx) * kRed;
+		return SampleTSROutputGray(Ctx) * kPink;
 	}
 	return SampleTSROutputGray(Ctx) * select(bSpatialAntiAlias, kGreen, kWhite);
 }
 
+tsr_half3 DisplayAntiFlickering(FContext Ctx)
+{
+	if (FlickeringFramePeriod == 0.0)
+	{
+		return SampleTSROutputGray(Ctx) * kPink;
+	}
+
+	const tsr_half BlendFinal = tsr_half(0.05);
+	const tsr_half MoireEncodingError = tsr_half(rcp(127.0));
+	const tsr_half MaxPrevTotalVariationCount = tsr_half(20.0);
+	
+	// Compute the threshold at which the variation count should quick in
+	const tsr_half TotalVariationCountThreshold = tsr_half(1.0 / (1.0 - pow(1.0 - BlendFinal, tsr_half(FlickeringFramePeriod))));
+
+	tsr_half4 MoireHistory = MoireHistoryTexture.SampleLevel(GlobalPointClampedSampler, Ctx.InputUV, 0);
+
+
+	// Unpack history
+	tsr_half FlickeringHistory   = MoireHistory[0];
+	tsr_half Gradient            = MoireHistory[1] * tsr_half(255.0 / 127.0) - tsr_half(1.0);
+	tsr_half TotalVariation      = MoireHistory[2] * (tsr_half(1.0) - BlendFinal);
+	tsr_half TotalVariationCount = MoireHistory[3] * MaxPrevTotalVariationCount * (tsr_half(1.0) - BlendFinal);
+	
+	// Compute the final luminance 
+	#if 1
+		tsr_half LocalCountFadeIn = saturate(TotalVariationCount / TotalVariationCountThreshold - tsr_half(0.5));
+	#else
+		tsr_half LocalCountFadeIn = saturate(TotalVariationCount - TotalVariationCountThreshold);
+	#endif
+	tsr_half LocalMoireError = (abs(TotalVariation * SafeRcp(TotalVariationCount)) + TotalVariationCount * MoireEncodingError);
+	
+	tsr_half3 Color;
+	Color = lerp(kGreen, kYellow, saturate(LocalMoireError * 16.0));
+	Color = lerp(Color,  kRed,    saturate(LocalMoireError * 16.0 - 1.0));
+	Color = lerp(kWhite, Color,   LocalCountFadeIn);
+
+	return SampleTSROutputGray(Ctx) * Color;
+}
+
 
 //------------------------------------------------------- ENTRY POINT
 
@@ -204,6 +251,10 @@ void MainCS(
 	{
 		OutputColor = DisplaySpatialAntiAliasingMask(Ctx);
 	}
+	else if (VisualizeId == 7)
+	{
+		OutputColor = DisplayAntiFlickering(Ctx);
+	}
 
 	{
 		tsr_short2 OutputPixelPos = InvalidateOutputPixelPos(Ctx.OutputPixelPos, OutputViewRectMax);
diff --git a/Engine/Source/Runtime/Renderer/Private/PostProcess/TemporalSuperResolution.cpp b/Engine/Source/Runtime/Renderer/Private/PostProcess/TemporalSuperResolution.cpp
--- a/Engine/Source/Runtime/Renderer/Private/PostProcess/TemporalSuperResolution.cpp
+++ b/Engine/Source/Runtime/Renderer/Private/PostProcess/TemporalSuperResolution.cpp
@@ -281,7 +281,9 @@ TAutoConsoleVariable<int32> CVarTSRVisualize(
 	TEXT("  2: Mask where the history is rejected;\n")
 	TEXT("  3: Mask where the history is clamped;\n")
 	TEXT("  4: Mask where the history is resurrected (with r.TSR.Resurrection=1);\n")
-	TEXT("  5: Mask where the history is resurrected in the resurrected frame (with r.TSR.Resurrection=1), particularily interesting to tune r.TSR.Resurrection.PersistentFrameInterval;\n"),
+	TEXT("  5: Mask where the history is resurrected in the resurrected frame (with r.TSR.Resurrection=1), particularily interesting to tune r.TSR.Resurrection.PersistentFrameInterval;\n")
+	TEXT("  6: Mask where spatial anti-aliasing is being computed;\n")
+	TEXT("  7: Mask where the anti-flickering heuristic is taking effects (with r.TSR.ShadingRejection.Flickering=1);\n"),
 	ECVF_RenderThreadSafe);
 
 #endif
@@ -940,11 +942,13 @@ class FTSRVisualizeCS : public FTSRShader
 		SHADER_PARAMETER(int32, bCanSpatialAntiAlias)
 		SHADER_PARAMETER(float, MaxHistorySampleCount)
 		SHADER_PARAMETER(float, OutputToHistoryResolutionFractionSquare)
+		SHADER_PARAMETER(float, FlickeringFramePeriod)
 
 		SHADER_PARAMETER_RDG_TEXTURE_SRV(Texture2D, SceneColorTexture)
 		SHADER_PARAMETER_RDG_TEXTURE(Texture2D, ClosestDepthTexture)
 		SHADER_PARAMETER_RDG_TEXTURE(Texture2D, DilatedVelocityTexture)
 		SHADER_PARAMETER_RDG_TEXTURE(Texture2D, HistoryRejectionTexture)
+		SHADER_PARAMETER_RDG_TEXTURE_SRV(Texture2D, MoireHistoryTexture)
 		SHADER_PARAMETER_RDG_TEXTURE(Texture2D, AntiAliasMaskTexture)
 		SHADER_PARAMETER_RDG_TEXTURE_SRV(Texture2D, HistoryMetadataTexture)
 		SHADER_PARAMETER_RDG_TEXTURE_SRV(Texture2D, ResurrectedHistoryColorTexture)
@@ -1304,7 +1308,8 @@ FDefaultTemporalUpscaler::FOutputs AddTemporalSuperResolutionPasses(
 
 	static auto CVarAntiAliasingQuality = IConsoleManager::Get().FindConsoleVariable(TEXT("sg.AntiAliasingQuality"));
 	check(CVarAntiAliasingQuality);
-
+	
+	if (IsVisualizeTSREnabled(View))
 	RDG_EVENT_SCOPE(GraphBuilder, "TemporalSuperResolution(sg.AntiAliasingQuality=%d%s) %dx%d -> %dx%d",
 		CVarAntiAliasingQuality->GetInt(),
 		bSupportsAlpha ? TEXT(" Alpha") : TEXT(""),
@@ -1943,6 +1948,7 @@ FDefaultTemporalUpscaler::FOutputs AddTemporalSuperResolutionPasses(
 	FRDGTextureRef HistoryRejectionTexture = nullptr;
 	FRDGTextureRef InputSceneColorLdrLumaTexture = nullptr;
 	FRDGTextureRef AntiAliasMaskTexture = nullptr;
+	FRDGTextureSRVRef MoireHistoryTexture = nullptr;
 	{
 		const bool bComputeInputSceneColorTexture = InputSceneColorLdrLumaTexture == nullptr;
 		if (bComputeInputSceneColorTexture)
@@ -2072,7 +2078,18 @@ FDefaultTemporalUpscaler::FOutputs AddTemporalSuperResolutionPasses(
 			}
 			else if (View.bStatePrevViewInfoIsReadOnly)
 			{
-				PassParameters->HistoryMoireOutput = CreateDummyUAV(GraphBuilder, History.MoireArray->Desc.Format);
+				FRDGTextureDesc Desc = FRDGTextureDesc::Create2D(
+					InputExtent,
+					History.MoireArray->Desc.Format,
+					FClearValueBinding::None,
+					/* InFlags = */ TexCreate_ShaderResource | TexCreate_UAV);
+
+				// Create an unused texture for the moire history so that the VisualizeTSR can still display the updated moire history.
+				FRDGTextureRef UnusedMoireHistoryTexture = GraphBuilder.CreateTexture(Desc, TEXT("TSR.History.Moire"));
+				GraphBuilder.RemoveUnusedTextureWarning(UnusedMoireHistoryTexture);
+
+				PassParameters->HistoryMoireOutput = GraphBuilder.CreateUAV(UnusedMoireHistoryTexture);
+				MoireHistoryTexture = GraphBuilder.CreateSRV(FRDGTextureSRVDesc(UnusedMoireHistoryTexture));
 			}
 			else
 			{
@@ -2082,6 +2099,8 @@ FDefaultTemporalUpscaler::FOutputs AddTemporalSuperResolutionPasses(
 				MoireUAVDesc.DimensionOverride = ETextureDimension::Texture2D;
 
 				PassParameters->HistoryMoireOutput = GraphBuilder.CreateUAV(MoireUAVDesc);
+
+				MoireHistoryTexture = GraphBuilder.CreateSRV(FRDGTextureSRVDesc::CreateForSlice(History.MoireArray, CurrentFrameSliceIndex));
 			}
 
 			// Output how the history should rejected in the HistoryUpdate
@@ -2471,6 +2490,7 @@ FDefaultTemporalUpscaler::FOutputs AddTemporalSuperResolutionPasses(
 			ResurrectionMask = 4,
 			ResurrectedColor = 5,
 			SpatialAntiAliasingMask = 6,
+			AntiFlickering = 7,
 			MAX,
 		};
 
@@ -2482,6 +2502,7 @@ FDefaultTemporalUpscaler::FOutputs AddTemporalSuperResolutionPasses(
 			TEXT("ResurrectionMask"),
 			TEXT("ResurrectedColor"),
 			TEXT("SpatialAntiAliasingMask"),
+			TEXT("AntiFlickering"),
 		};
 		static_assert(UE_ARRAY_COUNT(kVisualizationName) == int32(EVisualizeId::MAX), "kVisualizationName doesn't match EVisualizeId");
 
@@ -2515,11 +2536,13 @@ FDefaultTemporalUpscaler::FOutputs AddTemporalSuperResolutionPasses(
 			PassParameters->bCanSpatialAntiAlias = RejectionAntiAliasingQuality > 0;
 			PassParameters->MaxHistorySampleCount = MaxHistorySampleCount;
 			PassParameters->OutputToHistoryResolutionFractionSquare = OutputToHistoryResolutionFractionSquare;
+			PassParameters->FlickeringFramePeriod = FlickeringFramePeriod;
 
 			PassParameters->SceneColorTexture = SceneColorOutputTextureSRV;
 			PassParameters->ClosestDepthTexture = ClosestDepthTexture;
 			PassParameters->DilatedVelocityTexture = DilatedVelocityTexture;
 			PassParameters->HistoryRejectionTexture = HistoryRejectionTexture;
+			PassParameters->MoireHistoryTexture = MoireHistoryTexture ? MoireHistoryTexture : GraphBuilder.CreateSRV(FRDGTextureSRVDesc(BlackDummy));
 			PassParameters->AntiAliasMaskTexture = AntiAliasMaskTexture ? AntiAliasMaskTexture : BlackUintDummy;
 			PassParameters->HistoryMetadataTexture = GraphBuilder.CreateSRV(FRDGTextureSRVDesc::CreateForSlice(History.MetadataArray, CurrentFrameSliceIndex));
 			if (PrevHistory.ColorArray == BlackArrayDummy)
@@ -2565,6 +2588,7 @@ FDefaultTemporalUpscaler::FOutputs AddTemporalSuperResolutionPasses(
 					Tiles[4 * 2 + 0] = Visualize(EVisualizeId::ResurrectedColor, TEXT("Resurrected Frame"));
 				}
 				Tiles[4 * 3 + 0] = Visualize(EVisualizeId::SpatialAntiAliasingMask, TEXT("Spatial Anti-Aliasing"));
+				Tiles[4 * 1 + 3] = Visualize(EVisualizeId::AntiFlickering, TEXT("Flickering Temporal Analysis"));
 			}
 
 			{
