LFV - unified lighting to be closer to volumetric fog.

#prelfight https://horde.devtools.epicgames.com/job/65106c4bbb77e258eb347583

diff --git a/Engine/Shaders/Private/LocalFogVolumes/LocalFogVolumeCommon.ush b/Engine/Shaders/Private/LocalFogVolumes/LocalFogVolumeCommon.ush
--- a/Engine/Shaders/Private/LocalFogVolumes/LocalFogVolumeCommon.ush
+++ b/Engine/Shaders/Private/LocalFogVolumes/LocalFogVolumeCommon.ush
@@ -1,7 +1,9 @@
 // Copyright Epic Games, Inc. All Rights Reserved.
 #include "../Common.ush"
 
+#include "../SHCommon.ush"
 #include "../ParticipatingMediaCommon.ush"
+#include "../ReflectionEnvironmentShared.ush"
 
 
 
@@ -39,6 +41,9 @@ struct FLocalFogVolumeGPUInstanceData
 #define LFVMaxDensityIntoVolumetricFog	LFVStruct.LFV.LocalFogVolumeMaxDensityIntoVolumetricFog
 #define LFVRenderInVolumetricFog		LFVStruct.LFV.ShouldRenderLocalFogVolumeInVolumetricFog
 
+#define LFVDirectionalLightColor		LFVStruct.LFV.DirectionalLightColor
+#define LFVDirectionalLightDirection	LFVStruct.LFV.DirectionalLightDirection
+
 #else
 
 uint2									LFV_LocalFogVolumeTileDataTextureResolution;
@@ -47,6 +52,9 @@ uint									LFV_LocalFogVolumeTilePixelSize;
 float									LFV_LocalFogVolumeMaxDensityIntoVolumetricFog;
 uint									LFV_ShouldRenderLocalFogVolumeInVolumetricFog;
 
+float3									LFV_DirectionalLightColor;
+float3									LFV_DirectionalLightDirection;
+
 Buffer<float4>							LFV_LocalFogVolumeInstances;
 Texture2DArray<uint>					LFV_LocalFogVolumeTileDataTexture;
 
@@ -59,6 +67,9 @@ Texture2DArray<uint>					LFV_LocalFogVolumeTileDataTexture;
 #define LFVMaxDensityIntoVolumetricFog	LFV_LocalFogVolumeMaxDensityIntoVolumetricFog
 #define LFVRenderInVolumetricFog		LFV_ShouldRenderLocalFogVolumeInVolumetricFog
 
+#define LFVDirectionalLightColor		LFV_DirectionalLightColor
+#define LFVDirectionalLightDirection	LFV_DirectionalLightDirection
+
 #endif
 
 
@@ -227,29 +238,24 @@ LocalFogVolumeCommonLightingInput LocalFogVolumeEvaluateCommonLightingInput()
 
 float3 LocalFogVolumeEvaluateInScattering(in FLocalFogVolumeGPUInstanceData FogInstance, in FFogData FogData, in LocalFogVolumeCommonLightingInput CommonLightingInput, in float3 RayDirWorld)
 {
-	half3 InScattering = 0;
+	half3 InScattering = 0.0f;
 
-#if PROJECT_SUPPORT_SKY_ATMOSPHERE_AFFECTS_HEIGHFOG
+	// Lighting equation here matches the volumetric fog.
 
-	// Accumulate distant sky light
-	InScattering += CommonLightingInput.DistantSkyLight;
+	// Main forward directional light intensity
+	InScattering += LFVDirectionalLightColor * HenyeyGreensteinPhase(-FogInstance.PhaseG, dot(RayDirWorld, LFVDirectionalLightDirection));
 
-	// Accumulate directional lights if needed
-	if (View.AtmosphereLightIlluminanceOnGroundPostTransmittance[0].a > 0.0f)
-	{
-		InScattering += View.AtmosphereLightIlluminanceOnGroundPostTransmittance[0].rgb * HenyeyGreensteinPhase(-FogInstance.PhaseG, dot(RayDirWorld, View.AtmosphereLightDirection[0].xyz));
-	}
-	if (View.AtmosphereLightIlluminanceOnGroundPostTransmittance[1].a > 0.0f)
+	// Sky light
+	if (View.SkyLightVolumetricScatteringIntensity > 0)
 	{
-		InScattering += View.AtmosphereLightIlluminanceOnGroundPostTransmittance[0].rgb * HenyeyGreensteinPhase(-FogInstance.PhaseG, dot(RayDirWorld, View.AtmosphereLightDirection[1].xyz));
-	}
-
-	// Weight by user control
-	InScattering *= View.SkyAtmosphereHeightFogContribution.xxx;
+		float3 SkyLighting = View.SkyLightColor.rgb * GetSkySHDiffuseSimple(RayDirWorld * -FogInstance.PhaseG);
 
+		float SkyVisibility = 1.0f; // Not possible to reach with analytical integration
+		InScattering += (SkyVisibility * View.SkyLightVolumetricScatteringIntensity) * SkyLighting;
+	}
+	
 	// Apply the medium uniform single scattering albedo
 	InScattering *= FogInstance.Albedo;
-#endif
 
 	// Now account for emissive luminance after albedo and heighfog contribution has been accounted for.
 	// This do mean that the emissive color follows the medium density variation.
diff --git a/Engine/Shaders/Private/VolumetricFog.usf b/Engine/Shaders/Private/VolumetricFog.usf
--- a/Engine/Shaders/Private/VolumetricFog.usf
+++ b/Engine/Shaders/Private/VolumetricFog.usf
@@ -655,7 +655,6 @@ float GetDirectionalLightFunction(float3 TranslatedWorldPosition) // directional
 	return Texture2DSampleLevel(LightFunctionTexture, LightFunctionSampler, LightFunctionUV, 0).x;
 }
 
-float SkyLightVolumetricScatteringIntensity;
 float4 SkySH[3];
 float2 UseHeightFogColors;	// x=override directional light using height fog inscattering color, y=override sky light using heigh fog inscattering cubemap
 float UseDirectionalLightShadowing;
@@ -779,7 +778,7 @@ void LightScatteringCS(
 		float3 BrickTextureUVs = 0;
 
 #if ALLOW_STATIC_LIGHTING
-		if (SkyLightVolumetricScatteringIntensity > 0 || StaticLightingScatteringIntensity > 0)
+		if (View.SkyLightVolumetricScatteringIntensity > 0 || StaticLightingScatteringIntensity > 0)
 		{
 			BrickTextureUVs = ComputeVolumetricLightmapBrickTextureUVs(WorldPosition);
 		}
@@ -794,7 +793,7 @@ void LightScatteringCS(
 
 #else
 		// Skylight
-		if (SkyLightVolumetricScatteringIntensity > 0)
+		if (View.SkyLightVolumetricScatteringIntensity > 0)
 		{
 			float3 SkyLighting;
 
@@ -820,7 +819,7 @@ void LightScatteringCS(
 			}
 
 			float SkyVisibility = ComputeSkyVisibility(TranslatedWorldPosition, BrickTextureUVs);
-			LightScattering += (SkyVisibility * SkyLightVolumetricScatteringIntensity) * SkyLighting;
+			LightScattering += (SkyVisibility * View.SkyLightVolumetricScatteringIntensity) * SkyLighting;
 		}
 #endif
 
diff --git a/Engine/Source/Runtime/Engine/Public/SceneView.h b/Engine/Source/Runtime/Engine/Public/SceneView.h
--- a/Engine/Source/Runtime/Engine/Public/SceneView.h
+++ b/Engine/Source/Runtime/Engine/Public/SceneView.h
@@ -853,6 +853,7 @@ enum ETranslucencyVolumeCascade
 	VIEW_UNIFORM_BUFFER_MEMBER(float, SkyLightAffectReflectionFlag) \
 	VIEW_UNIFORM_BUFFER_MEMBER(float, SkyLightAffectGlobalIlluminationFlag) \
 	VIEW_UNIFORM_BUFFER_MEMBER(FLinearColor, SkyLightColor) \
+	VIEW_UNIFORM_BUFFER_MEMBER(float, SkyLightVolumetricScatteringIntensity) \
 	VIEW_UNIFORM_BUFFER_MEMBER_ARRAY(FVector4f, MobileSkyIrradianceEnvironmentMap, SKY_IRRADIANCE_ENVIRONMENT_MAP_VEC4_COUNT) \
 	VIEW_UNIFORM_BUFFER_MEMBER(float, MobilePreviewMode) \
 	VIEW_UNIFORM_BUFFER_MEMBER_PER_VIEW(float, HMDEyePaddingOffset) \
diff --git a/Engine/Source/Runtime/Renderer/Private/DeferredShadingRenderer.cpp b/Engine/Source/Runtime/Renderer/Private/DeferredShadingRenderer.cpp
--- a/Engine/Source/Runtime/Renderer/Private/DeferredShadingRenderer.cpp
+++ b/Engine/Source/Runtime/Renderer/Private/DeferredShadingRenderer.cpp
@@ -3175,7 +3175,8 @@ void FDeferredShadingSceneRenderer::Render(FRDGBuilder& GraphBuilder)
 				ViewFamily.ViewMode != VMI_Lit ||
 				bAnyLumenEnabled ||
 				VirtualShadowMapArray.IsEnabled() ||
-				ShouldVisualizeLightGrid());
+				ShouldVisualizeLightGrid() ||
+				ShouldRenderLocalFogVolume(Scene, ViewFamily)); // Needed when accessing forward light data for the directional light
 			bComputeLightGrid &= !ViewFamily.EngineShowFlags.PathTracing;
 		}
 	}
diff --git a/Engine/Source/Runtime/Renderer/Private/LightGridInjection.cpp b/Engine/Source/Runtime/Renderer/Private/LightGridInjection.cpp
--- a/Engine/Source/Runtime/Renderer/Private/LightGridInjection.cpp
+++ b/Engine/Source/Runtime/Renderer/Private/LightGridInjection.cpp
@@ -939,7 +939,8 @@ FComputeLightGridOutput FDeferredShadingSceneRenderer::GatherLightsAndComputeLig
 	for (int32 ViewIndex = 0; ViewIndex < Views.Num(); ViewIndex++)
 	{
 		const FViewInfo& View = Views[ViewIndex];
-		bAnyViewUsesForwardLighting |= View.bTranslucentSurfaceLighting || ShouldRenderVolumetricFog() || View.bHasSingleLayerWaterMaterial || VolumetricCloudWantsToSampleLocalLights(Scene, ViewFamily.EngineShowFlags) || ShouldVisualizeLightGrid();
+		bAnyViewUsesForwardLighting |= View.bTranslucentSurfaceLighting || ShouldRenderVolumetricFog() || View.bHasSingleLayerWaterMaterial 
+			|| VolumetricCloudWantsToSampleLocalLights(Scene, ViewFamily.EngineShowFlags) || ShouldVisualizeLightGrid() || ShouldRenderLocalFogVolume(Scene, ViewFamily);
 		bAnyViewUsesLumen |= GetViewPipelineState(View).DiffuseIndirectMethod == EDiffuseIndirectMethod::Lumen || GetViewPipelineState(View).ReflectionsMethod == EReflectionsMethod::Lumen;
 	}
 	
diff --git a/Engine/Source/Runtime/Renderer/Private/LocalFogVolumeRendering.cpp b/Engine/Source/Runtime/Renderer/Private/LocalFogVolumeRendering.cpp
--- a/Engine/Source/Runtime/Renderer/Private/LocalFogVolumeRendering.cpp
+++ b/Engine/Source/Runtime/Renderer/Private/LocalFogVolumeRendering.cpp
@@ -15,7 +15,7 @@ static TAutoConsoleVariable<int32> CVarLocalFogVolume(
 	ECVF_RenderThreadSafe);
 
 static TAutoConsoleVariable<int32> CVarLocalFogVolumeRenderIntoVolumetricFog(
-	TEXT("r.LocalFogVolume.RenderIntoVolumetricFog"), 0,
+	TEXT("r.LocalFogVolume.RenderIntoVolumetricFog"), 1,
 	TEXT("LocalFogVolume are going to be voxelised into the volumetric fog when this is not 0, otherwise it will remain isolated.\n"),
 	ECVF_RenderThreadSafe);
 
@@ -124,7 +124,10 @@ void SetDummyLocalFogVolumeForView(FRDGBuilder& GraphBuilder, FViewInfo& View)
 	View.LocalFogVolumeViewData.UniformParametersStruct.LocalFogVolumeCommon.LocalFogVolumeInstanceCount				= View.LocalFogVolumeViewData.GPUInstanceCount;
 	View.LocalFogVolumeViewData.UniformParametersStruct.LocalFogVolumeCommon.LocalFogVolumeTilePixelSize				= GetLocalFogVolumeTilePixelSize();
 	View.LocalFogVolumeViewData.UniformParametersStruct.LocalFogVolumeCommon.LocalFogVolumeMaxDensityIntoVolumetricFog	= GetLocalFogVolumeMaxDensityIntoVolumetricFog();
+	View.LocalFogVolumeViewData.UniformParametersStruct.LocalFogVolumeCommon.ShouldRenderLocalFogVolumeInVolumetricFog	= 0;
 	View.LocalFogVolumeViewData.UniformParametersStruct.LocalFogVolumeCommon.LocalFogVolumeInstances					= View.LocalFogVolumeViewData.GPUInstanceDataBufferSRV;
+	View.LocalFogVolumeViewData.UniformParametersStruct.LocalFogVolumeCommon.DirectionalLightColor						= FVector3f::Zero();
+	View.LocalFogVolumeViewData.UniformParametersStruct.LocalFogVolumeCommon.DirectionalLightDirection					= FVector3f::Zero();
 	View.LocalFogVolumeViewData.UniformParametersStruct.LocalFogVolumeTileDataTexture									= View.LocalFogVolumeViewData.TileDataTextureArraySRV;
 	View.LocalFogVolumeViewData.UniformBuffer																			= GraphBuilder.CreateUniformBuffer(&View.LocalFogVolumeViewData.UniformParametersStruct);
 
@@ -321,7 +324,7 @@ void GetLocalFogVolumeSortingData(const FScene* Scene, FRDGBuilder& GraphBuilder
 	Out.LocalFogVolumeSortKeys.SetNum(Out.LocalFogVolumeInstanceCountFinal, false/*bAllowShrinking*/);
 }
 
-void CreateViewLocalFogVolumeBufferSRV(FViewInfo& View, FRDGBuilder& GraphBuilder, FLocalFogVolumeSortingData& SortingData, bool bShouldRenderLocalFogVolumeInVolumetricFog)
+void CreateViewLocalFogVolumeBufferSRV(const FScene* Scene, FViewInfo& View, FRDGBuilder& GraphBuilder, FLocalFogVolumeSortingData& SortingData, bool bShouldRenderLocalFogVolumeInVolumetricFog)
 {
 	if (SortingData.LocalFogVolumeInstanceCountFinal == 0)
 	{
@@ -392,6 +395,34 @@ void CreateViewLocalFogVolumeBufferSRV(FViewInfo& View, FRDGBuilder& GraphBuilde
 	View.LocalFogVolumeViewData.UniformParametersStruct.LocalFogVolumeCommon.LocalFogVolumeMaxDensityIntoVolumetricFog	= GetLocalFogVolumeMaxDensityIntoVolumetricFog();
 	View.LocalFogVolumeViewData.UniformParametersStruct.LocalFogVolumeCommon.ShouldRenderLocalFogVolumeInVolumetricFog	= bShouldRenderLocalFogVolumeInVolumetricFog ? 1 : 0;
 	View.LocalFogVolumeViewData.UniformParametersStruct.LocalFogVolumeCommon.LocalFogVolumeInstances					= View.LocalFogVolumeViewData.GPUInstanceDataBufferSRV;
+	if (IsMobilePlatform(View.GetShaderPlatform()))
+	{
+		// On mobile there is a separate FMobileDirectionalLightShaderParameters UB which holds all directional light data.
+		// See SetupMobileDirectionalLightUniformParameters
+
+		View.LocalFogVolumeViewData.UniformParametersStruct.LocalFogVolumeCommon.DirectionalLightColor					= FVector3f::Zero();
+		View.LocalFogVolumeViewData.UniformParametersStruct.LocalFogVolumeCommon.DirectionalLightDirection				= FVector3f::Zero();
+		for (uint32 ChannelIdx = 0; ChannelIdx < UE_ARRAY_COUNT(Scene->MobileDirectionalLights); ChannelIdx++)
+		{
+			FLightSceneInfo* Light = Scene->MobileDirectionalLights[ChannelIdx];
+			if (Light != nullptr)
+			{
+				View.LocalFogVolumeViewData.UniformParametersStruct.LocalFogVolumeCommon.DirectionalLightColor			= FVector3f( Light->Proxy->GetSunIlluminanceAccountingForSkyAtmospherePerPixelTransmittance());
+				View.LocalFogVolumeViewData.UniformParametersStruct.LocalFogVolumeCommon.DirectionalLightDirection		= FVector3f(-Light->Proxy->GetDirection());
+			}
+		}
+	}
+	else if (View.ForwardLightingResources.SelectedForwardDirectionalLightProxy)
+	{
+		const FLightSceneProxy* SelectedForwardDirectionalLightProxy = View.ForwardLightingResources.SelectedForwardDirectionalLightProxy;
+		View.LocalFogVolumeViewData.UniformParametersStruct.LocalFogVolumeCommon.DirectionalLightColor					= FVector3f( SelectedForwardDirectionalLightProxy->GetSunIlluminanceAccountingForSkyAtmospherePerPixelTransmittance());
+		View.LocalFogVolumeViewData.UniformParametersStruct.LocalFogVolumeCommon.DirectionalLightDirection				= FVector3f(-SelectedForwardDirectionalLightProxy->GetDirection());
+	}
+	else
+	{
+		View.LocalFogVolumeViewData.UniformParametersStruct.LocalFogVolumeCommon.DirectionalLightColor					= FVector3f(View.CachedViewUniformShaderParameters->DirectionalLightColor);
+		View.LocalFogVolumeViewData.UniformParametersStruct.LocalFogVolumeCommon.DirectionalLightDirection				= View.CachedViewUniformShaderParameters->DirectionalLightDirection;
+	}
 	View.LocalFogVolumeViewData.UniformParametersStruct.LocalFogVolumeTileDataTexture									= View.LocalFogVolumeViewData.TileDataTextureArraySRV;
 	View.LocalFogVolumeViewData.UniformBuffer																			= GraphBuilder.CreateUniformBuffer(&View.LocalFogVolumeViewData.UniformParametersStruct);
 
@@ -425,7 +456,7 @@ void InitLocalFogVolumesForViews(
 
 		for (FViewInfo& View : Views)
 		{
-			CreateViewLocalFogVolumeBufferSRV(View, GraphBuilder, SortingData, ShouldRenderLocalFogVolumeInVolumetricFog(Scene, SceneViewFamily, bShouldRenderVolumetricFog));
+			CreateViewLocalFogVolumeBufferSRV(Scene, View, GraphBuilder, SortingData, ShouldRenderLocalFogVolumeInVolumetricFog(Scene, SceneViewFamily, bShouldRenderVolumetricFog));
 
 			LocalFogVolumeViewTiledCullingPass(View, GraphBuilder);
 		}
diff --git a/Engine/Source/Runtime/Renderer/Private/LocalFogVolumeRendering.h b/Engine/Source/Runtime/Renderer/Private/LocalFogVolumeRendering.h
--- a/Engine/Source/Runtime/Renderer/Private/LocalFogVolumeRendering.h
+++ b/Engine/Source/Runtime/Renderer/Private/LocalFogVolumeRendering.h
@@ -26,6 +26,8 @@ BEGIN_SHADER_PARAMETER_STRUCT(FLocalFogVolumeCommonParameters, )
 	SHADER_PARAMETER(uint32, LocalFogVolumeTilePixelSize)
 	SHADER_PARAMETER(float,  LocalFogVolumeMaxDensityIntoVolumetricFog)
 	SHADER_PARAMETER(uint32, ShouldRenderLocalFogVolumeInVolumetricFog)
+	SHADER_PARAMETER(FVector3f, DirectionalLightColor)
+	SHADER_PARAMETER(FVector3f, DirectionalLightDirection)
 END_SHADER_PARAMETER_STRUCT()
 
 BEGIN_SHADER_PARAMETER_STRUCT(FLocalFogVolumeUniformParameters, )
diff --git a/Engine/Source/Runtime/Renderer/Private/SceneRendering.cpp b/Engine/Source/Runtime/Renderer/Private/SceneRendering.cpp
--- a/Engine/Source/Runtime/Renderer/Private/SceneRendering.cpp
+++ b/Engine/Source/Runtime/Renderer/Private/SceneRendering.cpp
@@ -1761,6 +1761,7 @@ void FViewInfo::SetupUniformBufferParameters(
 		ViewUniformShaderParameters.SkyLightApplyPrecomputedBentNormalShadowingFlag = bApplyPrecomputedBentNormalShadowing ? 1.0f : 0.0f;
 		ViewUniformShaderParameters.SkyLightAffectReflectionFlag = SkyLight->bAffectReflection ? 1.0f : 0.0f;
 		ViewUniformShaderParameters.SkyLightAffectGlobalIlluminationFlag = SkyLight->bAffectGlobalIllumination ? 1.0f : 0.0f;
+		ViewUniformShaderParameters.SkyLightVolumetricScatteringIntensity = SkyLight->VolumetricScatteringIntensity;
 	}
 	else
 	{
@@ -1768,6 +1769,7 @@ void FViewInfo::SetupUniformBufferParameters(
 		ViewUniformShaderParameters.SkyLightApplyPrecomputedBentNormalShadowingFlag = 0.0f;
 		ViewUniformShaderParameters.SkyLightAffectReflectionFlag = 0.0f;
 		ViewUniformShaderParameters.SkyLightAffectGlobalIlluminationFlag = 0.0f;
+		ViewUniformShaderParameters.SkyLightVolumetricScatteringIntensity = 0.0f;
 	}
 
 	if (RHIFeatureLevel == ERHIFeatureLevel::ES3_1)
diff --git a/Engine/Source/Runtime/Renderer/Private/VolumetricFog.cpp b/Engine/Source/Runtime/Renderer/Private/VolumetricFog.cpp
--- a/Engine/Source/Runtime/Renderer/Private/VolumetricFog.cpp
+++ b/Engine/Source/Runtime/Renderer/Private/VolumetricFog.cpp
@@ -992,7 +992,6 @@ class FVolumetricFogLightScatteringCS : public FGlobalShader
 		SHADER_PARAMETER(FVector2f, UseHeightFogColors)
 		SHADER_PARAMETER(FVector2f, LightScatteringHistoryPreExposureAndInv)
 		SHADER_PARAMETER(float, StaticLightingScatteringIntensity)
-		SHADER_PARAMETER(float, SkyLightVolumetricScatteringIntensity)
 		SHADER_PARAMETER(float, SkyLightUseStaticShadowing)
 		SHADER_PARAMETER(float, PhaseG)
 		SHADER_PARAMETER(float, InverseSquaredLightDistanceBiasScale)
@@ -1564,7 +1563,6 @@ void FSceneRenderer::ComputeVolumetricFog(FRDGBuilder& GraphBuilder,
 				&& View.Family->EngineShowFlags.SkyLighting)
 			{
 				PassParameters->SkyLightUseStaticShadowing = SkyLight->bWantsStaticShadowing && SkyLight->bCastShadows ? 1.0f : 0.0f;
-				PassParameters->SkyLightVolumetricScatteringIntensity = SkyLight->VolumetricScatteringIntensity;
 
 				const FSHVectorRGB3& SkyIrradiance = SkyLight->IrradianceEnvironmentMap;
 				PassParameters->SkySH[0] = (FVector4f&)SkyIrradiance.R.V;
@@ -1574,7 +1572,6 @@ void FSceneRenderer::ComputeVolumetricFog(FRDGBuilder& GraphBuilder,
 			else
 			{
 				PassParameters->SkyLightUseStaticShadowing = 0.0f;
-				PassParameters->SkyLightVolumetricScatteringIntensity = 0.0f;
 				PassParameters->SkySH[0] = FVector4f(0, 0, 0, 0);
 				PassParameters->SkySH[1] = FVector4f(0, 0, 0, 0);
 				PassParameters->SkySH[2] = FVector4f(0, 0, 0, 0);
