Fix Substrate max BSDF count tracking based on view relevance.

Substrate view data was built before the view relevance computation was done, causing a race condition, leading to flickering.

Fix missing Substrate constant buffer binding.

diff --git a/Engine/Source/Runtime/Renderer/Private/DeferredShadingRenderer.cpp b/Engine/Source/Runtime/Renderer/Private/DeferredShadingRenderer.cpp
--- a/Engine/Source/Runtime/Renderer/Private/DeferredShadingRenderer.cpp
+++ b/Engine/Source/Runtime/Renderer/Private/DeferredShadingRenderer.cpp
@@ -3204,10 +3204,6 @@ void FDeferredShadingSceneRenderer::Render(FRDGBuilder& GraphBuilder)
 		Scene->SplineMeshSceneResources->Update(GraphBuilder, GetSceneUniforms());
 	}
 
-	// Substrate initialisation is always run even when not enabled.
-	const bool bSubstrateEnabled = Substrate::IsSubstrateEnabled();
-	Substrate::InitialiseSubstrateFrameSceneData(GraphBuilder, *this);
-
 	if (DepthPass.IsComputeStencilDitherEnabled())
 	{
 		AddDitheredStencilFillPass(GraphBuilder, Views, SceneTextures.Depth.Target, DepthPass);
@@ -3238,6 +3234,11 @@ void FDeferredShadingSceneRenderer::Render(FRDGBuilder& GraphBuilder)
 		DynamicReadBufferForInitViews.Commit(GraphBuilder.RHICmdList);
 	}
 
+	// Substrate initialisation is always run even when not enabled.
+	// Need to run after EndInitViews() to ensure ViewRelevance computation are completed
+	const bool bSubstrateEnabled = Substrate::IsSubstrateEnabled();
+	Substrate::InitialiseSubstrateFrameSceneData(GraphBuilder, *this);
+
 	UE::SVT::GetStreamingManager().EndAsyncUpdate(GraphBuilder);
 
 	FHairStrandsBookmarkParameters& HairStrandsBookmarkParameters = *GraphBuilder.AllocObject<FHairStrandsBookmarkParameters>();
