LFV - fixed crash on volumetric cloud.

#prelfight none

diff --git a/Engine/Source/Runtime/Renderer/Private/VolumetricCloudRendering.cpp b/Engine/Source/Runtime/Renderer/Private/VolumetricCloudRendering.cpp
--- a/Engine/Source/Runtime/Renderer/Private/VolumetricCloudRendering.cpp
+++ b/Engine/Source/Runtime/Renderer/Private/VolumetricCloudRendering.cpp
@@ -657,6 +657,8 @@ void SetupDefaultRenderVolumetricCloudGlobalParameters(FRDGBuilder& GraphBuilder
 	VolumetricCloudParams.EnableHeightFog = ViewInfo.Family->Scene->HasAnyExponentialHeightFog() && ShouldRenderFog(*ViewInfo.Family);
 	SetupFogUniformParameters(GraphBuilder, ViewInfo, VolumetricCloudParams.FogStruct);
 
+	VolumetricCloudParams.LFV = ViewInfo.LocalFogVolumeViewData.UniformParametersStruct;
+
 	VolumetricCloudParams.ForwardLightData = *ViewInfo.ForwardLightingResources.ForwardLightData;
 	VolumetricCloudParams.LocalLightsShadowSampleCount = FMath::Clamp(CVarVolumetricCloudLocalLightsShadowSampleCount.GetValueOnRenderThread(), 0.0f, 128.0f);
 
@@ -2167,8 +2169,6 @@ static TRDGUniformBufferRef<FRenderVolumetricCloudGlobalParameters> CreateCloudP
 	VolumetricCloudParams.AerialPerspectiveMieOnlyStartDistanceKm = CloudInfo.GetVolumetricCloudSceneProxy().AerialPespectiveMieScatteringStartDistance;
 	VolumetricCloudParams.AerialPerspectiveMieOnlyFadeDistanceKmInv = 1.0f / FMath::Max(CloudInfo.GetVolumetricCloudSceneProxy().AerialPespectiveMieScatteringFadeDistance, 0.00001f);// One centimeter minimum
 
-	VolumetricCloudParams.LFV = MainView.LocalFogVolumeViewData.UniformParametersStruct;
-
 	const FRDGTextureDesc& Desc = CloudRC.RenderTargets.Output[0].GetTexture()->Desc;
 	VolumetricCloudParams.OutputSizeInvSize = FVector4f(Desc.Extent.X, Desc.Extent.Y, 1.0f / Desc.Extent.X, 1.0f / Desc.Extent.Y);
 
