LFV - fixed crash when scene has LFV but all are not visible dueto density=0.

#prelfight none

diff --git a/Engine/Source/Runtime/Renderer/Private/LocalFogVolumeRendering.cpp b/Engine/Source/Runtime/Renderer/Private/LocalFogVolumeRendering.cpp
--- a/Engine/Source/Runtime/Renderer/Private/LocalFogVolumeRendering.cpp
+++ b/Engine/Source/Runtime/Renderer/Private/LocalFogVolumeRendering.cpp
@@ -216,6 +216,7 @@ IMPLEMENT_GLOBAL_SHADER(FLocalFogVolumeTiledCullingCS, "/Engine/Private/LocalFog
 
 static void LocalFogVolumeViewTiledCullingPass(FViewInfo& View, FRDGBuilder& GraphBuilder)
 {
+	check(View.LocalFogVolumeViewData.GPUInstanceCount > 0);
 	FIntVector TileDataTextureSize = View.LocalFogVolumeViewData.TileDataTextureArray->Desc.GetSize();
 
 	FLocalFogVolumeTiledCullingCS::FParameters* PassParameters = GraphBuilder.AllocParameters<FLocalFogVolumeTiledCullingCS::FParameters>();
@@ -458,7 +459,10 @@ void InitLocalFogVolumesForViews(
 		{
 			CreateViewLocalFogVolumeBufferSRV(Scene, View, GraphBuilder, SortingData, ShouldRenderLocalFogVolumeInVolumetricFog(Scene, SceneViewFamily, bShouldRenderVolumetricFog));
 
-			LocalFogVolumeViewTiledCullingPass(View, GraphBuilder);
+			if (View.LocalFogVolumeViewData.GPUInstanceCount > 0)
+			{
+				LocalFogVolumeViewTiledCullingPass(View, GraphBuilder);
+			}
 		}
 	}
 	else
