Fixes compilation failure on some platform

diff --git a/Engine/Source/Runtime/Renderer/Private/PostProcess/TemporalSuperResolution.cpp b/Engine/Source/Runtime/Renderer/Private/PostProcess/TemporalSuperResolution.cpp
--- a/Engine/Source/Runtime/Renderer/Private/PostProcess/TemporalSuperResolution.cpp
+++ b/Engine/Source/Runtime/Renderer/Private/PostProcess/TemporalSuperResolution.cpp
@@ -1309,7 +1309,6 @@ FDefaultTemporalUpscaler::FOutputs AddTemporalSuperResolutionPasses(
 	static auto CVarAntiAliasingQuality = IConsoleManager::Get().FindConsoleVariable(TEXT("sg.AntiAliasingQuality"));
 	check(CVarAntiAliasingQuality);
 	
-	if (IsVisualizeTSREnabled(View))
 	RDG_EVENT_SCOPE(GraphBuilder, "TemporalSuperResolution(sg.AntiAliasingQuality=%d%s) %dx%d -> %dx%d",
 		CVarAntiAliasingQuality->GetInt(),
 		bSupportsAlpha ? TEXT(" Alpha") : TEXT(""),
