Removed mutex from GetLODBias as it is used during VT update, and the lock cost is too high. Added a way to access the underlying VT update task. Prerequisite for parallel GDME work.

diff --git a/Engine/Source/Runtime/Renderer/Private/VT/VirtualTextureSystem.cpp b/Engine/Source/Runtime/Renderer/Private/VT/VirtualTextureSystem.cpp
--- a/Engine/Source/Runtime/Renderer/Private/VT/VirtualTextureSystem.cpp
+++ b/Engine/Source/Runtime/Renderer/Private/VT/VirtualTextureSystem.cpp
@@ -2851,8 +2851,6 @@ void FVirtualTextureSystem::ReleasePendingResources()
 
 float FVirtualTextureSystem::GetGlobalMipBias() const
 {
-	UE::TScopeLock Lock(Mutex);
-
 	float MaxResidencyMipMapBias = 0.f;
 	for (int32 SpaceIndex = 0; SpaceIndex < PhysicalSpaces.Num(); ++SpaceIndex)
 	{
diff --git a/Engine/Source/Runtime/Renderer/Private/VT/VirtualTextureSystem.h b/Engine/Source/Runtime/Renderer/Private/VT/VirtualTextureSystem.h
--- a/Engine/Source/Runtime/Renderer/Private/VT/VirtualTextureSystem.h
+++ b/Engine/Source/Runtime/Renderer/Private/VT/VirtualTextureSystem.h
@@ -82,6 +82,12 @@ struct FVirtualTextureUpdateSettings
 
 class FVirtualTextureUpdater
 {
+public:
+	UE::Tasks::FTask GetTask() const
+	{
+		return AsyncTask;
+	}
+
 private:
 	FVirtualTextureUpdateSettings Settings;
 	FConcurrentLinearBulkObjectAllocator Allocator;
