Renamed bConstrainToRenderThread to bSupportsParallelCreateDestroy and added bSupportsParallelGDME flag.

diff --git a/Engine/Plugins/FX/Niagara/Source/Niagara/Private/NiagaraComponent.cpp b/Engine/Plugins/FX/Niagara/Source/Niagara/Private/NiagaraComponent.cpp
--- a/Engine/Plugins/FX/Niagara/Source/Niagara/Private/NiagaraComponent.cpp
+++ b/Engine/Plugins/FX/Niagara/Source/Niagara/Private/NiagaraComponent.cpp
@@ -207,6 +207,9 @@ FNiagaraSceneProxy::FNiagaraSceneProxy(UNiagaraComponent* InComponent)
 		SystemStatString = NiagaraSystem->GetFName().ToString();
 #endif
 	}
+
+	// Niagara renderers reference a lot of common contexts that aren't locked (and would otherwise introduce a lot of contention).
+	bSupportsParallelGDME = false;
 }
 
 SIZE_T FNiagaraSceneProxy::GetTypeHash() const
diff --git a/Engine/Source/Runtime/Engine/Private/PrimitiveSceneProxy.cpp b/Engine/Source/Runtime/Engine/Private/PrimitiveSceneProxy.cpp
--- a/Engine/Source/Runtime/Engine/Private/PrimitiveSceneProxy.cpp
+++ b/Engine/Source/Runtime/Engine/Private/PrimitiveSceneProxy.cpp
@@ -433,7 +433,8 @@ FPrimitiveSceneProxy::FPrimitiveSceneProxy(const FPrimitiveSceneProxyDesc& InPro
 ,	bForceHidden(false)
 ,	bCollisionEnabled(InProxyDesc.IsCollisionEnabled())
 ,	bTreatAsBackgroundForOcclusion(InProxyDesc.bTreatAsBackgroundForOcclusion)
-,	bConstrainToRenderThread(false)
+,	bSupportsParallelCreateDestroy(true)
+,	bSupportsParallelGDME(true)
 ,	bVisibleInLumenScene(false)
 ,	bCanSkipRedundantTransformUpdates(true)
 ,	bGoodCandidateForCachedShadowmap(true)
@@ -1663,7 +1664,7 @@ bool FPrimitiveSceneProxy::IsShown(const FSceneView* View) const
 			return false;
 		}
 
-		const bool bOwnersContain = Owners.Contains(View->ViewActor);
+		const bool bOwnersContain = !Owners.IsEmpty() && Owners.Contains(View->ViewActor);
 		if (bOnlyOwnerSee && !bOwnersContain)
 		{
 			return false;
@@ -1706,8 +1707,8 @@ bool FPrimitiveSceneProxy::IsShadowCast(const FSceneView* View) const
 		{
 			return false;
 		}
-		
-		if (View->HiddenPrimitives.Contains(PrimitiveComponentId))
+
+		if (!View->HiddenPrimitives.IsEmpty() && View->HiddenPrimitives.Contains(PrimitiveComponentId))
 		{
 			return false;
 		}
diff --git a/Engine/Source/Runtime/Engine/Public/PrimitiveSceneProxy.h b/Engine/Source/Runtime/Engine/Public/PrimitiveSceneProxy.h
--- a/Engine/Source/Runtime/Engine/Public/PrimitiveSceneProxy.h
+++ b/Engine/Source/Runtime/Engine/Public/PrimitiveSceneProxy.h
@@ -683,7 +683,8 @@ public:
 	inline bool IsVisibleInLumenScene() const { return bVisibleInLumenScene; }
 	inline bool ShouldRenderInMainPass() const { return bRenderInMainPass; }
 	inline bool ShouldRenderInDepthPass() const { return bRenderInMainPass || bRenderInDepthPass; }
-	inline bool ShouldConstrainToRenderThread() const { return bConstrainToRenderThread; }
+	inline bool SupportsParallelCreateDestroy() const { return bSupportsParallelCreateDestroy; }
+	inline bool SupportsParallelGDME() const { return bSupportsParallelGDME; }
 	inline bool IsCollisionEnabled() const { return bCollisionEnabled; }
 	inline bool IsHovered() const { return bHovered; }
 	inline bool IsOwnedBy(const AActor* Actor) const { return Owners.Find(Actor) != INDEX_NONE; }
@@ -1290,8 +1291,14 @@ private:
 
 protected:
 
-	/** Whether to force legacy render thread commands for OnTransformChanged / CreateRenderThreadResources / DestroyRenderThreadResources. By default, commands can run asynchronously. */
-	uint8 bConstrainToRenderThread : 1;
+	/** Whether this scene proxy supports {Create, Destroy}RenderThreadResources and OnTransformUpdate calls on the Scene render command pipe. If enabled (default),
+	 *  CreateRenderThreadResources is deferred and called asynchronously from the scene update to avoid race conditions with other render commands. DestroyRenderThreadResources
+	 *  and OnTransformUpdate is called on the Scene render command pipe. Disable to revert to legacy behavior which runs everything on the render thread.
+	 */
+	uint8 bSupportsParallelCreateDestroy : 1;
+
+	/** Whether the proxy supports asynchronously calling GetDynamicMeshElements. If disabled, all calls for various proxies are serialized with respect to each other. */
+	uint8 bSupportsParallelGDME : 1;
 
 	/** Whether this component should be tracked by Lumen Scene. Turning this off will remove it from Lumen Scene and Lumen won't generate surface cache for it. */
 	uint8 bVisibleInLumenScene : 1;
diff --git a/Engine/Source/Runtime/Landscape/Private/LandscapeRender.cpp b/Engine/Source/Runtime/Landscape/Private/LandscapeRender.cpp
--- a/Engine/Source/Runtime/Landscape/Private/LandscapeRender.cpp
+++ b/Engine/Source/Runtime/Landscape/Private/LandscapeRender.cpp
@@ -1393,7 +1393,7 @@ FLandscapeComponentSceneProxy::FLandscapeComponentSceneProxy(ULandscapeComponent
 	UpdateVisibleInLumenScene();
 
 	// Force OnTransformChanged and {Create, Destroy}RenderThreadResources onto the render thread.
-	bConstrainToRenderThread = true;
+	bSupportsParallelCreateDestroy = false;
 }
 
 void FLandscapeComponentSceneProxy::CreateRenderThreadResources(FRHICommandListBase& RHICmdList)
@@ -4347,7 +4347,7 @@ FLandscapeMeshProxySceneProxy::FLandscapeMeshProxySceneProxy(UStaticMeshComponen
 	}
 
 	// Force OnTransformChanged and {Create, Destroy}RenderThreadResources onto the render thread.
-	bConstrainToRenderThread = true;
+	bSupportsParallelCreateDestroy = false;
 }
 
 void FLandscapeMeshProxySceneProxy::RegisterSections()
diff --git a/Engine/Source/Runtime/Renderer/Private/RendererScene.cpp b/Engine/Source/Runtime/Renderer/Private/RendererScene.cpp
--- a/Engine/Source/Runtime/Renderer/Private/RendererScene.cpp
+++ b/Engine/Source/Runtime/Renderer/Private/RendererScene.cpp
@@ -1954,7 +1954,7 @@ void FScene::BatchAddPrimitivesInternal(TArrayView<T*> InPrimitives)
 		FMatrix RenderMatrix = Primitive->GetRenderMatrix();
 		FVector AttachmentRootPosition = Primitive->GetActorPositionForRenderer();
 
-		TArray<FCreateCommand, TInlineAllocator<1>>& CreateCommands = GRenderCommandPipeMode != ERenderCommandPipeMode::All || PrimitiveSceneProxy->ShouldConstrainToRenderThread()
+		TArray<FCreateCommand, TInlineAllocator<1>>& CreateCommands = GRenderCommandPipeMode != ERenderCommandPipeMode::All || !PrimitiveSceneProxy->SupportsParallelCreateDestroy()
 			? CreateCommandsRenderThread
 			: CreateCommandsScenePipe;
 
@@ -2514,16 +2514,16 @@ void FScene::BatchRemovePrimitivesInternal(TArrayView<T*> InPrimitives)
 		{
 			FPrimitiveSceneInfo* PrimitiveSceneInfo = PrimitiveSceneProxy->GetPrimitiveSceneInfo();
 
-			const bool bConstrainToRenderThread = GRenderCommandPipeMode != ERenderCommandPipeMode::All || PrimitiveSceneProxy->ShouldConstrainToRenderThread();
+			const bool bRenderThreadSceneCommands = GRenderCommandPipeMode != ERenderCommandPipeMode::All || !PrimitiveSceneProxy->SupportsParallelCreateDestroy();
 
-			if (bConstrainToRenderThread)
+			if (bRenderThreadSceneCommands)
 			{
 				DestroyProxies.Emplace(PrimitiveSceneProxy);
 			}
 
 			// Disassociate the primitive's scene proxy.
 			Primitive->ReleaseSceneProxy();
-			DetachCommands.Add({ PrimitiveSceneInfo, !bConstrainToRenderThread ? PrimitiveSceneProxy : nullptr, &Primitive->GetSceneData().AttachmentCounter });
+			DetachCommands.Add({ PrimitiveSceneInfo, !bRenderThreadSceneCommands ? PrimitiveSceneProxy : nullptr, &Primitive->GetSceneData().AttachmentCounter });
 		}
 	}
 
@@ -3815,6 +3815,7 @@ void FSceneVelocityData::StartFrame(FScene* Scene)
 
 void FScene::GetPrimitiveUniformShaderParameters_RenderThread(const FPrimitiveSceneInfo* PrimitiveSceneInfo, bool& bHasPrecomputedVolumetricLightmap, FMatrix& PreviousLocalToWorld, int32& SingleCaptureIndex, bool& bOutputVelocity) const 
 {
+	SCOPED_NAMED_EVENT(GetPrimitiveUniformShaderParameters_RenderThread, FColor::Yellow);
 	const FMatrix LocalToWorld = PrimitiveSceneInfo->Proxy->GetLocalToWorld();
 	PreviousLocalToWorld = LocalToWorld;
 	bOutputVelocity = false;
@@ -5475,7 +5476,7 @@ void FScene::UpdateAllPrimitiveSceneInfos(FRDGBuilder& GraphBuilder, EUpdateAllP
 			TRACE_CPUPROFILER_EVENT_SCOPE(Scene::CreatePrimitiveResources);
 			for (FPrimitiveSceneInfo* Primitive : CreateResourcesPrimitives)
 			{
-				if (!Primitive->Proxy->ShouldConstrainToRenderThread())
+				if (Primitive->Proxy->SupportsParallelCreateDestroy())
 				{
 					Primitive->Proxy->CreateRenderThreadResources(RHICmdList);
 				}
