Fixed a case where volumetric fog ambient would not be lit correctly due to the CPU data being read, but this on is not updated when real time capture is in use.
The math and result are the same. It is a bit slower because that dat awill be fetched from a Structured buffer instead of a constant buffer. But at least it will go onto the scalar registers.

#prelfight none

diff --git a/Engine/Shaders/Private/VolumetricFog.usf b/Engine/Shaders/Private/VolumetricFog.usf
--- a/Engine/Shaders/Private/VolumetricFog.usf
+++ b/Engine/Shaders/Private/VolumetricFog.usf
@@ -15,6 +15,7 @@
 #endif
 #include "VolumeLightingCommon.ush"
 #include "VolumetricLightmapShared.ush"
+#include "ReflectionEnvironmentShared.ush"
 #include "ForwardShadowingCommon.ush"
 #include "ParticipatingMediaCommon.ush"
 #include "LightDataUniforms.ush"
@@ -656,7 +657,7 @@ float GetDirectionalLightFunction(float3 TranslatedWorldPosition) // directional
 	return Texture2DSampleLevel(LightFunctionTexture, LightFunctionSampler, LightFunctionUV, 0).x;
 }
 
-float4 SkySH[3];
+uint SampleSkyLightDiffuseEnvMap;
 float2 UseHeightFogColors;	// x=override directional light using height fog inscattering color, y=override sky light using heigh fog inscattering cubemap
 float UseDirectionalLightShadowing;
 float StaticLightingScatteringIntensity;
@@ -811,12 +812,7 @@ void LightScatteringCS(
 			}
 			else
 			{
-				FTwoBandSHVectorRGB SkyIrradianceSH;
-				SkyIrradianceSH.R.V = SkySH[0];
-				SkyIrradianceSH.G.V = SkySH[1];
-				SkyIrradianceSH.B.V = SkySH[2];
-
-				SkyLighting = View.SkyLightColor.rgb * max(DotSH(SkyIrradianceSH, RotatedHGZonalHarmonic), 0) / PI;
+				SkyLighting = SampleSkyLightDiffuseEnvMap > 0 ? View.SkyLightColor.rgb * GetSkySHDiffuseSimple(CameraVector * -PhaseG) : 0.0f.xxx;
 			}
 
 			float SkyVisibility = ComputeSkyVisibility(TranslatedWorldPosition, BrickTextureUVs);
diff --git a/Engine/Source/Runtime/Renderer/Private/VolumetricFog.cpp b/Engine/Source/Runtime/Renderer/Private/VolumetricFog.cpp
--- a/Engine/Source/Runtime/Renderer/Private/VolumetricFog.cpp
+++ b/Engine/Source/Runtime/Renderer/Private/VolumetricFog.cpp
@@ -985,7 +985,7 @@ class FVolumetricFogLightScatteringCS : public FGlobalShader
 		SHADER_PARAMETER_STRUCT_INCLUDE(FAOParameters, AOParameters)
 		SHADER_PARAMETER_STRUCT_INCLUDE(FGlobalDistanceFieldParameters2, GlobalDistanceFieldParameters)
 		SHADER_PARAMETER_RDG_TEXTURE_UAV(RWTexture2D, RWLightScattering)
-		SHADER_PARAMETER_ARRAY(FVector4f, SkySH, [3])
+		SHADER_PARAMETER(uint32, SampleSkyLightDiffuseEnvMap)
 		SHADER_PARAMETER(FMatrix44f, DirectionalLightFunctionTranslatedWorldToShadow)
 		SHADER_PARAMETER(FMatrix44f, CloudShadowmapTranslatedWorldToLightClipMatrix)
 		SHADER_PARAMETER(FVector2f, PrevConservativeDepthTextureSize)
@@ -1563,18 +1563,12 @@ void FSceneRenderer::ComputeVolumetricFog(FRDGBuilder& GraphBuilder,
 				&& View.Family->EngineShowFlags.SkyLighting)
 			{
 				PassParameters->SkyLightUseStaticShadowing = SkyLight->bWantsStaticShadowing && SkyLight->bCastShadows ? 1.0f : 0.0f;
-
-				const FSHVectorRGB3& SkyIrradiance = SkyLight->IrradianceEnvironmentMap;
-				PassParameters->SkySH[0] = (FVector4f&)SkyIrradiance.R.V;
-				PassParameters->SkySH[1] = (FVector4f&)SkyIrradiance.G.V;
-				PassParameters->SkySH[2] = (FVector4f&)SkyIrradiance.B.V;
+				PassParameters->SampleSkyLightDiffuseEnvMap = 1;
 			}
 			else
 			{
 				PassParameters->SkyLightUseStaticShadowing = 0.0f;
-				PassParameters->SkySH[0] = FVector4f(0, 0, 0, 0);
-				PassParameters->SkySH[1] = FVector4f(0, 0, 0, 0);
-				PassParameters->SkySH[2] = FVector4f(0, 0, 0, 0);
+				PassParameters->SampleSkyLightDiffuseEnvMap = 0;
 			}
 
 			float StaticLightingScatteringIntensityValue = 0;
