Only Destroy/Add Light Primitive Interactions when needed on light transform update
On mobile this means we only invalidate MDC's when needed
#tests win64, mobile, PIE

diff --git a/Engine/Source/Runtime/Renderer/Private/RendererScene.cpp b/Engine/Source/Runtime/Renderer/Private/RendererScene.cpp
--- a/Engine/Source/Runtime/Renderer/Private/RendererScene.cpp
+++ b/Engine/Source/Runtime/Renderer/Private/RendererScene.cpp
@@ -3838,17 +3838,8 @@ void FScene::UpdateLightTransform_RenderThread(int32 LightId, FLightSceneInfo* L
 
 	// This is called without a valid ID when the update is fused with an 'add' command (saves redundant scene updates to do the update first)
 	const bool bHasId = LightId != INDEX_NONE;
-	// Don't remove directional lights when their transform changes as nothing in RemoveFromScene() depends on their transform
-	const bool bRemove = bHasId && (Lights[LightId].LightType != LightType_Directional);
-
-	if (bHasId)
-	{
-		if (bRemove)
-		{
-			// Remove the light from the scene.
-			LightSceneInfo->RemoveFromScene();
-		}
-	}
+	// Don't Update Primitive Interactions for directional lights
+	const bool bUpdatePrimitiveInteractions = bHasId && (Lights[LightId].LightType != LightType_Directional);
 
 	// Invalidate the path tracer if the transform actually changed
 	// NOTE: Position is derived from the Matrix, so there is no need to check it separately
@@ -3866,11 +3857,58 @@ void FScene::UpdateLightTransform_RenderThread(int32 LightId, FLightSceneInfo* L
 		checkSlow(Lights[LightId].LightSceneInfo == LightSceneInfo);
 		Lights[LightId].Init(LightSceneInfo);
 
-		// Don't re-add directional lights when their transform changes as nothing in AddToScene() depends on their transform
-		if (bRemove)
+		if (bUpdatePrimitiveInteractions)
 		{
-			// Add the light to the scene at its new location.
-			LightSceneInfo->AddToScene();
+			using PrimitiveSceneInfoSet = TSet<FPrimitiveSceneInfo*, DefaultKeyFuncs<FPrimitiveSceneInfo*>, SceneRenderingSetAllocator>;
+			PrimitiveSceneInfoSet PrevPrimitivesInBounds;
+
+			TMap<FPrimitiveSceneInfo*, FLightPrimitiveInteraction*, SceneRenderingSetAllocator> PrimitivesToInteractions;
+			for (FLightPrimitiveInteraction* Interaction = LightSceneInfo->GetDynamicInteractionOftenMovingPrimitiveList();
+				Interaction;
+				Interaction = Interaction->GetNextPrimitive()
+				)
+			{
+				PrevPrimitivesInBounds.Add(Interaction->GetPrimitiveSceneInfo());
+				PrimitivesToInteractions.Add(Interaction->GetPrimitiveSceneInfo(), Interaction);
+			}
+
+			for (FLightPrimitiveInteraction* Interaction = LightSceneInfo->GetDynamicInteractionStaticPrimitiveList();
+				Interaction;
+				Interaction = Interaction->GetNextPrimitive()
+				)
+			{
+				PrevPrimitivesInBounds.Add(Interaction->GetPrimitiveSceneInfo());
+				PrimitivesToInteractions.Add(Interaction->GetPrimitiveSceneInfo(), Interaction);
+			}
+
+			PrimitiveSceneInfoSet CurrentPrimitivesInBounds;
+			const FLightSceneInfoCompact& LightSceneInfoCompact = Lights[LightId];
+
+			if (LightSceneInfo->OctreeId.IsValidId())
+			{
+				// Re-add the light to the octree after transform update.
+				LocalShadowCastingLightOctree.RemoveElement(LightSceneInfo->OctreeId);
+				LightSceneInfo->OctreeId = FOctreeElementId2();
+				LocalShadowCastingLightOctree.AddElement(LightSceneInfoCompact);
+			}
+
+			PrimitiveOctree.FindElementsWithBoundsTest(LightSceneInfo->GetBoundingBox(), [&LightSceneInfoCompact, &CurrentPrimitivesInBounds, this](const FPrimitiveSceneInfoCompact& PrimitiveSceneInfoCompact)
+				{
+					CurrentPrimitivesInBounds.Add(PrimitiveSceneInfoCompact.PrimitiveSceneInfo);
+				});
+
+			PrimitiveSceneInfoSet PrimitivesToBeRemoved = PrevPrimitivesInBounds.Difference(CurrentPrimitivesInBounds);
+			PrimitiveSceneInfoSet PrimitivesToAdd = CurrentPrimitivesInBounds.Difference(PrevPrimitivesInBounds);
+
+			for (FPrimitiveSceneInfo* PrimitiveToRemove : PrimitivesToBeRemoved)
+			{
+				FLightPrimitiveInteraction::Destroy(PrimitivesToInteractions[PrimitiveToRemove]);
+			}
+
+			for (FPrimitiveSceneInfo* PrimitiveToAdd : PrimitivesToAdd)
+			{
+				LightSceneInfo->CreateLightPrimitiveInteraction(LightSceneInfoCompact, PrimitiveToAdd);
+			}
 		}
 	}
 }
