Screen Space Shadows as a standalone pass.
- calculated at half-res + bilaterial upsample (similar to distance field and capsule shadows)
- only supports directional lights at the moment.
- configuration is similar to existing Contact Shadows except it doesn't read GBuffer so it doesn't support per-material overides.
- only used on mobile renderer.

diff --git a/Engine/Shaders/Private/ScreenSpaceShadows.usf b/Engine/Shaders/Private/ScreenSpaceShadows.usf
new file mode 100644
--- /dev/null
+++ b/Engine/Shaders/Private/ScreenSpaceShadows.usf
@@ -0,0 +1,84 @@
+// Copyright Epic Games, Inc. All Rights Reserved.
+
+/*=============================================================================
+	ScreenSpaceShadows.usf
+=============================================================================*/
+
+#include "ScreenSpaceShadowRayCast.ush"
+
+#include "Common.ush"
+#include "SceneTexturesCommon.ush"
+#include "ComputeShaderUtils.ush"
+
+#ifdef UPSAMPLE_PASS
+#	include "ShadowFactorsUpsampleCommon.ush"
+#endif
+
+#ifndef THREADGROUP_SIZEX
+#	define THREADGROUP_SIZEX 1
+#endif
+
+#ifndef THREADGROUP_SIZEY
+#	define THREADGROUP_SIZEY 1
+#endif
+
+float3 LightDirection;
+float ContactShadowLength;
+uint bContactShadowLengthInWS;
+float ContactShadowCastingIntensity;
+
+uint DownsampleFactor;
+int4 ScissorRectMinAndSize;
+
+RWTexture2D<float2> RWShadowFactors;
+
+[numthreads(THREADGROUP_SIZEX, THREADGROUP_SIZEY, 1)]
+void ScreenSpaceShadowsCS(
+	uint3 GroupId : SV_GroupID,
+	uint3 DispatchThreadId : SV_DispatchThreadID,
+    uint3 GroupThreadId : SV_GroupThreadID) 
+{
+	uint ThreadIndex = GroupThreadId.y * THREADGROUP_SIZEX + GroupThreadId.x;
+
+	float2 ScreenUV = float2((DispatchThreadId.xy * DownsampleFactor + ScissorRectMinAndSize.xy + .5f) * View.BufferSizeAndInvSize.zw);
+	float2 ScreenPosition = (ScreenUV.xy - View.ScreenPositionScaleBias.wz) / View.ScreenPositionScaleBias.xy;
+
+	float SceneDepth = CalcSceneDepth(ScreenUV);
+	float3 OpaqueTranslatedWorldPosition = mul(float4(GetScreenPositionForProjectionType(ScreenPosition, SceneDepth), SceneDepth, 1), PrimaryView.ScreenToTranslatedWorld).xyz;
+
+	const float ContactShadowLengthScreenScale = GetTanHalfFieldOfView().y * SceneDepth * DownsampleFactor;
+	const float ActualContactShadowLength = ContactShadowLength * (bContactShadowLengthInWS ? 1.0f : ContactShadowLengthScreenScale);
+
+	const float Dither = InterleavedGradientNoise(DispatchThreadId.xy + 0.5f, View.StateFrameIndexMod8);
+
+	float2 HitUV;
+	float HitDistance = CastScreenSpaceShadowRay(OpaqueTranslatedWorldPosition, LightDirection, ActualContactShadowLength, 8, Dither, 2.0f, false, HitUV);
+
+	float Result = HitDistance > 0.0 ? (1.0f - ContactShadowCastingIntensity) : 1.0f;
+
+	#if METAL_PROFILE 
+		// clamp max depth to avoid #inf
+		SceneDepth = min(SceneDepth, 65500.0f);
+	#endif
+	RWShadowFactors[DispatchThreadId.xy] = float2(Result, SceneDepth);
+}
+
+#ifdef UPSAMPLE_PASS
+
+Texture2D ShadowFactorsTexture;
+SamplerState ShadowFactorsSampler;
+
+float OneOverDownsampleFactor;
+
+void ScreenSpaceShadowsUpsamplePS(
+	in float4 SVPos : SV_POSITION,
+	out float4 OutColor : SV_Target0)
+{
+	float Output;
+	float SceneDepth;
+	UpsampleShadowFactors(SVPos, ScissorRectMinAndSize, OneOverDownsampleFactor, ShadowFactorsTexture, ShadowFactorsSampler, Output, SceneDepth);
+
+	OutColor = EncodeLightAttenuation(half4(Output, Output, Output, Output));
+}
+
+#endif // UPSAMPLE_PASS
diff --git a/Engine/Source/Runtime/Renderer/Private/ShadowRendering.cpp b/Engine/Source/Runtime/Renderer/Private/ShadowRendering.cpp
--- a/Engine/Source/Runtime/Renderer/Private/ShadowRendering.cpp
+++ b/Engine/Source/Runtime/Renderer/Private/ShadowRendering.cpp
@@ -15,6 +15,7 @@
 #include "HairStrands/HairStrandsRendering.h"
 #include "VirtualShadowMaps/VirtualShadowMapProjection.h"
 #include "Shadows/ShadowSceneRenderer.h"
+#include "Shadows/ScreenSpaceShadows.h"
 #include "RenderCore.h"
 #include "TranslucentLighting.h"
 #include "MobileBasePassRendering.h"
@@ -2414,11 +2415,15 @@ void FMobileSceneRenderer::RenderMobileShadowProjections(
 		// Local light shadows don't render to shadow mask texture on mobile deferred
 		if (LightSceneProxy->GetLightType() == LightType_Directional || !IsMobileDeferredShadingEnabled(ShaderPlatform))
 		{
+			const bool bProjectingForForwardShading = true;
+
 			RenderShadowProjections(GraphBuilder, SceneTextures,
 				ScreenShadowMaskTexture,
 				nullptr,
 				LightSceneInfo,
-				true);
+				bProjectingForForwardShading);
+
+			RenderScreenSpaceShadows(GraphBuilder, SceneTextures, Views, LightSceneInfo, bProjectingForForwardShading, ScreenShadowMaskTexture);
 		}
 	}
 }
diff --git a/Engine/Source/Runtime/Renderer/Private/Shadows/ScreenSpaceShadows.cpp b/Engine/Source/Runtime/Renderer/Private/Shadows/ScreenSpaceShadows.cpp
new file mode 100644
--- /dev/null
+++ b/Engine/Source/Runtime/Renderer/Private/Shadows/ScreenSpaceShadows.cpp
@@ -0,0 +1,252 @@
+// Copyright Epic Games, Inc. All Rights Reserved.
+
+/*
+=================================================================================
+	ScreenSpaceShadows.cpp: Functionality for rendering screen space shadows
+=================================================================================
+*/
+
+#include "ScreenSpaceShadows.h"
+
+#include "LightSceneInfo.h"
+#include "LightSceneProxy.h"
+#include "ShadowRendering.h"
+#include "SceneRendering.h"
+#include "RenderGraphUtils.h"
+#include "PixelShaderUtils.h"
+
+const int32 GScreenSpaceShadowsTileSizeX = 8;
+const int32 GScreenSpaceShadowsTileSizeY = 8;
+
+int32 GetScreenSpaceShadowDownsampleFactor()
+{
+	return 2;
+}
+
+FIntPoint GetBufferSizeForScreenSpaceShadows(const FViewInfo& View)
+{
+	return FIntPoint::DivideAndRoundDown(View.GetSceneTexturesConfig().Extent, GetScreenSpaceShadowDownsampleFactor());
+}
+
+class FScreenSpaceShadowsCS : public FGlobalShader
+{
+	DECLARE_GLOBAL_SHADER(FScreenSpaceShadowsCS);
+	SHADER_USE_PARAMETER_STRUCT(FScreenSpaceShadowsCS, FGlobalShader);
+
+	BEGIN_SHADER_PARAMETER_STRUCT(FParameters, )
+		SHADER_PARAMETER_RDG_TEXTURE_UAV(RWTexture2D<float2>, RWShadowFactors)
+		SHADER_PARAMETER_STRUCT_REF(FViewUniformShaderParameters, View)
+		SHADER_PARAMETER_STRUCT_INCLUDE(FSceneTextureShaderParameters, SceneTextures)
+		SHADER_PARAMETER(FVector3f, LightDirection)
+		SHADER_PARAMETER(float, ContactShadowLength)
+		SHADER_PARAMETER(uint32, bContactShadowLengthInWS)
+		SHADER_PARAMETER(float, ContactShadowCastingIntensity)
+		SHADER_PARAMETER(FIntRect, ScissorRectMinAndSize)
+		SHADER_PARAMETER(uint32, DownsampleFactor)
+	END_SHADER_PARAMETER_STRUCT()
+
+	static bool ShouldCompilePermutation(const FGlobalShaderPermutationParameters& Parameters)
+	{
+		return true;
+	}
+
+	static void ModifyCompilationEnvironment(const FGlobalShaderPermutationParameters& Parameters, FShaderCompilerEnvironment& OutEnvironment)
+	{
+		OutEnvironment.SetDefine(TEXT("THREADGROUP_SIZEX"), GScreenSpaceShadowsTileSizeX);
+		OutEnvironment.SetDefine(TEXT("THREADGROUP_SIZEY"), GScreenSpaceShadowsTileSizeY);
+		OutEnvironment.SetDefine(TEXT("FORCE_DEPTH_TEXTURE_READS"), 1);
+		OutEnvironment.SetDefine(TEXT("PLATFORM_SUPPORTS_TYPED_UAV_LOAD"), (int32)RHISupports4ComponentUAVReadWrite(Parameters.Platform));
+	}
+};
+
+IMPLEMENT_GLOBAL_SHADER(FScreenSpaceShadowsCS, "/Engine/Private/ScreenSpaceShadows.usf", "ScreenSpaceShadowsCS", SF_Compute);
+
+class FScreenSpaceShadowsUpsamplePS : public FGlobalShader
+{
+	DECLARE_GLOBAL_SHADER(FScreenSpaceShadowsUpsamplePS);
+	SHADER_USE_PARAMETER_STRUCT(FScreenSpaceShadowsUpsamplePS, FGlobalShader);
+
+	BEGIN_SHADER_PARAMETER_STRUCT(FParameters, )
+		SHADER_PARAMETER_STRUCT_INCLUDE(FSceneTextureShaderParameters, SceneTextures)
+		SHADER_PARAMETER_STRUCT_REF(FViewUniformShaderParameters, View)
+		SHADER_PARAMETER_RDG_TEXTURE(Texture2D, ShadowFactorsTexture)
+		SHADER_PARAMETER_SAMPLER(SamplerState, ShadowFactorsSampler)
+		SHADER_PARAMETER(FIntRect, ScissorRectMinAndSize)
+		SHADER_PARAMETER(float, OneOverDownsampleFactor)
+	END_SHADER_PARAMETER_STRUCT()
+
+	class FUpsample : SHADER_PERMUTATION_BOOL("SHADOW_FACTORS_UPSAMPLE_REQUIRED");
+	using FPermutationDomain = TShaderPermutationDomain<FUpsample>;
+
+	static bool ShouldCompilePermutation(const FGlobalShaderPermutationParameters& Parameters)
+	{
+		return true;
+	}
+
+	static void ModifyCompilationEnvironment(const FGlobalShaderPermutationParameters& Parameters, FShaderCompilerEnvironment& OutEnvironment)
+	{
+		OutEnvironment.SetDefine(TEXT("UPSAMPLE_PASS"), 1);
+		OutEnvironment.SetDefine(TEXT("FORCE_DEPTH_TEXTURE_READS"), 1);
+	}
+};
+
+IMPLEMENT_GLOBAL_SHADER(FScreenSpaceShadowsUpsamplePS, "/Engine/Private/ScreenSpaceShadows.usf", "ScreenSpaceShadowsUpsamplePS", SF_Pixel);
+
+BEGIN_SHADER_PARAMETER_STRUCT(FScreenSpaceShadowsUpsample, )
+	SHADER_PARAMETER_STRUCT_INCLUDE(FScreenSpaceShadowsUpsamplePS::FParameters, PS)
+	RENDER_TARGET_BINDING_SLOTS()
+END_SHADER_PARAMETER_STRUCT()
+
+void RenderScreenSpaceShadows(
+	FRDGBuilder& GraphBuilder,
+	bool bAsyncCompute,
+	const FMinimalSceneTextures& SceneTextures,
+	const FViewInfo& View,
+	FIntRect ScissorRect,
+	bool bProjectingForForwardShading,
+	const FLightSceneInfo* LightSceneInfo,
+	FRDGTextureRef ScreenShadowMaskTexture)
+{
+	check(ScissorRect.Area() > 0);
+
+	const FIntPoint BufferSize = GetBufferSizeForScreenSpaceShadows(View);
+	FRDGTextureDesc Desc(FRDGTextureDesc::Create2D(BufferSize, PF_G16R16F, FClearValueBinding::None, TexCreate_UAV | TexCreate_ShaderResource));
+	FRDGTextureRef ShadowsTexture = GraphBuilder.CreateTexture(Desc, TEXT("ScreenSpaceShadows"));
+
+	{
+		FScreenSpaceShadowsCS::FParameters* PassParameters = GraphBuilder.AllocParameters<FScreenSpaceShadowsCS::FParameters>();
+
+		PassParameters->RWShadowFactors = GraphBuilder.CreateUAV(ShadowsTexture);
+		PassParameters->View = View.ViewUniformBuffer;
+		PassParameters->SceneTextures = SceneTextures.GetSceneTextureShaderParameters(View.GetFeatureLevel());
+
+		PassParameters->ScissorRectMinAndSize = FIntRect(ScissorRect.Min, ScissorRect.Size());
+		PassParameters->DownsampleFactor = GetScreenSpaceShadowDownsampleFactor();
+
+		const FLightSceneProxy* LightProxy = LightSceneInfo->Proxy;
+		FLightRenderParameters LightParameters;
+		LightProxy->GetLightShaderParameters(LightParameters);
+
+		PassParameters->LightDirection = LightParameters.Direction;
+		PassParameters->ContactShadowLength = LightProxy->GetContactShadowLength();
+		PassParameters->bContactShadowLengthInWS = LightProxy->IsContactShadowLengthInWS();
+		PassParameters->ContactShadowCastingIntensity = LightProxy->GetContactShadowCastingIntensity();
+
+		auto ComputeShader = View.ShaderMap->GetShader<FScreenSpaceShadowsCS>();
+
+		uint32 GroupSizeX = FMath::DivideAndRoundUp(ScissorRect.Size().X / GetScreenSpaceShadowDownsampleFactor(), GScreenSpaceShadowsTileSizeX);
+		uint32 GroupSizeY = FMath::DivideAndRoundUp(ScissorRect.Size().Y / GetScreenSpaceShadowDownsampleFactor(), GScreenSpaceShadowsTileSizeY);
+
+		FComputeShaderUtils::AddPass(
+			GraphBuilder,
+			RDG_EVENT_NAME("ScreenSpaceShadowing %ux%u", GroupSizeX * GScreenSpaceShadowsTileSizeX, GroupSizeY * GScreenSpaceShadowsTileSizeY),
+			bAsyncCompute ? ERDGPassFlags::AsyncCompute : ERDGPassFlags::Compute,
+			ComputeShader,
+			PassParameters,
+			FIntVector(GroupSizeX, GroupSizeY, 1));
+	}
+
+	{
+		FScreenSpaceShadowsUpsample* PassParameters = GraphBuilder.AllocParameters<FScreenSpaceShadowsUpsample>();
+		PassParameters->RenderTargets[0] = FRenderTargetBinding(ScreenShadowMaskTexture, ERenderTargetLoadAction::ELoad);
+		PassParameters->RenderTargets.DepthStencil = FDepthStencilBinding(SceneTextures.Depth.Target, ERenderTargetLoadAction::ELoad, ERenderTargetLoadAction::ELoad, FExclusiveDepthStencil::DepthRead_StencilRead);
+
+		PassParameters->PS.View = GetShaderBinding(View.ViewUniformBuffer);
+		PassParameters->PS.SceneTextures = SceneTextures.GetSceneTextureShaderParameters(View.GetFeatureLevel());
+		PassParameters->PS.ShadowFactorsTexture = ShadowsTexture;
+		PassParameters->PS.ShadowFactorsSampler = TStaticSamplerState<SF_Bilinear>::GetRHI();
+		PassParameters->PS.ScissorRectMinAndSize = FIntRect(ScissorRect.Min, ScissorRect.Size());
+		PassParameters->PS.OneOverDownsampleFactor = 1.0f / GetScreenSpaceShadowDownsampleFactor();
+
+		FScreenSpaceShadowsUpsamplePS::FPermutationDomain PermutationVector;
+		PermutationVector.Set<FScreenSpaceShadowsUpsamplePS::FUpsample>(GetScreenSpaceShadowDownsampleFactor() != 1);
+		auto PixelShader = View.ShaderMap->GetShader<FScreenSpaceShadowsUpsamplePS>(PermutationVector);
+
+		// blend separately from CSM / DF Shadows since those interact with static lighting
+		// this matches behavior of GetShadowTerms(...) 
+		const bool bIsWholeSceneDirectionalShadow = false;
+
+		FRHIBlendState* BlendState = FProjectedShadowInfo::GetBlendStateForProjection(
+			LightSceneInfo->GetDynamicShadowMapChannel(),
+			bIsWholeSceneDirectionalShadow,
+			false,
+			bProjectingForForwardShading,
+			false);
+
+		ClearUnusedGraphResources(PixelShader, &PassParameters->PS);
+
+		GraphBuilder.AddPass(
+			RDG_EVENT_NAME("Upsample"),
+			PassParameters,
+			ERDGPassFlags::Raster,
+			[PassParameters, &View, PixelShader, BlendState, ScissorRect](FRHICommandList& RHICmdList)
+			{
+				RHICmdList.SetViewport(ScissorRect.Min.X, ScissorRect.Min.Y, 0.0f, ScissorRect.Max.X, ScissorRect.Max.Y, 1.0f);
+				RHICmdList.SetScissorRect(true, ScissorRect.Min.X, ScissorRect.Min.Y, ScissorRect.Max.X, ScissorRect.Max.Y);
+
+				FGraphicsPipelineStateInitializer GraphicsPSOInit;
+				FPixelShaderUtils::InitFullscreenPipelineState(RHICmdList, View.ShaderMap, PixelShader, GraphicsPSOInit);
+
+				GraphicsPSOInit.BoundShaderState.PixelShaderRHI = PixelShader.GetPixelShader();
+				GraphicsPSOInit.BlendState = BlendState;
+
+				SetGraphicsPipelineState(RHICmdList, GraphicsPSOInit, 0);
+
+				SetShaderParameters(RHICmdList, PixelShader, PixelShader.GetPixelShader(), PassParameters->PS);
+
+				FPixelShaderUtils::DrawFullscreenTriangle(RHICmdList);
+				RHICmdList.SetScissorRect(false, 0, 0, 0, 0);
+			});
+	}
+}
+
+void RenderScreenSpaceShadows(
+	FRDGBuilder& GraphBuilder,
+	const FMinimalSceneTextures& SceneTextures,
+	const TArray<FViewInfo>& Views,
+	const FLightSceneInfo* LightSceneInfo,
+	bool bProjectingForForwardShading,
+	FRDGTextureRef ScreenShadowMaskTexture)
+{
+	static auto* ContactShadowsCVar = IConsoleManager::Get().FindTConsoleVariableDataInt(TEXT("r.ContactShadows"));
+
+	if (ContactShadowsCVar && ContactShadowsCVar->GetValueOnRenderThread() == 0)
+	{
+		return;
+	}
+
+	RDG_EVENT_SCOPE(GraphBuilder, "ScreenSpaceShadows");
+
+	const FLightSceneProxy* LightSceneProxy = LightSceneInfo->Proxy;
+
+	for (int32 ViewIndex = 0; ViewIndex < Views.Num(); ViewIndex++)
+	{
+		const FViewInfo& View = Views[ViewIndex];
+		RDG_GPU_MASK_SCOPE(GraphBuilder, View.GPUMask);
+		RDG_EVENT_SCOPE_CONDITIONAL(GraphBuilder, Views.Num() > 1, "View%d", ViewIndex);
+
+		if (!View.Family->EngineShowFlags.ContactShadows)
+		{
+			continue;
+		}
+
+		FIntRect ScissorRect;
+		if (!LightSceneProxy->GetScissorRect(ScissorRect, View, View.ViewRect))
+		{
+			ScissorRect = View.ViewRect;
+		}
+
+		if (ScissorRect.Area() > 0)
+		{
+			RenderScreenSpaceShadows(
+				GraphBuilder,
+				false,
+				SceneTextures,
+				View,
+				ScissorRect,
+				bProjectingForForwardShading,
+				LightSceneInfo,
+				ScreenShadowMaskTexture);
+		}
+	}
+}
\ No newline at end of file
diff --git a/Engine/Source/Runtime/Renderer/Private/Shadows/ScreenSpaceShadows.h b/Engine/Source/Runtime/Renderer/Private/Shadows/ScreenSpaceShadows.h
new file mode 100644
--- /dev/null
+++ b/Engine/Source/Runtime/Renderer/Private/Shadows/ScreenSpaceShadows.h
@@ -0,0 +1,24 @@
+// Copyright Epic Games, Inc. All Rights Reserved.
+
+/*
+===============================================================================
+	ScreenSpaceShadows.h: Functionality for rendering screen space shadows
+===============================================================================
+*/
+
+#pragma once
+
+#include "CoreMinimal.h"
+#include "RendererInterface.h"
+#include "RenderGraphResources.h"
+
+struct FMinimalSceneTextures;
+class FLightSceneInfo;
+
+void RenderScreenSpaceShadows(
+	FRDGBuilder& GraphBuilder,
+	const FMinimalSceneTextures& SceneTextures,
+	const TArray<FViewInfo>& Views,
+	const FLightSceneInfo* LightSceneInfo,
+	bool bProjectingForForwardShading,
+	FRDGTextureRef ScreenShadowMaskTexture);
\ No newline at end of file
