Fix validation error when using Substrate and path tracer in the editor cased by unproduced resources.

diff --git a/Engine/Source/Runtime/Renderer/Private/PostProcess/PostProcessCompositeEditorPrimitives.cpp b/Engine/Source/Runtime/Renderer/Private/PostProcess/PostProcessCompositeEditorPrimitives.cpp
--- a/Engine/Source/Runtime/Renderer/Private/PostProcess/PostProcessCompositeEditorPrimitives.cpp
+++ b/Engine/Source/Runtime/Renderer/Private/PostProcess/PostProcessCompositeEditorPrimitives.cpp
@@ -8,6 +8,7 @@
 #include "BasePassRendering.h"
 #include "MobileBasePassRendering.h"
 #include "PixelShaderUtils.h"
+#include "Substrate/Substrate.h"
 
 namespace
 {
@@ -250,6 +251,14 @@ FScreenPassTexture AddEditorPrimitivePass(
 		EditorPrimitiveColor = GraphBuilder.CreateTexture(ColorDesc, TEXT("Editor.PrimitivesColor"));
 		EditorPrimitiveDepth = CreateCompositeDepthTexture(GraphBuilder, Extent, NumMSAASamples);
 	}
+	
+	// Subtrate data might not be produced in certain case (e.g., path-tracer). In such a case we force generate 
+	// them with a simple clear to please validation.
+	if (Substrate::IsSubstrateEnabled() && !HasBeenProduced(View.SubstrateViewData.SceneData->TopLayerTexture))
+	{
+		FRDGTextureClearInfo ClearInfo;
+		AddClearRenderTargetPass(GraphBuilder, View.SubstrateViewData.SceneData->TopLayerTexture, ClearInfo);
+	}
 
 	// Load the color target if it already exists.
 	const FScreenPassTextureViewport EditorPrimitivesViewport(EditorPrimitiveColor, Inputs.SceneColor.ViewRect);
