Fix a bug where Nanite materials would get mixed up when enabling/disabling sky lights in the scene.

diff --git a/Engine/Source/Runtime/Renderer/Private/RendererScene.cpp b/Engine/Source/Runtime/Renderer/Private/RendererScene.cpp
--- a/Engine/Source/Runtime/Renderer/Private/RendererScene.cpp
+++ b/Engine/Source/Runtime/Renderer/Private/RendererScene.cpp
@@ -6443,6 +6443,15 @@ void FScene::UpdateAllPrimitiveSceneInfos(FRDGBuilder& GraphBuilder, EUpdateAllP
 		{
 			FPrimitiveSceneInfo* PrimitiveSceneInfo = Primitives[PrimitiveIndex];
 			PrimitivesNeedingStaticMeshUpdate[PrimitiveSceneInfo->PackedIndex] = true;
+
+			// HACK: Update Nanite primitives that need re-caching in GPU Scene
+			// TODO: Should be able to remove this after the move to compute materials.
+			if (bScenesPrimitivesNeedStaticMeshElementUpdate &&
+				PrimitiveSceneInfo->Proxy &&
+				PrimitiveSceneInfo->Proxy->IsNaniteMesh())
+			{
+				GPUScene.AddPrimitiveToUpdate(PrimitiveSceneInfo->GetIndex(), EPrimitiveDirtyState::ChangedOther);
+			}
 		}
 
 		bScenesPrimitivesNeedStaticMeshElementUpdate = false;
