Implemented cvars to override contact shadow parameters
- also fixing incorrect shadow length in standalone pass when it is specified in world space units

diff --git a/Engine/Shaders/Private/ScreenSpaceShadows.usf b/Engine/Shaders/Private/ScreenSpaceShadows.usf
--- a/Engine/Shaders/Private/ScreenSpaceShadows.usf
+++ b/Engine/Shaders/Private/ScreenSpaceShadows.usf
@@ -46,7 +46,7 @@ void ScreenSpaceShadowsCS(
 	float SceneDepth = CalcSceneDepth(ScreenUV);
 	float3 OpaqueTranslatedWorldPosition = mul(float4(GetScreenPositionForProjectionType(ScreenPosition, SceneDepth), SceneDepth, 1), PrimaryView.ScreenToTranslatedWorld).xyz;
 
-	const float ContactShadowLengthScreenScale = GetTanHalfFieldOfView().y * SceneDepth * DownsampleFactor;
+	const float ContactShadowLengthScreenScale = GetTanHalfFieldOfView().y * SceneDepth;
 	const float ActualContactShadowLength = ContactShadowLength * (bContactShadowLengthInWS ? 1.0f : ContactShadowLengthScreenScale);
 
 	const float Dither = InterleavedGradientNoise(DispatchThreadId.xy + 0.5f, View.StateFrameIndexMod8);
diff --git a/Engine/Source/Runtime/Renderer/Private/LightRendering.cpp b/Engine/Source/Runtime/Renderer/Private/LightRendering.cpp
--- a/Engine/Source/Runtime/Renderer/Private/LightRendering.cpp
+++ b/Engine/Source/Runtime/Renderer/Private/LightRendering.cpp
@@ -141,6 +141,36 @@ static TAutoConsoleVariable<int32> CVarOnePassProjectionSkipScreenShadowMask(
 	TEXT("Should generally be left enabled outside of debugging."),
 	ECVF_RenderThreadSafe);
 
+static TAutoConsoleVariable<float> CVarContactShadowsOverrideLength(
+	TEXT("r.ContactShadows.OverrideLength"),
+	-1.0f,
+	TEXT("Allows overriding the contact shadow length for all directional lights.\n")
+	TEXT("Disabled when < 0.\n")
+	TEXT("Should generally be left disabled outside of debugging."),
+	ECVF_RenderThreadSafe);
+
+static TAutoConsoleVariable<bool> CVarContactShadowsOverrideLengthInWS(
+	TEXT("r.ContactShadows.OverrideLengthInWS"),
+	false,
+	TEXT("Whether r.ContactShadows.OverrideLength is in world space units or in screen space units."),
+	ECVF_RenderThreadSafe);
+
+static TAutoConsoleVariable<float> CVarContactShadowsOverrideShadowCastingIntensity(
+	TEXT("r.ContactShadows.OverrideShadowCastingIntensity"),
+	-1.0f,
+	TEXT("Allows overriding the contact shadow casting intensity for all directional lights.\n")
+	TEXT("Disabled when < 0.\n")
+	TEXT("Should generally be left disabled outside of debugging."),
+	ECVF_RenderThreadSafe);
+
+static TAutoConsoleVariable<float> CVarContactShadowsOverrideNonShadowCastingIntensity(
+	TEXT("r.ContactShadows.OverrideNonShadowCastingIntensity"),
+	-1.0f,
+	TEXT("Allows overriding the contact shadow non casting intensity for all directional lights.\n")
+	TEXT("Disabled when < 0.\n")
+	TEXT("Should generally be left disabled outside of debugging."),
+	ECVF_RenderThreadSafe);
+
 #if ENABLE_DEBUG_DISCARD_PROP
 static float GDebugLightDiscardProp = 0.0f;
 static FAutoConsoleVariableRef CVarDebugLightDiscardProp(
@@ -194,6 +224,36 @@ bool ShouldRenderRayTracingShadowsForLight(const FLightSceneInfoCompact& LightIn
 }
 #endif // RHI_RAYTRACING
 
+void GetLightContactShadowParameters(const FLightSceneProxy* Proxy, float& OutLength, bool& bOutLengthInWS, float& OutCastingIntensity, float& OutNonCastingIntensity)
+{
+	OutLength = Proxy->GetContactShadowLength();
+	bOutLengthInWS = Proxy->IsContactShadowLengthInWS();
+	OutCastingIntensity = Proxy->GetContactShadowCastingIntensity();
+	OutNonCastingIntensity = Proxy->GetContactShadowNonCastingIntensity();
+
+	if (CVarContactShadowsOverrideLength.GetValueOnAnyThread() >= 0.0f)
+	{
+		OutLength = CVarContactShadowsOverrideLength.GetValueOnAnyThread();
+		bOutLengthInWS = CVarContactShadowsOverrideLengthInWS.GetValueOnAnyThread();
+	}
+
+	if (CVarContactShadowsOverrideShadowCastingIntensity.GetValueOnAnyThread() >= 0.0f)
+	{
+		OutCastingIntensity = CVarContactShadowsOverrideShadowCastingIntensity.GetValueOnAnyThread();
+	}
+
+	if (CVarContactShadowsOverrideNonShadowCastingIntensity.GetValueOnAnyThread() >= 0.0f)
+	{
+		OutNonCastingIntensity = CVarContactShadowsOverrideNonShadowCastingIntensity.GetValueOnAnyThread();
+	}
+
+	if (!bOutLengthInWS)
+	{
+		// Multiply by 2 for screen space in order to preserve old values after introducing multiply by View.ClipToView[1][1] in shader.
+		OutLength *= 2.0f;
+	}
+}
+
 void FLightFunctionSharedParameters::Bind(const FShaderParameterMap& ParameterMap)
 {
 	LightFunctionParameters.Bind(ParameterMap, TEXT("LightFunctionParameters"));
@@ -535,12 +595,17 @@ FDeferredLightUniformStruct GetDeferredLightParameters(const FSceneView& View, c
 
 	if (ContactShadowsCVar && ContactShadowsCVar->GetValueOnRenderThread() != 0 && View.Family->EngineShowFlags.ContactShadows)
 	{
-		Out.ContactShadowLength = LightSceneInfo.Proxy->GetContactShadowLength();
+		float ContactShadowLength;
+		bool bContactShadowLengthInWS;
+		float ContactShadowCastingIntensity;
+		float ContactShadowNonCastingIntensity;
+		GetLightContactShadowParameters(LightSceneInfo.Proxy, ContactShadowLength, bContactShadowLengthInWS, ContactShadowCastingIntensity, ContactShadowNonCastingIntensity);
+
+		Out.ContactShadowLength = ContactShadowLength;
 		// Sign indicates if contact shadow length is in world space or screen space.
-		// Multiply by 2 for screen space in order to preserve old values after introducing multiply by View.ClipToView[1][1] in shader.
-		Out.ContactShadowLength *= LightSceneInfo.Proxy->IsContactShadowLengthInWS() ? -1.0f : 2.0f;
-		Out.ContactShadowCastingIntensity = LightSceneInfo.Proxy->GetContactShadowCastingIntensity();
-		Out.ContactShadowNonCastingIntensity = LightSceneInfo.Proxy->GetContactShadowNonCastingIntensity();
+		Out.ContactShadowLength *= bContactShadowLengthInWS ? -1.0f : 1.0f;
+		Out.ContactShadowCastingIntensity = ContactShadowCastingIntensity;
+		Out.ContactShadowNonCastingIntensity = ContactShadowNonCastingIntensity;
 	}
 
 	// When rendering reflection captures, the direct lighting of the light is actually the indirect specular from the main view
diff --git a/Engine/Source/Runtime/Renderer/Private/Shadows/ScreenSpaceShadows.cpp b/Engine/Source/Runtime/Renderer/Private/Shadows/ScreenSpaceShadows.cpp
--- a/Engine/Source/Runtime/Renderer/Private/Shadows/ScreenSpaceShadows.cpp
+++ b/Engine/Source/Runtime/Renderer/Private/Shadows/ScreenSpaceShadows.cpp
@@ -15,6 +15,8 @@
 #include "RenderGraphUtils.h"
 #include "PixelShaderUtils.h"
 
+extern void GetLightContactShadowParameters(const FLightSceneProxy* Proxy, float& OutLength, bool& bOutLengthInWS, float& OutCastingIntensity, float& OutNonCastingIntensity);
+
 const int32 GScreenSpaceShadowsTileSizeX = 8;
 const int32 GScreenSpaceShadowsTileSizeY = 8;
 
@@ -128,9 +130,16 @@ void RenderScreenSpaceShadows(
 		LightProxy->GetLightShaderParameters(LightParameters);
 
 		PassParameters->LightDirection = LightParameters.Direction;
-		PassParameters->ContactShadowLength = LightProxy->GetContactShadowLength();
-		PassParameters->bContactShadowLengthInWS = LightProxy->IsContactShadowLengthInWS();
-		PassParameters->ContactShadowCastingIntensity = LightProxy->GetContactShadowCastingIntensity();
+
+		float ContactShadowLength;
+		bool bContactShadowLengthInWS;
+		float ContactShadowCastingIntensity;
+		float ContactShadowNonCastingIntensity;
+		GetLightContactShadowParameters(LightProxy, ContactShadowLength, bContactShadowLengthInWS, ContactShadowCastingIntensity, ContactShadowNonCastingIntensity);
+
+		PassParameters->ContactShadowLength = ContactShadowLength;
+		PassParameters->bContactShadowLengthInWS = bContactShadowLengthInWS;
+		PassParameters->ContactShadowCastingIntensity = ContactShadowCastingIntensity;
 
 		auto ComputeShader = View.ShaderMap->GetShader<FScreenSpaceShadowsCS>();
 
