Fixed mobile directional light selection to really be the first one in the list.

#prelfight none

diff --git a/Engine/Source/Runtime/Renderer/Private/LocalFogVolumeRendering.cpp b/Engine/Source/Runtime/Renderer/Private/LocalFogVolumeRendering.cpp
--- a/Engine/Source/Runtime/Renderer/Private/LocalFogVolumeRendering.cpp
+++ b/Engine/Source/Runtime/Renderer/Private/LocalFogVolumeRendering.cpp
@@ -430,6 +430,7 @@ void CreateViewLocalFogVolumeBufferSRV(const FScene* Scene, FViewInfo& View, FRD
 			{
 				View.LocalFogVolumeViewData.UniformParametersStruct.LocalFogVolumeCommon.DirectionalLightColor			= FVector3f( Light->Proxy->GetSunIlluminanceAccountingForSkyAtmospherePerPixelTransmittance());
 				View.LocalFogVolumeViewData.UniformParametersStruct.LocalFogVolumeCommon.DirectionalLightDirection		= FVector3f(-Light->Proxy->GetDirection());
+				break;
 			}
 		}
 	}
