MeshDrawCommands - cvar to set batch size for cached mesh draw commands processing
It allows to adjust dynamically work for static mesh batches based on requirements

diff --git a/Engine/Source/Runtime/Renderer/Private/PrimitiveSceneInfo.cpp b/Engine/Source/Runtime/Renderer/Private/PrimitiveSceneInfo.cpp
--- a/Engine/Source/Runtime/Renderer/Private/PrimitiveSceneInfo.cpp
+++ b/Engine/Source/Runtime/Renderer/Private/PrimitiveSceneInfo.cpp
@@ -46,6 +46,13 @@ static FAutoConsoleVariableRef CVarDrawCommandsCacheMultithreaded(
 	TEXT("Enable multithreading of draw command caching for static meshes. 0=disabled, 1=enabled (default)"),
 	ECVF_RenderThreadSafe);
 
+static int32 GMeshDrawCommandsBatchSize = 12;
+static FAutoConsoleVariableRef CVarDrawCommandsCacheMultithreadedBatchSize(
+	TEXT("r.MeshDrawCommands.BatchSize"),
+	GMeshDrawCommandsBatchSize,
+	TEXT("Batch size of cache mesh draw commands when multithreading of draw command caching is enabled"),
+	ECVF_RenderThreadSafe);
+
 static int32 GNaniteDrawCommandCacheMultithreaded = 1;
 static FAutoConsoleVariableRef CVarNaniteDrawCommandCacheMultithreaded(
 	TEXT("r.Nanite.MeshDrawCommands.CacheMultithreaded"),
@@ -416,10 +423,10 @@ void FPrimitiveSceneInfo::CacheMeshDrawCommands(FScene* Scene, TArrayView<FPrimi
 
 	// This reduce stuttering in editor by improving balancing of all the 
 	// shadermap processing. Keep it as it is for runtime as the requirements are different.
-	static constexpr int BATCH_SIZE = WITH_EDITOR ? 1 : 64;
+	const int BATCH_SIZE = WITH_EDITOR ? 1 : GMeshDrawCommandsBatchSize;
 	const int NumBatches = (SceneInfos.Num() + BATCH_SIZE - 1) / BATCH_SIZE;
 
-	auto DoWorkLambda = [Scene, SceneInfos](FCachedPassMeshDrawListContext& DrawListContext, int32 Index)
+	auto DoWorkLambda = [Scene, SceneInfos, BATCH_SIZE](FCachedPassMeshDrawListContext& DrawListContext, int32 Index)
 	{
 		SCOPED_NAMED_EVENT(FPrimitiveSceneInfo_CacheMeshDrawCommand, FColor::Green);
 
