Moved GPU scene to happen prior to GDME to increase overlap with visibility tasks.

diff --git a/Engine/Source/Runtime/Renderer/Private/DeferredShadingRenderer.cpp b/Engine/Source/Runtime/Renderer/Private/DeferredShadingRenderer.cpp
--- a/Engine/Source/Runtime/Renderer/Private/DeferredShadingRenderer.cpp
+++ b/Engine/Source/Runtime/Renderer/Private/DeferredShadingRenderer.cpp
@@ -396,7 +396,6 @@ DECLARE_CYCLE_STAT(TEXT("InitViews Intentional Stall"), STAT_InitViews_Intention
 
 DECLARE_CYCLE_STAT(TEXT("DeferredShadingSceneRenderer UpdateDownsampledDepthSurface"), STAT_FDeferredShadingSceneRenderer_UpdateDownsampledDepthSurface, STATGROUP_SceneRendering);
 DECLARE_CYCLE_STAT(TEXT("DeferredShadingSceneRenderer Render Init"), STAT_FDeferredShadingSceneRenderer_Render_Init, STATGROUP_SceneRendering);
-DECLARE_CYCLE_STAT(TEXT("DeferredShadingSceneRenderer FGlobalDynamicVertexBuffer Commit"), STAT_FDeferredShadingSceneRenderer_FGlobalDynamicVertexBuffer_Commit, STATGROUP_SceneRendering);
 DECLARE_CYCLE_STAT(TEXT("DeferredShadingSceneRenderer FXSystem PreRender"), STAT_FDeferredShadingSceneRenderer_FXSystem_PreRender, STATGROUP_SceneRendering);
 DECLARE_CYCLE_STAT(TEXT("DeferredShadingSceneRenderer AllocGBufferTargets"), STAT_FDeferredShadingSceneRenderer_AllocGBufferTargets, STATGROUP_SceneRendering);
 DECLARE_CYCLE_STAT(TEXT("DeferredShadingSceneRenderer DBuffer"), STAT_FDeferredShadingSceneRenderer_DBuffer, STATGROUP_SceneRendering);
@@ -424,7 +423,6 @@ DECLARE_GPU_STAT(AllocateRendertargets);
 DECLARE_GPU_STAT(FrameRenderFinish);
 DECLARE_GPU_STAT(SortLights);
 DECLARE_GPU_STAT(PostRenderOpsFX);
-DECLARE_GPU_STAT(GPUSceneUpdate);
 DECLARE_GPU_STAT_NAMED(Unaccounted, TEXT("[unaccounted]"));
 DECLARE_GPU_STAT(WaterRendering);
 DECLARE_GPU_STAT(HairRendering);
@@ -2369,21 +2367,20 @@ void FDeferredShadingSceneRenderer::BeginInitDynamicShadows(FInitViewTaskDatas&
 	}
 }
 
-void FDeferredShadingSceneRenderer::FinishInitDynamicShadows(FRDGBuilder& GraphBuilder, FDynamicShadowsTaskData*& TaskData, FInstanceCullingManager& InstanceCullingManager, FRDGExternalAccessQueue& ExternalAccessQueue)
+void FDeferredShadingSceneRenderer::FinishInitDynamicShadows(FRDGBuilder& GraphBuilder, FDynamicShadowsTaskData*& TaskData, FInstanceCullingManager& InstanceCullingManager)
 {
 	if (ViewFamily.EngineShowFlags.DynamicShadows && !ViewFamily.EngineShowFlags.HitProxies && !HasRayTracedOverlay(ViewFamily))
 	{
 		// Setup dynamic shadows.
 		if (TaskData)
 		{
-			FSceneRenderer::FinishInitDynamicShadows(GraphBuilder, TaskData, InstanceCullingManager, ExternalAccessQueue);
+			FSceneRenderer::FinishInitDynamicShadows(GraphBuilder, TaskData, InstanceCullingManager);
 		}
 		else
 		{
-			TaskData = InitDynamicShadows(GraphBuilder, InstanceCullingManager, ExternalAccessQueue);
+			TaskData = InitDynamicShadows(GraphBuilder, InstanceCullingManager);
 		}
 
-		SCOPE_CYCLE_COUNTER(STAT_FDeferredShadingSceneRenderer_FGlobalDynamicVertexBuffer_Commit);
 		DynamicReadBufferForShadows.Commit(GraphBuilder.RHICmdList);
 	}
 }
@@ -2987,9 +2984,12 @@ void FDeferredShadingSceneRenderer::Render(FRDGBuilder& GraphBuilder)
 
 	::Substrate::PreInitViews(*Scene);
 
+	FSceneTextures::InitializeViewFamily(GraphBuilder, ViewFamily);
+	FSceneTextures& SceneTextures = GetActiveSceneTextures();
+
 	{
 		RDG_GPU_STAT_SCOPE(GraphBuilder, VisibilityCommands);
-		BeginInitViews(GraphBuilder, SceneTexturesConfig, BasePassDepthStencilAccess, InstanceCullingManager, VirtualTextureUpdater.Get(), InitViewTaskDatas);
+		BeginInitViews(GraphBuilder, SceneTexturesConfig, BasePassDepthStencilAccess, InstanceCullingManager, VirtualTextureUpdater.Get(), ExternalAccessQueue, InitViewTaskDatas);
 	}
 
 #if !UE_BUILD_SHIPPING
@@ -3110,20 +3110,16 @@ void FDeferredShadingSceneRenderer::Render(FRDGBuilder& GraphBuilder)
 			GraphBuilder.SetFlushResourcesRHI();
 		}
 
-		Scene->GPUScene.Update(GraphBuilder, GetSceneUniforms(), *Scene, ExternalAccessQueue);
-
 		for (int32 ViewIndex = 0; ViewIndex < AllViews.Num(); ViewIndex++)
 		{
 			FViewInfo& View = *AllViews[ViewIndex];
 			RDG_GPU_MASK_SCOPE(GraphBuilder, View.GPUMask);
 
-			Scene->GPUScene.UploadDynamicPrimitiveShaderDataForView(GraphBuilder, *Scene, View, ExternalAccessQueue);
+			Scene->GPUScene.UploadDynamicPrimitiveShaderDataForView(GraphBuilder, *Scene, View);
 
 			Scene->GPUScene.DebugRender(GraphBuilder, *Scene, GetSceneUniforms(), View);
 		}
 
-		InstanceCullingManager.BeginDeferredCulling(GraphBuilder, Scene->GPUScene);
-
 		if (Views.Num() > 0)
 		{
 			FViewInfo& View = Views[0];
@@ -3131,9 +3127,6 @@ void FDeferredShadingSceneRenderer::Render(FRDGBuilder& GraphBuilder)
 		}
 	}
 
-	FSceneTextures::InitializeViewFamily(GraphBuilder, ViewFamily);
-	FSceneTextures& SceneTextures = GetActiveSceneTextures();
-
 	const bool bUseGBuffer = IsUsingGBuffers(ShaderPlatform);
 	const bool bShouldRenderVolumetricFog = ShouldRenderVolumetricFog();
 	const bool bShouldRenderLocalFogVolume = ShouldRenderLocalFogVolume(Scene, ViewFamily);
@@ -3240,9 +3233,8 @@ void FDeferredShadingSceneRenderer::Render(FRDGBuilder& GraphBuilder)
 
 	{
 		RDG_RHI_GPU_STAT_SCOPE(GraphBuilder, VisibilityCommands);
-		EndInitViews(GraphBuilder, LumenFrameTemporaries, InstanceCullingManager, ExternalAccessQueue, InitViewTaskDatas);
+		EndInitViews(GraphBuilder, LumenFrameTemporaries, InstanceCullingManager, InitViewTaskDatas);
 
-		SCOPE_CYCLE_COUNTER(STAT_FDeferredShadingSceneRenderer_FGlobalDynamicVertexBuffer_Commit);
 		DynamicReadBufferForInitViews.Commit(GraphBuilder.RHICmdList);
 	}
 
@@ -3422,7 +3414,7 @@ void FDeferredShadingSceneRenderer::Render(FRDGBuilder& GraphBuilder)
 		if (!IsForwardShadingEnabled(ShaderPlatform))
 		{
 			// Dynamic shadows are synced later when using the deferred path to make more headroom for tasks.
-			FinishInitDynamicShadows(GraphBuilder, InitViewTaskDatas.DynamicShadows, InstanceCullingManager, ExternalAccessQueue);
+			FinishInitDynamicShadows(GraphBuilder, InitViewTaskDatas.DynamicShadows, InstanceCullingManager);
 		}
 
 		// Update groom only visible in shadow
diff --git a/Engine/Source/Runtime/Renderer/Private/DeferredShadingRenderer.h b/Engine/Source/Runtime/Renderer/Private/DeferredShadingRenderer.h
--- a/Engine/Source/Runtime/Renderer/Private/DeferredShadingRenderer.h
+++ b/Engine/Source/Runtime/Renderer/Private/DeferredShadingRenderer.h
@@ -553,7 +553,7 @@ private:
 	void PreVisibilityFrameSetup(FRDGBuilder& GraphBuilder);
 
 	void BeginInitDynamicShadows(FInitViewTaskDatas& TaskDatas);
-	void FinishInitDynamicShadows(FRDGBuilder& GraphBuilder, FDynamicShadowsTaskData*& TaskData, FInstanceCullingManager& InstanceCullingManager, FRDGExternalAccessQueue& ExternalAccessQueue);
+	void FinishInitDynamicShadows(FRDGBuilder& GraphBuilder, FDynamicShadowsTaskData*& TaskData, FInstanceCullingManager& InstanceCullingManager);
 
 	void ComputeLightVisibility();
 
@@ -564,13 +564,13 @@ private:
 		FExclusiveDepthStencil::Type BasePassDepthStencilAccess,
 		FInstanceCullingManager& InstanceCullingManager,
 		FVirtualTextureUpdater* VirtualTextureUpdater,
+		FRDGExternalAccessQueue& ExternalAccessQueue,
 		FInitViewTaskDatas& TaskDatas);
 
 	void EndInitViews(
 		FRDGBuilder& GraphBuilder,
 		FLumenSceneFrameTemporaries& FrameTemporaries,
 		FInstanceCullingManager& InstanceCullingManager,
-		FRDGExternalAccessQueue& ExternalAccessQueue,
 		FInitViewTaskDatas& TaskDatas);
 
 	void BeginUpdateLumenSceneTasks(FRDGBuilder& GraphBuilder, FLumenSceneFrameTemporaries& FrameTemporaries);
diff --git a/Engine/Source/Runtime/Renderer/Private/GPUScene.cpp b/Engine/Source/Runtime/Renderer/Private/GPUScene.cpp
--- a/Engine/Source/Runtime/Renderer/Private/GPUScene.cpp
+++ b/Engine/Source/Runtime/Renderer/Private/GPUScene.cpp
@@ -32,6 +32,8 @@
 // Defaults to being disabled, enable using the command line argument: -CsvCategory GPUScene
 CSV_DEFINE_CATEGORY(GPUScene, false);
 
+DEFINE_GPU_STAT(GPUSceneUpdate);
+
 #define LOG_INSTANCE_ALLOCATIONS 0
 
 static void ConstructDefault(FGPUSceneResourceParameters& GPUScene, FRDGBuilder& GraphBuilder)
@@ -632,7 +634,7 @@ void FGPUScene::EndRender()
 	ShaderParameters = {};
 }
 
-void FGPUScene::UpdateGPULights(FRDGBuilder& GraphBuilder, FScene& Scene)
+void FGPUScene::UpdateGPULights(FRDGBuilder& GraphBuilder, FScene& Scene, const UE::Tasks::FTask& PrerequisiteTask)
 {
 	FRDGUploadData<FLightSceneData> LightData(GraphBuilder, FMath::Max(1, Scene.Lights.Num()));
 
@@ -648,7 +650,8 @@ void FGPUScene::UpdateGPULights(FRDGBuilder& GraphBuilder, FScene& Scene)
 				InitLightData(Scene.Lights[Index], bAllowStaticLighting, LightData[Index]);
 			}
 		}
-	});
+
+	}, PrerequisiteTask);
 
 	GraphBuilder.QueueBufferUpload<FLightSceneData>(BufferState.LightDataBuffer, LightData, ERDGInitialDataFlags::NoCopy);
 }
@@ -687,7 +690,7 @@ void FGPUScene::InitLightData(const FLightSceneInfoCompact& LightInfoCompact, bo
 	DataOut.LightTypeAndShadowMapChannelMaskPacked = LightInfo.PackLightTypeAndShadowMapChannelMask(bAllowStaticLighting);
 }
 
-void FGPUScene::UpdateInternal(FRDGBuilder& GraphBuilder, FSceneUniformBuffer& SceneUB, FScene& Scene, FRDGExternalAccessQueue& ExternalAccessQueue)
+void FGPUScene::UpdateInternal(FRDGBuilder& GraphBuilder, FSceneUniformBuffer& SceneUB, FScene& Scene, FRDGExternalAccessQueue& ExternalAccessQueue, IVisibilityTaskData* VisibilityTaskData)
 {
 	LLM_SCOPE_BYTAG(GPUScene);
 
@@ -779,14 +782,22 @@ void FGPUScene::UpdateInternal(FRDGBuilder& GraphBuilder, FSceneUniformBuffer& S
 	PrimitivesToUpdate.Reset();
 	PrimitiveDirtyState.Init(EPrimitiveDirtyState::None, PrimitiveDirtyState.Num());
 
+	UE::Tasks::FTask PrerequisiteAsyncTask;
+
+	// An optimization to schedule the async task to occur after dynamic mesh elements are computed in order to reduce contention in the visibility task graph.
+	if (GraphBuilder.IsParallelSetupEnabled() && VisibilityTaskData)
+	{
+		PrerequisiteAsyncTask = VisibilityTaskData->GetComputeRelevanceTask();
+	}
+
 	{
 		SCOPED_NAMED_EVENT(UpdateGPUScene, FColor::Green);
 		QUICK_SCOPE_CYCLE_COUNTER(STAT_UpdateGPUScene);
 		SCOPE_CYCLE_COUNTER(STAT_UpdateGPUSceneTime);
 
-		UploadGeneral<FUploadDataSourceAdapterScenePrimitives>(GraphBuilder, Scene, ExternalAccessQueue, Adapter);
+		UploadGeneral<FUploadDataSourceAdapterScenePrimitives>(GraphBuilder, Scene, &ExternalAccessQueue, Adapter, PrerequisiteAsyncTask);
 
-		UpdateGPULights(GraphBuilder, Scene);
+		UpdateGPULights(GraphBuilder, Scene, PrerequisiteAsyncTask);
 	}
 }
 
@@ -932,7 +943,7 @@ struct FInstanceBatcher
 };
 
 template<typename FUploadDataSourceAdapter>
-void FGPUScene::UploadGeneral(FRDGBuilder& GraphBuilder, FScene& Scene, FRDGExternalAccessQueue& ExternalAccessQueue, const FUploadDataSourceAdapter& UploadDataSourceAdapter)
+void FGPUScene::UploadGeneral(FRDGBuilder& GraphBuilder, FScene& Scene, FRDGExternalAccessQueue* ExternalAccessQueue, const FUploadDataSourceAdapter& UploadDataSourceAdapter, const UE::Tasks::FTask& PrerequisiteTask)
 {
 	LLM_SCOPE_BYTAG(GPUScene);
 
@@ -1286,7 +1297,8 @@ void FGPUScene::UploadGeneral(FRDGBuilder& GraphBuilder, FScene& Scene, FRDGExte
 		{
 			UnlockIfValid(RHICmdList, Uploader);
 		}
-	});
+
+	}, PrerequisiteTask);
 
 	PrimitiveUploadBuffer.End(GraphBuilder, TaskContext.PrimitiveUploader);
 
@@ -1307,9 +1319,11 @@ void FGPUScene::UploadGeneral(FRDGBuilder& GraphBuilder, FScene& Scene, FRDGExte
 
 	if (TaskContext.bUseNaniteMaterialUploaders)
 	{
+		check(ExternalAccessQueue);
+
 		for (int32 NaniteMeshPassIndex = 0; NaniteMeshPassIndex < ENaniteMeshPass::Num; ++NaniteMeshPassIndex)
 		{
-			Scene.NaniteMaterials[NaniteMeshPassIndex].Finish(GraphBuilder, ExternalAccessQueue, TaskContext.NaniteMaterialUploaders[NaniteMeshPassIndex]);
+			Scene.NaniteMaterials[NaniteMeshPassIndex].Finish(GraphBuilder, *ExternalAccessQueue, TaskContext.NaniteMaterialUploaders[NaniteMeshPassIndex]);
 		}
 	}
 	const uint32 MaxPooledSize = uint32(CVarGPUSceneMaxPooledUploadBufferSize.GetValueOnRenderThread());
@@ -1472,11 +1486,12 @@ struct FUploadDataSourceAdapterDynamicPrimitives
 	TArray<uint32, SceneRenderingAllocator> PrimitivesIds;
 };
 
-void FGPUScene::UploadDynamicPrimitiveShaderDataForViewInternal(FRDGBuilder& GraphBuilder, FScene& Scene, FViewInfo& View, FRDGExternalAccessQueue& ExternalAccessQueue, bool bIsShadowView)
+void FGPUScene::UploadDynamicPrimitiveShaderDataForViewInternal(FRDGBuilder& GraphBuilder, FScene& Scene, FViewInfo& View, bool bIsShadowView)
 {
 	LLM_SCOPE_BYTAG(GPUScene);
 
 	RDG_EVENT_SCOPE(GraphBuilder, "GPUScene.UploadDynamicPrimitiveShaderDataForView");
+	SCOPED_NAMED_EVENT(FGPUScene_UploadDynamicPrimitiveShaderDataForView, FColor::Green);
 
 	ensure(bInBeginEndBlock);
 	ensure(DynamicPrimitivesOffset >= Scene.Primitives.Num());
@@ -1541,7 +1556,7 @@ void FGPUScene::UploadDynamicPrimitiveShaderDataForViewInternal(FRDGBuilder& Gra
 		// Run a pass that clears (Sets ID to invalid) any instances that need it.
 		AddClearInstancesPass(GraphBuilder, Scene.InstanceCullingOcclusionQueryRenderer);
 
-		UploadGeneral<FUploadDataSourceAdapterDynamicPrimitives>(GraphBuilder, Scene, ExternalAccessQueue, UploadAdapter);
+		UploadGeneral<FUploadDataSourceAdapterDynamicPrimitives>(GraphBuilder, Scene, nullptr, UploadAdapter, UE::Tasks::FTask{});
 	}
 
 	FSceneUniformBuffer& SceneUniforms = View.GetSceneUniforms();
@@ -1654,7 +1669,7 @@ void FGPUScene::AddPrimitiveToUpdate(int32 PrimitiveId, EPrimitiveDirtyState Dir
 }
 
 
-void FGPUScene::Update(FRDGBuilder& GraphBuilder, FSceneUniformBuffer& SceneUB, FScene& Scene, FRDGExternalAccessQueue& ExternalAccessQueue)
+void FGPUScene::Update(FRDGBuilder& GraphBuilder, FSceneUniformBuffer& SceneUB, FScene& Scene, FRDGExternalAccessQueue& ExternalAccessQueue, IVisibilityTaskData* VisibilityTaskData)
 {
 	if (bIsEnabled)
 	{
@@ -1662,17 +1677,17 @@ void FGPUScene::Update(FRDGBuilder& GraphBuilder, FSceneUniformBuffer& SceneUB,
 
 		ensure(bInBeginEndBlock);
 		
-		UpdateInternal(GraphBuilder, SceneUB, Scene, ExternalAccessQueue);
+		UpdateInternal(GraphBuilder, SceneUB, Scene, ExternalAccessQueue, VisibilityTaskData);
 	}
 }
 
-void FGPUScene::UploadDynamicPrimitiveShaderDataForView(FRDGBuilder& GraphBuilder, FScene& Scene, FViewInfo& View, FRDGExternalAccessQueue& ExternalAccessQueue, bool bIsShadowView)
+void FGPUScene::UploadDynamicPrimitiveShaderDataForView(FRDGBuilder& GraphBuilder, FScene& Scene, FViewInfo& View, bool bIsShadowView)
 {
 	if (bIsEnabled)
 	{
 		RDG_GPU_MASK_SCOPE(GraphBuilder, FRHIGPUMask::All());
 
-		UploadDynamicPrimitiveShaderDataForViewInternal(GraphBuilder, Scene, View, ExternalAccessQueue, bIsShadowView);
+		UploadDynamicPrimitiveShaderDataForViewInternal(GraphBuilder, Scene, View, bIsShadowView);
 	}
 }
 
diff --git a/Engine/Source/Runtime/Renderer/Private/GPUScene.h b/Engine/Source/Runtime/Renderer/Private/GPUScene.h
--- a/Engine/Source/Runtime/Renderer/Private/GPUScene.h
+++ b/Engine/Source/Runtime/Renderer/Private/GPUScene.h
@@ -23,6 +23,9 @@ class FGPUScene;
 class FGPUSceneDynamicContext;
 class FViewUniformShaderParameters;
 class FInstanceCullingOcclusionQueryRenderer;
+class IVisibilityTaskData;
+
+DECLARE_GPU_STAT_NAMED_EXTERN(GPUSceneUpdate, TEXT("GPUSceneUpdate"))
 
 BEGIN_SHADER_PARAMETER_STRUCT(FGPUSceneResourceParameters, RENDERER_API)
 	SHADER_PARAMETER_RDG_BUFFER_SRV(StructuredBuffer<float4>, GPUSceneInstanceSceneData)
@@ -228,7 +231,7 @@ public:
 	/**
 	 * Upload primitives from View.DynamicPrimitiveCollector.
 	 */
-	void UploadDynamicPrimitiveShaderDataForView(FRDGBuilder& GraphBuilder, FScene& Scene, FViewInfo& View, FRDGExternalAccessQueue& ExternalAccessQueue, bool bIsShadowView = false);
+	void UploadDynamicPrimitiveShaderDataForView(FRDGBuilder& GraphBuilder, FScene& Scene, FViewInfo& View, bool bIsShadowView = false);
 
 	/**
 	 * Modifies the GPUScene specific scene UB parameters to the current versions. Returns true if any of the parameters changed.
@@ -238,7 +241,7 @@ public:
 	/**
 	 * Pull all pending updates from Scene and upload primitive & instance data.
 	 */
-	void Update(FRDGBuilder& GraphBuilder, FSceneUniformBuffer& SceneUB, FScene& Scene, FRDGExternalAccessQueue& ExternalAccessQueue);
+	void Update(FRDGBuilder& GraphBuilder, FSceneUniformBuffer& SceneUB, FScene& Scene, FRDGExternalAccessQueue& ExternalAccessQueue, IVisibilityTaskData* VisibilityTaskData = nullptr);
 
 	/**
 	 * Queue the given primitive for upload to GPU at next call to Update.
@@ -411,18 +414,18 @@ private:
 	 * @parameter Scene may be null, as it is only needed for the Nanite material table update (which is coupled to the Scene at the moment).
 	 */
 	template<typename FUploadDataSourceAdapter>
-	void UploadGeneral(FRDGBuilder& GraphBuilder, FScene& Scene, FRDGExternalAccessQueue& ExternalAccessQueue, const FUploadDataSourceAdapter& UploadDataSourceAdapter);
+	void UploadGeneral(FRDGBuilder& GraphBuilder, FScene& Scene, FRDGExternalAccessQueue* ExternalAccessQueue, const FUploadDataSourceAdapter& UploadDataSourceAdapter, const UE::Tasks::FTask& PrerequisiteTask);
 
 	/**
 	 * Upload scene light data to gpu
 	 */
-	void UpdateGPULights(FRDGBuilder& GraphBuilder, FScene& Scene);
+	void UpdateGPULights(FRDGBuilder& GraphBuilder, FScene& Scene, const UE::Tasks::FTask& PrerequisiteTask);
 
 	static void InitLightData(const FLightSceneInfoCompact& LightInfoCompact, bool bAllowStaticLighting, FLightSceneData& DataOut);
 
-	void UploadDynamicPrimitiveShaderDataForViewInternal(FRDGBuilder& GraphBuilder, FScene& Scene, FViewInfo& View, FRDGExternalAccessQueue& ExternalAccessQueue, bool bIsShadowView);
+	void UploadDynamicPrimitiveShaderDataForViewInternal(FRDGBuilder& GraphBuilder, FScene& Scene, FViewInfo& View, bool bIsShadowView);
 
-	void UpdateInternal(FRDGBuilder& GraphBuilder, FSceneUniformBuffer& SceneUB, FScene& Scene, FRDGExternalAccessQueue& ExternalAccessQueue);
+	void UpdateInternal(FRDGBuilder& GraphBuilder, FSceneUniformBuffer& SceneUB, FScene& Scene, FRDGExternalAccessQueue& ExternalAccessQueue, IVisibilityTaskData* VisibilityTaskData);
 
 	void AddUpdatePrimitiveIdsPass(FRDGBuilder& GraphBuilder, FInstanceGPULoadBalancer& IdOnlyUpdateItems);
 
diff --git a/Engine/Source/Runtime/Renderer/Private/MobileShadingRenderer.cpp b/Engine/Source/Runtime/Renderer/Private/MobileShadingRenderer.cpp
--- a/Engine/Source/Runtime/Renderer/Private/MobileShadingRenderer.cpp
+++ b/Engine/Source/Runtime/Renderer/Private/MobileShadingRenderer.cpp
@@ -646,7 +646,7 @@ void FMobileSceneRenderer::InitViews(
 		Scene->GPUScene.Update(GraphBuilder, GetSceneUniforms(), *Scene, ExternalAccessQueue);
 		for (int32 ViewIndex = 0; ViewIndex < AllViews.Num(); ViewIndex++)
 		{
-			Scene->GPUScene.UploadDynamicPrimitiveShaderDataForView(GraphBuilder, *Scene, *AllViews[ViewIndex], ExternalAccessQueue);
+			Scene->GPUScene.UploadDynamicPrimitiveShaderDataForView(GraphBuilder, *Scene, *AllViews[ViewIndex]);
 		}
 	}
 
@@ -656,7 +656,7 @@ void FMobileSceneRenderer::InitViews(
 		if (bDynamicShadows)
 		{
 			// Setup dynamic shadows.
-			TaskDatas.DynamicShadows = InitDynamicShadows(GraphBuilder, InstanceCullingManager, ExternalAccessQueue);
+			TaskDatas.DynamicShadows = InitDynamicShadows(GraphBuilder, InstanceCullingManager);
 		}
 		else
 		{
diff --git a/Engine/Source/Runtime/Renderer/Private/SceneHitProxyRendering.cpp b/Engine/Source/Runtime/Renderer/Private/SceneHitProxyRendering.cpp
--- a/Engine/Source/Runtime/Renderer/Private/SceneHitProxyRendering.cpp
+++ b/Engine/Source/Runtime/Renderer/Private/SceneHitProxyRendering.cpp
@@ -621,7 +621,7 @@ void FDeferredShadingSceneRenderer::RenderHitProxies(FRDGBuilder& GraphBuilder)
 	FLumenSceneFrameTemporaries LumenFrameTemporaries;
 	FInitViewTaskDatas InitViewTaskDatas(VisibilityTaskData);
 	FRDGExternalAccessQueue ExternalAccessQueue;
-	BeginInitViews(GraphBuilder, SceneTexturesConfig, FExclusiveDepthStencil::DepthWrite_StencilWrite, InstanceCullingManager, nullptr, InitViewTaskDatas);
+	BeginInitViews(GraphBuilder, SceneTexturesConfig, FExclusiveDepthStencil::DepthWrite_StencilWrite, InstanceCullingManager, nullptr, ExternalAccessQueue, InitViewTaskDatas);
 
 	extern TSet<IPersistentViewUniformBufferExtension*> PersistentViewUniformBufferExtensions;
 
@@ -638,14 +638,12 @@ void FDeferredShadingSceneRenderer::RenderHitProxies(FRDGBuilder& GraphBuilder)
 
 	ShaderPrint::BeginViews(GraphBuilder, Views);
 
-	Scene->GPUScene.Update(GraphBuilder, GetSceneUniforms(), *Scene, ExternalAccessQueue);
-
 	for (int32 ViewIndex = 0; ViewIndex < Views.Num(); ViewIndex++)
 	{
-		Scene->GPUScene.UploadDynamicPrimitiveShaderDataForView(GraphBuilder, *Scene, Views[ViewIndex], ExternalAccessQueue);
+		Scene->GPUScene.UploadDynamicPrimitiveShaderDataForView(GraphBuilder, *Scene, Views[ViewIndex]);
 	}
 
-	EndInitViews(GraphBuilder, LumenFrameTemporaries, InstanceCullingManager, ExternalAccessQueue, InitViewTaskDatas);
+	EndInitViews(GraphBuilder, LumenFrameTemporaries, InstanceCullingManager, InitViewTaskDatas);
 
 	ExternalAccessQueue.Submit(GraphBuilder);
 
diff --git a/Engine/Source/Runtime/Renderer/Private/SceneRendering.h b/Engine/Source/Runtime/Renderer/Private/SceneRendering.h
--- a/Engine/Source/Runtime/Renderer/Private/SceneRendering.h
+++ b/Engine/Source/Runtime/Renderer/Private/SceneRendering.h
@@ -2306,8 +2306,8 @@ protected:
 	ERendererOutput GetRendererOutput() const;
 
 	FDynamicShadowsTaskData* BeginInitDynamicShadows( bool bRunningEarly, IVisibilityTaskData* VisibilityTaskData);
-	void FinishInitDynamicShadows(FRDGBuilder& GraphBuilder, FDynamicShadowsTaskData* TaskData, FInstanceCullingManager& InstanceCullingManager, FRDGExternalAccessQueue& ExternalAccessQueue);
-	FDynamicShadowsTaskData* InitDynamicShadows(FRDGBuilder& GraphBuilder, FInstanceCullingManager& InstanceCullingManager, FRDGExternalAccessQueue& ExternalAccessQueue);
+	void FinishInitDynamicShadows(FRDGBuilder& GraphBuilder, FDynamicShadowsTaskData* TaskData, FInstanceCullingManager& InstanceCullingManager);
+	FDynamicShadowsTaskData* InitDynamicShadows(FRDGBuilder& GraphBuilder, FInstanceCullingManager& InstanceCullingManager);
 
 	void CreateDynamicShadows(FDynamicShadowsTaskData& TaskData);
 	void FilterDynamicShadows(FDynamicShadowsTaskData& TaskData);
@@ -2624,7 +2624,7 @@ public:
 
 protected:
 	/** Finds the visible dynamic shadows for each view. */
-	FDynamicShadowsTaskData* InitDynamicShadows(FRDGBuilder& GraphBuilder, FInstanceCullingManager& FInstanceCullingManager, FRDGExternalAccessQueue& ExternalAccessQueue);
+	FDynamicShadowsTaskData* InitDynamicShadows(FRDGBuilder& GraphBuilder, FInstanceCullingManager& FInstanceCullingManager);
 
 	void PrepareViewVisibilityLists();
 
@@ -2823,12 +2823,6 @@ private:
 
 extern FFastVramConfig GFastVRamConfig;
 
-/**
- * Returns whether any shadows in the scene have ray traced shadows enabled. This can be called any time
- * after launching the shadow initialization tasks. It will sync the shadow creation task.
- */
-extern bool HasRayTracedDistanceFieldShadows(const FDynamicShadowsTaskData* TaskData);
-
 /** Returns the array of shadows with distance fields. Call only after finishing shadow initialization. */
 extern TConstArrayView<FProjectedShadowInfo*> GetProjectedDistanceFieldShadows(const FDynamicShadowsTaskData* TaskData);
 
diff --git a/Engine/Source/Runtime/Renderer/Private/SceneVisibility.cpp b/Engine/Source/Runtime/Renderer/Private/SceneVisibility.cpp
--- a/Engine/Source/Runtime/Renderer/Private/SceneVisibility.cpp
+++ b/Engine/Source/Runtime/Renderer/Private/SceneVisibility.cpp
@@ -1283,7 +1283,8 @@ void FRelevancePacket::ComputeRelevance(FDynamicPrimitiveIndexList& DynamicPrimi
 	{
 		FPrimitiveSceneInfo* PrimitiveSceneInfo = Scene.Primitives[BitIndex];
 		FPrimitiveViewRelevance& ViewRelevance = const_cast<FPrimitiveViewRelevance&>(View.PrimitiveViewRelevanceMap[BitIndex]);
-		ViewRelevance = PrimitiveSceneInfo->Proxy->GetViewRelevance(&View);
+		const FPrimitiveSceneProxy* PrimitiveSceneProxy = PrimitiveSceneInfo->Proxy;
+		ViewRelevance = PrimitiveSceneProxy->GetViewRelevance(&View);
 		ViewRelevance.bInitializedThisFrame = true;
 
 		const bool bStaticRelevance = ViewRelevance.bStaticRelevance;
@@ -1296,7 +1297,7 @@ void FRelevancePacket::ComputeRelevance(FDynamicPrimitiveIndexList& DynamicPrimi
 		const bool bTranslucentRelevance = ViewRelevance.HasTranslucency();
 		const bool bHairStrandsRelevance = bHairStrandsEnabled && ViewRelevance.bHairStrands;
 
-		if (View.bIsReflectionCapture && !PrimitiveSceneInfo->Proxy->IsVisibleInReflectionCaptures())
+		if (View.bIsReflectionCapture && !PrimitiveSceneProxy->IsVisibleInReflectionCaptures())
 		{
 			NotDrawRelevant.AddPrim(BitIndex);
 			continue;
@@ -1308,7 +1309,7 @@ void FRelevancePacket::ComputeRelevance(FDynamicPrimitiveIndexList& DynamicPrimi
 			const FPrimitiveBounds& Bounds = Scene.PrimitiveBounds[PrimitiveIndex];
 			const bool bIsPrimitiveDistanceCullFading = View.PrimitiveFadeUniformBufferMap[PrimitiveIndex];
 
-			const int8 CurFirstLODIdx = PrimitiveSceneInfo->Proxy->GetCurrentFirstLODIdx_RenderThread();
+			const int8 CurFirstLODIdx = PrimitiveSceneProxy->GetCurrentFirstLODIdx_RenderThread();
 			check(CurFirstLODIdx >= 0);
 			float MeshScreenSizeSquared = 0;
 			FLODMask LODToRender = ComputeLODForMeshes(PrimitiveSceneInfo->StaticMeshRelevances, View, Bounds.BoxSphereBounds.Origin, Bounds.BoxSphereBounds.SphereRadius, PrimitiveSceneInfo->GpuLodInstanceRadius, ViewData.ForcedLODLevel, MeshScreenSizeSquared, CurFirstLODIdx, ViewData.LODScale);
@@ -1409,8 +1410,6 @@ void FRelevancePacket::ComputeRelevance(FDynamicPrimitiveIndexList& DynamicPrimi
 							{
 								if (ViewRelevance.HasVelocity())
 								{
-									const FPrimitiveSceneProxy* PrimitiveSceneProxy = PrimitiveSceneInfo->Proxy;
-
 									if (FVelocityMeshProcessor::PrimitiveHasVelocityForView(View, PrimitiveSceneProxy))
 									{
 										if (ViewRelevance.bVelocityRelevance &&
@@ -1600,19 +1599,19 @@ void FRelevancePacket::ComputeRelevance(FDynamicPrimitiveIndexList& DynamicPrimi
 
 						if (ViewRelevance.bHasVolumeMaterialDomain)
 						{
-							if (ShouldRenderMeshBatchWithHeterogeneousVolumes(&StaticMesh, PrimitiveSceneInfo->Proxy, View.FeatureLevel))
+							if (ShouldRenderMeshBatchWithHeterogeneousVolumes(&StaticMesh, PrimitiveSceneProxy, View.FeatureLevel))
 							{
 								HeterogeneousVolumesMeshBatches.AddUninitialized(1);
 								FVolumetricMeshBatch& BatchAndProxy = HeterogeneousVolumesMeshBatches.Last();
 								BatchAndProxy.Mesh = &StaticMesh;
-								BatchAndProxy.Proxy = PrimitiveSceneInfo->Proxy;
+								BatchAndProxy.Proxy = PrimitiveSceneProxy;
 							}
 							else
 							{
 								VolumetricMeshBatches.AddUninitialized(1);
 								FVolumetricMeshBatch& BatchAndProxy = VolumetricMeshBatches.Last();
 								BatchAndProxy.Mesh = &StaticMesh;
-								BatchAndProxy.Proxy = PrimitiveSceneInfo->Proxy;
+								BatchAndProxy.Proxy = PrimitiveSceneProxy;
 							}
 						}
 
@@ -1621,17 +1620,17 @@ void FRelevancePacket::ComputeRelevance(FDynamicPrimitiveIndexList& DynamicPrimi
 							SkyMeshBatches.AddUninitialized(1);
 							FSkyMeshBatch& BatchAndProxy = SkyMeshBatches.Last();
 							BatchAndProxy.Mesh = &StaticMesh;
-							BatchAndProxy.Proxy = PrimitiveSceneInfo->Proxy;
+							BatchAndProxy.Proxy = PrimitiveSceneProxy;
 							BatchAndProxy.bVisibleInMainPass = ViewRelevance.bRenderInMainPass;
 							BatchAndProxy.bVisibleInRealTimeSkyCapture = PrimitiveSceneInfo->bVisibleInRealTimeSkyCapture;
 						}
 
-						if (ViewRelevance.HasTranslucency() && PrimitiveSceneInfo->Proxy->SupportsSortedTriangles()) // Need to check material as well
+						if (ViewRelevance.HasTranslucency() && PrimitiveSceneProxy->SupportsSortedTriangles()) // Need to check material as well
 						{
 							SortedTrianglesMeshBatches.AddUninitialized(1);
 							FSortedTrianglesMeshBatch& BatchAndProxy = SortedTrianglesMeshBatches.Last();
 							BatchAndProxy.Mesh = &StaticMesh;
-							BatchAndProxy.Proxy = PrimitiveSceneInfo->Proxy;
+							BatchAndProxy.Proxy = PrimitiveSceneProxy;
 						}
 
 						// FIXME: Now if a primitive has one batch with a decal material all primitive mesh batches will be added as decals
@@ -1641,8 +1640,8 @@ void FRelevancePacket::ComputeRelevance(FDynamicPrimitiveIndexList& DynamicPrimi
 							MeshDecalBatches.AddUninitialized(1);
 							FMeshDecalBatch& BatchAndProxy = MeshDecalBatches.Last();
 							BatchAndProxy.Mesh = &StaticMesh;
-							BatchAndProxy.Proxy = PrimitiveSceneInfo->Proxy;
-							BatchAndProxy.SortKey = PrimitiveSceneInfo->Proxy->GetTranslucencySortPriority();
+							BatchAndProxy.Proxy = PrimitiveSceneProxy;
+							BatchAndProxy.SortKey = PrimitiveSceneProxy->GetTranslucencySortPriority();
 						}
 					}
 
@@ -1663,17 +1662,18 @@ void FRelevancePacket::ComputeRelevance(FDynamicPrimitiveIndexList& DynamicPrimi
 #if WITH_EDITOR
 		auto CollectSelectedNaniteInstanceDraws = [](
 			const FPrimitiveSceneInfo& PrimitiveSceneInfo,
+			const FPrimitiveSceneProxy* PrimitiveSceneProxy,
 			TArray<Nanite::FInstanceDraw>& OutInstanceDraws,
 			TArray<uint32>* OutSelectedInstanceHitProxyIDs,
 			bool bSelectedInstancesOnly
 		)
 		{
-			if (!PrimitiveSceneInfo.Proxy->IsNaniteMesh())
+			if (!PrimitiveSceneProxy->IsNaniteMesh())
 			{
 				return;
 			}
 
-			auto* NaniteProxy = static_cast<const Nanite::FSceneProxyBase*>(PrimitiveSceneInfo.Proxy);
+			auto* NaniteProxy = static_cast<const Nanite::FSceneProxyBase*>(PrimitiveSceneProxy);
 			const int32 MaxInstances = PrimitiveSceneInfo.GetNumInstanceSceneDataEntries();
 			OutInstanceDraws.Reserve(OutInstanceDraws.Num() + MaxInstances);
 
@@ -1723,12 +1723,12 @@ void FRelevancePacket::ComputeRelevance(FDynamicPrimitiveIndexList& DynamicPrimi
 
 		if (bEditorVisualizeLevelInstanceRelevance)
 		{
-			CollectSelectedNaniteInstanceDraws(*PrimitiveSceneInfo, EditorVisualizeLevelInstancesNanite, nullptr, false);
+			CollectSelectedNaniteInstanceDraws(*PrimitiveSceneInfo, PrimitiveSceneProxy, EditorVisualizeLevelInstancesNanite, nullptr, false);
 		}
 
 		if (bEditorSelectionRelevance)
 		{
-			CollectSelectedNaniteInstanceDraws(*PrimitiveSceneInfo, EditorSelectedInstancesNanite, &EditorSelectedNaniteHitProxyIds, true);
+			CollectSelectedNaniteInstanceDraws(*PrimitiveSceneInfo, PrimitiveSceneProxy, EditorSelectedInstancesNanite, &EditorSelectedNaniteHitProxyIds, true);
 		}
 #endif
 
@@ -4870,6 +4870,7 @@ void FDeferredShadingSceneRenderer::BeginInitViews(
 	FExclusiveDepthStencil::Type BasePassDepthStencilAccess,
 	FInstanceCullingManager& InstanceCullingManager,
 	FVirtualTextureUpdater* VirtualTextureUpdater,
+	FRDGExternalAccessQueue& ExternalAccessQueue,
 	FInitViewTaskDatas& TaskDatas)
 {
 	SCOPED_NAMED_EVENT(FDeferredShadingSceneRenderer_InitViews, FColor::Emerald);
@@ -4919,6 +4920,16 @@ void FDeferredShadingSceneRenderer::BeginInitViews(
 
 	LumenScenePDIVisualization();
 
+	// Process GPU scene prior to visibility to maximize overlap.
+	{
+		RDG_CSV_STAT_EXCLUSIVE_SCOPE(GraphBuilder, UpdateGPUScene);
+		RDG_GPU_STAT_SCOPE(GraphBuilder, GPUSceneUpdate);
+
+		Scene->GPUScene.Update(GraphBuilder, GetSceneUniforms(), *Scene, ExternalAccessQueue, TaskDatas.VisibilityTaskData);
+
+		InstanceCullingManager.BeginDeferredCulling(GraphBuilder, Scene->GPUScene);
+	}
+
 	TaskDatas.VisibilityTaskData->ProcessRenderThreadTasks(BasePassDepthStencilAccess, InstanceCullingManager, VirtualTextureUpdater);
 
 	// Make a second attempt to launch shadow tasks it wasn't able to the first time due to visibility being deferred.
@@ -5045,7 +5056,6 @@ void FDeferredShadingSceneRenderer::EndInitViews(
 	FRDGBuilder& GraphBuilder,
 	FLumenSceneFrameTemporaries& FrameTemporaries,
 	FInstanceCullingManager& InstanceCullingManager,
-	FRDGExternalAccessQueue& ExternalAccessQueue,
 	FInitViewTaskDatas& TaskDatas)
 {
 	SCOPED_NAMED_EVENT(FDeferredShadingSceneRenderer_InitViewsAfterPrepass, FColor::Emerald);
@@ -5084,7 +5094,7 @@ void FDeferredShadingSceneRenderer::EndInitViews(
 	if (IsForwardShadingEnabled(ShaderPlatform))
 	{
 		// Dynamic shadows are synced earlier when forward shading is enabled.
-		FinishInitDynamicShadows(GraphBuilder, TaskDatas.DynamicShadows, InstanceCullingManager, ExternalAccessQueue);
+		FinishInitDynamicShadows(GraphBuilder, TaskDatas.DynamicShadows, InstanceCullingManager);
 	}
 }
 
diff --git a/Engine/Source/Runtime/Renderer/Private/ShadowSetup.cpp b/Engine/Source/Runtime/Renderer/Private/ShadowSetup.cpp
--- a/Engine/Source/Runtime/Renderer/Private/ShadowSetup.cpp
+++ b/Engine/Source/Runtime/Renderer/Private/ShadowSetup.cpp
@@ -1355,16 +1355,6 @@ TConstArrayView<FProjectedShadowInfo*> GetProjectedDistanceFieldShadows(const FD
 	return TConstArrayView<FProjectedShadowInfo*>();
 }
 
-bool HasRayTracedDistanceFieldShadows(const FDynamicShadowsTaskData* TaskData)
-{
-	if (TaskData)
-	{
-		TaskData->WaitForCreateDynamicShadowsTask();
-		return TaskData->bHasRayTracedDistanceFieldShadows;
-	}
-	return false;
-}
-
 struct FAddSubjectPrimitiveOverflowedIndices
 {
 	TArray<uint16> MDCIndices;
@@ -6249,7 +6239,7 @@ FDynamicShadowsTaskData* FSceneRenderer::BeginInitDynamicShadows(bool bRunningEa
 }
 
 
-void FSceneRenderer::FinishInitDynamicShadows(FRDGBuilder& GraphBuilder, FDynamicShadowsTaskData* TaskData, FInstanceCullingManager& InstanceCullingManager, FRDGExternalAccessQueue& ExternalAccessQueue)
+void FSceneRenderer::FinishInitDynamicShadows(FRDGBuilder& GraphBuilder, FDynamicShadowsTaskData* TaskData, FInstanceCullingManager& InstanceCullingManager)
 {
 	SCOPED_NAMED_EVENT_TEXT("FSceneRenderer::FinishInitDynamicShadows", FColor::Magenta);
 	SCOPE_CYCLE_COUNTER(STAT_InitDynamicShadowsTime);
@@ -6281,39 +6271,39 @@ void FSceneRenderer::FinishInitDynamicShadows(FRDGBuilder& GraphBuilder, FDynami
 	{
 		for (FProjectedShadowInfo* ProjectedShadowInfo : ShadowMapAtlas.Shadows)
 		{
-			Scene->GPUScene.UploadDynamicPrimitiveShaderDataForView(GraphBuilder, *Scene, *ProjectedShadowInfo->ShadowDepthView, ExternalAccessQueue, true);
+			Scene->GPUScene.UploadDynamicPrimitiveShaderDataForView(GraphBuilder, *Scene, *ProjectedShadowInfo->ShadowDepthView, true);
 		}
 	}
 	for (FSortedShadowMapAtlas& ShadowMap : SortedShadowsForShadowDepthPass.ShadowMapCubemaps)
 	{
 		check(ShadowMap.Shadows.Num() == 1);
 		FProjectedShadowInfo* ProjectedShadowInfo = ShadowMap.Shadows[0];
-		Scene->GPUScene.UploadDynamicPrimitiveShaderDataForView(GraphBuilder, *Scene, *ProjectedShadowInfo->ShadowDepthView, ExternalAccessQueue, true);
+		Scene->GPUScene.UploadDynamicPrimitiveShaderDataForView(GraphBuilder, *Scene, *ProjectedShadowInfo->ShadowDepthView, true);
 	}
 	for (FProjectedShadowInfo* ProjectedShadowInfo : SortedShadowsForShadowDepthPass.PreshadowCache.Shadows)
 	{
 		if (!ProjectedShadowInfo->bDepthsCached)
 		{
-			Scene->GPUScene.UploadDynamicPrimitiveShaderDataForView(GraphBuilder, *Scene, *ProjectedShadowInfo->ShadowDepthView, ExternalAccessQueue, true);
+			Scene->GPUScene.UploadDynamicPrimitiveShaderDataForView(GraphBuilder, *Scene, *ProjectedShadowInfo->ShadowDepthView, true);
 		}
 	}
 	for (const FSortedShadowMapAtlas& ShadowMapAtlas : SortedShadowsForShadowDepthPass.TranslucencyShadowMapAtlases)
 	{
 		for (FProjectedShadowInfo* ProjectedShadowInfo : ShadowMapAtlas.Shadows)
 		{
-			Scene->GPUScene.UploadDynamicPrimitiveShaderDataForView(GraphBuilder, *Scene, *ProjectedShadowInfo->ShadowDepthView, ExternalAccessQueue, true);
+			Scene->GPUScene.UploadDynamicPrimitiveShaderDataForView(GraphBuilder, *Scene, *ProjectedShadowInfo->ShadowDepthView, true);
 		}
 	}
 	for (FProjectedShadowInfo* ProjectedShadowInfo : SortedShadowsForShadowDepthPass.VirtualShadowMapShadows)
 	{
-		Scene->GPUScene.UploadDynamicPrimitiveShaderDataForView(GraphBuilder, *Scene, *ProjectedShadowInfo->ShadowDepthView, ExternalAccessQueue, true);
+		Scene->GPUScene.UploadDynamicPrimitiveShaderDataForView(GraphBuilder, *Scene, *ProjectedShadowInfo->ShadowDepthView, true);
 	}
 }
 
-FDynamicShadowsTaskData* FSceneRenderer::InitDynamicShadows(FRDGBuilder& GraphBuilder, FInstanceCullingManager& InstanceCullingManager, FRDGExternalAccessQueue& ExternalAccessQueue)
+FDynamicShadowsTaskData* FSceneRenderer::InitDynamicShadows(FRDGBuilder& GraphBuilder, FInstanceCullingManager& InstanceCullingManager)
 {
 	FDynamicShadowsTaskData* TaskData = BeginInitDynamicShadows(false, nullptr);
-	FinishInitDynamicShadows(GraphBuilder, TaskData, InstanceCullingManager, ExternalAccessQueue);
+	FinishInitDynamicShadows(GraphBuilder, TaskData, InstanceCullingManager);
 	return TaskData;
 }
 
diff --git a/Engine/Source/Runtime/Renderer/Private/ShadowSetupMobile.cpp b/Engine/Source/Runtime/Renderer/Private/ShadowSetupMobile.cpp
--- a/Engine/Source/Runtime/Renderer/Private/ShadowSetupMobile.cpp
+++ b/Engine/Source/Runtime/Renderer/Private/ShadowSetupMobile.cpp
@@ -284,7 +284,7 @@ static void VisualizeMobileDynamicCSMSubjectCapsules(FViewInfo& View, FLightScen
 }
 
 /** Finds the visible dynamic shadows for each view. */
-FDynamicShadowsTaskData* FMobileSceneRenderer::InitDynamicShadows(FRDGBuilder& GraphBuilder, FInstanceCullingManager& InstanceCullingManager, FRDGExternalAccessQueue& ExternalAccessQueue)
+FDynamicShadowsTaskData* FMobileSceneRenderer::InitDynamicShadows(FRDGBuilder& GraphBuilder, FInstanceCullingManager& InstanceCullingManager)
 {
 	static auto* MyCVarMobileEnableStaticAndCSMShadowReceivers = IConsoleManager::Get().FindTConsoleVariableDataInt(TEXT("r.Mobile.EnableStaticAndCSMShadowReceivers"));
 	const bool bCombinedStaticAndCSMEnabled = MyCVarMobileEnableStaticAndCSMShadowReceivers->GetValueOnRenderThread()!=0;
@@ -308,7 +308,7 @@ FDynamicShadowsTaskData* FMobileSceneRenderer::InitDynamicShadows(FRDGBuilder& G
 		}
 	}
 
-	FDynamicShadowsTaskData* TaskData = FSceneRenderer::InitDynamicShadows(GraphBuilder, InstanceCullingManager, ExternalAccessQueue);
+	FDynamicShadowsTaskData* TaskData = FSceneRenderer::InitDynamicShadows(GraphBuilder, InstanceCullingManager);
 
 	bool bAlwaysUseCSM = false;
 	const bool bSkipCSMShaderCulling = MobileBasePassAlwaysUsesCSM(Scene->GetShaderPlatform());
