Optimized cached shadow map size access (i.e. only iterate once through shadow linked list) and modified shadow render targets to take a non-immediate command list.

diff --git a/Engine/Source/Runtime/Renderer/Private/SceneRendering.h b/Engine/Source/Runtime/Renderer/Private/SceneRendering.h
--- a/Engine/Source/Runtime/Renderer/Private/SceneRendering.h
+++ b/Engine/Source/Runtime/Renderer/Private/SceneRendering.h
@@ -2369,21 +2369,22 @@ protected:
 		TArray<FProjectedShadowInfo*, SceneRenderingAllocator>& ShadowInfosThatNeedCulling, 
 		FVisibleLightInfo& VisibleLightInfo, 
 		FLightSceneInfo& LightSceneInfo,
+		int64 CachedShadowMapsSize,
 		uint32& NumCSMCachesUpdatedThisFrame);
 
 	void AllocateShadowDepthTargets(FRHICommandListImmediate& RHICmdList, const FDynamicShadowsTaskData& TaskData);
 	
-	void AllocateAtlasedShadowDepthTargets(FRHICommandListImmediate& RHICmdList, TConstArrayView<FProjectedShadowInfo*> Shadows, TArray<FSortedShadowMapAtlas,SceneRenderingAllocator>& OutAtlases);
+	void AllocateAtlasedShadowDepthTargets(FRHICommandListBase& RHICmdList, TConstArrayView<FProjectedShadowInfo*> Shadows, TArray<FSortedShadowMapAtlas,SceneRenderingAllocator>& OutAtlases);
 
-	void AllocateCachedShadowDepthTargets(FRHICommandListImmediate& RHICmdList, TConstArrayView<FProjectedShadowInfo*> CachedShadows);
+	void AllocateCachedShadowDepthTargets(FRHICommandListBase& RHICmdList, TConstArrayView<FProjectedShadowInfo*> CachedShadows);
 
-	void AllocateCSMDepthTargets(FRHICommandListImmediate& RHICmdList, TConstArrayView<FProjectedShadowInfo*> WholeSceneDirectionalShadows, TArray<FSortedShadowMapAtlas, SceneRenderingAllocator>& OutAtlases);
+	void AllocateCSMDepthTargets(FRHICommandListBase& RHICmdList, TConstArrayView<FProjectedShadowInfo*> WholeSceneDirectionalShadows, TArray<FSortedShadowMapAtlas, SceneRenderingAllocator>& OutAtlases);
 
-	void AllocateOnePassPointLightDepthTargets(FRHICommandListImmediate& RHICmdList, TConstArrayView<FProjectedShadowInfo*> WholeScenePointShadows);
+	void AllocateOnePassPointLightDepthTargets(FRHICommandListBase& RHICmdList, TConstArrayView<FProjectedShadowInfo*> WholeScenePointShadows);
 
-	void AllocateTranslucentShadowDepthTargets(FRHICommandListImmediate& RHICmdList, TConstArrayView<FProjectedShadowInfo*> TranslucentShadows);
+	void AllocateTranslucentShadowDepthTargets(FRHICommandListBase& RHICmdList, TConstArrayView<FProjectedShadowInfo*> TranslucentShadows);
 
-	void AllocateMobileCSMAndSpotLightShadowDepthTargets(FRHICommandListImmediate& RHICmdList, TConstArrayView<FProjectedShadowInfo*> MobileCSMAndSpotLightShadows);
+	void AllocateMobileCSMAndSpotLightShadowDepthTargets(FRHICommandListBase& RHICmdList, TConstArrayView<FProjectedShadowInfo*> MobileCSMAndSpotLightShadows);
 	/**
 	* Used by RenderLights to figure out if projected shadows need to be rendered to the attenuation buffer.
 	* Or to render a given shadowdepth map for forward rendering.
@@ -2406,7 +2407,7 @@ protected:
 	* Creates a projected shadow for all primitives affected by a light.
 	* @param LightSceneInfo - The light to create a shadow for.
 	*/
-	void CreateWholeSceneProjectedShadow(FDynamicShadowsTaskData& TaskData, FLightSceneInfo* LightSceneInfo, uint32& NumPointShadowCachesUpdatedThisFrame, uint32& NumSpotShadowCachesUpdatedThisFrame);
+	void CreateWholeSceneProjectedShadow(FDynamicShadowsTaskData& TaskData, FLightSceneInfo* LightSceneInfo, int64 CachedShadowMapsSize, uint32& NumPointShadowCachesUpdatedThisFrame, uint32& NumSpotShadowCachesUpdatedThisFrame);
 
 	/** Updates the preshadow cache, allocating new preshadows that can fit and evicting old ones. */
 	void UpdatePreshadowCache();
diff --git a/Engine/Source/Runtime/Renderer/Private/ShadowSetup.cpp b/Engine/Source/Runtime/Renderer/Private/ShadowSetup.cpp
--- a/Engine/Source/Runtime/Renderer/Private/ShadowSetup.cpp
+++ b/Engine/Source/Runtime/Renderer/Private/ShadowSetup.cpp
@@ -375,10 +375,8 @@ static TAutoConsoleVariable<bool> CVarCSMScissorOptim(
  * Helper function to maintain the common login of caching the whole scene shadow and be able to handle the specific-purpose through the CacheLambda and UnCacheLambda
  */
 template <typename CacheLambdaType, typename UnCacheLambdaType>
-void TryToCacheShadowMap(const FScene* Scene, int32& OutNumShadowMaps, EShadowDepthCacheMode* OutCacheModes, uint32& NumCachesUpdatedThisFrame, CacheLambdaType&& CacheLambda, UnCacheLambdaType&& UnCacheLambda)
+void TryToCacheShadowMap(const FScene* Scene, int64 CachedShadowMapsSize, int32& OutNumShadowMaps, EShadowDepthCacheMode* OutCacheModes, uint32& NumCachesUpdatedThisFrame, CacheLambdaType&& CacheLambda, UnCacheLambdaType&& UnCacheLambda)
 {
-	int64 CachedShadowMapsSize = Scene->GetCachedWholeSceneShadowMapsSize();
-
 	if (CachedShadowMapsSize < static_cast<int64>(GWholeSceneShadowCacheMb) * 1024 * 1024)
 	{
 		OutNumShadowMaps = 2;
@@ -1217,28 +1215,30 @@ struct FDrawDebugShadowFrustumOp
 	FProjectedShadowInfo* ProjectedShadowInfo = nullptr;
 };
 
+using FProjectedShadowInfoList = TArray<FProjectedShadowInfo*, SceneRenderingAllocator>;
+
 struct FFilteredShadowArrays
 {
 	FFilteredShadowArrays() = default;
 
 	// Sort visible shadows based on their allocation needs
 	// 2d shadowmaps for this frame only that can be atlased across lights
-	TArray<FProjectedShadowInfo*, SceneRenderingAllocator> Shadows;
+	FProjectedShadowInfoList Shadows;
 	// 2d shadowmaps that will persist across frames, can't be atlased
-	TArray<FProjectedShadowInfo*, SceneRenderingAllocator> CachedSpotlightShadows;
-	TArray<FProjectedShadowInfo*, SceneRenderingAllocator> TranslucentShadows;
+	FProjectedShadowInfoList CachedSpotlightShadows;
+	FProjectedShadowInfoList TranslucentShadows;
 	// 2d shadowmaps that persist across frames
-	TArray<FProjectedShadowInfo*, SceneRenderingAllocator> CachedPreShadows;
+	FProjectedShadowInfoList CachedPreShadows;
 	// Cubemaps, can't be atlased
-	TArray<FProjectedShadowInfo*, SceneRenderingAllocator> WholeScenePointShadows;
+	FProjectedShadowInfoList WholeScenePointShadows;
 
-	TArray<FProjectedShadowInfo*, SceneRenderingAllocator> WholeSceneDirectionalShadows;
-	TArray<FProjectedShadowInfo*, SceneRenderingAllocator> CachedWholeSceneDirectionalShadows;
+	FProjectedShadowInfoList WholeSceneDirectionalShadows;
+	FProjectedShadowInfoList CachedWholeSceneDirectionalShadows;
 
 	/** Distance field shadows to project. Used to avoid iterating through the scene lights array. */
 	TArray<FProjectedShadowInfo*, TInlineAllocator<2, SceneRenderingAllocator>> ProjectedDistanceFieldShadows;
 
-	TArray<FProjectedShadowInfo*, SceneRenderingAllocator> ShadowsToSetupViews;
+	FProjectedShadowInfoList ShadowsToSetupViews;
 
 #if !(UE_BUILD_SHIPPING || UE_BUILD_TEST)
 	TBitArray<SceneRenderingAllocator> OnePassShadowUnsupportedLights;
@@ -3473,6 +3473,7 @@ void ComputeWholeSceneShadowCacheModes(
 	FScene* Scene,
 	bool bNeedsVirtualShadowMap,
 	FWholeSceneProjectedShadowInitializer& InOutProjectedShadowInitializer,
+	const int64 CachedShadowMapsSize,
 	FIntPoint& InOutShadowMapSize,
 	uint32& InOutNumPointShadowCachesUpdatedThisFrame,
 	uint32& InOutNumSpotShadowCachesUpdatedThisFrame,
@@ -3521,7 +3522,7 @@ void ComputeWholeSceneShadowCacheModes(
 				}
 				else
 				{
-					TryToCacheShadowMap(Scene, OutNumShadowMaps, OutCacheModes, *NumCachesUpdatedThisFrame,
+					TryToCacheShadowMap(Scene, CachedShadowMapsSize, OutNumShadowMaps, OutCacheModes, *NumCachesUpdatedThisFrame,
 					[&]()
 					{
 						// Check if update is caused by resolution change
@@ -3583,7 +3584,7 @@ void ComputeWholeSceneShadowCacheModes(
 		}
 		else
 		{
-			TryToCacheShadowMap(Scene, OutNumShadowMaps, OutCacheModes, *NumCachesUpdatedThisFrame,
+			TryToCacheShadowMap(Scene, CachedShadowMapsSize, OutNumShadowMaps, OutCacheModes, *NumCachesUpdatedThisFrame,
 			[&]()
 			{
 				Scene->CachedShadowMaps.Add(LightSceneInfo->Id).Add(FCachedShadowMapData(InOutProjectedShadowInitializer, RealTime));
@@ -3619,6 +3620,7 @@ void ComputeViewDependentWholeSceneShadowCacheModes(
 	float RealTime,
 	FScene* Scene,
 	const FWholeSceneProjectedShadowInitializer& ProjectedShadowInitializer,
+	const int64 CachedShadowMapsSize,
 	const FIntPoint& ShadowMapSize,
 	uint32& InOutNumCSMCachesUpdatedThisFrame,
 	int32& OutNumShadowMaps,
@@ -3702,7 +3704,7 @@ void ComputeViewDependentWholeSceneShadowCacheModes(
 				}
 				else
 				{
-					TryToCacheShadowMap(Scene, OutNumShadowMaps, OutCacheModes, InOutNumCSMCachesUpdatedThisFrame,
+					TryToCacheShadowMap(Scene, CachedShadowMapsSize, OutNumShadowMaps, OutCacheModes, InOutNumCSMCachesUpdatedThisFrame,
 					[&]()
 					{
 						CachedShadowMapData->ShadowBufferResolution = ShadowMapSize;
@@ -3728,7 +3730,7 @@ void ComputeViewDependentWholeSceneShadowCacheModes(
 		}
 		else
 		{
-			TryToCacheShadowMap(Scene, OutNumShadowMaps, OutCacheModes, InOutNumCSMCachesUpdatedThisFrame,
+			TryToCacheShadowMap(Scene, CachedShadowMapsSize, OutNumShadowMaps, OutCacheModes, InOutNumCSMCachesUpdatedThisFrame,
 			[&]()
 			{
 				if (CachedShadowMapDatas == nullptr)
@@ -3878,6 +3880,7 @@ bool IntersectsConvexHulls(FLightViewFrustumConvexHulls const& ConvexHulls, FBox
 void FSceneRenderer::CreateWholeSceneProjectedShadow(
 	FDynamicShadowsTaskData& TaskData,
 	FLightSceneInfo* LightSceneInfo,
+	int64 CachedShadowMapsSize,
 	uint32& InOutNumPointShadowCachesUpdatedThisFrame,
 	uint32& InOutNumSpotShadowCachesUpdatedThisFrame)
 {
@@ -4022,6 +4025,7 @@ void FSceneRenderer::CreateWholeSceneProjectedShadow(
 						bNeedsVirtualShadowMap,
 						// Below are in-out or out parameters. They can change
 						ProjectedShadowInitializer,
+						CachedShadowMapsSize,
 						ShadowMapSize,
 						InOutNumPointShadowCachesUpdatedThisFrame,
 						InOutNumSpotShadowCachesUpdatedThisFrame,
@@ -5072,6 +5076,7 @@ void FSceneRenderer::AddViewDependentWholeSceneShadowsForView(
 	TArray<FProjectedShadowInfo*, SceneRenderingAllocator>& ShadowInfosThatNeedCulling,
 	FVisibleLightInfo& VisibleLightInfo, 
 	FLightSceneInfo& LightSceneInfo,
+	const int64 CachedShadowMapsSize,
 	uint32& NumCSMCachesUpdatedThisFrame)
 {
 	SCOPE_CYCLE_COUNTER(STAT_AddViewDependentWholeSceneShadowsForView);
@@ -5160,6 +5165,7 @@ void FSceneRenderer::AddViewDependentWholeSceneShadowsForView(
 								ViewFamily.Time.GetRealTimeSeconds(),
 								Scene,
 								ProjectedShadowInitializer,
+								CachedShadowMapsSize,
 								ShadowBufferResolution,
 								NumCSMCachesUpdatedThisFrame,
 								NumShadowMaps,
@@ -5399,7 +5405,7 @@ struct FLayoutAndAssignedShadows
 };
 
 void FSceneRenderer::AllocateAtlasedShadowDepthTargets(
-	FRHICommandListImmediate& RHICmdList,
+	FRHICommandListBase& RHICmdList,
 	TConstArrayView<FProjectedShadowInfo*> Shadows,
 	TArray<FSortedShadowMapAtlas,SceneRenderingAllocator>& OutAtlases)
 {
@@ -5521,7 +5527,7 @@ public:
 	}
 };
 
-void FSceneRenderer::AllocateCachedShadowDepthTargets(FRHICommandListImmediate& RHICmdList, TConstArrayView<FProjectedShadowInfo*> CachedShadows)
+void FSceneRenderer::AllocateCachedShadowDepthTargets(FRHICommandListBase& RHICmdList, TConstArrayView<FProjectedShadowInfo*> CachedShadows)
 {
 	for (int32 ShadowIndex = 0; ShadowIndex < CachedShadows.Num(); ShadowIndex++)
 	{
@@ -5551,7 +5557,7 @@ void FSceneRenderer::AllocateCachedShadowDepthTargets(FRHICommandListImmediate&
 }
 
 void FSceneRenderer::AllocateCSMDepthTargets(
-	FRHICommandListImmediate& RHICmdList,
+	FRHICommandListBase& RHICmdList,
 	TConstArrayView<FProjectedShadowInfo*> WholeSceneDirectionalShadows,
 	TArray<FSortedShadowMapAtlas, SceneRenderingAllocator>& OutAtlases
 	)
@@ -5611,7 +5617,7 @@ void FSceneRenderer::AllocateCSMDepthTargets(
 	}
 }
 
-void FSceneRenderer::AllocateOnePassPointLightDepthTargets(FRHICommandListImmediate& RHICmdList, TConstArrayView<FProjectedShadowInfo*> WholeScenePointShadows)
+void FSceneRenderer::AllocateOnePassPointLightDepthTargets(FRHICommandListBase& RHICmdList, TConstArrayView<FProjectedShadowInfo*> WholeScenePointShadows)
 {
 	if (FeatureLevel >= ERHIFeatureLevel::SM5)
 	{
@@ -5669,7 +5675,7 @@ TCHAR* const GetTranslucencyShadowTransmissionName(uint32 Id)
 	return (TCHAR*)TEXT("InvalidName");
 }
 
-void FSceneRenderer::AllocateTranslucentShadowDepthTargets(FRHICommandListImmediate& RHICmdList, TConstArrayView<FProjectedShadowInfo*> TranslucentShadows)
+void FSceneRenderer::AllocateTranslucentShadowDepthTargets(FRHICommandListBase& RHICmdList, TConstArrayView<FProjectedShadowInfo*> TranslucentShadows)
 {
 	if (TranslucentShadows.Num() > 0 && FeatureLevel >= ERHIFeatureLevel::SM5)
 	{
@@ -5762,6 +5768,18 @@ void FSceneRenderer::CreateDynamicShadows(FDynamicShadowsTaskData& TaskData)
 
 		if (GetShadowQuality() > 0)
 		{
+			int64 CachedShadowMapsSize = -1;
+
+			const auto GetCachedShadowMapsSize = [&]
+			{
+				if (CachedShadowMapsSize < 0)
+				{
+					// This call is quite expensive, so compute it lazily.
+					CachedShadowMapsSize = Scene->GetCachedWholeSceneShadowMapsSize();
+				}
+				return CachedShadowMapsSize;
+			};
+
 			for (auto LightIt = Scene->Lights.CreateConstIterator(); LightIt; ++LightIt)
 			{
 				const FLightSceneInfoCompact& LightSceneInfoCompact = *LightIt;
@@ -5856,7 +5874,7 @@ void FSceneRenderer::CreateDynamicShadows(FDynamicShadowsTaskData& TaskData)
 						if (bCreateShadowForMovableLight || bCreateShadowToPreviewStaticLight || bCreateShadowForOverflowStaticShadowing)
 						{
 							// Try to create a whole scene projected shadow.
-							CreateWholeSceneProjectedShadow(TaskData, LightSceneInfo, NumPointShadowCachesUpdatedThisFrame, NumSpotShadowCachesUpdatedThisFrame);
+							CreateWholeSceneProjectedShadow(TaskData, LightSceneInfo, GetCachedShadowMapsSize(), NumPointShadowCachesUpdatedThisFrame, NumSpotShadowCachesUpdatedThisFrame);
 						}
 
 						// Register visible lights for allowing hair strands to cast shadow (non-directional light)
@@ -5874,7 +5892,7 @@ void FSceneRenderer::CreateDynamicShadows(FDynamicShadowsTaskData& TaskData)
 									// Mobile uses the scene's MobileDirectionalLights only for whole scene shadows.
 									&& (LightSceneInfo == Scene->MobileDirectionalLights[0] || LightSceneInfo == Scene->MobileDirectionalLights[1] || LightSceneInfo == Scene->MobileDirectionalLights[2])))
 							{
-								AddViewDependentWholeSceneShadowsForView(ViewDependentWholeSceneShadows, ViewDependentWholeSceneShadowsThatNeedCulling, VisibleLightInfo, *LightSceneInfo, NumCSMCachesUpdatedThisFrame);
+								AddViewDependentWholeSceneShadowsForView(ViewDependentWholeSceneShadows, ViewDependentWholeSceneShadowsThatNeedCulling, VisibleLightInfo, *LightSceneInfo, GetCachedShadowMapsSize(), NumCSMCachesUpdatedThisFrame);
 							}
 
 							if (!bMobile || (LightSceneInfo->Proxy->CastsModulatedShadows() && !LightSceneInfo->Proxy->UseCSMForDynamicObjects() && LightSceneInfo->Proxy->HasStaticShadowing()))
@@ -6307,7 +6325,7 @@ FDynamicShadowsTaskData* FSceneRenderer::InitDynamicShadows(FRDGBuilder& GraphBu
 	return TaskData;
 }
 
-void FSceneRenderer::AllocateMobileCSMAndSpotLightShadowDepthTargets(FRHICommandListImmediate& RHICmdList, TConstArrayView<FProjectedShadowInfo*> MobileCSMAndSpotLightShadows)
+void FSceneRenderer::AllocateMobileCSMAndSpotLightShadowDepthTargets(FRHICommandListBase& RHICmdList, TConstArrayView<FProjectedShadowInfo*> MobileCSMAndSpotLightShadows)
 {
 	if (MobileCSMAndSpotLightShadows.Num() > 0)
 	{
