Enable ShaderPrint and GPUScene debugger on mobile renderer.

diff --git a/Engine/Source/Runtime/Renderer/Private/MobileShadingRenderer.cpp b/Engine/Source/Runtime/Renderer/Private/MobileShadingRenderer.cpp
--- a/Engine/Source/Runtime/Renderer/Private/MobileShadingRenderer.cpp
+++ b/Engine/Source/Runtime/Renderer/Private/MobileShadingRenderer.cpp
@@ -37,6 +37,7 @@
 #include "IXRTrackingSystem.h"
 #include "SceneViewExtension.h"
 #include "ScreenRendering.h"
+#include "ShaderPrint.h"
 #include "PipelineStateCache.h"
 #include "ClearQuad.h"
 #include "MobileSeparateTranslucencyPass.h"
@@ -646,7 +647,9 @@ void FMobileSceneRenderer::InitViews(
 		Scene->GPUScene.Update(GraphBuilder, GetSceneUniforms(), *Scene, ExternalAccessQueue);
 		for (int32 ViewIndex = 0; ViewIndex < AllViews.Num(); ViewIndex++)
 		{
-			Scene->GPUScene.UploadDynamicPrimitiveShaderDataForView(GraphBuilder, *Scene, *AllViews[ViewIndex]);
+			FViewInfo& View = *AllViews[ViewIndex];
+			Scene->GPUScene.UploadDynamicPrimitiveShaderDataForView(GraphBuilder, *Scene, View);
+			Scene->GPUScene.DebugRender(GraphBuilder, *Scene, GetSceneUniforms(), View);
 		}
 	}
 
@@ -871,6 +874,13 @@ void FMobileSceneRenderer::Render(FRDGBuilder& GraphBuilder)
 
 	FRDGSystemTextures::Create(GraphBuilder);
 
+	ShaderPrint::BeginViews(GraphBuilder, Views);
+
+	ON_SCOPE_EXIT
+	{
+		ShaderPrint::EndViews(Views);
+	};
+
 	TUniquePtr<FVirtualTextureUpdater> VirtualTextureUpdater;
 
 	if (bUseVirtualTexturing)
diff --git a/Engine/Source/Runtime/Renderer/Private/PostProcess/PostProcessing.cpp b/Engine/Source/Runtime/Renderer/Private/PostProcess/PostProcessing.cpp
--- a/Engine/Source/Runtime/Renderer/Private/PostProcess/PostProcessing.cpp
+++ b/Engine/Source/Runtime/Renderer/Private/PostProcess/PostProcessing.cpp
@@ -2779,6 +2779,11 @@ void AddMobilePostProcessingPasses(FRDGBuilder& GraphBuilder, FScene* Scene, con
 		}
 	}
 
+	if (ShaderPrint::IsEnabled(View.ShaderPrintData))
+	{
+		ShaderPrint::DrawView(GraphBuilder, View, SceneColor, SceneDepth);
+	}
+
 	if (PassSequence.IsEnabled(EPass::HMDDistortion))
 	{
 		FHMDDistortionInputs PassInputs;
diff --git a/Engine/Source/Runtime/Renderer/Private/ShaderPrint.cpp b/Engine/Source/Runtime/Renderer/Private/ShaderPrint.cpp
--- a/Engine/Source/Runtime/Renderer/Private/ShaderPrint.cpp
+++ b/Engine/Source/Runtime/Renderer/Private/ShaderPrint.cpp
@@ -178,6 +178,39 @@ namespace ShaderPrint
 	FBufferWithRDG* GEmptyBuffer = new TGlobalResource<FEmptyBuffer>();
 
 	//////////////////////////////////////////////////////////////////////////////////////////////////
+	// Global fixed index buffers
+
+	class FLineIndexBuffer : public FIndexBuffer
+	{
+	public:
+		void InitRHI(FRHICommandListBase& RHICmdList) override
+		{
+			FRHIResourceCreateInfo CreateInfo(TEXT("FLineIndexBuffer"));
+			IndexBufferRHI = RHICmdList.CreateIndexBuffer(sizeof(uint16), sizeof(uint16) * 2, BUF_Static, CreateInfo);
+			void* VoidPtr = RHICmdList.LockBuffer(IndexBufferRHI, 0, sizeof(uint16) * 2, RLM_WriteOnly);
+			static const uint16 Indices[] = { 0, 1 };
+			FMemory::Memcpy(VoidPtr, Indices, 2 * sizeof(uint16));
+			RHICmdList.UnlockBuffer(IndexBufferRHI);
+		}
+	};
+	TGlobalResource<FLineIndexBuffer, FRenderResource::EInitPhase::Pre> GLineIndexBuffer;
+
+	class FTriangleIndexBuffer : public FIndexBuffer
+	{
+	public:
+		void InitRHI(FRHICommandListBase& RHICmdList) override
+		{
+			FRHIResourceCreateInfo CreateInfo(TEXT("FTriangleIndexBuffer"));
+			IndexBufferRHI = RHICmdList.CreateIndexBuffer(sizeof(uint16), sizeof(uint16) * 3, BUF_Static, CreateInfo);
+			void* VoidPtr = RHICmdList.LockBuffer(IndexBufferRHI, 0, sizeof(uint16) * 3, RLM_WriteOnly);
+			static const uint16 Indices[] = { 0, 1, 2 };
+			FMemory::Memcpy(VoidPtr, Indices, 3 * sizeof(uint16));
+			RHICmdList.UnlockBuffer(IndexBufferRHI);
+		}
+	};
+	TGlobalResource<FTriangleIndexBuffer, FRenderResource::EInitPhase::Pre> GTriangleIndexBuffer;
+
+		//////////////////////////////////////////////////////////////////////////////////////////////////
 	// Uniform buffer
 	
 	// ShaderPrint uniform buffer
@@ -240,7 +273,9 @@ namespace ShaderPrint
 
 	bool IsSupported(EShaderPlatform InShaderPlatform)
 	{
-		return !IsMobilePlatform(InShaderPlatform) && !IsHlslccShaderPlatform(InShaderPlatform);
+		// Should be supported everywhere?
+		// If we want to restrict platforms then we can use existing logic here or add a DataDrivenPlatformInfo.
+		return true;
 	}
 
 	bool IsEnabled()
@@ -1107,7 +1142,7 @@ namespace ShaderPrint
 				// Marks the indirect draw parameter as used by the pass, given it's not used directly by any of the shaders.
 				FRHIBuffer* IndirectBufferRHI = PassParameters->VS.IndirectBuffer->GetIndirectRHICallBuffer();
 				check(IndirectBufferRHI != nullptr);
-				RHICmdList.DrawPrimitiveIndirect(IndirectBufferRHI, 0);
+				RHICmdList.DrawIndexedPrimitiveIndirect(bLines ? GLineIndexBuffer.IndexBufferRHI : GTriangleIndexBuffer.IndexBufferRHI, IndirectBufferRHI, 0);
 			});
 	}
 
