GPULM: Fix GPULM hit group shaders not compiling when Substrate is disabled.

diff --git a/Engine/Source/Runtime/Renderer/Private/PathTracing.cpp b/Engine/Source/Runtime/Renderer/Private/PathTracing.cpp
--- a/Engine/Source/Runtime/Renderer/Private/PathTracing.cpp
+++ b/Engine/Source/Runtime/Renderer/Private/PathTracing.cpp
@@ -1445,20 +1445,20 @@ public:
 		{
 			// only need to compile the intersection shader permutation if the VF actually requires it
 			return false;
-		}
-		if (SimplifySubstrate && (!Substrate::IsSubstrateEnabled() || CVarPathTracingSubstrateCompileSimplifiedMaterial.GetValueOnAnyThread() == 0))
-		{
-			// don't compile the extra Substrate permutation if:
-			//    Substrate is not enabled on this project
-			// or the user did not request the extra permutations to be compiled (default)
-			return false;
-		}
+		}		
 		if (IsGPULightmass)
 		{
 			return ShouldCompileGPULightmassShadersForProject(Parameters);
 		}
 		else
 		{
+			if (SimplifySubstrate && (!Substrate::IsSubstrateEnabled() || CVarPathTracingSubstrateCompileSimplifiedMaterial.GetValueOnAnyThread() == 0))
+			{
+				// don't compile the extra Substrate permutation if:
+				//    Substrate is not enabled on this project
+				// or the user did not request the extra permutations to be compiled (default)
+				return false;
+			}
 			return ShouldCompilePathTracingShadersForProject(Parameters.Platform);
 		}
 	}
@@ -1525,6 +1525,7 @@ using FPathTracingMaterialSimplifiedCHS_AHS_IS = TPathTracingMaterial<true , tru
 
 
 // NOTE: lightmass doesn't work with intersection shader VFs at the moment, so avoid instantiating permutations that will never generate any shaders
+// Also lightmass is always using simplified Substrate mode.
 using FGPULightmassCHS               = TPathTracingMaterial<false, false, true, true>;
 using FGPULightmassCHS_AHS           = TPathTracingMaterial<true , false, true, true>;
 
