-Change r.Shaders.ForceDXC cvar to a FShaderPlatformCachedIniValue to be able to read it per platform when cooking.
-Enable DXC shader compiler for console

diff --git a/Engine/Source/Runtime/RenderCore/Private/Shader.cpp b/Engine/Source/Runtime/RenderCore/Private/Shader.cpp
--- a/Engine/Source/Runtime/RenderCore/Private/Shader.cpp
+++ b/Engine/Source/Runtime/RenderCore/Private/Shader.cpp
@@ -1432,8 +1432,8 @@ bool IsDxcEnabledForPlatform(EShaderPlatform Platform, bool bHlslVersion2021)
 	// Check the generic console variable first (if DXC is supported)
 	if (FDataDrivenShaderPlatformInfo::GetSupportsDxc(Platform))
 	{
-		static const IConsoleVariable* CVar = IConsoleManager::Get().FindConsoleVariable(TEXT("r.Shaders.ForceDXC"));
-		if (bHlslVersion2021 || (CVar && CVar->GetInt() != 0))
+		static FShaderPlatformCachedIniValue<bool> ShaderForceDXC(TEXT("r.Shaders.ForceDXC"));
+		if (bHlslVersion2021 || (ShaderForceDXC.Get(Platform) != 0))
 		{
 			return true;
 		}
diff --git a/Engine/Source/Runtime/Renderer/Private/DepthRendering.h b/Engine/Source/Runtime/Renderer/Private/DepthRendering.h
--- a/Engine/Source/Runtime/Renderer/Private/DepthRendering.h
+++ b/Engine/Source/Runtime/Renderer/Private/DepthRendering.h
@@ -100,12 +100,6 @@ public:
 	static void ModifyCompilationEnvironment(const FMaterialShaderPermutationParameters& Parameters, FShaderCompilerEnvironment& OutEnvironment)
 	{
 		FMeshMaterialShader::ModifyCompilationEnvironment(Parameters, OutEnvironment);
-
-		// @lh-todo: Same workaround as for the VS of MobileBasePass. See TMobileBasePassVSPolicyParamType::ModifyCompilationEnvironment for details.
-		if (!Substrate::IsSubstrateEnabled())
-		{
-			OutEnvironment.SetCompileArgument(TEXT("WORKAROUND_DISABLE_rShadersForceDXC"), true);
-		}
 	}
 
 	void GetShaderBindings(
@@ -153,12 +147,6 @@ public:
 		
 		OutEnvironment.SetDefine(TEXT("ALLOW_DEBUG_VIEW_MODES"), AllowDebugViewmodes(Parameters.Platform));
 		OutEnvironment.SetDefine(TEXT("SCENE_TEXTURES_DISABLED"), 1u);
-
-		// @lh-todo: Same workaround as for the VS of MobileBasePass. See TMobileBasePassVSPolicyParamType::ModifyCompilationEnvironment for details.
-		if (!Substrate::IsSubstrateEnabled())
-		{
-			OutEnvironment.SetCompileArgument(TEXT("WORKAROUND_DISABLE_rShadersForceDXC"), true);
-		}
 	}
 
 	FDepthOnlyPS() {}
diff --git a/Engine/Source/Runtime/Renderer/Private/MobileBasePassRendering.h b/Engine/Source/Runtime/Renderer/Private/MobileBasePassRendering.h
--- a/Engine/Source/Runtime/Renderer/Private/MobileBasePassRendering.h
+++ b/Engine/Source/Runtime/Renderer/Private/MobileBasePassRendering.h
@@ -152,14 +152,6 @@ public:
 	static void ModifyCompilationEnvironment(const FMaterialShaderPermutationParameters& Parameters, FShaderCompilerEnvironment& OutEnvironment)
 	{
 		FMeshMaterialShader::ModifyCompilationEnvironment(Parameters, OutEnvironment);
-
-		// @lh-todo:
-		//		This is a workaround to prevent Z-fighting due to a mismatch between HLSLCC compiled shaders (from cooked materials) and DXC compiled shaders (from user materials).
-		//		It effectively disables CVar "r.Shaders.ForceDXC" which we partially enabled during the HLSLCC to DXC transition for certain platform backends.
-		if (!Substrate::IsSubstrateEnabled())
-		{
-			OutEnvironment.SetCompileArgument(TEXT("WORKAROUND_DISABLE_rShadersForceDXC"), true);
-		}
 	}
 
 	void GetShaderBindings(
@@ -252,12 +244,6 @@ public:
 		FMeshMaterialShader::ModifyCompilationEnvironment(Parameters, OutEnvironment);
 		// Modify compilation environment depending upon material shader quality level settings.
 		ModifyCompilationEnvironmentForQualityLevel(Parameters.Platform, Parameters.MaterialParameters.QualityLevel, OutEnvironment);
-
-		// @lh-todo: Same workaround as for the VS of MobileBasePass. See TMobileBasePassVSPolicyParamType::ModifyCompilationEnvironment for details.
-		if (!Substrate::IsSubstrateEnabled())
-		{
-			OutEnvironment.SetCompileArgument(TEXT("WORKAROUND_DISABLE_rShadersForceDXC"), true);
-		}
 	}
 
 	/** Initialization constructor. */
