Fix gbuffer picking shader complie fail.
Don't want this on mobile platforms.
It used to be that ShaderPrint disabled this, but that changed.

diff --git a/Engine/Source/Runtime/Renderer/Private/PostProcess/PostProcessing.cpp b/Engine/Source/Runtime/Renderer/Private/PostProcess/PostProcessing.cpp
--- a/Engine/Source/Runtime/Renderer/Private/PostProcess/PostProcessing.cpp
+++ b/Engine/Source/Runtime/Renderer/Private/PostProcess/PostProcessing.cpp
@@ -2923,7 +2923,7 @@ class FGBufferPickingCS : public FGlobalShader
 public:
 	static bool ShouldCompilePermutation(const FGlobalShaderPermutationParameters& Parameters) 
 	{ 
-		return ShaderPrint::IsSupported(Parameters.Platform) && !Substrate::IsSubstrateEnabled() && 
+		return ShaderPrint::IsSupported(Parameters.Platform) && !IsMobilePlatform(Parameters.Platform) && !Substrate::IsSubstrateEnabled() &&
 			EnumHasAllFlags(Parameters.Flags, EShaderPermutationFlags::HasEditorOnlyData);
 	}
 	static void ModifyCompilationEnvironment(const FGlobalShaderPermutationParameters& Parameters, FShaderCompilerEnvironment& OutEnvironment)
