LFV - packing the data to read less data (reducing instance bytes from 144 to 128 bytes)

diff --git a/Engine/Shaders/Private/LocalFogVolumes/LocalFogVolumeCommon.ush b/Engine/Shaders/Private/LocalFogVolumes/LocalFogVolumeCommon.ush
--- a/Engine/Shaders/Private/LocalFogVolumes/LocalFogVolumeCommon.ush
+++ b/Engine/Shaders/Private/LocalFogVolumes/LocalFogVolumeCommon.ush
@@ -15,15 +15,19 @@ struct FLocalFogVolumeGPUInstanceData
 	float4x3 Transform;
 	float4x3 InvTransform;
 
-	float RadialFogExtinction;
-	float HeightFogExtinction;
-	float HeightFogFalloff;
-	float HeightFogOffset;
-
-	float3 Albedo;
-	float  PhaseG;
-	float3 Emissive;
-	float  UniformScale;
+	// float4
+	float RadialFogExtinction;		// half1
+	float HeightFogExtinction;		// half1
+	float HeightFogFalloff;			// half1
+	float HeightFogOffset;			// half1
+	float3 Emissive;				// half3
+	//FREE							// half1
+
+	// float4 (note: albedo could be using half representation or even 3x8bits if needed later to gain space)
+	float  UniformScale;			// float
+	float3 Albedo;					// half3
+	float  PhaseG;					// half1
+	// FREE							// float2
 };
 
 
@@ -76,24 +80,32 @@ Texture2DArray<uint>					LFV_LocalFogVolumeTileDataTexture;
 FLocalFogVolumeGPUInstanceData GetLocalFogVolumeGPUInstanceData(uint Index)
 {
 	FLocalFogVolumeGPUInstanceData Data;
-	uint Offset = Index * 9;
+	uint Offset = Index * 8;
 
 	Data.Transform			= float4x3(LFVInsts[Offset + 0 ], LFVInsts[Offset + 1 ], LFVInsts[Offset + 2 ]); // reads 4x3 row major pack as 3 float4
 	Data.InvTransform		= float4x3(LFVInsts[Offset + 3 ], LFVInsts[Offset + 4 ], LFVInsts[Offset + 5 ]);
 
-	float4 Temp0			= LFVInsts[Offset + 6];
-	Data.RadialFogExtinction= Temp0.x;
-	Data.HeightFogExtinction= Temp0.y;
-	Data.HeightFogFalloff	= Temp0.z;
-	Data.HeightFogOffset	= Temp0.w;
+	float4 Data0			= LFVInsts[Offset + 6];
 
-	float4 Temp1			= LFVInsts[Offset + 7];
-	Data.Albedo				= Temp1.xyz;
-	Data.PhaseG				= Temp1.w;
+	float2 Data0X			= UnpackFloat2FromUInt(asuint(Data0.x));
+	float2 Data0Y			= UnpackFloat2FromUInt(asuint(Data0.y));
+	float2 Data0Z			= UnpackFloat2FromUInt(asuint(Data0.z));
+	float2 Data0W			= UnpackFloat2FromUInt(asuint(Data0.w));
 
-	float4 Temp2			= LFVInsts[Offset + 8];
-	Data.Emissive			= Temp2.xyz;
-	Data.UniformScale		= Temp2.w;
+	Data.RadialFogExtinction= Data0X.x;
+	Data.HeightFogExtinction= Data0X.y;
+	Data.HeightFogFalloff	= Data0Y.x;
+	Data.HeightFogOffset	= Data0Y.y;
+	Data.Emissive			= float3(Data0Z.x, Data0Z.y, Data0W.x);
+
+	float4 Data1			= LFVInsts[Offset + 7];
+
+	float2 Data1X			= UnpackFloat2FromUInt(asuint(Data1.y));
+	float2 Data1Y			= UnpackFloat2FromUInt(asuint(Data1.z));
+
+	Data.UniformScale		= Data1.x;
+	Data.Albedo				= float3(Data1X.x, Data1X.y, Data1Y.x);
+	Data.PhaseG				= Data1Y.y;
 
 	return Data;
 }
diff --git a/Engine/Source/Runtime/Core/Public/Math/Vector2DHalf.h b/Engine/Source/Runtime/Core/Public/Math/Vector2DHalf.h
--- a/Engine/Source/Runtime/Core/Public/Math/Vector2DHalf.h
+++ b/Engine/Source/Runtime/Core/Public/Math/Vector2DHalf.h
@@ -64,6 +64,12 @@ public:
 	void SetTruncate( const FVector2f& Vector2D );
 	void SetTruncate( const FVector2d& Vector2D );
 
+	FORCEINLINE uint32 AsUInt32()
+	{
+		uint32 PackedData = (X.Encoded & 0xFFFFFFFF) | (Y.Encoded << 16);
+		return PackedData;
+	}
+
 public:
 
 	/**
diff --git a/Engine/Source/Runtime/Renderer/Private/LocalFogVolumeRendering.cpp b/Engine/Source/Runtime/Renderer/Private/LocalFogVolumeRendering.cpp
--- a/Engine/Source/Runtime/Renderer/Private/LocalFogVolumeRendering.cpp
+++ b/Engine/Source/Runtime/Renderer/Private/LocalFogVolumeRendering.cpp
@@ -321,23 +321,31 @@ void GetLocalFogVolumeSortingData(const FScene* Scene, FRDGBuilder& GraphBuilder
 		ConvertFromMatrix44fTo4x3Array(Transform, LocalFogVolumeGPUInstanceDataIt->Transform);
 		ConvertFromMatrix44fTo4x3Array(InvTransform, LocalFogVolumeGPUInstanceDataIt->InvTransform);
 
-		LocalFogVolumeGPUInstanceDataIt->RadialFogExtinction = LHF->RadialFogExtinction;
-		LocalFogVolumeGPUInstanceDataIt->HeightFogExtinction = LHF->HeightFogExtinction;
-		LocalFogVolumeGPUInstanceDataIt->HeightFogFalloff = LHF->HeightFogFalloff * 0.01f;	// This scale is used to have artist author reasonable range.
-		LocalFogVolumeGPUInstanceDataIt->HeightFogOffset = LHF->HeightFogOffset;
+		FVector2DHalf Data0X = FVector2DHalf(LHF->RadialFogExtinction,			LHF->HeightFogExtinction);
+		FVector2DHalf Data0Y = FVector2DHalf(LHF->HeightFogFalloff * 0.01f,		LHF->HeightFogOffset);
+		FVector2DHalf Data0Z = FVector2DHalf(LHF->FogEmissive.R,				LHF->FogEmissive.G);
+		FVector2DHalf Data0W = FVector2DHalf(LHF->FogEmissive.B,				0.0f);
 
-		LocalFogVolumeGPUInstanceDataIt->UniformScale = LHF->FogUniformScale;
+		LocalFogVolumeGPUInstanceDataIt->Data0[0] = Data0X.AsUInt32();
+		LocalFogVolumeGPUInstanceDataIt->Data0[1] = Data0Y.AsUInt32();
+		LocalFogVolumeGPUInstanceDataIt->Data0[2] = Data0Z.AsUInt32();
+		LocalFogVolumeGPUInstanceDataIt->Data0[3] = Data0W.AsUInt32();
+
+		FVector2DHalf Data1X = FVector2DHalf(LHF->FogAlbedo.R, LHF->FogAlbedo.G);
+		FVector2DHalf Data1Y = FVector2DHalf(LHF->FogAlbedo.B, LHF->FogPhaseG);
 
-		LocalFogVolumeGPUInstanceDataIt->Albedo = FVector3f(LHF->FogAlbedo);
-		LocalFogVolumeGPUInstanceDataIt->PhaseG = LHF->FogPhaseG;
-		LocalFogVolumeGPUInstanceDataIt->Emissive = FVector3f(LHF->FogEmissive);
+		LocalFogVolumeGPUInstanceDataIt->UniformScale = LHF->FogUniformScale;
+		LocalFogVolumeGPUInstanceDataIt->Data1[0] = Data1X.AsUInt32();
+		LocalFogVolumeGPUInstanceDataIt->Data1[1] = Data1Y.AsUInt32();
+		LocalFogVolumeGPUInstanceDataIt->Data1[2] = 0;
 
+		// Register the sorting data
 		Out.LocalFogVolumeCenterPos[Out.LocalFogVolumeInstanceCountFinal] = LHF->FogTransform.GetTranslation();
 
-		FLocalFogVolumeSortKey* LocalFogVolumeSortKeysIt = &Out.LocalFogVolumeSortKeys[Out.LocalFogVolumeInstanceCountFinal];
-		LocalFogVolumeSortKeysIt->FogVolume.Index = Out.LocalFogVolumeInstanceCountFinal;
-		LocalFogVolumeSortKeysIt->FogVolume.Distance = 0;	// Filled up right before sorting according to a view
-		LocalFogVolumeSortKeysIt->FogVolume.Priority = LHF->FogSortPriority;
+		FLocalFogVolumeSortKey* LocalFogVolumeSortKeysIt= &Out.LocalFogVolumeSortKeys[Out.LocalFogVolumeInstanceCountFinal];
+		LocalFogVolumeSortKeysIt->FogVolume.Index		= Out.LocalFogVolumeInstanceCountFinal;
+		LocalFogVolumeSortKeysIt->FogVolume.Distance	= 0;	// Filled up right before sorting according to a view
+		LocalFogVolumeSortKeysIt->FogVolume.Priority	= LHF->FogSortPriority;
 
 		Out.LocalFogVolumeInstanceCountFinal++;
 	}
diff --git a/Engine/Source/Runtime/Renderer/Private/LocalFogVolumeRendering.h b/Engine/Source/Runtime/Renderer/Private/LocalFogVolumeRendering.h
--- a/Engine/Source/Runtime/Renderer/Private/LocalFogVolumeRendering.h
+++ b/Engine/Source/Runtime/Renderer/Private/LocalFogVolumeRendering.h
@@ -46,15 +46,10 @@ public:
 	float Transform[4 * 3];
 	float InvTransform[4 * 3];
 
-	float RadialFogExtinction;
-	float HeightFogExtinction;
-	float HeightFogFalloff;
-	float HeightFogOffset;
-
-	FVector3f Albedo;
-	float PhaseG;
-	FVector3f Emissive;
-	float UniformScale;
+	// See LocalFogVolumeCommon for the encoded representation
+	uint32 Data0[4];
+	float  UniformScale;
+	uint32 Data1[3];
 };
 
 class RENDERER_API FLocalFogVolumeSortKey
