Fix hash collisions in raytracing autoinstanced mesh batches by including the shader hash

This was noticed on the substrate content sample, though in theory the issue could manifest in other situations as well. The hash used by autoinstancing only used the bindings and not the actual shader being used. So materials with identical paramters but different shading models or other settings could in theory result in collisions.

diff --git a/Engine/Source/Runtime/Renderer/Private/PrimitiveSceneInfo.cpp b/Engine/Source/Runtime/Renderer/Private/PrimitiveSceneInfo.cpp
--- a/Engine/Source/Runtime/Renderer/Private/PrimitiveSceneInfo.cpp
+++ b/Engine/Source/Runtime/Renderer/Private/PrimitiveSceneInfo.cpp
@@ -982,8 +982,17 @@ void CacheRayTracingMeshBatch(
 	if (bMustEmitCommand || CommandContext.CommandIndex >= 0)
 	{
 		uint64& Hash = SceneInfo->CachedRayTracingMeshCommandsHashPerLOD[MeshBatch.LODIndex];
-		Hash <<= 1;
-		Hash ^= Commands[CommandContext.CommandIndex].ShaderBindings.GetDynamicInstancingHash();
+
+		// We want the hash to change if either the shader or the binding contents change. This is used by the autoinstance feature.
+		const FRHIShader* Shader = Commands[CommandContext.CommandIndex].MaterialShader;
+
+		// TODO: It would be better to use 64 bits for both of these to reduce the chance of hash collisions
+		//       but GetDynamicInstancingHash is currently a public function, so changing the return type would be an API change
+		uint32 ShaderHash = Shader != nullptr ? GetTypeHash(Shader->GetHash()) : 0;
+		uint32 ShaderBindingsHash = Commands[CommandContext.CommandIndex].ShaderBindings.GetDynamicInstancingHash();
+
+		Hash <<= 1; // TODO: It would probably be better to use some kind of proper 64 bit mix here?
+		Hash ^= (uint64(ShaderBindingsHash) << 32) | uint64(ShaderHash);
 
 		if (bDeferLODCommandIndices)
 		{
