LFV - reduces instance memory footprint from 128 to 80bytes.
Removed the non tiled shading path.
Tested mobile.

#prelfight none

diff --git a/Engine/Shaders/Private/LocalFogVolumes/LocalFogVolumeCommon.ush b/Engine/Shaders/Private/LocalFogVolumes/LocalFogVolumeCommon.ush
--- a/Engine/Shaders/Private/LocalFogVolumes/LocalFogVolumeCommon.ush
+++ b/Engine/Shaders/Private/LocalFogVolumes/LocalFogVolumeCommon.ush
@@ -12,7 +12,6 @@
 
 struct FLocalFogVolumeGPUInstanceData
 {
-	float4x3 Transform;
 	float4x3 InvTransform;
 
 	// float4
@@ -80,12 +79,11 @@ Texture2DArray<uint>					LFV_LocalFogVolumeTileDataTexture;
 FLocalFogVolumeGPUInstanceData GetLocalFogVolumeGPUInstanceData(uint Index)
 {
 	FLocalFogVolumeGPUInstanceData Data;
-	uint Offset = Index * 8;
+	uint Offset = Index * 5;
 
-	Data.Transform			= float4x3(LFVInsts[Offset + 0 ], LFVInsts[Offset + 1 ], LFVInsts[Offset + 2 ]); // reads 4x3 row major pack as 3 float4
-	Data.InvTransform		= float4x3(LFVInsts[Offset + 3 ], LFVInsts[Offset + 4 ], LFVInsts[Offset + 5 ]);
+	Data.InvTransform		= float4x3(LFVInsts[Offset + 0 ], LFVInsts[Offset + 1 ], LFVInsts[Offset + 2 ]); // reads 4x3 row major pack as 3 float4
 
-	float4 Data0			= LFVInsts[Offset + 6];
+	float4 Data0			= LFVInsts[Offset + 3];
 
 	float2 Data0X			= UnpackFloat2FromUInt(asuint(Data0.x));
 	float2 Data0Y			= UnpackFloat2FromUInt(asuint(Data0.y));
@@ -98,7 +96,7 @@ FLocalFogVolumeGPUInstanceData GetLocalFogVolumeGPUInstanceData(uint Index)
 	Data.HeightFogOffset	= Data0Y.y;
 	Data.Emissive			= float3(Data0Z.x, Data0Z.y, Data0W.x);
 
-	float4 Data1			= LFVInsts[Offset + 7];
+	float4 Data1			= LFVInsts[Offset + 4];
 
 	float2 Data1X			= UnpackFloat2FromUInt(asuint(Data1.y));
 	float2 Data1Y			= UnpackFloat2FromUInt(asuint(Data1.z));
@@ -243,10 +241,6 @@ float3 LocalFogVolumeEvaluateInScattering(in FLocalFogVolumeGPUInstanceData FogI
 ///////////////////////////////////////////////////////////////////////////////
 // Local Fog Volume Instance Evaluation
 
-#ifndef LFV_CAN_CLIP
-#define LFV_CAN_CLIP 0
-#endif
-
 // Returns color transmittance
 float4 GetLFVInstanceContribution(
 	uint InstanceIndex,
@@ -295,18 +289,6 @@ float4 GetLFVInstanceContribution(
 			LuminanceColor = LocalFogVolumeEvaluateInScattering(FogInstance, FogData, CamRayWorldDir);
 			Transmittance = 1.0 - FogData.Coverage;
 		}
-		else
-		{
-#if LFV_CAN_CLIP
-			clip(-1.0);
-#endif
-		}
-	}
-	else
-	{
-#if LFV_CAN_CLIP
-		clip(-1.0);
-#endif
 	}
 
 	return float4(LuminanceColor, Transmittance);
diff --git a/Engine/Shaders/Private/LocalFogVolumes/LocalFogVolumeSplat.usf b/Engine/Shaders/Private/LocalFogVolumes/LocalFogVolumeSplat.usf
--- a/Engine/Shaders/Private/LocalFogVolumes/LocalFogVolumeSplat.usf
+++ b/Engine/Shaders/Private/LocalFogVolumes/LocalFogVolumeSplat.usf
@@ -7,14 +7,8 @@
 
 #include "../SceneTexturesCommon.ush"
 
+#include "LocalFogVolumeCommon.ush"
 
-
-#ifndef LFV_SPLAT_VS
-#define LFV_SPLAT_VS 0
-#endif
-#ifndef LFV_SPLAT_PS
-#define LFV_SPLAT_PS 0
-#endif
 #ifndef LFV_TILED_VS
 #define LFV_TILED_VS 0
 #endif
@@ -24,89 +18,13 @@
 
 
 
-#if LFV_SPLAT_PS
-#define LFV_CAN_CLIP 1
-#else
-#define LFV_CAN_CLIP 0
-#endif
-#include "LocalFogVolumeCommon.ush"
-
-
-
 struct FLocalFogVolumeVertexOutput
 {
-#if LFV_SPLAT_VS || LFV_SPLAT_PS
-	nointerpolation uint InstanceId : TEXCOORD0;
-#endif
-
-#if LFV_TILED_VS || LFV_TILED_PS
+//#if LFV_TILED_VS || LFV_TILED_PS
 	nointerpolation uint TileData : TEXCOORD0;
-#endif
+//#endif
 };
 
-#if LFV_SPLAT_VS
-void LocalFogVolumeSplatVS(
-	in float4 InPosition : ATTRIBUTE0,
-	uint InstanceIndex : SV_InstanceID,
-	out float4 OutPosition : SV_POSITION,
-	out FLocalFogVolumeVertexOutput OtherVertexOutput
-)
-{
-	ResolvedView = ResolveView();
-
-	float4 LocalPosition = InPosition;
-
-	FLocalFogVolumeGPUInstanceData FogInstance = GetLocalFogVolumeGPUInstanceData(InstanceIndex);
-
-	float3 WorldPosition = mul(float4(LocalPosition.xyz, 1.0f), FogInstance.Transform).xyz;
-
-	float4 TranslatedWorldPosition = float4(WorldPosition + LWCHackToFloat(PrimaryView.PreViewTranslation), 1);
-
-	OutPosition = mul(TranslatedWorldPosition, PrimaryView.TranslatedWorldToClip);
-	OtherVertexOutput.InstanceId = InstanceIndex;
-}
-#endif // LFV_SPLAT_VS
-
-
-#if LFV_SPLAT_PS
-void LocalFogVolumeSplatPS(
-	in float4 SVPos : SV_POSITION,
-	in FLocalFogVolumeVertexOutput OtherVertexOutput,
-	out float4 OutColor : SV_Target0)
-{
-	ResolvedView = ResolveView();
-	uint InstanceIndex = OtherVertexOutput.InstanceId;
-
-	float3 CamRayTranslatedWorldOrigin	= PrimaryView.TranslatedWorldCameraOrigin;
-	float3 CamRayTranslatedWorldDir = normalize(SvPositionToTranslatedWorld(float4(SVPos.xy, 0.5, 1.0)));
-	float3 CamRayWorldDir = CamRayTranslatedWorldDir;
-	float3 CamRayWorldOrigin = CamRayTranslatedWorldOrigin - LWCHackToFloat(PrimaryView.PreViewTranslation);
-
-#if SHADING_PATH_MOBILE
-	// LookupDeviceZ with uint2 SVPos.xy as input does not work on mobile.
-	float2 UV = SVPos * View.BufferSizeAndInvSize.zw;
-	const float DeviceZ = LookupDeviceZ(UV);
-#else
-	const float DeviceZ = LookupDeviceZ(uint2(SVPos.xy));
-#endif
-
-	float3 DepthBufferTranslatedWorldPos = SvPositionToTranslatedWorld(float4(SVPos.xy, DeviceZ, 1.0));
-	float3 DepthBufferWorldPos = DepthBufferTranslatedWorldPos - LWCHackToFloat(ResolvedView.PreViewTranslation);
-
-	float4 LuminanceColorTransmittance = GetLFVInstanceContribution(
-		InstanceIndex,
-		CamRayWorldOrigin,
-		CamRayWorldDir,
-		DepthBufferWorldPos);
-
-	OutColor = float4(
-		LuminanceColorTransmittance.rgb * PrimaryView.PreExposure,
-		LuminanceColorTransmittance.a);
-}
-#endif // LFV_SPLAT_PS
-
-
-
 
 
 #if LFV_TILED_VS
diff --git a/Engine/Source/Runtime/Renderer/Private/LocalFogVolumeRendering.cpp b/Engine/Source/Runtime/Renderer/Private/LocalFogVolumeRendering.cpp
--- a/Engine/Source/Runtime/Renderer/Private/LocalFogVolumeRendering.cpp
+++ b/Engine/Source/Runtime/Renderer/Private/LocalFogVolumeRendering.cpp
@@ -29,11 +29,6 @@ static TAutoConsoleVariable<int32> CVarLocalFogVolumeApplyOnTranslucent(
 	TEXT("Project settings enabling the sampling of local fog volumes on translucent elements.\n"),
 	ECVF_ReadOnly | ECVF_RenderThreadSafe);
 
-static TAutoConsoleVariable<int32> CVarLocalFogVolumeTiledRendering(
-	TEXT("r.LocalFogVolume.TiledRendering"), 1,
-	TEXT("Use a tiled approach to render the local fog volume on screen.\n"),
-	ECVF_RenderThreadSafe);
-
 static TAutoConsoleVariable<int32> CVarLocalFogVolumeTilePixelSize(
 	TEXT("r.LocalFogVolume.TilePixelSize"), 128,
 	TEXT("Tile size on screen in pixel at which we cull the local fog volumes.\n"),
@@ -53,11 +48,6 @@ static TAutoConsoleVariable<int32> CVarLocalFogVolumeTileDebug(
 //  - 1920x1080 => 15x9 tiles
 //  - Allowing max 32 volumes at once => culling list buffer = 15 * 9 * 32 * 1 byte = 4320 bytes = 4.3KB
 
-static bool GetLocalFogVolumeTiledRenderingEnable()
-{
-	return CVarLocalFogVolumeTiledRendering.GetValueOnRenderThread() > 0;
-}
-
 static uint32 GetLocalFogVolumeTilePixelSize()
 {
 	return FMath::Max(8u, FMath::Min(512u, (uint32)CVarLocalFogVolumeTilePixelSize.GetValueOnRenderThread()));
@@ -318,7 +308,6 @@ void GetLocalFogVolumeSortingData(const FScene* Scene, FRDGBuilder& GraphBuilder
 
 		FMatrix44f Transform    = FMatrix44f(LHF->FogTransform.ToMatrixWithScale());
 		FMatrix44f InvTransform = Transform.Inverse();
-		ConvertFromMatrix44fTo4x3Array(Transform, LocalFogVolumeGPUInstanceDataIt->Transform);
 		ConvertFromMatrix44fTo4x3Array(InvTransform, LocalFogVolumeGPUInstanceDataIt->InvTransform);
 
 		FVector2DHalf Data0X = FVector2DHalf(LHF->RadialFogExtinction,			LHF->HeightFogExtinction);
@@ -507,67 +496,6 @@ void InitLocalFogVolumesForViews(
 	Local height fog rendering - non mobile
 =============================================================================*/
 
-class FLocalFogVolumeSplatVS : public FGlobalShader
-{
-	DECLARE_GLOBAL_SHADER(FLocalFogVolumeSplatVS);
-	SHADER_USE_PARAMETER_STRUCT(FLocalFogVolumeSplatVS, FGlobalShader);
-
-	BEGIN_SHADER_PARAMETER_STRUCT(FParameters, )
-		SHADER_PARAMETER_STRUCT_REF(FViewUniformShaderParameters, View)
-		SHADER_PARAMETER_STRUCT(FLocalFogVolumeUniformParameters, LFV)
-	END_SHADER_PARAMETER_STRUCT()
-
-	using FPermutationDomain = TShaderPermutationDomain<>;
-
-	static bool ShouldCompilePermutation(const FGlobalShaderPermutationParameters& Parameters)
-	{
-		if (IsMobilePlatform(Parameters.Platform))
-		{
-			return false;
-		}
-		return true;
-	}
-
-	static void ModifyCompilationEnvironment(const FGlobalShaderPermutationParameters& Parameters, FShaderCompilerEnvironment& OutEnvironment)
-	{
-		FGlobalShader::ModifyCompilationEnvironment(Parameters, OutEnvironment);
-		OutEnvironment.SetDefine(TEXT("LFV_SPLAT_VS"), 1);
-	}
-};
-
-IMPLEMENT_GLOBAL_SHADER(FLocalFogVolumeSplatVS, "/Engine/Private/LocalFogVolumes/LocalFogVolumeSplat.usf", "LocalFogVolumeSplatVS", SF_Vertex);
-
-class FLocalFogVolumeSplatPS : public FGlobalShader
-{
-	DECLARE_GLOBAL_SHADER(FLocalFogVolumeSplatPS);
-	SHADER_USE_PARAMETER_STRUCT(FLocalFogVolumeSplatPS, FGlobalShader);
-
-	BEGIN_SHADER_PARAMETER_STRUCT(FParameters, )
-		SHADER_PARAMETER_STRUCT_REF(FViewUniformShaderParameters, View)
-		SHADER_PARAMETER_STRUCT(FLocalFogVolumeUniformParameters, LFV)
-	END_SHADER_PARAMETER_STRUCT()
-
-	using FPermutationDomain = TShaderPermutationDomain<>;
-
-	static bool ShouldCompilePermutation(const FGlobalShaderPermutationParameters& Parameters)
-	{
-		if (IsMobilePlatform(Parameters.Platform))
-		{
-			return false;
-		}
-		return true;
-	}
-
-	static void ModifyCompilationEnvironment(const FGlobalShaderPermutationParameters& Parameters, FShaderCompilerEnvironment& OutEnvironment)
-	{
-		FGlobalShader::ModifyCompilationEnvironment(Parameters, OutEnvironment);
-		OutEnvironment.SetDefine(TEXT("LFV_SPLAT_PS"), 1);
-	}
-};
-
-IMPLEMENT_GLOBAL_SHADER(FLocalFogVolumeSplatPS, "/Engine/Private/LocalFogVolumes/LocalFogVolumeSplat.usf", "LocalFogVolumeSplatPS", SF_Pixel);
-
-
 class FLocalFogVolumeTiledRenderVS : public FGlobalShader
 {
 	DECLARE_GLOBAL_SHADER(FLocalFogVolumeTiledRenderVS);
@@ -630,13 +558,6 @@ class FLocalFogVolumeTiledRenderPS : public FGlobalShader
 
 IMPLEMENT_GLOBAL_SHADER(FLocalFogVolumeTiledRenderPS, "/Engine/Private/LocalFogVolumes/LocalFogVolumeSplat.usf", "LocalFogVolumeTiledPS", SF_Pixel);
 
-BEGIN_SHADER_PARAMETER_STRUCT(FLocalFogVolumePassParameters, )
-	SHADER_PARAMETER_STRUCT_INCLUDE(FLocalFogVolumeSplatVS::FParameters, VS)
-	SHADER_PARAMETER_STRUCT_INCLUDE(FLocalFogVolumeSplatPS::FParameters, PS)
-	SHADER_PARAMETER_RDG_UNIFORM_BUFFER(FSceneTextureUniformParameters, SceneTextures)
-	RENDER_TARGET_BINDING_SLOTS()
-END_SHADER_PARAMETER_STRUCT()
-
 BEGIN_SHADER_PARAMETER_STRUCT(FLocalFogVolumeTiledPassParameters, )
 	SHADER_PARAMETER_STRUCT_INCLUDE(FLocalFogVolumeTiledRenderVS::FParameters, VS)
 	SHADER_PARAMETER_STRUCT_INCLUDE(FLocalFogVolumeTiledRenderPS::FParameters, PS)
@@ -667,128 +588,62 @@ void RenderLocalFogVolume(
 				continue;
 			}
 
-			if (GetLocalFogVolumeTiledRenderingEnable())
-			{
-				FLocalFogVolumeTiledPassParameters* PassParameters = GraphBuilder.AllocParameters<FLocalFogVolumeTiledPassParameters>();
-
-				PassParameters->VS.View = GetShaderBinding(View.ViewUniformBuffer);
-				PassParameters->VS.LFV = View.LocalFogVolumeViewData.UniformParametersStruct;
-				PassParameters->VS.TileDataBuffer = View.LocalFogVolumeViewData.GPUTileDataBufferSRV;
-
-				PassParameters->PS.View = GetShaderBinding(View.ViewUniformBuffer);
-				PassParameters->PS.LFV = View.LocalFogVolumeViewData.UniformParametersStruct;
-				PassParameters->PS.LocalFogVolumeTileDebug = FMath::Clamp(CVarLocalFogVolumeTileDebug.GetValueOnRenderThread(), 0, 2);
-
-				PassParameters->SceneTextures = SceneTextures.UniformBuffer;
-				PassParameters->TileDrawIndirectBuffer = View.LocalFogVolumeViewData.GPUTileDrawIndirectBuffer;
-				PassParameters->RenderTargets[0] = FRenderTargetBinding(SceneColorTexture, ERenderTargetLoadAction::ENoAction);
-
-				FLocalFogVolumeTiledRenderVS::FPermutationDomain VSPermutationVector;
-				auto VertexShader = View.ShaderMap->GetShader< FLocalFogVolumeTiledRenderVS >(VSPermutationVector);
-
-				FLocalFogVolumeTiledRenderPS::FPermutationDomain PsPermutationVector;
-				auto PixelShader = View.ShaderMap->GetShader< FLocalFogVolumeTiledRenderPS >(PsPermutationVector);
-
-				const FIntRect ViewRect = View.ViewRect;
+			FLocalFogVolumeTiledPassParameters* PassParameters = GraphBuilder.AllocParameters<FLocalFogVolumeTiledPassParameters>();
 
-				ClearUnusedGraphResources(VertexShader, &PassParameters->VS);
-				ClearUnusedGraphResources(PixelShader, &PassParameters->PS);
+			PassParameters->VS.View = GetShaderBinding(View.ViewUniformBuffer);
+			PassParameters->VS.LFV = View.LocalFogVolumeViewData.UniformParametersStruct;
+			PassParameters->VS.TileDataBuffer = View.LocalFogVolumeViewData.GPUTileDataBufferSRV;
 
-				FUintVector2& LocalFogVolumeTileDataTextureResolution = View.LocalFogVolumeViewData.UniformParametersStruct.LocalFogVolumeCommon.LocalFogVolumeTileDataTextureResolution;
+			PassParameters->PS.View = GetShaderBinding(View.ViewUniformBuffer);
+			PassParameters->PS.LFV = View.LocalFogVolumeViewData.UniformParametersStruct;
+			PassParameters->PS.LocalFogVolumeTileDebug = FMath::Clamp(CVarLocalFogVolumeTileDebug.GetValueOnRenderThread(), 0, 2);
 
-				GraphBuilder.AddPass(
-					RDG_EVENT_NAME("LocalFogVolume.Tiled (%u X %u)", LocalFogVolumeTileDataTextureResolution.X, LocalFogVolumeTileDataTextureResolution.Y),
-					PassParameters,
-					ERDGPassFlags::Raster,
-					[VertexShader, PixelShader, PassParameters, ViewRect](FRHICommandList& RHICmdList)
-				{
-					FGraphicsPipelineStateInitializer GraphicsPSOInit;
-					RHICmdList.ApplyCachedRenderTargets(GraphicsPSOInit);
+			PassParameters->SceneTextures = SceneTextures.UniformBuffer;
+			PassParameters->TileDrawIndirectBuffer = View.LocalFogVolumeViewData.GPUTileDrawIndirectBuffer;
+			PassParameters->RenderTargets[0] = FRenderTargetBinding(SceneColorTexture, ERenderTargetLoadAction::ENoAction);
 
-					RHICmdList.SetViewport(ViewRect.Min.X, ViewRect.Min.Y, 0.0f, ViewRect.Max.X, ViewRect.Max.Y, 1.0f);
+			FLocalFogVolumeTiledRenderVS::FPermutationDomain VSPermutationVector;
+			auto VertexShader = View.ShaderMap->GetShader< FLocalFogVolumeTiledRenderVS >(VSPermutationVector);
 
-					// Render back faces only since camera may intersect
-					GraphicsPSOInit.RasterizerState = TStaticRasterizerState<>::GetRHI();
-					GraphicsPSOInit.DepthStencilState = TStaticDepthStencilState<false, CF_Always>::GetRHI();
-					GraphicsPSOInit.BlendState = TStaticBlendState<CW_RGBA, BO_Add, BF_One, BF_SourceAlpha, BO_Add, BF_Zero, BF_SourceAlpha>::GetRHI();
-					GraphicsPSOInit.PrimitiveType = PT_TriangleList; // LFV_TODO check if rects are supported and use them if so
+			FLocalFogVolumeTiledRenderPS::FPermutationDomain PsPermutationVector;
+			auto PixelShader = View.ShaderMap->GetShader< FLocalFogVolumeTiledRenderPS >(PsPermutationVector);
 
-					GraphicsPSOInit.BoundShaderState.VertexDeclarationRHI = GFilterVertexDeclaration.VertexDeclarationRHI;
-					GraphicsPSOInit.BoundShaderState.VertexShaderRHI = VertexShader.GetVertexShader();
-					GraphicsPSOInit.BoundShaderState.PixelShaderRHI = PixelShader.GetPixelShader();
+			const FIntRect ViewRect = View.ViewRect;
 
-					SetGraphicsPipelineState(RHICmdList, GraphicsPSOInit, 0);
+			ClearUnusedGraphResources(VertexShader, &PassParameters->VS);
+			ClearUnusedGraphResources(PixelShader, &PassParameters->PS);
 
-					SetShaderParameters(RHICmdList, VertexShader, VertexShader.GetVertexShader(), PassParameters->VS);
-					SetShaderParameters(RHICmdList, PixelShader, PixelShader.GetPixelShader(), PassParameters->PS);
+			FUintVector2& LocalFogVolumeTileDataTextureResolution = View.LocalFogVolumeViewData.UniformParametersStruct.LocalFogVolumeCommon.LocalFogVolumeTileDataTextureResolution;
 
-					RHICmdList.SetStreamSource(0, nullptr, 0);
-					RHICmdList.DrawPrimitiveIndirect(PassParameters->TileDrawIndirectBuffer->GetIndirectRHICallBuffer(), 0);
-				});
-			}
-			else
+			GraphBuilder.AddPass(
+				RDG_EVENT_NAME("LocalFogVolume.Tiled (%u X %u)", LocalFogVolumeTileDataTextureResolution.X, LocalFogVolumeTileDataTextureResolution.Y),
+				PassParameters,
+				ERDGPassFlags::Raster,
+				[VertexShader, PixelShader, PassParameters, ViewRect](FRHICommandList& RHICmdList)
 			{
-				FLocalFogVolumePassParameters* PassParameters = GraphBuilder.AllocParameters<FLocalFogVolumePassParameters>();
-
-				PassParameters->VS.View = GetShaderBinding(View.ViewUniformBuffer);
-				PassParameters->VS.LFV = View.LocalFogVolumeViewData.UniformParametersStruct;
-
-				PassParameters->PS.View = GetShaderBinding(View.ViewUniformBuffer);
-				PassParameters->PS.LFV = View.LocalFogVolumeViewData.UniformParametersStruct;
-
-				PassParameters->SceneTextures = SceneTextures.UniformBuffer;
-				PassParameters->RenderTargets[0] = FRenderTargetBinding(SceneColorTexture, ERenderTargetLoadAction::ENoAction);
-
-				FLocalFogVolumeSplatVS::FPermutationDomain VSPermutationVector;
-				auto VertexShader = View.ShaderMap->GetShader< FLocalFogVolumeSplatVS >(VSPermutationVector);
-
-				FLocalFogVolumeSplatPS::FPermutationDomain PsPermutationVector;
-				auto PixelShader = View.ShaderMap->GetShader< FLocalFogVolumeSplatPS >(PsPermutationVector);
-
-				const FIntRect ViewRect = View.ViewRect;
-
-				ClearUnusedGraphResources(VertexShader, &PassParameters->VS);
-				ClearUnusedGraphResources(PixelShader, &PassParameters->PS);
-
-				uint32 LocalFogVolumeGPUInstanceCount = View.LocalFogVolumeViewData.GPUInstanceCount;
-				GraphBuilder.AddPass(
-					RDG_EVENT_NAME("LocalFogVolume.Splat (%u inst.)", LocalFogVolumeGPUInstanceCount),
-					PassParameters,
-					ERDGPassFlags::Raster,
-					[VertexShader, PixelShader, PassParameters, LocalFogVolumeGPUInstanceCount, ViewRect](FRHICommandList& RHICmdList)
-				{
-					FGraphicsPipelineStateInitializer GraphicsPSOInit;
-					RHICmdList.ApplyCachedRenderTargets(GraphicsPSOInit);
+				FGraphicsPipelineStateInitializer GraphicsPSOInit;
+				RHICmdList.ApplyCachedRenderTargets(GraphicsPSOInit);
 
-					RHICmdList.SetViewport(ViewRect.Min.X, ViewRect.Min.Y, 0.0f, ViewRect.Max.X, ViewRect.Max.Y, 1.0f);
+				RHICmdList.SetViewport(ViewRect.Min.X, ViewRect.Min.Y, 0.0f, ViewRect.Max.X, ViewRect.Max.Y, 1.0f);
 
-					// Render back faces only since camera may intersect
-					GraphicsPSOInit.RasterizerState = TStaticRasterizerState<FM_Solid, CM_CW>::GetRHI();
-					GraphicsPSOInit.DepthStencilState = TStaticDepthStencilState<false, CF_Always>::GetRHI();
-					GraphicsPSOInit.BlendState = TStaticBlendState<CW_RGBA, BO_Add, BF_One, BF_SourceAlpha, BO_Add, BF_Zero, BF_SourceAlpha>::GetRHI();
-					GraphicsPSOInit.PrimitiveType = PT_TriangleList;
+				// Render back faces only since camera may intersect
+				GraphicsPSOInit.RasterizerState = TStaticRasterizerState<>::GetRHI();
+				GraphicsPSOInit.DepthStencilState = TStaticDepthStencilState<false, CF_Always>::GetRHI();
+				GraphicsPSOInit.BlendState = TStaticBlendState<CW_RGBA, BO_Add, BF_One, BF_SourceAlpha, BO_Add, BF_Zero, BF_SourceAlpha>::GetRHI();
+				GraphicsPSOInit.PrimitiveType = PT_TriangleList; // LFV_TODO check if rects are supported and use them if so
 
-					GraphicsPSOInit.BoundShaderState.VertexDeclarationRHI = GetVertexDeclarationFVector4();
-					GraphicsPSOInit.BoundShaderState.VertexShaderRHI = VertexShader.GetVertexShader();
-					GraphicsPSOInit.BoundShaderState.PixelShaderRHI = PixelShader.GetPixelShader();
+				GraphicsPSOInit.BoundShaderState.VertexDeclarationRHI = GFilterVertexDeclaration.VertexDeclarationRHI;
+				GraphicsPSOInit.BoundShaderState.VertexShaderRHI = VertexShader.GetVertexShader();
+				GraphicsPSOInit.BoundShaderState.PixelShaderRHI = PixelShader.GetPixelShader();
 
-					SetGraphicsPipelineState(RHICmdList, GraphicsPSOInit, 0);
+				SetGraphicsPipelineState(RHICmdList, GraphicsPSOInit, 0);
 
-					SetShaderParameters(RHICmdList, VertexShader, VertexShader.GetVertexShader(), PassParameters->VS);
-					SetShaderParameters(RHICmdList, PixelShader, PixelShader.GetPixelShader(), PassParameters->PS);
+				SetShaderParameters(RHICmdList, VertexShader, VertexShader.GetVertexShader(), PassParameters->VS);
+				SetShaderParameters(RHICmdList, PixelShader, PixelShader.GetPixelShader(), PassParameters->PS);
 
-					RHICmdList.SetStreamSource(0, GetUnitCubeVertexBuffer(), 0);
-
-					RHICmdList.DrawIndexedPrimitive(GetUnitCubeIndexBuffer()
-						, 0									//BaseVertexIndex
-						, 0									//FirstInstance
-						, 8									//uint32 NumVertices
-						, 0									//uint32 StartIndex
-						, UE_ARRAY_COUNT(GCubeIndices) / 3	//uint32 NumPrimitives
-						, LocalFogVolumeGPUInstanceCount	//uint32 NumInstances
-					);
-				});
-			}
+				RHICmdList.SetStreamSource(0, nullptr, 0);
+				RHICmdList.DrawPrimitiveIndirect(PassParameters->TileDrawIndirectBuffer->GetIndirectRHICallBuffer(), 0);
+			});
 		}
 	}
 }
@@ -796,73 +651,6 @@ void RenderLocalFogVolume(
 /*=============================================================================
 	Local height fog rendering - mobile
 =============================================================================*/
-	
-
-class FMobileLocalFogVolumeSplatVS : public FGlobalShader
-{
-	DECLARE_GLOBAL_SHADER(FMobileLocalFogVolumeSplatVS);
-	SHADER_USE_PARAMETER_STRUCT(FMobileLocalFogVolumeSplatVS, FGlobalShader);
-
-	BEGIN_SHADER_PARAMETER_STRUCT(FParameters, )
-		SHADER_PARAMETER_STRUCT_INCLUDE(FViewShaderParameters, View)
-		SHADER_PARAMETER_STRUCT(FLocalFogVolumeUniformParameters, LFV)
-	END_SHADER_PARAMETER_STRUCT()
-
-	using FPermutationDomain = TShaderPermutationDomain<>;
-
-	static bool ShouldCompilePermutation(const FGlobalShaderPermutationParameters& Parameters)
-	{
-		if (!IsMobilePlatform(Parameters.Platform))
-		{
-			return false;
-		}
-		return true;
-	}
-
-	static void ModifyCompilationEnvironment(const FGlobalShaderPermutationParameters& Parameters, FShaderCompilerEnvironment& OutEnvironment)
-	{
-		FGlobalShader::ModifyCompilationEnvironment(Parameters, OutEnvironment);
-		const bool bMobileForceDepthRead = MobileUsesFullDepthPrepass(Parameters.Platform);
-		OutEnvironment.SetDefine(TEXT("LFV_SPLAT_VS"), 1);
-		OutEnvironment.SetDefine(TEXT("IS_MOBILE_DEPTHREAD_SUBPASS"), bMobileForceDepthRead ? 0u : 1u);
-	}
-};
-
-IMPLEMENT_GLOBAL_SHADER(FMobileLocalFogVolumeSplatVS, "/Engine/Private/LocalFogVolumes/LocalFogVolumeSplat.usf", "LocalFogVolumeSplatVS", SF_Vertex);
-
-class FMobileLocalFogVolumeSplatPS : public FGlobalShader
-{
-	DECLARE_GLOBAL_SHADER(FMobileLocalFogVolumeSplatPS);
-	SHADER_USE_PARAMETER_STRUCT(FMobileLocalFogVolumeSplatPS, FGlobalShader);
-
-	BEGIN_SHADER_PARAMETER_STRUCT(FParameters, )
-		SHADER_PARAMETER_STRUCT_INCLUDE(FViewShaderParameters, View)
-		SHADER_PARAMETER_RDG_UNIFORM_BUFFER(FMobileBasePassUniformParameters, MobileBasePass)
-		SHADER_PARAMETER_STRUCT(FLocalFogVolumeUniformParameters, LFV)
-	END_SHADER_PARAMETER_STRUCT()
-
-	using FPermutationDomain = TShaderPermutationDomain<>;
-
-	static bool ShouldCompilePermutation(const FGlobalShaderPermutationParameters& Parameters)
-	{
-		if (!IsMobilePlatform(Parameters.Platform))
-		{
-			return false;
-		}
-		return true;
-	}
-
-	static void ModifyCompilationEnvironment(const FGlobalShaderPermutationParameters& Parameters, FShaderCompilerEnvironment& OutEnvironment)
-	{
-		FGlobalShader::ModifyCompilationEnvironment(Parameters, OutEnvironment);
-		const bool bMobileForceDepthRead = MobileUsesFullDepthPrepass(Parameters.Platform);
-		OutEnvironment.SetDefine(TEXT("LFV_SPLAT_PS"), 1);
-		OutEnvironment.SetDefine(TEXT("IS_MOBILE_DEPTHREAD_SUBPASS"), bMobileForceDepthRead ? 0u : 1u);
-	}
-};
-
-IMPLEMENT_GLOBAL_SHADER(FMobileLocalFogVolumeSplatPS, "/Engine/Private/LocalFogVolumes/LocalFogVolumeSplat.usf", "LocalFogVolumeSplatPS", SF_Pixel);
-
 
 class FMobileLocalFogVolumeTiledRenderVS : public FGlobalShader
 {
@@ -943,99 +731,45 @@ void RenderLocalFogVolumeMobile(
 
 	SCOPED_DRAW_EVENT(RHICmdList, LocalFogVolumeVolumes);
 
+	FMobileLocalFogVolumeTiledRenderVS::FPermutationDomain VSPermutationVector;
+	auto VertexShader = View.ShaderMap->GetShader< FMobileLocalFogVolumeTiledRenderVS >(VSPermutationVector);
 
-	if (GetLocalFogVolumeTiledRenderingEnable())
-	{
-		FMobileLocalFogVolumeTiledRenderVS::FPermutationDomain VSPermutationVector;
-		auto VertexShader = View.ShaderMap->GetShader< FMobileLocalFogVolumeTiledRenderVS >(VSPermutationVector);
-
-		FMobileLocalFogVolumeTiledRenderPS::FPermutationDomain PsPermutationVector;
-		auto PixelShader = View.ShaderMap->GetShader< FMobileLocalFogVolumeTiledRenderPS >(PsPermutationVector);
-
-		FGraphicsPipelineStateInitializer GraphicsPSOInit;
-		RHICmdList.ApplyCachedRenderTargets(GraphicsPSOInit);
+	FMobileLocalFogVolumeTiledRenderPS::FPermutationDomain PsPermutationVector;
+	auto PixelShader = View.ShaderMap->GetShader< FMobileLocalFogVolumeTiledRenderPS >(PsPermutationVector);
 
-		const FIntRect ViewRect = View.ViewRect;
-		RHICmdList.SetViewport(ViewRect.Min.X, ViewRect.Min.Y, 0.0f, ViewRect.Max.X, ViewRect.Max.Y, 1.0f);
+	FGraphicsPipelineStateInitializer GraphicsPSOInit;
+	RHICmdList.ApplyCachedRenderTargets(GraphicsPSOInit);
 
-		// Render back faces only since camera may intersect
-		GraphicsPSOInit.RasterizerState = TStaticRasterizerState<>::GetRHI();
-		GraphicsPSOInit.DepthStencilState = TStaticDepthStencilState<false, CF_Always>::GetRHI();
-		GraphicsPSOInit.BlendState = TStaticBlendState<CW_RGBA, BO_Add, BF_One, BF_SourceAlpha, BO_Add, BF_Zero, BF_SourceAlpha>::GetRHI();
-		GraphicsPSOInit.PrimitiveType = PT_TriangleList; // LFV_TODO check if rects are supported and use them if so
+	const FIntRect ViewRect = View.ViewRect;
+	RHICmdList.SetViewport(ViewRect.Min.X, ViewRect.Min.Y, 0.0f, ViewRect.Max.X, ViewRect.Max.Y, 1.0f);
 
-		GraphicsPSOInit.BoundShaderState.VertexDeclarationRHI = GFilterVertexDeclaration.VertexDeclarationRHI;
-		GraphicsPSOInit.BoundShaderState.VertexShaderRHI = VertexShader.GetVertexShader();
-		GraphicsPSOInit.BoundShaderState.PixelShaderRHI = PixelShader.GetPixelShader();
+	// Render back faces only since camera may intersect
+	GraphicsPSOInit.RasterizerState = TStaticRasterizerState<>::GetRHI();
+	GraphicsPSOInit.DepthStencilState = TStaticDepthStencilState<false, CF_Always>::GetRHI();
+	GraphicsPSOInit.BlendState = TStaticBlendState<CW_RGBA, BO_Add, BF_One, BF_SourceAlpha, BO_Add, BF_Zero, BF_SourceAlpha>::GetRHI();
+	GraphicsPSOInit.PrimitiveType = PT_TriangleList; // LFV_TODO check if rects are supported and use them if so
 
-		SetGraphicsPipelineState(RHICmdList, GraphicsPSOInit, 0);
+	GraphicsPSOInit.BoundShaderState.VertexDeclarationRHI = GFilterVertexDeclaration.VertexDeclarationRHI;
+	GraphicsPSOInit.BoundShaderState.VertexShaderRHI = VertexShader.GetVertexShader();
+	GraphicsPSOInit.BoundShaderState.PixelShaderRHI = PixelShader.GetPixelShader();
 
+	SetGraphicsPipelineState(RHICmdList, GraphicsPSOInit, 0);
 
-		FMobileLocalFogVolumeTiledRenderVS::FParameters VSParameters;
-		VSParameters.View = View.GetShaderParameters();
-		VSParameters.LFV = View.LocalFogVolumeViewData.UniformParametersStruct;
-		VSParameters.TileDataBuffer = View.LocalFogVolumeViewData.GPUTileDataBufferSRV;
-		VSParameters.TileDrawIndirectBuffer = View.LocalFogVolumeViewData.GPUTileDrawIndirectBuffer;
-		SetShaderParameters(RHICmdList, VertexShader, VertexShader.GetVertexShader(), VSParameters);
-
-		FMobileLocalFogVolumeTiledRenderPS::FParameters PSParameters;
-		PSParameters.View = View.GetShaderParameters();
-		PSParameters.LFV = View.LocalFogVolumeViewData.UniformParametersStruct;
-		PSParameters.LocalFogVolumeTileDebug = FMath::Clamp(CVarLocalFogVolumeTileDebug.GetValueOnRenderThread(), 0, 2);
-		// PSParameters.MobileBasePass filled up by the RDG pass parameters.
-		SetShaderParameters(RHICmdList, PixelShader, PixelShader.GetPixelShader(), PSParameters);
-
-		RHICmdList.SetStreamSource(0, nullptr, 0);
-		RHICmdList.DrawPrimitiveIndirect(View.LocalFogVolumeViewData.GPUTileDrawIndirectBuffer->GetIndirectRHICallBuffer(), 0);
-	}
-	else
-	{
-		FMobileLocalFogVolumeSplatVS::FPermutationDomain VSPermutationVector;
-		auto VertexShader = View.ShaderMap->GetShader< FMobileLocalFogVolumeSplatVS >(VSPermutationVector);
 
-		FMobileLocalFogVolumeSplatPS::FPermutationDomain PsPermutationVector;
-		auto PixelShader = View.ShaderMap->GetShader< FMobileLocalFogVolumeSplatPS >(PsPermutationVector);
+	FMobileLocalFogVolumeTiledRenderVS::FParameters VSParameters;
+	VSParameters.View = View.GetShaderParameters();
+	VSParameters.LFV = View.LocalFogVolumeViewData.UniformParametersStruct;
+	VSParameters.TileDataBuffer = View.LocalFogVolumeViewData.GPUTileDataBufferSRV;
+	VSParameters.TileDrawIndirectBuffer = View.LocalFogVolumeViewData.GPUTileDrawIndirectBuffer;
+	SetShaderParameters(RHICmdList, VertexShader, VertexShader.GetVertexShader(), VSParameters);
 
-		const FIntRect ViewRect = View.ViewRect;
+	FMobileLocalFogVolumeTiledRenderPS::FParameters PSParameters;
+	PSParameters.View = View.GetShaderParameters();
+	PSParameters.LFV = View.LocalFogVolumeViewData.UniformParametersStruct;
+	PSParameters.LocalFogVolumeTileDebug = FMath::Clamp(CVarLocalFogVolumeTileDebug.GetValueOnRenderThread(), 0, 2);
+	// PSParameters.MobileBasePass filled up by the RDG pass parameters.
+	SetShaderParameters(RHICmdList, PixelShader, PixelShader.GetPixelShader(), PSParameters);
 
-		FGraphicsPipelineStateInitializer GraphicsPSOInit;
-		RHICmdList.ApplyCachedRenderTargets(GraphicsPSOInit);
-
-		RHICmdList.SetViewport(ViewRect.Min.X, ViewRect.Min.Y, 0.0f, ViewRect.Max.X, ViewRect.Max.Y, 1.0f);
-
-		// Render back faces only since camera may intersect
-		GraphicsPSOInit.RasterizerState = TStaticRasterizerState<FM_Solid, CM_CW>::GetRHI();
-		GraphicsPSOInit.DepthStencilState = TStaticDepthStencilState<false, CF_Always>::GetRHI();
-		GraphicsPSOInit.BlendState = TStaticBlendState<CW_RGBA, BO_Add, BF_One, BF_SourceAlpha, BO_Add, BF_Zero, BF_SourceAlpha>::GetRHI();
-		GraphicsPSOInit.PrimitiveType = PT_TriangleList;
-
-		GraphicsPSOInit.BoundShaderState.VertexDeclarationRHI = GetVertexDeclarationFVector4();
-		GraphicsPSOInit.BoundShaderState.VertexShaderRHI = VertexShader.GetVertexShader();
-		GraphicsPSOInit.BoundShaderState.PixelShaderRHI = PixelShader.GetPixelShader();
-
-		SetGraphicsPipelineState(RHICmdList, GraphicsPSOInit, 0);
-
-		FMobileLocalFogVolumeSplatVS::FParameters VSParameters;
-		VSParameters.View = View.GetShaderParameters();
-		VSParameters.LFV = View.LocalFogVolumeViewData.UniformParametersStruct;
-		SetShaderParameters(RHICmdList, VertexShader, VertexShader.GetVertexShader(), VSParameters);
-
-		FMobileLocalFogVolumeSplatPS::FParameters PSParameters;
-		PSParameters.View = View.GetShaderParameters();
-		PSParameters.LFV = View.LocalFogVolumeViewData.UniformParametersStruct;
-		// PSParameters.MobileBasePass filled up by the RDG pass parameters.
-		SetShaderParameters(RHICmdList, PixelShader, PixelShader.GetPixelShader(), PSParameters);
-
-		RHICmdList.SetStreamSource(0, GetUnitCubeVertexBuffer(), 0);
-
-		RHICmdList.DrawIndexedPrimitive(
-			GetUnitCubeIndexBuffer()
-			, 0												//BaseVertexIndex
-			, 0												//FirstInstance
-			, 8												//uint32 NumVertices
-			, 0												//uint32 StartIndex
-			, UE_ARRAY_COUNT(GCubeIndices) / 3				//uint32 NumPrimitives
-			, View.LocalFogVolumeViewData.GPUInstanceCount	//uint32 NumInstances
-		);
-	}
+	RHICmdList.SetStreamSource(0, nullptr, 0);
+	RHICmdList.DrawPrimitiveIndirect(View.LocalFogVolumeViewData.GPUTileDrawIndirectBuffer->GetIndirectRHICallBuffer(), 0);
 }
diff --git a/Engine/Source/Runtime/Renderer/Private/LocalFogVolumeRendering.h b/Engine/Source/Runtime/Renderer/Private/LocalFogVolumeRendering.h
--- a/Engine/Source/Runtime/Renderer/Private/LocalFogVolumeRendering.h
+++ b/Engine/Source/Runtime/Renderer/Private/LocalFogVolumeRendering.h
@@ -43,7 +43,6 @@ class FLocalFogVolumeGPUInstanceData
 {
 public:
 
-	float Transform[4 * 3];
 	float InvTransform[4 * 3];
 
 	// See LocalFogVolumeCommon for the encoded representation
