LFV - added support for tiled culling on the async compute pipe.

#prelfight none

diff --git a/Engine/Source/Runtime/Renderer/Private/DeferredShadingRenderer.cpp b/Engine/Source/Runtime/Renderer/Private/DeferredShadingRenderer.cpp
--- a/Engine/Source/Runtime/Renderer/Private/DeferredShadingRenderer.cpp
+++ b/Engine/Source/Runtime/Renderer/Private/DeferredShadingRenderer.cpp
@@ -3443,7 +3443,7 @@ void FDeferredShadingSceneRenderer::Render(FRDGBuilder& GraphBuilder)
 		CSV_CUSTOM_STAT(LightCount, Batched, float(SortedLightSet.UnbatchedLightStart), ECsvCustomStatOp::Set);
 		CSV_CUSTOM_STAT(LightCount, Unbatched, float(SortedLightSet.SortedLights.Num()) - float(SortedLightSet.UnbatchedLightStart), ECsvCustomStatOp::Set);
 
-		// Run local fog volume initialization before base pass for when data is needed in forward
+		// Run local fog volume initialization before base pass for when data is needed in forward  and volumetric cloud view initialization which can bind LFV data.
 		InitLocalFogVolumesForViews(Scene, Views, ViewFamily, GraphBuilder, bShouldRenderVolumetricFog);
 
 		// Run before RenderSkyAtmosphereLookUpTables for cloud shadows to be valid.
diff --git a/Engine/Source/Runtime/Renderer/Private/LocalFogVolumeRendering.cpp b/Engine/Source/Runtime/Renderer/Private/LocalFogVolumeRendering.cpp
--- a/Engine/Source/Runtime/Renderer/Private/LocalFogVolumeRendering.cpp
+++ b/Engine/Source/Runtime/Renderer/Private/LocalFogVolumeRendering.cpp
@@ -39,6 +39,11 @@ static TAutoConsoleVariable<int32> CVarLocalFogVolumeTileMaxInstanceCount(
 	TEXT("Maximum number of local fog volume to account for per view (and per tile or consistency).\n"),
 	ECVF_RenderThreadSafe);
 
+static TAutoConsoleVariable<int32> CVarLocalFogVolumeTileCullingUseAsync(
+	TEXT("r.LocalFogVolume.TileCullingUseAsync"), 1,
+	TEXT("True if we want to try and use culling on the async pipe.\n"),
+	ECVF_RenderThreadSafe);
+
 static TAutoConsoleVariable<int32> CVarLocalFogVolumeTileDebug(
 	TEXT("r.LocalFogVolume.TileDebug"), 0,
 	TEXT("Debug the tiled rendering data complexity. 1: show per tile LFV count as color ; 2: same as one but also show the effect of pixel discard/clipping.\n"),
@@ -53,6 +58,11 @@ static uint32 GetLocalFogVolumeTilePixelSize()
 	return FMath::Max(8u, FMath::Min(512u, (uint32)CVarLocalFogVolumeTilePixelSize.GetValueOnRenderThread()));
 }
 
+static bool GetLocalFogVolumeTileCullingUseAsync()
+{
+	return CVarLocalFogVolumeTileCullingUseAsync.GetValueOnRenderThread() > 0;
+}
+
 static float GetLocalFogVolumeMaxDensityIntoVolumetricFog()
 {
 	return FMath::Max(0.0f, CVarLocalFogVolumeMaxDensityIntoVolumetricFog.GetValueOnRenderThread());
@@ -256,7 +266,8 @@ static void LocalFogVolumeViewTiledCullingPass(FViewInfo& View, FRDGBuilder& Gra
 	float TileCoveredResolutionY = View.LocalFogVolumeViewData.UniformParametersStruct.LocalFogVolumeCommon.LocalFogVolumeTilePixelSize * View.LocalFogVolumeViewData.UniformParametersStruct.LocalFogVolumeCommon.LocalFogVolumeTileDataTextureResolution.Y;
 	PassParameters->ViewToTileSpaceRatio = FVector2f(TileCoveredResolutionX * View.CachedViewUniformShaderParameters->ViewSizeAndInvSize.Z, TileCoveredResolutionY * View.CachedViewUniformShaderParameters->ViewSizeAndInvSize.W);
 
-	ERDGPassFlags PassFlag = ERDGPassFlags::Compute; // LFV_TODO try ERDGPassFlags::AsyncCompute later
+	const bool bUseAsyncCompute = GSupportsEfficientAsyncCompute && GetLocalFogVolumeTileCullingUseAsync();
+	ERDGPassFlags PassFlag = bUseAsyncCompute ? ERDGPassFlags::AsyncCompute : ERDGPassFlags::Compute;
 
 	TileDataTextureSize.Z = 1;
 	const FIntVector NumGroups = FIntVector::DivideAndRoundUp(TileDataTextureSize, FLocalFogVolumeTiledCullingCS::GroupSize);
