Disables r.TSR.Subpixel.Method for time being as it is no longer working

diff --git a/Engine/Source/Runtime/Renderer/Private/PostProcess/TemporalSuperResolution.cpp b/Engine/Source/Runtime/Renderer/Private/PostProcess/TemporalSuperResolution.cpp
--- a/Engine/Source/Runtime/Renderer/Private/PostProcess/TemporalSuperResolution.cpp
+++ b/Engine/Source/Runtime/Renderer/Private/PostProcess/TemporalSuperResolution.cpp
@@ -248,7 +248,7 @@ TAutoConsoleVariable<float> CVarTSRWeightClampingPixelSpeed(
 //	ECVF_Scalability | ECVF_RenderThreadSafe);
 
 TAutoConsoleVariable<int32> CVarTSRSubpixelMethod(
-	TEXT("r.TSR.Subpixel.Method"), 2,
+	TEXT("r.TSR.Subpixel.Method"), 0,
 	TEXT("One particular challenge of Nanite amount of details is that sometimes these details can be thiner than a rendering pixel in which case ")
 	TEXT("they only render in some frames. When that happens, it means neither depth or velocity buffer to be able to reproject them. This is for instance ")
 	TEXT("visible with the `vis SceneDepthZ` command.")
