Adds TSR's shader permutation ids in draw events

diff --git a/Engine/Source/Runtime/Renderer/Private/PostProcess/TemporalSuperResolution.cpp b/Engine/Source/Runtime/Renderer/Private/PostProcess/TemporalSuperResolution.cpp
--- a/Engine/Source/Runtime/Renderer/Private/PostProcess/TemporalSuperResolution.cpp
+++ b/Engine/Source/Runtime/Renderer/Private/PostProcess/TemporalSuperResolution.cpp
@@ -1826,7 +1826,8 @@ FDefaultTemporalUpscaler::FOutputs AddTemporalSuperResolutionPasses(
 		TShaderMapRef<FTSRDilateVelocityCS> ComputeShader(View.ShaderMap, PermutationVector);
 		FComputeShaderUtils::AddPass(
 			GraphBuilder,
-			RDG_EVENT_NAME("TSR DilateVelocity(MotionBlurDirections=%d%s%s) %dx%d",
+			RDG_EVENT_NAME("TSR DilateVelocity(#%d MotionBlurDirections=%d%s%s) %dx%d",
+				PermutationVector.ToDimensionValueId(),
 				int32(PermutationVector.Get<FTSRDilateVelocityCS::FMotionBlurDirectionsDim>()),
 				bOutputIsMovingTexture ? TEXT(" OutputIsMoving") : TEXT(""),
 				PermutationVector.Get<FTSRDilateVelocityCS::FSubpixelDepthDim>() ? TEXT(" SubpixelDepth") : TEXT(""),
@@ -1925,7 +1926,8 @@ FDefaultTemporalUpscaler::FOutputs AddTemporalSuperResolutionPasses(
 		TShaderMapRef<FTSRDecimateHistoryCS> ComputeShader(View.ShaderMap, PermutationVector);
 		FComputeShaderUtils::AddPass(
 			GraphBuilder,
-			RDG_EVENT_NAME("TSR DecimateHistory(%s%s) %dx%d",
+			RDG_EVENT_NAME("TSR DecimateHistory(#%d %s%s) %dx%d",
+				PermutationVector.ToDimensionValueId(),
 				PermutationVector.Get<FTSRDecimateHistoryCS::FMoireReprojectionDim>() ? TEXT("ReprojectMoire") : TEXT(""),
 				PermutationVector.Get<FTSRDecimateHistoryCS::FResurrectionReprojectionDim>() ? TEXT(" ReprojectResurrection") : TEXT(""),
 				InputRect.Width(), InputRect.Height()),
@@ -2125,7 +2127,8 @@ FDefaultTemporalUpscaler::FOutputs AddTemporalSuperResolutionPasses(
 		TShaderMapRef<FTSRRejectShadingCS> ComputeShader(View.ShaderMap, PermutationVector);
 		FComputeShaderUtils::AddPass(
 			GraphBuilder,
-			RDG_EVENT_NAME("TSR RejectShading(TileSize=%d PaddingCostMultiplier=%1.1f WaveSize=%d FlickeringFramePeriod=%f VALU=%s%s) %dx%d",
+			RDG_EVENT_NAME("TSR RejectShading(#%d TileSize=%d PaddingCostMultiplier=%1.1f WaveSize=%d FlickeringFramePeriod=%f VALU=%s%s) %dx%d",
+				PermutationVector.ToDimensionValueId(),
 				TileSize,
 				FMath::Pow(float(GroupTileSize) / float(TileSize), 2),
 				int32(PermutationVector.Get<FTSRRejectShadingCS::FWaveSizeOps>()),
@@ -2189,7 +2192,8 @@ FDefaultTemporalUpscaler::FOutputs AddTemporalSuperResolutionPasses(
 		TShaderMapRef<FTSRSpatialAntiAliasingCS> ComputeShader(View.ShaderMap, PermutationVector);
 		FComputeShaderUtils::AddPass(
 			GraphBuilder,
-			RDG_EVENT_NAME("TSR SpatialAntiAliasing(Quality=%d) %dx%d",
+			RDG_EVENT_NAME("TSR SpatialAntiAliasing(#%d Quality=%d) %dx%d",
+				PermutationVector.ToDimensionValueId(),
 				RejectionAntiAliasingQuality,
 				InputRect.Width(), InputRect.Height()),
 			AsyncComputePasses >= 3 ? ERDGPassFlags::AsyncCompute : ERDGPassFlags::Compute,
@@ -2348,7 +2352,8 @@ FDefaultTemporalUpscaler::FOutputs AddTemporalSuperResolutionPasses(
 		TShaderMapRef<FTSRUpdateHistoryCS> ComputeShader(View.ShaderMap, PermutationVector);
 		FComputeShaderUtils::AddPass(
 			GraphBuilder,
-			RDG_EVENT_NAME("TSR UpdateHistory(Quality=%s%s%s%s) %dx%d",
+			RDG_EVENT_NAME("TSR UpdateHistory(#%d Quality=%s%s%s%s) %dx%d",
+				PermutationVector.ToDimensionValueId(),
 				kUpdateQualityNames[int32(PermutationVector.Get<FTSRUpdateHistoryCS::FQualityDim>())],
 				PermutationVector.Get<FTSRShader::F16BitVALUDim>() ? TEXT(" 16bit") : TEXT(""),
 				HistoryColorFormat == PF_FloatR11G11B10 ? TEXT(" R11G11B10") : TEXT(""),
@@ -2407,8 +2412,9 @@ FDefaultTemporalUpscaler::FOutputs AddTemporalSuperResolutionPasses(
 		TShaderMapRef<FTSRResolveHistoryCS> ComputeShader(View.ShaderMap, PermutationVector);
 		FComputeShaderUtils::AddPass(
 			GraphBuilder,
-			RDG_EVENT_NAME("TSR ResolveHistory(%s%s) %dx%d", 
-				PermutationVector.Get<FTSRResolveHistoryCS::FNyquistDim>() ? TEXT("WaveOps") : TEXT(""),
+			RDG_EVENT_NAME("TSR ResolveHistory(#%d%s%s) %dx%d",
+				PermutationVector.ToDimensionValueId(),
+				PermutationVector.Get<FTSRResolveHistoryCS::FNyquistDim>() ? TEXT(" WaveOps") : TEXT(""),
 				PermutationVector.Get<FTSRShader::F16BitVALUDim>() ? TEXT(" 16bit") : TEXT(""),
 				OutputRect.Width(), OutputRect.Height()),
 			AsyncComputePasses >= 3 ? ERDGPassFlags::AsyncCompute : ERDGPassFlags::Compute,
