Update to CL 27111705 to fix frustum culled primitives not being properly merged. Changes approach to merging primitive lists as well as visibility maps between occlusion and relevance, rather than passing through all primitives for primary views even when culled and filtering later.

diff --git a/Engine/Source/Runtime/Renderer/Private/SceneVisibility.cpp b/Engine/Source/Runtime/Renderer/Private/SceneVisibility.cpp
--- a/Engine/Source/Runtime/Renderer/Private/SceneVisibility.cpp
+++ b/Engine/Source/Runtime/Renderer/Private/SceneVisibility.cpp
@@ -1869,15 +1869,26 @@ FComputeAndMarkRelevance::FComputeAndMarkRelevance(FVisibilityTaskData& InTaskDa
 	FMemory::Memzero(MarkMasks, NumMeshes + 31);
 	Packets.Reserve(InTaskData.TaskConfig.Relevance.NumEstimatedPackets);
 	CreateRelevancePacket();
+	InstancedPrimitiveAddedMap.Init(false, InScene.Primitives.Num());
 }
 
 void FComputeAndMarkRelevance::AddPrimitives(FPrimitiveIndexList&& PrimitiveIndexList)
 {
-	// In ISR only, primitives that have been occlusion culled in the primary view are still piped in because they may be visible in secondary views.
-	// They can be filtered out now using the visibility map, which is merged with that of the secondary view once occlusion tasks are complete.
-	const bool bIsPrimaryISRView = View.bIsMultiViewportEnabled && View.StereoPass == EStereoscopicPass::eSSP_PRIMARY;
+	// In ISR only, the primary view will also receive all primitives visible in secondary views, since all rendering is handled from the primary
+	if (View.bIsMultiViewportEnabled && View.StereoPass == EStereoscopicPass::eSSP_PRIMARY)
+	{
+		for (int32 Index : PrimitiveIndexList)
+		{
+			// Because we've queued primitives from both the primary and secondary views, we need to filter duplicates
+			if (!InstancedPrimitiveAddedMap[Index])
+			{
+				AddPrimitive(Index);
+				InstancedPrimitiveAddedMap[Index] = true;
+			}
+		}
+	}
 
-	if (!bIsPrimaryISRView && PrimitiveIndexList.Num() == NumPrimitivesPerPacket)
+	else if (PrimitiveIndexList.Num() == NumPrimitivesPerPacket)
 	{
 		// Create a one-off packet that will take all the primitives.
 		FRelevancePacket* Packet = CreateRelevancePacket();
@@ -1891,14 +1902,12 @@ void FComputeAndMarkRelevance::AddPrimitives(FPrimitiveIndexList&& PrimitiveInde
 		// Put the previous last packet that was in progress back in the last slot.
 		Swap(Packets[Packets.Num() - 1], Packets[Packets.Num() - 2]);
 	}
+
 	else
 	{
 		for (int32 Index : PrimitiveIndexList)
 		{
-			if (!bIsPrimaryISRView || View.PrimitiveVisibilityMap[Index])
-			{
-				AddPrimitive(Index);
-			}
+			AddPrimitive(Index);
 		}
 	}
 }
@@ -2886,14 +2895,9 @@ FOcclusionCullResult FGPUOcclusionParallelPacket::OcclusionCullTask(FPrimitiveIn
 
 	PrimitiveIndexList.Reserve(Input.Num());
 
-	const bool bIsPrimaryISRView = View.bIsMultiViewportEnabled && View.StereoPass == EStereoscopicPass::eSSP_PRIMARY;
-
 	for (int32 Index : Input)
 	{
-		// When in ISR, primary stereo views can have additional visible primitives derived from secondary views.
-		// In that case, forward the primary view primitives down the pipe even if they're occluded.
-		// We also need to update the visibility map so we can later filter primitives culled in both views.
-		if (OcclusionCullPrimitive(RecordVisitor, Result, Index) || bIsPrimaryISRView)
+		if (OcclusionCullPrimitive(RecordVisitor, Result, Index))
 		{
 			PrimitiveIndexList.Emplace(Index);
 		}
@@ -3201,6 +3205,23 @@ FVisibilityViewPacket::FVisibilityViewPacket(FVisibilityTaskData& InTaskData, FS
 			{
 				OcclusionCull.ContextIfParallel->AddPrimitives(PrimitiveRange);
 			}
+
+			else if (View.bIsMultiViewportEnabled)
+			{
+				// In ISR, we still need to use the relevance pipe even when occlusion is disabled to ensure all commands are forwarded to the primary view.
+				FPrimitiveIndexList PrimitiveIndexList;
+				for (FSceneSetBitIterator BitIt(View.PrimitiveVisibilityMap, PrimitiveRange.StartIndex); BitIt.GetIndex() < PrimitiveRange.EndIndex; ++BitIt)
+				{
+					PrimitiveIndexList.Emplace(BitIt.GetIndex());
+				}
+
+				if (!PrimitiveIndexList.IsEmpty())
+				{
+					Relevance.CommandPipe.AddNumCommands(1);
+					Relevance.CommandPipe.EnqueueCommand(PrimitiveIndexList);
+				}
+			}
+
 			else
 			{
 				// When occlusion is disabled primitives are queued directly to the relevance context rather than as a command on the relevance pipe.
@@ -3240,22 +3261,54 @@ FVisibilityViewPacket::FVisibilityViewPacket(FVisibilityTaskData& InTaskData, FS
 		OcclusionCull.CommandPipe.AddNumCommands(1);
 
 		// Callback for when a relevance command is queued from occlusion.
-		Relevance.CommandPipe.SetCommandFunction([this](FPrimitiveIndexList&& PrimitiveIndexList)
+		if (View.bIsMultiViewportEnabled && View.StereoPass == EStereoscopicPass::eSSP_SECONDARY)
 		{
-			Relevance.Context->AddPrimitives(MoveTemp(PrimitiveIndexList));
-		});
+			Relevance.CommandPipe.SetCommandFunction([this](FPrimitiveIndexList&& PrimitiveIndexList)
+			{
+				// For instanced stereo secondary views, also send this command to the primary view
+				Relevance.PrimaryViewCommandPipe->AddNumCommands(1);
+				Relevance.PrimaryViewCommandPipe->EnqueueCommand(PrimitiveIndexList);
+
+				Relevance.Context->AddPrimitives(MoveTemp(PrimitiveIndexList));
+			});
+		}
+		else
+		{
+			Relevance.CommandPipe.SetCommandFunction([this](FPrimitiveIndexList&& PrimitiveIndexList)
+			{
+				Relevance.Context->AddPrimitives(MoveTemp(PrimitiveIndexList));
+			});
+		}
 
 		// Callback for when the relevance pipe is done accepting commands.
-		Relevance.CommandPipe.SetEmptyFunction([this]
+		if (View.bIsMultiViewportEnabled && View.StereoPass == EStereoscopicPass::eSSP_SECONDARY)
 		{
-			Relevance.Context->Finish(Tasks.ComputeRelevance);
+			Relevance.CommandPipe.SetEmptyFunction([this]
+			{
+				Relevance.Context->Finish(Tasks.ComputeRelevance);
 
-			if (TaskData.DynamicMeshElements.CommandPipe)
+				if (TaskData.DynamicMeshElements.CommandPipe)
+				{
+					// Release our reference on the dynamic mesh element pipe. We want to keep it alive until all relevance packets have been launched.
+					TaskData.DynamicMeshElements.CommandPipe->ReleaseNumCommands(1);
+				}
+
+				// Also release reference on the instanced primary pipe
+				Relevance.PrimaryViewCommandPipe->ReleaseNumCommands(1);
+			});
+		}
+		else
+		{
+			Relevance.CommandPipe.SetEmptyFunction([this]
 			{
-				// Release our reference on the dynamic mesh element pipe. We want to keep it alive until all relevance packets have been launched.
-				TaskData.DynamicMeshElements.CommandPipe->ReleaseNumCommands(1);
-			}
-		});
+				Relevance.Context->Finish(Tasks.ComputeRelevance);
+
+				if (TaskData.DynamicMeshElements.CommandPipe)
+				{
+					TaskData.DynamicMeshElements.CommandPipe->ReleaseNumCommands(1);
+				}
+			});
+		}
 
 		// Take a reference on the relevance command pipe that is released by the occlusion pipe when all occlusion commands are complete.
 		Relevance.CommandPipe.AddNumCommands(1);
@@ -3492,66 +3545,61 @@ void FVisibilityTaskData::LaunchVisibilityTasks()
 
 	if (TaskConfig.Schedule == EVisibilityTaskSchedule::Parallel)
 	{
-		if (Views.Num() == 1 && !Views[0]->bIsMultiViewportEnabled)
+		if (Views.Num() > 1 && Views[0]->bIsMultiViewportEnabled)
 		{
-			// When using a single view, dynamic mesh elements are pushed into a pipe that is executed on the render thread which allows for some overlap with compute relevance work.
-			DynamicMeshElements.CommandPipe = Allocator.Create<TCommandPipe<FDynamicPrimitiveIndexList>>(TEXT("GatherDynamicMeshElements"), ENamedThreads::GetRenderThread_Local());
-
-			DynamicMeshElements.CommandPipe->SetCommandFunction([this](const FDynamicPrimitiveIndexList& DynamicPrimitiveIndexList)
-			{
-				GatherDynamicMeshElements(DynamicPrimitiveIndexList);
-			});
-
-			// Fire an event when the pipe has completed all work.
-			DynamicMeshElements.CommandPipeCompleteEvent = FGraphEvent::CreateGraphEvent();
-			DynamicMeshElements.CommandPipe->SetEmptyTaskEvent(DynamicMeshElements.CommandPipeCompleteEvent);
-
-			// Take a reference that is released when the relevance pipe has completed. We only need to take one since there can only be one view.
-			DynamicMeshElements.CommandPipe->AddNumCommands(1);
-
-			// We don't need the primitive view masks when in parallel mode with a single view.
-			bAllocatePrimitiveViewMasks = false;
-		}
-		else
-		{
-			TArray<UE::Tasks::FTask, SceneRenderingAllocator> OcclusionCullTasks;
-
 			// Instanced multi-view scenarios have secondary views which feed visibility data into primary views. In this
-			// case we make all views finish occlusion culling first before launching a task that merges visibility
-			// into the primary view, which then becomes a prerequisite task for processing relevance pipe commands in the
-			// primary view.
+			// case we make all views finish culling first before launching a task that merges secondary visibility
+			// maps into the primary view, which then becomes a prerequisite task for processing relevance pipe commands.
+			// All views will share a single relevance pipe and therefore a single relevance context.
+
+			TArray<UE::Tasks::FTask, SceneRenderingAllocator> PreSyncCullingTasks;
 
 			for (FVisibilityViewPacket& ViewPacket : ViewPackets)
 			{
-				OcclusionCullTasks.Emplace(ViewPacket.Tasks.OcclusionCull);
+				// Ensure we merge after all culling is complete.
+				PreSyncCullingTasks.Emplace(ViewPacket.Tasks.FrustumCull);
+				PreSyncCullingTasks.Emplace(ViewPacket.Tasks.OcclusionCull);
 			}
 
-			if (!OcclusionCullTasks.IsEmpty())
-			{
-				// Wait for frustum culling to complete to avoid race conditions with writing to the visibility bits.
-				OcclusionCullTasks.Emplace(Tasks.FrustumCull);
+			FGraphEventRef MergeSecondaryViewsTask = FGraphEvent::CreateGraphEvent();
 
-				FGraphEventRef MergeSecondaryViewsTask = FGraphEvent::CreateGraphEvent();
-
-				UE::Tasks::Launch(UE_SOURCE_LOCATION, [this, MergeSecondaryViewsTask]
+			UE::Tasks::Launch(UE_SOURCE_LOCATION, [this, MergeSecondaryViewsTask]
 				{
 					MergeSecondaryViewVisibility();
 					MergeSecondaryViewsTask->DispatchSubsequents();
 
-				}, OcclusionCullTasks, TaskConfig.OcclusionCull.TaskPriority);
+				}, PreSyncCullingTasks, TaskConfig.OcclusionCull.TaskPriority);
 
-				for (FVisibilityViewPacket& ViewPacket : ViewPackets)
+			for (FVisibilityViewPacket& ViewPacket : ViewPackets)
+			{
+				// Force the shared primary view relevance pipe to wait until the merge task completes.
+				if (ViewPacket.View.StereoPass == EStereoscopicPass::eSSP_PRIMARY)
 				{
-					const FViewInfo& View = ViewPacket.View;
-
-					if (View.StereoPass == EStereoscopicPass::eSSP_PRIMARY)
-					{
-						// Force the primary view relevance pipe to wait until the merge task completes.
-						ViewPacket.Relevance.CommandPipe.SetPrerequisiteTask(MergeSecondaryViewsTask);
-					}
+					ViewPacket.Relevance.CommandPipe.SetPrerequisiteTask(MergeSecondaryViewsTask);
 				}
 			}
 		}
+
+		else if (Views.Num() == 1)
+		{
+			// When using a single view, dynamic mesh elements are pushed into a pipe that is executed on the render thread which allows for some overlap with compute relevance work.
+			DynamicMeshElements.CommandPipe = Allocator.Create<TCommandPipe<FDynamicPrimitiveIndexList>>(TEXT("GatherDynamicMeshElements"), ENamedThreads::GetRenderThread_Local());
+
+			DynamicMeshElements.CommandPipe->SetCommandFunction([this](const FDynamicPrimitiveIndexList& DynamicPrimitiveIndexList)
+				{
+					GatherDynamicMeshElements(DynamicPrimitiveIndexList);
+				});
+
+			// Fire an event when the pipe has completed all work.
+			DynamicMeshElements.CommandPipeCompleteEvent = FGraphEvent::CreateGraphEvent();
+			DynamicMeshElements.CommandPipe->SetEmptyTaskEvent(DynamicMeshElements.CommandPipeCompleteEvent);
+
+			// Take a reference that is released when the relevance pipe has completed. We only need to take one since there can only be one view.
+			DynamicMeshElements.CommandPipe->AddNumCommands(1);
+
+			// We don't need the primitive view masks when in parallel mode with a single view.
+			bAllocatePrimitiveViewMasks = false;
+		}
 	}
 
 	if (bAllocatePrimitiveViewMasks)
@@ -3629,6 +3677,20 @@ void FVisibilityTaskData::LaunchVisibilityTasks()
 	{
 		SceneRenderer.WaitOcclusionTests(RHICmdList);
 
+		// In instanced stereo, we'll also redirect all primitives to the primary view's relevance command pipe, so secondary viewports need a reference
+		if (Views[0]->bIsMultiViewportEnabled)
+		{
+			int32 PrimaryViewIndex = Views[0]->PrimaryViewIndex;
+			for (FVisibilityViewPacket& ViewPacket : ViewPackets)
+			{
+				if (ViewPacket.View.StereoPass == EStereoscopicPass::eSSP_SECONDARY)
+				{
+					ViewPacket.Relevance.PrimaryViewCommandPipe = &ViewPackets[PrimaryViewIndex].Relevance.CommandPipe;
+					ViewPacket.Relevance.PrimaryViewCommandPipe->AddNumCommands(1); // Ensure the primary pipe doesn't close until the secondary pipe has forwarded all commands
+				}
+			}
+		}
+
 		for (FVisibilityViewPacket& ViewPacket : ViewPackets)
 		{
 			if (ViewPacket.OcclusionCull.ContextIfParallel)
diff --git a/Engine/Source/Runtime/Renderer/Private/SceneVisibilityPrivate.h b/Engine/Source/Runtime/Renderer/Private/SceneVisibilityPrivate.h
--- a/Engine/Source/Runtime/Renderer/Private/SceneVisibilityPrivate.h
+++ b/Engine/Source/Runtime/Renderer/Private/SceneVisibilityPrivate.h
@@ -333,6 +333,7 @@ private:
 	{
 		FComputeAndMarkRelevance* Context = nullptr;
 		TCommandPipe<FPrimitiveIndexList> CommandPipe{ TEXT("RelevancePipe") };
+		TCommandPipe<FPrimitiveIndexList>* PrimaryViewCommandPipe; // When using instanced stereo, secondary views will also send their commands to the primary view's command pipe to merge data
 
 	} Relevance;
 
@@ -702,6 +703,7 @@ private:
 	const bool bLaunchOnAddPrimitive;
 	bool bFinished = false;
 	bool bFinalized = false;
+	FSceneBitArray InstancedPrimitiveAddedMap;
 };
 
 ///////////////////////////////////////////////////////////////////////////////
