Fix hair flickering when LOD transit between strands and cards LOD.

This was due to a race condition between Groom LOD selection and ViewRelevance computation, causing the card/meshes to not be displayed for one frame.

diff --git a/Engine/Source/Runtime/Renderer/Private/HairStrands/HairStrandsInterface.cpp b/Engine/Source/Runtime/Renderer/Private/HairStrands/HairStrandsInterface.cpp
--- a/Engine/Source/Runtime/Renderer/Private/HairStrands/HairStrandsInterface.cpp
+++ b/Engine/Source/Runtime/Renderer/Private/HairStrands/HairStrandsInterface.cpp
@@ -381,41 +381,45 @@ void RunHairStrandsBookmark(EHairStrandsBookmark Bookmark, FHairStrandsBookmarkP
 	}
 }
 
-void CreateHairStrandsBookmarkParameters(FScene* Scene, FViewInfo& View, FHairStrandsBookmarkParameters& Out)
+void CreateHairStrandsBookmarkParameters(FScene* Scene, FViewInfo& View, FHairStrandsBookmarkParameters& Out, bool bComputeVisibleInstances)
 {
-	const int32 ActiveInstanceCount = Scene->HairStrandsSceneData.RegisteredProxies.Num();
-	Out.InstancesVisibility.Init(false, ActiveInstanceCount);
-
-	// 1. Strands - Add all visible strands instances
-	Out.VisibleStrands.Reserve(View.HairStrandsMeshElements.Num());
-	for (const FMeshBatchAndRelevance& MeshBatch : View.HairStrandsMeshElements)
+	// Only compute visible instances when required as this is expensive.
+	if (bComputeVisibleInstances)
 	{
-		check(MeshBatch.PrimitiveSceneProxy && MeshBatch.PrimitiveSceneProxy->ShouldRenderInMainPass());
-		if (MeshBatch.Mesh && MeshBatch.Mesh->Elements.Num() > 0)
+		const int32 ActiveInstanceCount = Scene->HairStrandsSceneData.RegisteredProxies.Num();
+		Out.InstancesVisibility.Init(false, ActiveInstanceCount);
+	
+		// 1. Strands - Add all visible strands instances
+		Out.VisibleStrands.Reserve(View.HairStrandsMeshElements.Num());
+		for (const FMeshBatchAndRelevance& MeshBatch : View.HairStrandsMeshElements)
 		{
-			FHairGroupPublicData* HairData = HairStrands::GetHairData(MeshBatch.Mesh);
-			if (HairData && HairData->Instance)
+			check(MeshBatch.PrimitiveSceneProxy && MeshBatch.PrimitiveSceneProxy->ShouldRenderInMainPass());
+			if (MeshBatch.Mesh && MeshBatch.Mesh->Elements.Num() > 0)
 			{
-				Out.VisibleStrands.Add(HairData->Instance);
-				Out.InstancesVisibility[HairData->Instance->RegisteredIndex] = true;
-				Out.InstanceCountPerType[uint32(EHairInstanceCount::StrandsPrimaryView)]++;
+				FHairGroupPublicData* HairData = HairStrands::GetHairData(MeshBatch.Mesh);
+				if (HairData && HairData->Instance)
+				{
+					Out.VisibleStrands.Add(HairData->Instance);
+					Out.InstancesVisibility[HairData->Instance->RegisteredIndex] = true;
+					Out.InstanceCountPerType[uint32(EHairInstanceCount::StrandsPrimaryView)]++;
+				}
 			}
 		}
-	}
-
-	// 2. Cards/Meshes - Add all visible cards instances
-	Out.VisibleCardsOrMeshes_Primary.Reserve(View.HairStrandsMeshElements.Num());
-	for (const FMeshBatchAndRelevance& MeshBatch : View.HairCardsMeshElements)
-	{
-		check(MeshBatch.PrimitiveSceneProxy && MeshBatch.PrimitiveSceneProxy->ShouldRenderInMainPass());
-		if (MeshBatch.Mesh && MeshBatch.Mesh->Elements.Num() > 0)
+	
+		// 2. Cards/Meshes - Add all visible cards instances
+		Out.VisibleCardsOrMeshes_Primary.Reserve(View.HairCardsMeshElements.Num());
+		for (const FMeshBatchAndRelevance& MeshBatch : View.HairCardsMeshElements)
 		{
-			FHairGroupPublicData* HairData = HairStrands::GetHairData(MeshBatch.Mesh);
-			if (HairData && HairData->Instance)
+			check(MeshBatch.PrimitiveSceneProxy && MeshBatch.PrimitiveSceneProxy->ShouldRenderInMainPass());
+			if (MeshBatch.Mesh && MeshBatch.Mesh->Elements.Num() > 0)
 			{
-				Out.VisibleCardsOrMeshes_Primary.Add(HairData->Instance);
-				Out.InstancesVisibility[HairData->Instance->RegisteredIndex] = true;
-				Out.InstanceCountPerType[uint32(EHairInstanceCount::CardsOrMeshesPrimaryView)]++;
+				FHairGroupPublicData* HairData = HairStrands::GetHairData(MeshBatch.Mesh);
+				if (HairData && HairData->Instance)
+				{
+					Out.VisibleCardsOrMeshes_Primary.Add(HairData->Instance);
+					Out.InstancesVisibility[HairData->Instance->RegisteredIndex] = true;
+					Out.InstanceCountPerType[uint32(EHairInstanceCount::CardsOrMeshesPrimaryView)]++;
+				}
 			}
 		}
 	}
@@ -467,9 +471,9 @@ void UpdateHairStrandsBookmarkParameters(FScene* Scene, TArray<FViewInfo>& Views
 	}
 }
 
-void CreateHairStrandsBookmarkParameters(FScene* Scene, TArray<FViewInfo>& Views, TArray<const FSceneView*>& AllFamilyViews, FHairStrandsBookmarkParameters& Out)
+void CreateHairStrandsBookmarkParameters(FScene* Scene, TArray<FViewInfo>& Views, TArray<const FSceneView*>& AllFamilyViews, FHairStrandsBookmarkParameters& Out, bool bComputeVisibleInstances)
 {
-	CreateHairStrandsBookmarkParameters(Scene, Views[0], Out);
+	CreateHairStrandsBookmarkParameters(Scene, Views[0], Out, bComputeVisibleInstances);
 	Out.AllViews = AllFamilyViews;
 }
 
diff --git a/Engine/Source/Runtime/Renderer/Private/HairStrands/HairStrandsRendering.h b/Engine/Source/Runtime/Renderer/Private/HairStrands/HairStrandsRendering.h
--- a/Engine/Source/Runtime/Renderer/Private/HairStrands/HairStrandsRendering.h
+++ b/Engine/Source/Runtime/Renderer/Private/HairStrands/HairStrandsRendering.h
@@ -43,6 +43,6 @@ void RunHairStrandsBookmark(
 	EHairStrandsBookmark Bookmark,
 	FHairStrandsBookmarkParameters& Parameters);
 
-void CreateHairStrandsBookmarkParameters(FScene* Scene, FViewInfo& View, FHairStrandsBookmarkParameters& Out);
-void CreateHairStrandsBookmarkParameters(FScene* Scene, TArray<FViewInfo>& Views, TArray<const FSceneView*>& AllFamilyViews, FHairStrandsBookmarkParameters& Out);
+void CreateHairStrandsBookmarkParameters(FScene* Scene, FViewInfo& View, FHairStrandsBookmarkParameters& Out, bool bComputeVisibleInstances=true);
+void CreateHairStrandsBookmarkParameters(FScene* Scene, TArray<FViewInfo>& Views, TArray<const FSceneView*>& AllFamilyViews, FHairStrandsBookmarkParameters& Out, bool bComputeVisibleInstances=true);
 void UpdateHairStrandsBookmarkParameters(FScene* Scene, TArray<FViewInfo>& Views, FHairStrandsBookmarkParameters& Out);
diff --git a/Engine/Source/Runtime/Renderer/Private/SceneRendering.cpp b/Engine/Source/Runtime/Renderer/Private/SceneRendering.cpp
--- a/Engine/Source/Runtime/Renderer/Private/SceneRendering.cpp
+++ b/Engine/Source/Runtime/Renderer/Private/SceneRendering.cpp
@@ -3451,6 +3451,21 @@ IVisibilityTaskData* FSceneRenderer::OnRenderBegin(FRDGBuilder& GraphBuilder)
 
 	PrepareViewStateForVisibility(SceneTexturesConfig);
 
+	// Run Groom LOD selection prior to visibility for selecting appropriate LOD & geometry type
+	if (IsGroomEnabled())
+	{
+		if (Views.Num() > 0 && !ViewFamily.EngineShowFlags.HitProxies)
+		{
+			FHairStrandsBookmarkParameters Parameters;
+			CreateHairStrandsBookmarkParameters(Scene, Views, AllFamilyViews, Parameters, false/*bComputeVisibleInstances*/);
+			if (Parameters.HasInstances())
+			{
+				// 1. Select appropriate LOD & geometry type
+				RunHairStrandsBookmark(GraphBuilder, EHairStrandsBookmark::ProcessLODSelection, Parameters);
+			}
+		}
+	}
+
 	FVisualizeTexturePresent::OnStartRender(Views[0]);
 
 	GraphBuilder.RHICmdList.BeginScene();
@@ -4632,7 +4647,7 @@ static void RenderViewFamilies_RenderThread(FRHICommandListImmediate& RHICmdList
 		if (IsHairStrandsEnabled(EHairStrandsShaderType::All, Scene->GetShaderPlatform()) && (SceneRenderers[0]->AllFamilyViews.Num() > 0) && !bAnyShowHitProxies)
 		{
 			FHairStrandsBookmarkParameters Parameters;
-			CreateHairStrandsBookmarkParameters(Scene, SceneRenderers[0]->Views, SceneRenderers[0]->AllFamilyViews, Parameters);
+			CreateHairStrandsBookmarkParameters(Scene, SceneRenderers[0]->Views, SceneRenderers[0]->AllFamilyViews, Parameters, false /*bComputeVisibleInstances*/);
 			if (Parameters.HasInstances())
 			{
 				RunHairStrandsBookmark(EHairStrandsBookmark::ProcessEndOfFrame, Parameters);
diff --git a/Engine/Source/Runtime/Renderer/Private/SceneVisibility.cpp b/Engine/Source/Runtime/Renderer/Private/SceneVisibility.cpp
--- a/Engine/Source/Runtime/Renderer/Private/SceneVisibility.cpp
+++ b/Engine/Source/Runtime/Renderer/Private/SceneVisibility.cpp
@@ -4114,14 +4114,10 @@ void FSceneRenderer::PreVisibilityFrameSetup(FRDGBuilder& GraphBuilder)
 		if (Views.Num() > 0 && !ViewFamily.EngineShowFlags.HitProxies)
 		{
 			FHairStrandsBookmarkParameters Parameters; 
-			CreateHairStrandsBookmarkParameters(Scene, Views, AllFamilyViews, Parameters);
-
+			CreateHairStrandsBookmarkParameters(Scene, Views, AllFamilyViews, Parameters, false /*bComputeVisibleInstances*/);
 			if (Parameters.HasInstances())
 			{
-				// 1. Select appropriate LOD & geometry type
-				RunHairStrandsBookmark(GraphBuilder, EHairStrandsBookmark::ProcessLODSelection, Parameters);
-
-				// 2. Run guide/simulation update
+				// Prepare (skel.) data for guide/simulation update
 				// If we are rendering from scene capture we don't need to run another time the hair bookmarks.
 				if (IsHairStrandsEnabled(EHairStrandsShaderType::All, Scene->GetShaderPlatform()) && Views[0].AllowGPUParticleUpdate())
 				{
