LFV - render LFV underwater as part of the fog pass

diff --git a/Engine/Source/Runtime/Renderer/Private/FogRendering.cpp b/Engine/Source/Runtime/Renderer/Private/FogRendering.cpp
--- a/Engine/Source/Runtime/Renderer/Private/FogRendering.cpp
+++ b/Engine/Source/Runtime/Renderer/Private/FogRendering.cpp
@@ -305,9 +305,6 @@ static void RenderViewFog(
 	GraphicsPSOInit.BoundShaderState.VertexDeclarationRHI = GFogVertexDeclaration.VertexDeclarationRHI;
 	GraphicsPSOInit.BoundShaderState.VertexShaderRHI = VertexShader.GetVertexShader();
 
-	// We should only render local fog volume here when volumetric fog is enabled, making sure that the depth bound test is not used.
-	check(!bFogComposeLocalFogVolumes || (bFogComposeLocalFogVolumes && bShouldRenderVolumetricFog));
-
 	const bool bUseFogInscatteringColorCubemap = View.FogInscatteringColorCubemap != nullptr;
 	FExponentialHeightFogPS::FPermutationDomain PsPermutationVector;
 	PsPermutationVector.Set<FExponentialHeightFogPS::FSupportFogInScatteringTexture>(bUseFogInscatteringColorCubemap);
@@ -449,10 +446,11 @@ void FDeferredShadingSceneRenderer::RenderUnderWaterFog(
 				PassParameters->PS.bUseWaterDepthTexture = true;
 				PassParameters->PS.WaterDepthTextureMinMaxUV = SceneWithoutWaterView.MinMaxUV;
 				PassParameters->RenderTargets[0] = FRenderTargetBinding(SceneWithoutWaterTextures.ColorTexture, ERenderTargetLoadAction::ELoad);
+				// No depth/stencil bound so depth bound clip will not work. If we enable this at some point, we will have to check LocalFogVolume to disable depth bound. Or have a start depth for it.
 
-				GraphBuilder.AddPass(RDG_EVENT_NAME("FogBehindWater"), PassParameters, ERDGPassFlags::Raster, [this, &View, SceneWithoutWaterView, PassParameters, bShouldRenderVolumetricFog](FRHICommandList& RHICmdList)
+				const bool bFogComposeLocalFogVolumes = ShouldRenderLocalFogVolume(Scene, ViewFamily); // Always render LFV as part of underwater fog, if present, to see them through the water.
+				GraphBuilder.AddPass(RDG_EVENT_NAME("FogBehindWater"), PassParameters, ERDGPassFlags::Raster, [this, &View, SceneWithoutWaterView, PassParameters, bShouldRenderVolumetricFog, bFogComposeLocalFogVolumes](FRHICommandList& RHICmdList)
 				{
-					const bool bFogComposeLocalFogVolumes = false; // LFV_TODO implement for under water
 					RenderViewFog(RHICmdList, View, SceneWithoutWaterView.ViewRect, PassParameters, bShouldRenderVolumetricFog, bFogComposeLocalFogVolumes);
 				});
 			}
