LFV - added a project wide local fov volume toggle to remove all potential usage and for it off.

diff --git a/Engine/Shaders/Private/BasePassPixelShader.usf b/Engine/Shaders/Private/BasePassPixelShader.usf
--- a/Engine/Shaders/Private/BasePassPixelShader.usf
+++ b/Engine/Shaders/Private/BasePassPixelShader.usf
@@ -181,7 +181,8 @@
 #include "/Engine/Private/Substrate/SubstrateExport.ush"
 #endif
 
-#if (NEEDS_BASEPASS_PIXEL_VOLUMETRIC_FOGGING && OPAQUE_NEEDS_BASEPASS_FOGGING)
+#define ENABLE_LOCAL_FOG_VOLUMES_ON_OPAQUE_FORWARD (PROJECT_SUPPORTS_LOCALFOGVOLUME && NEEDS_BASEPASS_PIXEL_VOLUMETRIC_FOGGING && OPAQUE_NEEDS_BASEPASS_FOGGING)
+#if ENABLE_LOCAL_FOG_VOLUMES_ON_OPAQUE_FORWARD
 #include "LocalFogVolumes/LocalFogVolumeCommon.ush"
 #endif
 
@@ -1310,7 +1311,7 @@ void FPixelShaderInOut_MainPS(
 #if NEEDS_BASEPASS_PIXEL_VOLUMETRIC_FOGGING
 	if (FogStruct.ApplyVolumetricFog > 0) 
 	{
-	#if OPAQUE_NEEDS_BASEPASS_FOGGING
+	#if ENABLE_LOCAL_FOG_VOLUMES_ON_OPAQUE_FORWARD
 		float3 TranslatedWorldPosition = SvPositionToResolvedTranslatedWorld(In.SvPosition);
 		uint2 TilePos = clamp(uint2(In.SvPosition.xy / float(LFVTilePixelSize)), uint2(0, 0), LFVTileDataResolution - 1);
 		float4 LocalFogVolumeContrib = GetLFVContribution(ResolvedView, TilePos, TranslatedWorldPosition);
diff --git a/Engine/Shaders/Public/ShaderVersion.ush b/Engine/Shaders/Public/ShaderVersion.ush
--- a/Engine/Shaders/Public/ShaderVersion.ush
+++ b/Engine/Shaders/Public/ShaderVersion.ush
@@ -2,4 +2,4 @@
 // Each time the console command is executed it generates a new GUID. As this file is included
 // in Platform.ush (which should be included in any shader) it allows to invalidate the shader DDC.
 // 
-#pragma message("UESHADERMETADATA_VERSION 959736DE-626A-4776-A2DC-36960B1A0099")
+#pragma message("UESHADERMETADATA_VERSION 2B90878D-98D5-4E42-933B-178267A501C0")
diff --git a/Engine/Source/Runtime/Engine/Classes/Engine/RendererSettings.h b/Engine/Source/Runtime/Engine/Classes/Engine/RendererSettings.h
--- a/Engine/Source/Runtime/Engine/Classes/Engine/RendererSettings.h
+++ b/Engine/Source/Runtime/Engine/Classes/Engine/RendererSettings.h
@@ -855,6 +855,16 @@ class URendererSettings : public UDeveloperSettings
 		ConfigRestartRequired = true))
 		uint32 bSupportSkyAtmosphereAffectsHeightFog : 1;
 
+	/**
+	"Local fog volume components can will need to be applied on translucent, and opaque in forward, so resources will need to be bound to apply aerial perspective on transparent surfaces (and all surfaces on mobile via per vertex evaluation)."
+	"It requires r.SupportLocalFogVolumes to be true."
+	*/
+	UPROPERTY(config, EditAnywhere, Category = Optimizations, meta = (
+		ConsoleVariable = "r.SupportLocalFogVolumes", DisplayName = "Support Local Fog Volumes",
+		ToolTip = "Local fog volume components can will need to be applied on translucent, and opaque in forward, so resources will need to be bound to apply aerial perspective on transparent surfaces (and all surfaces on mobile via per vertex evaluation). It requires r.SupportLocalFogVolumes to be true.",
+		ConfigRestartRequired = true))
+		uint32 bSupportLocalFogVolumes : 1;
+
 	/**
 	"Enable cloud shadow on translucent surface. This is evaluated per vertex to reduce GPU cost. The cloud system requires extra samplers/textures to be bound to vertex shaders."
 	*/
diff --git a/Engine/Source/Runtime/Engine/Private/ShaderCompiler/ShaderCompiler.cpp b/Engine/Source/Runtime/Engine/Private/ShaderCompiler/ShaderCompiler.cpp
--- a/Engine/Source/Runtime/Engine/Private/ShaderCompiler/ShaderCompiler.cpp
+++ b/Engine/Source/Runtime/Engine/Private/ShaderCompiler/ShaderCompiler.cpp
@@ -8012,10 +8012,17 @@ void GlobalBeginCompileShader(
 		SET_SHADER_DEFINE(Input.Environment, PROJECT_MOBILE_DISABLE_VERTEX_FOG, CVar ? (CVar->GetInt() != 0) : 0);
 	}
 
+	bool bSupportLocalFogVolumes = false;
+	{
+		static IConsoleVariable* CVar = IConsoleManager::Get().FindConsoleVariable(TEXT("r.SupportLocalFogVolumes"));
+		bSupportLocalFogVolumes = CVar && CVar->GetInt() > 0;
+		SET_SHADER_DEFINE(Input.Environment, PROJECT_SUPPORTS_LOCALFOGVOLUME, (bSupportLocalFogVolumes ? 1 : 0));
+	}
+
 	{
 		static IConsoleVariable* CVar = IConsoleManager::Get().FindConsoleVariable(TEXT("r.LocalFogVolume.ApplyOnTranslucent"));
 		const bool bLocalFogVolumesApplyOnTranclucent = CVar && CVar->GetInt() > 0;
-		SET_SHADER_DEFINE(Input.Environment, PROJECT_LOCALFOGVOLUME_APPLYONTRANSLUCENT, (bLocalFogVolumesApplyOnTranclucent ? 1 : 0));
+		SET_SHADER_DEFINE(Input.Environment, PROJECT_LOCALFOGVOLUME_APPLYONTRANSLUCENT, ((bSupportLocalFogVolumes && bLocalFogVolumesApplyOnTranclucent) ? 1 : 0));
 	}
 
 	bool bSupportSkyAtmosphere = false;
diff --git a/Engine/Source/Runtime/RenderCore/Private/Shader.cpp b/Engine/Source/Runtime/RenderCore/Private/Shader.cpp
--- a/Engine/Source/Runtime/RenderCore/Private/Shader.cpp
+++ b/Engine/Source/Runtime/RenderCore/Private/Shader.cpp
@@ -1935,10 +1935,20 @@ void ShaderMapAppendKeyString(EShaderPlatform Platform, FString& KeyString)
 		}
 	}
 
+	bool bSupportLocalFogVolumes = false;
+	{
+		static IConsoleVariable* CVar = IConsoleManager::Get().FindConsoleVariable(TEXT("r.SupportLocalFogVolumes"));
+		bSupportLocalFogVolumes = CVar && CVar->GetInt() > 0;
+		if (bSupportLocalFogVolumes)
+		{
+			KeyString += TEXT("_LFV");
+		}
+	}
+
 	{
 		static IConsoleVariable* CVar = IConsoleManager::Get().FindConsoleVariable(TEXT("r.LocalFogVolume.ApplyOnTranslucent"));
 		const bool bLocalFogVolumesApplyOnTranclucent = CVar && CVar->GetInt() > 0;
-		if (bLocalFogVolumesApplyOnTranclucent)
+		if (bSupportLocalFogVolumes && bLocalFogVolumesApplyOnTranclucent)
 		{
 			KeyString += TEXT("_LFVTRA");
 		}
diff --git a/Engine/Source/Runtime/Renderer/Private/LocalFogVolumeRendering.cpp b/Engine/Source/Runtime/Renderer/Private/LocalFogVolumeRendering.cpp
--- a/Engine/Source/Runtime/Renderer/Private/LocalFogVolumeRendering.cpp
+++ b/Engine/Source/Runtime/Renderer/Private/LocalFogVolumeRendering.cpp
@@ -14,6 +14,13 @@ static TAutoConsoleVariable<int32> CVarLocalFogVolume(
 	TEXT("LocalFogVolume components are rendered when this is not 0, otherwise ignored.\n"),
 	ECVF_RenderThreadSafe);
 
+// The project setting (disable runtime and shader code)
+static TAutoConsoleVariable<int32> CVarSupportLocalFogVolumes(
+	TEXT("r.SupportLocalFogVolumes"),
+	1,
+	TEXT("Enables local fog volume rendering and shader code."),
+	ECVF_ReadOnly | ECVF_RenderThreadSafe);
+
 static TAutoConsoleVariable<int32> CVarLocalFogVolumeRenderIntoVolumetricFog(
 	TEXT("r.LocalFogVolume.RenderIntoVolumetricFog"), 1,
 	TEXT("LocalFogVolume are going to be voxelised into the volumetric fog when this is not 0, otherwise it will remain isolated.\n"),
@@ -79,7 +86,7 @@ bool ShouldRenderLocalFogVolume(const FScene* Scene, const FSceneViewFamily& Sce
 	const FEngineShowFlags EngineShowFlags = SceneViewFamily.EngineShowFlags;
 	if (Scene && Scene->HasAnyLocalFogVolume() && EngineShowFlags.Fog && !SceneViewFamily.UseDebugViewPS())
 	{
-		return CVarLocalFogVolume.GetValueOnRenderThread() > 0;
+		return (CVarSupportLocalFogVolumes.GetValueOnRenderThread() > 0) && (CVarLocalFogVolume.GetValueOnRenderThread() > 0);
 	}
 	return false;
 }
