LFV - added an option to render LFV during the height fog pass. A non mobile only path experiment. Not faster as of today it seems.

diff --git a/Engine/Source/Runtime/Renderer/Private/DeferredShadingRenderer.cpp b/Engine/Source/Runtime/Renderer/Private/DeferredShadingRenderer.cpp
--- a/Engine/Source/Runtime/Renderer/Private/DeferredShadingRenderer.cpp
+++ b/Engine/Source/Runtime/Renderer/Private/DeferredShadingRenderer.cpp
@@ -3131,6 +3131,7 @@ void FDeferredShadingSceneRenderer::Render(FRDGBuilder& GraphBuilder)
 	const bool bUseGBuffer = IsUsingGBuffers(ShaderPlatform);
 	const bool bShouldRenderVolumetricFog = ShouldRenderVolumetricFog();
 	const bool bShouldRenderLocalFogVolume = ShouldRenderLocalFogVolume(Scene, ViewFamily);
+	const bool bShouldRenderLocalFogVolumeDuringHeightFogPass = ShouldRenderLocalFogVolumeDuringHeightFogPass(Scene, ViewFamily);
 	const bool bShouldRenderLocalFogVolumeInVolumetricFog = ShouldRenderLocalFogVolumeInVolumetricFog(Scene, ViewFamily, bShouldRenderLocalFogVolume);
 
 	const bool bRenderDeferredLighting = ViewFamily.EngineShowFlags.Lighting
@@ -4122,7 +4123,7 @@ void FDeferredShadingSceneRenderer::Render(FRDGBuilder& GraphBuilder)
 			RDG_CSV_STAT_EXCLUSIVE_SCOPE(GraphBuilder, RenderFog);
 			SCOPED_NAMED_EVENT(RenderFog, FColor::Emerald);
 			SCOPE_CYCLE_COUNTER(STAT_FDeferredShadingSceneRenderer_RenderFog);
-			const bool bFogComposeLocalFogVolumes = bShouldRenderLocalFogVolumeInVolumetricFog && bShouldRenderVolumetricFog;
+			const bool bFogComposeLocalFogVolumes = bShouldRenderLocalFogVolumeInVolumetricFog && bShouldRenderVolumetricFog || bShouldRenderLocalFogVolumeDuringHeightFogPass;
 			RenderFog(GraphBuilder, SceneTextures, LightShaftOcclusionTexture, bFogComposeLocalFogVolumes);
 			bHeightFogHasComposedLocalFogVolume = bFogComposeLocalFogVolumes;
 		}
diff --git a/Engine/Source/Runtime/Renderer/Private/FogRendering.cpp b/Engine/Source/Runtime/Renderer/Private/FogRendering.cpp
--- a/Engine/Source/Runtime/Renderer/Private/FogRendering.cpp
+++ b/Engine/Source/Runtime/Renderer/Private/FogRendering.cpp
@@ -288,7 +288,8 @@ static void RenderViewFog(
 	FIntRect ViewRect, 
 	FFogPassParameters* PassParameters, 
 	bool bShouldRenderVolumetricFog,
-	bool bFogComposeLocalFogVolumes)
+	bool bFogComposeLocalFogVolumes,
+	bool bSkipDepthBound = false)
 {
 	FGraphicsPipelineStateInitializer GraphicsPSOInit;
 	RHICmdList.ApplyCachedRenderTargets(GraphicsPSOInit);
@@ -315,7 +316,7 @@ static void RenderViewFog(
 	GraphicsPSOInit.BoundShaderState.PixelShaderRHI = PixelShader.GetPixelShader();
 
 	// Setup the depth bound optimization if possible on that platform.
-	GraphicsPSOInit.bDepthBounds = GSupportsDepthBoundsTest && CVarFogUseDepthBounds.GetValueOnAnyThread();
+	GraphicsPSOInit.bDepthBounds = GSupportsDepthBoundsTest && CVarFogUseDepthBounds.GetValueOnAnyThread() && !bSkipDepthBound;
 	if (GraphicsPSOInit.bDepthBounds)
 	{
 		float ExpFogStartDistance = View.ExponentialFogParameters.W;
@@ -380,6 +381,9 @@ void FDeferredShadingSceneRenderer::RenderFog(
 
 		const bool bShouldRenderVolumetricFog = ShouldRenderVolumetricFog();
 
+		// Without volumetric fog, LFVs would not be rendered correctly according to depth and height fog start distance. Since LFVs do not have any start distance as of today.
+		const bool bSkipDepthBound = !bShouldRenderVolumetricFog && bFogComposeLocalFogVolumes && ShouldRenderLocalFogVolumeDuringHeightFogPass(Scene, ViewFamily);
+
 		for(int32 ViewIndex = 0;ViewIndex < Views.Num();ViewIndex++)
 		{
 			const FViewInfo& View = Views[ViewIndex];
@@ -399,9 +403,9 @@ void FDeferredShadingSceneRenderer::RenderFog(
 				PassParameters->RenderTargets.DepthStencil = FDepthStencilBinding(SceneTextures.Depth.Target, ERenderTargetLoadAction::ELoad, ERenderTargetLoadAction::ELoad, FExclusiveDepthStencil::DepthRead_StencilWrite);
 
 				GraphBuilder.AddPass(RDG_EVENT_NAME("Fog"), PassParameters, ERDGPassFlags::Raster, 
-					[this, &View, PassParameters, bShouldRenderVolumetricFog, bFogComposeLocalFogVolumes](FRHICommandList& RHICmdList)
+					[this, &View, PassParameters, bShouldRenderVolumetricFog, bFogComposeLocalFogVolumes, bSkipDepthBound](FRHICommandList& RHICmdList)
 				{
-					RenderViewFog(RHICmdList, View, View.ViewRect, PassParameters, bShouldRenderVolumetricFog, bFogComposeLocalFogVolumes);
+					RenderViewFog(RHICmdList, View, View.ViewRect, PassParameters, bShouldRenderVolumetricFog, bFogComposeLocalFogVolumes, bSkipDepthBound);
 				});
 			}
 		}
diff --git a/Engine/Source/Runtime/Renderer/Private/LocalFogVolumeRendering.cpp b/Engine/Source/Runtime/Renderer/Private/LocalFogVolumeRendering.cpp
--- a/Engine/Source/Runtime/Renderer/Private/LocalFogVolumeRendering.cpp
+++ b/Engine/Source/Runtime/Renderer/Private/LocalFogVolumeRendering.cpp
@@ -21,6 +21,11 @@ static TAutoConsoleVariable<int32> CVarSupportLocalFogVolumes(
 	TEXT("Enables local fog volume rendering and shader code."),
 	ECVF_ReadOnly | ECVF_RenderThreadSafe);
 
+static TAutoConsoleVariable<int32> CVarLocalFogVolumeRenderDuringHeightFogPass(
+	TEXT("r.LocalFogVolume.RenderDuringHeightFogPass"), 0,
+	TEXT("LocalFogVolume are going to be rendered during the height fog pass, skipping the tiled rendering pass specific to them. Only work on the non mobile path as an experiment.\n"),
+	ECVF_RenderThreadSafe);
+
 static TAutoConsoleVariable<int32> CVarLocalFogVolumeRenderIntoVolumetricFog(
 	TEXT("r.LocalFogVolume.RenderIntoVolumetricFog"), 1,
 	TEXT("LocalFogVolume are going to be voxelised into the volumetric fog when this is not 0, otherwise it will remain isolated.\n"),
@@ -91,6 +96,15 @@ bool ShouldRenderLocalFogVolume(const FScene* Scene, const FSceneViewFamily& Sce
 	return false;
 }
 
+bool ShouldRenderLocalFogVolumeDuringHeightFogPass(const FScene* Scene, const FSceneViewFamily& SceneViewFamily)
+{
+	if (ShouldRenderLocalFogVolume(Scene, SceneViewFamily))
+	{
+		return CVarLocalFogVolumeRenderDuringHeightFogPass.GetValueOnRenderThread() > 0;
+	}
+	return false;
+}
+
 bool ShouldRenderLocalFogVolumeInVolumetricFog(const FScene* Scene, const FSceneViewFamily& SceneViewFamily, bool bShouldRenderVolumetricFog)
 {
 	if (ShouldRenderLocalFogVolume(Scene, SceneViewFamily) && bShouldRenderVolumetricFog)
diff --git a/Engine/Source/Runtime/Renderer/Private/LocalFogVolumeRendering.h b/Engine/Source/Runtime/Renderer/Private/LocalFogVolumeRendering.h
--- a/Engine/Source/Runtime/Renderer/Private/LocalFogVolumeRendering.h
+++ b/Engine/Source/Runtime/Renderer/Private/LocalFogVolumeRendering.h
@@ -117,6 +117,7 @@ struct FLocalFogVolumeViewData
 =============================================================================*/
 
 bool ShouldRenderLocalFogVolume(const FScene* Scene, const FSceneViewFamily& SceneViewFamily);
+bool ShouldRenderLocalFogVolumeDuringHeightFogPass(const FScene* Scene, const FSceneViewFamily& SceneViewFamily);
 bool ShouldRenderLocalFogVolumeInVolumetricFog(const FScene* Scene, const FSceneViewFamily& SceneViewFamily, bool bShouldRenderVolumetricFog);
 
 void GetLocalFogVolumeSortingData(const FScene* Scene, FRDGBuilder& GraphBuilder, FLocalFogVolumeSortingData& Out);
