LFV - Also move the indirect draw command clear onto the async pipe.

diff --git a/Engine/Source/Runtime/Renderer/Private/LocalFogVolumeRendering.cpp b/Engine/Source/Runtime/Renderer/Private/LocalFogVolumeRendering.cpp
--- a/Engine/Source/Runtime/Renderer/Private/LocalFogVolumeRendering.cpp
+++ b/Engine/Source/Runtime/Renderer/Private/LocalFogVolumeRendering.cpp
@@ -72,7 +72,7 @@ static uint32 GetLocalFogVolumeTilePixelSize()
 
 static bool GetLocalFogVolumeTileCullingUseAsync()
 {
-	return CVarLocalFogVolumeTileCullingUseAsync.GetValueOnRenderThread() > 0;
+	return GSupportsEfficientAsyncCompute && CVarLocalFogVolumeTileCullingUseAsync.GetValueOnRenderThread() > 0;
 }
 
 static float GetLocalFogVolumeMaxDensityIntoVolumetricFog()
@@ -286,9 +286,8 @@ static void LocalFogVolumeViewTiledCullingPass(FViewInfo& View, FRDGBuilder& Gra
 	float TileCoveredResolutionX = View.LocalFogVolumeViewData.UniformParametersStruct.LocalFogVolumeCommon.LocalFogVolumeTilePixelSize * View.LocalFogVolumeViewData.UniformParametersStruct.LocalFogVolumeCommon.LocalFogVolumeTileDataTextureResolution.X;
 	float TileCoveredResolutionY = View.LocalFogVolumeViewData.UniformParametersStruct.LocalFogVolumeCommon.LocalFogVolumeTilePixelSize * View.LocalFogVolumeViewData.UniformParametersStruct.LocalFogVolumeCommon.LocalFogVolumeTileDataTextureResolution.Y;
 	PassParameters->ViewToTileSpaceRatio = FVector2f(TileCoveredResolutionX * View.CachedViewUniformShaderParameters->ViewSizeAndInvSize.Z, TileCoveredResolutionY * View.CachedViewUniformShaderParameters->ViewSizeAndInvSize.W);
-
-	const bool bUseAsyncCompute = GSupportsEfficientAsyncCompute && GetLocalFogVolumeTileCullingUseAsync();
-	ERDGPassFlags PassFlag = bUseAsyncCompute ? ERDGPassFlags::AsyncCompute : ERDGPassFlags::Compute;
+ 
+	ERDGPassFlags PassFlag = GetLocalFogVolumeTileCullingUseAsync() ? ERDGPassFlags::AsyncCompute : ERDGPassFlags::Compute;
 
 	TileDataTextureSize.Z = 1;
 	const FIntVector NumGroups = FIntVector::DivideAndRoundUp(TileDataTextureSize, FLocalFogVolumeTiledCullingCS::GroupSize);
@@ -486,7 +485,7 @@ void CreateViewLocalFogVolumeBufferSRV(const FScene* Scene, FViewInfo& View, FRD
 
 	View.LocalFogVolumeViewData.GPUTileDrawIndirectBuffer = GraphBuilder.CreateBuffer(FRDGBufferDesc::CreateIndirectDesc<FRHIDrawIndirectParameters>(), TEXT("LocalFogVolume.DispatchIndirectBuffer"));
 	View.LocalFogVolumeViewData.GPUTileDrawIndirectBufferUAV = GraphBuilder.CreateUAV(View.LocalFogVolumeViewData.GPUTileDrawIndirectBuffer, PF_R32_UINT);
-	AddClearUAVPass(GraphBuilder, View.LocalFogVolumeViewData.GPUTileDrawIndirectBufferUAV, 0);
+	AddClearUAVPass(GraphBuilder, View.LocalFogVolumeViewData.GPUTileDrawIndirectBufferUAV, 0, GetLocalFogVolumeTileCullingUseAsync() ? ERDGPassFlags::AsyncCompute : ERDGPassFlags::Compute);
 }
 
 void InitLocalFogVolumesForViews(
