LFV - added warning on screen when the scene has LFVs but it is disabled from the project settings.

diff --git a/Engine/Source/Runtime/Renderer/Private/LocalFogVolumeRendering.cpp b/Engine/Source/Runtime/Renderer/Private/LocalFogVolumeRendering.cpp
--- a/Engine/Source/Runtime/Renderer/Private/LocalFogVolumeRendering.cpp
+++ b/Engine/Source/Runtime/Renderer/Private/LocalFogVolumeRendering.cpp
@@ -86,12 +86,17 @@ static uint32 GetLocalFogVolumeTileMaxInstanceCount()
 	return FMath::Max(1u, FMath::Min(256u, (uint32)CVarLocalFogVolumeTileMaxInstanceCount.GetValueOnRenderThread()));
 }
 
+bool ProjectSupportsLocalFogVolumes()
+{
+	return CVarSupportLocalFogVolumes.GetValueOnRenderThread() > 0;
+}
+
 bool ShouldRenderLocalFogVolume(const FScene* Scene, const FSceneViewFamily& SceneViewFamily)
 {
 	const FEngineShowFlags EngineShowFlags = SceneViewFamily.EngineShowFlags;
 	if (Scene && Scene->HasAnyLocalFogVolume() && EngineShowFlags.Fog && !SceneViewFamily.UseDebugViewPS())
 	{
-		return (CVarSupportLocalFogVolumes.GetValueOnRenderThread() > 0) && (CVarLocalFogVolume.GetValueOnRenderThread() > 0);
+		return ProjectSupportsLocalFogVolumes() && (CVarLocalFogVolume.GetValueOnRenderThread() > 0);
 	}
 	return false;
 }
diff --git a/Engine/Source/Runtime/Renderer/Private/LocalFogVolumeRendering.h b/Engine/Source/Runtime/Renderer/Private/LocalFogVolumeRendering.h
--- a/Engine/Source/Runtime/Renderer/Private/LocalFogVolumeRendering.h
+++ b/Engine/Source/Runtime/Renderer/Private/LocalFogVolumeRendering.h
@@ -116,6 +116,7 @@ struct FLocalFogVolumeViewData
 	Local height fog rendering functions
 =============================================================================*/
 
+bool ProjectSupportsLocalFogVolumes();
 bool ShouldRenderLocalFogVolume(const FScene* Scene, const FSceneViewFamily& SceneViewFamily);
 bool ShouldRenderLocalFogVolumeDuringHeightFogPass(const FScene* Scene, const FSceneViewFamily& SceneViewFamily);
 bool ShouldRenderLocalFogVolumeInVolumetricFog(const FScene* Scene, const FSceneViewFamily& SceneViewFamily, bool bShouldRenderVolumetricFog);
diff --git a/Engine/Source/Runtime/Renderer/Private/SceneRendering.cpp b/Engine/Source/Runtime/Renderer/Private/SceneRendering.cpp
--- a/Engine/Source/Runtime/Renderer/Private/SceneRendering.cpp
+++ b/Engine/Source/Runtime/Renderer/Private/SceneRendering.cpp
@@ -93,6 +93,7 @@
 #include "Engine/VolumeTexture.h"
 #include "GPUDebugCrashUtils.h"
 #include "MeshDrawCommandStats.h"
+#include "LocalFogVolumeRendering.h"
 
 /*-----------------------------------------------------------------------------
 	Globals
@@ -3537,6 +3538,8 @@ void FSceneRenderer::OnRenderFinish(FRDGBuilder& GraphBuilder, FRDGTextureRef Vi
 		const bool bShowPointLightWarning = UsedWholeScenePointLightNames.Num() > 0 && !ReadOnlyCVARCache.bEnablePointLightShadows;
 		const bool bShowShadowedLightOverflowWarning = Scene->OverflowingDynamicShadowedLights.Num() > 0;
 
+		const bool bLocalFogVolumeInSceneButProjectDisabled = Scene->HasAnyLocalFogVolume() && !ProjectSupportsLocalFogVolumes();
+
 		bool bLumenEnabledButHasNoDataForTracing = false;
 		bool bLumenEnabledButDisabledForTheProject = false;
 		bool bNaniteEnabledButDisabledInProject = false;
@@ -3623,7 +3626,7 @@ void FSceneRenderer::OnRenderFinish(FRDGBuilder& GraphBuilder, FRDGTextureRef Vi
 			|| bShowDFAODisabledWarning || bShowShadowedLightOverflowWarning || bShowMobileDynamicCSMWarning || bShowMobileLowQualityLightmapWarning || bShowMobileMovableDirectionalLightWarning
 			|| bMobileMissingSkyMaterial || bShowSkinCacheOOM || bSingleLayerWaterWarning || bShowDFDisabledWarning || bShowNoSkyAtmosphereComponentWarning || bFxDebugDraw 
 			|| bLumenEnabledButHasNoDataForTracing || bLumenEnabledButDisabledForTheProject || bNaniteEnabledButNoAtomics || bNaniteEnabledButDisabledInProject || bRealTimeSkyCaptureButNothingToCapture || bShowWaitingSkylight
-			|| bShowLocalExposureDisabledWarning || bHasDelegateWarnings || bContactShadowIntensityCvarUsed
+			|| bShowLocalExposureDisabledWarning || bHasDelegateWarnings || bContactShadowIntensityCvarUsed || bLocalFogVolumeInSceneButProjectDisabled
 			;
 
 		for(int32 ViewIndex = 0;ViewIndex < Views.Num();ViewIndex++)
@@ -3670,7 +3673,8 @@ void FSceneRenderer::OnRenderFinish(FRDGBuilder& GraphBuilder, FRDGTextureRef Vi
 						bIsFrozen, bShowSkylightWarning, bShowPointLightWarning, bShowShadowedLightOverflowWarning,
 						bShowMobileLowQualityLightmapWarning, bShowMobileMovableDirectionalLightWarning, bShowMobileDynamicCSMWarning, bMobileMissingSkyMaterial, 
 						bShowSkinCacheOOM, bSingleLayerWaterWarning, bShowNoSkyAtmosphereComponentWarning, bFxDebugDraw, FXInterface, bShowLocalExposureDisabledWarning,
-						bLumenEnabledButHasNoDataForTracing, bLumenEnabledButDisabledForTheProject, bNaniteEnabledButNoAtomics, bNaniteEnabledButDisabledInProject, bRealTimeSkyCaptureButNothingToCapture, bShowWaitingSkylight, bShowAmbientCubemapMipGenSettingsWarning,
+						bLumenEnabledButHasNoDataForTracing, bLumenEnabledButDisabledForTheProject, bNaniteEnabledButNoAtomics, bNaniteEnabledButDisabledInProject, 
+						bRealTimeSkyCaptureButNothingToCapture, bShowWaitingSkylight, bShowAmbientCubemapMipGenSettingsWarning, bLocalFogVolumeInSceneButProjectDisabled,
 						bStereoView, bPrimaryStereoView, bIsInstancedStereoEnabled, bIsMultiViewportEnabled, bIsMobileMultiViewEnabled, bContactShadowIntensityCvarUsed]
 						(FCanvas& Canvas)
 					{
@@ -3826,6 +3830,12 @@ void FSceneRenderer::OnRenderFinish(FRDGBuilder& GraphBuilder, FRDGTextureRef Vi
 							Writer.DrawLine(Message);
 						}
 
+						if (bLocalFogVolumeInSceneButProjectDisabled)
+						{
+							static const FText Message = NSLOCTEXT("Renderer", "LocalFogVolumeDisabled", "There are Local Fog Volumes in the scene, but your project does not support rendering them. This can be enabled from the project settings panel (r.SupportLocalFogVolumes).");
+							Writer.DrawLine(Message);
+						}
+
 #if !UE_BUILD_SHIPPING
 						if (bStereoView)
 						{
