Fix compilation issue on CIS

diff --git a/Engine/Source/Runtime/Renderer/Private/DeferredShadingRenderer.cpp b/Engine/Source/Runtime/Renderer/Private/DeferredShadingRenderer.cpp
--- a/Engine/Source/Runtime/Renderer/Private/DeferredShadingRenderer.cpp
+++ b/Engine/Source/Runtime/Renderer/Private/DeferredShadingRenderer.cpp
@@ -4123,7 +4123,7 @@ void FDeferredShadingSceneRenderer::Render(FRDGBuilder& GraphBuilder)
 			RDG_CSV_STAT_EXCLUSIVE_SCOPE(GraphBuilder, RenderFog);
 			SCOPED_NAMED_EVENT(RenderFog, FColor::Emerald);
 			SCOPE_CYCLE_COUNTER(STAT_FDeferredShadingSceneRenderer_RenderFog);
-			const bool bFogComposeLocalFogVolumes = bShouldRenderLocalFogVolumeInVolumetricFog && bShouldRenderVolumetricFog || bShouldRenderLocalFogVolumeDuringHeightFogPass;
+			const bool bFogComposeLocalFogVolumes = (bShouldRenderLocalFogVolumeInVolumetricFog && bShouldRenderVolumetricFog) || bShouldRenderLocalFogVolumeDuringHeightFogPass;
 			RenderFog(GraphBuilder, SceneTextures, LightShaftOcclusionTexture, bFogComposeLocalFogVolumes);
 			bHeightFogHasComposedLocalFogVolume = bFogComposeLocalFogVolumes;
 		}
