Works arrounds a AMD driver crash on a TSR shader

diff --git a/Engine/Source/Runtime/Renderer/Private/PostProcess/TemporalSuperResolution.cpp b/Engine/Source/Runtime/Renderer/Private/PostProcess/TemporalSuperResolution.cpp
--- a/Engine/Source/Runtime/Renderer/Private/PostProcess/TemporalSuperResolution.cpp
+++ b/Engine/Source/Runtime/Renderer/Private/PostProcess/TemporalSuperResolution.cpp
@@ -1201,7 +1201,14 @@ FDefaultTemporalUpscaler::FOutputs AddTemporalSuperResolutionPasses(
 	
 	// Whether to use 16bit VALU
 	const ERHIFeatureSupport VALU16BitSupport = FTSRShader::Supports16BitVALU(View.GetShaderPlatform());
-	const bool bUse16BitVALU = (CVarTSR16BitVALU.GetValueOnRenderThread() != 0 && GRHIGlobals.SupportsNative16BitOps && VALU16BitSupport == ERHIFeatureSupport::RuntimeDependent) || VALU16BitSupport == ERHIFeatureSupport::RuntimeGuaranteed;
+	bool bUse16BitVALU = (CVarTSR16BitVALU.GetValueOnRenderThread() != 0 && GRHIGlobals.SupportsNative16BitOps && VALU16BitSupport == ERHIFeatureSupport::RuntimeDependent) || VALU16BitSupport == ERHIFeatureSupport::RuntimeGuaranteed;
+
+#if PLATFORM_DESKTOP
+	if (IsRHIDeviceAMD())
+	{
+		bUse16BitVALU = false;
+	}
+#endif
 
 	// Whether should accumulate sub pixel depth.
 	const ETSRSubpixelMethod SubpixelMethod = ETSRSubpixelMethod(FMath::Clamp(CVarTSRSubpixelMethod.GetValueOnRenderThread(), 0, 2));
