Fix directional light shadow rendering incorrectly in vulkan preview with mobile multi view and mobile hdr = false.

diff --git a/Engine/Source/Runtime/Renderer/Private/ShadowRendering.cpp b/Engine/Source/Runtime/Renderer/Private/ShadowRendering.cpp
--- a/Engine/Source/Runtime/Renderer/Private/ShadowRendering.cpp
+++ b/Engine/Source/Runtime/Renderer/Private/ShadowRendering.cpp
@@ -1740,7 +1740,15 @@ FMatrix FProjectedShadowInfo::GetScreenToShadowMatrix(const FSceneView& View, ui
 				)
 			);
 
-	if (View.bIsMobileMultiViewEnabled && View.Family->Views.Num() > 0)
+	// Aspects creation is embedded in lambda so we don't pay the price for it when it's not needed.
+	auto IsMobileMultiViewEnabledInAspects = [](const FSceneView& View) {
+		const UE::StereoRenderUtils::FStereoShaderAspects Aspects(View.GetShaderPlatform());
+		return Aspects.IsMobileMultiViewEnabled();
+	};
+	
+	// Checking the Aspects is a workaround for editor windows or scene captures where View.bIsMobileMultiViewEnabled
+	// is false but shaders have been compiled with MOBILE_MULTI_VIEW enabled.
+	if (View.bIsMobileMultiViewEnabled || IsMobileMultiViewEnabledInAspects(View))
 	{
 		// In Multiview, we split ViewDependentTransform out into ViewUniformShaderParameters.MobileMultiviewShadowTransform
 		// So we can multiply it later in shader.
