[Renderer], [nDisplay]: moving some 'Renderer module' headers from 'Private' to 'Internal' and removing references to the 'Private' directory in nDisplay modules.
- Renderer: The FXAA renderpass has been moved from private to internal.
- Renderer: The FSceneTextures structure from the IRendererModule::FOnResolvedSceneColor callback has been moved from private to internal.
- nDisplay: Removed references to the Private directory of the RendererModule

diff --git a/Engine/Plugins/Runtime/nDisplay/Source/DisplayCluster/DisplayCluster.Build.cs b/Engine/Plugins/Runtime/nDisplay/Source/DisplayCluster/DisplayCluster.Build.cs
--- a/Engine/Plugins/Runtime/nDisplay/Source/DisplayCluster/DisplayCluster.Build.cs
+++ b/Engine/Plugins/Runtime/nDisplay/Source/DisplayCluster/DisplayCluster.Build.cs
@@ -7,10 +7,9 @@ public class DisplayCluster : ModuleRules
 {
 	public DisplayCluster(ReadOnlyTargetRules ROTargetRules) : base(ROTargetRules)
 	{
-		PrivateIncludePaths.AddRange(
-			new string[] {
-				System.IO.Path.Combine(GetModuleDirectory("Renderer"), "Private"), //required for FPostProcessMaterialInputs
-			});
+		// The DisplayCluster plugin is distributed with engine hot fixes and thus isn't tied to binary
+		// compatibility between hotfixes by only using Public/ interface of the renderer, but also Internal/ ones.
+		bTreatAsEngineModule = true;
 
 		PublicDependencyModuleNames.AddRange(
 			new string[] {
diff --git a/Engine/Plugins/Runtime/nDisplay/Source/DisplayCluster/Private/Render/Viewport/DisplayClusterViewportFrameStatsViewExtension.cpp b/Engine/Plugins/Runtime/nDisplay/Source/DisplayCluster/Private/Render/Viewport/DisplayClusterViewportFrameStatsViewExtension.cpp
--- a/Engine/Plugins/Runtime/nDisplay/Source/DisplayCluster/Private/Render/Viewport/DisplayClusterViewportFrameStatsViewExtension.cpp
+++ b/Engine/Plugins/Runtime/nDisplay/Source/DisplayCluster/Private/Render/Viewport/DisplayClusterViewportFrameStatsViewExtension.cpp
@@ -10,13 +10,11 @@
 #include "GlobalShader.h"
 #include "Misc/App.h"
 #include "Misc/Timecode.h"
-#include "SceneRendering.h"
 #include "SceneView.h"
 #include "ScreenPass.h"
 #include "ShaderParameters.h"
 #include "ShaderParameterStruct.h"
 #include "SystemTextures.h"
-#include "PostProcess/PostProcessTestImage.h"
 #include "PostProcess/PostProcessMaterialInputs.h"
 
 
@@ -104,7 +102,6 @@ void FDisplayClusterViewportFrameStatsViewExtension::BeginRenderViewFamily(FScen
 
 FScreenPassTexture FDisplayClusterViewportFrameStatsViewExtension::PostProcessPassAfterTonemap_RenderThread(FRDGBuilder& GraphBuilder, const FSceneView& View, const FPostProcessMaterialInputs& Inputs)
 {
-	const FViewInfo& ViewInfo = static_cast<const FViewInfo&>(View);
 	const FScreenPassTexture& SceneColor = Inputs.GetInput(EPostProcessMaterialInput::SceneColor);
 	check(SceneColor.IsValid());
 
@@ -127,7 +124,7 @@ FScreenPassTexture FDisplayClusterViewportFrameStatsViewExtension::PostProcessPa
 		Parameters->FrameCount = FrameCount;
 		Parameters->Timecode = EncodedTimecode;
 
-		AddDrawScreenPass(GraphBuilder, RDG_EVENT_NAME("DisplayClusterViewportFrameStatsPass"), ViewInfo, Viewport, Viewport, VertexShader, FrameStatsPixelShader, Parameters);
+		AddDrawScreenPass(GraphBuilder, RDG_EVENT_NAME("DisplayClusterViewportFrameStatsPass"), View, Viewport, Viewport, VertexShader, FrameStatsPixelShader, Parameters);
 	}
 
 	FScreenPassRenderTarget Output = Inputs.OverrideOutput;
@@ -149,7 +146,7 @@ FScreenPassTexture FDisplayClusterViewportFrameStatsViewExtension::PostProcessPa
 		Parameters->FrameStatsTextureSampler = TStaticSamplerState<SF_Point, AM_Wrap, AM_Wrap>::GetRHI();
 		Parameters->RenderTargets[0] = Output.GetRenderTargetBinding();
 
-		AddDrawScreenPass(GraphBuilder, RDG_EVENT_NAME("DisplayClusterViewportFrameStatsOutputPass"), ViewInfo, OutputViewport, InputViewport, VertexShader, OutputPixelShader, Parameters);
+		AddDrawScreenPass(GraphBuilder, RDG_EVENT_NAME("DisplayClusterViewportFrameStatsOutputPass"), View, OutputViewport, InputViewport, VertexShader, OutputPixelShader, Parameters);
 	}
 
 	return Output;
diff --git a/Engine/Plugins/Runtime/nDisplay/Source/DisplayCluster/Private/Render/Viewport/DisplayClusterViewportProxy.cpp b/Engine/Plugins/Runtime/nDisplay/Source/DisplayCluster/Private/Render/Viewport/DisplayClusterViewportProxy.cpp
--- a/Engine/Plugins/Runtime/nDisplay/Source/DisplayCluster/Private/Render/Viewport/DisplayClusterViewportProxy.cpp
+++ b/Engine/Plugins/Runtime/nDisplay/Source/DisplayCluster/Private/Render/Viewport/DisplayClusterViewportProxy.cpp
@@ -42,8 +42,7 @@
 #include "PostProcess/DrawRectangle.h"
 
 #include "ScreenPass.h"
-#include "PostProcess/PostProcessing.h"
-// for FPostProcessMaterialInputs
+#include "SceneTextures.h"
 #include "PostProcess/PostProcessMaterialInputs.h"
 
 static TAutoConsoleVariable<int32> CVarDisplayClusterRenderOverscanResolve(
@@ -1087,10 +1086,8 @@ void FDisplayClusterViewportProxy::OnPostRenderViewFamily_RenderThread(FRDGBuild
 	// Get the GPUIndex used to render this viewport
 	if (Contexts.IsValidIndex(InContextNum))
 	{
-		checkSlow(InSceneView.bIsViewInfo);
-		const FViewInfo& ViewInfo = static_cast<const FViewInfo&>(InSceneView);
 
-		const uint32 GPUIndex = ViewInfo.GPUMask.GetFirstIndex();
+		const uint32 GPUIndex = InSceneView.GPUMask.GetFirstIndex();
 		Contexts[InContextNum].RenderThreadData.GPUIndex = (GPUIndex < GNumExplicitGPUsForRendering) ? GPUIndex : -1;
 	}
 #endif
@@ -1129,9 +1126,6 @@ void FDisplayClusterViewportProxy::OnPostRenderViewFamily_RenderThread(FRDGBuild
 					// Copy Alpha channels back from'InputShaderResource' to 'InternalRenderTargetResource'
 					CopyResource_RenderThread(GraphBuilder, EDisplayClusterTextureCopyMode::Alpha, InContextNum, EDisplayClusterViewportResourceType::InputShaderResource, EDisplayClusterViewportResourceType::InternalRenderTargetResource);
 
-					checkSlow(InSceneView.bIsViewInfo);
-					const FViewInfo& ViewInfo = static_cast<const FViewInfo&>(InSceneView);
-
 					// 1. Copy RGB channels from 'InternalRenderTargetResource' to 'InputShaderResource'
 					CopyResource_RenderThread(GraphBuilder, EDisplayClusterTextureCopyMode::RGB, InContextNum, EDisplayClusterViewportResourceType::InternalRenderTargetResource, EDisplayClusterViewportResourceType::InputShaderResource);
 
@@ -1145,7 +1139,7 @@ void FDisplayClusterViewportProxy::OnPostRenderViewFamily_RenderThread(FRDGBuild
 					PassInputs.Quality = FXAAQuality;
 
 					// 2.1. Do FXAA
-					FScreenPassTexture OutputColorTexture = AddFXAAPass(GraphBuilder, ViewInfo, PassInputs);
+					FScreenPassTexture OutputColorTexture = AddFXAAPass(GraphBuilder, InSceneView, PassInputs);
 
 					// 2.2. Copy FXAA result from 'OutputTexture' to the 'InternalRenderTargetResource'
 					if (OutputColorTexture.Texture)
diff --git a/Engine/Plugins/Runtime/nDisplay/Source/DisplayCluster/Private/Render/Viewport/DisplayClusterViewport_OpenColorIO.cpp b/Engine/Plugins/Runtime/nDisplay/Source/DisplayCluster/Private/Render/Viewport/DisplayClusterViewport_OpenColorIO.cpp
--- a/Engine/Plugins/Runtime/nDisplay/Source/DisplayCluster/Private/Render/Viewport/DisplayClusterViewport_OpenColorIO.cpp
+++ b/Engine/Plugins/Runtime/nDisplay/Source/DisplayCluster/Private/Render/Viewport/DisplayClusterViewport_OpenColorIO.cpp
@@ -11,13 +11,10 @@
 #include "OpenColorIODisplayExtension.h"
 #include "OpenColorIOShader.h"
 
-#include "Shader.h"
-#include "SceneRendering.h"
 #include "SceneView.h"
 #include "ScreenPass.h"
 #include "CommonRenderResources.h"
-// for FPostProcessMaterialInputs
-#include "PostProcess/PostProcessMaterial.h"
+#include "PostProcess/PostProcessMaterialInputs.h"
 
 //////////////////////////////////////////////////////////////////////////
 // FDisplayClusterViewport_OpenColorIO
@@ -101,19 +98,17 @@ FScreenPassTexture FDisplayClusterViewport_OpenColorIO::PostProcessPassAfterTone
 {
 	const FScreenPassTexture& SceneColor = InOutInputs.GetInput(EPostProcessMaterialInput::SceneColor);
 	check(SceneColor.IsValid());
-	checkSlow(View.bIsViewInfo);
-	const FViewInfo& ViewInfo = static_cast<const FViewInfo&>(View);
 	FScreenPassRenderTarget Output = InOutInputs.OverrideOutput;
 
 	// If the override output is provided, it means that this is the last pass in post processing.
 	if (!Output.IsValid())
 	{
-		Output = FScreenPassRenderTarget::CreateFromInput(GraphBuilder, SceneColor, ViewInfo.GetOverwriteLoadAction(), TEXT("OCIORenderTarget"));
+		Output = FScreenPassRenderTarget::CreateFromInput(GraphBuilder, SceneColor, View.GetOverwriteLoadAction(), TEXT("OCIORenderTarget"));
 	}
 
 	FOpenColorIORendering::AddPass_RenderThread(
 		GraphBuilder,
-		ViewInfo,
+		View,
 		SceneColor,
 		Output,
 		CachedResourcesRenderThread
diff --git a/Engine/Source/Runtime/Renderer/Internal/CustomDepthRendering.h b/Engine/Source/Runtime/Renderer/Internal/CustomDepthRendering.h
new file mode 100644
--- /dev/null
+++ b/Engine/Source/Runtime/Renderer/Internal/CustomDepthRendering.h
@@ -0,0 +1,42 @@
+// Copyright Epic Games, Inc. All Rights Reserved.
+
+#pragma once
+
+#include "RenderGraphFwd.h"
+#include "RHIDefinitions.h"
+#include "RHIFwd.h"
+#include "RHIShaderPlatform.h"
+
+enum class ECustomDepthPassLocation : uint32
+{
+	// Renders custom depth before the base pass. Can be more efficient with AsyncCompute and enables use with DBuffer decals.
+	BeforeBasePass,
+
+	// Renders after the base pass.
+	AfterBasePass
+};
+
+// Returns the location in the frame where custom depth is rendered.
+extern ECustomDepthPassLocation GetCustomDepthPassLocation(EShaderPlatform Platform);
+
+struct FCustomDepthTextures
+{
+	static FCustomDepthTextures Create(FRDGBuilder& GraphBuilder, FIntPoint CustomDepthExtent, EShaderPlatform ShaderPlatform);
+
+	bool IsValid() const
+	{
+		return Depth != nullptr;
+	}
+
+	FRDGTextureRef Depth{};
+	FRDGTextureSRVRef Stencil{};
+
+	// Denotes that the depth and stencil buffers had to be split to separate, non-depth textures (and thus Depth cannot be bound
+	// as a depth/stencil buffer). This can happen when Nanite renders custom depth on platforms with HW that cannot write stencil
+	// values per-pixel from a shader.
+	bool bSeparateStencilBuffer = false;
+
+	// Actions to use when initially rendering to custom depth / stencil.
+	ERenderTargetLoadAction DepthAction = ERenderTargetLoadAction::EClear;
+	ERenderTargetLoadAction StencilAction = ERenderTargetLoadAction::EClear;
+};
diff --git a/Engine/Source/Runtime/Renderer/Internal/OverridePassSequence.h b/Engine/Source/Runtime/Renderer/Internal/OverridePassSequence.h
new file mode 100644
--- /dev/null
+++ b/Engine/Source/Runtime/Renderer/Internal/OverridePassSequence.h
@@ -0,0 +1,198 @@
+// Copyright Epic Games, Inc. All Rights Reserved.
+
+#pragma once
+
+#include "ScreenPass.h"
+#include "SceneViewExtension.h"
+
+
+/** This class provides a solution for overriding the final render target / view rect in a chain of render graph passes.
+ *  RDG pass setup is immediate, which makes it difficult to determine the final pass in a chain in order to swap a final
+ *  texture target (e.g. the view family target). This class helps by providing some utilities to mark passes as enabled
+ *  prior to setup, then seamlessly assign the override texture to the correct last pass. It assumes that all passes adhere
+ *  to the pattern of using an input struct that derives from FScreenPassRenderTarget (though this is not strictly required).
+ *
+ *  The user provides a custom enum class that is the set of ordered passes. The user can enable them in any order, but
+ *  all passes must be assigned (i.e. have SetEnabled called on them). After enabling / disabling all passes. Call Finalize
+ *  to compute internal state. You cannot enable / disable passes after finalizing. Each pass must be accepted in order,
+ *  and all passes marked as enabled must be accepted. The helper, AcceptOverrideIfLastPass, will check and override if
+ *  the pass is the last one, and then accept the pass.
+ */
+template <typename EPass>
+class TOverridePassSequence final
+{
+public:
+	TOverridePassSequence(const FScreenPassRenderTarget& InOverrideOutput)
+		: OverrideOutput(InOverrideOutput)
+	{}
+
+	~TOverridePassSequence()
+	{
+#if RDG_ENABLE_DEBUG
+		if (bFinalized)
+		{
+			for (int32 PassIndex = 0; PassIndex < PassCountMax; ++PassIndex)
+			{
+				if (Passes[PassIndex].bEnabled)
+				{
+					checkf(Passes[PassIndex].bAccepted, TEXT("Pass was enabled but not accepted: %s."), Passes[PassIndex].Name);
+				}
+			}
+		}
+#endif
+	}
+
+	void SetName(EPass Pass, const TCHAR* Name)
+	{
+#if RDG_ENABLE_DEBUG
+		Passes[(int32)Pass].Name = Name;
+#endif
+	}
+
+	void SetNames(const TCHAR* const* Names, uint32 NameCount)
+	{
+#if RDG_ENABLE_DEBUG
+		check(NameCount == PassCountMax);
+		for (int32 PassIndex = 0; PassIndex < PassCountMax; ++PassIndex)
+		{
+			Passes[PassIndex].Name = Names[PassIndex];
+		}
+#endif
+	}
+
+	void SetEnabled(EPass Pass, bool bEnabled)
+	{
+		const int32 PassIndex = (int32)Pass;
+
+#if RDG_ENABLE_DEBUG
+		check(!bFinalized);
+		checkf(!Passes[PassIndex].bAssigned, TEXT("Pass cannot be assigned multiple times: %s."), Passes[PassIndex].Name);
+		Passes[PassIndex].bAssigned = true;
+#endif
+
+		Passes[PassIndex].bEnabled = bEnabled;
+	}
+
+	bool IsEnabled(EPass Pass) const
+	{
+		const int32 PassIndex = (int32)Pass;
+#if RDG_ENABLE_DEBUG
+		check(Passes[PassIndex].bAssigned);
+#endif
+		return Passes[PassIndex].bEnabled;
+	}
+
+	bool IsLastPass(EPass Pass) const
+	{
+#if RDG_ENABLE_DEBUG
+		check(bFinalized);
+#endif
+		return Pass == LastPass;
+	}
+
+	void AcceptPass(EPass Pass)
+	{
+#if RDG_ENABLE_DEBUG
+		const int32 PassIndex = (int32)Pass;
+
+		check(bFinalized);
+		checkf(NextPass == Pass, TEXT("Pass was accepted out of order: %s. Expected %s."), Passes[PassIndex].Name, Passes[(int32)NextPass].Name);
+		checkf(Passes[PassIndex].bEnabled, TEXT("Only accepted passes can be enabled: %s."), Passes[PassIndex].Name);
+
+		Passes[PassIndex].bAccepted = true;
+
+		// Walk the remaining passes until we hit one that's enabled. This will be the next pass to add.
+		for (int32 NextPassIndex = int32(NextPass) + 1; NextPassIndex < PassCountMax; ++NextPassIndex)
+		{
+			if (Passes[NextPassIndex].bEnabled)
+			{
+				NextPass = EPass(NextPassIndex);
+				break;
+			}
+		}
+#endif
+	}
+
+	bool AcceptOverrideIfLastPass(EPass Pass, FScreenPassRenderTarget& OutTargetToOverride, const TOptional<int32>& AfterPassCallbackIndex = TOptional<int32>())
+	{
+		bool bLastAfterPass = AfterPass[(int32)Pass].Num() == 0;
+
+		if (AfterPassCallbackIndex)
+		{
+			bLastAfterPass = AfterPassCallbackIndex.GetValue() == AfterPass[(int32)Pass].Num() - 1;
+		}
+		else
+		{
+			// Display debug information for a Pass unless it is an after pass.
+			AcceptPass(Pass);
+		}
+
+		// We need to override output only if this is the last pass and the last after pass.
+		if (IsLastPass(Pass) && bLastAfterPass)
+		{
+			OutTargetToOverride = OverrideOutput;
+			return true;
+		}
+
+		return false;
+	}
+
+	void Finalize()
+	{
+#if RDG_ENABLE_DEBUG
+		check(!bFinalized);
+		bFinalized = true;
+
+		for (int32 PassIndex = 0; PassIndex < PassCountMax; ++PassIndex)
+		{
+			checkf(Passes[PassIndex].bAssigned, TEXT("Pass was not assigned to enabled or disabled: %s."), Passes[PassIndex].Name);
+		}
+#endif
+
+		bool bFirstPass = true;
+
+		for (int32 PassIndex = 0; PassIndex < PassCountMax; ++PassIndex)
+		{
+			if (Passes[PassIndex].bEnabled)
+			{
+				if (bFirstPass)
+				{
+#if RDG_ENABLE_DEBUG
+					NextPass = (EPass)PassIndex;
+#endif
+					bFirstPass = false;
+				}
+				LastPass = (EPass)PassIndex;
+			}
+		}
+	}
+
+	FAfterPassCallbackDelegateArray& GetAfterPassCallbacks(EPass Pass)
+	{ 
+		const int32 PassIndex = (int32)Pass;
+		return AfterPass[PassIndex]; 
+	}
+
+private:
+	static const int32 PassCountMax = (int32)EPass::MAX;
+
+	struct FPassInfo
+	{
+#if RDG_ENABLE_DEBUG
+		const TCHAR* Name = nullptr;
+		bool bAssigned = false;
+		bool bAccepted = false;
+#endif
+		bool bEnabled = false;
+	};
+
+	FScreenPassRenderTarget OverrideOutput;
+	TStaticArray<FPassInfo, PassCountMax> Passes;
+	TStaticArray<FAfterPassCallbackDelegateArray, PassCountMax> AfterPass;
+	EPass LastPass = EPass::MAX;
+
+#if RDG_ENABLE_DEBUG
+	EPass NextPass = EPass(0);
+	bool bFinalized = false;
+#endif
+};
diff --git a/Engine/Source/Runtime/Renderer/Internal/PostProcess/PostProcessAA.h b/Engine/Source/Runtime/Renderer/Internal/PostProcess/PostProcessAA.h
new file mode 100644
--- /dev/null
+++ b/Engine/Source/Runtime/Renderer/Internal/PostProcess/PostProcessAA.h
@@ -0,0 +1,35 @@
+// Copyright Epic Games, Inc. All Rights Reserved.
+
+#pragma once
+
+#include "ScreenPass.h"
+#include "OverridePassSequence.h"
+
+enum class EFXAAQuality : uint32
+{
+	// Lowest Quality / Fastest
+	Q0,
+	Q1,
+	Q2,
+	Q3,
+	Q4,
+	Q5,
+	// Highest Quality / Slowest
+	MAX
+};
+
+EFXAAQuality GetFXAAQuality();
+
+struct FFXAAInputs
+{
+	// [Optional] Render to the specified output. If invalid, a new texture is created and returned.
+	FScreenPassRenderTarget OverrideOutput;
+
+	// [Required] HDR scene color to filter.
+	FScreenPassTexture SceneColor;
+
+	// FXAA filter quality.
+	EFXAAQuality Quality = EFXAAQuality::MAX;
+};
+
+FScreenPassTexture RENDERER_API AddFXAAPass(FRDGBuilder& GraphBuilder, const FSceneView& View, const FFXAAInputs& Inputs);
\ No newline at end of file
diff --git a/Engine/Source/Runtime/Renderer/Internal/SceneTextures.h b/Engine/Source/Runtime/Renderer/Internal/SceneTextures.h
new file mode 100644
--- /dev/null
+++ b/Engine/Source/Runtime/Renderer/Internal/SceneTextures.h
@@ -0,0 +1,107 @@
+// Copyright Epic Games, Inc. All Rights Reserved.
+
+#pragma once
+
+#include "SceneView.h"
+#include "RenderGraphUtils.h"
+#include "CustomDepthRendering.h"
+#include "SceneRenderTargetParameters.h"
+#include "GBufferInfo.h"
+
+struct FSceneTextures;
+class FViewInfo;
+class FViewFamilyInfo;
+
+/** Initializes a scene textures config instance from the view family. */
+extern RENDERER_API void InitializeSceneTexturesConfig(FSceneTexturesConfig& Config, const FSceneViewFamily& ViewFamily);
+
+/** RDG struct containing the minimal set of scene textures common across all rendering configurations. */
+struct FMinimalSceneTextures
+{
+	// Initializes the minimal scene textures structure in the FViewFamilyInfo
+	static RENDERER_API void InitializeViewFamily(FRDGBuilder& GraphBuilder, FViewFamilyInfo& ViewFamily);
+
+	// Immutable copy of the config used to create scene textures.
+	FSceneTexturesConfig Config;
+
+	// Uniform buffers for deferred or mobile.
+	TRDGUniformBufferRef<FSceneTextureUniformParameters> UniformBuffer{};
+	TRDGUniformBufferRef<FMobileSceneTextureUniformParameters> MobileUniformBuffer{};
+
+	// Setup modes used when creating uniform buffers. These are updated on demand.
+	ESceneTextureSetupMode SetupMode = ESceneTextureSetupMode::None;
+	EMobileSceneTextureSetupMode MobileSetupMode = EMobileSceneTextureSetupMode::None;
+
+	// Texture containing scene color information with lighting but without post processing. Will be two textures if MSAA.
+	FRDGTextureMSAA Color{};
+
+	// Texture containing scene depth. Will be two textures if MSAA.
+	FRDGTextureMSAA Depth{};
+
+	// Texture containing a stencil view of the resolved (if MSAA) scene depth. 
+	FRDGTextureSRVRef Stencil{};
+
+	// Textures containing primary depth buffer copied before other meshes are rendered in the secondary depth pass.
+	FRDGTextureMSAA PartialDepth{};
+
+	// Textures containing depth / stencil information from the custom depth pass.
+	FCustomDepthTextures CustomDepth{};
+
+	RENDERER_API FSceneTextureShaderParameters GetSceneTextureShaderParameters(ERHIFeatureLevel::Type FeatureLevel) const;
+};
+
+/** RDG struct containing the complete set of scene textures for the deferred or mobile renderers. */
+struct FSceneTextures : public FMinimalSceneTextures
+{
+	// Initializes the scene textures structure in the FViewFamilyInfo
+	static RENDERER_API void InitializeViewFamily(FRDGBuilder& GraphBuilder, FViewFamilyInfo& ViewFamily);
+	static RENDERER_API EPixelFormat GetGBufferFFormatAndCreateFlags(ETextureCreateFlags& OutCreateFlags);
+
+	// Configures an array of render targets for the GBuffer pass.
+	RENDERER_API uint32 GetGBufferRenderTargets(
+		TArrayView<FTextureRenderTargetBinding> RenderTargets,
+		EGBufferLayout Layout = GBL_Default) const;
+	RENDERER_API uint32 GetGBufferRenderTargets(
+		ERenderTargetLoadAction LoadAction,
+		FRenderTargetBindingSlots& RenderTargets,
+		EGBufferLayout Layout = GBL_Default) const;
+	
+	// (Deferred) Texture containing conservative downsampled depth for occlusion.
+	FRDGTextureRef SmallDepth{};
+
+	// (Deferred) Textures containing geometry information for deferred shading.
+	FRDGTextureRef GBufferA{};
+	FRDGTextureRef GBufferB{};
+	FRDGTextureRef GBufferC{};
+	FRDGTextureRef GBufferD{};
+	FRDGTextureRef GBufferE{};
+	FRDGTextureRef GBufferF{};
+
+	// Additional Buffer texture used by mobile
+	FRDGTextureMSAA DepthAux{};
+
+	// Texture containing dynamic motion vectors. Can be bound by the base pass or its own velocity pass.
+	FRDGTextureRef Velocity{};
+	
+	// (Mobile Local Light Prepass) Textures containing LocalLight Direction and Color
+	FRDGTextureRef MobileLocalLightTextureA {};
+	FRDGTextureRef MobileLocalLightTextureB {};
+
+	// Texture containing the screen space ambient occlusion result.
+	FRDGTextureRef ScreenSpaceAO{};
+
+	// Texture used by the quad overdraw debug view mode when enabled.
+	FRDGTextureRef QuadOverdraw{};
+
+	// (Mobile) Texture used by mobile PPR in the next frame.
+	FRDGTextureRef PixelProjectedReflection{};
+
+	// Textures used to composite editor primitives. Also used by the base pass when in wireframe mode.
+#if WITH_EDITOR
+	FRDGTextureRef EditorPrimitiveColor{};
+	FRDGTextureRef EditorPrimitiveDepth{};
+#endif
+};
+
+/** Extracts scene textures into the global extraction instance. */
+void QueueSceneTextureExtractions(FRDGBuilder& GraphBuilder, const FSceneTextures& SceneTextures);
diff --git a/Engine/Source/Runtime/Renderer/Private/CustomDepthRendering.h b/Engine/Source/Runtime/Renderer/Private/CustomDepthRendering.h
deleted file mode 100644
diff --git a/Engine/Source/Runtime/Renderer/Private/OverridePassSequence.h b/Engine/Source/Runtime/Renderer/Private/OverridePassSequence.h
deleted file mode 100644
diff --git a/Engine/Source/Runtime/Renderer/Private/PostProcess/PostProcessAA.cpp b/Engine/Source/Runtime/Renderer/Private/PostProcess/PostProcessAA.cpp
--- a/Engine/Source/Runtime/Renderer/Private/PostProcess/PostProcessAA.cpp
+++ b/Engine/Source/Runtime/Renderer/Private/PostProcess/PostProcessAA.cpp
@@ -51,11 +51,15 @@ EFXAAQuality GetFXAAQuality()
 	return EFXAAQuality(FMath::Clamp(CVarFXAAQuality.GetValueOnRenderThread(), 0, 5));
 }
 
-FScreenPassTexture AddFXAAPass(FRDGBuilder& GraphBuilder, const FViewInfo& View, const FFXAAInputs& Inputs)
+FScreenPassTexture AddFXAAPass(FRDGBuilder& GraphBuilder, const FSceneView& InSceneView, const FFXAAInputs& Inputs)
 {
+	
 	check(Inputs.SceneColor.IsValid());
 	check(Inputs.Quality != EFXAAQuality::MAX);
 
+	checkSlow(InSceneView.bIsViewInfo);
+	const FViewInfo& View = static_cast<const FViewInfo&>(InSceneView);
+
 	FScreenPassRenderTarget Output = Inputs.OverrideOutput;
 	
 	if (!Output.IsValid())
diff --git a/Engine/Source/Runtime/Renderer/Private/PostProcess/PostProcessAA.h b/Engine/Source/Runtime/Renderer/Private/PostProcess/PostProcessAA.h
deleted file mode 100644
diff --git a/Engine/Source/Runtime/Renderer/Private/SceneTextures.h b/Engine/Source/Runtime/Renderer/Private/SceneTextures.h
deleted file mode 100644
