Deferred OIT sorted index buffer deletion to next frame to use after free issue.

This is a speculative fix based on crash report.

diff --git a/Engine/Source/Runtime/Renderer/Private/OIT/OIT.cpp b/Engine/Source/Runtime/Renderer/Private/OIT/OIT.cpp
--- a/Engine/Source/Runtime/Renderer/Private/OIT/OIT.cpp
+++ b/Engine/Source/Runtime/Renderer/Private/OIT/OIT.cpp
@@ -296,14 +296,7 @@ public:
 	FUnorderedAccessViewRHIRef SortedIndexUAV = nullptr;
 };
 
-static void RemoveAllocation(FSortedIndexBuffer* InBuffer)
-{
-	InBuffer->ReleaseResource();
-	delete InBuffer;
-	InBuffer = nullptr;
-}
-
-static void TrimSortedIndexBuffers(TArray<FSortedIndexBuffer*>& FreeBuffers, uint32 FrameId)
+static void TrimSortedIndexBuffers(TArray<FSortedIndexBuffer*>& FreeBuffers, TQueue<FSortedIndexBuffer*>& DeleteQueue, uint32 FrameId)
 {
 	uint32 FreeCount = FreeBuffers.Num();
 	for (uint32 FreeIt = 0; FreeIt < FreeCount;)
@@ -315,8 +308,7 @@ static void TrimSortedIndexBuffers(TArray<FSortedIndexBuffer*>& FreeBuffers, uin
 			const int32 ElapsedFrame = FMath::Abs(int32(FrameId) - int32(LastFrameId));
 			if (ElapsedFrame > CVarOIT_SortedTriangles_PoolReleaseThreshold.GetValueOnRenderThread())
 			{
-				FSortedIndexBuffer* Buffer = FreeBuffers[FreeIt];
-				RemoveAllocation(Buffer);
+				DeleteQueue.Enqueue(FreeBuffers[FreeIt]);
 				FreeBuffers[FreeIt] = FreeBuffers[FreeCount - 1];
 				--FreeCount;
 				FreeBuffers.SetNum(FreeCount);
@@ -408,7 +400,7 @@ void FOITSceneData::Deallocate(FIndexBuffer* InIndexBuffer)
 			FSortedTriangleData& In = Allocations[Slot];
 			In.SortedIndexUAV = nullptr;
 			In.SourceIndexSRV = nullptr;
-			RemoveAllocation((FSortedIndexBuffer*)In.SortedIndexBuffer);
+			PendingDeletes.Enqueue((FSortedIndexBuffer*)In.SortedIndexBuffer);
 			In = FSortedTriangleData();
 		}
 		FreeSlots.Enqueue(Slot);
@@ -897,7 +889,7 @@ namespace OIT
 		OITSceneData.FrameIndex = View.Family->FrameNumber;
 		if (CVarOIT_SortedTriangles_Pool.GetValueOnRenderThread() > 0 && CVarOIT_SortedTriangles_PoolReleaseThreshold.GetValueOnRenderThread() > 0)
 		{
-			TrimSortedIndexBuffers(OITSceneData.FreeBuffers, OITSceneData.FrameIndex);
+			TrimSortedIndexBuffers(OITSceneData.FreeBuffers, OITSceneData.PendingDeletes, OITSceneData.FrameIndex);
 		}
 	}
 
@@ -985,4 +977,16 @@ namespace OIT
 	{
 		AddInternalOITComposePass(GraphBuilder, View, OITData, SceneColorTexture);
 	}
+
+	void OnRenderBegin(FOITSceneData& OITSceneData)
+	{
+		// Delete all enqueue buffer for deletion
+		FSortedIndexBuffer* Buffer = nullptr;
+		while (OITSceneData.PendingDeletes.Dequeue(Buffer))
+		{
+			Buffer->ReleaseResource();
+			delete Buffer;
+			Buffer = nullptr;
+		}
+	}
 }
\ No newline at end of file
diff --git a/Engine/Source/Runtime/Renderer/Private/OIT/OIT.h b/Engine/Source/Runtime/Renderer/Private/OIT/OIT.h
--- a/Engine/Source/Runtime/Renderer/Private/OIT/OIT.h
+++ b/Engine/Source/Runtime/Renderer/Private/OIT/OIT.h
@@ -78,6 +78,7 @@ struct FOITSceneData
 
 	TArray<FSortedTriangleData> Allocations;
 	TArray<FSortedIndexBuffer*> FreeBuffers;
+	TQueue<FSortedIndexBuffer*> PendingDeletes;
 	TQueue<uint32> FreeSlots;
 	uint32 FrameIndex = 0;
 };
@@ -102,4 +103,7 @@ namespace OIT
 
 	/* Compose all OIT samples into the target color buffer */
 	void AddOITComposePass(FRDGBuilder& GraphBuilder, const FViewInfo& View, FOITData& OITData, FRDGTextureRef SceneColorTexture);
+
+	/* Call on SceneRenderer OnRenderBegin */
+	void OnRenderBegin(FOITSceneData& OITSceneData);
 }
diff --git a/Engine/Source/Runtime/Renderer/Private/SceneRendering.cpp b/Engine/Source/Runtime/Renderer/Private/SceneRendering.cpp
--- a/Engine/Source/Runtime/Renderer/Private/SceneRendering.cpp
+++ b/Engine/Source/Runtime/Renderer/Private/SceneRendering.cpp
@@ -94,6 +94,7 @@
 #include "GPUDebugCrashUtils.h"
 #include "MeshDrawCommandStats.h"
 #include "LocalFogVolumeRendering.h"
+#include "OIT/OIT.h"
 
 /*-----------------------------------------------------------------------------
 	Globals
@@ -3408,6 +3409,9 @@ IVisibilityTaskData* FSceneRenderer::OnRenderBegin(FRDGBuilder& GraphBuilder)
 	// This is called prior to scene update to avoid a race condition with the MDC caching task.
 	FVirtualTextureSystem::Get().CallPendingCallbacks();
 
+	// This is called prior to scene update
+	OIT::OnRenderBegin(Scene->OITSceneData);
+
 	EUpdateAllPrimitiveSceneInfosAsyncOps AsyncOps = EUpdateAllPrimitiveSceneInfosAsyncOps::None;
 
 	if (GAsyncCreateLightPrimitiveInteractions > 0)
