Add neural post processing (experimental).
* Add neural profile to config which neural network to use. Texture index is used by default where [-1,3,H,W] is supported. Use Buffer index to support [-1,C,H,W]. Need to config the batch size override or the tile size correspondingly.
* Config in the post process material which neural profile to use and enable `used with neural network`.
* Use neural input/output to control the input and output of the neural network. For neural input, [B,C, U, V] is the index, input0 takes 3 effective channels, mask determines if we should use it as input to the network. Only one node can be used per material; Multiple neural outputs can be used. For texture indexing, automatic bilinear downscaling and upscaling is applied. For Buffer indexing, the user has full control to access the value ( user needs to perform their own filtering).

diff --git a/Engine/Config/BaseEditor.ini b/Engine/Config/BaseEditor.ini
--- a/Engine/Config/BaseEditor.ini
+++ b/Engine/Config/BaseEditor.ini
@@ -64,6 +64,7 @@ bCleanupExternalObjectsEmptyFolders=false
 +RenderableThumbnailTypes=(ClassNeedingThumbnailName="/Script/Engine.SlateBrushAsset",RendererClassName="/Script/UnrealEd.SlateBrushThumbnailRenderer")
 +RenderableThumbnailTypes=(ClassNeedingThumbnailName="/Script/Engine.SubsurfaceProfile",RendererClassName="/Script/UnrealEd.SubsurfaceProfileRenderer")
 +RenderableThumbnailTypes=(ClassNeedingThumbnailName="/Script/Engine.SpecularProfile",RendererClassName="/Script/UnrealEd.SpecularProfileRenderer")
++RenderableThumbnailTypes=(ClassNeedingThumbnailName="/Script/Engine.NeuralProfile",RendererClassName="/Script/UnrealEd.NeuralProfileRenderer")
 +RenderableThumbnailTypes=(ClassNeedingThumbnailName="/Script/Engine.PhysicsAsset",RendererClassName="/Script/UnrealEd.PhysicsAssetThumbnailRenderer")
 +RenderableThumbnailTypes=(ClassNeedingThumbnailName="/Script/Engine.CurveLinearColor",RendererClassName="/Script/UnrealEd.CurveLinearColorThumbnailRenderer")
 +RenderableThumbnailTypes=(ClassNeedingThumbnailName="/Script/Engine.CurveFloat",RendererClassName="/Script/UnrealEd.CurveFloatThumbnailRenderer")
diff --git a/Engine/Plugins/Editor/EngineAssetDefinitions/Source/Private/AssetDefinition_NeuralProfile.h b/Engine/Plugins/Editor/EngineAssetDefinitions/Source/Private/AssetDefinition_NeuralProfile.h
new file mode 100644
--- /dev/null
+++ b/Engine/Plugins/Editor/EngineAssetDefinitions/Source/Private/AssetDefinition_NeuralProfile.h
@@ -0,0 +1,26 @@
+// Copyright Epic Games, Inc. All Rights Reserved.
+
+#pragma once
+
+#include "Engine/NeuralProfile.h"
+#include "AssetDefinitionDefault.h"
+
+#include "AssetDefinition_NeuralProfile.generated.h"
+
+UCLASS()
+class UAssetDefinition_NeuralProfile : public UAssetDefinitionDefault
+{
+	GENERATED_BODY()
+
+public:
+	// UAssetDefinition Begin
+	virtual FText GetAssetDisplayName() const override { return NSLOCTEXT("AssetTypeActions", "AssetTypeActions_NeuralProfile", "Neural Profile"); }
+	virtual FLinearColor GetAssetColor() const override { return FLinearColor(FColor(128, 128, 255)); }
+	virtual TSoftClassPtr<UObject> GetAssetClass() const override { return UNeuralProfile::StaticClass(); }
+	virtual TConstArrayView<FAssetCategoryPath> GetAssetCategories() const override
+    {
+    	static const auto Categories = { EAssetCategoryPaths::Material };
+    	return Categories;
+    }
+	// UAssetDefinition End
+};
diff --git a/Engine/Plugins/Experimental/NeuralRendering/NeuralRendering.uplugin b/Engine/Plugins/Experimental/NeuralRendering/NeuralRendering.uplugin
new file mode 100644
--- /dev/null
+++ b/Engine/Plugins/Experimental/NeuralRendering/NeuralRendering.uplugin
@@ -0,0 +1,41 @@
+{
+	"FileVersion": 3,
+	"Version": 0.1,
+	"VersionName": "0.1",
+	"FriendlyName": "Neural Rendering",
+	"Description": "Enable neural rendering features including: neural post processing",
+	"Category": "Experimental",
+	"CreatedBy": "Epic Games, Inc.",
+	"CreatedByURL": "https://epicgames.com",
+	"DocsURL": "",
+	"MarketplaceURL": "",
+	"SupportURL": "",
+	"EnabledByDefault": false,
+	"CanContainContent": false,
+	"IsBetaVersion": false,
+	"IsExperimentalVersion": true,
+	"Installed": false,
+	"SupportedTargetPlatforms": [
+		"Win64",
+		"Linux",
+		"Mac"
+	],
+	"Modules": [
+		{
+			"Name": "NeuralPostProcessing",
+			"Type": "RuntimeAndProgram",
+			"LoadingPhase": "PostConfigInit",
+			"PlatformAllowList": [ "Win64", "Linux", "Mac" ]
+		}
+	],
+	"Plugins": [
+		{
+			"Name": "NNERuntimeRDG",
+			"Enabled": true
+		},
+		{
+			"Name": "NNERuntimeORTCpu",
+			"Enabled": true
+		}
+	]
+}
\ No newline at end of file
diff --git a/Engine/Plugins/Experimental/NeuralRendering/Source/NeuralPostProcessing/NeuralPostProcessing.Build.cs b/Engine/Plugins/Experimental/NeuralRendering/Source/NeuralPostProcessing/NeuralPostProcessing.Build.cs
new file mode 100644
--- /dev/null
+++ b/Engine/Plugins/Experimental/NeuralRendering/Source/NeuralPostProcessing/NeuralPostProcessing.Build.cs
@@ -0,0 +1,59 @@
+// Copyright Epic Games, Inc. All Rights Reserved.
+
+using UnrealBuildTool;
+
+public class NeuralPostProcessing : ModuleRules
+{
+	public NeuralPostProcessing(ReadOnlyTargetRules Target) : base(Target)
+	{
+		PCHUsage = ModuleRules.PCHUsageMode.UseExplicitOrSharedPCHs;
+		
+		PublicIncludePaths.AddRange(
+			new string[] {
+				// ... add public include paths required here ...
+			}
+			);
+				
+		
+		PrivateIncludePaths.AddRange(
+			new string[] {
+				// ... add other private include paths required here ...
+				System.IO.Path.Combine(EngineDirectory, "Source/Runtime/Renderer/Private")
+			}
+			);
+			
+		
+		PublicDependencyModuleNames.AddRange(
+			new string[]
+			{
+				"Core",
+				// ... add other public dependencies that you statically link with here ...
+			}
+			);
+			
+		
+		PrivateDependencyModuleNames.AddRange(
+			new string[]
+			{
+				"CoreUObject",
+				"Engine",
+				"Slate",
+				"SlateCore",
+				// ... add private dependencies that you statically link with here ...	
+				"Projects",
+				"RenderCore",
+				"RHI",
+				"Renderer",
+				"NNE"
+			}
+			);
+			
+
+		DynamicallyLoadedModuleNames.AddRange(
+			new string[]
+			{
+				// ... add any modules that your module loads dynamically here ...
+			}
+			);
+	}
+}
diff --git a/Engine/Plugins/Experimental/NeuralRendering/Source/NeuralPostProcessing/Private/NeuralPostProcessModelInstance.cpp b/Engine/Plugins/Experimental/NeuralRendering/Source/NeuralPostProcessing/Private/NeuralPostProcessModelInstance.cpp
new file mode 100644
--- /dev/null
+++ b/Engine/Plugins/Experimental/NeuralRendering/Source/NeuralPostProcessing/Private/NeuralPostProcessModelInstance.cpp
@@ -0,0 +1,287 @@
+// Copyright Epic Games, Inc. All Rights Reserved.
+
+#include "NeuralPostProcessModelInstance.h"
+#include "NNERuntimeRDG.h"
+#include "RenderGraphUtils.h"
+#include "Engine/NeuralProfile.h"
+
+TSharedPtr<UE::NNE::IModelInstanceRDG> CreateNNEModelInstance(UNNEModelData* NNEModelData, FString RuntimeName)
+{
+	check(NNEModelData);
+
+	TWeakInterfacePtr<INNERuntime> Runtime = UE::NNE::GetRuntime<INNERuntime>(RuntimeName);
+	TWeakInterfacePtr<INNERuntimeRDG> RuntimeRDG = UE::NNE::GetRuntime<INNERuntimeRDG>(RuntimeName);
+	if (!Runtime.IsValid())
+	{
+		UE_LOG(LogNeuralPostProcessing, Error, TEXT("Can't get %s runtime."), *RuntimeName);
+		return nullptr;
+	}
+
+	if (!RuntimeRDG.IsValid())
+	{
+		UE_LOG(LogNeuralPostProcessing, Error, TEXT("No RDG runtime '%s' found"), *RuntimeName);
+		return nullptr;
+	}
+
+	TSharedPtr<UE::NNE::IModelRDG> ModelRDG = RuntimeRDG->CreateModelRDG(NNEModelData);
+	if (!ModelRDG.IsValid())
+	{
+		UE_LOG(LogNeuralPostProcessing, Error, TEXT("CreateModelRDG failed for Runtime = %s"), *RuntimeName);
+		return nullptr;
+	}
+
+	return ModelRDG->CreateModelInstanceRDG();
+}
+
+TSharedPtr<UE::NNE::IModelInstanceCPU> CreateNNECpuModelInstance(UNNEModelData* NNEModelData)
+{
+	check(NNEModelData);
+	const FString RuntimeName = "NNERuntimeORTCpu";
+	TWeakInterfacePtr<INNERuntime> Runtime = UE::NNE::GetRuntime<INNERuntime>(RuntimeName);
+	TWeakInterfacePtr<INNERuntimeCPU> RuntimeCPU = UE::NNE::GetRuntime<INNERuntimeCPU>(RuntimeName);
+	if (!Runtime.IsValid())
+	{
+		UE_LOG(LogNeuralPostProcessing, Error, TEXT("Can't get %s runtime."), *RuntimeName);
+		return nullptr;
+	}
+
+	if (!RuntimeCPU.IsValid())
+	{
+		UE_LOG(LogNeuralPostProcessing, Error, TEXT("No CPU runtime '%s' found"), *RuntimeName);
+		return nullptr;
+	}
+
+	TSharedPtr<UE::NNE::IModelCPU> ModelCPU = RuntimeCPU->CreateModelCPU(NNEModelData);
+	if (!ModelCPU.IsValid())
+	{
+		UE_LOG(LogNeuralPostProcessing, Error, TEXT("CreateModelCPU failed for Runtime = %s"), *RuntimeName);
+		return nullptr;
+	}
+
+	return ModelCPU->CreateModelInstanceCPU();
+}
+
+UNeuralPostProcessModelInstance::UNeuralPostProcessModelInstance(FObjectInitializer const& ObjectInitializer)
+	:Super(ObjectInitializer)
+{
+	RDGInputBuffer = nullptr;
+	RDGOutputBuffer = nullptr;
+	RDGTiledInputBuffer = nullptr;
+	RDGTiledOutputBuffer = nullptr;
+	TileSize = 1;
+	ModelTileSize = ENeuralModelTileType::OneByOne;
+	DimensionOverride = FIntVector4(-1);
+}
+
+void UNeuralPostProcessModelInstance::Update(UNNEModelData* NNEModelData, FString RuntimeName)
+{
+	this->CreateDefaultNNEModel(NNEModelData, RuntimeName);
+}
+
+void UNeuralPostProcessModelInstance::UpdateTileSize(int inTileSize)
+{
+	this->TileSize = inTileSize;
+}
+
+void UNeuralPostProcessModelInstance::Execute(FRDGBuilder& GraphBuilder)
+{
+	TArray<UE::NNE::FTensorBindingRDG> InputBindings;
+	TArray<UE::NNE::FTensorBindingRDG> OutputBindings;
+	UE::NNE::FTensorBindingRDG& Input = InputBindings.Emplace_GetRef();
+	UE::NNE::FTensorBindingRDG& Output = OutputBindings.Emplace_GetRef();
+
+	UE::NNE::FTensorShape InputShape = GetResolvedInputTensorShape();
+	int32 InputBatchSize = InputShape.GetData()[0];
+
+	auto RunSingleDispatch = [&]() {
+		Input.Buffer = RDGInputBuffer;
+		Output.Buffer = RDGOutputBuffer;
+		int Status = ModelInstanceRDG->EnqueueRDG(GraphBuilder, InputBindings, OutputBindings);
+
+		checkf(!Status, TEXT("EnqueueRDG failed: %d"), Status);
+
+		AddPass(GraphBuilder, RDG_EVENT_NAME("EnqueueRDGFence"), [this](FRHICommandListImmediate& RHICmdList)
+			{
+				// Need to wait until the output buffer, RDGOutputBuffer, is ready.
+				RHICmdList.SubmitCommandsAndFlushGPU();
+				RHICmdList.BlockUntilGPUIdle();
+			});
+	};
+
+	if (TileSize <= InputBatchSize)
+	{
+		RunSingleDispatch();
+	}
+	else
+	{
+		// Since we have more tiles than batch per dispatch we need to copy the data from tiled
+		// to and from the buffer.
+		int32 NumOfDispatch = (TileSize + InputBatchSize - 1) / InputBatchSize;
+		int32 InputBufferSize = RDGInputBuffer->Desc.GetSize();
+		int32 OutputBufferSize = RDGOutputBuffer->Desc.GetSize();
+
+		for (int i = 0; i < NumOfDispatch; ++i)
+		{
+			AddCopyBufferPass(GraphBuilder, RDGInputBuffer, 0, RDGTiledInputBuffer, InputBufferSize * i, InputBufferSize);
+			RunSingleDispatch();
+			AddCopyBufferPass(GraphBuilder, RDGTiledOutputBuffer, OutputBufferSize * i, RDGOutputBuffer, 0, OutputBufferSize);
+		}
+	}
+}
+
+UE::NNE::FTensorShape UNeuralPostProcessModelInstance::GetResolvedInputTensorShape()
+{
+	return ResolvedInputTensorShape;
+}
+
+UE::NNE::FTensorShape UNeuralPostProcessModelInstance::GetResolvedOutputTensorShape()
+{
+	return ResolvedOutputTensorShape;
+}
+
+FRDGBufferRef UNeuralPostProcessModelInstance::GetInputBuffer()
+{
+	return RDGInputBuffer;
+}
+
+FRDGBufferRef UNeuralPostProcessModelInstance::GetOutputBuffer()
+{
+	return RDGOutputBuffer;
+}
+
+FRDGBufferRef UNeuralPostProcessModelInstance::GetTiledInputBuffer()
+{
+	return RDGTiledInputBuffer;
+}
+
+FRDGBufferRef UNeuralPostProcessModelInstance::GetTiledOutputBuffer()
+{
+	return RDGTiledOutputBuffer;
+}
+
+bool UNeuralPostProcessModelInstance::ModifyInputShape(int Dim, int Size)
+{
+	// setup inputs
+	TConstArrayView<UE::NNE::FTensorDesc> InputTensorDescs = ModelInstanceRDG->GetInputTensorDescs();
+
+	UE::NNE::FSymbolicTensorShape InputShapeTemplate = InputTensorDescs[0].GetShape();
+	checkf(InputShapeTemplate.Rank() == 4, TEXT("Post Processing requires models with input shape batch x channel x height x width!"));
+
+	// We can only modify dimensions supporting dynamic Size
+	UE::NNE::FTensorShape CurrentResolvedInputTensorShape = GetResolvedInputTensorShape();
+
+	uint32 Batch = (Dim == 0) && (InputShapeTemplate.GetData()[0] < 0) ? Size : CurrentResolvedInputTensorShape.GetData()[0];
+	uint32 Channel = (Dim == 1) && (InputShapeTemplate.GetData()[1] < 0) ? Size : CurrentResolvedInputTensorShape.GetData()[1];
+	uint32 Height = (Dim == 2) && (InputShapeTemplate.GetData()[2] < 0) ? Size : CurrentResolvedInputTensorShape.GetData()[2];
+	uint32 Width = (Dim == 3) && (InputShapeTemplate.GetData()[3] < 0) ? Size : CurrentResolvedInputTensorShape.GetData()[3];
+	
+	TConstArrayView<uint32> NewResolvedInputShape = {Batch,Channel,Height,Width};
+
+	if (NewResolvedInputShape[Dim] != Size)
+	{
+		UE_LOG(LogNeuralPostProcessing, Error, TEXT("Cannot set dimension %d to %d. It is dynamic. revert back to %d"),Dim, Size, NewResolvedInputShape[Dim]);
+		return false;
+	}
+
+	ResolvedInputTensorShape = UE::NNE::FTensorShape::Make(NewResolvedInputShape);
+
+	ModelInstanceRDG->SetInputTensorShapes({ ResolvedInputTensorShape });
+
+	// Update the output shape as well due to the change of the input tensor shape
+
+	TConstArrayView<UE::NNE::FTensorShape> ResolvedOutputTensorShapes = ModelInstanceRDG->GetOutputTensorShapes();
+	if (ResolvedOutputTensorShapes.Num() < 1)
+	{
+		UE_LOG(LogNeuralPostProcessing, Error, TEXT("Cannot set dimension %d to %d due to NNE eror. Revert back to %d"), Dim, Size, CurrentResolvedInputTensorShape.GetData()[Dim]);
+		ModelInstanceRDG->SetInputTensorShapes({ CurrentResolvedInputTensorShape });
+		return false;
+	}
+
+	ResolvedOutputTensorShape = ResolvedOutputTensorShapes[0];
+
+	return true;
+}
+
+void UNeuralPostProcessModelInstance::CreateRDGBuffers(class FRDGBuilder& GraphBuilder)
+{
+	// Crate input buffers
+	UE::NNE::FTensorShape InputShape = GetResolvedInputTensorShape();
+	FIntPoint NeuralNetworkInputSize = { (int32) InputShape.GetData()[3], (int32) InputShape.GetData()[2] };
+	int32 InputBatch = InputShape.GetData()[0];
+	int32 InputChannels = InputShape.GetData()[1];
+
+	FRDGBufferDesc InputBufferDesc = FRDGBufferDesc::CreateBufferDesc(sizeof(float), NeuralNetworkInputSize.X * NeuralNetworkInputSize.Y * InputChannels * InputBatch);
+	
+	RDGInputBuffer = GraphBuilder.CreateBuffer(InputBufferDesc, TEXT("NeuralPostProcessing.InputBuffer"));
+
+	// Create output buffers
+	checkf(ModelInstanceRDG->GetOutputTensorShapes().Num() == 1, TEXT("Post Processing requires models with a single output tensor!"));
+	
+	ResolvedOutputTensorShape = ModelInstanceRDG->GetOutputTensorShapes()[0];
+	checkf(ResolvedOutputTensorShape.Rank() == 4, TEXT("Neural Post Processing requires models dim = 4 [static/dynamic x C x height x width]!"));
+	checkf(ResolvedOutputTensorShape.GetData()[0] != 0 , TEXT("Neural Post Processing requires models with output shape BatchSize != 0"));
+	checkf(ResolvedOutputTensorShape.GetData()[1] >= 1, TEXT("Neural Post Processing requires models with output shape channel C >= 1!"));
+	checkf(ResolvedOutputTensorShape.GetData()[2] > 0, TEXT("Neural Post Processing requires models with output height > 0!"));
+	checkf(ResolvedOutputTensorShape.GetData()[3] > 0, TEXT("Neural Post Processing requires models with output width > 0!"));
+	FIntPoint NeuralNetworkOutputSize = { (int32)ResolvedOutputTensorShape.GetData()[3], (int32)ResolvedOutputTensorShape.GetData()[2] };
+	int32 OutputBatch = ResolvedOutputTensorShape.GetData()[0];
+	int32 OutputChannels = ResolvedOutputTensorShape.GetData()[1];
+
+	FRDGBufferDesc OutputBufferDesc = FRDGBufferDesc::CreateBufferDesc(sizeof(float), 
+		NeuralNetworkOutputSize.X * NeuralNetworkOutputSize.Y * OutputChannels * OutputBatch);
+	
+	RDGOutputBuffer = GraphBuilder.CreateBuffer(OutputBufferDesc, *FString("SubsurfacePostProcessing.OutputBuffer"));
+
+	if (TileSize <= InputBatch)
+	{
+		RDGTiledInputBuffer = RDGInputBuffer;
+		RDGTiledOutputBuffer = RDGOutputBuffer;
+	}
+	else if (TileSize > InputBatch)
+	{
+		FRDGBufferDesc InputTiledBufferDesc = FRDGBufferDesc::CreateBufferDesc(sizeof(float), 
+			NeuralNetworkInputSize.X * NeuralNetworkInputSize.Y * InputChannels * InputBatch * TileSize);
+		RDGTiledInputBuffer = GraphBuilder.CreateBuffer(InputTiledBufferDesc, TEXT("NeuralPostProcessing.TiledInputBuffer"));
+
+		FRDGBufferDesc OutputTiledBufferDesc = FRDGBufferDesc::CreateBufferDesc(sizeof(float), 
+			NeuralNetworkOutputSize.X * NeuralNetworkOutputSize.Y * OutputChannels * OutputBatch * TileSize);
+		RDGTiledOutputBuffer = GraphBuilder.CreateBuffer(OutputTiledBufferDesc, TEXT("NeuralPostProcessing.TiledOutputBuffer"));
+	}
+}
+
+void UNeuralPostProcessModelInstance::CreateRDGBuffersIfNeeded(FRDGBuilder& GraphBuilder, bool bForceCreate)
+{
+	if (!RDGInputBuffer || !RDGOutputBuffer || bForceCreate)
+	{
+		this->CreateRDGBuffers(GraphBuilder);
+	}
+}
+
+void UNeuralPostProcessModelInstance::CreateDefaultNNEModel(UNNEModelData* NNEModelData, FString RuntimeName)
+{
+	ModelInstanceRDG = CreateNNEModelInstance(NNEModelData, RuntimeName);
+
+	if (ModelInstanceRDG)
+	{
+		// setup inputs
+		TConstArrayView<UE::NNE::FTensorDesc> InputTensorDescs = ModelInstanceRDG->GetInputTensorDescs();
+
+		// dimension check
+		checkf(InputTensorDescs.Num() >= 1, TEXT("Post Processing requires models with a single input tensor!"));
+		
+		UE::NNE::FSymbolicTensorShape InputShape = InputTensorDescs[0].GetShape();
+		checkf(InputShape.Rank() == 4, TEXT("Post Processing requires models with input shape 1 x N(N>=1) x height x width!"));
+		checkf(InputShape.GetData()[0] != 0, TEXT("Post Processing requires models with input shape ? x N(N>=1) x height x width!"));
+		checkf(InputShape.GetData()[1] >= 1, TEXT("Post Processing requires models with input shape ? x N(N>=1) x height x width!"));
+
+		// Set up the input tensor shape
+		// All dynamic dimensions are set to 1 by default.
+		ResolvedInputTensorShape = UE::NNE::FTensorShape::MakeFromSymbolic(InputShape);
+
+		ModelInstanceRDG->SetInputTensorShapes({ ResolvedInputTensorShape });
+	}
+	else
+	{
+		UE_LOG(LogNeuralPostProcessing, Error, TEXT("Failed to create NNE RDG Model."));
+	}
+}
diff --git a/Engine/Plugins/Experimental/NeuralRendering/Source/NeuralPostProcessing/Private/NeuralPostProcessModelInstance.h b/Engine/Plugins/Experimental/NeuralRendering/Source/NeuralPostProcessing/Private/NeuralPostProcessModelInstance.h
new file mode 100644
--- /dev/null
+++ b/Engine/Plugins/Experimental/NeuralRendering/Source/NeuralPostProcessing/Private/NeuralPostProcessModelInstance.h
@@ -0,0 +1,93 @@
+// Copyright Epic Games, Inc. All Rights Reserved.
+
+#pragma once
+
+#include "NNE.h"
+#include "NNETypes.h"
+#include "NNERuntimeRDG.h"
+#include "NNERuntimeCPU.h"
+#include "NNEModelData.h"
+
+#include "NeuralPostProcessModelInstance.generated.h"
+
+#if WITH_EDITOR
+	DECLARE_LOG_CATEGORY_EXTERN(LogNeuralPostProcessing, Log, All);
+#endif
+
+//Util functions
+TSharedPtr<UE::NNE::IModelInstanceRDG> CreateNNEModelInstance(UNNEModelData* NNEModelData, FString RuntimeName);
+TSharedPtr<UE::NNE::IModelInstanceCPU> CreateNNECpuModelInstance(UNNEModelData* NNEModelData);
+
+enum class ENeuralModelTileType : uint8;
+
+UCLASS()
+class UNeuralPostProcessModelInstance : public UObject
+{
+	GENERATED_UCLASS_BODY()
+public:
+
+	// Engine thread 
+	// If the model data and the Runtime name is different, recreate
+	// the model
+	void Update(UNNEModelData* NNEModelData, FString RuntimeName);
+
+	// Render thread functions
+	void CreateRDGBuffers(class FRDGBuilder& GraphBuilder);
+
+	// Create the input and output buffer on demand
+	// useful to use multiple post processing material to read/write
+	void CreateRDGBuffersIfNeeded(class FRDGBuilder& GraphBuilder, bool bForceCreate = false);
+
+	void Execute(class FRDGBuilder& GraphBuilder);
+
+	UE::NNE::FTensorShape GetResolvedInputTensorShape();
+
+	UE::NNE::FTensorShape GetResolvedOutputTensorShape();
+
+	//TODO: different type support
+	FRDGBufferRef GetInputBuffer();
+	FRDGBufferRef GetOutputBuffer();
+
+	FRDGBufferRef GetTiledInputBuffer();
+	FRDGBufferRef GetTiledOutputBuffer();
+
+	void UpdateTileSize(int TileSize = 1);
+	int GetTileSize() const { return TileSize; }
+
+	void UpdateModelTileType(ENeuralModelTileType inTileType) { ModelTileSize = inTileType; }
+	ENeuralModelTileType GetModelTileType()const { return ModelTileSize; }
+
+	bool IsValid() { return ModelInstanceRDG.IsValid();}
+
+	bool ModifyInputShape(int Dim, int Size);
+
+private:
+	void CreateDefaultNNEModel(UNNEModelData* NNEModelData, FString RuntimeName);
+
+	// Input Buffer for neural post process network pass
+	FRDGBufferRef RDGInputBuffer;
+	
+	// Output Buffer for Neural post process network pass
+	FRDGBufferRef RDGOutputBuffer;
+
+	// Temporary buffers to emulate batchsize
+	FRDGBufferRef RDGTiledInputBuffer;
+
+	FRDGBufferRef RDGTiledOutputBuffer;
+
+	ENeuralModelTileType ModelTileSize;
+
+	int TileSize = 1;
+
+	// The NNE RDG Model 
+	TSharedPtr<UE::NNE::IModelInstanceRDG> ModelInstanceRDG;
+
+	// Resolved Input Tensorshape
+	UE::NNE::FTensorShape ResolvedInputTensorShape;
+
+	// Output Tensorshape
+	UE::NNE::FTensorShape ResolvedOutputTensorShape;
+
+	FIntVector4 DimensionOverride;
+};
+
diff --git a/Engine/Plugins/Experimental/NeuralRendering/Source/NeuralPostProcessing/Private/NeuralPostProcessing.cpp b/Engine/Plugins/Experimental/NeuralRendering/Source/NeuralPostProcessing/Private/NeuralPostProcessing.cpp
new file mode 100644
--- /dev/null
+++ b/Engine/Plugins/Experimental/NeuralRendering/Source/NeuralPostProcessing/Private/NeuralPostProcessing.cpp
@@ -0,0 +1,473 @@
+// Copyright Epic Games, Inc. All Rights Reserved.
+
+#include "NeuralPostProcessing.h"
+#include "Interfaces/IPluginManager.h"
+#include "NNE.h"
+#include "NNEModelData.h"
+#include "NNERuntimeRDG.h"
+#include "Engine/NeuralProfile.h"
+#include "NeuralPostProcessModelInstance.h"
+#include "NeuralPostProcessingCS.h"
+#include "PostProcess/PostProcessing.h"
+#include "PostProcess/NeuralPostProcess.h"
+#include "PixelShaderUtils.h"
+#include "RenderGraphEvent.h"
+
+#define LOCTEXT_NAMESPACE "FNeuralPostProcessingModule"
+
+#if WITH_EDITOR
+DEFINE_LOG_CATEGORY(LogNeuralPostProcessing);
+#endif
+
+namespace
+{
+	TAutoConsoleVariable<int32> CVarNeuralPostProcessApply(
+		TEXT("r.neuralpostprocess.apply"),
+		1,
+		TEXT(" 0: disabled\n")
+		TEXT(" 1: enabled (default)"),
+		ECVF_RenderThreadSafe | ECVF_Scalability);
+
+	DECLARE_GPU_STAT(NeuralPostProcessing)
+}
+
+
+/**
+*
+===============================================================================================================================
+                    FNeuralPostProcessModelInstanceManager: Hold the mapping between neural profile and neural network models.
+										/|\   /|\
+										 | USE |
+										 |     |
+	|-ENGINE_API FNeuralProfileManager---|	   |
+	|(INeuralProfileManager)			 |	   |
+	|									 |	   |----RENDERING_API FNeuralPostProcess----|
+	|									 |	   |	(INeuralPostProcessInterface)		|
+	|	Manage the neural profile		 |	   |										|
+	|									 |	   |										|
+	|------------------------------------|	   |	Interface accessed by RDGBuilder	|
+											   |	in post process material.			|
+											   |----------------------------------------|
+											   |    AllocateBuffer()                    |
+											   |    Apply(): apply neural networks      |
+===============================================================================================================================
+*/
+
+class FNeuralPostProcessModelInstanceManager
+{
+public:
+	static FNeuralPostProcessModelInstanceManager* Get()
+	{
+		static FNeuralPostProcessModelInstanceManager Manager;
+		return &Manager;
+	}
+
+	UNeuralPostProcessModelInstance* GetModelInstance(int32 ProfileId)
+	{
+		if (NeuralPostProcessModelInstances.Contains(ProfileId))
+		{
+			return NeuralPostProcessModelInstances[ProfileId];
+		}
+		else
+		{
+			return nullptr;
+		}
+	}
+
+	void SetModelInstance(int32 ProfileId, UNeuralPostProcessModelInstance* Instance)
+	{
+		NeuralPostProcessModelInstances.Emplace(ProfileId, Instance);
+	}
+
+	UNeuralPostProcessModelInstance* Remove(int32 ProfileId)
+	{
+		return NeuralPostProcessModelInstances.FindAndRemoveChecked(ProfileId);
+	}
+
+	~FNeuralPostProcessModelInstanceManager()
+	{
+		NeuralPostProcessModelInstances.Reset();
+	}
+public:
+	TMap<int32, UNeuralPostProcessModelInstance*> NeuralPostProcessModelInstances;
+};
+
+class FNeuralProfileManager : public NeuralProfile::INeuralProfileManager
+{
+public:
+	FNeuralProfileManager() {}
+
+	virtual void UpdateModel(int32 AllocationId, UObject* NNEModelData, FString RuntimeName) override
+	{
+		check(IsInRenderingThread())
+
+		UNNEModelData* RawNNEModelData = Cast<UNNEModelData>(NNEModelData);
+
+		if (!IsValid(RawNNEModelData))
+		{
+			UE_LOG(LogNeuralPostProcessing, Error, TEXT("NNEModelData is invalid at Slot %d."), AllocationId);
+			return;
+		}
+
+		UNeuralPostProcessModelInstance* ModelInstance = NewObject<UNeuralPostProcessModelInstance>();
+		ModelInstance->Update(RawNNEModelData, RuntimeName);
+
+		FNeuralPostProcessModelInstanceManager::Get()->SetModelInstance(AllocationId, ModelInstance->IsValid()? ModelInstance : nullptr);
+	}
+
+	virtual void UpdateTileType(int32 AllocationId, ENeuralModelTileType ModelTileType) override
+	{
+		UNeuralPostProcessModelInstance* Instance = FNeuralPostProcessModelInstanceManager::Get()->GetModelInstance(AllocationId);
+		Instance->UpdateModelTileType(ModelTileType);
+	}
+
+	virtual bool UpdateBatchSize(int32 AllocationId, int32 BatchSize) override
+	{
+		UNeuralPostProcessModelInstance* Instance = FNeuralPostProcessModelInstanceManager::Get()->GetModelInstance(AllocationId);
+		return Instance->ModifyInputShape(0, BatchSize);
+	}
+
+	virtual void RemoveModel(int32 AllocationId) override
+	{
+		FNeuralPostProcessModelInstanceManager::Get()->Remove(AllocationId);
+	}
+
+	virtual FIntVector4 GetInputDimension(UObject* NNEModelData, FString RuntimeName) override
+	{
+		UE::NNE::FTensorShape TensorShape;
+		FIntVector4 InputDimension = FIntVector4(-1, -1, -1, -1);
+
+		if (UNNEModelData* ModelData = Cast<UNNEModelData>(NNEModelData))
+		{
+			// Need to create the ModelInstance in order to get the dimension in case the Model is not created
+			// in the rendering thread.
+			TSharedPtr<UE::NNE::IModelInstanceCPU> ModelInstance = CreateNNECpuModelInstance(ModelData);
+
+			if (ModelInstance)
+			{
+				TConstArrayView<UE::NNE::FTensorDesc> InputTensorDescs = ModelInstance->GetInputTensorDescs();
+				UE::NNE::FSymbolicTensorShape InputTensorShape = InputTensorDescs[0].GetShape();
+
+				// Support only output dimension of rank 4
+				checkf(InputTensorShape.Rank() == 4, TEXT("Neural Post Processing requires models with input shape N x channel x height x width!"));
+
+				InputDimension = FIntVector4(
+					InputTensorShape.GetData()[0],
+					InputTensorShape.GetData()[1],
+					InputTensorShape.GetData()[2],
+					InputTensorShape.GetData()[3]);
+			}
+			ModelInstance.Reset();
+		}
+
+		return InputDimension;
+	}
+
+	virtual FIntVector4 GetOutputDimension(UObject* NNEModelData, FString RuntimeName) override
+	{
+		UE::NNE::FTensorShape TensorShape;
+		FIntVector4 OutputDimension = FIntVector4(-1,-1,-1,-1);
+
+		if (UNNEModelData* ModelData = Cast<UNNEModelData>(NNEModelData))
+		{
+			TSharedPtr<UE::NNE::IModelInstanceCPU> ModelInstance = CreateNNECpuModelInstance(ModelData);
+			
+			if (ModelInstance)
+			{
+				TConstArrayView<UE::NNE::FTensorDesc> OutputTensorDescs = ModelInstance->GetOutputTensorDescs();
+				UE::NNE::FSymbolicTensorShape OutputTensorShape = OutputTensorDescs[0].GetShape();
+
+				checkf(OutputTensorShape.Rank() == 4, TEXT("Neural Post Processing requires models with output shape N x channel x height x width!"));
+
+				OutputDimension = FIntVector4(
+					OutputTensorShape.GetData()[0],
+					OutputTensorShape.GetData()[1],
+					OutputTensorShape.GetData()[2],
+					OutputTensorShape.GetData()[3]);
+			}
+
+			ModelInstance.Reset();
+		}
+		return OutputDimension;
+	}
+
+	~FNeuralProfileManager() {}
+};
+
+
+int32 GetTotalModelTileCount(ENeuralModelTileType ModelTileSize)
+{
+	int TotalTileCount = 1;
+
+	//@TODO: uncomment and implement the corresponding function.
+	switch (ModelTileSize)
+	{
+	case ENeuralModelTileType::TwoByTwo:
+		TotalTileCount = 2 * 2;
+		break;
+	case ENeuralModelTileType::FourByFour:
+		TotalTileCount = 4 * 4;
+		break;
+	case ENeuralModelTileType::EightByEight:
+		TotalTileCount = 8 * 8;
+		break;
+	//case ENeuralModelTileType::Auto:
+	//	TotalTileCount = -1;	// determined at runtime by the viewport size and W,H dimension of the network 
+	//	break;
+	case ENeuralModelTileType::OneByOne:
+	default:
+		break;
+	}
+
+	return TotalTileCount;
+}
+
+static void AllocateInputBuffer_RenderingThread(
+	FRDGBuilder& GraphBuilder,
+	const FScreenPassTextureViewport& Viewport,
+	int32 ProfileId,
+	FRDGBufferRef& InputNeuralBuffer,
+	FVector4f& InputBufferDimension)
+{
+	check(IsInRenderingThread())
+		check(ProfileId >= 0);
+
+	UNeuralPostProcessModelInstance* Model = FNeuralPostProcessModelInstanceManager::Get()->GetModelInstance(ProfileId);
+
+	if (!IsValid(Model))
+	{
+		InputNeuralBuffer = nullptr;
+		return;
+	}
+
+	UE::NNE::FTensorShape InputShape = Model->GetResolvedInputTensorShape();
+	InputBufferDimension = FVector4f(InputShape.GetData()[0], InputShape.GetData()[1], InputShape.GetData()[2], InputShape.GetData()[3]);
+	
+	// Calculate the number of tiles
+	int NumTiles = 1;
+	{
+		ENeuralModelTileType ModelTileSize = Model->GetModelTileType();
+		// @todo: Finish the implement of auto
+		//if (ModelTileSize == ENeuralModelTileType::Auto)
+		//{
+		//	// allocate the number of tiles based on the size of the viewport and the size of the buffer.
+		//	FScreenPassTextureViewportParameters ViewportParameters = GetScreenPassTextureViewportParameters(Viewport);
+		//	FVector2f TileSizeWH = FMath::DivideAndRoundUp(ViewportParameters.ViewportSize, FVector2f(InputBufferDimension.W, InputBufferDimension.Z));
+		//	float BatchDim = InputBufferDimension.X;
+		//	NumTiles = FMath::DivideAndRoundUp(TileSizeWH.X * TileSizeWH.Y, BatchDim);
+		//}
+		//else
+		{
+			NumTiles = GetTotalModelTileCount(ModelTileSize);
+		}
+
+		Model->UpdateTileSize(NumTiles);
+	}
+
+	Model->CreateRDGBuffersIfNeeded(GraphBuilder, true);
+
+	// Output the buffer and dimension for use.
+	InputNeuralBuffer = Model->GetTiledInputBuffer();
+	InputBufferDimension.X *= NumTiles;
+}
+
+static void	ApplyNeuralNetworks_RenderingThread(
+	FRDGBuilder& GraphBuilder,
+	const FViewInfo& View,
+	int32 ProfileId,
+	FRDGTextureRef NeuralTexture,
+	FIntRect ViewRect,
+	FRDGBufferRef InputSourceType,
+	FRDGBufferRef& OutputNeuralBuffer,
+	FVector4f& BufferDimension)
+{
+	check(IsInRenderingThread())
+		check(ProfileId >= 0);
+
+	UNeuralPostProcessModelInstance* Model = FNeuralPostProcessModelInstanceManager::Get()->GetModelInstance(ProfileId);
+
+	if (!IsValid(Model))
+	{
+		return;
+	}
+
+	RDG_GPU_STAT_SCOPE(GraphBuilder, NeuralPostProcessing);
+
+	FIntPoint TextureSize = NeuralTexture->Desc.Extent;
+	FGlobalShaderMap* GlobalShaderMap = GetGlobalShaderMap(GMaxRHIFeatureLevel);
+	float Scale = 1.0;
+	
+	const FScreenPassTextureViewport NeuralPostProcessViewport(NeuralTexture, ViewRect);
+	const FScreenPassTextureViewportParameters InputViewportParameters = GetScreenPassTextureViewportParameters(NeuralPostProcessViewport);
+
+	// 1. Preprocess the input data by copying from Texture to buffer if required
+	{
+		// Build the indirect dispatch parameters
+		FRDGBufferRef IndirectDispatchBuffer = GraphBuilder.CreateBuffer(
+			FRDGBufferDesc::CreateIndirectDesc<FRHIDispatchIndirectParameters>(1), TEXT("NeuralPostProcessing.IndirectDispatchBuffer"));
+
+		UE::NNE::FTensorShape InputShape = Model->GetResolvedInputTensorShape();
+		FIntPoint NeuralNetworkInputSize = { (int32)InputShape.GetData()[3], (int32)InputShape.GetData()[2] };// Width, height
+		int32 TileSize = Model->GetTileSize();
+
+		// Build the indirect dispatch parameters with InputSourceType
+		{
+			typedef FNeuralPostProcessingBuildIndirectDispatchArgsCS ARGSETUPSHADER;
+			ARGSETUPSHADER::FParameters* PassParameters = GraphBuilder.AllocParameters<ARGSETUPSHADER::FParameters>();
+			PassParameters->NetworkTextureSize = NeuralNetworkInputSize;
+			PassParameters->SourceType = GraphBuilder.CreateSRV(InputSourceType, EPixelFormat::PF_R32_UINT);
+			PassParameters->RWIndirectDispatchArgsBuffer = GraphBuilder.CreateUAV(IndirectDispatchBuffer, EPixelFormat::PF_R32_UINT);
+			TShaderMapRef<ARGSETUPSHADER> ComputeShader(View.ShaderMap);
+			FComputeShaderUtils::AddPass(GraphBuilder, FRDGEventName(TEXT("NeuralPostProcessing::BuildIndirectArgs(Dispatch)")), ComputeShader, PassParameters, FIntVector(1,1,1));
+		}
+
+		// Use the neural network's input dimension configuration if available, otherwise use the dynamic size of the buffer
+		if (TileSize == 1)
+		{
+			typedef FNeuralPostProcessingPrepareInputCS SHADER;
+			SHADER::FParameters* PassParameters = GraphBuilder.AllocParameters<SHADER::FParameters>();
+			{
+				PassParameters->Input0 = GetNeuralPostProcessInput(NeuralTexture, InputViewportParameters);
+				PassParameters->InputTextureSampler = TStaticSamplerState<SF_Bilinear, AM_Clamp, AM_Clamp, AM_Clamp>::GetRHI();
+				PassParameters->InputBufferWidth = NeuralNetworkInputSize.X;
+				PassParameters->InputBufferHeight = NeuralNetworkInputSize.Y;
+				PassParameters->InputBuffer = GraphBuilder.CreateUAV(FRDGBufferUAVDesc(Model->GetInputBuffer(), PF_R32_FLOAT));
+				PassParameters->ColorScale = Scale;
+				PassParameters->IndirectDispatchArgsBuffer = IndirectDispatchBuffer;
+			}
+
+			TShaderMapRef<FNeuralPostProcessingPrepareInputCS> ComputeShader(GlobalShaderMap);
+			FComputeShaderUtils::AddPass(
+				GraphBuilder,
+				RDG_EVENT_NAME("NeuralPostProcessing::PrepareInput"),
+				ERDGPassFlags::Compute, //| ERDGPassFlags::NeverCull
+				ComputeShader,
+				PassParameters,
+				PassParameters->IndirectDispatchArgsBuffer,0);
+		}
+		else
+		{
+			// @TODO: Implement texture copy when neural network buffer is accessed by texture index.
+		}
+	}
+
+	// 2. Run the network
+	const bool bNeuralPostProcessApply = CVarNeuralPostProcessApply.GetValueOnRenderThread() > 0;
+	if (bNeuralPostProcessApply)
+	{
+		Model->Execute(GraphBuilder);
+	}
+
+	// 3. Pass the output from the neural network and fill the scene color texture
+	// when batch and channel dimension matches between the input and output buffer.
+
+	auto ShouldUpdateNeuralTexture = [&]()->bool {
+		UE::NNE::FTensorShape OutputShape = Model->GetResolvedOutputTensorShape();
+		UE::NNE::FTensorShape InputShape = Model->GetResolvedInputTensorShape();
+
+		return OutputShape.GetData()[0] == InputShape.GetData()[0] &&
+			OutputShape.GetData()[1] == InputShape.GetData()[1] &&
+			Model->GetTileSize() == 1; //@TODO: add texture copy when tile size is larger than batch size. 
+	};
+
+	if (ShouldUpdateNeuralTexture())
+	{
+		FRDGBufferRef OutputBuffer = bNeuralPostProcessApply ? Model->GetOutputBuffer() : Model->GetInputBuffer();
+		FRDGBufferUAVRef OutputBufferUAV = GraphBuilder.CreateUAV(FRDGBufferUAVDesc(OutputBuffer, PF_R32_FLOAT));
+
+		UE::NNE::FTensorShape OutputShape = bNeuralPostProcessApply ? Model->GetResolvedOutputTensorShape() : Model->GetResolvedInputTensorShape();
+		FIntPoint NeuralNetworkOutputSize = { (int32)OutputShape.GetData()[3], (int32)OutputShape.GetData()[2] };
+
+		const FIntRect Viewport = FIntRect(0,TextureSize);
+
+		FNeuralPostProcessingProcessOutputPS::FParameters* ProcessOutputParameters = GraphBuilder.AllocParameters<FNeuralPostProcessingProcessOutputPS::FParameters>();
+		ProcessOutputParameters->Input0.Texture = nullptr;
+		ProcessOutputParameters->Input0.Viewport = InputViewportParameters;
+		ProcessOutputParameters->OutputBufferWidth = NeuralNetworkOutputSize.X;
+		ProcessOutputParameters->OutputBufferHeight = NeuralNetworkOutputSize.Y;
+		ProcessOutputParameters->OutputBuffer = OutputBufferUAV;
+		ProcessOutputParameters->ColorScale = Scale;
+		ProcessOutputParameters->RenderTargets[0] = FRenderTargetBinding(NeuralTexture, ERenderTargetLoadAction::ELoad);
+
+		TShaderMapRef<FNeuralPostProcessingProcessOutputPS> WriteOutputShader(GlobalShaderMap);
+
+		FPixelShaderUtils::AddFullscreenPass(
+			GraphBuilder,
+			GlobalShaderMap,
+			RDG_EVENT_NAME("NeuralPostProcessing::ProcessOutput"),
+			WriteOutputShader,
+			ProcessOutputParameters,
+			ViewRect);
+	}
+
+	// 4. Read back the buffer, so the user can directly decode in the post process material.
+	OutputNeuralBuffer = bNeuralPostProcessApply ? Model->GetTiledOutputBuffer() : Model->GetTiledInputBuffer();
+	UE::NNE::FTensorShape OutputShape = bNeuralPostProcessApply ? Model->GetResolvedOutputTensorShape() : Model->GetResolvedInputTensorShape();		
+	BufferDimension = FVector4f(OutputShape.GetData()[0], OutputShape.GetData()[1], OutputShape.GetData()[2], OutputShape.GetData()[3]);
+	BufferDimension.X *= Model->GetTileSize();
+}
+
+class FNeuralPostProcess : public INeuralPostProcessInterface
+{
+public:
+	virtual void Apply(FRDGBuilder& GraphBuilder, const FViewInfo& View, int32 NeuralProfileId,
+		FRDGTexture* NeuralTexture, FIntRect ViewRect, FRDGBufferRef InputSourceType,
+		FRDGBufferRef& OutputNeuralBuffer, FVector4f& BufferDimension) override
+	{
+		ApplyNeuralNetworks_RenderingThread(
+			GraphBuilder,
+			View,
+			NeuralProfileId,
+			NeuralTexture,
+			ViewRect,
+			InputSourceType,
+			OutputNeuralBuffer,
+			BufferDimension);
+	}
+
+	virtual void AllocateBuffer(FRDGBuilder& GraphBuilder, const FScreenPassTextureViewport& Viewport,
+		int32 NeuralProfileId, FRDGBufferRef& InputNeuralBuffer, FVector4f& InputBufferDimension) override
+	{
+		AllocateInputBuffer_RenderingThread(
+			GraphBuilder, 
+			Viewport, 
+			NeuralProfileId, 
+			InputNeuralBuffer, 
+			InputBufferDimension);
+	}
+};
+
+void FNeuralPostProcessingModule::StartupModule()
+{
+#if WITH_EDITOR
+	UE_LOG(LogNeuralPostProcessing, Log, TEXT("NeuralPostProcessing starting up"));
+#endif
+
+	const TSharedPtr<IPlugin> Plugin = IPluginManager::Get().FindPlugin(TEXT("NeuralRendering"));
+	if (Plugin.IsValid())
+	{
+		FString ModuleDir = Plugin->GetBaseDir() + TEXT("/Source/NeuralPostProcessing");
+		AddShaderSourceDirectoryMapping(TEXT("/NeuralRendering"), FPaths::Combine(ModuleDir, TEXT("Shaders")));
+	}
+	else
+	{
+		UE_LOG(LogNeuralPostProcessing, Error, TEXT("Shaders directory not added. Failed to find NeuralPostProcessing plugin"));
+	}
+
+	GNeuralProfileManager.Reset(new FNeuralProfileManager());
+	GNeuralPostProcess.Reset(new FNeuralPostProcess());
+}
+
+void FNeuralPostProcessingModule::ShutdownModule()
+{
+#if WITH_EDITOR
+	UE_LOG(LogNeuralPostProcessing, Log, TEXT("NeuralPostProcessing function shutting down"));
+#endif
+
+	GNeuralProfileManager.Reset(nullptr);
+	GNeuralPostProcess.Reset(nullptr);
+}
+
+#undef LOCTEXT_NAMESPACE
+	
+IMPLEMENT_MODULE(FNeuralPostProcessingModule, NeuralPostProcessing)
\ No newline at end of file
diff --git a/Engine/Plugins/Experimental/NeuralRendering/Source/NeuralPostProcessing/Private/NeuralPostProcessingCS.cpp b/Engine/Plugins/Experimental/NeuralRendering/Source/NeuralPostProcessing/Private/NeuralPostProcessingCS.cpp
new file mode 100644
--- /dev/null
+++ b/Engine/Plugins/Experimental/NeuralRendering/Source/NeuralPostProcessing/Private/NeuralPostProcessingCS.cpp
@@ -0,0 +1,41 @@
+// Copyright Epic Games, Inc. All Rights Reserved.
+
+
+#include "NeuralPostProcessingCS.h"
+
+
+FNueralPostProcessInput GetNeuralPostProcessInput(FRDGTextureRef Texture, const FScreenPassTextureViewportParameters& ViewportParameters)
+{
+	FNueralPostProcessInput Input;
+	Input.Texture = Texture;
+	Input.Viewport = ViewportParameters;
+	return Input;
+}
+
+
+void FNeuralPostProcessingBuildIndirectDispatchArgsCS::ModifyCompilationEnvironment(const FGlobalShaderPermutationParameters& Parameters, FShaderCompilerEnvironment& OutEnvironment)
+{
+	FGlobalShader::ModifyCompilationEnvironment(Parameters, OutEnvironment);
+	OutEnvironment.SetDefine(TEXT("THREAD_GROUP_SIZE"), NEURAL_POST_PROCESSING_THREAD_GROUP_SIZE);
+}
+
+bool FNeuralPostProcessingBuildIndirectDispatchArgsCS::ShouldCompilePermutation(const FGlobalShaderPermutationParameters& Parameters)
+{
+	return true;
+}
+
+void FNeuralPostProcessingPrepareInputCS::ModifyCompilationEnvironment(const FGlobalShaderPermutationParameters& InParameters, FShaderCompilerEnvironment& OutEnvironment)
+{
+	FGlobalShader::ModifyCompilationEnvironment(InParameters, OutEnvironment);
+	OutEnvironment.SetDefine(TEXT("THREAD_GROUP_SIZE"), NEURAL_POST_PROCESSING_THREAD_GROUP_SIZE);
+}
+
+void FNeuralPostProcessingProcessOutputPS::ModifyCompilationEnvironment(const FGlobalShaderPermutationParameters& InParameters, FShaderCompilerEnvironment& OutEnvironment)
+{
+	FGlobalShader::ModifyCompilationEnvironment(InParameters, OutEnvironment);
+	OutEnvironment.SetDefine(TEXT("THREAD_GROUP_SIZE"), NEURAL_POST_PROCESSING_THREAD_GROUP_SIZE);
+}
+
+IMPLEMENT_GLOBAL_SHADER(FNeuralPostProcessingBuildIndirectDispatchArgsCS, "/NeuralRendering/NeuralPostProcessing.usf", "BuildIndirectDispatchArgsCS", SF_Compute);
+IMPLEMENT_GLOBAL_SHADER(FNeuralPostProcessingPrepareInputCS, "/NeuralRendering/NeuralPostProcessing.usf", "PrepareInput", SF_Compute);
+IMPLEMENT_GLOBAL_SHADER(FNeuralPostProcessingProcessOutputPS, "/NeuralRendering/NeuralPostProcessing.usf", "ProcessOutput", SF_Pixel);
\ No newline at end of file
diff --git a/Engine/Plugins/Experimental/NeuralRendering/Source/NeuralPostProcessing/Private/NeuralPostProcessingCS.h b/Engine/Plugins/Experimental/NeuralRendering/Source/NeuralPostProcessing/Private/NeuralPostProcessingCS.h
new file mode 100644
--- /dev/null
+++ b/Engine/Plugins/Experimental/NeuralRendering/Source/NeuralPostProcessing/Private/NeuralPostProcessingCS.h
@@ -0,0 +1,73 @@
+// Copyright Epic Games, Inc. All Rights Reserved.
+
+#pragma once
+
+#include "GlobalShader.h"
+#include "ShaderParameterUtils.h"
+#include "RenderGraphUtils.h"
+#include "ScreenPass.h"
+
+#define NEURAL_POST_PROCESSING_THREAD_GROUP_SIZE 32
+
+BEGIN_SHADER_PARAMETER_STRUCT(FNueralPostProcessInput, )
+	SHADER_PARAMETER_STRUCT_INCLUDE(FScreenPassTextureViewportParameters, Viewport)
+	SHADER_PARAMETER_RDG_TEXTURE(Texture2D, Texture)
+END_SHADER_PARAMETER_STRUCT()
+
+FNueralPostProcessInput GetNeuralPostProcessInput(FRDGTextureRef Texture, const FScreenPassTextureViewportParameters& ViewportParameters);
+
+class FNeuralPostProcessingBuildIndirectDispatchArgsCS : public FGlobalShader
+{
+public:
+	DECLARE_GLOBAL_SHADER(FNeuralPostProcessingBuildIndirectDispatchArgsCS);
+	SHADER_USE_PARAMETER_STRUCT(FNeuralPostProcessingBuildIndirectDispatchArgsCS, FGlobalShader)
+
+	static void ModifyCompilationEnvironment(const FGlobalShaderPermutationParameters& Parameters, FShaderCompilerEnvironment& OutEnvironment);
+
+	BEGIN_SHADER_PARAMETER_STRUCT(FParameters, )
+		SHADER_PARAMETER_RDG_BUFFER_UAV(RWBuffer<uint>, RWIndirectDispatchArgsBuffer)
+		SHADER_PARAMETER_RDG_BUFFER_SRV(Buffer<uint>, SourceType)
+		SHADER_PARAMETER(FIntPoint, NetworkTextureSize)
+	END_SHADER_PARAMETER_STRUCT()
+
+	static bool ShouldCompilePermutation(const FGlobalShaderPermutationParameters& Parameters);
+};
+
+class FNeuralPostProcessingPrepareInputCS : public FGlobalShader
+{
+	DECLARE_GLOBAL_SHADER(FNeuralPostProcessingPrepareInputCS)
+	SHADER_USE_PARAMETER_STRUCT(FNeuralPostProcessingPrepareInputCS, FGlobalShader)
+
+public:
+
+	BEGIN_SHADER_PARAMETER_STRUCT(FParameters, )
+		SHADER_PARAMETER_STRUCT(FNueralPostProcessInput, Input0)
+		SHADER_PARAMETER_SAMPLER(SamplerState, InputTextureSampler)
+		RDG_BUFFER_ACCESS(IndirectDispatchArgsBuffer, ERHIAccess::IndirectArgs)
+		SHADER_PARAMETER(int32, InputBufferWidth)
+		SHADER_PARAMETER(int32, InputBufferHeight)
+		SHADER_PARAMETER_RDG_BUFFER_UAV(RWBuffer<float>, InputBuffer)
+		SHADER_PARAMETER(float, ColorScale)
+	END_SHADER_PARAMETER_STRUCT()
+
+	static void ModifyCompilationEnvironment(const FGlobalShaderPermutationParameters& InParameters, FShaderCompilerEnvironment& OutEnvironment);
+};
+
+class FNeuralPostProcessingProcessOutputPS : public FGlobalShader
+{
+	DECLARE_GLOBAL_SHADER(FNeuralPostProcessingProcessOutputPS)
+	SHADER_USE_PARAMETER_STRUCT(FNeuralPostProcessingProcessOutputPS, FGlobalShader)
+
+public:
+
+	BEGIN_SHADER_PARAMETER_STRUCT(FParameters, )
+		SHADER_PARAMETER_STRUCT(FNueralPostProcessInput, Input0)
+		SHADER_PARAMETER(int32, OutputBufferWidth)
+		SHADER_PARAMETER(int32, OutputBufferHeight)
+		SHADER_PARAMETER_RDG_BUFFER_UAV(RWBuffer<float>, OutputBuffer)
+		SHADER_PARAMETER(float, ColorScale)
+		RENDER_TARGET_BINDING_SLOTS()
+	END_SHADER_PARAMETER_STRUCT()
+
+	static void ModifyCompilationEnvironment(const FGlobalShaderPermutationParameters& InParameters, FShaderCompilerEnvironment& OutEnvironment);
+};
diff --git a/Engine/Plugins/Experimental/NeuralRendering/Source/NeuralPostProcessing/Public/NeuralPostProcessing.h b/Engine/Plugins/Experimental/NeuralRendering/Source/NeuralPostProcessing/Public/NeuralPostProcessing.h
new file mode 100644
--- /dev/null
+++ b/Engine/Plugins/Experimental/NeuralRendering/Source/NeuralPostProcessing/Public/NeuralPostProcessing.h
@@ -0,0 +1,15 @@
+// Copyright Epic Games, Inc. All Rights Reserved.
+
+#pragma once
+
+#include "CoreMinimal.h"
+#include "Modules/ModuleManager.h"
+
+class FNeuralPostProcessingModule : public IModuleInterface
+{
+public:
+
+	/** IModuleInterface implementation */
+	virtual void StartupModule() override;
+	virtual void ShutdownModule() override;
+};
diff --git a/Engine/Plugins/Experimental/NeuralRendering/Source/NeuralPostProcessing/Shaders/NeuralPostProcessing.usf b/Engine/Plugins/Experimental/NeuralRendering/Source/NeuralPostProcessing/Shaders/NeuralPostProcessing.usf
new file mode 100644
--- /dev/null
+++ b/Engine/Plugins/Experimental/NeuralRendering/Source/NeuralPostProcessing/Shaders/NeuralPostProcessing.usf
@@ -0,0 +1,140 @@
+
+
+#include "/Engine/Public/Platform.ush"
+#include "/Engine/Private/Common.ush"
+#include "/Engine/Private/ScreenPass.ush"
+
+#define SOURCE_TYPE_TEXTURE		0
+#define SOURCE_TYPE_BUFFER		1
+
+Texture2D Input0_Texture;
+SCREEN_PASS_TEXTURE_VIEWPORT(Input0)
+
+float2 ConvertGridPos2UV(uint2 GridPosition)
+{
+	float2 GripPositionF = float2(GridPosition);
+	return Input0_ExtentInverse * (GripPositionF + 0.5f);
+}
+
+Buffer<uint> SourceType;
+RWBuffer<uint> RWIndirectDispatchArgsBuffer;
+int2 NetworkTextureSize;
+
+[numthreads(8, 1, 1)]
+void BuildIndirectDispatchArgsCS(uint2 DispatchThreadId : SV_DispatchThreadID)
+{
+	if (DispatchThreadId.y == 0 && DispatchThreadId.x == 0)
+	{
+		if (SourceType[0] == SOURCE_TYPE_TEXTURE)
+		{
+			uint IndirectDispatchX = (NetworkTextureSize.x + THREAD_GROUP_SIZE - 1) / THREAD_GROUP_SIZE;
+			uint IndirectDispatchY = (NetworkTextureSize.y + THREAD_GROUP_SIZE - 1) / THREAD_GROUP_SIZE;
+
+			WriteDispatchIndirectArgs(RWIndirectDispatchArgsBuffer, 0, IndirectDispatchX, IndirectDispatchY, 1);
+		}
+		else
+		{
+			WriteDispatchIndirectArgs(RWIndirectDispatchArgsBuffer, 0, 0, 0, 1);
+		}
+	}
+}
+
+SamplerState InputTextureSampler;
+int InputBufferWidth;
+int InputBufferHeight;
+RWBuffer<float> InputBuffer;
+float ColorScale;
+
+// from buffer to neural input
+[numthreads(THREAD_GROUP_SIZE, THREAD_GROUP_SIZE, 1)]
+void PrepareInput(in const uint3 DispatchThreadID : SV_DispatchThreadID)
+{
+	if (DispatchThreadID.x >= InputBufferWidth || DispatchThreadID.y >= InputBufferHeight)
+	{
+		return;
+	}
+
+	// calculate the grid size mapping from the rendered output to the input of the neuralnetwork
+	const float2 BufferSize = float2(InputBufferWidth, InputBufferHeight) - 0.5f;
+	const float2 WidthFactors = Input0_ViewportSize.xy / BufferSize;
+	const int2 LoXY = Input0_ViewportMin.xy + (int2)(WidthFactors * (float2)DispatchThreadID.xy);
+	const int2 HiXY = Input0_ViewportMin.xy + (int2)ceil(WidthFactors * (float2)(DispatchThreadID.xy + 1));
+	int LoX = LoXY.x;
+	int LoY = LoXY.y;
+	int HiX = HiXY.x;
+	int HiY = HiXY.y;
+
+	// averaging each grid
+	float4 Result = float4(0.0, 0.0, 0.0, 0.0);
+	float Div = 0.0;
+	for (int x = LoX; x <= HiX; x++)
+	{
+		if (x < Input0_Extent.x)
+		{
+			for (int y = LoY; y <= HiY; y++)
+			{
+				if (y < Input0_Extent.y)
+				{
+					float2 BufferUV = ConvertGridPos2UV(uint2(x,y));
+					float2 ClampedBufferUV = clamp(BufferUV, Input0_UVViewportBilinearMin, Input0_UVViewportBilinearMax);
+					Result += Input0_Texture.SampleLevel(InputTextureSampler, ClampedBufferUV, 0);
+					Div += 1.0;
+				}
+			}
+		}
+	}
+	Result /= Div;
+
+	// write to the input buffer of the neural network
+	int Idx = DispatchThreadID.y * InputBufferWidth + DispatchThreadID.x;
+	int Offset1 = InputBufferWidth * InputBufferHeight;
+	int Offset2 = Offset1 + Offset1;
+
+	InputBuffer[Idx] = Result.r;
+	InputBuffer[Idx + Offset1] = Result.g;
+	InputBuffer[Idx + Offset2] = Result.b;
+}
+
+int OutputBufferWidth;
+int OutputBufferHeight;
+RWBuffer<float> OutputBuffer;
+
+// from neural output to buffer, full screenpass pixel shader
+float4 ProcessOutput(float4 Position : SV_POSITION) : SV_Target0
+{
+	float X = Position.x * float(OutputBufferWidth - 1) / (float)Input0_ViewportSize.x;
+	float Y = Position.y * float(OutputBufferHeight - 1) / (float)Input0_ViewportSize.y;
+
+	int LoX =  (int)clamp(floor(X), 0.0, float(OutputBufferWidth - 1));
+	int HiX = (int)clamp(floor(X + 1.0), 0.0, float(OutputBufferWidth - 1));
+
+	int LoY = (int)clamp(floor(Y), 0.0, float(OutputBufferHeight - 1));
+	int HiY = (int)clamp(floor(Y + 1.0), 0.0, float(OutputBufferHeight - 1));
+
+	int Offset1 = OutputBufferWidth * OutputBufferHeight;
+	int Offset2 = Offset1 + Offset1;
+
+	// bilinear filtering
+	int Idx;
+
+	Idx = LoY * OutputBufferWidth + LoX;
+	float3 A = float3(OutputBuffer[Idx], OutputBuffer[Idx + Offset1], OutputBuffer[Idx + Offset2]);
+
+	Idx = LoY * OutputBufferWidth + HiX;
+	float3 B = float3(OutputBuffer[Idx], OutputBuffer[Idx + Offset1], OutputBuffer[Idx + Offset2]);
+
+	Idx = HiY * OutputBufferWidth + LoX;
+	float3 C = float3(OutputBuffer[Idx], OutputBuffer[Idx + Offset1], OutputBuffer[Idx + Offset2]);
+
+	Idx = HiY * OutputBufferWidth + HiX;
+	float3 D = float3(OutputBuffer[Idx], OutputBuffer[Idx + Offset1], OutputBuffer[Idx + Offset2]);
+
+	float Alpha = clamp(X - (float)LoX, 0.0, 1.0);
+	float Beta = clamp(Y - (float)LoY, 0.0, 1.0);
+
+	float3 Result = (1.0 - Beta) * ((1.0 - Alpha) * A + Alpha * B) + Beta * ((1.0 - Alpha) * C + Alpha * D);
+
+	Result = Result.rgb;
+
+	return float4(Result, 1.0);
+}
\ No newline at end of file
diff --git a/Engine/Shaders/Private/MaterialTemplate.ush b/Engine/Shaders/Private/MaterialTemplate.ush
--- a/Engine/Shaders/Private/MaterialTemplate.ush
+++ b/Engine/Shaders/Private/MaterialTemplate.ush
@@ -2710,6 +2710,37 @@ float4 MaterialExpressionPathTracingBufferTextureLookup(FMaterialPixelParameters
 	return float4(1, 0, 0, 1);
 }
 
+float4 NeuralTextureOutput(FMaterialPixelParameters Parameters, float2 ViewportUV)
+{
+#if POST_PROCESS_MATERIAL && SHADING_PATH_DEFERRED && MATERIAL_NEURAL_POST_PROCESS
+	
+	float2 BufferUV = ViewportUV * PostProcessInput_0_UVViewportSize + PostProcessInput_0_UVViewportMin;
+	float4 TextureValue = Texture2DSampleLevel(NeuralTexture, PostProcessInput_BilinearSampler, BufferUV, 0);
+	return TextureValue;
+#else
+	return float4(0.0f, 0.0f, 0.0f, 1.0f);
+#endif
+
+}
+
+float4 NeuralBufferOutput(FMaterialPixelParameters Parameters, float4 BufferIndices)
+{
+#if POST_PROCESS_MATERIAL && SHADING_PATH_DEFERRED && MATERIAL_NEURAL_POST_PROCESS
+
+	float3 TextureValue = 0;
+	
+	int Batch = (int)BufferIndices.x;
+	int StartChannel = (int)BufferIndices.y;
+	float2 BufferWH = BufferIndices.zw;
+
+	ReadFromNeuralBuffer(Batch, StartChannel, BufferWH, TextureValue);
+
+	return float4(TextureValue, 1.0f);
+#else
+	return float4(0.0f, 0.0f, 0.0f, 1.0f);
+#endif
+}
+
 // DERIV_BASE_VALUE() is a disgusting macro to manage backwards compatibility while changing the generated materials. The existing nodes are not "aware"
 // of partial derivatives, and will use sprintf() to put the same line in both CalcPixelMaterialInputs() and CalcPixelMaterialInputsAnalyticDerivatives().
 // If we have a line like this in CalcPixelMaterialInputs():
diff --git a/Engine/Shaders/Private/NeuralPostProcessCommon.usf b/Engine/Shaders/Private/NeuralPostProcessCommon.usf
new file mode 100644
--- /dev/null
+++ b/Engine/Shaders/Private/NeuralPostProcessCommon.usf
@@ -0,0 +1,98 @@
+// Copyright Epic Games, Inc. All Rights Reserved.
+
+/*=============================================================================
+	NeuralPostProcessCommon.usf: Shaders for neural post process
+=============================================================================*/
+
+#ifndef MATERIAL_NEURAL_POST_PROCESS
+#define MATERIAL_NEURAL_POST_PROCESS 0
+#endif
+
+#ifndef NEURAL_POSTPROCESS_PREPASS
+#define NEURAL_POSTPROCESS_PREPASS 0
+#endif
+
+#if MATERIAL_NEURAL_POST_PROCESS
+
+// Input buffer/textue set by the prepass
+RWTexture2D<float4> RWNeuralTexture;
+RWBuffer<float> InputNeuralBuffer;
+RWBuffer<uint> NeuralSourceType;
+float4 InputNeuralBufferDimension;
+
+// Buffer/textures read by the main pass
+Texture2D<float4> NeuralTexture;
+Buffer<float> OutputNeuralBuffer;
+float4 OutputNeuralBufferDimension;
+
+#define SOURCE_TYPE_TEXTURE 0
+#define SOURCE_TYPE_BUFFER  1
+
+#endif
+
+void UpdateNeuralProfileSourceType(int InputSourceType)
+{
+#if NEURAL_POSTPROCESS_PREPASS
+	NeuralSourceType[0] = InputSourceType;
+#endif
+}
+
+
+// Save a 3 channel data to OutputNeuralBuffer from StartChannel to StartChannel+2 
+void SaveToNeuralBuffer(int Batch, int StartChannel, float2 BufferWH, float3 Value)
+{
+#if NEURAL_POSTPROCESS_PREPASS
+	float2 HWIndex = BufferWH.yx * InputNeuralBufferDimension.zw - 0.5f;
+
+	// point sampling
+	const int	IX = (int)clamp(HWIndex.y, 0, InputNeuralBufferDimension.w - 1);
+	const int	IY = (int)clamp(HWIndex.x, 0, InputNeuralBufferDimension.z - 1);
+
+	int BufferHeight = InputNeuralBufferDimension.z;
+	int BufferWidth = InputNeuralBufferDimension.w;
+	int BufferGridSize = BufferHeight * BufferWidth;
+	int NumChannels = InputNeuralBufferDimension.y;
+	int BatchSize = NumChannels * BufferGridSize;
+
+	int Idx = (Batch * BatchSize + StartChannel * BufferGridSize + IY * BufferWidth + IX);
+	int Offset = BufferGridSize;
+
+	// TODO: array of struct vs struct of array performance evaluation.
+	InputNeuralBuffer[Idx] = Value[0];
+	const int MaxOffset = min(3, InputNeuralBufferDimension.y - StartChannel) - 1;
+	for (int i = 1; i <= MaxOffset; ++i)
+	{
+		InputNeuralBuffer[Idx + i * Offset] = Value[i];
+	}
+#endif
+
+}
+
+void ReadFromNeuralBuffer(int Batch, int StartChannel, float2 BufferWH, inout float3 Value)
+{
+#if MATERIAL_NEURAL_POST_PROCESS
+	float2 HWIndex = BufferWH.yx * OutputNeuralBufferDimension.zw - 0.5f;
+	
+	// point sampling
+	const int	IX = (int)clamp(HWIndex.y, 0, OutputNeuralBufferDimension.w - 1);
+	const int	IY = (int)clamp(HWIndex.x, 0, OutputNeuralBufferDimension.z - 1);
+
+	int BufferHeight = OutputNeuralBufferDimension.z;
+	int BufferWidth = OutputNeuralBufferDimension.w;
+	int BufferGridSize = BufferHeight * BufferWidth;
+	int NumChannels = OutputNeuralBufferDimension.y;
+	int BatchSize = NumChannels * BufferGridSize;
+
+	int Idx = (Batch * BatchSize + StartChannel * BufferGridSize + IY * BufferWidth + IX);
+	int Offset = BufferGridSize;
+
+	// TODO: array of struct vs struct of array performance evaluation.
+	const int MaxOffset = min(3, OutputNeuralBufferDimension.y- StartChannel) - 1;
+	for (int i = 0; i <= MaxOffset; ++i)
+	{
+		Value[i] = OutputNeuralBuffer[Idx + i * Offset];
+	}
+#else
+	Value = 0;
+#endif
+}
\ No newline at end of file
diff --git a/Engine/Shaders/Private/PostProcessMaterialShaders.usf b/Engine/Shaders/Private/PostProcessMaterialShaders.usf
--- a/Engine/Shaders/Private/PostProcessMaterialShaders.usf
+++ b/Engine/Shaders/Private/PostProcessMaterialShaders.usf
@@ -76,6 +76,8 @@ SamplerState PathTracingPostProcessInput_3_Sampler;
 SamplerState PathTracingPostProcessInput_4_Sampler;
 #endif
 
+#include "NeuralPostProcessCommon.usf"
+
 #define EYE_ADAPTATION_LOOSE_PARAMETERS 1
 
 #include "/Engine/Generated/Material.ush"
@@ -242,6 +244,60 @@ void MainPS(
 	// fill out other related material parameters
 	CalcMaterialParametersPost(Parameters, PixelMaterialInputs, SvPosition, true);
 
+#if NEURAL_POSTPROCESS_PREPASS && MATERIAL_NEURAL_POST_PROCESS
+
+	OutColor = 0.5f;
+
+#if HAVE_GetNeuralInput1
+
+	float4 Index = GetNeuralInput0(Parameters);
+	float3 NeuralInput0 = GetNeuralInput1(Parameters).rgb;
+	float  Mask = GetNeuralInput2(Parameters);
+	const bool bWriteToBuffer = Index.x >= 0; //Configuration the same across all pixels.
+
+	// Update the source type so we can copy on demand
+	BRANCH
+	if (all((uint2)SvPosition.xy == uint2(0, 0)))
+	{
+		const int SourceType = bWriteToBuffer ? SOURCE_TYPE_BUFFER : SOURCE_TYPE_TEXTURE;
+		UpdateNeuralProfileSourceType(SourceType);
+	}
+
+	// Write to buffer/texture based on config
+	BRANCH
+	if (Mask != 0)
+	{
+		BRANCH
+		if (bWriteToBuffer)
+		{
+			int Batch = (int)Index.x;
+			int StartChannel = (int)Index.y;
+			float2 BufferWH = Index.zw;
+			SaveToNeuralBuffer(Batch, StartChannel, BufferWH, NeuralInput0);
+		}
+		else
+		{
+			int StartChannel = (int)Index.y;
+			float2 CustomViewportUV = Index.zw;
+			float2 Position = CustomViewportUV * PostProcessOutput_ViewportSize.xy + PostProcessOutput_ViewportMin.xy;
+			if (StartChannel >= 0 && StartChannel <= 2)
+			{
+				RWNeuralTexture[(uint2)Position.xy] = float4(NeuralInput0, 0.0f);
+			}
+		}
+	}
+	
+	// This is for debug only
+	OutColor.xyz = NeuralInput0;
+#if !MATERIALBLENDING_MODULATE && POST_PROCESS_MATERIAL_BEFORE_TONEMAP
+	OutColor.xyz *= View.PreExposure;
+#endif
+
+	return;
+
+#endif // HAVE_GetNeuralInput1
+#endif // NEURAL_POSTPROCESS_PREPASS	
+
 	// In Unreal, an alpha channel value of 0.0f means that the pixel is opaque and a value of 1.0f means that there is translucency.
 	// This is the opposite of how the alpha channel is interpreted in formats like PNG. For this reason, rendered images that are exported
 	// from the engine will have their alpha channel inverted to match the PNG convention. However, for consistency inside the rendering chain,
diff --git a/Engine/Source/Editor/UnrealEd/Classes/Factories/NeuralProfileFactory.h b/Engine/Source/Editor/UnrealEd/Classes/Factories/NeuralProfileFactory.h
new file mode 100644
--- /dev/null
+++ b/Engine/Source/Editor/UnrealEd/Classes/Factories/NeuralProfileFactory.h
@@ -0,0 +1,20 @@
+// Copyright Epic Games, Inc. All Rights Reserved.
+
+#pragma once
+
+#include "CoreMinimal.h"
+#include "UObject/ObjectMacros.h"
+#include "Factories/Factory.h"
+#include "NeuralProfileFactory.generated.h"
+
+UCLASS(hidecategories = Object)
+class UNeuralProfileFactory : public UFactory
+{
+	GENERATED_UCLASS_BODY()
+
+	//~ Begin UFactory Interface
+	virtual UObject* FactoryCreateNew(UClass* Class, UObject* InParent, FName Name, EObjectFlags Flags, UObject* Context, FFeedbackContext* Warn) override;
+	//~ Begin UFactory Interface
+};
+
+
diff --git a/Engine/Source/Editor/UnrealEd/Classes/ThumbnailRendering/NeuralProfileRenderer.h b/Engine/Source/Editor/UnrealEd/Classes/ThumbnailRendering/NeuralProfileRenderer.h
new file mode 100644
--- /dev/null
+++ b/Engine/Source/Editor/UnrealEd/Classes/ThumbnailRendering/NeuralProfileRenderer.h
@@ -0,0 +1,29 @@
+// Copyright Epic Games, Inc. All Rights Reserved.
+
+/**
+ *
+ * This thumbnail renderer displays the texture for the object in question
+ */
+
+#pragma once
+
+#include "CoreMinimal.h"
+#include "UObject/ObjectMacros.h"
+#include "ThumbnailRendering/TextureThumbnailRenderer.h"
+#include "NeuralProfileRenderer.generated.h"
+
+class FCanvas;
+class FRenderTarget;
+
+UCLASS()
+class UNeuralProfileRenderer : public UTextureThumbnailRenderer
+{
+	GENERATED_UCLASS_BODY()
+
+
+	// Begin UThumbnailRenderer Object
+	virtual void GetThumbnailSize(UObject* Object, float Zoom, uint32& OutWidth, uint32& OutHeight) const override;
+	virtual void Draw(UObject* Object, int32 X, int32 Y, uint32 Width, uint32 Height, FRenderTarget*, FCanvas* Canvas, bool bAdditionalViewFamily) override;
+	// End UThumbnailRenderer Object
+};
+
diff --git a/Engine/Source/Editor/UnrealEd/Private/Factories/EditorFactories.cpp b/Engine/Source/Editor/UnrealEd/Private/Factories/EditorFactories.cpp
--- a/Engine/Source/Editor/UnrealEd/Private/Factories/EditorFactories.cpp
+++ b/Engine/Source/Editor/UnrealEd/Private/Factories/EditorFactories.cpp
@@ -111,6 +111,7 @@
 #include "Factories/StringTableFactory.h"
 #include "Factories/SubsurfaceProfileFactory.h"
 #include "Factories/SpecularProfileFactory.h"
+#include "Factories/NeuralProfileFactory.h"
 #include "Factories/Texture2dFactoryNew.h"
 #include "Engine/Texture.h"
 #include "Factories/TextureFactory.h"
@@ -147,6 +148,7 @@
 #include "GameFramework/DefaultPhysicsVolume.h"
 #include "Engine/SubsurfaceProfile.h"
 #include "Engine/SpecularProfile.h"
+#include "Engine/NeuralProfile.h"
 #include "Misc/ConfigCacheIni.h"
 #include "Misc/FeedbackContext.h"
 #include "GameFramework/WorldSettings.h"
@@ -7919,6 +7921,25 @@ UObject* USpecularProfileFactory::FactoryCreateNew(UClass* InClass, UObject* InP
 	return Object;
 }
 
+/*-----------------------------------------------------------------------------
+	UNeuralProfileFactory implementation.
+	-----------------------------------------------------------------------------*/
+UNeuralProfileFactory::UNeuralProfileFactory(const FObjectInitializer& ObjectInitializer)
+	: Super(ObjectInitializer)
+{
+	SupportedClass = UNeuralProfile::StaticClass();
+	bCreateNew = true;
+	bEditorImport = false;
+	bEditAfterNew = true;
+}
+
+UObject* UNeuralProfileFactory::FactoryCreateNew(UClass* Class, UObject* InParent, FName InName, EObjectFlags Flags, UObject* Context, FFeedbackContext* Warn)
+{
+	UNeuralProfile* Object = NewObject<UNeuralProfile>(InParent, InName, Flags);
+	Object->Guid = FGuid::NewGuid();
+	return Object;
+}
+
 /*-----------------------------------------------------------------------------
 	UTouchInterfaceFactory implementation.
 -----------------------------------------------------------------------------*/
diff --git a/Engine/Source/Editor/UnrealEd/Private/ThumbnailRendering/NeuralProfileRenderer.cpp b/Engine/Source/Editor/UnrealEd/Private/ThumbnailRendering/NeuralProfileRenderer.cpp
new file mode 100644
--- /dev/null
+++ b/Engine/Source/Editor/UnrealEd/Private/ThumbnailRendering/NeuralProfileRenderer.cpp
@@ -0,0 +1,44 @@
+// Copyright Epic Games, Inc. All Rights Reserved.
+
+#include "ThumbnailRendering/NeuralProfileRenderer.h"
+#include "CanvasItem.h"
+#include "Engine/Engine.h"
+#include "EngineGlobals.h"
+#include "Engine/NeuralProfile.h"
+#include "CanvasTypes.h"
+
+UNeuralProfileRenderer::UNeuralProfileRenderer(const FObjectInitializer& ObjectInitializer)
+	: Super(ObjectInitializer)
+{
+}
+
+void UNeuralProfileRenderer::GetThumbnailSize(UObject* Object, float Zoom, uint32& OutWidth, uint32& OutHeight) const
+{
+	OutWidth = 128;
+	OutHeight = 128;
+}
+
+void UNeuralProfileRenderer::Draw(UObject* Object, int32 X, int32 Y, uint32 Width, uint32 Height, FRenderTarget*, FCanvas* Canvas, bool bAdditionalViewFamily)
+{
+	UNeuralProfile* LocalNeuralProfile = Cast<UNeuralProfile>(Object);
+	if (LocalNeuralProfile)
+	{
+		const bool bHasProfile = LocalNeuralProfile->Settings.NNEModelData != nullptr;
+		FLinearColor Col;
+
+		Col = bHasProfile? FLinearColor(0.0f,1.0f,0.0f):FLinearColor(0.0f,0.0f,0.0f); Col.A = 1;
+		Canvas->DrawTile(0,          0, Width, Height / 2, 0, 0, 1, 1, Col);
+		Col = FLinearColor(0.0f, 0.0f, 0.0f); Col.A = 1;
+		Canvas->DrawTile(0, Height / 2, Width, Height / 2, 0, 0, 1, 1, Col);
+
+		FText ChannelText = FText::AsNumber(LocalNeuralProfile->Settings.InputDimension.Y);
+		//FText DimenionHeight = FText::AsNumber(LocalNeuralProfile->Settings.Dimension.Z);
+		//FText DimenionWidth = FText::AsNumber(LocalNeuralProfile->Settings.Dimension.W);
+		
+		FCanvasTextItem TextItem(FVector2D(5.0f, 5.0f), ChannelText, GEngine->GetLargeFont(), FLinearColor::White);
+		TextItem.EnableShadow(FLinearColor::Black);
+		TextItem.Scale = FVector2D(Width / 128.0f, Height / 128.0f);
+		TextItem.Draw(Canvas);
+
+	}
+}
diff --git a/Engine/Source/Editor/UnrealEd/Public/UnrealEd.h b/Engine/Source/Editor/UnrealEd/Public/UnrealEd.h
--- a/Engine/Source/Editor/UnrealEd/Public/UnrealEd.h
+++ b/Engine/Source/Editor/UnrealEd/Public/UnrealEd.h
@@ -269,6 +269,7 @@ MONOLITHIC_HEADER_BOILERPLATE()
 #include "ThumbnailRendering/ParticleSystemThumbnailRenderer.h"
 #include "ThumbnailRendering/SubsurfaceProfileRenderer.h"
 #include "ThumbnailRendering/SpecularProfileRenderer.h"
+#include "ThumbnailRendering/NeuralProfileRenderer.h"
 #include "Factories/TextureCubeThumbnailRenderer.h"
 #include "Factories/Texture2dArrayThumbnailRenderer.h"
 #include "Editor/TransBuffer.h"
diff --git a/Engine/Source/Runtime/Engine/Classes/Engine/NeuralProfile.h b/Engine/Source/Runtime/Engine/Classes/Engine/NeuralProfile.h
new file mode 100644
--- /dev/null
+++ b/Engine/Source/Runtime/Engine/Classes/Engine/NeuralProfile.h
@@ -0,0 +1,153 @@
+// Copyright Epic Games, Inc. All Rights Reserved.
+#pragma once
+
+#include "CoreMinimal.h"
+#include "UObject/ObjectMacros.h"
+#include "UObject/Object.h"
+#include "Misc/Guid.h"
+#include "NeuralProfile.generated.h"
+
+#define MAX_NEURAL_PROFILE_COUNT 64
+
+class UNNEModelData;
+
+UENUM(BlueprintType)
+enum class ENeuralProfileFormat : uint8
+{
+	Type32 UMETA(DisplayName = "32bit"),
+	Type16 UMETA(DisplayName = "16bit"),
+};
+
+UENUM(BlueprintType)
+enum class ENeuralModelTileType : uint8
+{
+	/** The NNE model is loaded and used as it is. No dimension augmentation. E.g.,
+if the input texture has different dimensions, it will be scaled down before application*/
+	OneByOne		UMETA(DisplayName = "1x1"),
+
+	TwoByTwo		UMETA(DisplayName = "2x2"),
+	
+	FourByFour		UMETA(DisplayName = "4x4"),
+	
+	EightByEight	UMETA(DisplayName = "8x8"),
+
+	/** Create tiled buffers in batch dimension automatically, where each tile runs the neural model
+	e.g., if the model input dimension is (1x3x200x200) and the used buffer size of the post processing
+	is 1000x1000, then 5x5 tiles ((5x5)x3x200x200) will be run and recombined.
+	*/
+	//Auto			UMETA(DisplayName = "Auto")
+};
+
+// struct with all the settings we want in UNeuralProfile, separate to make it easer to pass this data around in the engine.
+USTRUCT(BlueprintType)
+struct FNeuralProfileStruct
+{
+	GENERATED_USTRUCT_BODY()
+	
+	/**
+	 * Define the expected input format, if any output from material is not this format, a custom conversion
+	 * will be applied for this conversion.
+	 */
+	UPROPERTY(Category = "Common", EditAnywhere, BlueprintReadOnly, meta=(DisplayName="Input Format", editcondition = "false", EditConditionHides))
+	ENeuralProfileFormat InputFormat;
+
+	/**
+	 * Define the expected output format. A conversion between the output format and the actual format will
+	 * be applied automatically.
+	 */
+	UPROPERTY(Category = "Common", EditAnywhere, BlueprintReadOnly, meta = (DisplayName = "Output Format", editcondition = "false", EditConditionHides))
+	ENeuralProfileFormat OutputFormat;
+
+	//runtime name (support "NNERuntimeRDGDml" only at this moment)
+	UPROPERTY(Category = "Model", EditAnywhere, BlueprintReadOnly)
+	FString RuntimeName;
+
+	/** Stores the NNEModelData imported from e.g., onnx model */
+	UPROPERTY(EditAnywhere, BlueprintReadOnly, Category = Model, meta = (DisplayName = "NNE Model Data"))
+	TObjectPtr<UNNEModelData> NNEModelData;
+
+	/** Input dimension of the NNEModelData model */
+	UPROPERTY(EditAnywhere, BlueprintReadOnly, Category = Model, meta = (DisplayName = "Input Dimension", editcondition = "false"))
+	FIntVector4 InputDimension;
+
+	/** Output dimension of the NNEModelData model */
+	UPROPERTY(EditAnywhere, BlueprintReadOnly, Category = Model, meta = (DisplayName = "Output Dimension", editcondition = "false"))
+	FIntVector4 OutputDimension;
+
+	/** Except for 1x1, all other model requires an intermediate buffer.*/
+	UPROPERTY(EditAnywhere, BlueprintReadOnly, Category = "Quality and Performance")
+	ENeuralModelTileType TileSize;
+
+	/** Used to override the batch size if the batch dimension is dynamic (-1)*/
+	UPROPERTY(EditAnywhere, BlueprintReadOnly, Category = "Quality and Performance", meta = (EditCondition = "InputDimension[0] < 0"))
+	int32 BatchSizeOverride;
+
+	FNeuralProfileStruct()
+	{
+		InputFormat = ENeuralProfileFormat::Type32;
+		OutputFormat = ENeuralProfileFormat::Type32;
+		RuntimeName = "NNERuntimeRDGDml";
+		NNEModelData = nullptr;
+		InputDimension = FIntVector4(0);
+		OutputDimension = FIntVector4(0);
+		TileSize = ENeuralModelTileType::OneByOne;
+		BatchSizeOverride = 1;
+	}
+
+	void Invalidate()
+	{
+		*this = FNeuralProfileStruct();
+	}
+};
+
+/**
+ * This class represents assets that stores the neural network model data and the conversion from/to the model.
+ *
+ * Neural network model profile typically consist of a neural network model data and the config of preprocessing/post-processing
+ * of applying this neural network model.
+ */
+UCLASS(autoexpandcategories = NeuralProfile, MinimalAPI)
+class UNeuralProfile : public UObject
+{
+	GENERATED_UCLASS_BODY()
+
+public:
+
+	UPROPERTY(Category = UNeuralProfile, EditAnywhere, meta = (ShowOnlyInnerProperties))
+	struct FNeuralProfileStruct Settings;
+
+	UPROPERTY()
+	FGuid Guid;
+
+	void BeginDestroy();
+	virtual void PostEditChangeProperty(struct FPropertyChangedEvent& PropertyChangedEvent);
+};
+
+namespace NeuralProfile
+{
+	// Profile - Initializes or updates the contents of the neural profile.
+	ENGINE_API int32 AddOrUpdateProfile(const UNeuralProfile* InProfile, const FGuid& InGuid, const FNeuralProfileStruct InSettings);
+
+	ENGINE_API FNeuralProfileStruct GetProfileSetting(int32 AllocationId);
+
+	// Profile - Returns the profile ID for a given neural profile object
+	ENGINE_API int32 GetNeuralProfileId(const UNeuralProfile* In);
+
+	// Experimental APIs to manage neural profiles.
+	class INeuralProfileManager
+	{
+	public:
+		virtual ~INeuralProfileManager() {}
+
+		virtual void UpdateModel(int32 AllocationId, UObject* NNEModelData, FString RuntimeName) = 0;
+		virtual void RemoveModel(int32 AllocationId) = 0;
+
+		virtual void UpdateTileType(int32 AllocationId, ENeuralModelTileType ModelTileSize) = 0;
+		virtual bool UpdateBatchSize(int32 AllocationId, int32 BatchSize) = 0;
+	
+		virtual FIntVector4 GetInputDimension(UObject* NNEModelData, FString RuntimeName) = 0;
+		virtual FIntVector4 GetOutputDimension(UObject* NNEModelData, FString RuntimeName) = 0;
+	};
+}
+
+extern ENGINE_API TUniquePtr<NeuralProfile::INeuralProfileManager> GNeuralProfileManager;
diff --git a/Engine/Source/Runtime/Engine/Classes/Materials/Material.h b/Engine/Source/Runtime/Engine/Classes/Materials/Material.h
--- a/Engine/Source/Runtime/Engine/Classes/Materials/Material.h
+++ b/Engine/Source/Runtime/Engine/Classes/Materials/Material.h
@@ -972,6 +972,15 @@ public:
 	UPROPERTY(EditAnywhere, BlueprintReadWrite, Category = PostProcessMaterial, meta = (DisplayName = "Output Alpha"))
 	uint8 BlendableOutputAlpha : 1;
 
+	/**
+	* Indicates that the material and its instances can be used with neural network engine.
+	* This will result in the shaders required to support neural network engine being compiled which will increase shader compile time and memory usage.
+	* In addition, an additional pass will run before the postprocess pass for neural network engine, which will increase the rendering cost due to
+	* buffer copy and inference.
+	*/
+	UPROPERTY(EditAnywhere, Category = PostProcessMaterial)
+	uint8 bUsedWithNeuralNetworks : 1;
+
 	/** 
 	 * Selectively execute post process material only for pixels that pass the stencil test against the Custom Depth/Stencil buffer. 
 	 * Pixels that fail the stencil test are filled with the previous post process material output or scene color.
@@ -985,6 +994,12 @@ public:
 	UPROPERTY(EditAnywhere, Category = PostProcessMaterial, AdvancedDisplay, meta = (EditCondition = "bEnableStencilTest"))
 	uint8 StencilRefValue = 0;
 
+	/**
+	* Set by reference object cannot be modified.
+	*/
+	UPROPERTY(EditAnywhere, Category = PostProcessMaterial, AdvancedDisplay, meta = (EditCondition = "false"))
+	int8 NeuralProfileId = INDEX_NONE;
+
 	UPROPERTY()
 	TEnumAsByte<ERefractionMode> RefractionMode_DEPRECATED;
 
@@ -1140,6 +1155,7 @@ public:
 	ENGINE_API virtual USubsurfaceProfile* GetSubsurfaceProfile_Internal() const override;
 	ENGINE_API virtual uint32 NumSpecularProfile_Internal() const override;
 	ENGINE_API virtual USpecularProfile* GetSpecularProfile_Internal(uint32 Index) const override;
+	ENGINE_API virtual UNeuralProfile* GetNeuralProfile_Internal() const override;
 	ENGINE_API virtual bool CastsRayTracedShadows() const override;
 	ENGINE_API virtual FDisplacementScaling GetDisplacementScaling() const override;
 	ENGINE_API virtual float GetMaxWorldPositionOffsetDisplacement() const override;
diff --git a/Engine/Source/Runtime/Engine/Classes/Materials/MaterialExpressionNeuralPostProcessNode.h b/Engine/Source/Runtime/Engine/Classes/Materials/MaterialExpressionNeuralPostProcessNode.h
new file mode 100644
--- /dev/null
+++ b/Engine/Source/Runtime/Engine/Classes/Materials/MaterialExpressionNeuralPostProcessNode.h
@@ -0,0 +1,103 @@
+// Copyright Epic Games, Inc. All Rights Reserved.
+
+/**
+ * Node acts as a base class for TextureSamples and TextureObjects 
+ * to cover their shared functionality. 
+ */
+
+#pragma once
+
+#include "CoreMinimal.h"
+#include "UObject/ObjectMacros.h"
+#include "Materials/MaterialExpressionCustomOutput.h"
+#include "MaterialExpressionNeuralPostProcessNode.generated.h"
+
+struct FPropertyChangedEvent;
+enum EMaterialSamplerType : int;
+
+UENUM()
+enum ENeuralIndexType : int
+{
+	/** Texture index UV in 0...1 range*/
+	NIT_TextureIndex UMETA(DisplayName = "Texture Index"),
+
+	/** Buffer index UV in LayerId (batch) x Channels x Height x Wdith*/
+	NIT_BufferIndex UMETA(DisplayName = "Buffer Index"),
+
+};
+
+UCLASS(MinimalAPI, collapsecategories, hidecategories=Object)
+class UMaterialExpressionNeuralNetworkInput : public UMaterialExpressionCustomOutput 
+{
+	GENERATED_UCLASS_BODY()
+	
+	/**
+	*	Coordinate to read values:
+			Texture index (float4): Batch[ignored], StartChannel[int,0], ViewportUV in 0..1 range
+			Buffer  index (float4): Batch[int,0], StartChannel[int,0], WidthHeight[ViewportUV] in 0..1 range
+	*/
+	UPROPERTY(meta = (RequiredInput = "false"))
+	FExpressionInput Coordinates;
+
+	/** First input node for neural post processing in post process material*/
+	UPROPERTY()
+	FExpressionInput Input0;
+
+	/** Input is ignored if Mask is zero. 
+		TODO: Used to optimize performance
+	*/
+	UPROPERTY(meta = (RequiredInput = "false"))
+	FExpressionInput Mask;
+
+	/**Indexing type: write to the texture or buffer */
+	UPROPERTY(EditAnywhere, CateGory = "Default", meta = (DisplayName = "Index Type"))
+	TEnumAsByte<ENeuralIndexType> NeuralIndexType;
+	
+#if WITH_EDITOR
+
+	virtual int32 Compile(class FMaterialCompiler* Compiler, int32 OutputIndex) override;
+	virtual void GetCaption(TArray<FString>& OutCaptions) const override;
+	virtual uint32 GetInputType(int32 InputIndex) override;
+
+	virtual UE::Shader::EValueType GetCustomOutputType(int32 OutputIndex) const override;
+	virtual bool GenerateHLSLExpression(FMaterialHLSLGenerator& Generator, UE::HLSLTree::FScope& Scope, int32 OutputIndex, UE::HLSLTree::FExpression const*& OutExpression) const;
+#endif // WITH_EDITOR
+
+	//~ Begin UMaterialExpressionCustomOutput Interface
+	virtual int32 GetNumOutputs() const override;
+	virtual FString GetFunctionName() const override;
+	//~ End UMaterialExpressionCustomOutput Interface
+};
+
+
+// The input node that holds the output from the neural network to the post process material 
+UCLASS(MinimalAPI, collapsecategories, hidecategories = Object)
+class UMaterialExpressionNeuralNetworkOutput : public UMaterialExpression
+{
+	GENERATED_UCLASS_BODY()
+
+	/** 
+	*	Coordinate to store the values:
+			Texture index (float4): Batch[ignored], StartChannel[int,0], ViewportUV in 0..1 range 
+			Buffer  index (float4): Batch[int,0], StartChannel[int,0], WidthHeight[ViewportUV] in 0..1 range
+	*/
+	UPROPERTY(meta = (RequiredInput = "false"))
+	FExpressionInput Coordinates;
+
+	/**Indexing type: read from the texture or buffer. 
+		Texture index is valid only when the first 2 dimension of the input and output dimension matches.
+	*/
+	UPROPERTY(EditAnywhere, Category = "Default", meta = (DisplayName = "Index Type"))
+	TEnumAsByte<ENeuralIndexType> NeuralIndexType;
+
+#if WITH_EDITOR
+	virtual int32 Compile(class FMaterialCompiler* Compiler, int32 OutputIndex) override;
+	virtual void GetCaption(TArray<FString>& OutCaptions) const override;
+	virtual uint32 GetInputType(int32 InputIndex) override;
+
+	virtual bool GenerateHLSLExpression(FMaterialHLSLGenerator& Generator, UE::HLSLTree::FScope& Scope, int32 OutputIndex, UE::HLSLTree::FExpression const*& OutExpression) const;
+#endif // WITH_EDITOR
+
+};
+
+
diff --git a/Engine/Source/Runtime/Engine/Classes/Materials/MaterialInterface.h b/Engine/Source/Runtime/Engine/Classes/Materials/MaterialInterface.h
--- a/Engine/Source/Runtime/Engine/Classes/Materials/MaterialInterface.h
+++ b/Engine/Source/Runtime/Engine/Classes/Materials/MaterialInterface.h
@@ -51,6 +51,7 @@ class UPhysicalMaterial;
 class UPhysicalMaterialMask;
 class USubsurfaceProfile;
 class USpecularProfile;
+class UNeuralProfile;
 class UTexture;
 class UMaterialInstance;
 struct FDebugShaderTypeInfo;
@@ -281,6 +282,10 @@ public:
 	UPROPERTY()
 	TArray<TObjectPtr<class USpecularProfile>> SpecularProfiles;
 
+	/** Neural network profile. For internal usage, not editable/visible */
+	UPROPERTY(EditAnywhere, BlueprintReadOnly, Category = PostProcessMaterial, meta = (DisplayName = "Neural Profile"))
+	TObjectPtr<class UNeuralProfile> NeuralProfile;
+
 	/** Whether this material interface is included in the base game (and not in a DLC) */
 	UPROPERTY()
 	uint8 bIncludedInBaseGame : 1;
@@ -918,6 +923,7 @@ public:
 	ENGINE_API virtual USubsurfaceProfile* GetSubsurfaceProfile_Internal() const;
 	ENGINE_API virtual uint32 NumSpecularProfile_Internal() const;
 	ENGINE_API virtual USpecularProfile* GetSpecularProfile_Internal(uint32 Index) const;
+	ENGINE_API virtual UNeuralProfile* GetNeuralProfile_Internal() const;
 	ENGINE_API virtual bool CastsRayTracedShadows() const;
 
 	/**
diff --git a/Engine/Source/Runtime/Engine/Engine.Build.cs b/Engine/Source/Runtime/Engine/Engine.Build.cs
--- a/Engine/Source/Runtime/Engine/Engine.Build.cs
+++ b/Engine/Source/Runtime/Engine/Engine.Build.cs
@@ -158,6 +158,7 @@ public class Engine : ModuleRules
 				"ColorManagement",
 				"Icmp",
 				"XmlParser",
+				"NNE"
 			}
 		);
 
diff --git a/Engine/Source/Runtime/Engine/Private/Materials/HLSLMaterialTranslator.cpp b/Engine/Source/Runtime/Engine/Private/Materials/HLSLMaterialTranslator.cpp
--- a/Engine/Source/Runtime/Engine/Private/Materials/HLSLMaterialTranslator.cpp
+++ b/Engine/Source/Runtime/Engine/Private/Materials/HLSLMaterialTranslator.cpp
@@ -2130,6 +2130,7 @@ void FHLSLMaterialTranslator::GetMaterialEnvironment(EShaderPlatform InPlatform,
 	OutEnvironment.SetDefine(TEXT("MATERIAL_DECAL_READ_MASK"), MaterialCompilationOutput.UsedDBufferTextures);
 	OutEnvironment.SetDefine(TEXT("MATERIAL_USES_DECAL_LOOKUP"), MaterialCompilationOutput.bUsesDBufferTextureLookup);
 	OutEnvironment.SetDefine(TEXT("MATERIAL_PATH_TRACING_BUFFER_READ"), MaterialCompilationOutput.UsedPathTracingBufferTextures);
+	OutEnvironment.SetDefine(TEXT("MATERIAL_NEURAL_POST_PROCESS"), (MaterialCompilationOutput.bUsedWithNeuralNetworks || Material->IsUsedWithNeuralNetworks()) && Material->IsPostProcessMaterial());
 
 	// Count the number of VTStacks (each stack will allocate a feedback slot)
 	OutEnvironment.SetDefine(TEXT("NUM_VIRTUALTEXTURE_SAMPLES"), VTStacks.Num());
@@ -11307,6 +11308,47 @@ int32 FHLSLMaterialTranslator::GetLocal(const FName& LocalName)
 	return AddInlinedCodeChunk(MCT_Float1, TEXT("%s"), *Entry->Name);
 }
 
+int32 FHLSLMaterialTranslator::NeuralOutput(int32 ViewportUV, uint32 NeuralIndexType)
+{
+	if (Material->GetMaterialDomain() != MD_PostProcess)
+	{
+		Errorf(TEXT("NNE Output Node are only available on post process material."));
+	}
+
+	AddEstimatedTextureSample();
+	MaterialCompilationOutput.bUsedWithNeuralNetworks = true;
+	
+	if (NeuralIndexType == 0)
+	{
+		if (ViewportUV == INDEX_NONE)
+		{
+			ViewportUV = AddInlinedCodeChunk(MCT_Float2, TEXT("GetViewportUV(Parameters)"));
+		}
+
+		return AddCodeChunk(MCT_Float4, TEXT("NeuralTextureOutput(Parameters,%s)"),
+			*CoerceParameter(ViewportUV, MCT_Float2));
+	}
+	else if (NeuralIndexType == 1)
+	{
+		int32 BufferIndex = INDEX_NONE;
+
+		if (ViewportUV == INDEX_NONE)
+		{
+			ViewportUV = AddInlinedCodeChunk(MCT_Float2, TEXT("GetViewportUV(Parameters)"));
+			BufferIndex = AppendVector(Constant2(0.0f, 0.0f), ViewportUV);
+		}
+		else
+		{
+			BufferIndex = ViewportUV;
+		}
+
+		return AddCodeChunk(MCT_Float4, TEXT("NeuralBufferOutput(Parameters,%s)"),
+			*CoerceParameter(BufferIndex, MCT_Float4));
+	}
+
+	return INDEX_NONE;
+}
+
 FSubstrateOperator& FHLSLMaterialTranslator::SubstrateCompilationRegisterOperator(int32 OperatorType, FGuid SubstrateExpressionGuid, UMaterialExpression* Child, UMaterialExpression* Parent, FGuid SubstrateParentExpressionGuid, bool bUseParameterBlending)
 {
 	FSubstrateCompilationContext& SubstrateCtx = SubstrateCompilationContext[CurrentSubstrateCompilationContext];
diff --git a/Engine/Source/Runtime/Engine/Private/Materials/HLSLMaterialTranslator.h b/Engine/Source/Runtime/Engine/Private/Materials/HLSLMaterialTranslator.h
--- a/Engine/Source/Runtime/Engine/Private/Materials/HLSLMaterialTranslator.h
+++ b/Engine/Source/Runtime/Engine/Private/Materials/HLSLMaterialTranslator.h
@@ -1197,6 +1197,9 @@ protected:
 	virtual int32 SetLocal(const FName& LocalName, int32 Value) override;
 	virtual int32 GetLocal(const FName& LocalName) override;
 
+	// Neural network nodes
+	virtual int32 NeuralOutput(int32 ViewportUV, uint32 NeuralIndexType) override;
+
 	// Substrate
 	virtual int32 SubstrateCreateAndRegisterNullMaterial() override;
 	virtual int32 SubstrateSlabBSDF(
diff --git a/Engine/Source/Runtime/Engine/Private/Materials/Material.cpp b/Engine/Source/Runtime/Engine/Private/Materials/Material.cpp
--- a/Engine/Source/Runtime/Engine/Private/Materials/Material.cpp
+++ b/Engine/Source/Runtime/Engine/Private/Materials/Material.cpp
@@ -4431,7 +4431,9 @@ bool UMaterial::CanEditChange(const FProperty* InProperty) const
 			PropertyName == GET_MEMBER_NAME_STRING_CHECKED(UMaterial, bIsBlendable) ||
 			PropertyName == GET_MEMBER_NAME_STRING_CHECKED(UMaterial, bEnableStencilTest) ||
 			PropertyName == GET_MEMBER_NAME_STRING_CHECKED(UMaterial, StencilCompare) ||
-			PropertyName == GET_MEMBER_NAME_STRING_CHECKED(UMaterial, StencilRefValue)
+			PropertyName == GET_MEMBER_NAME_STRING_CHECKED(UMaterial, StencilRefValue) ||
+			PropertyName == GET_MEMBER_NAME_STRING_CHECKED(UMaterial, NeuralProfileId) ||
+			PropertyName == GET_MEMBER_NAME_STRING_CHECKED(UMaterial, bUsedWithNeuralNetworks)
 			)
 		{
 			return MaterialDomain == MD_PostProcess;
@@ -6737,6 +6739,12 @@ USpecularProfile* UMaterial::GetSpecularProfile_Internal(uint32 Index) const
 	return SpecularProfiles[Index];
 }
 
+UNeuralProfile* UMaterial::GetNeuralProfile_Internal() const
+{
+	checkSlow(IsInGameThread());
+	return NeuralProfile;
+}
+
 bool UMaterial::CastsRayTracedShadows() const
 {
 	return bCastRayTracedShadows;
diff --git a/Engine/Source/Runtime/Engine/Private/Materials/MaterialExpressionHLSL.cpp b/Engine/Source/Runtime/Engine/Private/Materials/MaterialExpressionHLSL.cpp
--- a/Engine/Source/Runtime/Engine/Private/Materials/MaterialExpressionHLSL.cpp
+++ b/Engine/Source/Runtime/Engine/Private/Materials/MaterialExpressionHLSL.cpp
@@ -113,6 +113,7 @@
 #include "Materials/MaterialExpressionNaniteReplace.h"
 #include "Materials/MaterialExpressionNoise.h"
 #include "Materials/MaterialExpressionNormalize.h"
+#include "Materials/MaterialExpressionNeuralPostProcessNode.h"
 #include "Materials/MaterialExpressionObjectBounds.h"
 #include "Materials/MaterialExpressionObjectOrientation.h"
 #include "Materials/MaterialExpressionObjectPositionWS.h"
@@ -4382,4 +4383,135 @@ bool GenerateStaticTerrainLayerWeightExpression(FName LayerName, float PreviewWe
 	return true;
 }
 
+UE::Shader::EValueType UMaterialExpressionNeuralNetworkInput::GetCustomOutputType(int32 OutputIndex) const
+{
+	using namespace UE::Shader;
+
+	if (OutputIndex == 0)
+	{
+		return EValueType::Float4;
+	}
+	else if (OutputIndex == 1)
+	{
+		return EValueType::Float3;
+	}
+	else if (OutputIndex == 2)
+	{
+		return EValueType::Float1;
+	}
+	else
+	{
+		return EValueType::Void;
+	}
+}
+
+bool UMaterialExpressionNeuralNetworkInput::GenerateHLSLExpression(FMaterialHLSLGenerator& Generator, UE::HLSLTree::FScope& Scope, int32 OutputIndex, UE::HLSLTree::FExpression const*& OutExpression) const
+{
+	using namespace UE::HLSLTree;
+
+	const UMaterial* BaseMaterial = Generator.GetTargetMaterial();
+	if (BaseMaterial)
+	{
+		const EMaterialDomain MaterialDomain = BaseMaterial->MaterialDomain;
+		if (MaterialDomain != MD_PostProcess)
+		{
+			return Generator.Error(TEXT("Neural Output is only available in post process material."));
+		}
+	}
+
+	const bool bUseTextureAsInput = NeuralIndexType == ENeuralIndexType::NIT_TextureIndex;
+	FTree& Tree = Generator.GetTree();
+	
+	if (OutputIndex == 0)
+	{
+		if (Coordinates.IsConnected())
+		{
+			const FExpression* TracedCoordinates = Coordinates.AcquireHLSLExpression(Generator, Scope);
+			if (bUseTextureAsInput)
+			{
+				const FExpression* R = Tree.NewSwizzle(FSwizzleParameters(0), TracedCoordinates);
+				const FExpression* G = Tree.NewSwizzle(FSwizzleParameters(1), TracedCoordinates);
+				const FExpression* B = Tree.NewSwizzle(FSwizzleParameters(2), TracedCoordinates);
+				const FExpression* A = Tree.NewSwizzle(FSwizzleParameters(3), TracedCoordinates);
+				OutExpression = Tree.NewAppend(Tree.NewConstant(-1.0f), G, B, A);
+			}
+			else
+			{
+				OutExpression = TracedCoordinates;
+			}
+		}
+		else
+		{
+			const FExpression* ViewportUV = Tree.NewExpression<Material::FExpressionExternalInput>(Material::EExternalInput::ViewportUV);
+			int32 BatchIndex = bUseTextureAsInput ? -1.0f : 0.0f;
+			OutExpression = Tree.NewAppend(Generator.GetTree().NewConstant(FVector2f(BatchIndex, 0.0f)), ViewportUV);
+		}
+	}
+	else if (OutputIndex == 1)
+	{
+		OutExpression = Input0.AcquireHLSLExpressionOrConstant(Generator, Scope, FVector3f(0.5f));
+	}
+	else if (OutputIndex == 2)
+	{
+		OutExpression = Mask.AcquireHLSLExpressionOrConstant(Generator, Scope, 1.0f);
+	}
+
+	return OutExpression != nullptr;
+}
+
+bool UMaterialExpressionNeuralNetworkOutput::GenerateHLSLExpression(FMaterialHLSLGenerator& Generator, UE::HLSLTree::FScope& Scope, int32 OutputIndex, UE::HLSLTree::FExpression const*& OutExpression) const
+{
+	using namespace UE::HLSLTree;
+
+	const UMaterial* BaseMaterial = Generator.GetTargetMaterial();
+	if (BaseMaterial)
+	{
+		const EMaterialDomain MaterialDomain = BaseMaterial->MaterialDomain;
+		if (MaterialDomain != MD_PostProcess)
+		{
+			return Generator.Error(TEXT("Neural Output is only available in post process material."));
+		}
+	}
+	FTree& Tree = Generator.GetTree();
+
+	if (OutputIndex == 0)
+	{
+		auto GetCoordinateByNeuralIndexType = [&]()-> const FExpression* {
+
+			if (NeuralIndexType == ENeuralIndexType::NIT_TextureIndex)
+			{
+				return Coordinates.GetTracedInput().Expression ?
+					Coordinates.AcquireHLSLExpression(Generator, Scope) :
+					Tree.NewExpression<Material::FExpressionExternalInput>(Material::EExternalInput::ViewportUV);
+			}
+			else if (NeuralIndexType == ENeuralIndexType::NIT_BufferIndex)
+			{
+				if (Coordinates.GetTracedInput().Expression)
+				{
+					return Coordinates.AcquireHLSLExpression(Generator, Scope);
+				}
+				else
+				{
+					const FExpression* ViewportUV = Tree.NewExpression<Material::FExpressionExternalInput>(Material::EExternalInput::ViewportUV);
+					return Tree.NewAppend(Tree.NewConstant(FVector2f(0.0f)), ViewportUV);
+				}
+			}
+			else
+			{
+				return nullptr;
+			}
+		};
+
+		const FExpression* TracedCoordinates = GetCoordinateByNeuralIndexType();
+		if (TracedCoordinates)
+		{
+			OutExpression = Generator.GetTree().NewExpression<Material::FExpressionNeuralNetworkOutput>(TracedCoordinates, NeuralIndexType);
+		}
+
+		return TracedCoordinates != nullptr;
+	}
+
+	return true;
+}
+
 #endif // WITH_EDITOR
diff --git a/Engine/Source/Runtime/Engine/Private/Materials/MaterialExpressions.cpp b/Engine/Source/Runtime/Engine/Private/Materials/MaterialExpressions.cpp
--- a/Engine/Source/Runtime/Engine/Private/Materials/MaterialExpressions.cpp
+++ b/Engine/Source/Runtime/Engine/Private/Materials/MaterialExpressions.cpp
@@ -142,6 +142,7 @@
 #include "Materials/MaterialExpressionNaniteReplace.h"
 #include "Materials/MaterialExpressionNoise.h"
 #include "Materials/MaterialExpressionNormalize.h"
+#include "Materials/MaterialExpressionNeuralPostProcessNode.h"
 #include "Materials/MaterialExpressionObjectBounds.h"
 #include "Materials/MaterialExpressionObjectLocalBounds.h"
 #include "Materials/MaterialExpressionObjectOrientation.h"
@@ -27840,4 +27841,177 @@ FString UMaterialExpressionSubsurfaceMediumMaterialOutput::GetDisplayName() cons
 	return TEXT("Subsurface Medium");
 }
 
+
+///////////////////////////Neural network nodes
+
+UMaterialExpressionNeuralNetworkInput::UMaterialExpressionNeuralNetworkInput(const FObjectInitializer& ObjectInitializer)
+	: Super(ObjectInitializer)
+{
+	// Structure to hold one-time initialization
+	struct FConstructorStatics
+	{
+		FText NAME_NeuralNetwork;
+		FConstructorStatics()
+			: NAME_NeuralNetwork(LOCTEXT("NeuralNetwork", "NeuralNetwork"))
+		{
+		}
+	};
+	static FConstructorStatics ConstructorStatics;
+
+#if WITH_EDITORONLY_DATA
+	MenuCategories.Add(ConstructorStatics.NAME_NeuralNetwork);
+#endif
+
+#if WITH_EDITOR
+	Outputs.Reset();
+#endif
+}
+
+#if WITH_EDITOR
+
+int32 UMaterialExpressionNeuralNetworkInput::Compile(class FMaterialCompiler* Compiler, int32 OutputIndex)
+{
+	int32 CodeInput = INDEX_NONE;
+	const bool bUseTextureAsInput = NeuralIndexType == ENeuralIndexType::NIT_TextureIndex;
+
+	if (OutputIndex == 0)
+	{
+		if (Coordinates.IsConnected())
+		{
+			CodeInput = Coordinates.Compile(Compiler);
+			if (bUseTextureAsInput)
+			{
+				CodeInput = Compiler->ComponentMask(CodeInput, false, true, true, true);
+				CodeInput = Compiler->AppendVector(Compiler->Constant(-1.0f), CodeInput);
+			}
+		}
+		else
+		{
+			int32 ViewportUV = Compiler->GetViewportUV();
+			int32 BatchIndex = bUseTextureAsInput ? -1.0f : 0.0f;
+				
+			CodeInput = Compiler->AppendVector(Compiler->Constant2(BatchIndex, 0.0f), ViewportUV);
+		}
+	}
+	else if (OutputIndex == 1)
+	{
+		CodeInput = Input0.IsConnected() ? Input0.Compile(Compiler) : Compiler->Constant3(0.5f, 0.5f, 0.5f);
+	}
+	else if (OutputIndex == 2)
+	{
+		CodeInput = Mask.IsConnected() ? Mask.Compile(Compiler) : Compiler->Constant(1.0f);
+	}
+
+	return Compiler->CustomOutput(this, OutputIndex, CodeInput);
+}
+
+void UMaterialExpressionNeuralNetworkInput::GetCaption(TArray<FString>& OutCaptions) const
+{
+	if (NeuralIndexType == ENeuralIndexType::NIT_TextureIndex)
+	{
+		OutCaptions.Add(TEXT("Neural Input (Texture)"));
+	}
+	else if (NeuralIndexType == ENeuralIndexType::NIT_BufferIndex)
+	{
+		OutCaptions.Add(TEXT("Neural Input (Buffer)"));
+	}
+}
+
+uint32 UMaterialExpressionNeuralNetworkInput::GetInputType(int32 InputIndex)
+{
+
+	switch (InputIndex)
+	{
+	case 0:return MCT_Float4;break;
+	case 1:return MCT_Float3;break;
+	case 2:return MCT_Float1;break;
+	default: checkNoEntry(); break;
+	}
+
+	return MCT_Float1;
+}
+
+#endif // WITH_EDITOR
+
+int32 UMaterialExpressionNeuralNetworkInput::GetNumOutputs() const
+{
+	return 3;
+}
+
+FString UMaterialExpressionNeuralNetworkInput::GetFunctionName() const
+{
+	return TEXT("GetNeuralInput");
+}
+
+///////////////////////////////////////////////////////////////////////////////
+// UMaterialExpressionNeuralNetworkOutput
+///////////////////////////////////////////////////////////////////////////////
+UMaterialExpressionNeuralNetworkOutput::UMaterialExpressionNeuralNetworkOutput(const FObjectInitializer& ObjectInitializer)
+	: Super(ObjectInitializer)
+{
+
+	// Structure to hold one-time initialization
+	struct FConstructorStatics
+	{
+		FText NAME_NeuralNetwork;
+		FConstructorStatics()
+			: NAME_NeuralNetwork(LOCTEXT("NeuralNetwork", "NeuralNetwork"))
+		{
+		}
+	};
+	static FConstructorStatics ConstructorStatics;
+
+#if WITH_EDITORONLY_DATA
+	MenuCategories.Add(ConstructorStatics.NAME_NeuralNetwork);
+
+	bShaderInputData = true;
+	bShowOutputNameOnPin = true;
+#endif
+
+#if WITH_EDITOR
+	Outputs.Reset();
+	Outputs.Add(FExpressionOutput(TEXT("RGBA"), 1, 1, 1, 1, 1));
+#endif
+}
+
+#if WITH_EDITOR
+int32 UMaterialExpressionNeuralNetworkOutput::Compile(class FMaterialCompiler* Compiler, int32 OutputIndex)
+{
+	int32 ViewportUV = INDEX_NONE;
+
+	if (Coordinates.GetTracedInput().Expression)
+	{
+		ViewportUV = Coordinates.Compile(Compiler);
+	}
+
+	return Compiler->NeuralOutput(ViewportUV, NeuralIndexType);
+}
+
+void UMaterialExpressionNeuralNetworkOutput::GetCaption(TArray<FString>& OutCaptions) const
+{
+	if (NeuralIndexType == ENeuralIndexType::NIT_TextureIndex)
+	{
+		OutCaptions.Add(TEXT("Neural Output (Texture)"));
+	}
+	else if (NeuralIndexType == ENeuralIndexType::NIT_BufferIndex)
+	{
+		OutCaptions.Add(TEXT("Neural Output (Buffer)"));
+	}
+}
+uint32 UMaterialExpressionNeuralNetworkOutput::GetInputType(int32 InputIndex)
+{
+	if (NeuralIndexType == ENeuralIndexType::NIT_TextureIndex)
+	{
+		return MCT_Float2;
+	}
+	else if (NeuralIndexType == ENeuralIndexType::NIT_BufferIndex)
+	{
+		return MCT_Float4;
+	}
+
+	checkNoEntry();
+	return MCT_Float2;
+}
+#endif // WITH_EDITOR
+
 #undef LOCTEXT_NAMESPACE
diff --git a/Engine/Source/Runtime/Engine/Private/Materials/MaterialHLSLEmitter.cpp b/Engine/Source/Runtime/Engine/Private/Materials/MaterialHLSLEmitter.cpp
--- a/Engine/Source/Runtime/Engine/Private/Materials/MaterialHLSLEmitter.cpp
+++ b/Engine/Source/Runtime/Engine/Private/Materials/MaterialHLSLEmitter.cpp
@@ -487,6 +487,8 @@ static void GetMaterialEnvironment(EShaderPlatform InPlatform,
 	OutEnvironment.SetDefine(TEXT("MATERIAL_FULLY_ROUGH"), bIsFullyRough);
 	OutEnvironment.SetDefine(TEXT("MATERIAL_USES_ANISOTROPY"), bUsesAnisotropy);
 
+	OutEnvironment.SetDefine(TEXT("MATERIAL_NEURAL_POST_PROCESS"), (MaterialCompilationOutput.bUsedWithNeuralNetworks || InMaterial.IsUsedWithNeuralNetworks()) && InMaterial.IsPostProcessMaterial());
+
 	// Count the number of VTStacks (each stack will allocate a feedback slot)
 	OutEnvironment.SetDefine(TEXT("NUM_VIRTUALTEXTURE_SAMPLES"), EmitMaterialData.VTStacks.Num());
 
diff --git a/Engine/Source/Runtime/Engine/Private/Materials/MaterialHLSLTree.cpp b/Engine/Source/Runtime/Engine/Private/Materials/MaterialHLSLTree.cpp
--- a/Engine/Source/Runtime/Engine/Private/Materials/MaterialHLSLTree.cpp
+++ b/Engine/Source/Runtime/Engine/Private/Materials/MaterialHLSLTree.cpp
@@ -2795,6 +2795,57 @@ void FExpressionAtmosphericFogColorFunction::EmitValueShader(FEmitContext& Conte
 		EmitPosition);
 }
 
+bool FExpressionNeuralNetworkOutput::PrepareValue(FEmitContext& Context, FEmitScope& Scope, const FRequestedType& RequestedType, FPrepareValueResult& OutResult) const
+{
+	if (NeuralIndexType == 0)
+	{
+		const FPreparedType& TextureIndexType = Context.PrepareExpression(CoordinatesExpression, Scope, Shader::EValueType::Float2);
+		if (TextureIndexType.IsVoid())
+		{
+			return false;
+		}
+	}
+	else if (NeuralIndexType == 1)
+	{
+		const FPreparedType& BufferIndexType = Context.PrepareExpression(CoordinatesExpression, Scope, Shader::EValueType::Float4);
+		if (BufferIndexType.IsVoid())
+		{
+			return false;
+		}
+	}
+
+	if (Context.bMarkLiveValues && Context.MaterialCompilationOutput)
+	{
+		Context.MaterialCompilationOutput->bUsedWithNeuralNetworks = true;
+	}
+
+	return OutResult.SetType(Context, RequestedType, EExpressionEvaluation::Shader, Shader::EValueType::Float4);
+}
+
+void FExpressionNeuralNetworkOutput::EmitValueShader(FEmitContext& Context, FEmitScope& Scope, const FRequestedType& RequestedType, FEmitValueShaderResult& OutResult) const
+{
+	if (NeuralIndexType == 0)
+	{
+		FEmitShaderExpression* EmitTextureIndex = CoordinatesExpression->GetValueShader(Context, Scope, Shader::EValueType::Float2);
+		if (EmitTextureIndex)
+		{
+			OutResult.Code = Context.EmitInlineExpression(Scope, Shader::EValueType::Float4,
+				TEXT("NeuralTextureOutput(Parameters, %)"),
+				EmitTextureIndex);
+		}
+	}
+	else if (NeuralIndexType == 1)
+	{
+		FEmitShaderExpression* EmitBufferIndex = CoordinatesExpression->GetValueShader(Context, Scope, Shader::EValueType::Float4);
+		if (EmitBufferIndex)
+		{
+			OutResult.Code = Context.EmitInlineExpression(Scope, Shader::EValueType::Float4,
+				TEXT("NeuralBufferOutput(Parameters, %)"),
+				EmitBufferIndex);
+		}
+	}
+}
+
 int32 FEmitData::FindInterpolatorIndex(const FExpression* Expression) const
 {
 	for (int32 Index = 0; Index < VertexInterpolators.Num(); ++Index)
diff --git a/Engine/Source/Runtime/Engine/Private/Materials/MaterialInterface.cpp b/Engine/Source/Runtime/Engine/Private/Materials/MaterialInterface.cpp
--- a/Engine/Source/Runtime/Engine/Private/Materials/MaterialInterface.cpp
+++ b/Engine/Source/Runtime/Engine/Private/Materials/MaterialInterface.cpp
@@ -23,6 +23,7 @@
 #include "ObjectCacheEventSink.h"
 #include "Engine/SubsurfaceProfile.h"
 #include "Engine/SpecularProfile.h"
+#include "Engine/NeuralProfile.h"
 #include "Interfaces/ITargetPlatform.h"
 #include "Components/PrimitiveComponent.h"
 #include "ContentStreaming.h"
@@ -1331,6 +1332,12 @@ USpecularProfile* UMaterialInterface::GetSpecularProfile_Internal(uint32 Index)
 	return nullptr;
 }
 
+UNeuralProfile* UMaterialInterface::GetNeuralProfile_Internal() const
+{
+	return nullptr;
+}
+
+
 uint32 UMaterialInterface::NumSpecularProfile_Internal() const
 {
 	return 0u;
@@ -1407,6 +1414,40 @@ void UMaterialInterface::UpdateMaterialRenderProxy(FMaterialRenderProxy& Proxy)
 		});
 	}
 
+	UMaterial* Material = GetMaterial();
+	if (Material && Material->IsPostProcessMaterial())
+	{
+		struct FEntry
+		{
+			UNeuralProfile* Profile = nullptr;
+			FNeuralProfileStruct Setting;
+			FGuid Guid;
+		};
+		
+		UNeuralProfile* LocalNeuralProfile = GetNeuralProfile_Internal();
+		
+		if (LocalNeuralProfile)
+		{
+			FEntry Entry;
+			Entry.Profile = LocalNeuralProfile;
+			Entry.Setting = LocalNeuralProfile->Settings;
+			Entry.Guid = LocalNeuralProfile->Guid;
+
+			FMaterialRenderProxy* InProxy = &Proxy;
+			ENQUEUE_RENDER_COMMAND(UpdateMaterialRenderProxyNNEModelData)(
+				[LocalNeuralProfile, InProxy, Entry, Material](FRHICommandListImmediate& RHICmdList)
+				{
+					
+					const uint32 AllocationId = NeuralProfile::AddOrUpdateProfile(Entry.Profile, Entry.Guid, Entry.Setting);
+					check(AllocationId >= 0 && AllocationId < MAX_NEURAL_PROFILE_COUNT);
+
+					Material->NeuralProfileId = AllocationId;
+					
+					InProxy->SetNeuralProfileRT(LocalNeuralProfile);
+				});
+		}
+	}
+
 	if (Substrate::IsSubstrateEnabled())
 	{
 		struct FEntry
diff --git a/Engine/Source/Runtime/Engine/Private/Materials/MaterialShared.cpp b/Engine/Source/Runtime/Engine/Private/Materials/MaterialShared.cpp
--- a/Engine/Source/Runtime/Engine/Private/Materials/MaterialShared.cpp
+++ b/Engine/Source/Runtime/Engine/Private/Materials/MaterialShared.cpp
@@ -55,6 +55,7 @@
 #include "SubstrateDefinitions.h"
 #include "DataDrivenShaderPlatformInfo.h"
 #include "ProfilingDebugging/CookStats.h"
+#include "Engine/NeuralProfile.h"
 
 #define LOCTEXT_NAMESPACE "MaterialShared"
 
@@ -1725,6 +1726,11 @@ bool FMaterialResource::IsUsedWithVirtualHeightfieldMesh() const
 	return Material->bUsedWithVirtualHeightfieldMesh;
 }
 
+bool FMaterialResource::IsUsedWithNeuralNetworks() const
+{
+	return Material->bUsedWithNeuralNetworks && Material->IsPostProcessMaterial();
+}
+
 bool FMaterialResource::IsUsedWithLandscape() const
 {
 	return false;
@@ -2177,6 +2183,11 @@ uint32 FMaterialResource::GetStencilRefValue() const
 	return GetMaterialDomain() == MD_PostProcess ? Material->StencilRefValue : 0;
 }
 
+int32 FMaterialResource::GetNeuralProfileId() const
+{ 
+	return GetMaterialDomain() == MD_PostProcess ? Material->NeuralProfileId : INDEX_NONE;
+}
+
 uint32 FMaterialResource::GetStencilCompare() const
 {
 	return GetMaterialDomain() == MD_PostProcess ? uint32(Material->StencilCompare.GetValue()) : 0;
@@ -5345,6 +5356,7 @@ FMaterialShaderParameters::FMaterialShaderParameters(const FMaterial* InMaterial
 	bHasRuntimeVirtualTextureOutput = InMaterial->HasRuntimeVirtualTextureOutput();
 	bIsUsedWithLidarPointCloud = InMaterial->IsUsedWithLidarPointCloud();
 	bIsUsedWithVirtualHeightfieldMesh = InMaterial->IsUsedWithVirtualHeightfieldMesh();
+	bIsUsedWithNeuralNetworks = InMaterial->IsUsedWithNeuralNetworks();
 	bIsUsedWithNanite = InMaterial->IsUsedWithNanite();
 	bIsStencilTestEnabled = InMaterial->IsStencilTestEnabled();
 	bIsTranslucencySurface = InMaterial->GetTranslucencyLightingMode() == ETranslucencyLightingMode::TLM_Surface || InMaterial->GetTranslucencyLightingMode() == ETranslucencyLightingMode::TLM_SurfacePerPixelLighting;
diff --git a/Engine/Source/Runtime/Engine/Private/Rendering/NeuralProfile.cpp b/Engine/Source/Runtime/Engine/Private/Rendering/NeuralProfile.cpp
new file mode 100644
--- /dev/null
+++ b/Engine/Source/Runtime/Engine/Private/Rendering/NeuralProfile.cpp
@@ -0,0 +1,286 @@
+// Copyright Epic Games, Inc. All Rights Reserved.
+
+#include "Engine/NeuralProfile.h"
+#include "Engine/Texture2D.h"
+#include "DataDrivenShaderPlatformInfo.h"
+#include "Math/Float16.h"
+#include "Rendering/BurleyNormalizedSSS.h"
+#include "EngineModule.h"
+#include "RenderTargetPool.h"
+#include "PixelShaderUtils.h"
+#include "RenderingThread.h"
+#include "Rendering/Texture2DResource.h"
+#include "RenderGraphResources.h"
+#include "RenderGraphBuilder.h"
+#include "NNEModelData.h"
+
+#include UE_INLINE_GENERATED_CPP_BY_NAME(NeuralProfile)
+
+DEFINE_LOG_CATEGORY_STATIC(LogNeuralProfile, Log, All);
+
+TUniquePtr<NeuralProfile::INeuralProfileManager> GNeuralProfileManager;
+
+///////////////////////////////////////////////////////////////////////////////////////////////////
+// FNeuralProfileModelTextureManager
+
+//@TODO: add texture management.
+class FNeuralProfileModelTextureManager : public FRenderResource
+{
+public:
+	FNeuralProfileModelTextureManager()
+	{
+		check(IsInGameThread());
+		//@TODO: add a default network?
+		//		 add param texture
+
+		//FProfileEntry& Entry = NeuralProfileEntries.AddDefaulted_GetRef();
+		//Entry.Settings = FNeuralProfileStruct();
+		//Entry.Profile = nullptr;
+	}
+
+	~FNeuralProfileModelTextureManager() {}
+
+	// convenience, can be optimized 
+	// @param Profile must not be nullptr, game thread pointer, do not dereference, only for comparison
+	int32 AddOrUpdateProfile(const UNeuralProfile* InProfile, const FGuid& InGuid, const FNeuralProfileStruct InSettings)
+	{
+		check(InProfile);
+		
+		int32 AllocationId = FindAllocationId(InProfile);
+
+		if (AllocationId != INDEX_NONE)
+		{
+			UpdateProfile(AllocationId, InSettings);
+		}
+		else
+		{
+			AllocationId = AddProfile(InProfile, InGuid, InSettings);
+		}
+
+		return AllocationId;
+	}
+
+	int32 AddProfile(const UNeuralProfile* InProfile, const FGuid& InGuid, const FNeuralProfileStruct Settings)
+	{
+		check(InProfile);
+		check(FindAllocationId(InProfile) == INDEX_NONE);
+
+		int32 AllocationId = INDEX_NONE;
+		{
+			for (int i = 0; i < NeuralProfileEntries.Num(); ++i)
+			{
+				if (NeuralProfileEntries[i].Profile == nullptr)
+				{
+					AllocationId = i;
+					FProfileEntry& Entry = NeuralProfileEntries[AllocationId];
+					Entry.Profile = InProfile;
+
+					break;
+				}
+			}
+
+			if (AllocationId == INDEX_NONE)
+			{
+				AllocationId = NeuralProfileEntries.Num();
+				FProfileEntry& Entry = NeuralProfileEntries.AddDefaulted_GetRef();
+				Entry.Profile = InProfile;
+			}
+		}
+
+		UpdateProfile(AllocationId, Settings);
+
+		return AllocationId;
+	}
+
+	void UpdateProfile(int32 AllocationId, const FNeuralProfileStruct Settings)
+	{
+		check(IsInRenderingThread());
+		
+		if (AllocationId != INDEX_NONE && GNeuralProfileManager)
+		{
+			check(AllocationId < NeuralProfileEntries.Num());
+
+			FNeuralProfileStruct& CurrentSettings = NeuralProfileEntries[AllocationId].Settings;
+			const bool bModelIsTheSame =
+				CurrentSettings.RuntimeName == Settings.RuntimeName &&
+				CurrentSettings.NNEModelData == Settings.NNEModelData;
+
+			const bool bNeedToUpdateModel = !bModelIsTheSame;
+			
+			if (bNeedToUpdateModel)
+			{
+				// any changes to the setting require an update of the NNE model.
+				if (Settings.NNEModelData.Get())
+				{
+					GNeuralProfileManager->UpdateModel(AllocationId, Settings.NNEModelData.Get(), Settings.RuntimeName);
+				}
+			}
+
+			//update model tile type
+			if (CurrentSettings.TileSize != Settings.TileSize)
+			{
+				GNeuralProfileManager->UpdateTileType(AllocationId, Settings.TileSize);
+			}
+			
+			const bool bBufferSizeNeedsUpdate =
+				CurrentSettings.BatchSizeOverride != Settings.BatchSizeOverride;
+			int32 BatchSizeOverride = NeuralProfileEntries[AllocationId].Settings.BatchSizeOverride;
+			if (bBufferSizeNeedsUpdate)
+			{
+				if (GNeuralProfileManager->UpdateBatchSize(AllocationId, Settings.BatchSizeOverride))
+				{
+					BatchSizeOverride = Settings.BatchSizeOverride;
+				}
+			}
+
+			NeuralProfileEntries[AllocationId].Settings = Settings;
+			NeuralProfileEntries[AllocationId].Settings.BatchSizeOverride = BatchSizeOverride;
+
+		}
+	}
+
+	void UpdateProfile(const UNeuralProfile* InProfile, const FNeuralProfileStruct Settings)
+	{
+		UpdateProfile(FindAllocationId(InProfile), Settings);
+	}
+
+	void RemoveProfile(const UNeuralProfile* InProfile)
+	{
+		int32 AllocationId = FindAllocationId(InProfile);
+
+		if (AllocationId != INDEX_NONE)
+		{
+			check(AllocationId >= 0);
+			check(NeuralProfileEntries[AllocationId].Profile == InProfile);
+			FNeuralProfileStruct LocalSettings = NeuralProfileEntries[AllocationId].Settings;
+
+			ENQUEUE_RENDER_COMMAND(RemoveNeuralProfile)(
+				[LocalSettings, AllocationId](FRHICommandListImmediate& RHICmdList)
+				{
+					// any changes to the setting require an update of the NNE model.
+					if (GNeuralProfileManager)
+					{
+						GNeuralProfileManager->RemoveModel(AllocationId);
+					}
+				});
+
+			NeuralProfileEntries[AllocationId].Profile = nullptr;
+			NeuralProfileEntries[AllocationId].Settings.Invalidate();
+		}
+	}
+
+	int32 FindAllocationId(const UNeuralProfile* InProfile) const
+	{
+		for (int i = 0; i < NeuralProfileEntries.Num(); ++i)
+		{
+			if (NeuralProfileEntries[i].Profile == InProfile)
+			{
+				return i;
+			}
+		}
+
+		return INDEX_NONE;
+	}
+
+	FNeuralProfileStruct GetProfileSetting(int32 AllocationId)
+	{
+		FNeuralProfileStruct Settings;
+		if (AllocationId >= 0 && AllocationId < NeuralProfileEntries.Num())
+		{
+			Settings = NeuralProfileEntries[AllocationId].Settings;
+		}
+
+		return Settings;
+	}
+
+private:
+
+	struct FProfileEntry
+	{
+		FNeuralProfileStruct Settings;
+		const UNeuralProfile* Profile = nullptr;
+	};
+
+	TArray<FProfileEntry> NeuralProfileEntries;
+};
+
+// Global resources - lives on the render thread
+TGlobalResource<FNeuralProfileModelTextureManager> GNeuralProfileModelTextureManager;
+
+///////////////////////////////////////////////////////////////////////////////////////////////////
+// UNeuralProfile
+
+UNeuralProfile::UNeuralProfile(const FObjectInitializer& ObjectInitializer)
+	: Super(ObjectInitializer)
+{
+}
+
+void UNeuralProfile::BeginDestroy()
+{
+	UNeuralProfile* Ref = this;
+	ENQUEUE_RENDER_COMMAND(RemoveNeuralProfile)(
+	[Ref](FRHICommandList& RHICmdList)
+	{
+		GNeuralProfileModelTextureManager.RemoveProfile(Ref);
+	});
+
+	Super::BeginDestroy();
+}
+
+void UNeuralProfile::PostEditChangeProperty(struct FPropertyChangedEvent& PropertyChangedEvent)
+{
+	const FNeuralProfileStruct LocalSettings = this->Settings;
+	UNeuralProfile* LocalProfile = this;
+	GetRendererModule().InvalidatePathTracedOutput();
+
+	ENQUEUE_RENDER_COMMAND(UpdateNeuralProfile)(
+		[LocalSettings, LocalProfile](FRHICommandListImmediate& RHICmdList)
+		{
+			// any changes to the setting require an update of the texture
+			GNeuralProfileModelTextureManager.UpdateProfile(LocalProfile, LocalSettings);
+		});
+
+	if (GNeuralProfileManager)
+	{
+		ENQUEUE_RENDER_COMMAND(UpdateDimensionParams)(
+			[this](FRHICommandListImmediate& RHICmdList)
+			{
+				Settings.InputDimension = GNeuralProfileManager->GetInputDimension(
+					Settings.NNEModelData.Get(),
+					Settings.RuntimeName);
+
+				Settings.OutputDimension = GNeuralProfileManager->GetOutputDimension(
+					Settings.NNEModelData.Get(),
+					Settings.RuntimeName);
+			});
+	}
+}
+
+/////////////////////////////////////////////////////////////////////////////////////////////////////
+//// Public API
+//
+namespace NeuralProfile
+{
+
+	int32 GetNeuralProfileId(const UNeuralProfile* In)
+	{
+		int32 AllocationId = INDEX_NONE;
+		if (In)
+		{
+			AllocationId = GNeuralProfileModelTextureManager.FindAllocationId(In);
+		}
+
+		return AllocationId;
+	}
+
+	int32 AddOrUpdateProfile(const UNeuralProfile* InProfile, const FGuid& InGuid, const FNeuralProfileStruct InSettings)
+	{
+		return GNeuralProfileModelTextureManager.AddOrUpdateProfile(InProfile, InGuid, InSettings);
+	}
+
+	::FNeuralProfileStruct GetProfileSetting(int32 AllocationId)
+	{
+		return GNeuralProfileModelTextureManager.GetProfileSetting(AllocationId);
+	}
+
+} // namespace NeuralProfile
diff --git a/Engine/Source/Runtime/Engine/Public/MaterialCompiler.h b/Engine/Source/Runtime/Engine/Public/MaterialCompiler.h
--- a/Engine/Source/Runtime/Engine/Public/MaterialCompiler.h
+++ b/Engine/Source/Runtime/Engine/Public/MaterialCompiler.h
@@ -553,6 +553,9 @@ public:
 	virtual int32 SetLocal(const FName& LocalName, int32 Value) = 0;
 	virtual int32 GetLocal(const FName& LocalName) = 0;
 
+	// Neural network nodes
+	virtual int32 NeuralOutput(int32 ViewportUV, uint32 NeuralIndexType) = 0;
+	
 	// Substrate
 	virtual int32 SubstrateCreateAndRegisterNullMaterial() = 0;
 	virtual int32 SubstrateSlabBSDF(
@@ -1182,6 +1185,11 @@ public:
 		return Compiler->GetLocal(LocalName);
 	}
 
+	virtual int32 NeuralOutput(int32 ViewportUV, uint32 NeuralIndexType) override
+	{
+		return Compiler->NeuralOutput(ViewportUV, NeuralIndexType);
+	}
+
 	virtual int32 SubstrateCreateAndRegisterNullMaterial() override
 	{
 		return Compiler->SubstrateCreateAndRegisterNullMaterial();
diff --git a/Engine/Source/Runtime/Engine/Public/MaterialHLSLTree.h b/Engine/Source/Runtime/Engine/Public/MaterialHLSLTree.h
--- a/Engine/Source/Runtime/Engine/Public/MaterialHLSLTree.h
+++ b/Engine/Source/Runtime/Engine/Public/MaterialHLSLTree.h
@@ -916,6 +916,23 @@ public:
 	virtual void EmitValueShader(FEmitContext& Context, FEmitScope& Scope, const FRequestedType& RequestedType, FEmitValueShaderResult& OutResult) const override;
 };
 
+class FExpressionNeuralNetworkOutput : public FExpression
+{
+public:
+	const FExpression* CoordinatesExpression;
+	int32 NeuralIndexType;
+
+	FExpressionNeuralNetworkOutput(
+		const FExpression* InCoordinateExpression,
+		int32 InNeuralIndexType)
+		: CoordinatesExpression(InCoordinateExpression),
+		NeuralIndexType(InNeuralIndexType)
+	{}
+
+	virtual bool PrepareValue(FEmitContext& Context, FEmitScope& Scope, const FRequestedType& RequestedType, FPrepareValueResult& OutResult) const override;
+	virtual void EmitValueShader(FEmitContext& Context, FEmitScope& Scope, const FRequestedType& RequestedType, FEmitValueShaderResult& OutResult) const override;
+};
+
 class FExpressionLightVector : public FExpressionForward
 {
 public:
diff --git a/Engine/Source/Runtime/Engine/Public/MaterialShared.h b/Engine/Source/Runtime/Engine/Public/MaterialShared.h
--- a/Engine/Source/Runtime/Engine/Public/MaterialShared.h
+++ b/Engine/Source/Runtime/Engine/Public/MaterialShared.h
@@ -717,7 +717,8 @@ public:
 		bUsesVertexInterpolator(false),
 		bHasRuntimeVirtualTextureOutputNode(false),
 		bUsesAnisotropy(false),
-		bUsesDisplacement(false)
+		bUsesDisplacement(false),
+		bUsedWithNeuralNetworks(false)
 	{
 #if WITH_EDITOR
 		FMemory::Memzero(EstimatedLWCFuncUsages);
@@ -839,6 +840,11 @@ public:
 	
 	/** Whether the material uses scalar displacement. */
 	LAYOUT_BITFIELD(uint8, bUsesDisplacement, 1);
+
+	/** Whether the material uses NNE. */
+	LAYOUT_BITFIELD(uint8, bUsedWithNeuralNetworks, 1);
+	
+
 };
 
 struct FDebugShaderPipelineInfo
@@ -1911,6 +1917,7 @@ public:
 	virtual bool IsUsedWithHairStrands() const { return false; }
 	virtual bool IsUsedWithLidarPointCloud() const { return false; }
 	virtual bool IsUsedWithVirtualHeightfieldMesh() const { return false; }
+	virtual bool IsUsedWithNeuralNetworks() const { return false; }
 	virtual bool IsFullyRough() const { return false; }
 	virtual bool UseNormalCurvatureToRoughness() const { return false; }
 	virtual enum EMaterialFloatPrecisionMode GetMaterialFloatPrecisionMode() const { return EMaterialFloatPrecisionMode::MFPM_Default; };
@@ -1995,6 +2002,7 @@ public:
 	virtual bool CastsRayTracedShadows() const { return true; }
 	virtual bool HasRenderTracePhysicalMaterialOutputs() const { return false; }
 	virtual EMaterialShadingRate GetShadingRate() const { return MSR_1x1; }
+	virtual int32 GetNeuralProfileId() const { return INDEX_NONE; }
 	/**
 	 * Should shaders compiled for this material be saved to disk?
 	 */
@@ -2651,6 +2659,7 @@ public:
 	ENGINE_API virtual bool IsUsedWithHairStrands() const override;
 	ENGINE_API virtual bool IsUsedWithLidarPointCloud() const override;
 	ENGINE_API virtual bool IsUsedWithVirtualHeightfieldMesh() const override;
+	ENGINE_API virtual bool IsUsedWithNeuralNetworks() const override;
 	ENGINE_API virtual bool IsUsedWithNanite() const override;
 	ENGINE_API virtual bool IsUsedWithVolumetricCloud() const override;
 	ENGINE_API virtual bool IsUsedWithHeterogeneousVolumes() const override;
@@ -2726,6 +2735,7 @@ public:
 	ENGINE_API virtual bool CastsRayTracedShadows() const override;
 	ENGINE_API virtual bool HasRenderTracePhysicalMaterialOutputs() const override;
 	ENGINE_API virtual UMaterialInterface* GetMaterialInterface() const override;
+	ENGINE_API virtual int32 GetNeuralProfileId() const override;
 	/**
 	 * Should shaders compiled for this material be saved to disk?
 	 */
@@ -3069,6 +3079,7 @@ struct FMaterialShaderParameters
 			uint64 bHasRuntimeVirtualTextureOutput : 1;
 			uint64 bIsUsedWithLidarPointCloud : 1;
 			uint64 bIsUsedWithVirtualHeightfieldMesh : 1;
+			uint64 bIsUsedWithNeuralNetworks : 1;
 			uint64 bIsUsedWithNanite : 1;
 			uint64 bIsStencilTestEnabled : 1;
 			uint64 bIsTranslucencySurface : 1;
diff --git a/Engine/Source/Runtime/Engine/Public/Materials/MaterialRenderProxy.h b/Engine/Source/Runtime/Engine/Public/Materials/MaterialRenderProxy.h
--- a/Engine/Source/Runtime/Engine/Public/Materials/MaterialRenderProxy.h
+++ b/Engine/Source/Runtime/Engine/Public/Materials/MaterialRenderProxy.h
@@ -25,6 +25,7 @@ class URuntimeVirtualTexture;
 class USparseVolumeTexture;
 class USubsurfaceProfile;
 class USpecularProfile;
+class UNeuralProfile;
 class UTexture;
 
 struct FMaterialParameterValue;
@@ -213,6 +214,10 @@ public:
 	const USpecularProfile* GetSpecularProfileRT(uint32 Index) const { check(Index<uint32(SpecularProfilesRT.Num())); return SpecularProfilesRT[Index]; }
 	const uint32 NumSpecularProfileRT() const { return SpecularProfilesRT.Num(); }
 
+	// Neural profiles
+	void SetNeuralProfileRT(const UNeuralProfile* Ptr) { NeuralProfileRT = Ptr; }
+	const UNeuralProfile* GetNeuralProfileRT() const { return NeuralProfileRT; }
+
 	static ENGINE_API void UpdateDeferredCachedUniformExpressions();
 	static ENGINE_API void UpdateDeferredCachedUniformExpressions(FRHICommandListBase& RHICmdList, UE::Tasks::FTask* TaskIfAsync = nullptr);
 
@@ -232,6 +237,7 @@ private:
 	/** 0 if not set, game thread pointer, do not dereference, only for comparison */
 	const USubsurfaceProfile* SubsurfaceProfileRT;
 	TArray<const USpecularProfile*> SpecularProfilesRT;
+	const UNeuralProfile* NeuralProfileRT;
 	FString MaterialName;
 
 	/** Incremented each time UniformExpressionCache is modified */
diff --git a/Engine/Source/Runtime/Renderer/Private/PostProcess/NeuralPostProcess.cpp b/Engine/Source/Runtime/Renderer/Private/PostProcess/NeuralPostProcess.cpp
new file mode 100644
--- /dev/null
+++ b/Engine/Source/Runtime/Renderer/Private/PostProcess/NeuralPostProcess.cpp
@@ -0,0 +1,160 @@
+// Copyright Epic Games, Inc. All Rights Reserved.
+
+#include "NeuralPostProcess.h"
+#include "SystemTextures.h"
+#include "MaterialShared.h"
+
+TUniquePtr<INeuralPostProcessInterface> GNeuralPostProcess;
+
+bool IsNeuralPostProcessEnabled()
+{
+	return GNeuralPostProcess.IsValid();
+}
+
+bool ShouldApplyNeuralPostProcessForMaterial(const FMaterial* Material)
+{
+	const bool bUseWithNeuralPostProcess = Material->IsUsedWithNeuralNetworks();
+	const int32 NeuralProfileId = Material->GetNeuralProfileId();
+	bool bApply = true;
+	if (!bUseWithNeuralPostProcess ||
+		!IsNeuralPostProcessEnabled() ||
+		NeuralProfileId == INDEX_NONE)
+	{
+		bApply = false;
+	}
+
+	return bApply;
+}
+
+bool IsNeuralPostProcessShaderParameterUsed(FNeuralPostProcessShaderParameters& NeuralPostProcessShaderParameters)
+{
+	bool bInputNodeInUse = 
+		NeuralPostProcessShaderParameters.InputNeuralBuffer != nullptr || 
+		NeuralPostProcessShaderParameters.RWNeuralTexture != nullptr;
+
+	return bInputNodeInUse;
+}
+
+FNeuralPostProcessShaderParameters GetDefaultNeuralPostProcessShaderParameters(
+	FRDGBuilder& GraphBuilder
+)
+{
+	FNeuralPostProcessShaderParameters NeuralPostProcessShaderParameters;
+	NeuralPostProcessShaderParameters.NeuralTexture = GraphBuilder.CreateSRV(FRDGTextureSRVDesc(GraphBuilder.RegisterExternalTexture(GSystemTextures.BlackDummy)));
+	NeuralPostProcessShaderParameters.OutputNeuralBuffer = GraphBuilder.CreateSRV(GSystemTextures.GetDefaultBuffer(GraphBuilder, 4), PF_R32_FLOAT);
+
+	return NeuralPostProcessShaderParameters;
+}
+
+void SetupNeuralPostProcessShaderParametersForWrite(
+	FNeuralPostProcessShaderParameters& NeuralPostProcessShaderParameters,
+	FRDGBuilder& GraphBuilder,
+	FRDGTextureRef Texture,
+	FRDGBufferRef Buffer,
+	FVector4f Dimension,
+	FRDGBufferRef SourceTypeBuffer
+)
+{
+	NeuralPostProcessShaderParameters.RWNeuralTexture = GraphBuilder.CreateUAV(FRDGTextureUAVDesc(Texture));
+	NeuralPostProcessShaderParameters.InputNeuralBuffer = GraphBuilder.CreateUAV(Buffer, PF_R32_FLOAT);
+	NeuralPostProcessShaderParameters.InputNeuralBufferDimension = Dimension;
+	NeuralPostProcessShaderParameters.NeuralSourceType = GraphBuilder.CreateUAV(SourceTypeBuffer, PF_R32_UINT);
+}
+
+void SetupNeuralPostProcessShaderParametersForWrite(
+	FNeuralPostProcessShaderParameters& NeuralPostProcessShaderParameters, 
+	FRDGBuilder& GraphBuilder, 
+	const FNeuralPostProcessResource& NeuralPostProcessResource
+)
+{
+	SetupNeuralPostProcessShaderParametersForWrite(NeuralPostProcessShaderParameters, GraphBuilder,
+		NeuralPostProcessResource.Texture,
+		NeuralPostProcessResource.InputBuffer,
+		NeuralPostProcessResource.InputBufferDimension,
+		NeuralPostProcessResource.SourceTypeBuffer);
+}
+
+void SetupNeuralPostProcessShaderParametersForRead(FNeuralPostProcessShaderParameters& NeuralPostProcessShaderParameters, 
+	FRDGBuilder& GraphBuilder, 
+	FRDGTextureRef Texture, 
+	FRDGBufferRef Buffer, 
+	FVector4f Dimension
+)
+{
+	if (Texture && Buffer)
+	{
+		NeuralPostProcessShaderParameters.NeuralTexture = GraphBuilder.CreateSRV(FRDGTextureSRVDesc(Texture));
+		NeuralPostProcessShaderParameters.OutputNeuralBuffer = GraphBuilder.CreateSRV(Buffer, PF_R32_FLOAT);
+		NeuralPostProcessShaderParameters.OutputNeuralBufferDimension = Dimension;
+	}
+}
+
+void SetupNeuralPostProcessShaderParametersForRead(
+	FNeuralPostProcessShaderParameters& NeuralPostProcessShaderParameters, 
+	FRDGBuilder& GraphBuilder, 
+	const FNeuralPostProcessResource& NeuralPostProcessResource
+)
+{
+	SetupNeuralPostProcessShaderParametersForRead(
+		NeuralPostProcessShaderParameters,
+		GraphBuilder,
+		NeuralPostProcessResource.Texture,
+		NeuralPostProcessResource.OutputBuffer,
+		NeuralPostProcessResource.OutputBufferDimension);
+}
+
+FNeuralPostProcessResource AllocateNeuralPostProcessingResourcesIfNeeded(
+	FRDGBuilder& GraphBuilder,
+	const FScreenPassTextureViewport& OutputViewport,
+	int32 NeuralProfileId,
+	bool bUsedWithNeuralPass)
+{
+	FNeuralPostProcessResource NeuralPostProcessResource;
+	
+	if (bUsedWithNeuralPass && GNeuralPostProcess)
+	{
+		const FRDGTextureDesc TextureDesc = FRDGTextureDesc::Create2D(
+			OutputViewport.Extent,
+			PF_FloatRGBA,
+			FClearValueBinding(),
+			TexCreate_RenderTargetable | TexCreate_ShaderResource | TexCreate_UAV);
+
+		// Allocate neural buffer and texture
+		NeuralPostProcessResource.SourceTypeBuffer =
+			GraphBuilder.CreateBuffer(FRDGBufferDesc::CreateBufferDesc(sizeof(uint32), 1), TEXT("NeuralPostprocess.SourceType"));
+		NeuralPostProcessResource.Texture = GraphBuilder.CreateTexture(TextureDesc, TEXT("NeuralPostprocess.Texture"));
+		NeuralPostProcessResource.NeuralProfileId = NeuralProfileId;
+
+		// Allocate and get the input buffer for the profile
+		GNeuralPostProcess->AllocateBuffer(
+			GraphBuilder,
+			OutputViewport,
+			NeuralProfileId,
+			NeuralPostProcessResource.InputBuffer,
+			NeuralPostProcessResource.InputBufferDimension);
+	}
+
+	return NeuralPostProcessResource;
+}
+
+void ApplyNeuralPostProcess(
+	FRDGBuilder& GraphBuilder, 
+	const FViewInfo& View, 
+	FIntRect Rect,
+	FNeuralPostProcessResource& NeuralPostProcessResource
+)
+{
+	if (IsNeuralPostProcessEnabled() && NeuralPostProcessResource.IsValid())
+	{
+		GNeuralPostProcess->Apply(
+			GraphBuilder,
+			View,
+			NeuralPostProcessResource.NeuralProfileId,
+			NeuralPostProcessResource.Texture,
+			Rect,
+			NeuralPostProcessResource.SourceTypeBuffer,
+			NeuralPostProcessResource.OutputBuffer,
+			NeuralPostProcessResource.OutputBufferDimension);
+	}
+}
+
diff --git a/Engine/Source/Runtime/Renderer/Private/PostProcess/NeuralPostProcess.h b/Engine/Source/Runtime/Renderer/Private/PostProcess/NeuralPostProcess.h
new file mode 100644
--- /dev/null
+++ b/Engine/Source/Runtime/Renderer/Private/PostProcess/NeuralPostProcess.h
@@ -0,0 +1,82 @@
+// Copyright Epic Games, Inc. All Rights Reserved.
+
+#pragma once
+
+#include "ScreenPass.h"
+#include "OverridePassSequence.h"
+#include "PostProcess/PostProcessMaterialInputs.h"
+
+BEGIN_SHADER_PARAMETER_STRUCT(FNeuralPostProcessShaderParameters, )
+	SHADER_PARAMETER_RDG_TEXTURE_UAV(RWTexture2D<float4>, RWNeuralTexture)
+	SHADER_PARAMETER_RDG_TEXTURE_SRV(Texture2D, NeuralTexture)
+	SHADER_PARAMETER_RDG_BUFFER_UAV(RWBuffer<float>, InputNeuralBuffer)
+	SHADER_PARAMETER_RDG_BUFFER_UAV(RWBuffer<uint>, NeuralSourceType)
+	SHADER_PARAMETER_RDG_BUFFER_SRV(Buffer<float>, OutputNeuralBuffer)
+	SHADER_PARAMETER(FVector4f, InputNeuralBufferDimension)
+	SHADER_PARAMETER(FVector4f, OutputNeuralBufferDimension)
+END_SHADER_PARAMETER_STRUCT()
+
+
+FNeuralPostProcessShaderParameters GetDefaultNeuralPostProcessShaderParameters(
+	FRDGBuilder& GraphBuilder
+);
+
+struct FNeuralPostProcessResource
+{
+	FRDGTextureRef Texture = nullptr;
+	FRDGBufferRef InputBuffer = nullptr;
+	FRDGBufferRef OutputBuffer = nullptr;
+	FRDGBufferRef SourceTypeBuffer = nullptr;
+	
+	FVector4f InputBufferDimension;
+	FVector4f OutputBufferDimension;
+
+	int32 NeuralProfileId = INDEX_NONE;
+
+	bool IsValid() const { return InputBuffer != nullptr && NeuralProfileId != INDEX_NONE;}
+};
+
+void SetupNeuralPostProcessShaderParametersForWrite(
+	FNeuralPostProcessShaderParameters& NeuralPostProcessShaderParameters,
+	FRDGBuilder& GraphBuilder,
+	const FNeuralPostProcessResource& NeuralPostProcessResource);
+
+void SetupNeuralPostProcessShaderParametersForRead(
+	FNeuralPostProcessShaderParameters& NeuralPostProcessShaderParameters,
+	FRDGBuilder& GraphBuilder,
+	const FNeuralPostProcessResource& NeuralPostProcessResource);
+
+FNeuralPostProcessResource AllocateNeuralPostProcessingResourcesIfNeeded(
+	FRDGBuilder& GraphBuilder,
+	const class FScreenPassTextureViewport& OutputViewport, 
+	int32 NeuralProfileId,
+	bool bUsedWithNeuralNetworks
+);
+
+// Apply neural post process to the resource
+void ApplyNeuralPostProcess(
+	FRDGBuilder& GraphBuilder,
+	const FViewInfo& View,
+	FIntRect Rect,
+	FNeuralPostProcessResource& NeuralPostProcessResource
+);
+
+bool IsNeuralPostProcessEnabled();
+bool ShouldApplyNeuralPostProcessForMaterial(const class FMaterial* Material);
+bool IsNeuralPostProcessShaderParameterUsed(FNeuralPostProcessShaderParameters& NeuralPostProcessShaderParameters);
+
+class INeuralPostProcessInterface
+{
+public:
+
+	virtual ~INeuralPostProcessInterface() {}
+
+	virtual void Apply(FRDGBuilder& GraphBuilder, const FViewInfo& View, int32 NeuralProfileId,
+		FRDGTexture* NeuralTexture, FIntRect ViewRect, FRDGBufferRef InputSourceType,
+		FRDGBufferRef& OutputNeuralBuffer, FVector4f& BufferDimension) = 0;
+
+	virtual void AllocateBuffer(FRDGBuilder& GraphBuilder, const FScreenPassTextureViewport& Viewport,
+		int32 NeuralProfileId, FRDGBufferRef& InputNeuralBuffer, FVector4f& InputBufferDimension) = 0;
+};
+
+extern RENDERER_API TUniquePtr<INeuralPostProcessInterface> GNeuralPostProcess;
diff --git a/Engine/Source/Runtime/Renderer/Private/PostProcess/PostProcessMaterial.cpp b/Engine/Source/Runtime/Renderer/Private/PostProcess/PostProcessMaterial.cpp
--- a/Engine/Source/Runtime/Renderer/Private/PostProcess/PostProcessMaterial.cpp
+++ b/Engine/Source/Runtime/Renderer/Private/PostProcess/PostProcessMaterial.cpp
@@ -27,6 +27,7 @@
 #include "SystemTextures.h"
 #include "Substrate/Substrate.h"
 #include "SingleLayerWaterRendering.h"
+#include "Engine/NeuralProfile.h"
 
 namespace
 {
@@ -290,7 +291,8 @@
 	DECLARE_SHADER_TYPE(FPostProcessMaterialPS, Material);
 
 	class FManualStencilTestDim : SHADER_PERMUTATION_BOOL("MANUAL_STENCIL_TEST");
-	using FPermutationDomain = TShaderPermutationDomain<FManualStencilTestDim>;
+	class FNeuralPostProcessPrePass : SHADER_PERMUTATION_BOOL("NEURAL_POSTPROCESS_PREPASS");
+	using FPermutationDomain = TShaderPermutationDomain<FManualStencilTestDim,FNeuralPostProcessPrePass>;
 
 	FPostProcessMaterialPS() = default;
 	FPostProcessMaterialPS(const ShaderMetaType::CompiledShaderInitializerType& Initializer)
@@ -347,7 +349,8 @@
 	const FMaterialRenderProxy*& OutMaterialProxy,
 	const FMaterialShaderMap*& OutMaterialShaderMap,
 	TShaderRef<FPostProcessMaterialVS>& OutVertexShader,
-	TShaderRef<FPostProcessMaterialPS>& OutPixelShader)
+	TShaderRef<FPostProcessMaterialPS>& OutPixelShader,
+	bool bNeuralPostProcessPrepass = false)
 {
 	const FMaterialRenderProxy* MaterialProxy = InMaterialInterface->GetRenderProxy();
 	check(MaterialProxy);
@@ -365,6 +368,7 @@
 
 				FPostProcessMaterialPS::FPermutationDomain PermutationVectorPS;
 				PermutationVectorPS.Set<FPostProcessMaterialPS::FManualStencilTestDim>(bManualStencilTest);
+				PermutationVectorPS.Set<FPostProcessMaterialPS::FNeuralPostProcessPrePass>(bNeuralPostProcessPrepass);
 
 				ShaderTypes.AddShaderType<FPostProcessMaterialVS>();
 				ShaderTypes.AddShaderType<FPostProcessMaterialPS>(PermutationVectorPS.ToDimensionValueId());
@@ -420,6 +424,262 @@
 	AddDrawScreenPass(GraphBuilder, RDG_EVENT_NAME("MobileMSAADecodeAndDrawTexture"), View, OutputViewport, InputViewport, PixelShader, Parameters);
 }
 
+FPostProcessMaterialParameters* GetPostProcessMaterialParameters(
+	FRDGBuilder& GraphBuilder, 
+	const FPostProcessMaterialInputs& Inputs, 
+	const FViewInfo& View,
+	const FScreenPassTextureViewport& OutputViewport,
+	FScreenPassRenderTarget& Output, 
+	FRDGTextureRef DepthStencilTexture, 
+	const uint32 MaterialStencilRef, 
+	const FMaterial* Material, 
+	const FMaterialShaderMap* MaterialShaderMap)
+{
+	FPostProcessMaterialParameters* PostProcessMaterialParameters = GraphBuilder.AllocParameters<FPostProcessMaterialParameters>();
+	PostProcessMaterialParameters->SceneTextures = Inputs.SceneTextures;
+	PostProcessMaterialParameters->View = View.ViewUniformBuffer;
+	PostProcessMaterialParameters->EyeAdaptationBuffer = GraphBuilder.CreateSRV(GetEyeAdaptationBuffer(GraphBuilder, View));
+	PostProcessMaterialParameters->PostProcessOutput = GetScreenPassTextureViewportParameters(OutputViewport);
+	PostProcessMaterialParameters->RenderTargets[0] = Output.GetRenderTargetBinding();
+
+	// The target color will be decoded if bForceIntermediateTarget is true in any case, but we might still need to decode the input color
+	PostProcessMaterialParameters->bMetalMSAAHDRDecode = Inputs.bMetalMSAAHDRDecode ? 1 : 0;
+
+	if (DepthStencilTexture)
+	{
+		PostProcessMaterialParameters->RenderTargets.DepthStencil = FDepthStencilBinding(
+			DepthStencilTexture,
+			ERenderTargetLoadAction::ELoad,
+			ERenderTargetLoadAction::ELoad,
+			FExclusiveDepthStencil::DepthRead_StencilRead);
+	}
+	PostProcessMaterialParameters->ManualStencilReferenceValue = MaterialStencilRef;
+	PostProcessMaterialParameters->ManualStencilTestMask = GetManualStencilTestMask(Material->GetStencilCompare());
+
+	PostProcessMaterialParameters->PostProcessInput_BilinearSampler = TStaticSamplerState<SF_Bilinear, AM_Clamp, AM_Clamp, AM_Clamp>::GetRHI();;
+
+	const FScreenPassTexture BlackDummy(GSystemTextures.GetBlackDummy(GraphBuilder));
+
+	// This gets passed in whether or not it's used.
+	GraphBuilder.RemoveUnusedTextureWarning(BlackDummy.Texture);
+
+	FRHISamplerState* PointClampSampler = TStaticSamplerState<SF_Point, AM_Clamp, AM_Clamp, AM_Clamp>::GetRHI();
+
+	for (uint32 InputIndex = 0; InputIndex < kPostProcessMaterialInputCountMax; ++InputIndex)
+	{
+		FScreenPassTexture Input = Inputs.GetInput((EPostProcessMaterialInput)InputIndex);
+
+		// Need to provide valid textures for when shader compilation doesn't cull unused parameters.
+		if (!Input.Texture || !MaterialShaderMap->UsesSceneTexture(PPI_PostProcessInput0 + InputIndex))
+		{
+			Input = BlackDummy;
+		}
+
+		PostProcessMaterialParameters->PostProcessInput[InputIndex] = GetScreenPassTextureInput(Input, PointClampSampler);
+	}
+
+	// Path tracing buffer textures
+	for (uint32 InputIndex = 0; InputIndex < kPathTracingPostProcessMaterialInputCountMax; ++InputIndex)
+	{
+		FScreenPassTexture Input = Inputs.GetPathTracingInput((EPathTracingPostProcessMaterialInput)InputIndex);
+
+		if (!Input.Texture || !MaterialShaderMap->UsesPathTracingBufferTexture(InputIndex))
+		{
+			Input = BlackDummy;
+		}
+
+		PostProcessMaterialParameters->PathTracingPostProcessInput[InputIndex] = GetScreenPassTextureInput(Input, PointClampSampler);
+	}
+
+	PostProcessMaterialParameters->Substrate = Substrate::BindSubstrateGlobalUniformParameters(View);
+
+	// SceneDepthWithoutWater
+	const bool bHasValidSceneDepthWithoutWater = Inputs.SceneWithoutWaterTextures && Inputs.SceneWithoutWaterTextures->DepthTexture;
+	const bool bShouldUseBilinearSamplerForDepth = bHasValidSceneDepthWithoutWater && ShouldUseBilinearSamplerForDepthWithoutSingleLayerWater(Inputs.SceneWithoutWaterTextures->DepthTexture->Desc.Format);
+	PostProcessMaterialParameters->bSceneDepthWithoutWaterTextureAvailable = bHasValidSceneDepthWithoutWater;
+	PostProcessMaterialParameters->SceneDepthWithoutSingleLayerWaterSampler = bShouldUseBilinearSamplerForDepth ? TStaticSamplerState<SF_Bilinear>::GetRHI() : TStaticSamplerState<SF_Point>::GetRHI();
+	PostProcessMaterialParameters->SceneDepthWithoutSingleLayerWaterTexture = FRDGSystemTextures::Get(GraphBuilder).Black;
+	PostProcessMaterialParameters->SceneWithoutSingleLayerWaterMinMaxUV = FVector4f(0.0f, 0.0f, 1.0f, 1.0f);
+	PostProcessMaterialParameters->SceneWithoutSingleLayerWaterTextureSize = FVector2f(0.0f, 0.0f);
+	PostProcessMaterialParameters->SceneWithoutSingleLayerWaterInvTextureSize = FVector2f(0.0f, 0.0f);
+	if (bHasValidSceneDepthWithoutWater)
+	{
+		const bool bIsInstancedStereoSideBySide = View.bIsInstancedStereoEnabled && !View.bIsMobileMultiViewEnabled && IStereoRendering::IsStereoEyeView(View);
+		int32 WaterViewIndex = INDEX_NONE;
+		if (bIsInstancedStereoSideBySide)
+		{
+			WaterViewIndex = View.PrimaryViewIndex; // The instanced view does not have MinMaxUV initialized, instead the primary view MinMaxUV covers both eyes
+		}
+		else
+		{
+			verify(View.Family->Views.Find(&View, WaterViewIndex));
+		}
+
+		PostProcessMaterialParameters->SceneDepthWithoutSingleLayerWaterTexture = Inputs.SceneWithoutWaterTextures->DepthTexture;
+		PostProcessMaterialParameters->SceneWithoutSingleLayerWaterMinMaxUV = Inputs.SceneWithoutWaterTextures->Views[WaterViewIndex].MinMaxUV;
+
+		const FIntVector DepthTextureSize = Inputs.SceneWithoutWaterTextures->DepthTexture->Desc.GetSize();
+		PostProcessMaterialParameters->SceneWithoutSingleLayerWaterTextureSize = FVector2f(DepthTextureSize.X, DepthTextureSize.Y);
+		PostProcessMaterialParameters->SceneWithoutSingleLayerWaterInvTextureSize = FVector2f(1.0f / DepthTextureSize.X, 1.0f / DepthTextureSize.Y);
+	}
+
+	PostProcessMaterialParameters->NeuralPostProcessParameters = GetDefaultNeuralPostProcessShaderParameters(GraphBuilder);
+
+	return PostProcessMaterialParameters;
+}
+
+void AddNeuralPostProcessPass(
+	FRDGBuilder& GraphBuilder,
+	const FViewInfo& View,
+	const FPostProcessMaterialInputs& Inputs,
+	const UMaterialInterface* MaterialInterface,
+	FNeuralPostProcessResource& NeuralPostProcessResource)
+{
+	Inputs.Validate();
+
+	const FScreenPassTexture SceneColor = Inputs.GetInput(EPostProcessMaterialInput::SceneColor);
+
+	const ERHIFeatureLevel::Type FeatureLevel = View.GetFeatureLevel();
+
+	const FMaterial* Material = nullptr;
+	const FMaterialRenderProxy* MaterialRenderProxy = nullptr;
+	const FMaterialShaderMap* MaterialShaderMap = nullptr;
+	TShaderRef<FPostProcessMaterialVS> NeuralPostProcessPassVertexShader;
+	TShaderRef<FPostProcessMaterialPS> NeuralPostProcessPassPixelShader;
+	GetMaterialInfo(MaterialInterface, FeatureLevel, Inputs, Material, MaterialRenderProxy, MaterialShaderMap, NeuralPostProcessPassVertexShader, NeuralPostProcessPassPixelShader, true);
+	
+	check(NeuralPostProcessPassVertexShader.IsValid());
+	check(NeuralPostProcessPassPixelShader.IsValid());
+
+	int32 NeuralProfileId = Material->GetNeuralProfileId();
+
+	FRHIDepthStencilState* DefaultDepthStencilState = FScreenPassPipelineState::FDefaultDepthStencilState::GetRHI();
+	FRHIDepthStencilState* DepthStencilState = DefaultDepthStencilState;
+
+	FRDGTextureRef DepthStencilTexture = nullptr;
+
+	// Allocate custom depth stencil texture(s) and depth stencil state.
+	const EMaterialCustomDepthPolicy CustomStencilPolicy = GetMaterialCustomDepthPolicy(Material);
+
+	if (CustomStencilPolicy == EMaterialCustomDepthPolicy::Enabled &&
+		!Inputs.bManualStencilTest &&
+		HasBeenProduced(Inputs.CustomDepthTexture))
+	{
+		check(Inputs.CustomDepthTexture);
+		DepthStencilTexture = Inputs.CustomDepthTexture;
+		DepthStencilState = GetMaterialStencilState(Material);
+	}
+
+	FRHIBlendState* DefaultBlendState = FScreenPassPipelineState::FDefaultBlendState::GetRHI();
+	FRHIBlendState* BlendState = DefaultBlendState;
+
+	if (IsMaterialBlendEnabled(Material))
+	{
+		BlendState = GetMaterialBlendState(Material);
+	}
+
+	FScreenPassRenderTarget Output = Inputs.OverrideOutput;
+	// Create a new texture instead of reusing the scene color output in the pre pass. Should not pollute the scene color texture.
+	{
+		// Allocate new transient output texture.
+		{
+			FRDGTextureDesc OutputDesc = SceneColor.Texture->Desc;
+			OutputDesc.Reset();
+			if (Inputs.OutputFormat != PF_Unknown)
+			{
+				OutputDesc.Format = Inputs.OutputFormat;
+			}
+			OutputDesc.ClearValue = FClearValueBinding(FLinearColor::Black);
+			OutputDesc.Flags &= (~ETextureCreateFlags::FastVRAM);
+			OutputDesc.Flags |= GFastVRamConfig.PostProcessMaterial;
+
+			Output = FScreenPassRenderTarget(GraphBuilder.CreateTexture(OutputDesc, TEXT("PostProcessTempOutput")), SceneColor.ViewRect, View.GetOverwriteLoadAction());
+		}
+	}
+
+	const FScreenPassTextureViewport SceneColorViewport(SceneColor);
+	const FScreenPassTextureViewport OutputViewport(Output);
+
+	RDG_EVENT_SCOPE(GraphBuilder, "PostProcessMaterial::NeuralPass");
+
+	const uint32 MaterialStencilRef = Material->GetStencilRefValue();
+
+	const bool bMobilePlatform = IsMobilePlatform(View.GetShaderPlatform());
+
+
+	EScreenPassDrawFlags ScreenPassFlags = EScreenPassDrawFlags::AllowHMDHiddenAreaMask;
+
+	// check if we can skip that draw call in case if all pixels will fail the stencil test of the material
+	bool bSkipPostProcess = false;
+
+	if (Material->IsStencilTestEnabled() && IsPostProcessStencilTestAllowed())
+	{
+		bool bFailStencil = true;
+
+		const uint32 StencilComp = Material->GetStencilCompare();
+
+		// Always check against clear value, since a material might want to perform operations against that value
+		const uint32 StencilClearValue = Inputs.CustomDepthTexture ? Inputs.CustomDepthTexture->Desc.ClearValue.Value.DSValue.Stencil : 0;
+		bFailStencil &= PostProcessStencilTest(StencilClearValue, StencilComp, MaterialStencilRef);
+
+
+		for (const uint32& Value : View.CustomDepthStencilValues)
+		{
+			bFailStencil &= PostProcessStencilTest(Value, StencilComp, MaterialStencilRef);
+
+			if (!bFailStencil)
+			{
+				break;
+			}
+		}
+
+		bSkipPostProcess = bFailStencil;
+	}
+
+	if (!bSkipPostProcess)
+	{
+		NeuralPostProcessResource = AllocateNeuralPostProcessingResourcesIfNeeded(
+			GraphBuilder, OutputViewport, NeuralProfileId, Material->IsUsedWithNeuralNetworks());
+
+		if (NeuralPostProcessResource.IsValid())
+		{ 
+			// Prepass to extract the input to the NNE Engine
+			FPostProcessMaterialParameters* PostProcessMaterialParameters =
+				GetPostProcessMaterialParameters(GraphBuilder, Inputs, View, OutputViewport, Output, DepthStencilTexture, MaterialStencilRef, Material, MaterialShaderMap);
+
+			SetupNeuralPostProcessShaderParametersForWrite(PostProcessMaterialParameters->NeuralPostProcessParameters, GraphBuilder, NeuralPostProcessResource);
+
+			ClearUnusedGraphResources(NeuralPostProcessPassVertexShader, NeuralPostProcessPassPixelShader, PostProcessMaterialParameters);
+
+			//Only call the neural network when the shader resource is actually used.
+			if (IsNeuralPostProcessShaderParameterUsed(PostProcessMaterialParameters->NeuralPostProcessParameters))
+			{
+				AddDrawScreenPass(
+					GraphBuilder,
+#if RDG_EVENTS != RDG_EVENTS_STRING_COPY
+					RDG_EVENT_NAME("PostProcessMaterial(Neural Prepass)"),
+#else
+					FRDGEventName(*Material->GetAssetName()),
+#endif
+					View,
+					OutputViewport,
+					SceneColorViewport,
+					// Uses default depth stencil on mobile since the stencil test is done in pixel shader.
+					FScreenPassPipelineState(NeuralPostProcessPassVertexShader, NeuralPostProcessPassPixelShader, BlendState, DepthStencilState, MaterialStencilRef),
+					PostProcessMaterialParameters,
+					ScreenPassFlags,
+					[&View, NeuralPostProcessPassVertexShader, NeuralPostProcessPassPixelShader, MaterialRenderProxy, Material, PostProcessMaterialParameters](FRHICommandList& RHICmdList)
+					{
+						SetShaderParametersMixedVS(RHICmdList, NeuralPostProcessPassVertexShader, *PostProcessMaterialParameters, View, MaterialRenderProxy, *Material);
+						SetShaderParametersMixedPS(RHICmdList, NeuralPostProcessPassPixelShader, *PostProcessMaterialParameters, View, MaterialRenderProxy, *Material);
+					});
+
+				ApplyNeuralPostProcess(GraphBuilder, View, Output.ViewRect, NeuralPostProcessResource);
+			}
+		}
+	}
+}
+
 FScreenPassTexture AddPostProcessMaterialPass(
 	FRDGBuilder& GraphBuilder,
 	const FViewInfo& View,
@@ -541,96 +801,6 @@
 
 	const bool bMobilePlatform = IsMobilePlatform(View.GetShaderPlatform());
 
-	FPostProcessMaterialParameters* PostProcessMaterialParameters = GraphBuilder.AllocParameters<FPostProcessMaterialParameters>();
-	PostProcessMaterialParameters->SceneTextures = Inputs.SceneTextures;
-	PostProcessMaterialParameters->View = View.ViewUniformBuffer;
-	PostProcessMaterialParameters->EyeAdaptationBuffer = GraphBuilder.CreateSRV(GetEyeAdaptationBuffer(GraphBuilder, View));
-	PostProcessMaterialParameters->PostProcessOutput = GetScreenPassTextureViewportParameters(OutputViewport);
-	PostProcessMaterialParameters->RenderTargets[0] = Output.GetRenderTargetBinding();
-
-	// The target color will be decoded if bForceIntermediateTarget is true in any case, but we might still need to decode the input color
-	PostProcessMaterialParameters->bMetalMSAAHDRDecode = Inputs.bMetalMSAAHDRDecode ? 1 : 0;
-
-	if (DepthStencilTexture)
-	{
-		PostProcessMaterialParameters->RenderTargets.DepthStencil = FDepthStencilBinding(
-			DepthStencilTexture,
-			ERenderTargetLoadAction::ELoad,
-			ERenderTargetLoadAction::ELoad,
-			FExclusiveDepthStencil::DepthRead_StencilRead);
-	}
-	PostProcessMaterialParameters->ManualStencilReferenceValue = MaterialStencilRef;
-	PostProcessMaterialParameters->ManualStencilTestMask = GetManualStencilTestMask(Material->GetStencilCompare());
-
-	PostProcessMaterialParameters->PostProcessInput_BilinearSampler = TStaticSamplerState<SF_Bilinear, AM_Clamp, AM_Clamp, AM_Clamp>::GetRHI();;
-
-	const FScreenPassTexture BlackDummy(GSystemTextures.GetBlackDummy(GraphBuilder));
-
-    // This gets passed in whether or not it's used.
-	GraphBuilder.RemoveUnusedTextureWarning(BlackDummy.Texture);
-
-	FRHISamplerState* PointClampSampler = TStaticSamplerState<SF_Point, AM_Clamp, AM_Clamp, AM_Clamp>::GetRHI();
-
-	for (uint32 InputIndex = 0; InputIndex < kPostProcessMaterialInputCountMax; ++InputIndex)
-	{
-		FScreenPassTexture Input = Inputs.GetInput((EPostProcessMaterialInput)InputIndex);
-
-		// Need to provide valid textures for when shader compilation doesn't cull unused parameters.
-		if (!Input.Texture || !MaterialShaderMap->UsesSceneTexture(PPI_PostProcessInput0 + InputIndex))
-		{
-			Input = BlackDummy;
-		}
-
-		PostProcessMaterialParameters->PostProcessInput[InputIndex] = GetScreenPassTextureInput(Input, PointClampSampler);
-	}
-
-	// Path tracing buffer textures
-	for (uint32 InputIndex = 0; InputIndex < kPathTracingPostProcessMaterialInputCountMax; ++InputIndex)
-	{
-		FScreenPassTexture Input = Inputs.GetPathTracingInput((EPathTracingPostProcessMaterialInput)InputIndex);
-
-		if (!Input.Texture || !MaterialShaderMap->UsesPathTracingBufferTexture(InputIndex))
-		{
-			Input = BlackDummy;
-		}
-
-		PostProcessMaterialParameters->PathTracingPostProcessInput[InputIndex] = GetScreenPassTextureInput(Input, PointClampSampler);
-	}
-
-	PostProcessMaterialParameters->Substrate = Substrate::BindSubstrateGlobalUniformParameters(View);
-
-	// SceneDepthWithoutWater
-	const bool bHasValidSceneDepthWithoutWater = Inputs.SceneWithoutWaterTextures && Inputs.SceneWithoutWaterTextures->DepthTexture;
-	const bool bShouldUseBilinearSamplerForDepth = bHasValidSceneDepthWithoutWater && ShouldUseBilinearSamplerForDepthWithoutSingleLayerWater(Inputs.SceneWithoutWaterTextures->DepthTexture->Desc.Format);
-	PostProcessMaterialParameters->bSceneDepthWithoutWaterTextureAvailable = bHasValidSceneDepthWithoutWater;
-	PostProcessMaterialParameters->SceneDepthWithoutSingleLayerWaterSampler = bShouldUseBilinearSamplerForDepth ? TStaticSamplerState<SF_Bilinear>::GetRHI() : TStaticSamplerState<SF_Point>::GetRHI();
-	PostProcessMaterialParameters->SceneDepthWithoutSingleLayerWaterTexture = FRDGSystemTextures::Get(GraphBuilder).Black;
-	PostProcessMaterialParameters->SceneWithoutSingleLayerWaterMinMaxUV = FVector4f(0.0f, 0.0f, 1.0f, 1.0f);
-	PostProcessMaterialParameters->SceneWithoutSingleLayerWaterTextureSize = FVector2f(0.0f, 0.0f);
-	PostProcessMaterialParameters->SceneWithoutSingleLayerWaterInvTextureSize = FVector2f(0.0f, 0.0f);
-	if (bHasValidSceneDepthWithoutWater)
-	{
-		const bool bIsInstancedStereoSideBySide = View.bIsInstancedStereoEnabled && !View.bIsMobileMultiViewEnabled && IStereoRendering::IsStereoEyeView(View);
-		int32 WaterViewIndex = INDEX_NONE;
-		if (bIsInstancedStereoSideBySide)
-		{
-			WaterViewIndex = View.PrimaryViewIndex; // The instanced view does not have MinMaxUV initialized, instead the primary view MinMaxUV covers both eyes
-		}
-		else
-		{
-			verify(View.Family->Views.Find(&View, WaterViewIndex));
-		}
-
-		PostProcessMaterialParameters->SceneDepthWithoutSingleLayerWaterTexture = Inputs.SceneWithoutWaterTextures->DepthTexture;
-		PostProcessMaterialParameters->SceneWithoutSingleLayerWaterMinMaxUV = Inputs.SceneWithoutWaterTextures->Views[WaterViewIndex].MinMaxUV;
-
-		const FIntVector DepthTextureSize = Inputs.SceneWithoutWaterTextures->DepthTexture->Desc.GetSize();
-		PostProcessMaterialParameters->SceneWithoutSingleLayerWaterTextureSize = FVector2f(DepthTextureSize.X, DepthTextureSize.Y);
-		PostProcessMaterialParameters->SceneWithoutSingleLayerWaterInvTextureSize = FVector2f(1.0f / DepthTextureSize.X, 1.0f / DepthTextureSize.Y);
-	}
-
-	ClearUnusedGraphResources(VertexShader, PixelShader, PostProcessMaterialParameters);
-
 	EScreenPassDrawFlags ScreenPassFlags = EScreenPassDrawFlags::AllowHMDHiddenAreaMask;
 
 	// check if we can skip that draw call in case if all pixels will fail the stencil test of the material
@@ -662,25 +832,46 @@
 
 	if (!bSkipPostProcess)
 	{
-		AddDrawScreenPass(
-			GraphBuilder,
+
+		FNeuralPostProcessResource NeuralPostProcessResource;
+		const bool bShouldApplyNeuralPostProcessing = ShouldApplyNeuralPostProcessForMaterial(Material);
+
+		if (bShouldApplyNeuralPostProcessing)
+		{
+			AddNeuralPostProcessPass(GraphBuilder, View, Inputs, MaterialInterface, NeuralPostProcessResource);
+		}
+
+		{
+			FPostProcessMaterialParameters* PostProcessMaterialParameters =
+				GetPostProcessMaterialParameters(GraphBuilder, Inputs, View, OutputViewport, Output, DepthStencilTexture, MaterialStencilRef, Material, MaterialShaderMap);
+
+			if (bShouldApplyNeuralPostProcessing)
+			{
+				SetupNeuralPostProcessShaderParametersForRead(PostProcessMaterialParameters->NeuralPostProcessParameters, GraphBuilder, NeuralPostProcessResource);
+			}
+
+			ClearUnusedGraphResources(VertexShader, PixelShader, PostProcessMaterialParameters);
+
+			AddDrawScreenPass(
+				GraphBuilder,
 #if RDG_EVENTS != RDG_EVENTS_STRING_COPY
-			RDG_EVENT_NAME("PostProcessMaterial"),
+				RDG_EVENT_NAME("PostProcessMaterial"),
 #else
-			FRDGEventName(*Material->GetAssetName()),
+				FRDGEventName(*Material->GetAssetName()),
 #endif
-			View,
-			OutputViewport,
-			SceneColorViewport,
-			// Uses default depth stencil on mobile since the stencil test is done in pixel shader.
-			FScreenPassPipelineState(VertexShader, PixelShader, BlendState, DepthStencilState, MaterialStencilRef),
-			PostProcessMaterialParameters,
-			ScreenPassFlags,
-			[&View, VertexShader, PixelShader, MaterialRenderProxy, Material, PostProcessMaterialParameters](FRHICommandList& RHICmdList)
-			{
-				SetShaderParametersMixedVS(RHICmdList, VertexShader, *PostProcessMaterialParameters, View, MaterialRenderProxy, *Material);
-				SetShaderParametersMixedPS(RHICmdList, PixelShader, *PostProcessMaterialParameters, View, MaterialRenderProxy, *Material);
-			});
+				View,
+				OutputViewport,
+				SceneColorViewport,
+				// Uses default depth stencil on mobile since the stencil test is done in pixel shader.
+				FScreenPassPipelineState(VertexShader, PixelShader, BlendState, DepthStencilState, MaterialStencilRef),
+				PostProcessMaterialParameters,
+				ScreenPassFlags,
+				[&View, VertexShader, PixelShader, MaterialRenderProxy, Material, PostProcessMaterialParameters](FRHICommandList& RHICmdList)
+				{
+					SetShaderParametersMixedVS(RHICmdList, VertexShader, *PostProcessMaterialParameters, View, MaterialRenderProxy, *Material);
+					SetShaderParametersMixedPS(RHICmdList, PixelShader, *PostProcessMaterialParameters, View, MaterialRenderProxy, *Material);
+				});
+		}
 
 		if (bForceIntermediateTarget && !bCompositeWithInputAndDecode)
 		{
diff --git a/Engine/Source/Runtime/Renderer/Private/PostProcess/PostProcessMaterial.h b/Engine/Source/Runtime/Renderer/Private/PostProcess/PostProcessMaterial.h
--- a/Engine/Source/Runtime/Renderer/Private/PostProcess/PostProcessMaterial.h
+++ b/Engine/Source/Runtime/Renderer/Private/PostProcess/PostProcessMaterial.h
@@ -6,6 +6,7 @@
 #include "OverridePassSequence.h"
 #include "Substrate/Substrate.h"
 
+#include "NeuralPostProcess.h"
 #include "PostProcess/PostProcessMaterialInputs.h"
 
 class UMaterialInterface;
@@ -18,6 +19,7 @@ BEGIN_SHADER_PARAMETER_STRUCT(FPostProcessMaterialParameters, )
 	SHADER_PARAMETER_STRUCT_REF(FViewUniformShaderParameters, View)
 	SHADER_PARAMETER_STRUCT_INCLUDE(FSceneTextureShaderParameters, SceneTextures)
 	SHADER_PARAMETER_RDG_UNIFORM_BUFFER(FSubstrateGlobalUniformParameters, Substrate)
+	SHADER_PARAMETER_STRUCT_INCLUDE(FNeuralPostProcessShaderParameters, NeuralPostProcessParameters)
 	SHADER_PARAMETER_STRUCT(FScreenPassTextureViewportParameters, PostProcessOutput)
 	SHADER_PARAMETER_STRUCT_ARRAY(FScreenPassTextureInput, PostProcessInput, [kPostProcessMaterialInputCountMax])
 	SHADER_PARAMETER_STRUCT_ARRAY(FScreenPassTextureInput, PathTracingPostProcessInput, [kPathTracingPostProcessMaterialInputCountMax])
