Disable async VT update in editor.
Temp fix to unblock users until issue with wireframe rendering is fixed.

diff --git a/Engine/Source/Runtime/Renderer/Private/VT/VirtualTextureSystem.cpp b/Engine/Source/Runtime/Renderer/Private/VT/VirtualTextureSystem.cpp
--- a/Engine/Source/Runtime/Renderer/Private/VT/VirtualTextureSystem.cpp
+++ b/Engine/Source/Runtime/Renderer/Private/VT/VirtualTextureSystem.cpp
@@ -161,7 +161,11 @@ static TAutoConsoleVariable<int32> CVarVTCsvStats(
 );
 static TAutoConsoleVariable<int32> CVarVTAsyncPageRequestTask(
 	TEXT("r.VT.AsyncPageRequestTask"),
+#if WITH_EDITOR // Turn off async VT as quick fix for race condition seen in editor wireframe mode.
+	0,
+#else
 	1,
+#endif
 	TEXT("Performs VT page requests on an async task."),
 	ECVF_RenderThreadSafe
 );
