[Backout] - CL28467774
Original CL Desc
-----------------------------------------------------------------
Add neural post processing (experimental).
* Add neural profile to config which neural network to use. Texture index is used by default where [-1,3,H,W] is supported. Use Buffer index to support [-1,C,H,W]. Need to config the batch size override or the tile size correspondingly.
* Config in the post process material which neural profile to use and enable `used with neural network`.
* Use neural input/output to control the input and output of the neural network. For neural input, [B,C, U, V] is the index, input0 takes 3 effective channels, mask determines if we should use it as input to the network. Only one node can be used per material; Multiple neural outputs can be used. For texture indexing, automatic bilinear downscaling and upscaling is applied. For Buffer indexing, the user has full control to access the value ( user needs to perform their own filtering).

diff --git a/Engine/Config/BaseEditor.ini b/Engine/Config/BaseEditor.ini
--- a/Engine/Config/BaseEditor.ini
+++ b/Engine/Config/BaseEditor.ini
@@ -64,7 +64,6 @@ bCleanupExternalObjectsEmptyFolders=false
 +RenderableThumbnailTypes=(ClassNeedingThumbnailName="/Script/Engine.SlateBrushAsset",RendererClassName="/Script/UnrealEd.SlateBrushThumbnailRenderer")
 +RenderableThumbnailTypes=(ClassNeedingThumbnailName="/Script/Engine.SubsurfaceProfile",RendererClassName="/Script/UnrealEd.SubsurfaceProfileRenderer")
 +RenderableThumbnailTypes=(ClassNeedingThumbnailName="/Script/Engine.SpecularProfile",RendererClassName="/Script/UnrealEd.SpecularProfileRenderer")
-+RenderableThumbnailTypes=(ClassNeedingThumbnailName="/Script/Engine.NeuralProfile",RendererClassName="/Script/UnrealEd.NeuralProfileRenderer")
 +RenderableThumbnailTypes=(ClassNeedingThumbnailName="/Script/Engine.PhysicsAsset",RendererClassName="/Script/UnrealEd.PhysicsAssetThumbnailRenderer")
 +RenderableThumbnailTypes=(ClassNeedingThumbnailName="/Script/Engine.CurveLinearColor",RendererClassName="/Script/UnrealEd.CurveLinearColorThumbnailRenderer")
 +RenderableThumbnailTypes=(ClassNeedingThumbnailName="/Script/Engine.CurveFloat",RendererClassName="/Script/UnrealEd.CurveFloatThumbnailRenderer")
diff --git a/Engine/Plugins/Editor/EngineAssetDefinitions/Source/Private/AssetDefinition_NeuralProfile.h b/Engine/Plugins/Editor/EngineAssetDefinitions/Source/Private/AssetDefinition_NeuralProfile.h
deleted file mode 100644
diff --git a/Engine/Plugins/Experimental/NeuralRendering/NeuralRendering.uplugin b/Engine/Plugins/Experimental/NeuralRendering/NeuralRendering.uplugin
deleted file mode 100644
diff --git a/Engine/Plugins/Experimental/NeuralRendering/Source/NeuralPostProcessing/NeuralPostProcessing.Build.cs b/Engine/Plugins/Experimental/NeuralRendering/Source/NeuralPostProcessing/NeuralPostProcessing.Build.cs
deleted file mode 100644
diff --git a/Engine/Plugins/Experimental/NeuralRendering/Source/NeuralPostProcessing/Private/NeuralPostProcessModelInstance.cpp b/Engine/Plugins/Experimental/NeuralRendering/Source/NeuralPostProcessing/Private/NeuralPostProcessModelInstance.cpp
deleted file mode 100644
diff --git a/Engine/Plugins/Experimental/NeuralRendering/Source/NeuralPostProcessing/Private/NeuralPostProcessModelInstance.h b/Engine/Plugins/Experimental/NeuralRendering/Source/NeuralPostProcessing/Private/NeuralPostProcessModelInstance.h
deleted file mode 100644
diff --git a/Engine/Plugins/Experimental/NeuralRendering/Source/NeuralPostProcessing/Private/NeuralPostProcessing.cpp b/Engine/Plugins/Experimental/NeuralRendering/Source/NeuralPostProcessing/Private/NeuralPostProcessing.cpp
deleted file mode 100644
diff --git a/Engine/Plugins/Experimental/NeuralRendering/Source/NeuralPostProcessing/Private/NeuralPostProcessingCS.cpp b/Engine/Plugins/Experimental/NeuralRendering/Source/NeuralPostProcessing/Private/NeuralPostProcessingCS.cpp
deleted file mode 100644
diff --git a/Engine/Plugins/Experimental/NeuralRendering/Source/NeuralPostProcessing/Private/NeuralPostProcessingCS.h b/Engine/Plugins/Experimental/NeuralRendering/Source/NeuralPostProcessing/Private/NeuralPostProcessingCS.h
deleted file mode 100644
diff --git a/Engine/Plugins/Experimental/NeuralRendering/Source/NeuralPostProcessing/Public/NeuralPostProcessing.h b/Engine/Plugins/Experimental/NeuralRendering/Source/NeuralPostProcessing/Public/NeuralPostProcessing.h
deleted file mode 100644
diff --git a/Engine/Plugins/Experimental/NeuralRendering/Source/NeuralPostProcessing/Shaders/NeuralPostProcessing.usf b/Engine/Plugins/Experimental/NeuralRendering/Source/NeuralPostProcessing/Shaders/NeuralPostProcessing.usf
deleted file mode 100644
diff --git a/Engine/Shaders/Private/MaterialTemplate.ush b/Engine/Shaders/Private/MaterialTemplate.ush
--- a/Engine/Shaders/Private/MaterialTemplate.ush
+++ b/Engine/Shaders/Private/MaterialTemplate.ush
@@ -2710,37 +2710,6 @@ float4 MaterialExpressionPathTracingBufferTextureLookup(FMaterialPixelParameters
 	return float4(1, 0, 0, 1);
 }
 
-float4 NeuralTextureOutput(FMaterialPixelParameters Parameters, float2 ViewportUV)
-{
-#if POST_PROCESS_MATERIAL && SHADING_PATH_DEFERRED && MATERIAL_NEURAL_POST_PROCESS
-	
-	float2 BufferUV = ViewportUV * PostProcessInput_0_UVViewportSize + PostProcessInput_0_UVViewportMin;
-	float4 TextureValue = Texture2DSampleLevel(NeuralTexture, PostProcessInput_BilinearSampler, BufferUV, 0);
-	return TextureValue;
-#else
-	return float4(0.0f, 0.0f, 0.0f, 1.0f);
-#endif
-
-}
-
-float4 NeuralBufferOutput(FMaterialPixelParameters Parameters, float4 BufferIndices)
-{
-#if POST_PROCESS_MATERIAL && SHADING_PATH_DEFERRED && MATERIAL_NEURAL_POST_PROCESS
-
-	float3 TextureValue = 0;
-	
-	int Batch = (int)BufferIndices.x;
-	int StartChannel = (int)BufferIndices.y;
-	float2 BufferWH = BufferIndices.zw;
-
-	ReadFromNeuralBuffer(Batch, StartChannel, BufferWH, TextureValue);
-
-	return float4(TextureValue, 1.0f);
-#else
-	return float4(0.0f, 0.0f, 0.0f, 1.0f);
-#endif
-}
-
 // DERIV_BASE_VALUE() is a disgusting macro to manage backwards compatibility while changing the generated materials. The existing nodes are not "aware"
 // of partial derivatives, and will use sprintf() to put the same line in both CalcPixelMaterialInputs() and CalcPixelMaterialInputsAnalyticDerivatives().
 // If we have a line like this in CalcPixelMaterialInputs():
diff --git a/Engine/Shaders/Private/NeuralPostProcessCommon.usf b/Engine/Shaders/Private/NeuralPostProcessCommon.usf
deleted file mode 100644
diff --git a/Engine/Shaders/Private/PostProcessMaterialShaders.usf b/Engine/Shaders/Private/PostProcessMaterialShaders.usf
--- a/Engine/Shaders/Private/PostProcessMaterialShaders.usf
+++ b/Engine/Shaders/Private/PostProcessMaterialShaders.usf
@@ -76,8 +76,6 @@ SamplerState PathTracingPostProcessInput_3_Sampler;
 SamplerState PathTracingPostProcessInput_4_Sampler;
 #endif
 
-#include "NeuralPostProcessCommon.usf"
-
 #define EYE_ADAPTATION_LOOSE_PARAMETERS 1
 
 #include "/Engine/Generated/Material.ush"
@@ -244,60 +242,6 @@ void MainPS(
 	// fill out other related material parameters
 	CalcMaterialParametersPost(Parameters, PixelMaterialInputs, SvPosition, true);
 
-#if NEURAL_POSTPROCESS_PREPASS && MATERIAL_NEURAL_POST_PROCESS
-
-	OutColor = 0.5f;
-
-#if HAVE_GetNeuralInput1
-
-	float4 Index = GetNeuralInput0(Parameters);
-	float3 NeuralInput0 = GetNeuralInput1(Parameters).rgb;
-	float  Mask = GetNeuralInput2(Parameters);
-	const bool bWriteToBuffer = Index.x >= 0; //Configuration the same across all pixels.
-
-	// Update the source type so we can copy on demand
-	BRANCH
-	if (all((uint2)SvPosition.xy == uint2(0, 0)))
-	{
-		const int SourceType = bWriteToBuffer ? SOURCE_TYPE_BUFFER : SOURCE_TYPE_TEXTURE;
-		UpdateNeuralProfileSourceType(SourceType);
-	}
-
-	// Write to buffer/texture based on config
-	BRANCH
-	if (Mask != 0)
-	{
-		BRANCH
-		if (bWriteToBuffer)
-		{
-			int Batch = (int)Index.x;
-			int StartChannel = (int)Index.y;
-			float2 BufferWH = Index.zw;
-			SaveToNeuralBuffer(Batch, StartChannel, BufferWH, NeuralInput0);
-		}
-		else
-		{
-			int StartChannel = (int)Index.y;
-			float2 CustomViewportUV = Index.zw;
-			float2 Position = CustomViewportUV * PostProcessOutput_ViewportSize.xy + PostProcessOutput_ViewportMin.xy;
-			if (StartChannel >= 0 && StartChannel <= 2)
-			{
-				RWNeuralTexture[(uint2)Position.xy] = float4(NeuralInput0, 0.0f);
-			}
-		}
-	}
-	
-	// This is for debug only
-	OutColor.xyz = NeuralInput0;
-#if !MATERIALBLENDING_MODULATE && POST_PROCESS_MATERIAL_BEFORE_TONEMAP
-	OutColor.xyz *= View.PreExposure;
-#endif
-
-	return;
-
-#endif // HAVE_GetNeuralInput1
-#endif // NEURAL_POSTPROCESS_PREPASS	
-
 	// In Unreal, an alpha channel value of 0.0f means that the pixel is opaque and a value of 1.0f means that there is translucency.
 	// This is the opposite of how the alpha channel is interpreted in formats like PNG. For this reason, rendered images that are exported
 	// from the engine will have their alpha channel inverted to match the PNG convention. However, for consistency inside the rendering chain,
diff --git a/Engine/Source/Editor/UnrealEd/Classes/Factories/NeuralProfileFactory.h b/Engine/Source/Editor/UnrealEd/Classes/Factories/NeuralProfileFactory.h
deleted file mode 100644
diff --git a/Engine/Source/Editor/UnrealEd/Classes/ThumbnailRendering/NeuralProfileRenderer.h b/Engine/Source/Editor/UnrealEd/Classes/ThumbnailRendering/NeuralProfileRenderer.h
deleted file mode 100644
diff --git a/Engine/Source/Editor/UnrealEd/Private/Factories/EditorFactories.cpp b/Engine/Source/Editor/UnrealEd/Private/Factories/EditorFactories.cpp
--- a/Engine/Source/Editor/UnrealEd/Private/Factories/EditorFactories.cpp
+++ b/Engine/Source/Editor/UnrealEd/Private/Factories/EditorFactories.cpp
@@ -111,7 +111,6 @@
 #include "Factories/StringTableFactory.h"
 #include "Factories/SubsurfaceProfileFactory.h"
 #include "Factories/SpecularProfileFactory.h"
-#include "Factories/NeuralProfileFactory.h"
 #include "Factories/Texture2dFactoryNew.h"
 #include "Engine/Texture.h"
 #include "Factories/TextureFactory.h"
@@ -148,7 +147,6 @@
 #include "GameFramework/DefaultPhysicsVolume.h"
 #include "Engine/SubsurfaceProfile.h"
 #include "Engine/SpecularProfile.h"
-#include "Engine/NeuralProfile.h"
 #include "Misc/ConfigCacheIni.h"
 #include "Misc/FeedbackContext.h"
 #include "GameFramework/WorldSettings.h"
@@ -7921,25 +7919,6 @@ UObject* USpecularProfileFactory::FactoryCreateNew(UClass* InClass, UObject* InP
 	return Object;
 }
 
-/*-----------------------------------------------------------------------------
-	UNeuralProfileFactory implementation.
-	-----------------------------------------------------------------------------*/
-UNeuralProfileFactory::UNeuralProfileFactory(const FObjectInitializer& ObjectInitializer)
-	: Super(ObjectInitializer)
-{
-	SupportedClass = UNeuralProfile::StaticClass();
-	bCreateNew = true;
-	bEditorImport = false;
-	bEditAfterNew = true;
-}
-
-UObject* UNeuralProfileFactory::FactoryCreateNew(UClass* Class, UObject* InParent, FName InName, EObjectFlags Flags, UObject* Context, FFeedbackContext* Warn)
-{
-	UNeuralProfile* Object = NewObject<UNeuralProfile>(InParent, InName, Flags);
-	Object->Guid = FGuid::NewGuid();
-	return Object;
-}
-
 /*-----------------------------------------------------------------------------
 	UTouchInterfaceFactory implementation.
 -----------------------------------------------------------------------------*/
diff --git a/Engine/Source/Editor/UnrealEd/Private/ThumbnailRendering/NeuralProfileRenderer.cpp b/Engine/Source/Editor/UnrealEd/Private/ThumbnailRendering/NeuralProfileRenderer.cpp
deleted file mode 100644
diff --git a/Engine/Source/Editor/UnrealEd/Public/UnrealEd.h b/Engine/Source/Editor/UnrealEd/Public/UnrealEd.h
--- a/Engine/Source/Editor/UnrealEd/Public/UnrealEd.h
+++ b/Engine/Source/Editor/UnrealEd/Public/UnrealEd.h
@@ -269,7 +269,6 @@ MONOLITHIC_HEADER_BOILERPLATE()
 #include "ThumbnailRendering/ParticleSystemThumbnailRenderer.h"
 #include "ThumbnailRendering/SubsurfaceProfileRenderer.h"
 #include "ThumbnailRendering/SpecularProfileRenderer.h"
-#include "ThumbnailRendering/NeuralProfileRenderer.h"
 #include "Factories/TextureCubeThumbnailRenderer.h"
 #include "Factories/Texture2dArrayThumbnailRenderer.h"
 #include "Editor/TransBuffer.h"
diff --git a/Engine/Source/Runtime/Engine/Classes/Engine/NeuralProfile.h b/Engine/Source/Runtime/Engine/Classes/Engine/NeuralProfile.h
deleted file mode 100644
diff --git a/Engine/Source/Runtime/Engine/Classes/Materials/Material.h b/Engine/Source/Runtime/Engine/Classes/Materials/Material.h
--- a/Engine/Source/Runtime/Engine/Classes/Materials/Material.h
+++ b/Engine/Source/Runtime/Engine/Classes/Materials/Material.h
@@ -972,15 +972,6 @@ public:
 	UPROPERTY(EditAnywhere, BlueprintReadWrite, Category = PostProcessMaterial, meta = (DisplayName = "Output Alpha"))
 	uint8 BlendableOutputAlpha : 1;
 
-	/**
-	* Indicates that the material and its instances can be used with neural network engine.
-	* This will result in the shaders required to support neural network engine being compiled which will increase shader compile time and memory usage.
-	* In addition, an additional pass will run before the postprocess pass for neural network engine, which will increase the rendering cost due to
-	* buffer copy and inference.
-	*/
-	UPROPERTY(EditAnywhere, Category = PostProcessMaterial)
-	uint8 bUsedWithNeuralNetworks : 1;
-
 	/** 
 	 * Selectively execute post process material only for pixels that pass the stencil test against the Custom Depth/Stencil buffer. 
 	 * Pixels that fail the stencil test are filled with the previous post process material output or scene color.
@@ -994,12 +985,6 @@ public:
 	UPROPERTY(EditAnywhere, Category = PostProcessMaterial, AdvancedDisplay, meta = (EditCondition = "bEnableStencilTest"))
 	uint8 StencilRefValue = 0;
 
-	/**
-	* Set by reference object cannot be modified.
-	*/
-	UPROPERTY(EditAnywhere, Category = PostProcessMaterial, AdvancedDisplay, meta = (EditCondition = "false"))
-	int8 NeuralProfileId = INDEX_NONE;
-
 	UPROPERTY()
 	TEnumAsByte<ERefractionMode> RefractionMode_DEPRECATED;
 
@@ -1155,7 +1140,6 @@ public:
 	ENGINE_API virtual USubsurfaceProfile* GetSubsurfaceProfile_Internal() const override;
 	ENGINE_API virtual uint32 NumSpecularProfile_Internal() const override;
 	ENGINE_API virtual USpecularProfile* GetSpecularProfile_Internal(uint32 Index) const override;
-	ENGINE_API virtual UNeuralProfile* GetNeuralProfile_Internal() const override;
 	ENGINE_API virtual bool CastsRayTracedShadows() const override;
 	ENGINE_API virtual FDisplacementScaling GetDisplacementScaling() const override;
 	ENGINE_API virtual float GetMaxWorldPositionOffsetDisplacement() const override;
diff --git a/Engine/Source/Runtime/Engine/Classes/Materials/MaterialExpressionNeuralPostProcessNode.h b/Engine/Source/Runtime/Engine/Classes/Materials/MaterialExpressionNeuralPostProcessNode.h
deleted file mode 100644
diff --git a/Engine/Source/Runtime/Engine/Classes/Materials/MaterialInterface.h b/Engine/Source/Runtime/Engine/Classes/Materials/MaterialInterface.h
--- a/Engine/Source/Runtime/Engine/Classes/Materials/MaterialInterface.h
+++ b/Engine/Source/Runtime/Engine/Classes/Materials/MaterialInterface.h
@@ -51,7 +51,6 @@ class UPhysicalMaterial;
 class UPhysicalMaterialMask;
 class USubsurfaceProfile;
 class USpecularProfile;
-class UNeuralProfile;
 class UTexture;
 class UMaterialInstance;
 struct FDebugShaderTypeInfo;
@@ -282,10 +281,6 @@ public:
 	UPROPERTY()
 	TArray<TObjectPtr<class USpecularProfile>> SpecularProfiles;
 
-	/** Neural network profile. For internal usage, not editable/visible */
-	UPROPERTY(EditAnywhere, BlueprintReadOnly, Category = PostProcessMaterial, meta = (DisplayName = "Neural Profile"))
-	TObjectPtr<class UNeuralProfile> NeuralProfile;
-
 	/** Whether this material interface is included in the base game (and not in a DLC) */
 	UPROPERTY()
 	uint8 bIncludedInBaseGame : 1;
@@ -923,7 +918,6 @@ public:
 	ENGINE_API virtual USubsurfaceProfile* GetSubsurfaceProfile_Internal() const;
 	ENGINE_API virtual uint32 NumSpecularProfile_Internal() const;
 	ENGINE_API virtual USpecularProfile* GetSpecularProfile_Internal(uint32 Index) const;
-	ENGINE_API virtual UNeuralProfile* GetNeuralProfile_Internal() const;
 	ENGINE_API virtual bool CastsRayTracedShadows() const;
 
 	/**
diff --git a/Engine/Source/Runtime/Engine/Engine.Build.cs b/Engine/Source/Runtime/Engine/Engine.Build.cs
--- a/Engine/Source/Runtime/Engine/Engine.Build.cs
+++ b/Engine/Source/Runtime/Engine/Engine.Build.cs
@@ -158,7 +158,6 @@ public class Engine : ModuleRules
 				"ColorManagement",
 				"Icmp",
 				"XmlParser",
-				"NNE"
 			}
 		);
 
diff --git a/Engine/Source/Runtime/Engine/Private/Materials/HLSLMaterialTranslator.cpp b/Engine/Source/Runtime/Engine/Private/Materials/HLSLMaterialTranslator.cpp
--- a/Engine/Source/Runtime/Engine/Private/Materials/HLSLMaterialTranslator.cpp
+++ b/Engine/Source/Runtime/Engine/Private/Materials/HLSLMaterialTranslator.cpp
@@ -2130,7 +2130,6 @@ void FHLSLMaterialTranslator::GetMaterialEnvironment(EShaderPlatform InPlatform,
 	OutEnvironment.SetDefine(TEXT("MATERIAL_DECAL_READ_MASK"), MaterialCompilationOutput.UsedDBufferTextures);
 	OutEnvironment.SetDefine(TEXT("MATERIAL_USES_DECAL_LOOKUP"), MaterialCompilationOutput.bUsesDBufferTextureLookup);
 	OutEnvironment.SetDefine(TEXT("MATERIAL_PATH_TRACING_BUFFER_READ"), MaterialCompilationOutput.UsedPathTracingBufferTextures);
-	OutEnvironment.SetDefine(TEXT("MATERIAL_NEURAL_POST_PROCESS"), (MaterialCompilationOutput.bUsedWithNeuralNetworks || Material->IsUsedWithNeuralNetworks()) && Material->IsPostProcessMaterial());
 
 	// Count the number of VTStacks (each stack will allocate a feedback slot)
 	OutEnvironment.SetDefine(TEXT("NUM_VIRTUALTEXTURE_SAMPLES"), VTStacks.Num());
@@ -11308,47 +11307,6 @@ int32 FHLSLMaterialTranslator::GetLocal(const FName& LocalName)
 	return AddInlinedCodeChunk(MCT_Float1, TEXT("%s"), *Entry->Name);
 }
 
-int32 FHLSLMaterialTranslator::NeuralOutput(int32 ViewportUV, uint32 NeuralIndexType)
-{
-	if (Material->GetMaterialDomain() != MD_PostProcess)
-	{
-		Errorf(TEXT("NNE Output Node are only available on post process material."));
-	}
-
-	AddEstimatedTextureSample();
-	MaterialCompilationOutput.bUsedWithNeuralNetworks = true;
-	
-	if (NeuralIndexType == 0)
-	{
-		if (ViewportUV == INDEX_NONE)
-		{
-			ViewportUV = AddInlinedCodeChunk(MCT_Float2, TEXT("GetViewportUV(Parameters)"));
-		}
-
-		return AddCodeChunk(MCT_Float4, TEXT("NeuralTextureOutput(Parameters,%s)"),
-			*CoerceParameter(ViewportUV, MCT_Float2));
-	}
-	else if (NeuralIndexType == 1)
-	{
-		int32 BufferIndex = INDEX_NONE;
-
-		if (ViewportUV == INDEX_NONE)
-		{
-			ViewportUV = AddInlinedCodeChunk(MCT_Float2, TEXT("GetViewportUV(Parameters)"));
-			BufferIndex = AppendVector(Constant2(0.0f, 0.0f), ViewportUV);
-		}
-		else
-		{
-			BufferIndex = ViewportUV;
-		}
-
-		return AddCodeChunk(MCT_Float4, TEXT("NeuralBufferOutput(Parameters,%s)"),
-			*CoerceParameter(BufferIndex, MCT_Float4));
-	}
-
-	return INDEX_NONE;
-}
-
 FSubstrateOperator& FHLSLMaterialTranslator::SubstrateCompilationRegisterOperator(int32 OperatorType, FGuid SubstrateExpressionGuid, UMaterialExpression* Child, UMaterialExpression* Parent, FGuid SubstrateParentExpressionGuid, bool bUseParameterBlending)
 {
 	FSubstrateCompilationContext& SubstrateCtx = SubstrateCompilationContext[CurrentSubstrateCompilationContext];
diff --git a/Engine/Source/Runtime/Engine/Private/Materials/HLSLMaterialTranslator.h b/Engine/Source/Runtime/Engine/Private/Materials/HLSLMaterialTranslator.h
--- a/Engine/Source/Runtime/Engine/Private/Materials/HLSLMaterialTranslator.h
+++ b/Engine/Source/Runtime/Engine/Private/Materials/HLSLMaterialTranslator.h
@@ -1197,9 +1197,6 @@ protected:
 	virtual int32 SetLocal(const FName& LocalName, int32 Value) override;
 	virtual int32 GetLocal(const FName& LocalName) override;
 
-	// Neural network nodes
-	virtual int32 NeuralOutput(int32 ViewportUV, uint32 NeuralIndexType) override;
-
 	// Substrate
 	virtual int32 SubstrateCreateAndRegisterNullMaterial() override;
 	virtual int32 SubstrateSlabBSDF(
diff --git a/Engine/Source/Runtime/Engine/Private/Materials/Material.cpp b/Engine/Source/Runtime/Engine/Private/Materials/Material.cpp
--- a/Engine/Source/Runtime/Engine/Private/Materials/Material.cpp
+++ b/Engine/Source/Runtime/Engine/Private/Materials/Material.cpp
@@ -4431,9 +4431,7 @@ bool UMaterial::CanEditChange(const FProperty* InProperty) const
 			PropertyName == GET_MEMBER_NAME_STRING_CHECKED(UMaterial, bIsBlendable) ||
 			PropertyName == GET_MEMBER_NAME_STRING_CHECKED(UMaterial, bEnableStencilTest) ||
 			PropertyName == GET_MEMBER_NAME_STRING_CHECKED(UMaterial, StencilCompare) ||
-			PropertyName == GET_MEMBER_NAME_STRING_CHECKED(UMaterial, StencilRefValue) ||
-			PropertyName == GET_MEMBER_NAME_STRING_CHECKED(UMaterial, NeuralProfileId) ||
-			PropertyName == GET_MEMBER_NAME_STRING_CHECKED(UMaterial, bUsedWithNeuralNetworks)
+			PropertyName == GET_MEMBER_NAME_STRING_CHECKED(UMaterial, StencilRefValue)
 			)
 		{
 			return MaterialDomain == MD_PostProcess;
@@ -6739,12 +6737,6 @@ USpecularProfile* UMaterial::GetSpecularProfile_Internal(uint32 Index) const
 	return SpecularProfiles[Index];
 }
 
-UNeuralProfile* UMaterial::GetNeuralProfile_Internal() const
-{
-	checkSlow(IsInGameThread());
-	return NeuralProfile;
-}
-
 bool UMaterial::CastsRayTracedShadows() const
 {
 	return bCastRayTracedShadows;
diff --git a/Engine/Source/Runtime/Engine/Private/Materials/MaterialExpressionHLSL.cpp b/Engine/Source/Runtime/Engine/Private/Materials/MaterialExpressionHLSL.cpp
--- a/Engine/Source/Runtime/Engine/Private/Materials/MaterialExpressionHLSL.cpp
+++ b/Engine/Source/Runtime/Engine/Private/Materials/MaterialExpressionHLSL.cpp
@@ -119,7 +119,6 @@
 #include "Materials/MaterialExpressionNaniteReplace.h"
 #include "Materials/MaterialExpressionNoise.h"
 #include "Materials/MaterialExpressionNormalize.h"
-#include "Materials/MaterialExpressionNeuralPostProcessNode.h"
 #include "Materials/MaterialExpressionObjectBounds.h"
 #include "Materials/MaterialExpressionObjectOrientation.h"
 #include "Materials/MaterialExpressionObjectPositionWS.h"
@@ -4540,135 +4539,4 @@ bool GenerateStaticTerrainLayerWeightExpression(FName LayerName, float PreviewWe
 	return true;
 }
 
-UE::Shader::EValueType UMaterialExpressionNeuralNetworkInput::GetCustomOutputType(int32 OutputIndex) const
-{
-	using namespace UE::Shader;
-
-	if (OutputIndex == 0)
-	{
-		return EValueType::Float4;
-	}
-	else if (OutputIndex == 1)
-	{
-		return EValueType::Float3;
-	}
-	else if (OutputIndex == 2)
-	{
-		return EValueType::Float1;
-	}
-	else
-	{
-		return EValueType::Void;
-	}
-}
-
-bool UMaterialExpressionNeuralNetworkInput::GenerateHLSLExpression(FMaterialHLSLGenerator& Generator, UE::HLSLTree::FScope& Scope, int32 OutputIndex, UE::HLSLTree::FExpression const*& OutExpression) const
-{
-	using namespace UE::HLSLTree;
-
-	const UMaterial* BaseMaterial = Generator.GetTargetMaterial();
-	if (BaseMaterial)
-	{
-		const EMaterialDomain MaterialDomain = BaseMaterial->MaterialDomain;
-		if (MaterialDomain != MD_PostProcess)
-		{
-			return Generator.Error(TEXT("Neural Output is only available in post process material."));
-		}
-	}
-
-	const bool bUseTextureAsInput = NeuralIndexType == ENeuralIndexType::NIT_TextureIndex;
-	FTree& Tree = Generator.GetTree();
-	
-	if (OutputIndex == 0)
-	{
-		if (Coordinates.IsConnected())
-		{
-			const FExpression* TracedCoordinates = Coordinates.AcquireHLSLExpression(Generator, Scope);
-			if (bUseTextureAsInput)
-			{
-				const FExpression* R = Tree.NewSwizzle(FSwizzleParameters(0), TracedCoordinates);
-				const FExpression* G = Tree.NewSwizzle(FSwizzleParameters(1), TracedCoordinates);
-				const FExpression* B = Tree.NewSwizzle(FSwizzleParameters(2), TracedCoordinates);
-				const FExpression* A = Tree.NewSwizzle(FSwizzleParameters(3), TracedCoordinates);
-				OutExpression = Tree.NewAppend(Tree.NewConstant(-1.0f), G, B, A);
-			}
-			else
-			{
-				OutExpression = TracedCoordinates;
-			}
-		}
-		else
-		{
-			const FExpression* ViewportUV = Tree.NewExpression<Material::FExpressionExternalInput>(Material::EExternalInput::ViewportUV);
-			int32 BatchIndex = bUseTextureAsInput ? -1.0f : 0.0f;
-			OutExpression = Tree.NewAppend(Generator.GetTree().NewConstant(FVector2f(BatchIndex, 0.0f)), ViewportUV);
-		}
-	}
-	else if (OutputIndex == 1)
-	{
-		OutExpression = Input0.AcquireHLSLExpressionOrConstant(Generator, Scope, FVector3f(0.5f));
-	}
-	else if (OutputIndex == 2)
-	{
-		OutExpression = Mask.AcquireHLSLExpressionOrConstant(Generator, Scope, 1.0f);
-	}
-
-	return OutExpression != nullptr;
-}
-
-bool UMaterialExpressionNeuralNetworkOutput::GenerateHLSLExpression(FMaterialHLSLGenerator& Generator, UE::HLSLTree::FScope& Scope, int32 OutputIndex, UE::HLSLTree::FExpression const*& OutExpression) const
-{
-	using namespace UE::HLSLTree;
-
-	const UMaterial* BaseMaterial = Generator.GetTargetMaterial();
-	if (BaseMaterial)
-	{
-		const EMaterialDomain MaterialDomain = BaseMaterial->MaterialDomain;
-		if (MaterialDomain != MD_PostProcess)
-		{
-			return Generator.Error(TEXT("Neural Output is only available in post process material."));
-		}
-	}
-	FTree& Tree = Generator.GetTree();
-
-	if (OutputIndex == 0)
-	{
-		auto GetCoordinateByNeuralIndexType = [&]()-> const FExpression* {
-
-			if (NeuralIndexType == ENeuralIndexType::NIT_TextureIndex)
-			{
-				return Coordinates.GetTracedInput().Expression ?
-					Coordinates.AcquireHLSLExpression(Generator, Scope) :
-					Tree.NewExpression<Material::FExpressionExternalInput>(Material::EExternalInput::ViewportUV);
-			}
-			else if (NeuralIndexType == ENeuralIndexType::NIT_BufferIndex)
-			{
-				if (Coordinates.GetTracedInput().Expression)
-				{
-					return Coordinates.AcquireHLSLExpression(Generator, Scope);
-				}
-				else
-				{
-					const FExpression* ViewportUV = Tree.NewExpression<Material::FExpressionExternalInput>(Material::EExternalInput::ViewportUV);
-					return Tree.NewAppend(Tree.NewConstant(FVector2f(0.0f)), ViewportUV);
-				}
-			}
-			else
-			{
-				return nullptr;
-			}
-		};
-
-		const FExpression* TracedCoordinates = GetCoordinateByNeuralIndexType();
-		if (TracedCoordinates)
-		{
-			OutExpression = Generator.GetTree().NewExpression<Material::FExpressionNeuralNetworkOutput>(TracedCoordinates, NeuralIndexType);
-		}
-
-		return TracedCoordinates != nullptr;
-	}
-
-	return true;
-}
-
 #endif // WITH_EDITOR
diff --git a/Engine/Source/Runtime/Engine/Private/Materials/MaterialExpressions.cpp b/Engine/Source/Runtime/Engine/Private/Materials/MaterialExpressions.cpp
--- a/Engine/Source/Runtime/Engine/Private/Materials/MaterialExpressions.cpp
+++ b/Engine/Source/Runtime/Engine/Private/Materials/MaterialExpressions.cpp
@@ -142,7 +142,6 @@
 #include "Materials/MaterialExpressionNaniteReplace.h"
 #include "Materials/MaterialExpressionNoise.h"
 #include "Materials/MaterialExpressionNormalize.h"
-#include "Materials/MaterialExpressionNeuralPostProcessNode.h"
 #include "Materials/MaterialExpressionObjectBounds.h"
 #include "Materials/MaterialExpressionObjectLocalBounds.h"
 #include "Materials/MaterialExpressionObjectOrientation.h"
@@ -27841,177 +27840,4 @@ FString UMaterialExpressionSubsurfaceMediumMaterialOutput::GetDisplayName() cons
 	return TEXT("Subsurface Medium");
 }
 
-
-///////////////////////////Neural network nodes
-
-UMaterialExpressionNeuralNetworkInput::UMaterialExpressionNeuralNetworkInput(const FObjectInitializer& ObjectInitializer)
-	: Super(ObjectInitializer)
-{
-	// Structure to hold one-time initialization
-	struct FConstructorStatics
-	{
-		FText NAME_NeuralNetwork;
-		FConstructorStatics()
-			: NAME_NeuralNetwork(LOCTEXT("NeuralNetwork", "NeuralNetwork"))
-		{
-		}
-	};
-	static FConstructorStatics ConstructorStatics;
-
-#if WITH_EDITORONLY_DATA
-	MenuCategories.Add(ConstructorStatics.NAME_NeuralNetwork);
-#endif
-
-#if WITH_EDITOR
-	Outputs.Reset();
-#endif
-}
-
-#if WITH_EDITOR
-
-int32 UMaterialExpressionNeuralNetworkInput::Compile(class FMaterialCompiler* Compiler, int32 OutputIndex)
-{
-	int32 CodeInput = INDEX_NONE;
-	const bool bUseTextureAsInput = NeuralIndexType == ENeuralIndexType::NIT_TextureIndex;
-
-	if (OutputIndex == 0)
-	{
-		if (Coordinates.IsConnected())
-		{
-			CodeInput = Coordinates.Compile(Compiler);
-			if (bUseTextureAsInput)
-			{
-				CodeInput = Compiler->ComponentMask(CodeInput, false, true, true, true);
-				CodeInput = Compiler->AppendVector(Compiler->Constant(-1.0f), CodeInput);
-			}
-		}
-		else
-		{
-			int32 ViewportUV = Compiler->GetViewportUV();
-			int32 BatchIndex = bUseTextureAsInput ? -1.0f : 0.0f;
-				
-			CodeInput = Compiler->AppendVector(Compiler->Constant2(BatchIndex, 0.0f), ViewportUV);
-		}
-	}
-	else if (OutputIndex == 1)
-	{
-		CodeInput = Input0.IsConnected() ? Input0.Compile(Compiler) : Compiler->Constant3(0.5f, 0.5f, 0.5f);
-	}
-	else if (OutputIndex == 2)
-	{
-		CodeInput = Mask.IsConnected() ? Mask.Compile(Compiler) : Compiler->Constant(1.0f);
-	}
-
-	return Compiler->CustomOutput(this, OutputIndex, CodeInput);
-}
-
-void UMaterialExpressionNeuralNetworkInput::GetCaption(TArray<FString>& OutCaptions) const
-{
-	if (NeuralIndexType == ENeuralIndexType::NIT_TextureIndex)
-	{
-		OutCaptions.Add(TEXT("Neural Input (Texture)"));
-	}
-	else if (NeuralIndexType == ENeuralIndexType::NIT_BufferIndex)
-	{
-		OutCaptions.Add(TEXT("Neural Input (Buffer)"));
-	}
-}
-
-uint32 UMaterialExpressionNeuralNetworkInput::GetInputType(int32 InputIndex)
-{
-
-	switch (InputIndex)
-	{
-	case 0:return MCT_Float4;break;
-	case 1:return MCT_Float3;break;
-	case 2:return MCT_Float1;break;
-	default: checkNoEntry(); break;
-	}
-
-	return MCT_Float1;
-}
-
-#endif // WITH_EDITOR
-
-int32 UMaterialExpressionNeuralNetworkInput::GetNumOutputs() const
-{
-	return 3;
-}
-
-FString UMaterialExpressionNeuralNetworkInput::GetFunctionName() const
-{
-	return TEXT("GetNeuralInput");
-}
-
-///////////////////////////////////////////////////////////////////////////////
-// UMaterialExpressionNeuralNetworkOutput
-///////////////////////////////////////////////////////////////////////////////
-UMaterialExpressionNeuralNetworkOutput::UMaterialExpressionNeuralNetworkOutput(const FObjectInitializer& ObjectInitializer)
-	: Super(ObjectInitializer)
-{
-
-	// Structure to hold one-time initialization
-	struct FConstructorStatics
-	{
-		FText NAME_NeuralNetwork;
-		FConstructorStatics()
-			: NAME_NeuralNetwork(LOCTEXT("NeuralNetwork", "NeuralNetwork"))
-		{
-		}
-	};
-	static FConstructorStatics ConstructorStatics;
-
-#if WITH_EDITORONLY_DATA
-	MenuCategories.Add(ConstructorStatics.NAME_NeuralNetwork);
-
-	bShaderInputData = true;
-	bShowOutputNameOnPin = true;
-#endif
-
-#if WITH_EDITOR
-	Outputs.Reset();
-	Outputs.Add(FExpressionOutput(TEXT("RGBA"), 1, 1, 1, 1, 1));
-#endif
-}
-
-#if WITH_EDITOR
-int32 UMaterialExpressionNeuralNetworkOutput::Compile(class FMaterialCompiler* Compiler, int32 OutputIndex)
-{
-	int32 ViewportUV = INDEX_NONE;
-
-	if (Coordinates.GetTracedInput().Expression)
-	{
-		ViewportUV = Coordinates.Compile(Compiler);
-	}
-
-	return Compiler->NeuralOutput(ViewportUV, NeuralIndexType);
-}
-
-void UMaterialExpressionNeuralNetworkOutput::GetCaption(TArray<FString>& OutCaptions) const
-{
-	if (NeuralIndexType == ENeuralIndexType::NIT_TextureIndex)
-	{
-		OutCaptions.Add(TEXT("Neural Output (Texture)"));
-	}
-	else if (NeuralIndexType == ENeuralIndexType::NIT_BufferIndex)
-	{
-		OutCaptions.Add(TEXT("Neural Output (Buffer)"));
-	}
-}
-uint32 UMaterialExpressionNeuralNetworkOutput::GetInputType(int32 InputIndex)
-{
-	if (NeuralIndexType == ENeuralIndexType::NIT_TextureIndex)
-	{
-		return MCT_Float2;
-	}
-	else if (NeuralIndexType == ENeuralIndexType::NIT_BufferIndex)
-	{
-		return MCT_Float4;
-	}
-
-	checkNoEntry();
-	return MCT_Float2;
-}
-#endif // WITH_EDITOR
-
 #undef LOCTEXT_NAMESPACE
diff --git a/Engine/Source/Runtime/Engine/Private/Materials/MaterialHLSLEmitter.cpp b/Engine/Source/Runtime/Engine/Private/Materials/MaterialHLSLEmitter.cpp
--- a/Engine/Source/Runtime/Engine/Private/Materials/MaterialHLSLEmitter.cpp
+++ b/Engine/Source/Runtime/Engine/Private/Materials/MaterialHLSLEmitter.cpp
@@ -482,8 +482,6 @@ static void GetMaterialEnvironment(EShaderPlatform InPlatform,
 	OutEnvironment.SetDefine(TEXT("MATERIAL_FULLY_ROUGH"), bIsFullyRough);
 	OutEnvironment.SetDefine(TEXT("MATERIAL_USES_ANISOTROPY"), bUsesAnisotropy);
 
-	OutEnvironment.SetDefine(TEXT("MATERIAL_NEURAL_POST_PROCESS"), (MaterialCompilationOutput.bUsedWithNeuralNetworks || InMaterial.IsUsedWithNeuralNetworks()) && InMaterial.IsPostProcessMaterial());
-
 	// Count the number of VTStacks (each stack will allocate a feedback slot)
 	OutEnvironment.SetDefine(TEXT("NUM_VIRTUALTEXTURE_SAMPLES"), EmitMaterialData.VTStacks.Num());
 
diff --git a/Engine/Source/Runtime/Engine/Private/Materials/MaterialHLSLTree.cpp b/Engine/Source/Runtime/Engine/Private/Materials/MaterialHLSLTree.cpp
--- a/Engine/Source/Runtime/Engine/Private/Materials/MaterialHLSLTree.cpp
+++ b/Engine/Source/Runtime/Engine/Private/Materials/MaterialHLSLTree.cpp
@@ -3000,57 +3000,6 @@ void FExpressionAtmosphericFogColorFunction::EmitValueShader(FEmitContext& Conte
 		EmitPosition);
 }
 
-bool FExpressionNeuralNetworkOutput::PrepareValue(FEmitContext& Context, FEmitScope& Scope, const FRequestedType& RequestedType, FPrepareValueResult& OutResult) const
-{
-	if (NeuralIndexType == 0)
-	{
-		const FPreparedType& TextureIndexType = Context.PrepareExpression(CoordinatesExpression, Scope, Shader::EValueType::Float2);
-		if (TextureIndexType.IsVoid())
-		{
-			return false;
-		}
-	}
-	else if (NeuralIndexType == 1)
-	{
-		const FPreparedType& BufferIndexType = Context.PrepareExpression(CoordinatesExpression, Scope, Shader::EValueType::Float4);
-		if (BufferIndexType.IsVoid())
-		{
-			return false;
-		}
-	}
-
-	if (Context.bMarkLiveValues && Context.MaterialCompilationOutput)
-	{
-		Context.MaterialCompilationOutput->bUsedWithNeuralNetworks = true;
-	}
-
-	return OutResult.SetType(Context, RequestedType, EExpressionEvaluation::Shader, Shader::EValueType::Float4);
-}
-
-void FExpressionNeuralNetworkOutput::EmitValueShader(FEmitContext& Context, FEmitScope& Scope, const FRequestedType& RequestedType, FEmitValueShaderResult& OutResult) const
-{
-	if (NeuralIndexType == 0)
-	{
-		FEmitShaderExpression* EmitTextureIndex = CoordinatesExpression->GetValueShader(Context, Scope, Shader::EValueType::Float2);
-		if (EmitTextureIndex)
-		{
-			OutResult.Code = Context.EmitInlineExpression(Scope, Shader::EValueType::Float4,
-				TEXT("NeuralTextureOutput(Parameters, %)"),
-				EmitTextureIndex);
-		}
-	}
-	else if (NeuralIndexType == 1)
-	{
-		FEmitShaderExpression* EmitBufferIndex = CoordinatesExpression->GetValueShader(Context, Scope, Shader::EValueType::Float4);
-		if (EmitBufferIndex)
-		{
-			OutResult.Code = Context.EmitInlineExpression(Scope, Shader::EValueType::Float4,
-				TEXT("NeuralBufferOutput(Parameters, %)"),
-				EmitBufferIndex);
-		}
-	}
-}
-
 int32 FEmitData::FindInterpolatorIndex(const FExpression* Expression) const
 {
 	for (int32 Index = 0; Index < VertexInterpolators.Num(); ++Index)
diff --git a/Engine/Source/Runtime/Engine/Private/Materials/MaterialInterface.cpp b/Engine/Source/Runtime/Engine/Private/Materials/MaterialInterface.cpp
--- a/Engine/Source/Runtime/Engine/Private/Materials/MaterialInterface.cpp
+++ b/Engine/Source/Runtime/Engine/Private/Materials/MaterialInterface.cpp
@@ -23,7 +23,6 @@
 #include "ObjectCacheEventSink.h"
 #include "Engine/SubsurfaceProfile.h"
 #include "Engine/SpecularProfile.h"
-#include "Engine/NeuralProfile.h"
 #include "Interfaces/ITargetPlatform.h"
 #include "Components/PrimitiveComponent.h"
 #include "ContentStreaming.h"
@@ -1332,12 +1331,6 @@ USpecularProfile* UMaterialInterface::GetSpecularProfile_Internal(uint32 Index)
 	return nullptr;
 }
 
-UNeuralProfile* UMaterialInterface::GetNeuralProfile_Internal() const
-{
-	return nullptr;
-}
-
-
 uint32 UMaterialInterface::NumSpecularProfile_Internal() const
 {
 	return 0u;
@@ -1414,40 +1407,6 @@ void UMaterialInterface::UpdateMaterialRenderProxy(FMaterialRenderProxy& Proxy)
 		});
 	}
 
-	UMaterial* Material = GetMaterial();
-	if (Material && Material->IsPostProcessMaterial())
-	{
-		struct FEntry
-		{
-			UNeuralProfile* Profile = nullptr;
-			FNeuralProfileStruct Setting;
-			FGuid Guid;
-		};
-		
-		UNeuralProfile* LocalNeuralProfile = GetNeuralProfile_Internal();
-		
-		if (LocalNeuralProfile)
-		{
-			FEntry Entry;
-			Entry.Profile = LocalNeuralProfile;
-			Entry.Setting = LocalNeuralProfile->Settings;
-			Entry.Guid = LocalNeuralProfile->Guid;
-
-			FMaterialRenderProxy* InProxy = &Proxy;
-			ENQUEUE_RENDER_COMMAND(UpdateMaterialRenderProxyNNEModelData)(
-				[LocalNeuralProfile, InProxy, Entry, Material](FRHICommandListImmediate& RHICmdList)
-				{
-					
-					const uint32 AllocationId = NeuralProfile::AddOrUpdateProfile(Entry.Profile, Entry.Guid, Entry.Setting);
-					check(AllocationId >= 0 && AllocationId < MAX_NEURAL_PROFILE_COUNT);
-
-					Material->NeuralProfileId = AllocationId;
-					
-					InProxy->SetNeuralProfileRT(LocalNeuralProfile);
-				});
-		}
-	}
-
 	if (Substrate::IsSubstrateEnabled())
 	{
 		struct FEntry
diff --git a/Engine/Source/Runtime/Engine/Private/Materials/MaterialShared.cpp b/Engine/Source/Runtime/Engine/Private/Materials/MaterialShared.cpp
--- a/Engine/Source/Runtime/Engine/Private/Materials/MaterialShared.cpp
+++ b/Engine/Source/Runtime/Engine/Private/Materials/MaterialShared.cpp
@@ -55,7 +55,6 @@
 #include "SubstrateDefinitions.h"
 #include "DataDrivenShaderPlatformInfo.h"
 #include "ProfilingDebugging/CookStats.h"
-#include "Engine/NeuralProfile.h"
 
 #define LOCTEXT_NAMESPACE "MaterialShared"
 
@@ -1726,11 +1725,6 @@ bool FMaterialResource::IsUsedWithVirtualHeightfieldMesh() const
 	return Material->bUsedWithVirtualHeightfieldMesh;
 }
 
-bool FMaterialResource::IsUsedWithNeuralNetworks() const
-{
-	return Material->bUsedWithNeuralNetworks && Material->IsPostProcessMaterial();
-}
-
 bool FMaterialResource::IsUsedWithLandscape() const
 {
 	return false;
@@ -2183,11 +2177,6 @@ uint32 FMaterialResource::GetStencilRefValue() const
 	return GetMaterialDomain() == MD_PostProcess ? Material->StencilRefValue : 0;
 }
 
-int32 FMaterialResource::GetNeuralProfileId() const
-{ 
-	return GetMaterialDomain() == MD_PostProcess ? Material->NeuralProfileId : INDEX_NONE;
-}
-
 uint32 FMaterialResource::GetStencilCompare() const
 {
 	return GetMaterialDomain() == MD_PostProcess ? uint32(Material->StencilCompare.GetValue()) : 0;
@@ -5361,7 +5350,6 @@ FMaterialShaderParameters::FMaterialShaderParameters(const FMaterial* InMaterial
 	bHasRuntimeVirtualTextureOutput = InMaterial->HasRuntimeVirtualTextureOutput();
 	bIsUsedWithLidarPointCloud = InMaterial->IsUsedWithLidarPointCloud();
 	bIsUsedWithVirtualHeightfieldMesh = InMaterial->IsUsedWithVirtualHeightfieldMesh();
-	bIsUsedWithNeuralNetworks = InMaterial->IsUsedWithNeuralNetworks();
 	bIsUsedWithNanite = InMaterial->IsUsedWithNanite();
 	bIsStencilTestEnabled = InMaterial->IsStencilTestEnabled();
 	bIsTranslucencySurface = InMaterial->GetTranslucencyLightingMode() == ETranslucencyLightingMode::TLM_Surface || InMaterial->GetTranslucencyLightingMode() == ETranslucencyLightingMode::TLM_SurfacePerPixelLighting;
diff --git a/Engine/Source/Runtime/Engine/Private/Rendering/NeuralProfile.cpp b/Engine/Source/Runtime/Engine/Private/Rendering/NeuralProfile.cpp
deleted file mode 100644
diff --git a/Engine/Source/Runtime/Engine/Public/MaterialCompiler.h b/Engine/Source/Runtime/Engine/Public/MaterialCompiler.h
--- a/Engine/Source/Runtime/Engine/Public/MaterialCompiler.h
+++ b/Engine/Source/Runtime/Engine/Public/MaterialCompiler.h
@@ -553,9 +553,6 @@ public:
 	virtual int32 SetLocal(const FName& LocalName, int32 Value) = 0;
 	virtual int32 GetLocal(const FName& LocalName) = 0;
 
-	// Neural network nodes
-	virtual int32 NeuralOutput(int32 ViewportUV, uint32 NeuralIndexType) = 0;
-	
 	// Substrate
 	virtual int32 SubstrateCreateAndRegisterNullMaterial() = 0;
 	virtual int32 SubstrateSlabBSDF(
@@ -1185,11 +1182,6 @@ public:
 		return Compiler->GetLocal(LocalName);
 	}
 
-	virtual int32 NeuralOutput(int32 ViewportUV, uint32 NeuralIndexType) override
-	{
-		return Compiler->NeuralOutput(ViewportUV, NeuralIndexType);
-	}
-
 	virtual int32 SubstrateCreateAndRegisterNullMaterial() override
 	{
 		return Compiler->SubstrateCreateAndRegisterNullMaterial();
diff --git a/Engine/Source/Runtime/Engine/Public/MaterialHLSLTree.h b/Engine/Source/Runtime/Engine/Public/MaterialHLSLTree.h
--- a/Engine/Source/Runtime/Engine/Public/MaterialHLSLTree.h
+++ b/Engine/Source/Runtime/Engine/Public/MaterialHLSLTree.h
@@ -1015,23 +1015,6 @@ public:
 	virtual void EmitValueShader(FEmitContext& Context, FEmitScope& Scope, const FRequestedType& RequestedType, FEmitValueShaderResult& OutResult) const override;
 };
 
-class FExpressionNeuralNetworkOutput : public FExpression
-{
-public:
-	const FExpression* CoordinatesExpression;
-	int32 NeuralIndexType;
-
-	FExpressionNeuralNetworkOutput(
-		const FExpression* InCoordinateExpression,
-		int32 InNeuralIndexType)
-		: CoordinatesExpression(InCoordinateExpression),
-		NeuralIndexType(InNeuralIndexType)
-	{}
-
-	virtual bool PrepareValue(FEmitContext& Context, FEmitScope& Scope, const FRequestedType& RequestedType, FPrepareValueResult& OutResult) const override;
-	virtual void EmitValueShader(FEmitContext& Context, FEmitScope& Scope, const FRequestedType& RequestedType, FEmitValueShaderResult& OutResult) const override;
-};
-
 class FExpressionLightVector : public FExpressionForward
 {
 public:
diff --git a/Engine/Source/Runtime/Engine/Public/MaterialShared.h b/Engine/Source/Runtime/Engine/Public/MaterialShared.h
--- a/Engine/Source/Runtime/Engine/Public/MaterialShared.h
+++ b/Engine/Source/Runtime/Engine/Public/MaterialShared.h
@@ -717,8 +717,7 @@ public:
 		bUsesVertexInterpolator(false),
 		bHasRuntimeVirtualTextureOutputNode(false),
 		bUsesAnisotropy(false),
-		bUsesDisplacement(false),
-		bUsedWithNeuralNetworks(false)
+		bUsesDisplacement(false)
 	{
 #if WITH_EDITOR
 		FMemory::Memzero(EstimatedLWCFuncUsages);
@@ -840,11 +839,6 @@ public:
 	
 	/** Whether the material uses scalar displacement. */
 	LAYOUT_BITFIELD(uint8, bUsesDisplacement, 1);
-
-	/** Whether the material uses NNE. */
-	LAYOUT_BITFIELD(uint8, bUsedWithNeuralNetworks, 1);
-	
-
 };
 
 struct FDebugShaderPipelineInfo
@@ -1917,7 +1911,6 @@ public:
 	virtual bool IsUsedWithHairStrands() const { return false; }
 	virtual bool IsUsedWithLidarPointCloud() const { return false; }
 	virtual bool IsUsedWithVirtualHeightfieldMesh() const { return false; }
-	virtual bool IsUsedWithNeuralNetworks() const { return false; }
 	virtual bool IsFullyRough() const { return false; }
 	virtual bool UseNormalCurvatureToRoughness() const { return false; }
 	virtual enum EMaterialFloatPrecisionMode GetMaterialFloatPrecisionMode() const { return EMaterialFloatPrecisionMode::MFPM_Default; };
@@ -2002,7 +1995,6 @@ public:
 	virtual bool CastsRayTracedShadows() const { return true; }
 	virtual bool HasRenderTracePhysicalMaterialOutputs() const { return false; }
 	virtual EMaterialShadingRate GetShadingRate() const { return MSR_1x1; }
-	virtual int32 GetNeuralProfileId() const { return INDEX_NONE; }
 	/**
 	 * Should shaders compiled for this material be saved to disk?
 	 */
@@ -2659,7 +2651,6 @@ public:
 	ENGINE_API virtual bool IsUsedWithHairStrands() const override;
 	ENGINE_API virtual bool IsUsedWithLidarPointCloud() const override;
 	ENGINE_API virtual bool IsUsedWithVirtualHeightfieldMesh() const override;
-	ENGINE_API virtual bool IsUsedWithNeuralNetworks() const override;
 	ENGINE_API virtual bool IsUsedWithNanite() const override;
 	ENGINE_API virtual bool IsUsedWithVolumetricCloud() const override;
 	ENGINE_API virtual bool IsUsedWithHeterogeneousVolumes() const override;
@@ -2735,7 +2726,6 @@ public:
 	ENGINE_API virtual bool CastsRayTracedShadows() const override;
 	ENGINE_API virtual bool HasRenderTracePhysicalMaterialOutputs() const override;
 	ENGINE_API virtual UMaterialInterface* GetMaterialInterface() const override;
-	ENGINE_API virtual int32 GetNeuralProfileId() const override;
 	/**
 	 * Should shaders compiled for this material be saved to disk?
 	 */
@@ -3079,7 +3069,6 @@ struct FMaterialShaderParameters
 			uint64 bHasRuntimeVirtualTextureOutput : 1;
 			uint64 bIsUsedWithLidarPointCloud : 1;
 			uint64 bIsUsedWithVirtualHeightfieldMesh : 1;
-			uint64 bIsUsedWithNeuralNetworks : 1;
 			uint64 bIsUsedWithNanite : 1;
 			uint64 bIsStencilTestEnabled : 1;
 			uint64 bIsTranslucencySurface : 1;
diff --git a/Engine/Source/Runtime/Engine/Public/Materials/MaterialRenderProxy.h b/Engine/Source/Runtime/Engine/Public/Materials/MaterialRenderProxy.h
--- a/Engine/Source/Runtime/Engine/Public/Materials/MaterialRenderProxy.h
+++ b/Engine/Source/Runtime/Engine/Public/Materials/MaterialRenderProxy.h
@@ -25,7 +25,6 @@ class URuntimeVirtualTexture;
 class USparseVolumeTexture;
 class USubsurfaceProfile;
 class USpecularProfile;
-class UNeuralProfile;
 class UTexture;
 
 struct FMaterialParameterValue;
@@ -214,10 +213,6 @@ public:
 	const USpecularProfile* GetSpecularProfileRT(uint32 Index) const { check(Index<uint32(SpecularProfilesRT.Num())); return SpecularProfilesRT[Index]; }
 	const uint32 NumSpecularProfileRT() const { return SpecularProfilesRT.Num(); }
 
-	// Neural profiles
-	void SetNeuralProfileRT(const UNeuralProfile* Ptr) { NeuralProfileRT = Ptr; }
-	const UNeuralProfile* GetNeuralProfileRT() const { return NeuralProfileRT; }
-
 	static ENGINE_API void UpdateDeferredCachedUniformExpressions();
 	static ENGINE_API void UpdateDeferredCachedUniformExpressions(FRHICommandListBase& RHICmdList, UE::Tasks::FTask* TaskIfAsync = nullptr);
 
@@ -237,7 +232,6 @@ private:
 	/** 0 if not set, game thread pointer, do not dereference, only for comparison */
 	const USubsurfaceProfile* SubsurfaceProfileRT;
 	TArray<const USpecularProfile*> SpecularProfilesRT;
-	const UNeuralProfile* NeuralProfileRT;
 	FString MaterialName;
 
 	/** Incremented each time UniformExpressionCache is modified */
diff --git a/Engine/Source/Runtime/Renderer/Private/PostProcess/NeuralPostProcess.cpp b/Engine/Source/Runtime/Renderer/Private/PostProcess/NeuralPostProcess.cpp
deleted file mode 100644
diff --git a/Engine/Source/Runtime/Renderer/Private/PostProcess/NeuralPostProcess.h b/Engine/Source/Runtime/Renderer/Private/PostProcess/NeuralPostProcess.h
deleted file mode 100644
diff --git a/Engine/Source/Runtime/Renderer/Private/PostProcess/PostProcessMaterial.cpp b/Engine/Source/Runtime/Renderer/Private/PostProcess/PostProcessMaterial.cpp
--- a/Engine/Source/Runtime/Renderer/Private/PostProcess/PostProcessMaterial.cpp
+++ b/Engine/Source/Runtime/Renderer/Private/PostProcess/PostProcessMaterial.cpp
@@ -27,7 +27,6 @@
 #include "SystemTextures.h"
 #include "Substrate/Substrate.h"
 #include "SingleLayerWaterRendering.h"
-#include "Engine/NeuralProfile.h"
 
 namespace
 {
@@ -291,8 +290,7 @@ public:
 	DECLARE_SHADER_TYPE(FPostProcessMaterialPS, Material);
 
 	class FManualStencilTestDim : SHADER_PERMUTATION_BOOL("MANUAL_STENCIL_TEST");
-	class FNeuralPostProcessPrePass : SHADER_PERMUTATION_BOOL("NEURAL_POSTPROCESS_PREPASS");
-	using FPermutationDomain = TShaderPermutationDomain<FManualStencilTestDim,FNeuralPostProcessPrePass>;
+	using FPermutationDomain = TShaderPermutationDomain<FManualStencilTestDim>;
 
 	FPostProcessMaterialPS() = default;
 	FPostProcessMaterialPS(const ShaderMetaType::CompiledShaderInitializerType& Initializer)
@@ -349,8 +347,7 @@ static void GetMaterialInfo(
 	const FMaterialRenderProxy*& OutMaterialProxy,
 	const FMaterialShaderMap*& OutMaterialShaderMap,
 	TShaderRef<FPostProcessMaterialVS>& OutVertexShader,
-	TShaderRef<FPostProcessMaterialPS>& OutPixelShader,
-	bool bNeuralPostProcessPrepass = false)
+	TShaderRef<FPostProcessMaterialPS>& OutPixelShader)
 {
 	const FMaterialRenderProxy* MaterialProxy = InMaterialInterface->GetRenderProxy();
 	check(MaterialProxy);
@@ -368,7 +365,6 @@ static void GetMaterialInfo(
 
 				FPostProcessMaterialPS::FPermutationDomain PermutationVectorPS;
 				PermutationVectorPS.Set<FPostProcessMaterialPS::FManualStencilTestDim>(bManualStencilTest);
-				PermutationVectorPS.Set<FPostProcessMaterialPS::FNeuralPostProcessPrePass>(bNeuralPostProcessPrepass);
 
 				ShaderTypes.AddShaderType<FPostProcessMaterialVS>();
 				ShaderTypes.AddShaderType<FPostProcessMaterialPS>(PermutationVectorPS.ToDimensionValueId());
@@ -424,262 +420,6 @@ void AddMobileMSAADecodeAndDrawTexturePass(
 	AddDrawScreenPass(GraphBuilder, RDG_EVENT_NAME("MobileMSAADecodeAndDrawTexture"), View, OutputViewport, InputViewport, PixelShader, Parameters);
 }
 
-FPostProcessMaterialParameters* GetPostProcessMaterialParameters(
-	FRDGBuilder& GraphBuilder, 
-	const FPostProcessMaterialInputs& Inputs, 
-	const FViewInfo& View,
-	const FScreenPassTextureViewport& OutputViewport,
-	FScreenPassRenderTarget& Output, 
-	FRDGTextureRef DepthStencilTexture, 
-	const uint32 MaterialStencilRef, 
-	const FMaterial* Material, 
-	const FMaterialShaderMap* MaterialShaderMap)
-{
-	FPostProcessMaterialParameters* PostProcessMaterialParameters = GraphBuilder.AllocParameters<FPostProcessMaterialParameters>();
-	PostProcessMaterialParameters->SceneTextures = Inputs.SceneTextures;
-	PostProcessMaterialParameters->View = View.ViewUniformBuffer;
-	PostProcessMaterialParameters->EyeAdaptationBuffer = GraphBuilder.CreateSRV(GetEyeAdaptationBuffer(GraphBuilder, View));
-	PostProcessMaterialParameters->PostProcessOutput = GetScreenPassTextureViewportParameters(OutputViewport);
-	PostProcessMaterialParameters->RenderTargets[0] = Output.GetRenderTargetBinding();
-
-	// The target color will be decoded if bForceIntermediateTarget is true in any case, but we might still need to decode the input color
-	PostProcessMaterialParameters->bMetalMSAAHDRDecode = Inputs.bMetalMSAAHDRDecode ? 1 : 0;
-
-	if (DepthStencilTexture)
-	{
-		PostProcessMaterialParameters->RenderTargets.DepthStencil = FDepthStencilBinding(
-			DepthStencilTexture,
-			ERenderTargetLoadAction::ELoad,
-			ERenderTargetLoadAction::ELoad,
-			FExclusiveDepthStencil::DepthRead_StencilRead);
-	}
-	PostProcessMaterialParameters->ManualStencilReferenceValue = MaterialStencilRef;
-	PostProcessMaterialParameters->ManualStencilTestMask = GetManualStencilTestMask(Material->GetStencilCompare());
-
-	PostProcessMaterialParameters->PostProcessInput_BilinearSampler = TStaticSamplerState<SF_Bilinear, AM_Clamp, AM_Clamp, AM_Clamp>::GetRHI();;
-
-	const FScreenPassTexture BlackDummy(GSystemTextures.GetBlackDummy(GraphBuilder));
-
-	// This gets passed in whether or not it's used.
-	GraphBuilder.RemoveUnusedTextureWarning(BlackDummy.Texture);
-
-	FRHISamplerState* PointClampSampler = TStaticSamplerState<SF_Point, AM_Clamp, AM_Clamp, AM_Clamp>::GetRHI();
-
-	for (uint32 InputIndex = 0; InputIndex < kPostProcessMaterialInputCountMax; ++InputIndex)
-	{
-		FScreenPassTexture Input = Inputs.GetInput((EPostProcessMaterialInput)InputIndex);
-
-		// Need to provide valid textures for when shader compilation doesn't cull unused parameters.
-		if (!Input.Texture || !MaterialShaderMap->UsesSceneTexture(PPI_PostProcessInput0 + InputIndex))
-		{
-			Input = BlackDummy;
-		}
-
-		PostProcessMaterialParameters->PostProcessInput[InputIndex] = GetScreenPassTextureInput(Input, PointClampSampler);
-	}
-
-	// Path tracing buffer textures
-	for (uint32 InputIndex = 0; InputIndex < kPathTracingPostProcessMaterialInputCountMax; ++InputIndex)
-	{
-		FScreenPassTexture Input = Inputs.GetPathTracingInput((EPathTracingPostProcessMaterialInput)InputIndex);
-
-		if (!Input.Texture || !MaterialShaderMap->UsesPathTracingBufferTexture(InputIndex))
-		{
-			Input = BlackDummy;
-		}
-
-		PostProcessMaterialParameters->PathTracingPostProcessInput[InputIndex] = GetScreenPassTextureInput(Input, PointClampSampler);
-	}
-
-	PostProcessMaterialParameters->Substrate = Substrate::BindSubstrateGlobalUniformParameters(View);
-
-	// SceneDepthWithoutWater
-	const bool bHasValidSceneDepthWithoutWater = Inputs.SceneWithoutWaterTextures && Inputs.SceneWithoutWaterTextures->DepthTexture;
-	const bool bShouldUseBilinearSamplerForDepth = bHasValidSceneDepthWithoutWater && ShouldUseBilinearSamplerForDepthWithoutSingleLayerWater(Inputs.SceneWithoutWaterTextures->DepthTexture->Desc.Format);
-	PostProcessMaterialParameters->bSceneDepthWithoutWaterTextureAvailable = bHasValidSceneDepthWithoutWater;
-	PostProcessMaterialParameters->SceneDepthWithoutSingleLayerWaterSampler = bShouldUseBilinearSamplerForDepth ? TStaticSamplerState<SF_Bilinear>::GetRHI() : TStaticSamplerState<SF_Point>::GetRHI();
-	PostProcessMaterialParameters->SceneDepthWithoutSingleLayerWaterTexture = FRDGSystemTextures::Get(GraphBuilder).Black;
-	PostProcessMaterialParameters->SceneWithoutSingleLayerWaterMinMaxUV = FVector4f(0.0f, 0.0f, 1.0f, 1.0f);
-	PostProcessMaterialParameters->SceneWithoutSingleLayerWaterTextureSize = FVector2f(0.0f, 0.0f);
-	PostProcessMaterialParameters->SceneWithoutSingleLayerWaterInvTextureSize = FVector2f(0.0f, 0.0f);
-	if (bHasValidSceneDepthWithoutWater)
-	{
-		const bool bIsInstancedStereoSideBySide = View.bIsInstancedStereoEnabled && !View.bIsMobileMultiViewEnabled && IStereoRendering::IsStereoEyeView(View);
-		int32 WaterViewIndex = INDEX_NONE;
-		if (bIsInstancedStereoSideBySide)
-		{
-			WaterViewIndex = View.PrimaryViewIndex; // The instanced view does not have MinMaxUV initialized, instead the primary view MinMaxUV covers both eyes
-		}
-		else
-		{
-			verify(View.Family->Views.Find(&View, WaterViewIndex));
-		}
-
-		PostProcessMaterialParameters->SceneDepthWithoutSingleLayerWaterTexture = Inputs.SceneWithoutWaterTextures->DepthTexture;
-		PostProcessMaterialParameters->SceneWithoutSingleLayerWaterMinMaxUV = Inputs.SceneWithoutWaterTextures->Views[WaterViewIndex].MinMaxUV;
-
-		const FIntVector DepthTextureSize = Inputs.SceneWithoutWaterTextures->DepthTexture->Desc.GetSize();
-		PostProcessMaterialParameters->SceneWithoutSingleLayerWaterTextureSize = FVector2f(DepthTextureSize.X, DepthTextureSize.Y);
-		PostProcessMaterialParameters->SceneWithoutSingleLayerWaterInvTextureSize = FVector2f(1.0f / DepthTextureSize.X, 1.0f / DepthTextureSize.Y);
-	}
-
-	PostProcessMaterialParameters->NeuralPostProcessParameters = GetDefaultNeuralPostProcessShaderParameters(GraphBuilder);
-
-	return PostProcessMaterialParameters;
-}
-
-void AddNeuralPostProcessPass(
-	FRDGBuilder& GraphBuilder,
-	const FViewInfo& View,
-	const FPostProcessMaterialInputs& Inputs,
-	const UMaterialInterface* MaterialInterface,
-	FNeuralPostProcessResource& NeuralPostProcessResource)
-{
-	Inputs.Validate();
-
-	const FScreenPassTexture SceneColor = Inputs.GetInput(EPostProcessMaterialInput::SceneColor);
-
-	const ERHIFeatureLevel::Type FeatureLevel = View.GetFeatureLevel();
-
-	const FMaterial* Material = nullptr;
-	const FMaterialRenderProxy* MaterialRenderProxy = nullptr;
-	const FMaterialShaderMap* MaterialShaderMap = nullptr;
-	TShaderRef<FPostProcessMaterialVS> NeuralPostProcessPassVertexShader;
-	TShaderRef<FPostProcessMaterialPS> NeuralPostProcessPassPixelShader;
-	GetMaterialInfo(MaterialInterface, FeatureLevel, Inputs, Material, MaterialRenderProxy, MaterialShaderMap, NeuralPostProcessPassVertexShader, NeuralPostProcessPassPixelShader, true);
-	
-	check(NeuralPostProcessPassVertexShader.IsValid());
-	check(NeuralPostProcessPassPixelShader.IsValid());
-
-	int32 NeuralProfileId = Material->GetNeuralProfileId();
-
-	FRHIDepthStencilState* DefaultDepthStencilState = FScreenPassPipelineState::FDefaultDepthStencilState::GetRHI();
-	FRHIDepthStencilState* DepthStencilState = DefaultDepthStencilState;
-
-	FRDGTextureRef DepthStencilTexture = nullptr;
-
-	// Allocate custom depth stencil texture(s) and depth stencil state.
-	const EMaterialCustomDepthPolicy CustomStencilPolicy = GetMaterialCustomDepthPolicy(Material);
-
-	if (CustomStencilPolicy == EMaterialCustomDepthPolicy::Enabled &&
-		!Inputs.bManualStencilTest &&
-		HasBeenProduced(Inputs.CustomDepthTexture))
-	{
-		check(Inputs.CustomDepthTexture);
-		DepthStencilTexture = Inputs.CustomDepthTexture;
-		DepthStencilState = GetMaterialStencilState(Material);
-	}
-
-	FRHIBlendState* DefaultBlendState = FScreenPassPipelineState::FDefaultBlendState::GetRHI();
-	FRHIBlendState* BlendState = DefaultBlendState;
-
-	if (IsMaterialBlendEnabled(Material))
-	{
-		BlendState = GetMaterialBlendState(Material);
-	}
-
-	FScreenPassRenderTarget Output = Inputs.OverrideOutput;
-	// Create a new texture instead of reusing the scene color output in the pre pass. Should not pollute the scene color texture.
-	{
-		// Allocate new transient output texture.
-		{
-			FRDGTextureDesc OutputDesc = SceneColor.Texture->Desc;
-			OutputDesc.Reset();
-			if (Inputs.OutputFormat != PF_Unknown)
-			{
-				OutputDesc.Format = Inputs.OutputFormat;
-			}
-			OutputDesc.ClearValue = FClearValueBinding(FLinearColor::Black);
-			OutputDesc.Flags &= (~ETextureCreateFlags::FastVRAM);
-			OutputDesc.Flags |= GFastVRamConfig.PostProcessMaterial;
-
-			Output = FScreenPassRenderTarget(GraphBuilder.CreateTexture(OutputDesc, TEXT("PostProcessTempOutput")), SceneColor.ViewRect, View.GetOverwriteLoadAction());
-		}
-	}
-
-	const FScreenPassTextureViewport SceneColorViewport(SceneColor);
-	const FScreenPassTextureViewport OutputViewport(Output);
-
-	RDG_EVENT_SCOPE(GraphBuilder, "PostProcessMaterial::NeuralPass");
-
-	const uint32 MaterialStencilRef = Material->GetStencilRefValue();
-
-	const bool bMobilePlatform = IsMobilePlatform(View.GetShaderPlatform());
-
-
-	EScreenPassDrawFlags ScreenPassFlags = EScreenPassDrawFlags::AllowHMDHiddenAreaMask;
-
-	// check if we can skip that draw call in case if all pixels will fail the stencil test of the material
-	bool bSkipPostProcess = false;
-
-	if (Material->IsStencilTestEnabled() && IsPostProcessStencilTestAllowed())
-	{
-		bool bFailStencil = true;
-
-		const uint32 StencilComp = Material->GetStencilCompare();
-
-		// Always check against clear value, since a material might want to perform operations against that value
-		const uint32 StencilClearValue = Inputs.CustomDepthTexture ? Inputs.CustomDepthTexture->Desc.ClearValue.Value.DSValue.Stencil : 0;
-		bFailStencil &= PostProcessStencilTest(StencilClearValue, StencilComp, MaterialStencilRef);
-
-
-		for (const uint32& Value : View.CustomDepthStencilValues)
-		{
-			bFailStencil &= PostProcessStencilTest(Value, StencilComp, MaterialStencilRef);
-
-			if (!bFailStencil)
-			{
-				break;
-			}
-		}
-
-		bSkipPostProcess = bFailStencil;
-	}
-
-	if (!bSkipPostProcess)
-	{
-		NeuralPostProcessResource = AllocateNeuralPostProcessingResourcesIfNeeded(
-			GraphBuilder, OutputViewport, NeuralProfileId, Material->IsUsedWithNeuralNetworks());
-
-		if (NeuralPostProcessResource.IsValid())
-		{ 
-			// Prepass to extract the input to the NNE Engine
-			FPostProcessMaterialParameters* PostProcessMaterialParameters =
-				GetPostProcessMaterialParameters(GraphBuilder, Inputs, View, OutputViewport, Output, DepthStencilTexture, MaterialStencilRef, Material, MaterialShaderMap);
-
-			SetupNeuralPostProcessShaderParametersForWrite(PostProcessMaterialParameters->NeuralPostProcessParameters, GraphBuilder, NeuralPostProcessResource);
-
-			ClearUnusedGraphResources(NeuralPostProcessPassVertexShader, NeuralPostProcessPassPixelShader, PostProcessMaterialParameters);
-
-			//Only call the neural network when the shader resource is actually used.
-			if (IsNeuralPostProcessShaderParameterUsed(PostProcessMaterialParameters->NeuralPostProcessParameters))
-			{
-				AddDrawScreenPass(
-					GraphBuilder,
-#if RDG_EVENTS != RDG_EVENTS_STRING_COPY
-					RDG_EVENT_NAME("PostProcessMaterial(Neural Prepass)"),
-#else
-					FRDGEventName(*Material->GetAssetName()),
-#endif
-					View,
-					OutputViewport,
-					SceneColorViewport,
-					// Uses default depth stencil on mobile since the stencil test is done in pixel shader.
-					FScreenPassPipelineState(NeuralPostProcessPassVertexShader, NeuralPostProcessPassPixelShader, BlendState, DepthStencilState, MaterialStencilRef),
-					PostProcessMaterialParameters,
-					ScreenPassFlags,
-					[&View, NeuralPostProcessPassVertexShader, NeuralPostProcessPassPixelShader, MaterialRenderProxy, Material, PostProcessMaterialParameters](FRHICommandList& RHICmdList)
-					{
-						SetShaderParametersMixedVS(RHICmdList, NeuralPostProcessPassVertexShader, *PostProcessMaterialParameters, View, MaterialRenderProxy, *Material);
-						SetShaderParametersMixedPS(RHICmdList, NeuralPostProcessPassPixelShader, *PostProcessMaterialParameters, View, MaterialRenderProxy, *Material);
-					});
-
-				ApplyNeuralPostProcess(GraphBuilder, View, Output.ViewRect, NeuralPostProcessResource);
-			}
-		}
-	}
-}
-
 FScreenPassTexture AddPostProcessMaterialPass(
 	FRDGBuilder& GraphBuilder,
 	const FViewInfo& View,
@@ -801,6 +541,96 @@ FScreenPassTexture AddPostProcessMaterialPass(
 
 	const bool bMobilePlatform = IsMobilePlatform(View.GetShaderPlatform());
 
+	FPostProcessMaterialParameters* PostProcessMaterialParameters = GraphBuilder.AllocParameters<FPostProcessMaterialParameters>();
+	PostProcessMaterialParameters->SceneTextures = Inputs.SceneTextures;
+	PostProcessMaterialParameters->View = View.ViewUniformBuffer;
+	PostProcessMaterialParameters->EyeAdaptationBuffer = GraphBuilder.CreateSRV(GetEyeAdaptationBuffer(GraphBuilder, View));
+	PostProcessMaterialParameters->PostProcessOutput = GetScreenPassTextureViewportParameters(OutputViewport);
+	PostProcessMaterialParameters->RenderTargets[0] = Output.GetRenderTargetBinding();
+
+	// The target color will be decoded if bForceIntermediateTarget is true in any case, but we might still need to decode the input color
+	PostProcessMaterialParameters->bMetalMSAAHDRDecode = Inputs.bMetalMSAAHDRDecode ? 1 : 0;
+
+	if (DepthStencilTexture)
+	{
+		PostProcessMaterialParameters->RenderTargets.DepthStencil = FDepthStencilBinding(
+			DepthStencilTexture,
+			ERenderTargetLoadAction::ELoad,
+			ERenderTargetLoadAction::ELoad,
+			FExclusiveDepthStencil::DepthRead_StencilRead);
+	}
+	PostProcessMaterialParameters->ManualStencilReferenceValue = MaterialStencilRef;
+	PostProcessMaterialParameters->ManualStencilTestMask = GetManualStencilTestMask(Material->GetStencilCompare());
+
+	PostProcessMaterialParameters->PostProcessInput_BilinearSampler = TStaticSamplerState<SF_Bilinear, AM_Clamp, AM_Clamp, AM_Clamp>::GetRHI();;
+
+	const FScreenPassTexture BlackDummy(GSystemTextures.GetBlackDummy(GraphBuilder));
+
+    // This gets passed in whether or not it's used.
+	GraphBuilder.RemoveUnusedTextureWarning(BlackDummy.Texture);
+
+	FRHISamplerState* PointClampSampler = TStaticSamplerState<SF_Point, AM_Clamp, AM_Clamp, AM_Clamp>::GetRHI();
+
+	for (uint32 InputIndex = 0; InputIndex < kPostProcessMaterialInputCountMax; ++InputIndex)
+	{
+		FScreenPassTexture Input = Inputs.GetInput((EPostProcessMaterialInput)InputIndex);
+
+		// Need to provide valid textures for when shader compilation doesn't cull unused parameters.
+		if (!Input.Texture || !MaterialShaderMap->UsesSceneTexture(PPI_PostProcessInput0 + InputIndex))
+		{
+			Input = BlackDummy;
+		}
+
+		PostProcessMaterialParameters->PostProcessInput[InputIndex] = GetScreenPassTextureInput(Input, PointClampSampler);
+	}
+
+	// Path tracing buffer textures
+	for (uint32 InputIndex = 0; InputIndex < kPathTracingPostProcessMaterialInputCountMax; ++InputIndex)
+	{
+		FScreenPassTexture Input = Inputs.GetPathTracingInput((EPathTracingPostProcessMaterialInput)InputIndex);
+
+		if (!Input.Texture || !MaterialShaderMap->UsesPathTracingBufferTexture(InputIndex))
+		{
+			Input = BlackDummy;
+		}
+
+		PostProcessMaterialParameters->PathTracingPostProcessInput[InputIndex] = GetScreenPassTextureInput(Input, PointClampSampler);
+	}
+
+	PostProcessMaterialParameters->Substrate = Substrate::BindSubstrateGlobalUniformParameters(View);
+
+	// SceneDepthWithoutWater
+	const bool bHasValidSceneDepthWithoutWater = Inputs.SceneWithoutWaterTextures && Inputs.SceneWithoutWaterTextures->DepthTexture;
+	const bool bShouldUseBilinearSamplerForDepth = bHasValidSceneDepthWithoutWater && ShouldUseBilinearSamplerForDepthWithoutSingleLayerWater(Inputs.SceneWithoutWaterTextures->DepthTexture->Desc.Format);
+	PostProcessMaterialParameters->bSceneDepthWithoutWaterTextureAvailable = bHasValidSceneDepthWithoutWater;
+	PostProcessMaterialParameters->SceneDepthWithoutSingleLayerWaterSampler = bShouldUseBilinearSamplerForDepth ? TStaticSamplerState<SF_Bilinear>::GetRHI() : TStaticSamplerState<SF_Point>::GetRHI();
+	PostProcessMaterialParameters->SceneDepthWithoutSingleLayerWaterTexture = FRDGSystemTextures::Get(GraphBuilder).Black;
+	PostProcessMaterialParameters->SceneWithoutSingleLayerWaterMinMaxUV = FVector4f(0.0f, 0.0f, 1.0f, 1.0f);
+	PostProcessMaterialParameters->SceneWithoutSingleLayerWaterTextureSize = FVector2f(0.0f, 0.0f);
+	PostProcessMaterialParameters->SceneWithoutSingleLayerWaterInvTextureSize = FVector2f(0.0f, 0.0f);
+	if (bHasValidSceneDepthWithoutWater)
+	{
+		const bool bIsInstancedStereoSideBySide = View.bIsInstancedStereoEnabled && !View.bIsMobileMultiViewEnabled && IStereoRendering::IsStereoEyeView(View);
+		int32 WaterViewIndex = INDEX_NONE;
+		if (bIsInstancedStereoSideBySide)
+		{
+			WaterViewIndex = View.PrimaryViewIndex; // The instanced view does not have MinMaxUV initialized, instead the primary view MinMaxUV covers both eyes
+		}
+		else
+		{
+			verify(View.Family->Views.Find(&View, WaterViewIndex));
+		}
+
+		PostProcessMaterialParameters->SceneDepthWithoutSingleLayerWaterTexture = Inputs.SceneWithoutWaterTextures->DepthTexture;
+		PostProcessMaterialParameters->SceneWithoutSingleLayerWaterMinMaxUV = Inputs.SceneWithoutWaterTextures->Views[WaterViewIndex].MinMaxUV;
+
+		const FIntVector DepthTextureSize = Inputs.SceneWithoutWaterTextures->DepthTexture->Desc.GetSize();
+		PostProcessMaterialParameters->SceneWithoutSingleLayerWaterTextureSize = FVector2f(DepthTextureSize.X, DepthTextureSize.Y);
+		PostProcessMaterialParameters->SceneWithoutSingleLayerWaterInvTextureSize = FVector2f(1.0f / DepthTextureSize.X, 1.0f / DepthTextureSize.Y);
+	}
+
+	ClearUnusedGraphResources(VertexShader, PixelShader, PostProcessMaterialParameters);
+
 	EScreenPassDrawFlags ScreenPassFlags = EScreenPassDrawFlags::AllowHMDHiddenAreaMask;
 
 	// check if we can skip that draw call in case if all pixels will fail the stencil test of the material
@@ -832,46 +662,25 @@ FScreenPassTexture AddPostProcessMaterialPass(
 
 	if (!bSkipPostProcess)
 	{
-
-		FNeuralPostProcessResource NeuralPostProcessResource;
-		const bool bShouldApplyNeuralPostProcessing = ShouldApplyNeuralPostProcessForMaterial(Material);
-
-		if (bShouldApplyNeuralPostProcessing)
-		{
-			AddNeuralPostProcessPass(GraphBuilder, View, Inputs, MaterialInterface, NeuralPostProcessResource);
-		}
-
-		{
-			FPostProcessMaterialParameters* PostProcessMaterialParameters =
-				GetPostProcessMaterialParameters(GraphBuilder, Inputs, View, OutputViewport, Output, DepthStencilTexture, MaterialStencilRef, Material, MaterialShaderMap);
-
-			if (bShouldApplyNeuralPostProcessing)
-			{
-				SetupNeuralPostProcessShaderParametersForRead(PostProcessMaterialParameters->NeuralPostProcessParameters, GraphBuilder, NeuralPostProcessResource);
-			}
-
-			ClearUnusedGraphResources(VertexShader, PixelShader, PostProcessMaterialParameters);
-
-			AddDrawScreenPass(
-				GraphBuilder,
+		AddDrawScreenPass(
+			GraphBuilder,
 #if RDG_EVENTS != RDG_EVENTS_STRING_COPY
-				RDG_EVENT_NAME("PostProcessMaterial"),
+			RDG_EVENT_NAME("PostProcessMaterial"),
 #else
-				FRDGEventName(*Material->GetAssetName()),
+			FRDGEventName(*Material->GetAssetName()),
 #endif
-				View,
-				OutputViewport,
-				SceneColorViewport,
-				// Uses default depth stencil on mobile since the stencil test is done in pixel shader.
-				FScreenPassPipelineState(VertexShader, PixelShader, BlendState, DepthStencilState, MaterialStencilRef),
-				PostProcessMaterialParameters,
-				ScreenPassFlags,
-				[&View, VertexShader, PixelShader, MaterialRenderProxy, Material, PostProcessMaterialParameters](FRHICommandList& RHICmdList)
-				{
-					SetShaderParametersMixedVS(RHICmdList, VertexShader, *PostProcessMaterialParameters, View, MaterialRenderProxy, *Material);
-					SetShaderParametersMixedPS(RHICmdList, PixelShader, *PostProcessMaterialParameters, View, MaterialRenderProxy, *Material);
-				});
-		}
+			View,
+			OutputViewport,
+			SceneColorViewport,
+			// Uses default depth stencil on mobile since the stencil test is done in pixel shader.
+			FScreenPassPipelineState(VertexShader, PixelShader, BlendState, DepthStencilState, MaterialStencilRef),
+			PostProcessMaterialParameters,
+			ScreenPassFlags,
+			[&View, VertexShader, PixelShader, MaterialRenderProxy, Material, PostProcessMaterialParameters](FRHICommandList& RHICmdList)
+			{
+				SetShaderParametersMixedVS(RHICmdList, VertexShader, *PostProcessMaterialParameters, View, MaterialRenderProxy, *Material);
+				SetShaderParametersMixedPS(RHICmdList, PixelShader, *PostProcessMaterialParameters, View, MaterialRenderProxy, *Material);
+			});
 
 		if (bForceIntermediateTarget && !bCompositeWithInputAndDecode)
 		{
diff --git a/Engine/Source/Runtime/Renderer/Private/PostProcess/PostProcessMaterial.h b/Engine/Source/Runtime/Renderer/Private/PostProcess/PostProcessMaterial.h
--- a/Engine/Source/Runtime/Renderer/Private/PostProcess/PostProcessMaterial.h
+++ b/Engine/Source/Runtime/Renderer/Private/PostProcess/PostProcessMaterial.h
@@ -6,7 +6,6 @@
 #include "OverridePassSequence.h"
 #include "Substrate/Substrate.h"
 
-#include "NeuralPostProcess.h"
 #include "PostProcess/PostProcessMaterialInputs.h"
 
 class UMaterialInterface;
@@ -19,7 +18,6 @@ BEGIN_SHADER_PARAMETER_STRUCT(FPostProcessMaterialParameters, )
 	SHADER_PARAMETER_STRUCT_REF(FViewUniformShaderParameters, View)
 	SHADER_PARAMETER_STRUCT_INCLUDE(FSceneTextureShaderParameters, SceneTextures)
 	SHADER_PARAMETER_RDG_UNIFORM_BUFFER(FSubstrateGlobalUniformParameters, Substrate)
-	SHADER_PARAMETER_STRUCT_INCLUDE(FNeuralPostProcessShaderParameters, NeuralPostProcessParameters)
 	SHADER_PARAMETER_STRUCT(FScreenPassTextureViewportParameters, PostProcessOutput)
 	SHADER_PARAMETER_STRUCT_ARRAY(FScreenPassTextureInput, PostProcessInput, [kPostProcessMaterialInputCountMax])
 	SHADER_PARAMETER_STRUCT_ARRAY(FScreenPassTextureInput, PathTracingPostProcessInput, [kPathTracingPostProcessMaterialInputCountMax])
