Move SceneUniformBuffer header to public API

diff --git a/Engine/Source/Runtime/Renderer/Private/SceneUniformBuffer.h b/Engine/Source/Runtime/Renderer/Private/SceneUniformBuffer.h
deleted file mode 100644
diff --git a/Engine/Source/Runtime/Renderer/Public/SceneUniformBuffer.h b/Engine/Source/Runtime/Renderer/Public/SceneUniformBuffer.h
new file mode 100644
--- /dev/null
+++ b/Engine/Source/Runtime/Renderer/Public/SceneUniformBuffer.h
@@ -0,0 +1,218 @@
+// Copyright Epic Games, Inc. All Rights Reserved.
+
+#pragma once
+
+#include "RenderGraphDefinitions.h"
+#include "ShaderParameterMacros.h"
+#include "ShaderParameterMetadata.h"
+#include "ShaderParameterMetadataBuilder.h"
+
+// Construct a shader parameter struct with the specified type at Obj.
+// The resulting struct must be valid for usage (i.e. all members initialized)
+using FShaderParameterStructConstructor = TFunction<void(void* Obj, const FShaderParametersMetadata& Metadata, FRDGBuilder& GraphBuilder)>;
+template<typename TParameterStruct>
+using TShaderParameterStructConstructor = TFunction<void(TParameterStruct& Obj, FRDGBuilder& GraphBuilder)>;
+
+/**
+ * Collects scene UB members during static initialization,
+ * and uses these to construct the buffer layout during Main.
+ */
+class FSceneUniformBufferTypeRegistry
+{
+public:
+	using FMemberId = int;
+	struct FMemberInfo;
+	class FImpl;
+
+	// Must be called during static initialization
+	static FMemberId RegisterMember(
+		const TCHAR* Name,
+		const FShaderParametersMetadata& StructMetadata,
+		// Used to set default values for this member if no value is set explicitly
+		FShaderParameterStructConstructor DefaultValueConstructor);
+};
+
+class FSceneUniformBufferMemberRegistration
+{
+public:
+	using FMemberId = FSceneUniformBufferTypeRegistry::FMemberId;
+
+	// Called at runtime
+	virtual void Commit() = 0;
+	static void CommitAll();
+
+	FString Name;
+	FMemberId MemberId = -1;
+	FShaderParameterStructConstructor DefaultValueConstructor;
+
+protected:
+	static TArray<FSceneUniformBufferMemberRegistration*>& GetInstances();
+
+	// Called during static initialization
+	FSceneUniformBufferMemberRegistration(const TCHAR* Name)
+		: Name(Name)
+	{
+		GetInstances().Add(this);
+	}
+	virtual ~FSceneUniformBufferMemberRegistration() = default;
+};
+
+template<typename TMember>
+class TSceneUniformBufferMemberRegistration : public FSceneUniformBufferMemberRegistration
+{
+public:
+	TSceneUniformBufferMemberRegistration(const TCHAR* Name, TShaderParameterStructConstructor<TMember> TemplatedDefaultValueConstructor)
+		: FSceneUniformBufferMemberRegistration(Name)
+	{
+		DefaultValueConstructor = [=](void* Obj, const FShaderParametersMetadata& Metadata, FRDGBuilder& GraphBuilder)
+		{
+			check(&Metadata == TMember::FTypeInfo::GetStructMetadata());
+			TMember& Member = *(new(Obj) TMember());
+			TemplatedDefaultValueConstructor(Member, GraphBuilder);
+		};
+	}
+
+	void Commit() override
+	{
+		MemberId = FSceneUniformBufferTypeRegistry::RegisterMember(
+			*Name,
+			*TShaderParameterStructTypeInfo<TMember>::GetStructMetadata(),
+			DefaultValueConstructor);
+	}
+};
+
+#define DECLARE_SCENE_UB_STRUCT(StructType, FieldName, PrefixKeywords) \
+	namespace SceneUB { \
+		PrefixKeywords extern TSceneUniformBufferMemberRegistration<StructType> FieldName; \
+	}
+
+#define IMPLEMENT_SCENE_UB_STRUCT(StructType, FieldName, DefaultValueFactoryType) \
+	TSceneUniformBufferMemberRegistration<StructType> SceneUB::FieldName { TEXT(#FieldName), DefaultValueFactoryType } 
+
+
+DECLARE_UNIFORM_BUFFER_STRUCT(FSceneUniformParameters, RENDERER_API)
+/**
+ * RDG shader parameter struct containing data for FSceneUniformBuffer.
+ */
+MS_ALIGN(SHADER_PARAMETER_STRUCT_ALIGNMENT) class FSceneUniformParameters final
+{
+public:
+	struct FTypeInfo {
+		static constexpr int32 NumRows = 1;
+		static constexpr int32 NumColumns = 1;
+		static constexpr int32 NumElements = 0;
+		static constexpr int32 Alignment = SHADER_PARAMETER_STRUCT_ALIGNMENT;
+		static constexpr bool bIsStoredInConstantBuffer = true;
+		static constexpr const ANSICHAR* const FileName = UE_LOG_SOURCE_FILE(__FILE__);
+		static constexpr int32 FileLine = __LINE__;
+		using TAlignedType = FSceneUniformParameters;
+		static RENDERER_API const FShaderParametersMetadata* GetStructMetadata();
+	};
+
+	const void* GetContents() const
+	{
+		return Data;
+	}
+
+	FSceneUniformParameters();
+
+private:
+	friend class FSceneUniformBuffer;
+	friend class FSceneUniformBufferTypeRegistry::FImpl;
+
+	// TODO: Decouple parameter struct and actual data, so the data can be allocated dynamically (variable size)
+	// This requires some work in buffer binding/parameter macros/RHI, as these all assume the data to be inline.
+	static constexpr size_t MaxSize = 4096;
+	uint8_t Data[MaxSize];
+} GCC_ALIGN(SHADER_PARAMETER_STRUCT_ALIGNMENT);
+
+
+/**
+ * Holds scene-scoped parameters and stores these in uniform (constant) buffers for access on GPU.
+ */
+class FSceneUniformBuffer final
+{
+public:
+	using FMemberId = FSceneUniformBufferTypeRegistry::FMemberId;
+	using FRegistry = FSceneUniformBufferTypeRegistry;
+
+	RENDERER_API FSceneUniformBuffer();
+
+	/**
+	 * Set a field in the parameter struct.
+	 * The change will be reflected in any buffer that GetBuffer() returns after this call.
+	 * Returns true if anything actually changed.
+	 */
+	template<typename TMember>
+	bool Set(const TSceneUniformBufferMemberRegistration<TMember>& Registration, const TMember& Value)
+	{
+		return Set(Registration.MemberId, &Value, TMember::FTypeInfo::GetStructMetadata()->GetSize());
+	}
+
+	/**
+	 * Retrieve a field in the parameter struct.
+	 * If the field has not been set, this will crash.
+	 */
+	template<typename TMember>
+	const TMember& Get(const TSceneUniformBufferMemberRegistration<TMember>& Registration) const UE_LIFETIMEBOUND
+	{
+		const void* Ptr = Get(Registration.MemberId, TMember::FTypeInfo::GetStructMetadata()->GetSize());
+		checkf(Ptr, TEXT("SceneUB::%s has not been set yet"), *Registration.Name);
+		return *reinterpret_cast<const TMember*>(Ptr);
+	}
+
+	/**
+	 * Retrieve a field in the parameter struct.
+	 * If the field has not been set, this will set a default value.
+	 */
+	template<typename TMember>
+	const TMember& GetOrDefault(const TSceneUniformBufferMemberRegistration<TMember>& Registration, FRDGBuilder& GraphBuilder) const UE_LIFETIMEBOUND
+	{
+		return *reinterpret_cast<const TMember*>(GetOrDefault(Registration.MemberId, TMember::FTypeInfo::GetStructMetadata()->GetSize(), GraphBuilder));
+	}
+
+	// Get and re-create the UB if the cached parameters are modified.
+	RENDERER_API TRDGUniformBufferRef<FSceneUniformParameters> GetBuffer(FRDGBuilder& GraphBuilder);
+	RENDERER_API FRHIUniformBuffer* GetBufferRHI(FRDGBuilder& GraphBuilder);
+
+#if !UE_BUILD_SHIPPING
+	struct FDebugMemberInfo
+	{
+		FString ShaderType;
+		FString Name;
+		const void* Data;
+	};
+	TArray<FDebugMemberInfo> GetDebugInfo() const;
+#endif
+
+private:
+	template<typename TMember>
+	friend class TSceneUniformBufferMemberRegistration;
+	friend class FSceneUniformParameters;
+
+	/**
+	 * Set a field in the parameter struct.
+	 * The change will be reflected in any buffer that GetBuffer() returns after this call.
+	 * Returns true if anything actually changed.
+	 */
+	RENDERER_API bool Set(const FMemberId MemberId, const void* Value, const int32 ValueSize);
+
+	/**
+	 * Retrieve a field in the parameter struct.
+	 * If the field has not been set, this will return null;
+	 */
+	RENDERER_API const void* Get(const FMemberId MemberId, const int32 ExpectedValueSize) const UE_LIFETIMEBOUND;
+
+	/**
+	 * Retrieve a field in the parameter struct.
+	 * If the field has not been set, this will set a default value;
+	 */
+	RENDERER_API const void* GetOrDefault(const FMemberId MemberId, const int32 ExpectedValueSize, FRDGBuilder& GraphBuilder) UE_LIFETIMEBOUND;
+
+	TRDGUniformBufferRef<FSceneUniformParameters> Buffer;
+	FRHIUniformBuffer* RHIBuffer;
+
+	bool bAnyMemberDirty;
+	TBitArray<> MemberHasBeenSet;
+	TArray<uint8_t, TSizedHeapAllocator<32>> CachedData;
+};
