Skip precache and MDC building if custom water info texture rendering is disabled.

diff --git a/Engine/Source/Runtime/Renderer/Private/WaterInfoTextureRendering.cpp b/Engine/Source/Runtime/Renderer/Private/WaterInfoTextureRendering.cpp
--- a/Engine/Source/Runtime/Renderer/Private/WaterInfoTextureRendering.cpp
+++ b/Engine/Source/Runtime/Renderer/Private/WaterInfoTextureRendering.cpp
@@ -109,6 +109,12 @@ public:
 
 IMPLEMENT_GLOBAL_SHADER(FWaterInfoTextureBlurPS, "/Engine/Private/WaterInfoTextureBlur.usf", "Main", SF_Pixel);
 
+static bool IsCustomWaterInfoTextureRenderEnabled()
+{
+	static const IConsoleVariable* WaterRenderMethodCVar = IConsoleManager::Get().FindConsoleVariable(TEXT("r.Water.WaterInfo.RenderMethod"));
+	return WaterRenderMethodCVar && WaterRenderMethodCVar->GetInt();
+}
+
 /**
 * MeshPassProcessor for the "color" pass required for generating the water info texture. The associated pass draws water body meshes with an unlit material
 * in order to write water surface depth, river velocity and possibly other data too.
@@ -145,7 +151,7 @@ FWaterInfoTexturePassMeshProcessor::FWaterInfoTexturePassMeshProcessor(const FSc
 
 void FWaterInfoTexturePassMeshProcessor::AddMeshBatch(const FMeshBatch& RESTRICT MeshBatch, uint64 BatchElementMask, const FPrimitiveSceneProxy* RESTRICT PrimitiveSceneProxy, int32 StaticMeshId)
 {
-	if (!MeshBatch.bUseForMaterial)
+	if (!MeshBatch.bUseForMaterial || !IsCustomWaterInfoTextureRenderEnabled())
 	{
 		return;
 	}
@@ -171,7 +177,7 @@ void FWaterInfoTexturePassMeshProcessor::CollectPSOInitializers(const FSceneText
 	// Try to reduce the number of materials considered for this mesh pass:
 	// Materials for drawing the water info texture are supposed to be unlit (they write velocity and possibly other data into Emissive).
 	// They also need to be applied to meshes (MD_Surface) and we can safely exclude sky materials which also use the unlit shading model.
-	if (Material.GetShadingModels().IsUnlit() && Material.GetMaterialDomain() == MD_Surface && !Material.IsSky())
+	if (IsCustomWaterInfoTextureRenderEnabled() && Material.GetShadingModels().IsUnlit() && Material.GetMaterialDomain() == MD_Surface && !Material.IsSky())
 	{
 		// Determine the mesh's material and blend mode.
 		const FMeshDrawingPolicyOverrideSettings OverrideSettings = ComputeMeshOverrideSettings(PreCacheParams);
@@ -235,6 +241,10 @@ void FWaterInfoTexturePassMeshProcessor::CollectPSOInitializers(const FSceneText
 
 			Shaders.TryGetVertexShader(PassShaders.VertexShader);
 			Shaders.TryGetPixelShader(PassShaders.PixelShader);
+	
+			// TODO: find the correct values for mobile.
+			uint8 SubpassIndex = 0;
+			ESubpassHint SubpassHint = ESubpassHint::None;
 
 			AddGraphicsPipelineStateInitializer(
 				VertexFactoryData,
@@ -246,6 +256,8 @@ void FWaterInfoTexturePassMeshProcessor::CollectPSOInitializers(const FSceneText
 				MeshCullMode,
 				(EPrimitiveType)PreCacheParams.PrimitiveType,
 				EMeshPassFeatures::Default,
+				SubpassHint,
+				SubpassIndex,
 				true /*bRequired*/,
 				PSOInitializers);
 		}
@@ -421,6 +433,11 @@ FWaterInfoTextureDepthPassMeshProcessor::FWaterInfoTextureDepthPassMeshProcessor
 
 void FWaterInfoTextureDepthPassMeshProcessor::AddMeshBatch(const FMeshBatch& RESTRICT MeshBatch, uint64 BatchElementMask, const FPrimitiveSceneProxy* RESTRICT PrimitiveSceneProxy, int32 StaticMeshId)
 {
+	if(!IsCustomWaterInfoTextureRenderEnabled())
+	{
+		return;
+	}
+
 	const FMaterialRenderProxy* MaterialRenderProxy = MeshBatch.MaterialRenderProxy;
 	while (MaterialRenderProxy)
 	{
@@ -536,7 +553,7 @@ bool FWaterInfoTextureDepthPassMeshProcessor::Process(
 void FWaterInfoTextureDepthPassMeshProcessor::CollectPSOInitializers(const FSceneTexturesConfig& SceneTexturesConfig, const FMaterial& Material, const FPSOPrecacheVertexFactoryData& VertexFactoryData, const FPSOPrecacheParams& PreCacheParams, TArray<FPSOPrecacheData>& PSOInitializers)
 {
 	// We need to support all materials that could possibly be rendered in a depth-only pass. Unfortunately there doesn't seem to be a way to filter for bUseForWaterInfoTextureDepth at this point.
-	if (Material.GetMaterialDomain() == MD_Surface && !IsTranslucentBlendMode(Material))
+	if (IsCustomWaterInfoTextureRenderEnabled() && Material.GetMaterialDomain() == MD_Surface && !IsTranslucentBlendMode(Material))
 	{
 		// Determine the mesh's material and blend mode.
 		const FMeshDrawingPolicyOverrideSettings OverrideSettings = ComputeMeshOverrideSettings(PreCacheParams);
