Fix crash due to trying to render screen space shadows for lights without a valid DynamicShadowMapChannel assigment.

diff --git a/Engine/Source/Runtime/Renderer/Private/ShadowRendering.cpp b/Engine/Source/Runtime/Renderer/Private/ShadowRendering.cpp
--- a/Engine/Source/Runtime/Renderer/Private/ShadowRendering.cpp
+++ b/Engine/Source/Runtime/Renderer/Private/ShadowRendering.cpp
@@ -2431,7 +2431,12 @@ void FMobileSceneRenderer::RenderMobileShadowProjections(
 				LightSceneInfo,
 				bProjectingForForwardShading);
 
-			RenderScreenSpaceShadows(GraphBuilder, SceneTextures, Views, LightSceneInfo, bProjectingForForwardShading, ScreenShadowMaskTexture);
+			if (LightSceneInfo->GetDynamicShadowMapChannel() != -1)
+			{
+				// Dynamic shadows are projected into channels of the light attenuation texture based on their assigned DynamicShadowMapChannel
+				// Only render screen space shadows if light is assigned to a valid DynamicShadowMapChannel
+				RenderScreenSpaceShadows(GraphBuilder, SceneTextures, Views, LightSceneInfo, bProjectingForForwardShading, ScreenShadowMaskTexture);
+			}
 		}
 	}
 }
