Make sure readback buffers are also destroyed when not assigned to any passes - fix leak in RHIMisc

diff --git a/Engine/Source/Runtime/Renderer/Private/MeshDrawCommandStats.cpp b/Engine/Source/Runtime/Renderer/Private/MeshDrawCommandStats.cpp
--- a/Engine/Source/Runtime/Renderer/Private/MeshDrawCommandStats.cpp
+++ b/Engine/Source/Runtime/Renderer/Private/MeshDrawCommandStats.cpp
@@ -53,12 +53,12 @@ static TAutoConsoleVariable<int32> CVarExportMeshDrawCommandStats(
 FMeshDrawCommandStatsManager::FFrameData::~FFrameData()
 {
 	// Collect set of unique readback buffers for deletion (can be shared between MDCs and passes)
-	TSet<FRHIGPUBufferReadback*> ReadbackBuffers;
+	TSet<FRHIGPUBufferReadback*> ReadbackBuffers(RDGIndirectArgsReadbackBuffers);
 	for (FMeshDrawCommandPassStats* PassStats : PassData)
 	{
 		if (PassStats->InstanceCullingGPUBufferReadback)
 		{
-			ReadbackBuffers.Add(PassStats->InstanceCullingGPUBufferReadback);
+			check(ReadbackBuffers.Contains(PassStats->InstanceCullingGPUBufferReadback));
 			PassStats->InstanceCullingGPUBufferReadback = nullptr;
 		}
 		delete PassStats;
@@ -214,6 +214,14 @@ FRHIGPUBufferReadback* FMeshDrawCommandStatsManager::QueueDrawRDGIndirectArgsRea
 		{
 			GPUBufferReadback->EnqueueCopy(RHICmdList, DrawIndirectArgsRDG->GetRHI(), 0u);
 		});
+
+	// Make sure the readback buffer is stored for later deletion because the batch could be empty and then readback buffer might never be deleted
+	{
+		FScopeLock ScopeLock(&FrameDataCS);	
+		FFrameData* FrameData = GetOrAddFrameData();
+		FrameData->RDGIndirectArgsReadbackBuffers.Add(GPUBufferReadback);
+	}
+
 	return GPUBufferReadback;
 }
 
diff --git a/Engine/Source/Runtime/Renderer/Private/MeshDrawCommandStats.h b/Engine/Source/Runtime/Renderer/Private/MeshDrawCommandStats.h
--- a/Engine/Source/Runtime/Renderer/Private/MeshDrawCommandStats.h
+++ b/Engine/Source/Runtime/Renderer/Private/MeshDrawCommandStats.h
@@ -106,6 +106,7 @@ private:
 		int32 FrameNumber = 0;							//< Unique ever incrementing frame number
 		TArray<FMeshDrawCommandPassStats*> PassData;	//< Stats for all the MDC passes running during this frame
 		TMap<FRHIBuffer*, FIndirectArgsBufferResult> CustomIndirectArgsBufferResults; //< Custom indirect arg readback result buffer lookup map
+		TArray<FRHIGPUBufferReadback*> RDGIndirectArgsReadbackBuffers; //< All indirect args buffers requested via RDG passes
 		bool bIndirectArgReadbackRequested = false;		
 	};
 
