Fixes D3D11 validation failures in TSR

diff --git a/Engine/Shaders/Private/TemporalSuperResolution/TSRRejectShading.usf b/Engine/Shaders/Private/TemporalSuperResolution/TSRRejectShading.usf
--- a/Engine/Shaders/Private/TemporalSuperResolution/TSRRejectShading.usf
+++ b/Engine/Shaders/Private/TemporalSuperResolution/TSRRejectShading.usf
@@ -156,7 +156,7 @@ Texture2D<tsr_half2> DecimateMaskTexture;
 Texture2D<tsr_ushort> IsMovingMaskTexture;
 
 RWTexture2DArray<tsr_half4> HistoryGuideOutput;
-RWTexture2D<tsr_half4> HistoryMoireOutput;
+RWTexture2DArray<tsr_half4> HistoryMoireOutput;
 RWTexture2D<tsr_half4> HistoryRejectionOutput;
 RWTexture2D<tsr_halfC> InputSceneColorOutput;
 RWTexture2D<tsr_half> InputSceneColorLdrLumaOutput;
@@ -750,7 +750,7 @@ void MainCS(
 				SimdIndex);
 
 			#if CONFIG_MOIRE_DETECTION
-				HistoryMoireOutput[OutputPixelPos] = MoireHistory.GetElement(SimdIndex);
+				HistoryMoireOutput[tsr_short3(OutputPixelPos, 0)] = MoireHistory.GetElement(SimdIndex);
 			#endif
 		}
 	}
diff --git a/Engine/Shaders/Private/TemporalSuperResolution/TSRUpdateHistory.usf b/Engine/Shaders/Private/TemporalSuperResolution/TSRUpdateHistory.usf
--- a/Engine/Shaders/Private/TemporalSuperResolution/TSRUpdateHistory.usf
+++ b/Engine/Shaders/Private/TemporalSuperResolution/TSRUpdateHistory.usf
@@ -124,7 +124,7 @@ RWTexture2D<tsr_halfC> SceneColorOutputMip0;
 RWTexture2D<tsr_halfC> SceneColorOutputMip1;
 
 RWTexture2DArray<tsr_halfC> HistoryColorOutput;
-RWTexture2D<tsr_halfM> HistoryMetadataOutput;
+RWTexture2DArray<tsr_halfM> HistoryMetadataOutput;
 
 
 //------------------------------------------------------- LDS
@@ -1086,14 +1086,14 @@ void MainCS(
 			// Output final history lo pixel.
 			{
 				HistoryColorOutput[tsr_short3(dpv_lo(LocalHistoryPixelPos), HistoryArrayIndices_HighFrequency)] = dpv_lo(FinalHighFrequencyColor);
-				HistoryMetadataOutput[dpv_lo(LocalHistoryPixelPos)]  = dpv_lo(FinalMetadata);
+				HistoryMetadataOutput[tsr_short3(dpv_lo(LocalHistoryPixelPos), 0)]  = dpv_lo(FinalMetadata);
 			}
 		
 			// Output final history hi pixel.
 			#if CONFIG_ENABLE_DUAL_PIXEL_VECTORIZATION
 			{
 				HistoryColorOutput[tsr_short3(dpv_hi(LocalHistoryPixelPos), HistoryArrayIndices_HighFrequency)] = dpv_hi(FinalHighFrequencyColor);
-				HistoryMetadataOutput[dpv_hi(LocalHistoryPixelPos)] = dpv_hi(FinalMetadata);
+				HistoryMetadataOutput[tsr_short3(dpv_hi(LocalHistoryPixelPos), 0)] = dpv_hi(FinalMetadata);
 			}
 			#endif // !CONFIG_ENABLE_DUAL_PIXEL_VECTORIZATION
 		}
diff --git a/Engine/Source/Runtime/Renderer/Private/PostProcess/TemporalSuperResolution.cpp b/Engine/Source/Runtime/Renderer/Private/PostProcess/TemporalSuperResolution.cpp
--- a/Engine/Source/Runtime/Renderer/Private/PostProcess/TemporalSuperResolution.cpp
+++ b/Engine/Source/Runtime/Renderer/Private/PostProcess/TemporalSuperResolution.cpp
@@ -625,7 +625,7 @@ class FTSRRejectShadingCS : public FTSRShader
 		SHADER_PARAMETER_RDG_TEXTURE_SRV(Texture2D, IsMovingMaskTexture)
 
 		SHADER_PARAMETER_RDG_TEXTURE_UAV(RWTexture2DArray, HistoryGuideOutput)
-		SHADER_PARAMETER_RDG_TEXTURE_UAV(RWTexture2D, HistoryMoireOutput)
+		SHADER_PARAMETER_RDG_TEXTURE_UAV(RWTexture2DArray, HistoryMoireOutput)
 		SHADER_PARAMETER_RDG_TEXTURE_UAV(RWTexture2D, HistoryRejectionOutput)
 		SHADER_PARAMETER_RDG_TEXTURE_UAV(RWTexture2D, InputSceneColorOutput)
 		SHADER_PARAMETER_RDG_TEXTURE_UAV(RWTexture2D, InputSceneColorLdrLumaOutput)
@@ -816,7 +816,7 @@ class FTSRUpdateHistoryCS : public FTSRShader
 		SHADER_PARAMETER_RDG_TEXTURE_UAV(RWTexture2D, SceneColorOutputMip0)
 		SHADER_PARAMETER_RDG_TEXTURE_UAV(RWTexture2D, SceneColorOutputMip1)
 		SHADER_PARAMETER_RDG_TEXTURE_UAV(RWTexture2DArray, HistoryColorOutput)
-		SHADER_PARAMETER_RDG_TEXTURE_UAV(RWTexture2D, HistoryMetadataOutput)
+		SHADER_PARAMETER_RDG_TEXTURE_UAV(RWTexture2DArray, HistoryMetadataOutput)
 		SHADER_PARAMETER_RDG_TEXTURE_UAV(RWTexture2DArray, DebugOutput)
 	END_SHADER_PARAMETER_STRUCT()
 
@@ -2087,25 +2087,25 @@ FDefaultTemporalUpscaler::FOutputs AddTemporalSuperResolutionPasses(
 			}
 			else if (View.bStatePrevViewInfoIsReadOnly)
 			{
-				FRDGTextureDesc Desc = FRDGTextureDesc::Create2D(
+				FRDGTextureDesc Desc = FRDGTextureDesc::Create2DArray(
 					InputExtent,
 					History.MoireArray->Desc.Format,
 					FClearValueBinding::None,
-					/* InFlags = */ TexCreate_ShaderResource | TexCreate_UAV);
+					/* InFlags = */ TexCreate_ShaderResource | TexCreate_UAV,
+					/* InArraySize = */ 1);
 
 				// Create an unused texture for the moire history so that the VisualizeTSR can still display the updated moire history.
 				FRDGTextureRef UnusedMoireHistoryTexture = GraphBuilder.CreateTexture(Desc, TEXT("TSR.History.Moire"));
 				GraphBuilder.RemoveUnusedTextureWarning(UnusedMoireHistoryTexture);
 
 				PassParameters->HistoryMoireOutput = GraphBuilder.CreateUAV(UnusedMoireHistoryTexture);
-				MoireHistoryTexture = GraphBuilder.CreateSRV(FRDGTextureSRVDesc(UnusedMoireHistoryTexture));
+				MoireHistoryTexture = GraphBuilder.CreateSRV(FRDGTextureSRVDesc::CreateForSlice(UnusedMoireHistoryTexture, 0));
 			}
 			else
 			{
 				FRDGTextureUAVDesc MoireUAVDesc(History.MoireArray);
 				MoireUAVDesc.FirstArraySlice = CurrentFrameSliceIndex;
 				MoireUAVDesc.NumArraySlices = 1;
-				MoireUAVDesc.DimensionOverride = ETextureDimension::Texture2D;
 
 				PassParameters->HistoryMoireOutput = GraphBuilder.CreateUAV(MoireUAVDesc);
 
@@ -2299,7 +2299,6 @@ FDefaultTemporalUpscaler::FOutputs AddTemporalSuperResolutionPasses(
 			FRDGTextureUAVDesc MetadataUAVDesc(History.MetadataArray);
 			MetadataUAVDesc.FirstArraySlice = CurrentFrameSliceIndex;
 			MetadataUAVDesc.NumArraySlices = 1;
-			MetadataUAVDesc.DimensionOverride = ETextureDimension::Texture2D;
             
 			PassParameters->HistoryArrayIndices = HistoryArrayIndices;
 			PassParameters->HistoryColorOutput = GraphBuilder.CreateUAV(ColorUAVDesc);
