Implements TSR's shader permutations for processing the alpha channel

Motivations:
1) more easily tests alpha channel in any project,
2) add automated tests in EngineTests
2) catch shader compilation failures in preflights,

diff --git a/Engine/Shaders/Private/TemporalSuperResolution/TSRCommon.ush b/Engine/Shaders/Private/TemporalSuperResolution/TSRCommon.ush
--- a/Engine/Shaders/Private/TemporalSuperResolution/TSRCommon.ush
+++ b/Engine/Shaders/Private/TemporalSuperResolution/TSRCommon.ush
@@ -43,8 +43,8 @@
 
 #define CONFIG_SCENE_COLOR_OVERFLOW 1
 
-#if POST_PROCESS_ALPHA
-	#define CONFIG_SCENE_COLOR_ALPHA 1
+#if defined(DIM_ALPHA_CHANNEL)
+	#define CONFIG_SCENE_COLOR_ALPHA DIM_ALPHA_CHANNEL
 #else
 	#define CONFIG_SCENE_COLOR_ALPHA 0
 #endif
diff --git a/Engine/Shaders/Private/TemporalSuperResolution/TSRRejectShading.usf b/Engine/Shaders/Private/TemporalSuperResolution/TSRRejectShading.usf
--- a/Engine/Shaders/Private/TemporalSuperResolution/TSRRejectShading.usf
+++ b/Engine/Shaders/Private/TemporalSuperResolution/TSRRejectShading.usf
@@ -724,7 +724,7 @@ void MainCS(
 	}
 	#endif
 	
-	#if POST_PROCESS_ALPHA
+	#if CONFIG_SCENE_COLOR_ALPHA
 		OriginalInput = min(OriginalInput, tsr_simd_halfC::Const(LargestSceneColorRGBA));
 	#else
 		OriginalInput = min(OriginalInput, tsr_simd_halfC::Const(LargestSceneColorRGB));
diff --git a/Engine/Source/Runtime/Renderer/Private/PostProcess/TemporalSuperResolution.cpp b/Engine/Source/Runtime/Renderer/Private/PostProcess/TemporalSuperResolution.cpp
--- a/Engine/Source/Runtime/Renderer/Private/PostProcess/TemporalSuperResolution.cpp
+++ b/Engine/Source/Runtime/Renderer/Private/PostProcess/TemporalSuperResolution.cpp
@@ -18,6 +18,14 @@
 
 namespace
 {
+	
+TAutoConsoleVariable<int32> CVarTSRAlphaChannel(
+	TEXT("r.TSR.AplhaChannel"), -1,
+	TEXT("Controls whether TSR should process the scene color's alpha channel.\n")
+	TEXT(" -1: based of r.PostProcessing.PropagateAlpha (default);\n")
+	TEXT("  0: disabled;\n")
+	TEXT("  1: enabled.\n"),
+	ECVF_RenderThreadSafe);
 
 TAutoConsoleVariable<float> CVarTSRHistorySampleCount(
 	TEXT("r.TSR.History.SampleCount"), 16.0f,
@@ -339,6 +347,7 @@ enum class ETSRHistoryFormatBits : uint32
 {
 	None = 0,
 	Moire = 1 << 0,
+	AlphaChannel = 1 << 1,
 };
 ENUM_CLASS_FLAGS(ETSRHistoryFormatBits);
 
@@ -369,6 +378,7 @@ public:
 	static constexpr int32 kSupportMaxWaveSize = 64;
 
 	class F16BitVALUDim : SHADER_PERMUTATION_BOOL("DIM_16BIT_VALU");
+	class FAlphaChannelDim : SHADER_PERMUTATION_BOOL("DIM_ALPHA_CHANNEL");
 
 	FTSRShader(const ShaderMetaType::CompiledShaderInitializerType& Initializer)
 		: FGlobalShader(Initializer)
@@ -555,7 +565,7 @@ class FTSRDecimateHistoryCS : public FTSRShader
 
 	class FMoireReprojectionDim : SHADER_PERMUTATION_BOOL("DIM_MOIRE_REPROJECTION");
 	class FResurrectionReprojectionDim : SHADER_PERMUTATION_BOOL("DIM_RESURRECTION_REPROJECTION");
-	using FPermutationDomain = TShaderPermutationDomain<FMoireReprojectionDim, FResurrectionReprojectionDim>;
+	using FPermutationDomain = TShaderPermutationDomain<FMoireReprojectionDim, FResurrectionReprojectionDim, FTSRShader::FAlphaChannelDim>;
 
 	BEGIN_SHADER_PARAMETER_STRUCT(FParameters, )
 		SHADER_PARAMETER_STRUCT_INCLUDE(FTSRCommonParameters, CommonParameters)
@@ -599,7 +609,7 @@ class FTSRRejectShadingCS : public FTSRShader
 	class FFlickeringDetectionDim : SHADER_PERMUTATION_BOOL("DIM_FLICKERING_DETECTION");
 	class FHistoryResurrectionDim : SHADER_PERMUTATION_BOOL("DIM_HISTORY_RESURRECTION");
 
-	using FPermutationDomain = TShaderPermutationDomain<FWaveSizeOps, FFlickeringDetectionDim, FHistoryResurrectionDim, FTSRShader::F16BitVALUDim>;
+	using FPermutationDomain = TShaderPermutationDomain<FWaveSizeOps, FFlickeringDetectionDim, FHistoryResurrectionDim, FTSRShader::F16BitVALUDim, FTSRShader::FAlphaChannelDim>;
 
 	BEGIN_SHADER_PARAMETER_STRUCT(FParameters, )
 		SHADER_PARAMETER_STRUCT_INCLUDE(FTSRCommonParameters, CommonParameters)
@@ -775,7 +785,7 @@ class FTSRUpdateHistoryCS : public FTSRShader
 
 	class FQualityDim : SHADER_PERMUTATION_ENUM_CLASS("DIM_UPDATE_QUALITY", EQuality);
 
-	using FPermutationDomain = TShaderPermutationDomain<FQualityDim, FTSRShader::F16BitVALUDim>;
+	using FPermutationDomain = TShaderPermutationDomain<FQualityDim, FTSRShader::F16BitVALUDim, FTSRShader::FAlphaChannelDim>;
 
 	BEGIN_SHADER_PARAMETER_STRUCT(FParameters, )
 		SHADER_PARAMETER_STRUCT_INCLUDE(FTSRCommonParameters, CommonParameters)
@@ -852,7 +862,7 @@ class FTSRResolveHistoryCS : public FTSRShader
 
 	class FNyquistDim : SHADER_PERMUTATION_BOOL("DIM_NYQUIST");
 
-	using FPermutationDomain = TShaderPermutationDomain<FNyquistDim, FTSRShader::F16BitVALUDim>;
+	using FPermutationDomain = TShaderPermutationDomain<FNyquistDim, FTSRShader::F16BitVALUDim, FTSRShader::FAlphaChannelDim>;
 
 	BEGIN_SHADER_PARAMETER_STRUCT(FParameters, )
 		SHADER_PARAMETER_STRUCT_INCLUDE(FTSRCommonParameters, CommonParameters)
@@ -1214,7 +1224,7 @@ FDefaultTemporalUpscaler::FOutputs AddTemporalSuperResolutionPasses(
 	const ETSRSubpixelMethod SubpixelMethod = ETSRSubpixelMethod(FMath::Clamp(CVarTSRSubpixelMethod.GetValueOnRenderThread(), 0, 2));
 
 	// Whether alpha channel is supported.
-	const bool bSupportsAlpha = IsPostProcessingWithAlphaChannelSupported();
+	const bool bSupportsAlpha = CVarTSRAlphaChannel.GetValueOnRenderThread() >= 0 ? (CVarTSRAlphaChannel.GetValueOnRenderThread() > 0) : IsPostProcessingWithAlphaChannelSupported();
 
 	const float RefreshRateToFrameRateCap = (View.Family->Time.GetDeltaRealTimeSeconds() > 0.0f && CVarTSRFlickeringAdjustToFrameRate.GetValueOnRenderThread())
 		? View.Family->Time.GetDeltaRealTimeSeconds() * CVarTSRFlickeringFrameRateCap.GetValueOnRenderThread() : 1.0f;
@@ -1234,6 +1244,11 @@ FDefaultTemporalUpscaler::FOutputs AddTemporalSuperResolutionPasses(
 		{
 			HistoryFormatBits |= ETSRHistoryFormatBits::Moire;
 		}
+
+		if (bSupportsAlpha)
+		{
+			HistoryFormatBits |= ETSRHistoryFormatBits::AlphaChannel;
+		}
 	}
 	FTSRHistoryArrayIndices HistoryArrayIndices = TranslateHistoryFormatBitsToArrayIndices(HistoryFormatBits);
 
@@ -1319,7 +1334,7 @@ FDefaultTemporalUpscaler::FOutputs AddTemporalSuperResolutionPasses(
 	
 	RDG_EVENT_SCOPE(GraphBuilder, "TemporalSuperResolution(sg.AntiAliasingQuality=%d%s) %dx%d -> %dx%d",
 		CVarAntiAliasingQuality->GetInt(),
-		bSupportsAlpha ? TEXT(" Alpha") : TEXT(""),
+		bSupportsAlpha ? TEXT(" AlphaChannel") : TEXT(""),
 		InputRect.Width(), InputRect.Height(),
 		OutputRect.Width(), OutputRect.Height());
 	RDG_GPU_STAT_SCOPE(GraphBuilder, TemporalSuperResolution);
@@ -1929,14 +1944,16 @@ FDefaultTemporalUpscaler::FOutputs AddTemporalSuperResolutionPasses(
 		FTSRDecimateHistoryCS::FPermutationDomain PermutationVector;
 		PermutationVector.Set<FTSRDecimateHistoryCS::FMoireReprojectionDim>(FlickeringFramePeriod > 0.0f);
 		PermutationVector.Set<FTSRDecimateHistoryCS::FResurrectionReprojectionDim>(bCanResurrectHistory);
+		PermutationVector.Set<FTSRShader::FAlphaChannelDim>(bSupportsAlpha);
 
 		TShaderMapRef<FTSRDecimateHistoryCS> ComputeShader(View.ShaderMap, PermutationVector);
 		FComputeShaderUtils::AddPass(
 			GraphBuilder,
-			RDG_EVENT_NAME("TSR DecimateHistory(#%d %s%s) %dx%d",
+			RDG_EVENT_NAME("TSR DecimateHistory(#%d %s%s%s) %dx%d",
 				PermutationVector.ToDimensionValueId(),
 				PermutationVector.Get<FTSRDecimateHistoryCS::FMoireReprojectionDim>() ? TEXT("ReprojectMoire") : TEXT(""),
 				PermutationVector.Get<FTSRDecimateHistoryCS::FResurrectionReprojectionDim>() ? TEXT(" ReprojectResurrection") : TEXT(""),
+				PermutationVector.Get<FTSRShader::FAlphaChannelDim>() ? TEXT(" AlphaChannel") : TEXT(""),
 				InputRect.Width(), InputRect.Height()),
 			AsyncComputePasses >= 2 ? ERDGPassFlags::AsyncCompute : ERDGPassFlags::Compute,
 			ComputeShader,
@@ -2012,6 +2029,7 @@ FDefaultTemporalUpscaler::FOutputs AddTemporalSuperResolutionPasses(
 		PermutationVector.Set<FTSRRejectShadingCS::FFlickeringDetectionDim>(FlickeringFramePeriod > 0.0f);
 		PermutationVector.Set<FTSRRejectShadingCS::FHistoryResurrectionDim>(bCanResurrectHistory);
 		PermutationVector.Set<FTSRShader::F16BitVALUDim>(bUse16BitVALU);
+		PermutationVector.Set<FTSRShader::FAlphaChannelDim>(bSupportsAlpha);
 		PermutationVector = FTSRRejectShadingCS::RemapPermutation(PermutationVector);
 
 		const int32 GroupTileSize = 32;
@@ -2134,7 +2152,7 @@ FDefaultTemporalUpscaler::FOutputs AddTemporalSuperResolutionPasses(
 		TShaderMapRef<FTSRRejectShadingCS> ComputeShader(View.ShaderMap, PermutationVector);
 		FComputeShaderUtils::AddPass(
 			GraphBuilder,
-			RDG_EVENT_NAME("TSR RejectShading(#%d TileSize=%d PaddingCostMultiplier=%1.1f WaveSize=%d FlickeringFramePeriod=%f VALU=%s%s) %dx%d",
+			RDG_EVENT_NAME("TSR RejectShading(#%d TileSize=%d PaddingCostMultiplier=%1.1f WaveSize=%d FlickeringFramePeriod=%f VALU=%s%s%s) %dx%d",
 				PermutationVector.ToDimensionValueId(),
 				TileSize,
 				FMath::Pow(float(GroupTileSize) / float(TileSize), 2),
@@ -2142,6 +2160,7 @@ FDefaultTemporalUpscaler::FOutputs AddTemporalSuperResolutionPasses(
 				PassParameters->FlickeringFramePeriod,
 				PermutationVector.Get<FTSRShader::F16BitVALUDim>() ? TEXT("16bit") : TEXT("32bit"),
 				PermutationVector.Get<FTSRRejectShadingCS::FHistoryResurrectionDim>() ? TEXT(" Resurrection") : TEXT(""),
+				PermutationVector.Get<FTSRShader::FAlphaChannelDim>() ? TEXT(" AlphaChannel") : TEXT(""),
 				InputRect.Width(), InputRect.Height()),
 			AsyncComputePasses >= 3 ? ERDGPassFlags::AsyncCompute : ERDGPassFlags::Compute,
 			ComputeShader,
@@ -2354,14 +2373,16 @@ FDefaultTemporalUpscaler::FOutputs AddTemporalSuperResolutionPasses(
 		FTSRUpdateHistoryCS::FPermutationDomain PermutationVector;
 		PermutationVector.Set<FTSRUpdateHistoryCS::FQualityDim>(UpdateHistoryQuality);
 		PermutationVector.Set<FTSRShader::F16BitVALUDim>(bUse16BitVALU);
+		PermutationVector.Set<FTSRShader::FAlphaChannelDim>(bSupportsAlpha);
 
 		TShaderMapRef<FTSRUpdateHistoryCS> ComputeShader(View.ShaderMap, PermutationVector);
 		FComputeShaderUtils::AddPass(
 			GraphBuilder,
-			RDG_EVENT_NAME("TSR UpdateHistory(#%d Quality=%s%s%s%s) %dx%d",
+			RDG_EVENT_NAME("TSR UpdateHistory(#%d Quality=%s%s%s%s%s) %dx%d",
 				PermutationVector.ToDimensionValueId(),
 				kUpdateQualityNames[int32(PermutationVector.Get<FTSRUpdateHistoryCS::FQualityDim>())],
 				PermutationVector.Get<FTSRShader::F16BitVALUDim>() ? TEXT(" 16bit") : TEXT(""),
+				PermutationVector.Get<FTSRShader::FAlphaChannelDim>() ? TEXT(" AlphaChannel") : TEXT(""),
 				HistoryColorFormat == PF_FloatR11G11B10 ? TEXT(" R11G11B10") : TEXT(""),
 				PassParameters->bGenerateOutputMip2 ? TEXT(" OutputMip2") : (PassParameters->bGenerateOutputMip1 ? TEXT(" OutputMip1") : TEXT("")),
 				HistorySize.X, HistorySize.Y),
@@ -2413,15 +2434,17 @@ FDefaultTemporalUpscaler::FOutputs AddTemporalSuperResolutionPasses(
 		FTSRResolveHistoryCS::FPermutationDomain PermutationVector;
 		PermutationVector.Set<FTSRResolveHistoryCS::FNyquistDim>(bNyquistHistory && bUseWaveOps && GRHIMaximumWaveSize >= 32 && GRHIMinimumWaveSize <= 32);
 		PermutationVector.Set<FTSRShader::F16BitVALUDim>(bUse16BitVALU);
+		PermutationVector.Set<FTSRShader::FAlphaChannelDim>(bSupportsAlpha);
 		PermutationVector = FTSRResolveHistoryCS::RemapPermutation(PermutationVector);
 
 		TShaderMapRef<FTSRResolveHistoryCS> ComputeShader(View.ShaderMap, PermutationVector);
 		FComputeShaderUtils::AddPass(
 			GraphBuilder,
-			RDG_EVENT_NAME("TSR ResolveHistory(#%d%s%s) %dx%d",
+			RDG_EVENT_NAME("TSR ResolveHistory(#%d%s%s%s) %dx%d",
 				PermutationVector.ToDimensionValueId(),
 				PermutationVector.Get<FTSRResolveHistoryCS::FNyquistDim>() ? TEXT(" WaveOps") : TEXT(""),
 				PermutationVector.Get<FTSRShader::F16BitVALUDim>() ? TEXT(" 16bit") : TEXT(""),
+				PermutationVector.Get<FTSRShader::FAlphaChannelDim>() ? TEXT(" AlphaChannel") : TEXT(""),
 				OutputRect.Width(), OutputRect.Height()),
 			AsyncComputePasses >= 3 ? ERDGPassFlags::AsyncCompute : ERDGPassFlags::Compute,
 			ComputeShader,
