Modified RenderMobileShadowProjections to only calculate screen space shadows for directional lights.

diff --git a/Engine/Source/Runtime/Renderer/Private/ShadowRendering.cpp b/Engine/Source/Runtime/Renderer/Private/ShadowRendering.cpp
--- a/Engine/Source/Runtime/Renderer/Private/ShadowRendering.cpp
+++ b/Engine/Source/Runtime/Renderer/Private/ShadowRendering.cpp
@@ -2420,23 +2420,23 @@ void FMobileSceneRenderer::RenderMobileShadowProjections(
 		const FVisibleLightInfo& VisibleLightInfo = VisibleLightInfos[LightSceneInfo->Id];
 		const FLightSceneProxy* LightSceneProxy = LightSceneInfo->Proxy;
 
+		const bool bProjectingForForwardShading = true;
+
 		// Local light shadows don't render to shadow mask texture on mobile deferred
 		if (LightSceneProxy->GetLightType() == LightType_Directional || !IsMobileDeferredShadingEnabled(ShaderPlatform))
 		{
-			const bool bProjectingForForwardShading = true;
-
 			RenderShadowProjections(GraphBuilder, SceneTextures,
 				ScreenShadowMaskTexture,
 				nullptr,
 				LightSceneInfo,
 				bProjectingForForwardShading);
-
-			if (LightSceneInfo->GetDynamicShadowMapChannel() != -1)
-			{
-				// Dynamic shadows are projected into channels of the light attenuation texture based on their assigned DynamicShadowMapChannel
-				// Only render screen space shadows if light is assigned to a valid DynamicShadowMapChannel
-				RenderScreenSpaceShadows(GraphBuilder, SceneTextures, Views, LightSceneInfo, bProjectingForForwardShading, ScreenShadowMaskTexture);
-			}
+		}
+		
+		if (LightSceneProxy->GetLightType() == LightType_Directional && LightSceneInfo->GetDynamicShadowMapChannel() != -1)
+		{
+			// Dynamic shadows are projected into channels of the light attenuation texture based on their assigned DynamicShadowMapChannel
+			// Only render screen space shadows if light is assigned to a valid DynamicShadowMapChannel
+			RenderScreenSpaceShadows(GraphBuilder, SceneTextures, Views, LightSceneInfo, bProjectingForForwardShading, ScreenShadowMaskTexture);
 		}
 	}
 }
