Output the tonemapper's luminance history viewport for use with VRS image generation to avoid issues when using View.ViewRect or UnconstrainedViewRect for views with constrained aspect ratios (cinematic bars). Restructures some portions of the VRS manager and CAS generator to ensure initial generation is always performed once per view-family.

diff --git a/Engine/Shaders/Private/VariableRateShading/VRSShadingRateCalculate.usf b/Engine/Shaders/Private/VariableRateShading/VRSShadingRateCalculate.usf
--- a/Engine/Shaders/Private/VariableRateShading/VRSShadingRateCalculate.usf
+++ b/Engine/Shaders/Private/VariableRateShading/VRSShadingRateCalculate.usf
@@ -125,7 +125,6 @@ void CalculateShadingRateImageInner(
 	}
 
 	const float2 VRSTileSize = float2(THREADGROUP_SIZEX, THREADGROUP_SIZEY);
-	uint2 GroupOffset = ceil(ViewRect.xy / VRSTileSize);
 #if USE_WAVE_OP
 	if (InvalidThread)
 	{
@@ -135,7 +134,8 @@ void CalculateShadingRateImageInner(
 	int WaveSupportedShadingRate = WaveActiveBitAnd(~UnsupportedShadingRate);
 	if (GroupThreadId.x == 1 && GroupThreadId.y == 1)
 	{
-		VariableRateShadingTexture[GroupId.xy + GroupOffset] = VRSBitMask & WaveSupportedShadingRate;
+		// Generation is always done in one pass for a view family, so no need to account for stereo offsets
+		VariableRateShadingTexture[GroupId.xy] = VRSBitMask & WaveSupportedShadingRate;
 	}
 #else
 	uint unused;
@@ -147,7 +147,7 @@ void CalculateShadingRateImageInner(
 	GroupMemoryBarrierWithGroupSync();
 	if (GroupThreadId.x == 1 && GroupThreadId.y == 1)
 	{
-		VariableRateShadingTexture[GroupId.xy + GroupOffset] = GroupShadingRate;
+		VariableRateShadingTexture[GroupId.xy] = GroupShadingRate;
 	}
 #endif
 }
diff --git a/Engine/Source/Runtime/Renderer/Private/PostProcess/PostProcessTonemap.cpp b/Engine/Source/Runtime/Renderer/Private/PostProcess/PostProcessTonemap.cpp
--- a/Engine/Source/Runtime/Renderer/Private/PostProcess/PostProcessTonemap.cpp
+++ b/Engine/Source/Runtime/Renderer/Private/PostProcess/PostProcessTonemap.cpp
@@ -996,6 +996,7 @@ FScreenPassTexture AddTonemapPass(FRDGBuilder& GraphBuilder, const FViewInfo& Vi
 
 	if (OutputLuminance && View.ViewState)
 	{
+		View.ViewState->PrevFrameViewInfo.LuminanceViewRectHistory = OutputViewport.Rect;
 		GraphBuilder.QueueTextureExtraction(OutputLuminance, &View.ViewState->PrevFrameViewInfo.LuminanceHistory);
 	}
 
diff --git a/Engine/Source/Runtime/Renderer/Private/SceneRendering.h b/Engine/Source/Runtime/Renderer/Private/SceneRendering.h
--- a/Engine/Source/Runtime/Renderer/Private/SceneRendering.h
+++ b/Engine/Source/Runtime/Renderer/Private/SceneRendering.h
@@ -1139,6 +1139,9 @@ struct FPreviousViewInfo
 	// History for the final back buffer luminance
 	TRefCountPtr<IPooledRenderTarget> LuminanceHistory;
 
+	// History for the final back buffer luminance view rect
+	FIntRect LuminanceViewRectHistory;
+
 	// Mobile bloom setup eye adaptation surface.
 	TRefCountPtr<IPooledRenderTarget> MobileBloomSetup_EyeAdaptation;
 
diff --git a/Engine/Source/Runtime/Renderer/Private/VariableRateShading/ContrastAdaptiveImageGenerator.cpp b/Engine/Source/Runtime/Renderer/Private/VariableRateShading/ContrastAdaptiveImageGenerator.cpp
--- a/Engine/Source/Runtime/Renderer/Private/VariableRateShading/ContrastAdaptiveImageGenerator.cpp
+++ b/Engine/Source/Runtime/Renderer/Private/VariableRateShading/ContrastAdaptiveImageGenerator.cpp
@@ -21,7 +21,7 @@
 #include "DataDrivenShaderPlatformInfo.h"
 
 /**
- * Contrast Adaptive Shading (CAS) is a Tier 2 Variable Rate Shading method which generates a VRS image by examining the contrast from the previous frame.
+ * Contrast Adaptive Shading (CAS) is a Tier 2 Variable Rate Shading method which generates a shading rate image (SRI) by examining the contrast from the previous frame.
  * An image is generated which designates lower shading rates for areas of lower contrast in which reductions are unlikely to be noticed.
  * This image is then reprojected and rescaled in accordance with camera movement and dynamic resolution changes before being provided to the manager.
  */
@@ -251,6 +251,30 @@ namespace ESRIPreviewType
 	}
 };
 
+static FIntRect GetFullPostProcessOutputRect(const FSceneViewFamily& ViewFamily)
+{
+	// Get initial ViewRect based on luminance texture from previous frame
+	FIntRect FullLuminanceViewRect = FIntRect(0,0,0,0);
+	int32 NumViews = ViewFamily.Views.Num();
+
+	for (int32 ViewIndex = 0; ViewIndex < NumViews; ViewIndex++)
+	{
+		const FSceneView* View = ViewFamily.Views[ViewIndex];
+		check(View->bIsViewInfo);
+
+		if (ViewIndex == 0)
+		{
+			FullLuminanceViewRect = static_cast<const FViewInfo*>(View)->PrevViewInfo.LuminanceViewRectHistory;
+		}
+		else
+		{
+			FullLuminanceViewRect.Union(static_cast<const FViewInfo*>(View)->PrevViewInfo.LuminanceViewRectHistory); // Varies from luminance extent if cinematic bars are applied (constrained aspect ratio)
+		}
+	}
+
+	return FullLuminanceViewRect;
+}
+
 static const TCHAR* ShadingRateTextureName = TEXT("ShadingRateTexture");
 static const TCHAR* ScaledShadingRateTextureName = TEXT("ScaledShadingRateTexture");
 static const TCHAR* ScaledConservativeShadingRateTextureName = TEXT("ConservativeScaledShadingRateTexture");
@@ -269,11 +293,11 @@ struct RENDERER_API FVRSTextures
 		const FVRSTextures* VRSTextures = GraphBuilder.Blackboard.Get<FVRSTextures>();
 		return VRSTextures != nullptr;
 	}
-	void Create(FRDGBuilder& GraphBuilder, const FSceneViewFamily& ViewFamily)
+	void Create(FRDGBuilder& GraphBuilder, const FSceneViewFamily& ViewFamily, FIntRect PostProcessViewRect)
 	{
-		FRDGTextureDesc ConstructedSRIDesc = CreateSRIDesc(ViewFamily, false);
+		FRDGTextureDesc ConstructedSRIDesc = CreateSRIDesc(ViewFamily, false, PostProcessViewRect);
 		ConstructedSRI = GraphBuilder.CreateTexture(ConstructedSRIDesc, ShadingRateTextureName);
-		FRDGTextureDesc ScaledSRIDesc = CreateSRIDesc(ViewFamily, true);
+		FRDGTextureDesc ScaledSRIDesc = CreateSRIDesc(ViewFamily, true, PostProcessViewRect);
 		ScaledSRI = GraphBuilder.CreateTexture(ScaledSRIDesc, ScaledShadingRateTextureName);
 		ScaledConservativeSRI = GraphBuilder.CreateTexture(ScaledSRIDesc, ScaledConservativeShadingRateTextureName);
 	}
@@ -281,21 +305,17 @@ struct RENDERER_API FVRSTextures
 	FRDGTextureRef ScaledSRI;
 	FRDGTextureRef ScaledConservativeSRI;
 private:
-	static FRDGTextureDesc CreateSRIDesc(const FSceneViewFamily& ViewFamily, bool bIsForDynResScaled)
+	static FRDGTextureDesc CreateSRIDesc(const FSceneViewFamily& ViewFamily, bool bIsForDynResScaled, FIntRect PostProcessViewRect)
 	{
 		if (bIsForDynResScaled)
 		{
-			// Use SceneTextures
-			return FVariableRateShadingImageManager::GetSRIDesc();
+			// Use final ViewRect for final scaled SRI
+			return FVariableRateShadingImageManager::GetSRIDesc(ViewFamily);
 		}
 		else
 		{
-			// Get initial size based on luminance texture from previous frame
-			check(ViewFamily.Views[0]->bIsViewInfo);
-			const FViewInfo* ViewInfo = static_cast<const FViewInfo*>(ViewFamily.Views[0]);
-			const FIntPoint ViewTargetExtents = ViewInfo->PrevViewInfo.LuminanceHistory->GetDesc().Extent;
-
-			const FIntPoint SRIDimensions = FMath::DivideAndRoundUp(ViewTargetExtents, FVariableRateShadingImageManager::GetSRITileSize());
+			// Use luminance ViewRect to create initial unscaled image
+			const FIntPoint SRIDimensions = FMath::DivideAndRoundUp(PostProcessViewRect.Size(), FVariableRateShadingImageManager::GetSRITileSize());
 			return FRDGTextureDesc::Create2D(
 				SRIDimensions,
 				GRHIVariableRateShadingImageFormat,
@@ -320,24 +340,19 @@ static bool IsHDR10(const EDisplayOutputFormat& OutputFormat)
 		OutputFormat == EDisplayOutputFormat::HDR_ACES_2000nit_ST2084;
 }
 
-static FIntRect GetPostProcessOutputRect(const FViewInfo& ViewInfo)
-{
-	// If TAA/TSR is enabled, upscaling is done at the start of post-processing so the final output will match UnscaledViewRect. Otherwise use the dynamically rescale view rect since
-	// the secondary upscale will happen after post processing
-	return ViewInfo.PrimaryScreenPercentageMethod == EPrimaryScreenPercentageMethod::TemporalUpscale ? ViewInfo.UnscaledViewRect.Scale(ViewInfo.Family->SecondaryViewFraction) : ViewInfo.ViewRect;
-}
-
 bool AddCreateShadingRateImagePass(
 	FRDGBuilder& GraphBuilder,
-	const FViewInfo& View)
+	const FSceneViewFamily& ViewFamily,
+	FIntRect PostProcessViewRect)
 {
-	if (!View.PrevViewInfo.LuminanceHistory)
+	const FViewInfo& ViewInfo = *static_cast<const FViewInfo*>(ViewFamily.Views[0]);
+
+	if (!ViewInfo.PrevViewInfo.LuminanceHistory)
 	{
 		// Shading Rate Image unsupported
 		return false;
 	}
-
-	FRDGTextureRef Luminance = GraphBuilder.RegisterExternalTexture(View.PrevViewInfo.LuminanceHistory);
+	FRDGTextureRef Luminance = GraphBuilder.RegisterExternalTexture(ViewInfo.PrevViewInfo.LuminanceHistory);
 	const FVRSTextures& VRSTextures = FVRSTextures::Get(GraphBuilder);
 
 	{
@@ -347,14 +362,13 @@ bool AddCreateShadingRateImagePass(
 		PermutationVector.Set<FCalculateShadingRateImageCS::FThreadGroupX>(TileSize.X);
 		PermutationVector.Set<FCalculateShadingRateImageCS::FThreadGroupY>(TileSize.Y);
 
-		TShaderMapRef<FCalculateShadingRateImageCS> ComputeShader(View.ShaderMap, PermutationVector);
+		TShaderMapRef<FCalculateShadingRateImageCS> ComputeShader(ViewInfo.ShaderMap, PermutationVector);
 		auto* PassParameters = GraphBuilder.AllocParameters<FCalculateShadingRateImageCS::FParameters>();
-		EDisplayOutputFormat OutputDisplayFormat = GetDisplayOutputFormat(View);
-		FIntRect PostProcessRect = View.UnscaledViewRect.Scale(View.Family->SecondaryViewFraction);
+		EDisplayOutputFormat OutputDisplayFormat = GetDisplayOutputFormat(ViewInfo);
 		FCalculateShadingRateImageCS::InitParameters(
 			*PassParameters,
 			Luminance,
-			PostProcessRect,
+			PostProcessViewRect,
 			IsHDR10(OutputDisplayFormat),
 			GraphBuilder.CreateUAV(VRSTextures.ConstructedSRI));
 		FComputeShaderUtils::AddPass(
@@ -363,13 +377,13 @@ bool AddCreateShadingRateImagePass(
 			ERDGPassFlags::AsyncCompute | ERDGPassFlags::NeverCull,
 			ComputeShader,
 			PassParameters,
-			FComputeShaderUtils::GetGroupCount(PostProcessRect.Size(), TileSize));
+			FComputeShaderUtils::GetGroupCount(PostProcessViewRect.Size(), TileSize));
 	}
 
 	return true;
 }
 
-void AddPrepareImageBasedVRSPass(
+void AddReprojectImageBasedVRSPass(
 	FRDGBuilder& GraphBuilder,
 	const FMinimalSceneTextures& SceneTextures,
 	const FSceneViewFamily& ViewFamily)
@@ -390,11 +404,6 @@ void AddPrepareImageBasedVRSPass(
 		RDG_EVENT_SCOPE_CONDITIONAL(GraphBuilder, ViewFamily.Views.Num() > 1, "View%d", ViewIndex);
 		check(View->bIsViewInfo);
 		const FViewInfo* ViewInfo = static_cast<const FViewInfo*>(View);
-		if (View->bCameraCut || !FVariableRateShadingImageManager::IsVRSCompatibleWithView(*ViewInfo) || !ViewInfo->PrevViewInfo.LuminanceHistory)
-		{
-			break;
-		}
-		FIntPoint SrcBufferSize = FSceneTexturesConfig::Get().Extent;
 
 		TShaderMapRef<FRescaleVariableRateShadingCS> RescaleVariableRateShadingCS(ViewInfo->ShaderMap);
 
@@ -404,22 +413,23 @@ void AddPrepareImageBasedVRSPass(
 		int32 ScaledTilesWide = FMath::DivideAndRoundUp(ViewportWidth, TileSize.X);
 		int32 ScaledTilesHigh = FMath::DivideAndRoundUp(ViewportHeight, TileSize.Y);
 		FVector2f ScaledSRIDimensions(ScaledTilesWide, ScaledTilesHigh);
-
-		FIntRect PostProcessRect = GetPostProcessOutputRect(*ViewInfo);
+		
+		FIntRect SRIViewRect = ViewInfo->ViewRect.Scale(static_cast<float>(ViewInfo->PrevViewInfo.LuminanceViewRectHistory.Height()) / ViewportHeight);
 
 		FVector2f SRIViewRectMin(
-			static_cast<float>(FMath::DivideAndRoundDown(PostProcessRect.Min.X, TileSize.X)),
-			static_cast<float>(FMath::DivideAndRoundDown(PostProcessRect.Min.Y, TileSize.Y)));
+			static_cast<float>(FMath::DivideAndRoundDown(SRIViewRect.Min.X, TileSize.X)),
+			static_cast<float>(FMath::DivideAndRoundDown(SRIViewRect.Min.Y, TileSize.Y)));
 
 		FVector2f SRIViewRectMax(
-			static_cast<float>(FMath::DivideAndRoundUp(PostProcessRect.Max.X, TileSize.X)),
-			static_cast<float>(FMath::DivideAndRoundUp(PostProcessRect.Max.Y, TileSize.Y)));
+			static_cast<float>(FMath::DivideAndRoundUp(SRIViewRect.Max.X, TileSize.X)),
+			static_cast<float>(FMath::DivideAndRoundUp(SRIViewRect.Max.Y, TileSize.Y)));
 
 		FVector2f UVOffset(
-			(float)ViewInfo->ViewRect.Min.X / (float)SrcBufferSize.X,
-			(float)ViewInfo->ViewRect.Min.Y / (float)SrcBufferSize.Y);
+			static_cast<float>(ViewInfo->ViewRect.Min.X / TileSize.X) / TextureSize.X,
+			static_cast<float>(ViewInfo->ViewRect.Min.Y / TileSize.Y) / TextureSize.Y);
 
-		float DynamicResolutionScale = (float)ViewportWidth / (float)PostProcessRect.Width();
+		check(ViewInfo->PrevViewInfo.LuminanceViewRectHistory.Width());
+		float DynamicResolutionScale = static_cast<float>(ViewportWidth) / (ViewInfo->PrevViewInfo.LuminanceViewRectHistory.Width());
 
 		FRescaleVariableRateShadingCS::FParameters* PassParameters = GraphBuilder.AllocParameters<FRescaleVariableRateShadingCS::FParameters>();
 
@@ -471,6 +481,9 @@ void FContrastAdaptiveImageGenerator::PrepareImages(FRDGBuilder& GraphBuilder, c
 {
 	RDG_EVENT_SCOPE(GraphBuilder, "ContrastAdaptiveShading");
 
+	check(!ViewFamily.Views.IsEmpty());
+	check(ViewFamily.Views[0]->bIsViewInfo);
+
 	for (const FSceneView* View : ViewFamily.Views)
 	{
 		check(View->bIsViewInfo);
@@ -483,17 +496,12 @@ void FContrastAdaptiveImageGenerator::PrepareImages(FRDGBuilder& GraphBuilder, c
 	}
 
 	FVRSTextures& VRSTextures = GraphBuilder.Blackboard.Create<FVRSTextures>();
-	VRSTextures.Create(GraphBuilder, ViewFamily);
 
-	for (int32 ViewIndex = 0; ViewIndex < ViewFamily.Views.Num(); ViewIndex++)
-	{
-		check(ViewFamily.Views[ViewIndex]->bIsViewInfo);
-		const FViewInfo& View = *(FViewInfo*)ViewFamily.Views[ViewIndex];
-		RDG_GPU_MASK_SCOPE(GraphBuilder, View.GPUMask);
-		RDG_EVENT_SCOPE_CONDITIONAL(GraphBuilder, ViewFamily.Views.Num() > 1, "View%d", ViewIndex);
-		AddCreateShadingRateImagePass(GraphBuilder, View);
-	}
-	AddPrepareImageBasedVRSPass(GraphBuilder, SceneTextures, ViewFamily);
+	FIntRect PostProcessViewRect = GetFullPostProcessOutputRect(ViewFamily);
+
+	VRSTextures.Create(GraphBuilder, ViewFamily, PostProcessViewRect);
+	AddCreateShadingRateImagePass(GraphBuilder, ViewFamily, PostProcessViewRect);
+	AddReprojectImageBasedVRSPass(GraphBuilder, SceneTextures, ViewFamily);
 }
 
 bool FContrastAdaptiveImageGenerator::IsEnabled() const
diff --git a/Engine/Source/Runtime/Renderer/Private/VariableRateShading/FoveatedImageGenerator.cpp b/Engine/Source/Runtime/Renderer/Private/VariableRateShading/FoveatedImageGenerator.cpp
--- a/Engine/Source/Runtime/Renderer/Private/VariableRateShading/FoveatedImageGenerator.cpp
+++ b/Engine/Source/Runtime/Renderer/Private/VariableRateShading/FoveatedImageGenerator.cpp
@@ -194,7 +194,7 @@ void FFoveatedImageGenerator::PrepareImages(FRDGBuilder& GraphBuilder, const FSc
 	check(GRHIVariableRateShadingImageTileMinWidth >= 8 && GRHIVariableRateShadingImageTileMinWidth <= 64 && GRHIVariableRateShadingImageTileMinHeight >= 8 && GRHIVariableRateShadingImageTileMinHeight <= 64);
 
 	// Create texture to hold shading rate image
-	FRDGTextureDesc Desc = FVariableRateShadingImageManager::GetSRIDesc();
+	FRDGTextureDesc Desc = FVariableRateShadingImageManager::GetSRIDesc(ViewFamily);
 	FRDGTextureRef ShadingRateTexture = GraphBuilder.CreateTexture(Desc, TEXT("FoveatedShadingRateTexture"));
 
 	// Setup shader parameters and flags
diff --git a/Engine/Source/Runtime/Renderer/Private/VariableRateShading/VariableRateShadingImageManager.cpp b/Engine/Source/Runtime/Renderer/Private/VariableRateShading/VariableRateShadingImageManager.cpp
--- a/Engine/Source/Runtime/Renderer/Private/VariableRateShading/VariableRateShadingImageManager.cpp
+++ b/Engine/Source/Runtime/Renderer/Private/VariableRateShading/VariableRateShadingImageManager.cpp
@@ -353,18 +353,22 @@ FIntPoint FVariableRateShadingImageManager::GetSRITileSize()
 	return FIntPoint(GRHIVariableRateShadingImageTileMinWidth, GRHIVariableRateShadingImageTileMinHeight);
 }
 
-FRDGTextureDesc FVariableRateShadingImageManager::GetSRIDesc()
+FRDGTextureDesc FVariableRateShadingImageManager::GetSRIDesc(const FSceneViewFamily& ViewFamily)
 {
-	const FIntPoint Size = FSceneTexturesConfig::Get().Extent;
+	check(!ViewFamily.Views.IsEmpty())
+	check(ViewFamily.Views[0]->bIsViewInfo);
+
+	const FViewInfo* ViewInfo = static_cast<const FViewInfo*>(ViewFamily.Views[0]);
+	const FIntRect FamilyViewRect = ViewInfo->GetFamilyViewRect(); // May vary from the size of the scene textures if using constrained aspect ratios
+	
 	const FIntPoint TileSize = GetSRITileSize();
-	const FIntPoint SRISize = FMath::DivideAndRoundUp(Size, TileSize);
+	const FIntPoint SRISize = FMath::DivideAndRoundUp(FamilyViewRect.Size(), TileSize);
 
 	return FRDGTextureDesc::Create2D(
 		SRISize,
 		GRHIVariableRateShadingImageFormat,
 		FClearValueBinding::None,
 		TexCreate_Foveation | TexCreate_UAV | TexCreate_ShaderResource | TexCreate_DisableDCC);
-
 }
 
 int32 FVariableRateShadingImageManager::GetNumberOfSupportedRates()
@@ -398,7 +402,7 @@ FRDGTextureRef FVariableRateShadingImageManager::GetVariableRateShadingImage(FRD
 	// Use debug rate if provided
 	if (VRSForceRateForFrame >= 0)
 	{
-		return GetForceRateImage(GraphBuilder, VRSForceRateForFrame, ImageType);
+		return GetForceRateImage(GraphBuilder, *ViewInfo.Family, VRSForceRateForFrame, GetImageTypeFromPassType(PassType));
 	}
 
 	// Otherwise collate all internal sources
@@ -421,7 +425,7 @@ FRDGTextureRef FVariableRateShadingImageManager::GetVariableRateShadingImage(FRD
 	// If we have exactly one, the combiner will just return that
 	if (InternalVRSSources.Num())
 	{
-		return CombineShadingRateImages(GraphBuilder, ViewInfo, InternalVRSSources);
+		return CombineShadingRateImages(GraphBuilder, *ViewInfo.Family, InternalVRSSources);
 	}
 
 	// Default to nullptr if no sources are available
@@ -532,7 +536,7 @@ void FVariableRateShadingImageManager::DrawDebugPreview(FRDGBuilder& GraphBuilde
 			// Use debug rate if provided
 			if (VRSForceRateForFrame >= 0)
 			{
-				PreviewTexture = GetForceRateImage(GraphBuilder, VRSForceRateForFrame, PreviewImageType);
+				PreviewTexture = GetForceRateImage(GraphBuilder, ViewFamily, VRSForceRateForFrame, PreviewImageType);
 			}
 
 			// Otherwise collate debug images
@@ -554,12 +558,12 @@ void FVariableRateShadingImageManager::DrawDebugPreview(FRDGBuilder& GraphBuilde
 					}
 				}
 
-				PreviewTexture = CombineShadingRateImages(GraphBuilder, *ViewInfo, InternalVRSSources);
+				PreviewTexture = CombineShadingRateImages(GraphBuilder, ViewFamily, InternalVRSSources);
 
 				// Generate a dummy 1x1 image if we have no VRS sources
 				if (!PreviewTexture)
 				{
-					PreviewTexture = GetForceRateImage(GraphBuilder);
+					PreviewTexture = GetForceRateImage(GraphBuilder, ViewFamily);
 				}
 			}
 
@@ -636,7 +640,7 @@ void FVariableRateShadingImageManager::UnregisterExternalImageGenerator(IVariabl
  * Private functions
  */
 
-FRDGTextureRef FVariableRateShadingImageManager::CombineShadingRateImages(FRDGBuilder& GraphBuilder, const FViewInfo& ViewInfo, TArray<FRDGTextureRef> Sources)
+FRDGTextureRef FVariableRateShadingImageManager::CombineShadingRateImages(FRDGBuilder& GraphBuilder, const FSceneViewFamily& ViewFamily, TArray<FRDGTextureRef> Sources)
 {
 	// If we have more than one source, combine the first available two
 	// TODO: Support combining more textures
@@ -653,7 +657,7 @@ FRDGTextureRef FVariableRateShadingImageManager::CombineShadingRateImages(FRDGBu
 		SCOPED_NAMED_EVENT(CombineShadingRateImages, FColor::Green);
 
 		// Create texture to hold shading rate image
-		FRDGTextureRef CombinedShadingRateTexture = GraphBuilder.CreateTexture(GetSRIDesc(), TEXT("CombinedShadingRateTexture"));
+		FRDGTextureRef CombinedShadingRateTexture = GraphBuilder.CreateTexture(Sources[0]->Desc, TEXT("CombinedShadingRateTexture"));
 
 		FCombineShadingRateTexturesCS::FParameters* PassParameters = GraphBuilder.AllocParameters<FCombineShadingRateTexturesCS::FParameters>();
 		PassParameters->SourceTexture0 = Sources[0];
@@ -675,7 +679,7 @@ FRDGTextureRef FVariableRateShadingImageManager::CombineShadingRateImages(FRDGBu
 
 }
 
-FRDGTextureRef FVariableRateShadingImageManager::GetForceRateImage(FRDGBuilder& GraphBuilder, int RateIndex /* = 0*/, EVRSImageType ImageType /* = EVRSImageType::Full*/)
+FRDGTextureRef FVariableRateShadingImageManager::GetForceRateImage(FRDGBuilder& GraphBuilder, const FSceneViewFamily& ViewFamily, int RateIndex /* = 0*/, EVRSImageType ImageType /* = EVRSImageType::Full*/)
 {
 	static const TArray<uint32> ValidShadingRates = { VRSSR_1x1, VRSSR_1x2, VRSSR_2x1, VRSSR_2x2, VRSSR_2x4, VRSSR_4x2, VRSSR_4x4 };
 
@@ -691,7 +695,7 @@ FRDGTextureRef FVariableRateShadingImageManager::GetForceRateImage(FRDGBuilder&
 		RateIndex = 0; // Force to minimum shading rate if VRS is disabled for this pass
 	}
 
-	FRDGTextureRef ForceShadingRateTexture = GraphBuilder.CreateTexture(GetSRIDesc(), TEXT("ForceShadingRateTexture"));
+	FRDGTextureRef ForceShadingRateTexture = GraphBuilder.CreateTexture(GetSRIDesc(ViewFamily), TEXT("ForceShadingRateTexture"));
 	FRDGTextureUAVRef ForceShadingRateUAV = GraphBuilder.CreateUAV(ForceShadingRateTexture);
 	AddClearUAVPass(GraphBuilder, ForceShadingRateUAV, ValidShadingRates[RateIndex]);
 
diff --git a/Engine/Source/Runtime/Renderer/Public/VariableRateShadingImageManager.h b/Engine/Source/Runtime/Renderer/Public/VariableRateShadingImageManager.h
--- a/Engine/Source/Runtime/Renderer/Public/VariableRateShadingImageManager.h
+++ b/Engine/Source/Runtime/Renderer/Public/VariableRateShadingImageManager.h
@@ -108,7 +108,7 @@ public:
 	static bool IsVRSCompatibleWithOutputType(const EDisplayOutputFormat& DisplayOutputFormat);
 
 	static FIntPoint GetSRITileSize();
-	static FRDGTextureDesc GetSRIDesc();
+	static FRDGTextureDesc GetSRIDesc(const FSceneViewFamily& ViewFamily);
 	static int32 GetNumberOfSupportedRates();
 
 	void DrawDebugPreview(FRDGBuilder& GraphBuilder, const FSceneViewFamily& ViewFamily, FRDGTextureRef OutputSceneColor);
@@ -124,8 +124,8 @@ private:
 	// Register/UnregisterExternalImageGenerator could come from the game thread.
 	FRWLock	GeneratorsMutex;
 
-	FRDGTextureRef CombineShadingRateImages(FRDGBuilder& GraphBuilder, const FViewInfo& ViewInfo, TArray<FRDGTextureRef> Sources);
-	FRDGTextureRef GetForceRateImage(FRDGBuilder& GraphBuilder, int RateIndex = 0, EVRSImageType ImageType = EVRSImageType::Full);
+	FRDGTextureRef CombineShadingRateImages(FRDGBuilder& GraphBuilder, const FSceneViewFamily& ViewFamily, TArray<FRDGTextureRef> Sources);
+	FRDGTextureRef GetForceRateImage(FRDGBuilder& GraphBuilder, const FSceneViewFamily& ViewFamily, int RateIndex = 0, EVRSImageType ImageType = EVRSImageType::Full);
 
 	EVRSImageType GetImageTypeFromPassType(EVRSPassType PassType);
 
