LFV - Fixed LFV not showing on mobile preview anymore.

diff --git a/Engine/Source/Runtime/Renderer/Private/LocalFogVolumeRendering.cpp b/Engine/Source/Runtime/Renderer/Private/LocalFogVolumeRendering.cpp
--- a/Engine/Source/Runtime/Renderer/Private/LocalFogVolumeRendering.cpp
+++ b/Engine/Source/Runtime/Renderer/Private/LocalFogVolumeRendering.cpp
@@ -167,6 +167,14 @@ void SetDummyLocalFogVolumeForView(FRDGBuilder& GraphBuilder, FViewInfo& View)
 	View.LocalFogVolumeViewData.GPUTileDrawIndirectBufferUAV= nullptr;
 }
 
+void SetDummyLocalFogVolumeForViews(FRDGBuilder& GraphBuilder, TArray<FViewInfo>& Views)
+{
+	for (FViewInfo& View : Views)
+	{
+		SetDummyLocalFogVolumeForView(GraphBuilder, View);
+	}
+}
+
 
 /*=============================================================================
 	FScene functions
@@ -526,10 +534,7 @@ void InitLocalFogVolumesForViews(
 	}
 	else
 	{
-		for (FViewInfo& View : Views)
-		{
-			SetDummyLocalFogVolumeForView(GraphBuilder, View);
-		}
+		SetDummyLocalFogVolumeForViews(GraphBuilder, Views);
 	}
 }
 
diff --git a/Engine/Source/Runtime/Renderer/Private/LocalFogVolumeRendering.h b/Engine/Source/Runtime/Renderer/Private/LocalFogVolumeRendering.h
--- a/Engine/Source/Runtime/Renderer/Private/LocalFogVolumeRendering.h
+++ b/Engine/Source/Runtime/Renderer/Private/LocalFogVolumeRendering.h
@@ -123,6 +123,7 @@ bool ShouldRenderLocalFogVolumeInVolumetricFog(const FScene* Scene, const FScene
 
 void GetLocalFogVolumeSortingData(const FScene* Scene, FRDGBuilder& GraphBuilder, FLocalFogVolumeSortingData& Out);
 
+void SetDummyLocalFogVolumeForViews(FRDGBuilder& GraphBuilder, TArray<FViewInfo>& Views);
 void SetDummyLocalFogVolumeForView(FRDGBuilder& GraphBuilder, FViewInfo& View);
 
 void InitLocalFogVolumesForViews(
diff --git a/Engine/Source/Runtime/Renderer/Private/MobileShadingRenderer.cpp b/Engine/Source/Runtime/Renderer/Private/MobileShadingRenderer.cpp
--- a/Engine/Source/Runtime/Renderer/Private/MobileShadingRenderer.cpp
+++ b/Engine/Source/Runtime/Renderer/Private/MobileShadingRenderer.cpp
@@ -755,7 +755,6 @@ void FMobileSceneRenderer::RenderFullDepthPrepass(FRDGBuilder& GraphBuilder, TAr
 	for (FRenderViewContext& ViewContext : RenderViews)
 	{
 		FViewInfo& View = *ViewContext.ViewInfo;
-		SetDummyLocalFogVolumeForView(GraphBuilder, View);
 
 		if (!ViewContext.bIsFirstView)
 		{
@@ -1029,6 +1028,10 @@ void FMobileSceneRenderer::Render(FRDGBuilder& GraphBuilder)
 
 		BeginOcclusionScope(GraphBuilder, Views);
 	}
+	else
+	{
+		SetDummyLocalFogVolumeForViews(GraphBuilder, Views);
+	}
 	
 	// Sort objects' triangles
 	for (FViewInfo& View : Views)
