Console commands to recreate render state and update cached ray tracing state.

diff --git a/Engine/Source/Runtime/Renderer/Private/DeferredShadingRenderer.cpp b/Engine/Source/Runtime/Renderer/Private/DeferredShadingRenderer.cpp
--- a/Engine/Source/Runtime/Renderer/Private/DeferredShadingRenderer.cpp
+++ b/Engine/Source/Runtime/Renderer/Private/DeferredShadingRenderer.cpp
@@ -315,9 +315,20 @@ static TAutoConsoleVariable<int32> CVarTranslucencyVelocity(
 	TEXT("Whether translucency can draws depth/velocity (enabled by default)"),
 	ECVF_RenderThreadSafe);
 
+static FAutoConsoleCommand RecreateRenderStateContextCmd(
+	TEXT("r.RecreateRenderStateContext"),
+	TEXT("Recreate render state."),
+	FConsoleCommandDelegate::CreateStatic([] { FGlobalComponentRecreateRenderStateContext Context; }));
 
 #if RHI_RAYTRACING
 
+static bool bUpdateCachedRayTracingState = false;
+
+static FAutoConsoleCommand UpdateCachedRayTracingStateCmd(
+	TEXT("r.RayTracing.UpdateCachedState"),
+	TEXT("Update cached ray tracing state (mesh commands and instances)."),
+	FConsoleCommandDelegate::CreateStatic([] { bUpdateCachedRayTracingState = true; }));
+
 static bool bRefreshRayTracingInstances = false;
 
 static void RefreshRayTracingInstancesSinkFunction()
@@ -2830,7 +2841,11 @@ void FDeferredShadingSceneRenderer::Render(FRDGBuilder& GraphBuilder)
 #endif // WITH_EDITOR
 			const bool bNaniteRayTracingModeChanged = Nanite::GRayTracingManager.CheckModeChanged();
 
-			if (CurrentMode != Scene->CachedRayTracingMeshCommandsMode || bNaniteCoarseMeshStreamingModeChanged || bNaniteRayTracingModeChanged || bHasRayTracingEnableChanged)
+			if (CurrentMode != Scene->CachedRayTracingMeshCommandsMode
+				|| bNaniteCoarseMeshStreamingModeChanged
+				|| bNaniteRayTracingModeChanged
+				|| bHasRayTracingEnableChanged
+				|| bUpdateCachedRayTracingState)
 			{
 				Scene->WaitForCacheRayTracingPrimitivesTask();
 
@@ -2839,6 +2854,7 @@ void FDeferredShadingSceneRenderer::Render(FRDGBuilder& GraphBuilder)
 				Scene->CachedRayTracingMeshCommandsMode = CurrentMode;
 				Scene->RefreshRayTracingMeshCommandCache();
 				bHasRayTracingEnableChanged = false;
+				bUpdateCachedRayTracingState = false;
 			}
 
 			if (bRefreshRayTracingInstances)
@@ -2846,6 +2862,7 @@ void FDeferredShadingSceneRenderer::Render(FRDGBuilder& GraphBuilder)
 				Scene->WaitForCacheRayTracingPrimitivesTask();
 
 				// In some situations, we need to refresh the cached ray tracing instance.
+				// This assumes that cached instances will keep using the same LOD since CachedRayTracingMeshCommands is not recalculated
 				// eg: Need to update PrimitiveRayTracingFlags
 				// This operation is a bit expensive but only happens once as we transition between modes which should be rare.
 				Scene->RefreshRayTracingInstances();
