Fixes a source of ghosting in TSR's alpha channel

This notably implements new r.Test.Alpha.* cvars to more easily generate and debug an alpha channel in any project

diff --git a/Engine/Shaders/Private/TemporalSuperResolution/TSRRejectShading.usf b/Engine/Shaders/Private/TemporalSuperResolution/TSRRejectShading.usf
--- a/Engine/Shaders/Private/TemporalSuperResolution/TSRRejectShading.usf
+++ b/Engine/Shaders/Private/TemporalSuperResolution/TSRRejectShading.usf
@@ -513,6 +513,9 @@ void MeasureRejection(
 
 	tsr_simd_halfC RawFactor        = saturate(tsr_simd_halfC::Const(1.0) - RawClampedEnergy * rcp(Delta));
 	tsr_simd_half RawRejection      = min3(RawFactor[0], RawFactor[1], RawFactor[2]);
+	#if CONFIG_SCENE_COLOR_ALPHA
+		RawRejection = min(RawRejection, RawFactor[3]);
+	#endif
 
 	#if CONFIG_SCENE_COLOR_ALPHA
 		OutRejectionClampBlend = Min3x3(Median3x3(RawRejection));
@@ -524,6 +527,9 @@ void MeasureRejection(
 
 	tsr_simd_halfC FilteredFactor   = saturate(tsr_simd_halfC::Const(1.0) - FilteredClampedEnergy * rcp(Delta));
 	tsr_simd_half FilteredRejection = min3(FilteredFactor[0], FilteredFactor[1], FilteredFactor[2]);
+	#if CONFIG_SCENE_COLOR_ALPHA
+		FilteredRejection = min(FilteredRejection, FilteredFactor[3]);
+	#endif
 	OutRejectionBlendFinal          = FilteredRejection;
 } // MeasureRejection()
 
@@ -848,7 +854,11 @@ void MainCS(
 			tsr_simd_halfC ResurrectionMatch = abs(ExposedInput - ExposedResurrectedHistory);
 
 			tsr_simd_halfC MatchDiff = PrevMatch - ResurrectionMatch;
-			tsr_simd_half TotalMatchDiff = BlurEven3x3(MatchDiff[0] + MatchDiff[1] + MatchDiff[2]);
+			#if CONFIG_SCENE_COLOR_ALPHA
+				tsr_simd_half TotalMatchDiff = BlurEven3x3(MatchDiff[0] + MatchDiff[1] + MatchDiff[2] + MatchDiff[3]);
+			#else
+				tsr_simd_half TotalMatchDiff = BlurEven3x3(MatchDiff[0] + MatchDiff[1] + MatchDiff[2]);
+			#endif
 
 			tsr_simd_bool bResurrectionIsCloser = MedianBool3x3(TotalMatchDiff > tsr_simd_half::Const(tsr_half(0.05 * 3.0)));
 			tsr_simd_bool bPrevIsNotGoodEnough = RejectionBlendFinal < tsr_simd_half::Const(0.5);
diff --git a/Engine/Shaders/Private/Tools/DebugAlphaChannel.usf b/Engine/Shaders/Private/Tools/DebugAlphaChannel.usf
new file mode 100644
--- /dev/null
+++ b/Engine/Shaders/Private/Tools/DebugAlphaChannel.usf
@@ -0,0 +1,41 @@
+// Copyright Epic Games, Inc. All Rights Reserved.
+
+#include "../Common.ush"
+
+
+//------------------------------------------------------- PARAMETERS
+
+uint DebugModeId;
+float OpaqueWorldDistance;
+float OpaqueLerpWorldRange;
+
+Texture2D<float4> SceneColorTexture;
+Texture2D<float> SceneDepthTexture;
+RWTexture2D<float4> SceneColorOutput;
+
+
+//------------------------------------------------------- ENTRY POINT
+
+[numthreads(8, 8, 1)]
+void MainCS(uint2 DispatchThreadId : SV_DispatchThreadID)
+{
+	int2 PixelPos = View.ViewRectMinAndSize.xy + DispatchThreadId;
+
+	float DeviceZ = SceneDepthTexture[PixelPos];
+	float4 SceneColor = SceneColorTexture[PixelPos];
+
+	float WorldDepth = ConvertFromDeviceZ(DeviceZ);
+
+	float4 Output = SceneColor;
+
+	BRANCH
+	if (DebugModeId == 0)
+	{
+		float IsTranslucent = saturate((WorldDepth - OpaqueWorldDistance) * rcp(OpaqueLerpWorldRange));
+
+		SceneColor.rgb *= 1.0 - IsTranslucent;
+		SceneColor.a = IsTranslucent;
+	}
+
+	SceneColorOutput[PixelPos] = SceneColor; 
+}
diff --git a/Engine/Source/Runtime/Renderer/Private/DeferredShadingRenderer.cpp b/Engine/Source/Runtime/Renderer/Private/DeferredShadingRenderer.cpp
--- a/Engine/Source/Runtime/Renderer/Private/DeferredShadingRenderer.cpp
+++ b/Engine/Source/Runtime/Renderer/Private/DeferredShadingRenderer.cpp
@@ -96,6 +96,7 @@
 #include "SparseVolumeTexture/ISparseVolumeTextureStreamingManager.h"
 #include "WaterInfoTextureRendering.h"
 #include "SplineMeshSceneResources.h"
+#include "PostProcess/DebugAlphaChannel.h"
 
 #if !UE_BUILD_SHIPPING
 #include "RenderCaptureInterface.h"
@@ -4197,6 +4198,14 @@ void FDeferredShadingSceneRenderer::Render(FRDGBuilder& GraphBuilder)
 			RenderHairComposition(GraphBuilder, Views, SceneTextures.Color.Target, SceneTextures.Depth.Target, SceneTextures.Velocity, TranslucencyResourceMap);
 		}
 
+#if DEBUG_ALPHA_CHANNEL
+		if (ShouldMakeDistantGeometryTranslucent())
+		{
+			SceneTextures.Color = MakeDistanceGeometryTranslucent(GraphBuilder, Views, SceneTextures);
+			SceneTextures.UniformBuffer = CreateSceneTextureUniformBuffer(GraphBuilder, &SceneTextures, FeatureLevel, SceneTextures.SetupMode);
+		}
+#endif
+
 		// Experimental voxel test code
 		for (int32 ViewIndex = 0; ViewIndex < Views.Num(); ViewIndex++)
 		{
diff --git a/Engine/Source/Runtime/Renderer/Private/PostProcess/DebugAlphaChannel.cpp b/Engine/Source/Runtime/Renderer/Private/PostProcess/DebugAlphaChannel.cpp
new file mode 100644
--- /dev/null
+++ b/Engine/Source/Runtime/Renderer/Private/PostProcess/DebugAlphaChannel.cpp
@@ -0,0 +1,91 @@
+// Copyright Epic Games, Inc. All Rights Reserved.
+
+#include "PostProcess/DebugAlphaChannel.h"
+#include "ScenePrivate.h"
+#include "ScreenPass.h"
+#include "PixelShaderUtils.h"
+#include "SceneTextures.h"
+#include "DataDrivenShaderPlatformInfo.h"
+
+
+#if DEBUG_ALPHA_CHANNEL
+
+TAutoConsoleVariable<float> CVarTestAlphaOpaqueWorldDistance(
+	TEXT("r.Test.Aplha.OpaqueWorldDistance"), 0.0f,
+	TEXT("Sets the world distance beyond which the opaque pixel are lerped to translucent for testing purposes."),
+	ECVF_RenderThreadSafe);
+
+TAutoConsoleVariable<float> CVarTestAlphaOpaqueLerpWorldRange(
+	TEXT("r.Test.Aplha.OpaqueLerpWorldRange"), 100.0f,
+	TEXT("Sets the gradient length in world unit on which opaque pixel are lerped to translucent for testing purposes."),
+	ECVF_RenderThreadSafe);
+
+#endif // DEBUG_ALPHA_CHANNEL
+
+class FDebugAlphaChannelCS : public FGlobalShader
+{
+	DECLARE_GLOBAL_SHADER(FDebugAlphaChannelCS);
+	SHADER_USE_PARAMETER_STRUCT(FDebugAlphaChannelCS, FGlobalShader);
+
+	BEGIN_SHADER_PARAMETER_STRUCT(FParameters, )
+		SHADER_PARAMETER_STRUCT_REF(FViewUniformShaderParameters, ViewUniformBuffer)
+		SHADER_PARAMETER(int32, DebugModeId)
+		SHADER_PARAMETER(float, OpaqueWorldDistance)
+		SHADER_PARAMETER(float, OpaqueLerpWorldRange)
+
+		SHADER_PARAMETER_RDG_TEXTURE(Texture2D, SceneColorTexture)
+		SHADER_PARAMETER_RDG_TEXTURE(Texture2D, SceneDepthTexture)
+		SHADER_PARAMETER_RDG_TEXTURE_UAV(RWTexture2D, SceneColorOutput)
+	END_SHADER_PARAMETER_STRUCT()
+
+	static bool ShouldCompilePermutation(const FGlobalShaderPermutationParameters& Parameters)
+	{
+		return IsFeatureLevelSupported(Parameters.Platform, ERHIFeatureLevel::SM5);
+	}
+}; // class FDebugAlphaChannelCS
+
+IMPLEMENT_GLOBAL_SHADER(FDebugAlphaChannelCS, "/Engine/Private/Tools/DebugAlphaChannel.usf", "MainCS", SF_Compute);
+
+#if DEBUG_ALPHA_CHANNEL
+
+bool ShouldMakeDistantGeometryTranslucent()
+{
+	return CVarTestAlphaOpaqueWorldDistance.GetValueOnRenderThread() > 0.0f;
+}
+
+FRDGTextureMSAA MakeDistanceGeometryTranslucent(
+	FRDGBuilder& GraphBuilder,
+	TArrayView<const FViewInfo> Views,
+	FMinimalSceneTextures SceneTextures)
+{
+	FRDGTextureDesc Desc = SceneTextures.Color.Resolve->Desc;
+	Desc.Flags |= TexCreate_UAV;
+
+	FRDGTexture* NewSceneColorTexture = GraphBuilder.CreateTexture(Desc, TEXT("SceneColorWithAlpha"));
+
+	for (const FViewInfo& View : Views)
+	{
+		FDebugAlphaChannelCS::FParameters* PassParameters = GraphBuilder.AllocParameters<FDebugAlphaChannelCS::FParameters>();
+		PassParameters->ViewUniformBuffer = View.ViewUniformBuffer;
+		PassParameters->DebugModeId = 0; // MakeDistanceOpaqueTranslucent
+		PassParameters->OpaqueWorldDistance = CVarTestAlphaOpaqueWorldDistance.GetValueOnRenderThread();
+		PassParameters->OpaqueLerpWorldRange = CVarTestAlphaOpaqueLerpWorldRange.GetValueOnRenderThread();
+
+		PassParameters->SceneColorTexture = SceneTextures.Color.Resolve;
+		PassParameters->SceneDepthTexture = SceneTextures.Depth.Resolve;
+
+		PassParameters->SceneColorOutput = GraphBuilder.CreateUAV(NewSceneColorTexture);
+
+		TShaderMapRef<FDebugAlphaChannelCS> ComputeShader(View.ShaderMap);
+		FComputeShaderUtils::AddPass(
+			GraphBuilder,
+			RDG_EVENT_NAME("DebugAlphaChannel(MakeDistanceOpaqueTranslucent) %dx%d", View.ViewRect.Width(), View.ViewRect.Height()),
+			ComputeShader,
+			PassParameters,
+			FComputeShaderUtils::GetGroupCount(View.ViewRect.Size(), 8));
+	}
+
+	return FRDGTextureMSAA(NewSceneColorTexture);
+}
+
+#endif // DEBUG_ALPHA_CHANNEL
diff --git a/Engine/Source/Runtime/Renderer/Private/PostProcess/DebugAlphaChannel.h b/Engine/Source/Runtime/Renderer/Private/PostProcess/DebugAlphaChannel.h
new file mode 100644
--- /dev/null
+++ b/Engine/Source/Runtime/Renderer/Private/PostProcess/DebugAlphaChannel.h
@@ -0,0 +1,27 @@
+// Copyright Epic Games, Inc. All Rights Reserved.
+
+#pragma once
+
+#include "CoreMinimal.h"
+#include "RendererInterface.h"
+#include "ScreenPass.h"
+
+class FViewInfo;
+struct FMinimalSceneTextures;
+
+#if !UE_BUILD_SHIPPING
+	#define DEBUG_ALPHA_CHANNEL 1
+#else
+	#define DEBUG_ALPHA_CHANNEL 0
+#endif
+
+#if DEBUG_ALPHA_CHANNEL
+
+bool ShouldMakeDistantGeometryTranslucent();
+
+FRDGTextureMSAA MakeDistanceGeometryTranslucent(
+	FRDGBuilder& GraphBuilder,
+	TArrayView<const FViewInfo> Views,
+	FMinimalSceneTextures SceneTextures);
+
+#endif // DEBUG_ALPHA_CHANNEL
