LFV - fix for volumetric fog and local fog volume composition as part of the mobile fog rendering pass: now a common function is used to evalaute the start distance anmd avoid visual issues.

diff --git a/Engine/Source/Runtime/Renderer/Private/FogRendering.cpp b/Engine/Source/Runtime/Renderer/Private/FogRendering.cpp
--- a/Engine/Source/Runtime/Renderer/Private/FogRendering.cpp
+++ b/Engine/Source/Runtime/Renderer/Private/FogRendering.cpp
@@ -319,12 +319,7 @@ static void RenderViewFog(
 	GraphicsPSOInit.bDepthBounds = GSupportsDepthBoundsTest && CVarFogUseDepthBounds.GetValueOnAnyThread() && !bSkipDepthBound;
 	if (GraphicsPSOInit.bDepthBounds)
 	{
-		float ExpFogStartDistance = View.ExponentialFogParameters.W;
-		float VolFogStartDistance = bShouldRenderVolumetricFog ? View.VolumetricFogStartDistance : ExpFogStartDistance;
-
-		// The fog can be set to start at a certain euclidean distance.
-		// clamp the value to be behind the near plane z, according to the smallest distance between volumetric fog and height fog (if they are enabled). 
-		float FogStartDistance = FMath::Max(30.0f, FMath::Min(ExpFogStartDistance, VolFogStartDistance));
+		float FogStartDistance = GetViewFogCommonStartDistance(View, bShouldRenderVolumetricFog);
 
 		// Here we compute the nearest z value the fog can start
 		// to skip shader execution on pixels that are closer.
@@ -473,3 +468,14 @@ bool ShouldRenderFog(const FSceneViewFamily& Family)
 		&& !EngineShowFlags.StationaryLightOverlap 
 		&& !EngineShowFlags.LightMapDensity;
 }
+
+float GetViewFogCommonStartDistance(const FViewInfo& View, bool bShouldRenderVolumetricFog)
+{
+	float ExpFogStartDistance = View.ExponentialFogParameters.W;
+	float VolFogStartDistance = bShouldRenderVolumetricFog ? View.VolumetricFogStartDistance : ExpFogStartDistance;
+
+	// The fog can be set to start at a certain euclidean distance.
+	// clamp the value to be behind the near plane z, according to the smallest distance between volumetric fog and height fog (if they are enabled). 
+	float FogCommonStartDistance = FMath::Max(30.0f, FMath::Min(ExpFogStartDistance, VolFogStartDistance));
+	return FogCommonStartDistance;
+}
diff --git a/Engine/Source/Runtime/Renderer/Private/FogRendering.h b/Engine/Source/Runtime/Renderer/Private/FogRendering.h
--- a/Engine/Source/Runtime/Renderer/Private/FogRendering.h
+++ b/Engine/Source/Runtime/Renderer/Private/FogRendering.h
@@ -35,3 +35,5 @@ extern void SetupFogUniformParameters(FRDGBuilder& GraphBuilder, const FViewInfo
 TRDGUniformBufferRef<FFogUniformParameters> CreateFogUniformBuffer(FRDGBuilder& GraphBuilder, const FViewInfo& View);
 
 extern bool ShouldRenderFog(const FSceneViewFamily& Family);
+
+extern float GetViewFogCommonStartDistance(const FViewInfo& View, bool bShouldRenderVolumetricFog);
diff --git a/Engine/Source/Runtime/Renderer/Private/MobileFogRendering.cpp b/Engine/Source/Runtime/Renderer/Private/MobileFogRendering.cpp
--- a/Engine/Source/Runtime/Renderer/Private/MobileFogRendering.cpp
+++ b/Engine/Source/Runtime/Renderer/Private/MobileFogRendering.cpp
@@ -211,7 +211,7 @@ void FMobileSceneRenderer::RenderFog(FRHICommandList& RHICmdList, const FViewInf
 	SetGraphicsPipelineState(RHICmdList, GraphicsPSOInit, 0);
 	
 	// Use height fog start distance by default and fallback to AP distance
-	float FogStartDistance = View.ExponentialFogParameters.W;
+	float FogStartDistance = GetViewFogCommonStartDistance(View, bShouldRenderVolumetricFog);
 	if (!bUseHeightFog)
 	{
 		FSkyAtmosphereRenderSceneInfo& SkyInfo = *Scene->GetSkyAtmosphereSceneInfo();
