substrate - Fix for light function compi8lation with Vulkan.

diff --git a/Engine/Source/Runtime/Renderer/Private/MobileLocalLightsBuffer.cpp b/Engine/Source/Runtime/Renderer/Private/MobileLocalLightsBuffer.cpp
--- a/Engine/Source/Runtime/Renderer/Private/MobileLocalLightsBuffer.cpp
+++ b/Engine/Source/Runtime/Renderer/Private/MobileLocalLightsBuffer.cpp
@@ -144,6 +144,7 @@ public:
 		FMaterialShader::ModifyCompilationEnvironment(Parameters, OutEnvironment);
 		OutEnvironment.SetDefine(TEXT("LIGHT_FUNCTION"), 1);
 		OutEnvironment.SetDefine(TEXT("COMPUTE_SHADER"), 0);
+		OutEnvironment.SetDefine(TEXT("SUBSTRATE_INLINE_SHADING"), 1);
 		
 		if (MobileLocalLightsBufferPrepassEnabled(Parameters.Platform))
 		{
