Fix hair instance visibility (primary/shadow) tracking.

diff --git a/Engine/Plugins/Runtime/HairStrands/Source/HairStrandsCore/Private/GroomDebug.cpp b/Engine/Plugins/Runtime/HairStrands/Source/HairStrandsCore/Private/GroomDebug.cpp
--- a/Engine/Plugins/Runtime/HairStrands/Source/HairStrandsCore/Private/GroomDebug.cpp
+++ b/Engine/Plugins/Runtime/HairStrands/Source/HairStrandsCore/Private/GroomDebug.cpp
@@ -864,7 +864,7 @@ static void AddHairDebugPrintInstancePass(
 	const FSceneView& View,
 	const FShaderPrintData* ShaderPrintData,
 	const FHairStrandsInstances& Instances,
-	const FUintVector4& InstanceCountPerType)
+	const TArray<EHairInstanceVisibilityType>& InstancesVisibilityType)
 {
 	if (!ShaderPrint::IsSupported(View.GetShaderPlatform()))
 	{
@@ -873,6 +873,14 @@ static void AddHairDebugPrintInstancePass(
 
 	const uint32 InstanceCount = Instances.Num();
 
+	// Compute instace count per visibility types
+	TArray<uint32> InstanceCountPerType;
+	InstanceCountPerType.Init(0u, uint32(EHairInstanceVisibilityType::Count));
+	for (uint32 InstanceIndex = 0; InstanceIndex < InstanceCount; ++InstanceIndex)
+	{
+		InstanceCountPerType[uint32(InstancesVisibilityType[Instances[InstanceIndex]->RegisteredIndex])]++;
+	}
+
 	// Force ShaderPrint on.
 	ShaderPrint::SetEnabled(true);
 	// Request more drawing primitives & characters for printing if needed	
@@ -904,6 +912,7 @@ static void AddHairDebugPrintInstancePass(
 		const uint32 IntLODIndex = Instance->HairGroupPublicData->LODIndex;
 		const uint32 LODCount = Instance->HairGroupPublicData->GetLODScreenSizes().Num();
 		const EGroomCacheType ActiveGroomCacheType = GetHairInstanceCacheType(Instance);
+		const EHairInstanceVisibilityType VisibilityType = InstancesVisibilityType[Instance->RegisteredIndex];
 
 		FInstanceInfos& D = Infos.AddDefaulted_GetRef();
 		D.Data0.X =
@@ -939,8 +948,8 @@ static void AddHairDebugPrintInstancePass(
 
 				{
 					D.Data2 = FUintVector4(0);
-					D.Data2.X |= InstanceIndex < InstanceCountPerType[uint32(EHairInstanceCount::StrandsPrimaryView)]  ? 0x1u  : 0u;
-					D.Data2.X |= InstanceIndex < InstanceCountPerType[uint32(EHairInstanceCount::StrandsShadowView)]   ? 0x2u  : 0u;
+					D.Data2.X |= VisibilityType == EHairInstanceVisibilityType::StrandsPrimaryView           ? 0x1u  : 0u;
+					D.Data2.X |= VisibilityType == EHairInstanceVisibilityType::StrandsShadowView            ? 0x2u  : 0u;
 					D.Data2.X |= Instance->HairGroupPublicData->VFInput.Strands.Common.bScatterSceneLighting ? 0x4u  : 0u;
 					D.Data2.X |= Instance->HairGroupPublicData->VFInput.Strands.Common.bRaytracingGeometry   ? 0x8u  : 0u;
 					D.Data2.X |= Instance->HairGroupPublicData->VFInput.Strands.Common.bStableRasterization  ? 0x10u : 0u;
@@ -979,8 +988,11 @@ static void AddHairDebugPrintInstancePass(
 				D.Data0.W = Instance->Cards.LODs[IntLODIndex].Data->GetNumVertices();
 
 				D.Data2 = FUintVector4(0);
-				D.Data2.X |= InstanceIndex < InstanceCountPerType[uint32(EHairInstanceCount::CardsOrMeshesPrimaryView)] ? 0x1u : 0u;
-				D.Data2.X |= InstanceIndex < InstanceCountPerType[uint32(EHairInstanceCount::CardsOrMeshesShadowView)]  ? 0x2u : 0u;
+				D.Data2.X |= VisibilityType == EHairInstanceVisibilityType::CardsOrMeshesPrimaryView ? 0x1u : 0u;
+				D.Data2.X |= VisibilityType == EHairInstanceVisibilityType::CardsOrMeshesShadowView  ? 0x2u : 0u;
+
+				D.Data2.Y = D.Data0.W;
+				D.Data2.Z = D.Data0.Z;
 
 				#if RHI_RAYTRACING
 				bHasRaytracing = Instance->Cards.LODs[IntLODIndex].RaytracingResource != nullptr;
@@ -994,8 +1006,11 @@ static void AddHairDebugPrintInstancePass(
 				D.Data0.W = Instance->Meshes.LODs[IntLODIndex].Data->GetNumVertices();
 
 				D.Data2 = FUintVector4(0);
-				D.Data2.X |= InstanceIndex < InstanceCountPerType[uint32(EHairInstanceCount::CardsOrMeshesPrimaryView)] ? 0x1u : 0u;
-				D.Data2.X |= InstanceIndex < InstanceCountPerType[uint32(EHairInstanceCount::CardsOrMeshesShadowView)]  ? 0x2u : 0u;
+				D.Data2.X |= VisibilityType == EHairInstanceVisibilityType::CardsOrMeshesPrimaryView ? 0x1u : 0u;
+				D.Data2.X |= VisibilityType == EHairInstanceVisibilityType::CardsOrMeshesShadowView  ? 0x2u : 0u;
+
+				D.Data2.Y = D.Data0.W;
+				D.Data2.Z = D.Data0.Z;
 
 				#if RHI_RAYTRACING
 				bHasRaytracing = Instance->Meshes.LODs[IntLODIndex].RaytracingResource != nullptr;
@@ -1026,10 +1041,10 @@ static void AddHairDebugPrintInstancePass(
 	{
 		FHairDebugPrintInstanceCS::FParameters* Parameters = GraphBuilder.AllocParameters<FHairDebugPrintInstanceCS::FParameters>();
 		Parameters->InstanceCount = InstanceCount;
-		Parameters->InstanceCount_StrandsPrimaryView = InstanceCountPerType[uint32(EHairInstanceCount::StrandsPrimaryView)];
-		Parameters->InstanceCount_StrandsShadowView = InstanceCountPerType[uint32(EHairInstanceCount::StrandsShadowView)];
-		Parameters->InstanceCount_CardsOrMeshesPrimaryView = InstanceCountPerType[uint32(EHairInstanceCount::CardsOrMeshesPrimaryView)];
-		Parameters->InstanceCount_CardsOrMeshesShadowView = InstanceCountPerType[uint32(EHairInstanceCount::CardsOrMeshesShadowView)];
+		Parameters->InstanceCount_StrandsPrimaryView = InstanceCountPerType[uint32(EHairInstanceVisibilityType::StrandsPrimaryView)];
+		Parameters->InstanceCount_StrandsShadowView = InstanceCountPerType[uint32(EHairInstanceVisibilityType::StrandsShadowView)];
+		Parameters->InstanceCount_CardsOrMeshesPrimaryView = InstanceCountPerType[uint32(EHairInstanceVisibilityType::CardsOrMeshesPrimaryView)];
+		Parameters->InstanceCount_CardsOrMeshesShadowView = InstanceCountPerType[uint32(EHairInstanceVisibilityType::CardsOrMeshesShadowView)];
 		Parameters->InstanceNames = InstanceNames.GetParameters(GraphBuilder);
 		Parameters->AttributeNames = AttributeNames.GetParameters(GraphBuilder);
 		Parameters->Infos = GraphBuilder.CreateSRV(InfoBuffer, PF_R32_UINT);
@@ -1337,7 +1352,7 @@ void RunHairStrandsDebug(
 	FSceneInterface* Scene,
 	const FSceneView& View,
 	const FHairStrandsInstances& Instances,
-	const FUintVector4& InstanceCountPerType,
+	const TArray<EHairInstanceVisibilityType>& InstancesVisibilityType,
 	const FShaderPrintData* ShaderPrintData,
 	FRDGTextureRef SceneColorTexture,
 	FRDGTextureRef SceneDepthTexture,
@@ -1348,7 +1363,7 @@ void RunHairStrandsDebug(
 
 	if (ViewMode == EGroomViewMode::MacroGroups)
 	{
-		AddHairDebugPrintInstancePass(GraphBuilder, ShaderMap, View, ShaderPrintData, Instances, InstanceCountPerType);
+		AddHairDebugPrintInstancePass(GraphBuilder, ShaderMap, View, ShaderPrintData, Instances, InstancesVisibilityType);
 	}
 
 	if (ViewMode == EGroomViewMode::MeshProjection)
diff --git a/Engine/Plugins/Runtime/HairStrands/Source/HairStrandsCore/Private/GroomManager.cpp b/Engine/Plugins/Runtime/HairStrands/Source/HairStrandsCore/Private/GroomManager.cpp
--- a/Engine/Plugins/Runtime/HairStrands/Source/HairStrandsCore/Private/GroomManager.cpp
+++ b/Engine/Plugins/Runtime/HairStrands/Source/HairStrandsCore/Private/GroomManager.cpp
@@ -1275,7 +1275,7 @@ void RunHairStrandsDebug(
 	FSceneInterface* Scene,
 	const FSceneView& View,
 	const FHairStrandsInstances& Instances,
-	const FUintVector4& InstanceCountPerType,
+	const TArray<EHairInstanceVisibilityType>& InstancesVisibilityType,
 	const FShaderPrintData* ShaderPrintData,
 	FRDGTextureRef SceneColor,
 	FRDGTextureRef SceneDepth,
@@ -1418,7 +1418,7 @@ void ProcessHairStrandsBookmark(
 			Parameters.Scene,
 			*Parameters.View,
 			*Instances,
-			Parameters.InstanceCountPerType,
+			Parameters.InstancesVisibilityType,
 			Parameters.ShaderPrintData,
 			Parameters.SceneColorTexture,
 			Parameters.SceneDepthTexture,
diff --git a/Engine/Source/Runtime/Renderer/Private/DeferredShadingRenderer.cpp b/Engine/Source/Runtime/Renderer/Private/DeferredShadingRenderer.cpp
--- a/Engine/Source/Runtime/Renderer/Private/DeferredShadingRenderer.cpp
+++ b/Engine/Source/Runtime/Renderer/Private/DeferredShadingRenderer.cpp
@@ -4310,7 +4310,9 @@ void FDeferredShadingSceneRenderer::Render(FRDGBuilder& GraphBuilder)
 		{
 			if (HairStrandsBookmarkParameters.HasInstances())
 			{
-				RenderHairStrandsDebugInfo(GraphBuilder, Scene, Views, SceneTextures.Color.Target, SceneTextures.Depth.Target);
+				HairStrandsBookmarkParameters.SceneColorTexture = SceneTextures.Color.Target;
+				HairStrandsBookmarkParameters.SceneDepthTexture = SceneTextures.Depth.Target;
+				RenderHairStrandsDebugInfo(GraphBuilder, Scene, Views, HairStrandsBookmarkParameters);
 			}
 		}
 
diff --git a/Engine/Source/Runtime/Renderer/Private/HairStrands/HairStrandsDebug.cpp b/Engine/Source/Runtime/Renderer/Private/HairStrands/HairStrandsDebug.cpp
--- a/Engine/Source/Runtime/Renderer/Private/HairStrands/HairStrandsDebug.cpp
+++ b/Engine/Source/Runtime/Renderer/Private/HairStrands/HairStrandsDebug.cpp
@@ -1107,18 +1107,13 @@ static void InternalRenderHairStrandsDebugInfo(
 	FRDGBuilder& GraphBuilder,
 	FScene* Scene,
 	FViewInfo& View,
-	FRDGTextureRef SceneColorTexture,
-	FRDGTextureRef SceneDepthTexture)
+	FHairStrandsBookmarkParameters& Params)
 {
 	if (!ShaderPrint::IsSupported(View.GetShaderPlatform()))
 	{
 		return;
 	}
 
-	FHairStrandsBookmarkParameters Params;
-	CreateHairStrandsBookmarkParameters(Scene, View, Params);
-	Params.SceneColorTexture = SceneColorTexture;
-	Params.SceneDepthTexture = SceneDepthTexture;
 	if (!Params.HasInstances())
 	{
 		return;
@@ -1142,7 +1137,7 @@ static void InternalRenderHairStrandsDebugInfo(
 	// Display tangent vector for strands/cards/meshes
 	if (ViewMode == EGroomViewMode::Tangent)
 		{
-			AddDebugHairTangentPass(GraphBuilder, View, SceneTextures, SceneColorTexture);
+			AddDebugHairTangentPass(GraphBuilder, View, SceneTextures, Params.SceneColorTexture);
 		}
 
 	// Draw LOD info 
@@ -1163,14 +1158,13 @@ static void InternalRenderHairStrandsDebugInfo(
 		}
 	}
 
-
 	// Pass this point, all debug rendering concern only hair strands data
 	if (!HairStrands::HasViewHairStrandsData(View))
 	{
 		return;
 	}
 
-	const FScreenPassRenderTarget SceneColor(SceneColorTexture, View.ViewRect, ERenderTargetLoadAction::ELoad);
+	const FScreenPassRenderTarget SceneColor(Params.SceneColorTexture, View.ViewRect, ERenderTargetLoadAction::ELoad);
 
 	const FHairStrandsViewData& HairData = View.HairStrandsViewData;
 
@@ -1178,11 +1172,11 @@ static void InternalRenderHairStrandsDebugInfo(
 	{
 		if (GHairStrandsDebugPlotBsdf > 0)
 		{
-			AddPlotBSDFPass(GraphBuilder, View, SceneColorTexture);
+			AddPlotBSDFPass(GraphBuilder, View, Params.SceneColorTexture);
 		}
 		if (HairData.DebugData.IsPlotDataValid())
 		{
-			AddPlotSamplePass(GraphBuilder, View, HairData.DebugData.PlotData, SceneColorTexture);
+			AddPlotSamplePass(GraphBuilder, View, HairData.DebugData.PlotData, Params.SceneColorTexture);
 		}	
 	}
 
@@ -1213,7 +1207,7 @@ static void InternalRenderHairStrandsDebugInfo(
 				if (DeepShadowData.AtlasSlotIndex != DomIndex)
 					continue;
 
-				AddDebugDeepShadowTexturePass(GraphBuilder, &View, FIntRect(), &DeepShadowData, &HairData.DeepShadowResources, SceneColorTexture);
+				AddDebugDeepShadowTexturePass(GraphBuilder, &View, FIntRect(), &DeepShadowData, &HairData.DeepShadowResources, Params.SceneColorTexture);
 			}
 		}
 	}
@@ -1223,11 +1217,11 @@ static void InternalRenderHairStrandsDebugInfo(
 	{
 		for (const FHairStrandsMacroGroupData& MacroGroupData : HairData.MacroGroupDatas)
 		{
-			AddDebugDeepShadowTexturePass(GraphBuilder, &View, MacroGroupData.ScreenRect, nullptr, nullptr, SceneColorTexture);
+			AddDebugDeepShadowTexturePass(GraphBuilder, &View, MacroGroupData.ScreenRect, nullptr, nullptr, Params.SceneColorTexture);
 		}
 
 		const FIntRect TotalRect = ComputeVisibleHairStrandsMacroGroupsRect(View.ViewRect, HairData.MacroGroupDatas);
-		AddDebugDeepShadowTexturePass(GraphBuilder, &View, TotalRect, nullptr, nullptr, SceneColorTexture);
+		AddDebugDeepShadowTexturePass(GraphBuilder, &View, TotalRect, nullptr, nullptr, Params.SceneColorTexture);
 	}
 	
 
@@ -1237,7 +1231,7 @@ static void InternalRenderHairStrandsDebugInfo(
 		{
 			if (HairData.DeepShadowResources.bIsGPUDriven)
 			{
-				AddDeepShadowInfoPass(GraphBuilder, View, HairData.DeepShadowResources, HairData.MacroGroupResources, SceneColorTexture);
+				AddDeepShadowInfoPass(GraphBuilder, View, HairData.DeepShadowResources, HairData.MacroGroupResources, Params.SceneColorTexture);
 			}
 		}
 	}
@@ -1254,13 +1248,13 @@ static void InternalRenderHairStrandsDebugInfo(
 		ViewMode == EGroomViewMode::MaterialTangent;
 	if (bRunDebugPass)
 	{
-		AddDebugHairPass(GraphBuilder, &View, ViewMode, HairData.VisibilityData, SceneTextures.Stencil, SceneColorTexture);
+		AddDebugHairPass(GraphBuilder, &View, ViewMode, HairData.VisibilityData, SceneTextures.Stencil, Params.SceneColorTexture);
 		AddDebugHairPrintPass(GraphBuilder, Scene, &View, ViewMode, HairData.VisibilityData, HairData.MacroGroupDatas, HairData.MacroGroupResources, SceneTextures.Stencil);
 	}
 	else if (ViewMode == EGroomViewMode::Tile)
 	{
 		check(HairData.VisibilityData.TileData.IsValid());
-		AddHairStrandsDebugTilePass(GraphBuilder, View, SceneColorTexture, HairData.VisibilityData.TileData);
+		AddHairStrandsDebugTilePass(GraphBuilder, View, Params.SceneColorTexture, HairData.VisibilityData.TileData);
 	}
 
 	const bool bIsVoxelMode = ViewMode == EGroomViewMode::VoxelsDensity;
@@ -1268,7 +1262,7 @@ static void InternalRenderHairStrandsDebugInfo(
 	{
 		if (HairData.VirtualVoxelResources.IsValid())
 		{
-			AddVoxelPageRaymarchingPass(GraphBuilder, View, HairData.MacroGroupDatas, HairData.VirtualVoxelResources, SceneColorTexture);
+			AddVoxelPageRaymarchingPass(GraphBuilder, View, HairData.MacroGroupDatas, HairData.VirtualVoxelResources, Params.SceneColorTexture);
 		}
 	}
 
@@ -1286,12 +1280,12 @@ static void InternalRenderHairStrandsDebugInfo(
 		PassParameters->PPLLNodeIndex = HairData.DebugData.PPLLData.NodeIndexTexture;
 		PassParameters->PPLLNodeData = GraphBuilder.CreateSRV(FRDGBufferSRVDesc(HairData.DebugData.PPLLData.NodeDataBuffer));
 		PassParameters->ViewUniformBuffer = View.ViewUniformBuffer;
-		PassParameters->SceneColorTextureUAV = GraphBuilder.CreateUAV(SceneColorTexture);
+		PassParameters->SceneColorTextureUAV = GraphBuilder.CreateUAV(Params.SceneColorTexture);
 		ShaderPrint::SetParameters(GraphBuilder, View.ShaderPrintData, PassParameters->ShaderPrintParameters);
 
 		FHairVisibilityDebugPPLLCS::FPermutationDomain PermutationVector;
 		TShaderMapRef<FHairVisibilityDebugPPLLCS> ComputeShader(View.ShaderMap, PermutationVector);
-		FIntVector TextureSize = SceneColorTexture->Desc.GetSize(); TextureSize.Z = 1;
+		FIntVector TextureSize = Params.SceneColorTexture->Desc.GetSize(); TextureSize.Z = 1;
 		FComputeShaderUtils::AddPass(GraphBuilder, RDG_EVENT_NAME("HairStrands::PPLLDebug"), ComputeShader, PassParameters, FIntVector::DivideAndRoundUp(TextureSize, FIntVector(8, 8, 1)));
 	}
 
@@ -1316,12 +1310,11 @@ void RenderHairStrandsDebugInfo(
 	FRDGBuilder& GraphBuilder,
 	FScene* Scene,
 	TArrayView<FViewInfo> Views,
-	FRDGTextureRef SceneColorTexture,
-	FRDGTextureRef SceneDepthTexture)
+	FHairStrandsBookmarkParameters& Parameters)
 {
 	bool bHasHairData = false;
 	for (FViewInfo& View : Views)
 	{
-		InternalRenderHairStrandsDebugInfo(GraphBuilder, Scene, View, SceneColorTexture, SceneDepthTexture);
+		InternalRenderHairStrandsDebugInfo(GraphBuilder, Scene, View, Parameters);
 	}
 }
diff --git a/Engine/Source/Runtime/Renderer/Private/HairStrands/HairStrandsDebug.h b/Engine/Source/Runtime/Renderer/Private/HairStrands/HairStrandsDebug.h
--- a/Engine/Source/Runtime/Renderer/Private/HairStrands/HairStrandsDebug.h
+++ b/Engine/Source/Runtime/Renderer/Private/HairStrands/HairStrandsDebug.h
@@ -13,5 +13,4 @@ void RenderHairStrandsDebugInfo(
 	FRDGBuilder& GraphBuilder,
 	FScene* Scene, 
 	TArrayView<FViewInfo> Views,
-	FRDGTextureRef SceneColorTexture, 
-	FRDGTextureRef SceneDepthTexture);
+	FHairStrandsBookmarkParameters& Parameters);
diff --git a/Engine/Source/Runtime/Renderer/Private/HairStrands/HairStrandsInterface.cpp b/Engine/Source/Runtime/Renderer/Private/HairStrands/HairStrandsInterface.cpp
--- a/Engine/Source/Runtime/Renderer/Private/HairStrands/HairStrandsInterface.cpp
+++ b/Engine/Source/Runtime/Renderer/Private/HairStrands/HairStrandsInterface.cpp
@@ -387,7 +387,7 @@ void CreateHairStrandsBookmarkParameters(FScene* Scene, FViewInfo& View, FHairSt
 	if (bComputeVisibleInstances)
 	{
 		const int32 ActiveInstanceCount = Scene->HairStrandsSceneData.RegisteredProxies.Num();
-		Out.InstancesVisibility.Init(false, ActiveInstanceCount);
+		Out.InstancesVisibilityType.Init(EHairInstanceVisibilityType::NotVisible, ActiveInstanceCount);
 	
 		// 1. Strands - Add all visible strands instances
 		Out.VisibleStrands.Reserve(View.HairStrandsMeshElements.Num());
@@ -400,8 +400,7 @@ void CreateHairStrandsBookmarkParameters(FScene* Scene, FViewInfo& View, FHairSt
 				if (HairData && HairData->Instance)
 				{
 					Out.VisibleStrands.Add(HairData->Instance);
-					Out.InstancesVisibility[HairData->Instance->RegisteredIndex] = true;
-					Out.InstanceCountPerType[uint32(EHairInstanceCount::StrandsPrimaryView)]++;
+					Out.InstancesVisibilityType[HairData->Instance->RegisteredIndex] = EHairInstanceVisibilityType::StrandsPrimaryView;
 				}
 			}
 		}
@@ -417,8 +416,7 @@ void CreateHairStrandsBookmarkParameters(FScene* Scene, FViewInfo& View, FHairSt
 				if (HairData && HairData->Instance)
 				{
 					Out.VisibleCardsOrMeshes_Primary.Add(HairData->Instance);
-					Out.InstancesVisibility[HairData->Instance->RegisteredIndex] = true;
-					Out.InstanceCountPerType[uint32(EHairInstanceCount::CardsOrMeshesPrimaryView)]++;
+					Out.InstancesVisibilityType[HairData->Instance->RegisteredIndex] = EHairInstanceVisibilityType::CardsOrMeshesPrimaryView;
 				}
 			}
 		}
@@ -431,6 +429,7 @@ void CreateHairStrandsBookmarkParameters(FScene* Scene, FViewInfo& View, FHairSt
 	Out.ViewRect				= View.ViewRect;
 	Out.ViewUniqueID			= View.ViewState ? View.ViewState->UniqueID : ~0;
 	Out.SceneColorTexture		= nullptr;
+	Out.SceneDepthTexture		= nullptr;
 	Out.Scene					= Scene;
 }
 
@@ -451,19 +450,19 @@ void UpdateHairStrandsBookmarkParameters(FScene* Scene, TArray<FViewInfo>& Views
 			const bool bCardsOrMeshes = Instance->GetHairGeometry() == EHairGeometryType::Cards || Instance->GetHairGeometry() == EHairGeometryType::Meshes;
 			const bool bCompatible = bStrands || bCardsOrMeshes;
 
-			if (Out.InstancesVisibility.IsValidIndex(Instance->RegisteredIndex) && !Out.InstancesVisibility[Instance->RegisteredIndex] && bCompatible)
+			if (Out.InstancesVisibilityType.IsValidIndex(Instance->RegisteredIndex) && Out.InstancesVisibilityType[Instance->RegisteredIndex] == EHairInstanceVisibilityType::NotVisible && bCompatible)
 			{
 				if (IsHairStrandsVisibleInShadows(Views[0], *Instance))
 				{
 					if (bStrands) 		
 					{ 
 						Out.VisibleStrands.Add(Instance);
-						Out.InstanceCountPerType[uint32(EHairInstanceCount::StrandsShadowView)]++; 
+						Out.InstancesVisibilityType[Instance->RegisteredIndex] = EHairInstanceVisibilityType::StrandsShadowView;
 					}
 					else if (bCardsOrMeshes) 
 					{ 
 						Out.VisibleCardsOrMeshes_Shadow.Add(Instance);
-						Out.InstanceCountPerType[uint32(EHairInstanceCount::CardsOrMeshesShadowView)]++; 
+						Out.InstancesVisibilityType[Instance->RegisteredIndex] = EHairInstanceVisibilityType::CardsOrMeshesShadowView;
 					}
 				}
 			}
diff --git a/Engine/Source/Runtime/Renderer/Public/HairStrandsInterface.h b/Engine/Source/Runtime/Renderer/Public/HairStrandsInterface.h
--- a/Engine/Source/Runtime/Renderer/Public/HairStrandsInterface.h
+++ b/Engine/Source/Runtime/Renderer/Public/HairStrandsInterface.h
@@ -480,12 +480,13 @@ enum class EHairStrandsBookmark : uint8
 	ProcessEndOfFrame
 };
 
-enum class EHairInstanceCount : uint8
+enum class EHairInstanceVisibilityType : uint8
 {
-	StrandsPrimaryView = 0,
-	StrandsShadowView = 1,
-	CardsOrMeshesPrimaryView = 2,
-	CardsOrMeshesShadowView = 3,
+	NotVisible = 0,
+	StrandsPrimaryView = 1,
+	StrandsShadowView = 2,
+	CardsOrMeshesPrimaryView = 3,
+	CardsOrMeshesShadowView = 4,
 	Count
 };
 
@@ -502,15 +503,13 @@ struct FHairStrandsBookmarkParameters
 	FHairStrandsInstances VisibleCardsOrMeshes_Shadow;
 
 	FHairStrandsInstances* Instances = nullptr;
-	TBitArray<> InstancesVisibility;
+	TArray<EHairInstanceVisibilityType> InstancesVisibilityType;
 	const FSceneView* View = nullptr;// // View 0
 	FSceneInterface* Scene = nullptr;
 	TArray<const FSceneView*> AllViews;
 	FRDGTextureRef SceneColorTexture = nullptr;
 	FRDGTextureRef SceneDepthTexture = nullptr; 
 
-	FUintVector4 InstanceCountPerType = FUintVector4(0);
-
 	inline bool HasInstances() const { return Instances != nullptr && Instances->Num() > 0; }
 };
 
