Landscape Distance Field - don't require DFAO for composing landscape into the Global Distance Field, as other systems like Lumen may also be using it independently.

diff --git a/Engine/Source/Runtime/Renderer/Private/GlobalDistanceField.cpp b/Engine/Source/Runtime/Renderer/Private/GlobalDistanceField.cpp
--- a/Engine/Source/Runtime/Renderer/Private/GlobalDistanceField.cpp
+++ b/Engine/Source/Runtime/Renderer/Private/GlobalDistanceField.cpp
@@ -2290,7 +2290,6 @@ void UpdateGlobalDistanceFieldVolume(
 							&& NumUpdateBounds > 0
 							&& NumHeightfieldPrimitives > 0
 							&& GAOGlobalDistanceFieldRepresentHeightfields
-							&& SupportsDistanceFieldAO(Scene->GetFeatureLevel(), Scene->GetShaderPlatform())
 							&& !IsVulkanMobileSM5Platform(Scene->GetShaderPlatform()))
 						{
 							for (int32 HeightfieldPrimitiveIndex = 0; HeightfieldPrimitiveIndex < NumHeightfieldPrimitives; HeightfieldPrimitiveIndex++)
