Substrate - fix for sm5 preview that is now no longer working well due to intrinsic identified as supported but not working.

diff --git a/Engine/Source/Runtime/Renderer/Private/Substrate/Substrate.cpp b/Engine/Source/Runtime/Renderer/Private/Substrate/Substrate.cpp
--- a/Engine/Source/Runtime/Renderer/Private/Substrate/Substrate.cpp
+++ b/Engine/Source/Runtime/Renderer/Private/Substrate/Substrate.cpp
@@ -645,8 +645,8 @@ static void BindSubstratePublicGlobalUniformParameters(FRDGBuilder& GraphBuilder
 
 static ERHIFeatureSupport SubstrateSupportsWaveOps(EShaderPlatform Platform)
 {
-	// D3D11 / SM5 do not support wave-ops by default, that fixes classification and black/wrong tiling.
-	if (Platform == SP_PCD3D_SM5)
+	// D3D11 / SM5 or preview do not support, or work well with, wave-ops by default (or SM5 preview has issues with wave intrinsics too), that fixes classification and black/wrong tiling.
+	if (Platform == SP_PCD3D_SM5 || FDataDrivenShaderPlatformInfo::GetIsPreviewPlatform(Platform))
 	{
 		return ERHIFeatureSupport::Unsupported;
 	}
