Fix ShaderPrint flickering in when using Request() API and rendering shader print data into different views.

diff --git a/Engine/Source/Runtime/Renderer/Private/ShaderPrint.cpp b/Engine/Source/Runtime/Renderer/Private/ShaderPrint.cpp
--- a/Engine/Source/Runtime/Renderer/Private/ShaderPrint.cpp
+++ b/Engine/Source/Runtime/Renderer/Private/ShaderPrint.cpp
@@ -110,12 +110,28 @@ namespace ShaderPrint
 	//////////////////////////////////////////////////////////////////////////////////////////////////
 	// Global states
 
-	static uint32 GWidgetRequestCount = 0;
-	static uint32 GCharacterRequestCount = 0;
-	static uint32 GLineRequestCount = 0;
-	static uint32 GTriangleRequestCount = 0;
+	struct FShaderPrintRequest
+	{
+		uint32 WidgetCount 		= 0;
+		uint32 CharacterCount 	= 0;
+		uint32 LineCount 		= 0;
+		uint32 TriangleCount 	= 0;
+	};
+
+	static FShaderPrintRequest MaxRequests(const FShaderPrintRequest& A, const FShaderPrintRequest& B)
+	{
+		FShaderPrintRequest Out;
+		Out.CharacterCount 	= FMath::Max(A.CharacterCount, 	B.CharacterCount);
+		Out.LineCount 		= FMath::Max(A.LineCount, 		B.LineCount);
+		Out.TriangleCount 	= FMath::Max(A.TriangleCount, 	B.TriangleCount);
+		Out.WidgetCount 	= FMath::Max(A.WidgetCount, 	B.WidgetCount);
+		return Out;
+	}
+
+	static FShaderPrintRequest GActiveShaderPrintRequest = FShaderPrintRequest();
+	static FShaderPrintRequest GCachedShaderPrintMaxRequest = FShaderPrintRequest();
 	static FViewInfo* GDefaultView = nullptr;
-	
+
 	struct FQueuedRenderItem
 	{
 		FFrozenShaderPrintData Payload;
@@ -129,22 +145,22 @@ namespace ShaderPrint
 
 	static uint32 GetMaxValueCount()
 	{
-		return FMath::Max(CVarMaxCharacterCount.GetValueOnAnyThread() + int32(GCharacterRequestCount), 0);
+		return FMath::Max(CVarMaxCharacterCount.GetValueOnAnyThread() + int32(GCachedShaderPrintMaxRequest.CharacterCount), 0);
 	}
 
 	static uint32 GetMaxWidgetCount()
 	{
-		return FMath::Max(CVarMaxWidgetCount.GetValueOnAnyThread() + int32(GWidgetRequestCount), 0);
+		return FMath::Max(CVarMaxWidgetCount.GetValueOnAnyThread() + int32(GCachedShaderPrintMaxRequest.WidgetCount), 0);
 	}
 
 	static uint32 GetMaxLineCount()
 	{
-		return FMath::Max(CVarMaxLineCount.GetValueOnAnyThread() + int32(GLineRequestCount), 0);
+		return FMath::Max(CVarMaxLineCount.GetValueOnAnyThread() + int32(GCachedShaderPrintMaxRequest.LineCount), 0);
 	}
 
 	static uint32 GetMaxTriangleCount()
 	{
-		return FMath::Max(CVarMaxTriangleCount.GetValueOnAnyThread() + int32(GTriangleRequestCount), 0);
+		return FMath::Max(CVarMaxTriangleCount.GetValueOnAnyThread() + int32(GCachedShaderPrintMaxRequest.TriangleCount), 0);
 	}
 
 	// Returns the number of uints used for counters, a line element, and a triangle elements
@@ -311,17 +327,17 @@ namespace ShaderPrint
 
 	void RequestSpaceForCharacters(uint32 InCount)
 	{
-		GCharacterRequestCount += InCount;
+		GActiveShaderPrintRequest.CharacterCount += InCount;
 	}
 
 	void RequestSpaceForLines(uint32 InCount)
 	{
-		GLineRequestCount += InCount;
+		GActiveShaderPrintRequest.LineCount += InCount;
 	}
 
 	void RequestSpaceForTriangles(uint32 InCount)
 	{
-		GTriangleRequestCount += InCount;
+		GActiveShaderPrintRequest.TriangleCount += InCount;
 	}
 
 	void SubmitShaderPrintData(FFrozenShaderPrintData& InData, FSceneInterface const* InScene)
@@ -908,11 +924,10 @@ namespace ShaderPrint
 			BeginView(GraphBuilder, View);
 		}
 
-		// Reset counters which are read on the next BeginViews().
-		GCharacterRequestCount = 0;
-		GWidgetRequestCount = 0;
-		GLineRequestCount = 0;
-		GTriangleRequestCount = 0;
+		// * Track the max requests accross frame to avoid flickering when switching between views/renderer having different requests
+		// * Then Reset counters which are read on the next BeginViews().
+		GCachedShaderPrintMaxRequest = MaxRequests(GCachedShaderPrintMaxRequest, GActiveShaderPrintRequest);
+		GActiveShaderPrintRequest = FShaderPrintRequest();
 	}
 
 	static void InternalDrawView_Characters(
