Support NonCastingIntensity in ScreenSpaceShadowsCS on Mobile Forward renderer.
- Set stencil bit depending on CastContactShadow primitive flag during mobile depth prepass when using forward shading.

diff --git a/Engine/Shaders/Private/ScreenSpaceShadows.usf b/Engine/Shaders/Private/ScreenSpaceShadows.usf
--- a/Engine/Shaders/Private/ScreenSpaceShadows.usf
+++ b/Engine/Shaders/Private/ScreenSpaceShadows.usf
@@ -72,6 +72,11 @@ int4 ScissorRectMinAndSize;
 
 RWTexture2D<float2> RWShadowFactors;
 
+Texture2D<uint2> StencilTexture;
+
+// Must match C++ (SceneRenderTargets.h)
+#define STENCIL_MOBILE_CAST_CONTACT_SHADOW_MASK 2
+
 float GetContactShadowOcclusion(float2 HitUV)
 {
 	const uint2 HitPixelCoord = uint2(HitUV * View.BufferSizeAndInvSize.xy);
@@ -80,8 +85,9 @@ float GetContactShadowOcclusion(float2 HitUV)
 #	if MOBILE_DEFERRED_SHADING
 		FGBufferData SampleGBuffer = MobileFetchAndDecodeGBuffer(HitUV, HitPixelCoord);
 		const bool bHitCastContactShadow = CastContactShadow(SampleGBuffer);
-#	else 
-		const bool bHitCastContactShadow = true; // TODO: support mobile forward path
+#	else
+		uint Stencil = StencilTexture.Load(int3(HitPixelCoord, 0)) STENCIL_COMPONENT_SWIZZLE;
+		const bool bHitCastContactShadow = Stencil & STENCIL_MOBILE_CAST_CONTACT_SHADOW_MASK;
 #	endif
 #else
 #	if SUBSTRATE_ENABLED
diff --git a/Engine/Source/Runtime/Renderer/Private/DepthRendering.cpp b/Engine/Source/Runtime/Renderer/Private/DepthRendering.cpp
--- a/Engine/Source/Runtime/Renderer/Private/DepthRendering.cpp
+++ b/Engine/Source/Runtime/Renderer/Private/DepthRendering.cpp
@@ -765,6 +765,11 @@ void SetMobileDepthPassRenderState(const FPrimitiveSceneProxy* RESTRICT Primitiv
 		StencilValue |= GET_STENCIL_MOBILE_SM_MASK(ShadingModel);
 		StencilValue |= STENCIL_LIGHTING_CHANNELS_MASK(PrimitiveSceneProxy ? PrimitiveSceneProxy->GetLightingChannelStencilValue() : 0x00);
 	}
+	else
+	{
+		uint8 CastContactShadows = (PrimitiveSceneProxy && PrimitiveSceneProxy->CastsContactShadow() ? 0x01 : 0x00);
+		StencilValue |= GET_STENCIL_BIT_MASK(MOBILE_CAST_CONTACT_SHADOW, CastContactShadows);
+	}
 
 	DrawRenderState.SetStencilRef(StencilValue);
 }
diff --git a/Engine/Source/Runtime/Renderer/Private/PostProcess/SceneRenderTargets.h b/Engine/Source/Runtime/Renderer/Private/PostProcess/SceneRenderTargets.h
--- a/Engine/Source/Runtime/Renderer/Private/PostProcess/SceneRenderTargets.h
+++ b/Engine/Source/Runtime/Renderer/Private/PostProcess/SceneRenderTargets.h
@@ -52,7 +52,10 @@
 // Mobile specific
 // Store shading model into stencil [1-2] bits (only used in mobile deferred)
 #define GET_STENCIL_MOBILE_SM_MASK(Value) uint8(((Value) & 0x3) << 1)
+
+// Cast contact shadow mask - bit 1 (only used in mobile forward) / Must match shader (ScreenSpaceShadows.usf)
+#define STENCIL_MOBILE_CAST_CONTACT_SHADOW_BIT_ID	1
 // Forward local light buffer mask for light function - bit 2
-#define STENCIL_MOBILE_LIGHTFUNCTION_MASK uint8(1 << 2)
+#define STENCIL_MOBILE_LIGHTFUNCTION_MASK			uint8(1 << 2)
 // Sky material mask - bit 3
-#define STENCIL_MOBILE_SKY_MASK uint8(1 << 3)
+#define STENCIL_MOBILE_SKY_MASK						uint8(1 << 3)
diff --git a/Engine/Source/Runtime/Renderer/Private/Shadows/ScreenSpaceShadows.cpp b/Engine/Source/Runtime/Renderer/Private/Shadows/ScreenSpaceShadows.cpp
--- a/Engine/Source/Runtime/Renderer/Private/Shadows/ScreenSpaceShadows.cpp
+++ b/Engine/Source/Runtime/Renderer/Private/Shadows/ScreenSpaceShadows.cpp
@@ -57,6 +57,7 @@ class FScreenSpaceShadowsCS : public FGlobalShader
 		SHADER_PARAMETER_RDG_TEXTURE_UAV(RWTexture2D<float2>, RWShadowFactors)
 		SHADER_PARAMETER_STRUCT_REF(FViewUniformShaderParameters, View)
 		SHADER_PARAMETER_STRUCT_INCLUDE(FSceneTextureShaderParameters, SceneTextures)
+		SHADER_PARAMETER_RDG_TEXTURE_SRV(Texture2D<uint2>, StencilTexture)
 		SHADER_PARAMETER(FVector3f, LightDirection)
 		SHADER_PARAMETER(float, ContactShadowLength)
 		SHADER_PARAMETER(uint32, bContactShadowLengthInWS)
@@ -90,6 +91,7 @@ class FScreenSpaceShadowsBendCS : public FGlobalShader
 	BEGIN_SHADER_PARAMETER_STRUCT(FParameters, )
 		SHADER_PARAMETER_RDG_TEXTURE_UAV(RWTexture2D<float>, OutputTexture)
 		SHADER_PARAMETER_RDG_TEXTURE(Texture2D<float>, DepthTexture)
+		SHADER_PARAMETER_RDG_TEXTURE_SRV(Texture2D<uint2>, StencilTexture)
 		SHADER_PARAMETER_SAMPLER(SamplerState, PointBorderSampler)
 		SHADER_PARAMETER_STRUCT_REF(FViewUniformShaderParameters, View)
 		SHADER_PARAMETER(FVector3f, LightDirection)
@@ -244,6 +246,7 @@ void RenderScreenSpaceShadows(
 		PassParameters->RWShadowFactors = GraphBuilder.CreateUAV(ShadowsTexture);
 		PassParameters->View = View.ViewUniformBuffer;
 		PassParameters->SceneTextures = SceneTextures.GetSceneTextureShaderParameters(View.GetFeatureLevel());
+		PassParameters->StencilTexture = SceneTextures.Stencil;
 
 		PassParameters->ScissorRectMinAndSize = FIntRect(ScissorRect.Min, ScissorRect.Size());
 		PassParameters->DownsampleFactor = GetScreenSpaceShadowDownsampleFactor();
@@ -334,6 +337,8 @@ void RenderScreenSpaceShadowsBend(
 			PassParameters->DepthTexture = SceneTextures.Depth.Resolve;
 			PassParameters->InvDepthTextureSize = FVector2f(1.0f / DepthDesc.Extent.X, 1.0f / DepthDesc.Extent.Y);
 
+			PassParameters->StencilTexture = SceneTextures.Stencil;
+
 			// A point sampler, with Wrap Mode set to Clamp-To-Border-Color (D3D12_TEXTURE_ADDRESS_MODE_BORDER), and Border Color set to "FarDepthValue" (typically zero), or some other far-depth value out of DepthBounds.
 			// If you have issues where invalid shadows are appearing from off-screen, it is likely that this sampler is not correctly setup
 			PassParameters->PointBorderSampler = TStaticSamplerState<SF_Point, AM_Border, AM_Border, AM_Border, 0, 1, 0>::GetRHI();
