The virtual texture feedback UAV must be in the UAVGraphics | UAVCompute state because Nanite compute materials can write to it.

#rnx

diff --git a/Engine/Source/Runtime/Renderer/Private/VT/VirtualTextureFeedbackBuffer.cpp b/Engine/Source/Runtime/Renderer/Private/VT/VirtualTextureFeedbackBuffer.cpp
--- a/Engine/Source/Runtime/Renderer/Private/VT/VirtualTextureFeedbackBuffer.cpp
+++ b/Engine/Source/Runtime/Renderer/Private/VT/VirtualTextureFeedbackBuffer.cpp
@@ -91,7 +91,7 @@ void FVirtualTextureFeedbackBuffer::Begin(FRDGBuilder& GraphBuilder, const FVirt
 		// Clear virtual texture feedback to default value
 		RHICmdList.Transition(FRHITransitionInfo(UAV, ERHIAccess::Unknown, ERHIAccess::UAVCompute));
 		RHICmdList.ClearUAVUint(UAV, FUintVector4(~0u, ~0u, ~0u, ~0u));
-		RHICmdList.Transition(FRHITransitionInfo(UAV, ERHIAccess::UAVCompute, ERHIAccess::UAVGraphics));
+		RHICmdList.Transition(FRHITransitionInfo(UAV, ERHIAccess::UAVCompute, ERHIAccess::UAVMask));
 		RHICmdList.BeginUAVOverlap(UAV);
 	});
 }
@@ -101,7 +101,7 @@ void FVirtualTextureFeedbackBuffer::End(FRDGBuilder& GraphBuilder)
 	AddPass(GraphBuilder, RDG_EVENT_NAME("VirtualTextureFeedbackCopy"), [this](FRHICommandListImmediate& RHICmdList)
 	{
 		RHICmdList.EndUAVOverlap(UAV);
-		RHICmdList.Transition(FRHITransitionInfo(UAV, ERHIAccess::UAVGraphics, ERHIAccess::CopySrc));
+		RHICmdList.Transition(FRHITransitionInfo(UAV, ERHIAccess::UAVMask, ERHIAccess::CopySrc));
 		SubmitVirtualTextureFeedbackBuffer(RHICmdList, PooledBuffer->GetRHI(), Desc);
 	});
 }
