* Rename r.Substrate.TileOverflowFromMaterial -> r.Substrate.UseClosureCountFromMaterial
* Remove r.Substrate.TileOverflow, as layer count will now be driven by  r.Substrate.ClosuresPerPixel when r.Substrate.UseClosureCountFromMaterial=0

#perflight skip

diff --git a/Engine/Shaders/Private/Substrate/SubstrateVisualize.usf b/Engine/Shaders/Private/Substrate/SubstrateVisualize.usf
--- a/Engine/Shaders/Private/Substrate/SubstrateVisualize.usf
+++ b/Engine/Shaders/Private/Substrate/SubstrateVisualize.usf
@@ -1537,14 +1537,13 @@ uint ClassificationCMask;
 uint ClassificationAsync;
 uint Classification8bits;
 uint bRoughRefraction;
-uint bTileOverflowUseMaterialData;
+uint bUseClosureCountFromMaterialData;
 uint ViewsMaxBytesPerPixel;
 uint ProjectMaxBytesPerPixel;
 uint MaterialBufferAllocationInBytes;
 uint MaterialBufferAllocationMode;
 
 uint LayerCount;
-float TileOverflowRatio;
 Buffer<uint> ClassificationTileDrawIndirectBuffer;
 
 uint GetTileCount(uint InType)
@@ -1608,19 +1607,12 @@ void MainCS(uint3 DispatchThreadId : SV_DispatchThreadID)
 		PrintBool(Context, bEnergyPreservation);
 		Newline(Context);
 
-		Print(Context, TEXT(" Tile overflow mat. : "), LegendColor);
-		PrintBool(Context, bTileOverflowUseMaterialData);
+		Print(Context, TEXT(" Use #closures mat. : "), LegendColor);
+		PrintBool(Context, bUseClosureCountFromMaterialData);
 		Newline(Context); 
 
-		Print(Context, TEXT(" Tile overflow      : "), LegendColor);
-		Print(Context, TileOverflowRatio, ValueColor, 3, 1);
-		Print(Context, TEXT(" ["), LegendColor);
-		Print(Context, TileOverflowRatio +1, ValueColor, 2, 0);
-		Print(Context, TEXT(" BSDFs]"), LegendColor);
-		Newline(Context);
-
 		Print(Context, TEXT(" Tile LayerCount.   : "), LegendColor);
-		PrintBool(Context, LayerCount);
+		Print(Context, LayerCount, FontOrange);
 		Newline(Context); 
 
 		Print(Context, TEXT(" DBuffer pass       : "), LegendColor);
diff --git a/Engine/Source/Runtime/Renderer/Private/Substrate/Substrate.cpp b/Engine/Source/Runtime/Renderer/Private/Substrate/Substrate.cpp
--- a/Engine/Source/Runtime/Renderer/Private/Substrate/Substrate.cpp
+++ b/Engine/Source/Runtime/Renderer/Private/Substrate/Substrate.cpp
@@ -27,16 +27,10 @@ static TAutoConsoleVariable<int32> CVarUseCmaskClear(
 	TEXT("TEST."),
 	ECVF_RenderThreadSafe);
 
-static TAutoConsoleVariable<float> CVarSubstrateTileOverflow(
-	TEXT("r.Substrate.TileOverflow"),
-	1.f,
-	TEXT("Scale the number of Substrate tile for overflowing tiles containing multi-BSDFs pixels. (0: 0%, 1: 100%. Default 1.0f)."),
-	ECVF_RenderThreadSafe);
-
-static TAutoConsoleVariable<int32> CVarSubstrateTileOverflowFromMaterial(
-	TEXT("r.Substrate.TileOverflowFromMaterial"),
+static TAutoConsoleVariable<int32> CVarSubstrateUseClosureCountFromMaterial(
+	TEXT("r.Substrate.UseClosureCountFromMaterial"),
 	1,
-	TEXT("When enable, scale the number of Substrate tile for overflowing tiles containing multi-BSDFs pixels based on material data. Otherwise use global overflow scale defined by r.Substrate.TileOverflow."),
+	TEXT("When enable, scale the number of Lumen's layers for multi-BSDFs pixels based on material data. Otherwise use r.Substrate.ClosuresPerPixel."),
 	ECVF_RenderThreadSafe);
 
 static TAutoConsoleVariable<int32> CVarSubstrateDebugPeelLayersAboveDepth(
@@ -114,21 +108,9 @@ uint32 GetMaterialBufferAllocationMode()
 	return FMath::Clamp(CVarSubstrateBytePerPixelMode.GetValueOnAnyThread(), 0, 2);
 }
 
-bool DoesSubstrateTileOverflowUseMaterialData() 
+bool UsesSubstrateClosureCountFromMaterialData() 
 {
-	return CVarSubstrateTileOverflowFromMaterial.GetValueOnRenderThread() > 0;
-}
-
-uint32 GetSubstrateTileOverflowRatio(const FViewInfo& View)
-{
-	if (DoesSubstrateTileOverflowUseMaterialData())
-	{
-		return FMath::Clamp(View.SubstrateViewData.MaxBSDFCount, 1u, 4u);
-	}
-	else
-	{
-		return FMath::Clamp(uint32(CVarSubstrateTileOverflow.GetValueOnRenderThread()), 1u, 4u);
-	}
+	return CVarSubstrateUseClosureCountFromMaterial.GetValueOnRenderThread() > 0;
 }
 
 uint32 GetSubstrateTextureLayerCount(const FViewInfo& View)
@@ -136,7 +118,14 @@ uint32 GetSubstrateTextureLayerCount(const FViewInfo& View)
 	uint32 Out = 1;
 	if (Substrate::IsSubstrateEnabled())
 	{
-		Out = GetSubstrateTileOverflowRatio(View);
+		if (UsesSubstrateClosureCountFromMaterialData())
+		{
+			Out = FMath::Clamp(View.SubstrateViewData.MaxBSDFCount, 1u, SUBSTRATE_MAX_BSDF_COUNT);
+		}
+		else
+		{
+			Out = FMath::Clamp(uint32(GetClosurePerPixel(View.GetShaderPlatform())), 1u, SUBSTRATE_MAX_BSDF_COUNT);
+		}
 	}
 	return Out;
 }
@@ -254,7 +243,6 @@ static void InitialiseSubstrateViewData(FRDGBuilder& GraphBuilder, FViewInfo& Vi
 			const FIntPoint TileCount = GetSubstrateTextureTileResolution(View, DynResIndependentViewSize);
 			const uint32 LayerCount = GetSubstrateTextureLayerCount(View);
 			const uint32 MaxTileCount = TileCount.X * TileCount.Y * LayerCount;
-			const FIntPoint BaseOverflowTileOffset = FIntPoint(0, FMath::DivideAndRoundUp(DynResIndependentViewSize.Y, SUBSTRATE_TILE_SIZE));
 
 			Out.TileCount	= TileCount;
 			Out.LayerCount  = LayerCount;
diff --git a/Engine/Source/Runtime/Renderer/Private/Substrate/SubstrateVisualize.cpp b/Engine/Source/Runtime/Renderer/Private/Substrate/SubstrateVisualize.cpp
--- a/Engine/Source/Runtime/Renderer/Private/Substrate/SubstrateVisualize.cpp
+++ b/Engine/Source/Runtime/Renderer/Private/Substrate/SubstrateVisualize.cpp
@@ -119,13 +119,12 @@ class FSubstrateSystemInfoCS : public FGlobalShader
 		SHADER_PARAMETER(uint32, ClassificationAsync)
 		SHADER_PARAMETER(uint32, Classification8bits)
 		SHADER_PARAMETER(uint32, bRoughRefraction)
-		SHADER_PARAMETER(uint32, bTileOverflowUseMaterialData)
+		SHADER_PARAMETER(uint32, bUseClosureCountFromMaterialData)
 		SHADER_PARAMETER(uint32, ProjectMaxBytesPerPixel)
 		SHADER_PARAMETER(uint32, ViewsMaxBytesPerPixel)
 		SHADER_PARAMETER(uint32, MaterialBufferAllocationInBytes)
 		SHADER_PARAMETER(uint32, MaterialBufferAllocationMode)
 		SHADER_PARAMETER(uint32, LayerCount)
-		SHADER_PARAMETER(float, TileOverflowRatio)
 		SHADER_PARAMETER_RDG_BUFFER_SRV(Buffer<uint>, ClassificationTileDrawIndirectBuffer)
 		SHADER_PARAMETER_STRUCT_REF(FViewUniformShaderParameters, ViewUniformBuffer)
 		SHADER_PARAMETER_RDG_UNIFORM_BUFFER(FSubstrateGlobalUniformParameters, Substrate)
@@ -278,10 +277,9 @@ static void AddVisualizeMaterialCountPasses(FRDGBuilder & GraphBuilder, const FV
 }
 
 uint32 GetSubstrateTextureLayerCount(const FViewInfo& View);
-uint32 GetSubstrateTileOverflowRatio(const FViewInfo& View);
 bool IsClassificationAsync();
 bool SupportsCMask(const FStaticShaderPlatform InPlatform);
-bool DoesSubstrateTileOverflowUseMaterialData();
+bool UsesSubstrateClosureCountFromMaterialData();
 uint32 GetMaterialBufferAllocationMode();
 
 static void AddVisualizeSystemInfoPasses(FRDGBuilder& GraphBuilder, const FViewInfo& View, FScreenPassTexture& ScreenPassSceneColor, EShaderPlatform Platform)
@@ -301,9 +299,8 @@ static void AddVisualizeSystemInfoPasses(FRDGBuilder& GraphBuilder, const FViewI
 	PassParameters->ClassificationCMask = SupportsCMask(View.GetShaderPlatform()) ? 1 : 0;
 	PassParameters->ClassificationAsync = IsClassificationAsync() ? 1 : 0;
 	PassParameters->Classification8bits = Is8bitTileCoordEnabled() ? 1 : 0;
-	PassParameters->TileOverflowRatio = GetSubstrateTileOverflowRatio(View);
 	PassParameters->LayerCount = GetSubstrateTextureLayerCount(View);
-	PassParameters->bTileOverflowUseMaterialData = DoesSubstrateTileOverflowUseMaterialData() ? 1 : 0;
+	PassParameters->bUseClosureCountFromMaterialData = UsesSubstrateClosureCountFromMaterialData() ? 1 : 0;
 	PassParameters->bRoughRefraction = IsOpaqueRoughRefractionEnabled() ? 1 : 0;
 	PassParameters->ClassificationTileDrawIndirectBuffer = GraphBuilder.CreateSRV(View.SubstrateViewData.ClassificationTileDrawIndirectBuffer, PF_R32_UINT);
 	PassParameters->ViewUniformBuffer = View.ViewUniformBuffer;
