Add RVT displacement format intended for use with nanite displacement.

diff --git a/Engine/Shaders/Private/BCCompressionCommon.ush b/Engine/Shaders/Private/BCCompressionCommon.ush
--- a/Engine/Shaders/Private/BCCompressionCommon.ush
+++ b/Engine/Shaders/Private/BCCompressionCommon.ush
@@ -281,9 +281,8 @@ uint2 CompressBC4Block( in float Block[16] )
 
 	uint2 BC4Block = 0;
 	BC4Block.x = (MinAlphaUint << 8) | MaxAlphaUint;
-	BC4Block.x |= Indices.x >> 16;
-	BC4Block.y = Indices.y & 0xffff;
-	BC4Block.y |= Indices.y >> 16;
+	BC4Block.x |= Indices.x;
+	BC4Block.y = Indices.y;
 	return BC4Block;
 }
 
diff --git a/Engine/Shaders/Private/VirtualTextureCompress.usf b/Engine/Shaders/Private/VirtualTextureCompress.usf
--- a/Engine/Shaders/Private/VirtualTextureCompress.usf
+++ b/Engine/Shaders/Private/VirtualTextureCompress.usf
@@ -308,11 +308,34 @@ void CopyBaseColorNormalSpecularMaskYCoCgPS(
 	OutColor2 = float4(RoughnessSpecularNormalZ, NormalXYMask.z);
 }
 
-/** Copy world height (without compression)
-  * This path is used when want to copy the world height data into a thumbnail. */
+/** Compress displacement to a single BC4 target. */
+[numthreads(8, 8, 1)]
+void CompressDisplacementCS(
+	uint3 ThreadId : SV_DispatchThreadID)
+{
+	if (any(ThreadId.xy * 4 >= SourceRect.zw))
+	{
+		return;
+	}
+
+	uint2 SamplePos = SourceRect.xy + ThreadId.xy * 4;
+	float2 TexelUVSize = 1.f / float2(SourceRect.zw);
+	float2 SampleUV = (float2(SamplePos) + 0.5f) * TexelUVSize;
+
+	float Block[16];
+	ReadBlockX(RenderTexture0, TextureSampler0, SampleUV, TexelUVSize, Block);
+#if ETC_PROFILE
+	OutCompressTexture0_64bit[ThreadId.xy + DestPos0] = PackRG32(CompressBlock_ETC2_Alpha(Block));
+#else
+	OutCompressTexture0_64bit[ThreadId.xy + DestPos0] = PackRG32(CompressBC4Block(Block));
+#endif
+}
+
+/** Copy world height or displacement (without compression)
+  * This path is used when want to copy the world height or displacement data into a thumbnail. */
 void CopyWorldHeightPS(
-	in float4 InPosition : SV_POSITION, 
-	in noperspective float2 InTexCoord : TEXCOORD0, 
+	in float4 InPosition : SV_POSITION,
+	in noperspective float2 InTexCoord : TEXCOORD0,
 	out float4 OutColor0 : SV_Target0)
 {
 	float WorldHeight = RenderTexture0.SampleLevel(TextureSampler0, InTexCoord, 0).r;
diff --git a/Engine/Shaders/Private/VirtualTextureMaterial.usf b/Engine/Shaders/Private/VirtualTextureMaterial.usf
--- a/Engine/Shaders/Private/VirtualTextureMaterial.usf
+++ b/Engine/Shaders/Private/VirtualTextureMaterial.usf
@@ -97,6 +97,7 @@ void FPixelShaderInOut_MainPS(
 	FLWCScalar WorldHeight = GetVirtualTextureOutput4_LWC(MaterialParameters);
 	float Opacity = GetVirtualTextureOutput5(MaterialParameters);
 	float Mask = GetVirtualTextureOutput6(MaterialParameters);
+	float Displacement = GetVirtualTextureOutput7(MaterialParameters);
 #else
 	// Output is from standard material output attribute node
 	half3 BaseColor = GetMaterialBaseColor(PixelMaterialInputs);
@@ -106,6 +107,7 @@ void FPixelShaderInOut_MainPS(
 	FLWCScalar WorldHeight = LWCGetZ(MaterialParameters.AbsoluteWorldPosition);
 	float Opacity = GetMaterialOpacity(PixelMaterialInputs);
 	float Mask = 0.f;
+	float Displacement = 0.f;
 #endif
 
 	// Apply debug shading
@@ -127,6 +129,8 @@ void FPixelShaderInOut_MainPS(
 #elif defined(OUT_WORLDHEIGHT)
 	float PackedHeight = PackWorldHeight(WorldHeight, ResolvedView.RuntimeVirtualTexturePackHeight);
 	Out.MRT[0] = float4(PackedHeight, 0, 0, 1);
+#elif defined(OUT_DISPLACEMENT)
+	Out.MRT[0] = float4(Displacement, 0, 0, 1);
 #endif
 }
 
@@ -140,6 +144,8 @@ void FPixelShaderInOut_MainPS(
 //#define PIXELSHADEROUTPUT_MRT2 1
 //#elif defined(OUT_WORLDHEIGHT)
 //#define PIXELSHADEROUTPUT_MRT0 1
+//#elif defined(OUT_DISPLACEMENT)
+//#define PIXELSHADEROUTPUT_MRT0 1
 //#endif
 
 // all PIXELSHADEROUTPUT_ and "void FPixelShaderInOut_MainPS()" need to be setup before this include
diff --git a/Engine/Source/Runtime/Engine/Classes/Materials/MaterialExpressionRuntimeVirtualTextureOutput.h b/Engine/Source/Runtime/Engine/Classes/Materials/MaterialExpressionRuntimeVirtualTextureOutput.h
--- a/Engine/Source/Runtime/Engine/Classes/Materials/MaterialExpressionRuntimeVirtualTextureOutput.h
+++ b/Engine/Source/Runtime/Engine/Classes/Materials/MaterialExpressionRuntimeVirtualTextureOutput.h
@@ -41,6 +41,10 @@ class UMaterialExpressionRuntimeVirtualTextureOutput : public UMaterialExpressio
 	UPROPERTY()
 	FExpressionInput Mask;
 
+	/** Input for World Height to output to virtual texture. */
+	UPROPERTY()
+	FExpressionInput Displacement;
+
 public:
 #if WITH_EDITOR
 	//~ Begin UMaterialExpression Interface
diff --git a/Engine/Source/Runtime/Engine/Private/Components/RuntimeVirtualTextureComponent.cpp b/Engine/Source/Runtime/Engine/Private/Components/RuntimeVirtualTextureComponent.cpp
--- a/Engine/Source/Runtime/Engine/Private/Components/RuntimeVirtualTextureComponent.cpp
+++ b/Engine/Source/Runtime/Engine/Private/Components/RuntimeVirtualTextureComponent.cpp
@@ -345,7 +345,9 @@ bool URuntimeVirtualTextureComponent::CanEditChange(const FProperty* InProperty)
 	}
 	else if (InProperty->GetFName() == TEXT("bBuildDebugStreamingMips"))
 	{
-		bCanEdit = GetVirtualTexture() != nullptr && GetVirtualTexture()->GetMaterialType() != ERuntimeVirtualTextureMaterialType::WorldHeight;
+		bCanEdit = GetVirtualTexture() != nullptr && 
+			GetVirtualTexture()->GetMaterialType() != ERuntimeVirtualTextureMaterialType::WorldHeight && 
+			GetVirtualTexture()->GetMaterialType() != ERuntimeVirtualTextureMaterialType::Displacement;
 	}
 	return bCanEdit;
 }
diff --git a/Engine/Source/Runtime/Engine/Private/Materials/HLSLMaterialTranslator.cpp b/Engine/Source/Runtime/Engine/Private/Materials/HLSLMaterialTranslator.cpp
--- a/Engine/Source/Runtime/Engine/Private/Materials/HLSLMaterialTranslator.cpp
+++ b/Engine/Source/Runtime/Engine/Private/Materials/HLSLMaterialTranslator.cpp
@@ -7861,6 +7861,11 @@ int32 FHLSLMaterialTranslator::VirtualTextureUnpack(int32 CodeIndex0, int32 Code
 		FString	SampleCode(TEXT("VirtualTextureUnpackHeight(%s, %s)"));
 		return CodeIndex0 == INDEX_NONE ? INDEX_NONE : AddCodeChunk(MCT_Float, *SampleCode, *GetParameterCode(CodeIndex0), *GetParameterCode(P0));
 	}
+	else if (UnpackType == EVirtualTextureUnpackType::DisplacementR16)
+	{
+		FString	SampleCode(TEXT("%s.r"));
+		return CodeIndex0 == INDEX_NONE ? INDEX_NONE : AddCodeChunk(MCT_Float, *SampleCode, *GetParameterCode(CodeIndex0));
+	}
 	else if (UnpackType == EVirtualTextureUnpackType::NormalBGR565)
 	{
 		FString	SampleCode(TEXT("VirtualTextureUnpackNormalBGR565(%s)"));
diff --git a/Engine/Source/Runtime/Engine/Private/Materials/MaterialExpressionHLSL.cpp b/Engine/Source/Runtime/Engine/Private/Materials/MaterialExpressionHLSL.cpp
--- a/Engine/Source/Runtime/Engine/Private/Materials/MaterialExpressionHLSL.cpp
+++ b/Engine/Source/Runtime/Engine/Private/Materials/MaterialExpressionHLSL.cpp
@@ -1662,6 +1662,7 @@ bool UMaterialExpressionRuntimeVirtualTextureSample::GenerateHLSLExpression(FMat
 	bool bIsNormalValid = false;
 	bool bIsWorldHeightValid = false;
 	bool bIsMaskValid = false;
+	bool bIsDisplacementValid = false;
 
 	switch (MaterialType)
 	{
@@ -1671,6 +1672,7 @@ bool UMaterialExpressionRuntimeVirtualTextureSample::GenerateHLSLExpression(FMat
 	case ERuntimeVirtualTextureMaterialType::BaseColor_Normal_Specular_YCoCg: bIsRoughnessValid = bIsBaseColorValid = bIsNormalValid = bIsSpecularValid = true; break;
 	case ERuntimeVirtualTextureMaterialType::BaseColor_Normal_Specular_Mask_YCoCg: bIsRoughnessValid = bIsBaseColorValid = bIsNormalValid = bIsSpecularValid = bIsMaskValid = true; break;
 	case ERuntimeVirtualTextureMaterialType::WorldHeight: bIsWorldHeightValid = true; break;
+	case ERuntimeVirtualTextureMaterialType::Displacement: bIsDisplacementValid = true; break;
 	}
 
 	switch (OutputIndex)
@@ -1764,6 +1766,17 @@ bool UMaterialExpressionRuntimeVirtualTextureSample::GenerateHLSLExpression(FMat
 			return true;
 		}
 		break;
+	case 6:
+		if (bIsVirtualTextureValid && bIsDisplacementValid)
+		{
+			UnpackType = EVirtualTextureUnpackType::DisplacementR16;
+		}
+		else
+		{
+			OutExpression = Generator.GetTree().NewConstant(0.f);
+			return true;
+		}
+		break;
 	default:
 		return false;
 	}
diff --git a/Engine/Source/Runtime/Engine/Private/Materials/MaterialExpressions.cpp b/Engine/Source/Runtime/Engine/Private/Materials/MaterialExpressions.cpp
--- a/Engine/Source/Runtime/Engine/Private/Materials/MaterialExpressions.cpp
+++ b/Engine/Source/Runtime/Engine/Private/Materials/MaterialExpressions.cpp
@@ -3070,6 +3070,12 @@ int32 UMaterialExpressionRuntimeVirtualTextureOutput::Compile(class FMaterialCom
 		OutputAttributeMask |= Mask.IsConnected() ? (1 << (uint8)ERuntimeVirtualTextureAttributeType::Mask) : 0;
 		ValidateRuntimeVirtualTextureOutput(Compiler, Mask);
 	}
+	else if (OutputIndex == 7)
+	{
+		CodeInput = Displacement.IsConnected() ? Displacement.Compile(Compiler) : Compiler->Constant(0.f);
+		OutputAttributeMask |= Displacement.IsConnected() ? (1 << (uint8)ERuntimeVirtualTextureAttributeType::Displacement) : 0;
+		ValidateRuntimeVirtualTextureOutput(Compiler, Displacement);
+	}
 
 	Compiler->VirtualTextureOutput(OutputAttributeMask);
 	return Compiler->CustomOutput(this, OutputIndex, CodeInput);
@@ -3084,7 +3090,7 @@ void UMaterialExpressionRuntimeVirtualTextureOutput::GetCaption(TArray<FString>&
 
 int32 UMaterialExpressionRuntimeVirtualTextureOutput::GetNumOutputs() const
 {
-	return 7; 
+	return 8; 
 }
 
 FString UMaterialExpressionRuntimeVirtualTextureOutput::GetFunctionName() const
@@ -3148,6 +3154,7 @@ void UMaterialExpressionRuntimeVirtualTextureSample::InitOutputs()
 	Outputs.Add(FExpressionOutput(TEXT("Normal")));
 	Outputs.Add(FExpressionOutput(TEXT("WorldHeight")));
 	Outputs.Add(FExpressionOutput(TEXT("Mask")));
+	Outputs.Add(FExpressionOutput(TEXT("Displacement")));
 #endif // WITH_EDITORONLY_DATA
 }
 
@@ -3262,6 +3269,7 @@ int32 UMaterialExpressionRuntimeVirtualTextureSample::Compile(class FMaterialCom
 	bool bIsNormalValid = false;
 	bool bIsWorldHeightValid = false;
 	bool bIsMaskValid = false;
+	bool bIsDisplacementValid = false;
 
 	switch (MaterialType)
 	{
@@ -3271,6 +3279,7 @@ int32 UMaterialExpressionRuntimeVirtualTextureSample::Compile(class FMaterialCom
 	case ERuntimeVirtualTextureMaterialType::BaseColor_Normal_Specular_YCoCg: bIsRoughnessValid = bIsBaseColorValid = bIsNormalValid = bIsSpecularValid = true; break;
 	case ERuntimeVirtualTextureMaterialType::BaseColor_Normal_Specular_Mask_YCoCg: bIsRoughnessValid = bIsBaseColorValid = bIsNormalValid = bIsSpecularValid = bIsMaskValid = true; break;
 	case ERuntimeVirtualTextureMaterialType::WorldHeight: bIsWorldHeightValid = true; break;
+	case ERuntimeVirtualTextureMaterialType::Displacement: bIsDisplacementValid = true; break;
 	}
 
 	switch (OutputIndex)
@@ -3358,6 +3367,16 @@ int32 UMaterialExpressionRuntimeVirtualTextureSample::Compile(class FMaterialCom
 			return Compiler->Constant(1.f);
 		}
 		break;
+	case 6:
+		if (bIsVirtualTextureValid && bIsDisplacementValid)
+		{
+			UnpackType = EVirtualTextureUnpackType::DisplacementR16;
+		}
+		else
+		{
+			return Compiler->Constant(0.f);
+		}
+		break;
 	default:
 		return INDEX_NONE;
 	}
diff --git a/Engine/Source/Runtime/Engine/Private/Materials/MaterialHLSLTree.cpp b/Engine/Source/Runtime/Engine/Private/Materials/MaterialHLSLTree.cpp
--- a/Engine/Source/Runtime/Engine/Private/Materials/MaterialHLSLTree.cpp
+++ b/Engine/Source/Runtime/Engine/Private/Materials/MaterialHLSLTree.cpp
@@ -2047,6 +2047,15 @@ bool FExpressionVirtualTextureUnpack::PrepareValue(FEmitContext& Context, FEmitS
 		}
 		return OutResult.SetType(Context, RequestedType, EExpressionEvaluation::Shader, Shader::EValueType::Float1);
 	}
+	else if (UnpackType == EVirtualTextureUnpackType::DisplacementR16)
+	{
+		const FPreparedType& Sample0Type = Context.PrepareExpression(SampleLayer0Expression, Scope, Shader::EValueType::Float4);
+		if (Sample0Type.IsVoid())
+		{
+			return false;
+		}
+		return OutResult.SetType(Context, RequestedType, EExpressionEvaluation::Shader, Shader::EValueType::Float1);
+	}
 	else if (UnpackType == EVirtualTextureUnpackType::NormalBGR565)
 	{
 		const FPreparedType& SampleType = Context.PrepareExpression(SampleLayer1Expression, Scope, Shader::EValueType::Float4);
@@ -2104,6 +2113,11 @@ void FExpressionVirtualTextureUnpack::EmitValueShader(FEmitContext& Context, FEm
 		FEmitShaderExpression* EmitHeightScaleBias = WorldHeightUnpackUniformExpression->GetValueShader(Context, Scope, Shader::EValueType::Float2);
 		OutResult.Code = Context.EmitExpression(Scope, Shader::EValueType::Float1, TEXT("VirtualTextureUnpackHeight(%, %)"), EmitSampleLayer0, EmitHeightScaleBias);
 	}
+	else if (UnpackType == EVirtualTextureUnpackType::DisplacementR16)
+	{
+		FEmitShaderExpression* EmitSampleLayer0 = SampleLayer0Expression->GetValueShader(Context, Scope, Shader::EValueType::Float4);
+		OutResult.Code = Context.EmitExpression(Scope, Shader::EValueType::Float1, TEXT("%.r"), EmitSampleLayer0);
+	}
 	else if (UnpackType == EVirtualTextureUnpackType::NormalBGR565)
 	{
 		FEmitShaderExpression* EmitSampleLayer1 = SampleLayer1Expression->GetValueShader(Context, Scope, Shader::EValueType::Float4);
diff --git a/Engine/Source/Runtime/Engine/Private/ShaderCompiler/ShaderGenerationUtil.cpp b/Engine/Source/Runtime/Engine/Private/ShaderCompiler/ShaderGenerationUtil.cpp
--- a/Engine/Source/Runtime/Engine/Private/ShaderCompiler/ShaderGenerationUtil.cpp
+++ b/Engine/Source/Runtime/Engine/Private/ShaderCompiler/ShaderGenerationUtil.cpp
@@ -205,6 +205,7 @@ void ApplyFetchEnvironmentInternal(FShaderMaterialPropertyDefines& SrcDefines, c
 	FETCH_COMPILE_BOOL(OUT_BASECOLOR_NORMAL_ROUGHNESS);
 	FETCH_COMPILE_BOOL(OUT_BASECOLOR_NORMAL_SPECULAR);
 	FETCH_COMPILE_BOOL(OUT_WORLDHEIGHT);
+	FETCH_COMPILE_BOOL(OUT_DISPLACEMENT);
 
 	FETCH_COMPILE_BOOL(IS_VIRTUAL_TEXTURE_MATERIAL);
 	FETCH_COMPILE_BOOL(IS_DECAL);
@@ -1964,6 +1965,10 @@ void FShaderCompileUtilities::ApplyDerivedDefines(FShaderCompilerEnvironment& Ou
 		{
 			bTargetUsage[0] = 1;
 		}
+		else if (MaterialDefines.OUT_DISPLACEMENT)
+		{
+			bTargetUsage[0] = 1;
+		}
 	}
 	else if (MaterialDefines.IS_DECAL)
 	{
diff --git a/Engine/Source/Runtime/Engine/Private/VT/RuntimeVirtualTexture.cpp b/Engine/Source/Runtime/Engine/Private/VT/RuntimeVirtualTexture.cpp
--- a/Engine/Source/Runtime/Engine/Private/VT/RuntimeVirtualTexture.cpp
+++ b/Engine/Source/Runtime/Engine/Private/VT/RuntimeVirtualTexture.cpp
@@ -355,6 +355,7 @@ int32 URuntimeVirtualTexture::GetLayerCount(ERuntimeVirtualTextureMaterialType I
 	{
 	case ERuntimeVirtualTextureMaterialType::BaseColor:
 	case ERuntimeVirtualTextureMaterialType::WorldHeight:
+	case ERuntimeVirtualTextureMaterialType::Displacement:
 		return 1;
 	case ERuntimeVirtualTextureMaterialType::BaseColor_Normal_Specular:
 	case ERuntimeVirtualTextureMaterialType::BaseColor_Normal_Roughness:
@@ -392,6 +393,9 @@ static EPixelFormat PlatformCompressedRVTFormat(EPixelFormat Format)
 		case PF_DXT5:
 			Format = PF_ETC2_RGBA;
 			break;
+		case PF_BC4:
+			Format = PF_ETC2_R11_EAC;
+			break;
 		case PF_BC5:
 			Format = PF_ETC2_RG11_EAC;
 			break;
@@ -427,6 +431,8 @@ EPixelFormat URuntimeVirtualTexture::GetLayerFormat(int32 LayerIndex) const
 			return bCompressTextures ? PlatformCompressedRVTFormat(PF_DXT5) : PF_B8G8R8A8;
 		case ERuntimeVirtualTextureMaterialType::WorldHeight:
 			return PF_G16;
+		case ERuntimeVirtualTextureMaterialType::Displacement:
+			return bCompressTextures ? PlatformCompressedRVTFormat(PF_BC4) : PF_G16;
 		default:
 			break;
 		}
@@ -475,6 +481,7 @@ bool URuntimeVirtualTexture::IsLayerSRGB(int32 LayerIndex) const
 	case ERuntimeVirtualTextureMaterialType::BaseColor_Normal_Specular_YCoCg:
 	case ERuntimeVirtualTextureMaterialType::BaseColor_Normal_Specular_Mask_YCoCg:
 	case ERuntimeVirtualTextureMaterialType::WorldHeight:
+	case ERuntimeVirtualTextureMaterialType::Displacement:
 	case ERuntimeVirtualTextureMaterialType::BaseColor_Normal_Roughness:
 		return false;
 	default:
diff --git a/Engine/Source/Runtime/Engine/Public/MaterialCompiler.h b/Engine/Source/Runtime/Engine/Public/MaterialCompiler.h
--- a/Engine/Source/Runtime/Engine/Public/MaterialCompiler.h
+++ b/Engine/Source/Runtime/Engine/Public/MaterialCompiler.h
@@ -48,7 +48,8 @@ enum class EVirtualTextureUnpackType
 	NormalBC5BC1,
 	HeightR16,
 	NormalBGR565,
-	BaseColorSRGB
+	BaseColorSRGB,
+	DisplacementR16,
 };
 
 /** What type of compiler is this? Used by material expressions that select input based on compile context */
diff --git a/Engine/Source/Runtime/Engine/Public/VT/RuntimeVirtualTextureEnum.h b/Engine/Source/Runtime/Engine/Public/VT/RuntimeVirtualTextureEnum.h
--- a/Engine/Source/Runtime/Engine/Public/VT/RuntimeVirtualTextureEnum.h
+++ b/Engine/Source/Runtime/Engine/Public/VT/RuntimeVirtualTextureEnum.h
@@ -20,6 +20,7 @@ enum class ERuntimeVirtualTextureAttributeType : uint8
 	Specular,
 	Mask,
 	WorldHeight,
+	Displacement,
 
 	Count
 };
@@ -41,6 +42,7 @@ enum class ERuntimeVirtualTextureMaterialType : uint8
 	BaseColor_Normal_Specular_YCoCg UMETA(DisplayName = "YCoCg Base Color, Normal, Roughness, Specular", ToolTip = "Base Color is stored in YCoCg space. This requires more memory but may provide better quality."),
 	BaseColor_Normal_Specular_Mask_YCoCg UMETA(DisplayName = "YCoCg Base Color, Normal, Roughness, Specular, Mask", ToolTip="Base Color is stored in YCoCg space. This requires more memory but may provide better quality."),
 	WorldHeight UMETA(DisplayName = "World Height"),
+	Displacement UMETA(DisplayName = "Displacement"),
 	Count UMETA(Hidden),
 };
 
diff --git a/Engine/Source/Runtime/Landscape/Private/LandscapeRender.cpp b/Engine/Source/Runtime/Landscape/Private/LandscapeRender.cpp
--- a/Engine/Source/Runtime/Landscape/Private/LandscapeRender.cpp
+++ b/Engine/Source/Runtime/Landscape/Private/LandscapeRender.cpp
@@ -3742,11 +3742,13 @@ public:
 			FName(TEXT("TVirtualTextureVSBaseColorNormalSpecular")),
 			FName(TEXT("TVirtualTextureVSBaseColorNormalRoughness")),
 			FName(TEXT("TVirtualTextureVSWorldHeight")),
+			FName(TEXT("TVirtualTextureVSDisplacement")),
 			FName(TEXT("TVirtualTexturePSBaseColor")),
 			FName(TEXT("TVirtualTexturePSBaseColorNormal")),
 			FName(TEXT("TVirtualTexturePSBaseColorNormalSpecular")),
 			FName(TEXT("TVirtualTexturePSBaseColorNormalRoughness")),
 			FName(TEXT("TVirtualTexturePSWorldHeight")),
+			FName(TEXT("TVirtualTexturePSDisplacement")),
 		};
 		return AllowedShaderTypes;
 	}
@@ -3908,11 +3910,13 @@ public:
 			FName(TEXT("TVirtualTextureVSBaseColorNormalSpecular")),
 			FName(TEXT("TVirtualTextureVSBaseColorNormalRoughness")),
 			FName(TEXT("TVirtualTextureVSWorldHeight")),
+			FName(TEXT("TVirtualTextureVSDisplacement")),
 			FName(TEXT("TVirtualTexturePSBaseColor")),
 			FName(TEXT("TVirtualTexturePSBaseColorNormal")),
 			FName(TEXT("TVirtualTexturePSBaseColorNormalSpecular")),
 			FName(TEXT("TVirtualTexturePSBaseColorNormalRoughness")),
 			FName(TEXT("TVirtualTexturePSWorldHeight")),
+			FName(TEXT("TVirtualTexturePSDisplacement")),
 		};
 		return ShaderTypes;
 	}
diff --git a/Engine/Source/Runtime/RenderCore/Private/ShaderMaterialDerivedHelpers.cpp b/Engine/Source/Runtime/RenderCore/Private/ShaderMaterialDerivedHelpers.cpp
--- a/Engine/Source/Runtime/RenderCore/Private/ShaderMaterialDerivedHelpers.cpp
+++ b/Engine/Source/Runtime/RenderCore/Private/ShaderMaterialDerivedHelpers.cpp
@@ -124,6 +124,10 @@ FShaderMaterialDerivedDefines RENDERCORE_API CalculateDerivedMaterialParameters(
 			{
 				Dst.PIXELSHADEROUTPUT_MRT0 = 1;
 			}
+			else if (Mat.OUT_DISPLACEMENT)
+			{
+				Dst.PIXELSHADEROUTPUT_MRT0 = 1;
+			}
 		}
 	}
 	else if (Mat.IS_DECAL)
diff --git a/Engine/Source/Runtime/RenderCore/Public/ShaderMaterial.h b/Engine/Source/Runtime/RenderCore/Public/ShaderMaterial.h
--- a/Engine/Source/Runtime/RenderCore/Public/ShaderMaterial.h
+++ b/Engine/Source/Runtime/RenderCore/Public/ShaderMaterial.h
@@ -179,6 +179,7 @@ struct FShaderMaterialPropertyDefines
 	uint8 OUT_BASECOLOR_NORMAL_ROUGHNESS : 1;
 	uint8 OUT_BASECOLOR_NORMAL_SPECULAR : 1;
 	uint8 OUT_WORLDHEIGHT : 1;
+	uint8 OUT_DISPLACEMENT : 1;
 
 	uint8 SUBSTRATE_ENABLED : 1;
 	uint8 MATERIAL_IS_SUBSTRATE : 1;
diff --git a/Engine/Source/Runtime/Renderer/Private/VT/RuntimeVirtualTextureRender.cpp b/Engine/Source/Runtime/Renderer/Private/VT/RuntimeVirtualTextureRender.cpp
--- a/Engine/Source/Runtime/Renderer/Private/VT/RuntimeVirtualTextureRender.cpp
+++ b/Engine/Source/Runtime/Renderer/Private/VT/RuntimeVirtualTextureRender.cpp
@@ -346,6 +346,7 @@ namespace RuntimeVirtualTexture
 			return GetBlendStateImpl(OutputAttributeMask);
 		}
 	};
+
 	/** Specialization for ERuntimeVirtualTextureMaterialType::WorldHeight */
 	class FMaterialPolicy_WorldHeight
 	{
@@ -362,6 +363,22 @@ namespace RuntimeVirtualTexture
 		}
 	};
 
+	/** Specialization for ERuntimeVirtualTextureMaterialType::Displacement */
+	class FMaterialPolicy_Displacement
+	{
+	public:
+		static void ModifyCompilationEnvironment(FShaderCompilerEnvironment& OutEnvironment)
+		{
+			OutEnvironment.SetDefine(TEXT("OUT_DISPLACEMENT"), 1);
+			OutEnvironment.SetRenderTargetOutputFormat(0, PF_R32_FLOAT);
+		}
+
+		static FRHIBlendState* GetBlendState(uint8 OutputAttributeMask)
+		{
+			return TStaticBlendState< CW_RED, BO_Max, BF_One, BF_One, BO_Add, BF_One, BF_One >::GetRHI();
+		}
+	};
+
 
 	/** Vertex shader derivation of material shader. Templated on policy for virtual texture layout. */
 	template< class MaterialPolicy >
@@ -416,6 +433,7 @@ namespace RuntimeVirtualTexture
 	IMPLEMENT_VIRTUALTEXTURE_SHADER_TYPE(FMaterialPolicy_BaseColor, BaseColor);
 	IMPLEMENT_VIRTUALTEXTURE_SHADER_TYPE(FMaterialPolicy_BaseColorNormalSpecular, BaseColorNormalSpecular);
 	IMPLEMENT_VIRTUALTEXTURE_SHADER_TYPE(FMaterialPolicy_WorldHeight, WorldHeight);
+	IMPLEMENT_VIRTUALTEXTURE_SHADER_TYPE(FMaterialPolicy_Displacement, Displacement);
 	IMPLEMENT_VIRTUALTEXTURE_SHADER_TYPE(FMaterialPolicy_BaseColorNormalRoughness, BaseColorNormalRoughness);
 
 	/** Structure to localize the setup of our render graph based on the virtual texture setup. */
@@ -451,6 +469,7 @@ namespace RuntimeVirtualTexture
 				AddRenderTargetInfo(PF_B8G8R8A8, RTCreateFlags, RenderTargetsInfo);
 				break;
 			case ERuntimeVirtualTextureMaterialType::WorldHeight:
+			case ERuntimeVirtualTextureMaterialType::Displacement:
 				AddRenderTargetInfo(PF_G16, RTCreateFlags, RenderTargetsInfo);
 				break;
 			}
@@ -680,6 +699,30 @@ namespace RuntimeVirtualTexture
 					OutputAlias0 = CopyTexture0 = GraphBuilder.CreateTexture(FRDGTextureDesc::Create2D(Desc.TextureSize, PF_B8G8R8A8, FClearValueBinding::Black, RTCreateFlags), TEXT("CopyTexture0"));
 				}
 				break;
+			case ERuntimeVirtualTextureMaterialType::Displacement:
+				if (bRenderPass)
+				{
+					OutputAlias0 = RenderTexture0 = GraphBuilder.CreateTexture(FRDGTextureDesc::Create2D(Desc.TextureSize, PF_G16, FClearValueBinding::Black, RTCreateFlags), TEXT("RenderTexture0"));
+				}
+				if (bCompressPass)
+				{
+					if (bDirectAliasing)
+					{
+						CompressTexture0 = GraphBuilder.RegisterExternalTexture(Desc.OutputTargets[0], ERDGTextureFlags::None);
+						CompressTextureUAV0_64bit = GraphBuilder.CreateUAV(FRDGTextureUAVDesc(CompressTexture0, 0, Compressed64BitFormat));
+						OutputAlias0 = nullptr;
+					}
+					else
+					{
+						OutputAlias0 = CompressTexture0 = GraphBuilder.CreateTexture(FRDGTextureDesc::Create2D(Desc.TextureSize / 4, Compressed64BitFormat, FClearValueBinding::None, TexCreate_UAV), TEXT("CompressTexture0"));
+						CompressTextureUAV0_64bit = GraphBuilder.CreateUAV(FRDGTextureUAVDesc(CompressTexture0));
+					}
+				}
+				if (bCopyThumbnailPass)
+				{
+					OutputAlias0 = CopyTexture0 = GraphBuilder.CreateTexture(FRDGTextureDesc::Create2D(Desc.TextureSize, PF_B8G8R8A8, FClearValueBinding::Black, RTCreateFlags), TEXT("CopyTexture0"));
+				}
+				break;
 			}
 		}
 
@@ -741,18 +784,16 @@ namespace RuntimeVirtualTexture
 				{
 				case ERuntimeVirtualTextureMaterialType::BaseColor:
 					return Process<FMaterialPolicy_BaseColor>(MeshBatch, BatchElementMask, StaticMeshId, OutputAttributeMask, PrimitiveSceneProxy, *MaterialRenderProxy, *Material);
-					break;
 				case ERuntimeVirtualTextureMaterialType::BaseColor_Normal_Roughness:
 					return Process<FMaterialPolicy_BaseColorNormalRoughness>(MeshBatch, BatchElementMask, StaticMeshId, OutputAttributeMask, PrimitiveSceneProxy, *MaterialRenderProxy, *Material);
-					break;
 				case ERuntimeVirtualTextureMaterialType::BaseColor_Normal_Specular:
 				case ERuntimeVirtualTextureMaterialType::BaseColor_Normal_Specular_YCoCg:
 				case ERuntimeVirtualTextureMaterialType::BaseColor_Normal_Specular_Mask_YCoCg:
 					return Process<FMaterialPolicy_BaseColorNormalSpecular>(MeshBatch, BatchElementMask, StaticMeshId, OutputAttributeMask, PrimitiveSceneProxy, *MaterialRenderProxy, *Material);
-					break;
 				case ERuntimeVirtualTextureMaterialType::WorldHeight:
 					return Process<FMaterialPolicy_WorldHeight>(MeshBatch, BatchElementMask, StaticMeshId, OutputAttributeMask, PrimitiveSceneProxy, *MaterialRenderProxy, *Material);
-					break;
+				case ERuntimeVirtualTextureMaterialType::Displacement:
+					return Process<FMaterialPolicy_Displacement>(MeshBatch, BatchElementMask, StaticMeshId, OutputAttributeMask, PrimitiveSceneProxy, *MaterialRenderProxy, *Material);
 				default:
 					break;
 				}
@@ -911,6 +952,7 @@ namespace RuntimeVirtualTexture
 				CollectPSOInitializers<FMaterialPolicy_BaseColorNormalRoughness>(VertexFactoryData, Material, MeshFillMode, MeshCullMode, OutputAttributeMask, ERuntimeVirtualTextureMaterialType::BaseColor_Normal_Roughness, PSOInitializers);
 				CollectPSOInitializers<FMaterialPolicy_BaseColorNormalSpecular>(VertexFactoryData, Material, MeshFillMode, MeshCullMode, OutputAttributeMask, ERuntimeVirtualTextureMaterialType::BaseColor_Normal_Specular, PSOInitializers);
 				CollectPSOInitializers<FMaterialPolicy_WorldHeight>(VertexFactoryData, Material, MeshFillMode, MeshCullMode, OutputAttributeMask, ERuntimeVirtualTextureMaterialType::WorldHeight, PSOInitializers);
+				CollectPSOInitializers<FMaterialPolicy_Displacement>(VertexFactoryData, Material, MeshFillMode, MeshCullMode, OutputAttributeMask, ERuntimeVirtualTextureMaterialType::Displacement, PSOInitializers);
 			}
 		}
 
@@ -1101,6 +1143,7 @@ namespace RuntimeVirtualTexture
 	IMPLEMENT_SHADER_TYPE(template<>, FShader_VirtualTextureCompress_CS< ERuntimeVirtualTextureMaterialType::BaseColor_Normal_Roughness >, TEXT("/Engine/Private/VirtualTextureCompress.usf"), TEXT("CompressBaseColorNormalRoughnessCS"), SF_Compute);
 	IMPLEMENT_SHADER_TYPE(template<>, FShader_VirtualTextureCompress_CS< ERuntimeVirtualTextureMaterialType::BaseColor_Normal_Specular_YCoCg >, TEXT("/Engine/Private/VirtualTextureCompress.usf"), TEXT("CompressBaseColorNormalSpecularYCoCgCS"), SF_Compute);
 	IMPLEMENT_SHADER_TYPE(template<>, FShader_VirtualTextureCompress_CS< ERuntimeVirtualTextureMaterialType::BaseColor_Normal_Specular_Mask_YCoCg >, TEXT("/Engine/Private/VirtualTextureCompress.usf"), TEXT("CompressBaseColorNormalSpecularMaskYCoCgCS"), SF_Compute);
+	IMPLEMENT_SHADER_TYPE(template<>, FShader_VirtualTextureCompress_CS< ERuntimeVirtualTextureMaterialType::Displacement >, TEXT("/Engine/Private/VirtualTextureCompress.usf"), TEXT("CompressDisplacementCS"), SF_Compute);
 
 
 	/** Add the BC compression pass to the graph. */
@@ -1138,6 +1181,8 @@ namespace RuntimeVirtualTexture
 			break;
 		case ERuntimeVirtualTextureMaterialType::BaseColor_Normal_Specular_Mask_YCoCg:
 			AddCompressPass<ERuntimeVirtualTextureMaterialType::BaseColor_Normal_Specular_Mask_YCoCg>(GraphBuilder, FeatureLevel, Parameters, GroupCount);
+		case ERuntimeVirtualTextureMaterialType::Displacement:
+			AddCompressPass<ERuntimeVirtualTextureMaterialType::Displacement>(GraphBuilder, FeatureLevel, Parameters, GroupCount);
 			break;
 		}
 	}
@@ -1267,6 +1312,7 @@ namespace RuntimeVirtualTexture
 			AddCopyPass<ERuntimeVirtualTextureMaterialType::BaseColor>(GraphBuilder, FeatureLevel, Parameters, TextureSize);
 			break;
 		case ERuntimeVirtualTextureMaterialType::WorldHeight:
+		case ERuntimeVirtualTextureMaterialType::Displacement:
 			AddCopyPass<ERuntimeVirtualTextureMaterialType::WorldHeight>(GraphBuilder, FeatureLevel, Parameters, TextureSize);
 			break;
 		}
