Quick Fix for some editor crashes. This adds a redundant test before hitting a variable that could be garbage.

diff --git a/Engine/Source/Runtime/Renderer/Private/PrimitiveSceneInfo.cpp b/Engine/Source/Runtime/Renderer/Private/PrimitiveSceneInfo.cpp
--- a/Engine/Source/Runtime/Renderer/Private/PrimitiveSceneInfo.cpp
+++ b/Engine/Source/Runtime/Renderer/Private/PrimitiveSceneInfo.cpp
@@ -1919,6 +1919,11 @@ bool FPrimitiveSceneInfo::RequestUniformBufferUpdate()
 
 const FInstanceSceneDataBuffers *FPrimitiveSceneInfo::GetInstanceSceneDataBuffers() const
 { 
+	if (!HasInstanceDataBuffers())
+	{
+		return nullptr;
+	}
+
 	if (InstanceDataUpdateTaskInfo)
 	{
 		InstanceDataUpdateTaskInfo->WaitForUpdateCompletion();
