RayTracing: Remove deprecated RTGI indirect lighting technique

Lumen is a superior alternative in all cases and is the main path we wish to support going forward. Since RTGI has been deprecated for several releases now, it can be safely removed.

diff --git a/Engine/Shaders/Private/DiffuseIndirectComposite.usf b/Engine/Shaders/Private/DiffuseIndirectComposite.usf
--- a/Engine/Shaders/Private/DiffuseIndirectComposite.usf
+++ b/Engine/Shaders/Private/DiffuseIndirectComposite.usf
@@ -8,8 +8,7 @@
 //------------------------------------------------------- CONFIGS
 
 #define DIM_APPLY_DIFFUSE_INDIRECT_SSGI					1
-#define DIM_APPLY_DIFFUSE_INDIRECT_RTGI					2
-#define DIM_APPLY_DIFFUSE_INDIRECT_SCREEN_PROBE_GATHER	3
+#define DIM_APPLY_DIFFUSE_INDIRECT_SCREEN_PROBE_GATHER	2
 
 #if DIM_APPLY_DIFFUSE_INDIRECT == DIM_APPLY_DIFFUSE_INDIRECT_SSGI
 	#define CONFIG_SIGNAL_PROCESSING SIGNAL_PROCESSING_SSGI
@@ -20,15 +19,6 @@
 	#define CONFIG_SIGNAL_INPUT_LAYOUT  SIGNAL_BUFFER_LAYOUT_SSGI_HISTORY_R11G11B10
 	#define CONFIG_INPUT_TEXTURE_COUNT 2
 
-#elif DIM_APPLY_DIFFUSE_INDIRECT == DIM_APPLY_DIFFUSE_INDIRECT_RTGI
-	#define CONFIG_SIGNAL_PROCESSING SIGNAL_PROCESSING_DIFFUSE_INDIRECT_AND_AO
-	#define COMPILE_SIGNAL_COLOR 1
-	#define MAX_SIGNAL_BATCH_SIZE 1
-	#define SIGNAL_ARRAY_SIZE 1
-
-	#define CONFIG_SIGNAL_INPUT_LAYOUT  SIGNAL_BUFFER_LAYOUT_DIFFUSE_INDIRECT_AND_AO_HISTORY
-	#define CONFIG_INPUT_TEXTURE_COUNT 2
-
 #elif DIM_APPLY_DIFFUSE_INDIRECT == DIM_APPLY_DIFFUSE_INDIRECT_SCREEN_PROBE_GATHER
 	#define CONFIG_SIGNAL_PROCESSING SIGNAL_PROCESSING_SSGI
 	#define COMPILE_SIGNAL_COLOR_ARRAY 3
@@ -456,7 +446,7 @@ void MainPS(
 					Sample = UncompressedAccumulators.Array[0].Moment1;
 				#endif
 				
-				#if DIM_APPLY_DIFFUSE_INDIRECT == DIM_APPLY_DIFFUSE_INDIRECT_SSGI || DIM_APPLY_DIFFUSE_INDIRECT == DIM_APPLY_DIFFUSE_INDIRECT_RTGI
+				#if DIM_APPLY_DIFFUSE_INDIRECT == DIM_APPLY_DIFFUSE_INDIRECT_SSGI
 				{
 					DiffuseIndirectLighting = Sample.SceneColor.rgb;
 				}
diff --git a/Engine/Shaders/Private/PathTracing/Light/PathTracingLightSampling.ush b/Engine/Shaders/Private/PathTracing/Light/PathTracingLightSampling.ush
--- a/Engine/Shaders/Private/PathTracing/Light/PathTracingLightSampling.ush
+++ b/Engine/Shaders/Private/PathTracing/Light/PathTracingLightSampling.ush
@@ -167,53 +167,3 @@ void SelectLight(float RandSample, int NumLights, inout float LightPickingCdf[RA
 #endif
 #endif
 }
-
-bool SampleLight(
-	float3 RandSample,
-	float3 TranslatedWorldPos,
-	float3 WorldNormal,
-	uint PrimitiveLightingChannelMask,
-	bool IsTransmissiveMaterial,
-	out int OutLightId,
-	out float3 OutDirection,
-	out float  OutDistance,
-	out float3 OutRadianceOverPdf,
-	out float OutPdf)
-{
-	// set default values in case sampling fails
-	OutLightId = -1;
-	OutDirection = 0;
-	OutDistance = 0;
-	OutRadianceOverPdf = 0;
-	OutPdf = 0;
-
-	if (SceneLightCount == 0)
-	{
-		// no lights in the scene
-		return false;
-	}
-
-	float LightPickingCdf[RAY_TRACING_LIGHT_COUNT_MAXIMUM];
-
-	float LightPickingCdfSum = InitLightPickingCdf(TranslatedWorldPos, WorldNormal, PrimitiveLightingChannelMask, IsTransmissiveMaterial, LightPickingCdf);
-
-	if (LightPickingCdfSum == 0.0)
-	{
-		// light picking cdf is all 0
-		return false;
-	}
-
-	float LightPickPdf = 0;
-
-	SelectLight(RandSample.x * LightPickingCdfSum, SceneLightCount, LightPickingCdf, OutLightId, LightPickPdf);
-
-	LightPickPdf /= LightPickingCdfSum;
-
-	FLightSample Result = SampleLight(OutLightId, RandSample.yz, TranslatedWorldPos, WorldNormal);
-	OutDirection = Result.Direction;
-	OutDistance = Result.Distance;
-	OutRadianceOverPdf = Result.RadianceOverPdf / LightPickPdf;
-	OutPdf = Result.Pdf * LightPickPdf;
-
-	return true;
-}
diff --git a/Engine/Shaders/Private/RayTracing/RayTracingCreateGatherPointsRGS.usf b/Engine/Shaders/Private/RayTracing/RayTracingCreateGatherPointsRGS.usf
deleted file mode 100644
diff --git a/Engine/Shaders/Private/RayTracing/RayTracingFinalGatherRGS.usf b/Engine/Shaders/Private/RayTracing/RayTracingFinalGatherRGS.usf
deleted file mode 100644
diff --git a/Engine/Shaders/Private/RayTracing/RayTracingGlobalIlluminationRGS.usf b/Engine/Shaders/Private/RayTracing/RayTracingGlobalIlluminationRGS.usf
deleted file mode 100644
diff --git a/Engine/Source/Runtime/Engine/Classes/Engine/EngineTypes.h b/Engine/Source/Runtime/Engine/Classes/Engine/EngineTypes.h
--- a/Engine/Source/Runtime/Engine/Classes/Engine/EngineTypes.h
+++ b/Engine/Source/Runtime/Engine/Classes/Engine/EngineTypes.h
@@ -435,9 +435,6 @@ namespace EDynamicGlobalIlluminationMethod
 		/** Standalone Screen Space Global Illumination.  Low cost, but limited by screen space information. */
 		ScreenSpace UMETA(DisplayName="Screen Space (Beta)"),
 
-		/** Standalone Ray Traced Global Illumination technique.  Deprecated, use Lumen Global Illumination instead. */
-		RayTraced UMETA(DisplayName="Standalone Ray Traced (Deprecated)"),
-
 		/** Use a plugin for Global Illumination */
 		Plugin UMETA(DisplayName="Plugin"),
 	};
diff --git a/Engine/Source/Runtime/Engine/Classes/Engine/Scene.h b/Engine/Source/Runtime/Engine/Classes/Engine/Scene.h
--- a/Engine/Source/Runtime/Engine/Classes/Engine/Scene.h
+++ b/Engine/Source/Runtime/Engine/Classes/Engine/Scene.h
@@ -101,14 +101,6 @@ enum class ETranslucencyType : uint8
 };
 
 
-UENUM()
-enum class ERayTracingGlobalIlluminationType : uint8
-{
-	Disabled    UMETA(DisplayName = "Disabled"),
-	BruteForce  UMETA(DisplayName = "Brute Force"),
-	FinalGather UMETA(DisplayName = "Final Gather")
-};
-
 UENUM()
 enum class EReflectedAndRefractedRayTracedShadows : uint8
 {
@@ -1584,11 +1576,6 @@ struct FPostProcessSettings
 	UPROPERTY(interp, BlueprintReadWrite, AdvancedDisplay, Category="Global Illumination", meta=(ClampMin = "0", UIMax = "4.0", editcondition = "bOverride_IndirectLightingIntensity", DisplayName = "Indirect Lighting Intensity"))
 	float IndirectLightingIntensity;
 
-#if WITH_EDITORONLY_DATA
-	UPROPERTY()
-	uint32 RayTracingGI_DEPRECATED : 1;
-#endif
-
 	/** Scales Lumen Scene's quality.  Larger scales cause Lumen Scene to be calculated with a higher fidelity, which can be visible in reflections, but increase GPU cost. */
 	UPROPERTY(interp, EditAnywhere, BlueprintReadWrite, Category = "Global Illumination|Lumen Global Illumination", meta = (ClampMin = ".25", UIMax = "2", editcondition = "bOverride_LumenSceneLightingQuality", DisplayName = "Lumen Scene Lighting Quality"))
 	float LumenSceneLightingQuality;
@@ -1633,18 +1620,6 @@ struct FPostProcessSettings
 	UPROPERTY(EditAnywhere, BlueprintReadWrite, Category = "Global Illumination|Lumen Global Illumination", meta = (ClampMin = ".5", ClampMax = "1", editcondition = "bOverride_LumenSurfaceCacheResolution", DisplayName = "Scene Capture Cache Resolution Scale"))
 	float LumenSurfaceCacheResolution;
 
-	/** Sets the ray tracing global illumination type. */
-	UPROPERTY(interp, EditAnywhere, BlueprintReadWrite, Category = "Global Illumination|Ray Tracing Global Illumination", meta = (editcondition = "bOverride_RayTracingGI", DisplayName = "Type"))
-	ERayTracingGlobalIlluminationType RayTracingGIType;
-
-	/** Sets the ray tracing global illumination maximum bounces. */
-	UPROPERTY(interp, EditAnywhere, BlueprintReadWrite, Category = "Global Illumination|Ray Tracing Global Illumination", meta = (ClampMin = "0", ClampMax = "50", editcondition = "bOverride_RayTracingGIMaxBounces", DisplayName = "Max. Bounces"))
-	int32 RayTracingGIMaxBounces;
-
-	/** Sets the samples per pixel for ray tracing global illumination. */
-	UPROPERTY(interp, EditAnywhere, BlueprintReadWrite, Category = "Global Illumination|Ray Tracing Global Illumination", meta = (ClampMin = "1", ClampMax = "65536", editcondition = "bOverride_RayTracingGISamplesPerPixel", DisplayName = "Samples Per Pixel"))
-	int32 RayTracingGISamplesPerPixel;
-	
 	/** Chooses the Reflection method. Not compatible with Forward Shading. */
 	UPROPERTY(interp, EditAnywhere, BlueprintReadWrite, Category = "Reflections", meta = (editcondition = "bOverride_ReflectionMethod", DisplayName = "Method"))
 	TEnumAsByte<EReflectionMethod::Type> ReflectionMethod;
@@ -2311,10 +2286,6 @@ struct FPostProcessSettings
 			BloomConvolutionPreFilterMax = BloomConvolutionPreFilter_DEPRECATED.Y;
 			BloomConvolutionPreFilterMult = BloomConvolutionPreFilter_DEPRECATED.Z;
 		}
-		if (RayTracingGI_DEPRECATED)
-		{
-			RayTracingGIType = (ERayTracingGlobalIlluminationType)(RayTracingGI_DEPRECATED == 1);
-		}
 
 		if (bOverride_LocalExposureContrastScale_DEPRECATED)
 		{
diff --git a/Engine/Source/Runtime/Engine/Private/Scene.cpp b/Engine/Source/Runtime/Engine/Private/Scene.cpp
--- a/Engine/Source/Runtime/Engine/Private/Scene.cpp
+++ b/Engine/Source/Runtime/Engine/Private/Scene.cpp
@@ -572,9 +572,6 @@ FPostProcessSettings::FPostProcessSettings()
 	LumenFullSkylightLeakingDistance = 1000.0f;
 
 	ColorGradingIntensity = 1.0f;
-	RayTracingGIType = ERayTracingGlobalIlluminationType::Disabled;
-	RayTracingGIMaxBounces = 1;
-	RayTracingGISamplesPerPixel = 4;
 
 	DepthOfFieldFocalDistance = 0; // Intentionally invalid to disable DOF by default.
 	DepthOfFieldFstop = 4.0f; 
@@ -903,12 +900,6 @@ void FPostProcessSettings::PostSerialize(const FArchive& Ar)
 
 		if (Ar.CustomVer(FUE5ReleaseStreamObjectVersion::GUID) < FUE5ReleaseStreamObjectVersion::ReflectionMethodEnum)
 		{
-			if (bOverride_RayTracingGI && RayTracingGIType != ERayTracingGlobalIlluminationType::Disabled)
-			{
-				bOverride_DynamicGlobalIlluminationMethod = true;
-				DynamicGlobalIlluminationMethod = EDynamicGlobalIlluminationMethod::RayTraced;
-			}
-
 			if (bOverride_ReflectionsType_DEPRECATED)
 			{
 				bOverride_ReflectionMethod = true;
diff --git a/Engine/Source/Runtime/Engine/Private/SceneView.cpp b/Engine/Source/Runtime/Engine/Private/SceneView.cpp
--- a/Engine/Source/Runtime/Engine/Private/SceneView.cpp
+++ b/Engine/Source/Runtime/Engine/Private/SceneView.cpp
@@ -1628,11 +1628,6 @@ void FSceneView::OverridePostProcessSettings(const FPostProcessSettings& Src, fl
 		SET_PP(RayTracingTranslucencyShadows);
 		SET_PP(RayTracingTranslucencyRefraction);
 
-		if (Src.bOverride_RayTracingGI)
-		{
-			Dest.RayTracingGIType = Src.RayTracingGIType;
-		}
-
 		SET_PP(DynamicGlobalIlluminationMethod);
 		SET_PP(LumenSurfaceCacheResolution);
 		SET_PP(LumenSceneLightingQuality);
@@ -1653,8 +1648,6 @@ void FSceneView::OverridePostProcessSettings(const FPostProcessSettings& Src, fl
 		SET_PP(LumenMaxReflectionBounces);
 		SET_PP(ReflectionMethod);
 		SET_PP(LumenReflectionQuality);
-		SET_PP(RayTracingGIMaxBounces);
-		SET_PP(RayTracingGISamplesPerPixel);
 		SET_PP(RayTracingAO);
 		SET_PP(RayTracingAOSamplesPerPixel);
 		SET_PP(RayTracingAOIntensity);
diff --git a/Engine/Source/Runtime/Renderer/Private/DeferredShadingRenderer.cpp b/Engine/Source/Runtime/Renderer/Private/DeferredShadingRenderer.cpp
--- a/Engine/Source/Runtime/Renderer/Private/DeferredShadingRenderer.cpp
+++ b/Engine/Source/Runtime/Renderer/Private/DeferredShadingRenderer.cpp
@@ -1902,7 +1902,6 @@ bool FDeferredShadingSceneRenderer::SetupRayTracingPipelineStates(FRDGBuilder& G
 					PrepareRayTracingShadows(View, *Scene, RayGenShaders);
 					PrepareRayTracingAmbientOcclusion(View, RayGenShaders);
 					PrepareRayTracingSkyLight(View, *Scene, RayGenShaders);
-					PrepareRayTracingGlobalIllumination(View, RayGenShaders);
 					PrepareRayTracingGlobalIlluminationPlugin(View, RayGenShaders);
 					PrepareRayTracingTranslucency(View, RayGenShaders);
 					PrepareRayTracingVolumetricFogShadows(View, *Scene, RayGenShaders);
@@ -1938,7 +1937,6 @@ bool FDeferredShadingSceneRenderer::SetupRayTracingPipelineStates(FRDGBuilder& G
 			{
 				PrepareRayTracingReflectionsDeferredMaterial(View, *Scene, DeferredMaterialRayGenShaders);
 				PrepareRayTracingDeferredReflectionsDeferredMaterial(View, *Scene, DeferredMaterialRayGenShaders);
-				PrepareRayTracingGlobalIlluminationDeferredMaterial(View, DeferredMaterialRayGenShaders);
 				if (DoesPlatformSupportLumenGI(ShaderPlatform))
 				{
 					PrepareLumenHardwareRayTracingReflectionsDeferredMaterial(View, DeferredMaterialRayGenShaders);
@@ -4627,7 +4625,6 @@ bool AnyRayTracingPassEnabled(const FScene* Scene, const FViewInfo& View)
 
 	return ShouldRenderRayTracingAmbientOcclusion(View)
 		|| ShouldRenderRayTracingReflections(View)
-		|| ShouldRenderRayTracingGlobalIllumination(View)
 		|| ShouldRenderRayTracingTranslucency(View)
 		|| ShouldRenderRayTracingSkyLight(Scene->SkyLight)
 		|| ShouldRenderRayTracingShadows()
diff --git a/Engine/Source/Runtime/Renderer/Private/DeferredShadingRenderer.h b/Engine/Source/Runtime/Renderer/Private/DeferredShadingRenderer.h
--- a/Engine/Source/Runtime/Renderer/Private/DeferredShadingRenderer.h
+++ b/Engine/Source/Runtime/Renderer/Private/DeferredShadingRenderer.h
@@ -280,7 +280,6 @@ enum class EDiffuseIndirectMethod
 {
 	Disabled,
 	SSGI,
-	RTGI,
 	Lumen,
 	Plugin,
 };
@@ -1036,38 +1035,6 @@ private:
 		FRDGTextureRef SkyLightRT,
 		FRDGTextureRef HitDistanceRT);
 	
-	bool RenderRayTracingGlobalIllumination(
-		FRDGBuilder& GraphBuilder, 
-		FSceneTextureParameters& SceneTextures,
-		FViewInfo& View,
-		IScreenSpaceDenoiser::FAmbientOcclusionRayTracingConfig* RayTracingConfig,
-		IScreenSpaceDenoiser::FDiffuseIndirectInputs* OutDenoiserInputs);
-	
-	void RenderRayTracingGlobalIlluminationBruteForce(
-		FRDGBuilder& GraphBuilder,
-		FSceneTextureParameters& SceneTextures,
-		FViewInfo& View,
-		const IScreenSpaceDenoiser::FAmbientOcclusionRayTracingConfig& RayTracingConfig,
-		int32 UpscaleFactor,
-		IScreenSpaceDenoiser::FDiffuseIndirectInputs* OutDenoiserInputs);
-
-	void RayTracingGlobalIlluminationCreateGatherPoints(
-		FRDGBuilder& GraphBuilder,
-		FSceneTextureParameters& SceneTextures,
-		FViewInfo& View,
-		int32 UpscaleFactor,
-		int32 SampleIndex,
-		FRDGBufferRef& GatherPointsBuffer,
-		FIntVector& GatherPointsResolution);
-
-	void RenderRayTracingGlobalIlluminationFinalGather(
-		FRDGBuilder& GraphBuilder,
-		FSceneTextureParameters& SceneTextures,
-		FViewInfo& View,
-		const IScreenSpaceDenoiser::FAmbientOcclusionRayTracingConfig& RayTracingConfig,
-		int32 UpscaleFactor,
-		IScreenSpaceDenoiser::FDiffuseIndirectInputs* OutDenoiserInputs);
-	
 	void RenderRayTracingAmbientOcclusion(
 		FRDGBuilder& GraphBuilder,
 		FViewInfo& View,
@@ -1169,7 +1136,6 @@ private:
 	static void PrepareRayTracingShadows(const FViewInfo& View, const FScene& Scene, TArray<FRHIRayTracingShader*>& OutRayGenShaders);
 	static void PrepareRayTracingAmbientOcclusion(const FViewInfo& View, TArray<FRHIRayTracingShader*>& OutRayGenShaders);
 	static void PrepareRayTracingSkyLight(const FViewInfo& View, const FScene& Scene, TArray<FRHIRayTracingShader*>& OutRayGenShaders);
-	static void PrepareRayTracingGlobalIllumination(const FViewInfo& View, TArray<FRHIRayTracingShader*>& OutRayGenShaders);
 	static void PrepareRayTracingGlobalIlluminationPlugin(const FViewInfo& View, TArray<FRHIRayTracingShader*>& OutRayGenShaders);
 	static void PrepareRayTracingTranslucency(const FViewInfo& View, TArray<FRHIRayTracingShader*>& OutRayGenShaders);
 	static void PrepareRayTracingVolumetricFogShadows(const FViewInfo& View, const FScene& Scene, TArray<FRHIRayTracingShader*>& OutRayGenShaders);
@@ -1190,8 +1156,7 @@ private:
 	// Versions for setting up the deferred material pipeline
 	static void PrepareRayTracingReflectionsDeferredMaterial(const FViewInfo& View, const FScene& Scene, TArray<FRHIRayTracingShader*>& OutRayGenShaders);
 	static void PrepareRayTracingDeferredReflectionsDeferredMaterial(const FViewInfo& View, const FScene& Scene, TArray<FRHIRayTracingShader*>& OutRayGenShaders);
-	static void PrepareRayTracingGlobalIlluminationDeferredMaterial(const FViewInfo& View, TArray<FRHIRayTracingShader*>& OutRayGenShaders);
-
+	
 	// Versions for setting up the lumen material pipeline
 	static void PrepareLumenHardwareRayTracingTranslucencyVolumeLumenMaterial(const FViewInfo& View, TArray<FRHIRayTracingShader*>& OutRayGenShaders);
 	void PrepareLumenHardwareRayTracingVisualizeLumenMaterial(const FViewInfo& View, TArray<FRHIRayTracingShader*>& OutRayGenShaders);
diff --git a/Engine/Source/Runtime/Renderer/Private/IndirectLightRendering.cpp b/Engine/Source/Runtime/Renderer/Private/IndirectLightRendering.cpp
--- a/Engine/Source/Runtime/Renderer/Private/IndirectLightRendering.cpp
+++ b/Engine/Source/Runtime/Renderer/Private/IndirectLightRendering.cpp
@@ -39,8 +39,7 @@ static TAutoConsoleVariable<int32> CVarDynamicGlobalIlluminationMethod(
 	TEXT("0 - None.  Global Illumination can be baked into Lightmaps but no technique will be used for Dynamic Global Illumination.\n")
 	TEXT("1 - Lumen.  Use Lumen Global Illumination for all lights, emissive materials casting light and SkyLight Occlusion.  Requires 'Generate Mesh Distance Fields' enabled for Software Ray Tracing and 'Support Hardware Ray Tracing' enabled for Hardware Ray Tracing.\n")
 	TEXT("2 - SSGI.  Standalone Screen Space Global Illumination.  Low cost, but limited by screen space information.\n")
-	TEXT("3 - RTGI.  Ray Traced Global Illumination technique.  Deprecated, use Lumen Global Illumination instead.\n")
-	TEXT("4 - Plugin.  Use a plugin for Global Illumination."),
+	TEXT("3 - Plugin.  Use a plugin for Global Illumination."),
 	ECVF_RenderThreadSafe);
 
 // This is the project default reflection method, NOT the scalability setting (see r.Lumen.Reflections.Allow for scalability)
@@ -120,7 +119,7 @@ class FDiffuseIndirectCompositePS : public FGlobalShader
 	DECLARE_GLOBAL_SHADER(FDiffuseIndirectCompositePS)
 	SHADER_USE_PARAMETER_STRUCT(FDiffuseIndirectCompositePS, FGlobalShader)
 
-	class FApplyDiffuseIndirectDim : SHADER_PERMUTATION_INT("DIM_APPLY_DIFFUSE_INDIRECT", 4);
+	class FApplyDiffuseIndirectDim : SHADER_PERMUTATION_INT("DIM_APPLY_DIFFUSE_INDIRECT", 3);
 	class FUpscaleDiffuseIndirectDim : SHADER_PERMUTATION_BOOL("DIM_UPSCALE_DIFFUSE_INDIRECT");
 	class FScreenBentNormal : SHADER_PERMUTATION_BOOL("DIM_SCREEN_BENT_NORMAL");
 	class FSubstrateTileType : SHADER_PERMUTATION_INT("SUBSTRATE_TILETYPE", 4);
@@ -144,7 +143,7 @@ class FDiffuseIndirectCompositePS : public FGlobalShader
 		}
 
 		// Only support Bent Normal for ScreenProbeGather
-		if (PermutationVector.Get<FApplyDiffuseIndirectDim>() != 3 && PermutationVector.Get<FScreenBentNormal>())
+		if (PermutationVector.Get<FApplyDiffuseIndirectDim>() != 2 && PermutationVector.Get<FScreenBentNormal>())
 		{
 			return false;
 		}
@@ -152,7 +151,7 @@ class FDiffuseIndirectCompositePS : public FGlobalShader
 		// Build Substrate tile permutation only for Lumen
 		if (PermutationVector.Get<FSubstrateTileType>() != ESubstrateTileType::EComplex)
 		{
-			return Substrate::IsSubstrateEnabled() && PermutationVector.Get<FApplyDiffuseIndirectDim>() == 3;
+			return Substrate::IsSubstrateEnabled() && PermutationVector.Get<FApplyDiffuseIndirectDim>() == 2;
 		}
 
 		return IsFeatureLevelSupported(Parameters.Platform, ERHIFeatureLevel::SM5);
@@ -456,11 +455,6 @@ void FDeferredShadingSceneRenderer::CommitIndirectLightingState()
 			DiffuseIndirectMethod = EDiffuseIndirectMethod::SSGI;
 			DiffuseIndirectDenoiser = IScreenSpaceDenoiser::GetDenoiserMode(CVarDiffuseIndirectDenoiser);
 		}
-		else if (ShouldRenderRayTracingGlobalIllumination(View))
-		{
-			DiffuseIndirectMethod = EDiffuseIndirectMethod::RTGI;
-			DiffuseIndirectDenoiser = IScreenSpaceDenoiser::GetDenoiserMode(CVarDiffuseIndirectDenoiser);
-		}
 		else if (ShouldRenderPluginGlobalIllumination(View))
 		{
 			DiffuseIndirectMethod = EDiffuseIndirectMethod::Plugin;
@@ -1005,12 +999,6 @@ void FDeferredShadingSceneRenderer::RenderDiffuseIndirectAndAmbientOcclusion(
 			DenoiserInputs = ScreenSpaceRayTracing::CastStandaloneDiffuseIndirectRays(
 				GraphBuilder, CommonDiffuseParameters, PrevSceneColorMip, View);
 		}
-		else if (ViewPipelineState.DiffuseIndirectMethod == EDiffuseIndirectMethod::RTGI)
-		{
-			// TODO: Refactor under the HybridIndirectLighting standard API.
-			// TODO: hybrid SSGI / RTGI
-			RenderRayTracingGlobalIllumination(GraphBuilder, SceneTextureParameters, View, /* out */ &RayTracingConfig, /* out */ &DenoiserInputs);
-		}
 		else if (ViewPipelineState.DiffuseIndirectMethod == EDiffuseIndirectMethod::Lumen)
 		{
 			check(ViewPipelineState.DiffuseIndirectDenoiser == IScreenSpaceDenoiser::EMode::Disabled);
@@ -1142,19 +1130,7 @@ void FDeferredShadingSceneRenderer::RenderDiffuseIndirectAndAmbientOcclusion(
 				DenoiserToUse->GetDebugName(),
 				View.ViewRect.Width(), View.ViewRect.Height());
 
-			if (ViewPipelineState.DiffuseIndirectMethod == EDiffuseIndirectMethod::RTGI)
-			{
-				DenoiserOutputs = DenoiserToUse->DenoiseDiffuseIndirect(
-					GraphBuilder,
-					View,
-					&View.PrevViewInfo,
-					SceneTextureParameters,
-					DenoiserInputs,
-					RayTracingConfig);
-
-				AmbientOcclusionMask = DenoiserOutputs.Textures[1];
-			}
-			else if (ViewPipelineState.DiffuseIndirectMethod == EDiffuseIndirectMethod::SSGI)
+			if (ViewPipelineState.DiffuseIndirectMethod == EDiffuseIndirectMethod::SSGI)
 			{
 				DenoiserOutputs = DenoiserToUse->DenoiseScreenSpaceDiffuseIndirect(
 					GraphBuilder,
@@ -1349,14 +1325,9 @@ void FDeferredShadingSceneRenderer::RenderDiffuseIndirectAndAmbientOcclusion(
 
 			if (DenoiserOutputs.Textures[0])
 			{
-				if (ViewPipelineState.DiffuseIndirectMethod == EDiffuseIndirectMethod::RTGI)
+				if (ViewPipelineState.DiffuseIndirectMethod == EDiffuseIndirectMethod::Lumen)
 				{
 					PermutationVector.Set<FDiffuseIndirectCompositePS::FApplyDiffuseIndirectDim>(2);
-					DiffuseIndirectSampling = TEXT("RTGI");
-				}
-				else if (ViewPipelineState.DiffuseIndirectMethod == EDiffuseIndirectMethod::Lumen)
-				{
-					PermutationVector.Set<FDiffuseIndirectCompositePS::FApplyDiffuseIndirectDim>(3);
 					PermutationVector.Set<FDiffuseIndirectCompositePS::FScreenBentNormal>(ScreenBentNormalParameters.UseShortRangeAO != 0);
 					if (Substrate::IsSubstrateEnabled() && TileType != ESubstrateTileType::ECount)
 					{
diff --git a/Engine/Source/Runtime/Renderer/Private/RayTracing/RayTracingGlobalIllumination.cpp b/Engine/Source/Runtime/Renderer/Private/RayTracing/RayTracingGlobalIllumination.cpp
deleted file mode 100644
diff --git a/Engine/Source/Runtime/Renderer/Private/RayTracing/RaytracingOptions.h b/Engine/Source/Runtime/Renderer/Private/RayTracing/RaytracingOptions.h
--- a/Engine/Source/Runtime/Renderer/Private/RayTracing/RaytracingOptions.h
+++ b/Engine/Source/Runtime/Renderer/Private/RayTracing/RaytracingOptions.h
@@ -83,7 +83,6 @@ extern bool ShouldRenderRayTracingSkyLight(const FSkyLightSceneProxy* SkyLightSc
 extern bool ShouldRenderRayTracingAmbientOcclusion(const FViewInfo& View);
 extern bool ShouldRenderRayTracingReflections(const FViewInfo& View);
 extern bool ShouldRenderRayTracingReflectionsWater(const FViewInfo& View);
-extern bool ShouldRenderRayTracingGlobalIllumination(const FViewInfo& View);
 extern bool ShouldRenderRayTracingTranslucency(const FViewInfo& View);
 extern bool ShouldRenderRayTracingShadows();
 extern bool ShouldRenderRayTracingShadowsForLight(const FLightSceneProxy& LightProxy);
@@ -136,11 +135,6 @@ FORCEINLINE bool ShouldRenderRayTracingReflectionsWater(const FViewInfo& View)
 	return false;
 }
 
-FORCEINLINE bool ShouldRenderRayTracingGlobalIllumination(const FViewInfo& View)
-{
-	return false;
-}
-
 FORCEINLINE bool ShouldRenderRayTracingTranslucency(const FViewInfo& View)
 {
 	return false;
diff --git a/Engine/Source/Runtime/Renderer/Private/RayTracing/RaytracingSkylight.cpp b/Engine/Source/Runtime/Renderer/Private/RayTracing/RaytracingSkylight.cpp
--- a/Engine/Source/Runtime/Renderer/Private/RayTracing/RaytracingSkylight.cpp
+++ b/Engine/Source/Runtime/Renderer/Private/RayTracing/RaytracingSkylight.cpp
@@ -504,7 +504,7 @@ void FDeferredShadingSceneRenderer::RenderRayTracingSkyLight(
 		if (GRayTracingSkyLightDenoiser != 0)
 		{
 			const IScreenSpaceDenoiser* DefaultDenoiser = IScreenSpaceDenoiser::GetDefaultDenoiser();
-			const IScreenSpaceDenoiser* DenoiserToUse = DefaultDenoiser;// GRayTracingGlobalIlluminationDenoiser == 1 ? DefaultDenoiser : GScreenSpaceDenoiser;
+			const IScreenSpaceDenoiser* DenoiserToUse = DefaultDenoiser;
 
 			IScreenSpaceDenoiser::FDiffuseIndirectInputs DenoiserInputs;
 			DenoiserInputs.Color = OutSkyLightTexture;
diff --git a/Engine/Source/Runtime/Renderer/Private/RendererScene.cpp b/Engine/Source/Runtime/Renderer/Private/RendererScene.cpp
--- a/Engine/Source/Runtime/Renderer/Private/RendererScene.cpp
+++ b/Engine/Source/Runtime/Renderer/Private/RendererScene.cpp
@@ -393,11 +393,6 @@ FSceneViewState::FSceneViewState(ERHIFeatureLevel::Type FeatureLevel, FSceneView
 	PreExposure = 1.f;
 	bUpdateLastExposure = false;
 
-#if RHI_RAYTRACING
-	GatherPointsBuffer = nullptr;
-	GatherPointsResolution = FIntVector(0, 0, 0);
-#endif
-
 	bLumenSceneDataAdded = false;
 	LumenSurfaceCacheResolution = 1.0f;
 
@@ -1115,7 +1110,6 @@ uint64 FSceneViewState::GetGPUSizeBytes(bool bLogSizes) const
 	TotalSize += GetRenderTargetGPUSizeBytes(ImaginaryReflectionDepthZ, bLogSizes);
 	TotalSize += GetRenderTargetGPUSizeBytes(ImaginaryReflectionVelocity, bLogSizes);
 	TotalSize += GetBufferGPUSizeBytes(SkyLightVisibilityRaysBuffer, bLogSizes);
-	TotalSize += GetBufferGPUSizeBytes(GatherPointsBuffer, bLogSizes);
 #endif
 	TotalSize += GetRenderTargetGPUSizeBytes(LightScatteringHistory, bLogSizes);
 	TotalSize += GetRenderTargetGPUSizeBytes(PrevLightScatteringConservativeDepthTexture, bLogSizes);
diff --git a/Engine/Source/Runtime/Renderer/Private/ScenePrivate.h b/Engine/Source/Runtime/Renderer/Private/ScenePrivate.h
--- a/Engine/Source/Runtime/Renderer/Private/ScenePrivate.h
+++ b/Engine/Source/Runtime/Renderer/Private/ScenePrivate.h
@@ -1094,13 +1094,6 @@ public:
 	TRefCountPtr<FRDGPooledBuffer> SkyLightVisibilityRaysBuffer;
 	FIntVector SkyLightVisibilityRaysDimensions;
 
-	// Ray Traced Global Illumination Gather Point Data
-	TRefCountPtr<FRDGPooledBuffer> GatherPointsBuffer;
-	FIntVector GatherPointsResolution;
-	// todo: shared definition for maximum gather points per-pixel (32)
-	TStaticArray<FMatrix, 32> GatherPointsViewHistory;
-	uint32 GatherPointsCount;
-
 	// Last valid RTPSO is saved, so it could be used as fallback in future frames if background PSO compilation is enabled.
 	// This RTPSO can be used only if the only difference from previous PSO is the material hit shaders.
 	FRayTracingPipelineStateSignature LastRayTracingMaterialPipelineSignature;
