Stochastic Shadows - fixed shading artifacs on screen borders due to out of bounds writes.

diff --git a/Engine/Shaders/Private/StochasticShadows/StochasticShadows.ush b/Engine/Shaders/Private/StochasticShadows/StochasticShadows.ush
--- a/Engine/Shaders/Private/StochasticShadows/StochasticShadows.ush
+++ b/Engine/Shaders/Private/StochasticShadows/StochasticShadows.ush
@@ -368,9 +368,10 @@ FShadingTile UnpackShadingTile(uint PackedShadingTile)
 	return ShadingTile;
 }
 
-uint2 ShadingTileGridStride;
+uint2 ShadingTileGridSize;
+uint MaxShadingTilesPerGridCell;
 
 uint GetShadingTileGridOffset(uint2 ShadingTileCoord, uint IndexInCell)
 {
-	return ShadingTileCoord.x * ShadingTileGridStride.x + ShadingTileCoord.y * ShadingTileGridStride.y + IndexInCell;
+	return (ShadingTileCoord.x + ShadingTileCoord.y * ShadingTileGridSize.x) * MaxShadingTilesPerGridCell + IndexInCell;
 }
\ No newline at end of file
diff --git a/Engine/Shaders/Private/StochasticShadows/StochasticShadowsSpatial.usf b/Engine/Shaders/Private/StochasticShadows/StochasticShadowsSpatial.usf
--- a/Engine/Shaders/Private/StochasticShadows/StochasticShadowsSpatial.usf
+++ b/Engine/Shaders/Private/StochasticShadows/StochasticShadowsSpatial.usf
@@ -252,17 +252,22 @@ void ShadowMaskSpatialPassCS(
 		{	
 			const uint2 ShadingTileCoord = (BaseShadowMaskScreenCoord * SHADOW_MASK_DOWNSAMPLE_FACTOR) / SHADOW_MASK_TILE_SIZE + ShadingTileOffset;
 
-			uint ShadingTileIndex = 0;
-			InterlockedAdd(RWShadingTileAllocator[0], 1u, ShadingTileIndex);
-
-			if (ShadingTileIndex < MaxShadingTiles)
+			if (all(ShadingTileCoord < ShadingTileGridSize))
 			{
-				uint NumShadingTilesInCell = 0;
-				InterlockedAdd(RWShadingTileGridAllocator[ShadingTileCoord], 1u, NumShadingTilesInCell); // #sdl_todo: handle cell overflow
+				uint ShadingTileIndex = 0;
+				InterlockedAdd(RWShadingTileAllocator[0], 1u, ShadingTileIndex);
 
-				RWShadingTileGrid[GetShadingTileGridOffset(ShadingTileCoord, NumShadingTilesInCell)] = PackShadingTile(LightSceneId, ShadingTileIndex);
+				if (ShadingTileIndex < MaxShadingTiles)
+				{
+					uint ShadingTileIndexInCell = 0;
+					InterlockedAdd(RWShadingTileGridAllocator[ShadingTileCoord], 1u, ShadingTileIndexInCell);
 
-				SharedShadingTileIndex[ShadingTileOffset.x][ShadingTileOffset.y] = ShadingTileIndex;
+					if (ShadingTileIndexInCell < MaxShadingTilesPerGridCell)
+					{
+						RWShadingTileGrid[GetShadingTileGridOffset(ShadingTileCoord, ShadingTileIndexInCell)] = PackShadingTile(LightSceneId, ShadingTileIndex);
+						SharedShadingTileIndex[ShadingTileOffset.x][ShadingTileOffset.y] = ShadingTileIndex;
+					}
+				}
 			}
 		}
 	}
diff --git a/Engine/Source/Runtime/Renderer/Private/StochasticShadows/StochasticShadows.cpp b/Engine/Source/Runtime/Renderer/Private/StochasticShadows/StochasticShadows.cpp
--- a/Engine/Source/Runtime/Renderer/Private/StochasticShadows/StochasticShadows.cpp
+++ b/Engine/Source/Runtime/Renderer/Private/StochasticShadows/StochasticShadows.cpp
@@ -274,9 +274,10 @@ BEGIN_SHADER_PARAMETER_STRUCT(FStochasticShadowsParameters, )
 	SHADER_PARAMETER(uint32, StochasticShadowsStateFrameIndex)
 	SHADER_PARAMETER(uint32, MaxShadowMaskTiles)
 	SHADER_PARAMETER(uint32, MaxShadingTiles)
+	SHADER_PARAMETER(uint32, MaxShadingTilesPerGridCell)
 	SHADER_PARAMETER(uint32, ShadowMaskHashTableIndexWrapMask)
 	SHADER_PARAMETER(FIntPoint, ShadowMaskPageTablePerLightSize)
-	SHADER_PARAMETER(FIntPoint, ShadingTileGridStride)
+	SHADER_PARAMETER(FIntPoint, ShadingTileGridSize)
 	SHADER_PARAMETER(FVector2f, DownsampledBufferInvSize)
 	SHADER_PARAMETER(float, MinLightSampleWeight)
 	SHADER_PARAMETER(int32, DebugMode)
@@ -1011,10 +1012,10 @@ void FDeferredShadingSceneRenderer::RenderStochasticShadows(FRDGBuilder& GraphBu
 		StochasticShadowsParameters.DownsampledSceneDepth = DownsampledSceneDepth;
 		StochasticShadowsParameters.MaxShadowMaskTiles = MaxShadowMaskTiles;
 		StochasticShadowsParameters.MaxShadingTiles = (ShadingTileAtlasSize.X * ShadingTileAtlasSize.Y) / (StochasticShadows::ShadowMaskTileSize * StochasticShadows::ShadowMaskTileSize);
+		StochasticShadowsParameters.MaxShadingTilesPerGridCell = StochasticShadows::MaxShadingTilesPerGridCell;
+		StochasticShadowsParameters.ShadingTileGridSize = ShadingTileGridSize;
 		StochasticShadowsParameters.ShadowMaskHashTableIndexWrapMask = ShadowMaskHashTableSize - 1;
 		StochasticShadowsParameters.ShadowMaskPageTablePerLightSize = ShadowMaskPageTablePerLightSize;
-		StochasticShadowsParameters.ShadingTileGridStride.X = StochasticShadows::MaxShadingTilesPerGridCell;
-		StochasticShadowsParameters.ShadingTileGridStride.Y = ShadingTileGridSize.X * StochasticShadows::MaxShadingTilesPerGridCell;
 		StochasticShadowsParameters.DownsampledBufferInvSize = FVector2f(1.0f) / DownsampledBufferSize;
 		StochasticShadowsParameters.MinLightSampleWeight = CVarStochasticShadowsMinSampleWeight.GetValueOnRenderThread();
 		StochasticShadowsParameters.DebugMode = CVarStochasticShadowsDebug.GetValueOnRenderThread();
