OIT sort triangle.

* Fix debug visualization and add more information to track errors
* Fix missing base vertex index offset, and add min/max vertex index clamp to avoid out of bound vertex position read.
* Fix remove ForceDXC flag which cause invalid binding data on dx11. Instead remove ShaderPrint issue on hlslcc platform by removing debug permutation on these platforms

diff --git a/Engine/Shaders/Private/OIT/OITSorting.usf b/Engine/Shaders/Private/OIT/OITSorting.usf
--- a/Engine/Shaders/Private/OIT/OITSorting.usf
+++ b/Engine/Shaders/Private/OIT/OITSorting.usf
@@ -2,6 +2,9 @@
 
 #pragma once
 
+// RECOMPILE HASH
+#pragma message("UESHADERMETADATA_VERSION 78C3EBBB-59EA-459D-A2D4-67170905D7FC")
+
 #include "../Common.ush"
 
 ///////////////////////////////////////////////////////////////////////////////////////////////////
@@ -61,6 +64,10 @@ bool IsValid(FTriIndex InTri, uint InNumIndices)
 	return InTri.I0 < InNumIndices && InTri.I1 < InNumIndices && InTri.I2 < InNumIndices;
 }
 
+#if SORTING_SLICE_COUNT > 256
+#error Slice index is encoded onto 8bits and no longer fit into stored data
+#endif
+
 FPackedSliceAndOffset PackSliceAndOffset(uint InSliceIndex, uint InLocalOffset)
 {
 #if ENCODING_TYPE
@@ -90,17 +97,27 @@ struct FOITDebugData
 	uint PrimitiveCount;
 	uint Type;
 	uint SizeInBytes;
+
+	uint FirstIndex;
+	uint BaseVertexIndex;
+	uint MinVertexIndex;
+	uint MaxVertexIndex;
 };
 
 void AddDebugInfo(RWBuffer<uint> InDebugData, FOITDebugData In)
 {
-	uint Index = 1;
+	uint Index = 0;
 	InterlockedAdd(InDebugData[0], 1, Index);
 
-	const uint Index3 = 1 + Index * 3;		
+	const uint Index3 = 1 + Index * 7u;
 	InDebugData[Index3 + 0] = In.PrimitiveCount;
 	InDebugData[Index3 + 1] = In.Type;
 	InDebugData[Index3 + 2] = In.SizeInBytes;
+
+	InDebugData[Index3 + 3] = In.FirstIndex;
+	InDebugData[Index3 + 4] = In.BaseVertexIndex;
+	InDebugData[Index3 + 5] = In.MinVertexIndex;
+	InDebugData[Index3 + 6] = In.MaxVertexIndex;
 }
 
 uint GetDebugInfoCount(Buffer<uint> InDebugData)
@@ -113,11 +130,16 @@ FOITDebugData GetDebugInfo(Buffer<uint> InDebugData, uint Index)
 	FOITDebugData Out;
 	if (Index < InDebugData[0])
 	{
-		const uint Index3 = 1 + Index * 3;
+		const uint Index3 = 1 + Index * 7u;
 		Out.Index			= Index;
 		Out.PrimitiveCount	= InDebugData[Index3 + 0];
 		Out.Type			= InDebugData[Index3 + 1];
-		Out.SizeInBytes		= InDebugData[Index3 + 2];
+		Out.SizeInBytes	= InDebugData[Index3 + 2];
+
+		Out.FirstIndex		= InDebugData[Index3 + 3];
+		Out.BaseVertexIndex	= InDebugData[Index3 + 4];
+		Out.MinVertexIndex	= InDebugData[Index3 + 5];
+		Out.MaxVertexIndex	= InDebugData[Index3 + 6];
 	}
 
 	return Out;
@@ -136,13 +158,15 @@ float3 ViewBound_Max;
 #endif
 
 float4x4 LocalToWorld; // GPU scene buffer instead? How?
-float4x4 WorldToView;
 float4x4 LocalToView;
 
 uint SourcePrimitiveType;
 uint NumPrimitives;
 uint NumIndices;
 uint SourceFirstIndex;
+uint SourceBaseVertexIndex;
+uint SourceMinVertexIndex;
+uint SourceMaxVertexIndex;
 uint SortType;
 uint SortedIndexBufferSizeInByte;
 
@@ -181,9 +205,9 @@ void MainCS(uint3 DispatchThreadId : SV_DispatchThreadID, uint GroupThreadId : S
 	if (bIsValid)
 	{
 		const FTriIndex Tri = GetTriangleIndices(PrimitiveId, SourceFirstIndex, SourcePrimitiveType);
-		const uint I0 = 3 * IndexBuffer[Tri.I0];
-		const uint I1 = 3 * IndexBuffer[Tri.I1];
-		const uint I2 = 3 * IndexBuffer[Tri.I2];
+		const uint I0 = 3 * clamp(IndexBuffer[Tri.I0] + SourceBaseVertexIndex, SourceMinVertexIndex, SourceMaxVertexIndex);
+		const uint I1 = 3 * clamp(IndexBuffer[Tri.I1] + SourceBaseVertexIndex, SourceMinVertexIndex, SourceMaxVertexIndex);
+		const uint I2 = 3 * clamp(IndexBuffer[Tri.I2] + SourceBaseVertexIndex, SourceMinVertexIndex, SourceMaxVertexIndex);
 
 		// Position buffer store data as a float, not float3
 		const float3 P0 = float3(PositionBuffer[I0], PositionBuffer[I0+1], PositionBuffer[I0+2]);
@@ -200,12 +224,13 @@ void MainCS(uint3 DispatchThreadId : SV_DispatchThreadID, uint GroupThreadId : S
 
 		#if PERMUTATION_DEBUG
 		{
+			FShaderPrintContext Ctx = InitShaderPrintContext();
 			const float4 TriColor = float4(ColorMapViridis(SliceIndex / float(SORTING_SLICE_COUNT - 1)), 1);
 
 			const float3 WP0 = mul(float4(P0, 1), LocalToWorld).xyz;
 			const float3 WP1 = mul(float4(P1, 1), LocalToWorld).xyz;
 			const float3 WP2 = mul(float4(P2, 1), LocalToWorld).xyz;
-			AddLineTriangleWS(WP0, WP1, WP2, TriColor);
+			AddLineTriangleWS(Ctx, WP0, WP1, WP2, TriColor);
 		}
 		#endif
 
@@ -219,9 +244,15 @@ void MainCS(uint3 DispatchThreadId : SV_DispatchThreadID, uint GroupThreadId : S
 		AddOBBWS(ViewBound_Min, ViewBound_Max, ColorCyan, ViewToWorld);
 
 		FOITDebugData DebugData;
-		DebugData.PrimitiveCount= NumPrimitives;
-		DebugData.Type			= SourcePrimitiveType;
-		DebugData.SizeInBytes	= SortedIndexBufferSizeInByte;
+		DebugData.PrimitiveCount	= NumPrimitives;
+		DebugData.Type				= SourcePrimitiveType;
+		DebugData.SizeInBytes		= SortedIndexBufferSizeInByte;
+
+		DebugData.FirstIndex		= SourceFirstIndex;
+		DebugData.BaseVertexIndex	= SourceBaseVertexIndex;
+		DebugData.MinVertexIndex	= SourceMinVertexIndex;
+		DebugData.MaxVertexIndex	= SourceMaxVertexIndex;
+
 		AddDebugInfo(OutDebugData, DebugData);
 	}
 	#endif
@@ -364,18 +395,18 @@ void PrintSize(inout FShaderPrintContext Context, uint SizeInBytes, FFontColor C
 {
 	if (SizeInBytes < 1000)
 	{
-		Print(Context, SizeInBytes / 1000.f, Color);
-		Print(Context, TEXT("bytes"), Color);
+		Print(Context, SizeInBytes, Color, 5, 2);
+		Print(Context, TEXT("   bytes "), Color);
 	}
 	else if (SizeInBytes < 1000000)
 	{
-		Print(Context, SizeInBytes / 1000.f, Color);
-		Print(Context, TEXT("Kbytes"), Color);
+		Print(Context, SizeInBytes / 1000.f, Color, 5, 2);
+		Print(Context, TEXT("  Kbytes "), Color);
 	}
 	else
 	{
-		Print(Context, SizeInBytes / 1000.f, Color);
-		Print(Context, TEXT("Mbytes"), Color);
+		Print(Context, SizeInBytes / 1000000.f, Color, 5, 2);
+		Print(Context, TEXT("  Mbytes "), Color);
 	}
 }
 
@@ -389,37 +420,22 @@ void MainCS(uint DispatchThreadId : SV_DispatchThreadID)
 	// Visible
 	{
 		Print(Context, TEXT("Visible"), FontOrange); Newline(Context);
+		Print(Context, TEXT("#Sort slices   : "), FontWhite); Print(Context, uint(SORTING_SLICE_COUNT), FontSilver); Newline(Context);
 		Print(Context, TEXT("Instances      : "), FontWhite); Print(Context, VisibleInstances, FontSilver); Newline(Context);
 		Print(Context, TEXT("Primitives     : "), FontWhite); Print(Context, VisiblePrimitives, FontSilver); Newline(Context);
 		Print(Context, TEXT("Size           : "), FontYellow); PrintSize(Context, VisibleIndexSizeInBytes, FontYellow); Newline(Context);
 	}
-
-	Newline(Context);
-
-	// Allocated
-	{
-		Print(Context, TEXT("Allocated"), FontOrange); Newline(Context);
-		Print(Context, TEXT("Buffers        : "), FontWhite); Print(Context, AllocatedBuffers, FontSilver); Newline(Context);
-		Print(Context, TEXT("Size           : "), FontYellow); PrintSize(Context, AllocatedIndexSizeInBytes, FontYellow); Newline(Context);
-	}
-
 	Newline(Context);
 
-	// Unused
+	// Memory
 	{
-		Print(Context, TEXT("Unused"), FontOrange); Newline(Context);
-		Print(Context, TEXT("Buffers        : "), FontWhite); Print(Context, UnusedBuffers, FontSilver); Newline(Context);
-		Print(Context, TEXT("Size           : "), FontYellow); PrintSize(Context, UnusedIndexSizeInBytes, FontYellow); Newline(Context);
+		Print(Context, TEXT("Memory"), FontOrange); Newline(Context);
+		Print(Context, TEXT("Alloc. Buffers : "), FontWhite); Print(Context, AllocatedBuffers, FontSilver, 3, 0); Print(Context, TEXT("/"), FontSilver); Print(Context, TotalEntries, FontSilver, 3, 0); Newline(Context);
+		Print(Context, TEXT("Alloc. Size    : "), FontWhite); PrintSize(Context, AllocatedIndexSizeInBytes, FontSilver); Newline(Context);
+		Print(Context, TEXT("Unused Buffers : "), FontWhite); Print(Context, UnusedBuffers, FontSilver); Newline(Context);
+		Print(Context, TEXT("Unused Size    : "), FontWhite); PrintSize(Context, UnusedIndexSizeInBytes, FontSilver); Newline(Context);
+		Print(Context, TEXT("Total Size     : "), FontYellow); PrintSize(Context, AllocatedIndexSizeInBytes + UnusedIndexSizeInBytes, FontYellow); Newline(Context);
 	}
-
-	Newline(Context);
-
-	{
-		Print(Context, TEXT("Total"), FontOrange); Newline(Context);
-		Print(Context, TEXT("Entries        : "), FontWhite); Print(Context, TotalEntries, FontWhite); Newline(Context);
-		Print(Context, TEXT("Size           : "), FontOrange); PrintSize(Context, AllocatedIndexSizeInBytes + UnusedIndexSizeInBytes, FontOrange); Newline(Context);
-	}
-
 	Newline(Context);
 
 	// 1: NumPrimitive | PrimitiveType | 0.4 Mb
@@ -430,7 +446,13 @@ void MainCS(uint DispatchThreadId : SV_DispatchThreadID)
 		Print(Context, TEXT("Index       "));
 		Print(Context, TEXT("Primitive   ")); 
 		Print(Context, TEXT("Type  "));
-		Print(Context, TEXT("Size"));
+		Print(Context, TEXT("Size          "));
+
+		Print(Context, TEXT("FirstIdx    "));
+		Print(Context, TEXT("BaseIdx     "));
+		Print(Context, TEXT("MinVtxIdx   "));
+		Print(Context, TEXT("MaxVtxIdx   "));
+		Print(Context, TEXT("#Vertex     "));
 		Newline(Context);
 
 		for (uint Index = 0; Index < VisibleInstances; ++Index)
@@ -446,6 +468,12 @@ void MainCS(uint DispatchThreadId : SV_DispatchThreadID)
 
 			PrintSize(Context, Data.SizeInBytes, FontYellow);
 
+			Print(Context, Data.FirstIndex, FontSilver);
+			Print(Context, Data.BaseVertexIndex, FontSilver);
+			Print(Context, Data.MinVertexIndex, FontSilver);
+			Print(Context, Data.MaxVertexIndex, FontSilver);
+			Print(Context, Data.MaxVertexIndex - Data.MinVertexIndex + 1, FontSilver);
+
 			Newline(Context);
 		}
 	}
diff --git a/Engine/Source/Runtime/Renderer/Private/OIT/OIT.cpp b/Engine/Source/Runtime/Renderer/Private/OIT/OIT.cpp
--- a/Engine/Source/Runtime/Renderer/Private/OIT/OIT.cpp
+++ b/Engine/Source/Runtime/Renderer/Private/OIT/OIT.cpp
@@ -319,10 +319,10 @@ static void TrimSortedIndexBuffers(TArray<FSortedIndexBuffer*>& FreeBuffers, TQu
 	}
 }
 
-FSortedTriangleData FOITSceneData::Allocate(FRHICommandListBase& RHICmdList, const FIndexBuffer* InSource, EPrimitiveType PrimitiveType, uint32 InFirstIndex, uint32 InNumPrimitives)
+FSortedTriangleData FOITSceneData::Allocate(FRHICommandListBase& RHICmdList, EPrimitiveType InPrimitiveType, const FMeshBatchElement& InMeshElement)
 {
-	check(InSource && InSource->IndexBufferRHI);
-	check(PrimitiveType == PT_TriangleList || PrimitiveType == PT_TriangleStrip);
+	check(InMeshElement.IndexBuffer && InMeshElement.IndexBuffer->IndexBufferRHI);
+	check(InPrimitiveType == PT_TriangleList || InPrimitiveType == PT_TriangleStrip);
 
 	// Find a free slot, or create a new one
 	FSortedTriangleData* Out = nullptr;
@@ -339,7 +339,7 @@ FSortedTriangleData FOITSceneData::Allocate(FRHICommandListBase& RHICmdList, con
 	}
 
 	// Linear scan if there are some free resource which are large enough
-	const uint32 NumIndices = InNumPrimitives * 3; // Sorted index always has triangle list topology
+	const uint32 NumIndices = InMeshElement.NumPrimitives * 3; // Sorted index always has triangle list topology
 	FSortedIndexBuffer* OITIndexBuffer = nullptr;
 	if (CVarOIT_SortedTriangles_Pool.GetValueOnRenderThread() > 0)
 	{
@@ -360,16 +360,19 @@ FSortedTriangleData FOITSceneData::Allocate(FRHICommandListBase& RHICmdList, con
 	// Otherwise create a new one
 	if (OITIndexBuffer == nullptr)
 	{
-		OITIndexBuffer = new FSortedIndexBuffer(FreeSlot, InSource->IndexBufferRHI, NumIndices, TEXT("OIT::SortedIndexBuffer"));
+		OITIndexBuffer = new FSortedIndexBuffer(FreeSlot, InMeshElement.IndexBuffer->IndexBufferRHI, NumIndices, TEXT("OIT::SortedIndexBuffer"));
 		OITIndexBuffer->InitResource(RHICmdList);	
 	}
-	Out->NumPrimitives = InNumPrimitives;
+	Out->NumPrimitives = InMeshElement.NumPrimitives;
 	Out->NumIndices = NumIndices;
-	Out->SourceFirstIndex = InFirstIndex;
+	Out->SourceFirstIndex = InMeshElement.FirstIndex;
+	Out->SourceBaseVertexIndex = InMeshElement.BaseVertexIndex;
+	Out->SourceMinVertexIndex = InMeshElement.MinVertexIndex;
+	Out->SourceMaxVertexIndex = InMeshElement.MaxVertexIndex;
 	Out->SortedFirstIndex = 0u;
-	Out->SourcePrimitiveType = PrimitiveType;
+	Out->SourcePrimitiveType = InPrimitiveType;
 	Out->SortedPrimitiveType = PT_TriangleList;
-	Out->SourceIndexBuffer = InSource;
+	Out->SourceIndexBuffer = InMeshElement.IndexBuffer;
 	Out->SortedIndexBuffer = OITIndexBuffer;
 	Out->SortedIndexUAV = OITIndexBuffer->SortedIndexUAV;
 	Out->SourceIndexSRV = OITIndexBuffer->SourceIndexSRV;
@@ -418,45 +421,53 @@ class FOITSortTriangleIndex_ScanCS : public FGlobalShader
 	class FDebug : SHADER_PERMUTATION_BOOL("PERMUTATION_DEBUG");
 	using FPermutationDomain = TShaderPermutationDomain<FDebug>;
 
-	BEGIN_SHADER_PARAMETER_STRUCT(FParameters, )
-
-		// For Debug
-		SHADER_PARAMETER_STRUCT_INCLUDE(ShaderPrint::FShaderParameters, ShaderPrintParameters)
-		SHADER_PARAMETER(FMatrix44f, ViewToWorld)
-		SHADER_PARAMETER(FVector3f, WorldBound_Min)
-		SHADER_PARAMETER(FVector3f, WorldBound_Max)
-		SHADER_PARAMETER(FVector3f, ViewBound_Min)
-		SHADER_PARAMETER(FVector3f, ViewBound_Max)
-
-		SHADER_PARAMETER(FMatrix44f, LocalToWorld)
-		SHADER_PARAMETER(FMatrix44f, WorldToView)
-		SHADER_PARAMETER(FMatrix44f, LocalToView)
-
-		SHADER_PARAMETER(uint32, SourcePrimitiveType)
-		SHADER_PARAMETER(uint32, NumPrimitives)
-		SHADER_PARAMETER(uint32, NumIndices)
-		SHADER_PARAMETER(uint32, SourceFirstIndex)
-		SHADER_PARAMETER(uint32, SortType)
-		SHADER_PARAMETER(uint32, SortedIndexBufferSizeInByte)
+	BEGIN_SHADER_PARAMETER_STRUCT(FParameters,)
+
+	// For Debug
+	SHADER_PARAMETER_STRUCT_INCLUDE(ShaderPrint::FShaderParameters, ShaderPrintParameters)
+	SHADER_PARAMETER(FMatrix44f, ViewToWorld)
+	SHADER_PARAMETER(FVector3f, WorldBound_Min)
+	SHADER_PARAMETER(FVector3f, WorldBound_Max)
+	SHADER_PARAMETER(FVector3f, ViewBound_Min)
+	SHADER_PARAMETER(FVector3f, ViewBound_Max)
+
+	SHADER_PARAMETER(FMatrix44f, LocalToWorld)
+	SHADER_PARAMETER(FMatrix44f, LocalToView)
+
+	SHADER_PARAMETER(uint32, SourcePrimitiveType)
+	SHADER_PARAMETER(uint32, SourceFirstIndex)
+	SHADER_PARAMETER(uint32, SourceBaseVertexIndex)
+	SHADER_PARAMETER(uint32, SourceMinVertexIndex)
+	SHADER_PARAMETER(uint32, SourceMaxVertexIndex)
+
+	SHADER_PARAMETER(uint32, NumPrimitives)
+	SHADER_PARAMETER(uint32, NumIndices)
+	SHADER_PARAMETER(uint32, SortType)
+	SHADER_PARAMETER(uint32, SortedIndexBufferSizeInByte)
 		
-		SHADER_PARAMETER(float, ViewBoundMinZ)
-		SHADER_PARAMETER(float, ViewBoundMaxZ)
+	SHADER_PARAMETER(float, ViewBoundMinZ)
+	SHADER_PARAMETER(float, ViewBoundMaxZ)
 
-		SHADER_PARAMETER_SRV(Buffer<float>, PositionBuffer)
-		SHADER_PARAMETER_SRV(Buffer<uint>, IndexBuffer)
-		SHADER_PARAMETER_UAV(Buffer<uint>, OutIndexBuffer)
-		SHADER_PARAMETER_RDG_BUFFER_UAV(RWBuffer<uint2>, OutSliceCounterBuffer)
-		SHADER_PARAMETER_RDG_BUFFER_UAV(RWBuffer<uint>, OutPrimitiveSliceBuffer)
-		SHADER_PARAMETER_RDG_BUFFER_UAV(RWBuffer<uint>, OutDebugData)
+	SHADER_PARAMETER_SRV(Buffer < float > , PositionBuffer)
+	SHADER_PARAMETER_SRV(Buffer < uint > , IndexBuffer)
+	SHADER_PARAMETER_UAV(Buffer < uint > , OutIndexBuffer)
+	SHADER_PARAMETER_RDG_BUFFER_UAV(RWBuffer < uint2 > , OutSliceCounterBuffer)
+	SHADER_PARAMETER_RDG_BUFFER_UAV(RWBuffer < uint > , OutPrimitiveSliceBuffer)
+	SHADER_PARAMETER_RDG_BUFFER_UAV(RWBuffer < uint > , OutDebugData)
 
 
 	END_SHADER_PARAMETER_STRUCT()
 public:
 
+	static bool SupportDebugMode(EShaderPlatform InPlatform)
+	{
+		return ShaderPrint::IsSupported(InPlatform) && !IsHlslccShaderPlatform(InPlatform);
+	}
+
 	static bool ShouldCompilePermutation(const FGlobalShaderPermutationParameters& Parameters) 
 	{ 
 		FPermutationDomain PermutationVector(Parameters.PermutationId);
-		if (PermutationVector.Get<FDebug>() == 1 && !ShaderPrint::IsSupported(Parameters.Platform))
+		if (PermutationVector.Get<FDebug>() == 1 && !SupportDebugMode(Parameters.Platform))
 		{
 			return false;
 		}
@@ -468,7 +479,6 @@ public:
 		FGlobalShader::ModifyCompilationEnvironment(Parameters, OutEnvironment);
 		OutEnvironment.SetDefine(TEXT("SHADER_SCAN"), 1);
 		OutEnvironment.SetDefine(TEXT("SORTING_SLICE_COUNT"), FSortedIndexBuffer::SliceCount);
-		OutEnvironment.CompilerFlags.Add(CFLAG_ForceDXC);
 	}
 };
 
@@ -550,7 +560,7 @@ public:
 
 	static bool ShouldCompilePermutation(const FGlobalShaderPermutationParameters& Parameters) 
 	{
-		return ShaderPrint::IsSupported(Parameters.Platform); 
+		return FOITSortTriangleIndex_ScanCS::SupportDebugMode(Parameters.Platform);
 	}
 
 	static void ModifyCompilationEnvironment(const FGlobalShaderPermutationParameters& Parameters, FShaderCompilerEnvironment& OutEnvironment)
@@ -558,7 +568,6 @@ public:
 		FGlobalShader::ModifyCompilationEnvironment(Parameters, OutEnvironment);
 		OutEnvironment.SetDefine(TEXT("SHADER_DEBUG"), 1);
 		OutEnvironment.SetDefine(TEXT("SORTING_SLICE_COUNT"), FSortedIndexBuffer::SliceCount);
-		OutEnvironment.CompilerFlags.Add(CFLAG_ForceDXC);
 	}
 };
 
@@ -632,7 +641,7 @@ static void AddOITSortTriangleIndexPass(
 	AddClearUAVPass(GraphBuilder, SliceCounterUAV, 0u);
 
 	const EShaderPlatform Platform = View.Family->GetShaderPlatform();
-	const bool bDebugEnable = DebugData.IsValid() && ShaderPrint::IsSupported(Platform) && ShaderPrint::IsValid(View.ShaderPrintData);
+	const bool bDebugEnable = DebugData.IsValid() && FOITSortTriangleIndex_ScanCS::SupportDebugMode(Platform) && ShaderPrint::IsValid(View.ShaderPrintData);
 	FRHIBuffer* SortedIndexBufferRHI = Allocation.SortedIndexBuffer->IndexBufferRHI;
 
 	// 1. Scan the primitive and assign each primitive to a slice
@@ -645,6 +654,7 @@ static void AddOITSortTriangleIndexPass(
 
 		FOITSortTriangleIndex_ScanCS::FParameters* Parameters = GraphBuilder.AllocParameters<FOITSortTriangleIndex_ScanCS::FParameters>();
 		Parameters->LocalToView				= FMatrix44f(MeshBatch.Proxy->GetLocalToWorld() * View.ViewMatrices.GetViewMatrix());
+		Parameters->LocalToWorld 			= FMatrix44f(MeshBatch.Proxy->GetLocalToWorld());
 		Parameters->SourcePrimitiveType		= Allocation.SourcePrimitiveType == PT_TriangleStrip ? 1u : 0u;
 		Parameters->NumPrimitives			= Allocation.NumPrimitives;
 		Parameters->NumIndices				= Allocation.NumIndices;
@@ -654,6 +664,9 @@ static void AddOITSortTriangleIndexPass(
 		Parameters->SortedIndexBufferSizeInByte = Allocation.SortedIndexBuffer->IndexBufferRHI->GetSize();
 		Parameters->PositionBuffer			= VertexPosition;
 		Parameters->SourceFirstIndex		= Allocation.SourceFirstIndex;
+		Parameters->SourceBaseVertexIndex	= Allocation.SourceBaseVertexIndex;
+		Parameters->SourceMinVertexIndex	= Allocation.SourceMinVertexIndex;
+		Parameters->SourceMaxVertexIndex	= Allocation.SourceMaxVertexIndex;
 		Parameters->IndexBuffer				= Allocation.SourceIndexSRV;
 		Parameters->OutIndexBuffer			= Allocation.SortedIndexUAV;
 		Parameters->OutSliceCounterBuffer	= SliceCounterUAV;
@@ -662,16 +675,12 @@ static void AddOITSortTriangleIndexPass(
 		// Debug
 		if (bDebugEnable)
 		{
-			ShaderPrint::SetEnabled(true);
-			ShaderPrint::RequestSpaceForCharacters(DebugData.VisibleInstances * 256 + 512);
-			ShaderPrint::RequestSpaceForLines(DebugData.VisiblePrimitives * 8 + DebugData.VisibleInstances * 32);
-
 			Parameters->ViewToWorld		= FMatrix44f(View.ViewMatrices.GetViewMatrix().Inverse());	 // LWC_TODO: Precision loss?
 			Parameters->WorldBound_Min	= (FVector3f)Bounds.GetBox().Min;
 			Parameters->WorldBound_Max	= (FVector3f)Bounds.GetBox().Max;
 			Parameters->ViewBound_Min	= (FVector3f)ViewBounds.GetBox().Min;
 			Parameters->ViewBound_Max	= (FVector3f)ViewBounds.GetBox().Max;
-				ShaderPrint::SetParameters(GraphBuilder, View.ShaderPrintData, Parameters->ShaderPrintParameters);
+			ShaderPrint::SetParameters(GraphBuilder, View.ShaderPrintData, Parameters->ShaderPrintParameters);
 
 			check(DebugData.Buffer);
 
@@ -852,19 +861,25 @@ namespace OIT
 		if (bDebugEnable)
 		{
 			const uint32 ValidAllocationCount = OITSceneData.Allocations.Num();
-			DebugData.Buffer = GraphBuilder.CreateBuffer(FRDGBufferDesc::CreateBufferDesc(4u, ValidAllocationCount + 1), TEXT("OIT.DebugData"));
+			DebugData.Buffer = GraphBuilder.CreateBuffer(FRDGBufferDesc::CreateBufferDesc(4u, 10*ValidAllocationCount + 1), TEXT("OIT.DebugData"));
 			AddClearUAVPass(GraphBuilder, GraphBuilder.CreateUAV(DebugData.Buffer, FOITDebugData::Format), 0u);
 
 			// Allocated/used
+			uint32 AllocatedNumPrimitives = 0;
 			for (const FSortedTriangleData& Allocated : OITSceneData.Allocations)
 			{
 				if (Allocated.IsValid())
 				{
 					++DebugData.AllocatedBuffers;
 					DebugData.AllocatedIndexSizeInBytes += Allocated.SortedIndexBuffer->IndexBufferRHI->GetSize();
+					AllocatedNumPrimitives += Allocated.NumPrimitives;
 				}
 			}
 
+			ShaderPrint::SetEnabled(true);
+			ShaderPrint::RequestSpaceForCharacters(DebugData.AllocatedBuffers * 256 + 512);
+			ShaderPrint::RequestSpaceForLines(AllocatedNumPrimitives * 8 + DebugData.AllocatedBuffers * 32);
+
 			// Unused
 			for (const FSortedIndexBuffer* FreeBuffer : OITSceneData.FreeBuffers)
 			{
diff --git a/Engine/Source/Runtime/Renderer/Private/OIT/OIT.h b/Engine/Source/Runtime/Renderer/Private/OIT/OIT.h
--- a/Engine/Source/Runtime/Renderer/Private/OIT/OIT.h
+++ b/Engine/Source/Runtime/Renderer/Private/OIT/OIT.h
@@ -42,6 +42,9 @@ struct FSortedTriangleData
 
 	uint32 SortedFirstIndex = 0;
 	uint32 SourceFirstIndex	= 0;
+	uint32 SourceBaseVertexIndex = 0;
+	uint32 SourceMinVertexIndex = 0;
+	uint32 SourceMaxVertexIndex = 0;
 	uint32 NumPrimitives = 0;
 	uint32 NumIndices = 0;
 
@@ -71,7 +74,7 @@ struct FOITData
 struct FOITSceneData
 {
 	/* Allocate sorted-triangle data for a instance */
-	FSortedTriangleData Allocate(FRHICommandListBase& RHICmdList, const FIndexBuffer* InSource, EPrimitiveType PrimitiveType, uint32 InFirstIndex, uint32 InNumPrimitives);
+	FSortedTriangleData Allocate(FRHICommandListBase& RHICmdList, EPrimitiveType PrimitiveType, const FMeshBatchElement& InMeshElement);
 
 	/* Deallocate sorted-triangle data */
 	void Deallocate(FIndexBuffer* IndexBuffer);
diff --git a/Engine/Source/Runtime/Renderer/Private/PrimitiveSceneInfo.cpp b/Engine/Source/Runtime/Renderer/Private/PrimitiveSceneInfo.cpp
--- a/Engine/Source/Runtime/Renderer/Private/PrimitiveSceneInfo.cpp
+++ b/Engine/Source/Runtime/Renderer/Private/PrimitiveSceneInfo.cpp
@@ -1293,7 +1293,7 @@ void FPrimitiveSceneInfo::AddStaticMeshes(FScene* Scene, TArrayView<FPrimitiveSc
 
 				if (bAllocateSortedTriangles && OIT::IsCompatible(Mesh, FeatureLevel))
 				{
-					FSortedTriangleData Allocation = Scene->OITSceneData.Allocate(RHICmdList, Mesh.Elements[0].IndexBuffer, EPrimitiveType(Mesh.Type), Mesh.Elements[0].FirstIndex, Mesh.Elements[0].NumPrimitives);
+					FSortedTriangleData Allocation = Scene->OITSceneData.Allocate(RHICmdList, EPrimitiveType(Mesh.Type), Mesh.Elements[0]);
 					OIT::ConvertSortedIndexToDynamicIndex(&Allocation, &Mesh.Elements[0].DynamicIndexBuffer);
 				}
 			}
