Fix shader compilation error on dx11 with Substrate with RoughRefraction.

diff --git a/Engine/Source/Runtime/Renderer/Private/IndirectLightRendering.cpp b/Engine/Source/Runtime/Renderer/Private/IndirectLightRendering.cpp
--- a/Engine/Source/Runtime/Renderer/Private/IndirectLightRendering.cpp
+++ b/Engine/Source/Runtime/Renderer/Private/IndirectLightRendering.cpp
@@ -209,6 +209,10 @@ class FDiffuseIndirectCompositePS : public FGlobalShader
 	{
 		OutEnvironment.CompilerFlags.Add(CFLAG_ForceOptimization);
 		OutEnvironment.SetDefine(TEXT("USE_HAIR_COMPLEX_TRANSMITTANCE"), IsHairStrandsSupported(EHairStrandsShaderType::All, Parameters.Platform) ? 1u : 0u);
+		if (Substrate::IsSubstrateEnabled() && Substrate::IsOpaqueRoughRefractionEnabled())
+		{
+			OutEnvironment.CompilerFlags.Add(CFLAG_AllowTypedUAVLoads);
+		}
 	}
 
 };
