RayTracing: Remove deprecated RayTraced Reflections

This functionality is now well covered by Lumen reflections which are the supported path going forward.

diff --git a/Engine/Config/Windows/DataDrivenPlatformInfo.ini b/Engine/Config/Windows/DataDrivenPlatformInfo.ini
--- a/Engine/Config/Windows/DataDrivenPlatformInfo.ini
+++ b/Engine/Config/Windows/DataDrivenPlatformInfo.ini
@@ -45,7 +45,6 @@ bSupportsRayTracingCallableShaders=true
 bSupportsRayTracingProceduralPrimitive=true
 bSupportsRayTracingIndirectInstanceData=true
 bSupportsPathTracing=true
-bSupportsHighEndRayTracingReflections=true
 bSupportsGPUScene=true
 bSupportsNanite=false
 bSupportsSceneDataCompressedTransforms=true
@@ -96,7 +95,6 @@ bSupportsRayTracingCallableShaders=true
 bSupportsRayTracingProceduralPrimitive=true
 bSupportsRayTracingIndirectInstanceData=true
 bSupportsPathTracing=true
-bSupportsHighEndRayTracingReflections=true
 bSupportsGPUScene=true
 bSupportsNanite=true
 bSupportsSceneDataCompressedTransforms=true
diff --git a/Engine/Shaders/Private/RayTracing/SplitImaginaryReflectionGBufferCS.usf b/Engine/Shaders/Private/RayTracing/SplitImaginaryReflectionGBufferCS.usf
deleted file mode 100644
diff --git a/Engine/Shaders/Private/ReflectionEnvironmentPixelShader.usf b/Engine/Shaders/Private/ReflectionEnvironmentPixelShader.usf
--- a/Engine/Shaders/Private/ReflectionEnvironmentPixelShader.usf
+++ b/Engine/Shaders/Private/ReflectionEnvironmentPixelShader.usf
@@ -143,14 +143,8 @@ float3 ReflectionEnvironment(FGBufferData GBuffer, float AmbientOcclusion, float
 
 	Color = CompositeReflections(ReflectionInput, BufferUV, GBuffer.Roughness, ShadingModelID);
 
-#if RAY_TRACED_REFLECTIONS
-	float4 SavedColor = Color;	// When a clear coat material is encountered, we save the reflection buffer color for it to not be affected by operations.
-#endif
 	if(GBuffer.ShadingModelID == SHADINGMODELID_CLEAR_COAT )
 	{
-#if RAY_TRACED_REFLECTIONS
-		Color = float4(0, 0, 0, 1); // Clear coat reflection is expected to be computed on a black background
-#endif
 		const float ClearCoat = GBuffer.CustomData.x;
 		Color = lerp( Color, float4(0,0,0,1), ClearCoat );
 
@@ -223,19 +217,14 @@ float3 ReflectionEnvironment(FGBufferData GBuffer, float AmbientOcclusion, float
 		Color.rgb *= LayerAttenuation;
 		Color.a = F;
 		
-#if !RAY_TRACED_REFLECTIONS	
 		Color.rgb += ReflectionInput.rgb * F;
 		Color.a *= 1 - ReflectionInput.a;
-#endif
 			
 		Color.a *= SpecularOcclusion;
 
 		float3 TopLayerR = 2 * dot( V, N ) * N - V;
 		Color.rgb += View.PreExposure * GatherRadiance(Color.a, TranslatedWorldPosition, TopLayerR, ClearCoatRoughness, BentNormal, IndirectIrradiance, GBuffer.ShadingModelID, NumCulledReflectionCaptures, DataStartIndex);
 
-	#if RAY_TRACED_REFLECTIONS
-		Color.rgb = SavedColor.rgb +  Color.rgb * SavedColor.a; // Compose default clear coat reflection over regular refelction (using Premultiplied alpha where SaveColor.a=transmittance)
-	#endif
 	}
 	else
 	{
diff --git a/Engine/Source/Runtime/Engine/Classes/Engine/EngineTypes.h b/Engine/Source/Runtime/Engine/Classes/Engine/EngineTypes.h
--- a/Engine/Source/Runtime/Engine/Classes/Engine/EngineTypes.h
+++ b/Engine/Source/Runtime/Engine/Classes/Engine/EngineTypes.h
@@ -454,9 +454,6 @@ namespace EReflectionMethod
 
 		/** Standalone Screen Space Reflections.  Low cost, but limited by screen space information. */
 		ScreenSpace UMETA(DisplayName="Screen Space"),
-
-		/** Standalone Ray Traced Reflections technique.  Deprecated, use Lumen Reflections instead. */
-		RayTraced UMETA(DisplayName="Standalone Ray Traced (Deprecated)"),
 	};
 }
 
diff --git a/Engine/Source/Runtime/Engine/Classes/Engine/Scene.h b/Engine/Source/Runtime/Engine/Classes/Engine/Scene.h
--- a/Engine/Source/Runtime/Engine/Classes/Engine/Scene.h
+++ b/Engine/Source/Runtime/Engine/Classes/Engine/Scene.h
@@ -1659,26 +1659,6 @@ struct FPostProcessSettings
 	UPROPERTY(interp, BlueprintReadWrite, Category="Reflections|Screen Space Reflections", meta=(ClampMin = "0.01", ClampMax = "1.0", editcondition = "bOverride_ScreenSpaceReflectionMaxRoughness", DisplayName = "Max Roughness"))
 	float ScreenSpaceReflectionMaxRoughness;
 
-	/** Sets the maximum roughness until which ray tracing reflections will be visible (lower value is faster). Reflection contribution is smoothly faded when close to roughness threshold. This parameter behaves similarly to ScreenSpaceReflectionMaxRoughness. */
-	UPROPERTY(interp, BlueprintReadWrite, Category = "Reflections|Ray Tracing Reflections", meta = (ClampMin = "0.01", ClampMax = "1.0", editcondition = "bOverride_RayTracingReflectionsMaxRoughness", DisplayName = "Max Roughness"))
-	float RayTracingReflectionsMaxRoughness;
-
-	/** Sets the maximum number of ray tracing reflection bounces. */
-	UPROPERTY(interp, EditAnywhere, BlueprintReadWrite, Category = "Reflections|Ray Tracing Reflections", meta = (ClampMin = "0", ClampMax = "50", editcondition = "bOverride_RayTracingReflectionsMaxBounces", DisplayName = "Max. Bounces"))
-	int32 RayTracingReflectionsMaxBounces;
-
-	/** Sets the samples per pixel for ray traced reflections. */
-	UPROPERTY(interp, EditAnywhere, BlueprintReadWrite, Category = "Reflections|Ray Tracing Reflections", meta = (ClampMin = "1", ClampMax = "65536", editcondition = "bOverride_RayTracingReflectionsSamplesPerPixel", DisplayName = "Samples Per Pixel"))
-	int32 RayTracingReflectionsSamplesPerPixel;
-
-	/** Sets the reflected shadows type. */
-	UPROPERTY(interp, EditAnywhere, BlueprintReadWrite, Category = "Reflections|Ray Tracing Reflections", meta = (editcondition = "bOverride_RayTracingReflectionsShadows", DisplayName = "Shadows"))
-	EReflectedAndRefractedRayTracedShadows RayTracingReflectionsShadows;
-
-	/** Enables ray tracing translucency in reflections. */
-	UPROPERTY(interp, EditAnywhere, BlueprintReadWrite, Category = "Reflections|Ray Tracing Reflections", meta = (editcondition = "bOverride_RayTracingReflectionsTranslucency", DisplayName = "Include Translucent Objects"))
-	uint8 RayTracingReflectionsTranslucency : 1;
-
 	/** AmbientCubemap tint color */
 	UPROPERTY(interp, BlueprintReadWrite, Category="Rendering Features|Ambient Cubemap", meta=(editcondition = "bOverride_AmbientCubemapTint", DisplayName = "Tint", HideAlphaChannel))
 	FLinearColor AmbientCubemapTint;
diff --git a/Engine/Source/Runtime/Engine/Private/Scene.cpp b/Engine/Source/Runtime/Engine/Private/Scene.cpp
--- a/Engine/Source/Runtime/Engine/Private/Scene.cpp
+++ b/Engine/Source/Runtime/Engine/Private/Scene.cpp
@@ -619,11 +619,6 @@ FPostProcessSettings::FPostProcessSettings()
 	ScreenSpaceReflectionIntensity = 100.0f;
 	ScreenSpaceReflectionQuality = 50.0f;
 	ScreenSpaceReflectionMaxRoughness = 0.6f;
-	RayTracingReflectionsMaxRoughness = 0.6f;
-	RayTracingReflectionsMaxBounces = 1;
-	RayTracingReflectionsSamplesPerPixel = 1;
-	RayTracingReflectionsShadows = EReflectedAndRefractedRayTracedShadows::Hard_shadows;
-	RayTracingReflectionsTranslucency = 0;
 
 	TranslucencyType = ETranslucencyType::Raster;
 	RayTracingTranslucencyMaxRoughness = 0.6f;
@@ -908,10 +903,6 @@ void FPostProcessSettings::PostSerialize(const FArchive& Ar)
 				{
 					ReflectionMethod = EReflectionMethod::ScreenSpace;
 				}
-				else if (ReflectionsType_DEPRECATED == EReflectionsType::RayTracing)
-				{
-					ReflectionMethod = EReflectionMethod::RayTraced;
-				}
 			}
 		}
 
diff --git a/Engine/Source/Runtime/Engine/Private/SceneView.cpp b/Engine/Source/Runtime/Engine/Private/SceneView.cpp
--- a/Engine/Source/Runtime/Engine/Private/SceneView.cpp
+++ b/Engine/Source/Runtime/Engine/Private/SceneView.cpp
@@ -1616,11 +1616,6 @@ void FSceneView::OverridePostProcessSettings(const FPostProcessSettings& Src, fl
 		LERP_PP(ScreenSpaceReflectionIntensity);
 		LERP_PP(ScreenSpaceReflectionMaxRoughness);
 
-		SET_PP(RayTracingReflectionsMaxRoughness);
-		SET_PP(RayTracingReflectionsMaxBounces);
-		SET_PP(RayTracingReflectionsSamplesPerPixel);
-		SET_PP(RayTracingReflectionsShadows);
-		SET_PP(RayTracingReflectionsTranslucency);
 		SET_PP(TranslucencyType);
 		SET_PP(RayTracingTranslucencyMaxRoughness);
 		SET_PP(RayTracingTranslucencyRefractionRays);
diff --git a/Engine/Source/Runtime/RHI/Private/DataDrivenShaderPlatformInfo.cpp b/Engine/Source/Runtime/RHI/Private/DataDrivenShaderPlatformInfo.cpp
--- a/Engine/Source/Runtime/RHI/Private/DataDrivenShaderPlatformInfo.cpp
+++ b/Engine/Source/Runtime/RHI/Private/DataDrivenShaderPlatformInfo.cpp
@@ -226,7 +226,6 @@ void FGenericDataDrivenShaderPlatformInfo::ParseDataDrivenShaderInfo(const FConf
 	GET_SECTION_BOOL_HELPER(bSupportsRayTracingTraversalStatistics);
 	GET_SECTION_BOOL_HELPER(bSupportsRayTracingIndirectInstanceData);
 	GET_SECTION_BOOL_HELPER(bSupportsPathTracing);
-	GET_SECTION_BOOL_HELPER(bSupportsHighEndRayTracingReflections);
 	GET_SECTION_BOOL_HELPER(bSupportsByteBufferComputeShaders);
 	GET_SECTION_BOOL_HELPER(bSupportsGPUScene);
 	GET_SECTION_BOOL_HELPER(bSupportsPrimitiveShaders);
diff --git a/Engine/Source/Runtime/RHI/Public/DataDrivenShaderPlatformInfo.h b/Engine/Source/Runtime/RHI/Public/DataDrivenShaderPlatformInfo.h
--- a/Engine/Source/Runtime/RHI/Public/DataDrivenShaderPlatformInfo.h
+++ b/Engine/Source/Runtime/RHI/Public/DataDrivenShaderPlatformInfo.h
@@ -50,7 +50,6 @@ class FGenericDataDrivenShaderPlatformInfo
 	uint32 bSupportsRayTracingProceduralPrimitive : 1;
 	uint32 bSupportsRayTracingTraversalStatistics : 1;
 	uint32 bSupportsRayTracingIndirectInstanceData : 1; // Whether instance transforms can be copied from the GPU to the TLAS instances buffer
-	uint32 bSupportsHighEndRayTracingReflections : 1; // Whether fully-featured RT reflections can be used on the platform (with multi-bounce, translucency, etc.)
 	uint32 bSupportsPathTracing : 1; // Whether real-time path tracer is supported on this platform (avoids compiling unnecessary shaders)
 	uint32 bSupportsGPUScene : 1;
 	uint32 bSupportsByteBufferComputeShaders : 1;
@@ -382,12 +381,6 @@ public:
 		return Infos[Platform].bSupportsRayTracing && Infos[Platform].bSupportsPathTracing;
 	}
 
-	static FORCEINLINE_DEBUGGABLE const bool GetSupportsHighEndRayTracingReflections(const FStaticShaderPlatform Platform)
-	{
-		check(IsValid(Platform));
-		return Infos[Platform].bSupportsRayTracing && Infos[Platform].bSupportsHighEndRayTracingReflections;
-	}
-
 	static FORCEINLINE_DEBUGGABLE const bool GetSupportsComputeFramework(const FStaticShaderPlatform Platform)
 	{
 		check(IsValid(Platform));
diff --git a/Engine/Source/Runtime/Renderer/Private/DeferredShadingRenderer.cpp b/Engine/Source/Runtime/Renderer/Private/DeferredShadingRenderer.cpp
--- a/Engine/Source/Runtime/Renderer/Private/DeferredShadingRenderer.cpp
+++ b/Engine/Source/Runtime/Renderer/Private/DeferredShadingRenderer.cpp
@@ -1897,8 +1897,6 @@ bool FDeferredShadingSceneRenderer::SetupRayTracingPipelineStates(FRDGBuilder& G
 			{
 				for (const FViewInfo& View : Views)
 				{
-					PrepareRayTracingReflections(View, *Scene, RayGenShaders);
-					PrepareSingleLayerWaterRayTracingReflections(View, *Scene, RayGenShaders);
 					PrepareRayTracingShadows(View, *Scene, RayGenShaders);
 					PrepareRayTracingAmbientOcclusion(View, RayGenShaders);
 					PrepareRayTracingSkyLight(View, *Scene, RayGenShaders);
@@ -1935,8 +1933,6 @@ bool FDeferredShadingSceneRenderer::SetupRayTracingPipelineStates(FRDGBuilder& G
 		{
 			for (const FViewInfo& View : Views)
 			{
-				PrepareRayTracingReflectionsDeferredMaterial(View, *Scene, DeferredMaterialRayGenShaders);
-				PrepareRayTracingDeferredReflectionsDeferredMaterial(View, *Scene, DeferredMaterialRayGenShaders);
 				if (DoesPlatformSupportLumenGI(ShaderPlatform))
 				{
 					PrepareLumenHardwareRayTracingReflectionsDeferredMaterial(View, DeferredMaterialRayGenShaders);
@@ -4624,7 +4620,6 @@ bool AnyRayTracingPassEnabled(const FScene* Scene, const FViewInfo& View)
 	}
 
 	return ShouldRenderRayTracingAmbientOcclusion(View)
-		|| ShouldRenderRayTracingReflections(View)
 		|| ShouldRenderRayTracingTranslucency(View)
 		|| ShouldRenderRayTracingSkyLight(Scene->SkyLight)
 		|| ShouldRenderRayTracingShadows()
@@ -4632,7 +4627,6 @@ bool AnyRayTracingPassEnabled(const FScene* Scene, const FViewInfo& View)
 		|| ShouldRenderPluginRayTracingGlobalIllumination(View)
         || Lumen::AnyLumenHardwareRayTracingPassEnabled(Scene, View)
 		|| StochasticShadows::UseHardwareRayTracing()
-		|| ShouldRenderRayTracingReflectionsWater(View)
 		|| HasRayTracedOverlay(*View.Family);
 }
 
diff --git a/Engine/Source/Runtime/Renderer/Private/DeferredShadingRenderer.h b/Engine/Source/Runtime/Renderer/Private/DeferredShadingRenderer.h
--- a/Engine/Source/Runtime/Renderer/Private/DeferredShadingRenderer.h
+++ b/Engine/Source/Runtime/Renderer/Private/DeferredShadingRenderer.h
@@ -296,7 +296,6 @@ enum class EReflectionsMethod
 {
 	Disabled,
 	SSR,
-	RTR,
 	Lumen
 };
 
@@ -1130,9 +1129,6 @@ private:
 	void SetupLumenHardwareRayTracingHitGroupBuffer(FRDGBuilder& GraphBuilder, FViewInfo& View);
 	void SetupLumenHardwareRayTracingUniformBuffer(FRDGBuilder& GraphBuilder, FViewInfo& View);
 	// #dxr_todo: UE-72565: refactor ray tracing effects to not be member functions of DeferredShadingRenderer. Register each effect at startup and just loop over them automatically
-	static void PrepareRayTracingReflections(const FViewInfo& View, const FScene& Scene, TArray<FRHIRayTracingShader*>& OutRayGenShaders);
-	static void PrepareRayTracingDeferredReflections(const FViewInfo& View, const FScene& Scene, TArray<FRHIRayTracingShader*>& OutRayGenShaders);
-	static void PrepareSingleLayerWaterRayTracingReflections(const FViewInfo& View, const FScene& Scene, TArray<FRHIRayTracingShader*>& OutRayGenShaders);
 	static void PrepareRayTracingShadows(const FViewInfo& View, const FScene& Scene, TArray<FRHIRayTracingShader*>& OutRayGenShaders);
 	static void PrepareRayTracingAmbientOcclusion(const FViewInfo& View, TArray<FRHIRayTracingShader*>& OutRayGenShaders);
 	static void PrepareRayTracingSkyLight(const FViewInfo& View, const FScene& Scene, TArray<FRHIRayTracingShader*>& OutRayGenShaders);
@@ -1153,10 +1149,6 @@ private:
 	static void PrepareLumenHardwareRayTracingVisualizeDeferredMaterial(const FViewInfo& View, TArray<FRHIRayTracingShader*>& OutRayGenShaders);
 	static void PrepareStochasticShadows(const FViewInfo& View, const FScene& Scene, TArray<FRHIRayTracingShader*>& OutRayGenShaders);
 
-	// Versions for setting up the deferred material pipeline
-	static void PrepareRayTracingReflectionsDeferredMaterial(const FViewInfo& View, const FScene& Scene, TArray<FRHIRayTracingShader*>& OutRayGenShaders);
-	static void PrepareRayTracingDeferredReflectionsDeferredMaterial(const FViewInfo& View, const FScene& Scene, TArray<FRHIRayTracingShader*>& OutRayGenShaders);
-	
 	// Versions for setting up the lumen material pipeline
 	static void PrepareLumenHardwareRayTracingTranslucencyVolumeLumenMaterial(const FViewInfo& View, TArray<FRHIRayTracingShader*>& OutRayGenShaders);
 	void PrepareLumenHardwareRayTracingVisualizeLumenMaterial(const FViewInfo& View, TArray<FRHIRayTracingShader*>& OutRayGenShaders);
diff --git a/Engine/Source/Runtime/Renderer/Private/IndirectLightRendering.cpp b/Engine/Source/Runtime/Renderer/Private/IndirectLightRendering.cpp
--- a/Engine/Source/Runtime/Renderer/Private/IndirectLightRendering.cpp
+++ b/Engine/Source/Runtime/Renderer/Private/IndirectLightRendering.cpp
@@ -18,7 +18,6 @@
 #include "PostProcess/PostProcessing.h" // for FPostProcessVS
 #include "RendererModule.h" 
 #include "RayTracing/RaytracingOptions.h"
-#include "RayTracing/RayTracingReflections.h"
 #include "DistanceFieldAmbientOcclusion.h"
 #include "VolumetricCloudRendering.h"
 #include "Lumen/LumenSceneData.h"
@@ -49,8 +48,7 @@ static TAutoConsoleVariable<int32> CVarReflectionMethod(
 	TEXT("r.ReflectionMethod"), 2,
 	TEXT("0 - None.  Reflections can come from placed Reflection Captures, Planar Reflections and Skylight but no global reflection method will be used.\n")
 	TEXT("1 - Lumen.  Use Lumen Reflections, which supports Screen / Software / Hardware Ray Tracing together and integrates with Lumen Global Illumination for rough reflections and Global Illumination seen in reflections.\n")
-	TEXT("2 - SSR.  Standalone Screen Space Reflections.  Low cost, but limited by screen space information.\n")
-	TEXT("3 - RT Reflections.  Ray Traced Reflections technique.  Deprecated, use Lumen Reflections instead."),
+	TEXT("2 - SSR.  Standalone Screen Space Reflections.  Low cost, but limited by screen space information.\n"),
 	ECVF_RenderThreadSafe);
 
 static TAutoConsoleVariable<int32> CVarDiffuseIndirectHalfRes(
@@ -289,7 +287,6 @@ class FReflectionEnvironmentSkyLightingPS : public FGlobalShader
 	class FSkyLight : SHADER_PERMUTATION_BOOL("ENABLE_SKY_LIGHT");
 	class FDynamicSkyLight : SHADER_PERMUTATION_BOOL("ENABLE_DYNAMIC_SKY_LIGHT");
 	class FSkyShadowing : SHADER_PERMUTATION_BOOL("APPLY_SKY_SHADOWING");
-	class FRayTracedReflections : SHADER_PERMUTATION_BOOL("RAY_TRACED_REFLECTIONS");
 	class FLumenStandaloneReflections : SHADER_PERMUTATION_BOOL("LUMEN_STANDALONE_REFLECTIONS");
 	class FSubstrateTileType : SHADER_PERMUTATION_INT("SUBSTRATE_TILETYPE", 4);
 
@@ -300,7 +297,6 @@ class FReflectionEnvironmentSkyLightingPS : public FGlobalShader
 		FSkyLight,
 		FDynamicSkyLight,
 		FSkyShadowing,
-		FRayTracedReflections,
 		FLumenStandaloneReflections,
 		FSubstrateTileType>;
 
@@ -323,20 +319,10 @@ class FReflectionEnvironmentSkyLightingPS : public FGlobalShader
 			PermutationVector.Set<FSubstrateTileType>(0);
 		}
 
-		if (PermutationVector.Get<FLumenStandaloneReflections>())
-		{
-			PermutationVector.Set<FRayTracedReflections>(false);
-		}
-
-		if (PermutationVector.Get<FRayTracedReflections>())
-		{
-			PermutationVector.Set<FLumenStandaloneReflections>(false);
-		}
-
 		return PermutationVector;
 	}
 
-	static FPermutationDomain BuildPermutationVector(const FViewInfo& View, bool bBoxCapturesOnly, bool bSphereCapturesOnly, bool bSupportDFAOIndirectOcclusion, bool bEnableSkyLight, bool bEnableDynamicSkyLight, bool bApplySkyShadowing, bool bRayTracedReflections, bool bLumenStandaloneReflections, ESubstrateTileType TileType)
+	static FPermutationDomain BuildPermutationVector(const FViewInfo& View, bool bBoxCapturesOnly, bool bSphereCapturesOnly, bool bSupportDFAOIndirectOcclusion, bool bEnableSkyLight, bool bEnableDynamicSkyLight, bool bApplySkyShadowing, bool bLumenStandaloneReflections, ESubstrateTileType TileType)
 	{
 		FPermutationDomain PermutationVector;
 
@@ -346,7 +332,6 @@ class FReflectionEnvironmentSkyLightingPS : public FGlobalShader
 		PermutationVector.Set<FSkyLight>(bEnableSkyLight);
 		PermutationVector.Set<FDynamicSkyLight>(bEnableDynamicSkyLight);
 		PermutationVector.Set<FSkyShadowing>(bApplySkyShadowing);
-		PermutationVector.Set<FRayTracedReflections>(bRayTracedReflections);
 		PermutationVector.Set<FLumenStandaloneReflections>(bLumenStandaloneReflections);
 		PermutationVector.Set<FSubstrateTileType>(0);
 		if (Substrate::IsSubstrateEnabled())
@@ -490,10 +475,6 @@ void FDeferredShadingSceneRenderer::CommitIndirectLightingState()
 		{
 			ReflectionsMethod = EReflectionsMethod::Lumen;
 		}
-		else if (ShouldRenderRayTracingReflections(View))
-		{
-			ReflectionsMethod = EReflectionsMethod::RTR;
-		}
 		else if (ScreenSpaceRayTracing::ShouldRenderScreenSpaceReflections(View))
 		{
 			ReflectionsMethod = EReflectionsMethod::SSR;
@@ -503,10 +484,6 @@ void FDeferredShadingSceneRenderer::CommitIndirectLightingState()
 		{
 			ReflectionsMethodWater = EReflectionsMethod::Lumen;
 		}
-		else if (ShouldRenderRayTracingReflectionsWater(View))
-		{
-			ReflectionsMethodWater = EReflectionsMethod::RTR;
-		}
 		else if (ScreenSpaceRayTracing::ShouldRenderScreenSpaceReflectionsWater(View))
 		{
 			ReflectionsMethodWater = EReflectionsMethod::SSR;
@@ -518,8 +495,7 @@ void FDeferredShadingSceneRenderer::CommitIndirectLightingState()
 		ViewPipelineState.Set(&FPerViewPipelineState::ReflectionsMethod, ReflectionsMethod);
 		ViewPipelineState.Set(&FPerViewPipelineState::ReflectionsMethodWater, ReflectionsMethodWater);
 
-		ViewPipelineState.Set(&FPerViewPipelineState::bComposePlanarReflections,
-			ReflectionsMethod != EReflectionsMethod::RTR && HasDeferredPlanarReflections(View));
+		ViewPipelineState.Set(&FPerViewPipelineState::bComposePlanarReflections, HasDeferredPlanarReflections(View));
 	}
 }
 
@@ -1646,7 +1622,6 @@ static void AddSkyReflectionPass(
 	const FSceneTextures& SceneTextures,
 	FRDGTextureRef DynamicBentNormalAOTexture,
 	FRDGTextureRef ReflectionsColor,
-	const FRayTracingReflectionOptions& RayTracingReflectionOptions,
 	FSceneTextureParameters& SceneTextureParameters,
 	bool bSkyLight, 
 	bool bDynamicSkyLight, 
@@ -1741,15 +1716,9 @@ static void AddSkyReflectionPass(
 	// Bind hair data
 	const bool bCheckerboardSubsurfaceRendering = IsSubsurfaceCheckerboardFormat(SceneColorTexture.Target->Desc.Format);
 
-	// ScreenSpace and SortedDeferred ray traced reflections use the same reflection environment shader,
-	// but main RT reflection shader requires a custom path as it evaluates the clear coat BRDF differently.
-	const bool bRequiresSpecializedReflectionEnvironmentShader = RayTracingReflectionOptions.bEnabled
-		&& RayTracingReflectionOptions.Algorithm != FRayTracingReflectionOptions::EAlgorithm::SortedDeferred;
-
 	auto PermutationVector = FReflectionEnvironmentSkyLightingPS::BuildPermutationVector(
 		View, bHasBoxCaptures, bHasSphereCaptures, DynamicBentNormalAO != 0.0f,
 		bSkyLight, bDynamicSkyLight, bApplySkyShadowing,
-		bRequiresSpecializedReflectionEnvironmentShader,
 		bLumenStandaloneReflections,
 		SubstrateTileMaterialType);
 
@@ -1905,31 +1874,6 @@ void FDeferredShadingSceneRenderer::RenderDeferredReflectionsAndSkyLighting(
 	FSceneTextureParameters SceneTextureParameters = GetSceneTextureParameters(GraphBuilder, SceneTextures);
 	const auto& SceneColorTexture = SceneTextures.Color;
 
-	IScreenSpaceDenoiser::FReflectionsInputs DenoiserInputs;
-	IScreenSpaceDenoiser::FReflectionsRayTracingConfig RayTracingConfig;
-
-	extern float GetRayTracingReflectionScreenPercentage();
-	RayTracingConfig.ResolutionFraction = GetRayTracingReflectionScreenPercentage();
-	int32 UpscaleFactor = int32(1.0f / RayTracingConfig.ResolutionFraction);
-
-	{
-		FRDGTextureDesc Desc = FRDGTextureDesc::Create2D(
-			SceneTextureParameters.SceneDepthTexture->Desc.Extent / UpscaleFactor,
-			PF_FloatRGBA,
-			FClearValueBinding::None,
-			TexCreate_ShaderResource | TexCreate_RenderTargetable | TexCreate_UAV);
-
-		DenoiserInputs.Color = GraphBuilder.CreateTexture(Desc, TEXT("RayTracingReflections"));
-
-		Desc.Format = PF_R16F;
-		DenoiserInputs.RayHitDistance = GraphBuilder.CreateTexture(Desc, TEXT("RayTracingReflectionsHitDistance"));
-		DenoiserInputs.RayImaginaryDepth = GraphBuilder.CreateTexture(Desc, TEXT("RayTracingReflectionsImaginaryDepth"));
-	}
-
-	FRDGTextureUAV* ReflectionColorOutputUAV = GraphBuilder.CreateUAV(FRDGTextureUAVDesc(DenoiserInputs.Color));
-	FRDGTextureUAV* RayHitDistanceOutputUAV = GraphBuilder.CreateUAV(FRDGTextureUAVDesc(DenoiserInputs.RayHitDistance));
-	FRDGTextureUAV* RayImaginaryDepthOutputUAV = GraphBuilder.CreateUAV(FRDGTextureUAVDesc(DenoiserInputs.RayImaginaryDepth));
-
 	uint32 ViewIndex = 0;
 	for (FViewInfo& View : Views)
 	{
@@ -1938,10 +1882,8 @@ void FDeferredShadingSceneRenderer::RenderDeferredReflectionsAndSkyLighting(
 		const uint32 CurrentViewIndex = ViewIndex++;
 		const FPerViewPipelineState& ViewPipelineState = GetViewPipelineState(View);
 
-		const FRayTracingReflectionOptions RayTracingReflectionOptions = GetRayTracingReflectionOptions(View, *Scene);
-
-		const bool bScreenSpaceReflections = !RayTracingReflectionOptions.bEnabled && ViewPipelineState.ReflectionsMethod == EReflectionsMethod::SSR;
-		const bool bComposePlanarReflections = !RayTracingReflectionOptions.bEnabled && HasDeferredPlanarReflections(View);
+		const bool bScreenSpaceReflections = ViewPipelineState.ReflectionsMethod == EReflectionsMethod::SSR;
+		const bool bComposePlanarReflections = HasDeferredPlanarReflections(View);
 
 		FRDGTextureRef ReflectionsColor = nullptr;
 		if ((ViewPipelineState.ReflectionsMethod == EReflectionsMethod::Lumen && ViewPipelineState.DiffuseIndirectMethod == EDiffuseIndirectMethod::Lumen)
@@ -1971,42 +1913,17 @@ void FDeferredShadingSceneRenderer::RenderDeferredReflectionsAndSkyLighting(
 			// Lumen needs its own depth history because things like Translucency velocities write to depth
 			StoreLumenDepthHistory(GraphBuilder, SceneTextures, View);
 		}
-		else if (RayTracingReflectionOptions.bEnabled || bScreenSpaceReflections)
+		else if (ViewPipelineState.ReflectionsMethod == EReflectionsMethod::SSR)
 		{
 			int32 DenoiserMode = GetReflectionsDenoiserMode();
 
-			bool bDenoise = false;
-			bool bTemporalFilter = false;
-
 			// Traces the reflections, either using screen space reflection, or ray tracing.
-			//IScreenSpaceDenoiser::FReflectionsInputs DenoiserInputs;
+			IScreenSpaceDenoiser::FReflectionsInputs DenoiserInputs;
 			IScreenSpaceDenoiser::FReflectionsRayTracingConfig DenoiserConfig;
-			if (RayTracingReflectionOptions.bEnabled)
-			{
-				RDG_EVENT_SCOPE(GraphBuilder, "RayTracingReflections %d", CurrentViewIndex);
-				RDG_GPU_STAT_SCOPE(GraphBuilder, RayTracingReflections);
-
-				bDenoise = DenoiserMode != 0;
-
-				DenoiserConfig.ResolutionFraction = RayTracingReflectionOptions.ResolutionFraction;
-				DenoiserConfig.RayCountPerPixel = RayTracingReflectionOptions.SamplesPerPixel;
+			bool bDenoise = DenoiserMode != 0 && CVarDenoiseSSR.GetValueOnRenderThread();
+			bool bTemporalFilter = !bDenoise && View.ViewState && ScreenSpaceRayTracing::IsSSRTemporalPassRequired(View);
 
-				check(RayTracingReflectionOptions.bReflectOnlyWater == false);
-
-				RenderRayTracingReflections(
-					GraphBuilder,
-					SceneTextures,
-					View,
-					DenoiserMode,
-					RayTracingReflectionOptions,
-					&DenoiserInputs);
-			}
-			else if (
-				ViewPipelineState.ReflectionsMethod == EReflectionsMethod::SSR)
 			{
-				bDenoise = DenoiserMode != 0 && CVarDenoiseSSR.GetValueOnRenderThread();
-				bTemporalFilter = !bDenoise && View.ViewState && ScreenSpaceRayTracing::IsSSRTemporalPassRequired(View);
-
 				ESSRQuality SSRQuality;
 				ScreenSpaceRayTracing::GetSSRQualityForView(View, &SSRQuality, &DenoiserConfig);
 
@@ -2015,11 +1932,6 @@ void FDeferredShadingSceneRenderer::RenderDeferredReflectionsAndSkyLighting(
 				ScreenSpaceRayTracing::RenderScreenSpaceReflections(
 					GraphBuilder, SceneTextureParameters, SceneColorTexture.Resolve, View, SSRQuality, bDenoise, &DenoiserInputs);
 			}
-			else
-			{
-				check(0);
-			}
-
 			if (bDenoise)
 			{
 				const IScreenSpaceDenoiser* DefaultDenoiser = IScreenSpaceDenoiser::GetDefaultDenoiser();
@@ -2064,22 +1976,12 @@ void FDeferredShadingSceneRenderer::RenderDeferredReflectionsAndSkyLighting(
 			}
 			else
 			{
-				if (RayTracingReflectionOptions.bEnabled && DenoiserInputs.RayHitDistance)
-				{
-					// The performance of ray tracing does not allow to run without a denoiser in real time.
-					// Multiple rays per pixel is unsupported by the denoiser that will most likely more bound by to
-					// many rays than exporting the hit distance buffer. Therefore no permutation of the ray generation
-					// shader has been judged required to be supported.
-					GraphBuilder.RemoveUnusedTextureWarning(DenoiserInputs.RayHitDistance);
-				}
-
 				ReflectionsColor = DenoiserInputs.Color;
 			}
 		} // if (RayTracingReflectionOptions.bEnabled || bScreenSpaceReflections)
 
 		if (ViewPipelineState.bComposePlanarReflections)
 		{
-			check(!RayTracingReflectionOptions.bEnabled);
 			RenderDeferredPlanarReflections(GraphBuilder, SceneTextureParameters, View, /* inout */ ReflectionsColor);
 		}
 
@@ -2101,7 +2003,6 @@ void FDeferredShadingSceneRenderer::RenderDeferredReflectionsAndSkyLighting(
 					SceneTextures,
 					DynamicBentNormalAOTexture,
 					ReflectionsColor,
-					RayTracingReflectionOptions,
 					SceneTextureParameters,
 					bSkyLight,
 					bDynamicSkyLight,
@@ -2116,7 +2017,6 @@ void FDeferredShadingSceneRenderer::RenderDeferredReflectionsAndSkyLighting(
 					SceneTextures,
 					DynamicBentNormalAOTexture,
 					ReflectionsColor,
-					RayTracingReflectionOptions,
 					SceneTextureParameters,
 					bSkyLight,
 					bDynamicSkyLight,
@@ -2131,7 +2031,6 @@ void FDeferredShadingSceneRenderer::RenderDeferredReflectionsAndSkyLighting(
 					SceneTextures,
 					DynamicBentNormalAOTexture,
 					ReflectionsColor,
-					RayTracingReflectionOptions,
 					SceneTextureParameters,
 					bSkyLight,
 					bDynamicSkyLight,
@@ -2146,7 +2045,6 @@ void FDeferredShadingSceneRenderer::RenderDeferredReflectionsAndSkyLighting(
 					SceneTextures,
 					DynamicBentNormalAOTexture,
 					ReflectionsColor,
-					RayTracingReflectionOptions,
 					SceneTextureParameters,
 					bSkyLight,
 					bDynamicSkyLight,
@@ -2164,7 +2062,6 @@ void FDeferredShadingSceneRenderer::RenderDeferredReflectionsAndSkyLighting(
 					SceneTextures,
 					DynamicBentNormalAOTexture,
 					ReflectionsColor,
-					RayTracingReflectionOptions,
 					SceneTextureParameters,
 					bSkyLight,
 					bDynamicSkyLight,
diff --git a/Engine/Source/Runtime/Renderer/Private/RayTracing/RayTracingDeferredReflections.cpp b/Engine/Source/Runtime/Renderer/Private/RayTracing/RayTracingDeferredReflections.cpp
deleted file mode 100644
diff --git a/Engine/Source/Runtime/Renderer/Private/RayTracing/RayTracingReflections.cpp b/Engine/Source/Runtime/Renderer/Private/RayTracing/RayTracingReflections.cpp
deleted file mode 100644
diff --git a/Engine/Source/Runtime/Renderer/Private/RayTracing/RayTracingReflections.h b/Engine/Source/Runtime/Renderer/Private/RayTracing/RayTracingReflections.h
deleted file mode 100644
diff --git a/Engine/Source/Runtime/Renderer/Private/RayTracing/RaytracingOptions.h b/Engine/Source/Runtime/Renderer/Private/RayTracing/RaytracingOptions.h
--- a/Engine/Source/Runtime/Renderer/Private/RayTracing/RaytracingOptions.h
+++ b/Engine/Source/Runtime/Renderer/Private/RayTracing/RaytracingOptions.h
@@ -81,8 +81,6 @@ extern FRayTracingPrimaryRaysOptions GetRayTracingTranslucencyOptions(const FVie
 
 extern bool ShouldRenderRayTracingSkyLight(const FSkyLightSceneProxy* SkyLightSceneProxy);
 extern bool ShouldRenderRayTracingAmbientOcclusion(const FViewInfo& View);
-extern bool ShouldRenderRayTracingReflections(const FViewInfo& View);
-extern bool ShouldRenderRayTracingReflectionsWater(const FViewInfo& View);
 extern bool ShouldRenderRayTracingTranslucency(const FViewInfo& View);
 extern bool ShouldRenderRayTracingShadows();
 extern bool ShouldRenderRayTracingShadowsForLight(const FLightSceneProxy& LightProxy);
@@ -125,16 +123,6 @@ FORCEINLINE bool ShouldRenderRayTracingAmbientOcclusion(const FViewInfo& View)
 	return false;
 }
 
-FORCEINLINE bool ShouldRenderRayTracingReflections(const FViewInfo& View)
-{
-	return false;
-}
-
-FORCEINLINE bool ShouldRenderRayTracingReflectionsWater(const FViewInfo& View)
-{
-	return false;
-}
-
 FORCEINLINE bool ShouldRenderRayTracingTranslucency(const FViewInfo& View)
 {
 	return false;
diff --git a/Engine/Source/Runtime/Renderer/Private/RendererScene.cpp b/Engine/Source/Runtime/Renderer/Private/RendererScene.cpp
--- a/Engine/Source/Runtime/Renderer/Private/RendererScene.cpp
+++ b/Engine/Source/Runtime/Renderer/Private/RendererScene.cpp
@@ -845,12 +845,9 @@ uint64 FPreviousViewInfo::GetGPUSizeBytes(bool bLogSizes) const
 		GetRenderTargetGPUSizeBytes(GBufferA, bLogSizes) +
 		GetRenderTargetGPUSizeBytes(GBufferB, bLogSizes) +
 		GetRenderTargetGPUSizeBytes(GBufferC, bLogSizes) +
-		GetRenderTargetGPUSizeBytes(ImaginaryReflectionDepthBuffer, bLogSizes) +
-		GetRenderTargetGPUSizeBytes(ImaginaryReflectionGBufferA, bLogSizes) +
 		GetRenderTargetGPUSizeBytes(HZB, bLogSizes) +
 		GetRenderTargetGPUSizeBytes(NaniteHZB, bLogSizes) +
 		GetRenderTargetGPUSizeBytes(CompressedDepthViewNormal, bLogSizes) +
-		GetRenderTargetGPUSizeBytes(ImaginaryReflectionCompressedDepthViewNormal, bLogSizes) +
 		GetRenderTargetGPUSizeBytes(CompressedOpaqueDepth, bLogSizes) +
 		GetRenderTargetGPUSizeBytes(CompressedOpaqueShadingModel, bLogSizes) +
 		GetRenderTargetGPUSizeBytes(ScreenSpaceRayTracingInput, bLogSizes) +
@@ -1106,9 +1103,6 @@ uint64 FSceneViewState::GetGPUSizeBytes(bool bLogSizes) const
 	TotalSize += GetBufferGPUSizeBytes(BloomFFTKernel.ConstantsBuffer, bLogSizes);
 	TotalSize += GetBufferGPUSizeBytes(FilmGrainCache.ConstantsBuffer, bLogSizes);
 #if RHI_RAYTRACING
-	TotalSize += GetRenderTargetGPUSizeBytes(ImaginaryReflectionGBufferA, bLogSizes);
-	TotalSize += GetRenderTargetGPUSizeBytes(ImaginaryReflectionDepthZ, bLogSizes);
-	TotalSize += GetRenderTargetGPUSizeBytes(ImaginaryReflectionVelocity, bLogSizes);
 	TotalSize += GetBufferGPUSizeBytes(SkyLightVisibilityRaysBuffer, bLogSizes);
 #endif
 	TotalSize += GetRenderTargetGPUSizeBytes(LightScatteringHistory, bLogSizes);
diff --git a/Engine/Source/Runtime/Renderer/Private/ScenePrivate.h b/Engine/Source/Runtime/Renderer/Private/ScenePrivate.h
--- a/Engine/Source/Runtime/Renderer/Private/ScenePrivate.h
+++ b/Engine/Source/Runtime/Renderer/Private/ScenePrivate.h
@@ -1085,11 +1085,6 @@ public:
 	TPimplPtr<FPathTracingState> PathTracingState;
 	uint32 PathTracingInvalidationCounter = 0;
 
-	// Ray Traced Reflection Imaginary GBuffer Data containing a pseudo-geometric representation of the reflected surface(s)
-	TRefCountPtr<IPooledRenderTarget> ImaginaryReflectionGBufferA;
-	TRefCountPtr<IPooledRenderTarget> ImaginaryReflectionDepthZ;
-	TRefCountPtr<IPooledRenderTarget> ImaginaryReflectionVelocity;
-
 	// Ray Traced Sky Light Sample Direction Data
 	TRefCountPtr<FRDGPooledBuffer> SkyLightVisibilityRaysBuffer;
 	FIntVector SkyLightVisibilityRaysDimensions;
diff --git a/Engine/Source/Runtime/Renderer/Private/SceneRendering.h b/Engine/Source/Runtime/Renderer/Private/SceneRendering.h
--- a/Engine/Source/Runtime/Renderer/Private/SceneRendering.h
+++ b/Engine/Source/Runtime/Renderer/Private/SceneRendering.h
@@ -1053,8 +1053,6 @@ struct FPreviousViewInfo
 	TRefCountPtr<IPooledRenderTarget> GBufferA;
 	TRefCountPtr<IPooledRenderTarget> GBufferB;
 	TRefCountPtr<IPooledRenderTarget> GBufferC;
-	TRefCountPtr<IPooledRenderTarget> ImaginaryReflectionDepthBuffer;
-	TRefCountPtr<IPooledRenderTarget> ImaginaryReflectionGBufferA;
 
 	TRefCountPtr<IPooledRenderTarget> HZB;
 	TRefCountPtr<IPooledRenderTarget> NaniteHZB;
@@ -1067,7 +1065,6 @@ struct FPreviousViewInfo
 	// Compressed scene textures for bandwidth efficient bilateral kernel rejection.
 	// DeviceZ as float16, and normal in view space.
 	TRefCountPtr<IPooledRenderTarget> CompressedDepthViewNormal;
-	TRefCountPtr<IPooledRenderTarget> ImaginaryReflectionCompressedDepthViewNormal;
 
 	// 16bit compressed depth buffer with opaque only.
 	TRefCountPtr<IPooledRenderTarget> CompressedOpaqueDepth;
diff --git a/Engine/Source/Runtime/Renderer/Private/ScreenSpaceDenoise.cpp b/Engine/Source/Runtime/Renderer/Private/ScreenSpaceDenoise.cpp
--- a/Engine/Source/Runtime/Renderer/Private/ScreenSpaceDenoise.cpp
+++ b/Engine/Source/Runtime/Renderer/Private/ScreenSpaceDenoise.cpp
@@ -1258,25 +1258,6 @@ void SetupSceneViewInfoPooledRenderTargets(
 	OutViewInfoPooledRenderTargets->NextCompressedDepthViewNormal = &PrevFrameViewInfo.CompressedDepthViewNormal;
 }
 
-void SetupImaginaryReflectionViewInfoPooledRenderTargets(
-	const FViewInfo& View,
-	FViewInfoPooledRenderTargets* OutViewInfoPooledRenderTargets)
-{
-	auto&& PrevViewInfo = View.PrevViewInfo;
-	auto&& PrevFrameViewInfo = View.ViewState->PrevFrameViewInfo;
-
-	OutViewInfoPooledRenderTargets->PrevDepthBuffer = PrevViewInfo.ImaginaryReflectionDepthBuffer;
-	OutViewInfoPooledRenderTargets->PrevGBufferA = PrevViewInfo.ImaginaryReflectionGBufferA;
-	OutViewInfoPooledRenderTargets->PrevGBufferB = nullptr; // GBufferB not used
-	OutViewInfoPooledRenderTargets->PrevCompressedDepthViewNormal = PrevViewInfo.ImaginaryReflectionCompressedDepthViewNormal;
-
-	OutViewInfoPooledRenderTargets->NextDepthBuffer = &PrevFrameViewInfo.ImaginaryReflectionDepthBuffer;
-	OutViewInfoPooledRenderTargets->NextGBufferA = &PrevFrameViewInfo.ImaginaryReflectionGBufferA;
-	OutViewInfoPooledRenderTargets->NextGBufferB = nullptr; // GBufferB not used
-	OutViewInfoPooledRenderTargets->NextCompressedDepthViewNormal = &PrevFrameViewInfo.ImaginaryReflectionCompressedDepthViewNormal;
-}
-
-
 void Denoiser::SetupCommonShaderParameters(
 	const FViewInfo& View,
 	const FSceneTextureParameters& SceneTextures,
@@ -2787,53 +2768,7 @@ public:
 		GlobalIlluminationOutputs.Color = SignalOutput.Textures[0];
 		return GlobalIlluminationOutputs;
 	}
-
-	FDiffuseIndirectOutputs DenoiseReflectedSkyLight(
-		FRDGBuilder& GraphBuilder,
-		const FViewInfo& View,
-		FPreviousViewInfo* PreviousViewInfos,
-		const FSceneTextureParameters& SceneTextures,
-		const FDiffuseIndirectInputs& Inputs,
-		const FAmbientOcclusionRayTracingConfig Config) const override
-	{
-		RDG_GPU_STAT_SCOPE(GraphBuilder, DiffuseIndirectDenoiser);
-
-		FViewInfoPooledRenderTargets ViewInfoPooledRenderTargets;
-		SetupImaginaryReflectionViewInfoPooledRenderTargets(View, &ViewInfoPooledRenderTargets);
-
-		FSSDSignalTextures InputSignal;
-		InputSignal.Textures[0] = Inputs.Color;
-		InputSignal.Textures[1] = Inputs.RayHitDistance;
-
-		FSSDConstantPixelDensitySettings Settings;
-		Settings.FullResViewport = View.ViewRect;
-		Settings.SignalProcessing = ESignalProcessing::DiffuseAndAmbientOcclusion;
-		Settings.InputResolutionFraction = Config.ResolutionFraction;
-		Settings.ReconstructionSamples = FMath::Clamp(CVarGIReconstructionSampleCount.GetValueOnRenderThread(), 1, kStackowiakMaxSampleCountPerSet);
-		Settings.PreConvolutionCount = CVarGIPreConvolutionCount.GetValueOnRenderThread();
-		Settings.bUseTemporalAccumulation = CVarGITemporalAccumulation.GetValueOnRenderThread() != 0;
-		Settings.HistoryConvolutionSampleCount = CVarGIHistoryConvolutionSampleCount.GetValueOnRenderThread();
-		Settings.HistoryConvolutionKernelSpreadFactor = CVarGIHistoryConvolutionKernelSpreadFactor.GetValueOnRenderThread();
-		Settings.MaxInputSPP = Config.RayCountPerPixel;
-
-		TStaticArray<FScreenSpaceDenoiserHistory*, IScreenSpaceDenoiser::kMaxBatchSize> PrevHistories;
-		TStaticArray<FScreenSpaceDenoiserHistory*, IScreenSpaceDenoiser::kMaxBatchSize> NewHistories;
-		PrevHistories[0] = &PreviousViewInfos->ReflectedSkyLightHistory;
-		NewHistories[0] = View.ViewState ? &View.ViewState->PrevFrameViewInfo.ReflectedSkyLightHistory : nullptr;
-
-		FSSDSignalTextures SignalOutput;
-		DenoiseSignalAtConstantPixelDensity(
-			GraphBuilder, View, SceneTextures, ViewInfoPooledRenderTargets,
-			InputSignal, Settings,
-			PrevHistories,
-			NewHistories,
-			&SignalOutput);
-
-		FDiffuseIndirectOutputs GlobalIlluminationOutputs;
-		GlobalIlluminationOutputs.Color = SignalOutput.Textures[0];
-		return GlobalIlluminationOutputs;
-	}
-
+	
 	FSSDSignalTextures DenoiseDiffuseIndirectHarmonic(
 		FRDGBuilder& GraphBuilder,
 		const FViewInfo& View,
diff --git a/Engine/Source/Runtime/Renderer/Private/ScreenSpaceDenoise.h b/Engine/Source/Runtime/Renderer/Private/ScreenSpaceDenoise.h
--- a/Engine/Source/Runtime/Renderer/Private/ScreenSpaceDenoise.h
+++ b/Engine/Source/Runtime/Renderer/Private/ScreenSpaceDenoise.h
@@ -296,15 +296,6 @@ public:
 		const FDiffuseIndirectInputs& Inputs,
 		const FAmbientOcclusionRayTracingConfig Config) const = 0;
 
-	/** Entry point to denoise reflected SkyLight diffuse indirect. */
-	virtual FDiffuseIndirectOutputs DenoiseReflectedSkyLight(
-		FRDGBuilder& GraphBuilder,
-		const FViewInfo& View,
-		FPreviousViewInfo* PreviousViewInfos,
-		const FSceneTextureParameters& SceneTextures,
-		const FDiffuseIndirectInputs& Inputs,
-		const FAmbientOcclusionRayTracingConfig Config) const = 0;
-
 	/** Entry point to denoise spherical harmonic for diffuse indirect. */
 	virtual FSSDSignalTextures DenoiseDiffuseIndirectHarmonic(
 		FRDGBuilder& GraphBuilder,
diff --git a/Engine/Source/Runtime/Renderer/Private/SingleLayerWaterRendering.cpp b/Engine/Source/Runtime/Renderer/Private/SingleLayerWaterRendering.cpp
--- a/Engine/Source/Runtime/Renderer/Private/SingleLayerWaterRendering.cpp
+++ b/Engine/Source/Runtime/Renderer/Private/SingleLayerWaterRendering.cpp
@@ -10,7 +10,6 @@
 #include "PostProcess/SceneRenderTargets.h"
 #include "PostProcess/TemporalAA.h"
 #include "RayTracing/RaytracingOptions.h"
-#include "RayTracing/RayTracingReflections.h"
 #include "VolumetricRenderTarget.h"
 #include "RenderGraph.h"
 #include "ScenePrivate.h"
@@ -205,18 +204,6 @@ bool ShouldRenderScreenSpaceReflectionsWater(const FViewInfo& View)
 }
 }
 
-#if RHI_RAYTRACING
-bool ShouldRenderRayTracingReflectionsWater(const FViewInfo& View)
-{
-	// This only returns true if using the default reflections method and having RTR enabled in the scene. It can't be forced with r.Water.SingleLayer.Reflection.
-	const bool bEffectEnabled = !View.bIsReflectionCapture 
-		&& GetSingleLayerWaterReflectionTechnique() == ESingleLayerWaterReflections::Enabled
-		&& ShouldRenderRayTracingReflections(View)
-		&& FDataDrivenShaderPlatformInfo::GetSupportsHighEndRayTracingReflections(View.GetShaderPlatform());
-	return ShouldRenderRayTracingEffect(bEffectEnabled, ERayTracingPipelineCompatibilityFlags::FullPipeline, nullptr);
-}
-#endif // RHI_RAYTRACING
-
 bool ShouldRenderLumenReflectionsWater(const FViewInfo& View, bool bSkipTracingDataCheck, bool bSkipProjectCheck)
 {
 	// This only returns true if using the default reflections method and having Lumen enabled in the scene. It can't be forced with r.Water.SingleLayer.Reflection.
@@ -1024,86 +1011,6 @@ void FDeferredShadingSceneRenderer::RenderSingleLayerWaterReflections(
 				nullptr,
 				ERDGPassFlags::Compute);
 		}
-		else if (ViewPipelineState.ReflectionsMethodWater == EReflectionsMethod::RTR)
-		{
-			check(ShouldRenderRayTracingReflectionsWater(View));
-			RDG_EVENT_SCOPE(GraphBuilder, "SLW::RayTracingReflections");
-			RDG_GPU_STAT_SCOPE(GraphBuilder, RayTracingWaterReflections);
-
-			IScreenSpaceDenoiser::FReflectionsInputs DenoiserInputs;
-			IScreenSpaceDenoiser::FReflectionsRayTracingConfig RayTracingConfig;
-
-			//RayTracingConfig.ResolutionFraction = FMath::Clamp(GetRayTracingReflectionsScreenPercentage() / 100.0f, 0.25f, 1.0f);
-			RayTracingConfig.ResolutionFraction = 1.0f;
-			//RayTracingConfig.RayCountPerPixel = GetRayTracingReflectionsSamplesPerPixel(View) > -1 ? GetRayTracingReflectionsSamplesPerPixel(View) : View.FinalPostProcessSettings.RayTracingReflectionsSamplesPerPixel;
-			RayTracingConfig.RayCountPerPixel = 1;
-
-			// Water is assumed to have zero roughness and is not currently denoised.
-			//int32 DenoiserMode = GetReflectionsDenoiserMode();
-			//bool bDenoise = DenoiserMode != 0;
-			int32 DenoiserMode = 0;
-			bool bDenoise = false;
-
-			if (!bDenoise)
-			{
-				RayTracingConfig.ResolutionFraction = 1.0f;
-			}
-
-			FRayTracingReflectionOptions Options;
-			Options.Algorithm = FRayTracingReflectionOptions::BruteForce;
-			Options.SamplesPerPixel = 1;
-			Options.ResolutionFraction = 1.0;
-			Options.bReflectOnlyWater = true;
-
-			{
-				float UpscaleFactor = 1.0;
-				FRDGTextureDesc Desc = FRDGTextureDesc::Create2D(
-					SceneTextures.Config.Extent / UpscaleFactor,
-					PF_FloatRGBA,
-					FClearValueBinding::None,
-					TexCreate_ShaderResource | TexCreate_RenderTargetable | TexCreate_UAV);
-
-				DenoiserInputs.Color = GraphBuilder.CreateTexture(Desc, TEXT("SLW.RayTracingReflections"));
-
-				Desc.Format = PF_R16F;
-				DenoiserInputs.RayHitDistance = GraphBuilder.CreateTexture(Desc, TEXT("SLW.RayTracingReflectionsHitDistance"));
-				DenoiserInputs.RayImaginaryDepth = GraphBuilder.CreateTexture(Desc, TEXT("SLW.RayTracingReflectionsImaginaryDepth"));
-			}
-
-			RenderRayTracingReflections(
-				GraphBuilder,
-				SceneTextures,
-				View,
-				DenoiserMode,
-				Options,
-				&DenoiserInputs);
-
-			if (bDenoise)
-			{
-				const IScreenSpaceDenoiser* DefaultDenoiser = IScreenSpaceDenoiser::GetDefaultDenoiser();
-				const IScreenSpaceDenoiser* DenoiserToUse = DenoiserMode == 1 ? DefaultDenoiser : GScreenSpaceDenoiser;
-
-				// Standard event scope for denoiser to have all profiling information not matter what, and with explicit detection of third party.
-				RDG_EVENT_SCOPE(GraphBuilder, "%s%s(WaterReflections) %dx%d",
-					DenoiserToUse != DefaultDenoiser ? TEXT("ThirdParty ") : TEXT(""),
-					DenoiserToUse->GetDebugName(),
-					View.ViewRect.Width(), View.ViewRect.Height());
-
-				IScreenSpaceDenoiser::FReflectionsOutputs DenoiserOutputs = DenoiserToUse->DenoiseWaterReflections(
-					GraphBuilder,
-					View,
-					&View.PrevViewInfo,
-					SceneTextureParameters,
-					DenoiserInputs,
-					RayTracingConfig);
-
-				ReflectionsColor = DenoiserOutputs.Color;
-			}
-			else
-			{
-				ReflectionsColor = DenoiserInputs.Color;
-			}
-		}
 		else if (ViewPipelineState.ReflectionsMethodWater == EReflectionsMethod::SSR)
 		{
 			check(ScreenSpaceRayTracing::ShouldRenderScreenSpaceReflectionsWater(View));
