#UE IWYU Fix

diff --git a/Engine/Source/Runtime/Renderer/Private/OIT/OIT.cpp b/Engine/Source/Runtime/Renderer/Private/OIT/OIT.cpp
--- a/Engine/Source/Runtime/Renderer/Private/OIT/OIT.cpp
+++ b/Engine/Source/Runtime/Renderer/Private/OIT/OIT.cpp
@@ -3,6 +3,7 @@
 #include "OIT.h"
 #include "DataDrivenShaderPlatformInfo.h"
 #include "LocalVertexFactory.h"
+#include "MeshBatch.h"
 #include "OITParameters.h"
 #include "Shader.h"
 #include "GlobalShader.h"
diff --git a/Engine/Source/Runtime/Renderer/Private/OIT/OIT.h b/Engine/Source/Runtime/Renderer/Private/OIT/OIT.h
--- a/Engine/Source/Runtime/Renderer/Private/OIT/OIT.h
+++ b/Engine/Source/Runtime/Renderer/Private/OIT/OIT.h
@@ -9,6 +9,7 @@
 
 class FViewInfo;
 class FSortedIndexBuffer;
+struct FMeshBatchElement;
 struct FMeshBatchElementDynamicIndexBuffer;
 class FViewInfo;
 
