RayTracing: Fix RT Debug PrimaryRays mode (SceneColor is no longer valid since the BasePass is skipped)

diff --git a/Engine/Source/Runtime/Renderer/Private/RayTracing/RayTracingPrimaryRays.cpp b/Engine/Source/Runtime/Renderer/Private/RayTracing/RayTracingPrimaryRays.cpp
--- a/Engine/Source/Runtime/Renderer/Private/RayTracing/RayTracingPrimaryRays.cpp
+++ b/Engine/Source/Runtime/Renderer/Private/RayTracing/RayTracingPrimaryRays.cpp
@@ -160,7 +160,7 @@ void FDeferredShadingSceneRenderer::RenderRayTracingPrimaryRaysView(
 
 	PassParameters->SceneTextures = SceneTextureParameters;
 
-	PassParameters->SceneColorTexture = SceneTextures.Color.Resolve;
+	PassParameters->SceneColorTexture = GetIfProduced(SceneTextures.Color.Resolve, FRDGSystemTextures::Get(GraphBuilder).Black);
 
 	PassParameters->ReflectionStruct = CreateReflectionUniformBuffer(GraphBuilder, View);
 	PassParameters->FogUniformParameters = CreateFogUniformBuffer(GraphBuilder, View);
