Add an option to use a single permutation for mobile BasePass regardless of local light state
When enabled this removes the need to invalidate primitive MDC when local light state changes

diff --git a/Engine/Source/Runtime/Renderer/Private/MobileBasePass.cpp b/Engine/Source/Runtime/Renderer/Private/MobileBasePass.cpp
--- a/Engine/Source/Runtime/Renderer/Private/MobileBasePass.cpp
+++ b/Engine/Source/Runtime/Renderer/Private/MobileBasePass.cpp
@@ -31,6 +31,19 @@ bool MobileLocalLightsBufferPostprocessEnabled(const FStaticShaderPlatform Platf
 	return MobileLocalLightsBufferEnabled(Platform) && !MobileUsesFullDepthPrepass(Platform);
 }
 
+int32 GMobileForwardLocalLightsSinglePermutation = 0;
+FAutoConsoleVariableRef CVarMobileForwardLocalLightsSinglePermutation(
+	TEXT("r.Mobile.Forward.LocalLightsSinglePermutation"),
+	GMobileForwardLocalLightsSinglePermutation,
+	TEXT("Whether to use the same permutation regardless of local lights state. This may improve RT time at expense of some GPU time"),
+	ECVF_Scalability | ECVF_RenderThreadSafe
+);
+
+bool MobileLocalLightsUseSinglePermutation()
+{ 
+	return GMobileForwardLocalLightsSinglePermutation != 0;
+}
+
 EMobileLocalLightSetting GetMobileForwardLocalLightSetting(EShaderPlatform ShaderPlatform, bool bIsTranslucent)
 {
 	static FShaderPlatformCachedIniValue<int32> MobileForwardLocalLightsIniValueCvar(TEXT("r.Mobile.Forward.EnableLocalLights"));
@@ -879,7 +892,11 @@ bool FMobileBasePassMeshProcessor::Process(
 	EMobileLocalLightSetting LocalLightSetting = EMobileLocalLightSetting::LOCAL_LIGHTS_DISABLED;
 	if (Scene && PrimitiveSceneProxy && ShadingModels.IsLit())
 	{
-		if (!bPassUsesDeferredShading && (PrimitiveSceneProxy->GetPrimitiveSceneInfo()->NumMobileDynamicLocalLights > 0))
+		if (!bPassUsesDeferredShading && 
+			// we can choose to use a single permutation for opaque meshes regarless of local light state
+			// this is to avoid re-caching MDC on light state changes
+			// Translucency always use un-cached MDC so this option does not apply to it
+			((MobileLocalLightsUseSinglePermutation() && !bIsTranslucent) || PrimitiveSceneProxy->GetPrimitiveSceneInfo()->NumMobileDynamicLocalLights > 0))
 		{
 			LocalLightSetting = GetMobileForwardLocalLightSetting(Scene->GetShaderPlatform(), bIsTranslucent);
 		}
diff --git a/Engine/Source/Runtime/Renderer/Private/SceneCore.cpp b/Engine/Source/Runtime/Renderer/Private/SceneCore.cpp
--- a/Engine/Source/Runtime/Renderer/Private/SceneCore.cpp
+++ b/Engine/Source/Runtime/Renderer/Private/SceneCore.cpp
@@ -156,6 +156,18 @@ void FLightPrimitiveInteraction::Destroy(FLightPrimitiveInteraction* LightPrimit
 
 extern bool ShouldCreateObjectShadowForStationaryLight(const FLightSceneInfo* LightSceneInfo, const FPrimitiveSceneProxy* PrimitiveSceneProxy, bool bInteractionShadowMapped);
 
+static bool MobileRequiresStaticMeshUpdateOnLocalLightChange(const FStaticShaderPlatform Platform)
+{
+	extern bool MobileLocalLightsUseSinglePermutation();
+	extern bool MobileLocalLightsBufferPostprocessEnabled(const FStaticShaderPlatform Platform);
+	
+	if (!IsMobileDeferredShadingEnabled(Platform))
+	{
+		return MobileForwardEnableLocalLights(Platform) && !(MobileLocalLightsUseSinglePermutation() || MobileLocalLightsBufferPostprocessEnabled(Platform));
+	}
+	return false;
+}
+
 FLightPrimitiveInteraction::FLightPrimitiveInteraction(
 	FLightSceneInfo* InLightSceneInfo,
 	FPrimitiveSceneInfo* InPrimitiveSceneInfo,
@@ -250,7 +262,8 @@ FLightPrimitiveInteraction::FLightPrimitiveInteraction(
 			{
 				bMobileDynamicLocalLight = true;
 				PrimitiveSceneInfo->NumMobileDynamicLocalLights++;
-				if (PrimitiveSceneInfo->NumMobileDynamicLocalLights == 1)
+				if (PrimitiveSceneInfo->NumMobileDynamicLocalLights == 1 && 
+					MobileRequiresStaticMeshUpdateOnLocalLightChange(PrimitiveSceneInfo->Scene->GetShaderPlatform()))
 				{
 					// Update static meshes to choose the shader permutation with local lights.
 					PrimitiveSceneInfo->RequestStaticMeshUpdate();
@@ -322,7 +335,8 @@ FLightPrimitiveInteraction::~FLightPrimitiveInteraction()
 	if (bMobileDynamicLocalLight)
 	{
 		PrimitiveSceneInfo->NumMobileDynamicLocalLights--;
-		if (PrimitiveSceneInfo->NumMobileDynamicLocalLights == 0)
+		if (PrimitiveSceneInfo->NumMobileDynamicLocalLights == 0 &&
+			MobileRequiresStaticMeshUpdateOnLocalLightChange(PrimitiveSceneInfo->Scene->GetShaderPlatform()))
 		{
 			// Update static meshes to choose the shader permutation without local lights.
 			PrimitiveSceneInfo->RequestStaticMeshUpdate();
