Implement new shader permutation for [wavesize(16)] for Intel Arc GPUs

Adds a new r.TSR.WaveSize to manually override which WaveSize to use

Submitting on behalf of Hugo Lamarche

diff --git a/Engine/Shaders/Private/TemporalSuperResolution/TSRRejectShading.usf b/Engine/Shaders/Private/TemporalSuperResolution/TSRRejectShading.usf
--- a/Engine/Shaders/Private/TemporalSuperResolution/TSRRejectShading.usf
+++ b/Engine/Shaders/Private/TemporalSuperResolution/TSRRejectShading.usf
@@ -33,17 +33,25 @@
 
 //------------------------------------------------------- DERIVED
 	
-#if DIM_WAVE_SIZE == 64 || DIM_WAVE_SIZE == 32
+#if DIM_WAVE_SIZE == 64 || DIM_WAVE_SIZE == 32 || DIM_WAVE_SIZE == 16
 	/** Layouts wave in row major 32x32 tile so an entire row of pixels fits within a single wave, and waveops horizontally, but LDS vertically.
 	 * Each lane worth of a 1x4 of data per lane to minimise to trade VGPR pressure for both number of wave ops and LDS.
 	 */
 	#define GROUP_TILE_SIZE 32
 
-	#define LANE_STRIDE_X 1
-	#if CONFIG_COMPILE_FP16
-		#define LANE_STRIDE_Y 4
+	#if DIM_WAVE_SIZE == 16
+		// We need a LANE_STRIDE_X of at least 2 to have the entire row of pixels fits within a single wave.
+		#define LANE_STRIDE_X 2
+
+		// Only process 2 pixels per lane to reduce register pressure on Intel Arc Alchemist GPUs
+		#define LANE_STRIDE_Y 1
 	#else
-		#define LANE_STRIDE_Y 2
+		#define LANE_STRIDE_X 1
+		#if CONFIG_COMPILE_FP16
+			#define LANE_STRIDE_Y 4
+		#else
+			#define LANE_STRIDE_Y 2
+		#endif
 	#endif
 
 	#define WAVE_COUNT_X 1
diff --git a/Engine/Source/Runtime/Renderer/Private/PostProcess/TemporalSuperResolution.cpp b/Engine/Source/Runtime/Renderer/Private/PostProcess/TemporalSuperResolution.cpp
--- a/Engine/Source/Runtime/Renderer/Private/PostProcess/TemporalSuperResolution.cpp
+++ b/Engine/Source/Runtime/Renderer/Private/PostProcess/TemporalSuperResolution.cpp
@@ -70,6 +70,15 @@ TAutoConsoleVariable<int32> CVarTSRWaveOps(
 	TEXT("backend of DXC which is not great for editor startup."),
 	ECVF_RenderThreadSafe);
 
+TAutoConsoleVariable<int32> CVarTSRWaveSize(
+	TEXT("r.TSR.WaveSize"), 0,
+	TEXT("Overrides the WaveSize to use.\n")
+	TEXT(" 0: Automatic (default);\n")
+	TEXT(" 16: WaveSizeOps 16;\n")
+	TEXT(" 32: WaveSizeOps 32;\n")
+	TEXT(" 64: WaveSizeOps 64;\n"),
+	ECVF_RenderThreadSafe);
+
 TAutoConsoleVariable<int32> CVarTSR16BitVALU(
 	TEXT("r.TSR.16BitVALU"), 1,
 	TEXT("Whether to use 16bit VALU on platform that have bSupportsRealTypes=RuntimeDependent"),
@@ -605,7 +614,7 @@ class FTSRRejectShadingCS : public FTSRShader
 	DECLARE_GLOBAL_SHADER(FTSRRejectShadingCS);
 	SHADER_USE_PARAMETER_STRUCT(FTSRRejectShadingCS, FTSRShader);
 
-	class FWaveSizeOps : SHADER_PERMUTATION_SPARSE_INT("DIM_WAVE_SIZE", 0, 32, 64);
+	class FWaveSizeOps : SHADER_PERMUTATION_SPARSE_INT("DIM_WAVE_SIZE", 0, 16, 32, 64);
 	class FFlickeringDetectionDim : SHADER_PERMUTATION_BOOL("DIM_FLICKERING_DETECTION");
 	class FHistoryResurrectionDim : SHADER_PERMUTATION_BOOL("DIM_HISTORY_RESURRECTION");
 
@@ -645,6 +654,14 @@ class FTSRRejectShadingCS : public FTSRShader
 
 	static FPermutationDomain RemapPermutation(FPermutationDomain PermutationVector)
 	{
+		int32 WaveSize = PermutationVector.Get<FWaveSizeOps>();
+
+		// WaveSize is for Intel Arc GPU which also supports 16bits ops, so compiling WaveSize=16 32bit ops is useless and should instead fall back to WaveSize=0.
+		if (WaveSize == 16 && !PermutationVector.Get<FTSRShader::F16BitVALUDim>())
+		{
+			PermutationVector.Set<FWaveSizeOps>(0);
+		}
+
 		return PermutationVector;
 	}
 
@@ -1208,6 +1225,7 @@ FDefaultTemporalUpscaler::FOutputs AddTemporalSuperResolutionPasses(
 	const ERHIFeatureSupport WaveOpsSupport = FTSRShader::SupportsWaveOps(View.GetShaderPlatform());
 	const bool bSupportsLDS = FTSRShader::SupportsLDS(View.GetShaderPlatform());
 	const bool bUseWaveOps = (CVarTSRWaveOps.GetValueOnRenderThread() != 0 && GRHISupportsWaveOperations && bSupportsLDS && (WaveOpsSupport == ERHIFeatureSupport::RuntimeDependent || WaveOpsSupport == ERHIFeatureSupport::RuntimeGuaranteed)) || !bSupportsLDS;
+	const int32 WaveSizeOverride = bUseWaveOps ? CVarTSRWaveSize.GetValueOnAnyThread() : 0;
 	
 	// Whether to use 16bit VALU
 	const ERHIFeatureSupport VALU16BitSupport = FTSRShader::Supports16BitVALU(View.GetShaderPlatform());
@@ -2024,8 +2042,23 @@ FDefaultTemporalUpscaler::FOutputs AddTemporalSuperResolutionPasses(
 		FScreenPassTextureViewport TranslucencyViewport(
 			SeparateTranslucencyTexture->Desc.Extent, SeparateTranslucencyRect);
 
+		int32 WaveSizeOps = 0;
+
+		if (bUseWaveOps)
+		{
+			if ((WaveSizeOverride == 16 || WaveSizeOverride == 32 || WaveSizeOverride == 64) && WaveSizeOverride >= GRHIMinimumWaveSize && WaveSizeOverride <= GRHIMaximumWaveSize)
+			{
+				WaveSizeOps = WaveSizeOverride;
+			}
+			else
+			{
+				const int32 MinimumWaveSizeWithPermutation = FMath::Max(GRHIMinimumWaveSize, 16);
+				WaveSizeOps = MinimumWaveSizeWithPermutation >= 16 && MinimumWaveSizeWithPermutation <= 64 ? MinimumWaveSizeWithPermutation : 0;
+			}
+		}
+
 		FTSRRejectShadingCS::FPermutationDomain PermutationVector;
-		PermutationVector.Set<FTSRRejectShadingCS::FWaveSizeOps>(bUseWaveOps&& GRHIMinimumWaveSize >= 32 && GRHIMinimumWaveSize <= 64 ? GRHIMinimumWaveSize : 0);
+		PermutationVector.Set<FTSRRejectShadingCS::FWaveSizeOps>(WaveSizeOps);
 		PermutationVector.Set<FTSRRejectShadingCS::FFlickeringDetectionDim>(FlickeringFramePeriod > 0.0f);
 		PermutationVector.Set<FTSRRejectShadingCS::FHistoryResurrectionDim>(bCanResurrectHistory);
 		PermutationVector.Set<FTSRShader::F16BitVALUDim>(bUse16BitVALU);
