Setup MeshDrawCommand budgets and categories w/ tech art

diff --git a/Engine/Source/Runtime/Engine/Public/MeshDrawCommandStatsSettings.h b/Engine/Source/Runtime/Engine/Public/MeshDrawCommandStatsSettings.h
--- a/Engine/Source/Runtime/Engine/Public/MeshDrawCommandStatsSettings.h
+++ b/Engine/Source/Runtime/Engine/Public/MeshDrawCommandStatsSettings.h
@@ -22,9 +22,9 @@ struct FMeshDrawCommandStatsBudget
 	/** The category primitive budget. This is the maximum triangles expected, post-culling, summed across all passes. */
 	UPROPERTY(config, EditAnywhere, Category = Engine)
 	int32 PrimitiveBudget = 0;
-	/** If true, exports the category's stats to the CSVProfiler. */
+	/** The collection which contains this budget. */
 	UPROPERTY(config, EditAnywhere, Category = Engine)
-	bool bExportToCsv = false;
+	int32 Collection = 0;
 };
 
 /** User settings used by the MeshDrawCommandStats system. */
@@ -40,4 +40,7 @@ public:
 	/** The total primitive budget. This is the maximimum triangles expected, post-culling, summed across all passes. */
 	UPROPERTY(config, EditAnywhere, Category = Engine)
 	int32 TotalPrimitiveBudget = 0;
+	/** Which collection to export to CSV */
+	UPROPERTY(config, EditAnywhere, Category = Engine)
+	int32 CollectionForCsvProfiler = 1;
 };
diff --git a/Engine/Source/Runtime/Renderer/Private/MeshDrawCommandStats.cpp b/Engine/Source/Runtime/Renderer/Private/MeshDrawCommandStats.cpp
--- a/Engine/Source/Runtime/Renderer/Private/MeshDrawCommandStats.cpp
+++ b/Engine/Source/Runtime/Renderer/Private/MeshDrawCommandStats.cpp
@@ -30,26 +30,26 @@ DECLARE_DWORD_COUNTER_STAT(TEXT("InstanceCulling Indirect Rendered Instances"),
 DECLARE_DWORD_COUNTER_STAT(TEXT("Custom Indirect Rendered Primitives"), STAT_Culling_CustomIndirectNumPrimitives, STATGROUP_Culling);
 DECLARE_DWORD_COUNTER_STAT(TEXT("Custom Indirect Rendered Instances"), STAT_Culling_CustomIndirectNumInstances, STATGROUP_Culling);
 
-CSV_DEFINE_CATEGORY(MeshDrawCommandStats, true);
+CSV_DEFINE_CATEGORY(MeshDrawCommandStats, false);
 
-static TAutoConsoleVariable<int32> CVarShowMeshDrawCommandStats(
+enum class MeshDrawStatsCollection : int32
+{
+	None,
+	Pass,
+	User,
+};
+
+static TAutoConsoleVariable<int32> CVarMeshDrawCommandStats(
 	TEXT("r.MeshDrawCommands.Stats"),
 	0,
 	TEXT("Show on screen mesh draw command stats.\n")
 	TEXT("The stats for visible triangles are post GPU culling.\n")
-	TEXT(" 1 = Show stats per category. The stats are accumulated across passes.\n")
-	TEXT(" 2 = Show stats per pass.\n")
+	TEXT(" 1 = Show stats per pass.\n")
+	TEXT(" 2...N = Show the collection of stats matching the 'Collection' parameter in the ini file.\n")
 	TEXT("You can also use 'stat culling' to see global culling stats.\n"),
 	ECVF_RenderThreadSafe
 );
 
-static TAutoConsoleVariable<int32> CVarExportMeshDrawCommandStats(
-	TEXT("r.MeshDrawCommands.Export"),
-	false,
-	TEXT("Export mesh draw command stats to the CSVProfiler.\n"),
-	ECVF_RenderThreadSafe
-);
-
 FMeshDrawCommandStatsManager::FFrameData::~FFrameData()
 {
 	// Collect set of unique readback buffers for deletion (can be shared between MDCs and passes)
@@ -146,20 +146,24 @@ FMeshDrawCommandStatsManager::FMeshDrawCommandStatsManager()
 	// Is it fine to keep the screen message delegate always registered even if we are not showing anything?
 	ScreenMessageDelegate = FRendererOnScreenNotification::Get().AddLambda([this](TMultiMap<FCoreDelegates::EOnScreenMessageSeverity, FText >& OutMessages)
 		{	
-			const bool bShowStats = CVarShowMeshDrawCommandStats->GetInt() == 0 ? false : true;
+			const bool bShowStats = CVarMeshDrawCommandStats->GetInt() != (int)MeshDrawStatsCollection::None;
 			if (bShowStats)
 			{
 				OutMessages.Add(FCoreDelegates::EOnScreenMessageSeverity::Info, FText::FromString(FString::Printf(TEXT("MeshDrawCommandStats (Triangles / Budget - Category):"), Stats.TotalPrimitives / 1000)));
+				int Collection = CVarMeshDrawCommandStats->GetInt();
 
 				// Show budgeted stats first.
 				const UMeshDrawCommandStatsSettings* Settings = GetDefault<UMeshDrawCommandStatsSettings>();
 				for (FMeshDrawCommandStatsBudget const& CategoryBudget : Settings->Budgets)
 				{
-					uint64* PrimitiveCount = BudgetedPrimitives.Find(CategoryBudget.CategoryName);
-					if (PrimitiveCount && *PrimitiveCount > 0)
+					if (CategoryBudget.Collection == Collection)
 					{
-						FCoreDelegates::EOnScreenMessageSeverity Severity = CategoryBudget.PrimitiveBudget < *PrimitiveCount ? FCoreDelegates::EOnScreenMessageSeverity::Warning : FCoreDelegates::EOnScreenMessageSeverity::Info;
-						OutMessages.Add(Severity, FText::FromString(FString::Printf(TEXT("%5dK / %5dK - %s"), *PrimitiveCount / 1000, CategoryBudget.PrimitiveBudget / 1000, *(CategoryBudget.CategoryName.ToString()))));
+						uint64* PrimitiveCount = BudgetedPrimitives.Find(CategoryBudget.CategoryName);
+						if (PrimitiveCount && *PrimitiveCount > 0)
+						{
+							FCoreDelegates::EOnScreenMessageSeverity Severity = CategoryBudget.PrimitiveBudget < *PrimitiveCount ? FCoreDelegates::EOnScreenMessageSeverity::Warning : FCoreDelegates::EOnScreenMessageSeverity::Info;
+							OutMessages.Add(Severity, FText::FromString(FString::Printf(TEXT("%5dK / %5dK - %s"), *PrimitiveCount / 1000, CategoryBudget.PrimitiveBudget / 1000, *(CategoryBudget.CategoryName.ToString()))));
+						}
 					}
 				}
 
@@ -261,7 +265,7 @@ void FMeshDrawCommandStatsManager::Update()
 
 	++CurrentFrameNumber;
 
-	const bool bShowPassNameStats = CVarShowMeshDrawCommandStats->GetInt() == 2 ? true : false;
+	const bool bShowPassNameStats = CVarMeshDrawCommandStats->GetInt() == (int)MeshDrawStatsCollection::Pass;
 
 	FScopeLock ScopeLock(&FrameDataCS);
 
@@ -386,26 +390,28 @@ void FMeshDrawCommandStatsManager::Update()
 	SET_DWORD_STAT(STAT_Culling_CustomIndirectNumInstances, Stats.CustomIndirectInstances);
 
 	// Collect stats during the next frame (check if STATGROUP_Culling is also visible somehow)
-	const bool bShowStats = CVarShowMeshDrawCommandStats->GetInt() == 0 ? false : true;
+	const bool bShowStats = CVarMeshDrawCommandStats->GetInt() != (int)MeshDrawStatsCollection::None;
 	bCollectStats = bShowStats || bRequestDumpStats;
 
 #if CSV_PROFILER
-	bCollectStats |= (FCsvProfiler::Get()->IsCapturing_Renderthread() && CVarExportMeshDrawCommandStats->GetBool());
+	const bool bCsvExport = FCsvProfiler::Get()->IsCapturing_Renderthread() && FCsvProfiler::Get()->IsCategoryEnabled(CSV_CATEGORY_INDEX(MeshDrawCommandStats));
+	bCollectStats |= bCsvExport;
 #endif
 
 	if (bCollectStats)
 	{
 		// First time - Build associative map for quick Stat -> Budget lookup
-		if (Budgets.IsEmpty())
+		if (!StatCollections.Num())
 		{
 			const UMeshDrawCommandStatsSettings* Settings = GetDefault<UMeshDrawCommandStatsSettings>();	
 			for (const FMeshDrawCommandStatsBudget& CategoryBudget : Settings->Budgets)
 			{
-				Budgets.Add(CategoryBudget.CategoryName, CategoryBudget.CategoryName);
+				StatCollection& Collection = StatCollections.FindOrAdd(CategoryBudget.Collection);
+				Collection.Add(CategoryBudget.CategoryName, CategoryBudget.CategoryName);
 
 				for (FName Name : CategoryBudget.LinkedStatNames)
 				{
-					Budgets.Add(Name, CategoryBudget.CategoryName);
+					Collection.Add(Name, CategoryBudget.CategoryName);
 				}
 			}
 		}
@@ -414,18 +420,32 @@ void FMeshDrawCommandStatsManager::Update()
 		BudgetedPrimitives.Reset();
 		UntrackedPrimitives.Reset();
 
-		for (const FStats::FCategoryStats& CategoryStat : Stats.CategoryStats)
-		{ 
-			FName* BudgetName = Budgets.Find(CategoryStat.CategoryName);
-			TMap<FName, uint64>* Map = BudgetName ? &BudgetedPrimitives : &UntrackedPrimitives;
+		int CollectionIdx = CVarMeshDrawCommandStats->GetInt();
 
-			uint64& Count = Map->FindOrAdd(BudgetName ? *BudgetName : CategoryStat.CategoryName);
-			Count += CategoryStat.PrimitiveCount;
+	#if CSV_PROFILER // If capturing for CSV, override the collection to the one requested in the ini file
+		if (FCsvProfiler::Get()->IsCapturing_Renderthread() && FCsvProfiler::Get()->IsCategoryEnabled(CSV_CATEGORY_INDEX(MeshDrawCommandStats)))
+		{
+			CollectionIdx = GetDefault<UMeshDrawCommandStatsSettings>()->CollectionForCsvProfiler;
+		}
+	#endif
+
+		StatCollection* Collection = StatCollections.Find(CollectionIdx);
+
+		if (Collection || CollectionIdx == (int)MeshDrawStatsCollection::Pass)
+		{
+			for (const FStats::FCategoryStats& CategoryStat : Stats.CategoryStats)
+			{ 
+				FName* BudgetName = Collection ? Collection->Find(CategoryStat.CategoryName) : nullptr;
+				TMap<FName, uint64>* Map = BudgetName ? &BudgetedPrimitives : &UntrackedPrimitives;
+
+				uint64& Count = Map->FindOrAdd(BudgetName ? *BudgetName : CategoryStat.CategoryName);
+				Count += CategoryStat.PrimitiveCount;
+			}
 		}
 	}
 
 #if CSV_PROFILER
-	if (FCsvProfiler::Get()->IsCapturing_Renderthread() && CVarExportMeshDrawCommandStats->GetBool())
+	if (bCsvExport)
 	{
 		// Output Budget totals
 		for (const TPair<FName, uint64>& Pair : BudgetedPrimitives)
diff --git a/Engine/Source/Runtime/Renderer/Private/MeshDrawCommandStats.h b/Engine/Source/Runtime/Renderer/Private/MeshDrawCommandStats.h
--- a/Engine/Source/Runtime/Renderer/Private/MeshDrawCommandStats.h
+++ b/Engine/Source/Runtime/Renderer/Private/MeshDrawCommandStats.h
@@ -164,6 +164,9 @@ private:
 		return Frames.Last();
 	}
 
+	using StatCollection = TMap<FName, FName>;
+	using StatCollectionMap = TMap<int32, StatCollection>;
+
 	/** Dump given frame data stats to csv file on disc in profiling folder */
 	void DumpStats(FFrameData* FrameData);
 
@@ -174,7 +177,7 @@ private:
 	FCriticalSection FrameDataCS;			
 	TArray<FFrameData*> Frames;				//< All active frames (contains the frame for which we are collecting stats now and all frames waiting for GPU readback)
 	FStats Stats;							//< Last updated frame stats
-	TMap<FName, FName> Budgets;				//< LinkedStatName to Budget CategoryName
+	StatCollectionMap StatCollections; 		//< Per Collection LinkedStatName to Budget CategoryName
 	TMap<FName, uint64> BudgetedPrimitives;	//< Budget CategoryName to Total Primitive Count
 	TMap<FName, uint64> UntrackedPrimitives;//< Primitives which aren't tracked by any Budgets
 
