GlobalDistanceField - fixed crash when opening material editor

diff --git a/Engine/Source/Runtime/Renderer/Private/GlobalDistanceField.cpp b/Engine/Source/Runtime/Renderer/Private/GlobalDistanceField.cpp
--- a/Engine/Source/Runtime/Renderer/Private/GlobalDistanceField.cpp
+++ b/Engine/Source/Runtime/Renderer/Private/GlobalDistanceField.cpp
@@ -2729,6 +2729,7 @@ void UpdateGlobalDistanceFieldCache(
 		}
 
 		// Return to the free list
+		if (PackedClipmaps.Num() > 0)
 		{
 			FRDGBufferRef FreeListReturnIndirectArgBuffer = GraphBuilder.CreateBuffer(FRDGBufferDesc::CreateIndirectDesc<FRHIDispatchIndirectParameters>(1), TEXT("GlobalDistanceField.FreeListReturnIndirectArgs"));
 
@@ -2833,6 +2834,7 @@ void UpdateGlobalDistanceFieldCache(
 		}
 
 		// Composite mesh SDFs into allocated object grid pages
+		if (PageObjectGridBuffer)
 		{
 			const FRDGBufferUAVRef PageObjectGridBufferUAV = GraphBuilder.CreateUAV(PageObjectGridBuffer, ERDGUnorderedAccessViewFlags::SkipBarrier);
 
