Fix crash during RVT compression setup.
Due to falling through case due to missing a break.

diff --git a/Engine/Source/Runtime/Renderer/Private/VT/RuntimeVirtualTextureRender.cpp b/Engine/Source/Runtime/Renderer/Private/VT/RuntimeVirtualTextureRender.cpp
--- a/Engine/Source/Runtime/Renderer/Private/VT/RuntimeVirtualTextureRender.cpp
+++ b/Engine/Source/Runtime/Renderer/Private/VT/RuntimeVirtualTextureRender.cpp
@@ -1181,6 +1181,7 @@ namespace RuntimeVirtualTexture
 			break;
 		case ERuntimeVirtualTextureMaterialType::BaseColor_Normal_Specular_Mask_YCoCg:
 			AddCompressPass<ERuntimeVirtualTextureMaterialType::BaseColor_Normal_Specular_Mask_YCoCg>(GraphBuilder, FeatureLevel, Parameters, GroupCount);
+			break;
 		case ERuntimeVirtualTextureMaterialType::Displacement:
 			AddCompressPass<ERuntimeVirtualTextureMaterialType::Displacement>(GraphBuilder, FeatureLevel, Parameters, GroupCount);
 			break;
