Cvar to control Bend SSS parameter SurfaceThickness.
- once both methods to calculate contact shadows use the same surface thickness logic this value should be exposed on the Directional Light component, but for now we use a cvar.

diff --git a/Engine/Shaders/Private/ScreenSpaceShadows.usf b/Engine/Shaders/Private/ScreenSpaceShadows.usf
--- a/Engine/Shaders/Private/ScreenSpaceShadows.usf
+++ b/Engine/Shaders/Private/ScreenSpaceShadows.usf
@@ -21,7 +21,7 @@ Texture2D<float> DepthTexture;
 RWTexture2D<float> OutputTexture;
 SamplerState PointBorderSampler;
 
-DispatchParameters SetupBendDispatchParameters(float NearDepthVal, float FarDepthVal)
+DispatchParameters SetupBendDispatchParameters(float NearDepthVal, float FarDepthVal, float SurfaceThickness)
 {
 	DispatchParameters Parameters;
 	Parameters.SetDefaults();
@@ -37,6 +37,8 @@ DispatchParameters SetupBendDispatchParameters(float NearDepthVal, float FarDept
 	Parameters.DepthTexture = DepthTexture;
 	Parameters.PointBorderSampler = PointBorderSampler;
 
+	Parameters.SurfaceThickness = SurfaceThickness;
+
 	return Parameters;
 }
 
@@ -65,6 +67,7 @@ float ContactShadowLength;
 uint bContactShadowLengthInWS;
 float ContactShadowCastingIntensity;
 float ContactShadowNonCastingIntensity;
+float SurfaceThickness;
 
 uint DownsampleFactor;
 int4 ScissorRectMinAndSize;
@@ -144,7 +147,7 @@ void ScreenSpaceShadowsBendCS(
 	uint3 DispatchThreadId : SV_DispatchThreadID,
 	uint3 GroupThreadId : SV_GroupThreadID)
 {
-	DispatchParameters Parameters = SetupBendDispatchParameters(NearDepthValue, FarDepthValue);
+	DispatchParameters Parameters = SetupBendDispatchParameters(NearDepthValue, FarDepthValue, SurfaceThickness);
 
 	float Result;
 	int2 ResultXY;
diff --git a/Engine/Source/Runtime/Renderer/Private/Shadows/ScreenSpaceShadows.cpp b/Engine/Source/Runtime/Renderer/Private/Shadows/ScreenSpaceShadows.cpp
--- a/Engine/Source/Runtime/Renderer/Private/Shadows/ScreenSpaceShadows.cpp
+++ b/Engine/Source/Runtime/Renderer/Private/Shadows/ScreenSpaceShadows.cpp
@@ -33,6 +33,14 @@ static FAutoConsoleVariableRef CVarContactShadowsMethod(
 	ECVF_Scalability | ECVF_RenderThreadSafe
 );
 
+static float GBendShadowsOverrideSurfaceThickness = 0.005f;
+static FAutoConsoleVariableRef CVarBendShadowsOverrideSurfaceThickness(
+	TEXT("r.ContactShadows.Bend.OverrideSurfaceThickness"),
+	GBendShadowsOverrideSurfaceThickness,
+	TEXT("How thick the surface represented by a pixel is assumed to be when determining whether a ray intersects it."),
+	ECVF_Scalability | ECVF_RenderThreadSafe
+);
+
 extern void GetLightContactShadowParameters(const FLightSceneProxy* Proxy, float& OutLength, bool& bOutLengthInWS, float& OutCastingIntensity, float& OutNonCastingIntensity);
 
 const int32 GScreenSpaceShadowsTileSizeX = 8;
@@ -99,6 +107,7 @@ class FScreenSpaceShadowsBendCS : public FGlobalShader
 		SHADER_PARAMETER(uint32, bContactShadowLengthInWS)
 		SHADER_PARAMETER(float, ContactShadowCastingIntensity)
 		SHADER_PARAMETER(float, ContactShadowNonCastingIntensity)
+		SHADER_PARAMETER(float, SurfaceThickness)
 		SHADER_PARAMETER(FIntRect, ScissorRectMinAndSize)
 		SHADER_PARAMETER(uint32, DownsampleFactor)
 		SHADER_PARAMETER(FVector2f, InvDepthTextureSize)
@@ -363,6 +372,7 @@ void RenderScreenSpaceShadowsBend(
 			PassParameters->bContactShadowLengthInWS = bContactShadowLengthInWS;
 			PassParameters->ContactShadowCastingIntensity = ContactShadowCastingIntensity;
 			PassParameters->ContactShadowNonCastingIntensity = ContactShadowNonCastingIntensity;
+			PassParameters->SurfaceThickness = GBendShadowsOverrideSurfaceThickness;
 
 			PassParameters->LightCoordinate = FVector4f(DispatchList.LightCoordinate_Shader[0], DispatchList.LightCoordinate_Shader[1], DispatchList.LightCoordinate_Shader[2], DispatchList.LightCoordinate_Shader[3]);
 			PassParameters->WaveOffset = FIntVector(Dispatch.WaveOffset_Shader[0], Dispatch.WaveOffset_Shader[1], 0);
