Fixed GDME refactor fallout.
 - Sync virtual texture task before GDME runs; works around race condition between GDME and VT update.
 - Moved GDME read buffer to commit immediately after GDME instead of after GPU scene updates (fixes RDG immediate crash).

diff --git a/Engine/Source/Runtime/Renderer/Private/DeferredShadingRenderer.cpp b/Engine/Source/Runtime/Renderer/Private/DeferredShadingRenderer.cpp
--- a/Engine/Source/Runtime/Renderer/Private/DeferredShadingRenderer.cpp
+++ b/Engine/Source/Runtime/Renderer/Private/DeferredShadingRenderer.cpp
@@ -3256,9 +3256,6 @@ void FDeferredShadingSceneRenderer::Render(FRDGBuilder& GraphBuilder)
 	{
 		RDG_RHI_GPU_STAT_SCOPE(GraphBuilder, VisibilityCommands);
 		EndInitViews(GraphBuilder, LumenFrameTemporaries, InstanceCullingManager, ExternalAccessQueue, InitViewTaskDatas);
-
-		SCOPE_CYCLE_COUNTER(STAT_FDeferredShadingSceneRenderer_FGlobalDynamicVertexBuffer_Commit);
-		DynamicReadBufferForInitViews.Commit(GraphBuilder.RHICmdList);
 	}
 
 	// Substrate initialisation is always run even when not enabled.
diff --git a/Engine/Source/Runtime/Renderer/Private/SceneRendering.cpp b/Engine/Source/Runtime/Renderer/Private/SceneRendering.cpp
--- a/Engine/Source/Runtime/Renderer/Private/SceneRendering.cpp
+++ b/Engine/Source/Runtime/Renderer/Private/SceneRendering.cpp
@@ -2503,8 +2503,8 @@ TGlobalResource<FGlobalDynamicReadBuffer> FSceneRenderer::DynamicReadBufferForSh
 FSceneRenderer::FSceneRenderer(const FSceneViewFamily* InViewFamily, FHitProxyConsumer* HitProxyConsumer)
 :	Scene(CheckPointer(InViewFamily->Scene)->GetRenderScene())
 ,	ViewFamily(*CheckPointer(InViewFamily))
-,	MeshCollector(InViewFamily->GetFeatureLevel(), Allocator)
-,	EditorMeshCollector(InViewFamily->GetFeatureLevel(), Allocator)
+,	MeshCollector(InViewFamily->GetFeatureLevel(), Allocator, FMeshElementCollector::ECommitFlags::DeferAll)
+,	EditorMeshCollector(InViewFamily->GetFeatureLevel(), Allocator, FMeshElementCollector::ECommitFlags::DeferAll)
 ,	VirtualShadowMapArray(*CheckPointer(Scene))
 ,	bHasRequestedToggleFreeze(false)
 ,	bUsedPrecomputedVisibility(false)
diff --git a/Engine/Source/Runtime/Renderer/Private/SceneVisibility.cpp b/Engine/Source/Runtime/Renderer/Private/SceneVisibility.cpp
--- a/Engine/Source/Runtime/Renderer/Private/SceneVisibility.cpp
+++ b/Engine/Source/Runtime/Renderer/Private/SceneVisibility.cpp
@@ -3756,11 +3756,6 @@ void FVisibilityTaskData::FinishGatherDynamicMeshElements(FVirtualTextureUpdater
 {
 	check(IsInRenderingThread());
 
-	// Sync the virtual texture update task before finishing the mesh collectors. Render proxies can register new
-	// materials which require evaluating uniform expression caches, which can contain virtual textures. Newly allocated
-	// virtual textures are processed later.
-	FVirtualTextureSystem::Get().WaitForTasks(VirtualTextureUpdater);
-
 	if (DynamicMeshElements.CommandPipe)
 	{
 		SCOPED_NAMED_EVENT(WaitForGatherDynamicMeshElements, FColor::Magenta);
@@ -3781,8 +3776,14 @@ void FVisibilityTaskData::FinishGatherDynamicMeshElements(FVirtualTextureUpdater
 		EditorMeshCollector.Finish();
 	}
 
+	// Sync the virtual texture update task before finishing the mesh collectors. Render proxies can register new
+	// materials which require evaluating uniform expression caches, which can contain virtual textures. Newly allocated
+	// virtual textures are processed later.
+	FVirtualTextureSystem::Get().WaitForTasks(VirtualTextureUpdater);
+
 	DynamicMeshElements.DynamicVertexBuffer.Commit();
 	DynamicMeshElements.DynamicIndexBuffer.Commit();
+	SceneRenderer.DynamicReadBufferForInitViews.Commit(RHICmdList);
 }
 
 void FVisibilityTaskData::GatherDynamicMeshElementsForPrimitive(int32 PrimitiveIndex, uint8 ViewMask)
diff --git a/Engine/Source/Runtime/Renderer/Private/VT/VirtualTextureSystem.cpp b/Engine/Source/Runtime/Renderer/Private/VT/VirtualTextureSystem.cpp
--- a/Engine/Source/Runtime/Renderer/Private/VT/VirtualTextureSystem.cpp
+++ b/Engine/Source/Runtime/Renderer/Private/VT/VirtualTextureSystem.cpp
@@ -161,11 +161,7 @@ static TAutoConsoleVariable<int32> CVarVTCsvStats(
 );
 static TAutoConsoleVariable<int32> CVarVTAsyncPageRequestTask(
 	TEXT("r.VT.AsyncPageRequestTask"),
-#if WITH_EDITOR // Turn off async VT as quick fix for race condition seen in editor wireframe mode.
-	0,
-#else
 	1,
-#endif
 	TEXT("Performs VT page requests on an async task."),
 	ECVF_RenderThreadSafe
 );
