Made ShadowCacheInvalidationBehavior override HasDeformableMesh such that it can suppress that behavior if needed.

diff --git a/Engine/Source/Runtime/Renderer/Private/VirtualShadowMaps/VirtualShadowMapCacheManager.cpp b/Engine/Source/Runtime/Renderer/Private/VirtualShadowMaps/VirtualShadowMapCacheManager.cpp
--- a/Engine/Source/Runtime/Renderer/Private/VirtualShadowMaps/VirtualShadowMapCacheManager.cpp
+++ b/Engine/Source/Runtime/Renderer/Private/VirtualShadowMaps/VirtualShadowMapCacheManager.cpp
@@ -851,7 +851,8 @@ void FVirtualShadowMapPerLightCacheEntry::OnPrimitiveRendered(const FPrimitiveSc
 	if (GVSMCacheDeformableMeshesInvalidate != 0)
 	{
 		// Deformable mesh primitives need to trigger invalidation (even if they did not move) or we get artifacts, for example skinned meshes that are animating but not currently moving.
-		if (PrimitiveSceneInfo->Proxy->HasDeformableMesh())
+		// Skip if the invalidation mode is NOT auto (because Always will do it elsewhere & the others should prevent this).
+		if (PrimitiveSceneInfo->Proxy->HasDeformableMesh() && PrimitiveSceneInfo->Proxy->GetShadowCacheInvalidationBehavior() == EShadowCacheInvalidationBehavior::Auto)
 		{
 			PrimitiveInstancesToInvalidate.Add(FInstanceRange{ 
 				PrimitiveSceneInfo->GetInstanceSceneDataOffset(),
