Proxies with zero bounds will fail visibility early.

diff --git a/Engine/Source/Runtime/Renderer/Private/SceneVisibility.cpp b/Engine/Source/Runtime/Renderer/Private/SceneVisibility.cpp
--- a/Engine/Source/Runtime/Renderer/Private/SceneVisibility.cpp
+++ b/Engine/Source/Runtime/Renderer/Private/SceneVisibility.cpp
@@ -676,6 +676,9 @@ static int32 FrustumCull(const FScene& Scene, FViewInfo& View, FFrustumCullingFl
 
 			const FPrimitiveBounds& RESTRICT Bounds = Scene.PrimitiveBounds[Index];
 
+			// Zero sized bounds indicates that we are not visible.
+			bIsVisible &= Bounds.BoxSphereBounds.SphereRadius > 0;
+
 			// Handle primitives that are not always visible.
 			if (Flags.bShouldVisibilityCull && bIsVisible && !IsAlwaysVisible(Scene, Index))
 			{
