Fix distance field shadow on SLW.

diff --git a/Engine/Shaders/Private/DistanceFieldShadowing.usf b/Engine/Shaders/Private/DistanceFieldShadowing.usf
--- a/Engine/Shaders/Private/DistanceFieldShadowing.usf
+++ b/Engine/Shaders/Private/DistanceFieldShadowing.usf
@@ -810,8 +810,11 @@ void ShadowTileVS(
 	out float4 Position : SV_POSITION)
 {
 	uint TileData = TileListData[InstanceId.x];
-	uint2 TileOrigin = uint2((TileData & 0x0000FFFF) * WORK_TILE_SIZE, (TileData >> 16) * WORK_TILE_SIZE);
-
+#if PERMUTATION_TILE_TYPE == 1
+	const uint2 TileOrigin = UnpackTileCoord12bits(TileData) * WORK_TILE_SIZE;
+#else
+	const uint2 TileOrigin = UnpackTileCoord16bits(TileData) * WORK_TILE_SIZE;
+#endif
 	uint2 TileVertex = TileOrigin;
 	TileVertex.x += VertexId == 1 || VertexId == 2 || VertexId == 4 ? WORK_TILE_SIZE : 0;
 	TileVertex.y += VertexId == 2 || VertexId == 4 || VertexId == 5 ? WORK_TILE_SIZE : 0;
diff --git a/Engine/Source/Runtime/Renderer/Private/DistanceFieldShadowing.cpp b/Engine/Source/Runtime/Renderer/Private/DistanceFieldShadowing.cpp
--- a/Engine/Source/Runtime/Renderer/Private/DistanceFieldShadowing.cpp
+++ b/Engine/Source/Runtime/Renderer/Private/DistanceFieldShadowing.cpp
@@ -380,7 +380,8 @@ class FShadowTileVS : public FGlobalShader
 	DECLARE_GLOBAL_SHADER(FShadowTileVS);
 	SHADER_USE_PARAMETER_STRUCT(FShadowTileVS, FGlobalShader);
 
-	using FPermutationDomain = TShaderPermutationDomain<>;
+	class FTileType : SHADER_PERMUTATION_INT("PERMUTATION_TILE_TYPE",2);
+	using FPermutationDomain = TShaderPermutationDomain<FTileType>;
 
 	BEGIN_SHADER_PARAMETER_STRUCT(FParameters, )
 		SHADER_PARAMETER_STRUCT_REF(FViewUniformShaderParameters, ViewUniformBuffer)
@@ -1089,7 +1090,9 @@ void FProjectedShadowInfo::RenderRayTracedDistanceFieldProjection(
 		{
 			check(TiledShadowRendering->TileSize == FShadowTileVS::GetTileSize());
 
-			auto VertexShader = View.ShaderMap->GetShader<FShadowTileVS>();
+			FShadowTileVS::FPermutationDomain VSPermutationVector;
+			VSPermutationVector.Set<FShadowTileVS::FTileType>(TiledShadowRendering->TileType == FTiledShadowRendering::ETileType::Tile16bits ? 0 : 1);
+			auto VertexShader = View.ShaderMap->GetShader<FShadowTileVS>(VSPermutationVector);
 			ClearUnusedGraphResources(VertexShader, &PassParameters->VS);
 
 			GraphBuilder.AddPass(
diff --git a/Engine/Source/Runtime/Renderer/Private/ShadowRendering.h b/Engine/Source/Runtime/Renderer/Private/ShadowRendering.h
--- a/Engine/Source/Runtime/Renderer/Private/ShadowRendering.h
+++ b/Engine/Source/Runtime/Renderer/Private/ShadowRendering.h
@@ -206,9 +206,15 @@ inline bool IsShadowCacheModeOcclusionQueryable(EShadowDepthCacheMode CacheMode)
 
 struct FTiledShadowRendering
 {
+	enum class ETileType : uint8
+	{
+		Tile16bits,
+		Tile12bits
+	};
 	FRDGBufferRef		DrawIndirectParametersBuffer;
 	FRDGBufferSRVRef	TileListDataBufferSRV;
 	uint32				TileSize;
+	ETileType			TileType = ETileType::Tile16bits;
 };
 
 class FShadowMapRenderTargets
diff --git a/Engine/Source/Runtime/Renderer/Private/SingleLayerWaterRendering.cpp b/Engine/Source/Runtime/Renderer/Private/SingleLayerWaterRendering.cpp
--- a/Engine/Source/Runtime/Renderer/Private/SingleLayerWaterRendering.cpp
+++ b/Engine/Source/Runtime/Renderer/Private/SingleLayerWaterRendering.cpp
@@ -968,6 +968,7 @@ void FDeferredShadingSceneRenderer::RenderSingleLayerWaterReflections(
 						TiledShadowRendering.DrawIndirectParametersBuffer = TiledScreenSpaceReflection.DrawIndirectParametersBuffer;
 						TiledShadowRendering.TileListDataBufferSRV = TiledScreenSpaceReflection.TileListDataBufferSRV;
 						TiledShadowRendering.TileSize = TiledScreenSpaceReflection.TileSize;
+						TiledShadowRendering.TileType = FTiledShadowRendering::ETileType::Tile12bits;
 					}
 
 					const bool bProjectingForForwardShading = false;
