Fix for new water info texture rendering path pulling in unrelated materials, causing issues if they require a more complex setup (uniforms, scene textures etc). Since this whole path is a temporary solution, this CL simply introduces an allow list with known good (and required) water info materials (currently only a single one). This CL also undoes the previous fix which simply disabled the water info mesh processor if the runtime cvar for the new rendering path is disabled. That fix broke water rendering if the new path is enabled at runtime after draw calls were already cached/processed and the cvar was disabled at the time.

#rnx

diff --git a/Engine/Source/Runtime/Renderer/Private/WaterInfoTextureRendering.cpp b/Engine/Source/Runtime/Renderer/Private/WaterInfoTextureRendering.cpp
--- a/Engine/Source/Runtime/Renderer/Private/WaterInfoTextureRendering.cpp
+++ b/Engine/Source/Runtime/Renderer/Private/WaterInfoTextureRendering.cpp
@@ -109,12 +109,6 @@ public:
 
 IMPLEMENT_GLOBAL_SHADER(FWaterInfoTextureBlurPS, "/Engine/Private/WaterInfoTextureBlur.usf", "Main", SF_Pixel);
 
-static bool IsCustomWaterInfoTextureRenderEnabled()
-{
-	static const IConsoleVariable* WaterRenderMethodCVar = IConsoleManager::Get().FindConsoleVariable(TEXT("r.Water.WaterInfo.RenderMethod"));
-	return WaterRenderMethodCVar && WaterRenderMethodCVar->GetInt();
-}
-
 /**
 * MeshPassProcessor for the "color" pass required for generating the water info texture. The associated pass draws water body meshes with an unlit material
 * in order to write water surface depth, river velocity and possibly other data too.
@@ -137,6 +131,7 @@ private:
 		const FMaterial& Material);
 
 	FMeshPassProcessorRenderState PassDrawRenderState;
+	TArray<const TCHAR*, TInlineAllocator<1>> MaterialAllowList;
 	bool bIsMobile;
 };
 
@@ -147,11 +142,15 @@ FWaterInfoTexturePassMeshProcessor::FWaterInfoTexturePassMeshProcessor(const FSc
 	PassDrawRenderState.SetBlendState(TStaticBlendState<>::GetRHI());
 	PassDrawRenderState.SetDepthStencilAccess(FExclusiveDepthStencil::DepthWrite_StencilNop);
 	PassDrawRenderState.SetDepthStencilState(TStaticDepthStencilState<true, CF_DepthNearOrEqual>::GetRHI());
+
+	// HACK: This whole path for rendering the water info texture is a temporary solution, so in order to avoid supporting generic material setups, we use an allow list to restrict
+	// the set of supported materials to known working (and required) materials.
+	MaterialAllowList.Add(TEXT("DrawWaterInfo"));
 }
 
 void FWaterInfoTexturePassMeshProcessor::AddMeshBatch(const FMeshBatch& RESTRICT MeshBatch, uint64 BatchElementMask, const FPrimitiveSceneProxy* RESTRICT PrimitiveSceneProxy, int32 StaticMeshId)
 {
-	if (!MeshBatch.bUseForMaterial || !IsCustomWaterInfoTextureRenderEnabled())
+	if (!MeshBatch.bUseForMaterial)
 	{
 		return;
 	}
@@ -160,7 +159,7 @@ void FWaterInfoTexturePassMeshProcessor::AddMeshBatch(const FMeshBatch& RESTRICT
 	while (MaterialRenderProxy)
 	{
 		const FMaterial* Material = MaterialRenderProxy->GetMaterialNoFallback(FeatureLevel);
-		if (Material)
+		if (Material && MaterialAllowList.Contains(Material->GetAssetName()))
 		{
 			if (TryAddMeshBatch(MeshBatch, BatchElementMask, PrimitiveSceneProxy, StaticMeshId, *MaterialRenderProxy, *Material))
 			{
@@ -177,7 +176,7 @@ void FWaterInfoTexturePassMeshProcessor::CollectPSOInitializers(const FSceneText
 	// Try to reduce the number of materials considered for this mesh pass:
 	// Materials for drawing the water info texture are supposed to be unlit (they write velocity and possibly other data into Emissive).
 	// They also need to be applied to meshes (MD_Surface) and we can safely exclude sky materials which also use the unlit shading model.
-	if (IsCustomWaterInfoTextureRenderEnabled() && Material.GetShadingModels().IsUnlit() && Material.GetMaterialDomain() == MD_Surface && !Material.IsSky())
+	if (Material.GetShadingModels().IsUnlit() && Material.GetMaterialDomain() == MD_Surface && !Material.IsSky() && MaterialAllowList.Contains(Material.GetAssetName()))
 	{
 		// Determine the mesh's material and blend mode.
 		const FMeshDrawingPolicyOverrideSettings OverrideSettings = ComputeMeshOverrideSettings(PreCacheParams);
@@ -242,10 +241,6 @@ void FWaterInfoTexturePassMeshProcessor::CollectPSOInitializers(const FSceneText
 			Shaders.TryGetVertexShader(PassShaders.VertexShader);
 			Shaders.TryGetPixelShader(PassShaders.PixelShader);
 	
-			// TODO: find the correct values for mobile.
-			uint8 SubpassIndex = 0;
-			ESubpassHint SubpassHint = ESubpassHint::None;
-
 			AddGraphicsPipelineStateInitializer(
 				VertexFactoryData,
 				Material,
@@ -256,8 +251,6 @@ void FWaterInfoTexturePassMeshProcessor::CollectPSOInitializers(const FSceneText
 				MeshCullMode,
 				(EPrimitiveType)PreCacheParams.PrimitiveType,
 				EMeshPassFeatures::Default,
-				SubpassHint,
-				SubpassIndex,
 				true /*bRequired*/,
 				PSOInitializers);
 		}
@@ -433,11 +426,6 @@ FWaterInfoTextureDepthPassMeshProcessor::FWaterInfoTextureDepthPassMeshProcessor
 
 void FWaterInfoTextureDepthPassMeshProcessor::AddMeshBatch(const FMeshBatch& RESTRICT MeshBatch, uint64 BatchElementMask, const FPrimitiveSceneProxy* RESTRICT PrimitiveSceneProxy, int32 StaticMeshId)
 {
-	if(!IsCustomWaterInfoTextureRenderEnabled())
-	{
-		return;
-	}
-
 	const FMaterialRenderProxy* MaterialRenderProxy = MeshBatch.MaterialRenderProxy;
 	while (MaterialRenderProxy)
 	{
@@ -553,7 +541,7 @@ bool FWaterInfoTextureDepthPassMeshProcessor::Process(
 void FWaterInfoTextureDepthPassMeshProcessor::CollectPSOInitializers(const FSceneTexturesConfig& SceneTexturesConfig, const FMaterial& Material, const FPSOPrecacheVertexFactoryData& VertexFactoryData, const FPSOPrecacheParams& PreCacheParams, TArray<FPSOPrecacheData>& PSOInitializers)
 {
 	// We need to support all materials that could possibly be rendered in a depth-only pass. Unfortunately there doesn't seem to be a way to filter for bUseForWaterInfoTextureDepth at this point.
-	if (IsCustomWaterInfoTextureRenderEnabled() && Material.GetMaterialDomain() == MD_Surface && !IsTranslucentBlendMode(Material))
+	if (Material.GetMaterialDomain() == MD_Surface && !IsTranslucentBlendMode(Material))
 	{
 		// Determine the mesh's material and blend mode.
 		const FMeshDrawingPolicyOverrideSettings OverrideSettings = ComputeMeshOverrideSettings(PreCacheParams);
