Fixes VisualizeTemporalUpscaler show flag showing CopyToScreenPassTexture2D on TSR output

#tests win64

diff --git a/Engine/Source/Runtime/Renderer/Private/PostProcess/VisualizeTemporalUpscaler.cpp b/Engine/Source/Runtime/Renderer/Private/PostProcess/VisualizeTemporalUpscaler.cpp
--- a/Engine/Source/Runtime/Renderer/Private/PostProcess/VisualizeTemporalUpscaler.cpp
+++ b/Engine/Source/Runtime/Renderer/Private/PostProcess/VisualizeTemporalUpscaler.cpp
@@ -33,7 +33,31 @@ FScreenPassTexture AddVisualizeTemporalUpscalerPass(FRDGBuilder& GraphBuilder, c
 
 	if (Inputs.TAAConfig != EMainTAAPassConfig::Disabled)
 	{
-		FScreenPassTexture OutputTexture = FScreenPassTexture::CopyFromSlice(GraphBuilder, Inputs.Outputs.FullRes);
+		FScreenPassTexture OutputTexture;
+		if (Inputs.Outputs.FullRes.TextureSRV->Desc.Texture->Desc.IsTextureArray())
+		{
+			FRDGTextureDesc Desc = Inputs.Outputs.FullRes.TextureSRV->Desc.Texture->Desc;
+			Desc.Dimension = ETextureDimension::Texture2D;
+			Desc.ArraySize = 1;
+			Desc.NumMips = 1;
+
+			FRDGTextureRef NewTexture = GraphBuilder.CreateTexture(Desc, Inputs.Outputs.FullRes.TextureSRV->Desc.Texture->Name);
+
+			FRHICopyTextureInfo CopyInfo;
+			CopyInfo.SourceSliceIndex = Inputs.Outputs.FullRes.TextureSRV->Desc.FirstArraySlice;
+
+			AddCopyTexturePass(
+				GraphBuilder,
+				Inputs.Outputs.FullRes.TextureSRV->Desc.Texture,
+				NewTexture,
+				CopyInfo);
+
+			OutputTexture = FScreenPassTexture(NewTexture, Inputs.Outputs.FullRes.ViewRect);
+		}
+		else
+		{
+			OutputTexture = FScreenPassTexture(Inputs.Outputs.FullRes.TextureSRV->Desc.Texture, Inputs.Outputs.FullRes.ViewRect);
+		}
 
 		auto VisualizeTextureLabel = [](FRDGTextureRef Texture, const TCHAR* Suffix = TEXT(""))
 		{
