Fixes FScreenPassTexture::CopyFromSlice() not copying up all mips

diff --git a/Engine/Source/Runtime/Renderer/Private/ScreenPass.cpp b/Engine/Source/Runtime/Renderer/Private/ScreenPass.cpp
--- a/Engine/Source/Runtime/Renderer/Private/ScreenPass.cpp
+++ b/Engine/Source/Runtime/Renderer/Private/ScreenPass.cpp
@@ -62,6 +62,7 @@ FScreenPassTexture FScreenPassTexture::CopyFromSlice(FRDGBuilder& GraphBuilder,
 
 	FRHICopyTextureInfo CopyInfo;
 	CopyInfo.SourceSliceIndex = ScreenTextureSlice.TextureSRV->Desc.FirstArraySlice;
+	CopyInfo.NumMips = ScreenTextureSlice.TextureSRV->Desc.Texture->Desc.NumMips;
 
 	AddCopyTexturePass(
 		GraphBuilder,
