Works arround fullscreen corruption on Nvidia 2xxx series GPU

diff --git a/Engine/Source/Runtime/Renderer/Private/PostProcess/TemporalSuperResolution.cpp b/Engine/Source/Runtime/Renderer/Private/PostProcess/TemporalSuperResolution.cpp
--- a/Engine/Source/Runtime/Renderer/Private/PostProcess/TemporalSuperResolution.cpp
+++ b/Engine/Source/Runtime/Renderer/Private/PostProcess/TemporalSuperResolution.cpp
@@ -1232,6 +1232,12 @@ FDefaultTemporalUpscaler::FOutputs AddTemporalSuperResolutionPasses(
 	bool bUse16BitVALU = (CVarTSR16BitVALU.GetValueOnRenderThread() != 0 && GRHIGlobals.SupportsNative16BitOps && VALU16BitSupport == ERHIFeatureSupport::RuntimeDependent) || VALU16BitSupport == ERHIFeatureSupport::RuntimeGuaranteed;
 
 #if PLATFORM_DESKTOP
+	// PLAY-18908: work around shader corruption on Nvidia GPU until fix is done in the driver
+	if (IsRHIDeviceNVIDIA())
+	{
+		bUse16BitVALU = false;
+	}
+
 	if (IsRHIDeviceAMD())
 	{
 		bUse16BitVALU = false;
