Fixed aerial perspective being applied on the sky dome mesh (material with IsSky=true) on mobile.
Also improved the logic computing the start depth.

diff --git a/Engine/Shaders/Private/MobileFog.usf b/Engine/Shaders/Private/MobileFog.usf
--- a/Engine/Shaders/Private/MobileFog.usf
+++ b/Engine/Shaders/Private/MobileFog.usf
@@ -41,7 +41,8 @@ void MobileFogPS(
 {
 	ResolvedView = ResolveView();
 	half4 FinalFog = half4(0,0,0,1);
-	float SceneDepth = CalcSceneDepth(TexCoord);
+	float DeviceZ = LookupDeviceZ(TexCoord);
+	float SceneDepth = ConvertFromDeviceZ(DeviceZ);
 	float3 WorldPositionRelativeToCamera = ScreenVector.xyz * SceneDepth;
 	float3 WorldPosition = WorldPositionRelativeToCamera + ResolvedView.TranslatedWorldCameraOrigin - LWCHackToFloat(ResolvedView.PreViewTranslation);
 
@@ -66,10 +67,8 @@ void MobileFogPS(
 #endif
 
 #if PERMUTATION_SUPPORT_AERIAL_PERSPECTIVE
-	if (ResolvedView.SkyAtmosphereApplyCameraAerialPerspectiveVolume > 0.0f)
+	if (ResolvedView.SkyAtmosphereApplyCameraAerialPerspectiveVolume > 0.0f && (DeviceZ != FarDepthValue))
 	{
-		
-
 	//	float4 NDCPosition = LWCMultiply(MakeLWCVector4(WorldPosition, 1.0f), PrimaryView.WorldToClip);		// LWC_TODO: MakeLWCVector4 seems broken
 		float4 NDCPosition = LWCMultiply(LWCPromote(float4(WorldPosition, 1.0f)), PrimaryView.WorldToClip);	// Using LWCPromote make the math work as expected.
 
diff --git a/Engine/Source/Runtime/Renderer/Private/FogRendering.cpp b/Engine/Source/Runtime/Renderer/Private/FogRendering.cpp
--- a/Engine/Source/Runtime/Renderer/Private/FogRendering.cpp
+++ b/Engine/Source/Runtime/Renderer/Private/FogRendering.cpp
@@ -468,6 +468,10 @@ bool ShouldRenderFog(const FSceneViewFamily& Family)
 		&& !EngineShowFlags.StationaryLightOverlap 
 		&& !EngineShowFlags.LightMapDensity;
 }
+float GetFogDefaultStartDistance()
+{
+	return 30.0f;
+}
 
 float GetViewFogCommonStartDistance(const FViewInfo& View, bool bShouldRenderVolumetricFog)
 {
@@ -476,6 +480,6 @@ float GetViewFogCommonStartDistance(const FViewInfo& View, bool bShouldRenderVol
 
 	// The fog can be set to start at a certain euclidean distance.
 	// clamp the value to be behind the near plane z, according to the smallest distance between volumetric fog and height fog (if they are enabled). 
-	float FogCommonStartDistance = FMath::Max(30.0f, FMath::Min(ExpFogStartDistance, VolFogStartDistance));
+	float FogCommonStartDistance = FMath::Max(GetFogDefaultStartDistance(), FMath::Min(ExpFogStartDistance, VolFogStartDistance));
 	return FogCommonStartDistance;
 }
diff --git a/Engine/Source/Runtime/Renderer/Private/FogRendering.h b/Engine/Source/Runtime/Renderer/Private/FogRendering.h
--- a/Engine/Source/Runtime/Renderer/Private/FogRendering.h
+++ b/Engine/Source/Runtime/Renderer/Private/FogRendering.h
@@ -37,3 +37,4 @@ TRDGUniformBufferRef<FFogUniformParameters> CreateFogUniformBuffer(FRDGBuilder&
 extern bool ShouldRenderFog(const FSceneViewFamily& Family);
 
 extern float GetViewFogCommonStartDistance(const FViewInfo& View, bool bShouldRenderVolumetricFog);
+extern float GetFogDefaultStartDistance();
diff --git a/Engine/Source/Runtime/Renderer/Private/MobileFogRendering.cpp b/Engine/Source/Runtime/Renderer/Private/MobileFogRendering.cpp
--- a/Engine/Source/Runtime/Renderer/Private/MobileFogRendering.cpp
+++ b/Engine/Source/Runtime/Renderer/Private/MobileFogRendering.cpp
@@ -211,12 +211,20 @@ void FMobileSceneRenderer::RenderFog(FRHICommandList& RHICmdList, const FViewInf
 	SetGraphicsPipelineState(RHICmdList, GraphicsPSOInit, 0);
 	
 	// Use height fog start distance by default and fallback to AP distance
-	float FogStartDistance = GetViewFogCommonStartDistance(View, bShouldRenderVolumetricFog);
-	if (!bUseHeightFog)
+	float FogStartDistance = GetFogDefaultStartDistance();
+	if (!bFogHasComposedLocalFogVolumes) // Local fog volume do not have any start distance at the moment.
 	{
-		FSkyAtmosphereRenderSceneInfo& SkyInfo = *Scene->GetSkyAtmosphereSceneInfo();
-		const FSkyAtmosphereSceneProxy& SkyAtmosphereSceneProxy = SkyInfo.GetSkyAtmosphereSceneProxy();
-		FogStartDistance = GetValidAerialPerspectiveStartDepthInCm(View, SkyAtmosphereSceneProxy);
+		if (bUseHeightFog)
+		{
+			FogStartDistance = GetViewFogCommonStartDistance(View, bShouldRenderVolumetricFog);
+		}
+		if (bUseAerialPerspective)
+		{
+			FSkyAtmosphereRenderSceneInfo& SkyInfo = *Scene->GetSkyAtmosphereSceneInfo();
+			const FSkyAtmosphereSceneProxy& SkyAtmosphereSceneProxy = SkyInfo.GetSkyAtmosphereSceneProxy();
+			const float AerialPerspectiveStartDepthInCm = GetValidAerialPerspectiveStartDepthInCm(View, SkyAtmosphereSceneProxy);
+			FogStartDistance = bUseHeightFog ? FMath::Min(FogStartDistance, AerialPerspectiveStartDepthInCm) : AerialPerspectiveStartDepthInCm;
+		}
 	}
 		
 	float StartDepthZ = 0.1;	
