Added a cvar to force skylight cubemap blending per platform for forward shading. This allows for higher quality refelction on some platforms using forward shading, for instance when using the mobile forward shading path.

diff --git a/Engine/Shaders/Private/ForwardLightingCommon.ush b/Engine/Shaders/Private/ForwardLightingCommon.ush
--- a/Engine/Shaders/Private/ForwardLightingCommon.ush
+++ b/Engine/Shaders/Private/ForwardLightingCommon.ush
@@ -387,7 +387,7 @@ uint MortonCode( uint x )
 #define USE_BLENDED_REFLECTION_CAPTURES_FORWARD (FORWARD_SHADING && MATERIAL_HQ_FORWARD_REFLECTION_CAPTURES)
 #define REFLECTION_COMPOSITE_USE_BLENDED_REFLECTION_CAPTURES ((USE_BLENDED_REFLECTION_CAPTURES_DEFERRED || USE_BLENDED_REFLECTION_CAPTURES_FORWARD) && FEATURE_LEVEL >= FEATURE_LEVEL_SM5)
 // REFLECTION_COMPOSITE_SUPPORT_SKYLIGHT_BLEND covers forward rendered material for deferred and forward shading pathes.
-#define REFLECTION_COMPOSITE_SUPPORT_SKYLIGHT_BLEND (MATERIAL_FORWARD_BLENDS_SKYLIGHT_CUBEMAPS)
+#define REFLECTION_COMPOSITE_SUPPORT_SKYLIGHT_BLEND (MATERIAL_FORWARD_BLENDS_SKYLIGHT_CUBEMAPS || FORWARD_SHADING_FORCES_SKYLIGHT_CUBEMAPS_BLENDING)
 #define REFLECTION_COMPOSITE_HAS_BOX_CAPTURES 1
 #define REFLECTION_COMPOSITE_HAS_SPHERE_CAPTURES 1
 #include "ReflectionEnvironmentComposite.ush"
diff --git a/Engine/Shaders/Private/MobileLightingCommon.ush b/Engine/Shaders/Private/MobileLightingCommon.ush
--- a/Engine/Shaders/Private/MobileLightingCommon.ush
+++ b/Engine/Shaders/Private/MobileLightingCommon.ush
@@ -308,6 +308,14 @@ half3 GetMobileSkyLightReflection(half3 ReflectionVector, half Roughness, half C
 {
 	half AbsoluteSpecularMip = ComputeReflectionCaptureMipFromRoughness(Roughness, CubemapMaxMip);
 	half4 Reflection = MobileReflectionCapture.Texture.SampleLevel(MobileReflectionCapture.TextureSampler, ReflectionVector, AbsoluteSpecularMip);
+#if FORWARD_SHADING_FORCES_SKYLIGHT_CUBEMAPS_BLENDING
+	float BlendFraction = MobileReflectionCapture.Params.w;
+	if (BlendFraction > 0.0f)
+	{
+		half4 ReflectionBlend = MobileReflectionCapture.TextureBlend.SampleLevel(MobileReflectionCapture.TextureBlendSampler, ReflectionVector, AbsoluteSpecularMip);
+		Reflection = lerp(Reflection, ReflectionBlend, BlendFraction.xxxx);
+	}
+#endif
 	return Reflection.rgb * ResolvedView.SkyLightColor.rgb;
 }
 #endif
diff --git a/Engine/Source/Runtime/Engine/Private/Components/ReflectionCaptureComponent.cpp b/Engine/Source/Runtime/Engine/Private/Components/ReflectionCaptureComponent.cpp
--- a/Engine/Source/Runtime/Engine/Private/Components/ReflectionCaptureComponent.cpp
+++ b/Engine/Source/Runtime/Engine/Private/Components/ReflectionCaptureComponent.cpp
@@ -1325,6 +1325,8 @@ void FReflectionCaptureProxy::UpdateMobileUniformBuffer(FRHICommandListBase& RHI
 	Parameters.Params = FVector4f(EncodedHDRAverageBrightness, 0.f, 0.0f, Brightness);
 	Parameters.Texture = CaptureTexture->TextureRHI;
 	Parameters.TextureSampler = CaptureTexture->SamplerStateRHI;
+	Parameters.TextureBlend = Parameters.Texture;
+	Parameters.TextureBlendSampler = Parameters.TextureSampler;
 
 	if (MobileUniformBuffer.GetReference())
 	{
diff --git a/Engine/Source/Runtime/Engine/Private/SceneManagement.cpp b/Engine/Source/Runtime/Engine/Private/SceneManagement.cpp
--- a/Engine/Source/Runtime/Engine/Private/SceneManagement.cpp
+++ b/Engine/Source/Runtime/Engine/Private/SceneManagement.cpp
@@ -1714,6 +1714,8 @@ void FDefaultMobileReflectionCaptureUniformBuffer::InitRHI(FRHICommandListBase&
 	Parameters.Params = FVector4f(1.f, 0.f, 0.f, 0.f);
 	Parameters.Texture = GBlackTextureCube->TextureRHI;
 	Parameters.TextureSampler = GBlackTextureCube->SamplerStateRHI;
+	Parameters.TextureBlend = Parameters.Texture;
+	Parameters.TextureBlendSampler = Parameters.TextureSampler;
 	SetContentsNoUpdate(Parameters);
 	Super::InitRHI(RHICmdList);
 }
diff --git a/Engine/Source/Runtime/Engine/Private/ShaderCompiler/ShaderCompiler.cpp b/Engine/Source/Runtime/Engine/Private/ShaderCompiler/ShaderCompiler.cpp
--- a/Engine/Source/Runtime/Engine/Private/ShaderCompiler/ShaderCompiler.cpp
+++ b/Engine/Source/Runtime/Engine/Private/ShaderCompiler/ShaderCompiler.cpp
@@ -7553,6 +7553,8 @@ void GlobalBeginCompileShader(
 		SET_SHADER_DEFINE(Input.Environment, RAYMISSSHADER,			Target.Frequency == SF_RayMiss);
 	}
 
+	SET_SHADER_DEFINE(Input.Environment, FORWARD_SHADING_FORCES_SKYLIGHT_CUBEMAPS_BLENDING, ForwardShadingForcesSkyLightCubemapBlending(ShaderPlatform) ? 1 : 0);
+
 	// Enables HLSL 2021
 	uint32 EnablesHLSL2021ByDefault = FDataDrivenShaderPlatformInfo::GetEnablesHLSL2021ByDefault(EShaderPlatform(Target.Platform));
 	if((EnablesHLSL2021ByDefault == uint32(1) && DebugGroupName == TEXT("Global")) ||
diff --git a/Engine/Source/Runtime/Engine/Public/SceneManagement.h b/Engine/Source/Runtime/Engine/Public/SceneManagement.h
--- a/Engine/Source/Runtime/Engine/Public/SceneManagement.h
+++ b/Engine/Source/Runtime/Engine/Public/SceneManagement.h
@@ -1096,6 +1096,8 @@ BEGIN_GLOBAL_SHADER_PARAMETER_STRUCT(FMobileReflectionCaptureShaderParameters,EN
 	SHADER_PARAMETER(FVector4f, Params) // x - inv average brightness, y - sky cubemap max mip, z - unused, w - brightness of reflection capture
 	SHADER_PARAMETER_TEXTURE(TextureCube, Texture)
 	SHADER_PARAMETER_SAMPLER(SamplerState, TextureSampler)
+	SHADER_PARAMETER_TEXTURE(TextureCube, TextureBlend)			// Only used when this refelction is a sky light
+	SHADER_PARAMETER_SAMPLER(SamplerState, TextureBlendSampler)	// Idem
 END_GLOBAL_SHADER_PARAMETER_STRUCT()
 
 class FDefaultMobileReflectionCaptureUniformBuffer : public TUniformBuffer<FMobileReflectionCaptureShaderParameters>
diff --git a/Engine/Source/Runtime/RenderCore/Private/RenderUtils.cpp b/Engine/Source/Runtime/RenderCore/Private/RenderUtils.cpp
--- a/Engine/Source/Runtime/RenderCore/Private/RenderUtils.cpp
+++ b/Engine/Source/Runtime/RenderCore/Private/RenderUtils.cpp
@@ -1959,4 +1959,16 @@ bool DoesPlatformSupportLumenGI(EShaderPlatform Platform, bool bSkipProjectCheck
 	return (bSkipProjectCheck || bLumenSupported)
 		&& FDataDrivenShaderPlatformInfo::GetSupportsLumenGI(Platform)
 		&& !IsForwardShadingEnabled(Platform);
-}
\ No newline at end of file
+}
+
+static TAutoConsoleVariable<int32> CVarForwardShadingForceSkyBoxBlending(
+	TEXT("r.ForwardShading.ForceSkyLightCubemapBlending"),
+	0,
+	TEXT("Forces in shader sky box blending for forward shading. More expenssive but higher quality. This is similar to having Blend Sky Light Cubemaps = true on all forward materials."),
+	ECVF_ReadOnly | ECVF_RenderThreadSafe);
+
+bool ForwardShadingForcesSkyLightCubemapBlending(const FStaticShaderPlatform Platform)
+{
+	static FShaderPlatformCachedIniValue<int32> CVar(TEXT("r.ForwardShading.ForceSkyLightCubemapBlending"));
+	return CVar.Get(Platform) > 0;
+}
diff --git a/Engine/Source/Runtime/RenderCore/Private/Shader.cpp b/Engine/Source/Runtime/RenderCore/Private/Shader.cpp
--- a/Engine/Source/Runtime/RenderCore/Private/Shader.cpp
+++ b/Engine/Source/Runtime/RenderCore/Private/Shader.cpp
@@ -1661,6 +1661,11 @@ void ShaderMapAppendKeyString(EShaderPlatform Platform, FString& KeyString)
 		KeyString += TEXT("_UnInt");
 	}
 
+	if (ForwardShadingForcesSkyLightCubemapBlending(Platform))
+	{
+		KeyString += TEXT("_FwdSkyBlnd");
+	}
+
 	if (IsMobilePlatform(Platform))
 	{
 		{
diff --git a/Engine/Source/Runtime/RenderCore/Public/RenderUtils.h b/Engine/Source/Runtime/RenderCore/Public/RenderUtils.h
--- a/Engine/Source/Runtime/RenderCore/Public/RenderUtils.h
+++ b/Engine/Source/Runtime/RenderCore/Public/RenderUtils.h
@@ -320,6 +320,9 @@ RENDERCORE_API bool IsUsingDBuffers(const FStaticShaderPlatform Platform);
 /** Returns if ForwardShading is enabled. Only valid for the current platform (otherwise call ITargetPlatform::UsesForwardShading()). */
 RENDERCORE_API bool IsForwardShadingEnabled(const FStaticShaderPlatform Platform);
 
+/** Return true if all forward shaded material should blend the interpolated sky boxes for higher quality. */
+RENDERCORE_API bool ForwardShadingForcesSkyLightCubemapBlending(const FStaticShaderPlatform Platform);
+
 /** Returns if the GBuffer is used. Only valid for the current platform. */
 RENDERCORE_API bool IsUsingGBuffers(const FStaticShaderPlatform Platform);
 
diff --git a/Engine/Source/Runtime/Renderer/Private/MobileBasePassRendering.cpp b/Engine/Source/Runtime/Renderer/Private/MobileBasePassRendering.cpp
--- a/Engine/Source/Runtime/Renderer/Private/MobileBasePassRendering.cpp
+++ b/Engine/Source/Runtime/Renderer/Private/MobileBasePassRendering.cpp
@@ -358,6 +358,8 @@ void SetupMobileSkyReflectionUniformParameters(FSkyLightSceneProxy* SkyLight, FM
 	float Brightness = 0.f;
 	float SkyMaxMipIndex = 0.f;
 	FTexture* CaptureTexture = GBlackTextureCube;
+	FTexture* CaptureTextureBlend = GBlackTextureCube;
+	float BlendFraction = 0.f;
 
 	bool bSkyLightIsDynamic = false;
 	if (SkyLight && SkyLight->ProcessedTexture)
@@ -367,12 +369,20 @@ void SetupMobileSkyReflectionUniformParameters(FSkyLightSceneProxy* SkyLight, FM
 		SkyMaxMipIndex = FMath::Log2(static_cast<float>(CaptureTexture->GetSizeX()));
 		Brightness = SkyLight->AverageBrightness;
 		bSkyLightIsDynamic = !SkyLight->bHasStaticLighting && !SkyLight->bWantsStaticShadowing;
+
+		BlendFraction = SkyLight->BlendFraction;
+		if (BlendFraction > 0.0 && SkyLight->BlendDestinationProcessedTexture)
+		{
+			CaptureTextureBlend = SkyLight->BlendDestinationProcessedTexture;
+		}
 	}
 	
 	//To keep ImageBasedReflectionLighting coherence with PC, use AverageBrightness instead of InvAverageBrightness to calculate the IBL contribution
-	Parameters.Params = FVector4f(Brightness, SkyMaxMipIndex, bSkyLightIsDynamic ? 1.0f : 0.0f, 0.f);
+	Parameters.Params = FVector4f(Brightness, SkyMaxMipIndex, bSkyLightIsDynamic ? 1.0f : 0.0f, BlendFraction);
 	Parameters.Texture = CaptureTexture->TextureRHI;
 	Parameters.TextureSampler = CaptureTexture->SamplerStateRHI;
+	Parameters.TextureBlend = CaptureTextureBlend->TextureRHI;
+	Parameters.TextureBlendSampler = CaptureTextureBlend->SamplerStateRHI;
 }
 
 void FMobileSceneRenderer::RenderMobileBasePass(FRHICommandList& RHICmdList, const FViewInfo& View, const FInstanceCullingDrawParams* InstanceCullingDrawParams)
