VolumetricCloud - fixed bug with high quality aerial perspective option where the secondary far traced cloud buffer was not receiving it.

diff --git a/Engine/Source/Runtime/Renderer/Private/VolumetricCloudRendering.cpp b/Engine/Source/Runtime/Renderer/Private/VolumetricCloudRendering.cpp
--- a/Engine/Source/Runtime/Renderer/Private/VolumetricCloudRendering.cpp
+++ b/Engine/Source/Runtime/Renderer/Private/VolumetricCloudRendering.cpp
@@ -2617,7 +2617,9 @@ bool FSceneRenderer::RenderVolumetricCloud(
 				CloudRC.bIsReflectionRendering = ViewInfo.bIsReflectionCapture;
 
 				FRDGTextureRef IntermediateRT = nullptr;
+				FRDGTextureRef IntermediateSecondaryRT = nullptr;
 				FRDGTextureRef DestinationRT = nullptr;
+				FRDGTextureRef DestinationSecondaryRT = nullptr;
 				FRDGTextureRef DestinationRTDepth = nullptr;
 				CloudRC.TracingCoordToZbufferCoordScaleBias = FUintVector4(1, 1, 0, 0);
 				CloudRC.TracingCoordToFullResPixelCoordScaleBias = FUintVector4(1, 1, 0, 0);
@@ -2635,6 +2637,9 @@ bool FSceneRenderer::RenderVolumetricCloud(
 						IntermediateRT = GraphBuilder.CreateTexture(
 							FRDGTextureDesc::Create2D(IntermadiateTargetResolution, PF_FloatRGBA, FClearValueBinding(FLinearColor(63000.0f, 63000.0f, 63000.0f, 63000.0f)),
 								TexCreate_ShaderResource | TexCreate_RenderTargetable | TexCreate_UAV), TEXT("Cloud.HighQualityAPIntermediate"));
+						IntermediateSecondaryRT = GraphBuilder.CreateTexture(
+							FRDGTextureDesc::Create2D(IntermadiateTargetResolution, PF_FloatRGBA, FClearValueBinding(FLinearColor(63000.0f, 63000.0f, 63000.0f, 63000.0f)),
+								TexCreate_ShaderResource | TexCreate_RenderTargetable | TexCreate_UAV), TEXT("Cloud.HighQualityAPIntermediateSecondary"));
 					}
 
 					// No action because we only need to render volumetric clouds so we do not blend in that render target.
@@ -2650,7 +2655,8 @@ bool FSceneRenderer::RenderVolumetricCloud(
 					const bool bShouldVolumetricCloudTraceWithMinMaxDepth = ShouldVolumetricCloudTraceWithMinMaxDepth(ViewInfo);
 					if (bShouldVolumetricCloudTraceWithMinMaxDepth)
 					{
-						CloudRC.SecondaryCloudTracingDataTexture = VRT.GetOrCreateVolumetricSecondaryTracingRT(GraphBuilder); 
+						DestinationSecondaryRT = VRT.GetOrCreateVolumetricSecondaryTracingRT(GraphBuilder);
+						CloudRC.SecondaryCloudTracingDataTexture = bShouldUseHighQualityAerialPerspective ? IntermediateSecondaryRT : DestinationSecondaryRT;
 					}
 
 					CloudRC.NoiseFrameIndexModPattern = VRT.GetNoiseFrameIndexModPattern();
@@ -2804,6 +2810,15 @@ bool FSceneRenderer::RenderVolumetricCloud(
 					SkyRC.VolumetricCloudSkyAO = CloudShadowAOData.VolumetricCloudSkyAO;
 
 					RenderSkyAtmosphereInternal(GraphBuilder, GetSceneTextureShaderParameters(SceneTextures.UniformBuffer), SkyRC);
+
+					if (DestinationSecondaryRT)
+					{
+						// We need to execute a second tracing for the far depth cloud result in order to correctly fix up edges.
+						SkyRC.VolumetricCloudDepthTexture = DestinationRTDepth;
+						SkyRC.InputCloudLuminanceTransmittanceTexture = IntermediateSecondaryRT;
+						SkyRC.RenderTargets[0] = FRenderTargetBinding(DestinationSecondaryRT, ERenderTargetLoadAction::ENoAction);
+						RenderSkyAtmosphereInternal(GraphBuilder, GetSceneTextureShaderParameters(SceneTextures.UniformBuffer), SkyRC);
+					}
 				}
 
 				if (DebugCloudShadowMap || DebugCloudSkyAO)
