Fix for broken CombineLUT pass when switching between PS and CS based post processing

diff --git a/Engine/Source/Runtime/Renderer/Private/PostProcess/PostProcessCombineLUTs.cpp b/Engine/Source/Runtime/Renderer/Private/PostProcess/PostProcessCombineLUTs.cpp
--- a/Engine/Source/Runtime/Renderer/Private/PostProcess/PostProcessCombineLUTs.cpp
+++ b/Engine/Source/Runtime/Renderer/Private/PostProcess/PostProcessCombineLUTs.cpp
@@ -178,6 +178,7 @@ struct FCachedLUTSettings
 	EShaderPlatform ShaderPlatform = GMaxRHIShaderPlatform;
 	FCombineLUTParameters Parameters;
 	FWorkingColorSpaceShaderParameters WorkingColorSpaceShaderParameters;
+	bool bUseCompute = false;
 
 	bool UpdateCachedValues(const FViewInfo& View, const FTexture* const* Textures, const float* Weights, uint32 BlendCount)
 	{
@@ -185,6 +186,7 @@ struct FCachedLUTSettings
 		GetCombineLUTParameters(View, Textures, Weights, BlendCount, bHasChanged);
 		UPDATE_CACHE_SETTINGS(UniqueID, View.State ? View.State->GetViewKey() : 0, bHasChanged);
 		UPDATE_CACHE_SETTINGS(ShaderPlatform, View.GetShaderPlatform(), bHasChanged);
+		UPDATE_CACHE_SETTINGS(bUseCompute, View.bUseComputePasses, bHasChanged);
 
 		const FWorkingColorSpaceShaderParameters* InWorkingColorSpaceShaderParameters = reinterpret_cast<const FWorkingColorSpaceShaderParameters*>(GDefaultWorkingColorSpaceUniformBuffer.GetContents());
 		if (InWorkingColorSpaceShaderParameters)
