VolumetricFog: add support for Rect Light texture.
* Controlled with r.VolumetricFog.RectLightTexture and disabled by default.

diff --git a/Engine/Shaders/Private/VolumetricFog.usf b/Engine/Shaders/Private/VolumetricFog.usf
--- a/Engine/Shaders/Private/VolumetricFog.usf
+++ b/Engine/Shaders/Private/VolumetricFog.usf
@@ -5,6 +5,7 @@
 #include "Common.ush"
 #include "Definitions.usf"
 #define SUPPORT_CONTACT_SHADOWS 0
+#define USE_SOURCE_TEXTURE USE_RECT_LIGHT_TEXTURE
 #include "DeferredLightingCommon.ush"
 #include "LightGridCommon.ush"
 #include "HeightFogCommon.ush"
@@ -456,6 +457,15 @@ float4 InjectShadowedLocalLightCommon(uint3 GridCoordinate)
 				FRect Rect = GetRect(ToLight, LightData);
 
 				Lighting = IntegrateLight(Rect);
+
+#if USE_RECT_LIGHT_TEXTURE
+				if (IsRectVisible(Rect))
+				{
+					const FRectTexture SourceTexture = ConvertToRectTexture(LightData);
+					const float3 FalloffColor = SampleSourceTexture(Lighting, Rect, SourceTexture);
+					LightColor *= FalloffColor;
+				}
+#endif // USE_RECT_LIGHT_TEXTURE
 			}
 			else
 			{
diff --git a/Engine/Source/Runtime/Renderer/Private/VolumetricFog.cpp b/Engine/Source/Runtime/Renderer/Private/VolumetricFog.cpp
--- a/Engine/Source/Runtime/Renderer/Private/VolumetricFog.cpp
+++ b/Engine/Source/Runtime/Renderer/Private/VolumetricFog.cpp
@@ -117,6 +117,14 @@ FAutoConsoleVariableRef CVarVolumetricFogEmissive(
 	ECVF_Scalability | ECVF_RenderThreadSafe
 );
 
+int32 GVolumetricFogRectLightTexture = 0;
+FAutoConsoleVariableRef CVarVolumetricRectLightTexture(
+	TEXT("r.VolumetricFog.RectLightTexture"),
+	GVolumetricFogRectLightTexture,
+	TEXT("Whether to allow the volumetric fog to use rect light source texture."),
+	ECVF_RenderThreadSafe
+);
+
 int32 GVolumetricFogConservativeDepth = 0;
 FAutoConsoleVariableRef CVarVolumetricFogConservativeDepth(
 	TEXT("r.VolumetricFog.ConservativeDepth"),
@@ -360,13 +368,15 @@ class FInjectShadowedLocalLightPS : public FGlobalShader
 	class FLightFunction		: SHADER_PERMUTATION_BOOL("USE_LIGHT_FUNCTION");
 	class FEnableShadows		: SHADER_PERMUTATION_BOOL("ENABLE_SHADOW_COMPUTATION");
 	class FVirtualShadowMap		: SHADER_PERMUTATION_BOOL("VIRTUAL_SHADOW_MAP");
+	class FRectLightTexture		: SHADER_PERMUTATION_BOOL("USE_RECT_LIGHT_TEXTURE");
 
 	using FPermutationDomain = TShaderPermutationDomain<
 		FDynamicallyShadowed,
 		FTemporalReprojection,
 		FLightFunction,
 		FEnableShadows,
-		FVirtualShadowMap >;
+		FVirtualShadowMap,
+		FRectLightTexture >;
 
 	static void ModifyCompilationEnvironment(const FGlobalShaderPermutationParameters& Parameters, FShaderCompilerEnvironment& OutEnvironment)
 	{
@@ -393,11 +403,13 @@ class FInjectShadowedLocalLightRGS : public FGlobalShader
 	END_SHADER_PARAMETER_STRUCT()
 
 	class FTemporalReprojection : SHADER_PERMUTATION_BOOL("USE_TEMPORAL_REPROJECTION");
-	class FLightFunction : SHADER_PERMUTATION_BOOL("USE_LIGHT_FUNCTION");
+	class FLightFunction		: SHADER_PERMUTATION_BOOL("USE_LIGHT_FUNCTION");
+	class FRectLightTexture		: SHADER_PERMUTATION_BOOL("USE_RECT_LIGHT_TEXTURE");
 
 	using FPermutationDomain = TShaderPermutationDomain<
 		FTemporalReprojection,
-		FLightFunction >;
+		FLightFunction,
+		FRectLightTexture >;
 
 	static bool ShouldCompilePermutation(const FGlobalShaderPermutationParameters& Parameters)
 	{
@@ -542,12 +554,16 @@ void FDeferredShadingSceneRenderer::PrepareRayTracingVolumetricFogShadows(const
 	{
 		for (int32 UseLightFunction = 0; UseLightFunction < 2; ++UseLightFunction)
 		{
-			FInjectShadowedLocalLightRGS::FPermutationDomain PermutationVector;
-			PermutationVector.Set<FInjectShadowedLocalLightRGS::FTemporalReprojection>((bool)TemporalReprojection);
-			PermutationVector.Set<FInjectShadowedLocalLightRGS::FLightFunction>((bool)UseLightFunction);
+			for (int32 UseRectLightTexture = 0; UseRectLightTexture < 2; ++UseRectLightTexture)
+			{
+				FInjectShadowedLocalLightRGS::FPermutationDomain PermutationVector;
+				PermutationVector.Set<FInjectShadowedLocalLightRGS::FTemporalReprojection>((bool)TemporalReprojection);
+				PermutationVector.Set<FInjectShadowedLocalLightRGS::FLightFunction>((bool)UseLightFunction);
+				PermutationVector.Set<FInjectShadowedLocalLightRGS::FRectLightTexture>((bool)UseRectLightTexture);
 
-			TShaderMapRef<FInjectShadowedLocalLightRGS> RayGenerationShader(View.ShaderMap, PermutationVector);
-			OutRayGenShaders.Add(RayGenerationShader.GetRayTracingShader());
+				TShaderMapRef<FInjectShadowedLocalLightRGS> RayGenerationShader(View.ShaderMap, PermutationVector);
+				OutRayGenShaders.Add(RayGenerationShader.GetRayTracingShader());
+			}
 		}
 	}
 
@@ -723,10 +739,11 @@ void FSceneRenderer::RenderLocalLightsForVolumetricFog(
 		bool bIsShadowed = LightNeedsSeparateInjectionIntoVolumetricFogForOpaqueShadow(View, LightSceneInfo, VisibleLightInfos[LightSceneInfo->Id]);
 		bool bUsesLightFunction = ViewFamily.EngineShowFlags.LightFunctions 
 			&& CheckForLightFunction(LightSceneInfo) && LightNeedsSeparateInjectionIntoVolumetricFogForLightFunction(LightSceneInfo);
+		bool bUsesRectLightTexture = GVolumetricFogRectLightTexture && LightSceneInfo->Proxy->HasSourceTexture();
 
 		if (LightSceneInfo->ShouldRenderLightViewIndependent()
 			&& LightSceneInfo->ShouldRenderLight(View)
-			&& (bIsShadowed || bUsesLightFunction)
+			&& (bIsShadowed || bUsesLightFunction || bUsesRectLightTexture)
 			&& LightSceneInfo->Proxy->GetVolumetricScatteringIntensity() > 0)
 		{
 			const FSphere LightBounds = LightSceneInfo->Proxy->GetBoundingSphere();
@@ -762,6 +779,7 @@ void FSceneRenderer::RenderLocalLightsForVolumetricFog(
 				bool bIsShadowed = LightNeedsSeparateInjectionIntoVolumetricFogForOpaqueShadow(View, LightSceneInfo, VisibleLightInfo);
 				bool bUsesLightFunction = ViewFamily.EngineShowFlags.LightFunctions
 					&& CheckForLightFunction(LightSceneInfo) && LightNeedsSeparateInjectionIntoVolumetricFogForLightFunction(LightSceneInfo);
+				bool bUsesRectLightTexture = GVolumetricFogRectLightTexture && LightSceneInfo->Proxy->HasSourceTexture();
 
 				int32 VirtualShadowMapId = VisibleLightInfo.GetVirtualShadowMapId(&View);
 				const bool bUseVSM = bIsShadowed && VirtualShadowMapArray.IsAllocated() && VirtualShadowMapId != INDEX_NONE;
@@ -805,6 +823,7 @@ void FSceneRenderer::RenderLocalLightsForVolumetricFog(
 				PermutationVector.Set< FInjectShadowedLocalLightPS::FLightFunction >(bUsesLightFunction);
 				PermutationVector.Set< FInjectShadowedLocalLightPS::FEnableShadows >(bIsShadowed);
 				PermutationVector.Set< FInjectShadowedLocalLightPS::FVirtualShadowMap >(bUseVSM);
+				PermutationVector.Set< FInjectShadowedLocalLightPS::FRectLightTexture >(bUsesRectLightTexture);
 
 				auto VertexShader = View.ShaderMap->GetShader< FWriteToBoundingSphereVS >();
 				TOptionalShaderMapRef<FWriteToSliceGS> GeometryShader(View.ShaderMap);
@@ -886,6 +905,7 @@ void FSceneRenderer::RenderLocalLightsForVolumetricFog(
 			{
 				bool bUsesLightFunction = ViewFamily.EngineShowFlags.LightFunctions
 					&& CheckForLightFunction(LightSceneInfo) && LightNeedsSeparateInjectionIntoVolumetricFogForLightFunction(LightSceneInfo);
+				bool bUsesRectLightTexture = GVolumetricFogRectLightTexture && LightSceneInfo->Proxy->HasSourceTexture();
 
 				FInjectShadowedLocalLightRGS::FParameters* PassParameters = GraphBuilder.AllocParameters<FInjectShadowedLocalLightRGS::FParameters>();
 				PassParameters->OutVolumeTexture = GraphBuilder.CreateUAV(OutLocalShadowedLightScattering);
@@ -910,6 +930,7 @@ void FSceneRenderer::RenderLocalLightsForVolumetricFog(
 				FInjectShadowedLocalLightRGS::FPermutationDomain PermutationVector;
 				PermutationVector.Set< FInjectShadowedLocalLightRGS::FTemporalReprojection >(bUseTemporalReprojection);
 				PermutationVector.Set< FInjectShadowedLocalLightRGS::FLightFunction >(bUsesLightFunction);
+				PermutationVector.Set< FInjectShadowedLocalLightRGS::FRectLightTexture >(bUsesRectLightTexture);
 
 				TShaderMapRef<FInjectShadowedLocalLightRGS> RayGenerationShader(GetGlobalShaderMap(FeatureLevel), PermutationVector);
 
