Fixed build error from bad merge.

diff --git a/Engine/Source/Runtime/Renderer/Private/SceneVisibility.cpp b/Engine/Source/Runtime/Renderer/Private/SceneVisibility.cpp
--- a/Engine/Source/Runtime/Renderer/Private/SceneVisibility.cpp
+++ b/Engine/Source/Runtime/Renderer/Private/SceneVisibility.cpp
@@ -1159,7 +1159,7 @@ void FRelevancePacket::Finalize()
 	NaniteCustomDepthInstances.AppendTo(WriteView.NaniteCustomDepthInstances);
 	WriteView.bUsesCustomDepth |= bUsesCustomDepth;
 	WriteView.bUsesCustomStencil |= bUsesCustomStencil;
-	WriteView.SubstrateViewData.MaxClosureCount = FMath::Max(WriteView.SubstrateViewData.MaxClosureCount, 8u - FMath::CountLeadingZeros8(SubstrateClosureCountMask));
+	WriteView.SubstrateViewData.MaxClosurePerPixel = FMath::Max(WriteView.SubstrateViewData.MaxClosurePerPixel, 8u - FMath::CountLeadingZeros8(SubstrateClosureCountMask));
 	WriteView.SubstrateViewData.MaxBytesPerPixel = FMath::Max(WriteView.SubstrateViewData.MaxBytesPerPixel, SubstrateUintPerPixel * 4u);
 	WriteView.SubstrateViewData.bUsesComplexSpecialRenderPath |= bUsesComplexSpecialRenderPath;
 	DirtyIndirectLightingCacheBufferPrimitives.AppendTo(WriteView.DirtyIndirectLightingCacheBufferPrimitives);
