Implement new shader permutation in TSRResolveHistory for [wavesize(16)] for Intel Arc GPUs

Submitting on behalf of Hugo Lamarche

diff --git a/Engine/Shaders/Private/TemporalSuperResolution/TSRResolveHistory.usf b/Engine/Shaders/Private/TemporalSuperResolution/TSRResolveHistory.usf
--- a/Engine/Shaders/Private/TemporalSuperResolution/TSRResolveHistory.usf
+++ b/Engine/Shaders/Private/TemporalSuperResolution/TSRResolveHistory.usf
@@ -5,7 +5,7 @@
 
 //------------------------------------------------------- CONFIG
 
-#if DIM_NYQUIST
+#if DIM_NYQUIST_WAVE_SIZE != 0
 	#define WAVE_COUNT_X 1
 	#define WAVE_COUNT_Y 1
 
@@ -16,7 +16,13 @@
 	#define TILE_SIZE (GROUP_TILE_SIZE - 2u)
 
 	#define LANE_OUTPUT_STRIDE_X 2u
-	#define LANE_OUTPUT_STRIDE_Y 1u
+	#if DIM_NYQUIST_WAVE_SIZE == 16
+		#define LANE_OUTPUT_STRIDE_Y 2u
+	#elif DIM_NYQUIST_WAVE_SIZE == 32
+		#define LANE_OUTPUT_STRIDE_Y 1u
+	#else
+		#error Unknown wave size
+	#endif
 	#define LANE_INPUT_STRIDE_X (LANE_OUTPUT_STRIDE_X * DOWNSAMPLE_FACTOR)
 	#define LANE_INPUT_STRIDE_Y (LANE_OUTPUT_STRIDE_Y * DOWNSAMPLE_FACTOR)
 	
@@ -44,7 +50,7 @@
 
 //------------------------------------------------------- CONVOLUTIONS
 
-#if DIM_NYQUIST
+#if DIM_NYQUIST_WAVE_SIZE != 0
 	#include "TSRConvolutionNetwork.ush"
 #endif
 
@@ -133,10 +139,10 @@ RWTexture2D<tsr_halfC> SceneColorOutputMip1;
 
 //------------------------------------------------------- ENTRY POINT NYQUIST DOWNSAMPLE
 
-#if DIM_NYQUIST
+#if DIM_NYQUIST_WAVE_SIZE != 0
 
 #if COMPILER_SUPPORTS_WAVE_SIZE
-	WAVESIZE(32)
+	WAVESIZE(DIM_NYQUIST_WAVE_SIZE)
 #endif
 [numthreads(LANE_COUNT, 1, 1)]
 void MainCS(
@@ -269,7 +275,7 @@ void MainCS(
 	}
 }
 
-#else // !DIM_NYQUIST
+#else // DIM_NYQUIST_WAVE_SIZE == 0
 
 //------------------------------------------------------- ENTRY POINT REFERENCE
 
@@ -388,4 +394,4 @@ void MainCS(
 	#endif
 }
 
-#endif // !DIM_NYQUIST
+#endif // DIM_NYQUIST_WAVE_SIZE == 0
diff --git a/Engine/Source/Runtime/Renderer/Private/PostProcess/TemporalSuperResolution.cpp b/Engine/Source/Runtime/Renderer/Private/PostProcess/TemporalSuperResolution.cpp
--- a/Engine/Source/Runtime/Renderer/Private/PostProcess/TemporalSuperResolution.cpp
+++ b/Engine/Source/Runtime/Renderer/Private/PostProcess/TemporalSuperResolution.cpp
@@ -877,7 +877,7 @@ class FTSRResolveHistoryCS : public FTSRShader
 	DECLARE_GLOBAL_SHADER(FTSRResolveHistoryCS);
 	SHADER_USE_PARAMETER_STRUCT(FTSRResolveHistoryCS, FTSRShader);
 
-	class FNyquistDim : SHADER_PERMUTATION_BOOL("DIM_NYQUIST");
+	class FNyquistDim : SHADER_PERMUTATION_SPARSE_INT("DIM_NYQUIST_WAVE_SIZE", 0, 16, 32);
 
 	using FPermutationDomain = TShaderPermutationDomain<FNyquistDim, FTSRShader::F16BitVALUDim, FTSRShader::FAlphaChannelDim>;
 
@@ -898,7 +898,15 @@ class FTSRResolveHistoryCS : public FTSRShader
 
 	static FPermutationDomain RemapPermutation(FPermutationDomain PermutationVector)
 	{
-		if (!PermutationVector.Get<FNyquistDim>())
+		int32 WaveSize = PermutationVector.Get<FNyquistDim>();
+
+		// WaveSize=16 is for Intel Arc GPU which also supports 16bits ops, so compiling WaveSize=16 32bit ops is useless and should instead fall back to WaveSize=0.
+		if (WaveSize == 16 && !PermutationVector.Get<FTSRShader::F16BitVALUDim>())
+		{
+			PermutationVector.Set<FNyquistDim>(0);
+		}
+
+		if (WaveSize == 0)
 		{
 			PermutationVector.Set<FTSRShader::F16BitVALUDim>(false);
 		}
@@ -917,10 +925,18 @@ class FTSRResolveHistoryCS : public FTSRShader
 
 		if (PermutationVector.Get<FNyquistDim>())
 		{
+			int32 WaveSize = PermutationVector.Get<FNyquistDim>();
+
 			if (FTSRShader::SupportsWaveOps(Parameters.Platform) == ERHIFeatureSupport::Unsupported)
 			{
 				return false;
 			}
+
+			if (WaveSize < int32(FDataDrivenShaderPlatformInfo::GetMinimumWaveSize(Parameters.Platform)) ||
+				WaveSize > int32(FDataDrivenShaderPlatformInfo::GetMaximumWaveSize(Parameters.Platform)))
+			{
+				return false;
+			}
 		}
 
 		if (!FTSRShader::ShouldCompile32or16BitPermutation(Parameters.Platform, PermutationVector.Get<FTSRShader::F16BitVALUDim>()))
@@ -1437,6 +1453,27 @@ FDefaultTemporalUpscaler::FOutputs AddTemporalSuperResolutionPasses(
 		return GraphBuilder.CreateUAV(DebugTexture);
 	};
 
+	auto SelectWaveSize = [&](const TArray<int32> WaveSizeDomain)
+	{
+		check(!WaveSizeDomain.IsEmpty());
+		int32 WaveSizeOps = 0;
+
+		if (bUseWaveOps)
+		{
+			if (WaveSizeOverride != 0 && WaveSizeDomain.Contains(WaveSizeOverride) && WaveSizeOverride >= GRHIMinimumWaveSize && WaveSizeOverride <= GRHIMaximumWaveSize)
+			{
+				WaveSizeOps = WaveSizeOverride;
+			}
+			else
+			{
+				const int32 MinimumWaveSizeWithPermutation = FMath::Max(GRHIMinimumWaveSize, WaveSizeDomain[0]);
+				WaveSizeOps = MinimumWaveSizeWithPermutation >= WaveSizeDomain[0] && MinimumWaveSizeWithPermutation <= WaveSizeDomain.Last() ? MinimumWaveSizeWithPermutation : 0;
+			}
+		}
+
+		return WaveSizeOps;
+	};
+
 	// Output
 	FRDGTextureRef UpdateHistoryOutputTexture = nullptr;
 	FRDGTextureRef SceneColorOutputTexture = nullptr;
@@ -2048,23 +2085,8 @@ FDefaultTemporalUpscaler::FOutputs AddTemporalSuperResolutionPasses(
 		FScreenPassTextureViewport TranslucencyViewport(
 			SeparateTranslucencyTexture->Desc.Extent, SeparateTranslucencyRect);
 
-		int32 WaveSizeOps = 0;
-
-		if (bUseWaveOps)
-		{
-			if ((WaveSizeOverride == 16 || WaveSizeOverride == 32 || WaveSizeOverride == 64) && WaveSizeOverride >= GRHIMinimumWaveSize && WaveSizeOverride <= GRHIMaximumWaveSize)
-			{
-				WaveSizeOps = WaveSizeOverride;
-			}
-			else
-			{
-				const int32 MinimumWaveSizeWithPermutation = FMath::Max(GRHIMinimumWaveSize, 16);
-				WaveSizeOps = MinimumWaveSizeWithPermutation >= 16 && MinimumWaveSizeWithPermutation <= 64 ? MinimumWaveSizeWithPermutation : 0;
-			}
-		}
-
 		FTSRRejectShadingCS::FPermutationDomain PermutationVector;
-		PermutationVector.Set<FTSRRejectShadingCS::FWaveSizeOps>(WaveSizeOps);
+		PermutationVector.Set<FTSRRejectShadingCS::FWaveSizeOps>(SelectWaveSize({ 16, 32, 64 }));
 		PermutationVector.Set<FTSRRejectShadingCS::FFlickeringDetectionDim>(FlickeringFramePeriod > 0.0f);
 		PermutationVector.Set<FTSRRejectShadingCS::FHistoryResurrectionDim>(bCanResurrectHistory);
 		PermutationVector.Set<FTSRShader::F16BitVALUDim>(bUse16BitVALU);
@@ -2471,7 +2493,7 @@ FDefaultTemporalUpscaler::FOutputs AddTemporalSuperResolutionPasses(
 		PassParameters->DebugOutput = CreateDebugUAV(OutputExtent, TEXT("Debug.TSR.ResolveHistory"));
 
 		FTSRResolveHistoryCS::FPermutationDomain PermutationVector;
-		PermutationVector.Set<FTSRResolveHistoryCS::FNyquistDim>(bNyquistHistory && bUseWaveOps && GRHIMaximumWaveSize >= 32 && GRHIMinimumWaveSize <= 32);
+		PermutationVector.Set<FTSRResolveHistoryCS::FNyquistDim>(bNyquistHistory ? SelectWaveSize({ 16, 32 }) : 0);
 		PermutationVector.Set<FTSRShader::F16BitVALUDim>(bUse16BitVALU);
 		PermutationVector.Set<FTSRShader::FAlphaChannelDim>(bSupportsAlpha);
 		PermutationVector = FTSRResolveHistoryCS::RemapPermutation(PermutationVector);
@@ -2479,9 +2501,9 @@ FDefaultTemporalUpscaler::FOutputs AddTemporalSuperResolutionPasses(
 		TShaderMapRef<FTSRResolveHistoryCS> ComputeShader(View.ShaderMap, PermutationVector);
 		FComputeShaderUtils::AddPass(
 			GraphBuilder,
-			RDG_EVENT_NAME("TSR ResolveHistory(#%d%s%s%s) %dx%d",
+			RDG_EVENT_NAME("TSR ResolveHistory(#%d WaveSize=%d%s%s) %dx%d",
 				PermutationVector.ToDimensionValueId(),
-				PermutationVector.Get<FTSRResolveHistoryCS::FNyquistDim>() ? TEXT(" WaveOps") : TEXT(""),
+				PermutationVector.Get<FTSRResolveHistoryCS::FNyquistDim>(),
 				PermutationVector.Get<FTSRShader::F16BitVALUDim>() ? TEXT(" 16bit") : TEXT(""),
 				PermutationVector.Get<FTSRShader::FAlphaChannelDim>() ? TEXT(" AlphaChannel") : TEXT(""),
 				OutputRect.Width(), OutputRect.Height()),
