Merging CL#29037036

Make 'UniformView' GPUScene path work with DynamicMeshPass.
This needs some extra care as DynamicMeshPass can be mixed with a regular mesh passes, so code that uses it has to restore static 'UniformView' binding after DynamicMeshPass call

diff --git a/Engine/Source/Runtime/Engine/Private/PrimitiveUniformShaderParameters.cpp b/Engine/Source/Runtime/Engine/Private/PrimitiveUniformShaderParameters.cpp
--- a/Engine/Source/Runtime/Engine/Private/PrimitiveUniformShaderParameters.cpp
+++ b/Engine/Source/Runtime/Engine/Private/PrimitiveUniformShaderParameters.cpp
@@ -281,123 +281,4 @@ FPrimitiveUniformShaderParameters GetIdentityPrimitiveParameters()
 			.WorldBounds(FBoxSphereBounds(EForceInit::ForceInit))
 			.LocalBounds(FBoxSphereBounds(EForceInit::ForceInit))
 		.Build();
-}
-
-
-FBatchedPrimitiveShaderData::FBatchedPrimitiveShaderData(const FPrimitiveSceneProxy* RESTRICT Proxy)
-	: Data(InPlace, NoInit)
-{
-	FPrimitiveUniformShaderParametersBuilder Builder = FPrimitiveUniformShaderParametersBuilder{};
-	Proxy->BuildUniformShaderParameters(Builder);
-	Setup(Builder.Build());
-}
-
-void FBatchedPrimitiveShaderData::Setup(const FPrimitiveUniformShaderParameters& PrimitiveUniformShaderParameters)
-{
-	static_assert(NUM_LIGHTING_CHANNELS == 3, "The FPrimitiveSceneShaderData packing currently assumes a maximum of 3 lighting channels.");
-	
-	// Note: layout must match LoadPrimitiveDataUBO in SceneDataMobileLoader.ush
-	int32 i = 0;
-
-	const bool bAllowStaticLighting = FReadOnlyCVARCache::Get().bAllowStaticLighting;
-	if (bAllowStaticLighting)
-	{
-		FVector4f LightMapUVScaleBias = FVector4f(1, 1, 0, 0);
-		FVector4f ShadowMapUVScaleBias = FVector4f(1, 1, 0, 0);
-		uint32 LightMapDataIdx = 0;
-
-		// FIXME: valid lightmap data
-		Data[i+0] = LightMapUVScaleBias;
-		Data[i+1] = ShadowMapUVScaleBias;
-		Data[i+2] = FVector4f(FMath::AsFloat(LightMapDataIdx), 0.f, 0.f, 0.f);
-		i+=3;
-	}
-
-	// TilePosition, Flags
-	{
-		Data[i].X = PrimitiveUniformShaderParameters.TilePosition.X;
-		Data[i].Y = PrimitiveUniformShaderParameters.TilePosition.Y;
-		Data[i].Z = PrimitiveUniformShaderParameters.TilePosition.Z;
-		Data[i].W = FMath::AsFloat(PrimitiveUniformShaderParameters.Flags);
-		i+=1;
-	}
-	
-	// LocalToWorld
-	{
-		FMatrix44f LocalToRelativeWorldTranspose = PrimitiveUniformShaderParameters.LocalToRelativeWorld.GetTransposed();
-		Data[i+0] = *(const FVector4f*)&LocalToRelativeWorldTranspose.M[0][0];
-		Data[i+1] = *(const FVector4f*)&LocalToRelativeWorldTranspose.M[1][0];
-		Data[i+2] = *(const FVector4f*)&LocalToRelativeWorldTranspose.M[2][0];
-		i+=3;
-	}
-
-	// InvNonUniformScale, TODO .w
-	{
-		Data[i+0] = FVector4f(PrimitiveUniformShaderParameters.InvNonUniformScale, 0.0f);
-		i+=1;
-	}
-
-	// ObjectWorldPosition, Radius
-	{
-		Data[i+0] = PrimitiveUniformShaderParameters.ObjectRelativeWorldPositionAndRadius;
-		i+=1;
-	}
-
-	// ActorWorldPosition, TODO .w
-	{
-		Data[i+0] = FVector4f(PrimitiveUniformShaderParameters.ActorRelativeWorldPosition, 0.0f);
-		i+=1;
-	}
-
-	// ObjectOrientation, ObjectBoundsX
-	{
-		Data[i+0] = FVector4f(PrimitiveUniformShaderParameters.ObjectOrientation, PrimitiveUniformShaderParameters.ObjectBoundsX);
-		i+=1;
-	}
-
-	// LocalObjectBoundsMin, ObjectBoundsY
-	{
-		Data[i+0] = FVector4f(PrimitiveUniformShaderParameters.LocalObjectBoundsMin, PrimitiveUniformShaderParameters.ObjectBoundsY);
-		i+=1;
-	}
-
-	// LocalObjectBoundsMax, ObjectBoundsZ
-	{
-		Data[i+0] = FVector4f(PrimitiveUniformShaderParameters.LocalObjectBoundsMax, PrimitiveUniformShaderParameters.ObjectBoundsZ);
-		i+=1;
-	}
-
-	// WorldToLocal
-	{
-		FMatrix44f RelativeWorldToLocalTranspose = PrimitiveUniformShaderParameters.RelativeWorldToLocal.GetTransposed();
-		Data[i+0] = *(const FVector4f*)&RelativeWorldToLocalTranspose.M[0][0];
-		Data[i+1] = *(const FVector4f*)&RelativeWorldToLocalTranspose.M[1][0];
-		Data[i+2] = *(const FVector4f*)&RelativeWorldToLocalTranspose.M[2][0];
-		i+=3;
-	}
-
-	// PreviousLocalToWorld
-	{
-		FMatrix44f PreviousLocalToRelativeWorldTranspose = PrimitiveUniformShaderParameters.PreviousLocalToRelativeWorld.GetTransposed();
-		Data[i+0] = *(const FVector4f*)&PreviousLocalToRelativeWorldTranspose.M[0][0];
-		Data[i+1] = *(const FVector4f*)&PreviousLocalToRelativeWorldTranspose.M[1][0];
-		Data[i+2] = *(const FVector4f*)&PreviousLocalToRelativeWorldTranspose.M[2][0];
-		i+=3;
-	}
-
-	// PreviousWorldToLocal
-	{
-		FMatrix44f PreviousRelativeWorldToLocalTranspose = PrimitiveUniformShaderParameters.PreviousRelativeWorldToLocal.GetTransposed();
-		Data[i+0] = *(const FVector4f*)&PreviousRelativeWorldToLocalTranspose.M[0][0];
-		Data[i+1] = *(const FVector4f*)&PreviousRelativeWorldToLocalTranspose.M[1][0];
-		Data[i+2] = *(const FVector4f*)&PreviousRelativeWorldToLocalTranspose.M[2][0];
-		i+=3;
-	}
-	
-	// Set all the custom primitive data float4. This matches the loop in SceneData.ush
-	int32 NumCustomData = FMath::Min<int32>(FCustomPrimitiveData::NumCustomPrimitiveDataFloat4s, DataStrideInFloat4s - i);
-	for (int32 DataIndex = 0; DataIndex < FCustomPrimitiveData::NumCustomPrimitiveDataFloat4s; ++DataIndex)
-	{
-		Data[i + DataIndex] = PrimitiveUniformShaderParameters.CustomPrimitiveData[DataIndex];
-	}
 }
\ No newline at end of file
diff --git a/Engine/Source/Runtime/Engine/Public/PrimitiveSceneShaderData.h b/Engine/Source/Runtime/Engine/Public/PrimitiveSceneShaderData.h
--- a/Engine/Source/Runtime/Engine/Public/PrimitiveSceneShaderData.h
+++ b/Engine/Source/Runtime/Engine/Public/PrimitiveSceneShaderData.h
@@ -122,27 +122,3 @@ public:
 extern ENGINE_API TGlobalResource<FSinglePrimitiveStructured> GIdentityPrimitiveBuffer;
 UE_DEPRECATED(5.3, "Use IRendererModule::CreateSinglePrimitiveSceneUniformBuffer instead")
 extern ENGINE_API TGlobalResource<FSinglePrimitiveStructured> GTilePrimitiveBuffer;
-
-
-struct FBatchedPrimitiveShaderData
-{
-	static const uint32 DataStrideInFloat4s = BATCHED_PRIMITIVE_DATA_STRIDE_FLOAT4;
-
-	TStaticArray<FVector4f, DataStrideInFloat4s> Data;
-
-	FBatchedPrimitiveShaderData()
-		: Data(InPlace, NoInit)
-	{
-		Setup(GetIdentityPrimitiveParameters());
-	}
-
-	explicit FBatchedPrimitiveShaderData(const FPrimitiveUniformShaderParameters& PrimitiveUniformShaderParameters)
-		: Data(InPlace, NoInit)
-	{
-		Setup(PrimitiveUniformShaderParameters);
-	}
-
-	ENGINE_API FBatchedPrimitiveShaderData(const class FPrimitiveSceneProxy* RESTRICT Proxy);
-
-	ENGINE_API void Setup(const FPrimitiveUniformShaderParameters& PrimitiveUniformShaderParameters);
-};
diff --git a/Engine/Source/Runtime/Renderer/Private/GPUScene.cpp b/Engine/Source/Runtime/Renderer/Private/GPUScene.cpp
--- a/Engine/Source/Runtime/Renderer/Private/GPUScene.cpp
+++ b/Engine/Source/Runtime/Renderer/Private/GPUScene.cpp
@@ -195,6 +195,19 @@ void FGPUScenePrimitiveCollector::Add(
 	OutInstanceSceneDataOffset = PrimitiveData.LocalInstanceSceneDataOffset;
 }
 
+const FPrimitiveUniformShaderParameters* FGPUScenePrimitiveCollector::GetPrimitiveShaderParameters(int32 DrawPrimitiveId) const
+{
+	if (UploadData != nullptr && (DrawPrimitiveId & GPrimIDDynamicFlag) != 0)
+	{
+		int32 DynamicPrimitiveIndex = DrawPrimitiveId & (~GPrimIDDynamicFlag);
+		if (UploadData->PrimitiveData.IsValidIndex(DynamicPrimitiveIndex))
+		{
+			return UploadData->PrimitiveData[DynamicPrimitiveIndex].ShaderParams;
+		}
+	}
+	return nullptr;
+}
+
 #if DO_CHECK
 
 void FGPUScenePrimitiveCollector::CheckPrimitiveProcessed(uint32 PrimitiveIndex, const FGPUScene& GPUScene) const
@@ -2215,56 +2228,130 @@ void FGPUScene::AddUpdatePrimitiveIdsPass(FRDGBuilder& GraphBuilder, FInstanceGP
 	}
 }
 
-void FGPUSceneCompactInstanceData::Init(const FGPUScenePrimitiveCollector* PrimitiveCollector, int32 PrimitiveId)
+FBatchedPrimitiveShaderData::FBatchedPrimitiveShaderData(const FPrimitiveSceneProxy* Proxy)
+	: Data(InPlace, NoInit)
 {
-	FMatrix44f LocalToRelativeWorld = FMatrix44f::Identity;
-	FVector3f TilePosition = FVector3f::ZeroVector;
-	FVector3f InvNonUniformScale = FVector3f::OneVector;
-	float PrimitiveFlags = 0;
+	FPrimitiveUniformShaderParametersBuilder Builder = FPrimitiveUniformShaderParametersBuilder{};
+	Proxy->BuildUniformShaderParameters(Builder);
+	Setup(Builder.Build());
+}
 
-	int32 DynamicPrimitiveId = PrimitiveId;
-	if (PrimitiveCollector && PrimitiveCollector->UploadData && !PrimitiveCollector->GetPrimitiveIdRange().IsEmpty())
-	{
-		DynamicPrimitiveId = PrimitiveCollector->GetPrimitiveIdRange().GetLowerBoundValue() + PrimitiveId;
-		if (PrimitiveCollector->GetPrimitiveIdRange().Contains(DynamicPrimitiveId))
-		{
-			const FPrimitiveUniformShaderParameters& PrimitiveData = *PrimitiveCollector->UploadData->PrimitiveData[PrimitiveId].ShaderParams;
-			LocalToRelativeWorld = PrimitiveData.LocalToRelativeWorld;
-			TilePosition = PrimitiveData.TilePosition;
-			InvNonUniformScale = PrimitiveData.InvNonUniformScale;
-			PrimitiveFlags = *(float*)&PrimitiveData.Flags;
-		}
-	}
+void FBatchedPrimitiveShaderData::Emplace(FBatchedPrimitiveShaderData* Dest, const FPrimitiveUniformShaderParameters& ShaderParams)
+{
+	new (Dest) FBatchedPrimitiveShaderData(ShaderParams);
+}
 
-	// must match packing in SceneDataMobileWriter.ush
-	LocalToWorld0	= LocalToRelativeWorld.GetScaledAxis(EAxis::X);
-	LocalToWorld1	= LocalToRelativeWorld.GetScaledAxis(EAxis::Y);
-	LocalToWorld2	= LocalToRelativeWorld.GetScaledAxis(EAxis::Z);
-	LocalToWorld3	= LocalToRelativeWorld.GetOrigin();
-	LocalToWorld0.W = TilePosition.X;
-	LocalToWorld1.W = TilePosition.Y;
-	LocalToWorld2.W = TilePosition.Z;
-	InvNonUniformScaleAndFlags = FVector4f(InvNonUniformScale, PrimitiveFlags);
+void FBatchedPrimitiveShaderData::Emplace(FBatchedPrimitiveShaderData* Dest, const FPrimitiveSceneProxy* Proxy)
+{
+	new (Dest) FBatchedPrimitiveShaderData(Proxy);
 }
 
-void FGPUSceneCompactInstanceData::Init(const FScene* Scene, int32 PrimitiveId)
+void FBatchedPrimitiveShaderData::Setup(const FPrimitiveUniformShaderParameters& PrimitiveUniformShaderParameters)
 {
-	FPrimitiveUniformShaderParametersBuilder Builder = FPrimitiveUniformShaderParametersBuilder{}.Defaults();
-	if (Scene && PrimitiveId >= 0 && PrimitiveId < Scene->PrimitiveTransforms.Num())
+	// Note: layout must match LoadPrimitiveDataUBO in SceneDataMobileLoader.ush
+	int32 i = 0;
+
+	const bool bAllowStaticLighting = FReadOnlyCVARCache::Get().bAllowStaticLighting;
+	if (bAllowStaticLighting)
 	{
-		Builder.LocalToWorld(Scene->PrimitiveTransforms[PrimitiveId]);
+		FVector4f LightMapUVScaleBias = FVector4f(1, 1, 0, 0);
+		FVector4f ShadowMapUVScaleBias = FVector4f(1, 1, 0, 0);
+		uint32 LightMapDataIdx = 0;
+
+		// FIXME: valid lightmap data
+		Data[i+0] = LightMapUVScaleBias;
+		Data[i+1] = ShadowMapUVScaleBias;
+		Data[i+2] = FVector4f(FMath::AsFloat(LightMapDataIdx), 0.f, 0.f, 0.f);
+		i+=3;
+	}
+
+	// TilePosition, Flags
+	{
+		Data[i].X = PrimitiveUniformShaderParameters.TilePosition.X;
+		Data[i].Y = PrimitiveUniformShaderParameters.TilePosition.Y;
+		Data[i].Z = PrimitiveUniformShaderParameters.TilePosition.Z;
+		Data[i].W = FMath::AsFloat(PrimitiveUniformShaderParameters.Flags);
+		i+=1;
+	}
+
+	// LocalToWorld
+	{
+		FMatrix44f LocalToRelativeWorldTranspose = PrimitiveUniformShaderParameters.LocalToRelativeWorld.GetTransposed();
+		Data[i+0] = *(const FVector4f*)&LocalToRelativeWorldTranspose.M[0][0];
+		Data[i+1] = *(const FVector4f*)&LocalToRelativeWorldTranspose.M[1][0];
+		Data[i+2] = *(const FVector4f*)&LocalToRelativeWorldTranspose.M[2][0];
+		i+=3;
+	}
+
+	// InvNonUniformScale, TODO .w
+	{
+		Data[i+0] = FVector4f(PrimitiveUniformShaderParameters.InvNonUniformScale, 0.0f);
+		i+=1;
+	}
+
+	// ObjectWorldPosition, Radius
+	{
+		Data[i+0] = PrimitiveUniformShaderParameters.ObjectRelativeWorldPositionAndRadius;
+		i+=1;
+	}
+
+	// ActorWorldPosition, TODO .w
+	{
+		Data[i+0] = FVector4f(PrimitiveUniformShaderParameters.ActorRelativeWorldPosition, 0.0f);
+		i+=1;
+	}
+
+	// ObjectOrientation, ObjectBoundsX
+	{
+		Data[i+0] = FVector4f(PrimitiveUniformShaderParameters.ObjectOrientation, PrimitiveUniformShaderParameters.ObjectBoundsX);
+		i+=1;
+	}
+
+	// LocalObjectBoundsMin, ObjectBoundsY
+	{
+		Data[i+0] = FVector4f(PrimitiveUniformShaderParameters.LocalObjectBoundsMin, PrimitiveUniformShaderParameters.ObjectBoundsY);
+		i+=1;
+	}
+
+	// LocalObjectBoundsMax, ObjectBoundsZ
+	{
+		Data[i+0] = FVector4f(PrimitiveUniformShaderParameters.LocalObjectBoundsMax, PrimitiveUniformShaderParameters.ObjectBoundsZ);
+		i+=1;
+	}
+
+	// WorldToLocal
+	{
+		FMatrix44f RelativeWorldToLocalTranspose = PrimitiveUniformShaderParameters.RelativeWorldToLocal.GetTransposed();
+		Data[i+0] = *(const FVector4f*)&RelativeWorldToLocalTranspose.M[0][0];
+		Data[i+1] = *(const FVector4f*)&RelativeWorldToLocalTranspose.M[1][0];
+		Data[i+2] = *(const FVector4f*)&RelativeWorldToLocalTranspose.M[2][0];
+		i+=3;
+	}
+
+	// PreviousLocalToWorld
+	{
+		FMatrix44f PreviousLocalToRelativeWorldTranspose = PrimitiveUniformShaderParameters.PreviousLocalToRelativeWorld.GetTransposed();
+		Data[i+0] = *(const FVector4f*)&PreviousLocalToRelativeWorldTranspose.M[0][0];
+		Data[i+1] = *(const FVector4f*)&PreviousLocalToRelativeWorldTranspose.M[1][0];
+		Data[i+2] = *(const FVector4f*)&PreviousLocalToRelativeWorldTranspose.M[2][0];
+		i+=3;
+	}
+
+	// PreviousWorldToLocal
+	{
+		FMatrix44f PreviousRelativeWorldToLocalTranspose = PrimitiveUniformShaderParameters.PreviousRelativeWorldToLocal.GetTransposed();
+		Data[i+0] = *(const FVector4f*)&PreviousRelativeWorldToLocalTranspose.M[0][0];
+		Data[i+1] = *(const FVector4f*)&PreviousRelativeWorldToLocalTranspose.M[1][0];
+		Data[i+2] = *(const FVector4f*)&PreviousRelativeWorldToLocalTranspose.M[2][0];
+		i+=3;
+	}
+
+	// Set all the custom primitive data float4. This matches the loop in SceneData.ush
+	int32 NumCustomData = FMath::Min<int32>(FCustomPrimitiveData::NumCustomPrimitiveDataFloat4s, DataStrideInFloat4s - i);
+	for (int32 DataIndex = 0; DataIndex < FCustomPrimitiveData::NumCustomPrimitiveDataFloat4s; ++DataIndex)
+	{
+		Data[i + DataIndex] = PrimitiveUniformShaderParameters.CustomPrimitiveData[DataIndex];
 	}
-	const FPrimitiveUniformShaderParameters PrimitiveData = Builder.Build();
-	
-	// must match packing in SceneDataMobileWriter.ush
-	LocalToWorld0	= PrimitiveData.LocalToRelativeWorld.GetScaledAxis(EAxis::X);
-	LocalToWorld1	= PrimitiveData.LocalToRelativeWorld.GetScaledAxis(EAxis::Y);
-	LocalToWorld2	= PrimitiveData.LocalToRelativeWorld.GetScaledAxis(EAxis::Z);
-	LocalToWorld3	= PrimitiveData.LocalToRelativeWorld.GetOrigin();
-	LocalToWorld0.W = PrimitiveData.TilePosition.X;
-	LocalToWorld1.W = PrimitiveData.TilePosition.Y;
-	LocalToWorld2.W = PrimitiveData.TilePosition.Z;
-	InvNonUniformScaleAndFlags = FVector4f(PrimitiveData.InvNonUniformScale, *(float*)&PrimitiveData.Flags);
 }
 
 void FGPUScene::AddClearInstancesPass(FRDGBuilder& GraphBuilder, FInstanceCullingOcclusionQueryRenderer* OcclusionQueryRenderer)
diff --git a/Engine/Source/Runtime/Renderer/Private/GPUScene.h b/Engine/Source/Runtime/Renderer/Private/GPUScene.h
--- a/Engine/Source/Runtime/Renderer/Private/GPUScene.h
+++ b/Engine/Source/Runtime/Renderer/Private/GPUScene.h
@@ -23,6 +23,7 @@ class FGPUScene;
 class FGPUSceneDynamicContext;
 class FViewUniformShaderParameters;
 class FInstanceCullingOcclusionQueryRenderer;
+class FPrimitiveSceneProxy;
 class IVisibilityTaskData;
 
 DECLARE_GPU_STAT_NAMED_EXTERN(GPUSceneUpdate, TEXT("GPUSceneUpdate"))
@@ -97,6 +98,7 @@ public:
 	int32 Num() const {	return UploadData != nullptr ? UploadData->PrimitiveData.Num() : 0; }
 	int32 NumInstances() const { return UploadData != nullptr ? UploadData->TotalInstanceCount : 0; }
 	int32 NumPayloadDataSlots() const { return UploadData != nullptr ? UploadData->InstancePayloadDataFloat4Count : 0; }
+	const FPrimitiveUniformShaderParameters* GetPrimitiveShaderParameters(int32 PrimitiveId) const;
 
 #if DO_CHECK
 	/**
@@ -109,7 +111,6 @@ private:
 
 	friend class FGPUScene;
 	friend class FGPUSceneDynamicContext;
-	friend struct FGPUSceneCompactInstanceData;
 	friend struct FUploadDataSourceAdapterDynamicPrimitives;
 
 	struct FPrimitiveData
@@ -464,14 +465,29 @@ private:
 	FGPUScene& GPUScene;
 };
 
-struct FGPUSceneCompactInstanceData
+struct FBatchedPrimitiveShaderData
 {
-	FVector4f LocalToWorld0;
-	FVector4f LocalToWorld1;
-	FVector4f LocalToWorld2;
-	FVector4f LocalToWorld3;
-	FVector4f InvNonUniformScaleAndFlags;
-
-	void Init(const FGPUScenePrimitiveCollector* PrimitiveCollector, int32 PrimitiveId);
-	void Init(const FScene* Scene, int32 PrimitiveId);
+	static const uint32 DataStrideInFloat4s = BATCHED_PRIMITIVE_DATA_STRIDE_FLOAT4;
+
+	TStaticArray<FVector4f, DataStrideInFloat4s> Data;
+
+	FBatchedPrimitiveShaderData()
+		: Data(InPlace, NoInit)
+	{
+		Setup(GetIdentityPrimitiveParameters());
+	}
+
+	explicit FBatchedPrimitiveShaderData(const FPrimitiveUniformShaderParameters& PrimitiveUniformShaderParameters)
+		: Data(InPlace, NoInit)
+	{
+		Setup(PrimitiveUniformShaderParameters);
+	}
+
+	explicit FBatchedPrimitiveShaderData(const FPrimitiveSceneProxy* Proxy);
+	
+	static void Emplace(FBatchedPrimitiveShaderData* Dest, const FPrimitiveUniformShaderParameters& ShaderParams);
+	static void Emplace(FBatchedPrimitiveShaderData* Dest, const FPrimitiveSceneProxy* Proxy);
+
+private:
+	void Setup(const FPrimitiveUniformShaderParameters& PrimitiveUniformShaderParameters);
 };
diff --git a/Engine/Source/Runtime/Renderer/Private/MeshDrawCommands.cpp b/Engine/Source/Runtime/Renderer/Private/MeshDrawCommands.cpp
--- a/Engine/Source/Runtime/Renderer/Private/MeshDrawCommands.cpp
+++ b/Engine/Source/Runtime/Renderer/Private/MeshDrawCommands.cpp
@@ -96,7 +96,7 @@ FPrimitiveIdVertexBufferPoolEntry FPrimitiveIdVertexBufferPool::Allocate(FRHICom
 		NewEntry.LastDiscardId = DiscardId;
 		NewEntry.BufferSize = BufferSize;
 		FRHIResourceCreateInfo CreateInfo(TEXT("FPrimitiveIdVertexBufferPool"));
-		NewEntry.BufferRHI = RHICmdList.CreateBuffer(NewEntry.BufferSize, BUF_VertexBuffer | BUF_Volatile | BUF_UniformBuffer, 0, ERHIAccess::VertexOrIndexBuffer, CreateInfo);
+		NewEntry.BufferRHI = RHICmdList.CreateBuffer(NewEntry.BufferSize, BUF_VertexBuffer | BUF_Volatile | BUF_UniformBuffer | BUF_ShaderResource | BUF_ByteAddressBuffer, 0, ERHIAccess::VertexOrIndexBuffer, CreateInfo);
 
 		return NewEntry;
 	}
@@ -465,7 +465,7 @@ static void BuildMeshDrawCommandPrimitiveIdBuffer(
 	int32& NewPassVisibleMeshDrawCommandsNum,
 	EShaderPlatform ShaderPlatform,
 	uint32 InstanceFactor,
-	TFunctionRef<void(int32, int32)> WritePrimitiveDataFn)
+	TFunctionRef<void(int32, int32, int32)> WritePrimitiveDataFn)
 {
 	QUICK_SCOPE_CYCLE_COUNTER(STAT_BuildMeshDrawCommandPrimitiveIdBuffer);
 
@@ -548,7 +548,7 @@ static void BuildMeshDrawCommandPrimitiveIdBuffer(
 			//@todo - refactor into instance step rate in the RHI
 			for (uint32 InstanceFactorIndex = 0; InstanceFactorIndex < InstanceFactor; InstanceFactorIndex++, PrimitiveIdIndex++)
 			{
-				WritePrimitiveDataFn(PrimitiveIdIndex, VisibleMeshDrawCommand.PrimitiveIdInfo.DrawPrimitiveId);
+				WritePrimitiveDataFn(PrimitiveIdIndex, VisibleMeshDrawCommand.PrimitiveIdInfo.DrawPrimitiveId, VisibleMeshDrawCommand.PrimitiveIdInfo.ScenePrimitiveId);
 			}
 		}
 
@@ -569,7 +569,7 @@ static void BuildMeshDrawCommandPrimitiveIdBuffer(
 			const FVisibleMeshDrawCommand& VisibleMeshDrawCommand = VisibleMeshDrawCommands[DrawCommandIndex];
 			for (uint32 InstanceFactorIndex = 0; InstanceFactorIndex < InstanceFactor; InstanceFactorIndex++, PrimitiveIdIndex++)
 			{
-				WritePrimitiveDataFn(PrimitiveIdIndex, VisibleMeshDrawCommand.PrimitiveIdInfo.DrawPrimitiveId);
+				WritePrimitiveDataFn(PrimitiveIdIndex, VisibleMeshDrawCommand.PrimitiveIdInfo.DrawPrimitiveId, VisibleMeshDrawCommand.PrimitiveIdInfo.ScenePrimitiveId);
 			}
 		}
 	}
@@ -1234,33 +1234,46 @@ void SortAndMergeDynamicPassMeshDrawCommands(
 				DynamicPrimitiveIdMax = DynamicPrimitiveIdRange.GetUpperBoundValue();
 			}
 
+			const bool bUsesUniformView = PlatformGPUSceneUsesUniformBufferView(SceneView.GetShaderPlatform());
+
 			const uint32 PrimitiveIdBufferStride = FInstanceCullingContext::GetInstanceIdBufferStride(SceneView.GetShaderPlatform());
 			const int32 MaxNumPrimitives = InstanceFactor * NumDrawCommands;
-			const int32 PrimitiveIdBufferDataSize = MaxNumPrimitives * PrimitiveIdBufferStride;
+			const int32 PrimitiveIdBufferDataSize = MaxNumPrimitives * PrimitiveIdBufferStride + (bUsesUniformView ? (PLATFORM_MAX_UNIFORM_BUFFER_RANGE - PrimitiveIdBufferStride) : 0);
 			FPrimitiveIdVertexBufferPoolEntry Entry = GPrimitiveIdVertexBufferPool.Allocate(RHICmdList, PrimitiveIdBufferDataSize);
 			OutPrimitiveIdVertexBuffer = Entry.BufferRHI;
 			void* PrimitiveIdBufferData = RHICmdList.LockBuffer(OutPrimitiveIdVertexBuffer, 0, PrimitiveIdBufferDataSize, RLM_WriteOnly);
 
-			if (PlatformGPUSceneUsesUniformBufferView(SceneView.GetShaderPlatform()))
+			if (bUsesUniformView)
 			{
-				check(sizeof(FGPUSceneCompactInstanceData) <= PrimitiveIdBufferStride);
-
-				auto WritePrimitiveDataFn = [&](int32 PrimitiveIndex, int32 PrimitiveId) 
+				check(PrimitiveIdBufferStride == sizeof(FBatchedPrimitiveShaderData));
+				const FPrimitiveUniformShaderParameters* IdentityShaderParams = (const FPrimitiveUniformShaderParameters*)GIdentityPrimitiveUniformBuffer.GetContents();
+			
+				auto WritePrimitiveDataFn = [&](int32 PrimitiveIndex, int32 DrawPrimitiveId, int32 ScenePrimitiveId)
 				{
 					checkSlow(PrimitiveIndex < MaxNumPrimitives);
-					uint8* PrimitiveIdBufferDataBytes = reinterpret_cast<uint8*>(PrimitiveIdBufferData) + (PrimitiveIndex * PrimitiveIdBufferStride);
-					FMemory::Memzero(PrimitiveIdBufferDataBytes, PrimitiveIdBufferStride);
-					FGPUSceneCompactInstanceData* PrimitiveData = reinterpret_cast<FGPUSceneCompactInstanceData*>(PrimitiveIdBufferDataBytes);
+					FBatchedPrimitiveShaderData* PrimitiveData = reinterpret_cast<FBatchedPrimitiveShaderData*>(PrimitiveIdBufferData) + PrimitiveIndex;
 
-					if ((PrimitiveId & GPrimIDDynamicFlag) != 0)
+					if ((DrawPrimitiveId & GPrimIDDynamicFlag) != 0)
 					{
-						PrimitiveId = PrimitiveId & (~GPrimIDDynamicFlag);
-						PrimitiveData->Init(DynamicPrimitiveCollector, PrimitiveId);
+						const FPrimitiveUniformShaderParameters* ShaderParams = DynamicPrimitiveCollector->GetPrimitiveShaderParameters(DrawPrimitiveId);
+						if (ShaderParams == nullptr)
+						{
+							ShaderParams = IdentityShaderParams;
+						}
+						FBatchedPrimitiveShaderData::Emplace(PrimitiveData, *ShaderParams);
 					}
 					else
 					{
 						const FScene* Scene = SceneView.Family->Scene->GetRenderScene();
-						PrimitiveData->Init(Scene, PrimitiveId);
+						if (ScenePrimitiveId >= 0 && ScenePrimitiveId < Scene->Primitives.Num())
+						{
+							FPrimitiveSceneProxy* PrimitiveProxy = Scene->PrimitiveSceneProxies[ScenePrimitiveId];
+							FBatchedPrimitiveShaderData::Emplace(PrimitiveData, PrimitiveProxy);
+						}
+						else
+						{
+							FBatchedPrimitiveShaderData::Emplace(PrimitiveData, *IdentityShaderParams);
+						}
 					}
 				};
 				
@@ -1280,7 +1293,7 @@ void SortAndMergeDynamicPassMeshDrawCommands(
 			{
 				int32* RESTRICT PrimitiveIds = reinterpret_cast<int32*>(PrimitiveIdBufferData);
 			
-				auto WritePrimitiveDataFn = [&](int32 PrimitiveIndex, int32 DrawPrimitiveId) 
+				auto WritePrimitiveDataFn = [&](int32 PrimitiveIndex, int32 DrawPrimitiveId, int32 ScenePrimitiveId) 
 				{
 					checkSlow(PrimitiveIndex < MaxNumPrimitives);
 					PrimitiveIds[PrimitiveIndex] = TranslatePrimitiveId(DrawPrimitiveId, DynamicPrimitiveIdOffset, DynamicPrimitiveIdMax);
diff --git a/Engine/Source/Runtime/Renderer/Private/MeshPassProcessor.cpp b/Engine/Source/Runtime/Renderer/Private/MeshPassProcessor.cpp
--- a/Engine/Source/Runtime/Renderer/Private/MeshPassProcessor.cpp
+++ b/Engine/Source/Runtime/Renderer/Private/MeshPassProcessor.cpp
@@ -1775,6 +1775,36 @@ void ApplyViewOverridesToMeshDrawCommands(const FSceneView& View, FMeshCommandOn
 	}
 }
 
+class FDynamicBatchedPrimitiveLayout : public FRenderResource
+{
+public:
+	void InitRHI(FRHICommandListBase& RHICmdList) override 
+	{
+		// This Layout fully replicates BatchedPrimitive UB
+		// we replace RDG_SRV with a regular SRV to be able to update UB inside RDG passes
+		static FName BatchedPrimitiveSlotName = "BatchedPrimitive";
+		FRHIUniformBufferLayoutInitializer Initialzer(TEXT("DynamicBatchedPrimitive"), 16u);
+		Initialzer.bUniformView = true;
+		Initialzer.Resources.Add({0, UBMT_RDG_BUFFER_SRV});
+		Initialzer.StaticSlot = FUniformBufferStaticSlotRegistry::Get().FindSlotByName(BatchedPrimitiveSlotName);
+		Initialzer.BindingFlags = EUniformBufferBindingFlags::StaticAndShader;
+		Initialzer.ComputeHash();
+		// set view source to a regular SRV after hash computation, to make sure hash matches BatchedPrimitive layout
+		Initialzer.Resources[0].MemberType = UBMT_SRV;
+
+		LayoutRHI = RHICreateUniformBufferLayout(Initialzer);
+	}
+
+	void ReleaseRHI() override 
+	{
+		LayoutRHI = nullptr;
+	}
+
+	FUniformBufferLayoutRHIRef LayoutRHI;
+};
+
+TGlobalResource<FDynamicBatchedPrimitiveLayout> GDynamicBatchedPrimitiveLayout;
+
 void DrawDynamicMeshPassPrivate(
 	const FSceneView& View,
 	FRHICommandList& RHICmdList,
@@ -1786,22 +1816,19 @@ void DrawDynamicMeshPassPrivate(
 {
 	if (VisibleMeshDrawCommands.Num() > 0)
 	{
-		const ERHIFeatureLevel::Type FeatureLevel = View.GetFeatureLevel();
-		const bool bUseGPUScene = UseGPUScene(GMaxRHIShaderPlatform, FeatureLevel);
-
-		if (bUseGPUScene && PlatformGPUSceneUsesUniformBufferView(View.GetShaderPlatform()))
-		{
-			// FIXME: UniformBufferView path
-			return;
-		}
-				
 		// GPUCULL_TODO: workaround for the fact that DrawDynamicMeshPassPrivate et al. don't work with GPU-Scene instancing
 		//               we don't support dynamic instancing for this path since we require one primitive per draw command
 		//               This is because the stride on the instance data buffer is set to 0 so only the first will ever be fetched.
 		const bool bDynamicInstancing = false;
-
+		const ERHIFeatureLevel::Type FeatureLevel = View.GetFeatureLevel();
+		const bool bUseGPUScene = UseGPUScene(GMaxRHIShaderPlatform, FeatureLevel);
+		
 		FMeshDrawCommandSceneArgs SceneArgs;
-		SceneArgs.BatchedPrimitiveSlot = FInstanceCullingContext::GetUniformBufferViewStaticSlot(View.GetShaderPlatform());
+		if (bUseGPUScene)
+		{
+			SceneArgs.BatchedPrimitiveSlot = FInstanceCullingContext::GetUniformBufferViewStaticSlot(View.GetShaderPlatform());
+		}
+	
 		const uint32 PrimitiveIdBufferStride = FInstanceCullingContext::GetInstanceIdBufferStride(View.GetShaderPlatform());
 		
 		ApplyViewOverridesToMeshDrawCommands(View, VisibleMeshDrawCommands, DynamicMeshDrawCommandStorage, GraphicsMinimalPipelineStateSet, InNeedsShaderInitialisation);
@@ -1813,9 +1840,11 @@ void DrawDynamicMeshPassPrivate(
 
 		if (IsUniformBufferStaticSlotValid(SceneArgs.BatchedPrimitiveSlot))
 		{
-			// FIXME: UniformBufferView path
-			FRHIUniformBuffer* Buffer = nullptr;
-			RHICmdList.SetStaticUniformBuffer(SceneArgs.BatchedPrimitiveSlot, Buffer);
+			FShaderResourceViewRHIRef BatchedPrimitiveSRV = RHICmdList.CreateShaderResourceView(SceneArgs.PrimitiveIdsBuffer);
+			SceneArgs.PrimitiveIdsBuffer = nullptr;
+			FRHIShaderResourceView* SRV = BatchedPrimitiveSRV.GetReference();
+			FUniformBufferRHIRef UBRef = RHICreateUniformBuffer(&SRV, GDynamicBatchedPrimitiveLayout.LayoutRHI, UniformBuffer_SingleFrame, EUniformBufferValidation::None);
+			RHICmdList.SetStaticUniformBuffer(SceneArgs.BatchedPrimitiveSlot, UBRef);
 		}
 
 		SubmitMeshDrawCommandsRange(VisibleMeshDrawCommands, GraphicsMinimalPipelineStateSet, SceneArgs, PrimitiveIdBufferStride, bDynamicInstancing, 0, VisibleMeshDrawCommands.Num(), InstanceFactor, RHICmdList);
diff --git a/Engine/Source/Runtime/Renderer/Private/MobileBasePassRendering.cpp b/Engine/Source/Runtime/Renderer/Private/MobileBasePassRendering.cpp
--- a/Engine/Source/Runtime/Renderer/Private/MobileBasePassRendering.cpp
+++ b/Engine/Source/Runtime/Renderer/Private/MobileBasePassRendering.cpp
@@ -405,10 +405,10 @@ void FMobileSceneRenderer::RenderMobileBasePass(FRHICommandList& RHICmdList, con
 	DrawRenderState.SetBlendState(TStaticBlendStateWriteMask<CW_RGBA>::GetRHI());
 	DrawRenderState.SetDepthStencilAccess(Scene->DefaultBasePassDepthStencilAccess);
 	DrawRenderState.SetDepthStencilState(TStaticDepthStencilState<true, CF_DepthNearOrEqual>::GetRHI());
-	RenderMobileEditorPrimitives(RHICmdList, View, DrawRenderState);
+	RenderMobileEditorPrimitives(RHICmdList, View, DrawRenderState, InstanceCullingDrawParams);
 }
 
-void FMobileSceneRenderer::RenderMobileEditorPrimitives(FRHICommandList& RHICmdList, const FViewInfo& View, const FMeshPassProcessorRenderState& DrawRenderState)
+void FMobileSceneRenderer::RenderMobileEditorPrimitives(FRHICommandList& RHICmdList, const FViewInfo& View, const FMeshPassProcessorRenderState& DrawRenderState, const FInstanceCullingDrawParams* InstanceCullingDrawParams)
 {
 	QUICK_SCOPE_CYCLE_COUNTER(STAT_EditorDynamicPrimitiveDrawTime);
 	SCOPED_DRAW_EVENT(RHICmdList, DynamicEd);
@@ -460,6 +460,15 @@ void FMobileSceneRenderer::RenderMobileEditorPrimitives(FRHICommandList& RHICmdL
 				PassMeshProcessor.AddMeshBatch(MeshBatch, DefaultBatchElementMask, nullptr);
 			}
 		});
+		
+		// DrawDynamicMeshPass may change BatchedPrimitive binding, so we need to restore it
+		FUniformBufferStaticSlot BatchedPrimitiveSlot = FInstanceCullingContext::GetUniformBufferViewStaticSlot(View.GetShaderPlatform());
+		if (IsUniformBufferStaticSlotValid(BatchedPrimitiveSlot))
+		{
+			FRHIUniformBuffer* BatchedPrimitiveBufferRHI = InstanceCullingDrawParams->BatchedPrimitive.GetUniformBuffer()->GetRHI();
+			check(BatchedPrimitiveBufferRHI);
+			RHICmdList.SetStaticUniformBuffer(BatchedPrimitiveSlot, BatchedPrimitiveBufferRHI);
+		}
 
 		// Draw the view's batched simple elements(lines, sprites, etc).
 		View.TopBatchedViewElements.Draw(RHICmdList, DrawRenderState, FeatureLevel, View, false);
diff --git a/Engine/Source/Runtime/Renderer/Private/MobileDecalRendering.cpp b/Engine/Source/Runtime/Renderer/Private/MobileDecalRendering.cpp
--- a/Engine/Source/Runtime/Renderer/Private/MobileDecalRendering.cpp
+++ b/Engine/Source/Runtime/Renderer/Private/MobileDecalRendering.cpp
@@ -44,7 +44,7 @@ static bool DoesPlatformSupportDecals(EShaderPlatform ShaderPlatform)
 	return true;
 }
 
-void FMobileSceneRenderer::RenderDecals(FRHICommandList& RHICmdList, const FViewInfo& View)
+void FMobileSceneRenderer::RenderDecals(FRHICommandList& RHICmdList, const FViewInfo& View, const FInstanceCullingDrawParams* InstanceCullingDrawParams)
 {
 	if (!DoesPlatformSupportDecals(View.GetShaderPlatform()) || !ViewFamily.EngineShowFlags.Decals || View.bIsPlanarReflection)
 	{
@@ -70,6 +70,15 @@ void FMobileSceneRenderer::RenderDecals(FRHICommandList& RHICmdList, const FView
 	{
 		SCOPED_DRAW_EVENT(RHICmdList, MeshDecals);
 		RenderMeshDecalsMobile(RHICmdList, View, DecalRenderStage, RenderTargetMode);
+
+		// MeshDecals use DrawDynamicMeshPass which may change BatchedPrimitive binding, so we need to restore it
+		FUniformBufferStaticSlot BatchedPrimitiveSlot = FInstanceCullingContext::GetUniformBufferViewStaticSlot(View.GetShaderPlatform());
+		if (IsUniformBufferStaticSlotValid(BatchedPrimitiveSlot))
+		{
+			FRHIUniformBuffer* BatchedPrimitiveBufferRHI = InstanceCullingDrawParams->BatchedPrimitive.GetUniformBuffer()->GetRHI();
+			check(BatchedPrimitiveBufferRHI);
+			RHICmdList.SetStaticUniformBuffer(BatchedPrimitiveSlot, BatchedPrimitiveBufferRHI);
+		}
 	}
 }
 
diff --git a/Engine/Source/Runtime/Renderer/Private/MobileShadingRenderer.cpp b/Engine/Source/Runtime/Renderer/Private/MobileShadingRenderer.cpp
--- a/Engine/Source/Runtime/Renderer/Private/MobileShadingRenderer.cpp
+++ b/Engine/Source/Runtime/Renderer/Private/MobileShadingRenderer.cpp
@@ -1385,7 +1385,7 @@ void FMobileSceneRenderer::RenderForwardSinglePass(FRDGBuilder& GraphBuilder, FM
 		// scene depth is read only and can be fetched
 		RHICmdList.NextSubpass();
 		RHICmdList.SetCurrentStat(GET_STATID(STAT_CLMM_Translucency));
-		RenderDecals(RHICmdList, View);
+		RenderDecals(RHICmdList, View, &PassParameters->InstanceCullingDrawParams);
 		RenderModulatedShadowProjections(RHICmdList, ViewContext.ViewIndex, View);
 		if (GMaxRHIShaderPlatform != SP_METAL_SIM)
 		{
@@ -1487,7 +1487,7 @@ void FMobileSceneRenderer::RenderForwardMultiPass(FRDGBuilder& GraphBuilder, FMo
 			
 		// scene depth is read only and can be fetched
 		RHICmdList.SetCurrentStat(GET_STATID(STAT_CLMM_Translucency));
-		RenderDecals(RHICmdList, View);
+		RenderDecals(RHICmdList, View, &SecondPassParameters->InstanceCullingDrawParams);
 		RenderModulatedShadowProjections(RHICmdList, ViewContext.ViewIndex, View);
 		RenderFog(RHICmdList, View);
 		// Draw translucency.
@@ -1713,7 +1713,7 @@ void FMobileSceneRenderer::RenderDeferredSinglePass(FRDGBuilder& GraphBuilder, c
 		// SceneColor + GBuffer write, SceneDepth is read only
 		RHICmdList.NextSubpass();
 		RHICmdList.SetCurrentStat(GET_STATID(STAT_CLMM_Translucency));
-		RenderDecals(RHICmdList, View);
+		RenderDecals(RHICmdList, View, &PassParameters->InstanceCullingDrawParams);
 		// SceneColor write, SceneDepth is read only
 		RHICmdList.NextSubpass();
 		MobileDeferredShadingPass(RHICmdList, ViewContext.ViewIndex, Views.Num(), View, *Scene, SortedLightSet, VisibleLightInfos);
@@ -1778,7 +1778,7 @@ void FMobileSceneRenderer::RenderDeferredMultiPass(FRDGBuilder& GraphBuilder, cl
 		[this, SecondPassParameters, ViewContext](FRHICommandList& RHICmdList)
 	{
 		FViewInfo& View = *ViewContext.ViewInfo;
-		RenderDecals(RHICmdList, View);
+		RenderDecals(RHICmdList, View, &SecondPassParameters->InstanceCullingDrawParams);
 	});
 
 	auto* ThirdPassParameters = GraphBuilder.AllocParameters<FMobileRenderPassParameters>();
diff --git a/Engine/Source/Runtime/Renderer/Private/Renderer.cpp b/Engine/Source/Runtime/Renderer/Private/Renderer.cpp
--- a/Engine/Source/Runtime/Renderer/Private/Renderer.cpp
+++ b/Engine/Source/Runtime/Renderer/Private/Renderer.cpp
@@ -338,10 +338,6 @@ void FRendererModule::DrawTileMesh(FCanvasRenderContext& RenderContext, FMeshPas
 		PassParameters->View = View.GetShaderParameters();
 		PassParameters->InstanceCullingDrawParams.Scene = SceneUniforms.GetBuffer(GraphBuilder);
 		PassParameters->InstanceCullingDrawParams.InstanceCulling = FInstanceCullingContext::CreateDummyInstanceCullingUniformBuffer(GraphBuilder);
-		if (UseGPUScene(View.GetShaderPlatform(), FeatureLevel) && PlatformGPUSceneUsesUniformBufferView(View.GetShaderPlatform()))
-		{
-			PassParameters->InstanceCullingDrawParams.BatchedPrimitive = CreateSinglePrimitiveUniformView(GraphBuilder, View, Mesh);
-		}
 		PassParameters->ReflectionCapture = EmptyReflectionCaptureUniformBuffer;
 
 		// handle translucent material blend modes, not relevant in MaterialTexCoordScalesAnalysis since it outputs the scales.
diff --git a/Engine/Source/Runtime/Renderer/Private/SceneRendering.h b/Engine/Source/Runtime/Renderer/Private/SceneRendering.h
--- a/Engine/Source/Runtime/Renderer/Private/SceneRendering.h
+++ b/Engine/Source/Runtime/Renderer/Private/SceneRendering.h
@@ -2661,7 +2661,7 @@ protected:
 
 	void PostRenderBasePass(FRHICommandList& RHICmdList, FViewInfo& View);
 
-	void RenderMobileEditorPrimitives(FRHICommandList& RHICmdList, const FViewInfo& View, const FMeshPassProcessorRenderState& DrawRenderState);
+	void RenderMobileEditorPrimitives(FRHICommandList& RHICmdList, const FViewInfo& View, const FMeshPassProcessorRenderState& DrawRenderState, const FInstanceCullingDrawParams* InstanceCullingDrawParams);
 
 	/** Renders the debug view pass for mobile. */
 	void RenderMobileDebugView(FRHICommandList& RHICmdList, const FViewInfo& View);
@@ -2684,7 +2684,7 @@ protected:
 
 
 	/** Renders decals. */
-	void RenderDecals(FRHICommandList& RHICmdList, const FViewInfo& View);
+	void RenderDecals(FRHICommandList& RHICmdList, const FViewInfo& View, const FInstanceCullingDrawParams* InstanceCullingDrawParams);
 	void RenderDBuffer(FRDGBuilder& GraphBuilder, FSceneTextures& SceneTextures, FDBufferTextures& DBufferTextures);
 
 	/** Renders the atmospheric and height fog */
