VolumetricCloud - fixed exposure issue after large camera movement but wen it is not a camera cut. Cloud keeptrack internally of exposure. And still does reset on camera cut also.

diff --git a/Engine/Source/Runtime/Renderer/Private/VolumetricRenderTarget.cpp b/Engine/Source/Runtime/Renderer/Private/VolumetricRenderTarget.cpp
--- a/Engine/Source/Runtime/Renderer/Private/VolumetricRenderTarget.cpp
+++ b/Engine/Source/Runtime/Renderer/Private/VolumetricRenderTarget.cpp
@@ -161,6 +161,7 @@ FVolumetricRenderTargetViewStateData::FVolumetricRenderTargetViewStateData()
 	: CurrentRT(1)
 	, bFirstTimeUsed(true)
 	, bHistoryValid(false)
+	, PreViewExposure(1.0f)
 	, FullResolution(FIntPoint::ZeroValue)
 	, VolumetricReconstructRTResolution(FIntPoint::ZeroValue)
 	, VolumetricTracingRTResolution(FIntPoint::ZeroValue)
@@ -187,6 +188,7 @@ void FVolumetricRenderTargetViewStateData::Initialise(
 	{
 		bFirstTimeUsed = false;
 		bHistoryValid = false;
+		PreViewExposure = 1.0f;
 		FrameId = 0;
 		NoiseFrameIndex = 0;
 		NoiseFrameIndexModPattern = 0;
@@ -266,6 +268,7 @@ void FVolumetricRenderTargetViewStateData::Reset()
 {
 	bFirstTimeUsed = false;
 	bHistoryValid = false;
+	PreViewExposure = 1.0f;
 	FrameId = 0;
 	NoiseFrameIndex = 0;
 	NoiseFrameIndexModPattern = 0;
@@ -614,7 +617,7 @@ void ReconstructVolumetricRenderTarget(
 		PassParameters->VolumetricRenderTargetMode = VolumetricCloudRT.GetMode();
 		PassParameters->HalfResDepthTexture = (VolumetricCloudRT.GetMode() == 0 || VolumetricCloudRT.GetMode() == 3) ? HalfResolutionDepthCheckerboardMinMaxTexture : SceneDepthTexture;
 		PassParameters->MinimumDistanceKmToEnableReprojection = FMath::Max(0.0f, CVarVolumetricRenderTargetMinimumDistanceKmToEnableReprojection.GetValueOnRenderThread());
-		PassParameters->HistoryPreExposureCorrection = ViewInfo.PreExposure / ViewInfo.PrevViewInfo.SceneColorPreExposure;
+		PassParameters->HistoryPreExposureCorrection = ViewInfo.PreExposure / VolumetricCloudRT.GetPrevViewExposure();
 
 		const bool bVisualizeConservativeDensity = ShouldViewVisualizeVolumetricCloudConservativeDensity(ViewInfo, ViewInfo.Family->EngineShowFlags);
 		PassParameters->HalfResDepthTexture = bVisualizeConservativeDensity ?
@@ -635,6 +638,8 @@ void ReconstructVolumetricRenderTarget(
 		FPixelShaderUtils::AddFullscreenPass<FReconstructVolumetricRenderTargetPS>(
 			GraphBuilder, ViewInfo.ShaderMap, RDG_EVENT_NAME("VolumetricReconstruct"), PixelShader, PassParameters,
 			FIntRect(0, 0, DstVolumetricSize.X, DstVolumetricSize.Y));
+
+		VolumetricCloudRT.PostRenderUpdate(ViewInfo.PreExposure);
 	}
 
 }
diff --git a/Engine/Source/Runtime/Renderer/Private/VolumetricRenderTargetViewStateData.h b/Engine/Source/Runtime/Renderer/Private/VolumetricRenderTargetViewStateData.h
--- a/Engine/Source/Runtime/Renderer/Private/VolumetricRenderTargetViewStateData.h
+++ b/Engine/Source/Runtime/Renderer/Private/VolumetricRenderTargetViewStateData.h
@@ -29,6 +29,16 @@ public:
 
 	void Reset();
 
+	void PostRenderUpdate(float ViewExposure)
+	{
+		PreViewExposure = ViewExposure;
+	}
+
+	float GetPrevViewExposure()
+	{
+		return PreViewExposure;
+	}
+
 	FRDGTextureRef GetOrCreateVolumetricTracingRT(FRDGBuilder& GraphBuilder);
 	FRDGTextureRef GetOrCreateVolumetricSecondaryTracingRT(FRDGBuilder& GraphBuilder);
 	FRDGTextureRef GetOrCreateVolumetricTracingRTDepth(FRDGBuilder& GraphBuilder);
@@ -68,6 +78,7 @@ private:
 	uint32 CurrentRT;
 	bool bFirstTimeUsed;
 	bool bHistoryValid;
+	float PreViewExposure;
 
 	int32 FrameId;
 	uint32 NoiseFrameIndex;	// This is only incremented once all Volumetric render target samples have been iterated
