Fixed bug where non atomic compatible dummy textures were assigned to atomic arrays
Fix MSAA bug where texture setting was overridden
Added GRHISupportsArrayIndexFromAnyShader for supported devices

diff --git a/Engine/Source/Runtime/Apple/MetalRHI/Private/MetalRHI.cpp b/Engine/Source/Runtime/Apple/MetalRHI/Private/MetalRHI.cpp
--- a/Engine/Source/Runtime/Apple/MetalRHI/Private/MetalRHI.cpp
+++ b/Engine/Source/Runtime/Apple/MetalRHI/Private/MetalRHI.cpp
@@ -670,6 +670,12 @@ FMetalDynamicRHI::FMetalDynamicRHI(ERHIFeatureLevel::Type RequestedFeatureLevel)
 	GMaxShadowDepthBufferSizeY = GMaxTextureDimensions;
 #endif
 
+    if([Device.GetPtr() supportsFamily:MTLGPUFamilyApple6] ||
+       [Device.GetPtr() supportsFamily:MTLGPUFamilyMac2])
+    {
+        GRHISupportsArrayIndexFromAnyShader = true;
+    }
+            
 	GRHIMaxDispatchThreadGroupsPerDimension.X = MAX_uint16;
 	GRHIMaxDispatchThreadGroupsPerDimension.Y = MAX_uint16;
 	GRHIMaxDispatchThreadGroupsPerDimension.Z = MAX_uint16;
diff --git a/Engine/Source/Runtime/Apple/MetalRHI/Private/MetalUAV.cpp b/Engine/Source/Runtime/Apple/MetalRHI/Private/MetalUAV.cpp
--- a/Engine/Source/Runtime/Apple/MetalRHI/Private/MetalUAV.cpp
+++ b/Engine/Source/Runtime/Apple/MetalRHI/Private/MetalUAV.cpp
@@ -233,7 +233,10 @@ void FMetalShaderResourceView::UpdateView()
 		}
 		else
 		{
-            TextureType = SRVDimensionToMetalTextureType(Info.Dimension);
+            if(TextureType != mtlpp::TextureType::Texture2DMultisample)
+            {
+                TextureType = SRVDimensionToMetalTextureType(Info.Dimension);
+            }
 
             uint32_t ArrayStart = Info.ArrayRange.First;
             uint32_t ArraySize = Info.ArrayRange.Num;
diff --git a/Engine/Source/Runtime/Renderer/Private/SystemTextures.cpp b/Engine/Source/Runtime/Renderer/Private/SystemTextures.cpp
--- a/Engine/Source/Runtime/Renderer/Private/SystemTextures.cpp
+++ b/Engine/Source/Runtime/Renderer/Private/SystemTextures.cpp
@@ -926,6 +926,25 @@ void FSystemTextures::InitializeFeatureLevelDependentTextures(FRHICommandListImm
 		GTAOPreIntegrated = CreateRenderTarget(Texture, Desc.DebugName);
 	}
 
+    // Create a texture array that is a single UInt32 value set to 0, this texture is AtomicCompatible on SM6
+    {
+        ETextureCreateFlags TextureCreateFlags = ETextureCreateFlags::ShaderResource;
+        
+        if(InFeatureLevel >= ERHIFeatureLevel::SM6)
+        {
+            TextureCreateFlags |= ETextureCreateFlags::AtomicCompatible;
+        }
+        
+        const FRHITextureCreateDesc Desc =
+            FRHITextureCreateDesc::Create2DArray(TEXT("ZeroUIntArrayAtomicCompatDummy"), 1, 1, 1, PF_R32_UINT)
+            .SetFlags(TextureCreateFlags)
+            .SetClassName(SystemTexturesName);
+
+        FTextureRHIRef Texture = RHICreateTexture(Desc);
+        SetDummyTextureArrayData<uint32>(Texture, 0u);
+        ZeroUIntArrayAtomicCompatDummy = CreateRenderTarget(Texture, Desc.DebugName);
+    }
+    
 	// Initialize textures only once.
 	FeatureLevelInitializedTo = InFeatureLevel;
 }
@@ -960,8 +979,9 @@ void FSystemTextures::ReleaseRHI()
 	CubeBlackDummy.SafeRelease();
 	CubeArrayBlackDummy.SafeRelease();
 	ZeroUIntDummy.SafeRelease();
-	ZeroUIntArrayDummy.SafeRelease();
-	MidGreyDummy.SafeRelease();
+    ZeroUIntArrayDummy.SafeRelease();
+    ZeroUIntArrayAtomicCompatDummy.SafeRelease();
+    MidGreyDummy.SafeRelease();
 	StencilDummy.SafeRelease();
 	StencilDummySRV.SafeRelease();
 	BlackDepthCube.SafeRelease();
@@ -1059,6 +1079,11 @@ FRDGTextureRef FSystemTextures::GetZeroUIntArrayDummy(FRDGBuilder& GraphBuilder)
 	return GraphBuilder.RegisterExternalTexture(ZeroUIntArrayDummy, TEXT("ZeroUIntArrayDummy"), ERDGTextureFlags::SkipTracking);
 }
 
+FRDGTextureRef FSystemTextures::GetZeroUIntArrayAtomicCompatDummy(FRDGBuilder& GraphBuilder) const
+{
+    return GraphBuilder.RegisterExternalTexture(ZeroUIntArrayAtomicCompatDummy, TEXT("ZeroUIntArrayAtomicCompatDummy"), ERDGTextureFlags::SkipTracking);
+}
+
 //////////////////////////////////////////////////////////////////////////////////////////////////////////
 // Default textures
 
diff --git a/Engine/Source/Runtime/Renderer/Private/VirtualShadowMaps/VirtualShadowMapArray.cpp b/Engine/Source/Runtime/Renderer/Private/VirtualShadowMaps/VirtualShadowMapArray.cpp
--- a/Engine/Source/Runtime/Renderer/Private/VirtualShadowMaps/VirtualShadowMapArray.cpp
+++ b/Engine/Source/Runtime/Renderer/Private/VirtualShadowMaps/VirtualShadowMapArray.cpp
@@ -494,7 +494,7 @@ void FVirtualShadowMapArray::Initialize(
 	else
 	{
 		CacheManager->FreePhysicalPool(GraphBuilder);
-		UniformParameters.PhysicalPagePool = GSystemTextures.GetZeroUIntArrayDummy(GraphBuilder);
+		UniformParameters.PhysicalPagePool = GSystemTextures.GetZeroUIntArrayAtomicCompatDummy(GraphBuilder);
 	}
 
 	if (bEnabled && bUseHzbOcclusion)
diff --git a/Engine/Source/Runtime/Renderer/Private/VirtualShadowMaps/VirtualShadowMapCacheManager.cpp b/Engine/Source/Runtime/Renderer/Private/VirtualShadowMaps/VirtualShadowMapCacheManager.cpp
--- a/Engine/Source/Runtime/Renderer/Private/VirtualShadowMaps/VirtualShadowMapCacheManager.cpp
+++ b/Engine/Source/Runtime/Renderer/Private/VirtualShadowMaps/VirtualShadowMapCacheManager.cpp
@@ -1305,7 +1305,7 @@ void FVirtualShadowMapArrayCacheManager::ProcessInvalidations(
 			UniformParameters->PageFlags = RegExtCreateSrv(PrevBuffers.PageFlags, TEXT("Shadow.Virtual.PrevPageFlags"));
 			UniformParameters->PageRectBounds = RegExtCreateSrv(PrevBuffers.PageRectBounds, TEXT("Shadow.Virtual.PrevPageRectBounds"));
 			// Unused in this path... may be a better way to handle this
-			UniformParameters->PhysicalPagePool = GSystemTextures.GetZeroUIntArrayDummy(GraphBuilder);
+			UniformParameters->PhysicalPagePool = GSystemTextures.GetZeroUIntArrayAtomicCompatDummy(GraphBuilder);
 			UniformParameters->LightGridData = GraphBuilder.CreateSRV(GSystemTextures.GetDefaultStructuredBuffer(GraphBuilder, sizeof(uint32)));
 			UniformParameters->NumCulledLightsGrid = GraphBuilder.CreateSRV(GSystemTextures.GetDefaultStructuredBuffer(GraphBuilder, sizeof(uint32)));
 		}
diff --git a/Engine/Source/Runtime/Renderer/Public/SystemTextures.h b/Engine/Source/Runtime/Renderer/Public/SystemTextures.h
--- a/Engine/Source/Runtime/Renderer/Public/SystemTextures.h
+++ b/Engine/Source/Runtime/Renderer/Public/SystemTextures.h
@@ -159,6 +159,8 @@ public:
 	TRefCountPtr<IPooledRenderTarget> ZeroUIntDummy;
 	// Dummy 0 Uint texture for RHIs that need explicit overloads
 	TRefCountPtr<IPooledRenderTarget> ZeroUIntArrayDummy;
+    // Dummy 0 Uint texture for RHIs that need explicit overloads, specific version supporting atomics on Metal
+    TRefCountPtr<IPooledRenderTarget> ZeroUIntArrayAtomicCompatDummy;
 
 	// SRV for WhiteDummy Texture.
 	TRefCountPtr<FRHIShaderResourceView> WhiteDummySRV;
@@ -174,6 +176,7 @@ public:
 	FRDGTextureRef RENDERER_API GetBlackArrayDummy(FRDGBuilder& GraphBuilder) const;
 	FRDGTextureRef RENDERER_API GetZeroUIntDummy(FRDGBuilder& GraphBuilder) const;
 	FRDGTextureRef RENDERER_API GetZeroUIntArrayDummy(FRDGBuilder& GraphBuilder) const;
+    FRDGTextureRef RENDERER_API GetZeroUIntArrayAtomicCompatDummy(FRDGBuilder& GraphBuilder) const;
 	FRDGTextureRef RENDERER_API GetZeroUShort4Dummy(FRDGBuilder& GraphBuilder) const;
 	FRDGTextureRef RENDERER_API GetBlackAlphaOneDummy(FRDGBuilder& GraphBuilder) const;
 	FRDGTextureRef RENDERER_API GetMaxFP16Depth(FRDGBuilder& GraphBuilder) const;
