GlobalDistanceField - remove unused code

diff --git a/Engine/Source/Runtime/Renderer/Private/GlobalDistanceField.cpp b/Engine/Source/Runtime/Renderer/Private/GlobalDistanceField.cpp
--- a/Engine/Source/Runtime/Renderer/Private/GlobalDistanceField.cpp
+++ b/Engine/Source/Runtime/Renderer/Private/GlobalDistanceField.cpp
@@ -2671,7 +2671,7 @@ void UpdateGlobalDistanceFieldCache(
 		FRDGBufferRef PageFreeListReturnBuffer = GraphBuilder.CreateBuffer(FRDGBufferDesc::CreateStructuredDesc(sizeof(uint32), GGlobalDistanceFieldMaxPageNum), TEXT("GlobalDistanceField.PageFreeListReturn"));
 		AddClearUAVPass(GraphBuilder, GraphBuilder.CreateUAV(PageFreeListReturnAllocatorBuffer, PF_R32_UINT), 0);
 
-		// TODO: replace single dispatch per clipmap with a single one in order to remove barriers between them
+		// TODO: atlas PageComposeTileBuffer and use SkipBarrier in order to allow overlap in this pass.
 		// Allocate pages for objects and build page lists
 		for (FGlobalDistanceFieldPackedClipmap& PackedClipmap : PackedClipmaps)
 		{
@@ -3141,19 +3141,6 @@ void UpdateGlobalDistanceFieldVolume(
 			FRDGTextureRef MipTexture = GlobalDistanceFieldInfoRDG.MipTexture;
 			FRDGTextureRef (&PageTableLayerTextures)[GDF_Num] = GlobalDistanceFieldInfoRDG.PageTableLayerTextures;
 
-			// #kris_todo: REMOVE
-			FRDGTextureRef TempMipTexture = nullptr;
-			{
-				const int32 ClipmapMipResolution = GlobalDistanceField::GetClipmapMipResolution(bLumenEnabled);
-				FRDGTextureDesc TempMipDesc(FRDGTextureDesc::Create3D(
-					FIntVector(ClipmapMipResolution),
-					PF_R8,
-					FClearValueBinding::Black,
-					TexCreate_ShaderResource | TexCreate_UAV | TexCreate_3DTiling));
-
-				TempMipTexture = GraphBuilder.CreateTexture(TempMipDesc, TEXT("GlobalDistanceField.TempMip"));
-			}
-
 			if (View.ViewState && View.ViewState->GlobalDistanceFieldData->bPendingReset)
 			{
 				// Reset all allocators to default
