GPUSkinCacheTask - move wait after GDME. In heavy skeletal mesh scene GPUSkinCacheTask may stall UpdateAllPrimitiveSceneInfos.

diff --git a/Engine/Source/Runtime/Renderer/Private/RendererScene.cpp b/Engine/Source/Runtime/Renderer/Private/RendererScene.cpp
--- a/Engine/Source/Runtime/Renderer/Private/RendererScene.cpp
+++ b/Engine/Source/Runtime/Renderer/Private/RendererScene.cpp
@@ -5443,8 +5443,6 @@ void FScene::UpdateAllPrimitiveSceneInfos(FRDGBuilder& GraphBuilder, EUpdateAllP
 	check(IsInRenderingThread());
 	check(!UE::RenderCommandPipe::IsReplaying());
 
-	UE::Tasks::FTask GPUSkinCacheTask;
-
 	if (GPUSkinCache && GPUSkinCache->HasWork())
 	{
 		GPUSkinCacheTask = GraphBuilder.AddCommandListSetupTask([this] (FRHICommandList& RHICmdList)
@@ -6413,7 +6411,6 @@ void FScene::UpdateAllPrimitiveSceneInfos(FRDGBuilder& GraphBuilder, EUpdateAllP
 	UpdateCachedShadowState(SceneUpdateChangeSetStorage.GetPreUpdateSet(), SceneUpdateChangeSetStorage.GetPostUpdateSet());
 	ShadowScene->PostSceneUpdate(SceneUpdateChangeSetStorage.GetPreUpdateSet(), SceneUpdateChangeSetStorage.GetPostUpdateSet());
 
-	GPUSkinCacheTask.Wait();
 	UpdateUniformExpressionsTask.Wait();
 
 	if (SceneInfosWithStaticDrawListUpdate.Num() > 0)
diff --git a/Engine/Source/Runtime/Renderer/Private/ScenePrivate.h b/Engine/Source/Runtime/Renderer/Private/ScenePrivate.h
--- a/Engine/Source/Runtime/Renderer/Private/ScenePrivate.h
+++ b/Engine/Source/Runtime/Renderer/Private/ScenePrivate.h
@@ -3601,6 +3601,16 @@ public:
 		return CreateLightPrimitiveInteractionsTask;
 	}
 
+	void WaitForGPUSkinCacheTask()
+	{
+		GPUSkinCacheTask.Wait();
+	}
+
+	UE::Tasks::FTask GetGPUSkinCacheTask() const
+	{
+		return GPUSkinCacheTask;
+	}
+
 	void WaitForCacheMeshDrawCommandsTask()
 	{
 		CacheMeshDrawCommandsTask.Wait();
@@ -3835,6 +3845,7 @@ private:
 	TSet<FPrimitiveSceneInfo*> DeletedPrimitiveSceneInfos;
 
 	UE::Tasks::FTask CreateLightPrimitiveInteractionsTask;
+	UE::Tasks::FTask GPUSkinCacheTask;
 	UE::Tasks::FTask CacheMeshDrawCommandsTask;
 	UE::Tasks::FTask CacheNaniteMaterialBinsTask;
 #if RHI_RAYTRACING
diff --git a/Engine/Source/Runtime/Renderer/Private/SceneVisibility.cpp b/Engine/Source/Runtime/Renderer/Private/SceneVisibility.cpp
--- a/Engine/Source/Runtime/Renderer/Private/SceneVisibility.cpp
+++ b/Engine/Source/Runtime/Renderer/Private/SceneVisibility.cpp
@@ -4053,6 +4053,7 @@ void FVisibilityTaskData::ProcessRenderThreadTasks(FExclusiveDepthStencil::Type
 		SetupMeshPasses(BasePassDepthStencilAccess, InstanceCullingManager);
 	}
 
+	Scene.WaitForGPUSkinCacheTask();
 	Tasks.LightVisibility.Wait();
 
 	INC_DWORD_STAT_BY(STAT_ProcessedPrimitives, Scene.Primitives.Num() * Views.Num());
