Moved distance field async update to actually run async.

diff --git a/Engine/Source/Runtime/Renderer/Private/DistanceFieldStreaming.cpp b/Engine/Source/Runtime/Renderer/Private/DistanceFieldStreaming.cpp
--- a/Engine/Source/Runtime/Renderer/Private/DistanceFieldStreaming.cpp
+++ b/Engine/Source/Runtime/Renderer/Private/DistanceFieldStreaming.cpp
@@ -329,22 +329,27 @@ public:
 	FReadBuffer& BrickUploadCoordinatesBuffer;
 	FReadBuffer& BrickUploadDataBuffer;
 
-	FIntVector4* BrickUploadCoordinatesPtr;
-	uint8* BrickUploadDataPtr;
+	FIntVector4* BrickUploadCoordinatesPtr = nullptr;
+	uint8* BrickUploadDataPtr = nullptr;
+
+	uint32 NumBrickUploadCoordinatesBytes = 0;
+	uint32 NumBrickUploadDataBytes = 0;
 
 	FDistanceFieldAtlasUpload(
-		FReadBuffer& InBrickUploadCoordinatesBuffer, 
-		FReadBuffer& InBrickUploadDataBuffer) : 
+		FReadBuffer& InBrickUploadCoordinatesBuffer,
+		FReadBuffer& InBrickUploadDataBuffer) :
 		BrickUploadCoordinatesBuffer(InBrickUploadCoordinatesBuffer),
 		BrickUploadDataBuffer(InBrickUploadDataBuffer)
 	{}
 
-	void AllocateAndLock(FRHICommandListBase& RHICmdList, uint32 NumBrickUploads, uint32 BrickSize)
+	void Allocate(FRHICommandListBase& RHICmdList, uint32 NumBrickUploads, uint32 BrickSize)
 	{
 		const uint32 NumCoordElements = FMath::RoundUpToPowerOfTwo(NumBrickUploads);
 		const uint32 CoordNumBytesPerElement = GPixelFormats[PF_R32G32B32A32_UINT].BlockBytes;
 
-		if (BrickUploadCoordinatesBuffer.NumBytes < NumCoordElements * CoordNumBytesPerElement)
+		NumBrickUploadCoordinatesBytes = NumCoordElements * CoordNumBytesPerElement;
+
+		if (BrickUploadCoordinatesBuffer.NumBytes < NumBrickUploadCoordinatesBytes)
 		{
 			BrickUploadCoordinatesBuffer.Initialize(RHICmdList, TEXT("DistanceFields.BrickUploadCoordinatesBuffer"), CoordNumBytesPerElement, NumCoordElements, PF_R32G32B32A32_UINT, BUF_Volatile);
 		}
@@ -352,20 +357,38 @@ public:
 		const uint32 NumBrickDataElements = FMath::RoundUpToPowerOfTwo(NumBrickUploads) * BrickSize * BrickSize * BrickSize;
 		const uint32 BrickDataNumBytesPerElement = GPixelFormats[DistanceField::DistanceFieldFormat].BlockBytes;
 
-		if (BrickUploadDataBuffer.NumBytes < NumBrickDataElements * BrickDataNumBytesPerElement 
-			|| BrickUploadDataBuffer.NumBytes > 2 * NumBrickDataElements * BrickDataNumBytesPerElement)
+		NumBrickUploadDataBytes = NumBrickDataElements * BrickDataNumBytesPerElement;
+
+		if (BrickUploadDataBuffer.NumBytes < NumBrickUploadDataBytes || BrickUploadDataBuffer.NumBytes > 2 * NumBrickUploadDataBytes)
 		{
 			BrickUploadDataBuffer.Initialize(RHICmdList, TEXT("DistanceFields.BrickUploadDataBuffer"), BrickDataNumBytesPerElement, NumBrickDataElements, DistanceField::DistanceFieldFormat, BUF_Volatile);
 		}
+	}
+
+	void Lock(FRHICommandListBase& RHICmdList)
+	{
+		if (NumBrickUploadCoordinatesBytes != 0)
+		{
+			BrickUploadCoordinatesPtr = (FIntVector4*)RHICmdList.LockBuffer(BrickUploadCoordinatesBuffer.Buffer, 0, NumBrickUploadCoordinatesBytes, RLM_WriteOnly);
+		}
 
-		BrickUploadCoordinatesPtr = (FIntVector4*)RHICmdList.LockBuffer(BrickUploadCoordinatesBuffer.Buffer, 0, NumCoordElements * CoordNumBytesPerElement, RLM_WriteOnly);
-		BrickUploadDataPtr = (uint8*)RHICmdList.LockBuffer(BrickUploadDataBuffer.Buffer, 0, NumBrickDataElements * BrickDataNumBytesPerElement, RLM_WriteOnly);
+		if (NumBrickUploadDataBytes != 0)
+		{
+			BrickUploadDataPtr = (uint8*)RHICmdList.LockBuffer(BrickUploadDataBuffer.Buffer, 0, NumBrickUploadDataBytes, RLM_WriteOnly);
+		}
 	}
 
-	void Unlock(FRHICommandListBase& RHICmdList) const
+	void Unlock(FRHICommandListBase& RHICmdList)
 	{
-		RHICmdList.UnlockBuffer(BrickUploadCoordinatesBuffer.Buffer);
-		RHICmdList.UnlockBuffer(BrickUploadDataBuffer.Buffer);
+		if (BrickUploadCoordinatesPtr)
+		{
+			RHICmdList.UnlockBuffer(BrickUploadCoordinatesBuffer.Buffer);
+		}
+
+		if (BrickUploadDataPtr)
+		{
+			RHICmdList.UnlockBuffer(BrickUploadDataBuffer.Buffer);
+		}
 	}
 };
 
@@ -374,8 +397,12 @@ class FDistanceFieldIndirectionAtlasUpload
 public:
 	FReadBuffer& IndirectionUploadIndicesBuffer;
 	FReadBuffer& IndirectionUploadDataBuffer;
-	uint32* IndirectionUploadIndicesPtr;
-	FVector4f* IndirectionUploadDataPtr;
+
+	uint32* IndirectionUploadIndicesPtr = nullptr;
+	FVector4f* IndirectionUploadDataPtr = nullptr;
+
+	uint32 NumIndirectionUploadIndicesBytes = 0;
+	uint32 NumIndirectionUploadDataBytes = 0;
 
 	FDistanceFieldIndirectionAtlasUpload(
 		FReadBuffer& InIndirectionUploadIndicesBuffer,
@@ -384,28 +411,49 @@ public:
 		IndirectionUploadDataBuffer(InIndirectionUploadDataBuffer)
 	{}
 
-	void AllocateAndLock(FRHICommandListBase& RHICmdList, uint32 NumIndirectionUploads)
+	void Allocate(FRHICommandListBase& RHICmdList, uint32 NumIndirectionUploads)
 	{
 		const uint32 NumElements = FMath::RoundUpToPowerOfTwo(NumIndirectionUploads);
 
-		if (IndirectionUploadIndicesBuffer.NumBytes < NumElements * sizeof(uint32))
+		NumIndirectionUploadIndicesBytes = NumElements * sizeof(uint32);
+
+		if (IndirectionUploadIndicesBuffer.NumBytes < NumIndirectionUploadIndicesBytes)
 		{
 			IndirectionUploadIndicesBuffer.Initialize(RHICmdList, TEXT("DistanceFields.IndirectionUploadIndicesBuffer"), sizeof(uint32), NumElements, PF_R32_UINT, BUF_Volatile);
 		}
-		
-		if (IndirectionUploadDataBuffer.NumBytes < NumElements * sizeof(FVector4f))
+
+		NumIndirectionUploadDataBytes = NumElements * sizeof(FVector4f);
+
+		if (IndirectionUploadDataBuffer.NumBytes < NumIndirectionUploadDataBytes)
 		{
 			IndirectionUploadDataBuffer.Initialize(RHICmdList, TEXT("DistanceFields.IndirectionUploadDataBuffer"), sizeof(FVector4f), NumElements, PF_A32B32G32R32F, BUF_Volatile);
 		}
-		
-		IndirectionUploadIndicesPtr = (uint32*)RHICmdList.LockBuffer(IndirectionUploadIndicesBuffer.Buffer, 0, NumElements * sizeof(uint32), RLM_WriteOnly);
-		IndirectionUploadDataPtr = (FVector4f*)RHICmdList.LockBuffer(IndirectionUploadDataBuffer.Buffer, 0, NumElements * sizeof(FVector4f), RLM_WriteOnly);
+	}
+
+	void Lock(FRHICommandListBase& RHICmdList)
+	{
+		if (NumIndirectionUploadIndicesBytes != 0)
+		{
+			IndirectionUploadIndicesPtr = (uint32*)RHICmdList.LockBuffer(IndirectionUploadIndicesBuffer.Buffer, 0, NumIndirectionUploadIndicesBytes, RLM_WriteOnly);
+		}
+
+		if (NumIndirectionUploadDataBytes != 0)
+		{
+			IndirectionUploadDataPtr = (FVector4f*)RHICmdList.LockBuffer(IndirectionUploadDataBuffer.Buffer, 0, NumIndirectionUploadDataBytes, RLM_WriteOnly);
+		}
 	}
 
 	void Unlock(FRHICommandListBase& RHICmdList) const
 	{
-		RHICmdList.UnlockBuffer(IndirectionUploadIndicesBuffer.Buffer);
-		RHICmdList.UnlockBuffer(IndirectionUploadDataBuffer.Buffer);
+		if (IndirectionUploadIndicesPtr)
+		{
+			RHICmdList.UnlockBuffer(IndirectionUploadIndicesBuffer.Buffer);
+		}
+
+		if (IndirectionUploadDataPtr)
+		{
+			RHICmdList.UnlockBuffer(IndirectionUploadDataBuffer.Buffer);
+		}
 	}
 };
 
@@ -465,12 +513,43 @@ struct FDistanceFieldReadRequest
 
 struct FDistanceFieldAsyncUpdateParameters
 {
+	FDistanceFieldAsyncUpdateParameters(
+		FReadBuffer& InBrickUploadCoordinatesBuffer,
+		FReadBuffer& InBrickUploadDataBuffer,
+		FReadBuffer& InIndirectionUploadIndicesBuffer,
+		FReadBuffer& InIndirectionUploadDataBuffer)
+		: AtlasUpload(InBrickUploadCoordinatesBuffer, InBrickUploadDataBuffer)
+		, IndirectionAtlasUpload(InIndirectionUploadIndicesBuffer, InIndirectionUploadDataBuffer)
+	{}
+
+	void Lock(FRHICommandListBase& RHICmdList)
+	{
+		AtlasUpload.Lock(RHICmdList);
+		IndirectionAtlasUpload.Lock(RHICmdList);
+
+		if (IndirectionTableUploader)
+		{
+			IndirectionTableUploader->Lock(RHICmdList);
+		}
+	}
+
+	void Unlock(FRHICommandListBase& RHICmdList)
+	{
+		AtlasUpload.Unlock(RHICmdList);
+		IndirectionAtlasUpload.Unlock(RHICmdList);
+
+		if (IndirectionTableUploader)
+		{
+			IndirectionTableUploader->Unlock(RHICmdList);
+		}
+	}
+
 	FDistanceFieldSceneData* DistanceFieldSceneData = nullptr;
-	FIntVector4* BrickUploadCoordinatesPtr = nullptr;
-	uint8* BrickUploadDataPtr = nullptr;
 
-	uint32* IndirectionIndicesUploadPtr = nullptr;
-	FVector4f* IndirectionDataUploadPtr = nullptr;
+	FDistanceFieldAtlasUpload AtlasUpload;
+	FDistanceFieldIndirectionAtlasUpload IndirectionAtlasUpload;
+
+	FRDGScatterUploader* IndirectionTableUploader = nullptr;
 
 	TArray<FDistanceFieldReadRequest> NewReadRequests;
 	TArray<FDistanceFieldReadRequest> ReadRequestsToUpload;
@@ -501,24 +580,7 @@ FDistanceFieldSceneData::~FDistanceFieldSceneData()
 	delete HeightFieldObjectBuffers;
 }
 
-class FDistanceFieldStreamingUpdateTask
-{
-public:
-	explicit FDistanceFieldStreamingUpdateTask(const FDistanceFieldAsyncUpdateParameters& InParams) : Parameters(InParams) {}
-
-	FDistanceFieldAsyncUpdateParameters Parameters;
-
-	void DoTask(ENamedThreads::Type CurrentThread, const FGraphEventRef& MyCompletionGraphEvent)
-	{
-		Parameters.DistanceFieldSceneData->AsyncUpdate(Parameters);
-	}
-
-	static ESubsequentsMode::Type	GetSubsequentsMode()	{ return ESubsequentsMode::TrackSubsequents; }
-	ENamedThreads::Type				GetDesiredThread()		{ return ENamedThreads::AnyNormalThreadNormalTask; }
-	FORCEINLINE TStatId				GetStatId() const		{ return TStatId(); }
-};
-
-void FDistanceFieldSceneData::AsyncUpdate(FDistanceFieldAsyncUpdateParameters UpdateParameters)
+void FDistanceFieldSceneData::AsyncUpdate(FRHICommandListBase& RHICmdList, FDistanceFieldAsyncUpdateParameters& UpdateParameters)
 {
 	QUICK_SCOPE_CYCLE_COUNTER(STAT_FDistanceFieldSceneData_AsyncUpdate);
 	TRACE_CPUPROFILER_EVENT_SCOPE(FDistanceFieldSceneData::AsyncUpdate);
@@ -530,6 +592,8 @@ void FDistanceFieldSceneData::AsyncUpdate(FDistanceFieldAsyncUpdateParameters Up
 	int32 BrickUploadIndex = 0;
 	int32 IndirectionUploadIndex = 0;
 
+	UpdateParameters.Lock(RHICmdList);
+
 	for (FDistanceFieldReadRequest ReadRequest : UpdateParameters.ReadRequestsToUpload)
 	{
 		const FDistanceFieldAssetState& AssetState = AssetStateArray[ReadRequest.AssetSetId];
@@ -566,11 +630,11 @@ void FDistanceFieldSceneData::AsyncUpdate(FDistanceFieldAsyncUpdateParameters Up
 		FVector4f* DestIndirection2Table = nullptr;
 		if (GDistanceFieldOffsetDataStructure == 0)
 		{
-			DestIndirectionTable = (uint32*)IndirectionTableUploadBuffer.Add_GetRef(MipState.IndirectionTableOffset, NumIndirectionEntries);
+			DestIndirectionTable = (uint32*)UpdateParameters.IndirectionTableUploader->Add_GetRef(MipState.IndirectionTableOffset, NumIndirectionEntries);
 		}
 		else if (GDistanceFieldOffsetDataStructure == 1)
 		{
-			DestIndirection2Table = (FVector4f*)IndirectionTableUploadBuffer.Add_GetRef(MipState.IndirectionTableOffset, NumIndirectionEntries);
+			DestIndirection2Table = (FVector4f*)UpdateParameters.IndirectionTableUploader->Add_GetRef(MipState.IndirectionTableOffset, NumIndirectionEntries);
 		}
 
 		// Add global allocated brick offset to indirection table entries as we upload them
@@ -617,8 +681,8 @@ void FDistanceFieldSceneData::AsyncUpdate(FDistanceFieldAsyncUpdateParameters Up
 				const FIntVector IndirectionAtlasPosition = MipState.IndirectionAtlasOffset + IndirectionCoord;
 				const uint32 IndirectionIndex = IndirectionAtlasPosition.X + IndirectionTextureSize.X * (IndirectionAtlasPosition.Y + IndirectionTextureSize.Y * IndirectionAtlasPosition.Z);
 
-				UpdateParameters.IndirectionIndicesUploadPtr[IndirectionUploadIndex] = IndirectionIndex;
-				UpdateParameters.IndirectionDataUploadPtr[IndirectionUploadIndex] = BrickOffset;
+				UpdateParameters.IndirectionAtlasUpload.IndirectionUploadIndicesPtr[IndirectionUploadIndex] = IndirectionIndex;
+				UpdateParameters.IndirectionAtlasUpload.IndirectionUploadDataPtr[IndirectionUploadIndex] = BrickOffset;
 
 				++IndirectionUploadIndex;
 			}
@@ -627,13 +691,13 @@ void FDistanceFieldSceneData::AsyncUpdate(FDistanceFieldAsyncUpdateParameters Up
 		check(MipState.NumBricks == ReadRequest.NumDistanceFieldBricks);
 		const uint8* DistanceFieldBrickDataPtr = BulkDataReadPtr + NumIndirectionEntries * sizeof(uint32);
 		const SIZE_T DistanceFieldBrickDataSizeBytes = ReadRequest.NumDistanceFieldBricks * BrickSizeBytes;
-		FMemory::Memcpy(UpdateParameters.BrickUploadDataPtr + BrickUploadIndex * BrickSizeBytes, DistanceFieldBrickDataPtr, DistanceFieldBrickDataSizeBytes);
+		FMemory::Memcpy(UpdateParameters.AtlasUpload.BrickUploadDataPtr + BrickUploadIndex * BrickSizeBytes, DistanceFieldBrickDataPtr, DistanceFieldBrickDataSizeBytes);
 
 		for (int32 BrickIndex = 0; BrickIndex < MipState.NumBricks; BrickIndex++)
 		{
 			const int32 GlobalBrickIndex = BrickIndex % GDistanceFieldBlockAllocatorSizeInBricks + GlobalBlockOffsets[BrickIndex / GDistanceFieldBlockAllocatorSizeInBricks] * GDistanceFieldBlockAllocatorSizeInBricks;
 			const FIntVector BrickTextureCoordinate = GetBrickCoordinate(GlobalBrickIndex, BrickTextureDimensionsInBricks);
-			UpdateParameters.BrickUploadCoordinatesPtr[BrickUploadIndex + BrickIndex] = FIntVector4(BrickTextureCoordinate.X, BrickTextureCoordinate.Y, BrickTextureCoordinate.Z, 0);
+			UpdateParameters.AtlasUpload.BrickUploadCoordinatesPtr[BrickUploadIndex + BrickIndex] = FIntVector4(BrickTextureCoordinate.X, BrickTextureCoordinate.Y, BrickTextureCoordinate.Z, 0);
 		}
 
 #if WITH_EDITOR
@@ -694,6 +758,8 @@ void FDistanceFieldSceneData::AsyncUpdate(FDistanceFieldAsyncUpdateParameters Up
 #endif
 
 	ReadRequests.Append(UpdateParameters.NewReadRequests);
+
+	UpdateParameters.Unlock(RHICmdList);
 }
 
 bool AssetHasOutstandingRequest(FSetElementId AssetSetId, const TArray<FDistanceFieldReadRequest>& ReadRequests)
@@ -1492,41 +1558,36 @@ void FDistanceFieldSceneData::UpdateDistanceFieldAtlas(
 	}
 
 	{
-		FDistanceFieldAsyncUpdateParameters UpdateParameters;
+		FDistanceFieldAsyncUpdateParameters& UpdateParameters = *GraphBuilder.AllocObject<FDistanceFieldAsyncUpdateParameters>(BrickUploadCoordinatesBuffer, BrickUploadDataBuffer, IndirectionUploadIndicesBuffer, IndirectionUploadDataBuffer);
 		UpdateParameters.DistanceFieldSceneData = this;
 
-		check(ReadRequestsToUpload.Num() == 0 && NumIndirectionTableAdds == 0 || ReadRequestsToUpload.Num() > 0 && NumIndirectionTableAdds > 0);
+		FDistanceFieldIndirectionAtlasUpload& IndirectionAtlasUpload = UpdateParameters.IndirectionAtlasUpload;
+		FDistanceFieldAtlasUpload& AtlasUpload = UpdateParameters.AtlasUpload;
 
-		FDistanceFieldIndirectionAtlasUpload IndirectionAtlasUpload(IndirectionUploadIndicesBuffer, IndirectionUploadDataBuffer);
+		check(ReadRequestsToUpload.Num() == 0 && NumIndirectionTableAdds == 0 || ReadRequestsToUpload.Num() > 0 && NumIndirectionTableAdds > 0);
 
 		if (NumIndirectionTableAdds > 0)
 		{
 			// Allocate staging buffer space for the indirection table compute scatter
 			if (GDistanceFieldOffsetDataStructure == 0)
 			{
-				IndirectionTableUploadBuffer.Init(GraphBuilder, NumIndirectionTableAdds, sizeof(uint32), false, TEXT("DistanceFields.IndirectionTableUploadBuffer.Uint"));
+				UpdateParameters.IndirectionTableUploader = IndirectionTableUploadBuffer.Begin(GraphBuilder, IndirectionTableRDG, NumIndirectionTableAdds, sizeof(uint32), TEXT("DistanceFields.IndirectionTableUploadBuffer.Uint"));
 			}
 			else if (GDistanceFieldOffsetDataStructure == 1)
 			{
-				IndirectionTableUploadBuffer.Init(GraphBuilder, NumIndirectionTableAdds, sizeof(FVector4f), true, TEXT("DistanceFields.IndirectionFloatUploadBuffer.Float"));
+				UpdateParameters.IndirectionTableUploader = IndirectionTableUploadBuffer.Begin(GraphBuilder, IndirectionTableRDG, NumIndirectionTableAdds, sizeof(FVector4f), TEXT("DistanceFields.IndirectionFloatUploadBuffer.Float"));
 			}
 			else
 			{
 				// Allocate staging buffer space for the indirection atlas compute scatter
-				IndirectionAtlasUpload.AllocateAndLock(GraphBuilder.RHICmdList, NumIndirectionTableAdds);
-				UpdateParameters.IndirectionIndicesUploadPtr = IndirectionAtlasUpload.IndirectionUploadIndicesPtr;
-				UpdateParameters.IndirectionDataUploadPtr = IndirectionAtlasUpload.IndirectionUploadDataPtr;
+				IndirectionAtlasUpload.Allocate(GraphBuilder.RHICmdList, NumIndirectionTableAdds);
 			}
 		}
 
-		FDistanceFieldAtlasUpload AtlasUpload(BrickUploadCoordinatesBuffer, BrickUploadDataBuffer);
-
 		if (NumBrickUploads > 0)
 		{
 			// Allocate staging buffer space for the brick atlas compute scatter
-			AtlasUpload.AllocateAndLock(GraphBuilder.RHICmdList, NumBrickUploads, DistanceField::BrickSize);
-			UpdateParameters.BrickUploadDataPtr = AtlasUpload.BrickUploadDataPtr;
-			UpdateParameters.BrickUploadCoordinatesPtr = AtlasUpload.BrickUploadCoordinatesPtr;
+			AtlasUpload.Allocate(GraphBuilder.RHICmdList, NumBrickUploads, DistanceField::BrickSize);
 		}
 
 		if (NewReadRequests.Num() || ReadRequestsToUpload.Num() || ReadRequestsToCleanUp.Num())
@@ -1535,8 +1596,10 @@ void FDistanceFieldSceneData::UpdateDistanceFieldAtlas(
 			UpdateParameters.ReadRequestsToUpload = MoveTemp(ReadRequestsToUpload);
 			UpdateParameters.ReadRequestsToCleanUp = MoveTemp(ReadRequestsToCleanUp);
 
-			// TODO: We actually run this synchronously now after the RDG conversion, as it would otherwise immediately sync.
-			AsyncUpdate(UpdateParameters);
+			GraphBuilder.AddCommandListSetupTask([this, &UpdateParameters] (FRHICommandListBase& RHICmdList)
+			{
+				AsyncUpdate(RHICmdList, UpdateParameters);
+			});
 		}
 
 		if (NumBrickUploads > 0 || NumIndirectionTableAdds > 0)
@@ -1544,22 +1607,15 @@ void FDistanceFieldSceneData::UpdateDistanceFieldAtlas(
 			QUICK_SCOPE_CYCLE_COUNTER(STAT_WaitOnDistanceFieldStreamingUpdate);
 			TRACE_CPUPROFILER_EVENT_SCOPE(WaitOnDistanceFieldStreamingUpdate);
 
-			if (NumBrickUploads > 0)
-			{
-				AtlasUpload.Unlock(GraphBuilder.RHICmdList);
-			}
-
 			if (NumIndirectionTableAdds > 0)
 			{
 				if (GDistanceFieldOffsetDataStructure == 0 || GDistanceFieldOffsetDataStructure == 1)
 				{
-					IndirectionTableUploadBuffer.ResourceUploadTo(GraphBuilder, IndirectionTableRDG);
+					IndirectionTableUploadBuffer.End(GraphBuilder, UpdateParameters.IndirectionTableUploader);
 					ExternalAccessQueue.Add(IndirectionTableRDG, ERHIAccess::SRVMask, ERHIPipeline::All);
 				}
 				else
 				{
-					IndirectionAtlasUpload.Unlock(GraphBuilder.RHICmdList);
-
 					TShaderMapRef<FScatterUploadDistanceFieldIndirectionAtlasCS> ComputeShader(GlobalShaderMap);
 
 					auto* PassParameters = GraphBuilder.AllocParameters<FScatterUploadDistanceFieldIndirectionAtlasCS::FParameters>();
diff --git a/Engine/Source/Runtime/Renderer/Private/ScenePrivate.h b/Engine/Source/Runtime/Renderer/Private/ScenePrivate.h
--- a/Engine/Source/Runtime/Renderer/Private/ScenePrivate.h
+++ b/Engine/Source/Runtime/Renderer/Private/ScenePrivate.h
@@ -2131,7 +2131,7 @@ public:
 
 	FGrowOnlySpanAllocator IndirectionTableAllocator;
 	TRefCountPtr<FRDGPooledBuffer> IndirectionTable;
-	FRDGScatterUploadBuffer IndirectionTableUploadBuffer;
+	FRDGAsyncScatterUploadBuffer IndirectionTableUploadBuffer;
 
 	TRefCountPtr<IPooledRenderTarget> IndirectionAtlas;
 	FTextureLayout3d IndirectionAtlasLayout;
@@ -2186,7 +2186,7 @@ private:
 	
 	void UploadAllAssetData(FRDGBuilder& GraphBuilder, FRDGBuffer* AssetDataBufferRDG);
 
-	void AsyncUpdate(FDistanceFieldAsyncUpdateParameters UpdateParameters);
+	void AsyncUpdate(FRHICommandListBase& RHICmdList, FDistanceFieldAsyncUpdateParameters& UpdateParameters);
 
 	void GenerateStreamingRequests(
 		FRDGBuilder& GraphBuilder, 
