Fix missing Substrate special complex tile dispatch.

This issue was introduced in CL 28959729.

diff --git a/Engine/Source/Runtime/Renderer/Private/Substrate/Substrate.cpp b/Engine/Source/Runtime/Renderer/Private/Substrate/Substrate.cpp
--- a/Engine/Source/Runtime/Renderer/Private/Substrate/Substrate.cpp
+++ b/Engine/Source/Runtime/Renderer/Private/Substrate/Substrate.cpp
@@ -73,7 +73,12 @@ void FSubstrateViewData::Reset()
 	// * MaxClosurePerPixel
 	// * SubstrateClosureCountMask
 	static_assert(SUBSTRATE_MAX_CLOSURE_COUNT <= 8u);
+
+	// Propagate bUsesComplexSpecialRenderPath after reset as we use the per-view (vs. the per-scene) 
+	// value to know if a view needs special complex path or not
+	const bool OldUsesComplexSpecialRenderPath = bUsesComplexSpecialRenderPath;
 	*this = FSubstrateViewData();
+	bUsesComplexSpecialRenderPath = OldUsesComplexSpecialRenderPath;
 }
 
 const TCHAR* ToString(ESubstrateTileType Type)
@@ -148,6 +153,8 @@ bool GetSubstrateUsesComplexSpecialPath(const FViewInfo& View)
 {
 	if (Substrate::IsSubstrateEnabled())
 	{
+		// Use the per-view value rather than the per-scene data to have more accurate dispatching of special complex tiles 
+		// and avoid unecessary empty-dispatch
 		return View.SubstrateViewData.bUsesComplexSpecialRenderPath;
 	}
 	return false;
