Stochastic Shadows - Support all shading models.
* Tile Classification pass generates tiles for simple (at the moment Default Lit) and complex shading (all shading models)
* Sample generation and shading passes are dispatched twice in order to have lower occupancy on simple shading tiles

diff --git a/Engine/Shaders/Private/StochasticShadows/StochasticShadows.usf b/Engine/Shaders/Private/StochasticShadows/StochasticShadows.usf
--- a/Engine/Shaders/Private/StochasticShadows/StochasticShadows.usf
+++ b/Engine/Shaders/Private/StochasticShadows/StochasticShadows.usf
@@ -5,6 +5,96 @@
 #include "../ShaderPrint.ush"
 #include "StochasticShadows.ush"
 
+groupshared uint SharedTileMode;
+
+RWStructuredBuffer<uint> RWTileAllocator;
+RWStructuredBuffer<uint> RWTileData;
+uint TileDataStride;
+uint DownsampledTileDataStride;
+uint2 DownsampledViewSize;
+
+/**
+ * Run tile classification to generate tiles for each subsequent pass
+ */
+[numthreads(THREADGROUP_SIZE, THREADGROUP_SIZE, 1)]
+void TileClassificationCS(
+	uint3 GroupId : SV_GroupID,
+	uint3 GroupThreadId : SV_GroupThreadID,
+	uint3 DispatchThreadId : SV_DispatchThreadID)
+{
+	if (all(GroupThreadId == 0))
+	{
+		SharedTileMode = 0;
+	}
+	
+	GroupMemoryBarrierWithGroupSync();
+
+	uint2 ScreenCoord = DispatchThreadId.xy;
+
+#if DOWNSAMPLED_CLASSIFICATION
+	if (all(ScreenCoord < DownsampledViewSize))
+#else
+	if (all(ScreenCoord < View.ViewRectMinAndSize.zw))
+#endif
+	{
+#if DOWNSAMPLED_CLASSIFICATION
+		float2 ScreenUV = DownsampledScreenCoordToScreenUV(ScreenCoord);
+#else
+		float2 ScreenUV = (ScreenCoord + 0.5f) * View.BufferSizeAndInvSize.zw;
+#endif
+		FGBufferData GBuffer = GetGBufferData(ScreenUV);
+
+		if (GBuffer.ShadingModelID == SHADINGMODELID_UNLIT)
+		{
+			// Skip mark
+		}
+		else if (GBuffer.ShadingModelID == SHADINGMODELID_DEFAULT_LIT)
+		{
+			InterlockedOr(SharedTileMode, 0x1);
+		}
+		else
+		{
+			InterlockedOr(SharedTileMode, 0x2);
+		}
+	}
+
+	GroupMemoryBarrierWithGroupSync();
+
+	if (all(GroupThreadId == 0))
+	{
+		const uint DataStride = DOWNSAMPLED_CLASSIFICATION ? DownsampledTileDataStride : TileDataStride;
+
+		if (SharedTileMode & 0x2)
+		{
+			uint StochasticShadowsTileIndex = 0;
+			InterlockedAdd(RWTileAllocator[TILE_MODE_COMPLEX_SHADING], 1, StochasticShadowsTileIndex);
+			RWTileData[DataStride * TILE_MODE_COMPLEX_SHADING + StochasticShadowsTileIndex] = PackTile(ScreenCoord / 8);
+		}
+		else if (SharedTileMode & 0x1)
+		{
+			uint StochasticShadowsTileIndex = 0;
+			InterlockedAdd(RWTileAllocator[TILE_MODE_SIMPLE_SHADING], 1, StochasticShadowsTileIndex);
+			RWTileData[DataStride * TILE_MODE_SIMPLE_SHADING + StochasticShadowsTileIndex] = PackTile(ScreenCoord / 8);
+		}
+	}
+}
+
+RWBuffer<uint> RWTileIndirectArgs;
+RWBuffer<uint> RWDownsampledTileIndirectArgs;
+StructuredBuffer<uint> TileAllocator;
+StructuredBuffer<uint> DownsampledTileAllocator;
+
+[numthreads(THREADGROUP_SIZE, 1, 1)]
+void InitTileIndirectArgsCS(uint3 DispatchThreadId : SV_DispatchThreadID)
+{
+	uint TileMode = DispatchThreadId.x;
+	if (TileMode < TILE_MODE_MAX)
+	{
+		WriteDispatchIndirectArgs(RWTileIndirectArgs, TileMode, TileAllocator[TileMode], 1, 1);
+		WriteDispatchIndirectArgs(RWDownsampledTileIndirectArgs, TileMode, DownsampledTileAllocator[TileMode], 1, 1);
+	}
+}
+
 RWBuffer<uint> RWIndirectArgs;
 StructuredBuffer<uint> ShadowMaskTileAllocator;
 uint MaxShadowMaskTiles;
@@ -91,114 +181,4 @@ void InitShadowMaskUpsampleWeightsCS(
 
 		RWShadowMaskUpsampleWeights[ScreenCoord] = InterpolationWeights;
 	}
-}
-
-RWTexture2D<uint> RWShadowMaskTileAtlas;
-StructuredBuffer<uint2> ShadowMaskTileHeader;
-Texture2D<uint2> ShadowMaskHistoryScreenCoord00;
-Texture2D<float4> ShadowMaskHistoryWeights;
-Texture2D<uint> LightSamples;
-
-/**
- * Run one thread per shadow mask pixel and composite current frame's traces.
- * Also output a compacted tile list for shading.
- */
-[numthreads(THREADGROUP_SIZE, THREADGROUP_SIZE, 1)]
-void CompositeShadowMaskTracesCS(
-	uint3 GroupId : SV_GroupID,
-	uint3 GroupThreadId : SV_GroupThreadID,
-	uint3 DispatchThreadId : SV_DispatchThreadID)
-{
-	uint ShadowMaskTileIndex = GroupId.x;
-	uint2 LocalShadowMaskTileCoord = GroupThreadId.xy;
-
-	uint LightSceneId = 0;
-	uint2 ShadowMaskScreenCoord = 0;	
-	UnpackShadowMaskTileHeader(ShadowMaskTileHeader[ShadowMaskTileIndex], LightSceneId, ShadowMaskScreenCoord);
-	ShadowMaskScreenCoord += LocalShadowMaskTileCoord;
-
-#define DEBUG_MODE 0
-#if DEBUG_MODE
-	int2 DebugScreenCoord = View.CursorPosition.x >= 0 ? View.CursorPosition * View.ViewResolutionFraction : -1;
-	bool bDebug = all(ShadowMaskScreenCoord == DebugScreenCoord / SHADOW_MASK_DOWNSAMPLE_FACTOR);
-	FShaderPrintContext Context = InitShaderPrintContext(true, float2(0.55, 0.2));
-
-	if (bDebug)
-	{
-		Print(Context, TEXT("CompositeShadowMaskTracesCS"));
-	}
-#endif
-
-	if (all(ShadowMaskScreenCoord < ShadowMaskViewSize))
-	{
-		FShadowMask ShadowMask = InitShadowMask();
-		uint2 ShadowMaskTileAtlasCoord = GetShadowMaskTileAtlasCoord(ShadowMaskTileIndex, LocalShadowMaskTileCoord);
-
-		#if TEMPORAL_ACCUMULATION
-		{
-			int2 ScreenCoord00 = ShadowMaskHistoryScreenCoord00[ShadowMaskScreenCoord];
-			float4 Weights = ShadowMaskHistoryWeights[ShadowMaskScreenCoord];
-
-			if (any(Weights > 0.0f))
-			{
-				FShadowMaskAccumulator Accumulator = InitShadowMaskAccumulator();
-
-				AccumulateHistoryShadowMask(LightSceneId, ScreenCoord00 + uint2(0, 0), Weights.x, Accumulator);
-				AccumulateHistoryShadowMask(LightSceneId, ScreenCoord00 + uint2(1, 0), Weights.y, Accumulator);
-				AccumulateHistoryShadowMask(LightSceneId, ScreenCoord00 + uint2(0, 1), Weights.z, Accumulator);
-				AccumulateHistoryShadowMask(LightSceneId, ScreenCoord00 + uint2(1, 1), Weights.w, Accumulator);
-
-				if (Accumulator.WeightSum > 0.0f)
-				{
-					ShadowMask = Accumulator.ResolveShadowMask();
-				}
-			}
-		}
-		#endif
-
-		// Upsample and composite current frame traces
-		{
-			// Accumulate traces
-			FShadowMaskAccumulator Accumulator = InitShadowMaskAccumulator();
-			for (uint LocalSampleY = 0; LocalSampleY < NUM_SAMPLES_PER_PIXEL_Y; ++LocalSampleY)
-			{
-				for (uint LocalSampleX = 0; LocalSampleX < NUM_SAMPLES_PER_PIXEL_X; ++LocalSampleX)
-				{
-					uint2 SampleCoord = ShadowMaskScreenCoord * uint2(NUM_SAMPLES_PER_PIXEL_X, NUM_SAMPLES_PER_PIXEL_Y) + uint2(LocalSampleX, LocalSampleY);
-					FLightSample LightSample = UnpackLightSample(LightSamples[SampleCoord]);
-
-					if (LightSample.LightSceneId == LightSceneId)
-					{
-						float Factor = LightSample.bVisible ? 1.0f : 0.0f;
-
-						float Weight = 1.0f;
-						Accumulator.FactorSum += Factor * Weight;
-						Accumulator.FactorSqSum += Pow2(Factor) * Weight;
-						Accumulator.WeightSum += Weight;
-					}
-				}
-			}
-
-			if (Accumulator.WeightSum > 0.0f)
-			{
-				// Merge with reprojected history data
-				if (ShadowMask.IsValid())
-				{
-					float NewNumFramesAccumulated = max(Accumulator.MaxHistoryLength, 8.0f);
-					float Alpha = 1.0f / (1.0f + NewNumFramesAccumulated);
-					ShadowMask.Factor = lerp(ShadowMask.Factor, Accumulator.FactorSum / Accumulator.WeightSum, Alpha);
-					ShadowMask.FactorSq = lerp(ShadowMask.FactorSq, Accumulator.FactorSqSum / Accumulator.WeightSum, Alpha);
-					ShadowMask.Age = 0;
-				}
-				else
-				{
-					ShadowMask.Factor = Accumulator.FactorSum / Accumulator.WeightSum;
-					ShadowMask.FactorSq = Accumulator.FactorSqSum / Accumulator.WeightSum;
-					ShadowMask.Age = 0;
-				}
-			}
-		}
-
-		RWShadowMaskTileAtlas[ShadowMaskTileAtlasCoord] = PackShadowMask(ShadowMask);
-	}
 }
\ No newline at end of file
diff --git a/Engine/Shaders/Private/StochasticShadows/StochasticShadows.ush b/Engine/Shaders/Private/StochasticShadows/StochasticShadows.ush
--- a/Engine/Shaders/Private/StochasticShadows/StochasticShadows.ush
+++ b/Engine/Shaders/Private/StochasticShadows/StochasticShadows.ush
@@ -17,6 +17,7 @@
 #define INVALID_LIGHT_SCENE_ID 0x3FFF
 #define LIGHT_SCENE_ID_MASK_SIZE (MAX_LIGHT_SCENE_ID / 32)
 
+#define TILE_SIZE 8
 #define SHADOW_MASK_TILE_SIZE 8
 #define INVALID_SHADOW_MASK_TILE_INDEX 0xFFFFFFFF
 #define MAX_SHADOW_MASK_AGE 7
@@ -35,6 +36,21 @@
 	#define NUM_SAMPLES_PER_PIXEL_X 2
 #endif
 
+// Keep in sync with StochasticShadows::ETileType
+#define TILE_MODE_SIMPLE_SHADING	0
+#define TILE_MODE_COMPLEX_SHADING	1
+#define TILE_MODE_MAX				2
+
+uint PackTile(uint2 TileCoord)
+{
+	return TileCoord.x | (TileCoord.y << 16);
+}
+
+uint2 UnpackTile(uint PackedTile)
+{
+	return uint2(PackedTile & 0xFFFF, PackedTile >> 16);
+}
+
 float3 GetTranslatedWorldPositionFromScreenUV(float2 ScreenUV, float SceneDepth)
 {
 	float2 ScreenPosition = (ScreenUV - View.ScreenPositionScaleBias.wz) / View.ScreenPositionScaleBias.xy;
@@ -232,21 +248,21 @@ float2 DownsampledScreenCoordToScreenUV(uint2 DownsampledScreenCoord)
 	return ScreenUV;
 }
 
-uint2 PackShadowMaskTileHeader(uint LightIndex, uint2 ShadowMaskScreenCoord)
+uint2 PackShadowMaskTile(uint LightIndex, uint2 ShadowMaskScreenCoord)
 {
 	ShadowMaskScreenCoord /= SHADOW_MASK_TILE_SIZE;
 
-	uint2 PackedHeader;
-	PackedHeader.x = LightIndex;
-	PackedHeader.y = (ShadowMaskScreenCoord.x & 0xFFFF) | ((ShadowMaskScreenCoord.y & 0xFFFF) << 16);
-	return PackedHeader;
+	uint2 PackedTile;
+	PackedTile.x = LightIndex;
+	PackedTile.y = (ShadowMaskScreenCoord.x & 0xFFFF) | ((ShadowMaskScreenCoord.y & 0xFFFF) << 16);
+	return PackedTile;
 }
 
-void UnpackShadowMaskTileHeader(uint2 PackedHeader, inout uint LightIndex, inout uint2 ShadowMaskScreenCoord)
+void UnpackShadowMaskTile(uint2 PackedTile, inout uint LightIndex, inout uint2 ShadowMaskScreenCoord)
 {
-	LightIndex = PackedHeader.x;
-	ShadowMaskScreenCoord.x = PackedHeader.y & 0xFFFF;
-	ShadowMaskScreenCoord.y = PackedHeader.y >> 16;
+	LightIndex = PackedTile.x;
+	ShadowMaskScreenCoord.x = PackedTile.y & 0xFFFF;
+	ShadowMaskScreenCoord.y = PackedTile.y >> 16;
 
 	ShadowMaskScreenCoord *= SHADOW_MASK_TILE_SIZE;
 }
diff --git a/Engine/Shaders/Private/StochasticShadows/StochasticShadowsComposite.usf b/Engine/Shaders/Private/StochasticShadows/StochasticShadowsComposite.usf
new file mode 100644
--- /dev/null
+++ b/Engine/Shaders/Private/StochasticShadows/StochasticShadowsComposite.usf
@@ -0,0 +1,120 @@
+// Copyright Epic Games, Inc. All Rights Reserved.
+
+#include "../Common.ush"
+#include "../LightGridCommon.ush"
+#include "../ShaderPrint.ush"
+#include "StochasticShadows.ush"
+
+RWTexture2D<uint> RWShadowMaskTileAtlas;
+StructuredBuffer<uint2> ShadowMaskTileData;
+Texture2D<uint2> ShadowMaskHistoryScreenCoord00;
+Texture2D<float4> ShadowMaskHistoryWeights;
+Texture2D<uint> LightSamples;
+
+/**
+ * Run one thread per shadow mask pixel and composite current frame's traces.
+ * Also output a compacted tile list for shading.
+ */
+[numthreads(THREADGROUP_SIZE, THREADGROUP_SIZE, 1)]
+void CompositeShadowMaskTracesCS(
+	uint3 GroupId : SV_GroupID,
+	uint3 GroupThreadId : SV_GroupThreadID,
+	uint3 DispatchThreadId : SV_DispatchThreadID)
+{
+	uint ShadowMaskTileIndex = GroupId.x;
+	uint2 LocalShadowMaskTileCoord = GroupThreadId.xy;
+
+	uint LightSceneId = 0;
+	uint2 ShadowMaskScreenCoord = 0;	
+	UnpackShadowMaskTile(ShadowMaskTileData[ShadowMaskTileIndex], LightSceneId, ShadowMaskScreenCoord);
+	ShadowMaskScreenCoord += LocalShadowMaskTileCoord;
+
+#define DEBUG_MODE 0
+#if DEBUG_MODE
+	int2 DebugScreenCoord = View.CursorPosition.x >= 0 ? View.CursorPosition * View.ViewResolutionFraction : -1;
+	bool bDebug = all(ShadowMaskScreenCoord == DebugScreenCoord / SHADOW_MASK_DOWNSAMPLE_FACTOR);
+	FShaderPrintContext Context = InitShaderPrintContext(true, float2(0.1, 0.3));
+
+	if (bDebug)
+	{
+		Print(Context, TEXT("CompositeShadowMaskTracesCS"));
+	}
+#endif
+
+	if (all(ShadowMaskScreenCoord < ShadowMaskViewSize))
+	{
+		// Accumulate current frame traces
+		FShadowMask ShadowMask = InitShadowMask();
+		{
+			FShadowMaskAccumulator Accumulator = InitShadowMaskAccumulator();
+			for (uint LocalSampleY = 0; LocalSampleY < NUM_SAMPLES_PER_PIXEL_Y; ++LocalSampleY)
+			{
+				for (uint LocalSampleX = 0; LocalSampleX < NUM_SAMPLES_PER_PIXEL_X; ++LocalSampleX)
+				{
+					uint2 SampleCoord = ShadowMaskScreenCoord * uint2(NUM_SAMPLES_PER_PIXEL_X, NUM_SAMPLES_PER_PIXEL_Y) + uint2(LocalSampleX, LocalSampleY);
+					FLightSample LightSample = UnpackLightSample(LightSamples[SampleCoord]);
+
+					if (LightSample.LightSceneId == LightSceneId)
+					{
+						float Factor = LightSample.bVisible ? 1.0f : 0.0f;
+
+						float Weight = 1.0f;
+						Accumulator.FactorSum += Factor * Weight;
+						Accumulator.FactorSqSum += Pow2(Factor) * Weight;
+						Accumulator.WeightSum += Weight;
+					}
+				}
+			}
+
+			if (Accumulator.WeightSum > 0.0f)
+			{
+				ShadowMask.Age = 0;
+				ShadowMask.Factor = Accumulator.FactorSum / Accumulator.WeightSum;
+				ShadowMask.FactorSq = Accumulator.FactorSqSum / Accumulator.WeightSum;
+			}
+		}
+
+		// History reprojection
+		#if TEMPORAL_ACCUMULATION
+		{
+			int2 ScreenCoord00 = ShadowMaskHistoryScreenCoord00[ShadowMaskScreenCoord];
+			float4 Weights = ShadowMaskHistoryWeights[ShadowMaskScreenCoord];
+
+			if (any(Weights > 0.0f))
+			{
+				FShadowMaskAccumulator Accumulator = InitShadowMaskAccumulator();
+
+				AccumulateHistoryShadowMask(LightSceneId, ScreenCoord00 + uint2(0, 0), Weights.x, Accumulator);
+				AccumulateHistoryShadowMask(LightSceneId, ScreenCoord00 + uint2(1, 0), Weights.y, Accumulator);
+				AccumulateHistoryShadowMask(LightSceneId, ScreenCoord00 + uint2(0, 1), Weights.z, Accumulator);
+				AccumulateHistoryShadowMask(LightSceneId, ScreenCoord00 + uint2(1, 1), Weights.w, Accumulator);
+
+				if (Accumulator.WeightSum > 0.0f)
+				{
+					float HistoryShadowFactor = Accumulator.FactorSum / Accumulator.WeightSum;
+					float HistoryShadowFactorSq = Accumulator.FactorSqSum / Accumulator.WeightSum;
+
+					// Merge with reprojected history data
+					if (ShadowMask.IsValid())
+					{
+						float NewNumFramesAccumulated = max(Accumulator.MaxHistoryLength, 8.0f);
+						float Alpha = 1.0f / (1.0f + NewNumFramesAccumulated);
+						ShadowMask.Factor = lerp(HistoryShadowFactor, ShadowMask.Factor, Alpha);
+						ShadowMask.FactorSq = lerp(HistoryShadowFactorSq, ShadowMask.FactorSq, Alpha);
+						ShadowMask.Age = 0;
+					}
+					else
+					{
+						ShadowMask.Factor = HistoryShadowFactor;
+						ShadowMask.FactorSq = HistoryShadowFactorSq;
+						ShadowMask.Age = Accumulator.MinAge;
+					}
+				}
+			}
+		}
+		#endif
+
+		uint2 ShadowMaskTileAtlasCoord = GetShadowMaskTileAtlasCoord(ShadowMaskTileIndex, LocalShadowMaskTileCoord);
+		RWShadowMaskTileAtlas[ShadowMaskTileAtlasCoord] = PackShadowMask(ShadowMask);
+	}
+}
\ No newline at end of file
diff --git a/Engine/Shaders/Private/StochasticShadows/StochasticShadowsHardwareRayTracing.usf b/Engine/Shaders/Private/StochasticShadows/StochasticShadowsHardwareRayTracing.usf
--- a/Engine/Shaders/Private/StochasticShadows/StochasticShadowsHardwareRayTracing.usf
+++ b/Engine/Shaders/Private/StochasticShadows/StochasticShadowsHardwareRayTracing.usf
@@ -28,7 +28,7 @@ RAY_TRACING_ENTRY_RAYGEN_OR_INLINE(HardwareRayTraceLightSamples)
 	uint TraceTexelIndex = GroupIndex * 64 + ThreadIndex;
 
 #if DEBUG_MODE
-	FShaderPrintContext Context = InitShaderPrintContext(true, float2(0.55, 0.3));
+	FShaderPrintContext Context = InitShaderPrintContext(true, TraceTexelIndex == 0 ? float2(0.55, 0.3) : float2(0.55, 0.35));
 	if (TraceTexelIndex == 0)
 	{
 		Print(Context, TEXT("HardwareRayTraceLightSamples"));
diff --git a/Engine/Shaders/Private/StochasticShadows/StochasticShadowsSampling.usf b/Engine/Shaders/Private/StochasticShadows/StochasticShadowsSampling.usf
--- a/Engine/Shaders/Private/StochasticShadows/StochasticShadowsSampling.usf
+++ b/Engine/Shaders/Private/StochasticShadows/StochasticShadowsSampling.usf
@@ -50,8 +50,12 @@ float GetLocalLightTargetPDF(FDeferredLightData LightData, float3 TranslatedWorl
 			float SurfaceShadow = 1.0f;
 			float AmbientOcclusion = 1.0f;
 
-			// #sdl_todo: tile classification and better integration
-			GBuffer.ShadingModelID = SHADINGMODELID_DEFAULT_LIT;
+			// #sdl_todo: Move more shading models to TILE_MODE_SIMPLE_SHADING
+			#if TILE_TYPE == TILE_MODE_SIMPLE_SHADING
+			{
+				GBuffer.ShadingModelID = SHADINGMODELID_DEFAULT_LIT;
+			}
+			#endif
 			LightData.ShadowedBits = 0;
 
 			float4 LightRadiance = GetDynamicLighting(
@@ -89,7 +93,7 @@ RWTexture2D<float> RWDownsampledSceneDepth;
 RWTexture2D<uint> RWShadowMaskPageTable;
 RWStructuredBuffer<uint> RWShadowMaskHashTable;
 RWStructuredBuffer<uint> RWShadowMaskTileAllocator;
-RWStructuredBuffer<uint2> RWShadowMaskTileHeader;
+RWStructuredBuffer<uint2> RWShadowMaskTileData;
 RWTexture2D<uint> RWLightSamples;
 
 float4 HistoryScreenPositionScaleBias;
@@ -159,7 +163,7 @@ void WriteShadowMaskTileToPageTable(uint LightSceneId, uint2 ShadowMaskScreenCoo
 {
 	uint2 ShadowMaskPageTableCoord = GetShadowMaskPageTableCoord(LightSceneId, ShadowMaskScreenCoord);
 
-	RWShadowMaskTileHeader[ShadowMaskTileIndex] = PackShadowMaskTileHeader(LightSceneId, ShadowMaskScreenCoord);
+	RWShadowMaskTileData[ShadowMaskTileIndex] = PackShadowMaskTile(LightSceneId, ShadowMaskScreenCoord);
 
 	#if HASH_TABLE
 	{
@@ -204,6 +208,9 @@ void WriteShadowMaskTileToPageTable(uint LightSceneId, uint2 ShadowMaskScreenCoo
 	#endif
 }
 
+StructuredBuffer<uint> DownsampledTileAllocator;
+StructuredBuffer<uint> DownsampledTileData;
+uint DownsampledTileDataStride;
 float2 DownsampledBufferInvSize;
 
 /**
@@ -238,435 +245,441 @@ void GenerateSamplesCS(
 
 	GroupMemoryBarrierWithGroupSync();
 
-	uint2 DownsampledScreenCoord = DispatchThreadId.xy;
-#if CANDIDATE_LIGHT_MASK
-	if (all(DownsampledScreenCoord < DownsampledViewSize))
+	uint DownsampledTileIndex = GroupId.x;
+	if (DownsampledTileIndex < DownsampledTileAllocator[TILE_TYPE])
 	{
-		uint2 ScreenCoord = DownsampledScreenCoordToScreenCoord(DownsampledScreenCoord);	
-		const float2 ScreenUV = (ScreenCoord + 0.5f) * View.BufferSizeAndInvSize.zw;
-		const FGBufferData GBuffer = GetGBufferData(ScreenUV);
-		const float SceneDepth = GBuffer.Depth;
+		uint2 DownsampledTileCoord = UnpackTile(DownsampledTileData[DownsampledTileIndex + TILE_TYPE * DownsampledTileDataStride]);
+		uint2 DownsampledScreenCoord = DownsampledTileCoord * TILE_SIZE + GroupThreadId.xy;
 
-		float3 TranslatedWorldPosition = GetTranslatedWorldPositionFromScreenUV(ScreenUV, SceneDepth);		
+		#if CANDIDATE_LIGHT_MASK
+		if (all(DownsampledScreenCoord < DownsampledViewSize))
+		{
+			uint2 ScreenCoord = DownsampledScreenCoordToScreenCoord(DownsampledScreenCoord);	
+			const float2 ScreenUV = (ScreenCoord + 0.5f) * View.BufferSizeAndInvSize.zw;
+			const FGBufferData GBuffer = GetGBufferData(ScreenUV);
+			const float SceneDepth = GBuffer.Depth;
 
-		const uint EyeIndex = 0;
-		const uint GridIndex = ComputeLightGridCellIndex(ScreenCoord, SceneDepth, EyeIndex);
-		const FCulledLightsGridData CulledLightGridData = GetCulledLightsGrid(GridIndex, EyeIndex);
-		const uint NumLocalLights = min(CulledLightGridData.NumLocalLights, GetNumLocalLights(EyeIndex));
+			float3 TranslatedWorldPosition = GetTranslatedWorldPositionFromScreenUV(ScreenUV, SceneDepth);		
 
-		for (uint LocalLightIndex = 0; LocalLightIndex < NumLocalLights; ++LocalLightIndex)
-		{
-			const FLocalLightData LocalLightData = GetLocalLightData(CulledLightGridData.DataStartIndex + LocalLightIndex, EyeIndex);
-			if (LocalLightData.bStochasticShadowsSupported && LocalLightData.LightSceneId < MAX_LIGHT_SCENE_ID)
+			const uint EyeIndex = 0;
+			const uint GridIndex = ComputeLightGridCellIndex(ScreenCoord, SceneDepth, EyeIndex);
+			const FCulledLightsGridData CulledLightGridData = GetCulledLightsGrid(GridIndex, EyeIndex);
+			const uint NumLocalLights = min(CulledLightGridData.NumLocalLights, GetNumLocalLights(EyeIndex));
+
+			for (uint LocalLightIndex = 0; LocalLightIndex < NumLocalLights; ++LocalLightIndex)
 			{
-				uint DWORDIndex = LocalLightData.LightSceneId / 32;
-				uint BitMask = 1u << (LocalLightData.LightSceneId % 32);
-				InterlockedOr(SharedCandidateLightMask[DWORDIndex], BitMask);
+				const FLocalLightData LocalLightData = GetLocalLightData(CulledLightGridData.DataStartIndex + LocalLightIndex, EyeIndex);
+				if (LocalLightData.bStochasticShadowsSupported && LocalLightData.LightSceneId < MAX_LIGHT_SCENE_ID)
+				{
+					uint DWORDIndex = LocalLightData.LightSceneId / 32;
+					uint BitMask = 1u << (LocalLightData.LightSceneId % 32);
+					InterlockedOr(SharedCandidateLightMask[DWORDIndex], BitMask);
+				}
 			}
 		}
-	}
-#endif
-
-	GroupMemoryBarrierWithGroupSync();
-
-	if (all(DownsampledScreenCoord < DownsampledViewSize))
-	{
-		uint2 ScreenCoord = DownsampledScreenCoordToScreenCoord(DownsampledScreenCoord);
-
-#if DEBUG_MODE
-		int2 DebugScreenCoord = View.CursorPosition.x >= 0 ? View.CursorPosition * View.ViewResolutionFraction : -1;
-		bool bDebug = all(DownsampledScreenCoord == DebugScreenCoord / DOWNSAMPLE_FACTOR);
-		FShaderPrintContext Context = InitShaderPrintContext(true, float2(0.05, 0.05));
-#endif
+		#endif
+	
+		GroupMemoryBarrierWithGroupSync();
 
-		uint SelectedLights[NUM_SAMPLES_PER_PIXEL];
-		for (uint LightSampleIndex = 0; LightSampleIndex < NUM_SAMPLES_PER_PIXEL; ++LightSampleIndex)
+		if (all(DownsampledScreenCoord < DownsampledViewSize))
 		{
-			SelectedLights[LightSampleIndex] = INVALID_LIGHT_SCENE_ID;
-		}
+			uint2 ScreenCoord = DownsampledScreenCoordToScreenCoord(DownsampledScreenCoord);
 
-		const uint2 LightSampleOffsets[4] = 
-		{
-			uint2(0, 0),
-			uint2(1, 0),
-			uint2(0, 1),
-			uint2(1, 1),
-		};
+	#if DEBUG_MODE
+			int2 DebugScreenCoord = View.CursorPosition.x >= 0 ? View.CursorPosition * View.ViewResolutionFraction : -1;
+			bool bDebug = all(DownsampledScreenCoord == DebugScreenCoord / DOWNSAMPLE_FACTOR);
+			FShaderPrintContext Context = InitShaderPrintContext(true, float2(0.05, 0.05));
+	#endif
+
+			uint SelectedLights[NUM_SAMPLES_PER_PIXEL];
+			for (uint LightSampleIndex = 0; LightSampleIndex < NUM_SAMPLES_PER_PIXEL; ++LightSampleIndex)
+			{
+				SelectedLights[LightSampleIndex] = INVALID_LIGHT_SCENE_ID;
+			}
 
-		uint2 ShadowMaskHistoryScreenCoord00 = 0;
-		float4 ShadowMaskHistoryWeights = 0.0f;
+			const uint2 LightSampleOffsets[4] = 
+			{
+				uint2(0, 0),
+				uint2(1, 0),
+				uint2(0, 1),
+				uint2(1, 1),
+			};
 
-		const float2 ScreenUV = (ScreenCoord + 0.5f) * View.BufferSizeAndInvSize.zw;
-		const FGBufferData GBuffer = GetGBufferData(ScreenUV);
-		const float SceneDepth = GBuffer.Depth;
+			uint2 ShadowMaskHistoryScreenCoord00 = 0;
+			float4 ShadowMaskHistoryWeights = 0.0f;
 
-		if (SceneDepth > 0)
-		{
-			// Add a bit of stochastic noise to hide transitions
-			const float StochasticMinLightSampleWeight = MinLightSampleWeight * (1.0f + InterleavedGradientNoise(DownsampledScreenCoord, StochasticShadowsStateFrameIndex));
+			const float2 ScreenUV = (ScreenCoord + 0.5f) * View.BufferSizeAndInvSize.zw;
+			const FGBufferData GBuffer = GetGBufferData(ScreenUV);
+			const float SceneDepth = GBuffer.Depth;
 
-			// Prepare weights for history reprojection
-			#if SHADOW_MASK_REPROJECTION_WEIGHTS
+			if (SceneDepth > 0)
 			{
-				// Reproject center of downsampled texel
-				float2 DownsampledScreenUV = (DownsampledScreenCoord + 0.5f) * DOWNSAMPLE_FACTOR * View.BufferSizeAndInvSize.zw;
-
-				float DeviceZ = ConvertToDeviceZ(SceneDepth);
-				float2 ScreenPosition = (DownsampledScreenUV - View.ScreenPositionScaleBias.wz) / View.ScreenPositionScaleBias.xy;
-				float3 HistoryScreenPosition = GetHistoryScreenPosition(ScreenPosition, DownsampledScreenUV, DeviceZ);
-				float2 HistoryScreenUV = HistoryScreenPosition.xy * HistoryScreenPositionScaleBias.xy + HistoryScreenPositionScaleBias.wz;
+				// Add a bit of stochastic noise to hide transitions
+				const float StochasticMinLightSampleWeight = MinLightSampleWeight * (1.0f + InterleavedGradientNoise(DownsampledScreenCoord, StochasticShadowsStateFrameIndex));
 
-				bool bHistoryWasOnScreen = all(HistoryScreenUV >= HistoryUVMinMax.xy) && all (HistoryScreenUV <= HistoryUVMinMax.zw);
-				if (bHistoryWasOnScreen)
+				// Prepare weights for history reprojection
+				#if SHADOW_MASK_REPROJECTION_WEIGHTS
 				{
-					uint2 HistoryScreenCoord = floor(HistoryScreenUV * View.BufferSizeAndInvSize.xy / DOWNSAMPLE_FACTOR - 0.5f);
-					float2 HistoryBilinearWeights = frac(HistoryScreenUV * View.BufferSizeAndInvSize.xy / DOWNSAMPLE_FACTOR - 0.5f);
-					float2 HistoryGatherUV = (HistoryScreenCoord + 1.0f) * DownsampledBufferInvSize;
+					// Reproject center of downsampled texel
+					float2 DownsampledScreenUV = (DownsampledScreenCoord + 0.5f) * DOWNSAMPLE_FACTOR * View.BufferSizeAndInvSize.zw;
 
-					float4 HistorySampleSceneDepth4 = ShadowMaskSceneDepthHistory.GatherRed(GlobalPointClampedSampler, HistoryGatherUV).wzxy;
+					float DeviceZ = ConvertToDeviceZ(SceneDepth);
+					float2 ScreenPosition = (DownsampledScreenUV - View.ScreenPositionScaleBias.wz) / View.ScreenPositionScaleBias.xy;
+					float3 HistoryScreenPosition = GetHistoryScreenPosition(ScreenPosition, DownsampledScreenUV, DeviceZ);
+					float2 HistoryScreenUV = HistoryScreenPosition.xy * HistoryScreenPositionScaleBias.xy + HistoryScreenPositionScaleBias.wz;
 
-					// #sdl_todo: plane weighting
-					float ReprojectedSceneDepth = ConvertFromDeviceZ(HistoryScreenPosition.z);
-					float DisocclusionDistanceThreshold = 0.03f;
-					const float4 DistanceToHistoryValue = abs(HistorySampleSceneDepth4 - ReprojectedSceneDepth);
-					float4 DepthWeights = select(DistanceToHistoryValue >= ReprojectedSceneDepth * DisocclusionDistanceThreshold, 0.0f, 1.0f);
+					bool bHistoryWasOnScreen = all(HistoryScreenUV >= HistoryUVMinMax.xy) && all (HistoryScreenUV <= HistoryUVMinMax.zw);
+					if (bHistoryWasOnScreen)
+					{
+						uint2 HistoryScreenCoord = floor(HistoryScreenUV * View.BufferSizeAndInvSize.xy / DOWNSAMPLE_FACTOR - 0.5f);
+						float2 HistoryBilinearWeights = frac(HistoryScreenUV * View.BufferSizeAndInvSize.xy / DOWNSAMPLE_FACTOR - 0.5f);
+						float2 HistoryGatherUV = (HistoryScreenCoord + 1.0f) * DownsampledBufferInvSize;
 
-					float4 HistoryWeights = float4(
-						(1 - HistoryBilinearWeights.y) * (1 - HistoryBilinearWeights.x),
-						(1 - HistoryBilinearWeights.y) * HistoryBilinearWeights.x,
-						HistoryBilinearWeights.y * (1 - HistoryBilinearWeights.x),
-						HistoryBilinearWeights.y * HistoryBilinearWeights.x);
+						float4 HistorySampleSceneDepth4 = ShadowMaskSceneDepthHistory.GatherRed(GlobalPointClampedSampler, HistoryGatherUV).wzxy;
 
-					HistoryWeights *= DepthWeights;
+						// #sdl_todo: plane weighting
+						float ReprojectedSceneDepth = ConvertFromDeviceZ(HistoryScreenPosition.z);
+						float DisocclusionDistanceThreshold = 0.03f;
+						const float4 DistanceToHistoryValue = abs(HistorySampleSceneDepth4 - ReprojectedSceneDepth);
+						float4 DepthWeights = select(DistanceToHistoryValue >= ReprojectedSceneDepth * DisocclusionDistanceThreshold, 0.0f, 1.0f);
 
-					// Use history only if we found at least one valid sample
-					if (any(HistoryWeights > 0.01f))
-					{
-						ShadowMaskHistoryScreenCoord00 = HistoryScreenCoord;
-						ShadowMaskHistoryWeights = HistoryWeights;
+						float4 HistoryWeights = float4(
+							(1 - HistoryBilinearWeights.y) * (1 - HistoryBilinearWeights.x),
+							(1 - HistoryBilinearWeights.y) * HistoryBilinearWeights.x,
+							HistoryBilinearWeights.y * (1 - HistoryBilinearWeights.x),
+							HistoryBilinearWeights.y * HistoryBilinearWeights.x);
+
+						HistoryWeights *= DepthWeights;
+
+						// Use history only if we found at least one valid sample
+						if (any(HistoryWeights > 0.01f))
+						{
+							ShadowMaskHistoryScreenCoord00 = HistoryScreenCoord;
+							ShadowMaskHistoryWeights = HistoryWeights;
+						}
 					}
 				}
-			}
-			#endif
+				#endif
 
-			float3 TranslatedWorldPosition = GetTranslatedWorldPositionFromScreenUV(ScreenUV, SceneDepth);		
+				float3 TranslatedWorldPosition = GetTranslatedWorldPositionFromScreenUV(ScreenUV, SceneDepth);		
 
-			const uint EyeIndex = 0;
-			const uint GridIndex = ComputeLightGridCellIndex(ScreenCoord, SceneDepth, EyeIndex);
-			const FCulledLightsGridData CulledLightGridData = GetCulledLightsGrid(GridIndex, EyeIndex);
-			const uint NumLocalLights = min(CulledLightGridData.NumLocalLights, GetNumLocalLights(EyeIndex));
+				const uint EyeIndex = 0;
+				const uint GridIndex = ComputeLightGridCellIndex(ScreenCoord, SceneDepth, EyeIndex);
+				const FCulledLightsGridData CulledLightGridData = GetCulledLightsGrid(GridIndex, EyeIndex);
+				const uint NumLocalLights = min(CulledLightGridData.NumLocalLights, GetNumLocalLights(EyeIndex));
 
-			float LightIndexRandom[NUM_SAMPLES_PER_PIXEL];
-			for (uint LightSampleIndex = 0; LightSampleIndex < NUM_SAMPLES_PER_PIXEL; ++LightSampleIndex)
-			{
-				LightIndexRandom[LightSampleIndex] = BlueNoiseScalar(DispatchThreadId.xy * DOWNSAMPLE_FACTOR + LightSampleOffsets[LightSampleIndex], StochasticShadowsStateFrameIndex);
-			}
+				float LightIndexRandom[NUM_SAMPLES_PER_PIXEL];
+				for (uint LightSampleIndex = 0; LightSampleIndex < NUM_SAMPLES_PER_PIXEL; ++LightSampleIndex)
+				{
+					LightIndexRandom[LightSampleIndex] = BlueNoiseScalar(DownsampledScreenCoord * DOWNSAMPLE_FACTOR + LightSampleOffsets[LightSampleIndex], StochasticShadowsStateFrameIndex);
+				}
 
-			#if DEBUG_MODE
-			if (bDebug)
-			{
-				Print(Context, TEXT("GenerateSamplesCS"));
-				Newline(Context);
-				Print(Context, TEXT("ScreenCoord: "));
-				Print(Context, ScreenCoord.x);
-				Print(Context, ScreenCoord.y);
-				Newline(Context);
-				Print(Context, TEXT("View.PreExposure: "));
-				Print(Context, View.PreExposure);
-				Newline(Context);
-				Print(Context, TEXT("NumLocalLights: "));
-				Print(Context, NumLocalLights);
-				AddCrossTWS(Context, TranslatedWorldPosition, 5.0f, float4(1, 1, 0, 1));
-				Newline(Context);
-				Print(Context, TEXT("LightId  | SampleWeight | Weight   | Estimate | Shadow | ShadowValid | ShadowAge | HistoryLength"));
-			}
-			#endif
+				#if DEBUG_MODE
+				if (bDebug)
+				{
+					Print(Context, TEXT("GenerateSamplesCS"));
+					Newline(Context);
+					Print(Context, TEXT("ScreenCoord: "));
+					Print(Context, ScreenCoord.x);
+					Print(Context, ScreenCoord.y);
+					Newline(Context);
+					Print(Context, TEXT("TileType: "));
+					Print(Context, uint(TILE_TYPE));
+					Print(Context, DownsampledTileAllocator[TILE_TYPE]);
+					Newline(Context);
+					Print(Context, TEXT("View.PreExposure: "));
+					Print(Context, View.PreExposure);
+					Newline(Context);
+					Print(Context, TEXT("NumLocalLights: "));
+					Print(Context, NumLocalLights);
+					AddCrossTWS(Context, TranslatedWorldPosition, 5.0f, float4(1, 1, 0, 1));
+					Newline(Context);
+					Print(Context, TEXT("LightId  | SampleWeight | Weight   | Estimate | Shadow | ShadowValid | ShadowAge | HistoryLength"));
+				}
+				#endif
 
-			float WeightSum = 0.0f;
+				float WeightSum = 0.0f;
 
-			// Weighted blue noise sampling
-			float Xi[NUM_SAMPLES_PER_PIXEL] = LightIndexRandom;
+				// Weighted blue noise sampling
+				float Xi[NUM_SAMPLES_PER_PIXEL] = LightIndexRandom;
 
-#if CANDIDATE_LIGHT_MASK
-			for (uint MaskIndex = 0; MaskIndex < LIGHT_SCENE_ID_MASK_SIZE; ++MaskIndex)
-			{
-				uint CandidateLightMask = SharedCandidateLightMask[MaskIndex];
-				while (CandidateLightMask != 0)
+	#if CANDIDATE_LIGHT_MASK
+				for (uint MaskIndex = 0; MaskIndex < LIGHT_SCENE_ID_MASK_SIZE; ++MaskIndex)
 				{
-					const uint NextBitIndex = firstbitlow(CandidateLightMask);
-					const uint NextBitMask = 1u << NextBitIndex;
-					CandidateLightMask ^= NextBitMask;
+					uint CandidateLightMask = SharedCandidateLightMask[MaskIndex];
+					while (CandidateLightMask != 0)
+					{
+						const uint NextBitIndex = firstbitlow(CandidateLightMask);
+						const uint NextBitMask = 1u << NextBitIndex;
+						CandidateLightMask ^= NextBitMask;
 
-					uint LightSceneId = MaskIndex * 32 + NextBitIndex;
-					FLightSceneData LightSceneData = GetLightSceneData(LightSceneId);
-					FDeferredLightData LightData = ConvertToDeferredLight(LightSceneData);
+						uint LightSceneId = MaskIndex * 32 + NextBitIndex;
+						FLightSceneData LightSceneData = GetLightSceneData(LightSceneId);
+						FDeferredLightData LightData = ConvertToDeferredLight(LightSceneData);
 					
-					float SampleWeight = GetLocalLightTargetPDF(LightData, TranslatedWorldPosition, GBuffer, ScreenCoord);
+						float SampleWeight = GetLocalLightTargetPDF(LightData, TranslatedWorldPosition, GBuffer, ScreenCoord);
 
-					// Simulate tonemapping
-					SampleWeight = log2(max(SampleWeight + 1.0f, 1.0f));
+						// Simulate tonemapping
+						SampleWeight = log2(max(SampleWeight + 1.0f, 1.0f));
 
-					// Get shadow factor from the last frame
-					FShadowMask HistoryShadowMask = InitShadowMask();
-					float ShadowFactorEstimate = 1.0f;
+						// Get shadow factor from the last frame
+						FShadowMask HistoryShadowMask = InitShadowMask();
+						float ShadowFactorEstimate = 1.0f;
 
-					#if SHADOW_FACTOR_ESTIMATE
-					{
-						if (any(ShadowMaskHistoryWeights > 0))
+						#if SHADOW_FACTOR_ESTIMATE
 						{
-							FShadowMaskAccumulator Accumulator = InitShadowMaskAccumulator();
+							if (any(ShadowMaskHistoryWeights > 0))
+							{
+								FShadowMaskAccumulator Accumulator = InitShadowMaskAccumulator();
 
-							AccumulateHistoryShadowMask(LightSceneId, ShadowMaskHistoryScreenCoord00 + uint2(0, 0), ShadowMaskHistoryWeights.x, Accumulator);
-							AccumulateHistoryShadowMask(LightSceneId, ShadowMaskHistoryScreenCoord00 + uint2(1, 0), ShadowMaskHistoryWeights.y, Accumulator);
-							AccumulateHistoryShadowMask(LightSceneId, ShadowMaskHistoryScreenCoord00 + uint2(0, 1), ShadowMaskHistoryWeights.z, Accumulator);
-							AccumulateHistoryShadowMask(LightSceneId, ShadowMaskHistoryScreenCoord00 + uint2(1, 1), ShadowMaskHistoryWeights.w, Accumulator);
+								AccumulateHistoryShadowMask(LightSceneId, ShadowMaskHistoryScreenCoord00 + uint2(0, 0), ShadowMaskHistoryWeights.x, Accumulator);
+								AccumulateHistoryShadowMask(LightSceneId, ShadowMaskHistoryScreenCoord00 + uint2(1, 0), ShadowMaskHistoryWeights.y, Accumulator);
+								AccumulateHistoryShadowMask(LightSceneId, ShadowMaskHistoryScreenCoord00 + uint2(0, 1), ShadowMaskHistoryWeights.z, Accumulator);
+								AccumulateHistoryShadowMask(LightSceneId, ShadowMaskHistoryScreenCoord00 + uint2(1, 1), ShadowMaskHistoryWeights.w, Accumulator);
 
-							if (Accumulator.WeightSum > 0.0f)
-							{
-								HistoryShadowMask = Accumulator.ResolveShadowMask();
+								if (Accumulator.WeightSum > 0.0f)
+								{
+									HistoryShadowMask = Accumulator.ResolveShadowMask();
 
-								float AgeCertainty = (HistoryShadowMask.Age - 1) / float(MAX_SHADOW_MASK_AGE - 1.0f);
-								ShadowFactorEstimate = lerp(HistoryShadowMask.Factor, 1.0f, saturate(AgeCertainty));
+									float AgeCertainty = (HistoryShadowMask.Age - 1) / float(MAX_SHADOW_MASK_AGE - 1.0f);
+									ShadowFactorEstimate = lerp(HistoryShadowMask.Factor, 1.0f, saturate(AgeCertainty));
+								}
 							}
 						}
-					}
-					#endif
+						#endif
 
-					#if DEBUG_MODE
-					if (bDebug)
-					{
-						Newline(Context);
-						Print(Context, LightSceneId);
-						Print(Context, SampleWeight);
-						Print(Context, SampleWeight * ShadowFactorEstimate);
-						Print(Context, ShadowFactorEstimate);
-						Print(Context, HistoryShadowMask.Factor);
-						Print(Context, HistoryShadowMask.IsValid() ? 1u : 0u);
-						Print(Context, HistoryShadowMask.Age);
-						Print(Context, HistoryShadowMask.HistoryLength);
-					}
-					#endif
-
-					if (SampleWeight > StochasticMinLightSampleWeight)
-					{
-						// Allocate required shadow mask tiles for this frame
-						MarkLightSceneId(LightSceneId);
-
-						// Apply shadowing estimate from the last frame
-						if (HistoryShadowMask.IsValid())
+						#if DEBUG_MODE
+						if (bDebug)
 						{
-							SampleWeight *= ShadowFactorEstimate;
+							Newline(Context);
+							Print(Context, LightSceneId);
+							Print(Context, SampleWeight);
+							Print(Context, SampleWeight * ShadowFactorEstimate);
+							Print(Context, ShadowFactorEstimate);
+							Print(Context, HistoryShadowMask.Factor);
+							Print(Context, HistoryShadowMask.IsValid() ? 1u : 0u);
+							Print(Context, HistoryShadowMask.Age);
+							Print(Context, HistoryShadowMask.HistoryLength);
 						}
+						#endif
 
 						if (SampleWeight > StochasticMinLightSampleWeight)
 						{
-							float Tau = WeightSum / (WeightSum + SampleWeight);
-							WeightSum += SampleWeight;
+							// Allocate required shadow mask tiles for this frame
+							MarkLightSceneId(LightSceneId);
 
-							for (uint LightSampleIndex = 0; LightSampleIndex < NUM_SAMPLES_PER_PIXEL; ++LightSampleIndex)
+							// Apply shadowing estimate from the last frame
+							if (HistoryShadowMask.IsValid())
 							{
-								if (Xi[LightSampleIndex] < Tau)
-								{
-									Xi[LightSampleIndex] /= Tau;
-								}
-								else
+								SampleWeight *= ShadowFactorEstimate;
+							}
+
+							if (SampleWeight > StochasticMinLightSampleWeight)
+							{
+								float Tau = WeightSum / (WeightSum + SampleWeight);
+								WeightSum += SampleWeight;
+
+								for (uint LightSampleIndex = 0; LightSampleIndex < NUM_SAMPLES_PER_PIXEL; ++LightSampleIndex)
 								{
-									// Select this sample
-									Xi[LightSampleIndex] = (Xi[LightSampleIndex] - Tau) / (1.0f - Tau);
-									SelectedLights[LightSampleIndex] = LightSceneId;
+									if (Xi[LightSampleIndex] < Tau)
+									{
+										Xi[LightSampleIndex] /= Tau;
+									}
+									else
+									{
+										// Select this sample
+										Xi[LightSampleIndex] = (Xi[LightSampleIndex] - Tau) / (1.0f - Tau);
+										SelectedLights[LightSampleIndex] = LightSceneId;
+									}
+
+									Xi[LightSampleIndex] = clamp(Xi[LightSampleIndex], 0, 0.9999f);
 								}
-
-								Xi[LightSampleIndex] = clamp(Xi[LightSampleIndex], 0, 0.9999f);
 							}
-						}
-					}						
+						}						
+					}
 				}
-			}
 
-#else
-			for (uint LocalLightIndex = 0; LocalLightIndex < NumLocalLights; ++LocalLightIndex)
-			{
-				const FLocalLightData LocalLightData = GetLocalLightData(CulledLightGridData.DataStartIndex + LocalLightIndex, EyeIndex);
-				if (LocalLightData.bStochasticShadowsSupported)
+	#else
+				for (uint LocalLightIndex = 0; LocalLightIndex < NumLocalLights; ++LocalLightIndex)
 				{
-					FDeferredLightData LightData = ConvertToDeferredLight(LocalLightData);
-					float SampleWeight = GetLocalLightTargetPDF(LightData, TranslatedWorldPosition, GBuffer, ScreenCoord);
+					const FLocalLightData LocalLightData = GetLocalLightData(CulledLightGridData.DataStartIndex + LocalLightIndex, EyeIndex);
+					if (LocalLightData.bStochasticShadowsSupported)
+					{
+						FDeferredLightData LightData = ConvertToDeferredLight(LocalLightData);
+						float SampleWeight = GetLocalLightTargetPDF(LightData, TranslatedWorldPosition, GBuffer, ScreenCoord);
 
-					// Simulate tonemapping
-					SampleWeight = log2(max(SampleWeight + 1.0f, 1.0f));
+						// Simulate tonemapping
+						SampleWeight = log2(max(SampleWeight + 1.0f, 1.0f));
 
-					// Get shadow factor from the last frame
-					FShadowMask HistoryShadowMask = InitShadowMask();
-					float ShadowFactorEstimate = 1.0f;
+						// Get shadow factor from the last frame
+						FShadowMask HistoryShadowMask = InitShadowMask();
+						float ShadowFactorEstimate = 1.0f;
 
-					#if SHADOW_FACTOR_ESTIMATE
-					{
-						if (any(ShadowMaskHistoryWeights > 0))
+						#if SHADOW_FACTOR_ESTIMATE
 						{
-							FShadowMaskAccumulator Accumulator = InitShadowMaskAccumulator();
+							if (any(ShadowMaskHistoryWeights > 0))
+							{
+								FShadowMaskAccumulator Accumulator = InitShadowMaskAccumulator();
 
-							AccumulateHistoryShadowMask(LocalLightData.LightSceneId, ShadowMaskHistoryScreenCoord00 + uint2(0, 0), ShadowMaskHistoryWeights.x, Accumulator);
-							AccumulateHistoryShadowMask(LocalLightData.LightSceneId, ShadowMaskHistoryScreenCoord00 + uint2(1, 0), ShadowMaskHistoryWeights.y, Accumulator);
-							AccumulateHistoryShadowMask(LocalLightData.LightSceneId, ShadowMaskHistoryScreenCoord00 + uint2(0, 1), ShadowMaskHistoryWeights.z, Accumulator);
-							AccumulateHistoryShadowMask(LocalLightData.LightSceneId, ShadowMaskHistoryScreenCoord00 + uint2(1, 1), ShadowMaskHistoryWeights.w, Accumulator);
+								AccumulateHistoryShadowMask(LocalLightData.LightSceneId, ShadowMaskHistoryScreenCoord00 + uint2(0, 0), ShadowMaskHistoryWeights.x, Accumulator);
+								AccumulateHistoryShadowMask(LocalLightData.LightSceneId, ShadowMaskHistoryScreenCoord00 + uint2(1, 0), ShadowMaskHistoryWeights.y, Accumulator);
+								AccumulateHistoryShadowMask(LocalLightData.LightSceneId, ShadowMaskHistoryScreenCoord00 + uint2(0, 1), ShadowMaskHistoryWeights.z, Accumulator);
+								AccumulateHistoryShadowMask(LocalLightData.LightSceneId, ShadowMaskHistoryScreenCoord00 + uint2(1, 1), ShadowMaskHistoryWeights.w, Accumulator);
 
-							if (Accumulator.WeightSum > 0.0f)
-							{
-								HistoryShadowMask = Accumulator.ResolveShadowMask();
+								if (Accumulator.WeightSum > 0.0f)
+								{
+									HistoryShadowMask = Accumulator.ResolveShadowMask();
 
-								float AgeCertainty = (HistoryShadowMask.Age - 1) / float(MAX_SHADOW_MASK_AGE - 1.0f);
-								ShadowFactorEstimate = lerp(HistoryShadowMask.Factor, 1.0f, saturate(AgeCertainty));
+									float AgeCertainty = (HistoryShadowMask.Age - 1) / float(MAX_SHADOW_MASK_AGE - 1.0f);
+									ShadowFactorEstimate = lerp(HistoryShadowMask.Factor, 1.0f, saturate(AgeCertainty));
+								}
 							}
 						}
-					}
-					#endif
-
-					#if DEBUG_MODE
-					if (bDebug)
-					{
-						Newline(Context);
-						Print(Context, LocalLightData.LightSceneId);
-						Print(Context, SampleWeight);
-						Print(Context, SampleWeight * ShadowFactorEstimate);
-						Print(Context, ShadowFactorEstimate);
-						Print(Context, HistoryShadowMask.Factor);
-						Print(Context, HistoryShadowMask.IsValid() ? 1u : 0u);
-						Print(Context, HistoryShadowMask.Age);
-						Print(Context, HistoryShadowMask.HistoryLength);
-					}
-					#endif
+						#endif
 
-					if (SampleWeight > StochasticMinLightSampleWeight)
-					{
-						// Allocate required shadow mask tiles for this frame
-						MarkLightSceneId(LocalLightData.LightSceneId);
-
-						// Apply shadowing estimate from the last frame
-						if (HistoryShadowMask.IsValid())
+						#if DEBUG_MODE
+						if (bDebug)
 						{
-							SampleWeight *= ShadowFactorEstimate;
+							Newline(Context);
+							Print(Context, LocalLightData.LightSceneId);
+							Print(Context, SampleWeight);
+							Print(Context, SampleWeight * ShadowFactorEstimate);
+							Print(Context, ShadowFactorEstimate);
+							Print(Context, HistoryShadowMask.Factor);
+							Print(Context, HistoryShadowMask.IsValid() ? 1u : 0u);
+							Print(Context, HistoryShadowMask.Age);
+							Print(Context, HistoryShadowMask.HistoryLength);
 						}
+						#endif
 
 						if (SampleWeight > StochasticMinLightSampleWeight)
 						{
-							float Tau = WeightSum / (WeightSum + SampleWeight);
-							WeightSum += SampleWeight;
+							// Allocate required shadow mask tiles for this frame
+							MarkLightSceneId(LocalLightData.LightSceneId);
 
-							for (uint LightSampleIndex = 0; LightSampleIndex < NUM_SAMPLES_PER_PIXEL; ++LightSampleIndex)
+							// Apply shadowing estimate from the last frame
+							if (HistoryShadowMask.IsValid())
 							{
-								if (Xi[LightSampleIndex] < Tau)
-								{
-									Xi[LightSampleIndex] /= Tau;
-								}
-								else
+								SampleWeight *= ShadowFactorEstimate;
+							}
+
+							if (SampleWeight > StochasticMinLightSampleWeight)
+							{
+								float Tau = WeightSum / (WeightSum + SampleWeight);
+								WeightSum += SampleWeight;
+
+								for (uint LightSampleIndex = 0; LightSampleIndex < NUM_SAMPLES_PER_PIXEL; ++LightSampleIndex)
 								{
-									// Select this sample
-									Xi[LightSampleIndex] = (Xi[LightSampleIndex] - Tau) / (1.0f - Tau);
-									SelectedLights[LightSampleIndex] = LocalLightData.LightSceneId;
+									if (Xi[LightSampleIndex] < Tau)
+									{
+										Xi[LightSampleIndex] /= Tau;
+									}
+									else
+									{
+										// Select this sample
+										Xi[LightSampleIndex] = (Xi[LightSampleIndex] - Tau) / (1.0f - Tau);
+										SelectedLights[LightSampleIndex] = LocalLightData.LightSceneId;
+									}
+
+									Xi[LightSampleIndex] = clamp(Xi[LightSampleIndex], 0, 0.9999f);
 								}
-
-								Xi[LightSampleIndex] = clamp(Xi[LightSampleIndex], 0, 0.9999f);
 							}
 						}
 					}
 				}
-			}
-#endif
-
-			#if DEBUG_MODE
-			if (bDebug)
-			{
-				Newline(Context);
-				Print(Context, TEXT("Selected lights: "));
+	#endif
 
-				for (uint LightSampleIndex = 0; LightSampleIndex < NUM_SAMPLES_PER_PIXEL; ++LightSampleIndex)
+				#if DEBUG_MODE
+				if (bDebug)
 				{
-					Print(Context, SelectedLights[LightSampleIndex]);
+					Newline(Context);
+					Print(Context, TEXT("Selected lights: "));
 
-					if (DebugMode == DEBUG_MODE_VISUALIZE_SAMPLING)
+					for (uint LightSampleIndex = 0; LightSampleIndex < NUM_SAMPLES_PER_PIXEL; ++LightSampleIndex)
 					{
-						const uint2 SampleCoord = DownsampledScreenCoord * DOWNSAMPLE_FACTOR + LightSampleOffsets[LightSampleIndex];
-						FLightSceneData LightSceneData = GetLightSceneData(SelectedLights[LightSampleIndex]);
-						FLightSampleTrace LightSampleTrace = GetLightSampleTrace(TranslatedWorldPosition, SelectedLights[LightSampleIndex], SampleCoord);
-						AddLineTWS(Context, TranslatedWorldPosition, TranslatedWorldPosition + LightSampleTrace.Direction * LightSampleTrace.Distance, float4(LightSceneData.Color.xyz / Luminance(LightSceneData.Color.xyz), 1.0f));
+						Print(Context, SelectedLights[LightSampleIndex]);
+
+						if (DebugMode == DEBUG_MODE_VISUALIZE_SAMPLING)
+						{
+							const uint2 SampleCoord = DownsampledScreenCoord * DOWNSAMPLE_FACTOR + LightSampleOffsets[LightSampleIndex];
+							FLightSceneData LightSceneData = GetLightSceneData(SelectedLights[LightSampleIndex]);
+							FLightSampleTrace LightSampleTrace = GetLightSampleTrace(TranslatedWorldPosition, SelectedLights[LightSampleIndex], SampleCoord);
+							AddLineTWS(Context, TranslatedWorldPosition, TranslatedWorldPosition + LightSampleTrace.Direction * LightSampleTrace.Distance, float4(LightSceneData.Color.xyz / Luminance(LightSceneData.Color.xyz), 1.0f));
+						}
 					}
 				}
+				#endif
 			}
-			#endif
-		}
 
-		// Cache common parts for later
-		#if SHADOW_MASK_REPROJECTION_WEIGHTS
-		{
-			RWShadowMaskHistoryScreenCoord00[DownsampledScreenCoord] = ShadowMaskHistoryScreenCoord00;
-			RWShadowMaskHistoryWeights[DownsampledScreenCoord] = ShadowMaskHistoryWeights;
-		}
-		#endif
-		RWDownsampledSceneDepth[DownsampledScreenCoord] = SceneDepth;
+			// Cache common parts for later
+			#if SHADOW_MASK_REPROJECTION_WEIGHTS
+			{
+				RWShadowMaskHistoryScreenCoord00[DownsampledScreenCoord] = ShadowMaskHistoryScreenCoord00;
+				RWShadowMaskHistoryWeights[DownsampledScreenCoord] = ShadowMaskHistoryWeights;
+			}
+			#endif
+			RWDownsampledSceneDepth[DownsampledScreenCoord] = SceneDepth;
 
-		for (uint LightSampleIndex = 0; LightSampleIndex < NUM_SAMPLES_PER_PIXEL; ++LightSampleIndex)
-		{
-			FLightSample LightSample;
-			LightSample.LightSceneId = SelectedLights[LightSampleIndex];
-			LightSample.bCompleted = LightSample.LightSceneId == INVALID_LIGHT_SCENE_ID ? true : false;
-			LightSample.bVisible = true;
-			RWLightSamples[DownsampledScreenCoord * uint2(NUM_SAMPLES_PER_PIXEL_X, NUM_SAMPLES_PER_PIXEL_Y) + LightSampleOffsets[LightSampleIndex]] = PackLightSample(LightSample);
+			for (uint LightSampleIndex = 0; LightSampleIndex < NUM_SAMPLES_PER_PIXEL; ++LightSampleIndex)
+			{
+				FLightSample LightSample;
+				LightSample.LightSceneId = SelectedLights[LightSampleIndex];
+				LightSample.bCompleted = LightSample.LightSceneId == INVALID_LIGHT_SCENE_ID ? true : false;
+				LightSample.bVisible = true;
+				RWLightSamples[DownsampledScreenCoord * uint2(NUM_SAMPLES_PER_PIXEL_X, NUM_SAMPLES_PER_PIXEL_Y) + LightSampleOffsets[LightSampleIndex]] = PackLightSample(LightSample);
+			}
 		}
-	}
-
-	GroupMemoryBarrierWithGroupSync();
-
-	// Write out allocated shadow mask tiles
 
-#if HASH_TABLE
-	if (LinearThreadIndex < HASH_TABLE_SIZE)
-	{
-		uint2 ShadowMaskScreenCoord = GroupId.xy * SHADOW_MASK_TILE_SIZE;
+		GroupMemoryBarrierWithGroupSync();
 
-		uint LightIndex = SharedHashTable[LinearThreadIndex];
-
-		if (LightIndex != INVALID_LIGHT_SCENE_ID)
+		// Write out allocated shadow mask tiles
+	#if HASH_TABLE
+		if (LinearThreadIndex < HASH_TABLE_SIZE)
 		{
-			// #sdl_todo: merge allocations through groupshared and do only one atomic per group
-			uint ShadowMaskTileIndex = 0;
-			InterlockedAdd(RWShadowMaskTileAllocator[0], 1, ShadowMaskTileIndex);
+			uint2 ShadowMaskScreenCoord = DownsampledTileCoord * SHADOW_MASK_TILE_SIZE;
 
-			if (ShadowMaskTileIndex < MaxShadowMaskTiles)
+			uint LightIndex = SharedHashTable[LinearThreadIndex];
+
+			if (LightIndex != INVALID_LIGHT_SCENE_ID)
 			{
-				WriteShadowMaskTileToPageTable(LightIndex, ShadowMaskScreenCoord, ShadowMaskTileIndex);
+				// #sdl_todo: merge allocations through groupshared and do only one atomic per group
+				uint ShadowMaskTileIndex = 0;
+				InterlockedAdd(RWShadowMaskTileAllocator[0], 1, ShadowMaskTileIndex);
+
+				if (ShadowMaskTileIndex < MaxShadowMaskTiles)
+				{
+					WriteShadowMaskTileToPageTable(LightIndex, ShadowMaskScreenCoord, ShadowMaskTileIndex);
+				}
 			}
 		}
-	}
-
-#else
-
-	// #sdl_todo: process part of the list per tread
-	if (LinearThreadIndex < LIGHT_SCENE_ID_MASK_SIZE)
-	{
-		const uint LightIndex = LinearThreadIndex;
-		const uint BitMask = 1u << LightIndex;
+	#else
+		// #sdl_todo: process part of the list per tread
+		if (LinearThreadIndex < LIGHT_SCENE_ID_MASK_SIZE)
+		{
+			const uint LightIndex = LinearThreadIndex;
+			const uint BitMask = 1u << LightIndex;
 
-		uint2 ShadowMaskScreenCoord = GroupId.xy * SHADOW_MASK_TILE_SIZE;
+			uint2 ShadowMaskScreenCoord = DownsampledTileCoord * SHADOW_MASK_TILE_SIZE;
 	
-		uint LightMask = SharedLightMask[LinearThreadIndex];
-		while (LightMask != 0)
-		{
-			const uint NextBitIndex = firstbitlow(LightMask);
-			const uint NextBitMask = 1u << NextBitIndex;
-			LightMask ^= NextBitMask;
+			uint LightMask = SharedLightMask[LinearThreadIndex];
+			while (LightMask != 0)
+			{
+				const uint NextBitIndex = firstbitlow(LightMask);
+				const uint NextBitMask = 1u << NextBitIndex;
+				LightMask ^= NextBitMask;
 
-			uint LightIndex = LinearThreadIndex + NextBitIndex * LIGHT_SCENE_ID_MASK_SIZE;
+				uint LightIndex = LinearThreadIndex + NextBitIndex * LIGHT_SCENE_ID_MASK_SIZE;
 
-			// #sdl_todo: merge allocations through groupshared and do only one atomic per group
-			uint ShadowMaskTileIndex = 0;
-			InterlockedAdd(RWShadowMaskTileAllocator[0], 1, ShadowMaskTileIndex);
+				// #sdl_todo: merge allocations through groupshared and do only one atomic per group
+				uint ShadowMaskTileIndex = 0;
+				InterlockedAdd(RWShadowMaskTileAllocator[0], 1, ShadowMaskTileIndex);
 
-			if (ShadowMaskTileIndex < MaxShadowMaskTiles)
-			{
-				WriteShadowMaskTileToPageTable(LightIndex, ShadowMaskScreenCoord, ShadowMaskTileIndex);
+				if (ShadowMaskTileIndex < MaxShadowMaskTiles)
+				{
+					WriteShadowMaskTileToPageTable(LightIndex, ShadowMaskScreenCoord, ShadowMaskTileIndex);
+				}
 			}
 		}
+	#endif
 	}
-
-#endif
 }
\ No newline at end of file
diff --git a/Engine/Shaders/Private/StochasticShadows/StochasticShadowsShading.usf b/Engine/Shaders/Private/StochasticShadows/StochasticShadowsShading.usf
--- a/Engine/Shaders/Private/StochasticShadows/StochasticShadowsShading.usf
+++ b/Engine/Shaders/Private/StochasticShadows/StochasticShadowsShading.usf
@@ -21,132 +21,145 @@ StructuredBuffer<uint> ShadingTileAllocator;
 uint MaxShadowMaskTiles;
 uint MaxShadingTiles;
 
+StructuredBuffer<uint> TileAllocator;
+StructuredBuffer<uint> TileData;
+uint TileDataStride;
+
 /**
- * Apply all lights per pixel using shadow data from the shadow mask.
+ * Apply all lights per pixel to affected tiles using shadow data from the shadow mask.
  */
 [numthreads(THREADGROUP_SIZE, THREADGROUP_SIZE, 1)]
-void ShadeLightSamplesWithShadowMaskCS(
+void ShadeLightSamplesCS(
 	uint3 GroupId : SV_GroupID,
 	uint3 GroupThreadId : SV_GroupThreadID,
 	uint3 DispatchThreadId : SV_DispatchThreadID)
 {
-	uint2 ScreenCoord = DispatchThreadId.xy;
-
-	if (all(ScreenCoord < View.ViewRectMinAndSize.zw))
+	uint TileIndex = GroupId.x;
+	if (TileIndex < TileAllocator[TILE_TYPE])
 	{
-#if DEBUG_MODE
-		int2 DebugScreenCoord = View.CursorPosition.x >= 0 ? View.CursorPosition * View.ViewResolutionFraction : -1;
-		bool bDebug = all(ScreenCoord == DebugScreenCoord);
-		FShaderPrintContext Context = InitShaderPrintContext(true, float2(0.55, 0.05));
-#endif
-
-		float2 ScreenUV = (ScreenCoord + 0.5f) * View.BufferSizeAndInvSize.zw;
-		FGBufferData GBuffer = GetGBufferData(ScreenUV);
+		uint2 TileCoord = UnpackTile(TileData[TileIndex + TILE_TYPE * TileDataStride]);
+		uint2 ScreenCoord = TileCoord * TILE_SIZE + GroupThreadId.xy;
 
-		if (GBuffer.Depth > 0.0f)
+		if (all(ScreenCoord < View.ViewRectMinAndSize.zw))
 		{
-			uint2 ShadingTileCoord = GroupId.xy;
-			uint2 CoordInShadingTile = GroupThreadId.xy;
-
-			float3 TranslatedWorldPosition = GetTranslatedWorldPositionFromScreenUV(ScreenUV, GBuffer.Depth);		
-			float3 CameraVector = normalize(TranslatedWorldPosition - View.TranslatedWorldCameraOrigin);
-
-			float4 Radiance = 0.0f;
-			uint NumShadingTilesInCell = ShadingTileGridAllocator.Load(int3(ShadingTileCoord, 0));
+	#if DEBUG_MODE
+			int2 DebugScreenCoord = View.CursorPosition.x >= 0 ? View.CursorPosition * View.ViewResolutionFraction : -1;
+			bool bDebug = all(ScreenCoord == DebugScreenCoord);
+			FShaderPrintContext Context = InitShaderPrintContext(true, float2(0.55, 0.05));
+	#endif
 
-			#if DEBUG_MODE
-			if (bDebug)
-			{
-				Print(Context, TEXT("ShadeSamplesCS"));
-				Newline(Context);
-				Print(Context, TEXT("ScreenCoord: "));
-				Print(Context, ScreenCoord.x);
-				Print(Context, ScreenCoord.y);
-				Newline(Context);
-				Print(Context, TEXT("NumShadingTilesInCell: "));
-				Print(Context, NumShadingTilesInCell);
-				Newline(Context);
-				Print(Context, TEXT("ShadowMaskTiles: "));
-				Print(Context, MaxShadowMaskTiles);
-				Print(Context, ShadowMaskTileAllocator[0]);
-				Print(Context, float(ShadowMaskTileAllocator[0] * 8 * 8) / (View.ViewSizeAndInvSize.x * View.ViewSizeAndInvSize.y));
-				Newline(Context);
-				Print(Context, TEXT("ShadingTiles: "));
-				Print(Context, MaxShadingTiles);
-				Print(Context, ShadingTileAllocator[0]);
-				Print(Context, ShadingTileAllocator[0] / float(MaxShadingTiles));
-				Newline(Context);
-				Print(Context, TEXT("LightId   | Shadow    | Radiance"));
-			}
-			#endif
+			float2 ScreenUV = (ScreenCoord + 0.5f) * View.BufferSizeAndInvSize.zw;
+			FGBufferData GBuffer = GetGBufferData(ScreenUV);
 
-			for (uint IndexInCell = 0; IndexInCell < NumShadingTilesInCell; ++IndexInCell)
+			if (GBuffer.Depth > 0.0f)
 			{
-				const uint PackedShadingTile = ShadingTileGrid[GetShadingTileGridOffset(ShadingTileCoord, IndexInCell)];
-				const FShadingTile ShadingTile = UnpackShadingTile(PackedShadingTile);
-				FDeferredLightData LightData = ConvertToDeferredLight(GetLightSceneData(ShadingTile.LightSceneId));
+				uint2 ShadingTileCoord = TileCoord;
+				uint2 CoordInShadingTile = GroupThreadId.xy;
 
-				uint2 ShadingTileAtlasCoord = GetShadingTileAtlasCoord(ShadingTile.ShadingTileIndex, CoordInShadingTile);
+				float3 TranslatedWorldPosition = GetTranslatedWorldPositionFromScreenUV(ScreenUV, GBuffer.Depth);		
+				float3 CameraVector = normalize(TranslatedWorldPosition - View.TranslatedWorldCameraOrigin);
 
-				float ShadowFactor = ShadingTileAtlas[ShadingTileAtlasCoord];
+				float4 Irradiance = 0.0f;
+				uint NumShadingTilesInCell = ShadingTileGridAllocator.Load(int3(ShadingTileCoord, 0));
 
 				#if DEBUG_MODE
 				if (bDebug)
 				{
+					Print(Context, TEXT("ShadeSamplesCS"));
+					Newline(Context);
+					Print(Context, TEXT("ScreenCoord: "));
+					Print(Context, ScreenCoord.x);
+					Print(Context, ScreenCoord.y);
+					Newline(Context);
+					Print(Context, TEXT("NumShadingTilesInCell: "));
+					Print(Context, NumShadingTilesInCell);
+					Newline(Context);
+					Print(Context, TEXT("ShadowMaskTiles: "));
+					Print(Context, MaxShadowMaskTiles);
+					Print(Context, ShadowMaskTileAllocator[0]);
+					Print(Context, float(ShadowMaskTileAllocator[0] * 8 * 8) / (View.ViewSizeAndInvSize.x * View.ViewSizeAndInvSize.y));
 					Newline(Context);
-					Print(Context, ShadingTile.LightSceneId);
-					Print(Context, ShadowFactor);
+					Print(Context, TEXT("ShadingTiles: "));
+					Print(Context, MaxShadingTiles);
+					Print(Context, ShadingTileAllocator[0]);
+					Print(Context, ShadingTileAllocator[0] / float(MaxShadingTiles));
+					Newline(Context);
+					Print(Context, TEXT("LightId   | Shadow    | Irradiance"));
 				}
 				#endif
 
-				#if DEBUG_SHADOW_MASK
+				for (uint IndexInCell = 0; IndexInCell < NumShadingTilesInCell; ++IndexInCell)
 				{
-					Radiance += float4(normalize(LightData.Color) * ShadowFactor, 0.0f);
-				}
-				#endif
+					const uint PackedShadingTile = ShadingTileGrid[GetShadingTileGridOffset(ShadingTileCoord, IndexInCell)];
+					const FShadingTile ShadingTile = UnpackShadingTile(PackedShadingTile);
+					FDeferredLightData LightData = ConvertToDeferredLight(GetLightSceneData(ShadingTile.LightSceneId));
 
-				if (ShadowFactor > 0.0f)
-				{
-					GBuffer.Roughness = max(GBuffer.Roughness, View.MinRoughness);
-
-					float4 LightAttenuation = 1.0f;
-					float Dither = 0.5f;
-					float SurfaceShadow = 1;
-					float AmbientOcclusion = ShadowFactor;
-
-					// #sdl_todo: tile classification and better integration
-					GBuffer.ShadingModelID = SHADINGMODELID_DEFAULT_LIT;
-					LightData.ShadowedBits = 0;
-
-					float4 LightRadiance = GetDynamicLighting(
-						TranslatedWorldPosition,
-						CameraVector,
-						GBuffer, 
-						AmbientOcclusion,
-						GBuffer.ShadingModelID,
-						LightData,
-						LightAttenuation,
-						Dither,
-						ScreenCoord, 
-						SurfaceShadow);
+					uint2 ShadingTileAtlasCoord = GetShadingTileAtlasCoord(ShadingTile.ShadingTileIndex, CoordInShadingTile);
+
+					float ShadowFactor = ShadingTileAtlas[ShadingTileAtlasCoord];
 
 					#if DEBUG_MODE
 					if (bDebug)
 					{
-						Print(Context, LightRadiance);
+						Newline(Context);
+						Print(Context, ShadingTile.LightSceneId);
+						Print(Context, ShadowFactor);
 					}
 					#endif
 
-					#if !DEBUG_SHADOW_MASK
+					#if DEBUG_SHADOW_MASK
 					{
-						Radiance += LightRadiance * View.PreExposure;
+						Irradiance += float4(normalize(LightData.Color) * ShadowFactor, 0.0f);
 					}
 					#endif
+
+					if (ShadowFactor > 0.0f)
+					{
+						GBuffer.Roughness = max(GBuffer.Roughness, View.MinRoughness);
+
+						float4 LightAttenuation = 1.0f;
+						float Dither = 0.5f;
+						float SurfaceShadow = 1;
+						float AmbientOcclusion = ShadowFactor;
+
+						// #sdl_todo: Move more shading models to TILE_MODE_SIMPLE_SHADING
+						#if TILE_TYPE == TILE_MODE_SIMPLE_SHADING
+						{
+							GBuffer.ShadingModelID = SHADINGMODELID_DEFAULT_LIT;
+						}
+						#endif
+						LightData.ShadowedBits = 0;
+
+						float4 LightRadiance = GetDynamicLighting(
+							TranslatedWorldPosition,
+							CameraVector,
+							GBuffer, 
+							AmbientOcclusion,
+							GBuffer.ShadingModelID,
+							LightData,
+							LightAttenuation,
+							Dither,
+							ScreenCoord, 
+							SurfaceShadow);
+
+						#if DEBUG_MODE
+						if (bDebug)
+						{
+							Print(Context, LightRadiance);
+						}
+						#endif
+
+						#if !DEBUG_SHADOW_MASK
+						{
+							Irradiance += LightRadiance * View.PreExposure;
+						}
+						#endif
+					}
 				}
-			}
 
-			// #sdl_todo: replace with a pixel shader?
-			RWSceneColor[ScreenCoord] = RWSceneColor[ScreenCoord] + Radiance;
+				// #sdl_todo: replace with a pixel shader?
+				RWSceneColor[ScreenCoord] = RWSceneColor[ScreenCoord] + Irradiance;
+			}
 		}
 	}
 }
\ No newline at end of file
diff --git a/Engine/Shaders/Private/StochasticShadows/StochasticShadowsSpatial.usf b/Engine/Shaders/Private/StochasticShadows/StochasticShadowsSpatial.usf
--- a/Engine/Shaders/Private/StochasticShadows/StochasticShadowsSpatial.usf
+++ b/Engine/Shaders/Private/StochasticShadows/StochasticShadowsSpatial.usf
@@ -8,7 +8,7 @@
 Texture2D<uint> ShadowMaskTileAtlas;
 Texture2D<uint> ShadowMaskPageTable;
 StructuredBuffer<uint> ShadowMaskHashTable;
-StructuredBuffer<uint2> ShadowMaskTileHeader;
+StructuredBuffer<uint2> ShadowMaskTileData;
 Texture2D<float> DownsampledSceneDepth;
 uint MaxShadingTiles;
 
@@ -58,7 +58,7 @@ void ShadowMaskSpatialPassCS(
 
 	uint LightSceneId = 0;
 	uint2 BaseShadowMaskScreenCoord = 0;	
-	UnpackShadowMaskTileHeader(ShadowMaskTileHeader[ShadowMaskTileIndex], LightSceneId, BaseShadowMaskScreenCoord);
+	UnpackShadowMaskTile(ShadowMaskTileData[ShadowMaskTileIndex], LightSceneId, BaseShadowMaskScreenCoord);
 
 	if (all(GroupThreadId == 0))
 	{
diff --git a/Engine/Shaders/Private/StochasticShadows/StochasticShadowsTracing.ush b/Engine/Shaders/Private/StochasticShadows/StochasticShadowsTracing.ush
--- a/Engine/Shaders/Private/StochasticShadows/StochasticShadowsTracing.ush
+++ b/Engine/Shaders/Private/StochasticShadows/StochasticShadowsTracing.ush
@@ -6,6 +6,10 @@
 #include "../RayTracing/RayTracingSphereLight.ush"
 #include "../RayTracing/RayTracingRectLight.ush"
 
+#ifndef NUM_SAMPLES_PER_PIXEL_X
+	#error NUM_SAMPLES_PER_PIXEL_X must be defined
+#endif
+
 uint2 SampleCoordToDownsampledScreenCoord(uint2 SampleCoord)
 {
 	return SampleCoord / uint2(NUM_SAMPLES_PER_PIXEL_X, NUM_SAMPLES_PER_PIXEL_Y);
diff --git a/Engine/Source/Runtime/Renderer/Private/StochasticShadows/StochasticShadows.cpp b/Engine/Source/Runtime/Renderer/Private/StochasticShadows/StochasticShadows.cpp
--- a/Engine/Source/Runtime/Renderer/Private/StochasticShadows/StochasticShadows.cpp
+++ b/Engine/Source/Runtime/Renderer/Private/StochasticShadows/StochasticShadows.cpp
@@ -160,6 +160,7 @@
 
 namespace StochasticShadows
 {
+	constexpr int32 TileSize = 8;
 	constexpr int32 ShadowMaskTileSize = 8;	// Stored downsampled
 	constexpr int32 MaxLightSceneIdXY = 16; // 16 * 16 = 256
 	constexpr int32 MaxShadingTilesPerGridCell = 16;
@@ -258,6 +259,14 @@
 		NumTraces      = 2 * sizeof(FRHIDispatchIndirectParameters),
 		MAX = 3
 	};
+
+	// Keep in sync with TILE_TYPE_* in shaders
+	enum class ETileType : uint8
+	{
+		SimpleShading = 0,
+		ComplexShading = 1,
+		MAX = 2
+	};
 }
 
 BEGIN_SHADER_PARAMETER_STRUCT(FStochasticShadowsParameters, )
@@ -280,10 +289,77 @@
 	SHADER_PARAMETER(FIntPoint, ShadingTileGridSize)
 	SHADER_PARAMETER(FVector2f, DownsampledBufferInvSize)
 	SHADER_PARAMETER(float, MinLightSampleWeight)
+	SHADER_PARAMETER(int32, TileDataStride)
+	SHADER_PARAMETER(int32, DownsampledTileDataStride)
 	SHADER_PARAMETER(int32, DebugMode)
 	SHADER_PARAMETER_RDG_TEXTURE(Texture2D<float>, DownsampledSceneDepth)
 END_SHADER_PARAMETER_STRUCT()
 
+class FTileClassificationCS : public FGlobalShader
+{
+	DECLARE_GLOBAL_SHADER(FTileClassificationCS)
+	SHADER_USE_PARAMETER_STRUCT(FTileClassificationCS, FGlobalShader)
+
+	BEGIN_SHADER_PARAMETER_STRUCT(FParameters, )
+		SHADER_PARAMETER_STRUCT_INCLUDE(FStochasticShadowsParameters, StochasticShadowsParameters)
+		SHADER_PARAMETER_RDG_BUFFER_UAV(RWStructuredBuffer<uint>, RWTileAllocator)
+		SHADER_PARAMETER_RDG_BUFFER_UAV(RWStructuredBuffer<uint>, RWTileData)
+	END_SHADER_PARAMETER_STRUCT()
+
+	class FDownsampledClassification : SHADER_PERMUTATION_BOOL("DOWNSAMPLED_CLASSIFICATION");
+	using FPermutationDomain = TShaderPermutationDomain<FDownsampledClassification>;
+
+	static int32 GetGroupSize()
+	{	
+		return 8;
+	}
+
+	static bool ShouldCompilePermutation(const FGlobalShaderPermutationParameters& Parameters)
+	{
+		return StochasticShadows::ShouldCompileShaders(Parameters);
+	}
+
+	static void ModifyCompilationEnvironment(const FGlobalShaderPermutationParameters& Parameters, FShaderCompilerEnvironment& OutEnvironment)
+	{
+		FGlobalShader::ModifyCompilationEnvironment(Parameters, OutEnvironment);
+		OutEnvironment.SetDefine(TEXT("THREADGROUP_SIZE"), GetGroupSize());
+	}
+};
+
+IMPLEMENT_GLOBAL_SHADER(FTileClassificationCS, "/Engine/Private/StochasticShadows/StochasticShadows.usf", "TileClassificationCS", SF_Compute);
+
+class FInitTileIndirectArgsCS : public FGlobalShader
+{
+	DECLARE_GLOBAL_SHADER(FInitTileIndirectArgsCS)
+	SHADER_USE_PARAMETER_STRUCT(FInitTileIndirectArgsCS, FGlobalShader);
+
+	BEGIN_SHADER_PARAMETER_STRUCT(FParameters, )
+		SHADER_PARAMETER_STRUCT_INCLUDE(FStochasticShadowsParameters, StochasticShadowsParameters)
+		SHADER_PARAMETER_RDG_BUFFER_UAV(RWBuffer<uint>, RWTileIndirectArgs)
+		SHADER_PARAMETER_RDG_BUFFER_UAV(RWBuffer<uint>, RWDownsampledTileIndirectArgs)
+		SHADER_PARAMETER_RDG_BUFFER_SRV(StructuredBuffer<uint>, TileAllocator)
+		SHADER_PARAMETER_RDG_BUFFER_SRV(StructuredBuffer<uint>, DownsampledTileAllocator)
+	END_SHADER_PARAMETER_STRUCT()
+
+	static bool ShouldCompilePermutation(const FGlobalShaderPermutationParameters& Parameters)
+	{
+		return StochasticShadows::ShouldCompileShaders(Parameters);
+	}
+
+	static void ModifyCompilationEnvironment(const FGlobalShaderPermutationParameters& Parameters, FShaderCompilerEnvironment& OutEnvironment)
+	{
+		FGlobalShader::ModifyCompilationEnvironment(Parameters, OutEnvironment);
+		OutEnvironment.SetDefine(TEXT("THREADGROUP_SIZE"), GetGroupSize());
+	}
+
+	static int32 GetGroupSize()
+	{
+		return 64;
+	}
+};
+
+IMPLEMENT_GLOBAL_SHADER(FInitTileIndirectArgsCS, "/Engine/Private/StochasticShadows/StochasticShadows.usf", "InitTileIndirectArgsCS", SF_Compute);
+
 class FCompactLightSampleTracesCS : public FGlobalShader
 {
 	DECLARE_GLOBAL_SHADER(FCompactLightSampleTracesCS)
@@ -395,11 +471,13 @@
 		SHADER_PARAMETER_RDG_TEXTURE_UAV(RWTexture2D<uint32>, RWShadowMaskHistoryScreenCoord00)
 		SHADER_PARAMETER_RDG_TEXTURE_UAV(RWTexture2D<float4>, RWShadowMaskHistoryWeights)
 		SHADER_PARAMETER_RDG_BUFFER_UAV(RWStructuredBuffer<uint>, RWShadowMaskTileAllocator)
-		SHADER_PARAMETER_RDG_BUFFER_UAV(RWStructuredBuffer<uint2>, RWShadowMaskTileHeader)
+		SHADER_PARAMETER_RDG_BUFFER_UAV(RWStructuredBuffer<uint2>, RWShadowMaskTileData)
 		SHADER_PARAMETER_RDG_TEXTURE_UAV(RWTexture2D<uint>, RWShadowMaskPageTable)
 		SHADER_PARAMETER_RDG_BUFFER_UAV(RWStructuredBuffer<uint>, RWShadowMaskHashTable)
 		SHADER_PARAMETER_RDG_TEXTURE_UAV(RWTexture2D<uint>, RWLightSamples)
 		SHADER_PARAMETER_RDG_UNIFORM_BUFFER(FForwardLightData, ForwardLightData)
+		SHADER_PARAMETER_RDG_BUFFER_SRV(StructuredBuffer<uint>, DownsampledTileAllocator)
+		SHADER_PARAMETER_RDG_BUFFER_SRV(StructuredBuffer<uint>, DownsampledTileData)
 		SHADER_PARAMETER_RDG_TEXTURE(Texture2D<uint>, ShadowMaskPageTableHistory)
 		SHADER_PARAMETER_RDG_BUFFER_SRV(StructuredBuffer<uint>, ShadowMaskHashTableHistory)
 		SHADER_PARAMETER_RDG_TEXTURE(Texture2D<uint>, ShadowMaskAtlasHistory)
@@ -409,13 +487,14 @@
 		SHADER_PARAMETER(FVector4f, HistoryScreenPositionScaleBias)
 	END_SHADER_PARAMETER_STRUCT()
 
+	class FTileType : SHADER_PERMUTATION_INT("TILE_TYPE", (int32)StochasticShadows::ETileType::MAX);
 	class FNumSamplesPerPixel : SHADER_PERMUTATION_SPARSE_INT("NUM_SAMPLES_PER_PIXEL", 1, 2, 4);
 	class FShadowMaskReprojectionWeights : SHADER_PERMUTATION_BOOL("SHADOW_MASK_REPROJECTION_WEIGHTS");
 	class FShadowFactorEstimate : SHADER_PERMUTATION_BOOL("SHADOW_FACTOR_ESTIMATE");
 	class FHashTable : SHADER_PERMUTATION_BOOL("HASH_TABLE");
 	class FCandidateLightMask : SHADER_PERMUTATION_BOOL("CANDIDATE_LIGHT_MASK");
 	class FDebugMode : SHADER_PERMUTATION_BOOL("DEBUG_MODE");
-	using FPermutationDomain = TShaderPermutationDomain<FNumSamplesPerPixel, FShadowMaskReprojectionWeights, FShadowFactorEstimate, FHashTable, FCandidateLightMask, FDebugMode>;
+	using FPermutationDomain = TShaderPermutationDomain<FTileType, FNumSamplesPerPixel, FShadowMaskReprojectionWeights, FShadowFactorEstimate, FHashTable, FCandidateLightMask, FDebugMode>;
 
 	static int32 GetGroupSize()
 	{	
@@ -692,7 +771,7 @@
 		RDG_BUFFER_ACCESS(IndirectArgs, ERHIAccess::IndirectArgs)
 		SHADER_PARAMETER_RDG_TEXTURE_UAV(RWTexture2D<uint>, RWShadowMaskTileAtlas)
 		SHADER_PARAMETER_RDG_TEXTURE(Texture2D<uint>, ShadowMaskPageTable)
-		SHADER_PARAMETER_RDG_BUFFER_SRV(StructuredBuffer<uint2>, ShadowMaskTileHeader)
+		SHADER_PARAMETER_RDG_BUFFER_SRV(StructuredBuffer<uint2>, ShadowMaskTileData)
 		SHADER_PARAMETER_RDG_TEXTURE(Texture2D<uint32>, ShadowMaskHistoryScreenCoord00)
 		SHADER_PARAMETER_RDG_TEXTURE(Texture2D<float4>, ShadowMaskHistoryWeights)
 		SHADER_PARAMETER_RDG_TEXTURE(Texture2D<uint>, ShadowMaskPageTableHistory)
@@ -728,7 +807,7 @@
 	}
 };
 
-IMPLEMENT_GLOBAL_SHADER(FCompositeShadowMaskTracesCS, "/Engine/Private/StochasticShadows/StochasticShadows.usf", "CompositeShadowMaskTracesCS", SF_Compute);
+IMPLEMENT_GLOBAL_SHADER(FCompositeShadowMaskTracesCS, "/Engine/Private/StochasticShadows/StochasticShadowsComposite.usf", "CompositeShadowMaskTracesCS", SF_Compute);
 
 class FShadowMaskSpatialPassCS : public FGlobalShader
 {
@@ -743,7 +822,7 @@
 		SHADER_PARAMETER_RDG_BUFFER_UAV(RWStructuredBuffer<uint>, RWShadingTileGrid)
 		SHADER_PARAMETER_RDG_TEXTURE_UAV(RWTexture2D<UNORM float>, RWShadingTileAtlas)
 		SHADER_PARAMETER_RDG_TEXTURE(Texture2D<float4>, ShadowMaskUpsampleWeights)
-		SHADER_PARAMETER_RDG_BUFFER_SRV(StructuredBuffer<uint2>, ShadowMaskTileHeader)
+		SHADER_PARAMETER_RDG_BUFFER_SRV(StructuredBuffer<uint2>, ShadowMaskTileData)
 		SHADER_PARAMETER_RDG_TEXTURE(Texture2D<uint>, ShadowMaskPageTable)
 		SHADER_PARAMETER_RDG_TEXTURE(Texture2D<uint>, ShadowMaskTileAtlas)
 		SHADER_PARAMETER_RDG_BUFFER_SRV(StructuredBuffer<uint>, ShadowMaskHashTable)
@@ -772,16 +851,18 @@
 
 IMPLEMENT_GLOBAL_SHADER(FShadowMaskSpatialPassCS, "/Engine/Private/StochasticShadows/StochasticShadowsSpatial.usf", "ShadowMaskSpatialPassCS", SF_Compute);
 
-class FShadeLightSamplesWithShadowMaskCS : public FGlobalShader
+class FShadeLightSamplesCS : public FGlobalShader
 {
-	DECLARE_GLOBAL_SHADER(FShadeLightSamplesWithShadowMaskCS)
-	SHADER_USE_PARAMETER_STRUCT(FShadeLightSamplesWithShadowMaskCS, FGlobalShader)
+	DECLARE_GLOBAL_SHADER(FShadeLightSamplesCS)
+	SHADER_USE_PARAMETER_STRUCT(FShadeLightSamplesCS, FGlobalShader)
 
 	BEGIN_SHADER_PARAMETER_STRUCT(FParameters, )
 		SHADER_PARAMETER_STRUCT_INCLUDE(FStochasticShadowsParameters, StochasticShadowsParameters)
 		SHADER_PARAMETER_RDG_TEXTURE_UAV(RWTexture2D<float4>, RWSceneColor)
 		SHADER_PARAMETER_RDG_TEXTURE(Texture2D<uint>, ShadowMaskPageTable)
 		SHADER_PARAMETER_RDG_TEXTURE(Texture2D<uint>, ShadowMaskTileAtlas)
+		SHADER_PARAMETER_RDG_BUFFER_SRV(StructuredBuffer<uint>, TileAllocator)
+		SHADER_PARAMETER_RDG_BUFFER_SRV(StructuredBuffer<uint>, TileData)
 		SHADER_PARAMETER_RDG_BUFFER_SRV(StructuredBuffer<uint>, ShadowMaskTileAllocator)
 		SHADER_PARAMETER_RDG_BUFFER_SRV(StructuredBuffer<uint>, ShadingTileAllocator)
 		SHADER_PARAMETER_RDG_TEXTURE(Texture2D<uint>, ShadingTileGridAllocator)
@@ -796,8 +877,9 @@
 		return 8;
 	}
 
+	class FTileType : SHADER_PERMUTATION_INT("TILE_TYPE", (int32)StochasticShadows::ETileType::MAX);
 	class FDebugMode : SHADER_PERMUTATION_BOOL("DEBUG_MODE");
-	using FPermutationDomain = TShaderPermutationDomain<FDebugMode>;
+	using FPermutationDomain = TShaderPermutationDomain<FTileType, FDebugMode>;
 
 	static bool ShouldCompilePermutation(const FGlobalShaderPermutationParameters& Parameters)
 	{
@@ -811,7 +893,7 @@
 	}
 };
 
-IMPLEMENT_GLOBAL_SHADER(FShadeLightSamplesWithShadowMaskCS, "/Engine/Private/StochasticShadows/StochasticShadowsShading.usf", "ShadeLightSamplesWithShadowMaskCS", SF_Compute);
+IMPLEMENT_GLOBAL_SHADER(FShadeLightSamplesCS, "/Engine/Private/StochasticShadows/StochasticShadowsShading.usf", "ShadeLightSamplesCS", SF_Compute);
 
 #if RHI_RAYTRACING
 void FDeferredShadingSceneRenderer::PrepareStochasticShadows(const FViewInfo& View, const FScene& Scene, TArray<FRHIRayTracingShader*>& OutRayGenShaders)
@@ -996,37 +1078,6 @@
 	FRDGTextureRef ShadowMaskTileAtlas = GraphBuilder.CreateTexture(
 		FRDGTextureDesc::Create2D(ShadowMaskTileAtlasSize, ShadowMaskTileAtlasFormat, FClearValueBinding::Black, TexCreate_ShaderResource | TexCreate_UAV),
 		TEXT("StochasticShadows.ShadowMaskTileAtlas"));
-
-	FStochasticShadowsParameters StochasticShadowsParameters;
-	{
-		StochasticShadowsParameters.ViewUniformBuffer = View.ViewUniformBuffer;
-		StochasticShadowsParameters.Scene = View.GetSceneUniforms().GetBuffer(GraphBuilder);
-		StochasticShadowsParameters.SceneTextures = GetSceneTextureParameters(GraphBuilder, SceneTextures.UniformBuffer);
-		StochasticShadowsParameters.SceneTexturesStruct = SceneTextures.UniformBuffer;
-		StochasticShadowsParameters.Substrate = Substrate::BindSubstrateGlobalUniformParameters(View);
-		StochasticShadowsParameters.BlueNoise = BlueNoiseUniformBuffer;
-		StochasticShadowsParameters.DownsampledViewSize = DownsampledViewSize;
-		StochasticShadowsParameters.SampleViewSize = SampleViewSize;
-		StochasticShadowsParameters.ShadowMaskViewSize = ShadowMaskViewSize;
-		StochasticShadowsParameters.StochasticShadowsStateFrameIndex = StochasticShadows::GetStateFrameIndex(View.ViewState);
-		StochasticShadowsParameters.DownsampledSceneDepth = DownsampledSceneDepth;
-		StochasticShadowsParameters.MaxShadowMaskTiles = MaxShadowMaskTiles;
-		StochasticShadowsParameters.MaxShadingTiles = (ShadingTileAtlasSize.X * ShadingTileAtlasSize.Y) / (StochasticShadows::ShadowMaskTileSize * StochasticShadows::ShadowMaskTileSize);
-		StochasticShadowsParameters.MaxShadingTilesPerGridCell = StochasticShadows::MaxShadingTilesPerGridCell;
-		StochasticShadowsParameters.ShadingTileGridSize = ShadingTileGridSize;
-		StochasticShadowsParameters.ShadowMaskHashTableIndexWrapMask = ShadowMaskHashTableSize - 1;
-		StochasticShadowsParameters.ShadowMaskPageTablePerLightSize = ShadowMaskPageTablePerLightSize;
-		StochasticShadowsParameters.DownsampledBufferInvSize = FVector2f(1.0f) / DownsampledBufferSize;
-		StochasticShadowsParameters.MinLightSampleWeight = CVarStochasticShadowsMinSampleWeight.GetValueOnRenderThread();
-		StochasticShadowsParameters.DebugMode = CVarStochasticShadowsDebug.GetValueOnRenderThread();
-
-		if (bDebug)
-		{
-			ShaderPrint::SetEnabled(true);
-			ShaderPrint::RequestSpaceForLines(1024);
-			ShaderPrint::SetParameters(GraphBuilder, View.ShaderPrintData, StochasticShadowsParameters.ShaderPrintUniformBuffer);
-		}
-	}
 	
 	bool bTemporal = CVarStochasticShadowsTemporal.GetValueOnRenderThread() != 0;
 	FRDGTextureRef ShadowMaskSceneDepthHistory = nullptr;
@@ -1082,6 +1133,118 @@
 		}
 	}
 
+	const FIntPoint ViewSizeInTiles = FIntPoint::DivideAndRoundUp(View.ViewRect.Size(), StochasticShadows::TileSize);
+	const int32 TileDataStride = ViewSizeInTiles.X * ViewSizeInTiles.Y;
+
+	const FIntPoint DownsampledViewSizeInTiles = FIntPoint::DivideAndRoundUp(DownsampledViewSize, StochasticShadows::TileSize);
+	const int32 DownsampledTileDataStride = DownsampledViewSizeInTiles.X * DownsampledViewSizeInTiles.Y;
+
+	FStochasticShadowsParameters StochasticShadowsParameters;
+	{
+		StochasticShadowsParameters.ViewUniformBuffer = View.ViewUniformBuffer;
+		StochasticShadowsParameters.Scene = View.GetSceneUniforms().GetBuffer(GraphBuilder);
+		StochasticShadowsParameters.SceneTextures = GetSceneTextureParameters(GraphBuilder, SceneTextures.UniformBuffer);
+		StochasticShadowsParameters.SceneTexturesStruct = SceneTextures.UniformBuffer;
+		StochasticShadowsParameters.Substrate = Substrate::BindSubstrateGlobalUniformParameters(View);
+		StochasticShadowsParameters.BlueNoise = BlueNoiseUniformBuffer;
+		StochasticShadowsParameters.DownsampledViewSize = DownsampledViewSize;
+		StochasticShadowsParameters.SampleViewSize = SampleViewSize;
+		StochasticShadowsParameters.ShadowMaskViewSize = ShadowMaskViewSize;
+		StochasticShadowsParameters.StochasticShadowsStateFrameIndex = StochasticShadows::GetStateFrameIndex(View.ViewState);
+		StochasticShadowsParameters.DownsampledSceneDepth = DownsampledSceneDepth;
+		StochasticShadowsParameters.MaxShadowMaskTiles = MaxShadowMaskTiles;
+		StochasticShadowsParameters.MaxShadingTiles = (ShadingTileAtlasSize.X * ShadingTileAtlasSize.Y) / (StochasticShadows::ShadowMaskTileSize * StochasticShadows::ShadowMaskTileSize);
+		StochasticShadowsParameters.MaxShadingTilesPerGridCell = StochasticShadows::MaxShadingTilesPerGridCell;
+		StochasticShadowsParameters.ShadingTileGridSize = ShadingTileGridSize;
+		StochasticShadowsParameters.ShadowMaskHashTableIndexWrapMask = ShadowMaskHashTableSize - 1;
+		StochasticShadowsParameters.ShadowMaskPageTablePerLightSize = ShadowMaskPageTablePerLightSize;
+		StochasticShadowsParameters.DownsampledBufferInvSize = FVector2f(1.0f) / DownsampledBufferSize;
+		StochasticShadowsParameters.MinLightSampleWeight = CVarStochasticShadowsMinSampleWeight.GetValueOnRenderThread();
+		StochasticShadowsParameters.TileDataStride = TileDataStride;
+		StochasticShadowsParameters.DownsampledTileDataStride = DownsampledTileDataStride;
+		StochasticShadowsParameters.DebugMode = CVarStochasticShadowsDebug.GetValueOnRenderThread();
+
+		if (bDebug)
+		{
+			ShaderPrint::SetEnabled(true);
+			ShaderPrint::RequestSpaceForLines(1024);
+			ShaderPrint::SetParameters(GraphBuilder, View.ShaderPrintData, StochasticShadowsParameters.ShaderPrintUniformBuffer);
+		}
+	}
+
+	FRDGBufferRef TileAllocator = GraphBuilder.CreateBuffer(FRDGBufferDesc::CreateStructuredDesc(sizeof(uint32), (int32)StochasticShadows::ETileType::MAX), TEXT("StochasticShadows.TileAllocator"));
+	FRDGBufferRef TileData = GraphBuilder.CreateBuffer(FRDGBufferDesc::CreateStructuredDesc(sizeof(uint32), TileDataStride * (int32)StochasticShadows::ETileType::MAX), TEXT("StochasticShadows.TileData"));
+	AddClearUAVPass(GraphBuilder, GraphBuilder.CreateUAV(TileAllocator), 0);
+
+	FRDGBufferRef DownsampledTileAllocator = GraphBuilder.CreateBuffer(FRDGBufferDesc::CreateStructuredDesc(sizeof(uint32), (int32)StochasticShadows::ETileType::MAX), TEXT("StochasticShadows.DownsampledTileAllocator"));
+	FRDGBufferRef DownsampledTileData = GraphBuilder.CreateBuffer(FRDGBufferDesc::CreateStructuredDesc(sizeof(uint32), DownsampledTileDataStride * (int32)StochasticShadows::ETileType::MAX), TEXT("StochasticShadows.DownsampledTileData"));
+	AddClearUAVPass(GraphBuilder, GraphBuilder.CreateUAV(DownsampledTileAllocator), 0);
+
+	// Run tile classification to generate tiles for the subsequent passes
+	{
+		{
+			FTileClassificationCS::FParameters* PassParameters = GraphBuilder.AllocParameters<FTileClassificationCS::FParameters>();
+			PassParameters->StochasticShadowsParameters = StochasticShadowsParameters;
+			PassParameters->RWTileAllocator = GraphBuilder.CreateUAV(TileAllocator);
+			PassParameters->RWTileData = GraphBuilder.CreateUAV(TileData);
+
+			FTileClassificationCS::FPermutationDomain PermutationVector;
+			PermutationVector.Set<FTileClassificationCS::FDownsampledClassification>(false);
+			auto ComputeShader = View.ShaderMap->GetShader<FTileClassificationCS>(PermutationVector);
+
+			const FIntVector GroupCount = FComputeShaderUtils::GetGroupCount(View.ViewRect.Size(), FTileClassificationCS::GetGroupSize());
+
+			FComputeShaderUtils::AddPass(
+				GraphBuilder,
+				RDG_EVENT_NAME("TileClassification %dx%d", View.ViewRect.Size().X, View.ViewRect.Size().Y),
+				ComputeShader,
+				PassParameters,
+				GroupCount);
+		}
+
+		{
+			FTileClassificationCS::FParameters* PassParameters = GraphBuilder.AllocParameters<FTileClassificationCS::FParameters>();
+			PassParameters->StochasticShadowsParameters = StochasticShadowsParameters;
+			PassParameters->RWTileAllocator = GraphBuilder.CreateUAV(DownsampledTileAllocator);
+			PassParameters->RWTileData = GraphBuilder.CreateUAV(DownsampledTileData);
+
+			FTileClassificationCS::FPermutationDomain PermutationVector;
+			PermutationVector.Set<FTileClassificationCS::FDownsampledClassification>(true);
+			auto ComputeShader = View.ShaderMap->GetShader<FTileClassificationCS>(PermutationVector);
+
+			const FIntVector GroupCount = FComputeShaderUtils::GetGroupCount(View.ViewRect.Size(), FTileClassificationCS::GetGroupSize());
+
+			FComputeShaderUtils::AddPass(
+				GraphBuilder,
+				RDG_EVENT_NAME("DownsampledTileClassification %dx%d", DownsampledViewSize.X, DownsampledViewSize.Y),
+				ComputeShader,
+				PassParameters,
+				GroupCount);
+		}
+	}
+
+	FRDGBufferRef TileIndirectArgs = GraphBuilder.CreateBuffer(FRDGBufferDesc::CreateIndirectDesc<FRHIDispatchIndirectParameters>((int32)StochasticShadows::ETileType::MAX), TEXT("StochasticShadows.TileIndirectArgs"));
+	FRDGBufferRef DownsampledTileIndirectArgs = GraphBuilder.CreateBuffer(FRDGBufferDesc::CreateIndirectDesc<FRHIDispatchIndirectParameters>((int32)StochasticShadows::ETileType::MAX), TEXT("StochasticShadows.DownsampledTileIndirectArgs"));
+
+	// Setup indirect args for classified tiles
+	{
+		FInitTileIndirectArgsCS::FParameters* PassParameters = GraphBuilder.AllocParameters<FInitTileIndirectArgsCS::FParameters>();
+		PassParameters->StochasticShadowsParameters = StochasticShadowsParameters;
+		PassParameters->RWTileIndirectArgs = GraphBuilder.CreateUAV(TileIndirectArgs);
+		PassParameters->RWDownsampledTileIndirectArgs = GraphBuilder.CreateUAV(DownsampledTileIndirectArgs);
+		PassParameters->TileAllocator = GraphBuilder.CreateSRV(TileAllocator);
+		PassParameters->DownsampledTileAllocator = GraphBuilder.CreateSRV(DownsampledTileAllocator);
+
+		auto ComputeShader = View.ShaderMap->GetShader<FInitTileIndirectArgsCS>();
+
+		FComputeShaderUtils::AddPass(
+			GraphBuilder,
+			RDG_EVENT_NAME("InitTileIndirectArgs"),
+			ComputeShader,
+			PassParameters,
+			FIntVector(1, 1, 1));
+	}
+
 	FRDGTextureRef ShadowMaskHistoryWeights = GraphBuilder.CreateTexture(
 		FRDGTextureDesc::Create2D(ShadowMaskBufferSize, PF_FloatRGBA, FClearValueBinding::Black, TexCreate_ShaderResource | TexCreate_UAV),
 		TEXT("StochasticShadows.ShadowMaskHistoryWeights"));
@@ -1090,49 +1253,63 @@
 		FRDGTextureDesc::Create2D(ShadowMaskBufferSize, PF_R16G16_UINT, FClearValueBinding::Black, TexCreate_ShaderResource | TexCreate_UAV),
 		TEXT("StochasticShadows.ShadowMaskHistoryScreenCoord00"));
 
-	FRDGBufferRef ShadowMaskTileAllocator = GraphBuilder.CreateBuffer(FRDGBufferDesc::CreateStructuredDesc(sizeof(uint32), 1), TEXT("StochasticShadowsParameters.ShadowMaskTileAllocator"));
-	FRDGBufferRef ShadowMaskTileHeader = GraphBuilder.CreateBuffer(FRDGBufferDesc::CreateStructuredDesc(2 * sizeof(uint32), MaxShadowMaskTiles), TEXT("StochasticShadowsParameters.ShadowMaskTileHeader"));
+	FRDGBufferRef ShadowMaskTileAllocator = GraphBuilder.CreateBuffer(FRDGBufferDesc::CreateStructuredDesc(sizeof(uint32), 1), TEXT("StochasticShadows.ShadowMaskTileAllocator"));
+	FRDGBufferRef ShadowMaskTileData = GraphBuilder.CreateBuffer(FRDGBufferDesc::CreateStructuredDesc(2 * sizeof(uint32), MaxShadowMaskTiles), TEXT("StochasticShadows.ShadowMaskTileData"));
 	AddClearUAVPass(GraphBuilder, GraphBuilder.CreateUAV(ShadowMaskTileAllocator), 0);
 	
 	// Generate new candidate light samples
 	{
-		FGenerateSamplesCS::FParameters* PassParameters = GraphBuilder.AllocParameters<FGenerateSamplesCS::FParameters>();
-		PassParameters->StochasticShadowsParameters = StochasticShadowsParameters;
-		PassParameters->RWDownsampledSceneDepth = GraphBuilder.CreateUAV(DownsampledSceneDepth);
-		PassParameters->RWShadowMaskHistoryScreenCoord00 = GraphBuilder.CreateUAV(ShadowMaskHistoryScreenCoord00);
-		PassParameters->RWShadowMaskHistoryWeights = GraphBuilder.CreateUAV(ShadowMaskHistoryWeights);
-		PassParameters->RWShadowMaskTileAllocator = GraphBuilder.CreateUAV(ShadowMaskTileAllocator);
-		PassParameters->RWShadowMaskTileHeader = GraphBuilder.CreateUAV(ShadowMaskTileHeader);
-		PassParameters->RWShadowMaskPageTable = ShadowMaskPageTable ? GraphBuilder.CreateUAV(ShadowMaskPageTable) : nullptr;
-		PassParameters->RWShadowMaskHashTable = ShadowMaskHashTable ? GraphBuilder.CreateUAV(ShadowMaskHashTable) : nullptr;
-		PassParameters->RWLightSamples = GraphBuilder.CreateUAV(LightSamples);
-		PassParameters->ForwardLightData = View.ForwardLightingResources.ForwardLightUniformBuffer;
-		PassParameters->ShadowMaskPageTableHistory = ShadowMaskPageTableHistory;
-		PassParameters->ShadowMaskHashTableHistory = ShadowMaskHashTableHistory ? GraphBuilder.CreateSRV(ShadowMaskHashTableHistory) : nullptr;
-		PassParameters->ShadowMaskAtlasHistory = ShadowMaskAtlasHistory;
-		PassParameters->ShadowMaskTileAtlas = ShadowMaskTileAtlas;
-		PassParameters->ShadowMaskSceneDepthHistory = ShadowMaskSceneDepthHistory;
-		PassParameters->HistoryScreenPositionScaleBias = HistoryScreenPositionScaleBias;
-		PassParameters->HistoryUVMinMax = HistoryUVMinMax;
+		FRDGTextureUAVRef DownsampledSceneDepthUAV = GraphBuilder.CreateUAV(DownsampledSceneDepth, ERDGUnorderedAccessViewFlags::SkipBarrier);
+		FRDGTextureUAVRef ShadowMaskHistoryScreenCoord00UAV = GraphBuilder.CreateUAV(ShadowMaskHistoryScreenCoord00, ERDGUnorderedAccessViewFlags::SkipBarrier);
+		FRDGTextureUAVRef ShadowMaskHistoryWeightsUAV = GraphBuilder.CreateUAV(ShadowMaskHistoryWeights, ERDGUnorderedAccessViewFlags::SkipBarrier);
+		FRDGBufferUAVRef ShadowMaskTileAllocatorUAV = GraphBuilder.CreateUAV(ShadowMaskTileAllocator, ERDGUnorderedAccessViewFlags::SkipBarrier);
+		FRDGBufferUAVRef ShadowMaskTileDataUAV = GraphBuilder.CreateUAV(ShadowMaskTileData, ERDGUnorderedAccessViewFlags::SkipBarrier);
+		FRDGTextureUAVRef ShadowMaskPageTableUAV = ShadowMaskPageTable ? GraphBuilder.CreateUAV(ShadowMaskPageTable, ERDGUnorderedAccessViewFlags::SkipBarrier) : nullptr;
+		FRDGBufferUAVRef ShadowMaskHashTableUAV = ShadowMaskHashTable ? GraphBuilder.CreateUAV(ShadowMaskHashTable, ERDGUnorderedAccessViewFlags::SkipBarrier) : nullptr;
+		FRDGTextureUAVRef LightSamplesUAV = GraphBuilder.CreateUAV(LightSamples, ERDGUnorderedAccessViewFlags::SkipBarrier);
 
-		FGenerateSamplesCS::FPermutationDomain PermutationVector;
-		PermutationVector.Set<FGenerateSamplesCS::FNumSamplesPerPixel>(NumSamplesPerPixel1d);
-		PermutationVector.Set<FGenerateSamplesCS::FShadowMaskReprojectionWeights>(bTemporal);
-		PermutationVector.Set<FGenerateSamplesCS::FShadowFactorEstimate>(bTemporal && CVarStochasticShadowsGuiding.GetValueOnRenderThread() != 0);
-		PermutationVector.Set<FGenerateSamplesCS::FDebugMode>(bDebug);
-		PermutationVector.Set<FGenerateSamplesCS::FHashTable>(bUseHashTable);
-		PermutationVector.Set<FGenerateSamplesCS::FCandidateLightMask>(ShadowMaskAtlasHistory && CVarStochasticShadowsCandidateLightMask.GetValueOnRenderThread() != 0);
-		PermutationVector = FGenerateSamplesCS::RemapPermutation(PermutationVector);
-		auto ComputeShader = View.ShaderMap->GetShader<FGenerateSamplesCS>(PermutationVector);
+		for (int32 TileType = 0; TileType < (int32)StochasticShadows::ETileType::MAX; ++TileType)
+		{
+			FGenerateSamplesCS::FParameters* PassParameters = GraphBuilder.AllocParameters<FGenerateSamplesCS::FParameters>();
+			PassParameters->StochasticShadowsParameters = StochasticShadowsParameters;
+			PassParameters->RWDownsampledSceneDepth = DownsampledSceneDepthUAV;
+			PassParameters->RWShadowMaskHistoryScreenCoord00 = ShadowMaskHistoryScreenCoord00UAV;
+			PassParameters->RWShadowMaskHistoryWeights = ShadowMaskHistoryWeightsUAV;
+			PassParameters->RWShadowMaskTileAllocator = ShadowMaskTileAllocatorUAV;
+			PassParameters->RWShadowMaskTileData = ShadowMaskTileDataUAV;
+			PassParameters->RWShadowMaskPageTable = ShadowMaskPageTableUAV;
+			PassParameters->RWShadowMaskHashTable = ShadowMaskHashTableUAV;
+			PassParameters->RWLightSamples = LightSamplesUAV;
+			PassParameters->DownsampledTileAllocator = GraphBuilder.CreateSRV(DownsampledTileAllocator);
+			PassParameters->DownsampledTileData = GraphBuilder.CreateSRV(DownsampledTileData);
+			PassParameters->ForwardLightData = View.ForwardLightingResources.ForwardLightUniformBuffer;
+			PassParameters->ShadowMaskPageTableHistory = ShadowMaskPageTableHistory;
+			PassParameters->ShadowMaskHashTableHistory = ShadowMaskHashTableHistory ? GraphBuilder.CreateSRV(ShadowMaskHashTableHistory) : nullptr;
+			PassParameters->ShadowMaskAtlasHistory = ShadowMaskAtlasHistory;
+			PassParameters->ShadowMaskTileAtlas = ShadowMaskTileAtlas;
+			PassParameters->ShadowMaskSceneDepthHistory = ShadowMaskSceneDepthHistory;
+			PassParameters->HistoryScreenPositionScaleBias = HistoryScreenPositionScaleBias;
+			PassParameters->HistoryUVMinMax = HistoryUVMinMax;
 
-		const FIntVector GroupCount = FComputeShaderUtils::GetGroupCount(DownsampledViewSize, FGenerateSamplesCS::GetGroupSize());
+			FGenerateSamplesCS::FPermutationDomain PermutationVector;
+			PermutationVector.Set<FGenerateSamplesCS::FTileType>(TileType);
+			PermutationVector.Set<FGenerateSamplesCS::FNumSamplesPerPixel>(NumSamplesPerPixel1d);
+			PermutationVector.Set<FGenerateSamplesCS::FShadowMaskReprojectionWeights>(bTemporal);
+			PermutationVector.Set<FGenerateSamplesCS::FShadowFactorEstimate>(bTemporal && CVarStochasticShadowsGuiding.GetValueOnRenderThread() != 0);
+			PermutationVector.Set<FGenerateSamplesCS::FDebugMode>(bDebug);
+			PermutationVector.Set<FGenerateSamplesCS::FHashTable>(bUseHashTable);
+			PermutationVector.Set<FGenerateSamplesCS::FCandidateLightMask>(ShadowMaskAtlasHistory && CVarStochasticShadowsCandidateLightMask.GetValueOnRenderThread() != 0);
+			PermutationVector = FGenerateSamplesCS::RemapPermutation(PermutationVector);
+			auto ComputeShader = View.ShaderMap->GetShader<FGenerateSamplesCS>(PermutationVector);
 
-		FComputeShaderUtils::AddPass(
-			GraphBuilder,
-			RDG_EVENT_NAME("GenerateSamples SamplesPerPixel:%dx%d", NumSamplesPerPixel2d.X, NumSamplesPerPixel2d.Y),
-			ComputeShader,
-			PassParameters,
-			GroupCount);
+			FComputeShaderUtils::AddPass(
+				GraphBuilder,
+				RDG_EVENT_NAME("GenerateSamples SamplesPerPixel:%dx%d TileType:%d", NumSamplesPerPixel2d.X, NumSamplesPerPixel2d.Y, TileType),
+				ComputeShader,
+				PassParameters,
+				DownsampledTileIndirectArgs,
+				TileType * sizeof(FRHIDispatchIndirectParameters));
+		}
 	}
 
 	FRDGBufferRef ShadowMaskTileIndirectArgs = GraphBuilder.CreateBuffer(FRDGBufferDesc::CreateIndirectDesc<FRHIDispatchIndirectParameters>(1), TEXT("StochasticShadows.ShadowMaskTileIndirectArgs"));
@@ -1303,36 +1480,6 @@
 		}
 	}
 
-	// Composite shadow masks traces
-	{
-		FCompositeShadowMaskTracesCS::FParameters* PassParameters = GraphBuilder.AllocParameters<FCompositeShadowMaskTracesCS::FParameters>();
-		PassParameters->StochasticShadowsParameters = StochasticShadowsParameters;
-		PassParameters->IndirectArgs = ShadowMaskTileIndirectArgs;
-		PassParameters->RWShadowMaskTileAtlas = GraphBuilder.CreateUAV(ShadowMaskTileAtlas);
-		PassParameters->ShadowMaskPageTable = ShadowMaskPageTable;
-		PassParameters->ShadowMaskTileHeader = GraphBuilder.CreateSRV(ShadowMaskTileHeader);
-		PassParameters->ShadowMaskHistoryScreenCoord00 = ShadowMaskHistoryScreenCoord00;
-		PassParameters->ShadowMaskHistoryWeights = ShadowMaskHistoryWeights;
-		PassParameters->ShadowMaskPageTableHistory = ShadowMaskPageTableHistory;
-		PassParameters->ShadowMaskHashTableHistory = ShadowMaskHashTableHistory ? GraphBuilder.CreateSRV(ShadowMaskHashTableHistory) : nullptr;
-		PassParameters->ShadowMaskAtlasHistory = ShadowMaskAtlasHistory;
-		PassParameters->LightSamples = LightSamples;
-
-		FCompositeShadowMaskTracesCS::FPermutationDomain PermutationVector;
-		PermutationVector.Set<FCompositeShadowMaskTracesCS::FNumSamplesPerPixel>(NumSamplesPerPixel1d);
-		PermutationVector.Set<FCompositeShadowMaskTracesCS::FTemporalAccumulation>(bTemporal);
-		PermutationVector.Set<FCompositeShadowMaskTracesCS::FHashTable>(bUseHashTable);
-		auto ComputeShader = View.ShaderMap->GetShader<FCompositeShadowMaskTracesCS>(PermutationVector);
-
-		FComputeShaderUtils::AddPass(
-			GraphBuilder,
-			RDG_EVENT_NAME("CompositeShadowMaskTraces"),
-			ComputeShader,
-			PassParameters,
-			ShadowMaskTileIndirectArgs,
-			0);
-	}
-
 	FRDGTextureRef ShadowMaskUpsampleWeights = GraphBuilder.CreateTexture(
 		FRDGTextureDesc::Create2D(SceneTextures.Config.Extent, PF_FloatRGBA, FClearValueBinding::Black, TexCreate_ShaderResource | TexCreate_UAV),
 		TEXT("StochasticShadows.ShadowMaskUpsampleWeights"));
@@ -1377,6 +1524,37 @@
 	AddClearUAVPass(GraphBuilder, GraphBuilder.CreateUAV(ShadingTileAllocator), 0u);
 	AddClearUAVPass(GraphBuilder, GraphBuilder.CreateUAV(ShadingTileGridAllocator), 0u);
 
+	// Composite shadow masks traces
 	{
+		FCompositeShadowMaskTracesCS::FParameters* PassParameters = GraphBuilder.AllocParameters<FCompositeShadowMaskTracesCS::FParameters>();
+		PassParameters->StochasticShadowsParameters = StochasticShadowsParameters;
+		PassParameters->IndirectArgs = ShadowMaskTileIndirectArgs;
+		PassParameters->RWShadowMaskTileAtlas = GraphBuilder.CreateUAV(ShadowMaskTileAtlas);
+		PassParameters->ShadowMaskPageTable = ShadowMaskPageTable;
+		PassParameters->ShadowMaskTileData = GraphBuilder.CreateSRV(ShadowMaskTileData);
+		PassParameters->ShadowMaskHistoryScreenCoord00 = ShadowMaskHistoryScreenCoord00;
+		PassParameters->ShadowMaskHistoryWeights = ShadowMaskHistoryWeights;
+		PassParameters->ShadowMaskPageTableHistory = ShadowMaskPageTableHistory;
+		PassParameters->ShadowMaskHashTableHistory = ShadowMaskHashTableHistory ? GraphBuilder.CreateSRV(ShadowMaskHashTableHistory) : nullptr;
+		PassParameters->ShadowMaskAtlasHistory = ShadowMaskAtlasHistory;
+		PassParameters->LightSamples = LightSamples;
+
+		FCompositeShadowMaskTracesCS::FPermutationDomain PermutationVector;
+		PermutationVector.Set<FCompositeShadowMaskTracesCS::FNumSamplesPerPixel>(NumSamplesPerPixel1d);
+		PermutationVector.Set<FCompositeShadowMaskTracesCS::FTemporalAccumulation>(bTemporal);
+		PermutationVector.Set<FCompositeShadowMaskTracesCS::FHashTable>(bUseHashTable);
+		auto ComputeShader = View.ShaderMap->GetShader<FCompositeShadowMaskTracesCS>(PermutationVector);
+
+		FComputeShaderUtils::AddPass(
+			GraphBuilder,
+			RDG_EVENT_NAME("CompositeShadowMaskTraces"),
+			ComputeShader,
+			PassParameters,
+			ShadowMaskTileIndirectArgs,
+			0);
+	}
+
+	// Spatial denoising pass
+	{
 		FShadowMaskSpatialPassCS::FParameters* PassParameters = GraphBuilder.AllocParameters<FShadowMaskSpatialPassCS::FParameters>();
 		PassParameters->StochasticShadowsParameters = StochasticShadowsParameters;
@@ -1386,7 +1564,7 @@
 		PassParameters->RWShadingTileGrid = GraphBuilder.CreateUAV(ShadingTileGrid);
 		PassParameters->RWShadingTileAtlas = GraphBuilder.CreateUAV(ShadingTileAtlas);
 		PassParameters->ShadowMaskUpsampleWeights = ShadowMaskUpsampleWeights;
-		PassParameters->ShadowMaskTileHeader = GraphBuilder.CreateSRV(ShadowMaskTileHeader);
+		PassParameters->ShadowMaskTileData = GraphBuilder.CreateSRV(ShadowMaskTileData);
 		PassParameters->ShadowMaskPageTable = ShadowMaskPageTable;
 		PassParameters->ShadowMaskHashTable = ShadowMaskHashTable ? GraphBuilder.CreateSRV(ShadowMaskHashTable) : nullptr;
 		PassParameters->ShadowMaskTileAtlas = ShadowMaskTileAtlas;
@@ -1405,51 +1583,39 @@
 			0);
 	}
 
-	FRDGBufferRef ShadingTileIndirectArgs = GraphBuilder.CreateBuffer(FRDGBufferDesc::CreateIndirectDesc<FRHIDispatchIndirectParameters>(1), TEXT("StochasticShadows.ShadingTileIndirectArgs"));
-
-	// Setup indirect args for shadow mask tile updates reusing FInitShadowMaskUpdateIndirectArgsCS
-	{
-		FInitShadowMaskUpdateIndirectArgsCS::FParameters* PassParameters = GraphBuilder.AllocParameters<FInitShadowMaskUpdateIndirectArgsCS::FParameters>();
-		PassParameters->StochasticShadowsParameters = StochasticShadowsParameters;
-		PassParameters->RWIndirectArgs = GraphBuilder.CreateUAV(ShadingTileIndirectArgs);
-		PassParameters->ShadowMaskTileAllocator = GraphBuilder.CreateSRV(ShadingTileAllocator);
-
-		auto ComputeShader = View.ShaderMap->GetShader<FInitShadowMaskUpdateIndirectArgsCS>();
-
-		FComputeShaderUtils::AddPass(
-			GraphBuilder,
-			RDG_EVENT_NAME("InitShadingTileIndirectArgs"),
-			ComputeShader,
-			PassParameters,
-			FIntVector(1, 1, 1));
-	}
-
 	// Shade light samples
 	{
-		FShadeLightSamplesWithShadowMaskCS::FParameters* PassParameters = GraphBuilder.AllocParameters<FShadeLightSamplesWithShadowMaskCS::FParameters>();
-		PassParameters->StochasticShadowsParameters = StochasticShadowsParameters;
-		PassParameters->ShadowMaskPageTable = ShadowMaskPageTable;
-		PassParameters->ShadowMaskTileAtlas = ShadowMaskTileAtlas;
-		PassParameters->ShadowMaskTileAllocator = GraphBuilder.CreateSRV(ShadowMaskTileAllocator);
-		PassParameters->ShadingTileAllocator = GraphBuilder.CreateSRV(ShadingTileAllocator);
-		PassParameters->ShadingTileGridAllocator = ShadingTileGridAllocator;
-		PassParameters->ShadingTileGrid = GraphBuilder.CreateSRV(ShadingTileGrid);
-		PassParameters->ShadingTileAtlas = ShadingTileAtlas;
-		PassParameters->RWSceneColor = GraphBuilder.CreateUAV(SceneTextures.Color.Target);
-		PassParameters->ShadowMaskHashTable = ShadowMaskHashTable ? GraphBuilder.CreateSRV(ShadowMaskHashTable) : nullptr;
+		FRDGTextureUAVRef SceneColorUAV = GraphBuilder.CreateUAV(SceneTextures.Color.Target, ERDGUnorderedAccessViewFlags::SkipBarrier);
 
-		FShadeLightSamplesWithShadowMaskCS::FPermutationDomain PermutationVector;
-		PermutationVector.Set<FShadeLightSamplesWithShadowMaskCS::FDebugMode>(bDebug);
-		auto ComputeShader = View.ShaderMap->GetShader<FShadeLightSamplesWithShadowMaskCS>(PermutationVector);
+		for (int32 TileType = 0; TileType < (int32)StochasticShadows::ETileType::MAX; ++TileType)
+		{
+			FShadeLightSamplesCS::FParameters* PassParameters = GraphBuilder.AllocParameters<FShadeLightSamplesCS::FParameters>();
+			PassParameters->StochasticShadowsParameters = StochasticShadowsParameters;
+			PassParameters->ShadowMaskPageTable = ShadowMaskPageTable;
+			PassParameters->ShadowMaskTileAtlas = ShadowMaskTileAtlas;
+			PassParameters->ShadowMaskTileAllocator = GraphBuilder.CreateSRV(ShadowMaskTileAllocator);
+			PassParameters->TileAllocator = GraphBuilder.CreateSRV(TileAllocator);
+			PassParameters->TileData = GraphBuilder.CreateSRV(TileData);
+			PassParameters->ShadingTileAllocator = GraphBuilder.CreateSRV(ShadingTileAllocator);
+			PassParameters->ShadingTileGridAllocator = ShadingTileGridAllocator;
+			PassParameters->ShadingTileGrid = GraphBuilder.CreateSRV(ShadingTileGrid);
+			PassParameters->ShadingTileAtlas = ShadingTileAtlas;
+			PassParameters->RWSceneColor = SceneColorUAV;
+			PassParameters->ShadowMaskHashTable = ShadowMaskHashTable ? GraphBuilder.CreateSRV(ShadowMaskHashTable) : nullptr;
 
-		const FIntVector GroupCount = FComputeShaderUtils::GetGroupCount(View.ViewRect.Size(), FShadeLightSamplesWithShadowMaskCS::GetGroupSize());
+			FShadeLightSamplesCS::FPermutationDomain PermutationVector;
+			PermutationVector.Set<FShadeLightSamplesCS::FTileType>(TileType);
+			PermutationVector.Set<FShadeLightSamplesCS::FDebugMode>(bDebug);
+			auto ComputeShader = View.ShaderMap->GetShader<FShadeLightSamplesCS>(PermutationVector);
 
-		FComputeShaderUtils::AddPass(
-			GraphBuilder,
-			RDG_EVENT_NAME("ShadeLightSamplesWithShadowMask"),
-			ComputeShader,
-			PassParameters,
-			GroupCount);
+			FComputeShaderUtils::AddPass(
+				GraphBuilder,
+				RDG_EVENT_NAME("ShadeLightSamples TileType:%d", TileType),
+				ComputeShader,
+				PassParameters,
+				TileIndirectArgs,
+				TileType * sizeof(FRHIDispatchIndirectParameters));
+		}
 	}
 
 	if (View.ViewState && !View.bStatePrevViewInfoIsReadOnly)
