Changed visibility compute relevance task to wait on finalize relevance, as tasks like shadows need the full finalize to complete before results are valid.

diff --git a/Engine/Source/Runtime/Renderer/Private/SceneVisibilityPrivate.h b/Engine/Source/Runtime/Renderer/Private/SceneVisibilityPrivate.h
--- a/Engine/Source/Runtime/Renderer/Private/SceneVisibilityPrivate.h
+++ b/Engine/Source/Runtime/Renderer/Private/SceneVisibilityPrivate.h
@@ -410,7 +410,7 @@ public:
 
 	UE::Tasks::FTask GetComputeRelevanceTask() const override
 	{
-		return Tasks.ComputeRelevance;
+		return Tasks.FinalizeRelevance;
 	}
 
 	UE::Tasks::FTask GetLightVisibilityTask() const override
