Wait on GPU skin cache task prior to running hair bookmarks prior to visibility.

diff --git a/Engine/Source/Runtime/Renderer/Private/SceneRendering.cpp b/Engine/Source/Runtime/Renderer/Private/SceneRendering.cpp
--- a/Engine/Source/Runtime/Renderer/Private/SceneRendering.cpp
+++ b/Engine/Source/Runtime/Renderer/Private/SceneRendering.cpp
@@ -3473,6 +3473,8 @@ IVisibilityTaskData* FSceneRenderer::OnRenderBegin(FRDGBuilder& GraphBuilder)
 			CreateHairStrandsBookmarkParameters(Scene, Views, AllFamilyViews, Parameters, false/*bComputeVisibleInstances*/);
 			if (Parameters.HasInstances())
 			{
+				Scene->WaitForGPUSkinCacheTask();
+
 				// 1. Select appropriate LOD & geometry type
 				RunHairStrandsBookmark(GraphBuilder, EHairStrandsBookmark::ProcessLODSelection, Parameters);
 			}
diff --git a/Engine/Source/Runtime/Renderer/Private/SceneVisibility.cpp b/Engine/Source/Runtime/Renderer/Private/SceneVisibility.cpp
--- a/Engine/Source/Runtime/Renderer/Private/SceneVisibility.cpp
+++ b/Engine/Source/Runtime/Renderer/Private/SceneVisibility.cpp
@@ -4126,6 +4126,8 @@ void FSceneRenderer::PreVisibilityFrameSetup(FRDGBuilder& GraphBuilder)
 				// If we are rendering from scene capture we don't need to run another time the hair bookmarks.
 				if (IsHairStrandsEnabled(EHairStrandsShaderType::All, Scene->GetShaderPlatform()) && Views[0].AllowGPUParticleUpdate())
 				{
+					Scene->WaitForGPUSkinCacheTask();
+
 					RunHairStrandsBookmark(GraphBuilder, EHairStrandsBookmark::ProcessGuideInterpolation, Parameters);
 				}
 			}
