Fixed rendering crash on Android when project has GPUScene disabled (introduced in CL# 29034277)

diff --git a/Engine/Source/Runtime/Renderer/Private/MobileBasePassRendering.cpp b/Engine/Source/Runtime/Renderer/Private/MobileBasePassRendering.cpp
--- a/Engine/Source/Runtime/Renderer/Private/MobileBasePassRendering.cpp
+++ b/Engine/Source/Runtime/Renderer/Private/MobileBasePassRendering.cpp
@@ -463,7 +463,7 @@ void FMobileSceneRenderer::RenderMobileEditorPrimitives(FRHICommandList& RHICmdL
 		
 		// DrawDynamicMeshPass may change BatchedPrimitive binding, so we need to restore it
 		FUniformBufferStaticSlot BatchedPrimitiveSlot = FInstanceCullingContext::GetUniformBufferViewStaticSlot(View.GetShaderPlatform());
-		if (IsUniformBufferStaticSlotValid(BatchedPrimitiveSlot))
+		if (IsUniformBufferStaticSlotValid(BatchedPrimitiveSlot) && InstanceCullingDrawParams->BatchedPrimitive)
 		{
 			FRHIUniformBuffer* BatchedPrimitiveBufferRHI = InstanceCullingDrawParams->BatchedPrimitive.GetUniformBuffer()->GetRHI();
 			check(BatchedPrimitiveBufferRHI);
diff --git a/Engine/Source/Runtime/Renderer/Private/MobileDecalRendering.cpp b/Engine/Source/Runtime/Renderer/Private/MobileDecalRendering.cpp
--- a/Engine/Source/Runtime/Renderer/Private/MobileDecalRendering.cpp
+++ b/Engine/Source/Runtime/Renderer/Private/MobileDecalRendering.cpp
@@ -73,7 +73,7 @@ void FMobileSceneRenderer::RenderDecals(FRHICommandList& RHICmdList, const FView
 
 		// MeshDecals use DrawDynamicMeshPass which may change BatchedPrimitive binding, so we need to restore it
 		FUniformBufferStaticSlot BatchedPrimitiveSlot = FInstanceCullingContext::GetUniformBufferViewStaticSlot(View.GetShaderPlatform());
-		if (IsUniformBufferStaticSlotValid(BatchedPrimitiveSlot))
+		if (IsUniformBufferStaticSlotValid(BatchedPrimitiveSlot) && InstanceCullingDrawParams->BatchedPrimitive)
 		{
 			FRHIUniformBuffer* BatchedPrimitiveBufferRHI = InstanceCullingDrawParams->BatchedPrimitive.GetUniformBuffer()->GetRHI();
 			check(BatchedPrimitiveBufferRHI);
