Remove GBufferPicking support from HLSLcc platform.

diff --git a/Engine/Source/Runtime/Renderer/Private/PostProcess/PostProcessing.cpp b/Engine/Source/Runtime/Renderer/Private/PostProcess/PostProcessing.cpp
--- a/Engine/Source/Runtime/Renderer/Private/PostProcess/PostProcessing.cpp
+++ b/Engine/Source/Runtime/Renderer/Private/PostProcess/PostProcessing.cpp
@@ -2923,14 +2923,18 @@ class FGBufferPickingCS : public FGlobalShader
 	END_SHADER_PARAMETER_STRUCT()
 
 public:
-	static bool ShouldCompilePermutation(const FGlobalShaderPermutationParameters& Parameters) 
+	static bool IsSupported(EShaderPlatform Platform) 
 	{ 
 		return 
-			IsFeatureLevelSupported(Parameters.Platform, ERHIFeatureLevel::SM5) &&
-			ShaderPrint::IsSupported(Parameters.Platform) && 
-			!IsMobilePlatform(Parameters.Platform) && 
-			!Substrate::IsSubstrateEnabled() &&
-			EnumHasAllFlags(Parameters.Flags, EShaderPermutationFlags::HasEditorOnlyData);
+			IsFeatureLevelSupported(Platform, ERHIFeatureLevel::SM5) &&
+			ShaderPrint::IsSupported(Platform) && 
+			!IsHlslccShaderPlatform(Platform) &&
+			!IsMobilePlatform(Platform) && 
+			!Substrate::IsSubstrateEnabled();
+	}
+	static bool ShouldCompilePermutation(const FGlobalShaderPermutationParameters& Parameters) 
+	{ 
+		return IsSupported(Parameters.Platform) && EnumHasAllFlags(Parameters.Flags, EShaderPermutationFlags::HasEditorOnlyData);
 	}
 	static void ModifyCompilationEnvironment(const FGlobalShaderPermutationParameters& Parameters, FShaderCompilerEnvironment& OutEnvironment)
 	{
@@ -2945,7 +2949,7 @@ IMPLEMENT_GLOBAL_SHADER(FGBufferPickingCS, "/Engine/Private/PostProcessGBufferHi
 #if WITH_EDITOR
 static void AddGBufferPicking(FRDGBuilder& GraphBuilder, const FViewInfo& View, const TRDGUniformBufferRef<FSceneTextureUniformParameters>& SceneTextures)
 {
-	if (CVarGBufferPicking.GetValueOnRenderThread() <= 0 || !ShaderPrint::IsSupported(View.Family->GetShaderPlatform()))
+	if (CVarGBufferPicking.GetValueOnRenderThread() <= 0 || !FGBufferPickingCS::IsSupported(View.Family->GetShaderPlatform()))
 	{
 		return;
 	}
