Fixed artifacts in Niagara due to multiple locks to the same resource.

diff --git a/Engine/Source/Runtime/Renderer/Private/DeferredShadingRenderer.cpp b/Engine/Source/Runtime/Renderer/Private/DeferredShadingRenderer.cpp
--- a/Engine/Source/Runtime/Renderer/Private/DeferredShadingRenderer.cpp
+++ b/Engine/Source/Runtime/Renderer/Private/DeferredShadingRenderer.cpp
@@ -1103,7 +1103,7 @@ bool FDeferredShadingSceneRenderer::GatherRayTracingWorldInstancesForView(FRDGBu
 		ViewFamily,
 		GraphBuilder,
 		*View.RayTracingMeshResourceCollector,
-		DynamicReadBufferForInitViews
+		DynamicReadBufferForRayTracing
 	);
 
 	const float CurrentWorldTime = View.Family->Time.GetWorldTimeSeconds();
diff --git a/Engine/Source/Runtime/Renderer/Private/SceneRendering.cpp b/Engine/Source/Runtime/Renderer/Private/SceneRendering.cpp
--- a/Engine/Source/Runtime/Renderer/Private/SceneRendering.cpp
+++ b/Engine/Source/Runtime/Renderer/Private/SceneRendering.cpp
@@ -2495,6 +2495,7 @@ FViewFamilyInfo::~FViewFamilyInfo()
 }
 
 TGlobalResource<FGlobalDynamicReadBuffer> FSceneRenderer::DynamicReadBufferForInitViews;
+TGlobalResource<FGlobalDynamicReadBuffer> FSceneRenderer::DynamicReadBufferForRayTracing;
 TGlobalResource<FGlobalDynamicReadBuffer> FSceneRenderer::DynamicReadBufferForShadows;
 
 /*-----------------------------------------------------------------------------
diff --git a/Engine/Source/Runtime/Renderer/Private/SceneRendering.h b/Engine/Source/Runtime/Renderer/Private/SceneRendering.h
--- a/Engine/Source/Runtime/Renderer/Private/SceneRendering.h
+++ b/Engine/Source/Runtime/Renderer/Private/SceneRendering.h
@@ -2268,6 +2268,7 @@ public:
 	inline TConstStridedView<FSceneView> GetSceneViews() const { return MakeStridedViewOfBase<const FSceneView>(MakeArrayView(Views)); }
 
 	static TGlobalResource<FGlobalDynamicReadBuffer> DynamicReadBufferForInitViews;
+	static TGlobalResource<FGlobalDynamicReadBuffer> DynamicReadBufferForRayTracing;
 	static TGlobalResource<FGlobalDynamicReadBuffer> DynamicReadBufferForShadows;
 
 protected:
