RayTracing: Fix size mismatch in buffer description for raytracing light grid

diff --git a/Engine/Source/Runtime/Renderer/Private/RayTracing/RayTracingLighting.cpp b/Engine/Source/Runtime/Renderer/Private/RayTracing/RayTracingLighting.cpp
--- a/Engine/Source/Runtime/Renderer/Private/RayTracing/RayTracingLighting.cpp
+++ b/Engine/Source/Runtime/Renderer/Private/RayTracing/RayTracingLighting.cpp
@@ -277,14 +277,12 @@ TRDGUniformBufferRef<FRaytracingLightDataPacked> CreateRayTracingLightData(
 	FRDGBufferRef LightIndicesBuffer;
 	CreateRaytracingLightCullingStructure(GraphBuilder, Scene->Lights, View, ShaderMap, LightIndices, LightCullVolume, LightIndicesBuffer);
 
-	FRDGUploadData<FRTLightingData> LightDataArray(GraphBuilder, FMath::Max(LightIndices.Num(), 1));
+	FRDGUploadData<FRTLightingData> LightDataArray(GraphBuilder, LightIndices.Num());
 	SetupRaytracingLightDataPacked(GraphBuilder, Scene, LightIndices, View, *LightData, LightDataArray);
 
 	check(LightData->Count == LightIndices.Num());
 
-	static_assert(sizeof(FRTLightingData) % sizeof(FUintVector4) == 0, "sizeof(FRTLightingData) must be a multiple of sizeof(FUintVector4)");
-	const uint32 NumUintVector4Elements = LightDataArray.GetTotalSize() / sizeof(FUintVector4);
-	FRDGBufferRef LightBuffer = CreateStructuredBuffer(GraphBuilder, TEXT("LightBuffer"), sizeof(FUintVector4), NumUintVector4Elements, LightDataArray.GetData(), LightDataArray.GetTotalSize(), ERDGInitialDataFlags::NoCopy);
+	FRDGBufferRef LightBuffer = CreateStructuredBuffer(GraphBuilder, TEXT("LightBuffer"), LightDataArray);
 
 	LightData->LightDataBuffer = GraphBuilder.CreateSRV(LightBuffer);
 	LightData->LightIndices = GraphBuilder.CreateSRV(LightIndicesBuffer, EPixelFormat::PF_R16_UINT);
