Removed rotten code that is now done in FPostProcessMaterialInputs::Validate()

diff --git a/Engine/Source/Runtime/Renderer/Private/PostProcess/PostProcessMaterial.cpp b/Engine/Source/Runtime/Renderer/Private/PostProcess/PostProcessMaterial.cpp
--- a/Engine/Source/Runtime/Renderer/Private/PostProcess/PostProcessMaterial.cpp
+++ b/Engine/Source/Runtime/Renderer/Private/PostProcess/PostProcessMaterial.cpp
@@ -384,13 +384,6 @@ static void GetMaterialInfo(
 
 	check(Material);
 
-	if (Material->IsStencilTestEnabled() || Material->GetBlendableOutputAlpha())
-	{
-		// Only allowed to have blend/stencil test if output format is compatible with ePId_Input0. 
-		// PF_Unknown implies output format is that of EPId_Input0
-		ensure(Inputs.OutputFormat == PF_Unknown);
-	}
-
 	const FMaterialShaderMap* MaterialShaderMap = Material->GetRenderingThreadShaderMap();
 	check(MaterialShaderMap);
 
