VolumetricCloud - added tracing start distance from camera

diff --git a/Engine/Shaders/Private/VolumetricCloud.usf b/Engine/Shaders/Private/VolumetricCloud.usf
--- a/Engine/Shaders/Private/VolumetricCloud.usf
+++ b/Engine/Shaders/Private/VolumetricCloud.usf
@@ -655,6 +655,10 @@ void MainCommon(in FMaterialPixelParameters MaterialParameters, in float4 SvPosi
 			}
 		}
 
+		// Account for the tracing start distance
+		TMin = max(TMin, RenderVolumetricCloudParameters.TracingStartDistanceFromCamera);
+		TMax = max(TMin, TMax);
+
 		// When we do trace, we want a minimum amount of sample count. This avoid situation with 0 sample count due to InvDistanceToSampleCountMax.
 		const uint  IStepCount = max(RenderVolumetricCloudParameters.SampleCountMin, RenderVolumetricCloudParameters.SampleCountMax * saturate((TMax - TMin) * RenderVolumetricCloudParameters.InvDistanceToSampleCountMax));
 		const float StepCount = float(IStepCount);
diff --git a/Engine/Source/Runtime/Engine/Classes/Components/VolumetricCloudComponent.h b/Engine/Source/Runtime/Engine/Classes/Components/VolumetricCloudComponent.h
--- a/Engine/Source/Runtime/Engine/Classes/Components/VolumetricCloudComponent.h
+++ b/Engine/Source/Runtime/Engine/Classes/Components/VolumetricCloudComponent.h
@@ -39,10 +39,14 @@ class UVolumetricCloudComponent : public USceneComponent
 	UPROPERTY(EditAnywhere, BlueprintReadOnly, interp, Category = "Layer", meta = (UIMin = 0.1f, UIMax = 20.0f, ClampMin = 0.1, SliderExponent = 2.0))
 	float LayerHeight;
 
-	/** The maximum distance of the volumetric surface before which we will accept to start tracing. (kilometers) */
+	/** The maximum distance of the volumetric surface, i.e. cloud layer upper and lower bound, before which we will accept to start tracing. (kilometers) */
 	UPROPERTY(EditAnywhere, BlueprintReadOnly, interp, Category = "Layer", meta = (UIMin = 100.0f, UIMax = 500.0f, ClampMin = 1.0f, SliderExponent = 2.0))
 	float TracingStartMaxDistance;
 
+	/** The distance from which the tracing will start. This is useful when the camera for instance is inside the layer of cloud. (kilometers) */
+	UPROPERTY(EditAnywhere, BlueprintReadOnly, interp, Category = "Layer", meta = (UIMin = 0.0f, UIMax = 100.0f, ClampMin = 0.0f, SliderExponent = 3.0))
+	float TracingStartDistanceFromCamera;
+
 	/** Mode to select how the tracing max distance should be interpreted. DistanceFromPointOfView is useful to avoid the top of the cloud layer to be clipped when TracingMaxDistance is shorten for performance. */
 	UPROPERTY(EditAnywhere, BlueprintReadOnly, interp, Category = "Layer")
 	EVolumetricCloudTracingMaxDistanceMode TracingMaxDistanceMode;
@@ -69,7 +73,7 @@ class UVolumetricCloudComponent : public USceneComponent
 	/** Whether to apply atmosphere transmittance per sample, instead of using the light global transmittance. */
 	UPROPERTY(EditAnywhere, BlueprintReadOnly, Category = "Cloud Tracing")
 	uint32 bUsePerSampleAtmosphericLightTransmittance : 1; 
-	// bUsePerSampleAtmosphericLightTransmittance is there on the cloud component and not on the light because otherwise we would need optimisation permutations of the cloud shader.
+	// bUsePerSampleAtmosphericLightTransmittance is there on the cloud component and not on the light because otherwise we would need optimization permutations of the cloud shader.
 	// And this for the two atmospheric lights ON or OFF. Keeping it simple for now because this changes the look of the cloud, so it is an art/look decision.
 
 	/** Occlude the sky light contribution at the bottom of the cloud layer. This is a fast approximation to sky lighting being occluded by cloud without having to trace rays or sample AO texture. Ignored if the cloud material explicitely sets the ambient occlusion value. */
@@ -143,6 +147,8 @@ class UVolumetricCloudComponent : public USceneComponent
 	UFUNCTION(BlueprintCallable, Category = "Rendering")
 	ENGINE_API void SetTracingStartMaxDistance(float NewValue);
 	UFUNCTION(BlueprintCallable, Category = "Rendering")
+	ENGINE_API void SetTracingStartDistanceFromCamera(float NewValue);
+	UFUNCTION(BlueprintCallable, Category = "Rendering")
 	ENGINE_API void SetTracingMaxDistance(float NewValue);
 	UFUNCTION(BlueprintCallable, Category = "Rendering")
 	ENGINE_API void SetPlanetRadius(float NewValue);
diff --git a/Engine/Source/Runtime/Engine/Private/Components/VolumetricCloudComponent.cpp b/Engine/Source/Runtime/Engine/Private/Components/VolumetricCloudComponent.cpp
--- a/Engine/Source/Runtime/Engine/Private/Components/VolumetricCloudComponent.cpp
+++ b/Engine/Source/Runtime/Engine/Private/Components/VolumetricCloudComponent.cpp
@@ -31,6 +31,7 @@ UVolumetricCloudComponent::UVolumetricCloudComponent(const FObjectInitializer& O
 	, LayerBottomAltitude(5.0f)
 	, LayerHeight(10.0f)
 	, TracingStartMaxDistance(350.0f)
+	, TracingStartDistanceFromCamera(0.0f)
 	, TracingMaxDistanceMode(EVolumetricCloudTracingMaxDistanceMode::DistanceFromCloudLayerEntryPoint)
 	, TracingMaxDistance(50.0f)
 	, PlanetRadius(6360.0f)					// Default to earth-like
@@ -153,6 +154,7 @@ void UVolumetricCloudComponent::Serialize(FArchive& Ar)
 CLOUD_DECLARE_BLUEPRINT_SETFUNCTION(float, LayerBottomAltitude);
 CLOUD_DECLARE_BLUEPRINT_SETFUNCTION(float, LayerHeight);
 CLOUD_DECLARE_BLUEPRINT_SETFUNCTION(float, TracingStartMaxDistance);
+CLOUD_DECLARE_BLUEPRINT_SETFUNCTION(float, TracingStartDistanceFromCamera);
 CLOUD_DECLARE_BLUEPRINT_SETFUNCTION(float, TracingMaxDistance);
 CLOUD_DECLARE_BLUEPRINT_SETFUNCTION(float, PlanetRadius);
 CLOUD_DECLARE_BLUEPRINT_SETFUNCTION(FColor, GroundAlbedo);
diff --git a/Engine/Source/Runtime/Engine/Private/VolumetricCloudProxy.cpp b/Engine/Source/Runtime/Engine/Private/VolumetricCloudProxy.cpp
--- a/Engine/Source/Runtime/Engine/Private/VolumetricCloudProxy.cpp
+++ b/Engine/Source/Runtime/Engine/Private/VolumetricCloudProxy.cpp
@@ -13,6 +13,7 @@ FVolumetricCloudSceneProxy::FVolumetricCloudSceneProxy(const UVolumetricCloudCom
 	: LayerBottomAltitudeKm(InComponent->LayerBottomAltitude)
 	, LayerHeightKm(InComponent->LayerHeight)
 	, TracingStartMaxDistance(InComponent->TracingStartMaxDistance)
+	, TracingStartDistanceFromCamera(InComponent->TracingStartDistanceFromCamera)
 	, TracingMaxDistanceMode(uint8(InComponent->TracingMaxDistanceMode))
 	, TracingMaxDistance(InComponent->TracingMaxDistance)
 	, PlanetRadiusKm(InComponent->PlanetRadius)
diff --git a/Engine/Source/Runtime/Engine/Public/VolumetricCloudProxy.h b/Engine/Source/Runtime/Engine/Public/VolumetricCloudProxy.h
--- a/Engine/Source/Runtime/Engine/Public/VolumetricCloudProxy.h
+++ b/Engine/Source/Runtime/Engine/Public/VolumetricCloudProxy.h
@@ -33,6 +33,7 @@ public:
 	float LayerHeightKm;
 
 	float TracingStartMaxDistance;
+	float TracingStartDistanceFromCamera;
 	uint8 TracingMaxDistanceMode;
 	float TracingMaxDistance;
 
diff --git a/Engine/Source/Runtime/Renderer/Private/VolumetricCloudRendering.cpp b/Engine/Source/Runtime/Renderer/Private/VolumetricCloudRendering.cpp
--- a/Engine/Source/Runtime/Renderer/Private/VolumetricCloudRendering.cpp
+++ b/Engine/Source/Runtime/Renderer/Private/VolumetricCloudRendering.cpp
@@ -207,7 +207,7 @@ static TAutoConsoleVariable<float> CVarVolumetricCloudLocalLightsShadowSampleCou
 
 static TAutoConsoleVariable<int32> CVarVolumetricCloudEmptySpaceSkipping(
 	TEXT("r.VolumetricCloud.EmptySpaceSkipping"), 0,
-	TEXT("Enable/disable empty space skipping to accelerate cloud tracing through emty areas."),
+	TEXT("Enable/disable empty space skipping to accelerate cloud tracing through emty areas. EXPERIMENTAL"),
 	ECVF_RenderThreadSafe | ECVF_Scalability);
 
 static TAutoConsoleVariable<float> CVarVolumetricCloudEmptySpaceSkippingVolumeDepth(
@@ -1596,23 +1596,23 @@ void FSceneRenderer::InitVolumetricCloudsForViews(FRDGBuilder& GraphBuilder, boo
 			{
 				CloudGlobalShaderParams.CloudLayerCenterKm = FVector3f(0.0f, 0.0f, -PlanetRadiusKm);
 			}
-			CloudGlobalShaderParams.PlanetRadiusKm = PlanetRadiusKm;
-			CloudGlobalShaderParams.BottomRadiusKm = PlanetRadiusKm + CloudProxy.LayerBottomAltitudeKm;
-			CloudGlobalShaderParams.TopRadiusKm = CloudGlobalShaderParams.BottomRadiusKm + CloudProxy.LayerHeightKm;
-			CloudGlobalShaderParams.GroundAlbedo = FLinearColor(CloudProxy.GroundAlbedo);
-			CloudGlobalShaderParams.SkyLightCloudBottomVisibility = 1.0f - CloudProxy.SkyLightCloudBottomOcclusion;
-
-			CloudGlobalShaderParams.TracingStartMaxDistance = KilometersToCentimeters * CloudProxy.TracingStartMaxDistance;
-			CloudGlobalShaderParams.TracingMaxDistanceMode	= CloudProxy.TracingMaxDistanceMode;
-			CloudGlobalShaderParams.TracingMaxDistance		= KilometersToCentimeters * CloudProxy.TracingMaxDistance;
-
-			CloudGlobalShaderParams.SampleCountMin		= FMath::Max(0, CVarVolumetricCloudSampleMinCount.GetValueOnAnyThread());
-			CloudGlobalShaderParams.SampleCountMax		= FMath::Max(2.0f, FMath::Min(UVolumetricCloudComponent::BaseViewRaySampleCount   * CloudProxy.ViewSampleCountScale,       CVarVolumetricCloudViewRaySampleMaxCount.GetValueOnAnyThread()));
-			CloudGlobalShaderParams.ShadowSampleCountMax= FMath::Max(2.0f, FMath::Min(UVolumetricCloudComponent::BaseShadowRaySampleCount * CloudProxy.ShadowViewSampleCountScale, CVarVolumetricCloudShadowViewRaySampleMaxCount.GetValueOnAnyThread()));
-			CloudGlobalShaderParams.ShadowTracingMaxDistance = KilometersToCentimeters * FMath::Max(0.1f, CloudProxy.ShadowTracingDistance);
-			CloudGlobalShaderParams.InvDistanceToSampleCountMax = 1.0f / FMath::Max(1.0f, KilometersToCentimeters * CVarVolumetricCloudDistanceToSampleMaxCount.GetValueOnAnyThread());
-			CloudGlobalShaderParams.StopTracingTransmittanceThreshold = FMath::Clamp(CloudProxy.StopTracingTransmittanceThreshold, 0.0f, 1.0f);
-
+			CloudGlobalShaderParams.PlanetRadiusKm						= PlanetRadiusKm;
+			CloudGlobalShaderParams.BottomRadiusKm						= PlanetRadiusKm + CloudProxy.LayerBottomAltitudeKm;
+			CloudGlobalShaderParams.TopRadiusKm							= CloudGlobalShaderParams.BottomRadiusKm + CloudProxy.LayerHeightKm;
+			CloudGlobalShaderParams.GroundAlbedo						= FLinearColor(CloudProxy.GroundAlbedo);
+			CloudGlobalShaderParams.SkyLightCloudBottomVisibility		= 1.0f - CloudProxy.SkyLightCloudBottomOcclusion;
+
+			CloudGlobalShaderParams.TracingStartMaxDistance				= FMath::Max(0.0f, KilometersToCentimeters * CloudProxy.TracingStartMaxDistance);
+			CloudGlobalShaderParams.TracingStartDistanceFromCamera		= FMath::Max(0.0f, KilometersToCentimeters * CloudProxy.TracingStartDistanceFromCamera);
+			CloudGlobalShaderParams.TracingMaxDistance					= FMath::Max(0.0f, KilometersToCentimeters * CloudProxy.TracingMaxDistance);
+			CloudGlobalShaderParams.TracingMaxDistanceMode				= CloudProxy.TracingMaxDistanceMode;
+
+			CloudGlobalShaderParams.SampleCountMin						= FMath::Max(0, CVarVolumetricCloudSampleMinCount.GetValueOnAnyThread());
+			CloudGlobalShaderParams.SampleCountMax						= FMath::Max(2.0f, FMath::Min(UVolumetricCloudComponent::BaseViewRaySampleCount   * CloudProxy.ViewSampleCountScale,       CVarVolumetricCloudViewRaySampleMaxCount.GetValueOnAnyThread()));
+			CloudGlobalShaderParams.ShadowSampleCountMax				= FMath::Max(2.0f, FMath::Min(UVolumetricCloudComponent::BaseShadowRaySampleCount * CloudProxy.ShadowViewSampleCountScale, CVarVolumetricCloudShadowViewRaySampleMaxCount.GetValueOnAnyThread()));
+			CloudGlobalShaderParams.ShadowTracingMaxDistance			= KilometersToCentimeters * FMath::Max(0.1f, CloudProxy.ShadowTracingDistance);
+			CloudGlobalShaderParams.InvDistanceToSampleCountMax			= 1.0f / FMath::Max(1.0f, KilometersToCentimeters * CVarVolumetricCloudDistanceToSampleMaxCount.GetValueOnAnyThread());
+			CloudGlobalShaderParams.StopTracingTransmittanceThreshold	= FMath::Clamp(CloudProxy.StopTracingTransmittanceThreshold, 0.0f, 1.0f);
 
 			auto PrepareCloudShadowMapLightData = [&](FLightSceneProxy* AtmosphericLight, int LightIndex)
 			{
diff --git a/Engine/Source/Runtime/Renderer/Private/VolumetricCloudRendering.h b/Engine/Source/Runtime/Renderer/Private/VolumetricCloudRendering.h
--- a/Engine/Source/Runtime/Renderer/Private/VolumetricCloudRendering.h
+++ b/Engine/Source/Runtime/Renderer/Private/VolumetricCloudRendering.h
@@ -30,6 +30,7 @@ BEGIN_SHADER_PARAMETER_STRUCT(FVolumetricCloudCommonShaderParameters, )
 	SHADER_PARAMETER(float, PlanetRadiusKm)
 	SHADER_PARAMETER(float, BottomRadiusKm)
 	SHADER_PARAMETER(float, TopRadiusKm)
+	SHADER_PARAMETER(float, TracingStartDistanceFromCamera)
 	SHADER_PARAMETER(float, TracingStartMaxDistance)
 	SHADER_PARAMETER(int32, TracingMaxDistanceMode)
 	SHADER_PARAMETER(float, TracingMaxDistance)
