Remove Substrate uniform buffer binding from the distortion pass, since it is not used. The FSubstrateIntegrationSettings are hardcoded into DistortionAccumulatePS.usf.

diff --git a/Engine/Source/Runtime/Renderer/Private/DistortionRendering.cpp b/Engine/Source/Runtime/Renderer/Private/DistortionRendering.cpp
--- a/Engine/Source/Runtime/Renderer/Private/DistortionRendering.cpp
+++ b/Engine/Source/Runtime/Renderer/Private/DistortionRendering.cpp
@@ -257,7 +257,6 @@ BEGIN_SHADER_PARAMETER_STRUCT(FDistortionPassParameters, RENDERER_API)
 	SHADER_PARAMETER_STRUCT_INCLUDE(FViewShaderParameters, View)
 	SHADER_PARAMETER_RDG_UNIFORM_BUFFER(FDistortionPassUniformParameters, Pass)
 	SHADER_PARAMETER_STRUCT_INCLUDE(FInstanceCullingDrawParams, InstanceCullingDrawParams)
-	SHADER_PARAMETER_RDG_UNIFORM_BUFFER(FSubstrateGlobalUniformParameters, Substrate)
 	RENDER_TARGET_BINDING_SLOTS()
 END_SHADER_PARAMETER_STRUCT()
 
@@ -743,7 +742,6 @@ void FDeferredShadingSceneRenderer::RenderDistortion(
 			auto* PassParameters = GraphBuilder.AllocParameters<FDistortionPassParameters>();
 			PassParameters->View = View.GetShaderParameters();
 			PassParameters->Pass = CreateDistortionPassUniformBuffer(GraphBuilder, View);
-			PassParameters->Substrate = Substrate::BindSubstrateGlobalUniformParameters(View);
 			PassParameters->RenderTargets[0] = FRenderTargetBinding(DistortionTexture, LoadAction);
 			if (bUseRoughRefraction)
 			{
