Support ThinTranslucent Materials in mobile renderer

OpenGL ES and Metal use framebuffer fetch.
Vulkan uses dual source blending.
For Vulkan and OpenGL ES there is a fallback shader permutation for drivers that don't support this. The fallback is the same as the existing solution that uses regular blending (i.e. looks different).
Others uses dual source blending and we force use of DXC for those shaders.

diff --git a/Engine/Config/Android/DataDrivenPlatformInfo.ini b/Engine/Config/Android/DataDrivenPlatformInfo.ini
--- a/Engine/Config/Android/DataDrivenPlatformInfo.ini
+++ b/Engine/Config/Android/DataDrivenPlatformInfo.ini
@@ -86,6 +86,8 @@ bSupportsShaderPipelines = false
 bSupportsVariableRateShading = true
 bSupportsVertexShaderSRVs=false
 bSupportsUniformBufferObjects = true
+bSupportsDualSourceBlending = true
+
 FriendlyName=LOCTEXT("FriendlyShaderPlatformName_AndroidVulkan", "Android Vulkan Mobile")
 
 [ShaderPlatform VULKAN_SM5_ANDROID]
diff --git a/Engine/Config/Mac/DataDrivenPlatformInfo.ini b/Engine/Config/Mac/DataDrivenPlatformInfo.ini
--- a/Engine/Config/Mac/DataDrivenPlatformInfo.ini
+++ b/Engine/Config/Mac/DataDrivenPlatformInfo.ini
@@ -126,4 +126,5 @@ bSupports4ComponentUAVReadWrite=true
 bSupportsManualVertexFetch=false
 bSupportsVolumeTextureAtomics=false
 bSupportsShaderPipelines = false
+bSupportsDualSourceBlending = true
 FriendlyName=LOCTEXT("FriendlyShaderPlatformName_MacMobile", "Mac Metal Mobile")
diff --git a/Engine/Config/VulkanPC/DataDrivenPlatformInfo.ini b/Engine/Config/VulkanPC/DataDrivenPlatformInfo.ini
--- a/Engine/Config/VulkanPC/DataDrivenPlatformInfo.ini
+++ b/Engine/Config/VulkanPC/DataDrivenPlatformInfo.ini
@@ -271,6 +271,8 @@ bSupportsSceneDataCompressedTransforms=true
 bSupportsShaderPipelines = false
 bSupportsVertexShaderSRVs = false
 bSupportsUniformBufferObjects = true
+bSupportsDualSourceBlending = true
+
 FriendlyName=LOCTEXT("FriendlyShaderPlatformName_VulkanPC_Mobile", "PC Vulkan Mobile")
 
 [PreviewPlatform VULKAN_SM5]
diff --git a/Engine/Config/Windows/DataDrivenPlatformInfo.ini b/Engine/Config/Windows/DataDrivenPlatformInfo.ini
--- a/Engine/Config/Windows/DataDrivenPlatformInfo.ini
+++ b/Engine/Config/Windows/DataDrivenPlatformInfo.ini
@@ -157,6 +157,7 @@ bSupportsManualVertexFetch=false
 bSupportsSceneDataCompressedTransforms=true
 bSupportsDebugViewShaders=true
 bSupportsShaderPipelines = false
+bSupportsDualSourceBlending=true
 FriendlyName=LOCTEXT("FriendlyShaderPlatformName_D3D_Mobile", "PC D3D Mobile")
 SupportsMultiViewport=RuntimeDependent
 
diff --git a/Engine/Shaders/Private/MobileBasePassPixelShader.usf b/Engine/Shaders/Private/MobileBasePassPixelShader.usf
--- a/Engine/Shaders/Private/MobileBasePassPixelShader.usf
+++ b/Engine/Shaders/Private/MobileBasePassPixelShader.usf
@@ -159,6 +159,18 @@ half3 FrameBufferBlendOp(half4 Source)
 #endif
 }
 
+#if MOBILE_TRANSLUCENT_COLOR_TRANSMITTANCE_PROGRAMMABLE_BLENDING
+	#if VULKAN_PROFILE
+		#define FramebufferFetchColor0 VulkanSubpassFetch0
+	#elif METAL_PROFILE && !MAC
+		#define FramebufferFetchColor0 SubpassFetchRGBA_0
+	#elif COMPILER_GLSL_ES3_1
+		#define FramebufferFetchColor0 FramebufferFetchES2
+	#else
+		#error Platform unsupported for programmable blending
+	#endif
+#endif // MOBILE_TRANSLUCENT_COLOR_TRANSMITTANCE_PROGRAMMABLE_BLENDING
+
 void ApplyPixelDepthOffsetForMobileBasePass(inout FMaterialPixelParameters MaterialParameters, FPixelMaterialInputs PixelMaterialInputs, out float OutDepth)
 {
     float PixelDepthOffset = ApplyPixelDepthOffsetToMaterialParameters(MaterialParameters, PixelMaterialInputs, OutDepth);
@@ -279,6 +291,8 @@ void Main(
 #if DEFERRED_SHADING_PATH
 	#if USE_GLES_FBF_DEFERRED
 	, out HALF4_TYPE OutProxy : SV_Target0
+	#elif MOBILE_TRANSLUCENT_COLOR_TRANSMITTANCE_PROGRAMMABLE_BLENDING
+	, out HALF4_TYPE OutProgrammableBlending : SV_Target0
 	#else
 	, out HALF4_TYPE OutColor : SV_Target0
 	#endif
@@ -289,9 +303,11 @@ void Main(
 	, out HALF4_TYPE OutGBufferD : SV_Target4
 	#endif
 #else
-	#if DUAL_SOURCE_COLOR_BLENDING_ENABLED
+	#if MOBILE_TRANSLUCENT_COLOR_TRANSMITTANCE_DUAL_SRC_BLENDING
 	, out HALF4_TYPE OutColor DUAL_SOURCE_BLENDING_SLOT(0) : SV_Target0
 	, out HALF4_TYPE OutColor1 DUAL_SOURCE_BLENDING_SLOT(1) : SV_Target1
+	#elif MOBILE_TRANSLUCENT_COLOR_TRANSMITTANCE_PROGRAMMABLE_BLENDING
+	, out HALF4_TYPE OutProgrammableBlending : SV_Target0
 	#else
 	, out HALF4_TYPE OutColor : SV_Target0
 	#endif
@@ -305,13 +321,17 @@ void Main(
 	)
 {  
 #if DEFERRED_SHADING_PATH 
-	#if USE_GLES_FBF_DEFERRED
+	#if USE_GLES_FBF_DEFERRED && !MOBILE_TRANSLUCENT_COLOR_TRANSMITTANCE_PROGRAMMABLE_BLENDING
 		half4 OutColor;
 	#endif
 	#if !MOBILE_EXTENDED_GBUFFER
 		half4 OutGBufferD;
 	#endif
 #endif
+#if MOBILE_TRANSLUCENT_COLOR_TRANSMITTANCE_PROGRAMMABLE_BLENDING
+	half4 OutColor;
+	half4 OutColor1;
+#endif
 
 #if MOBILE_MULTI_VIEW
 	ResolvedView = ResolveView(BasePassInterpolants.MultiViewId);
@@ -978,13 +998,13 @@ void Main(
 		// Fade the surface color transmittance out to 1 according to the surface coverage, and take into account the fog coverage to the surface.
 		half3 AdjustedDualBlendMul = lerp(1.0f, VertexFog.a * DualBlendSurfaceTransmittancePreCoverage, DualBlendSurfaceCoverage);
 
-		#if DUAL_SOURCE_COLOR_BLENDING_ENABLED
-			OutColor = half4(AdjustedDualBlendAdd,0.0);
-			OutColor1 = half4(AdjustedDualBlendMul,1.0);
-		#else
+		#if MOBILE_TRANSLUCENT_COLOR_TRANSMITTANCE_SINGLE_SRC_BLENDING
 			// In the fallback case, use standard alpha blending(grey scale transmittance)
 			half AdjustedAlpha = saturate(1 - dot(AdjustedDualBlendMul, half3(1.0f, 1.0f, 1.0f) / 3.0f));
 			OutColor = half4(AdjustedDualBlendAdd, AdjustedAlpha);
+		#else
+			OutColor = half4(AdjustedDualBlendAdd, 0.0);
+			OutColor1 = half4(AdjustedDualBlendMul, 1.0);
 		#endif
 	// On mobile, water (an opaque material) is rendered as trnaslucent with forced premultiplied alpha blending (see MobileBasePass::SetTranslucentRenderState)
 	#elif MATERIALBLENDING_ALPHACOMPOSITE || MATERIAL_SHADINGMODEL_SINGLELAYERWATER
@@ -1050,8 +1070,9 @@ void Main(
 	);
 #endif
 
-#if DEFERRED_SHADING_PATH && USE_GLES_FBF_DEFERRED 
+#if DEFERRED_SHADING_PATH && USE_GLES_FBF_DEFERRED
 	OutProxy.rgb = OutColor.rgb;
+#elif MOBILE_TRANSLUCENT_COLOR_TRANSMITTANCE_PROGRAMMABLE_BLENDING
+	OutProgrammableBlending = OutColor1 * FramebufferFetchColor0() + OutColor;
 #endif
-
 }
diff --git a/Engine/Source/Developer/ShaderFormatOpenGL/Private/OpenGLShaderCompiler.cpp b/Engine/Source/Developer/ShaderFormatOpenGL/Private/OpenGLShaderCompiler.cpp
--- a/Engine/Source/Developer/ShaderFormatOpenGL/Private/OpenGLShaderCompiler.cpp
+++ b/Engine/Source/Developer/ShaderFormatOpenGL/Private/OpenGLShaderCompiler.cpp
@@ -2095,37 +2095,39 @@ static void ConvertToEmulatedUBs(std::string& GlslSource, const ReflectionData&
 }
 
 const ANSICHAR* GlslFrameBufferExtensions =
-"\n\n#ifdef GL_ARM_shader_framebuffer_fetch_depth_stencil\n"
-"\t#extension GL_ARM_shader_framebuffer_fetch_depth_stencil : enable\n"
-"#elif defined(GL_EXT_shader_framebuffer_fetch)\n"
+"\n\n#ifdef UE_MRT_FRAMEBUFFER_FETCH\n"
 "\t#extension GL_EXT_shader_framebuffer_fetch : enable\n"
 "\t#define FBF_STORAGE_QUALIFIER inout\n"
+"#elif defined(GL_ARM_shader_framebuffer_fetch)\n"
+"\t#extension GL_ARM_shader_framebuffer_fetch : enable\n"
+"#endif\n"
+"#ifdef GL_ARM_shader_framebuffer_fetch_depth_stencil\n"
+"\t#extension GL_ARM_shader_framebuffer_fetch_depth_stencil : enable\n"
 "#endif\n"
-"#extension GL_EXT_texture_buffer : enable\n"
 "// end extensions";
 
 // GLSL framebuffer macro definitions. Used to patch GLSL output source.
 const ANSICHAR* GlslFrameBufferDefines =
-"\n\n#ifdef UE_EXT_shader_framebuffer_fetch\n"
-"#define _Globals_ARM_shader_framebuffer_fetch 0\n"
-"#define FRAME_BUFFERFETCH_STORAGE_QUALIFIER inout\n"
-"#define _Globals_gl_FragColor out_var_SV_Target0\n"
-"#define _Globals_gl_LastFragColorARM vec4(0.0, 0.0, 0.0, 0.0)\n"
-"#elif defined( GL_ARM_shader_framebuffer_fetch)\n"
-"#define _Globals_ARM_shader_framebuffer_fetch 1\n"
-"#define FRAME_BUFFERFETCH_STORAGE_QUALIFIER out\n"
-"#define _Globals_gl_FragColor vec4(0.0, 0.0, 0.0, 0.0)\n"
-"#define _Globals_gl_LastFragColorARM gl_LastFragDepthARM\n"
+"\n\n#ifdef UE_MRT_FRAMEBUFFER_FETCH\n"
+"\t#define _Globals_ARM_shader_framebuffer_fetch 0u\n"
+"\t#define FRAME_BUFFERFETCH_STORAGE_QUALIFIER inout\n"
+"\t#define _Globals_gl_FragColor out_var_SV_Target0\n"
+"\t#define _Globals_gl_LastFragColorARM vec4(0.0, 0.0, 0.0, 0.0)\n"
+"#elif defined(GL_ARM_shader_framebuffer_fetch)\n"
+"\t#define _Globals_ARM_shader_framebuffer_fetch 1u\n"
+"\t#define FRAME_BUFFERFETCH_STORAGE_QUALIFIER out\n"
+"\t#define _Globals_gl_FragColor vec4(0.0, 0.0, 0.0, 0.0)\n"
+"\t#define _Globals_gl_LastFragColorARM gl_LastFragColorARM\n"
 "#else\n"
-"#define FRAME_BUFFERFETCH_STORAGE_QUALIFIER out\n"
-"#define _Globals_ARM_shader_framebuffer_fetch 0\n"
-"#define _Globals_gl_FragColor vec4(0.0, 0.0, 0.0, 0.0)\n"
-"#define _Globals_gl_LastFragColorARM vec4(0.0, 0.0, 0.0, 0.0)\n"
+"\t#define FRAME_BUFFERFETCH_STORAGE_QUALIFIER out\n"
+"\t#define _Globals_ARM_shader_framebuffer_fetch 0u\n"
+"\t#define _Globals_gl_FragColor vec4(0.0, 0.0, 0.0, 0.0)\n"
+"\t#define _Globals_gl_LastFragColorARM vec4(0.0, 0.0, 0.0, 0.0)\n"
 "#endif\n"
 "#ifdef GL_ARM_shader_framebuffer_fetch_depth_stencil\n"
-"#define _Globals_ARM_shader_framebuffer_fetch_depth_stencil 1u\n"
+"\t#define _Globals_ARM_shader_framebuffer_fetch_depth_stencil 1u\n"
 "#else\n"
-"#define _Globals_ARM_shader_framebuffer_fetch_depth_stencil 0u\n"
+"\t#define _Globals_ARM_shader_framebuffer_fetch_depth_stencil 0u\n"
 "#endif\n";
 
 struct GLSLCompileParameters
@@ -2187,8 +2189,8 @@ bool GenerateGlslShader(std::string& OutString, GLSLCompileParameters& GLSLCompi
 	}
 
 	// Perform FBF replacements
+	if (GLSLCompileParams.Frequency == SF_Pixel)
 	{
-
 		size_t MainPos = OutString.find("#version 320 es");
 
 		// Fallback if the shader is 310 
@@ -2204,105 +2206,132 @@ bool GenerateGlslShader(std::string& OutString, GLSLCompileParameters& GLSLCompi
 
 		// Framebuffer Depth Fetch
 		{
-			std::string FBDFString = "_Globals.ARM_shader_framebuffer_fetch_depth_stencil";
-			std::string FBDFReplaceString = "1u";
-
 			std::string FBFString = "_Globals.ARM_shader_framebuffer_fetch";
 			std::string FBFReplaceString = "_Globals_ARM_shader_framebuffer_fetch";
 
-			bool UsesFramebufferDepthFetch = GLSLCompileParams.Frequency == SF_Pixel && OutString.find(FBDFString) != std::string::npos;
-			bool UsesFramebufferFetch = GLSLCompileParams.Frequency == SF_Pixel && OutString.find(FBFString) != std::string::npos;
+			bool UsesFramebufferFetch = OutString.find(FBFString) != std::string::npos;
 
-			if (UsesFramebufferDepthFetch || UsesFramebufferFetch)
+			if (UsesFramebufferFetch)
 			{
 				OutString.insert(MainPos, GlslFrameBufferDefines);
 				OutString.insert(MainPos, GlslFrameBufferExtensions);
 
-				size_t FramebufferDepthFetchPos = OutString.find(FBDFString);
-
-				OutString.erase(FramebufferDepthFetchPos, FBDFString.length());
-				OutString.insert(FramebufferDepthFetchPos, FBDFReplaceString);
-
-				std::string LastFragDepthARMString = "_Globals._RESERVED_IDENTIFIER_FIXUP_gl_LastFragDepthARM";
-				std::string LastFragDepthARMReplaceString = "GLFetchDepthBuffer()";
+				std::string FBDFString = "_Globals.ARM_shader_framebuffer_fetch_depth_stencil";
+				std::string FBDFReplaceString = "1u";
 
-				size_t LastFragDepthARMStringPos = OutString.find(LastFragDepthARMString);
-				
-				while (LastFragDepthARMStringPos != std::string::npos)
+				bool UsesFramebufferDepthFetch = OutString.find(FBDFString) != std::string::npos;
+				if (UsesFramebufferDepthFetch)
 				{
-					OutString.erase(LastFragDepthARMStringPos, LastFragDepthARMString.length());
-					OutString.insert(LastFragDepthARMStringPos, LastFragDepthARMReplaceString);
+					size_t FramebufferDepthFetchPos = OutString.find(FBDFString);
 
-					LastFragDepthARMStringPos = OutString.find(LastFragDepthARMString);
-				}
+					OutString.erase(FramebufferDepthFetchPos, FBDFString.length());
+					OutString.insert(FramebufferDepthFetchPos, FBDFReplaceString);
 
-				MainPos = OutString.find("void main()");
+					std::string LastFragDepthARMString = "_Globals._RESERVED_IDENTIFIER_FIXUP_gl_LastFragDepthARM";
+					std::string LastFragDepthARMReplaceString = "GLFetchDepthBuffer()";
 
-				// Add support for framebuffer fetch depth when ARM extension is not supported
-				if (MainPos != std::string::npos)
-				{
-					std::string DepthBufferIndex = bIsDeferred ? "4" : "1";
-					std::string DepthBufferOutVarString = "out_var_SV_Target" + DepthBufferIndex;
-
-					// Insert function declaration to handle retrieving depth
-					OutString.insert(MainPos, "float GLFetchDepthBuffer()\n"
-						"{\n"
-						"\t#if defined(GL_ARM_shader_framebuffer_fetch_depth_stencil)\n"
-						"\treturn gl_LastFragDepthARM;\n"
-						"\t#elif defined(GL_EXT_shader_framebuffer_fetch)\n"
-						"\treturn " + DepthBufferOutVarString + ".x;\n"
-						"\t#else\n"
-						"\treturn 0.0f;\n"
-						"\t#endif\n"
-						"}\n");
-
-					
-					// If SceneDepthAux is not declared then declare it, otherwise modify so that we only enable it on devices that don't support
-					// GL_ARM_shader_framebuffer_fetch_depth_stencil and do support GL_EXT_shader_framebuffer_fetch
-					size_t DepthBufferOutVarPos = OutString.find(DepthBufferOutVarString + ";");
-
-					std::string DepthBufferDeclString = "layout(location = " + DepthBufferIndex + ") inout highp vec4 " + DepthBufferOutVarString + ";\n";
-					std::string DepthBufferOutString = "\n#if !defined(GL_ARM_shader_framebuffer_fetch_depth_stencil) && defined(GL_EXT_shader_framebuffer_fetch)\n" +
-														DepthBufferDeclString +
-														"#endif\n";
-
-					// If we cannot find a declararation of out_var_SV_Target(n) in the shader, insert one
-					if (DepthBufferOutVarPos == std::string::npos)
-					{
-						OutString.insert(MainPos, DepthBufferOutString);
-					}
-					else 
+					size_t LastFragDepthARMStringPos = OutString.find(LastFragDepthARMString);
+
+					while (LastFragDepthARMStringPos != std::string::npos)
 					{
-						// If we have a declaration, replace with one that will be stripped if GL_ARM_shader_framebuffer_fetch_depth_stencil is enabled
-						size_t StringStartPos = OutString.rfind("layout", DepthBufferOutVarPos - 1);
-						size_t StringEndPos = OutString.find(";", StringStartPos);
+						OutString.erase(LastFragDepthARMStringPos, LastFragDepthARMString.length());
+						OutString.insert(LastFragDepthARMStringPos, LastFragDepthARMReplaceString);
 
-						OutString.erase(StringStartPos, (StringEndPos+1) - StringStartPos);
-						OutString.insert(StringStartPos, DepthBufferOutString);
+						LastFragDepthARMStringPos = OutString.find(LastFragDepthARMString);
 					}
 
-					// Make SceneDepthAux assignment conditional
-					// We only need to write the depth when we don't support GL_ARM_shader_framebuffer_fetch_depth_stencil
-					std::string DepthBufferAssignment = DepthBufferOutVarString + " =";
-					size_t DepthBufferAssignmentPos = OutString.find(DepthBufferAssignment);
+					MainPos = OutString.find("void main()");
 
-					if (DepthBufferAssignmentPos != std::string::npos)
+					// Add support for framebuffer fetch depth when ARM extension is not supported
+					if (MainPos != std::string::npos)
 					{
-						size_t LineEnd = OutString.find_first_of(";", DepthBufferAssignmentPos);
-						uint32_t AssignmentValueStart = DepthBufferAssignmentPos + DepthBufferAssignment.size();
-						std::string AssignmentValue = OutString.substr(AssignmentValueStart + 1, LineEnd - AssignmentValueStart);
+						std::string DepthBufferIndex = bIsDeferred ? "4" : "1";
+						std::string DepthBufferOutVarString = "out_var_SV_Target" + DepthBufferIndex;
+
+						// Insert function declaration to handle retrieving depth
+						OutString.insert(MainPos, "float GLFetchDepthBuffer()\n"
+							"{\n"
+							"\t#if defined(GL_ARM_shader_framebuffer_fetch_depth_stencil)\n"
+							"\treturn gl_LastFragDepthARM;\n"
+							"\t#elif defined(GL_EXT_shader_framebuffer_fetch)\n"
+							"\treturn " + DepthBufferOutVarString + ".x;\n"
+							"\t#else\n"
+							"\treturn 0.0f;\n"
+							"\t#endif\n"
+							"}\n");
+
+
+						// If SceneDepthAux is not declared then declare it, otherwise modify so that we only enable it on devices that don't support
+						// GL_ARM_shader_framebuffer_fetch_depth_stencil and do support GL_EXT_shader_framebuffer_fetch
+						size_t DepthBufferOutVarPos = OutString.find(DepthBufferOutVarString + ";");
+
+						std::string DepthBufferDeclString = "layout(location = " + DepthBufferIndex + ") inout highp vec4 " + DepthBufferOutVarString + ";\n";
+						std::string DepthBufferOutString = "\n#if !defined(GL_ARM_shader_framebuffer_fetch_depth_stencil) && defined(GL_EXT_shader_framebuffer_fetch)\n" +
+							DepthBufferDeclString +
+							"#endif\n";
+
+						// If we cannot find a declararation of out_var_SV_Target(n) in the shader, insert one
+						if (DepthBufferOutVarPos == std::string::npos)
+						{
+							OutString.insert(MainPos, DepthBufferOutString);
+						}
+						else
+						{
+							// If we have a declaration, replace with one that will be stripped if GL_ARM_shader_framebuffer_fetch_depth_stencil is enabled
+							size_t StringStartPos = OutString.rfind("layout", DepthBufferOutVarPos - 1);
+							size_t StringEndPos = OutString.find(";", StringStartPos);
 
-						if (LineEnd != std::string::npos)
+							OutString.erase(StringStartPos, (StringEndPos + 1) - StringStartPos);
+							OutString.insert(StringStartPos, DepthBufferOutString);
+						}
+
+						// Make SceneDepthAux assignment conditional
+						// We only need to write the depth when we don't support GL_ARM_shader_framebuffer_fetch_depth_stencil
+						std::string DepthBufferAssignment = DepthBufferOutVarString + " =";
+						size_t DepthBufferAssignmentPos = OutString.find(DepthBufferAssignment);
+
+						if (DepthBufferAssignmentPos != std::string::npos)
 						{
-							OutString.erase(DepthBufferAssignmentPos, LineEnd + 1 - DepthBufferAssignmentPos);
-							OutString.insert(DepthBufferAssignmentPos, std::string("#if !defined(GL_ARM_shader_framebuffer_fetch_depth_stencil) && defined(GL_EXT_shader_framebuffer_fetch)\n") + DepthBufferAssignment + AssignmentValue + std::string("\n#endif\n"));
+							size_t LineEnd = OutString.find_first_of(";", DepthBufferAssignmentPos);
+							uint32_t AssignmentValueStart = DepthBufferAssignmentPos + DepthBufferAssignment.size();
+							std::string AssignmentValue = OutString.substr(AssignmentValueStart + 1, LineEnd - AssignmentValueStart);
+
+							if (LineEnd != std::string::npos)
+							{
+								OutString.erase(DepthBufferAssignmentPos, LineEnd + 1 - DepthBufferAssignmentPos);
+								OutString.insert(DepthBufferAssignmentPos, std::string("#if !defined(GL_ARM_shader_framebuffer_fetch_depth_stencil) && defined(GL_EXT_shader_framebuffer_fetch)\n") + DepthBufferAssignment + AssignmentValue + std::string("\n#endif\n"));
+							}
 						}
 					}
+				} // UsesFramebufferDepthFetch
+			}
+
+			// check again, before it might have been a match for ARM_shader_framebuffer_fetch_depth_stencil
+			UsesFramebufferFetch = GLSLCompileParams.Frequency == SF_Pixel && OutString.find(FBFString) != std::string::npos;
+			if (UsesFramebufferFetch)
+			{
+				std::string ReservedIdentifierFixupString = "_RESERVED_IDENTIFIER_FIXUP_";
+				size_t ReservedIdentifierFixupStringPos = OutString.find(ReservedIdentifierFixupString);
+
+				while (ReservedIdentifierFixupStringPos != std::string::npos)
+				{
+					OutString.erase(ReservedIdentifierFixupStringPos, ReservedIdentifierFixupString.length());
+					ReservedIdentifierFixupStringPos = OutString.find(ReservedIdentifierFixupString);
+				}
+
+				std::string Target0String = "layout(location = 0) out ";
+				std::string Target0ReplaceString = "layout(location = 0) FRAME_BUFFERFETCH_STORAGE_QUALIFIER ";
+				size_t Target0StringPos = OutString.find(Target0String);
+				if (Target0StringPos != std::string::npos)
+				{
+					OutString.erase(Target0StringPos, Target0String.length());
+					OutString.insert(Target0StringPos, Target0ReplaceString);
 				}
 			}
 		}
 	}
 
+
 	// If we are rendering deferred, then we only need SceneDepthAux on devices that don't support framebuffer fetch depth
 	if (bIsDeferred)
 	{
diff --git a/Engine/Source/Runtime/Engine/Private/Materials/MaterialShared.cpp b/Engine/Source/Runtime/Engine/Private/Materials/MaterialShared.cpp
--- a/Engine/Source/Runtime/Engine/Private/Materials/MaterialShared.cpp
+++ b/Engine/Source/Runtime/Engine/Private/Materials/MaterialShared.cpp
@@ -1860,7 +1860,7 @@ bool FMaterialResource::IsDualBlendingEnabled(EShaderPlatform Platform) const
 	{
 		bMaterialRequestsDualSourceBlending = GetBlendMode() == EBlendMode::BLEND_TranslucentColoredTransmittance;
 	}
-	const bool bIsPlatformSupported = RHISupportsDualSourceBlending(Platform);
+	const bool bIsPlatformSupported = RHISupportsDualSourceBlending(Platform) || IsMobilePlatform(Platform); // Mobile renderer has runtime fallbacks
 	return bMaterialRequestsDualSourceBlending && bIsPlatformSupported;
 }
 
diff --git a/Engine/Source/Runtime/OpenGLDrv/Private/OpenGLDevice.cpp b/Engine/Source/Runtime/OpenGLDrv/Private/OpenGLDevice.cpp
--- a/Engine/Source/Runtime/OpenGLDrv/Private/OpenGLDevice.cpp
+++ b/Engine/Source/Runtime/OpenGLDrv/Private/OpenGLDevice.cpp
@@ -1124,6 +1124,8 @@ static void InitRHICapabilitiesForGL()
 	GShaderPlatformForFeatureLevel[ERHIFeatureLevel::SM4_REMOVED] = SP_NumPlatforms;
 	GShaderPlatformForFeatureLevel[ERHIFeatureLevel::SM5] = SP_NumPlatforms;
 
+	// Not implemented in OpenGL RHI. Supported in GLcore 3.3, not in OpenGL ES, extension GL_EXT_blend_func_extended not widely supported.
+	GSupportsDualSrcBlending = false;
 
 	GRHISupportsTextureStreaming = true;
 
diff --git a/Engine/Source/Runtime/RHI/Public/RHIGlobals.h b/Engine/Source/Runtime/RHI/Public/RHIGlobals.h
--- a/Engine/Source/Runtime/RHI/Public/RHIGlobals.h
+++ b/Engine/Source/Runtime/RHI/Public/RHIGlobals.h
@@ -212,6 +212,9 @@ struct FRHIGlobals
 	/** True if the RHI supports separate blend states per render target. */
 	bool SupportsSeparateRenderTargetBlendState = false;
 
+	/** True if the RHI supports dual src blending. */
+	bool SupportsDualSrcBlending = true;
+
 	/** True if the RHI has artifacts with atlased CSM depths. */
 	bool NeedsUnatlasedCSMDepthsWorkaround = false;
 
@@ -677,6 +680,7 @@ extern RHI_API FRHIGlobals GRHIGlobals;
 #define GRHISupportsExactOcclusionQueries                      GRHIGlobals.SupportsExactOcclusionQueries
 #define GSupportsVolumeTextureRendering                        GRHIGlobals.SupportsVolumeTextureRendering
 #define GSupportsSeparateRenderTargetBlendState                GRHIGlobals.SupportsSeparateRenderTargetBlendState
+#define GSupportsDualSrcBlending                               GRHIGlobals.SupportsDualSrcBlending
 #define GRHINeedsUnatlasedCSMDepthsWorkaround                  GRHIGlobals.NeedsUnatlasedCSMDepthsWorkaround
 #define GSupportsTexture3D                                     GRHIGlobals.SupportsTexture3D
 #define GUseTexture3DBulkDataRHI                               GRHIGlobals.UseTexture3DBulkData
diff --git a/Engine/Source/Runtime/Renderer/Private/MobileBasePass.cpp b/Engine/Source/Runtime/Renderer/Private/MobileBasePass.cpp
--- a/Engine/Source/Runtime/Renderer/Private/MobileBasePass.cpp
+++ b/Engine/Source/Runtime/Renderer/Private/MobileBasePass.cpp
@@ -99,11 +99,39 @@ bool MobileUsesNoLightMapPermutation(const FMeshMaterialShaderPermutationParamet
 	return true;
 }
 
+template <ELightMapPolicyType Policy, EMobileLocalLightSetting LocalLightSetting, EMobileTranslucentColorTransmittanceMode ThinTranslucencyFallback>
+static void AddMobileBasePassPixelShaderTypes(FMaterialShaderTypes& ShaderTypes, const bool bIsMobileHDR, const bool bEnableSkyLight)
+{
+	if (bIsMobileHDR)
+	{
+		if (bEnableSkyLight)
+		{
+			ShaderTypes.AddShaderType<TMobileBasePassPS<TUniformLightMapPolicy<Policy>, HDR_LINEAR_64, true, LocalLightSetting, ThinTranslucencyFallback>>();
+		}
+		else
+		{
+			ShaderTypes.AddShaderType<TMobileBasePassPS<TUniformLightMapPolicy<Policy>, HDR_LINEAR_64, false, LocalLightSetting, ThinTranslucencyFallback>>();
+		}
+	}
+	else
+	{
+		if (bEnableSkyLight)
+		{
+			ShaderTypes.AddShaderType<TMobileBasePassPS<TUniformLightMapPolicy<Policy>, LDR_GAMMA_32, true, LocalLightSetting, ThinTranslucencyFallback>>();
+		}
+		else
+		{
+			ShaderTypes.AddShaderType<TMobileBasePassPS<TUniformLightMapPolicy<Policy>, LDR_GAMMA_32, false, LocalLightSetting, ThinTranslucencyFallback>>();
+		}
+	}
+}
+
 template <ELightMapPolicyType Policy, EMobileLocalLightSetting LocalLightSetting>
 bool GetUniformMobileBasePassShaders(
 	const FMaterial& Material, 
 	const FVertexFactoryType* VertexFactoryType, 
 	bool bEnableSkyLight,
+	EMobileTranslucentColorTransmittanceMode ThinTranslucentFallback,
 	TShaderRef<TMobileBasePassVSPolicyParamType<FUniformLightMapPolicy>>& VertexShader,
 	TShaderRef<TMobileBasePassPSPolicyParamType<FUniformLightMapPolicy>>& PixelShader
 	)
@@ -111,32 +139,24 @@ bool GetUniformMobileBasePassShaders(
 	using FVertexShaderType = TMobileBasePassVSPolicyParamType<FUniformLightMapPolicy>;
 	using FPixelShaderType = TMobileBasePassPSPolicyParamType<FUniformLightMapPolicy>;
 
+	const bool bIsMobileHDR = IsMobileHDR();
 	FMaterialShaderTypes ShaderTypes;
-	if (IsMobileHDR())
+	if (bIsMobileHDR)
 	{
-		ShaderTypes.AddShaderType<TMobileBasePassVS<TUniformLightMapPolicy<Policy>, HDR_LINEAR_64>>();
-
-		if (bEnableSkyLight)
-		{
-			ShaderTypes.AddShaderType<TMobileBasePassPS<TUniformLightMapPolicy<Policy>, HDR_LINEAR_64, true, LocalLightSetting>>();
-		}
-		else
-		{
-			ShaderTypes.AddShaderType<TMobileBasePassPS<TUniformLightMapPolicy<Policy>, HDR_LINEAR_64, false, LocalLightSetting>>();
-		}	
+		ShaderTypes.AddShaderType<TMobileBasePassVS<TUniformLightMapPolicy<Policy>, HDR_LINEAR_64>>();	
 	}
 	else
 	{
-		ShaderTypes.AddShaderType<TMobileBasePassVS<TUniformLightMapPolicy<Policy>, LDR_GAMMA_32>>();
+		ShaderTypes.AddShaderType<TMobileBasePassVS<TUniformLightMapPolicy<Policy>, LDR_GAMMA_32>>();			
+	}
 
-		if (bEnableSkyLight)
-		{
-			ShaderTypes.AddShaderType<TMobileBasePassPS<TUniformLightMapPolicy<Policy>, LDR_GAMMA_32, true, LocalLightSetting>>();
-		}
-		else
-		{
-			ShaderTypes.AddShaderType<TMobileBasePassPS<TUniformLightMapPolicy<Policy>, LDR_GAMMA_32, false, LocalLightSetting>>();
-		}			
+	switch (ThinTranslucentFallback)
+	{
+	default:
+	case EMobileTranslucentColorTransmittanceMode::DEFAULT:
+		AddMobileBasePassPixelShaderTypes<Policy, LocalLightSetting, EMobileTranslucentColorTransmittanceMode::DEFAULT>(ShaderTypes, bIsMobileHDR, bEnableSkyLight); break;
+	case EMobileTranslucentColorTransmittanceMode::SINGLE_SRC_BLENDING:
+		AddMobileBasePassPixelShaderTypes<Policy, LocalLightSetting, EMobileTranslucentColorTransmittanceMode::SINGLE_SRC_BLENDING>(ShaderTypes, bIsMobileHDR, bEnableSkyLight); break;
 	}
 
 	FMaterialShaders Shaders;
@@ -156,6 +176,7 @@ bool GetMobileBasePassShaders(
 	const FMaterial& Material, 
 	const FVertexFactoryType* VertexFactoryType, 
 	bool bEnableSkyLight,
+	EMobileTranslucentColorTransmittanceMode ThinTranslucentFallback,
 	TShaderRef<TMobileBasePassVSPolicyParamType<FUniformLightMapPolicy>>& VertexShader,
 	TShaderRef<TMobileBasePassPSPolicyParamType<FUniformLightMapPolicy>>& PixelShader
 	)
@@ -163,25 +184,25 @@ bool GetMobileBasePassShaders(
 	switch (LightMapPolicyType)
 	{
 	case LMP_NO_LIGHTMAP:
-		return GetUniformMobileBasePassShaders<LMP_NO_LIGHTMAP, LocalLightSetting>(Material, VertexFactoryType, bEnableSkyLight, VertexShader, PixelShader);
+		return GetUniformMobileBasePassShaders<LMP_NO_LIGHTMAP, LocalLightSetting>(Material, VertexFactoryType, bEnableSkyLight, ThinTranslucentFallback, VertexShader, PixelShader);
 	case LMP_LQ_LIGHTMAP:
-		return GetUniformMobileBasePassShaders<LMP_LQ_LIGHTMAP, LocalLightSetting>(Material, VertexFactoryType, bEnableSkyLight, VertexShader, PixelShader);
+		return GetUniformMobileBasePassShaders<LMP_LQ_LIGHTMAP, LocalLightSetting>(Material, VertexFactoryType, bEnableSkyLight, ThinTranslucentFallback, VertexShader, PixelShader);
 	case LMP_MOBILE_DISTANCE_FIELD_SHADOWS_AND_LQ_LIGHTMAP:
-		return GetUniformMobileBasePassShaders<LMP_MOBILE_DISTANCE_FIELD_SHADOWS_AND_LQ_LIGHTMAP, LocalLightSetting>(Material, VertexFactoryType, bEnableSkyLight, VertexShader, PixelShader);
+		return GetUniformMobileBasePassShaders<LMP_MOBILE_DISTANCE_FIELD_SHADOWS_AND_LQ_LIGHTMAP, LocalLightSetting>(Material, VertexFactoryType, bEnableSkyLight, ThinTranslucentFallback, VertexShader, PixelShader);
 	case LMP_MOBILE_DISTANCE_FIELD_SHADOWS_LIGHTMAP_AND_CSM:
-		return GetUniformMobileBasePassShaders<LMP_MOBILE_DISTANCE_FIELD_SHADOWS_LIGHTMAP_AND_CSM, LocalLightSetting>(Material, VertexFactoryType, bEnableSkyLight, VertexShader, PixelShader);
+		return GetUniformMobileBasePassShaders<LMP_MOBILE_DISTANCE_FIELD_SHADOWS_LIGHTMAP_AND_CSM, LocalLightSetting>(Material, VertexFactoryType, bEnableSkyLight, ThinTranslucentFallback, VertexShader, PixelShader);
 	case LMP_MOBILE_DIRECTIONAL_LIGHT_CSM_AND_LIGHTMAP:
-		return GetUniformMobileBasePassShaders<LMP_MOBILE_DIRECTIONAL_LIGHT_CSM_AND_LIGHTMAP, LocalLightSetting>(Material, VertexFactoryType, bEnableSkyLight, VertexShader, PixelShader);
+		return GetUniformMobileBasePassShaders<LMP_MOBILE_DIRECTIONAL_LIGHT_CSM_AND_LIGHTMAP, LocalLightSetting>(Material, VertexFactoryType, bEnableSkyLight, ThinTranslucentFallback, VertexShader, PixelShader);
 	case LMP_MOBILE_DIRECTIONAL_LIGHT_AND_SH_INDIRECT:
-		return GetUniformMobileBasePassShaders<LMP_MOBILE_DIRECTIONAL_LIGHT_AND_SH_INDIRECT, LocalLightSetting>(Material, VertexFactoryType, bEnableSkyLight, VertexShader, PixelShader);
+		return GetUniformMobileBasePassShaders<LMP_MOBILE_DIRECTIONAL_LIGHT_AND_SH_INDIRECT, LocalLightSetting>(Material, VertexFactoryType, bEnableSkyLight, ThinTranslucentFallback, VertexShader, PixelShader);
 	case LMP_MOBILE_DIRECTIONAL_LIGHT_CSM_AND_SH_INDIRECT:
-		return GetUniformMobileBasePassShaders<LMP_MOBILE_DIRECTIONAL_LIGHT_CSM_AND_SH_INDIRECT, LocalLightSetting>(Material, VertexFactoryType, bEnableSkyLight, VertexShader, PixelShader);
+		return GetUniformMobileBasePassShaders<LMP_MOBILE_DIRECTIONAL_LIGHT_CSM_AND_SH_INDIRECT, LocalLightSetting>(Material, VertexFactoryType, bEnableSkyLight, ThinTranslucentFallback, VertexShader, PixelShader);
 	case LMP_MOBILE_MOVABLE_DIRECTIONAL_LIGHT_WITH_LIGHTMAP:
-		return GetUniformMobileBasePassShaders<LMP_MOBILE_MOVABLE_DIRECTIONAL_LIGHT_WITH_LIGHTMAP, LocalLightSetting>(Material, VertexFactoryType, bEnableSkyLight, VertexShader, PixelShader);
+		return GetUniformMobileBasePassShaders<LMP_MOBILE_MOVABLE_DIRECTIONAL_LIGHT_WITH_LIGHTMAP, LocalLightSetting>(Material, VertexFactoryType, bEnableSkyLight, ThinTranslucentFallback, VertexShader, PixelShader);
 	case LMP_MOBILE_MOVABLE_DIRECTIONAL_LIGHT_CSM_WITH_LIGHTMAP:
-		return GetUniformMobileBasePassShaders<LMP_MOBILE_MOVABLE_DIRECTIONAL_LIGHT_CSM_WITH_LIGHTMAP, LocalLightSetting>(Material, VertexFactoryType, bEnableSkyLight, VertexShader, PixelShader);
+		return GetUniformMobileBasePassShaders<LMP_MOBILE_MOVABLE_DIRECTIONAL_LIGHT_CSM_WITH_LIGHTMAP, LocalLightSetting>(Material, VertexFactoryType, bEnableSkyLight, ThinTranslucentFallback, VertexShader, PixelShader);
 	case LMP_MOBILE_DIRECTIONAL_LIGHT_CSM:
-		return GetUniformMobileBasePassShaders<LMP_MOBILE_DIRECTIONAL_LIGHT_CSM, LocalLightSetting>(Material, VertexFactoryType, bEnableSkyLight, VertexShader, PixelShader);
+		return GetUniformMobileBasePassShaders<LMP_MOBILE_DIRECTIONAL_LIGHT_CSM, LocalLightSetting>(Material, VertexFactoryType, bEnableSkyLight, ThinTranslucentFallback, VertexShader, PixelShader);
 	default:										
 		check(false);
 		return true;
@@ -203,6 +224,13 @@ bool MobileBasePass::GetShaders(
 		bEnableSkyLight = !bEnableSkyLight;
 	}
 
+	EMobileTranslucentColorTransmittanceMode ThinTranslucencyFallback = EMobileTranslucentColorTransmittanceMode::DEFAULT;
+	if (MaterialResource.GetShadingModels().HasShadingModel(MSM_ThinTranslucent))
+	{
+		const EShaderPlatform ShaderPlatform = GetFeatureLevelShaderPlatform(MaterialResource.GetFeatureLevel());
+		ThinTranslucencyFallback = MobileActiveTranslucentColorTransmittanceMode(ShaderPlatform, false);
+	}
+
 	switch (LocalLightSetting)
 	{
 		case EMobileLocalLightSetting::LOCAL_LIGHTS_DISABLED:
@@ -212,6 +240,7 @@ bool MobileBasePass::GetShaders(
 				MaterialResource,
 				VertexFactoryType,
 				bEnableSkyLight,
+				ThinTranslucencyFallback,
 				VertexShader,
 				PixelShader
 				);
@@ -224,6 +253,7 @@ bool MobileBasePass::GetShaders(
 				MaterialResource,
 				VertexFactoryType,
 				bEnableSkyLight,
+				ThinTranslucencyFallback,
 				VertexShader,
 				PixelShader
 				);
@@ -236,6 +266,7 @@ bool MobileBasePass::GetShaders(
 				MaterialResource,
 				VertexFactoryType,
 				bEnableSkyLight,
+				ThinTranslucencyFallback,
 				VertexShader,
 				PixelShader
 				);
@@ -553,70 +584,75 @@ void MobileBasePass::SetOpaqueRenderState(FMeshPassProcessorRenderState& DrawRen
 	}
 }
 
-void MobileBasePass::SetTranslucentRenderState(FMeshPassProcessorRenderState& DrawRenderState, const FMaterial& Material, FMaterialShadingModelField ShadingModels)
+static FRHIBlendState* GetBlendStateForColorTransmittanceBlending(const EShaderPlatform ShaderPlatform)
 {
-	EShaderPlatform ShaderPlatform = GetFeatureLevelShaderPlatform(Material.GetFeatureLevel());
+	switch (MobileActiveTranslucentColorTransmittanceMode(ShaderPlatform, true))
+	{
+	case EMobileTranslucentColorTransmittanceMode::DUAL_SRC_BLENDING:
+		// Blend by putting add in target 0 and multiply by background in target 1.
+		return TStaticBlendState<CW_RGBA, BO_Add, BF_One, BF_Source1Color, BO_Add, BF_One, BF_Source1Alpha,
+			CW_NONE, BO_Add, BF_One, BF_Zero, BO_Add, BF_One, BF_Zero,
+			CW_NONE, BO_Add, BF_One, BF_Zero, BO_Add, BF_One, BF_Zero,
+			CW_NONE, BO_Add, BF_One, BF_Zero, BO_Add, BF_One, BF_Zero,
+			CW_NONE, BO_Add, BF_One, BF_Zero, BO_Add, BF_One, BF_Zero,
+			CW_NONE, BO_Add, BF_One, BF_Zero, BO_Add, BF_One, BF_Zero,
+			CW_NONE, BO_Add, BF_One, BF_Zero, BO_Add, BF_One, BF_Zero,
+			CW_NONE, BO_Add, BF_One, BF_Zero, BO_Add, BF_One, BF_Zero>::GetRHI();
+	case EMobileTranslucentColorTransmittanceMode::SINGLE_SRC_BLENDING:
+		// If a material was requesting dual source blending, the shader will use static platform knowledge to convert colored transmittance to a grey scale transmittance.
+		return TStaticBlendState<CW_RGBA, BO_Add, BF_One, BF_InverseSourceAlpha, BO_Add, BF_Zero, BF_InverseSourceAlpha,
+			CW_NONE, BO_Add, BF_One, BF_Zero, BO_Add, BF_One, BF_Zero,
+			CW_NONE, BO_Add, BF_One, BF_Zero, BO_Add, BF_One, BF_Zero,
+			CW_NONE, BO_Add, BF_One, BF_Zero, BO_Add, BF_One, BF_Zero,
+			CW_NONE, BO_Add, BF_One, BF_Zero, BO_Add, BF_One, BF_Zero,
+			CW_NONE, BO_Add, BF_One, BF_Zero, BO_Add, BF_One, BF_Zero,
+			CW_NONE, BO_Add, BF_One, BF_Zero, BO_Add, BF_One, BF_Zero,
+			CW_NONE, BO_Add, BF_One, BF_Zero, BO_Add, BF_One, BF_Zero>::GetRHI();
+	default:
+	case EMobileTranslucentColorTransmittanceMode::PROGRAMMABLE_BLENDING:
+		// Blending is done in shader
+		return TStaticBlendState<CW_RGBA, BO_Add, BF_One, BF_Zero, BO_Add, BF_One, BF_Zero,
+			CW_NONE, BO_Add, BF_One, BF_Zero, BO_Add, BF_One, BF_Zero,
+			CW_NONE, BO_Add, BF_One, BF_Zero, BO_Add, BF_One, BF_Zero,
+			CW_NONE, BO_Add, BF_One, BF_Zero, BO_Add, BF_One, BF_Zero,
+			CW_NONE, BO_Add, BF_One, BF_Zero, BO_Add, BF_One, BF_Zero,
+			CW_NONE, BO_Add, BF_One, BF_Zero, BO_Add, BF_One, BF_Zero,
+			CW_NONE, BO_Add, BF_One, BF_Zero, BO_Add, BF_One, BF_Zero,
+			CW_NONE, BO_Add, BF_One, BF_Zero, BO_Add, BF_One, BF_Zero>::GetRHI();
+	}
+}
 
+void MobileBasePass::SetTranslucentRenderState(FMeshPassProcessorRenderState& DrawRenderState, const FMaterial& Material, FMaterialShadingModelField ShadingModels)
+{
+	const EShaderPlatform ShaderPlatform = GetFeatureLevelShaderPlatform(Material.GetFeatureLevel());
 	const bool bIsUsingMobilePixelProjectedReflection = Material.IsUsingPlanarForwardReflections() && IsUsingMobilePixelProjectedReflection(ShaderPlatform);
-	const bool bIsDualSourceBlending = RHISupportsDualSourceBlending(ShaderPlatform);
-	const EShaderPlatform Platform = GetFeatureLevelShaderPlatform(Material.GetFeatureLevel());
 
 	if (Substrate::IsSubstrateEnabled())
 	{
-		if (Material.IsDualBlendingEnabled(Platform))
+		if (Material.IsDualBlendingEnabled(ShaderPlatform) || Material.GetBlendMode() == BLEND_TranslucentColoredTransmittance)
 		{
-			// If requested and available, we do standard dual blending, and the alpha gets ignored
-			// Blend by putting add in target 0 and multiply by background in target 1.
-			DrawRenderState.SetBlendState(TStaticBlendState<CW_RGBA, BO_Add, BF_One, BF_Source1Color, BO_Add, BF_One, BF_Source1Alpha>::GetRHI());
+			DrawRenderState.SetBlendState(GetBlendStateForColorTransmittanceBlending(ShaderPlatform));
 		}
-		else
+		else if (Material.GetBlendMode() == BLEND_ColoredTransmittanceOnly)
 		{
-			if (Material.GetBlendMode() == BLEND_ColoredTransmittanceOnly)
-			{
-				// Modulate with the existing scene color, preserve destination alpha.
-				DrawRenderState.SetBlendState(TStaticBlendState<CW_RGB, BO_Add, BF_DestColor, BF_Zero>::GetRHI());
-			}
-			else if (Material.GetBlendMode() == BLEND_AlphaHoldout)
-			{
-				// Blend by holding out the matte shape of the source alpha
-				DrawRenderState.SetBlendState(TStaticBlendState<CW_RGBA, BO_Add, BF_Zero, BF_InverseSourceAlpha, BO_Add, BF_One, BF_InverseSourceAlpha>::GetRHI());
-			}
-			else
-			{
-				// We always use premultiplied alpha for translucent rendering.
-				// If a material was requesting dual source blending, the shader will use static platform knowledge to convert colored transmittance to a grey scale transmittance.
-				DrawRenderState.SetBlendState(TStaticBlendState<CW_RGBA, BO_Add, BF_One, BF_InverseSourceAlpha, BO_Add, BF_Zero, BF_InverseSourceAlpha>::GetRHI());
-			}
+			// Modulate with the existing scene color, preserve destination alpha.
+			DrawRenderState.SetBlendState(TStaticBlendState<CW_RGB, BO_Add, BF_DestColor, BF_Zero>::GetRHI());
 		}
-	}
-	else if (ShadingModels.HasShadingModel(MSM_ThinTranslucent))
-	{
-		const bool bRequiresSceneDepthAux = MobileRequiresSceneDepthAux(ShaderPlatform);
-		if (bIsDualSourceBlending && !bRequiresSceneDepthAux)
+		else if (Material.GetBlendMode() == BLEND_AlphaHoldout)
 		{
-			// Blend by putting add in target 0 and multiply by background in target 1.
-			DrawRenderState.SetBlendState(TStaticBlendState<CW_RGBA, BO_Add, BF_One, BF_Source1Color, BO_Add, BF_One, BF_Source1Alpha,
-															CW_NONE, BO_Add, BF_One, BF_Zero, BO_Add, BF_One, BF_Zero,
-															CW_NONE, BO_Add, BF_One, BF_Zero, BO_Add, BF_One, BF_Zero,
-															CW_NONE, BO_Add, BF_One, BF_Zero, BO_Add, BF_One, BF_Zero,
-															CW_NONE, BO_Add, BF_One, BF_Zero, BO_Add, BF_One, BF_Zero,
-															CW_NONE, BO_Add, BF_One, BF_Zero, BO_Add, BF_One, BF_Zero,
-															CW_NONE, BO_Add, BF_One, BF_Zero, BO_Add, BF_One, BF_Zero,
-															CW_NONE, BO_Add, BF_One, BF_Zero, BO_Add, BF_One, BF_Zero>::GetRHI());
+			// Blend by holding out the matte shape of the source alpha
+			DrawRenderState.SetBlendState(TStaticBlendState<CW_RGBA, BO_Add, BF_Zero, BF_InverseSourceAlpha, BO_Add, BF_One, BF_InverseSourceAlpha>::GetRHI());
 		}
 		else
 		{
-			// the mobile thin translucent fallback uses a similar mode as BLEND_Translucent, but multiplies color by 1 insead of SrcAlpha.
-			DrawRenderState.SetBlendState(TStaticBlendState<CW_RGBA, BO_Add, BF_One, BF_InverseSourceAlpha, BO_Add, BF_Zero, BF_InverseSourceAlpha,
-															CW_NONE, BO_Add, BF_One, BF_Zero, BO_Add, BF_One, BF_Zero,
-															CW_NONE, BO_Add, BF_One, BF_Zero, BO_Add, BF_One, BF_Zero,
-															CW_NONE, BO_Add, BF_One, BF_Zero, BO_Add, BF_One, BF_Zero,
-															CW_NONE, BO_Add, BF_One, BF_Zero, BO_Add, BF_One, BF_Zero,
-															CW_NONE, BO_Add, BF_One, BF_Zero, BO_Add, BF_One, BF_Zero,
-															CW_NONE, BO_Add, BF_One, BF_Zero, BO_Add, BF_One, BF_Zero,
-															CW_NONE, BO_Add, BF_One, BF_Zero, BO_Add, BF_One, BF_Zero>::GetRHI());
+			// We always use premultiplied alpha for translucent rendering.
+			DrawRenderState.SetBlendState(TStaticBlendState<CW_RGBA, BO_Add, BF_One, BF_InverseSourceAlpha, BO_Add, BF_Zero, BF_InverseSourceAlpha>::GetRHI());
 		}
 	}
+	else if (ShadingModels.HasShadingModel(MSM_ThinTranslucent))
+	{
+		DrawRenderState.SetBlendState(GetBlendStateForColorTransmittanceBlending(ShaderPlatform));
+	}
 	else
 	{
 		switch (Material.GetBlendMode())
diff --git a/Engine/Source/Runtime/Renderer/Private/MobileBasePassRendering.cpp b/Engine/Source/Runtime/Renderer/Private/MobileBasePassRendering.cpp
--- a/Engine/Source/Runtime/Renderer/Private/MobileBasePassRendering.cpp
+++ b/Engine/Source/Runtime/Renderer/Private/MobileBasePassRendering.cpp
@@ -66,6 +66,61 @@ static TAutoConsoleVariable<int32> CVarMobileUseHWsRGBEncoding(
 	TEXT("1: Use GPU HW to convert linear to sRGB automatically (device must support sRGB write control)\n"),
 	ECVF_RenderThreadSafe);
 
+EMobileTranslucentColorTransmittanceMode MobileDefaultTranslucentColorTransmittanceMode(EShaderPlatform Platform)
+{
+	if (FDataDrivenShaderPlatformInfo::GetSupportsDualSourceBlending(Platform) || IsSimulatedPlatform(Platform))
+	{
+		return EMobileTranslucentColorTransmittanceMode::DUAL_SRC_BLENDING;
+	}
+	if (IsMetalMobilePlatform(Platform) || IsAndroidOpenGLESPlatform(Platform))
+	{
+		return EMobileTranslucentColorTransmittanceMode::PROGRAMMABLE_BLENDING;
+	}
+	return EMobileTranslucentColorTransmittanceMode::SINGLE_SRC_BLENDING;
+}
+
+static bool SupportsTranslucentColorTransmittanceFallback(EShaderPlatform Platform, EMobileTranslucentColorTransmittanceMode Fallback)
+{
+	switch (Fallback)
+	{
+	default:
+		return true;
+	case EMobileTranslucentColorTransmittanceMode::SINGLE_SRC_BLENDING:
+		return IsSimulatedPlatform(Platform) || IsAndroidPlatform(Platform);
+	}
+}
+
+EMobileTranslucentColorTransmittanceMode MobileActiveTranslucentColorTransmittanceMode(EShaderPlatform Platform, bool bExplicitDefaultMode)
+{
+	const EMobileTranslucentColorTransmittanceMode DefaultMode = MobileDefaultTranslucentColorTransmittanceMode(Platform);
+
+	if (DefaultMode == EMobileTranslucentColorTransmittanceMode::DUAL_SRC_BLENDING)
+	{
+		if (!GSupportsDualSrcBlending)
+		{
+			if (SupportsTranslucentColorTransmittanceFallback(Platform, EMobileTranslucentColorTransmittanceMode::PROGRAMMABLE_BLENDING) && GSupportsShaderFramebufferFetch)
+			{
+				return EMobileTranslucentColorTransmittanceMode::PROGRAMMABLE_BLENDING;
+			}
+			else
+			{
+				check(SupportsTranslucentColorTransmittanceFallback(Platform, EMobileTranslucentColorTransmittanceMode::SINGLE_SRC_BLENDING));
+				return EMobileTranslucentColorTransmittanceMode::SINGLE_SRC_BLENDING;
+			}
+		}
+	}
+	else if (DefaultMode == EMobileTranslucentColorTransmittanceMode::PROGRAMMABLE_BLENDING)
+	{
+		if (!GSupportsShaderFramebufferFetch)
+		{
+			check(SupportsTranslucentColorTransmittanceFallback(Platform, EMobileTranslucentColorTransmittanceMode::SINGLE_SRC_BLENDING));
+			return EMobileTranslucentColorTransmittanceMode::SINGLE_SRC_BLENDING;
+		}
+	}
+
+	return bExplicitDefaultMode ? DefaultMode : EMobileTranslucentColorTransmittanceMode::DEFAULT;
+}
+
 bool ShouldCacheShaderByPlatformAndOutputFormat(EShaderPlatform Platform, EOutputFormat OutputFormat)
 {
 	bool bSupportsMobileHDR = IsMobileHDR();
@@ -82,17 +137,21 @@ bool ShouldCacheShaderByPlatformAndOutputFormat(EShaderPlatform Platform, EOutpu
 	IMPLEMENT_MATERIAL_SHADER_TYPE(template<>, TMobileBasePassVS##LightMapPolicyName##LDRGamma32, TEXT("/Engine/Private/MobileBasePassVertexShader.usf"), TEXT("Main"), SF_Vertex); \
 	IMPLEMENT_MATERIAL_SHADER_TYPE(template<>, TMobileBasePassVS##LightMapPolicyName##HDRLinear64, TEXT("/Engine/Private/MobileBasePassVertexShader.usf"), TEXT("Main"), SF_Vertex);
 
-#define IMPLEMENT_MOBILE_SHADING_BASEPASS_LIGHTMAPPED_PIXEL_SHADER_TYPE(LightMapPolicyType,LightMapPolicyName, LocalLightSetting) \
-	typedef TMobileBasePassPS< LightMapPolicyType, LDR_GAMMA_32, false, LocalLightSetting > TMobileBasePassPS##LightMapPolicyName##LDRGamma32##LocalLightSetting; \
-	typedef TMobileBasePassPS< LightMapPolicyType, HDR_LINEAR_64, false, LocalLightSetting > TMobileBasePassPS##LightMapPolicyName##HDRLinear64##LocalLightSetting; \
-	typedef TMobileBasePassPS< LightMapPolicyType, LDR_GAMMA_32, true, LocalLightSetting > TMobileBasePassPS##LightMapPolicyName##LDRGamma32##Skylight##LocalLightSetting; \
-	typedef TMobileBasePassPS< LightMapPolicyType, HDR_LINEAR_64, true, LocalLightSetting > TMobileBasePassPS##LightMapPolicyName##HDRLinear64##Skylight##LocalLightSetting; \
-	IMPLEMENT_MATERIAL_SHADER_TYPE(template<>, TMobileBasePassPS##LightMapPolicyName##LDRGamma32##LocalLightSetting, TEXT("/Engine/Private/MobileBasePassPixelShader.usf"), TEXT("Main"), SF_Pixel); \
-	IMPLEMENT_MATERIAL_SHADER_TYPE(template<>, TMobileBasePassPS##LightMapPolicyName##HDRLinear64##LocalLightSetting, TEXT("/Engine/Private/MobileBasePassPixelShader.usf"), TEXT("Main"), SF_Pixel); \
-	IMPLEMENT_MATERIAL_SHADER_TYPE(template<>, TMobileBasePassPS##LightMapPolicyName##LDRGamma32##Skylight##LocalLightSetting, TEXT("/Engine/Private/MobileBasePassPixelShader.usf"), TEXT("Main"), SF_Pixel); \
-	IMPLEMENT_MATERIAL_SHADER_TYPE(template<>, TMobileBasePassPS##LightMapPolicyName##HDRLinear64##Skylight##LocalLightSetting, TEXT("/Engine/Private/MobileBasePassPixelShader.usf"), TEXT("Main"), SF_Pixel)
-
-#define IMPLEMENT_MOBILE_SHADING_BASEPASS_LIGHTMAPPED_SHADER_TYPE(LightMapPolicyType,LightMapPolicyName) \
+#define IMPLEMENT_MOBILE_SHADING_BASEPASS_LIGHTMAPPED_PIXEL_SHADER_TYPE2(LightMapPolicyType, LightMapPolicyName, LocalLightSetting, ThinTranslucencyEnum, ThinTranslucencyName) \
+	typedef TMobileBasePassPS< LightMapPolicyType, LDR_GAMMA_32, false, LocalLightSetting, EMobileTranslucentColorTransmittanceMode::ThinTranslucencyEnum > TMobileBasePassPS##LightMapPolicyName##LDRGamma32##LocalLightSetting##ThinTranslucencyName; \
+	typedef TMobileBasePassPS< LightMapPolicyType, HDR_LINEAR_64, false, LocalLightSetting, EMobileTranslucentColorTransmittanceMode::ThinTranslucencyEnum > TMobileBasePassPS##LightMapPolicyName##HDRLinear64##LocalLightSetting##ThinTranslucencyName; \
+	typedef TMobileBasePassPS< LightMapPolicyType, LDR_GAMMA_32, true, LocalLightSetting, EMobileTranslucentColorTransmittanceMode::ThinTranslucencyEnum > TMobileBasePassPS##LightMapPolicyName##LDRGamma32##Skylight##LocalLightSetting##ThinTranslucencyName; \
+	typedef TMobileBasePassPS< LightMapPolicyType, HDR_LINEAR_64, true, LocalLightSetting, EMobileTranslucentColorTransmittanceMode::ThinTranslucencyEnum > TMobileBasePassPS##LightMapPolicyName##HDRLinear64##Skylight##LocalLightSetting##ThinTranslucencyName; \
+	IMPLEMENT_MATERIAL_SHADER_TYPE(template<>, TMobileBasePassPS##LightMapPolicyName##LDRGamma32##LocalLightSetting##ThinTranslucencyName, TEXT("/Engine/Private/MobileBasePassPixelShader.usf"), TEXT("Main"), SF_Pixel); \
+	IMPLEMENT_MATERIAL_SHADER_TYPE(template<>, TMobileBasePassPS##LightMapPolicyName##HDRLinear64##LocalLightSetting##ThinTranslucencyName, TEXT("/Engine/Private/MobileBasePassPixelShader.usf"), TEXT("Main"), SF_Pixel); \
+	IMPLEMENT_MATERIAL_SHADER_TYPE(template<>, TMobileBasePassPS##LightMapPolicyName##LDRGamma32##Skylight##LocalLightSetting##ThinTranslucencyName, TEXT("/Engine/Private/MobileBasePassPixelShader.usf"), TEXT("Main"), SF_Pixel); \
+	IMPLEMENT_MATERIAL_SHADER_TYPE(template<>, TMobileBasePassPS##LightMapPolicyName##HDRLinear64##Skylight##LocalLightSetting##ThinTranslucencyName, TEXT("/Engine/Private/MobileBasePassPixelShader.usf"), TEXT("Main"), SF_Pixel);
+
+#define IMPLEMENT_MOBILE_SHADING_BASEPASS_LIGHTMAPPED_PIXEL_SHADER_TYPE(LightMapPolicyType, LightMapPolicyName, LocalLightSetting) \
+	IMPLEMENT_MOBILE_SHADING_BASEPASS_LIGHTMAPPED_PIXEL_SHADER_TYPE2(LightMapPolicyType, LightMapPolicyName, LocalLightSetting, DEFAULT,) \
+	IMPLEMENT_MOBILE_SHADING_BASEPASS_LIGHTMAPPED_PIXEL_SHADER_TYPE2(LightMapPolicyType, LightMapPolicyName, LocalLightSetting, SINGLE_SRC_BLENDING, ThinTranslGrey)
+
+#define IMPLEMENT_MOBILE_SHADING_BASEPASS_LIGHTMAPPED_SHADER_TYPE(LightMapPolicyType, LightMapPolicyName) \
 	IMPLEMENT_MOBILE_SHADING_BASEPASS_LIGHTMAPPED_VERTEX_SHADER_TYPE(LightMapPolicyType, LightMapPolicyName) \
 	IMPLEMENT_MOBILE_SHADING_BASEPASS_LIGHTMAPPED_PIXEL_SHADER_TYPE(LightMapPolicyType, LightMapPolicyName, LOCAL_LIGHTS_DISABLED) \
 	IMPLEMENT_MOBILE_SHADING_BASEPASS_LIGHTMAPPED_PIXEL_SHADER_TYPE(LightMapPolicyType, LightMapPolicyName, LOCAL_LIGHTS_ENABLED) \
@@ -110,6 +169,26 @@ IMPLEMENT_MOBILE_SHADING_BASEPASS_LIGHTMAPPED_SHADER_TYPE(TUniformLightMapPolicy
 IMPLEMENT_MOBILE_SHADING_BASEPASS_LIGHTMAPPED_SHADER_TYPE(TUniformLightMapPolicy<LMP_MOBILE_MOVABLE_DIRECTIONAL_LIGHT_CSM_WITH_LIGHTMAP>, FMobileMovableDirectionalLightCSMWithLightmapPolicy);
 IMPLEMENT_MOBILE_SHADING_BASEPASS_LIGHTMAPPED_SHADER_TYPE(TUniformLightMapPolicy<LMP_MOBILE_DIRECTIONAL_LIGHT_CSM>, FMobileDirectionalLightAndCSMPolicy);
 
+bool MaterialRequiresColorTransmittanceBlending(const FMaterialShaderParameters& MaterialParameters)
+{
+	return MaterialParameters.ShadingModels.HasShadingModel(MSM_ThinTranslucent) || MaterialParameters.BlendMode == BLEND_TranslucentColoredTransmittance;
+}
+
+bool ShouldCacheShaderForColorTransmittanceFallback(const FMaterialShaderPermutationParameters& Parameters, const EMobileTranslucentColorTransmittanceMode TranslucentColorTransmittanceFallback)
+{
+	if (TranslucentColorTransmittanceFallback == EMobileTranslucentColorTransmittanceMode::DEFAULT)
+	{
+		return true;
+	}
+	
+	if (!MaterialRequiresColorTransmittanceBlending(Parameters.MaterialParameters))
+	{
+		return false;
+	}
+
+	return SupportsTranslucentColorTransmittanceFallback(Parameters.Platform, TranslucentColorTransmittanceFallback);
+}
+
 // shared defines for mobile base pass VS and PS
 void MobileBasePassModifyCompilationEnvironment(const FMaterialShaderPermutationParameters& Parameters, FShaderCompilerEnvironment& OutEnvironment, EOutputFormat OutputFormat)
 {
@@ -137,6 +216,13 @@ void MobileBasePassModifyCompilationEnvironment(const FMaterialShaderPermutation
 	// translucency is in the same subpass with deferred shading shaders, so it has access to GBuffer
 	const bool bDeferredShadingSubpass = (bDeferredShadingEnabled && bTranslucentMaterial && !Parameters.MaterialParameters.bIsMobileSeparateTranslucencyEnabled);
 	OutEnvironment.SetDefine(TEXT("IS_MOBILE_DEFERREDSHADING_SUBPASS"), bDeferredShadingSubpass ? 1u : 0u);
+
+	// HLSLcc does not support dual source blending, so force DXC if needed
+	if (bTranslucentMaterial && FDataDrivenShaderPlatformInfo::GetSupportsDxc(Parameters.Platform) && IsHlslccShaderPlatform(Parameters.Platform) &&
+		MaterialRequiresColorTransmittanceBlending(Parameters.MaterialParameters) && MobileDefaultTranslucentColorTransmittanceMode(Parameters.Platform) == EMobileTranslucentColorTransmittanceMode::DUAL_SRC_BLENDING)
+	{
+		OutEnvironment.CompilerFlags.Add(CFLAG_ForceDXC);
+	}
 }
 
 template<typename LightMapPolicyType>
@@ -473,4 +559,4 @@ void FMobileSceneRenderer::RenderMobileEditorPrimitives(FRHICommandList& RHICmdL
 		// Draw the view's batched simple elements(lines, sprites, etc).
 		View.TopBatchedViewElements.Draw(RHICmdList, DrawRenderState, FeatureLevel, View, false);
 	}
-}
\ No newline at end of file
+}
diff --git a/Engine/Source/Runtime/Renderer/Private/MobileBasePassRendering.h b/Engine/Source/Runtime/Renderer/Private/MobileBasePassRendering.h
--- a/Engine/Source/Runtime/Renderer/Private/MobileBasePassRendering.h
+++ b/Engine/Source/Runtime/Renderer/Private/MobileBasePassRendering.h
@@ -107,6 +107,19 @@ enum EOutputFormat
 
 EMobileLocalLightSetting GetMobileForwardLocalLightSetting(EShaderPlatform ShaderPlatform, bool bIsTranslucent);
 
+enum class EMobileTranslucentColorTransmittanceMode
+{
+	DEFAULT,
+	DUAL_SRC_BLENDING,
+	PROGRAMMABLE_BLENDING,
+	SINGLE_SRC_BLENDING, // Grey transmittance
+};
+
+EMobileTranslucentColorTransmittanceMode MobileDefaultTranslucentColorTransmittanceMode(EShaderPlatform Platform);
+EMobileTranslucentColorTransmittanceMode MobileActiveTranslucentColorTransmittanceMode(EShaderPlatform Platform, bool bExplicitDefaultMode);
+bool MaterialRequiresColorTransmittanceBlending(const FMaterialShaderParameters& MaterialParameters);
+bool ShouldCacheShaderForColorTransmittanceFallback(const FMaterialShaderPermutationParameters& Parameters, EMobileTranslucentColorTransmittanceMode TranslucentColorTransmittanceFallback);
+
 bool ShouldCacheShaderByPlatformAndOutputFormat(EShaderPlatform Platform, EOutputFormat OutputFormat);
 // shared defines for mobile base pass VS and PS
 void MobileBasePassModifyCompilationEnvironment(const FMaterialShaderPermutationParameters& Parameters, FShaderCompilerEnvironment& OutEnvironment, EOutputFormat OutputFormat);
@@ -345,7 +358,7 @@ namespace MobileBasePass
 	}
 };
 
-template< typename LightMapPolicyType, EOutputFormat OutputFormat, bool bEnableSkyLight, EMobileLocalLightSetting LocalLightSetting>
+template< typename LightMapPolicyType, EOutputFormat OutputFormat, bool bEnableSkyLight, EMobileLocalLightSetting LocalLightSetting, EMobileTranslucentColorTransmittanceMode TranslucentColorTransmittanceFallback = EMobileTranslucentColorTransmittanceMode::DEFAULT>
 class TMobileBasePassPS : public TMobileBasePassPSBaseType<LightMapPolicyType>
 {
 	DECLARE_SHADER_TYPE(TMobileBasePassPS,MeshMaterial);
@@ -384,7 +397,8 @@ public:
 		return TMobileBasePassPSBaseType<LightMapPolicyType>::ShouldCompilePermutation(Parameters) && 
 				ShouldCacheShaderByPlatformAndOutputFormat(Parameters.Platform, OutputFormat) && 
 				bShouldCacheBySkylight && 
-				bShouldCacheByLocalLights;
+				bShouldCacheByLocalLights && 
+				ShouldCacheShaderForColorTransmittanceFallback(Parameters, TranslucentColorTransmittanceFallback);
 	}
 	
 	static void ModifyCompilationEnvironment(const FMaterialShaderPermutationParameters& Parameters, FShaderCompilerEnvironment& OutEnvironment)
@@ -407,6 +421,15 @@ public:
 		OutEnvironment.SetDefine(TEXT("ENABLE_CLUSTERED_REFLECTION"), bEnableClusteredReflections ? 1u : 0u);
 		OutEnvironment.SetDefine(TEXT("USE_SHADOWMASKTEXTURE"), bMobileUsesShadowMaskTexture && !bTranslucentMaterial ? 1u : 0u);
 		OutEnvironment.SetDefine(TEXT("ENABLE_DBUFFER_TEXTURES"), Parameters.MaterialParameters.MaterialDomain == MD_Surface ? 1u : 0u);
+
+		EMobileTranslucentColorTransmittanceMode TranslucentColorTransmittanceMode = EMobileTranslucentColorTransmittanceMode::DEFAULT;
+		if (Parameters.MaterialParameters.ShadingModels.HasShadingModel(MSM_ThinTranslucent))
+		{
+			TranslucentColorTransmittanceMode = (TranslucentColorTransmittanceFallback == EMobileTranslucentColorTransmittanceMode::DEFAULT) ? MobileDefaultTranslucentColorTransmittanceMode(Parameters.Platform) : TranslucentColorTransmittanceFallback;
+		}
+		OutEnvironment.SetDefine(TEXT("MOBILE_TRANSLUCENT_COLOR_TRANSMITTANCE_DUAL_SRC_BLENDING"), TranslucentColorTransmittanceMode == EMobileTranslucentColorTransmittanceMode::DUAL_SRC_BLENDING ? 1u : 0u);
+		OutEnvironment.SetDefine(TEXT("MOBILE_TRANSLUCENT_COLOR_TRANSMITTANCE_PROGRAMMABLE_BLENDING"), TranslucentColorTransmittanceMode == EMobileTranslucentColorTransmittanceMode::PROGRAMMABLE_BLENDING ? 1u : 0u);
+		OutEnvironment.SetDefine(TEXT("MOBILE_TRANSLUCENT_COLOR_TRANSMITTANCE_SINGLE_SRC_BLENDING"), TranslucentColorTransmittanceMode == EMobileTranslucentColorTransmittanceMode::SINGLE_SRC_BLENDING ? 1u : 0u);
 	}
 	
 	/** Initialization constructor. */
diff --git a/Engine/Source/Runtime/VulkanRHI/Private/VulkanRHI.cpp b/Engine/Source/Runtime/VulkanRHI/Private/VulkanRHI.cpp
--- a/Engine/Source/Runtime/VulkanRHI/Private/VulkanRHI.cpp
+++ b/Engine/Source/Runtime/VulkanRHI/Private/VulkanRHI.cpp
@@ -980,6 +980,7 @@ void FVulkanDynamicRHI::InitInstance()
 		GMaxTextureArrayLayers = Props.limits.maxImageArrayLayers;
 		GRHISupportsBaseVertexIndex = true;
 		GSupportsSeparateRenderTargetBlendState = true;
+		GSupportsDualSrcBlending = Device->GetPhysicalDeviceFeatures().Core_1_0.dualSrcBlend == VK_TRUE;
 		GRHISupportsSeparateDepthStencilCopyAccess = Device->SupportsParallelRendering();
 		GRHIBindlessSupport = Device->SupportsBindless() ? RHIGetBindlessSupport(GMaxRHIShaderPlatform) : ERHIBindlessSupport::Unsupported;
 		GMaxTextureSamplers = Props.limits.maxPerStageDescriptorSamplers;
