Improve VSM cvar documentation

diff --git a/Engine/Shaders/Private/VirtualShadowMaps/VirtualShadowMapPageOverlap.ush b/Engine/Shaders/Private/VirtualShadowMaps/VirtualShadowMapPageOverlap.ush
--- a/Engine/Shaders/Private/VirtualShadowMaps/VirtualShadowMapPageOverlap.ush
+++ b/Engine/Shaders/Private/VirtualShadowMaps/VirtualShadowMapPageOverlap.ush
@@ -233,7 +233,7 @@ uint GetDetailGeometryFlagMaskForCulling(bool bShouldCacheAsStatic, bool bIsNani
 		return bIsNaniteGeometry ? 0U : VSM_DETAIL_GEOMETRY_FLAG;
 	}
 
-	// We use a separate threshold for dynamically cached geometry, this is becase the overhead is lower when caching is enabled as cached page culling removes the instances anyway.
+	// We use a separate threshold for dynamically cached geometry, this is because the overhead is lower when caching is enabled as cached page culling removes the instances anyway.
 	// Nanite geometry has a separate smaller footprint as it has better LOD and lower overhead for replicated drawing.
 	if (bShouldCacheAsStatic)
 	{
diff --git a/Engine/Shaders/Private/VirtualShadowMaps/VirtualShadowMapProjection.usf b/Engine/Shaders/Private/VirtualShadowMaps/VirtualShadowMapProjection.usf
--- a/Engine/Shaders/Private/VirtualShadowMaps/VirtualShadowMapProjection.usf
+++ b/Engine/Shaders/Private/VirtualShadowMaps/VirtualShadowMapProjection.usf
@@ -684,7 +684,7 @@ void VirtualShadowMapProjection(
 #if DIRECTIONAL_LIGHT
 		// Increases light source radius for subsurface material as the opacity reduces to emulates light diffusion.
 		// TODO: Does something similar make sense for local lights?
-		// Substrate: this approximation has no concrete fundation so it is currently skipped for Substrate using PerPixel SSS, e.g. when bDataIsOpacity==false and data is MFP.
+		// Substrate: this approximation has no concrete foundation so it is currently skipped for Substrate using PerPixel SSS, e.g. when bDataIsOpacity==false and data is MFP.
 		if (SubsurfaceOpacityMFP.bDataIsOpacity && SubsurfaceOpacityMFP.Data < 1.0f)
 		{
 			const float SubsurfaceOpacity = SubsurfaceOpacityMFP.Data;
diff --git a/Engine/Source/Editor/CommonMenuExtensions/Private/VirtualShadowMapVisualizationMenuCommands.cpp b/Engine/Source/Editor/CommonMenuExtensions/Private/VirtualShadowMapVisualizationMenuCommands.cpp
--- a/Engine/Source/Editor/CommonMenuExtensions/Private/VirtualShadowMapVisualizationMenuCommands.cpp
+++ b/Engine/Source/Editor/CommonMenuExtensions/Private/VirtualShadowMapVisualizationMenuCommands.cpp
@@ -21,7 +21,7 @@ int32 GVirtualShadowMapVisualizeAdvanced = 0;
 static FAutoConsoleVariableRef CVarVirtualShadowMapVisualizeAdvanced(
 	TEXT("r.Shadow.Virtual.Visualize.Advanced"),
 	GVirtualShadowMapVisualizeAdvanced,
-	TEXT("")
+	TEXT("Enable to show advanced VSM debug modes in the visualization UI menu.")
 );
 
 #define LOCTEXT_NAMESPACE "VirtualShadowMapVisualizationMenuCommands"
diff --git a/Engine/Source/Runtime/Engine/Private/VirtualShadowMapVisualizationData.cpp b/Engine/Source/Runtime/Engine/Private/VirtualShadowMapVisualizationData.cpp
--- a/Engine/Source/Runtime/Engine/Private/VirtualShadowMapVisualizationData.cpp
+++ b/Engine/Source/Runtime/Engine/Private/VirtualShadowMapVisualizationData.cpp
@@ -65,7 +65,7 @@ void FVirtualShadowMapVisualizationData::Initialize()
 		AddVisualizationMode(
 			TEXT("invalid"),
 			LOCTEXT("InvalidPage", "GPU Invalidated Page"),
-			LOCTEXT("InvalidPageDebugDesc", "Show the pages marked for GPU-driven invalidaiton (World Position Offset)."),
+			LOCTEXT("InvalidPageDebugDesc", "Show the pages marked for GPU-driven invalidation (World Position Offset)."),
 			FModeType::Advanced,
 			VIRTUAL_SHADOW_MAP_VISUALIZE_GPU_INVALIDATED_PAGE);
 
diff --git a/Engine/Source/Runtime/Renderer/Private/Shadows/ShadowSceneRenderer.cpp b/Engine/Source/Runtime/Renderer/Private/Shadows/ShadowSceneRenderer.cpp
--- a/Engine/Source/Runtime/Renderer/Private/Shadows/ShadowSceneRenderer.cpp
+++ b/Engine/Source/Runtime/Renderer/Private/Shadows/ShadowSceneRenderer.cpp
@@ -27,14 +27,15 @@ TAutoConsoleVariable<int32> CVarVSMMaterialVisibility(
 TAutoConsoleVariable<int32> CVarMaxDistantLightsPerFrame(
 	TEXT("r.Shadow.Virtual.MaxDistantUpdatePerFrame"),
 	1,
-	TEXT("Maximum number of distant lights to update each frame."),
+	TEXT("Maximum number of distant lights to update each frame. Invalidated lights that were missed may be updated in a later frame (round-robin)."),
 	ECVF_Scalability | ECVF_RenderThreadSafe
 );
 
 static TAutoConsoleVariable<int32> CVarDistantLightMode(
 	TEXT("r.Shadow.Virtual.DistantLightMode"),
 	1,
-	TEXT("Control whether distant light mode is enabled for local lights.\n0 == Off, \n1 == On (default), \n2 == Force All."),
+	TEXT("Control whether distant light mode is enabled for local lights.\n0 == Off, \n1 == On (default), \n2 == Force All.\n")
+	TEXT("When on, lights with a pixel footprint below the threshold are marked as distant. Updates to distant lights are throttled (force-cached), they use simpler page-table logic and the memory cost is lower."),
 	ECVF_RenderThreadSafe
 );
 
@@ -42,6 +43,7 @@ static TAutoConsoleVariable<float> CVarDistantLightForceCacheFootprintFraction(
 	TEXT("r.Shadow.Virtual.DistantLightForceCacheFootprintFraction"),
 	0.0f,
 	TEXT("Fraction of footprint size below which start force-caching lights that are invalidated (i.e., are moving or re-added)\n")
+	TEXT("  Larger values may improve performance but may also produce more visible artifacts\n")
 	TEXT("  The base footprint is based on the page size.\n")
 	TEXT("  0.0 == Never force-cache (default), 1.0 == Always force-cache."),
 	ECVF_Scalability | ECVF_RenderThreadSafe
@@ -70,7 +72,8 @@ static TAutoConsoleVariable<float> CVarResolutionLodBiasLocal(
 static TAutoConsoleVariable<float> CVarResolutionLodBiasLocalMoving(
 	TEXT("r.Shadow.Virtual.ResolutionLodBiasLocalMoving"),
 	1.0f,
-	TEXT("Bias applied to LOD calculations for moving local lights. -1.0 doubles resolution, 1.0 halves it and so on."),
+	TEXT("Bias applied to LOD calculations for local lights that are moving. -1.0 doubles resolution, 1.0 halves it and so on.\n")
+	TEXT("The bias transitions smoothly back to ResolutionLodBiasLocal as the light transitions to non-moving, see 'r.Shadow.Scene.LightActiveFrameCount'."),
 	ECVF_Scalability | ECVF_RenderThreadSafe
 );
 
diff --git a/Engine/Source/Runtime/Renderer/Private/VirtualShadowMaps/VirtualShadowMapArray.cpp b/Engine/Source/Runtime/Renderer/Private/VirtualShadowMaps/VirtualShadowMapArray.cpp
--- a/Engine/Source/Runtime/Renderer/Private/VirtualShadowMaps/VirtualShadowMapArray.cpp
+++ b/Engine/Source/Runtime/Renderer/Private/VirtualShadowMaps/VirtualShadowMapArray.cpp
@@ -51,7 +51,8 @@ FAutoConsoleVariableRef CVarVSMShowLightDrawEvents(
 TAutoConsoleVariable<int32> CVarEnableVirtualShadowMaps(
 	TEXT("r.Shadow.Virtual.Enable"),
 	0,
-	TEXT("Enable Virtual Shadow Maps."),
+	TEXT("Enable Virtual Shadow Maps. Renders geometry into virtualized shadow depth maps for shadowing.\n")
+	TEXT("Provides high - quality shadows for next - gen projects with simplified setup.High efficiency culling when used with Nanite."),
 	FConsoleVariableDelegate::CreateLambda([](IConsoleVariable* InVariable)
 	{
 		// Needed because the depth state changes with method (so cached draw commands must be re-created) see SetStateForShadowDepth
@@ -63,7 +64,11 @@ TAutoConsoleVariable<int32> CVarEnableVirtualShadowMaps(
 TAutoConsoleVariable<int32> CVarMaxPhysicalPages(
 	TEXT("r.Shadow.Virtual.MaxPhysicalPages"),
 	2048,
-	TEXT("Maximum number of physical pages in the pool."),
+	TEXT("Maximum number of physical pages in the pool.\n")
+	TEXT("More space for pages means more memory usage, but allows for higher resolution shadows.\n")
+	TEXT("Ideally this value is large enough to fit enough pages for all the lights in the scene, but not too large to waste memory.\n")
+	TEXT("Enable 'ShowStats' to see how many pages are allocated in the pool right now.\n")
+	TEXT("For more page pool control, see the 'ResolutionLodBias*', 'DynamicRes.*' and 'Cache.StaticSeparate' cvars."),
 	ECVF_Scalability | ECVF_RenderThreadSafe
 );
 
@@ -78,7 +83,7 @@ TAutoConsoleVariable<int32> CVarCacheStaticSeparate(
 static TAutoConsoleVariable<int32> CVarShowStats(
 	TEXT("r.Shadow.Virtual.ShowStats"),
 	0,
-	TEXT("ShowStats, also toggle shaderprint one!"),
+	TEXT("Show VSM statistics."),
 	ECVF_RenderThreadSafe
 );
 
@@ -142,8 +147,11 @@ TAutoConsoleVariable<int32> CVarCoarsePagesIncludeNonNanite(
 TAutoConsoleVariable<float> CVarNonNaniteCulledInstanceAllocationFactor(
 	TEXT("r.Shadow.Virtual.NonNanite.CulledInstanceAllocationFactor"),
 	1.0f,
-	TEXT("Scale factor multiplied by the conservative size required if all instances are emitted into every clip or mip level, setting to 1.0 is fully conservative.")
-	TEXT("The actual number cannot be known on the CPU as the culling emits an instace for each clip/mip level that is overlapped."),
+	TEXT("Allocation size scale factor for the buffer used to store instances after culling.\n")
+	TEXT("The total size accounts for the worst-case scenario in which all instances are emitted into every clip or mip level.\n")
+	TEXT("This is far more than we'd expect in reasonable circumstances, so this scale factor is used to reduce memory pressure.\n")
+	TEXT("The actual number cannot be known on the CPU as the culling emits an instance for each clip/mip level that is overlapped.\n")
+	TEXT("Setting to 1.0 is fully conservative. Lowering this is likely to produce artifacts unless you're certain the buffer won't overflow."),
 	ECVF_RenderThreadSafe
 );
 
@@ -179,7 +187,10 @@ FAutoConsoleVariableRef CVarEnableNonNaniteVSM(
 static TAutoConsoleVariable<int32> CVarNonNaniteVsmUseHzb(
 	TEXT("r.Shadow.Virtual.NonNanite.UseHZB"),
 	2,
-	TEXT("Cull Non-Nanite instances using HZB. If set to 2, attempt to use Nanite-HZB from the current frame."),
+	TEXT("Cull Non-Nanite instances using HZB.\n")
+	TEXT("  Set to 0 to disable.\n")
+	TEXT("  Set to 1 to use HZB from previous frame. Can incorrectly cull in some cases due to outdated data.\n")
+	TEXT("  Set to 2 to use two-pass Nanite culling with HZB from the current frame."),
 	ECVF_RenderThreadSafe);
 
 TAutoConsoleVariable<int32> CVarVirtualShadowOnePassProjectionMaxLights(
@@ -200,28 +211,30 @@ TAutoConsoleVariable<int32> CVarDoNonNaniteBatching(
 static TAutoConsoleVariable<float> CVarCoarsePagePixelThresholdDynamic(
 	TEXT("r.Shadow.Virtual.CoarsePagePixelThresholdDynamic"),
 	16.0f,
-	TEXT("If a dynamic (non-nanite) instance has a smaller footprint, it should not be drawn into a coarse page."),
+	TEXT("If a dynamic (non-nanite) instance has a smaller estimated pixel footprint than this value, it should not be drawn into a coarse page. Higher values cull away more instances."),
 	ECVF_Scalability | ECVF_RenderThreadSafe
 );
 
 static TAutoConsoleVariable<float> CVarCoarsePagePixelThresholdStatic(
 	TEXT("r.Shadow.Virtual.CoarsePagePixelThresholdStatic"),
 	1.0f,
-	TEXT("If a static (non-nanite) instance has a smaller footprint, it should not be drawn into a coarse page."),
+	TEXT("If a static (non-nanite) instance has a smaller estimated pixel footprint than this value, it should not be drawn into a coarse page. Higher values cull away more instances.\n")
+	TEXT("This value is typically lower than the non-static one because the static pages have better caching."),
 	ECVF_Scalability | ECVF_RenderThreadSafe
 );
 
 static TAutoConsoleVariable<float> CVarCoarsePagePixelThresholdDynamicNanite(
 	TEXT("r.Shadow.Virtual.CoarsePagePixelThresholdDynamicNanite"),
 	4.0f,
-	TEXT("If a dynamic Nanite instance has a smaller footprint, it should not be drawn into a coarse page."),
+	TEXT("If a dynamic Nanite instance has a smaller estimated pixel footprint than this value, it should not be drawn into a coarse page. Higher values cull away more instances.\n")
+	TEXT("This value is typically lower than the non-Nanite one because Nanite has lower overhead for drawing small objects."),
 	ECVF_Scalability | ECVF_RenderThreadSafe
 );
 
 static TAutoConsoleVariable<int32> CVarCacheAllocateViaLRU(
 	TEXT("r.Shadow.Virtual.Cache.AllocateViaLRU"),
 	0,
-	TEXT("Prioritizes keeping more recently requested cached physical pages when allocating for new requests.\n"),
+	TEXT("Prioritizes keeping more recently requested cached physical pages when allocating for new requests."),
 	ECVF_RenderThreadSafe
 );
 
@@ -239,7 +252,7 @@ void DumpVSMLightNames()
 
 FAutoConsoleCommand CmdDumpVSMLightNames(
 	TEXT("r.Shadow.Virtual.Visualize.DumpLightNames"),
-	TEXT("Dump light names with virtual shadow maps (for developer use in non-shiping builds)"),
+	TEXT("Dump light names with virtual shadow maps (for developer use in non-shipping builds)"),
 	FConsoleCommandDelegate::CreateStatic(DumpVSMLightNames)
 );
 
@@ -247,7 +260,7 @@ FString GVirtualShadowMapVisualizeLightName;
 FAutoConsoleVariableRef CVarVisualizeLightName(
 	TEXT("r.Shadow.Virtual.Visualize.LightName"),
 	GVirtualShadowMapVisualizeLightName,
-	TEXT("Sets the name of a specific light to visualize (for developer use in non-shiping builds)"),
+	TEXT("Sets the name of a specific light to visualize (for developer use in non-shipping builds)"),
 	ECVF_RenderThreadSafe
 );
 
@@ -264,28 +277,29 @@ static TAutoConsoleVariable<int32> CVarVisualizeLayout(
 TAutoConsoleVariable<int32> CVarDebugSkipMergePhysical(
 	TEXT("r.Shadow.Virtual.DebugSkipMergePhysical"),
 	0,
-	TEXT(""),
+	TEXT("Skip the merging of the static VSM cache into the dynamic one. This will create obvious visual artifacts when disabled."),
 	ECVF_RenderThreadSafe
 );
 
 TAutoConsoleVariable<int32> CVarDebugSkipDynamicPageInvalidation(
 	TEXT("r.Shadow.Virtual.Cache.DebugSkipDynamicPageInvalidation"),
 	0,
-	TEXT("Skip invalidation of cached pages when geometry moves for debugging purposes. This will create obvious visual artifacts when disabled.\n"),
+	TEXT("Skip invalidation of cached pages when geometry moves for debugging purposes. This will create obvious visual artifacts when disabled."),
 	ECVF_RenderThreadSafe
 );
 
 TAutoConsoleVariable<int32> CVarNumPageAreaDiagSlots(
 	TEXT("r.Shadow.Virtual.NonNanite.NumPageAreaDiagSlots"),
 	0,
-	TEXT("Feed back diagnostics to host to report page area coverage for the K first occurrences, < 0 uses the max number allowed, 0 disables."),
+	TEXT("Number of slots in diagnostics to report non-nanite instances with the largest page area coverage, < 0 uses the max number allowed, 0 disables."),
 	ECVF_RenderThreadSafe
 );
 
 TAutoConsoleVariable<int32> CVarLargeInstancePageAreaThreshold(
 	TEXT("r.Shadow.Virtual.NonNanite.LargeInstancePageAreaThreshold"),
 	-1,
-	TEXT("How large area is considered a 'large' footprint, summed over all overlapped levels, if set to -1 uses the physical page pool size / 8."),
+	TEXT("How large area is considered a 'large' footprint, summed over all overlapped levels, if set to -1 uses the physical page pool size / 8.\n")
+	TEXT("Used as a threshold when storing page area coverage stats for diagnostics."),
 	ECVF_RenderThreadSafe
 );
 #endif // !UE_BUILD_SHIPPING
@@ -315,13 +329,16 @@ static TAutoConsoleVariable<int32> CVarVirtualShadowSinglePassBatched(
 static TAutoConsoleVariable<int32> CVarVirtualShadowMapPageMarkingPixelStrideX(
 	TEXT("r.Shadow.Virtual.PageMarkingPixelStrideX"),
 	1,
-	TEXT("."),
+	TEXT("During page marking, instead of testing every screen pixel, test every Nth pixel.\n")
+	TEXT("Page marking from screen pixels is used to determine which VSM pages are seen from the camera and need to be rendered.\n")
+	TEXT("Increasing this value reduces page-marking costs, but could introduce artifacts due to missing pages.\n")
+	TEXT("With sufficiently low values, it is likely a neighbouring pixel will mark the required page anyway."),
 	ECVF_RenderThreadSafe);
 
 static TAutoConsoleVariable<int32> CVarVirtualShadowMapPageMarkingPixelStrideY(
 	TEXT("r.Shadow.Virtual.PageMarkingPixelStrideY"),
 	1,
-	TEXT("."),
+	TEXT("Same as PageMarkingPixelStrideX, but on the vertical axis of the screen."),
 	ECVF_RenderThreadSafe);
 
 namespace Nanite
diff --git a/Engine/Source/Runtime/Renderer/Private/VirtualShadowMaps/VirtualShadowMapCacheManager.cpp b/Engine/Source/Runtime/Renderer/Private/VirtualShadowMaps/VirtualShadowMapCacheManager.cpp
--- a/Engine/Source/Runtime/Renderer/Private/VirtualShadowMaps/VirtualShadowMapCacheManager.cpp
+++ b/Engine/Source/Runtime/Renderer/Private/VirtualShadowMaps/VirtualShadowMapCacheManager.cpp
@@ -21,7 +21,7 @@ CSV_DECLARE_CATEGORY_EXTERN(VSM);
 static TAutoConsoleVariable<int32> CVarAccumulateStats(
 	TEXT("r.Shadow.Virtual.AccumulateStats"),
 	0,
-	TEXT("AccumulateStats"),
+	TEXT("When enabled, VSM stats will be collected over multiple frames and written to a CSV file"),
 	ECVF_RenderThreadSafe
 );
 
@@ -45,14 +45,14 @@ static TAutoConsoleVariable<int32> CVarDrawInvalidatingBounds(
 static TAutoConsoleVariable<int32> CVarCacheVsmUseHzb(
 	TEXT("r.Shadow.Virtual.Cache.InvalidateUseHZB"),
 	1,
-	TEXT("Enables testing HZB for Virtual Shadow Map invalidations."),
+	TEXT(" When enabled, instances invalidations are tested against the HZB. Instances that are fully occluded will not cause page invalidations."),
 	ECVF_RenderThreadSafe);
 
 int32 GClipmapPanning = 1;
 FAutoConsoleVariableRef CVarEnableClipmapPanning(
 	TEXT("r.Shadow.Virtual.Cache.ClipmapPanning"),
 	GClipmapPanning,
-	TEXT("Enable support for panning cached clipmap pages for directional lights."),
+	TEXT("Enable support for panning cached clipmap pages for directional lights, allowing re-use of cached data when the camera moves. Keep this enabled outside of debugging."),
 	ECVF_RenderThreadSafe
 );
 
@@ -60,7 +60,7 @@ static int32 GVSMCacheDeformableMeshesInvalidate = 1;
 FAutoConsoleVariableRef CVarCacheInvalidateOftenMoving(
 	TEXT("r.Shadow.Virtual.Cache.DeformableMeshesInvalidate"),
 	GVSMCacheDeformableMeshesInvalidate,
-	TEXT("If enabled, Primitive Proxies that are marked as having deformable meshes (HasDeformableMesh() == true) causes invalidations regardless of whether their transforms are updated."),
+	TEXT("If enabled, Primitive Proxies that are marked as having deformable meshes (HasDeformableMesh() == true) cause invalidations regardless of whether their transforms are updated."),
 	ECVF_RenderThreadSafe);
 
 int32 GForceInvalidateDirectionalVSM = 0;
@@ -89,7 +89,8 @@ static TAutoConsoleVariable<int32> CVarVSMReservedResource(
 static TAutoConsoleVariable<float> CVarVSMDynamicResolutionMaxLodBias(
 	TEXT("r.Shadow.Virtual.DynamicRes.MaxResolutionLodBias"),
 	2.0f,
-	TEXT("Maximum LOD bias to clamp to for global dynamic resolution reduction. 0 = disabled"),
+	TEXT("As page allocation approaches the pool capacity, VSM resolution ramps down by biasing the LOD up, similar to 'ResolutionLodBiasDirectional'.\n")
+	TEXT("This is the maximum LOD bias to clamp to for global dynamic shadow resolution reduction. 0 = disabled"),
 	ECVF_RenderThreadSafe
 );
 
diff --git a/Engine/Source/Runtime/Renderer/Private/VirtualShadowMaps/VirtualShadowMapClipmap.cpp b/Engine/Source/Runtime/Renderer/Private/VirtualShadowMaps/VirtualShadowMapClipmap.cpp
--- a/Engine/Source/Runtime/Renderer/Private/VirtualShadowMaps/VirtualShadowMapClipmap.cpp
+++ b/Engine/Source/Runtime/Renderer/Private/VirtualShadowMaps/VirtualShadowMapClipmap.cpp
@@ -28,20 +28,21 @@ static TAutoConsoleVariable<float> CVarVirtualShadowMapResolutionLodBiasDirectio
 static TAutoConsoleVariable<float> CVarVirtualShadowMapResolutionLodBiasDirectionalMoving(
 	TEXT( "r.Shadow.Virtual.ResolutionLodBiasDirectionalMoving" ),
 	0.5f,
-	TEXT( "Bias applied to LOD calculations for directional lights that are moving. -1.0 doubles resolution, 1.0 halves it and so on." ),
+	TEXT( "Bias applied to LOD calculations for directional lights that are moving. -1.0 doubles resolution, 1.0 halves it and so on.\n" )
+	TEXT( "The bias transitions smoothly back to ResolutionLodBiasDirectional as the light transitions to non-moving, see 'r.Shadow.Scene.LightActiveFrameCount'." ),
 	ECVF_Scalability | ECVF_RenderThreadSafe
 );
 
 static TAutoConsoleVariable<int32> CVarVirtualShadowMapClipmapFirstLevel(
 	TEXT( "r.Shadow.Virtual.Clipmap.FirstLevel" ),
 	6,
-	TEXT( "First level of the virtual clipmap. Lower values allow higher resolution shadows closer to the camera." ),
+	TEXT( "First level of the virtual clipmap. Lower values allow higher resolution shadows closer to the camera, but may increase page count." ),
 	ECVF_Scalability | ECVF_RenderThreadSafe
 );
 static TAutoConsoleVariable<int32> CVarVirtualShadowMapClipmapLastLevel(
 	TEXT( "r.Shadow.Virtual.Clipmap.LastLevel" ),
 	22,
-	TEXT( "Last level of the virtual climap. Indirectly determines radius the clipmap can cover." ),
+	TEXT( "Last level of the virtual clipmap. Indirectly determines radius the clipmap can cover. Each extra level doubles the maximum range, but may increase page count." ),
 	ECVF_Scalability | ECVF_RenderThreadSafe
 );
 
@@ -62,7 +63,7 @@ TAutoConsoleVariable<int32> CVarVirtualShadowMapClipmapLastCoarseLevel(
 TAutoConsoleVariable<float> CVarVirtualShadowMapClipmapZRangeScale(
 	TEXT("r.Shadow.Virtual.Clipmap.ZRangeScale"),
 	1000.0f,
-	TEXT("Scale of the clipmap level depth range relative to the radius. Should generally be at least 10 or it will result in excessive cache invalidations."),
+	TEXT("Scale of the clipmap level depth range relative to the radius. Affects z-near/z-far of the shadow map. Should generally be at least 10 or it will result in excessive cache invalidations. Values that are too large cause depth imprecisions and shadow flickering."),
 	ECVF_RenderThreadSafe
 );
 
diff --git a/Engine/Source/Runtime/Renderer/Private/VirtualShadowMaps/VirtualShadowMapProjection.cpp b/Engine/Source/Runtime/Renderer/Private/VirtualShadowMaps/VirtualShadowMapProjection.cpp
--- a/Engine/Source/Runtime/Renderer/Private/VirtualShadowMaps/VirtualShadowMapProjection.cpp
+++ b/Engine/Source/Runtime/Renderer/Private/VirtualShadowMaps/VirtualShadowMapProjection.cpp
@@ -33,7 +33,7 @@
 static TAutoConsoleVariable<float> CVarScreenRayLength(
 	TEXT( "r.Shadow.Virtual.ScreenRayLength" ),
 	0.015f,
-	TEXT( "Length of the screen space shadow trace (smart shadow bias) before the virtual shadow map lookup." ),
+	TEXT( "Length of the screen space shadow trace away from receiver surface (smart shadow bias) before the VSM / SMRT lookup." ),
 	ECVF_Scalability | ECVF_RenderThreadSafe
 );
 
@@ -142,21 +142,24 @@ TAutoConsoleVariable<float> CVarSMRTMaxSlopeBiasLocal(
 static TAutoConsoleVariable<int32> CVarForcePerLightShadowMaskClear(
 	TEXT( "r.Shadow.Virtual.ForcePerLightShadowMaskClear" ),
 	0,
-	TEXT( "" ),
+	TEXT( "For debugging purposes. When enabled, the shadow mask texture is cleared before the projection pass writes to it. Projection pass writes all relevant pixels, so clearing should be unnecessary." ),
 	ECVF_RenderThreadSafe
 );
 
 static TAutoConsoleVariable<int32> CVarVSMTranslucentQuality(
 	TEXT("r.Shadow.Virtual.TranslucentQuality"),
 	0,
-	TEXT("Quality of shadow for lit translucent surfaces. This will be applied on all translucent surfaces, and has high-performance impact."),
+	TEXT("Quality of shadow for lit translucent surfaces. This will be applied on all translucent surfaces, and has high-performance impact.\n")
+	TEXT("Set to 1 to enable the high-quality mode."),
 	ECVF_ReadOnly | ECVF_RenderThreadSafe
 );
 
 static TAutoConsoleVariable<int32> CVarSubsurfaceShadowMinSourceAngle(
 	TEXT("r.Shadow.Virtual.SubsurfaceShadowMinSourceAngle"),
 	5,
-	TEXT("Minimum source angle (in degrees) used for shadow & transmittance of sub-surface materials with directional lights"),
+	TEXT("Minimum source angle (in degrees) used for shadow & transmittance of sub-surface materials with directional lights.\n")
+	TEXT("To emulate light diffusion with sub-surface materials, VSM can increase the light source radius depending on the material opacity.\n")
+	TEXT("The higher this value, the more diffuse the shadowing with these materials will appear."),
 	ECVF_RenderThreadSafe
 );
 
