Path Tracing: Remove bounds field from the path tracing light structure

This field was only used for two purposes - getting a bounding box over all lights (on the CPU), and improving culling of spot lights (on the GPU) when building the light grid.

Remove the field saves 24 bytes from the light struct which is a nice savings, and makes it more similar to what will be required for making a similar structure with FRTLightingData which does not include a bounding box as is already 128 bytes.

Spend 12 bytes relaxing the encoding of certain half-precision values back to float precision.

diff --git a/Engine/Plugins/Experimental/GPULightmass/Source/GPULightmass/Private/PathTracingLightParameters.inl b/Engine/Plugins/Experimental/GPULightmass/Source/GPULightmass/Private/PathTracingLightParameters.inl
--- a/Engine/Plugins/Experimental/GPULightmass/Source/GPULightmass/Private/PathTracingLightParameters.inl
+++ b/Engine/Plugins/Experimental/GPULightmass/Source/GPULightmass/Private/PathTracingLightParameters.inl
@@ -94,13 +94,6 @@ void SetupPathTracingLightParameters(
 		DestLight.Flags |= Light.IsInverseSquared ? 0 : PATHTRACER_FLAG_NON_INVERSE_SQUARE_FALLOFF_MASK;
 		DestLight.Flags |= Light.bStationary ? PATHTRACER_FLAG_STATIONARY_MASK : 0;
 		DestLight.Flags |= PATHTRACING_LIGHT_POINT;
-
-		float Radius = Light.AttenuationRadius;
-		FVector3f Center = DestLight.TranslatedWorldPosition;
-
-		// simple sphere of influence
-		DestLight.TranslatedBoundMin = Center - FVector3f(Radius, Radius, Radius);
-		DestLight.TranslatedBoundMax = Center + FVector3f(Radius, Radius, Radius);
 	}
 
 	for (auto Light : LightScene.SpotLights.Elements)
@@ -113,7 +106,7 @@ void SetupPathTracingLightParameters(
 		DestLight.dPdv = (FVector3f)Light.Tangent;
 		DestLight.Color = FVector3f(Light.Color);
 		DestLight.Dimensions = FVector2f(Light.SourceRadius, Light.SourceLength);
-		DestLight.Shaping = EncodeToF16x2(FVector2f(Light.SpotAngles));
+		DestLight.Shaping = FVector2f(Light.SpotAngles);
 		DestLight.Attenuation = 1.0f / Light.AttenuationRadius;
 		DestLight.FalloffExponent = Light.FalloffExponent;
 		DestLight.IESAtlasIndex = Light.GetIESAtlasIndex();
@@ -124,27 +117,6 @@ void SetupPathTracingLightParameters(
 		DestLight.Flags |= Light.IsInverseSquared ? 0 : PATHTRACER_FLAG_NON_INVERSE_SQUARE_FALLOFF_MASK;
 		DestLight.Flags |= Light.bStationary ? PATHTRACER_FLAG_STATIONARY_MASK : 0;
 		DestLight.Flags |= PATHTRACING_LIGHT_SPOT;
-
-		// LWC_TODO: Precision Loss
-		float Radius = Light.AttenuationRadius;
-		FVector3f Center = DestLight.TranslatedWorldPosition;
-		FVector3f Normal = DestLight.Normal;
-		FVector3f Disc = FVector3f(
-			FMath::Sqrt(FMath::Clamp(1 - Normal.X * Normal.X, 0.0f, 1.0f)),
-			FMath::Sqrt(FMath::Clamp(1 - Normal.Y * Normal.Y, 0.0f, 1.0f)),
-			FMath::Sqrt(FMath::Clamp(1 - Normal.Z * Normal.Z, 0.0f, 1.0f))
-		);
-		// box around ray from light center to tip of the cone
-		FVector3f Tip = Center + Normal * Radius;
-		DestLight.TranslatedBoundMin = Center.ComponentMin(Tip);
-		DestLight.TranslatedBoundMax = Center.ComponentMax(Tip);
-		// expand by disc around the farthest part of the cone
-
-		float CosOuter = Light.SpotAngles.X;
-		float SinOuter = FMath::Sqrt(1.0f - CosOuter * CosOuter);
-
-		DestLight.TranslatedBoundMin = DestLight.TranslatedBoundMin.ComponentMin(Center + Radius * (Normal * CosOuter - Disc * SinOuter));
-		DestLight.TranslatedBoundMax = DestLight.TranslatedBoundMax.ComponentMax(Center + Radius * (Normal * CosOuter + Disc * SinOuter));
 	}
 
 	for (auto Light : LightScene.RectLights.Elements)
@@ -162,7 +134,7 @@ void SetupPathTracingLightParameters(
 
 		DestLight.Dimensions = FVector2f(Light.SourceWidth, Light.SourceHeight);
 		DestLight.Attenuation = 1.0f / Light.AttenuationRadius;
-		DestLight.Shaping = EncodeToF16x2(FVector2f(FMath::Cos(FMath::DegreesToRadians(Light.BarnDoorAngle)), Light.BarnDoorLength));
+		DestLight.Shaping = FVector2f(FMath::Cos(FMath::DegreesToRadians(Light.BarnDoorAngle)), Light.BarnDoorLength);
 		DestLight.IESAtlasIndex = Light.GetIESAtlasIndex();
 
 		DestLight.Flags = PATHTRACER_FLAG_TRANSMISSION_MASK;
@@ -171,22 +143,8 @@ void SetupPathTracingLightParameters(
 		DestLight.Flags |= Light.bStationary ? PATHTRACER_FLAG_STATIONARY_MASK : 0;
 		DestLight.Flags |= PATHTRACING_LIGHT_RECT;
 
-		float Radius = Light.AttenuationRadius;
-		FVector3f Center = DestLight.TranslatedWorldPosition;
-		FVector3f Normal = DestLight.Normal;
-		FVector3f Disc = FVector3f(
-			FMath::Sqrt(FMath::Clamp(1 - Normal.X * Normal.X, 0.0f, 1.0f)),
-			FMath::Sqrt(FMath::Clamp(1 - Normal.Y * Normal.Y, 0.0f, 1.0f)),
-			FMath::Sqrt(FMath::Clamp(1 - Normal.Z * Normal.Z, 0.0f, 1.0f))
-		);
-		// quad bbox is the bbox of the disc +  the tip of the hemisphere
-		// TODO: is it worth trying to account for barndoors? seems unlikely to cut much empty space since the volume _inside_ the barndoor receives light
-		FVector3f Tip = Center + Normal * Radius;
-		DestLight.TranslatedBoundMin = Tip.ComponentMin(Center - Radius * Disc);
-		DestLight.TranslatedBoundMax = Tip.ComponentMax(Center + Radius * Disc);
-
-		DestLight.RectLightAtlasUVOffset = EncodeToF16x2(Light.RectLightAtlasUVOffset);
-		DestLight.RectLightAtlasUVScale  = EncodeToF16x2(Light.RectLightAtlasUVScale);
+		DestLight.RectLightAtlasUVOffset = Light.RectLightAtlasUVOffset;
+		DestLight.RectLightAtlasUVScale  = Light.RectLightAtlasUVScale;
 		if (Light.RectLightAtlasMaxLevel < 16)
 		{
 			DestLight.Flags |= PATHTRACER_FLAG_HAS_RECT_TEXTURE_MASK;
diff --git a/Engine/Shaders/Private/PathTracing/Light/PathTracingLightCommon.ush b/Engine/Shaders/Private/PathTracing/Light/PathTracingLightCommon.ush
--- a/Engine/Shaders/Private/PathTracing/Light/PathTracingLightCommon.ush
+++ b/Engine/Shaders/Private/PathTracing/Light/PathTracingLightCommon.ush
@@ -220,17 +220,9 @@ float2 GetRectSize(int LightId)
 	return SceneLights[LightId].Dimensions.xy;
 }
 
-float2 DecodeLightF16x2(uint Encoded)
-{
-	return float2(
-		f16tof32(Encoded),
-		f16tof32(Encoded >> 16)
-	);
-}
-
 float2 GetCosConeAngles(int LightId)
 {
-	return DecodeLightF16x2(SceneLights[LightId].Shaping);
+	return SceneLights[LightId].Shaping;
 }
 
 float3 GetColor(int LightId)
@@ -255,22 +247,22 @@ float GetAttenuation(int LightId)
 
 float GetRectLightBarnCosAngle(int LightId)
 {
-	return DecodeLightF16x2(SceneLights[LightId].Shaping).x;
+	return SceneLights[LightId].Shaping.x;
 }
 
 float GetRectLightBarnLength(int LightId)
 {
-	return DecodeLightF16x2(SceneLights[LightId].Shaping).y;
+	return SceneLights[LightId].Shaping.y;
 }
 
 float2 GetRectLightAtlasUVScale(int LightId)
 {
-	return DecodeLightF16x2(SceneLights[LightId].RectLightAtlasUVScale);
+	return SceneLights[LightId].RectLightAtlasUVScale;
 }
 
 float2 GetRectLightAtlasUVOffset(int LightId)
 {
-	return DecodeLightF16x2(SceneLights[LightId].RectLightAtlasUVOffset);
+	return SceneLights[LightId].RectLightAtlasUVOffset;
 }
 
 #ifndef USE_ATTENUATION_TERM
diff --git a/Engine/Shaders/Private/PathTracing/PathTracingBuildLightGrid.usf b/Engine/Shaders/Private/PathTracing/PathTracingBuildLightGrid.usf
--- a/Engine/Shaders/Private/PathTracing/PathTracingBuildLightGrid.usf
+++ b/Engine/Shaders/Private/PathTracing/PathTracingBuildLightGrid.usf
@@ -47,6 +47,7 @@ void PathTracingBuildLightGridCS(uint2 DispatchThreadId : SV_DispatchThreadID)
 		uint LightType = SceneLights[LightIndex].Flags & PATHTRACER_FLAG_TYPE_MASK;
 		float3 Center = GetTranslatedPosition(LightIndex);
 		float Radius = 1.0 / GetAttenuation(LightIndex);
+		float3 Normal = GetNormal(LightIndex);
 
 		float3 E = max(Lo - Center, 0) + max(Center - Hi, 0);
 		float D2 = dot(E, E);
@@ -54,12 +55,7 @@ void PathTracingBuildLightGridCS(uint2 DispatchThreadId : SV_DispatchThreadID)
 		// reject against normal
 		if (LightType != PATHTRACING_LIGHT_POINT)
 		{
-			float3 N = GetNormal(LightIndex);
-			float x = N.x > 0 ? Hi.x : Lo.x;
-			float y = N.y > 0 ? Hi.y : Lo.y;
-			float z = N.z > 0 ? Hi.z : Lo.z;
-
-			if (dot(float3(x, y, z) - Center, N) < 0)
+			if (dot(select(Normal > 0, Hi, Lo) - Center, Normal) < 0)
 			{
 				continue;
 			}
@@ -67,8 +63,18 @@ void PathTracingBuildLightGridCS(uint2 DispatchThreadId : SV_DispatchThreadID)
 
 		if (LightType == PATHTRACING_LIGHT_SPOT)
 		{
-			float3 LightBoundLo = SceneLights[LightIndex].TranslatedBoundMin;
-			float3 LightBoundHi = SceneLights[LightIndex].TranslatedBoundMax;
+			float3 Disc = sqrt(saturate(1.0 - Normal * Normal));
+			// box around ray from light center to tip of the cone
+			float3 Tip = Center + Normal * Radius;
+			float3 LightBoundLo = min(Center, Tip);
+			float3 LightBoundHi = max(Center, Tip);
+			// expand by disc around the farthest part of the cone
+
+			float CosOuter = GetCosConeAngles(LightIndex).x;
+			float SinOuter = sqrt(1.0f - CosOuter * CosOuter);
+
+			LightBoundLo = min(LightBoundLo, Center + Radius * (Normal * CosOuter - Disc * SinOuter));
+			LightBoundHi = max(LightBoundHi, Center + Radius * (Normal * CosOuter + Disc * SinOuter));
 
 			if (any(Hi < LightBoundLo) || any(Lo > LightBoundHi))
 			{
diff --git a/Engine/Shaders/Shared/RayTracingTypes.h b/Engine/Shaders/Shared/RayTracingTypes.h
--- a/Engine/Shaders/Shared/RayTracingTypes.h
+++ b/Engine/Shaders/Shared/RayTracingTypes.h
@@ -49,7 +49,7 @@ struct FPathTracingLight {
 	float3  dPdv;
 	float3  Color;
 	float2  Dimensions; // Radius,Length or RectWidth,RectHeight or Sin(Angle/2),0 depending on light type
-	uint    Shaping;    // Barndoor controls for RectLights, Cone angles for spots lights, encoded as f16x2
+	float2  Shaping;    // Barndoor controls for RectLights, Cone angles for spots lights
 	float   SpecularScale;
 	float   Attenuation;
 	float   FalloffExponent; // for non-inverse square decay lights only
@@ -57,12 +57,10 @@ struct FPathTracingLight {
 	int     IESAtlasIndex;
 	uint    Flags; // see defines PATHTRACER_FLAG_*
 	uint    MissShaderIndex;  // used to implement light functions
-	float3  TranslatedBoundMin;
-	float3  TranslatedBoundMax;
-	uint	RectLightAtlasUVScale;  // Rect. light atlas UV transformation, encoded as f16x2
-	uint	RectLightAtlasUVOffset; // Rect. light atlas UV transformation, encoded as f16x2
+	float2  RectLightAtlasUVScale;  // Rect. light atlas UV transformation
+	float2  RectLightAtlasUVOffset; // Rect. light atlas UV transformation
 };
-HLSL_STATIC_ASSERT(sizeof(FPathTracingLight) == 132, "Path tracing light structure should be kept as small as possible");
+HLSL_STATIC_ASSERT(sizeof(FPathTracingLight) == 120, "Path tracing light structure should be kept as small as possible");
 
 struct FPathTracingPackedPathState {
 	uint      RandSeqSampleIndex;
diff --git a/Engine/Source/Runtime/Renderer/Private/PathTracing.cpp b/Engine/Source/Runtime/Renderer/Private/PathTracing.cpp
--- a/Engine/Source/Runtime/Renderer/Private/PathTracing.cpp
+++ b/Engine/Source/Runtime/Renderer/Private/PathTracing.cpp
@@ -1285,7 +1285,6 @@ public:
 
 	static void ModifyCompilationEnvironment(const FMaterialShaderPermutationParameters& Parameters, FShaderCompilerEnvironment& OutEnvironment)
 	{
-		OutEnvironment.SetDefine(TEXT("SUPPORT_LIGHT_FUNCTION"), 1);
 		FMaterialShader::ModifyCompilationEnvironment(Parameters, OutEnvironment);
 	}
 
@@ -1895,8 +1894,59 @@ RENDERER_API void PrepareLightGrid(FRDGBuilder& GraphBuilder, ERHIFeatureLevel::
 		for (int Index = 0; Index < NumFiniteLights; Index++)
 		{
 			const FPathTracingLight& Light = FiniteLights[Index];
-			LightGridParameters->SceneLightsTranslatedBoundMin = FVector3f::Min(LightGridParameters->SceneLightsTranslatedBoundMin, Light.TranslatedBoundMin);
-			LightGridParameters->SceneLightsTranslatedBoundMax = FVector3f::Max(LightGridParameters->SceneLightsTranslatedBoundMax, Light.TranslatedBoundMax);
+			FVector3f Lo = FVector3f(-Inf, -Inf, -Inf);
+			FVector3f Hi = FVector3f( Inf,  Inf,  Inf);
+
+			const float Radius = 1.0f / Light.Attenuation;
+			const FVector3f Center = Light.TranslatedWorldPosition;
+			const FVector3f Normal = Light.Normal;
+			const FVector3f Disc = FVector3f(
+				FMath::Sqrt(FMath::Clamp(1 - Normal.X * Normal.X, 0.0f, 1.0f)),
+				FMath::Sqrt(FMath::Clamp(1 - Normal.Y * Normal.Y, 0.0f, 1.0f)),
+				FMath::Sqrt(FMath::Clamp(1 - Normal.Z * Normal.Z, 0.0f, 1.0f))
+			);
+			switch (Light.Flags & PATHTRACER_FLAG_TYPE_MASK)
+			{
+				case PATHTRACING_LIGHT_POINT:
+				{
+					// simple sphere of influence
+					Lo = Center - FVector3f(Radius, Radius, Radius);
+					Hi = Center + FVector3f(Radius, Radius, Radius);
+					break;
+				}
+				case PATHTRACING_LIGHT_SPOT:
+				{
+					// box around ray from light center to tip of the cone
+					const FVector3f Tip = Center + Normal * Radius;
+					Lo = FVector3f::Min(Center, Tip);
+					Hi = FVector3f::Max(Center, Tip);
+
+					// expand by disc around the farthest part of the cone
+					const float CosOuter = Light.Shaping.X;
+					const float SinOuter = FMath::Sqrt(1.0f - CosOuter * CosOuter);
+
+					Lo = FVector3f::Min(Lo, Center + Radius * (Normal * CosOuter - Disc * SinOuter));
+					Hi = FVector3f::Max(Hi, Center + Radius * (Normal * CosOuter + Disc * SinOuter));
+					break;
+				}
+				case PATHTRACING_LIGHT_RECT:
+				{
+					// quad bbox is the bbox of the disc +  the tip of the hemisphere
+					// TODO: is it worth trying to account for barndoors? seems unlikely to cut much empty space since the volume _inside_ the barndoor receives light
+					const FVector3f Tip = Center + Normal * Radius;
+					Lo = FVector3f::Min(Tip, Center - Radius * Disc);
+					Hi = FVector3f::Max(Tip, Center + Radius * Disc);
+					break;
+				}
+				default:
+				{
+					// non-finite lights should not appear in this case
+					checkSlow(false);
+					break;
+				}
+			}
+			LightGridParameters->SceneLightsTranslatedBoundMin = FVector3f::Min(LightGridParameters->SceneLightsTranslatedBoundMin, Lo);
+			LightGridParameters->SceneLightsTranslatedBoundMax = FVector3f::Max(LightGridParameters->SceneLightsTranslatedBoundMax, Hi);
 		}
 
 		const uint32 Resolution = FMath::RoundUpToPowerOfTwo(CVarPathTracingLightGridResolution.GetValueOnRenderThread());
@@ -1989,11 +2039,6 @@ RENDERER_API void PrepareLightGrid(FRDGBuilder& GraphBuilder, ERHIFeatureLevel::
 	}
 }
 
-static uint32 EncodeToF16x2(const FVector2f& In)
-{
-	return FFloat16(In.X).Encoded | (FFloat16(In.Y).Encoded << 16);
-}
-
 void SetLightParameters(FRDGBuilder& GraphBuilder, FPathTracingRG::FParameters* PassParameters, FScene* Scene, const FViewInfo& View, bool UseMISCompensation)
 {
 	PassParameters->SceneVisibleLightCount = 0;
@@ -2005,7 +2050,6 @@ void SetLightParameters(FRDGBuilder& GraphBuilder, FPathTracingRG::FParameters*
 	uint32 NumLights = 0;
 
 	// Prepend SkyLight to light buffer since it is not part of the regular light list
-	const float Inf = std::numeric_limits<float>::infinity();
 	// skylight should be excluded if we are using the reference atmosphere calculation (don't bother checking again if an atmosphere is present)
 	const bool bUseAtmosphere = PassParameters->PathTracingData.EnableAtmosphere != 0;
 	const bool bEnableSkydome = !bUseAtmosphere;
@@ -2023,8 +2067,6 @@ void SetLightParameters(FRDGBuilder& GraphBuilder, FPathTracingRG::FParameters*
 		DestLight.VolumetricScatteringIntensity = Scene->SkyLight->VolumetricScatteringIntensity;
 		DestLight.IESAtlasIndex = INDEX_NONE;
 		DestLight.MissShaderIndex = 0;
-		DestLight.TranslatedBoundMin = FVector3f(-Inf, -Inf, -Inf);
-		DestLight.TranslatedBoundMax = FVector3f( Inf,  Inf,  Inf);
 		if (Scene->SkyLight->bRealTimeCaptureEnabled || CVarPathTracingVisibleLights.GetValueOnRenderThread() == 2)
 		{
 			// When using the realtime capture system, always make the skylight visible
@@ -2037,7 +2079,7 @@ void SetLightParameters(FRDGBuilder& GraphBuilder, FPathTracingRG::FParameters*
 	// Add directional lights next (all lights with infinite bounds should come first)
 	if (View.Family->EngineShowFlags.DirectionalLights)
 	{
-		for (auto Light : Scene->Lights)
+		for (const FLightSceneInfoCompact& Light : Scene->Lights)
 		{
 			ELightComponentType LightComponentType = (ELightComponentType)Light.LightSceneInfo->Proxy->GetLightType();
 
@@ -2071,20 +2113,17 @@ void SetLightParameters(FRDGBuilder& GraphBuilder, FPathTracingRG::FParameters*
 			DestLight.Normal = -LightParameters.Direction;
 			DestLight.dPdu = FVector3f::CrossProduct(LightParameters.Tangent, LightParameters.Direction);
 			DestLight.dPdv = LightParameters.Tangent;
-			DestLight.Shaping = 0;
+			DestLight.Shaping = FVector2f(0.0f, 0.0f);
 			DestLight.SpecularScale = LightParameters.SpecularScale;
 			DestLight.Attenuation = LightParameters.InvRadius;
 			DestLight.FalloffExponent = 0;
 			DestLight.VolumetricScatteringIntensity = Light.LightSceneInfo->Proxy->GetVolumetricScatteringIntensity();
-			DestLight.RectLightAtlasUVOffset = 0;
-			DestLight.RectLightAtlasUVScale = 0;
+			DestLight.RectLightAtlasUVOffset = FVector2f(0.0f, 0.0f);
+			DestLight.RectLightAtlasUVScale = FVector2f(0.0f, 0.0f);
 
 			DestLight.Normal = LightParameters.Direction;
 			DestLight.Dimensions = FVector2f(LightParameters.SourceRadius, 0.0f);
 			DestLight.Flags |= PATHTRACING_LIGHT_DIRECTIONAL;
-
-			DestLight.TranslatedBoundMin = FVector3f(-Inf, -Inf, -Inf);
-			DestLight.TranslatedBoundMax = FVector3f( Inf,  Inf,  Inf);
 		}
 	}
 
@@ -2100,7 +2139,7 @@ void SetLightParameters(FRDGBuilder& GraphBuilder, FPathTracingRG::FParameters*
 	int32 NextRectTextureIndex = 0;
 
 	const FRayTracingLightFunctionMap* RayTracingLightFunctionMap = GraphBuilder.Blackboard.Get<FRayTracingLightFunctionMap>();
-	for (auto Light : Scene->Lights)
+	for (const FLightSceneInfoCompact& Light : Scene->Lights)
 	{
 		ELightComponentType LightComponentType = (ELightComponentType)Light.LightSceneInfo->Proxy->GetLightType();
 
@@ -2139,13 +2178,13 @@ void SetLightParameters(FRDGBuilder& GraphBuilder, FPathTracingRG::FParameters*
 		DestLight.Normal = -LightParameters.Direction;
 		DestLight.dPdu = FVector3f::CrossProduct(LightParameters.Tangent, LightParameters.Direction);
 		DestLight.dPdv = LightParameters.Tangent;
-		DestLight.Shaping = 0;
+		DestLight.Shaping = FVector2f(0.0f, 0.0f);
 		DestLight.SpecularScale = LightParameters.SpecularScale;
 		DestLight.Attenuation = LightParameters.InvRadius;
 		DestLight.FalloffExponent = 0;
 		DestLight.VolumetricScatteringIntensity = Light.LightSceneInfo->Proxy->GetVolumetricScatteringIntensity();
-		DestLight.RectLightAtlasUVOffset = 0;
-		DestLight.RectLightAtlasUVScale = 0;
+		DestLight.RectLightAtlasUVOffset = FVector2f(0.0f, 0.0f);
+		DestLight.RectLightAtlasUVScale = FVector2f(0.0f, 0.0f);
 
 		if (RayTracingLightFunctionMap)
 		{
@@ -2161,28 +2200,15 @@ void SetLightParameters(FRDGBuilder& GraphBuilder, FPathTracingRG::FParameters*
 			case LightType_Rect:
 			{
 				DestLight.Dimensions = FVector2f(2.0f * LightParameters.SourceRadius, 2.0f * LightParameters.SourceLength);
-				DestLight.Shaping = EncodeToF16x2(FVector2f(LightParameters.RectLightBarnCosAngle, LightParameters.RectLightBarnLength));
+				DestLight.Shaping = FVector2f(LightParameters.RectLightBarnCosAngle, LightParameters.RectLightBarnLength);
 				DestLight.FalloffExponent = LightParameters.FalloffExponent;
 				DestLight.Flags |= Light.LightSceneInfo->Proxy->IsInverseSquared() ? 0 : PATHTRACER_FLAG_NON_INVERSE_SQUARE_FALLOFF_MASK;
 				DestLight.Flags |= PATHTRACING_LIGHT_RECT;
 
-				float Radius = 1.0f / LightParameters.InvRadius;
-				FVector3f Center = DestLight.TranslatedWorldPosition;
-				FVector3f Normal = DestLight.Normal;
-				FVector3f Disc = FVector3f(
-					FMath::Sqrt(FMath::Clamp(1 - Normal.X * Normal.X, 0.0f, 1.0f)),
-					FMath::Sqrt(FMath::Clamp(1 - Normal.Y * Normal.Y, 0.0f, 1.0f)),
-					FMath::Sqrt(FMath::Clamp(1 - Normal.Z * Normal.Z, 0.0f, 1.0f))
-				);
-				// quad bbox is the bbox of the disc +  the tip of the hemisphere
-				// TODO: is it worth trying to account for barndoors? seems unlikely to cut much empty space since the volume _inside_ the barndoor receives light
-				FVector3f Tip = Center + Normal * Radius;
-				DestLight.TranslatedBoundMin = FVector3f::Min(Tip, Center - Radius * Disc);
-				DestLight.TranslatedBoundMax = FVector3f::Max(Tip, Center + Radius * Disc);
 
 				// Rect light atlas UV transformation
-				DestLight.RectLightAtlasUVOffset = EncodeToF16x2(LightParameters.RectLightAtlasUVOffset);
-				DestLight.RectLightAtlasUVScale  = EncodeToF16x2(LightParameters.RectLightAtlasUVScale);
+				DestLight.RectLightAtlasUVOffset = LightParameters.RectLightAtlasUVOffset;
+				DestLight.RectLightAtlasUVScale  = LightParameters.RectLightAtlasUVScale;
 				if (LightParameters.RectLightAtlasMaxLevel < 16)
 				{
 					DestLight.Flags |= PATHTRACER_FLAG_HAS_RECT_TEXTURE_MASK;
@@ -2192,30 +2218,10 @@ void SetLightParameters(FRDGBuilder& GraphBuilder, FPathTracingRG::FParameters*
 			case LightType_Spot:
 			{
 				DestLight.Dimensions = FVector2f(LightParameters.SourceRadius, LightParameters.SourceLength);
-				DestLight.Shaping = EncodeToF16x2(LightParameters.SpotAngles);
+				DestLight.Shaping = LightParameters.SpotAngles;
 				DestLight.FalloffExponent = LightParameters.FalloffExponent;
 				DestLight.Flags |= Light.LightSceneInfo->Proxy->IsInverseSquared() ? 0 : PATHTRACER_FLAG_NON_INVERSE_SQUARE_FALLOFF_MASK;
 				DestLight.Flags |= PATHTRACING_LIGHT_SPOT;
-
-				float Radius = 1.0f / LightParameters.InvRadius;
-				FVector3f Center = DestLight.TranslatedWorldPosition;
-				FVector3f Normal = DestLight.Normal;
-				FVector3f Disc = FVector3f(
-					FMath::Sqrt(FMath::Clamp(1 - Normal.X * Normal.X, 0.0f, 1.0f)),
-					FMath::Sqrt(FMath::Clamp(1 - Normal.Y * Normal.Y, 0.0f, 1.0f)),
-					FMath::Sqrt(FMath::Clamp(1 - Normal.Z * Normal.Z, 0.0f, 1.0f))
-				);
-				// box around ray from light center to tip of the cone
-				FVector3f Tip = Center + Normal * Radius;
-				DestLight.TranslatedBoundMin = FVector3f::Min(Center, Tip);
-				DestLight.TranslatedBoundMax = FVector3f::Max(Center, Tip);
-				// expand by disc around the farthest part of the cone
-
-				float CosOuter = LightParameters.SpotAngles.X;
-				float SinOuter = FMath::Sqrt(1.0f - CosOuter * CosOuter);
-
-				DestLight.TranslatedBoundMin = FVector3f::Min(DestLight.TranslatedBoundMin, Center + Radius * (Normal * CosOuter - Disc * SinOuter));
-				DestLight.TranslatedBoundMax = FVector3f::Max(DestLight.TranslatedBoundMax, Center + Radius * (Normal * CosOuter + Disc * SinOuter));
 				break;
 			}
 			case LightType_Point:
@@ -2224,12 +2230,6 @@ void SetLightParameters(FRDGBuilder& GraphBuilder, FPathTracingRG::FParameters*
 				DestLight.FalloffExponent = LightParameters.FalloffExponent;
 				DestLight.Flags |= Light.LightSceneInfo->Proxy->IsInverseSquared() ? 0 : PATHTRACER_FLAG_NON_INVERSE_SQUARE_FALLOFF_MASK;
 				DestLight.Flags |= PATHTRACING_LIGHT_POINT;
-				float Radius = 1.0f / LightParameters.InvRadius;
-				FVector3f Center = DestLight.TranslatedWorldPosition;
-
-				// simple sphere of influence
-				DestLight.TranslatedBoundMin = Center - FVector3f(Radius, Radius, Radius);
-				DestLight.TranslatedBoundMax = Center + FVector3f(Radius, Radius, Radius);
 				break;
 			}
 			default:
