Fix bug where downsampling capsule shadows would cause a crash due to a groupsize of 0

diff --git a/Engine/Source/Runtime/Renderer/Private/CapsuleShadowRendering.cpp b/Engine/Source/Runtime/Renderer/Private/CapsuleShadowRendering.cpp
--- a/Engine/Source/Runtime/Renderer/Private/CapsuleShadowRendering.cpp
+++ b/Engine/Source/Runtime/Renderer/Private/CapsuleShadowRendering.cpp
@@ -549,9 +549,10 @@ bool FDeferredShadingSceneRenderer::RenderCapsuleDirectShadows(
 				ScissorRect = View.ViewRect;
 			}
 
-			const FIntPoint GroupSize(
+			const FIntPoint GroupSize = FIntPoint(
 				FMath::DivideAndRoundUp(ScissorRect.Size().X / GetCapsuleShadowDownsampleFactor(), GShadowShapeTileSize),
-				FMath::DivideAndRoundUp(ScissorRect.Size().Y / GetCapsuleShadowDownsampleFactor(), GShadowShapeTileSize));
+				FMath::DivideAndRoundUp(ScissorRect.Size().Y / GetCapsuleShadowDownsampleFactor(), GShadowShapeTileSize))
+				.ComponentMax(FIntPoint(1, 1));
 
 			FRDGBufferRef CapsuleTileIntersectionCountsBuffer = GraphBuilder.CreateBuffer(FRDGBufferDesc::CreateStructuredDesc(sizeof(uint32), GroupSize.X * GroupSize.Y), TEXT("CapsuleTileIntersectionCountsBuffer"));
 			FRDGBufferUAVRef CapsuleTileIntersectionCountsUAV = GraphBuilder.CreateUAV(CapsuleTileIntersectionCountsBuffer);
