Implemented parallel Gather Dynamic Mesh Elements (GDME) for visibility and shadows.
 - Added r.ParallelGatherDynamicMeshElements to control whether the feature is enabled. When disabled, tasks are constrained to the render thread as before.
 - Added r.Visibility.DynamicMeshElements.NumTasks and r.Shadow.DynamicMeshElements.NumTasks which both default to 4. Thread contention becomes an issue at higher task counts for proxies which do a lot of uniform buffer or resource creation / updates.

Implementation Details:

The scene view visibility path typically has many dynamic elements per view, whereas there are typically more shadow views processing fewer elements each.

The main visibility path pushes dynamic primitives into a queue which is then processed by N async tasks, each with its own mesh collector and global vertex / index buffer. If a scene proxy is marked as not supporting parallel GDME, it is processed instead on the render thread in parallel with other async tasks. Niagara is currently the only proxy which is marked as unsupported by parallel GDME, as it has many potential sync points that would need to be untangled. When the async tasks complete, the task contexts are merged back to the view in an async task prior to launching dynamic mesh passes. Another important detail is that all material and GPU scene ops are deferred by the mesh collector and finalized at the end on the render thread. Deferring materials is necessary to allow the VT update task to overlap with GDME. The GPU scene updates are not thread safe and must be deferred.

The shadow path instead creates a queue of FProjectedShadowInfo's which is processed in parallel by async gather tasks (each with its own collector context). If a particular shadow has proxies that do not support parallel GDME, then a second pass is done over that FProjectedShadowInfo on the render thread to process the render thread only proxies. The dynamic mesh passes are launched asynchronously as well and are waited on by the shadow depth rendering passes later in the frame.

diff --git a/Engine/Config/Android/AndroidEngine.ini b/Engine/Config/Android/AndroidEngine.ini
--- a/Engine/Config/Android/AndroidEngine.ini
+++ b/Engine/Config/Android/AndroidEngine.ini
@@ -102,6 +102,8 @@ r.Substrate.SpecularProfile=0
 r.Substrate.ClosuresPerPixel=1
 ; Enable GPU culling for landscape
 landscape.SupportGPUCulling=1
+; Disables parallel GDME by default due to not enough threads
+r.Visibility.DynamicMeshElements.Parallel=0
 
 [AlternateTextureCompression]
 ;disabling TextureCompressionFormat for Android/IOS no longer needed
diff --git a/Engine/Source/Runtime/Engine/Private/Components/MeshComponent.cpp b/Engine/Source/Runtime/Engine/Private/Components/MeshComponent.cpp
--- a/Engine/Source/Runtime/Engine/Private/Components/MeshComponent.cpp
+++ b/Engine/Source/Runtime/Engine/Private/Components/MeshComponent.cpp
@@ -29,7 +29,7 @@ UMeshComponent::UMeshComponent(const FObjectInitializer& ObjectInitializer)
 	bUseAsOccluder = true;
 	bCanEverAffectNavigation = true;
 	bCachedMaterialParameterIndicesAreDirty = true;
-	bEnableMaterialParameterCaching = false;
+	bEnableMaterialParameterCaching = true;
 }
 
 UMaterialInterface* UMeshComponent::GetMaterial(int32 ElementIndex) const
diff --git a/Engine/Source/Runtime/Engine/Private/PrimitiveSceneProxy.cpp b/Engine/Source/Runtime/Engine/Private/PrimitiveSceneProxy.cpp
--- a/Engine/Source/Runtime/Engine/Private/PrimitiveSceneProxy.cpp
+++ b/Engine/Source/Runtime/Engine/Private/PrimitiveSceneProxy.cpp
@@ -106,6 +106,19 @@ static TAutoConsoleVariable<int32> CVarApproximateOcclusionQueries(
 	ECVF_RenderThreadSafe
 );
 
+bool GParallelGatherDynamicMeshElements = true;
+static FAutoConsoleVariableRef CVarParallelGatherDynamicMeshElements(
+	TEXT("r.Visibility.DynamicMeshElements.Parallel"),
+	GParallelGatherDynamicMeshElements,
+	TEXT("Enables parallel processing of the gather dynamic mesh elements visibility phase."),
+	ECVF_RenderThreadSafe
+);
+
+bool IsParallelGatherDynamicMeshElementsEnabled()
+{
+	return GParallelGatherDynamicMeshElements;
+}
+
 bool FPrimitiveSceneProxy::ShouldRenderCustomDepth() const
 {
 	return IsCustomDepthPassEnabled() && bRenderCustomDepth;
diff --git a/Engine/Source/Runtime/Engine/Private/SkeletalRenderGPUSkin.cpp b/Engine/Source/Runtime/Engine/Private/SkeletalRenderGPUSkin.cpp
--- a/Engine/Source/Runtime/Engine/Private/SkeletalRenderGPUSkin.cpp
+++ b/Engine/Source/Runtime/Engine/Private/SkeletalRenderGPUSkin.cpp
@@ -476,7 +476,7 @@ static TAutoConsoleVariable<int32> CVarDeferSkeletalDynamicDataUpdateUntilGDME(
 
 FORCEINLINE bool IsDeferredSkeletalDynamicDataUpdateEnabled()
 {
-	return CVarDeferSkeletalDynamicDataUpdateUntilGDME.GetValueOnRenderThread() > 0;
+	return CVarDeferSkeletalDynamicDataUpdateUntilGDME.GetValueOnRenderThread() > 0 && IsParallelGatherDynamicMeshElementsEnabled();
 }
 
 void FSkeletalMeshObjectGPUSkin::UpdateDynamicData_RenderThread(FGPUSkinCache* GPUSkinCache, FRHICommandList& RHICmdList, FDynamicSkelMeshObjectDataGPUSkin* InDynamicData, FSceneInterface* Scene, uint64 FrameNumberToPrepare, uint32 RevisionNumber)
diff --git a/Engine/Source/Runtime/Engine/Public/PrimitiveSceneProxy.h b/Engine/Source/Runtime/Engine/Public/PrimitiveSceneProxy.h
--- a/Engine/Source/Runtime/Engine/Public/PrimitiveSceneProxy.h
+++ b/Engine/Source/Runtime/Engine/Public/PrimitiveSceneProxy.h
@@ -1592,3 +1592,6 @@ ENGINE_API extern bool SupportsNaniteRendering(const FVertexFactory* RESTRICT Ve
 ENGINE_API extern bool SupportsNaniteRendering(const FVertexFactory* RESTRICT VertexFactory, const FPrimitiveSceneProxy* RESTRICT PrimitiveSceneProxy, const class FMaterialRenderProxy* MaterialRenderProxy, ERHIFeatureLevel::Type FeatureLevel);
 
 ENGINE_API extern bool SupportsNaniteRendering(const class FVertexFactoryType* RESTRICT VertexFactoryType, const class FMaterial& Material, ERHIFeatureLevel::Type FeatureLevel);
+
+// Whether scene proxies will have GetDynamicMeshElements called in parallel.
+ENGINE_API bool IsParallelGatherDynamicMeshElementsEnabled();
\ No newline at end of file
diff --git a/Engine/Source/Runtime/Renderer/Private/DeferredShadingRenderer.cpp b/Engine/Source/Runtime/Renderer/Private/DeferredShadingRenderer.cpp
--- a/Engine/Source/Runtime/Renderer/Private/DeferredShadingRenderer.cpp
+++ b/Engine/Source/Runtime/Renderer/Private/DeferredShadingRenderer.cpp
@@ -2337,7 +2337,7 @@ void FDeferredShadingSceneRenderer::WaitForRayTracingScene(FRDGBuilder& GraphBui
 
 #endif // RHI_RAYTRACING
 
-void FDeferredShadingSceneRenderer::BeginInitDynamicShadows(FInitViewTaskDatas& TaskDatas)
+void FDeferredShadingSceneRenderer::BeginInitDynamicShadows(FRDGBuilder& GraphBuilder, FInitViewTaskDatas& TaskDatas, FInstanceCullingManager& InstanceCullingManager)
 {
 	extern int32 GEarlyInitDynamicShadows;
 
@@ -2349,7 +2349,7 @@ void FDeferredShadingSceneRenderer::BeginInitDynamicShadows(FInitViewTaskDatas&
 		&& !HasRayTracedOverlay(ViewFamily)
 		&& TaskDatas.VisibilityTaskData->IsTaskWaitingAllowed())
 	{
-		TaskDatas.DynamicShadows = FSceneRenderer::BeginInitDynamicShadows(true, TaskDatas.VisibilityTaskData);
+		TaskDatas.DynamicShadows = FSceneRenderer::BeginInitDynamicShadows(GraphBuilder, true, TaskDatas.VisibilityTaskData, InstanceCullingManager);
 	}
 }
 
@@ -2360,14 +2360,12 @@ void FDeferredShadingSceneRenderer::FinishInitDynamicShadows(FRDGBuilder& GraphB
 		// Setup dynamic shadows.
 		if (TaskData)
 		{
-			FSceneRenderer::FinishInitDynamicShadows(GraphBuilder, TaskData, InstanceCullingManager);
+			FSceneRenderer::FinishInitDynamicShadows(GraphBuilder, TaskData);
 		}
 		else
 		{
 			TaskData = InitDynamicShadows(GraphBuilder, InstanceCullingManager);
 		}
-
-		DynamicReadBufferForShadows.Commit(GraphBuilder.RHICmdList);
 	}
 }
 
@@ -2981,7 +2979,7 @@ void FDeferredShadingSceneRenderer::Render(FRDGBuilder& GraphBuilder)
 
 	{
 		RDG_GPU_STAT_SCOPE(GraphBuilder, VisibilityCommands);
-		BeginInitViews(GraphBuilder, SceneTexturesConfig, BasePassDepthStencilAccess, InstanceCullingManager, VirtualTextureUpdater.Get(), ExternalAccessQueue, InitViewTaskDatas);
+		BeginInitViews(GraphBuilder, SceneTexturesConfig, InstanceCullingManager, ExternalAccessQueue, InitViewTaskDatas);
 	}
 
 #if !UE_BUILD_SHIPPING
@@ -3076,6 +3074,15 @@ void FDeferredShadingSceneRenderer::Render(FRDGBuilder& GraphBuilder)
 		}
 #endif // RHI_RAYTRACING
 
+#if !(UE_BUILD_SHIPPING || UE_BUILD_TEST)
+		Scene->DebugRender(Views);
+#endif
+	}
+
+	InitViewTaskDatas.VisibilityTaskData->FinishGatherDynamicMeshElements(BasePassDepthStencilAccess, InstanceCullingManager, VirtualTextureUpdater.Get());
+
+	if (RendererOutput == ERendererOutput::FinalSceneColor)
+	{
 		// Notify the FX system that the scene is about to be rendered.
 		if (FXSystem && Views.IsValidIndex(0))
 		{
@@ -3087,10 +3094,6 @@ void FDeferredShadingSceneRenderer::Render(FRDGBuilder& GraphBuilder)
 				GPUSortManager->OnPreRender(GraphBuilder);
 			}
 		}
-
-#if !(UE_BUILD_SHIPPING || UE_BUILD_TEST)
-		Scene->DebugRender(Views);
-#endif
 	}
 
 	{
@@ -3541,7 +3544,7 @@ void FDeferredShadingSceneRenderer::Render(FRDGBuilder& GraphBuilder)
 		{
 			// With forward shading we need to render shadow maps early
 			ensureMsgf(!VirtualShadowMapArray.IsEnabled(), TEXT("Virtual shadow maps are not supported in the forward shading path"));
-			RenderShadowDepthMaps(GraphBuilder, InstanceCullingManager, ExternalAccessQueue);
+			RenderShadowDepthMaps(GraphBuilder, InitViewTaskDatas.DynamicShadows, InstanceCullingManager, ExternalAccessQueue);
 			bShadowMapsRenderedEarly = true;
 
 			if (bHairStrandsEnable)
@@ -3565,7 +3568,7 @@ void FDeferredShadingSceneRenderer::Render(FRDGBuilder& GraphBuilder)
 			ensureMsgf(!VirtualShadowMapArray.IsEnabled(), TEXT("Virtual shadow maps are not supported with r.shadow.ShadowMapsRenderEarly. Early shadows will be disabled"));
 			if (!VirtualShadowMapArray.IsEnabled())
 			{
-				RenderShadowDepthMaps(GraphBuilder, InstanceCullingManager, ExternalAccessQueue);
+				RenderShadowDepthMaps(GraphBuilder, InitViewTaskDatas.DynamicShadows, InstanceCullingManager, ExternalAccessQueue);
 				bShadowMapsRenderedEarly = true;
 			}
 		}
@@ -3797,7 +3800,7 @@ void FDeferredShadingSceneRenderer::Render(FRDGBuilder& GraphBuilder)
 					VirtualShadowMapArray.BuildPageAllocations(GraphBuilder, GetActiveSceneTextures(), Views, SortedLightSet, VisibleLightInfos, SingleLayerWaterPrePassResult, FrontLayerTranslucencyData);
 				}
 
-				RenderShadowDepthMaps(GraphBuilder, InstanceCullingManager, ExternalAccessQueue);
+				RenderShadowDepthMaps(GraphBuilder, InitViewTaskDatas.DynamicShadows, InstanceCullingManager, ExternalAccessQueue);
 			}
 			CheckShadowDepthRenderCompleted();
 
diff --git a/Engine/Source/Runtime/Renderer/Private/DeferredShadingRenderer.h b/Engine/Source/Runtime/Renderer/Private/DeferredShadingRenderer.h
--- a/Engine/Source/Runtime/Renderer/Private/DeferredShadingRenderer.h
+++ b/Engine/Source/Runtime/Renderer/Private/DeferredShadingRenderer.h
@@ -550,7 +550,7 @@ private:
 
 	void PreVisibilityFrameSetup(FRDGBuilder& GraphBuilder);
 
-	void BeginInitDynamicShadows(FInitViewTaskDatas& TaskDatas);
+	void BeginInitDynamicShadows(FRDGBuilder& GraphBuilder, FInitViewTaskDatas& TaskDatas, FInstanceCullingManager& InstanceCullingManager);
 	void FinishInitDynamicShadows(FRDGBuilder& GraphBuilder, FDynamicShadowsTaskData*& TaskData, FInstanceCullingManager& InstanceCullingManager);
 
 	void ComputeLightVisibility();
@@ -559,9 +559,7 @@ private:
 	void BeginInitViews(
 		FRDGBuilder& GraphBuilder,
 		const FSceneTexturesConfig& SceneTexturesConfig,
-		FExclusiveDepthStencil::Type BasePassDepthStencilAccess,
 		FInstanceCullingManager& InstanceCullingManager,
-		FVirtualTextureUpdater* VirtualTextureUpdater,
 		FRDGExternalAccessQueue& ExternalAccessQueue,
 		FInitViewTaskDatas& TaskDatas);
 
diff --git a/Engine/Source/Runtime/Renderer/Private/InstanceCulling/InstanceCullingManager.cpp b/Engine/Source/Runtime/Renderer/Private/InstanceCulling/InstanceCullingManager.cpp
--- a/Engine/Source/Runtime/Renderer/Private/InstanceCulling/InstanceCullingManager.cpp
+++ b/Engine/Source/Runtime/Renderer/Private/InstanceCulling/InstanceCullingManager.cpp
@@ -31,6 +31,14 @@ FInstanceCullingManager::~FInstanceCullingManager()
 {
 }
 
+void FInstanceCullingManager::AllocateViews(int32 NumViews)
+{
+	if (bIsEnabled)
+	{
+		CullingViews.AddUninitialized(CullingViews.Num() + NumViews);
+	}
+}
+
 int32 FInstanceCullingManager::RegisterView(const FViewInfo& ViewInfo)
 {
 	if (!bIsEnabled)
@@ -60,16 +68,21 @@ int32 FInstanceCullingManager::RegisterView(const Nanite::FPackedViewParams& Par
 	{
 		return 0;
 	}
-	CullingViews.Add(CreatePackedView(Params));
-	return CullingViews.Num() - 1;
+
+	const int32 ViewIndex = NumRegisteredViews.fetch_add(1, std::memory_order_relaxed);
+	check(ViewIndex <= CullingViews.Num());
+	CullingViews[ViewIndex] = CreatePackedView(Params);
+	return ViewIndex;
 }
 
 void FInstanceCullingManager::FlushRegisteredViews(FRDGBuilder& GraphBuilder)
 {
-	if (CullingIntermediate.NumViews != CullingViews.Num())
+	const int32 LocalNumRegisteredViews = NumRegisteredViews.load(std::memory_order_relaxed);
+
+	if (CullingIntermediate.NumViews != LocalNumRegisteredViews)
 	{
-		CullingIntermediate.CullingViews = CreateStructuredBuffer(GraphBuilder, TEXT("InstanceCulling.CullingViews"), CullingViews);
-		CullingIntermediate.NumViews = CullingViews.Num();
+		CullingIntermediate.CullingViews = CreateStructuredBuffer(GraphBuilder, TEXT("InstanceCulling.CullingViews"), MakeArrayView<const Nanite::FPackedView>(CullingViews.GetData(), LocalNumRegisteredViews));
+		CullingIntermediate.NumViews = LocalNumRegisteredViews;
 	}
 }
 
@@ -80,6 +93,7 @@ bool FInstanceCullingManager::AllowBatchedBuildRenderingCommands(const FGPUScene
 
 void FInstanceCullingManager::BeginDeferredCulling(FRDGBuilder& GraphBuilder, FGPUScene& GPUScene)
 {
+	TRACE_CPUPROFILER_EVENT_SCOPE(FInstanceCullingManager::BeginDeferredCulling);
 	FlushRegisteredViews(GraphBuilder);
 
 	// Cannot defer pass execution in immediate mode.
@@ -89,7 +103,7 @@ void FInstanceCullingManager::BeginDeferredCulling(FRDGBuilder& GraphBuilder, FG
 	}
 
 	// If there are no instances, there can be no work to perform later.
-	if (GPUScene.GetNumInstances() == 0 || CullingViews.Num() == 0)
+	if (GPUScene.GetNumInstances() == 0 || NumRegisteredViews.load(std::memory_order_relaxed) == 0)
 	{
 		return;
 	}
diff --git a/Engine/Source/Runtime/Renderer/Private/InstanceCulling/InstanceCullingManager.h b/Engine/Source/Runtime/Renderer/Private/InstanceCulling/InstanceCullingManager.h
--- a/Engine/Source/Runtime/Renderer/Private/InstanceCulling/InstanceCullingManager.h
+++ b/Engine/Source/Runtime/Renderer/Private/InstanceCulling/InstanceCullingManager.h
@@ -5,6 +5,7 @@
 #include "InstanceCulling/InstanceCullingContext.h"
 #include "InstanceCullingLoadBalancer.h"
 #include "RenderGraphResources.h"
+#include "Async/Mutex.h"
 
 class FGPUScene;
 class FSceneUniformBuffer;
@@ -65,6 +66,9 @@ public:
 	// Helper to translate from view info, extracts the needed data for setting up instance culling.
 	int32 RegisterView(const FViewInfo& ViewInfo);
 
+	// Allocate space for views ahead of time prior to calling RegisterView.
+	void AllocateViews(int32 NumViews);
+
 	/**
 	 * Upload all registered views if the number has changed (grown). This must be done before calls to FInstanceCullingContext::BuildRenderingCommands that references
 	 * the views.
@@ -93,9 +97,6 @@ public:
 	// Populated by CullInstances, used when performing final culling & rendering 
 	FInstanceCullingIntermediate CullingIntermediate;
 
-	const TArray<Nanite::FPackedView>& GetCullingViews() { return CullingViews; }
-
-
 	// Reference to a buffer owned by FInstanceCullingOcclusionQueryRenderer
 	FRDGBufferRef InstanceOcclusionQueryBuffer = {};
 	EPixelFormat InstanceOcclusionQueryBufferFormat = PF_Unknown;
@@ -111,6 +112,7 @@ private:
 	FInstanceCullingManager(FInstanceCullingManager &) = delete;
 
 	FSceneUniformBuffer& SceneUB;
+	std::atomic_int32_t NumRegisteredViews = {0};
 	TArray<Nanite::FPackedView> CullingViews;
 	bool bIsEnabled;
 };
diff --git a/Engine/Source/Runtime/Renderer/Private/MobileShadingRenderer.cpp b/Engine/Source/Runtime/Renderer/Private/MobileShadingRenderer.cpp
--- a/Engine/Source/Runtime/Renderer/Private/MobileShadingRenderer.cpp
+++ b/Engine/Source/Runtime/Renderer/Private/MobileShadingRenderer.cpp
@@ -436,6 +436,7 @@ void FMobileSceneRenderer::InitViews(
 	}
 
 	// Create GPU-side representation of the view for instance culling.
+	InstanceCullingManager.AllocateViews(Views.Num());
 	for (int32 ViewIndex = 0; ViewIndex < Views.Num(); ++ViewIndex)
 	{
 		Views[ViewIndex].GPUSceneViewId = InstanceCullingManager.RegisterView(Views[ViewIndex]);
@@ -453,7 +454,8 @@ void FMobileSceneRenderer::InitViews(
 		FXSystem->PostInitViews(GraphBuilder, GetSceneViews(), !ViewFamily.EngineShowFlags.HitProxies);
 	}
 
-	TaskDatas.VisibilityTaskData->ProcessRenderThreadTasks(BasePassDepthStencilAccess, InstanceCullingManager, VirtualTextureUpdater);
+	TaskDatas.VisibilityTaskData->ProcessRenderThreadTasks();
+	TaskDatas.VisibilityTaskData->FinishGatherDynamicMeshElements(BasePassDepthStencilAccess, InstanceCullingManager, VirtualTextureUpdater);
 
 	if (ShouldRenderVolumetricFog() && bRendererOutputFinalSceneColor)
 	{
@@ -938,9 +940,6 @@ void FMobileSceneRenderer::Render(FRDGBuilder& GraphBuilder)
 
 	GEngine->GetPreRenderDelegateEx().Broadcast(GraphBuilder);
 
-	// Global dynamic buffers need to be committed before rendering.
-	DynamicReadBufferForInitViews.Commit(GraphBuilder.RHICmdList);
-	DynamicReadBufferForShadows.Commit(GraphBuilder.RHICmdList);
 	
 	GraphBuilder.SetCommandListStat(GET_STATID(STAT_CLMM_SceneSim));
 
@@ -1006,7 +1005,7 @@ void FMobileSceneRenderer::Render(FRDGBuilder& GraphBuilder)
 		}
 
 		GraphBuilder.SetCommandListStat(GET_STATID(STAT_CLMM_Shadows));
-		RenderShadowDepthMaps(GraphBuilder, InstanceCullingManager, ExternalAccessQueue);
+		RenderShadowDepthMaps(GraphBuilder, nullptr, InstanceCullingManager, ExternalAccessQueue);
 
 		// Run local fog volume initialization before base pass and volumetric fog for all the culled instance instance data to be ready.
 		InitLocalFogVolumesForViews(Scene, Views, ViewFamily, GraphBuilder, ShouldRenderVolumetricFog());
diff --git a/Engine/Source/Runtime/Renderer/Private/SceneHitProxyRendering.cpp b/Engine/Source/Runtime/Renderer/Private/SceneHitProxyRendering.cpp
--- a/Engine/Source/Runtime/Renderer/Private/SceneHitProxyRendering.cpp
+++ b/Engine/Source/Runtime/Renderer/Private/SceneHitProxyRendering.cpp
@@ -578,10 +578,6 @@ void FMobileSceneRenderer::RenderHitProxies(FRDGBuilder& GraphBuilder)
 
 	GEngine->GetPreRenderDelegateEx().Broadcast(GraphBuilder);
 
-	// Global dynamic buffers need to be committed before rendering.
-	DynamicReadBufferForInitViews.Commit(GraphBuilder.RHICmdList);
-	DynamicReadBufferForShadows.Commit(GraphBuilder.RHICmdList);
-
 	InstanceCullingManager.FlushRegisteredViews(GraphBuilder);
 
 	TArray<Nanite::FRasterResults, TInlineAllocator<2>> NaniteRasterResults;
@@ -621,7 +617,7 @@ void FDeferredShadingSceneRenderer::RenderHitProxies(FRDGBuilder& GraphBuilder)
 	FLumenSceneFrameTemporaries LumenFrameTemporaries;
 	FInitViewTaskDatas InitViewTaskDatas(VisibilityTaskData);
 	FRDGExternalAccessQueue ExternalAccessQueue;
-	BeginInitViews(GraphBuilder, SceneTexturesConfig, FExclusiveDepthStencil::DepthWrite_StencilWrite, InstanceCullingManager, nullptr, ExternalAccessQueue, InitViewTaskDatas);
+	BeginInitViews(GraphBuilder, SceneTexturesConfig, InstanceCullingManager, ExternalAccessQueue, InitViewTaskDatas);
 
 	extern TSet<IPersistentViewUniformBufferExtension*> PersistentViewUniformBufferExtensions;
 
@@ -638,6 +634,8 @@ void FDeferredShadingSceneRenderer::RenderHitProxies(FRDGBuilder& GraphBuilder)
 
 	ShaderPrint::BeginViews(GraphBuilder, Views);
 
+	InitViewTaskDatas.VisibilityTaskData->FinishGatherDynamicMeshElements(FExclusiveDepthStencil::DepthWrite_StencilWrite, InstanceCullingManager, nullptr);
+
 	for (int32 ViewIndex = 0; ViewIndex < Views.Num(); ViewIndex++)
 	{
 		Scene->GPUScene.UploadDynamicPrimitiveShaderDataForView(GraphBuilder, *Scene, Views[ViewIndex]);
@@ -663,10 +661,6 @@ void FDeferredShadingSceneRenderer::RenderHitProxies(FRDGBuilder& GraphBuilder)
 		Scene->SplineMeshSceneResources->Update(GraphBuilder, GetSceneUniforms());
 	}
 
-	// Global dynamic buffers need to be committed before rendering.
-	DynamicReadBufferForInitViews.Commit(GraphBuilder.RHICmdList);
-	DynamicReadBufferForShadows.Commit(GraphBuilder.RHICmdList);
-
 	// Notify the FX system that the scene is about to be rendered.
 	if (FXSystem && Views.IsValidIndex(0))
 	{
diff --git a/Engine/Source/Runtime/Renderer/Private/SceneRendering.cpp b/Engine/Source/Runtime/Renderer/Private/SceneRendering.cpp
--- a/Engine/Source/Runtime/Renderer/Private/SceneRendering.cpp
+++ b/Engine/Source/Runtime/Renderer/Private/SceneRendering.cpp
@@ -2504,8 +2504,6 @@ TGlobalResource<FGlobalDynamicReadBuffer> FSceneRenderer::DynamicReadBufferForSh
 FSceneRenderer::FSceneRenderer(const FSceneViewFamily* InViewFamily, FHitProxyConsumer* HitProxyConsumer)
 :	Scene(CheckPointer(InViewFamily->Scene)->GetRenderScene())
 ,	ViewFamily(*CheckPointer(InViewFamily))
-,	MeshCollector(InViewFamily->GetFeatureLevel(), Allocator, FMeshElementCollector::ECommitFlags::DeferAll)
-,	EditorMeshCollector(InViewFamily->GetFeatureLevel(), Allocator, FMeshElementCollector::ECommitFlags::DeferAll)
 ,	VirtualShadowMapArray(*CheckPointer(Scene))
 ,	bHasRequestedToggleFreeze(false)
 ,	bUsedPrecomputedVisibility(false)
diff --git a/Engine/Source/Runtime/Renderer/Private/SceneRendering.h b/Engine/Source/Runtime/Renderer/Private/SceneRendering.h
--- a/Engine/Source/Runtime/Renderer/Private/SceneRendering.h
+++ b/Engine/Source/Runtime/Renderer/Private/SceneRendering.h
@@ -2036,13 +2036,11 @@ public:
 	DynamicRenderScaling::TMap<float> DynamicResolutionFractions;
 	DynamicRenderScaling::TMap<float> DynamicResolutionUpperBounds;
 
-	FMeshElementCollector MeshCollector;
-	FMeshElementCollector EditorMeshCollector;
-
 	/** Information about the visible lights. */
 	TArray<FVisibleLightInfo, SceneRenderingAllocator> VisibleLightInfos;
 
 	/** Array of dispatched parallel shadow depth passes. */
+	UE::FMutex DispatchedShadowDepthPassesMutex;
 	TArray<FParallelMeshDrawCommandPass*, SceneRenderingAllocator> DispatchedShadowDepthPasses;
 
 	FSortedShadowMaps SortedShadowsForShadowDepthPass;
@@ -2304,8 +2302,9 @@ protected:
 
 	ERendererOutput GetRendererOutput() const;
 
-	FDynamicShadowsTaskData* BeginInitDynamicShadows( bool bRunningEarly, IVisibilityTaskData* VisibilityTaskData);
-	void FinishInitDynamicShadows(FRDGBuilder& GraphBuilder, FDynamicShadowsTaskData* TaskData, FInstanceCullingManager& InstanceCullingManager);
+	FDynamicShadowsTaskData* BeginInitDynamicShadows(FRDGBuilder& GraphBuilder, bool bRunningEarly, IVisibilityTaskData* VisibilityTaskData, FInstanceCullingManager& InstanceCullingManager);
+	void FinishInitDynamicShadows(FRDGBuilder& GraphBuilder, FDynamicShadowsTaskData* TaskData);
+	void FinishDynamicShadowMeshPassSetup(FRDGBuilder& GraphBuilder, FDynamicShadowsTaskData* TaskData);
 	FDynamicShadowsTaskData* InitDynamicShadows(FRDGBuilder& GraphBuilder, FInstanceCullingManager& InstanceCullingManager);
 
 	void CreateDynamicShadows(FDynamicShadowsTaskData& TaskData);
@@ -2371,7 +2370,7 @@ protected:
 		int64 CachedShadowMapsSize,
 		uint32& NumCSMCachesUpdatedThisFrame);
 
-	void AllocateShadowDepthTargets(FRHICommandListImmediate& RHICmdList, const FDynamicShadowsTaskData& TaskData);
+	void AllocateShadowDepthTargets(FDynamicShadowsTaskData& TaskData);
 	
 	void AllocateAtlasedShadowDepthTargets(FRHICommandListBase& RHICmdList, TConstArrayView<FProjectedShadowInfo*> Shadows, TArray<FSortedShadowMapAtlas,SceneRenderingAllocator>& OutAtlases);
 
@@ -2397,7 +2396,7 @@ protected:
 	void BeginGatherShadowPrimitives(FDynamicShadowsTaskData* TaskData, IVisibilityTaskData* VisibilityTaskData);
 	void FinishGatherShadowPrimitives(FDynamicShadowsTaskData* TaskData);
 
-	void RenderShadowDepthMaps(FRDGBuilder& GraphBuilder, FInstanceCullingManager& InstanceCullingManager, FRDGExternalAccessQueue& ExternalAccessQueue);
+	void RenderShadowDepthMaps(FRDGBuilder& GraphBuilder, FDynamicShadowsTaskData* DynamicShadowsTaskData, FInstanceCullingManager& InstanceCullingManager, FRDGExternalAccessQueue& ExternalAccessQueue);
 	void RenderVirtualShadowMaps(FRDGBuilder& GraphBuilder, bool bNaniteEnabled);
 
 	void RenderShadowDepthMapAtlases(FRDGBuilder& GraphBuilder);
@@ -2421,7 +2420,7 @@ protected:
 	void DrawDebugShadowFrustum(FViewInfo& View, FProjectedShadowInfo& ProjectedShadowInfo);
 
 	/** Gathers dynamic mesh elements for all shadows. */
-	void GatherShadowDynamicMeshElements(FRHICommandList& RHICmdList, FInstanceCullingManager& InstanceCullingManager);
+	void GatherShadowDynamicMeshElements(FDynamicShadowsTaskData& TaskData);
 
 	/** Performs once per frame setup prior to visibility determination. */
 	void PreVisibilityFrameSetup(FRDGBuilder& GraphBuilder);
@@ -2828,6 +2827,9 @@ extern FFastVramConfig GFastVRamConfig;
 /** Returns the array of shadows with distance fields. Call only after finishing shadow initialization. */
 extern TConstArrayView<FProjectedShadowInfo*> GetProjectedDistanceFieldShadows(const FDynamicShadowsTaskData* TaskData);
 
+/** Triggers shadow gather dynamic mesh elements task graph to start processing. */
+extern void BeginShadowGatherDynamicMeshElements(FDynamicShadowsTaskData* TaskData);
+
 extern bool UseCachedMeshDrawCommands();
 extern bool UseCachedMeshDrawCommands_AnyThread();
 extern bool IsDynamicInstancingEnabled(ERHIFeatureLevel::Type FeatureLevel);
diff --git a/Engine/Source/Runtime/Renderer/Private/SceneVisibility.cpp b/Engine/Source/Runtime/Renderer/Private/SceneVisibility.cpp
--- a/Engine/Source/Runtime/Renderer/Private/SceneVisibility.cpp
+++ b/Engine/Source/Runtime/Renderer/Private/SceneVisibility.cpp
@@ -299,6 +299,24 @@ static FAutoConsoleVariableRef CVarOcclusionCullMaxQueriesPerTask(
 	ECVF_RenderThreadSafe
 );
 
+static int32 GNumDynamicMeshElementTasks = 4;
+static FAutoConsoleVariableRef CVarNumDynamicMeshElementTasks(
+	TEXT("r.Visibility.DynamicMeshElements.NumMainViewTasks"),
+	GNumDynamicMeshElementTasks,
+	TEXT("Controls the number of gather dynamic mesh elements tasks to run asynchronously during view visibility."),
+	ECVF_RenderThreadSafe
+);
+
+inline uint32 GetNumDynamicMeshElementTasks()
+{
+	if (!IsParallelGatherDynamicMeshElementsEnabled())
+	{
+		return 0;
+	}
+
+	return FMath::Clamp<int32>(GNumDynamicMeshElementTasks, 0, LowLevelTasks::FScheduler::Get().GetNumWorkers());
+}
+
 static bool GOcclusionCullEnabled = true;
 static FAutoConsoleVariableRef CVarOcclusionCullEnable(
 	// TODO: Move to r.Visibility.OcclusionCull.Enable. Still several explicit references.
@@ -1262,7 +1280,7 @@ void FRelevancePacket::ComputeRelevance(FDynamicPrimitiveIndexList& DynamicPrimi
 			}
 			else
 			{
-				DynamicPrimitiveIndexList.EditorPrimitives.Add(PrimitiveIndex);
+				DynamicPrimitiveIndexList.EditorPrimitives.Emplace(PrimitiveIndex, ViewBit);
 			}
 		}
 #endif
@@ -1278,7 +1296,7 @@ void FRelevancePacket::ComputeRelevance(FDynamicPrimitiveIndexList& DynamicPrimi
 		}
 		else
 		{
-			DynamicPrimitiveIndexList.Primitives.Add(PrimitiveIndex);
+			DynamicPrimitiveIndexList.Primitives.Emplace(PrimitiveIndex, ViewBit);
 		}
 	};
 
@@ -3512,6 +3530,347 @@ void FVisibilityViewPacket::BeginInitVisibility()
 
 ///////////////////////////////////////////////////////////////////////////////
 
+FDynamicMeshElementContext::FDynamicMeshElementContext(FSceneRenderer& SceneRenderer)
+	: ViewFamily(SceneRenderer.ViewFamily)
+	, Views(SceneRenderer.AllViews)
+	, Primitives(SceneRenderer.Scene->Primitives)
+	// Defer committing GPU scene and Material updates until the mesh collectors are finished. Deferring materials allows VT updates to
+	// overlap with GDME (uniform expression updates can't happen at the same time as the async VT system update), and since we are going
+	// wide across a single view we would otherwise have to lock for GPU scene updates.
+	, MeshCollector(SceneRenderer.FeatureLevel, SceneRenderer.Allocator, FMeshElementCollector::ECommitFlags::DeferAll)
+#if WITH_EDITOR
+	, EditorMeshCollector(SceneRenderer.FeatureLevel, SceneRenderer.Allocator, FMeshElementCollector::ECommitFlags::DeferAll)
+#endif
+	, RHICmdList(new FRHICommandList(FRHIGPUMask::All()))
+	, DynamicVertexBuffer(*RHICmdList)
+	, DynamicIndexBuffer(*RHICmdList)
+{
+	RHICmdList->SwitchPipeline(ERHIPipeline::Graphics);
+
+	ViewMeshArraysPerView.SetNum(Views.Num());
+
+	MeshCollector.Start(
+		*RHICmdList,
+		DynamicVertexBuffer,
+		DynamicIndexBuffer,
+		SceneRenderer.DynamicReadBufferForInitViews
+	);
+
+	for (int32 ViewIndex = 0; ViewIndex < Views.Num(); ViewIndex++)
+	{
+		FViewInfo& View = *Views[ViewIndex];
+		FViewMeshArrays& ViewMeshArrays = ViewMeshArraysPerView[ViewIndex];
+
+		MeshCollector.AddViewMeshArrays(
+			&View,
+			&ViewMeshArrays.DynamicMeshElements,
+			&ViewMeshArrays.SimpleElementCollector,
+			&View.DynamicPrimitiveCollector
+#if UE_ENABLE_DEBUG_DRAWING
+			, &ViewMeshArrays.DebugSimpleElementCollector
+#endif
+		);
+	}
+
+#if WITH_EDITOR
+	if (GIsEditor)
+	{
+		EditorMeshCollector.Start(
+			*RHICmdList,
+			DynamicVertexBuffer,
+			DynamicIndexBuffer,
+			SceneRenderer.DynamicReadBufferForInitViews
+		);
+
+		for (int32 ViewIndex = 0; ViewIndex < Views.Num(); ViewIndex++)
+		{
+			FViewInfo& View = *Views[ViewIndex];
+			FViewMeshArrays& ViewMeshArrays = ViewMeshArraysPerView[ViewIndex];
+
+			EditorMeshCollector.AddViewMeshArrays(
+				&View,
+				&ViewMeshArrays.DynamicEditorMeshElements,
+				&ViewMeshArrays.EditorSimpleElementCollector,
+				&View.DynamicPrimitiveCollector
+#if UE_ENABLE_DEBUG_DRAWING
+				, &ViewMeshArrays.DebugSimpleElementCollector
+#endif
+			);
+		}
+	}
+#endif
+}
+
+FGraphEventRef FDynamicMeshElementContext::LaunchRenderThreadTask(FDynamicPrimitiveIndexList&& PrimitiveIndexList)
+{
+	return FFunctionGraphTask::CreateAndDispatchWhenReady([this, PrimitiveIndexList = MoveTemp(PrimitiveIndexList)]
+	{
+		for (FDynamicPrimitiveIndex PrimitiveIndex : PrimitiveIndexList.Primitives)
+		{
+			GatherDynamicMeshElementsForPrimitive(Primitives[PrimitiveIndex.Index], PrimitiveIndex.ViewMask);
+		}
+
+#if WITH_EDITOR
+		for (FDynamicPrimitiveIndex PrimitiveIndex : PrimitiveIndexList.EditorPrimitives)
+		{
+			GatherDynamicMeshElementsForEditorPrimitive(Primitives[PrimitiveIndex.Index], PrimitiveIndex.ViewMask);
+		}
+#endif
+	}, TStatId{}, nullptr, ENamedThreads::GetRenderThread_Local());
+}
+
+UE::Tasks::FTask FDynamicMeshElementContext::LaunchAsyncTask(FDynamicPrimitiveIndexQueue* PrimitiveIndexQueue, UE::Tasks::ETaskPriority TaskPriority)
+{
+	return Pipe.Launch(UE_SOURCE_LOCATION, [this, PrimitiveIndexQueue]
+	{
+		FOptionalTaskTagScope Scope(ETaskTag::EParallelRenderingThread);
+		FDynamicPrimitiveIndex PrimitiveIndex;
+
+		while (PrimitiveIndexQueue->Pop(PrimitiveIndex))
+		{
+			GatherDynamicMeshElementsForPrimitive(Primitives[PrimitiveIndex.Index], PrimitiveIndex.ViewMask);
+		}
+
+#if WITH_EDITOR
+		while (PrimitiveIndexQueue->PopEditor(PrimitiveIndex))
+		{
+			GatherDynamicMeshElementsForEditorPrimitive(Primitives[PrimitiveIndex.Index], PrimitiveIndex.ViewMask);
+		}
+#endif
+	}, TaskPriority);
+}
+
+void FDynamicMeshElementContext::GatherDynamicMeshElementsForPrimitive(FPrimitiveSceneInfo* Primitive, uint8 ViewMask)
+{
+	SCOPED_NAMED_EVENT(DynamicPrimitive, FColor::Magenta);
+
+	TArray<int32, TInlineAllocator<4>> MeshBatchCountBefore;
+	MeshBatchCountBefore.SetNumUninitialized(Views.Num());
+	for (int32 ViewIndex = 0; ViewIndex < Views.Num(); ViewIndex++)
+	{
+		MeshBatchCountBefore[ViewIndex] = MeshCollector.GetMeshBatchCount(ViewIndex);
+	}
+
+	MeshCollector.SetPrimitive(Primitive->Proxy, Primitive->DefaultDynamicHitProxyId);
+	Primitive->Proxy->GetDynamicMeshElements(ViewFamily.AllViews, ViewFamily, ViewMask, MeshCollector);
+
+	for (int32 ViewIndex = 0; ViewIndex < Views.Num(); ViewIndex++)
+	{
+		FViewInfo& View = *Views[ViewIndex];
+
+		if (ViewMask & (1 << ViewIndex))
+		{
+			FDynamicPrimitive& DynamicPrimitive = DynamicPrimitives.Emplace_GetRef();
+			DynamicPrimitive.PrimitiveIndex = Primitive->GetIndex();
+			DynamicPrimitive.ViewIndex = ViewIndex;
+			DynamicPrimitive.StartElementIndex = MeshBatchCountBefore[ViewIndex];
+			DynamicPrimitive.EndElementIndex = MeshCollector.GetMeshBatchCount(ViewIndex);
+		}
+	}
+}
+
+void FDynamicMeshElementContext::GatherDynamicMeshElementsForEditorPrimitive(FPrimitiveSceneInfo* Primitive, uint8 ViewMask)
+{
+#if WITH_EDITOR
+	EditorMeshCollector.SetPrimitive(Primitive->Proxy, Primitive->DefaultDynamicHitProxyId);
+	Primitive->Proxy->GetDynamicMeshElements(ViewFamily.AllViews, ViewFamily, ViewMask, EditorMeshCollector);
+#endif
+}
+
+void FDynamicMeshElementContext::Finish()
+{
+	MeshCollector.Finish();
+
+#if WITH_EDITOR
+	if (GIsEditor)
+	{
+		EditorMeshCollector.Finish();
+	}
+#endif
+
+	DynamicVertexBuffer.Commit();
+	DynamicIndexBuffer.Commit();
+	RHICmdList->FinishRecording();
+}
+
+FDynamicMeshElementContextContainer::~FDynamicMeshElementContextContainer()
+{
+	check(CommandLists.IsEmpty());
+}
+
+UE::Tasks::FTask FDynamicMeshElementContextContainer::LaunchAsyncTask(FDynamicPrimitiveIndexQueue* PrimitiveIndexQueue, int32 Index, UE::Tasks::ETaskPriority TaskPriority)
+{
+	return Contexts[Index]->LaunchAsyncTask(PrimitiveIndexQueue, TaskPriority);
+}
+
+FGraphEventRef FDynamicMeshElementContextContainer::LaunchRenderThreadTask(FDynamicPrimitiveIndexList&& PrimitiveIndexList)
+{
+	return Contexts.Last()->LaunchRenderThreadTask(MoveTemp(PrimitiveIndexList));
+}
+
+void FDynamicMeshElementContextContainer::Init(FSceneRenderer& SceneRenderer, int32 NumAsyncContexts)
+{
+	const int32 NumRenderThreadContexts = 1;
+	const int32 NumContexts = NumAsyncContexts + NumRenderThreadContexts;
+	Views = SceneRenderer.AllViews;
+	Contexts.Reserve(NumContexts);
+	CommandLists.Reserve(Contexts.Num());
+
+	for (int32 Index = 0; Index < NumContexts; ++Index)
+	{
+		FDynamicMeshElementContext* Context = SceneRenderer.Allocator.Create<FDynamicMeshElementContext>(SceneRenderer);
+		Contexts.Emplace(Context);
+		CommandLists.Emplace(Context->RHICmdList);
+	}
+}
+
+void FDynamicMeshElementContextContainer::MergeContexts(TArray<FDynamicPrimitive, SceneRenderingAllocator>& OutDynamicPrimitives)
+{
+	SCOPED_NAMED_EVENT(MergeGatherDynamicMeshElementContexts, FColor::Magenta);
+
+	check(!Views.IsEmpty());
+
+	// Fast path for one context; just move the memory instead of copying.
+	if (Contexts.Num() == 1)
+	{
+		FDynamicMeshElementContext* Context = Contexts[0];
+		OutDynamicPrimitives = MoveTemp(Context->DynamicPrimitives);
+
+		for (int32 ViewIndex = 0; ViewIndex < Views.Num(); ++ViewIndex)
+		{
+			FViewInfo& View = *Views[ViewIndex];
+			FDynamicMeshElementContext::FViewMeshArrays& ViewMeshArrays = Context->ViewMeshArraysPerView[ViewIndex];
+
+			View.SimpleElementCollector = MoveTemp(ViewMeshArrays.SimpleElementCollector);
+			View.DynamicMeshElements    = MoveTemp(ViewMeshArrays.DynamicMeshElements);
+
+#if WITH_EDITOR
+			View.EditorSimpleElementCollector = MoveTemp(ViewMeshArrays.EditorSimpleElementCollector);
+			View.DynamicEditorMeshElements    = MoveTemp(ViewMeshArrays.DynamicEditorMeshElements);
+#endif
+
+#if UE_ENABLE_DEBUG_DRAWING
+			View.DebugSimpleElementCollector = MoveTemp(ViewMeshArrays.DebugSimpleElementCollector);
+#endif
+		}
+	}
+	// >1 context means we have to merge the containers.
+	else
+	{
+		struct FViewAllocationInfo
+		{
+			FSimpleElementCollector::FAllocationInfo SimpleElementCollector;
+			uint32 NumDynamicMeshElements = 0;
+#if WITH_EDITOR
+			FSimpleElementCollector::FAllocationInfo EditorSimpleElementCollector;
+			uint32 NumDynamicEditorMeshElements = 0;
+#endif
+#if UE_ENABLE_DEBUG_DRAWING
+			FSimpleElementCollector::FAllocationInfo DebugSimpleElementCollector;
+#endif
+		};
+
+		TArray<FViewAllocationInfo, TInlineAllocator<2>> AllocationInfosPerView;
+		AllocationInfosPerView.AddDefaulted(Views.Num());
+		uint32 NumDynamicPrimitives = 0;
+
+		for (FDynamicMeshElementContext* Context : Contexts)
+		{
+			NumDynamicPrimitives += Context->DynamicPrimitives.Num();
+			check(AllocationInfosPerView.Num() == Context->ViewMeshArraysPerView.Num());
+
+			// Accumulate allocation info for each context in order to reserve container memory once.
+			for (int32 ViewIndex = 0; ViewIndex < AllocationInfosPerView.Num(); ++ViewIndex)
+			{
+				const FDynamicMeshElementContext::FViewMeshArrays& ViewMeshArrays = Context->ViewMeshArraysPerView[ViewIndex];
+				FViewAllocationInfo& AllocationInfo = AllocationInfosPerView[ViewIndex];
+
+				ViewMeshArrays.SimpleElementCollector.AddAllocationInfo(AllocationInfo.SimpleElementCollector);
+				AllocationInfo.NumDynamicMeshElements += ViewMeshArrays.DynamicMeshElements.Num();
+#if WITH_EDITOR
+				ViewMeshArrays.EditorSimpleElementCollector.AddAllocationInfo(AllocationInfo.EditorSimpleElementCollector);
+				AllocationInfo.NumDynamicEditorMeshElements += ViewMeshArrays.DynamicEditorMeshElements.Num();
+#endif
+#if UE_ENABLE_DEBUG_DRAWING
+				ViewMeshArrays.DebugSimpleElementCollector.AddAllocationInfo(AllocationInfo.DebugSimpleElementCollector);
+#endif
+			}
+		}
+
+		OutDynamicPrimitives.Reserve(NumDynamicPrimitives);
+
+		// Reserve memory for merged containers.
+		for (int32 ViewIndex = 0; ViewIndex < Views.Num(); ++ViewIndex)
+		{
+			FViewAllocationInfo& AllocationInfo = AllocationInfosPerView[ViewIndex];
+			FViewInfo& View = *Views[ViewIndex];
+
+			View.SimpleElementCollector.Reserve(AllocationInfo.SimpleElementCollector);
+			View.DynamicMeshElements.Reserve(AllocationInfo.NumDynamicMeshElements);
+#if WITH_EDITOR
+			View.EditorSimpleElementCollector.Reserve(AllocationInfo.EditorSimpleElementCollector);
+			View.DynamicEditorMeshElements.Reserve(AllocationInfo.NumDynamicEditorMeshElements);
+#endif
+#if UE_ENABLE_DEBUG_DRAWING
+			View.DebugSimpleElementCollector.Reserve(AllocationInfo.DebugSimpleElementCollector);
+#endif
+
+			// Reset dynamic element count to use as offset for copying ranges in the next loop.
+			AllocationInfo.NumDynamicMeshElements = 0;
+		}
+
+		for (FDynamicMeshElementContext* Context : Contexts)
+		{
+			for (FDynamicPrimitive DynamicPrimitive : Context->DynamicPrimitives)
+			{
+				const uint32 NumDynamicMeshElements = AllocationInfosPerView[DynamicPrimitive.ViewIndex].NumDynamicMeshElements;
+
+				// Offset the dynamic element range by the current number of meshes in the final container.
+				DynamicPrimitive.StartElementIndex += NumDynamicMeshElements;
+				DynamicPrimitive.EndElementIndex   += NumDynamicMeshElements;
+
+				OutDynamicPrimitives.Emplace(DynamicPrimitive);
+
+				Views[DynamicPrimitive.ViewIndex]->DynamicMeshElementRanges[DynamicPrimitive.PrimitiveIndex] = FInt32Vector2(DynamicPrimitive.StartElementIndex, DynamicPrimitive.EndElementIndex);
+			}
+
+			for (int32 ViewIndex = 0; ViewIndex < Views.Num(); ++ViewIndex)
+			{
+				FViewInfo& View = *Views[ViewIndex];
+				FDynamicMeshElementContext::FViewMeshArrays& ViewMeshArrays = Context->ViewMeshArraysPerView[ViewIndex];
+
+				AllocationInfosPerView[ViewIndex].NumDynamicMeshElements += ViewMeshArrays.DynamicMeshElements.Num();
+
+				View.SimpleElementCollector.Append(ViewMeshArrays.SimpleElementCollector);
+				View.DynamicMeshElements.Append(ViewMeshArrays.DynamicMeshElements);
+#if WITH_EDITOR
+				View.EditorSimpleElementCollector.Append(ViewMeshArrays.EditorSimpleElementCollector);
+				View.DynamicEditorMeshElements.Append(ViewMeshArrays.DynamicEditorMeshElements);
+#endif
+#if UE_ENABLE_DEBUG_DRAWING
+				View.DebugSimpleElementCollector.Append(ViewMeshArrays.DebugSimpleElementCollector);
+#endif
+			}
+
+			Context->DynamicPrimitives.Empty();
+			Context->ViewMeshArraysPerView.Empty();
+		}
+	}
+}
+
+void FDynamicMeshElementContextContainer::Submit(FRHICommandListImmediate& RHICmdList)
+{
+	for (FDynamicMeshElementContext* Context : Contexts)
+	{
+		Context->Finish();
+	}
+
+	RHICmdList.QueueAsyncCommandListSubmit(CommandLists);
+	CommandLists.Empty();
+}
+
+///////////////////////////////////////////////////////////////////////////////
+
 FVisibilityTaskData::FVisibilityTaskData(FRHICommandListImmediate& InRHICmdList, FSceneRenderer& InSceneRenderer)
 	: RHICmdList(InRHICmdList)
 	, SceneRenderer(InSceneRenderer)
@@ -3519,8 +3878,6 @@ FVisibilityTaskData::FVisibilityTaskData(FRHICommandListImmediate& InRHICmdList,
 	, Views(SceneRenderer.AllViews)
 	, ViewFamily(SceneRenderer.ViewFamily)
 	, ShadingPath(Scene.GetShadingPath())
-	, MeshCollector(SceneRenderer.MeshCollector)
-	, EditorMeshCollector(SceneRenderer.EditorMeshCollector)
 	, TaskConfig(Scene, Views)
 	, bAddLightmapDensityCommands(ViewFamily.EngineShowFlags.LightMapDensity&& AllowDebugViewmodes())
 {
@@ -3544,6 +3901,12 @@ void FVisibilityTaskData::LaunchVisibilityTasks()
 		Tasks.ComputeRelevance.AddPrerequisites(ViewPacket.Tasks.ComputeRelevance);
 	}
 
+	// Each relevance task should have this as a prerequisite, but in case there aren't any tasks we make it explicit.
+	Tasks.ComputeRelevance.AddPrerequisites(Scene.GetCacheMeshDrawCommandsTask());
+
+	// Wait on the GPU skin cache task prior to GDME.
+	Tasks.DynamicMeshElementsPrerequisites.AddPrerequisites(Scene.GetGPUSkinCacheTask());
+
 	bool bAllocatePrimitiveViewMasks = true;
 
 	if (TaskConfig.Schedule == EVisibilityTaskSchedule::Parallel)
@@ -3564,14 +3927,11 @@ void FVisibilityTaskData::LaunchVisibilityTasks()
 				PreSyncCullingTasks.Emplace(ViewPacket.Tasks.OcclusionCull);
 			}
 
-			FGraphEventRef MergeSecondaryViewsTask = FGraphEvent::CreateGraphEvent();
-
-			UE::Tasks::Launch(UE_SOURCE_LOCATION, [this, MergeSecondaryViewsTask]
-				{
-					MergeSecondaryViewVisibility();
-					MergeSecondaryViewsTask->DispatchSubsequents();
+			UE::Tasks::FTask MergeSecondaryViewsTask = UE::Tasks::Launch(UE_SOURCE_LOCATION, [this]
+			{
+				MergeSecondaryViewVisibility();
 
-				}, PreSyncCullingTasks, TaskConfig.OcclusionCull.TaskPriority);
+			}, PreSyncCullingTasks, TaskConfig.OcclusionCull.TaskPriority);
 
 			for (FVisibilityViewPacket& ViewPacket : ViewPackets)
 			{
@@ -3586,16 +3946,22 @@ void FVisibilityTaskData::LaunchVisibilityTasks()
 		else if (Views.Num() == 1)
 		{
 			// When using a single view, dynamic mesh elements are pushed into a pipe that is executed on the render thread which allows for some overlap with compute relevance work.
-			DynamicMeshElements.CommandPipe = Allocator.Create<TCommandPipe<FDynamicPrimitiveIndexList>>(TEXT("GatherDynamicMeshElements"), ENamedThreads::GetRenderThread_Local());
+			DynamicMeshElements.CommandPipe = Allocator.Create<TCommandPipe<FDynamicPrimitiveIndexList>>(TEXT("GatherDynamicMeshElements"));
 
-			DynamicMeshElements.CommandPipe->SetCommandFunction([this](const FDynamicPrimitiveIndexList& DynamicPrimitiveIndexList)
-				{
-					GatherDynamicMeshElements(DynamicPrimitiveIndexList);
-				});
+			DynamicMeshElements.CommandPipe->SetCommandFunction([this](FDynamicPrimitiveIndexList&& DynamicPrimitiveIndexList)
+			{
+				GatherDynamicMeshElements(MoveTemp(DynamicPrimitiveIndexList));
+			});
 
-			// Fire an event when the pipe has completed all work.
-			DynamicMeshElements.CommandPipeCompleteEvent = FGraphEvent::CreateGraphEvent();
-			DynamicMeshElements.CommandPipe->SetEmptyTaskEvent(DynamicMeshElements.CommandPipeCompleteEvent);
+			DynamicMeshElements.CommandPipe->SetPrerequisiteTask(Tasks.DynamicMeshElementsPrerequisites);
+
+			Tasks.DynamicMeshElementsPipe = FGraphEvent::CreateGraphEvent();
+
+			DynamicMeshElements.CommandPipe->SetEmptyFunction([this]
+			{
+				Tasks.DynamicMeshElementsPipe->DispatchSubsequents();
+				Tasks.DynamicMeshElements.Trigger();
+			});
 
 			// Take a reference that is released when the relevance pipe has completed. We only need to take one since there can only be one view.
 			DynamicMeshElements.CommandPipe->AddNumCommands(1);
@@ -3618,56 +3984,10 @@ void FVisibilityTaskData::LaunchVisibilityTasks()
 	#endif
 	}
 
-	// Reserve a full page for dynamic primitive indices as a starting point.
-	DynamicMeshElements.DynamicPrimitives.Reserve(FSceneRenderingBlockAllocationTag::BlockSize / sizeof(FDynamicPrimitive));
-	DynamicMeshElements.ViewCommandsPerView.SetNum(Views.Num());
-	DynamicMeshElements.LastElementPerView.SetNum(Views.Num());
-
-	DynamicMeshElements.DynamicVertexBuffer.Init(RHICmdList);
-	DynamicMeshElements.DynamicIndexBuffer.Init(RHICmdList);
-
-	MeshCollector.Start(
-		RHICmdList,
-		DynamicMeshElements.DynamicVertexBuffer,
-		DynamicMeshElements.DynamicIndexBuffer,
-		SceneRenderer.DynamicReadBufferForInitViews
-	);
-
-	for (int32 ViewIndex = 0; ViewIndex < Views.Num(); ViewIndex++)
-	{
-		MeshCollector.AddViewMeshArrays(
-			Views[ViewIndex],
-			&Views[ViewIndex]->DynamicMeshElements,
-			&Views[ViewIndex]->SimpleElementCollector,
-			&Views[ViewIndex]->DynamicPrimitiveCollector
-#if UE_ENABLE_DEBUG_DRAWING
-			, &Views[ViewIndex]->DebugSimpleElementCollector
-#endif
-		);
-	}
-
-	if (GIsEditor)
-	{
-		EditorMeshCollector.Start(
-			RHICmdList,
-			DynamicMeshElements.DynamicVertexBuffer,
-			DynamicMeshElements.DynamicIndexBuffer,
-			SceneRenderer.DynamicReadBufferForInitViews
-		);
+	const int32 NumAsyncDynamicMeshElementContexts = TaskConfig.Schedule == EVisibilityTaskSchedule::Parallel ? GetNumDynamicMeshElementTasks() : 0;
 
-		for (int32 ViewIndex = 0; ViewIndex < Views.Num(); ViewIndex++)
-		{
-			EditorMeshCollector.AddViewMeshArrays(
-				Views[ViewIndex],
-				&Views[ViewIndex]->DynamicEditorMeshElements,
-				&Views[ViewIndex]->EditorSimpleElementCollector,
-				&Views[ViewIndex]->DynamicPrimitiveCollector
-#if UE_ENABLE_DEBUG_DRAWING
-				, &Views[ViewIndex]->DebugSimpleElementCollector
-#endif
-			);
-		}
-	}
+	DynamicMeshElements.ContextContainer.Init(SceneRenderer, NumAsyncDynamicMeshElementContexts);
+	DynamicMeshElements.ViewCommandsPerView.SetNum(Views.Num());
 
 	Tasks.LightVisibility.AddPrerequisites(UE::Tasks::Launch(UE_SOURCE_LOCATION, [this]
 	{
@@ -3755,104 +4075,65 @@ void FVisibilityTaskData::MergeSecondaryViewVisibility()
 	}
 }
 
-void FVisibilityTaskData::FinishGatherDynamicMeshElements(FVirtualTextureUpdater* VirtualTextureUpdater)
+void FVisibilityTaskData::GatherDynamicMeshElements(FDynamicPrimitiveIndexList&& DynamicPrimitiveIndexList)
 {
-	check(IsInRenderingThread());
-	Scene.WaitForGPUSkinCacheTask();
+	FDynamicPrimitiveIndexList RenderThreadDynamicPrimitiveIndexList;
 
-	if (DynamicMeshElements.CommandPipe)
+	const int32 NumAsyncContexts = DynamicMeshElements.ContextContainer.GetNumAsyncContexts();
+
+	if (NumAsyncContexts > 0)
 	{
-		SCOPED_NAMED_EVENT(WaitForGatherDynamicMeshElements, FColor::Magenta);
-		DynamicMeshElements.CommandPipeCompleteEvent->Wait(ENamedThreads::GetRenderThread_Local());
+		const auto FilterDynamicPrimitives = [] (TArrayView<FPrimitiveSceneProxy*> PrimitiveSceneProxies, FDynamicPrimitiveIndexList::FList& Primitives, FDynamicPrimitiveIndexList::FList& RenderThreadPrimitives)
+		{
+			for (int32 Index = 0; Index < Primitives.Num(); )
+			{
+				const FDynamicPrimitiveIndex PrimitiveIndex = Primitives[Index];
+
+				if (!PrimitiveSceneProxies[PrimitiveIndex.Index]->SupportsParallelGDME())
+				{
+					RenderThreadPrimitives.Emplace(PrimitiveIndex);
+					Primitives.RemoveAtSwap(Index, 1, false);
+				}
+				else
+				{
+					Index++;
+				}
+			}
+		};
+
+		FilterDynamicPrimitives(Scene.PrimitiveSceneProxies, DynamicPrimitiveIndexList.Primitives, RenderThreadDynamicPrimitiveIndexList.Primitives);
+#if WITH_EDITOR
+		FilterDynamicPrimitives(Scene.PrimitiveSceneProxies, DynamicPrimitiveIndexList.EditorPrimitives, RenderThreadDynamicPrimitiveIndexList.EditorPrimitives);
+#endif
 	}
 	else
 	{
-		Tasks.ComputeRelevance.Wait();
-
-		check(DynamicMeshElements.PrimitiveViewMasks);
-		GatherDynamicMeshElements(*DynamicMeshElements.PrimitiveViewMasks);
+		RenderThreadDynamicPrimitiveIndexList = MoveTemp(DynamicPrimitiveIndexList);
+		DynamicPrimitiveIndexList = {};
 	}
 
-	// Sync the virtual texture update task before finishing the mesh collectors. Render proxies can register new
-	// materials which require evaluating uniform expression caches, which can contain virtual textures. Newly allocated
-	// virtual textures are processed later.
-	FVirtualTextureSystem::Get().WaitForTasks(VirtualTextureUpdater);
-
-	MeshCollector.Finish();
-
-	if (GIsEditor)
+	if (!RenderThreadDynamicPrimitiveIndexList.IsEmpty())
 	{
-		EditorMeshCollector.Finish();
+		Tasks.DynamicMeshElementsRenderThread = DynamicMeshElements.ContextContainer.LaunchRenderThreadTask(MoveTemp(RenderThreadDynamicPrimitiveIndexList));
 	}
 
-	DynamicMeshElements.DynamicVertexBuffer.Commit();
-	DynamicMeshElements.DynamicIndexBuffer.Commit();
-	SceneRenderer.DynamicReadBufferForInitViews.Commit(RHICmdList);
-}
-
-void FVisibilityTaskData::GatherDynamicMeshElementsForPrimitive(int32 PrimitiveIndex, uint8 ViewMask)
-{
-	SCOPED_NAMED_EVENT(DynamicPrimitive, FColor::Magenta);
-
-	TArray<int32, TInlineAllocator<4>> MeshBatchCountBefore;
-	MeshBatchCountBefore.SetNumUninitialized(Views.Num());
-	for (int32 ViewIndex = 0; ViewIndex < Views.Num(); ViewIndex++)
+	if (!DynamicPrimitiveIndexList.IsEmpty())
 	{
-		MeshBatchCountBefore[ViewIndex] = MeshCollector.GetMeshBatchCount(ViewIndex);
-	}
-
-	FPrimitiveSceneInfo* PrimitiveSceneInfo = Scene.Primitives[PrimitiveIndex];
-	MeshCollector.SetPrimitive(PrimitiveSceneInfo->Proxy, PrimitiveSceneInfo->DefaultDynamicHitProxyId);
+		FDynamicPrimitiveIndexQueue* Queue = Allocator.Create<FDynamicPrimitiveIndexQueue>(MoveTemp(DynamicPrimitiveIndexList));
 
-	PrimitiveSceneInfo->Proxy->GetDynamicMeshElements(ViewFamily.AllViews, ViewFamily, ViewMask, MeshCollector);
-
-	for (int32 ViewIndex = 0; ViewIndex < Views.Num(); ViewIndex++)
-	{
-		FViewInfo& View = *Views[ViewIndex];
-		if (ViewMask & (1 << ViewIndex))
+		for (int32 Index = 0; Index < DynamicMeshElements.ContextContainer.GetNumAsyncContexts(); ++Index)
 		{
-			FDynamicPrimitive& DynamicPrimitive = DynamicMeshElements.DynamicPrimitives.Emplace_GetRef();
-			DynamicPrimitive.PrimitiveIndex = PrimitiveIndex;
-			DynamicPrimitive.ViewIndex = ViewIndex;
-			DynamicPrimitive.StartElementIndex = DynamicMeshElements.LastElementPerView[ViewIndex];
-			DynamicPrimitive.EndElementIndex = View.DynamicMeshElements.Num();
-
-			DynamicMeshElements.LastElementPerView[ViewIndex] = View.DynamicMeshElements.Num();
+			Tasks.DynamicMeshElements.AddPrerequisites(DynamicMeshElements.ContextContainer.LaunchAsyncTask(Queue, Index, TaskConfig.TaskPriority));
 		}
-
-		View.DynamicMeshElementRanges[PrimitiveIndex] = FInt32Vector2(MeshBatchCountBefore[ViewIndex], MeshCollector.GetMeshBatchCount(ViewIndex));
 	}
-};
-
-void FVisibilityTaskData::GatherDynamicMeshElementsForEditorPrimitive(int32 PrimitiveIndex, uint8 ViewMask)
-{
-	FPrimitiveSceneInfo* PrimitiveSceneInfo = Scene.Primitives[PrimitiveIndex];
-	EditorMeshCollector.SetPrimitive(PrimitiveSceneInfo->Proxy, PrimitiveSceneInfo->DefaultDynamicHitProxyId);
-
-	PrimitiveSceneInfo->Proxy->GetDynamicMeshElements(ViewFamily.AllViews, ViewFamily, ViewMask, EditorMeshCollector);
-}
-
-void FVisibilityTaskData::GatherDynamicMeshElements(const FDynamicPrimitiveIndexList& DynamicPrimitiveIndexList)
-{
-	constexpr uint8 ViewMask = 1;
-
-	for (int32 PrimitiveIndex : DynamicPrimitiveIndexList.Primitives)
-	{
-		GatherDynamicMeshElementsForPrimitive(PrimitiveIndex, ViewMask);
-	}
-
-#if WITH_EDITOR
-	for (int32 PrimitiveIndex : DynamicPrimitiveIndexList.EditorPrimitives)
-	{
-		GatherDynamicMeshElementsForEditorPrimitive(PrimitiveIndex, ViewMask);
-	}
-#endif
 }
 
 void FVisibilityTaskData::GatherDynamicMeshElements(const FDynamicPrimitiveViewMasks& DynamicPrimitiveViewMasks)
 {
 	SCOPED_NAMED_EVENT(GatherDynamicMeshElements, FColor::Magenta);
 
+	Tasks.DynamicMeshElementsPrerequisites.Wait();
+
 	const auto GetPrimaryViewMask = [this] (uint8 ViewMask) -> uint8
 	{
 		// If a mesh is visible in a secondary view, mark it as visible in the primary view
@@ -3871,6 +4152,9 @@ void FVisibilityTaskData::GatherDynamicMeshElements(const FDynamicPrimitiveViewM
 		return ViewMaskFinal;
 	};
 
+	FDynamicPrimitiveIndexList DynamicPrimitiveIndexList;
+	DynamicPrimitiveIndexList.Primitives.Reserve(128);
+
 	const int32 NumPrimitives = DynamicPrimitiveViewMasks.Primitives.Num();
 
 	for (int32 PrimitiveIndex = 0; PrimitiveIndex < NumPrimitives; ++PrimitiveIndex)
@@ -3879,28 +4163,35 @@ void FVisibilityTaskData::GatherDynamicMeshElements(const FDynamicPrimitiveViewM
 
 		if (ViewMask != 0)
 		{
-			GatherDynamicMeshElementsForPrimitive(PrimitiveIndex, GetPrimaryViewMask(ViewMask));
+			DynamicPrimitiveIndexList.Primitives.Emplace(PrimitiveIndex, GetPrimaryViewMask(ViewMask));
 		}
 	}
 
 #if WITH_EDITOR
 	if (GIsEditor)
 	{
+		DynamicPrimitiveIndexList.EditorPrimitives.Reserve(128);
+
 		for (int32 PrimitiveIndex = 0; PrimitiveIndex < NumPrimitives; ++PrimitiveIndex)
 		{
 			const uint8 ViewMask = DynamicPrimitiveViewMasks.EditorPrimitives[PrimitiveIndex];
 
 			if (ViewMask != 0)
 			{
-				GatherDynamicMeshElementsForEditorPrimitive(PrimitiveIndex, GetPrimaryViewMask(ViewMask));
+				DynamicPrimitiveIndexList.EditorPrimitives.Emplace(PrimitiveIndex, GetPrimaryViewMask(ViewMask));
 			}
 		}
 	}
 #endif
+
+	GatherDynamicMeshElements(MoveTemp(DynamicPrimitiveIndexList));
+	Tasks.DynamicMeshElements.Trigger();
 }
 
 void FVisibilityTaskData::SetupMeshPasses(FExclusiveDepthStencil::Type BasePassDepthStencilAccess, FInstanceCullingManager& InstanceCullingManager)
 {
+	DynamicMeshElements.ContextContainer.MergeContexts(DynamicMeshElements.DynamicPrimitives);
+
 	{
 		SCOPED_NAMED_EVENT(DynamicRelevance, FColor::Magenta);
 
@@ -3968,11 +4259,15 @@ void FVisibilityTaskData::SetupMeshPasses(FExclusiveDepthStencil::Type BasePassD
 	}
 }
 
-void FVisibilityTaskData::ProcessRenderThreadTasks(FExclusiveDepthStencil::Type BasePassDepthStencilAccess, FInstanceCullingManager& InstanceCullingManager, FVirtualTextureUpdater* VirtualTextureUpdater)
+void FVisibilityTaskData::ProcessRenderThreadTasks()
 {
 	SCOPE_CYCLE_COUNTER(STAT_ViewVisibilityTime);
 	SCOPED_NAMED_EVENT(ProcessVisibilityTasks, FColor::Magenta);
 
+	UE::Tasks::FTask VirtualTextureTask;
+
+	StartGatherDynamicMeshElements();
+
 	if (TaskConfig.Schedule == EVisibilityTaskSchedule::RenderThread)
 	{
 		FPrimitiveRange PrimitiveRange;
@@ -4023,23 +4318,25 @@ void FVisibilityTaskData::ProcessRenderThreadTasks(FExclusiveDepthStencil::Type
 		}
 
 		Tasks.bWaitingAllowed = true;
-	}
 
-	FinishGatherDynamicMeshElements(VirtualTextureUpdater);
-
-	if (TaskConfig.Schedule == EVisibilityTaskSchedule::Parallel)
+		check(DynamicMeshElements.PrimitiveViewMasks);
+		GatherDynamicMeshElements(*DynamicMeshElements.PrimitiveViewMasks);
+	}
+	else
 	{
-		// We sync cached mesh draw commands here to make things simple, as there are several places where cached mesh commands
-		// are referenced later on in the frame and it becomes difficult to reason about it and it's easy to miss. In practice
-		// mesh command caching for most use cases should be complete by now.
-		Scene.WaitForCacheMeshDrawCommandsTask();
-
-		Tasks.MeshPassSetup = UE::Tasks::Launch(UE_SOURCE_LOCATION, [this, BasePassDepthStencilAccess, &InstanceCullingManager]
+		if (DynamicMeshElements.CommandPipe)
 		{
-			FOptionalTaskTagScope Scope(ETaskTag::EParallelRenderingThread);
-			SetupMeshPasses(BasePassDepthStencilAccess, InstanceCullingManager);
+			SCOPED_NAMED_EVENT(WaitForGatherDynamicMeshElements, FColor::Magenta);
 
-		}, Tasks.FinalizeRelevance, TaskConfig.TaskPriority);
+			// Wait on the command pipe first as it will be continually updating the render thread event (and process tasks while we wait).
+			Tasks.DynamicMeshElementsPipe->Wait(ENamedThreads::GetRenderThread_Local());
+		}
+		else
+		{
+			Tasks.ComputeRelevance.Wait();
+			check(DynamicMeshElements.PrimitiveViewMasks);
+			GatherDynamicMeshElements(*DynamicMeshElements.PrimitiveViewMasks);
+		}
 
 		for (FVisibilityViewPacket& ViewPacket : ViewPackets)
 		{
@@ -4049,12 +4346,15 @@ void FVisibilityTaskData::ProcessRenderThreadTasks(FExclusiveDepthStencil::Type
 			}
 		}
 	}
-	else
+
+	// Now process all gather dynamic mesh element tasks that were queued up to run on the render thread.
+	if (Tasks.DynamicMeshElementsRenderThread)
 	{
-		SetupMeshPasses(BasePassDepthStencilAccess, InstanceCullingManager);
+		Tasks.DynamicMeshElementsRenderThread->Wait(ENamedThreads::GetRenderThread_Local());
 	}
 
 	Tasks.LightVisibility.Wait();
+	Tasks.FinalizeRelevance.Wait();
 
 	INC_DWORD_STAT_BY(STAT_ProcessedPrimitives, Scene.Primitives.Num() * Views.Num());
 	INC_DWORD_STAT_BY(STAT_CulledPrimitives, TaskConfig.FrustumCull.NumCulledPrimitives);
@@ -4068,6 +4368,25 @@ void FVisibilityTaskData::ProcessRenderThreadTasks(FExclusiveDepthStencil::Type
 	TRACE_COUNTER_SET(Scene_Visibility_Relevance_NumPrimitivesPerPacket, TaskConfig.Relevance.NumPrimitivesPerPacket);
 }
 
+void FVisibilityTaskData::FinishGatherDynamicMeshElements(FExclusiveDepthStencil::Type BasePassDepthStencilAccess, FInstanceCullingManager& InstanceCullingManager, FVirtualTextureUpdater* VirtualTextureUpdater)
+{
+	check(IsInRenderingThread());
+	SCOPED_NAMED_EVENT(FinishDynamicMeshElements, FColor::Magenta);
+
+	FVirtualTextureSystem::Get().WaitForTasks(VirtualTextureUpdater);
+	Tasks.DynamicMeshElements.Wait();
+	DynamicMeshElements.ContextContainer.Submit(RHICmdList);
+
+	Tasks.MeshPassSetup = UE::Tasks::Launch(UE_SOURCE_LOCATION, [this, BasePassDepthStencilAccess, &InstanceCullingManager]
+	{
+		FOptionalTaskTagScope Scope(ETaskTag::EParallelRenderingThread);
+		SetupMeshPasses(BasePassDepthStencilAccess, InstanceCullingManager);
+
+	}, TaskConfig.TaskPriority);
+
+	FSceneRenderer::DynamicReadBufferForInitViews.Commit(RHICmdList);
+}
+
 void FVisibilityTaskData::Finish()
 {
 	SCOPED_NAMED_EVENT(FinishVisibility, FColor::Magenta);
@@ -4076,15 +4395,10 @@ void FVisibilityTaskData::Finish()
 
 	check(Tasks.ComputeRelevance.IsCompleted());
 	check(Tasks.FinalizeRelevance.IsCompleted());
+	check(Tasks.DynamicMeshElements.IsCompleted());
 	check(Tasks.MeshPassSetup.IsCompleted());
 
-	if (DynamicMeshElements.CommandPipe)
-	{
-		DynamicMeshElements.CommandPipe->Wait();
-	}
-
 	ViewPackets.Empty();
-	DynamicMeshElements.LastElementPerView.Empty();
 	DynamicMeshElements.DynamicPrimitives.Empty();
 	Allocator.BulkDelete();
 	bFinished = true;
@@ -4934,9 +5248,7 @@ void FDeferredShadingSceneRenderer::PreVisibilityFrameSetup(FRDGBuilder& GraphBu
 void FDeferredShadingSceneRenderer::BeginInitViews(
 	FRDGBuilder& GraphBuilder,
 	const FSceneTexturesConfig& SceneTexturesConfig,
-	FExclusiveDepthStencil::Type BasePassDepthStencilAccess,
 	FInstanceCullingManager& InstanceCullingManager,
-	FVirtualTextureUpdater* VirtualTextureUpdater,
 	FRDGExternalAccessQueue& ExternalAccessQueue,
 	FInitViewTaskDatas& TaskDatas)
 {
@@ -4948,10 +5260,13 @@ void FDeferredShadingSceneRenderer::BeginInitViews(
 
 	PreVisibilityFrameSetup(GraphBuilder);
 
+	// Start processing dynamic mesh elements tasks early enough to overlap with shadows and GPU scene update.
+	TaskDatas.VisibilityTaskData->StartGatherDynamicMeshElements();
+
 	// Attempt to launch dynamic shadow tasks early before finalizing visibility.
 	if (bRendererOutputFinalSceneColor)
 	{
-		BeginInitDynamicShadows(TaskDatas);
+		BeginInitDynamicShadows(GraphBuilder, TaskDatas, InstanceCullingManager);
 	}
 
 	FRHICommandListImmediate& RHICmdList = GraphBuilder.RHICmdList;
@@ -4965,6 +5280,7 @@ void FDeferredShadingSceneRenderer::BeginInitViews(
 	}
 
 	// Create GPU-side representation of the view for instance culling.
+	InstanceCullingManager.AllocateViews(Views.Num());
 	for (int32 ViewIndex = 0; ViewIndex < Views.Num(); ++ViewIndex)
 	{
 		Views[ViewIndex].GPUSceneViewId = InstanceCullingManager.RegisterView(Views[ViewIndex]);
@@ -4995,14 +5311,6 @@ void FDeferredShadingSceneRenderer::BeginInitViews(
 		Scene->GPUScene.Update(GraphBuilder, GetSceneUniforms(), *Scene, ExternalAccessQueue, TaskDatas.VisibilityTaskData);
 	}
 
-	TaskDatas.VisibilityTaskData->ProcessRenderThreadTasks(BasePassDepthStencilAccess, InstanceCullingManager, VirtualTextureUpdater);
-
-	// Make a second attempt to launch shadow tasks it wasn't able to the first time due to visibility being deferred.
-	if (bRendererOutputFinalSceneColor)
-	{
-		BeginInitDynamicShadows(TaskDatas);
-	}
-
 	// This must happen before we start initialising and using views.
 	if (Scene)
 	{
@@ -5015,8 +5323,6 @@ void FDeferredShadingSceneRenderer::BeginInitViews(
 		InitSkyAtmosphereForViews(RHICmdList);
 	}
 
-	PostVisibilityFrameSetup(TaskDatas.ILCUpdatePrim);
-
 	{
 		QUICK_SCOPE_CYCLE_COUNTER(STAT_InitViews_InitRHIResources);
 		// initialize per-view uniform buffer. Do it from back to front because secondary stereo view follows its primary one, but primary needs to know the instanced's params
@@ -5040,10 +5346,15 @@ void FDeferredShadingSceneRenderer::BeginInitViews(
 		}
 	}
 
-	if (GDynamicRHI->RHIIncludeOptionalFlushes())
+	TaskDatas.VisibilityTaskData->ProcessRenderThreadTasks();
+
+	// Make a second attempt to launch shadow tasks it wasn't able to the first time due to visibility being deferred.
+	if (bRendererOutputFinalSceneColor)
 	{
-		RHICmdList.ImmediateFlush(EImmediateFlushType::DispatchToRHIThread);
+		BeginInitDynamicShadows(GraphBuilder, TaskDatas, InstanceCullingManager);
 	}
+
+	PostVisibilityFrameSetup(TaskDatas.ILCUpdatePrim);
 }
 
 template<class T>
@@ -5126,6 +5437,8 @@ void FDeferredShadingSceneRenderer::EndInitViews(
 	SCOPED_NAMED_EVENT(FDeferredShadingSceneRenderer_InitViewsAfterPrepass, FColor::Emerald);
 	SCOPE_CYCLE_COUNTER(STAT_InitViewsPossiblyAfterPrepass);
 
+	BeginShadowGatherDynamicMeshElements(TaskDatas.DynamicShadows);
+
 	TaskDatas.VisibilityTaskData->Finish();
 
 	FRHICommandListImmediate& RHICmdList = GraphBuilder.RHICmdList;
diff --git a/Engine/Source/Runtime/Renderer/Private/SceneVisibility.h b/Engine/Source/Runtime/Renderer/Private/SceneVisibility.h
--- a/Engine/Source/Runtime/Renderer/Private/SceneVisibility.h
+++ b/Engine/Source/Runtime/Renderer/Private/SceneVisibility.h
@@ -34,8 +34,14 @@ class IVisibilityTaskData
 public:
 	virtual ~IVisibilityTaskData() {}
 
+	/** [Optional] Call to allow early processing of async GDME tasks when it is safe to do so. Otherwise, this is automatically called from ProcessRenderThreadTasks. */
+	virtual void StartGatherDynamicMeshElements() = 0;
+
 	/** Processes all visibility tasks that must be performed on the render thread. */
-	virtual void ProcessRenderThreadTasks(FExclusiveDepthStencil::Type BasePassDepthStencilAccess, FInstanceCullingManager& InstanceCullingManager, FVirtualTextureUpdater* VirtualTextureUpdater) = 0;
+	virtual void ProcessRenderThreadTasks() = 0;
+
+	/** Called to finish processing of the GDME tasks. */
+	virtual void FinishGatherDynamicMeshElements(FExclusiveDepthStencil::Type BasePassDepthStencilAccess, FInstanceCullingManager& InstanceCullingManager, FVirtualTextureUpdater* VirtualTextureUpdater) = 0;
 
 	/** Waits for the task graph and cleans up. */
 	virtual void Finish() = 0;
diff --git a/Engine/Source/Runtime/Renderer/Private/SceneVisibilityPrivate.h b/Engine/Source/Runtime/Renderer/Private/SceneVisibilityPrivate.h
--- a/Engine/Source/Runtime/Renderer/Private/SceneVisibilityPrivate.h
+++ b/Engine/Source/Runtime/Renderer/Private/SceneVisibilityPrivate.h
@@ -32,12 +32,9 @@ public:
 	using CommandFunctionType = TFunction<void(CommandType&&)>;
 	using EmptyFunctionType = TFunction<void()>;
 
-	TCommandPipe(const TCHAR* InName, ENamedThreads::Type InDesiredThread = ENamedThreads::AnyHiPriThreadHiPriTask)
-		: Name(InName)
-		, DesiredThread(InDesiredThread)
-	{
-		checkf(!EnumHasAnyFlags(DesiredThread, ENamedThreads::LocalQueue) || DesiredThread == ENamedThreads::GetRenderThread_Local(), TEXT("The only supported local queue is on the render thread."));
-	}
+	TCommandPipe(const TCHAR* InName)
+		: Pipe(InName)
+	{}
 
 	~TCommandPipe()
 	{
@@ -55,12 +52,7 @@ public:
 		EmptyFunction = Forward<EmptyFunctionType&&>(InEmptyFunction);
 	}
 
-	void SetEmptyTaskEvent(const FGraphEventRef& InEmptyTaskEvent)
-	{
-		EmptyTaskEvent = InEmptyTaskEvent;
-	}
-
-	void SetPrerequisiteTask(const FGraphEventRef& InPrerequisiteTask)
+	void SetPrerequisiteTask(const UE::Tasks::FTask& InPrerequisiteTask)
 	{
 		PrerequisiteTask = InPrerequisiteTask;
 	}
@@ -82,12 +74,6 @@ public:
 				EmptyFunction();
 			}
 
-			if (EmptyTaskEvent)
-			{
-				EmptyTaskEvent->DispatchSubsequents();
-				EmptyTaskEvent = nullptr;
-			}
-
 			NumEmptyEvents.fetch_add(1, std::memory_order_relaxed);
 		}
 	}
@@ -97,11 +83,38 @@ public:
 	{
 		check(CommandFunction);
 
-		const UE::EConsumeAllMpmcQueueResult Result = Queue.ProduceItem(Forward<ArgTypes>(Args)...);
+		QueueMutex.Lock();
+		const bool bWasEmpty = Queue.IsEmpty();
+		Queue.Emplace(Forward<ArgTypes>(Args)...);
+		QueueMutex.Unlock();
 
-		if (Result == UE::EConsumeAllMpmcQueueResult::WasEmpty)
+		if (bWasEmpty)
 		{
-			Launch();
+			Pipe.Launch(Pipe.GetDebugName(), [this]
+			{
+				FOptionalTaskTagScope Scope(ETaskTag::EParallelRenderingThread);
+				SCOPED_NAMED_EVENT_TCHAR(Pipe.GetDebugName(), FColor::Magenta);
+
+				TArray<CommandType, SceneRenderingAllocator> Commands;
+
+				QueueMutex.Lock();
+				Commands = MoveTemp(Queue);
+				QueueMutex.Unlock();
+
+				int32 NumProcessedCommands = 0;
+
+				for (CommandType& Command : Commands)
+				{
+					CommandFunction(MoveTemp(Command));
+					NumProcessedCommands++;
+				}
+
+				if (NumProcessedCommands)
+				{
+					ReleaseNumCommands(NumProcessedCommands);
+				}
+
+			}, PrerequisiteTask);
 		}
 	}
 
@@ -112,67 +125,18 @@ public:
 
 	void Wait()
 	{
-		if (CurrentTask)
-		{
-			check(IsInRenderingThread());
-			CurrentTask->Wait(ENamedThreads::GetRenderThread_Local());
-			CurrentTask = nullptr;
-		}
+		Pipe.WaitUntilEmpty();
 	}
 
 private:
-	void Launch()
-	{
-		UE::TUniqueLock Lock(LaunchMutex);
-		FGraphEventArray Prereqs;
-
-		if (CurrentTask)
-		{
-			Prereqs.Add(CurrentTask);
-		}
-
-		if (PrerequisiteTask)
-		{
-			Prereqs.Add(PrerequisiteTask);
-		}
-
-		CurrentTask = FFunctionGraphTask::CreateAndDispatchWhenReady([this]
-		{
-			FOptionalTaskTagScope Scope(ETaskTag::EParallelRenderingThread);
-			Process();
-
-		}, TStatId(), &Prereqs, DesiredThread);
-	}
-
-	void Process()
-	{
-		SCOPED_NAMED_EVENT_TCHAR(Name, FColor::Magenta);
-
-		int32 NumProcessedCommands = 0;
-
-		const UE::EConsumeAllMpmcQueueResult Result = Queue.ConsumeAllLifo([&](CommandType&& Command)
-		{
-			CommandFunction(MoveTemp(Command));
-			NumProcessedCommands++;
-		});
-
-		if (NumProcessedCommands)
-		{
-			ReleaseNumCommands(NumProcessedCommands);
-		}
-	}
-
-	const TCHAR* Name;
-	FGraphEventRef CurrentTask;
-	FGraphEventRef PrerequisiteTask;
-	FGraphEventRef EmptyTaskEvent;
+	UE::Tasks::FTask PrerequisiteTask;
 	CommandFunctionType CommandFunction;
 	EmptyFunctionType EmptyFunction;
-	UE::TConsumeAllMpmcQueue<CommandType, FConcurrentLinearAllocator> Queue;
+	UE::FMutex QueueMutex;
+	TArray<CommandType, SceneRenderingAllocator> Queue;
+	UE::Tasks::FPipe Pipe;
 	std::atomic_int32_t NumCommands{ 0 };
 	std::atomic_int32_t NumEmptyEvents{ 0 };
-	UE::FMutex LaunchMutex;
-	ENamedThreads::Type DesiredThread;
 };
 
 ///////////////////////////////////////////////////////////////////////////////
@@ -182,6 +146,8 @@ inline UE::Tasks::EExtendedTaskPriority GetExtendedTaskPriority(bool bExecuteInP
 	return bExecuteInParallel ? UE::Tasks::EExtendedTaskPriority::None : UE::Tasks::EExtendedTaskPriority::Inline;
 }
 
+using FPrimitiveIndexList = TArray<int32, SceneRenderingAllocator>;
+
 struct FPrimitiveRange
 {
 	int32 StartIndex;
@@ -196,10 +162,23 @@ struct FDynamicPrimitive
 	int32 EndElementIndex;
 };
 
-using FPrimitiveIndexList = TArray<int32, SceneRenderingAllocator>;
+struct FDynamicPrimitiveIndex
+{
+	FDynamicPrimitiveIndex() = default;
+
+	FDynamicPrimitiveIndex(int32 InIndex, uint8 InViewMask)
+		: Index(InIndex)
+		, ViewMask(InViewMask)
+	{}
+
+	uint32 Index    : 24;
+	uint32 ViewMask : 8;
+};
 
 struct FDynamicPrimitiveIndexList
 {
+	using FList = TArray<FDynamicPrimitiveIndex, SceneRenderingAllocator>;
+
 	bool IsEmpty() const
 	{
 		return Primitives.IsEmpty()
@@ -209,10 +188,49 @@ struct FDynamicPrimitiveIndexList
 			;
 	}
 
-	FPrimitiveIndexList Primitives;
+	FList Primitives;
 
 #if WITH_EDITOR
-	FPrimitiveIndexList EditorPrimitives;
+	FList EditorPrimitives;
+#endif
+};
+
+class FDynamicPrimitiveIndexQueue
+{
+public:
+	FDynamicPrimitiveIndexQueue(FDynamicPrimitiveIndexList&& InList)
+		: List(InList)
+	{}
+
+	bool Pop(FDynamicPrimitiveIndex& PrimitiveIndex)
+	{
+		const int32 Index = NextIndex.fetch_add(1, std::memory_order_relaxed);
+		if (Index < List.Primitives.Num())
+		{
+			PrimitiveIndex = List.Primitives[Index];
+			return true;
+		}
+		return false;
+	}
+
+#if WITH_EDITOR
+	bool PopEditor(FDynamicPrimitiveIndex& PrimitiveIndex)
+	{
+		const int32 Index = NextEditorIndex.fetch_add(1, std::memory_order_relaxed);
+		if (Index < List.EditorPrimitives.Num())
+		{
+			PrimitiveIndex = List.EditorPrimitives[Index];
+			return true;
+		}
+		return false;
+	}
+#endif
+
+private:
+	FDynamicPrimitiveIndexList List;
+	std::atomic_int32_t NextIndex = { 0 };
+#if WITH_EDITOR
+	std::atomic_int32_t NextEditorIndex = { 0 };
 #endif
 };
 
@@ -225,6 +243,93 @@ struct FDynamicPrimitiveViewMasks
 #endif
 };
 
+///////////////////////////////////////////////////////////////////////////////
+
+class FDynamicMeshElementContext
+{
+public:
+	FDynamicMeshElementContext(FSceneRenderer& SceneRenderer);
+
+	FGraphEventRef LaunchRenderThreadTask(FDynamicPrimitiveIndexList&& PrimitiveIndexList);
+
+	UE::Tasks::FTask LaunchAsyncTask(FDynamicPrimitiveIndexQueue* PrimitiveIndexQueue, UE::Tasks::ETaskPriority TaskPriority);
+
+	void GatherDynamicMeshElementsForPrimitive(FPrimitiveSceneInfo* Primitive, uint8 ViewMask);
+
+	void GatherDynamicMeshElementsForEditorPrimitive(FPrimitiveSceneInfo* Primitive, uint8 ViewMask);
+
+private:
+	void Finish();
+
+	struct FViewMeshArrays
+	{
+		TArray<FMeshBatchAndRelevance, SceneRenderingAllocator> DynamicMeshElements;
+		FSimpleElementCollector SimpleElementCollector;
+
+#if WITH_EDITOR
+		TArray<FMeshBatchAndRelevance, SceneRenderingAllocator> DynamicEditorMeshElements;
+		FSimpleElementCollector EditorSimpleElementCollector;
+#endif
+
+#if UE_ENABLE_DEBUG_DRAWING
+		FSimpleElementCollector DebugSimpleElementCollector;
+#endif
+	};
+
+	const FSceneViewFamily& ViewFamily;
+	TArrayView<FViewInfo*> Views;
+	TArrayView<FPrimitiveSceneInfo*> Primitives;
+	TArray<FViewMeshArrays, TInlineAllocator<2>> ViewMeshArraysPerView;
+	TArray<FDynamicPrimitive, SceneRenderingAllocator> DynamicPrimitives;
+	FMeshElementCollector MeshCollector;
+#if WITH_EDITOR
+	FMeshElementCollector EditorMeshCollector;
+#endif
+	FRHICommandList* RHICmdList;
+	FGlobalDynamicVertexBuffer DynamicVertexBuffer;
+	FGlobalDynamicIndexBuffer DynamicIndexBuffer;
+	UE::Tasks::FPipe Pipe{UE_SOURCE_LOCATION};
+
+	friend class FDynamicMeshElementContextContainer;
+};
+
+class FDynamicMeshElementContextContainer
+{
+public:
+	~FDynamicMeshElementContextContainer();
+
+	int32 GetNumAsyncContexts() const
+	{
+		return Contexts.Num() - 1;
+	}
+
+	FDynamicMeshElementContext* GetRenderThreadContext() const
+	{
+		check(!bFinished);
+		return Contexts.Last();
+	}
+
+	FGraphEventRef LaunchRenderThreadTask(FDynamicPrimitiveIndexList&& PrimitiveIndexList);
+
+	UE::Tasks::FTask LaunchAsyncTask(FDynamicPrimitiveIndexQueue* PrimitiveIndexQueue, int32 Index, UE::Tasks::ETaskPriority TaskPriority);
+
+	void Init(FSceneRenderer& InSceneRenderer, int32 NumAsyncContexts);
+
+	void MergeContexts(TArray<FDynamicPrimitive, SceneRenderingAllocator>& OutDynamicPrimitives);
+
+	void Submit(FRHICommandListImmediate& RHICmdList);
+
+private:
+	using FDynamicMeshElementContextArray = TArray<FDynamicMeshElementContext*>;
+
+	TArrayView<FViewInfo*> Views;
+	FDynamicMeshElementContextArray Contexts;
+	TArray<FRHICommandListImmediate::FQueuedCommandList, FConcurrentLinearArrayAllocator> CommandLists;
+	bool bFinished = false;
+};
+
+///////////////////////////////////////////////////////////////////////////////
+
 enum class EVisibilityTaskSchedule
 {
 	// Visibility is processed on the render thread with assistance from other threads via parallel for.
@@ -395,7 +500,20 @@ public:
 	}
 
 	void LaunchVisibilityTasks();
-	void ProcessRenderThreadTasks(FExclusiveDepthStencil::Type BasePassDepthStencilAccess, FInstanceCullingManager& InstanceCullingManager, FVirtualTextureUpdater* VirtualTextureUpdater) override;
+
+	void ProcessRenderThreadTasks() override;
+
+	void StartGatherDynamicMeshElements() override
+	{
+		if (!Tasks.bDynamicMeshElementsPrerequisitesTriggered)
+		{
+			Tasks.DynamicMeshElementsPrerequisites.Trigger();
+			Tasks.bDynamicMeshElementsPrerequisitesTriggered = true;
+		}
+	}
+
+	void FinishGatherDynamicMeshElements(FExclusiveDepthStencil::Type BasePassDepthStencilAccess, FInstanceCullingManager& InstanceCullingManager, FVirtualTextureUpdater* VirtualTextureUpdater) override;
+
 	void Finish() override;
 
 	TArrayView<FViewCommands> GetViewCommandsPerView() override
@@ -426,14 +544,9 @@ public:
 private:
 	void MergeSecondaryViewVisibility();
 
-	void FinishGatherDynamicMeshElements(FVirtualTextureUpdater* VirtualTextureUpdater);
-
-	void GatherDynamicMeshElements(const FDynamicPrimitiveIndexList& Primitives);
+	void GatherDynamicMeshElements(FDynamicPrimitiveIndexList&& Primitives);
 	void GatherDynamicMeshElements(const FDynamicPrimitiveViewMasks& Primitives);
 
-	void GatherDynamicMeshElementsForPrimitive(int32 PrimitiveIndex, uint8 ViewMask);
-	void GatherDynamicMeshElementsForEditorPrimitive(int32 PrimitiveIndex, uint8 ViewMask);
-
 	void SetupMeshPasses(FExclusiveDepthStencil::Type BasePassDepthStencilAccess, FInstanceCullingManager& InstanceCullingManager);
 
 	FRHICommandListImmediate& RHICmdList;
@@ -442,26 +555,21 @@ private:
 	TArrayView<FViewInfo*> Views;
 	FViewFamilyInfo& ViewFamily;
 	EShadingPath ShadingPath;
-	FMeshElementCollector& MeshCollector;
-	FMeshElementCollector& EditorMeshCollector;
 	FSceneRenderingBulkObjectAllocator Allocator;
 	TArray<FVisibilityViewPacket, SceneRenderingAllocator> ViewPackets;
 
 	struct FDynamicMeshElements
 	{
 		// The command pipe and event are only non-null when in single-view mode.
-		FGraphEventRef CommandPipeCompleteEvent;
 		TCommandPipe<FDynamicPrimitiveIndexList>* CommandPipe = nullptr;
 
 		// Primitive view masks are only non-null when in multi-view mode.
 		FDynamicPrimitiveViewMasks* PrimitiveViewMasks = nullptr;
 
-		FGlobalDynamicVertexBuffer DynamicVertexBuffer;
-		FGlobalDynamicIndexBuffer DynamicIndexBuffer;
-
+		FDynamicMeshElementContextContainer ContextContainer;
 		TArray<FViewCommands, TInlineAllocator<4>> ViewCommandsPerView;
-		TArray<uint32, TInlineAllocator<4, SceneRenderingAllocator>> LastElementPerView;
 		TArray<FDynamicPrimitive, SceneRenderingAllocator> DynamicPrimitives;
+		TArray<UE::Tasks::FPipe, SceneRenderingAllocator> DynamicPrimitiveTaskPipes;
 
 	} DynamicMeshElements;
 
@@ -471,14 +579,20 @@ private:
 		FGraphEventRef FrustumCullLegacyTask;
 		FGraphEventRef ComputeRelevanceLegacyTask;
 
+		FGraphEventRef DynamicMeshElementsPipe;
+		FGraphEventRef DynamicMeshElementsRenderThread;
+
 		UE::Tasks::FTaskEvent LightVisibility{ UE_SOURCE_LOCATION };
 		UE::Tasks::FTaskEvent BeginInitVisibility{ UE_SOURCE_LOCATION };
 		UE::Tasks::FTaskEvent FrustumCull{ UE_SOURCE_LOCATION };
 		UE::Tasks::FTaskEvent OcclusionCull{ UE_SOURCE_LOCATION };
 		UE::Tasks::FTaskEvent ComputeRelevance{ UE_SOURCE_LOCATION };
+		UE::Tasks::FTaskEvent DynamicMeshElementsPrerequisites{ UE_SOURCE_LOCATION };
+		UE::Tasks::FTaskEvent DynamicMeshElements{ UE_SOURCE_LOCATION };
 		UE::Tasks::FTask FinalizeRelevance;
 		UE::Tasks::FTask MeshPassSetup;
 
+		bool bDynamicMeshElementsPrerequisitesTriggered = false;
 		bool bWaitingAllowed = false;
 
 	} Tasks;
diff --git a/Engine/Source/Runtime/Renderer/Private/ShadowDepthRendering.cpp b/Engine/Source/Runtime/Renderer/Private/ShadowDepthRendering.cpp
--- a/Engine/Source/Runtime/Renderer/Private/ShadowDepthRendering.cpp
+++ b/Engine/Source/Runtime/Renderer/Private/ShadowDepthRendering.cpp
@@ -1700,7 +1700,7 @@ void FSceneRenderer::RenderVirtualShadowMaps(FRDGBuilder& GraphBuilder, bool bNa
 	}
 }
 
-void FSceneRenderer::RenderShadowDepthMaps(FRDGBuilder& GraphBuilder, FInstanceCullingManager &InstanceCullingManager, FRDGExternalAccessQueue& ExternalAccessQueue)
+void FSceneRenderer::RenderShadowDepthMaps(FRDGBuilder& GraphBuilder, FDynamicShadowsTaskData* DynamicShadowsTaskData, FInstanceCullingManager &InstanceCullingManager, FRDGExternalAccessQueue& ExternalAccessQueue)
 {
 	ensureMsgf(!bShadowDepthRenderCompleted, TEXT("RenderShadowDepthMaps called twice in the same frame"));
 
@@ -1712,6 +1712,11 @@ void FSceneRenderer::RenderShadowDepthMaps(FRDGBuilder& GraphBuilder, FInstanceC
 	RDG_EVENT_SCOPE(GraphBuilder, "ShadowDepths");
 	RDG_GPU_STAT_SCOPE(GraphBuilder, ShadowDepths);
 
+	if (DynamicShadowsTaskData)
+	{
+		FinishDynamicShadowMeshPassSetup(GraphBuilder, DynamicShadowsTaskData);
+	}
+
 	// Begin new deferred culling batching scope to catch shadow render passes, as there can use dynamic primitives that have not been uploaded before 
 	// the previous batching scope. Also flushes the culling views registered during the setup (in InitViewsAfterPrepass) that are referenced in the shadow view
 	// culling.
diff --git a/Engine/Source/Runtime/Renderer/Private/ShadowRendering.h b/Engine/Source/Runtime/Renderer/Private/ShadowRendering.h
--- a/Engine/Source/Runtime/Renderer/Private/ShadowRendering.h
+++ b/Engine/Source/Runtime/Renderer/Private/ShadowRendering.h
@@ -693,8 +693,20 @@ public:
 	 */
 	void AddReceiverPrimitive(FPrimitiveSceneInfo* PrimitiveSceneInfo);
 
+	enum class EGatherDynamicMeshElementsPass : uint8
+	{
+		// Processes all operations in a single pass.
+		All,
+
+		// Parallel pass runs first and processes elements in parallel with other shadows.
+		Parallel,
+
+		// Serial pass runs second and processes elements serially other shadows.
+		Serial,
+	};
+
 	/** Gathers dynamic mesh elements for all the shadow's primitives arrays. */
-	void GatherDynamicMeshElements(FMeshElementCollector& MeshCollector, FSceneRenderer& Renderer, class FVisibleLightInfo& VisibleLightInfo, TArray<const FSceneView*>& ReusedViewsArray, FInstanceCullingManager& InstanceCullingManager);
+	bool GatherDynamicMeshElements(FMeshElementCollector& MeshCollector, FSceneRenderer& Renderer, class FVisibleLightInfo& VisibleLightInfo, TArray<const FSceneView*>& ReusedViewsArray, EGatherDynamicMeshElementsPass Pass);
 
 	void SetupMeshDrawCommandsForShadowDepth(FSceneRenderer& Renderer, FInstanceCullingManager& InstanceCullingManager);
 
@@ -931,13 +943,14 @@ private:
 	int32 UpdateShadowCastingObjectBuffers() const;
 
 	/** Gathers dynamic mesh elements for the given primitive array. */
-	void GatherDynamicMeshElementsArray(
+	bool GatherDynamicMeshElementsArray(
 		FMeshElementCollector& Collector,
 		const PrimitiveArrayType& PrimitiveArray, 
 		const TArray<const FSceneView*>& Views,
 		const FSceneViewFamily& ViewFamily,
 		TArray<FMeshBatchAndRelevance,SceneRenderingAllocator>& OutDynamicMeshElements,
-		int32& OutNumDynamicSubjectMeshElements);
+		int32& OutNumDynamicSubjectMeshElements,
+		EGatherDynamicMeshElementsPass Pass);
 
 	void SetupFrustumForProjection(const FViewInfo* View, TArray<FVector4f, TInlineAllocator<8>>& OutFrustumVertices, bool& bOutCameraInsideShadowFrustum, FPlane* OutPlanes) const;
 
diff --git a/Engine/Source/Runtime/Renderer/Private/ShadowSetup.cpp b/Engine/Source/Runtime/Renderer/Private/ShadowSetup.cpp
--- a/Engine/Source/Runtime/Renderer/Private/ShadowSetup.cpp
+++ b/Engine/Source/Runtime/Renderer/Private/ShadowSetup.cpp
@@ -267,6 +267,24 @@ static FAutoConsoleVariableRef CVarParallelInitDynamicShadows(
 	ECVF_RenderThreadSafe
 );
 
+int32 GNumShadowDynamicMeshElementTasks = 4;
+static FAutoConsoleVariableRef CVarNumShadowDynamicMeshElementTasks(
+	TEXT("r.Visibility.DynamicMeshElements.NumShadowViewTasks"),
+	GNumShadowDynamicMeshElementTasks,
+	TEXT("Controls the number of gather dynamic mesh elements tasks to run asynchronously during shadow visibility."),
+	ECVF_RenderThreadSafe
+);
+
+uint32 GetNumShadowDynamicMeshElementTasks()
+{
+	if (!IsParallelGatherDynamicMeshElementsEnabled())
+	{
+		return 0;
+	}
+
+	return FMath::Clamp<int32>(GNumShadowDynamicMeshElementTasks, 0, LowLevelTasks::FScheduler::Get().GetNumWorkers());
+}
+
 static TAutoConsoleVariable<int32> CVarParallelGatherNumPrimitivesPerPacket(
 	TEXT("r.ParallelGatherNumPrimitivesPerPacket"),
 	256,  
@@ -1260,24 +1278,28 @@ public:
 	void Finish()
 	{
 		Collector.Finish();
-
 		DynamicVertexBuffer.Commit();
 		DynamicIndexBuffer.Commit();
 	}
 
+	UE::Tasks::FPipe& GetPipe()
+	{
+		return Pipe;
+	}
+
 private:
 	FShadowMeshCollector(FRHICommandList& RHICmdList, FSceneRenderer& SceneRenderer)
 		: Collector(SceneRenderer.FeatureLevel, SceneRenderer.Allocator)
 	{
 		DynamicVertexBuffer.Init(RHICmdList);
 		DynamicIndexBuffer.Init(RHICmdList);
-
 		Collector.Start(RHICmdList, DynamicVertexBuffer, DynamicIndexBuffer, SceneRenderer.DynamicReadBufferForShadows);
 	}
 
 	FMeshElementCollector Collector;
 	FGlobalDynamicVertexBuffer DynamicVertexBuffer;
 	FGlobalDynamicIndexBuffer DynamicIndexBuffer;
+	UE::Tasks::FPipe Pipe{ UE_SOURCE_LOCATION };
 
 	template <typename T>
 	friend class TConcurrentLinearBulkObjectAllocator;
@@ -1290,15 +1312,19 @@ struct FDynamicShadowsTaskData
 	FSceneRenderer* SceneRenderer;
 	const FScene* Scene;
 	TArrayView<FViewInfo> Views;
+	FSceneRenderingBulkObjectAllocator& Allocator;
+	FInstanceCullingManager& InstanceCullingManager;
 	const ERHIFeatureLevel::Type FeatureLevel;
 	const EShaderPlatform ShaderPlatform;
 	const bool bStaticSceneOnly;
 	const bool bRunningEarly;
 	const bool bMultithreaded;
 	const bool bMultithreadedCreateAndFilterShadows;
+	const bool bMultithreadedGDME;
 
 	// Whether any light in the scene has a ray traced distance field shadow.
 	bool bHasRayTracedDistanceFieldShadows = false;
+	bool bFinishedMeshPassSetup = false;
 
 	// Generated from prepare task
 	TArray<FProjectedShadowInfo*, SceneRenderingAllocator> PreShadows;
@@ -1309,41 +1335,82 @@ struct FDynamicShadowsTaskData
 	TArray<struct FGatherShadowPrimitivesPacket*, SceneRenderingAllocator> Packets;
 	FPerShadowGatherStats GatherStats;
 	FFilteredShadowArrays ShadowArrays;
+	UE::Tasks::FTaskEvent FilterDynamicShadowsTask{ UE_SOURCE_LOCATION };
+
+	// Gather Dynamic Mesh Elements state
+	TArray<FShadowMeshCollector*, TInlineAllocator<1, SceneRenderingAllocator>> MeshCollectors;
+	TArray<FRHICommandListImmediate::FQueuedCommandList, FConcurrentLinearArrayAllocator> CommandLists;
+	TArray<FProjectedShadowInfo*, SceneRenderingAllocator> ShadowsToGather;
+	TBitArray<SceneRenderingBitArrayAllocator> ShadowsToGatherInSerialPass;
+	UE::Tasks::FTaskEvent BeginGatherDynamicMeshElementsTask{ UE_SOURCE_LOCATION };
+	UE::Tasks::FTaskEvent SetupMeshPassTask{ UE_SOURCE_LOCATION };
+	std::atomic_int32_t ShadowsToGatherNextIndex = { 0 };
+
+	// Command list used to allocate shadow depth render targets.
+	FRHICommandList* RHICmdListForAllocateTargets = nullptr;
 
 	// Used by RenderThread
 	FGraphEventRef TaskEvent;
-	mutable FGraphEventRef CreateDynamicShadowsTaskEvent;
 
-	void WaitForCreateDynamicShadowsTask() const
+	FShadowMeshCollector& GetSerialMeshCollector() const
 	{
-		if (CreateDynamicShadowsTaskEvent)
-		{
-			check(IsInRenderingThread());
-			CreateDynamicShadowsTaskEvent->Wait(ENamedThreads::GetRenderThread_Local());
-			CreateDynamicShadowsTaskEvent = nullptr;
-		}
+		return *MeshCollectors.Last();
 	}
 
 	FDynamicShadowsTaskData(const FDynamicShadowsTaskData&) = delete;
 
-	FDynamicShadowsTaskData(FSceneRenderer* InSceneRenderer, bool bInRunningEarly)
+	FDynamicShadowsTaskData(FRHICommandListImmediate& InRHICmdList, FSceneRenderer* InSceneRenderer, FInstanceCullingManager& InInstanceCullingManager, bool bInRunningEarly)
 		: SceneRenderer(InSceneRenderer)
 		, Scene(InSceneRenderer->Scene)
 		, Views(InSceneRenderer->Views)
+		, Allocator(SceneRenderer->Allocator)
+		, InstanceCullingManager(InInstanceCullingManager)
 		, FeatureLevel(Scene->GetFeatureLevel())
 		, ShaderPlatform(GShaderPlatformForFeatureLevel[FeatureLevel])
 		, bStaticSceneOnly(AreAnyViewsStaticSceneOnly(Views))
 		, bRunningEarly(bInRunningEarly)
 		, bMultithreaded((FApp::ShouldUseThreadingForPerformance() || FForkProcessHelper::IsForkedMultithreadInstance()) && CVarParallelGatherShadowPrimitives.GetValueOnRenderThread() > 0)
 		, bMultithreadedCreateAndFilterShadows(bRunningEarly && bMultithreaded && GRHISupportsAsyncGetRenderQueryResult && GParallelInitDynamicShadows > 0)
+		, bMultithreadedGDME(bMultithreadedCreateAndFilterShadows && GetNumShadowDynamicMeshElementTasks() > 0)
 	{
 		if (bMultithreadedCreateAndFilterShadows)
 		{
-			CreateDynamicShadowsTaskEvent = FGraphEvent::CreateGraphEvent();
+			RHICmdListForAllocateTargets = new FRHICommandList(FRHIGPUMask::All());
+			RHICmdListForAllocateTargets->SwitchPipeline(ERHIPipeline::Graphics);
+			CommandLists.Emplace(RHICmdListForAllocateTargets);
+		}
+		else
+		{
+			RHICmdListForAllocateTargets = &InRHICmdList;
+		}
+
+		if (bMultithreadedGDME)
+		{
+			MeshCollectors.Reserve(GetNumShadowDynamicMeshElementTasks());
+		}
+		else
+		{
+			MeshCollectors.Emplace(Allocator.Create<FShadowMeshCollector>(InRHICmdList, *SceneRenderer));
+			BeginGatherDynamicMeshElementsTask.Trigger();
 		}
 	}
+
+	~FDynamicShadowsTaskData()
+	{
+		check(bFinishedMeshPassSetup);
+		check(CommandLists.IsEmpty());
+		check(MeshCollectors.IsEmpty());
+	}
 };
 
+void BeginShadowGatherDynamicMeshElements(FDynamicShadowsTaskData* TaskData)
+{
+	if (TaskData)
+	{
+		TaskData->BeginGatherDynamicMeshElementsTask.Trigger();
+	}
+}
+
 TConstArrayView<FProjectedShadowInfo*> GetProjectedDistanceFieldShadows(const FDynamicShadowsTaskData* TaskData)
 {
 	if (TaskData)
@@ -2425,6 +2492,7 @@ void FProjectedShadowInfo::SetupMeshDrawCommandsForShadowDepth(FSceneRenderer& R
 		NumSubjectMeshCommandBuildRequestElements * InstanceFactor,
 		ShadowDepthPassVisibleCommands);
 
+	UE::TScopeLock ScopeLock(Renderer.DispatchedShadowDepthPassesMutex);
 	Renderer.DispatchedShadowDepthPasses.Add(&ShadowDepthPass);
 }
 
@@ -2552,20 +2620,22 @@ void FProjectedShadowInfo::SetupMeshDrawCommandsForProjectionStenciling(FSceneRe
 	}
 }
 
-void FProjectedShadowInfo::GatherDynamicMeshElements(FMeshElementCollector& MeshCollector, FSceneRenderer& Renderer, FVisibleLightInfo& VisibleLightInfo, TArray<const FSceneView*>& ReusedViewsArray, FInstanceCullingManager& InstanceCullingManager)
+bool FProjectedShadowInfo::GatherDynamicMeshElements(
+	FMeshElementCollector& MeshCollector,
+	FSceneRenderer& Renderer,
+	FVisibleLightInfo& VisibleLightInfo,
+	TArray<const FSceneView*>& ReusedViewsArray,
+	EGatherDynamicMeshElementsPass Pass)
 {
 	QUICK_SCOPE_CYCLE_COUNTER(STAT_Shadow_GatherDynamicMeshElements);
+	TRACE_CPUPROFILER_EVENT_SCOPE(FProjectedShadowInfo::GatherDynamicMeshElements);
 
 	check(ShadowDepthView);
 
+	bool bProcessedAllPrimitives = true;
+
 	if (DynamicSubjectPrimitives.Num() > 0 || ReceiverPrimitives.Num() > 0 || SubjectTranslucentPrimitives.Num() > 0)
 	{
-		// Backup properties of the view that we will override
-		FMatrix OriginalViewMatrix = ShadowDepthView->ViewMatrices.GetViewMatrix();
-
-		// Override the view matrix so that billboarding primitives will be aligned to the light
-		ShadowDepthView->ViewMatrices.HackOverrideViewMatrixForShadows(TranslatedWorldToView);
-
 		ReusedViewsArray[0] = ShadowDepthView;
 
 		if (bPreShadow && GPreshadowsForceLowestLOD)
@@ -2582,36 +2652,34 @@ void FProjectedShadowInfo::GatherDynamicMeshElements(FMeshElementCollector& Mesh
 		{
 			ShadowDepthView->SetPreShadowTranslation(FVector(0, 0, 0));
 			ShadowDepthView->SetDynamicMeshElementsShadowCullFrustum(&CascadeSettings.ShadowBoundsAccurate);
-			GatherDynamicMeshElementsArray(MeshCollector, DynamicSubjectPrimitives, ReusedViewsArray, Renderer.ViewFamily, DynamicSubjectMeshElements, NumDynamicSubjectMeshElements);
+			bProcessedAllPrimitives &= GatherDynamicMeshElementsArray(MeshCollector, DynamicSubjectPrimitives, ReusedViewsArray, Renderer.ViewFamily, DynamicSubjectMeshElements, NumDynamicSubjectMeshElements, Pass);
 			ShadowDepthView->SetPreShadowTranslation(PreShadowTranslation);
 		}
 		else
 		{
 			ShadowDepthView->SetPreShadowTranslation(PreShadowTranslation);
 			ShadowDepthView->SetDynamicMeshElementsShadowCullFrustum(&CasterOuterFrustum);
-			GatherDynamicMeshElementsArray(MeshCollector, DynamicSubjectPrimitives, ReusedViewsArray, Renderer.ViewFamily, DynamicSubjectMeshElements, NumDynamicSubjectMeshElements);
+			bProcessedAllPrimitives &= GatherDynamicMeshElementsArray(MeshCollector, DynamicSubjectPrimitives, ReusedViewsArray, Renderer.ViewFamily, DynamicSubjectMeshElements, NumDynamicSubjectMeshElements, Pass);
 		}
 
 		ShadowDepthView->DrawDynamicFlags = EDrawDynamicFlags::None;
 
 		int32 NumDynamicSubjectTranslucentMeshElements = 0;
 		ShadowDepthView->SetDynamicMeshElementsShadowCullFrustum(&CasterOuterFrustum);
-		GatherDynamicMeshElementsArray(MeshCollector, SubjectTranslucentPrimitives, ReusedViewsArray, Renderer.ViewFamily, DynamicSubjectTranslucentMeshElements, NumDynamicSubjectTranslucentMeshElements);
+		bProcessedAllPrimitives &= GatherDynamicMeshElementsArray(MeshCollector, SubjectTranslucentPrimitives, ReusedViewsArray, Renderer.ViewFamily, DynamicSubjectTranslucentMeshElements, NumDynamicSubjectTranslucentMeshElements, Pass);
 	}
 
-	MeshCollector.Commit();
-
-	SetupMeshDrawCommandsForProjectionStenciling(Renderer, InstanceCullingManager);
-	SetupMeshDrawCommandsForShadowDepth(Renderer, InstanceCullingManager);
+	return bProcessedAllPrimitives;
 }
 
-void FProjectedShadowInfo::GatherDynamicMeshElementsArray(
+bool FProjectedShadowInfo::GatherDynamicMeshElementsArray(
 	FMeshElementCollector& MeshCollector,
 	const PrimitiveArrayType& Primitives,
 	const TArray<const FSceneView*>& Views,
 	const FSceneViewFamily& ViewFamily,
 	TArray<FMeshBatchAndRelevance,SceneRenderingAllocator>& OutDynamicMeshElements,
-	int32& OutNumDynamicSubjectMeshElements)
+	int32& OutNumDynamicSubjectMeshElements,
+	EGatherDynamicMeshElementsPass Pass)
 {
 	// Simple elements not supported in shadow passes
 	FSimpleElementCollector DynamicSubjectSimpleElements;
@@ -2624,6 +2692,8 @@ void FProjectedShadowInfo::GatherDynamicMeshElementsArray(
 
 	const uint32 PrimitiveCount = Primitives.Num();
 
+	bool bProcessedAllPrimitives = true;
+
 	for (uint32 PrimitiveIndex = 0; PrimitiveIndex < PrimitiveCount; ++PrimitiveIndex)
 	{
 		TRACE_CPUPROFILER_EVENT_SCOPE(DynamicPrimitive);
@@ -2642,14 +2712,27 @@ void FProjectedShadowInfo::GatherDynamicMeshElementsArray(
 		// Only draw if the subject primitive is shadow relevant.
 		if (ViewRelevance.bShadowRelevance && ViewRelevance.bDynamicRelevance)
 		{
+			if (Pass != EGatherDynamicMeshElementsPass::All)
+			{
+				const bool bSerialPrimitivesOnly = Pass == EGatherDynamicMeshElementsPass::Serial;
+
+				// Only process primitives that match the correct pass.
+				if (bSerialPrimitivesOnly == PrimitiveSceneProxy->SupportsParallelGDME())
+				{
+					bProcessedAllPrimitives = false;
+					continue;
+				}
+			}
+
 			MeshCollector.SetPrimitive(PrimitiveSceneProxy, PrimitiveSceneInfo->DefaultDynamicHitProxyId);
 
 			PrimitiveSceneProxy->GetDynamicMeshElements(Views, ViewFamily, 0x1, MeshCollector);
 		}
 	}
 
-	OutNumDynamicSubjectMeshElements = MeshCollector.GetMeshElementCount(0);
+	OutNumDynamicSubjectMeshElements += MeshCollector.GetMeshElementCount(0);
 	MeshCollector.ClearViewMeshArrays();
+	return bProcessedAllPrimitives;
 }
 
 /** 
@@ -4400,70 +4483,165 @@ void FSceneRenderer::DrawDebugShadowFrustum(FViewInfo& View, FProjectedShadowInf
 	}
 }
 
-void FSceneRenderer::GatherShadowDynamicMeshElements(FRHICommandList& RHICmdList, FInstanceCullingManager& InstanceCullingManager)
+void FSceneRenderer::GatherShadowDynamicMeshElements(FDynamicShadowsTaskData& TaskData)
 {
-	FShadowMeshCollector* ShadowMeshCollector = FShadowMeshCollector::Create(RHICmdList, *this);
+	TConstArrayView<FProjectedShadowInfo*> ShadowsToSetupViews = TaskData.ShadowArrays.ShadowsToSetupViews;
+	const int32 SetupShadowDepthViewBatchSize = 8;
 
-	TArray<const FSceneView*> ReusedViewsArray;
-	ReusedViewsArray.AddZeroed(1);
+	ParallelForTemplate(
+		TEXT("SetupShadowDepthView"), ShadowsToSetupViews.Num(), SetupShadowDepthViewBatchSize,
+		[this, ShadowsToSetupViews](int32 Index)
+		{
+			FOptionalTaskTagScope Scope(ETaskTag::EParallelRenderingThread);
+			ShadowsToSetupViews[Index]->SetupShadowDepthView(this);
+		},
+		EParallelForFlags::None
+	);
+
+	TaskData.ShadowsToGather.Reserve(128);
 
 	for (int32 AtlasIndex = 0; AtlasIndex < SortedShadowsForShadowDepthPass.ShadowMapAtlases.Num(); AtlasIndex++)
 	{
-		FSortedShadowMapAtlas& Atlas = SortedShadowsForShadowDepthPass.ShadowMapAtlases[AtlasIndex];
-
-		for (int32 ShadowIndex = 0; ShadowIndex < Atlas.Shadows.Num(); ShadowIndex++)
-		{
-			FProjectedShadowInfo* ProjectedShadowInfo = Atlas.Shadows[ShadowIndex];
-			FVisibleLightInfo& VisibleLightInfo = VisibleLightInfos[ProjectedShadowInfo->GetLightSceneInfo().Id];
-			ProjectedShadowInfo->GatherDynamicMeshElements(ShadowMeshCollector->GetCollector(), *this, VisibleLightInfo, ReusedViewsArray, InstanceCullingManager);
-		}
+		TaskData.ShadowsToGather.Append(SortedShadowsForShadowDepthPass.ShadowMapAtlases[AtlasIndex].Shadows);
 	}
 
 	for (int32 AtlasIndex = 0; AtlasIndex < SortedShadowsForShadowDepthPass.ShadowMapCubemaps.Num(); AtlasIndex++)
 	{
-		FSortedShadowMapAtlas& Atlas = SortedShadowsForShadowDepthPass.ShadowMapCubemaps[AtlasIndex];
+		TaskData.ShadowsToGather.Append(SortedShadowsForShadowDepthPass.ShadowMapCubemaps[AtlasIndex].Shadows);
+	}
+
+	TaskData.ShadowsToGather.Append(SortedShadowsForShadowDepthPass.PreshadowCache.Shadows);
+
+	for (FSortedShadowMapAtlas& Atlas : SortedShadowsForShadowDepthPass.TranslucencyShadowMapAtlases)
+	{
+		TaskData.ShadowsToGather.Append(Atlas.Shadows);
+	}
 
-		for (int32 ShadowIndex = 0; ShadowIndex < Atlas.Shadows.Num(); ShadowIndex++)
+	if (UseNonNaniteVirtualShadowMaps(ShaderPlatform, FeatureLevel))
+	{
+		// GPUCULL_TODO: Replace with new shadow culling processor thing
+		for (FProjectedShadowInfo* ProjectedShadowInfo : SortedShadowsForShadowDepthPass.VirtualShadowMapShadows)
 		{
-			FProjectedShadowInfo* ProjectedShadowInfo = Atlas.Shadows[ShadowIndex];
-			FVisibleLightInfo& VisibleLightInfo = VisibleLightInfos[ProjectedShadowInfo->GetLightSceneInfo().Id];
-			ProjectedShadowInfo->GatherDynamicMeshElements(ShadowMeshCollector->GetCollector(), *this, VisibleLightInfo, ReusedViewsArray, InstanceCullingManager);
+			if (ProjectedShadowInfo->bShouldRenderVSM)
+			{
+				TaskData.ShadowsToGather.Emplace(ProjectedShadowInfo);
+			}
 		}
 	}
 
-	for (int32 ShadowIndex = 0; ShadowIndex < SortedShadowsForShadowDepthPass.PreshadowCache.Shadows.Num(); ShadowIndex++)
+	// Process all shadows in the serial path when not in multithreaded mode.
+	TaskData.ShadowsToGatherInSerialPass.Init(!TaskData.bMultithreadedGDME, TaskData.ShadowsToGather.Num());
+
+	const UE::Tasks::ETaskPriority TaskPriority = UE::Tasks::ETaskPriority::High;
+
+	const UE::Tasks::EExtendedTaskPriority ExtendedTaskPriority = TaskData.bMultithreadedGDME
+		? UE::Tasks::EExtendedTaskPriority::None
+		: UE::Tasks::EExtendedTaskPriority::Inline;
+
+	UE::Tasks::FTask MeshCollectorsTask;
+
+	int32 NumShadowViews = 0;
+	for (FProjectedShadowInfo* ProjectedShadowInfo : TaskData.ShadowsToGather)
 	{
-		FProjectedShadowInfo* ProjectedShadowInfo = SortedShadowsForShadowDepthPass.PreshadowCache.Shadows[ShadowIndex];
-		FVisibleLightInfo& VisibleLightInfo = VisibleLightInfos[ProjectedShadowInfo->GetLightSceneInfo().Id];
-		ProjectedShadowInfo->GatherDynamicMeshElements(ShadowMeshCollector->GetCollector(), *this, VisibleLightInfo, ReusedViewsArray, InstanceCullingManager);
+		NumShadowViews += ProjectedShadowInfo->bOnePassPointLightShadow ? 6 : 1;
 	}
+	TaskData.InstanceCullingManager.AllocateViews(NumShadowViews);
 
-	for (int32 AtlasIndex = 0; AtlasIndex < SortedShadowsForShadowDepthPass.TranslucencyShadowMapAtlases.Num(); AtlasIndex++)
+	if (TaskData.bMultithreadedGDME)
 	{
-		FSortedShadowMapAtlas& Atlas = SortedShadowsForShadowDepthPass.TranslucencyShadowMapAtlases[AtlasIndex];
+		UE::Tasks::FTaskEvent MeshCollectorsTaskEvent{ UE_SOURCE_LOCATION };
+
+		const int32 NumTasks = FMath::Min<int32>(GetNumShadowDynamicMeshElementTasks(), TaskData.ShadowsToGather.Num());
 
-		for (int32 ShadowIndex = 0; ShadowIndex < Atlas.Shadows.Num(); ShadowIndex++)
+		for (int32 TaskIndex = 0; TaskIndex < NumTasks; ++TaskIndex)
 		{
-			FProjectedShadowInfo* ProjectedShadowInfo = Atlas.Shadows[ShadowIndex];
-			FVisibleLightInfo& VisibleLightInfo = VisibleLightInfos[ProjectedShadowInfo->GetLightSceneInfo().Id];
-			ProjectedShadowInfo->GatherDynamicMeshElements(ShadowMeshCollector->GetCollector(), *this, VisibleLightInfo, ReusedViewsArray, InstanceCullingManager);
+			FRHICommandList* RHICmdList = new FRHICommandList(FRHIGPUMask::All());
+			RHICmdList->SwitchPipeline(ERHIPipeline::Graphics);
+			TaskData.CommandLists.Emplace(RHICmdList);
+
+			FShadowMeshCollector* MeshCollector = Allocator.Create<FShadowMeshCollector>(*RHICmdList, *this);
+			TaskData.MeshCollectors.Emplace(MeshCollector);
+
+			MeshCollectorsTaskEvent.AddPrerequisites(
+				MeshCollector->GetPipe().Launch(UE_SOURCE_LOCATION,
+					[this, &TaskData, MeshCollector]() mutable
+			{
+				FOptionalTaskTagScope Scope(ETaskTag::EParallelRenderingThread);
+				TArray<const FSceneView*> LocalViews;
+				LocalViews.AddZeroed(1);
+
+				while (true)
+				{
+					// Atomically increment to get the next shadow to process in this task until empty.
+					const int32 ShadowIndex = TaskData.ShadowsToGatherNextIndex.fetch_add(1, std::memory_order_relaxed);
+					if (ShadowIndex >= TaskData.ShadowsToGather.Num())
+					{
+						break;
+					}
+
+					FProjectedShadowInfo& ProjectedShadowInfo = *TaskData.ShadowsToGather[ShadowIndex];
+					FVisibleLightInfo& VisibleLightInfo = VisibleLightInfos[ProjectedShadowInfo.GetLightSceneInfo().Id];
+
+					const bool bProcessedAllPrimitives = ProjectedShadowInfo.GatherDynamicMeshElements(MeshCollector->GetCollector(), *this, VisibleLightInfo, LocalViews, FProjectedShadowInfo::EGatherDynamicMeshElementsPass::Parallel);
+
+					// Atomically mark shadows that require a second serial pass.
+					if (!bProcessedAllPrimitives)
+					{
+						TaskData.ShadowsToGatherInSerialPass[ShadowIndex].AtomicSet(true);
+					}
+
+					// Setup mesh draw command passes immediately if we don't have any more primitives to process in the serial pass.
+					if (bProcessedAllPrimitives)
+					{
+						ProjectedShadowInfo.SetupMeshDrawCommandsForProjectionStenciling(*this, TaskData.InstanceCullingManager);
+						ProjectedShadowInfo.SetupMeshDrawCommandsForShadowDepth(*this, TaskData.InstanceCullingManager);
+					}
+				}
+
+			}, TaskData.BeginGatherDynamicMeshElementsTask, TaskPriority));
 		}
+
+		MeshCollectorsTaskEvent.Trigger();
+		MeshCollectorsTask = MoveTemp(MeshCollectorsTaskEvent);
 	}
 
-	if (UseNonNaniteVirtualShadowMaps(ShaderPlatform, FeatureLevel))
+	// Process only serial primitives if in multithreaded mode, otherwise process all of them.
+	const FProjectedShadowInfo::EGatherDynamicMeshElementsPass SerialPass = TaskData.bMultithreadedGDME
+		? FProjectedShadowInfo::EGatherDynamicMeshElementsPass::Serial
+		: FProjectedShadowInfo::EGatherDynamicMeshElementsPass::All;
+
+	TaskData.SetupMeshPassTask.AddPrerequisites(UE::Tasks::Launch(UE_SOURCE_LOCATION, [this, &TaskData, SerialPass] () mutable
 	{
-		// GPUCULL_TODO: Replace with new shadow culling processor thing
-		for (FProjectedShadowInfo* ProjectedShadowInfo : SortedShadowsForShadowDepthPass.VirtualShadowMapShadows)
+		FOptionalTaskTagScope Scope(ETaskTag::EParallelRenderingThread);
+		TRACE_CPUPROFILER_EVENT_SCOPE(SetupShadowMeshPass);
+
+		TArray<const FSceneView*> LocalViews;
+
+		for (TConstSetBitIterator<SceneRenderingBitArrayAllocator> BitIt(TaskData.ShadowsToGatherInSerialPass); BitIt; ++BitIt)
 		{
-			if (ProjectedShadowInfo->bShouldRenderVSM)
+			if (LocalViews.IsEmpty())
 			{
-				FVisibleLightInfo& VisibleLightInfo = VisibleLightInfos[ProjectedShadowInfo->GetLightSceneInfo().Id];
-				ProjectedShadowInfo->GatherDynamicMeshElements(ShadowMeshCollector->GetCollector(), *this, VisibleLightInfo, ReusedViewsArray, InstanceCullingManager);
+				LocalViews.AddZeroed(1);
 			}
+
+			FShadowMeshCollector& MeshCollector = TaskData.GetSerialMeshCollector();
+			FProjectedShadowInfo* ProjectedShadowInfo = TaskData.ShadowsToGather[BitIt.GetIndex()];
+			FVisibleLightInfo& VisibleLightInfo = VisibleLightInfos[ProjectedShadowInfo->GetLightSceneInfo().Id];
+			ProjectedShadowInfo->GatherDynamicMeshElements(MeshCollector.GetCollector(), *this, VisibleLightInfo, LocalViews, SerialPass);
+
+			// Setup mesh draw command passes now that all primitives are processed.
+			ProjectedShadowInfo->SetupMeshDrawCommandsForProjectionStenciling(*this, TaskData.InstanceCullingManager);
+			ProjectedShadowInfo->SetupMeshDrawCommandsForShadowDepth(*this, TaskData.InstanceCullingManager);
 		}
-	}
 
-	ShadowMeshCollector->Finish();
+		for (FShadowMeshCollector* MeshCollector : TaskData.MeshCollectors)
+		{
+			MeshCollector->Finish();
+		}
+	
+	}, MeshCollectorsTask, TaskPriority, ExtendedTaskPriority));
+
+	TaskData.SetupMeshPassTask.Trigger();
 }
 
 FAutoConsoleTaskPriority CPrio_GatherShadowPrimitives(
@@ -4811,7 +4989,6 @@ struct FGatherShadowPrimitivesPrepareTask
 		if (TaskData.bMultithreadedCreateAndFilterShadows)
 		{
 			TaskData.SceneRenderer->CreateDynamicShadows(TaskData);
-			TaskData.CreateDynamicShadowsTaskEvent->DispatchSubsequents();
 		}
 
 		AnyThreadTask();
@@ -4873,6 +5050,22 @@ struct FGatherShadowPrimitivesPrepareTask
 				CPrio_GatherShadowPrimitives.Get());
 
 			MyCompletionGraphEvent->DontCompleteUntil(FinalizeTask);
+
+			if (TaskData.bMultithreadedGDME)
+			{
+				Prereqs.Reset();
+				Prereqs.Emplace(FinalizeTask);
+
+				FFunctionGraphTask::CreateAndDispatchWhenReady([TaskData = &TaskData]
+					{
+						SCOPED_NAMED_EVENT_TEXT("GatherShadowDynamicMeshElements", FColor::Green);
+						FOptionalTaskTagScope Scope(ETaskTag::EParallelRenderingThread);
+						TaskData->SceneRenderer->GatherShadowDynamicMeshElements(*TaskData);
+					},
+					TStatId(),
+					&Prereqs,
+					CPrio_GatherShadowPrimitives.Get());
+			}
 		}
 	}
 
@@ -5039,6 +5232,13 @@ void FSceneRenderer::FinishGatherShadowPrimitives(FDynamicShadowsTaskData* TaskD
 		FilterDynamicShadows(*TaskData);
 	}
 
+	TaskData->FilterDynamicShadowsTask.Wait();
+
+	if (!TaskData->bMultithreadedGDME)
+	{
+		GatherShadowDynamicMeshElements(*TaskData);
+	}
+
 	for (FDrawDebugShadowFrustumOp Op : TaskData->DrawDebugShadowFrustumOps)
 	{
 		DrawDebugShadowFrustum(*Op.View, *Op.ProjectedShadowInfo);
@@ -5326,7 +5526,7 @@ void FSceneRenderer::AddViewDependentWholeSceneShadowsForView(
 	}
 }
 
-void FSceneRenderer::AllocateShadowDepthTargets(FRHICommandListImmediate& RHICmdList, const FDynamicShadowsTaskData& TaskData)
+void FSceneRenderer::AllocateShadowDepthTargets(FDynamicShadowsTaskData& TaskData)
 {
 	SCOPED_NAMED_EVENT_TEXT("FSceneRenderer::AllocateShadowDepthTargets", FColor::Magenta);
 
@@ -5334,6 +5534,8 @@ void FSceneRenderer::AllocateShadowDepthTargets(FRHICommandListImmediate& RHICmd
 
 	const FFilteredShadowArrays& ShadowArrays = TaskData.ShadowArrays;
 
+	FRHICommandList& RHICmdList = *TaskData.RHICmdListForAllocateTargets;
+
 	if (ShadowArrays.CachedPreShadows.Num() > 0)
 	{
 		if (!Scene->PreShadowCacheDepthZ)
@@ -6217,28 +6419,19 @@ void FSceneRenderer::FilterDynamicShadows(FDynamicShadowsTaskData& TaskData)
 
 	// Sort the projected shadows by resolution.
 	ShadowArrays.TranslucentShadows.Sort(FCompareFProjectedShadowInfoByResolution());
-	
-	TConstArrayView<FProjectedShadowInfo*> ShadowsToSetupViews = TaskData.ShadowArrays.ShadowsToSetupViews;
-	const int32 SetupShadowDepthViewBatchSize = 8;
 
-	ParallelForTemplate(
-		TEXT("SetupShadowDepthView"), ShadowsToSetupViews.Num(), SetupShadowDepthViewBatchSize,
-		[this, ShadowsToSetupViews](int32 Index)
-		{
-			FOptionalTaskTagScope Scope(ETaskTag::EParallelRenderingThread);
-			ShadowsToSetupViews[Index]->SetupShadowDepthView(this);
-		},
-		TaskData.bMultithreadedCreateAndFilterShadows ? EParallelForFlags::None : EParallelForFlags::ForceSingleThread
-	);
+	AllocateShadowDepthTargets(TaskData);
+
+	TaskData.FilterDynamicShadowsTask.Trigger();
 }
 
-FDynamicShadowsTaskData* FSceneRenderer::BeginInitDynamicShadows(bool bRunningEarly, IVisibilityTaskData* VisibilityTaskData)
+FDynamicShadowsTaskData* FSceneRenderer::BeginInitDynamicShadows(FRDGBuilder& GraphBuilder, bool bRunningEarly, IVisibilityTaskData* VisibilityTaskData, FInstanceCullingManager& InstanceCullingManager)
 {
 	SCOPE_CYCLE_COUNTER(STAT_DynamicShadowSetupTime);
 	CSV_SCOPED_TIMING_STAT_EXCLUSIVE(InitViews_Shadows);
 	SCOPED_NAMED_EVENT_TEXT("FSceneRenderer::BeginInitDynamicShadows", FColor::Magenta);
 
-	FDynamicShadowsTaskData* DynamicShadowTaskData = Allocator.Create<FDynamicShadowsTaskData>(this, bRunningEarly);
+	FDynamicShadowsTaskData* DynamicShadowTaskData = Allocator.Create<FDynamicShadowsTaskData>(GraphBuilder.RHICmdList, this, InstanceCullingManager, bRunningEarly);
 
 	// Gathers the list of primitives used to draw various shadow types
 	BeginGatherShadowPrimitives(DynamicShadowTaskData, VisibilityTaskData);
@@ -6247,7 +6440,7 @@ FDynamicShadowsTaskData* FSceneRenderer::BeginInitDynamicShadows(bool bRunningEa
 }
 
 
-void FSceneRenderer::FinishInitDynamicShadows(FRDGBuilder& GraphBuilder, FDynamicShadowsTaskData* TaskData, FInstanceCullingManager& InstanceCullingManager)
+void FSceneRenderer::FinishInitDynamicShadows(FRDGBuilder& GraphBuilder, FDynamicShadowsTaskData* TaskData)
 {
 	SCOPED_NAMED_EVENT_TEXT("FSceneRenderer::FinishInitDynamicShadows", FColor::Magenta);
 	SCOPE_CYCLE_COUNTER(STAT_InitDynamicShadowsTime);
@@ -6263,11 +6456,25 @@ void FSceneRenderer::FinishInitDynamicShadows(FRDGBuilder& GraphBuilder, FDynami
 		ShadowSceneRenderer->DispatchVirtualShadowMapViewAndCullingSetup(GraphBuilder, SortedShadowsForShadowDepthPass.VirtualShadowMapShadows);
 		ShadowSceneRenderer->PostSetupDebugRender();
 	}
+}
+
+void FSceneRenderer::FinishDynamicShadowMeshPassSetup(FRDGBuilder& GraphBuilder, FDynamicShadowsTaskData* TaskData)
+{
+	// This can be called early in InitDynamicShadows.
+	if (TaskData->bFinishedMeshPassSetup)
+	{
+		return;
+	}
 
-	AllocateShadowDepthTargets(GraphBuilder.RHICmdList, *TaskData);
+	TaskData->SetupMeshPassTask.Wait();
+	TaskData->MeshCollectors.Empty();
 
-	// Generate mesh element arrays from shadow primitive arrays
-	GatherShadowDynamicMeshElements(GraphBuilder.RHICmdList, InstanceCullingManager);
+	for (FRHICommandListImmediate::FQueuedCommandList& QueueCmdList : TaskData->CommandLists)
+	{
+		QueueCmdList.CmdList->FinishRecording();
+	}
+	GraphBuilder.RHICmdList.QueueAsyncCommandListSubmit(TaskData->CommandLists);
+	TaskData->CommandLists.Empty();
 
 	// Ensure all shadow view dynamic primitives are uploaded before shadow-culling batching pass.
 	// TODO: automate this such that:
@@ -6306,12 +6513,17 @@ void FSceneRenderer::FinishInitDynamicShadows(FRDGBuilder& GraphBuilder, FDynami
 	{
 		Scene->GPUScene.UploadDynamicPrimitiveShaderDataForView(GraphBuilder, *Scene, *ProjectedShadowInfo->ShadowDepthView, true);
 	}
+
+	DynamicReadBufferForShadows.Commit(GraphBuilder.RHICmdList);
+
+	TaskData->bFinishedMeshPassSetup = true;
 }
 
 FDynamicShadowsTaskData* FSceneRenderer::InitDynamicShadows(FRDGBuilder& GraphBuilder, FInstanceCullingManager& InstanceCullingManager)
 {
-	FDynamicShadowsTaskData* TaskData = BeginInitDynamicShadows(false, nullptr);
-	FinishInitDynamicShadows(GraphBuilder, TaskData, InstanceCullingManager);
+	FDynamicShadowsTaskData* TaskData = BeginInitDynamicShadows(GraphBuilder, false, nullptr, InstanceCullingManager);
+	FinishInitDynamicShadows(GraphBuilder, TaskData);
+	FinishDynamicShadowMeshPassSetup(GraphBuilder, TaskData);
 	return TaskData;
 }
 
