GPU Occlusion Instance Culling : add ForceInstanceCulling cvar

diff --git a/Engine/Source/Runtime/Renderer/Private/InstanceCulling/InstanceCullingContext.cpp b/Engine/Source/Runtime/Renderer/Private/InstanceCulling/InstanceCullingContext.cpp
--- a/Engine/Source/Runtime/Renderer/Private/InstanceCulling/InstanceCullingContext.cpp
+++ b/Engine/Source/Runtime/Renderer/Private/InstanceCulling/InstanceCullingContext.cpp
@@ -31,6 +31,13 @@ static TAutoConsoleVariable<int32> CVarOcclusionCullInstances(
 	TEXT("Whether to do per instance occlusion culling for GPU instance culling."),
 	ECVF_RenderThreadSafe | ECVF_Preview);
 
+static int32 GOcclusionForceInstanceCulling = 0;
+static FAutoConsoleVariableRef CVarOcclusionForceInstanceCulling(
+	TEXT("r.InstanceCulling.ForceInstanceCulling"),
+	GOcclusionForceInstanceCulling,
+	TEXT("Whether to force per instance occlusion culling."),
+	ECVF_RenderThreadSafe);
+
 static int32 GInstanceCullingAllowOrderPreservation = 1;
 static FAutoConsoleVariableRef CVarInstanceCullingAllowOrderPreservation(
 	TEXT("r.InstanceCulling.AllowInstanceOrderPreservation"),
@@ -376,7 +383,7 @@ void FInstanceCullingContext::AddInstanceRunsToDrawCommand(uint32 IndirectArgsOf
 		uint32 RunStart = Runs[Index * 2];
 		uint32 RunEndIncl = Runs[Index * 2 + 1];
 		uint32 NumInstances = (RunEndIncl + 1U) - RunStart;
-		AddInstancesToDrawCommand(IndirectArgsOffset, InstanceDataOffset + RunStart, NumInstancesInRuns, NumInstances, InstanceFlags, MaxBatchSize);
+		AddInstancesToDrawCommand(IndirectArgsOffset, InstanceDataOffset + RunStart, NumInstancesInRuns, NumInstances, InstanceFlags | EInstanceFlags::ForceInstanceCulling, MaxBatchSize);
 		NumInstancesInRuns += NumInstances;
 	}
 }
@@ -1488,7 +1495,7 @@ void FInstanceCullingContext::SetupDrawCommands(
 		const bool bSupportsGPUSceneInstancing = EnumHasAnyFlags(VisibleMeshDrawCommand.Flags, EFVisibleMeshDrawCommandFlags::HasPrimitiveIdStreamIndex);
 		const bool bMaterialUsesWorldPositionOffset = EnumHasAnyFlags(VisibleMeshDrawCommand.Flags, EFVisibleMeshDrawCommandFlags::MaterialUsesWorldPositionOffset);
 		const bool bMaterialAlwaysEvaluatesWorldPositionOffset = EnumHasAnyFlags(VisibleMeshDrawCommand.Flags, EFVisibleMeshDrawCommandFlags::MaterialAlwaysEvaluatesWorldPositionOffset);
-		const bool bForceInstanceCulling = EnumHasAnyFlags(VisibleMeshDrawCommand.Flags, EFVisibleMeshDrawCommandFlags::ForceInstanceCulling);
+		const bool bForceInstanceCulling = EnumHasAnyFlags(VisibleMeshDrawCommand.Flags, EFVisibleMeshDrawCommandFlags::ForceInstanceCulling) || (GOcclusionForceInstanceCulling != 0);
 		const bool bPreserveInstanceOrder = bOrderPreservationEnabled && EnumHasAnyFlags(VisibleMeshDrawCommand.Flags, EFVisibleMeshDrawCommandFlags::PreserveInstanceOrder);
 		const bool bUseIndirectDraw = bFetchInstanceCountFromScene || bAlwaysUseIndirectDraws || bForceInstanceCulling || (VisibleMeshDrawCommand.NumRuns > 0 || MeshDrawCommand->NumInstances > 1);
 		// UniformBufferView path does not support merging ISM draws atm
