Fix for VRS getting disabled on projects using film grain (i.e. CitySample)

diff --git a/Engine/Source/Runtime/Renderer/Private/PostProcess/PostProcessTonemap.cpp b/Engine/Source/Runtime/Renderer/Private/PostProcess/PostProcessTonemap.cpp
--- a/Engine/Source/Runtime/Renderer/Private/PostProcess/PostProcessTonemap.cpp
+++ b/Engine/Source/Runtime/Renderer/Private/PostProcess/PostProcessTonemap.cpp
@@ -185,13 +185,11 @@ FDesktopDomain RemapPermutation(FDesktopDomain PermutationVector, ERHIFeatureLev
 		PermutationVector.Set<FTonemapperColorFringeDim>(true);
 	}
 
-	// If we're going down a slow path, don't bother with VRS either to keep shader permutations down
-	if (!FVariableRateShadingImageManager::IsVRSCompatibleWithOutputType(PermutationVector.Get<FTonemapperOutputDeviceDim>()) || bFallbackToSlowest)
+	if (!FVariableRateShadingImageManager::IsVRSCompatibleWithOutputType(PermutationVector.Get<FTonemapperOutputDeviceDim>()))
 	{
 		PermutationVector.Set<FTonemapperOutputLuminance>(false);
 	}
 
-
 	// You most likely need Bloom anyway.
 	CommonPermutationVector.Set<FTonemapperBloomDim>(true);
 
