Ensures Scene->EarlyZPassMode is updated correctly when changing MSAA rates. Prevents certain masked materials from becoming invisible at 1x MSAA count.

diff --git a/Engine/Source/Runtime/Renderer/Private/RendererScene.cpp b/Engine/Source/Runtime/Renderer/Private/RendererScene.cpp
--- a/Engine/Source/Runtime/Renderer/Private/RendererScene.cpp
+++ b/Engine/Source/Runtime/Renderer/Private/RendererScene.cpp
@@ -1345,18 +1345,22 @@ void FScene::CheckPrimitiveArrays(int MaxTypeOffsetIndex)
 static void UpdateEarlyZPassModeCVarSinkFunction()
 {
 	static auto* CVarAntiAliasingMethod = IConsoleManager::Get().FindTConsoleVariableDataInt(TEXT("r.AntiAliasingMethod"));
+	static auto* CVarMSAACount = IConsoleManager::Get().FindTConsoleVariableDataInt(TEXT("r.MSAACount"));
 	static int32 CachedAntiAliasingMethod = CVarAntiAliasingMethod->GetValueOnGameThread();
+	static int32 CachedMSAACount = CVarMSAACount->GetValueOnGameThread();
 	static int32 CachedEarlyZPass = CVarEarlyZPass.GetValueOnGameThread();
 	static int32 CachedBasePassWriteDepthEvenWithFullPrepass = CVarBasePassWriteDepthEvenWithFullPrepass.GetValueOnGameThread();
 	static int32 CachedMobileEarlyZPass = CVarMobileEarlyZPass.GetValueOnGameThread();
 
 	const int32 AntiAliasingMethod = CVarAntiAliasingMethod->GetValueOnGameThread();
+	const int32 MSAACount = CVarMSAACount->GetValueOnGameThread();
 	const int32 EarlyZPass = CVarEarlyZPass.GetValueOnGameThread();
 	const int32 BasePassWriteDepthEvenWithFullPrepass = CVarBasePassWriteDepthEvenWithFullPrepass.GetValueOnGameThread();
 	const int32 MobileEarlyZPass = CVarMobileEarlyZPass.GetValueOnGameThread();
 
 	// Switching between MSAA and another AA in forward shading mode requires EarlyZPassMode to update.
 	if (AntiAliasingMethod != CachedAntiAliasingMethod
+		|| MSAACount != CachedMSAACount
 		|| EarlyZPass != CachedEarlyZPass
 		|| BasePassWriteDepthEvenWithFullPrepass != CachedBasePassWriteDepthEvenWithFullPrepass
 		|| MobileEarlyZPass != CachedMobileEarlyZPass)
@@ -1372,6 +1376,7 @@ static void UpdateEarlyZPassModeCVarSinkFunction()
 		}
 
 		CachedAntiAliasingMethod = AntiAliasingMethod;
+		CachedMSAACount = MSAACount;
 		CachedEarlyZPass = EarlyZPass;
 		CachedBasePassWriteDepthEvenWithFullPrepass = BasePassWriteDepthEvenWithFullPrepass;
 		CachedMobileEarlyZPass = MobileEarlyZPass;
