Avoid running VSM invalidations when VSM page pool is not allocated

diff --git a/Engine/Source/Runtime/Renderer/Private/VirtualShadowMaps/VirtualShadowMapCacheManager.cpp b/Engine/Source/Runtime/Renderer/Private/VirtualShadowMaps/VirtualShadowMapCacheManager.cpp
--- a/Engine/Source/Runtime/Renderer/Private/VirtualShadowMaps/VirtualShadowMapCacheManager.cpp
+++ b/Engine/Source/Runtime/Renderer/Private/VirtualShadowMaps/VirtualShadowMapCacheManager.cpp
@@ -796,10 +796,12 @@ bool FVirtualShadowMapArrayCacheManager::IsCacheEnabled()
 bool FVirtualShadowMapArrayCacheManager::IsCacheDataAvailable()
 {
 	return IsCacheEnabled() &&
+		PhysicalPagePool &&
+		PhysicalPageMetaData &&
 		PrevBuffers.PageTable &&
 		PrevBuffers.PageFlags &&
 		PrevBuffers.PageRectBounds &&
-		PrevBuffers.ProjectionData && 
+		PrevBuffers.ProjectionData &&
 		PrevBuffers.PhysicalPageLists;
 }
 
