Updates to MeshDrawCommands Stats
Stat Categories can now specify which passes they care about collecting primitive counts from
Stat Categories can share LinkedStatNames, this allows multiple categories to track the same stat in different passes (if ever desired).
Stat Display shows which passes contributed to the budget
After chatting with tech art, changed categories to use single Basepass.
Temporarily remove "TotalPrimitives" from budget because it's now confusing because it tracks all passes, not just specified passes.  I'll rethink this after some critical tasks for launch have been tackled.

diff --git a/Engine/Source/Runtime/Engine/Public/MeshDrawCommandStatsSettings.h b/Engine/Source/Runtime/Engine/Public/MeshDrawCommandStatsSettings.h
--- a/Engine/Source/Runtime/Engine/Public/MeshDrawCommandStatsSettings.h
+++ b/Engine/Source/Runtime/Engine/Public/MeshDrawCommandStatsSettings.h
@@ -25,6 +25,9 @@ struct FMeshDrawCommandStatsBudget
 	/** The collection which contains this budget. */
 	UPROPERTY(config, EditAnywhere, Category = Engine)
 	int32 Collection = 0;
+	/** Which passes contribute to this budget. */
+	UPROPERTY(config, EditAnywhere, Category = Engine)
+	TArray<FName> Passes;
 };
 
 /** User settings used by the MeshDrawCommandStats system. */
diff --git a/Engine/Source/Runtime/Renderer/Private/MeshDrawCommandStats.cpp b/Engine/Source/Runtime/Renderer/Private/MeshDrawCommandStats.cpp
--- a/Engine/Source/Runtime/Renderer/Private/MeshDrawCommandStats.cpp
+++ b/Engine/Source/Runtime/Renderer/Private/MeshDrawCommandStats.cpp
@@ -161,8 +161,15 @@ FMeshDrawCommandStatsManager::FMeshDrawCommandStatsManager()
 						uint64* PrimitiveCount = BudgetedPrimitives.Find(CategoryBudget.CategoryName);
 						if (PrimitiveCount && *PrimitiveCount > 0)
 						{
+							FString& PassFriendlyNames = StatCollections[CategoryBudget.Collection].CategoryPassFriendlyNames[CategoryBudget.CategoryName];
+
 							FCoreDelegates::EOnScreenMessageSeverity Severity = CategoryBudget.PrimitiveBudget < *PrimitiveCount ? FCoreDelegates::EOnScreenMessageSeverity::Warning : FCoreDelegates::EOnScreenMessageSeverity::Info;
-							OutMessages.Add(Severity, FText::FromString(FString::Printf(TEXT("%5dK / %5dK - %s"), *PrimitiveCount / 1000, CategoryBudget.PrimitiveBudget / 1000, *(CategoryBudget.CategoryName.ToString()))));
+							OutMessages.Add(Severity, FText::FromString(FString::Printf(TEXT("%5dK / %5dK - %s (%s)"), 
+									*PrimitiveCount / 1000, 
+									CategoryBudget.PrimitiveBudget / 1000, 
+									*(CategoryBudget.CategoryName.ToString()),
+									*PassFriendlyNames
+							)));
 						}
 					}
 				}
@@ -184,7 +191,7 @@ FMeshDrawCommandStatsManager::FMeshDrawCommandStatsManager()
 				if (TotalBudget > 0)
 				{
 					FCoreDelegates::EOnScreenMessageSeverity Severity = TotalBudget < Stats.TotalPrimitives ? FCoreDelegates::EOnScreenMessageSeverity::Warning : FCoreDelegates::EOnScreenMessageSeverity::Info;
-					OutMessages.Add(Severity, FText::FromString(FString::Printf(TEXT("%5dK / %5dK - TOTAL"), Stats.TotalPrimitives / 1000, TotalBudget / 1000)));
+					OutMessages.Add(Severity, FText::FromString(FString::Printf(TEXT("%5dK / %5dK - TOTAL (All Passes)"), Stats.TotalPrimitives / 1000, TotalBudget / 1000)));
 				}
 				else
 				{
@@ -292,7 +299,9 @@ void FMeshDrawCommandStatsManager::Update()
 					CustomArgsBufferResult.DrawIndexedIndirectParameters = reinterpret_cast<const FRHIDrawIndexedIndirectParameters*>(CustomArgsBufferResult.GPUBufferReadback->Lock(CustomArgsBufferResult.GPUBufferReadback->GetGPUSizeBytes()));
 				}
 
-				TMap<FName, uint64> CategoryStats;
+				using PassCategoryStats = TMap<FName, uint64>;
+				TMap<FName, PassCategoryStats> Passes;
+
 				for (FMeshDrawCommandPassStats* PassStats : FrameData->PassData)
 				{
 					// make sure the pass was kicked
@@ -301,6 +310,8 @@ void FMeshDrawCommandStatsManager::Update()
 						continue;
 					}
 
+					PassCategoryStats& CategoryStats = Passes.FindOrAdd(PassStats->PassName);
+
 					const uint8* InstanceCullingReadBackData = PassStats->InstanceCullingGPUBufferReadback ? reinterpret_cast<const uint8*>(PassStats->InstanceCullingGPUBufferReadback->Lock(PassStats->DrawData.Num())) : nullptr;
 					const FRHIDrawIndexedIndirectParameters* IndirectArgsPtr = reinterpret_cast<const FRHIDrawIndexedIndirectParameters*>(InstanceCullingReadBackData);
 					
@@ -355,9 +366,14 @@ void FMeshDrawCommandStatsManager::Update()
 					CustomArgsBufferResult.DrawIndexedIndirectParameters = nullptr;
 				}
 
-				for (auto Iter = CategoryStats.CreateConstIterator(); Iter; ++Iter)
+				for (auto PassIter = Passes.CreateConstIterator(); PassIter; ++PassIter)
 				{
-					Stats.CategoryStats.Add(FStats::FCategoryStats(Iter.Key(), Iter.Value()));
+					PassCategoryStats CatMap = PassIter.Value();
+
+					for (auto CatIter = CatMap.CreateConstIterator(); CatIter; ++CatIter)
+					{
+						Stats.CategoryStats.Add(FStats::FCategoryStats(PassIter.Key(), CatIter.Key(), CatIter.Value()));
+					}
 				}
 				Algo::Sort(Stats.CategoryStats, [this](FStats::FCategoryStats& LHS, FStats::FCategoryStats& RHS) { return LHS.CategoryName.ToString() < RHS.CategoryName.ToString(); });
 
@@ -406,14 +422,41 @@ void FMeshDrawCommandStatsManager::Update()
 			const UMeshDrawCommandStatsSettings* Settings = GetDefault<UMeshDrawCommandStatsSettings>();	
 			for (const FMeshDrawCommandStatsBudget& CategoryBudget : Settings->Budgets)
 			{
-				StatCollection& Collection = StatCollections.FindOrAdd(CategoryBudget.Collection);
-				Collection.Add(CategoryBudget.CategoryName, CategoryBudget.CategoryName);
+				FStatCollection& Collection = StatCollections.FindOrAdd(CategoryBudget.Collection);
+
+				FCollectionCategory& Category = Collection.Categories.AddDefaulted_GetRef();
+				Category.Name = CategoryBudget.CategoryName;
+
+				FString& FriendlyNames = Collection.CategoryPassFriendlyNames.FindOrAdd(CategoryBudget.CategoryName);
+
+				if (CategoryBudget.Passes.Num())
+				{
+					for (int i = 0; i < CategoryBudget.Passes.Num(); i++)
+					{
+						const FName& Pass = CategoryBudget.Passes[i];
+						Category.Passes.Add(Pass);
+
+						FriendlyNames += i ? " | " : "";
+						FriendlyNames += Pass.ToString();
+					}
+				}
+				else
+				{
+					FriendlyNames = "All Passes";
+				}
+
+				Category.LinkedNames.Add(CategoryBudget.CategoryName);
 
 				for (FName Name : CategoryBudget.LinkedStatNames)
 				{
-					Collection.Add(Name, CategoryBudget.CategoryName);
+					Category.LinkedNames.Add(Name);
 				}
 			}
+			
+			for (TPair<int32, FStatCollection>& Pair : StatCollections)
+			{
+				Pair.Value.Finish();
+			}
 		}
 
 		// Total up Primitive across stats to their respective Budgets
@@ -429,17 +472,38 @@ void FMeshDrawCommandStatsManager::Update()
 		}
 	#endif
 
-		StatCollection* Collection = StatCollections.Find(CollectionIdx);
+		FStatCollection* Collection = StatCollections.Find(CollectionIdx);
 
 		if (Collection || CollectionIdx == (int)MeshDrawStatsCollection::Pass)
 		{
 			for (const FStats::FCategoryStats& CategoryStat : Stats.CategoryStats)
 			{ 
-				FName* BudgetName = Collection ? Collection->Find(CategoryStat.CategoryName) : nullptr;
-				TMap<FName, uint64>* Map = BudgetName ? &BudgetedPrimitives : &UntrackedPrimitives;
+				TArray<int>* CategoryIndices = nullptr;
+
+				if (Collection)
+				{
+					CategoryIndices = Collection->CategoriesThatLinkStat(CategoryStat.CategoryName);
+				}
+				
+				if (CategoryIndices)
+				{
+					for (int CategoryIndex : *CategoryIndices)
+					{ 
+						FCollectionCategory& Category = Collection->Categories[CategoryIndex];
 
-				uint64& Count = Map->FindOrAdd(BudgetName ? *BudgetName : CategoryStat.CategoryName);
-				Count += CategoryStat.PrimitiveCount;
+						if (!Category.Passes.Num() || Category.Passes.Contains(CategoryStat.PassName))
+						{
+							uint64& Count = BudgetedPrimitives.FindOrAdd(Category.Name);
+							Count += CategoryStat.PrimitiveCount;
+						}
+					}
+				}
+				else
+				{
+					// No collection categories care about this stat, so it was probably missed
+					uint64& Count = UntrackedPrimitives.FindOrAdd(CategoryStat.CategoryName);
+					Count += CategoryStat.PrimitiveCount;
+				}
 			}
 		}
 	}
diff --git a/Engine/Source/Runtime/Renderer/Private/MeshDrawCommandStats.h b/Engine/Source/Runtime/Renderer/Private/MeshDrawCommandStats.h
--- a/Engine/Source/Runtime/Renderer/Private/MeshDrawCommandStats.h
+++ b/Engine/Source/Runtime/Renderer/Private/MeshDrawCommandStats.h
@@ -143,7 +143,9 @@ private:
 		 */
 		struct FCategoryStats
 		{
-			FCategoryStats(FName InCategoryName, uint64 InPrimitiveCount) : CategoryName(InCategoryName), PrimitiveCount(InPrimitiveCount) {}
+			FCategoryStats(FName InPassName, FName InCategoryName, uint64 InPrimitiveCount)
+				: PassName(InPassName), CategoryName(InCategoryName), PrimitiveCount(InPrimitiveCount) {}
+			FName PassName;
 			FName CategoryName;
 			uint64 PrimitiveCount;
 		};
@@ -164,9 +166,38 @@ private:
 		return Frames.Last();
 	}
 
-	using StatCollection = TMap<FName, FName>;
-	using StatCollectionMap = TMap<int32, StatCollection>;
+	struct FCollectionCategory
+	{
+		FName Name;					// User supplied name for this category
+		TSet<FName> Passes;			// Passes this category cares about
+		TSet<FName> LinkedNames;	// LinkedStats this collection cares about
+	};
+
+	struct FStatCollection
+	{
+		TArray<int>* CategoriesThatLinkStat(FName Stat)
+		{
+			return StatToCategoryIndices.Find(Stat);
+		}
+
+		void Finish()
+		{
+			for (int i = 0; i < Categories.Num(); i++)
+			{
+				for (FName& LinkedName : Categories[i].LinkedNames)
+				{
+					StatToCategoryIndices.FindOrAdd(LinkedName).Add(i);
+				}
+			}
+		}
+
+		TArray<FCollectionCategory> Categories;
+		TMap<FName, TArray<int>> StatToCategoryIndices;
+		TMap<FName, FString> CategoryPassFriendlyNames; //< Pipe delimited passes tracked by this category, used for on screen display
+	};
 
+	using FStatCollectionMap = TMap<int32, FStatCollection>;
+	 
 	/** Dump given frame data stats to csv file on disc in profiling folder */
 	void DumpStats(FFrameData* FrameData);
 
@@ -175,11 +206,11 @@ private:
 	bool bCollectStats = false;				//< Collect stats during next frame?
 	
 	FCriticalSection FrameDataCS;			
-	TArray<FFrameData*> Frames;				//< All active frames (contains the frame for which we are collecting stats now and all frames waiting for GPU readback)
-	FStats Stats;							//< Last updated frame stats
-	StatCollectionMap StatCollections; 		//< Per Collection LinkedStatName to Budget CategoryName
-	TMap<FName, uint64> BudgetedPrimitives;	//< Budget CategoryName to Total Primitive Count
-	TMap<FName, uint64> UntrackedPrimitives;//< Primitives which aren't tracked by any Budgets
+	TArray<FFrameData*> Frames;						//< All active frames (contains the frame for which we are collecting stats now and all frames waiting for GPU readback)
+	FStats Stats;									//< Last updated frame stats
+	FStatCollectionMap StatCollections; 			//< Per Collection LinkedStatName to Budget CategoryName
+	TMap<FName, uint64> BudgetedPrimitives;			//< Budget CategoryName to Total Primitive Count
+	TMap<FName, uint64> UntrackedPrimitives;		//< Primitives which aren't tracked by any Budgets
 
 	FDelegateHandle ScreenMessageDelegate;	//< Delegate used to render optional screen stats
 
