Add draw command count as a condition for rendering SingleLayerWater pass

#rnx

diff --git a/Engine/Source/Runtime/Renderer/Private/SingleLayerWaterRendering.cpp b/Engine/Source/Runtime/Renderer/Private/SingleLayerWaterRendering.cpp
--- a/Engine/Source/Runtime/Renderer/Private/SingleLayerWaterRendering.cpp
+++ b/Engine/Source/Runtime/Renderer/Private/SingleLayerWaterRendering.cpp
@@ -140,7 +140,7 @@ bool ShouldRenderSingleLayerWater(TArrayView<const FViewInfo> Views)
 	{
 		for (const FViewInfo& View : Views)
 		{
-			if (View.bHasSingleLayerWaterMaterial)
+			if (View.bHasSingleLayerWaterMaterial && View.ParallelMeshDrawCommandPasses[EMeshPass::SingleLayerWaterPass].HasAnyDraw())
 			{
 				return true;
 			}
