GPUScene Lights - fixed GPU upload which was skipping indices if there were some holes in the sparse array.

diff --git a/Engine/Source/Runtime/Renderer/Private/GPUScene.cpp b/Engine/Source/Runtime/Renderer/Private/GPUScene.cpp
--- a/Engine/Source/Runtime/Renderer/Private/GPUScene.cpp
+++ b/Engine/Source/Runtime/Renderer/Private/GPUScene.cpp
@@ -707,19 +707,25 @@ void FGPUScene::UpdateGPULights(FRDGBuilder& GraphBuilder, FScene& InScene, cons
 {
 	// TODO: remove scene as parameter
 	check(&InScene == &Scene);
-	FRDGUploadData<FLightSceneData> LightData(GraphBuilder, FMath::Max(1, Scene.Lights.Num()));
+	FRDGUploadData<FLightSceneData> LightData(GraphBuilder, FMath::Max(1, Scene.Lights.GetMaxIndex()));
 
 	GraphBuilder.AddSetupTask([this, LightData]
 	{
 		SCOPED_NAMED_EVENT(UpdateGPUScene_Lights, FColor::Green);
 		const bool bAllowStaticLighting = IsStaticLightingAllowed();
 
-		for (int32 Index = 0; Index < Scene.Lights.Num(); ++Index)
+		for (int32 Index = 0; Index < Scene.Lights.GetMaxIndex(); ++Index)
 		{
 			if (Scene.Lights.IsAllocated(Index))
 			{
 				InitLightData(Scene.Lights[Index], bAllowStaticLighting, LightData[Index]);
 			}
+			else
+			{
+				LightData[Index].WorldPosition = TLargeWorldRenderPosition<float>(FVector::ZeroVector);
+				LightData[Index].Color = FVector3f::ZeroVector;
+				LightData[Index].InvRadius = 0.0f;
+			}
 		}
 
 	}, PrerequisiteTask);
@@ -926,8 +932,8 @@ void FGPUScene::UpdateBufferState(FRDGBuilder& GraphBuilder, FSceneUniformBuffer
 	
 	if (bIsMainUpdate)
 	{
-		const uint32 LightDataBufferSize = FMath::RoundUpToPowerOfTwo(FMath::Max(Scene.Lights.Num(), InitialBufferSize));
-		BufferState.LightDataBuffer = GraphBuilder.CreateBuffer(FRDGBufferDesc::CreateStructuredDesc(sizeof(FLightSceneData), FMath::Max(1, Scene.Lights.Num())), TEXT("GPUScene.LightData"));
+		const uint32 LightDataBufferSize = FMath::RoundUpToPowerOfTwo(FMath::Max(Scene.Lights.GetMaxIndex(), InitialBufferSize));
+		BufferState.LightDataBuffer = GraphBuilder.CreateBuffer(FRDGBufferDesc::CreateStructuredDesc(sizeof(FLightSceneData), FMath::Max(1, Scene.Lights.GetMaxIndex())), TEXT("GPUScene.LightData"));
 	}
 
 	ShaderParameters.GPUSceneInstanceSceneData = GraphBuilder.CreateSRV(BufferState.InstanceSceneDataBuffer);
