Modules linked with Renderer can access AddPostProcessMaterialPass

#fordurian

diff --git a/Engine/Source/Runtime/Renderer/Private/PostProcess/PostProcessMaterial.cpp b/Engine/Source/Runtime/Renderer/Private/PostProcess/PostProcessMaterial.cpp
--- a/Engine/Source/Runtime/Renderer/Private/PostProcess/PostProcessMaterial.cpp
+++ b/Engine/Source/Runtime/Renderer/Private/PostProcess/PostProcessMaterial.cpp
@@ -893,6 +893,20 @@ FScreenPassTexture AddPostProcessMaterialPass(
 	return MoveTemp(Output);
 }
 
+FScreenPassTexture AddPostProcessMaterialPass(
+	FRDGBuilder& GraphBuilder,
+	const FSceneView& View,
+	const FPostProcessMaterialInputs& Inputs,
+	const UMaterialInterface* MaterialInterface)
+{
+	if (!ensureMsgf(View.bIsViewInfo, TEXT("AddPostProcessMaterialPass requires that its View parameter is an FViewInfo.")))
+	{
+		return Inputs.GetInput(EPostProcessMaterialInput::SceneColor);
+	}
+
+	return AddPostProcessMaterialPass(GraphBuilder, static_cast<const FViewInfo&>(View), Inputs, MaterialInterface);
+}
+
 static bool IsPostProcessMaterialsEnabledForView(const FViewInfo& View)
 {
 	if (!View.Family->EngineShowFlags.PostProcessing ||
diff --git a/Engine/Source/Runtime/Renderer/Public/PostProcess/PostProcessMaterialInputs.h b/Engine/Source/Runtime/Renderer/Public/PostProcess/PostProcessMaterialInputs.h
--- a/Engine/Source/Runtime/Renderer/Public/PostProcess/PostProcessMaterialInputs.h
+++ b/Engine/Source/Runtime/Renderer/Public/PostProcess/PostProcessMaterialInputs.h
@@ -124,3 +124,11 @@ struct FPostProcessMaterialInputs
 
 	bool bMetalMSAAHDRDecode = false;
 };
+
+class UMaterialInterface;
+
+FScreenPassTexture RENDERER_API AddPostProcessMaterialPass(
+	FRDGBuilder& GraphBuilder,
+	const FSceneView& View,
+	const FPostProcessMaterialInputs& Inputs,
+	const UMaterialInterface* MaterialInterface);
