Overlaps the SkyAtmosphereLUTs on asyncompte with the prepass

RenderSkyAtmosphereLookUpTables keeps the new RDG  resources in a new FSkyAtmospherePendingRDGResources that only CommitToSceneAndViewUniformBuffers() right before the prepass to make sure there is no dependency created between the prepass and the sky atmosphere LUT which doesn't draw any colors.

diff --git a/Engine/Source/Runtime/Renderer/Private/DeferredShadingRenderer.cpp b/Engine/Source/Runtime/Renderer/Private/DeferredShadingRenderer.cpp
--- a/Engine/Source/Runtime/Renderer/Private/DeferredShadingRenderer.cpp
+++ b/Engine/Source/Runtime/Renderer/Private/DeferredShadingRenderer.cpp
@@ -3266,6 +3266,15 @@ void FDeferredShadingSceneRenderer::Render(FRDGBuilder& GraphBuilder)
 
 	ExternalAccessQueue.Submit(GraphBuilder);
 
+	const bool bShouldRenderSkyAtmosphere = ShouldRenderSkyAtmosphere(Scene, ViewFamily.EngineShowFlags);
+	const ESkyAtmospherePassLocation SkyAtmospherePassLocation = GetSkyAtmospherePassLocation();
+	FSkyAtmospherePendingRDGResources SkyAtmospherePendingRDGResources;
+	if (SkyAtmospherePassLocation == ESkyAtmospherePassLocation::BeforePrePass && bShouldRenderSkyAtmosphere)
+	{
+		// Generate the Sky/Atmosphere look up tables overlaping the pre-pass
+		RenderSkyAtmosphereLookUpTables(GraphBuilder, /* out */ SkyAtmospherePendingRDGResources);
+	}
+
 	RenderWaterInfoTexture(GraphBuilder, *this, Scene);
 
 	const bool bShouldRenderVelocities = ShouldRenderVelocities();
@@ -3390,8 +3399,6 @@ void FDeferredShadingSceneRenderer::Render(FRDGBuilder& GraphBuilder)
 		CompositionLighting.ProcessAfterOcclusion(GraphBuilder);
 	};
 
-	const bool bShouldRenderSkyAtmosphere = ShouldRenderSkyAtmosphere(Scene, ViewFamily.EngineShowFlags);
-	const ESkyAtmospherePassLocation SkyAtmospherePassLocation = GetSkyAtmospherePassLocation();
 	const bool bShouldRenderVolumetricCloudBase = ShouldRenderVolumetricCloud(Scene, ViewFamily.EngineShowFlags);
 	const bool bShouldRenderVolumetricCloud = bShouldRenderVolumetricCloudBase && (!ViewFamily.EngineShowFlags.VisualizeVolumetricCloudConservativeDensity && !ViewFamily.EngineShowFlags.VisualizeVolumetricCloudEmptySpaceSkipping);
 	const bool bShouldVisualizeVolumetricCloud = bShouldRenderVolumetricCloudBase && (!!ViewFamily.EngineShowFlags.VisualizeVolumetricCloudConservativeDensity || !!ViewFamily.EngineShowFlags.VisualizeVolumetricCloudEmptySpaceSkipping);
@@ -3454,9 +3461,9 @@ void FDeferredShadingSceneRenderer::Render(FRDGBuilder& GraphBuilder)
 		if (SkyAtmospherePassLocation == ESkyAtmospherePassLocation::BeforeOcclusion && bShouldRenderSkyAtmosphere)
 		{
 			// Generate the Sky/Atmosphere look up tables
-			RenderSkyAtmosphereLookUpTables(GraphBuilder, ExternalAccessQueue);
+			RenderSkyAtmosphereLookUpTables(GraphBuilder, /* out */ SkyAtmospherePendingRDGResources);
 
-			// Sky env map capture uses the view UB, which contains the LUTs computed above. We need to transition them to readable now.
+			SkyAtmospherePendingRDGResources.CommitToSceneAndViewUniformBuffers(GraphBuilder, /* out */ ExternalAccessQueue);
 			ExternalAccessQueue.Submit(GraphBuilder);
 		}
 	
@@ -3487,9 +3494,14 @@ void FDeferredShadingSceneRenderer::Render(FRDGBuilder& GraphBuilder)
 		if (SkyAtmospherePassLocation == ESkyAtmospherePassLocation::BeforeBasePass && bShouldRenderSkyAtmosphere)
 		{
 			// Generate the Sky/Atmosphere look up tables
-			RenderSkyAtmosphereLookUpTables(GraphBuilder, ExternalAccessQueue);
+			RenderSkyAtmosphereLookUpTables(GraphBuilder, /* out */ SkyAtmospherePendingRDGResources);
 
-			// Sky env map capture uses the view UB, which contains the LUTs computed above. We need to transition them to readable now.
+			SkyAtmospherePendingRDGResources.CommitToSceneAndViewUniformBuffers(GraphBuilder, /* out */ ExternalAccessQueue);
+			ExternalAccessQueue.Submit(GraphBuilder);
+		}
+		else if (SkyAtmospherePassLocation == ESkyAtmospherePassLocation::BeforePrePass && bShouldRenderSkyAtmosphere)
+		{
+			SkyAtmospherePendingRDGResources.CommitToSceneAndViewUniformBuffers(GraphBuilder, /* out */ ExternalAccessQueue);
 			ExternalAccessQueue.Submit(GraphBuilder);
 		}
 
diff --git a/Engine/Source/Runtime/Renderer/Private/MobileShadingRenderer.cpp b/Engine/Source/Runtime/Renderer/Private/MobileShadingRenderer.cpp
--- a/Engine/Source/Runtime/Renderer/Private/MobileShadingRenderer.cpp
+++ b/Engine/Source/Runtime/Renderer/Private/MobileShadingRenderer.cpp
@@ -1001,7 +1001,9 @@ void FMobileSceneRenderer::Render(FRDGBuilder& GraphBuilder)
 		const bool bShouldRenderSkyAtmosphere = ShouldRenderSkyAtmosphere(Scene, ViewFamily.EngineShowFlags);
 		if (bShouldRenderSkyAtmosphere)
 		{
-			RenderSkyAtmosphereLookUpTables(GraphBuilder, ExternalAccessQueue);
+			FSkyAtmospherePendingRDGResources PendingRDGResources;
+			RenderSkyAtmosphereLookUpTables(GraphBuilder, /* out */ PendingRDGResources);
+			PendingRDGResources.CommitToSceneAndViewUniformBuffers(GraphBuilder, /* out */ ExternalAccessQueue);
 		}
 
 		GraphBuilder.SetCommandListStat(GET_STATID(STAT_CLMM_Shadows));
diff --git a/Engine/Source/Runtime/Renderer/Private/SceneRendering.h b/Engine/Source/Runtime/Renderer/Private/SceneRendering.h
--- a/Engine/Source/Runtime/Renderer/Private/SceneRendering.h
+++ b/Engine/Source/Runtime/Renderer/Private/SceneRendering.h
@@ -2456,7 +2456,7 @@ protected:
 	void InitSkyAtmosphereForViews(FRHICommandListImmediate& RHICmdList);
 	
 	/** Render the sky atmosphere look up table needed for this frame.*/
-	void RenderSkyAtmosphereLookUpTables(FRDGBuilder& GraphBuilder, FRDGExternalAccessQueue& ExternalAccessQueue);
+	void RenderSkyAtmosphereLookUpTables(FRDGBuilder& GraphBuilder, class FSkyAtmospherePendingRDGResources& PendingRDGResources);
 
 	/** Render the sky atmosphere over the scene.*/
 	void RenderSkyAtmosphere(FRDGBuilder& GraphBuilder, const FMinimalSceneTextures& SceneTextures);
diff --git a/Engine/Source/Runtime/Renderer/Private/SkyAtmosphereRendering.cpp b/Engine/Source/Runtime/Renderer/Private/SkyAtmosphereRendering.cpp
--- a/Engine/Source/Runtime/Renderer/Private/SkyAtmosphereRendering.cpp
+++ b/Engine/Source/Runtime/Renderer/Private/SkyAtmosphereRendering.cpp
@@ -217,8 +217,8 @@ static TAutoConsoleVariable<int32> CVarSkyAtmosphereEditorNotifications(
 	TEXT("Enable the rendering of in editor notification to warn the user about missing sky dome pixels on screen. It is better to keep it enabled and will be removed when shipping.\n"),
 	ECVF_RenderThreadSafe | ECVF_Scalability);
 
-static TAutoConsoleVariable<bool> CVarSkyAtmosphereASyncCompute(
-	TEXT("r.SkyAtmosphereASyncCompute"), false,
+static TAutoConsoleVariable<int32> CVarSkyAtmosphereASyncCompute(
+	TEXT("r.SkyAtmosphereASyncCompute"), 0,
 	TEXT("SkyAtmosphere on async compute (default: false). When running on the async pipe, SkyAtmosphere lut generation will overlap with the occlusion pass.\n"),
 	ECVF_RenderThreadSafe);
 
@@ -284,11 +284,15 @@ IMPLEMENT_GLOBAL_SHADER_PARAMETER_STRUCT(FSkyAtmosphereInternalCommonParameters,
 
 ESkyAtmospherePassLocation GetSkyAtmospherePassLocation()
 {
-	if (CVarSkyAtmosphereASyncCompute.GetValueOnAnyThread())
+	if (CVarSkyAtmosphereASyncCompute.GetValueOnAnyThread() == 1)
 	{
 		// When SkyAtmLUT is running on async compute, try to kick it before occlusion pass to have better wave occupancy
 		return ESkyAtmospherePassLocation::BeforeOcclusion;
 	}
+	else if (CVarSkyAtmosphereASyncCompute.GetValueOnAnyThread() == 2)
+	{
+		return ESkyAtmospherePassLocation::BeforePrePass;
+	}
 	// When SkyAtmLUT is running on graphics pipe, kickbefore BasePass to have a better overlap when DistanceFieldShadows is running async
 	return ESkyAtmospherePassLocation::BeforeBasePass;
 }
@@ -1291,7 +1295,7 @@ static void SetupSkyAtmosphereInternalCommonParameters(
 	InternalCommonParameters.CameraAerialPerspectiveVolumeSizeAndInvSize = GetSizeAndInvSize(CameraAerialPerspectiveVolumeScreenResolution, CameraAerialPerspectiveVolumeScreenResolution);
 }
 
-void FSceneRenderer::RenderSkyAtmosphereLookUpTables(FRDGBuilder& GraphBuilder, FRDGExternalAccessQueue& ExternalAccessQueue)
+void FSceneRenderer::RenderSkyAtmosphereLookUpTables(FRDGBuilder& GraphBuilder, class FSkyAtmospherePendingRDGResources& PendingRDGResources)
 {
 	check(ShouldRenderSkyAtmosphere(Scene, ViewFamily.EngineShowFlags)); // This should not be called if we should not render SkyAtmosphere
 
@@ -1303,6 +1307,9 @@ void FSceneRenderer::RenderSkyAtmosphereLookUpTables(FRDGBuilder& GraphBuilder,
 	FSkyAtmosphereRenderSceneInfo& SkyInfo = *Scene->GetSkyAtmosphereSceneInfo();
 	const FSkyAtmosphereSceneProxy& SkyAtmosphereSceneProxy = SkyInfo.GetSkyAtmosphereSceneProxy();
 
+	PendingRDGResources.SceneRenderer = this;
+	PendingRDGResources.ViewResources.SetNum(Views.Num());
+
 	const bool bHighQualityMultiScattering = CVarSkyAtmosphereMultiScatteringLUTHighQuality.GetValueOnRenderThread() > 0;
 	const bool bSecondAtmosphereLightEnabled = Scene->IsSecondAtmosphereLightEnabled();
 	const bool bSeparatedAtmosphereMieRayLeigh = VolumetricCloudWantsSeparatedAtmosphereMieRayLeigh(Scene);
@@ -1406,7 +1413,7 @@ void FSceneRenderer::RenderSkyAtmosphereLookUpTables(FRDGBuilder& GraphBuilder,
 		const FIntVector NumGroups = FIntVector::DivideAndRoundUp(TextureSize, FRenderDistantSkyLightLutCS::GroupSize);
 		FComputeShaderUtils::AddPass(GraphBuilder, RDG_EVENT_NAME("DistantSkyLightLut"), PassFlag, ComputeShader, PassParameters, NumGroups);
 
-		SkyInfo.GetDistantSkyLightLutTexture() = ConvertToExternalAccessTexture(GraphBuilder, ExternalAccessQueue, DistantSkyLightLut, ERHIAccess::SRVMask, ERHIPipeline::All);
+		PendingRDGResources.DistantSkyLightLut = DistantSkyLightLut;
 	}
 
 	SkyAtmosphereLightShadowData LightShadowData;
@@ -1506,11 +1513,7 @@ void FSceneRenderer::RenderSkyAtmosphereLookUpTables(FRDGBuilder& GraphBuilder,
 			const FIntVector NumGroups = FIntVector::DivideAndRoundUp(TextureSize, FRenderSkyViewLutCS::GroupSize);
 			FComputeShaderUtils::AddPass(GraphBuilder, RDG_EVENT_NAME("RealTimeCaptureSkyViewLut"), PassFlag, ComputeShader, PassParameters, NumGroups);
 
-			Scene->RealTimeReflectionCaptureSkyAtmosphereViewLutTexture = ConvertToExternalAccessTexture(GraphBuilder, ExternalAccessQueue, RealTimeReflectionCaptureSkyAtmosphereViewLutTexture, ERHIAccess::SRVMask, ERHIPipeline::All);
-		}
-		else
-		{
-			Scene->RealTimeReflectionCaptureSkyAtmosphereViewLutTexture = nullptr;
+			PendingRDGResources.RealTimeReflectionCaptureSkyAtmosphereViewLutTexture = RealTimeReflectionCaptureSkyAtmosphereViewLutTexture;
 		}
 
 		if (bRealTimeReflectionCapture360APLutTexture)
@@ -1560,18 +1563,9 @@ void FSceneRenderer::RenderSkyAtmosphereLookUpTables(FRDGBuilder& GraphBuilder,
 			const FIntVector NumGroups = FIntVector::DivideAndRoundUp(TextureSize, FRenderCameraAerialPerspectiveVolumeCS::GroupSize);
 			FComputeShaderUtils::AddPass(GraphBuilder, RDG_EVENT_NAME("RealTimeCaptureCamera360VolumeLut"), PassFlag, ComputeShader, PassParameters, NumGroups);
 
-			Scene->RealTimeReflectionCaptureCamera360APLutTexture = ConvertToExternalAccessTexture(GraphBuilder, ExternalAccessQueue, RealTimeReflectionCaptureCamera360APLutTexture, ERHIAccess::SRVMask, ERHIPipeline::All);
-		}
-		else
-		{
-			Scene->RealTimeReflectionCaptureCamera360APLutTexture = nullptr;
+			PendingRDGResources.RealTimeReflectionCaptureCamera360APLutTexture = RealTimeReflectionCaptureCamera360APLutTexture;
 		}
 	}
-	else
-	{
-		Scene->RealTimeReflectionCaptureSkyAtmosphereViewLutTexture = nullptr;
-		Scene->RealTimeReflectionCaptureCamera360APLutTexture = nullptr;
-	}
 
 	for (int32 ViewIndex = 0; ViewIndex < Views.Num(); ViewIndex++)
 	{
@@ -1697,13 +1691,51 @@ void FSceneRenderer::RenderSkyAtmosphereLookUpTables(FRDGBuilder& GraphBuilder,
 			FComputeShaderUtils::AddPass(GraphBuilder, RDG_EVENT_NAME("CameraVolumeLut"), PassFlag, ComputeShader, PassParameters, NumGroups);
 		}
 
-		View.SkyAtmosphereViewLutTexture = ConvertToExternalAccessTexture(GraphBuilder, ExternalAccessQueue, SkyAtmosphereViewLutTexture);
-		View.SkyAtmosphereCameraAerialPerspectiveVolume = ConvertToExternalAccessTexture(GraphBuilder, ExternalAccessQueue, SkyAtmosphereCameraAerialPerspectiveVolume, ERHIAccess::SRVMask, ERHIPipeline::All);
+		PendingRDGResources.ViewResources[ViewIndex].SkyAtmosphereViewLutTexture = SkyAtmosphereViewLutTexture;
+		PendingRDGResources.ViewResources[ViewIndex].SkyAtmosphereCameraAerialPerspectiveVolume = SkyAtmosphereCameraAerialPerspectiveVolume;
 	}
 
-	SkyInfo.GetTransmittanceLutTexture() = ConvertToExternalAccessTexture(GraphBuilder, ExternalAccessQueue, TransmittanceLut);
+	PendingRDGResources.TransmittanceLut = TransmittanceLut;
+}
+
+void FSkyAtmospherePendingRDGResources::CommitToSceneAndViewUniformBuffers(FRDGBuilder& GraphBuilder, FRDGExternalAccessQueue& ExternalAccessQueue) const
+{
+	check(SceneRenderer);
+	FScene* Scene = SceneRenderer->Scene;
+	FSkyAtmosphereRenderSceneInfo& SkyInfo = *Scene->GetSkyAtmosphereSceneInfo();
+
+	if (DistantSkyLightLut)
+	{
+		SkyInfo.GetDistantSkyLightLutTexture() = ConvertToExternalAccessTexture(GraphBuilder, ExternalAccessQueue, DistantSkyLightLut, ERHIAccess::SRVMask, ERHIPipeline::All);
+	}
 
-	// TODO have RDG execute those above passes with compute overlap similarly to using AutomaticCacheFlushAfterComputeShader(true);
+	if (RealTimeReflectionCaptureSkyAtmosphereViewLutTexture)
+	{
+		Scene->RealTimeReflectionCaptureSkyAtmosphereViewLutTexture = ConvertToExternalAccessTexture(GraphBuilder, ExternalAccessQueue, RealTimeReflectionCaptureSkyAtmosphereViewLutTexture, ERHIAccess::SRVMask, ERHIPipeline::All);
+	}
+	else
+	{
+		Scene->RealTimeReflectionCaptureSkyAtmosphereViewLutTexture = nullptr;
+	}
+
+	if (RealTimeReflectionCaptureCamera360APLutTexture)
+	{
+		Scene->RealTimeReflectionCaptureCamera360APLutTexture = ConvertToExternalAccessTexture(GraphBuilder, ExternalAccessQueue, RealTimeReflectionCaptureCamera360APLutTexture, ERHIAccess::SRVMask, ERHIPipeline::All);
+	}
+	else
+	{
+		Scene->RealTimeReflectionCaptureCamera360APLutTexture = nullptr;
+	}
+
+	for (int32 ViewIndex = 0; ViewIndex < SceneRenderer->Views.Num(); ViewIndex++)
+	{
+		FViewInfo& View = SceneRenderer->Views[ViewIndex];
+		View.SkyAtmosphereViewLutTexture = ConvertToExternalAccessTexture(GraphBuilder, ExternalAccessQueue, ViewResources[ViewIndex].SkyAtmosphereViewLutTexture, ERHIAccess::SRVMask, ERHIPipeline::All);
+		View.SkyAtmosphereCameraAerialPerspectiveVolume = ConvertToExternalAccessTexture(GraphBuilder, ExternalAccessQueue, ViewResources[ViewIndex].SkyAtmosphereCameraAerialPerspectiveVolume, ERHIAccess::SRVMask, ERHIPipeline::All);
+	}
+
+	check(TransmittanceLut);
+	SkyInfo.GetTransmittanceLutTexture() = ConvertToExternalAccessTexture(GraphBuilder, ExternalAccessQueue, TransmittanceLut);
 }
 
 FSkyAtmosphereRenderContext::FSkyAtmosphereRenderContext()
diff --git a/Engine/Source/Runtime/Renderer/Private/SkyAtmosphereRendering.h b/Engine/Source/Runtime/Renderer/Private/SkyAtmosphereRendering.h
--- a/Engine/Source/Runtime/Renderer/Private/SkyAtmosphereRendering.h
+++ b/Engine/Source/Runtime/Renderer/Private/SkyAtmosphereRendering.h
@@ -168,13 +168,41 @@ private:
 	TRefCountPtr<IPooledRenderTarget> DistantSkyLightLutTexture;
 };
 
+/** Pending RDG resource to commit after the pre-pass / nanite rasterization so that RenderSkyAtmosphereLookUpTables() can overlap them on async compute. */
+class FSkyAtmospherePendingRDGResources
+{
+public:
+	// Sky env map capture uses the view UB, which contains the LUTs computed above. We need to transition them to readable now.
+	void CommitToSceneAndViewUniformBuffers(FRDGBuilder& GraphBuilder, FRDGExternalAccessQueue& ExternalAccessQueue) const;
+
+private:
+	struct FViewRDGResources
+	{
+		FRDGTextureRef SkyAtmosphereViewLutTexture = nullptr;
+		FRDGTextureRef SkyAtmosphereCameraAerialPerspectiveVolume = nullptr;
+	};
+
+	FSceneRenderer* SceneRenderer = nullptr;
+	TArray<FViewRDGResources, TInlineAllocator<4>> ViewResources;
+
+	FRDGTextureRef DistantSkyLightLut = nullptr;
+	FRDGTextureRef RealTimeReflectionCaptureSkyAtmosphereViewLutTexture = nullptr;
+	FRDGTextureRef RealTimeReflectionCaptureCamera360APLutTexture = nullptr;
+	FRDGTextureRef TransmittanceLut = nullptr;
+
+	friend class FSceneRenderer;
+};
+
 enum class ESkyAtmospherePassLocation : uint32
 {
 	// Renders just before the Occlusion queries. Good for wave occuppency when SkyAtmosphere is async
 	BeforeOcclusion,
 
 	// Renders just before the BasePass.
-	BeforeBasePass
+	BeforeBasePass,
+
+	// Renders just before prepass but can overlap on async compute with it too.
+	BeforePrePass,
 };
 
 // Returns the location in the frame where SkyAtmosphere is rendered.
