Add support for not unmapping all pages when invalidating RVT.
Previously the behavior was to only keep locked root pages mapped and produce them in the current frame.
Unmapping all other pages can cause flickering, becuase it will take a couple of frames to get GPU feedback requesting them for remap.
While that is happening we end falling back the the root page.
This change adds r.VT.RVT.DirtyPagesKeptMappedFrames and defaults it to 8.
Pages that had GPU feedback in the previous number of frames set by this cvar are kept mapped but queued for update.
The behavior of the pages queued for mapped update is also changed so that it is no longer handled in the current frame but timesliced along with other unmapped updates.
Potentially we may end up producing more pages with this logic, expecially if the cvar is set to a high value. But default value seems neutral in tests.
Setting the cvar to 0 should give the old behavior.

diff --git a/Engine/Source/Runtime/Renderer/Private/VT/RuntimeVirtualTextureSceneProxy.cpp b/Engine/Source/Runtime/Renderer/Private/VT/RuntimeVirtualTextureSceneProxy.cpp
--- a/Engine/Source/Runtime/Renderer/Private/VT/RuntimeVirtualTextureSceneProxy.cpp
+++ b/Engine/Source/Runtime/Renderer/Private/VT/RuntimeVirtualTextureSceneProxy.cpp
@@ -25,6 +25,12 @@ static TAutoConsoleVariable<int32> CVarVTStreamingMips(
 	}),
 	ECVF_Default);
 
+static TAutoConsoleVariable<int32> CVarDirtyPagesKeptMappedFrames(
+	TEXT("r.VT.RVT.DirtyPagesKeptMappedFrames"),
+	8,
+	TEXT("When invalidating RVT pages, we keep them mapped if they gave feedback in the last N frames."),
+	ECVF_RenderThreadSafe);
+
 int32 FRuntimeVirtualTextureSceneProxy::ProducerIdGenerator = 1;
 
 FRuntimeVirtualTextureSceneProxy::FRuntimeVirtualTextureSceneProxy(URuntimeVirtualTextureComponent* InComponent)
@@ -186,6 +192,10 @@ void FRuntimeVirtualTextureSceneProxy::FlushDirtyPages()
 		// Don't do any work if we won't mark anything dirty.
 		if (MaxDirtyLevel >= 0 && CombinedDirtyRect.Width() != 0 && CombinedDirtyRect.Height() != 0)
 		{
+			// Keeping visible pages mapped reduces update flicker due to the latency in the unmap/feedback/map sequence.
+			// But it potentially creates more page update work since more pages may get updated.
+			const uint32 MaxAgeToKeepMapped = CVarDirtyPagesKeptMappedFrames.GetValueOnRenderThread();
+
 			//todo[vt]: 
 			// Profile to work out best heuristic for when we should use the CombinedDirtyRect
 			// Also consider using some other structure to represent dirty area such as a course 2D bitfield
@@ -193,13 +203,13 @@ void FRuntimeVirtualTextureSceneProxy::FlushDirtyPages()
 
 			if (bCombinedFlush)
 			{
-				FVirtualTextureSystem::Get().FlushCache(ProducerHandle, SpaceID, CombinedDirtyRect, MaxDirtyLevel);
+				FVirtualTextureSystem::Get().FlushCache(ProducerHandle, SpaceID, CombinedDirtyRect, MaxDirtyLevel, MaxAgeToKeepMapped);
 			}
 			else
 			{
 				for (FIntRect Rect : DirtyRects)
 				{
-					FVirtualTextureSystem::Get().FlushCache(ProducerHandle, SpaceID, Rect, MaxDirtyLevel);
+					FVirtualTextureSystem::Get().FlushCache(ProducerHandle, SpaceID, Rect, MaxDirtyLevel, MaxAgeToKeepMapped);
 				}
 			}
 		}
diff --git a/Engine/Source/Runtime/Renderer/Private/VT/TexturePagePool.cpp b/Engine/Source/Runtime/Renderer/Private/VT/TexturePagePool.cpp
--- a/Engine/Source/Runtime/Renderer/Private/VT/TexturePagePool.cpp
+++ b/Engine/Source/Runtime/Renderer/Private/VT/TexturePagePool.cpp
@@ -122,7 +122,14 @@ void FTexturePagePool::EvictPages(FVirtualTextureSystem* System, const FVirtualT
 	}
 }
 
-void FTexturePagePool::EvictPages(FVirtualTextureSystem* System, FVirtualTextureProducerHandle const& ProducerHandle, FVTProducerDescription const& Desc, FIntRect const& TextureRegion, uint32 MaxLevel, TArray<union FVirtualTextureLocalTile>& OutLocked)
+void FTexturePagePool::EvictPages(
+	FVirtualTextureSystem* System, 
+	FVirtualTextureProducerHandle const& ProducerHandle, 
+	FVTProducerDescription const& Desc, 
+	FIntRect const& TextureRegion, 
+	uint32 MaxLevelToEvict, 
+	uint32 MinFrameToKeepMapped,
+	TArray<union FVirtualTextureLocalTile>& OutDirtyMapped)
 {
 	TArray<uint32, TInlineAllocator<256>> ToEvict;
 	const uint32 Hash = MurmurFinalize32(ProducerHandle.PackedValue);
@@ -134,7 +141,7 @@ void FTexturePagePool::EvictPages(FVirtualTextureSystem* System, FVirtualTexture
 			const uint32 vAddress = Pages[pAddress].Local_vAddress;
 			const uint32 vLevel = Pages[pAddress].Local_vLevel;
 
-			if (vLevel <= MaxLevel)
+			if (vLevel <= MaxLevelToEvict)
 			{
 				const int32 TileSize = Desc.TileSize << vLevel;
 				const int32 X = FMath::ReverseMortonCode2(vAddress) * TileSize;
@@ -146,11 +153,21 @@ void FTexturePagePool::EvictPages(FVirtualTextureSystem* System, FVirtualTexture
 				{
 					if (!FreeHeap.IsPresent(pAddress))
 					{
-						OutLocked.Add(FVirtualTextureLocalTile(ProducerHandle, vAddress, vLevel));
+						// Locked pages aren't unmapped but are added to the dirty output array.
+						OutDirtyMapped.Add(FVirtualTextureLocalTile(ProducerHandle, vAddress, vLevel));
 					}
 					else
 					{
-						ToEvict.Add(pAddress);
+						const uint32 PageFrame = FreeHeap.GetKey(pAddress) >> 4;
+						if (PageFrame >= MinFrameToKeepMapped)
+						{
+							// Visible pages aren't unmapped but are added to the dirty output array.
+							OutDirtyMapped.Add(FVirtualTextureLocalTile(ProducerHandle, vAddress, vLevel));
+						}
+						else
+						{
+							ToEvict.Add(pAddress);
+						}
 					}
 				}
 			}
diff --git a/Engine/Source/Runtime/Renderer/Private/VT/TexturePagePool.h b/Engine/Source/Runtime/Renderer/Private/VT/TexturePagePool.h
--- a/Engine/Source/Runtime/Renderer/Private/VT/TexturePagePool.h
+++ b/Engine/Source/Runtime/Renderer/Private/VT/TexturePagePool.h
@@ -49,7 +49,7 @@ public:
 	 * Unmap/remove any pages that were allocated by the given producer and are inside the TextureRegion.
 	 * Outputs the locked pages that can't be unmapped to the OutLocked array.
 	 */
-	void EvictPages(FVirtualTextureSystem* System, FVirtualTextureProducerHandle const& ProducerHandle, FVTProducerDescription const& Desc, FIntRect const& TextureRegion, uint32 MaxLevel, TArray<union FVirtualTextureLocalTile>& OutLocked);
+	void EvictPages(FVirtualTextureSystem* System, FVirtualTextureProducerHandle const& ProducerHandle, FVTProducerDescription const& Desc, FIntRect const& TextureRegion, uint32 MaxLevelToEvict, uint32 MinFrameToKeepMapped, TArray<union FVirtualTextureLocalTile>& OutDirtyMapped);
 
 	/**
 	* Unmap all pages from the given address range in a space. Pages will remain resident in the pool, but no longer by mapped to any page table.
diff --git a/Engine/Source/Runtime/Renderer/Private/VT/VirtualTextureSystem.cpp b/Engine/Source/Runtime/Renderer/Private/VT/VirtualTextureSystem.cpp
--- a/Engine/Source/Runtime/Renderer/Private/VT/VirtualTextureSystem.cpp
+++ b/Engine/Source/Runtime/Renderer/Private/VT/VirtualTextureSystem.cpp
@@ -396,7 +396,7 @@ void FVirtualTextureSystem::FlushCache()
 	bFlushCaches = true;
 }
 
-void FVirtualTextureSystem::FlushCache(FVirtualTextureProducerHandle const& ProducerHandle, int32 SpaceID, FIntRect const& TextureRegion, uint32 MaxLevel)
+void FVirtualTextureSystem::FlushCache(FVirtualTextureProducerHandle const& ProducerHandle, int32 SpaceID, FIntRect const& TextureRegion, uint32 MaxLevelToEvict, uint32 MaxAgeToKeepMapped)
 {
 	check(!bUpdating);
 	UE::TScopeLock Lock(Mutex);
@@ -417,11 +417,13 @@ void FVirtualTextureSystem::FlushCache(FVirtualTextureProducerHandle const& Prod
 
 		check(TransientCollectedPages.Num() == 0);
 
+		const uint32 MinFrameToKeepMapped = Frame > MaxAgeToKeepMapped ? Frame - MaxAgeToKeepMapped : 0;
+
 		// If this is an Adaptive VT we need to collect all of the associated Producers to flush.
 		TArray<FAdaptiveVirtualTexture::FProducerInfo> ProducerInfos;
 		if (AdaptiveVTs[SpaceID] != nullptr)
 		{
-			AdaptiveVTs[SpaceID]->GetProducers(TextureRegion, MaxLevel, ProducerInfos);
+			AdaptiveVTs[SpaceID]->GetProducers(TextureRegion, MaxLevelToEvict, ProducerInfos);
 		}
 
 		for (int32 i = 0; i < PhysicalSpacesForProducer.Num(); ++i)
@@ -433,13 +435,13 @@ void FVirtualTextureSystem::FlushCache(FVirtualTextureProducerHandle const& Prod
 				// Adaptive VT flushes.
 				for (FAdaptiveVirtualTexture::FProducerInfo& Info : ProducerInfos)
 				{
-					Pool.EvictPages(this, Info.ProducerHandle, ProducerDescription, Info.RemappedTextureRegion, Info.RemappedMaxLevel, TransientCollectedPages);
+					Pool.EvictPages(this, Info.ProducerHandle, ProducerDescription, Info.RemappedTextureRegion, Info.RemappedMaxLevel, MinFrameToKeepMapped, TransientCollectedPages);
 				}
 			}
 			else
 			{
 				// Regular flush.
-				Pool.EvictPages(this, ProducerHandle, ProducerDescription, TextureRegion, MaxLevel, TransientCollectedPages);
+				Pool.EvictPages(this, ProducerHandle, ProducerDescription, TextureRegion, MaxLevelToEvict, MinFrameToKeepMapped, TransientCollectedPages);
 			}
 		}
 
@@ -1932,10 +1934,11 @@ void FVirtualTextureSystem::UpdateResidencyTracking() const
 	}
 }
 
-void FVirtualTextureSystem::SubmitRequestsFromLocalTileList(FRHICommandList& RHICmdList, TArray<FVirtualTextureLocalTile>& OutDeferredTiles, const TSet<FVirtualTextureLocalTile>& LocalTileList, EVTProducePageFlags Flags, ERHIFeatureLevel::Type FeatureLevel)
+void FVirtualTextureSystem::SubmitRequestsFromLocalTileList(FRHICommandList& RHICmdList, TArray<FVirtualTextureLocalTile>& OutDeferredTiles, const TSet<FVirtualTextureLocalTile>& LocalTileList, EVTProducePageFlags Flags, ERHIFeatureLevel::Type FeatureLevel, uint32 MaxRequestsToProduce)
 {
 	LLM_SCOPE(ELLMTag::VirtualTextureSystem);
 
+	uint32 NumPagesProduced = 0;
 	for (const FVirtualTextureLocalTile& Tile : LocalTileList)
 	{
 		const FVirtualTextureProducerHandle ProducerHandle = Tile.GetProducerHandle();
@@ -1980,6 +1983,13 @@ void FVirtualTextureSystem::SubmitRequestsFromLocalTileList(FRHICommandList& RHI
 			continue;
 		}
 
+		if (MaxRequestsToProduce > 0 && NumPagesProduced >= MaxRequestsToProduce)
+		{
+			// Keep the request for the next frame?
+			OutDeferredTiles.Add(Tile);
+			continue;
+		}
+
 		FVTRequestPageResult RequestPageResult = Producer->GetVirtualTexture()->RequestPageData(
 			RHICmdList, ProducerHandle, LayerMask, Tile.Local_vLevel, Tile.Local_vAddress, EVTRequestPagePriority::High);
 
@@ -2002,16 +2012,18 @@ void FVirtualTextureSystem::SubmitRequestsFromLocalTileList(FRHICommandList& RHI
 			// Add the finalizer here but note that we don't call Finalize until SubmitRequests()
 			Finalizers.AddUnique(VTFinalizer);
 		}
+
+		NumPagesProduced++;
 	}
 }
 
-void FVirtualTextureSystem::SubmitPreMappedRequests(FRHICommandList& RHICmdList, ERHIFeatureLevel::Type FeatureLevel)
+void FVirtualTextureSystem::SubmitPreMappedRequests(FRHICommandList& RHICmdList, ERHIFeatureLevel::Type FeatureLevel, uint32 MaxMappedTilesToProduce)
 {
 	check(TransientCollectedPages.Num() == 0);
 
 	{
 		INC_DWORD_STAT_BY(STAT_NumMappedPageUpdate, MappedTilesToProduce.Num());
-		SubmitRequestsFromLocalTileList(RHICmdList, TransientCollectedPages, MappedTilesToProduce, EVTProducePageFlags::None, FeatureLevel);
+		SubmitRequestsFromLocalTileList(RHICmdList, TransientCollectedPages, MappedTilesToProduce, EVTProducePageFlags::None, FeatureLevel, MaxMappedTilesToProduce);
 		MappedTilesToProduce.Reset();
 		MappedTilesToProduce.Append(TransientCollectedPages);
 		TransientCollectedPages.Reset();
@@ -2019,7 +2031,7 @@ void FVirtualTextureSystem::SubmitPreMappedRequests(FRHICommandList& RHICmdList,
 
 	{
 		INC_DWORD_STAT_BY(STAT_NumContinuousPageUpdate, ContinuousUpdateTilesToProduce.Num());
-		SubmitRequestsFromLocalTileList(RHICmdList, TransientCollectedPages, ContinuousUpdateTilesToProduce, EVTProducePageFlags::ContinuousUpdate, FeatureLevel);
+		SubmitRequestsFromLocalTileList(RHICmdList, TransientCollectedPages, ContinuousUpdateTilesToProduce, EVTProducePageFlags::ContinuousUpdate, FeatureLevel, 0);
 		ContinuousUpdateTilesToProduce.Reset();
 		TransientCollectedPages.Reset();
 	}
@@ -2076,7 +2088,7 @@ void FVirtualTextureSystem::SubmitThrottledRequests(FRHICommandList& RHICmdList,
 	Updater->NumProcessedLoadRequests += MergedRequestList->GetNumLoadRequests();
 
 	// Submit the requests to produce pages that are already mapped
-	SubmitPreMappedRequests(RHICmdList, FeatureLevel);
+	SubmitPreMappedRequests(RHICmdList, FeatureLevel, Settings.MaxRVTPageUploads);
 
 	// Submit the merged requests
 	SubmitRequests(RHICmdList, FeatureLevel, Allocator, Settings, MergedRequestList, true);
diff --git a/Engine/Source/Runtime/Renderer/Private/VT/VirtualTextureSystem.h b/Engine/Source/Runtime/Renderer/Private/VT/VirtualTextureSystem.h
--- a/Engine/Source/Runtime/Renderer/Private/VT/VirtualTextureSystem.h
+++ b/Engine/Source/Runtime/Renderer/Private/VT/VirtualTextureSystem.h
@@ -160,7 +160,7 @@ public:
 	void RequestRecordedTiles(TArray<uint64>&& InPageRequests);
 
 	void FlushCache();
-	void FlushCache(FVirtualTextureProducerHandle const& ProducerHandle, int32 SpaceID, FIntRect const& TextureRegion, uint32 MaxLevel);
+	void FlushCache(FVirtualTextureProducerHandle const& ProducerHandle, int32 SpaceID, FIntRect const& TextureRegion, uint32 MaxLevelToEvict, uint32 MaxAgeToKeepMapped);
 
 	float GetGlobalMipBias() const;
 
@@ -204,13 +204,13 @@ private:
 	void RequestTilesInternal(const IAllocatedVirtualTexture* AllocatedVT, int32 InMipLevel);
 	void RequestTilesInternal(const IAllocatedVirtualTexture* AllocatedVT, const FVector2D& InScreenSpaceSize, int32 InMipLevel);
 	
-	void SubmitRequestsFromLocalTileList(FRHICommandList& RHICmdList, TArray<FVirtualTextureLocalTile>& OutDeferredTiles, const TSet<FVirtualTextureLocalTile>& LocalTileList, EVTProducePageFlags Flags, ERHIFeatureLevel::Type FeatureLevel);
+	void SubmitRequestsFromLocalTileList(FRHICommandList& RHICmdList, TArray<FVirtualTextureLocalTile>& OutDeferredTiles, const TSet<FVirtualTextureLocalTile>& LocalTileList, EVTProducePageFlags Flags, ERHIFeatureLevel::Type FeatureLevel, uint32 MaxRequestsToProduce);
 
 	void GatherFeedbackRequests(FConcurrentLinearBulkObjectAllocator& Allocator, const FVirtualTextureUpdateSettings& Settings, const FVirtualTextureFeedback::FMapResult& FeedbackResult, FUniqueRequestList* MergedRequestList);
 	void GatherLockedTileRequests(FUniqueRequestList* MergedRequestList);
 	void GatherPackedTileRequests(FConcurrentLinearBulkObjectAllocator& Allocator, const FVirtualTextureUpdateSettings& Settings, FUniqueRequestList* MergedRequestList);
 
-	void SubmitPreMappedRequests(FRHICommandList& RHICmdList, ERHIFeatureLevel::Type FeatureLevel);
+	void SubmitPreMappedRequests(FRHICommandList& RHICmdList, ERHIFeatureLevel::Type FeatureLevel, uint32 MaxMappedTilesToProduce);
 
 	void SubmitThrottledRequests(FRHICommandList& RHICmdList, FVirtualTextureUpdater* Updater, bool bContinuousUpdates);
 	void SubmitRequests(FRHICommandList& RHICmdList, ERHIFeatureLevel::Type FeatureLevel, FConcurrentLinearBulkObjectAllocator& Allocator, FVirtualTextureUpdateSettings const& Settings, FUniqueRequestList* RequestList, bool bAsync);
