Fixed crash in model scene proxy GetDynamicMeshElements.

diff --git a/Engine/Source/Runtime/Engine/Private/PrimitiveSceneProxy.cpp b/Engine/Source/Runtime/Engine/Private/PrimitiveSceneProxy.cpp
--- a/Engine/Source/Runtime/Engine/Private/PrimitiveSceneProxy.cpp
+++ b/Engine/Source/Runtime/Engine/Private/PrimitiveSceneProxy.cpp
@@ -1002,7 +1002,7 @@ void FPrimitiveSceneProxy::SetDebugMassData(const TArray<FDebugMassData>& InDebu
  */
 void FPrimitiveSceneProxy::SetSelection_RenderThread(const bool bInParentSelected, const bool bInIndividuallySelected)
 {
-	check(IsInRenderingThread());
+	check(IsInParallelRenderingThread());
 
 	const bool bWasSelected = IsSelected();
 	bParentSelected = bInParentSelected;
diff --git a/Engine/Source/Runtime/Renderer/Private/RendererScene.cpp b/Engine/Source/Runtime/Renderer/Private/RendererScene.cpp
--- a/Engine/Source/Runtime/Renderer/Private/RendererScene.cpp
+++ b/Engine/Source/Runtime/Renderer/Private/RendererScene.cpp
@@ -2261,14 +2261,14 @@ void FScene::UpdatePrimitiveInstances(FUpdateInstanceCommand& UpdateParams)
 
 void FScene::UpdatePrimitiveSelectedState_RenderThread(const FPrimitiveSceneInfo* PrimitiveSceneInfo, bool bIsSelected)
 {
-	check(IsInRenderingThread());
+	check(IsInParallelRenderingThread());
 
 #if WITH_EDITOR
 	if (PrimitiveSceneInfo)
 	{
 		if (PrimitiveSceneInfo->GetIndex() != INDEX_NONE)
 		{
-			PrimitivesSelected[PrimitiveSceneInfo->GetIndex()] = bIsSelected;
+			PrimitivesSelected[PrimitiveSceneInfo->GetIndex()].AtomicSet(bIsSelected);
 		}
 	}
 #endif // WITH_EDITOR
