Set default value for ForceOnlyVirtualShadowMaps to 1

diff --git a/Engine/Source/Runtime/Renderer/Private/ShadowSetup.cpp b/Engine/Source/Runtime/Renderer/Private/ShadowSetup.cpp
--- a/Engine/Source/Runtime/Renderer/Private/ShadowSetup.cpp
+++ b/Engine/Source/Runtime/Renderer/Private/ShadowSetup.cpp
@@ -324,7 +324,7 @@ static TAutoConsoleVariable<int32> CVarClipmapUseConservativeCulling(
 
 static TAutoConsoleVariable<int32> CVarForceOnlyVirtualShadowMaps(
 	TEXT("r.Shadow.Virtual.ForceOnlyVirtualShadowMaps"),
-	0,
+	1,
 	TEXT("If enabled, disallow creation of conventional non-virtual shadow maps for any lights that get a virtual shadow map.\n")
 	TEXT("This can improve performance and save memory, but any geometric primitives that cannot be rendered into the virtual shadow map will not cast shadows."),
 	ECVF_RenderThreadSafe
