Fix r.Water.WaterInfo.RenderMethod on mobile devices

diff --git a/Engine/Shaders/Private/WaterInfoTextureMerge.usf b/Engine/Shaders/Private/WaterInfoTextureMerge.usf
--- a/Engine/Shaders/Private/WaterInfoTextureMerge.usf
+++ b/Engine/Shaders/Private/WaterInfoTextureMerge.usf
@@ -3,9 +3,13 @@
 #include "Common.ush"
 
 Texture2D<float4> GroundDepthTexture;
+SamplerState GroundDepthTextureSampler;
 Texture2D<float4> WaterBodyTexture;
+SamplerState WaterBodyTextureSampler;
 Texture2D<float4> WaterBodyDepthTexture;
+SamplerState WaterBodyDepthTextureSampler;
 Texture2D<float4> DilatedWaterBodyDepthTexture;
+SamplerState DilatedWaterBodyDepthTextureSampler;
 
 float2 WaterHeightExtents;
 float GroundZMin;
@@ -30,10 +34,13 @@ void Main(
 	out float4 OutColor : SV_Target0
 )
 {
-	const float2 WaterBodyVelocity = WaterBodyTexture.Load(int3(Position.xy, 0)).xy;
-	const float WaterBodyDeviceDepth = WaterBodyDepthTexture.Load(int3(Position.xy, 0)).x;
-	const float DilationDeviceDepth = DilatedWaterBodyDepthTexture.Load(int3(Position.xy, 0)).x;
-	const float GroundDeviceDepth = GroundDepthTexture.Load(int3(Position.xy, 0)).x;
+	const float2 ScreenPosition = SvPositionToScreenPosition(Position).xy;
+	float2 ScreenUV = ScreenPosition * View.ScreenPositionScaleBias.xy + View.ScreenPositionScaleBias.wz;
+
+	const float2 WaterBodyVelocity = WaterBodyTexture.SampleLevel(WaterBodyTextureSampler, ScreenUV, 0);
+	const float WaterBodyDeviceDepth =  WaterBodyDepthTexture.SampleLevel(WaterBodyDepthTextureSampler, ScreenUV, 0);
+	const float DilationDeviceDepth = DilatedWaterBodyDepthTexture.SampleLevel(DilatedWaterBodyDepthTextureSampler, ScreenUV, 0);
+	const float GroundDeviceDepth =  GroundDepthTexture.SampleLevel(GroundDepthTextureSampler, ScreenUV, 0);
 
 	const float GroundSceneDepth = ConvertFromDeviceZ(GroundDeviceDepth);
 	float GroundZ = CaptureZ - GroundSceneDepth;
diff --git a/Engine/Source/Runtime/Renderer/Private/WaterInfoTextureRendering.cpp b/Engine/Source/Runtime/Renderer/Private/WaterInfoTextureRendering.cpp
--- a/Engine/Source/Runtime/Renderer/Private/WaterInfoTextureRendering.cpp
+++ b/Engine/Source/Runtime/Renderer/Private/WaterInfoTextureRendering.cpp
@@ -54,9 +54,13 @@ public:
 		SHADER_PARAMETER_RDG_UNIFORM_BUFFER(FViewUniformShaderParameters, View)
 		SHADER_PARAMETER_STRUCT_INCLUDE(FSceneTextureShaderParameters, SceneTextures)
 		SHADER_PARAMETER_RDG_TEXTURE(Texture2D<float4>, GroundDepthTexture)
+		SHADER_PARAMETER_SAMPLER(SamplerState, GroundDepthTextureSampler)
 		SHADER_PARAMETER_RDG_TEXTURE(Texture2D<float4>, WaterBodyTexture)
+		SHADER_PARAMETER_SAMPLER(SamplerState, WaterBodyTextureSampler)
 		SHADER_PARAMETER_RDG_TEXTURE(Texture2D<float4>, WaterBodyDepthTexture)
+		SHADER_PARAMETER_SAMPLER(SamplerState, WaterBodyDepthTextureSampler)
 		SHADER_PARAMETER_RDG_TEXTURE(Texture2D<float4>, DilatedWaterBodyDepthTexture)
+		SHADER_PARAMETER_SAMPLER(SamplerState, DilatedWaterBodyDepthTextureSampler)
 		SHADER_PARAMETER(FVector2f, WaterHeightExtents)
 		SHADER_PARAMETER(float, GroundZMin)
 		SHADER_PARAMETER(float, CaptureZ)
@@ -970,9 +974,13 @@ void RenderWaterInfoTexture(
 			PassParameters->RenderTargets[0] = FRenderTargetBinding(MergedTexture, ERenderTargetLoadAction::ENoAction);
 			PassParameters->SceneTextures = GetSceneTextureShaderParameters(WaterView);
 			PassParameters->GroundDepthTexture = TerrainDepthBuffer;
+			PassParameters->GroundDepthTextureSampler = TStaticSamplerState<SF_Point, AM_Clamp, AM_Clamp, AM_Clamp>::GetRHI();
 			PassParameters->WaterBodyTexture = WaterInfoColorTexture;
+			PassParameters->WaterBodyTextureSampler = TStaticSamplerState<SF_Point, AM_Clamp, AM_Clamp, AM_Clamp>::GetRHI();
 			PassParameters->WaterBodyDepthTexture = WaterInfoDepthBuffer;
+			PassParameters->WaterBodyDepthTextureSampler = TStaticSamplerState<SF_Point, AM_Clamp, AM_Clamp, AM_Clamp>::GetRHI();
 			PassParameters->DilatedWaterBodyDepthTexture = DilatedDepthBuffer;
+			PassParameters->DilatedWaterBodyDepthTextureSampler = TStaticSamplerState<SF_Point, AM_Clamp, AM_Clamp, AM_Clamp>::GetRHI();
 			PassParameters->CaptureZ = RenderingParams.CaptureZ;
 			PassParameters->WaterHeightExtents = RenderingParams.WaterHeightExtents;
 			PassParameters->GroundZMin = RenderingParams.GroundZMin;
