Prevent UpdateStaticMeshesForMaterials from updating primitives which have not yet been added to the scene.

Processing a component during UWorld::SendAllEndOfFrameUpdates can trigger an update for another component which was just processed by the same function (e.g. creating the proxy adds a missing usage flag on the primitive's material, and another primitive which was just processed is using the same material). In that case, we'll try to call FPrimitiveSceneInfo::UpdateStaticMeshes() on a primitive which was just created, but hasn't yet been added to the scene, which crashes. It should be safe to skip primitives which don't have a valid index, since they'll be added to the scene after the material is updated, so they'll use the correct material data.

#rnx

diff --git a/Engine/Source/Runtime/Renderer/Private/RendererScene.cpp b/Engine/Source/Runtime/Renderer/Private/RendererScene.cpp
--- a/Engine/Source/Runtime/Renderer/Private/RendererScene.cpp
+++ b/Engine/Source/Runtime/Renderer/Private/RendererScene.cpp
@@ -6894,7 +6894,7 @@ void UpdateStaticMeshesForMaterials(const TArray<const FMaterial*>& MaterialReso
 	{
 		UPrimitiveComponent* PrimitiveComponent = *PrimitiveIt;
 
-		if (PrimitiveComponent->IsRenderStateCreated() && PrimitiveComponent->SceneProxy)
+		if (PrimitiveComponent->IsRenderStateCreated() && PrimitiveComponent->SceneProxy && PrimitiveComponent->SceneProxy->GetPrimitiveSceneInfo()->IsIndexValid())
 		{
 			UsedMaterialsDependencies.Reset();
 			UsedMaterials.Reset();
