Fix Local lights for some mobile devices drivers
Change UAV of LocalLightBuffer from SINT to UINT

diff --git a/Engine/Shaders/Private/MobileLocalLightsBuffer.usf b/Engine/Shaders/Private/MobileLocalLightsBuffer.usf
--- a/Engine/Shaders/Private/MobileLocalLightsBuffer.usf
+++ b/Engine/Shaders/Private/MobileLocalLightsBuffer.usf
@@ -1,6 +1,7 @@
 // Copyright Epic Games, Inc. All Rights Reserved.
 
 #define SUPPORT_CONTACT_SHADOWS 0
+#define MAX_UINT 4294967295
 
 #if !COMPUTE_SHADER
 #include "Common.ush"
@@ -19,7 +20,7 @@
 #if COMPUTE_SHADER
 
 uint2 GroupSize;
-RWBuffer<int> RWTileInfo;
+RWBuffer<uint> RWTileInfo;
 [numthreads(THREADGROUP_SIZEX, 1, 1)]
 void MainCS(
 	uint3 GroupId : SV_GroupID,
@@ -57,8 +58,8 @@ void MainCS(
 	}
 	else
 	{
-		RWTileInfo[2 * DispatchThreadId.x] = -1;
-		RWTileInfo[2 * DispatchThreadId.x + 1] = -1;
+		RWTileInfo[2 * DispatchThreadId.x] = MAX_UINT;
+		RWTileInfo[2 * DispatchThreadId.x + 1] = MAX_UINT;
 	}
 }
 
@@ -67,7 +68,7 @@ void MainCS(
 #if !LIGHT_FUNCTION
 
 int LightGridPixelSize;
-Buffer<int> TileInfo;
+Buffer<uint> TileInfo;
 void MainVS(
 	float2 TexCoord : ATTRIBUTE0,
 	uint VertexId : SV_VertexID,
@@ -81,7 +82,7 @@ void MainVS(
 	float2 ScreenPosition = (ScreenUV.xy - View.ScreenPositionScaleBias.wz) / View.ScreenPositionScaleBias.xy;
 	OutPosition = float4(ScreenPosition, 0, 1);
 
-	if(TileInfo[2 * InstanceId] < 0)
+	if(TileInfo[2 * InstanceId] == MAX_UINT)
 	{
 		OutPosition.xy = 0;
 	}
diff --git a/Engine/Source/Runtime/Renderer/Private/MobileLocalLightsBuffer.cpp b/Engine/Source/Runtime/Renderer/Private/MobileLocalLightsBuffer.cpp
--- a/Engine/Source/Runtime/Renderer/Private/MobileLocalLightsBuffer.cpp
+++ b/Engine/Source/Runtime/Renderer/Private/MobileLocalLightsBuffer.cpp
@@ -27,7 +27,7 @@ public:
 	SHADER_USE_PARAMETER_STRUCT(FLocalLightBufferCS, FGlobalShader);
 
 	BEGIN_SHADER_PARAMETER_STRUCT(FParameters, )
-		SHADER_PARAMETER_RDG_BUFFER_UAV(RWBuffer<int32>, RWTileInfo)
+		SHADER_PARAMETER_RDG_BUFFER_UAV(RWBuffer<uint>, RWTileInfo)
 		SHADER_PARAMETER_RDG_UNIFORM_BUFFER(FForwardLightData, ForwardLightData)
 		SHADER_PARAMETER(FIntPoint, GroupSize)
 	END_SHADER_PARAMETER_STRUCT()
@@ -55,7 +55,7 @@ class FLocalLightBufferVS : public FGlobalShader
 
 	BEGIN_SHADER_PARAMETER_STRUCT(FParameters, )
 		SHADER_PARAMETER_STRUCT_REF(FViewUniformShaderParameters, View)
-		SHADER_PARAMETER_RDG_BUFFER_SRV(Buffer<int32>, TileInfo)
+		SHADER_PARAMETER_RDG_BUFFER_SRV(Buffer<uint>, TileInfo)
 		SHADER_PARAMETER(int32, LightGridPixelSize)
 	END_SHADER_PARAMETER_STRUCT()
 
@@ -247,9 +247,9 @@ void FMobileSceneRenderer::RenderMobileLocalLightsBuffer(FRDGBuilder& GraphBuild
 			FMath::DivideAndRoundUp(View.ViewRect.Size().X, LightGridPixelSize),
 			FMath::DivideAndRoundUp(View.ViewRect.Size().Y, LightGridPixelSize));
 
-		FRDGBufferRef TileInfoBuffer = GraphBuilder.CreateBuffer(FRDGBufferDesc::CreateBufferDesc(sizeof(int32), 2 * GroupSize.X * GroupSize.Y), TEXT("TileInfoBuffer"));
-		FRDGBufferUAVRef TileInfoBufferUAV = GraphBuilder.CreateUAV(TileInfoBuffer, PF_R32_SINT);
-		FRDGBufferSRVRef TileInfoBufferSRV = GraphBuilder.CreateSRV(TileInfoBuffer, PF_R32_SINT);
+		FRDGBufferRef TileInfoBuffer = GraphBuilder.CreateBuffer(FRDGBufferDesc::CreateBufferDesc(sizeof(uint32), 2 * GroupSize.X * GroupSize.Y), TEXT("TileInfoBuffer"));
+		FRDGBufferUAVRef TileInfoBufferUAV = GraphBuilder.CreateUAV(TileInfoBuffer, PF_R32_UINT);
+		FRDGBufferSRVRef TileInfoBufferSRV = GraphBuilder.CreateSRV(TileInfoBuffer, PF_R32_UINT);
 
 		{
 			auto* PassParameters = GraphBuilder.AllocParameters<FLocalLightBufferCS::FParameters>();
