LFV - added start distance as an optimisation.

diff --git a/Engine/Shaders/Private/LocalFogVolumes/LocalFogVolumeCommon.ush b/Engine/Shaders/Private/LocalFogVolumes/LocalFogVolumeCommon.ush
--- a/Engine/Shaders/Private/LocalFogVolumes/LocalFogVolumeCommon.ush
+++ b/Engine/Shaders/Private/LocalFogVolumes/LocalFogVolumeCommon.ush
@@ -43,6 +43,7 @@ struct FLocalFogVolumeGPUInstanceData
 #define LFVTilePixelSize				LFVStruct.LFV.LocalFogVolumeTilePixelSize
 #define LFVMaxDensityIntoVolumetricFog	LFVStruct.LFV.LocalFogVolumeMaxDensityIntoVolumetricFog
 #define LFVRenderInVolumetricFog		LFVStruct.LFV.ShouldRenderLocalFogVolumeInVolumetricFog
+#define LFVGlobalStartDistance			LFVStruct.LFV.GlobalStartDistance
 
 #define LFVDirectionalLightColor		LFVStruct.LFV.DirectionalLightColor
 #define LFVDirectionalLightDirection	LFVStruct.LFV.DirectionalLightDirection
@@ -54,6 +55,7 @@ uint									LFV_LocalFogVolumeInstanceCount;
 uint									LFV_LocalFogVolumeTilePixelSize;
 float									LFV_LocalFogVolumeMaxDensityIntoVolumetricFog;
 uint									LFV_ShouldRenderLocalFogVolumeInVolumetricFog;
+float									LFV_GlobalStartDistance;
 
 float3									LFV_DirectionalLightColor;
 float3									LFV_DirectionalLightDirection;
@@ -69,6 +71,7 @@ Texture2DArray<uint>					LFV_LocalFogVolumeTileDataTexture;
 #define LFVTilePixelSize				LFV_LocalFogVolumeTilePixelSize
 #define LFVMaxDensityIntoVolumetricFog	LFV_LocalFogVolumeMaxDensityIntoVolumetricFog
 #define LFVRenderInVolumetricFog		LFV_ShouldRenderLocalFogVolumeInVolumetricFog
+#define LFVGlobalStartDistance			LFV_GlobalStartDistance
 
 #define LFVDirectionalLightColor		LFV_DirectionalLightColor
 #define LFVDirectionalLightDirection	LFV_DirectionalLightDirection
@@ -312,7 +315,7 @@ float4 GetLFVInstanceContribution(
 	float2 TsU = RayIntersectSphere(RayPosU, RayDirU, float4(0.0, 0.0, 0.0, 1.0));
 
 	//Clamp to make sure we only consider intersection in front of the camera
-	TsU = max(0.0f.xx, TsU);
+	TsU = max(LFVGlobalStartDistance.xx / FogInstance.UniformScale, TsU);
 
 	// Clamp the traced distance to the depth buffer and compute world length.
 	float LengthDepthBufferU = length(DepthBufferPosU - RayPosU);
diff --git a/Engine/Shaders/Private/VolumetricFog.usf b/Engine/Shaders/Private/VolumetricFog.usf
--- a/Engine/Shaders/Private/VolumetricFog.usf
+++ b/Engine/Shaders/Private/VolumetricFog.usf
@@ -158,86 +158,89 @@ void MaterialSetupCS(
 	float3 AdditionalEmissive = 0.0f;
 
 #if USE_LOCAL_FOG_VOLUMES
-	float3 CamRayTranslatedWorldOrigin = PrimaryView.TranslatedWorldCameraOrigin;
-	float3 PreViewTranslation = LWCHackToFloat(PrimaryView.PreViewTranslation);
-	//float3 TranslatedWorldPosition = WorldPosition + PreViewTranslation;
-
 	float CellOffset = .5f;
 	float SceneDepth;
 	float3 TranslatedWorldPosition = ComputeCellTranslatedWorldPosition(GridCoordinate, CellOffset, SceneDepth);
+	if(length(TranslatedWorldPosition - PrimaryView.TranslatedWorldCameraOrigin) > LFVGlobalStartDistance)
+	{
+		float3 CamRayTranslatedWorldOrigin = PrimaryView.TranslatedWorldCameraOrigin;
+		float3 PreViewTranslation = LWCHackToFloat(PrimaryView.PreViewTranslation);
+		//float3 TranslatedWorldPosition = WorldPosition + PreViewTranslation;
 
-	float4 VertexClipSpacePosition = mul(float4(TranslatedWorldPosition.xyz, 1), PrimaryView.TranslatedWorldToClip);
-	float2 SvPosition = (VertexClipSpacePosition.xy / VertexClipSpacePosition.w * float2(.5f, -.5f) + .5f) * PrimaryView.ViewSizeAndInvSize.xy;
-	uint2 TilePos = clamp(uint2(SvPosition.xy / float(LFVTilePixelSize)), uint2(0, 0), LFVTileDataResolution - 1);
 
-	uint LFVCount = LFVTileDataTex[uint3(TilePos, 0)];
+		float4 VertexClipSpacePosition = mul(float4(TranslatedWorldPosition.xyz, 1), PrimaryView.TranslatedWorldToClip);
+		float2 SvPosition = (VertexClipSpacePosition.xy / VertexClipSpacePosition.w * float2(.5f, -.5f) + .5f) * PrimaryView.ViewSizeAndInvSize.xy;
+		uint2 TilePos = clamp(uint2(SvPosition.xy / float(LFVTilePixelSize)), uint2(0, 0), LFVTileDataResolution - 1);
 
-	for (int LFVIndex = 0; LFVIndex < LFVCount; ++LFVIndex)
-	{
-		uint InstanceIndex = LFVTileDataTex[uint3(TilePos, 1 + LFVIndex)];
+		uint LFVCount = LFVTileDataTex[uint3(TilePos, 0)];
 
-		float3 CamRayWorldOrigin = CamRayTranslatedWorldOrigin - PreViewTranslation;
-		float3 CamRayTranslatedWorldDir = normalize(TranslatedWorldPosition - CamRayTranslatedWorldOrigin);
-		float3 CamRayWorldDir = CamRayTranslatedWorldDir;
+		for (int LFVIndex = 0; LFVIndex < LFVCount; ++LFVIndex)
+		{
+			uint InstanceIndex = LFVTileDataTex[uint3(TilePos, 1 + LFVIndex)];
 
-		// LFV_TODO skip if out of sphere of influence
-		FLocalFogVolumeGPUInstanceData FogInstance = GetLocalFogVolumeGPUInstanceData(InstanceIndex);
+			float3 CamRayWorldOrigin = CamRayTranslatedWorldOrigin - PreViewTranslation;
+			float3 CamRayTranslatedWorldDir = normalize(TranslatedWorldPosition - CamRayTranslatedWorldOrigin);
+			float3 CamRayWorldDir = CamRayTranslatedWorldDir;
 
-		float3 UnitSpacePos = mul(float4(WorldPosition, 1.0), FogInstance.InvTransform).xyz;
+			// LFV_TODO skip if out of sphere of influence
+			FLocalFogVolumeGPUInstanceData FogInstance = GetLocalFogVolumeGPUInstanceData(InstanceIndex);
 
-		float DistToCenter = length(UnitSpacePos - float3(0.0f, 0.0f, 0.0f));
-		float SphereFade = saturate(1.0f - DistToCenter * DistToCenter);
+			float3 UnitSpacePos = mul(float4(WorldPosition, 1.0), FogInstance.InvTransform).xyz;
 
-		// Evaluate both radial and height local fog contribution
-		float LFVExtinctionRadial = SphereFade <= 0.0f ? 0.0f : FogInstance.HeightFogExtinction * exp(-FogInstance.HeightFogFalloff * (UnitSpacePos.z - FogInstance.HeightFogOffset));
-		float LFVExtinctionHeight = FogInstance.RadialFogExtinction * SphereFade;
+			float DistToCenter = length(UnitSpacePos - float3(0.0f, 0.0f, 0.0f));
+			float SphereFade = saturate(1.0f - DistToCenter * DistToCenter);
 
-		// As for the local fog volume analitical evaluation, we want to combine both two optical depth for the "opacity"=1-transmittance to be combined together.
-		// This helps to have soft edges when using height fog. 
-		/*
-		In this case, we have to find math so that the combined extinction resulting in a transmittance representing the multiplication of coverage = 1 - transmittance factors.
+			// Evaluate both radial and height local fog contribution
+			float LFVExtinctionRadial = SphereFade <= 0.0f ? 0.0f : FogInstance.HeightFogExtinction * exp(-FogInstance.HeightFogFalloff * (UnitSpacePos.z - FogInstance.HeightFogOffset));
+			float LFVExtinctionHeight = FogInstance.RadialFogExtinction * SphereFade;
 
-		Reminder
-			Transmittance as T
-			Optical Depth as Tau
-			T   =  exp(-Tau)
-			Tau = -log(T)		(*)
+			// As for the local fog volume analitical evaluation, we want to combine both two optical depth for the "opacity"=1-transmittance to be combined together.
+			// This helps to have soft edges when using height fog. 
+			/*
+			In this case, we have to find math so that the combined extinction resulting in a transmittance representing the multiplication of coverage = 1 - transmittance factors.
 
-		We want to combine the both mode transmittance Tr (radial) and Th (height) as done in LocalFogVolumeEvaluateAnalyticalIntegral to have
-			 T		= 1 - (1-Tr)(1-Th)											// Where (1-Tr) is virtual coverage of the medium expressed from transmittance for the radial fog.
-			 T		= 1 - (1 - exp(-Tau_r)) (1 - exp(-Tau_h)))					// Where Tr = exp(-Tau_r)
-			-Tau	= log(1 - (1 - exp(-Tau_r)) (1 - exp(-Tau_h)))				// Using (*)
-			 Tau	=-log(exp(-Tau_r) - exp((-Tau_r-Tau_h) + exp(-Tau_h)))
+			Reminder
+				Transmittance as T
+				Optical Depth as Tau
+				T   =  exp(-Tau)
+				Tau = -log(T)		(*)
 
-		Optical Depth Tau = Sigma_t * dt so it is linear with respect to dt and extinction Sigma_t. So we can also write
-			Sigma_T = -log(exp(-Sigma_T_r) - exp((-Sigma_T_r-Sigma_T_h) + exp(-Sigma_T_h)))
+			We want to combine the both mode transmittance Tr (radial) and Th (height) as done in LocalFogVolumeEvaluateAnalyticalIntegral to have
+				 T		= 1 - (1-Tr)(1-Th)											// Where (1-Tr) is virtual coverage of the medium expressed from transmittance for the radial fog.
+				 T		= 1 - (1 - exp(-Tau_r)) (1 - exp(-Tau_h)))					// Where Tr = exp(-Tau_r)
+				-Tau	= log(1 - (1 - exp(-Tau_r)) (1 - exp(-Tau_h)))				// Using (*)
+				 Tau	=-log(exp(-Tau_r) - exp((-Tau_r-Tau_h) + exp(-Tau_h)))
 
-		Factoring some values a bit we can get
-			A		= exp(-Sigma_T_r)
-			B		= exp(-Sigma_T_h)
-			Sigma_T	= -log(A - A*B + B)
-		*/
+			Optical Depth Tau = Sigma_t * dt so it is linear with respect to dt and extinction Sigma_t. So we can also write
+				Sigma_T = -log(exp(-Sigma_T_r) - exp((-Sigma_T_r-Sigma_T_h) + exp(-Sigma_T_h)))
 
-		// Optical Depth is linear with Extinction so there is no need to account for the step length.
-		float FactorA = exp(-LFVExtinctionHeight);
-		float FactorB = exp(-LFVExtinctionRadial);
-		float LFVExtinction = -log(FactorA - FactorA* FactorB + FactorB);
-		LFVExtinction = max(0.0f, LFVExtinction);
+			Factoring some values a bit we can get
+				A		= exp(-Sigma_T_r)
+				B		= exp(-Sigma_T_h)
+				Sigma_T	= -log(A - A*B + B)
+			*/
 
-		// Convert extinction from unit volume 1/m into extinction for volumetric fog as 1/cm
-		LFVExtinction *= (1.0f / METER_TO_CENTIMETER);
+			// Optical Depth is linear with Extinction so there is no need to account for the step length.
+			float FactorA = exp(-LFVExtinctionHeight);
+			float FactorB = exp(-LFVExtinctionRadial);
+			float LFVExtinction = -log(FactorA - FactorA * FactorB + FactorB);
+			LFVExtinction = max(0.0f, LFVExtinction);
 
-		LFVExtinction = min(LFVMaxDensityIntoVolumetricFog, LFVExtinction);
+			// Convert extinction from unit volume 1/m into extinction for volumetric fog as 1/cm
+			LFVExtinction *= (1.0f / METER_TO_CENTIMETER);
 
-		float3 LFVScattering = LFVExtinction * FogInstance.Albedo;
+			LFVExtinction = min(LFVMaxDensityIntoVolumetricFog, LFVExtinction);
 
-		Scattering += LFVScattering;
-		Absorption += max(LFVExtinction - Luminance(LFVScattering), 0.0f);
+			float3 LFVScattering = LFVExtinction * FogInstance.Albedo;
 
-#if USE_EMISSIVE
-		// Emissive follow the density of the medium
-		AdditionalEmissive += LFVExtinction * FogInstance.Emissive;
-#endif
+			Scattering += LFVScattering;
+			Absorption += max(LFVExtinction - Luminance(LFVScattering), 0.0f);
+
+	#if USE_EMISSIVE
+			// Emissive follow the density of the medium
+			AdditionalEmissive += LFVExtinction * FogInstance.Emissive;
+	#endif
+		}
 	}
 #endif // USE_LOCAL_FOG_VOLUMES
 
diff --git a/Engine/Source/Runtime/Renderer/Private/FogRendering.cpp b/Engine/Source/Runtime/Renderer/Private/FogRendering.cpp
--- a/Engine/Source/Runtime/Renderer/Private/FogRendering.cpp
+++ b/Engine/Source/Runtime/Renderer/Private/FogRendering.cpp
@@ -288,8 +288,7 @@ static void RenderViewFog(
 	FIntRect ViewRect, 
 	FFogPassParameters* PassParameters, 
 	bool bShouldRenderVolumetricFog,
-	bool bFogComposeLocalFogVolumes,
-	bool bSkipDepthBound = false)
+	bool bFogComposeLocalFogVolumes)
 {
 	FGraphicsPipelineStateInitializer GraphicsPSOInit;
 	RHICmdList.ApplyCachedRenderTargets(GraphicsPSOInit);
@@ -316,10 +315,10 @@ static void RenderViewFog(
 	GraphicsPSOInit.BoundShaderState.PixelShaderRHI = PixelShader.GetPixelShader();
 
 	// Setup the depth bound optimization if possible on that platform.
-	GraphicsPSOInit.bDepthBounds = GSupportsDepthBoundsTest && CVarFogUseDepthBounds.GetValueOnAnyThread() && !bSkipDepthBound;
+	GraphicsPSOInit.bDepthBounds = GSupportsDepthBoundsTest && CVarFogUseDepthBounds.GetValueOnAnyThread();
 	if (GraphicsPSOInit.bDepthBounds)
 	{
-		float FogStartDistance = GetViewFogCommonStartDistance(View, bShouldRenderVolumetricFog);
+		float FogStartDistance = GetViewFogCommonStartDistance(View, bShouldRenderVolumetricFog, bFogComposeLocalFogVolumes);
 
 		// Here we compute the nearest z value the fog can start
 		// to skip shader execution on pixels that are closer.
@@ -376,9 +375,6 @@ void FDeferredShadingSceneRenderer::RenderFog(
 
 		const bool bShouldRenderVolumetricFog = ShouldRenderVolumetricFog();
 
-		// Without volumetric fog, LFVs would not be rendered correctly according to depth and height fog start distance. Since LFVs do not have any start distance as of today.
-		const bool bSkipDepthBound = !bShouldRenderVolumetricFog && bFogComposeLocalFogVolumes && ShouldRenderLocalFogVolumeDuringHeightFogPass(Scene, ViewFamily);
-
 		for(int32 ViewIndex = 0;ViewIndex < Views.Num();ViewIndex++)
 		{
 			const FViewInfo& View = Views[ViewIndex];
@@ -398,9 +394,9 @@ void FDeferredShadingSceneRenderer::RenderFog(
 				PassParameters->RenderTargets.DepthStencil = FDepthStencilBinding(SceneTextures.Depth.Target, ERenderTargetLoadAction::ELoad, ERenderTargetLoadAction::ELoad, FExclusiveDepthStencil::DepthRead_StencilWrite);
 
 				GraphBuilder.AddPass(RDG_EVENT_NAME("Fog"), PassParameters, ERDGPassFlags::Raster, 
-					[this, &View, PassParameters, bShouldRenderVolumetricFog, bFogComposeLocalFogVolumes, bSkipDepthBound](FRHICommandList& RHICmdList)
+					[this, &View, PassParameters, bShouldRenderVolumetricFog, bFogComposeLocalFogVolumes](FRHICommandList& RHICmdList)
 				{
-					RenderViewFog(RHICmdList, View, View.ViewRect, PassParameters, bShouldRenderVolumetricFog, bFogComposeLocalFogVolumes, bSkipDepthBound);
+					RenderViewFog(RHICmdList, View, View.ViewRect, PassParameters, bShouldRenderVolumetricFog, bFogComposeLocalFogVolumes);
 				});
 			}
 		}
@@ -473,13 +469,21 @@ float GetFogDefaultStartDistance()
 	return 30.0f;
 }
 
-float GetViewFogCommonStartDistance(const FViewInfo& View, bool bShouldRenderVolumetricFog)
+float GetViewFogCommonStartDistance(const FViewInfo& View, bool bShouldRenderVolumetricFog, bool bShouldRenderLocalFogVolumes)
 {
 	float ExpFogStartDistance = View.ExponentialFogParameters.W;
 	float VolFogStartDistance = bShouldRenderVolumetricFog ? View.VolumetricFogStartDistance : ExpFogStartDistance;
 
 	// The fog can be set to start at a certain euclidean distance.
 	// clamp the value to be behind the near plane z, according to the smallest distance between volumetric fog and height fog (if they are enabled). 
-	float FogCommonStartDistance = FMath::Max(GetFogDefaultStartDistance(), FMath::Min(ExpFogStartDistance, VolFogStartDistance));
+	float FogCommonStartDistance = FMath::Min(ExpFogStartDistance, VolFogStartDistance);
+
+	if (bShouldRenderLocalFogVolumes)
+	{
+		FogCommonStartDistance = FMath::Min(GetLocalFogVolumeGlobalStartDistance(), FogCommonStartDistance);
+	}
+
+	FogCommonStartDistance = FMath::Max(GetFogDefaultStartDistance(), FogCommonStartDistance);
+
 	return FogCommonStartDistance;
 }
diff --git a/Engine/Source/Runtime/Renderer/Private/FogRendering.h b/Engine/Source/Runtime/Renderer/Private/FogRendering.h
--- a/Engine/Source/Runtime/Renderer/Private/FogRendering.h
+++ b/Engine/Source/Runtime/Renderer/Private/FogRendering.h
@@ -36,5 +36,5 @@ TRDGUniformBufferRef<FFogUniformParameters> CreateFogUniformBuffer(FRDGBuilder&
 
 extern bool ShouldRenderFog(const FSceneViewFamily& Family);
 
-extern float GetViewFogCommonStartDistance(const FViewInfo& View, bool bShouldRenderVolumetricFog);
+extern float GetViewFogCommonStartDistance(const FViewInfo& View, bool bShouldRenderVolumetricFog, bool bShouldRenderLocalFogVolumes);
 extern float GetFogDefaultStartDistance();
diff --git a/Engine/Source/Runtime/Renderer/Private/LocalFogVolumeRendering.cpp b/Engine/Source/Runtime/Renderer/Private/LocalFogVolumeRendering.cpp
--- a/Engine/Source/Runtime/Renderer/Private/LocalFogVolumeRendering.cpp
+++ b/Engine/Source/Runtime/Renderer/Private/LocalFogVolumeRendering.cpp
@@ -11,7 +11,7 @@
 // The runtime ON/OFF toggle
 static TAutoConsoleVariable<int32> CVarLocalFogVolume(
 	TEXT("r.LocalFogVolume"), 1,
-	TEXT("LocalFogVolume components are rendered when this is not 0, otherwise ignored.\n"),
+	TEXT("LocalFogVolume components are rendered when this is not 0, otherwise ignored."),
 	ECVF_RenderThreadSafe);
 
 // The project setting (disable runtime and shader code)
@@ -23,42 +23,47 @@ static TAutoConsoleVariable<int32> CVarSupportLocalFogVolumes(
 
 static TAutoConsoleVariable<int32> CVarLocalFogVolumeRenderDuringHeightFogPass(
 	TEXT("r.LocalFogVolume.RenderDuringHeightFogPass"), 0,
-	TEXT("LocalFogVolume are going to be rendered during the height fog pass, skipping the tiled rendering pass specific to them. Only work on the non mobile path as an experiment.\n"),
+	TEXT("LocalFogVolume are going to be rendered during the height fog pass, skipping the tiled rendering pass specific to them. Only work on the non mobile path as an experiment."),
 	ECVF_RenderThreadSafe);
 
 static TAutoConsoleVariable<int32> CVarLocalFogVolumeRenderIntoVolumetricFog(
 	TEXT("r.LocalFogVolume.RenderIntoVolumetricFog"), 1,
-	TEXT("LocalFogVolume are going to be voxelised into the volumetric fog when this is not 0, otherwise it will remain isolated.\n"),
+	TEXT("LocalFogVolume are going to be voxelised into the volumetric fog when this is not 0, otherwise it will remain isolated."),
 	ECVF_RenderThreadSafe);
 
 static TAutoConsoleVariable<float> CVarLocalFogVolumeMaxDensityIntoVolumetricFog(
 	TEXT("r.LocalFogVolume.MaxDensityIntoVolumetricFog"), 0.01f,
-	TEXT("LocalFogVolume height fog mode can become exponentially dense in the bottom part. VolumetricFog temporal reprojection then can leak du to high density. Clamping density is a way to get that visual artefact under control.\n"),
+	TEXT("LocalFogVolume height fog mode can become exponentially dense in the bottom part. VolumetricFog temporal reprojection then can leak du to high density. Clamping density is a way to get that visual artefact under control."),
 	ECVF_RenderThreadSafe);
 
 static TAutoConsoleVariable<int32> CVarLocalFogVolumeApplyOnTranslucent(
 	TEXT("r.LocalFogVolume.ApplyOnTranslucent"), 0,
-	TEXT("Project settings enabling the sampling of local fog volumes on translucent elements.\n"),
+	TEXT("Project settings enabling the sampling of local fog volumes on translucent elements."),
 	ECVF_ReadOnly | ECVF_RenderThreadSafe);
 
 static TAutoConsoleVariable<int32> CVarLocalFogVolumeTilePixelSize(
 	TEXT("r.LocalFogVolume.TilePixelSize"), 128,
-	TEXT("Tile size on screen in pixel at which we cull the local fog volumes.\n"),
+	TEXT("Tile size on screen in pixel at which we cull the local fog volumes."),
 	ECVF_RenderThreadSafe);
 
 static TAutoConsoleVariable<int32> CVarLocalFogVolumeTileMaxInstanceCount(
 	TEXT("r.LocalFogVolume.TileMaxInstanceCount"), 32,
-	TEXT("Maximum number of local fog volume to account for per view (and per tile or consistency).\n"),
+	TEXT("Maximum number of local fog volume to account for per view (and per tile or consistency)."),
 	ECVF_RenderThreadSafe);
 
 static TAutoConsoleVariable<int32> CVarLocalFogVolumeTileCullingUseAsync(
 	TEXT("r.LocalFogVolume.TileCullingUseAsync"), 1,
-	TEXT("True if we want to try and use culling on the async pipe.\n"),
+	TEXT("True if we want to try and use culling on the async pipe."),
 	ECVF_RenderThreadSafe);
 
 static TAutoConsoleVariable<int32> CVarLocalFogVolumeTileDebug(
 	TEXT("r.LocalFogVolume.TileDebug"), 0,
-	TEXT("Debug the tiled rendering data complexity. 1: show per tile LFV count as color ; 2: same as one but also show the effect of pixel discard/clipping.\n"),
+	TEXT("Debug the tiled rendering data complexity. 1: show per tile LFV count as color ; 2: same as one but also show the effect of pixel discard/clipping."),
+	ECVF_RenderThreadSafe);
+
+static TAutoConsoleVariable<float> CVarLocalFogVolumeGlobalStartDistance(
+	TEXT("r.LocalFogVolume.GlobalStartDistance"), 2000.0f,
+	TEXT("The start distance in centimeter from which local fog volumes starts to appear."),
 	ECVF_RenderThreadSafe);
 
 // Example of tile setup
@@ -119,6 +124,11 @@ bool ShouldRenderLocalFogVolumeInVolumetricFog(const FScene* Scene, const FScene
 	return false;
 }
 
+float GetLocalFogVolumeGlobalStartDistance()
+{
+	return FMath::Max(0.0f, CVarLocalFogVolumeGlobalStartDistance.GetValueOnRenderThread());
+}
+
 
 DECLARE_GPU_STAT(LocalFogVolumeVolumes);
 
@@ -154,6 +164,7 @@ void SetDummyLocalFogVolumeForView(FRDGBuilder& GraphBuilder, FViewInfo& View)
 	View.LocalFogVolumeViewData.UniformParametersStruct.LocalFogVolumeCommon.LocalFogVolumeInstances					= View.LocalFogVolumeViewData.GPUInstanceDataBufferSRV;
 	View.LocalFogVolumeViewData.UniformParametersStruct.LocalFogVolumeCommon.DirectionalLightColor						= FVector3f::Zero();
 	View.LocalFogVolumeViewData.UniformParametersStruct.LocalFogVolumeCommon.DirectionalLightDirection					= FVector3f::Zero();
+	View.LocalFogVolumeViewData.UniformParametersStruct.LocalFogVolumeCommon.GlobalStartDistance						= 0.0f;
 	View.LocalFogVolumeViewData.UniformParametersStruct.LocalFogVolumeTileDataTexture									= View.LocalFogVolumeViewData.TileDataTextureArraySRV;
 	View.LocalFogVolumeViewData.UniformBuffer																			= GraphBuilder.CreateUniformBuffer(&View.LocalFogVolumeViewData.UniformParametersStruct);
 
@@ -462,6 +473,7 @@ void CreateViewLocalFogVolumeBufferSRV(const FScene* Scene, FViewInfo& View, FRD
 	View.LocalFogVolumeViewData.UniformParametersStruct.LocalFogVolumeCommon.LocalFogVolumeMaxDensityIntoVolumetricFog	= GetLocalFogVolumeMaxDensityIntoVolumetricFog();
 	View.LocalFogVolumeViewData.UniformParametersStruct.LocalFogVolumeCommon.ShouldRenderLocalFogVolumeInVolumetricFog	= bShouldRenderLocalFogVolumeInVolumetricFog ? 1 : 0;
 	View.LocalFogVolumeViewData.UniformParametersStruct.LocalFogVolumeCommon.LocalFogVolumeInstances					= View.LocalFogVolumeViewData.GPUInstanceDataBufferSRV;
+	View.LocalFogVolumeViewData.UniformParametersStruct.LocalFogVolumeCommon.GlobalStartDistance						= GetLocalFogVolumeGlobalStartDistance();
 	if (IsMobilePlatform(View.GetShaderPlatform()))
 	{
 		// On mobile there is a separate FMobileDirectionalLightShaderParameters UB which holds all directional light data.
@@ -611,6 +623,40 @@ BEGIN_SHADER_PARAMETER_STRUCT(FLocalFogVolumeTiledPassParameters, )
 	RDG_BUFFER_ACCESS(TileDrawIndirectBuffer, ERHIAccess::IndirectArgs)
 	RENDER_TARGET_BINDING_SLOTS()
 END_SHADER_PARAMETER_STRUCT()
+
+static void ApplyDepthBoundIfNeeded(FGraphicsPipelineStateInitializer& GraphicsPSOInit, FMatrix ViewProjectionMatrix, FMatrix ViewInvProjectionMatrix, FRHICommandList& RHICmdList)
+{
+	GraphicsPSOInit.bDepthBounds = GSupportsDepthBoundsTest;
+	if (GraphicsPSOInit.bDepthBounds)
+	{
+		float FogStartDistance = GetLocalFogVolumeGlobalStartDistance();
+
+		// Here we compute the nearest z value the fog can start
+		// to skip shader execution on pixels that are closer.
+		// This means with a bigger distance specified more pixels are
+		// are culled and don't need to be rendered. This is faster if
+		// there is opaque content nearer than the computed z.
+		// This optimization is achieved using depth bound tests.
+		// Mobile platforms typically does not support that feature 
+		// but typically renders the world using forward shading 
+		// with height fog evaluated as part of the material vertex or pixel shader.
+		FVector ViewSpaceCorner = ViewInvProjectionMatrix.TransformFVector4(FVector4(1, 1, 1, 1));
+		float Ratio = ViewSpaceCorner.Z / ViewSpaceCorner.Size();
+		FVector ViewSpaceStartFogPoint(0.0f, 0.0f, FogStartDistance * Ratio);
+		FVector4f ClipSpaceMaxDistance = (FVector4f)ViewProjectionMatrix.TransformPosition(ViewSpaceStartFogPoint); // LWC_TODO: precision loss
+		float FogClipSpaceZ = ClipSpaceMaxDistance.Z / ClipSpaceMaxDistance.W;
+		FogClipSpaceZ = FMath::Clamp(FogClipSpaceZ, 0.f, 1.f);
+
+		if (bool(ERHIZBuffer::IsInverted))
+		{
+			RHICmdList.SetDepthBounds(0.0f, FogClipSpaceZ);
+		}
+		else
+		{
+			RHICmdList.SetDepthBounds(FogClipSpaceZ, 1.0f);
+		}
+	}
+}
  
 void RenderLocalFogVolume(
 	const FScene* Scene,
@@ -647,6 +693,7 @@ void RenderLocalFogVolume(
 			PassParameters->SceneTextures = SceneTextures.UniformBuffer;
 			PassParameters->TileDrawIndirectBuffer = View.LocalFogVolumeViewData.GPUTileDrawIndirectBuffer;
 			PassParameters->RenderTargets[0] = FRenderTargetBinding(SceneColorTexture, ERenderTargetLoadAction::ENoAction);
+			PassParameters->RenderTargets.DepthStencil = FDepthStencilBinding(SceneTextures.Depth.Target, ERenderTargetLoadAction::ELoad, ERenderTargetLoadAction::ELoad, FExclusiveDepthStencil::DepthRead_StencilWrite);
 
 			FLocalFogVolumeTiledRenderVS::FPermutationDomain VSPermutationVector;
 			auto VertexShader = View.ShaderMap->GetShader< FLocalFogVolumeTiledRenderVS >(VSPermutationVector);
@@ -655,6 +702,8 @@ void RenderLocalFogVolume(
 			auto PixelShader = View.ShaderMap->GetShader< FLocalFogVolumeTiledRenderPS >(PsPermutationVector);
 
 			const FIntRect ViewRect = View.ViewRect;
+			const FMatrix ViewProjectionMatrix = View.ViewMatrices.GetProjectionMatrix();
+			const FMatrix ViewInvProjectionMatrix = View.ViewMatrices.GetInvProjectionMatrix();
 
 			ClearUnusedGraphResources(VertexShader, &PassParameters->VS);
 			ClearUnusedGraphResources(PixelShader, &PassParameters->PS);
@@ -665,7 +714,7 @@ void RenderLocalFogVolume(
 				RDG_EVENT_NAME("LocalFogVolume.Tiled (%u X %u)", LocalFogVolumeTileDataTextureResolution.X, LocalFogVolumeTileDataTextureResolution.Y),
 				PassParameters,
 				ERDGPassFlags::Raster,
-				[VertexShader, PixelShader, PassParameters, ViewRect](FRHICommandList& RHICmdList)
+				[VertexShader, PixelShader, PassParameters, ViewRect, ViewProjectionMatrix, ViewInvProjectionMatrix](FRHICommandList& RHICmdList)
 			{
 				FGraphicsPipelineStateInitializer GraphicsPSOInit;
 				RHICmdList.ApplyCachedRenderTargets(GraphicsPSOInit);
@@ -682,6 +731,8 @@ void RenderLocalFogVolume(
 				GraphicsPSOInit.BoundShaderState.VertexShaderRHI = VertexShader.GetVertexShader();
 				GraphicsPSOInit.BoundShaderState.PixelShaderRHI = PixelShader.GetPixelShader();
 
+				ApplyDepthBoundIfNeeded(GraphicsPSOInit, ViewProjectionMatrix, ViewInvProjectionMatrix, RHICmdList);
+
 				SetGraphicsPipelineState(RHICmdList, GraphicsPSOInit, 0);
 
 				SetShaderParameters(RHICmdList, VertexShader, VertexShader.GetVertexShader(), PassParameters->VS);
@@ -799,6 +850,8 @@ void RenderLocalFogVolumeMobile(
 	GraphicsPSOInit.BoundShaderState.VertexShaderRHI = VertexShader.GetVertexShader();
 	GraphicsPSOInit.BoundShaderState.PixelShaderRHI = PixelShader.GetPixelShader();
 
+	ApplyDepthBoundIfNeeded(GraphicsPSOInit, View.ViewMatrices.GetProjectionMatrix(), View.ViewMatrices.GetInvProjectionMatrix(), RHICmdList);
+
 	SetGraphicsPipelineState(RHICmdList, GraphicsPSOInit, 0);
 
 
diff --git a/Engine/Source/Runtime/Renderer/Private/LocalFogVolumeRendering.h b/Engine/Source/Runtime/Renderer/Private/LocalFogVolumeRendering.h
--- a/Engine/Source/Runtime/Renderer/Private/LocalFogVolumeRendering.h
+++ b/Engine/Source/Runtime/Renderer/Private/LocalFogVolumeRendering.h
@@ -26,6 +26,7 @@ BEGIN_SHADER_PARAMETER_STRUCT(FLocalFogVolumeCommonParameters, )
 	SHADER_PARAMETER(uint32, LocalFogVolumeTilePixelSize)
 	SHADER_PARAMETER(float,  LocalFogVolumeMaxDensityIntoVolumetricFog)
 	SHADER_PARAMETER(uint32, ShouldRenderLocalFogVolumeInVolumetricFog)
+	SHADER_PARAMETER(float,  GlobalStartDistance)
 	SHADER_PARAMETER(FVector3f, DirectionalLightColor)
 	SHADER_PARAMETER(FVector3f, DirectionalLightDirection)
 END_SHADER_PARAMETER_STRUCT()
@@ -120,6 +121,7 @@ bool ProjectSupportsLocalFogVolumes();
 bool ShouldRenderLocalFogVolume(const FScene* Scene, const FSceneViewFamily& SceneViewFamily);
 bool ShouldRenderLocalFogVolumeDuringHeightFogPass(const FScene* Scene, const FSceneViewFamily& SceneViewFamily);
 bool ShouldRenderLocalFogVolumeInVolumetricFog(const FScene* Scene, const FSceneViewFamily& SceneViewFamily, bool bShouldRenderVolumetricFog);
+float GetLocalFogVolumeGlobalStartDistance();
 
 void GetLocalFogVolumeSortingData(const FScene* Scene, FRDGBuilder& GraphBuilder, FLocalFogVolumeSortingData& Out);
 
diff --git a/Engine/Source/Runtime/Renderer/Private/MobileFogRendering.cpp b/Engine/Source/Runtime/Renderer/Private/MobileFogRendering.cpp
--- a/Engine/Source/Runtime/Renderer/Private/MobileFogRendering.cpp
+++ b/Engine/Source/Runtime/Renderer/Private/MobileFogRendering.cpp
@@ -212,11 +212,10 @@ void FMobileSceneRenderer::RenderFog(FRHICommandList& RHICmdList, const FViewInf
 	
 	// Use height fog start distance by default and fallback to AP distance
 	float FogStartDistance = GetFogDefaultStartDistance();
-	if (!bFogHasComposedLocalFogVolumes) // Local fog volume do not have any start distance at the moment.
 	{
 		if (bUseHeightFog)
 		{
-			FogStartDistance = GetViewFogCommonStartDistance(View, bShouldRenderVolumetricFog);
+			FogStartDistance = GetViewFogCommonStartDistance(View, bShouldRenderVolumetricFog, bFogHasComposedLocalFogVolumes);
 		}
 		if (bUseAerialPerspective)
 		{
