Add support for mobile SSAO at half resolution.
This works with SSAO only, and not GTAO.
Use r.Mobile.AmbientOcclusionHalfResolution to enable.
The setting from 1-3 determines the upsample quality.
1 is simple bilinear. 2-3 are bilateral upsamples with increasing sample count.

diff --git a/Engine/Shaders/Private/PostProcessAmbientOcclusionCommon.ush b/Engine/Shaders/Private/PostProcessAmbientOcclusionCommon.ush
--- a/Engine/Shaders/Private/PostProcessAmbientOcclusionCommon.ush
+++ b/Engine/Shaders/Private/PostProcessAmbientOcclusionCommon.ush
@@ -47,6 +47,7 @@ Texture2D SSAO_DownsampledAO;
 SamplerState SSAO_Sampler;
 
 float2 SSAO_DownsampledAOInverseSize;
+float2 SSAO_SvPositionScaleBias;
 
 // could be moved to a more central spot
 // @param ScreenPos -1 .. 1
diff --git a/Engine/Shaders/Private/PostProcessAmbientOcclusionMobile.usf b/Engine/Shaders/Private/PostProcessAmbientOcclusionMobile.usf
--- a/Engine/Shaders/Private/PostProcessAmbientOcclusionMobile.usf
+++ b/Engine/Shaders/Private/PostProcessAmbientOcclusionMobile.usf
@@ -964,6 +964,16 @@ void GTAOSpatialFilterPS(
 // --------------------------------------------------------------------------------------------------------------------
 // SSAO
 
+float3 PackSceneDepth(float InSceneDepth)
+{
+	return UnpackRGBA8(PackR24F(InSceneDepth)).rgb;
+}
+
+float UnpackSceneDepth(float3 InPackedSceneDepth)
+{
+	return UnpackR24F(PackRGBA8(float4(InPackedSceneDepth, 0)));
+}
+
 #if SSAO
 
 void MainPSandCS(in float4 UVAndScreenPos, float4 SvPosition, out float4 OutColor)
@@ -988,8 +998,10 @@ void MainPSandCS(in float4 UVAndScreenPos, float4 SvPosition, out float4 OutColo
 	float3 FovFix = float3(InvTanHalfFov, Ratio * InvTanHalfFov, 1);
 	float3 InvFovFix = 1.0f / FovFix;
 
+	float4 ModifiedSvPosition = float4(SvPosition.xy * SSAO_SvPositionScaleBias.xx + SSAO_SvPositionScaleBias.yy, SvPosition.zw);
+
 	float SceneDepth = GetDepthFromAOInput(UV);
-	float3 WorldNormal = GetWorldSpaceNormalFromAOInput(UV, SvPosition);
+	float3 WorldNormal = GetWorldSpaceNormalFromAOInput(UV, ModifiedSvPosition);
 
 	// can be NaN if WorldNormal=0,0,0 which happens when !USE_NORMALS
 	float3 ViewSpaceNormal = normalize(mul(WorldNormal, (float3x3)View.TranslatedWorldToView));
@@ -1208,6 +1220,10 @@ void MainPSandCS(in float4 UVAndScreenPos, float4 SvPosition, out float4 OutColo
 //		OutColor = (WeightA - 0.333f) / 0.666f;
 	}
 #endif
+
+#if OUTPUT_DEPTH
+	OutColor.gba = PackSceneDepth(SceneDepth);
+#endif
 }
 
 void MainPS(in noperspective float4 UVAndScreenPos : TEXCOORD0, float4 SvPosition : SV_POSITION, out float4 OutColor : SV_Target0)
@@ -1217,3 +1233,104 @@ void MainPS(in noperspective float4 UVAndScreenPos : TEXCOORD0, float4 SvPositio
 
 #endif // SSAO
 
+
+#ifdef UPSAMPLE_PASS
+
+Texture2D AOTexture;
+SamplerState AOSampler;
+
+float UpsampleSSAO(float2 BufferUV)
+{
+#if UPSAMPLE_QUALITY == 0
+	// Bilinear sample.
+	return Texture2DSampleLevel(AOTexture, AOSampler, BufferUV, 0).x;
+
+#elif UPSAMPLE_QUALITY == 1
+	// 4 sample bilinear with depth weighting.
+	float2 LowResBufferSize = floor(View.BufferSizeAndInvSize.xy / 2);
+	float2 LowResTexelSize = 1.0f / LowResBufferSize;
+
+	float2 Corner00UV = floor(BufferUV * LowResBufferSize - 0.5f) / LowResBufferSize + 0.5f * LowResTexelSize;
+	float2 BilinearWeights = (BufferUV - Corner00UV) * LowResBufferSize;
+
+	float4 AOValues00 = Texture2DSampleLevel(AOTexture, AOSampler, Corner00UV, 0);
+	float4 AOValues10 = Texture2DSampleLevel(AOTexture, AOSampler, Corner00UV + float2(LowResTexelSize.x, 0), 0);
+	float4 AOValues01 = Texture2DSampleLevel(AOTexture, AOSampler, Corner00UV + float2(0, LowResTexelSize.y), 0);
+	float4 AOValues11 = Texture2DSampleLevel(AOTexture, AOSampler, Corner00UV + LowResTexelSize, 0);
+
+	float4 CornerWeights = float4(
+		(1 - BilinearWeights.y) * (1 - BilinearWeights.x),
+		(1 - BilinearWeights.y) * BilinearWeights.x,
+		BilinearWeights.y * (1 - BilinearWeights.x),
+		BilinearWeights.y * BilinearWeights.x);
+
+	float SceneDepth00 = UnpackSceneDepth(AOValues00.gba);
+	float SceneDepth10 = UnpackSceneDepth(AOValues10.gba);
+	float SceneDepth01 = UnpackSceneDepth(AOValues01.gba);
+	float SceneDepth11 = UnpackSceneDepth(AOValues11.gba);
+
+	float MaxDepth = max(max(max(SceneDepth00, SceneDepth10), SceneDepth01), SceneDepth11);
+	float MinDepth = min(min(min(SceneDepth00, SceneDepth10), SceneDepth01), SceneDepth11);
+
+	if (MaxDepth / MinDepth > 1.02f)
+	{
+		float Epsilon = 0.0001f;
+		float SceneDepth = GetDepthFromAOInput(BufferUV);
+		float4 CornerDepths = abs(float4(SceneDepth00, SceneDepth10, SceneDepth01, SceneDepth11));
+		float4 DepthWeights = 1.0f / (abs(CornerDepths - SceneDepth.xxxx) + Epsilon);
+		CornerWeights *= DepthWeights;
+	}
+
+	float InterpolatedResult =
+		(CornerWeights.x * AOValues00.r
+			+ CornerWeights.y * AOValues10.r
+			+ CornerWeights.z * AOValues01.r
+			+ CornerWeights.w * AOValues11.r)
+		/ dot(CornerWeights, 1);
+
+	return InterpolatedResult;
+
+#elif UPSAMPLE_QUALITY == 2
+	// 9 sample with depth weighting.
+	float2 LowResBufferSize = floor(View.BufferSizeAndInvSize.xy * 0.5);
+	float2 LowResTexelSize = 1.0f / LowResBufferSize;
+
+	float4 AOValues[9];
+	AOValues[0] = Texture2DSampleLevel(AOTexture, AOSampler, BufferUV + float2(-1, -1) * LowResTexelSize, 0);
+	AOValues[1] = Texture2DSampleLevel(AOTexture, AOSampler, BufferUV + float2(0, -1) * LowResTexelSize, 0);
+	AOValues[2] = Texture2DSampleLevel(AOTexture, AOSampler, BufferUV + float2(1, -1) * LowResTexelSize, 0);
+	AOValues[3] = Texture2DSampleLevel(AOTexture, AOSampler, BufferUV + float2(-1, 0) * LowResTexelSize, 0);
+	AOValues[4] = Texture2DSampleLevel(AOTexture, AOSampler, BufferUV + float2(0, 0) * LowResTexelSize, 0);
+	AOValues[5] = Texture2DSampleLevel(AOTexture, AOSampler, BufferUV + float2(1, 0) * LowResTexelSize, 0);
+	AOValues[6] = Texture2DSampleLevel(AOTexture, AOSampler, BufferUV + float2(-1, 1) * LowResTexelSize, 0);
+	AOValues[7] = Texture2DSampleLevel(AOTexture, AOSampler, BufferUV + float2(0, 1) * LowResTexelSize, 0);
+	AOValues[8] = Texture2DSampleLevel(AOTexture, AOSampler, BufferUV + float2(1, 1) * LowResTexelSize, 0);
+
+	float SceneDepth = GetDepthFromAOInput(BufferUV);
+	float Epsilon = 0.0001f;
+
+	float TotalSum = 0;
+	float TotalWeight = 0;
+	UNROLL for (int i = 0; i < 9; ++i)
+	{
+		float SampleValue = AOValues[i].r;
+		float SampleDepth = UnpackSceneDepth(AOValues[i].gba);
+
+		float Weight = 1.0f / (abs(SceneDepth - SampleDepth) + Epsilon);
+		TotalWeight += Weight;
+		TotalSum += SampleValue * Weight;
+	}
+
+	return TotalSum / TotalWeight;
+#endif
+}
+
+void AmbientOcclusionUpsamplePS(
+	in float4 SvPosition : SV_POSITION,
+	out float4 OutColor : SV_Target0)
+{
+	float2 BufferUV = SvPositionToBufferUV(SvPosition);
+	OutColor = UpsampleSSAO(BufferUV).xxxx;
+}
+
+#endif // UPSAMPLE_PASS
diff --git a/Engine/Source/Runtime/Renderer/Private/CompositionLighting/PostProcessAmbientOcclusion.cpp b/Engine/Source/Runtime/Renderer/Private/CompositionLighting/PostProcessAmbientOcclusion.cpp
--- a/Engine/Source/Runtime/Renderer/Private/CompositionLighting/PostProcessAmbientOcclusion.cpp
+++ b/Engine/Source/Runtime/Renderer/Private/CompositionLighting/PostProcessAmbientOcclusion.cpp
@@ -665,6 +665,7 @@ BEGIN_SHADER_PARAMETER_STRUCT(FAmbientOcclusionParameters, )
 	SHADER_PARAMETER_STRUCT_INCLUDE(FSSAOShaderParameters, SSAOParameters)
 
 	SHADER_PARAMETER(FVector2f, SSAO_DownsampledAOInverseSize)
+	SHADER_PARAMETER(FVector2f, SSAO_SvPositionScaleBias)
 
 	SHADER_PARAMETER_RDG_TEXTURE(Texture2D, SSAO_SetupTexture)
 	SHADER_PARAMETER_RDG_TEXTURE(Texture2D, SSAO_NormalsTexture)
@@ -827,6 +828,8 @@ void AddAmbientOcclusionPass(
 		SharedParameters.SSAO_DownsampledAOInverseSize = FVector2f(1.0f, 1.0f);
 	}
 
+	SharedParameters.SSAO_SvPositionScaleBias = FVector2f(1, 0);
+	
 	SharedParameters.SSAO_Sampler = TStaticSamplerState<SF_Point, AM_Clamp, AM_Clamp, AM_Clamp>::GetRHI();
 
 	SharedParameters.RandomNormalTexture = GSystemTextures.SSAORandomization->GetRHI();
diff --git a/Engine/Source/Runtime/Renderer/Private/PostProcess/PostProcessAmbientOcclusionMobile.cpp b/Engine/Source/Runtime/Renderer/Private/PostProcess/PostProcessAmbientOcclusionMobile.cpp
--- a/Engine/Source/Runtime/Renderer/Private/PostProcess/PostProcessAmbientOcclusionMobile.cpp
+++ b/Engine/Source/Runtime/Renderer/Private/PostProcess/PostProcessAmbientOcclusionMobile.cpp
@@ -18,11 +18,12 @@
 #include "SceneTextureParameters.h"
 #include "SceneRenderTargetParameters.h"
 #include "ClearQuad.h"
+#include "PixelShaderUtils.h"
 
 static TAutoConsoleVariable<int32> CVarMobileAmbientOcclusion(
 	TEXT("r.Mobile.AmbientOcclusion"),
 	0,
-	TEXT("Causion: An extra sampler will be occupied in mobile base pass pixel shader after enable the mobile ambient occlusion.\n")
+	TEXT("Caution: An extra sampler will be occupied in mobile base pass pixel shader after enable the mobile ambient occlusion.\n")
 	TEXT("0: Disable Ambient Occlusion on mobile platform. [default]\n")
 	TEXT("1: Enable Ambient Occlusion on mobile platform.\n"),
 	ECVF_ReadOnly | ECVF_RenderThreadSafe
@@ -70,6 +71,16 @@ static TAutoConsoleVariable<int32> CVarMobileAmbientOcclusionDepthBoundsTest(
 	TEXT("Whether to use depth bounds test to cull distant pixels during AO pass. This option is only valid when pixel shader path is used"),
 	ECVF_RenderThreadSafe);
 
+static TAutoConsoleVariable<int32> CVarMobileSSAOHalfResolution(
+	TEXT("r.Mobile.SSAOHalfResolution"),
+	1,
+	TEXT("Whether to calculate SSAO at half resolution.\n")
+	TEXT("0: Disabled.\n")
+	TEXT("1: Half Resolution with bilinear upsample\n")
+	TEXT("2: Half Resolution with 4 tap bilateral upsample\n")
+	TEXT("3: Half Resolution with 9 tap bilateral upsample\n"),
+	ECVF_RenderThreadSafe);
+
 // --------------------------------------------------------------------------------------------------------------------
 DECLARE_GPU_STAT_NAMED(MobileSSAO, TEXT("SSAO"));
 
@@ -684,6 +695,7 @@ BEGIN_SHADER_PARAMETER_STRUCT(FMobileAmbientOcclusionParameters, RENDERER_API)
 	SHADER_PARAMETER_SAMPLER(SamplerState, SSAO_Sampler)
 
 	SHADER_PARAMETER(FVector2f, SSAO_DownsampledAOInverseSize)
+	SHADER_PARAMETER(FVector2f, SSAO_SvPositionScaleBias)
 
 	SHADER_PARAMETER_RDG_TEXTURE(Texture2D, RandomNormalTexture)
 	SHADER_PARAMETER_SAMPLER(SamplerState, RandomNormalTextureSampler)
@@ -696,8 +708,8 @@ public:
 	SHADER_USE_PARAMETER_STRUCT(FMobileAmbientOcclusionPS, FGlobalShader);
 
 	class FShaderQualityDim : SHADER_PERMUTATION_INT("SHADER_QUALITY", 5);
-
-	using FPermutationDomain = TShaderPermutationDomain<FShaderQualityDim>;
+	class FOutputDepth : SHADER_PERMUTATION_BOOL("OUTPUT_DEPTH");
+	using FPermutationDomain = TShaderPermutationDomain<FShaderQualityDim, FOutputDepth>;
 
 	static void ModifyCompilationEnvironment(const FGlobalShaderPermutationParameters& Parameters, FShaderCompilerEnvironment& OutEnvironment)
 	{
@@ -713,12 +725,40 @@ public:
 
 	BEGIN_SHADER_PARAMETER_STRUCT(FParameters, )
 		SHADER_PARAMETER_STRUCT_INCLUDE(FMobileAmbientOcclusionParameters, SharedParameters)
-
 		RENDER_TARGET_BINDING_SLOTS()
-		END_GLOBAL_SHADER_PARAMETER_STRUCT();
+	END_SHADER_PARAMETER_STRUCT();
 };
 IMPLEMENT_GLOBAL_SHADER(FMobileAmbientOcclusionPS, "/Engine/Private/PostProcessAmbientOcclusionMobile.usf", "MainPS", SF_Pixel);
 
+class FMobileAmbientOcclusionUpsamplePS : public FGlobalShader
+{
+	DECLARE_GLOBAL_SHADER(FMobileAmbientOcclusionUpsamplePS);
+	SHADER_USE_PARAMETER_STRUCT(FMobileAmbientOcclusionUpsamplePS, FGlobalShader);
+
+	class FUpsampleQualityDim : SHADER_PERMUTATION_INT("UPSAMPLE_QUALITY", 3);
+	using FPermutationDomain = TShaderPermutationDomain <FUpsampleQualityDim>;
+
+	BEGIN_SHADER_PARAMETER_STRUCT(FParameters, )
+		SHADER_PARAMETER_RDG_UNIFORM_BUFFER(FMobileSceneTextureUniformParameters, SceneTextures)
+		SHADER_PARAMETER_STRUCT_REF(FViewUniformShaderParameters, View)
+		SHADER_PARAMETER_RDG_TEXTURE(Texture2D, AOTexture)
+		SHADER_PARAMETER_SAMPLER(SamplerState, AOSampler)
+		RENDER_TARGET_BINDING_SLOTS()
+	END_SHADER_PARAMETER_STRUCT()
+
+	static bool ShouldCompilePermutation(const FGlobalShaderPermutationParameters& Parameters)
+	{
+		return IsMobileAmbientOcclusionEnabled(Parameters.Platform);
+	}
+
+	static void ModifyCompilationEnvironment(const FGlobalShaderPermutationParameters& Parameters, FShaderCompilerEnvironment& OutEnvironment)
+	{
+		OutEnvironment.SetDefine(TEXT("UPSAMPLE_PASS"), 1);
+	
+	}
+};
+IMPLEMENT_GLOBAL_SHADER(FMobileAmbientOcclusionUpsamplePS, "/Engine/Private/PostProcessAmbientOcclusionMobile.usf", "AmbientOcclusionUpsamplePS", SF_Pixel);
+
 // --------------------------------------------------------------------------------------------------------------------
 static FMobileSSAOCommonParameters GetMobileSSAOCommonParameters(
 	FRDGBuilder& GraphBuilder,
@@ -743,7 +783,7 @@ static FMobileSSAOShaderParameters GetMobileSSAOShaderParameters(
 {
 	const FFinalPostProcessSettings& Settings = View.FinalPostProcessSettings;
 
-	FIntPoint RandomizationSize = GSystemTextures.SSAORandomization->GetDesc().Extent;
+	FIntPoint RandomizationSize = GSystemTextures.SSAORandomization->GetDesc().Extent * InputViewport.Extent / OutputViewport.Extent;
 	FVector2D ViewportUVToRandomUV(InputViewport.Extent.X / (float)RandomizationSize.X, InputViewport.Extent.Y / (float)RandomizationSize.Y);
 
 	// e.g. 4 means the input texture is 4x smaller than the buffer size
@@ -805,25 +845,23 @@ static void AddMobileAmbientOcclusionPass(
 	FRDGBuilder& GraphBuilder,
 	const FViewInfo& View,
 	const FMobileSSAOCommonParameters& CommonParameters,
+	bool bHalfResolution,
+	bool bOutputDepth,
 	FScreenPassRenderTarget Output)
 {
-	RDG_GPU_STAT_SCOPE(GraphBuilder, MobileSSAO);
-
 	check(Output.IsValid());
 
 	const FScreenPassTextureViewport InputViewport = CommonParameters.SceneTexturesViewport;
-	const FScreenPassTextureViewport OutputViewport(Output);
-	const uint32 ScaleToFullRes = CommonParameters.SceneTexturesViewport.Extent.X / InputViewport.Extent.X;
+	const FScreenPassTextureViewport OutputViewport = GetDownscaledViewport(CommonParameters.SceneTexturesViewport, bHalfResolution ? 2 : 1);
 	
 	const bool bDepthBoundsTestEnabled =
 		CVarMobileAmbientOcclusionDepthBoundsTest.GetValueOnRenderThread()
-		&& ScaleToFullRes == 1
+		&& !bHalfResolution
 		&& GSupportsDepthBoundsTest
 		&& CommonParameters.SceneDepth.IsValid()
 		&& CommonParameters.SceneDepth.Texture->Desc.NumSamples == 1;
 
 	float DepthFar = 0.0f;
-
 	FDepthStencilBinding DepthStencilBinding(CommonParameters.SceneDepth.Texture, ERenderTargetLoadAction::ELoad, ERenderTargetLoadAction::ELoad, FExclusiveDepthStencil::DepthRead_StencilWrite);
 
 	if (bDepthBoundsTestEnabled)
@@ -876,8 +914,9 @@ static void AddMobileAmbientOcclusionPass(
 	SharedParameters.SceneTextures = CommonParameters.SceneTexturesUniformBufferRDG;
 	SharedParameters.HZBParameters = GetHZBParameters(View, CommonParameters.HZBInput, CommonParameters.SceneTexturesViewport.Extent, EAOTechnique::SSAO);
 	SharedParameters.SSAOParameters = GetMobileSSAOShaderParameters(View, InputViewport, OutputViewport, CommonParameters.SceneTexturesViewport);
-
 	SharedParameters.SSAO_Sampler = TStaticSamplerState<SF_Point, AM_Clamp, AM_Clamp, AM_Clamp>::GetRHI();
+	SharedParameters.SSAO_SvPositionScaleBias = bHalfResolution ? FVector2f(2, -0.5) : FVector2f(1, 0);
+	SharedParameters.SSAO_DownsampledAOInverseSize = bHalfResolution ? FVector2f(1, 1) / FVector2f(OutputViewport.Extent) : FVector2f(1, 1);
 	SharedParameters.RandomNormalTexture = GraphBuilder.RegisterExternalTexture(GSystemTextures.SSAORandomization, TEXT("SSAORandomization"));
 	SharedParameters.RandomNormalTextureSampler = TStaticSamplerState<SF_Point, AM_Wrap, AM_Wrap, AM_Wrap>::GetRHI();
 
@@ -895,6 +934,7 @@ static void AddMobileAmbientOcclusionPass(
 	const int32 MobileAmbientOcclusionQuality = FMath::Min(CVarMobileAmbientOcclusionQuality.GetValueOnRenderThread(), 3) - 1; 
 	FMobileAmbientOcclusionPS::FPermutationDomain PermutationVector;
 	PermutationVector.Set<FMobileAmbientOcclusionPS::FShaderQualityDim>(MobileAmbientOcclusionQuality);
+	PermutationVector.Set<FMobileAmbientOcclusionPS::FOutputDepth>(bOutputDepth);
 
 	TShaderMapRef<FMobileAmbientOcclusionPS> PixelShader(View.ShaderMap, PermutationVector);
 	TShaderMapRef<FScreenPassVS> VertexShader(View.ShaderMap);
@@ -954,24 +994,158 @@ static void AddMobileAmbientOcclusionPass(
 		});	
 }
 
+static void AddMobileAmbientOcclusionUpsamplePass(
+	FRDGBuilder& GraphBuilder,
+	const FViewInfo& View,
+	const FMobileSSAOCommonParameters& CommonParameters,
+	int32 UpsampleQuality,
+	FRDGTextureRef SourceTexture,
+	FScreenPassRenderTarget Output)
+{
+	const bool bDepthBoundsTestEnabled =
+		CVarMobileAmbientOcclusionDepthBoundsTest.GetValueOnRenderThread()
+		&& GSupportsDepthBoundsTest
+		&& CommonParameters.SceneDepth.IsValid()
+		&& CommonParameters.SceneDepth.Texture->Desc.NumSamples == 1;
+
+	const FScreenPassTextureViewport OutputViewport = CommonParameters.SceneTexturesViewport;
+
+	float DepthFar = 0.0f;
+	FDepthStencilBinding DepthStencilBinding(CommonParameters.SceneDepth.Texture, ERenderTargetLoadAction::ELoad, ERenderTargetLoadAction::ELoad, FExclusiveDepthStencil::DepthRead_StencilWrite);
+
+	if (bDepthBoundsTestEnabled)
+	{
+		const FFinalPostProcessSettings& Settings = View.FinalPostProcessSettings;
+		const FMatrix& ProjectionMatrix = View.ViewMatrices.GetProjectionMatrix();
+		const FVector4f Far = (FVector4f)ProjectionMatrix.TransformFVector4(FVector4(0, 0, Settings.AmbientOcclusionFadeDistance));
+		DepthFar = FMath::Min(1.0f, Far.Z / Far.W);
+
+		static_assert(bool(ERHIZBuffer::IsInverted), "Inverted depth buffer is assumed when setting depth bounds test for AO.");
+
+		FRenderTargetParameters* ClearParameters = GraphBuilder.AllocParameters<FRenderTargetParameters>();
+		ClearParameters->RenderTargets[0] = Output.GetRenderTargetBinding();
+		ClearParameters->RenderTargets.DepthStencil = DepthStencilBinding;
+
+		GraphBuilder.AddPass(
+			RDG_EVENT_NAME("DepthBounds ClearQuad(%s)", Output.Texture->Name),
+			ClearParameters,
+			ERDGPassFlags::Raster,
+			[OutputViewport, DepthFar](FRHICommandListImmediate& RHICmdList)
+			{
+				RHICmdList.SetViewport(OutputViewport.Rect.Min.X, OutputViewport.Rect.Min.Y, 0.0f, OutputViewport.Rect.Max.X, OutputViewport.Rect.Max.Y, 1.0f);
+
+				// We must clear all pixels that won't be touched by AO shader.
+				FClearQuadCallbacks Callbacks;
+				Callbacks.PSOModifier = [](FGraphicsPipelineStateInitializer& PSOInitializer)
+				{
+					PSOInitializer.bDepthBounds = true;
+				};
+				Callbacks.PreClear = [DepthFar](FRHICommandList& InRHICmdList)
+				{
+					// This is done by rendering a clear quad over a depth range from AmbientOcclusionFadeDistance to far plane.
+					InRHICmdList.SetDepthBounds(0, DepthFar);	// NOTE: Inverted depth
+				};
+				Callbacks.PostClear = [DepthFar](FRHICommandList& InRHICmdList)
+				{
+					// Set depth bounds test to cover everything from near plane to AmbientOcclusionFadeDistance and run AO pixel shader.
+					InRHICmdList.SetDepthBounds(DepthFar, 1.0f);
+				};
+
+				DrawClearQuad(RHICmdList, FLinearColor::White, Callbacks);
+			});
+
+		// Make sure the following pass doesn't clear or ignore the data
+		Output.LoadAction = ERenderTargetLoadAction::ELoad;
+	}
+
+	FMobileAmbientOcclusionUpsamplePS::FParameters* PassParameters = GraphBuilder.AllocParameters<FMobileAmbientOcclusionUpsamplePS::FParameters>();
+	PassParameters->RenderTargets[0] = Output.GetRenderTargetBinding();
+	if (bDepthBoundsTestEnabled)
+	{
+		PassParameters->RenderTargets.DepthStencil = DepthStencilBinding;
+	}
+
+	PassParameters->View = GetShaderBinding(View.ViewUniformBuffer);
+	PassParameters->SceneTextures = CommonParameters.SceneTexturesUniformBufferRDG;
+	PassParameters->AOTexture = SourceTexture;
+	PassParameters->AOSampler = TStaticSamplerState<SF_Bilinear>::GetRHI();
+
+	typename FMobileAmbientOcclusionUpsamplePS::FPermutationDomain PermutationVector;
+	PermutationVector.Set<typename FMobileAmbientOcclusionUpsamplePS::FUpsampleQualityDim>(UpsampleQuality);
+	TShaderMapRef<FMobileAmbientOcclusionUpsamplePS> PixelShader(View.ShaderMap, PermutationVector);
+
+	TShaderMapRef<FScreenPassVS> VertexShader(View.ShaderMap);
+
+	ClearUnusedGraphResources(PixelShader, PassParameters);
+
+	GraphBuilder.AddPass(
+		RDG_EVENT_NAME("UpsamplePS Quality %d", UpsampleQuality),
+		PassParameters,
+		ERDGPassFlags::Raster,
+		[PassParameters, &View, OutputViewport, VertexShader, PixelShader, bDepthBoundsTestEnabled, DepthFar](FRHICommandList& RHICmdList)
+		{
+			RHICmdList.SetViewport(OutputViewport.Rect.Min.X, OutputViewport.Rect.Min.Y, 0.0f, OutputViewport.Rect.Max.X, OutputViewport.Rect.Max.Y, 1.0f);
+
+			FGraphicsPipelineStateInitializer GraphicsPSOInit;
+			FPixelShaderUtils::InitFullscreenPipelineState(RHICmdList, View.ShaderMap, PixelShader, GraphicsPSOInit);
+			GraphicsPSOInit.bDepthBounds = bDepthBoundsTestEnabled;
+
+			SetGraphicsPipelineState(RHICmdList, GraphicsPSOInit, 0);
+			SetShaderParameters(RHICmdList, PixelShader, PixelShader.GetPixelShader(), *PassParameters);
+
+			if (bDepthBoundsTestEnabled)
+			{
+				RHICmdList.SetDepthBounds(DepthFar, 1.0f);
+			}
+
+			FPixelShaderUtils::DrawFullscreenTriangle(RHICmdList);
+			
+			if (bDepthBoundsTestEnabled)
+			{
+				RHICmdList.SetDepthBounds(0.0f, 1.0f);
+			}
+		});
+}
+
 static void RenderSSAO(FRDGBuilder& GraphBuilder, FRDGTextureRef SceneDepthTexture, FRDGTextureRef AmbientOcclusionTexture, const TArray<FViewInfo>& Views)
 {
-	
-	TRDGUniformBufferRef<FMobileSceneTextureUniformParameters> SceneTexturesUniformBufferRDG = CreateMobileSceneTextureUniformBuffer(GraphBuilder, GetViewFamilyInfo(Views).GetSceneTexturesChecked(), EMobileSceneTextureSetupMode::SceneDepth);
-	
-	for (int32 ViewIndex = 0; ViewIndex < Views.Num(); ViewIndex++)
+	RDG_EVENT_SCOPE(GraphBuilder, "MobileSSAO");
+	RDG_GPU_STAT_SCOPE(GraphBuilder, MobileSSAO);
+
+	const int32 HalfResolutionSetting = CVarMobileSSAOHalfResolution.GetValueOnRenderThread();
+	const bool bHalfResolution = HalfResolutionSetting > 0;
+	const int32 UpsampleQuality = FMath::Clamp(HalfResolutionSetting - 1, 0, 2);
+	const bool bBilateralUpsample = UpsampleQuality > 0;
+
+	FRDGTextureRef HalfResolutionTexture = nullptr;
+	if (bHalfResolution)
+	{
+		// Bilateral requires 32bit format for AO + depth.
+		const EPixelFormat Format = bBilateralUpsample ? PF_R8G8B8A8 : AmbientOcclusionTexture->Desc.Format;
+		const FIntPoint Extent = GetDownscaledExtent(AmbientOcclusionTexture->Desc.Extent, 2);
+		HalfResolutionTexture = GraphBuilder.CreateTexture(
+			FRDGTextureDesc::Create2D(Extent, Format, FClearValueBinding::None, TexCreate_ShaderResource | TexCreate_RenderTargetable),
+			TEXT("HalfResolutionScreenSpaceAO"));
+	}
+
+	TRDGUniformBufferRef<FMobileSceneTextureUniformParameters> SceneTexturesUniformBufferRDG = CreateMobileSceneTextureUniformBuffer(
+		GraphBuilder, GetViewFamilyInfo(Views).GetSceneTexturesChecked(), EMobileSceneTextureSetupMode::SceneDepth);
+
+	for (FViewInfo const& View : Views)
 	{
-		const FViewInfo& View = Views[ViewIndex];
 		FMobileSSAOCommonParameters Parameters = GetMobileSSAOCommonParameters(GraphBuilder, View, SceneDepthTexture, SceneTexturesUniformBufferRDG);
 		FScreenPassRenderTarget FinalTarget = FScreenPassRenderTarget(AmbientOcclusionTexture, View.ViewRect, ERenderTargetLoadAction::ENoAction);
 
-		AddMobileAmbientOcclusionPass
-		(
-			GraphBuilder,
-			View,
-			Parameters,
-			FinalTarget
-		);
+		if (HalfResolutionTexture != nullptr)
+		{
+			FScreenPassRenderTarget HalfResolutionTarget = FScreenPassRenderTarget(HalfResolutionTexture, View.ViewRect, ERenderTargetLoadAction::ENoAction);
+			AddMobileAmbientOcclusionPass(GraphBuilder,	View, Parameters, bHalfResolution, bBilateralUpsample, HalfResolutionTarget);
+			AddMobileAmbientOcclusionUpsamplePass(GraphBuilder, View, Parameters, UpsampleQuality, HalfResolutionTexture, FinalTarget);
+		}
+		else
+		{
+			AddMobileAmbientOcclusionPass(GraphBuilder, View, Parameters, bHalfResolution, bBilateralUpsample, FinalTarget);
+		}
 	}
 }
 
