Fix warning V763: Parameter is always rewritten in function body before being used.
#rnx

diff --git a/Engine/Source/Runtime/Experimental/IoStoreOnDemand/Private/JournaledCache.cpp b/Engine/Source/Runtime/Experimental/IoStoreOnDemand/Private/JournaledCache.cpp
--- a/Engine/Source/Runtime/Experimental/IoStoreOnDemand/Private/JournaledCache.cpp
+++ b/Engine/Source/Runtime/Experimental/IoStoreOnDemand/Private/JournaledCache.cpp
@@ -605,8 +605,7 @@ FDiskCache::FDiskCache(FString&& Path, uint64 InMaxDataSize, uint32 InJournalSiz
 , Journal(BinPath, InJournalSize)
 {
 	// Align down to keep to some assumptions
-	InJournalSize = Journal.GetMaxSize();
-	MaxDataSize = (MaxDataSize - InJournalSize) & ~((1ull << 20) - 1);
+	MaxDataSize = (MaxDataSize - Journal.GetMaxSize()) & ~((1ull << 20) - 1);
 
 	OpenDataFile();
 	
diff --git a/Engine/Source/Runtime/Renderer/Private/MeshPassProcessor.cpp b/Engine/Source/Runtime/Renderer/Private/MeshPassProcessor.cpp
--- a/Engine/Source/Runtime/Renderer/Private/MeshPassProcessor.cpp
+++ b/Engine/Source/Runtime/Renderer/Private/MeshPassProcessor.cpp
@@ -1708,7 +1708,7 @@ void SubmitMeshDrawCommandsRange(
 	//               we don't support dynamic instancing for this path since we require one primitive per draw command
 	//               This is because the stride on the instance data buffer is set to 0 so only the first will ever be fetched.
 	checkSlow(!bDynamicInstancing);
-	bDynamicInstancing = false;
+	bDynamicInstancing = false; //-V763
 
 	FMeshDrawCommandStateCache StateCache;
 	FMeshDrawCommandSceneArgs LocalSceneArgs = InSceneArgs;
