LFV - added HZB culling support to cull volumes when hidden.

diff --git a/Engine/Shaders/Private/LocalFogVolumes/LocalFogVolumeTiledCulling.usf b/Engine/Shaders/Private/LocalFogVolumes/LocalFogVolumeTiledCulling.usf
--- a/Engine/Shaders/Private/LocalFogVolumes/LocalFogVolumeTiledCulling.usf
+++ b/Engine/Shaders/Private/LocalFogVolumes/LocalFogVolumeTiledCulling.usf
@@ -18,6 +18,8 @@ RWTexture2DArray<uint> LocalFogVolumeTileDataTextureUAV;
 RWBuffer<uint> TileDataBufferUAV;
 RWBuffer<uint> TileDrawIndirectBufferUAV;
 
+int4 HZBViewRect;
+
 #ifndef THREADGROUP_SIZE
 #error THREADGROUP_SIZE must be defined
 #endif
@@ -25,6 +27,8 @@ RWBuffer<uint> TileDrawIndirectBufferUAV;
 [numthreads(THREADGROUP_SIZE, THREADGROUP_SIZE, 1)]
 void LocalFogVolumeTiledCullingCS(uint3 DispatchThreadId : SV_DispatchThreadID)
 {
+	ResolvedView = ResolveView();
+
 	uint2 TileCoord = uint2(DispatchThreadId.xy);
 
 	if (any(TileCoord >= LFVTileDataResolution))
@@ -68,6 +72,15 @@ void LocalFogVolumeTiledCullingCS(uint3 DispatchThreadId : SV_DispatchThreadID)
 		TilePlanes[i] *= Scale;
 	}
 #endif
+#if USE_HZB
+	static float4x4 IdentityLocalToTranslatedWorldMatrix =
+	{
+		{ 1, 0, 0, 0 },
+		{ 0, 1, 0, 0 },
+		{ 0, 0, 1, 0 },
+		{ 0, 0, 0, 1 }
+	};
+#endif
 
 	uint LFVCount = 0;
 	for (int LFVIndex = 0; LFVIndex < LFVInstcount; ++LFVIndex)
@@ -84,6 +97,46 @@ void LocalFogVolumeTiledCullingCS(uint3 DispatchThreadId : SV_DispatchThreadID)
 				break;
 			}
 		}
+
+	#if USE_HZB
+		if (!bCompletelyOutside)
+		{
+		#if 0
+			// Sphere
+			float4 SphereCenterView = mul(Sphere.xyz + LWCToFloat(PrimaryView.PreViewTranslation), ResolvedView.TranslatedWorldToView);
+			FFrustumCullData Cull = SphereCullFrustum(SphereCenterView, Sphere.w, ResolvedView.ViewToClip, 0.0/*float ZNear*/, false/*bool bNearClip*/ , true/*bool bSkipFrustumCull*/);
+		#else
+			FFrustumCullData Cull = BoxCullFrustum(
+				Sphere.xyz + LWCToFloat(PrimaryView.PreViewTranslation),
+				Sphere.www,
+				IdentityLocalToTranslatedWorldMatrix,
+				ResolvedView.TranslatedWorldToClip,
+				ResolvedView.ViewToClip,
+				false /*bIsOrtho*/, true /* near clip */, true /* skip culling */
+			);
+		#endif
+
+		#if 0 // USE_TILED_HZB // Not working yet...
+			Cull.RectMin.xy = float2(TileCoord.x,			(TileCoord.y))			/ float2(LFVTileDataResolution);
+			Cull.RectMax.xy = float2(TileCoord.x + 1.0f,	(TileCoord.y + 1.0f))	/ float2(LFVTileDataResolution);
+
+			Cull.RectMin.x = Cull.RectMin.x * ViewToTileSpaceRatio.x * 2.0 - 1.0;
+			Cull.RectMax.x = Cull.RectMax.x * ViewToTileSpaceRatio.x * 2.0 - 1.0;
+
+			Cull.RectMin.y = 1.0/ViewToTileSpaceRatio.y - Cull.RectMin.y;
+			Cull.RectMax.y = 1.0/ViewToTileSpaceRatio.y - Cull.RectMax.y;
+
+			Cull.RectMin.y = Cull.RectMin.y * 2 - 1;
+			Cull.RectMax.y = Cull.RectMax.y * 2 - 1;
+		#endif
+
+			int4 ActualHZBViewRect = int4(0.0.xx, ResolvedView.ViewRectMinAndSize.zw);
+			FScreenRect Rect = GetScreenRect(ActualHZBViewRect, Cull, 4);
+
+			bCompletelyOutside = !IsVisibleHZB(Rect, true /*bSample4x4*/);
+		}
+	#endif
+
 		if (!bCompletelyOutside)
 		{
 			LocalFogVolumeTileDataTextureUAV[uint3(TileCoord, 1 + LFVCount)] = LFVIndex;
diff --git a/Engine/Source/Runtime/Renderer/Private/DeferredShadingRenderer.cpp b/Engine/Source/Runtime/Renderer/Private/DeferredShadingRenderer.cpp
--- a/Engine/Source/Runtime/Renderer/Private/DeferredShadingRenderer.cpp
+++ b/Engine/Source/Runtime/Renderer/Private/DeferredShadingRenderer.cpp
@@ -3452,8 +3452,8 @@ void FDeferredShadingSceneRenderer::Render(FRDGBuilder& GraphBuilder)
 		CSV_CUSTOM_STAT(LightCount, Batched, float(SortedLightSet.UnbatchedLightStart), ECsvCustomStatOp::Set);
 		CSV_CUSTOM_STAT(LightCount, Unbatched, float(SortedLightSet.SortedLights.Num()) - float(SortedLightSet.UnbatchedLightStart), ECsvCustomStatOp::Set);
 
-		// Run local fog volume initialization before base pass for when data is needed in forward  and volumetric cloud view initialization which can bind LFV data.
-		InitLocalFogVolumesForViews(Scene, Views, ViewFamily, GraphBuilder, bShouldRenderVolumetricFog);
+		// Initialise local fog volume with dummy data before volumetric cloud view initialization which can bind LFV data.
+		SetDummyLocalFogVolumeForViews(GraphBuilder, Views);
 
 		// Run before RenderSkyAtmosphereLookUpTables for cloud shadows to be valid.
 		InitVolumetricCloudsForViews(GraphBuilder, bShouldRenderVolumetricCloudBase, InstanceCullingManager);
@@ -3479,6 +3479,9 @@ void FDeferredShadingSceneRenderer::Render(FRDGBuilder& GraphBuilder)
 
 		BeginAsyncDistanceFieldShadowProjections(GraphBuilder, SceneTextures, InitViewTaskDatas.DynamicShadows);
 
+		// Run local fog volume culling before base pass and after HZB generation tyo benefit from more culling.
+		InitLocalFogVolumesForViews(Scene, Views, ViewFamily, GraphBuilder, bShouldRenderVolumetricFog);
+
 		if (bShouldRenderVolumetricCloudBase)
 		{
 			InitVolumetricRenderTargetForViews(GraphBuilder, Views);
diff --git a/Engine/Source/Runtime/Renderer/Private/LocalFogVolumeRendering.cpp b/Engine/Source/Runtime/Renderer/Private/LocalFogVolumeRendering.cpp
--- a/Engine/Source/Runtime/Renderer/Private/LocalFogVolumeRendering.cpp
+++ b/Engine/Source/Runtime/Renderer/Private/LocalFogVolumeRendering.cpp
@@ -66,6 +66,11 @@ static TAutoConsoleVariable<float> CVarLocalFogVolumeGlobalStartDistance(
 	TEXT("The start distance in centimeter from which local fog volumes starts to appear."),
 	ECVF_RenderThreadSafe);
 
+static TAutoConsoleVariable<int32> CVarLocalFogVolumeUseHZB(
+	TEXT("r.LocalFogVolume.UseHZB"), 1,
+	TEXT("Use the HZB to cull loca lfog volume away.\n"),
+	ECVF_RenderThreadSafe);
+
 // Example of tile setup
 //  - 1920x1080 => 15x9 tiles
 //  - Allowing max 32 volumes at once => culling list buffer = 15 * 9 * 32 * 1 byte = 4320 bytes = 4.3KB
@@ -231,7 +236,8 @@ class FLocalFogVolumeTiledCullingCS : public FGlobalShader
 	DECLARE_GLOBAL_SHADER(FLocalFogVolumeTiledCullingCS);
 	SHADER_USE_PARAMETER_STRUCT(FLocalFogVolumeTiledCullingCS, FGlobalShader);
 
-	using FPermutationDomain = TShaderPermutationDomain<>;
+	class FUseHZB : SHADER_PERMUTATION_BOOL("USE_HZB");
+	using FPermutationDomain = TShaderPermutationDomain<FUseHZB>;
 
 public:
 	const static uint32 GroupSize = 8;
@@ -243,6 +249,11 @@ public:
 		SHADER_PARAMETER_RDG_BUFFER_SRV(Buffer<float4>, LocalFogVolumeCullingDataBuffer)
 		SHADER_PARAMETER_RDG_BUFFER_UAV(RWBuffer<uint>, TileDataBufferUAV)
 		SHADER_PARAMETER_RDG_BUFFER_UAV(RWBuffer<uint>, TileDrawIndirectBufferUAV)
+		SHADER_PARAMETER_RDG_TEXTURE(Texture2D, HZBTexture)
+		SHADER_PARAMETER_SAMPLER(SamplerState, HZBSampler)
+		SHADER_PARAMETER(FVector2f, HZBSize)
+		SHADER_PARAMETER(FVector2f, HZBViewSize)
+		SHADER_PARAMETER(FIntRect, HZBViewRect)
 		SHADER_PARAMETER(FVector4f, LeftPlane)
 		SHADER_PARAMETER(FVector4f, RightPlane)
 		SHADER_PARAMETER(FVector4f, TopPlane)
@@ -272,6 +283,12 @@ static void LocalFogVolumeViewTiledCullingPass(FViewInfo& View, FRDGBuilder& Gra
 	PassParameters->TileDataBufferUAV					= View.LocalFogVolumeViewData.GPUTileDataBufferUAV;
 	PassParameters->TileDrawIndirectBufferUAV			= View.LocalFogVolumeViewData.GPUTileDrawIndirectBufferUAV;
 
+	PassParameters->HZBTexture							= View.HZB;
+	PassParameters->HZBSampler							= TStaticSamplerState<SF_Point, AM_Clamp, AM_Clamp, AM_Clamp>::GetRHI();
+	PassParameters->HZBSize								= FVector2f(View.HZBMipmap0Size);
+	PassParameters->HZBViewSize							= FVector2f(View.ViewRect.Size());
+	PassParameters->HZBViewRect							= FIntRect(0, 0, View.ViewRect.Width(), View.ViewRect.Height());
+
 	auto ConvertPlanToVector4f = [&](FVector4f& OutVec4f, auto& Plane, bool bFlipPlane)
 	{
 		OutVec4f.X = Plane.X;
@@ -316,7 +333,10 @@ static void LocalFogVolumeViewTiledCullingPass(FViewInfo& View, FRDGBuilder& Gra
 	TileDataTextureSize.Z = 1;
 	const FIntVector NumGroups = FIntVector::DivideAndRoundUp(TileDataTextureSize, FLocalFogVolumeTiledCullingCS::GroupSize);
 
-	TShaderMapRef<FLocalFogVolumeTiledCullingCS> ComputeShader(View.ShaderMap);
+	
+	FLocalFogVolumeTiledCullingCS::FPermutationDomain PermutationVector;
+	PermutationVector.Set<FLocalFogVolumeTiledCullingCS::FUseHZB>(View.HZB != nullptr && CVarLocalFogVolumeUseHZB.GetValueOnRenderThread());
+	TShaderMapRef<FLocalFogVolumeTiledCullingCS> ComputeShader(View.ShaderMap, PermutationVector);
 	FComputeShaderUtils::AddPass(GraphBuilder, RDG_EVENT_NAME("LocalFogVolume.TiledCulling"), PassFlag, ComputeShader, PassParameters, NumGroups);
 }
 
