Disable half res mobile SSAO by default.
(Defaulted on by mistake in 29405978)

diff --git a/Engine/Source/Runtime/Renderer/Private/PostProcess/PostProcessAmbientOcclusionMobile.cpp b/Engine/Source/Runtime/Renderer/Private/PostProcess/PostProcessAmbientOcclusionMobile.cpp
--- a/Engine/Source/Runtime/Renderer/Private/PostProcess/PostProcessAmbientOcclusionMobile.cpp
+++ b/Engine/Source/Runtime/Renderer/Private/PostProcess/PostProcessAmbientOcclusionMobile.cpp
@@ -73,7 +73,7 @@ static TAutoConsoleVariable<int32> CVarMobileAmbientOcclusionDepthBoundsTest(
 
 static TAutoConsoleVariable<int32> CVarMobileSSAOHalfResolution(
 	TEXT("r.Mobile.SSAOHalfResolution"),
-	1,
+	0,
 	TEXT("Whether to calculate SSAO at half resolution.\n")
 	TEXT("0: Disabled.\n")
 	TEXT("1: Half Resolution with bilinear upsample\n")
