Add fog gpu tag for the mobile renderer.

diff --git a/Engine/Source/Runtime/Renderer/Private/MobileFogRendering.cpp b/Engine/Source/Runtime/Renderer/Private/MobileFogRendering.cpp
--- a/Engine/Source/Runtime/Renderer/Private/MobileFogRendering.cpp
+++ b/Engine/Source/Runtime/Renderer/Private/MobileFogRendering.cpp
@@ -20,6 +20,8 @@
 #include "SkyAtmosphereRendering.h"
 #include "LocalFogVolumeRendering.h"
 
+DECLARE_GPU_STAT(MobileFog);
+
 static TAutoConsoleVariable<int32> CVarPixelFogQuality(
 	TEXT("r.Mobile.PixelFogQuality"),
 	1,
@@ -122,6 +124,8 @@ IMPLEMENT_SHADER_TYPE(, FMobileFogPS, TEXT("/Engine/Private/MobileFog.usf"), TEX
 
 void FMobileSceneRenderer::RenderFog(FRHICommandList& RHICmdList, const FViewInfo& View)
 {
+	SCOPED_GPU_STAT(RHICmdList, MobileFog);
+
 	// RenderFog has some extra logic to skip the rendering of fog. So we account for that inside this function using a lambda.
 	bool bFogHasComposedLocalFogVolumes = false;
 	auto RenderLocalFogVolumeMobileLambda = [&]()
