Disable r.SceneCapture.AllowRenderInMainRenderer when Platform is Mobile and  MobileUsesFullDepthPrepass is false

diff --git a/Engine/Source/Runtime/Renderer/Private/SceneCaptureRendering.cpp b/Engine/Source/Runtime/Renderer/Private/SceneCaptureRendering.cpp
--- a/Engine/Source/Runtime/Renderer/Private/SceneCaptureRendering.cpp
+++ b/Engine/Source/Runtime/Renderer/Private/SceneCaptureRendering.cpp
@@ -896,7 +896,11 @@ void FScene::UpdateSceneCaptureContents(USceneCaptureComponent2D* CaptureCompone
 		}
 
 		// As optimization for depth capture modes, render scene capture as additional render passes inside the main renderer.
-		if (GSceneCaptureAllowRenderInMainRenderer && CaptureComponent->bRenderInMainRenderer && (CaptureComponent->CaptureSource == ESceneCaptureSource::SCS_SceneDepth || CaptureComponent->CaptureSource == ESceneCaptureSource::SCS_DeviceDepth))
+		if (GSceneCaptureAllowRenderInMainRenderer && 
+			CaptureComponent->bRenderInMainRenderer && 
+			(CaptureComponent->CaptureSource == ESceneCaptureSource::SCS_SceneDepth || CaptureComponent->CaptureSource == ESceneCaptureSource::SCS_DeviceDepth) &&
+			(!IsMobilePlatform(GetShaderPlatform()) || MobileUsesFullDepthPrepass(GetShaderPlatform()))
+			)
 		{
 			FSceneCaptureInfo CaptureInfo;
 			CaptureInfo.ViewLocation = ViewLocation;
