Modified GetEyeAdaptationParameters(...) to use GetBasicEyeAdaptationMinFeatureLevel() instead to taking MinFeatureLevel as an argument.

diff --git a/Engine/Source/Runtime/Renderer/Private/PostProcess/PostProcessEyeAdaptation.cpp b/Engine/Source/Runtime/Renderer/Private/PostProcess/PostProcessEyeAdaptation.cpp
--- a/Engine/Source/Runtime/Renderer/Private/PostProcess/PostProcessEyeAdaptation.cpp
+++ b/Engine/Source/Runtime/Renderer/Private/PostProcess/PostProcessEyeAdaptation.cpp
@@ -364,7 +364,7 @@ float CalculateManualAutoExposure(const FViewInfo& View, bool bForceDisablePhysi
 	return FoundLuminance;
 }
 
-FEyeAdaptationParameters GetEyeAdaptationParameters(const FViewInfo& View, ERHIFeatureLevel::Type MinFeatureLevel)
+FEyeAdaptationParameters GetEyeAdaptationParameters(const FViewInfo& View)
 {
 	const bool bExtendedLuminanceRange = IsExtendLuminanceRangeEnabled();
 
@@ -418,7 +418,7 @@ FEyeAdaptationParameters GetEyeAdaptationParameters(const FViewInfo& View, ERHIF
 		MinWhitePointLuminance = MaxWhitePointLuminance = CalculateFixedAutoExposure(View);
 	}
 	// The feature level check should always pass unless on mobile with MobileHDR is false
-	else if (EngineShowFlags.EyeAdaptation && View.GetFeatureLevel() >= MinFeatureLevel)
+	else if (EngineShowFlags.EyeAdaptation && View.GetFeatureLevel() >= GetBasicEyeAdaptationMinFeatureLevel())
 	{
 		if (AutoExposureMethod == EAutoExposureMethod::AEM_Manual)
 		{
@@ -538,7 +538,7 @@ FEyeAdaptationParameters GetEyeAdaptationParameters(const FViewInfo& View, ERHIF
 
 float GetEyeAdaptationFixedExposure(const FViewInfo& View)
 {
-	const FEyeAdaptationParameters Parameters = GetEyeAdaptationParameters(View, GetBasicEyeAdaptationMinFeatureLevel());
+	const FEyeAdaptationParameters Parameters = GetEyeAdaptationParameters(View);
 
 	const float Exposure = (Parameters.MinAverageLuminance + Parameters.MaxAverageLuminance) * 0.5f;
 
@@ -742,7 +742,7 @@ FRDGTextureRef AddCalculateExposureIlluminancePass(
 			PassParameters->RWIlluminanceTexture = GraphBuilder.CreateUAV(ExposureIlluminanceSetup);
 			PassParameters->Illuminance = GetScreenPassTextureViewportParameters(OutputViewport);
 			PassParameters->IllumiananceDownscaleFactor = GetAutoExposureIlluminanceDownscaleFactor();
-			PassParameters->EyeAdaptation = GetEyeAdaptationParameters(View, ERHIFeatureLevel::SM5);
+			PassParameters->EyeAdaptation = GetEyeAdaptationParameters(View);
 
 			PassParameters->PreIntegratedGF = GSystemTextures.PreintegratedGF->GetRHI();
 			PassParameters->PreIntegratedGFSampler = TStaticSamplerState<SF_Bilinear, AM_Clamp, AM_Clamp, AM_Clamp>::GetRHI();
diff --git a/Engine/Source/Runtime/Renderer/Private/PostProcess/PostProcessEyeAdaptation.h b/Engine/Source/Runtime/Renderer/Private/PostProcess/PostProcessEyeAdaptation.h
--- a/Engine/Source/Runtime/Renderer/Private/PostProcess/PostProcessEyeAdaptation.h
+++ b/Engine/Source/Runtime/Renderer/Private/PostProcess/PostProcessEyeAdaptation.h
@@ -77,7 +77,7 @@ BEGIN_SHADER_PARAMETER_STRUCT(FEyeAdaptationParameters, )
 	SHADER_PARAMETER_SAMPLER(SamplerState, MeterMaskSampler)
 END_SHADER_PARAMETER_STRUCT()
 
-FEyeAdaptationParameters GetEyeAdaptationParameters(const FViewInfo& ViewInfo, ERHIFeatureLevel::Type MinFeatureLevel);
+FEyeAdaptationParameters GetEyeAdaptationParameters(const FViewInfo& ViewInfo);
 
 // Computes the a fixed exposure to be used to replace the dynamic exposure when it's not supported (< SM5).
 float GetEyeAdaptationFixedExposure(const FViewInfo& View);
diff --git a/Engine/Source/Runtime/Renderer/Private/PostProcess/PostProcessing.cpp b/Engine/Source/Runtime/Renderer/Private/PostProcess/PostProcessing.cpp
--- a/Engine/Source/Runtime/Renderer/Private/PostProcess/PostProcessing.cpp
+++ b/Engine/Source/Runtime/Renderer/Private/PostProcess/PostProcessing.cpp
@@ -335,7 +335,7 @@ void AddPostProcessingPasses(
 	const FScreenPassTexture OriginalSceneColor = SceneColor;
 
 	// Default the new eye adaptation to the last one in case it's not generated this frame.
-	const FEyeAdaptationParameters EyeAdaptationParameters = GetEyeAdaptationParameters(View, ERHIFeatureLevel::SM5);
+	const FEyeAdaptationParameters EyeAdaptationParameters = GetEyeAdaptationParameters(View);
 	FRDGBufferRef LastEyeAdaptationBuffer = GetEyeAdaptationBuffer(GraphBuilder, View);
 	FRDGBufferRef EyeAdaptationBuffer = LastEyeAdaptationBuffer;
 
@@ -1729,7 +1729,7 @@ void AddDebugViewPostProcessingPasses(FRDGBuilder& GraphBuilder, const FViewInfo
 	PassSequence.SetEnabled(EPass::SecondaryUpscale, View.RequiresSecondaryUpscale() || View.Family->GetSecondarySpatialUpscalerInterface() != nullptr);
 	PassSequence.Finalize();
 
-	const FEyeAdaptationParameters EyeAdaptationParameters = GetEyeAdaptationParameters(View, ERHIFeatureLevel::SM5);
+	const FEyeAdaptationParameters EyeAdaptationParameters = GetEyeAdaptationParameters(View);
 
 	if (bTonemapBefore)
 	{
@@ -1996,7 +1996,7 @@ void AddMobilePostProcessingPasses(FRDGBuilder& GraphBuilder, FScene* Scene, con
 	FScreenPassTexture SceneDepthAux((*Inputs.SceneTextures)->SceneDepthAuxTexture, FinalOutputViewRect);
 
 	// Default the new eye adaptation to the last one in case it's not generated this frame.
-	const FEyeAdaptationParameters EyeAdaptationParameters = GetEyeAdaptationParameters(View, ERHIFeatureLevel::ES3_1);
+	const FEyeAdaptationParameters EyeAdaptationParameters = GetEyeAdaptationParameters(View);
 	FRDGBufferRef LastEyeAdaptationBuffer = GetEyeAdaptationBuffer(GraphBuilder, View);
 
 	const FPaniniProjectionConfig PaniniConfig(View);
