Added cvar to control which weights are used to calculate luminance in auto exposure logic.
- affects auto exposure, local exposure and eye adaptation.
- by default use existing behaviour (uniform weights) to avoid modifying exposure in existing projects

diff --git a/Engine/Shaders/Private/PostProcessEyeAdaptation.usf b/Engine/Shaders/Private/PostProcessEyeAdaptation.usf
--- a/Engine/Shaders/Private/PostProcessEyeAdaptation.usf
+++ b/Engine/Shaders/Private/PostProcessEyeAdaptation.usf
@@ -372,7 +372,7 @@ uint IllumiananceDownscaleFactor;
 // but skip the min since we apply EyeAdaptation_IgnoreMaterialsMinBaseColorLuminance here
 float CalculateEyeAdaptationLuminanceWithoutMin(float3 Color)
 {
-	return dot(Color, float3(1.0f, 1.0f, 1.0f) / 3.0f);
+	return dot(Color, EyeAdaptation_LuminanceWeights);
 }
 
 [numthreads(8, 8, 1)]
diff --git a/Engine/Shaders/Private/PostProcessHistogramCommon.ush b/Engine/Shaders/Private/PostProcessHistogramCommon.ush
--- a/Engine/Shaders/Private/PostProcessHistogramCommon.ush
+++ b/Engine/Shaders/Private/PostProcessHistogramCommon.ush
@@ -22,6 +22,7 @@ float EyeAdaptation_IgnoreMaterialsLuminanceScale;
 float EyeAdaptation_IgnoreMaterialsMinBaseColorLuminance;
 uint EyeAdaptation_IgnoreMaterialsReconstructFromSceneColor;
 float EyeAdaptation_GreyMult;
+float3 EyeAdaptation_LuminanceWeights;
 float EyeAdaptation_ExponentialUpM;
 float EyeAdaptation_ExponentialDownM;
 float EyeAdaptation_StartDistance;
@@ -43,7 +44,7 @@ float LocalExposure_ShadowThreshold;
 
 float CalculateEyeAdaptationLuminance(float3 Color)
 {
-	return max(dot(Color, float3(1.0f, 1.0f, 1.0f) / 3.0f), EyeAdaptation_LuminanceMin);
+	return max(dot(Color, EyeAdaptation_LuminanceWeights), EyeAdaptation_LuminanceMin);
 }
 
 // inverse of ComputeLogLuminanceFromHistogramPosition
diff --git a/Engine/Source/Runtime/Renderer/Private/PostProcess/PostProcessEyeAdaptation.cpp b/Engine/Source/Runtime/Renderer/Private/PostProcess/PostProcessEyeAdaptation.cpp
--- a/Engine/Source/Runtime/Renderer/Private/PostProcess/PostProcessEyeAdaptation.cpp
+++ b/Engine/Source/Runtime/Renderer/Private/PostProcess/PostProcessEyeAdaptation.cpp
@@ -78,6 +78,14 @@ namespace
 		TEXT("(default: false)"),
 		ECVF_Scalability | ECVF_RenderThreadSafe);
 
+	TAutoConsoleVariable<int32> CVarAutoExposureLuminanceMethod(
+		TEXT("r.AutoExposure.LuminanceMethod"),
+		0,
+		TEXT("0 - Uniform.\n")
+		TEXT("1 - NSTC.\n")
+		TEXT("2 - Rec709."),
+		ECVF_RenderThreadSafe);
+
 	TAutoConsoleVariable<bool> CVarAutoExposureIgnoreMaterialsReconstructFromSceneColor(
 		TEXT("r.AutoExposure.IgnoreMaterials.ReconstructFromSceneColor"),
 		true,
@@ -533,6 +541,24 @@ FEyeAdaptationParameters GetEyeAdaptationParameters(const FViewInfo& View)
 	Parameters.VisualizeDebugType = CVarEyeAdaptationVisualizeDebugType.GetValueOnRenderThread();
 	Parameters.MeterMaskTexture = MeterMask;
 	Parameters.MeterMaskSampler = TStaticSamplerState<SF_Bilinear, AM_Clamp, AM_Clamp, AM_Clamp>::GetRHI();
+
+	const int32 LuminanceMethod = CVarAutoExposureLuminanceMethod.GetValueOnRenderThread();
+	if (LuminanceMethod == 1)
+	{
+		// NTSC / match weights in Common.ush
+		Parameters.LuminanceWeights = FVector3f(0.3f, 0.59f, 0.11f);
+	}
+	else if (LuminanceMethod == 2)
+	{
+		// Rec 709
+		Parameters.LuminanceWeights = FVector3f(0.2126f, 0.7152f, 0.0722f);
+	}
+	else
+	{
+		// default (uniform weights)
+		Parameters.LuminanceWeights = FVector3f(1.0f / 3.0f);
+	}
+
 	return Parameters;
 }
 
@@ -613,6 +639,7 @@ class FSetupExposureIlluminanceCS : public FGlobalShader
 		SHADER_PARAMETER_RDG_TEXTURE(Texture2D, ColorTexture)
 		SHADER_PARAMETER_RDG_TEXTURE_UAV(RWTexture2D<float4>, RWIlluminanceTexture)
 		SHADER_PARAMETER_STRUCT(FScreenPassTextureViewportParameters, Illuminance)
+		SHADER_PARAMETER_STRUCT(FEyeAdaptationParameters, EyeAdaptation)
 		SHADER_PARAMETER(uint32, IllumiananceDownscaleFactor)
 	END_SHADER_PARAMETER_STRUCT()
 
@@ -663,6 +690,7 @@ FRDGTextureRef AddSetupExposureIlluminancePass(
 			PassParameters->ColorTexture = SceneTextures.Color.Resolve;
 			PassParameters->RWIlluminanceTexture = GraphBuilder.CreateUAV(OutputTexture);
 			PassParameters->Illuminance = GetScreenPassTextureViewportParameters(OutputViewport);
+			PassParameters->EyeAdaptation = GetEyeAdaptationParameters(View);
 			PassParameters->IllumiananceDownscaleFactor = GetAutoExposureIlluminanceDownscaleFactor();
 
 			auto ComputeShader = View.ShaderMap->GetShader<FSetupExposureIlluminanceCS>();
diff --git a/Engine/Source/Runtime/Renderer/Private/PostProcess/PostProcessEyeAdaptation.h b/Engine/Source/Runtime/Renderer/Private/PostProcess/PostProcessEyeAdaptation.h
--- a/Engine/Source/Runtime/Renderer/Private/PostProcess/PostProcessEyeAdaptation.h
+++ b/Engine/Source/Runtime/Renderer/Private/PostProcess/PostProcessEyeAdaptation.h
@@ -63,6 +63,7 @@ BEGIN_SHADER_PARAMETER_STRUCT(FEyeAdaptationParameters, )
 	SHADER_PARAMETER(float, LuminanceMin)
 	SHADER_PARAMETER(float, BlackHistogramBucketInfluence)
 	SHADER_PARAMETER(float, GreyMult)
+	SHADER_PARAMETER(FVector3f, LuminanceWeights)
 	SHADER_PARAMETER(float, ExponentialUpM)
 	SHADER_PARAMETER(float, ExponentialDownM)
 	SHADER_PARAMETER(float, StartDistance)
