Fix potential nullptr reference spotted by static analyzer.

diff --git a/Engine/Source/Runtime/Renderer/Private/MeshDrawCommandStats.cpp b/Engine/Source/Runtime/Renderer/Private/MeshDrawCommandStats.cpp
--- a/Engine/Source/Runtime/Renderer/Private/MeshDrawCommandStats.cpp
+++ b/Engine/Source/Runtime/Renderer/Private/MeshDrawCommandStats.cpp
@@ -335,7 +335,7 @@ void FMeshDrawCommandStatsManager::Update()
 								Stats.CustomIndirectPrimitives += DrawData.VisibleInstanceCount * DrawData.PrimitiveCount;
 							}
 						}
-						else if (DrawData.UseInstantCullingIndirectBuffer > 0 && InstanceCullingReadBackData)
+						else if (IndirectArgsPtr && DrawData.UseInstantCullingIndirectBuffer > 0 && InstanceCullingReadBackData)
 						{
 							const FRHIDrawIndexedIndirectParameters& IndirectArgs = IndirectArgsPtr[PassStats->IndirectArgParameterOffset + IndirectCommandIndex];
 							ensure(DrawData.PrimitiveCount == IndirectArgs.IndexCountPerInstance / 3);
