VulkanRHI: VKRT WIP: Avoid lock of size 0.

#rnx

diff --git a/Engine/Source/Runtime/Renderer/Private/RayTracing/RayTracingScene.cpp b/Engine/Source/Runtime/Renderer/Private/RayTracing/RayTracingScene.cpp
--- a/Engine/Source/Runtime/Renderer/Private/RayTracing/RayTracingScene.cpp
+++ b/Engine/Source/Runtime/Renderer/Private/RayTracing/RayTracingScene.cpp
@@ -170,7 +170,7 @@ void FRayTracingScene::CreateWithInitializationData(FRDGBuilder& GraphBuilder, c
 		const uint32 TransformUploadBytes = SceneWithGeometryInstances.NumNativeCPUInstances * 3 * sizeof(FVector4f);
 
 		FRayTracingInstanceDescriptorInput* InstanceUploadData = (FRayTracingInstanceDescriptorInput*)RHICmdList.LockBuffer(InstanceUploadBuffer, 0, InstanceUploadBytes, RLM_WriteOnly);
-		FVector4f* TransformUploadData = (FVector4f*)RHICmdList.LockBuffer(TransformUploadBuffer, 0, TransformUploadBytes, RLM_WriteOnly);
+		FVector4f* TransformUploadData = (TransformUploadBytes > 0) ? (FVector4f*)RHICmdList.LockBuffer(TransformUploadBuffer, 0, TransformUploadBytes, RLM_WriteOnly) : nullptr;
 
 		// Fill instance upload buffer on separate thread since results are only needed in RHI thread
 		FillInstanceUploadBufferTask = FFunctionGraphTask::CreateAndDispatchWhenReady(
@@ -238,7 +238,11 @@ void FRayTracingScene::CreateWithInitializationData(FRDGBuilder& GraphBuilder, c
 			{
 				WaitForTasks();
 				RHICmdList.UnlockBuffer(InstanceUploadBuffer);
-				RHICmdList.UnlockBuffer(TransformUploadBuffer);
+
+				if (NumNativeCPUInstances > 0)
+				{
+					RHICmdList.UnlockBuffer(TransformUploadBuffer);
+				}
 
 				// Pull this out here, because command list playback (where the lambda is executed) doesn't update the GPU mask
 				FRHIGPUMask IterateGPUMasks = RHICmdList.GetGPUMask();
