Fix incorrect Substrate tile classification visualization when using screen percentage and realtime update is turned off.

diff --git a/Engine/Source/Runtime/Renderer/Private/Substrate/Substrate.cpp b/Engine/Source/Runtime/Renderer/Private/Substrate/Substrate.cpp
--- a/Engine/Source/Runtime/Renderer/Private/Substrate/Substrate.cpp
+++ b/Engine/Source/Runtime/Renderer/Private/Substrate/Substrate.cpp
@@ -1029,7 +1029,7 @@ void AddSubstrateInternalClassificationTilePass(
 		// The ColorTexture can be larger than the screen resolution if the screen percentage has be manipulated to be >100%.
 		// So we simply re-use the previously computed ViewResolutionFraction to recover the targeted resolution in the editor.
 		// TODO fix this for split screen.
-		const float InvViewResolutionFraction = 1.0f / View.CachedViewUniformShaderParameters->ViewResolutionFraction;
+		const float InvViewResolutionFraction = View.Family->bRealtimeUpdate ? 1.0f / View.CachedViewUniformShaderParameters->ViewResolutionFraction : 1.0f;
 		DebugOutputResolution = FIntPoint(float(ViewRect.Width()) * InvViewResolutionFraction, float(ViewRect.Height()) * InvViewResolutionFraction);
 
 		check(ColorTexture);
