* Fix non-deterministic hair strands macro groom build. Due to rounding logic / GPU feedback this was causing voxel size to be inconsistent frame to frame when macro groups are different. Now preorder visible instance by register index to get consitent groups.
* Fix incorrect materialId <-> hit proxy color and materialId <-> selected proxy mapping
* Fix insufficient shader print line count for displaying groom instance view mode.

diff --git a/Engine/Plugins/Runtime/HairStrands/Source/HairStrandsCore/Private/GroomDebug.cpp b/Engine/Plugins/Runtime/HairStrands/Source/HairStrandsCore/Private/GroomDebug.cpp
--- a/Engine/Plugins/Runtime/HairStrands/Source/HairStrandsCore/Private/GroomDebug.cpp
+++ b/Engine/Plugins/Runtime/HairStrands/Source/HairStrandsCore/Private/GroomDebug.cpp
@@ -885,7 +885,7 @@ static void AddHairDebugPrintInstancePass(
 	// Force ShaderPrint on.
 	ShaderPrint::SetEnabled(true);
 	// Request more drawing primitives & characters for printing if needed	
-	ShaderPrint::RequestSpaceForLines(InstanceCount * 16u);
+	ShaderPrint::RequestSpaceForLines(InstanceCount * 64u);
 	ShaderPrint::RequestSpaceForCharacters(InstanceCount * 256 + 512);
 
 	if (!ShaderPrintData || InstanceCount == 0) { return; }
diff --git a/Engine/Source/Runtime/Renderer/Private/HairStrands/HairStrandsCluster.cpp b/Engine/Source/Runtime/Renderer/Private/HairStrands/HairStrandsCluster.cpp
--- a/Engine/Source/Runtime/Renderer/Private/HairStrands/HairStrandsCluster.cpp
+++ b/Engine/Source/Runtime/Renderer/Private/HairStrands/HairStrandsCluster.cpp
@@ -263,21 +263,35 @@ void CreateHairStrandsMacroGroups(
 		}
 	};
 
-	// 1. Add all visible hair-strands instances
-	static FBoxSphereBounds EmptyBound(ForceInit);
-	const int32 ActiveInstanceCount = Scene->HairStrandsSceneData.RegisteredProxies.Num();
-	TBitArray InstancesVisibility(false, ActiveInstanceCount);
+	// 0. Pre-sort visible instances by register index to get stable macro-groups
+	struct FVisibleBatch
+	{
+		const FMeshBatchAndRelevance* Batch = nullptr;
+		FHairGroupPublicData* HairData = nullptr;
+	};
+	TArray<FVisibleBatch> VisibleBatches;
+	VisibleBatches.Reserve(View.HairStrandsMeshElements.Num());
 	for (const FMeshBatchAndRelevance& MeshBatchAndRelevance : View.HairStrandsMeshElements)
 	{
 		if (HairStrands::IsHairStrandsVF(MeshBatchAndRelevance.Mesh))
 		{
 			if (FHairGroupPublicData* HairData = HairStrands::GetHairData(MeshBatchAndRelevance.Mesh))
 			{
-				UpdateMacroGroup(HairData, MeshBatchAndRelevance.Mesh, MeshBatchAndRelevance.PrimitiveSceneProxy, EmptyBound);
-				InstancesVisibility[HairData->Instance->RegisteredIndex] = true;
+				VisibleBatches.Add( { &MeshBatchAndRelevance, HairData });
 			}
 		}
 	}
+	VisibleBatches.Sort([](const FVisibleBatch& A, const FVisibleBatch& B) { return A.HairData->Instance->RegisteredIndex < B.HairData->Instance->RegisteredIndex; });
+
+	// 1. Add all visible hair-strands instances
+	static FBoxSphereBounds EmptyBound(ForceInit);
+	const int32 ActiveInstanceCount = Scene->HairStrandsSceneData.RegisteredProxies.Num();
+	TBitArray InstancesVisibility(false, ActiveInstanceCount);
+	for (FVisibleBatch& VisibleBatch : VisibleBatches)
+	{
+		UpdateMacroGroup(VisibleBatch.HairData, VisibleBatch.Batch->Mesh, VisibleBatch.Batch->PrimitiveSceneProxy, EmptyBound);
+		InstancesVisibility[VisibleBatch.HairData->Instance->RegisteredIndex] = true;
+	}
 
 	// 2. Add all hair-strands instances which are non-visible in primary view(s) but visible in shadow view(s)
 	// Slow Linear search
diff --git a/Engine/Source/Runtime/Renderer/Private/HairStrands/HairStrandsVisibility.cpp b/Engine/Source/Runtime/Renderer/Private/HairStrands/HairStrandsVisibility.cpp
--- a/Engine/Source/Runtime/Renderer/Private/HairStrands/HairStrandsVisibility.cpp
+++ b/Engine/Source/Runtime/Renderer/Private/HairStrands/HairStrandsVisibility.cpp
@@ -3853,6 +3853,7 @@ void AddHairStrandsSelectionOutlinePass(
 	FRDGBuilder& GraphBuilder,
 	const FViewInfo& View,
 	const FIntRect& ViewportRect,
+	const FHairStrandsMacroGroupDatas& MacroGroupDatas,
 	FRDGTextureRef VisNodeIndex,
 	FRDGBufferRef VisNodeData,
 	FRDGTextureRef CoverageTexture,
@@ -3864,13 +3865,19 @@ void AddHairStrandsSelectionOutlinePass(
 	}
 
 #if WITH_EDITOR
-	// Create mapping table between PrimitiveId and BatchId
+	// Create mapping table between MaterialId and BatchId
 	TArray<uint32> SelectionMaterialId;
-	SelectionMaterialId.Reserve(View.HairStrandsMeshElements.Num());
-	for (const FMeshBatchAndRelevance& MeshBatch : View.HairStrandsMeshElements)
+	SelectionMaterialId.Init(0u, View.HairStrandsMeshElements.Num());
+	for (const FHairStrandsMacroGroupData& MacroGroupData : MacroGroupDatas)
 	{
-		const uint32 bSelected = MeshBatch.PrimitiveSceneProxy->IsSelected() ? 1u : 0u;
-		SelectionMaterialId.Add(bSelected);
+		for (const FHairStrandsMacroGroupData::PrimitiveInfo& PrimitiveInfo : MacroGroupData.PrimitivesInfos)
+		{
+			if (SelectionMaterialId.IsValidIndex(PrimitiveInfo.MaterialId) && PrimitiveInfo.PrimitiveSceneProxy)
+			{
+				const uint32 bSelected = PrimitiveInfo.PrimitiveSceneProxy->IsSelected() ? 1u : 0u;
+				SelectionMaterialId[PrimitiveInfo.MaterialId] = bSelected;
+			}
+		}
 	}
 
 	FRDGBufferRef SelectionMaterialIdBuffer = CreateUploadBuffer(GraphBuilder, TEXT("Hair.MaterialIdToHitProxyIdBuffer"), sizeof(uint32), SelectionMaterialId.Num(), SelectionMaterialId.GetData(), sizeof(uint32) * SelectionMaterialId.Num());
@@ -3938,6 +3945,7 @@ void AddHairStrandsHitProxyIdPass(
 	FRDGBuilder& GraphBuilder,
 	const FScene& Scene,
 	const FViewInfo& View,
+	const FHairStrandsMacroGroupDatas& MacroGroupDatas,
 	FRDGTextureRef VisNodeIndex,
 	FRDGBufferRef VisNodeData,
 	FRDGTextureRef CoverageTexture,
@@ -3950,13 +3958,22 @@ void AddHairStrandsHitProxyIdPass(
 		return;
 	}
 
-	// Create mapping table between PrimitiveId and BatchId
+	// Create mapping table between MaterialId and BatchId
 	TArray<uint32> MaterialIdToHitProxyId;
-	MaterialIdToHitProxyId.Reserve(View.HairStrandsMeshElements.Num());
-	for (const FMeshBatchAndRelevance& MeshBatch : View.HairStrandsMeshElements)
+	MaterialIdToHitProxyId.Init(0u, View.HairStrandsMeshElements.Num());
+	for (const FHairStrandsMacroGroupData& MacroGroupData : MacroGroupDatas)
 	{
-		uint32 HitColor = MeshBatch.Mesh->BatchHitProxyId.GetColor().DWColor();
-		MaterialIdToHitProxyId.Add(HitColor);
+		for (const FHairStrandsMacroGroupData::PrimitiveInfo& PrimitiveInfo : MacroGroupData.PrimitivesInfos)
+		{
+			if (MaterialIdToHitProxyId.IsValidIndex(PrimitiveInfo.MaterialId))
+			{
+				if (const FMeshBatch* MeshBatch = PrimitiveInfo.Mesh)
+				{
+					const uint32 HitColor = MeshBatch->BatchHitProxyId.GetColor().DWColor();
+					MaterialIdToHitProxyId[PrimitiveInfo.MaterialId] = HitColor;
+				}
+			}
+		}
 	}
 
 	FRDGBufferRef MaterialIdToHitProxyIdBuffer = CreateUploadBuffer(GraphBuilder, TEXT("Hair.MaterialIdToHitProxyIdBuffer"), sizeof(uint32), MaterialIdToHitProxyId.Num(), MaterialIdToHitProxyId.GetData(), sizeof(uint32) * MaterialIdToHitProxyId.Num());
@@ -4073,6 +4090,7 @@ void DrawEditorSelection(FRDGBuilder& GraphBuilder, const FViewInfo& View, const
 		GraphBuilder,
 		View,
 		ViewportRect,
+		View.HairStrandsViewData.MacroGroupDatas,
 		View.HairStrandsViewData.VisibilityData.NodeIndex,
 		View.HairStrandsViewData.VisibilityData.NodeVisData,
 		View.HairStrandsViewData.VisibilityData.CoverageTexture,
@@ -4211,7 +4229,7 @@ void DrawHitProxies(
 			OutMaxNodeCount);
 	}
 
-	AddHairStrandsHitProxyIdPass(GraphBuilder, Scene, View, VisNodeIndex, VisNodeData, CoverageTexture, HitProxyTexture, HitProxyDepthTexture);
+	AddHairStrandsHitProxyIdPass(GraphBuilder, Scene, View, MacroGroupDatas, VisNodeIndex, VisNodeData, CoverageTexture, HitProxyTexture, HitProxyDepthTexture);
 }
 
 // Check if any simulated/skinned-bound groom has its positions updated (e.g. for invalidating the path-tracer accumulation)
