Removed RenderTargetTemp.h, which isn't used anywhere anymore

#tests editor

diff --git a/Engine/Source/Runtime/Renderer/Private/HairStrands/HairStrandsDebug.cpp b/Engine/Source/Runtime/Renderer/Private/HairStrands/HairStrandsDebug.cpp
--- a/Engine/Source/Runtime/Renderer/Private/HairStrands/HairStrandsDebug.cpp
+++ b/Engine/Source/Runtime/Renderer/Private/HairStrands/HairStrandsDebug.cpp
@@ -22,7 +22,6 @@
 #include "PostProcess/SceneRenderTargets.h"
 #include "SceneTextureParameters.h"
 #include "DynamicPrimitiveDrawing.h"
-#include "RenderTargetTemp.h"
 #include "CanvasTypes.h"
 #include "ShaderPrintParameters.h"
 #include "RenderGraphUtils.h"
diff --git a/Engine/Source/Runtime/Renderer/Private/PhysicsFieldRendering.cpp b/Engine/Source/Runtime/Renderer/Private/PhysicsFieldRendering.cpp
--- a/Engine/Source/Runtime/Renderer/Private/PhysicsFieldRendering.cpp
+++ b/Engine/Source/Runtime/Renderer/Private/PhysicsFieldRendering.cpp
@@ -16,7 +16,6 @@
 #include "RenderGraphUtils.h"
 #include "SceneTextureParameters.h"
 #include "DynamicPrimitiveDrawing.h"
-#include "RenderTargetTemp.h"
 #include "ShaderPrint.h"
 
 // Console variables
diff --git a/Engine/Source/Runtime/Renderer/Private/PostProcess/PostProcessBokehDOF.cpp b/Engine/Source/Runtime/Renderer/Private/PostProcess/PostProcessBokehDOF.cpp
--- a/Engine/Source/Runtime/Renderer/Private/PostProcess/PostProcessBokehDOF.cpp
+++ b/Engine/Source/Runtime/Renderer/Private/PostProcess/PostProcessBokehDOF.cpp
@@ -3,8 +3,8 @@
 #include "PostProcess/PostProcessBokehDOF.h"
 #include "PostProcess/PostProcessDOF.h"
 #include "CanvasTypes.h"
-#include "RenderTargetTemp.h"
 #include "PostProcess/DiaphragmDOF.h"
+#include "SceneRendering.h"
 #include "DataDrivenShaderPlatformInfo.h"
 #include "UnrealEngine.h"
 
diff --git a/Engine/Source/Runtime/Renderer/Private/PostProcess/PostProcessGBufferHints.cpp b/Engine/Source/Runtime/Renderer/Private/PostProcess/PostProcessGBufferHints.cpp
--- a/Engine/Source/Runtime/Renderer/Private/PostProcess/PostProcessGBufferHints.cpp
+++ b/Engine/Source/Runtime/Renderer/Private/PostProcess/PostProcessGBufferHints.cpp
@@ -3,7 +3,7 @@
 #include "PostProcess/PostProcessGBufferHints.h"
 #include "CanvasTypes.h"
 #include "DataDrivenShaderPlatformInfo.h"
-#include "RenderTargetTemp.h"
+#include "SceneRendering.h"
 #include "SceneTextureParameters.h"
 #include "UnrealEngine.h"
 
diff --git a/Engine/Source/Runtime/Renderer/Private/PostProcess/PostProcessMotionBlur.cpp b/Engine/Source/Runtime/Renderer/Private/PostProcess/PostProcessMotionBlur.cpp
--- a/Engine/Source/Runtime/Renderer/Private/PostProcess/PostProcessMotionBlur.cpp
+++ b/Engine/Source/Runtime/Renderer/Private/PostProcess/PostProcessMotionBlur.cpp
@@ -5,7 +5,6 @@
 #include "CanvasTypes.h"
 #include "DataDrivenShaderPlatformInfo.h"
 #include "Engine/Engine.h"
-#include "RenderTargetTemp.h"
 #include "ScenePrivate.h"
 #include "SpriteIndexBuffer.h"
 #include "PostProcess/PostProcessing.h"
diff --git a/Engine/Source/Runtime/Renderer/Private/PostProcess/PostProcessSelectionOutline.cpp b/Engine/Source/Runtime/Renderer/Private/PostProcess/PostProcessSelectionOutline.cpp
--- a/Engine/Source/Runtime/Renderer/Private/PostProcess/PostProcessSelectionOutline.cpp
+++ b/Engine/Source/Runtime/Renderer/Private/PostProcess/PostProcessSelectionOutline.cpp
@@ -7,7 +7,6 @@
 #include "EditorPrimitivesRendering.h"
 #include "SceneTextureParameters.h"
 #include "CanvasTypes.h"
-#include "RenderTargetTemp.h"
 #include "ClearQuad.h"
 #include "ScenePrivate.h"
 #include "PostProcess/SceneRenderTargets.h"
diff --git a/Engine/Source/Runtime/Renderer/Private/PostProcess/PostProcessSubsurface.cpp b/Engine/Source/Runtime/Renderer/Private/PostProcess/PostProcessSubsurface.cpp
--- a/Engine/Source/Runtime/Renderer/Private/PostProcess/PostProcessSubsurface.cpp
+++ b/Engine/Source/Runtime/Renderer/Private/PostProcess/PostProcessSubsurface.cpp
@@ -15,7 +15,6 @@
 #include "DataDrivenShaderPlatformInfo.h"
 #include "Engine/SubsurfaceProfile.h"
 #include "CanvasTypes.h"
-#include "RenderTargetTemp.h"
 #include "ScenePrivate.h"
 #include "GenerateMips.h"
 #include "ClearQuad.h"
diff --git a/Engine/Source/Runtime/Renderer/Private/PostProcess/PostProcessTestImage.cpp b/Engine/Source/Runtime/Renderer/Private/PostProcess/PostProcessTestImage.cpp
--- a/Engine/Source/Runtime/Renderer/Private/PostProcess/PostProcessTestImage.cpp
+++ b/Engine/Source/Runtime/Renderer/Private/PostProcess/PostProcessTestImage.cpp
@@ -4,7 +4,7 @@
 #include "PostProcess/PostProcessCombineLUTs.h"
 #include "CanvasTypes.h"
 #include "DataDrivenShaderPlatformInfo.h"
-#include "RenderTargetTemp.h"
+#include "SceneRendering.h"
 #include "UnrealEngine.h"
 
 class FTestImagePS : public FGlobalShader
diff --git a/Engine/Source/Runtime/Renderer/Private/PostProcess/PostProcessVisualizeComplexity.cpp b/Engine/Source/Runtime/Renderer/Private/PostProcess/PostProcessVisualizeComplexity.cpp
--- a/Engine/Source/Runtime/Renderer/Private/PostProcess/PostProcessVisualizeComplexity.cpp
+++ b/Engine/Source/Runtime/Renderer/Private/PostProcess/PostProcessVisualizeComplexity.cpp
@@ -2,8 +2,8 @@
 
 #include "PostProcess/PostProcessVisualizeComplexity.h"
 #include "CanvasTypes.h"
-#include "RenderTargetTemp.h"
 #include "UnrealEngine.h"
+#include "SceneRendering.h"
 #include "SystemTextures.h"
 #include "DataDrivenShaderPlatformInfo.h"
 
diff --git a/Engine/Source/Runtime/Renderer/Private/PostProcess/PostProcessVisualizeLevelInstance.cpp b/Engine/Source/Runtime/Renderer/Private/PostProcess/PostProcessVisualizeLevelInstance.cpp
--- a/Engine/Source/Runtime/Renderer/Private/PostProcess/PostProcessVisualizeLevelInstance.cpp
+++ b/Engine/Source/Runtime/Renderer/Private/PostProcess/PostProcessVisualizeLevelInstance.cpp
@@ -4,9 +4,9 @@
 
 #include "PostProcess/PostProcessVisualizeLevelInstance.h"
 #include "PostProcess/PostProcessCompositePrimitivesCommon.h"
+#include "SceneRendering.h"
 #include "SceneTextureParameters.h"
 #include "CanvasTypes.h"
-#include "RenderTargetTemp.h"
 #include "ClearQuad.h"
 #include "ScenePrivate.h"
 
diff --git a/Engine/Source/Runtime/Renderer/Private/PostProcess/VisualizeShadingModels.cpp b/Engine/Source/Runtime/Renderer/Private/PostProcess/VisualizeShadingModels.cpp
--- a/Engine/Source/Runtime/Renderer/Private/PostProcess/VisualizeShadingModels.cpp
+++ b/Engine/Source/Runtime/Renderer/Private/PostProcess/VisualizeShadingModels.cpp
@@ -2,7 +2,7 @@
 
 #include "PostProcess/VisualizeShadingModels.h"
 #include "CanvasTypes.h"
-#include "RenderTargetTemp.h"
+#include "SceneRendering.h"
 #include "SceneTextureParameters.h"
 #include "UnrealEngine.h"
 #include "DataDrivenShaderPlatformInfo.h"
diff --git a/Engine/Source/Runtime/Renderer/Private/ReflectionEnvironmentRealTimeCapture.cpp b/Engine/Source/Runtime/Renderer/Private/ReflectionEnvironmentRealTimeCapture.cpp
--- a/Engine/Source/Runtime/Renderer/Private/ReflectionEnvironmentRealTimeCapture.cpp
+++ b/Engine/Source/Runtime/Renderer/Private/ReflectionEnvironmentRealTimeCapture.cpp
@@ -22,7 +22,6 @@
 
 #if WITH_EDITOR
 #include "CanvasTypes.h"
-#include "RenderTargetTemp.h"
 #endif
 
 extern float GReflectionCaptureNearPlane;
diff --git a/Engine/Source/Runtime/Renderer/Private/RenderTargetTemp.h b/Engine/Source/Runtime/Renderer/Private/RenderTargetTemp.h
deleted file mode 100644
diff --git a/Engine/Source/Runtime/Renderer/Private/RendererPrivate.h b/Engine/Source/Runtime/Renderer/Private/RendererPrivate.h
--- a/Engine/Source/Runtime/Renderer/Private/RendererPrivate.h
+++ b/Engine/Source/Runtime/Renderer/Private/RendererPrivate.h
@@ -14,5 +14,4 @@
 #include "LightMapDensityRendering.h"
 #include "ClearQuad.h"
 #include "ScenePrivate.h"
-#include "RenderTargetTemp.h"
 #include "RendererModule.h"
diff --git a/Engine/Source/Runtime/Renderer/Private/SceneRendering.cpp b/Engine/Source/Runtime/Renderer/Private/SceneRendering.cpp
--- a/Engine/Source/Runtime/Renderer/Private/SceneRendering.cpp
+++ b/Engine/Source/Runtime/Renderer/Private/SceneRendering.cpp
@@ -17,7 +17,6 @@
 #include "DeferredShadingRenderer.h"
 #include "DumpGPU.h"
 #include "DynamicPrimitiveDrawing.h"
-#include "RenderTargetTemp.h"
 #include "RendererModule.h"
 #include "ScenePrivate.h"
 #include "PostProcess/DiaphragmDOF.h"
diff --git a/Engine/Source/Runtime/Renderer/Private/SkyAtmosphereRendering.cpp b/Engine/Source/Runtime/Renderer/Private/SkyAtmosphereRendering.cpp
--- a/Engine/Source/Runtime/Renderer/Private/SkyAtmosphereRendering.cpp
+++ b/Engine/Source/Runtime/Renderer/Private/SkyAtmosphereRendering.cpp
@@ -10,7 +10,6 @@
 #include "DeferredShadingRenderer.h"
 #include "LightSceneInfo.h"
 #include "PixelShaderUtils.h"
-#include "RenderTargetTemp.h"
 #include "Rendering/SkyAtmosphereCommonData.h"
 #include "ScenePrivate.h"
 #include "SceneRenderTargetParameters.h"
diff --git a/Engine/Source/Runtime/Renderer/Private/VolumetricCloudRendering.cpp b/Engine/Source/Runtime/Renderer/Private/VolumetricCloudRendering.cpp
--- a/Engine/Source/Runtime/Renderer/Private/VolumetricCloudRendering.cpp
+++ b/Engine/Source/Runtime/Renderer/Private/VolumetricCloudRendering.cpp
@@ -20,7 +20,6 @@
 #include "DynamicPrimitiveDrawing.h"
 #include "ShaderPrint.h"
 #include "CanvasTypes.h"
-#include "RenderTargetTemp.h"
 #include "VolumetricRenderTarget.h"
 #include "BlueNoise.h"
 #include "FogRendering.h"
