Enable FPopulateCompositeDepthPS for Metal platforms as MSAA is supported on SM3.1. Also added decode dim shader permutation logic that is required for other Metal 3.1 post process passes to ensure colour output is correct.

diff --git a/Engine/Shaders/Private/PostProcessCompositePrimitives.usf b/Engine/Shaders/Private/PostProcessCompositePrimitives.usf
--- a/Engine/Shaders/Private/PostProcessCompositePrimitives.usf
+++ b/Engine/Shaders/Private/PostProcessCompositePrimitives.usf
@@ -264,6 +264,11 @@ void MainPopulateSceneDepthPS(
 		Color_UVViewportBilinearMin, Color_UVViewportBilinearMax);
 	ColorUV = clamp(ColorUV, float2(0, 0), float2(1, 1));
 	OutColor = Texture2DSampleLevel(ColorTexture, ColorSampler, ColorUV, 0).rgba;
+
+#if METAL_MSAA_HDR_DECODE
+	OutColor.rgb *= rcp(OutColor.r * (-0.299) + OutColor.g * (-0.587) + OutColor.b * (-0.114) + 1.0);
+#endif
+
 #else
 	OutColor = 0.0;
 #endif
diff --git a/Engine/Source/Runtime/Renderer/Private/PostProcess/PostProcessCompositeDebugPrimitives.cpp b/Engine/Source/Runtime/Renderer/Private/PostProcess/PostProcessCompositeDebugPrimitives.cpp
--- a/Engine/Source/Runtime/Renderer/Private/PostProcess/PostProcessCompositeDebugPrimitives.cpp
+++ b/Engine/Source/Runtime/Renderer/Private/PostProcess/PostProcessCompositeDebugPrimitives.cpp
@@ -56,7 +56,8 @@ FScreenPassTexture AddDebugPrimitivePass(
 		DebugPrimitiveDepth,
 		SceneDepthJitter,
 		NumMSAASamples,
-		true);	// - Draw Background Colour
+		true,
+		Inputs.bUseMetalMSAAHDRDecode);
 
 	//Composite the debug draw elements into the scene
 	FDebugPrimitivesPassParameters* PassParameters = GraphBuilder.AllocParameters<FDebugPrimitivesPassParameters>();
diff --git a/Engine/Source/Runtime/Renderer/Private/PostProcess/PostProcessCompositePrimitivesCommon.cpp b/Engine/Source/Runtime/Renderer/Private/PostProcess/PostProcessCompositePrimitivesCommon.cpp
--- a/Engine/Source/Runtime/Renderer/Private/PostProcess/PostProcessCompositePrimitivesCommon.cpp
+++ b/Engine/Source/Runtime/Renderer/Private/PostProcess/PostProcessCompositePrimitivesCommon.cpp
@@ -92,7 +92,8 @@ public:
 
 	class FUseMSAADimension : SHADER_PERMUTATION_BOOL("USE_MSAA");
 	class FForceDebugDrawColorOutput : SHADER_PERMUTATION_BOOL("FORCE_DEBUG_DRAW_COLOR"); //Allow the option to draw out the scene color too to lower draw calls
-	using FPermutationDomain = TShaderPermutationDomain<FUseMSAADimension, FForceDebugDrawColorOutput>;
+	class FUseMetalMSAAHDRDecodeDim : SHADER_PERMUTATION_BOOL("METAL_MSAA_HDR_DECODE");
+	using FPermutationDomain = TShaderPermutationDomain<FUseMSAADimension, FForceDebugDrawColorOutput, FUseMetalMSAAHDRDecodeDim>;
 
 	BEGIN_SHADER_PARAMETER_STRUCT(FParameters, )
 		SHADER_PARAMETER_STRUCT_REF(FViewUniformShaderParameters, View)
@@ -122,15 +123,28 @@ public:
 	{
 		const FPermutationDomain PermutationVector(Parameters.PermutationId);
 		const bool bUseMSAA = PermutationVector.Get<FUseMSAADimension>();
+		const bool bUseMetalMSAAHDRDecodeDim = PermutationVector.Get<FUseMetalMSAAHDRDecodeDim>();
 
-		// Only SM5+ platforms supports MSAA.
-		if (!IsFeatureLevelSupported(Parameters.Platform, ERHIFeatureLevel::SM5) && bUseMSAA)
+		if (bUseMSAA)
 		{
-			return false;
+			// Only SM5+ and Metal ES3.1 platforms support MSAA.
+			if (IsFeatureLevelSupported(Parameters.Platform, ERHIFeatureLevel::SM5))
+			{
+				return true;
+			}
+			else if (IsMetalMobilePlatform(Parameters.Platform) && bUseMetalMSAAHDRDecodeDim)
+			{
+				//Only compile with decode dim if using metal platform on 3.1 + color output
+				return PermutationVector.Get<FForceDebugDrawColorOutput>();
+			}
+		}
+		else
+		{
+			//Never compile Metal decode dim permutations for non-MSAA passes
+			return !bUseMetalMSAAHDRDecodeDim;
 		}
 
-		// This pass can be used with any AA method so return true for all other circumstances
-		return true;
+		return false;
 	}
 };
 
@@ -245,7 +259,8 @@ void PopulateDepthPass(FRDGBuilder& GraphBuilder,
 	FRDGTextureRef OutPopDepth,
 	const FVector2f& SceneDepthJitter,
 	uint32 NumMSAASamples,
-	bool bForceDrawColor)
+	bool bForceDrawColor,
+	bool bUseMetalPlatformHDRDecode)
 {
 	FRHISamplerState* PointClampSampler = TStaticSamplerState<SF_Point, AM_Clamp, AM_Clamp, AM_Clamp>::GetRHI();
 
@@ -271,6 +286,7 @@ void PopulateDepthPass(FRDGBuilder& GraphBuilder,
 	FPopulateCompositeDepthPS::FPermutationDomain PermutationVector;
 	PermutationVector.Set<FPopulateCompositeDepthPS::FUseMSAADimension>(NumMSAASamples > 1);
 	PermutationVector.Set< FPopulateCompositeDepthPS::FForceDebugDrawColorOutput>(bForceDrawColor);
+	PermutationVector.Set< FPopulateCompositeDepthPS::FUseMetalMSAAHDRDecodeDim>(bForceDrawColor && bUseMetalPlatformHDRDecode);
 	TShaderMapRef<FPopulateCompositeDepthPS> PixelShader(View.ShaderMap, PermutationVector);
 
 	FPixelShaderUtils::AddFullscreenPass(
diff --git a/Engine/Source/Runtime/Renderer/Private/PostProcess/PostProcessCompositePrimitivesCommon.h b/Engine/Source/Runtime/Renderer/Private/PostProcess/PostProcessCompositePrimitivesCommon.h
--- a/Engine/Source/Runtime/Renderer/Private/PostProcess/PostProcessCompositePrimitivesCommon.h
+++ b/Engine/Source/Runtime/Renderer/Private/PostProcess/PostProcessCompositePrimitivesCommon.h
@@ -63,6 +63,8 @@ struct FCompositePrimitiveInputs
 
 	// [Required] The scene depth to composite with editor primitives.
 	FScreenPassTexture SceneDepth;
+
+	bool bUseMetalMSAAHDRDecode = false;
 };
 
 // Constructs a new view suitable for rendering debug primitives.
@@ -86,6 +88,7 @@ void PopulateDepthPass(FRDGBuilder& GraphBuilder,
 	FRDGTextureRef OutPopDepth, 
 	const FVector2f& SceneDepthJitter,
 	uint32 NumMSAASamples,
-	bool bForceDrawColor = false);
+	bool bForceDrawColor = false,
+	bool bUseMetalPlatformHDRDecode = false);
 
 #endif //#if UE_ENABLE_DEBUG_DRAWING
\ No newline at end of file
diff --git a/Engine/Source/Runtime/Renderer/Private/PostProcess/PostProcessing.cpp b/Engine/Source/Runtime/Renderer/Private/PostProcess/PostProcessing.cpp
--- a/Engine/Source/Runtime/Renderer/Private/PostProcess/PostProcessing.cpp
+++ b/Engine/Source/Runtime/Renderer/Private/PostProcess/PostProcessing.cpp
@@ -2674,6 +2674,7 @@ void AddMobilePostProcessingPasses(FRDGBuilder& GraphBuilder, FScene* Scene, con
 		PassSequence.AcceptOverrideIfLastPass(EPass::DebugPrimitive, PassInputs.OverrideOutput);
 		PassInputs.SceneColor = SceneColor;
 		PassInputs.SceneDepth = SceneDepth;
+		PassInputs.bUseMetalMSAAHDRDecode = bMetalMSAAHDRDecode;
 
 		SceneColor = AddDebugPrimitivePass(GraphBuilder, View, PassInputs);
 	}
