Removed render thread as the known thread for launching mesh setup tasks. Fixes check in debug mode.

diff --git a/Engine/Source/Runtime/Renderer/Private/MeshDrawCommands.cpp b/Engine/Source/Runtime/Renderer/Private/MeshDrawCommands.cpp
--- a/Engine/Source/Runtime/Renderer/Private/MeshDrawCommands.cpp
+++ b/Engine/Source/Runtime/Renderer/Private/MeshDrawCommands.cpp
@@ -1422,13 +1422,13 @@ void FParallelMeshDrawCommandPass::DispatchPassSetup(
 		{
 			if (IsOnDemandShaderCreationEnabled())
 			{
-				TaskEventRef = TGraphTask<FMeshDrawCommandPassSetupTask>::CreateTask(nullptr, ENamedThreads::GetRenderThread()).ConstructAndDispatchWhenReady(TaskContext);
+				TaskEventRef = TGraphTask<FMeshDrawCommandPassSetupTask>::CreateTask().ConstructAndDispatchWhenReady(TaskContext);
 			}
 			else
 			{
 				FGraphEventArray DependentGraphEvents;
-				DependentGraphEvents.Add(TGraphTask<FMeshDrawCommandPassSetupTask>::CreateTask(nullptr, ENamedThreads::GetRenderThread()).ConstructAndDispatchWhenReady(TaskContext));
-				TaskEventRef = TGraphTask<FMeshDrawCommandInitResourcesTask>::CreateTask(&DependentGraphEvents, ENamedThreads::GetRenderThread()).ConstructAndDispatchWhenReady(TaskContext);
+				DependentGraphEvents.Add(TGraphTask<FMeshDrawCommandPassSetupTask>::CreateTask().ConstructAndDispatchWhenReady(TaskContext));
+				TaskEventRef = TGraphTask<FMeshDrawCommandInitResourcesTask>::CreateTask(&DependentGraphEvents).ConstructAndDispatchWhenReady(TaskContext);
 			}
 		}
 		else
