Fix mismatching VS/PS interpolation for full screen passes in hair code.

diff --git a/Engine/Source/Runtime/Renderer/Private/HairStrands/HairStrandsDebug.cpp b/Engine/Source/Runtime/Renderer/Private/HairStrands/HairStrandsDebug.cpp
--- a/Engine/Source/Runtime/Renderer/Private/HairStrands/HairStrandsDebug.cpp
+++ b/Engine/Source/Runtime/Renderer/Private/HairStrands/HairStrandsDebug.cpp
@@ -507,7 +507,7 @@ static void AddDebugHairPass(
 	Parameters->DebugMode = InternalDebugMode;
 	Parameters->SampleIndex = 0;
 	Parameters->RenderTargets[0] = FRenderTargetBinding(OutTarget, ERenderTargetLoadAction::ELoad, 0);
-	TShaderMapRef<FPostProcessVS> VertexShader(View->ShaderMap);
+	TShaderMapRef<FScreenVS> VertexShader(View->ShaderMap);
 
 	TShaderMapRef<FHairDebugPS> PixelShader(View->ShaderMap);
 
@@ -617,7 +617,7 @@ static void AddDebugDeepShadowTexturePass(
 	Parameters->LinearSampler = TStaticSamplerState<SF_Bilinear, AM_Clamp, AM_Clamp, AM_Clamp>::GetRHI();
 	Parameters->HairViewRect = FIntVector4(HairViewRect.Min.X, HairViewRect.Min.Y, HairViewRect.Width(), HairViewRect.Height());
 	Parameters->RenderTargets[0] = FRenderTargetBinding(OutTarget, ERenderTargetLoadAction::ELoad, 0);
-	TShaderMapRef<FPostProcessVS> VertexShader(View->ShaderMap);
+	TShaderMapRef<FScreenVS> VertexShader(View->ShaderMap);
 	FDeepShadowVisualizePS::FPermutationDomain PermutationVector;
 	PermutationVector.Set<FDeepShadowVisualizePS::FOutputType>(ShadowData ? 0 : 1);
 	TShaderMapRef<FDeepShadowVisualizePS> PixelShader(View->ShaderMap, PermutationVector);
@@ -931,7 +931,7 @@ static void AddPlotBSDFPass(
 	Parameters->RenderTargets[0] = FRenderTargetBinding(OutputTexture, ERenderTargetLoadAction::ELoad);
 
 	const FIntPoint OutputResolution = SceneTextures.SceneDepthTexture->Desc.Extent;
-	TShaderMapRef<FPostProcessVS> VertexShader(View.ShaderMap);
+	TShaderMapRef<FScreenVS> VertexShader(View.ShaderMap);
 	TShaderMapRef<FHairStrandsPlotBSDFPS> PixelShader(View.ShaderMap);
 	const FGlobalShaderMap* GlobalShaderMap = View.ShaderMap;
 	const FIntRect Viewport = View.ViewRect;
@@ -1021,7 +1021,7 @@ static void AddPlotSamplePass(
 	Parameters->RenderTargets[0] = FRenderTargetBinding(OutputTexture, ERenderTargetLoadAction::ELoad);
 
 	const FIntPoint OutputResolution = SceneTextures.SceneDepthTexture->Desc.Extent;
-	TShaderMapRef<FPostProcessVS> VertexShader(View.ShaderMap);
+	TShaderMapRef<FScreenVS> VertexShader(View.ShaderMap);
 	TShaderMapRef<FHairStrandsPlotSamplePS> PixelShader(View.ShaderMap);
 	const FGlobalShaderMap* GlobalShaderMap = View.ShaderMap;
 	const FIntRect Viewport = View.ViewRect;
diff --git a/Engine/Source/Runtime/Renderer/Private/HairStrands/HairStrandsTransmittance.cpp b/Engine/Source/Runtime/Renderer/Private/HairStrands/HairStrandsTransmittance.cpp
--- a/Engine/Source/Runtime/Renderer/Private/HairStrands/HairStrandsTransmittance.cpp
+++ b/Engine/Source/Runtime/Renderer/Private/HairStrands/HairStrandsTransmittance.cpp
@@ -506,7 +506,7 @@ static void AddHairStrandsVoxelShadowMaskPass(
 	FHairStrandsVoxelShadowMaskPS::FPermutationDomain PermutationVector;
 	PermutationVector.Set<FHairStrandsVoxelShadowMaskPS::FOnePass>(bOnePass);
 
-	TShaderMapRef<FPostProcessVS> VertexShader(View.ShaderMap);
+	TShaderMapRef<FScreenVS> VertexShader(View.ShaderMap);
 	TShaderMapRef<FHairStrandsVoxelShadowMaskPS> PixelShader(View.ShaderMap, PermutationVector);
 	const FGlobalShaderMap* GlobalShaderMap = View.ShaderMap;
 	const FIntRect Viewport = View.ViewRect;
@@ -701,7 +701,7 @@ static void AddHairStrandsDeepShadowMaskPass(
 	PermutationVector.Set<FHairStrandsDeepShadowMaskPS::FKernelType>(FMath::Clamp(GStrandHairShadowMaskKernelType, 0, 4));
 
 	const FIntPoint OutputResolution = SceneDepthTexture->Desc.Extent;
-	TShaderMapRef<FPostProcessVS> VertexShader(View.ShaderMap);
+	TShaderMapRef<FScreenVS> VertexShader(View.ShaderMap);
 	TShaderMapRef<FHairStrandsDeepShadowMaskPS> PixelShader(View.ShaderMap, PermutationVector);
 	const FGlobalShaderMap* GlobalShaderMap = View.ShaderMap;
 	const FIntRect Viewport = View.ViewRect;
diff --git a/Engine/Source/Runtime/Renderer/Private/HairStrands/HairStrandsVisibility.cpp b/Engine/Source/Runtime/Renderer/Private/HairStrands/HairStrandsVisibility.cpp
--- a/Engine/Source/Runtime/Renderer/Private/HairStrands/HairStrandsVisibility.cpp
+++ b/Engine/Source/Runtime/Renderer/Private/HairStrands/HairStrandsVisibility.cpp
@@ -2415,7 +2415,7 @@ static void AddHairViewTransmittanceDepthPass(
 	FHairViewTransmittanceDepthPS::FPermutationDomain PermutationVector;
 	PermutationVector.Set<FHairViewTransmittanceDepthPS::FOutputFormat>(HairCountTexture->Desc.Format == PF_G32R32F ? 1 : 0);
 
-	TShaderMapRef<FPostProcessVS> VertexShader(View.ShaderMap);
+	TShaderMapRef<FScreenVS> VertexShader(View.ShaderMap);
 	TShaderMapRef<FHairViewTransmittanceDepthPS> PixelShader(View.ShaderMap, PermutationVector);
 	const FGlobalShaderMap* GlobalShaderMap = View.ShaderMap;
 	const FIntRect Viewport = View.ViewRect;
