Fix GPU crash when using OIT pixel sorting.

diff --git a/Engine/Shaders/Private/OITCommon.ush b/Engine/Shaders/Private/OITCommon.ush
--- a/Engine/Shaders/Private/OITCommon.ush
+++ b/Engine/Shaders/Private/OITCommon.ush
@@ -252,7 +252,9 @@ void AddOITSample_MLAB(uint2 InPixelCoord, float3 InColor, float3 InTrans, float
 		{
 			if (all(LastSample.Trans < TranslucentBasePass.OIT.TransmittanceThreshold))
 			{
+				// The return is import here as it avoid GPU hang. Discard is not enough here, probably due to ROV/UAV output
 				discard;
+				return;
 			}
 		}
 	}
diff --git a/Engine/Source/Runtime/Renderer/Private/OIT/OIT.cpp b/Engine/Source/Runtime/Renderer/Private/OIT/OIT.cpp
--- a/Engine/Source/Runtime/Renderer/Private/OIT/OIT.cpp
+++ b/Engine/Source/Runtime/Renderer/Private/OIT/OIT.cpp
@@ -934,7 +934,7 @@ namespace OIT
 			Out.SampleColorTexture = GraphBuilder.CreateTexture(FRDGTextureDesc::Create2D(FIntPoint(1, 1), PF_R32_UINT, FClearValueBinding::None, TexCreate_ShaderResource | TexCreate_UAV), TEXT("OIT.SampleColor"));
 			Out.SampleTransTexture = GraphBuilder.CreateTexture(FRDGTextureDesc::Create2D(FIntPoint(1, 1), PF_R32_UINT, FClearValueBinding::None, TexCreate_ShaderResource | TexCreate_UAV), TEXT("OIT.SampleTrans"));
 			Out.SampleDepthTexture = GraphBuilder.CreateTexture(FRDGTextureDesc::Create2D(FIntPoint(1, 1), PF_R32_FLOAT, FClearValueBinding::Black, TexCreate_ShaderResource | TexCreate_UAV), TEXT("OIT.SampleDepth"));
-			Out.SampleCountTexture = GraphBuilder.CreateTexture(FRDGTextureDesc::Create2D(FIntPoint(1, 1), PF_R32_UINT, FClearValueBinding::None, TexCreate_ShaderResource | TexCreate_UAV), TEXT("OIT.SampleCount"));
+			Out.SampleCountTexture = GraphBuilder.CreateTexture(FRDGTextureDesc::Create2D(FIntPoint(1, 1), PF_R32_UINT, FClearValueBinding::Black, TexCreate_ShaderResource | TexCreate_UAV), TEXT("OIT.SampleCount"));
 			return Out;
 		}
 
diff --git a/Engine/Source/Runtime/Renderer/Private/OIT/OITParameters.h b/Engine/Source/Runtime/Renderer/Private/OIT/OITParameters.h
--- a/Engine/Source/Runtime/Renderer/Private/OIT/OITParameters.h
+++ b/Engine/Source/Runtime/Renderer/Private/OIT/OITParameters.h
@@ -13,7 +13,7 @@ BEGIN_SHADER_PARAMETER_STRUCT(FOITBasePassUniformParameters, )
 	SHADER_PARAMETER(uint32, OITMethod)
 	SHADER_PARAMETER(uint32, MaxSideSamplePerPixel)
 	SHADER_PARAMETER(uint32, MaxSamplePerPixel)
-	SHADER_PARAMETER(uint32, TransmittanceThreshold)
+	SHADER_PARAMETER(float, TransmittanceThreshold)
 	SHADER_PARAMETER_RDG_TEXTURE_UAV(RasterizerOrderedTexture2D<uint>, OutOITSampleCount)
 	SHADER_PARAMETER_RDG_TEXTURE_UAV(RWTexture2D<uint>,  OutOITSampleColor)
 	SHADER_PARAMETER_RDG_TEXTURE_UAV(RWTexture2D<uint>,  OutOITSampleTrans)
