Fixed memory stomp when shadows allocate views from instance culling manager while main views do as well.
 - Solution is to move view allocation into a task that cannot start until shadow GDME is triggered.

diff --git a/Engine/Source/Runtime/Renderer/Private/ShadowSetup.cpp b/Engine/Source/Runtime/Renderer/Private/ShadowSetup.cpp
--- a/Engine/Source/Runtime/Renderer/Private/ShadowSetup.cpp
+++ b/Engine/Source/Runtime/Renderer/Private/ShadowSetup.cpp
@@ -4538,14 +4538,18 @@ void FSceneRenderer::GatherShadowDynamicMeshElements(FDynamicShadowsTaskData& Ta
 		? UE::Tasks::EExtendedTaskPriority::None
 		: UE::Tasks::EExtendedTaskPriority::Inline;
 
-	UE::Tasks::FTask MeshCollectorsTask;
-
 	int32 NumShadowViews = 0;
 	for (FProjectedShadowInfo* ProjectedShadowInfo : TaskData.ShadowsToGather)
 	{
 		NumShadowViews += ProjectedShadowInfo->bOnePassPointLightShadow ? 6 : 1;
 	}
-	TaskData.InstanceCullingManager.AllocateViews(NumShadowViews);
+
+	UE::Tasks::FTask MeshCollectorsTask = UE::Tasks::Launch(UE_SOURCE_LOCATION, [&InstanceCullingManager = TaskData.InstanceCullingManager, NumShadowViews]
+	{
+		// Wait to allocate views until after the task event has triggered.
+		InstanceCullingManager.AllocateViews(NumShadowViews);
+
+	}, TaskData.BeginGatherDynamicMeshElementsTask, TaskPriority, ExtendedTaskPriority);
 
 	if (TaskData.bMultithreadedGDME)
 	{
@@ -4598,7 +4602,7 @@ void FSceneRenderer::GatherShadowDynamicMeshElements(FDynamicShadowsTaskData& Ta
 					}
 				}
 
-			}, TaskData.BeginGatherDynamicMeshElementsTask, TaskPriority));
+			}, MeshCollectorsTask, TaskPriority));
 		}
 
 		MeshCollectorsTaskEvent.Trigger();
