Add mobile lens flare pass

diff --git a/Engine/Config/Android/AndroidScalability.ini b/Engine/Config/Android/AndroidScalability.ini
--- a/Engine/Config/Android/AndroidScalability.ini
+++ b/Engine/Config/Android/AndroidScalability.ini
@@ -93,6 +93,7 @@ r.LightShaftQuality=0
 r.EyeAdaptationQuality=1
 r.Mobile.AmbientOcclusionQuality=0
 r.Mobile.PixelProjectedReflectionQuality=0
+r.LensFlareQuality=0
 
 [PostProcessQuality@1]
 r.Tonemapper.Quality=2
@@ -102,6 +103,7 @@ r.LightShaftQuality=1
 r.EyeAdaptationQuality=1
 r.Mobile.AmbientOcclusionQuality=0
 r.Mobile.PixelProjectedReflectionQuality=0
+r.LensFlareQuality=1
 
 [PostProcessQuality@2]
 r.Tonemapper.Quality=3
@@ -111,6 +113,7 @@ r.LightShaftQuality=1
 r.EyeAdaptationQuality=1
 r.Mobile.AmbientOcclusionQuality=1
 r.Mobile.PixelProjectedReflectionQuality=1
+r.LensFlareQuality=2
 
 [PostProcessQuality@3]
 r.Tonemapper.Quality=3
@@ -120,6 +123,7 @@ r.LightShaftQuality=1
 r.EyeAdaptationQuality=1
 r.Mobile.AmbientOcclusionQuality=1
 r.Mobile.PixelProjectedReflectionQuality=1
+r.LensFlareQuality=2
 
 [AntiAliasingQuality@0]
 r.FXAA.Quality=0
diff --git a/Engine/Config/IOS/IOSScalability.ini b/Engine/Config/IOS/IOSScalability.ini
--- a/Engine/Config/IOS/IOSScalability.ini
+++ b/Engine/Config/IOS/IOSScalability.ini
@@ -93,6 +93,7 @@ r.LightShaftQuality=0
 r.EyeAdaptationQuality=1
 r.Mobile.AmbientOcclusionQuality=0
 r.Mobile.PixelProjectedReflectionQuality=0
+r.LensFlareQuality=0
 
 [PostProcessQuality@1]
 r.Tonemapper.Quality=2
@@ -102,6 +103,7 @@ r.LightShaftQuality=0
 r.EyeAdaptationQuality=1
 r.Mobile.AmbientOcclusionQuality=0
 r.Mobile.PixelProjectedReflectionQuality=0
+r.LensFlareQuality=1
 
 [PostProcessQuality@2]
 r.Tonemapper.Quality=3
@@ -111,6 +113,7 @@ r.LightShaftQuality=0
 r.EyeAdaptationQuality=1
 r.Mobile.AmbientOcclusionQuality=0
 r.Mobile.PixelProjectedReflectionQuality=0
+r.LensFlareQuality=2
 
 [PostProcessQuality@3]
 r.Tonemapper.Quality=3
@@ -120,6 +123,7 @@ r.LightShaftQuality=1
 r.EyeAdaptationQuality=1
 r.Mobile.AmbientOcclusionQuality=1
 r.Mobile.PixelProjectedReflectionQuality=1
+r.LensFlareQuality=2
 
 [AntiAliasingQuality@0]
 r.FXAA.Quality=0
diff --git a/Engine/Source/Runtime/Renderer/Private/PostProcess/PostProcessLensFlares.cpp b/Engine/Source/Runtime/Renderer/Private/PostProcess/PostProcessLensFlares.cpp
--- a/Engine/Source/Runtime/Renderer/Private/PostProcess/PostProcessLensFlares.cpp
+++ b/Engine/Source/Runtime/Renderer/Private/PostProcess/PostProcessLensFlares.cpp
@@ -36,7 +36,7 @@ class FLensFlareShader : public FGlobalShader
 public:
 	static bool ShouldCompilePermutation(const FGlobalShaderPermutationParameters& Parameters)
 	{
-		return IsFeatureLevelSupported(Parameters.Platform, ERHIFeatureLevel::SM5);
+		return true;
 	}
 
 	static void ModifyCompilationEnvironment(const FGlobalShaderPermutationParameters& Parameters, FShaderCompilerEnvironment& OutEnvironment)
@@ -100,7 +100,7 @@ public:
 
 	static bool ShouldCompilePermutation(const FGlobalShaderPermutationParameters& Parameters)
 	{
-		return IsFeatureLevelSupported(Parameters.Platform, ERHIFeatureLevel::SM5);
+		return true;
 	}
 };
 
@@ -363,10 +363,10 @@ FScreenPassTexture AddLensFlaresPass(
 	FRDGBuilder& GraphBuilder,
 	const FViewInfo& View,
 	FScreenPassTexture Bloom,
-	const FSceneDownsampleChain& SceneDownsampleChain)
+	FScreenPassTexture QualitySceneDownsample,
+	FScreenPassTexture DefaultSceneDownsample)
 {
 	ensure(IsLensFlaresEnabled(View));
-	const ELensFlareQuality LensFlareQuality = GetLensFlareQuality();
 
 	const FPostProcessSettings& Settings = View.FinalPostProcessSettings;
 
@@ -394,12 +394,9 @@ FScreenPassTexture AddLensFlaresPass(
 		}
 	}
 
-	// The quality level controls which downsample stage we use as the flare input texture.
-	const uint32 LensFlareDownsampleStageIndex = static_cast<uint32>(ELensFlareQuality::MAX) - static_cast<uint32>(LensFlareQuality) - 1;
-
 	FLensFlareInputs LensFlareInputs;
 	LensFlareInputs.Bloom = Bloom;
-	LensFlareInputs.Flare = SceneDownsampleChain.GetTexture(LensFlareDownsampleStageIndex);
+	LensFlareInputs.Flare = QualitySceneDownsample;
 	LensFlareInputs.BokehShapeTexture = BokehTextureRHI;
 	LensFlareInputs.TintColorsPerFlare = Settings.LensFlareTints;
 	LensFlareInputs.TintColor = Settings.LensFlareTint;
@@ -411,9 +408,9 @@ FScreenPassTexture AddLensFlaresPass(
 	// composition. The pass needs a primary input in order to access the image descriptor and viewport for output.
 	if (!Bloom.IsValid())
 	{
-		LensFlareInputs.Bloom = SceneDownsampleChain.GetFirstTexture();
+		LensFlareInputs.Bloom = DefaultSceneDownsample;
 		LensFlareInputs.bCompositeWithBloom = false;
 	}
 
 	return AddLensFlaresPass(GraphBuilder, View, LensFlareInputs);
-}
+}
\ No newline at end of file
diff --git a/Engine/Source/Runtime/Renderer/Private/PostProcess/PostProcessLensFlares.h b/Engine/Source/Runtime/Renderer/Private/PostProcess/PostProcessLensFlares.h
--- a/Engine/Source/Runtime/Renderer/Private/PostProcess/PostProcessLensFlares.h
+++ b/Engine/Source/Runtime/Renderer/Private/PostProcess/PostProcessLensFlares.h
@@ -65,4 +65,5 @@ FScreenPassTexture AddLensFlaresPass(
 	FRDGBuilder& GraphBuilder,
 	const FViewInfo& View,
 	FScreenPassTexture Bloom,
-	const FSceneDownsampleChain& SceneDownsampleChain);
\ No newline at end of file
+	FScreenPassTexture QualitySceneDownsample,
+	FScreenPassTexture DefaultSceneDownsample);
\ No newline at end of file
diff --git a/Engine/Source/Runtime/Renderer/Private/PostProcess/PostProcessing.cpp b/Engine/Source/Runtime/Renderer/Private/PostProcess/PostProcessing.cpp
--- a/Engine/Source/Runtime/Renderer/Private/PostProcess/PostProcessing.cpp
+++ b/Engine/Source/Runtime/Renderer/Private/PostProcess/PostProcessing.cpp
@@ -1125,7 +1125,9 @@ void AddPostProcessingPasses(
 
 			if (bLensFlareEnabled)
 			{
-				Bloom = AddLensFlaresPass(GraphBuilder, View, Bloom, *LensFlareSceneDownsampleChain);
+				const ELensFlareQuality LensFlareQuality = GetLensFlareQuality();
+				const uint32 LensFlareDownsampleStageIndex = static_cast<uint32>(ELensFlareQuality::MAX) - static_cast<uint32>(LensFlareQuality) - 1;
+				Bloom = AddLensFlaresPass(GraphBuilder, View, Bloom, LensFlareSceneDownsampleChain->GetTexture(LensFlareDownsampleStageIndex), LensFlareSceneDownsampleChain->GetFirstTexture());
 			}
 		}
 
@@ -2416,6 +2418,13 @@ void AddMobilePostProcessingPasses(FRDGBuilder& GraphBuilder, FScene* Scene, con
 
 				BloomUpOutputs = AddBloomUpPass(PostProcessDownsample_Bloom[0], BloomUpOutputs, BloomUpScale, TintA, TintB);
 			}
+
+			if (IsLensFlaresEnabled(View))
+			{
+				const ELensFlareQuality LensFlareQuality = GetLensFlareQuality();
+				const uint32 LensFlareDownsampleStageIndex = static_cast<uint32>(ELensFlareQuality::MAX) - static_cast<uint32>(LensFlareQuality) - 1;
+				BloomUpOutputs = AddLensFlaresPass(GraphBuilder, View, BloomUpOutputs, PostProcessDownsample_Bloom[LensFlareDownsampleStageIndex], PostProcessDownsample_Bloom[0]);
+			}
 		}
 
 		if (PassSequence.IsEnabled(EPass::EyeAdaptation))
