Work around shadow HZB issue by starting a new batch of deferred culling after shadow rendering (preventing invasion of HZB-using passes from later) the change can be turned off using  r.Shadow.ShouldBeginDeferredCullingAfterShadowRendering for hotfixability.

diff --git a/Engine/Source/Runtime/Renderer/Private/ShadowDepthRendering.cpp b/Engine/Source/Runtime/Renderer/Private/ShadowDepthRendering.cpp
--- a/Engine/Source/Runtime/Renderer/Private/ShadowDepthRendering.cpp
+++ b/Engine/Source/Runtime/Renderer/Private/ShadowDepthRendering.cpp
@@ -108,6 +108,12 @@ IMPLEMENT_SHADER_TYPE(template<>, TScreenVSForGS<false>, TEXT("/Engine/Private/S
 IMPLEMENT_SHADER_TYPE(template<>, TScreenVSForGS<true>, TEXT("/Engine/Private/ScreenVertexShader.usf"), TEXT("MainForGS"), SF_Vertex);
 
 
+static TAutoConsoleVariable<int32> CVarShouldBeginDeferredCullingAfterShadowRendering(
+	TEXT("r.Shadow.ShouldBeginDeferredCullingAfterShadowRendering"),
+	1,
+	TEXT("Temporary: If turned on, a new deferred scope of instance culling will be started directly after shadow depth rendering (iff instance occlusion culling is enabled) to work around a bug."),
+	ECVF_RenderThreadSafe);
+
 static TAutoConsoleVariable<int32> CVarShadowForceSerialSingleRenderPass(
 	TEXT("r.Shadow.ForceSerialSingleRenderPass"),
 	0,
@@ -1963,6 +1969,13 @@ void FSceneRenderer::RenderShadowDepthMaps(FRDGBuilder& GraphBuilder, FDynamicSh
 		}
 
 	}, &GPersistentShadowsPipe);
+	
+	// Begin another deferred batch here to avoid any subsequent passes enabling HZB for all (an unfortunate side effect).
+	if (FInstanceCullingContext::IsOcclusionCullingEnabled() && CVarShouldBeginDeferredCullingAfterShadowRendering.GetValueOnRenderThread() != 0)
+	{
+		InstanceCullingManager.BeginDeferredCulling(GraphBuilder, Scene->GPUScene);
+	}
+
 
 	bShadowDepthRenderCompleted = true;
 }
