Fix Tiled SSS reflection artifact on some platforms where:
1) Specular is small but contributes to reflection. Set the MinSpecular to 0.0f.
2) Visible seam of tiles. Make tile classification logic match the SSR max roughness threshold to removing tiling artifacts.
3) clearcoat material will not be culled away.

diff --git a/Engine/Shaders/Private/DefaultSSRTiles.usf b/Engine/Shaders/Private/DefaultSSRTiles.usf
--- a/Engine/Shaders/Private/DefaultSSRTiles.usf
+++ b/Engine/Shaders/Private/DefaultSSRTiles.usf
@@ -70,6 +70,15 @@ bool DoesPixelContainSSR(uint2 PixelPos)
 	// Ignore two sided foliage 
 	bContainValidPixels = bContainValidPixels && (bEnableTwoSidedFoliage ? true : (GBuffer.ShadingModelID != SHADINGMODELID_TWOSIDED_FOLIAGE));
 
+	// Cull special case where specular=0, and roughness=0
+	if (GBuffer.Roughness == 0 && GBuffer.Specular == 0)
+	{
+		bContainValidPixels = false;
+	}
+
+	// Clearcoat: if the material is clear coat it contains valid pixel.
+	bContainValidPixels = bContainValidPixels || (GBuffer.ShadingModelID == SHADINGMODELID_CLEAR_COAT);
+
 	return bContainValidPixels;
 #endif
 #endif // USE_SSR_PRE_PASS_STENCIL
diff --git a/Engine/Source/Runtime/Renderer/Private/ScreenSpaceReflectionTiles.cpp b/Engine/Source/Runtime/Renderer/Private/ScreenSpaceReflectionTiles.cpp
--- a/Engine/Source/Runtime/Renderer/Private/ScreenSpaceReflectionTiles.cpp
+++ b/Engine/Source/Runtime/Renderer/Private/ScreenSpaceReflectionTiles.cpp
@@ -9,13 +9,14 @@ static TAutoConsoleVariable<int32> CVarSSRTiledComposite(
 	TEXT("Enable tiled optimization of the screen space reflection."),
 	ECVF_RenderThreadSafe | ECVF_Scalability);
 
-static TAutoConsoleVariable<float> CVarSSRTiledCompositeMaxRoughness(
-	TEXT("r.SSR.TiledComposite.MaxRoughness"), 0.3f,
-	TEXT("Ignore pixels with roughness larger than this value."),
+static TAutoConsoleVariable<float> CVarSSRTiledCompositeOverrideMaxRoughness(
+	TEXT("r.SSR.TiledComposite.OverrideMaxRoughness"), -1.0f,
+	TEXT("Ignore pixels with roughness larger than this value.")
+	TEXT("<0: use derived value from ScreenSpaceReflectionMaxRoughness of FinalPostProcessSettings."),
 	ECVF_RenderThreadSafe | ECVF_Scalability);
 
 static TAutoConsoleVariable<float> CVarSSRTiledCompositeMinSpecular(
-	TEXT("r.SSR.TiledComposite.MinSpecular"), 0.01f,
+	TEXT("r.SSR.TiledComposite.MinSpecular"), 0.0f,
 	TEXT("Ignore pixels with very small specular contribution in case max roughness cannot filter them out"),
 	ECVF_RenderThreadSafe | ECVF_Scalability);
 
@@ -30,6 +31,17 @@ static TAutoConsoleVariable<bool> CVarSSRTiledCompositeVisualize(
 	ECVF_RenderThreadSafe | ECVF_Scalability);
 
 
+static float GetScreenSpaceReflectionMaxRoughnessScale(const FViewInfo& View)
+{
+	float MaxRoughness = CVarSSRTiledCompositeOverrideMaxRoughness.GetValueOnRenderThread();
+
+	if (MaxRoughness < 0)
+	{
+		MaxRoughness = FMath::Clamp(View.FinalPostProcessSettings.ScreenSpaceReflectionMaxRoughness, 0.01f, 1.0f);
+	}
+	return MaxRoughness;
+}
+
 bool UseSSRIndirectDraw(EShaderPlatform ShaderPlatform)
 {
 	return IsFeatureLevelSupported(ShaderPlatform, ERHIFeatureLevel::SM5)
@@ -167,7 +179,7 @@ FScreenSpaceReflectionTileClassification ClassifySSRTiles(FRDGBuilder& GraphBuil
 			PassParameters->TiledViewRes = Result.TiledViewRes;
 			PassParameters->SSRDepthStencilTexture = DepthPrepassTexture ? GraphBuilder.CreateSRV(FRDGTextureSRVDesc::CreateWithPixelFormat(DepthPrepassTexture, PF_X24_G8)) : nullptr;
 			PassParameters->TileMaskBufferOut = TileMaskBufferUAV;
-			PassParameters->MaxRoughness = FMath::Clamp(CVarSSRTiledCompositeMaxRoughness.GetValueOnRenderThread(), -0.001f, 1.001f);
+			PassParameters->MaxRoughness = GetScreenSpaceReflectionMaxRoughnessScale(View);
 			PassParameters->MinSpecular = FMath::Clamp(CVarSSRTiledCompositeMinSpecular.GetValueOnRenderThread(), -0.001f, 1.001f);
 			PassParameters->bEnableTwoSidedFoliage = CVarSSRTiledCompositeTwoSidedFoliage.GetValueOnRenderThread() != 0;
 
