Fix mismatching texture2d/texture2darray binding.

diff --git a/Engine/Shaders/Private/DiffuseIndirectComposite.usf b/Engine/Shaders/Private/DiffuseIndirectComposite.usf
--- a/Engine/Shaders/Private/DiffuseIndirectComposite.usf
+++ b/Engine/Shaders/Private/DiffuseIndirectComposite.usf
@@ -386,16 +386,19 @@ void MainPS(
 				DiffuseIndirectLighting 		= DiffuseIndirect_Lumen_0.SampleLevel(GlobalPointClampedSampler, float3(SceneBufferUV, 0), 0).rgb;
 				RoughSpecularIndirectLighting 	= DiffuseIndirect_Lumen_2.SampleLevel(GlobalPointClampedSampler, float3(SceneBufferUV, 0), 0).rgb;
 				
-				#if SSR_TILED_COMPOSITION
-				bool bSSRTileUsed = IsSSRTileUsed(PixelPos, SSRTiledViewRes, SSRTileMaskBuffer);
-
-				if (bSSRTileUsed && bLumenReflectionInputIsSSR)
+				if (bLumenReflectionInputIsSSR)
 				{
-					SpecularIndirectLighting	= DiffuseIndirect_Lumen_3.SampleLevel(GlobalPointClampedSampler, float3(SceneBufferUV, 0), 0).rgba;
+					#if SSR_TILED_COMPOSITION
+					if (IsSSRTileUsed(PixelPos, SSRTiledViewRes, SSRTileMaskBuffer))
+					#endif
+					{
+						SpecularIndirectLighting = DiffuseIndirect_Textures_3.SampleLevel(GlobalPointClampedSampler, SceneBufferUV, 0).rgba;
+					}
+				}
+				else
+				{
+					SpecularIndirectLighting = DiffuseIndirect_Lumen_3.SampleLevel(GlobalPointClampedSampler, float3(SceneBufferUV, 0), 0).rgba;
 				}
-				#else
-				SpecularIndirectLighting 		= DiffuseIndirect_Lumen_3.SampleLevel(GlobalPointClampedSampler, float3(SceneBufferUV, 0), 0).rgba;
-				#endif
 			#else
 			{
 				// Sample the output of the denoiser.
diff --git a/Engine/Source/Runtime/Renderer/Private/IndirectLightRendering.cpp b/Engine/Source/Runtime/Renderer/Private/IndirectLightRendering.cpp
--- a/Engine/Source/Runtime/Renderer/Private/IndirectLightRendering.cpp
+++ b/Engine/Source/Runtime/Renderer/Private/IndirectLightRendering.cpp
@@ -1249,6 +1249,16 @@ void FDeferredShadingSceneRenderer::RenderDiffuseIndirectAndAmbientOcclusion(
 			SceneColorCopyTexture = GraphBuilder.CreateTexture(SceneColorCopyTextureDesc, TEXT("SceneCopyTextureCopy0"));
 		}
 
+		auto SetDefaultBlackIfNotTexture2DArray = [&SystemTextures](FRDGTextureRef In)
+		{
+			return In ? (In->Desc.Dimension == ETextureDimension::Texture2DArray ? In : SystemTextures.BlackArray) : nullptr;
+		};
+
+		auto SetDefaultBlackIfNotTexture2D = [&SystemTextures](FRDGTextureRef In)
+		{
+			return In ? (In->Desc.Dimension == ETextureDimension::Texture2D ? In : SystemTextures.Black) : nullptr;
+		};
+
 		auto ApplyDiffuseIndirect = [&](ESubstrateTileType TileType)
 		{
 			if (bEnableCopyPass)
@@ -1293,12 +1303,17 @@ void FDeferredShadingSceneRenderer::RenderDiffuseIndirectAndAmbientOcclusion(
 
 			PassParameters->bVisualizeDiffuseIndirect = bIsVisualizePass;
 
-			PassParameters->DiffuseIndirect_Lumen_0 = DenoiserOutputs.Textures[0];
-			PassParameters->DiffuseIndirect_Lumen_1 = DenoiserOutputs.Textures[1];
-			PassParameters->DiffuseIndirect_Lumen_2 = DenoiserOutputs.Textures[2];
-			PassParameters->DiffuseIndirect_Lumen_3 = DenoiserOutputs.Textures[3];
+			PassParameters->DiffuseIndirect_Lumen_0 = SetDefaultBlackIfNotTexture2DArray(DenoiserOutputs.Textures[0]);
+			PassParameters->DiffuseIndirect_Lumen_1 = SetDefaultBlackIfNotTexture2DArray(DenoiserOutputs.Textures[1]);
+			PassParameters->DiffuseIndirect_Lumen_2 = SetDefaultBlackIfNotTexture2DArray(DenoiserOutputs.Textures[2]);
+			PassParameters->DiffuseIndirect_Lumen_3 = SetDefaultBlackIfNotTexture2DArray(DenoiserOutputs.Textures[3]);
 
 			PassParameters->DiffuseIndirect = DenoiserOutputs;
+			for (uint32 EntryIt = 0; EntryIt < kMaxDenoiserBufferProcessingCount; ++EntryIt)
+			{
+				PassParameters->DiffuseIndirect.Textures[EntryIt] = SetDefaultBlackIfNotTexture2D(PassParameters->DiffuseIndirect.Textures[EntryIt]);
+			}
+
 			PassParameters->DiffuseIndirectSampler = TStaticSamplerState<SF_Point>::GetRHI();
 
 			PassParameters->PreIntegratedGF = GSystemTextures.PreintegratedGF->GetRHI();
diff --git a/Engine/Source/Runtime/Renderer/Private/SingleLayerWaterRendering.cpp b/Engine/Source/Runtime/Renderer/Private/SingleLayerWaterRendering.cpp
--- a/Engine/Source/Runtime/Renderer/Private/SingleLayerWaterRendering.cpp
+++ b/Engine/Source/Runtime/Renderer/Private/SingleLayerWaterRendering.cpp
@@ -250,7 +250,7 @@ bool NeedsSeparatedMainDirectionalLightTexture_Runtime(const FStaticShaderPlatfo
 }
 
 BEGIN_SHADER_PARAMETER_STRUCT(FSingleLayerWaterCommonShaderParameters, )
-	SHADER_PARAMETER_RDG_TEXTURE(Texture2D, ScreenSpaceReflectionsTexture)
+	SHADER_PARAMETER_RDG_TEXTURE_SRV(Texture2D, ScreenSpaceReflectionsTexture)
 	SHADER_PARAMETER_SAMPLER(SamplerState, ScreenSpaceReflectionsSampler)
 	SHADER_PARAMETER_TEXTURE(Texture2D, PreIntegratedGF)
 	SHADER_PARAMETER_SAMPLER(SamplerState, PreIntegratedGFSampler)
@@ -847,7 +847,16 @@ void FDeferredShadingSceneRenderer::RenderSingleLayerWaterReflections(
 
 			const bool bIsInstancedStereoSideBySide = View.bIsInstancedStereoEnabled && !View.bIsMobileMultiViewEnabled && IStereoRendering::IsStereoEyeView(View);
 
-			Parameters.ScreenSpaceReflectionsTexture = ReflectionsColor ? ReflectionsColor : BlackDummyTexture;
+			FRDGTextureRef ScreenSpaceReflectionsTexture = ReflectionsColor ? ReflectionsColor : BlackDummyTexture;
+			if (ReflectionsColor && ReflectionsColor->Desc.Dimension == ETextureDimension::Texture2DArray)
+			{
+				Parameters.ScreenSpaceReflectionsTexture = GraphBuilder.CreateSRV(FRDGTextureSRVDesc::CreateForSlice(ScreenSpaceReflectionsTexture, 0));
+			}
+			else
+			{
+				Parameters.ScreenSpaceReflectionsTexture = GraphBuilder.CreateSRV(ScreenSpaceReflectionsTexture);
+			}
+
 			Parameters.ScreenSpaceReflectionsSampler = TStaticSamplerState<SF_Point>::GetRHI();
 			Parameters.PreIntegratedGF = GSystemTextures.PreintegratedGF->GetRHI();
 			Parameters.PreIntegratedGFSampler = TStaticSamplerState<SF_Bilinear, AM_Clamp, AM_Clamp, AM_Clamp>::GetRHI();
