* Merge all per-pixel OIT sample data into a single texture.
* Remove unecessary clear
* Merge/clean load/store logic

diff --git a/Engine/Shaders/Private/OITCombine.usf b/Engine/Shaders/Private/OITCombine.usf
--- a/Engine/Shaders/Private/OITCombine.usf
+++ b/Engine/Shaders/Private/OITCombine.usf
@@ -19,20 +19,9 @@ uint Method;
 uint PassType;
 uint SupportedPass;
 
-Texture2D<uint>	SampleColorTexture;
-Texture2D<uint>	SampleTransTexture;
-Texture2D<float>SampleDepthTexture;
+Texture2DArray<uint> SampleDataTexture;
 Texture2D<uint>	SampleCountTexture;
 
-FOITSample OITLoadSample(uint2 InPixelCoord)
-{
-	FOITPackedSample Packed = (FOITPackedSample)0;
-	Packed.Depth = SampleDepthTexture.Load(uint3(InPixelCoord, 0));
-	Packed.Color = SampleColorTexture.Load(uint3(InPixelCoord, 0));
-	Packed.Trans = SampleTransTexture.Load(uint3(InPixelCoord, 0));
-	return UnpackOITSample(Packed);
-}
-
 ///////////////////////////////////////////////////////////////////////////////////////////////////
 // Combine
 
@@ -61,7 +50,7 @@ void MainCS(uint3 DispatchThreadId : SV_DispatchThreadID)
 		const uint2 SampleOffset = uint2(SampleIt % MaxSideSampleCount, SampleIt / MaxSideSampleCount);
 		const uint2 SampleCoord  = SampleBaseCoord + SampleOffset;
 
-		FOITSample Current = OITLoadSample(SampleCoord);
+		FOITSample Current = OITLoadSample(SampleDataTexture, SampleCoord);
 		if (SampleIt > 0)
 		{
 			for (uint j = 0; j <= SampleIt; ++j)
@@ -200,7 +189,7 @@ void MainCS(uint3 DispatchThreadId : SV_DispatchThreadID)
 			const uint2 SampleOffset = uint2(SampleIt % MaxSideSampleCount, SampleIt / MaxSideSampleCount);
 			const uint2 SampleCoord = SampleBaseCoord + SampleOffset;
 
-			const FOITSample Sample = OITLoadSample(SampleCoord);
+			const FOITSample Sample = OITLoadSample(SampleDataTexture, SampleCoord);
 
 			Print(Context, TEXT("Depth "), FontYellow);
 			Print(Context, Sample.Depth, FontYellow);
diff --git a/Engine/Shaders/Private/OITCommon.ush b/Engine/Shaders/Private/OITCommon.ush
--- a/Engine/Shaders/Private/OITCommon.ush
+++ b/Engine/Shaders/Private/OITCommon.ush
@@ -30,14 +30,6 @@ struct FOITSample
 	float  Pad;
 };
 
-struct FOITPackedSample
-{
-	uint  Color;
-	uint  Trans;
-	float Depth;
-	uint  Pad;
-};
-
 //////////////////////////////////////////////////////////////////////////////////////////
 
 uint OITPack10F(float Value)
@@ -88,21 +80,25 @@ float3 OITUnpackR11G11B10F(uint rgb)
 // Per-attribut compaction 12bytes/96bits
 
 #if OIT_FORMAT == OIT_FORMAT_10BITS
+
+#define FOITPackedSample uint3
+#define OIT_DATA_COUNT 3
+
 FOITPackedSample PackOITSample(FOITSample In)
 {
 	FOITPackedSample Out = (FOITPackedSample)0;
-	Out.Color = OITPackR11G11B10F(In.Color);
-	Out.Trans = OITPackR11G11B10F(saturate(In.Trans));
-	Out.Depth = In.Depth;
+	Out.x = asuint(In.Depth);
+	Out.y = OITPackR11G11B10F(In.Color);
+	Out.z = OITPackR11G11B10F(saturate(In.Trans));
 	return Out;
 }
 
 FOITSample UnpackOITSample(FOITPackedSample In)
 {
 	FOITSample Out = (FOITSample)0;
-	Out.Color = OITUnpackR11G11B10F(In.Color);
-	Out.Trans = saturate(OITUnpackR11G11B10F(In.Trans));
-	Out.Depth = In.Depth;
+	Out.Depth = asfloat(In.x);
+	Out.Color = OITUnpackR11G11B10F(In.y);
+	Out.Trans = saturate(OITUnpackR11G11B10F(In.z));
 	return Out;
 }
 #endif
@@ -110,8 +106,12 @@ FOITSample UnpackOITSample(FOITPackedSample In)
 //////////////////////////////////////////////////////////////////////////////////////////
 // Fat compaction 16bytes/128bits
 
-#if OIT_FORMAT_16BITS
-uint4 PackOITSample(float3 InColor, float3 InTrans, float InDepth)
+#if OIT_FORMAT == OIT_FORMAT_16BITS
+
+#define FOITPackedSample uint4
+#define OIT_DATA_COUNT 4
+
+FOITPackedSample PackOITSample(float3 InColor, float3 InTrans, float InDepth)
 {
 	uint4 Out = 0;
 	Out.x = f32tof16(InColor.r) | (f32tof16(InColor.g) << 16);
@@ -121,7 +121,7 @@ uint4 PackOITSample(float3 InColor, float3 InTrans, float InDepth)
 	return Out;
 }
 
-void UnpackOITSample(uint4 In, inout float3 OutColor, inout float3 OutTrans, inout float OutDepth)
+void UnpackOITSample(FOITPackedSample In, inout float3 OutColor, inout float3 OutTrans, inout float OutDepth)
 {
 	OutColor.r = f16tof32(0xFFFF & (In.x));
 	OutColor.g = f16tof32(0xFFFF & (In.x >> 16));
@@ -139,7 +139,11 @@ void UnpackOITSample(uint4 In, inout float3 OutColor, inout float3 OutTrans, ino
 // Compact sample into 8bytes/64bits (untested|add dithering if banding)
 
 #if OIT_FORMAT == OIT_FORMAT_COMPACT
-uint2 PackOITSample(float3 InColor, float3 InTrans, float InDepth)
+
+#define FOITPackedSample uint2
+#define OIT_DATA_COUNT 2
+
+FOITPackedSample PackOITSample(float3 InColor, float3 InTrans, float InDepth)
 {
 	uint ColorR10 = OITPack10F(InColor.r);
 	uint ColorG10 = OITPack10F(InColor.g);
@@ -151,13 +155,13 @@ uint2 PackOITSample(float3 InColor, float3 InTrans, float InDepth)
 
 	uint Depth16 = f32tof16(InDepth);
 	
-	uint2 Out = 0;
+	FOITPackedSample Out = 0;
 	Out.x = ColorR10 | (ColorG10 << 10) | (ColorB10 << 20) | (((TransB6 >> 4) & 0x3) << 30);
 	Out.y = Depth16  | (TransR6  << 16) | (TransG6  << 22) | (TransB6 & 0xF) << 28);
 	return Out;
 }
 
-void UnpackOITSample(uint2 In, inout float3 OutColor, inout float3 OutTrans, inout float OutDepth)
+void UnpackOITSample(FOITPackedSample In, inout float3 OutColor, inout float3 OutTrans, inout float OutDepth)
 {
 	OutColor.r = OITUnpack10F((In.x)     & 3FF);
 	OutColor.g = OITUnpack10F((In.x>>10) & 3FF);
@@ -173,6 +177,22 @@ void UnpackOITSample(uint2 In, inout float3 OutColor, inout float3 OutTrans, ino
 
 //////////////////////////////////////////////////////////////////////////////////////////
 
+FOITSample OITLoadSample(Texture2DArray<uint> InTexture, uint2 InPixelCoord)
+{
+	FOITPackedSample Packed = (FOITPackedSample)0;
+	Packed.x = InTexture[uint3(InPixelCoord, 0)];
+#if OIT_DATA_COUNT > 1
+	Packed.y = InTexture[uint3(InPixelCoord, 1)];
+#endif
+#if OIT_DATA_COUNT > 2
+	Packed.z = InTexture[uint3(InPixelCoord, 2)];
+#endif
+#if OIT_DATA_COUNT > 3
+	Packed.w = InTexture[uint3(InPixelCoord, 3)];
+#endif
+	return UnpackOITSample(Packed);
+}
+
 #if OIT_IS_BASEPASS
 
 //////////////////////////////////////////////////////////////////////////////////////////
@@ -181,19 +201,32 @@ void UnpackOITSample(uint2 In, inout float3 OutColor, inout float3 OutTrans, ino
 FOITSample OITLoadSample(uint2 InPixelCoord)
 {
 	FOITPackedSample Packed = (FOITPackedSample)0;
-	Packed.Color = TranslucentBasePass.OIT.OutOITSampleColor[InPixelCoord];
-	Packed.Trans = TranslucentBasePass.OIT.OutOITSampleTrans[InPixelCoord];
-	Packed.Depth = TranslucentBasePass.OIT.OutOITSampleDepth[InPixelCoord];
-
+	Packed.x = TranslucentBasePass.OIT.OutOITSampleData[uint3(InPixelCoord, 0)];
+#if OIT_DATA_COUNT > 1
+	Packed.y = TranslucentBasePass.OIT.OutOITSampleData[uint3(InPixelCoord, 1)];
+#endif
+#if OIT_DATA_COUNT > 2
+	Packed.z = TranslucentBasePass.OIT.OutOITSampleData[uint3(InPixelCoord, 2)];
+#endif
+#if OIT_DATA_COUNT > 3
+	Packed.w = TranslucentBasePass.OIT.OutOITSampleData[uint3(InPixelCoord, 3)];
+#endif
 	return UnpackOITSample(Packed);
 }
 
 void OITStoreSample(uint2 InPixelCoord, FOITSample In)
 {
 	FOITPackedSample Packed = PackOITSample(In);
-	TranslucentBasePass.OIT.OutOITSampleColor[InPixelCoord] = Packed.Color;
-	TranslucentBasePass.OIT.OutOITSampleTrans[InPixelCoord] = Packed.Trans;
-	TranslucentBasePass.OIT.OutOITSampleDepth[InPixelCoord] = Packed.Depth;
+	TranslucentBasePass.OIT.OutOITSampleData[uint3(InPixelCoord, 0)] = Packed.x;
+#if OIT_DATA_COUNT > 1
+	TranslucentBasePass.OIT.OutOITSampleData[uint3(InPixelCoord, 1)] = Packed.y;
+#endif
+#if OIT_DATA_COUNT > 2
+	TranslucentBasePass.OIT.OutOITSampleData[uint3(InPixelCoord, 2)] = Packed.z;
+#endif
+#if OIT_DATA_COUNT > 3
+	TranslucentBasePass.OIT.OutOITSampleData[uint3(InPixelCoord, 4)] = Packed.w;
+#endif
 }
 
 // Simple k-buffer insertion
diff --git a/Engine/Source/Runtime/Renderer/Private/OIT/OIT.cpp b/Engine/Source/Runtime/Renderer/Private/OIT/OIT.cpp
--- a/Engine/Source/Runtime/Renderer/Private/OIT/OIT.cpp
+++ b/Engine/Source/Runtime/Renderer/Private/OIT/OIT.cpp
@@ -106,9 +106,7 @@ class FOITPixelDebugCS : public FGlobalShader
 		SHADER_PARAMETER(uint32, Method)
 		SHADER_PARAMETER(uint32, PassType)
 		SHADER_PARAMETER(uint32, SupportedPass)
-		SHADER_PARAMETER_RDG_TEXTURE(Texture2D<float3>, SampleColorTexture)
-		SHADER_PARAMETER_RDG_TEXTURE(Texture2D<float3>, SampleTransTexture)
-		SHADER_PARAMETER_RDG_TEXTURE(Texture2D<float>, SampleDepthTexture)
+		SHADER_PARAMETER_RDG_TEXTURE(Texture2DArray<uint>, SampleDataTexture)
 		SHADER_PARAMETER_RDG_TEXTURE(Texture2D<uint>, SampleCountTexture)
 		SHADER_PARAMETER_STRUCT_INCLUDE(ShaderPrint::FShaderParameters, ShaderPrintUniformBuffer)
 		END_SHADER_PARAMETER_STRUCT()
@@ -128,9 +126,7 @@ static void AddOITPixelDebugPass(
 	const FViewInfo& View,
 	FOITData& OITData)
 {
-	if (!OITData.SampleColorTexture ||
-		!OITData.SampleTransTexture ||
-		!OITData.SampleDepthTexture ||
+	if (!OITData.SampleDataTexture ||
 		!OITData.SampleCountTexture || 
 		!ShaderPrint::IsSupported(View.GetShaderPlatform()))
 		return;
@@ -147,9 +143,7 @@ static void AddOITPixelDebugPass(
 	Parameters->MaxSampleCount = OITData.MaxSamplePerPixel;
 	Parameters->Method = OITData.Method;
 	Parameters->MaxSideSampleCount = OITData.MaxSideSamplePerPixel;
-	Parameters->SampleColorTexture = OITData.SampleColorTexture;
-	Parameters->SampleTransTexture = OITData.SampleTransTexture;
-	Parameters->SampleDepthTexture = OITData.SampleDepthTexture;
+	Parameters->SampleDataTexture  = OITData.SampleDataTexture;
 	Parameters->SampleCountTexture = OITData.SampleCountTexture;
 	ShaderPrint::SetParameters(GraphBuilder, View.ShaderPrintData, Parameters->ShaderPrintUniformBuffer);
 
@@ -171,18 +165,16 @@ class FOITPixelCombineCS : public FGlobalShader
 	DECLARE_GLOBAL_SHADER(FOITPixelCombineCS);
 	SHADER_USE_PARAMETER_STRUCT(FOITPixelCombineCS, FGlobalShader);
 
-	BEGIN_SHADER_PARAMETER_STRUCT(FParameters, )
+	BEGIN_SHADER_PARAMETER_STRUCT(FParameters,)
 		SHADER_PARAMETER(FIntPoint, Resolution)
 		SHADER_PARAMETER(uint32, MaxSideSampleCount)
 		SHADER_PARAMETER(uint32, MaxSampleCount)
 		SHADER_PARAMETER(uint32, Method)
 		SHADER_PARAMETER(uint32, PassType)
-		SHADER_PARAMETER_RDG_TEXTURE(Texture2D<float3>, SampleColorTexture)
-		SHADER_PARAMETER_RDG_TEXTURE(Texture2D<float3>, SampleTransTexture)
-		SHADER_PARAMETER_RDG_TEXTURE(Texture2D<float>, SampleDepthTexture)
+		SHADER_PARAMETER_RDG_TEXTURE(Texture2DArray<uint>, SampleDataTexture)
 		SHADER_PARAMETER_RDG_TEXTURE(Texture2D<uint>, SampleCountTexture)
 		SHADER_PARAMETER_RDG_TEXTURE_UAV(RWTexture2D<float4>, OutColorTexture)
-		END_SHADER_PARAMETER_STRUCT()
+	END_SHADER_PARAMETER_STRUCT()
 public:
 	static bool ShouldCompilePermutation(const FGlobalShaderPermutationParameters& Parameters) { return IsOITSortedPixelsSupported(Parameters.Platform); }
 	static void ModifyCompilationEnvironment(const FGlobalShaderPermutationParameters& Parameters, FShaderCompilerEnvironment& OutEnvironment)
@@ -205,13 +197,11 @@ static void AddInternalOITComposePass(
 
 	if (!SceneColorTexture ||
 		 SceneColorTexture->Desc.NumSamples > 1 ||
-		!OITData.SampleColorTexture ||
-		!OITData.SampleTransTexture ||
-		!OITData.SampleDepthTexture ||
+		!OITData.SampleDataTexture ||
 		!OITData.SampleCountTexture)
 		return;
 
-	FIntPoint InResolution = OITData.SampleColorTexture->Desc.Extent;
+	FIntPoint InResolution = OITData.SampleDataTexture->Desc.Extent;
 	FIntPoint OutResolution = SceneColorTexture->Desc.Extent;
 	FIntPoint Resolution = FIntPoint(FMath::Min(InResolution.X, OutResolution.X), FMath::Min(InResolution.Y, OutResolution.Y));
 
@@ -221,9 +211,7 @@ static void AddInternalOITComposePass(
 	Parameters->Method = OITData.Method;
 	Parameters->MaxSampleCount = OITData.MaxSamplePerPixel;
 	Parameters->MaxSideSampleCount = OITData.MaxSideSamplePerPixel;
-	Parameters->SampleColorTexture = OITData.SampleColorTexture;
-	Parameters->SampleTransTexture = OITData.SampleTransTexture;
-	Parameters->SampleDepthTexture = OITData.SampleDepthTexture;
+	Parameters->SampleDataTexture = OITData.SampleDataTexture;
 	Parameters->SampleCountTexture = OITData.SampleCountTexture;
 	Parameters->OutColorTexture = GraphBuilder.CreateUAV(SceneColorTexture);
 
@@ -231,11 +219,8 @@ static void AddInternalOITComposePass(
 	TShaderMapRef<FOITPixelCombineCS > ComputeShader(View.ShaderMap, PermutationVector);
 
 	// Add 64 threads permutation
-	const uint32 GroupSize = 8;
-	const FIntVector DispatchCount = FIntVector(
-		(Resolution.X + GroupSize - 1) / GroupSize,
-		(Resolution.Y + GroupSize - 1) / GroupSize,
-		1);
+	const int32 GroupSize = 8;
+	const FIntVector DispatchCount = FIntVector(FMath::DivideAndRoundUp(Resolution.X, GroupSize), FMath::DivideAndRoundUp(Resolution.Y, GroupSize), 1);
 	FComputeShaderUtils::AddPass(
 		GraphBuilder,
 		RDG_EVENT_NAME("Translucency::OITCombine(%s)", !!(OITData.PassType & OITPass_SeperateTranslucency) ? TEXT("SeparateTransluency") : TEXT("Regular")),
@@ -922,7 +907,7 @@ namespace OIT
 		const bool bOIT = IsEnabled(EOITSortingType::SortedPixels, View.GetShaderPlatform());
 		const uint32 PassTypeBits = FMath::Clamp(CVarOIT_SortedPixels_PassType.GetValueOnRenderThread(), 0, 3);
 		const bool bPassValid = !!(PassTypeBits & PassType);
-
+		const uint32 LayerCount = 3u; /*Depth/Color/Trans*/
 		FOITData Out;
 		if (!bOIT || !bPassValid)
 		{
@@ -931,9 +916,7 @@ namespace OIT
 			Out.Method = 0;
 			Out.TransmittanceThreshold = 0;
 
-			Out.SampleColorTexture = GraphBuilder.CreateTexture(FRDGTextureDesc::Create2D(FIntPoint(1, 1), PF_R32_UINT, FClearValueBinding::None, TexCreate_ShaderResource | TexCreate_UAV), TEXT("OIT.SampleColor"));
-			Out.SampleTransTexture = GraphBuilder.CreateTexture(FRDGTextureDesc::Create2D(FIntPoint(1, 1), PF_R32_UINT, FClearValueBinding::None, TexCreate_ShaderResource | TexCreate_UAV), TEXT("OIT.SampleTrans"));
-			Out.SampleDepthTexture = GraphBuilder.CreateTexture(FRDGTextureDesc::Create2D(FIntPoint(1, 1), PF_R32_FLOAT, FClearValueBinding::Black, TexCreate_ShaderResource | TexCreate_UAV), TEXT("OIT.SampleDepth"));
+			Out.SampleDataTexture  = GraphBuilder.CreateTexture(FRDGTextureDesc::Create2DArray(FIntPoint(1, 1), PF_R32_UINT, FClearValueBinding::None, TexCreate_ShaderResource | TexCreate_UAV, LayerCount), TEXT("OIT.SampleData"));
 			Out.SampleCountTexture = GraphBuilder.CreateTexture(FRDGTextureDesc::Create2D(FIntPoint(1, 1), PF_R32_UINT, FClearValueBinding::Black, TexCreate_ShaderResource | TexCreate_UAV), TEXT("OIT.SampleCount"));
 			return Out;
 		}
@@ -960,17 +943,13 @@ namespace OIT
 			Out.Method = FMath::Clamp(CVarOIT_SortedPixels_Method.GetValueOnRenderThread(), 0, 1);
 			Out.TransmittanceThreshold = FMath::Clamp(CVarOIT_SortedPixels_TransmittanceThreshold.GetValueOnRenderThread(), 0.f, 1.f);
 
-			Out.SampleColorTexture = GraphBuilder.CreateTexture(FRDGTextureDesc::Create2D(EffectiveBufferSize * MaxSideSamplePerPixel, PF_R32_UINT, FClearValueBinding::None, TexCreate_ShaderResource | TexCreate_UAV), TEXT("OIT.SampleColor"));
-			Out.SampleTransTexture = GraphBuilder.CreateTexture(FRDGTextureDesc::Create2D(EffectiveBufferSize * MaxSideSamplePerPixel, PF_R32_UINT, FClearValueBinding::None, TexCreate_ShaderResource | TexCreate_UAV), TEXT("OIT.SampleTrans"));
-			Out.SampleDepthTexture = GraphBuilder.CreateTexture(FRDGTextureDesc::Create2D(EffectiveBufferSize * MaxSideSamplePerPixel, PF_R32_FLOAT, FClearValueBinding::Black, TexCreate_ShaderResource | TexCreate_UAV), TEXT("OIT.SampleDepth"));
+			Out.SampleDataTexture = GraphBuilder.CreateTexture(FRDGTextureDesc::Create2DArray(EffectiveBufferSize * MaxSideSamplePerPixel, PF_R32_UINT, FClearValueBinding::None, TexCreate_ShaderResource | TexCreate_UAV, LayerCount), TEXT("OIT.SampleData"));
 			Out.SampleCountTexture = GraphBuilder.CreateTexture(FRDGTextureDesc::Create2D(EffectiveBufferSize, PF_R32_UINT, FClearValueBinding::None, TexCreate_ShaderResource | TexCreate_UAV), TEXT("OIT.SampleCount"));
 		}
 
 		// TODO: Add tile clear based on the coarse translucent raster AABB buffer
-		AddClearUAVPass(GraphBuilder, GraphBuilder.CreateUAV(Out.SampleColorTexture), 0u);
-		AddClearUAVPass(GraphBuilder, GraphBuilder.CreateUAV(Out.SampleTransTexture), 0u);
+		//AddClearUAVPass(GraphBuilder, GraphBuilder.CreateUAV(Out.SampleDataTexture), 0u);
 		AddClearUAVPass(GraphBuilder, GraphBuilder.CreateUAV(Out.SampleCountTexture), 0u);
-		AddClearUAVPass(GraphBuilder, GraphBuilder.CreateUAV(Out.SampleDepthTexture), 0.f);
 
 		return Out;
 	}
@@ -983,9 +962,7 @@ namespace OIT
 		OutOIT.MaxSideSamplePerPixel = InOITData.MaxSideSamplePerPixel;
 		OutOIT.MaxSamplePerPixel = InOITData.MaxSamplePerPixel;
 		OutOIT.TransmittanceThreshold = InOITData.TransmittanceThreshold;
-		OutOIT.OutOITSampleColor = GraphBuilder.CreateUAV(InOITData.SampleColorTexture);
-		OutOIT.OutOITSampleTrans = GraphBuilder.CreateUAV(InOITData.SampleTransTexture);
-		OutOIT.OutOITSampleDepth = GraphBuilder.CreateUAV(InOITData.SampleDepthTexture);
+		OutOIT.OutOITSampleData  = GraphBuilder.CreateUAV(InOITData.SampleDataTexture);
 		OutOIT.OutOITSampleCount = GraphBuilder.CreateUAV(InOITData.SampleCountTexture);
 	}
 
diff --git a/Engine/Source/Runtime/Renderer/Private/OIT/OIT.h b/Engine/Source/Runtime/Renderer/Private/OIT/OIT.h
--- a/Engine/Source/Runtime/Renderer/Private/OIT/OIT.h
+++ b/Engine/Source/Runtime/Renderer/Private/OIT/OIT.h
@@ -65,9 +65,7 @@ struct FOITData
 	uint32 MaxSideSamplePerPixel = 0;
 	float TransmittanceThreshold = 0.f;
 
-	FRDGTextureRef SampleColorTexture = nullptr;
-	FRDGTextureRef SampleTransTexture = nullptr;
-	FRDGTextureRef SampleDepthTexture = nullptr;
+	FRDGTextureRef SampleDataTexture = nullptr;
 	FRDGTextureRef SampleCountTexture = nullptr;
 };
 
diff --git a/Engine/Source/Runtime/Renderer/Private/OIT/OITParameters.h b/Engine/Source/Runtime/Renderer/Private/OIT/OITParameters.h
--- a/Engine/Source/Runtime/Renderer/Private/OIT/OITParameters.h
+++ b/Engine/Source/Runtime/Renderer/Private/OIT/OITParameters.h
@@ -15,9 +15,7 @@ BEGIN_SHADER_PARAMETER_STRUCT(FOITBasePassUniformParameters, )
 	SHADER_PARAMETER(uint32, MaxSamplePerPixel)
 	SHADER_PARAMETER(float, TransmittanceThreshold)
 	SHADER_PARAMETER_RDG_TEXTURE_UAV(RasterizerOrderedTexture2D<uint>, OutOITSampleCount)
-	SHADER_PARAMETER_RDG_TEXTURE_UAV(RWTexture2D<uint>,  OutOITSampleColor)
-	SHADER_PARAMETER_RDG_TEXTURE_UAV(RWTexture2D<uint>,  OutOITSampleTrans)
-	SHADER_PARAMETER_RDG_TEXTURE_UAV(RWTexture2D<float>, OutOITSampleDepth)
+	SHADER_PARAMETER_RDG_TEXTURE_UAV(RWTexture2DArray<uint>, OutOITSampleData)
 END_SHADER_PARAMETER_STRUCT()
 
 class FViewInfo;
