Fix for mismatched depth buffer sizing when drawing editor primitive debug views. Separated the depth buffer resizing logic from the color texture/editor "prior pass" logic to ensure that depth temporal upscaling occurs in all circumstances that require it.

diff --git a/Engine/Source/Runtime/Renderer/Private/PostProcess/PostProcessCompositeEditorPrimitives.cpp b/Engine/Source/Runtime/Renderer/Private/PostProcess/PostProcessCompositeEditorPrimitives.cpp
--- a/Engine/Source/Runtime/Renderer/Private/PostProcess/PostProcessCompositeEditorPrimitives.cpp
+++ b/Engine/Source/Runtime/Renderer/Private/PostProcess/PostProcessCompositeEditorPrimitives.cpp
@@ -224,22 +224,16 @@ FScreenPassTexture AddEditorPrimitivePass(
 	const FViewInfo* EditorView = CreateCompositePrimitiveView(View, Inputs.SceneColor.ViewRect, NumMSAASamples);
 
 	// Load the color target if it already exists.
-	const bool bProducedByPriorPass = HasBeenProduced(SceneTextures.EditorPrimitiveColor);
-
+	bool bProducedByPriorPass = HasBeenProduced(SceneTextures.EditorPrimitiveColor);
+	FIntPoint Extent = Inputs.SceneColor.Texture->Desc.Extent;
 	FRDGTextureRef EditorPrimitiveColor;
 	FRDGTextureRef EditorPrimitiveDepth;
 	if (bProducedByPriorPass)
 	{
-		ensureMsgf(
-			Inputs.SceneColor.ViewRect == Inputs.SceneDepth.ViewRect,
-			TEXT("Temporal upsampling should be disabled when drawing directly to EditorPrimitivesColor."));
 		EditorPrimitiveColor = SceneTextures.EditorPrimitiveColor;
-		EditorPrimitiveDepth = SceneTextures.EditorPrimitiveDepth;
 	}
 	else
 	{
-		FIntPoint Extent = Inputs.SceneColor.Texture->Desc.Extent;
-
 		const FRDGTextureDesc ColorDesc = FRDGTextureDesc::Create2D(
 			Extent,
 			PF_B8G8R8A8,
@@ -249,7 +243,17 @@ FScreenPassTexture AddEditorPrimitivePass(
 			NumMSAASamples);
 
 		EditorPrimitiveColor = GraphBuilder.CreateTexture(ColorDesc, TEXT("Editor.PrimitivesColor"));
+	}
+
+	if (bProducedByPriorPass && Inputs.SceneColor.ViewRect == Inputs.SceneDepth.ViewRect)
+	{
+		EditorPrimitiveDepth = SceneTextures.EditorPrimitiveDepth;
+	}
+	else
+	{
 		EditorPrimitiveDepth = CreateCompositeDepthTexture(GraphBuilder, Extent, NumMSAASamples);
+		//bProducedByPriorPass no longer true as this pass had to create a depth texture for temporal upscaling
+		bProducedByPriorPass = false;
 	}
 	
 	// Subtrate data might not be produced in certain case (e.g., path-tracer). In such a case we force generate 
