LFV - added support for half res rendering of local fog volume on the mobile shading path.
Tested with split screen.

See also https://horde.devtools.epicgames.com/job/654cba80387f0780bed2243f

diff --git a/Engine/Shaders/Private/LocalFogVolumes/LocalFogVolumeCommon.ush b/Engine/Shaders/Private/LocalFogVolumes/LocalFogVolumeCommon.ush
--- a/Engine/Shaders/Private/LocalFogVolumes/LocalFogVolumeCommon.ush
+++ b/Engine/Shaders/Private/LocalFogVolumes/LocalFogVolumeCommon.ush
@@ -44,6 +44,7 @@ struct FLocalFogVolumeGPUInstanceData
 #define LFVMaxDensityIntoVolumetricFog	LFVStruct.LFV.LocalFogVolumeMaxDensityIntoVolumetricFog
 #define LFVRenderInVolumetricFog		LFVStruct.LFV.ShouldRenderLocalFogVolumeInVolumetricFog
 #define LFVGlobalStartDistance			LFVStruct.LFV.GlobalStartDistance
+#define LFVHalfResTextureSizeAndInvSize	LFVStruct.LFV.HalfResTextureSizeAndInvSize
 
 #define LFVDirectionalLightColor		LFVStruct.LFV.DirectionalLightColor
 #define LFVDirectionalLightDirection	LFVStruct.LFV.DirectionalLightDirection
@@ -56,6 +57,7 @@ uint									LFV_LocalFogVolumeTilePixelSize;
 float									LFV_LocalFogVolumeMaxDensityIntoVolumetricFog;
 uint									LFV_ShouldRenderLocalFogVolumeInVolumetricFog;
 float									LFV_GlobalStartDistance;
+float4									LFV_HalfResTextureSizeAndInvSize;
 
 float3									LFV_DirectionalLightColor;
 float3									LFV_DirectionalLightDirection;
@@ -72,6 +74,7 @@ Texture2DArray<uint>					LFV_LocalFogVolumeTileDataTexture;
 #define LFVMaxDensityIntoVolumetricFog	LFV_LocalFogVolumeMaxDensityIntoVolumetricFog
 #define LFVRenderInVolumetricFog		LFV_ShouldRenderLocalFogVolumeInVolumetricFog
 #define LFVGlobalStartDistance			LFV_GlobalStartDistance
+#define LFVHalfResTextureSizeAndInvSize	LFV_HalfResTextureSizeAndInvSize
 
 #define LFVDirectionalLightColor		LFV_DirectionalLightColor
 #define LFVDirectionalLightDirection	LFV_DirectionalLightDirection
diff --git a/Engine/Shaders/Private/LocalFogVolumes/LocalFogVolumeSplat.usf b/Engine/Shaders/Private/LocalFogVolumes/LocalFogVolumeSplat.usf
--- a/Engine/Shaders/Private/LocalFogVolumes/LocalFogVolumeSplat.usf
+++ b/Engine/Shaders/Private/LocalFogVolumes/LocalFogVolumeSplat.usf
@@ -15,6 +15,9 @@
 #ifndef LFV_TILED_PS
 #define LFV_TILED_PS 0
 #endif
+#ifndef LFV_HALFRES_PS
+#define LFV_HALFRES_PS 0
+#endif
 
 
 
@@ -63,24 +66,40 @@ int LocalFogVolumeTileDebug;
 
 void LocalFogVolumeTiledPS(
 	in float4 SVPos : SV_POSITION,
+#if LFV_HALFRES_PS
+	out float4 OutColorThroughput : SV_Target0,
+	out float OutDepth : SV_Target1)
+#else
 	in FLocalFogVolumeVertexOutput OtherVertexOutput,
 	out float4 OutColor : SV_Target0)
+#endif
 {
 	ResolvedView = ResolveView();
 
+#if LFV_HALFRES_PS
+	const float HalfResScale = ResolvedView.ViewSizeAndInvSize.xy * LFVHalfResTextureSizeAndInvSize.zw;
+#else
+	const float HalfResScale = 1.0f;
+#endif
+	const float2 SourcePixelPos = SVPos.xy * HalfResScale + ResolvedView.ViewRectMin.xy;
+
 #if SHADING_PATH_MOBILE
-	// LookupDeviceZ with uint2 SVPos.xy as input does not work on mobile.
-	float2 UV = SVPos * View.BufferSizeAndInvSize.zw;
-	float DeviceZ = LookupDeviceZ(UV);
+	float2 UVDepth = SourcePixelPos * ResolvedView.BufferSizeAndInvSize.zw;
+	float DeviceZ = LookupDeviceZ(UVDepth); // LookupDeviceZ with uint2(SVPos.xy) as input does not work on mobile.
 #else
 	float DeviceZ = LookupDeviceZ(uint2(SVPos.xy));
 #endif
+
 #if HAS_INVERTED_Z_BUFFER
 	DeviceZ = max(0.000000000001, DeviceZ);	// Make sure SvPositionToWorld does not return bad values when DeviceZ is far=0 when using inverted z
 #endif
-	float3 DepthBufferTranslatedWorldPos = SvPositionToTranslatedWorld(float4(SVPos.xy, DeviceZ, 1.0));
+	float3 DepthBufferTranslatedWorldPos = SvPositionToTranslatedWorld(float4(SourcePixelPos, DeviceZ, 1.0));
 
+#if LFV_HALFRES_PS
+	uint2 TilePos = (SourcePixelPos - ResolvedView.ViewRectMin.xy) / LFVTilePixelSize;
+#else
 	uint2 TilePos = LFVUnpackTile(OtherVertexOutput.TileData);
+#endif
 
 	uint LFVCount = 0;
 	const bool bPlatformSupportsVolumetricFogOntranslucent = false;
@@ -90,6 +109,9 @@ void LocalFogVolumeTiledPS(
 #if USE_DEVELOPMENT_SHADERS
 	bDoClip = LocalFogVolumeTileDebug != 1;
 #endif
+#if LFV_HALFRES_PS
+	bDoClip = false; // We want to write all pixel to avoid clear operation
+#endif
 
 	const float4 IdentityLuminanceTransmittance = float4(0.0, 0.0, 0.0, 1.0);
 	if (all(LFVContribution == IdentityLuminanceTransmittance) && bDoClip)
@@ -97,11 +119,18 @@ void LocalFogVolumeTiledPS(
 		clip(-1.0f);
 	}
 
+#if LFV_HALFRES_PS
+	OutColorThroughput = float4(
+		LFVContribution.rgb * PrimaryView.PreExposure,
+		LFVContribution.a);
+	OutDepth = ConvertFromDeviceZ(DeviceZ) * CENTIMETER_TO_METER; // Allows to represent depth up to 65 km while keeping enough accuracy close to the camera
+#else
 	OutColor = float4(
 		LFVContribution.rgb * PrimaryView.PreExposure,
 		LFVContribution.a);
+#endif
 
-#if USE_DEVELOPMENT_SHADERS
+#if USE_DEVELOPMENT_SHADERS && !LFV_HALFRES_PS
 	if (LocalFogVolumeTileDebug > 0)
 	{
 		float DebugTransmittance = 0.5;
diff --git a/Engine/Shaders/Private/MobileFog.usf b/Engine/Shaders/Private/MobileFog.usf
--- a/Engine/Shaders/Private/MobileFog.usf
+++ b/Engine/Shaders/Private/MobileFog.usf
@@ -9,10 +9,17 @@
 #include "HeightFogCommon.ush"
 #include "SkyAtmosphereCommon.ush"
 
-#if PERMUTATION_SUPPORT_LOCAL_FOG_VOLUME
+#if PERMUTATION_SUPPORT_LOCAL_FOG_VOLUME > 0
 #include "LocalFogVolumes/LocalFogVolumeCommon.ush"
 #endif
 
+#if PERMUTATION_SUPPORT_LOCAL_FOG_VOLUME == 2
+Texture2D<float4> HalfResLocalFogVolumeViewSRV;
+SamplerState HalfResLocalFogVolumeViewSRVSampler;
+Texture2D<float4>  HalfResLocalFogVolumeDepthSRV;
+SamplerState HalfResLocalFogVolumeDepthSRVSampler;
+#endif
+
 float StartDepthZ;
 
 void MobileFogVS(
@@ -51,11 +58,74 @@ void MobileFogPS(
 	FinalFog = CalculateHeightFog(WorldPositionRelativeToCamera);
 #endif
 
-#if PERMUTATION_SUPPORT_LOCAL_FOG_VOLUME
+#if PERMUTATION_SUPPORT_LOCAL_FOG_VOLUME == 1
+
+	// Evaluate the local fog volumes in line when there is volumetric fog.
 	uint2 TilePos = uint2(SVPos.xy - ResolvedView.ViewRectMin.xy) / LFVTilePixelSize.xx;
 	float4 LFVContribution = GetLFVContribution(ResolvedView, TilePos, WorldPositionRelativeToCamera);
 	FinalFog = float4(LFVContribution.rgb + FinalFog.rgb * LFVContribution.a, LFVContribution.a * FinalFog.a);
-#endif
+
+#elif PERMUTATION_SUPPORT_LOCAL_FOG_VOLUME == 2
+
+	// Upsample the half resolution local fog volumes previously traced in a half resolution texture.
+	const int2 FullResolutionToVolumetricBufferResolutionScale = uint2(2, 2);
+
+	int2 PixelPos = SVPos.xy - ResolvedView.ViewRectMin.xy;
+	int2 VolumeCoordUInt = PixelPos / int(FullResolutionToVolumetricBufferResolutionScale.y);
+	int OffsetX = (VolumeCoordUInt.x * int(FullResolutionToVolumetricBufferResolutionScale.y)) == PixelPos.x ? -1 : 1;
+	int OffsetY = (VolumeCoordUInt.y * int(FullResolutionToVolumetricBufferResolutionScale.y)) == PixelPos.y ? -1 : 1;
+
+	uint2 VolumeCoord0 = max(0, int2(VolumeCoordUInt) + int2(0, 0));
+	uint2 VolumeCoord1 = max(0, int2(VolumeCoordUInt) + int2(OffsetX, 0));
+	uint2 VolumeCoord2 = max(0, int2(VolumeCoordUInt) + int2(OffsetX, OffsetY));
+	uint2 VolumeCoord3 = max(0, int2(VolumeCoordUInt) + int2(0, OffsetY));
+
+	float VolumeFrontDepth0 = HalfResLocalFogVolumeDepthSRV.Load(uint3(VolumeCoord0, 0)).r * METER_TO_CENTIMETER;
+	float VolumeFrontDepth1 = HalfResLocalFogVolumeDepthSRV.Load(uint3(VolumeCoord1, 0)).r * METER_TO_CENTIMETER;
+	float VolumeFrontDepth2 = HalfResLocalFogVolumeDepthSRV.Load(uint3(VolumeCoord2, 0)).r * METER_TO_CENTIMETER;
+	float VolumeFrontDepth3 = HalfResLocalFogVolumeDepthSRV.Load(uint3(VolumeCoord3, 0)).r * METER_TO_CENTIMETER;
+
+	float MaxFrontDepth4Diff = max(max(VolumeFrontDepth0, VolumeFrontDepth1), max(VolumeFrontDepth2, VolumeFrontDepth3));
+	float UpsampleSceneDepth = min(SceneDepth, MaxFrontDepth4Diff);	// needed because FP16 is limited in distance adn that fixes upsampling artefacts when depth is super large..
+
+	float4 Depth4 = float4(VolumeFrontDepth0, VolumeFrontDepth1, VolumeFrontDepth2, VolumeFrontDepth3);
+	float4 Depth4Diff = UpsampleSceneDepth - Depth4;
+	Depth4Diff = abs(Depth4Diff);
+	float MaxDepth4Diff = max(max(Depth4Diff.x, Depth4Diff.y), max(Depth4Diff.z, Depth4Diff.w));
+
+	const float WeightMultiplier = 10.0f;
+	const float ThresholdToBilinear = 30.0f; // in centimeters
+	float4 LFVContribution;
+	if (MaxDepth4Diff > ThresholdToBilinear)
+	{
+		float MaxDepth4Diff = max(max(Depth4Diff.x, Depth4Diff.y), max(Depth4Diff.z, Depth4Diff.w));
+		// Depth discontinuities edges
+		float4 weights = 1.0f / (Depth4Diff * WeightMultiplier + 1.0f);
+		const float weightsSum = dot(weights, float4(1.0f, 1.0f, 1.0f, 1.0f));
+		weights /= weightsSum;
+
+		float4 VolumeRGBT0 = HalfResLocalFogVolumeViewSRV.Load(uint3(VolumeCoord0, 0));
+		float4 VolumeRGBT1 = HalfResLocalFogVolumeViewSRV.Load(uint3(VolumeCoord1, 0));
+		float4 VolumeRGBT2 = HalfResLocalFogVolumeViewSRV.Load(uint3(VolumeCoord2, 0));
+		float4 VolumeRGBT3 = HalfResLocalFogVolumeViewSRV.Load(uint3(VolumeCoord3, 0));
+
+		LFVContribution = weights.x * VolumeRGBT0 + weights.y * VolumeRGBT1 + weights.z * VolumeRGBT2 + weights.w * VolumeRGBT3;
+	}
+	else
+	{
+		// This is needed when half res is not exactly a multiple of 2, for the bilinear filtering to work nicely.
+		const float2 ResolutionRatioCorrection = (ResolvedView.ViewSizeAndInvSize.xy * LFVHalfResTextureSizeAndInvSize.zw * 0.5);
+		const float2 HalfResSourceTexCoord = float2(SVPos.xy - ResolvedView.ViewRectMin.xy) * ResolvedView.ViewSizeAndInvSize.zw * ResolutionRatioCorrection;
+		LFVContribution = Texture2DSampleLevel(HalfResLocalFogVolumeViewSRV, HalfResLocalFogVolumeViewSRVSampler, HalfResSourceTexCoord, 0.0f); // Simple bilinear filtering for increased performance
+	}
+
+	// HalfResLocalFogVolumeViewSRV contains pre exposed color, so remove it for composition.
+	LFVContribution.rgb /= ResolvedView.PreExposure;
+
+	// Combine with the curren fog
+	FinalFog = float4(LFVContribution.rgb + FinalFog.rgb * LFVContribution.a, LFVContribution.a * FinalFog.a);
+
+#endif // PERMUTATION_SUPPORT_LOCAL_FOG_VOLUME
 
 #if PERMUTATION_SUPPORT_VOLUMETRIC_FOG
 	if (FogStruct.ApplyVolumetricFog > 0) 
diff --git a/Engine/Source/Runtime/Renderer/Private/DeferredShadingRenderer.cpp b/Engine/Source/Runtime/Renderer/Private/DeferredShadingRenderer.cpp
--- a/Engine/Source/Runtime/Renderer/Private/DeferredShadingRenderer.cpp
+++ b/Engine/Source/Runtime/Renderer/Private/DeferredShadingRenderer.cpp
@@ -3464,7 +3464,7 @@ void FDeferredShadingSceneRenderer::Render(FRDGBuilder& GraphBuilder)
 		BeginAsyncDistanceFieldShadowProjections(GraphBuilder, SceneTextures, InitViewTaskDatas.DynamicShadows);
 
 		// Run local fog volume culling before base pass and after HZB generation tyo benefit from more culling.
-		InitLocalFogVolumesForViews(Scene, Views, ViewFamily, GraphBuilder, bShouldRenderVolumetricFog);
+		InitLocalFogVolumesForViews(Scene, Views, ViewFamily, GraphBuilder, bShouldRenderVolumetricFog, false /*bool bUseHalfResLocalFogVolume*/);
 
 		if (bShouldRenderVolumetricCloudBase)
 		{
diff --git a/Engine/Source/Runtime/Renderer/Private/LocalFogVolumeRendering.cpp b/Engine/Source/Runtime/Renderer/Private/LocalFogVolumeRendering.cpp
--- a/Engine/Source/Runtime/Renderer/Private/LocalFogVolumeRendering.cpp
+++ b/Engine/Source/Runtime/Renderer/Private/LocalFogVolumeRendering.cpp
@@ -6,6 +6,7 @@
 #include "ScreenPass.h"
 #include "LocalFogVolumeSceneProxy.h"
 #include "MobileBasePassRendering.h"
+#include "PixelShaderUtils.h"
 
 
 // The runtime ON/OFF toggle
@@ -71,6 +72,11 @@ static TAutoConsoleVariable<int32> CVarLocalFogVolumeUseHZB(
 	TEXT("Use the HZB to cull loca lfog volume away.\n"),
 	ECVF_RenderThreadSafe);
 
+static TAutoConsoleVariable<int32> CVarLocalFogVolumeHalfResolution(
+	TEXT("r.LocalFogVolume.HalfResolution"), 0,
+	TEXT("Set to one to render local fog volumes at half resoltuion with an upsampling to full resolution later. Only work for the mobile path for now.\n"),
+	ECVF_RenderThreadSafe);
+
 // Example of tile setup
 //  - 1920x1080 => 15x9 tiles
 //  - Allowing max 32 volumes at once => culling list buffer = 15 * 9 * 32 * 1 byte = 4320 bytes = 4.3KB
@@ -134,6 +140,10 @@ float GetLocalFogVolumeGlobalStartDistance()
 	return FMath::Max(0.0f, CVarLocalFogVolumeGlobalStartDistance.GetValueOnRenderThread());
 }
 
+bool IsLocalFogVolumeHalfResolution()
+{
+	return CVarLocalFogVolumeHalfResolution.GetValueOnRenderThread() > 0;
+}
 
 DECLARE_GPU_STAT(LocalFogVolumeVolumes);
 
@@ -170,6 +180,7 @@ void SetDummyLocalFogVolumeForView(FRDGBuilder& GraphBuilder, FViewInfo& View)
 	View.LocalFogVolumeViewData.UniformParametersStruct.LocalFogVolumeCommon.DirectionalLightColor						= FVector3f::Zero();
 	View.LocalFogVolumeViewData.UniformParametersStruct.LocalFogVolumeCommon.DirectionalLightDirection					= FVector3f::Zero();
 	View.LocalFogVolumeViewData.UniformParametersStruct.LocalFogVolumeCommon.GlobalStartDistance						= 0.0f;
+	View.LocalFogVolumeViewData.UniformParametersStruct.LocalFogVolumeCommon.HalfResTextureSizeAndInvSize				= FVector4f(1.0f, 1.0f, 1.0f, 1.0f);
 	View.LocalFogVolumeViewData.UniformParametersStruct.LocalFogVolumeTileDataTexture									= View.LocalFogVolumeViewData.TileDataTextureArraySRV;
 	View.LocalFogVolumeViewData.UniformBuffer																			= GraphBuilder.CreateUniformBuffer(&View.LocalFogVolumeViewData.UniformParametersStruct);
 
@@ -181,6 +192,12 @@ void SetDummyLocalFogVolumeForView(FRDGBuilder& GraphBuilder, FViewInfo& View)
 
 	View.LocalFogVolumeViewData.GPUTileDrawIndirectBuffer	= GraphBuilder.CreateBuffer(FRDGBufferDesc::CreateIndirectDesc<FRHIDrawIndirectParameters>(), TEXT("LocalFogVolume.DispatchIndirectBuffer"));
 	View.LocalFogVolumeViewData.GPUTileDrawIndirectBufferUAV= nullptr;
+
+	View.LocalFogVolumeViewData.bUseHalfResLocalFogVolume = false;
+	View.LocalFogVolumeViewData.HalfResLocalFogVolumeView = GSystemTextures.GetBlackDummy(GraphBuilder);
+	View.LocalFogVolumeViewData.HalfResLocalFogVolumeViewSRV = GraphBuilder.CreateSRV(View.LocalFogVolumeViewData.HalfResLocalFogVolumeView);
+	View.LocalFogVolumeViewData.HalfResLocalFogVolumeDepth = GSystemTextures.GetBlackDummy(GraphBuilder);
+	View.LocalFogVolumeViewData.HalfResLocalFogVolumeDepthSRV = GraphBuilder.CreateSRV(View.LocalFogVolumeViewData.HalfResLocalFogVolumeDepth);
 }
 
 void SetDummyLocalFogVolumeForViews(FRDGBuilder& GraphBuilder, TArray<FViewInfo>& Views)
@@ -422,7 +439,7 @@ void GetLocalFogVolumeSortingData(const FScene* Scene, FRDGBuilder& GraphBuilder
 	Out.LocalFogVolumeSortKeys.SetNum(Out.LocalFogVolumeInstanceCountFinal, false/*bAllowShrinking*/);
 }
 
-void CreateViewLocalFogVolumeBufferSRV(const FScene* Scene, FViewInfo& View, FRDGBuilder& GraphBuilder, FLocalFogVolumeSortingData& SortingData, bool bShouldRenderLocalFogVolumeInVolumetricFog)
+void CreateViewLocalFogVolumeBufferSRV(const FScene* Scene, FViewInfo& View, FRDGBuilder& GraphBuilder, FLocalFogVolumeSortingData& SortingData, bool bShouldRenderLocalFogVolumeInVolumetricFog, bool bUseHalfResLocalFogVolume)
 {
 	if (SortingData.LocalFogVolumeInstanceCountFinal == 0)
 	{
@@ -523,8 +540,6 @@ void CreateViewLocalFogVolumeBufferSRV(const FScene* Scene, FViewInfo& View, FRD
 		View.LocalFogVolumeViewData.UniformParametersStruct.LocalFogVolumeCommon.DirectionalLightColor					= FVector3f(View.CachedViewUniformShaderParameters->DirectionalLightColor);
 		View.LocalFogVolumeViewData.UniformParametersStruct.LocalFogVolumeCommon.DirectionalLightDirection				= View.CachedViewUniformShaderParameters->DirectionalLightDirection;
 	}
-	View.LocalFogVolumeViewData.UniformParametersStruct.LocalFogVolumeTileDataTexture									= View.LocalFogVolumeViewData.TileDataTextureArraySRV;
-	View.LocalFogVolumeViewData.UniformBuffer																			= GraphBuilder.CreateUniformBuffer(&View.LocalFogVolumeViewData.UniformParametersStruct);
 
 	// This buffer must remain a basic vertex buffer for mobile to be able to read it from vertex shader
 	View.LocalFogVolumeViewData.GPUTileDataBuffer = CreateVertexBuffer(
@@ -536,6 +551,39 @@ void CreateViewLocalFogVolumeBufferSRV(const FScene* Scene, FViewInfo& View, FRD
 	View.LocalFogVolumeViewData.GPUTileDrawIndirectBuffer = GraphBuilder.CreateBuffer(FRDGBufferDesc::CreateIndirectDesc<FRHIDrawIndirectParameters>(), TEXT("LocalFogVolume.DispatchIndirectBuffer"));
 	View.LocalFogVolumeViewData.GPUTileDrawIndirectBufferUAV = GraphBuilder.CreateUAV(View.LocalFogVolumeViewData.GPUTileDrawIndirectBuffer, PF_R32_UINT);
 	AddClearUAVPass(GraphBuilder, View.LocalFogVolumeViewData.GPUTileDrawIndirectBufferUAV, 0, GetLocalFogVolumeTileCullingUseAsync() ? ERDGPassFlags::AsyncCompute : ERDGPassFlags::Compute);
+
+	View.LocalFogVolumeViewData.bUseHalfResLocalFogVolume = bUseHalfResLocalFogVolume;
+	if (View.LocalFogVolumeViewData.bUseHalfResLocalFogVolume)
+	{
+		FIntRect ViewRectAtOrigin = View.ViewRect;
+		ViewRectAtOrigin.Max -= ViewRectAtOrigin.Min;
+		ViewRectAtOrigin.Min -= ViewRectAtOrigin.Min;
+		FIntRect HalfResRect = GetDownscaledViewport(FScreenPassTextureViewport(ViewRectAtOrigin), FIntPoint(2, 2)).Rect;
+		View.LocalFogVolumeViewData.HalfResResolution = HalfResRect.Max - HalfResRect.Min;
+
+		const FIntPoint& HalfResSize = View.LocalFogVolumeViewData.HalfResResolution;
+		View.LocalFogVolumeViewData.UniformParametersStruct.LocalFogVolumeCommon.HalfResTextureSizeAndInvSize = FVector4f(HalfResSize.X, HalfResSize.Y, 1.0f/float(HalfResSize.X), 1.0f/float(HalfResSize.Y));
+
+		FRDGTextureDesc Texture2DHalfResLFVDesc(FRDGTextureDesc::Create2D(View.LocalFogVolumeViewData.HalfResResolution, PF_FloatRGBA, FClearValueBinding(EClearBinding::ENoneBound), TexCreate_RenderTargetable | TexCreate_ShaderResource | TexCreate_ReduceMemoryWithTilingMode | TexCreate_NoFastClear));
+
+		View.LocalFogVolumeViewData.HalfResLocalFogVolumeView = GraphBuilder.CreateTexture(Texture2DHalfResLFVDesc, TEXT("LocalFogVolume.HalfResLocalFogVolumeView"));
+		View.LocalFogVolumeViewData.HalfResLocalFogVolumeViewSRV = GraphBuilder.CreateSRV(FRDGTextureSRVDesc::Create(View.LocalFogVolumeViewData.HalfResLocalFogVolumeView));
+
+		Texture2DHalfResLFVDesc.Format = PF_R16F;
+		View.LocalFogVolumeViewData.HalfResLocalFogVolumeDepth = GraphBuilder.CreateTexture(Texture2DHalfResLFVDesc, TEXT("LocalFogVolume.HalfResLocalFogVolumeDepth"));
+		View.LocalFogVolumeViewData.HalfResLocalFogVolumeDepthSRV = GraphBuilder.CreateSRV(FRDGTextureSRVDesc::Create(View.LocalFogVolumeViewData.HalfResLocalFogVolumeDepth));
+	}
+	else
+	{
+		View.LocalFogVolumeViewData.UniformParametersStruct.LocalFogVolumeCommon.HalfResTextureSizeAndInvSize = FVector4f(1.0f, 1.0f, 1.0f, 1.0f);
+		View.LocalFogVolumeViewData.HalfResLocalFogVolumeView = GSystemTextures.GetBlackDummy(GraphBuilder);
+		View.LocalFogVolumeViewData.HalfResLocalFogVolumeViewSRV = GraphBuilder.CreateSRV(View.LocalFogVolumeViewData.HalfResLocalFogVolumeView);
+		View.LocalFogVolumeViewData.HalfResLocalFogVolumeDepth = GSystemTextures.GetBlackDummy(GraphBuilder);
+		View.LocalFogVolumeViewData.HalfResLocalFogVolumeDepthSRV = GraphBuilder.CreateSRV(View.LocalFogVolumeViewData.HalfResLocalFogVolumeDepth);
+	}
+
+	View.LocalFogVolumeViewData.UniformParametersStruct.LocalFogVolumeTileDataTexture = View.LocalFogVolumeViewData.TileDataTextureArraySRV;
+	View.LocalFogVolumeViewData.UniformBuffer = GraphBuilder.CreateUniformBuffer(&View.LocalFogVolumeViewData.UniformParametersStruct);
 }
 
 void InitLocalFogVolumesForViews(
@@ -543,7 +591,8 @@ void InitLocalFogVolumesForViews(
 	TArray<FViewInfo>& Views,
 	const FSceneViewFamily& SceneViewFamily,
 	FRDGBuilder& GraphBuilder,
-	bool bShouldRenderVolumetricFog)
+	bool bShouldRenderVolumetricFog,
+	bool bUseHalfResLocalFogVolume)
 {
 	const uint32 LocalFogVolumeInstanceCount = Scene->LocalFogVolumes.Num();
 	const bool bShouldRenderLocalFogVolume = ShouldRenderLocalFogVolume(Scene, SceneViewFamily);
@@ -556,7 +605,7 @@ void InitLocalFogVolumesForViews(
 
 		for (FViewInfo& View : Views)
 		{
-			CreateViewLocalFogVolumeBufferSRV(Scene, View, GraphBuilder, SortingData, ShouldRenderLocalFogVolumeInVolumetricFog(Scene, SceneViewFamily, bShouldRenderVolumetricFog));
+			CreateViewLocalFogVolumeBufferSRV(Scene, View, GraphBuilder, SortingData, ShouldRenderLocalFogVolumeInVolumetricFog(Scene, SceneViewFamily, bShouldRenderVolumetricFog), bUseHalfResLocalFogVolume);
 
 			if (View.LocalFogVolumeViewData.GPUInstanceCount > 0)
 			{
@@ -837,6 +886,43 @@ class FMobileLocalFogVolumeTiledRenderPS : public FGlobalShader
 
 IMPLEMENT_GLOBAL_SHADER(FMobileLocalFogVolumeTiledRenderPS, "/Engine/Private/LocalFogVolumes/LocalFogVolumeSplat.usf", "LocalFogVolumeTiledPS", SF_Pixel);
 
+class FMobileLocalFogVolumeTiledRenderHalfResPS : public FGlobalShader
+{
+	DECLARE_GLOBAL_SHADER(FMobileLocalFogVolumeTiledRenderHalfResPS);
+	SHADER_USE_PARAMETER_STRUCT(FMobileLocalFogVolumeTiledRenderHalfResPS, FGlobalShader);
+
+	BEGIN_SHADER_PARAMETER_STRUCT(FParameters, )
+		SHADER_PARAMETER_STRUCT_INCLUDE(FViewShaderParameters, View)
+		SHADER_PARAMETER_RDG_UNIFORM_BUFFER(FMobileBasePassUniformParameters, MobileBasePass)
+		//SHADER_PARAMETER_RDG_UNIFORM_BUFFER(FMobileSceneTextureUniformParameters, SceneTextures)
+		SHADER_PARAMETER_STRUCT(FLocalFogVolumeUniformParameters, LFV)
+		SHADER_PARAMETER(int32, LocalFogVolumeTileDebug)
+		RENDER_TARGET_BINDING_SLOTS()
+	END_SHADER_PARAMETER_STRUCT()
+
+	using FPermutationDomain = TShaderPermutationDomain<>;
+
+	static bool ShouldCompilePermutation(const FGlobalShaderPermutationParameters& Parameters)
+	{
+		if (!IsMobilePlatform(Parameters.Platform))
+		{
+			return false;
+		}
+		return true;
+	}
+
+	static void ModifyCompilationEnvironment(const FGlobalShaderPermutationParameters& Parameters, FShaderCompilerEnvironment& OutEnvironment)
+	{
+		FGlobalShader::ModifyCompilationEnvironment(Parameters, OutEnvironment);
+		const bool bMobileForceDepthRead = MobileUsesFullDepthPrepass(Parameters.Platform);
+		OutEnvironment.SetDefine(TEXT("LFV_TILED_PS"), 1);
+		OutEnvironment.SetDefine(TEXT("LFV_HALFRES_PS"), 1);
+		OutEnvironment.SetDefine(TEXT("IS_MOBILE_DEPTHREAD_SUBPASS"), bMobileForceDepthRead ? 0u : 1u);
+	}
+};
+
+IMPLEMENT_GLOBAL_SHADER(FMobileLocalFogVolumeTiledRenderHalfResPS, "/Engine/Private/LocalFogVolumes/LocalFogVolumeSplat.usf", "LocalFogVolumeTiledPS", SF_Pixel);
+
 void RenderLocalFogVolumeMobile(
 	FRHICommandList& RHICmdList,
 	const FViewInfo& View)
@@ -892,3 +978,30 @@ void RenderLocalFogVolumeMobile(
 	RHICmdList.SetStreamSource(0, nullptr, 0);
 	RHICmdList.DrawPrimitiveIndirect(View.LocalFogVolumeViewData.GPUTileDrawIndirectBuffer->GetIndirectRHICallBuffer(), 0);
 }
+
+void RenderLocalFogVolumeHalfResMobile(
+	FRDGBuilder& GraphBuilder,
+	const FViewInfo& View)
+{
+	if (View.LocalFogVolumeViewData.GPUInstanceCount == 0)
+	{
+		return;
+	}
+	FMobileLocalFogVolumeTiledRenderHalfResPS::FParameters* PassParameters = GraphBuilder.AllocParameters<FMobileLocalFogVolumeTiledRenderHalfResPS::FParameters>();
+
+	PassParameters->View = View.GetShaderParameters();;
+	PassParameters->LFV = View.LocalFogVolumeViewData.UniformParametersStruct;
+	PassParameters->LocalFogVolumeTileDebug = FMath::Clamp(CVarLocalFogVolumeTileDebug.GetValueOnRenderThread(), 0, 2);
+	PassParameters->MobileBasePass = CreateMobileBasePassUniformBuffer(GraphBuilder, View, EMobileBasePass::Opaque, EMobileSceneTextureSetupMode::SceneDepth);
+	PassParameters->RenderTargets[0] = FRenderTargetBinding(View.LocalFogVolumeViewData.HalfResLocalFogVolumeView, ERenderTargetLoadAction::ENoAction);
+	PassParameters->RenderTargets[1] = FRenderTargetBinding(View.LocalFogVolumeViewData.HalfResLocalFogVolumeDepth, ERenderTargetLoadAction::ENoAction);
+
+	FMobileLocalFogVolumeTiledRenderHalfResPS::FPermutationDomain PsPermutationVector;
+	auto PixelShader = View.ShaderMap->GetShader< FMobileLocalFogVolumeTiledRenderHalfResPS >(PsPermutationVector);
+	ClearUnusedGraphResources(PixelShader, PassParameters);
+
+	const FIntPoint HalfResolution = View.LocalFogVolumeViewData.HalfResResolution;
+	FPixelShaderUtils::AddFullscreenPass<FMobileLocalFogVolumeTiledRenderHalfResPS>(
+		GraphBuilder, View.ShaderMap, RDG_EVENT_NAME("LocalFogVolume.HalfRes"), PixelShader, PassParameters,
+		FIntRect(0, 0, HalfResolution.X, HalfResolution.Y));
+}
diff --git a/Engine/Source/Runtime/Renderer/Private/LocalFogVolumeRendering.h b/Engine/Source/Runtime/Renderer/Private/LocalFogVolumeRendering.h
--- a/Engine/Source/Runtime/Renderer/Private/LocalFogVolumeRendering.h
+++ b/Engine/Source/Runtime/Renderer/Private/LocalFogVolumeRendering.h
@@ -27,6 +27,7 @@ BEGIN_SHADER_PARAMETER_STRUCT(FLocalFogVolumeCommonParameters, )
 	SHADER_PARAMETER(float,  LocalFogVolumeMaxDensityIntoVolumetricFog)
 	SHADER_PARAMETER(uint32, ShouldRenderLocalFogVolumeInVolumetricFog)
 	SHADER_PARAMETER(float,  GlobalStartDistance)
+	SHADER_PARAMETER(FVector4f, HalfResTextureSizeAndInvSize)
 	SHADER_PARAMETER(FVector3f, DirectionalLightColor)
 	SHADER_PARAMETER(FVector3f, DirectionalLightDirection)
 END_SHADER_PARAMETER_STRUCT()
@@ -107,9 +108,16 @@ struct FLocalFogVolumeViewData
 
 	FLocalFogVolumeUniformParameters UniformParametersStruct;
 
-	FRDGTextureRef		TileDataTextureArray;		// First slice is the instance count, later slices are instance indices.
-	FRDGTextureSRVRef	TileDataTextureArraySRV;
-	FRDGTextureUAVRef	TileDataTextureArrayUAV;
+	FRDGTextureRef		TileDataTextureArray = nullptr;		// First slice is the instance count, later slices are instance indices.
+	FRDGTextureSRVRef	TileDataTextureArraySRV = nullptr;
+	FRDGTextureUAVRef	TileDataTextureArrayUAV = nullptr;
+
+	bool				bUseHalfResLocalFogVolume = false;	// Only for the mobile path for now
+	FIntPoint			HalfResResolution;
+	FRDGTextureRef		HalfResLocalFogVolumeView = nullptr;
+	FRDGTextureSRVRef	HalfResLocalFogVolumeViewSRV = nullptr;
+	FRDGTextureRef		HalfResLocalFogVolumeDepth = nullptr;
+	FRDGTextureSRVRef	HalfResLocalFogVolumeDepthSRV = nullptr;
 };
 
 
@@ -122,6 +130,7 @@ bool ShouldRenderLocalFogVolume(const FScene* Scene, const FSceneViewFamily& Sce
 bool ShouldRenderLocalFogVolumeDuringHeightFogPass(const FScene* Scene, const FSceneViewFamily& SceneViewFamily);
 bool ShouldRenderLocalFogVolumeInVolumetricFog(const FScene* Scene, const FSceneViewFamily& SceneViewFamily, bool bShouldRenderVolumetricFog);
 float GetLocalFogVolumeGlobalStartDistance();
+bool IsLocalFogVolumeHalfResolution();
 
 void GetLocalFogVolumeSortingData(const FScene* Scene, FRDGBuilder& GraphBuilder, FLocalFogVolumeSortingData& Out);
 
@@ -133,7 +142,8 @@ void InitLocalFogVolumesForViews(
 	TArray<FViewInfo>& Views,
 	const FSceneViewFamily& SceneViewFamily,
 	FRDGBuilder& GraphBuilder,
-	bool bShouldRenderVolumetricFog);
+	bool bShouldRenderVolumetricFog,
+	bool bUseHalfResLocalFogVolume);
 
 void RenderLocalFogVolume(
 	const FScene* Scene,
@@ -146,4 +156,7 @@ void RenderLocalFogVolume(
 void RenderLocalFogVolumeMobile(
 	FRHICommandList& RHICmdList, 
 	const FViewInfo& View);
+void RenderLocalFogVolumeHalfResMobile(
+	FRDGBuilder& GraphBuilder,
+	const FViewInfo& View);
 
diff --git a/Engine/Source/Runtime/Renderer/Private/MobileFogRendering.cpp b/Engine/Source/Runtime/Renderer/Private/MobileFogRendering.cpp
--- a/Engine/Source/Runtime/Renderer/Private/MobileFogRendering.cpp
+++ b/Engine/Source/Runtime/Renderer/Private/MobileFogRendering.cpp
@@ -66,7 +66,7 @@ class FMobileFogPS : public FGlobalShader
 	class FSupportFogDirectionalLightInScattering	: SHADER_PERMUTATION_BOOL("PERMUTATION_SUPPORT_FOG_DIRECTIONAL_LIGHT_INSCATTERING");
 	class FSupportAerialPerspective					: SHADER_PERMUTATION_BOOL("PERMUTATION_SUPPORT_AERIAL_PERSPECTIVE");
 	class FSupportVolumetricFog						: SHADER_PERMUTATION_BOOL("PERMUTATION_SUPPORT_VOLUMETRIC_FOG");
-	class FSupportLocalFogVolume					: SHADER_PERMUTATION_BOOL("PERMUTATION_SUPPORT_LOCAL_FOG_VOLUME");
+	class FSupportLocalFogVolume					: SHADER_PERMUTATION_INT("PERMUTATION_SUPPORT_LOCAL_FOG_VOLUME", 3); // 0 disabled, 1 sample LFVs, 2 compose halfres LFV texture
 	
 	using FPermutationDomain = TShaderPermutationDomain< 
 		FSupportHeightFog, 
@@ -117,6 +117,10 @@ class FMobileFogPS : public FGlobalShader
 	BEGIN_SHADER_PARAMETER_STRUCT(FParameters, )
 		SHADER_PARAMETER_STRUCT_INCLUDE(FViewShaderParameters, View)
 		SHADER_PARAMETER_STRUCT(FLocalFogVolumeUniformParameters, LFV)
+		SHADER_PARAMETER_RDG_TEXTURE_SRV(Texture2D<float4>, HalfResLocalFogVolumeViewSRV)
+		SHADER_PARAMETER_SAMPLER(SamplerState, HalfResLocalFogVolumeViewSRVSampler)
+		SHADER_PARAMETER_RDG_TEXTURE_SRV(Texture2D<float>, HalfResLocalFogVolumeDepthSRV)
+		SHADER_PARAMETER_SAMPLER(SamplerState, HalfResLocalFogVolumeDepthSRVSampler)
 	END_SHADER_PARAMETER_STRUCT()
 };
 
@@ -195,6 +199,8 @@ void FMobileSceneRenderer::RenderFog(FRHICommandList& RHICmdList, const FViewInf
 	
 	const bool bShouldRenderVolumetricFog = ShouldRenderVolumetricFog();
 	bFogHasComposedLocalFogVolumes = bShouldRenderVolumetricFog && ShouldRenderLocalFogVolume(Scene, *View.Family) && ShouldRenderLocalFogVolumeInVolumetricFog(Scene, *View.Family, bShouldRenderVolumetricFog);
+	const bool bComposeHalfResLFV = View.LocalFogVolumeViewData.bUseHalfResLocalFogVolume;
+	bFogHasComposedLocalFogVolumes |= bComposeHalfResLFV;
 
 	FMobileFogPS::FPermutationDomain PsPermutationVector;
 	PsPermutationVector.Set<FMobileFogPS::FSupportHeightFog>(bUseHeightFog);
@@ -204,7 +210,7 @@ void FMobileSceneRenderer::RenderFog(FRHICommandList& RHICmdList, const FViewInf
 	PsPermutationVector.Set<FMobileFogPS::FSupportAerialPerspective>(bUseAerialPerspective);
 	PsPermutationVector.Set<FMobileFogPS::FSupportFogDirectionalLightInScattering>(bUseFogDirectionalInscatering);
 	PsPermutationVector.Set<FMobileFogPS::FSupportVolumetricFog>(bShouldRenderVolumetricFog);
-	PsPermutationVector.Set<FMobileFogPS::FSupportLocalFogVolume>(bFogHasComposedLocalFogVolumes);
+	PsPermutationVector.Set<FMobileFogPS::FSupportLocalFogVolume>(bFogHasComposedLocalFogVolumes ? (bComposeHalfResLFV ? 2 : 1) : 0);
 	
 	TShaderMapRef<FMobileFogPS> PixelShader(View.ShaderMap, PsPermutationVector);
 		
@@ -250,6 +256,10 @@ void FMobileSceneRenderer::RenderFog(FRHICommandList& RHICmdList, const FViewInf
 	FMobileFogPS::FParameters PSParameters;
 	PSParameters.View = View.GetShaderParameters();
 	PSParameters.LFV = View.LocalFogVolumeViewData.UniformParametersStruct;
+	PSParameters.HalfResLocalFogVolumeViewSRV = View.LocalFogVolumeViewData.HalfResLocalFogVolumeViewSRV;
+	PSParameters.HalfResLocalFogVolumeDepthSRV = View.LocalFogVolumeViewData.HalfResLocalFogVolumeDepthSRV;
+	PSParameters.HalfResLocalFogVolumeViewSRVSampler = TStaticSamplerState<SF_Bilinear, AM_Clamp, AM_Clamp, AM_Clamp>::GetRHI();
+	PSParameters.HalfResLocalFogVolumeDepthSRVSampler = TStaticSamplerState<SF_Bilinear, AM_Clamp, AM_Clamp, AM_Clamp>::GetRHI();
 	SetShaderParameters(RHICmdList, PixelShader, PixelShader.GetPixelShader(), PSParameters);
 
 	// Draw a quad covering the view.
diff --git a/Engine/Source/Runtime/Renderer/Private/MobileShadingRenderer.cpp b/Engine/Source/Runtime/Renderer/Private/MobileShadingRenderer.cpp
--- a/Engine/Source/Runtime/Renderer/Private/MobileShadingRenderer.cpp
+++ b/Engine/Source/Runtime/Renderer/Private/MobileShadingRenderer.cpp
@@ -197,6 +197,10 @@ BEGIN_SHADER_PARAMETER_STRUCT(FMobileRenderPassParameters, RENDERER_API)
 	RDG_BUFFER_ACCESS(LocalFogVolumeTileDrawIndirectBuffer, ERHIAccess::IndirectArgs)
 	SHADER_PARAMETER_RDG_TEXTURE_SRV(Texture2DArray<uint>, LocalFogVolumeTileDataTexture)
 	SHADER_PARAMETER_RDG_BUFFER_SRV(Buffer<uint>, LocalFogVolumeTileDataBuffer)
+	SHADER_PARAMETER_RDG_TEXTURE_SRV(Texture2D<float4>, HalfResLocalFogVolumeViewSRV)
+	SHADER_PARAMETER_RDG_TEXTURE_SRV(Texture2D<float>, HalfResLocalFogVolumeDepthSRV)
+	SHADER_PARAMETER_RDG_TEXTURE_SRV(Texture2D<float4>, HalfResLocalFogVolumeViewSRV)
+	SHADER_PARAMETER_RDG_TEXTURE_SRV(Texture2D<float>, HalfResLocalFogVolumeDepthSRV)
 	RENDER_TARGET_BINDING_SLOTS()
 END_SHADER_PARAMETER_STRUCT()
 
@@ -951,6 +955,8 @@ void FMobileSceneRenderer::Render(FRDGBuilder& GraphBuilder)
 	SceneTextures.MobileSetupMode = EMobileSceneTextureSetupMode::None;
 	SceneTextures.MobileUniformBuffer = CreateMobileSceneTextureUniformBuffer(GraphBuilder, &SceneTextures, SceneTextures.MobileSetupMode);
 
+	const bool bUseHalfResLocalFogVolume = bIsFullDepthPrepassEnabled && IsLocalFogVolumeHalfResolution(); // We must have a full depth buffer in order to render half res and upsample
+
 	if (bRendererOutputFinalSceneColor)
 	{
 #if WITH_DEBUG_VIEW_MODES
@@ -1010,7 +1016,7 @@ void FMobileSceneRenderer::Render(FRDGBuilder& GraphBuilder)
 		RenderShadowDepthMaps(GraphBuilder, nullptr, InstanceCullingManager, ExternalAccessQueue);
 
 		// Run local fog volume initialization before base pass and volumetric fog for all the culled instance instance data to be ready.
-		InitLocalFogVolumesForViews(Scene, Views, ViewFamily, GraphBuilder, ShouldRenderVolumetricFog());
+		InitLocalFogVolumesForViews(Scene, Views, ViewFamily, GraphBuilder, ShouldRenderVolumetricFog(), bUseHalfResLocalFogVolume);
 
 		if (ShouldRenderVolumetricFog())
 		{
@@ -1109,6 +1115,15 @@ void FMobileSceneRenderer::Render(FRDGBuilder& GraphBuilder)
 				RenderDBuffer(GraphBuilder, SceneTextures, DBufferTextures);
 			}
 		}
+
+		// Render half res local fog volume here
+		for (FViewInfo& View : Views)
+		{
+			if (View.LocalFogVolumeViewData.bUseHalfResLocalFogVolume)
+			{
+				RenderLocalFogVolumeHalfResMobile(GraphBuilder, View);
+			}
+		}
 	}
 
 	if (bRendererOutputFinalSceneColor)
@@ -1338,6 +1353,8 @@ void FMobileSceneRenderer::RenderForward(FRDGBuilder& GraphBuilder, FRDGTextureR
 		PassParameters->LocalFogVolumeTileDrawIndirectBuffer = View.LocalFogVolumeViewData.GPUTileDrawIndirectBuffer;
 		PassParameters->LocalFogVolumeTileDataTexture = View.LocalFogVolumeViewData.TileDataTextureArraySRV;
 		PassParameters->LocalFogVolumeTileDataBuffer = View.LocalFogVolumeViewData.GPUTileDataBufferSRV;
+		PassParameters->HalfResLocalFogVolumeViewSRV = View.LocalFogVolumeViewData.HalfResLocalFogVolumeViewSRV;
+		PassParameters->HalfResLocalFogVolumeDepthSRV = View.LocalFogVolumeViewData.HalfResLocalFogVolumeDepthSRV;
 	
 		BuildInstanceCullingDrawParams(GraphBuilder, View, PassParameters);
 
@@ -1672,6 +1689,8 @@ void FMobileSceneRenderer::RenderDeferred(FRDGBuilder& GraphBuilder, const FSort
 		PassParameters->LocalFogVolumeTileDrawIndirectBuffer = View.LocalFogVolumeViewData.GPUTileDrawIndirectBuffer;
 		PassParameters->LocalFogVolumeTileDataTexture = View.LocalFogVolumeViewData.TileDataTextureArraySRV;
 		PassParameters->LocalFogVolumeTileDataBuffer = View.LocalFogVolumeViewData.GPUTileDataBufferSRV;
+		PassParameters->HalfResLocalFogVolumeViewSRV = View.LocalFogVolumeViewData.HalfResLocalFogVolumeViewSRV;
+		PassParameters->HalfResLocalFogVolumeDepthSRV = View.LocalFogVolumeViewData.HalfResLocalFogVolumeDepthSRV;
 		
 		BuildInstanceCullingDrawParams(GraphBuilder, View, PassParameters);
 
