VolumetricCloud - fix high quality aerial perspective to render correct ordering of AP, Fog, VFog and LFVs.

diff --git a/Engine/Shaders/Private/HeightFogPixelShader.usf b/Engine/Shaders/Private/HeightFogPixelShader.usf
--- a/Engine/Shaders/Private/HeightFogPixelShader.usf
+++ b/Engine/Shaders/Private/HeightFogPixelShader.usf
@@ -28,6 +28,13 @@ SamplerState OcclusionSampler;
 Texture2D WaterDepthTexture;
 SamplerState WaterDepthSampler;
 
+#if PERMUTATION_SAMPLE_FOG_ON_CLOUDS
+Texture2D<float> SrcCloudDepthTexture;
+SamplerState SrcCloudDepthSampler;
+Texture2D<float4> SrcCloudViewTexture;
+SamplerState SrcCloudViewSampler;
+#endif
+
 void ExponentialPixelMain(
 	float2 TexCoord : TEXCOORD0,
 	float3 ScreenVector : TEXCOORD1,
@@ -37,6 +44,22 @@ void ExponentialPixelMain(
 {
 	float SceneDepth;
 	float DeviceZ;
+#if PERMUTATION_SAMPLE_FOG_ON_CLOUDS
+
+	float2 PixPos = SVPos.xy;
+	const float4 CloudLuminanceTransmittance = SrcCloudViewTexture.Load(int3(PixPos, 0));
+	const float CloudCoverage = 1.0f - CloudLuminanceTransmittance.a;
+	if (CloudLuminanceTransmittance.a > 0.999)
+	{
+		OutColor = float4(0.0f, 0.0f, 0.0f, 1.0f);
+		return;
+	}
+
+	const float CloudDepthKm = SrcCloudDepthTexture.Load(int3(PixPos, 0)).r;
+	SceneDepth = CloudDepthKm * KILOMETER_TO_CENTIMETER;
+	DeviceZ = ConvertToDeviceZ(SceneDepth); // Warning: for simplicity, we use DeviceZ as world distance in kilometer when SAMPLE_ATMOSPHERE_ON_CLOUDS. See special case in IntegrateSingleScatteredLuminance.
+
+#else
 	if (bUseWaterDepthTexture)
 	{
 		float2 SampleUV = clamp(TexCoord, WaterDepthTextureMinMaxUV.xy, WaterDepthTextureMinMaxUV.zw);
@@ -53,6 +76,7 @@ void ExponentialPixelMain(
 		SceneDepth = ConvertFromDeviceZ(DeviceZ); // Fetch the depth buffer Z / W value, solve for W
 	#endif
 	}
+#endif
 
 	bool bIsRendered = (DeviceZ != 0.0);
 
@@ -88,6 +112,16 @@ void ExponentialPixelMain(
 
 	FogInscatteringAndOpacity.rgb *= View.PreExposure;
 
+
+#if PERMUTATION_SAMPLE_FOG_ON_CLOUDS
+
+	// Reduce cloud luminance according to the fog transmittance and add the fog in scattering luminance according to the cloud coverage.
+	OutColor.rgb = CloudLuminanceTransmittance.rgb* FogInscatteringAndOpacity.a + FogInscatteringAndOpacity.rgb * CloudCoverage;
+	// Coverage of the cloud layer itself does not change.
+	OutColor.a   = CloudLuminanceTransmittance.a;
+
+#else
+
 	if (bOnlyOnRenderedOpaque > 0.0 && !bIsRendered)
 	{
 		FogInscatteringAndOpacity.rgb = 0;
@@ -95,4 +129,5 @@ void ExponentialPixelMain(
 	}
 
 	OutColor = RETURN_COLOR(float4(FogInscatteringAndOpacity.rgb, FogInscatteringAndOpacity.w));
+#endif
 }
diff --git a/Engine/Source/Runtime/Renderer/Private/FogRendering.cpp b/Engine/Source/Runtime/Renderer/Private/FogRendering.cpp
--- a/Engine/Source/Runtime/Renderer/Private/FogRendering.cpp
+++ b/Engine/Source/Runtime/Renderer/Private/FogRendering.cpp
@@ -122,7 +122,8 @@ class FExponentialHeightFogPS : public FGlobalShader
 	class FSupportFogDirectionalLightInScattering : SHADER_PERMUTATION_BOOL("PERMUTATION_SUPPORT_FOG_DIRECTIONAL_LIGHT_INSCATTERING");
 	class FSupportVolumetricFog : SHADER_PERMUTATION_BOOL("PERMUTATION_SUPPORT_VOLUMETRIC_FOG");
 	class FSupportLocalFogVolume : SHADER_PERMUTATION_BOOL("PERMUTATION_SUPPORT_LOCAL_FOG_VOLUME");
-	using FPermutationDomain = TShaderPermutationDomain<FSupportFogInScatteringTexture, FSupportFogDirectionalLightInScattering, FSupportVolumetricFog, FSupportLocalFogVolume>;
+	class FSampleFogOnClouds : SHADER_PERMUTATION_BOOL("PERMUTATION_SAMPLE_FOG_ON_CLOUDS");
+	using FPermutationDomain = TShaderPermutationDomain<FSupportFogInScatteringTexture, FSupportFogDirectionalLightInScattering, FSupportVolumetricFog, FSupportLocalFogVolume, FSampleFogOnClouds>;
 
 	BEGIN_SHADER_PARAMETER_STRUCT(FParameters, )
 		SHADER_PARAMETER_STRUCT_REF(FViewUniformShaderParameters, ViewUniformBuffer)
@@ -137,6 +138,10 @@ class FExponentialHeightFogPS : public FGlobalShader
 		SHADER_PARAMETER(float, UpsampleJitterMultiplier)
 		SHADER_PARAMETER(FVector4f, WaterDepthTextureMinMaxUV)
 		SHADER_PARAMETER(FVector4f, OcclusionTextureMinMaxUV)
+		SHADER_PARAMETER_RDG_TEXTURE(Texture2D, SrcCloudDepthTexture)
+		SHADER_PARAMETER_SAMPLER(SamplerState, SrcCloudDepthSampler)
+		SHADER_PARAMETER_RDG_TEXTURE(Texture2D, SrcCloudViewTexture)
+		SHADER_PARAMETER_SAMPLER(SamplerState, SrcCloudViewSampler)
 	END_SHADER_PARAMETER_STRUCT()
 
 	static bool ShouldCompilePermutation(const FGlobalShaderPermutationParameters& Parameters)
@@ -279,6 +284,12 @@ static FFogPassParameters* CreateDefaultFogPassParameters(
 	PassParameters->PS.UpsampleJitterMultiplier = CVarUpsampleJitterMultiplier.GetValueOnRenderThread() * GVolumetricFogGridPixelSize;
 	PassParameters->PS.bOnlyOnRenderedOpaque = View.bFogOnlyOnRenderedOpaque;
 	PassParameters->PS.bUseWaterDepthTexture = false;
+
+	PassParameters->PS.SrcCloudDepthTexture = GSystemTextures.GetWhiteDummy(GraphBuilder);
+	PassParameters->PS.SrcCloudDepthSampler = TStaticSamplerState<SF_Point, AM_Clamp, AM_Clamp, AM_Clamp>::GetRHI();
+	PassParameters->PS.SrcCloudViewTexture = GSystemTextures.GetWhiteDummy(GraphBuilder);
+	PassParameters->PS.SrcCloudViewSampler = TStaticSamplerState<SF_Point, AM_Clamp, AM_Clamp, AM_Clamp>::GetRHI();
+
 	return PassParameters;
 }
 
@@ -288,7 +299,9 @@ static void RenderViewFog(
 	FIntRect ViewRect, 
 	FFogPassParameters* PassParameters, 
 	bool bShouldRenderVolumetricFog,
-	bool bFogComposeLocalFogVolumes)
+	bool bFogComposeLocalFogVolumes,
+	bool bSampleFogOnClouds = false,
+	bool bEnableBlending = true)
 {
 	FGraphicsPipelineStateInitializer GraphicsPSOInit;
 	RHICmdList.ApplyCachedRenderTargets(GraphicsPSOInit);
@@ -298,8 +311,15 @@ static void RenderViewFog(
 	GraphicsPSOInit.DepthStencilState = TStaticDepthStencilState<false, CF_Always>::GetRHI();
 	GraphicsPSOInit.PrimitiveType = PT_TriangleList;
 
-	// disable alpha writes in order to preserve scene depth values on PC
-	GraphicsPSOInit.BlendState = TStaticBlendState<CW_RGB, BO_Add, BF_One, BF_SourceAlpha>::GetRHI();
+	if (bEnableBlending)
+	{
+		// disable alpha writes in order to preserve scene depth values on PC
+		GraphicsPSOInit.BlendState = TStaticBlendState<CW_RGB, BO_Add, BF_One, BF_SourceAlpha>::GetRHI();
+	}
+	else
+	{
+		GraphicsPSOInit.BlendState = TStaticBlendState<>::GetRHI();
+	}
 
 	TShaderMapRef<FHeightFogVS> VertexShader(View.ShaderMap);
 	GraphicsPSOInit.BoundShaderState.VertexDeclarationRHI = GFogVertexDeclaration.VertexDeclarationRHI;
@@ -311,11 +331,12 @@ static void RenderViewFog(
 	PsPermutationVector.Set<FExponentialHeightFogPS::FSupportFogDirectionalLightInScattering>(!bUseFogInscatteringColorCubemap && View.bUseDirectionalInscattering);
 	PsPermutationVector.Set<FExponentialHeightFogPS::FSupportVolumetricFog>(bShouldRenderVolumetricFog);
 	PsPermutationVector.Set<FExponentialHeightFogPS::FSupportLocalFogVolume>(bFogComposeLocalFogVolumes);
+	PsPermutationVector.Set<FExponentialHeightFogPS::FSampleFogOnClouds>(bSampleFogOnClouds);
 	TShaderMapRef<FExponentialHeightFogPS> PixelShader(View.ShaderMap, PsPermutationVector);
 	GraphicsPSOInit.BoundShaderState.PixelShaderRHI = PixelShader.GetPixelShader();
 
 	// Setup the depth bound optimization if possible on that platform.
-	GraphicsPSOInit.bDepthBounds = GSupportsDepthBoundsTest && CVarFogUseDepthBounds.GetValueOnAnyThread();
+	GraphicsPSOInit.bDepthBounds = GSupportsDepthBoundsTest && CVarFogUseDepthBounds.GetValueOnAnyThread() && !bSampleFogOnClouds;
 	if (GraphicsPSOInit.bDepthBounds)
 	{
 		float FogStartDistance = GetViewFogCommonStartDistance(View, bShouldRenderVolumetricFog, bFogComposeLocalFogVolumes);
@@ -360,6 +381,62 @@ static void RenderViewFog(
 	RHICmdList.DrawIndexedPrimitive(GTwoTrianglesIndexBuffer.IndexBufferRHI, 0, 0, 4, 0, 2, 1);
 }
 
+void RenderFogOnClouds(
+	FRDGBuilder& GraphBuilder,
+	const FScene* Scene,
+	const FViewInfo& View,
+	FRDGTextureRef SrcCloudDepth,
+	FRDGTextureRef SrcCloudView,
+	FRDGTextureRef DstCloudView,
+	const bool bShouldRenderVolumetricFog,
+	const bool bUseVolumetricRenderTarget)
+{
+	if (Scene->ExponentialFogs.Num() > 0)
+	{
+		RDG_EVENT_SCOPE(GraphBuilder, "ExponentialHeightFog on Clouds");
+		RDG_GPU_STAT_SCOPE(GraphBuilder, Fog);
+
+		if (View.IsPerspectiveProjection())
+		{
+			RDG_GPU_MASK_SCOPE(GraphBuilder, View.GPUMask);
+
+			TRDGUniformBufferRef<FFogUniformParameters> FogUniformBuffer = CreateFogUniformBuffer(GraphBuilder, View);
+
+			// Light shaft is not accounted for in this case
+			const FScreenPassTextureViewportParameters LightShaftParameters;
+			// Local fog volume are not accounted for in this case
+			const bool bFogComposeLocalFogVolumes = false;
+
+			FFogPassParameters* PassParameters = CreateDefaultFogPassParameters(
+				GraphBuilder, View, 
+				CreateSceneTextureUniformBuffer(GraphBuilder,View, ESceneTextureSetupMode::None),
+				FogUniformBuffer, nullptr /*LightShaftOcclusionTexture*/, LightShaftParameters);
+
+			// Patch the pass parameter for it to work on clouds
+			PassParameters->VS.ViewUniformBuffer = GetShaderBinding(bUseVolumetricRenderTarget ? View.VolumetricRenderTargetViewUniformBuffer : View.ViewUniformBuffer);
+			PassParameters->PS.ViewUniformBuffer = GetShaderBinding(bUseVolumetricRenderTarget ? View.VolumetricRenderTargetViewUniformBuffer : View.ViewUniformBuffer);
+			PassParameters->PS.bOnlyOnRenderedOpaque = false;
+
+			PassParameters->PS.SrcCloudDepthTexture = SrcCloudDepth;
+			PassParameters->PS.SrcCloudViewTexture = SrcCloudView;
+
+			PassParameters->RenderTargets[0] = FRenderTargetBinding(DstCloudView, ERenderTargetLoadAction::ENoAction);
+			// No depth target
+
+			// We enable the blending when volumetric render target is not enabled. Because in this case, the fog pass is compositing directly over the scene.
+			const bool bEnableBlending = !bUseVolumetricRenderTarget;
+
+			FIntRect ViewRect(0, 0, SrcCloudView->Desc.Extent.X, SrcCloudView->Desc.Extent.Y);
+			GraphBuilder.AddPass(RDG_EVENT_NAME("Fog"), PassParameters, ERDGPassFlags::Raster,
+				[&View, ViewRect, PassParameters, bShouldRenderVolumetricFog, bFogComposeLocalFogVolumes, bEnableBlending](FRHICommandList& RHICmdList)
+				{
+					const bool bSampleFogOnClouds = true;
+					RenderViewFog(RHICmdList, View, ViewRect, PassParameters, bShouldRenderVolumetricFog, bFogComposeLocalFogVolumes, bSampleFogOnClouds, bEnableBlending);
+				});
+		}
+	}
+}
+
 void FDeferredShadingSceneRenderer::RenderFog(
 	FRDGBuilder& GraphBuilder,
 	const FMinimalSceneTextures& SceneTextures,
diff --git a/Engine/Source/Runtime/Renderer/Private/FogRendering.h b/Engine/Source/Runtime/Renderer/Private/FogRendering.h
--- a/Engine/Source/Runtime/Renderer/Private/FogRendering.h
+++ b/Engine/Source/Runtime/Renderer/Private/FogRendering.h
@@ -10,6 +10,7 @@
 #include "RenderGraphFwd.h"
 
 class FSceneViewFamily;
+class FScene;
 class FViewInfo;
 
 BEGIN_GLOBAL_SHADER_PARAMETER_STRUCT(FFogUniformParameters,)
@@ -38,3 +39,13 @@ extern bool ShouldRenderFog(const FSceneViewFamily& Family);
 
 extern float GetViewFogCommonStartDistance(const FViewInfo& View, bool bShouldRenderVolumetricFog, bool bShouldRenderLocalFogVolumes);
 extern float GetFogDefaultStartDistance();
+
+extern void RenderFogOnClouds(
+	FRDGBuilder& GraphBuilder,
+	const FScene* Scene,
+	const FViewInfo& View,
+	FRDGTextureRef SrcCloudDepth,
+	FRDGTextureRef SrcCloudView,
+	FRDGTextureRef DstCloudView,
+	const bool bShouldRenderVolumetricFog,
+	const bool bUseVolumetricRenderTarget);
\ No newline at end of file
diff --git a/Engine/Source/Runtime/Renderer/Private/VolumetricCloudRendering.cpp b/Engine/Source/Runtime/Renderer/Private/VolumetricCloudRendering.cpp
--- a/Engine/Source/Runtime/Renderer/Private/VolumetricCloudRendering.cpp
+++ b/Engine/Source/Runtime/Renderer/Private/VolumetricCloudRendering.cpp
@@ -2058,6 +2058,7 @@ FCloudRenderContext::FCloudRenderContext()
 	bSkipAtmosphericLightShadowmap = false;
 
 	bSkipAerialPerspective = false;
+	bSkipHeightFog = false;
 
 	bAsyncCompute = false;
 	bCloudDebugViewModeEnabled = false;
@@ -2162,6 +2163,7 @@ static TRDGUniformBufferRef<FRenderVolumetricCloudGlobalParameters> CreateCloudP
 	}
 
 	VolumetricCloudParams.EnableAerialPerspectiveSampling = CloudRC.bSkipAerialPerspective ? 0 : 1;
+	VolumetricCloudParams.EnableHeightFog = VolumetricCloudParams.EnableHeightFog && !CloudRC.bSkipHeightFog ? 1 : 0;
 	VolumetricCloudParams.EnableDistantSkyLightSampling = CVarVolumetricCloudEnableDistantSkyLightSampling.GetValueOnAnyThread() > 0 ? 1 : 0;
 	VolumetricCloudParams.EnableAtmosphericLightsSampling = CVarVolumetricCloudEnableAtmosphericLightsSampling.GetValueOnAnyThread() > 0 ? 1 : 0;
 
@@ -2595,6 +2597,7 @@ bool FSceneRenderer::RenderVolumetricCloud(
 				ViewToProcess.bCloudDebugViewModeEnabled = bCloudDebugViewModeEnabled;
 			}
 
+			const bool bHeightFogInScene = Scene->ExponentialFogs.Num() > 0;
 			for (int32 ViewIndex = 0; ViewIndex < ViewsToProcess.Num(); ViewIndex++)
 			{
 				RDG_EVENT_SCOPE_CONDITIONAL(GraphBuilder, ViewsToProcess.Num() > 1, "View%d", ViewIndex);
@@ -2614,10 +2617,13 @@ bool FSceneRenderer::RenderVolumetricCloud(
 				CloudRC.ViewUniformBuffer = bShouldViewRenderVolumetricCloudRenderTarget ? ViewInfo.VolumetricRenderTargetViewUniformBuffer : ViewInfo.ViewUniformBuffer;
 
 				CloudRC.bSkipAerialPerspective = !bEnableAerialPerspectiveSampling || bShouldUseHighQualityAerialPerspective; // Skip AP on clouds if we are going to trace it separately in a second pass
+				CloudRC.bSkipHeightFog = bShouldUseHighQualityAerialPerspective;
 				CloudRC.bIsReflectionRendering = ViewInfo.bIsReflectionCapture;
 
 				FRDGTextureRef IntermediateRT = nullptr;
 				FRDGTextureRef IntermediateSecondaryRT = nullptr;
+				FRDGTextureRef IntermediateRT2 = nullptr;
+				FRDGTextureRef IntermediateSecondaryRT2 = nullptr;
 				FRDGTextureRef DestinationRT = nullptr;
 				FRDGTextureRef DestinationSecondaryRT = nullptr;
 				FRDGTextureRef DestinationRTDepth = nullptr;
@@ -2640,6 +2646,16 @@ bool FSceneRenderer::RenderVolumetricCloud(
 						IntermediateSecondaryRT = GraphBuilder.CreateTexture(
 							FRDGTextureDesc::Create2D(IntermadiateTargetResolution, PF_FloatRGBA, FClearValueBinding(FLinearColor(63000.0f, 63000.0f, 63000.0f, 63000.0f)),
 								TexCreate_ShaderResource | TexCreate_RenderTargetable | TexCreate_UAV), TEXT("Cloud.HighQualityAPIntermediateSecondary"));
+
+						if (bHeightFogInScene)
+						{
+							IntermediateRT2 = GraphBuilder.CreateTexture(
+								FRDGTextureDesc::Create2D(IntermadiateTargetResolution, PF_FloatRGBA, FClearValueBinding(FLinearColor(63000.0f, 63000.0f, 63000.0f, 63000.0f)),
+									TexCreate_ShaderResource | TexCreate_RenderTargetable | TexCreate_UAV), TEXT("Cloud.HighQualityAPIntermediate2"));
+							IntermediateSecondaryRT2 = GraphBuilder.CreateTexture(
+								FRDGTextureDesc::Create2D(IntermadiateTargetResolution, PF_FloatRGBA, FClearValueBinding(FLinearColor(63000.0f, 63000.0f, 63000.0f, 63000.0f)),
+									TexCreate_ShaderResource | TexCreate_RenderTargetable | TexCreate_UAV), TEXT("Cloud.HighQualityAPIntermediateSecondary2"));
+						}
 					}
 
 					// No action because we only need to render volumetric clouds so we do not blend in that render target.
@@ -2676,6 +2692,13 @@ bool FSceneRenderer::RenderVolumetricCloud(
 						IntermediateRT = GraphBuilder.CreateTexture(
 							FRDGTextureDesc::Create2D(IntermadiateTargetResolution, PF_FloatRGBA, FClearValueBinding(FLinearColor(0.0f, 0.0f, 0.0f, 1.0f)),
 								TexCreate_ShaderResource | TexCreate_RenderTargetable | TexCreate_UAV), TEXT("Cloud.HighQualityAPIntermediate"));
+
+						if (bHeightFogInScene)
+						{
+							IntermediateRT2 = GraphBuilder.CreateTexture(
+								FRDGTextureDesc::Create2D(IntermadiateTargetResolution, PF_FloatRGBA, FClearValueBinding(FLinearColor(63000.0f, 63000.0f, 63000.0f, 63000.0f)),
+									TexCreate_ShaderResource | TexCreate_RenderTargetable | TexCreate_UAV), TEXT("Cloud.HighQualityAPIntermediate2"));
+						}
 					}
 
 					// Create a texture for cloud depth data that will be dropped out just after.
@@ -2759,7 +2782,9 @@ bool FSceneRenderer::RenderVolumetricCloud(
 					SkyRC.bUseDepthBoundTestIfPossible = false;
 					SkyRC.bForceRayMarching = true;				// We do not have any valid view LUT
 					SkyRC.bDepthReadDisabled = true;
-					SkyRC.bDisableBlending = bShouldViewRenderVolumetricCloudRenderTarget ? true : false;
+
+					// Disable blending if volumetric render target is used, or if height fog is present (meaning we are going to render in a separate render target before blending over the scene)
+					SkyRC.bDisableBlending = (bShouldViewRenderVolumetricCloudRenderTarget || bHeightFogInScene) ? true : false;
 
 					SkyRC.TransmittanceLut = GraphBuilder.RegisterExternalTexture(SkyInfo.GetTransmittanceLutTexture());
 					SkyRC.MultiScatteredLuminanceLut = GraphBuilder.RegisterExternalTexture(SkyInfo.GetMultiScatteredLuminanceLutTexture());
@@ -2768,7 +2793,7 @@ bool FSceneRenderer::RenderVolumetricCloud(
 					SkyRC.bAPOnCloudMode = true;
 					SkyRC.VolumetricCloudDepthTexture = DestinationRTDepth;
 					SkyRC.InputCloudLuminanceTransmittanceTexture = IntermediateRT;
-					SkyRC.RenderTargets[0] = FRenderTargetBinding(DestinationRT, ERenderTargetLoadAction::ENoAction);
+					SkyRC.RenderTargets[0] = FRenderTargetBinding(bHeightFogInScene ? IntermediateRT2 : DestinationRT, ERenderTargetLoadAction::ENoAction);
 
 					SkyRC.ViewMatrices = &ViewInfo.ViewMatrices;
 					SkyRC.ViewUniformBuffer = bShouldViewRenderVolumetricCloudRenderTarget ? ViewInfo.VolumetricRenderTargetViewUniformBuffer : ViewInfo.ViewUniformBuffer;
@@ -2816,9 +2841,45 @@ bool FSceneRenderer::RenderVolumetricCloud(
 						// We need to execute a second tracing for the far depth cloud result in order to correctly fix up edges.
 						SkyRC.VolumetricCloudDepthTexture = DestinationRTDepth;
 						SkyRC.InputCloudLuminanceTransmittanceTexture = IntermediateSecondaryRT;
-						SkyRC.RenderTargets[0] = FRenderTargetBinding(DestinationSecondaryRT, ERenderTargetLoadAction::ENoAction);
+						SkyRC.RenderTargets[0] = FRenderTargetBinding(bHeightFogInScene ? IntermediateSecondaryRT2 : DestinationSecondaryRT, ERenderTargetLoadAction::ENoAction);
 						RenderSkyAtmosphereInternal(GraphBuilder, GetSceneTextureShaderParameters(SceneTextures.UniformBuffer), SkyRC);
 					}
+
+					// Also render fog on top of AP
+					if (bHeightFogInScene)
+					{
+						const bool bShouldRenderVolumetricFog = ShouldRenderVolumetricFog();
+
+						FRDGTextureRef SrcCloudDepth = DestinationRTDepth;
+						FRDGTextureRef SrcCloudView  = IntermediateRT2;
+						FRDGTextureRef DstCloudView  = DestinationRT;
+						RenderFogOnClouds(
+							GraphBuilder,
+							Scene,
+							ViewInfo, 
+							SrcCloudDepth,
+							SrcCloudView,
+							DstCloudView,
+							bShouldRenderVolumetricFog,
+							bShouldViewRenderVolumetricCloudRenderTarget);
+
+						if (DestinationSecondaryRT)
+						{
+							// We need to execute a second tracing for the far depth cloud result in order to correctly fix up edges.
+							SrcCloudDepth = DestinationRTDepth;
+							SrcCloudView = IntermediateSecondaryRT2;
+							DstCloudView = DestinationSecondaryRT;
+							RenderFogOnClouds(
+								GraphBuilder,
+								Scene,
+								ViewInfo,
+								SrcCloudDepth,
+								SrcCloudView,
+								DstCloudView,
+								bShouldRenderVolumetricFog,
+								bShouldViewRenderVolumetricCloudRenderTarget);
+						}
+					}
 				}
 
 				if (DebugCloudShadowMap || DebugCloudSkyAO)
diff --git a/Engine/Source/Runtime/Renderer/Private/VolumetricCloudRendering.h b/Engine/Source/Runtime/Renderer/Private/VolumetricCloudRendering.h
--- a/Engine/Source/Runtime/Renderer/Private/VolumetricCloudRendering.h
+++ b/Engine/Source/Runtime/Renderer/Private/VolumetricCloudRendering.h
@@ -138,6 +138,7 @@ struct FCloudRenderContext
 
 	bool bShouldViewRenderVolumetricRenderTarget;
 	bool bSkipAerialPerspective;
+	bool bSkipHeightFog;
 	bool bIsReflectionRendering;				// Reflection capture and real time sky capture
 	bool bIsSkyRealTimeReflectionRendering;		// Real time sky capture only
 	bool bSkipAtmosphericLightShadowmap;
