Mac: Fixes to enable groom on M2+ devices.

Addresses a hang that was occuring. This is leads to a reduction in hairless nightmare fuel.

diff --git a/Engine/Config/Mac/DataDrivenPlatformInfo.ini b/Engine/Config/Mac/DataDrivenPlatformInfo.ini
--- a/Engine/Config/Mac/DataDrivenPlatformInfo.ini
+++ b/Engine/Config/Mac/DataDrivenPlatformInfo.ini
@@ -57,7 +57,7 @@ bSupportsRayTracingShaders = false
 bSupportsInlineRayTracing = false
 bSupportsRayTracingIndirectInstanceData = false
 bSupportsByteBufferComputeShaders = true
-bSupportsHairStrandGeometry = false
+bSupportsHairStrandGeometry = true
 bSupportsNanite=false
 bSupportsUInt64ImageAtomics=false
 bSupportsNNEShaders=true
@@ -97,7 +97,7 @@ bSupportsRayTracingShaders = false
 bSupportsInlineRayTracing = false
 bSupportsRayTracingIndirectInstanceData = false
 bSupportsByteBufferComputeShaders = true
-bSupportsHairStrandGeometry = false
+bSupportsHairStrandGeometry = true
 bSupportsNanite=true
 bSupportsUInt64ImageAtomics=true
 bSupportsNNEShaders=true
diff --git a/Engine/Config/Mac/MacEngine.ini b/Engine/Config/Mac/MacEngine.ini
--- a/Engine/Config/Mac/MacEngine.ini
+++ b/Engine/Config/Mac/MacEngine.ini
@@ -42,4 +42,4 @@ PlatformRequiresDataCrypto=True
 MaxShaderJobBatchSize=25
 
 [ConsoleVariables]
-r.HairStrands.Strands=0
+;r.HairStrands.Strands=0
diff --git a/Engine/Source/Runtime/Renderer/Private/HairStrands/HairStrandsVisibility.cpp b/Engine/Source/Runtime/Renderer/Private/HairStrands/HairStrandsVisibility.cpp
--- a/Engine/Source/Runtime/Renderer/Private/HairStrands/HairStrandsVisibility.cpp
+++ b/Engine/Source/Runtime/Renderer/Private/HairStrands/HairStrandsVisibility.cpp
@@ -1826,7 +1826,7 @@ static void AddHairVisibilityControlPointIdCompactionPass(
 	}
 
 	{
-		FRDGTextureDesc Desc = FRDGTextureDesc::Create2D(FIntPoint(1, 1), PF_R32_UINT, FClearValueBinding::None, TexCreate_UAV | TexCreate_ShaderResource);
+		FRDGTextureDesc Desc = FRDGTextureDesc::Create2D(FIntPoint(1, 1), PF_R32_UINT, FClearValueBinding::None, TexCreate_UAV | TexCreate_ShaderResource | TexCreate_AtomicCompatible);
 		OutCompactCounter = GraphBuilder.CreateTexture(Desc, TEXT("Hair.VisibilityCompactCounter"));
 	}
 
@@ -1971,7 +1971,7 @@ static void AddHairVisibilityCompactionComputeRasterPass(
 	FIntPoint Resolution = RasterComputeData.DepthCovTexture->Desc.Extent;
 
 	{
-		FRDGTextureDesc Desc = FRDGTextureDesc::Create2D(FIntPoint(1,1), PF_R32_UINT, FClearValueBinding::None, TexCreate_UAV);
+		FRDGTextureDesc Desc = FRDGTextureDesc::Create2D(FIntPoint(1,1), PF_R32_UINT, FClearValueBinding::None, TexCreate_UAV | TexCreate_AtomicCompatible);
 		OutCompactCounter = GraphBuilder.CreateTexture(Desc, TEXT("Hair.VisibilityCompactCounter"));
 	}
 
