* Add hair transient resources struct
* Regroup group aabb into single resource and batch some passes

diff --git a/Engine/Plugins/Runtime/HairStrands/Source/HairStrandsCore/Private/GroomDebug.cpp b/Engine/Plugins/Runtime/HairStrands/Source/HairStrandsCore/Private/GroomDebug.cpp
--- a/Engine/Plugins/Runtime/HairStrands/Source/HairStrandsCore/Private/GroomDebug.cpp
+++ b/Engine/Plugins/Runtime/HairStrands/Source/HairStrandsCore/Private/GroomDebug.cpp
@@ -409,7 +409,8 @@ class FDrawDebugClusterAABBCS : public FGlobalShader
 	BEGIN_SHADER_PARAMETER_STRUCT(FParameters,)
 		SHADER_PARAMETER_STRUCT_REF(FViewUniformShaderParameters, ViewUniformBuffer)
 		SHADER_PARAMETER_SRV(Buffer, ClusterAABBBuffer)
-		SHADER_PARAMETER_SRV(Buffer, GroupAABBBuffer)
+		SHADER_PARAMETER_RDG_BUFFER_SRV(Buffer, GroupAABBBuffer)
+		SHADER_PARAMETER(uint32, InstanceRegisteredIndex)
 		SHADER_PARAMETER(uint32, ClusterCount)
 		SHADER_PARAMETER(uint32, PointCount)
 		SHADER_PARAMETER(uint32, CurveCount)
@@ -437,6 +438,7 @@ void AddDrawDebugClusterPass(
 	FRDGBuilder& GraphBuilder,
 	const FSceneView& View,
 	FGlobalShaderMap* ShaderMap,
+	FHairTransientResources& TransientResources,
 	const FShaderPrintData* ShaderPrintData,
 	EGroomViewMode ViewMode,
 	FHairStrandClusterData& HairClusterData)
@@ -459,6 +461,7 @@ void AddDrawDebugClusterPass(
 		TShaderMapRef<FDrawDebugClusterAABBCS> ComputeShader(ShaderMap);
 
 		FDrawDebugClusterAABBCS::FParameters* Parameters = GraphBuilder.AllocParameters<FDrawDebugClusterAABBCS::FParameters>();
+		Parameters->InstanceRegisteredIndex = HairGroupClusters.InstanceRegisteredIndex;
 		Parameters->ViewUniformBuffer = View.ViewUniformBuffer;
 		Parameters->ClusterCount = HairGroupClusters.ClusterCount;
 		Parameters->PointCount = HairGroupClusters.HairGroupPublicPtr->GetActiveStrandsPointCount();
@@ -466,7 +469,7 @@ void AddDrawDebugClusterPass(
 		Parameters->HairGroupId = DataIndex++;
 		Parameters->bDrawAABB = ViewMode == EGroomViewMode::ClusterAABB ? 1 : 0;
 		Parameters->ClusterAABBBuffer = HairGroupClusters.ClusterAABBBuffer->SRV;
-		Parameters->GroupAABBBuffer = HairGroupClusters.GroupAABBBuffer->SRV;
+		Parameters->GroupAABBBuffer = TransientResources.GroupAABBSRV;
 		ShaderPrint::SetParameters(GraphBuilder, *ShaderPrintData, Parameters->ShaderPrintParameters);
 
 		const FIntVector DispatchCount = DispatchCount.DivideAndRoundUp(FIntVector(Parameters->ClusterCount, 1, 1), FIntVector(FDrawDebugClusterAABBCS::GetGroupSize(), 1, 1));
@@ -840,7 +843,7 @@ class FHairDebugPrintInstanceCS : public FGlobalShader
 
 	BEGIN_SHADER_PARAMETER_STRUCT(FParameters, )
 		SHADER_PARAMETER(uint32, InstanceCount)
-
+		SHADER_PARAMETER(uint32, InstanceRegisteredIndex)
 		SHADER_PARAMETER(uint32, InstanceCount_StrandsPrimaryView)
 		SHADER_PARAMETER(uint32, InstanceCount_StrandsShadowView)
 		SHADER_PARAMETER(uint32, InstanceCount_CardsOrMeshesPrimaryView)
@@ -877,6 +880,7 @@ static void AddHairDebugPrintInstancePass(
 	FRDGBuilder& GraphBuilder, 
 	FGlobalShaderMap* ShaderMap,
 	const FSceneView& View,
+	FHairTransientResources& TransientResources,
 	const FShaderPrintData* ShaderPrintData,
 	const FHairStrandsInstances& Instances,
 	const TArray<EHairInstanceVisibilityType>& InstancesVisibilityType)
@@ -1093,7 +1097,8 @@ static void AddHairDebugPrintInstancePass(
 			const float ContinousLODRadius = Instance->HairGroupPublicData->ContinuousLODScreenSize * MaxRectSizeInPixels * 0.5f; // Diameter->Radius
 
 			FHairDebugPrintInstanceCS::FParameters* Parameters = GraphBuilder.AllocParameters<FHairDebugPrintInstanceCS::FParameters>();
-			Parameters->InstanceAABB = Register(GraphBuilder, Instance->HairGroupPublicData->GetGroupAABBBuffer(), ERDGImportedBufferFlags::CreateSRV).SRV;
+			Parameters->InstanceRegisteredIndex = Instance->RegisteredIndex;
+			Parameters->InstanceAABB = TransientResources.GroupAABBSRV;
 			Parameters->InstanceScreenSphereBound = FVector4f(Instance->HairGroupPublicData->ContinuousLODScreenPos.X, Instance->HairGroupPublicData->ContinuousLODScreenPos.Y, 0.f, ContinousLODRadius);
 			ShaderPrint::SetParameters(GraphBuilder, *ShaderPrintData, Parameters->ShaderPrintUniformBuffer);
 			ClearUnusedGraphResources(ComputeShader, Parameters);
@@ -1369,6 +1374,7 @@ void RunHairStrandsDebug(
 	const FSceneView& View,
 	const FHairStrandsInstances& Instances,
 	const TArray<EHairInstanceVisibilityType>& InstancesVisibilityType,
+	FHairTransientResources& TransientResources,
 	const FShaderPrintData* ShaderPrintData,
 	FRDGTextureRef SceneColorTexture,
 	FRDGTextureRef SceneDepthTexture,
@@ -1379,7 +1385,7 @@ void RunHairStrandsDebug(
 
 	if (ViewMode == EGroomViewMode::MacroGroups)
 	{
-		AddHairDebugPrintInstancePass(GraphBuilder, ShaderMap, View, ShaderPrintData, Instances, InstancesVisibilityType);
+		AddHairDebugPrintInstancePass(GraphBuilder, ShaderMap, View, TransientResources, ShaderPrintData, Instances, InstancesVisibilityType);
 	}
 
 	if (ViewMode == EGroomViewMode::MeshProjection)
diff --git a/Engine/Plugins/Runtime/HairStrands/Source/HairStrandsCore/Private/GroomManager.cpp b/Engine/Plugins/Runtime/HairStrands/Source/HairStrandsCore/Private/GroomManager.cpp
--- a/Engine/Plugins/Runtime/HairStrands/Source/HairStrandsCore/Private/GroomManager.cpp
+++ b/Engine/Plugins/Runtime/HairStrands/Source/HairStrandsCore/Private/GroomManager.cpp
@@ -82,7 +82,8 @@ void AddHairClusterAABBPass(
 	FHairStrandClusterData::FHairGroup* ClusterData,
 	FRDGHairStrandsCullingData& ClusterAABBData,
 	FRDGBufferSRVRef RenderPositionBufferSRV,
-	FRDGBufferSRVRef& DrawIndirectRasterComputeBuffer);
+	FRDGBufferSRVRef& DrawIndirectRasterComputeBuffer,
+	FRDGBufferUAVRef GroupAABBUAV);
 
 ///////////////////////////////////////////////////////////////////////////////////////////////////
 // Utils
@@ -667,6 +668,7 @@ void AddDrawDebugClusterPass(
 	FRDGBuilder& GraphBuilder,
 	const FSceneView& View,
 	FGlobalShaderMap* ShaderMap,
+	FHairTransientResources& TransientResources,
 	const FShaderPrintData* ShaderPrintData,
 	EGroomViewMode ViewMode,
 	FHairStrandClusterData& HairClusterData);
@@ -679,6 +681,7 @@ static void RunHairStrandsInterpolation_Strands(
 	const uint32 ViewUniqueID,
 	const FHairStrandsInstances& Instances,
 	const FShaderPrintData* ShaderPrintData,
+	FHairTransientResources& TransientResources,
 	FGlobalShaderMap* ShaderMap)
 {
 	if (Instances.IsEmpty()) { return; }
@@ -709,6 +712,7 @@ static void RunHairStrandsInterpolation_Strands(
 
 	struct FInstanceData
 	{
+		uint32 RegisteredIndex = ~0;
 		int32 HairLODIndex = -1;
 		int32 MeshLODIndex = -1;
 		uint32 ActivePointCount = 0;
@@ -739,7 +743,6 @@ static void RunHairStrandsInterpolation_Strands(
 		}
 	};
 
-
 	// Gather all strands instances
 	const uint32 ViewRayTracingMask = RHI_RAYTRACING ? (View->Family->EngineShowFlags.PathTracing ? EHairViewRayTracingMask::PathTracing : EHairViewRayTracingMask::RayTracing) : 0u;
 	bool bHasAnySimCacheInstances = false;
@@ -757,6 +760,7 @@ static void RunHairStrandsInterpolation_Strands(
 		check(Instance->HairGroupPublicData);
 
 		FInstanceData& InstanceData = InstanceDatas.AddDefaulted_GetRef();
+		InstanceData.RegisteredIndex			= Instance->RegisteredIndex;
 		InstanceData.Instance 					= Instance;
 		InstanceData.ActivePointCount 			= Instance->HairGroupPublicData->GetActiveStrandsPointCount();
 		InstanceData.ActiveCurveCount 			= Instance->HairGroupPublicData->GetActiveStrandsCurveCount();
@@ -869,13 +873,14 @@ static void RunHairStrandsInterpolation_Strands(
 				ShaderMap,
 				Views[0],
 				ShaderPrintData,
-				ClusterDatas);
+				ClusterDatas,
+				TransientResources.IndirectDispatchArgsUAV);
 		}
 
 		// Run cluster debug view here (instead of GroomDebug.h/.cpp, as we need to have the (transient) cluster data 
 		if (ViewMode == EGroomViewMode::Cluster || ViewMode == EGroomViewMode::ClusterAABB)
 		{
-			AddDrawDebugClusterPass(GraphBuilder, *View, ShaderMap, ShaderPrintData, ViewMode, ClusterDatas);
+			AddDrawDebugClusterPass(GraphBuilder, *View, ShaderMap, TransientResources, ShaderPrintData, ViewMode, ClusterDatas);
 		}
 	}
 
@@ -955,17 +960,15 @@ static void RunHairStrandsInterpolation_Strands(
 	// Clear cluster AABBs (used optionally  for voxel allocation & for culling)
 	for (FInstanceData& InstanceData : InstanceDatas)
 	{
-		//if (InstanceData.bNeedDeformation)
-		{
-			InstanceData.CullingData = ImportCullingData(GraphBuilder, InstanceData.Instance->HairGroupPublicData);
-			AddClearClusterAABBPass(
-				GraphBuilder,
-				ShaderMap,
-				InstanceData.bNeedDeformation ? EHairAABBUpdateType::UpdateClusterAABB : EHairAABBUpdateType::UpdateGroupAABB,
-				InstanceData.CullingData.ClusterCount,
-				InstanceData.CullingData.ClusterAABBBuffer,
-				InstanceData.CullingData.GroupAABBBuffer);
-		}
+		InstanceData.CullingData = ImportCullingData(GraphBuilder, InstanceData.Instance->HairGroupPublicData);
+		AddClearClusterAABBPass(
+			GraphBuilder,
+			ShaderMap,
+			InstanceData.bNeedDeformation ? EHairAABBUpdateType::UpdateClusterAABB : EHairAABBUpdateType::UpdateGroupAABB,
+			InstanceData.RegisteredIndex,
+			InstanceData.CullingData.ClusterCount,
+			InstanceData.CullingData.ClusterAABBBuffer,
+			TransientResources.GroupAABBUAV);
 	}
 
 	for (FInstanceData& InstanceData : InstanceDatas)
@@ -1169,7 +1172,6 @@ static void RunHairStrandsInterpolation_Strands(
 		}
 			
 		InstanceData.Instance->HairGroupPublicData->SetClusterAABBValid(false);
-		InstanceData.Instance->HairGroupPublicData->SetGroupAABBValid(false);
 			
 		if (bNeedGPUAABB)
 		{
@@ -1186,10 +1188,12 @@ static void RunHairStrandsInterpolation_Strands(
 				HairGroupCluster,
 				InstanceData.CullingData,
 				InstanceData.RDGResources.PositionSRV,
-				Strands_CulledVertexCount.SRV);
+				Strands_CulledVertexCount.SRV,
+				TransientResources.GroupAABBUAV);
 			
 			InstanceData.Instance->HairGroupPublicData->SetClusterAABBValid(UpdateType == EHairAABBUpdateType::UpdateClusterAABB);
-			InstanceData.Instance->HairGroupPublicData->SetGroupAABBValid(true);
+
+			TransientResources.bIsGroupAABBValid[InstanceData.RegisteredIndex] = true;
 		}
 	}
 
@@ -1988,6 +1992,7 @@ void RunHairStrandsDebug(
 	const FSceneView& View,
 	const FHairStrandsInstances& Instances,
 	const TArray<EHairInstanceVisibilityType>& InstancesVisibilityType,
+	FHairTransientResources& TransientResources,
 	const FShaderPrintData* ShaderPrintData,
 	FRDGTextureRef SceneColor,
 	FRDGTextureRef SceneDepth,
@@ -2101,6 +2106,7 @@ void ProcessHairStrandsBookmark(
 	else if (Bookmark == EHairStrandsBookmark::ProcessStrandsInterpolation)
 	{
 		check(GraphBuilder);
+		check(Parameters.TransientResources);
 		RunHairStrandsInterpolation_Strands(
 			*GraphBuilder,
 			Parameters.Scene,
@@ -2109,6 +2115,7 @@ void ProcessHairStrandsBookmark(
 			Parameters.ViewUniqueID,
 			*Instances,
 			Parameters.ShaderPrintData,
+			*Parameters.TransientResources,
 			Parameters.ShaderMap);
 	}
 	else if (Bookmark == EHairStrandsBookmark::ProcessDebug)
@@ -2128,6 +2135,7 @@ void ProcessHairStrandsBookmark(
 		}
 
 		check(GraphBuilder);
+		check(Parameters.TransientResources);
 		RunHairStrandsDebug(
 			*GraphBuilder,
 			Parameters.ShaderMap,
@@ -2135,6 +2143,7 @@ void ProcessHairStrandsBookmark(
 			*Parameters.View,
 			*Instances,
 			Parameters.InstancesVisibilityType,
+			*Parameters.TransientResources,
 			Parameters.ShaderPrintData,
 			Parameters.SceneColorTexture,
 			Parameters.SceneDepthTexture,
diff --git a/Engine/Plugins/Runtime/HairStrands/Source/HairStrandsCore/Private/GroomResources.cpp b/Engine/Plugins/Runtime/HairStrands/Source/HairStrandsCore/Private/GroomResources.cpp
--- a/Engine/Plugins/Runtime/HairStrands/Source/HairStrandsCore/Private/GroomResources.cpp
+++ b/Engine/Plugins/Runtime/HairStrands/Source/HairStrandsCore/Private/GroomResources.cpp
@@ -1310,7 +1310,6 @@ void FHairStrandsCullingResource::InternalAllocate(FRDGBuilder& GraphBuilder, ui
 		if (Resources.DrawIndirectBuffer.Buffer == nullptr)
 		{
 			InternalCreateVertexBufferRDG<FHairStrandsSintFormat>(GraphBuilder, ClusterCount * 6, Resources.ClusterAABBBuffer, TEXT("Hair.Cluster_ClusterAABBBuffer"), GetOwnerName(), EHairResourceUsageType::Dynamic);
-			InternalCreateVertexBufferRDG<FHairStrandsSintFormat>(GraphBuilder, 6, Resources.GroupAABBBuffer, TEXT("Hair.Cluster_GroupAABBBuffer"), GetOwnerName(), EHairResourceUsageType::Dynamic);
 			InternalCreateIndirectBufferRDG(GraphBuilder, Resources.DrawIndirectBuffer, TEXT("Hair.Cluster_DrawIndirectBuffer"), GetOwnerName());
 			InternalCreateIndirectBufferRDG(GraphBuilder, Resources.DrawIndirectRasterComputeBuffer, TEXT("Hair.Cluster_DrawIndirectRasterComputeBuffer"), GetOwnerName());
 			GraphBuilder.SetBufferAccessFinal(Register(GraphBuilder, Resources.DrawIndirectBuffer, ERDGImportedBufferFlags::None).Buffer, ERHIAccess::IndirectArgs);
@@ -1323,7 +1322,6 @@ void FHairStrandsCullingResource::InternalRelease()
 	Resources.DrawIndirectBuffer.Release();
 	Resources.DrawIndirectRasterComputeBuffer.Release();
 	Resources.ClusterAABBBuffer.Release();
-	Resources.GroupAABBBuffer.Release();
 	Resources.CulledCurveBuffer.Release();
 	Resources.CulledVertexIdBuffer.Release();
 	Resources.CulledVertexRadiusScaleBuffer.Release();
diff --git a/Engine/Plugins/Runtime/HairStrands/Source/HairStrandsCore/Private/HairStrandsClusterCulling.cpp b/Engine/Plugins/Runtime/HairStrands/Source/HairStrandsCore/Private/HairStrandsClusterCulling.cpp
--- a/Engine/Plugins/Runtime/HairStrands/Source/HairStrandsCore/Private/HairStrandsClusterCulling.cpp
+++ b/Engine/Plugins/Runtime/HairStrands/Source/HairStrandsCore/Private/HairStrandsClusterCulling.cpp
@@ -64,10 +64,12 @@ class FHairClusterCullArgsCS: public FGlobalShader
 	SHADER_USE_PARAMETER_STRUCT(FHairClusterCullArgsCS, FGlobalShader);
 
 	BEGIN_SHADER_PARAMETER_STRUCT(FParameters,)
+		SHADER_PARAMETER(uint32, InstanceRegisteredIndex)
 		SHADER_PARAMETER(uint32, ClusterGroupIndex)
 		SHADER_PARAMETER_RDG_BUFFER_SRV(StructuredBuffer<uint>, PointCounterBuffer)
 		SHADER_PARAMETER_RDG_BUFFER_UAV(RWBuffer<uint>, RWIndirectDrawArgsBuffer)
 		SHADER_PARAMETER_RDG_BUFFER_UAV(RWBuffer<uint>, RWIndirectDispatchArgsBuffer)
+		SHADER_PARAMETER_RDG_BUFFER_UAV(RWBuffer<uint>, RWIndirectDispatchArgsGlobalBuffer)
 		END_SHADER_PARAMETER_STRUCT()
 
 public:
@@ -87,7 +89,8 @@ void AddClusterCullingPass(
 	FGlobalShaderMap* ShaderMap,
 	const FSceneView* View,
 	const FShaderPrintData* ShaderPrintData,
-	FHairStrandClusterData& ClusterDatas)
+	FHairStrandClusterData& ClusterDatas,
+	FRDGBufferUAVRef IndirectDispatchArgsGlobalUAV)
 {
 	check(View);
 
@@ -165,10 +168,12 @@ void AddClusterCullingPass(
 		FRDGImportedBuffer DrawIndirectParametersRasterComputeBuffer = Register(GraphBuilder, ClusterData.HairGroupPublicPtr->GetDrawIndirectRasterComputeBuffer(), ERDGImportedBufferFlags::CreateViews);
 
 		FHairClusterCullArgsCS::FParameters* Parameters = GraphBuilder.AllocParameters<FHairClusterCullArgsCS::FParameters>();
+		Parameters->InstanceRegisteredIndex = ClusterData.InstanceRegisteredIndex;
 		Parameters->ClusterGroupIndex = ClusterGroupIt++;
 		Parameters->PointCounterBuffer = PointCounterSRV;
 		Parameters->RWIndirectDrawArgsBuffer = DrawIndirectParametersBuffer.UAV;
 		Parameters->RWIndirectDispatchArgsBuffer = DrawIndirectParametersRasterComputeBuffer.UAV;
+		Parameters->RWIndirectDispatchArgsGlobalBuffer = IndirectDispatchArgsGlobalUAV;
 
 		TShaderMapRef<FHairClusterCullArgsCS> ComputeShader(ShaderMap);
 		FComputeShaderUtils::AddPass(
@@ -182,7 +187,6 @@ void AddClusterCullingPass(
 		GraphBuilder.SetBufferAccessFinal(DrawIndirectParametersRasterComputeBuffer.Buffer, ERHIAccess::IndirectArgs | ERHIAccess::SRVMask);
 		ClusterData.SetCullingResultAvailable(true);
 	}
-
 }
 
 ///////////////////////////////////////////////////////////////////////////////////////////////////
@@ -195,7 +199,6 @@ void AddInstanceToClusterData(FHairGroupInstance* In, FHairStrandClusterData& Ou
 	HairGroupCluster.MaxPointPerCurve = In->Strands.Data ? In->Strands.Data->Header.MaxPointPerCurve : 0;
 	HairGroupCluster.ClusterScale = In->HairGroupPublicData->GetClusterScale();
 	HairGroupCluster.ClusterCount = In->HairGroupPublicData->GetClusterCount();
-	HairGroupCluster.GroupAABBBuffer = &In->HairGroupPublicData->GetGroupAABBBuffer();
 	HairGroupCluster.ClusterAABBBuffer = &In->HairGroupPublicData->GetClusterAABBBuffer();
 
 	HairGroupCluster.CurveBuffer = &In->Strands.RestResource->CurveBuffer;
@@ -205,6 +208,7 @@ void AddInstanceToClusterData(FHairGroupInstance* In, FHairStrandClusterData& Ou
 	HairGroupCluster.ClusterInfoBuffer = &In->Strands.ClusterResource->ClusterInfoBuffer;
 	HairGroupCluster.CurveToClusterIdBuffer = &In->Strands.ClusterResource->CurveToClusterIdBuffer;
 
+	HairGroupCluster.InstanceRegisteredIndex = In->RegisteredIndex;
 	HairGroupCluster.HairGroupPublicPtr = In->HairGroupPublicData;
 	HairGroupCluster.LODBias  = In->HairGroupPublicData->GetLODBias();
 	HairGroupCluster.LODIndex = In->HairGroupPublicData->GetLODIndex();
diff --git a/Engine/Plugins/Runtime/HairStrands/Source/HairStrandsCore/Private/HairStrandsClusterCulling.h b/Engine/Plugins/Runtime/HairStrands/Source/HairStrandsCore/Private/HairStrandsClusterCulling.h
--- a/Engine/Plugins/Runtime/HairStrands/Source/HairStrandsCore/Private/HairStrandsClusterCulling.h
+++ b/Engine/Plugins/Runtime/HairStrands/Source/HairStrandsCore/Private/HairStrandsClusterCulling.h
@@ -15,6 +15,7 @@ struct FHairStrandClusterData
 {
 	struct FHairGroup
 	{
+		uint32 InstanceRegisteredIndex = ~0;
 		uint32 ClusterCount = 0;
 		float ClusterScale = 0;
 		uint32 MaxPointPerCurve = 0;
@@ -27,7 +28,6 @@ struct FHairStrandClusterData
 		FRDGExternalBuffer* PointLODBuffer = nullptr;
 
 		// See FHairStrandsClusterResource fro details about those buffers.
-		FRDGExternalBuffer* GroupAABBBuffer = nullptr;
 		FRDGExternalBuffer* ClusterAABBBuffer = nullptr;
 		FRDGExternalBuffer* ClusterInfoBuffer = nullptr; // SRV
 		FRDGExternalBuffer* CurveToClusterIdBuffer = nullptr; // SRV
@@ -47,6 +47,7 @@ struct FHairStrandClusterData
 	};
 
 	TArray<FHairGroup> HairGroups;
+	FRDGBufferRef SharedIndirectDispatchArgsBuffer = nullptr;
 };
 
 void AddInstanceToClusterData(
@@ -58,4 +59,5 @@ void AddClusterCullingPass(
 	FGlobalShaderMap* ShaderMap,
 	const FSceneView* View,
 	const FShaderPrintData* ShaderPrintData,
-	FHairStrandClusterData& ClusterDatas);
+	FHairStrandClusterData& ClusterDatas,
+	FRDGBufferUAVRef IndirectDispatchArgsGlobalUAV);
diff --git a/Engine/Plugins/Runtime/HairStrands/Source/HairStrandsCore/Private/HairStrandsInterpolation.cpp b/Engine/Plugins/Runtime/HairStrands/Source/HairStrandsCore/Private/HairStrandsInterpolation.cpp
--- a/Engine/Plugins/Runtime/HairStrands/Source/HairStrandsCore/Private/HairStrandsInterpolation.cpp
+++ b/Engine/Plugins/Runtime/HairStrands/Source/HairStrandsCore/Private/HairStrandsInterpolation.cpp
@@ -451,7 +451,6 @@ FRDGHairStrandsCullingData ImportCullingData(FRDGBuilder& GraphBuilder, FHairGro
 
 	Out.ClusterCount		= In->GetClusterCount();
 	Out.ClusterAABBBuffer	= Register(GraphBuilder, In->GetClusterAABBBuffer(), ERDGImportedBufferFlags::CreateViews);
-	Out.GroupAABBBuffer		= Register(GraphBuilder, In->GetGroupAABBBuffer(), ERDGImportedBufferFlags::CreateViews);
 
 	return Out;
 }
@@ -790,6 +789,7 @@ class FHairClusterAABBCS : public FGlobalShader
 
 	BEGIN_SHADER_PARAMETER_STRUCT(FParameters, )
 		SHADER_PARAMETER_STRUCT_INCLUDE(ShaderPrint::FShaderParameters, ShaderDrawParameters)
+		SHADER_PARAMETER(uint32, InstanceRegisteredIndex)
 		SHADER_PARAMETER(float, LODIndex)
 		SHADER_PARAMETER(float, ClusterScale)
 		SHADER_PARAMETER(uint32, ClusterCount)
@@ -830,7 +830,8 @@ void AddHairClusterAABBPass(
 	FHairStrandClusterData::FHairGroup* ClusterData,
 	FRDGHairStrandsCullingData& ClusterAABBData,
 	FRDGBufferSRVRef RenderPositionBufferSRV,
-	FRDGBufferSRVRef& DrawIndirectRasterComputeBuffer)
+	FRDGBufferSRVRef& DrawIndirectRasterComputeBuffer,
+	FRDGBufferUAVRef GroupAABBBUAV)
 {
 	// Clusters AABB are only update if the groom is deformed.
 	const FBoxSphereBounds& Bounds = Instance->Strands.Data->Header.BoundingBox;
@@ -840,6 +841,7 @@ void AddHairClusterAABBPass(
 	Instance->HairGroupPublicData->SetClusterAABBValid(UpdateType == EHairAABBUpdateType::UpdateClusterAABB);
 
 	FHairClusterAABBCS::FParameters* Parameters = GraphBuilder.AllocParameters<FHairClusterAABBCS::FParameters>();
+	Parameters->InstanceRegisteredIndex = Instance->RegisteredIndex;
 	Parameters->LODIndex = ClusterData ? ClusterData->LODIndex : 1;;
 	Parameters->CPUBoundMin = (FVector3f)TransformedBounds.GetBox().Min;
 	Parameters->CPUBoundMax = (FVector3f)TransformedBounds.GetBox().Max;
@@ -852,7 +854,7 @@ void AddHairClusterAABBPass(
 	Parameters->RenCurveBuffer = RegisterAsSRV(GraphBuilder, Instance->Strands.RestResource->CurveBuffer);
 	Parameters->RenPointLODBuffer = ClusterData ? RegisterAsSRV(GraphBuilder, *ClusterData->PointLODBuffer) : nullptr;
 	Parameters->OutClusterAABBBuffer = ClusterAABBData.ClusterAABBBuffer.UAV;
-	Parameters->OutGroupAABBBuffer = ClusterAABBData.GroupAABBBuffer.UAV;
+	Parameters->OutGroupAABBBuffer = GroupAABBBUAV;
 
 	if (ShaderPrintData)
 	{
@@ -884,7 +886,6 @@ void AddHairClusterAABBPass(
 	}
 
 	GraphBuilder.SetBufferAccessFinal(ClusterAABBData.ClusterAABBBuffer.Buffer, ERHIAccess::SRVMask);
-	GraphBuilder.SetBufferAccessFinal(ClusterAABBData.GroupAABBBuffer.Buffer, ERHIAccess::SRVMask);
 }
 
 ///////////////////////////////////////////////////////////////////////////////////////////////////
@@ -1243,6 +1244,7 @@ class FClearClusterAABBCS : public FGlobalShader
 	BEGIN_SHADER_PARAMETER_STRUCT(FParameters, )
 		SHADER_PARAMETER_RDG_BUFFER_UAV(RWBuffer, OutClusterAABBBuffer)
 		SHADER_PARAMETER_RDG_BUFFER_UAV(RWBuffer, OutGroupAABBBuffer)
+		SHADER_PARAMETER(uint32, InstanceRegisteredIndex)
 		SHADER_PARAMETER(uint32, ClusterCount)
 		SHADER_PARAMETER(uint32, bClearClusterAABBs)
 	END_SHADER_PARAMETER_STRUCT()
@@ -1262,17 +1264,19 @@ void AddClearClusterAABBPass(
 	FRDGBuilder& GraphBuilder,
 	FGlobalShaderMap* ShaderMap,
 	const EHairAABBUpdateType UpdateType,
+	uint32 InstanceRegisteredIndex,
 	uint32 ClusterCount,
 	FRDGImportedBuffer& OutClusterAABBBuffer,
-	FRDGImportedBuffer& OutGroupAABBBuffer)
+	FRDGBufferUAVRef& OutGroupAABBUAV)
 {
 	check(OutClusterAABBBuffer.Buffer);
 
 	FClearClusterAABBCS::FParameters* Parameters = GraphBuilder.AllocParameters<FClearClusterAABBCS::FParameters>();
+	Parameters->InstanceRegisteredIndex = InstanceRegisteredIndex;
 	Parameters->ClusterCount = ClusterCount;
 	Parameters->bClearClusterAABBs = UpdateType == EHairAABBUpdateType::UpdateClusterAABB ? 1 : 0;
 	Parameters->OutClusterAABBBuffer = OutClusterAABBBuffer.UAV;
-	Parameters->OutGroupAABBBuffer = OutGroupAABBBuffer.UAV;
+	Parameters->OutGroupAABBBuffer = OutGroupAABBUAV;
 
 	TShaderMapRef<FClearClusterAABBCS> ComputeShader(ShaderMap);
 
@@ -1288,8 +1292,7 @@ void AddClearClusterAABBPass(
 		Parameters,
 		DispatchCount);
 
-	GraphBuilder.SetBufferAccessFinal(OutClusterAABBBuffer.Buffer, ERHIAccess::SRVMask),
-	GraphBuilder.SetBufferAccessFinal(OutGroupAABBBuffer.Buffer, ERHIAccess::SRVMask);
+	GraphBuilder.SetBufferAccessFinal(OutClusterAABBBuffer.Buffer, ERHIAccess::SRVMask);
 }
 
 ///////////////////////////////////////////////////////////////////////////////////////////////////
diff --git a/Engine/Plugins/Runtime/HairStrands/Source/HairStrandsCore/Public/GroomResources.h b/Engine/Plugins/Runtime/HairStrands/Source/HairStrandsCore/Public/GroomResources.h
--- a/Engine/Plugins/Runtime/HairStrands/Source/HairStrandsCore/Public/GroomResources.h
+++ b/Engine/Plugins/Runtime/HairStrands/Source/HairStrandsCore/Public/GroomResources.h
@@ -370,7 +370,6 @@ struct FHairStrandsCullingResource : public FHairCommonResource
 		Total += GetBufferTotalNumBytes(Resources.DrawIndirectBuffer);
 		Total += GetBufferTotalNumBytes(Resources.DrawIndirectRasterComputeBuffer);
 		Total += GetBufferTotalNumBytes(Resources.ClusterAABBBuffer);
-		Total += GetBufferTotalNumBytes(Resources.GroupAABBBuffer);
 		Total += GetBufferTotalNumBytes(Resources.CulledCurveBuffer);
 		Total += GetBufferTotalNumBytes(Resources.CulledVertexIdBuffer);
 		Total += GetBufferTotalNumBytes(Resources.CulledVertexRadiusScaleBuffer);
diff --git a/Engine/Plugins/Runtime/HairStrands/Source/HairStrandsCore/Public/HairStrandsInterpolation.h b/Engine/Plugins/Runtime/HairStrands/Source/HairStrandsCore/Public/HairStrandsInterpolation.h
--- a/Engine/Plugins/Runtime/HairStrands/Source/HairStrandsCore/Public/HairStrandsInterpolation.h
+++ b/Engine/Plugins/Runtime/HairStrands/Source/HairStrandsCore/Public/HairStrandsInterpolation.h
@@ -98,7 +98,6 @@ struct FRDGHairStrandsCullingData
 
 	uint32 ClusterCount = 0;
 	FRDGImportedBuffer ClusterAABBBuffer;
-	FRDGImportedBuffer GroupAABBBuffer;
 };
 
 FRDGHairStrandsCullingData ImportCullingData(FRDGBuilder& GraphBuilder, FHairGroupPublicData* In);
@@ -113,9 +112,10 @@ void AddClearClusterAABBPass(
 	FRDGBuilder& GraphBuilder,
 	FGlobalShaderMap* ShaderMap,
 	const EHairAABBUpdateType UpdateType,
+	uint32 InstanceRegisteredIndex,
 	uint32 ClusterCount,
 	FRDGImportedBuffer& OutClusterAABBBuffer,
-	FRDGImportedBuffer& OutGroupAABBBuffer);
+	FRDGBufferUAVRef& OutGroupAABBUAV);
 
 void AddHairStrandsInterpolationPass(
 	FRDGBuilder& GraphBuilder,
diff --git a/Engine/Shaders/Private/HairStrands/HairStrandsAABB.usf b/Engine/Shaders/Private/HairStrands/HairStrandsAABB.usf
--- a/Engine/Shaders/Private/HairStrands/HairStrandsAABB.usf
+++ b/Engine/Shaders/Private/HairStrands/HairStrandsAABB.usf
@@ -5,23 +5,16 @@
 
 #if SHADER_AABBUPDATE 
 
-uint bClearBuffer;
 uint MacroGroupId;
-uint MacroGroupValid;
+uint MacroGroupCount;
 
 float PixelSizeAtDepth1;
 float NumPixelPerVoxel;
 uint  VoxelPageResolution;
 
 // Group data packed as {uint3 AABBMin, uint3 AABBMax}
-Buffer<int> InGroupAABBBuffer0;
-Buffer<int> InGroupAABBBuffer1;
-Buffer<int> InGroupAABBBuffer2;
-Buffer<int> InGroupAABBBuffer3;
-Buffer<int> InGroupAABBBuffer4;
-Buffer<int> InGroupAABBBuffer5;
-Buffer<int> InGroupAABBBuffer6;
-Buffer<int> InGroupAABBBuffer7;
+Buffer<uint> RegisteredIndexBuffer;
+Buffer<int> InGroupAABBBuffer;
 RWBuffer<int> OutMacroGroupAABBBuffer;
 RWBuffer<float> OutMacroGroupVoxelSizeBuffer;
 
@@ -30,27 +23,32 @@ void UpdateValue(inout FHairAABB B, in uint Index, in Buffer<int> InBuffer)
 	FHairAABB A = ReadHairAABB(Index, InBuffer);
 	B = UnionHairAABB(A, B);
 }
-
-[numthreads(1, 1, 1)]
+groupshared float3 g_MinBounds[GROUP_SIZE];
+groupshared float3 g_MaxBounds[GROUP_SIZE];
+[numthreads(GROUP_SIZE, 1, 1)]
 void Main(uint2 DispatchThreadId : SV_DispatchThreadID)
 {
-	if (DispatchThreadId.x == 0)
+	const uint ThreadId = DispatchThreadId.x;
+	if (ThreadId < MacroGroupCount)
+	{
+		const uint RegisteredIndex = RegisteredIndexBuffer[ThreadId];
+		FHairAABB Bound = ReadHairAABB(RegisteredIndex, InGroupAABBBuffer);
+		g_MinBounds[ThreadId] = Bound.Min;
+		g_MaxBounds[ThreadId] = Bound.Max;
+	}
+	GroupMemoryBarrierWithGroupSync();
+	
+	if (ThreadId == 0)
 	{
 		FHairAABB Bound = InitHairAABB();
-		if (bClearBuffer == 0)
+		for (uint BoundIt = 0; BoundIt < MacroGroupCount; ++BoundIt)
 		{
-			Bound = ReadHairAABB(MacroGroupId, OutMacroGroupAABBBuffer);
+			FHairAABB CurrBound;
+			CurrBound.Min = g_MinBounds[BoundIt];
+			CurrBound.Max = g_MaxBounds[BoundIt];
+			Bound = UnionHairAABB(Bound, CurrBound);
 		}
 
-		if (MacroGroupValid & 1)   { UpdateValue(Bound, 0, InGroupAABBBuffer0); }
-		if (MacroGroupValid & 2)   { UpdateValue(Bound, 0, InGroupAABBBuffer1); }
-		if (MacroGroupValid & 4)   { UpdateValue(Bound, 0, InGroupAABBBuffer2); }
-		if (MacroGroupValid & 8)   { UpdateValue(Bound, 0, InGroupAABBBuffer3); }
-		if (MacroGroupValid & 16)  { UpdateValue(Bound, 0, InGroupAABBBuffer4); }
-		if (MacroGroupValid & 32)  { UpdateValue(Bound, 0, InGroupAABBBuffer5); }
-		if (MacroGroupValid & 64)  { UpdateValue(Bound, 0, InGroupAABBBuffer6); }
-		if (MacroGroupValid & 128) { UpdateValue(Bound, 0, InGroupAABBBuffer7); }
-
 		// Transform positions in camera space
 		const FHairAABB ViewBound = Transform(Bound, View.TranslatedWorldToCameraView);
 
diff --git a/Engine/Shaders/Private/HairStrands/HairStrandsClusterCulling.usf b/Engine/Shaders/Private/HairStrands/HairStrandsClusterCulling.usf
--- a/Engine/Shaders/Private/HairStrands/HairStrandsClusterCulling.usf
+++ b/Engine/Shaders/Private/HairStrands/HairStrandsClusterCulling.usf
@@ -37,6 +37,7 @@ uint3 ComputeDispatchGroupCount2D(uint GroupCount)
 /// Clear clusters AABB
 #ifdef SHADER_CLEARCLUSTERAABB
 
+uint InstanceRegisteredIndex;
 uint ClusterCount;
 uint bClearClusterAABBs;
 
@@ -59,12 +60,12 @@ void MainClearClusterAABBCS(uint2 DispatchThreadId : SV_DispatchThreadID)
 	// Clear 'group' AABB
 	if (DispatchThreadId.x == 0)
 	{
-		OutGroupAABBBuffer[0] =  2147483647;
-		OutGroupAABBBuffer[1] =  2147483647;
-		OutGroupAABBBuffer[2] =  2147483647;
-		OutGroupAABBBuffer[3] = -2147483647;
-		OutGroupAABBBuffer[4] = -2147483647;
-		OutGroupAABBBuffer[5] = -2147483647;
+		OutGroupAABBBuffer[InstanceRegisteredIndex*6 + 0] =  2147483647;
+		OutGroupAABBBuffer[InstanceRegisteredIndex*6 + 1] =  2147483647;
+		OutGroupAABBBuffer[InstanceRegisteredIndex*6 + 2] =  2147483647;
+		OutGroupAABBBuffer[InstanceRegisteredIndex*6 + 3] = -2147483647;
+		OutGroupAABBBuffer[InstanceRegisteredIndex*6 + 4] = -2147483647;
+		OutGroupAABBBuffer[InstanceRegisteredIndex*6 + 5] = -2147483647;
 	}
 }
 
@@ -417,10 +418,12 @@ void Main(uint3 DispatchThreadId : SV_DispatchThreadID, uint ThreadIndex : SV_Gr
 
 #if SHADER_CLUSTER_CULL_ARGS
 
+uint InstanceRegisteredIndex;
 uint ClusterGroupIndex;
 StructuredBuffer<uint> PointCounterBuffer;
 RWBuffer<uint> RWIndirectDrawArgsBuffer;
 RWBuffer<uint> RWIndirectDispatchArgsBuffer;
+RWBuffer<uint> RWIndirectDispatchArgsGlobalBuffer;
 
 [numthreads(1, 1, 1)]
 void Main(uint DispatchThreadId : SV_DispatchThreadID, uint ThreadIndex : SV_GroupIndex)
@@ -437,6 +440,11 @@ void Main(uint DispatchThreadId : SV_DispatchThreadID, uint ThreadIndex : SV_Gro
 	RWIndirectDispatchArgsBuffer[1] = 1;
 	RWIndirectDispatchArgsBuffer[2] = 1;
 	RWIndirectDispatchArgsBuffer[3] = PointCount;	
+
+	RWIndirectDispatchArgsGlobalBuffer[InstanceRegisteredIndex * 4 + 0] = DispatchCount;
+	RWIndirectDispatchArgsGlobalBuffer[InstanceRegisteredIndex * 4 + 1] = 1;
+	RWIndirectDispatchArgsGlobalBuffer[InstanceRegisteredIndex * 4 + 2] = 1;
+	RWIndirectDispatchArgsGlobalBuffer[InstanceRegisteredIndex * 4 + 3] = PointCount;
 }
 
 #endif // SHADER_CLUSTER_CULL_ARGS
@@ -446,6 +454,7 @@ void Main(uint DispatchThreadId : SV_DispatchThreadID, uint ThreadIndex : SV_Gro
 
 #if SHADER_CLUSTERAABB
 
+uint InstanceRegisteredIndex;
 uint CurveCount;
 uint ClusterCount;
 float ClusterScale;
@@ -478,7 +487,7 @@ groupshared float3 OutHairPositionOffset;
 		FHairAABB Bound;
 		Bound.Min = CPUBoundMin;
 		Bound.Max = CPUBoundMax;
-		WriteHairAABB(0, Bound, OutGroupAABBBuffer);
+		WriteHairAABB(InstanceRegisteredIndex, Bound, OutGroupAABBBuffer);
 	}
 #else // PERMUTATION_USE_CPU_AABB
 
@@ -592,7 +601,7 @@ groupshared float3 OutHairPositionOffset;
 		Bound.Max.x = SharedClusterAABBMax[0].x;
 		Bound.Max.y = SharedClusterAABBMax[0].y;
 		Bound.Max.z = SharedClusterAABBMax[0].z;
-		InterlockHairAABB(0, Bound, OutGroupAABBBuffer);
+		InterlockHairAABB(InstanceRegisteredIndex, Bound, OutGroupAABBBuffer);
 	}
 
 #if DEBUG_ENABLE
diff --git a/Engine/Shaders/Private/HairStrands/HairStrandsDebug.usf b/Engine/Shaders/Private/HairStrands/HairStrandsDebug.usf
--- a/Engine/Shaders/Private/HairStrands/HairStrandsDebug.usf
+++ b/Engine/Shaders/Private/HairStrands/HairStrandsDebug.usf
@@ -1152,6 +1152,7 @@ void MainCS(uint InGroupIndex : SV_GroupIndex, uint3 DispatchThreadId : SV_Dispa
 
 #include "HairStrandsAABBCommon.ush"
 
+uint					InstanceRegisteredIndex;
 int2					MaxResolution;
 uint					InstanceCount;
 uint					InstanceCount_StrandsPrimaryView;
@@ -1555,7 +1556,7 @@ void MainCS(uint GroupIndex : SV_GroupIndex, uint3 DispatchThreadId : SV_Dispatc
 	if (all(DispatchThreadId == 0))
 	{
 		// Bound
-		const FHairAABB AABB = ReadHairAABB(0, InstanceAABB);
+		const FHairAABB AABB = ReadHairAABB(InstanceRegisteredIndex, InstanceAABB);
 		AddAABBTWS(AABB.Min, AABB.Max, GetGPUAABBInstanceGroupBoundColor());
 
 		// Screen size
@@ -2872,6 +2873,7 @@ void MainCS(uint GroupIndex : SV_GroupIndex, uint3 DispatchThreadId : SV_Dispatc
 
 #include "HairStrandsAABBCommon.ush"
 
+uint InstanceRegisteredIndex;
 uint ClusterCount;
 uint HairGroupId;
 uint PointCount;
@@ -2914,7 +2916,7 @@ void MainDrawDebugAABBCS(uint3 DispatchThreadId : SV_DispatchThreadID)
 		// Group
 		if (DispatchThreadId.x == 0)
 		{
-			const FHairAABB GroupBound = ReadHairAABB(0, GroupAABBBuffer);
+			const FHairAABB GroupBound = ReadHairAABB(InstanceRegisteredIndex, GroupAABBBuffer);
 			AddAABBTWS(GroupBound.Min, GroupBound.Max, float4(1.0f, 0.5f, 0.0f, 1.0f));
 		}
 	
diff --git a/Engine/Shaders/Private/HairStrands/HairStrandsInterpolation.usf b/Engine/Shaders/Private/HairStrands/HairStrandsInterpolation.usf
--- a/Engine/Shaders/Private/HairStrands/HairStrandsInterpolation.usf
+++ b/Engine/Shaders/Private/HairStrands/HairStrandsInterpolation.usf
@@ -91,7 +91,6 @@ uint HairStrandsVFTODO_bCullingEnable;
 Buffer<uint>	HairStrandsVFTODO_CullingIndirectBuffer; 
 Buffer<uint>	HairStrandsVFTODO_CullingIndexBuffer;
 Buffer<float>	HairStrandsVFTODO_CullingRadiusScaleBuffer;
-Buffer<uint>	HairStrandsVFTODO_CullingIndirectBufferArgs;
 #endif
 
 #if PERMUTATION_DYNAMIC_GEOMETRY >= 1 || PERMUTATION_DEBUG == 1 
diff --git a/Engine/Shaders/Private/HairStrands/HairStrandsRaytracingGeometry.usf b/Engine/Shaders/Private/HairStrands/HairStrandsRaytracingGeometry.usf
--- a/Engine/Shaders/Private/HairStrands/HairStrandsRaytracingGeometry.usf
+++ b/Engine/Shaders/Private/HairStrands/HairStrandsRaytracingGeometry.usf
@@ -298,6 +298,7 @@ void MainCS(uint DispatchThreadId : SV_DispatchThreadID)
 #error GROUP_SIZE has changed, please update the reduction code.
 #endif
 
+uint InstanceResgisteredIndex;
 uint PointCount;
 uint bDrawInvalidElement;
 float PositionThreshold2;
@@ -351,7 +352,7 @@ void MainCS(uint2 DispatchThreadId : SV_DispatchThreadID, uint LinearIndex : SV_
 				InterlockedAdd(InvalidationPrintCounter[0], 1, Offset);
 				if (Offset < 2)
 				{
-					const FHairAABB Bound = ReadHairAABB(0, GroupAABBBuffer);
+					const FHairAABB Bound = ReadHairAABB(InstanceResgisteredIndex, GroupAABBBuffer);
 					AddAABBTWS(Bound.Min, Bound.Max, ColorRed);
 				}
 			}
diff --git a/Engine/Source/Runtime/Renderer/Private/DeferredShadingRenderer.cpp b/Engine/Source/Runtime/Renderer/Private/DeferredShadingRenderer.cpp
--- a/Engine/Source/Runtime/Renderer/Private/DeferredShadingRenderer.cpp
+++ b/Engine/Source/Runtime/Renderer/Private/DeferredShadingRenderer.cpp
@@ -3230,6 +3230,7 @@ void FDeferredShadingSceneRenderer::Render(FRDGBuilder& GraphBuilder)
 	if (IsHairStrandsEnabled(EHairStrandsShaderType::All, Scene->GetShaderPlatform()) && RendererOutput == ERendererOutput::FinalSceneColor)
 	{
 		CreateHairStrandsBookmarkParameters(Scene, Views, AllFamilyViews, HairStrandsBookmarkParameters);
+		HairStrandsBookmarkParameters.TransientResources = AllocateHairTransientResourcse(GraphBuilder, Scene);
 		RunHairStrandsBookmark(GraphBuilder, EHairStrandsBookmark::ProcessTasks, HairStrandsBookmarkParameters);
 
 		// Interpolation needs to happen after the skin cache run as there is a dependency 
@@ -4573,6 +4574,7 @@ void FDeferredShadingSceneRenderer::Render(FRDGBuilder& GraphBuilder)
 	QueueSceneTextureExtractions(GraphBuilder, SceneTextures);
 
 	::Substrate::PostRender(*Scene);
+	HairStrands::PostRender(*Scene);
 
 	// Release the view's previous frame histories so that their memory can be reused at the graph's execution.
 	for (int32 ViewIndex = 0; ViewIndex < Views.Num(); ViewIndex++)
diff --git a/Engine/Source/Runtime/Renderer/Private/HairStrands/HairStrandsCluster.cpp b/Engine/Source/Runtime/Renderer/Private/HairStrands/HairStrandsCluster.cpp
--- a/Engine/Source/Runtime/Renderer/Private/HairStrands/HairStrandsCluster.cpp
+++ b/Engine/Source/Runtime/Renderer/Private/HairStrands/HairStrandsCluster.cpp
@@ -23,32 +23,27 @@ class FHairMacroGroupAABBCS : public FGlobalShader
 
 	BEGIN_SHADER_PARAMETER_STRUCT(FParameters, )
 		SHADER_PARAMETER(uint32, MacroGroupId)
-		SHADER_PARAMETER(uint32, MacroGroupValid)
-		SHADER_PARAMETER(uint32, bClearBuffer)
+		SHADER_PARAMETER(uint32, MacroGroupCount)
 
 		SHADER_PARAMETER(float, PixelSizeAtDepth1)
 		SHADER_PARAMETER(float, NumPixelPerVoxel)
 		SHADER_PARAMETER(uint32, VoxelPageResolution)
 
-		SHADER_PARAMETER_RDG_BUFFER_SRV(Buffer, InGroupAABBBuffer0)
-		SHADER_PARAMETER_RDG_BUFFER_SRV(Buffer, InGroupAABBBuffer1)
-		SHADER_PARAMETER_RDG_BUFFER_SRV(Buffer, InGroupAABBBuffer2)
-		SHADER_PARAMETER_RDG_BUFFER_SRV(Buffer, InGroupAABBBuffer3)
-		SHADER_PARAMETER_RDG_BUFFER_SRV(Buffer, InGroupAABBBuffer4)
-		SHADER_PARAMETER_RDG_BUFFER_SRV(Buffer, InGroupAABBBuffer5)
-		SHADER_PARAMETER_RDG_BUFFER_SRV(Buffer, InGroupAABBBuffer6)
-		SHADER_PARAMETER_RDG_BUFFER_SRV(Buffer, InGroupAABBBuffer7)
+		SHADER_PARAMETER_RDG_BUFFER_SRV(Buffer<uint>, RegisteredIndexBuffer)
+		SHADER_PARAMETER_RDG_BUFFER_SRV(Buffer<int>, InGroupAABBBuffer)
 		SHADER_PARAMETER_RDG_BUFFER_UAV(RWBuffer, OutMacroGroupAABBBuffer)
 		SHADER_PARAMETER_RDG_BUFFER_UAV(RWBuffer, OutMacroGroupVoxelSizeBuffer)
 		SHADER_PARAMETER_STRUCT_REF(FViewUniformShaderParameters, View)
 	END_SHADER_PARAMETER_STRUCT()
 
 public:
+	static uint32 GetGroupSize() { return 32; }
 	static bool ShouldCompilePermutation(const FGlobalShaderPermutationParameters& Parameters) { return IsHairStrandsSupported(EHairStrandsShaderType::Strands, Parameters.Platform); }
 	static void ModifyCompilationEnvironment(const FGlobalShaderPermutationParameters& Parameters, FShaderCompilerEnvironment& OutEnvironment)
 	{
 		FGlobalShader::ModifyCompilationEnvironment(Parameters, OutEnvironment);
 		OutEnvironment.SetDefine(TEXT("SHADER_AABBUPDATE"), 1);
+		OutEnvironment.SetDefine(TEXT("GROUP_SIZE"), GetGroupSize());
 	}
 };
 
@@ -59,6 +54,7 @@ void GetVoxelPageResolution(uint32& OutPageResolution, uint32& OutPageResolution
 static void AddHairMacroGroupAABBPass(
 	FRDGBuilder& GraphBuilder,
 	const FViewInfo& View, 
+	FHairTransientResources& TransientResources,
 	FHairStrandsMacroGroupData& MacroGroup,
 	FRDGBufferUAVRef& OutHairMacroGroupAABBBufferUAV,
 	FRDGBufferUAVRef& OutMacroGroupVoxelSizeBufferUAV)
@@ -71,77 +67,44 @@ static void AddHairMacroGroupAABBPass(
 	const FIntPoint Resolution(View.ViewRect.Width(), View.ViewRect.Height());
 	const float vFOV = FMath::DegreesToRadians(View.FOV);
 	const float PixelSizeAtDepth1 = FMath::Tan(vFOV * 0.5f) / (0.5f * Resolution.Y);
+	const uint32 MacroGroupCount = MacroGroup.PrimitivesInfos.Num();
 	
 	uint32 VoxelPageResolution = 0;
 	uint32 VoxelPageResolutionLog2 = 1;
 	GetVoxelPageResolution(VoxelPageResolution, VoxelPageResolutionLog2);
 
-	const float NumPixelPerVoxel = FMath::Clamp(GHairVirtualVoxel_NumPixelPerVoxel, 1.f, 50.f);
-
-	const uint32 GroupPerPass = 8;
-	bool bNeedClear = true;
-	const uint32 MacroGroupId = MacroGroup.MacroGroupId;
-	const uint32 IterationCount = FMath::CeilToInt(PrimitiveCount / float(GroupPerPass));
-	for (uint32 PassIt = 0; PassIt < IterationCount; ++PassIt)
+	TArray<uint32> RegisteredIndices;
+	RegisteredIndices.Reserve(MacroGroupCount);
+	for (const FHairStrandsMacroGroupData::PrimitiveInfo& PrimitiveInfo : MacroGroup.PrimitivesInfos)
 	{
-		FHairMacroGroupAABBCS::FParameters* Parameters = GraphBuilder.AllocParameters<FHairMacroGroupAABBCS::FParameters>();
-		Parameters->MacroGroupId = MacroGroupId;
-		Parameters->OutMacroGroupAABBBuffer = OutHairMacroGroupAABBBufferUAV;
-		Parameters->OutMacroGroupVoxelSizeBuffer = OutMacroGroupVoxelSizeBufferUAV;
-		Parameters->PixelSizeAtDepth1   = PixelSizeAtDepth1;
-		Parameters->NumPixelPerVoxel    = NumPixelPerVoxel;
-		Parameters->VoxelPageResolution = VoxelPageResolution;
-		Parameters->View			    = View.ViewUniformBuffer;
-
-		uint32 MacroGroupValid = 1;
-		uint32 CurrentGroupIt = 1;
-		for (uint32 PassPrimitiveIt = 0, PassPrimitiveCount = FMath::Min(GroupPerPass, PrimitiveCount - PassIt * GroupPerPass); PassPrimitiveIt < PassPrimitiveCount; ++PassPrimitiveIt)
-		{
-			const uint32 PrimitiveIndex = PassIt * GroupPerPass + PassPrimitiveIt;
-			const FHairStrandsMacroGroupData::PrimitiveInfo& PrimitiveInfo = MacroGroup.PrimitivesInfos[PrimitiveIndex];
-			FRDGBufferSRVRef GroupAABBBufferSRV = RegisterAsSRV(GraphBuilder, PrimitiveInfo.PublicDataPtr->GetGroupAABBBuffer());
-
-			// Default value
-			if (PassPrimitiveIt == 0 && PassPrimitiveCount != GroupPerPass)
-			{
-				Parameters->InGroupAABBBuffer0 = GroupAABBBufferSRV;
-				Parameters->InGroupAABBBuffer1 = GroupAABBBufferSRV;
-				Parameters->InGroupAABBBuffer2 = GroupAABBBufferSRV;
-				Parameters->InGroupAABBBuffer3 = GroupAABBBufferSRV;
-				Parameters->InGroupAABBBuffer4 = GroupAABBBufferSRV;
-				Parameters->InGroupAABBBuffer5 = GroupAABBBufferSRV;
-				Parameters->InGroupAABBBuffer6 = GroupAABBBufferSRV;
-				Parameters->InGroupAABBBuffer7 = GroupAABBBufferSRV;
-			}
-
-			switch (PassPrimitiveIt)
-			{
-				case 0 : Parameters->InGroupAABBBuffer0 = GroupAABBBufferSRV; break;
-				case 1 : Parameters->InGroupAABBBuffer1 = GroupAABBBufferSRV; break;
-				case 2 : Parameters->InGroupAABBBuffer2 = GroupAABBBufferSRV; break;
-				case 3 : Parameters->InGroupAABBBuffer3 = GroupAABBBufferSRV; break;
-				case 4 : Parameters->InGroupAABBBuffer4 = GroupAABBBufferSRV; break;
-				case 5 : Parameters->InGroupAABBBuffer5 = GroupAABBBufferSRV; break;
-				case 6 : Parameters->InGroupAABBBuffer6 = GroupAABBBufferSRV; break;
-				case 7 : Parameters->InGroupAABBBuffer7 = GroupAABBBufferSRV; break;
-			}
-			MacroGroupValid |= 1 << PassPrimitiveIt;
-
-		}
-
-		Parameters->MacroGroupValid = MacroGroupValid;
-		Parameters->bClearBuffer = bNeedClear ? 1 : 0;
+		RegisteredIndices.Add(PrimitiveInfo.PublicDataPtr->Instance->RegisteredIndex);
+	}
+	FRDGBufferRef RegisteredIndexBuffer = CreateVertexBuffer(GraphBuilder, TEXT("Hair.RegisteredIndexBuffer"), FRDGBufferDesc::CreateBufferDesc(4, RegisteredIndices.Num()), RegisteredIndices.GetData(), 4u * RegisteredIndices.Num());
 
-		TShaderMapRef<FHairMacroGroupAABBCS> ComputeShader(View.ShaderMap);
-		FComputeShaderUtils::AddPass(
-			GraphBuilder,
-			RDG_EVENT_NAME("HairStrands::MacroGroupAABBUpdate"),
-			ComputeShader,
-			Parameters,
-			FIntVector(1,1,1));
+	const float NumPixelPerVoxel = FMath::Clamp(GHairVirtualVoxel_NumPixelPerVoxel, 1.f, 50.f);
 
-		bNeedClear = false;
-	}
+	// Can only aggregate 32 GroupAABBs 
+	check(uint32(RegisteredIndices.Num()) < FHairMacroGroupAABBCS::GetGroupSize());
+
+	FHairMacroGroupAABBCS::FParameters* Parameters = GraphBuilder.AllocParameters<FHairMacroGroupAABBCS::FParameters>();
+	Parameters->MacroGroupId 					= MacroGroup.MacroGroupId;
+	Parameters->RegisteredIndexBuffer			= GraphBuilder.CreateSRV(RegisteredIndexBuffer, PF_R32_UINT);
+	Parameters->InGroupAABBBuffer				= TransientResources.GroupAABBSRV;
+	Parameters->OutMacroGroupAABBBuffer 		= OutHairMacroGroupAABBBufferUAV;
+	Parameters->OutMacroGroupVoxelSizeBuffer 	= OutMacroGroupVoxelSizeBufferUAV;
+	Parameters->PixelSizeAtDepth1   			= PixelSizeAtDepth1;
+	Parameters->NumPixelPerVoxel    			= NumPixelPerVoxel;
+	Parameters->VoxelPageResolution 			= VoxelPageResolution;
+	Parameters->View			    			= View.ViewUniformBuffer;
+	Parameters->MacroGroupCount					= MacroGroupCount;
+	
+	TShaderMapRef<FHairMacroGroupAABBCS> ComputeShader(View.ShaderMap);
+	FComputeShaderUtils::AddPass(
+		GraphBuilder,
+		RDG_EVENT_NAME("HairStrands::MacroGroupAABBUpdate"),
+		ComputeShader,
+		Parameters,
+		FIntVector(1, 1, 1));
 }
 
 static bool DoesGroupExists(uint32 ResourceId, uint32 GroupIndex, const FHairStrandsMacroGroupData::TPrimitiveInfos& PrimitivesGroups)
@@ -187,7 +150,8 @@ void CreateHairStrandsMacroGroups(
 	FRDGBuilder& GraphBuilder,
 	const FScene* Scene,
 	const FViewInfo& View, 
-	FHairStrandsViewData& OutHairStrandsViewData)
+	FHairStrandsViewData& OutHairStrandsViewData,
+	bool bBuildGPUAABB)
 {
 	const bool bHasHairStrandsElements = View.HairStrandsMeshElements.Num() != 0 || Scene->HairStrandsSceneData.RegisteredProxies.Num() != 0;
 	if (!View.Family || !bHasHairStrandsElements || View.bIsReflectionCapture)
@@ -316,11 +280,12 @@ void CreateHairStrandsMacroGroups(
 	}
 	// Sanity check
 	check(MacroGroups.Num() <= FHairStrandsMacroGroupData::MaxMacroGroupCount);
+	check(Scene->HairStrandsSceneData.TransientResources);
 
 	// Build hair macro group AABBB
 	FHairStrandsMacroGroupResources& MacroGroupResources = OutHairStrandsViewData.MacroGroupResources;
 	const uint32 MacroGroupCount = MacroGroups.Num();
-	if (MacroGroupCount > 0)
+	if (MacroGroupCount > 0 && bBuildGPUAABB)
 	{
 		DECLARE_GPU_STAT(HairStrandsAABB);
 		RDG_EVENT_SCOPE(GraphBuilder, "HairStrandsAABB");
@@ -333,7 +298,7 @@ void CreateHairStrandsMacroGroups(
 		FRDGBufferUAVRef MacroGroupVoxelSizeBufferUAV = GraphBuilder.CreateUAV(MacroGroupResources.MacroGroupVoxelSizeBuffer, PF_R16F);
 		for (FHairStrandsMacroGroupData& MacroGroup : MacroGroups)
 		{				
-			AddHairMacroGroupAABBPass(GraphBuilder, View, MacroGroup, MacroGroupAABBBufferUAV, MacroGroupVoxelSizeBufferUAV);
+			AddHairMacroGroupAABBPass(GraphBuilder, View, *Scene->HairStrandsSceneData.TransientResources, MacroGroup, MacroGroupAABBBufferUAV, MacroGroupVoxelSizeBufferUAV);
 		}
 		MacroGroupResources.MacroGroupCount = MacroGroups.Num();
 	}
diff --git a/Engine/Source/Runtime/Renderer/Private/HairStrands/HairStrandsCluster.h b/Engine/Source/Runtime/Renderer/Private/HairStrands/HairStrandsCluster.h
--- a/Engine/Source/Runtime/Renderer/Private/HairStrands/HairStrandsCluster.h
+++ b/Engine/Source/Runtime/Renderer/Private/HairStrands/HairStrandsCluster.h
@@ -16,4 +16,5 @@ void CreateHairStrandsMacroGroups(
 	FRDGBuilder& GraphBuilder,
 	const FScene* Scene,
 	const FViewInfo& View, 
-	FHairStrandsViewData& OutHairStrandsViewData);
\ No newline at end of file
+	FHairStrandsViewData& OutHairStrandsViewData,
+	bool bBuildGPUAABB=true);
\ No newline at end of file
diff --git a/Engine/Source/Runtime/Renderer/Private/HairStrands/HairStrandsData.h b/Engine/Source/Runtime/Renderer/Private/HairStrands/HairStrandsData.h
--- a/Engine/Source/Runtime/Renderer/Private/HairStrands/HairStrandsData.h
+++ b/Engine/Source/Runtime/Renderer/Private/HairStrands/HairStrandsData.h
@@ -487,7 +487,7 @@ namespace HairStrands
 	bool HasViewHairStrandsData(const TArray<FViewInfo>& Views);
 	bool HasViewHairStrandsVoxelData(const FViewInfo& View);
 
-	bool HasPositionsChanged(FRDGBuilder& GraphBuilder, const FViewInfo& View);
+	bool HasPositionsChanged(FRDGBuilder& GraphBuilder, const FScene& Scene, const FViewInfo& View);
 	void DrawHitProxies(FRDGBuilder& GraphBuilder, const FScene& Scene, const FViewInfo& View, FInstanceCullingManager& InstanceCullingManager, FRDGTextureRef HitProxyTexture, FRDGTextureRef HitProxyDepthTexture);
 	void DrawEditorSelection(FRDGBuilder& GraphBuilder, const FViewInfo& View, const FIntRect& ViewportRect, FRDGTextureRef SelectionDepthTexture);
 
@@ -504,4 +504,6 @@ namespace HairStrands
 	bool HasHairStrandsVisible(const TArray<FViewInfo>& Views);
 
 	void AddVisibleShadowCastingLight(const FScene& Scene, TArray<FViewInfo>& Views, const FLightSceneInfo* LightSceneInfo);
+
+	void PostRender(FScene& Scene);
 }
\ No newline at end of file
diff --git a/Engine/Source/Runtime/Renderer/Private/HairStrands/HairStrandsInterface.cpp b/Engine/Source/Runtime/Renderer/Private/HairStrands/HairStrandsInterface.cpp
--- a/Engine/Source/Runtime/Renderer/Private/HairStrands/HairStrandsInterface.cpp
+++ b/Engine/Source/Runtime/Renderer/Private/HairStrands/HairStrandsInterface.cpp
@@ -431,6 +431,7 @@ void CreateHairStrandsBookmarkParameters(FScene* Scene, FViewInfo& View, FHairSt
 	Out.SceneColorTexture		= nullptr;
 	Out.SceneDepthTexture		= nullptr;
 	Out.Scene					= Scene;
+	Out.TransientResources		= Scene->HairStrandsSceneData.TransientResources;
 }
 
 void UpdateHairStrandsBookmarkParameters(FScene* Scene, TArray<FViewInfo>& Views, FHairStrandsBookmarkParameters& Out)
diff --git a/Engine/Source/Runtime/Renderer/Private/HairStrands/HairStrandsRendering.cpp b/Engine/Source/Runtime/Renderer/Private/HairStrands/HairStrandsRendering.cpp
--- a/Engine/Source/Runtime/Renderer/Private/HairStrands/HairStrandsRendering.cpp
+++ b/Engine/Source/Runtime/Renderer/Private/HairStrands/HairStrandsRendering.cpp
@@ -85,6 +85,32 @@ void AddMeshDrawTransitionPass(
 	const FViewInfo& ViewInfo,
 	const FHairStrandsMacroGroupDatas& MacroGroupDatas);
 
+FHairTransientResources* AllocateHairTransientResourcse(FRDGBuilder& GraphBuilder, FScene* Scene)
+{
+	if (Scene->HairStrandsSceneData.TransientResources == nullptr)
+	{
+		Scene->HairStrandsSceneData.TransientResources = new FHairTransientResources();
+	}
+
+	const uint32 InstanceCount = Scene->HairStrandsSceneData.RegisteredProxies.Num();
+	FHairTransientResources* Out = Scene->HairStrandsSceneData.TransientResources;
+	*Out = FHairTransientResources();
+	if (InstanceCount > 0)
+	{
+		Out->bIsGroupAABBValid.Init(false, InstanceCount);
+
+		// Change this into a structure buffer
+		Out->GroupAABBBuffer = GraphBuilder.CreateBuffer(FRDGBufferDesc::CreateBufferDesc(4, 6 * InstanceCount), TEXT("Hair.TransientResources.GroupAABB"));
+		Out->GroupAABBUAV = GraphBuilder.CreateUAV(Out->GroupAABBBuffer, PF_R32_SINT);
+		Out->GroupAABBSRV = GraphBuilder.CreateSRV(Out->GroupAABBBuffer, PF_R32_SINT);
+
+		Out->IndirectDispatchArgsBuffer = GraphBuilder.CreateBuffer(FRDGBufferDesc::CreateIndirectDesc<FRHIDispatchIndirectParameters>(InstanceCount), TEXT("Hair.TransientResources.IndirectDispatchArgs"));
+		Out->IndirectDispatchArgsUAV = GraphBuilder.CreateUAV(Out->IndirectDispatchArgsBuffer);
+		Out->IndirectDispatchArgsSRV = GraphBuilder.CreateSRV(Out->IndirectDispatchArgsBuffer);
+	}
+	return Out;
+}
+
 void RenderHairPrePass(
 	FRDGBuilder& GraphBuilder,
 	FScene* Scene,
@@ -313,4 +339,13 @@ bool HasHairInstanceInScene(const FScene& Scene)
 	return Scene.HairStrandsSceneData.RegisteredProxies.Num() > 0;
 }
 
+void PostRender(FScene& Scene)
+{
+	// Dellocate transient resourcse
+	if (Scene.HairStrandsSceneData.TransientResources)
+	{
+		*Scene.HairStrandsSceneData.TransientResources = FHairTransientResources();
+	}
+}
+
 } // HairStrands
diff --git a/Engine/Source/Runtime/Renderer/Private/HairStrands/HairStrandsRendering.h b/Engine/Source/Runtime/Renderer/Private/HairStrands/HairStrandsRendering.h
--- a/Engine/Source/Runtime/Renderer/Private/HairStrands/HairStrandsRendering.h
+++ b/Engine/Source/Runtime/Renderer/Private/HairStrands/HairStrandsRendering.h
@@ -21,6 +21,8 @@
 #include "HairStrandsInterface.h"
 #include "HairStrandsData.h"
 
+FHairTransientResources* AllocateHairTransientResourcse(FRDGBuilder& GraphBuilder, FScene* Scene);
+
 void RenderHairPrePass(
 	FRDGBuilder& GraphBuilder,
 	FScene* Scene,
diff --git a/Engine/Source/Runtime/Renderer/Private/HairStrands/HairStrandsVisibility.cpp b/Engine/Source/Runtime/Renderer/Private/HairStrands/HairStrandsVisibility.cpp
--- a/Engine/Source/Runtime/Renderer/Private/HairStrands/HairStrandsVisibility.cpp
+++ b/Engine/Source/Runtime/Renderer/Private/HairStrands/HairStrandsVisibility.cpp
@@ -4011,6 +4011,7 @@ class FHairStrandsPositionChangedCS : public FGlobalShader
 	SHADER_USE_PARAMETER_STRUCT(FHairStrandsPositionChangedCS, FGlobalShader);
 
 	BEGIN_SHADER_PARAMETER_STRUCT(FParameters, )
+		SHADER_PARAMETER(uint32, InstanceResgisteredIndex)
 		SHADER_PARAMETER(uint32, PointCount)
 		SHADER_PARAMETER(float, PositionThreshold2)
 		SHADER_PARAMETER(uint32, HairStrandsVF_bCullingEnable)
@@ -4040,8 +4041,11 @@ static void AddHairStrandsHasPositionChangedPass(
 	FRDGBuilder& GraphBuilder,
 	const FViewInfo* View,
 	const FHairGroupPublicData* HairGroupPublicData,
+	FHairTransientResources* TransientResources,
 	FRDGBufferUAVRef InvalidationBuffer)
 {
+	check(TransientResources);
+
 	const uint32 PointCount = HairGroupPublicData->GetActiveStrandsPointCount();
 
 	FRDGBufferRef InvalidationPrintCounter = GraphBuilder.CreateBuffer(FRDGBufferDesc::CreateBufferDesc(sizeof(uint32), 1), TEXT("Hair.InvalidationPrintCounter"));
@@ -4049,6 +4053,7 @@ static void AddHairStrandsHasPositionChangedPass(
 	AddClearUAVPass(GraphBuilder, InvalidationPrintCounterUAV, 0u);
 
 	FHairStrandsPositionChangedCS::FParameters* Parameters = GraphBuilder.AllocParameters<FHairStrandsPositionChangedCS::FParameters>();
+	Parameters->InstanceResgisteredIndex = HairGroupPublicData->Instance->RegisteredIndex;
 	Parameters->PointCount = PointCount;
 	Parameters->PositionThreshold2 = FMath::Square(GHairStrands_InvalidationPosition_Threshold);
 	Parameters->bDrawInvalidElement = GHairStrands_InvalidationPosition_Debug > 0 ? 1u : 0u;
@@ -4057,7 +4062,7 @@ static void AddHairStrandsHasPositionChangedPass(
 	Parameters->HairStrandsVF_CullingIndirectBuffer = Parameters->HairStrandsVF_CullingIndexBuffer;
 	Parameters->CurrPositionBuffer = HairGroupPublicData->VFInput.Strands.PositionBuffer.SRV;
 	Parameters->PrevPositionBuffer = HairGroupPublicData->VFInput.Strands.PrevPositionBuffer.SRV;
-	Parameters->GroupAABBBuffer = Register(GraphBuilder, HairGroupPublicData->GetGroupAABBBuffer(), ERDGImportedBufferFlags::CreateSRV).SRV;
+	Parameters->GroupAABBBuffer = TransientResources->GroupAABBSRV;
 	Parameters->InvalidationBuffer = InvalidationBuffer;
 	Parameters->InvalidationPrintCounter = InvalidationPrintCounterUAV;
 	ShaderPrint::SetParameters(GraphBuilder, View->ShaderPrintData, Parameters->ShaderPrintParameters);
@@ -4118,7 +4123,7 @@ void DrawHitProxies(
 	// Geometry won't be updated for proxy view
 	const FIntPoint Resolution = HitProxyTexture->Desc.Extent;
 	FHairStrandsViewData HairStrandsViewData;
-	CreateHairStrandsMacroGroups(GraphBuilder, &Scene, View, HairStrandsViewData);
+	CreateHairStrandsMacroGroups(GraphBuilder, &Scene, View, HairStrandsViewData, false /*bBuildGPUAABB*/);
 
 	// We don't compute the transmittance texture as there is no need for picking.
 	FRDGTextureRef DummyTransmittanceTexture = GraphBuilder.CreateTexture(FRDGTextureDesc::Create2D(Resolution, PF_R32_FLOAT, FClearValueBinding::White, ETextureCreateFlags::ShaderResource | ETextureCreateFlags::UAV), TEXT("Hair.DummyTransmittanceTextureForHitProxyId"));
@@ -4233,7 +4238,7 @@ void DrawHitProxies(
 }
 
 // Check if any simulated/skinned-bound groom has its positions updated (e.g. for invalidating the path-tracer accumulation)
-bool HasPositionsChanged(FRDGBuilder& GraphBuilder, const FViewInfo& View)
+bool HasPositionsChanged(FRDGBuilder& GraphBuilder, const FScene& Scene, const FViewInfo& View)
 {
 	if (View.HairStrandsMeshElements.IsEmpty())
 	{
@@ -4281,7 +4286,7 @@ bool HasPositionsChanged(FRDGBuilder& GraphBuilder, const FViewInfo& View)
 	// Compare current/previous and enqueue aggregated comparison
 	for (const FHairGroupPublicData* GroupData : GroupDatas)
 	{
-		AddHairStrandsHasPositionChangedPass(GraphBuilder, &View, GroupData, InvalidationUAV);
+		AddHairStrandsHasPositionChangedPass(GraphBuilder, &View, GroupData, Scene.HairStrandsSceneData.TransientResources, InvalidationUAV);
 	}
 
 	// Pull a 'ready' previous frame value
diff --git a/Engine/Source/Runtime/Renderer/Private/HairStrands/HairStrandsVoxelization.cpp b/Engine/Source/Runtime/Renderer/Private/HairStrands/HairStrandsVoxelization.cpp
--- a/Engine/Source/Runtime/Renderer/Private/HairStrands/HairStrandsVoxelization.cpp
+++ b/Engine/Source/Runtime/Renderer/Private/HairStrands/HairStrandsVoxelization.cpp
@@ -486,6 +486,7 @@ inline FIntVector CeilToInt(const FVector& V)
 static void AddAllocateVoxelPagesPass(
 	FRDGBuilder& GraphBuilder,
 	const FViewInfo& View,
+	FHairTransientResources& TransientResources,
 	FHairStrandsMacroGroupDatas& MacroGroupDatas,
 	FHairStrandsMacroGroupResources& MacroGroupResources,
 	const FIntVector PageCountResolution,
@@ -674,7 +675,7 @@ static void AddAllocateVoxelPagesPass(
 				FVoxelMarkValidPageIndex_PrepareCS::FParameters* Parameters = GraphBuilder.AllocParameters<FVoxelMarkValidPageIndex_PrepareCS::FParameters>();
 				Parameters->MaxClusterCount								= HairGroupData->GetClusterCount();
 				Parameters->MacroGroupId								= MacroGroup.MacroGroupId;
-				Parameters->GroupAABBsBuffer							= RegisterAsSRV(GraphBuilder, HairGroupData->GetGroupAABBBuffer());
+				Parameters->GroupAABBsBuffer							= TransientResources.GroupAABBSRV;
 				Parameters->ClusterAABBsBuffer							= RegisterAsSRV(GraphBuilder, HairGroupData->GetClusterAABBBuffer());
 				Parameters->MacroGroupVoxelAlignedAABBBuffer			= GraphBuilder.CreateSRV(MacroGroupResources.MacroGroupVoxelAlignedAABBsBuffer, PF_R32_SINT);
 				Parameters->PageIndexResolutionAndOffsetBuffer			= PageIndexResolutionAndOffsetBufferSRV;
@@ -834,6 +835,7 @@ static FHairStrandsVoxelResources AllocateVirtualVoxelResources(
 	FRDGBuilder& GraphBuilder,
 	const FViewInfo& View,
 	const FVector& PreViewStereoCorrection,
+	FHairTransientResources& TransientResources,
 	FHairStrandsMacroGroupDatas& MacroGroupDatas,
 	FHairStrandsMacroGroupResources& MacroGroupResources,
 	FRDGBufferRef& PageToPageIndexBuffer,
@@ -899,6 +901,7 @@ static FHairStrandsVoxelResources AllocateVirtualVoxelResources(
 	AddAllocateVoxelPagesPass(
 		GraphBuilder, 
 		View, 
+		TransientResources,
 		MacroGroupDatas,
 		MacroGroupResources,
 		Out.Parameters.Common.PageCountResolution,
@@ -1401,6 +1404,7 @@ void VoxelizeHairStrands(
 	FHairStrandsMacroGroupDatas& MacroGroupDatas = View.HairStrandsViewData.MacroGroupDatas;
 	FHairStrandsVoxelResources& VirtualVoxelResources = View.HairStrandsViewData.VirtualVoxelResources;
 	FHairStrandsMacroGroupResources& MacroGroupResources = View.HairStrandsViewData.MacroGroupResources;
+	FHairTransientResources* TransientResources = Scene->HairStrandsSceneData.TransientResources;
 
 	// Simple early out to check if voxelization is needed
 	if (!IsHairStrandsVoxelizationEnable() || MacroGroupDatas.Num() == 0)
@@ -1430,10 +1434,11 @@ void VoxelizeHairStrands(
 	RDG_EVENT_SCOPE(GraphBuilder, "HairStrandsVoxelization");
 	RDG_GPU_STAT_SCOPE(GraphBuilder, HairStrandsVoxelization);
 
+	check(TransientResources);
 	{
 		FRDGBufferRef PageToPageIndexBuffer = nullptr;
 		FHairStrandsViewStateData* HairStrandsViewStateData = View.ViewState ? &View.ViewState->HairStrandsViewStateData : nullptr;
-		VirtualVoxelResources = AllocateVirtualVoxelResources(GraphBuilder, View, PreViewStereoCorrection, MacroGroupDatas, MacroGroupResources, PageToPageIndexBuffer, HairStrandsViewStateData);
+		VirtualVoxelResources = AllocateVirtualVoxelResources(GraphBuilder, View, PreViewStereoCorrection, *TransientResources, MacroGroupDatas, MacroGroupResources, PageToPageIndexBuffer, HairStrandsViewStateData);
 
 		FRDGBufferRef ClearIndArgsBuffer = IndirectVoxelPageClear(GraphBuilder, View, VirtualVoxelResources);
 
diff --git a/Engine/Source/Runtime/Renderer/Private/PathTracing.cpp b/Engine/Source/Runtime/Renderer/Private/PathTracing.cpp
--- a/Engine/Source/Runtime/Renderer/Private/PathTracing.cpp
+++ b/Engine/Source/Runtime/Renderer/Private/PathTracing.cpp
@@ -2648,7 +2648,7 @@ void FDeferredShadingSceneRenderer::RenderPathTracing(
 	if (FirstTime ||
 		Config.IsDifferent(PathTracingState->LastConfig) ||
 		(!View.bIsOfflineRender && Config.IsDOFDifferent(PathTracingState->LastConfig)) ||
-		(!View.bIsOfflineRender && HairStrands::HasPositionsChanged(GraphBuilder, View)))
+		(!View.bIsOfflineRender && HairStrands::HasPositionsChanged(GraphBuilder, *Scene, View)))
 	{
 		// remember the options we used for next time
 		PathTracingState->LastConfig = Config;
diff --git a/Engine/Source/Runtime/Renderer/Private/ScenePrivate.h b/Engine/Source/Runtime/Renderer/Private/ScenePrivate.h
--- a/Engine/Source/Runtime/Renderer/Private/ScenePrivate.h
+++ b/Engine/Source/Runtime/Renderer/Private/ScenePrivate.h
@@ -1830,6 +1830,7 @@ class FHairStrandsSceneData
 {
 public:
 	TArray<FHairStrandsInstance*> RegisteredProxies;
+	FHairTransientResources* TransientResources = nullptr;
 };
 
 class FVolumetricLightmapInterpolation
diff --git a/Engine/Source/Runtime/Renderer/Public/HairStrandsInterface.h b/Engine/Source/Runtime/Renderer/Public/HairStrandsInterface.h
--- a/Engine/Source/Runtime/Renderer/Public/HairStrandsInterface.h
+++ b/Engine/Source/Runtime/Renderer/Public/HairStrandsInterface.h
@@ -221,8 +221,6 @@ public:
 	const FRDGExternalBuffer& GetDrawIndirectRasterComputeBuffer() const { return Culling->DrawIndirectRasterComputeBuffer; }
 	FRDGExternalBuffer& GetDrawIndirectBuffer() { return Culling->DrawIndirectBuffer; }
 	FRDGExternalBuffer& GetClusterAABBBuffer() { return Culling->ClusterAABBBuffer; }
-	FRDGExternalBuffer& GetGroupAABBBuffer() { return Culling->GroupAABBBuffer; }
-	const FRDGExternalBuffer& GetGroupAABBBuffer() const { return Culling->GroupAABBBuffer; }
 
 	const FRDGExternalBuffer& GetCulledCurveBuffer() const { return Culling->CulledCurveBuffer; }
 	const FRDGExternalBuffer& GetCulledVertexIdBuffer() const { return Culling->CulledVertexIdBuffer; }
@@ -238,8 +236,6 @@ public:
 	void SetClusterAABBValid(bool In) { Culling->bClusterAABBValid = In;  }
 	bool GetClusterAABBValid() const  { return Culling->bClusterAABBValid;  }
 
-	void SetGroupAABBValid(bool In) { Culling->bGroupAABBValid = In;  }
-
 	void SupportVoxelization(bool InVoxelize) { bSupportVoxelization = InVoxelize; }
 	bool DoesSupportVoxelization() const { return bSupportVoxelization; }
 
@@ -365,8 +361,6 @@ public:
 
 		/* Hair Cluster & Hair Group bounding box buffer */
 		FRDGExternalBuffer ClusterAABBBuffer;
-		FRDGExternalBuffer GroupAABBBuffer;
-		bool bGroupAABBValid = false;
 		bool bClusterAABBValid = false;
 
 		/* Culling & LODing results for a hair group */ // Better to be transient?
@@ -427,7 +421,22 @@ public:
 };
 
 ////////////////////////////////////////////////////////////////////////////////////////////////////
+// Transient resources 
+// Resources are indexed by hair instance's registered index
+struct FHairTransientResources
+{
+	FRDGBufferRef    GroupAABBBuffer = nullptr;
+	FRDGBufferUAVRef GroupAABBUAV = nullptr;
+	FRDGBufferSRVRef GroupAABBSRV = nullptr;
+
+	FRDGBufferRef    IndirectDispatchArgsBuffer = nullptr;
+	FRDGBufferUAVRef IndirectDispatchArgsUAV = nullptr;
+	FRDGBufferSRVRef IndirectDispatchArgsSRV = nullptr;
 
+	TBitArray<> bIsGroupAABBValid;
+};
+
+////////////////////////////////////////////////////////////////////////////////////////////////////
 // API for enabling/disabling the various geometry representation
 enum class EHairStrandsShaderType
 {
@@ -502,6 +511,7 @@ struct FHairStrandsBookmarkParameters
 	FHairStrandsInstances VisibleStrands; // Primary & Shadow
 	FHairStrandsInstances VisibleCardsOrMeshes_Primary;
 	FHairStrandsInstances VisibleCardsOrMeshes_Shadow;
+	FHairTransientResources* TransientResources = nullptr;
 
 	FHairStrandsInstances* Instances = nullptr;
 	TArray<EHairInstanceVisibilityType> InstancesVisibilityType;
@@ -515,4 +525,4 @@ struct FHairStrandsBookmarkParameters
 };
 
 typedef void (*THairStrandsBookmarkFunction)(FRDGBuilder* GraphBuilder, EHairStrandsBookmark Bookmark, FHairStrandsBookmarkParameters& Parameters);
-RENDERER_API void RegisterBookmarkFunction(THairStrandsBookmarkFunction Bookmark);
+RENDERER_API void RegisterBookmarkFunction(THairStrandsBookmarkFunction Bookmark);
\ No newline at end of file
