Skip UAV barrier during hair voxel page allocation.

diff --git a/Engine/Source/Runtime/Renderer/Private/HairStrands/HairStrandsVoxelization.cpp b/Engine/Source/Runtime/Renderer/Private/HairStrands/HairStrandsVoxelization.cpp
--- a/Engine/Source/Runtime/Renderer/Private/HairStrands/HairStrandsVoxelization.cpp
+++ b/Engine/Source/Runtime/Renderer/Private/HairStrands/HairStrandsVoxelization.cpp
@@ -193,7 +193,8 @@ static void AddVirtualVoxelInjectOpaquePass(
 	FRDGBuilder& GraphBuilder,
 	const FViewInfo& View,
 	const FHairStrandsVoxelResources& VoxelResources,
-	const FHairStrandsMacroGroupData& MacroGroup)
+	const FHairStrandsMacroGroupData& MacroGroup,
+	FRDGTextureUAVRef OutTexturePage)
 {
 	FSceneTextureParameters SceneTextures = GetSceneTextureParameters(GraphBuilder, View);
 
@@ -211,7 +212,7 @@ static void AddVirtualVoxelInjectOpaquePass(
 	Parameters->SceneDepthResolution = SceneTextures.SceneDepthTexture->Desc.Extent;
 	Parameters->SceneTextures = SceneTextures;
 	Parameters->MacroGroupId = MacroGroup.MacroGroupId;
-	Parameters->OutPageTexture = GraphBuilder.CreateUAV(VoxelResources.PageTexture);
+	Parameters->OutPageTexture = OutTexturePage;
 	Parameters->IndirectDispatchArgs = VoxelResources.IndirectArgsBuffer;
 	TShaderMapRef<FVirtualVoxelInjectOpaqueCS> ComputeShader(View.ShaderMap);
 	const FGlobalShaderMap* GlobalShaderMap = View.ShaderMap;
@@ -1452,11 +1453,12 @@ void VoxelizeHairStrands(
 
 		if (GHairVoxelInjectOpaqueDepthEnable > 0)
 		{
+			FRDGTextureUAVRef PageTextureUAV = GraphBuilder.CreateUAV(VirtualVoxelResources.PageTexture, ERDGUnorderedAccessViewFlags::SkipBarrier);
 			for (FHairStrandsMacroGroupData& MacroGroup : MacroGroupDatas)
 			{
 				if (MacroGroup.bSupportVoxelization)
 				{
-					AddVirtualVoxelInjectOpaquePass(GraphBuilder, View, VirtualVoxelResources, MacroGroup);
+					AddVirtualVoxelInjectOpaquePass(GraphBuilder, View, VirtualVoxelResources, MacroGroup, PageTextureUAV);
 				}
 			}
 		}
