Removed PostProcess local light buffers when r.Mobile.Forward.EnableLocalLights=2
When r.Mobile.Forward.EnableLocalLights=2 and Depth Prepass is not enabled we will merge the local lights in the BasePass and call the BRDF once
#localization none
#tests  Android game

diff --git a/Engine/Shaders/Private/MobileBasePassPixelShader.usf b/Engine/Shaders/Private/MobileBasePassPixelShader.usf
--- a/Engine/Shaders/Private/MobileBasePassPixelShader.usf
+++ b/Engine/Shaders/Private/MobileBasePassPixelShader.usf
@@ -864,7 +864,7 @@ void Main(
 #endif
 
 	// Local lights
-#if  PREPASS_LOCAL_LIGHTS_MOBILE
+#if  MERGED_LOCAL_LIGHTS_MOBILE == 1
 	{
 		FDeferredLightData MergedLightData = (FDeferredLightData)0;
 		half3 LocalLightA = MobileSceneTextures.LocalLightTextureA.Load(int3(SvPosition.xy, 0)).xyz; 
@@ -888,6 +888,47 @@ void Main(
 		FLightAccumulator NewLighting = AccumulateDynamicLighting(MaterialParameters.WorldPosition_CamRelative, CameraVector, GBuffer, 1, GBuffer.ShadingModelID, MergedLightData, LightAttenuation, 0, uint2(0, 0), SurfaceShadow);
 		DirectLighting = LightAccumulator_Add(DirectLighting, NewLighting);
 	}
+#elif  MERGED_LOCAL_LIGHTS_MOBILE == 2
+	{
+#if MOBILE_MULTI_VIEW
+		const uint EyeIndex = BasePassInterpolants.MultiViewId;
+#else
+		const uint EyeIndex = 0;
+#endif
+
+		const FCulledLightsGridData CulledLightGridData = GetCulledLightsGrid(GridIndex, EyeIndex);
+
+		float3 MergedLightColor = uint3(0.f, 0.f , 0.f);
+		float3 MergedLightL = float3(0,0,0);
+		float MergedSpecularScale = 0.f;
+		float MergedTotalWeight = 0;
+		MergeLocalLights(CulledLightGridData, MaterialParameters.WorldPosition_CamRelative, EyeIndex, false, MergedLightColor, MergedLightL, MergedSpecularScale, MergedTotalWeight);
+
+		if(MergedTotalWeight > 0.f)
+		{
+			MergedLightL = MergedLightL / MergedTotalWeight;
+			MergedSpecularScale = MergedSpecularScale / MergedTotalWeight;
+
+			FDeferredLightData MergedLightData = (FDeferredLightData)0;
+			MergedLightData.SourceRadius	= 0;
+			MergedLightData.SourceLength	= 0;
+			MergedLightData.Tangent			= 0;
+			MergedLightData.SoftSourceRadius= 0;
+			MergedLightData.SpecularScale	= MergedSpecularScale;
+			MergedLightData.bSpotLight		= false;
+			MergedLightData.bRectLight		= false;
+			MergedLightData.Color = MergedLightColor;
+			MergedLightData.bRadialLight = false;
+
+			half4 LightAttenuation = half4(1, 1, 1, 1); 
+			float SurfaceShadow = 0;
+
+			MergedLightData.Direction = MergedLightL;
+
+			FLightAccumulator NewLighting = AccumulateDynamicLighting(MaterialParameters.WorldPosition_CamRelative, CameraVector, GBuffer, 1, GBuffer.ShadingModelID, MergedLightData, LightAttenuation, 0, uint2(0, 0), SurfaceShadow);
+			DirectLighting = LightAccumulator_Add(DirectLighting, NewLighting);
+		}
+	}
 #elif  (!MATERIAL_SHADINGMODEL_SINGLELAYERWATER && ENABLE_CLUSTERED_LIGHTS)
 	{
 #if MOBILE_MULTI_VIEW
diff --git a/Engine/Shaders/Private/MobileLightingCommon.ush b/Engine/Shaders/Private/MobileLightingCommon.ush
--- a/Engine/Shaders/Private/MobileLightingCommon.ush
+++ b/Engine/Shaders/Private/MobileLightingCommon.ush
@@ -601,4 +601,71 @@ void AccumulateLightGridLocalLighting(const FCulledLightsGridData InLightGridDat
 		DirectLighting = LightAccumulator_Add(DirectLighting, NewLighting);
 	}
 }
+#endif
+
+#if ENABLE_CLUSTERED_LIGHTS || MERGED_LOCAL_LIGHTS_MOBILE
+void MergeLocalLights(const FCulledLightsGridData InLightGridData
+	, float3 TranslatedWorldPosition
+	, uint EyeIndex
+	, float bSupportLightFunctions
+	, inout float3 MergedLightColor
+	, inout float3 MergedLightL
+	, inout float MergedSpecularScale
+	, inout float MergedTotalWeight)
+{
+	MergedLightColor = uint3(0.f, 0.f , 0.f);
+	MergedLightL = float3(0,0,0);
+	MergedTotalWeight = 0;
+	MergedSpecularScale = 0.f;
+	
+	uint NumLocalLights = min(InLightGridData.NumLocalLights, GetMaxLightsPerCell(EyeIndex));
+
+	LOOP
+	for (uint LocalLightListIndex = 0; LocalLightListIndex < NumLocalLights; LocalLightListIndex++)
+	{
+		const FLocalLightData LocalLight = GetLocalLightData(InLightGridData.DataStartIndex + LocalLightListIndex, EyeIndex);
+
+		// extra-early out since we know light grid is sloppy and all lights in list are radial (have a range)
+		// appears useless
+		float InvLightRadiusSq = LocalLight.LightPositionAndInvRadius.w * LocalLight.LightPositionAndInvRadius.w;
+		float DistLight = length2(TranslatedWorldPosition - LocalLight.LightPositionAndInvRadius.xyz) * InvLightRadiusSq;
+		
+		if (DistLight > 1.0f)
+		{
+			continue;
+		}
+
+		FDeferredLightData LightData = ConvertToDeferredLight_Mobile(LocalLight);
+		LightData.SpecularScale	= UnpackLightSpecularScale(LocalLight);
+
+		float3 L = LightData.Direction;	// Already normalized
+		float3 ToLight = L;
+		float3 MaskedLightColor = LightData.Color;
+		float LightMask = 1;
+		if (LightData.bRadialLight)
+		{
+			LightMask = GetLocalLightAttenuation( TranslatedWorldPosition, LightData, ToLight, L );
+			MaskedLightColor *= LightMask;
+		}
+
+		float WeightLight = 1 - DistLight;
+		if (LightData.bSpotLight)
+		{
+			WeightLight *= SpotAttenuation(L, -LightData.Direction, LightData.SpotAngles);
+		}
+		
+		if( LightMask > 0)
+		{
+			MergedSpecularScale += LightData.SpecularScale * WeightLight;
+			MergedTotalWeight += WeightLight;
+			// For lights with light functions we only only use the compute the direction (OutColorB)
+			// The Color is added in additional draws using MainLightFunction()
+			if (!bSupportLightFunctions || !UnpackHasLightFunction(LocalLight))
+			{
+				MergedLightColor += MaskedLightColor;
+			}
+			MergedLightL += L *  WeightLight;
+		}
+	}
+}
 #endif
\ No newline at end of file
diff --git a/Engine/Shaders/Private/MobileLocalLightsBuffer.usf b/Engine/Shaders/Private/MobileLocalLightsBuffer.usf
--- a/Engine/Shaders/Private/MobileLocalLightsBuffer.usf
+++ b/Engine/Shaders/Private/MobileLocalLightsBuffer.usf
@@ -7,6 +7,8 @@
 #include "Common.ush"
 #include "SceneTextureParameters.ush"
 #include "DeferredLightingCommon.ush"
+#define MERGED_LOCAL_LIGHTS_MOBILE 1
+#include "MobileLightingCommon.ush"
 #endif
 
 #if LIGHT_FUNCTION
@@ -90,11 +92,7 @@ void MainVS(
 
 void Main(float4 SvPosition : SV_POSITION
 	, out float3 OutColorA : SV_Target0
-#if POST_PROCESS_LOCAL_LIGHTS == 0
 	, out float4 OutColorB : SV_Target1
-#elif GENERATE_LIGHT_FUNCTION_DEPTH_STENCIL && USE_SCENE_DEPTH_AUX
-	, out float OutDepth : SV_Depth
-#endif
 )
 {
 	const uint2 PixelPos = SvPosition.xy;
@@ -111,70 +109,18 @@ void Main(float4 SvPosition : SV_POSITION
 	uint GridIndex = ComputeLightGridCellIndex(uint2(PixelPos.x, PixelPos.y), SceneDepth);
 	const FCulledLightsGridData CulledLightGridData = GetCulledLightsGrid(GridIndex, EyeIndex);
 
-	uint NumLocalLights = min(CulledLightGridData.NumLocalLights, GetMaxLightsPerCell(EyeIndex));
-
-	float3 Color = uint3(0.f, 0.f , 0.f);
-	float3 SavedL = float3(0,0,0);
-	float TotalWeight = 0;
-	float SpecularScale = 0.f;
-
-	LOOP
-	for (uint LocalLightListIndex = 0; LocalLightListIndex < NumLocalLights; LocalLightListIndex++)
-	{
-		const FLocalLightData LocalLight = GetLocalLightData(CulledLightGridData.DataStartIndex + LocalLightListIndex, EyeIndex);
-
-		// extra-early out since we know light grid is sloppy and all lights in list are radial (have a range)
-		// appears useless
-		float InvLightRadiusSq = LocalLight.LightPositionAndInvRadius.w * LocalLight.LightPositionAndInvRadius.w;
-		float DistLight = length2(TranslatedWorldPosition - LocalLight.LightPositionAndInvRadius.xyz) * InvLightRadiusSq;
-		
-		if (DistLight > 1.0f)
-		{
-			continue;
-		}
-
-		FDeferredLightData LightData = ConvertToDeferredLight_Mobile(LocalLight);
-		LightData.SpecularScale	= UnpackLightSpecularScale(LocalLight);
-
-		float3 L = LightData.Direction;	// Already normalized
-		float3 ToLight = L;
-		float3 MaskedLightColor = LightData.Color;
-		float LightMask = 1;
-		if (LightData.bRadialLight)
-		{
-			LightMask = GetLocalLightAttenuation( TranslatedWorldPosition, LightData, ToLight, L );
-			MaskedLightColor *= LightMask;
-		}
-
-		float WeightLight = 1 - DistLight;
-		if (LightData.bSpotLight)
-		{
-			WeightLight *= SpotAttenuation(L, -LightData.Direction, LightData.SpotAngles);
-		}
-		
-		if( LightMask > 0)
-		{
-			SpecularScale += LightData.SpecularScale * WeightLight;
-
-			TotalWeight += WeightLight;
-			// For lights with light functions we only only use the compute the direction (OutColorB)
-			// The Color is added in additional draws using MainLightFunction()
-			if (!UnpackHasLightFunction(LocalLight))
-			{
-				Color += MaskedLightColor;
-			}
+	float3 MergedLightColor = uint3(0.f, 0.f , 0.f);
+	float3 MergedLightL = float3(0,0,0);
+	float MergedSpecularScale = 0.f;
+	float MergedTotalWeight = 0;
+	MergeLocalLights(CulledLightGridData, TranslatedWorldPosition, EyeIndex, true, MergedLightColor, MergedLightL, MergedSpecularScale, MergedTotalWeight);
 
-			SavedL += L *  WeightLight;
-		}
-	}
-
-#if POST_PROCESS_LOCAL_LIGHTS == 0
-	if (TotalWeight > 0)
+	if (MergedTotalWeight > 0)
 	{
-		SavedL = (SavedL / TotalWeight);
-		OutColorA = float3(Color.x, Color.y, Color.z);
-		OutColorB.rgb = Pack1212To888(UnitVectorToOctahedron(SavedL) * 0.5 + 0.5);
-		OutColorB.a = SpecularScale / TotalWeight;
+		MergedLightL = (MergedLightL / MergedTotalWeight);
+		OutColorA = float3(MergedLightColor.x, MergedLightColor.y, MergedLightColor.z);
+		OutColorB.rgb = Pack1212To888(UnitVectorToOctahedron(MergedLightL) * 0.5 + 0.5);
+		OutColorB.a = MergedSpecularScale / MergedTotalWeight;
 	}
 	else
 	{
@@ -185,29 +131,6 @@ void Main(float4 SvPosition : SV_POSITION
 			clip(-1);
 		#endif
 	}
-#else
-	if (TotalWeight > 0)
-	{
-		SavedL = (SavedL / TotalWeight);
-		float3 N = normalize(cross(ddx(TranslatedWorldPosition.xyz), ddy(TranslatedWorldPosition.xyz)));
-		float3 L = SavedL;	// Already normalized
-		float NoL = max(0, dot(N, L));
-		float3 Diffuse =  (1 / PI) * NoL * Color;
-		OutColorA = 1 + Diffuse;
-		#if GENERATE_LIGHT_FUNCTION_DEPTH_STENCIL && USE_SCENE_DEPTH_AUX
-			// Copy depth value from scene depth aux to the real depth buffer
-			OutDepth = DeviceZ;
-		#endif
-	}
-	else
-	{
-		OutColorA = float3(1, 1, 1);
-		#if GENERATE_LIGHT_FUNCTION_DEPTH_STENCIL
-			// Clip to avoid writing to stencil
-			clip(-1);
-		#endif
-	}
-#endif
 }
 
 #else // !LIGHT_FUNCTION
@@ -258,14 +181,7 @@ void MainLightFunction(
 
 	float3 Color = LightData.Color * LightMask * EncodedLightAttenuation;
 
-#if POST_PROCESS_LOCAL_LIGHTS == 0
 	OutColorA = Color;
-#else
-	float3 N = normalize(cross(ddx(TranslatedWorldPosition.xyz), ddy(TranslatedWorldPosition.xyz)));
-	float NoL = max(0, dot(N, L));
-	float3 Diffuse =  (1 / PI) * NoL * Color;
-	OutColorA = Diffuse;
-#endif
 }
 #endif // !LIGHT_FUNCTION
 
diff --git a/Engine/Source/Runtime/Renderer/Private/MobileBasePass.cpp b/Engine/Source/Runtime/Renderer/Private/MobileBasePass.cpp
--- a/Engine/Source/Runtime/Renderer/Private/MobileBasePass.cpp
+++ b/Engine/Source/Runtime/Renderer/Private/MobileBasePass.cpp
@@ -21,12 +21,12 @@ bool MobileLocalLightsBufferEnabled(const FStaticShaderPlatform Platform)
 	return MobileForwardEnablePrepassLocalLightsIniValue.Get(Platform) == 2;
 }
 
-bool MobileLocalLightsBufferPrepassEnabled(const FStaticShaderPlatform Platform)
+bool MobileMergeLocalLightsInPrepassEnabled(const FStaticShaderPlatform Platform)
 {
 	return MobileLocalLightsBufferEnabled(Platform) && MobileUsesFullDepthPrepass(Platform);
 }
 
-bool MobileLocalLightsBufferPostprocessEnabled(const FStaticShaderPlatform Platform)
+bool MobileMergeLocalLightsInBasepassEnabled(const FStaticShaderPlatform Platform)
 {
 	return MobileLocalLightsBufferEnabled(Platform) && !MobileUsesFullDepthPrepass(Platform);
 }
@@ -57,10 +57,7 @@ EMobileLocalLightSetting GetMobileForwardLocalLightSetting(EShaderPlatform Shade
 		}
 		else if (MobileForwardLocalLightsIniValue == 2)
 		{
-			if (MobileUsesFullDepthPrepass(ShaderPlatform))
-			{
-				return EMobileLocalLightSetting::LOCAL_LIGHTS_BUFFER;
-			}
+			return EMobileLocalLightSetting::LOCAL_LIGHTS_BUFFER;
 		}
 	}
 
diff --git a/Engine/Source/Runtime/Renderer/Private/MobileBasePassRendering.h b/Engine/Source/Runtime/Renderer/Private/MobileBasePassRendering.h
--- a/Engine/Source/Runtime/Renderer/Private/MobileBasePassRendering.h
+++ b/Engine/Source/Runtime/Renderer/Private/MobileBasePassRendering.h
@@ -30,8 +30,8 @@
 #include "CompositionLighting/PostProcessDeferredDecals.h"
 
 bool MobileLocalLightsBufferEnabled(const FStaticShaderPlatform Platform);
-bool MobileLocalLightsBufferPrepassEnabled(const FStaticShaderPlatform Platform);
-bool MobileLocalLightsBufferPostprocessEnabled(const FStaticShaderPlatform Platform);
+bool MobileMergeLocalLightsInPrepassEnabled(const FStaticShaderPlatform Platform);
+bool MobileMergeLocalLightsInBasepassEnabled(const FStaticShaderPlatform Platform);
 
 struct FMobileBasePassTextures
 {
@@ -417,11 +417,25 @@ public:
 		OutEnvironment.SetDefine(TEXT("ENABLE_CLUSTERED_LIGHTS"), (LocalLightSetting == EMobileLocalLightSetting::LOCAL_LIGHTS_ENABLED) ? 1u : 0u);
 
 		// Translucent materials don't write to depth so cannot use prepass
-		OutEnvironment.SetDefine(TEXT("PREPASS_LOCAL_LIGHTS_MOBILE"), (LocalLightSetting == EMobileLocalLightSetting::LOCAL_LIGHTS_BUFFER) ? 1u: 0u);
+		uint32 MergedLocalLights = 0u;
+		if (LocalLightSetting == EMobileLocalLightSetting::LOCAL_LIGHTS_BUFFER)
+		{
+			if (MobileMergeLocalLightsInPrepassEnabled(Parameters.Platform))
+			{
+				MergedLocalLights = 1u;
+			}
+			else if (MobileMergeLocalLightsInBasepassEnabled(Parameters.Platform))
+			{
+				MergedLocalLights = 2u;
+			}
+		}
+		OutEnvironment.SetDefine(TEXT("MERGED_LOCAL_LIGHTS_MOBILE"), MergedLocalLights);
 		OutEnvironment.SetDefine(TEXT("ENABLE_CLUSTERED_REFLECTION"), bEnableClusteredReflections ? 1u : 0u);
 		OutEnvironment.SetDefine(TEXT("USE_SHADOWMASKTEXTURE"), bMobileUsesShadowMaskTexture && !bTranslucentMaterial ? 1u : 0u);
 		OutEnvironment.SetDefine(TEXT("ENABLE_DBUFFER_TEXTURES"), Parameters.MaterialParameters.MaterialDomain == MD_Surface ? 1u : 0u);
 
+		const bool bSupportsDepthPrepass = MobileUsesFullDepthPrepass(Parameters.Platform);
+		OutEnvironment.SetDefine(TEXT("SUPPORTS_DEPTH_PREPASS"), bSupportsDepthPrepass ? 1u : 0u);
 		EMobileTranslucentColorTransmittanceMode TranslucentColorTransmittanceMode = EMobileTranslucentColorTransmittanceMode::DEFAULT;
 		if (Parameters.MaterialParameters.ShadingModels.HasShadingModel(MSM_ThinTranslucent))
 		{
diff --git a/Engine/Source/Runtime/Renderer/Private/MobileLocalLightsBuffer.cpp b/Engine/Source/Runtime/Renderer/Private/MobileLocalLightsBuffer.cpp
--- a/Engine/Source/Runtime/Renderer/Private/MobileLocalLightsBuffer.cpp
+++ b/Engine/Source/Runtime/Renderer/Private/MobileLocalLightsBuffer.cpp
@@ -16,7 +16,8 @@ static bool CompileShaderPermutationsForMobileLocalLightsBuffer(const FStaticSha
 {
 	return !IsMobileDeferredShadingEnabled(Platform) && 
 			IsMobilePlatform(Platform) && 
-			MobileLocalLightsBufferEnabled(Platform);
+			MobileLocalLightsBufferEnabled(Platform) &&
+			MobileUsesFullDepthPrepass(Platform);
 }
 
 const int32 GLocalLightPrepassTileSizeX = 8;
@@ -89,18 +90,8 @@ class FLocalLightBufferPS : public FGlobalShader
 	{
 		FForwardLightingParameters::ModifyCompilationEnvironment(Parameters.Platform, OutEnvironment);
 		OutEnvironment.SetDefine(TEXT("LIGHT_FUNCTION"), 0);
-
-		if (MobileLocalLightsBufferPrepassEnabled(Parameters.Platform))
-		{
-			OutEnvironment.SetRenderTargetOutputFormat(0, PF_FloatR11G11B10);
-			OutEnvironment.SetRenderTargetOutputFormat(1, PF_A2B10G10R10);
-			OutEnvironment.SetDefine(TEXT("POST_PROCESS_LOCAL_LIGHTS"), 0);
-		}
-		else
-		{
-			OutEnvironment.SetRenderTargetOutputFormat(0, PF_FloatR11G11B10);
-			OutEnvironment.SetDefine(TEXT("POST_PROCESS_LOCAL_LIGHTS"), 1);
-		}
+		OutEnvironment.SetRenderTargetOutputFormat(0, PF_FloatR11G11B10);
+		OutEnvironment.SetRenderTargetOutputFormat(1, PF_B8G8R8A8);
 	}
 };
 
@@ -145,18 +136,8 @@ public:
 		OutEnvironment.SetDefine(TEXT("LIGHT_FUNCTION"), 1);
 		OutEnvironment.SetDefine(TEXT("COMPUTE_SHADER"), 0);
 		OutEnvironment.SetDefine(TEXT("SUBSTRATE_INLINE_SHADING"), 1);
-		
-		if (MobileLocalLightsBufferPrepassEnabled(Parameters.Platform))
-		{
-			OutEnvironment.SetRenderTargetOutputFormat(0, PF_FloatR11G11B10);
-			OutEnvironment.SetRenderTargetOutputFormat(1, PF_A2B10G10R10);
-			OutEnvironment.SetDefine(TEXT("POST_PROCESS_LOCAL_LIGHTS"), 0);
-		}
-		else
-		{
-			OutEnvironment.SetRenderTargetOutputFormat(0, PF_FloatR11G11B10);
-			OutEnvironment.SetDefine(TEXT("POST_PROCESS_LOCAL_LIGHTS"), 1);
-		}
+		OutEnvironment.SetRenderTargetOutputFormat(0, PF_FloatR11G11B10);
+		OutEnvironment.SetRenderTargetOutputFormat(1, PF_B8G8R8A8);
 	}
 
 	void SetParameters(FRHIBatchedShaderParameters& BatchedParameters, const FViewInfo& View, const FMaterialRenderProxy* MaterialProxy, const FMaterial& Material)
@@ -218,10 +199,9 @@ using LightFunctionMaterialPassDepthStencilState = TStaticDepthStencilState<
 	true, CF_Equal, SO_Keep, SO_Keep, SO_Keep,
 	STENCIL_MOBILE_LIGHTFUNCTION_MASK, 0>;
 
-void FMobileSceneRenderer::RenderMobileLocalLightsBuffer(FRDGBuilder& GraphBuilder, FSceneTextures& SceneTextures, bool bIsPrepass, const FSortedLightSetSceneInfo& SortedLights)
+void FMobileSceneRenderer::RenderMobileLocalLightsBuffer(FRDGBuilder& GraphBuilder, FSceneTextures& SceneTextures, const FSortedLightSetSceneInfo& SortedLights)
 {
 	if (!CompileShaderPermutationsForMobileLocalLightsBuffer(ShaderPlatform) ||
-		(bIsPrepass != MobileLocalLightsBufferPrepassEnabled(ShaderPlatform)) || 
 		IsMobileDeferredShadingEnabled(ShaderPlatform))
 	{
 		return;
@@ -290,15 +270,9 @@ void FMobileSceneRenderer::RenderMobileLocalLightsBuffer(FRDGBuilder& GraphBuild
 
 		{
 			FLocalLightBufferPrepassParameters* PassParameters = GraphBuilder.AllocParameters<FLocalLightBufferPrepassParameters>();
-			if (bIsPrepass)
-			{
-				PassParameters->RenderTargets[0] = FRenderTargetBinding(SceneTextures.MobileLocalLightTextureA, ERenderTargetLoadAction::EClear);
-				PassParameters->RenderTargets[1] = FRenderTargetBinding(SceneTextures.MobileLocalLightTextureB, ERenderTargetLoadAction::EClear);
-			}
-			else
-			{
-				PassParameters->RenderTargets[0] = FRenderTargetBinding(SceneTextures.Color.Resolve, ERenderTargetLoadAction::ELoad);
-			}
+			PassParameters->RenderTargets[0] = FRenderTargetBinding(SceneTextures.MobileLocalLightTextureA, ERenderTargetLoadAction::EClear);
+			PassParameters->RenderTargets[1] = FRenderTargetBinding(SceneTextures.MobileLocalLightTextureB, ERenderTargetLoadAction::EClear);
+
 
 			bool bRestoreDepthBuffer = false;
 			if (bRenderLightFunctions)
@@ -332,24 +306,17 @@ void FMobileSceneRenderer::RenderMobileLocalLightsBuffer(FRDGBuilder& GraphBuild
 			auto PixelShader = View.ShaderMap->GetShader<FLocalLightBufferPS>(PermutationVectorPS);
 
 			GraphBuilder.AddPass(
-				RDG_EVENT_NAME("RenderMobileLocalLightsBuffer %s", bIsPrepass ? TEXT("Prepass") : TEXT("PostProcess")),
+				RDG_EVENT_NAME("RenderMobileLocalLightsBuffer Prepass"),
 				PassParameters,
 				ERDGPassFlags::Raster,
-				[PassParameters, VertexShader, PixelShader, &View, GroupSize, bIsPrepass, &SortedLights, bRenderLightFunctions, bRestoreDepthBuffer](FRHICommandList& RHICmdList)
+				[PassParameters, VertexShader, PixelShader, &View, GroupSize, &SortedLights, bRenderLightFunctions, bRestoreDepthBuffer](FRHICommandList& RHICmdList)
 				{
 					FGraphicsPipelineStateInitializer GraphicsPSOInit;
 					RHICmdList.ApplyCachedRenderTargets(GraphicsPSOInit);
 
 					RHICmdList.SetViewport(View.ViewRect.Min.X, View.ViewRect.Min.Y, 0.0f, View.ViewRect.Max.X, View.ViewRect.Max.Y, 1.0f);
 					GraphicsPSOInit.RasterizerState = TStaticRasterizerState<FM_Solid, CM_None>::GetRHI();
-					if (bIsPrepass)
-					{
-						GraphicsPSOInit.BlendState = TStaticBlendState<>::GetRHI();
-					}
-					else
-					{
-						GraphicsPSOInit.BlendState = TStaticBlendState<CW_RGB, BO_Add, BF_DestColor, BF_Zero>::GetRHI();
-					}
+					GraphicsPSOInit.BlendState = TStaticBlendState<>::GetRHI();
 					
 					// The main non-lightfunction pass creates a stencil mask of the lit area
 					// Later lightfunction passes use it for stencil test to avoid redundant PS execution
diff --git a/Engine/Source/Runtime/Renderer/Private/MobileShadingRenderer.cpp b/Engine/Source/Runtime/Renderer/Private/MobileShadingRenderer.cpp
--- a/Engine/Source/Runtime/Renderer/Private/MobileShadingRenderer.cpp
+++ b/Engine/Source/Runtime/Renderer/Private/MobileShadingRenderer.cpp
@@ -586,7 +586,7 @@ void FMobileSceneRenderer::InitViews(
 	}
 
 	// When we capturing scene depth, use a more precise format for SceneDepthAux as it will be used as a source DepthTexture
-	if (bSceneDepthCapture || MobileLocalLightsBufferPostprocessEnabled(ShaderPlatform))
+	if (bSceneDepthCapture)
 	{
 		SceneTexturesConfig.bPreciseDepthAux = true;
 	}
@@ -1102,7 +1102,7 @@ void FMobileSceneRenderer::Render(FRDGBuilder& GraphBuilder)
 
 		if (bRendererOutputFinalSceneColor)
 		{
-			RenderMobileLocalLightsBuffer(GraphBuilder, SceneTextures, true, SortedLightSet);
+			RenderMobileLocalLightsBuffer(GraphBuilder, SceneTextures, SortedLightSet);
 		}
 
 		if (bRendererOutputFinalSceneColor)
@@ -1161,8 +1161,6 @@ void FMobileSceneRenderer::Render(FRDGBuilder& GraphBuilder)
 	
 		GraphBuilder.SetCommandListStat(GET_STATID(STAT_CLMM_Post));
 
-		RenderMobileLocalLightsBuffer(GraphBuilder, SceneTextures, false, SortedLightSet);
-
 		FRendererModule& RendererModule = static_cast<FRendererModule&>(GetRendererModule());
 		RendererModule.RenderPostOpaqueExtensions(GraphBuilder, Views, SceneTextures);
 
diff --git a/Engine/Source/Runtime/Renderer/Private/SceneCore.cpp b/Engine/Source/Runtime/Renderer/Private/SceneCore.cpp
--- a/Engine/Source/Runtime/Renderer/Private/SceneCore.cpp
+++ b/Engine/Source/Runtime/Renderer/Private/SceneCore.cpp
@@ -159,11 +159,10 @@ extern bool ShouldCreateObjectShadowForStationaryLight(const FLightSceneInfo* Li
 static bool MobileRequiresStaticMeshUpdateOnLocalLightChange(const FStaticShaderPlatform Platform)
 {
 	extern bool MobileLocalLightsUseSinglePermutation();
-	extern bool MobileLocalLightsBufferPostprocessEnabled(const FStaticShaderPlatform Platform);
 	
 	if (!IsMobileDeferredShadingEnabled(Platform))
 	{
-		return MobileForwardEnableLocalLights(Platform) && !(MobileLocalLightsUseSinglePermutation() || MobileLocalLightsBufferPostprocessEnabled(Platform));
+		return MobileForwardEnableLocalLights(Platform) && !MobileLocalLightsUseSinglePermutation();
 	}
 	return false;
 }
diff --git a/Engine/Source/Runtime/Renderer/Private/SceneRendering.h b/Engine/Source/Runtime/Renderer/Private/SceneRendering.h
--- a/Engine/Source/Runtime/Renderer/Private/SceneRendering.h
+++ b/Engine/Source/Runtime/Renderer/Private/SceneRendering.h
@@ -2651,7 +2651,7 @@ protected:
 	void RenderMaskedPrePass(FRHICommandList& RHICmdList, const FViewInfo& View);
 	void RenderFullDepthPrepass(FRDGBuilder& GraphBuilder, TArrayView<FViewInfo> InViews, FSceneTextures& SceneTextures, bool bIsSceneCaptureRenderPass=false);
 
-	void RenderMobileLocalLightsBuffer(FRDGBuilder& GraphBuilder, FSceneTextures& SceneTextures, bool bIsPrepass, const FSortedLightSetSceneInfo& SortedLights);
+	void RenderMobileLocalLightsBuffer(FRDGBuilder& GraphBuilder, FSceneTextures& SceneTextures, const FSortedLightSetSceneInfo& SortedLights);
 
 	/** Renders the opaque base pass for mobile. */
 	void RenderMobileBasePass(FRHICommandList& RHICmdList, const FViewInfo& View, const FInstanceCullingDrawParams* InstanceCullingDrawParams);
diff --git a/Engine/Source/Runtime/Renderer/Private/SceneTextures.cpp b/Engine/Source/Runtime/Renderer/Private/SceneTextures.cpp
--- a/Engine/Source/Runtime/Renderer/Private/SceneTextures.cpp
+++ b/Engine/Source/Runtime/Renderer/Private/SceneTextures.cpp
@@ -618,8 +618,8 @@ void FSceneTextures::InitializeViewFamily(FRDGBuilder& GraphBuilder, FViewFamily
 	}
 #endif
 
-	extern bool MobileLocalLightsBufferPrepassEnabled(const FStaticShaderPlatform Platform);
-	if(MobileLocalLightsBufferPrepassEnabled(Config.ShaderPlatform))
+	extern bool MobileMergeLocalLightsInPrepassEnabled(const FStaticShaderPlatform Platform);
+	if(MobileMergeLocalLightsInPrepassEnabled(Config.ShaderPlatform))
 	{
 		FRDGTextureDesc MobileLocalLightTextureADesc = FRDGTextureDesc::Create2D(Config.Extent, PF_FloatR11G11B10, FClearValueBinding::Transparent, TexCreate_RenderTargetable | TexCreate_ShaderResource);
 		SceneTextures.MobileLocalLightTextureA = GraphBuilder.CreateTexture(MobileLocalLightTextureADesc, TEXT("MobileLocalLightTextureA"));
