Remove SUPPORTS_DEPTH_PREPASS define from MobileBasePassRendering.h, was submitted by mistake

diff --git a/Engine/Source/Runtime/Renderer/Private/MobileBasePassRendering.h b/Engine/Source/Runtime/Renderer/Private/MobileBasePassRendering.h
--- a/Engine/Source/Runtime/Renderer/Private/MobileBasePassRendering.h
+++ b/Engine/Source/Runtime/Renderer/Private/MobileBasePassRendering.h
@@ -433,9 +433,6 @@ public:
 		OutEnvironment.SetDefine(TEXT("ENABLE_CLUSTERED_REFLECTION"), bEnableClusteredReflections ? 1u : 0u);
 		OutEnvironment.SetDefine(TEXT("USE_SHADOWMASKTEXTURE"), bMobileUsesShadowMaskTexture && !bTranslucentMaterial ? 1u : 0u);
 		OutEnvironment.SetDefine(TEXT("ENABLE_DBUFFER_TEXTURES"), Parameters.MaterialParameters.MaterialDomain == MD_Surface ? 1u : 0u);
-
-		const bool bSupportsDepthPrepass = MobileUsesFullDepthPrepass(Parameters.Platform);
-		OutEnvironment.SetDefine(TEXT("SUPPORTS_DEPTH_PREPASS"), bSupportsDepthPrepass ? 1u : 0u);
 		EMobileTranslucentColorTransmittanceMode TranslucentColorTransmittanceMode = EMobileTranslucentColorTransmittanceMode::DEFAULT;
 		if (Parameters.MaterialParameters.ShadingModels.HasShadingModel(MSM_ThinTranslucent))
 		{
