Build LightGrid when Stochastic Shadows are enabled.

diff --git a/Engine/Source/Runtime/Renderer/Private/LightGridInjection.cpp b/Engine/Source/Runtime/Renderer/Private/LightGridInjection.cpp
--- a/Engine/Source/Runtime/Renderer/Private/LightGridInjection.cpp
+++ b/Engine/Source/Runtime/Renderer/Private/LightGridInjection.cpp
@@ -945,7 +945,7 @@ FComputeLightGridOutput FDeferredShadingSceneRenderer::GatherLightsAndComputeLig
 	
 	const bool bCullLightsToGrid = GLightCullingQuality 
 		&& (IsForwardShadingEnabled(ShaderPlatform) || bAnyViewUsesForwardLighting || IsRayTracingEnabled() || ShouldUseClusteredDeferredShading() ||
-			bAnyViewUsesLumen || ViewFamily.EngineShowFlags.VisualizeMeshDistanceFields || VirtualShadowMapArray.IsEnabled());
+			bAnyViewUsesLumen || ViewFamily.EngineShowFlags.VisualizeMeshDistanceFields || VirtualShadowMapArray.IsEnabled() || StochasticShadows::IsEnabled());
 
 	// Store this flag if lights are injected in the grids, check with 'AreLightsInLightGrid()'
 	bAreLightsInLightGrid = bCullLightsToGrid;
diff --git a/Engine/Source/Runtime/Renderer/Private/StochasticShadows/StochasticShadows.cpp b/Engine/Source/Runtime/Renderer/Private/StochasticShadows/StochasticShadows.cpp
--- a/Engine/Source/Runtime/Renderer/Private/StochasticShadows/StochasticShadows.cpp
+++ b/Engine/Source/Runtime/Renderer/Private/StochasticShadows/StochasticShadows.cpp
@@ -555,6 +555,8 @@ void FDeferredShadingSceneRenderer::RenderStochasticShadows(FRDGBuilder& GraphBu
 		return;
 	}
 
+	check(AreLightsInLightGrid());
+
 	RDG_EVENT_SCOPE(GraphBuilder, "StochasticShadows");
 
 	const FViewInfo& View = Views[0];
