LFV - mobile render path optimisations
Depth test is always used now to cull pixel, wether depth bound is available or not. 2.5ms => 1.8ms

diff --git a/Engine/Shaders/Private/LocalFogVolumes/LocalFogVolumeSplat.usf b/Engine/Shaders/Private/LocalFogVolumes/LocalFogVolumeSplat.usf
--- a/Engine/Shaders/Private/LocalFogVolumes/LocalFogVolumeSplat.usf
+++ b/Engine/Shaders/Private/LocalFogVolumes/LocalFogVolumeSplat.usf
@@ -33,6 +33,7 @@ struct FLocalFogVolumeVertexOutput
 #if LFV_TILED_VS
 
 Buffer<uint> TileDataBuffer;
+float StartDepthZ;
 
 void LocalFogVolumeTiledVS(
 	in uint VertexIndex : SV_VertexID,
@@ -50,7 +51,7 @@ void LocalFogVolumeTiledVS(
 	TileVertex.x += VertexIndex == 1 || VertexIndex == 2 || VertexIndex == 4 ? LFVTilePixelSize : 0;
 	TileVertex.y += VertexIndex == 2 || VertexIndex == 4 || VertexIndex == 5 ? LFVTilePixelSize : 0;
 
-	OutPosition = float4(float2(TileVertex) * ResolvedView.ViewSizeAndInvSize.zw * float2(2.0f, -2.0f) + float2(-1.0, 1.0f), 0.5f, 1.0f);
+	OutPosition = float4(float2(TileVertex) * ResolvedView.ViewSizeAndInvSize.zw * float2(2.0f, -2.0f) + float2(-1.0, 1.0f), StartDepthZ, 1.0f);
 
 	OtherVertexOutput.TileData = PackedTileCoord;
 }
diff --git a/Engine/Source/Runtime/Renderer/Private/LocalFogVolumeRendering.cpp b/Engine/Source/Runtime/Renderer/Private/LocalFogVolumeRendering.cpp
--- a/Engine/Source/Runtime/Renderer/Private/LocalFogVolumeRendering.cpp
+++ b/Engine/Source/Runtime/Renderer/Private/LocalFogVolumeRendering.cpp
@@ -632,6 +632,7 @@ class FLocalFogVolumeTiledRenderVS : public FGlobalShader
 		SHADER_PARAMETER_STRUCT_REF(FViewUniformShaderParameters, View)
 		SHADER_PARAMETER_STRUCT(FLocalFogVolumeUniformParameters, LFV)
 		SHADER_PARAMETER_RDG_BUFFER_SRV(Buffer<uint>, TileDataBuffer)
+		SHADER_PARAMETER(float, StartDepthZ)
 	END_SHADER_PARAMETER_STRUCT()
 
 	using FPermutationDomain = TShaderPermutationDomain<>;
@@ -693,29 +694,29 @@ BEGIN_SHADER_PARAMETER_STRUCT(FLocalFogVolumeTiledPassParameters, )
 	RENDER_TARGET_BINDING_SLOTS()
 END_SHADER_PARAMETER_STRUCT()
 
-static void ApplyDepthBoundIfNeeded(FGraphicsPipelineStateInitializer& GraphicsPSOInit, FMatrix ViewProjectionMatrix, FMatrix ViewInvProjectionMatrix, FRHICommandList& RHICmdList)
+static float ApplyDepthBoundIfNeeded(FGraphicsPipelineStateInitializer& GraphicsPSOInit, FMatrix ViewProjectionMatrix, FMatrix ViewInvProjectionMatrix, FRHICommandList& RHICmdList)
 {
+	float FogStartDistance = GetLocalFogVolumeGlobalStartDistance();
+
+	// Here we compute the nearest z value the fog can start
+	// to skip shader execution on pixels that are closer.
+	// This means with a bigger distance specified more pixels are
+	// are culled and don't need to be rendered. This is faster if
+	// there is opaque content nearer than the computed z.
+	// This optimization is achieved using depth bound tests.
+	// Mobile platforms typically does not support that feature 
+	// but typically renders the world using forward shading 
+	// with height fog evaluated as part of the material vertex or pixel shader.
+	FVector ViewSpaceCorner = ViewInvProjectionMatrix.TransformFVector4(FVector4(1, 1, 1, 1));
+	float Ratio = ViewSpaceCorner.Z / ViewSpaceCorner.Size();
+	FVector ViewSpaceStartFogPoint(0.0f, 0.0f, FogStartDistance * Ratio);
+	FVector4f ClipSpaceMaxDistance = (FVector4f)ViewProjectionMatrix.TransformPosition(ViewSpaceStartFogPoint); // LWC_TODO: precision loss
+	float FogClipSpaceZ = ClipSpaceMaxDistance.Z / ClipSpaceMaxDistance.W;
+	FogClipSpaceZ = FMath::Clamp(FogClipSpaceZ, 0.f, 1.f);
+
 	GraphicsPSOInit.bDepthBounds = GSupportsDepthBoundsTest;
 	if (GraphicsPSOInit.bDepthBounds)
 	{
-		float FogStartDistance = GetLocalFogVolumeGlobalStartDistance();
-
-		// Here we compute the nearest z value the fog can start
-		// to skip shader execution on pixels that are closer.
-		// This means with a bigger distance specified more pixels are
-		// are culled and don't need to be rendered. This is faster if
-		// there is opaque content nearer than the computed z.
-		// This optimization is achieved using depth bound tests.
-		// Mobile platforms typically does not support that feature 
-		// but typically renders the world using forward shading 
-		// with height fog evaluated as part of the material vertex or pixel shader.
-		FVector ViewSpaceCorner = ViewInvProjectionMatrix.TransformFVector4(FVector4(1, 1, 1, 1));
-		float Ratio = ViewSpaceCorner.Z / ViewSpaceCorner.Size();
-		FVector ViewSpaceStartFogPoint(0.0f, 0.0f, FogStartDistance * Ratio);
-		FVector4f ClipSpaceMaxDistance = (FVector4f)ViewProjectionMatrix.TransformPosition(ViewSpaceStartFogPoint); // LWC_TODO: precision loss
-		float FogClipSpaceZ = ClipSpaceMaxDistance.Z / ClipSpaceMaxDistance.W;
-		FogClipSpaceZ = FMath::Clamp(FogClipSpaceZ, 0.f, 1.f);
-
 		if (bool(ERHIZBuffer::IsInverted))
 		{
 			RHICmdList.SetDepthBounds(0.0f, FogClipSpaceZ);
@@ -725,6 +726,8 @@ static void ApplyDepthBoundIfNeeded(FGraphicsPipelineStateInitializer& GraphicsP
 			RHICmdList.SetDepthBounds(FogClipSpaceZ, 1.0f);
 		}
 	}
+
+	return FogClipSpaceZ;
 }
  
 void RenderLocalFogVolume(
@@ -792,7 +795,7 @@ void RenderLocalFogVolume(
 
 				// Render back faces only since camera may intersect
 				GraphicsPSOInit.RasterizerState = TStaticRasterizerState<>::GetRHI();
-				GraphicsPSOInit.DepthStencilState = TStaticDepthStencilState<false, CF_Always>::GetRHI();
+				GraphicsPSOInit.DepthStencilState = TStaticDepthStencilState<false, CF_DepthNearOrEqual>::GetRHI();
 				GraphicsPSOInit.BlendState = TStaticBlendState<CW_RGBA, BO_Add, BF_One, BF_SourceAlpha, BO_Add, BF_Zero, BF_SourceAlpha>::GetRHI();
 				GraphicsPSOInit.PrimitiveType = PT_TriangleList; // LFV_TODO check if rects are supported and use them if so
 
@@ -800,7 +803,7 @@ void RenderLocalFogVolume(
 				GraphicsPSOInit.BoundShaderState.VertexShaderRHI = VertexShader.GetVertexShader();
 				GraphicsPSOInit.BoundShaderState.PixelShaderRHI = PixelShader.GetPixelShader();
 
-				ApplyDepthBoundIfNeeded(GraphicsPSOInit, ViewProjectionMatrix, ViewInvProjectionMatrix, RHICmdList);
+				PassParameters->VS.StartDepthZ = ApplyDepthBoundIfNeeded(GraphicsPSOInit, ViewProjectionMatrix, ViewInvProjectionMatrix, RHICmdList);
 
 				SetGraphicsPipelineState(RHICmdList, GraphicsPSOInit, 0);
 
@@ -828,6 +831,7 @@ class FMobileLocalFogVolumeTiledRenderVS : public FGlobalShader
 		SHADER_PARAMETER_STRUCT(FLocalFogVolumeUniformParameters, LFV)
 		SHADER_PARAMETER_RDG_BUFFER_SRV(Buffer<uint>, TileDataBuffer)
 		RDG_BUFFER_ACCESS(TileDrawIndirectBuffer, ERHIAccess::IndirectArgs)
+		SHADER_PARAMETER(float, StartDepthZ)
 	END_SHADER_PARAMETER_STRUCT()
 
 	using FPermutationDomain = TShaderPermutationDomain<>;
@@ -948,7 +952,7 @@ void RenderLocalFogVolumeMobile(
 
 	// Render back faces only since camera may intersect
 	GraphicsPSOInit.RasterizerState = TStaticRasterizerState<>::GetRHI();
-	GraphicsPSOInit.DepthStencilState = TStaticDepthStencilState<false, CF_Always>::GetRHI();
+	GraphicsPSOInit.DepthStencilState = TStaticDepthStencilState<false, CF_DepthNearOrEqual>::GetRHI();
 	GraphicsPSOInit.BlendState = TStaticBlendState<CW_RGBA, BO_Add, BF_One, BF_SourceAlpha, BO_Add, BF_Zero, BF_SourceAlpha>::GetRHI();
 	GraphicsPSOInit.PrimitiveType = PT_TriangleList; // LFV_TODO check if rects are supported and use them if so
 
@@ -956,7 +960,7 @@ void RenderLocalFogVolumeMobile(
 	GraphicsPSOInit.BoundShaderState.VertexShaderRHI = VertexShader.GetVertexShader();
 	GraphicsPSOInit.BoundShaderState.PixelShaderRHI = PixelShader.GetPixelShader();
 
-	ApplyDepthBoundIfNeeded(GraphicsPSOInit, View.ViewMatrices.GetProjectionMatrix(), View.ViewMatrices.GetInvProjectionMatrix(), RHICmdList);
+	float StartDepthZ = ApplyDepthBoundIfNeeded(GraphicsPSOInit, View.ViewMatrices.GetProjectionMatrix(), View.ViewMatrices.GetInvProjectionMatrix(), RHICmdList);
 
 	SetGraphicsPipelineState(RHICmdList, GraphicsPSOInit, 0);
 
@@ -966,6 +970,7 @@ void RenderLocalFogVolumeMobile(
 	VSParameters.LFV = View.LocalFogVolumeViewData.UniformParametersStruct;
 	VSParameters.TileDataBuffer = View.LocalFogVolumeViewData.GPUTileDataBufferSRV;
 	VSParameters.TileDrawIndirectBuffer = View.LocalFogVolumeViewData.GPUTileDrawIndirectBuffer;
+	VSParameters.StartDepthZ = StartDepthZ;
 	SetShaderParameters(RHICmdList, VertexShader, VertexShader.GetVertexShader(), VSParameters);
 
 	FMobileLocalFogVolumeTiledRenderPS::FParameters PSParameters;
