LFV - mobile render path optimisations
- Packed LFV representation into 3 float4 instead of 5. We are no longer bound by tex fetch on low end platofrm now.
- Optimised ALu a bit but did not help much.

diff --git a/Engine/Shaders/Private/LocalFogVolumes/LocalFogVolumeCommon.ush b/Engine/Shaders/Private/LocalFogVolumes/LocalFogVolumeCommon.ush
--- a/Engine/Shaders/Private/LocalFogVolumes/LocalFogVolumeCommon.ush
+++ b/Engine/Shaders/Private/LocalFogVolumes/LocalFogVolumeCommon.ush
@@ -4,6 +4,7 @@
 #include "../SHCommon.ush"
 #include "../ParticipatingMediaCommon.ush"
 #include "../ReflectionEnvironmentShared.ush"
+#include "../FastMath.ush"
 
 
 
@@ -14,19 +15,17 @@ struct FLocalFogVolumeGPUInstanceData
 {
 	float4x3 InvTransform;
 
-	// float4
-	float RadialFogExtinction;		// half1
-	float HeightFogExtinction;		// half1
-	float HeightFogFalloff;			// half1
-	float HeightFogOffset;			// half1
-	float3 Emissive;				// half3
-	//FREE							// half1
-
-	// float4 (note: albedo could be using half representation or even 3x8bits if needed later to gain space)
-	float  UniformScale;			// float
-	float3 Albedo;					// half3
-	float  PhaseG;					// half1
-	// FREE							// float2
+	float RadialFogExtinction;
+	float HeightFogExtinction;
+	float HeightFogFalloff;
+	float HeightFogOffset;
+	float3 Emissive;
+
+	float  UniformScale;
+	float  UniformScaleInv;		// Derived from UniformScale
+
+	float3 Albedo;
+	float  PhaseG;
 };
 
 
@@ -85,31 +84,51 @@ Texture2DArray<uint>					LFV_LocalFogVolumeTileDataTexture;
 FLocalFogVolumeGPUInstanceData GetLocalFogVolumeGPUInstanceData(uint Index)
 {
 	FLocalFogVolumeGPUInstanceData Data;
-	uint Offset = Index * 5;
-
-	Data.InvTransform		= float4x3(LFVInsts[Offset + 0 ], LFVInsts[Offset + 1 ], LFVInsts[Offset + 2 ]); // reads 4x3 row major pack as 3 float4
-
-	float4 Data0			= LFVInsts[Offset + 3];
-
-	float2 Data0X			= UnpackFloat2FromUInt(asuint(Data0.x));
-	float2 Data0Y			= UnpackFloat2FromUInt(asuint(Data0.y));
-	float2 Data0Z			= UnpackFloat2FromUInt(asuint(Data0.z));
-	float2 Data0W			= UnpackFloat2FromUInt(asuint(Data0.w));
-
-	Data.RadialFogExtinction= Data0X.x;
-	Data.HeightFogExtinction= Data0X.y;
-	Data.HeightFogFalloff	= Data0Y.x;
-	Data.HeightFogOffset	= Data0Y.y;
-	Data.Emissive			= float3(Data0Z.x, Data0Z.y, Data0W.x);
-
-	float4 Data1			= LFVInsts[Offset + 4];
-
-	float2 Data1X			= UnpackFloat2FromUInt(asuint(Data1.y));
-	float2 Data1Y			= UnpackFloat2FromUInt(asuint(Data1.z));
+	uint Offset = Index * 3;
+
+	float4 Data0 = LFVInsts[Offset + 0];
+	float4 Data1 = LFVInsts[Offset + 1];
+	float4 Data2 = LFVInsts[Offset + 2];
+	// See GetLocalFogVolumeSortingData for the packing mathod.
+
+	// Data0
+	float3 Trans		= Data0.xyz;
+	Data.UniformScale	= Data0.w;
+	Data.UniformScaleInv= 1.0f / Data.UniformScale;
+
+	// Data1
+	float3 XVec;
+	float3 YVec;
+	XVec.xy		= UnpackFloat2FromUInt(asuint(Data1.x));
+	float2 Temp	= UnpackFloat2FromUInt(asuint(Data1.y));
+	XVec.z		= Temp.x;
+	YVec.x		= Temp.y;
+	YVec.yz		= UnpackFloat2FromUInt(asuint(Data1.z));
+
+	float3 ZVec = cross(XVec, YVec);
+
+	Data.InvTransform = float4x3(float3(XVec * Data.UniformScaleInv), float3(YVec * Data.UniformScaleInv), float3(ZVec * Data.UniformScaleInv), Trans);
+
+	// Data2
+	float3 Data2X			= UnpackR11G11B10F(asuint(Data2.x));
+	Data.RadialFogExtinction= Data2X.x;
+	Data.HeightFogExtinction= Data2X.y;
+	Data.HeightFogFalloff	= Data2X.z;
+
+	uint PackEmissiveData   = asuint(Data2.y);
+	if (PackEmissiveData > 0)
+	{
+		Data.Emissive = UnpackR11G11B10F(PackEmissiveData);
+	}
+	else
+	{
+		Data.Emissive = 0.0f.xxx;
+	}
 
-	Data.UniformScale		= Data1.x;
-	Data.Albedo				= float3(Data1X.x, Data1X.y, Data1Y.x);
-	Data.PhaseG				= Data1Y.y;
+	float4 Data2Z			= UnpackRGBA8(asuint(Data2.z));
+	Data.Albedo				= Data2Z.rgb;
+	Data.PhaseG				= Data2Z.a;
+	Data.HeightFogOffset	= Data2.w;
 
 	return Data;
 }
@@ -238,15 +257,24 @@ FFogData LocalFogVolumeEvaluateAnalyticalIntegral(FLocalFogVolumeGPUInstanceData
 		*/
 	}
 
+#if 1 
+	// => Reference path
 	// We now combine the two optical depth for the "opacity"=1-transmittance to be combined together.
 	// This helps to have soft edges when using height fog. We also account for that when voxelising in the volumetric fog.
 	float TR = exp(-RadialOpticalDepth);
 	float TH = exp(-HeightOpticalDepth);
-	float T = 1 - (1-TR) * (1-TH);
+	float T = 1 - (1 - TR) * (1 - TH);
 	float OpticalDepth = -log(T);
 
 	// Evaluate transmittance, accounting for the volume scale converted to local unit space meter.
 	float Transmittance = exp(-OpticalDepth * FogInstance.UniformScale * CENTIMETER_TO_METER);
+#else 
+	// => Optimised path, but does not match due to CommonFactor not being applied at the same place
+	const float CommonFactor = FogInstance.UniformScale * CENTIMETER_TO_METER;
+	float TR = exp(-RadialOpticalDepth);
+	float TH = exp(-HeightOpticalDepth);
+	float Transmittance = 1 - (1 - TR) * (1 - TH);
+#endif
 
 	// Assuming extinction is grey scale, we can compute a single coverage value from transmittance
 	FogData.Coverage = (1.0 - Transmittance);
@@ -318,10 +346,12 @@ float4 GetLFVInstanceContribution(
 	float2 TsU = RayIntersectSphere(RayPosU, RayDirU, float4(0.0, 0.0, 0.0, 1.0));
 
 	//Clamp to make sure we only consider intersection in front of the camera
-	TsU = max(LFVGlobalStartDistance.xx / FogInstance.UniformScale, TsU);
+	TsU = max(LFVGlobalStartDistance.xx * FogInstance.UniformScaleInv, TsU);
 
 	// Clamp the traced distance to the depth buffer and compute world length.
-	float LengthDepthBufferU = length(DepthBufferPosU - RayPosU);
+	float3 ViewToDepthVector = DepthBufferPosU - RayPosU;
+	float ViewToDepthVectorSqrLength = dot(ViewToDepthVector, ViewToDepthVector);
+	float LengthDepthBufferU = sqrt(ViewToDepthVectorSqrLength);
 	TsU = min(TsU, LengthDepthBufferU);
 
 	if (any(TsU > 0.0))
@@ -332,7 +362,7 @@ float4 GetLFVInstanceContribution(
 			// Volumetric fog covers up to MaxDistance along ViewZ, exclude analytical fog from this range
 			float CosAngle = dot(CamRayWorldDir, View.ViewForward);
 			float InvCosAngle = (CosAngle > 0.0001) ? rcp(CosAngle) : 0;
-			float ExcludeDistance = View.VolumetricFogMaxDistance * InvCosAngle / FogInstance.UniformScale;
+			float ExcludeDistance = View.VolumetricFogMaxDistance * InvCosAngle * FogInstance.UniformScaleInv;
 			TsU = max(TsU, ExcludeDistance);
 		}
 
diff --git a/Engine/Source/Runtime/Renderer/Private/LocalFogVolumeRendering.cpp b/Engine/Source/Runtime/Renderer/Private/LocalFogVolumeRendering.cpp
--- a/Engine/Source/Runtime/Renderer/Private/LocalFogVolumeRendering.cpp
+++ b/Engine/Source/Runtime/Renderer/Private/LocalFogVolumeRendering.cpp
@@ -381,49 +381,66 @@ void GetLocalFogVolumeSortingData(const FScene* Scene, FRDGBuilder& GraphBuilder
 
 		FLocalFogVolumeGPUInstanceData* LocalFogVolumeGPUInstanceDataIt = &Out.LocalFogVolumeGPUInstanceData[Out.LocalFogVolumeInstanceCountFinal];
 
-		auto ConvertFromMatrix44fTo4x3Array = [](FMatrix44f& InMat, float* OutArr)
-			{
-				// Row major order
-				OutArr[0 ] = InMat.M[0][0];
-				OutArr[1 ] = InMat.M[0][1];
-				OutArr[2 ] = InMat.M[0][2];
-				OutArr[3 ] = InMat.M[1][0];
-				OutArr[4 ] = InMat.M[1][1];
-				OutArr[5 ] = InMat.M[1][2];
-				OutArr[6 ] = InMat.M[2][0];
-				OutArr[7 ] = InMat.M[2][1];
-				OutArr[8 ] = InMat.M[2][2];
-				OutArr[9 ] = InMat.M[3][0];
-				OutArr[10] = InMat.M[3][1];
-				OutArr[11] = InMat.M[3][2];
-			};
-
-		FMatrix44f Transform    = FMatrix44f(LHF->FogTransform.ToMatrixWithScale());
-		FMatrix44f InvTransform = Transform.Inverse();
-		ConvertFromMatrix44fTo4x3Array(InvTransform, LocalFogVolumeGPUInstanceDataIt->InvTransform);
-
 		// Falloff needs to be made safe in order to avoid artifact when the camera is looking toward the horizon at the level of the offset.
 		const float SafeFalloffThreshold = 1.0f;
 		const float FalloffScaleUI = 0.01f;
 		const float SafeFallOff = FMath::Max(LHF->HeightFogFalloff, SafeFalloffThreshold) * FalloffScaleUI;
 
-		FVector2DHalf Data0X = FVector2DHalf(LHF->RadialFogExtinction,			LHF->HeightFogExtinction);
-		FVector2DHalf Data0Y = FVector2DHalf(SafeFallOff,						LHF->HeightFogOffset);
-		FVector2DHalf Data0Z = FVector2DHalf(LHF->FogEmissive.R,				LHF->FogEmissive.G);
-		FVector2DHalf Data0W = FVector2DHalf(LHF->FogEmissive.B,				0.0f);
+		FMatrix44f Transform = FMatrix44f(LHF->FogTransform.ToMatrixWithScale());
+		FMatrix44f InvTransform = Transform.Inverse();
+
+		FVector3f XVec(InvTransform.M[0][0], InvTransform.M[0][1], InvTransform.M[0][2]);
+		FVector3f YVec(InvTransform.M[1][0], InvTransform.M[1][1], InvTransform.M[1][2]);
+		FVector3f ZVec(InvTransform.M[2][0], InvTransform.M[2][1], InvTransform.M[2][2]);
+		FVector3f Tran(InvTransform.M[3][0], InvTransform.M[3][1], InvTransform.M[3][2]);
 
-		LocalFogVolumeGPUInstanceDataIt->Data0[0] = Data0X.AsUInt32();
-		LocalFogVolumeGPUInstanceDataIt->Data0[1] = Data0Y.AsUInt32();
-		LocalFogVolumeGPUInstanceDataIt->Data0[2] = Data0Z.AsUInt32();
-		LocalFogVolumeGPUInstanceDataIt->Data0[3] = Data0W.AsUInt32();
+		XVec.Normalize();
+		YVec.Normalize();
+		ZVec.Normalize();
 
-		FVector2DHalf Data1X = FVector2DHalf(LHF->FogAlbedo.R, LHF->FogAlbedo.G);
-		FVector2DHalf Data1Y = FVector2DHalf(LHF->FogAlbedo.B, LHF->FogPhaseG);
+		auto AsUint32 = [](float X)
+		{
+			union { float F; uint32 U; } FU = { X };
+			return FU.U;
+		};
+
+		auto AsFloat111110 = [](float x, float y, float z)
+		{
+			FVector2DHalf HalfXY(x, y);
+			FVector2DHalf HalfZ0(z, 0.0f);
 
-		LocalFogVolumeGPUInstanceDataIt->UniformScale = LHF->FogUniformScale;
-		LocalFogVolumeGPUInstanceDataIt->Data1[0] = Data1X.AsUInt32();
-		LocalFogVolumeGPUInstanceDataIt->Data1[1] = Data1Y.AsUInt32();
-		LocalFogVolumeGPUInstanceDataIt->Data1[2] = 0;
+			uint32 r = (uint32(HalfXY.X.Encoded) << 17) & 0xFFE00000;
+			uint32 g = (uint32(HalfXY.Y.Encoded) <<  6) & 0x001FFC00;
+			uint32 b = (uint32(HalfZ0.X.Encoded) >>  5) & 0x000003FF;
+			return r | g | b;
+		};
+
+		auto AsUNorm8888 = [](float x, float y, float z, float w)
+		{
+			uint32 r = (uint32(FMath::Clamp(x * 255.0f, 0.0f, 255.0f)) & 0xFFu);
+			uint32 g = (uint32(FMath::Clamp(y * 255.0f, 0.0f, 255.0f)) & 0xFFu) << 8u;
+			uint32 b = (uint32(FMath::Clamp(z * 255.0f, 0.0f, 255.0f)) & 0xFFu) << 16u;
+			uint32 a = (uint32(FMath::Clamp(w * 255.0f, 0.0f, 255.0f)) & 0xFFu) << 24u;
+			return r | g | b | a;
+		};
+
+		// Translation and scale at fp32 for stability. Further optimization: we could use fp16 if translated/view space position would be sent.
+		LocalFogVolumeGPUInstanceDataIt->Data0[0] = AsUint32(Tran.X);
+		LocalFogVolumeGPUInstanceDataIt->Data0[1] = AsUint32(Tran.Y);
+		LocalFogVolumeGPUInstanceDataIt->Data0[2] = AsUint32(Tran.Z);
+		LocalFogVolumeGPUInstanceDataIt->Data0[3] = AsUint32(LHF->FogUniformScale);
+
+		// Store X and Y from the rotation matrix and recover Z in the shader. Further optimization: could be fp16.
+		LocalFogVolumeGPUInstanceDataIt->Data1[0] = FVector2DHalf(XVec.X, XVec.Y).AsUInt32();
+		LocalFogVolumeGPUInstanceDataIt->Data1[1] = FVector2DHalf(XVec.Z, YVec.X).AsUInt32();
+		LocalFogVolumeGPUInstanceDataIt->Data1[2] = FVector2DHalf(YVec.Y, YVec.Z).AsUInt32();
+		LocalFogVolumeGPUInstanceDataIt->Data1[3] = 0; // FREE
+
+		// All the remaining data are packed as small as possible w.r.t. their range of value.
+		LocalFogVolumeGPUInstanceDataIt->Data2[0] = AsFloat111110(	LHF->RadialFogExtinction,	LHF->HeightFogExtinction,	SafeFallOff);
+		LocalFogVolumeGPUInstanceDataIt->Data2[1] = AsFloat111110(	LHF->FogEmissive.R,			LHF->FogEmissive.G,			LHF->FogEmissive.B);
+		LocalFogVolumeGPUInstanceDataIt->Data2[2] = AsUNorm8888(	LHF->FogAlbedo.R,			LHF->FogAlbedo.G,			LHF->FogAlbedo.B,		LHF->FogPhaseG);
+		LocalFogVolumeGPUInstanceDataIt->Data2[3] = AsUint32(		LHF->HeightFogOffset);
 
 		// Register the sorting data
 		Out.LocalFogVolumeCenterPos[Out.LocalFogVolumeInstanceCountFinal] = LHF->FogTransform.GetTranslation();
@@ -473,7 +490,7 @@ void CreateViewLocalFogVolumeBufferSRV(const FScene* Scene, FViewInfo& View, FRD
 		LocalFogVolumeGPUSortedInstanceData[i] = SortingData.LocalFogVolumeGPUInstanceData[LFVKey.FogVolume.Index];
 
 		FVector& LFVPosition = SortingData.LocalFogVolumeCenterPos[LFVKey.FogVolume.Index];
-		LocalFogVolumeGPUSortedInstanceCullingData[i] = FVector4f(LFVPosition.X, LFVPosition.Y, LFVPosition.Z, LocalFogVolumeGPUSortedInstanceData[i].UniformScale);
+		LocalFogVolumeGPUSortedInstanceCullingData[i] = FVector4f(LFVPosition.X, LFVPosition.Y, LFVPosition.Z, LocalFogVolumeGPUSortedInstanceData[i].GetUniformScale());
 	}
 
 	// 3. Allocate buffer and initialize with sorted data to upload to GPU
diff --git a/Engine/Source/Runtime/Renderer/Private/LocalFogVolumeRendering.h b/Engine/Source/Runtime/Renderer/Private/LocalFogVolumeRendering.h
--- a/Engine/Source/Runtime/Renderer/Private/LocalFogVolumeRendering.h
+++ b/Engine/Source/Runtime/Renderer/Private/LocalFogVolumeRendering.h
@@ -45,12 +45,15 @@ class FLocalFogVolumeGPUInstanceData
 {
 public:
 
-	float InvTransform[4 * 3];
-
-	// See LocalFogVolumeCommon for the encoded representation
 	uint32 Data0[4];
-	float  UniformScale;
-	uint32 Data1[3];
+	uint32 Data1[4];
+	uint32 Data2[4];
+
+	float GetUniformScale() 
+	{
+		union { float F; uint32 U; } FU = { Data0[3] };
+		return FU.F;
+	}
 };
 
 class RENDERER_API FLocalFogVolumeSortKey
