Stochastic Shadows
* Restored rect lights
* Removed StochasticMinLightSampleWeight as it causing some artifacts without really improving final image
* Fixed ShadowMaskHashTableHistory resize

diff --git a/Engine/Shaders/Private/StochasticShadows/StochasticShadowsSampling.usf b/Engine/Shaders/Private/StochasticShadows/StochasticShadowsSampling.usf
--- a/Engine/Shaders/Private/StochasticShadows/StochasticShadowsSampling.usf
+++ b/Engine/Shaders/Private/StochasticShadows/StochasticShadowsSampling.usf
@@ -84,7 +84,7 @@ Texture2D<uint2> ShadowMaskHistoryScreenCoord00;
 Texture2D<float4> ShadowMaskHistoryWeights;
 uint2 DownsampledViewSize;
 uint MaxShadowMaskTiles;
-float MinLightSampleWeight;
+float SamplingMinWeight;
 int DebugMode;
 
 RWTexture2D<uint2> RWShadowMaskHistoryScreenCoord00;
@@ -315,9 +315,6 @@ void GenerateSamplesCS(
 
 			if (SceneDepth > 0)
 			{
-				// Add a bit of stochastic noise to hide transitions
-				const float StochasticMinLightSampleWeight = MinLightSampleWeight * (1.0f + InterleavedGradientNoise(DownsampledScreenCoord, StochasticShadowsStateFrameIndex));
-
 				// Prepare weights for history reprojection
 				#if SHADOW_MASK_REPROJECTION_WEIGHTS
 				{
@@ -416,9 +413,8 @@ void GenerateSamplesCS(
 
 						uint LightSceneId = MaskIndex * 32 + NextBitIndex;
 						FLightSceneData LightSceneData = GetLightSceneData(LightSceneId);
-						FDeferredLightData LightData = ConvertToDeferredLight(LightSceneData);\
-						LightData.bRectLight = false;
-					
+						FDeferredLightData LightData = ConvertToDeferredLight(LightSceneData);
+
 						float SampleWeight = GetLocalLightTargetPDF(LightData, TranslatedWorldPosition, GBuffer, ScreenCoord);
 
 						// Simulate tonemapping
@@ -465,7 +461,7 @@ void GenerateSamplesCS(
 						}
 						#endif
 
-						if (SampleWeight > StochasticMinLightSampleWeight)
+						if (SampleWeight > SamplingMinWeight)
 						{
 							// Allocate required shadow mask tiles for this frame
 							MarkLightSceneId(LightSceneId);
@@ -476,7 +472,7 @@ void GenerateSamplesCS(
 								SampleWeight *= ShadowFactorEstimate;
 							}
 
-							if (SampleWeight > StochasticMinLightSampleWeight)
+							if (SampleWeight > SamplingMinWeight)
 							{
 								float Tau = WeightSum / (WeightSum + SampleWeight);
 								WeightSum += SampleWeight;
@@ -554,7 +550,7 @@ void GenerateSamplesCS(
 						}
 						#endif
 
-						if (SampleWeight > StochasticMinLightSampleWeight)
+						if (SampleWeight > SamplingMinWeight)
 						{
 							// Allocate required shadow mask tiles for this frame
 							MarkLightSceneId(LocalLightData.LightSceneId);
@@ -565,7 +561,7 @@ void GenerateSamplesCS(
 								SampleWeight *= ShadowFactorEstimate;
 							}
 
-							if (SampleWeight > StochasticMinLightSampleWeight)
+							if (SampleWeight > SamplingMinWeight)
 							{
 								float Tau = WeightSum / (WeightSum + SampleWeight);
 								WeightSum += SampleWeight;
diff --git a/Engine/Shaders/Private/StochasticShadows/StochasticShadowsShading.usf b/Engine/Shaders/Private/StochasticShadows/StochasticShadowsShading.usf
--- a/Engine/Shaders/Private/StochasticShadows/StochasticShadowsShading.usf
+++ b/Engine/Shaders/Private/StochasticShadows/StochasticShadowsShading.usf
@@ -93,7 +93,6 @@ void ShadeLightSamplesCS(
 					const uint PackedShadingTile = ShadingTileGrid[GetShadingTileGridOffset(ShadingTileCoord, IndexInCell)];
 					const FShadingTile ShadingTile = UnpackShadingTile(PackedShadingTile);
 					FDeferredLightData LightData = ConvertToDeferredLight(GetLightSceneData(ShadingTile.LightSceneId));
-					LightData.bRectLight = false;
 
 					uint2 ShadingTileAtlasCoord = GetShadingTileAtlasCoord(ShadingTile.ShadingTileIndex, CoordInShadingTile);
 
@@ -143,7 +142,7 @@ void ShadeLightSamplesCS(
 							ScreenCoord, 
 							SurfaceShadow);
 
-						#if DEBUG_MODE
+						#if DEBUG_MODE 
 						if (bDebug)
 						{
 							Print(Context, LightRadiance);
diff --git a/Engine/Source/Runtime/Renderer/Private/StochasticShadows/StochasticShadows.cpp b/Engine/Source/Runtime/Renderer/Private/StochasticShadows/StochasticShadows.cpp
--- a/Engine/Source/Runtime/Renderer/Private/StochasticShadows/StochasticShadows.cpp
+++ b/Engine/Source/Runtime/Renderer/Private/StochasticShadows/StochasticShadows.cpp
@@ -25,13 +25,6 @@ static TAutoConsoleVariable<int32> CVarStochasticShadowsNumSamplesPerPixel(
 	ECVF_Scalability | ECVF_RenderThreadSafe
 );
 
-static TAutoConsoleVariable<float> CVarStochasticShadowsMinSampleWeight(
-	TEXT("r.StochasticShadows.MinSampleWeight"),
-	0.002f,
-	TEXT("Determines minimal sample influence on final pixels. Used to skip samples which would have minimal impact to the final image even if light is fully visible."),
-	ECVF_Scalability | ECVF_RenderThreadSafe
-);
-
 static TAutoConsoleVariable<int32> CVarStochasticShadowsTemporal(
 	TEXT("r.StochasticShadows.Temporal"),
 	1,
@@ -53,6 +46,13 @@ static TAutoConsoleVariable<float> CVarStochasticShadowsTemporalStdDevOffset(
 	ECVF_Scalability | ECVF_RenderThreadSafe
 );
 
+static TAutoConsoleVariable<float> CVarStochasticShadowsSamplingMinWeight(
+	TEXT("r.StochasticShadows.Sampling.MinWeight"),
+	0.002f,
+	TEXT("Determines minimal sample influence on final pixels. Used to skip samples which would have minimal impact to the final image even if light is fully visible."),
+	ECVF_Scalability | ECVF_RenderThreadSafe
+);
+
 static TAutoConsoleVariable<int32> CVarStochasticSamplingShadowEstimate(
 	TEXT("r.StochasticShadows.Sampling.ShadowEstimate"),
 	0,
@@ -117,7 +117,7 @@ static TAutoConsoleVariable<int> CVarStochasticShadowsHashTable(
 static TAutoConsoleVariable<int> CVarStochasticShadowsCandidateLightMask(
 	TEXT("r.StochasticShadows.CandidateLightMask"),
 	1,
-	TEXT("#kris_todo: finish and pick one shader path."),
+	TEXT("#sdl_todo: finish and pick one shader path."),
 	ECVF_Scalability | ECVF_RenderThreadSafe
 );
 
@@ -670,9 +670,8 @@ void FDeferredShadingSceneRenderer::RenderStochasticShadows(FRDGBuilder& GraphBu
 				ShadowMaskPageTableHistory = GraphBuilder.RegisterExternalTexture(LightingViewState.ShadowMaskPageTableHistory);
 			}
 
-			// #kris_todo: fix it
 			if (LightingViewState.ShadowMaskHashTableHistory 
-				/* && LightingViewState.ShadowMaskHashTableHistory->GetAlignedDesc().NumElements == ShadowMaskHashTableSize*/)
+				&& LightingViewState.ShadowMaskHashTableHistory->GetSize() == sizeof(uint32) * ShadowMaskHashTableSize)
 			{
 				ShadowMaskHashTableHistory = GraphBuilder.RegisterExternalBuffer(LightingViewState.ShadowMaskHashTableHistory);
 			}
@@ -728,7 +727,7 @@ void FDeferredShadingSceneRenderer::RenderStochasticShadows(FRDGBuilder& GraphBu
 		StochasticShadowsParameters.ShadowMaskHashTableIndexWrapMask = ShadowMaskHashTableSize - 1;
 		StochasticShadowsParameters.ShadowMaskPageTablePerLightSize = ShadowMaskPageTablePerLightSize;
 		StochasticShadowsParameters.DownsampledBufferInvSize = FVector2f(1.0f) / DownsampledBufferSize;
-		StochasticShadowsParameters.MinLightSampleWeight = CVarStochasticShadowsMinSampleWeight.GetValueOnRenderThread();
+		StochasticShadowsParameters.SamplingMinWeight = FMath::Max(CVarStochasticShadowsSamplingMinWeight.GetValueOnRenderThread(), 0.0f);
 		StochasticShadowsParameters.TileDataStride = TileDataStride;
 		StochasticShadowsParameters.DownsampledTileDataStride = DownsampledTileDataStride;
 		StochasticShadowsParameters.TemporalMaxFramesAccumulated = CVarStochasticShadowsTemporalMaxFramesAccumulated.GetValueOnRenderThread();
diff --git a/Engine/Source/Runtime/Renderer/Private/StochasticShadows/StochasticShadowsInternal.h b/Engine/Source/Runtime/Renderer/Private/StochasticShadows/StochasticShadowsInternal.h
--- a/Engine/Source/Runtime/Renderer/Private/StochasticShadows/StochasticShadowsInternal.h
+++ b/Engine/Source/Runtime/Renderer/Private/StochasticShadows/StochasticShadowsInternal.h
@@ -25,7 +25,7 @@ BEGIN_SHADER_PARAMETER_STRUCT(FStochasticShadowsParameters, )
 	SHADER_PARAMETER(FIntPoint, ShadowMaskPageTablePerLightSize)
 	SHADER_PARAMETER(FIntPoint, ShadingTileGridSize)
 	SHADER_PARAMETER(FVector2f, DownsampledBufferInvSize)
-	SHADER_PARAMETER(float, MinLightSampleWeight)
+	SHADER_PARAMETER(float, SamplingMinWeight)
 	SHADER_PARAMETER(int32, TileDataStride)
 	SHADER_PARAMETER(int32, DownsampledTileDataStride)
 	SHADER_PARAMETER(int32, TemporalMaxFramesAccumulated)
