Fixed race condition between main visibility GDME context merging and shadow projection stenciling.

diff --git a/Engine/Source/Runtime/Renderer/Private/SceneVisibility.cpp b/Engine/Source/Runtime/Renderer/Private/SceneVisibility.cpp
--- a/Engine/Source/Runtime/Renderer/Private/SceneVisibility.cpp
+++ b/Engine/Source/Runtime/Renderer/Private/SceneVisibility.cpp
@@ -5437,10 +5437,11 @@ void FDeferredShadingSceneRenderer::EndInitViews(
 	SCOPED_NAMED_EVENT(FDeferredShadingSceneRenderer_InitViewsAfterPrepass, FColor::Emerald);
 	SCOPE_CYCLE_COUNTER(STAT_InitViewsPossiblyAfterPrepass);
 
-	BeginShadowGatherDynamicMeshElements(TaskDatas.DynamicShadows);
-
 	TaskDatas.VisibilityTaskData->Finish();
 
+	// Trigger shadow GDME tasks after the main visibility tasks are synced. Projection stencil shadows reference the main view dynamic elements.
+	BeginShadowGatherDynamicMeshElements(TaskDatas.DynamicShadows);
+
 	FRHICommandListImmediate& RHICmdList = GraphBuilder.RHICmdList;
 
 	// If parallel ILC update is disabled, then process it in place.
