Fix incorret hair group AABBs indexing during voxel page allocation.

diff --git a/Engine/Shaders/Private/HairStrands/HairStrandsVoxelRasterCompute.usf b/Engine/Shaders/Private/HairStrands/HairStrandsVoxelRasterCompute.usf
--- a/Engine/Shaders/Private/HairStrands/HairStrandsVoxelRasterCompute.usf
+++ b/Engine/Shaders/Private/HairStrands/HairStrandsVoxelRasterCompute.usf
@@ -159,6 +159,7 @@ void AllocatePageIndex(uint2 DispatchThreadId : SV_DispatchThreadID)
 ///////////////////////////////////////////////////////////////////////////////////////////////////
 
 #if SHADER_MARKVALID_PREPARE
+uint 			InstanceRegisteredIndex;
 uint			MaxClusterCount;
 uint			MacroGroupId;
 uint			bUseMacroGroupBoundCPU;
@@ -273,7 +274,7 @@ void MarkValid_PrepareCS(uint3 DispatchThreadId : SV_DispatchThreadID)
 	else
 	{
 		MacroGroupBound = ReadHairAABB(MacroGroupId, MacroGroupVoxelAlignedAABBBuffer);
-		GroupBound = ReadHairAABB(0, GroupAABBsBuffer);
+		GroupBound = ReadHairAABB(InstanceRegisteredIndex, GroupAABBsBuffer);
 	}
 
 	const uint3 MinCoord = PositionToCoord(GroupBound.Min, MacroGroupBound.Min, MacroGroupBound.Max, PageIndexResolution);
diff --git a/Engine/Source/Runtime/Renderer/Private/HairStrands/HairStrandsVoxelization.cpp b/Engine/Source/Runtime/Renderer/Private/HairStrands/HairStrandsVoxelization.cpp
--- a/Engine/Source/Runtime/Renderer/Private/HairStrands/HairStrandsVoxelization.cpp
+++ b/Engine/Source/Runtime/Renderer/Private/HairStrands/HairStrandsVoxelization.cpp
@@ -277,6 +277,7 @@ class FVoxelMarkValidPageIndex_PrepareCS : public FGlobalShader
 	using FPermutationDomain = TShaderPermutationDomain<FUseCluster>;
 
 	BEGIN_SHADER_PARAMETER_STRUCT(FParameters, )
+		SHADER_PARAMETER(uint32, InstanceRegisteredIndex)
 		SHADER_PARAMETER(uint32, MaxClusterCount)
 		SHADER_PARAMETER(uint32, MacroGroupId)
 		SHADER_PARAMETER(uint32, bUseMacroGroupBoundCPU)
@@ -674,6 +675,7 @@ static void AddAllocateVoxelPagesPass(
 				const bool bUseClusterAABB = HairGroupData->GetClusterAABBValid() && bIsGPUDriven;
 				
 				FVoxelMarkValidPageIndex_PrepareCS::FParameters* Parameters = GraphBuilder.AllocParameters<FVoxelMarkValidPageIndex_PrepareCS::FParameters>();
+				Parameters->InstanceRegisteredIndex						= HairGroupData->Instance->RegisteredIndex;
 				Parameters->MaxClusterCount								= HairGroupData->GetClusterCount();
 				Parameters->MacroGroupId								= MacroGroup.MacroGroupId;
 				Parameters->GroupAABBsBuffer							= TransientResources.GroupAABBSRV;
