Help compiler to compile out unused ShadingPath renderer code on platforms with a static feature level. (cuts ~2.5MB from a mobile only binary size)

diff --git a/Engine/Plugins/FX/Niagara/Source/Niagara/Private/NiagaraGpuComputeDispatch.cpp b/Engine/Plugins/FX/Niagara/Source/Niagara/Private/NiagaraGpuComputeDispatch.cpp
--- a/Engine/Plugins/FX/Niagara/Source/Niagara/Private/NiagaraGpuComputeDispatch.cpp
+++ b/Engine/Plugins/FX/Niagara/Source/Niagara/Private/NiagaraGpuComputeDispatch.cpp
@@ -1082,11 +1082,11 @@ void FNiagaraGpuComputeDispatch::ExecuteTicks(FRDGBuilder& GraphBuilder, TConstS
 	// Setup Parameters that can be read from data interfaces
 	SimulationSceneViews = Views;
 
-	if (FSceneInterface::GetShadingPath(FeatureLevel) == EShadingPath::Deferred)
+	if (GetFeatureLevelShadingPath(FeatureLevel) == EShadingPath::Deferred)
 	{
 		SceneTexturesUniformParams = UE::FXRenderingUtils::GetOrCreateSceneTextureUniformBuffer(GraphBuilder, SimulationSceneViews, FeatureLevel, ESceneTextureSetupMode::SceneVelocity);
 	}
-	else if (FSceneInterface::GetShadingPath(FeatureLevel) == EShadingPath::Mobile)
+	else if (GetFeatureLevelShadingPath(FeatureLevel) == EShadingPath::Mobile)
 	{
 		MobileSceneTexturesUniformParams = UE::FXRenderingUtils::GetOrCreateMobileSceneTextureUniformBuffer(GraphBuilder, SimulationSceneViews, EMobileSceneTextureSetupMode::None);
 	}
diff --git a/Engine/Plugins/FX/Niagara/Source/Niagara/Private/NiagaraRendererProperties.cpp b/Engine/Plugins/FX/Niagara/Source/Niagara/Private/NiagaraRendererProperties.cpp
--- a/Engine/Plugins/FX/Niagara/Source/Niagara/Private/NiagaraRendererProperties.cpp
+++ b/Engine/Plugins/FX/Niagara/Source/Niagara/Private/NiagaraRendererProperties.cpp
@@ -1020,7 +1020,7 @@ bool UNiagaraRendererProperties::ShouldGpuTranslucentThisFrame(ENiagaraRendererG
 bool UNiagaraRendererProperties::IsGpuTranslucentThisFrame(ERHIFeatureLevel::Type FeatureLevel, ENiagaraRendererGpuTranslucentLatency Latency)
 {
 	// We can not support low latency on the mobile renderer path as it calls PostRenderOpaque after translucent in some paths
-	if (FSceneInterface::GetShadingPath(FeatureLevel) != EShadingPath::Deferred)
+	if (GetFeatureLevelShadingPath(FeatureLevel) != EShadingPath::Deferred)
 	{
 		return false;
 	}
diff --git a/Engine/Source/Runtime/Engine/Private/GameViewportClient.cpp b/Engine/Source/Runtime/Engine/Private/GameViewportClient.cpp
--- a/Engine/Source/Runtime/Engine/Private/GameViewportClient.cpp
+++ b/Engine/Source/Runtime/Engine/Private/GameViewportClient.cpp
@@ -1379,7 +1379,7 @@ void UGameViewportClient::Draw(FViewport* InViewport, FCanvas* SceneCanvas)
 	{
 		static const auto MobileMultiViewCVar = IConsoleManager::Get().FindTConsoleVariableDataInt(TEXT("vr.MobileMultiView"));
 		const bool bSkipPostprocessing = !IsMobileHDR();
-		const bool bUsingMobileRenderer = FSceneInterface::GetShadingPath(MyWorld->Scene->GetFeatureLevel()) == EShadingPath::Mobile;
+		const bool bUsingMobileRenderer = GetFeatureLevelShadingPath(MyWorld->Scene->GetFeatureLevel()) == EShadingPath::Mobile;
 		bRequireMultiView = (GSupportsMobileMultiView || GRHISupportsArrayIndexFromAnyShader) && bUsingMobileRenderer && bSkipPostprocessing && (MobileMultiViewCVar && MobileMultiViewCVar->GetValueOnAnyThread() != 0);
 	}
 
diff --git a/Engine/Source/Runtime/Engine/Private/Materials/MaterialShader.cpp b/Engine/Source/Runtime/Engine/Private/Materials/MaterialShader.cpp
--- a/Engine/Source/Runtime/Engine/Private/Materials/MaterialShader.cpp
+++ b/Engine/Source/Runtime/Engine/Private/Materials/MaterialShader.cpp
@@ -2727,7 +2727,7 @@ FPSOPrecacheRequestResultArray FMaterialShaderMap::CollectPSOs(const FMaterialPS
 	FSceneTexturesConfig SceneTexturesConfig;
 	SceneTexturesConfig.Init(SceneTexturesConfigInitSettings);
 
-	const EShadingPath ShadingPath = FSceneInterface::GetShadingPath(PrecacheParams.FeatureLevel);
+	const EShadingPath ShadingPath = GetFeatureLevelShadingPath(PrecacheParams.FeatureLevel);
 
 	TArray<FPSOPrecacheData> PSOInitializers;
 	PSOInitializers.Reserve(32);
diff --git a/Engine/Source/Runtime/Engine/Private/PSOPrecache.cpp b/Engine/Source/Runtime/Engine/Private/PSOPrecache.cpp
--- a/Engine/Source/Runtime/Engine/Private/PSOPrecache.cpp
+++ b/Engine/Source/Runtime/Engine/Private/PSOPrecache.cpp
@@ -620,7 +620,7 @@ void PSOCollectorStats::FPrecacheUsageData::UpdateStats(uint32 MeshPassType, con
 bool PSOCollectorStats::FPrecacheStatsCollector::IsStateTracked(uint32 MeshPassType, const FVertexFactoryType* VertexFactoryType) const {
 	bool bTracked = true;
 	if (MeshPassType < FPSOCollectorCreateManager::MaxPSOCollectorCount) {
-		const EShadingPath ShadingPath = FSceneInterface::GetShadingPath(GMaxRHIFeatureLevel);
+		const EShadingPath ShadingPath = GetFeatureLevelShadingPath(GMaxRHIFeatureLevel);
 		bool bCollectPSOs = FPSOCollectorCreateManager::GetCreateFunction(ShadingPath, MeshPassType) != nullptr;
 		bTracked = bCollectPSOs && (VertexFactoryType == nullptr || VertexFactoryType->SupportsPSOPrecaching());
 	}
diff --git a/Engine/Source/Runtime/Engine/Private/SceneInterface.cpp b/Engine/Source/Runtime/Engine/Private/SceneInterface.cpp
--- a/Engine/Source/Runtime/Engine/Private/SceneInterface.cpp
+++ b/Engine/Source/Runtime/Engine/Private/SceneInterface.cpp
@@ -28,12 +28,5 @@ EShaderPlatform FSceneInterface::GetShaderPlatform() const
 
 EShadingPath FSceneInterface::GetShadingPath(ERHIFeatureLevel::Type InFeatureLevel)
 {
-	if (InFeatureLevel >= ERHIFeatureLevel::SM5)
-	{
-		return EShadingPath::Deferred;
-	}
-	else
-	{
-		return EShadingPath::Mobile;
-	}
-}
\ No newline at end of file
+	return GetFeatureLevelShadingPath(InFeatureLevel);
+}
diff --git a/Engine/Source/Runtime/Engine/Private/SceneTexturesConfig.cpp b/Engine/Source/Runtime/Engine/Private/SceneTexturesConfig.cpp
--- a/Engine/Source/Runtime/Engine/Private/SceneTexturesConfig.cpp
+++ b/Engine/Source/Runtime/Engine/Private/SceneTexturesConfig.cpp
@@ -132,7 +132,7 @@ static EPixelFormat GetSceneColorFormat(bool bRequiresAlphaChannel)
 
 static void GetSceneColorFormatAndCreateFlags(ERHIFeatureLevel::Type FeatureLevel, bool bRequiresAlphaChannel, ETextureCreateFlags ExtraSceneColorCreateFlags, uint32 NumSamples, bool bMemorylessMSAA, EPixelFormat& SceneColorFormat, ETextureCreateFlags& SceneColorCreateFlags)
 {
-	EShadingPath ShadingPath = FSceneInterface::GetShadingPath(FeatureLevel);
+	EShadingPath ShadingPath = GetFeatureLevelShadingPath(FeatureLevel);
 	switch (ShadingPath)
 	{
 	case EShadingPath::Deferred:
@@ -255,7 +255,7 @@ static void SetupMobileGBufferFlags(FGBufferBindings GBufferBindings[GBL_Num], b
 void FSceneTexturesConfig::Init(const FSceneTexturesConfigInitSettings& InitSettings)
 {
 	FeatureLevel			= InitSettings.FeatureLevel;
-	ShadingPath				= FSceneInterface::GetShadingPath(FeatureLevel);
+	ShadingPath				= GetFeatureLevelShadingPath(FeatureLevel);
 	ShaderPlatform			= GetFeatureLevelShaderPlatform(FeatureLevel);
 	Extent					= InitSettings.Extent;
 	NumSamples				= GetDefaultMSAACount(FeatureLevel, GDynamicRHI->RHIGetPlatformTextureMaxSampleCount());
diff --git a/Engine/Source/Runtime/Engine/Private/SceneView.cpp b/Engine/Source/Runtime/Engine/Private/SceneView.cpp
--- a/Engine/Source/Runtime/Engine/Private/SceneView.cpp
+++ b/Engine/Source/Runtime/Engine/Private/SceneView.cpp
@@ -855,7 +855,7 @@ PRAGMA_ENABLE_DEPRECATION_WARNINGS
 
 	// OpenGL Gamma space output in GLSL flips Y when rendering directly to the back buffer (so not needed on PC, as we never render directly into the back buffer)
 	auto ShaderPlatform = GShaderPlatformForFeatureLevel[FeatureLevel];
-	bool bUsingMobileRenderer = FSceneInterface::GetShadingPath(FeatureLevel) == EShadingPath::Mobile;
+	bool bUsingMobileRenderer = GetFeatureLevelShadingPath(FeatureLevel) == EShadingPath::Mobile;
 
 	// Setup transformation constants to be used by the graphics hardware to transform device normalized depth samples
 	// into world oriented z.
diff --git a/Engine/Source/Runtime/Engine/Private/VT/VirtualTextureBuilder.cpp b/Engine/Source/Runtime/Engine/Private/VT/VirtualTextureBuilder.cpp
--- a/Engine/Source/Runtime/Engine/Private/VT/VirtualTextureBuilder.cpp
+++ b/Engine/Source/Runtime/Engine/Private/VT/VirtualTextureBuilder.cpp
@@ -67,7 +67,7 @@ void UVirtualTextureBuilder::PostLoad()
 	// Discard one of the VTs on a cooked platform that support both rendering modes
 	if (FPlatformProperties::RequiresCookedData())
 	{
-		if (FSceneInterface::GetShadingPath(GMaxRHIFeatureLevel) == EShadingPath::Mobile && bSeparateTextureForMobile)
+		if (GetFeatureLevelShadingPath(GMaxRHIFeatureLevel) == EShadingPath::Mobile && bSeparateTextureForMobile)
 		{
 			Texture = nullptr;
 		}
diff --git a/Engine/Source/Runtime/Engine/Public/SceneUtils.h b/Engine/Source/Runtime/Engine/Public/SceneUtils.h
--- a/Engine/Source/Runtime/Engine/Public/SceneUtils.h
+++ b/Engine/Source/Runtime/Engine/Public/SceneUtils.h
@@ -31,6 +31,11 @@ enum class EShadingPath
 	Num,
 };
 
+inline EShadingPath GetFeatureLevelShadingPath(FStaticFeatureLevel InFeatureLevel)
+{
+	return (InFeatureLevel >= ERHIFeatureLevel::SM5) ? EShadingPath::Deferred : EShadingPath::Mobile;
+}
+
 enum class EMobileHDRMode
 {
 	Unset,
diff --git a/Engine/Source/Runtime/Renderer/Private/EditorPrimitivesRendering.cpp b/Engine/Source/Runtime/Renderer/Private/EditorPrimitivesRendering.cpp
--- a/Engine/Source/Runtime/Renderer/Private/EditorPrimitivesRendering.cpp
+++ b/Engine/Source/Runtime/Renderer/Private/EditorPrimitivesRendering.cpp
@@ -49,7 +49,7 @@ bool FEditorPrimitivesBasePassMeshProcessor::TryAddMeshBatch(const FMeshBatch& R
 		&& (!PrimitiveSceneProxy || PrimitiveSceneProxy->ShouldRenderInMainPass())
 		&& ShouldIncludeDomainInMeshPass(Material.GetMaterialDomain()))
 	{
-		if (Scene->GetShadingPath(FeatureLevel) == EShadingPath::Mobile)
+		if (GetFeatureLevelShadingPath(FeatureLevel) == EShadingPath::Mobile)
 		{
 			bResult = ProcessMobileShadingPath(MeshBatch, BatchElementMask, Material, MaterialRenderProxy, PrimitiveSceneProxy, StaticMeshId);
 		}
diff --git a/Engine/Source/Runtime/Renderer/Private/MeshDrawCommands.cpp b/Engine/Source/Runtime/Renderer/Private/MeshDrawCommands.cpp
--- a/Engine/Source/Runtime/Renderer/Private/MeshDrawCommands.cpp
+++ b/Engine/Source/Runtime/Renderer/Private/MeshDrawCommands.cpp
@@ -1352,7 +1352,7 @@ void FParallelMeshDrawCommandPass::DispatchPassSetup(
 
 	TaskContext.View = &View;
 	TaskContext.Scene = Scene;
-	TaskContext.ShadingPath = Scene->GetShadingPath();
+	TaskContext.ShadingPath = GetFeatureLevelShadingPath(View.GetFeatureLevel());
 	TaskContext.ShaderPlatform = Scene->GetShaderPlatform();
 	TaskContext.PassType = PassType;
 	TaskContext.bUseGPUScene = UseGPUScene(GMaxRHIShaderPlatform, View.GetFeatureLevel());
@@ -1468,7 +1468,7 @@ bool FParallelMeshDrawCommandPass::IsOnDemandShaderCreationEnabled()
 {
 	// GL rhi does not support multithreaded shader creation, however the engine can be configured to not run mesh drawing tasks in threads other than the RT 
 	// (see FRHICommandListExecutor::UseParallelAlgorithms()): if this condition is true, on demand shader creation can be enabled.
-	const bool bIsMobileRenderer = FSceneInterface::GetShadingPath(GMaxRHIFeatureLevel) == EShadingPath::Mobile;
+	const bool bIsMobileRenderer = GetFeatureLevelShadingPath(GMaxRHIFeatureLevel) == EShadingPath::Mobile;
 	return GAllowOnDemandShaderCreation && 
 		(GRHISupportsMultithreadedShaderCreation || (bIsMobileRenderer && (!GSupportsParallelRenderingTasksWithSeparateRHIThread && IsRunningRHIInSeparateThread())));
 }
diff --git a/Engine/Source/Runtime/Renderer/Private/PrimitiveSceneInfo.cpp b/Engine/Source/Runtime/Renderer/Private/PrimitiveSceneInfo.cpp
--- a/Engine/Source/Runtime/Renderer/Private/PrimitiveSceneInfo.cpp
+++ b/Engine/Source/Runtime/Renderer/Private/PrimitiveSceneInfo.cpp
@@ -468,7 +468,7 @@ void FPrimitiveSceneInfo::CacheMeshDrawCommands(FScene* Scene, TArrayView<FPrimi
 
 		for (int32 PassIndex = 0; PassIndex < EMeshPass::Num; PassIndex++)
 		{
-			const EShadingPath ShadingPath = Scene->GetShadingPath();
+			const EShadingPath ShadingPath = GetFeatureLevelShadingPath(Scene->GetFeatureLevel());
 			EMeshPass::Type PassType = (EMeshPass::Type)PassIndex;
 
 			if ((FPassProcessorManager::GetPassFlags(ShadingPath, PassType) & EMeshPassFlags::CachedMeshCommands) != EMeshPassFlags::None)
diff --git a/Engine/Source/Runtime/Renderer/Private/Renderer.cpp b/Engine/Source/Runtime/Renderer/Private/Renderer.cpp
--- a/Engine/Source/Runtime/Renderer/Private/Renderer.cpp
+++ b/Engine/Source/Runtime/Renderer/Private/Renderer.cpp
@@ -287,7 +287,7 @@ void FRendererModule::DrawTileMesh(FCanvasRenderContext& RenderContext, FMeshPas
 		ViewFamily->SceneTexturesConfig.Extent = View.ViewRect.Size();
 
 		const auto FeatureLevel = View.GetFeatureLevel();
-		const EShadingPath ShadingPath = FSceneInterface::GetShadingPath(FeatureLevel);
+		const EShadingPath ShadingPath = GetFeatureLevelShadingPath(FeatureLevel);
 
 		FScene* Scene = nullptr;
 		if (ViewFamily->Scene)
diff --git a/Engine/Source/Runtime/Renderer/Private/RendererScene.cpp b/Engine/Source/Runtime/Renderer/Private/RendererScene.cpp
--- a/Engine/Source/Runtime/Renderer/Private/RendererScene.cpp
+++ b/Engine/Source/Runtime/Renderer/Private/RendererScene.cpp
@@ -2635,7 +2635,7 @@ void FScene::AddLightSceneInfo_RenderThread(FLightSceneInfo* LightSceneInfo)
 			SimpleDirectionalLight = LightSceneInfo;
 		}
 
-		if(GetShadingPath() == EShadingPath::Mobile)
+		if(GetFeatureLevelShadingPath(FeatureLevel) == EShadingPath::Mobile)
 		{
 			const bool bUseCSMForDynamicObjects = LightSceneInfo->Proxy->UseCSMForDynamicObjects();
 #if !(UE_BUILD_SHIPPING || UE_BUILD_TEST)
@@ -3916,7 +3916,7 @@ void FScene::RemoveLightSceneInfo_RenderThread(FLightSceneInfo* LightSceneInfo)
 		SimpleDirectionalLight = nullptr;
 	}
 
-	if(GetShadingPath() == EShadingPath::Mobile)
+	if(GetFeatureLevelShadingPath(FeatureLevel) == EShadingPath::Mobile)
 	{
 		const bool bUseCSMForDynamicObjects = LightSceneInfo->Proxy->UseCSMForDynamicObjects();
 
@@ -4586,7 +4586,7 @@ FExclusiveDepthStencil::Type FScene::GetDefaultBasePassDepthStencilAccess(ERHIFe
 {
 	FExclusiveDepthStencil::Type BasePassDepthStencilAccess = FExclusiveDepthStencil::DepthWrite_StencilWrite;
 
-	if (GetShadingPath(InFeatureLevel) == EShadingPath::Deferred)
+	if (GetFeatureLevelShadingPath(InFeatureLevel) == EShadingPath::Deferred)
 	{
 		const EShaderPlatform ShaderPlatform = GetFeatureLevelShaderPlatform(InFeatureLevel);
 		if (ShouldForceFullDepthPass(ShaderPlatform)
@@ -4605,7 +4605,7 @@ void FScene::GetEarlyZPassMode(ERHIFeatureLevel::Type InFeatureLevel, EDepthDraw
 	bOutEarlyZPassMovable = false;
 
 	const EShaderPlatform ShaderPlatform = GetFeatureLevelShaderPlatform(InFeatureLevel);
-	if (GetShadingPath(InFeatureLevel) == EShadingPath::Deferred)
+	if (GetFeatureLevelShadingPath(InFeatureLevel) == EShadingPath::Deferred)
 	{
 		// developer override, good for profiling, can be useful as project setting
 		{
@@ -4628,7 +4628,7 @@ void FScene::GetEarlyZPassMode(ERHIFeatureLevel::Type InFeatureLevel, EDepthDraw
 			bOutEarlyZPassMovable = bDepthPassCanOutputVelocity ? false : true;
 		}
 	}
-	else if (GetShadingPath(InFeatureLevel) == EShadingPath::Mobile)
+	else if (GetFeatureLevelShadingPath(InFeatureLevel) == EShadingPath::Mobile)
 	{
 		OutZPassMode = DDM_None;
 				 
@@ -5185,7 +5185,7 @@ FLightSceneChangeSet FScene::UpdateAllLightSceneInfos(FRDGBuilder& GraphBuilder)
 			// Mobile renderer:
 			// a light with no color/intensity can cause the light to be ignored when rendering.
 			// thus, lights that change state in this way must update the draw lists.
-			if (GetShadingPath() == EShadingPath::Mobile 
+			if (GetFeatureLevelShadingPath(FeatureLevel) == EShadingPath::Mobile 
 				&& LightSceneInfo->Proxy->GetLightType() == LightType_Directional 
 				&& NewParameters.NewColor.IsAlmostBlack() != LightSceneInfo->Proxy->GetColor().IsAlmostBlack())
 			{
@@ -5324,7 +5324,7 @@ void UpdateReflectionSceneData(FScene* Scene)
 		if (Scene->ReflectionSceneData.bRegisteredReflectionCapturesHasChanged)
 		{
 			// Mobile needs to re-cache all mesh commands when scene capture data has changed
-			const bool bNeedsStaticMeshUpdate = Scene->GetShadingPath() == EShadingPath::Mobile;
+			const bool bNeedsStaticMeshUpdate = GetFeatureLevelShadingPath(Scene->GetFeatureLevel()) == EShadingPath::Mobile;
 
 			// Mark all primitives as needing an update
 			// Note: Only visible primitives will actually update their reflection proxy
diff --git a/Engine/Source/Runtime/Renderer/Private/SceneCaptureRendering.cpp b/Engine/Source/Runtime/Renderer/Private/SceneCaptureRendering.cpp
--- a/Engine/Source/Runtime/Renderer/Private/SceneCaptureRendering.cpp
+++ b/Engine/Source/Runtime/Renderer/Private/SceneCaptureRendering.cpp
@@ -504,7 +504,7 @@ static void UpdateSceneCaptureContent_RenderThread(
 {
 	FUniformExpressionCacheAsyncUpdateScope AsyncUpdateScope;
 
-	switch (SceneRenderer->Scene->GetShadingPath())
+	switch (GetFeatureLevelShadingPath(SceneRenderer->Scene->GetFeatureLevel()))
 	{
 		case EShadingPath::Mobile:
 		{
diff --git a/Engine/Source/Runtime/Renderer/Private/SceneRendering.cpp b/Engine/Source/Runtime/Renderer/Private/SceneRendering.cpp
--- a/Engine/Source/Runtime/Renderer/Private/SceneRendering.cpp
+++ b/Engine/Source/Runtime/Renderer/Private/SceneRendering.cpp
@@ -4030,7 +4030,7 @@ void FSceneRenderer::SetupMeshPass(FViewInfo& View, FExclusiveDepthStencil::Type
 {
 	SCOPE_CYCLE_COUNTER(STAT_SetupMeshPass);
 
-	const EShadingPath ShadingPath = Scene->GetShadingPath();
+	const EShadingPath ShadingPath = GetFeatureLevelShadingPath(Scene->GetFeatureLevel());
 
 	for (int32 PassIndex = 0; PassIndex < EMeshPass::Num; PassIndex++)
 	{
@@ -4117,7 +4117,7 @@ void FSceneRenderer::CreateSceneRenderers(TArrayView<const FSceneViewFamily*> In
 	const FSceneInterface* Scene = InViewFamilies[0]->Scene;
 	check(Scene);
 
-	EShadingPath ShadingPath = Scene->GetShadingPath();
+	EShadingPath ShadingPath = GetFeatureLevelShadingPath(Scene->GetFeatureLevel());
 
 	for (int32 FamilyIndex = 0; FamilyIndex < InViewFamilies.Num(); FamilyIndex++)
 	{
diff --git a/Engine/Source/Runtime/Renderer/Private/SceneTextures.cpp b/Engine/Source/Runtime/Renderer/Private/SceneTextures.cpp
--- a/Engine/Source/Runtime/Renderer/Private/SceneTextures.cpp
+++ b/Engine/Source/Runtime/Renderer/Private/SceneTextures.cpp
@@ -395,7 +395,7 @@ ENUM_CLASS_FLAGS(FSceneTextureExtentState::ERenderTargetHistory);
 void InitializeSceneTexturesConfig(FSceneTexturesConfig& Config, const FSceneViewFamily& ViewFamily)
 {
 	FIntPoint Extent = FSceneTextureExtentState::Get().Compute(ViewFamily);
-	EShadingPath ShadingPath = FSceneInterface::GetShadingPath(ViewFamily.GetFeatureLevel());
+	EShadingPath ShadingPath = GetFeatureLevelShadingPath(ViewFamily.GetFeatureLevel());
 
 	bool bRequiresAlphaChannel = ShadingPath == EShadingPath::Mobile ? IsMobilePropagateAlphaEnabled(ViewFamily.GetShaderPlatform()) : false;
 	int32 NumberOfViewsWithMultiviewEnabled = 0;
@@ -1077,11 +1077,11 @@ FSceneTextureShaderParameters CreateSceneTextureShaderParameters(
 	ESceneTextureSetupMode SetupMode)
 {
 	FSceneTextureShaderParameters Parameters;
-	if (FSceneInterface::GetShadingPath(FeatureLevel) == EShadingPath::Deferred)
+	if (GetFeatureLevelShadingPath(FeatureLevel) == EShadingPath::Deferred)
 	{
 		Parameters.SceneTextures = CreateSceneTextureUniformBuffer(GraphBuilder, SceneTextures, FeatureLevel, SetupMode);
 	}
-	else if (FSceneInterface::GetShadingPath(FeatureLevel) == EShadingPath::Mobile)
+	else if (GetFeatureLevelShadingPath(FeatureLevel) == EShadingPath::Mobile)
 	{
 		Parameters.MobileSceneTextures = CreateMobileSceneTextureUniformBuffer(GraphBuilder, SceneTextures, Translate(SetupMode));
 	}
@@ -1091,11 +1091,11 @@ FSceneTextureShaderParameters CreateSceneTextureShaderParameters(
 FSceneTextureShaderParameters CreateSceneTextureShaderParameters(FRDGBuilder& GraphBuilder, const FSceneView& View, ESceneTextureSetupMode SetupMode)
 {
 	FSceneTextureShaderParameters Parameters;
-	if (FSceneInterface::GetShadingPath(View.FeatureLevel) == EShadingPath::Deferred)
+	if (GetFeatureLevelShadingPath(View.FeatureLevel) == EShadingPath::Deferred)
 	{
 		Parameters.SceneTextures = CreateSceneTextureUniformBuffer(GraphBuilder, View, SetupMode);
 	}
-	else if (FSceneInterface::GetShadingPath(View.FeatureLevel) == EShadingPath::Mobile)
+	else if (GetFeatureLevelShadingPath(View.FeatureLevel) == EShadingPath::Mobile)
 	{
 		Parameters.MobileSceneTextures = CreateMobileSceneTextureUniformBuffer(GraphBuilder, View, Translate(SetupMode));
 	}
diff --git a/Engine/Source/Runtime/Renderer/Private/SceneVisibility.cpp b/Engine/Source/Runtime/Renderer/Private/SceneVisibility.cpp
--- a/Engine/Source/Runtime/Renderer/Private/SceneVisibility.cpp
+++ b/Engine/Source/Runtime/Renderer/Private/SceneVisibility.cpp
@@ -1036,7 +1036,7 @@ void FDrawCommandRelevancePacket::AddCommandsForMesh(
 	bool bCanCache, 
 	EMeshPass::Type PassType)
 {
-	const EShadingPath ShadingPath = Scene.GetShadingPath();
+	const EShadingPath ShadingPath = GetFeatureLevelShadingPath(Scene.GetFeatureLevel());
 	const bool bUseCachedMeshCommand = bUseCachedMeshDrawCommands
 		&& !!(FPassProcessorManager::GetPassFlags(ShadingPath, PassType) & EMeshPassFlags::CachedMeshCommands)
 		&& StaticMeshRelevance.bSupportsCachingMeshDrawCommands
@@ -1146,7 +1146,7 @@ void FRelevancePacket::Finalize()
 {
 	FViewInfo& WriteView = const_cast<FViewInfo&>(View);
 	FViewCommands& WriteViewCommands = const_cast<FViewCommands&>(ViewCommands);
-	const EShadingPath ShadingPath = Scene.GetShadingPath();
+	const EShadingPath ShadingPath = GetFeatureLevelShadingPath(Scene.GetFeatureLevel());
 
 	for (int32 BitIndex : NotDrawRelevant.Prims)
 	{
@@ -1253,7 +1253,7 @@ void FRelevancePacket::ComputeRelevance(FDynamicPrimitiveIndexList& DynamicPrimi
 	bUsesSecondStageDepthPass = 0;
 	bUsesLightingChannels = false;
 	bTranslucentSurfaceLighting = false;
-	const EShadingPath ShadingPath = Scene.GetShadingPath();
+	const EShadingPath ShadingPath = GetFeatureLevelShadingPath(Scene.GetFeatureLevel());
 	const bool bHairStrandsEnabled = IsHairStrandsEnabled(EHairStrandsShaderType::All, Scene.GetShaderPlatform());
 
 	int32 NumVisibleStaticMeshElements = 0;
@@ -3879,7 +3879,7 @@ FVisibilityTaskData::FVisibilityTaskData(FRHICommandListImmediate& InRHICmdList,
 	, Scene(*SceneRenderer.Scene)
 	, Views(SceneRenderer.AllViews)
 	, ViewFamily(SceneRenderer.ViewFamily)
-	, ShadingPath(Scene.GetShadingPath())
+	, ShadingPath(GetFeatureLevelShadingPath(Scene.GetFeatureLevel()))
 	, TaskConfig(Scene, Views)
 	, bAddLightmapDensityCommands(ViewFamily.EngineShowFlags.LightMapDensity&& AllowDebugViewmodes())
 {
diff --git a/Engine/Source/Runtime/Renderer/Private/ShadowDepthRendering.cpp b/Engine/Source/Runtime/Renderer/Private/ShadowDepthRendering.cpp
--- a/Engine/Source/Runtime/Renderer/Private/ShadowDepthRendering.cpp
+++ b/Engine/Source/Runtime/Renderer/Private/ShadowDepthRendering.cpp
@@ -415,12 +415,12 @@ public:
 	{
 		const ERHIFeatureLevel::Type FeatureLevel = GetMaxSupportedFeatureLevel((EShaderPlatform)Initializer.Target.Platform);
 
-		if (FSceneInterface::GetShadingPath(FeatureLevel) == EShadingPath::Deferred)
+		if (GetFeatureLevelShadingPath(FeatureLevel) == EShadingPath::Deferred)
 		{
 			PassUniformBuffer.Bind(Initializer.ParameterMap, FShadowDepthPassUniformParameters::FTypeInfo::GetStructMetadata()->GetShaderVariableName());
 		}
 
-		if (FSceneInterface::GetShadingPath(FeatureLevel) == EShadingPath::Mobile)
+		if (GetFeatureLevelShadingPath(FeatureLevel) == EShadingPath::Mobile)
 		{
 			PassUniformBuffer.Bind(Initializer.ParameterMap, FMobileShadowDepthPassUniformParameters::FTypeInfo::GetStructMetadata()->GetShaderVariableName());
 		}
@@ -1046,7 +1046,7 @@ void FProjectedShadowInfo::BeginRenderView(FRDGBuilder& GraphBuilder, FScene* Sc
 	if (DependentView)
 	{
 		const ERHIFeatureLevel::Type FeatureLevel = ShadowDepthView->FeatureLevel;
-		if (FSceneInterface::GetShadingPath(FeatureLevel) == EShadingPath::Deferred)
+		if (GetFeatureLevelShadingPath(FeatureLevel) == EShadingPath::Deferred)
 		{
 			extern TSet<IPersistentViewUniformBufferExtension*> PersistentViewUniformBufferExtensions;
 
@@ -1138,7 +1138,7 @@ void FProjectedShadowInfo::RenderDepth(
 
 	PassParameters->VirtualShadowMap = SceneRenderer->VirtualShadowMapArray.GetUniformBuffer();
 
-	switch (FSceneInterface::GetShadingPath(FeatureLevel))
+	switch (GetFeatureLevelShadingPath(FeatureLevel))
 	{
 	case EShadingPath::Deferred:
 	{
diff --git a/Engine/Source/Runtime/Renderer/Private/ShadowSetup.cpp b/Engine/Source/Runtime/Renderer/Private/ShadowSetup.cpp
--- a/Engine/Source/Runtime/Renderer/Private/ShadowSetup.cpp
+++ b/Engine/Source/Runtime/Renderer/Private/ShadowSetup.cpp
@@ -1166,7 +1166,7 @@ void FProjectedShadowInfo::AddCachedMeshDrawCommandsForPass(
 	TArray<EMeshDrawCommandCullingPayloadFlags, SceneRenderingAllocator> MeshCommandBuildFlags,
 	int32& NumMeshCommandBuildRequestElements)
 {
-	const EShadingPath ShadingPath = Scene->GetShadingPath();
+	const EShadingPath ShadingPath = GetFeatureLevelShadingPath(Scene->GetFeatureLevel());
 	const bool bUseCachedMeshCommand = UseCachedMeshDrawCommands()
 		&& !!(FPassProcessorManager::GetPassFlags(ShadingPath, PassType) & EMeshPassFlags::CachedMeshCommands)
 		&& StaticMeshRelevance.bSupportsCachingMeshDrawCommands;
@@ -1653,7 +1653,7 @@ void FProjectedShadowInfo::AddCachedMeshDrawCommands_AnyThread(
 	FAddSubjectPrimitiveOverflowedIndices& OverflowBuffer) const
 {
 	const EMeshPass::Type PassType = MeshPassTargetType;
-	const EShadingPath ShadingPath = Scene->GetShadingPath();
+	const EShadingPath ShadingPath = GetFeatureLevelShadingPath(Scene->GetFeatureLevel());
 	const bool bUseCachedMeshCommand = UseCachedMeshDrawCommands_AnyThread()
 		&& !!(FPassProcessorManager::GetPassFlags(ShadingPath, PassType) & EMeshPassFlags::CachedMeshCommands)
 		&& StaticMeshRelevance.bSupportsCachingMeshDrawCommands;
@@ -2084,7 +2084,7 @@ uint64 FProjectedShadowInfo::AddSubjectPrimitive_AnyThread(
 	FAddSubjectPrimitiveResult Result;
 	Result.Qword = 0;
 
-	if (FSceneInterface::GetShadingPath(FeatureLevel) == EShadingPath::Mobile)
+	if (GetFeatureLevelShadingPath(FeatureLevel) == EShadingPath::Mobile)
 	{
 		const bool bShouldRecordShadowSubjectsForMobile = GetLightSceneInfo().ShouldRecordShadowSubjectsForMobile();
 
@@ -2501,7 +2501,7 @@ void FProjectedShadowInfo::SetupMeshDrawCommandsForProjectionStenciling(FSceneRe
 {
 	QUICK_SCOPE_CYCLE_COUNTER(STAT_SetupMeshDrawCommandsForShadowDepth);
 
-	const EShadingPath ShadingPath = FSceneInterface::GetShadingPath(Renderer.FeatureLevel);
+	const EShadingPath ShadingPath = GetFeatureLevelShadingPath(Renderer.FeatureLevel);
 	static const auto EnableModulatedSelfShadowCVar = IConsoleManager::Get().FindTConsoleVariableDataInt(TEXT("r.Shadow.EnableModulatedSelfShadow"));
 	const bool bMobileModulatedShadowsAllowSelfShadow = !bSelfShadowOnly && (ShadingPath == EShadingPath::Mobile && !EnableModulatedSelfShadowCVar->GetValueOnRenderThread() && LightSceneInfo->Proxy && LightSceneInfo->Proxy->CastsModulatedShadows());
 	if (bPreShadow || bSelfShadowOnly || bMobileModulatedShadowsAllowSelfShadow)
diff --git a/Engine/Source/Runtime/Renderer/Private/SkyPassRendering.cpp b/Engine/Source/Runtime/Renderer/Private/SkyPassRendering.cpp
--- a/Engine/Source/Runtime/Renderer/Private/SkyPassRendering.cpp
+++ b/Engine/Source/Runtime/Renderer/Private/SkyPassRendering.cpp
@@ -66,7 +66,7 @@ bool FSkyPassMeshProcessor::Process(
 	FUniformLightMapPolicy NoLightmapPolicy(LMP_NO_LIGHTMAP);
 	const FVertexFactory* VertexFactory = MeshBatch.VertexFactory;
 
-	if (Scene->GetShadingPath()==EShadingPath::Deferred)
+	if (GetFeatureLevelShadingPath(FeatureLevel) == EShadingPath::Deferred)
 	{
 		TMeshProcessorShaders<
 			TBasePassVertexShaderPolicyParamType<LightMapPolicyType>,
@@ -167,7 +167,7 @@ void FSkyPassMeshProcessor::CollectPSOInitializers(const FSceneTexturesConfig& S
 	}
 
 	// Only do deferred path for now
-	if (FScene::GetShadingPath(FeatureLevel) != EShadingPath::Deferred)
+	if (GetFeatureLevelShadingPath(FeatureLevel) != EShadingPath::Deferred)
 	{
 		return;
 	}
