RHI: avoid out of bounds access when VirtualShadowMapId is INDEX_NONE in FTranslucentLightingInjectPS.

#rnx

diff --git a/Engine/Source/Runtime/Renderer/Private/TranslucentLighting.cpp b/Engine/Source/Runtime/Renderer/Private/TranslucentLighting.cpp
--- a/Engine/Source/Runtime/Renderer/Private/TranslucentLighting.cpp
+++ b/Engine/Source/Runtime/Renderer/Private/TranslucentLighting.cpp
@@ -1116,7 +1116,7 @@ void InjectTranslucencyLightingVolume(
 			const FVisibleLightInfo& VisibleLightInfo = Renderer.VisibleLightInfos[LightSceneInfo->Id];
 			const bool bInverseSquared = LightSceneInfo->Proxy->IsInverseSquared();
 			const bool bDirectionalLight = LightSceneInfo->Proxy->GetLightType() == LightType_Directional;
-			const bool bUseVSM = Renderer.VirtualShadowMapArray.IsAllocated();
+			bool bUseVSM = Renderer.VirtualShadowMapArray.IsAllocated();
 
 			const FVolumeBounds VolumeBounds = CalculateLightVolumeBounds(LightSceneInfo->Proxy->GetBoundingSphere(), View, VolumeCascadeIndex, bDirectionalLight);
 			if (VolumeBounds.IsValid())
@@ -1147,7 +1147,11 @@ void InjectTranslucencyLightingVolume(
 
 				GetVolumeShadowingShaderParameters(GraphBuilder, View, LightSceneInfo, InjectionData.ProjectedShadowInfo, PassParameters->PS.VolumeShadowingParameters);
 
-				int32 VirtualShadowMapId = bUseVSM ? Renderer.VisibleLightInfos[LightSceneInfo->Id].GetVirtualShadowMapId(&View) : INDEX_NONE;
+				const int32 VirtualShadowMapId = bUseVSM ? Renderer.VisibleLightInfos[LightSceneInfo->Id].GetVirtualShadowMapId(&View) : INDEX_NONE;
+				
+				// Switch it back off if there's no ID to avoid the FVirtualShadowMap permutation if we don't need it
+				bUseVSM = (VirtualShadowMapId != INDEX_NONE);
+
 				PassParameters->PS.VirtualShadowMapId = VirtualShadowMapId;
 				PassParameters->PS.LightFunctionParameters = FLightFunctionSharedParameters::GetLightFunctionSharedParameters(LightSceneInfo, 1.0f);
 				PassParameters->PS.VolumeCascadeIndex = VolumeCascadeIndex;
