Fix crash during -nullrhi game auto-test.

diff --git a/Engine/Source/Runtime/Renderer/Private/ScenePrivate.h b/Engine/Source/Runtime/Renderer/Private/ScenePrivate.h
--- a/Engine/Source/Runtime/Renderer/Private/ScenePrivate.h
+++ b/Engine/Source/Runtime/Renderer/Private/ScenePrivate.h
@@ -1508,7 +1508,7 @@ PRAGMA_ENABLE_DEPRECATION_WARNINGS
 			NewMID->CopyInterpParameters(InputAsMID);
 		}
 
-		check(NewMID->GetRenderProxy());
+		check(!FApp::CanEverRender() || NewMID->GetRenderProxy());
 		MIDUsedCount++;
 		return NewMID;
 	}
