Path Tracing Denoiser third party interface

Problem:
The current callback functions used to hook into the path tracer denoiser don't provide a RDG builder, which is required to do inference using the NNERuntimeRDG.

Solution:
* Introduced a simple, clean and easy to extendable interface that can be used by third party plugins to add their own path tracer denoising capabilities by adding dispatch passes through the RDG builder
* Updated existing denosing plugins to use the new interface

Tests:
* P4V Non-Unity build
* Verified OpenImageDenoise plugin works as expected
* Verified OptiX plugin works as expected

diff --git a/Engine/Plugins/Experimental/OpenImageDenoise/Source/OpenImageDenoise/OpenImageDenoise.Build.cs b/Engine/Plugins/Experimental/OpenImageDenoise/Source/OpenImageDenoise/OpenImageDenoise.Build.cs
--- a/Engine/Plugins/Experimental/OpenImageDenoise/Source/OpenImageDenoise/OpenImageDenoise.Build.cs
+++ b/Engine/Plugins/Experimental/OpenImageDenoise/Source/OpenImageDenoise/OpenImageDenoise.Build.cs
@@ -9,6 +9,7 @@ namespace UnrealBuildTool.Rules
 			PrivateDependencyModuleNames.AddRange(
 				new string[] {
 					"Core",
+					"RenderCore",
 					"Renderer",
 					"RHI",
 					"IntelOIDN"
diff --git a/Engine/Plugins/Experimental/OpenImageDenoise/Source/OpenImageDenoise/Private/Plugin.cpp b/Engine/Plugins/Experimental/OpenImageDenoise/Source/OpenImageDenoise/Private/Plugin.cpp
--- a/Engine/Plugins/Experimental/OpenImageDenoise/Source/OpenImageDenoise/Private/Plugin.cpp
+++ b/Engine/Plugins/Experimental/OpenImageDenoise/Source/OpenImageDenoise/Private/Plugin.cpp
@@ -5,13 +5,22 @@
 #include "Modules/ModuleInterface.h"
 #include "Modules/ModuleManager.h"
 #include "PathTracingDenoiser.h"
+#include "RenderGraphBuilder.h"
 #include "RHICommandList.h"
+#include "RHIGPUReadback.h"
 #include "RHIResources.h"
 #include "RHITypes.h"
-#include "RHIGPUReadback.h"
+#include "SceneView.h"
 
 #include "OpenImageDenoise/oidn.hpp"
 
+BEGIN_SHADER_PARAMETER_STRUCT(FDenoiseTextureParameters, )
+	RDG_TEXTURE_ACCESS(InputTexture, ERHIAccess::CopySrc)
+	RDG_TEXTURE_ACCESS(InputAlbedo, ERHIAccess::CopySrc)
+	RDG_TEXTURE_ACCESS(InputNormal, ERHIAccess::CopySrc)
+	RDG_TEXTURE_ACCESS(OutputTexture, ERHIAccess::CopyDest)
+END_SHADER_PARAMETER_STRUCT()
+
 class FOpenImageDenoiseModule : public IModuleInterface
 {
 public:
@@ -156,7 +165,6 @@ static void CopyTextureFromCPUToGPU(FRHICommandListImmediate& RHICmdList, const
 	RHICmdList.UnlockTexture2D(DstTexture, 0, false);
 }
 
-
 static void Denoise(FRHICommandListImmediate& RHICmdList, FRHITexture* ColorTex, FRHITexture* AlbedoTex, FRHITexture* NormalTex, FRHITexture* OutputTex, FRHIGPUMask GPUMask)
 {
 	static IConsoleVariable* DenoiseModeCVar = IConsoleManager::Get().FindConsoleVariable(TEXT("r.PathTracing.Denoiser"));
@@ -209,13 +217,42 @@ static void Denoise(FRHICommandListImmediate& RHICmdList, FRHITexture* ColorTex,
 #endif
 }
 
+using namespace UE::Renderer::Private;
+
+class FOIDNDenoiser : public IPathTracingDenoiser
+{
+public:
+	~FOIDNDenoiser() {}
+
+	void AddPasses(FRDGBuilder& GraphBuilder, const FSceneView& View, const FInputs& Inputs) const override
+	{
+		FDenoiseTextureParameters* DenoiseParameters = GraphBuilder.AllocParameters<FDenoiseTextureParameters>();
+		DenoiseParameters->InputTexture = Inputs.ColorTex;
+		DenoiseParameters->InputAlbedo = Inputs.AlbedoTex;
+		DenoiseParameters->InputNormal = Inputs.NormalTex;
+		DenoiseParameters->OutputTexture = Inputs.OutputTex;
+
+		// Need to read GPU mask outside Pass function, as the value is not refreshed inside the pass
+		GraphBuilder.AddPass(RDG_EVENT_NAME("OIDN Denoiser Plugin"), DenoiseParameters, ERDGPassFlags::Readback,
+			[DenoiseParameters, GPUMask = View.GPUMask](FRHICommandListImmediate& RHICmdList)
+		{
+			Denoise(RHICmdList,
+				DenoiseParameters->InputTexture->GetRHI()->GetTexture2D(),
+				DenoiseParameters->InputAlbedo->GetRHI()->GetTexture2D(),
+				DenoiseParameters->InputNormal->GetRHI()->GetTexture2D(),
+				DenoiseParameters->OutputTexture->GetRHI()->GetTexture2D(),
+				GPUMask);
+		});
+	}
+};
 
 void FOpenImageDenoiseModule::StartupModule()
 {
 #if WITH_EDITOR
 	UE_LOG(LogOpenImageDenoise, Log, TEXT("OIDN starting up"));
 #endif
-	GPathTracingDenoiserFunc = &Denoise;
+
+	GPathTracingDenoiserPlugin = MakeUnique<FOIDNDenoiser>();
 }
 
 void FOpenImageDenoiseModule::ShutdownModule()
@@ -223,7 +260,8 @@ void FOpenImageDenoiseModule::ShutdownModule()
 #if WITH_EDITOR
 	UE_LOG(LogOpenImageDenoise, Log, TEXT("OIDN shutting down"));
 #endif
+
 	// Release scratch memory and destroy the OIDN device and filters
 	DenoiserState.Reset();
-	GPathTracingDenoiserFunc = nullptr;
+	GPathTracingDenoiserPlugin.Reset();
 }
diff --git a/Engine/Plugins/Experimental/OptiXDenoise/Source/OptiXDenoise/Private/Plugin.cpp b/Engine/Plugins/Experimental/OptiXDenoise/Source/OptiXDenoise/Private/Plugin.cpp
--- a/Engine/Plugins/Experimental/OptiXDenoise/Source/OptiXDenoise/Private/Plugin.cpp
+++ b/Engine/Plugins/Experimental/OptiXDenoise/Source/OptiXDenoise/Private/Plugin.cpp
@@ -3,12 +3,29 @@
 #include "OptiXDenoiser.h"
 #include "HAL/IConsoleManager.h"
 #include "PathTracingDenoiser.h"
+#include "RenderGraphBuilder.h"
+#include "SceneView.h"
 #include <mutex>
 
 #define LOCTEXT_NAMESPACE "FOptiXDenoiseModule"
 
 DEFINE_LOG_CATEGORY(LogOptiXDenoise);
 
+BEGIN_SHADER_PARAMETER_STRUCT(FDenoiseTextureExtParameters, )
+	RDG_TEXTURE_ACCESS(InputTexture, ERHIAccess::CopySrc)
+	RDG_TEXTURE_ACCESS(InputAlbedo, ERHIAccess::CopySrc)
+	RDG_TEXTURE_ACCESS(InputNormal, ERHIAccess::CopySrc)
+	RDG_TEXTURE_ACCESS(InputFlow, ERHIAccess::CopySrc)
+	RDG_TEXTURE_ACCESS(InputPreviousOutput, ERHIAccess::CopySrc)
+	RDG_TEXTURE_ACCESS(OutputTexture, ERHIAccess::CopyDest)
+END_SHADER_PARAMETER_STRUCT()
+
+BEGIN_SHADER_PARAMETER_STRUCT(FMotionVectorParameters, )
+	RDG_TEXTURE_ACCESS(InputFrameTexture, ERHIAccess::CopySrc)
+	RDG_TEXTURE_ACCESS(ReferenceFrameTexture, ERHIAccess::CopySrc)
+	RDG_TEXTURE_ACCESS(OutputTexture, ERHIAccess::CopyDest)
+END_SHADER_PARAMETER_STRUCT()
+
 namespace
 {
 	TAutoConsoleVariable<int32> CVarOptixDenoiseTileWidth(
@@ -117,7 +134,10 @@ static void Denoise(
 	FRHITexture* AlbedoTex,
 	FRHITexture* NormalTex,
 	FRHITexture* OutputTex,
-	const FDenoisingArgumentsExt* DenoisingArgumentExt,
+	FRHITexture* FlowTex,
+	FRHITexture* PreviousOutputTex,
+	int DenoisingFrameId,
+	bool bForceSpatialDenoiserOnly,
 	FRHIGPUMask GPUMask)
 {
 	FOptiXCudaFunctionList::Get().InitializeCudaModule();
@@ -136,7 +156,7 @@ static void Denoise(
 		});
 
 	// (Re)initialize the denoiser if this is the first frame (e.g., new cut) or the image size has changed
-	if (DenoisingArgumentExt->bForceSpatialDenoiserOnly ||
+	if (bForceSpatialDenoiserOnly ||
 		FOptiXDenoiseModule::Denoiser->IsImageSizeChanged(Size.X, Size.Y))
 	{
 		FOptiXDenoiseModule::Denoiser->Finish();
@@ -148,12 +168,12 @@ static void Denoise(
 			CVarOptixDenoiseTileHeight.GetValueOnAnyThread());
 
 		UE_LOG(LogOptiXDenoise, Log, TEXT("Denoising Frame %d (%d x %d) Spatially"),
-			DenoisingArgumentExt->DenoisingFrameId, Size.X, Size.Y);
+			DenoisingFrameId, Size.X, Size.Y);
 	}
 	else
 	{
 		UE_LOG(LogOptiXDenoise, Log, TEXT("Denoising Frame %d (%d x %d) temporally"),
-			DenoisingArgumentExt->DenoisingFrameId, Size.X, Size.Y);
+			DenoisingFrameId, Size.X, Size.Y);
 	}
 
 	const CUstream& CudaStream = FOptiXDenoiseModule::Denoiser->GetCudaStream();
@@ -191,8 +211,8 @@ static void Denoise(
 		UploadTextureToDenoiser(ColorTex, EDenoisingImageType::COLOR, TEXT("ColorTex"));
 		UploadTextureToDenoiser(AlbedoTex, EDenoisingImageType::ALBEDO, TEXT("AlbedoTex"));
 		UploadTextureToDenoiser(NormalTex, EDenoisingImageType::NORMAL, TEXT("NormalTex"));
-		UploadTextureToDenoiser(DenoisingArgumentExt->FlowTex, EDenoisingImageType::FLOW, TEXT("FlowTex"));
-		UploadTextureToDenoiser(DenoisingArgumentExt->PreviousOutputTex, EDenoisingImageType::PREVOUTPUT, TEXT("PreviousOutputTex"));
+		UploadTextureToDenoiser(FlowTex, EDenoisingImageType::FLOW, TEXT("FlowTex"));
+		UploadTextureToDenoiser(PreviousOutputTex, EDenoisingImageType::PREVOUTPUT, TEXT("PreviousOutputTex"));
 		UploadTextureToDenoiser(OutputTex, EDenoisingImageType::OUTPUT, TEXT("OutputTex"));
 	}
 
@@ -209,6 +229,86 @@ static void Denoise(
 	FOptiXDenoiseModule::OptiXImageFactory->FlushImages();
 }
 
+using namespace UE::Renderer::Private;
+
+class FOptiXDenosier : public IPathTracingSpatialTemporalDenoiser
+{
+public:
+	class FHistory : public IHistory
+	{
+	public:
+		FHistory(const TCHAR* DebugName) : DebugName(DebugName) { }
+
+		virtual ~FHistory() { }
+
+		const TCHAR* GetDebugName() const override { return DebugName; };
+
+	private:
+		const TCHAR* DebugName;
+	};
+
+	~FOptiXDenosier() {}
+
+	const TCHAR* GetDebugName() const override { return *DebugName; }
+
+	virtual FOutputs AddPasses(FRDGBuilder& GraphBuilder, const FSceneView& View, const FInputs& Inputs) const
+	{
+		FDenoiseTextureExtParameters* DenoiseParameters = GraphBuilder.AllocParameters<FDenoiseTextureExtParameters>();
+		DenoiseParameters->InputTexture = Inputs.ColorTex;
+		DenoiseParameters->InputAlbedo = Inputs.AlbedoTex;
+		DenoiseParameters->InputNormal = Inputs.NormalTex;
+		DenoiseParameters->InputFlow = Inputs.FlowTex;
+		DenoiseParameters->InputPreviousOutput = Inputs.PreviousOutputTex;
+		DenoiseParameters->OutputTexture = Inputs.OutputTex;
+
+		const int DenoisingFrameId = Inputs.DenoisingFrameId;
+		const bool bForceSpatialDenoiserOnly = Inputs.bForceSpatialDenoiserOnly;
+
+		// Need to read GPU mask outside Pass function, as the value is not refreshed inside the pass
+		GraphBuilder.AddPass(RDG_EVENT_NAME("Path Tracer Denoiser Ext Plugin"), DenoiseParameters, ERDGPassFlags::Readback,
+			[DenoiseParameters, DenoisingFrameId, bForceSpatialDenoiserOnly, GPUMask = View.GPUMask](FRHICommandListImmediate& RHICmdList)
+		{
+			Denoise(RHICmdList,
+				DenoiseParameters->InputTexture->GetRHI()->GetTexture2D(),
+				DenoiseParameters->InputAlbedo->GetRHI()->GetTexture2D(),
+				DenoiseParameters->InputNormal->GetRHI()->GetTexture2D(),
+				DenoiseParameters->OutputTexture->GetRHI()->GetTexture2D(),
+				DenoiseParameters->InputFlow ? DenoiseParameters->InputFlow->GetRHI()->GetTexture2D() : nullptr,
+				DenoiseParameters->InputPreviousOutput ? DenoiseParameters->InputPreviousOutput->GetRHI()->GetTexture2D() : nullptr,
+				DenoisingFrameId,
+				bForceSpatialDenoiserOnly,
+				GPUMask);
+		});
+
+		return {};
+	}
+
+	virtual void AddMotionVectorPass(FRDGBuilder& GraphBuilder, const FSceneView& View, const FMotionVectorInputs& Inputs) const
+	{
+		FMotionVectorParameters* DenoiseParameters = GraphBuilder.AllocParameters<FMotionVectorParameters>();
+		DenoiseParameters->InputFrameTexture = Inputs.InputFrameTex;
+		DenoiseParameters->ReferenceFrameTexture = Inputs.ReferenceFrameTex;
+		DenoiseParameters->OutputTexture = Inputs.OutputTex;
+
+		const float PreExposure = Inputs.PreExposure;
+		
+		GraphBuilder.AddPass(RDG_EVENT_NAME("OptiX Motion Vector Pass"), DenoiseParameters, ERDGPassFlags::Readback,
+			[DenoiseParameters, PreExposure, GPUMask = View.GPUMask](FRHICommandListImmediate& RHICmdList)
+		{
+
+			OpticalFlow(RHICmdList,
+				DenoiseParameters->InputFrameTexture->GetRHI()->GetTexture2D(),
+				DenoiseParameters->ReferenceFrameTexture->GetRHI()->GetTexture2D(),
+				DenoiseParameters->OutputTexture->GetRHI()->GetTexture2D(),
+				PreExposure,
+				GPUMask);
+		});
+	}
+
+private:
+	inline static const FString DebugName = TEXT("FOptiXDenosier");
+};
+
 void FOptiXDenoiseModule::StartupModule()
 {
 #if WITH_EDITOR
@@ -224,8 +324,7 @@ void FOptiXDenoiseModule::StartupModule()
 	OptiXDenoiseBaseDLLHandle = FPlatformProcess::GetDllHandle(*OptiXDenoiseBaseDllPath);
 	FOptiXCudaFunctionList::Get().RegisterFunctionInstance<FOptiXDenoiserFunctionInstance>();
 
-	GPathTracingSpatialTemporalDenoiserFunc = &Denoise;
-	GPathTracingMotionVectorFunc = &OpticalFlow;
+	GPathTracingSpatialTemporalDenoiserPlugin = MakeUnique<FOptiXDenosier>();
 }
 
 void FOptiXDenoiseModule::ShutdownModule()
@@ -233,8 +332,8 @@ void FOptiXDenoiseModule::ShutdownModule()
 #if WITH_EDITOR
 	UE_LOG(LogOptiXDenoise, Log, TEXT("OptiXDenoise shutting down"));
 #endif
-	GPathTracingSpatialTemporalDenoiserFunc = nullptr;
-	GPathTracingMotionVectorFunc = nullptr;
+
+	GPathTracingSpatialTemporalDenoiserPlugin.Reset();
 
 	// Assure resources related to CUDA is released before the releasing of CUDA module.
 	Denoiser.Reset();
diff --git a/Engine/Source/Runtime/Renderer/Private/PathTracing.cpp b/Engine/Source/Runtime/Renderer/Private/PathTracing.cpp
--- a/Engine/Source/Runtime/Renderer/Private/PathTracing.cpp
+++ b/Engine/Source/Runtime/Renderer/Private/PathTracing.cpp
@@ -584,6 +584,9 @@ struct FPathTracingState {
 	// Buffer containing the starting medium extinction
 	TRefCountPtr<FRDGPooledBuffer>    StartingExtinctionCoefficient;
 
+	// Custom path tracing spacial temporal denoiser result, used by plugins
+	TRefCountPtr<UE::Renderer::Private::IPathTracingSpatialTemporalDenoiser::IHistory> SpatialTemporalDenoiserHistory;
+
 	// Current sample index to be rendered by the path tracer - this gets incremented each time the path tracer accumulates a frame of samples
 	uint32 SampleIndex = 0;
 
@@ -2382,6 +2385,8 @@ void FSceneViewState::PathTracingInvalidate(bool InvalidateAnimationStates)
 			State->LastNormalRT.SafeRelease();
 			State->LastAlbedoRT.SafeRelease();
 			State->LastVarianceBuffer.SafeRelease();
+
+			State->SpatialTemporalDenoiserHistory.SafeRelease();
 		}
 
 		State->RadianceRT.SafeRelease();
@@ -3220,7 +3225,8 @@ void FDeferredShadingSceneRenderer::RenderPathTracing(
 		}
 	}
 
-	FPathTracingSpatialTemporalDenoisingContext DenoisingContext;
+	FPathTracingSpatialTemporalDenoisingContext DenoisingContext = {};
+	DenoisingContext.SpatialTemporalDenoiserHistory = PathTracingState->SpatialTemporalDenoiserHistory;
 	const bool EnablePathTracingDenoiserRealtimeDebug = ShouldEnablePathTracingDenoiserRealtimeDebug();
 
 	if (IsDenoiserEnabled)
@@ -3280,6 +3286,8 @@ void FDeferredShadingSceneRenderer::RenderPathTracing(
 			GraphBuilder.QueueTextureExtraction(NormalTexture, &PathTracingState->LastNormalRT);
 			GraphBuilder.QueueTextureExtraction(AlbedoTexture, &PathTracingState->LastAlbedoRT);
 			GraphBuilder.QueueTextureExtraction(RadianceTexture, &PathTracingState->LastRadianceRT);
+
+			PathTracingState->SpatialTemporalDenoiserHistory = DenoisingContext.SpatialTemporalDenoiserHistory;
 		}
 
 		// 3. Update pixel variance
diff --git a/Engine/Source/Runtime/Renderer/Private/PathTracingSpatialTemporalDenoising.cpp b/Engine/Source/Runtime/Renderer/Private/PathTracingSpatialTemporalDenoising.cpp
--- a/Engine/Source/Runtime/Renderer/Private/PathTracingSpatialTemporalDenoising.cpp
+++ b/Engine/Source/Runtime/Renderer/Private/PathTracingSpatialTemporalDenoising.cpp
@@ -3,12 +3,10 @@
 #include "PathTracingSpatialTemporalDenoising.h"
 #include "PathTracing.h"
 #include "RHI.h"
-#include "PathTracingDenoiser.h"
-#include "RHIDefinitions.h"
 
-PathTracingDenoiserFunction* GPathTracingDenoiserFunc = nullptr;
-PathTracingSpatialTemporalDenoiserFunction* GPathTracingSpatialTemporalDenoiserFunc = nullptr;
-PathTracingMotionVectorFunction* GPathTracingMotionVectorFunc = nullptr;
+TUniquePtr<UE::Renderer::Private::IPathTracingDenoiser> GPathTracingDenoiserPlugin;
+TUniquePtr<UE::Renderer::Private::IPathTracingSpatialTemporalDenoiser> GPathTracingSpatialTemporalDenoiserPlugin;
+
 #if RHI_RAYTRACING
 
 #include "DeferredShadingRenderer.h"
@@ -26,28 +24,6 @@ PathTracingMotionVectorFunction* GPathTracingMotionVectorFunc = nullptr;
 
 DEFINE_LOG_CATEGORY_STATIC(LogPathTracingDenoising, Log, All);
 
-BEGIN_SHADER_PARAMETER_STRUCT(FDenoiseTextureParameters, )
-	RDG_TEXTURE_ACCESS(InputTexture, ERHIAccess::CopySrc)
-	RDG_TEXTURE_ACCESS(InputAlbedo, ERHIAccess::CopySrc)
-	RDG_TEXTURE_ACCESS(InputNormal, ERHIAccess::CopySrc)
-	RDG_TEXTURE_ACCESS(OutputTexture, ERHIAccess::CopyDest)
-END_SHADER_PARAMETER_STRUCT()
-
-BEGIN_SHADER_PARAMETER_STRUCT(FDenoiseTextureExtParameters, )
-	RDG_TEXTURE_ACCESS(InputTexture, ERHIAccess::CopySrc)
-	RDG_TEXTURE_ACCESS(InputAlbedo, ERHIAccess::CopySrc)
-	RDG_TEXTURE_ACCESS(InputNormal, ERHIAccess::CopySrc)
-	RDG_TEXTURE_ACCESS(InputFlow, ERHIAccess::CopySrc)
-	RDG_TEXTURE_ACCESS(InputPreviousOutput, ERHIAccess::CopySrc)
-	RDG_TEXTURE_ACCESS(OutputTexture, ERHIAccess::CopyDest)
-END_SHADER_PARAMETER_STRUCT()
-
-BEGIN_SHADER_PARAMETER_STRUCT(FMotionVectorParameters, )
-	RDG_TEXTURE_ACCESS(InputFrameTexture, ERHIAccess::CopySrc)
-	RDG_TEXTURE_ACCESS(ReferenceFrameTexture, ERHIAccess::CopySrc)
-	RDG_TEXTURE_ACCESS(OutputTexture, ERHIAccess::CopyDest)
-END_SHADER_PARAMETER_STRUCT()
-
 namespace {
 
 	TAutoConsoleVariable<int32> CVarPathTracingDenoiser(
@@ -280,10 +256,7 @@ int GetPathTracingDenoiserMode(const FViewInfo& View)
 
 bool IsPathTracingDenoiserEnabled(const FViewInfo& View)
 {
-	int DenoiserMode = GetPathTracingDenoiserMode(View);
-	const bool IsDenoiserEnabled = DenoiserMode != 0 && (GPathTracingDenoiserFunc != nullptr || GPathTracingSpatialTemporalDenoiserFunc != nullptr);
-
-	return IsDenoiserEnabled;
+	return GetPathTracingDenoiserMode(View) != 0 && (GPathTracingDenoiserPlugin || GPathTracingSpatialTemporalDenoiserPlugin);
 }
 
 static bool ShouldDenoiseWithNormalInCameraSpace()
@@ -342,12 +315,12 @@ static bool ShouldApplySpatialDenoiser()
 TArray<ESpatialDenoiserType> GetAvailableSpatialDenoiserTypes()
 {
 	TArray<ESpatialDenoiserType> Types;
-	if (GPathTracingSpatialTemporalDenoiserFunc)
+	if (GPathTracingSpatialTemporalDenoiserPlugin)
 	{
 		Types.Add(ESpatialDenoiserType::SPATIAL_TEMPORAL_DENOISER_PLUGIN);
 	}
 
-	if (GPathTracingDenoiserFunc)
+	if (GPathTracingDenoiserPlugin)
 	{
 		Types.Add(ESpatialDenoiserType::SPATIAL_DENOISER_PLUGIN);
 	}
@@ -369,11 +342,11 @@ ESpatialDenoiserType GetSpatialDenosierType()
 
 	ESpatialDenoiserType DenoiserType = static_cast<ESpatialDenoiserType>(Type);
 
-	if (DenoiserType == ESpatialDenoiserType::SPATIAL_DENOISER_PLUGIN && GPathTracingDenoiserFunc)
+	if (DenoiserType == ESpatialDenoiserType::SPATIAL_DENOISER_PLUGIN && GPathTracingDenoiserPlugin)
 	{
 		return DenoiserType;
 	}
-	else if (DenoiserType == ESpatialDenoiserType::SPATIAL_TEMPORAL_DENOISER_PLUGIN && GPathTracingSpatialTemporalDenoiserFunc)
+	else if (DenoiserType == ESpatialDenoiserType::SPATIAL_TEMPORAL_DENOISER_PLUGIN && GPathTracingSpatialTemporalDenoiserPlugin)
 	{
 		return DenoiserType;
 	}
@@ -431,9 +404,13 @@ ETextureCreateFlags GetExtraTextureCreateFlagsForDenoiser()
 	ETemporalDenoiserType TemporalDenoiserType = GetTemporalDenoiserType();
 	ETextureCreateFlags TextureCreateFlags = ETextureCreateFlags::None;
 
+	bool bThirdPartyDenoiserNeedsTextureCreateExtraFlags = false;
+	bThirdPartyDenoiserNeedsTextureCreateExtraFlags |= GPathTracingDenoiserPlugin && GPathTracingDenoiserPlugin->NeedTextureCreateExtraFlags();
+	bThirdPartyDenoiserNeedsTextureCreateExtraFlags |= GPathTracingSpatialTemporalDenoiserPlugin && GPathTracingSpatialTemporalDenoiserPlugin->NeedTextureCreateExtraFlags();
+
 	if (SpatialDenoiserType == ESpatialDenoiserType::SPATIAL_TEMPORAL_DENOISER_PLUGIN ||
 		TemporalDenoiserType == ETemporalDenoiserType::SPATIAL_TEMPORAL_DENOISER_PLUGIN ||
-		GPathTracingSpatialTemporalDenoiserFunc)
+		bThirdPartyDenoiserNeedsTextureCreateExtraFlags)
 	{
 		TextureCreateFlags |= TexCreate_Shared | TexCreate_RenderTargetable;
 	}
@@ -925,6 +902,8 @@ static void PathTracingDenoiserPlugin(FRDGBuilder& GraphBuilder,
 	FRDGTextureRef NormalTexture,
 	FRDGTextureRef OutputTexture)
 {
+	check(GPathTracingDenoiserPlugin);
+
 	FRDGTextureRef ProcessedNormalTexture = NormalTexture;
 
 	if (ShouldDenoiseWithNormalInCameraSpace())
@@ -943,23 +922,7 @@ static void PathTracingDenoiserPlugin(FRDGBuilder& GraphBuilder,
 		ConvertNormalSpace(GraphBuilder, View, ProcessedNormalTexture);
 	}
 
-	FDenoiseTextureParameters* DenoiseParameters = GraphBuilder.AllocParameters<FDenoiseTextureParameters>();
-	DenoiseParameters->InputTexture = InputTexture;
-	DenoiseParameters->InputAlbedo = AlbedoTexture;
-	DenoiseParameters->InputNormal = ProcessedNormalTexture;
-	DenoiseParameters->OutputTexture = OutputTexture;
-	// Need to read GPU mask outside Pass function, as the value is not refreshed inside the pass
-	GraphBuilder.AddPass(RDG_EVENT_NAME("Path Tracer Denoiser Plugin"), DenoiseParameters, ERDGPassFlags::Readback,
-		[DenoiseParameters, DenoiserMode, GPUMask = View.GPUMask](FRHICommandListImmediate& RHICmdList)
-	{
-		GPathTracingDenoiserFunc(RHICmdList,
-			DenoiseParameters->InputTexture->GetRHI()->GetTexture2D(),
-			DenoiseParameters->InputAlbedo->GetRHI()->GetTexture2D(),
-			DenoiseParameters->InputNormal->GetRHI()->GetTexture2D(),
-			DenoiseParameters->OutputTexture->GetRHI()->GetTexture2D(),
-			GPUMask);
-	}
-	);
+	GPathTracingDenoiserPlugin->AddPasses(GraphBuilder, View, {InputTexture, AlbedoTexture, NormalTexture, OutputTexture});
 }
 
 static void PathTracingSpatialTemporalDenoiserPlugin(FRDGBuilder& GraphBuilder,
@@ -972,8 +935,10 @@ static void PathTracingSpatialTemporalDenoiserPlugin(FRDGBuilder& GraphBuilder,
 	FRDGTextureRef PreviousOutputFrameTexture,
 	FRDGTextureRef OutputTexture,
 	int DenoisingFrameId,
-	bool bForceSpatialDenoiserOnly)
+	bool bForceSpatialDenoiserOnly,
+	FPathTracingSpatialTemporalDenoisingContext& Context)
 {
+	check(GPathTracingSpatialTemporalDenoiserPlugin);
 
 	FRDGTextureRef ProcessedNormalTexture = NormalTexture;
 
@@ -993,36 +958,29 @@ static void PathTracingSpatialTemporalDenoiserPlugin(FRDGBuilder& GraphBuilder,
 		ConvertNormalSpace(GraphBuilder, View, ProcessedNormalTexture);
 	}
 
-	FDenoiseTextureExtParameters* DenoiseParameters = GraphBuilder.AllocParameters<FDenoiseTextureExtParameters>();
-	DenoiseParameters->InputTexture = InputTexture;
-	DenoiseParameters->InputAlbedo = AlbedoTexture;
-	DenoiseParameters->InputNormal = ProcessedNormalTexture;
-	DenoiseParameters->InputFlow = FlowTexture;
-	DenoiseParameters->InputPreviousOutput = PreviousOutputFrameTexture;
-	DenoiseParameters->OutputTexture = OutputTexture;
-
-	// Need to read GPU mask outside Pass function, as the value is not refreshed inside the pass
-	GraphBuilder.AddPass(RDG_EVENT_NAME("Path Tracer Denoiser Ext Plugin"), DenoiseParameters, ERDGPassFlags::Readback,
-		[DenoiseParameters, DenoiserMode, DenoisingFrameId, bForceSpatialDenoiserOnly, GPUMask = View.GPUMask](FRHICommandListImmediate& RHICmdList)
-	{
-		check(GPathTracingSpatialTemporalDenoiserFunc);
-
-		FDenoisingArgumentsExt DenoisingArgumentsExt;
-		DenoisingArgumentsExt.FlowTex = DenoiseParameters->InputFlow ?
-			DenoiseParameters->InputFlow->GetRHI()->GetTexture2D() : nullptr;
-		DenoisingArgumentsExt.PreviousOutputTex = DenoiseParameters->InputPreviousOutput ?
-			DenoiseParameters->InputPreviousOutput->GetRHI()->GetTexture2D() : nullptr;
-		DenoisingArgumentsExt.DenoisingFrameId = DenoisingFrameId;
-		DenoisingArgumentsExt.bForceSpatialDenoiserOnly = bForceSpatialDenoiserOnly;
-
-		GPathTracingSpatialTemporalDenoiserFunc(RHICmdList,
-			DenoiseParameters->InputTexture->GetRHI()->GetTexture2D(),
-			DenoiseParameters->InputAlbedo->GetRHI()->GetTexture2D(),
-			DenoiseParameters->InputNormal->GetRHI()->GetTexture2D(),
-			DenoiseParameters->OutputTexture->GetRHI()->GetTexture2D(),
-			&DenoisingArgumentsExt,
-			GPUMask);
-	});
+	using UE::Renderer::Private::IPathTracingSpatialTemporalDenoiser;
+
+	IPathTracingSpatialTemporalDenoiser::FInputs Inputs;
+	Inputs.ColorTex = InputTexture;
+	Inputs.AlbedoTex = AlbedoTexture;
+	Inputs.NormalTex = NormalTexture;
+	Inputs.OutputTex = OutputTexture;
+	Inputs.FlowTex = FlowTexture;
+	Inputs.PreviousOutputTex = PreviousOutputFrameTexture;
+	Inputs.DenoisingFrameId = DenoisingFrameId;
+	Inputs.bForceSpatialDenoiserOnly = bForceSpatialDenoiserOnly;
+
+	if (Context.SpatialTemporalDenoiserHistory &&
+		Context.SpatialTemporalDenoiserHistory->GetDebugName() == GPathTracingSpatialTemporalDenoiserPlugin->GetDebugName())
+	{
+		Inputs.PrevHistory = Context.SpatialTemporalDenoiserHistory;
+	}
+
+	IPathTracingSpatialTemporalDenoiser::FOutputs Outputs = GPathTracingSpatialTemporalDenoiserPlugin->AddPasses(GraphBuilder, View, Inputs);
+	if (Outputs.NewHistory)
+	{
+		Context.SpatialTemporalDenoiserHistory = Outputs.NewHistory;
+	}
 }
 
 static bool ShouldApplyPreExposureToMotionVectorEstimation()
@@ -1042,29 +1000,12 @@ static void PathTracingMotionVectorPlugin(FRDGBuilder& GraphBuilder,
 	FRDGTextureRef ReferenceFrameTexture,
 	FRDGTextureRef OutputTexture)
 {
-	check(GPathTracingMotionVectorFunc);
+	check(GPathTracingSpatialTemporalDenoiserPlugin);
 
-	FMotionVectorParameters* DenoiseParameters = GraphBuilder.AllocParameters<FMotionVectorParameters>();
-	DenoiseParameters->InputFrameTexture = InputFrameTexture;
-	DenoiseParameters->ReferenceFrameTexture = ReferenceFrameTexture;
-	DenoiseParameters->OutputTexture = OutputTexture;
-	
 	bool bShouldApplyPreExposure = ShouldApplyPreExposureToMotionVectorEstimation();
 	float PreExposure = bShouldApplyPreExposure ? View.PreExposure : 1.0f;
 
-
-	GraphBuilder.AddPass(RDG_EVENT_NAME("Path Tracer Motion Vector Plugin"), DenoiseParameters, ERDGPassFlags::Readback,
-		[DenoiseParameters, PreExposure, GPUMask = View.GPUMask](FRHICommandListImmediate& RHICmdList)
-	{
-
-		GPathTracingMotionVectorFunc(RHICmdList,
-			DenoiseParameters->InputFrameTexture->GetRHI()->GetTexture2D(),
-			DenoiseParameters->ReferenceFrameTexture->GetRHI()->GetTexture2D(),
-			DenoiseParameters->OutputTexture->GetRHI()->GetTexture2D(),
-			PreExposure,
-			GPUMask);
-	});
-
+	GPathTracingSpatialTemporalDenoiserPlugin->AddMotionVectorPass(GraphBuilder, View, {InputFrameTexture, ReferenceFrameTexture, OutputTexture, PreExposure});
 }
 
 class FMotionVectorEstimationContext
@@ -1953,7 +1894,8 @@ void PathTracingSpatialTemporalDenoising(FRDGBuilder& GraphBuilder,
 			SpatialTemporalDenoisingContext.AlbedoTexture, SpatialTemporalDenoisingContext.NormalTexture,
 			MotionTexture, SourceTexture, TemporalDenoisedTexture,
 			SpatialTemporalDenoisingContext.FrameIndex,
-			bIsInitialFrame);	// zero frame will denoise without temporal
+			bIsInitialFrame,
+			SpatialTemporalDenoisingContext);	// zero frame will denoise without temporal
 
 	}
 
diff --git a/Engine/Source/Runtime/Renderer/Private/PathTracingSpatialTemporalDenoising.h b/Engine/Source/Runtime/Renderer/Private/PathTracingSpatialTemporalDenoising.h
--- a/Engine/Source/Runtime/Renderer/Private/PathTracingSpatialTemporalDenoising.h
+++ b/Engine/Source/Runtime/Renderer/Private/PathTracingSpatialTemporalDenoising.h
@@ -2,6 +2,7 @@
 
 #pragma once
 
+#include "PathTracingDenoiser.h"
 #include "RHIDefinitions.h"
 
 #if RHI_RAYTRACING
@@ -27,6 +28,9 @@ struct FPathTracingSpatialTemporalDenoisingContext
 	FRDGTexture*	LastAlbedoTexture = nullptr;
 	FRDGBuffer*		LastVarianceBuffer = nullptr;
 
+	// Custom path tracing spacial temporal denoiser result, used by plugins
+	TRefCountPtr<UE::Renderer::Private::IPathTracingSpatialTemporalDenoiser::IHistory> SpatialTemporalDenoiserHistory;
+
 	int FrameIndex;
 
 	static constexpr uint32 kNumberOfPixelShifts = 25;
diff --git a/Engine/Source/Runtime/Renderer/Public/PathTracingDenoiser.h b/Engine/Source/Runtime/Renderer/Public/PathTracingDenoiser.h
--- a/Engine/Source/Runtime/Renderer/Public/PathTracingDenoiser.h
+++ b/Engine/Source/Runtime/Renderer/Public/PathTracingDenoiser.h
@@ -3,29 +3,106 @@
 #pragma once
 
 #include "MultiGPU.h"
+#include "Templates/RefCounting.h"
+#include "Templates/UniquePtr.h"
 
-class FRHITexture;
+class FSceneView;
+class FRDGBuilder;
+class FRDGTexture;
 class FRHICommandListImmediate;
 
-// Spatial denoiser only Plugin
-using PathTracingDenoiserFunction = void(FRHICommandListImmediate& RHICmdList, FRHITexture* ColorTex, FRHITexture* AlbedoTex, FRHITexture* NormalTex, FRHITexture* OutputTex, FRHIGPUMask GPUMask);
+using FRDGTextureRef = FRDGTexture*;
 
-extern RENDERER_API PathTracingDenoiserFunction* GPathTracingDenoiserFunc;
+namespace UE::Renderer::Private
+{
+
+/** Interface for implementing third party path tracing spatial denoiser. */
+class IPathTracingDenoiser
+{
+public:
+
+	/** Inputs of the path tracing denoiser. */
+	struct FInputs
+	{
+		FRDGTextureRef ColorTex;
+		FRDGTextureRef AlbedoTex;
+		FRDGTextureRef NormalTex;
+		FRDGTextureRef OutputTex;
+	};
+
+	virtual ~IPathTracingDenoiser() {}
 
-struct FDenoisingArgumentsExt
+	virtual bool NeedTextureCreateExtraFlags() const { return false; }
+
+	/** Adds the necessary passes into RDG for denoising. */
+	virtual void AddPasses(FRDGBuilder& GraphBuilder, const FSceneView& View, const FInputs& Inputs) const = 0;
+};
+
+/** Interface for implementing third party path tracing spatial temporal denoiser. */
+class IPathTracingSpatialTemporalDenoiser
 {
-	FRHITexture* FlowTex;
-	FRHITexture* PreviousOutputTex;
+public:
+
+	/** Ref counted history to be saved in the history. */
+	class IHistory : public FRefCountBase //IRefCountedObject
+	{
+	public:
+		virtual ~IHistory() {}
+
+		/** Debug name of the history. Must exactly point to the same const TCHAR* as IPathTracingSpatialTemporalDenoiser::GetDebugName().
+		 * This is used to ensure IHistory is fed to a compatible IPathTracingSpatialTemporalDenoiser.
+		 */
+		virtual const TCHAR* GetDebugName() const = 0;
+	};
+
+	/** Inputs of the path tracing denoiser. */
+	struct FInputs
+	{
+		FRDGTextureRef ColorTex;
+		FRDGTextureRef AlbedoTex;
+		FRDGTextureRef NormalTex;
+		FRDGTextureRef OutputTex;
+
+		FRDGTextureRef FlowTex;
+		FRDGTextureRef PreviousOutputTex;
+
+		int DenoisingFrameId = -1;
+		bool bForceSpatialDenoiserOnly = false;
+
+		/** The history of the previous frame set by FOutputs::NewHistory. PrevHistory->GetDebugName() is guarentee to match the IPathTracingSpatialTemporalDenoiser. */
+		TRefCountPtr<IHistory> PrevHistory;
+	};
+
+	struct FMotionVectorInputs
+	{
+		FRDGTextureRef InputFrameTex;
+		FRDGTextureRef ReferenceFrameTex;
+		FRDGTextureRef OutputTex;
+		float PreExposure = 1.0f;
+	};
+
+	/** Outputs of the third party path tracing spacial temporal denoiser. */
+	struct FOutputs
+	{
+		/** New history to be kept alive for next frame. NewHistory->GetDebugName() must exactly point to the same const TCHAR* as IPathTracingSpatialTemporalDenoiser::GetDebugName(). */
+		TRefCountPtr<IHistory> NewHistory;
+	};
+
+	virtual ~IPathTracingSpatialTemporalDenoiser() {}
+
+	/** Debug name of the history. Must exactly point to the same const TCHAR* as IPathTracingSpatialTemporalDenoiser::IHistory::GetDebugName(). */
+	virtual const TCHAR* GetDebugName() const = 0;
+
+	virtual bool NeedTextureCreateExtraFlags() const { return false; }
+
+	/** Adds the necessary passes into RDG for denoising. */
+	virtual FOutputs AddPasses(FRDGBuilder& GraphBuilder, const FSceneView& View, const FInputs& Inputs) const = 0;
 
-	int Width;
-	int Height;
-	int DenoisingFrameId;
-	bool bForceSpatialDenoiserOnly;
+	/**  */
+	virtual void AddMotionVectorPass(FRDGBuilder& GraphBuilder, const FSceneView& View, const FMotionVectorInputs& Inputs) const = 0;
 };
 
-// Spatial-temporal denoiser plugin
-using PathTracingSpatialTemporalDenoiserFunction = void(FRHICommandListImmediate& RHICmdList, FRHITexture* ColorTex, FRHITexture* AlbedoTex, FRHITexture* NormalTex, FRHITexture* OutputTex, const FDenoisingArgumentsExt* DenoisingArgumentExt, FRHIGPUMask GPUMask);
-using PathTracingMotionVectorFunction = void(FRHICommandListImmediate& RHICmdList, FRHITexture* InputFrameTex, FRHITexture* ReferenceFrameTex, FRHITexture* FlowTex, float PreExposure, FRHIGPUMask GPUMask);
+} // namespace UE::Renderer::Private
 
-extern RENDERER_API PathTracingSpatialTemporalDenoiserFunction* GPathTracingSpatialTemporalDenoiserFunc;
-extern RENDERER_API PathTracingMotionVectorFunction* GPathTracingMotionVectorFunc;
\ No newline at end of file
+extern RENDERER_API TUniquePtr<UE::Renderer::Private::IPathTracingDenoiser> GPathTracingDenoiserPlugin;
+extern RENDERER_API TUniquePtr<UE::Renderer::Private::IPathTracingSpatialTemporalDenoiser> GPathTracingSpatialTemporalDenoiserPlugin;
\ No newline at end of file
