ShouldPrepareDistanceFieldScene should return true when VisualizeMeshDistanceFields or VisualizeGlobalDistanceField is active

diff --git a/Engine/Source/Runtime/Renderer/Private/DistanceFieldAmbientOcclusion.cpp b/Engine/Source/Runtime/Renderer/Private/DistanceFieldAmbientOcclusion.cpp
--- a/Engine/Source/Runtime/Renderer/Private/DistanceFieldAmbientOcclusion.cpp
+++ b/Engine/Source/Runtime/Renderer/Private/DistanceFieldAmbientOcclusion.cpp
@@ -717,8 +717,6 @@ bool FSceneRenderer::ShouldPrepareForDistanceFieldAO() const
 
 	return SupportsDistanceFieldAO(Scene->GetFeatureLevel(), Scene->GetShaderPlatform())
 		&& ((ShouldRenderDeferredDynamicSkyLight(Scene, ViewFamily) && bAnyViewHasGIMethodSupportingDFAO && Scene->SkyLight->bCastShadows && ViewFamily.EngineShowFlags.DistanceFieldAO)
-			|| ViewFamily.EngineShowFlags.VisualizeMeshDistanceFields
-			|| ViewFamily.EngineShowFlags.VisualizeGlobalDistanceField
 			|| ViewFamily.EngineShowFlags.VisualizeDistanceFieldAO
 			|| (GDistanceFieldAOApplyToStaticIndirect && bAnyViewHasGIMethodSupportingDFAO && ViewFamily.EngineShowFlags.DistanceFieldAO));
 }
@@ -751,7 +749,12 @@ bool FSceneRenderer::ShouldPrepareDistanceFieldScene() const
 	bool bShouldPrepareForDFInsetIndirectShadow = ShouldPrepareForDFInsetIndirectShadow();
 
 	// Prepare the distance field scene (object buffers and distance field atlas) if any feature needs it
-	return bShouldPrepareGlobalDistanceField || bShouldPrepareForAO || ShouldPrepareForDistanceFieldShadows() || bShouldPrepareForDFInsetIndirectShadow;
+	return bShouldPrepareGlobalDistanceField 
+		|| bShouldPrepareForAO 
+		|| ShouldPrepareForDistanceFieldShadows() 
+		|| bShouldPrepareForDFInsetIndirectShadow 
+		|| ViewFamily.EngineShowFlags.VisualizeGlobalDistanceField
+		|| ViewFamily.EngineShowFlags.VisualizeMeshDistanceFields;
 }
 
 bool FSceneRenderer::ShouldPrepareGlobalDistanceField() const
