PSO precaching fixes:
- Always precache CSM shadow PSOs for default material
- Precache tranlucent pass PSOs also with default render target formats for underwater rendering
- Precache sky capture PSOs with different render target info

diff --git a/Engine/Source/Runtime/Renderer/Private/BasePassRendering.cpp b/Engine/Source/Runtime/Renderer/Private/BasePassRendering.cpp
--- a/Engine/Source/Runtime/Renderer/Private/BasePassRendering.cpp
+++ b/Engine/Source/Runtime/Renderer/Private/BasePassRendering.cpp
@@ -1644,17 +1644,10 @@ void FBasePassMeshProcessor::CollectPSOInitializersForLMPolicy(
 		const bool bDepthTest = TranslucencyPassType != ETranslucencyPass::TPT_TranslucencyAfterMotionBlur;
 		const bool bRenderInSeparateTranslucency = IsSeparateTranslucencyEnabled(TranslucencyPassType, SeparateTranslucencyDimensions.Scale);
 
-		if (bRenderInSeparateTranslucency)
-		{
-			const FRDGTextureDesc TextureDesc = GetPostDOFTranslucentTextureDesc(TranslucencyPassType, SeparateTranslucencyDimensions, bIsModulate, ShaderPlatform);
-			AddRenderTargetInfo(TextureDesc.Format, TextureDesc.Flags, RenderTargetsInfo);
-		}
-		else
-		{
-			EPixelFormat SceneColorFormat = SceneTexturesConfig.ColorFormat;
-			ETextureCreateFlags SceneColorCreateFlags = SceneTexturesConfig.ColorCreateFlags;
-			AddRenderTargetInfo(SceneColorFormat, SceneColorCreateFlags, RenderTargetsInfo);
-		}
+		// Always create PSO without separate translucency (could be used when under water for all translucent passes)
+		EPixelFormat SceneColorFormat = SceneTexturesConfig.ColorFormat;
+		ETextureCreateFlags SceneColorCreateFlags = SceneTexturesConfig.ColorCreateFlags;
+		AddRenderTargetInfo(SceneColorFormat, SceneColorCreateFlags, RenderTargetsInfo);
 
 		if (bDepthTest)
 		{
@@ -1662,25 +1655,59 @@ void FBasePassMeshProcessor::CollectPSOInitializersForLMPolicy(
 			SetupDepthStencilInfo(PF_DepthStencil, DepthStencilCreateFlags, ERenderTargetLoadAction::ELoad,
 				ERenderTargetLoadAction::ELoad, FExclusiveDepthStencil::DepthRead_StencilWrite, RenderTargetsInfo);
 		}
+
+		AddGraphicsPipelineStateInitializer(
+			VertexFactoryData,
+			MaterialResource,
+			DrawRenderState,
+			RenderTargetsInfo,
+			BasePassShaders,
+			MeshFillMode,
+			MeshCullMode,
+			PrimitiveType,
+			EMeshPassFeatures::Default,
+			true /*bRequired*/,
+			PSOInitializers);
+
+		// Add another PSO when render in separate translucency because render target format could have changed
+		if (bRenderInSeparateTranslucency)
+		{
+			const FRDGTextureDesc TextureDesc = GetPostDOFTranslucentTextureDesc(TranslucencyPassType, SeparateTranslucencyDimensions, bIsModulate, ShaderPlatform);
+			RenderTargetsInfo.RenderTargetFormats[0] = TextureDesc.Format;
+			RenderTargetsInfo.RenderTargetFlags[0] = TextureDesc.Flags;
+
+			AddGraphicsPipelineStateInitializer(
+				VertexFactoryData,
+				MaterialResource,
+				DrawRenderState,
+				RenderTargetsInfo,
+				BasePassShaders,
+				MeshFillMode,
+				MeshCullMode,
+				PrimitiveType,
+				EMeshPassFeatures::Default,
+				true /*bRequired*/,
+				PSOInitializers);
+		}
 	}
 	else
 	{
 		// Regular base pass with gbuffer bindings
 		SetupGBufferRenderTargetInfo(SceneTexturesConfig, RenderTargetsInfo, true /*bSetupDepthStencil*/);
-	}
-		
-	AddGraphicsPipelineStateInitializer(
-		VertexFactoryData,
-		MaterialResource,
-		DrawRenderState,
-		RenderTargetsInfo,
-		BasePassShaders,
-		MeshFillMode,
-		MeshCullMode,
-		PrimitiveType,
-		EMeshPassFeatures::Default,
-		true /*bRequired*/,
-		PSOInitializers);
+
+		AddGraphicsPipelineStateInitializer(
+			VertexFactoryData,
+			MaterialResource,
+			DrawRenderState,
+			RenderTargetsInfo,
+			BasePassShaders,
+			MeshFillMode,
+			MeshCullMode,
+			PrimitiveType,
+			EMeshPassFeatures::Default,
+			true /*bRequired*/,
+			PSOInitializers);
+	}	
 }
 
 template<typename LightMapPolicyType>
diff --git a/Engine/Source/Runtime/Renderer/Private/ReflectionEnvironmentRealTimeCapture.cpp b/Engine/Source/Runtime/Renderer/Private/ReflectionEnvironmentRealTimeCapture.cpp
--- a/Engine/Source/Runtime/Renderer/Private/ReflectionEnvironmentRealTimeCapture.cpp
+++ b/Engine/Source/Runtime/Renderer/Private/ReflectionEnvironmentRealTimeCapture.cpp
@@ -359,9 +359,7 @@ void FScene::AllocateAndCaptureFrameSkyEnvMap(
 	const FMatrix CubeProjectionMatrix = GetCubeProjectionMatrix(CubeView.FOV * 0.5f, (float)CubeWidth, GReflectionCaptureNearPlane);
 	CubeView.UpdateProjectionMatrix(CubeProjectionMatrix);
 
-	FPooledRenderTargetDesc SkyCubeTexDesc = FPooledRenderTargetDesc::CreateCubemapDesc(CubeWidth, 
-		PF_FloatR11G11B10, FClearValueBinding::Black, TexCreate_TargetArraySlicesIndependently,
-		TexCreate_ShaderResource | TexCreate_UAV | TexCreate_RenderTargetable, false, 1, CubeMipCount, false);
+	FPooledRenderTargetDesc SkyCubeTexDesc = Translate(FSkyPassMeshProcessor::GetCaptureFrameSkyEnvMapTextureDesc(CubeWidth, CubeMipCount));
 
 	const bool bTimeSlicedRealTimeCapture = CVarRealTimeReflectionCaptureTimeSlicing.GetValueOnRenderThread() > 0 && !MainView.Family->bCurrentlyBeingEdited;
 
diff --git a/Engine/Source/Runtime/Renderer/Private/ShadowDepthRendering.cpp b/Engine/Source/Runtime/Renderer/Private/ShadowDepthRendering.cpp
--- a/Engine/Source/Runtime/Renderer/Private/ShadowDepthRendering.cpp
+++ b/Engine/Source/Runtime/Renderer/Private/ShadowDepthRendering.cpp
@@ -2194,10 +2194,11 @@ void FShadowDepthPassMeshProcessor::AddMeshBatch(const FMeshBatch& RESTRICT Mesh
 void FShadowDepthPassMeshProcessor::CollectPSOInitializers(const FSceneTexturesConfig& SceneTexturesConfig, const FMaterial& Material, const FPSOPrecacheVertexFactoryData& VertexFactoryData, const FPSOPrecacheParams& PreCacheParams, TArray<FPSOPrecacheData>& PSOInitializers)
 {
 	// Early out if possible
-	if (!Material.ShouldCastDynamicShadows() ||
+	if (!PreCacheParams.bDefaultMaterial && 
+		(!Material.ShouldCastDynamicShadows() ||
 		!ShouldIncludeDomainInMeshPass(Material.GetMaterialDomain()) ||
 		!ShouldIncludeMaterialInDefaultOpaquePass(Material) ||
-		!EnumHasAnyFlags(MeshSelectionMask, VertexFactoryData.VertexFactoryType->SupportsPrimitiveIdStream() ? EShadowMeshSelection::VSM : EShadowMeshSelection::SM))
+		!EnumHasAnyFlags(MeshSelectionMask, VertexFactoryData.VertexFactoryType->SupportsPrimitiveIdStream() ? EShadowMeshSelection::VSM : EShadowMeshSelection::SM)))
 	{
 		return;
 	}
diff --git a/Engine/Source/Runtime/Renderer/Private/SkyPassRendering.cpp b/Engine/Source/Runtime/Renderer/Private/SkyPassRendering.cpp
--- a/Engine/Source/Runtime/Renderer/Private/SkyPassRendering.cpp
+++ b/Engine/Source/Runtime/Renderer/Private/SkyPassRendering.cpp
@@ -191,7 +191,7 @@ void FSkyPassMeshProcessor::CollectPSOInitializers(const FSceneTexturesConfig& S
 		FeatureLevel,
 		bRenderSkylight,
 		false,
-		GBL_Default,
+		GBL_Default, 
 		&SkyPassShaders.VertexShader,
 		&SkyPassShaders.PixelShader
 		))
@@ -201,7 +201,7 @@ void FSkyPassMeshProcessor::CollectPSOInitializers(const FSceneTexturesConfig& S
 	
 	FGraphicsPipelineRenderTargetsInfo RenderTargetsInfo;
 	SetupGBufferRenderTargetInfo(SceneTexturesConfig, RenderTargetsInfo, true /*bSetupDepthStencil*/);
-
+	
 	AddGraphicsPipelineStateInitializer(
 		VertexFactoryData,
 		Material,
@@ -214,6 +214,41 @@ void FSkyPassMeshProcessor::CollectPSOInitializers(const FSceneTexturesConfig& S
 		EMeshPassFeatures::Default,
 		true /*bRequired*/,
 		PSOInitializers);
+
+	// Also generate with depth write which is used during CaptureSkyMeshReflection
+	const FExclusiveDepthStencil::Type SceneBasePassDepthStencilAccess = FScene::GetDefaultBasePassDepthStencilAccess(FeatureLevel);
+	FMeshPassProcessorRenderState SkyCaptureDrawRenderState;
+	FExclusiveDepthStencil::Type BasePassDepthStencilAccess_Sky = FExclusiveDepthStencil::Type(SceneBasePassDepthStencilAccess | FExclusiveDepthStencil::DepthWrite);
+	SetupBasePassState(BasePassDepthStencilAccess_Sky, false, SkyCaptureDrawRenderState);
+	
+	// Also change render target format
+	FRDGTextureDesc SkyCaptureRenderTargetDesc = FSkyPassMeshProcessor::GetCaptureFrameSkyEnvMapTextureDesc(1, 1);
+
+	FGraphicsPipelineRenderTargetsInfo SkyCaptureRenderTargetsInfo;
+	SkyCaptureRenderTargetsInfo.NumSamples = 1;
+	AddRenderTargetInfo(SkyCaptureRenderTargetDesc.Format, SkyCaptureRenderTargetDesc.Flags, SkyCaptureRenderTargetsInfo);
+	SetupDepthStencilInfo(PF_DepthStencil, SceneTexturesConfig.DepthCreateFlags, ERenderTargetLoadAction::ELoad,
+		ERenderTargetLoadAction::ELoad, FExclusiveDepthStencil::DepthWrite_StencilWrite, SkyCaptureRenderTargetsInfo);
+
+	AddGraphicsPipelineStateInitializer(
+		VertexFactoryData,
+		Material,
+		SkyCaptureDrawRenderState,
+		SkyCaptureRenderTargetsInfo,
+		SkyPassShaders,
+		MeshFillMode,
+		MeshCullMode,
+		(EPrimitiveType)PreCacheParams.PrimitiveType,
+		EMeshPassFeatures::Default,
+		true /*bRequired*/,
+		PSOInitializers);
+}
+
+FRDGTextureDesc FSkyPassMeshProcessor::GetCaptureFrameSkyEnvMapTextureDesc(uint32 CubeWidth, uint32 CubeMipCount)
+{
+	return FRDGTextureDesc::CreateCube(CubeWidth,
+		PF_FloatR11G11B10, FClearValueBinding::Black, TexCreate_TargetArraySlicesIndependently |
+		TexCreate_ShaderResource | TexCreate_UAV | TexCreate_RenderTargetable, CubeMipCount);
 }
 
 FMeshPassProcessor* CreateSkyPassProcessor(ERHIFeatureLevel::Type FeatureLevel, const FScene* Scene, const FSceneView* InViewIfDynamicMeshCommand, FMeshPassDrawListContext* InDrawListContext)
diff --git a/Engine/Source/Runtime/Renderer/Private/SkyPassRendering.h b/Engine/Source/Runtime/Renderer/Private/SkyPassRendering.h
--- a/Engine/Source/Runtime/Renderer/Private/SkyPassRendering.h
+++ b/Engine/Source/Runtime/Renderer/Private/SkyPassRendering.h
@@ -21,6 +21,8 @@ public:
 	virtual void AddMeshBatch(const FMeshBatch& RESTRICT MeshBatch, uint64 BatchElementMask, const FPrimitiveSceneProxy* RESTRICT PrimitiveSceneProxy, int32 StaticMeshId = -1) override final;
 	virtual void CollectPSOInitializers(const FSceneTexturesConfig& SceneTexturesConfig, const FMaterial& Material, const FPSOPrecacheVertexFactoryData& VertexFactoryData, const FPSOPrecacheParams& PreCacheParams, TArray<FPSOPrecacheData>& PSOInitializers) override final;
 
+	static FRDGTextureDesc GetCaptureFrameSkyEnvMapTextureDesc(uint32 CubeWidth, uint32 CubeMipCount);
+
 	FMeshPassProcessorRenderState PassDrawRenderState;
 	EShadingPath ShadingPath;
 
