Stochastic Shadows
- Remove Hash Table code path

diff --git a/Engine/Shaders/Private/StochasticShadows/StochasticShadows.ush b/Engine/Shaders/Private/StochasticShadows/StochasticShadows.ush
--- a/Engine/Shaders/Private/StochasticShadows/StochasticShadows.ush
+++ b/Engine/Shaders/Private/StochasticShadows/StochasticShadows.ush
@@ -290,39 +290,9 @@ uint GetShadowMaskPageTableKey(uint LightSceneId, uint2 PageTableCoord)
 	return HashedKey;
 }
 
-struct FShadowMaskHashTableValue
-{
-	uint ShadowMaskTileIndex;
-	uint Checksum;
-};
-
-uint GetShadowMaskHashTableValueChecksum(uint HashTableKey)
-{
-	// Use top 12 bits of a key to detect collisions
-	return (HashTableKey >> 20) & 0xFFF;
-}
-
-uint PackShadowMaskHashTableValue(uint ShadowMaskTileIndex, uint HashTableKey)
-{
-	uint PackedValue = 0;
-	PackedValue = (ShadowMaskTileIndex & 0xFFFFF) << 12;
-	PackedValue |= GetShadowMaskHashTableValueChecksum(HashTableKey);
-	return PackedValue;
-}
-
-FShadowMaskHashTableValue UnpackShadowMaskHashTableValue(uint PackedValue)
-{
-	FShadowMaskHashTableValue TableValue;
-	TableValue.ShadowMaskTileIndex = PackedValue >> 12;
-	TableValue.Checksum = PackedValue & 0xFFF;
-	return TableValue;
-}
-
 Texture2D<uint> ShadowMaskPageTableHistory;
-StructuredBuffer<uint> ShadowMaskHashTableHistory;
 Texture2D<uint> ShadowMaskAtlasHistory;
 int2 ShadowMaskViewSize;
-uint ShadowMaskHashTableIndexWrapMask;
 
 void AccumulateHistoryShadowMask(uint LightSceneId, uint PrevLocalLightIndex, uint2 ShadowMaskScreenCoord, float Weight, inout FShadowMaskAccumulator Accumulator)
 {
@@ -330,31 +300,7 @@ void AccumulateHistoryShadowMask(uint LightSceneId, uint PrevLocalLightIndex, ui
 	{
 		uint2 HistoryShadowMaskPageTableCoord = GetShadowMaskPageTableCoord(PrevLocalLightIndex, ShadowMaskScreenCoord);
 
-#if HASH_TABLE
-		uint HistoryShadowMaskTileIndex = INVALID_SHADOW_MASK_TILE_INDEX;
-		const uint PageTableKey = GetShadowMaskPageTableKey(LightSceneId, HistoryShadowMaskPageTableCoord);
-		const uint MaxLinearProbingSteps = 4;
-
-		LOOP
-		for (uint LinearProbingStep = 0; LinearProbingStep < MaxLinearProbingSteps; ++LinearProbingStep)
-		{
-			uint Index = (PageTableKey + LinearProbingStep) & ShadowMaskHashTableIndexWrapMask;
-			uint PackedValue = ShadowMaskHashTableHistory[Index];
-
-			if (PackedValue != INVALID_SHADOW_MASK_TILE_INDEX)
-			{
-				FShadowMaskHashTableValue TableValue = UnpackShadowMaskHashTableValue(PackedValue);
-				if (TableValue.Checksum == GetShadowMaskHashTableValueChecksum(PageTableKey))
-				{
-					HistoryShadowMaskTileIndex = TableValue.ShadowMaskTileIndex;
-					break;
-				}
-			}
-		}
-
-#else
 		uint HistoryShadowMaskTileIndex = ShadowMaskPageTableHistory[HistoryShadowMaskPageTableCoord];
-#endif
 
 		if (HistoryShadowMaskTileIndex != INVALID_SHADOW_MASK_TILE_INDEX)
 		{
diff --git a/Engine/Shaders/Private/StochasticShadows/StochasticShadowsSampling.usf b/Engine/Shaders/Private/StochasticShadows/StochasticShadowsSampling.usf
--- a/Engine/Shaders/Private/StochasticShadows/StochasticShadowsSampling.usf
+++ b/Engine/Shaders/Private/StochasticShadows/StochasticShadowsSampling.usf
@@ -94,7 +94,6 @@ RWTexture2D<float4> RWShadowMaskHistoryWeights;
 RWTexture2D<float> RWDownsampledSceneDepth;
 RWTexture2D<UNORM float3> RWDownsampledSceneWorldNormal;
 RWTexture2D<uint> RWShadowMaskPageTable;
-RWStructuredBuffer<uint> RWShadowMaskHashTable;
 RWStructuredBuffer<uint> RWShadowMaskTileAllocator;
 RWStructuredBuffer<uint2> RWShadowMaskTileData;
 RWTexture2D<uint> RWLightSamples;
@@ -103,12 +102,7 @@ float4 HistoryScreenPositionScaleBias;
 float4 HistoryUVMinMax;
 Texture2D<float> ShadowMaskSceneDepthHistory;
 
-#if HASH_TABLE
-	#define HASH_TABLE_SIZE 32
-	groupshared uint SharedHashTable[HASH_TABLE_SIZE];
-#else
-	groupshared uint SharedLightMask[LIGHT_SCENE_ID_MASK_SIZE];
-#endif
+groupshared uint SharedLightMask[LIGHT_SCENE_ID_MASK_SIZE];
 
 #if CANDIDATE_LIGHT_MASK
 	groupshared uint SharedCandidateLightMask[LIGHT_SCENE_ID_MASK_SIZE];
@@ -119,47 +113,10 @@ Texture2D<float> ShadowMaskSceneDepthHistory;
  */
 void MarkLightSceneId(uint LightSceneId)
 {
-#if HASH_TABLE
-	// Just skip add if search gets too long
-	const uint MaxLinearProbingSteps = 4;
-
-	LOOP
-	for (uint LinearProbingStep = 0; LinearProbingStep < MaxLinearProbingSteps; ++LinearProbingStep)
-	{
-		uint Index = (LightSceneId + LinearProbingStep) % HASH_TABLE_SIZE;
-
-		uint StoredKey = SharedHashTable[Index];
-		if (StoredKey != LightSceneId)
-		{
-			if (StoredKey != INVALID_LIGHT_SCENE_ID)
-			{
-				continue;
-			}
-
-			uint PrevKey;
-			InterlockedCompareExchange(SharedHashTable[Index], INVALID_LIGHT_SCENE_ID, LightSceneId, PrevKey);
-
-			if (PrevKey == INVALID_LIGHT_SCENE_ID)
-			{
-				SharedHashTable[Index] = LightSceneId;
-				return;
-			}
-			else if (PrevKey != LightSceneId)
-			{
-				continue;
-			}
-		}
-
-		break;
-	}
-
-#else
-
 	// Re-distribute nearby ids to different bytes in order to balance gather workload
 	uint DWORDIndex = LightSceneId % LIGHT_SCENE_ID_MASK_SIZE;
 	uint BitMask = 1u << (LightSceneId / LIGHT_SCENE_ID_MASK_SIZE);
 	InterlockedOr(SharedLightMask[DWORDIndex], BitMask);
-#endif
 }
 
 void WriteShadowMaskTileToPageTable(uint LightSceneId, uint LocalLightIndex, uint2 ShadowMaskScreenCoord, uint ShadowMaskTileIndex)
@@ -168,47 +125,7 @@ void WriteShadowMaskTileToPageTable(uint LightSceneId, uint LocalLightIndex, uin
 
 	RWShadowMaskTileData[ShadowMaskTileIndex] = PackShadowMaskTile(LightSceneId, LocalLightIndex, ShadowMaskScreenCoord);
 
-	#if HASH_TABLE
-	{
-		const uint PageTableKey = GetShadowMaskPageTableKey(LightSceneId, ShadowMaskPageTableCoord);
-		const uint PageTableValue = PackShadowMaskHashTableValue(ShadowMaskTileIndex, PageTableKey);
-		const uint MaxLinearProbingSteps = 4;
-
-		LOOP
-		for (uint LinearProbingStep = 0; LinearProbingStep < MaxLinearProbingSteps; ++LinearProbingStep)
-		{
-			uint Index = (PageTableKey + LinearProbingStep) & ShadowMaskHashTableIndexWrapMask;
-
-			uint StoredValue = RWShadowMaskHashTable[Index];
-			if (StoredValue != PageTableValue)
-			{
-				if (StoredValue != INVALID_SHADOW_MASK_TILE_INDEX)
-				{
-					continue;
-				}
-
-				uint PrevValue;
-				InterlockedCompareExchange(RWShadowMaskHashTable[Index], INVALID_SHADOW_MASK_TILE_INDEX, PageTableValue, PrevValue);
-
-				if (PrevValue == INVALID_SHADOW_MASK_TILE_INDEX)
-				{
-					RWShadowMaskHashTable[Index] = PageTableValue;
-					return;
-				}
-				else if (PrevValue != PageTableValue)
-				{
-					continue;
-				}
-			}
-
-			break;
-		}
-	}
-	#else
-	{
-		RWShadowMaskPageTable[ShadowMaskPageTableCoord] = ShadowMaskTileIndex;
-	}
-	#endif
+	RWShadowMaskPageTable[ShadowMaskPageTableCoord] = ShadowMaskTileIndex;
 }
 
 StructuredBuffer<uint> DownsampledTileAllocator;
@@ -227,17 +144,10 @@ void GenerateSamplesCS(
 {
 	uint LinearThreadIndex = GroupThreadId.y * THREADGROUP_SIZE + GroupThreadId.x;
 
-#if HASH_TABLE
-	if (LinearThreadIndex < HASH_TABLE_SIZE)
-	{
-		SharedHashTable[LinearThreadIndex] = INVALID_LIGHT_SCENE_ID;
-	}
-#else
 	if (LinearThreadIndex < LIGHT_SCENE_ID_MASK_SIZE)
 	{
 		SharedLightMask[LinearThreadIndex] = 0;
 	}
-#endif
 
 #if CANDIDATE_LIGHT_MASK
 	if (LinearThreadIndex < LIGHT_SCENE_ID_MASK_SIZE)
@@ -654,30 +564,6 @@ void GenerateSamplesCS(
 		GroupMemoryBarrierWithGroupSync();
 
 		// Write out allocated shadow mask tiles
-	#if HASH_TABLE
-		if (LinearThreadIndex < HASH_TABLE_SIZE)
-		{
-			uint2 ShadowMaskScreenCoord = DownsampledTileCoord * SHADOW_MASK_TILE_SIZE;
-
-			const uint StoredValue = SharedHashTable[LinearThreadIndex];
-
-			if (StoredValue != INVALID_LIGHT_SCENE_ID)
-			{
-				const uint LocalLightIndex = StoredValue;
-				const FLocalLightData LocalLightData = GetLocalLightDataNonStereo(LocalLightIndex);
-				const uint LightIndex = LocalLightData.LightSceneId;
-
-				// #sdl_todo: merge allocations through groupshared and do only one atomic per group
-				uint ShadowMaskTileIndex = 0;
-				InterlockedAdd(RWShadowMaskTileAllocator[0], 1, ShadowMaskTileIndex);
-
-				if (ShadowMaskTileIndex < MaxShadowMaskTiles)
-				{
-					WriteShadowMaskTileToPageTable(LightIndex, LocalLightIndex, ShadowMaskScreenCoord, ShadowMaskTileIndex);
-				}
-			}
-		}
-	#else
 		// #sdl_todo: process part of the list per tread
 		if (LinearThreadIndex < LIGHT_SCENE_ID_MASK_SIZE)
 		{
@@ -707,6 +593,5 @@ void GenerateSamplesCS(
 				}
 			}
 		}
-	#endif
 	}
 }
\ No newline at end of file
diff --git a/Engine/Shaders/Private/StochasticShadows/StochasticShadowsSpatial.usf b/Engine/Shaders/Private/StochasticShadows/StochasticShadowsSpatial.usf
--- a/Engine/Shaders/Private/StochasticShadows/StochasticShadowsSpatial.usf
+++ b/Engine/Shaders/Private/StochasticShadows/StochasticShadowsSpatial.usf
@@ -9,7 +9,6 @@
 
 Texture2D<uint> ShadowMaskTileAtlas;
 Texture2D<uint> ShadowMaskPageTable;
-StructuredBuffer<uint> ShadowMaskHashTable;
 StructuredBuffer<uint2> ShadowMaskTileData;
 Texture2D<float> DownsampledSceneDepth;
 Texture2D<UNORM float3> DownsampledSceneWorldNormal;
@@ -84,35 +83,8 @@ void ShadowMaskSpatialPassCS(
 			uint2 SampleShadowMaskScreenCoord = clamp(int2(BaseShadowMaskScreenCoord) + int2(SharedShadowMaskX - SHARED_TILE_BORDER, SharedShadowMaskY - SHARED_TILE_BORDER), 0, ShadowMaskViewSize - 1);
 
 			uint2 SampleShadowMaskPageTableCoord = GetShadowMaskPageTableCoord(LocalLightIndex, SampleShadowMaskScreenCoord);
-
-			uint SampleShadowMaskTileIndex = INVALID_SHADOW_MASK_TILE_INDEX;
-			#if HASH_TABLE
-			{
-				const uint PageTableKey = GetShadowMaskPageTableKey(LightSceneId, SampleShadowMaskPageTableCoord);
-				const uint MaxLinearProbingSteps = 4;
-
-				LOOP
-				for (uint LinearProbingStep = 0; LinearProbingStep < MaxLinearProbingSteps; ++LinearProbingStep)
-				{
-					uint Index = (PageTableKey + LinearProbingStep) & ShadowMaskHashTableIndexWrapMask;
-					uint PackedValue = ShadowMaskHashTable[Index];
-
-					if (PackedValue != INVALID_SHADOW_MASK_TILE_INDEX)
-					{
-						FShadowMaskHashTableValue TableValue = UnpackShadowMaskHashTableValue(PackedValue);
-						if (TableValue.Checksum == GetShadowMaskHashTableValueChecksum(PageTableKey))
-						{
-							SampleShadowMaskTileIndex = TableValue.ShadowMaskTileIndex;
-							break;
-						}
-					}
-				}
-			}
-			#else
-			{
-				SampleShadowMaskTileIndex = ShadowMaskPageTable[SampleShadowMaskPageTableCoord];
-			}
-			#endif			
+			
+			uint SampleShadowMaskTileIndex = ShadowMaskPageTable[SampleShadowMaskPageTableCoord];
 
 			FShadowMask ShadowMask = InitShadowMask();
 
diff --git a/Engine/Source/Runtime/Renderer/Private/StochasticShadows/StochasticShadows.cpp b/Engine/Source/Runtime/Renderer/Private/StochasticShadows/StochasticShadows.cpp
--- a/Engine/Source/Runtime/Renderer/Private/StochasticShadows/StochasticShadows.cpp
+++ b/Engine/Source/Runtime/Renderer/Private/StochasticShadows/StochasticShadows.cpp
@@ -107,13 +107,6 @@ static TAutoConsoleVariable<int> CVarStochasticShadowsFixedStateFrameIndex(
 	ECVF_RenderThreadSafe
 );
 
-static TAutoConsoleVariable<int> CVarStochasticShadowsHashTable(
-	TEXT("r.StochasticShadows.HashTable"),
-	0,
-	TEXT("Whether to use bit array limites to 256 lights per frustum or unlimited hash table."),
-	ECVF_Scalability | ECVF_RenderThreadSafe
-);
-
 static TAutoConsoleVariable<int> CVarStochasticShadowsCandidateLightMask(
 	TEXT("r.StochasticShadows.CandidateLightMask"),
 	1,
@@ -130,7 +123,6 @@ namespace StochasticShadows
 	constexpr int32 MaxShadingTilesPerGridCell = 16;
 	constexpr int32 ShadowMaskAtlasSizeInTiles = 512;
 	constexpr uint32 InvalidShadowMaskTileIndex = 0xFFFFFFFF;
-	constexpr uint32 ShadowMaskHashTableFactor = 4; // How much hash table should be larger than max number of elements?
 
 	bool IsEnabled()
 	{
@@ -313,13 +305,11 @@ class FGenerateSamplesCS : public FGlobalShader
 		SHADER_PARAMETER_RDG_BUFFER_UAV(RWStructuredBuffer<uint>, RWShadowMaskTileAllocator)
 		SHADER_PARAMETER_RDG_BUFFER_UAV(RWStructuredBuffer<uint2>, RWShadowMaskTileData)
 		SHADER_PARAMETER_RDG_TEXTURE_UAV(RWTexture2D<uint>, RWShadowMaskPageTable)
-		SHADER_PARAMETER_RDG_BUFFER_UAV(RWStructuredBuffer<uint>, RWShadowMaskHashTable)
 		SHADER_PARAMETER_RDG_TEXTURE_UAV(RWTexture2D<uint>, RWLightSamples)
 		SHADER_PARAMETER_RDG_UNIFORM_BUFFER(FForwardLightData, ForwardLightData)
 		SHADER_PARAMETER_RDG_BUFFER_SRV(StructuredBuffer<uint>, DownsampledTileAllocator)
 		SHADER_PARAMETER_RDG_BUFFER_SRV(StructuredBuffer<uint>, DownsampledTileData)
 		SHADER_PARAMETER_RDG_TEXTURE(Texture2D<uint>, ShadowMaskPageTableHistory)
-		SHADER_PARAMETER_RDG_BUFFER_SRV(StructuredBuffer<uint>, ShadowMaskHashTableHistory)
 		SHADER_PARAMETER_RDG_TEXTURE(Texture2D<uint>, ShadowMaskAtlasHistory)
 		SHADER_PARAMETER_RDG_TEXTURE(Texture2D<uint>, ShadowMaskTileAtlas)
 		SHADER_PARAMETER_RDG_TEXTURE(Texture2D<float>, ShadowMaskSceneDepthHistory)
@@ -331,10 +321,9 @@ class FGenerateSamplesCS : public FGlobalShader
 	class FNumSamplesPerPixel : SHADER_PERMUTATION_SPARSE_INT("NUM_SAMPLES_PER_PIXEL", 1, 2, 4);
 	class FShadowMaskReprojectionWeights : SHADER_PERMUTATION_BOOL("SHADOW_MASK_REPROJECTION_WEIGHTS");
 	class FShadowFactorEstimate : SHADER_PERMUTATION_BOOL("SHADOW_FACTOR_ESTIMATE");
-	class FHashTable : SHADER_PERMUTATION_BOOL("HASH_TABLE");
 	class FCandidateLightMask : SHADER_PERMUTATION_BOOL("CANDIDATE_LIGHT_MASK");
 	class FDebugMode : SHADER_PERMUTATION_BOOL("DEBUG_MODE");
-	using FPermutationDomain = TShaderPermutationDomain<FTileType, FNumSamplesPerPixel, FShadowMaskReprojectionWeights, FShadowFactorEstimate, FHashTable, FCandidateLightMask, FDebugMode>;
+	using FPermutationDomain = TShaderPermutationDomain<FTileType, FNumSamplesPerPixel, FShadowMaskReprojectionWeights, FShadowFactorEstimate, FCandidateLightMask, FDebugMode>;
 
 	static int32 GetGroupSize()
 	{	
@@ -424,7 +413,6 @@ class FCompositeShadowMaskTracesCS : public FGlobalShader
 		SHADER_PARAMETER_RDG_TEXTURE(Texture2D<uint32>, ShadowMaskHistoryScreenCoord00)
 		SHADER_PARAMETER_RDG_TEXTURE(Texture2D<float4>, ShadowMaskHistoryWeights)
 		SHADER_PARAMETER_RDG_TEXTURE(Texture2D<uint>, ShadowMaskPageTableHistory)
-		SHADER_PARAMETER_RDG_BUFFER_SRV(StructuredBuffer<uint>, ShadowMaskHashTableHistory)
 		SHADER_PARAMETER_RDG_TEXTURE(Texture2D<uint>, ShadowMaskAtlasHistory)
 		SHADER_PARAMETER_RDG_TEXTURE(Texture2D<uint>, LightSamples)
 		SHADER_PARAMETER_RDG_UNIFORM_BUFFER(FForwardLightData, ForwardLightData)
@@ -437,9 +425,8 @@ class FCompositeShadowMaskTracesCS : public FGlobalShader
 
 	class FNumSamplesPerPixel : SHADER_PERMUTATION_SPARSE_INT("NUM_SAMPLES_PER_PIXEL", 1, 2, 4);
 	class FTemporalAccumulation : SHADER_PERMUTATION_BOOL("TEMPORAL_ACCUMULATION");
-	class FHashTable : SHADER_PERMUTATION_BOOL("HASH_TABLE");
 	class FDebugMode : SHADER_PERMUTATION_BOOL("DEBUG_MODE");
-	using FPermutationDomain = TShaderPermutationDomain<FNumSamplesPerPixel, FTemporalAccumulation, FHashTable, FDebugMode>;
+	using FPermutationDomain = TShaderPermutationDomain<FNumSamplesPerPixel, FTemporalAccumulation, FDebugMode>;
 
 	static bool ShouldCompilePermutation(const FGlobalShaderPermutationParameters& Parameters)
 	{
@@ -476,7 +463,6 @@ class FShadowMaskSpatialPassCS : public FGlobalShader
 		SHADER_PARAMETER_RDG_BUFFER_SRV(StructuredBuffer<uint2>, ShadowMaskTileData)
 		SHADER_PARAMETER_RDG_TEXTURE(Texture2D<uint>, ShadowMaskPageTable)
 		SHADER_PARAMETER_RDG_TEXTURE(Texture2D<uint>, ShadowMaskTileAtlas)
-		SHADER_PARAMETER_RDG_BUFFER_SRV(StructuredBuffer<uint>, ShadowMaskHashTable)
 	END_SHADER_PARAMETER_STRUCT()
 
 	static int32 GetGroupSize()
@@ -484,9 +470,8 @@ class FShadowMaskSpatialPassCS : public FGlobalShader
 		return 8;
 	}
 
-	class FHashTable : SHADER_PERMUTATION_BOOL("HASH_TABLE");
 	class FSpatialPass : SHADER_PERMUTATION_BOOL("SPATIAL_PASS");
-	using FPermutationDomain = TShaderPermutationDomain<FHashTable, FSpatialPass>;
+	using FPermutationDomain = TShaderPermutationDomain<FSpatialPass>;
 
 	static bool ShouldCompilePermutation(const FGlobalShaderPermutationParameters& Parameters)
 	{
@@ -521,7 +506,6 @@ class FShadeLightSamplesCS : public FGlobalShader
 		SHADER_PARAMETER_RDG_BUFFER_SRV(StructuredBuffer<uint>, ShadingTileGrid)
 		SHADER_PARAMETER_RDG_BUFFER_SRV(StructuredBuffer<uint>, ShadingTiles)
 		SHADER_PARAMETER_RDG_TEXTURE(Texture2D<UNORM float>, ShadingTileAtlas)
-		SHADER_PARAMETER_RDG_BUFFER_SRV(StructuredBuffer<uint>, ShadowMaskHashTable)
 	END_SHADER_PARAMETER_STRUCT()
 
 	static int32 GetGroupSize()
@@ -601,7 +585,6 @@ void FDeferredShadingSceneRenderer::RenderStochasticShadows(FRDGBuilder& GraphBu
 	const FIntPoint ShadowMaskPageTableSize = ShadowMaskPageTablePerLightSize * StochasticShadows::MaxLightSceneIdXY;
 	const FIntPoint ShadowMaskTileAtlasSize = StochasticShadows::ShadowMaskTileSize * StochasticShadows::ShadowMaskAtlasSizeInTiles;
 	const int32 MaxShadowMaskTiles = StochasticShadows::ShadowMaskAtlasSizeInTiles * StochasticShadows::ShadowMaskAtlasSizeInTiles;
-	const int32 ShadowMaskHashTableSize = FMath::RoundUpToPowerOfTwo(MaxShadowMaskTiles * StochasticShadows::ShadowMaskHashTableFactor);
 
 	const FIntPoint ShadingTileGridSize = FIntPoint::DivideAndRoundUp(SceneTextures.Config.Extent, StochasticShadows::ShadowMaskTileSize);
 	const FIntPoint ShadingTileAtlasSize = ShadowMaskTileAtlasSize;
@@ -621,28 +604,12 @@ void FDeferredShadingSceneRenderer::RenderStochasticShadows(FRDGBuilder& GraphBu
 	FRDGTextureRef LightSampleRayDistance = GraphBuilder.CreateTexture(
 		FRDGTextureDesc::Create2D(DonwnsampledSampleBufferSize, PF_R16F, FClearValueBinding::Black, TexCreate_ShaderResource | TexCreate_UAV),
 		TEXT("StochasticShadows.LightSampleRayDistance"));
-
-	const bool bUseHashTable = CVarStochasticShadowsHashTable.GetValueOnRenderThread() != 0;
 	
-	FRDGBufferRef ShadowMaskHashTable = nullptr;
-	FRDGTextureRef ShadowMaskPageTable = nullptr;
-
-	if (bUseHashTable)
-	{
-		ShadowMaskHashTable = GraphBuilder.CreateBuffer(
-			FRDGBufferDesc::CreateStructuredDesc(sizeof(uint32), ShadowMaskHashTableSize),
-			TEXT("StochasticShadows.ShadowMaskTableMap"));
-
-		AddClearUAVPass(GraphBuilder, GraphBuilder.CreateUAV(ShadowMaskHashTable), StochasticShadows::InvalidShadowMaskTileIndex);
-	}
-	else
-	{
-		ShadowMaskPageTable = GraphBuilder.CreateTexture(
-			FRDGTextureDesc::Create2D(ShadowMaskPageTableSize, PF_R32_UINT, FClearValueBinding::Black, TexCreate_ShaderResource | TexCreate_UAV),
-			TEXT("StochasticShadows.ShadowMaskPageTable"));
+	FRDGTextureRef ShadowMaskPageTable = GraphBuilder.CreateTexture(
+		FRDGTextureDesc::Create2D(ShadowMaskPageTableSize, PF_R32_UINT, FClearValueBinding::Black, TexCreate_ShaderResource | TexCreate_UAV),
+		TEXT("StochasticShadows.ShadowMaskPageTable"));
 
-		AddClearUAVPass(GraphBuilder, GraphBuilder.CreateUAV(ShadowMaskPageTable), StochasticShadows::InvalidShadowMaskTileIndex);
-	}
+	AddClearUAVPass(GraphBuilder, GraphBuilder.CreateUAV(ShadowMaskPageTable), StochasticShadows::InvalidShadowMaskTileIndex);
 
 	const EPixelFormat ShadowMaskTileAtlasFormat = PF_R32_UINT;
 	FRDGTextureRef ShadowMaskTileAtlas = GraphBuilder.CreateTexture(
@@ -652,7 +619,6 @@ void FDeferredShadingSceneRenderer::RenderStochasticShadows(FRDGBuilder& GraphBu
 	bool bTemporal = CVarStochasticShadowsTemporal.GetValueOnRenderThread() != 0;
 	FRDGTextureRef ShadowMaskSceneDepthHistory = nullptr;
 	FRDGTextureRef ShadowMaskPageTableHistory = nullptr;
-	FRDGBufferRef ShadowMaskHashTableHistory = nullptr;
 	FRDGTextureRef ShadowMaskAtlasHistory = nullptr;
 	FVector4f HistoryScreenPositionScaleBias = FVector4f(0.0f, 0.0f, 0.0f, 0.0f);
 	FVector4f HistoryUVMinMax = FVector4f(0.0f, 0.0f, 0.0f, 0.0f);
@@ -672,12 +638,6 @@ void FDeferredShadingSceneRenderer::RenderStochasticShadows(FRDGBuilder& GraphBu
 				ShadowMaskPageTableHistory = GraphBuilder.RegisterExternalTexture(LightingViewState.ShadowMaskPageTableHistory);
 			}
 
-			if (LightingViewState.ShadowMaskHashTableHistory 
-				&& LightingViewState.ShadowMaskHashTableHistory->GetSize() == sizeof(uint32) * ShadowMaskHashTableSize)
-			{
-				ShadowMaskHashTableHistory = GraphBuilder.RegisterExternalBuffer(LightingViewState.ShadowMaskHashTableHistory);
-			}
-
 			if (LightingViewState.ShadowMaskAtlasHistory 
 				&& LightingViewState.ShadowMaskAtlasHistory->GetDesc().Extent == ShadowMaskTileAtlasSize)
 			{
@@ -687,19 +647,9 @@ void FDeferredShadingSceneRenderer::RenderStochasticShadows(FRDGBuilder& GraphBu
 		}
 	}
 
-	if (bUseHashTable)
+	if (ShadowMaskAtlasHistory == nullptr || ShadowMaskPageTableHistory == nullptr)
 	{
-		if (ShadowMaskAtlasHistory == nullptr || ShadowMaskHashTableHistory == nullptr)
-		{
-			bTemporal = false;
-		}
-	}
-	else
-	{
-		if (ShadowMaskAtlasHistory == nullptr || ShadowMaskPageTableHistory == nullptr)
-		{
-			bTemporal = false;
-		}
+		bTemporal = false;
 	}
 
 	const FIntPoint ViewSizeInTiles = FIntPoint::DivideAndRoundUp(View.ViewRect.Size(), StochasticShadows::TileSize);
@@ -726,7 +676,6 @@ void FDeferredShadingSceneRenderer::RenderStochasticShadows(FRDGBuilder& GraphBu
 		StochasticShadowsParameters.MaxShadingTiles = (ShadingTileAtlasSize.X * ShadingTileAtlasSize.Y) / (StochasticShadows::ShadowMaskTileSize * StochasticShadows::ShadowMaskTileSize);
 		StochasticShadowsParameters.MaxShadingTilesPerGridCell = StochasticShadows::MaxShadingTilesPerGridCell;
 		StochasticShadowsParameters.ShadingTileGridSize = ShadingTileGridSize;
-		StochasticShadowsParameters.ShadowMaskHashTableIndexWrapMask = ShadowMaskHashTableSize - 1;
 		StochasticShadowsParameters.ShadowMaskPageTablePerLightSize = ShadowMaskPageTablePerLightSize;
 		StochasticShadowsParameters.DownsampledBufferInvSize = FVector2f(1.0f) / DownsampledBufferSize;
 		StochasticShadowsParameters.SamplingMinWeight = FMath::Max(CVarStochasticShadowsSamplingMinWeight.GetValueOnRenderThread(), 0.0f);
@@ -838,7 +787,6 @@ void FDeferredShadingSceneRenderer::RenderStochasticShadows(FRDGBuilder& GraphBu
 		FRDGBufferUAVRef ShadowMaskTileAllocatorUAV = GraphBuilder.CreateUAV(ShadowMaskTileAllocator, ERDGUnorderedAccessViewFlags::SkipBarrier);
 		FRDGBufferUAVRef ShadowMaskTileDataUAV = GraphBuilder.CreateUAV(ShadowMaskTileData, ERDGUnorderedAccessViewFlags::SkipBarrier);
 		FRDGTextureUAVRef ShadowMaskPageTableUAV = ShadowMaskPageTable ? GraphBuilder.CreateUAV(ShadowMaskPageTable, ERDGUnorderedAccessViewFlags::SkipBarrier) : nullptr;
-		FRDGBufferUAVRef ShadowMaskHashTableUAV = ShadowMaskHashTable ? GraphBuilder.CreateUAV(ShadowMaskHashTable, ERDGUnorderedAccessViewFlags::SkipBarrier) : nullptr;
 		FRDGTextureUAVRef LightSamplesUAV = GraphBuilder.CreateUAV(LightSamples, ERDGUnorderedAccessViewFlags::SkipBarrier);
 
 		for (int32 TileType = 0; TileType < (int32)StochasticShadows::ETileType::MAX; ++TileType)
@@ -853,13 +801,11 @@ void FDeferredShadingSceneRenderer::RenderStochasticShadows(FRDGBuilder& GraphBu
 			PassParameters->RWShadowMaskTileAllocator = ShadowMaskTileAllocatorUAV;
 			PassParameters->RWShadowMaskTileData = ShadowMaskTileDataUAV;
 			PassParameters->RWShadowMaskPageTable = ShadowMaskPageTableUAV;
-			PassParameters->RWShadowMaskHashTable = ShadowMaskHashTableUAV;
 			PassParameters->RWLightSamples = LightSamplesUAV;
 			PassParameters->DownsampledTileAllocator = GraphBuilder.CreateSRV(DownsampledTileAllocator);
 			PassParameters->DownsampledTileData = GraphBuilder.CreateSRV(DownsampledTileData);
 			PassParameters->ForwardLightData = View.ForwardLightingResources.ForwardLightUniformBuffer;
 			PassParameters->ShadowMaskPageTableHistory = ShadowMaskPageTableHistory;
-			PassParameters->ShadowMaskHashTableHistory = ShadowMaskHashTableHistory ? GraphBuilder.CreateSRV(ShadowMaskHashTableHistory) : nullptr;
 			PassParameters->ShadowMaskAtlasHistory = ShadowMaskAtlasHistory;
 			PassParameters->ShadowMaskTileAtlas = ShadowMaskTileAtlas;
 			PassParameters->ShadowMaskSceneDepthHistory = ShadowMaskSceneDepthHistory;
@@ -872,7 +818,6 @@ void FDeferredShadingSceneRenderer::RenderStochasticShadows(FRDGBuilder& GraphBu
 			PermutationVector.Set<FGenerateSamplesCS::FShadowMaskReprojectionWeights>(bTemporal);
 			PermutationVector.Set<FGenerateSamplesCS::FShadowFactorEstimate>(bTemporal && CVarStochasticSamplingShadowEstimate.GetValueOnRenderThread() != 0);
 			PermutationVector.Set<FGenerateSamplesCS::FDebugMode>(bDebug);
-			PermutationVector.Set<FGenerateSamplesCS::FHashTable>(bUseHashTable);
 			PermutationVector.Set<FGenerateSamplesCS::FCandidateLightMask>(ShadowMaskAtlasHistory && CVarStochasticShadowsCandidateLightMask.GetValueOnRenderThread() != 0);
 			PermutationVector = FGenerateSamplesCS::RemapPermutation(PermutationVector);
 			auto ComputeShader = View.ShaderMap->GetShader<FGenerateSamplesCS>(PermutationVector);
@@ -971,7 +916,6 @@ void FDeferredShadingSceneRenderer::RenderStochasticShadows(FRDGBuilder& GraphBu
 		PassParameters->ShadowMaskHistoryScreenCoord00 = ShadowMaskHistoryScreenCoord00;
 		PassParameters->ShadowMaskHistoryWeights = ShadowMaskHistoryWeights;
 		PassParameters->ShadowMaskPageTableHistory = ShadowMaskPageTableHistory;
-		PassParameters->ShadowMaskHashTableHistory = ShadowMaskHashTableHistory ? GraphBuilder.CreateSRV(ShadowMaskHashTableHistory) : nullptr;
 		PassParameters->ShadowMaskAtlasHistory = ShadowMaskAtlasHistory;
 		PassParameters->LightSamples = LightSamples;
 		PassParameters->ForwardLightData = View.ForwardLightingResources.ForwardLightUniformBuffer;
@@ -979,7 +923,6 @@ void FDeferredShadingSceneRenderer::RenderStochasticShadows(FRDGBuilder& GraphBu
 		FCompositeShadowMaskTracesCS::FPermutationDomain PermutationVector;
 		PermutationVector.Set<FCompositeShadowMaskTracesCS::FNumSamplesPerPixel>(NumSamplesPerPixel1d);
 		PermutationVector.Set<FCompositeShadowMaskTracesCS::FTemporalAccumulation>(bTemporal);
-		PermutationVector.Set<FCompositeShadowMaskTracesCS::FHashTable>(bUseHashTable);
 		PermutationVector.Set<FCompositeShadowMaskTracesCS::FDebugMode>(bDebug);
 		auto ComputeShader = View.ShaderMap->GetShader<FCompositeShadowMaskTracesCS>(PermutationVector);
 
@@ -1004,11 +947,9 @@ void FDeferredShadingSceneRenderer::RenderStochasticShadows(FRDGBuilder& GraphBu
 		PassParameters->ShadowMaskUpsampleWeights = ShadowMaskUpsampleWeights;
 		PassParameters->ShadowMaskTileData = GraphBuilder.CreateSRV(ShadowMaskTileData);
 		PassParameters->ShadowMaskPageTable = ShadowMaskPageTable;
-		PassParameters->ShadowMaskHashTable = ShadowMaskHashTable ? GraphBuilder.CreateSRV(ShadowMaskHashTable) : nullptr;
 		PassParameters->ShadowMaskTileAtlas = ShadowMaskTileAtlas;
 
 		FShadowMaskSpatialPassCS::FPermutationDomain PermutationVector;
-		PermutationVector.Set<FShadowMaskSpatialPassCS::FHashTable>(bUseHashTable);
 		PermutationVector.Set<FShadowMaskSpatialPassCS::FSpatialPass>(CVarStochasticShadowsSpatial.GetValueOnRenderThread() != 0);
 		auto ComputeShader = View.ShaderMap->GetShader<FShadowMaskSpatialPassCS>(PermutationVector);
 
@@ -1040,7 +981,6 @@ void FDeferredShadingSceneRenderer::RenderStochasticShadows(FRDGBuilder& GraphBu
 			PassParameters->ShadingTileGrid = GraphBuilder.CreateSRV(ShadingTileGrid);
 			PassParameters->ShadingTileAtlas = ShadingTileAtlas;
 			PassParameters->RWSceneColor = SceneColorUAV;
-			PassParameters->ShadowMaskHashTable = ShadowMaskHashTable ? GraphBuilder.CreateSRV(ShadowMaskHashTable) : nullptr;
 
 			FShadeLightSamplesCS::FPermutationDomain PermutationVector;
 			PermutationVector.Set<FShadeLightSamplesCS::FTileType>(TileType);
@@ -1061,15 +1001,9 @@ void FDeferredShadingSceneRenderer::RenderStochasticShadows(FRDGBuilder& GraphBu
 	{
 		FStochasticShadowsViewState& LightingViewState = View.ViewState->StochasticShadows;
 		LightingViewState.ShadowMaskPageTableHistory = nullptr;
-		LightingViewState.ShadowMaskHashTableHistory = nullptr;
 		LightingViewState.ShadowMaskAtlasHistory = nullptr;
 		LightingViewState.ShadowMaskSceneDepthHistory = nullptr;
 
-		if (ShadowMaskHashTable)
-		{
-			GraphBuilder.QueueBufferExtraction(ShadowMaskHashTable, &LightingViewState.ShadowMaskHashTableHistory);
-		}
-
 		if (ShadowMaskPageTable)
 		{
 			GraphBuilder.QueueTextureExtraction(ShadowMaskPageTable, &LightingViewState.ShadowMaskPageTableHistory);
diff --git a/Engine/Source/Runtime/Renderer/Private/StochasticShadows/StochasticShadowsInternal.h b/Engine/Source/Runtime/Renderer/Private/StochasticShadows/StochasticShadowsInternal.h
--- a/Engine/Source/Runtime/Renderer/Private/StochasticShadows/StochasticShadowsInternal.h
+++ b/Engine/Source/Runtime/Renderer/Private/StochasticShadows/StochasticShadowsInternal.h
@@ -21,7 +21,6 @@ BEGIN_SHADER_PARAMETER_STRUCT(FStochasticShadowsParameters, )
 	SHADER_PARAMETER(uint32, MaxShadowMaskTiles)
 	SHADER_PARAMETER(uint32, MaxShadingTiles)
 	SHADER_PARAMETER(uint32, MaxShadingTilesPerGridCell)
-	SHADER_PARAMETER(uint32, ShadowMaskHashTableIndexWrapMask)
 	SHADER_PARAMETER(FIntPoint, ShadowMaskPageTablePerLightSize)
 	SHADER_PARAMETER(FIntPoint, ShadingTileGridSize)
 	SHADER_PARAMETER(FVector2f, DownsampledBufferInvSize)
diff --git a/Engine/Source/Runtime/Renderer/Private/StochasticShadows/StochasticShadowsViewState.h b/Engine/Source/Runtime/Renderer/Private/StochasticShadows/StochasticShadowsViewState.h
--- a/Engine/Source/Runtime/Renderer/Private/StochasticShadows/StochasticShadowsViewState.h
+++ b/Engine/Source/Runtime/Renderer/Private/StochasticShadows/StochasticShadowsViewState.h
@@ -6,7 +6,6 @@ class FStochasticShadowsViewState
 {
 public:
 	TRefCountPtr<IPooledRenderTarget> ShadowMaskPageTableHistory;
-	TRefCountPtr<FRDGPooledBuffer> ShadowMaskHashTableHistory;
 	TRefCountPtr<IPooledRenderTarget> ShadowMaskAtlasHistory;
 	TRefCountPtr<IPooledRenderTarget> ShadowMaskSceneDepthHistory;
 
@@ -16,7 +15,6 @@ public:
 	void SafeRelease()
 	{
 		ShadowMaskPageTableHistory.SafeRelease();
-		ShadowMaskHashTableHistory.SafeRelease();
 		ShadowMaskAtlasHistory.SafeRelease();
 		ShadowMaskSceneDepthHistory.SafeRelease();
 	}
