Regroup hair voxel pass (0.27ms -> 0.21ms)

diff --git a/Engine/Plugins/Runtime/HairStrands/Source/HairStrandsCore/Private/HairStrandsMeshProjection.cpp b/Engine/Plugins/Runtime/HairStrands/Source/HairStrandsCore/Private/HairStrandsMeshProjection.cpp
--- a/Engine/Plugins/Runtime/HairStrands/Source/HairStrandsCore/Private/HairStrandsMeshProjection.cpp
+++ b/Engine/Plugins/Runtime/HairStrands/Source/HairStrandsCore/Private/HairStrandsMeshProjection.cpp
@@ -1206,103 +1206,4 @@ void AddHairStrandUpdatePositionOffsetPass(
 
 	GraphBuilder.SetBufferAccessFinal(OutCurrPositionOffsetBuffer.Buffer, ERHIAccess::SRVMask);
 	GraphBuilder.SetBufferAccessFinal(OutPrevPositionOffsetBuffer.Buffer, ERHIAccess::SRVMask);
-}
-
-///////////////////////////////////////////////////////////////////////////////////////////////////
-
-class FHairResourceTransitionPass : public FGlobalShader
-{
-public:
-private:
-	DECLARE_GLOBAL_SHADER(FHairResourceTransitionPass);
-	SHADER_USE_PARAMETER_STRUCT(FHairResourceTransitionPass, FGlobalShader);
-
-	class FStructuredBuffer : SHADER_PERMUTATION_BOOL("PERMUTATION_STRUCTURED_BUFFER");
-	using FPermutationDomain = TShaderPermutationDomain<FStructuredBuffer>;
-
-	BEGIN_SHADER_PARAMETER_STRUCT(FParameters, )
-		SHADER_PARAMETER(uint32, DummyValue)
-		SHADER_PARAMETER_RDG_BUFFER_SRV_ARRAY(Buffer, VertexBuffers, [16])
-		SHADER_PARAMETER_RDG_BUFFER_SRV_ARRAY(StructuredBuffer, StructuredBuffers, [16])
-		SHADER_PARAMETER_RDG_BUFFER_UAV(RWStructuredBuffer, DummyOutput)
-	END_SHADER_PARAMETER_STRUCT()
-
-public:
-	static bool ShouldCompilePermutation(const FGlobalShaderPermutationParameters& Parameters)
-	{
-		return IsHairStrandsSupported(EHairStrandsShaderType::All, Parameters.Platform);
-	}
-	static void ModifyCompilationEnvironment(const FGlobalShaderPermutationParameters& Parameters, FShaderCompilerEnvironment& OutEnvironment)
-	{
-		FGlobalShader::ModifyCompilationEnvironment(Parameters, OutEnvironment);
-		OutEnvironment.SetDefine(TEXT("SHADER_RESOURCE_TRANSITION"), 1);
-	}
-};
-
-IMPLEMENT_GLOBAL_SHADER(FHairResourceTransitionPass, "/Engine/Private/HairStrands/HairStrandsMesh.usf", "MainCS", SF_Compute);
-
-void AddTransitionPass(
-	FRDGBuilder& GraphBuilder,
-	FGlobalShaderMap* ShaderMap,
-	const TArray<FRDGBufferSRVRef>& Transitions)
-{
-	const uint32 ResourceCount = Transitions.Num();
-	if (ResourceCount == 0)
-	{
-		return;
-	}
-
-	FRDGBufferSRVRef DummyVertexInput = GraphBuilder.CreateSRV(GSystemTextures.GetDefaultBuffer(GraphBuilder, 4u, 1u), PF_R32_UINT);
-	FRDGBufferSRVRef DummyStructuredInput= GraphBuilder.CreateSRV(GSystemTextures.GetDefaultStructuredBuffer(GraphBuilder, 16u));
-
-	FRDGBufferRef DummyOutput = GraphBuilder.CreateBuffer(FRDGBufferDesc::CreateStructuredDesc(4,1),TEXT("DummyOutput"));
-	FRDGBufferUAVRef DummyOutputUAV = GraphBuilder.CreateUAV(DummyOutput, PF_R32_UINT, ERDGUnorderedAccessViewFlags::SkipBarrier);
-
-	check(Transitions[0]);
-	check(Transitions[0]->Desc.Buffer);
-	const bool bStructuredBuffer = EnumHasAnyFlags(Transitions[0]->Desc.Buffer->Desc.Usage, EBufferUsageFlags::StructuredBuffer);
-
-	const uint32 MaxBufferCount = 16;
-	const uint32 PassCount = FMath::DivideAndRoundUp(ResourceCount, MaxBufferCount);
-	for (uint32 PassIt=0; PassIt< PassCount; ++PassIt)
-	{
-		FHairResourceTransitionPass::FParameters* PassParameters = GraphBuilder.AllocParameters<FHairResourceTransitionPass::FParameters>();
-		PassParameters->DummyValue = 0;
-		PassParameters->DummyOutput = DummyOutputUAV;
-
-		const uint32 PassResourceOffset = PassIt * MaxBufferCount;
-		const uint32 PassResourceCount = FMath::Min(int32(MaxBufferCount), int32(ResourceCount) - int32(PassResourceOffset));
-		if (bStructuredBuffer)
-		{
-			for (uint32 ResourceIt = 0; ResourceIt < PassResourceCount; ++ResourceIt)
-			{
-				PassParameters->StructuredBuffers[ResourceIt] = Transitions[PassResourceOffset + ResourceIt];
-			}
-			for (uint32 ResourceIt = PassResourceCount; ResourceIt < MaxBufferCount; ++ResourceIt)
-			{
-				PassParameters->StructuredBuffers[ResourceIt] = DummyStructuredInput;
-			}
-		}
-		else
-		{
-			for (uint32 ResourceIt = 0; ResourceIt < PassResourceCount; ++ResourceIt)
-			{
-				PassParameters->VertexBuffers[ResourceIt] = Transitions[PassResourceOffset + ResourceIt];
-			}
-			for (uint32 ResourceIt = PassResourceCount; ResourceIt < MaxBufferCount; ++ResourceIt)
-			{
-				PassParameters->VertexBuffers[ResourceIt] = DummyVertexInput;
-			}
-		}
-		FHairResourceTransitionPass::FPermutationDomain PermutationVector;
-		PermutationVector.Set<FHairResourceTransitionPass::FStructuredBuffer>(bStructuredBuffer);
-		TShaderMapRef<FHairResourceTransitionPass> ComputeShader(ShaderMap, PermutationVector);
-		FComputeShaderUtils::AddPass(
-			GraphBuilder,
-			RDG_EVENT_NAME("HairStrands::ResourceTransitions"),
-			ERDGPassFlags::Compute | ERDGPassFlags::NeverCull,
-			ComputeShader,
-			PassParameters,
-			FIntVector(1, 1, 1));
-	}
 }
\ No newline at end of file
diff --git a/Engine/Plugins/Runtime/HairStrands/Source/HairStrandsCore/Public/HairStrandsMeshProjection.h b/Engine/Plugins/Runtime/HairStrands/Source/HairStrandsCore/Public/HairStrandsMeshProjection.h
--- a/Engine/Plugins/Runtime/HairStrands/Source/HairStrandsCore/Public/HairStrandsMeshProjection.h
+++ b/Engine/Plugins/Runtime/HairStrands/Source/HairStrandsCore/Public/HairStrandsMeshProjection.h
@@ -153,9 +153,4 @@ void AddHairStrandUpdatePositionOffsetPass(
 	FHairStrandsDeformedRootResource* DeformedRootResources,
 	FHairStrandsDeformedResource* DeformedResources);
 
-void AddTransitionPass(
-	FRDGBuilder& GraphBuilder,
-	FGlobalShaderMap* ShaderMap,
-	const TArray<FRDGBufferSRVRef>& Transitions);
-
 FHairStrandsProjectionMeshData ExtractMeshData(FSkeletalMeshRenderData* RenderData);
diff --git a/Engine/Source/Runtime/Renderer/Private/HairStrands/HairStrandsUtils.cpp b/Engine/Source/Runtime/Renderer/Private/HairStrands/HairStrandsUtils.cpp
--- a/Engine/Source/Runtime/Renderer/Private/HairStrands/HairStrandsUtils.cpp
+++ b/Engine/Source/Runtime/Renderer/Private/HairStrands/HairStrandsUtils.cpp
@@ -6,6 +6,7 @@
 #include "HairStrandsCluster.h"
 #include "Rendering/SkeletalMeshRenderData.h"
 #include "HairStrandsData.h"
+#include "SystemTextures.h"
 
 static float GHairR = 1;
 static float GHairTT = 1;
@@ -346,3 +347,101 @@ uint32 PackHairRenderInfoBits(
 	BitField |= bIsGPUDriven ? 0x2 : 0;
 	return BitField;
 }
+
+///////////////////////////////////////////////////////////////////////////////////////////////////
+
+class FHairResourceTransitionPass : public FGlobalShader
+{
+private:
+	DECLARE_GLOBAL_SHADER(FHairResourceTransitionPass);
+	SHADER_USE_PARAMETER_STRUCT(FHairResourceTransitionPass, FGlobalShader);
+
+	class FStructuredBuffer : SHADER_PERMUTATION_BOOL("PERMUTATION_STRUCTURED_BUFFER");
+	using FPermutationDomain = TShaderPermutationDomain<FStructuredBuffer>;
+
+	BEGIN_SHADER_PARAMETER_STRUCT(FParameters, )
+		SHADER_PARAMETER(uint32, DummyValue)
+		SHADER_PARAMETER_RDG_BUFFER_SRV_ARRAY(Buffer, VertexBuffers, [16])
+		SHADER_PARAMETER_RDG_BUFFER_SRV_ARRAY(StructuredBuffer, StructuredBuffers, [16])
+		SHADER_PARAMETER_RDG_BUFFER_UAV(RWStructuredBuffer, DummyOutput)
+		END_SHADER_PARAMETER_STRUCT()
+
+public:
+		static bool ShouldCompilePermutation(const FGlobalShaderPermutationParameters& Parameters)
+	{
+		return IsHairStrandsSupported(EHairStrandsShaderType::All, Parameters.Platform);
+	}
+	static void ModifyCompilationEnvironment(const FGlobalShaderPermutationParameters& Parameters, FShaderCompilerEnvironment& OutEnvironment)
+	{
+		FGlobalShader::ModifyCompilationEnvironment(Parameters, OutEnvironment);
+		OutEnvironment.SetDefine(TEXT("SHADER_RESOURCE_TRANSITION"), 1);
+	}
+};
+
+IMPLEMENT_GLOBAL_SHADER(FHairResourceTransitionPass, "/Engine/Private/HairStrands/HairStrandsMesh.usf", "MainCS", SF_Compute);
+
+void AddTransitionPass(
+	FRDGBuilder& GraphBuilder,
+	FGlobalShaderMap* ShaderMap,
+	const TArray<FRDGBufferSRVRef>& Transitions)
+{
+	const uint32 ResourceCount = Transitions.Num();
+	if (ResourceCount == 0)
+	{
+		return;
+	}
+
+	FRDGBufferSRVRef DummyVertexInput = GraphBuilder.CreateSRV(GSystemTextures.GetDefaultBuffer(GraphBuilder, 4u, 1u), PF_R32_UINT);
+	FRDGBufferSRVRef DummyStructuredInput= GraphBuilder.CreateSRV(GSystemTextures.GetDefaultStructuredBuffer(GraphBuilder, 16u));
+
+	FRDGBufferRef DummyOutput = GraphBuilder.CreateBuffer(FRDGBufferDesc::CreateStructuredDesc(4,1),TEXT("DummyOutput"));
+	FRDGBufferUAVRef DummyOutputUAV = GraphBuilder.CreateUAV(DummyOutput, PF_R32_UINT, ERDGUnorderedAccessViewFlags::SkipBarrier);
+
+	check(Transitions[0]);
+	check(Transitions[0]->Desc.Buffer);
+	const bool bStructuredBuffer = EnumHasAnyFlags(Transitions[0]->Desc.Buffer->Desc.Usage, EBufferUsageFlags::StructuredBuffer);
+
+	const uint32 MaxBufferCount = 16;
+	const uint32 PassCount = FMath::DivideAndRoundUp(ResourceCount, MaxBufferCount);
+	for (uint32 PassIt=0; PassIt< PassCount; ++PassIt)
+	{
+		FHairResourceTransitionPass::FParameters* PassParameters = GraphBuilder.AllocParameters<FHairResourceTransitionPass::FParameters>();
+		PassParameters->DummyValue = 0;
+		PassParameters->DummyOutput = DummyOutputUAV;
+
+		const uint32 PassResourceOffset = PassIt * MaxBufferCount;
+		const uint32 PassResourceCount = FMath::Min(int32(MaxBufferCount), int32(ResourceCount) - int32(PassResourceOffset));
+		if (bStructuredBuffer)
+		{
+			for (uint32 ResourceIt = 0; ResourceIt < PassResourceCount; ++ResourceIt)
+			{
+				PassParameters->StructuredBuffers[ResourceIt] = Transitions[PassResourceOffset + ResourceIt];
+			}
+			for (uint32 ResourceIt = PassResourceCount; ResourceIt < MaxBufferCount; ++ResourceIt)
+			{
+				PassParameters->StructuredBuffers[ResourceIt] = DummyStructuredInput;
+			}
+		}
+		else
+		{
+			for (uint32 ResourceIt = 0; ResourceIt < PassResourceCount; ++ResourceIt)
+			{
+				PassParameters->VertexBuffers[ResourceIt] = Transitions[PassResourceOffset + ResourceIt];
+			}
+			for (uint32 ResourceIt = PassResourceCount; ResourceIt < MaxBufferCount; ++ResourceIt)
+			{
+				PassParameters->VertexBuffers[ResourceIt] = DummyVertexInput;
+			}
+		}
+		FHairResourceTransitionPass::FPermutationDomain PermutationVector;
+		PermutationVector.Set<FHairResourceTransitionPass::FStructuredBuffer>(bStructuredBuffer);
+		TShaderMapRef<FHairResourceTransitionPass> ComputeShader(ShaderMap, PermutationVector);
+		FComputeShaderUtils::AddPass(
+			GraphBuilder,
+			RDG_EVENT_NAME("HairStrands::ResourceTransitions"),
+			ERDGPassFlags::Compute | ERDGPassFlags::NeverCull,
+			ComputeShader,
+			PassParameters,
+			FIntVector(1, 1, 1));
+	}
+}
\ No newline at end of file
diff --git a/Engine/Source/Runtime/Renderer/Private/HairStrands/HairStrandsVoxelization.cpp b/Engine/Source/Runtime/Renderer/Private/HairStrands/HairStrandsVoxelization.cpp
--- a/Engine/Source/Runtime/Renderer/Private/HairStrands/HairStrandsVoxelization.cpp
+++ b/Engine/Source/Runtime/Renderer/Private/HairStrands/HairStrandsVoxelization.cpp
@@ -193,7 +193,7 @@ static void AddVirtualVoxelInjectOpaquePass(
 	FRDGBuilder& GraphBuilder,
 	const FViewInfo& View,
 	const FHairStrandsVoxelResources& VoxelResources,
-	const FHairStrandsMacroGroupData& MacroGroup,
+	const uint32& MacroGroupId,
 	FRDGTextureUAVRef OutTexturePage)
 {
 	FSceneTextureParameters SceneTextures = GetSceneTextureParameters(GraphBuilder, View);
@@ -211,13 +211,13 @@ static void AddVirtualVoxelInjectOpaquePass(
 	Parameters->VoxelMarkCount = FMath::Max(0, GHairStransVoxelInjectOpaqueMarkCount);
 	Parameters->SceneDepthResolution = SceneTextures.SceneDepthTexture->Desc.Extent;
 	Parameters->SceneTextures = SceneTextures;
-	Parameters->MacroGroupId = MacroGroup.MacroGroupId;
+	Parameters->MacroGroupId = MacroGroupId;
 	Parameters->OutPageTexture = OutTexturePage;
 	Parameters->IndirectDispatchArgs = VoxelResources.IndirectArgsBuffer;
 	TShaderMapRef<FVirtualVoxelInjectOpaqueCS> ComputeShader(View.ShaderMap);
 	const FGlobalShaderMap* GlobalShaderMap = View.ShaderMap;
 
-	const uint32 ArgsOffset = sizeof(FRHIDispatchIndirectParameters) * Parameters->MacroGroupId;
+	const uint32 ArgsOffset = sizeof(FRHIDispatchIndirectParameters) * MacroGroupId;
 
 	FComputeShaderUtils::AddPass(
 		GraphBuilder, 
@@ -575,6 +575,7 @@ static void AddAllocateVoxelPagesPass(
 	FRDGBufferRef ReadBackBuffer = GraphBuilder.CreateBuffer(FRDGBufferDesc::CreateBufferDesc(sizeof(uint32), 1), TEXT("ReadBackAllocations"));
 
 	FRDGBufferUAVRef PageIndexBufferUAV = GraphBuilder.CreateUAV(PageIndexBuffer, PF_R32_UINT);
+	FRDGBufferUAVRef PageIndexBufferUAVSkipBarrier = GraphBuilder.CreateUAV(PageIndexBuffer, PF_R32_UINT, ERDGUnorderedAccessViewFlags::SkipBarrier);
 	FRDGBufferUAVRef PageIndexOccupancyBufferUAV = GraphBuilder.CreateUAV(PageIndexOccupancyBuffer, PF_R32G32_UINT);
 	FRDGBufferUAVRef PageIndexGlobalCounterUAV = GraphBuilder.CreateUAV(PageIndexGlobalCounter, PF_R32_UINT);
 	
@@ -651,15 +652,11 @@ static void AddAllocateVoxelPagesPass(
 	FRDGBufferSRVRef PageIndexResolutionAndOffsetBufferSRV = GraphBuilder.CreateSRV(PageIndexResolutionBuffer, PF_R32G32B32A32_UINT);
 
 	// Mark valid page index
-	for (uint32 MacroGroupIt=0; MacroGroupIt <MacroGroupCount;++MacroGroupIt)
+	for (uint32 MacroGroupIt = 0; MacroGroupIt < MacroGroupCount; ++MacroGroupIt)
 	{
-		DECLARE_GPU_STAT(HairStrandsAllocateMacroGroup);
-		SCOPED_DRAW_EVENT(GraphBuilder.RHICmdList, HairStrandsAllocateMacroGroup);
-		SCOPED_GPU_STAT(GraphBuilder.RHICmdList, HairStrandsAllocateMacroGroup);
-
 		const FHairStrandsMacroGroupData& MacroGroup = MacroGroupDatas[MacroGroupIt];
 		FCPUMacroGroupAllocation& CPUAllocationDesc = CPUAllocationDescs[MacroGroupIt];
-		{			
+		{
 			for (const FHairStrandsMacroGroupData::PrimitiveInfo& PrimitiveInfo : MacroGroup.PrimitivesInfos)
 			{
 				FHairGroupPublicData* HairGroupData = PrimitiveInfo.PublicDataPtr;
@@ -676,23 +673,23 @@ static void AddAllocateVoxelPagesPass(
 				const bool bUseClusterAABB = TransientResources.IsClusterAABBValid(HairGroupData->Instance->RegisteredIndex) && bIsGPUDriven;
 				
 				FVoxelMarkValidPageIndex_PrepareCS::FParameters* Parameters = GraphBuilder.AllocParameters<FVoxelMarkValidPageIndex_PrepareCS::FParameters>();
-				Parameters->InstanceRegisteredIndex						= HairGroupData->Instance->RegisteredIndex;
-				Parameters->ClusterOffset								= TransientResources.GetClusterOffset(HairGroupData->Instance->RegisteredIndex);
-				Parameters->ClusterCount 								= TransientResources.GetClusterCount(HairGroupData->Instance->RegisteredIndex);
-				Parameters->MacroGroupId								= MacroGroup.MacroGroupId;
-				Parameters->GroupAABBsBuffer							= TransientResources.GroupAABBSRV;
-				Parameters->ClusterAABBsBuffer							= TransientResources.ClusterAABBSRV;
-				Parameters->MacroGroupVoxelAlignedAABBBuffer			= GraphBuilder.CreateSRV(MacroGroupResources.MacroGroupVoxelAlignedAABBsBuffer, PF_R32_SINT);
-				Parameters->PageIndexResolutionAndOffsetBuffer			= PageIndexResolutionAndOffsetBufferSRV;
-				Parameters->bUseMacroGroupBoundCPU						= bIsGPUDriven ? 0 : 1;
-				Parameters->MacroGroupBoundCPU_TranslatedWorldMinAABB	= (FVector3f)CPUAllocationDesc.TranslatedMinAABB;
-				Parameters->MacroGroupBoundCPU_TranslatedWorldMaxAABB	= (FVector3f)CPUAllocationDesc.TranslatedMaxAABB;
-				Parameters->OutValidPageIndexBuffer						= PageIndexBufferUAV;
+				Parameters->InstanceRegisteredIndex = HairGroupData->Instance->RegisteredIndex;
+				Parameters->ClusterOffset = TransientResources.GetClusterOffset(HairGroupData->Instance->RegisteredIndex);
+				Parameters->ClusterCount = TransientResources.GetClusterCount(HairGroupData->Instance->RegisteredIndex);
+				Parameters->MacroGroupId = MacroGroup.MacroGroupId;
+				Parameters->GroupAABBsBuffer = TransientResources.GroupAABBSRV;
+				Parameters->ClusterAABBsBuffer = TransientResources.ClusterAABBSRV;
+				Parameters->MacroGroupVoxelAlignedAABBBuffer = GraphBuilder.CreateSRV(MacroGroupResources.MacroGroupVoxelAlignedAABBsBuffer, PF_R32_SINT);
+				Parameters->PageIndexResolutionAndOffsetBuffer = PageIndexResolutionAndOffsetBufferSRV;
+				Parameters->bUseMacroGroupBoundCPU = bIsGPUDriven ? 0 : 1;
+				Parameters->MacroGroupBoundCPU_TranslatedWorldMinAABB = (FVector3f)CPUAllocationDesc.TranslatedMinAABB;
+				Parameters->MacroGroupBoundCPU_TranslatedWorldMaxAABB = (FVector3f)CPUAllocationDesc.TranslatedMaxAABB;
+				Parameters->OutValidPageIndexBuffer = PageIndexBufferUAVSkipBarrier;
 
 				FIntVector DispatchCount = bUseClusterAABB ?
 					FIntVector(
-						FMath::DivideAndRoundUp(Parameters->ClusterCount, FVoxelMarkValidPageIndex_PrepareCS::GetGroupSize(true)), 
-						1, 
+						FMath::DivideAndRoundUp(Parameters->ClusterCount, FVoxelMarkValidPageIndex_PrepareCS::GetGroupSize(true)),
+						1,
 						1) :
 					FIntVector(
 						FMath::DivideAndRoundUp(uint32(CPUAllocationDesc.PageIndexResolution.X), FVoxelMarkValidPageIndex_PrepareCS::GetGroupSize(false)),
@@ -714,6 +711,14 @@ static void AddAllocateVoxelPagesPass(
 					DispatchCount);
 			}
 		}
+	}
+
+	// Do this into a single pass by creatin a bufer with CPU AABB page res
+	FRDGBufferUAVRef NodeDescUAVSkipBarrier = GraphBuilder.CreateUAV(NodeDescBuffer, ERDGUnorderedAccessViewFlags::SkipBarrier);
+	for (uint32 MacroGroupIt = 0; MacroGroupIt < MacroGroupCount; ++MacroGroupIt)
+	{
+		const FHairStrandsMacroGroupData& MacroGroup = MacroGroupDatas[MacroGroupIt];
+		FCPUMacroGroupAllocation& CPUAllocationDesc = CPUAllocationDescs[MacroGroupIt];
 
 		// Fill in hair-macro-group information.
 		// Note: This need to happen before the allocation as we copy the index global count. This global index is 
@@ -731,7 +736,7 @@ static void AddAllocateVoxelPagesPass(
 			Parameters->GPU_VoxelWorldSize = CurrGPUMinVoxelSizeSRV;
 			Parameters->bUseCPUVoxelWorldSize = bAdaptiveVoxelEnable ? 0u : 1u;
 			Parameters->MacroGroupVoxelSizeBuffer = GraphBuilder.CreateSRV(MacroGroupResources.MacroGroupVoxelSizeBuffer, PF_R16F);
-			Parameters->OutNodeDescBuffer = GraphBuilder.CreateUAV(NodeDescBuffer);
+			Parameters->OutNodeDescBuffer = NodeDescUAVSkipBarrier;
 
 			if (bIsGPUDriven)
 			{
@@ -746,6 +751,13 @@ static void AddAllocateVoxelPagesPass(
 			TShaderMapRef<FVoxelAddNodeDescCS> ComputeShader(View.ShaderMap, PermutationVector);
 			FComputeShaderUtils::AddPass(GraphBuilder, RDG_EVENT_NAME("HairStrands::AddNodeDesc"), ComputeShader, Parameters, DispatchCount);
 		}
+	}
+
+	// Need to run serial, as PageIndexGlobalCounterUAV is used pass to pass 
+	for (uint32 MacroGroupIt = 0; MacroGroupIt < MacroGroupCount; ++MacroGroupIt)
+	{
+		const FHairStrandsMacroGroupData& MacroGroup = MacroGroupDatas[MacroGroupIt];
+		FCPUMacroGroupAllocation& CPUAllocationDesc = CPUAllocationDescs[MacroGroupIt];
 
 		// Allocate pages
 		{
@@ -1121,21 +1133,27 @@ IMPLEMENT_GLOBAL_SHADER(FVoxelRasterComputeCS, "/Engine/Private/HairStrands/Hair
 
 bool IsHairStrandContinuousDecimationReorderingEnabled();
 
-static void AddVirtualVoxelizationComputeRasterPass(
+struct FInstanceData
+{
+	uint32 MacroGroupId = 0;
+	const FHairGroupPublicData* Data = nullptr;
+};
+
+static void AddVirtualVoxelizationRasterPass(
 	FRDGBuilder& GraphBuilder,
 	const FViewInfo* ViewInfo,
 	const FHairStrandsVoxelResources& VoxelResources,
-	const FHairStrandsMacroGroupData& MacroGroup)
+	const TArray<FInstanceData>& InstanceDatas)
 {
-	ensure(MacroGroup.bSupportVoxelization);
+	DECLARE_GPU_STAT(HairStrandsVoxelize);
+	SCOPED_DRAW_EVENT(GraphBuilder.RHICmdList, HairStrandsVoxelize);
+	SCOPED_GPU_STAT(GraphBuilder.RHICmdList, HairStrandsVoxelize);
 
 	if (ViewInfo)
 	{
-		const FHairStrandsMacroGroupData::TPrimitiveInfos& PrimitiveSceneInfos = MacroGroup.PrimitivesInfos;
-
 		FRDGTextureUAVRef PageTextureUAV = GraphBuilder.CreateUAV(VoxelResources.PageTexture, ERDGUnorderedAccessViewFlags::SkipBarrier);
 
-		const uint32 FrameIdMode8 = ViewInfo && ViewInfo->ViewState ? (ViewInfo->ViewState->GetFrameIndex() % 8) : 0;
+		const uint32 FrameIdMode8 = ViewInfo->ViewState ? (ViewInfo->ViewState->GetFrameIndex() % 8) : 0;
 		const FVector& TranslatedWorldOffset = ViewInfo->ViewMatrices.GetPreViewTranslation();
 
 		FVoxelRasterComputeCS::FPermutationDomain PermutationVector_Off;
@@ -1145,67 +1163,37 @@ static void AddVirtualVoxelizationComputeRasterPass(
 
 		TShaderMapRef<FVoxelRasterComputeCS> ComputeShader_CullingOff(ViewInfo->ShaderMap, PermutationVector_Off);
 		TShaderMapRef<FVoxelRasterComputeCS> ComputeShader_CullingOn(ViewInfo->ShaderMap, PermutationVector_On);
-		for (const FHairStrandsMacroGroupData::PrimitiveInfo& PrimitiveInfo : PrimitiveSceneInfos)
+		for (const FInstanceData& InstanceData : InstanceDatas)
 		{
-			check(PrimitiveInfo.PublicDataPtr);
-			const FHairGroupPublicData* HairGroupPublicData = PrimitiveInfo.PublicDataPtr;
-
-			if (HairGroupPublicData->DoesSupportVoxelization())
+			const uint32 PointCount = InstanceData.Data->GetActiveStrandsPointCount();
+			const float CoverageScale = InstanceData.Data->GetActiveStrandsCoverageScale();
+
+			FVoxelRasterComputeCS::FParameters* PassParameters = GraphBuilder.AllocParameters<FVoxelRasterComputeCS::FParameters>();
+			PassParameters->MaxRasterCount = FMath::Clamp(GHairStrandsVoxelComputeRasterMaxVoxelCount, 1, 256);
+			PassParameters->VirtualVoxelParams = VoxelResources.Parameters.Common;
+			PassParameters->MacroGroupId = InstanceData.MacroGroupId;
+			PassParameters->OutPageTexture = PageTextureUAV;
+			PassParameters->FrameIdMod8 = FrameIdMode8;
+			PassParameters->VertexCount = PointCount;
+			PassParameters->CoverageScale = CoverageScale;
+		
+			const bool bCullingEnable = InstanceData.Data->GetCullingResultAvailable();
+
+			PassParameters->HairInstance = GetHairStrandsInstanceParameters(GraphBuilder, *ViewInfo, InstanceData.Data, bCullingEnable, true);
+
+			if (bCullingEnable)
 			{
-				const FHairGroupPublicData::FVertexFactoryInput& VFInput = HairGroupPublicData->VFInput;
-				if (HairGroupPublicData->VFInput.Strands.PositionBuffer.Buffer == nullptr)
-				{
-					continue;
-				}
-
-				const uint32 PointCount = HairGroupPublicData->GetActiveStrandsPointCount();
-				const float CoverageScale = HairGroupPublicData->GetActiveStrandsCoverageScale();
-
-				FVoxelRasterComputeCS::FParameters* PassParameters = GraphBuilder.AllocParameters<FVoxelRasterComputeCS::FParameters>();
-				PassParameters->MaxRasterCount = FMath::Clamp(GHairStrandsVoxelComputeRasterMaxVoxelCount, 1, 256);
-				PassParameters->VirtualVoxelParams = VoxelResources.Parameters.Common;
-				PassParameters->MacroGroupId = MacroGroup.MacroGroupId;
-				PassParameters->OutPageTexture = PageTextureUAV;
-				PassParameters->FrameIdMod8 = FrameIdMode8;
-
-				PassParameters->VertexCount = PointCount;
-				PassParameters->CoverageScale = CoverageScale;
-			
-				const bool bCullingEnable = HairGroupPublicData->GetCullingResultAvailable();
-
-				PassParameters->HairInstance = GetHairStrandsInstanceParameters(GraphBuilder, *ViewInfo, HairGroupPublicData, bCullingEnable, true);
-
-				if (bCullingEnable)
-				{
-					FComputeShaderUtils::AddPass(GraphBuilder, RDG_EVENT_NAME("HairStrands::VoxelComputeRaster(culling=on)"), ComputeShader_CullingOn, PassParameters, PassParameters->HairInstance.HairStrandsVF.Culling.CullingIndirectBufferArgs, 0);
-				}
-				else
-				{
-					const FIntVector DispatchCount(FMath::DivideAndRoundUp(PassParameters->VertexCount, FVoxelRasterComputeCS::GetGroupSize()), 1, 1);
-					FComputeShaderUtils::AddPass(GraphBuilder, RDG_EVENT_NAME("HairStrands::VoxelComputeRaster(culling=off)"), ComputeShader_CullingOff, PassParameters, DispatchCount);
-				}
+				FComputeShaderUtils::AddPass(GraphBuilder, RDG_EVENT_NAME("HairStrands::VoxelComputeRaster(culling=on)"), ComputeShader_CullingOn, PassParameters, PassParameters->HairInstance.HairStrandsVF.Culling.CullingIndirectBufferArgs, 0);
+			}
+			else
+			{
+				const FIntVector DispatchCount(FMath::DivideAndRoundUp(PassParameters->VertexCount, FVoxelRasterComputeCS::GetGroupSize()), 1, 1);
+				FComputeShaderUtils::AddPass(GraphBuilder, RDG_EVENT_NAME("HairStrands::VoxelComputeRaster(culling=off)"), ComputeShader_CullingOff, PassParameters, DispatchCount);
 			}
 		}
 	}
 }
 
-///////////////////////////////////////////////////////////////////////////////////////////////////
-
-static void AddVirtualVoxelizationRasterPass(
-	FRDGBuilder& GraphBuilder,
-	const FScene* Scene,
-	const FViewInfo* ViewInfo,
-	FHairStrandsVoxelResources& VoxelResources,
-	FInstanceCullingManager& InstanceCullingManager,
-	const FHairStrandsMacroGroupData& MacroGroup)
-{
-	DECLARE_GPU_STAT(HairStrandsVoxelize);
-	SCOPED_DRAW_EVENT(GraphBuilder.RHICmdList, HairStrandsVoxelize);
-	SCOPED_GPU_STAT(GraphBuilder.RHICmdList, HairStrandsVoxelize);
-
-	AddVirtualVoxelizationComputeRasterPass(GraphBuilder, ViewInfo, VoxelResources, MacroGroup);
-}
-
 ///////////////////////////////////////////////////////////////////////////////////////////////////
 class FVirtualVoxelGenerateMipCS : public FGlobalShader
 {
@@ -1440,33 +1428,71 @@ void VoxelizeHairStrands(
 	RDG_GPU_STAT_SCOPE(GraphBuilder, HairStrandsVoxelization);
 
 	check(TransientResources);
-	{
-		FRDGBufferRef PageToPageIndexBuffer = nullptr;
-		FHairStrandsViewStateData* HairStrandsViewStateData = View.ViewState ? &View.ViewState->HairStrandsViewStateData : nullptr;
-		VirtualVoxelResources = AllocateVirtualVoxelResources(GraphBuilder, View, PreViewStereoCorrection, *TransientResources, MacroGroupDatas, MacroGroupResources, PageToPageIndexBuffer, HairStrandsViewStateData);
-
-		FRDGBufferRef ClearIndArgsBuffer = IndirectVoxelPageClear(GraphBuilder, View, VirtualVoxelResources);
 
-		for (FHairStrandsMacroGroupData& MacroGroup : MacroGroupDatas)
+	// Gather all instances which needs to be voxelized
+	uint32 TotalInstanceCount = 0;
+	TArray<FInstanceData> InstanceDatas;
+	InstanceDatas.Reserve(Scene->HairStrandsSceneData.RegisteredProxies.Num());
+	TArray<uint32> ValidMacroGroupIDs;
+	ValidMacroGroupIDs.Reserve(MacroGroupDatas.Num());
+	TArray<FUintVector2> MacroGroupInstanceOffsetAndCount;
+	MacroGroupInstanceOffsetAndCount.Init(FUintVector2::ZeroValue, MacroGroupDatas.Num());
+	for (FHairStrandsMacroGroupData& MacroGroup : MacroGroupDatas)
+	{
+		if (MacroGroup.bSupportVoxelization)
 		{
-			if (MacroGroup.bSupportVoxelization)
+			ValidMacroGroupIDs.Add(MacroGroup.MacroGroupId);
+			MacroGroupInstanceOffsetAndCount[MacroGroup.MacroGroupId].X = TotalInstanceCount;
+			for (const FHairStrandsMacroGroupData::PrimitiveInfo& PrimitiveInfo : MacroGroup.PrimitivesInfos)
 			{
-				AddVirtualVoxelizationRasterPass(GraphBuilder, Scene, &View, VirtualVoxelResources, InstanceCullingManager, MacroGroup);
-			}
-		}
+				check(PrimitiveInfo.PublicDataPtr);
+				const FHairGroupPublicData* HairGroupPublicData = PrimitiveInfo.PublicDataPtr;
+				const FHairGroupPublicData::FVertexFactoryInput& VFInput = HairGroupPublicData->VFInput;
+				if (HairGroupPublicData->VFInput.Strands.PositionBuffer.Buffer == nullptr)
+				{
+					continue;
+				}
 
-		if (GHairVoxelInjectOpaqueDepthEnable > 0)
-		{
-			FRDGTextureUAVRef PageTextureUAV = GraphBuilder.CreateUAV(VirtualVoxelResources.PageTexture, ERDGUnorderedAccessViewFlags::SkipBarrier);
-			for (FHairStrandsMacroGroupData& MacroGroup : MacroGroupDatas)
-			{
-				if (MacroGroup.bSupportVoxelization)
+				if (HairGroupPublicData->DoesSupportVoxelization())
 				{
-					AddVirtualVoxelInjectOpaquePass(GraphBuilder, View, VirtualVoxelResources, MacroGroup, PageTextureUAV);
+					FInstanceData& InstanceData = InstanceDatas.AddDefaulted_GetRef();
+					InstanceData.MacroGroupId = MacroGroup.MacroGroupId;
+					InstanceData.Data = HairGroupPublicData;
+					++TotalInstanceCount;
+					++MacroGroupInstanceOffsetAndCount[MacroGroup.MacroGroupId].Y;
 				}
 			}
 		}
+	}
+
+	// Force transition to ensure passes to be batched
+	TArray<FRDGBufferSRVRef> Transitions;
+	Transitions.Reserve(InstanceDatas.Num());
+	for (const FInstanceData& InstanceData : InstanceDatas)
+	{
+		if (InstanceData.Data->GetCullingResultAvailable())
+		{
+			Transitions.Add(RegisterAsSRV(GraphBuilder, InstanceData.Data->GetDrawIndirectRasterComputeBuffer()));
+
+		}
+	}
+	AddTransitionPass(GraphBuilder, View.ShaderMap, Transitions);
+
+	FRDGBufferRef PageToPageIndexBuffer = nullptr;
+	FHairStrandsViewStateData* HairStrandsViewStateData = View.ViewState ? &View.ViewState->HairStrandsViewStateData : nullptr;
+	VirtualVoxelResources = AllocateVirtualVoxelResources(GraphBuilder, View, PreViewStereoCorrection, *TransientResources, MacroGroupDatas, MacroGroupResources, PageToPageIndexBuffer, HairStrandsViewStateData);
 
-		AddVirtualVoxelGenerateMipPass(GraphBuilder, View, VirtualVoxelResources, ClearIndArgsBuffer, PageToPageIndexBuffer);
+	FRDGBufferRef ClearIndArgsBuffer = IndirectVoxelPageClear(GraphBuilder, View, VirtualVoxelResources);
+	AddVirtualVoxelizationRasterPass(GraphBuilder, &View, VirtualVoxelResources, InstanceDatas);
+
+	if (GHairVoxelInjectOpaqueDepthEnable > 0)
+	{
+		FRDGTextureUAVRef PageTextureUAV = GraphBuilder.CreateUAV(VirtualVoxelResources.PageTexture, ERDGUnorderedAccessViewFlags::SkipBarrier);
+		for (uint32 MacroGroupId : ValidMacroGroupIDs)
+		{
+			AddVirtualVoxelInjectOpaquePass(GraphBuilder, View, VirtualVoxelResources, MacroGroupId, PageTextureUAV);
+		}
 	}
+
+	AddVirtualVoxelGenerateMipPass(GraphBuilder, View, VirtualVoxelResources, ClearIndArgsBuffer, PageToPageIndexBuffer);
 }
diff --git a/Engine/Source/Runtime/Renderer/Public/HairStrandsInterface.h b/Engine/Source/Runtime/Renderer/Public/HairStrandsInterface.h
--- a/Engine/Source/Runtime/Renderer/Public/HairStrandsInterface.h
+++ b/Engine/Source/Runtime/Renderer/Public/HairStrandsInterface.h
@@ -20,6 +20,7 @@
 
 class UTexture2D;
 class FSceneInterface;
+class FGlobalShaderMap;
 
 ////////////////////////////////////////////////////////////////////////////////////////////////////
 // Shader parameters
@@ -161,6 +162,12 @@ struct FRDGImportedBuffer
 RENDERER_API FRDGImportedBuffer Register(FRDGBuilder& GraphBuilder, const FRDGExternalBuffer& In, ERDGImportedBufferFlags Flags, ERDGUnorderedAccessViewFlags UAVFlags = ERDGUnorderedAccessViewFlags::None);
 RENDERER_API FRDGBufferSRVRef   RegisterAsSRV(FRDGBuilder& GraphBuilder, const FRDGExternalBuffer& In);
 RENDERER_API FRDGBufferUAVRef   RegisterAsUAV(FRDGBuilder& GraphBuilder, const FRDGExternalBuffer& In, ERDGUnorderedAccessViewFlags Flags = ERDGUnorderedAccessViewFlags::None);
+
+RENDERER_API void AddTransitionPass(
+	FRDGBuilder& GraphBuilder,
+	FGlobalShaderMap* ShaderMap,
+	const TArray<FRDGBufferSRVRef>& Transitions);
+
 ////////////////////////////////////////////////////////////////////////////////////////////////////
 // Misc/Helpers
 
