Add wave permutation to hair tile classification.

Save 25% of the cost 0.15ms ->0.11ms

diff --git a/Engine/Shaders/Private/HairStrands/HairStrandsVisibilityTile.usf b/Engine/Shaders/Private/HairStrands/HairStrandsVisibilityTile.usf
--- a/Engine/Shaders/Private/HairStrands/HairStrandsVisibilityTile.usf
+++ b/Engine/Shaders/Private/HairStrands/HairStrandsVisibilityTile.usf
@@ -15,8 +15,13 @@ uint bUintTexture;
 uint bForceOutputAllTiles;
 float TransmittanceThreshold;
 uint IntCoverageThreshold;
+#if PERMUTATION_INPUT_TYPE == 0
 Texture2D<float> InputFloatTexture;
+#elif PERMUTATION_INPUT_TYPE == 1
 Texture2D<uint> InputUintTexture;
+#else
+	#error Input type not defined
+#endif
 RWBuffer<uint> TileCountBuffer;
 
 RWBuffer<uint2> TileHairAllBuffer;
@@ -24,10 +29,15 @@ RWBuffer<uint2> TileHairFullBuffer;
 RWBuffer<uint2> TileHairPartialBuffer;
 RWBuffer<uint2> TileOtherBuffer;
 
+#if PERMUTATION_WAVEOPS == 0
 // 8x8 bit mask
 groupshared uint2 s_fMask[GROUP_THREAD_COUNT];
 groupshared uint2 s_pMask[GROUP_THREAD_COUNT];
+#endif
 
+#if PERMUTATION_WAVE_OPS && COMPILER_SUPPORTS_WAVE_SIZE
+WAVESIZE(64) // PERMUTATION_WAVE_OPS is true only when wave>=64 are available
+#endif
 [numthreads(HAIR_TILE_SIZE, HAIR_TILE_SIZE, 1)]
 void TileMainCS(uint2 DispatchThreadId : SV_DispatchThreadID, uint LinearIndex : SV_GroupIndex, uint3 GroupId : SV_GroupID)
 {
@@ -44,20 +54,52 @@ void TileMainCS(uint2 DispatchThreadId : SV_DispatchThreadID, uint LinearIndex :
 			fHair = true;
 			pHair = false;
 		}
-		else if (bUintTexture)
+		else
+		#if PERMUTATION_INPUT_TYPE == 0
+		{
+			const float Transmittance = InputFloatTexture.Load(uint3(PixelCoord, 0));
+			fHair = Transmittance <= TransmittanceThreshold;
+			pHair = Transmittance < 1.f;
+		}
+		#elif PERMUTATION_INPUT_TYPE == 1
 		{
 			const uint Coverage = InputUintTexture.Load(uint3(PixelCoord, 0));
 			fHair = Coverage >= IntCoverageThreshold;
 			pHair = Coverage > 0;//&& Coverage < IntCoverageThreshold; //this is now consistent with non-int branch below
 		}
+		#endif
+	}
+#if PERMUTATION_WAVEOPS
+	const bool bIsFull    = WaveActiveAllTrue(fHair);
+	const bool bIsPartial = WaveActiveAnyTrue(pHair);
+	if (all(LinearIndex == 0)) 
+	{
+		if (bIsFull || bIsPartial)
+		{
+			uint WriteToIndex;
+			InterlockedAdd(TileCountBuffer[HAIRTILE_HAIR_ALL], 1, WriteToIndex);
+			TileHairAllBuffer[WriteToIndex] = GroupId.xy;
+		}
+		if (bIsFull)
+		{
+			uint WriteToIndex;
+			InterlockedAdd(TileCountBuffer[HAIRTILE_HAIR_FULL], 1, WriteToIndex);
+			TileHairFullBuffer[WriteToIndex] = GroupId.xy;
+		}
+		else if (bIsPartial)
+		{
+			uint WriteToIndex;
+			InterlockedAdd(TileCountBuffer[HAIRTILE_HAIR_PARTIAL], 1, WriteToIndex);
+			TileHairPartialBuffer[WriteToIndex] = GroupId.xy;
+		}
 		else
 		{
-			const float Transmittance = InputFloatTexture.Load(uint3(PixelCoord, 0));
-			fHair = Transmittance <= TransmittanceThreshold;
-			pHair = Transmittance < 1.f;
+			uint WriteToIndex;
+			InterlockedAdd(TileCountBuffer[HAIRTILE_OTHER], 1, WriteToIndex);
+			TileOtherBuffer[WriteToIndex] = GroupId.xy;
 		}
 	}
-
+#else
 	const uint2 Mask = LinearIndex < 32 ? uint2(1u << LinearIndex, 0u) : uint2(0u, 1u << (LinearIndex - 32u));
 	s_fMask[LinearIndex] = fHair ? Mask : 0u;
 	s_pMask[LinearIndex] = pHair ? Mask : 0u;
@@ -128,6 +170,7 @@ void TileMainCS(uint2 DispatchThreadId : SV_DispatchThreadID, uint LinearIndex :
 			TileOtherBuffer[WriteToIndex] = GroupId.xy;
 		}
 	}
+#endif
 }
 #endif // SHADER_TILE_GENERATION
 
diff --git a/Engine/Source/Runtime/Renderer/Private/HairStrands/HairStrandsTile.cpp b/Engine/Source/Runtime/Renderer/Private/HairStrands/HairStrandsTile.cpp
--- a/Engine/Source/Runtime/Renderer/Private/HairStrands/HairStrandsTile.cpp
+++ b/Engine/Source/Runtime/Renderer/Private/HairStrands/HairStrandsTile.cpp
@@ -46,6 +46,17 @@ FHairStrandsTilePassVS::FParameters GetHairStrandsTileParameters(const FViewInfo
 	return Out;
 }
 
+static ERHIFeatureSupport HairSupportsWaveOps(EShaderPlatform Platform)
+{
+	// D3D11 / SM5 or preview do not support, or work well with, wave-ops by default (or SM5 preview has issues with wave intrinsics too), that fixes classification and black/wrong tiling.
+	if (Platform == SP_PCD3D_SM5 || FDataDrivenShaderPlatformInfo::GetIsPreviewPlatform(Platform))
+	{
+		return ERHIFeatureSupport::Unsupported;
+	}
+
+	return FDataDrivenShaderPlatformInfo::GetSupportsWaveOperations(Platform);
+}
+
 ////////////////////////////////////////////////////////////////////////////////////////////////////////////
 // Generate indirect draw and indirect dispatch buffers
 
@@ -116,12 +127,13 @@ class FHairStrandsTileGenerationPassCS : public FGlobalShader
 	DECLARE_GLOBAL_SHADER(FHairStrandsTileGenerationPassCS);
 	SHADER_USE_PARAMETER_STRUCT(FHairStrandsTileGenerationPassCS, FGlobalShader);
 
-	using FPermutationDomain = TShaderPermutationDomain<>;
+	class FInputType : SHADER_PERMUTATION_INT("PERMUTATION_INPUT_TYPE", 2);
+	class FWaveOps : SHADER_PERMUTATION_BOOL("PERMUTATION_WAVEOPS");
+	using FPermutationDomain = TShaderPermutationDomain<FInputType, FWaveOps>;
 
 	BEGIN_SHADER_PARAMETER_STRUCT(FParameters, )
 		SHADER_PARAMETER_STRUCT_REF(FViewUniformShaderParameters, ViewUniformBuffer)
 		SHADER_PARAMETER(FIntPoint, BufferResolution)
-		SHADER_PARAMETER(uint32, bUintTexture)
 		SHADER_PARAMETER(uint32, bForceOutputAllTiles)
 		SHADER_PARAMETER(float, TransmittanceThreshold)
 		SHADER_PARAMETER(uint32, IntCoverageThreshold)
@@ -137,6 +149,11 @@ class FHairStrandsTileGenerationPassCS : public FGlobalShader
 
 	static bool ShouldCompilePermutation(const FGlobalShaderPermutationParameters& Parameters)
 	{
+		FPermutationDomain PermutationVector(Parameters.PermutationId);
+		if (PermutationVector.Get<FWaveOps>() && HairSupportsWaveOps(Parameters.Platform) == ERHIFeatureSupport::Unsupported)
+		{
+			return false;
+		}
 		return IsHairStrandsSupported(EHairStrandsShaderType::Strands, Parameters.Platform);
 	}
 
@@ -144,6 +161,12 @@ class FHairStrandsTileGenerationPassCS : public FGlobalShader
 	{
 		FGlobalShader::ModifyCompilationEnvironment(Parameters, OutEnvironment);
 		OutEnvironment.SetDefine(TEXT("SHADER_TILE_GENERATION"), 1);
+
+		FPermutationDomain PermutationVector(Parameters.PermutationId);
+		if (PermutationVector.Get<FWaveOps>())
+		{
+			OutEnvironment.CompilerFlags.Add(CFLAG_WaveOperations);
+		}
 	}
 };
 IMPLEMENT_GLOBAL_SHADER(FHairStrandsTileGenerationPassCS, "/Engine/Private/HairStrands/HairStrandsVisibilityTile.usf", "TileMainCS", SF_Compute);
@@ -159,6 +182,7 @@ static FHairStrandsTiles AddHairStrandsGenerateTilesPass_Internal(
 {
 	const bool bHasValidInput = InputTexture != nullptr;
 	const bool bUintTexture   = InputTexture && InputTexture->Desc.Format == PF_R32_UINT;
+	const bool bWaveOps       = GRHISupportsWaveOperations&& GRHIMaximumWaveSize >= 64 && HairSupportsWaveOps(View.GetShaderPlatform()) != ERHIFeatureSupport::Unsupported;
 
 	FHairStrandsTiles Out;
 
@@ -181,12 +205,15 @@ static FHairStrandsTiles AddHairStrandsGenerateTilesPass_Internal(
 	FRDGBufferUAVRef TileCountUAV = GraphBuilder.CreateUAV(Out.TileCountBuffer, PF_R32_UINT);
 	AddClearUAVPass(GraphBuilder, TileCountUAV, 0u);
 
-	TShaderMapRef<FHairStrandsTileGenerationPassCS> ComputeShader(View.ShaderMap);
+	FHairStrandsTileGenerationPassCS::FPermutationDomain PermutationVector;
+	PermutationVector.Set<FHairStrandsTileGenerationPassCS::FInputType>(bUintTexture ? 1u : 0u);
+	PermutationVector.Set<FHairStrandsTileGenerationPassCS::FWaveOps>(bWaveOps);
+
+	TShaderMapRef<FHairStrandsTileGenerationPassCS> ComputeShader(View.ShaderMap, PermutationVector);
 	FHairStrandsTileGenerationPassCS::FParameters* PassParameters = GraphBuilder.AllocParameters<FHairStrandsTileGenerationPassCS::FParameters>();
 	PassParameters->ViewUniformBuffer		= View.ViewUniformBuffer;
 	PassParameters->BufferResolution		= View.ViewRect.Size();
 	PassParameters->bForceOutputAllTiles	= bHasValidInput ? 0 : 1;
-	PassParameters->bUintTexture			= bUintTexture ? 1u : 0u;
 	PassParameters->TransmittanceThreshold	= 1.f - GetHairStrandsFullCoverageThreshold();	
 	PassParameters->IntCoverageThreshold 	= GetHairStrandsIntCoverageThreshold();
 	PassParameters->InputFloatTexture		= GSystemTextures.GetBlackDummy(GraphBuilder);
