* Remove depth compresion on hair vis buffer
* Hair vis buffer is now only partially clear
* Change hair sample offset/count texture & coverate texture to use tiled clear
* Change HitProxy and SelectionOutline to use tile data to avoid sampling uninitialized data

Save hair vis cost. 1.79ms -> 1.15ms

diff --git a/Engine/Shaders/Private/HairStrands/HairStrandsDebug.usf b/Engine/Shaders/Private/HairStrands/HairStrandsDebug.usf
--- a/Engine/Shaders/Private/HairStrands/HairStrandsDebug.usf
+++ b/Engine/Shaders/Private/HairStrands/HairStrandsDebug.usf
@@ -129,6 +129,12 @@ void MainPS(
 	if (!bIsValid)
 		discard;
 
+	const bool bIsDepthValid = HairStrands.HairOnlyDepthTexture.Load(uint3(PixelPos, 0)) > 0;
+	if (!bIsDepthValid)
+	{
+		discard;
+	}
+
 	const FNodeDesc NodeDesc = DecodeNodeDesc(HairStrands.HairSampleOffset.Load(uint3(PixelPos, 0)));
 
 	float3 Color = 0;
diff --git a/Engine/Shaders/Private/HairStrands/HairStrandsDeepShadowMask.usf b/Engine/Shaders/Private/HairStrands/HairStrandsDeepShadowMask.usf
--- a/Engine/Shaders/Private/HairStrands/HairStrandsDeepShadowMask.usf
+++ b/Engine/Shaders/Private/HairStrands/HairStrandsDeepShadowMask.usf
@@ -338,7 +338,7 @@ void MainPS(
 	const float2 UV = Input.UV;
 	const uint2 PixelCoord = floor(Input.Position.xy);
 	const float SceneDepth = ConvertFromDeviceZ(SceneDepthTexture.Load(uint3(PixelCoord, 0)).x);
-	const bool bIsHairPixel = HairStrands.HairCoverageTexture.Load(uint3(PixelCoord, 0)) >= 1;
+	const bool bIsHairPixel = HairStrands.HairOnlyDepthTexture.Load(uint3(PixelCoord, 0)) > 0;
 
 	const float3 TranslatedWorldPosition = GetTranslatedWorldPosition(UV, SceneDepth);
 	float4 AtlasScaleBias;
@@ -445,8 +445,9 @@ void MainPS(
 
 	const float DistanceThreshold = 100000.0f;
 	const float CoverageThreshold = 0.995f; // When Coverage is high, we do not trace shadow on opaque since hair/fur is covering the background.
+	const bool  bIsHairPixel = HairStrands.HairOnlyDepthTexture.Load(uint3(PixelCoord, 0)) > 0;
 	const float HairPixelCoverage = HairStrands.HairCoverageTexture.Load(uint3(PixelCoord, 0));
-	const bool bIsValidOpaquePixel = HairPixelCoverage < CoverageThreshold;
+	const bool bIsValidOpaquePixel = bIsHairPixel && HairPixelCoverage < CoverageThreshold;
 	
 	float3 TranslatedWorldPosition = GetTranslatedWorldPosition(UV, SceneDepth);
 	if (View.StereoPassIndex == 1)
diff --git a/Engine/Shaders/Private/HairStrands/HairStrandsEnvironmentAO.usf b/Engine/Shaders/Private/HairStrands/HairStrandsEnvironmentAO.usf
--- a/Engine/Shaders/Private/HairStrands/HairStrandsEnvironmentAO.usf
+++ b/Engine/Shaders/Private/HairStrands/HairStrandsEnvironmentAO.usf
@@ -132,20 +132,9 @@ void MainPS(
 	OutAO = 1;
 
 	// Trace AO for all pixel which are partially covered or not covered at all.
+	// Fully covered pixel are marked as unlit
 	const uint2 PixelCoord = floor(SvPosition.xy);
-	bool bIsValid = false;
-	if (bHalfRes)
-	{
-		const float Coverage0 = HairStrands.HairCoverageTexture.Load(uint3(PixelCoord              , 0));
-		const float Coverage1 = HairStrands.HairCoverageTexture.Load(uint3(PixelCoord + uint2(1, 0), 0));
-		const float Coverage2 = HairStrands.HairCoverageTexture.Load(uint3(PixelCoord + uint2(0, 1), 0));
-		const float Coverage3 = HairStrands.HairCoverageTexture.Load(uint3(PixelCoord + uint2(1, 1), 0));
-		bIsValid = Coverage0 < 1 || Coverage1 < 1 || Coverage2 < 1 || Coverage3 < 1 || GBuffer.ShadingModelID != SHADINGMODELID_UNLIT;
-	}
-	else
-	{
-		bIsValid = HairStrands.HairCoverageTexture.Load(uint3(PixelCoord, 0)) < 1;
-	}
+	const bool bIsValid = GBuffer.ShadingModelID != SHADINGMODELID_UNLIT;
 
 	if (bIsValid)
 	{
diff --git a/Engine/Shaders/Private/HairStrands/HairStrandsHitProxy.usf b/Engine/Shaders/Private/HairStrands/HairStrandsHitProxy.usf
--- a/Engine/Shaders/Private/HairStrands/HairStrandsHitProxy.usf
+++ b/Engine/Shaders/Private/HairStrands/HairStrandsHitProxy.usf
@@ -16,11 +16,11 @@ StructuredBuffer<FPackedHairVis> VisNodeData;
 Texture2D<float>				 CoverageTexture;
 
 void EmitPS(
-	in float4 SvPosition : SV_Position,
+	in FScreenVertexOutput In,
 	out float4 OutHitProxyId : SV_Target0,
 	out float OutDepth : SV_Depth)
 {
-	const uint2 PixelCoord = (uint2)SvPosition.xy;
+	const uint2 PixelCoord = (uint2)In.Position.xy;
 
 	// Take the first sample
 	const FNodeDesc NodeDesc = DecodeNodeDesc(VisNodeIndex.Load(uint3(PixelCoord,0)));
@@ -61,10 +61,10 @@ StructuredBuffer<FPackedHairVis> VisNodeData;
 Texture2D<float>				 CoverageTexture;
 
 void EmitPS(
-	in float4 SvPosition : SV_Position,
+	in FScreenVertexOutput In,
 	out float OutDepth : SV_Depth)
 {
-	const float2 UV = SvPosition.xy * InvViewportResolution;
+	const float2 UV = In.Position.xy * InvViewportResolution;
 	const uint2 PixelCoord = UV * View.ViewSizeAndInvSize.xy;
 
 	// Take the first sample
diff --git a/Engine/Shaders/Private/HairStrands/HairStrandsVisibilityClearPS.usf b/Engine/Shaders/Private/HairStrands/HairStrandsVisibilityClearPS.usf
--- a/Engine/Shaders/Private/HairStrands/HairStrandsVisibilityClearPS.usf
+++ b/Engine/Shaders/Private/HairStrands/HairStrandsVisibilityClearPS.usf
@@ -3,19 +3,24 @@
 #include "/Engine/Public/Platform.ush"
 
 
-uint ClearValue;
+uint uClearValue; 
+float fClearValue;
 
 void ClearPS(
 #if PERMUTATION_OUTPUT_FORMAT == 0
 	out uint OutColor : SV_Target0
 #elif PERMUTATION_OUTPUT_FORMAT == 1
 	out uint2 OutColor : SV_Target0
+#elif PERMUTATION_OUTPUT_FORMAT == 2
+	out float OutColor : SV_Target0
 #endif
 )
 {
 #if PERMUTATION_OUTPUT_FORMAT == 0
-	OutColor = ClearValue;
+	OutColor = uClearValue;
 #elif PERMUTATION_OUTPUT_FORMAT == 1
-	OutColor = ClearValue.xx;
+	OutColor = uClearValue.xx;
+#elif PERMUTATION_OUTPUT_FORMAT == 2
+	OutColor = fClearValue;
 #endif
 }
\ No newline at end of file
diff --git a/Engine/Source/Runtime/Renderer/Private/HairStrands/HairStrandsVisibility.cpp b/Engine/Source/Runtime/Renderer/Private/HairStrands/HairStrandsVisibility.cpp
--- a/Engine/Source/Runtime/Renderer/Private/HairStrands/HairStrandsVisibility.cpp
+++ b/Engine/Source/Runtime/Renderer/Private/HairStrands/HairStrandsVisibility.cpp
@@ -1576,11 +1576,12 @@ class FClearUIntGraphicPS : public FGlobalShader
 	DECLARE_GLOBAL_SHADER(FClearUIntGraphicPS);
 	SHADER_USE_PARAMETER_STRUCT(FClearUIntGraphicPS, FGlobalShader);
 
-	class FOutputFormat : SHADER_PERMUTATION_INT("PERMUTATION_OUTPUT_FORMAT", 2);
+	class FOutputFormat : SHADER_PERMUTATION_INT("PERMUTATION_OUTPUT_FORMAT", 3);
 	using FPermutationDomain = TShaderPermutationDomain<FOutputFormat>;
 
 	BEGIN_SHADER_PARAMETER_STRUCT(FParameters, )
-		SHADER_PARAMETER(uint32, ClearValue)
+		SHADER_PARAMETER(uint32, uClearValue)
+		SHADER_PARAMETER(float, fClearValue)
 		SHADER_PARAMETER_STRUCT_INCLUDE(FHairStrandsTilePassVS::FParameters, TileData)
 		RENDER_TARGET_BINDING_SLOTS()
 	END_SHADER_PARAMETER_STRUCT()
@@ -1618,7 +1619,8 @@ static void AddClearGraphicPass(
 	const FHairStrandsTiles::ETileType TileType = FHairStrandsTiles::ETileType::HairAll;
 
 	FClearUIntGraphicPS::FParameters* Parameters = GraphBuilder.AllocParameters<FClearUIntGraphicPS::FParameters>();
-	Parameters->ClearValue = ClearValue;
+	Parameters->uClearValue = ClearValue;
+	Parameters->fClearValue = ClearValue;
 	Parameters->TileData = GetHairStrandsTileParameters(*View, TileData, TileType);
 	Parameters->RenderTargets[0] = FRenderTargetBinding(OutTarget, ERenderTargetLoadAction::ENoAction, 0);
 
@@ -1631,6 +1633,10 @@ static void AddClearGraphicPass(
 	{
 		PermutationVector.Set<FClearUIntGraphicPS::FOutputFormat>(1);
 	}
+	else
+	{
+		PermutationVector.Set<FClearUIntGraphicPS::FOutputFormat>(2);
+	}
 
 	TShaderMapRef<FHairStrandsTilePassVS> TileVertexShader(View->ShaderMap);
 	TShaderMapRef<FClearUIntGraphicPS> PixelShader(View->ShaderMap, PermutationVector);
@@ -1831,19 +1837,19 @@ static void AddHairVisibilityControlPointIdCompactionPass(
 	}
 
 	{
-		FRDGTextureDesc Desc = FRDGTextureDesc::Create2D(Resolution, PF_R32_UINT, FClearValueBinding::None, TexCreate_UAV | TexCreate_ShaderResource);
+		FRDGTextureDesc Desc = FRDGTextureDesc::Create2D(Resolution, PF_R32_UINT, FClearValueBinding::None, TexCreate_UAV | TexCreate_ShaderResource | TexCreate_RenderTargetable);
 		OutCompactNodeIndex = GraphBuilder.CreateTexture(Desc, TEXT("Hair.VisibilityCompactNodeIndex"));
 	}
 
 	{
-		FRDGTextureDesc Desc = FRDGTextureDesc::Create2D(Resolution, FHairStrandsVisibilityData::CoverageFormat, FClearValueBinding::None, TexCreate_UAV | TexCreate_ShaderResource);
+		FRDGTextureDesc Desc = FRDGTextureDesc::Create2D(Resolution, FHairStrandsVisibilityData::CoverageFormat, FClearValueBinding::None, TexCreate_UAV | TexCreate_ShaderResource | TexCreate_RenderTargetable);
 		OutCoverageTexture = GraphBuilder.CreateTexture(Desc, TEXT("Hair.CoverageTexture"));
 	}
 
 	const uint32 ClearValues[4] = { 0u,0u,0u,0u };
 	AddClearUAVPass(GraphBuilder, GraphBuilder.CreateUAV(OutCompactCounter), ClearValues);
-	AddClearUAVPass(GraphBuilder, GraphBuilder.CreateUAV(OutCompactNodeIndex), ClearValues);
-	AddClearUAVPass(GraphBuilder, GraphBuilder.CreateUAV(OutCoverageTexture), 0.f);
+	AddClearGraphicPass(GraphBuilder, RDG_EVENT_NAME("HairStrands::NodeOffsetAndCount"), &View, 0, TileData, OutCompactNodeIndex);
+	AddClearGraphicPass(GraphBuilder, RDG_EVENT_NAME("HairStrands::CoverageTexture"), &View, 0, TileData, OutCoverageTexture);
 
 	// Adapt the buffer allocation based on the bounding box of the hair macro groups. This allows to reduce the overall allocation size
 	const FIntRect HairRect = ComputeVisibleHairStrandsMacroGroupsRect(View.ViewRect, MacroGroupDatas);
@@ -1976,12 +1982,12 @@ static void AddHairVisibilityCompactionComputeRasterPass(
 	}
 
 	{
-		FRDGTextureDesc Desc = FRDGTextureDesc::Create2D(Resolution, PF_R32_UINT, FClearValueBinding::None, TexCreate_UAV | TexCreate_ShaderResource);
+		FRDGTextureDesc Desc = FRDGTextureDesc::Create2D(Resolution, PF_R32_UINT, FClearValueBinding::None, TexCreate_UAV | TexCreate_ShaderResource | TexCreate_RenderTargetable);
 		OutCompactNodeIndex = GraphBuilder.CreateTexture(Desc, TEXT("Hair.VisibilityCompactNodeIndex"));
 	}
 
 	{
-		FRDGTextureDesc Desc = FRDGTextureDesc::Create2D(Resolution, FHairStrandsVisibilityData::CoverageFormat, FClearValueBinding::None, TexCreate_UAV | TexCreate_ShaderResource);
+		FRDGTextureDesc Desc = FRDGTextureDesc::Create2D(Resolution, FHairStrandsVisibilityData::CoverageFormat, FClearValueBinding::None, TexCreate_UAV | TexCreate_ShaderResource | TexCreate_RenderTargetable);
 		OutCoverageTexture = GraphBuilder.CreateTexture(Desc, TEXT("Hair.CoverageTexture"));
 	}
 
@@ -2077,7 +2083,7 @@ static FRDGTextureRef AddHairVisibilityFillOpaqueDepth(
 	check(TileData.IsValid());
 
 	const FHairStrandsTiles::ETileType TileType = FHairStrandsTiles::ETileType::HairAll;
-	FRDGTextureRef OutVisibilityDepthTexture = GraphBuilder.CreateTexture(FRDGTextureDesc::Create2D(Resolution, PF_D24, FClearValueBinding::DepthFar, TexCreate_DepthStencilTargetable | TexCreate_ShaderResource, 1, GetMaxSamplePerPixel(View.GetShaderPlatform())), TEXT("Hair.VisibilityDepthTexture"));
+	FRDGTextureRef OutVisibilityDepthTexture = GraphBuilder.CreateTexture(FRDGTextureDesc::Create2D(Resolution, PF_D24, FClearValueBinding::DepthFar, TexCreate_NoFastClear | TexCreate_DepthStencilTargetable | TexCreate_ShaderResource, 1, GetMaxSamplePerPixel(View.GetShaderPlatform())), TEXT("Hair.VisibilityDepthTexture"));
 
 	FHairVisibilityFillOpaqueDepthPS::FParameters* Parameters = GraphBuilder.AllocParameters<FHairVisibilityFillOpaqueDepthPS::FParameters>();
 	Parameters->TileData = GetHairStrandsTileParameters(View, TileData, TileType);
@@ -2085,7 +2091,7 @@ static FRDGTextureRef AddHairVisibilityFillOpaqueDepth(
 	Parameters->SceneDepthTexture = SceneDepthTexture;
 	Parameters->RenderTargets.DepthStencil = FDepthStencilBinding(
 		OutVisibilityDepthTexture,
-		ERenderTargetLoadAction::EClear,
+		ERenderTargetLoadAction::ENoAction,
 		ERenderTargetLoadAction::ENoAction,
 		FExclusiveDepthStencil::DepthWrite_StencilNop);
 
@@ -2107,7 +2113,7 @@ static FRDGTextureRef AddHairVisibilityFillOpaqueDepth(
 			RHICmdList.ApplyCachedRenderTargets(GraphicsPSOInit);
 			GraphicsPSOInit.BlendState = TStaticBlendState<CW_RGBA, BO_Add, BF_One, BF_Zero, BO_Add, BF_One, BF_Zero>::GetRHI();
 			GraphicsPSOInit.RasterizerState = TStaticRasterizerState<>::GetRHI();
-			GraphicsPSOInit.DepthStencilState = TStaticDepthStencilState<true, CF_DepthNearOrEqual>::GetRHI();
+			GraphicsPSOInit.DepthStencilState = TStaticDepthStencilState<true, CF_Always>::GetRHI();
 
 			GraphicsPSOInit.BoundShaderState.VertexDeclarationRHI = GFilterVertexDeclaration.VertexDeclarationRHI;
 			GraphicsPSOInit.BoundShaderState.VertexShaderRHI = TileVertexShader.GetVertexShader();
@@ -3833,6 +3839,7 @@ class FHairStrandsEmitSelectionPS : public FGlobalShader
 		SHADER_PARAMETER_RDG_TEXTURE(Texture2D<float>, CoverageTexture)
 		SHADER_PARAMETER_RDG_BUFFER_SRV(StructuredBuffer<FPackedHairVis>, VisNodeData)
 		SHADER_PARAMETER_RDG_BUFFER_SRV(Buffer<uint>, SelectionMaterialIdBuffer)
+		SHADER_PARAMETER_STRUCT_INCLUDE(FHairStrandsTilePassVS::FParameters, TileData)
 		RENDER_TARGET_BINDING_SLOTS()
 	END_SHADER_PARAMETER_STRUCT()
 
@@ -3854,6 +3861,7 @@ void AddHairStrandsSelectionOutlinePass(
 	const FViewInfo& View,
 	const FIntRect& ViewportRect,
 	const FHairStrandsMacroGroupDatas& MacroGroupDatas,
+	FHairStrandsTiles TileData,
 	FRDGTextureRef VisNodeIndex,
 	FRDGBufferRef VisNodeData,
 	FRDGTextureRef CoverageTexture,
@@ -3880,6 +3888,8 @@ void AddHairStrandsSelectionOutlinePass(
 		}
 	}
 
+	const FHairStrandsTiles::ETileType TileType = FHairStrandsTiles::ETileType::HairAll;
+
 	FRDGBufferRef SelectionMaterialIdBuffer = CreateUploadBuffer(GraphBuilder, TEXT("Hair.MaterialIdToHitProxyIdBuffer"), sizeof(uint32), SelectionMaterialId.Num(), SelectionMaterialId.GetData(), sizeof(uint32) * SelectionMaterialId.Num());
 	auto* PassParameters = GraphBuilder.AllocParameters<FHairStrandsEmitSelectionPS::FParameters>();
 	PassParameters->View = View.ViewUniformBuffer;
@@ -3891,21 +3901,41 @@ void AddHairStrandsSelectionOutlinePass(
 	PassParameters->CoverageTexture = CoverageTexture;
 	PassParameters->SelectionMaterialIdBuffer = GraphBuilder.CreateSRV(SelectionMaterialIdBuffer, PF_R32_UINT);
 	PassParameters->RenderTargets.DepthStencil = FDepthStencilBinding(SelectionDepthTexture, ERenderTargetLoadAction::ELoad, FExclusiveDepthStencil::DepthWrite_StencilWrite);
+	PassParameters->TileData = GetHairStrandsTileParameters(View, TileData, TileType);
 
+	const FIntRect Viewport = View.ViewRect;
 	auto PixelShader = View.ShaderMap->GetShader<FHairStrandsEmitSelectionPS>();
+	TShaderMapRef<FHairStrandsTilePassVS> TileVertexShader(View.ShaderMap);
 
 	const uint32 StencilRef = 3;
-	FPixelShaderUtils::AddFullscreenPass(
-		GraphBuilder,
-		View.ShaderMap,
-		RDG_EVENT_NAME("HairStrands::EmitSelection"),
-		PixelShader,
+	
+	GraphBuilder.AddPass(
+		RDG_EVENT_NAME("HairStrands::EmitSelection(Tile)"),
 		PassParameters,
-		ViewportRect,
-		TStaticBlendState<>::GetRHI(),
-		TStaticRasterizerState<>::GetRHI(),
-		TStaticDepthStencilState<true, CF_DepthNearOrEqual, true, CF_Always, SO_Keep, SO_Keep, SO_Replace>::GetRHI(), 
-		StencilRef);
+		ERDGPassFlags::Raster,
+		[PassParameters, TileVertexShader, PixelShader, Viewport, TileType, StencilRef](FRHICommandList& RHICmdList)
+		{
+			FHairStrandsTilePassVS::FParameters ParametersVS = PassParameters->TileData;
+
+			FGraphicsPipelineStateInitializer GraphicsPSOInit;
+			RHICmdList.ApplyCachedRenderTargets(GraphicsPSOInit);
+			GraphicsPSOInit.BlendState = TStaticBlendState<>::GetRHI();
+			GraphicsPSOInit.RasterizerState = TStaticRasterizerState<>::GetRHI();
+			GraphicsPSOInit.DepthStencilState = TStaticDepthStencilState<true, CF_DepthNearOrEqual, true, CF_Always, SO_Keep, SO_Keep, SO_Replace>::GetRHI();
+
+			GraphicsPSOInit.BoundShaderState.VertexDeclarationRHI = GFilterVertexDeclaration.VertexDeclarationRHI;
+			GraphicsPSOInit.BoundShaderState.VertexShaderRHI = TileVertexShader.GetVertexShader();
+			GraphicsPSOInit.BoundShaderState.PixelShaderRHI = PixelShader.GetPixelShader();
+			GraphicsPSOInit.PrimitiveType = PassParameters->TileData.bRectPrimitive > 0 ? PT_RectList : PT_TriangleList;
+			SetGraphicsPipelineState(RHICmdList, GraphicsPSOInit, StencilRef);
+
+			SetShaderParameters(RHICmdList, PixelShader, PixelShader.GetPixelShader(), *PassParameters);
+
+			RHICmdList.SetViewport(Viewport.Min.X, Viewport.Min.Y, 0.0f, Viewport.Max.X, Viewport.Max.Y, 1.0f);
+			SetShaderParameters(RHICmdList, TileVertexShader, TileVertexShader.GetVertexShader(), ParametersVS);
+			RHICmdList.SetStreamSource(0, nullptr, 0);
+			RHICmdList.DrawPrimitiveIndirect(PassParameters->TileData.TileIndirectBuffer->GetRHI(), FHairStrandsTiles::GetIndirectDrawArgOffset(TileType));
+		});
 #endif
 }
 
@@ -3925,6 +3955,7 @@ class FHairStrandsEmitHitProxyIdPS : public FGlobalShader
 		SHADER_PARAMETER_RDG_BUFFER_SRV(StructuredBuffer<FPackedHairVis>, VisNodeData)
 		SHADER_PARAMETER_RDG_BUFFER_SRV(Buffer<uint>, MaterialIdToHitProxyIdBuffer)
 		SHADER_PARAMETER_RDG_TEXTURE(Texture2D<float>, CoverageTexture)
+		SHADER_PARAMETER_STRUCT_INCLUDE(FHairStrandsTilePassVS::FParameters, TileData)
 		RENDER_TARGET_BINDING_SLOTS()
 	END_SHADER_PARAMETER_STRUCT()
 
@@ -3946,6 +3977,7 @@ void AddHairStrandsHitProxyIdPass(
 	const FScene& Scene,
 	const FViewInfo& View,
 	const FHairStrandsMacroGroupDatas& MacroGroupDatas,
+	FHairStrandsTiles TileData,
 	FRDGTextureRef VisNodeIndex,
 	FRDGBufferRef VisNodeData,
 	FRDGTextureRef CoverageTexture,
@@ -3976,6 +4008,8 @@ void AddHairStrandsHitProxyIdPass(
 		}
 	}
 
+	const FHairStrandsTiles::ETileType TileType = FHairStrandsTiles::ETileType::HairAll;
+
 	FRDGBufferRef MaterialIdToHitProxyIdBuffer = CreateUploadBuffer(GraphBuilder, TEXT("Hair.MaterialIdToHitProxyIdBuffer"), sizeof(uint32), MaterialIdToHitProxyId.Num(), MaterialIdToHitProxyId.GetData(), sizeof(uint32) * MaterialIdToHitProxyId.Num());
 	auto* PassParameters = GraphBuilder.AllocParameters<FHairStrandsEmitHitProxyIdPS::FParameters>();
 	PassParameters->View = View.ViewUniformBuffer;
@@ -3987,19 +4021,39 @@ void AddHairStrandsHitProxyIdPass(
 	PassParameters->CoverageTexture = CoverageTexture;
 	PassParameters->RenderTargets[0] = FRenderTargetBinding(HitProxyTexture, ERenderTargetLoadAction::ELoad);
 	PassParameters->RenderTargets.DepthStencil = FDepthStencilBinding(HitProxyDepthTexture, ERenderTargetLoadAction::ELoad, FExclusiveDepthStencil::DepthWrite_StencilWrite);
+	PassParameters->TileData = GetHairStrandsTileParameters(View, TileData, TileType);
 
+	const FIntRect Viewport = View.ViewRect;
 	auto PixelShader = View.ShaderMap->GetShader<FHairStrandsEmitHitProxyIdPS>();
+	TShaderMapRef<FHairStrandsTilePassVS> TileVertexShader(View.ShaderMap);
 
-	FPixelShaderUtils::AddFullscreenPass(
-		GraphBuilder,
-		View.ShaderMap,
-		RDG_EVENT_NAME("HairStrands::EmitHitProxyId"),
-		PixelShader,
+	GraphBuilder.AddPass(
+		RDG_EVENT_NAME("HairStrands::EmitHitProxyId(Tile)"),
 		PassParameters,
-		View.ViewRect,
-		TStaticBlendState<>::GetRHI(),
-		TStaticRasterizerState<>::GetRHI(),
-		TStaticDepthStencilState<true, CF_DepthNearOrEqual>::GetRHI());
+		ERDGPassFlags::Raster,
+		[PassParameters, TileVertexShader, PixelShader, Viewport, TileType](FRHICommandList& RHICmdList)
+		{
+			FHairStrandsTilePassVS::FParameters ParametersVS = PassParameters->TileData;
+
+			FGraphicsPipelineStateInitializer GraphicsPSOInit;
+			RHICmdList.ApplyCachedRenderTargets(GraphicsPSOInit);
+			GraphicsPSOInit.BlendState = TStaticBlendState<>::GetRHI();
+			GraphicsPSOInit.RasterizerState = TStaticRasterizerState<>::GetRHI();
+			GraphicsPSOInit.DepthStencilState = TStaticDepthStencilState<true, CF_DepthNearOrEqual>::GetRHI();
+
+			GraphicsPSOInit.BoundShaderState.VertexDeclarationRHI = GFilterVertexDeclaration.VertexDeclarationRHI;
+			GraphicsPSOInit.BoundShaderState.VertexShaderRHI = TileVertexShader.GetVertexShader();
+			GraphicsPSOInit.BoundShaderState.PixelShaderRHI = PixelShader.GetPixelShader();
+			GraphicsPSOInit.PrimitiveType = PassParameters->TileData.bRectPrimitive > 0 ? PT_RectList : PT_TriangleList;
+			SetGraphicsPipelineState(RHICmdList, GraphicsPSOInit, 0);
+
+			SetShaderParameters(RHICmdList, PixelShader, PixelShader.GetPixelShader(), *PassParameters);
+
+			RHICmdList.SetViewport(Viewport.Min.X, Viewport.Min.Y, 0.0f, Viewport.Max.X, Viewport.Max.Y, 1.0f);
+			SetShaderParameters(RHICmdList, TileVertexShader, TileVertexShader.GetVertexShader(), ParametersVS);
+			RHICmdList.SetStreamSource(0, nullptr, 0);
+			RHICmdList.DrawPrimitiveIndirect(PassParameters->TileData.TileIndirectBuffer->GetRHI(), FHairStrandsTiles::GetIndirectDrawArgOffset(TileType));
+		});
 #endif
 }
 
@@ -4096,6 +4150,7 @@ void DrawEditorSelection(FRDGBuilder& GraphBuilder, const FViewInfo& View, const
 		View,
 		ViewportRect,
 		View.HairStrandsViewData.MacroGroupDatas,
+		View.HairStrandsViewData.VisibilityData.TileData,
 		View.HairStrandsViewData.VisibilityData.NodeIndex,
 		View.HairStrandsViewData.VisibilityData.NodeVisData,
 		View.HairStrandsViewData.VisibilityData.CoverageTexture,
@@ -4234,7 +4289,8 @@ void DrawHitProxies(
 			OutMaxNodeCount);
 	}
 
-	AddHairStrandsHitProxyIdPass(GraphBuilder, Scene, View, MacroGroupDatas, VisNodeIndex, VisNodeData, CoverageTexture, HitProxyTexture, HitProxyDepthTexture);
+	// Make this tile dispatch
+	AddHairStrandsHitProxyIdPass(GraphBuilder, Scene, View, MacroGroupDatas, TileData, VisNodeIndex, VisNodeData, CoverageTexture, HitProxyTexture, HitProxyDepthTexture);
 }
 
 // Check if any simulated/skinned-bound groom has its positions updated (e.g. for invalidating the path-tracer accumulation)
@@ -4427,7 +4483,6 @@ void RenderHairStrandsVisibilityBuffer(
 
 				// Dummy coverage texture
 				CoverageTexture = GraphBuilder.CreateTexture(FRDGTextureDesc::Create2D(Resolution, FHairStrandsVisibilityData::CoverageFormat, FClearValueBinding::None, TexCreate_UAV | TexCreate_ShaderResource), TEXT("Hair.CoverageTexture"));
-				AddClearUAVPass(GraphBuilder, GraphBuilder.CreateUAV(CoverageTexture), 0.f);
 
 				CompactNodeData = PassOutput.NodeData;
 
