Prevent volume fog using an unsupported UAV format on Android GLES.
 - fixes GLES errors and pathological perf problems on some GPUs.

diff --git a/Engine/Source/Runtime/Renderer/Private/LocalFogVolumeRendering.cpp b/Engine/Source/Runtime/Renderer/Private/LocalFogVolumeRendering.cpp
--- a/Engine/Source/Runtime/Renderer/Private/LocalFogVolumeRendering.cpp
+++ b/Engine/Source/Runtime/Renderer/Private/LocalFogVolumeRendering.cpp
@@ -514,7 +514,14 @@ void CreateViewLocalFogVolumeBufferSRV(const FScene* Scene, FViewInfo& View, FRD
 
 	const FIntPoint TileDataTextureResolution				= FIntPoint::DivideAndRoundUp(View.ViewRect.Size(), FIntPoint(LocalFogVolumeTilePixelSize, LocalFogVolumeTilePixelSize));
 	const uint32 TileDataTextureSliceCount					= LocalFogVolumeTileMaxInstanceCount + 1; // +1 because the first slice is the culled instance count
-	FRDGTextureDesc Texture2DArrayDesc(FRDGTextureDesc::Create2DArray(TileDataTextureResolution, PF_R8_UINT, FClearValueBinding(EClearBinding::ENoneBound), TexCreate_ShaderResource | TexCreate_UAV | TexCreate_ReduceMemoryWithTilingMode | TexCreate_3DTiling, TileDataTextureSliceCount));
+
+	EPixelFormat TileDataFormat = PF_R8_UINT;	
+	if(IsMobilePlatform(View.GetShaderPlatform()) && IsOpenGLPlatform(View.GetShaderPlatform())) // TODO: !UE::PixelFormat::HasCapabilities(TileDataFormat, EPixelFormatCapabilities::UAV)
+	{
+		TileDataFormat = PF_R8G8B8A8_UINT;
+	}
+	// TODO:  check(UE::PixelFormat::HasCapabilities(TileDataFormat, EPixelFormatCapabilities::UAV));
+	FRDGTextureDesc Texture2DArrayDesc(FRDGTextureDesc::Create2DArray(TileDataTextureResolution, TileDataFormat, FClearValueBinding(EClearBinding::ENoneBound), TexCreate_ShaderResource | TexCreate_UAV | TexCreate_ReduceMemoryWithTilingMode | TexCreate_3DTiling, TileDataTextureSliceCount));
 	// LFV TODO, consider FFastVramConfig onto Texture2DArrayDesc
 
 	View.LocalFogVolumeViewData.TileDataTextureArray		= GraphBuilder.CreateTexture(Texture2DArrayDesc, TEXT("LocalFogVolume.CullingDataTexture"));
