Fix Global SDF not being built when using Stochastic Shadows or Global Distance Field visualization.
- it was relying on other techniques being enabled which is not always true.

diff --git a/Engine/Source/Runtime/Renderer/Private/DistanceFieldAmbientOcclusion.cpp b/Engine/Source/Runtime/Renderer/Private/DistanceFieldAmbientOcclusion.cpp
--- a/Engine/Source/Runtime/Renderer/Private/DistanceFieldAmbientOcclusion.cpp
+++ b/Engine/Source/Runtime/Renderer/Private/DistanceFieldAmbientOcclusion.cpp
@@ -20,6 +20,7 @@
 #include "VisualizeTexture.h"
 #include "RayTracing/RaytracingOptions.h"
 #include "Lumen/Lumen.h"
+#include "StochasticShadows/StochasticShadows.h"
 #include "ScenePrivate.h"
 #include "Substrate/Substrate.h"
 
@@ -774,14 +775,18 @@ bool FSceneRenderer::ShouldPrepareGlobalDistanceField() const
 		return false;
 	}
 
-	bool bShouldPrepareForAO = SupportsDistanceFieldAO(Scene->GetFeatureLevel(), Scene->GetShaderPlatform())
+	const bool bShouldPrepareForAO = SupportsDistanceFieldAO(Scene->GetFeatureLevel(), Scene->GetShaderPlatform())
 		&& (!GAOGlobalDistanceFieldDetailedNecessityCheck
 			|| ShouldPrepareForDistanceFieldAO()
 			|| ((FXSystem != nullptr) && FXSystem->UsesGlobalDistanceField()));
 
-	bShouldPrepareForAO = bShouldPrepareForAO || (IsLumenEnabled(Views[0]) && Lumen::UseGlobalSDFObjectGrid(*Views[0].Family));
+	const bool bShouldPrepareForLumen = IsLumenEnabled(Views[0]) && Lumen::UseGlobalSDFObjectGrid(*Views[0].Family);
 
-	return bShouldPrepareForAO && UseGlobalDistanceField();
+	const bool bShouldPrepareForStochasticShadows = StochasticShadows::UseGlobalSDF();
+
+	const bool bShouldPrepareForVisualization = ViewFamily.EngineShowFlags.VisualizeGlobalDistanceField;
+
+	return (bShouldPrepareForAO || bShouldPrepareForLumen || bShouldPrepareForStochasticShadows || bShouldPrepareForVisualization) && UseGlobalDistanceField();
 }
 
 void FDeferredShadingSceneRenderer::RenderDFAOAsIndirectShadowing(
diff --git a/Engine/Source/Runtime/Renderer/Private/StochasticShadows/StochasticShadows.h b/Engine/Source/Runtime/Renderer/Private/StochasticShadows/StochasticShadows.h
--- a/Engine/Source/Runtime/Renderer/Private/StochasticShadows/StochasticShadows.h
+++ b/Engine/Source/Runtime/Renderer/Private/StochasticShadows/StochasticShadows.h
@@ -17,4 +17,5 @@ namespace StochasticShadows
 	bool IsLightSupported(uint8 LightType, ECastRayTracedShadow::Type CastRayTracedShadow);
 	bool UseHardwareRayTracing();
 	bool UseInlineHardwareRayTracing();
+	bool UseGlobalSDF();
 };
\ No newline at end of file
diff --git a/Engine/Source/Runtime/Renderer/Private/StochasticShadows/StochasticShadowsRayTracing.cpp b/Engine/Source/Runtime/Renderer/Private/StochasticShadows/StochasticShadowsRayTracing.cpp
--- a/Engine/Source/Runtime/Renderer/Private/StochasticShadows/StochasticShadowsRayTracing.cpp
+++ b/Engine/Source/Runtime/Renderer/Private/StochasticShadows/StochasticShadowsRayTracing.cpp
@@ -92,6 +92,11 @@ namespace StochasticShadows
 		#endif
 	}
 
+	bool UseGlobalSDF()
+	{
+		return CVarStochasticShadowsWorldSpaceTraces.GetValueOnRenderThread() != 0 && !UseHardwareRayTracing();
+	}
+
 	bool IsUsingClosestHZB()
 	{
 		return IsEnabled() && CVarStochasticShadowsScreenTraces.GetValueOnRenderThread() != 0;
@@ -576,6 +581,8 @@ void StochasticShadows::RayTraceLightSamples(
 		}
 		else
 		{
+			ensure(StochasticShadows::UseGlobalSDF());
+
 			FSoftwareRayTraceLightSamplesCS::FParameters* PassParameters = GraphBuilder.AllocParameters<FSoftwareRayTraceLightSamplesCS::FParameters>();
 			PassParameters->CompactedTraceParameters = CompactedTraceParameters;
 			PassParameters->StochasticShadowsParameters = StochasticShadowsParameters;
