Fix issue where DF shadows and far shadow culling would not work correctly when ForceOnlyVirtualShadowMaps=1

diff --git a/Engine/Source/Runtime/Renderer/Private/ShadowSetup.cpp b/Engine/Source/Runtime/Renderer/Private/ShadowSetup.cpp
--- a/Engine/Source/Runtime/Renderer/Private/ShadowSetup.cpp
+++ b/Engine/Source/Runtime/Renderer/Private/ShadowSetup.cpp
@@ -5330,7 +5330,9 @@ void FSceneRenderer::AddViewDependentWholeSceneShadowsForView(
 			float MaxNonFarCascadeDistance = 0.0f;
 
 
-			if (!bNeedsVirtualShadowMap || CVarForceOnlyVirtualShadowMaps.GetValueOnRenderThread() == 0)
+			bool bForceOnlyVirtualShadowMaps = CVarForceOnlyVirtualShadowMaps.GetValueOnRenderThread() != 0;
+			bool bVsmUseFarShadowRules = CVarVsmUseFarShadowRules.GetValueOnRenderThread() != 0;
+			if (!bNeedsVirtualShadowMap || !bForceOnlyVirtualShadowMaps || bVsmUseFarShadowRules)
 			{
 				for (int32 Index = 0; Index < ProjectionCount; Index++)
 				{
@@ -5346,10 +5348,21 @@ void FSceneRenderer::AddViewDependentWholeSceneShadowsForView(
 
 					if (LightSceneInfo.Proxy->GetViewDependentWholeSceneProjectedShadowInitializer(View, LocalIndex, LightSceneInfo.IsPrecomputedLightingValid(), ProjectedShadowInitializer))
 					{
+						// Also needed for VSM-only when UseFarShadowCulling=1
 						if (!ProjectedShadowInitializer.CascadeSettings.bFarShadowCascade && !ProjectedShadowInitializer.bRayTracedDistanceField)
 						{
 							MaxNonFarCascadeDistance = FMath::Max(MaxNonFarCascadeDistance, ProjectedShadowInitializer.CascadeSettings.SplitFar);
 						}
+
+						if (bNeedsVirtualShadowMap && bForceOnlyVirtualShadowMaps)
+						{
+							// ForceOnlyVirtualShadowMaps disables CSM, but DF shadows can still be enabled
+							if (!ProjectedShadowInitializer.bRayTracedDistanceField)
+							{
+								continue;
+							}
+						}
+
 						uint32 ShadowBorder = NeedsUnatlasedCSMDepthsWorkaround( FeatureLevel ) ? 0 : SHADOW_BORDER;
 					
 						const int32 MaxCSMResolution = GetCachedScalabilityCVars().MaxCSMShadowResolution;
