Stochastic Shadows - warning fix "error: implicit conversion from 'float' to 'int'"

diff --git a/Engine/Source/Runtime/Renderer/Private/StochasticShadows/StochasticShadowsRayTracing.cpp b/Engine/Source/Runtime/Renderer/Private/StochasticShadows/StochasticShadowsRayTracing.cpp
--- a/Engine/Source/Runtime/Renderer/Private/StochasticShadows/StochasticShadowsRayTracing.cpp
+++ b/Engine/Source/Runtime/Renderer/Private/StochasticShadows/StochasticShadowsRayTracing.cpp
@@ -54,14 +54,14 @@ static TAutoConsoleVariable<int32> CVarStochasticShadowsHardwareRayTracingInline
 	ECVF_RenderThreadSafe | ECVF_Scalability
 );
 
-static TAutoConsoleVariable<int32> CVarStochasticShadowsHardwareRayTracingBias(
+static TAutoConsoleVariable<float> CVarStochasticShadowsHardwareRayTracingBias(
 	TEXT("r.StochasticShadows.HardwareRayTracing.Bias"),
 	1.0f,
 	TEXT("Constant bias for hardware ray traced shadow rays."),
 	ECVF_Scalability | ECVF_RenderThreadSafe
 );
 
-static TAutoConsoleVariable<int32> CVarStochasticShadowsHardwareRayTracingNormalBias(
+static TAutoConsoleVariable<float> CVarStochasticShadowsHardwareRayTracingNormalBias(
 	TEXT("r.StochasticShadows.HardwareRayTracing.NormalBias"),
 	0.1f,
 	TEXT("Normal bias for hardware ray traced shadow rays."),
