LFV - fixed local basis normalisation with new packing to not result in degenerated base.

diff --git a/Engine/Source/Runtime/Renderer/Private/LocalFogVolumeRendering.cpp b/Engine/Source/Runtime/Renderer/Private/LocalFogVolumeRendering.cpp
--- a/Engine/Source/Runtime/Renderer/Private/LocalFogVolumeRendering.cpp
+++ b/Engine/Source/Runtime/Renderer/Private/LocalFogVolumeRendering.cpp
@@ -391,12 +391,12 @@ void GetLocalFogVolumeSortingData(const FScene* Scene, FRDGBuilder& GraphBuilder
 
 		FVector3f XVec(InvTransform.M[0][0], InvTransform.M[0][1], InvTransform.M[0][2]);
 		FVector3f YVec(InvTransform.M[1][0], InvTransform.M[1][1], InvTransform.M[1][2]);
-		FVector3f ZVec(InvTransform.M[2][0], InvTransform.M[2][1], InvTransform.M[2][2]);
 		FVector3f Tran(InvTransform.M[3][0], InvTransform.M[3][1], InvTransform.M[3][2]);
 
-		XVec.Normalize();
-		YVec.Normalize();
-		ZVec.Normalize();
+		// Normalization requires small tolerance for large volumes.
+		const float NormalizeTolerance = 1.e-32;
+		XVec.Normalize(NormalizeTolerance);
+		YVec.Normalize(NormalizeTolerance);
 
 		auto AsUint32 = [](float X)
 		{
