Fix hair hit proxy not working in certain case.

diff --git a/Engine/Source/Runtime/Renderer/Private/HairStrands/HairStrandsData.h b/Engine/Source/Runtime/Renderer/Private/HairStrands/HairStrandsData.h
--- a/Engine/Source/Runtime/Renderer/Private/HairStrands/HairStrandsData.h
+++ b/Engine/Source/Runtime/Renderer/Private/HairStrands/HairStrandsData.h
@@ -494,7 +494,7 @@ namespace HairStrands
 	bool IsHairCompatible(const FMeshBatch* Mesh);
 	bool IsHairStrandsVF(const FMeshBatch* Mesh);
 	bool IsHairCardsVF(const FMeshBatch* Mesh);
-	bool IsHairVisible(const FMeshBatchAndRelevance& MeshBatch);
+	bool IsHairVisible(const FMeshBatchAndRelevance& MeshBatch, bool bCheckLengthScale);
 
 	// Hair helpers
 	bool HasHairInstanceInScene(const FScene& Scene);
diff --git a/Engine/Source/Runtime/Renderer/Private/HairStrands/HairStrandsInterface.cpp b/Engine/Source/Runtime/Renderer/Private/HairStrands/HairStrandsInterface.cpp
--- a/Engine/Source/Runtime/Renderer/Private/HairStrands/HairStrandsInterface.cpp
+++ b/Engine/Source/Runtime/Renderer/Private/HairStrands/HairStrandsInterface.cpp
@@ -507,14 +507,14 @@ bool IsHairCompatible(const FMeshBatch* Mesh)
 	return IsHairStrandsVF(Mesh) || IsHairCardsVF(Mesh);
 }
 
-bool IsHairVisible(const FMeshBatchAndRelevance& MeshBatch)
+bool IsHairVisible(const FMeshBatchAndRelevance& MeshBatch, bool bCheckLengthScale)
 {
 	if (MeshBatch.Mesh && MeshBatch.PrimitiveSceneProxy && MeshBatch.PrimitiveSceneProxy->ShouldRenderInMainPass())
 	{
 		const FHairGroupPublicData* Data = HairStrands::GetHairData(MeshBatch.Mesh);
 		switch (Data->VFInput.GeometryType)
 		{
-		case EHairGeometryType::Strands: return Data->VFInput.Strands.Common.LengthScale > 0;
+		case EHairGeometryType::Strands: return bCheckLengthScale ? Data->VFInput.Strands.Common.LengthScale > 0 : true;
 		case EHairGeometryType::Cards: return true;
 		case EHairGeometryType::Meshes: return true;
 		}
diff --git a/Engine/Source/Runtime/Renderer/Private/SceneVisibility.cpp b/Engine/Source/Runtime/Renderer/Private/SceneVisibility.cpp
--- a/Engine/Source/Runtime/Renderer/Private/SceneVisibility.cpp
+++ b/Engine/Source/Runtime/Renderer/Private/SceneVisibility.cpp
@@ -2302,7 +2302,9 @@ static void ComputeDynamicMeshRelevance(
 	const bool bIsHairStrandsCompatible = ViewRelevance.bHairStrands && IsHairStrandsEnabled(EHairStrandsShaderType::All, View.GetShaderPlatform());
 	if (bIsHairStrandsCompatible)
 	{
-		if (HairStrands::IsHairStrandsVF(MeshBatch.Mesh) && HairStrands::IsHairVisible(MeshBatch))
+		// Disable bCheckLengthScaleInitialize when running hit proxy as LengthScale is not initialized
+		const bool bCheckLengthScale = !View.Family->EngineShowFlags.HitProxies;
+		if (HairStrands::IsHairStrandsVF(MeshBatch.Mesh) && HairStrands::IsHairVisible(MeshBatch, bCheckLengthScale))
 		{
 			View.HairStrandsMeshElements.AddUninitialized(1);
 			FMeshBatchAndRelevance& BatchAndProxy = View.HairStrandsMeshElements.Last();
