SLW - added a cvar to test a potential fix for under water rendering of sky/fog/cloud/trnaslucent in correct order.
Before, translucent were rendered first and then fog/cloud composited above. Now the order is maintaine d above and below water when the fix is active.
Experimental fix until it is validated by artists.

diff --git a/Engine/Source/Runtime/Renderer/Private/DeferredShadingRenderer.cpp b/Engine/Source/Runtime/Renderer/Private/DeferredShadingRenderer.cpp
--- a/Engine/Source/Runtime/Renderer/Private/DeferredShadingRenderer.cpp
+++ b/Engine/Source/Runtime/Renderer/Private/DeferredShadingRenderer.cpp
@@ -323,6 +323,11 @@ static FAutoConsoleCommand RecreateRenderStateContextCmd(
 	TEXT("Recreate render state."),
 	FConsoleCommandDelegate::CreateStatic([] { FGlobalComponentRecreateRenderStateContext Context; }));
 
+static TAutoConsoleVariable<int32> CVarSingleLayerWaterUnderWaterExpFix(
+	TEXT("r.SingleLayerWater.UnderWaterFix"), 0,
+	TEXT("Experimental fix for sky, fog, cloud to be correctly ordered with translucent element. If proven valid we will enable it definintely."),
+	ECVF_RenderThreadSafe);
+
 #if RHI_RAYTRACING
 
 static bool bUpdateCachedRayTracingState = false;
@@ -4089,12 +4094,78 @@ void FDeferredShadingSceneRenderer::Render(FRDGBuilder& GraphBuilder)
 
 		FTranslucencyPassResourcesMap TranslucencyResourceMap(Views.Num());
 
+		const bool bSingleLayerWaterUnderWaterExpFix = CVarSingleLayerWaterUnderWaterExpFix.GetValueOnRenderThread() > 0;
+		const bool bIsCameraUnderWater = EnumHasAnyFlags(TranslucencyViewsToRender, ETranslucencyView::UnderWater);
+		FRDGTextureRef LightShaftOcclusionTexture = nullptr;
 		const bool bShouldRenderSingleLayerWater = !bHasRayTracedOverlay && ShouldRenderSingleLayerWater(Views);
 		FSceneWithoutWaterTextures SceneWithoutWaterTextures;
+		auto RenderLigthShaftSkyFogAndCloud = [&]()
+		{
+			// Draw Lightshafts
+			if (!bHasRayTracedOverlay && ViewFamily.EngineShowFlags.LightShafts)
+			{
+				SCOPE_CYCLE_COUNTER(STAT_FDeferredShadingSceneRenderer_RenderLightShaftOcclusion);
+				LightShaftOcclusionTexture = RenderLightShaftOcclusion(GraphBuilder, SceneTextures);
+			}
+
+			// Draw the sky atmosphere
+			if (!bHasRayTracedOverlay && bShouldRenderSkyAtmosphere && !IsForwardShadingEnabled(ShaderPlatform))
+			{
+				SCOPE_CYCLE_COUNTER(STAT_FDeferredShadingSceneRenderer_RenderSkyAtmosphere);
+				RenderSkyAtmosphere(GraphBuilder, SceneTextures);
+			}
+
+			// Draw fog.
+			bool bHeightFogHasComposedLocalFogVolume = false;
+			if (!bHasRayTracedOverlay && ShouldRenderFog(ViewFamily))
+			{
+				RDG_CSV_STAT_EXCLUSIVE_SCOPE(GraphBuilder, RenderFog);
+				SCOPED_NAMED_EVENT(RenderFog, FColor::Emerald);
+				SCOPE_CYCLE_COUNTER(STAT_FDeferredShadingSceneRenderer_RenderFog);
+				const bool bFogComposeLocalFogVolumes = (bShouldRenderLocalFogVolumeInVolumetricFog && bShouldRenderVolumetricFog) || bShouldRenderLocalFogVolumeDuringHeightFogPass;
+				RenderFog(GraphBuilder, SceneTextures, LightShaftOcclusionTexture, bFogComposeLocalFogVolumes);
+				bHeightFogHasComposedLocalFogVolume = bFogComposeLocalFogVolumes;
+			}
+
+			// Local Fog Volumes (LFV) rendering order is first HeightFog, then LFV, then volumetric fog on top.
+			// LFVs are rendered as part of the regular height fog + volumetric fog pass when volumetric fog is enabled and it is requested to voxelise LFVs into volumetric fog.
+			// Otherwise, they are rendered in an independent pass (this for instance make it independent of the near clip plane optimization).
+			if (!bHasRayTracedOverlay && !bHeightFogHasComposedLocalFogVolume)
+			{
+				RDG_CSV_STAT_EXCLUSIVE_SCOPE(GraphBuilder, RenderLocalFogVolume);
+				SCOPED_NAMED_EVENT(RenderLocalFogVolume, FColor::Emerald);
+				SCOPE_CYCLE_COUNTER(STAT_FDeferredShadingSceneRenderer_RenderLocalFogVolume);
+				RenderLocalFogVolume(Scene, Views, ViewFamily, GraphBuilder, SceneTextures, LightShaftOcclusionTexture);
+			}
+
+			// After the height fog, Draw volumetric clouds (having fog applied on them already) when using per pixel tracing,
+			if (!bHasRayTracedOverlay && bShouldRenderVolumetricCloud)
+			{
+				bool bSkipVolumetricRenderTarget = true;
+				bool bSkipPerPixelTracing = false;
+				RenderVolumetricCloud(GraphBuilder, SceneTextures, bSkipVolumetricRenderTarget, bSkipPerPixelTracing,
+					HalfResolutionDepthCheckerboardMinMaxTexture, QuarterResolutionDepthMinMaxTexture, false, InstanceCullingManager);
+			}
+
+			// or composite the off screen buffer over the scene.
+			if (bVolumetricRenderTargetRequired)
+			{
+				ComposeVolumetricRenderTargetOverScene(
+					GraphBuilder, Views, SceneTextures.Color.Target, SceneTextures.Depth.Target,
+					bSingleLayerWaterUnderWaterExpFix ? (bIsCameraUnderWater ? false : bShouldRenderSingleLayerWater) : bShouldRenderSingleLayerWater,
+					SceneWithoutWaterTextures, SceneTextures);
+			}
+		};
+
 		if (bShouldRenderSingleLayerWater)
 		{
-			if (EnumHasAnyFlags(TranslucencyViewsToRender, ETranslucencyView::UnderWater))
+			if (bIsCameraUnderWater)
 			{
+				if (bSingleLayerWaterUnderWaterExpFix)
+				{
+					RenderLigthShaftSkyFogAndCloud();
+				}
+
 				RDG_CSV_STAT_EXCLUSIVE_SCOPE(GraphBuilder, RenderTranslucency);
 				SCOPED_NAMED_EVENT(RenderTranslucency, FColor::Emerald);
 				SCOPE_CYCLE_COUNTER(STAT_TranslucencyDrawTime);
@@ -4105,7 +4176,8 @@ void FDeferredShadingSceneRenderer::Render(FRDGBuilder& GraphBuilder)
 			}
 
 			GraphBuilder.SetCommandListStat(GET_STATID(STAT_CLM_WaterPass));
-			RenderSingleLayerWater(GraphBuilder, SceneTextures, SingleLayerWaterPrePassResult, bShouldRenderVolumetricCloud, SceneWithoutWaterTextures, LumenFrameTemporaries);
+			RenderSingleLayerWater(GraphBuilder, SceneTextures, SingleLayerWaterPrePassResult, bShouldRenderVolumetricCloud, SceneWithoutWaterTextures, LumenFrameTemporaries, 
+				bSingleLayerWaterUnderWaterExpFix ? bIsCameraUnderWater : false); // false is the default value when the fix is not active.
 
 			// Replace main depth texture with the output of the SLW depth prepass which contains the scene + water.
 			// Note: Stencil now has all water bits marked with 1. As long as no other passes after this point want to read the depth buffer,
@@ -4119,58 +4191,9 @@ void FDeferredShadingSceneRenderer::Render(FRDGBuilder& GraphBuilder)
 		// Rebuild scene textures to include scene color.
 		SceneTextures.UniformBuffer = CreateSceneTextureUniformBuffer(GraphBuilder, &SceneTextures, FeatureLevel, SceneTextures.SetupMode);
 
-		FRDGTextureRef LightShaftOcclusionTexture = nullptr;
-
-		// Draw Lightshafts
-		if (!bHasRayTracedOverlay && ViewFamily.EngineShowFlags.LightShafts)
-		{
-			SCOPE_CYCLE_COUNTER(STAT_FDeferredShadingSceneRenderer_RenderLightShaftOcclusion);
-			LightShaftOcclusionTexture = RenderLightShaftOcclusion(GraphBuilder, SceneTextures);
-		}
-
-		// Draw the sky atmosphere
-		if (!bHasRayTracedOverlay && bShouldRenderSkyAtmosphere && !IsForwardShadingEnabled(ShaderPlatform))
-		{
-			SCOPE_CYCLE_COUNTER(STAT_FDeferredShadingSceneRenderer_RenderSkyAtmosphere);
-			RenderSkyAtmosphere(GraphBuilder, SceneTextures);
-		}
-
-		// Draw fog.
-		bool bHeightFogHasComposedLocalFogVolume = false;
-		if (!bHasRayTracedOverlay && ShouldRenderFog(ViewFamily))
-		{
-			RDG_CSV_STAT_EXCLUSIVE_SCOPE(GraphBuilder, RenderFog);
-			SCOPED_NAMED_EVENT(RenderFog, FColor::Emerald);
-			SCOPE_CYCLE_COUNTER(STAT_FDeferredShadingSceneRenderer_RenderFog);
-			const bool bFogComposeLocalFogVolumes = (bShouldRenderLocalFogVolumeInVolumetricFog && bShouldRenderVolumetricFog) || bShouldRenderLocalFogVolumeDuringHeightFogPass;
-			RenderFog(GraphBuilder, SceneTextures, LightShaftOcclusionTexture, bFogComposeLocalFogVolumes);
-			bHeightFogHasComposedLocalFogVolume = bFogComposeLocalFogVolumes;
-		}
-
-		// Local Fog Volumes (LFV) rendering order is first HeightFog, then LFV, then volumetric fog on top.
-		// LFVs are rendered as part of the regular height fog + volumetric fog pass when volumetric fog is enabled and it is requested to voxelise LFVs into volumetric fog.
-		// Otherwise, they are rendered in an independent pass (this for instance make it independent of the near clip plane optimization).
-		if (!bHasRayTracedOverlay && !bHeightFogHasComposedLocalFogVolume)
-		{
-			RDG_CSV_STAT_EXCLUSIVE_SCOPE(GraphBuilder, RenderLocalFogVolume);
-			SCOPED_NAMED_EVENT(RenderLocalFogVolume, FColor::Emerald);
-			SCOPE_CYCLE_COUNTER(STAT_FDeferredShadingSceneRenderer_RenderLocalFogVolume);
-			RenderLocalFogVolume(Scene, Views, ViewFamily, GraphBuilder, SceneTextures, LightShaftOcclusionTexture);
-		}
-
-		// After the height fog, Draw volumetric clouds (having fog applied on them already) when using per pixel tracing,
-		if (!bHasRayTracedOverlay && bShouldRenderVolumetricCloud)
-		{
-			bool bSkipVolumetricRenderTarget = true;
-			bool bSkipPerPixelTracing = false;
-			RenderVolumetricCloud(GraphBuilder, SceneTextures, bSkipVolumetricRenderTarget, bSkipPerPixelTracing, 
-				HalfResolutionDepthCheckerboardMinMaxTexture, QuarterResolutionDepthMinMaxTexture, false, InstanceCullingManager);
-		}
-
-		// or composite the off screen buffer over the scene.
-		if (bVolumetricRenderTargetRequired)
+		if (!bSingleLayerWaterUnderWaterExpFix || !bIsCameraUnderWater)
 		{
-			ComposeVolumetricRenderTargetOverScene(GraphBuilder, Views, SceneTextures.Color.Target, SceneTextures.Depth.Target, bShouldRenderSingleLayerWater, SceneWithoutWaterTextures, SceneTextures);
+			RenderLigthShaftSkyFogAndCloud();
 		}
 
 		FRDGTextureRef ExposureIlluminance = AddCalculateExposureIlluminancePass(GraphBuilder, Views, SceneTextures, TranslucencyLightingVolumeTextures, ExposureIlluminanceSetup);
diff --git a/Engine/Source/Runtime/Renderer/Private/DeferredShadingRenderer.h b/Engine/Source/Runtime/Renderer/Private/DeferredShadingRenderer.h
--- a/Engine/Source/Runtime/Renderer/Private/DeferredShadingRenderer.h
+++ b/Engine/Source/Runtime/Renderer/Private/DeferredShadingRenderer.h
@@ -390,7 +390,8 @@ public:
 		const FSingleLayerWaterPrePassResult* SingleLayerWaterPrePassResult,
 		bool bShouldRenderVolumetricCloud,
 		FSceneWithoutWaterTextures& SceneWithoutWaterTextures,
-		FLumenSceneFrameTemporaries& LumenFrameTemporaries);
+		FLumenSceneFrameTemporaries& LumenFrameTemporaries,
+		bool bIsCameraUnderWater);
 
 	void RenderSingleLayerWaterInner(
 		FRDGBuilder& GraphBuilder,
diff --git a/Engine/Source/Runtime/Renderer/Private/SingleLayerWaterRendering.cpp b/Engine/Source/Runtime/Renderer/Private/SingleLayerWaterRendering.cpp
--- a/Engine/Source/Runtime/Renderer/Private/SingleLayerWaterRendering.cpp
+++ b/Engine/Source/Runtime/Renderer/Private/SingleLayerWaterRendering.cpp
@@ -1154,7 +1154,8 @@ void FDeferredShadingSceneRenderer::RenderSingleLayerWater(
 	const FSingleLayerWaterPrePassResult* SingleLayerWaterPrePassResult,
 	bool bShouldRenderVolumetricCloud,
 	FSceneWithoutWaterTextures& SceneWithoutWaterTextures,
-	FLumenSceneFrameTemporaries& LumenFrameTemporaries)
+	FLumenSceneFrameTemporaries& LumenFrameTemporaries,
+	bool bIsCameraUnderWater)
 {
 	RDG_EVENT_SCOPE(GraphBuilder, "SingleLayerWater");
 	RDG_GPU_STAT_SCOPE(GraphBuilder, SingleLayerWater);
@@ -1164,11 +1165,11 @@ void FDeferredShadingSceneRenderer::RenderSingleLayerWater(
 	SceneWithoutWaterTextures = AddCopySceneWithoutWaterPass(GraphBuilder, ViewFamily, Views, SceneTextures.Color.Resolve, SceneTextures.Depth.Resolve);
 
 	// Render height fog over the color buffer if it is allocated, e.g. SingleLayerWaterUsesSimpleShading is true.
-	if (SceneWithoutWaterTextures.ColorTexture && ShouldRenderFog(ViewFamily))
+	if (!bIsCameraUnderWater && SceneWithoutWaterTextures.ColorTexture && ShouldRenderFog(ViewFamily))
 	{
 		RenderUnderWaterFog(GraphBuilder, SceneWithoutWaterTextures, SceneTextures.UniformBuffer);
 	}
-	if (SceneWithoutWaterTextures.ColorTexture && bShouldRenderVolumetricCloud)
+	if (!bIsCameraUnderWater && SceneWithoutWaterTextures.ColorTexture && bShouldRenderVolumetricCloud)
 	{
 		// This path is only taken when rendering the clouds in a render target that can be composited
 		ComposeVolumetricRenderTargetOverSceneUnderWater(GraphBuilder, Views, SceneWithoutWaterTextures, SceneTextures);
