Fix a bug of CSM caching that the static primitves are renderred twice.

diff --git a/Engine/Source/Runtime/Renderer/Private/ShadowSetup.cpp b/Engine/Source/Runtime/Renderer/Private/ShadowSetup.cpp
--- a/Engine/Source/Runtime/Renderer/Private/ShadowSetup.cpp
+++ b/Engine/Source/Runtime/Renderer/Private/ShadowSetup.cpp
@@ -4841,8 +4841,8 @@ struct FGatherShadowPrimitivesPacket
 				{
 					const FCachedShadowMapData& CachedShadowMapData = TaskData.Scene->GetCachedShadowMapDataRef(LightSceneInfo.Id, ProjectedShadowInfo->CascadeSettings.ShadowSplitIndex);
 
-					// if the mesh is dynamic or a new added mesh, it should cast csm.
-					bShadowRelevance = PrimitiveProxy->IsMeshShapeOftenMoving() || !CachedShadowMapData.StaticShadowSubjectPersistentPrimitiveIdMap[PrimitiveSceneInfo->GetPersistentIndex().Index];
+					// if the mesh is dynamic, it should cast csm.
+					bShadowRelevance = PrimitiveProxy->IsMeshShapeOftenMoving();
 				}
 				else if (ProjectedShadowInfo->CacheMode == SDCM_StaticPrimitivesOnly)
 				{
