Fixed: Case where mobile deferred shading renderer was applying movable skylight twice, once in base pass and in a lighting pass

diff --git a/Engine/Source/Runtime/Renderer/Private/MobileBasePassRendering.h b/Engine/Source/Runtime/Renderer/Private/MobileBasePassRendering.h
--- a/Engine/Source/Runtime/Renderer/Private/MobileBasePassRendering.h
+++ b/Engine/Source/Runtime/Renderer/Private/MobileBasePassRendering.h
@@ -28,6 +28,7 @@
 #include "LocalFogVolumeRendering.h"
 #include "DBufferTextures.h"
 #include "CompositionLighting/PostProcessDeferredDecals.h"
+#include "ReadOnlyCVARCache.h"
 
 bool MobileLocalLightsBufferEnabled(const FStaticShaderPlatform Platform);
 bool MobileMergeLocalLightsInPrepassEnabled(const FStaticShaderPlatform Platform);
@@ -407,10 +408,17 @@ public:
 
 		const bool bMobileUsesShadowMaskTexture = MobileUsesShadowMaskTexture(Parameters.Platform);
 		const bool bEnableClusteredReflections = MobileForwardEnableClusteredReflections(Parameters.Platform);
+		const bool bDeferredShadingEnabled = IsMobileDeferredShadingEnabled(Parameters.Platform);
 		const bool bTranslucentMaterial = IsTranslucentBlendMode(Parameters.MaterialParameters) || Parameters.MaterialParameters.ShadingModels.HasShadingModel(MSM_SingleLayerWater);
+		const bool bIsLit = Parameters.MaterialParameters.ShadingModels.IsLit();
+		const bool bMaterialUsesForwardShading = bIsLit && bTranslucentMaterial;
+		// Translucent materials always support clustered shading on mobile deferred
+		const bool bForwardShading = !bDeferredShadingEnabled || bMaterialUsesForwardShading;
+		// Only stationary skylights contribute into basepass with deferred shading materials. Has to do this test as on some project configurations we dont have a separate permutation for no-skylight
+		const bool bEnableSkylightInBasePass = bEnableSkyLight && (bForwardShading || FReadOnlyCVARCache::EnableStationarySkylight());
 
 		TMobileBasePassPSBaseType<LightMapPolicyType>::ModifyCompilationEnvironment(Parameters, OutEnvironment);
-		OutEnvironment.SetDefine(TEXT("ENABLE_SKY_LIGHT"), bEnableSkyLight);
+		OutEnvironment.SetDefine(TEXT("ENABLE_SKY_LIGHT"), bEnableSkylightInBasePass);
 		OutEnvironment.SetDefine(TEXT("ENABLE_AMBIENT_OCCLUSION"), IsMobileAmbientOcclusionEnabled(Parameters.Platform) ? 1u : 0u);
 		
 		FForwardLightingParameters::ModifyCompilationEnvironment(Parameters.Platform, OutEnvironment);
