Fix incorrect hair outline rendering with screen percentage.

diff --git a/Engine/Shaders/Private/HairStrands/HairStrandsHitProxy.usf b/Engine/Shaders/Private/HairStrands/HairStrandsHitProxy.usf
--- a/Engine/Shaders/Private/HairStrands/HairStrandsHitProxy.usf
+++ b/Engine/Shaders/Private/HairStrands/HairStrandsHitProxy.usf
@@ -59,27 +59,33 @@ Buffer<uint>					 SelectionMaterialIdBuffer;
 Texture2D<uint>					 VisNodeIndex;
 StructuredBuffer<FPackedHairVis> VisNodeData;
 Texture2D<float>				 CoverageTexture;
+Texture2D<float>				 HairOnlyDepthTexture;
 
 void EmitPS(
-	in FScreenVertexOutput In,
+	in float4 InPosition : SV_POSITION,
 	out float OutDepth : SV_Depth)
 {
-	const float2 UV = In.Position.xy * InvViewportResolution;
+
+	const float2 UV = InPosition.xy * InvViewportResolution;
 	const uint2 PixelCoord = UV * View.ViewSizeAndInvSize.xy;
+	const bool bValid = HairOnlyDepthTexture.Load(uint3(PixelCoord, 0)) > 0;
 
-	// Take the first sample
-	const FNodeDesc NodeDesc = DecodeNodeDesc(VisNodeIndex.Load(uint3(PixelCoord, 0)));
 	bool bIsSelected = false;
-	if (NodeDesc.Count > 0)
+	if (bValid)
 	{
-		const float Coverage = CoverageTexture.Load(uint3(PixelCoord, 0));
-		const uint SampleIndex = NodeDesc.Offset + 0;
-		const FHairVis InNode = UnpackHairVis(VisNodeData[SampleIndex]);
-		const uint MaterialId = clamp(InNode.MaterialId, 0, MaxMaterialCount - 1);
-		if (SelectionMaterialIdBuffer[MaterialId] > 0 && Coverage >= CoverageThreshold)
+		// Take the first sample
+		const FNodeDesc NodeDesc = DecodeNodeDesc(VisNodeIndex.Load(uint3(PixelCoord, 0)));
+		if (NodeDesc.Count > 0)
 		{
-			bIsSelected = true;
-			OutDepth = InNode.Depth;
+			const float Coverage = CoverageTexture.Load(uint3(PixelCoord, 0));
+			const uint SampleIndex = NodeDesc.Offset + 0;
+			const FHairVis InNode = UnpackHairVis(VisNodeData[SampleIndex]);
+			const uint MaterialId = clamp(InNode.MaterialId, 0, MaxMaterialCount - 1);
+			if (SelectionMaterialIdBuffer[MaterialId] > 0 && Coverage >= CoverageThreshold)
+			{
+				bIsSelected = true;
+				OutDepth = InNode.Depth;
+			}
 		}
 	}
 
diff --git a/Engine/Source/Runtime/Renderer/Private/HairStrands/HairStrandsVisibility.cpp b/Engine/Source/Runtime/Renderer/Private/HairStrands/HairStrandsVisibility.cpp
--- a/Engine/Source/Runtime/Renderer/Private/HairStrands/HairStrandsVisibility.cpp
+++ b/Engine/Source/Runtime/Renderer/Private/HairStrands/HairStrandsVisibility.cpp
@@ -3837,9 +3837,9 @@ class FHairStrandsEmitSelectionPS : public FGlobalShader
 		SHADER_PARAMETER_STRUCT_REF(FViewUniformShaderParameters, View)
 		SHADER_PARAMETER_RDG_TEXTURE(Texture2D, VisNodeIndex)
 		SHADER_PARAMETER_RDG_TEXTURE(Texture2D<float>, CoverageTexture)
+		SHADER_PARAMETER_RDG_TEXTURE(Texture2D<float>, HairOnlyDepthTexture)
 		SHADER_PARAMETER_RDG_BUFFER_SRV(StructuredBuffer<FPackedHairVis>, VisNodeData)
 		SHADER_PARAMETER_RDG_BUFFER_SRV(Buffer<uint>, SelectionMaterialIdBuffer)
-		SHADER_PARAMETER_STRUCT_INCLUDE(FHairStrandsTilePassVS::FParameters, TileData)
 		RENDER_TARGET_BINDING_SLOTS()
 	END_SHADER_PARAMETER_STRUCT()
 
@@ -3865,6 +3865,7 @@ void AddHairStrandsSelectionOutlinePass(
 	FRDGTextureRef VisNodeIndex,
 	FRDGBufferRef VisNodeData,
 	FRDGTextureRef CoverageTexture,
+	FRDGTextureRef HairOnlyDepthTexture,
 	FRDGTextureRef SelectionDepthTexture)
 {
 	if (View.HairStrandsMeshElements.Num() == 0 || !VisNodeData)
@@ -3899,43 +3900,27 @@ void AddHairStrandsSelectionOutlinePass(
 	PassParameters->VisNodeIndex = VisNodeIndex;
 	PassParameters->VisNodeData = GraphBuilder.CreateSRV(VisNodeData);
 	PassParameters->CoverageTexture = CoverageTexture;
+	PassParameters->HairOnlyDepthTexture = HairOnlyDepthTexture;
 	PassParameters->SelectionMaterialIdBuffer = GraphBuilder.CreateSRV(SelectionMaterialIdBuffer, PF_R32_UINT);
 	PassParameters->RenderTargets.DepthStencil = FDepthStencilBinding(SelectionDepthTexture, ERenderTargetLoadAction::ELoad, FExclusiveDepthStencil::DepthWrite_StencilWrite);
-	PassParameters->TileData = GetHairStrandsTileParameters(View, TileData, TileType);
 
 	const FIntRect Viewport = View.ViewRect;
 	auto PixelShader = View.ShaderMap->GetShader<FHairStrandsEmitSelectionPS>();
-	TShaderMapRef<FHairStrandsTilePassVS> TileVertexShader(View.ShaderMap);
 
+	// We don't use tile rendering for hair selection, because the outline buffer is unscaled, and does not match the visibility buffer which is unscaled
 	const uint32 StencilRef = 3;
-	
-	GraphBuilder.AddPass(
-		RDG_EVENT_NAME("HairStrands::EmitSelection(Tile)"),
+	FPixelShaderUtils::AddFullscreenPass(
+		GraphBuilder,
+		View.ShaderMap,
+		RDG_EVENT_NAME("HairStrands::EmitSelection"),
+		PixelShader,
 		PassParameters,
-		ERDGPassFlags::Raster,
-		[PassParameters, TileVertexShader, PixelShader, Viewport, TileType, StencilRef](FRHICommandList& RHICmdList)
-		{
-			FHairStrandsTilePassVS::FParameters ParametersVS = PassParameters->TileData;
-
-			FGraphicsPipelineStateInitializer GraphicsPSOInit;
-			RHICmdList.ApplyCachedRenderTargets(GraphicsPSOInit);
-			GraphicsPSOInit.BlendState = TStaticBlendState<>::GetRHI();
-			GraphicsPSOInit.RasterizerState = TStaticRasterizerState<>::GetRHI();
-			GraphicsPSOInit.DepthStencilState = TStaticDepthStencilState<true, CF_DepthNearOrEqual, true, CF_Always, SO_Keep, SO_Keep, SO_Replace>::GetRHI();
-
-			GraphicsPSOInit.BoundShaderState.VertexDeclarationRHI = GFilterVertexDeclaration.VertexDeclarationRHI;
-			GraphicsPSOInit.BoundShaderState.VertexShaderRHI = TileVertexShader.GetVertexShader();
-			GraphicsPSOInit.BoundShaderState.PixelShaderRHI = PixelShader.GetPixelShader();
-			GraphicsPSOInit.PrimitiveType = PassParameters->TileData.bRectPrimitive > 0 ? PT_RectList : PT_TriangleList;
-			SetGraphicsPipelineState(RHICmdList, GraphicsPSOInit, StencilRef);
-
-			SetShaderParameters(RHICmdList, PixelShader, PixelShader.GetPixelShader(), *PassParameters);
+		ViewportRect,
+		TStaticBlendState<>::GetRHI(),
+		TStaticRasterizerState<>::GetRHI(),
+		TStaticDepthStencilState<true, CF_DepthNearOrEqual, true, CF_Always, SO_Keep, SO_Keep, SO_Replace>::GetRHI(), 
+		StencilRef);
 
-			RHICmdList.SetViewport(Viewport.Min.X, Viewport.Min.Y, 0.0f, Viewport.Max.X, Viewport.Max.Y, 1.0f);
-			SetShaderParameters(RHICmdList, TileVertexShader, TileVertexShader.GetVertexShader(), ParametersVS);
-			RHICmdList.SetStreamSource(0, nullptr, 0);
-			RHICmdList.DrawPrimitiveIndirect(PassParameters->TileData.TileIndirectBuffer->GetRHI(), FHairStrandsTiles::GetIndirectDrawArgOffset(TileType));
-		});
 #endif
 }
 
@@ -4154,6 +4139,7 @@ void DrawEditorSelection(FRDGBuilder& GraphBuilder, const FViewInfo& View, const
 		View.HairStrandsViewData.VisibilityData.NodeIndex,
 		View.HairStrandsViewData.VisibilityData.NodeVisData,
 		View.HairStrandsViewData.VisibilityData.CoverageTexture,
+		View.HairStrandsViewData.VisibilityData.HairOnlyDepthTexture,
 		SelectionDepthTexture);
 }
 
