Merge the Local Lights in Basepass for translucent objects when r.Mobile.Forward.EnableLocalLights == 2
#localization none
#tests   game

diff --git a/Engine/Source/Runtime/Renderer/Private/MobileBasePass.cpp b/Engine/Source/Runtime/Renderer/Private/MobileBasePass.cpp
--- a/Engine/Source/Runtime/Renderer/Private/MobileBasePass.cpp
+++ b/Engine/Source/Runtime/Renderer/Private/MobileBasePass.cpp
@@ -44,13 +44,13 @@ bool MobileLocalLightsUseSinglePermutation()
 	return GMobileForwardLocalLightsSinglePermutation != 0;
 }
 
-EMobileLocalLightSetting GetMobileForwardLocalLightSetting(EShaderPlatform ShaderPlatform, bool bIsTranslucent)
+EMobileLocalLightSetting GetMobileForwardLocalLightSetting(EShaderPlatform ShaderPlatform)
 {
 	const int32 MobileForwardLocalLightsIniValue = FReadOnlyCVARCache::MobileForwardLocalLights(ShaderPlatform);
 
 	if (MobileForwardLocalLightsIniValue > 0)
 	{
-		if ((MobileForwardLocalLightsIniValue == 1) || bIsTranslucent)
+		if (MobileForwardLocalLightsIniValue == 1)
 		{
 			return EMobileLocalLightSetting::LOCAL_LIGHTS_ENABLED;
 		}
@@ -926,7 +926,7 @@ bool FMobileBasePassMeshProcessor::Process(
 			// Translucency always use un-cached MDC so this option does not apply to it
 			((MobileLocalLightsUseSinglePermutation() && !bIsTranslucent) || PrimitiveSceneProxy->GetPrimitiveSceneInfo()->NumMobileDynamicLocalLights > 0))
 		{
-			LocalLightSetting = GetMobileForwardLocalLightSetting(Scene->GetShaderPlatform(), bIsTranslucent);
+			LocalLightSetting = GetMobileForwardLocalLightSetting(Scene->GetShaderPlatform());
 		}
 	}
 
@@ -1116,7 +1116,7 @@ void FMobileBasePassMeshProcessor::CollectPSOInitializers(const FSceneTexturesCo
 	EMobileLocalLightSetting LocalLightSetting = EMobileLocalLightSetting::LOCAL_LIGHTS_DISABLED;
 	if (bLitMaterial && !bPassUsesDeferredShading)
 	{
-		LocalLightSetting = GetMobileForwardLocalLightSetting(ShaderPlatform, bTranslucentBasePass);
+		LocalLightSetting = GetMobileForwardLocalLightSetting(ShaderPlatform);
 	}
 	const bool bUseLocalLightPermutation = (LocalLightSetting != EMobileLocalLightSetting::LOCAL_LIGHTS_DISABLED);
 
diff --git a/Engine/Source/Runtime/Renderer/Private/MobileBasePassRendering.h b/Engine/Source/Runtime/Renderer/Private/MobileBasePassRendering.h
--- a/Engine/Source/Runtime/Renderer/Private/MobileBasePassRendering.h
+++ b/Engine/Source/Runtime/Renderer/Private/MobileBasePassRendering.h
@@ -106,7 +106,7 @@ enum EOutputFormat
 	HDR_LINEAR_64,
 };
 
-EMobileLocalLightSetting GetMobileForwardLocalLightSetting(EShaderPlatform ShaderPlatform, bool bIsTranslucent);
+EMobileLocalLightSetting GetMobileForwardLocalLightSetting(EShaderPlatform ShaderPlatform);
 
 enum class EMobileTranslucentColorTransmittanceMode
 {
@@ -388,7 +388,7 @@ public:
 		EMobileLocalLightSetting SupportedLocalLightsType = EMobileLocalLightSetting::LOCAL_LIGHTS_DISABLED;
 		if (bForwardShading && bIsLit)
 		{
-			SupportedLocalLightsType = GetMobileForwardLocalLightSetting(Parameters.Platform, bIsTranslucent);
+			SupportedLocalLightsType = GetMobileForwardLocalLightSetting(Parameters.Platform);
 		}
 		// Deferred shading does not need SkyLight and LocalLight permutations
 		// TODO: skip skylight permutations for deferred	
@@ -428,11 +428,11 @@ public:
 		uint32 MergedLocalLights = 0u;
 		if (LocalLightSetting == EMobileLocalLightSetting::LOCAL_LIGHTS_BUFFER)
 		{
-			if (MobileMergeLocalLightsInPrepassEnabled(Parameters.Platform))
+			if (MobileMergeLocalLightsInPrepassEnabled(Parameters.Platform) && !bTranslucentMaterial)
 			{
 				MergedLocalLights = 1u;
 			}
-			else if (MobileMergeLocalLightsInBasepassEnabled(Parameters.Platform))
+			else if (MobileMergeLocalLightsInBasepassEnabled(Parameters.Platform) || bTranslucentMaterial)
 			{
 				MergedLocalLights = 2u;
 			}
