Implements r.TSR.16BitVALU.{AMD,Intel,Nvidia}

diff --git a/Engine/Source/Runtime/Renderer/Private/PostProcess/TemporalSuperResolution.cpp b/Engine/Source/Runtime/Renderer/Private/PostProcess/TemporalSuperResolution.cpp
--- a/Engine/Source/Runtime/Renderer/Private/PostProcess/TemporalSuperResolution.cpp
+++ b/Engine/Source/Runtime/Renderer/Private/PostProcess/TemporalSuperResolution.cpp
@@ -84,6 +84,25 @@ TAutoConsoleVariable<int32> CVarTSR16BitVALU(
 	TEXT("Whether to use 16bit VALU on platform that have bSupportsRealTypes=RuntimeDependent"),
 	ECVF_RenderThreadSafe);
 
+#if PLATFORM_DESKTOP
+
+TAutoConsoleVariable<int32> CVarTSR16BitVALUOnAMD(
+	TEXT("r.TSR.16BitVALU.AMD"), 0,
+	TEXT("Overrides whether to use 16bit VALU on AMD desktop GPUs"),
+	ECVF_RenderThreadSafe);
+
+TAutoConsoleVariable<int32> CVarTSR16BitVALUOnIntel(
+	TEXT("r.TSR.16BitVALU.Intel"), 0,
+	TEXT("Overrides whether to use 16bit VALU on Intel desktop GPUs"),
+	ECVF_RenderThreadSafe);
+
+TAutoConsoleVariable<int32> CVarTSR16BitVALUOnNvidia(
+	TEXT("r.TSR.16BitVALU.Nvidia"), 0,
+	TEXT("Overrides whether to use 16bit VALU on Nvidia desktop GPUs"),
+	ECVF_RenderThreadSafe);
+
+#endif // PLATFORM_DESKTOP
+
 TAutoConsoleVariable<float> CVarTSRHistoryRejectionSampleCount(
 	TEXT("r.TSR.ShadingRejection.SampleCount"), 2.0f,
 	TEXT("Maximum number of sample in each output pixel of the history after total shading rejection.\n")
@@ -1259,18 +1278,21 @@ FDefaultTemporalUpscaler::FOutputs AddTemporalSuperResolutionPasses(
 	const ERHIFeatureSupport VALU16BitSupport = FTSRShader::Supports16BitVALU(View.GetShaderPlatform());
 	bool bUse16BitVALU = (CVarTSR16BitVALU.GetValueOnRenderThread() != 0 && GRHIGlobals.SupportsNative16BitOps && VALU16BitSupport == ERHIFeatureSupport::RuntimeDependent) || VALU16BitSupport == ERHIFeatureSupport::RuntimeGuaranteed;
 
+	// Controls whether to use 16bit ops on per GPU vendor in mean time each driver matures.
 #if PLATFORM_DESKTOP
-	// PLAY-18908: work around shader corruption on Nvidia GPU until fix is done in the driver
-	if (IsRHIDeviceNVIDIA())
+	if (IsRHIDeviceAMD())
 	{
-		bUse16BitVALU = false;
+		bUse16BitVALU = CVarTSR16BitVALUOnAMD.GetValueOnRenderThread() != 0;
 	}
-
-	if (IsRHIDeviceAMD())
+	else if (IsRHIDeviceIntel())
 	{
-		bUse16BitVALU = false;
+		bUse16BitVALU = CVarTSR16BitVALUOnIntel.GetValueOnRenderThread() != 0;
 	}
-#endif
+	else if (IsRHIDeviceNVIDIA())
+	{
+		bUse16BitVALU = CVarTSR16BitVALUOnNvidia.GetValueOnRenderThread() != 0;
+	}
+#endif // PLATFORM_DESKTOP
 
 	// Whether should accumulate sub pixel depth.
 	const ETSRSubpixelMethod SubpixelMethod = ETSRSubpixelMethod(FMath::Clamp(CVarTSRSubpixelMethod.GetValueOnRenderThread(), 0, 2));
