Remove ProxiesWithDirtyCachedInstance since cached instances are properly invalidated when mesh LODs are streamed.
- this was a temporary workaround to update cached instance when raytracing geometries were streamed out, but this is now correctly handled in UStaticMesh::RequestUpdateCachedRenderState().

diff --git a/Engine/Source/Runtime/Renderer/Private/DeferredShadingRenderer.cpp b/Engine/Source/Runtime/Renderer/Private/DeferredShadingRenderer.cpp
--- a/Engine/Source/Runtime/Renderer/Private/DeferredShadingRenderer.cpp
+++ b/Engine/Source/Runtime/Renderer/Private/DeferredShadingRenderer.cpp
@@ -748,7 +748,6 @@ struct FRayTracingRelevantPrimitive
 	bool bIsSky = false;
 	bool bAllSegmentsTranslucent = true;
 
-	bool bCachedRayTracingGeometryValid = true;
 	const FRayTracingGeometryInstance* CachedRayTracingInstance = nullptr;
 	TArrayView<const int32> CachedRayTracingMeshCommandIndices; // Pointer to FPrimitiveSceneInfo::CachedRayTracingMeshCommandIndicesPerLOD data
 
@@ -961,7 +960,7 @@ static void GatherRayTracingRelevantPrimitives(FScene& Scene, const FViewInfo& V
 				if (!bUsingNaniteRayTracing)
 				{
 					// Currently IsCachedRayTracingGeometryValid() can only be called for non-nanite geometries
-					RelevantPrimitive.bCachedRayTracingGeometryValid = SceneInfo->IsCachedRayTracingGeometryValid();
+					checkf(SceneInfo->IsCachedRayTracingGeometryValid(), TEXT("Cached ray tracing instance is expected to be valid. Was mesh LOD streamed but cached data was not invalidated?"));
 					checkf(SceneInfo->CachedRayTracingInstance.GeometryRHI, TEXT("Ray tracing instance must have a valid geometry."));
 				}
 
@@ -1439,21 +1438,18 @@ bool FDeferredShadingSceneRenderer::GatherRayTracingWorldInstancesForView(FRDGBu
 
 		FRayTracingScene& RayTracingScene; // New instances are added into FRayTracingScene::Instances and FRayTracingScene::Allocator is used for temporary data
 		TArray<FVisibleRayTracingMeshCommand>& VisibleRayTracingMeshCommands; // New elements are added here by this task
-		TArray<FPrimitiveSceneProxy*>& ProxiesWithDirtyCachedInstance;
 
 		FRayTracingSceneAddInstancesTask(const FScene& InScene,
 											TChunkedArray<FRayTracingRelevantPrimitive>& InRelevantStaticPrimitives,
 											const FRayTracingCullingParameters& InCullingParameters,
 											const bool bInIsPathTracing,
-											FRayTracingScene& InRayTracingScene, TArray<FVisibleRayTracingMeshCommand>& InVisibleRayTracingMeshCommands,
-											TArray<FPrimitiveSceneProxy*>& InProxiesWithDirtyCachedInstance)
+											FRayTracingScene& InRayTracingScene, TArray<FVisibleRayTracingMeshCommand>& InVisibleRayTracingMeshCommands)
 			: Scene(InScene)
 			, RelevantStaticPrimitives(InRelevantStaticPrimitives)
 			, CullingParameters(InCullingParameters)
 			, bIsPathTracing(bInIsPathTracing)
 			, RayTracingScene(InRayTracingScene)
 			, VisibleRayTracingMeshCommands(InVisibleRayTracingMeshCommands)
-			, ProxiesWithDirtyCachedInstance(InProxiesWithDirtyCachedInstance)
 		{
 			VisibleRayTracingMeshCommands.Reserve(RelevantStaticPrimitives.Num());
 		}
@@ -1554,12 +1550,6 @@ bool FDeferredShadingSceneRenderer::GatherRayTracingWorldInstancesForView(FRDGBu
 							continue;
 						}
 					}
-					else if (!RelevantPrimitive.bCachedRayTracingGeometryValid)
-					{
-						// cached instance is not valid (eg: was streamed out) need to invalidate for next frame
-						ProxiesWithDirtyCachedInstance.Add(Scene.PrimitiveSceneProxies[PrimitiveIndex]);
-						continue;
-					}
 					
 					// TODO: Consider requesting a recache of all ray tracing commands during which decals are excluded
 					
@@ -1803,7 +1793,7 @@ bool FDeferredShadingSceneRenderer::GatherRayTracingWorldInstancesForView(FRDGBu
 
 	FGraphEventRef AddInstancesTask = TGraphTask<FRayTracingSceneAddInstancesTask>::CreateTask(&AddInstancesTaskPrerequisites).ConstructAndDispatchWhenReady(
 		*Scene, RelevantPrimitiveList.StaticPrimitives, View.RayTracingCullingParameters, bool(View.Family->EngineShowFlags.PathTracing), // inputs 
-		RayTracingScene, View.VisibleRayTracingMeshCommands, View.ProxiesWithDirtyCachedInstance // outputs
+		RayTracingScene, View.VisibleRayTracingMeshCommands // outputs
 	);
 
 	// Scene init task can run only when all pre-init tasks are complete (including culling tasks that are spawned while adding instances)
@@ -2054,11 +2044,6 @@ bool FDeferredShadingSceneRenderer::DispatchRayTracingWorldUpdates(FRDGBuilder&
 
 	ReferenceView.RayTracingSceneInitTask = {};
 
-	for (FPrimitiveSceneProxy* SceneProxy : ReferenceView.ProxiesWithDirtyCachedInstance)
-	{
-		SceneProxy->GetScene().UpdateCachedRayTracingState(SceneProxy);
-	}
-
 	{
 		Nanite::GRayTracingManager.ProcessUpdateRequests(GraphBuilder, GetSceneUniforms());
 		const bool bAnyBlasRebuilt = Nanite::GRayTracingManager.ProcessBuildRequests(GraphBuilder);
diff --git a/Engine/Source/Runtime/Renderer/Private/SceneRendering.h b/Engine/Source/Runtime/Renderer/Private/SceneRendering.h
--- a/Engine/Source/Runtime/Renderer/Private/SceneRendering.h
+++ b/Engine/Source/Runtime/Renderer/Private/SceneRendering.h
@@ -1357,8 +1357,6 @@ public:
 	FRayTracingSceneWithGeometryInstances RayTracingSceneInitData;
 	FGraphEventRef RayTracingSceneInitTask; // Task to asynchronously create RayTracingSceneInitData
 
-	TArray<FPrimitiveSceneProxy*> ProxiesWithDirtyCachedInstance;
-
 	FGraphEventArray AddRayTracingMeshBatchTaskList;
 	TArray<FRayTracingMeshCommandOneFrameArray*, SceneRenderingAllocator> VisibleRayTracingMeshCommandsPerTask;
 	TArray<FDynamicRayTracingMeshCommandStorage*, SceneRenderingAllocator> DynamicRayTracingMeshCommandStoragePerTask;
