Stochastic Shadows:
* Removed shadow rays traced from empty tiles (e.g. sky) due to unitialized SampleBuffer in those tiles
* Removed a bunch of FNumSamplesPerPixel shader permutations and replaced them with a bit ALU as it doesn't have any influence on performance

diff --git a/Engine/Shaders/Private/StochasticShadows/StochasticShadows.usf b/Engine/Shaders/Private/StochasticShadows/StochasticShadows.usf
--- a/Engine/Shaders/Private/StochasticShadows/StochasticShadows.usf
+++ b/Engine/Shaders/Private/StochasticShadows/StochasticShadows.usf
@@ -11,6 +11,7 @@ groupshared uint SharedTileMode[THREADGROUP_SIZE * THREADGROUP_SIZE];
 
 RWStructuredBuffer<uint> RWTileAllocator;
 RWStructuredBuffer<uint> RWTileData;
+RWTexture2D<uint> RWDownsampledTileMask;
 uint TileDataStride;
 uint DownsampledTileDataStride;
 uint2 DownsampledViewSize;
@@ -110,17 +111,25 @@ void TileClassificationCS(
 		const uint DataStride = TileDataStride;
 #endif
 
+		uint2 TileCoord = ScreenCoord / TILE_SIZE;
+
+		#if DOWNSAMPLED_CLASSIFICATION
+		{
+			RWDownsampledTileMask[TileCoord] = SharedTileMode[0] != 0 ? 1 : 0;
+		}
+		#endif
+
 		if (SharedTileMode[0] & 0x2)
 		{
 			uint StochasticShadowsTileIndex = 0;
 			InterlockedAdd(RWTileAllocator[TILE_MODE_COMPLEX_SHADING], 1, StochasticShadowsTileIndex);
-			RWTileData[DataStride * TILE_MODE_COMPLEX_SHADING + StochasticShadowsTileIndex] = PackTile(ScreenCoord / 8);
+			RWTileData[DataStride * TILE_MODE_COMPLEX_SHADING + StochasticShadowsTileIndex] = PackTile(TileCoord);
 		}
 		else if (SharedTileMode[0] & 0x1)
 		{
 			uint StochasticShadowsTileIndex = 0;
 			InterlockedAdd(RWTileAllocator[TILE_MODE_SIMPLE_SHADING], 1, StochasticShadowsTileIndex);
-			RWTileData[DataStride * TILE_MODE_SIMPLE_SHADING + StochasticShadowsTileIndex] = PackTile(ScreenCoord / 8);
+			RWTileData[DataStride * TILE_MODE_SIMPLE_SHADING + StochasticShadowsTileIndex] = PackTile(TileCoord);
 		}
 	}
 }
diff --git a/Engine/Shaders/Private/StochasticShadows/StochasticShadows.ush b/Engine/Shaders/Private/StochasticShadows/StochasticShadows.ush
--- a/Engine/Shaders/Private/StochasticShadows/StochasticShadows.ush
+++ b/Engine/Shaders/Private/StochasticShadows/StochasticShadows.ush
@@ -36,12 +36,6 @@
 #define DOWNSAMPLE_FACTOR				2
 #define SHADOW_MASK_DOWNSAMPLE_FACTOR	2
 
-#ifndef NUM_SAMPLES_PER_PIXEL
-	#define NUM_SAMPLES_PER_PIXEL	4
-	#define NUM_SAMPLES_PER_PIXEL_Y 2
-	#define NUM_SAMPLES_PER_PIXEL_X 2
-#endif
-
 // Keep in sync with StochasticShadows::ETileType
 #define TILE_MODE_SIMPLE_SHADING	0
 #define TILE_MODE_COMPLEX_SHADING	1
diff --git a/Engine/Shaders/Private/StochasticShadows/StochasticShadowsComposite.usf b/Engine/Shaders/Private/StochasticShadows/StochasticShadowsComposite.usf
--- a/Engine/Shaders/Private/StochasticShadows/StochasticShadowsComposite.usf
+++ b/Engine/Shaders/Private/StochasticShadows/StochasticShadowsComposite.usf
@@ -15,6 +15,7 @@ Texture2D<float> DownsampledSceneDepth;
 Texture2D<UNORM float3> DownsampledSceneWorldNormal;
 uint TemporalMaxFramesAccumulated;
 float TemporalStdDevOffset;
+uint2 NumSamplesPerPixel;
 
 RWTexture2D<uint> RWShadowMaskTileAtlas;
 
@@ -72,11 +73,11 @@ void CompositeShadowMaskTracesCS(
 			SampleAccumulator.SampleSum = 0.0f;
 			SampleAccumulator.SampleWeightSum = 0.0f;
 
-			for (uint LocalSampleY = 0; LocalSampleY < NUM_SAMPLES_PER_PIXEL_Y; ++LocalSampleY)
+			for (uint LocalSampleY = 0; LocalSampleY < NumSamplesPerPixel.y; ++LocalSampleY)
 			{
-				for (uint LocalSampleX = 0; LocalSampleX < NUM_SAMPLES_PER_PIXEL_X; ++LocalSampleX)
+				for (uint LocalSampleX = 0; LocalSampleX < NumSamplesPerPixel.x; ++LocalSampleX)
 				{
-					uint2 SampleCoord = ShadowMaskScreenCoord * uint2(NUM_SAMPLES_PER_PIXEL_X, NUM_SAMPLES_PER_PIXEL_Y) + uint2(LocalSampleX, LocalSampleY);
+					uint2 SampleCoord = ShadowMaskScreenCoord * NumSamplesPerPixel + uint2(LocalSampleX, LocalSampleY);
 					FLightSample LightSample = UnpackLightSample(LightSamples[SampleCoord]);
 
 					if (LightSample.LightSceneId == LightSceneId)
diff --git a/Engine/Shaders/Private/StochasticShadows/StochasticShadowsSampling.usf b/Engine/Shaders/Private/StochasticShadows/StochasticShadowsSampling.usf
--- a/Engine/Shaders/Private/StochasticShadows/StochasticShadowsSampling.usf
+++ b/Engine/Shaders/Private/StochasticShadows/StochasticShadowsSampling.usf
@@ -1,10 +1,13 @@
 // Copyright Epic Games, Inc. All Rights Reserved.
 
 #define STOCHASTIC_SHADOWS 1
-
 #define NON_DIRECTIONAL_DIRECT_LIGHTING 1
 #define SUPPORT_CONTACT_SHADOWS 0
 
+#ifndef NUM_SAMPLES_PER_PIXEL_1D
+	#error NUM_SAMPLES_PER_PIXEL_1D must be defined
+#endif
+
 #include "../Common.ush"
 #include "../LightGridCommon.ush"
 #include "../BlueNoise.ush"
@@ -87,6 +90,7 @@ Texture2D<float4> ShadowMaskHistoryWeights;
 uint2 DownsampledViewSize;
 uint MaxShadowMaskTiles;
 float SamplingMinWeight;
+uint2 NumSamplesPerPixel;
 int DebugMode;
 
 RWTexture2D<uint2> RWShadowMaskHistoryScreenCoord00;
@@ -206,8 +210,8 @@ void GenerateSamplesCS(
 			FShaderPrintContext Context = InitShaderPrintContext(true, float2(0.05, 0.05));
 	#endif
 
-			uint SelectedLights[NUM_SAMPLES_PER_PIXEL];
-			for (uint LightSampleIndex = 0; LightSampleIndex < NUM_SAMPLES_PER_PIXEL; ++LightSampleIndex)
+			uint SelectedLights[NUM_SAMPLES_PER_PIXEL_1D];
+			for (uint LightSampleIndex = 0; LightSampleIndex < NUM_SAMPLES_PER_PIXEL_1D; ++LightSampleIndex)
 			{
 				SelectedLights[LightSampleIndex] = INVALID_LIGHT_SCENE_ID;
 			}
@@ -281,8 +285,8 @@ void GenerateSamplesCS(
 				const uint NumLightsInGridCell = min(CulledLightGridData.NumLights, GetMaxLightsPerCell(EyeIndex));
 				const uint NumLocalLights = GetNumLocalLights(EyeIndex);
 
-				float LightIndexRandom[NUM_SAMPLES_PER_PIXEL];
-				for (uint LightSampleIndex = 0; LightSampleIndex < NUM_SAMPLES_PER_PIXEL; ++LightSampleIndex)
+				float LightIndexRandom[NUM_SAMPLES_PER_PIXEL_1D];
+				for (uint LightSampleIndex = 0; LightSampleIndex < NUM_SAMPLES_PER_PIXEL_1D; ++LightSampleIndex)
 				{
 					LightIndexRandom[LightSampleIndex] = BlueNoiseScalar(DownsampledScreenCoord * DOWNSAMPLE_FACTOR + LightSampleOffsets[LightSampleIndex], StochasticShadowsStateFrameIndex);
 				}
@@ -323,7 +327,7 @@ void GenerateSamplesCS(
 				float WeightSum = 0.0f;
 
 				// Weighted blue noise sampling
-				float Xi[NUM_SAMPLES_PER_PIXEL] = LightIndexRandom;
+				float Xi[NUM_SAMPLES_PER_PIXEL_1D] = LightIndexRandom;
 
 	#if CANDIDATE_LIGHT_MASK
 				for (uint MaskIndex = 0; MaskIndex < LIGHT_SCENE_ID_MASK_SIZE; ++MaskIndex)
@@ -404,7 +408,7 @@ void GenerateSamplesCS(
 								float Tau = WeightSum / (WeightSum + SampleWeight);
 								WeightSum += SampleWeight;
 
-								for (uint LightSampleIndex = 0; LightSampleIndex < NUM_SAMPLES_PER_PIXEL; ++LightSampleIndex)
+								for (uint LightSampleIndex = 0; LightSampleIndex < NUM_SAMPLES_PER_PIXEL_1D; ++LightSampleIndex)
 								{
 									if (Xi[LightSampleIndex] < Tau)
 									{
@@ -497,7 +501,7 @@ void GenerateSamplesCS(
 								float Tau = WeightSum / (WeightSum + SampleWeight);
 								WeightSum += SampleWeight;
 
-								for (uint LightSampleIndex = 0; LightSampleIndex < NUM_SAMPLES_PER_PIXEL; ++LightSampleIndex)
+								for (uint LightSampleIndex = 0; LightSampleIndex < NUM_SAMPLES_PER_PIXEL_1D; ++LightSampleIndex)
 								{
 									if (Xi[LightSampleIndex] < Tau)
 									{
@@ -524,7 +528,7 @@ void GenerateSamplesCS(
 					Newline(Context);
 					Print(Context, TEXT("Selected lights: "));
 
-					for (uint LightSampleIndex = 0; LightSampleIndex < NUM_SAMPLES_PER_PIXEL; ++LightSampleIndex)
+					for (uint LightSampleIndex = 0; LightSampleIndex < NUM_SAMPLES_PER_PIXEL_1D; ++LightSampleIndex)
 					{
 						Print(Context, SelectedLights[LightSampleIndex]);
 
@@ -551,13 +555,13 @@ void GenerateSamplesCS(
 			RWDownsampledSceneDepth[DownsampledScreenCoord] = SceneDepth;
 			RWDownsampledSceneWorldNormal[DownsampledScreenCoord] = EncodeNormal(GBuffer.WorldNormal);
 
-			for (uint LightSampleIndex = 0; LightSampleIndex < NUM_SAMPLES_PER_PIXEL; ++LightSampleIndex)
+			for (uint LightSampleIndex = 0; LightSampleIndex < NUM_SAMPLES_PER_PIXEL_1D; ++LightSampleIndex)
 			{
 				FLightSample LightSample;
 				LightSample.LightSceneId = SelectedLights[LightSampleIndex];
 				LightSample.bCompleted = LightSample.LightSceneId == INVALID_LIGHT_SCENE_ID ? true : false;
 				LightSample.bVisible = true;
-				RWLightSamples[DownsampledScreenCoord * uint2(NUM_SAMPLES_PER_PIXEL_X, NUM_SAMPLES_PER_PIXEL_Y) + LightSampleOffsets[LightSampleIndex]] = PackLightSample(LightSample);
+				RWLightSamples[DownsampledScreenCoord * NumSamplesPerPixel + LightSampleOffsets[LightSampleIndex]] = PackLightSample(LightSample);
 			}
 		}
 
diff --git a/Engine/Shaders/Private/StochasticShadows/StochasticShadowsTracing.usf b/Engine/Shaders/Private/StochasticShadows/StochasticShadowsTracing.usf
--- a/Engine/Shaders/Private/StochasticShadows/StochasticShadowsTracing.usf
+++ b/Engine/Shaders/Private/StochasticShadows/StochasticShadowsTracing.usf
@@ -18,6 +18,7 @@
 
 StructuredBuffer<uint> CompactedTraceTexelAllocator;
 RWBuffer<uint> RWIndirectArgs;
+Texture2D<uint> DownsampledTileMask;
 
 [numthreads(THREADGROUP_SIZE, 1, 1)]
 void InitCompactedTraceTexelIndirectArgsCS(uint3 DispatchThreadId : SV_DispatchThreadID)
@@ -69,6 +70,9 @@ groupshared uint SharedGlobalTraceTexelStartOffset;
 	groupshared uint SharedTraceTexels[THREADGROUP_SIZE * THREADGROUP_SIZE];
 #endif
 
+// Allocated for max size, where we have only 1 sample per pixel
+groupshared uint SharedTileMask[THREADGROUP_SIZE / TILE_SIZE][THREADGROUP_SIZE / TILE_SIZE];
+
 #if COMPILER_SUPPORTS_WAVE_SIZE && WAVE_OPS
 	WAVESIZE(32)
 #endif
@@ -91,13 +95,22 @@ void CompactLightSampleTracesCS(
 	}
 	#endif
 
+	uint2 LocalTileMaskSize = SampleCoordToDownsampledScreenCoord(THREADGROUP_SIZE) / TILE_SIZE;
+	if (all(GroupThreadId.xy < LocalTileMaskSize))
+	{
+		// Skip unitialized SampleBuffer tiles
+		uint2 DownsampledTileCoord = SampleCoordToDownsampledScreenCoord(GroupId.xy * THREADGROUP_SIZE) / TILE_SIZE + GroupThreadId.xy;
+		SharedTileMask[GroupThreadId.x][GroupThreadId.y] = DownsampledTileMask[DownsampledTileCoord];
+	}
+
 	GroupMemoryBarrierWithGroupSync();
 
 	uint2 SampleCoord = DispatchThreadId.xy;
 	bool bTraceValid = false;
 	uint TraceTexelForThisThread = 0;
 
-	if (all(SampleCoord < SampleViewSize))
+	uint2 LocalTileMaskCoord = SampleCoordToDownsampledScreenCoord(SampleCoord - GroupId.xy * THREADGROUP_SIZE) / TILE_SIZE;
+	if (all(SampleCoord < SampleViewSize) && SharedTileMask[LocalTileMaskCoord.x][LocalTileMaskCoord.y] != 0)
 	{
 		// #sdl_todo: try to preserve z-order here for improved tracing coherency
 		FLightSample LightSample = UnpackLightSample(LightSamples[SampleCoord]);
diff --git a/Engine/Shaders/Private/StochasticShadows/StochasticShadowsTracing.ush b/Engine/Shaders/Private/StochasticShadows/StochasticShadowsTracing.ush
--- a/Engine/Shaders/Private/StochasticShadows/StochasticShadowsTracing.ush
+++ b/Engine/Shaders/Private/StochasticShadows/StochasticShadowsTracing.ush
@@ -6,13 +6,12 @@
 #include "../RayTracing/RayTracingSphereLight.ush"
 #include "../RayTracing/RayTracingRectLight.ush"
 
-#ifndef NUM_SAMPLES_PER_PIXEL_X
-	#error NUM_SAMPLES_PER_PIXEL_X must be defined
-#endif
+uint2 NumSamplesPerPixelDivideShift;
 
 uint2 SampleCoordToDownsampledScreenCoord(uint2 SampleCoord)
 {
-	return SampleCoord / uint2(NUM_SAMPLES_PER_PIXEL_X, NUM_SAMPLES_PER_PIXEL_Y);
+	// SampleCoord / NumSamplesPerPixel
+	return SampleCoord >> NumSamplesPerPixelDivideShift;
 }
 
 uint2 SampleCoordToScreenCoord(uint2 SampleCoord)
diff --git a/Engine/Source/Runtime/Renderer/Private/StochasticShadows/StochasticShadows.cpp b/Engine/Source/Runtime/Renderer/Private/StochasticShadows/StochasticShadows.cpp
--- a/Engine/Source/Runtime/Renderer/Private/StochasticShadows/StochasticShadows.cpp
+++ b/Engine/Source/Runtime/Renderer/Private/StochasticShadows/StochasticShadows.cpp
@@ -163,14 +163,9 @@ namespace StochasticShadows
 		return StateFrameIndex;
 	}
 
-	uint32 GetNumSamplesPerPixel()
-	{
-		const uint32 NumSamplesPerPixel = FMath::RoundUpToPowerOfTwo(FMath::Clamp(CVarStochasticShadowsNumSamplesPerPixel.GetValueOnRenderThread(), 1, 4));
-		return NumSamplesPerPixel;
-	}
-
-	FIntPoint GetNumSamplesPerPixel2d(uint32 NumSamplesPerPixel1d)
+	FIntPoint GetNumSamplesPerPixel2d()
 	{
+		const uint32 NumSamplesPerPixel1d = FMath::RoundUpToPowerOfTwo(FMath::Clamp(CVarStochasticShadowsNumSamplesPerPixel.GetValueOnRenderThread(), 1, 4));
 		return NumSamplesPerPixel1d == 4 ? FIntPoint(2, 2) : (NumSamplesPerPixel1d == 2 ? FIntPoint(2, 1) : FIntPoint(1, 1));
 	}
 
@@ -204,6 +199,7 @@ class FTileClassificationCS : public FGlobalShader
 		SHADER_PARAMETER_STRUCT_INCLUDE(FStochasticShadowsParameters, StochasticShadowsParameters)
 		SHADER_PARAMETER_RDG_BUFFER_UAV(RWStructuredBuffer<uint>, RWTileAllocator)
 		SHADER_PARAMETER_RDG_BUFFER_UAV(RWStructuredBuffer<uint>, RWTileData)
+		SHADER_PARAMETER_RDG_TEXTURE_UAV(RWTexture2D<uint>, RWDownsampledTileMask)
 	END_SHADER_PARAMETER_STRUCT()
 
 	class FDownsampledClassification : SHADER_PERMUTATION_BOOL("DOWNSAMPLED_CLASSIFICATION");
@@ -318,12 +314,12 @@ class FGenerateSamplesCS : public FGlobalShader
 	END_SHADER_PARAMETER_STRUCT()
 
 	class FTileType : SHADER_PERMUTATION_INT("TILE_TYPE", (int32)StochasticShadows::ETileType::MAX);
-	class FNumSamplesPerPixel : SHADER_PERMUTATION_SPARSE_INT("NUM_SAMPLES_PER_PIXEL", 1, 2, 4);
+	class FNumSamplesPerPixel1d : SHADER_PERMUTATION_SPARSE_INT("NUM_SAMPLES_PER_PIXEL_1D", 1, 2, 4);
 	class FShadowMaskReprojectionWeights : SHADER_PERMUTATION_BOOL("SHADOW_MASK_REPROJECTION_WEIGHTS");
 	class FShadowFactorEstimate : SHADER_PERMUTATION_BOOL("SHADOW_FACTOR_ESTIMATE");
 	class FCandidateLightMask : SHADER_PERMUTATION_BOOL("CANDIDATE_LIGHT_MASK");
 	class FDebugMode : SHADER_PERMUTATION_BOOL("DEBUG_MODE");
-	using FPermutationDomain = TShaderPermutationDomain<FTileType, FNumSamplesPerPixel, FShadowMaskReprojectionWeights, FShadowFactorEstimate, FCandidateLightMask, FDebugMode>;
+	using FPermutationDomain = TShaderPermutationDomain<FTileType, FNumSamplesPerPixel1d, FShadowMaskReprojectionWeights, FShadowFactorEstimate, FCandidateLightMask, FDebugMode>;
 
 	static int32 GetGroupSize()
 	{	
@@ -357,11 +353,6 @@ class FGenerateSamplesCS : public FGlobalShader
 		FGlobalShader::ModifyCompilationEnvironment(Parameters, OutEnvironment);
 		FForwardLightingParameters::ModifyCompilationEnvironment(Parameters.Platform, OutEnvironment);
 		OutEnvironment.SetDefine(TEXT("THREADGROUP_SIZE"), GetGroupSize());
-
-		FPermutationDomain PermutationVector(Parameters.PermutationId);
-		const uint32 NumSamplesPerPixel = PermutationVector.Get<FNumSamplesPerPixel>();
-		OutEnvironment.SetDefine(TEXT("NUM_SAMPLES_PER_PIXEL_X"), StochasticShadows::GetNumSamplesPerPixel2d(NumSamplesPerPixel).X);
-		OutEnvironment.SetDefine(TEXT("NUM_SAMPLES_PER_PIXEL_Y"), StochasticShadows::GetNumSamplesPerPixel2d(NumSamplesPerPixel).Y);
 	}
 };
 
@@ -423,10 +414,9 @@ class FCompositeShadowMaskTracesCS : public FGlobalShader
 		return 8;
 	}
 
-	class FNumSamplesPerPixel : SHADER_PERMUTATION_SPARSE_INT("NUM_SAMPLES_PER_PIXEL", 1, 2, 4);
 	class FTemporalAccumulation : SHADER_PERMUTATION_BOOL("TEMPORAL_ACCUMULATION");
 	class FDebugMode : SHADER_PERMUTATION_BOOL("DEBUG_MODE");
-	using FPermutationDomain = TShaderPermutationDomain<FNumSamplesPerPixel, FTemporalAccumulation, FDebugMode>;
+	using FPermutationDomain = TShaderPermutationDomain<FTemporalAccumulation, FDebugMode>;
 
 	static bool ShouldCompilePermutation(const FGlobalShaderPermutationParameters& Parameters)
 	{
@@ -438,11 +428,6 @@ class FCompositeShadowMaskTracesCS : public FGlobalShader
 		FGlobalShader::ModifyCompilationEnvironment(Parameters, OutEnvironment);
 		FForwardLightingParameters::ModifyCompilationEnvironment(Parameters.Platform, OutEnvironment);
 		OutEnvironment.SetDefine(TEXT("THREADGROUP_SIZE"), GetGroupSize());
-
-		FPermutationDomain PermutationVector(Parameters.PermutationId);
-		const uint32 NumSamplesPerPixel = PermutationVector.Get<FNumSamplesPerPixel>();
-		OutEnvironment.SetDefine(TEXT("NUM_SAMPLES_PER_PIXEL_X"), StochasticShadows::GetNumSamplesPerPixel2d(NumSamplesPerPixel).X);
-		OutEnvironment.SetDefine(TEXT("NUM_SAMPLES_PER_PIXEL_Y"), StochasticShadows::GetNumSamplesPerPixel2d(NumSamplesPerPixel).Y);
 	}
 };
 
@@ -569,8 +554,7 @@ void FDeferredShadingSceneRenderer::RenderStochasticShadows(FRDGBuilder& GraphBu
 		bResetHistory = true;
 	}
 
-	const int32 NumSamplesPerPixel1d = StochasticShadows::GetNumSamplesPerPixel();
-	const FIntPoint NumSamplesPerPixel2d = StochasticShadows::GetNumSamplesPerPixel2d(NumSamplesPerPixel1d);
+	const FIntPoint NumSamplesPerPixel2d = StochasticShadows::GetNumSamplesPerPixel2d();
 
 	const uint32 DownsampleFactor = 2;
 	const uint32 ShadowMaskDownsampleFactor = 2;
@@ -659,6 +643,10 @@ void FDeferredShadingSceneRenderer::RenderStochasticShadows(FRDGBuilder& GraphBu
 	const FIntPoint DownsampledViewSizeInTiles = FIntPoint::DivideAndRoundUp(DownsampledViewSize, StochasticShadows::TileSize);
 	const int32 DownsampledTileDataStride = DownsampledViewSizeInTiles.X * DownsampledViewSizeInTiles.Y;
 
+	FRDGTextureRef DownsampledTileMask = GraphBuilder.CreateTexture(
+		FRDGTextureDesc::Create2D(FMath::DivideAndRoundUp<FIntPoint>(DownsampledBufferSize, StochasticShadows::TileSize), PF_R8_UINT, FClearValueBinding::Black, TexCreate_ShaderResource | TexCreate_UAV),
+		TEXT("StochasticShadows.DownsampledTileMask"));
+
 	FStochasticShadowsParameters StochasticShadowsParameters;
 	{
 		StochasticShadowsParameters.ViewUniformBuffer = View.ViewUniformBuffer;
@@ -669,8 +657,12 @@ void FDeferredShadingSceneRenderer::RenderStochasticShadows(FRDGBuilder& GraphBu
 		StochasticShadowsParameters.BlueNoise = BlueNoiseUniformBuffer;
 		StochasticShadowsParameters.DownsampledViewSize = DownsampledViewSize;
 		StochasticShadowsParameters.SampleViewSize = SampleViewSize;
-		StochasticShadowsParameters.ShadowMaskViewSize = ShadowMaskViewSize;
+		StochasticShadowsParameters.ShadowMaskViewSize = ShadowMaskViewSize;	
+		StochasticShadowsParameters.NumSamplesPerPixel = NumSamplesPerPixel2d;
+		StochasticShadowsParameters.NumSamplesPerPixelDivideShift.X = FMath::FloorLog2(NumSamplesPerPixel2d.X);
+		StochasticShadowsParameters.NumSamplesPerPixelDivideShift.Y = FMath::FloorLog2(NumSamplesPerPixel2d.Y);
 		StochasticShadowsParameters.StochasticShadowsStateFrameIndex = StochasticShadows::GetStateFrameIndex(View.ViewState);
+		StochasticShadowsParameters.DownsampledTileMask = DownsampledTileMask;
 		StochasticShadowsParameters.DownsampledSceneDepth = DownsampledSceneDepth;
 		StochasticShadowsParameters.DownsampledSceneWorldNormal = DownsampledSceneWorldNormal;
 		StochasticShadowsParameters.MaxShadowMaskTiles = MaxShadowMaskTiles;
@@ -702,6 +694,7 @@ void FDeferredShadingSceneRenderer::RenderStochasticShadows(FRDGBuilder& GraphBu
 	FRDGBufferRef DownsampledTileData = GraphBuilder.CreateBuffer(FRDGBufferDesc::CreateStructuredDesc(sizeof(uint32), DownsampledTileDataStride * (int32)StochasticShadows::ETileType::MAX), TEXT("StochasticShadows.DownsampledTileData"));
 	AddClearUAVPass(GraphBuilder, GraphBuilder.CreateUAV(DownsampledTileAllocator), 0);
 
+	// #sdl_todo: merge classification passes or reuse downsampled one to create full res tiles
 	// Run tile classification to generate tiles for the subsequent passes
 	{
 		{
@@ -709,6 +702,7 @@ void FDeferredShadingSceneRenderer::RenderStochasticShadows(FRDGBuilder& GraphBu
 			PassParameters->StochasticShadowsParameters = StochasticShadowsParameters;
 			PassParameters->RWTileAllocator = GraphBuilder.CreateUAV(TileAllocator);
 			PassParameters->RWTileData = GraphBuilder.CreateUAV(TileData);
+			PassParameters->RWDownsampledTileMask = nullptr;
 
 			FTileClassificationCS::FPermutationDomain PermutationVector;
 			PermutationVector.Set<FTileClassificationCS::FDownsampledClassification>(false);
@@ -729,6 +723,7 @@ void FDeferredShadingSceneRenderer::RenderStochasticShadows(FRDGBuilder& GraphBu
 			PassParameters->StochasticShadowsParameters = StochasticShadowsParameters;
 			PassParameters->RWTileAllocator = GraphBuilder.CreateUAV(DownsampledTileAllocator);
 			PassParameters->RWTileData = GraphBuilder.CreateUAV(DownsampledTileData);
+			PassParameters->RWDownsampledTileMask = GraphBuilder.CreateUAV(DownsampledTileMask);
 
 			FTileClassificationCS::FPermutationDomain PermutationVector;
 			PermutationVector.Set<FTileClassificationCS::FDownsampledClassification>(true);
@@ -815,7 +810,7 @@ void FDeferredShadingSceneRenderer::RenderStochasticShadows(FRDGBuilder& GraphBu
 
 			FGenerateSamplesCS::FPermutationDomain PermutationVector;
 			PermutationVector.Set<FGenerateSamplesCS::FTileType>(TileType);
-			PermutationVector.Set<FGenerateSamplesCS::FNumSamplesPerPixel>(NumSamplesPerPixel1d);
+			PermutationVector.Set<FGenerateSamplesCS::FNumSamplesPerPixel1d>(NumSamplesPerPixel2d.X * NumSamplesPerPixel2d.Y);
 			PermutationVector.Set<FGenerateSamplesCS::FShadowMaskReprojectionWeights>(bTemporal);
 			PermutationVector.Set<FGenerateSamplesCS::FShadowFactorEstimate>(bTemporal && CVarStochasticSamplingShadowEstimate.GetValueOnRenderThread() != 0);
 			PermutationVector.Set<FGenerateSamplesCS::FDebugMode>(bDebug);
@@ -922,7 +917,6 @@ void FDeferredShadingSceneRenderer::RenderStochasticShadows(FRDGBuilder& GraphBu
 		PassParameters->ForwardLightData = View.ForwardLightingResources.ForwardLightUniformBuffer;
 
 		FCompositeShadowMaskTracesCS::FPermutationDomain PermutationVector;
-		PermutationVector.Set<FCompositeShadowMaskTracesCS::FNumSamplesPerPixel>(NumSamplesPerPixel1d);
 		PermutationVector.Set<FCompositeShadowMaskTracesCS::FTemporalAccumulation>(bTemporal);
 		PermutationVector.Set<FCompositeShadowMaskTracesCS::FDebugMode>(bDebug);
 		auto ComputeShader = View.ShaderMap->GetShader<FCompositeShadowMaskTracesCS>(PermutationVector);
diff --git a/Engine/Source/Runtime/Renderer/Private/StochasticShadows/StochasticShadowsInternal.h b/Engine/Source/Runtime/Renderer/Private/StochasticShadows/StochasticShadowsInternal.h
--- a/Engine/Source/Runtime/Renderer/Private/StochasticShadows/StochasticShadowsInternal.h
+++ b/Engine/Source/Runtime/Renderer/Private/StochasticShadows/StochasticShadowsInternal.h
@@ -17,6 +17,8 @@ BEGIN_SHADER_PARAMETER_STRUCT(FStochasticShadowsParameters, )
 	SHADER_PARAMETER(FIntPoint, SampleViewSize)
 	SHADER_PARAMETER(FIntPoint, DownsampledViewSize)
 	SHADER_PARAMETER(FIntPoint, ShadowMaskViewSize)
+	SHADER_PARAMETER(FIntPoint, NumSamplesPerPixel)
+	SHADER_PARAMETER(FIntPoint, NumSamplesPerPixelDivideShift)
 	SHADER_PARAMETER(uint32, StochasticShadowsStateFrameIndex)
 	SHADER_PARAMETER(uint32, MaxShadowMaskTiles)
 	SHADER_PARAMETER(uint32, MaxShadingTiles)
@@ -30,6 +32,7 @@ BEGIN_SHADER_PARAMETER_STRUCT(FStochasticShadowsParameters, )
 	SHADER_PARAMETER(int32, TemporalMaxFramesAccumulated)
 	SHADER_PARAMETER(float, TemporalStdDevOffset)
 	SHADER_PARAMETER(int32, DebugMode)
+	SHADER_PARAMETER_RDG_TEXTURE(Texture2D<uint>, DownsampledTileMask)
 	SHADER_PARAMETER_RDG_TEXTURE(Texture2D<float>, DownsampledSceneDepth)
 	SHADER_PARAMETER_RDG_TEXTURE(Texture2D<UNORM float3>, DownsampledSceneWorldNormal)
 END_SHADER_PARAMETER_STRUCT()
@@ -50,6 +53,4 @@ namespace StochasticShadows
 	bool ShouldCompileShaders(const FGlobalShaderPermutationParameters& Parameters);
 	bool UseWaveOps(EShaderPlatform ShaderPlatform);
 	int32 GetDebugMode();
-	uint32 GetNumSamplesPerPixel();
-	FIntPoint GetNumSamplesPerPixel2d(uint32 NumSamplesPerPixel1d);
 };
\ No newline at end of file
diff --git a/Engine/Source/Runtime/Renderer/Private/StochasticShadows/StochasticShadowsRayTracing.cpp b/Engine/Source/Runtime/Renderer/Private/StochasticShadows/StochasticShadowsRayTracing.cpp
--- a/Engine/Source/Runtime/Renderer/Private/StochasticShadows/StochasticShadowsRayTracing.cpp
+++ b/Engine/Source/Runtime/Renderer/Private/StochasticShadows/StochasticShadowsRayTracing.cpp
@@ -230,9 +230,8 @@ class FHardwareRayTraceLightSamples : public FLumenHardwareRayTracingShaderBase
 		SHADER_PARAMETER_RDG_UNIFORM_BUFFER(FLumenHardwareRayTracingUniformBufferParameters, LumenHardwareRayTracingUniformBuffer)
 	END_SHADER_PARAMETER_STRUCT()
 
-	class FNumSamplesPerPixel : SHADER_PERMUTATION_SPARSE_INT("NUM_SAMPLES_PER_PIXEL", 1, 2, 4);
 	class FDebugMode : SHADER_PERMUTATION_BOOL("DEBUG_MODE");
-	using FPermutationDomain = TShaderPermutationDomain<FNumSamplesPerPixel, FDebugMode>;
+	using FPermutationDomain = TShaderPermutationDomain<FDebugMode>;
 
 	static bool ShouldCompilePermutation(const FGlobalShaderPermutationParameters& Parameters, Lumen::ERayTracingShaderDispatchType ShaderDispatchType)
 	{
@@ -243,11 +242,6 @@ class FHardwareRayTraceLightSamples : public FLumenHardwareRayTracingShaderBase
 	static void ModifyCompilationEnvironment(const FGlobalShaderPermutationParameters& Parameters, Lumen::ERayTracingShaderDispatchType ShaderDispatchType, FShaderCompilerEnvironment& OutEnvironment)
 	{
 		FLumenHardwareRayTracingShaderBase::ModifyCompilationEnvironment(Parameters, ShaderDispatchType, Lumen::ESurfaceCacheSampling::AlwaysResidentPagesWithoutFeedback, OutEnvironment);
-
-		FPermutationDomain PermutationVector(Parameters.PermutationId);
-		const uint32 NumSamplesPerPixel = PermutationVector.Get<FNumSamplesPerPixel>();
-		OutEnvironment.SetDefine(TEXT("NUM_SAMPLES_PER_PIXEL_X"), StochasticShadows::GetNumSamplesPerPixel2d(NumSamplesPerPixel).X);
-		OutEnvironment.SetDefine(TEXT("NUM_SAMPLES_PER_PIXEL_Y"), StochasticShadows::GetNumSamplesPerPixel2d(NumSamplesPerPixel).Y);
 	}
 
 	static ERayTracingPayloadType GetRayTracingPayloadType(const int32 PermutationId)
@@ -280,9 +274,8 @@ class FSoftwareRayTraceLightSamplesCS : public FGlobalShader
 		return 64;
 	}
 
-	class FNumSamplesPerPixel : SHADER_PERMUTATION_SPARSE_INT("NUM_SAMPLES_PER_PIXEL", 1, 2, 4);
 	class FDebugMode : SHADER_PERMUTATION_BOOL("DEBUG_MODE");
-	using FPermutationDomain = TShaderPermutationDomain<FNumSamplesPerPixel, FDebugMode>;
+	using FPermutationDomain = TShaderPermutationDomain<FDebugMode>;
 
 	static bool ShouldCompilePermutation(const FGlobalShaderPermutationParameters& Parameters)
 	{
@@ -297,11 +290,6 @@ class FSoftwareRayTraceLightSamplesCS : public FGlobalShader
 
 		// GPU Scene definitions
 		OutEnvironment.SetDefine(TEXT("VF_SUPPORTS_PRIMITIVE_SCENE_DATA"), 1);
-
-		FPermutationDomain PermutationVector(Parameters.PermutationId);
-		const uint32 NumSamplesPerPixel = PermutationVector.Get<FNumSamplesPerPixel>();
-		OutEnvironment.SetDefine(TEXT("NUM_SAMPLES_PER_PIXEL_X"), StochasticShadows::GetNumSamplesPerPixel2d(NumSamplesPerPixel).X);
-		OutEnvironment.SetDefine(TEXT("NUM_SAMPLES_PER_PIXEL_Y"), StochasticShadows::GetNumSamplesPerPixel2d(NumSamplesPerPixel).Y);
 	}
 };
 
@@ -329,9 +317,8 @@ class FScreenSpaceRayTraceLightSamplesCS : public FGlobalShader
 		return 64;
 	}
 
-	class FNumSamplesPerPixel : SHADER_PERMUTATION_SPARSE_INT("NUM_SAMPLES_PER_PIXEL", 1, 2, 4);
 	class FDebugMode : SHADER_PERMUTATION_BOOL("DEBUG_MODE");
-	using FPermutationDomain = TShaderPermutationDomain<FNumSamplesPerPixel, FDebugMode>;
+	using FPermutationDomain = TShaderPermutationDomain<FDebugMode>;
 
 	static bool ShouldCompilePermutation(const FGlobalShaderPermutationParameters& Parameters)
 	{
@@ -343,11 +330,6 @@ class FScreenSpaceRayTraceLightSamplesCS : public FGlobalShader
 		FGlobalShader::ModifyCompilationEnvironment(Parameters, OutEnvironment);
 		OutEnvironment.CompilerFlags.Add(CFLAG_Wave32);
 		OutEnvironment.SetDefine(TEXT("THREADGROUP_SIZE"), GetGroupSize());
-
-		FPermutationDomain PermutationVector(Parameters.PermutationId);
-		const uint32 NumSamplesPerPixel = PermutationVector.Get<FNumSamplesPerPixel>();
-		OutEnvironment.SetDefine(TEXT("NUM_SAMPLES_PER_PIXEL_X"), StochasticShadows::GetNumSamplesPerPixel2d(NumSamplesPerPixel).X);
-		OutEnvironment.SetDefine(TEXT("NUM_SAMPLES_PER_PIXEL_Y"), StochasticShadows::GetNumSamplesPerPixel2d(NumSamplesPerPixel).Y);
 	}
 };
 
@@ -359,7 +341,6 @@ void FDeferredShadingSceneRenderer::PrepareStochasticShadowsLumenMaterial(const
 	if (StochasticShadows::IsEnabled() && StochasticShadows::UseHardwareRayTracing())
 	{
 		FHardwareRayTraceLightSamplesRGS::FPermutationDomain PermutationVector;
-		PermutationVector.Set<FHardwareRayTraceLightSamplesRGS::FNumSamplesPerPixel>(StochasticShadows::GetNumSamplesPerPixel());
 		PermutationVector.Set<FHardwareRayTraceLightSamplesRGS::FDebugMode>(StochasticShadows::GetDebugMode() != 0);
 		TShaderRef<FHardwareRayTraceLightSamplesRGS> RayGenerationShader = View.ShaderMap->GetShader<FHardwareRayTraceLightSamplesRGS>(PermutationVector);
 		OutRayGenShaders.Add(RayGenerationShader.GetRayTracingShader());
@@ -489,8 +470,6 @@ void StochasticShadows::RayTraceLightSamples(
 	const FStochasticShadowsParameters& StochasticShadowsParameters)
 {
 	const bool bDebug = StochasticShadows::GetDebugMode() != 0;
-	const int32 NumSamplesPerPixel1d = StochasticShadows::GetNumSamplesPerPixel();
-	const FIntPoint NumSamplesPerPixel2d = StochasticShadows::GetNumSamplesPerPixel2d(NumSamplesPerPixel1d);
 
 	if (CVarStochasticShadowsScreenTraces.GetValueOnRenderThread() != 0)
 	{
@@ -513,7 +492,6 @@ void StochasticShadows::RayTraceLightSamples(
 		PassParameters->MinimumTracingThreadOccupancy = CVarStochasticShadowsScreenTracesMinimumOccupancy.GetValueOnRenderThread();
 
 		FScreenSpaceRayTraceLightSamplesCS::FPermutationDomain PermutationVector;
-		PermutationVector.Set<FScreenSpaceRayTraceLightSamplesCS::FNumSamplesPerPixel>(NumSamplesPerPixel1d);
 		PermutationVector.Set<FScreenSpaceRayTraceLightSamplesCS::FDebugMode>(bDebug);
 		auto ComputeShader = View.ShaderMap->GetShader<FScreenSpaceRayTraceLightSamplesCS>(PermutationVector);
 
@@ -555,7 +533,6 @@ void StochasticShadows::RayTraceLightSamples(
 					PassParameters);
 
 				FHardwareRayTraceLightSamplesCS::FPermutationDomain PermutationVector;
-				PermutationVector.Set<FHardwareRayTraceLightSamplesCS::FNumSamplesPerPixel>(NumSamplesPerPixel1d);
 				PermutationVector.Set<FHardwareRayTraceLightSamplesCS::FDebugMode>(bDebug);
 				auto ComputeShader = View.ShaderMap->GetShader<FHardwareRayTraceLightSamplesCS>(PermutationVector);
 
@@ -581,7 +558,6 @@ void StochasticShadows::RayTraceLightSamples(
 					PassParameters);
 
 				FHardwareRayTraceLightSamplesRGS::FPermutationDomain PermutationVector;
-				PermutationVector.Set<FHardwareRayTraceLightSamplesRGS::FNumSamplesPerPixel>(NumSamplesPerPixel1d);
 				PermutationVector.Set<FHardwareRayTraceLightSamplesRGS::FDebugMode>(bDebug);
 				auto RayGenerationShader = View.ShaderMap->GetShader<FHardwareRayTraceLightSamplesRGS>(PermutationVector);
 
@@ -608,7 +584,6 @@ void StochasticShadows::RayTraceLightSamples(
 			PassParameters->LightSampleRayDistance = LightSampleRayDistance;
 
 			FSoftwareRayTraceLightSamplesCS::FPermutationDomain PermutationVector;
-			PermutationVector.Set<FSoftwareRayTraceLightSamplesCS::FNumSamplesPerPixel>(NumSamplesPerPixel1d);
 			PermutationVector.Set<FSoftwareRayTraceLightSamplesCS::FDebugMode>(bDebug);
 			auto ComputeShader = View.ShaderMap->GetShader<FSoftwareRayTraceLightSamplesCS>(PermutationVector);
 
