SLW - Enabling translucent ordering fix with and fog cloud under water.

diff --git a/Engine/Source/Runtime/Renderer/Private/DeferredShadingRenderer.cpp b/Engine/Source/Runtime/Renderer/Private/DeferredShadingRenderer.cpp
--- a/Engine/Source/Runtime/Renderer/Private/DeferredShadingRenderer.cpp
+++ b/Engine/Source/Runtime/Renderer/Private/DeferredShadingRenderer.cpp
@@ -323,11 +323,6 @@ static FAutoConsoleCommand RecreateRenderStateContextCmd(
 	TEXT("Recreate render state."),
 	FConsoleCommandDelegate::CreateStatic([] { FGlobalComponentRecreateRenderStateContext Context; }));
 
-static TAutoConsoleVariable<int32> CVarSingleLayerWaterUnderWaterExpFix(
-	TEXT("r.SingleLayerWater.UnderWaterFix"), 0,
-	TEXT("Experimental fix for sky, fog, cloud to be correctly ordered with translucent element. If proven valid we will enable it definintely."),
-	ECVF_RenderThreadSafe);
-
 #if RHI_RAYTRACING
 
 static bool bUpdateCachedRayTracingState = false;
@@ -4079,7 +4074,6 @@ void FDeferredShadingSceneRenderer::Render(FRDGBuilder& GraphBuilder)
 
 		FTranslucencyPassResourcesMap TranslucencyResourceMap(Views.Num());
 
-		const bool bSingleLayerWaterUnderWaterExpFix = CVarSingleLayerWaterUnderWaterExpFix.GetValueOnRenderThread() > 0;
 		const bool bIsCameraUnderWater = EnumHasAnyFlags(TranslucencyViewsToRender, ETranslucencyView::UnderWater);
 		FRDGTextureRef LightShaftOcclusionTexture = nullptr;
 		const bool bShouldRenderSingleLayerWater = !bHasRayTracedOverlay && ShouldRenderSingleLayerWater(Views);
@@ -4132,12 +4126,12 @@ void FDeferredShadingSceneRenderer::Render(FRDGBuilder& GraphBuilder)
 					HalfResolutionDepthCheckerboardMinMaxTexture, QuarterResolutionDepthMinMaxTexture, false, InstanceCullingManager);
 			}
 
-			// or composite the off screen buffer over the scene.
+			// Or composite the off screen buffer over the scene.
 			if (bVolumetricRenderTargetRequired)
 			{
 				ComposeVolumetricRenderTargetOverScene(
 					GraphBuilder, Views, SceneTextures.Color.Target, SceneTextures.Depth.Target,
-					bSingleLayerWaterUnderWaterExpFix ? (bIsCameraUnderWater ? false : bShouldRenderSingleLayerWater) : bShouldRenderSingleLayerWater,
+					bIsCameraUnderWater ? false : bShouldRenderSingleLayerWater,
 					SceneWithoutWaterTextures, SceneTextures);
 			}
 		};
@@ -4146,10 +4140,7 @@ void FDeferredShadingSceneRenderer::Render(FRDGBuilder& GraphBuilder)
 		{
 			if (bIsCameraUnderWater)
 			{
-				if (bSingleLayerWaterUnderWaterExpFix)
-				{
-					RenderLigthShaftSkyFogAndCloud();
-				}
+				RenderLigthShaftSkyFogAndCloud();
 
 				RDG_CSV_STAT_EXCLUSIVE_SCOPE(GraphBuilder, RenderTranslucency);
 				SCOPED_NAMED_EVENT(RenderTranslucency, FColor::Emerald);
@@ -4161,8 +4152,7 @@ void FDeferredShadingSceneRenderer::Render(FRDGBuilder& GraphBuilder)
 			}
 
 			GraphBuilder.SetCommandListStat(GET_STATID(STAT_CLM_WaterPass));
-			RenderSingleLayerWater(GraphBuilder, SceneTextures, SingleLayerWaterPrePassResult, bShouldRenderVolumetricCloud, SceneWithoutWaterTextures, LumenFrameTemporaries, 
-				bSingleLayerWaterUnderWaterExpFix ? bIsCameraUnderWater : false); // false is the default value when the fix is not active.
+			RenderSingleLayerWater(GraphBuilder, SceneTextures, SingleLayerWaterPrePassResult, bShouldRenderVolumetricCloud, SceneWithoutWaterTextures, LumenFrameTemporaries, bIsCameraUnderWater);
 
 			// Replace main depth texture with the output of the SLW depth prepass which contains the scene + water.
 			// Note: Stencil now has all water bits marked with 1. As long as no other passes after this point want to read the depth buffer,
@@ -4176,7 +4166,7 @@ void FDeferredShadingSceneRenderer::Render(FRDGBuilder& GraphBuilder)
 		// Rebuild scene textures to include scene color.
 		SceneTextures.UniformBuffer = CreateSceneTextureUniformBuffer(GraphBuilder, &SceneTextures, FeatureLevel, SceneTextures.SetupMode);
 
-		if (!bSingleLayerWaterUnderWaterExpFix || !bIsCameraUnderWater)
+		if (!bIsCameraUnderWater)
 		{
 			RenderLigthShaftSkyFogAndCloud();
 		}
