Fix Global SDF being built due to incorrect Stochastic Shadows check.

diff --git a/Engine/Source/Runtime/Renderer/Private/DistanceFieldAmbientOcclusion.cpp b/Engine/Source/Runtime/Renderer/Private/DistanceFieldAmbientOcclusion.cpp
--- a/Engine/Source/Runtime/Renderer/Private/DistanceFieldAmbientOcclusion.cpp
+++ b/Engine/Source/Runtime/Renderer/Private/DistanceFieldAmbientOcclusion.cpp
@@ -782,7 +782,8 @@ bool FSceneRenderer::ShouldPrepareGlobalDistanceField() const
 
 	const bool bShouldPrepareForLumen = IsLumenEnabled(Views[0]) && Lumen::UseGlobalSDFObjectGrid(*Views[0].Family);
 
-	const bool bShouldPrepareForStochasticShadows = StochasticShadows::UseGlobalSDF();
+	// TODO: Should check any light actually uses Stochastic Shadows
+	const bool bShouldPrepareForStochasticShadows = StochasticShadows::IsEnabled() && StochasticShadows::UseGlobalSDF();
 
 	const bool bShouldPrepareForVisualization = ViewFamily.EngineShowFlags.VisualizeGlobalDistanceField;
 
