Disabled IsHairStrandsForwardRasterSupported for Mac due to Spirv-v

diff --git a/Engine/Source/Runtime/Renderer/Private/HairStrands/HairStrandsForwardRaster.cpp b/Engine/Source/Runtime/Renderer/Private/HairStrands/HairStrandsForwardRaster.cpp
--- a/Engine/Source/Runtime/Renderer/Private/HairStrands/HairStrandsForwardRaster.cpp
+++ b/Engine/Source/Runtime/Renderer/Private/HairStrands/HairStrandsForwardRaster.cpp
@@ -28,7 +28,7 @@ FMinHairRadiusAtDepth1 ComputeMinStrandRadiusAtDepth1(const FIntPoint& Resolutio
 inline bool IsHairStrandsForwardRasterSupported(EShaderPlatform In) 
 { 
 	return IsFeatureLevelSupported(In, ERHIFeatureLevel::SM6) && IsHairStrandsSupported(EHairStrandsShaderType::Strands, In) && 
-		!IsVulkanPlatform(In); // :todo-jn: fix SPIR-V error during compilation
+		!IsVulkanPlatform(In) && !IsMetalPlatform(In); // :todo-jn: fix SPIR-V error during compilation (Mac and Vulkan)
 }
 
 /////////////////////////////////////////////////////////////////////////////////////////
