Fix incorrect hair macro group building with groom not in primary view.

diff --git a/Engine/Source/Runtime/Renderer/Private/DeferredShadingRenderer.cpp b/Engine/Source/Runtime/Renderer/Private/DeferredShadingRenderer.cpp
--- a/Engine/Source/Runtime/Renderer/Private/DeferredShadingRenderer.cpp
+++ b/Engine/Source/Runtime/Renderer/Private/DeferredShadingRenderer.cpp
@@ -3553,7 +3553,7 @@ void FDeferredShadingSceneRenderer::Render(FRDGBuilder& GraphBuilder)
 				RunHairStrandsBookmark(GraphBuilder, EHairStrandsBookmark::ProcessStrandsInterpolation, HairStrandsBookmarkParameters);
 				if (!bHasRayTracedOverlay)
 				{
-					RenderHairPrePass(GraphBuilder, Scene, Views, InstanceCullingManager);
+					RenderHairPrePass(GraphBuilder, Scene, Views, InstanceCullingManager, HairStrandsBookmarkParameters.InstancesVisibilityType);
 					RenderHairBasePass(GraphBuilder, Scene, SceneTextures, Views, InstanceCullingManager);
 				}
 			}
@@ -3727,7 +3727,7 @@ void FDeferredShadingSceneRenderer::Render(FRDGBuilder& GraphBuilder)
 			RunHairStrandsBookmark(GraphBuilder, EHairStrandsBookmark::ProcessStrandsInterpolation, HairStrandsBookmarkParameters);
 			if (!bHasRayTracedOverlay)
 			{
-				RenderHairPrePass(GraphBuilder, Scene, Views, InstanceCullingManager);
+				RenderHairPrePass(GraphBuilder, Scene, Views, InstanceCullingManager, HairStrandsBookmarkParameters.InstancesVisibilityType);
 				RenderHairBasePass(GraphBuilder, Scene, SceneTextures, Views, InstanceCullingManager);
 			}
 		}
diff --git a/Engine/Source/Runtime/Renderer/Private/HairStrands/HairStrandsCluster.cpp b/Engine/Source/Runtime/Renderer/Private/HairStrands/HairStrandsCluster.cpp
--- a/Engine/Source/Runtime/Renderer/Private/HairStrands/HairStrandsCluster.cpp
+++ b/Engine/Source/Runtime/Renderer/Private/HairStrands/HairStrandsCluster.cpp
@@ -120,7 +120,6 @@ static bool DoesGroupExists(uint32 ResourceId, uint32 GroupIndex, const FHairStr
 }
 
 bool IsHairStrandsNonVisibleShadowCastingEnable();
-bool IsHairStrandsVisibleInShadows(const FViewInfo& View, const FHairStrandsInstance& Instance);
 
 static void InternalUpdateMacroGroup(FHairStrandsMacroGroupData& MacroGroup, int32& MaterialId, FHairGroupPublicData* HairData, const FMeshBatch* Mesh, const FPrimitiveSceneProxy* Proxy)
 {
@@ -149,6 +148,7 @@ void CreateHairStrandsMacroGroups(
 	FRDGBuilder& GraphBuilder,
 	const FScene* Scene,
 	const FViewInfo& View, 
+	const TArray<EHairInstanceVisibilityType>& InstancesVisibilityType,
 	FHairStrandsViewData& OutHairStrandsViewData,
 	bool bBuildGPUAABB)
 {
@@ -248,26 +248,19 @@ void CreateHairStrandsMacroGroups(
 
 	// 1. Add all visible hair-strands instances
 	static FBoxSphereBounds EmptyBound(ForceInit);
-	const int32 ActiveInstanceCount = Scene->HairStrandsSceneData.RegisteredProxies.Num();
-	TBitArray InstancesVisibility(false, ActiveInstanceCount);
 	for (FVisibleBatch& VisibleBatch : VisibleBatches)
 	{
 		UpdateMacroGroup(VisibleBatch.HairData, VisibleBatch.Batch->Mesh, VisibleBatch.Batch->PrimitiveSceneProxy, EmptyBound);
-		InstancesVisibility[VisibleBatch.HairData->Instance->RegisteredIndex] = true;
 	}
 
 	// 2. Add all hair-strands instances which are non-visible in primary view(s) but visible in shadow view(s)
-	// Slow Linear search
 	if (IsHairStrandsNonVisibleShadowCastingEnable())
 	{
 		for (FHairStrandsInstance* Instance : Scene->HairStrandsSceneData.RegisteredProxies)
 		{
-			if (Instance && InstancesVisibility.IsValidIndex(Instance->RegisteredIndex) && !InstancesVisibility[Instance->RegisteredIndex])
+			if (Instance && InstancesVisibilityType.IsValidIndex(Instance->RegisteredIndex) && InstancesVisibilityType[Instance->RegisteredIndex] == EHairInstanceVisibilityType::StrandsShadowView)
 			{
-				if (IsHairStrandsVisibleInShadows(View, *Instance))
-				{
-					UpdateMacroGroup(const_cast<FHairGroupPublicData*>(Instance->GetHairData()), nullptr, nullptr, Instance->GetBounds());
-				}
+				UpdateMacroGroup(const_cast<FHairGroupPublicData*>(Instance->GetHairData()), nullptr, nullptr, Instance->GetBounds());
 			}
 		}
 	}
diff --git a/Engine/Source/Runtime/Renderer/Private/HairStrands/HairStrandsCluster.h b/Engine/Source/Runtime/Renderer/Private/HairStrands/HairStrandsCluster.h
--- a/Engine/Source/Runtime/Renderer/Private/HairStrands/HairStrandsCluster.h
+++ b/Engine/Source/Runtime/Renderer/Private/HairStrands/HairStrandsCluster.h
@@ -16,5 +16,6 @@ void CreateHairStrandsMacroGroups(
 	FRDGBuilder& GraphBuilder,
 	const FScene* Scene,
 	const FViewInfo& View, 
+	const TArray<EHairInstanceVisibilityType>& InstancesVisibilityType,
 	FHairStrandsViewData& OutHairStrandsViewData,
 	bool bBuildGPUAABB=true);
\ No newline at end of file
diff --git a/Engine/Source/Runtime/Renderer/Private/HairStrands/HairStrandsRendering.cpp b/Engine/Source/Runtime/Renderer/Private/HairStrands/HairStrandsRendering.cpp
--- a/Engine/Source/Runtime/Renderer/Private/HairStrands/HairStrandsRendering.cpp
+++ b/Engine/Source/Runtime/Renderer/Private/HairStrands/HairStrandsRendering.cpp
@@ -114,7 +114,8 @@ void RenderHairPrePass(
 	FRDGBuilder& GraphBuilder,
 	FScene* Scene,
 	TArray<FViewInfo>& Views,
-	FInstanceCullingManager& InstanceCullingManager)
+	FInstanceCullingManager& InstanceCullingManager,
+	const TArray<EHairInstanceVisibilityType>& InstancesVisibilityType)
 {
 	for (FViewInfo& View : Views)
 	{
@@ -150,7 +151,7 @@ void RenderHairPrePass(
 		}
 
 		//SCOPED_GPU_STAT(RHICmdList, HairRendering);
-		CreateHairStrandsMacroGroups(GraphBuilder, Scene, View, View.HairStrandsViewData);
+		CreateHairStrandsMacroGroups(GraphBuilder, Scene, View, InstancesVisibilityType, View.HairStrandsViewData);
 		GraphBuilder.AddDispatchHint();
 
 		// Voxelization and Deep Opacity Maps
diff --git a/Engine/Source/Runtime/Renderer/Private/HairStrands/HairStrandsRendering.h b/Engine/Source/Runtime/Renderer/Private/HairStrands/HairStrandsRendering.h
--- a/Engine/Source/Runtime/Renderer/Private/HairStrands/HairStrandsRendering.h
+++ b/Engine/Source/Runtime/Renderer/Private/HairStrands/HairStrandsRendering.h
@@ -27,7 +27,8 @@ void RenderHairPrePass(
 	FRDGBuilder& GraphBuilder,
 	FScene* Scene,
 	TArray<FViewInfo>& Views,
-	FInstanceCullingManager& InstanceCullingManager);
+	FInstanceCullingManager& InstanceCullingManager,
+	const TArray<EHairInstanceVisibilityType>& InstancesVisibilityType);
 
 void RenderHairBasePass(
 	FRDGBuilder& GraphBuilder,
diff --git a/Engine/Source/Runtime/Renderer/Private/HairStrands/HairStrandsVisibility.cpp b/Engine/Source/Runtime/Renderer/Private/HairStrands/HairStrandsVisibility.cpp
--- a/Engine/Source/Runtime/Renderer/Private/HairStrands/HairStrandsVisibility.cpp
+++ b/Engine/Source/Runtime/Renderer/Private/HairStrands/HairStrandsVisibility.cpp
@@ -4164,7 +4164,8 @@ void DrawHitProxies(
 	// Geometry won't be updated for proxy view
 	const FIntPoint Resolution = HitProxyTexture->Desc.Extent;
 	FHairStrandsViewData HairStrandsViewData;
-	CreateHairStrandsMacroGroups(GraphBuilder, &Scene, View, HairStrandsViewData, false /*bBuildGPUAABB*/);
+	TArray<EHairInstanceVisibilityType> EmptyInstancesVisibilityType;
+	CreateHairStrandsMacroGroups(GraphBuilder, &Scene, View, EmptyInstancesVisibilityType, HairStrandsViewData, false /*bBuildGPUAABB*/);
 
 	// We don't compute the transmittance texture as there is no need for picking.
 	FRDGTextureRef DummyTransmittanceTexture = GraphBuilder.CreateTexture(FRDGTextureDesc::Create2D(Resolution, PF_R32_FLOAT, FClearValueBinding::White, ETextureCreateFlags::ShaderResource | ETextureCreateFlags::UAV), TEXT("Hair.DummyTransmittanceTextureForHitProxyId"));
