Reduce FPrimitiveAndInstance size (208 bytes -> 104 bytes)
- store rotation/scale matrix and relative world bounds using floats + double precision origin.
- also modified FDistanceFieldSceneData::UpdateDistanceFieldObjectBuffers() to use WorldBounds from FPrimitiveAndInstance instead of recalculating them using LocalSpaceMeshBounds and LocalToWorld

diff --git a/Engine/Source/Runtime/Renderer/Private/DistanceFieldObjectManagement.cpp b/Engine/Source/Runtime/Renderer/Private/DistanceFieldObjectManagement.cpp
--- a/Engine/Source/Runtime/Renderer/Private/DistanceFieldObjectManagement.cpp
+++ b/Engine/Source/Runtime/Renderer/Private/DistanceFieldObjectManagement.cpp
@@ -249,7 +249,7 @@ void ProcessDistanceFieldObjectRemoves(FScene* Scene, TArray<FSetElementId>& Dis
 				{
 					// Mark region covered by instance in global distance field as modified
 					FGlobalDFCacheType CacheType = PrimitiveRemoveInfo.bOftenMoving ? GDF_Full : GDF_MostlyStatic;
-					AddModifiedBounds(Scene, CacheType, DistanceFieldSceneData.PrimitiveInstanceMapping[InstanceIndex].WorldBounds);
+					AddModifiedBounds(Scene, CacheType, DistanceFieldSceneData.PrimitiveInstanceMapping[InstanceIndex].GetWorldBounds());
 
 					// Add individual instances to temporary array for processing in the next pass
 					PendingRemoveOperations.Add(InstanceIndex);
@@ -430,26 +430,28 @@ void ProcessPrimitiveUpdate(
 					{
 						FPrimitiveAndInstance& Mapping = DistanceFieldSceneData.PrimitiveInstanceMapping[InstanceIndex];
 
+						const FMatrix PrevLocalToWorld = Mapping.GetLocalToWorld();
+						const FBox PrevWorldBounds = Mapping.GetWorldBounds();
+
 						// Filter out global distance field updates which were too small
-						if (!Mapping.WorldBounds.GetExtent().Equals(WorldBounds.GetExtent(), 0.01f)
-							|| !Mapping.LocalToWorld.Equals(LocalToWorld, 0.01f))
+						if (!PrevWorldBounds.GetExtent().Equals(WorldBounds.GetExtent(), 0.01f)
+							|| !PrevLocalToWorld.Equals(LocalToWorld, 0.01f))
 						{
 							// decide if we want to make a single global distance field update or two updates for large movement (teleport) case
-							const FBox MergedBounds = Mapping.WorldBounds + WorldBounds;
-							const FVector MergedExtentIncrease = MergedBounds.GetExtent() - Mapping.WorldBounds.GetExtent() - WorldBounds.GetExtent();
+							const FBox MergedBounds = PrevWorldBounds + WorldBounds;
+							const FVector MergedExtentIncrease = MergedBounds.GetExtent() - PrevWorldBounds.GetExtent() - WorldBounds.GetExtent();
 							if (MergedExtentIncrease.GetMax() < 100.0f)
 							{
 								AddModifiedBounds(Scene, CacheType, MergedBounds);
 							}
 							else
 							{
-								AddModifiedBounds(Scene, CacheType, Mapping.WorldBounds);
+								AddModifiedBounds(Scene, CacheType, PrevWorldBounds);
 								AddModifiedBounds(Scene, CacheType, WorldBounds);
 							}
 							LogDistanceFieldUpdate(PrimitiveSceneInfo, BoundingRadius, bIsAddOperation);
 
-							Mapping.LocalToWorld = LocalToWorld;
-							Mapping.WorldBounds = WorldBounds;
+							Mapping.SetTransformAndBounds(LocalToWorld, WorldBounds);
 						}
 					}
 				}
@@ -590,10 +592,10 @@ void FDistanceFieldSceneData::UpdateDistanceFieldObjectBuffers(
 
 								const FBox3f LocalSpaceMeshBounds = DistanceFieldData->LocalSpaceMeshBounds;
 			
-								const FMatrix LocalToWorld = PrimAndInst.LocalToWorld;
+								const FMatrix LocalToWorld = PrimAndInst.GetLocalToWorld();
 
 								{
-									const FBox WorldSpaceMeshBounds = ((FBox)LocalSpaceMeshBounds).TransformBy(LocalToWorld);
+									const FBox WorldSpaceMeshBounds = PrimAndInst.GetWorldBounds();
 
 									const FLargeWorldRenderPosition AbsoluteWorldPosition(WorldSpaceMeshBounds.GetCenter());
 
diff --git a/Engine/Source/Runtime/Renderer/Private/ScenePrivate.h b/Engine/Source/Runtime/Renderer/Private/ScenePrivate.h
--- a/Engine/Source/Runtime/Renderer/Private/ScenePrivate.h
+++ b/Engine/Source/Runtime/Renderer/Private/ScenePrivate.h
@@ -1879,17 +1879,58 @@ class FPrimitiveAndInstance
 public:
 
 	FPrimitiveAndInstance(const FMatrix& InLocalToWorld, const FBox& InWorldBounds, FPrimitiveSceneInfo* InPrimitive, int32 InInstanceIndex)
-	: LocalToWorld(InLocalToWorld)
-	, WorldBounds(InWorldBounds)
+	: Primitive(InPrimitive)
 	, InstanceIndex(InInstanceIndex)
-	, Primitive(InPrimitive)
 	{
+		SetTransformAndBounds(InLocalToWorld, InWorldBounds);
 	}
 
-	FMatrix LocalToWorld;
-	FBox WorldBounds;
-	int32 InstanceIndex;
 	FPrimitiveSceneInfo* Primitive;
+
+	FVector Origin;
+	FVector3f TransformRows[3];
+
+	FBox3f WorldBoundsRelativeToOrigin;
+
+	int32 InstanceIndex;
+
+	FORCEINLINE void SetTransformAndBounds(const FMatrix& InLocalToWorld, const FBox& InWorldBounds)
+	{
+		TransformRows[0] = FVector3f((float)InLocalToWorld.M[0][0], (float)InLocalToWorld.M[0][1], (float)InLocalToWorld.M[0][2]);
+		TransformRows[1] = FVector3f((float)InLocalToWorld.M[1][0], (float)InLocalToWorld.M[1][1], (float)InLocalToWorld.M[1][2]);
+		TransformRows[2] = FVector3f((float)InLocalToWorld.M[2][0], (float)InLocalToWorld.M[2][1], (float)InLocalToWorld.M[2][2]);
+		Origin = FVector(InLocalToWorld.M[3][0], InLocalToWorld.M[3][1], InLocalToWorld.M[3][2]);
+
+		WorldBoundsRelativeToOrigin = (FBox3f)(InWorldBounds.ShiftBy(-Origin));
+	}
+
+	FORCEINLINE FMatrix GetLocalToWorld() const
+	{
+		FMatrix Matrix;
+		Matrix.M[0][0] = TransformRows[0].X;
+		Matrix.M[0][1] = TransformRows[0].Y;
+		Matrix.M[0][2] = TransformRows[0].Z;
+		Matrix.M[0][3] = 0.0f;
+		Matrix.M[1][0] = TransformRows[1].X;
+		Matrix.M[1][1] = TransformRows[1].Y;
+		Matrix.M[1][2] = TransformRows[1].Z;
+		Matrix.M[1][3] = 0.0f;
+		Matrix.M[2][0] = TransformRows[2].X;
+		Matrix.M[2][1] = TransformRows[2].Y;
+		Matrix.M[2][2] = TransformRows[2].Z;
+		Matrix.M[2][3] = 0.0f;
+		Matrix.M[3][0] = Origin.X;
+		Matrix.M[3][1] = Origin.Y;
+		Matrix.M[3][2] = Origin.Z;
+		Matrix.M[3][3] = 1.0f;
+		return Matrix;
+	}
+
+	FORCEINLINE FBox GetWorldBounds() const
+	{
+		FBox WorldBoundsRelativeToOriginDoublePrecision = (FBox)WorldBoundsRelativeToOrigin;
+		return WorldBoundsRelativeToOriginDoublePrecision.ShiftBy(Origin);
+	}
 };
 
 class FPrimitiveRemoveInfo
