Fix bShowHitProxies is never true after scene capture rendered as scene view refactoring

diff --git a/Engine/Source/Runtime/Renderer/Private/SceneRendering.cpp b/Engine/Source/Runtime/Renderer/Private/SceneRendering.cpp
--- a/Engine/Source/Runtime/Renderer/Private/SceneRendering.cpp
+++ b/Engine/Source/Runtime/Renderer/Private/SceneRendering.cpp
@@ -4904,16 +4904,7 @@ void FRendererModule::BeginRenderingViewFamilies(FCanvas* Canvas, TArrayView<FSc
 			ViewFamiliesConst.Add(ViewFamily);
 		}
 
-		bool bShowHitProxies = false;
-		for (FSceneRenderer* SceneRenderer : SceneRenderers)
-		{
-			if (SceneRenderer->ViewFamily.EngineShowFlags.HitProxies)
-			{
-				bShowHitProxies = true;
-				break;
-			}
-		}
-
+		bool bShowHitProxies = (Canvas->GetHitProxyConsumer() != nullptr);
 		if (!bShowHitProxies)
 		{
 			USceneCaptureComponent::UpdateDeferredCaptures(Scene);
