fix .exr/.pngs Sequence alpha channel corrupted with MRQ.

1. When the target texture is RGBA without alpha channel write through, the write mask is CW_RGB. The operation on the target alpha channel is undefined. Manually override the load action to clear for this condition.
2. When alpha channel write through is needed, no need to fast clear the target channel as all the value will be overridden.

#tests DumpGPU for two frames in movie render queue verifies artifact gone.

diff --git a/Engine/Source/Runtime/Renderer/Private/PostProcess/PostProcessTonemap.cpp b/Engine/Source/Runtime/Renderer/Private/PostProcess/PostProcessTonemap.cpp
--- a/Engine/Source/Runtime/Renderer/Private/PostProcess/PostProcessTonemap.cpp
+++ b/Engine/Source/Runtime/Renderer/Private/PostProcess/PostProcessTonemap.cpp
@@ -546,6 +546,29 @@ FRDGBufferRef BuildFilmGrainConstants(FRDGBuilder& GraphBuilder, const FViewInfo
 	return FilmGrainConstantsBuffer;
 }
 
+bool ShouldWriteAlphaChannel(const FViewInfo& View, const FTonemapInputs& Inputs, const FRDGTextureRef Output)
+{
+	// If this is a stereo view, there's a good chance we need alpha out of the tonemapper
+	// @todo: Remove this once Oculus fix the bug in their runtime that requires alpha here.
+	const bool bIsStereo = IStereoRendering::IsStereoEyeView(View);
+	const bool bFormatNeedsAlphaWrite = Output->Desc.Format == PF_R9G9B9EXP5;
+	return (Inputs.bWriteAlphaChannel || bIsStereo || bFormatNeedsAlphaWrite);
+}
+
+bool ShouldOverrideOutputLoadActionToFastClear(const FRDGTextureRef Output, bool bShouldWriteAlphaChannel)
+{
+	bool bShouldOverride = false;
+	EPixelFormatChannelFlags OutputFlags = GetPixelFormatValidChannels(Output->Desc.Format);
+	// If we do not write through alpha channel but the output texture has alpha channel
+	// need to override to fast clear load action ERenderTargetLoadAction::Clear, otherwise,
+	// the alpha channel can be garbage data in terms of different driver implementation.
+	if ( (!bShouldWriteAlphaChannel) && EnumHasAnyFlags(OutputFlags, EPixelFormatChannelFlags::A))
+	{
+		bShouldOverride = true;
+	}
+	return bShouldOverride;
+}
+
 FScreenPassTexture AddTonemapPass(FRDGBuilder& GraphBuilder, const FViewInfo& View, const FTonemapInputs& Inputs)
 {
 	if (!Inputs.bGammaOnly)
@@ -614,7 +637,13 @@ FScreenPassTexture AddTonemapPass(FRDGBuilder& GraphBuilder, const FViewInfo& Vi
 		Output = FScreenPassRenderTarget(
 			GraphBuilder.CreateTexture(OutputDesc, TEXT("Tonemap")),
 			Inputs.SceneColor.ViewRect,
-			ERenderTargetLoadAction::EClear);
+			ERenderTargetLoadAction::ENoAction);
+	}
+
+	const bool bShouldWriteAlphaChannel = ShouldWriteAlphaChannel(View, Inputs, Output.Texture);
+	if (ShouldOverrideOutputLoadActionToFastClear(Output.Texture, bShouldWriteAlphaChannel))
+	{
+		Output.LoadAction = ERenderTargetLoadAction::EClear;
 	}
 
 	const FScreenPassTextureViewport OutputViewport(Output);
@@ -967,11 +996,8 @@ FScreenPassTexture AddTonemapPass(FRDGBuilder& GraphBuilder, const FViewInfo& Vi
 		TShaderMapRef<FTonemapVS> VertexShader(View.ShaderMap);
 		TShaderMapRef<FTonemapPS> PixelShader(View.ShaderMap, DesktopPermutationVector);
 
-		// If this is a stereo view, there's a good chance we need alpha out of the tonemapper
-		// @todo: Remove this once Oculus fix the bug in their runtime that requires alpha here.
-		const bool bIsStereo = IStereoRendering::IsStereoEyeView(View);
-		const bool bFormatNeedsAlphaWrite = Output.Texture->Desc.Format == PF_R9G9B9EXP5;
-		FRHIBlendState* BlendState = (Inputs.bWriteAlphaChannel || bIsStereo || bFormatNeedsAlphaWrite) ? FScreenPassPipelineState::FDefaultBlendState::GetRHI() : TStaticBlendStateWriteMask<CW_RGB>::GetRHI();
+		FRHIBlendState* BlendState = bShouldWriteAlphaChannel ? FScreenPassPipelineState::FDefaultBlendState::GetRHI() : TStaticBlendStateWriteMask<CW_RGB>::GetRHI();
+
 		FRHIDepthStencilState* DepthStencilState = FScreenPassPipelineState::FDefaultDepthStencilState::GetRHI();
 
 		EScreenPassDrawFlags DrawFlags = EScreenPassDrawFlags::AllowHMDHiddenAreaMask;
