Refactored logic in FDistanceFieldSceneData::UpdateDistanceFieldObjectBuffers(...) to minimize number of operations done in double precision.

diff --git a/Engine/Source/Runtime/Renderer/Private/DistanceFieldObjectManagement.cpp b/Engine/Source/Runtime/Renderer/Private/DistanceFieldObjectManagement.cpp
--- a/Engine/Source/Runtime/Renderer/Private/DistanceFieldObjectManagement.cpp
+++ b/Engine/Source/Runtime/Renderer/Private/DistanceFieldObjectManagement.cpp
@@ -388,9 +388,9 @@ void ProcessPrimitiveUpdate(
 
 				const FMatrix LocalToWorld = InstanceLocalToWorldTransforms[TransformIndex];
 
-				const FMatrix::FReal MaxScale = LocalToWorld.GetMaximumAxisScale();
+				const FMatrix::FReal MinScale = LocalToWorld.GetMinimumAxisScale();
 
-				if (bIsAddOperation && (MaxScale <= 0 || bInstanceCountOverflow))
+				if (bIsAddOperation && (MinScale <= 0 || bInstanceCountOverflow))
 				{
 					// Skip degenerate instances or when instance count limit is reached
 					PrimitiveSceneInfo->DistanceFieldInstanceIndices[TransformIndex] = -1;
@@ -591,17 +591,23 @@ void FDistanceFieldSceneData::UpdateDistanceFieldObjectBuffers(
 								PrimitiveSceneProxy->GetDistanceFieldAtlasData(DistanceFieldData, SelfShadowBias);
 
 								const FBox3f LocalSpaceMeshBounds = DistanceFieldData->LocalSpaceMeshBounds;
-			
-								const FMatrix LocalToWorld = PrimAndInst.GetLocalToWorld();
 
-								{
-									const FBox WorldSpaceMeshBounds = PrimAndInst.GetWorldBounds();
+								// Uniformly scale our Volume space to lie within [-1, 1] at the max extent
+								// This is mirrored in the SDF encoding
+								const FBox3f::FReal LocalToVolumeScale = 1.0f / LocalSpaceMeshBounds.GetExtent().GetMax();
+
+								const FLargeWorldRenderPosition WorldPosition(PrimAndInst.Origin);
+								const FVector TilePositionOffset = WorldPosition.GetTileOffset();
+
+								FMatrix44f LocalToRelativeWorld = FLargeWorldRenderScalar::MakeToRelativeWorldMatrix(TilePositionOffset, PrimAndInst.GetLocalToWorld());
+								FMatrix44f RelativeWorldToLocal = FMatrix44f(LocalToRelativeWorld.InverseFast());
 
-									const FLargeWorldRenderPosition AbsoluteWorldPosition(WorldSpaceMeshBounds.GetCenter());
+								{
+									const FBox3f RelativeWorldSpaceMeshBounds = PrimAndInst.WorldBoundsRelativeToOrigin.ShiftBy(FVector3f(PrimAndInst.Origin - TilePositionOffset));
 
-									const FVector4f ObjectBoundingSphere(AbsoluteWorldPosition.GetOffset(), WorldSpaceMeshBounds.GetExtent().Size());
+									const FVector4f ObjectBoundingSphere(RelativeWorldSpaceMeshBounds.GetCenter(), RelativeWorldSpaceMeshBounds.GetExtent().Size());
 
-									UploadObjectBounds[0] = AbsoluteWorldPosition.GetTile();
+									UploadObjectBounds[0] = WorldPosition.GetTile();
 									UploadObjectBounds[1] = ObjectBoundingSphere;
 
 									const FGlobalDFCacheType CacheType = PrimitiveSceneProxy->IsOftenMoving() ? GDF_Full : GDF_MostlyStatic;
@@ -620,27 +626,18 @@ void FDistanceFieldSceneData::UpdateDistanceFieldObjectBuffers(
 									Flags |= bVisible ? 16u : 0;
 									Flags |= bAffectIndirectLightingWhileHidden ? 32u : 0;
 
-									FVector4f ObjectWorldExtentAndFlags((FVector3f)WorldSpaceMeshBounds.GetExtent(), 0.0f);
+									FVector4f ObjectWorldExtentAndFlags(RelativeWorldSpaceMeshBounds.GetExtent(), 0.0f);
 									ObjectWorldExtentAndFlags.W = *(const float*)&Flags;
 									UploadObjectBounds[2] = ObjectWorldExtentAndFlags;
 								}
 
-								// Uniformly scale our Volume space to lie within [-1, 1] at the max extent
-								// This is mirrored in the SDF encoding
-								const FBox3f::FReal LocalToVolumeScale = 1.0f / LocalSpaceMeshBounds.GetExtent().GetMax();
-
-								const FMatrix VolumeToWorld = FScaleMatrix(1.0f / LocalToVolumeScale) * FTranslationMatrix((FVector)LocalSpaceMeshBounds.GetCenter()) * LocalToWorld;
-
-								const FLargeWorldRenderPosition WorldPosition(VolumeToWorld.GetOrigin());
-								const FVector TilePositionOffset = WorldPosition.GetTileOffset();
-
-								// Inverse on FMatrix44f can generate NaNs if the source matrix contains large scaling, so do it in double precision.
-								FMatrix LocalToRelativeWorld = FLargeWorldRenderScalar::MakeToRelativeWorldMatrixDouble(TilePositionOffset, VolumeToWorld);
+								const FMatrix44f VolumeToRelativeWorld = FScaleMatrix44f(1.0f / LocalToVolumeScale) * FTranslationMatrix44f(LocalSpaceMeshBounds.GetCenter()) * LocalToRelativeWorld;
+								const FMatrix44f RelativeWorldToVolume = RelativeWorldToLocal * FTranslationMatrix44f(-LocalSpaceMeshBounds.GetCenter()) * FScaleMatrix44f(LocalToVolumeScale);
 
 								// TilePosition
 								UploadObjectData[0] = WorldPosition.GetTile();
 
-								const FMatrix44f WorldToVolumeT = FMatrix44f(LocalToRelativeWorld.Inverse().GetTransposed());
+								const FMatrix44f WorldToVolumeT = RelativeWorldToVolume.GetTransposed();
 								// WorldToVolumeT
 								UploadObjectData[1] = (*(FVector4f*)&WorldToVolumeT.M[0]);
 								UploadObjectData[2] = (*(FVector4f*)&WorldToVolumeT.M[1]);
@@ -673,12 +670,12 @@ void FDistanceFieldSceneData::UpdateDistanceFieldObjectBuffers(
 								Vector4.W = *(const float*)&GPUSceneInstanceIndex;
 								UploadObjectData[5] = Vector4;
 
-								const FMatrix44f VolumeToWorldT = FMatrix44f(LocalToRelativeWorld.GetTransposed());
+								const FMatrix44f VolumeToWorldT = VolumeToRelativeWorld.GetTransposed();
 								UploadObjectData[6] = *(FVector4f*)&VolumeToWorldT.M[0];
 								UploadObjectData[7] = *(FVector4f*)&VolumeToWorldT.M[1];
 								UploadObjectData[8] = *(FVector4f*)&VolumeToWorldT.M[2];
 
-								FVector4f FloatVector8(FVector3f(VolumeToWorld.GetScaleVector()), 0.0f);
+								FVector4f FloatVector8(FVector3f(VolumeToRelativeWorld.GetScaleVector()), 0.0f);
 
 								// Bypass NaN checks in FVector4f ctor
 								FSetElementId AssetStateSetId = AssetStateArray.FindId(DistanceFieldData);
diff --git a/Engine/Source/Runtime/Renderer/Private/GlobalDistanceField.cpp b/Engine/Source/Runtime/Renderer/Private/GlobalDistanceField.cpp
--- a/Engine/Source/Runtime/Renderer/Private/GlobalDistanceField.cpp
+++ b/Engine/Source/Runtime/Renderer/Private/GlobalDistanceField.cpp
@@ -1148,8 +1148,7 @@ static void ComputeUpdateRegionsAndUpdateViewState(
 
 					const TArray<FBox>& PrimitiveModifiedBounds = ClipmapViewState.Cache[CacheType].PrimitiveModifiedBounds;
 
-					TArray<FBox, SceneRenderingAllocator> CulledPrimitiveModifiedBounds;
-					CulledPrimitiveModifiedBounds.Empty(PrimitiveModifiedBounds.Num() / 2);
+					uint32 NumCulledPrimitiveModifiedBounds = 0;
 
 					Clipmap.UpdateBounds.Empty(PrimitiveModifiedBounds.Num() / 2);
 
@@ -1159,7 +1158,7 @@ static void ComputeUpdateRegionsAndUpdateViewState(
 
 						if (ModifiedBounds.ComputeSquaredDistanceToBox(ClipmapBounds) < ClipmapInfluenceRadius * ClipmapInfluenceRadius)
 						{
-							CulledPrimitiveModifiedBounds.Add(ModifiedBounds);
+							++NumCulledPrimitiveModifiedBounds;
 
 							Clipmap.UpdateBounds.Add(FClipmapUpdateBounds(ModifiedBounds.GetCenter(), ModifiedBounds.GetExtent(), true));
 							
@@ -1201,7 +1200,7 @@ static void ComputeUpdateRegionsAndUpdateViewState(
 					}
 
 					// Only use partial updates with small numbers of primitive modifications
-					bool bUsePartialUpdatesForUpdateBounds = bUsePartialUpdates && CulledPrimitiveModifiedBounds.Num() < 1024;
+					bool bUsePartialUpdatesForUpdateBounds = bUsePartialUpdates && NumCulledPrimitiveModifiedBounds < 1024;
 
 					if (!bUsePartialUpdatesForUpdateBounds)
 					{
