GPUScene upload optimizations
- Removed some primitive setup work from the new InstanceSceneDataBuffers path (that is no longer needed)
- Removed copy step for FPrimitiveSceneShaderData instead writing directly to the upload buffer.
- Inlined FInstanceSceneShaderData build & cleaned up includes of InstanceUniformShaderParameters.h to reduce rebuilding
- Cleaned up the FPrimitiveSceneShaderData a bit
- Removed dummy identity transform and needless matrix work for single-instance primitives.

diff --git a/Engine/Plugins/Experimental/GPULightmass/Source/GPULightmass/Private/LightmapRenderer.cpp b/Engine/Plugins/Experimental/GPULightmass/Source/GPULightmass/Private/LightmapRenderer.cpp
--- a/Engine/Plugins/Experimental/GPULightmass/Source/GPULightmass/Private/LightmapRenderer.cpp
+++ b/Engine/Plugins/Experimental/GPULightmass/Source/GPULightmass/Private/LightmapRenderer.cpp
@@ -467,7 +467,9 @@ void FCachedRayTracingSceneData::SetupViewAndSceneUniformBufferFromSceneRenderSt
 					INVALID_LAST_UPDATE_FRAME,
 					0, /* Custom Data Count */
 					0.0f, /* Random ID */
-					InstanceGroup.InstanceSceneDataBuffers->GetInstanceToPrimitiveRelative(InstanceIdx)
+					InstanceGroup.InstanceSceneDataBuffers->GetInstanceToPrimitiveRelative(InstanceIdx),
+					true,
+					FInstanceSceneShaderData::SupportsCompressedTransforms()
 				);
 
 				InstancePayloadData.Emplace(InstanceData.InstanceLightShadowUVBias[InstanceIdx]);
@@ -525,7 +527,9 @@ void FCachedRayTracingSceneData::SetupViewAndSceneUniformBufferFromSceneRenderSt
 				INVALID_LAST_UPDATE_FRAME,
 				0, /* Custom Data Count */
 				0.0f, /* Random ID */
-				PrimitiveUniformShaderParameters.LocalToRelativeWorld
+				PrimitiveUniformShaderParameters.LocalToRelativeWorld,
+				true,
+				FInstanceSceneShaderData::SupportsCompressedTransforms()
 			);
 
 			PrimitiveSceneData.Add(FPrimitiveSceneShaderData(PrimitiveUniformShaderParameters));
diff --git a/Engine/Plugins/FX/Niagara/Source/Niagara/Public/NiagaraRendererMeshes.h b/Engine/Plugins/FX/Niagara/Source/Niagara/Public/NiagaraRendererMeshes.h
--- a/Engine/Plugins/FX/Niagara/Source/Niagara/Public/NiagaraRendererMeshes.h
+++ b/Engine/Plugins/FX/Niagara/Source/Niagara/Public/NiagaraRendererMeshes.h
@@ -12,6 +12,7 @@ NiagaraRenderer.h: Base class for Niagara render modules
 #include "NiagaraGPUSortInfo.h"
 #include "NiagaraSystemInstance.h"
 #include "NiagaraGPUSceneUtils.h"
+#include "InstanceUniformShaderParameters.h"
 
 class FNiagaraDataSet;
 struct FNiagaraDynamicDataMesh;
diff --git a/Engine/Source/Runtime/Core/Private/Misc/LargeWorldRenderPosition.cpp b/Engine/Source/Runtime/Core/Private/Misc/LargeWorldRenderPosition.cpp
--- a/Engine/Source/Runtime/Core/Private/Misc/LargeWorldRenderPosition.cpp
+++ b/Engine/Source/Runtime/Core/Private/Misc/LargeWorldRenderPosition.cpp
@@ -9,7 +9,6 @@
 
 #define UE_LWC_RENDER_TILE_SIZE_MIN (262144.0)	// LWC_TODO: Make this smaller?
 
-#define UE_LWC_RENDER_TILE_SIZE 2097152.0
 //#define UE_LWC_RENDER_TILE_SIZE UE_OLD_WORLD_MAX (FMath::Max(UE_LWC_RENDER_TILE_SIZE_MIN,  HALF_WORLD_MAX / 16777216.0))	// LWC_TODO: Tweak render tile divisor to maximise range.
 //static_assert(UE_LWC_RENDER_TILE_SIZE <= 1048576.0, "Current WORLD_MAX is too large and is likely to adversely affect world space coordinate precision within shaders!");	// LWC_TODO: Set this to something reasonable.
 
@@ -22,12 +21,6 @@
 // Normally offsets should be within +/-TileSizeDivideBy2, but we often rebase multiple quantities off a single tile origin
 #define UE_LWC_RENDER_MAX_OFFSET (2097152.0*0.5)
 
-template<typename TScalar>
-TScalar TLargeWorldRenderScalar<TScalar>::GetTileSize()
-{
-	return static_cast<TScalar>(UE_LWC_RENDER_TILE_SIZE);
-}
-
 template<typename TScalar>
 FVector3f TLargeWorldRenderScalar<TScalar>::GetTileFor(FVector InPosition)
 {
diff --git a/Engine/Source/Runtime/Core/Public/Misc/LargeWorldRenderPosition.h b/Engine/Source/Runtime/Core/Public/Misc/LargeWorldRenderPosition.h
--- a/Engine/Source/Runtime/Core/Public/Misc/LargeWorldRenderPosition.h
+++ b/Engine/Source/Runtime/Core/Public/Misc/LargeWorldRenderPosition.h
@@ -10,6 +10,8 @@
 #include "Math/Vector4.h"
 #include "Misc/LargeWorldCoordinates.h"
 
+static constexpr double UE_LWC_RENDER_TILE_SIZE = 2097152.0;
+
 template <typename TScalar>
 struct TLargeWorldRenderScalar
 {
@@ -17,8 +19,11 @@ private:
 	using VectorType = UE::Math::TVector<TScalar>;
 
 public:
-	// This should become inline at some point, but keeping it in cpp file for now to make iteration faster when experimenting with tile size
-	CORE_API static TScalar GetTileSize();
+	FORCEINLINE static TScalar GetTileSize()
+	{
+		return static_cast<TScalar>(UE_LWC_RENDER_TILE_SIZE);
+	}
+
 	CORE_API static FVector3f GetTileFor(FVector InPosition);
 
 	static TScalar MakeTile(double InValue)
diff --git a/Engine/Source/Runtime/Engine/Private/InstanceUniformShaderParameters.cpp b/Engine/Source/Runtime/Engine/Private/InstanceUniformShaderParameters.cpp
--- a/Engine/Source/Runtime/Engine/Private/InstanceUniformShaderParameters.cpp
+++ b/Engine/Source/Runtime/Engine/Private/InstanceUniformShaderParameters.cpp
@@ -4,6 +4,12 @@
 #include "DataDrivenShaderPlatformInfo.h"
 #include "SceneDefinitions.h"
 
+
+bool FInstanceSceneShaderData::SupportsCompressedTransforms()
+{
+	return FDataDrivenShaderPlatformInfo::GetSupportSceneDataCompressedTransforms(GMaxRHIShaderPlatform);
+}
+
 uint32 FInstanceSceneShaderData::GetDataStrideInFloat4s()
 {
 	if (FDataDrivenShaderPlatformInfo::GetSupportSceneDataCompressedTransforms(GMaxRHIShaderPlatform))
@@ -27,7 +33,7 @@ void FInstanceSceneShaderData::Build
 	bool bIsVisible
 )
 {
-	BuildInternal(PrimitiveId, RelativeId, InstanceFlags, LastUpdateFrame, CustomDataCount, RandomID, FRenderTransform::Identity, bIsVisible);
+	BuildInternal(PrimitiveId, RelativeId, InstanceFlags, LastUpdateFrame, CustomDataCount, RandomID, FRenderTransform::Identity, bIsVisible, SupportsCompressedTransforms());
 }
 
 void FInstanceSceneShaderData::Build
@@ -48,61 +54,6 @@ void FInstanceSceneShaderData::Build
 	// Remove shear
 	LocalToWorld.Orthogonalize();
 
-	BuildInternal(PrimitiveId, RelativeId, InstanceFlags, LastUpdateFrame, CustomDataCount, RandomID, LocalToWorld, bIsVisible);
+	BuildInternal(PrimitiveId, RelativeId, InstanceFlags, LastUpdateFrame, CustomDataCount, RandomID, LocalToWorld, bIsVisible, SupportsCompressedTransforms());
 }
 
-void FInstanceSceneShaderData::BuildInternal
-(
-	uint32 PrimitiveId,
-	uint32 RelativeId,
-	uint32 InstanceFlags,
-	uint32 LastUpdateFrame,
-	uint32 CustomDataCount,
-	float RandomID,
-	const FRenderTransform& LocalToWorld, // Assumes shear has been removed already
-	bool bIsVisible
-)
-{
-	// Note: layout must match GetInstanceData in SceneData.ush and InitializeInstanceSceneData in GPUSceneWriter.ush
-
-	const float RotDeterminant = LocalToWorld.RotDeterminant();
-	if (RotDeterminant < 0.0f)
-	{
-		InstanceFlags |= INSTANCE_SCENE_DATA_FLAG_DETERMINANT_SIGN;
-	}
-	else
-	{
-		InstanceFlags &= ~INSTANCE_SCENE_DATA_FLAG_DETERMINANT_SIGN;
-	}
-
-	// Mark zero scaled instances as hidden.
-	if (!bIsVisible || RotDeterminant == 0)
-	{
-		InstanceFlags |= INSTANCE_SCENE_DATA_FLAG_HIDDEN;
-	}
-
-	checkSlow((PrimitiveId		& 0x000FFFFF) == PrimitiveId);
-	checkSlow((InstanceFlags	& 0x00000FFF) == InstanceFlags);
-	checkSlow((RelativeId		& 0x00FFFFFF) == RelativeId);
-	checkSlow((CustomDataCount	& 0x000000FF) == CustomDataCount);
-
-	const uint32 Packed0 = (InstanceFlags   << 20u) | PrimitiveId;
-	const uint32 Packed1 = (CustomDataCount << 24u) | RelativeId;
-
-	Data[0].X  = *(const float*)&Packed0;
-	Data[0].Y  = *(const float*)&Packed1;
-	Data[0].Z  = *(const float*)&LastUpdateFrame;
-	Data[0].W  = *(const float*)&RandomID;
-
-	if (FDataDrivenShaderPlatformInfo::GetSupportSceneDataCompressedTransforms(GMaxRHIShaderPlatform))
-	{
-		FCompressedTransform CompressedLocalToWorld(LocalToWorld);
-		Data[1] = *(const FVector4f*)&CompressedLocalToWorld.Rotation[0];
-		Data[2] = *(const FVector3f*)&CompressedLocalToWorld.Translation;
-	}
-	else
-	{
-		// Note: writes 3x float4s
-		LocalToWorld.To3x4MatrixTranspose((float*)&Data[1]);
-	}
-}
diff --git a/Engine/Source/Runtime/Engine/Private/PrimitiveUniformShaderParameters.cpp b/Engine/Source/Runtime/Engine/Private/PrimitiveUniformShaderParameters.cpp
--- a/Engine/Source/Runtime/Engine/Private/PrimitiveUniformShaderParameters.cpp
+++ b/Engine/Source/Runtime/Engine/Private/PrimitiveUniformShaderParameters.cpp
@@ -134,117 +134,116 @@ PRAGMA_ENABLE_DEPRECATION_WARNINGS
 
 FPrimitiveSceneShaderData::FPrimitiveSceneShaderData(const FPrimitiveSceneProxy* RESTRICT Proxy)
 	: Data(InPlace, NoInit)
+{
+	FPrimitiveSceneShaderData::BuildDataFromProxy(Proxy, Data.GetData());
+}
+
+void FPrimitiveSceneShaderData::BuildDataFromProxy(const FPrimitiveSceneProxy* RESTRICT Proxy, FVector4f* RESTRICT OutData)
 {
 	FPrimitiveUniformShaderParametersBuilder Builder = FPrimitiveUniformShaderParametersBuilder{};
 	Proxy->BuildUniformShaderParameters(Builder);
-	Setup(Builder.Build());
+	Setup(Builder.Build(), OutData);
 }
 
-void FPrimitiveSceneShaderData::Setup(const FPrimitiveUniformShaderParameters& PrimitiveUniformShaderParameters)
+/**
+ * Helper struct to make sure integers are converted to float as needed.
+ */
+struct FAsFloat
 {
-	static_assert(NUM_LIGHTING_CHANNELS == 3, "The FPrimitiveSceneShaderData packing currently assumes a maximum of 3 lighting channels.");
-
-	// Note: layout must match GetPrimitiveData in usf
-
-	// Set W directly in order to bypass NaN check, when passing int through FVector to shader.
+	FORCEINLINE FAsFloat(uint32 InValue) : FloatValue(FMath::AsFloat(InValue)) {}
+	FORCEINLINE FAsFloat(float InValue) : FloatValue(InValue) {}
 
-	Data[0].X	= FMath::AsFloat(PrimitiveUniformShaderParameters.Flags);
-	Data[0].Y	= FMath::AsFloat(PrimitiveUniformShaderParameters.InstanceSceneDataOffset);
-	Data[0].Z	= FMath::AsFloat(PrimitiveUniformShaderParameters.NumInstanceSceneDataEntries);
-	Data[0].W	= FMath::AsFloat((uint32)PrimitiveUniformShaderParameters.SingleCaptureIndex |
-								 ((PrimitiveUniformShaderParameters.VisibilityFlags & 0xFFFFu) << 16u));
+	float FloatValue;
+};
 
-	Data[1].X	= PrimitiveUniformShaderParameters.TilePosition.X;
-	Data[1].Y	= PrimitiveUniformShaderParameters.TilePosition.Y;
-	Data[1].Z	= PrimitiveUniformShaderParameters.TilePosition.Z;
-	Data[1].W	= FMath::AsFloat(PrimitiveUniformShaderParameters.PrimitiveComponentId);
+FORCEINLINE void Store4(FVector4f *Data, int32 Offset, FAsFloat X, FAsFloat Y, FAsFloat Z, FAsFloat W)
+{
+	VectorRegister4f VR = MakeVectorRegisterFloat(X.FloatValue, Y.FloatValue, Z.FloatValue, W.FloatValue);
+	VectorStoreAligned(VR, &Data[Offset].X);
+}
 
-	// Pack these matrices into the buffer as float3x4 transposed
+FORCEINLINE void Store4(FVector4f *Data, int32 Offset, const FVector3f &XYZ, FAsFloat W)
+{
+	Store4(Data, Offset, XYZ.X, XYZ.Y, XYZ.Z, W.FloatValue);
+}
 
-	FMatrix44f LocalToRelativeWorldTranspose = PrimitiveUniformShaderParameters.LocalToRelativeWorld.GetTransposed();
-	Data[2]		= *(const FVector4f*)&LocalToRelativeWorldTranspose.M[0][0];
-	Data[3]		= *(const FVector4f*)&LocalToRelativeWorldTranspose.M[1][0];
-	Data[4]		= *(const FVector4f*)&LocalToRelativeWorldTranspose.M[2][0];
+FORCEINLINE void StoreTransposed(FVector4f *Data, int32 StartOffset, const FMatrix44f &Matrix)
+{
+	Store4(Data, StartOffset + 0, Matrix.M[0][0],   Matrix.M[1][0],   Matrix.M[2][0],   Matrix.M[3][0]);
+	Store4(Data, StartOffset + 1, Matrix.M[0][1],   Matrix.M[1][1],   Matrix.M[2][1],   Matrix.M[3][1]);
+	Store4(Data, StartOffset + 2, Matrix.M[0][2],   Matrix.M[1][2],   Matrix.M[2][2],   Matrix.M[3][2]);
+}
 
-	FMatrix44f RelativeWorldToLocalTranspose = PrimitiveUniformShaderParameters.RelativeWorldToLocal.GetTransposed();
-	Data[5]		= *(const FVector4f*)&RelativeWorldToLocalTranspose.M[0][0];
-	Data[6]		= *(const FVector4f*)&RelativeWorldToLocalTranspose.M[1][0];
-	Data[7]		= *(const FVector4f*)&RelativeWorldToLocalTranspose.M[2][0];
+void FPrimitiveSceneShaderData::Setup(const FPrimitiveUniformShaderParameters& PrimitiveUniformShaderParameters, FVector4f* RESTRICT OutData)
+{
+	static_assert(NUM_LIGHTING_CHANNELS == 3, "The FPrimitiveSceneShaderData packing currently assumes a maximum of 3 lighting channels.");
 
-	FMatrix44f PreviousLocalToRelativeWorldTranspose = PrimitiveUniformShaderParameters.PreviousLocalToRelativeWorld.GetTransposed();
-	Data[8]		= *(const FVector4f*)&PreviousLocalToRelativeWorldTranspose.M[0][0];
-	Data[9]		= *(const FVector4f*)&PreviousLocalToRelativeWorldTranspose.M[1][0];
-	Data[10]	= *(const FVector4f*)&PreviousLocalToRelativeWorldTranspose.M[2][0];
+	// Note: layout must match GetPrimitiveData in usf
 
-	FMatrix44f PreviousRelativeWorldToLocalTranspose = PrimitiveUniformShaderParameters.PreviousRelativeWorldToLocal.GetTransposed();
-	Data[11]	= *(const FVector4f*)&PreviousRelativeWorldToLocalTranspose.M[0][0];
-	Data[12]	= *(const FVector4f*)&PreviousRelativeWorldToLocalTranspose.M[1][0];
-	Data[13]	= *(const FVector4f*)&PreviousRelativeWorldToLocalTranspose.M[2][0];
+	Store4(OutData, 0,
+		PrimitiveUniformShaderParameters.Flags,
+		PrimitiveUniformShaderParameters.InstanceSceneDataOffset,
+		PrimitiveUniformShaderParameters.NumInstanceSceneDataEntries,
+		(uint32)PrimitiveUniformShaderParameters.SingleCaptureIndex | ((PrimitiveUniformShaderParameters.VisibilityFlags & 0xFFFFu) << 16u));
 
-	FMatrix44f WorldToPreviousWorldTranspose = PrimitiveUniformShaderParameters.WorldToPreviousWorld.GetTransposed();
-	Data[14]	= *(const FVector4f*)&WorldToPreviousWorldTranspose.M[0][0];
-	Data[15]	= *(const FVector4f*)&WorldToPreviousWorldTranspose.M[1][0];
-	Data[16]	= *(const FVector4f*)&WorldToPreviousWorldTranspose.M[2][0];
+	Store4(OutData, 1,
+		PrimitiveUniformShaderParameters.TilePosition.X,
+		PrimitiveUniformShaderParameters.TilePosition.Y,
+		PrimitiveUniformShaderParameters.TilePosition.Z,
+		PrimitiveUniformShaderParameters.PrimitiveComponentId);
 
-	Data[17]	= FVector4f(PrimitiveUniformShaderParameters.InvNonUniformScale, PrimitiveUniformShaderParameters.ObjectBoundsX);
-	Data[18]	= PrimitiveUniformShaderParameters.ObjectRelativeWorldPositionAndRadius;
+	// Pack these matrices into the buffer as float3x4 transposed
+	StoreTransposed(OutData, 2, PrimitiveUniformShaderParameters.LocalToRelativeWorld);
+	StoreTransposed(OutData, 5, PrimitiveUniformShaderParameters.RelativeWorldToLocal);
+	StoreTransposed(OutData, 8, PrimitiveUniformShaderParameters.PreviousLocalToRelativeWorld);
+	StoreTransposed(OutData, 11, PrimitiveUniformShaderParameters.PreviousRelativeWorldToLocal);
+	StoreTransposed(OutData, 14, PrimitiveUniformShaderParameters.WorldToPreviousWorld);
 
-	Data[19]	= FVector4f(PrimitiveUniformShaderParameters.ActorRelativeWorldPosition, 0.0f);
-	Data[19].W	= FMath::AsFloat(PrimitiveUniformShaderParameters.LightmapUVIndex);
 
-	Data[20]	= FVector4f(PrimitiveUniformShaderParameters.ObjectOrientation, 0.0f);
-	Data[20].W	= FMath::AsFloat(PrimitiveUniformShaderParameters.LightmapDataIndex);
+	OutData[17]	= FVector4f(PrimitiveUniformShaderParameters.InvNonUniformScale, PrimitiveUniformShaderParameters.ObjectBoundsX);
+	OutData[18]	= PrimitiveUniformShaderParameters.ObjectRelativeWorldPositionAndRadius;
 
-	Data[21]	= PrimitiveUniformShaderParameters.NonUniformScale;
+	Store4(OutData, 19, PrimitiveUniformShaderParameters.ActorRelativeWorldPosition, PrimitiveUniformShaderParameters.LightmapUVIndex);
+	Store4(OutData, 20, PrimitiveUniformShaderParameters.ObjectOrientation, PrimitiveUniformShaderParameters.LightmapDataIndex);
 
-	Data[22]	= FVector4f(PrimitiveUniformShaderParameters.PreSkinnedLocalBoundsMin, 0.0f);
-	Data[22].W	= FMath::AsFloat(PrimitiveUniformShaderParameters.NaniteResourceID);
+	OutData[21]	= PrimitiveUniformShaderParameters.NonUniformScale;
 
-	Data[23]	= FVector4f(PrimitiveUniformShaderParameters.PreSkinnedLocalBoundsMax, 0.0f);
-	Data[23].W	= FMath::AsFloat(PrimitiveUniformShaderParameters.NaniteHierarchyOffset);
+	Store4(OutData, 22, PrimitiveUniformShaderParameters.PreSkinnedLocalBoundsMin, PrimitiveUniformShaderParameters.NaniteResourceID);
+	Store4(OutData, 23, PrimitiveUniformShaderParameters.PreSkinnedLocalBoundsMax, PrimitiveUniformShaderParameters.NaniteHierarchyOffset);
 
-	Data[24]	= FVector4f(PrimitiveUniformShaderParameters.LocalObjectBoundsMin, PrimitiveUniformShaderParameters.ObjectBoundsY);
-	Data[25]	= FVector4f(PrimitiveUniformShaderParameters.LocalObjectBoundsMax, PrimitiveUniformShaderParameters.ObjectBoundsZ);
+	OutData[24]	= FVector4f(PrimitiveUniformShaderParameters.LocalObjectBoundsMin, PrimitiveUniformShaderParameters.ObjectBoundsY);
+	OutData[25]	= FVector4f(PrimitiveUniformShaderParameters.LocalObjectBoundsMax, PrimitiveUniformShaderParameters.ObjectBoundsZ);
 
-	Data[26].X = PrimitiveUniformShaderParameters.InstanceLocalBoundsCenter.X;
-	Data[26].Y = PrimitiveUniformShaderParameters.InstanceLocalBoundsCenter.Y;
-	Data[26].Z = PrimitiveUniformShaderParameters.InstanceLocalBoundsCenter.Z;
-	Data[26].W = FMath::AsFloat(PrimitiveUniformShaderParameters.InstancePayloadDataOffset);
+	Store4(OutData, 26, PrimitiveUniformShaderParameters.InstanceLocalBoundsCenter, PrimitiveUniformShaderParameters.InstancePayloadDataOffset);
+	Store4(OutData, 27, PrimitiveUniformShaderParameters.InstanceLocalBoundsExtent, (PrimitiveUniformShaderParameters.InstancePayloadDataStride & 0x00FFFFFFu) | (PrimitiveUniformShaderParameters.InstancePayloadExtensionSize << 24u));
+	Store4(OutData, 28, PrimitiveUniformShaderParameters.WireframeColor, PrimitiveUniformShaderParameters.PackedNaniteFlags);
 
-	Data[27].X = PrimitiveUniformShaderParameters.InstanceLocalBoundsExtent.X;
-	Data[27].Y = PrimitiveUniformShaderParameters.InstanceLocalBoundsExtent.Y;
-	Data[27].Z = PrimitiveUniformShaderParameters.InstanceLocalBoundsExtent.Z;
-	Data[27].W = FMath::AsFloat((PrimitiveUniformShaderParameters.InstancePayloadDataStride & 0x00FFFFFFu) |
-								(PrimitiveUniformShaderParameters.InstancePayloadExtensionSize << 24u));
+	Store4(OutData, 29, PrimitiveUniformShaderParameters.LevelColor, uint32(PrimitiveUniformShaderParameters.PersistentPrimitiveIndex));
+	Store4(OutData, 30, 
+		PrimitiveUniformShaderParameters.InstanceDrawDistanceMinMaxSquared.X, 
+		PrimitiveUniformShaderParameters.InstanceDrawDistanceMinMaxSquared.Y, 
+		PrimitiveUniformShaderParameters.InstanceWPODisableDistanceSquared, 
+		PrimitiveUniformShaderParameters.NaniteRayTracingDataOffset);
 	
-	Data[28].X = PrimitiveUniformShaderParameters.WireframeColor.X;
-	Data[28].Y = PrimitiveUniformShaderParameters.WireframeColor.Y;
-	Data[28].Z = PrimitiveUniformShaderParameters.WireframeColor.Z;
-	Data[28].W = FMath::AsFloat(PrimitiveUniformShaderParameters.PackedNaniteFlags);
-
-	Data[29].X = PrimitiveUniformShaderParameters.LevelColor.X;
-	Data[29].Y = PrimitiveUniformShaderParameters.LevelColor.Y;
-	Data[29].Z = PrimitiveUniformShaderParameters.LevelColor.Z;
-	Data[29].W = FMath::AsFloat(uint32(PrimitiveUniformShaderParameters.PersistentPrimitiveIndex));
-
-	Data[30].X = PrimitiveUniformShaderParameters.InstanceDrawDistanceMinMaxSquared.X;
-	Data[30].Y = PrimitiveUniformShaderParameters.InstanceDrawDistanceMinMaxSquared.Y;
-	Data[30].Z = PrimitiveUniformShaderParameters.InstanceWPODisableDistanceSquared;
-	Data[30].W = FMath::AsFloat(PrimitiveUniformShaderParameters.NaniteRayTracingDataOffset);
-
-	Data[31].X = PrimitiveUniformShaderParameters.MaxWPOExtent;
-	Data[31].Y = PrimitiveUniformShaderParameters.MinMaterialDisplacement;
-	Data[31].Z = PrimitiveUniformShaderParameters.MaxMaterialDisplacement;
-	Data[31].W = FMath::AsFloat(PrimitiveUniformShaderParameters.CustomStencilValueAndMask);
+	Store4(OutData, 31,
+		PrimitiveUniformShaderParameters.MaxWPOExtent,
+		PrimitiveUniformShaderParameters.MinMaterialDisplacement,
+		PrimitiveUniformShaderParameters.MaxMaterialDisplacement,
+		PrimitiveUniformShaderParameters.CustomStencilValueAndMask);
 
 	// Set all the custom primitive data float4. This matches the loop in SceneData.ush
 	const int32 CustomPrimitiveDataStartIndex = 32;
 	for (int32 DataIndex = 0; DataIndex < FCustomPrimitiveData::NumCustomPrimitiveDataFloat4s; ++DataIndex)
 	{
-		Data[CustomPrimitiveDataStartIndex + DataIndex] = PrimitiveUniformShaderParameters.CustomPrimitiveData[DataIndex];
+		OutData[CustomPrimitiveDataStartIndex + DataIndex] = PrimitiveUniformShaderParameters.CustomPrimitiveData[DataIndex];
 	}
 }
 
+void FPrimitiveSceneShaderData::Setup(const FPrimitiveUniformShaderParameters& PrimitiveUniformShaderParameters)
+{
+	FPrimitiveSceneShaderData::Setup(PrimitiveUniformShaderParameters, Data.GetData());
+}
+
 TUniformBufferRef<FPrimitiveUniformShaderParameters> CreatePrimitiveUniformBufferImmediate(
 	const FMatrix& LocalToWorld,
 	const FBoxSphereBounds& WorldBounds,
diff --git a/Engine/Source/Runtime/Engine/Public/InstanceUniformShaderParameters.h b/Engine/Source/Runtime/Engine/Public/InstanceUniformShaderParameters.h
--- a/Engine/Source/Runtime/Engine/Public/InstanceUniformShaderParameters.h
+++ b/Engine/Source/Runtime/Engine/Public/InstanceUniformShaderParameters.h
@@ -10,6 +10,7 @@
 #include "ShaderParameters.h"
 #include "UniformBuffer.h"
 #include "Containers/StaticArray.h"
+#include "SceneDefinitions.h"
 
 #define INVALID_LAST_UPDATE_FRAME 0xFFFFFFFFu
 
@@ -74,6 +75,7 @@ private:
 	static constexpr uint32 UnCompressedTransformDataStrideInFloat4s = 4;
 
 public:
+	static ENGINE_API bool SupportsCompressedTransforms();
 
 	static ENGINE_API uint32 GetDataStrideInFloat4s();
 
@@ -110,7 +112,7 @@ public:
 		bool bIsVisible = true
 	);
 
-	ENGINE_API void BuildInternal
+	FORCEINLINE void BuildInternal
 	(
 		uint32 PrimitiveId,
 		uint32 RelativeId,
@@ -119,8 +121,53 @@ public:
 		uint32 CustomDataCount,
 		float RandomID,
 		const FRenderTransform& LocalToWorld,
-		bool bIsVisible = true
-	);
+		bool bIsVisible,
+		bool bSupportsCompressedTransforms
+	)
+	{
+		// Note: layout must match GetInstanceData in SceneData.ush and InitializeInstanceSceneData in GPUSceneWriter.ush
+
+		const float RotDeterminant = LocalToWorld.RotDeterminant();
+		if (RotDeterminant < 0.0f)
+		{
+			InstanceFlags |= INSTANCE_SCENE_DATA_FLAG_DETERMINANT_SIGN;
+		}
+		else
+		{
+			InstanceFlags &= ~INSTANCE_SCENE_DATA_FLAG_DETERMINANT_SIGN;
+		}
+
+		// Mark zero scaled instances as hidden.
+		if (!bIsVisible || RotDeterminant == 0.0f)
+		{
+			InstanceFlags |= INSTANCE_SCENE_DATA_FLAG_HIDDEN;
+		}
+
+		checkSlow((PrimitiveId		& 0x000FFFFF) == PrimitiveId);
+		checkSlow((InstanceFlags	& 0x00000FFF) == InstanceFlags);
+		checkSlow((RelativeId		& 0x00FFFFFF) == RelativeId);
+		checkSlow((CustomDataCount	& 0x000000FF) == CustomDataCount);
+
+		const uint32 Packed0 = (InstanceFlags   << 20u) | PrimitiveId;
+		const uint32 Packed1 = (CustomDataCount << 24u) | RelativeId;
+
+		Data[0].X  = *(const float*)&Packed0;
+		Data[0].Y  = *(const float*)&Packed1;
+		Data[0].Z  = *(const float*)&LastUpdateFrame;
+		Data[0].W  = *(const float*)&RandomID;
+
+		if (bSupportsCompressedTransforms)
+		{
+			FCompressedTransform CompressedLocalToWorld(LocalToWorld);
+			Data[1] = *(const FVector4f*)&CompressedLocalToWorld.Rotation[0];
+			Data[2] = *(const FVector3f*)&CompressedLocalToWorld.Translation;
+		}
+		else
+		{
+			// Note: writes 3x float4s
+			LocalToWorld.To3x4MatrixTranspose((float*)&Data[1]);
+		}
+	}
 
 	TStaticArray<FVector4f, UnCompressedTransformDataStrideInFloat4s> Data;
 };
diff --git a/Engine/Source/Runtime/Engine/Public/PrimitiveSceneProxy.h b/Engine/Source/Runtime/Engine/Public/PrimitiveSceneProxy.h
--- a/Engine/Source/Runtime/Engine/Public/PrimitiveSceneProxy.h
+++ b/Engine/Source/Runtime/Engine/Public/PrimitiveSceneProxy.h
@@ -15,7 +15,9 @@
 #include "UniformBuffer.h"
 #include "SceneView.h"
 #include "PrimitiveUniformShaderParameters.h"
+#if UE_ENABLE_INCLUDE_ORDER_DEPRECATED_IN_5_4
 #include "InstanceUniformShaderParameters.h"
+#endif
 #include "DrawDebugHelpers.h"
 #include "Math/CapsuleShape.h"
 #include "SceneDefinitions.h"
@@ -40,6 +42,8 @@ class IHeterogeneousVolumeInterface;
 struct FPrimitiveUniformShaderParametersBuilder;
 struct FPrimitiveSceneProxyDesc;
 class IPrimitiveComponent;
+struct FRenderTransform;
+struct FRenderBounds;
 
 namespace Nanite
 {
diff --git a/Engine/Source/Runtime/Engine/Public/PrimitiveSceneShaderData.h b/Engine/Source/Runtime/Engine/Public/PrimitiveSceneShaderData.h
--- a/Engine/Source/Runtime/Engine/Public/PrimitiveSceneShaderData.h
+++ b/Engine/Source/Runtime/Engine/Public/PrimitiveSceneShaderData.h
@@ -12,6 +12,8 @@
 #include "PrimitiveUniformShaderParameters.h"
 #endif
 
+class FPrimitiveSceneProxy;
+
 struct FPrimitiveSceneShaderData
 {
 	static const uint32 DataStrideInFloat4s = PRIMITIVE_SCENE_DATA_STRIDE;
@@ -30,7 +32,16 @@ struct FPrimitiveSceneShaderData
 		Setup(PrimitiveUniformShaderParameters);
 	}
 
-	ENGINE_API FPrimitiveSceneShaderData(const class FPrimitiveSceneProxy* RESTRICT Proxy);
+	ENGINE_API FPrimitiveSceneShaderData(const FPrimitiveSceneProxy* RESTRICT Proxy);
+
+	/**
+	 * Directly construct the data from the proxy into an output array, removing the need to construct an intermediate.
+	 */
+	ENGINE_API static void BuildDataFromProxy(const FPrimitiveSceneProxy* RESTRICT Proxy, FVector4f* RESTRICT OutData);
+
+	/**
+	 */
+	ENGINE_API static void Setup(const FPrimitiveUniformShaderParameters& PrimitiveUniformShaderParameters, FVector4f* RESTRICT OutData);
 
 	ENGINE_API void Setup(const FPrimitiveUniformShaderParameters& PrimitiveUniformShaderParameters);
 };
diff --git a/Engine/Source/Runtime/Renderer/Private/GPUScene.cpp b/Engine/Source/Runtime/Renderer/Private/GPUScene.cpp
--- a/Engine/Source/Runtime/Renderer/Private/GPUScene.cpp
+++ b/Engine/Source/Runtime/Renderer/Private/GPUScene.cpp
@@ -30,6 +30,10 @@
 #include "PrimitiveUniformShaderParametersBuilder.h"
 #include "InstanceDataSceneProxy.h"
 
+// Useful for debugging
+#define FORCEINLINE_GPUSCENE FORCEINLINE
+//#define FORCEINLINE_GPUSCENE
+
 // Defaults to being disabled, enable using the command line argument: -CsvCategory GPUScene
 CSV_DEFINE_CATEGORY(GPUScene, false);
 
@@ -248,13 +252,6 @@ FGPUScenePrimitiveCollector::FUploadData* FGPUScenePrimitiveCollector::AllocateU
 	return GPUSceneDynamicContext->AllocateDynamicPrimitiveData();
 }
 
-struct FBVHNode
-{
-	uint32		ChildIndexes[4];
-	FVector4	ChildMin[3];
-	FVector4	ChildMax[3];
-};
-
 /**
  * Info needed by the uploader to prepare to upload a primitive.
  */
@@ -271,13 +268,7 @@ struct FPrimitiveUploadInfoHeader
 	const Nanite::FSceneProxyBase* NaniteSceneProxy = nullptr;
 	const FPrimitiveSceneInfo* PrimitiveSceneInfo = nullptr;
 };
-/**
- * Info needed by the uploader to update a primitive.
- */
-struct FPrimitiveUploadInfo : public FPrimitiveUploadInfoHeader
-{
-	FPrimitiveSceneShaderData PrimitiveSceneData;
-};
+
 /**
  * Info required by the uploader to update the instances that belong to a primitive.
  */
@@ -303,10 +294,6 @@ struct FInstanceUploadInfo
 	TConstArrayView<uint32> InstanceEditorData;
 #endif
 
-	// Used for primitives that need to create a dummy instance (they do not have instance data in the proxy)
-	FInstanceSceneData DummyInstance;
-	FRenderBounds DummyLocalBounds;
-
 	uint32 InstanceFlags = 0x0;
 
 	FRenderTransform PrimitiveToWorld;
@@ -320,9 +307,9 @@ struct FInstanceUploadInfo
 #endif
 };
 
+#if DO_CHECK
 void ValidateInstanceUploadInfo(const FInstanceUploadInfo& UploadInfo, FRDGBuffer* InstancePayloadDataBuffer)
 {
-#if DO_CHECK
 	const bool bHasRandomID			= (UploadInfo.InstanceFlags & INSTANCE_SCENE_DATA_FLAG_HAS_RANDOM) != 0u;
 	const bool bHasCustomData		= (UploadInfo.InstanceFlags & INSTANCE_SCENE_DATA_FLAG_HAS_CUSTOM_DATA) != 0u;
 	const bool bHasDynamicData		= (UploadInfo.InstanceFlags & INSTANCE_SCENE_DATA_FLAG_HAS_DYNAMIC_DATA) != 0u;
@@ -388,7 +375,7 @@ void ValidateInstanceUploadInfo(const FInstanceUploadInfo& UploadInfo, FRDGBuffe
 	else
 	{
 		check(UploadInfo.InstanceDynamicData.Num()	== (bHasDynamicData	? UploadInfo.NumInstances : 0));
-		check(UploadInfo.NumInstances == UploadInfo.PrimitiveInstances.Num());
+		check(UploadInfo.NumInstances == UploadInfo.PrimitiveInstances.Num()  || UploadInfo.NumInstances == 1 && UploadInfo.PrimitiveInstances.IsEmpty());
 	}
 
 	if (UploadInfo.PrimitiveSceneInfo)
@@ -396,8 +383,10 @@ void ValidateInstanceUploadInfo(const FInstanceUploadInfo& UploadInfo, FRDGBuffe
 		check(UploadInfo.PrimitiveSceneInfo->GetNumInstanceSceneDataEntries() == UploadInfo.NumInstances);
 	}
 
-#endif
 }
+#else
+FORCEINLINE void ValidateInstanceUploadInfo(const FInstanceUploadInfo& , FRDGBuffer* ) {}
+#endif
 
 /**
  * Info required by the uploader to update the lightmap data for a primitive.
@@ -425,7 +414,7 @@ struct FUploadDataSourceAdapterScenePrimitives
 	{}
 
 	/**
-	 * Return the number of primitives to upload N, GetPrimitiveInfo will be called with ItemIndex in [0,N).
+	 * Return the number of primitives to upload N, GetPrimitiveInfoHeader/GetPrimitiveShaderData will be called with ItemIndex in [0,N).
 	 */
 	FORCEINLINE int32 NumPrimitivesToUpload() const 
 	{ 
@@ -441,7 +430,7 @@ struct FUploadDataSourceAdapterScenePrimitives
 	 * Populate the primitive info for a given item index.
 	 * 
 	 */
-	FORCEINLINE void GetPrimitiveInfoHeader(int32 ItemIndex, FPrimitiveUploadInfoHeader& PrimitiveUploadInfo) const
+	FORCEINLINE_GPUSCENE void GetPrimitiveInfoHeader(int32 ItemIndex, FPrimitiveUploadInfoHeader& PrimitiveUploadInfo) const
 	{
 		const FPersistentPrimitiveIndex PersistentPrimitiveIndex = PrimitivesToUpdate[ItemIndex];
 		const int32 PrimitiveID = Scene.GetPrimitiveIndex(PersistentPrimitiveIndex);
@@ -482,21 +471,17 @@ struct FUploadDataSourceAdapterScenePrimitives
 	 * Populate the primitive info for a given item index.
 	 * 
 	 */
-	FORCEINLINE void GetPrimitiveInfo(int32 ItemIndex, FPrimitiveUploadInfo& PrimitiveUploadInfo) const
+	FORCEINLINE_GPUSCENE void GetPrimitiveShaderData(int32 ItemIndex, FVector4f* RESTRICT OutData) const
 	{
 		const FPersistentPrimitiveIndex PersistentPrimitiveIndex = PrimitivesToUpdate[ItemIndex];
 		const int32 PrimitiveID = Scene.GetPrimitiveIndex(PersistentPrimitiveIndex);
 		check(Scene.PrimitiveSceneProxies.IsValidIndex(PrimitiveID));
 
-		GetPrimitiveInfoHeader(ItemIndex, PrimitiveUploadInfo);
-
 		const FPrimitiveSceneProxy* RESTRICT PrimitiveSceneProxy = Scene.PrimitiveSceneProxies[PrimitiveID];
-		const FPrimitiveSceneInfo* PrimitiveSceneInfo = PrimitiveSceneProxy->GetPrimitiveSceneInfo();
-
-		PrimitiveUploadInfo.PrimitiveSceneData = FPrimitiveSceneShaderData(PrimitiveSceneProxy);
+		FPrimitiveSceneShaderData::BuildDataFromProxy(PrimitiveSceneProxy, OutData);
 	}
 
-	FORCEINLINE void GetInstanceInfo(int32 ItemIndex, FInstanceUploadInfo& InstanceUploadInfo) const
+	FORCEINLINE_GPUSCENE void GetInstanceInfo(int32 ItemIndex, FInstanceUploadInfo& InstanceUploadInfo) const
 	{
 		const FPersistentPrimitiveIndex PersistentPrimitiveIndex = PrimitivesToUpdate[ItemIndex];
 		const int32 PrimitiveID = Scene.GetPrimitiveIndex(PersistentPrimitiveIndex);
@@ -509,35 +494,18 @@ struct FUploadDataSourceAdapterScenePrimitives
 #if DO_CHECK
 		InstanceUploadInfo.PrimitiveSceneInfo = PrimitiveSceneInfo;
 #endif
-		const FMatrix LocalToWorld = PrimitiveSceneProxy->GetLocalToWorld();
-		InstanceUploadInfo.InstanceSceneDataBuffers = PrimitiveSceneInfo->GetInstanceSceneDataBuffers();
-
-		const FLargeWorldRenderPosition AbsoluteOrigin(LocalToWorld.GetOrigin());
 
 		InstanceUploadInfo.InstanceSceneDataOffset = PrimitiveSceneInfo->GetInstanceSceneDataOffset();
 		check(InstanceUploadInfo.InstanceSceneDataOffset >= 0);
 		InstanceUploadInfo.InstancePayloadDataOffset = PrimitiveSceneInfo->GetInstancePayloadDataOffset();
 		InstanceUploadInfo.InstancePayloadDataStride = PrimitiveSceneInfo->GetInstancePayloadDataStride();
-
 		InstanceUploadInfo.LastUpdateSceneFrameNumber = SceneFrameNumber;
 		InstanceUploadInfo.PrimitiveID = PersistentPrimitiveIndex.Index;
-		InstanceUploadInfo.PrimitiveToWorld = FLargeWorldRenderScalar::MakeToRelativeWorldMatrix(AbsoluteOrigin.GetTileOffset(), LocalToWorld);
-		
 		// HACK: ignoring IsForceHidden for non-Nanite due to issues that cropped up with water rendering.
 		// TODO: Remove the IsNaniteMesh() check
 		InstanceUploadInfo.bIsPrimitiveForceHidden = PrimitiveSceneProxy->IsNaniteMesh() && PrimitiveSceneProxy->IsForceHidden();
 
-		{
-			bool bHasPrecomputedVolumetricLightmap{};
-			bool bOutputVelocity{};
-			int32 SingleCaptureIndex{};
-
-			FMatrix PreviousLocalToWorld;
-			Scene.GetPrimitiveUniformShaderParameters_RenderThread(PrimitiveSceneInfo, bHasPrecomputedVolumetricLightmap, PreviousLocalToWorld, SingleCaptureIndex, bOutputVelocity);
-			InstanceUploadInfo.PrevPrimitiveToWorld = FLargeWorldRenderScalar::MakeClampedToRelativeWorldMatrix(AbsoluteOrigin.GetTileOffset(), PreviousLocalToWorld);;
-		}
-
-
+		InstanceUploadInfo.InstanceSceneDataBuffers = PrimitiveSceneInfo->GetInstanceSceneDataBuffers();
 		if (InstanceUploadInfo.InstanceSceneDataBuffers)
 		{
 			const FInstanceSceneDataBuffers::FReadView InstanceSceneDataBuffers = InstanceUploadInfo.InstanceSceneDataBuffers->GetReadView();
@@ -552,20 +520,41 @@ struct FUploadDataSourceAdapterScenePrimitives
 #if WITH_EDITOR
 			InstanceUploadInfo.InstanceEditorData = InstanceSceneDataBuffers.InstanceEditorData;
 #endif
+#if DO_CHECK
+			// This is already precomputed in the InstanceSceneDataBuffers and we don't need to do it again here, except for validation purposes
+			// TODO: this validation should (also?) move elsewhere and validate that the transform on RT matches that on GT
+			const FMatrix LocalToWorld = PrimitiveSceneProxy->GetLocalToWorld();
+			const FLargeWorldRenderPosition AbsoluteOrigin(LocalToWorld.GetOrigin());
+			InstanceUploadInfo.PrimitiveToWorld = FLargeWorldRenderScalar::MakeToRelativeWorldMatrix(AbsoluteOrigin.GetTileOffset(), LocalToWorld);
+#endif
+
 		}
 		else
 		{
 			// Old path, only taken for uninstanced primitives.
+			const FMatrix LocalToWorld = PrimitiveSceneProxy->GetLocalToWorld();
+			const FLargeWorldRenderPosition AbsoluteOrigin(LocalToWorld.GetOrigin());
+			InstanceUploadInfo.PrimitiveToWorld = FLargeWorldRenderScalar::MakeToRelativeWorldMatrix(AbsoluteOrigin.GetTileOffset(), LocalToWorld);
+
 			InstanceUploadInfo.InstanceFlags = 0u;
 			check(InstanceUploadInfo.InstancePayloadDataOffset == INDEX_NONE && InstanceUploadInfo.InstancePayloadDataStride == 0);
 
-			// We always create an instance to ensure that we can always use the same code paths in the shader
-			// In the future we should remove redundant data from the primitive, and then the instances should be
-			// provided by the proxy. However, this is a lot of work before we can just enable it in the base proxy class.
-			InstanceUploadInfo.DummyInstance.LocalToPrimitive.SetIdentity();
-			// TODO: replance DummyInstance with explicit primitive transform reference since it only ever is used for that case and it would be cheaper (no orthonorm etc)
-			InstanceUploadInfo.PrimitiveInstances = TConstArrayView<FInstanceSceneData>(&InstanceUploadInfo.DummyInstance, 1);
+			// empty array signals that we should use the PrimitiveToWorld as the instance transform (and that it is therefore pre-transformed).
+			InstanceUploadInfo.PrimitiveInstances = TConstArrayView<FInstanceSceneData>();
 			InstanceUploadInfo.InstanceDynamicData = TConstArrayView<FInstanceDynamicData>();
+
+#if 0
+			// NOTE: We only need this if not using the InstanceSceneDataBuffers and if the old path were to support dynamic data
+			{
+				bool bHasPrecomputedVolumetricLightmap{};
+				bool bOutputVelocity{};
+				int32 SingleCaptureIndex{};
+
+				FMatrix PreviousLocalToWorld;
+				Scene.GetPrimitiveUniformShaderParameters_RenderThread(PrimitiveSceneInfo, bHasPrecomputedVolumetricLightmap, PreviousLocalToWorld, SingleCaptureIndex, bOutputVelocity);
+				InstanceUploadInfo.PrevPrimitiveToWorld = FLargeWorldRenderScalar::MakeClampedToRelativeWorldMatrix(AbsoluteOrigin.GetTileOffset(), PreviousLocalToWorld);;
+			}
+#endif
 			InstanceUploadInfo.InstanceLightShadowUVBias = TConstArrayView<FVector4f>();
 			InstanceUploadInfo.InstanceCustomData = TConstArrayView<float>();
 			InstanceUploadInfo.InstanceRandomID = TConstArrayView<float>();
@@ -577,12 +566,6 @@ struct FUploadDataSourceAdapterScenePrimitives
 			InstanceUploadInfo.InstanceEditorData = TConstArrayView<uint32>();
 #endif
 		}
-		
-		if (InstanceUploadInfo.InstanceLocalBounds.Num() == 0)
-		{
-			InstanceUploadInfo.DummyLocalBounds = PrimitiveSceneProxy->GetLocalBounds();
-			InstanceUploadInfo.InstanceLocalBounds = TConstArrayView<FRenderBounds>(&InstanceUploadInfo.DummyLocalBounds, 1);
-		}
 
 		InstanceUploadInfo.InstancePayloadExtensionCount = 0;
 		if (InstanceUploadInfo.InstancePayloadExtension.Num() > 0)
@@ -600,7 +583,7 @@ struct FUploadDataSourceAdapterScenePrimitives
 		check(InstanceUploadInfo.NumInstances > 0);
 	}
 
-	FORCEINLINE bool GetLightMapInfo(int32 ItemIndex, FLightMapUploadInfo &UploadInfo) const
+	FORCEINLINE_GPUSCENE bool GetLightMapInfo(int32 ItemIndex, FLightMapUploadInfo &UploadInfo) const
 	{
 		const FPersistentPrimitiveIndex PersistentPrimitiveIndex = PrimitivesToUpdate[ItemIndex];
 		const int32 PrimitiveID = Scene.GetPrimitiveIndex(PersistentPrimitiveIndex);
@@ -943,9 +926,6 @@ void FGPUScene::UpdateBufferState(FRDGBuilder& GraphBuilder, FSceneUniformBuffer
  */
 struct FInstanceUploadBatch
 {
-	static constexpr int32 MaxItems = 64;
-	static constexpr int32 MaxCost = MaxItems * 2; // Selected to allow filling the array when 1:1 primitive / instances
-
 	struct FItem
 	{
 		int32 ItemIndex;
@@ -959,6 +939,9 @@ struct FInstanceUploadBatch
 
 struct FInstanceBatcher
 {
+	int32 MaxItems = 64;
+	int32 MaxCost = MaxItems * 2; // Selected to allow filling the array when 1:1 primitive / instances
+
 	struct FPrimitiveItemInfo
 	{
 		int32 InstanceSceneDataUploadOffset;
@@ -975,8 +958,20 @@ struct FInstanceBatcher
 	int32 InstanceSceneDataUploadOffset = 0;
 	int32 InstancePayloadDataUploadOffset = 0; // Count of float4s
 
-	FInstanceBatcher()
+	FInstanceBatcher(bool bExecuteInParallel, int32 NumPrimitiveDataUploads)
 	{
+		if (bExecuteInParallel)
+		{
+			MaxItems = 64;
+			MaxCost = MaxItems * 2;
+		}
+		else
+		{
+			// funnel  all items into a single batch with one item per primitive.
+			MaxItems = MAX_int32;
+			MaxCost = MAX_int32;
+			UpdateBatchItems.Reserve(NumPrimitiveDataUploads);
+		}
 		CurrentBatch = &UpdateBatches.AddDefaulted_GetRef();
 	}
 
@@ -991,7 +986,7 @@ struct FInstanceBatcher
 			// Fill to minimum cost (1
 
 			// Can add one less to account for primitive cost
-			int32 MaxInstancesThisBatch = FInstanceUploadBatch::MaxCost - CurrentBatchCost - 1;
+			int32 MaxInstancesThisBatch = MaxCost - CurrentBatchCost - 1;
 
 			if (MaxInstancesThisBatch > 0)
 			{
@@ -1003,7 +998,7 @@ struct FInstanceBatcher
 			}
 
 			// Flush batch if it is not possible to add any more items (for one of the reasons)
-			if (MaxInstancesThisBatch <= 0 || CurrentBatchCost > FInstanceUploadBatch::MaxCost - 1 || CurrentBatch->NumItems >= FInstanceUploadBatch::MaxItems)
+			if (MaxInstancesThisBatch <= 0 || CurrentBatchCost > MaxCost - 1 || CurrentBatch->NumItems >= MaxItems)
 			{
 				CurrentBatchCost = 0;
 				CurrentBatch = &UpdateBatches.AddDefaulted_GetRef();
@@ -1117,7 +1112,7 @@ void FGPUScene::UploadGeneral(FRDGBuilder& GraphBuilder, FScene& InScene, FRDGEx
 			LockIfValid(RHICmdList, Uploader);
 		}
 
-		FInstanceBatcher InstanceUpdates;
+		FInstanceBatcher InstanceUpdates(bExecuteInParallel, TaskContext.NumPrimitiveDataUploads);
 
 		{
 			SCOPED_NAMED_EVENT(Primitives, FColor::Green);
@@ -1191,15 +1186,8 @@ void FGPUScene::UploadGeneral(FRDGBuilder& GraphBuilder, FScene& InScene, FRDGEx
 			{
 				FTaskTagScope TaskTagScope(ETaskTag::EParallelRenderingThread);
 
-				FPrimitiveUploadInfo UploadInfo;
-				UploadDataSourceAdapter.GetPrimitiveInfo(ItemIndex, UploadInfo);
-
 				FVector4f* DstData = static_cast<FVector4f*>(TaskContext.PrimitiveUploader->GetRef(ItemIndex));
-				for (uint32 VectorIndex = 0; VectorIndex < FPrimitiveSceneShaderData::DataStrideInFloat4s; ++VectorIndex)
-				{
-					DstData[VectorIndex] = UploadInfo.PrimitiveSceneData.Data[VectorIndex];
-				}
-
+				UploadDataSourceAdapter.GetPrimitiveShaderData(ItemIndex, DstData);
 			}, bExecuteInParallel ? EParallelForFlags::None : EParallelForFlags::ForceSingleThread);
 		}
 
@@ -1210,6 +1198,8 @@ void FGPUScene::UploadGeneral(FRDGBuilder& GraphBuilder, FScene& InScene, FRDGEx
 			ParallelForTemplate(TEXT("GPUScene Upload Instances Task"), InstanceUpdates.UpdateBatches.Num(), 1, [&TaskContext, &InstanceUpdates, &UploadDataSourceAdapter](int32 BatchIndex)
 			{
 				const FInstanceUploadBatch Batch = InstanceUpdates.UpdateBatches[BatchIndex];
+				const uint32 InstanceDataStrideInFloat4s = FInstanceSceneShaderData::GetDataStrideInFloat4s();
+				const bool bSupportsCompressedTransforms = FInstanceSceneShaderData::SupportsCompressedTransforms();
 
 				for (int32 BatchItemIndex = 0; BatchItemIndex < Batch.NumItems; ++BatchItemIndex)
 				{
@@ -1241,7 +1231,24 @@ void FGPUScene::UploadGeneral(FRDGBuilder& GraphBuilder, FScene& InScene, FRDGEx
 								UploadInfo.InstanceCustomDataCount,
 								RandomID,
 								UploadInfo.InstanceSceneDataBuffers->GetInstanceToPrimitiveRelative(InstanceIndex),
-								!UploadInfo.bIsPrimitiveForceHidden && UploadInfo.InstanceSceneDataBuffers->GetInstanceVisible(InstanceIndex)
+								!UploadInfo.bIsPrimitiveForceHidden && UploadInfo.InstanceSceneDataBuffers->GetInstanceVisible(InstanceIndex),
+								bSupportsCompressedTransforms
+							);
+						}
+						else if (UploadInfo.PrimitiveInstances.IsEmpty())
+						{
+							// This path should only be taken for uninstanced primitives
+							check(UploadInfo.NumInstances == 1 && InstanceIndex == 0);
+							InstanceSceneData.BuildInternal(
+								UploadInfo.PrimitiveID,
+								InstanceIndex,
+								UploadInfo.InstanceFlags,
+								UploadInfo.LastUpdateSceneFrameNumber,
+								UploadInfo.InstanceCustomDataCount,
+								RandomID,
+								UploadInfo.PrimitiveToWorld,
+								!UploadInfo.bIsPrimitiveForceHidden,
+								bSupportsCompressedTransforms
 							);
 						}
 						else
@@ -1260,9 +1267,9 @@ void FGPUScene::UploadGeneral(FRDGBuilder& GraphBuilder, FScene& InScene, FRDGEx
 							);
 						}
 						// RefIndex* BufferState.InstanceSceneDataSOAStride + UploadInfo.InstanceSceneDataOffset + InstanceIndex
-						const uint32 UploadInstanceItemOffset = (PrimitiveItemInfo.InstanceSceneDataUploadOffset + InstanceIndex) * FInstanceSceneShaderData::GetDataStrideInFloat4s();
+						const uint32 UploadInstanceItemOffset = (PrimitiveItemInfo.InstanceSceneDataUploadOffset + InstanceIndex) * InstanceDataStrideInFloat4s;
 
-						for (uint32 RefIndex = 0; RefIndex < FInstanceSceneShaderData::GetDataStrideInFloat4s(); ++RefIndex)
+						for (uint32 RefIndex = 0; RefIndex < InstanceDataStrideInFloat4s; ++RefIndex)
 						{
 							FVector4f* DstVector = static_cast<FVector4f*>(TaskContext.InstanceSceneUploader->Set_GetRef(UploadInstanceItemOffset + RefIndex, RefIndex * TaskContext.InstanceSceneDataSOAStride + UploadInfo.InstanceSceneDataOffset + InstanceIndex));
 							*DstVector = InstanceSceneData.Data[RefIndex];
@@ -1312,7 +1319,7 @@ void FGPUScene::UploadGeneral(FRDGBuilder& GraphBuilder, FScene& InScene, FRDGEx
 									UploadInfo.InstanceSceneDataBuffers->GetPrevInstanceToPrimitiveRelative(InstanceIndex) :
 									UploadInfo.InstanceDynamicData[InstanceIndex].ComputePrevLocalToWorld(UploadInfo.PrevPrimitiveToWorld);
 
-								if (FDataDrivenShaderPlatformInfo::GetSupportSceneDataCompressedTransforms(GMaxRHIShaderPlatform))
+								if (bSupportsCompressedTransforms)
 								{
 									check(PayloadPosition + 1 < InstancePayloadData.Num()); // Sanity check
 									FCompressedTransform CompressedPrevLocalToWorld(PrevLocalToWorld);
@@ -1469,13 +1476,13 @@ struct FUploadDataSourceAdapterDynamicPrimitives
 	}
 
 
-	FORCEINLINE TArrayView<const uint32> GetItemPrimitiveIds() const
+	FORCEINLINE_GPUSCENE TArrayView<const uint32> GetItemPrimitiveIds() const
 	{
 		return TArrayView<const uint32>(PrimitivesIds.GetData(), PrimitivesIds.Num());
 	}
 
 
-	FORCEINLINE void GetPrimitiveInfoHeader(int32 ItemIndex, FPrimitiveUploadInfoHeader& PrimitiveUploadInfo) const
+	FORCEINLINE_GPUSCENE void GetPrimitiveInfoHeader(int32 ItemIndex, FPrimitiveUploadInfoHeader& PrimitiveUploadInfo) const
 	{
 		PrimitiveUploadInfo.LightmapUploadCount = 0;
 		PrimitiveUploadInfo.NaniteSceneProxy = nullptr;
@@ -1500,10 +1507,8 @@ struct FUploadDataSourceAdapterDynamicPrimitives
 		}
 	}
 
-	FORCEINLINE void GetPrimitiveInfo(int32 ItemIndex, FPrimitiveUploadInfo& PrimitiveUploadInfo) const
+	FORCEINLINE_GPUSCENE void GetPrimitiveShaderData(int32 ItemIndex, FVector4f* RESTRICT OutData) const
 	{
-		GetPrimitiveInfoHeader(ItemIndex, PrimitiveUploadInfo);
-
 		// Needed to ensure the link back to instance list is up to date
 		const FGPUScenePrimitiveCollector::FPrimitiveData& PrimData = PrimitiveData[ItemIndex];		
 		FPrimitiveUniformShaderParameters Tmp = *PrimData.ShaderParams;
@@ -1521,10 +1526,10 @@ struct FUploadDataSourceAdapterDynamicPrimitives
 			Tmp.InstancePayloadDataStride = 0;
 		}
 
-		PrimitiveUploadInfo.PrimitiveSceneData = FPrimitiveSceneShaderData(Tmp);
+		FPrimitiveSceneShaderData::Setup(Tmp, OutData);
 	}
 
-	FORCEINLINE bool GetInstanceInfo(int32 ItemIndex, FInstanceUploadInfo& InstanceUploadInfo) const
+	FORCEINLINE_GPUSCENE bool GetInstanceInfo(int32 ItemIndex, FInstanceUploadInfo& InstanceUploadInfo) const
 	{
 		if (ItemIndex < PrimitiveData.Num())
 		{
@@ -1551,11 +1556,9 @@ struct FUploadDataSourceAdapterDynamicPrimitives
 #endif
 			InstanceUploadInfo.NumInstances = PrimData.SourceData.InstanceSceneData.Num();
 			InstanceUploadInfo.bIsPrimitiveForceHidden = false;
-			// upload dummies where applicable
+			// Use primitive to world as instance transform:
 			if (InstanceUploadInfo.PrimitiveInstances.Num() == 0)
 			{
-				InstanceUploadInfo.DummyInstance.LocalToPrimitive.SetIdentity();
-				InstanceUploadInfo.PrimitiveInstances = TConstArrayView<FInstanceSceneData>(&InstanceUploadInfo.DummyInstance, 1);
 				InstanceUploadInfo.NumInstances = 1;
 			}
 
@@ -1565,7 +1568,7 @@ struct FUploadDataSourceAdapterDynamicPrimitives
 		return false;
 	}
 
-	FORCEINLINE bool GetLightMapInfo(int32 ItemIndex, FLightMapUploadInfo& UploadInfo) const
+	FORCEINLINE_GPUSCENE bool GetLightMapInfo(int32 ItemIndex, FLightMapUploadInfo& UploadInfo) const
 	{
 		return false;
 	}
diff --git a/Engine/Source/Runtime/Renderer/Private/Renderer.cpp b/Engine/Source/Runtime/Renderer/Private/Renderer.cpp
--- a/Engine/Source/Runtime/Renderer/Private/Renderer.cpp
+++ b/Engine/Source/Runtime/Renderer/Private/Renderer.cpp
@@ -206,7 +206,9 @@ FSceneUniformBuffer* FRendererModule::CreateSinglePrimitiveSceneUniformBuffer(FR
 				INVALID_LAST_UPDATE_FRAME,
 				0 /* Custom Data Count */,
 				0.0f /* Random ID */,
-				PrimitiveParams.LocalToRelativeWorld
+				PrimitiveParams.LocalToRelativeWorld,
+				true,
+				FInstanceSceneShaderData::SupportsCompressedTransforms()
 			);
 
 			// Set up the parameters for the LightmapSceneData from the given LCI data 
