Fix assert in Niagara when Ray Tracing Overlay is enabled due to SceneColor not being produced.

diff --git a/Engine/Source/Runtime/Renderer/Private/FXRenderingUtils.cpp b/Engine/Source/Runtime/Renderer/Private/FXRenderingUtils.cpp
--- a/Engine/Source/Runtime/Renderer/Private/FXRenderingUtils.cpp
+++ b/Engine/Source/Runtime/Renderer/Private/FXRenderingUtils.cpp
@@ -66,9 +66,17 @@ TRDGUniformBufferRef<FSceneTextureUniformParameters> UE::FXRenderingUtils::GetOr
 	TRDGUniformBufferRef<FSceneTextureUniformParameters> SceneTexturesUniformParams = nullptr;
 
 	const FViewInfo* View = Views.Num() > 0 ? static_cast<const FViewInfo*>(&Views[0]) : nullptr;
-	if (const FSceneTextures* SceneTextures = View ? static_cast<const FViewFamilyInfo*>(View->Family)->GetSceneTexturesChecked() : nullptr)
+	if (View)
 	{
-		SceneTexturesUniformParams = SceneTextures->UniformBuffer;
+		const FViewFamilyInfo& ViewFamily = *static_cast<const FViewFamilyInfo*>(View->Family);
+
+		if (!HasRayTracedOverlay(ViewFamily))
+		{
+			if (const FSceneTextures* SceneTextures = ViewFamily.GetSceneTexturesChecked())
+			{
+				SceneTexturesUniformParams = SceneTextures->UniformBuffer;
+			}
+		}
 	}
 
 	if (SceneTexturesUniformParams == nullptr)
