Fix assert when both Ray Tracing Overlay and r.AutoExposure.IgnoreMaterials are enabled.

diff --git a/Engine/Source/Runtime/Renderer/Private/DeferredShadingRenderer.cpp b/Engine/Source/Runtime/Renderer/Private/DeferredShadingRenderer.cpp
--- a/Engine/Source/Runtime/Renderer/Private/DeferredShadingRenderer.cpp
+++ b/Engine/Source/Runtime/Renderer/Private/DeferredShadingRenderer.cpp
@@ -3701,8 +3701,12 @@ void FDeferredShadingSceneRenderer::Render(FRDGBuilder& GraphBuilder)
 			// VisualizeVirtualShadowMap TODO
 		}
 
-		// Extract emissive from SceneColor (before lighting is applied) + material diffuse and subsurface colors
-		FRDGTextureRef ExposureIlluminanceSetup = AddSetupExposureIlluminancePass(GraphBuilder, Views, SceneTextures);
+		FRDGTextureRef ExposureIlluminanceSetup = nullptr;
+		if (!bHasRayTracedOverlay)
+		{
+			// Extract emissive from SceneColor (before lighting is applied)
+			ExposureIlluminanceSetup = AddSetupExposureIlluminancePass(GraphBuilder, Views, SceneTextures);
+		}
 
 		if (ViewFamily.EngineShowFlags.VisualizeLightCulling)
 		{
@@ -4201,7 +4205,11 @@ void FDeferredShadingSceneRenderer::Render(FRDGBuilder& GraphBuilder)
 			RenderLigthShaftSkyFogAndCloud();
 		}
 
-		FRDGTextureRef ExposureIlluminance = AddCalculateExposureIlluminancePass(GraphBuilder, Views, SceneTextures, TranslucencyLightingVolumeTextures, ExposureIlluminanceSetup);
+		FRDGTextureRef ExposureIlluminance = nullptr;
+		if (!bHasRayTracedOverlay)
+		{
+			ExposureIlluminance = AddCalculateExposureIlluminancePass(GraphBuilder, Views, SceneTextures, TranslucencyLightingVolumeTextures, ExposureIlluminanceSetup);
+		}
 
 		RenderOpaqueFX(GraphBuilder, GetSceneViews(), GetSceneUniforms(), FXSystem, FeatureLevel, SceneTextures.UniformBuffer);
 
