Fix assert when both Ray Tracing Debug visualization and TSR are enabled.
- don't generate TSR moire luma since Scene Color is not valid.

diff --git a/Engine/Source/Runtime/Renderer/Private/DeferredShadingRenderer.cpp b/Engine/Source/Runtime/Renderer/Private/DeferredShadingRenderer.cpp
--- a/Engine/Source/Runtime/Renderer/Private/DeferredShadingRenderer.cpp
+++ b/Engine/Source/Runtime/Renderer/Private/DeferredShadingRenderer.cpp
@@ -4086,18 +4086,21 @@ void FDeferredShadingSceneRenderer::Render(FRDGBuilder& GraphBuilder)
 		}
 
 		TArray<FScreenPassTexture, TInlineAllocator<4>> TSRMoireInputTextures;
-		// Extract TSR's moire heuristic luminance before renderering translucency into the scene color.
-		for (int32 ViewIndex = 0; ViewIndex < Views.Num(); ++ViewIndex)
+		if (!bHasRayTracedOverlay)
 		{
-			FViewInfo& View = Views[ViewIndex];
-			if (NeedTSRMoireLuma(View))
+			// Extract TSR's moire heuristic luminance before rendering translucency into the scene color.
+			for (int32 ViewIndex = 0; ViewIndex < Views.Num(); ++ViewIndex)
 			{
-				if (TSRMoireInputTextures.Num() == 0)
+				FViewInfo& View = Views[ViewIndex];
+				if (NeedTSRMoireLuma(View))
 				{
-					TSRMoireInputTextures.SetNum(Views.Num());
-				}
+					if (TSRMoireInputTextures.Num() == 0)
+					{
+						TSRMoireInputTextures.SetNum(Views.Num());
+					}
 
-				TSRMoireInputTextures[ViewIndex] = AddTSRComputeMoireLuma(GraphBuilder, View.ShaderMap, FScreenPassTexture(SceneTextures.Color.Target, View.ViewRect));
+					TSRMoireInputTextures[ViewIndex] = AddTSRComputeMoireLuma(GraphBuilder, View.ShaderMap, FScreenPassTexture(SceneTextures.Color.Target, View.ViewRect));
+				}
 			}
 		}
 
diff --git a/Engine/Source/Runtime/Renderer/Private/PostProcess/VisualizeTemporalUpscaler.cpp b/Engine/Source/Runtime/Renderer/Private/PostProcess/VisualizeTemporalUpscaler.cpp
--- a/Engine/Source/Runtime/Renderer/Private/PostProcess/VisualizeTemporalUpscaler.cpp
+++ b/Engine/Source/Runtime/Renderer/Private/PostProcess/VisualizeTemporalUpscaler.cpp
@@ -158,9 +158,17 @@ FScreenPassTexture AddVisualizeTemporalUpscalerPass(FRDGBuilder& GraphBuilder, c
 		if (Inputs.TAAConfig == EMainTAAPassConfig::TSR)
 		{
 			FVisualizeBufferTile& Tile = Tiles[4 * 1 + 3];
-			Tile.Input = Inputs.Inputs.MoireInputTexture;
-			Tile.Input.ViewRect = CropViewRectToCenter(View.ViewRect);
-			Tile.Label = VisualizeTextureLabel(Inputs.Inputs.MoireInputTexture.Texture);
+			if (Inputs.Inputs.MoireInputTexture.IsValid())
+			{
+				Tile.Input = Inputs.Inputs.MoireInputTexture;
+				Tile.Input.ViewRect = CropViewRectToCenter(View.ViewRect);
+				Tile.Label = VisualizeTextureLabel(Inputs.Inputs.MoireInputTexture.Texture);
+			}
+			else
+			{
+				Tile.Input = FScreenPassTexture(GSystemTextures.GetBlackDummy(GraphBuilder));
+				Tile.Label = TEXT("No Moire Luma!");
+			}
 		}
 
 		// Display UMaterial::bHasPixelAnimation used to disable TSR's anti-flickering heuristic (r.TSR.ShadingRejection.Flickering) on per pixel basis
