Add support for per platform enable/disable of RVT volume.
Add support for disable of RVT volume when nanite is disabled.
Both settings help with removing the overhead of physical pool allocations that aren't needed when material setup means that the RVT volume isn't actually in use.

diff --git a/Engine/Source/Runtime/Engine/Classes/Components/RuntimeVirtualTextureComponent.h b/Engine/Source/Runtime/Engine/Classes/Components/RuntimeVirtualTextureComponent.h
--- a/Engine/Source/Runtime/Engine/Classes/Components/RuntimeVirtualTextureComponent.h
+++ b/Engine/Source/Runtime/Engine/Classes/Components/RuntimeVirtualTextureComponent.h
@@ -4,6 +4,7 @@
 
 #include "CoreMinimal.h"
 #include "Engine/TextureDefines.h"
+#include "PerPlatformProperties.h"
 #include "RenderCommandFence.h"
 #include "SceneComponent.h"
 #include "RuntimeVirtualTextureComponent.generated.h"
@@ -36,6 +37,14 @@ protected:
 	UPROPERTY(EditAnywhere, BlueprintReadOnly, TextExportTransient, Category = VirtualTexture)
 	TObjectPtr<URuntimeVirtualTexture> VirtualTexture = nullptr;
 
+	/** Per platform overrides for enabling the virtual texture. Only affects In-Game and PIE. */
+	UPROPERTY(EditAnywhere, AdvancedDisplay, Category = VirtualTexture)
+	FPerPlatformBool EnableInGamePerPlatform;
+
+	/** Enable the virtual texture only when Nanite is enabled. Can be used for a Displacement virtual texture with Nanite tessellation. */
+	UPROPERTY(EditAnywhere, AdvancedDisplay, Category = VirtualTexture)
+	bool bEnableForNaniteOnly = false;
+
 	/** Set to true to enable scalability settings for the virtual texture. */
 	UPROPERTY(EditAnywhere, AdvancedDisplay, Category = VirtualTexture, meta = (InlineEditConditionToggle))
 	bool bEnableScalability = false;
@@ -102,6 +111,9 @@ public:
 	/** Get the runtime virtual texture object on this component. */
 	URuntimeVirtualTexture* GetVirtualTexture() const { return VirtualTexture; }
 
+	/** Get if the runtime virtual texture should be fully instantiated by it's render proxy. */
+	ENGINE_API bool IsEnabledInScene() const;
+
 	/** Get if scalability settings are enabled. */
 	bool IsScalable() const { return bEnableScalability; }
 
@@ -163,6 +175,7 @@ protected:
 	//~ End UObject Interface
 
 	//~ Begin UActorComponent Interface
+	ENGINE_API virtual bool ShouldCreateRenderState() const override;
 	ENGINE_API virtual void CreateRenderState_Concurrent(FRegisterComponentContext* Context) override;
 	ENGINE_API virtual void SendRenderTransform_Concurrent() override;
 	ENGINE_API virtual void DestroyRenderState_Concurrent() override;
@@ -176,7 +189,6 @@ protected:
 	ENGINE_API virtual void OnRegister() override;
 	ENGINE_API virtual void OnUnregister() override;
 #endif
-	ENGINE_API virtual bool IsVisible() const override;
 	ENGINE_API virtual FBoxSphereBounds CalcBounds(const FTransform& LocalToWorld) const override;
 	//~ End USceneComponent Interface
 
diff --git a/Engine/Source/Runtime/Engine/Private/Components/RuntimeVirtualTextureComponent.cpp b/Engine/Source/Runtime/Engine/Private/Components/RuntimeVirtualTextureComponent.cpp
--- a/Engine/Source/Runtime/Engine/Private/Components/RuntimeVirtualTextureComponent.cpp
+++ b/Engine/Source/Runtime/Engine/Private/Components/RuntimeVirtualTextureComponent.cpp
@@ -8,6 +8,8 @@
 #include "Engine/World.h"
 #include "Logging/MessageLog.h"
 #include "GameFramework/Actor.h"
+#include "Interfaces/ITargetPlatform.h"
+#include "Interfaces/ITargetPlatformManagerModule.h"
 #include "Misc/UObjectToken.h"
 #include "Misc/MapErrors.h"
 #include "SceneInterface.h"
@@ -16,6 +18,7 @@
 #include "VT/RuntimeVirtualTexture.h"
 #include "VT/VirtualTextureBuilder.h"
 #include "RenderUtils.h"
+#include "RHIGlobals.h"
 #include "SceneUtils.h"
 
 #include UE_INLINE_GENERATED_CPP_BY_NAME(RuntimeVirtualTextureComponent)
@@ -30,6 +33,7 @@ static TAutoConsoleVariable<int32> CVarRVTEnableVolumes(
 
 URuntimeVirtualTextureComponent::URuntimeVirtualTextureComponent(const FObjectInitializer& ObjectInitializer)
 	: Super(ObjectInitializer)
+	, EnableInGamePerPlatform(true)
 	, SceneProxy(nullptr)
 {
 	Mobility = EComponentMobility::Stationary;
@@ -94,10 +98,10 @@ void URuntimeVirtualTextureComponent::GetHidePrimitiveSettings(bool& OutHidePrim
 	HidePrimitivesDelegate.Broadcast(OutHidePrimitiveEditor, OutHidePrimitiveGame);
 }
 
-bool URuntimeVirtualTextureComponent::IsVisible() const
+bool URuntimeVirtualTextureComponent::ShouldCreateRenderState() const
 {
 	// Make sure to have the component do nothing if VT is disabled or if the world is not compatible with RVT
-	return Super::IsVisible() && IsActiveInWorld() && CVarRVTEnableVolumes.GetValueOnGameThread() != 0 && UseVirtualTexturing(GetScene()->GetShaderPlatform());
+	return Super::ShouldCreateRenderState() && IsActiveInWorld() && CVarRVTEnableVolumes.GetValueOnGameThread() != 0 && UseVirtualTexturing(GetScene()->GetShaderPlatform());
 }
 
 void URuntimeVirtualTextureComponent::ApplyWorldOffset(const FVector& InOffset, bool bWorldShift)
@@ -109,8 +113,7 @@ void URuntimeVirtualTextureComponent::ApplyWorldOffset(const FVector& InOffset,
 
 void URuntimeVirtualTextureComponent::CreateRenderState_Concurrent(FRegisterComponentContext* Context)
 {
-	// Make sure not to create a render state if the world is not compatible with RVT :
-	if (IsActiveInWorld() && ShouldRender() && VirtualTexture != nullptr)
+	if (VirtualTexture != nullptr)
 	{
 		// This will modify the URuntimeVirtualTexture and allocate its VT
 		GetScene()->AddRuntimeVirtualTexture(this);
@@ -121,11 +124,8 @@ void URuntimeVirtualTextureComponent::CreateRenderState_Concurrent(FRegisterComp
 
 void URuntimeVirtualTextureComponent::SendRenderTransform_Concurrent()
 {
-	// We don't have a render state if the world is not compatible with RVT :
-	if (IsActiveInWorld() && ShouldRender() && VirtualTexture != nullptr)
+	if (IsRenderStateCreated())
 	{
-		checkf(IsActiveInWorld(), TEXT("ShouldRender should never return true for a world where we're inactive"));
-
 		// We do a full recreate of the URuntimeVirtualTexture here which can cause a visual glitch.
 		// We do this because, for an arbitrary transform, there is no way to only modify the transform and maintain the VT contents.
 		// Possibly, with some work, the contents could be maintained for any transform change that is an exact multiple of the page size in world space.
@@ -137,18 +137,44 @@ void URuntimeVirtualTextureComponent::SendRenderTransform_Concurrent()
 
 void URuntimeVirtualTextureComponent::DestroyRenderState_Concurrent()
 {
-	if (IsActiveInWorld())
+	// This will modify the URuntimeVirtualTexture and free its VT
+	GetScene()->RemoveRuntimeVirtualTexture(this);
+
+	Super::DestroyRenderState_Concurrent();
+}
+
+bool URuntimeVirtualTextureComponent::IsEnabledInScene() const
+{
+	if (UWorld* World = GetWorld())
 	{
-		// This will modify the URuntimeVirtualTexture and free its VT
-		GetScene()->RemoveRuntimeVirtualTexture(this);
+		if (World->WorldType == EWorldType::Game || World->WorldType == EWorldType::PIE)
+		{
+#if WITH_EDITOR
+			if (!EnableInGamePerPlatform.GetValueForPlatform(*GetTargetPlatformManagerRef().GetRunningTargetPlatform()->IniPlatformName()))
+			{
+				return false;
+			}
+#else
+			if (!EnableInGamePerPlatform.GetValue())
+			{
+				return false;
+			}
+#endif
+		}
+
+		const bool bUseNanite = UseNanite(GetFeatureLevelShaderPlatform(World->GetFeatureLevel()));
+		if (bEnableForNaniteOnly && !bUseNanite)
+		{
+			return false;
+		}
 	}
 
-	Super::DestroyRenderState_Concurrent();
+	return true;
 }
 
 void URuntimeVirtualTextureComponent::Invalidate(FBoxSphereBounds const& InWorldBounds)
 {
-	if (IsActiveInWorld() && (GetScene() != nullptr))
+	if (GetScene() != nullptr)
 	{
 		GetScene()->InvalidateRuntimeVirtualTexture(this, InWorldBounds);
 	}
@@ -296,6 +322,7 @@ static void GetLayerFormatSettings(FTextureFormatSettings& OutFormatSettings, EP
 	OutFormatSettings.CompressionYCoCg = IsLayerYCoCg;
 	OutFormatSettings.SRGB = IsLayerSRGB;
 }
+
 void URuntimeVirtualTextureComponent::InitializeStreamingTexture(EShadingPath ShadingPath, uint32 InSizeX, uint32 InSizeY, uint8* InData)
 {
 	// We need an existing StreamingTexture object to update.
diff --git a/Engine/Source/Runtime/Renderer/Private/VT/RuntimeVirtualTextureSceneProxy.cpp b/Engine/Source/Runtime/Renderer/Private/VT/RuntimeVirtualTextureSceneProxy.cpp
--- a/Engine/Source/Runtime/Renderer/Private/VT/RuntimeVirtualTextureSceneProxy.cpp
+++ b/Engine/Source/Runtime/Renderer/Private/VT/RuntimeVirtualTextureSceneProxy.cpp
@@ -46,88 +46,97 @@ FRuntimeVirtualTextureSceneProxy::FRuntimeVirtualTextureSceneProxy(URuntimeVirtu
 
 	if (InComponent->GetVirtualTexture() != nullptr)
 	{
-		// We store a ProducerId here so that we will be able to find our SceneIndex from the Producer during rendering.
-		// We will need the SceneIndex to determine which primitives should render to this Producer.
-		ProducerId = ProducerIdGenerator++;
-
-		URuntimeVirtualTexture::FInitSettings InitSettings;
-		InitSettings.TileCountBias = InComponent->IsScalable() ? VirtualTextureScalability::GetRuntimeVirtualTextureSizeBias(InComponent->GetScalabilityGroup()) : 0;
-
-		VirtualTexture = InComponent->GetVirtualTexture();
-		Transform = InComponent->GetComponentTransform();
-		const FBox Bounds = InComponent->Bounds.GetBox();
-
-		// The producer description is calculated using the transform to determine the aspect ratio
-		FVTProducerDescription ProducerDesc;
-		VirtualTexture->GetProducerDescription(ProducerDesc, InitSettings, Transform);
-		VirtualTextureSize = FIntPoint(ProducerDesc.BlockWidthInTiles * ProducerDesc.TileSize, ProducerDesc.BlockHeightInTiles * ProducerDesc.TileSize);
-		// We only need to dirty flush up to the producer description MaxLevel which accounts for the RemoveLowMips
-		MaxDirtyLevel = ProducerDesc.MaxLevel;
-
-		const ERuntimeVirtualTextureMaterialType MaterialType = VirtualTexture->GetMaterialType();
-		const bool bClearTextures = VirtualTexture->GetClearTextures();
-		FSceneInterface* SceneInterface = InComponent->GetScene();
-		const EShadingPath ShadingPath = SceneInterface ? SceneInterface->GetShadingPath() : EShadingPath::Deferred;
-
-		// The producer object created here will be passed into the virtual texture system which will take ownership.
-		IVirtualTexture* Producer = new FRuntimeVirtualTextureProducer(ProducerDesc, ProducerId, MaterialType, bClearTextures, SceneInterface, Transform, Bounds);
-		
-		// Create a producer for the streaming low mips. 
-		// This is bound with the main producer so that one allocated VT can use both runtime or streaming producers dependent on mip level.
-		if (InComponent->IsStreamingLowMips(ShadingPath))
+		if (InComponent->IsEnabledInScene())
 		{
-			if (CVarVTStreamingMips.GetValueOnAnyThread() == 0)
+			// We store a ProducerId here so that we will be able to find our SceneIndex from the Producer during rendering.
+			// We will need the SceneIndex to determine which primitives should render to this Producer.
+			ProducerId = ProducerIdGenerator++;
+
+			URuntimeVirtualTexture::FInitSettings InitSettings;
+			InitSettings.TileCountBias = InComponent->IsScalable() ? VirtualTextureScalability::GetRuntimeVirtualTextureSizeBias(InComponent->GetScalabilityGroup()) : 0;
+
+			VirtualTexture = InComponent->GetVirtualTexture();
+			Transform = InComponent->GetComponentTransform();
+			const FBox Bounds = InComponent->Bounds.GetBox();
+
+			// The producer description is calculated using the transform to determine the aspect ratio
+			FVTProducerDescription ProducerDesc;
+			VirtualTexture->GetProducerDescription(ProducerDesc, InitSettings, Transform);
+			VirtualTextureSize = FIntPoint(ProducerDesc.BlockWidthInTiles * ProducerDesc.TileSize, ProducerDesc.BlockHeightInTiles * ProducerDesc.TileSize);
+			// We only need to dirty flush up to the producer description MaxLevel which accounts for the RemoveLowMips
+			MaxDirtyLevel = ProducerDesc.MaxLevel;
+
+			const ERuntimeVirtualTextureMaterialType MaterialType = VirtualTexture->GetMaterialType();
+			const bool bClearTextures = VirtualTexture->GetClearTextures();
+			FSceneInterface* SceneInterface = InComponent->GetScene();
+			const EShadingPath ShadingPath = SceneInterface ? SceneInterface->GetShadingPath() : EShadingPath::Deferred;
+
+			// The producer object created here will be passed into the virtual texture system which will take ownership.
+			IVirtualTexture* Producer = new FRuntimeVirtualTextureProducer(ProducerDesc, ProducerId, MaterialType, bClearTextures, SceneInterface, Transform, Bounds);
+		
+			// Create a producer for the streaming low mips. 
+			// This is bound with the main producer so that one allocated VT can use both runtime or streaming producers dependent on mip level.
+			if (InComponent->IsStreamingLowMips(ShadingPath))
 			{
-#if !UE_BUILD_SHIPPING
-				// Notify that streaming texture is turned off.
-				OnScreenWarningDelegateHandle = FRendererOnScreenNotification::Get().AddLambda([](FCoreDelegates::FSeverityMessageMap& OutMessages)
+				if (CVarVTStreamingMips.GetValueOnAnyThread() == 0)
 				{
-					OutMessages.Add(
-						FCoreDelegates::EOnScreenMessageSeverity::Warning,
-						LOCTEXT("SVTDisabled", "Runtime Virtual Texture streaming mips disabled."));
-				});
-#endif
-			}
-			else if (InComponent->IsStreamingTextureInvalid(ShadingPath))
-			{
 #if !UE_BUILD_SHIPPING
-				// Notify that streaming texture is invalid since this can cause performance regression.
-				const FString Name = InComponent->GetPathName();
-				OnScreenWarningDelegateHandle = FRendererOnScreenNotification::Get().AddLambda([Name](FCoreDelegates::FSeverityMessageMap& OutMessages)
+					// Notify that streaming texture is turned off.
+					OnScreenWarningDelegateHandle = FRendererOnScreenNotification::Get().AddLambda([](FCoreDelegates::FSeverityMessageMap& OutMessages)
+					{
+						OutMessages.Add(
+							FCoreDelegates::EOnScreenMessageSeverity::Warning,
+							LOCTEXT("SVTDisabled", "Runtime Virtual Texture streaming mips disabled."));
+					});
+#endif
+				}
+				else if (InComponent->IsStreamingTextureInvalid(ShadingPath))
 				{
-					OutMessages.Add(
-						FCoreDelegates::EOnScreenMessageSeverity::Warning,
-						FText::Format(LOCTEXT("SVTInvalid", "Runtime Virtual Texture '{0}' streaming mips needs to be rebuilt."), FText::FromString(Name)));
-				});
+#if !UE_BUILD_SHIPPING
+					// Notify that streaming texture is invalid since this can cause performance regression.
+					const FString Name = InComponent->GetPathName();
+					OnScreenWarningDelegateHandle = FRendererOnScreenNotification::Get().AddLambda([Name](FCoreDelegates::FSeverityMessageMap& OutMessages)
+					{
+						OutMessages.Add(
+							FCoreDelegates::EOnScreenMessageSeverity::Warning,
+							FText::Format(LOCTEXT("SVTInvalid", "Runtime Virtual Texture '{0}' streaming mips needs to be rebuilt."), FText::FromString(Name)));
+					});
 #endif
-			}
-			else
-			{
-				UVirtualTexture2D* StreamingTexture = InComponent->GetStreamingTexture()->GetVirtualTexture(ShadingPath);
+				}
+				else
+				{
+					UVirtualTexture2D* StreamingTexture = InComponent->GetStreamingTexture()->GetVirtualTexture(ShadingPath);
 
-				FVTProducerDescription StreamingProducerDesc;
-				IVirtualTexture* StreamingProducer = RuntimeVirtualTexture::CreateStreamingTextureProducer(StreamingTexture, ProducerDesc, StreamingProducerDesc);
+					FVTProducerDescription StreamingProducerDesc;
+					IVirtualTexture* StreamingProducer = RuntimeVirtualTexture::CreateStreamingTextureProducer(StreamingTexture, ProducerDesc, StreamingProducerDesc);
 
-				ensure(ProducerDesc.MaxLevel >= StreamingProducerDesc.MaxLevel);
-				const int32 TransitionLevel = ProducerDesc.MaxLevel - StreamingProducerDesc.MaxLevel;
+					ensure(ProducerDesc.MaxLevel >= StreamingProducerDesc.MaxLevel);
+					const int32 TransitionLevel = ProducerDesc.MaxLevel - StreamingProducerDesc.MaxLevel;
 
-				Producer = RuntimeVirtualTexture::BindStreamingTextureProducer(Producer, StreamingProducer, TransitionLevel);
+					Producer = RuntimeVirtualTexture::BindStreamingTextureProducer(Producer, StreamingProducer, TransitionLevel);
 
-				// Any dirty flushes don't need to flush the streaming mips (they only change with a build step).
-				MaxDirtyLevel = TransitionLevel - 1;
+					// Any dirty flushes don't need to flush the streaming mips (they only change with a build step).
+					MaxDirtyLevel = TransitionLevel - 1;
+				}
 			}
-		}
 
-		// The Initialize() call will allocate the virtual texture by spawning work on the render thread.
-		VirtualTexture->Initialize(Producer, ProducerDesc, Transform, Bounds);
+			// The Initialize() call will allocate the virtual texture by spawning work on the render thread.
+			VirtualTexture->Initialize(Producer, ProducerDesc, Transform, Bounds);
 
-		// Store the ProducerHandle and SpaceID immediately after virtual texture is initialized.
-		ENQUEUE_RENDER_COMMAND(GetProducerHandle)(
-			[this, VirtualTexturePtr = VirtualTexture](FRHICommandList& RHICmdList)
-			{
-				ProducerHandle = VirtualTexturePtr->GetProducerHandle();
-				SpaceID = VirtualTexturePtr->GetAllocatedVirtualTexture()->GetSpaceID();
-			});
+			// Store the ProducerHandle and SpaceID immediately after virtual texture is initialized.
+			ENQUEUE_RENDER_COMMAND(GetProducerHandle)(
+				[this, VirtualTexturePtr = VirtualTexture](FRHICommandList& RHICmdList)
+				{
+					ProducerHandle = VirtualTexturePtr->GetProducerHandle();
+					SpaceID = VirtualTexturePtr->GetAllocatedVirtualTexture()->GetSpaceID();
+				});
+		}
+		else
+		{
+			// When not enabled, ensure that the RVT asset has no allocated VT.
+			// In PIE this handles removing the RVT from the editor scene.
+			InComponent->GetVirtualTexture()->Release();
+		}
 	}
 }
 
