Renames TSR's ComputeMoireLuma pass to MeasureFlickeringLuma to better match 5.4 documentation

diff --git a/Engine/Shaders/Private/TemporalSuperResolution/TSRComputeMoireLuma.usf b/Engine/Shaders/Private/TemporalSuperResolution/TSRComputeMoireLuma.usf
deleted file mode 100644
diff --git a/Engine/Shaders/Private/TemporalSuperResolution/TSRMeasureFlickeringLuma.usf b/Engine/Shaders/Private/TemporalSuperResolution/TSRMeasureFlickeringLuma.usf
new file mode 100644
--- /dev/null
+++ b/Engine/Shaders/Private/TemporalSuperResolution/TSRMeasureFlickeringLuma.usf
@@ -0,0 +1,69 @@
+// Copyright Epic Games, Inc. All Rights Reserved.
+
+#include "TSRCommon.ush"
+
+
+//------------------------------------------------------- CONFIG
+
+#define TILE_SIZE 8
+
+#define TILE_ROWS 2
+
+
+//------------------------------------------------------- PARAMETERS
+
+float PerceptionAdd;
+
+Texture2D<tsr_half3> SceneColorTexture;
+RWTexture2D<tsr_half> FlickeringLumaOutput;
+
+
+//------------------------------------------------------- ENTRY POINT
+
+[numthreads(TILE_SIZE * TILE_SIZE, 1, 1)]
+void MainCS(
+	uint2 GroupId : SV_GroupID,
+	uint GroupThreadIndex : SV_GroupIndex)
+{
+	uint2 GroupThreadId = uint2(GroupThreadIndex % TILE_SIZE, GroupThreadIndex / TILE_SIZE);
+	uint2 OutputPixelOrigin = (InputInfo_ViewportMin + GroupId * uint(TILE_SIZE * TILE_ROWS)) + GroupThreadId;
+
+	tsr_half3 ColorArray[TILE_ROWS * TILE_ROWS];
+
+	ISOLATE
+	{
+		UNROLL_N(TILE_ROWS)
+		for (uint y = 0; y < TILE_ROWS; y++)
+		{
+			UNROLL_N(TILE_ROWS)
+			for (uint x = 0; x < TILE_ROWS; x++)
+			{
+				const uint2 PixelOffset = uint2(x, y) * TILE_SIZE;
+				uint2 InputPixelPos = OutputPixelOrigin + PixelOffset;
+			
+				ColorArray[x + y * TILE_ROWS] = SceneColorTexture[InputPixelPos];
+			}
+		}
+	}
+
+	ISOLATE
+	{
+		UNROLL_N(TILE_ROWS)
+		for (uint y = 0; y < TILE_ROWS; y++)
+		{
+			UNROLL_N(TILE_ROWS)
+			for (uint x = 0; x < TILE_ROWS; x++)
+			{
+				const uint2 PixelOffset = uint2(x, y) * TILE_SIZE;
+				uint2 InputPixelPos = OutputPixelOrigin + PixelOffset;
+			
+				tsr_half3 Color = ColorArray[x + y * TILE_ROWS];
+				Color *= rcp(tsr_half(PerceptionAdd) + Color);
+
+				InputPixelPos.x = select(all(InputPixelPos < InputInfo_ViewportMax), InputPixelPos.x, uint(~0));
+				
+				FlickeringLumaOutput[InputPixelPos] = dot(Color, kMoireLumaWeights);
+			}
+		}
+	}
+}
diff --git a/Engine/Source/Runtime/Renderer/Private/DeferredShadingRenderer.cpp b/Engine/Source/Runtime/Renderer/Private/DeferredShadingRenderer.cpp
--- a/Engine/Source/Runtime/Renderer/Private/DeferredShadingRenderer.cpp
+++ b/Engine/Source/Runtime/Renderer/Private/DeferredShadingRenderer.cpp
@@ -4085,7 +4085,7 @@ void FDeferredShadingSceneRenderer::Render(FRDGBuilder& GraphBuilder)
 			ReconstructVolumetricRenderTarget(GraphBuilder, Views, SceneTextures.Depth.Resolve, HalfResolutionDepthCheckerboardMinMaxTexture, bAsyncComputeVolumetricCloud);
 		}
 
-		TArray<FScreenPassTexture, TInlineAllocator<4>> TSRMoireInputTextures;
+		TArray<FScreenPassTexture, TInlineAllocator<4>> TSRFlickeringInputTextures;
 		if (!bHasRayTracedOverlay)
 		{
 			// Extract TSR's moire heuristic luminance before rendering translucency into the scene color.
@@ -4094,12 +4094,12 @@ void FDeferredShadingSceneRenderer::Render(FRDGBuilder& GraphBuilder)
 				FViewInfo& View = Views[ViewIndex];
 				if (NeedTSRMoireLuma(View))
 				{
-					if (TSRMoireInputTextures.Num() == 0)
+					if (TSRFlickeringInputTextures.Num() == 0)
 					{
-						TSRMoireInputTextures.SetNum(Views.Num());
+						TSRFlickeringInputTextures.SetNum(Views.Num());
 					}
 
-					TSRMoireInputTextures[ViewIndex] = AddTSRComputeMoireLuma(GraphBuilder, View.ShaderMap, FScreenPassTexture(SceneTextures.Color.Target, View.ViewRect));
+					TSRFlickeringInputTextures[ViewIndex] = AddTSRMeasureFlickeringLuma(GraphBuilder, View.ShaderMap, FScreenPassTexture(SceneTextures.Color.Target, View.ViewRect));
 				}
 			}
 		}
@@ -4562,10 +4562,10 @@ void FDeferredShadingSceneRenderer::Render(FRDGBuilder& GraphBuilder)
 						const bool bAnyLumenActive = ViewPipelineState.DiffuseIndirectMethod == EDiffuseIndirectMethod::Lumen || ViewPipelineState.ReflectionsMethod == EReflectionsMethod::Lumen;
 						const bool bLumenGIEnabled = ViewPipelineState.DiffuseIndirectMethod == EDiffuseIndirectMethod::Lumen;
 
-						FScreenPassTexture TSRMoireInput;
-						if (ViewIndex < TSRMoireInputTextures.Num())
+						FScreenPassTexture TSRFlickeringInput;
+						if (ViewIndex < TSRFlickeringInputTextures.Num())
 						{
-							TSRMoireInput = TSRMoireInputTextures[ViewIndex];
+							TSRFlickeringInput = TSRFlickeringInputTextures[ViewIndex];
 						}
 
 						AddPostProcessingPasses(
@@ -4581,7 +4581,7 @@ void FDeferredShadingSceneRenderer::Render(FRDGBuilder& GraphBuilder)
 							&VirtualShadowMapArray,
 							LumenFrameTemporaries,
 							SceneWithoutWaterTextures,
-							TSRMoireInput);
+							TSRFlickeringInput);
 					}
 				}
 			}
diff --git a/Engine/Source/Runtime/Renderer/Private/PostProcess/PostProcessing.cpp b/Engine/Source/Runtime/Renderer/Private/PostProcess/PostProcessing.cpp
--- a/Engine/Source/Runtime/Renderer/Private/PostProcess/PostProcessing.cpp
+++ b/Engine/Source/Runtime/Renderer/Private/PostProcess/PostProcessing.cpp
@@ -296,7 +296,7 @@ void AddPostProcessingPasses(
 	FVirtualShadowMapArray* VirtualShadowMapArray, 
 	FLumenSceneFrameTemporaries& LumenFrameTemporaries,
 	const FSceneWithoutWaterTextures& SceneWithoutWaterTextures,
-	FScreenPassTexture TSRMoireInput)
+	FScreenPassTexture TSRFlickeringInput)
 {
 	RDG_CSV_STAT_EXCLUSIVE_SCOPE(GraphBuilder, RenderPostProcessing);
 	QUICK_SCOPE_CYCLE_COUNTER(STAT_PostProcessing_Process);
@@ -770,7 +770,7 @@ void AddPostProcessingPasses(
 			UpscalerPassInputs.bAllowFullResSlice = PassSequence.IsEnabled(EPass::MotionBlur) || PassSequence.IsEnabled(EPass::Tonemap);
 			UpscalerPassInputs.DownsampleOverrideFormat = DownsampleOverrideFormat;
 			UpscalerPassInputs.PostDOFTranslucencyResources = PostDOFTranslucencyResources;
-			UpscalerPassInputs.MoireInputTexture = TSRMoireInput;
+			UpscalerPassInputs.FlickeringInputTexture = TSRFlickeringInput;
 			check(UpscalerPassInputs.SceneColor.ViewRect == View.ViewRect);
 
 			FDefaultTemporalUpscaler::FOutputs Outputs;
diff --git a/Engine/Source/Runtime/Renderer/Private/PostProcess/PostProcessing.h b/Engine/Source/Runtime/Renderer/Private/PostProcess/PostProcessing.h
--- a/Engine/Source/Runtime/Renderer/Private/PostProcess/PostProcessing.h
+++ b/Engine/Source/Runtime/Renderer/Private/PostProcess/PostProcessing.h
@@ -63,7 +63,7 @@ void AddPostProcessingPasses(
 	FVirtualShadowMapArray* VirtualShadowMapArray,
 	struct FLumenSceneFrameTemporaries& LumenFrameTemporaries,
 	const FSceneWithoutWaterTextures& SceneWithoutWaterTextures,
-	FScreenPassTexture TSRMoireInput);
+	FScreenPassTexture TSRFlickeringInput);
 
 void AddDebugViewPostProcessingPasses(FRDGBuilder& GraphBuilder, const FViewInfo& View, FSceneUniformBuffer &SceneUniformBuffer, const FPostProcessingInputs& Inputs, const Nanite::FRasterResults* NaniteRasterResults);
 
diff --git a/Engine/Source/Runtime/Renderer/Private/PostProcess/TemporalAA.h b/Engine/Source/Runtime/Renderer/Private/PostProcess/TemporalAA.h
--- a/Engine/Source/Runtime/Renderer/Private/PostProcess/TemporalAA.h
+++ b/Engine/Source/Runtime/Renderer/Private/PostProcess/TemporalAA.h
@@ -179,8 +179,8 @@ bool NeedTSRMoireLuma(const FViewInfo& View);
 /** Returns whether TSR internal visualization is enabled on the view. */
 bool IsVisualizeTSREnabled(const FViewInfo& View);
 
-/** Measure luminance of the scene color for moire anti-flickering. */
-FScreenPassTexture AddTSRComputeMoireLuma(FRDGBuilder& GraphBuilder, FGlobalShaderMap* ShaderMap, FScreenPassTexture SceneColor);
+/** Measure luminance of the scene color for anti-flickering. */
+FScreenPassTexture AddTSRMeasureFlickeringLuma(FRDGBuilder& GraphBuilder, FGlobalShaderMap* ShaderMap, FScreenPassTexture SceneColor);
 
 
 EMainTAAPassConfig GetMainTAAPassConfig(const FViewInfo& View);
@@ -200,7 +200,7 @@ struct FDefaultTemporalUpscaler
 		FScreenPassTexture SceneDepth;
 		FScreenPassTexture SceneVelocity;
 		FTranslucencyPassResources PostDOFTranslucencyResources;
-		FScreenPassTexture MoireInputTexture;
+		FScreenPassTexture FlickeringInputTexture;
 	};
 
 	struct FOutputs
diff --git a/Engine/Source/Runtime/Renderer/Private/PostProcess/TemporalSuperResolution.cpp b/Engine/Source/Runtime/Renderer/Private/PostProcess/TemporalSuperResolution.cpp
--- a/Engine/Source/Runtime/Renderer/Private/PostProcess/TemporalSuperResolution.cpp
+++ b/Engine/Source/Runtime/Renderer/Private/PostProcess/TemporalSuperResolution.cpp
@@ -492,18 +492,18 @@ public:
 	}
 }; // class FTemporalSuperResolutionShader
 
-class FTSRComputeMoireLumaCS : public FTSRShader
+class FTSRMeasureFlickeringLumaCS : public FTSRShader
 {
-	DECLARE_GLOBAL_SHADER(FTSRComputeMoireLumaCS);
-	SHADER_USE_PARAMETER_STRUCT(FTSRComputeMoireLumaCS, FTSRShader);
+	DECLARE_GLOBAL_SHADER(FTSRMeasureFlickeringLumaCS);
+	SHADER_USE_PARAMETER_STRUCT(FTSRMeasureFlickeringLumaCS, FTSRShader);
 
 	BEGIN_SHADER_PARAMETER_STRUCT(FParameters, )
 		SHADER_PARAMETER_STRUCT(FScreenPassTextureViewportParameters, InputInfo)
 		SHADER_PARAMETER(float, PerceptionAdd)
 		SHADER_PARAMETER_RDG_TEXTURE(Texture2D, SceneColorTexture)
-		SHADER_PARAMETER_RDG_TEXTURE_UAV(RWTexture2D, MoireLumaOutput)
+		SHADER_PARAMETER_RDG_TEXTURE_UAV(RWTexture2D, FlickeringLumaOutput)
 	END_SHADER_PARAMETER_STRUCT()
-}; // class FTSRComputeMoireLumaCS
+}; // class FTSRMeasureFlickeringLumaCS
 
 class FTSRClearPrevTexturesCS : public FTSRShader
 {
@@ -1033,7 +1033,7 @@ class FTSRVisualizeCS : public FTSRShader
 	END_SHADER_PARAMETER_STRUCT()
 }; // class FTSRVisualizeCS
 
-IMPLEMENT_GLOBAL_SHADER(FTSRComputeMoireLumaCS,      "/Engine/Private/TemporalSuperResolution/TSRComputeMoireLuma.usf",      "MainCS", SF_Compute);
+IMPLEMENT_GLOBAL_SHADER(FTSRMeasureFlickeringLumaCS, "/Engine/Private/TemporalSuperResolution/TSRMeasureFlickeringLuma.usf", "MainCS", SF_Compute);
 IMPLEMENT_GLOBAL_SHADER(FTSRClearPrevTexturesCS,     "/Engine/Private/TemporalSuperResolution/TSRClearPrevTextures.usf",     "MainCS", SF_Compute);
 IMPLEMENT_GLOBAL_SHADER(FTSRForwardScatterDepthCS,   "/Engine/Private/TemporalSuperResolution/TSRForwardScatterDepth.usf",   "MainCS", SF_Compute);
 IMPLEMENT_GLOBAL_SHADER(FTSRDilateVelocityCS,        "/Engine/Private/TemporalSuperResolution/TSRDilateVelocity.usf",        "MainCS", SF_Compute);
@@ -1211,12 +1211,12 @@ bool IsVisualizeTSREnabled(const FViewInfo& View)
 }
 #endif
 
-FScreenPassTexture AddTSRComputeMoireLuma(FRDGBuilder& GraphBuilder, FGlobalShaderMap* ShaderMap, FScreenPassTexture SceneColor)
+FScreenPassTexture AddTSRMeasureFlickeringLuma(FRDGBuilder& GraphBuilder, FGlobalShaderMap* ShaderMap, FScreenPassTexture SceneColor)
 {
 	check(SceneColor.Texture)
 	RDG_GPU_STAT_SCOPE(GraphBuilder, TemporalSuperResolution);
 
-	FScreenPassTexture MoireLuma;
+	FScreenPassTexture FlickeringLuma;
 	{
 		FRDGTextureDesc Desc = FRDGTextureDesc::Create2D(
 			SceneColor.Texture->Desc.Extent,
@@ -1224,26 +1224,26 @@ FScreenPassTexture AddTSRComputeMoireLuma(FRDGBuilder& GraphBuilder, FGlobalShad
 			FClearValueBinding::None,
 			/* InFlags = */ TexCreate_ShaderResource | TexCreate_UAV);
 
-		MoireLuma.Texture = GraphBuilder.CreateTexture(Desc, TEXT("TSR.Moire.Luma"));
-		MoireLuma.ViewRect = SceneColor.ViewRect;
+		FlickeringLuma.Texture = GraphBuilder.CreateTexture(Desc, TEXT("TSR.Flickering.Luminance"));
+		FlickeringLuma.ViewRect = SceneColor.ViewRect;
 	}
 
-	FTSRComputeMoireLumaCS::FParameters* PassParameters = GraphBuilder.AllocParameters<FTSRComputeMoireLumaCS::FParameters>();
+	FTSRMeasureFlickeringLumaCS::FParameters* PassParameters = GraphBuilder.AllocParameters<FTSRMeasureFlickeringLumaCS::FParameters>();
 	PassParameters->InputInfo = GetScreenPassTextureViewportParameters(FScreenPassTextureViewport(
 		SceneColor.Texture->Desc.Extent, SceneColor.ViewRect));
 	PassParameters->PerceptionAdd = FMath::Pow(0.5f, CVarTSRShadingExposureOffset.GetValueOnRenderThread());
 	PassParameters->SceneColorTexture = SceneColor.Texture;
-	PassParameters->MoireLumaOutput = GraphBuilder.CreateUAV(MoireLuma.Texture);
+	PassParameters->FlickeringLumaOutput = GraphBuilder.CreateUAV(FlickeringLuma.Texture);
 
-	TShaderMapRef<FTSRComputeMoireLumaCS> ComputeShader(ShaderMap);
+	TShaderMapRef<FTSRMeasureFlickeringLumaCS> ComputeShader(ShaderMap);
 	FComputeShaderUtils::AddPass(
 		GraphBuilder,
-		RDG_EVENT_NAME("TSR ComputeMoireLuma %dx%d", SceneColor.ViewRect.Width(), SceneColor.ViewRect.Height()),
+		RDG_EVENT_NAME("TSR MeasureFlickeringLuma %dx%d", SceneColor.ViewRect.Width(), SceneColor.ViewRect.Height()),
 		ComputeShader,
 		PassParameters,
-		FComputeShaderUtils::GetGroupCount(MoireLuma.ViewRect.Size(), 8 * 2));
+		FComputeShaderUtils::GetGroupCount(FlickeringLuma.ViewRect.Size(), 8 * 2));
 
-	return MoireLuma;
+	return FlickeringLuma;
 }
 
 FDefaultTemporalUpscaler::FOutputs AddTemporalSuperResolutionPasses(
@@ -2146,10 +2146,10 @@ FDefaultTemporalUpscaler::FOutputs AddTemporalSuperResolutionPasses(
 		PassParameters->PerceptionAdd = FMath::Pow(0.5f, CVarTSRShadingExposureOffset.GetValueOnRenderThread());
 
 		PassParameters->InputTexture = PassInputs.SceneColor.Texture;
-		if (PassInputs.MoireInputTexture.IsValid())
+		if (PassInputs.FlickeringInputTexture.IsValid())
 		{
-			ensure(InputRect == PassInputs.MoireInputTexture.ViewRect);
-			PassParameters->InputMoireLumaTexture = PassInputs.MoireInputTexture.Texture;
+			ensure(InputRect == PassInputs.FlickeringInputTexture.ViewRect);
+			PassParameters->InputMoireLumaTexture = PassInputs.FlickeringInputTexture.Texture;
 		}
 		else
 		{
diff --git a/Engine/Source/Runtime/Renderer/Private/PostProcess/VisualizeTemporalUpscaler.cpp b/Engine/Source/Runtime/Renderer/Private/PostProcess/VisualizeTemporalUpscaler.cpp
--- a/Engine/Source/Runtime/Renderer/Private/PostProcess/VisualizeTemporalUpscaler.cpp
+++ b/Engine/Source/Runtime/Renderer/Private/PostProcess/VisualizeTemporalUpscaler.cpp
@@ -158,11 +158,11 @@ FScreenPassTexture AddVisualizeTemporalUpscalerPass(FRDGBuilder& GraphBuilder, c
 		if (Inputs.TAAConfig == EMainTAAPassConfig::TSR)
 		{
 			FVisualizeBufferTile& Tile = Tiles[4 * 1 + 3];
-			if (Inputs.Inputs.MoireInputTexture.IsValid())
+			if (Inputs.Inputs.FlickeringInputTexture.IsValid())
 			{
-				Tile.Input = Inputs.Inputs.MoireInputTexture;
+				Tile.Input = Inputs.Inputs.FlickeringInputTexture;
 				Tile.Input.ViewRect = CropViewRectToCenter(View.ViewRect);
-				Tile.Label = VisualizeTextureLabel(Inputs.Inputs.MoireInputTexture.Texture);
+				Tile.Label = VisualizeTextureLabel(Inputs.Inputs.FlickeringInputTexture.Texture);
 			}
 			else
 			{
