Substrate - fix for crash with hair strand shading when using clustered lighting.

diff --git a/Engine/Source/Runtime/Renderer/Private/ClusteredDeferredShadingPass.cpp b/Engine/Source/Runtime/Renderer/Private/ClusteredDeferredShadingPass.cpp
--- a/Engine/Source/Runtime/Renderer/Private/ClusteredDeferredShadingPass.cpp
+++ b/Engine/Source/Runtime/Renderer/Private/ClusteredDeferredShadingPass.cpp
@@ -158,7 +158,7 @@ static void InternalAddClusteredDeferredShadingPass(
 	check(SortedLightsSet.ClusteredSupportedEnd > 0);
 	const FIntPoint SceneTextureExtent = SceneTextures.Config.Extent;
 	const bool bHairStrands = InputType == EClusterPassInputType::HairStrands;
-	const bool bSubstrate = Substrate::IsSubstrateEnabled();
+	const bool bSubstrate = Substrate::IsSubstrateEnabled() && !bHairStrands;
 	
 	FClusteredShadingPS::FParameters *PassParameters = GraphBuilder.AllocParameters<FClusteredShadingPS::FParameters>();
 	PassParameters->View = View.ViewUniformBuffer;
