Removed redundant FScene parameter passed in many places to GPU-Scene functions, fixed typo in OnPostLigtSceneInfoUpdate, & fixed the primitive ID in the debug rendering.

diff --git a/Engine/Source/Runtime/Renderer/Private/DeferredShadingRenderer.cpp b/Engine/Source/Runtime/Renderer/Private/DeferredShadingRenderer.cpp
--- a/Engine/Source/Runtime/Renderer/Private/DeferredShadingRenderer.cpp
+++ b/Engine/Source/Runtime/Renderer/Private/DeferredShadingRenderer.cpp
@@ -2676,7 +2676,7 @@ void FDeferredShadingSceneRenderer::Render(FRDGBuilder& GraphBuilder)
 	TUniquePtr<FVirtualTextureUpdater> VirtualTextureUpdater;
 	FLumenSceneFrameTemporaries LumenFrameTemporaries;
 
-	FGPUSceneScopeBeginEndHelper GPUSceneScopeBeginEndHelper(Scene->GPUScene, GPUSceneDynamicContext, Scene);
+	FGPUSceneScopeBeginEndHelper GPUSceneScopeBeginEndHelper(Scene->GPUScene, GPUSceneDynamicContext);
 
 	const bool bUseVirtualTexturing = UseVirtualTexturing(ShaderPlatform);
 
@@ -3086,9 +3086,9 @@ void FDeferredShadingSceneRenderer::Render(FRDGBuilder& GraphBuilder)
 			FViewInfo& View = *AllViews[ViewIndex];
 			RDG_GPU_MASK_SCOPE(GraphBuilder, View.GPUMask);
 
-			Scene->GPUScene.UploadDynamicPrimitiveShaderDataForView(GraphBuilder, *Scene, View);
+			Scene->GPUScene.UploadDynamicPrimitiveShaderDataForView(GraphBuilder, View);
 
-			Scene->GPUScene.DebugRender(GraphBuilder, *Scene, GetSceneUniforms(), View);
+			Scene->GPUScene.DebugRender(GraphBuilder, GetSceneUniforms(), View);
 		}
 
 		InstanceCullingManager.BeginDeferredCulling(GraphBuilder, Scene->GPUScene);
diff --git a/Engine/Source/Runtime/Renderer/Private/GPUScene.cpp b/Engine/Source/Runtime/Renderer/Private/GPUScene.cpp
--- a/Engine/Source/Runtime/Renderer/Private/GPUScene.cpp
+++ b/Engine/Source/Runtime/Renderer/Private/GPUScene.cpp
@@ -650,21 +650,14 @@ FGPUScene::~FGPUScene()
 #endif
 }
 
-void FGPUScene::BeginRender(const FScene* InScene, FGPUSceneDynamicContext &GPUSceneDynamicContext)
+void FGPUScene::BeginRender(FGPUSceneDynamicContext &GPUSceneDynamicContext)
 {
 	ensure(!bInBeginEndBlock);
 	ensure(CurrentDynamicContext == nullptr);
-	// TODO: IS the scene being null a case that exists???
-	if (InScene != nullptr)
-	{
-		check(InScene == &Scene);
-		ensure(bIsEnabled == UseGPUScene(GMaxRHIShaderPlatform, Scene.GetFeatureLevel()));
-		NumScenePrimitives = Scene.Primitives.Num();
-	}
-	else
-	{
-		NumScenePrimitives = 0;
-	}
+
+	ensure(bIsEnabled == UseGPUScene(GMaxRHIShaderPlatform, Scene.GetFeatureLevel()));
+	NumScenePrimitives = Scene.Primitives.Num();
+
 	CurrentDynamicContext = &GPUSceneDynamicContext;
 	DynamicPrimitivesOffset = Scene.GetMaxPersistentPrimitiveIndex();
 	bInBeginEndBlock = true;
@@ -681,10 +674,8 @@ void FGPUScene::EndRender()
 	ShaderParameters = {};
 }
 
-void FGPUScene::UpdateGPULights(FRDGBuilder& GraphBuilder, FScene& InScene, const UE::Tasks::FTask& PrerequisiteTask)
+void FGPUScene::UpdateGPULights(FRDGBuilder& GraphBuilder, const UE::Tasks::FTask& PrerequisiteTask)
 {
-	// TODO: remove scene as parameter
-	check(&InScene == &Scene);
 	FRDGUploadData<FLightSceneData> LightData(GraphBuilder, FMath::Max(1, Scene.Lights.GetMaxIndex()));
 
 	GraphBuilder.AddSetupTask([this, LightData]
@@ -745,11 +736,8 @@ void FGPUScene::InitLightData(const FLightSceneInfoCompact& LightInfoCompact, bo
 	DataOut.LightTypeAndShadowMapChannelMaskPacked = LightInfo.PackLightTypeAndShadowMapChannelMask(bAllowStaticLighting);
 }
 
-void FGPUScene::UpdateInternal(FRDGBuilder& GraphBuilder, FSceneUniformBuffer& SceneUB, FScene& InScene, FRDGExternalAccessQueue& ExternalAccessQueue, IVisibilityTaskData* VisibilityTaskData)
+void FGPUScene::UpdateInternal(FRDGBuilder& GraphBuilder, FSceneUniformBuffer& SceneUB, FRDGExternalAccessQueue& ExternalAccessQueue, IVisibilityTaskData* VisibilityTaskData)
 {
-	// TODO: remove scene as parameter
-	check(&InScene == &Scene);
-
 	LLM_SCOPE_BYTAG(GPUScene);
 
 	check(bInBeginEndBlock);
@@ -825,7 +813,7 @@ void FGPUScene::UpdateInternal(FRDGBuilder& GraphBuilder, FSceneUniformBuffer& S
 	check(!BufferState.IsValid());
 
 	FUploadDataSourceAdapterScenePrimitives& Adapter = *GraphBuilder.AllocObject<FUploadDataSourceAdapterScenePrimitives>(Scene, SceneFrameNumber, MoveTemp(PrimitivesToUpdate), MoveTemp(PrimitiveDirtyState));
-	UpdateBufferState(GraphBuilder, SceneUB, Scene, Adapter, true);
+	UpdateBufferState(GraphBuilder, SceneUB, Adapter, true);
 
 	// Run a pass that clears (Sets ID to invalid) any instances that need it
 	AddClearInstancesPass(GraphBuilder, Scene.InstanceCullingOcclusionQueryRenderer);
@@ -847,17 +835,15 @@ void FGPUScene::UpdateInternal(FRDGBuilder& GraphBuilder, FSceneUniformBuffer& S
 		QUICK_SCOPE_CYCLE_COUNTER(STAT_UpdateGPUScene);
 		SCOPE_CYCLE_COUNTER(STAT_UpdateGPUSceneTime);
 
-		UploadGeneral<FUploadDataSourceAdapterScenePrimitives>(GraphBuilder, Scene, &ExternalAccessQueue, Adapter, PrerequisiteAsyncTask);
+		UploadGeneral<FUploadDataSourceAdapterScenePrimitives>(GraphBuilder, &ExternalAccessQueue, Adapter, PrerequisiteAsyncTask);
 
-		UpdateGPULights(GraphBuilder, Scene, PrerequisiteAsyncTask);
+		UpdateGPULights(GraphBuilder, PrerequisiteAsyncTask);
 	}
 }
 
 template<typename FUploadDataSourceAdapter>
-void FGPUScene::UpdateBufferState(FRDGBuilder& GraphBuilder, FSceneUniformBuffer& SceneUB, FScene& InScene, const FUploadDataSourceAdapter& UploadDataSourceAdapter, bool bIsMainUpdate)
+void FGPUScene::UpdateBufferState(FRDGBuilder& GraphBuilder, FSceneUniformBuffer& SceneUB, const FUploadDataSourceAdapter& UploadDataSourceAdapter, bool bIsMainUpdate)
 {
-	// TODO: remove scene as parameter
-	check(&InScene == &Scene);
 	LLM_SCOPE_BYTAG(GPUScene);
 
 	check(bInBeginEndBlock);
@@ -1009,11 +995,8 @@ struct FInstanceBatcher
 };
 
 template<typename FUploadDataSourceAdapter>
-void FGPUScene::UploadGeneral(FRDGBuilder& GraphBuilder, FScene& InScene, FRDGExternalAccessQueue* ExternalAccessQueue, const FUploadDataSourceAdapter& UploadDataSourceAdapter, const UE::Tasks::FTask& PrerequisiteTask)
+void FGPUScene::UploadGeneral(FRDGBuilder& GraphBuilder, FRDGExternalAccessQueue* ExternalAccessQueue, const FUploadDataSourceAdapter& UploadDataSourceAdapter, const UE::Tasks::FTask& PrerequisiteTask)
 {
-	// TODO: remove scene as parameter
-	check(&InScene == &Scene);
-
 	LLM_SCOPE_BYTAG(GPUScene);
 
 	ensure(bIsEnabled == UseGPUScene(GMaxRHIShaderPlatform, Scene.GetFeatureLevel()));
@@ -1098,7 +1081,7 @@ void FGPUScene::UploadGeneral(FRDGBuilder& GraphBuilder, FScene& InScene, FRDGEx
 		TaskContext.bUseNaniteMaterialUploaders = true;
 	}
 
-	GraphBuilder.AddCommandListSetupTask([&TaskContext, &UploadDataSourceAdapter, &InScene, bNaniteEnabled, bExecuteInParallel, FeatureLevel = FeatureLevel](FRHICommandListBase& RHICmdList)
+	GraphBuilder.AddCommandListSetupTask([&TaskContext, &UploadDataSourceAdapter, bNaniteEnabled, bExecuteInParallel, FeatureLevel = FeatureLevel](FRHICommandListBase& RHICmdList)
 	{
 		SCOPED_NAMED_EVENT(UpdateGPUScene_Primitives, FColor::Green);
 
@@ -1581,11 +1564,8 @@ struct FUploadDataSourceAdapterDynamicPrimitives
 	TArray<uint32, SceneRenderingAllocator> PrimitivesIds;
 };
 
-void FGPUScene::UploadDynamicPrimitiveShaderDataForViewInternal(FRDGBuilder& GraphBuilder, FScene& InScene, FViewInfo& View, bool bIsShadowView)
+void FGPUScene::UploadDynamicPrimitiveShaderDataForViewInternal(FRDGBuilder& GraphBuilder, FViewInfo& View, bool bIsShadowView)
 {
-	// TODO: remove scene as parameter
-	check(&InScene == &Scene);
-
 	LLM_SCOPE_BYTAG(GPUScene);
 
 	RDG_EVENT_SCOPE(GraphBuilder, "GPUScene.UploadDynamicPrimitiveShaderDataForView");
@@ -1649,12 +1629,12 @@ void FGPUScene::UploadDynamicPrimitiveShaderDataForViewInternal(FRDGBuilder& Gra
 			Collector.UploadData->InstancePayloadDataOffset,
 			SceneFrameNumber);
 
-		UpdateBufferState(GraphBuilder, View.GetSceneUniforms(), Scene, UploadAdapter, false);
+		UpdateBufferState(GraphBuilder, View.GetSceneUniforms(), UploadAdapter, false);
 
 		// Run a pass that clears (Sets ID to invalid) any instances that need it.
 		AddClearInstancesPass(GraphBuilder, Scene.InstanceCullingOcclusionQueryRenderer);
 
-		UploadGeneral<FUploadDataSourceAdapterDynamicPrimitives>(GraphBuilder, Scene, nullptr, UploadAdapter, UE::Tasks::FTask{});
+		UploadGeneral<FUploadDataSourceAdapterDynamicPrimitives>(GraphBuilder, nullptr, UploadAdapter, UE::Tasks::FTask{});
 	}
 
 	FSceneUniformBuffer& SceneUniforms = View.GetSceneUniforms();
@@ -1775,31 +1755,25 @@ void FGPUScene::AddPrimitiveToUpdate(FPersistentPrimitiveIndex PersistentPrimiti
 }
 
 
-void FGPUScene::Update(FRDGBuilder& GraphBuilder, FSceneUniformBuffer& SceneUB, FScene& InScene, FRDGExternalAccessQueue& ExternalAccessQueue, IVisibilityTaskData* VisibilityTaskData)
+void FGPUScene::Update(FRDGBuilder& GraphBuilder, FSceneUniformBuffer& SceneUB, FRDGExternalAccessQueue& ExternalAccessQueue, IVisibilityTaskData* VisibilityTaskData)
 {
-	// TODO: remove scene as parameter
-	check(&InScene == &Scene);
-
 	if (bIsEnabled)
 	{
 		RDG_GPU_MASK_SCOPE(GraphBuilder, FRHIGPUMask::All());
 
 		ensure(bInBeginEndBlock);
 		
-		UpdateInternal(GraphBuilder, SceneUB, Scene, ExternalAccessQueue, VisibilityTaskData);
+		UpdateInternal(GraphBuilder, SceneUB, ExternalAccessQueue, VisibilityTaskData);
 	}
 }
 
-void FGPUScene::UploadDynamicPrimitiveShaderDataForView(FRDGBuilder& GraphBuilder, FScene& InScene, FViewInfo& View, bool bIsShadowView)
+void FGPUScene::UploadDynamicPrimitiveShaderDataForView(FRDGBuilder& GraphBuilder, FViewInfo& View, bool bIsShadowView)
 {
-	// TODO: remove scene as parameter
-	check(&InScene == &Scene);
-
 	if (bIsEnabled)
 	{
 		RDG_GPU_MASK_SCOPE(GraphBuilder, FRHIGPUMask::All());
 
-		UploadDynamicPrimitiveShaderDataForViewInternal(GraphBuilder, Scene, View, bIsShadowView);
+		UploadDynamicPrimitiveShaderDataForViewInternal(GraphBuilder, View, bIsShadowView);
 	}
 }
 
@@ -1934,11 +1908,8 @@ public:
 
 IMPLEMENT_GLOBAL_SHADER(FGPUSceneDebugRenderCS, "/Engine/Private/GPUSceneDebugRender.usf", "GPUSceneDebugRenderCS", SF_Compute);
 
-void FGPUScene::DebugRender(FRDGBuilder& GraphBuilder, FScene& InScene, FSceneUniformBuffer& SceneUniformBuffer, FViewInfo& View)
+void FGPUScene::DebugRender(FRDGBuilder& GraphBuilder, FSceneUniformBuffer& SceneUniformBuffer, FViewInfo& View)
 {
-	// TODO: remove scene as parameter
-	check(&InScene == &Scene);
-
 	int32 DebugMode = CVarGPUSceneDebugMode.GetValueOnRenderThread();
 	if (DebugMode > 0)
 	{
@@ -1964,17 +1935,20 @@ void FGPUScene::DebugRender(FRDGBuilder& GraphBuilder, FScene& InScene, FSceneUn
 			uint32 SelectedCount = 0;
 			TArray<uint32> SelectedPrimitiveFlags;
 			const int32 BitsPerWord = (sizeof(uint32) * 8U);
-			SelectedPrimitiveFlags.Init(0U, FMath::DivideAndRoundUp(Scene.Primitives.Num(), BitsPerWord));
-			for (int32 PrimitiveID = 0; PrimitiveID < Scene.PrimitiveSceneProxies.Num(); ++PrimitiveID)
+			SelectedPrimitiveFlags.Init(0U, FMath::DivideAndRoundUp(Scene.GetMaxPersistentPrimitiveIndex(), BitsPerWord));
+			for (int32 PackedIndex = 0; PackedIndex < Scene.PrimitiveSceneProxies.Num(); ++PackedIndex)
 			{
-				if (Scene.PrimitiveSceneProxies[PrimitiveID]->IsSelected())
+				if (Scene.PrimitiveSceneProxies[PackedIndex]->IsSelected())
 				{
-					SelectedPrimitiveFlags[PrimitiveID / BitsPerWord] |= 1U << uint32(PrimitiveID % BitsPerWord);
+					FPrimitiveSceneInfo* PrimitiveSceneInfo = Scene.Primitives[PackedIndex];
+					FPersistentPrimitiveIndex PersistentPrimitiveIndex = PrimitiveSceneInfo->GetPersistentIndex();
+
+					SelectedPrimitiveFlags[PersistentPrimitiveIndex.Index / BitsPerWord] |= 1U << uint32(PersistentPrimitiveIndex.Index % BitsPerWord);
 
 					// Collect Names
 					if (SelectedNameInfos.Num() < MaxPrimitiveNameCount)
 					{
-						const FString OwnerName = Scene.Primitives[PrimitiveID]->GetFullnameForDebuggingOnly();
+						const FString OwnerName = PrimitiveSceneInfo->GetFullnameForDebuggingOnly();
 						const uint32 NameOffset = SelectedNames.Num();
 						const uint32 NameLength = OwnerName.Len();
 						for (TCHAR C : OwnerName)
@@ -1983,7 +1957,7 @@ void FGPUScene::DebugRender(FRDGBuilder& GraphBuilder, FScene& InScene, FSceneUn
 						}
 
 						FPrimitiveSceneDebugNameInfo& NameInfo = SelectedNameInfos.AddDefaulted_GetRef();
-						NameInfo.PrimitiveID= PrimitiveID;
+						NameInfo.PrimitiveID= PersistentPrimitiveIndex.Index;
 						NameInfo.Length		= NameLength;
 						NameInfo.Offset		= NameOffset;
 						++SelectedCount;
diff --git a/Engine/Source/Runtime/Renderer/Private/GPUScene.h b/Engine/Source/Runtime/Renderer/Private/GPUScene.h
--- a/Engine/Source/Runtime/Renderer/Private/GPUScene.h
+++ b/Engine/Source/Runtime/Renderer/Private/GPUScene.h
@@ -209,7 +209,7 @@ public:
 	 * and prepare for dynamic primitive allocations.
 	 * Scene may be NULL which means there are zero scene primitives (but there may be dynamic ones added later).
 	 */
-	void BeginRender(const FScene* Scene, FGPUSceneDynamicContext &GPUSceneDynamicContext);
+	void BeginRender(FGPUSceneDynamicContext &GPUSceneDynamicContext);
 	inline bool IsRendering() const { return bInBeginEndBlock; }
 	void EndRender();
 
@@ -234,7 +234,7 @@ public:
 	/**
 	 * Upload primitives from View.DynamicPrimitiveCollector.
 	 */
-	void UploadDynamicPrimitiveShaderDataForView(FRDGBuilder& GraphBuilder, FScene& Scene, FViewInfo& View, bool bIsShadowView = false);
+	void UploadDynamicPrimitiveShaderDataForView(FRDGBuilder& GraphBuilder, FViewInfo& View, bool bIsShadowView = false);
 
 	/**
 	 * Modifies the GPUScene specific scene UB parameters to the current versions. Returns true if any of the parameters changed.
@@ -244,7 +244,7 @@ public:
 	/**
 	 * Pull all pending updates from Scene and upload primitive & instance data.
 	 */
-	void Update(FRDGBuilder& GraphBuilder, FSceneUniformBuffer& SceneUB, FScene& Scene, FRDGExternalAccessQueue& ExternalAccessQueue, IVisibilityTaskData* VisibilityTaskData = nullptr);
+	void Update(FRDGBuilder& GraphBuilder, FSceneUniformBuffer& SceneUB, FRDGExternalAccessQueue& ExternalAccessQueue, IVisibilityTaskData* VisibilityTaskData = nullptr);
 
 	/**
 	 * Queue the given primitive for upload to GPU at next call to Update.
@@ -288,7 +288,7 @@ public:
 	 * Draw GPU-Scene debug info, such as bounding boxes. Call once per view at some point in the frame after GPU scene has been updated fully.
 	 * What is drawn is controlled by the CVar: r.GPUScene.DebugMode. Enabling this cvar causes ShaderDraw to be being active (if supported). 
 	 */
-	void DebugRender(FRDGBuilder& GraphBuilder, FScene& Scene, FSceneUniformBuffer& SceneUniformBuffer, FViewInfo& View);
+	void DebugRender(FRDGBuilder& GraphBuilder, FSceneUniformBuffer& SceneUniformBuffer, FViewInfo& View);
 
 	/**
 	 * Manually trigger an allocator consolidate (will otherwise be done when an item is allocated).
@@ -400,25 +400,25 @@ private:
 	ERHIFeatureLevel::Type FeatureLevel;
 
 	template<typename FUploadDataSourceAdapter>
-	void UpdateBufferState(FRDGBuilder& GraphBuilder, FSceneUniformBuffer& SceneUB, FScene& Scene, const FUploadDataSourceAdapter& UploadDataSourceAdapter, bool bIsMainUpdate = false);
+	void UpdateBufferState(FRDGBuilder& GraphBuilder, FSceneUniformBuffer& SceneUB, const FUploadDataSourceAdapter& UploadDataSourceAdapter, bool bIsMainUpdate = false);
 
 	/**
 	 * Generalized upload that uses an adapter to abstract the data souce. Enables uploading scene primitives & dynamic primitives using a single path.
 	 * @parameter Scene may be null, as it is only needed for the Nanite material table update (which is coupled to the Scene at the moment).
 	 */
 	template<typename FUploadDataSourceAdapter>
-	void UploadGeneral(FRDGBuilder& GraphBuilder, FScene& Scene, FRDGExternalAccessQueue* ExternalAccessQueue, const FUploadDataSourceAdapter& UploadDataSourceAdapter, const UE::Tasks::FTask& PrerequisiteTask);
+	void UploadGeneral(FRDGBuilder& GraphBuilder, FRDGExternalAccessQueue* ExternalAccessQueue, const FUploadDataSourceAdapter& UploadDataSourceAdapter, const UE::Tasks::FTask& PrerequisiteTask);
 
 	/**
 	 * Upload scene light data to gpu
 	 */
-	void UpdateGPULights(FRDGBuilder& GraphBuilder, FScene& Scene, const UE::Tasks::FTask& PrerequisiteTask);
+	void UpdateGPULights(FRDGBuilder& GraphBuilder, const UE::Tasks::FTask& PrerequisiteTask);
 
 	static void InitLightData(const FLightSceneInfoCompact& LightInfoCompact, bool bAllowStaticLighting, FLightSceneData& DataOut);
 
-	void UploadDynamicPrimitiveShaderDataForViewInternal(FRDGBuilder& GraphBuilder, FScene& Scene, FViewInfo& View, bool bIsShadowView);
+	void UploadDynamicPrimitiveShaderDataForViewInternal(FRDGBuilder& GraphBuilder, FViewInfo& View, bool bIsShadowView);
 
-	void UpdateInternal(FRDGBuilder& GraphBuilder, FSceneUniformBuffer& SceneUB, FScene& Scene, FRDGExternalAccessQueue& ExternalAccessQueue, IVisibilityTaskData* VisibilityTaskData);
+	void UpdateInternal(FRDGBuilder& GraphBuilder, FSceneUniformBuffer& SceneUB, FRDGExternalAccessQueue& ExternalAccessQueue, IVisibilityTaskData* VisibilityTaskData);
 
 	void AddUpdatePrimitiveIdsPass(FRDGBuilder& GraphBuilder, FInstanceGPULoadBalancer& IdOnlyUpdateItems);
 
@@ -436,10 +436,10 @@ private:
 class FGPUSceneScopeBeginEndHelper
 {
 public:
-	FGPUSceneScopeBeginEndHelper(FGPUScene& InGPUScene, FGPUSceneDynamicContext &GPUSceneDynamicContext, const FScene* Scene) :
+	FGPUSceneScopeBeginEndHelper(FGPUScene& InGPUScene, FGPUSceneDynamicContext &GPUSceneDynamicContext) :
 		GPUScene(InGPUScene)
 	{
-		GPUScene.BeginRender(Scene, GPUSceneDynamicContext);
+		GPUScene.BeginRender(GPUSceneDynamicContext);
 	}
 
 	~FGPUSceneScopeBeginEndHelper()
diff --git a/Engine/Source/Runtime/Renderer/Private/MobileShadingRenderer.cpp b/Engine/Source/Runtime/Renderer/Private/MobileShadingRenderer.cpp
--- a/Engine/Source/Runtime/Renderer/Private/MobileShadingRenderer.cpp
+++ b/Engine/Source/Runtime/Renderer/Private/MobileShadingRenderer.cpp
@@ -665,12 +665,12 @@ void FMobileSceneRenderer::InitViews(
 	{
 		RDG_CSV_STAT_EXCLUSIVE_SCOPE(GraphBuilder, UpdateGPUScene);
 
-		Scene->GPUScene.Update(GraphBuilder, GetSceneUniforms(), *Scene, ExternalAccessQueue);
+		Scene->GPUScene.Update(GraphBuilder, GetSceneUniforms(), ExternalAccessQueue);
 		for (int32 ViewIndex = 0; ViewIndex < AllViews.Num(); ViewIndex++)
 		{
 			FViewInfo& View = *AllViews[ViewIndex];
-			Scene->GPUScene.UploadDynamicPrimitiveShaderDataForView(GraphBuilder, *Scene, View);
-			Scene->GPUScene.DebugRender(GraphBuilder, *Scene, GetSceneUniforms(), View);
+			Scene->GPUScene.UploadDynamicPrimitiveShaderDataForView(GraphBuilder, View);
+			Scene->GPUScene.DebugRender(GraphBuilder, GetSceneUniforms(), View);
 		}
 	}
 
@@ -917,7 +917,7 @@ void FMobileSceneRenderer::Render(FRDGBuilder& GraphBuilder)
 	GPU_MESSAGE_SCOPE(GraphBuilder);
 
 	// Establish scene primitive count (must be done after UpdateAllPrimitiveSceneInfos)
-	FGPUSceneScopeBeginEndHelper GPUSceneScopeBeginEndHelper(Scene->GPUScene, GPUSceneDynamicContext, Scene);
+	FGPUSceneScopeBeginEndHelper GPUSceneScopeBeginEndHelper(Scene->GPUScene, GPUSceneDynamicContext);
 
 	if (bRendererOutputFinalSceneColor)
 	{
diff --git a/Engine/Source/Runtime/Renderer/Private/RendererScene.cpp b/Engine/Source/Runtime/Renderer/Private/RendererScene.cpp
--- a/Engine/Source/Runtime/Renderer/Private/RendererScene.cpp
+++ b/Engine/Source/Runtime/Renderer/Private/RendererScene.cpp
@@ -5146,7 +5146,7 @@ FLightSceneChangeSet FScene::UpdateAllLightSceneInfos(FRDGBuilder& GraphBuilder)
 		}
 	}
 	// This can't access the scene light data if done async since it happens before the actual removals.
-	OnPreLigtSceneInfoUpdate.Broadcast(GraphBuilder, FLightSceneChangeSet{ ChangeSet.RemovedLightIds, TConstArrayView<int32>(), ChangeSet.TransformUpdatedLightIds, ChangeSet.ColorUpdatedLightIds });
+	OnPreLightSceneInfoUpdate.Broadcast(GraphBuilder, FLightSceneChangeSet{ ChangeSet.RemovedLightIds, TConstArrayView<int32>(), ChangeSet.TransformUpdatedLightIds, ChangeSet.ColorUpdatedLightIds });
 
 	// Batch process all light removes
 	for (int32 LightId : ChangeSet.RemovedLightIds)
@@ -5218,7 +5218,7 @@ FLightSceneChangeSet FScene::UpdateAllLightSceneInfos(FRDGBuilder& GraphBuilder)
 		}
 	}
 
-	OnPostLigtSceneInfoUpdate.Broadcast(GraphBuilder, FLightSceneChangeSet{ ChangeSet.RemovedLightIds, ChangeSet.AddedLightIds, ChangeSet.TransformUpdatedLightIds, ChangeSet.ColorUpdatedLightIds });
+	OnPostLightSceneInfoUpdate.Broadcast(GraphBuilder, FLightSceneChangeSet{ ChangeSet.RemovedLightIds, ChangeSet.AddedLightIds, ChangeSet.TransformUpdatedLightIds, ChangeSet.ColorUpdatedLightIds });
 
 	SceneLightInfoUpdates->Reset();
 
diff --git a/Engine/Source/Runtime/Renderer/Private/SceneHitProxyRendering.cpp b/Engine/Source/Runtime/Renderer/Private/SceneHitProxyRendering.cpp
--- a/Engine/Source/Runtime/Renderer/Private/SceneHitProxyRendering.cpp
+++ b/Engine/Source/Runtime/Renderer/Private/SceneHitProxyRendering.cpp
@@ -561,7 +561,7 @@ void FMobileSceneRenderer::RenderHitProxies(FRDGBuilder& GraphBuilder)
 
 	GPU_MESSAGE_SCOPE(GraphBuilder);
 
-	FGPUSceneScopeBeginEndHelper GPUSceneScopeBeginEndHelper(Scene->GPUScene, GPUSceneDynamicContext, Scene);
+	FGPUSceneScopeBeginEndHelper GPUSceneScopeBeginEndHelper(Scene->GPUScene, GPUSceneDynamicContext);
 
 #if WITH_EDITOR
 	FSceneTexturesConfig& SceneTexturesConfig = GetActiveSceneTexturesConfig();
@@ -599,7 +599,7 @@ void FDeferredShadingSceneRenderer::RenderHitProxies(FRDGBuilder& GraphBuilder)
 
 	GPU_MESSAGE_SCOPE(GraphBuilder);
 
-	FGPUSceneScopeBeginEndHelper GPUSceneScopeBeginEndHelper(Scene->GPUScene, GPUSceneDynamicContext, Scene);
+	FGPUSceneScopeBeginEndHelper GPUSceneScopeBeginEndHelper(Scene->GPUScene, GPUSceneDynamicContext);
 
 
 #if WITH_EDITOR
@@ -638,7 +638,7 @@ void FDeferredShadingSceneRenderer::RenderHitProxies(FRDGBuilder& GraphBuilder)
 
 	for (int32 ViewIndex = 0; ViewIndex < Views.Num(); ViewIndex++)
 	{
-		Scene->GPUScene.UploadDynamicPrimitiveShaderDataForView(GraphBuilder, *Scene, Views[ViewIndex]);
+		Scene->GPUScene.UploadDynamicPrimitiveShaderDataForView(GraphBuilder, Views[ViewIndex]);
 	}
 
 	EndInitViews(GraphBuilder, LumenFrameTemporaries, InstanceCullingManager, InitViewTaskDatas);
diff --git a/Engine/Source/Runtime/Renderer/Private/ScenePrivate.h b/Engine/Source/Runtime/Renderer/Private/ScenePrivate.h
--- a/Engine/Source/Runtime/Renderer/Private/ScenePrivate.h
+++ b/Engine/Source/Runtime/Renderer/Private/ScenePrivate.h
@@ -3700,13 +3700,13 @@ public:
 	 * Thus, AddedLightIds is not valid in the change set as the added lights do not have assigned IDs yet.
 	 * IF using this to drive an async task, care must be taken as the (light) scene will be modified directly after.
 	 */
-	FSceneLightSceneInfoUpdateDelegate OnPreLigtSceneInfoUpdate;
+	FSceneLightSceneInfoUpdateDelegate OnPreLightSceneInfoUpdate;
 	/**
 	 * This delegate is invoked during the scene update phase _after_ all light changes are applied.
 	 * Thus, RemovedLightIds may contain ID's that are no longer valid or are now referencing newly added lights.
 	 * IF using this to drive an async task, the core light scene info may be used, but primitive scene updates will still be ongoing (e.g., light/primitive interactions may change).
 	 */
-	FSceneLightSceneInfoUpdateDelegate OnPostLigtSceneInfoUpdate;
+	FSceneLightSceneInfoUpdateDelegate OnPostLightSceneInfoUpdate;
 protected:
 
 private:
diff --git a/Engine/Source/Runtime/Renderer/Private/SceneRendering.cpp b/Engine/Source/Runtime/Renderer/Private/SceneRendering.cpp
--- a/Engine/Source/Runtime/Renderer/Private/SceneRendering.cpp
+++ b/Engine/Source/Runtime/Renderer/Private/SceneRendering.cpp
@@ -5125,10 +5125,10 @@ public:
 		ViewFamily.SetSceneRenderer(&Renderer);
 
 		Scene.UpdateAllPrimitiveSceneInfos(GraphBuilder);
-		GPUScene.BeginRender(&Scene, GPUSceneDynamicContext);
+		GPUScene.BeginRender(GPUSceneDynamicContext);
 
 		FRDGExternalAccessQueue ExternalAccessQueue;
-		Scene.GPUScene.Update(GraphBuilder, Renderer.GetSceneUniforms(), Scene, ExternalAccessQueue);
+		Scene.GPUScene.Update(GraphBuilder, Renderer.GetSceneUniforms(), ExternalAccessQueue);
 		ExternalAccessQueue.Submit(GraphBuilder);
 	}
 
diff --git a/Engine/Source/Runtime/Renderer/Private/SceneVisibility.cpp b/Engine/Source/Runtime/Renderer/Private/SceneVisibility.cpp
--- a/Engine/Source/Runtime/Renderer/Private/SceneVisibility.cpp
+++ b/Engine/Source/Runtime/Renderer/Private/SceneVisibility.cpp
@@ -5310,7 +5310,7 @@ void FDeferredShadingSceneRenderer::BeginInitViews(
 		RDG_CSV_STAT_EXCLUSIVE_SCOPE(GraphBuilder, UpdateGPUScene);
 		RDG_GPU_STAT_SCOPE(GraphBuilder, GPUSceneUpdate);
 
-		Scene->GPUScene.Update(GraphBuilder, GetSceneUniforms(), *Scene, ExternalAccessQueue, TaskDatas.VisibilityTaskData);
+		Scene->GPUScene.Update(GraphBuilder, GetSceneUniforms(), ExternalAccessQueue, TaskDatas.VisibilityTaskData);
 	}
 
 	// This must happen before we start initialising and using views.
diff --git a/Engine/Source/Runtime/Renderer/Private/ShadowSetup.cpp b/Engine/Source/Runtime/Renderer/Private/ShadowSetup.cpp
--- a/Engine/Source/Runtime/Renderer/Private/ShadowSetup.cpp
+++ b/Engine/Source/Runtime/Renderer/Private/ShadowSetup.cpp
@@ -6502,32 +6502,32 @@ void FSceneRenderer::FinishDynamicShadowMeshPassSetup(FRDGBuilder& GraphBuilder,
 	{
 		for (FProjectedShadowInfo* ProjectedShadowInfo : ShadowMapAtlas.Shadows)
 		{
-			Scene->GPUScene.UploadDynamicPrimitiveShaderDataForView(GraphBuilder, *Scene, *ProjectedShadowInfo->ShadowDepthView, true);
+			Scene->GPUScene.UploadDynamicPrimitiveShaderDataForView(GraphBuilder, *ProjectedShadowInfo->ShadowDepthView, true);
 		}
 	}
 	for (FSortedShadowMapAtlas& ShadowMap : SortedShadowsForShadowDepthPass.ShadowMapCubemaps)
 	{
 		check(ShadowMap.Shadows.Num() == 1);
 		FProjectedShadowInfo* ProjectedShadowInfo = ShadowMap.Shadows[0];
-		Scene->GPUScene.UploadDynamicPrimitiveShaderDataForView(GraphBuilder, *Scene, *ProjectedShadowInfo->ShadowDepthView, true);
+		Scene->GPUScene.UploadDynamicPrimitiveShaderDataForView(GraphBuilder, *ProjectedShadowInfo->ShadowDepthView, true);
 	}
 	for (FProjectedShadowInfo* ProjectedShadowInfo : SortedShadowsForShadowDepthPass.PreshadowCache.Shadows)
 	{
 		if (!ProjectedShadowInfo->bDepthsCached)
 		{
-			Scene->GPUScene.UploadDynamicPrimitiveShaderDataForView(GraphBuilder, *Scene, *ProjectedShadowInfo->ShadowDepthView, true);
+			Scene->GPUScene.UploadDynamicPrimitiveShaderDataForView(GraphBuilder, *ProjectedShadowInfo->ShadowDepthView, true);
 		}
 	}
 	for (const FSortedShadowMapAtlas& ShadowMapAtlas : SortedShadowsForShadowDepthPass.TranslucencyShadowMapAtlases)
 	{
 		for (FProjectedShadowInfo* ProjectedShadowInfo : ShadowMapAtlas.Shadows)
 		{
-			Scene->GPUScene.UploadDynamicPrimitiveShaderDataForView(GraphBuilder, *Scene, *ProjectedShadowInfo->ShadowDepthView, true);
+			Scene->GPUScene.UploadDynamicPrimitiveShaderDataForView(GraphBuilder, *ProjectedShadowInfo->ShadowDepthView, true);
 		}
 	}
 	for (FProjectedShadowInfo* ProjectedShadowInfo : SortedShadowsForShadowDepthPass.VirtualShadowMapShadows)
 	{
-		Scene->GPUScene.UploadDynamicPrimitiveShaderDataForView(GraphBuilder, *Scene, *ProjectedShadowInfo->ShadowDepthView, true);
+		Scene->GPUScene.UploadDynamicPrimitiveShaderDataForView(GraphBuilder, *ProjectedShadowInfo->ShadowDepthView, true);
 	}
 
 	DynamicReadBufferForShadows.Commit(GraphBuilder.RHICmdList);
diff --git a/Engine/Source/Runtime/Renderer/Private/Shadows/ShadowScene.cpp b/Engine/Source/Runtime/Renderer/Private/Shadows/ShadowScene.cpp
--- a/Engine/Source/Runtime/Renderer/Private/Shadows/ShadowScene.cpp
+++ b/Engine/Source/Runtime/Renderer/Private/Shadows/ShadowScene.cpp
@@ -30,7 +30,7 @@ DECLARE_DWORD_COUNTER_STAT(TEXT("Active Light Count"), STAT_ActiveLightCount, ST
 FShadowScene::FShadowScene(FScene & InScene)
 	: Scene(InScene)
 {
-	Scene.OnPostLigtSceneInfoUpdate.AddLambda([this](FRDGBuilder& GraphBuilder, const FLightSceneChangeSet& LightSceneChangeSet) { PostLightsUpdate(GraphBuilder, LightSceneChangeSet); });
+	Scene.OnPostLightSceneInfoUpdate.AddLambda([this](FRDGBuilder& GraphBuilder, const FLightSceneChangeSet& LightSceneChangeSet) { PostLightsUpdate(GraphBuilder, LightSceneChangeSet); });
 }
 
 void FShadowScene::PostLightsUpdate(FRDGBuilder& GraphBuilder, const FLightSceneChangeSet& LightSceneChangeSet)
diff --git a/Engine/Source/Runtime/Renderer/Private/VT/RuntimeVirtualTextureRender.cpp b/Engine/Source/Runtime/Renderer/Private/VT/RuntimeVirtualTextureRender.cpp
--- a/Engine/Source/Runtime/Renderer/Private/VT/RuntimeVirtualTextureRender.cpp
+++ b/Engine/Source/Runtime/Renderer/Private/VT/RuntimeVirtualTextureRender.cpp
@@ -1569,11 +1569,11 @@ namespace RuntimeVirtualTexture
 		InDesc.Scene->UpdateAllPrimitiveSceneInfos(GraphBuilder);
 
 		// Call to let GPU-Scene determine if it is active and record scene primitive count
-		FGPUSceneScopeBeginEndHelper GPUSceneScopeBeginEndHelper(InDesc.Scene->GPUScene, GPUSceneDynamicContext, InDesc.Scene);
+		FGPUSceneScopeBeginEndHelper GPUSceneScopeBeginEndHelper(InDesc.Scene->GPUScene, GPUSceneDynamicContext);
 
 		FSceneUniformBuffer SceneUB {};
 		FRDGExternalAccessQueue ExternalAccessQueue;
-		InDesc.Scene->GPUScene.Update(GraphBuilder, SceneUB, *InDesc.Scene, ExternalAccessQueue);
+		InDesc.Scene->GPUScene.Update(GraphBuilder, SceneUB, ExternalAccessQueue);
 		if (InDesc.Scene->SplineMeshSceneResources)
 		{
 			InDesc.Scene->SplineMeshSceneResources->Update(GraphBuilder, SceneUB);
