Moved GPU-Scene update to the end of UpdateAllPrimitiveSceneInfos to enable more efficient and consistent scene updates for sub-systems that interact with GPU-scene.

diff --git a/Engine/Source/Runtime/Renderer/Private/DeferredShadingRenderer.cpp b/Engine/Source/Runtime/Renderer/Private/DeferredShadingRenderer.cpp
--- a/Engine/Source/Runtime/Renderer/Private/DeferredShadingRenderer.cpp
+++ b/Engine/Source/Runtime/Renderer/Private/DeferredShadingRenderer.cpp
@@ -2676,7 +2676,7 @@ void FDeferredShadingSceneRenderer::Render(FRDGBuilder& GraphBuilder)
 	TUniquePtr<FVirtualTextureUpdater> VirtualTextureUpdater;
 	FLumenSceneFrameTemporaries LumenFrameTemporaries;
 
-	FGPUSceneScopeBeginEndHelper GPUSceneScopeBeginEndHelper(Scene->GPUScene, GPUSceneDynamicContext);
+	FGPUSceneScopeBeginEndHelper GPUSceneScopeBeginEndHelper(GraphBuilder, Scene->GPUScene, GPUSceneDynamicContext);
 
 	const bool bUseVirtualTexturing = UseVirtualTexturing(ShaderPlatform);
 
diff --git a/Engine/Source/Runtime/Renderer/Private/GPUScene.cpp b/Engine/Source/Runtime/Renderer/Private/GPUScene.cpp
--- a/Engine/Source/Runtime/Renderer/Private/GPUScene.cpp
+++ b/Engine/Source/Runtime/Renderer/Private/GPUScene.cpp
@@ -650,28 +650,34 @@ FGPUScene::~FGPUScene()
 #endif
 }
 
-void FGPUScene::BeginRender(FGPUSceneDynamicContext &GPUSceneDynamicContext)
+void FGPUScene::BeginRender(FRDGBuilder& GraphBuilder, FGPUSceneDynamicContext &GPUSceneDynamicContext)
 {
 	ensure(!bInBeginEndBlock);
 	ensure(CurrentDynamicContext == nullptr);
+	check(NumScenePrimitives == Scene.Primitives.Num());
+	// Should always be reset to this as the neutral state.
+	check(DynamicPrimitivesOffset == Scene.GetMaxPersistentPrimitiveIndex());
 
-	ensure(bIsEnabled == UseGPUScene(GMaxRHIShaderPlatform, Scene.GetFeatureLevel()));
-	NumScenePrimitives = Scene.Primitives.Num();
-
-	CurrentDynamicContext = &GPUSceneDynamicContext;
+	// Do it anyway for old times sake
 	DynamicPrimitivesOffset = Scene.GetMaxPersistentPrimitiveIndex();
+	CurrentDynamicContext = &GPUSceneDynamicContext;
 	bInBeginEndBlock = true;
+
+	check(!CachedRegisteredBuffers.IsValid());
+	CachedRegisteredBuffers = RegisterBuffers(GraphBuilder);
 }
 
 void FGPUScene::EndRender()
 {
 	ensure(bInBeginEndBlock);
 	ensure(CurrentDynamicContext != nullptr);
-	DynamicPrimitivesOffset = -1;
+	check(DynamicPrimitivesOffset >= NumScenePrimitives);
+
+	// Pop all dynamic primitives off the stack
+	DynamicPrimitivesOffset = Scene.GetMaxPersistentPrimitiveIndex();
 	bInBeginEndBlock = false;
 	CurrentDynamicContext = nullptr;
-	BufferState = {};
-	ShaderParameters = {};
+	CachedRegisteredBuffers = {};
 }
 
 void FGPUScene::UpdateGPULights(FRDGBuilder& GraphBuilder, const UE::Tasks::FTask& PrerequisiteTask)
@@ -699,7 +705,10 @@ void FGPUScene::UpdateGPULights(FRDGBuilder& GraphBuilder, const UE::Tasks::FTas
 
 	}, PrerequisiteTask);
 
-	GraphBuilder.QueueBufferUpload<FLightSceneData>(BufferState.LightDataBuffer, LightData, ERDGInitialDataFlags::NoCopy);
+	const uint32 LightDataBufferSize = FMath::RoundUpToPowerOfTwo(FMath::Max(Scene.Lights.GetMaxIndex(), InitialBufferSize));
+	FRDGBuffer *LightDataBufferRDG = ResizeStructuredBufferIfNeeded(GraphBuilder, LightDataBuffer, LightDataBufferSize * sizeof(FLightSceneData), TEXT("GPUScene.LightData"));
+
+	GraphBuilder.QueueBufferUpload<FLightSceneData>(LightDataBufferRDG, LightData, ERDGInitialDataFlags::NoCopy);
 }
 
 void FGPUScene::InitLightData(const FLightSceneInfoCompact& LightInfoCompact, bool bAllowStaticLighting, FLightSceneData& DataOut)
@@ -740,7 +749,6 @@ void FGPUScene::UpdateInternal(FRDGBuilder& GraphBuilder, FSceneUniformBuffer& S
 {
 	LLM_SCOPE_BYTAG(GPUScene);
 
-	check(bInBeginEndBlock);
 	check(bIsEnabled == UseGPUScene(GMaxRHIShaderPlatform, Scene.GetFeatureLevel()));
 	check(NumScenePrimitives == Scene.Primitives.Num());
 	check(DynamicPrimitivesOffset >= Scene.Primitives.Num());
@@ -810,10 +818,9 @@ void FGPUScene::UpdateInternal(FRDGBuilder& GraphBuilder, FSceneUniformBuffer& S
 		}
 	}
 
-	check(!BufferState.IsValid());
-
 	FUploadDataSourceAdapterScenePrimitives& Adapter = *GraphBuilder.AllocObject<FUploadDataSourceAdapterScenePrimitives>(Scene, SceneFrameNumber, MoveTemp(PrimitivesToUpdate), MoveTemp(PrimitiveDirtyState));
-	UpdateBufferState(GraphBuilder, SceneUB, Adapter, true);
+	
+	FRegisteredBuffers BufferState = UpdateBufferAllocations(GraphBuilder, SceneUB, Adapter);
 
 	// Run a pass that clears (Sets ID to invalid) any instances that need it
 	AddClearInstancesPass(GraphBuilder, Scene.InstanceCullingOcclusionQueryRenderer);
@@ -835,24 +842,21 @@ void FGPUScene::UpdateInternal(FRDGBuilder& GraphBuilder, FSceneUniformBuffer& S
 		QUICK_SCOPE_CYCLE_COUNTER(STAT_UpdateGPUScene);
 		SCOPE_CYCLE_COUNTER(STAT_UpdateGPUSceneTime);
 
-		UploadGeneral<FUploadDataSourceAdapterScenePrimitives>(GraphBuilder, &ExternalAccessQueue, Adapter, PrerequisiteAsyncTask);
-
-		UpdateGPULights(GraphBuilder, PrerequisiteAsyncTask);
+		UploadGeneral<FUploadDataSourceAdapterScenePrimitives>(GraphBuilder, BufferState, &ExternalAccessQueue, Adapter, PrerequisiteAsyncTask);
 	}
 }
 
 template<typename FUploadDataSourceAdapter>
-void FGPUScene::UpdateBufferState(FRDGBuilder& GraphBuilder, FSceneUniformBuffer& SceneUB, const FUploadDataSourceAdapter& UploadDataSourceAdapter, bool bIsMainUpdate)
+FGPUScene::FRegisteredBuffers FGPUScene::UpdateBufferAllocations(FRDGBuilder& GraphBuilder, FSceneUniformBuffer& SceneUB, const FUploadDataSourceAdapter& UploadDataSourceAdapter)
 {
 	LLM_SCOPE_BYTAG(GPUScene);
 
-	check(bInBeginEndBlock);
 	check(bIsEnabled == UseGPUScene(GMaxRHIShaderPlatform, Scene.GetFeatureLevel()));
 	check(NumScenePrimitives == Scene.Primitives.Num());
 
 	RDG_GPU_MASK_SCOPE(GraphBuilder, FRHIGPUMask::All());
 
-	constexpr int32 InitialBufferSize = 256;
+	FRegisteredBuffers BufferState;
 
 	const uint32 SizeReserve = FMath::RoundUpToPowerOfTwo(FMath::Max(DynamicPrimitivesOffset, InitialBufferSize));
 	BufferState.PrimitiveBuffer = ResizeStructuredBufferIfNeeded(GraphBuilder, PrimitiveBuffer, SizeReserve * sizeof(FPrimitiveSceneShaderData::Data), TEXT("GPUScene.PrimitiveData"));
@@ -866,7 +870,6 @@ void FGPUScene::UpdateBufferState(FRDGBuilder& GraphBuilder, FSceneUniformBuffer
 
 	BufferState.InstanceSceneDataBuffer = ResizeStructuredBufferSOAIfNeeded(GraphBuilder, InstanceSceneDataBuffer, ResizeParams, TEXT("GPUScene.InstanceSceneData"));
 	InstanceSceneDataSOAStride = InstanceSceneDataSizeReserve;
-	BufferState.InstanceSceneDataSOAStride = InstanceSceneDataSizeReserve;
 
 	const uint32 PayloadFloat4Count = FMath::Max(InstancePayloadDataAllocator.GetMaxSize(), InitialBufferSize);
 	const uint32 InstancePayloadDataSizeReserve = FMath::RoundUpToPowerOfTwo(PayloadFloat4Count * sizeof(FVector4f));
@@ -886,14 +889,10 @@ void FGPUScene::UpdateBufferState(FRDGBuilder& GraphBuilder, FSceneUniformBuffer
 	
 	const uint32 LightMapDataBufferSize = FMath::RoundUpToPowerOfTwo(FMath::Max(LightmapDataAllocator.GetMaxSize(), InitialBufferSize));
 	BufferState.LightmapDataBuffer = ResizeStructuredBufferIfNeeded(GraphBuilder, LightmapDataBuffer, LightMapDataBufferSize * sizeof(FLightmapSceneShaderData::Data), TEXT("GPUScene.LightmapData"));
-	BufferState.LightMapDataBufferSize = LightMapDataBufferSize;
-	
-	if (bIsMainUpdate)
-	{
-		const uint32 LightDataBufferSize = FMath::RoundUpToPowerOfTwo(FMath::Max(Scene.Lights.GetMaxIndex(), InitialBufferSize));
-		BufferState.LightDataBuffer = GraphBuilder.CreateBuffer(FRDGBufferDesc::CreateStructuredDesc(sizeof(FLightSceneData), FMath::Max(1, Scene.Lights.GetMaxIndex())), TEXT("GPUScene.LightData"));
-	}
 
+	BufferState.LightDataBuffer = GraphBuilder.RegisterExternalBuffer(LightDataBuffer);
+
+	FGPUSceneResourceParameters ShaderParameters;
 	ShaderParameters.GPUSceneInstanceSceneData = GraphBuilder.CreateSRV(BufferState.InstanceSceneDataBuffer);
 	ShaderParameters.GPUSceneInstancePayloadData = GraphBuilder.CreateSRV(BufferState.InstancePayloadDataBuffer);
 	ShaderParameters.GPUScenePrimitiveSceneData = GraphBuilder.CreateSRV(BufferState.PrimitiveBuffer);
@@ -905,6 +904,31 @@ void FGPUScene::UpdateBufferState(FRDGBuilder& GraphBuilder, FSceneUniformBuffer
 	ShaderParameters.GPUSceneFrameNumber = GetSceneFrameNumber();
 
 	SceneUB.Set(SceneUB::GPUScene, ShaderParameters);
+
+	if (bInBeginEndBlock)
+	{
+		CachedRegisteredBuffers = BufferState;
+	}
+	else
+	{
+		check(!CachedRegisteredBuffers.IsValid());
+	}
+
+	return BufferState;
+}
+
+
+FGPUScene::FRegisteredBuffers FGPUScene::RegisterBuffers(FRDGBuilder& GraphBuilder) const
+{
+	FRegisteredBuffers Result;
+
+	Result.PrimitiveBuffer = GraphBuilder.RegisterExternalBuffer(PrimitiveBuffer);
+	Result.InstanceSceneDataBuffer = GraphBuilder.RegisterExternalBuffer(InstanceSceneDataBuffer);
+	Result.InstancePayloadDataBuffer = GraphBuilder.RegisterExternalBuffer(InstancePayloadDataBuffer);
+	Result.LightmapDataBuffer = GraphBuilder.RegisterExternalBuffer(LightmapDataBuffer);
+	Result.LightDataBuffer = GraphBuilder.RegisterExternalBuffer(LightDataBuffer);
+
+	return Result;
 }
 
 /**
@@ -995,7 +1019,7 @@ struct FInstanceBatcher
 };
 
 template<typename FUploadDataSourceAdapter>
-void FGPUScene::UploadGeneral(FRDGBuilder& GraphBuilder, FRDGExternalAccessQueue* ExternalAccessQueue, const FUploadDataSourceAdapter& UploadDataSourceAdapter, const UE::Tasks::FTask& PrerequisiteTask)
+void FGPUScene::UploadGeneral(FRDGBuilder& GraphBuilder, const FRegisteredBuffers &BufferState, FRDGExternalAccessQueue* ExternalAccessQueue, const FUploadDataSourceAdapter& UploadDataSourceAdapter, const UE::Tasks::FTask& PrerequisiteTask)
 {
 	LLM_SCOPE_BYTAG(GPUScene);
 
@@ -1041,7 +1065,7 @@ void FGPUScene::UploadGeneral(FRDGBuilder& GraphBuilder, FRDGExternalAccessQueue
 	FTaskContext& TaskContext = *GraphBuilder.AllocObject<FTaskContext>();
 
 	TaskContext.NumPrimitiveDataUploads = NumPrimitiveDataUploads;
-	TaskContext.InstanceSceneDataSOAStride = BufferState.InstanceSceneDataSOAStride;
+	TaskContext.InstanceSceneDataSOAStride = InstanceSceneDataSOAStride;
 	TaskContext.PrimitiveUploadInfos.SetNumUninitialized(NumPrimitiveDataUploads);
 
 	for (int32 ItemIndex = 0; ItemIndex < NumPrimitiveDataUploads; ++ItemIndex)
@@ -1589,8 +1613,6 @@ void FGPUScene::UploadDynamicPrimitiveShaderDataForViewInternal(FRDGBuilder& Gra
 	// Make sure we are not trying to upload data that lives in a different context.
 	ensure(Collector.UploadData == nullptr || CurrentDynamicContext->DymamicPrimitiveUploadData.Find(Collector.UploadData) != INDEX_NONE);
 
-	check(BufferState.IsValid());
-
 	// Skip uploading empty & already uploaded data
 	const bool bNeedsUpload = Collector.UploadData != nullptr && NumPrimitiveDataUploads > 0 && !Collector.UploadData->bIsUploaded;
 	if (bNeedsUpload)
@@ -1629,12 +1651,12 @@ void FGPUScene::UploadDynamicPrimitiveShaderDataForViewInternal(FRDGBuilder& Gra
 			Collector.UploadData->InstancePayloadDataOffset,
 			SceneFrameNumber);
 
-		UpdateBufferState(GraphBuilder, View.GetSceneUniforms(), UploadAdapter, false);
+		FRegisteredBuffers BufferState = UpdateBufferAllocations(GraphBuilder, View.GetSceneUniforms(), UploadAdapter);
 
 		// Run a pass that clears (Sets ID to invalid) any instances that need it.
 		AddClearInstancesPass(GraphBuilder, Scene.InstanceCullingOcclusionQueryRenderer);
 
-		UploadGeneral<FUploadDataSourceAdapterDynamicPrimitives>(GraphBuilder, nullptr, UploadAdapter, UE::Tasks::FTask{});
+		UploadGeneral<FUploadDataSourceAdapterDynamicPrimitives>(GraphBuilder, BufferState, nullptr, UploadAdapter, UE::Tasks::FTask{});
 	}
 
 	FSceneUniformBuffer& SceneUniforms = View.GetSceneUniforms();
@@ -1702,24 +1724,9 @@ bool FGPUScene::FillSceneUniformBuffer(FRDGBuilder& GraphBuilder, FSceneUniformB
 		return false;
 	}
 
-	if (ShaderParameters.GPUScenePrimitiveSceneData != nullptr)
+	if (PrimitiveBuffer != nullptr)
 	{
-		return SceneUB.Set(SceneUB::GPUScene, ShaderParameters);
-	}
-	else if (PrimitiveBuffer != nullptr)
-	{
-		FGPUSceneResourceParameters TmpParameters;
-		// Not in an active rendering context, must register the buffers and fill in the data structure.
-		TmpParameters.GPUSceneInstanceSceneData = GraphBuilder.CreateSRV(GraphBuilder.RegisterExternalBuffer(InstanceSceneDataBuffer));
-		TmpParameters.GPUSceneInstancePayloadData = GraphBuilder.CreateSRV(GraphBuilder.RegisterExternalBuffer(InstancePayloadDataBuffer));
-		TmpParameters.GPUScenePrimitiveSceneData = GraphBuilder.CreateSRV(GraphBuilder.RegisterExternalBuffer(PrimitiveBuffer));
-		TmpParameters.GPUSceneLightmapData = GraphBuilder.CreateSRV(GraphBuilder.RegisterExternalBuffer(LightmapDataBuffer));
-		TmpParameters.GPUSceneLightData = GraphBuilder.CreateSRV(GSystemTextures.GetDefaultStructuredBuffer(GraphBuilder, sizeof(FLightSceneData)));
-		TmpParameters.InstanceDataSOAStride = InstanceSceneDataSOAStride;
-		TmpParameters.NumScenePrimitives = NumScenePrimitives;
-		TmpParameters.NumInstances = InstanceSceneDataAllocator.GetMaxSize();
-		TmpParameters.GPUSceneFrameNumber = GetSceneFrameNumber();
-		return SceneUB.Set(SceneUB::GPUScene, TmpParameters);
+		return SceneUB.Set(SceneUB::GPUScene, GetShaderParameters(GraphBuilder));
 	}
 	else
 	{
@@ -1728,6 +1735,30 @@ bool FGPUScene::FillSceneUniformBuffer(FRDGBuilder& GraphBuilder, FSceneUniformB
 	}
 }
 
+FGPUSceneResourceParameters FGPUScene::GetShaderParameters(FRDGBuilder& GraphBuilder) const
+{
+	FRegisteredBuffers BufferState = CachedRegisteredBuffers;
+	// If we're not in a begin/end block we need to register the buffers here and now.
+	if (!BufferState.IsValid())
+	{
+		BufferState = RegisterBuffers(GraphBuilder);
+	}
+
+	FGPUSceneResourceParameters TmpParameters;
+	// Not in an active rendering context, must register the buffers and fill in the data structure.
+	TmpParameters.GPUSceneInstanceSceneData = GraphBuilder.CreateSRV(BufferState.InstanceSceneDataBuffer);
+	TmpParameters.GPUSceneInstancePayloadData = GraphBuilder.CreateSRV(BufferState.InstancePayloadDataBuffer);
+	TmpParameters.GPUScenePrimitiveSceneData = GraphBuilder.CreateSRV(BufferState.PrimitiveBuffer);
+	TmpParameters.GPUSceneLightmapData = GraphBuilder.CreateSRV(BufferState.LightmapDataBuffer);
+	TmpParameters.GPUSceneLightData = GraphBuilder.CreateSRV(BufferState.LightDataBuffer);
+	TmpParameters.InstanceDataSOAStride = InstanceSceneDataSOAStride;
+	TmpParameters.NumScenePrimitives = NumScenePrimitives;
+	TmpParameters.NumInstances = InstanceSceneDataAllocator.GetMaxSize();
+	TmpParameters.GPUSceneFrameNumber = GetSceneFrameNumber();
+
+	return TmpParameters;
+}
+
 void FGPUScene::AddPrimitiveToUpdate(FPersistentPrimitiveIndex PersistentPrimitiveIndex, EPrimitiveDirtyState DirtyState)
 {
 	LLM_SCOPE_BYTAG(GPUScene);
@@ -1761,8 +1792,14 @@ void FGPUScene::Update(FRDGBuilder& GraphBuilder, FSceneUniformBuffer& SceneUB,
 	{
 		RDG_GPU_MASK_SCOPE(GraphBuilder, FRHIGPUMask::All());
 
-		ensure(bInBeginEndBlock);
-		
+		ensure(!bInBeginEndBlock);
+
+		// record primitive high-watermark (dynamic ones are allocated after that point)
+		NumScenePrimitives = Scene.Primitives.Num();
+
+		// Default state when updated (no "dynamic primitives" pushed)
+		DynamicPrimitivesOffset = Scene.GetMaxPersistentPrimitiveIndex();
+
 		UpdateInternal(GraphBuilder, SceneUB, ExternalAccessQueue, VisibilityTaskData);
 	}
 }
@@ -2068,8 +2105,6 @@ bool FGPUScene::ExecuteDeferredGPUWritePass(FRDGBuilder& GraphBuilder, TArray<FV
 		return false;
 	}
 
-	check(BufferState.IsValid());
-
 	RDG_EVENT_SCOPE(GraphBuilder, "GPUScene.DeferredGPUWrites - Pass %u", uint32(GPUWritePass));
 
 	FGPUSceneWriteDelegateParams Params;
@@ -2130,6 +2165,15 @@ void FGPUScene::OnPostSceneUpdate(FRDGBuilder& GraphBuilder, const FScenePostUpd
 	}
 }
 
+void FGPUScene::OnPostLightSceneInfoUpdate(FRDGBuilder& GraphBuilder, const FLightSceneChangeSet& LightsPostUpdateData)
+{
+	const bool bAnythingChanged = !(LightsPostUpdateData.AddedLightIds.IsEmpty() || LightsPostUpdateData.RemovedLightIds.IsEmpty() || LightsPostUpdateData.TransformUpdatedLightIds.IsEmpty() || LightsPostUpdateData.ColorUpdatedLightIds.IsEmpty());
+	if (bAnythingChanged || !LightDataBuffer.IsValid())
+	{
+		UpdateGPULights(GraphBuilder, UE::Tasks::FTask{});
+	}
+}
+
 SIZE_T FGPUScene::GetAllocatedSize() const
 {
 	return PrimitivesToUpdate.GetAllocatedSize()
@@ -2176,6 +2220,12 @@ FGPUScenePrimitiveCollector::FUploadData* FGPUSceneDynamicContext::AllocateDynam
  */
 void FGPUScene::GetWriteParameters(FRDGBuilder& GraphBuilder, FGPUSceneWriterParameters& GPUSceneWriterParametersOut)
 {
+	FRegisteredBuffers BufferState = CachedRegisteredBuffers;
+	// If we're not in a begin/end block we need to register the buffers here and now.
+	if (!BufferState.IsValid())
+	{
+		BufferState = RegisterBuffers(GraphBuilder);
+	}
 	GPUSceneWriterParametersOut.GPUSceneFrameNumber = SceneFrameNumber;
 	GPUSceneWriterParametersOut.GPUSceneInstanceSceneDataSOAStride = InstanceSceneDataSOAStride;
 	GPUSceneWriterParametersOut.GPUSceneNumAllocatedInstances = InstanceSceneDataAllocator.GetMaxSize();
diff --git a/Engine/Source/Runtime/Renderer/Private/GPUScene.h b/Engine/Source/Runtime/Renderer/Private/GPUScene.h
--- a/Engine/Source/Runtime/Renderer/Private/GPUScene.h
+++ b/Engine/Source/Runtime/Renderer/Private/GPUScene.h
@@ -27,6 +27,7 @@ class FScenePreUpdateChangeSet;
 class FScenePostUpdateChangeSet;
 class FPrimitiveSceneProxy;
 class IVisibilityTaskData;
+struct FLightSceneChangeSet;
 
 DECLARE_GPU_STAT_NAMED_EXTERN(GPUSceneUpdate, TEXT("GPUSceneUpdate"))
 
@@ -172,24 +173,6 @@ private:
 	FGPUScene& GPUScene;
 };
 
-// Buffers used by GPU-Scene, since they can be resized during updates AND the render passes must retain the 
-// right copy (this is chiefly because the init of shadow views after pre-pass means we need to be able to set 
-// up GPU-Scene before pre-pass, but then may discover new primitives etc. As there is no way to know how many
-// dynamic primitives will turn up after Pre-pass, we can't guarantee a resize won't happen).
-struct FGPUSceneBufferState
-{
-	bool IsValid() const { return PrimitiveBuffer != nullptr; }
-
-	FRDGBuffer* PrimitiveBuffer = nullptr;
-	FRDGBuffer* InstanceSceneDataBuffer = nullptr;
-	uint32 InstanceSceneDataSOAStride = 1; // Distance between arrays in float4s
-	FRDGBuffer* InstancePayloadDataBuffer = nullptr;
-	FRDGBuffer* LightmapDataBuffer = nullptr;
-	uint32 LightMapDataBufferSize = 0;
-
-	FRDGBuffer* LightDataBuffer = nullptr;
-};
-
 struct FGPUSceneInstanceRange
 {
 	uint32 InstanceSceneDataOffset;
@@ -209,7 +192,7 @@ public:
 	 * and prepare for dynamic primitive allocations.
 	 * Scene may be NULL which means there are zero scene primitives (but there may be dynamic ones added later).
 	 */
-	void BeginRender(FGPUSceneDynamicContext &GPUSceneDynamicContext);
+	void BeginRender(FRDGBuilder& GraphBuilder, FGPUSceneDynamicContext &GPUSceneDynamicContext);
 	inline bool IsRendering() const { return bInBeginEndBlock; }
 	void EndRender();
 
@@ -282,7 +265,7 @@ public:
 	/**
 	 * Return the GPU scene resource
 	 */
-	FGPUSceneResourceParameters GetShaderParameters() const { check(ShaderParameters.GPUScenePrimitiveSceneData != nullptr); return ShaderParameters; }
+	FGPUSceneResourceParameters GetShaderParameters(FRDGBuilder& GraphBuilder) const;
 
 	/**
 	 * Draw GPU-Scene debug info, such as bounding boxes. Call once per view at some point in the frame after GPU scene has been updated fully.
@@ -313,6 +296,10 @@ public:
 	 */
 	void OnPostSceneUpdate(FRDGBuilder& GraphBuilder, const FScenePostUpdateChangeSet& ScenePostUpdateData);
 
+	/**
+	 */
+	void OnPostLightSceneInfoUpdate(FRDGBuilder& OnPostLightSceneInfoUpdate, const FLightSceneChangeSet& LightsPostUpdateData);
+
 	bool bUpdateAllPrimitives;
 
 	/** GPU mirror of Primitives */
@@ -344,6 +331,27 @@ public:
 	SIZE_T GetAllocatedSize() const;
 
 private:
+	static constexpr int32 InitialBufferSize = 256;
+
+	TRefCountPtr<FRDGPooledBuffer> LightDataBuffer;
+
+	// Buffers used by GPU-Scene, since they can be resized during updates AND the render passes must retain the 
+	// right copy (this is chiefly because the init of shadow views after pre-pass means we need to be able to set 
+	// up GPU-Scene before pre-pass, but then may discover new primitives etc. As there is no way to know how many
+	// dynamic primitives will turn up after Pre-pass, we can't guarantee a resize won't happen).
+	struct FRegisteredBuffers
+	{
+		bool IsValid() const { return PrimitiveBuffer != nullptr; }
+
+		FRDGBuffer* PrimitiveBuffer = nullptr;
+		FRDGBuffer* InstanceSceneDataBuffer = nullptr;
+		FRDGBuffer* InstancePayloadDataBuffer = nullptr;
+		FRDGBuffer* LightmapDataBuffer = nullptr;
+		FRDGBuffer* LightDataBuffer = nullptr;
+	};
+
+
+
 	FScene &Scene;
 	FSpanAllocator		           InstanceSceneDataAllocator;
 
@@ -360,8 +368,7 @@ private:
 	/** Indices of primitives that need to be updated in GPU Scene */
 	TArray<FPersistentPrimitiveIndex> PrimitivesToUpdate;
 
-	FGPUSceneBufferState BufferState;
-	FGPUSceneResourceParameters ShaderParameters;
+	FRegisteredBuffers CachedRegisteredBuffers;
 
 	TArray<EPrimitiveDirtyState> PrimitiveDirtyState;
 
@@ -400,14 +407,19 @@ private:
 	ERHIFeatureLevel::Type FeatureLevel;
 
 	template<typename FUploadDataSourceAdapter>
-	void UpdateBufferState(FRDGBuilder& GraphBuilder, FSceneUniformBuffer& SceneUB, const FUploadDataSourceAdapter& UploadDataSourceAdapter, bool bIsMainUpdate = false);
+	FRegisteredBuffers UpdateBufferAllocations(FRDGBuilder& GraphBuilder, FSceneUniformBuffer& SceneUB, const FUploadDataSourceAdapter& UploadDataSourceAdapter);
+
+	/**
+	 * Register the external buffers with the graphbuilder.
+	 */
+	FRegisteredBuffers RegisterBuffers(FRDGBuilder& GraphBuilder) const;
 
 	/**
 	 * Generalized upload that uses an adapter to abstract the data souce. Enables uploading scene primitives & dynamic primitives using a single path.
 	 * @parameter Scene may be null, as it is only needed for the Nanite material table update (which is coupled to the Scene at the moment).
 	 */
 	template<typename FUploadDataSourceAdapter>
-	void UploadGeneral(FRDGBuilder& GraphBuilder, FRDGExternalAccessQueue* ExternalAccessQueue, const FUploadDataSourceAdapter& UploadDataSourceAdapter, const UE::Tasks::FTask& PrerequisiteTask);
+	void UploadGeneral(FRDGBuilder& GraphBuilder, const FRegisteredBuffers& BufferState, FRDGExternalAccessQueue* ExternalAccessQueue, const FUploadDataSourceAdapter& UploadDataSourceAdapter, const UE::Tasks::FTask& PrerequisiteTask);
 
 	/**
 	 * Upload scene light data to gpu
@@ -436,10 +448,10 @@ private:
 class FGPUSceneScopeBeginEndHelper
 {
 public:
-	FGPUSceneScopeBeginEndHelper(FGPUScene& InGPUScene, FGPUSceneDynamicContext &GPUSceneDynamicContext) :
+	FGPUSceneScopeBeginEndHelper(FRDGBuilder& GraphBuilder, FGPUScene& InGPUScene, FGPUSceneDynamicContext &GPUSceneDynamicContext) :
 		GPUScene(InGPUScene)
 	{
-		GPUScene.BeginRender(GPUSceneDynamicContext);
+		GPUScene.BeginRender(GraphBuilder, GPUSceneDynamicContext);
 	}
 
 	~FGPUSceneScopeBeginEndHelper()
diff --git a/Engine/Source/Runtime/Renderer/Private/InstanceCulling/InstanceCullingContext.cpp b/Engine/Source/Runtime/Renderer/Private/InstanceCulling/InstanceCullingContext.cpp
--- a/Engine/Source/Runtime/Renderer/Private/InstanceCulling/InstanceCullingContext.cpp
+++ b/Engine/Source/Runtime/Renderer/Private/InstanceCulling/InstanceCullingContext.cpp
@@ -794,7 +794,7 @@ void FInstanceCullingContext::BuildRenderingCommandsInternal(
 
 	PassParametersTmp.InstanceCullingPayloads = GraphBuilder.CreateSRV(CreateStructuredBuffer(GraphBuilder, TEXT("InstanceCulling.PayloadData"), PayloadData));
 
-	const FGPUSceneResourceParameters GPUSceneParameters = GPUScene.GetShaderParameters();
+	const FGPUSceneResourceParameters GPUSceneParameters = GPUScene.GetShaderParameters(GraphBuilder);
 
 	// Because the view uniforms are not set up by the time this runs
 	// PassParametersTmp.View = View.ViewUniformBuffer;
@@ -1147,7 +1147,7 @@ FInstanceCullingDeferredContext *FInstanceCullingContext::CreateDeferredContext(
 	FRDGBufferUAVRef InstanceIdsBufferUAV = GraphBuilder.CreateUAV(InstanceIdsBuffer, ERDGUnorderedAccessViewFlags::SkipBarrier);
 	DeferredContext->InstanceDataBuffer = InstanceIdOffsetBuffer;
 
-	const FGPUSceneResourceParameters GPUSceneParameters = GPUScene.GetShaderParameters();
+	const FGPUSceneResourceParameters GPUSceneParameters = GPUScene.GetShaderParameters(GraphBuilder);
 
 	// Because the view uniforms are not set up by the time this runs
 	// PassParameters->View = View.ViewUniformBuffer;
diff --git a/Engine/Source/Runtime/Renderer/Private/InstanceCulling/InstanceCullingOcclusionQuery.cpp b/Engine/Source/Runtime/Renderer/Private/InstanceCulling/InstanceCullingOcclusionQuery.cpp
--- a/Engine/Source/Runtime/Renderer/Private/InstanceCulling/InstanceCullingOcclusionQuery.cpp
+++ b/Engine/Source/Runtime/Renderer/Private/InstanceCulling/InstanceCullingOcclusionQuery.cpp
@@ -532,7 +532,7 @@ uint32 FInstanceCullingOcclusionQueryRenderer::Render(
 
 	const FIntVector HZBSize = HZBTexture->Desc.GetSize();
 
-	const FGPUSceneResourceParameters GPUSceneParameters = GPUScene.GetShaderParameters();
+	const FGPUSceneResourceParameters GPUSceneParameters = GPUScene.GetShaderParameters(GraphBuilder);
 
 	const FIntPoint ViewRectSize = View.ViewRect.Size();
 
@@ -856,7 +856,7 @@ void FInstanceCullingOcclusionQueryRenderer::RenderDebug(FRDGBuilder& GraphBuild
 	const FIntVector HZBSize = HZBTexture->Desc.GetSize();
 
 	const int32 NumInstances = GPUScene.GetNumInstances();
-	const FGPUSceneResourceParameters GPUSceneParameters = GPUScene.GetShaderParameters();
+	const FGPUSceneResourceParameters GPUSceneParameters = GPUScene.GetShaderParameters(GraphBuilder);
 
 	FOcclusionInstanceCullingDebugParameters* PassParameters = GraphBuilder.AllocParameters<FOcclusionInstanceCullingDebugParameters>();
 
diff --git a/Engine/Source/Runtime/Renderer/Private/MobileShadingRenderer.cpp b/Engine/Source/Runtime/Renderer/Private/MobileShadingRenderer.cpp
--- a/Engine/Source/Runtime/Renderer/Private/MobileShadingRenderer.cpp
+++ b/Engine/Source/Runtime/Renderer/Private/MobileShadingRenderer.cpp
@@ -665,7 +665,6 @@ void FMobileSceneRenderer::InitViews(
 	{
 		RDG_CSV_STAT_EXCLUSIVE_SCOPE(GraphBuilder, UpdateGPUScene);
 
-		Scene->GPUScene.Update(GraphBuilder, GetSceneUniforms(), ExternalAccessQueue);
 		for (int32 ViewIndex = 0; ViewIndex < AllViews.Num(); ViewIndex++)
 		{
 			FViewInfo& View = *AllViews[ViewIndex];
@@ -917,7 +916,7 @@ void FMobileSceneRenderer::Render(FRDGBuilder& GraphBuilder)
 	GPU_MESSAGE_SCOPE(GraphBuilder);
 
 	// Establish scene primitive count (must be done after UpdateAllPrimitiveSceneInfos)
-	FGPUSceneScopeBeginEndHelper GPUSceneScopeBeginEndHelper(Scene->GPUScene, GPUSceneDynamicContext);
+	FGPUSceneScopeBeginEndHelper GPUSceneScopeBeginEndHelper(GraphBuilder, Scene->GPUScene, GPUSceneDynamicContext);
 
 	if (bRendererOutputFinalSceneColor)
 	{
diff --git a/Engine/Source/Runtime/Renderer/Private/RendererScene.cpp b/Engine/Source/Runtime/Renderer/Private/RendererScene.cpp
--- a/Engine/Source/Runtime/Renderer/Private/RendererScene.cpp
+++ b/Engine/Source/Runtime/Renderer/Private/RendererScene.cpp
@@ -5219,6 +5219,7 @@ FLightSceneChangeSet FScene::UpdateAllLightSceneInfos(FRDGBuilder& GraphBuilder)
 	}
 
 	OnPostLightSceneInfoUpdate.Broadcast(GraphBuilder, FLightSceneChangeSet{ ChangeSet.RemovedLightIds, ChangeSet.AddedLightIds, ChangeSet.TransformUpdatedLightIds, ChangeSet.ColorUpdatedLightIds });
+	GPUScene.OnPostLightSceneInfoUpdate(GraphBuilder, FLightSceneChangeSet{ ChangeSet.RemovedLightIds, ChangeSet.AddedLightIds, ChangeSet.TransformUpdatedLightIds, ChangeSet.ColorUpdatedLightIds });
 
 	SceneLightInfoUpdates->Reset();
 
@@ -6511,6 +6512,18 @@ void FScene::UpdateAllPrimitiveSceneInfos(FRDGBuilder& GraphBuilder, EUpdateAllP
 		PrimitiveOcclusionBounds[SceneInfo->PackedIndex] = NewOccBounds.ExpandBy(OCCLUSION_SLOP + OccSlackDelta.Value);
 	}
 
+	{
+		RDG_CSV_STAT_EXCLUSIVE_SCOPE(GraphBuilder, UpdateGPUScene);
+		RDG_GPU_STAT_SCOPE(GraphBuilder, GPUSceneUpdate);
+
+		FRDGExternalAccessQueue ExternalAccessQueue;
+
+		// Process GPU scene prior to visibility to maximize overlap.
+		GPUScene.Update(GraphBuilder, SceneUB, ExternalAccessQueue, nullptr);
+
+		ExternalAccessQueue.Submit(GraphBuilder);
+	}
+
 	// Need to do this here since we delete the proxy next 
 	// TODO: should refactor to add this as a dependency for a proxy-deletion task instead.
 	SceneCullingUpdater.GetAsyncProxyUseTaskHandle().Wait();
diff --git a/Engine/Source/Runtime/Renderer/Private/SceneHitProxyRendering.cpp b/Engine/Source/Runtime/Renderer/Private/SceneHitProxyRendering.cpp
--- a/Engine/Source/Runtime/Renderer/Private/SceneHitProxyRendering.cpp
+++ b/Engine/Source/Runtime/Renderer/Private/SceneHitProxyRendering.cpp
@@ -561,7 +561,7 @@ void FMobileSceneRenderer::RenderHitProxies(FRDGBuilder& GraphBuilder)
 
 	GPU_MESSAGE_SCOPE(GraphBuilder);
 
-	FGPUSceneScopeBeginEndHelper GPUSceneScopeBeginEndHelper(Scene->GPUScene, GPUSceneDynamicContext);
+	FGPUSceneScopeBeginEndHelper GPUSceneScopeBeginEndHelper(GraphBuilder, Scene->GPUScene, GPUSceneDynamicContext);
 
 #if WITH_EDITOR
 	FSceneTexturesConfig& SceneTexturesConfig = GetActiveSceneTexturesConfig();
@@ -599,7 +599,7 @@ void FDeferredShadingSceneRenderer::RenderHitProxies(FRDGBuilder& GraphBuilder)
 
 	GPU_MESSAGE_SCOPE(GraphBuilder);
 
-	FGPUSceneScopeBeginEndHelper GPUSceneScopeBeginEndHelper(Scene->GPUScene, GPUSceneDynamicContext);
+	FGPUSceneScopeBeginEndHelper GPUSceneScopeBeginEndHelper(GraphBuilder, Scene->GPUScene, GPUSceneDynamicContext);
 
 
 #if WITH_EDITOR
diff --git a/Engine/Source/Runtime/Renderer/Private/SceneRendering.cpp b/Engine/Source/Runtime/Renderer/Private/SceneRendering.cpp
--- a/Engine/Source/Runtime/Renderer/Private/SceneRendering.cpp
+++ b/Engine/Source/Runtime/Renderer/Private/SceneRendering.cpp
@@ -5125,11 +5125,7 @@ public:
 		ViewFamily.SetSceneRenderer(&Renderer);
 
 		Scene.UpdateAllPrimitiveSceneInfos(GraphBuilder);
-		GPUScene.BeginRender(GPUSceneDynamicContext);
-
-		FRDGExternalAccessQueue ExternalAccessQueue;
-		Scene.GPUScene.Update(GraphBuilder, Renderer.GetSceneUniforms(), ExternalAccessQueue);
-		ExternalAccessQueue.Submit(GraphBuilder);
+		GPUScene.BeginRender(GraphBuilder, GPUSceneDynamicContext);
 	}
 
 	virtual ~FScenePrimitiveRenderingContext()
diff --git a/Engine/Source/Runtime/Renderer/Private/SceneVisibility.cpp b/Engine/Source/Runtime/Renderer/Private/SceneVisibility.cpp
--- a/Engine/Source/Runtime/Renderer/Private/SceneVisibility.cpp
+++ b/Engine/Source/Runtime/Renderer/Private/SceneVisibility.cpp
@@ -5305,14 +5305,6 @@ void FDeferredShadingSceneRenderer::BeginInitViews(
 
 	LumenScenePDIVisualization();
 
-	// Process GPU scene prior to visibility to maximize overlap.
-	{
-		RDG_CSV_STAT_EXCLUSIVE_SCOPE(GraphBuilder, UpdateGPUScene);
-		RDG_GPU_STAT_SCOPE(GraphBuilder, GPUSceneUpdate);
-
-		Scene->GPUScene.Update(GraphBuilder, GetSceneUniforms(), ExternalAccessQueue, TaskDatas.VisibilityTaskData);
-	}
-
 	// This must happen before we start initialising and using views.
 	UpdateSkyIrradianceGpuBuffer(GraphBuilder, ViewFamily.EngineShowFlags, Scene->SkyLight, Scene->SkyIrradianceEnvironmentMap);
 
diff --git a/Engine/Source/Runtime/Renderer/Private/VT/RuntimeVirtualTextureRender.cpp b/Engine/Source/Runtime/Renderer/Private/VT/RuntimeVirtualTextureRender.cpp
--- a/Engine/Source/Runtime/Renderer/Private/VT/RuntimeVirtualTextureRender.cpp
+++ b/Engine/Source/Runtime/Renderer/Private/VT/RuntimeVirtualTextureRender.cpp
@@ -1563,23 +1563,19 @@ namespace RuntimeVirtualTexture
 
 	void RenderPagesStandAlone(FRDGBuilder& GraphBuilder, FRenderPageBatchDesc const& InDesc)
 	{
-		// This is required to collect dynamic primitives from the views (not used here, but we must provide one).
-		FGPUSceneDynamicContext GPUSceneDynamicContext(InDesc.Scene->GPUScene);
-
 		InDesc.Scene->UpdateAllPrimitiveSceneInfos(GraphBuilder);
 
+		// This is required to collect dynamic primitives from the views (not used here, but we must provide one).
+		FGPUSceneDynamicContext GPUSceneDynamicContext(InDesc.Scene->GPUScene);
 		// Call to let GPU-Scene determine if it is active and record scene primitive count
-		FGPUSceneScopeBeginEndHelper GPUSceneScopeBeginEndHelper(InDesc.Scene->GPUScene, GPUSceneDynamicContext);
+		FGPUSceneScopeBeginEndHelper GPUSceneScopeBeginEndHelper(GraphBuilder, InDesc.Scene->GPUScene, GPUSceneDynamicContext);
 
 		FSceneUniformBuffer SceneUB {};
-		FRDGExternalAccessQueue ExternalAccessQueue;
-		InDesc.Scene->GPUScene.Update(GraphBuilder, SceneUB, ExternalAccessQueue);
+		InDesc.Scene->GPUScene.FillSceneUniformBuffer(GraphBuilder, SceneUB);
 		if (InDesc.Scene->SplineMeshSceneResources)
 		{
 			InDesc.Scene->SplineMeshSceneResources->Update(GraphBuilder, SceneUB);
 		}
-		ExternalAccessQueue.Submit(GraphBuilder);
-
 		RenderPagesInternal(GraphBuilder, InDesc, SceneUB);
 	}
 
