StochasticShadows - added optional support for textured area lights under r.StochasticShadows.TexturedRectLights. Should be a part of tile classification in the future.

diff --git a/Engine/Shaders/Private/StochasticShadows/StochasticShadows.ush b/Engine/Shaders/Private/StochasticShadows/StochasticShadows.ush
--- a/Engine/Shaders/Private/StochasticShadows/StochasticShadows.ush
+++ b/Engine/Shaders/Private/StochasticShadows/StochasticShadows.ush
@@ -3,12 +3,11 @@
 #pragma once
 
 // Change this to force recompilation of all Stochastic Shadows shaders
-#pragma message("UESHADERMETADATA_VERSION 1A63AC4B-2335-2AD1-2CAB-82295C72AFFC")
+#pragma message("UESHADERMETADATA_VERSION 1A63AC4B-2335-2AD1-2CAB-82295C72CAFC")
 
 #define SUPPORT_CONTACT_SHADOWS 0
 #define USE_IES_PROFILE 0
 #define NON_DIRECTIONAL_DIRECT_LIGHTING 0
-#define USE_SOURCE_TEXTURE 0
 #include "../DeferredLightingCommon.ush"
 #include "../SceneData.ush"
 #include "../Hash.ush"
diff --git a/Engine/Shaders/Private/StochasticShadows/StochasticShadowsSampling.usf b/Engine/Shaders/Private/StochasticShadows/StochasticShadowsSampling.usf
--- a/Engine/Shaders/Private/StochasticShadows/StochasticShadowsSampling.usf
+++ b/Engine/Shaders/Private/StochasticShadows/StochasticShadowsSampling.usf
@@ -1,8 +1,6 @@
 // Copyright Epic Games, Inc. All Rights Reserved.
 
 #define STOCHASTIC_SHADOWS 1
-#define NON_DIRECTIONAL_DIRECT_LIGHTING 1
-#define SUPPORT_CONTACT_SHADOWS 0
 
 #ifndef NUM_SAMPLES_PER_PIXEL_1D
 	#error NUM_SAMPLES_PER_PIXEL_1D must be defined
diff --git a/Engine/Shaders/Private/StochasticShadows/StochasticShadowsShading.usf b/Engine/Shaders/Private/StochasticShadows/StochasticShadowsShading.usf
--- a/Engine/Shaders/Private/StochasticShadows/StochasticShadowsShading.usf
+++ b/Engine/Shaders/Private/StochasticShadows/StochasticShadowsShading.usf
@@ -2,9 +2,6 @@
 
 #define STOCHASTIC_SHADOWS 1
 
-#define NON_DIRECTIONAL_DIRECT_LIGHTING 1
-#define SUPPORT_CONTACT_SHADOWS 0
-
 #include "../Common.ush"
 #include "../BlueNoise.ush"
 #include "../ShaderPrint.ush"
diff --git a/Engine/Source/Runtime/Renderer/Private/StochasticShadows/StochasticShadows.cpp b/Engine/Source/Runtime/Renderer/Private/StochasticShadows/StochasticShadows.cpp
--- a/Engine/Source/Runtime/Renderer/Private/StochasticShadows/StochasticShadows.cpp
+++ b/Engine/Source/Runtime/Renderer/Private/StochasticShadows/StochasticShadows.cpp
@@ -114,13 +114,20 @@ static TAutoConsoleVariable<int> CVarStochasticShadowsCandidateLightMask(
 	ECVF_Scalability | ECVF_RenderThreadSafe
 );
 
+static TAutoConsoleVariable<int> CVarStochasticShadowsTexturedRectLights(
+	TEXT("r.StochasticShadows.TexturedRectLights"),
+	0,
+	TEXT("Whether to support textured rect lights."),
+	ECVF_Scalability | ECVF_RenderThreadSafe
+);
+
 namespace StochasticShadows
 {
 	// must match values in StochasticShadows.ush
 	constexpr int32 TileSize = 8;
 	constexpr int32 ShadowMaskTileSize = 8;	// Stored downsampled
 	constexpr int32 MaxLightSceneIdXY = 16; // 16 * 16 = 256
-	constexpr int32 MaxShadingTilesPerGridCell = 16;
+	constexpr int32 MaxShadingTilesPerGridCell = 32;
 	constexpr int32 ShadowMaskAtlasSizeInTiles = 512;
 	constexpr uint32 InvalidShadowMaskTileIndex = 0xFFFFFFFF;
 
@@ -314,12 +321,13 @@ class FGenerateSamplesCS : public FGlobalShader
 	END_SHADER_PARAMETER_STRUCT()
 
 	class FTileType : SHADER_PERMUTATION_INT("TILE_TYPE", (int32)StochasticShadows::ETileType::MAX);
+	class FTexturedRectLights : SHADER_PERMUTATION_BOOL("USE_SOURCE_TEXTURE");
 	class FNumSamplesPerPixel1d : SHADER_PERMUTATION_SPARSE_INT("NUM_SAMPLES_PER_PIXEL_1D", 1, 2, 4);
 	class FShadowMaskReprojectionWeights : SHADER_PERMUTATION_BOOL("SHADOW_MASK_REPROJECTION_WEIGHTS");
 	class FShadowFactorEstimate : SHADER_PERMUTATION_BOOL("SHADOW_FACTOR_ESTIMATE");
 	class FCandidateLightMask : SHADER_PERMUTATION_BOOL("CANDIDATE_LIGHT_MASK");
 	class FDebugMode : SHADER_PERMUTATION_BOOL("DEBUG_MODE");
-	using FPermutationDomain = TShaderPermutationDomain<FTileType, FNumSamplesPerPixel1d, FShadowMaskReprojectionWeights, FShadowFactorEstimate, FCandidateLightMask, FDebugMode>;
+	using FPermutationDomain = TShaderPermutationDomain<FTileType, FTexturedRectLights, FNumSamplesPerPixel1d, FShadowMaskReprojectionWeights, FShadowFactorEstimate, FCandidateLightMask, FDebugMode>;
 
 	static int32 GetGroupSize()
 	{	
@@ -500,8 +508,9 @@ class FShadeLightSamplesCS : public FGlobalShader
 	}
 
 	class FTileType : SHADER_PERMUTATION_INT("TILE_TYPE", (int32)StochasticShadows::ETileType::MAX);
+	class FTexturedRectLights : SHADER_PERMUTATION_BOOL("USE_SOURCE_TEXTURE");
 	class FDebugMode : SHADER_PERMUTATION_BOOL("DEBUG_MODE");
-	using FPermutationDomain = TShaderPermutationDomain<FTileType, FDebugMode>;
+	using FPermutationDomain = TShaderPermutationDomain<FTileType, FTexturedRectLights, FDebugMode>;
 
 	static bool ShouldCompilePermutation(const FGlobalShaderPermutationParameters& Parameters)
 	{
@@ -810,6 +819,7 @@ void FDeferredShadingSceneRenderer::RenderStochasticShadows(FRDGBuilder& GraphBu
 
 			FGenerateSamplesCS::FPermutationDomain PermutationVector;
 			PermutationVector.Set<FGenerateSamplesCS::FTileType>(TileType);
+			PermutationVector.Set<FGenerateSamplesCS::FTexturedRectLights>(CVarStochasticShadowsTexturedRectLights.GetValueOnRenderThread() != 0);
 			PermutationVector.Set<FGenerateSamplesCS::FNumSamplesPerPixel1d>(NumSamplesPerPixel2d.X * NumSamplesPerPixel2d.Y);
 			PermutationVector.Set<FGenerateSamplesCS::FShadowMaskReprojectionWeights>(bTemporal);
 			PermutationVector.Set<FGenerateSamplesCS::FShadowFactorEstimate>(bTemporal && CVarStochasticSamplingShadowEstimate.GetValueOnRenderThread() != 0);
@@ -979,6 +989,7 @@ void FDeferredShadingSceneRenderer::RenderStochasticShadows(FRDGBuilder& GraphBu
 
 			FShadeLightSamplesCS::FPermutationDomain PermutationVector;
 			PermutationVector.Set<FShadeLightSamplesCS::FTileType>(TileType);
+			PermutationVector.Set<FShadeLightSamplesCS::FTexturedRectLights>(CVarStochasticShadowsTexturedRectLights.GetValueOnRenderThread() != 0);
 			PermutationVector.Set<FShadeLightSamplesCS::FDebugMode>(bDebug);
 			auto ComputeShader = View.ShaderMap->GetShader<FShadeLightSamplesCS>(PermutationVector);
 
