Stochastic Shadows:
- Print CompositeShadowMaskTraces debug info for selected light or light ID specified using cvar.

diff --git a/Engine/Shaders/Private/StochasticShadows/StochasticShadowsComposite.usf b/Engine/Shaders/Private/StochasticShadows/StochasticShadowsComposite.usf
--- a/Engine/Shaders/Private/StochasticShadows/StochasticShadowsComposite.usf
+++ b/Engine/Shaders/Private/StochasticShadows/StochasticShadowsComposite.usf
@@ -16,6 +16,7 @@ Texture2D<UNORM float3> DownsampledSceneWorldNormal;
 uint TemporalMaxFramesAccumulated;
 float TemporalStdDevOffset;
 uint2 NumSamplesPerPixel;
+uint DebugLightId;
 
 RWTexture2D<uint> RWShadowMaskTileAtlas;
 
@@ -100,7 +101,7 @@ void CompositeShadowMaskTracesCS(
 
 #if DEBUG_MODE
 	int2 DebugScreenCoord = View.CursorPosition.x >= 0 ? View.CursorPosition * View.ViewResolutionFraction : -1;
-	bool bDebug = all(ShadowMaskScreenCoord == DebugScreenCoord / SHADOW_MASK_DOWNSAMPLE_FACTOR) && LightSceneId == 1;
+	bool bDebug = all(ShadowMaskScreenCoord == DebugScreenCoord / SHADOW_MASK_DOWNSAMPLE_FACTOR) && LightSceneId == DebugLightId;
 	FShaderPrintContext Context = InitShaderPrintContext(true, float2(0.05, 0.3));
 
 	if (bDebug)
@@ -268,4 +269,4 @@ void CompositeShadowMaskTracesCS(
 		uint2 ShadowMaskTileAtlasCoord = GetShadowMaskTileAtlasCoord(ShadowMaskTileIndex, LocalShadowMaskTileCoord);
 		RWShadowMaskTileAtlas[ShadowMaskTileAtlasCoord] = PackShadowMask(ShadowMask, Dither);
 	}
-}
\ No newline at end of file
+}
diff --git a/Engine/Source/Runtime/Renderer/Private/StochasticShadows/StochasticShadows.cpp b/Engine/Source/Runtime/Renderer/Private/StochasticShadows/StochasticShadows.cpp
--- a/Engine/Source/Runtime/Renderer/Private/StochasticShadows/StochasticShadows.cpp
+++ b/Engine/Source/Runtime/Renderer/Private/StochasticShadows/StochasticShadows.cpp
@@ -84,6 +84,13 @@ static TAutoConsoleVariable<int32> CVarStochasticShadowsDebug(
 	ECVF_RenderThreadSafe
 );
 
+static TAutoConsoleVariable<int32> CVarStochasticShadowsDebugLightId(
+	TEXT("r.StochasticShadows.Debug.LightId"),
+	-1,
+	TEXT("Which light to show debug info for. When set to -1, uses the currently selected light in editor."),
+	ECVF_RenderThreadSafe
+);
+
 int32 GStochasticShadowsReset = 0;
 FAutoConsoleVariableRef CVarStochasticShadowsReset(
 	TEXT("r.StochasticShadows.Reset"),
@@ -685,13 +692,31 @@ void FDeferredShadingSceneRenderer::RenderStochasticShadows(FRDGBuilder& GraphBu
 		StochasticShadowsParameters.DownsampledTileDataStride = DownsampledTileDataStride;
 		StochasticShadowsParameters.TemporalMaxFramesAccumulated = CVarStochasticShadowsTemporalMaxFramesAccumulated.GetValueOnRenderThread();
 		StochasticShadowsParameters.TemporalStdDevOffset = CVarStochasticShadowsTemporalStdDevOffset.GetValueOnRenderThread();
-		StochasticShadowsParameters.DebugMode = CVarStochasticShadowsDebug.GetValueOnRenderThread();
+		StochasticShadowsParameters.DebugMode = StochasticShadows::GetDebugMode();
+		StochasticShadowsParameters.DebugLightId = INDEX_NONE;
 
 		if (bDebug)
 		{
 			ShaderPrint::SetEnabled(true);
 			ShaderPrint::RequestSpaceForLines(1024);
 			ShaderPrint::SetParameters(GraphBuilder, View.ShaderPrintData, StochasticShadowsParameters.ShaderPrintUniformBuffer);
+
+			StochasticShadowsParameters.DebugLightId = CVarStochasticShadowsDebugLightId.GetValueOnRenderThread();
+
+			if (StochasticShadowsParameters.DebugLightId < 0)
+			{
+				for (auto LightIt = Scene->Lights.CreateConstIterator(); LightIt; ++LightIt)
+				{
+					const FLightSceneInfoCompact& LightSceneInfoCompact = *LightIt;
+					const FLightSceneInfo* const LightSceneInfo = LightSceneInfoCompact.LightSceneInfo;
+
+					if (LightSceneInfo->Proxy->IsSelected())
+					{
+						StochasticShadowsParameters.DebugLightId = LightSceneInfo->Id;
+						break;
+					}
+				}
+			}
 		}
 	}
 
diff --git a/Engine/Source/Runtime/Renderer/Private/StochasticShadows/StochasticShadowsInternal.h b/Engine/Source/Runtime/Renderer/Private/StochasticShadows/StochasticShadowsInternal.h
--- a/Engine/Source/Runtime/Renderer/Private/StochasticShadows/StochasticShadowsInternal.h
+++ b/Engine/Source/Runtime/Renderer/Private/StochasticShadows/StochasticShadowsInternal.h
@@ -32,6 +32,7 @@ BEGIN_SHADER_PARAMETER_STRUCT(FStochasticShadowsParameters, )
 	SHADER_PARAMETER(int32, TemporalMaxFramesAccumulated)
 	SHADER_PARAMETER(float, TemporalStdDevOffset)
 	SHADER_PARAMETER(int32, DebugMode)
+	SHADER_PARAMETER(int32, DebugLightId)
 	SHADER_PARAMETER_RDG_TEXTURE(Texture2D<uint>, DownsampledTileMask)
 	SHADER_PARAMETER_RDG_TEXTURE(Texture2D<float>, DownsampledSceneDepth)
 	SHADER_PARAMETER_RDG_TEXTURE(Texture2D<UNORM float3>, DownsampledSceneWorldNormal)
