Removes a TArray copy on TSR's SelectWaveSize

diff --git a/Engine/Source/Runtime/Renderer/Private/PostProcess/TemporalSuperResolution.cpp b/Engine/Source/Runtime/Renderer/Private/PostProcess/TemporalSuperResolution.cpp
--- a/Engine/Source/Runtime/Renderer/Private/PostProcess/TemporalSuperResolution.cpp
+++ b/Engine/Source/Runtime/Renderer/Private/PostProcess/TemporalSuperResolution.cpp
@@ -1488,7 +1488,7 @@ FDefaultTemporalUpscaler::FOutputs AddTemporalSuperResolutionPasses(
 		return GraphBuilder.CreateUAV(DebugTexture);
 	};
 
-	auto SelectWaveSize = [&](const TArray<int32> WaveSizeDomain)
+	auto SelectWaveSize = [&](const TArray<int32>& WaveSizeDomain)
 	{
 		check(!WaveSizeDomain.IsEmpty());
 		int32 WaveSizeOps = 0;
