Unshelved from pending changelist '29941045':

Heterogeneous Volumes: Correct live-shading compositing issues when disabling the lighting cache.

diff --git a/Engine/Source/Runtime/Renderer/Private/HeterogeneousVolumes/HeterogeneousVolumesLiveShadingPipeline.cpp b/Engine/Source/Runtime/Renderer/Private/HeterogeneousVolumes/HeterogeneousVolumesLiveShadingPipeline.cpp
--- a/Engine/Source/Runtime/Renderer/Private/HeterogeneousVolumes/HeterogeneousVolumesLiveShadingPipeline.cpp
+++ b/Engine/Source/Runtime/Renderer/Private/HeterogeneousVolumes/HeterogeneousVolumesLiveShadingPipeline.cpp
@@ -811,7 +811,7 @@ static void RenderWithInscatteringVolumePipeline(
 	int32 NumPasses = LightSceneInfoCompact.Num();
 	for (int32 PassIndex = 0; PassIndex < NumPasses; ++PassIndex)
 	{
-		bool bApplyEmissionAndTransmittance = PassIndex == 0;
+		bool bApplyEmissionAndTransmittance = (PassIndex == (NumPasses - 1));
 		bool bApplyDirectLighting = !LightSceneInfoCompact.IsEmpty();
 		bool bApplyShadowTransmittance = false;
 
