Fix checks that incorrectly fired when GPU-Scene is disabled.

diff --git a/Engine/Source/Runtime/Renderer/Private/GPUScene.cpp b/Engine/Source/Runtime/Renderer/Private/GPUScene.cpp
--- a/Engine/Source/Runtime/Renderer/Private/GPUScene.cpp
+++ b/Engine/Source/Runtime/Renderer/Private/GPUScene.cpp
@@ -654,24 +654,27 @@ void FGPUScene::BeginRender(FRDGBuilder& GraphBuilder, FGPUSceneDynamicContext &
 {
 	ensure(!bInBeginEndBlock);
 	ensure(CurrentDynamicContext == nullptr);
-	check(NumScenePrimitives == Scene.Primitives.Num());
-	// Should always be reset to this as the neutral state.
-	check(DynamicPrimitivesOffset == Scene.GetMaxPersistentPrimitiveIndex());
 
-	// Do it anyway for old times sake
-	DynamicPrimitivesOffset = Scene.GetMaxPersistentPrimitiveIndex();
 	CurrentDynamicContext = &GPUSceneDynamicContext;
 	bInBeginEndBlock = true;
-
 	check(!CachedRegisteredBuffers.IsValid());
-	CachedRegisteredBuffers = RegisterBuffers(GraphBuilder);
+	if (bIsEnabled)
+	{
+		check(NumScenePrimitives == Scene.Primitives.Num());
+		// Should always be reset to this as the neutral state.
+		check(DynamicPrimitivesOffset == Scene.GetMaxPersistentPrimitiveIndex());
+		// Do it anyway for old times sake
+		DynamicPrimitivesOffset = Scene.GetMaxPersistentPrimitiveIndex();
+
+		CachedRegisteredBuffers = RegisterBuffers(GraphBuilder);
+	}
 }
 
 void FGPUScene::EndRender()
 {
 	ensure(bInBeginEndBlock);
 	ensure(CurrentDynamicContext != nullptr);
-	check(DynamicPrimitivesOffset >= NumScenePrimitives);
+	check(!bIsEnabled || DynamicPrimitivesOffset >= NumScenePrimitives);
 
 	// Pop all dynamic primitives off the stack
 	DynamicPrimitivesOffset = Scene.GetMaxPersistentPrimitiveIndex();
