Kickoff ray tracing geometry builds early in the frame (before ray tracing instance gathering).
- also avoids one frame delay in including instances in raytracing scene due to logic checking HasPendingBuildRequest() to early out in GetStaticRayTracingGeometryInstance(...).
- for situations where a build is requested during ray tracing instance gathering and must be completed in the same frame, the correct way to handle that is to add those requests to RayTracingScene.GeometriesToBuild which are then processed using GRayTracingGeometryManager->ForceBuildIfPending(...) which skips the r.RayTracing.Geometry.MaxBuiltPrimitivesPerFrame limit.

diff --git a/Engine/Source/Runtime/Renderer/Private/DeferredShadingRenderer.cpp b/Engine/Source/Runtime/Renderer/Private/DeferredShadingRenderer.cpp
--- a/Engine/Source/Runtime/Renderer/Private/DeferredShadingRenderer.cpp
+++ b/Engine/Source/Runtime/Renderer/Private/DeferredShadingRenderer.cpp
@@ -2004,9 +2004,6 @@ bool FDeferredShadingSceneRenderer::DispatchRayTracingWorldUpdates(FRDGBuilder&
 	// We only need to update ray tracing scene for the first view family, if multiple are rendered in a single scene render call.
 	if (!bShouldUpdateRayTracingScene)
 	{
-		// This needs to happen even when ray tracing is not enabled
-		// - importers might batch BVH creation requests that need to be resolved in any case
-		GRayTracingGeometryManager->ProcessBuildRequests(GraphBuilder.RHICmdList);
 		// - Nanite ray tracing instances are already pointing at the new BLASes and RayTracingDataOffsets in GPUScene have been updated
 		Nanite::GRayTracingManager.ProcessBuildRequests(GraphBuilder);
 		return false;
@@ -2025,8 +2022,6 @@ bool FDeferredShadingSceneRenderer::DispatchRayTracingWorldUpdates(FRDGBuilder&
 		RayTracingSkinnedGeometryUpdateQueue->Commit(GraphBuilder);
 	}
 
-	GRayTracingGeometryManager->ProcessBuildRequests(GraphBuilder.RHICmdList);
-
 	const int32 ReferenceViewIndex = 0;
 	FViewInfo& ReferenceView = Views[ReferenceViewIndex];
 	FRayTracingScene& RayTracingScene = Scene->RayTracingScene;
@@ -2670,6 +2665,14 @@ void FDeferredShadingSceneRenderer::Render(FRDGBuilder& GraphBuilder)
 
 	GPU_MESSAGE_SCOPE(GraphBuilder);
 
+#if RHI_RAYTRACING
+	if (RendererOutput != FSceneRenderer::ERendererOutput::DepthPrepassOnly)
+	{
+		// TODO: should only process build requests once per frame
+		GRayTracingGeometryManager->ProcessBuildRequests(GraphBuilder.RHICmdList);
+	}
+#endif
+
 	FInitViewTaskDatas InitViewTaskDatas = OnRenderBegin(GraphBuilder);
 
 	FRDGExternalAccessQueue ExternalAccessQueue;
diff --git a/Engine/Source/Runtime/Renderer/Private/PrimitiveSceneInfo.cpp b/Engine/Source/Runtime/Renderer/Private/PrimitiveSceneInfo.cpp
--- a/Engine/Source/Runtime/Renderer/Private/PrimitiveSceneInfo.cpp
+++ b/Engine/Source/Runtime/Renderer/Private/PrimitiveSceneInfo.cpp
@@ -386,6 +386,8 @@ bool FPrimitiveSceneInfo::IsCachedRayTracingGeometryValid() const
 	{
 		// TODO: Doesn't take Nanite Ray Tracing into account
 		check(CachedRayTracingGeometry->RayTracingGeometryRHI == CachedRayTracingInstance.GeometryRHI);
+		check(!CachedRayTracingGeometry->GetRequiresBuild() && !CachedRayTracingGeometry->HasPendingBuildRequest());
+
 		return CachedRayTracingGeometry->IsValid();
 	}
 
