Only allocate ShadingTileAtlas space for tiles shadowed regions (not fully lit).
- if shading tile is unshadowed, mark tile with out-of-range ShadingTileIndex which causes ShadeLightSamplesCS to assume ShadowFactor = 1.
- added cvar to control MaxShadingTilesPerGridCell.

diff --git a/Engine/Shaders/Private/StochasticShadows/StochasticShadowsShading.usf b/Engine/Shaders/Private/StochasticShadows/StochasticShadowsShading.usf
--- a/Engine/Shaders/Private/StochasticShadows/StochasticShadowsShading.usf
+++ b/Engine/Shaders/Private/StochasticShadows/StochasticShadowsShading.usf
@@ -94,9 +94,14 @@ void ShadeLightSamplesCS(
 					const FShadingTile ShadingTile = UnpackShadingTile(PackedShadingTile);
 					FDeferredLightData LightData = ConvertToDeferredLight(GetLightSceneData(ShadingTile.LightSceneId));
 
-					uint2 ShadingTileAtlasCoord = GetShadingTileAtlasCoord(ShadingTile.ShadingTileIndex, CoordInShadingTile);
+					float ShadowFactor = 1.0f; // Initialize as unshadowed as tiles with index >= MaxShadingTiles are treated as fully unshadowed.
 
-					float ShadowFactor = ShadingTileAtlas[ShadingTileAtlasCoord];
+					if (ShadingTile.ShadingTileIndex < MaxShadingTiles)
+					{
+						uint2 ShadingTileAtlasCoord = GetShadingTileAtlasCoord(ShadingTile.ShadingTileIndex, CoordInShadingTile);
+
+						ShadowFactor = ShadingTileAtlas[ShadingTileAtlasCoord];
+					}
 
 					#if DEBUG_MODE
 					if (bDebug)
diff --git a/Engine/Shaders/Private/StochasticShadows/StochasticShadowsSpatial.usf b/Engine/Shaders/Private/StochasticShadows/StochasticShadowsSpatial.usf
--- a/Engine/Shaders/Private/StochasticShadows/StochasticShadowsSpatial.usf
+++ b/Engine/Shaders/Private/StochasticShadows/StochasticShadowsSpatial.usf
@@ -30,6 +30,38 @@ groupshared float SharedSceneDepth[SHARED_TILE_SIZE][SHARED_TILE_SIZE];
 groupshared uint SharedShadingTileMask;
 groupshared uint SharedShadingTileIndex[2][2];
 
+uint GetLitBitMask(uint2 ShadingTileCoord)
+{
+	const uint BitOffset = (ShadingTileCoord.x + ShadingTileCoord.y * 2);
+	return 0x1 << BitOffset;
+}
+
+uint GetShadowedBitMask(uint2 ShadingTileCoord)
+{
+	const uint BitOffset = (ShadingTileCoord.x + ShadingTileCoord.y * 2);
+	return 0x10 << BitOffset;
+}
+
+void SetSharedShadingTileMask(uint2 ShadingTileCoord, FShadowMask ShadowMask)
+{
+	uint BitMask = 0;
+
+	if (ShadowMask.IsValid() && ShadowMask.Factor > 0.0f) // lit
+	{
+		BitMask |= GetLitBitMask(ShadingTileCoord);
+	}
+
+	if (!ShadowMask.IsValid() || ShadowMask.Factor < 0.999f) // shadowed
+	{
+		BitMask |= GetShadowedBitMask(ShadingTileCoord);
+	}
+
+	if (BitMask != 0)
+	{
+		InterlockedOr(SharedShadingTileMask, BitMask);
+	}
+}
+
 void AccumulateShadowMask(uint2 SharedShadowMaskTileCoord, float UpsampleWeight, inout FShadowMaskAccumulator Accumulator)
 {
 	if (UpsampleWeight > 0.01f)
@@ -119,7 +151,6 @@ void ShadowMaskSpatialPassCS(
 	// Spatial pass
 	if (all(ShadowMaskScreenCoord < ShadowMaskViewSize))
 	{
-		uint2 ShadowMaskTileAtlasCoord = GetShadowMaskTileAtlasCoord(ShadowMaskTileIndex, LocalShadowMaskTileCoord);
 		FShadowMaskAccumulator ShadowMaskAccumulator = InitShadowMaskAccumulator();
 
 		uint2 CenterSampleSharedTileCoord = LocalShadowMaskTileCoord + SHARED_TILE_BORDER;
@@ -166,7 +197,7 @@ void ShadowMaskSpatialPassCS(
 						Newline(Context);
 						Print(Context, TEXT("SampleShadowMask: "));
 						Print(Context, SampleShadowMask.Factor);
-						Print(Context, GBuffer.Depth);
+						Print(Context, SampleDepth);
 					}
 					#endif
 				}
@@ -197,24 +228,14 @@ void ShadowMaskSpatialPassCS(
 		// Update shared shadow mask tile with filtered results
 		SharedShadowMask[GroupThreadId.x + SHARED_TILE_BORDER][GroupThreadId.y + SHARED_TILE_BORDER] = PackShadowMask(ShadowMask, /*Dither*/ 0.5f);
 
-		if (ShadowMask.IsValid() && ShadowMask.Factor > 0.0f)
-		{
-			uint2 BitCoord = LocalShadowMaskTileCoord / 4;
-			uint BitMask = 1u << (BitCoord.x + BitCoord.y * 2);
-			InterlockedOr(SharedShadingTileMask, BitMask);
-		}
+		SetSharedShadingTileMask(LocalShadowMaskTileCoord / 4, ShadowMask);
 	}
 #else
 	// Only mark sub-tiles requiring shading tiles
 	{
 		FShadowMask ShadowMask = UnpackShadowMask(SharedShadowMask[GroupThreadId.x + SHARED_TILE_BORDER][GroupThreadId.y + SHARED_TILE_BORDER]);
 
-		if (ShadowMask.IsValid() && ShadowMask.Factor > 0.0f)
-		{
-			uint2 BitCoord = LocalShadowMaskTileCoord / 4;
-			uint BitMask = 1u << (BitCoord.x + BitCoord.y * 2);
-			InterlockedOr(SharedShadingTileMask, BitMask);
-		}
+		SetSharedShadingTileMask(LocalShadowMaskTileCoord / 4, ShadowMask);
 	}
 #endif
 
@@ -225,17 +246,22 @@ void ShadowMaskSpatialPassCS(
 	// Allocate shading tiles
 	if (all(ShadingTileOffset < 2))
 	{
-		uint BitMask = 1u << (ShadingTileOffset.x + ShadingTileOffset.y * 2);
-		if (BitMask & SharedShadingTileMask)
+		if (SharedShadingTileMask & GetLitBitMask(ShadingTileOffset))
 		{	
 			const uint2 ShadingTileCoord = (BaseShadowMaskScreenCoord * SHADOW_MASK_DOWNSAMPLE_FACTOR) / SHADOW_MASK_TILE_SIZE + ShadingTileOffset;
 
 			if (all(ShadingTileCoord < ShadingTileGridSize))
 			{
-				uint ShadingTileIndex = 0;
-				InterlockedAdd(RWShadingTileAllocator[0], 1u, ShadingTileIndex);
+				const uint UnshadowedTileIndex = 0xFFFFFFFF;
+
+				uint ShadingTileIndex = UnshadowedTileIndex;
+
+				if(SharedShadingTileMask & GetShadowedBitMask(ShadingTileOffset))
+				{
+					InterlockedAdd(RWShadingTileAllocator[0], 1u, ShadingTileIndex);
+				}
 
-				if (ShadingTileIndex < MaxShadingTiles)
+				if (ShadingTileIndex < MaxShadingTiles || ShadingTileIndex == UnshadowedTileIndex)
 				{
 					uint ShadingTileIndexInCell = 0;
 					InterlockedAdd(RWShadingTileGridAllocator[ShadingTileCoord], 1u, ShadingTileIndexInCell);
diff --git a/Engine/Source/Runtime/Renderer/Private/StochasticShadows/StochasticShadows.cpp b/Engine/Source/Runtime/Renderer/Private/StochasticShadows/StochasticShadows.cpp
--- a/Engine/Source/Runtime/Renderer/Private/StochasticShadows/StochasticShadows.cpp
+++ b/Engine/Source/Runtime/Renderer/Private/StochasticShadows/StochasticShadows.cpp
@@ -25,6 +25,13 @@ static TAutoConsoleVariable<int32> CVarStochasticShadowsNumSamplesPerPixel(
 	ECVF_Scalability | ECVF_RenderThreadSafe
 );
 
+static TAutoConsoleVariable<int32> CVarStochasticShadowsMaxShadingTilesPerGridCell(
+	TEXT("r.StochasticShadows.MaxShadingTilesPerGridCell"),
+	32,
+	TEXT("Maximum number of shading tiles per grid cell."),
+	ECVF_Scalability | ECVF_RenderThreadSafe
+);
+
 static TAutoConsoleVariable<int32> CVarStochasticShadowsTemporal(
 	TEXT("r.StochasticShadows.Temporal"),
 	1,
@@ -134,9 +141,9 @@ namespace StochasticShadows
 	constexpr int32 TileSize = 8;
 	constexpr int32 ShadowMaskTileSize = 8;	// Stored downsampled
 	constexpr int32 MaxLightSceneIdXY = 16; // 16 * 16 = 256
-	constexpr int32 MaxShadingTilesPerGridCell = 32;
 	constexpr int32 ShadowMaskAtlasSizeInTiles = 512;
 	constexpr uint32 InvalidShadowMaskTileIndex = 0xFFFFFFFF;
+	constexpr uint32 ShadingTileIndexUnshadowed = 0xFFFFF; // limited by PackShadingTile()
 
 	bool IsEnabled()
 	{
@@ -586,9 +593,15 @@ void FDeferredShadingSceneRenderer::RenderStochasticShadows(FRDGBuilder& GraphBu
 	const FIntPoint ShadowMaskPageTableSize = ShadowMaskPageTablePerLightSize * StochasticShadows::MaxLightSceneIdXY;
 	const FIntPoint ShadowMaskTileAtlasSize = StochasticShadows::ShadowMaskTileSize * StochasticShadows::ShadowMaskAtlasSizeInTiles;
 	const int32 MaxShadowMaskTiles = StochasticShadows::ShadowMaskAtlasSizeInTiles * StochasticShadows::ShadowMaskAtlasSizeInTiles;
+	check(MaxShadowMaskTiles == (ShadowMaskTileAtlasSize.X * ShadowMaskTileAtlasSize.Y) / (StochasticShadows::ShadowMaskTileSize * StochasticShadows::ShadowMaskTileSize));
 
 	const FIntPoint ShadingTileGridSize = FIntPoint::DivideAndRoundUp(SceneTextures.Config.Extent, StochasticShadows::ShadowMaskTileSize);
 	const FIntPoint ShadingTileAtlasSize = ShadowMaskTileAtlasSize;
+	const int32 MaxShadingTiles = MaxShadowMaskTiles;
+	check(MaxShadingTiles == (ShadingTileAtlasSize.X * ShadingTileAtlasSize.Y) / (StochasticShadows::ShadowMaskTileSize * StochasticShadows::ShadowMaskTileSize));
+	check(MaxShadingTiles < StochasticShadows::ShadingTileIndexUnshadowed);
+
+	const int32 MaxShadingTilesPerGridCell = FMath::Max(CVarStochasticShadowsMaxShadingTilesPerGridCell.GetValueOnRenderThread(), 0);
 
 	FRDGTextureRef DownsampledSceneDepth = GraphBuilder.CreateTexture(
 		FRDGTextureDesc::Create2D(DownsampledBufferSize, PF_R32_FLOAT, FClearValueBinding::Black, TexCreate_ShaderResource | TexCreate_UAV),
@@ -682,8 +695,8 @@ void FDeferredShadingSceneRenderer::RenderStochasticShadows(FRDGBuilder& GraphBu
 		StochasticShadowsParameters.DownsampledSceneDepth = DownsampledSceneDepth;
 		StochasticShadowsParameters.DownsampledSceneWorldNormal = DownsampledSceneWorldNormal;
 		StochasticShadowsParameters.MaxShadowMaskTiles = MaxShadowMaskTiles;
-		StochasticShadowsParameters.MaxShadingTiles = (ShadingTileAtlasSize.X * ShadingTileAtlasSize.Y) / (StochasticShadows::ShadowMaskTileSize * StochasticShadows::ShadowMaskTileSize);
-		StochasticShadowsParameters.MaxShadingTilesPerGridCell = StochasticShadows::MaxShadingTilesPerGridCell;
+		StochasticShadowsParameters.MaxShadingTiles = MaxShadingTiles;
+		StochasticShadowsParameters.MaxShadingTilesPerGridCell = MaxShadingTilesPerGridCell;
 		StochasticShadowsParameters.ShadingTileGridSize = ShadingTileGridSize;
 		StochasticShadowsParameters.ShadowMaskPageTablePerLightSize = ShadowMaskPageTablePerLightSize;
 		StochasticShadowsParameters.DownsampledBufferInvSize = FVector2f(1.0f) / DownsampledBufferSize;
@@ -926,7 +939,7 @@ void FDeferredShadingSceneRenderer::RenderStochasticShadows(FRDGBuilder& GraphBu
 		TEXT("StochasticShadows.ShadingTileGridAllocator"));
 
 	FRDGBufferRef ShadingTileGrid = GraphBuilder.CreateBuffer(
-		FRDGBufferDesc::CreateStructuredDesc(sizeof(uint32), ShadingTileGridSize.X * ShadingTileGridSize.Y * StochasticShadows::MaxShadingTilesPerGridCell),
+		FRDGBufferDesc::CreateStructuredDesc(sizeof(uint32), ShadingTileGridSize.X * ShadingTileGridSize.Y * MaxShadingTilesPerGridCell),
 		TEXT("StochasticShadowsParameters.ShadingTileGrid"));
 
 	FRDGTextureRef ShadingTileAtlas = GraphBuilder.CreateTexture(
