SLW - Fix for clouds not composed correctly in Forward.

diff --git a/Engine/Shaders/Private/VolumetricRenderTarget.usf b/Engine/Shaders/Private/VolumetricRenderTarget.usf
--- a/Engine/Shaders/Private/VolumetricRenderTarget.usf
+++ b/Engine/Shaders/Private/VolumetricRenderTarget.usf
@@ -455,8 +455,9 @@ Texture2D<float4> VolumetricTexture;
 Texture2D<float2> VolumetricDepthTexture;
 uint4 VolumetricTextureValidCoordRect;
 float4 VolumetricTextureValidUvRect;
+uint ForwardShadingEnable;
 
-#if PERMUTATION_RENDER_UNDERWATER_BUFFER || PERMUTATION_RENDER_CAMERA_UNDERWATER
+#if PERMUTATION_RENDER_UNDERWATER_BUFFER || PERMUTATION_COMPOSE_WITH_WATER
 Texture2D WaterLinearDepthTexture;
 SamplerState WaterLinearDepthSampler;
 float4 SceneWithoutSingleLayerWaterViewRect;
@@ -610,7 +611,7 @@ void ComposeVolumetricRTOverScenePS(
 
 
 
-#if PERMUTATION_RENDER_CAMERA_UNDERWATER
+#if PERMUTATION_COMPOSE_WITH_WATER
 	// Now check that we are compositing a pixel that is not "water" to avoid applying clouds twice (they are already composited with the behind water layer scene).
 	// We also lack depth information behind the water surface now so the composition would be wrong anyway.
 	float WaterTestPixelLinearDepth = ConvertFromDeviceZ(WaterLinearDepthTexture.SampleLevel(WaterLinearDepthSampler, ScreenUV, 0).r);
@@ -766,8 +767,9 @@ void ComposeVolumetricRTOverScenePS(
 #else // PERMUTATION_RENDER_UNDERWATER_BUFFER
 
 	// We only want to run the special "cloud over water" code path on actual water pixels, otherwise depth threshold can cause visual issues.
+	// However, with Forward Shading, we do not have access to the gbuffer: so we still run that simple water depth test to occlude cloud on forrground meshes. Not the perfect but it will work in most cases.
 	FScreenSpaceData ScreenSpaceData = GetScreenSpaceData(ScreenUV);
-	const bool bIsWaterPixel = ScreenSpaceData.GBuffer.ShadingModelID == SHADINGMODELID_SINGLELAYERWATER;
+	const bool bIsWaterPixel = ForwardShadingEnable || ScreenSpaceData.GBuffer.ShadingModelID == SHADINGMODELID_SINGLELAYERWATER;
 
 	uint2 VolumeCoord0 = max(0, int2(VolumeCoordUInt) + int2(0, 0));
 	uint2 VolumeCoord1 = max(0, int2(VolumeCoordUInt) + int2(OffsetX, 0));
diff --git a/Engine/Source/Runtime/Renderer/Private/DeferredShadingRenderer.cpp b/Engine/Source/Runtime/Renderer/Private/DeferredShadingRenderer.cpp
--- a/Engine/Source/Runtime/Renderer/Private/DeferredShadingRenderer.cpp
+++ b/Engine/Source/Runtime/Renderer/Private/DeferredShadingRenderer.cpp
@@ -4169,9 +4169,10 @@ void FDeferredShadingSceneRenderer::Render(FRDGBuilder& GraphBuilder)
 			// Or composite the off screen buffer over the scene.
 			if (bVolumetricRenderTargetRequired)
 			{
+				const bool bComposeWithWater = bIsCameraUnderWater ? false : bShouldRenderSingleLayerWater;
 				ComposeVolumetricRenderTargetOverScene(
 					GraphBuilder, Views, SceneTextures.Color.Target, SceneTextures.Depth.Target,
-					bIsCameraUnderWater ? false : bShouldRenderSingleLayerWater,
+					bComposeWithWater,
 					SceneWithoutWaterTextures, SceneTextures);
 			}
 		};
diff --git a/Engine/Source/Runtime/Renderer/Private/VolumetricRenderTarget.cpp b/Engine/Source/Runtime/Renderer/Private/VolumetricRenderTarget.cpp
--- a/Engine/Source/Runtime/Renderer/Private/VolumetricRenderTarget.cpp
+++ b/Engine/Source/Runtime/Renderer/Private/VolumetricRenderTarget.cpp
@@ -653,9 +653,9 @@ class FComposeVolumetricRTOverScenePS : public FGlobalShader
 
 	class FUpsamplingMode : SHADER_PERMUTATION_RANGE_INT("PERMUTATION_UPSAMPLINGMODE", 0, 5);
 	class FRenderUnderWaterBuffer : SHADER_PERMUTATION_BOOL("PERMUTATION_RENDER_UNDERWATER_BUFFER");	// Render into the water scene color buffer (used when rendering from water system)
-	class FRenderCameraUnderWater : SHADER_PERMUTATION_BOOL("PERMUTATION_RENDER_CAMERA_UNDERWATER");	// When water us used and the camera is under water, use that permutation (to handle camera intersection with water and double cloud composition)
+	class FRenderCameraComposeWithWater : SHADER_PERMUTATION_BOOL("PERMUTATION_COMPOSE_WITH_WATER");	// When water us used and the camera is under water, use that permutation (to handle camera intersection with water and double cloud composition)
 	class FMSAASampleCount : SHADER_PERMUTATION_SPARSE_INT("MSAA_SAMPLE_COUNT", 1, 2, 4, 8);
-	using FPermutationDomain = TShaderPermutationDomain<FUpsamplingMode, FRenderUnderWaterBuffer, FRenderCameraUnderWater, FMSAASampleCount>;
+	using FPermutationDomain = TShaderPermutationDomain<FUpsamplingMode, FRenderUnderWaterBuffer, FRenderCameraComposeWithWater, FMSAASampleCount>;
 
 	BEGIN_SHADER_PARAMETER_STRUCT(FParameters, )
 		SHADER_PARAMETER_STRUCT_REF(FViewUniformShaderParameters, ViewUniformBuffer)
@@ -666,6 +666,7 @@ class FComposeVolumetricRTOverScenePS : public FGlobalShader
 		SHADER_PARAMETER_SAMPLER(SamplerState, LinearTextureSampler)
 		SHADER_PARAMETER_SAMPLER(SamplerState, WaterLinearDepthSampler)
 		RENDER_TARGET_BINDING_SLOTS()
+		SHADER_PARAMETER(uint32, ForwardShadingEnable)
 		SHADER_PARAMETER(float, UvOffsetSampleAcceptanceWeight)
 		SHADER_PARAMETER(FVector4f, VolumetricTextureSizeAndInvSize)
 		SHADER_PARAMETER(FVector2f, FullResolutionToVolumetricBufferResolutionScale)
@@ -741,7 +742,7 @@ void ComposeVolumetricRenderTargetOverScene(
 	TArrayView<FViewInfo> Views,
 	FRDGTextureRef SceneColorTexture,
 	FRDGTextureRef SceneDepthTexture,
-	bool bShouldRenderSingleLayerWater,
+	bool bComposeWithWater,
 	const FSceneWithoutWaterTextures& WaterPassData,
 	const FMinimalSceneTextures& SceneTextures)
 {
@@ -782,7 +783,7 @@ void ComposeVolumetricRenderTargetOverScene(
 		FComposeVolumetricRTOverScenePS::FPermutationDomain PermutationVector;
 		PermutationVector.Set<FComposeVolumetricRTOverScenePS::FUpsamplingMode>(UpsamplingMode);
 		PermutationVector.Set<FComposeVolumetricRTOverScenePS::FRenderUnderWaterBuffer>(0);
-		PermutationVector.Set<FComposeVolumetricRTOverScenePS::FRenderCameraUnderWater>((bShouldRenderSingleLayerWater && ViewInfo.IsUnderwater()) ? 1 : 0);
+		PermutationVector.Set<FComposeVolumetricRTOverScenePS::FRenderCameraComposeWithWater>((bComposeWithWater && ViewInfo.IsUnderwater()) ? 1 : 0);
 		PermutationVector.Set<FComposeVolumetricRTOverScenePS::FMSAASampleCount>(SceneDepthTexture->Desc.NumSamples);
 		TShaderMapRef<FComposeVolumetricRTOverScenePS> PixelShader(ViewInfo.ShaderMap, PermutationVector);
 
@@ -796,11 +797,12 @@ void ComposeVolumetricRenderTargetOverScene(
 		PassParameters->UvOffsetSampleAcceptanceWeight = GetUvNoiseSampleAcceptanceWeight();
 		PassParameters->FullResolutionToVolumetricBufferResolutionScale = GetCompositionFullResolutionToVolumetricBufferResolutionScale(VRTMode);
 		PassParameters->SceneTextures = SceneTextures.UniformBuffer;
+		PassParameters->ForwardShadingEnable = bForwardShading ? 1 : 0;
 		GetTextureSafeUvCoordBound(PassParameters->VolumetricTexture, PassParameters->VolumetricTextureValidCoordRect, PassParameters->VolumetricTextureValidUvRect);
 
 		PassParameters->WaterLinearDepthTexture = WaterPassData.DepthTexture;
 		PassParameters->WaterLinearDepthSampler = TStaticSamplerState<SF_Point>::GetRHI();
-		if (bShouldRenderSingleLayerWater)
+		if (bComposeWithWater)
 		{
 			const FSceneWithoutWaterTextures::FView& WaterPassViewData = WaterPassData.Views[ViewIndex];
 			PassParameters->FullResolutionToWaterBufferScale = FVector2f(1.0f / WaterPassData.RefractionDownsampleFactor, WaterPassData.RefractionDownsampleFactor);
@@ -858,7 +860,7 @@ void ComposeVolumetricRenderTargetOverSceneUnderWater(
 		FComposeVolumetricRTOverScenePS::FPermutationDomain PermutationVector;
 		PermutationVector.Set<FComposeVolumetricRTOverScenePS::FUpsamplingMode>(UpsamplingMode);
 		PermutationVector.Set<FComposeVolumetricRTOverScenePS::FRenderUnderWaterBuffer>(1);
-		PermutationVector.Set<FComposeVolumetricRTOverScenePS::FRenderCameraUnderWater>(0);
+		PermutationVector.Set<FComposeVolumetricRTOverScenePS::FRenderCameraComposeWithWater>(0);
 		PermutationVector.Set<FComposeVolumetricRTOverScenePS::FMSAASampleCount>(1);
 		TShaderMapRef<FComposeVolumetricRTOverScenePS> PixelShader(ViewInfo.ShaderMap, PermutationVector);
 
@@ -924,7 +926,7 @@ void ComposeVolumetricRenderTargetOverSceneForVisualization(
 		FComposeVolumetricRTOverScenePS::FPermutationDomain PermutationVector;
 		PermutationVector.Set<FComposeVolumetricRTOverScenePS::FUpsamplingMode>(0);
 		PermutationVector.Set<FComposeVolumetricRTOverScenePS::FRenderUnderWaterBuffer>(0);
-		PermutationVector.Set<FComposeVolumetricRTOverScenePS::FRenderCameraUnderWater>(0);
+		PermutationVector.Set<FComposeVolumetricRTOverScenePS::FRenderCameraComposeWithWater>(0);
 		PermutationVector.Set<FComposeVolumetricRTOverScenePS::FMSAASampleCount>(1);
 		TShaderMapRef<FComposeVolumetricRTOverScenePS> PixelShader(ViewInfo.ShaderMap, PermutationVector);
 
