Added TPersistentByteAddressBuffer and TByteAddressBufferScatterUploader alongside the structured buffer support.

diff --git a/Engine/Source/Runtime/RenderCore/Public/RenderGraphUtils.h b/Engine/Source/Runtime/RenderCore/Public/RenderGraphUtils.h
--- a/Engine/Source/Runtime/RenderCore/Public/RenderGraphUtils.h
+++ b/Engine/Source/Runtime/RenderCore/Public/RenderGraphUtils.h
@@ -1033,6 +1033,21 @@ FORCEINLINE FRDGBufferRef CreateStructuredUploadBuffer(
 	return Buffer;
 }
 
+/**
+ * Helper to create a byte address upload buffer with initial data from a TArray.
+ * NOTE: does not provide a 1-size fallback for empty initial data.
+ */
+template <typename ElementType, typename AllocatorType>
+FORCEINLINE FRDGBufferRef CreateByteAddressUploadBuffer(
+	FRDGBuilder& GraphBuilder,
+	const TCHAR* Name,
+	const TArray<ElementType, AllocatorType> &InitialData,
+	ERDGInitialDataFlags InitialDataFlags = ERDGInitialDataFlags::None)
+{
+	FRDGBufferRef Buffer = GraphBuilder.CreateBuffer(FRDGBufferDesc::CreateByteAddressUploadDesc(sizeof(ElementType) * InitialData.Num()), Name);
+	GraphBuilder.QueueBufferUpload(Buffer, TConstArrayView<ElementType>(InitialData), InitialDataFlags);
+	return Buffer;
+}
 
 /** Creates a vertex buffer with initial data by creating an upload pass. */
 RENDERCORE_API FRDGBufferRef CreateVertexBuffer(
diff --git a/Engine/Source/Runtime/Renderer/Private/RendererPrivateUtils.h b/Engine/Source/Runtime/Renderer/Private/RendererPrivateUtils.h
--- a/Engine/Source/Runtime/Renderer/Private/RendererPrivateUtils.h
+++ b/Engine/Source/Runtime/Renderer/Private/RendererPrivateUtils.h
@@ -86,43 +86,95 @@ void AddBufferLockReadbackPass(FRDGBuilder& GraphBuilder, TRefCountPtr<FRDGPoole
 }
 
 
+
+namespace UE::RendererPrivateUtils::Implementation
+{
+
 /** 
- * Helper class to manage a persistent structured buffer.
+ * Helper class to manage a persistent buffer.
  */
-class FPersistentStructuredBuffer
+class FPersistentBuffer
 {
 public:
-	FPersistentStructuredBuffer(int32 InMinimumNumElementsReserved, const TCHAR *InName, bool bInRoundUpToPOT = true);
+	FPersistentBuffer(int32 InMinimumNumElementsReserved, const TCHAR *InName, bool bInRoundUpToPOT = true);
 
-	FRDGBuffer* ResizeBufferIfNeeded(FRDGBuilder& GraphBuilder, int32 InNewMinNumElements, int32 BytesPerElement);
 	FRDGBuffer* Register(FRDGBuilder& GraphBuilder);
 
 	void Empty();
+
 protected:
+
+	FRDGBuffer* ResizeBufferIfNeeded(FRDGBuilder& GraphBuilder, const FRDGBufferDesc& BufferDesc);
+
 	int32 MinimumNumElementsReserved = 0;
 	const TCHAR *Name = nullptr;
 	bool bRoundUpToPOT = true;
 	TRefCountPtr<FRDGPooledBuffer> PooledBuffer;
 };
 
+class FBufferScatterUploader
+{
+public:
+	void UploadTo(FRDGBuilder& GraphBuilder, FRDGBuffer *DestBuffer, FRDGBuffer *ScatterOffsets, FRDGBuffer *Values, uint32 NumScatters, uint32 NumBytesPerElement, int32 NumValuesPerScatter);
+};
+
+struct FStructuredBufferTraits
+{
+	template <typename ElementType, typename AllocatorType>
+	static FRDGBuffer *CreateUploadBuffer(FRDGBuilder& GraphBuilder, const TCHAR* Name, const TArray<ElementType, AllocatorType> &InitialData)
+	{
+		return CreateStructuredUploadBuffer(GraphBuilder, Name, InitialData, ERDGInitialDataFlags::NoCopy);
+	}
+
+	static FRDGBufferDesc CreateDesc(uint32 BytesPerElement, uint32 NumElements)
+	{
+		return FRDGBufferDesc::CreateStructuredDesc(BytesPerElement, NumElements);	
+	}
+
+	static constexpr bool bAutoValuesPerScatter = false;
+};
+
+struct FByteAddressBufferTraits
+{
+	template <typename ElementType, typename AllocatorType>
+	static FRDGBuffer *CreateUploadBuffer(FRDGBuilder& GraphBuilder, const TCHAR* Name, const TArray<ElementType, AllocatorType> &InitialData)
+	{
+		return CreateByteAddressUploadBuffer(GraphBuilder, Name, InitialData, ERDGInitialDataFlags::NoCopy);
+	}
+
+	static FRDGBufferDesc CreateDesc(uint32 BytesPerElement, uint32 NumElements)
+	{
+		uint32 NumBytes = BytesPerElement * NumElements;
+		// Needs to be aligned to 16 bytes to MemcpyResource to work correctly (otherwise it skips last unaligned elements of the buffer during resize)
+		check((NumBytes & 15) == 0);
+		return FRDGBufferDesc::CreateByteAddressDesc(NumBytes);	
+	}
+
+	static constexpr bool bAutoValuesPerScatter = true;
+};
+
 /**
  * Typed version of FPersistentStructuredBuffer 
  */
-template <typename InValueType>
-class TPersistentStructuredBuffer : public FPersistentStructuredBuffer
+template <typename InValueType, typename InBufferTraits>
+class TPersistentBuffer : public FPersistentBuffer
 {
 public:
 	using ValueType = InValueType;
-	TPersistentStructuredBuffer(int32 InMinimumNumElementsReserved, const TCHAR *InName, bool bInRoundUpToPOT = true)
-	: FPersistentStructuredBuffer(InMinimumNumElementsReserved, InName, bInRoundUpToPOT)
+	using BufferTraits = InBufferTraits;
+
+	static constexpr uint32 BytesPerElement = sizeof(ValueType);
+
+	TPersistentBuffer(int32 InMinimumNumElementsReserved, const TCHAR *InName, bool bInRoundUpToPOT = true)
+	: FPersistentBuffer(InMinimumNumElementsReserved, InName, bInRoundUpToPOT)
 	{
 	}
 
 	FRDGBuffer* ResizeBufferIfNeeded(FRDGBuilder& GraphBuilder, int32 InNewMinNumElements)
 	{
-		return FPersistentStructuredBuffer::ResizeBufferIfNeeded(GraphBuilder, InNewMinNumElements, sizeof(ValueType));
-	}
-	
+		int32 NewMinNumElements = FMath::Max(MinimumNumElementsReserved, bRoundUpToPOT ? int32(FMath::RoundUpToPowerOfTwo(InNewMinNumElements)) : InNewMinNumElements);
+		return FPersistentBuffer::ResizeBufferIfNeeded(GraphBuilder, BufferTraits::CreateDesc(BytesPerElement, NewMinNumElements));
+	}	
 
 	template <typename ValueCheckFuncType>
 	void ValidateGPUData(FRDGBuilder& GraphBuilder, TConstArrayView<ValueType> HostValues, ValueCheckFuncType &&ValueCheckFunc)
@@ -148,24 +200,21 @@ public:
 	TRefCountPtr<FRDGPooledBuffer> &GetPooledBuffer() { return PooledBuffer; }
 };
 
-class FStructuredBufferScatterUploader
-{
-public:
-	void UploadTo(FRDGBuilder& GraphBuilder, FRDGBuffer *DestBuffer, FRDGBuffer *ScatterOffsets, FRDGBuffer *Values, uint32 NumScatters, uint32 NumBytesPerElement, uint32 NumValuesPerScatter);
-};
-
-template <typename InValueType, int32 InNumValuesPerScatter = 1>
-class TStructuredBufferScatterUploader : public FStructuredBufferScatterUploader
+template <typename InValueType, typename InBufferTraits,  int32 InNumValuesPerScatter = 1>
+class TBufferScatterUploader : public FBufferScatterUploader
 {
 public:
+	using BufferTraits = InBufferTraits;
 	using ValueType = InValueType;
 	static constexpr int32 BytesPerElement = sizeof(ValueType);
 	static constexpr int32 NumValuesPerScatter = InNumValuesPerScatter;
 
+	static_assert((BytesPerElement % 4) == 0, "The struct used must be 4-byte aligned");
+
 	/**
 	 * Optionally reserve space for NumScatters items.
 	 */
-	TStructuredBufferScatterUploader(int32 NumScatters = 0)
+	TBufferScatterUploader(int32 NumScatters = 0)
 	{
 		Reserve(NumScatters);
 	}
@@ -206,13 +255,29 @@ public:
 		UploadData.ScatterOffsets.Add(ScatterOffset);
 	}
 
+	/**
+	 * Add a number of values to scatter to consecutive destination offsets with a common start offset, 
+	 * NOTE: this allocates a new scatter offset for each NumValuesPerScatter elements.
+	 */
+	void AddMultiple(const TConstArrayView<ValueType> &InValues, int32 FirstScatterOffset)
+	{
+		check(InValues.Num() % NumValuesPerScatter = 0);
+		check(UploadDataProxy == nullptr);
+
+		UploadData.Values.Append(InValues);
+		for (int32 Index = 0; Index < InValues.Num() / NumValuesPerScatter; ++Index)
+		{
+			UploadData.ScatterOffsets.Add(FirstScatterOffset + Index);
+		}
+	}
+
 	int32 GetNumScatters() const { return UploadDataProxy != nullptr ? UploadDataProxy->ScatterOffsets.Num() : UploadData.ScatterOffsets.Num(); }
 
 	/**
 	 * Resize the destination persistent buffer (if needed) and upload & scatter the collected data to it.
 	 * This locks the uploader to prevent accidental resize (and thus realloc) of the buffer by adding more elements.
 	 */
-	FRDGBuffer *ResizeAndUploadTo(FRDGBuilder& GraphBuilder, TPersistentStructuredBuffer<ValueType> &DestDataBuffer, int32 DestDataMinimumSize)
+	FRDGBuffer *ResizeAndUploadTo(FRDGBuilder& GraphBuilder, TPersistentBuffer<ValueType, BufferTraits> &DestDataBuffer, int32 DestDataMinimumSize)
 	{
 		check(UploadDataProxy == nullptr);
 
@@ -225,10 +290,10 @@ public:
 			UploadDataProxy = GraphBuilder.AllocObject<FUploadData>(MoveTemp(UploadData));
 
 			// upload the values & offsets
-			FRDGBuffer *ScatterOffsetsRDG = CreateStructuredUploadBuffer(GraphBuilder, TEXT("ScatterUploader.Offsets"), UploadDataProxy->ScatterOffsets, ERDGInitialDataFlags::NoCopy);
-			FRDGBuffer *ValuesRDG = CreateStructuredUploadBuffer(GraphBuilder, TEXT("ScatterUploader.Values"), UploadDataProxy->Values, ERDGInitialDataFlags::NoCopy);
+			FRDGBuffer *ScatterOffsetsRDG = BufferTraits::CreateUploadBuffer(GraphBuilder, TEXT("ScatterUploader.Offsets"), UploadDataProxy->ScatterOffsets);
+			FRDGBuffer *ValuesRDG = BufferTraits::CreateUploadBuffer(GraphBuilder, TEXT("ScatterUploader.Values"), UploadDataProxy->Values);
 
-			FStructuredBufferScatterUploader::UploadTo(GraphBuilder, DestBufferRDG, ScatterOffsetsRDG, ValuesRDG, NumScatters, sizeof(ValueType), NumValuesPerScatter);
+			FBufferScatterUploader::UploadTo(GraphBuilder, DestBufferRDG, ScatterOffsetsRDG, ValuesRDG, NumScatters, sizeof(ValueType), BufferTraits::bAutoValuesPerScatter ? INDEX_NONE : NumValuesPerScatter);
 		}
 
 		return DestBufferRDG;
@@ -243,3 +308,16 @@ private:
 	FUploadData UploadData;
 };
 
+}
+
+template <typename InValueType>
+using TPersistentStructuredBuffer = UE::RendererPrivateUtils::Implementation::TPersistentBuffer<InValueType, UE::RendererPrivateUtils::Implementation::FStructuredBufferTraits>;
+
+template <typename InValueType, int32 InNumValuesPerScatter = 1>
+using TStructuredBufferScatterUploader = UE::RendererPrivateUtils::Implementation::TBufferScatterUploader<InValueType, UE::RendererPrivateUtils::Implementation::FStructuredBufferTraits, InNumValuesPerScatter>;
+
+template <typename InValueType>
+using TPersistentByteAddressBuffer = UE::RendererPrivateUtils::Implementation::TPersistentBuffer<InValueType, UE::RendererPrivateUtils::Implementation::FByteAddressBufferTraits>;
+
+template <typename InValueType, int32 InNumValuesPerScatter = 1>
+using TByteAddressBufferScatterUploader = UE::RendererPrivateUtils::Implementation::TBufferScatterUploader<InValueType, UE::RendererPrivateUtils::Implementation::FByteAddressBufferTraits, InNumValuesPerScatter>;
diff --git a/Engine/Source/Runtime/Renderer/Private/RendererUtils.cpp b/Engine/Source/Runtime/Renderer/Private/RendererUtils.cpp
--- a/Engine/Source/Runtime/Renderer/Private/RendererUtils.cpp
+++ b/Engine/Source/Runtime/Renderer/Private/RendererUtils.cpp
@@ -227,36 +227,39 @@ namespace Substrate
 	}
 }
 
-FPersistentStructuredBuffer::FPersistentStructuredBuffer(int32 InMinimumNumElementsReserved, const TCHAR *InName, bool bInRoundUpToPOT)
+namespace UE::RendererPrivateUtils::Implementation
+{
+
+FPersistentBuffer::FPersistentBuffer(int32 InMinimumNumElementsReserved, const TCHAR *InName, bool bInRoundUpToPOT)
 	: MinimumNumElementsReserved(InMinimumNumElementsReserved)
 	, Name(InName)
 	, bRoundUpToPOT(bInRoundUpToPOT)
 {
 }
 
-FRDGBuffer* FPersistentStructuredBuffer::ResizeBufferIfNeeded(FRDGBuilder& GraphBuilder, int32 InNewMinNumElements, int32 BytesPerElement)
-{
-	int32 NewMinNumElements = FMath::Max(MinimumNumElementsReserved, bRoundUpToPOT ? int32(FMath::RoundUpToPowerOfTwo(InNewMinNumElements)) : InNewMinNumElements);
-	FRDGBuffer* BufferRDG = ::ResizeBufferIfNeeded(GraphBuilder, PooledBuffer, FRDGBufferDesc::CreateStructuredDesc(BytesPerElement, NewMinNumElements), Name);
-	return BufferRDG;
-}
-
-FRDGBuffer* FPersistentStructuredBuffer::Register(FRDGBuilder& GraphBuilder) 
+FRDGBuffer* FPersistentBuffer::Register(FRDGBuilder& GraphBuilder) 
 { 
 	return GraphBuilder.RegisterExternalBuffer(PooledBuffer); 
 }
 
-void FPersistentStructuredBuffer::Empty()
+void FPersistentBuffer::Empty()
 {
 	PooledBuffer.SafeRelease();
 }
 
-void FStructuredBufferScatterUploader::UploadTo(FRDGBuilder& GraphBuilder, FRDGBuffer *DestBuffer, FRDGBuffer *ScatterOffsets, FRDGBuffer *Values, uint32 NumScatters, uint32 NumBytesPerElement, uint32 NumValuesPerScatter)
+FRDGBuffer* FPersistentBuffer::ResizeBufferIfNeeded(FRDGBuilder& GraphBuilder, const FRDGBufferDesc& BufferDesc)
+{
+	return ::ResizeBufferIfNeeded(GraphBuilder, PooledBuffer, BufferDesc, Name);
+}
+
+void FBufferScatterUploader::UploadTo(FRDGBuilder& GraphBuilder, FRDGBuffer *DestBuffer, FRDGBuffer *ScatterOffsets, FRDGBuffer *Values, uint32 NumScatters, uint32 NumBytesPerElement, int32 NumValuesPerScatter)
 {
 	FScatterCopyParams ScatterCopyParams { NumScatters, NumBytesPerElement, NumValuesPerScatter };
 	ScatterCopyResource(GraphBuilder, DestBuffer, GraphBuilder.CreateSRV(ScatterOffsets), GraphBuilder.CreateSRV(Values), ScatterCopyParams);
 }
 
+}
+
 TGlobalResource<FTileTexCoordVertexBuffer> GOneTileQuadVertexBuffer(1);
 TGlobalResource<FTileIndexBuffer> GOneTileQuadIndexBuffer(1);
 
