Fixed inverted boolean logic blocking all GPU-Scene light updates.

#rnx

diff --git a/Engine/Source/Runtime/Renderer/Private/GPUScene.cpp b/Engine/Source/Runtime/Renderer/Private/GPUScene.cpp
--- a/Engine/Source/Runtime/Renderer/Private/GPUScene.cpp
+++ b/Engine/Source/Runtime/Renderer/Private/GPUScene.cpp
@@ -2170,7 +2170,7 @@ void FGPUScene::OnPostSceneUpdate(FRDGBuilder& GraphBuilder, const FScenePostUpd
 
 void FGPUScene::OnPostLightSceneInfoUpdate(FRDGBuilder& GraphBuilder, const FLightSceneChangeSet& LightsPostUpdateData)
 {
-	const bool bAnythingChanged = !(LightsPostUpdateData.AddedLightIds.IsEmpty() || LightsPostUpdateData.RemovedLightIds.IsEmpty() || LightsPostUpdateData.TransformUpdatedLightIds.IsEmpty() || LightsPostUpdateData.ColorUpdatedLightIds.IsEmpty());
+	const bool bAnythingChanged = !(LightsPostUpdateData.AddedLightIds.IsEmpty() && LightsPostUpdateData.RemovedLightIds.IsEmpty() && LightsPostUpdateData.TransformUpdatedLightIds.IsEmpty() && LightsPostUpdateData.ColorUpdatedLightIds.IsEmpty());
 	if (bAnythingChanged || !LightDataBuffer.IsValid())
 	{
 		UpdateGPULights(GraphBuilder, UE::Tasks::FTask{});
