Adds the history's memory size in VisualizeTemporalUpscaler showflag

diff --git a/Engine/Source/Runtime/Renderer/Private/PostProcess/VisualizeTemporalUpscaler.cpp b/Engine/Source/Runtime/Renderer/Private/PostProcess/VisualizeTemporalUpscaler.cpp
--- a/Engine/Source/Runtime/Renderer/Private/PostProcess/VisualizeTemporalUpscaler.cpp
+++ b/Engine/Source/Runtime/Renderer/Private/PostProcess/VisualizeTemporalUpscaler.cpp
@@ -287,13 +287,42 @@ FScreenPassTexture AddVisualizeTemporalUpscalerPass(FRDGBuilder& GraphBuilder, c
 				QuickDrawSummary(/* Location = */ 4, TEXT("Support Alpha: ") + Text);
 			}
 
+			// Display memory size of the history.
+			{
+				static auto CVarAntiAliasingQuality = IConsoleManager::Get().FindConsoleVariable(TEXT("sg.AntiAliasingQuality"));
+				check(CVarAntiAliasingQuality);
+
+				const uint64 PingPongHistorySizeMultiplier = 2;
+				uint64 HistorySize = 0;
+				if (Inputs.TAAConfig == EMainTAAPassConfig::TAA)
+				{
+					HistorySize = View.PrevViewInfo.TemporalAAHistory.GetGPUSizeBytes(/* bLogSizes = */ false) * PingPongHistorySizeMultiplier;
+				}
+				else if (Inputs.TAAConfig == EMainTAAPassConfig::TSR && View.PrevViewInfo.TSRHistory.IsValid())
+				{
+					HistorySize = View.PrevViewInfo.TSRHistory.GetGPUSizeBytes(/* bLogSizes = */ false);
+
+					bool bHasTSRHistoryResurrection = View.PrevViewInfo.TSRHistory.MetadataArray->GetDesc().ArraySize > 1;
+					bool bHasTSRPingPongHistory = !bHasTSRHistoryResurrection;
+					if (bHasTSRPingPongHistory)
+					{
+						HistorySize *= PingPongHistorySizeMultiplier;
+					}
+				}
+				else if (Inputs.TAAConfig == EMainTAAPassConfig::ThirdParty && View.PrevViewInfo.ThirdPartyTemporalUpscalerHistory.IsValid())
+				{
+					HistorySize = View.PrevViewInfo.ThirdPartyTemporalUpscalerHistory->GetGPUSizeBytes();
+				}
+				QuickDrawSummary(/* Location = */ 5, FString::Printf(TEXT("History VRAM footprint: %.1f MB"), float(HistorySize) / float(1024 * 1024)));
+			}
+
 			// Display if any additional sharpening is happening
 			{
 				static auto CVarSharpen = IConsoleManager::Get().FindConsoleVariable(TEXT("r.Tonemapper.Sharpen"));
 				check(CVarSharpen);
 				float Sharpen = CVarSharpen->GetFloat();
 				Sharpen = (Sharpen < 0) ? View.FinalPostProcessSettings.Sharpen : Sharpen;
-				QuickDrawSummary(/* Location = */ 5, Sharpen > 0 ? FString::Printf(TEXT("Tonemapper Sharpen: %f"), Sharpen) : TEXT("Tonemapper Sharpen: Off"));
+				QuickDrawSummary(/* Location = */ 6, Sharpen > 0 ? FString::Printf(TEXT("Tonemapper Sharpen: %f"), Sharpen) : TEXT("Tonemapper Sharpen: Off"));
 			}
 
 		});
