Scene View Extensions (public header) use FPostProcessingInputs which are currently in private header, which is illegal. The following change moves out dependencies required for FPostProcessingInputs to be used by plugins that utilize SVE.

diff --git a/Engine/Source/Runtime/Renderer/Internal/PathTracingResources.h b/Engine/Source/Runtime/Renderer/Internal/PathTracingResources.h
new file mode 100644
--- /dev/null
+++ b/Engine/Source/Runtime/Renderer/Internal/PathTracingResources.h
@@ -0,0 +1,16 @@
+// Copyright Epic Games, Inc. All Rights Reserved.
+
+#pragma once
+
+class FRDGTexture;
+
+struct FPathTracingResources
+{
+	FRDGTexture* DenoisedRadiance = nullptr;
+	FRDGTexture* Radiance = nullptr;
+	FRDGTexture* Albedo = nullptr;
+	FRDGTexture* Normal = nullptr;
+	FRDGTexture* Variance = nullptr;
+	
+	bool bPostProcessEnabled = false;
+};
diff --git a/Engine/Source/Runtime/Renderer/Internal/PostProcess/PostProcessInputs.h b/Engine/Source/Runtime/Renderer/Internal/PostProcess/PostProcessInputs.h
new file mode 100644
--- /dev/null
+++ b/Engine/Source/Runtime/Renderer/Internal/PostProcess/PostProcessInputs.h
@@ -0,0 +1,25 @@
+// Copyright Epic Games, Inc. All Rights Reserved.
+
+#pragma once
+
+#include "TranslucentPassResource.h"
+#include "PathTracingResources.h"
+
+struct FPostProcessingInputs
+{
+	TRDGUniformBufferRef<FSceneTextureUniformParameters> SceneTextures = nullptr;
+	FRDGTextureRef ViewFamilyTexture = nullptr;
+	FRDGTextureRef CustomDepthTexture = nullptr;
+	FRDGTextureRef ExposureIlluminance = nullptr;
+	FTranslucencyViewResourcesMap TranslucencyViewResourcesMap;
+	FPathTracingResources PathTracingResources;
+
+	bool bSeparateCustomStencil = false;
+
+	void Validate() const
+	{
+		check(SceneTextures);
+		check(ViewFamilyTexture);
+		check(TranslucencyViewResourcesMap.IsValid());
+	}
+};
diff --git a/Engine/Source/Runtime/Renderer/Internal/TranslucentPassResource.h b/Engine/Source/Runtime/Renderer/Internal/TranslucentPassResource.h
new file mode 100644
--- /dev/null
+++ b/Engine/Source/Runtime/Renderer/Internal/TranslucentPassResource.h
@@ -0,0 +1,118 @@
+// Copyright Epic Games, Inc. All Rights Reserved.
+
+#pragma once
+
+#include "ScreenPass.h"
+#include "SystemTextures.h"
+
+
+// enum instead of bool to get better visibility when we pass around multiple bools, also allows for easier extensions
+namespace ETranslucencyPass
+{
+	enum Type : int
+	{
+		TPT_TranslucencyStandard,
+		TPT_TranslucencyStandardModulate,
+		TPT_TranslucencyAfterDOF,
+		TPT_TranslucencyAfterDOFModulate,
+		TPT_TranslucencyAfterMotionBlur,
+
+		/** Drawing all translucency, regardless of separate or standard.  Used when drawing translucency outside of the main renderer, eg FRendererModule::DrawTile. */
+		TPT_AllTranslucency,
+		TPT_MAX
+	};
+};
+
+/** Resources of a translucency pass. */
+struct FTranslucencyPassResources
+{
+	ETranslucencyPass::Type Pass = ETranslucencyPass::TPT_MAX;
+	FIntRect ViewRect = FIntRect(0, 0, 0, 0);
+	FRDGTextureMSAA ColorTexture;
+	FRDGTextureMSAA ColorModulateTexture;
+	FRDGTextureMSAA DepthTexture;
+
+	inline bool IsValid() const
+	{
+		return ViewRect.Width() > 0 && ViewRect.Height() > 0;
+	}
+
+	inline FRDGTextureRef GetColorForRead(FRDGBuilder& GraphBuilder) const
+	{
+		if (!ColorTexture.IsValid())
+		{
+			return GSystemTextures.GetBlackAlphaOneDummy(GraphBuilder);
+		}
+		return ColorTexture.Resolve;
+	}
+
+	inline FRDGTextureRef GetColorModulateForRead(FRDGBuilder& GraphBuilder) const
+	{
+		if (!ColorModulateTexture.IsValid())
+		{
+			return GSystemTextures.GetWhiteDummy(GraphBuilder);
+		}
+		return ColorModulateTexture.Resolve;
+	}
+
+	inline FRDGTextureRef GetDepthForRead(FRDGBuilder& GraphBuilder) const
+	{
+		if (!DepthTexture.IsValid())
+		{
+			return GSystemTextures.GetMaxFP16Depth(GraphBuilder);
+		}
+		return DepthTexture.Resolve;
+	}
+
+	inline FScreenPassTextureViewport GetTextureViewport() const
+	{
+		check(IsValid());
+		return FScreenPassTextureViewport(ColorTexture.Target->Desc.Extent, ViewRect);
+	}
+};
+
+/** All resources of all translucency passes for a view family. */
+struct FTranslucencyPassResourcesMap
+{
+	FTranslucencyPassResourcesMap(int32 NumViews);
+
+	inline FTranslucencyPassResources& Get(int32 ViewIndex, ETranslucencyPass::Type Translucency)
+	{
+		return Array[ViewIndex][int32(Translucency)];
+	};
+
+	inline const FTranslucencyPassResources& Get(int32 ViewIndex, ETranslucencyPass::Type Translucency) const
+	{
+		check(ViewIndex < Array.Num());
+		return Array[ViewIndex][int32(Translucency)];
+	};
+
+private:
+	TArray<TStaticArray<FTranslucencyPassResources, ETranslucencyPass::TPT_MAX>, TInlineAllocator<4>> Array;
+};
+
+/** All resources of all translucency for one view. */
+struct FTranslucencyViewResourcesMap
+{
+	FTranslucencyViewResourcesMap() = default;
+
+	FTranslucencyViewResourcesMap(const FTranslucencyPassResourcesMap& InTranslucencyPassResourcesMap, int32 InViewIndex)
+		: TranslucencyPassResourcesMap(&InTranslucencyPassResourcesMap)
+		, ViewIndex(InViewIndex)
+	{ }
+
+	bool IsValid() const
+	{
+		return TranslucencyPassResourcesMap != nullptr;
+	}
+
+	inline const FTranslucencyPassResources& Get(ETranslucencyPass::Type Translucency) const
+	{
+		check(IsValid());
+		return TranslucencyPassResourcesMap->Get(ViewIndex, Translucency);
+	};
+
+private:
+	const FTranslucencyPassResourcesMap* TranslucencyPassResourcesMap = nullptr;
+	int32 ViewIndex = 0;
+};
diff --git a/Engine/Source/Runtime/Renderer/Private/FXRenderingUtils.cpp b/Engine/Source/Runtime/Renderer/Private/FXRenderingUtils.cpp
--- a/Engine/Source/Runtime/Renderer/Private/FXRenderingUtils.cpp
+++ b/Engine/Source/Runtime/Renderer/Private/FXRenderingUtils.cpp
@@ -18,6 +18,7 @@ TConstStridedView<FSceneView> UE::FXRenderingUtils::ConvertViewArray(TConstArray
 
 FIntRect UE::FXRenderingUtils::GetRawViewRectUnsafe(const FSceneView& View)
 {
+	check(View.bIsViewInfo);
 	return static_cast<const FViewInfo&>(View).ViewRect;
 }
 
diff --git a/Engine/Source/Runtime/Renderer/Private/MeshDrawCommands.h b/Engine/Source/Runtime/Renderer/Private/MeshDrawCommands.h
--- a/Engine/Source/Runtime/Renderer/Private/MeshDrawCommands.h
+++ b/Engine/Source/Runtime/Renderer/Private/MeshDrawCommands.h
@@ -7,7 +7,7 @@ MeshDrawCommands.h: Mesh draw commands.
 #pragma once
 
 #include "MeshPassProcessor.h"
-#include "TranslucencyPass.h"
+#include "TranslucentPassResource.h"
 #include "InstanceCulling/InstanceCullingContext.h"
 #include "InstanceCulling/InstanceCullingManager.h"
 #include "ScenePrivateBase.h"
diff --git a/Engine/Source/Runtime/Renderer/Private/PathTracing.h b/Engine/Source/Runtime/Renderer/Private/PathTracing.h
--- a/Engine/Source/Runtime/Renderer/Private/PathTracing.h
+++ b/Engine/Source/Runtime/Renderer/Private/PathTracing.h
@@ -2,6 +2,7 @@
 
 #pragma once
 
+#include "PathTracingResources.h"
 #include "ShaderParameterMacros.h"
 
 // this struct holds skylight parameters
@@ -31,19 +32,6 @@ RENDERER_API FRHIRayTracingShader* GetGPULightmassDefaultHiddenHitShader(const F
 
 void PreparePathTracingRTPSO();
 
-class FRDGTexture;
-
-struct FPathTracingResources
-{
-	FRDGTexture* DenoisedRadiance = nullptr;
-	FRDGTexture* Radiance = nullptr;
-	FRDGTexture* Albedo = nullptr;
-	FRDGTexture* Normal = nullptr;
-	FRDGTexture* Variance = nullptr;
-	
-	bool bPostProcessEnabled = false;
-};
-
 namespace PathTracing
 {
 	bool UsesDecals(const FSceneViewFamily& ViewFamily);
diff --git a/Engine/Source/Runtime/Renderer/Private/TranslucencyPass.h b/Engine/Source/Runtime/Renderer/Private/TranslucencyPass.h
--- a/Engine/Source/Runtime/Renderer/Private/TranslucencyPass.h
+++ b/Engine/Source/Runtime/Renderer/Private/TranslucencyPass.h
@@ -4,23 +4,6 @@
 
 #include "Misc/EnumClassFlags.h"
 
-// enum instead of bool to get better visibility when we pass around multiple bools, also allows for easier extensions
-namespace ETranslucencyPass
-{
-	enum Type : int
-	{
-		TPT_TranslucencyStandard,
-		TPT_TranslucencyStandardModulate,
-		TPT_TranslucencyAfterDOF,
-		TPT_TranslucencyAfterDOFModulate,
-		TPT_TranslucencyAfterMotionBlur,
-
-		/** Drawing all translucency, regardless of separate or standard.  Used when drawing translucency outside of the main renderer, eg FRendererModule::DrawTile. */
-		TPT_AllTranslucency,
-		TPT_MAX
-	};
-};
-
 enum class ETranslucencyView
 {
 	None       = 0,
diff --git a/Engine/Source/Runtime/Renderer/Private/TranslucentRendering.h b/Engine/Source/Runtime/Renderer/Private/TranslucentRendering.h
--- a/Engine/Source/Runtime/Renderer/Private/TranslucentRendering.h
+++ b/Engine/Source/Runtime/Renderer/Private/TranslucentRendering.h
@@ -33,100 +33,6 @@ struct FSeparateTranslucencyDimensions
 	uint32 NumSamples = 1;
 };
 
-/** Resources of a translucency pass. */
-struct FTranslucencyPassResources
-{
-	ETranslucencyPass::Type Pass = ETranslucencyPass::TPT_MAX;
-	FIntRect ViewRect = FIntRect(0, 0, 0, 0);
-	FRDGTextureMSAA ColorTexture;
-	FRDGTextureMSAA ColorModulateTexture;
-	FRDGTextureMSAA DepthTexture;
-
-	inline bool IsValid() const
-	{
-		return ViewRect.Width() > 0 && ViewRect.Height() > 0;
-	}
-
-	inline FRDGTextureRef GetColorForRead(FRDGBuilder& GraphBuilder) const
-	{
-		if (!ColorTexture.IsValid())
-		{
-			return GSystemTextures.GetBlackAlphaOneDummy(GraphBuilder);
-		}
-		return ColorTexture.Resolve;
-	}
-
-	inline FRDGTextureRef GetColorModulateForRead(FRDGBuilder& GraphBuilder) const
-	{
-		if (!ColorModulateTexture.IsValid())
-		{
-			return GSystemTextures.GetWhiteDummy(GraphBuilder);
-		}
-		return ColorModulateTexture.Resolve;
-	}
-
-	inline FRDGTextureRef GetDepthForRead(FRDGBuilder& GraphBuilder) const
-	{
-		if (!DepthTexture.IsValid())
-		{
-			return GSystemTextures.GetMaxFP16Depth(GraphBuilder);
-		}
-		return DepthTexture.Resolve;
-	}
-
-	inline FScreenPassTextureViewport GetTextureViewport() const
-	{
-		check(IsValid());
-		return FScreenPassTextureViewport(ColorTexture.Target->Desc.Extent, ViewRect);
-	}
-};
-
-/** All resources of all translucency passes for a view family. */
-struct FTranslucencyPassResourcesMap
-{
-	FTranslucencyPassResourcesMap(int32 NumViews);
-
-	inline FTranslucencyPassResources& Get(int32 ViewIndex, ETranslucencyPass::Type Translucency)
-	{
-		return Array[ViewIndex][int32(Translucency)];
-	};
-
-	inline const FTranslucencyPassResources& Get(int32 ViewIndex, ETranslucencyPass::Type Translucency) const
-	{
-		check(ViewIndex < Array.Num());
-		return Array[ViewIndex][int32(Translucency)];
-	};
-
-private:
-	TArray<TStaticArray<FTranslucencyPassResources, ETranslucencyPass::TPT_MAX>, TInlineAllocator<4>> Array;
-};
-
-/** All resources of all translucency for one view. */
-struct FTranslucencyViewResourcesMap
-{
-	FTranslucencyViewResourcesMap() = default;
-
-	FTranslucencyViewResourcesMap(const FTranslucencyPassResourcesMap& InTranslucencyPassResourcesMap, int32 InViewIndex)
-		: TranslucencyPassResourcesMap(&InTranslucencyPassResourcesMap)
-		, ViewIndex(InViewIndex)
-	{ }
-
-	bool IsValid() const
-	{
-		return TranslucencyPassResourcesMap != nullptr;
-	}
-
-	inline const FTranslucencyPassResources& Get(ETranslucencyPass::Type Translucency) const
-	{
-		check(IsValid());
-		return TranslucencyPassResourcesMap->Get(ViewIndex, Translucency);
-	};
-
-private:
-	const FTranslucencyPassResourcesMap* TranslucencyPassResourcesMap = nullptr;
-	int32 ViewIndex = 0;
-};
-
 DECLARE_GPU_DRAWCALL_STAT_EXTERN(Translucency);
 
 /** Converts the the translucency pass into the respective mesh pass. */
