* Add OIT permutation to base pass shaders in order to handle runtime toggle.
* Add RHI ROV feature to know if the runtime RHI supports ROVs.

Details:
* OIT permutation is only compiled when the OIT project settings is enabled
* OIT permutation allows to handle case where a platform supports ROV, but the actual RHI doesn't
* This also allow to toggle OIT support at runtime (through quality settings for instance).

diff --git a/Engine/Shaders/Private/BasePassPixelShader.usf b/Engine/Shaders/Private/BasePassPixelShader.usf
--- a/Engine/Shaders/Private/BasePassPixelShader.usf
+++ b/Engine/Shaders/Private/BasePassPixelShader.usf
@@ -134,7 +134,7 @@
 #define MATERIAL_WORKS_WITH_DUAL_SOURCE_COLOR_BLENDING (MATERIAL_SHADINGMODEL_THIN_TRANSLUCENT || SUBSTRATE_BLENDING_TRANSLUCENT_COLOREDTRANSMITTANCE )
 #define OUTPUTS_COVERAGE (MATERIALBLENDING_MASKED_USING_COVERAGE && !EARLY_Z_PASS_ONLY_MATERIAL_MASKING)
 
-#define OIT_ENABLED (PROJECT_OIT && PLATFORM_SUPPORTS_ROV && (MATERIALBLENDING_TRANSLUCENT || MATERIALBLENDING_ADDITIVE || MATERIAL_WORKS_WITH_DUAL_SOURCE_COLOR_BLENDING))
+#define OIT_ENABLED (PROJECT_OIT && PERMUTATION_SUPPORTS_OIT && PLATFORM_SUPPORTS_ROV && (MATERIALBLENDING_TRANSLUCENT || MATERIALBLENDING_ADDITIVE || MATERIAL_WORKS_WITH_DUAL_SOURCE_COLOR_BLENDING))
 
 #if OIT_ENABLED
 #define OIT_IS_BASEPASS 1
diff --git a/Engine/Source/Runtime/D3D12RHI/Private/D3D12Adapter.cpp b/Engine/Source/Runtime/D3D12RHI/Private/D3D12Adapter.cpp
--- a/Engine/Source/Runtime/D3D12RHI/Private/D3D12Adapter.cpp
+++ b/Engine/Source/Runtime/D3D12RHI/Private/D3D12Adapter.cpp
@@ -1039,8 +1039,10 @@ void FD3D12Adapter::InitializeDevices()
 				RootDevice->CheckFeatureSupport(D3D12_FEATURE_D3D12_OPTIONS, &Features, sizeof(Features));
 
 				GRHISupportsStencilRefFromPixelShader = (Features.PSSpecifiedStencilRefSupported != 0);
+				GRHISupportsRasterOrderViews = (Features.ROVsSupported != 0);
 
 				UE_LOG(LogD3D12RHI, Log, TEXT("Stencil ref from pixel shader is %s"), GRHISupportsStencilRefFromPixelShader ? TEXT("supported") : TEXT("not supported"));
+				UE_LOG(LogD3D12RHI, Log, TEXT("Raster order views are %s"), GRHISupportsRasterOrderViews ? TEXT("supported") : TEXT("not supported"));
 			}
 
 			// Detect availability of shader model 6.0 wave operations
diff --git a/Engine/Source/Runtime/RHI/Public/RHIGlobals.h b/Engine/Source/Runtime/RHI/Public/RHIGlobals.h
--- a/Engine/Source/Runtime/RHI/Public/RHIGlobals.h
+++ b/Engine/Source/Runtime/RHI/Public/RHIGlobals.h
@@ -563,6 +563,9 @@ struct FRHIGlobals
 	/** True if the RHI supports setting the stencil ref at pixel granularity from the pixel shader */
 	bool SupportsStencilRefFromPixelShader = false;
 
+	/** True if the RHI supports raster order views. */
+	bool SupportsRasterOrderViews = false;
+
 	/** Whether current RHI supports overestimated conservative rasterization. */
 	bool SupportsConservativeRasterization = false;
 
@@ -787,6 +790,7 @@ extern RHI_API FRHIGlobals GRHIGlobals;
 #define GRHISupportsPipelineFileCache                          GRHIGlobals.SupportsPipelineFileCache
 #define GRHISupportsPSOPrecaching                              GRHIGlobals.SupportsPSOPrecaching
 #define GRHISupportsStencilRefFromPixelShader                  GRHIGlobals.SupportsStencilRefFromPixelShader
+#define GRHISupportsRasterOrderViews                           GRHIGlobals.SupportsRasterOrderViews
 #define GRHISupportsConservativeRasterization                  GRHIGlobals.SupportsConservativeRasterization
 #define GRHISupportsShaderRootConstants                        GRHIGlobals.SupportsShaderRootConstants
 #define GRHISupportsShaderBundleDispatch                       GRHIGlobals.SupportsShaderBundleDispatch
diff --git a/Engine/Source/Runtime/Renderer/Private/BasePassRendering.cpp b/Engine/Source/Runtime/Renderer/Private/BasePassRendering.cpp
--- a/Engine/Source/Runtime/Renderer/Private/BasePassRendering.cpp
+++ b/Engine/Source/Runtime/Renderer/Private/BasePassRendering.cpp
@@ -510,21 +510,31 @@ void SetupBasePassState(FExclusiveDepthStencil::Type BasePassDepthStencilAccess,
 }
 
 template <ELightMapPolicyType Policy, EGBufferLayout GBufferLayout>
-void AddUniformBasePassPixelShader(bool bEnableSkyLight, bool bUse128bitRT, FMaterialShaderTypes& OutShaderTypes)
+void AddUniformBasePassPixelShader(bool bEnableSkyLight, bool bUse128bitRT, FMaterialShaderTypes& OutShaderTypes, bool bIsForOITPass = false)
 {
+	int32 PermutationId = 0;
+
+	if (bIsForOITPass)
+	{
+		using FMyShader = TBasePassPS<TUniformLightMapPolicy<Policy>, true, GBufferLayout>;
+		typename FMyShader::FPermutationDomain PermutationVector;
+		PermutationVector.template Set<typename FMyShader::FSupportOITDim>(true);
+		PermutationId = PermutationVector.ToDimensionValueId();
+	}
+
 	if (bEnableSkyLight)
 	{
-		OutShaderTypes.AddShaderType<TBasePassPS<TUniformLightMapPolicy<Policy>, true, GBufferLayout>>();
+		OutShaderTypes.AddShaderType<TBasePassPS<TUniformLightMapPolicy<Policy>, true, GBufferLayout>>(PermutationId);
 	}
 	else
 	{
 		if (bUse128bitRT && (Policy == LMP_NO_LIGHTMAP) && (GBufferLayout == GBL_Default))
 		{
-			OutShaderTypes.AddShaderType<F128BitRTBasePassPS>();
+			OutShaderTypes.AddShaderType<F128BitRTBasePassPS>(PermutationId);
 		}
 		else
 		{
-			OutShaderTypes.AddShaderType<TBasePassPS<TUniformLightMapPolicy<Policy>, false, GBufferLayout>>();
+			OutShaderTypes.AddShaderType<TBasePassPS<TUniformLightMapPolicy<Policy>, false, GBufferLayout>>(PermutationId);
 		}
 	}
 }
@@ -542,7 +552,8 @@ bool GetUniformBasePassShaders(
 	bool bUse128bitRT,
 	EGBufferLayout GBufferLayout,
 	TShaderRef<TBasePassVertexShaderPolicyParamType<FUniformLightMapPolicy>>* VertexShader,
-	TShaderRef<TBasePassPixelShaderPolicyParamType<FUniformLightMapPolicy>>* PixelShader
+	TShaderRef<TBasePassPixelShaderPolicyParamType<FUniformLightMapPolicy>>* PixelShader,
+	bool bIsForOITPass = false
 )
 {
 	FMaterialShaderTypes ShaderTypes;
@@ -556,10 +567,10 @@ bool GetUniformBasePassShaders(
 		switch (GBufferLayout)
 		{
 		case GBL_Default:
-			AddUniformBasePassPixelShader<Policy, GBL_Default>(bEnableSkyLight, bUse128bitRT, ShaderTypes);
+			AddUniformBasePassPixelShader<Policy, GBL_Default>(bEnableSkyLight, bUse128bitRT, ShaderTypes, bIsForOITPass);
 			break;
 		case GBL_ForceVelocity:
-			AddUniformBasePassPixelShader<Policy, GBL_ForceVelocity>(bEnableSkyLight, bUse128bitRT, ShaderTypes);
+			AddUniformBasePassPixelShader<Policy, GBL_ForceVelocity>(bEnableSkyLight, bUse128bitRT, ShaderTypes, bIsForOITPass);
 			break;
 		default:
 			check(false);
@@ -588,25 +599,26 @@ bool GetBasePassShaders<FUniformLightMapPolicy>(
     bool bUse128bitRT,
 	EGBufferLayout GBufferLayout,
 	TShaderRef<TBasePassVertexShaderPolicyParamType<FUniformLightMapPolicy>>* VertexShader,
-	TShaderRef<TBasePassPixelShaderPolicyParamType<FUniformLightMapPolicy>>* PixelShader
+	TShaderRef<TBasePassPixelShaderPolicyParamType<FUniformLightMapPolicy>>* PixelShader,
+	bool bIsForOITPass
 	)
 {
 	switch (LightMapPolicy.GetIndirectPolicy())
 	{
 	case LMP_PRECOMPUTED_IRRADIANCE_VOLUME_INDIRECT_LIGHTING:
-		return GetUniformBasePassShaders<LMP_PRECOMPUTED_IRRADIANCE_VOLUME_INDIRECT_LIGHTING>(Material, VertexFactoryType, FeatureLevel, bEnableSkyLight, bUse128bitRT, GBufferLayout, VertexShader, PixelShader);
+		return GetUniformBasePassShaders<LMP_PRECOMPUTED_IRRADIANCE_VOLUME_INDIRECT_LIGHTING>(Material, VertexFactoryType, FeatureLevel, bEnableSkyLight, bUse128bitRT, GBufferLayout, VertexShader, PixelShader, bIsForOITPass);
 	case LMP_CACHED_VOLUME_INDIRECT_LIGHTING:
-		return GetUniformBasePassShaders<LMP_CACHED_VOLUME_INDIRECT_LIGHTING>(Material, VertexFactoryType, FeatureLevel, bEnableSkyLight, bUse128bitRT, GBufferLayout, VertexShader, PixelShader);
+		return GetUniformBasePassShaders<LMP_CACHED_VOLUME_INDIRECT_LIGHTING>(Material, VertexFactoryType, FeatureLevel, bEnableSkyLight, bUse128bitRT, GBufferLayout, VertexShader, PixelShader, bIsForOITPass);
 	case LMP_CACHED_POINT_INDIRECT_LIGHTING:
-		return GetUniformBasePassShaders<LMP_CACHED_POINT_INDIRECT_LIGHTING>(Material, VertexFactoryType, FeatureLevel, bEnableSkyLight, bUse128bitRT, GBufferLayout, VertexShader, PixelShader);
+		return GetUniformBasePassShaders<LMP_CACHED_POINT_INDIRECT_LIGHTING>(Material, VertexFactoryType, FeatureLevel, bEnableSkyLight, bUse128bitRT, GBufferLayout, VertexShader, PixelShader, bIsForOITPass);
 	case LMP_LQ_LIGHTMAP:
-		return GetUniformBasePassShaders<LMP_LQ_LIGHTMAP>(Material, VertexFactoryType, FeatureLevel, bEnableSkyLight, bUse128bitRT, GBufferLayout, VertexShader, PixelShader);
+		return GetUniformBasePassShaders<LMP_LQ_LIGHTMAP>(Material, VertexFactoryType, FeatureLevel, bEnableSkyLight, bUse128bitRT, GBufferLayout, VertexShader, PixelShader, bIsForOITPass);
 	case LMP_HQ_LIGHTMAP:
-		return GetUniformBasePassShaders<LMP_HQ_LIGHTMAP>(Material, VertexFactoryType, FeatureLevel, bEnableSkyLight, bUse128bitRT, GBufferLayout, VertexShader, PixelShader);
+		return GetUniformBasePassShaders<LMP_HQ_LIGHTMAP>(Material, VertexFactoryType, FeatureLevel, bEnableSkyLight, bUse128bitRT, GBufferLayout, VertexShader, PixelShader, bIsForOITPass);
 	case LMP_DISTANCE_FIELD_SHADOWS_AND_HQ_LIGHTMAP:
-		return GetUniformBasePassShaders<LMP_DISTANCE_FIELD_SHADOWS_AND_HQ_LIGHTMAP>(Material, VertexFactoryType, FeatureLevel, bEnableSkyLight, bUse128bitRT, GBufferLayout, VertexShader, PixelShader);
+		return GetUniformBasePassShaders<LMP_DISTANCE_FIELD_SHADOWS_AND_HQ_LIGHTMAP>(Material, VertexFactoryType, FeatureLevel, bEnableSkyLight, bUse128bitRT, GBufferLayout, VertexShader, PixelShader, bIsForOITPass);
 	case LMP_NO_LIGHTMAP:
-		return GetUniformBasePassShaders<LMP_NO_LIGHTMAP>(Material, VertexFactoryType, FeatureLevel, bEnableSkyLight, bUse128bitRT, GBufferLayout, VertexShader, PixelShader);
+		return GetUniformBasePassShaders<LMP_NO_LIGHTMAP>(Material, VertexFactoryType, FeatureLevel, bEnableSkyLight, bUse128bitRT, GBufferLayout, VertexShader, PixelShader, bIsForOITPass);
 	default:
 		check(false);
 		return false;
@@ -1627,7 +1639,8 @@ void FBasePassMeshProcessor::CollectPSOInitializersForLMPolicy(
 		Get128BitRequirement(),
 		GBL_Default, // Currently only Nanite supports non-default layout
 		&BasePassShaders.VertexShader,
-		&BasePassShaders.PixelShader
+		&BasePassShaders.PixelShader,
+		bOITBasePass
 		))
 	{
 		return;
@@ -1767,7 +1780,8 @@ bool FBasePassMeshProcessor::Process(
 		Get128BitRequirement(),
 		GBL_Default, // Currently only Nanite uses non-default layout
 		&BasePassShaders.VertexShader,
-		&BasePassShaders.PixelShader
+		&BasePassShaders.PixelShader,
+		bOITBasePass
 		))
 	{
 		return false;
@@ -2412,6 +2426,7 @@ FBasePassMeshProcessor::FBasePassMeshProcessor(
 	, PassDrawRenderState(InDrawRenderState)
 	, TranslucencyPassType(InTranslucencyPassType)
 	, bTranslucentBasePass(InTranslucencyPassType != ETranslucencyPass::TPT_MAX)
+	, bOITBasePass(InTranslucencyPassType != ETranslucencyPass::TPT_MAX && OIT::IsSortedPixelsEnabled(GetFeatureLevelShaderPlatform(InFeatureLevel)))
 	, bEnableReceiveDecalOutput((Flags & EFlags::CanUseDepthStencil) == EFlags::CanUseDepthStencil)
 	, EarlyZPassMode(Scene ? Scene->EarlyZPassMode : DDM_None)
 	, bRequiresExplicit128bitRT((Flags & EFlags::bRequires128bitRT) == EFlags::bRequires128bitRT)
diff --git a/Engine/Source/Runtime/Renderer/Private/BasePassRendering.h b/Engine/Source/Runtime/Renderer/Private/BasePassRendering.h
--- a/Engine/Source/Runtime/Renderer/Private/BasePassRendering.h
+++ b/Engine/Source/Runtime/Renderer/Private/BasePassRendering.h
@@ -44,6 +44,11 @@ struct FSceneWithoutWaterTextures;
 class FViewInfo;
 class UMaterialExpressionSingleLayerWaterMaterialOutput;
 
+namespace OIT
+{
+	bool IsSortedPixelsEnabledForProject(EShaderPlatform InPlatform);
+}
+
 /** Whether to allow the indirect lighting cache to be applied to dynamic objects. */
 extern int32 GIndirectLightingCache;
 
@@ -560,6 +565,9 @@ template<typename LightMapPolicyType, bool bEnableSkyLight, EGBufferLayout GBuff
 class TBasePassPS : public TBasePassPixelShaderBaseType<LightMapPolicyType>
 {
 	DECLARE_SHADER_TYPE(TBasePassPS,MeshMaterial);
+
+	class FSupportOITDim : SHADER_PERMUTATION_BOOL("PERMUTATION_SUPPORTS_OIT");
+	using FPermutationDomain = TShaderPermutationDomain<FSupportOITDim>;
 public:
 
 	static bool ShouldCompilePermutation(const FMeshMaterialShaderPermutationParameters& Parameters)
@@ -574,6 +582,16 @@ public:
 		const bool bForceAllPermutations = SupportAllShaderPermutations && SupportAllShaderPermutations->GetValueOnAnyThread() != 0;
 		const bool bProjectSupportsStationarySkylight = !SupportStationarySkylight || SupportStationarySkylight->GetValueOnAnyThread() != 0 || bForceAllPermutations;
 
+		// Only compiled OIT permutation for translucent surface, and if OIT sorted pixel is enabled for the current project
+		FPermutationDomain PermutationVector{ Parameters.PermutationId };
+		if (PermutationVector.template Get<FSupportOITDim>())
+		{
+			if (!bTranslucent || !OIT::IsSortedPixelsEnabledForProject(Parameters.Platform))
+			{
+				return false;
+			}
+		}
+
 		const bool bCacheShaders = !bEnableSkyLight
 			//translucent materials need to compile skylight support to support MOVABLE skylights also.
 			|| bTranslucent
@@ -683,15 +701,25 @@ bool GetBasePassShader<FUniformLightMapPolicy>(
  */
 
 template <typename LightMapPolicyType, EGBufferLayout GBufferLayout>
-void AddBasePassPixelShader(bool bEnableSkyLight, FMaterialShaderTypes& OutShaderTypes)
+void AddBasePassPixelShader(bool bEnableSkyLight, FMaterialShaderTypes& OutShaderTypes, bool bIsForOITPass = false)
 {
+	int32 PermutationId = 0;
+
+	if (bIsForOITPass)
+	{
+		using FMyShader = TBasePassPS<LightMapPolicyType, true, GBufferLayout>;
+		typename FMyShader::FPermutationDomain PermutationVector;
+		PermutationVector.template Set<typename FMyShader::FSupportOITDim>(true);
+		PermutationId = PermutationVector.ToDimensionValueId();
+	}
+
 	if (bEnableSkyLight)
 	{
-		OutShaderTypes.AddShaderType<TBasePassPS<LightMapPolicyType, true, GBufferLayout>>();
+		OutShaderTypes.AddShaderType<TBasePassPS<LightMapPolicyType, true, GBufferLayout>>(PermutationId);
 	}
 	else
 	{
-		OutShaderTypes.AddShaderType<TBasePassPS<LightMapPolicyType, false, GBufferLayout>>();
+		OutShaderTypes.AddShaderType<TBasePassPS<LightMapPolicyType, false, GBufferLayout>>(PermutationId);
 	}
 }
 
@@ -705,7 +733,8 @@ bool GetBasePassShaders(
 	bool bUse128bitRT,
 	EGBufferLayout GBufferLayout,
 	TShaderRef<TBasePassVertexShaderPolicyParamType<LightMapPolicyType>>* VertexShader,
-	TShaderRef<TBasePassPixelShaderPolicyParamType<LightMapPolicyType>>* PixelShader
+	TShaderRef<TBasePassPixelShaderPolicyParamType<LightMapPolicyType>>* PixelShader,
+	bool bIsForOITPass = false
 )
 {
 	FMaterialShaderTypes ShaderTypes;
@@ -719,10 +748,10 @@ bool GetBasePassShaders(
 		switch (GBufferLayout)
 		{
 		case GBL_Default:
-			AddBasePassPixelShader<LightMapPolicyType, GBL_Default>(bEnableSkyLight, ShaderTypes);
+			AddBasePassPixelShader<LightMapPolicyType, GBL_Default>(bEnableSkyLight, ShaderTypes, bIsForOITPass);
 			break;
 		case GBL_ForceVelocity:
-			AddBasePassPixelShader<LightMapPolicyType, GBL_ForceVelocity>(bEnableSkyLight, ShaderTypes);
+			AddBasePassPixelShader<LightMapPolicyType, GBL_ForceVelocity>(bEnableSkyLight, ShaderTypes, bIsForOITPass);
 			break;
 		default:
 			check(false);
@@ -751,7 +780,8 @@ bool GetBasePassShaders<FUniformLightMapPolicy>(
 	bool bUse128bitRT,
 	EGBufferLayout GBufferLayout,
 	TShaderRef<TBasePassVertexShaderPolicyParamType<FUniformLightMapPolicy>>* VertexShader,
-	TShaderRef<TBasePassPixelShaderPolicyParamType<FUniformLightMapPolicy>>* PixelShader
+	TShaderRef<TBasePassPixelShaderPolicyParamType<FUniformLightMapPolicy>>* PixelShader,
+	bool bIsForOITPass
 	);
 
 class FBasePassMeshProcessor : public FSceneRenderingAllocatorObject<FBasePassMeshProcessor>, public FMeshPassProcessor
@@ -842,6 +872,7 @@ private:
 
 	const ETranslucencyPass::Type TranslucencyPassType;
 	const bool bTranslucentBasePass;
+	const bool bOITBasePass;
 	const bool bEnableReceiveDecalOutput;
 	EDepthDrawingMode EarlyZPassMode;
 	bool bRequiresExplicit128bitRT;
diff --git a/Engine/Source/Runtime/Renderer/Private/DeferredShadingRenderer.cpp b/Engine/Source/Runtime/Renderer/Private/DeferredShadingRenderer.cpp
--- a/Engine/Source/Runtime/Renderer/Private/DeferredShadingRenderer.cpp
+++ b/Engine/Source/Runtime/Renderer/Private/DeferredShadingRenderer.cpp
@@ -4283,7 +4283,7 @@ void FDeferredShadingSceneRenderer::Render(FRDGBuilder& GraphBuilder)
 			// Sort objects' triangles
 			for (FViewInfo& View : Views)
 			{
-				if (OIT::IsEnabled(EOITSortingType::SortedTriangles, View))
+				if (OIT::IsSortedTrianglesEnabled(View.GetShaderPlatform()))
 				{
 					OIT::AddSortTrianglesPass(GraphBuilder, View, Scene->OITSceneData, FTriangleSortingOrder::BackToFront);
 				}
diff --git a/Engine/Source/Runtime/Renderer/Private/MeshDrawCommands.cpp b/Engine/Source/Runtime/Renderer/Private/MeshDrawCommands.cpp
--- a/Engine/Source/Runtime/Renderer/Private/MeshDrawCommands.cpp
+++ b/Engine/Source/Runtime/Renderer/Private/MeshDrawCommands.cpp
@@ -1095,7 +1095,17 @@ public:
 			{
 				// When per-pixel OIT is enabled, sort primitive from front to back ensure avoid 
 				// constantly resorting front-to-back samples list.
-				const bool bInverseSorting = OIT::IsEnabled(EOITSortingType::SortedPixels, Context.ShaderPlatform) && Context.View->AntiAliasingMethod != EAntiAliasingMethod::AAM_MSAA;
+				bool bInverseSorting = OIT::IsSortedPixelsEnabled(*Context.View);
+				if (Context.TranslucencyPass == ETranslucencyPass::TPT_AllTranslucency ||
+					Context.TranslucencyPass == ETranslucencyPass::TPT_TranslucencyStandard ||
+					Context.TranslucencyPass == ETranslucencyPass::TPT_TranslucencyStandardModulate)
+				{
+					bInverseSorting &= OIT::IsSortedPixelsEnabledForPass(OITPass_RegularTranslucency);
+				}
+				else
+				{
+					bInverseSorting &= OIT::IsSortedPixelsEnabledForPass(OITPass_SeperateTranslucency);
+				}
 
 				UpdateTranslucentMeshSortKeys(
 					Context.TranslucentSortPolicy,
diff --git a/Engine/Source/Runtime/Renderer/Private/MobileShadingRenderer.cpp b/Engine/Source/Runtime/Renderer/Private/MobileShadingRenderer.cpp
--- a/Engine/Source/Runtime/Renderer/Private/MobileShadingRenderer.cpp
+++ b/Engine/Source/Runtime/Renderer/Private/MobileShadingRenderer.cpp
@@ -1110,7 +1110,7 @@ void FMobileSceneRenderer::Render(FRDGBuilder& GraphBuilder)
 	// Sort objects' triangles
 	for (FViewInfo& View : Views)
 	{
-		if (View.ShouldRenderView() && OIT::IsEnabled(EOITSortingType::SortedTriangles, View))
+		if (View.ShouldRenderView() && OIT::IsSortedTrianglesEnabled(View.GetShaderPlatform()))
 		{
 			OIT::AddSortTrianglesPass(GraphBuilder, View, Scene->OITSceneData, FTriangleSortingOrder::BackToFront);
 		}
diff --git a/Engine/Source/Runtime/Renderer/Private/OIT/OIT.cpp b/Engine/Source/Runtime/Renderer/Private/OIT/OIT.cpp
--- a/Engine/Source/Runtime/Renderer/Private/OIT/OIT.cpp
+++ b/Engine/Source/Runtime/Renderer/Private/OIT/OIT.cpp
@@ -20,7 +20,7 @@
 ///////////////////////////////////////////////////////////////////////////////////////////////////
 // Variables: sorted triangles
 
-static TAutoConsoleVariable<int32> CVarOIT_SortedTriangles_Enable(
+static TAutoConsoleVariable<int32> CVarOIT_SortedTriangles_Enable_Project(
 	TEXT("r.OIT.SortedTriangles"), 
 	1, 
 	TEXT("Enable per-instance triangle sorting to avoid invalid triangle ordering."),
@@ -48,12 +48,18 @@ static TAutoConsoleVariable<int32> CVarOIT_SortedTriangles_PoolReleaseThreshold(
 // Variables: sorted pixels
 
 // Referenced in RendererSettings.h, as it is a project settings
-static TAutoConsoleVariable<int32> CVarOIT_SortedPixels_Enable(
+static TAutoConsoleVariable<int32> CVarOIT_SortedPixels_Enable_Project(
 	TEXT("r.OIT.SortedPixels"),
 	0,
 	TEXT("Enable OIT rendering (project settings, can't be changed at runtime)"),
 	ECVF_ReadOnly | ECVF_RenderThreadSafe);
 
+static TAutoConsoleVariable<int32> CVarOIT_SortedPixels_Enable_Runtime(
+	TEXT("r.OIT.SortedPixels.Enable"),
+	1,
+	TEXT("Enable OIT rendering (runtime setting, selects shader permutation)"),
+	ECVF_Scalability | ECVF_RenderThreadSafe);
+
 static TAutoConsoleVariable<int32> CVarOIT_SortedPixels_PassType(
 	TEXT("r.OIT.SortedPixels.PassType"),
 	3,
@@ -787,38 +793,29 @@ static void AddOITTriangleDebugPass(
 
 namespace OIT
 {
-	bool IsEnabled(EOITSortingType Type)
+	bool IsSortedTrianglesEnabled(EShaderPlatform InPlatform)
 	{
-		switch (Type)
-		{
-		case EOITSortingType::SortedTriangles: return CVarOIT_SortedTriangles_Enable.GetValueOnAnyThread() > 0;
-		case EOITSortingType::SortedPixels	 : return CVarOIT_SortedPixels_Enable.GetValueOnAnyThread() > 0;
-		}
-		return false;
+		return CVarOIT_SortedTriangles_Enable_Project.GetValueOnAnyThread() > 0;
 	}
 
-	bool IsEnabled(EOITSortingType Type, const FViewInfo& View)
+	bool IsSortedPixelsEnabledForProject(EShaderPlatform InPlatform)
 	{
-		switch (Type)
-		{
-		case EOITSortingType::SortedTriangles: return CVarOIT_SortedTriangles_Enable.GetValueOnAnyThread() > 0;
-		case EOITSortingType::SortedPixels	 : return CVarOIT_SortedPixels_Enable.GetValueOnAnyThread() > 0 && FDataDrivenShaderPlatformInfo::GetSupportsOIT(View.GetShaderPlatform());
-		}
-		return false;
-
-		
+		//const bool bMSAAEnabled = GetDefaultAntiAliasingMethod(GetMaxSupportedFeatureLevel(InPlatform)) != EAntiAliasingMethod::AAM_MSAA;
+		const bool bPixelOIT = CVarOIT_SortedPixels_Enable_Project.GetValueOnAnyThread() > 0 && FDataDrivenShaderPlatformInfo::GetSupportsOIT(EShaderPlatform(InPlatform));
+		return bPixelOIT;
 	}
 
-	bool IsEnabled(EOITSortingType Type, EShaderPlatform InPlatform)
+	bool InternalIsSortedPixelsEnabled(EShaderPlatform InPlatform, bool bMSAA)
 	{
-		const bool bMSAAEnabled = GetDefaultAntiAliasingMethod(GetMaxSupportedFeatureLevel(InPlatform)) != EAntiAliasingMethod::AAM_MSAA;
-
-		switch (Type)
-		{
-			case EOITSortingType::SortedTriangles: return CVarOIT_SortedTriangles_Enable.GetValueOnAnyThread() > 0;
-			case EOITSortingType::SortedPixels	 : return CVarOIT_SortedPixels_Enable.GetValueOnAnyThread() > 0 && FDataDrivenShaderPlatformInfo::GetSupportsOIT(EShaderPlatform(InPlatform));
-		}
-		return false;
+		return IsSortedPixelsEnabledForProject(InPlatform) && GRHISupportsRasterOrderViews && !!CVarOIT_SortedPixels_Enable_Runtime.GetValueOnRenderThread() && !bMSAA;
+	}
+	bool IsSortedPixelsEnabled(const FViewInfo& InView) 	{ return InternalIsSortedPixelsEnabled(InView.GetShaderPlatform(), InView.AntiAliasingMethod == EAntiAliasingMethod::AAM_MSAA); }
+	bool IsSortedPixelsEnabled(EShaderPlatform InPlatform)	{ return InternalIsSortedPixelsEnabled(InPlatform, false /*bMSAA*/); }
+	
+	bool IsSortedPixelsEnabledForPass(EOITPassType PassType)
+	{
+		const uint32 PassTypeBits = FMath::Clamp(CVarOIT_SortedPixels_PassType.GetValueOnRenderThread(), 0, 3);
+		return !!(PassTypeBits & PassType);
 	}
 
 	bool IsCompatible(const FMeshBatch& InMesh, ERHIFeatureLevel::Type InFeatureLevel)
@@ -835,7 +832,7 @@ namespace OIT
 
 	void AddSortTrianglesPass(FRDGBuilder& GraphBuilder, const FViewInfo& View, FOITSceneData& OITSceneData, FTriangleSortingOrder SortType)
 	{
-		if (!IsEnabled(EOITSortingType::SortedTriangles))
+		if (!IsSortedTrianglesEnabled(View.GetShaderPlatform()))
 		{
 			return;
 		}
@@ -904,10 +901,11 @@ namespace OIT
 
 	FOITData CreateOITData(FRDGBuilder& GraphBuilder, const FViewInfo& View, EOITPassType PassType)
 	{
-		const bool bOIT = IsEnabled(EOITSortingType::SortedPixels, View.GetShaderPlatform());
+		const bool bOIT = IsSortedPixelsEnabled(View);
 		const uint32 PassTypeBits = FMath::Clamp(CVarOIT_SortedPixels_PassType.GetValueOnRenderThread(), 0, 3);
 		const bool bPassValid = !!(PassTypeBits & PassType);
 		const uint32 LayerCount = 3u; /*Depth/Color/Trans*/
+
 		FOITData Out;
 		if (!bOIT || !bPassValid)
 		{
diff --git a/Engine/Source/Runtime/Renderer/Private/OIT/OIT.h b/Engine/Source/Runtime/Renderer/Private/OIT/OIT.h
--- a/Engine/Source/Runtime/Renderer/Private/OIT/OIT.h
+++ b/Engine/Source/Runtime/Renderer/Private/OIT/OIT.h
@@ -13,12 +13,6 @@ struct FMeshBatchElement;
 struct FMeshBatchElementDynamicIndexBuffer;
 class FViewInfo;
 
-enum class EOITSortingType
-{
-	SortedTriangles,
-	SortedPixels,
-};
-
 enum EOITPassType
 {
 	OITPass_None = 0,
@@ -87,13 +81,22 @@ struct FOITSceneData
 
 namespace OIT
 {
-	/* Return true if OIT techniques are enabled/supported */
-	bool IsEnabled(EOITSortingType Type, const FViewInfo& View);
-	bool IsEnabled(EOITSortingType Type, EShaderPlatform ShaderPlatform);
+	/* Return true if OIT sorted triangles is enabled/supported */
+	bool IsSortedTrianglesEnabled(EShaderPlatform InPlatform);
 
 	/* Return true if the current MeshBatch is compatible with per-instance sorted triangle */
 	bool IsCompatible(const FMeshBatch& Mesh, ERHIFeatureLevel::Type InFeatureLevel);
 
+	/* Return if OIT sorted pixel is enabled for the project*/
+	bool IsSortedPixelsEnabledForProject(EShaderPlatform InPlatform);
+
+	/* Return if OIT is enabled for the current runtime */
+	bool IsSortedPixelsEnabled(const FViewInfo& InView);
+	bool IsSortedPixelsEnabled(EShaderPlatform InPlatform); 
+
+	/* Return if OIT is enabled for the translucent pass current pass */
+	bool IsSortedPixelsEnabledForPass(EOITPassType PassType);
+
 	/* Sort triangles of all instances whose has the sorted triangle option enabled */
 	void AddSortTrianglesPass(FRDGBuilder& GraphBuilder, const FViewInfo& View, FOITSceneData& OITSceneData, FTriangleSortingOrder SortType);
 
diff --git a/Engine/Source/Runtime/Renderer/Private/PrimitiveSceneInfo.cpp b/Engine/Source/Runtime/Renderer/Private/PrimitiveSceneInfo.cpp
--- a/Engine/Source/Runtime/Renderer/Private/PrimitiveSceneInfo.cpp
+++ b/Engine/Source/Runtime/Renderer/Private/PrimitiveSceneInfo.cpp
@@ -1287,7 +1287,7 @@ void FPrimitiveSceneInfo::AddStaticMeshes(FScene* Scene, TArrayView<FPrimitiveSc
 		for (FPrimitiveSceneInfo* SceneInfo : SceneInfos)
 		{
 			// Allocate OIT index buffer where needed
-			const bool bAllocateSortedTriangles = OIT::IsEnabled(EOITSortingType::SortedTriangles, GMaxRHIShaderPlatform) && SceneInfo->Proxy->SupportsSortedTriangles();
+			const bool bAllocateSortedTriangles = OIT::IsSortedTrianglesEnabled(GMaxRHIShaderPlatform) && SceneInfo->Proxy->SupportsSortedTriangles();
 
 			for (int32 MeshIndex = 0; MeshIndex < SceneInfo->StaticMeshes.Num(); MeshIndex++)
 			{
@@ -1706,7 +1706,7 @@ void FPrimitiveSceneInfo::UpdateVirtualTextures(FScene* Scene, TArrayView<FPrimi
 void FPrimitiveSceneInfo::RemoveStaticMeshes()
 {
 	// Deallocate potential OIT dynamic index buffer
-	if (OIT::IsEnabled(EOITSortingType::SortedTriangles, GMaxRHIShaderPlatform))
+	if (OIT::IsSortedTrianglesEnabled(GMaxRHIShaderPlatform))
 	{
 		for (int32 MeshIndex = 0; MeshIndex < StaticMeshes.Num(); MeshIndex++)
 		{
diff --git a/Engine/Source/Runtime/Renderer/Private/TranslucentRendering.cpp b/Engine/Source/Runtime/Renderer/Private/TranslucentRendering.cpp
--- a/Engine/Source/Runtime/Renderer/Private/TranslucentRendering.cpp
+++ b/Engine/Source/Runtime/Renderer/Private/TranslucentRendering.cpp
@@ -746,7 +746,7 @@ const FRDGTextureDesc GetPostDOFTranslucentTextureDesc(
 	bool bIsModulate,
 	EShaderPlatform ShaderPlatform)
 {
-	const bool bNeedUAV = SeparateTranslucencyDimensions.NumSamples == 1 && OIT::IsEnabled(EOITSortingType::SortedPixels, ShaderPlatform);
+	const bool bNeedUAV = SeparateTranslucencyDimensions.NumSamples == 1 && OIT::IsSortedPixelsEnabled(ShaderPlatform);
 	return FRDGTextureDesc::Create2D(
 		SeparateTranslucencyDimensions.Extent,
 		bIsModulate ? PF_FloatR11G11B10 : PF_FloatRGBA,
