Flips r.TSR.Visualize -1 and -2 to allows more overviews in the negatives in the futur

diff --git a/Engine/Source/Runtime/Renderer/Private/PostProcess/TemporalSuperResolution.cpp b/Engine/Source/Runtime/Renderer/Private/PostProcess/TemporalSuperResolution.cpp
--- a/Engine/Source/Runtime/Renderer/Private/PostProcess/TemporalSuperResolution.cpp
+++ b/Engine/Source/Runtime/Renderer/Private/PostProcess/TemporalSuperResolution.cpp
@@ -308,10 +308,10 @@ TAutoConsoleVariable<int32> CVarTSRSubpixelIncludeMovingDepth(
 #if !UE_BUILD_OPTIMIZED_SHOWFLAGS
 
 TAutoConsoleVariable<int32> CVarTSRVisualize(
-	TEXT("r.TSR.Visualize"), -2,
+	TEXT("r.TSR.Visualize"), -1,
 	TEXT("Selects the TSR internal visualization mode.\n")
-	TEXT(" -2: Display an overview grid based on the VisualizeTSR show flag (default, opened with the `show VisualizeTSR` command at runtime or Show > Visualize > TSR in editor viewports);\n")
-	TEXT(" -1: Display an overview grid based regardless of VisualizeTSR show flag;\n")
+	TEXT(" -2: Display an overview grid based regardless of VisualizeTSR show flag;\n")
+	TEXT(" -1: Display an overview grid based on the VisualizeTSR show flag (default, opened with the `show VisualizeTSR` command at runtime or Show > Visualize > TSR in editor viewports);\n")
 	TEXT("  0: Number of accumulated samples in the history, particularily interesting to tune r.TSR.ShadingRejection.SampleCount and r.TSR.Velocity.WeightClampingSampleCount;\n")
 	TEXT("  1: Parallax disocclusion based of depth and velocity buffers;\n")
 	TEXT("  2: Mask where the history is rejected;\n")
@@ -1207,7 +1207,8 @@ bool IsVisualizeTSREnabled(const FViewInfo& View)
 }
 #else
 {
-	return GetMainTAAPassConfig(View) == EMainTAAPassConfig::TSR && (View.Family->EngineShowFlags.VisualizeTSR || CVarTSRVisualize.GetValueOnRenderThread() >= -1);
+	int32 VisualizeSettings = CVarTSRVisualize.GetValueOnRenderThread();
+	return GetMainTAAPassConfig(View) == EMainTAAPassConfig::TSR && (View.Family->EngineShowFlags.VisualizeTSR || VisualizeSettings != -1);
 }
 #endif
 
@@ -2637,7 +2638,7 @@ FDefaultTemporalUpscaler::FOutputs AddTemporalSuperResolutionPasses(
 		};
 		static_assert(UE_ARRAY_COUNT(kVisualizationName) == int32(EVisualizeId::MAX), "kVisualizationName doesn't match EVisualizeId");
 
-		const EVisualizeId Visualization = EVisualizeId(FMath::Clamp(CVarTSRVisualize.GetValueOnRenderThread(), -1, int32(EVisualizeId::MAX) - 1));
+		const EVisualizeId Visualization = EVisualizeId(FMath::Clamp(CVarTSRVisualize.GetValueOnRenderThread(), int32(EVisualizeId::Overview), int32(EVisualizeId::MAX) - 1));
 		FIntRect VisualizeRect = Visualization == EVisualizeId::Overview ? FIntRect(OutputRect.Min + OutputRect.Size() / 4, OutputRect.Min + (OutputRect.Size() * 3) / 4) : OutputRect;
 
 		auto Visualize = [&](EVisualizeId VisualizeId, FString Label)
