LightFunctionAtlas - new.
Store light functions in an atlas for all the views of a scene.
Replaces the volumetric fog only light function atlas.
Also can be applied to all lights of a scene, preventing them from going the not unbactched route.

Next:
- Debug of atlas
- RGB option on project + shader changes with defined type
- Apply during Forward shading option
- Super sample
- Convert other systems (Lumen, PT, SLW, etc.)

diff --git a/Engine/Shaders/Private/DeferredLightPixelShaders.usf b/Engine/Shaders/Private/DeferredLightPixelShaders.usf
--- a/Engine/Shaders/Private/DeferredLightPixelShaders.usf
+++ b/Engine/Shaders/Private/DeferredLightPixelShaders.usf
@@ -55,6 +55,8 @@
 #include "VirtualShadowMaps/VirtualShadowMapMaskBitsCommon.ush"
 #endif
 
+#include "LightFunctionAtlas/LightFunctionAtlasCommon.usf"
+
 // Remap light shape permutation into light type
 #if LIGHT_SOURCE_SHAPE == 0
 #define CURRENT_LIGHT_TYPE LIGHT_TYPE_DIRECTIONAL
@@ -94,6 +96,7 @@ FDerivedParams GetDerivedParams(in FInputParams Input, in float SceneDepth)
 	return Out;
 }
 
+int bUseLightFunctionAtlas;
 Texture2D<uint> LightingChannelsTexture;
 
 uint GetLightingChannelMask(float2 UV)
@@ -334,6 +337,13 @@ void DeferredLightPixelMain(
 		float4 LightAttenuation = GetLightAttenuationFromShadow(InputParams, SceneDepth, DerivedParams.TranslatedWorldPosition);
 		float4 Radiance = GetDynamicLighting(DerivedParams.TranslatedWorldPosition, DerivedParams.CameraVector, ScreenSpaceData.GBuffer, ScreenSpaceData.AmbientOcclusion, ScreenSpaceData.GBuffer.ShadingModelID, LightData, LightAttenuation, Dither, uint2(InputParams.PixelPos), SurfaceShadow);
 
+		if (bUseLightFunctionAtlas)
+		{
+			Radiance *= GetLocalLightFunctionCommon(
+				DerivedParams.TranslatedWorldPosition,
+				LightData.LightFunctionAtlasLightIndex);
+		}
+
 		OutColor += Radiance;
 	}
 
@@ -435,11 +445,18 @@ void DeferredLightPixelMain(
 
 	float SurfaceShadow = 1.0f;
 	const float4 Radiance = GetDynamicLighting(TranslatedWorldPosition, CameraVector, HairScreenSpaceData.GBuffer, HairScreenSpaceData.AmbientOcclusion, HairScreenSpaceData.GBuffer.ShadingModelID, LightData, LightAttenuation, Dither, PixelCoord, SurfaceShadow);
-	const float  Attenuation = ComputeLightProfileMultiplier(TranslatedWorldPosition, GetDeferredLightTranslatedWorldPosition(), -DeferredLightUniforms.Direction, DeferredLightUniforms.Tangent, DeferredLightUniforms.IESAtlasIndex);
+	float  Attenuation = ComputeLightProfileMultiplier(TranslatedWorldPosition, GetDeferredLightTranslatedWorldPosition(), -DeferredLightUniforms.Direction, DeferredLightUniforms.Tangent, DeferredLightUniforms.IESAtlasIndex);
 	const float3 AtmosphericTransmittance = GetAtmosphereAndCloudAttenuation(GetStableTranslatedLargeWorldPosition(ScreenPosition, Sample.Depth), UV);
 
 	const float LocalCoverage = From8bitCoverage(Sample.Coverage8bit);
 
+	if (bUseLightFunctionAtlas)
+	{
+		Attenuation *= GetLocalLightFunctionCommon(
+			TranslatedWorldPosition,
+			LightData.LightFunctionAtlasLightIndex);
+	}
+
 	OutColor.rgb = AtmosphericTransmittance * Radiance.xyz * Attenuation * LocalCoverage * GetExposure();
 	OutColor.a = LocalCoverage;
 }
diff --git a/Engine/Shaders/Private/LightData.ush b/Engine/Shaders/Private/LightData.ush
--- a/Engine/Shaders/Private/LightData.ush
+++ b/Engine/Shaders/Private/LightData.ush
@@ -128,6 +128,7 @@ struct FLightShaderParameters
 	float2 RectLightAtlasUVScale;
 	float  RectLightAtlasMaxLevel;
 	float  IESAtlasIndex;
+	uint  LightFunctionAtlasLightIndex;
 };
 
 ///////////////////////////////////////////////////////////////////////////////////////////////////
@@ -175,6 +176,8 @@ struct FDeferredLightData
 	FRectLightData RectLightData;
 	/** IES data. */
 	float  IESAtlasIndex;
+	// Light function atlas index 
+	uint  LightFunctionAtlasLightIndex;
 	/** Hair transmittance data. */
 	FHairTransmittanceData HairTransmittance;
 };
@@ -227,6 +230,12 @@ bool UnpackHasLightFunction(FLocalLightData In)
 	return ((LightTypeAndPackedShadowMapChannelMask >> 19) & 0x1) != 0;
 }
 
+uint UnpackLightFunctionAtlasIndex(FLocalLightData In)
+{
+	const uint LightTypeAndPackedShadowMapChannelMask = asuint(In.LightDirectionAndShadowMask.w);
+	return ((LightTypeAndPackedShadowMapChannelMask >> 20) & 0xFF);
+}
+
 float UnpackVolumetricScatteringIntensity(FLocalLightData In)
 {
 	return f16tof32(asuint(In.SpotAnglesAndSourceRadiusPacked.w) >> 16);
@@ -312,6 +321,7 @@ FDeferredLightData ConvertToDeferredLight(
 									asuint(In.RectData.y),
 									asuint(In.RectData.z));
 	Out.IESAtlasIndex			= UnpackLigthIESAtlasIndex(In);
+	Out.LightFunctionAtlasLightIndex = UnpackLightFunctionAtlasIndex(In);
 
 	// Shadow
 	// * Static shadowing uses ShadowMapChannel, 
@@ -391,6 +401,8 @@ FLightShaderParameters ConvertFromLocal(const FLocalLightData In)
 	Out.SoftSourceRadius		= UnpackLightSoftSourceRadius(In);
 	Out.SourceLength			= UnpackLightSourceLength(In);
 
+	Out.LightFunctionAtlasLightIndex = UnpackLightFunctionAtlasIndex(In);
+
 	const FRectLightData RectLightData = UnpackRectLightData(
 									asuint(In.RectData.x),
 									asuint(In.RectData.y),
@@ -468,6 +480,7 @@ FDeferredLightData ConvertToDeferredLight(FLightSceneData In)
 	Out.RectLightData.AtlasData.AtlasUVScale = In.RectLightAtlasUVScale;
 
 	Out.IESAtlasIndex = 0;
+	Out.LightFunctionAtlasLightIndex = 0;
 
 	// Shadow
 	// * Static shadowing uses ShadowMapChannel, 
diff --git a/Engine/Shaders/Private/LightDataUniforms.ush b/Engine/Shaders/Private/LightDataUniforms.ush
--- a/Engine/Shaders/Private/LightDataUniforms.ush
+++ b/Engine/Shaders/Private/LightDataUniforms.ush
@@ -40,6 +40,8 @@ FDeferredLightData InitDeferredLightFromUniforms(uint InLightType)
 	Out.RectLightData.AtlasData.AtlasUVOffset	= DeferredLightUniforms.RectLightAtlasUVOffset;
 	Out.RectLightData.AtlasData.AtlasUVScale	= DeferredLightUniforms.RectLightAtlasUVScale;
 
+	Out.LightFunctionAtlasLightIndex = DeferredLightUniforms.LightFunctionAtlasLightIndex;
+
 	Out.HairTransmittance		= InitHairTransmittanceData();
 	return Out;
 }
diff --git a/Engine/Shaders/Private/LightFunctionAtlas/LightFunctionAtlasCommon.usf b/Engine/Shaders/Private/LightFunctionAtlas/LightFunctionAtlasCommon.usf
new file mode 100644
--- /dev/null
+++ b/Engine/Shaders/Private/LightFunctionAtlas/LightFunctionAtlasCommon.usf
@@ -0,0 +1,76 @@
+/*=============================================================================
+	LightFunctionAtlasCommon.usf
+=============================================================================*/
+
+#define LIGHT_FUNCTION_ATLAS 1
+
+#include "../Common.ush"
+#include "/Engine/Shared/LightDefinitions.h"
+
+// In order to use function from that file, one must use FLightFunctionAtlasGlobalParameters
+
+float GetLocalLightFunctionCommon(float3 TranslatedWorldPosition, uint LightFunctionAtlasLightIndex)
+{
+	if (LightFunctionAtlasLightIndex == 0)
+	{
+		// No valid light function assigned to the light.
+		return 1.0f;
+	}
+
+	float4 LightFunctionLightDataParameters = LightFunctionAtlas.LightInfoDataParameters[LightFunctionAtlasLightIndex];
+	const float  FadeDistance				= LightFunctionLightDataParameters.x;
+	const uint   PackedLightInfoDataParams	= asuint(LightFunctionLightDataParameters.y);
+	const float2 Slot_MinUV					= LightFunctionLightDataParameters.zw;
+
+	const uint   LightType					= PackedLightInfoDataParams & 0xFF;
+	const float  DisabledBrightness			= f16tof32((PackedLightInfoDataParams >> 8) & 0xFFFF);
+
+	float4x4 ActualLightFunctionMatrix = LightFunctionAtlas.LightInfoDataMatrix[LightFunctionAtlasLightIndex];
+	float4 HomogeneousShadowPosition = mul(float4(TranslatedWorldPosition, 1), ActualLightFunctionMatrix);
+
+	float2 LightFunctionUV;
+	switch (LightType)
+	{
+
+	case LIGHT_TYPE_SPOT:
+	{
+		HomogeneousShadowPosition.xyz /= HomogeneousShadowPosition.www;
+		LightFunctionUV = HomogeneousShadowPosition.xy;
+		LightFunctionUV = 1 - LightFunctionUV;
+		break;
+	}
+
+	case LIGHT_TYPE_POINT:
+	{
+		float3 LightVec = normalize(HomogeneousShadowPosition).zyx;
+		LightFunctionUV = float2((atan2Fast(LightVec.y, LightVec.x) + PI) / (2 * PI), acosFast(LightVec.z) / PI);
+		break;
+	}
+
+	//case LIGHT_TYPE_RECT:
+	//case LIGHT_TYPE_DIRECTIONAL:
+	default:
+	{
+		LightFunctionUV = frac(HomogeneousShadowPosition.zy);
+		break;
+	}
+
+	}
+
+	// Slot min uv and size already account for half texel offsets
+	float2 AtlasUV = Slot_MinUV + saturate(LightFunctionUV) * LightFunctionAtlas.Slot_UVSize.xx;
+
+	float LightFunctionValue = Texture2DSampleLevel(LightFunctionAtlas.LightFunctionAtlasTexture, LightFunctionAtlas.LightFunctionAtlasSampler, AtlasUV, 0).x;
+	LightFunctionValue = LightFunctionValue * LightFunctionValue; // Simple Gamma decompression
+
+	// Calculate radial view distance for stable fading
+	float ViewDistance = length(PrimaryView.TranslatedWorldCameraOrigin - TranslatedWorldPosition);
+	float DistanceFadeAlpha = saturate((FadeDistance - ViewDistance) / (FadeDistance * .2f));
+	// Fade to disabled based on FadeDistance
+	LightFunctionValue = lerp(DisabledBrightness, LightFunctionValue, DistanceFadeAlpha);
+	// Fade to disabled based on ShadowFadeFraction
+	const float ShadowFadeFraction = 1.0f; // Not setup
+	LightFunctionValue = LightFunctionValue; // lerp(DisabledBrightness, LightFunctionValue, ShadowFadeFraction);
+
+	return LightFunctionValue;
+}
\ No newline at end of file
diff --git a/Engine/Shaders/Private/LightFunctionAtlas/LightFunctionAtlasRender.usf b/Engine/Shaders/Private/LightFunctionAtlas/LightFunctionAtlasRender.usf
new file mode 100644
--- /dev/null
+++ b/Engine/Shaders/Private/LightFunctionAtlas/LightFunctionAtlasRender.usf
@@ -0,0 +1,71 @@
+/*=============================================================================
+	RenderLightFunctionAtlas.usf.usf
+=============================================================================*/
+
+// This define is for the light function common evaluation to only be based on UV stored in the light vector only.
+#define LIGHT_FUNCTION_ATLAS 1
+
+#include "../Common.ush"
+#include "/Engine/Generated/Material.ush"
+#include "../LightFunctionCommon.ush"
+
+float2 LightFunctionTexelSize;
+float4 SvPositionToUVScaleBias;
+
+void Main(
+	in noperspective float4 UV : TEXCOORD0,
+	in float4 SvPosition : SV_Position,
+	out float4 OutColor : SV_Target0
+)
+{
+	float4 TranslatedWorldPosition = float4(0, 0, 0, 1);
+
+	float Depth = 0.000001f;
+	float2 ScreenPosition = 0.0f;
+
+	float3 LightVector = 0.0f;
+	float3 LightFunction = 0.0f;
+
+//	ScreenPosition = (SvPosition.xy * LightFunctionTexelSize - float2(0.5f, 0.5f)) * float2(2, -2);
+	float2 UVs = (SvPosition.xy - SvPositionToUVScaleBias.zw) * SvPositionToUVScaleBias.xy;
+
+#if 1
+
+	// Generate a local space light vector that is at least correct for point lights
+	float Theta = (UVs.x * (2 * PI)) - PI;
+	float Phi = UVs.y * PI;
+	float X = cos(Theta) * sin(Phi);
+	float Y = sin(Theta) * sin(Phi);
+	float Z = cos(Phi);
+	LightVector = normalize(float3(X, Y, Z)); 
+
+	LightFunction = GetLightFunctionColor(LightVector, TranslatedWorldPosition.xyz, UVs.xy);
+
+#else
+
+	// TODO Super sample 4x
+	const float2 UVOffset = 0.25f * LightFunctionTexelSize;
+
+	LightVector.xy = UVs.xy + UVOffset * float2(-1.0f, -1.0f);
+	LightFunction += GetLightFunctionColor(LightVector, TranslatedWorldPosition.xyz);
+
+	LightVector.xy = UVs.xy + UVOffset * float2(-1.0f, 1.0f);
+	LightFunction += GetLightFunctionColor(LightVector, TranslatedWorldPosition.xyz);
+
+	LightVector.xy = UVs.xy + UVOffset * float2(1.0f, 1.0f);
+	LightFunction += GetLightFunctionColor(LightVector, TranslatedWorldPosition.xyz);
+
+	LightVector.xy = UVs.xy + UVOffset * float2(1.0f, -1.0f);
+	LightFunction += GetLightFunctionColor(LightVector, TranslatedWorldPosition.xyz);
+
+	LightFunction *= 0.25f;
+
+#endif
+
+
+	float GreyScale = dot(LightFunction, .3333f).x;
+
+	GreyScale = sqrt(max(0.0f, GreyScale)); // Simple Gamma compression
+
+	OutColor = GreyScale;
+}
diff --git a/Engine/Shaders/Private/LightFunctionCommon.ush b/Engine/Shaders/Private/LightFunctionCommon.ush
--- a/Engine/Shaders/Private/LightFunctionCommon.ush
+++ b/Engine/Shaders/Private/LightFunctionCommon.ush
@@ -6,6 +6,11 @@
 
 #include "/Engine/Private/Substrate/Substrate.ush" 
 
+
+#ifndef LIGHT_FUNCTION_ATLAS
+#define LIGHT_FUNCTION_ATLAS 0
+#endif
+
 // This ifdef is to allow shaders(notably the ratracing miss shader implementing light functions) to override where these parameters come from.
 #ifndef LightFunctionParameters
 /** Tan of spotlight cone outer angle in x, ShadowFadeFraction in y, IsSpotLight in z, IsPointLight in w. */
@@ -17,8 +22,20 @@ float4x4 LightFunctionTranslatedWorldToLight;
  * Calculates the light function color with the given light vector. 
  * LightVector is the vector from the light to the position being shaded, in the light's coordinate space.
  */
-float3 GetLightFunctionColor(float3 LightVector, float3 TranslatedWorldSpace)
+float3 GetLightFunctionColor(float3 LightVector, float3 TranslatedWorldSpace
+#if LIGHT_FUNCTION_ATLAS
+	, float2 UVs
+#endif
+)
 {
+
+#if LIGHT_FUNCTION_ATLAS
+
+	// This path is used when rendering a light function into a 2D Atlas
+	float2 LightFunctionUVs = UVs;
+
+#else // LIGHT_FUNCTION_ATLAS
+
 	// Swizzle so that LightVector.xy are perpendicular to the light direction and LightVector.z is distance along the light direction
 	LightVector.xyz = LightVector.zyx;
 
@@ -39,6 +56,8 @@ float3 GetLightFunctionColor(float3 LightVector, float3 TranslatedWorldSpace)
 		LightFunctionUVs = float2((atan2(UnitLightVector.y, UnitLightVector.x) + PI) / (2 * PI), acos(UnitLightVector.z) / PI);
 	}
 
+#endif // LIGHT_FUNCTION_ATLAS
+
 	FMaterialPixelParameters MaterialParameters = MakeInitializedMaterialPixelParameters();
 #if NUM_MATERIAL_TEXCOORDS
 	for(int CoordinateIndex = 0;CoordinateIndex < NUM_MATERIAL_TEXCOORDS;CoordinateIndex++)
diff --git a/Engine/Shaders/Private/VolumetricFog.usf b/Engine/Shaders/Private/VolumetricFog.usf
--- a/Engine/Shaders/Private/VolumetricFog.usf
+++ b/Engine/Shaders/Private/VolumetricFog.usf
@@ -326,6 +326,8 @@ int VirtualShadowMapId;
 #define HISTORY_MISS_SUPER_SAMPLE_COUNT 1
 #endif
 
+#include "LightFunctionAtlas/LightFunctionAtlasCommon.usf"
+
 #if	USE_LIGHT_FUNCTION
 
 /** Matrix can be either WorldToShadow or WorldToLight depending on light type. */
@@ -335,40 +337,14 @@ float4 LightFunctionAtlasTileMinMaxUvBound;
 Texture2D LightFunctionAtlasTexture;
 SamplerState LightFunctionAtlasSampler;
 
-float GetLocalLightFunction(float3 TranslatedWorldPosition, bool bSpotLight, bool bRectLight)
+float GetLocalLightFunction(float3 TranslatedWorldPosition, FDeferredLightData LightData)
 {
-	float2 LightFunctionUV;
-
-	if(bSpotLight)
-	{
-		float4 HomogeneousShadowPosition = mul(float4(TranslatedWorldPosition, 1), LocalLightFunctionMatrix);
-		HomogeneousShadowPosition.xyz /= HomogeneousShadowPosition.w;
-		LightFunctionUV = HomogeneousShadowPosition.xy;
-		LightFunctionUV.x = 1 - LightFunctionUV.x;
-		LightFunctionUV.y = 1 - LightFunctionUV.y;
-	}
-	else if (bRectLight)
-	{
-		float4 HomogeneousShadowPosition = mul(float4(TranslatedWorldPosition, 1), LocalLightFunctionMatrix).zyxw;
-		HomogeneousShadowPosition.xyz /= HomogeneousShadowPosition.w;
-		LightFunctionUV = HomogeneousShadowPosition.xy;
-
-		LightFunctionUV = LightFunctionUV * 0.5 + 0.5;
-		LightFunctionUV = frac(LightFunctionUV);
-	}
-	else // Point light
-	{
-		// Get vector from light to cell world pos.
-		float3 LightVec = normalize(TranslatedWorldPosition - DeferredLightUniforms.TranslatedWorldPosition);
-		// Transform and swizzle
-		LightVec = mul(float4(LightVec.xyz, 0), LocalLightFunctionMatrix).zyx;
-		// Map to spherical coordinate
-		LightFunctionUV = float2((atan2Fast(LightVec.y, LightVec.x) + PI) / (2 * PI), acosFast(LightVec.z) / PI);
-	}
-	
-	return Texture2DSampleLevel(LightFunctionAtlasTexture, LightFunctionAtlasSampler, LightFunctionAtlasTileMinMaxUvBound.xy + LightFunctionUV * (LightFunctionAtlasTileMinMaxUvBound.zw - LightFunctionAtlasTileMinMaxUvBound.xy), 0).x;
+	return GetLocalLightFunctionCommon(
+		TranslatedWorldPosition, 
+		LightData.LightFunctionAtlasLightIndex);
 }
-#endif
+
+#endif // USE_LIGHT_FUNCTION
 
 Texture2D<float> ConservativeDepthTexture;
 float2 PrevConservativeDepthTextureSize;
@@ -503,7 +479,7 @@ float4 InjectShadowedLocalLightCommon(uint3 GridCoordinate)
 #endif
 	
 #if	USE_LIGHT_FUNCTION
-			ShadowFactor *= GetLocalLightFunction(TranslatedWorldPosition, LightData.bSpotLight, LightData.bRectLight);
+			ShadowFactor *= GetLocalLightFunction(TranslatedWorldPosition, LightData);
 #endif
 
 			OutScattering.rgb += LightColor * (PhaseFunction(PhaseG, dot(L, -CameraVector)) * CombinedAttenuation * ShadowFactor * VolumetricScatteringIntensity);
@@ -687,16 +663,25 @@ float ComputeSkyVisibility(float3 TranslatedWorldPosition, float3 BrickTextureUV
 }
 
 float4x4 DirectionalLightFunctionTranslatedWorldToShadow;
-Texture2D LightFunctionTexture;
-SamplerState LightFunctionSampler;
+Texture2D DirectionalLightLightFunctionTexture;
+SamplerState DirectionalLightLightFunctionSampler;
+uint DirectionalApplyLightFunctionFromAtlas;
+uint DirectionalLightFunctionAtlasLightIndex;
 
 float GetDirectionalLightFunction(float3 TranslatedWorldPosition) // directional light
 {
+#if USE_LIGHT_FUNCTION_ATLAS
+	if (DirectionalApplyLightFunctionFromAtlas > 0)
+	{
+		return GetLocalLightFunctionCommon(TranslatedWorldPosition, DirectionalLightFunctionAtlasLightIndex);
+	}
+#endif
+
 	float4 HomogeneousShadowPosition = mul(float4(TranslatedWorldPosition, 1), DirectionalLightFunctionTranslatedWorldToShadow);
 	float2 LightFunctionUV = HomogeneousShadowPosition.xy * .5f + .5f;
 	LightFunctionUV.y = 1 - LightFunctionUV.y;
 
-	return Texture2DSampleLevel(LightFunctionTexture, LightFunctionSampler, LightFunctionUV, 0).x;
+	return Texture2DSampleLevel(DirectionalLightLightFunctionTexture, DirectionalLightLightFunctionSampler, LightFunctionUV, 0).x;
 }
 
 uint SampleSkyLightDiffuseEnvMap;
@@ -925,6 +910,12 @@ void LightScatteringCS(
 
 				float CombinedAttenuation = Lighting * LightMask;
 
+			#if USE_LIGHT_FUNCTION_ATLAS
+				CombinedAttenuation *= GetLocalLightFunctionCommon(
+					TranslatedWorldPosition,
+					LightData.LightFunctionAtlasLightIndex);
+			#endif
+
 				LightScattering += LightData.Color * (PhaseFunction(PhaseG, dot(L, -CameraVector)) * CombinedAttenuation * VolumetricScatteringIntensity);
 
 				// To debug culling
diff --git a/Engine/Shaders/Public/ShaderVersion.ush b/Engine/Shaders/Public/ShaderVersion.ush
--- a/Engine/Shaders/Public/ShaderVersion.ush
+++ b/Engine/Shaders/Public/ShaderVersion.ush
@@ -3,4 +3,4 @@
 // Each time the console command is executed it generates a new GUID. As this file is included
 // in Platform.ush (which should be included in any shader) it allows to invalidate the shader DDC.
 // 
-#pragma message("UESHADERMETADATA_VERSION 6169D26A-4079-3DED-6B9D-25878110CFAC")
\ No newline at end of file
+#pragma message("UESHADERMETADATA_VERSION 08294D12-4290-46A4-8DD7-0FDE597F2914")
diff --git a/Engine/Source/Runtime/Engine/Private/Components/DirectionalLightComponent.cpp b/Engine/Source/Runtime/Engine/Private/Components/DirectionalLightComponent.cpp
--- a/Engine/Source/Runtime/Engine/Private/Components/DirectionalLightComponent.cpp
+++ b/Engine/Source/Runtime/Engine/Private/Components/DirectionalLightComponent.cpp
@@ -405,6 +405,7 @@ public:
 		LightParameters.RectLightAtlasUVScale = FVector2f::ZeroVector;
 		LightParameters.RectLightAtlasMaxLevel = FLightRenderParameters::GetRectLightAtlasInvalidMIPLevel();
 		LightParameters.IESAtlasIndex = INDEX_NONE;
+		LightParameters.LightFunctionAtlasLightIndex = GetLightFunctionAtlasLightIndex();
 		LightParameters.InverseExposureBlend = 0.0f;
 	}
 
diff --git a/Engine/Source/Runtime/Engine/Private/Components/LightComponent.cpp b/Engine/Source/Runtime/Engine/Private/Components/LightComponent.cpp
--- a/Engine/Source/Runtime/Engine/Private/Components/LightComponent.cpp
+++ b/Engine/Source/Runtime/Engine/Private/Components/LightComponent.cpp
@@ -350,6 +350,7 @@ void FLightRenderParameters::MakeShaderParameters(const FViewMatrices& ViewMatri
 	OutShaderParameters.RectLightAtlasUVScale = RectLightAtlasUVScale;
 	OutShaderParameters.RectLightAtlasMaxLevel = RectLightAtlasMaxLevel;
 	OutShaderParameters.IESAtlasIndex = IESAtlasIndex;
+	OutShaderParameters.LightFunctionAtlasLightIndex = LightFunctionAtlasLightIndex;
 }
 
 // match logic in EyeAdaptationInverseLookup(...)
diff --git a/Engine/Source/Runtime/Engine/Private/Components/PointLightComponent.cpp b/Engine/Source/Runtime/Engine/Private/Components/PointLightComponent.cpp
--- a/Engine/Source/Runtime/Engine/Private/Components/PointLightComponent.cpp
+++ b/Engine/Source/Runtime/Engine/Private/Components/PointLightComponent.cpp
@@ -55,6 +55,7 @@ void FPointLightSceneProxy::GetLightShaderParameters(FLightRenderParameters& Lig
 	LightParameters.RectLightAtlasMaxLevel = FLightRenderParameters::GetRectLightAtlasInvalidMIPLevel();
 	LightParameters.IESAtlasIndex = INDEX_NONE;
 	LightParameters.InverseExposureBlend = InverseExposureBlend;
+	LightParameters.LightFunctionAtlasLightIndex = GetLightFunctionAtlasLightIndex();
 
 	if (IESAtlasId != ~0)
 	{
diff --git a/Engine/Source/Runtime/Engine/Private/Components/RectLightComponent.cpp b/Engine/Source/Runtime/Engine/Private/Components/RectLightComponent.cpp
--- a/Engine/Source/Runtime/Engine/Private/Components/RectLightComponent.cpp
+++ b/Engine/Source/Runtime/Engine/Private/Components/RectLightComponent.cpp
@@ -249,6 +249,7 @@ void FRectLightSceneProxy::GetLightShaderParameters(FLightRenderParameters& Ligh
 	LightParameters.RectLightAtlasMaxLevel = FLightRenderParameters::GetRectLightAtlasInvalidMIPLevel();
 	LightParameters.IESAtlasIndex = INDEX_NONE;
 	LightParameters.InverseExposureBlend = InverseExposureBlend;
+	LightParameters.LightFunctionAtlasLightIndex = GetLightFunctionAtlasLightIndex();
 
 	if (IESAtlasId != ~0)
 	{
diff --git a/Engine/Source/Runtime/Engine/Private/Components/SpotLightComponent.cpp b/Engine/Source/Runtime/Engine/Private/Components/SpotLightComponent.cpp
--- a/Engine/Source/Runtime/Engine/Private/Components/SpotLightComponent.cpp
+++ b/Engine/Source/Runtime/Engine/Private/Components/SpotLightComponent.cpp
@@ -75,6 +75,7 @@ public:
 		LightParameters.RectLightAtlasUVScale = FVector2f::ZeroVector;
 		LightParameters.RectLightAtlasMaxLevel = FLightRenderParameters::GetRectLightAtlasInvalidMIPLevel();
 		LightParameters.IESAtlasIndex = INDEX_NONE;
+		LightParameters.LightFunctionAtlasLightIndex = GetLightFunctionAtlasLightIndex();
 		LightParameters.InverseExposureBlend = InverseExposureBlend;
 
 		if (IESAtlasId != uint32(INDEX_NONE))
diff --git a/Engine/Source/Runtime/Engine/Private/LightSceneProxy.cpp b/Engine/Source/Runtime/Engine/Private/LightSceneProxy.cpp
--- a/Engine/Source/Runtime/Engine/Private/LightSceneProxy.cpp
+++ b/Engine/Source/Runtime/Engine/Private/LightSceneProxy.cpp
@@ -67,6 +67,7 @@ FLightSceneProxy::FLightSceneProxy(const ULightComponent* InLightComponent)
 	, SamplesPerPixel(1)
 	, DeepShadowLayerDistribution(InLightComponent->DeepShadowLayerDistribution)
 	, IESAtlasId(~0u)
+	, LightFunctionAtlasLightIndex(0)
 #if ACTOR_HAS_LABELS
 	, OwnerNameOrLabel(InLightComponent->GetOwner() ? InLightComponent->GetOwner()->GetActorNameOrLabel() : InLightComponent->GetName())
 #endif
diff --git a/Engine/Source/Runtime/Engine/Public/LightSceneProxy.h b/Engine/Source/Runtime/Engine/Public/LightSceneProxy.h
--- a/Engine/Source/Runtime/Engine/Public/LightSceneProxy.h
+++ b/Engine/Source/Runtime/Engine/Public/LightSceneProxy.h
@@ -283,6 +283,10 @@ public:
 	virtual FLinearColor GetCloudScatteredLuminanceScale() const { return FLinearColor::White; }
 	virtual bool GetUsePerPixelAtmosphereTransmittance() const { return false; }
 
+	inline void  SetLightFunctionAtlasIndices(uint8 LightIndex) { LightFunctionAtlasLightIndex = LightIndex; }
+	inline bool  HasValidLightFunctionAtlasSlot() const { return LightFunctionAtlasLightIndex != 0; }
+	inline uint8 GetLightFunctionAtlasLightIndex() const { return LightFunctionAtlasLightIndex; }
+
 protected:
 
 	friend class FScene;
@@ -294,7 +298,7 @@ protected:
 	/** The scene the primitive is in. */
 	FSceneInterface* SceneInterface;
 
-	/** The homogenous position of the light. */
+	/** The homogeneous position of the light. */
 	FVector4 Position;
 
 	/** The light color. */
@@ -472,6 +476,12 @@ protected:
 	/** IES texture atlas id. */
 	uint32 IESAtlasId;
 
+	/**
+	 * The light index in order to be able to read matrix and parameters when reading the light function atlas for that light.
+	 * A value of 0 means this is the default identity light function and no light function sampling will be done in shader.
+	 */
+	uint8 LightFunctionAtlasLightIndex;
+
 	/**
 	 * Updates the light proxy's cached transforms.
 	 * @param InLightToWorld - The new light-to-world transform.
diff --git a/Engine/Source/Runtime/Engine/Public/SceneManagement.h b/Engine/Source/Runtime/Engine/Public/SceneManagement.h
--- a/Engine/Source/Runtime/Engine/Public/SceneManagement.h
+++ b/Engine/Source/Runtime/Engine/Public/SceneManagement.h
@@ -1268,9 +1268,12 @@ BEGIN_SHADER_PARAMETER_STRUCT(FLightShaderParameters, ENGINE_API)
 	SHADER_PARAMETER(FVector2f, RectLightAtlasUVScale)
 	SHADER_PARAMETER(float, RectLightAtlasMaxLevel)
 
-	// UES texture slice index
+	// IES texture slice index
 	SHADER_PARAMETER(float, IESAtlasIndex)
 
+	// Index of the light function in the atlas
+	SHADER_PARAMETER(uint32, LightFunctionAtlasLightIndex)
+
 END_SHADER_PARAMETER_STRUCT()
 
 
@@ -1343,6 +1346,9 @@ struct FLightRenderParameters
 	// IES atlas slice index
 	float IESAtlasIndex;
 
+	// Index of the light in the Light function atlas data
+	uint32 LightFunctionAtlasLightIndex;
+
 	float InverseExposureBlend;
 
 	// Return Invalid rect light atlas MIP level
diff --git a/Engine/Source/Runtime/Renderer/Private/DeferredShadingRenderer.cpp b/Engine/Source/Runtime/Renderer/Private/DeferredShadingRenderer.cpp
--- a/Engine/Source/Runtime/Renderer/Private/DeferredShadingRenderer.cpp
+++ b/Engine/Source/Runtime/Renderer/Private/DeferredShadingRenderer.cpp
@@ -3459,6 +3459,8 @@ void FDeferredShadingSceneRenderer::Render(FRDGBuilder& GraphBuilder)
 			ComputeLightGridOutput = GatherLightsAndComputeLightGrid(GraphBuilder, bComputeLightGrid, SortedLightSet);
 		}
 
+		LightFunctionAtlas.RenderLightFunctionAtlas(GraphBuilder, Views);
+
 		CSV_CUSTOM_STAT(LightCount, All,  float(SortedLightSet.SortedLights.Num()), ECsvCustomStatOp::Set);
 		CSV_CUSTOM_STAT(LightCount, Batched, float(SortedLightSet.UnbatchedLightStart), ECsvCustomStatOp::Set);
 		CSV_CUSTOM_STAT(LightCount, Unbatched, float(SortedLightSet.SortedLights.Num()) - float(SortedLightSet.UnbatchedLightStart), ECsvCustomStatOp::Set);
diff --git a/Engine/Source/Runtime/Renderer/Private/DeferredShadingRenderer.h b/Engine/Source/Runtime/Renderer/Private/DeferredShadingRenderer.h
--- a/Engine/Source/Runtime/Renderer/Private/DeferredShadingRenderer.h
+++ b/Engine/Source/Runtime/Renderer/Private/DeferredShadingRenderer.h
@@ -23,6 +23,7 @@
 #include "ScreenSpaceRayTracing.h"
 #include "RenderGraphUtils.h"
 #include "SceneCulling/SceneCullingRenderer.h"
+#include "LightFunctionAtlas.h"
 
 enum class ERayTracingPrimaryRaysFlag : uint32;
 
@@ -972,7 +973,8 @@ private:
 		FRDGTextureRef ScreenShadowMaskSubPixelTexture,
 		FRDGTextureRef LightingChannelsTexture,
 		const FHairStrandsTransmittanceMaskData& InTransmittanceMaskData,
-		const bool bForwardRendering);
+		const bool bForwardRendering,
+		TRDGUniformBufferRef<FLightFunctionAtlasGlobalParameters> LightFunctionAtlas);
 
 	/** Renders an array of simple lights using standard deferred shading. */
 	void RenderSimpleLightsStandardDeferred(
diff --git a/Engine/Source/Runtime/Renderer/Private/LightFunctionAtlas.cpp b/Engine/Source/Runtime/Renderer/Private/LightFunctionAtlas.cpp
new file mode 100644
--- /dev/null
+++ b/Engine/Source/Runtime/Renderer/Private/LightFunctionAtlas.cpp
@@ -0,0 +1,669 @@
+// Copyright Epic Games, Inc. All Rights Reserved.
+
+#include "LightFunctionAtlas.h"
+
+#include "HAL/IConsoleManager.h"
+#include "RendererPrivate.h"
+#include "LightSceneProxy.h"
+#include "LightSceneInfo.h"
+#include "LightRendering.h"
+#include "Materials/MaterialRenderProxy.h"
+#include "Materials/MaterialInterface.h"
+#include "RenderGraphBuilder.h"
+#include "SystemTextures.h"
+#include "PostProcess/PostProcessing.h"
+#include "PostProcess/SceneFilterRendering.h"
+#include "ShadowRendering.h"
+
+DECLARE_GPU_STAT(LightFunctionAtlasGeneration);
+
+/*
+
+This LightFunctionAtlas stores light functions as 2D sub region of a texture2D atlas for all the views of a scene.
+Each material is stored only onces, de-duplicated based on its unique ID. Material instances and MID are correctly handled separately.
+Each slot stores a texture for a material that then can be applied on any light type. This only works if
+ - only Tex Coord are used to generate the light function
+ - only light/view direction are used to generate light function for point light
+Later, we could have atlas slot for a light/material pair if needed. That could automatically detected for instance when a material is reading instance data. Or if a material is reading view/light dir or world position.
+
+To avoid allocating SRV when using the atlas we use a constant buffer to store all atlas slot and light mapping data:
+	- A single SRV is used, being the atlas Texture2D.
+	- Otherwise a single constant buffer entry is used for each view, storing: AtlasSlotIndex=>{SubUVs} read for a Light LightIndex=>{FadeParams, AtlasSlotIndex, TranslatedWorlViewProjectionMatrix}
+
+Systems currently using the LightFunctionAtlas:
+ - Volumetric fog to inject local lights scattering using the light grid (when shadow casting is not enabled)
+ - Deferred lighting: applied to all lights of a scene stored in the light grid, preventing them from going the slow un-bactched route.
+
+How to use the Atlas:
+ - Add FLightFunctionAtlasGlobalParameters to your shader
+ - #include "LightFunctionAtlas/LightFunctionAtlasCommon.usf" 
+ - Call GetLocalLightFunctionCommon(DerivedParams.TranslatedWorldPosition, LightData.LightFunctionAtlasLightIndex); where LightData is a FDeferredLightData recovered from uniform or the light grid.
+
+ What is next:
+ - Atlas debug
+ - RGB option on project + shader changes with defined type
+ - Forward shading option
+ - Super sample CVAR
+ - Convert systems:
+    - Lumen
+    - Path tracer
+    - Clustered shading
+    - SLW
+    - Check those path with Substrate
+
+*/
+
+static TAutoConsoleVariable<int32> CVarLightFunctionAtlas(
+	TEXT("r.LightFunctionAtlas"),
+	1,
+	TEXT("Experimental: enable the light function atlas generation at runtime. The atlas will only be generated if other systems are using it at runtime."),
+	ECVF_RenderThreadSafe);
+
+// We do not dynamically scale allocated slot resolution for now.
+static TAutoConsoleVariable<int32> CVarLightFunctionAtlasSlotResolution(
+	TEXT("r.LightFunctionAtlas.SlotResolution"),
+	128,
+	TEXT("Experimental: The resolution of each atlas slot."),
+	ECVF_RenderThreadSafe);
+
+// We do not dynamically scale allocated slot resolution for now.
+static TAutoConsoleVariable<int32> CVarLightFunctionAtlasSize(
+	TEXT("r.LightFunctionAtlas.Size"),
+	4,
+	TEXT("Experimental: The default size in atlas slot count of the edge of the 2D texture atlas."),
+	ECVF_RenderThreadSafe);
+
+
+
+//////////////////////////////////////////////////////////////////////////
+
+// The CVars here represent systems that can request the creation/sampling of the light function atlas.
+
+// Volumetric fog always generate a light function for the directional light.
+// So this alias really only controls the use of the LightFunctionAtlas on the local lights.
+int GVolumetricFogUsesLightFunctionAtlas = 1;
+FAutoConsoleVariableRef CVarVolumetricFogLightFunction(
+	TEXT("r.VolumetricFog.LightFunction"),
+	GVolumetricFogUsesLightFunctionAtlas,
+	TEXT("This is an alias, please use r.VolumetricFog.UsesLightFunctionAtlas."),
+	ECVF_Scalability | ECVF_RenderThreadSafe
+);
+FAutoConsoleVariableRef CVarVolumetricFogUsesLightFunctionAtlas(
+	TEXT("r.VolumetricFog.UsesLightFunctionAtlas"),
+	GVolumetricFogUsesLightFunctionAtlas,
+	TEXT("Whether the light function atlas is sampled when rendering local lights."),
+	ECVF_Scalability | ECVF_RenderThreadSafe
+);
+
+// This deferred CVar includes deferred lights splatting (batched or not) as well as clustered lighting.
+int GDeferredUsesLightFunctionAtlas = 0;
+FAutoConsoleVariableRef CVarDeferredLightsUsesLightFunctionAtlas(
+	TEXT("r.Deferred.UsesLightFunctionAtlas"),
+	GDeferredUsesLightFunctionAtlas,
+	TEXT("Whether the light function atlas is sampled when rendering local lights."),
+	ECVF_RenderThreadSafe
+);
+
+
+
+//////////////////////////////////////////////////////////////////////////
+
+static uint32 GetAtlasSlotResolution()
+{
+	const uint32 AtlasSlotResolution = FMath::Clamp(CVarLightFunctionAtlasSlotResolution.GetValueOnRenderThread(), 32, 256);
+	return AtlasSlotResolution;
+}
+
+static uint32 GetAtlasEdgeSize()
+{
+	const uint32 AtlasEdgeSize = FMath::Clamp(CVarLightFunctionAtlasSize.GetValueOnRenderThread(), 4, 16);// 16x16 is the maximum slot count of LIGHT_FUNCTION_ATLAS_MAX_LIGHT_FUNCTION_COUNT=256 we currently allow
+	return AtlasEdgeSize;
+}
+
+IMPLEMENT_GLOBAL_SHADER_PARAMETER_STRUCT(FLightFunctionAtlasGlobalParameters, "LightFunctionAtlas");
+
+
+
+//////////////////////////////////////////////////////////////////////////
+
+class FLightFunctionAtlasSlotPS : public FMaterialShader
+{
+	DECLARE_SHADER_TYPE(FLightFunctionAtlasSlotPS, Material);
+
+	using FPermutationDomain = TShaderPermutationDomain<>;
+
+	BEGIN_SHADER_PARAMETER_STRUCT(FParameters, )
+		SHADER_PARAMETER(FVector4f, SvPositionToUVScaleBias)
+		SHADER_PARAMETER(FVector2f, LightFunctionTexelSize)
+	END_SHADER_PARAMETER_STRUCT()
+
+	FLightFunctionAtlasSlotPS() {}
+
+	FLightFunctionAtlasSlotPS(const ShaderMetaType::CompiledShaderInitializerType& Initializer)
+		: FMaterialShader(Initializer)
+	{
+		Bindings.BindForLegacyShaderParameters(
+			this,
+			Initializer.PermutationId,
+			Initializer.ParameterMap,
+			*FParameters::FTypeInfo::GetStructMetadata(),
+			// Don't require full bindings, we use FMaterialShader::SetParameters
+			false);
+	}
+
+	static bool ShouldCompilePermutation(const FMaterialShaderPermutationParameters& Parameters)
+	{
+		return Parameters.MaterialParameters.MaterialDomain == MD_LightFunction;
+	}
+
+	void SetParameters(FRHIBatchedShaderParameters& BatchedParameters, const FSceneView& View, const FMaterialRenderProxy* MaterialProxy)
+	{
+		const FMaterial& Material = MaterialProxy->GetMaterialWithFallback(View.GetFeatureLevel(), MaterialProxy);
+		FMaterialShader::SetViewParameters(BatchedParameters, View, View.ViewUniformBuffer);
+		FMaterialShader::SetParameters(BatchedParameters, MaterialProxy, Material, View);
+	}
+
+	FParameters GetParameters(const FVector2f& LightFunctionTexelSize, const FVector4f& SvPositionToUVScaleBias)
+	{
+		FParameters PS;
+		PS.SvPositionToUVScaleBias = SvPositionToUVScaleBias;
+		PS.LightFunctionTexelSize  = LightFunctionTexelSize;
+		return PS;
+	}
+};
+
+IMPLEMENT_MATERIAL_SHADER_TYPE(, FLightFunctionAtlasSlotPS, TEXT("/Engine/Private/LightFunctionAtlas/LightFunctionAtlasRender.usf"), TEXT("Main"), SF_Pixel);
+
+
+
+//////////////////////////////////////////////////////////////////////////
+
+FLightFunctionSlotKey::FLightFunctionSlotKey(FLightSceneInfo* InLightSceneInfo)
+{
+	const FMaterialRenderProxy* LightFunctionMaterial = InLightSceneInfo->Proxy->GetLightFunctionMaterial();
+	if (LightFunctionMaterial != nullptr)
+	{
+		const UMaterialInterface* LightFunctionMaterialInterface = LightFunctionMaterial->GetMaterialInterface();
+		LFMaterialUniqueID = LightFunctionMaterialInterface->GetUniqueID();
+	}
+	// else Default Key
+}
+
+
+
+//////////////////////////////////////////////////////////////////////////
+
+FLightFunctionAtlas::FLightFunctionAtlas()
+{
+	RegisteredLights.Reserve(64); // Reserve a minimal amount to avoid multiple allocations
+}
+
+FLightFunctionAtlas::~FLightFunctionAtlas()
+{
+}
+
+void FLightFunctionAtlas::UpdateRegisterLightSceneInfo(FLightSceneInfo* LightSceneInfo)
+{
+	if (LightSceneInfo->Proxy->GetLightFunctionMaterial() != nullptr && IsLightFunctionAtlasEnabled())
+	{
+	#if !UE_BUILD_SHIPPING
+		check(!RegisteredLights.Contains(LightSceneInfo));
+	#endif
+		RegisteredLights.Push(LightSceneInfo);
+	}
+	else
+	{
+		LightSceneInfo->Proxy->SetLightFunctionAtlasIndices(0);
+	}
+}
+
+void FLightFunctionAtlas::BeginSceneFrame(FViewFamilyInfo& ViewFamily, TArray<FViewInfo>& Views, FLightFunctionAtlasSceneData& LightFunctionAtlasSceneData, bool bShouldRenderVolumetricFog)
+{
+	RegisteredLights.Empty(64);
+
+	// Reset UBs
+	DefaultLightFunctionAtlasGlobalParametersUB = nullptr;
+	ViewLightFunctionAtlasGlobalParametersUBArray.Empty(4);
+
+	// Now lets check if we need to generate the atlas for this frame
+	bLightFunctionAtlasEnabled = CVarLightFunctionAtlas.GetValueOnRenderThread() > 0 && ViewFamily.EngineShowFlags.LightFunctions > 0;
+
+	// But only really enable the atlas generation if a system asks for it
+	bool VolumetricFogRequestsLightFunctionAtlas = false;
+	bool DeferredlightingRequestsLightFunctionAtlas = false;
+	if (bLightFunctionAtlasEnabled)
+	{
+		VolumetricFogRequestsLightFunctionAtlas = bShouldRenderVolumetricFog && GVolumetricFogUsesLightFunctionAtlas > 0;
+
+		DeferredlightingRequestsLightFunctionAtlas = GDeferredUsesLightFunctionAtlas > 0;
+
+		bLightFunctionAtlasEnabled = bLightFunctionAtlasEnabled && (VolumetricFogRequestsLightFunctionAtlas || DeferredlightingRequestsLightFunctionAtlas); 
+	}
+
+	// We propagate bLightFunctionAtlasEnabled to all the views to ease later shader parameter decision and binding for lighting, shadow or volumetric fog for instance (avoid sending lots of parameters all over the place)
+	LightFunctionAtlasSceneData.bLightFunctionAtlasEnabled				= bLightFunctionAtlasEnabled;
+	LightFunctionAtlasSceneData.bVolumetricFogUsesLightFunctionAtlas	= bLightFunctionAtlasEnabled && VolumetricFogRequestsLightFunctionAtlas;
+	LightFunctionAtlasSceneData.bDeferredlightingUsesLightFunctionAtlas	= bLightFunctionAtlasEnabled && DeferredlightingRequestsLightFunctionAtlas;
+	for (auto& View : Views)
+	{
+		View.LightFunctionAtlasViewData = FLightFunctionAtlasViewData(&LightFunctionAtlasSceneData);
+	}
+}
+
+void FLightFunctionAtlas::UpdateLightFunctionAtlas(const TArray<FViewInfo>& Views)
+{
+	if (!IsLightFunctionAtlasEnabled())
+	{
+		return;
+	}
+
+	AllocateAtlasSlots(Views);
+}
+
+void FLightFunctionAtlas::AllocateAtlasSlots(const TArray<FViewInfo>& Views)
+{
+	if (Views.Num() == 0)
+	{
+		return;
+	}
+
+	//
+	// Sort the list of lights registered as having light function in order to keep directional lights first, then each lights closer to each views
+	//
+	struct FSortedRegisteredLights
+	{
+		float MinDistanceToViews;
+		uint32 RegisteredLightIndex;
+		FORCEINLINE bool operator!=(FSortedRegisteredLights B) const
+		{
+			return MinDistanceToViews != B.MinDistanceToViews;
+		}
+
+		FORCEINLINE bool operator<(FSortedRegisteredLights B) const
+		{
+			return MinDistanceToViews < B.MinDistanceToViews;
+		}
+	};
+	TArray<FSortedRegisteredLights> SortedRegisteredLights;
+	const uint32 ViewCount = Views.Num();
+	{
+		SortedRegisteredLights.Reserve(RegisteredLights.Num());
+
+		const FVector View0Pos = Views[0].ViewMatrices.GetViewOrigin();
+		uint32 RegisteredLightCount = RegisteredLights.Num();
+		for (uint32 RegisteredLightIndex = 0; RegisteredLightIndex < RegisteredLightCount; ++RegisteredLightIndex)
+		{
+			const FLightSceneInfo* LightSceneInfo = RegisteredLights[RegisteredLightIndex];
+			FLightSceneProxy* Proxy = LightSceneInfo->Proxy;
+
+			if (Proxy->GetLightType() == LightType_Directional)
+			{
+				// Directional light are considered at a 0 distance from each view
+				SortedRegisteredLights.Push({ 0.0f, RegisteredLightIndex });
+			}
+			else
+			{
+				const FVector ProxyPos = FVector(Proxy->GetPosition());
+				float MinDistanceToViews = (View0Pos - ProxyPos).SquaredLength();
+				for (uint32 ViewId = 1; ViewId < ViewCount; ++ViewId)
+				{
+					FVector ViewXPos = Views[ViewId].ViewMatrices.GetViewOrigin();
+					MinDistanceToViews = FMath::Min(MinDistanceToViews, (ViewXPos - Proxy->GetPosition()).SquaredLength());
+				}
+				SortedRegisteredLights.Push({ MinDistanceToViews, RegisteredLightIndex });
+			}
+		}
+
+		// Now sort according to priority
+		SortedRegisteredLights.Sort();
+	}
+
+	//
+	// Allocate slots until we cannot anymore and set light function slot index on FLightSceneInfo to be send to the GPU later
+	//
+	const uint32 AtlasSlotResolution = GetAtlasSlotResolution();
+	const float AtlasEdgeSize = GetAtlasEdgeSize();
+	const float AtlasResolution = AtlasSlotResolution * AtlasEdgeSize;
+	check(AtlasEdgeSize * AtlasEdgeSize <= LIGHT_FUNCTION_ATLAS_MAX_LIGHT_FUNCTION_COUNT);
+	const uint32 AtlasMaxLightFunctionCount = AtlasEdgeSize * AtlasEdgeSize;
+
+	LightFunctionsSet.Reset();
+	EffectiveLightFunctionSlotArray.Reserve(AtlasMaxLightFunctionCount);
+
+	EffectiveLocalLightSlotArray.Reset();
+	EffectiveLocalLightSlotArray.Reserve(FMath::Min(LIGHT_FUNCTION_ATLAS_MAX_LIGHT_COUNT, SortedRegisteredLights.Num()));
+
+	LightFunctionsSet.Reset();
+	LightFunctionsSet.Reserve(AtlasMaxLightFunctionCount);
+
+	// Reserve the default slot at index 0 as the default identity slot
+	LightFunctionsSet.Add(FLightFunctionSlotKey());
+
+	uint32 NextAtlasSlotX = 0;
+	uint32 NextAtlasSlotY = 0;
+	auto AddAtlasSlot = [&](FLightFunctionSlotKey Key, const FMaterialRenderProxy* LightFunctionMaterial)
+	{
+		check(uint32(LightFunctionsSet.Num()) < (AtlasMaxLightFunctionCount - 1));
+		const uint32 NewSlotIndex = EffectiveLightFunctionSlotArray.Num();
+
+		Key.EffectiveLightFunctionSlotIndex = NewSlotIndex;
+
+		LightFunctionsSet.Add(Key);
+
+		EffectiveLightFunctionSlot& NewAtlasSlot = EffectiveLightFunctionSlotArray.Emplace_GetRef();
+
+		NewAtlasSlot.Min = FIntPoint(NextAtlasSlotX * AtlasSlotResolution, NextAtlasSlotY * AtlasSlotResolution);
+		NewAtlasSlot.Max = NewAtlasSlot.Min + FIntPoint(AtlasSlotResolution, AtlasSlotResolution);
+		NewAtlasSlot.MinU = (float(NewAtlasSlot.Min.X) + 0.5f) / AtlasResolution;
+		NewAtlasSlot.MinV = (float(NewAtlasSlot.Min.Y) + 0.5f) / AtlasResolution;
+
+		NewAtlasSlot.LightFunctionMaterial = LightFunctionMaterial;
+
+		NextAtlasSlotX++;
+		if (NextAtlasSlotX == AtlasEdgeSize)
+		{
+			NextAtlasSlotX = 0;
+			NextAtlasSlotY++;
+		}
+		return NewSlotIndex;
+	};
+
+	uint32 LocalLightWithLightFunctionCount = 0;
+	auto AddLightSlot = [&](FLightSceneInfo* LightSceneInfo, uint8 LightFunctionAtlasSlotIndex)
+	{
+		check(LocalLightWithLightFunctionCount < LIGHT_FUNCTION_ATLAS_MAX_LIGHT_COUNT);
+		EffectiveLocalLightSlot& NewLightSlot = EffectiveLocalLightSlotArray.Emplace_GetRef();
+		NewLightSlot.LightSceneInfo = LightSceneInfo;
+		NewLightSlot.LightFunctionAtlasSlotIndex = LightFunctionAtlasSlotIndex;
+		const uint32 LightSlotIndex = LocalLightWithLightFunctionCount;
+		LocalLightWithLightFunctionCount++;
+		return LightSlotIndex;
+	};
+
+	// Add the default invalid light slot at the beginning
+	AddLightSlot(nullptr, 0);
+
+	for (FSortedRegisteredLights& SortedRegisteredLight : SortedRegisteredLights)
+	{
+		FLightSceneInfo* LightSceneInfo = RegisteredLights[SortedRegisteredLight.RegisteredLightIndex];
+		FLightSceneProxy* Proxy = LightSceneInfo->Proxy;
+
+		if (LocalLightWithLightFunctionCount > LIGHT_FUNCTION_ATLAS_MAX_LIGHT_COUNT)
+		{
+			// We cannot register anymore light, so set them to no light function
+			Proxy->SetLightFunctionAtlasIndices(0);
+			continue;
+		}
+
+		FLightFunctionSlotKey NewKey = FLightFunctionSlotKey(LightSceneInfo);
+
+		FLightFunctionSlotKey* ExistingKey = LightFunctionsSet.Find(NewKey);
+
+		if (ExistingKey == nullptr && (uint32(LightFunctionsSet.Num()) < AtlasMaxLightFunctionCount))
+		{
+			const FMaterialRenderProxy* LightFunctionMaterial = Proxy->GetLightFunctionMaterial();
+
+			// Allocate slots for the light and views
+			Proxy->SetLightFunctionAtlasIndices(AddLightSlot(LightSceneInfo, AddAtlasSlot(NewKey, LightFunctionMaterial)));
+
+		}
+		else if (ExistingKey != nullptr)
+		{
+			// The key already exist, make the light point to the existing light function slot
+			Proxy->SetLightFunctionAtlasIndices(AddLightSlot(LightSceneInfo, ExistingKey->EffectiveLightFunctionSlotIndex));
+		}
+		else
+		{
+			// The key does not exist, or there is no space to allocate a new slot. Disable light function on that light
+			Proxy->SetLightFunctionAtlasIndices(0);
+		}
+	}
+
+	// TODO we could do all the constant buffer setup inline above (done in RenderAtlasSlots right now) if we would send a GraphBuilder here.
+}
+
+TRDGUniformBufferRef<FLightFunctionAtlasGlobalParameters> FLightFunctionAtlas::GetLightFunctionAtlasGlobalParameters(uint32 ViewIndex, FRDGBuilder& GraphBuilder)
+{
+	if (!IsLightFunctionAtlasEnabled())
+	{
+		return GetDefaultLightFunctionAtlasGlobalParameters(GraphBuilder);
+	}
+
+	const bool bViewIndexIsValid = ViewIndex < uint32(ViewLightFunctionAtlasGlobalParametersUBArray.Num());
+	check(bViewIndexIsValid);
+	if (bViewIndexIsValid)
+	{
+		return ViewLightFunctionAtlasGlobalParametersUBArray[ViewIndex];
+	}
+
+	return GetDefaultLightFunctionAtlasGlobalParameters(GraphBuilder);
+}
+
+TRDGUniformBufferRef<FLightFunctionAtlasGlobalParameters> FLightFunctionAtlas::GetDefaultLightFunctionAtlasGlobalParameters(FRDGBuilder& GraphBuilder)
+{
+	if (DefaultLightFunctionAtlasGlobalParametersUB == nullptr) // Only create the default buffer once per frame
+	{
+		FLightFunctionAtlasGlobalParameters* LightFunctionAtlasGlobalParams = GraphBuilder.AllocParameters<FLightFunctionAtlasGlobalParameters>();
+		LightFunctionAtlasGlobalParams->LightFunctionAtlasTexture = GSystemTextures.GetWhiteDummy(GraphBuilder);
+		LightFunctionAtlasGlobalParams->LightFunctionAtlasSampler = TStaticSamplerState<SF_Bilinear>::GetRHI();
+		LightFunctionAtlasGlobalParams->Slot_UVSize = 1.0f;
+		DefaultLightFunctionAtlasGlobalParametersUB = GraphBuilder.CreateUniformBuffer(LightFunctionAtlasGlobalParams);
+	}
+	return DefaultLightFunctionAtlasGlobalParametersUB;
+}
+
+void FLightFunctionAtlas::RenderLightFunctionAtlas(FRDGBuilder& GraphBuilder, TArray<FViewInfo>& Views)
+{
+	if (!IsLightFunctionAtlasEnabled())
+	{
+		return;
+	}
+
+	//
+	// Render the atlas
+	//
+	RenderAtlasSlots(GraphBuilder, Views);
+
+	//
+	// Allocate and fill up the global light function atlas UB
+	//
+	for (uint32 ViewIndex = 0; ViewIndex < uint32(Views.Num()); ++ViewIndex)
+	{
+		FViewInfo& View = Views[ViewIndex];
+
+		FLightFunctionAtlasGlobalParameters* LightFunctionAtlasGlobalParameters = GraphBuilder.AllocParameters<FLightFunctionAtlasGlobalParameters>();
+
+		check(RDGAtlasTexture2D != nullptr);
+		LightFunctionAtlasGlobalParameters->LightFunctionAtlasTexture = RDGAtlasTexture2D;
+		LightFunctionAtlasGlobalParameters->LightFunctionAtlasSampler = TStaticSamplerState<SF_Bilinear>::GetRHI();
+
+		const uint32 AtlasSlotResolution = GetAtlasSlotResolution();
+		float AtlasEdgeSize = GetAtlasEdgeSize();
+		float AtlasResolution = AtlasSlotResolution * AtlasEdgeSize;
+		LightFunctionAtlasGlobalParameters->Slot_UVSize = (AtlasSlotResolution - 1.0f) / AtlasResolution; // -1.0 because we remove a bit more than half a texel at the border.
+
+		const FIntPoint LightFunctionResolution = FIntPoint(AtlasSlotResolution, AtlasSlotResolution);
+		const FMatrix TranslatedWorldToWorld = FTranslationMatrix(-View.ViewMatrices.GetPreViewTranslation());
+
+		// Write the light data needed to rotate and fade the light function in the world.
+		uint32 LightIndex = 0;
+		for (EffectiveLocalLightSlot& LightSlot : EffectiveLocalLightSlotArray)
+		{
+			FLightSceneInfo* LightSceneInfo = LightSlot.LightSceneInfo;
+			const bool bIsDefaultSlot = LightSceneInfo == nullptr;
+
+			if (bIsDefaultSlot)
+			{
+				LightFunctionAtlasGlobalParameters->LightInfoDataParameters[LightIndex] = FVector4f(1.0f, 1.0f, 1.0f, 0.0f);
+				LightIndex++;
+				continue;
+			}
+
+			const FLightSceneProxy* Proxy = LightSceneInfo->Proxy;
+
+			float ShadowFadeFraction = 1.0f;
+
+			uint8 LightType = Proxy->GetLightType();
+			if (LightType == LightType_Spot)
+			{
+				// Spotlights needs to utilize their shadow projection when rendering the light function 
+				TArray<FWholeSceneProjectedShadowInitializer, TInlineAllocator<6>> Initializers;
+				Proxy->GetWholeSceneProjectedShadowInitializer(*View.Family, Initializers);
+
+				FWholeSceneProjectedShadowInitializer ShadowInitializer = Initializers[0];
+				FProjectedShadowInfo ProjectedShadowInfo;
+				ProjectedShadowInfo.SetupWholeSceneProjection(
+					LightSceneInfo,
+					&View,
+					ShadowInitializer,
+					LightFunctionResolution.X,
+					LightFunctionResolution.Y,
+					LightFunctionResolution.X,
+					LightFunctionResolution.Y,
+					0
+				);
+				//ShadowFadeFraction = ProjectedShadowInfo.FadeAlphas.Num() == 0 ? 1.0f : ProjectedShadowInfo.FadeAlphas[0];
+
+				FVector4f ShadowmapMinMaxValue;
+				LightFunctionAtlasGlobalParameters->LightInfoDataMatrix[LightIndex] = FMatrix44f(TranslatedWorldToWorld * ProjectedShadowInfo.GetWorldToShadowMatrix(ShadowmapMinMaxValue, &LightFunctionResolution));
+			}
+			else if (LightType == LightType_Point)
+			{
+				LightFunctionAtlasGlobalParameters->LightInfoDataMatrix[LightIndex] = FMatrix44f(TranslatedWorldToWorld * Proxy->GetWorldToLight());
+			}
+			else if (LightType == LightType_Rect)
+			{
+				const FVector Scale = Proxy->GetLightFunctionScale();
+				// Switch x and z so that z of the user specified scale affects the distance along the light direction
+				const FVector InverseScale = FVector(1.0 / Scale.Z, 1.0 / Scale.Y, 1.0 / Scale.X);
+				const FMatrix WorldToLight = Proxy->GetWorldToLight() * FScaleMatrix(InverseScale);
+
+				LightFunctionAtlasGlobalParameters->LightInfoDataMatrix[LightIndex] = FMatrix44f(TranslatedWorldToWorld * WorldToLight);
+			}
+			else if (LightType == LightType_Directional)
+			{
+				const FVector LightDirection = Proxy->GetDirection().GetSafeNormal();
+
+				const FVector Scale = LightSceneInfo->Proxy->GetLightFunctionScale();
+				// Switch x and z so that z of the user specified scale affects the distance along the light direction
+				const FVector InverseScale = FVector(1.0 / Scale.Z, 1.0 / Scale.Y, 1.0 / Scale.X);
+
+				const FMatrix WorldToLight = LightSceneInfo->Proxy->GetWorldToLight() * FScaleMatrix(FVector(InverseScale));
+				LightFunctionAtlasGlobalParameters->LightInfoDataMatrix[LightIndex] = FMatrix44f(TranslatedWorldToWorld * WorldToLight);
+			}
+			
+			const uint8 LightFunctionAtlasSlotIndex = LightSlot.LightFunctionAtlasSlotIndex;
+			FFloat16 PackedDisabledBrightness(Proxy->GetLightFunctionDisabledBrightness());
+			uint32 PackedLightInfoDataParams = uint32(LightType) | (uint32(PackedDisabledBrightness.Encoded) << 8);
+
+			EffectiveLightFunctionSlot& AtlasSlot = EffectiveLightFunctionSlotArray[LightFunctionAtlasSlotIndex];
+
+			// ShadowFadeFraction is unused.
+			LightFunctionAtlasGlobalParameters->LightInfoDataParameters[LightIndex] = FVector4f(Proxy->GetLightFunctionFadeDistance(), FMath::AsFloat(PackedLightInfoDataParams), AtlasSlot.MinU, AtlasSlot.MinV);
+
+			LightIndex++;
+		}
+
+		ViewLightFunctionAtlasGlobalParametersUBArray.Add(GraphBuilder.CreateUniformBuffer(LightFunctionAtlasGlobalParameters));
+	}
+}
+
+void FLightFunctionAtlas::AllocateTexture2DAtlas(FRDGBuilder& GraphBuilder)
+{
+	uint32 AtlasSlotResolution = GetAtlasSlotResolution();
+	uint32 AtlasEdgeSize = GetAtlasEdgeSize();
+	uint32 AtlasResolution = AtlasSlotResolution * AtlasEdgeSize;
+	const uint32 MipCount = 1;
+
+	RDGAtlasTexture2D = GraphBuilder.CreateTexture(FRDGTextureDesc::Create2D(
+		FIntPoint(AtlasResolution, AtlasResolution),
+		PF_R8,
+		FClearValueBinding::Black,
+		ETextureCreateFlags::UAV | ETextureCreateFlags::ShaderResource | ETextureCreateFlags::RenderTargetable,
+		MipCount),
+		TEXT("LightFunction.Atlas"),
+		ERDGTextureFlags::MultiFrame);
+}
+
+BEGIN_SHADER_PARAMETER_STRUCT(FLightFunctionAtlasRenderParameters, )
+	RENDER_TARGET_BINDING_SLOTS()
+END_SHADER_PARAMETER_STRUCT()
+
+void FLightFunctionAtlas::RenderAtlasSlots(FRDGBuilder& GraphBuilder, const TArray<FViewInfo>& Views)
+{
+	AllocateTexture2DAtlas(GraphBuilder);
+
+	SCOPED_NAMED_EVENT(LightFunctionAtlasGeneration, FColor::Emerald);
+	RDG_GPU_STAT_SCOPE(GraphBuilder, LightFunctionAtlasGeneration);
+	RDG_CSV_STAT_EXCLUSIVE_SCOPE(GraphBuilder, LightFunctionAtlasGeneration);
+	
+	FLightFunctionAtlasRenderParameters* PassParameters = GraphBuilder.AllocParameters<FLightFunctionAtlasRenderParameters>();
+	PassParameters->RenderTargets[0] = FRenderTargetBinding(RDGAtlasTexture2D, ERenderTargetLoadAction::ENoAction);
+	GraphBuilder.AddPass(
+		RDG_EVENT_NAME("LightFunctionAtlas Generation"),
+		PassParameters,
+		ERDGPassFlags::Raster,
+		[PassParameters, &Views, this](FRHICommandList& RHICmdList)
+		{
+			uint32 AtlasSlotResolution = GetAtlasSlotResolution();
+			uint32 AtlasEdgeSize = GetAtlasEdgeSize();
+			uint32 AtlasResolution = AtlasSlotResolution * AtlasEdgeSize;
+
+			// Render all light functions and update light info
+			for (EffectiveLightFunctionSlot& Slot : EffectiveLightFunctionSlotArray)
+			{
+				// This always work because in this case we do not need anything from any view.
+				const FViewInfo& View = Views[0];
+
+				const FMaterialRenderProxy* MaterialProxyForRendering = Slot.LightFunctionMaterial;
+				if (MaterialProxyForRendering == nullptr)
+				{
+					continue;	// This is the default invalid slot at index 0 to notify disabled light function on local light data.
+				}
+				const FMaterial& Material = MaterialProxyForRendering->GetMaterialWithFallback(View.GetFeatureLevel(), MaterialProxyForRendering);
+
+				RHICmdList.SetViewport(0.f, 0.f, 0.f, AtlasResolution, AtlasResolution, 1.f);
+
+				FGraphicsPipelineStateInitializer GraphicsPSOInit;
+				RHICmdList.ApplyCachedRenderTargets(GraphicsPSOInit);
+				GraphicsPSOInit.RasterizerState = TStaticRasterizerState<FM_Solid, CM_None>::GetRHI();
+				GraphicsPSOInit.BlendState = TStaticBlendState<>::GetRHI();
+				GraphicsPSOInit.DepthStencilState = TStaticDepthStencilState<false, CF_Always>::GetRHI();
+				GraphicsPSOInit.PrimitiveType = PT_TriangleList;
+
+				const FMaterialShaderMap* MaterialShaderMap = Material.GetRenderingThreadShaderMap();
+				TShaderMapRef<FPostProcessVS> VertexShader(View.ShaderMap); 
+
+				FLightFunctionAtlasSlotPS::FPermutationDomain PermutationVector;
+				TShaderRef<FLightFunctionAtlasSlotPS> PixelShader = MaterialShaderMap->GetShader<FLightFunctionAtlasSlotPS>(PermutationVector);
+
+				GraphicsPSOInit.BoundShaderState.VertexDeclarationRHI = GFilterVertexDeclaration.VertexDeclarationRHI;
+				GraphicsPSOInit.BoundShaderState.VertexShaderRHI = VertexShader.GetVertexShader();
+				GraphicsPSOInit.BoundShaderState.PixelShaderRHI = PixelShader.GetPixelShader();
+
+				SetGraphicsPipelineState(RHICmdList, GraphicsPSOInit, 0);
+
+				FVector2f LightFunctionTexelSize  = FVector2f(1.0f / AtlasSlotResolution, 1.0f / AtlasSlotResolution);
+				FVector4f SvPositionToUVScaleBias = FVector4f(
+					1.0f / float(AtlasSlotResolution - 1), 1.0f / float(AtlasSlotResolution - 1),
+					float(Slot.Min.X) + 0.5f, float(Slot.Min.Y) + 0.5f);
+
+				FLightFunctionAtlasSlotPS::FParameters PS = PixelShader->GetParameters(LightFunctionTexelSize, SvPositionToUVScaleBias);
+
+				ClearUnusedGraphResources(PixelShader, &PS);
+				SetShaderParametersMixedPS(RHICmdList, PixelShader, PS, View, MaterialProxyForRendering);
+
+				FIntPoint RectSize = FIntPoint(Slot.Max.X - Slot.Min.X, Slot.Max.Y - Slot.Min.Y);
+				DrawRectangle(
+					RHICmdList,
+					Slot.Min.X, Slot.Min.Y,
+					RectSize.X, RectSize.Y,
+					Slot.Min.X, Slot.Min.Y,
+					RectSize.X, RectSize.Y,
+					AtlasResolution,
+					AtlasSlotResolution,
+					VertexShader);
+			}
+		}
+	);
+}
diff --git a/Engine/Source/Runtime/Renderer/Private/LightFunctionAtlas.h b/Engine/Source/Runtime/Renderer/Private/LightFunctionAtlas.h
new file mode 100644
--- /dev/null
+++ b/Engine/Source/Runtime/Renderer/Private/LightFunctionAtlas.h
@@ -0,0 +1,156 @@
+// Copyright Epic Games, Inc. All Rights Reserved.
+
+#pragma once
+
+#include "CoreMinimal.h"
+#include "LightSceneProxy.h"
+#include "RenderGraphFwd.h"
+#include "ShaderParameterMacros.h"
+
+class FViewFamilyInfo;
+class FViewInfo;
+class FScene;
+class FMaterialRenderProxy;
+class FRDGBuilder;
+class FLightSceneInfo;
+class FRDGTexture;
+
+struct FSortedLightSetSceneInfo;
+struct IPooledRenderTarget;
+
+
+
+struct FLightFunctionAtlasSceneData
+{
+	bool bLightFunctionAtlasEnabled = false;
+	bool bVolumetricFogUsesLightFunctionAtlas = false;
+	bool bDeferredlightingUsesLightFunctionAtlas = false;
+
+	bool GetLightFunctionAtlasEnabled()					const { return bLightFunctionAtlasEnabled; }
+	bool GetVolumetricFogUsesLightFunctionAtlas()		const { return bVolumetricFogUsesLightFunctionAtlas; }
+	bool GetDeferredlightingUsesLightFunctionAtlas()	const { return bDeferredlightingUsesLightFunctionAtlas; }
+};
+
+struct FLightFunctionAtlasViewData
+{
+	FLightFunctionAtlasViewData() : SceneData(nullptr) {}
+	FLightFunctionAtlasViewData(FLightFunctionAtlasSceneData* InSceneData) : SceneData(InSceneData) {}
+
+	bool GetLightFunctionAtlasEnabled()					const { return SceneData ? SceneData->bLightFunctionAtlasEnabled				: false; }
+	bool GetVolumetricFogUsesLightFunctionAtlas()		const { return SceneData ? SceneData->bVolumetricFogUsesLightFunctionAtlas	: false; }
+	bool GetDeferredlightingUsesLightFunctionAtlas()	const { return SceneData ? SceneData->bDeferredlightingUsesLightFunctionAtlas	: false; }
+
+private:
+	FLightFunctionAtlasSceneData* SceneData;
+};
+
+
+
+struct FLightFunctionSlotKey
+{
+	uint32 LFMaterialUniqueID = 0;
+
+	uint32 EffectiveLightFunctionSlotIndex = 0;	// Not used to de-duplicate light function. It is the index of the effective lighting function for this frame in EffectiveLightFunctionSlotArray.
+
+	FLightFunctionSlotKey() {};
+
+	FLightFunctionSlotKey(FLightSceneInfo* InLightSceneInfo);
+
+	inline bool operator==(const FLightFunctionSlotKey& Other) const 
+	{ 
+		return Other.LFMaterialUniqueID == LFMaterialUniqueID;
+	}
+};
+
+inline uint32 GetTypeHash(FLightFunctionSlotKey Key)
+{
+	return GetTypeHash(Key.LFMaterialUniqueID) ;
+}
+
+struct FLightFunctionSlot
+{
+	const FMaterialRenderProxy* LightFunctionMaterial;
+	FIntPoint Min;
+	FIntPoint Max;
+};
+
+#define LIGHT_FUNCTION_ATLAS_MAX_LIGHT_FUNCTION_COUNT 256
+
+/**
+ * On devices using OpenGL, we are limited to 256 vec4s. 
+ * Light parameters take 5 vec4s. So we can fit at maximum of 51.2 light info data. We round it down to 48 lights.
+ * This is already plenty. And we can move to use a buffer later on.
+ */
+#define LIGHT_FUNCTION_ATLAS_MAX_LIGHT_COUNT 48
+
+// This allows to not have to store more data per GPU light representation on GPU. The light only needs an index into the array.
+// Using a constant buffer also workaround the fact that we would otherwise need another SRV in forward shaders. 
+// The light atlas texture itself already use 1 extra SRV. We could an extra SRV and have LightInfoDataXXX be in a buffer that scale with amount of light in the scene.
+BEGIN_GLOBAL_SHADER_PARAMETER_STRUCT(FLightFunctionAtlasGlobalParameters, )
+	SHADER_PARAMETER_RDG_TEXTURE(Texture2D, LightFunctionAtlasTexture)
+	SHADER_PARAMETER_SAMPLER(SamplerState, LightFunctionAtlasSampler)
+	SHADER_PARAMETER(float, Slot_UVSize)
+	SHADER_PARAMETER_ARRAY(FMatrix44f, LightInfoDataMatrix, [LIGHT_FUNCTION_ATLAS_MAX_LIGHT_COUNT])		// Light data
+	SHADER_PARAMETER_ARRAY(FVector4f,  LightInfoDataParameters, [LIGHT_FUNCTION_ATLAS_MAX_LIGHT_COUNT]) // Light data
+END_GLOBAL_SHADER_PARAMETER_STRUCT()
+
+
+
+// This class holds all data and resources related light function for a single scene, including multiple views.
+struct FLightFunctionAtlas
+{
+	FLightFunctionAtlas();
+	virtual ~FLightFunctionAtlas();
+
+	bool IsLightFunctionAtlasEnabled() const { return bLightFunctionAtlasEnabled; }
+
+	void BeginSceneFrame(FViewFamilyInfo& ViewFamily, TArray<FViewInfo>& Views, FLightFunctionAtlasSceneData& LightFunctionAtlasSceneData, bool bShouldRenderVolumetricFog);
+
+	void UpdateRegisterLightSceneInfo(FLightSceneInfo* LightSceneInfo);
+
+	void UpdateLightFunctionAtlas(const TArray<FViewInfo>& Views);
+
+	void RenderLightFunctionAtlas(FRDGBuilder& GraphBuilder, TArray<FViewInfo>& Views);
+
+	TRDGUniformBufferRef<FLightFunctionAtlasGlobalParameters> GetLightFunctionAtlasGlobalParameters(uint32 ViewIndex, FRDGBuilder& GraphBuilder);
+	TRDGUniformBufferRef<FLightFunctionAtlasGlobalParameters> GetDefaultLightFunctionAtlasGlobalParameters(FRDGBuilder& GraphBuilder);
+
+private:
+
+	void AllocateAtlasSlots(const TArray<FViewInfo>& Views);
+
+	void AllocateTexture2DAtlas(FRDGBuilder& GraphBuilder);
+
+	void RenderAtlasSlots(FRDGBuilder& GraphBuilder, const TArray<FViewInfo>& Views);
+
+	bool bLightFunctionAtlasEnabled = false;
+
+	FRDGTextureRef RDGAtlasTexture2D = nullptr;
+
+	// All the lights that wants to sample light functions
+	TArray<FLightSceneInfo*> RegisteredLights;
+
+	// This set is used to de-duplicate light functions in the atlas. An alternative to set would be to use simple array with simple loops.
+	TSet<FLightFunctionSlotKey> LightFunctionsSet;
+
+	// The structure used to render light function and generate the associated constant buffer, containing UVs and transformation matrices.
+	struct EffectiveLightFunctionSlot
+	{
+		const FMaterialRenderProxy* LightFunctionMaterial = nullptr;
+		FIntPoint Min;
+		FIntPoint Max;
+		float MinU;
+		float MinV;
+	};
+	TArray<EffectiveLightFunctionSlot> EffectiveLightFunctionSlotArray;
+
+	struct EffectiveLocalLightSlot
+	{
+		FLightSceneInfo*	LightSceneInfo = nullptr;
+		uint8				LightFunctionAtlasSlotIndex = 0;
+	};
+	TArray<EffectiveLocalLightSlot> EffectiveLocalLightSlotArray;
+
+	TRDGUniformBufferRef<FLightFunctionAtlasGlobalParameters>	DefaultLightFunctionAtlasGlobalParametersUB;
+	TArray<TRDGUniformBufferRef<FLightFunctionAtlasGlobalParameters>> ViewLightFunctionAtlasGlobalParametersUBArray;
+};
diff --git a/Engine/Source/Runtime/Renderer/Private/LightFunctionRendering.cpp b/Engine/Source/Runtime/Renderer/Private/LightFunctionRendering.cpp
--- a/Engine/Source/Runtime/Renderer/Private/LightFunctionRendering.cpp
+++ b/Engine/Source/Runtime/Renderer/Private/LightFunctionRendering.cpp
@@ -170,7 +170,7 @@ public:
 		auto DeferredLightParameter = GetUniformBufferParameter<FDeferredLightUniformStruct>();
 		if (DeferredLightParameter.IsBound())
 		{
-			SetDeferredLightParameters(BatchedParameters, DeferredLightParameter, LightSceneInfo, View);
+			SetDeferredLightParameters(BatchedParameters, DeferredLightParameter, LightSceneInfo, View, View.LightFunctionAtlasViewData.GetDeferredlightingUsesLightFunctionAtlas());
 		}
 	}
 
diff --git a/Engine/Source/Runtime/Renderer/Private/LightGridInjection.cpp b/Engine/Source/Runtime/Renderer/Private/LightGridInjection.cpp
--- a/Engine/Source/Runtime/Renderer/Private/LightGridInjection.cpp
+++ b/Engine/Source/Runtime/Renderer/Private/LightGridInjection.cpp
@@ -376,7 +376,7 @@ static void PackLocalLightData(
 	FForwardLocalLightData& Out, 
 	const FViewInfo& View,
 	const FLightRenderParameters& LightParameters,
-	const uint32 LightTypeAndShadowMapChannelMaskPacked,
+	const uint32 LightTypeAndShadowMapChannelMaskAndLightFunctionIndexPacked,
 	const int32 LightSceneId,
 	const int32 VirtualShadowMapId,
 	const int32 PrevLocalLightIndex,
@@ -400,7 +400,7 @@ static void PackLocalLightData(
 	RectPackedZ |= uint32(FMath::Clamp(LightParameters.RectLightAtlasMaxLevel, 0.f, 63.f)) << 26;			//  6 bits
 
 	// Pack specular scale and IES profile index
-	const uint32 SpecularScaleAndIESData = PackRG16(LightParameters.SpecularScale, LightParameters.IESAtlasIndex);
+	const uint32 SpecularScaleAndIESData = PackRG16(LightParameters.SpecularScale, LightParameters.IESAtlasIndex); // pack atlas id here? 16bit specular 8bit IES and 8 bit LightFunction
 
 	const FVector2f LightColorPacked = PackLightColor(FVector3f(LightParameters.Color));
 
@@ -415,7 +415,7 @@ static void PackLocalLightData(
 	// NOTE: SpotAngles needs full-precision for VSM one pass projection
 	Out.LightPositionAndInvRadius							= FVector4f(LightTranslatedWorldPosition, LightParameters.InvRadius);
 	Out.LightColorAndIdAndFalloffExponent					= FVector4f(LightColorPacked.X, LightColorPacked.Y, LightSceneId, LightParameters.FalloffExponent);
-	Out.LightDirectionAndShadowMapChannelMask				= FVector4f(LightParameters.Direction, FMath::AsFloat(LightTypeAndShadowMapChannelMaskPacked));
+	Out.LightDirectionAndShadowMapChannelMask				= FVector4f(LightParameters.Direction, FMath::AsFloat(LightTypeAndShadowMapChannelMaskAndLightFunctionIndexPacked));
 	Out.SpotAnglesAndSourceRadiusPacked						= FVector4f(LightParameters.SpotAngles.X, LightParameters.SpotAngles.Y, FMath::AsFloat(PackedZ), FMath::AsFloat(PackedW));
 	Out.LightTangentAndIESDataAndSpecularScale				= FVector4f(LightParameters.Tangent, FMath::AsFloat(SpecularScaleAndIESData));
 	Out.RectDataAndVirtualShadowMapIdOrPrevLocalLightIndex	= FVector4f(FMath::AsFloat(RectPackedX), FMath::AsFloat(RectPackedY), FMath::AsFloat(RectPackedZ), float(VirtualShadowMapIdOrPrevLocalLightIndex));
@@ -583,8 +583,7 @@ FComputeLightGridOutput FSceneRenderer::ComputeLightGrid(FRDGBuilder& GraphBuild
 						LightParameters.Color *= LightParameters.GetLightExposureScale(Exposure);
 
 						float VolumetricScatteringIntensity = LightProxy->GetVolumetricScatteringIntensity();
-						if (LightNeedsSeparateInjectionIntoVolumetricFogForOpaqueShadow(View, LightSceneInfo, VisibleLightInfos[LightSceneInfo->Id])
-							|| (LightNeedsSeparateInjectionIntoVolumetricFogForLightFunction(LightSceneInfo) && CheckForLightFunction(LightSceneInfo)))
+						if (LightNeedsSeparateInjectionIntoVolumetricFogForOpaqueShadow(View, LightSceneInfo, VisibleLightInfos[LightSceneInfo->Id]))
 						{
 							// Disable this lights forward shading volumetric scattering contribution
 							VolumetricScatteringIntensity = 0;
diff --git a/Engine/Source/Runtime/Renderer/Private/LightRendering.cpp b/Engine/Source/Runtime/Renderer/Private/LightRendering.cpp
--- a/Engine/Source/Runtime/Renderer/Private/LightRendering.cpp
+++ b/Engine/Source/Runtime/Renderer/Private/LightRendering.cpp
@@ -31,6 +31,7 @@
 #include "MobileBasePassRendering.h"
 #include "TranslucentLighting.h"
 #include "StochasticShadows/StochasticShadows.h"
+#include "LightFunctionAtlas.h"
 
 // ENABLE_DEBUG_DISCARD_PROP is used to test the lighting code by allowing to discard lights to see how performance scales
 // It ought never to be enabled in a shipping build, and is probably only really useful when woring on the shading code.
@@ -547,11 +548,12 @@ static void RenderLight(
 	FRDGTextureRef LightingChannelsTexture,
 	bool bRenderOverlap,
 	bool bCloudShadow,
+	TRDGUniformBufferRef<FLightFunctionAtlasGlobalParameters> LightFunctionAtlasGlobalParameters,
 	TRDGUniformBufferRef<FVirtualShadowMapUniformParameters> VirtualShadowMapUniformBuffer = nullptr,
 	FRDGTextureRef ShadowMaskBits = nullptr,
 	int32 VirtualShadowMapId = INDEX_NONE);
 
-FDeferredLightUniformStruct GetDeferredLightParameters(const FSceneView& View, const FLightSceneInfo& LightSceneInfo, uint32 LightFlags)
+FDeferredLightUniformStruct GetDeferredLightParameters(const FSceneView& View, const FLightSceneInfo& LightSceneInfo, bool bUseLightFunctionAtlas, uint32 LightFlags)
 {
 	FDeferredLightUniformStruct Out;
 
@@ -582,7 +584,7 @@ FDeferredLightUniformStruct GetDeferredLightParameters(const FSceneView& View, c
 		ShadowMapChannel == 3 ? 1 : 0);
 
 	const bool bDynamicShadows = View.Family->EngineShowFlags.DynamicShadows && GetShadowQuality() > 0;
-	const bool bHasLightFunction = LightSceneInfo.Proxy->GetLightFunctionMaterial() != NULL;
+	const bool bHasLightFunction = LightSceneInfo.Proxy->GetLightFunctionMaterial() != NULL && !bUseLightFunctionAtlas;
 	Out.ShadowedBits = LightSceneInfo.Proxy->CastsStaticShadow() || bHasLightFunction ? 1 : 0;
 	Out.ShadowedBits |= LightSceneInfo.Proxy->CastsDynamicShadow() && View.Family->EngineShowFlags.DynamicShadows ? 3 : 0;
 
@@ -663,6 +665,7 @@ FDeferredLightUniformStruct GetSimpleDeferredLightParameters(
 	Out.LightParameters.RectLightAtlasUVScale = FVector2f::ZeroVector;
 	Out.LightParameters.RectLightAtlasMaxLevel = FLightRenderParameters::GetRectLightAtlasInvalidMIPLevel();
 	Out.LightParameters.IESAtlasIndex = INDEX_NONE;
+	Out.LightParameters.LightFunctionAtlasLightIndex = 0;
 	return Out;
 }
 FDeferredLightUniformStruct GetSimpleDeferredLightParameters(
@@ -850,9 +853,11 @@ class FDeferredLightPS : public FGlobalShader
 		SHADER_PARAMETER_RDG_TEXTURE(Texture2D, ScreenShadowMaskSubPixelTexture)
 		SHADER_PARAMETER_STRUCT_REF(FViewUniformShaderParameters, View)
 		SHADER_PARAMETER_RDG_UNIFORM_BUFFER(FDeferredLightUniformStruct, DeferredLight)
+		SHADER_PARAMETER_RDG_UNIFORM_BUFFER(FLightFunctionAtlasGlobalParameters, LightFunctionAtlas)
 		// For virtual shadow map mask
 		SHADER_PARAMETER_RDG_UNIFORM_BUFFER(FVirtualShadowMapUniformParameters, VirtualShadowMap)
 		SHADER_PARAMETER(int32, VirtualShadowMapId)
+		SHADER_PARAMETER(int32, bUseLightFunctionAtlas)
 		SHADER_PARAMETER_RDG_TEXTURE(Texture2D, ShadowMaskBits)
 		RENDER_TARGET_BINDING_SLOTS()
 	END_SHADER_PARAMETER_STRUCT()
@@ -1116,6 +1121,8 @@ bool FSceneRenderer::AllowSimpleLights() const
 
 void FSceneRenderer::GatherAndSortLights(FSortedLightSetSceneInfo& OutSortedLights, bool bShadowedLightsInClustered)
 {
+	LightFunctionAtlas.BeginSceneFrame(ViewFamily, Views, Scene->LightFunctionAtlasSceneData, ShouldRenderVolumetricFog());
+
 	if (AllowSimpleLights())
 	{
 		GatherSimpleLights(ViewFamily, Views, OutSortedLights.SimpleLights);
@@ -1126,6 +1133,7 @@ void FSceneRenderer::GatherAndSortLights(FSortedLightSetSceneInfo& OutSortedLigh
 	// NOTE: we allocate space also for simple lights such that they can be referenced in the same sorted range
 	SortedLights.Empty(Scene->Lights.Num() + SimpleLights.InstanceData.Num());
 
+	const bool bUseLightFunctionAtlas = Scene->LightFunctionAtlasSceneData.GetDeferredlightingUsesLightFunctionAtlas();
 	bool bDynamicShadows = ViewFamily.EngineShowFlags.DynamicShadows && GetShadowQuality() > 0;
 
 #if ENABLE_DEBUG_DISCARD_PROP
@@ -1140,6 +1148,8 @@ void FSceneRenderer::GatherAndSortLights(FSortedLightSetSceneInfo& OutSortedLigh
 		const FLightSceneInfoCompact& LightSceneInfoCompact = *LightIt;
 		const FLightSceneInfo* const LightSceneInfo = LightSceneInfoCompact.LightSceneInfo;
 
+		LightFunctionAtlas.UpdateRegisterLightSceneInfo(LightSceneInfoCompact.LightSceneInfo);
+
 #if ENABLE_DEBUG_DISCARD_PROP
 		{
 			int32 PrevCounter = int32(DebugDiscardCounter);
@@ -1183,7 +1193,7 @@ void FSceneRenderer::GatherAndSortLights(FSortedLightSetSceneInfo& OutSortedLigh
 					const bool bClusteredDeferredSupported =
 						!SortedLightInfo->SortKey.Fields.bTextureProfile &&
 						(!SortedLightInfo->SortKey.Fields.bShadowed || bShadowedLightsInClustered) &&
-						!SortedLightInfo->SortKey.Fields.bLightFunction &&
+						(!SortedLightInfo->SortKey.Fields.bLightFunction || bUseLightFunctionAtlas) &&
 						!SortedLightInfo->SortKey.Fields.bUsesLightingChannels
 						&& LightSceneInfoCompact.LightType != LightType_Directional
 						&& LightSceneInfoCompact.LightType != LightType_Rect
@@ -1246,10 +1256,12 @@ void FSceneRenderer::GatherAndSortLights(FSortedLightSetSceneInfo& OutSortedLigh
 	{
 		const FSortedLightSceneInfo& SortedLightInfo = SortedLights[LightIndex];
 		const bool bDrawShadows = SortedLightInfo.SortKey.Fields.bShadowed;
-		const bool bDrawLightFunction = SortedLightInfo.SortKey.Fields.bLightFunction;
 		const bool bTextureLightProfile = SortedLightInfo.SortKey.Fields.bTextureProfile;
 		const bool bLightingChannels = SortedLightInfo.SortKey.Fields.bUsesLightingChannels;
 
+		// Do not schedule unbatched lights if the atlas is used and enabled
+		const bool bDrawLightFunction = SortedLightInfo.SortKey.Fields.bLightFunction && !bUseLightFunctionAtlas;
+
 		if (SortedLightInfo.SortKey.Fields.bHandledByStochasticShadows && OutSortedLights.StochasticShadowsLightStart == SortedLights.Num())
 		{
 			// Mark the first index that needs to be rendered
@@ -1285,6 +1297,9 @@ void FSceneRenderer::GatherAndSortLights(FSortedLightSetSceneInfo& OutSortedLigh
 	{
 		OutSortedLights.UnbatchedLightStart = OutSortedLights.StochasticShadowsLightStart;
 	}
+
+	// Update the light function atlas according to registered lights and views
+	LightFunctionAtlas.UpdateLightFunctionAtlas(Views);
 }
 
 FHairStrandsTransmittanceMaskData CreateDummyHairStrandsTransmittanceMaskData(FRDGBuilder& GraphBuilder, FGlobalShaderMap* ShaderMap);
@@ -1302,6 +1317,7 @@ void FDeferredShadingSceneRenderer::RenderLights(
 #else
 	const bool bEnableRayTracing = false;
 #endif // RHI_RAYTRACING
+	const bool bUseLightFunctionAtlas = Scene->LightFunctionAtlasSceneData.GetDeferredlightingUsesLightFunctionAtlas();
 
 	RDG_EVENT_SCOPE(GraphBuilder, "Lights");
 	RDG_GPU_STAT_SCOPE(GraphBuilder, Lights);
@@ -1431,7 +1447,8 @@ void FDeferredShadingSceneRenderer::RenderLights(
 				{
 					// Render the light to the scene color buffer, using a 1x1 white texture as input
 					const FLightSceneInfo* LightSceneInfo = SortedLights[LightIndex].LightSceneInfo;
-					RenderLight(GraphBuilder, Scene, View, SceneTextures, LightSceneInfo, nullptr, LightingChannelsTexture, false /*bRenderOverlap*/, false /*bCloudShadow*/);
+					RenderLight(GraphBuilder, Scene, View, SceneTextures, LightSceneInfo, nullptr, LightingChannelsTexture, false /*bRenderOverlap*/, false /*bCloudShadow*/,
+						LightFunctionAtlas.GetLightFunctionAtlasGlobalParameters(ViewIndex, GraphBuilder));
 				}
 			}
 
@@ -1439,15 +1456,18 @@ void FDeferredShadingSceneRenderer::RenderLights(
 			if (bUseHairLighting)
 			{
 				FRDGTextureRef NullScreenShadowMaskSubPixelTexture = nullptr;
-				for (FViewInfo& View : Views)
+				for (int32 ViewIndex = 0, ViewCount = Views.Num(); ViewIndex < ViewCount; ++ViewIndex)
 				{
+					FViewInfo& View = Views[ViewIndex];
 					if (HairStrands::HasViewHairStrandsData(View))
 					{
 						// Draw non-shadowed non-light function lights without changing render targets between them
 						for (int32 LightIndex = StandardDeferredStart; LightIndex < UnbatchedLightStart; LightIndex++)
 						{
 							const FLightSceneInfo* LightSceneInfo = SortedLights[LightIndex].LightSceneInfo;
-							RenderLightForHair(GraphBuilder, View, SceneTextures, LightSceneInfo, NullScreenShadowMaskSubPixelTexture, LightingChannelsTexture, DummyTransmittanceMaskData, false /*bForwardRendering*/);
+							RenderLightForHair(GraphBuilder, View, SceneTextures, LightSceneInfo, 
+								NullScreenShadowMaskSubPixelTexture, LightingChannelsTexture, DummyTransmittanceMaskData, false /*bForwardRendering*/,
+								LightFunctionAtlas.GetLightFunctionAtlasGlobalParameters(ViewIndex, GraphBuilder));
 						}
 					}
 				}
@@ -1515,7 +1535,7 @@ void FDeferredShadingSceneRenderer::RenderLights(
 				const FLightOcclusionType OcclusionType = GetLightOcclusionType(LightSceneProxy);
 
 				const bool bDrawShadows = SortedLightInfo.SortKey.Fields.bShadowed;
-				const bool bDrawLightFunction = SortedLightInfo.SortKey.Fields.bLightFunction;
+				const bool bDrawLightFunction = SortedLightInfo.SortKey.Fields.bLightFunction && !bUseLightFunctionAtlas;
 				const bool bDrawPreviewIndicator = ViewFamily.EngineShowFlags.PreviewShadowsIndicator && !LightSceneInfo.IsPrecomputedLightingValid() && LightSceneProxy.HasStaticShadowing();
 				const bool bDrawHairShadow = bDrawShadows && bUseHairLighting;
 				const bool bUseHairDeepShadow = bDrawShadows && bUseHairLighting && LightSceneProxy.CastsHairStrandsDeepShadow();
@@ -2051,14 +2071,20 @@ void FDeferredShadingSceneRenderer::RenderLights(
 
 						RDG_EVENT_SCOPE_CONDITIONAL(GraphBuilder, ViewCount > 1, "View%d", ViewIndex);
 						SCOPED_GPU_MASK(GraphBuilder.RHICmdList, View.GPUMask);
-						RenderLight(GraphBuilder, Scene, View, SceneTextures, &LightSceneInfo, VirtualShadowMapId != INDEX_NONE ? nullptr : ScreenShadowMaskTexture, LightingChannelsTexture, false /*bRenderOverlap*/, true /*bCloudShadow*/, VirtualShadowMapArray.GetUniformBuffer(), ShadowSceneRenderer->VirtualShadowMapMaskBits, VirtualShadowMapId);
+						RenderLight(
+							GraphBuilder, Scene, View, SceneTextures, &LightSceneInfo, 
+							VirtualShadowMapId != INDEX_NONE ? nullptr : ScreenShadowMaskTexture, LightingChannelsTexture, false /*bRenderOverlap*/, true /*bCloudShadow*/, 
+							LightFunctionAtlas.GetLightFunctionAtlasGlobalParameters(ViewIndex, GraphBuilder),
+							VirtualShadowMapArray.GetUniformBuffer(), ShadowSceneRenderer->VirtualShadowMapMaskBits, VirtualShadowMapId);
 					}
 				}
 
 				if (bUseHairLighting)
 				{
-					for (FViewInfo& View : Views)
+					for (int32 ViewIndex = 0, ViewCount = Views.Num(); ViewIndex < ViewCount; ++ViewIndex)
 					{
+						FViewInfo& View = Views[ViewIndex];
+
 						if (HairStrands::HasViewHairStrandsData(View))
 						{
 							FHairStrandsTransmittanceMaskData TransmittanceMaskData;
@@ -2077,7 +2103,10 @@ void FDeferredShadingSceneRenderer::RenderLights(
 
 							// Note: ideally the light should still be evaluated for hair when not casting shadow, but for preserving the old behavior, and not adding 
 							// any perf. regression, we disable this light for hair rendering 
-							RenderLightForHair(GraphBuilder, View, SceneTextures, &LightSceneInfo, HairShadowMask, LightingChannelsTexture, TransmittanceMaskData, false /*bForwardRendering*/);
+							RenderLightForHair(
+								GraphBuilder, View, SceneTextures, &LightSceneInfo, 
+								HairShadowMask, LightingChannelsTexture, TransmittanceMaskData, false /*bForwardRendering*/, 
+								LightFunctionAtlas.GetLightFunctionAtlasGlobalParameters(ViewIndex, GraphBuilder));
 						}
 					}
 				}
@@ -2092,7 +2121,8 @@ static void RenderLightArrayForOverlapViewmode(
 	const TArray<FViewInfo>& Views,
 	const FMinimalSceneTextures& SceneTextures,
 	FRDGTextureRef LightingChannelsTexture,
-	const TSparseArray<FLightSceneInfoCompact, TAlignedSparseArrayAllocator<alignof(FLightSceneInfoCompact)>>& LightArray)
+	const TSparseArray<FLightSceneInfoCompact, TAlignedSparseArrayAllocator<alignof(FLightSceneInfoCompact)>>& LightArray,
+	TRDGUniformBufferRef<FLightFunctionAtlasGlobalParameters> LightFunctionAtlasGlobalParameters)
 {
 	for (auto LightIt = LightArray.CreateConstIterator(); LightIt; ++LightIt)
 	{
@@ -2117,7 +2147,7 @@ static void RenderLightArrayForOverlapViewmode(
 		for (const FViewInfo& View : Views)
 		{
 			SCOPED_GPU_MASK(GraphBuilder.RHICmdList, View.GPUMask);
-			RenderLight(GraphBuilder, Scene, View, SceneTextures, LightSceneInfo, nullptr, LightingChannelsTexture, true /*bRenderOverlap*/, false /*bCloudShadow*/);
+			RenderLight(GraphBuilder, Scene, View, SceneTextures, LightSceneInfo, nullptr, LightingChannelsTexture, true /*bRenderOverlap*/, false /*bCloudShadow*/, LightFunctionAtlasGlobalParameters);
 		}
 	}
 }
@@ -2132,11 +2162,11 @@ void FDeferredShadingSceneRenderer::RenderStationaryLightOverlap(
 		// Clear to discard base pass values in scene color since we didn't skip that, to have valid scene depths
 		AddClearRenderTargetPass(GraphBuilder, SceneTextures.Color.Target, FLinearColor::Black);
 
-		RenderLightArrayForOverlapViewmode(GraphBuilder, Scene, Views, SceneTextures, LightingChannelsTexture, Scene->Lights);
+		RenderLightArrayForOverlapViewmode(GraphBuilder, Scene, Views, SceneTextures, LightingChannelsTexture, Scene->Lights, LightFunctionAtlas.GetDefaultLightFunctionAtlasGlobalParameters(GraphBuilder));
 
 		//Note: making use of FScene::InvisibleLights, which contains lights that haven't been added to the scene in the same way as visible lights
 		// So code called by RenderLightArrayForOverlapViewmode must be careful what it accesses
-		RenderLightArrayForOverlapViewmode(GraphBuilder, Scene, Views, SceneTextures, LightingChannelsTexture, Scene->InvisibleLights);
+		RenderLightArrayForOverlapViewmode(GraphBuilder, Scene, Views, SceneTextures, LightingChannelsTexture, Scene->InvisibleLights, LightFunctionAtlas.GetDefaultLightFunctionAtlasGlobalParameters(GraphBuilder));
 	}
 }
 
@@ -2230,7 +2260,7 @@ static void CalculateLightNearFarDepthFromBounds(const FViewInfo& View, const FS
 static TRDGUniformBufferRef<FDeferredLightUniformStruct> CreateDeferredLightUniformBuffer(FRDGBuilder& GraphBuilder, const FViewInfo& View, const FLightSceneInfo& LightSceneInfo)
 {
 	auto* DeferredLightStruct = GraphBuilder.AllocParameters<FDeferredLightUniformStruct>();
-	*DeferredLightStruct = GetDeferredLightParameters(View, LightSceneInfo);
+	*DeferredLightStruct = GetDeferredLightParameters(View, LightSceneInfo, View.LightFunctionAtlasViewData.GetDeferredlightingUsesLightFunctionAtlas());
 	return GraphBuilder.CreateUniformBuffer(DeferredLightStruct);
 }
 
@@ -2253,6 +2283,7 @@ static FDeferredLightPS::FParameters GetDeferredLightPSParameters(
 	FRDGTextureRef ShadowMaskTexture,
 	FRDGTextureRef LightingChannelsTexture,
 	bool bCloudShadow,
+	TRDGUniformBufferRef<FLightFunctionAtlasGlobalParameters> LightFunctionAtlasGlobalParameters,
 	TRDGUniformBufferRef<FVirtualShadowMapUniformParameters> VirtualShadowMapUniformBuffer = nullptr,
 	FRDGTextureRef ShadowMaskBits = nullptr,
 	int32 VirtualShadowMapId = INDEX_NONE)
@@ -2290,6 +2321,11 @@ static FDeferredLightPS::FParameters GetDeferredLightPSParameters(
 	Out.VirtualShadowMap = VirtualShadowMapUniformBuffer;
 	Out.VirtualShadowMapId = VirtualShadowMapId;
 	Out.ShadowMaskBits = ShadowMaskBits ? ShadowMaskBits : GSystemTextures.GetZeroUIntDummy(GraphBuilder);
+
+	// If the light is not batched, it could be due to shadow, so we still specify light function atlas sampling.
+	Out.bUseLightFunctionAtlas = LightSceneInfo->Proxy->HasValidLightFunctionAtlasSlot() && View.LightFunctionAtlasViewData.GetDeferredlightingUsesLightFunctionAtlas();
+	Out.LightFunctionAtlas = LightFunctionAtlasGlobalParameters;
+
 	// PS - Render Targets
 	Out.RenderTargets[0] = FRenderTargetBinding(SceneColorTexture, ERenderTargetLoadAction::ELoad);
 	if (Substrate::IsOpaqueRoughRefractionEnabled())
@@ -2488,8 +2524,9 @@ static void RenderLight(
 	const FLightSceneInfo* LightSceneInfo,
 	FRDGTextureRef ScreenShadowMaskTexture,
 	FRDGTextureRef LightingChannelsTexture,
-	bool bRenderOverlap, 
+	bool bRenderOverlap,
 	bool bCloudShadow,
+	TRDGUniformBufferRef<FLightFunctionAtlasGlobalParameters> LightFunctionAtlasGlobalParameters,
 	TRDGUniformBufferRef<FVirtualShadowMapUniformParameters> VirtualShadowMapUniformBuffer,
 	FRDGTextureRef ShadowMaskBits,
 	int32 VirtualShadowMapId)
@@ -2549,7 +2586,9 @@ static void RenderLight(
 	{
 		FRenderLightParameters* PassParameters = GraphBuilder.AllocParameters<FRenderLightParameters>();
 		// PS - Generatl parameters
-		PassParameters->PS = GetDeferredLightPSParameters(GraphBuilder, Scene, View, LightSceneInfo, SceneTextures.Color.Target, SceneTextures.Depth.Target, SceneTextures.UniformBuffer, View.HairStrandsViewData.UniformBuffer, ScreenShadowMaskTexture, LightingChannelsTexture, bCloudShadow, VirtualShadowMapUniformBuffer, ShadowMaskBits, VirtualShadowMapId);
+		PassParameters->PS = GetDeferredLightPSParameters(
+			GraphBuilder, Scene, View, LightSceneInfo, SceneTextures.Color.Target, SceneTextures.Depth.Target, SceneTextures.UniformBuffer, View.HairStrandsViewData.UniformBuffer, 
+			ScreenShadowMaskTexture, LightingChannelsTexture, bCloudShadow, LightFunctionAtlasGlobalParameters, VirtualShadowMapUniformBuffer, ShadowMaskBits, VirtualShadowMapId);
 		// VS - General parameters
 		if (bIsRadial)
 		{
@@ -2658,7 +2697,9 @@ void FDeferredShadingSceneRenderer::RenderLightForHair(
 	FRDGTextureRef HairShadowMaskTexture,
 	FRDGTextureRef LightingChannelsTexture,
 	const FHairStrandsTransmittanceMaskData& InTransmittanceMaskData,
-	const bool bForwardRendering)
+	const bool bForwardRendering,
+	TRDGUniformBufferRef<FLightFunctionAtlasGlobalParameters> LightFunctionAtlasGlobalParameters
+	)
 {
 	// Ensure the light is valid for this view
 	const bool bHairRenderingEnabled = HairStrands::HasViewHairStrandsData(View);
@@ -2693,7 +2734,8 @@ void FDeferredShadingSceneRenderer::RenderLightForHair(
 		HairStrands::BindHairStrandsViewUniformParameters(View),
 		HairShadowMaskTexture,
 		LightingChannelsTexture,
-		bCloudShadow);
+		bCloudShadow,
+		LightFunctionAtlasGlobalParameters);
 
 	// PS - Hair parameters
 	const FIntPoint SampleLightingViewportResolution = View.HairStrandsViewData.VisibilityData.SampleLightingViewportResolution;
@@ -2788,8 +2830,10 @@ void FDeferredShadingSceneRenderer::RenderLightsForHair(
 	{
 		RDG_EVENT_SCOPE(GraphBuilder, "DirectLighting");
 
-		for (FViewInfo& View : Views)
+		for (int32 ViewIndex = 0, ViewCount = Views.Num(); ViewIndex < ViewCount; ++ViewIndex)
 		{
+			FViewInfo& View = Views[ViewIndex];
+
 			if (!HairStrands::HasViewHairStrandsData(View))
 			{
 				continue;
@@ -2817,7 +2861,8 @@ void FDeferredShadingSceneRenderer::RenderLightsForHair(
 						ScreenShadowMaskSubPixelTexture,
 						LightingChannelsTexture,
 						TransmittanceMaskData,
-						true /*bForwardRendering*/);
+						true /*bForwardRendering*/,
+						LightFunctionAtlas.GetLightFunctionAtlasGlobalParameters(ViewIndex, GraphBuilder));
 				}
 			}
 		}
diff --git a/Engine/Source/Runtime/Renderer/Private/LightRendering.h b/Engine/Source/Runtime/Renderer/Private/LightRendering.h
--- a/Engine/Source/Runtime/Renderer/Private/LightRendering.h
+++ b/Engine/Source/Runtime/Renderer/Private/LightRendering.h
@@ -35,15 +35,16 @@ extern uint32 GetShadowQuality();
 
 extern float GetLightFadeFactor(const FSceneView& View, const FLightSceneProxy* Proxy);
 
-extern FDeferredLightUniformStruct GetDeferredLightParameters(const FSceneView& View, const FLightSceneInfo& LightSceneInfo, uint32 LightFlags=0);
+extern FDeferredLightUniformStruct GetDeferredLightParameters(const FSceneView& View, const FLightSceneInfo& LightSceneInfo, bool bUseLightFunctionAtlas=false, uint32 LightFlags=0);
 
 inline void SetDeferredLightParameters(
 	FRHIBatchedShaderParameters& BatchedParameters,
 	const TShaderUniformBufferParameter<FDeferredLightUniformStruct>& DeferredLightUniformBufferParameter, 
 	const FLightSceneInfo* LightSceneInfo,
-	const FSceneView& View)
+	const FSceneView& View,
+	bool bUseLightFunctionAtlas)
 {
-	SetUniformBufferParameterImmediate(BatchedParameters, DeferredLightUniformBufferParameter, GetDeferredLightParameters(View, *LightSceneInfo));
+	SetUniformBufferParameterImmediate(BatchedParameters, DeferredLightUniformBufferParameter, GetDeferredLightParameters(View, *LightSceneInfo, bUseLightFunctionAtlas));
 }
 
 extern FDeferredLightUniformStruct GetSimpleDeferredLightParameters(
diff --git a/Engine/Source/Runtime/Renderer/Private/LightSceneInfo.cpp b/Engine/Source/Runtime/Renderer/Private/LightSceneInfo.cpp
--- a/Engine/Source/Runtime/Renderer/Private/LightSceneInfo.cpp
+++ b/Engine/Source/Runtime/Renderer/Private/LightSceneInfo.cpp
@@ -376,13 +376,22 @@ uint32 FLightSceneInfo::PackLightTypeAndShadowMapChannelMask(bool bAllowStaticLi
 		(CurrentDynamicShadowMapChannel == 2 ? 64 : 0) |
 		(CurrentDynamicShadowMapChannel == 3 ? 128 : 0);
 
-	Result |= Proxy->GetLightingChannelMask() << 8;
+	uint32 BitOffset = 8;
+
+	Result |= Proxy->GetLightingChannelMask() << BitOffset;				BitOffset += 8;					// This could be 3 bits
+
 	// pack light type in this uint32 as well
-	Result |= ((uint32)Proxy->GetLightType()) << 16;
+	Result |= ((uint32)Proxy->GetLightType()) << BitOffset;				BitOffset += LightType_NumBits;
+
 	const uint32 CastShadows = Proxy->CastsDynamicShadow() ? 1 : 0;
-	Result |= CastShadows << (16 + LightType_NumBits);
-	uint32 HasLightFunction = bLightFunction ? 1 : 0;;
-	Result |= HasLightFunction << (16 + LightType_NumBits + 1);
+	Result |= CastShadows << BitOffset;									BitOffset += 1;
+
+	uint32 HasLightFunction = bLightFunction ? 1 : 0;
+	Result |= HasLightFunction << BitOffset;							BitOffset += 1;
+
+	Result |= Proxy->LightFunctionAtlasLightIndex << (BitOffset);		BitOffset += 8;
+
+	// 28 bits used, 4 bits free
 
 	return Result;
 }
diff --git a/Engine/Source/Runtime/Renderer/Private/ScenePrivate.h b/Engine/Source/Runtime/Renderer/Private/ScenePrivate.h
--- a/Engine/Source/Runtime/Renderer/Private/ScenePrivate.h
+++ b/Engine/Source/Runtime/Renderer/Private/ScenePrivate.h
@@ -55,6 +55,7 @@
 #include "GlobalDistanceField.h"
 #include "Algo/RemoveIf.h"
 #include "UObject/Package.h"
+#include "LightFunctionAtlas.h"
 
 /** Factor by which to grow occlusion tests **/
 #define OCCLUSION_SLOP (1.0f)
@@ -3181,6 +3182,8 @@ public:
 	/** Substrate data shared between all views. */
 	FSubstrateSceneData SubstrateSceneData;
 
+	FLightFunctionAtlasSceneData LightFunctionAtlasSceneData;
+
 	/** Mask used to determine whether primitives that draw to a runtime virtual texture should also be drawn in the main pass. */
 	uint8 RuntimeVirtualTexturePrimitiveHideMaskEditor;
 	uint8 RuntimeVirtualTexturePrimitiveHideMaskGame;
diff --git a/Engine/Source/Runtime/Renderer/Private/SceneRendering.h b/Engine/Source/Runtime/Renderer/Private/SceneRendering.h
--- a/Engine/Source/Runtime/Renderer/Private/SceneRendering.h
+++ b/Engine/Source/Runtime/Renderer/Private/SceneRendering.h
@@ -48,6 +48,7 @@
 #include "Async/Mutex.h"
 #include "LocalFogVolumeRendering.h"
 #include "Nanite/NaniteShared.h"
+#include "LightFunctionAtlas.h"
 
 #if RHI_RAYTRACING
 #include "RayTracingInstanceBufferUtil.h"
@@ -822,22 +823,12 @@ public:
 
 	FRDGTextureRef IntegratedLightScatteringTexture = nullptr;
 
-	// TODO: right now the lightfunction atlas is dedicated to the volumetric fog.
-	// Later we could put the allocated atlas tiles on FLightSceneInfo and uploaded as light data on GPU
-	// so that the lightfunction atlas can be used for forward rendering or tiled lighting.
-	// For this to work we would also need to add the default white light functoin as an atlas item.
-	// Note: this is not a smart pointer since it is allocated using the GraphBuilder frame transient memory.
-	FTransientLightFunctionTextureAtlas* TransientLightFunctionTextureAtlas = nullptr;
-
-	TMap<FLightSceneInfo*, FVolumetricFogLocalLightFunctionInfo> LocalLightFunctionData;
-
 	FVolumetricFogViewResources()
 	{}
 
 	void Release()
 	{
 		IntegratedLightScatteringTexture = nullptr;
-		TransientLightFunctionTextureAtlas = nullptr;
 	}
 };
 
@@ -1373,6 +1364,8 @@ public:
 
 	FHairStrandsViewData HairStrandsViewData;
 
+	FLightFunctionAtlasViewData LightFunctionAtlasViewData;
+
 	/** Parameters for exponential height fog. */
 	FVector4f ExponentialFogParameters;
 	FVector4f ExponentialFogParameters2;
@@ -2054,6 +2047,8 @@ public:
 	// TODO: Move to deferred scene renderer
 	TUniquePtr<FShadowSceneRenderer> ShadowSceneRenderer;
 
+	FLightFunctionAtlas LightFunctionAtlas;
+
 	/** If a freeze request has been made */
 	bool bHasRequestedToggleFreeze;
 
@@ -2221,11 +2216,13 @@ public:
 		float VolumetricFogMaxDistance,
 		FMatrix44f& OutLightFunctionTranslatedWorldToShadow,
 		FRDGTexture*& OutLightFunctionTexture,
-		bool& bOutUseDirectionalLightShadowing);
+		bool& bOutUseDirectionalLightShadowing,
+		FLightSceneInfo*& OutDirectionalLightSceneInfo);
 
 	void RenderLocalLightsForVolumetricFog(
 		FRDGBuilder& GraphBuilder,
 		FViewInfo& View,
+		uint32 ViewIndex,
 		bool bUseTemporalReprojection,
 		const struct FVolumetricFogIntegrationParameterData& IntegrationData,
 		const FExponentialHeightFogSceneInfo& FogInfo,
diff --git a/Engine/Source/Runtime/Renderer/Private/ShadowRendering.cpp b/Engine/Source/Runtime/Renderer/Private/ShadowRendering.cpp
--- a/Engine/Source/Runtime/Renderer/Private/ShadowRendering.cpp
+++ b/Engine/Source/Runtime/Renderer/Private/ShadowRendering.cpp
@@ -466,7 +466,8 @@ static void BindShaderShaders(FRHICommandList& RHICmdList, FGraphicsPipelineStat
 
 	FRHIBatchedShaderParameters& BatchedParameters = RHICmdList.GetScratchShaderParameters();
 
-	PixelShader->SetParameters(BatchedParameters, ViewIndex, View, ShadowInfo);
+	const bool bUseLightFunctionAtlas = View.LightFunctionAtlasViewData.GetDeferredlightingUsesLightFunctionAtlas();
+	PixelShader->SetParameters(BatchedParameters, ViewIndex, View, ShadowInfo, bUseLightFunctionAtlas);
 
 	if (Substrate::IsSubstrateEnabled())
 	{
diff --git a/Engine/Source/Runtime/Renderer/Private/ShadowRendering.h b/Engine/Source/Runtime/Renderer/Private/ShadowRendering.h
--- a/Engine/Source/Runtime/Renderer/Private/ShadowRendering.h
+++ b/Engine/Source/Runtime/Renderer/Private/ShadowRendering.h
@@ -1251,7 +1251,8 @@ public:
 		FRHIBatchedShaderParameters& BatchedParameters,
 		int32 ViewIndex,
 		const FSceneView& View,
-		const FProjectedShadowInfo* ShadowInfo)
+		const FProjectedShadowInfo* ShadowInfo,
+		bool bUseLightFunctionAtlas)
 	{
 		const FVector PreViewTranslation = View.ViewMatrices.GetPreViewTranslation();
 		const bool bUseFadePlaneEnable = ShadowInfo->CascadeSettings.FadePlaneLength > 0;
@@ -1267,7 +1268,7 @@ public:
 
 		if (DeferredLightParameter.IsBound())
 		{
-			SetDeferredLightParameters(BatchedParameters, DeferredLightParameter, &ShadowInfo->GetLightSceneInfo(), View);
+			SetDeferredLightParameters(BatchedParameters, DeferredLightParameter, &ShadowInfo->GetLightSceneInfo(), View, bUseLightFunctionAtlas);
 		}
 
 
@@ -1314,9 +1315,10 @@ public:
 		FRHIBatchedShaderParameters& BatchedParameters,
 		int32 ViewIndex,
 		const FSceneView& View,
-		const FProjectedShadowInfo* ShadowInfo)
+		const FProjectedShadowInfo* ShadowInfo,
+		bool bUseLightFunctionAtlas)
 	{
-		TShadowProjectionPS<Quality, false, true>::SetParameters(BatchedParameters, ViewIndex, View, ShadowInfo);
+		TShadowProjectionPS<Quality, false, true>::SetParameters(BatchedParameters, ViewIndex, View, ShadowInfo, bUseLightFunctionAtlas);
 		SetShaderValue(BatchedParameters, ModulatedShadowColorParameter, ShadowInfo->GetLightSceneInfo().Proxy->GetModulatedShadowColor());
 	}
 
@@ -1381,9 +1383,10 @@ public:
 		FRHIBatchedShaderParameters& BatchedParameters,
 		int32 ViewIndex,
 		const FSceneView& View,
-		const FProjectedShadowInfo* ShadowInfo)
+		const FProjectedShadowInfo* ShadowInfo,
+		bool bUseLightFunctionAtlas)
 	{
-		TShadowProjectionPS<Quality>::SetParameters(BatchedParameters, ViewIndex, View, ShadowInfo);
+		TShadowProjectionPS<Quality>::SetParameters(BatchedParameters, ViewIndex, View, ShadowInfo, bUseLightFunctionAtlas);
 
 		FTranslucentSelfShadowUniformParameters TranslucentSelfShadowUniformParameters;
 		SetupTranslucentSelfShadowUniformParameters(ShadowInfo, TranslucentSelfShadowUniformParameters);
@@ -1595,7 +1598,7 @@ public:
 
 		if (DeferredLightParameter.IsBound())
 		{
-			SetDeferredLightParameters(BatchedParameters, DeferredLightParameter, &ShadowInfo->GetLightSceneInfo(), View);
+			SetDeferredLightParameters(BatchedParameters, DeferredLightParameter, &ShadowInfo->GetLightSceneInfo(), View, View.LightFunctionAtlasViewData.GetDeferredlightingUsesLightFunctionAtlas());
 		}
 	}
 
@@ -1673,9 +1676,10 @@ public:
 		FRHIBatchedShaderParameters& BatchedParameters,
 		int32 ViewIndex,
 		const FSceneView& View,
-		const FProjectedShadowInfo* ShadowInfo)
+		const FProjectedShadowInfo* ShadowInfo,
+		bool bUseLightFunctionAtlas)
 	{
-		TShadowProjectionPS<Quality, bUseFadePlane>::SetParameters(BatchedParameters, ViewIndex, View, ShadowInfo);
+		TShadowProjectionPS<Quality, bUseFadePlane>::SetParameters(BatchedParameters, ViewIndex, View, ShadowInfo, bUseLightFunctionAtlas);
 
 		// GetLightSourceAngle returns the full angle.
 		float TanLightSourceAngle = FMath::Tan(0.5 * FMath::DegreesToRadians(ShadowInfo->GetLightSceneInfo().Proxy->GetLightSourceAngle()));
@@ -1727,11 +1731,12 @@ public:
 		FRHIBatchedShaderParameters& BatchedParameters,
 		int32 ViewIndex,
 		const FSceneView& View,
-		const FProjectedShadowInfo* ShadowInfo)
+		const FProjectedShadowInfo* ShadowInfo,
+		bool bUseLightFunctionAtlas)
 	{
 		check(ShadowInfo->GetLightSceneInfo().Proxy->GetLightType() == LightType_Spot);
 
-		TShadowProjectionPS<Quality, bUseFadePlane>::SetParameters(BatchedParameters, ViewIndex, View, ShadowInfo);
+		TShadowProjectionPS<Quality, bUseFadePlane>::SetParameters(BatchedParameters, ViewIndex, View, ShadowInfo, bUseLightFunctionAtlas);
 
 		static IConsoleVariable* CVarMaxSoftShadowKernelSize = IConsoleManager::Get().FindConsoleVariable(TEXT("r.Shadow.MaxSoftKernelSize"));
 		check(CVarMaxSoftShadowKernelSize);
diff --git a/Engine/Source/Runtime/Renderer/Private/VolumetricFog.cpp b/Engine/Source/Runtime/Renderer/Private/VolumetricFog.cpp
--- a/Engine/Source/Runtime/Renderer/Private/VolumetricFog.cpp
+++ b/Engine/Source/Runtime/Renderer/Private/VolumetricFog.cpp
@@ -27,6 +27,7 @@ VolumetricFog.cpp
 #include "VirtualShadowMaps/VirtualShadowMapClipmap.h"
 #include "ProfilingDebugging/CpuProfilerTrace.h"
 #include "DataDrivenShaderPlatformInfo.h"
+#include "LightFunctionAtlas.h"
 
 int32 GVolumetricFog = 1;
 FAutoConsoleVariableRef CVarVolumetricFog(
@@ -133,14 +134,6 @@ FAutoConsoleVariableRef CVarVolumetricFogConservativeDepth(
 	ECVF_Scalability | ECVF_RenderThreadSafe
 );
 
-int GVolumetricFogLightFunction = 1;
-FAutoConsoleVariableRef CVarVolumetricFogLightFunction(
-	TEXT("r.VolumetricFog.LightFunction"),
-	GVolumetricFogLightFunction,
-	TEXT("Whether light functions are generated to be sampled when rendering volumetric fog."),
-	ECVF_Scalability | ECVF_RenderThreadSafe
-);
-
 int GVolumetricFogInjectRaytracedLights = 0;
 FAutoConsoleVariableRef CVarVolumetricInjectRaytracedLights(
 	TEXT("r.VolumetricFog.InjectRaytracedLights"),
@@ -292,14 +285,12 @@ IMPLEMENT_GLOBAL_SHADER(FWriteToBoundingSphereVS, "/Engine/Private/VolumetricFog
 BEGIN_SHADER_PARAMETER_STRUCT(FInjectShadowedLocalLightCommonParameters, )
 	SHADER_PARAMETER_STRUCT_REF(FViewUniformShaderParameters, ViewUniformBuffer)
 	SHADER_PARAMETER_RDG_UNIFORM_BUFFER(FDeferredLightUniformStruct, DeferredLight)
+	SHADER_PARAMETER_RDG_UNIFORM_BUFFER(FLightFunctionAtlasGlobalParameters, LightFunctionAtlas)
 	SHADER_PARAMETER_RDG_TEXTURE(Texture2D, WhiteDummyTexture)
-	SHADER_PARAMETER_RDG_TEXTURE(Texture2D, LightFunctionAtlasTexture)
-	SHADER_PARAMETER_SAMPLER(SamplerState, LightFunctionAtlasSampler)
 	SHADER_PARAMETER_STRUCT_INCLUDE(FVolumetricFogIntegrationParameters, VolumetricFogParameters)
-	SHADER_PARAMETER(FMatrix44f, LocalLightFunctionMatrix)
-	SHADER_PARAMETER(FVector4f, LightFunctionAtlasTileMinMaxUvBound)
 	SHADER_PARAMETER(float, PhaseG)
 	SHADER_PARAMETER(float, InverseSquaredLightDistanceBiasScale)
+	SHADER_PARAMETER(uint32, LightFunctionAtlasLightIndex)
 END_SHADER_PARAMETER_STRUCT()
 
 static bool SetupInjectShadowedLocalLightCommonParameters(
@@ -307,33 +298,11 @@ static bool SetupInjectShadowedLocalLightCommonParameters(
 	FViewInfo& View,
 	const FVolumetricFogIntegrationParameterData& IntegrationData,
 	const FExponentialHeightFogSceneInfo& FogInfo,
-	const FLightSceneInfo* LightSceneInfo,	
-	bool bUsesLightFunction,
+	const FLightSceneInfo* LightSceneInfo,
 	FInjectShadowedLocalLightCommonParameters& OutCommonParameters)
 {
-	OutCommonParameters.LightFunctionAtlasTexture = View.VolumetricFogResources.TransientLightFunctionTextureAtlas ? View.VolumetricFogResources.TransientLightFunctionTextureAtlas->GetTransientLightFunctionAtlasTexture() : GSystemTextures.GetWhiteDummy(GraphBuilder);;
-	OutCommonParameters.LocalLightFunctionMatrix = FMatrix44f::Identity;
-	OutCommonParameters.LightFunctionAtlasTileMinMaxUvBound = FVector4f(ForceInitToZero);
-	if (bUsesLightFunction)
-	{
-		const TMap<FLightSceneInfo*, FVolumetricFogLocalLightFunctionInfo>& LocalLightFunctionData = View.VolumetricFogResources.LocalLightFunctionData;
-
-		const FVolumetricFogLocalLightFunctionInfo* LightFunctionData = LocalLightFunctionData.Find(LightSceneInfo);
-
-		if (!ensure(LightFunctionData != nullptr))
-		{
-			// The light function data is missing but the light requires it. Skip this light for now.
-			return false;
-		}
-
-		OutCommonParameters.LocalLightFunctionMatrix = LightFunctionData->LightFunctionTranslatedWorldToLightMatrix;
-		OutCommonParameters.LightFunctionAtlasTexture = LightFunctionData->AtlasTile.Texture;
-		OutCommonParameters.LightFunctionAtlasTileMinMaxUvBound = LightFunctionData->AtlasTile.MinMaxUvBound;
-	}
-
-	OutCommonParameters.LightFunctionAtlasSampler = TStaticSamplerState<SF_Bilinear, AM_Clamp, AM_Clamp, AM_Clamp>::GetRHI();
 	// We also bind the default light function texture because when we are out of atlas tile, we fallback to use a white light function so we need the RHI to be created
-	OutCommonParameters.WhiteDummyTexture = View.VolumetricFogResources.TransientLightFunctionTextureAtlas ? View.VolumetricFogResources.TransientLightFunctionTextureAtlas->GetDefaultLightFunctionTexture() : GSystemTextures.GetWhiteDummy(GraphBuilder);
+	OutCommonParameters.WhiteDummyTexture = GSystemTextures.GetWhiteDummy(GraphBuilder);
 	SetupVolumetricFogIntegrationParameters(OutCommonParameters.VolumetricFogParameters, View, IntegrationData);
 
 	OutCommonParameters.ViewUniformBuffer = View.ViewUniformBuffer;
@@ -617,12 +586,6 @@ bool LightNeedsSeparateInjectionIntoVolumetricFogForOpaqueShadow(const FViewInfo
 	return false;
 }
 
-bool LightNeedsSeparateInjectionIntoVolumetricFogForLightFunction(const FLightSceneInfo* LightSceneInfo)
-{
-	// No directional light type because it is handled in a specific way in RenderLightFunctionForVolumetricFog.
-	return GVolumetricFogLightFunction > 0 && (LightSceneInfo->Proxy->GetLightType() == LightType_Point || LightSceneInfo->Proxy->GetLightType() == LightType_Spot || LightSceneInfo->Proxy->GetLightType() == LightType_Rect);
-}
-
 FIntPoint CalculateVolumetricFogBoundsForLight(const FSphere& LightBounds, const FViewInfo& View, FIntVector VolumetricFogGridSize, FVector GridZParams)
 {
 	FIntPoint VolumeZBounds;
@@ -718,7 +681,8 @@ TGlobalResource<FCircleRasterizeIndexBuffer> GCircleRasterizeIndexBuffer;
 
 void FSceneRenderer::RenderLocalLightsForVolumetricFog(
 	FRDGBuilder& GraphBuilder,
-	FViewInfo& View,
+	FViewInfo& View, 
+	uint32 ViewIndex,
 	bool bUseTemporalReprojection,
 	const FVolumetricFogIntegrationParameterData& IntegrationData,
 	const FExponentialHeightFogSceneInfo& FogInfo,
@@ -737,13 +701,11 @@ void FSceneRenderer::RenderLocalLightsForVolumetricFog(
 		const FLightSceneInfo* LightSceneInfo = LightSceneInfoCompact.LightSceneInfo;
 
 		bool bIsShadowed = LightNeedsSeparateInjectionIntoVolumetricFogForOpaqueShadow(View, LightSceneInfo, VisibleLightInfos[LightSceneInfo->Id]);
-		bool bUsesLightFunction = ViewFamily.EngineShowFlags.LightFunctions 
-			&& CheckForLightFunction(LightSceneInfo) && LightNeedsSeparateInjectionIntoVolumetricFogForLightFunction(LightSceneInfo);
 		bool bUsesRectLightTexture = GVolumetricFogRectLightTexture && LightSceneInfo->Proxy->HasSourceTexture();
 
 		if (LightSceneInfo->ShouldRenderLightViewIndependent()
 			&& LightSceneInfo->ShouldRenderLight(View)
-			&& (bIsShadowed || bUsesLightFunction || bUsesRectLightTexture)
+			&& (bIsShadowed || bUsesRectLightTexture)
 			&& LightSceneInfo->Proxy->GetVolumetricScatteringIntensity() > 0)
 		{
 			const FSphere LightBounds = LightSceneInfo->Proxy->GetBoundingSphere();
@@ -763,6 +725,10 @@ void FSceneRenderer::RenderLocalLightsForVolumetricFog(
 		}
 	}
 
+	// Setup the light function atlas
+	const bool bUseLightFunctionAtlas = View.LightFunctionAtlasViewData.GetVolumetricFogUsesLightFunctionAtlas();
+	TRDGUniformBufferRef<FLightFunctionAtlasGlobalParameters> LightFunctionAtlasGlobalParameters = LightFunctionAtlas.GetLightFunctionAtlasGlobalParameters(ViewIndex, GraphBuilder);
+
 	// Now voxelise all the light we have just gathered.
 	bool bClearExecuted = false;
 	if (LightsToInject.Num() > 0)
@@ -777,8 +743,6 @@ void FSceneRenderer::RenderLocalLightsForVolumetricFog(
 			if (VolumeZBounds.X < VolumeZBounds.Y)
 			{
 				bool bIsShadowed = LightNeedsSeparateInjectionIntoVolumetricFogForOpaqueShadow(View, LightSceneInfo, VisibleLightInfo);
-				bool bUsesLightFunction = ViewFamily.EngineShowFlags.LightFunctions
-					&& CheckForLightFunction(LightSceneInfo) && LightNeedsSeparateInjectionIntoVolumetricFogForLightFunction(LightSceneInfo);
 				bool bUsesRectLightTexture = GVolumetricFogRectLightTexture && LightSceneInfo->Proxy->HasSourceTexture();
 
 				int32 VirtualShadowMapId = VisibleLightInfo.GetVirtualShadowMapId(&View);
@@ -793,9 +757,9 @@ void FSceneRenderer::RenderLocalLightsForVolumetricFog(
 					IntegrationData,
 					FogInfo,
 					LightSceneInfo,
-					bUsesLightFunction,
 					PassParameters->Common
 				);
+				PassParameters->Common.LightFunctionAtlas = LightFunctionAtlasGlobalParameters;
 
 				if (!bValid)
 				{
@@ -820,7 +784,7 @@ void FSceneRenderer::RenderLocalLightsForVolumetricFog(
 				FInjectShadowedLocalLightPS::FPermutationDomain PermutationVector;
 				PermutationVector.Set< FInjectShadowedLocalLightPS::FDynamicallyShadowed >(bDynamicallyShadowed);
 				PermutationVector.Set< FInjectShadowedLocalLightPS::FTemporalReprojection >(bUseTemporalReprojection);
-				PermutationVector.Set< FInjectShadowedLocalLightPS::FLightFunction >(bUsesLightFunction);
+				PermutationVector.Set< FInjectShadowedLocalLightPS::FLightFunction >(bUseLightFunctionAtlas);
 				PermutationVector.Set< FInjectShadowedLocalLightPS::FEnableShadows >(bIsShadowed);
 				PermutationVector.Set< FInjectShadowedLocalLightPS::FVirtualShadowMap >(bUseVSM);
 				PermutationVector.Set< FInjectShadowedLocalLightPS::FRectLightTexture >(bUsesRectLightTexture);
@@ -903,8 +867,6 @@ void FSceneRenderer::RenderLocalLightsForVolumetricFog(
 			const FIntPoint VolumeZBounds = CalculateVolumetricFogBoundsForLight(LightBounds, View, VolumetricFogViewGridSize, GridZParams);
 			if (VolumeZBounds.X < VolumeZBounds.Y)
 			{
-				bool bUsesLightFunction = ViewFamily.EngineShowFlags.LightFunctions
-					&& CheckForLightFunction(LightSceneInfo) && LightNeedsSeparateInjectionIntoVolumetricFogForLightFunction(LightSceneInfo);
 				bool bUsesRectLightTexture = GVolumetricFogRectLightTexture && LightSceneInfo->Proxy->HasSourceTexture();
 
 				FInjectShadowedLocalLightRGS::FParameters* PassParameters = GraphBuilder.AllocParameters<FInjectShadowedLocalLightRGS::FParameters>();
@@ -918,9 +880,9 @@ void FSceneRenderer::RenderLocalLightsForVolumetricFog(
 					IntegrationData,
 					FogInfo,
 					LightSceneInfo,
-					bUsesLightFunction,
 					PassParameters->Common
 				);
+				PassParameters->Common.LightFunctionAtlas = LightFunctionAtlasGlobalParameters;
 
 				if (!bValid)
 				{
@@ -929,7 +891,7 @@ void FSceneRenderer::RenderLocalLightsForVolumetricFog(
 
 				FInjectShadowedLocalLightRGS::FPermutationDomain PermutationVector;
 				PermutationVector.Set< FInjectShadowedLocalLightRGS::FTemporalReprojection >(bUseTemporalReprojection);
-				PermutationVector.Set< FInjectShadowedLocalLightRGS::FLightFunction >(bUsesLightFunction);
+				PermutationVector.Set< FInjectShadowedLocalLightRGS::FLightFunction >(bUseLightFunctionAtlas);
 				PermutationVector.Set< FInjectShadowedLocalLightRGS::FRectLightTexture >(bUsesRectLightTexture);
 
 				TShaderMapRef<FInjectShadowedLocalLightRGS> RayGenerationShader(GetGlobalShaderMap(FeatureLevel), PermutationVector);
@@ -973,6 +935,7 @@ class FVolumetricFogLightScatteringCS : public FGlobalShader
 	class FVirtualShadowMap				: SHADER_PERMUTATION_BOOL("VIRTUAL_SHADOW_MAP");
 	class FCloudTransmittance			: SHADER_PERMUTATION_BOOL("USE_CLOUD_TRANSMITTANCE");
 	class FRaytracedShadowsVolume		: SHADER_PERMUTATION_BOOL("USE_RAYTRACED_SHADOWS_VOLUME");
+	class FSampleLightFunctionAtlas		: SHADER_PERMUTATION_BOOL("USE_LIGHT_FUNCTION_ATLAS");
 	
 	using FPermutationDomain = TShaderPermutationDomain<
 		FSuperSampleCount,
@@ -981,19 +944,21 @@ class FVolumetricFogLightScatteringCS : public FGlobalShader
 		FLumenGI,
 		FVirtualShadowMap,
 		FCloudTransmittance,
-		FRaytracedShadowsVolume>;
+		FRaytracedShadowsVolume,
+		FSampleLightFunctionAtlas>;
 
 	BEGIN_SHADER_PARAMETER_STRUCT(FParameters, )
 		SHADER_PARAMETER_STRUCT_REF(FViewUniformShaderParameters, View)
 		SHADER_PARAMETER_RDG_UNIFORM_BUFFER(FForwardLightData, Forward)
 		SHADER_PARAMETER_RDG_UNIFORM_BUFFER(FFogUniformParameters, Fog)
+		SHADER_PARAMETER_RDG_UNIFORM_BUFFER(FLightFunctionAtlasGlobalParameters, LightFunctionAtlas)
 		SHADER_PARAMETER_STRUCT_INCLUDE(FVolumetricFogIntegrationParameters, VolumetricFogParameters)
 
 		SHADER_PARAMETER_RDG_TEXTURE(Texture2D, VBufferA)
 		SHADER_PARAMETER_RDG_TEXTURE(Texture2D, VBufferB)
 		SHADER_PARAMETER_RDG_TEXTURE(Texture2D, LocalShadowedLightScattering)
-		SHADER_PARAMETER_RDG_TEXTURE(Texture2D, LightFunctionTexture)
-		SHADER_PARAMETER_SAMPLER(SamplerState, LightFunctionSampler)
+		SHADER_PARAMETER_RDG_TEXTURE(Texture2D, DirectionalLightLightFunctionTexture)
+		SHADER_PARAMETER_SAMPLER(SamplerState, DirectionalLightLightFunctionSampler)
 		SHADER_PARAMETER_RDG_TEXTURE(Texture2D, CloudShadowmapTexture)
 		SHADER_PARAMETER_SAMPLER(SamplerState, CloudShadowmapSampler)
 		SHADER_PARAMETER_RDG_TEXTURE(Texture2D, ConservativeDepthTexture)
@@ -1025,6 +990,8 @@ class FVolumetricFogLightScatteringCS : public FGlobalShader
 		SHADER_PARAMETER(uint32, UseConservativeDepthTexture)
 		SHADER_PARAMETER(uint32, UseEmissive)
 		SHADER_PARAMETER(uint32, MobileHasDirectionalLight)
+		SHADER_PARAMETER(uint32, DirectionalApplyLightFunctionFromAtlas)
+		SHADER_PARAMETER(uint32, DirectionalLightFunctionAtlasLightIndex)
 	END_SHADER_PARAMETER_STRUCT()
 
 	static FIntVector GetGroupSize()
@@ -1338,6 +1305,7 @@ void FSceneRenderer::ComputeVolumetricFog(FRDGBuilder& GraphBuilder,
 	}
 
 	const FExponentialHeightFogSceneInfo& FogInfo = Scene->ExponentialFogs[0];
+	const bool bUseLightFunctionAtlas = Scene->LightFunctionAtlasSceneData.GetVolumetricFogUsesLightFunctionAtlas();
 
 	TRACE_CPUPROFILER_EVENT_SCOPE(FSceneRenderer::ComputeVolumetricFog);
 	QUICK_SCOPE_CYCLE_COUNTER(STAT_VolumetricFog);
@@ -1429,7 +1397,8 @@ void FSceneRenderer::ComputeVolumetricFog(FRDGBuilder& GraphBuilder,
 				AtmosphericDirectionalLightIndex = 1;
 			}
 		}
-
+		
+		FLightSceneInfo* DirectionalLightSceneInfo = nullptr;
 		RenderLightFunctionForVolumetricFog(
 			GraphBuilder,
 			View,
@@ -1438,7 +1407,8 @@ void FSceneRenderer::ComputeVolumetricFog(FRDGBuilder& GraphBuilder,
 			FogInfo.VolumetricFogDistance,
 			DirectionalLightFunctionTranslatedWorldToShadow,
 			DirectionalLightFunctionTexture,
-			bUseDirectionalLightShadowing);
+			bUseDirectionalLightShadowing,
+			DirectionalLightSceneInfo);
 			
 		View.VolumetricFogResources.IntegratedLightScatteringTexture = nullptr;
 		TRDGUniformBufferRef<FFogUniformParameters> FogUniformBuffer = CreateFogUniformBuffer(GraphBuilder, View);
@@ -1463,7 +1433,7 @@ void FSceneRenderer::ComputeVolumetricFog(FRDGBuilder& GraphBuilder,
 		}
 
 		FRDGTexture* LocalShadowedLightScattering = GraphBuilder.RegisterExternalTexture(GSystemTextures.VolumetricBlackDummy);
-		RenderLocalLightsForVolumetricFog(GraphBuilder, View, bUseTemporalReprojection, IntegrationData, FogInfo,
+		RenderLocalLightsForVolumetricFog(GraphBuilder, View, uint32(ViewIndex), bUseTemporalReprojection, IntegrationData, FogInfo,
 			VolumetricFogViewGridSize, GridZParams, VolumeDescFastVRAM, LocalShadowedLightScattering, ConservativeDepthTexture);
 	
 		FRDGTextureRef RaytracedShadowsVolume = nullptr;
@@ -1547,8 +1517,8 @@ void FSceneRenderer::ComputeVolumetricFog(FRDGBuilder& GraphBuilder,
 			}
 
 			PassParameters->DirectionalLightFunctionTranslatedWorldToShadow = DirectionalLightFunctionTranslatedWorldToShadow;
-			PassParameters->LightFunctionTexture = DirectionalLightFunctionTexture;
-			PassParameters->LightFunctionSampler = TStaticSamplerState<SF_Bilinear, AM_Clamp, AM_Clamp, AM_Clamp>::GetRHI();
+			PassParameters->DirectionalLightLightFunctionTexture = DirectionalLightFunctionTexture;
+			PassParameters->DirectionalLightLightFunctionSampler = TStaticSamplerState<SF_Bilinear, AM_Clamp, AM_Clamp, AM_Clamp>::GetRHI();
 
 			auto* LumenUniforms = GraphBuilder.AllocParameters<FLumenTranslucencyLightingUniforms>();
 			LumenUniforms->Parameters = GetLumenTranslucencyLightingParameters(GraphBuilder, View.GetLumenTranslucencyGIVolume(), View.LumenFrontLayerTranslucency);
@@ -1646,6 +1616,18 @@ void FSceneRenderer::ComputeVolumetricFog(FRDGBuilder& GraphBuilder,
 
 			PassParameters->RaytracedShadowsVolume = RaytracedShadowsVolume ? GraphBuilder.CreateSRV(RaytracedShadowsVolume) : nullptr;
 
+			PassParameters->LightFunctionAtlas = LightFunctionAtlas.GetLightFunctionAtlasGlobalParameters(ViewIndex, GraphBuilder);
+			if (bUseLightFunctionAtlas)
+			{
+				PassParameters->DirectionalApplyLightFunctionFromAtlas = bUseLightFunctionAtlas && DirectionalLightSceneInfo && DirectionalLightSceneInfo->Proxy->HasValidLightFunctionAtlasSlot() ? 1 :0;
+				PassParameters->DirectionalLightFunctionAtlasLightIndex = DirectionalLightSceneInfo ? DirectionalLightSceneInfo->Proxy->GetLightFunctionAtlasLightIndex() : 0;
+			}
+			else
+			{
+				PassParameters->DirectionalApplyLightFunctionFromAtlas = 0;
+				PassParameters->DirectionalLightFunctionAtlasLightIndex = 0;
+			}
+
 			const bool bUseLumenGI = View.GetLumenTranslucencyGIVolume().Texture0 != nullptr && FDataDrivenShaderPlatformInfo::GetSupportsLumenGI(View.GetShaderPlatform());
 			const bool bUseGlobalDistanceField = UseGlobalDistanceField() && Scene->DistanceFieldSceneData.NumObjectsInBuffer > 0;
 			const bool bUseRaytracedShadowsVolume = RaytracedShadowsVolume != nullptr;
@@ -1673,6 +1655,7 @@ void FSceneRenderer::ComputeVolumetricFog(FRDGBuilder& GraphBuilder,
 			PermutationVector.Set< FVolumetricFogLightScatteringCS::FVirtualShadowMap >(VirtualShadowMapArray.IsAllocated() );
 			PermutationVector.Set< FVolumetricFogLightScatteringCS::FCloudTransmittance >(AtmosphericDirectionalLightIndex >= 0);
 			PermutationVector.Set< FVolumetricFogLightScatteringCS::FRaytracedShadowsVolume >(bUseRaytracedShadowsVolume);
+			PermutationVector.Set< FVolumetricFogLightScatteringCS::FSampleLightFunctionAtlas >(bUseLightFunctionAtlas);
 
 			auto ComputeShader = View.ShaderMap->GetShader< FVolumetricFogLightScatteringCS >(PermutationVector);
 			ClearUnusedGraphResources(ComputeShader, PassParameters);
@@ -1684,7 +1667,7 @@ void FSceneRenderer::ComputeVolumetricFog(FRDGBuilder& GraphBuilder,
 					VolumetricFogViewGridSize.Z,
 					SuperSampleCount,
 					bUseDistanceFieldSkyOcclusion ? TEXT("DFAO") : TEXT(""),
-					PassParameters->LightFunctionTexture ? TEXT("LF") : TEXT(""),
+					PassParameters->DirectionalLightLightFunctionTexture ? TEXT("LF") : TEXT(""),
 					bUseLumenGI ? TEXT("Lumen") : TEXT("")),
 				PassParameters,
 				ERDGPassFlags::Compute,
diff --git a/Engine/Source/Runtime/Renderer/Private/VolumetricFog.h b/Engine/Source/Runtime/Renderer/Private/VolumetricFog.h
--- a/Engine/Source/Runtime/Renderer/Private/VolumetricFog.h
+++ b/Engine/Source/Runtime/Renderer/Private/VolumetricFog.h
@@ -10,43 +10,6 @@
 #include "SceneView.h"
 #include "SceneRendering.h"
 
-
-class FTransientLightFunctionTextureAtlas
-{
-public:
-
-	FTransientLightFunctionTextureAtlas(FRDGBuilder& GraphBuilder);
-	~FTransientLightFunctionTextureAtlas();
-
-	// FTransientLightFunctionTextureAtlasTile will never be null, but it can be a default white light function
-	FTransientLightFunctionTextureAtlasTile AllocateAtlasTile();
-
-	FRDGTextureRef GetTransientLightFunctionAtlasTexture()
-	{
-		return TransientLightFunctionAtlasTexture;
-	}
-	FRDGTextureRef GetDefaultLightFunctionTexture()
-	{
-		return DefaultLightFunctionAtlasItemTexture;
-	}
-
-	uint32 GetAtlasTextureWidth()
-	{
-		return AtlasTextureWidth;
-	}
-
-private:
-	FTransientLightFunctionTextureAtlas() {}
-
-	uint32 AtlasItemWidth;
-	uint32 AtlasTextureWidth;
-	uint32 AllocatedAtlasTiles;
-	float HalfTexelSize;
-
-	FRDGTextureRef TransientLightFunctionAtlasTexture;
-	FRDGTextureRef DefaultLightFunctionAtlasItemTexture;
-};
-
 // Grid size for resource allocation to be independent of dynamic resolution
 extern FIntVector GetVolumetricFogResourceGridSize(const FViewInfo& View, int32& OutVolumetricFogGridPixelSize);
 // Grid size for the view rectangle within the allocated resource
@@ -61,4 +24,3 @@ extern bool ShouldRenderVolumetricFog(const FScene* Scene, const FSceneViewFamil
 extern const FProjectedShadowInfo* GetShadowForInjectionIntoVolumetricFog(const FVisibleLightInfo& VisibleLightInfo);
 
 extern bool LightNeedsSeparateInjectionIntoVolumetricFogForOpaqueShadow(const FViewInfo& View, const FLightSceneInfo* LightSceneInfo, const FVisibleLightInfo& VisibleLightInfo);
-extern bool LightNeedsSeparateInjectionIntoVolumetricFogForLightFunction(const FLightSceneInfo* LightSceneInfo);
diff --git a/Engine/Source/Runtime/Renderer/Private/VolumetricFogLightFunction.cpp b/Engine/Source/Runtime/Renderer/Private/VolumetricFogLightFunction.cpp
--- a/Engine/Source/Runtime/Renderer/Private/VolumetricFogLightFunction.cpp
+++ b/Engine/Source/Runtime/Renderer/Private/VolumetricFogLightFunction.cpp
@@ -14,22 +14,6 @@
 #include "PipelineStateCache.h"
 #include "ShadowRendering.h"
 
-int GVolumetricFogLightFunctionResolution = 128;
-FAutoConsoleVariableRef CVarVolumetricFogLightFunctionResolution(
-	TEXT("r.VolumetricFog.LightFunction.Resolution"),
-	GVolumetricFogLightFunctionResolution,
-	TEXT("The resolution of all light functions generated to be sampled when rendering volumetric fog."),
-	ECVF_Scalability | ECVF_RenderThreadSafe
-	);
-
-int GVolumetricFogLightFunctionCount = 16;
-FAutoConsoleVariableRef CVarVolumetricFogLightFunctionCount(
-	TEXT("r.VolumetricFog.LightFunction.LightFunctionCount"),
-	GVolumetricFogLightFunctionCount,
-	TEXT("The maximum light function that can be rendered per frame."),
-	ECVF_Scalability | ECVF_RenderThreadSafe
-	);
-
 float GVolumetricFogLightFunctionDirectionalLightSupersampleScale = 2.0f;
 FAutoConsoleVariableRef CVarVolumetricFogLightFunctionSupersampleScale(
 	TEXT("r.VolumetricFog.LightFunction.DirectionalLightSupersampleScale"),
@@ -44,12 +28,6 @@ static bool inline LocalLightLighFunctionsEnabled()
 	return GVolumetricFogLightFunction > 0;
 }
 
-static int32 GetVolumetricFogLightFunctionResolution()
-{
-	const int32 VolumetricFogLightFunctionResolution = FMath::Clamp(GVolumetricFogLightFunctionResolution, 32, 256); // Clamp to reasonable values
-	return VolumetricFogLightFunctionResolution;
-}
-
 class FVolumetricFogLightFunctionPS : public FMaterialShader
 {
 	DECLARE_SHADER_TYPE(FVolumetricFogLightFunctionPS, Material);
@@ -130,8 +108,6 @@ BEGIN_SHADER_PARAMETER_STRUCT(FVolumetricFogLightFunctionParameters, )
 	RENDER_TARGET_BINDING_SLOTS()
 END_SHADER_PARAMETER_STRUCT()
 
-
-
 void FSceneRenderer::RenderLightFunctionForVolumetricFog(
 	FRDGBuilder& GraphBuilder,
 	FViewInfo& View,
@@ -140,18 +116,15 @@ void FSceneRenderer::RenderLightFunctionForVolumetricFog(
 	float VolumetricFogMaxDistance,
 	FMatrix44f& OutDirectionalLightFunctionTranslatedWorldToShadow,
 	FRDGTexture*& OutDirectionalLightFunctionTexture,
-	bool& bOutUseDirectionalLightShadowing)
+	bool& bOutUseDirectionalLightShadowing,
+	FLightSceneInfo*& OutDirectionalLightSceneInfo)
 {
-	TMap<FLightSceneInfo*, FVolumetricFogLocalLightFunctionInfo>& LocalLightFunctionData = View.VolumetricFogResources.LocalLightFunctionData;
-	LocalLightFunctionData.Reset();
-	TArray<FLightSceneInfo*> LocalLightsToEvaluate;
-
 	// The only directional light we can accept in the volumetric fog because we use the forward lighting data in the Scattering compute shader.
 	const FLightSceneProxy* SelectedForwardDirectionalLightProxy = View.ForwardLightingResources.SelectedForwardDirectionalLightProxy;
 	// Default directional light properties
 	OutDirectionalLightFunctionTranslatedWorldToShadow = FMatrix44f::Identity;
 	bOutUseDirectionalLightShadowing = false;
-	FLightSceneInfo* DirectionalLightSceneInfo = NULL;
+	OutDirectionalLightSceneInfo = NULL;
 
 	// Gather lights that need to evaluate light functions
 	for (auto LightIt = Scene->Lights.CreateConstIterator(); LightIt; ++LightIt)
@@ -159,48 +132,30 @@ void FSceneRenderer::RenderLightFunctionForVolumetricFog(
 		const FLightSceneInfoCompact& LightSceneInfoCompact = *LightIt;
 		FLightSceneInfo* LightSceneInfo = LightSceneInfoCompact.LightSceneInfo;
 
-		if (ViewFamily.EngineShowFlags.LightFunctions
-			&& LightSceneInfo
+		if (   LightSceneInfo
 			&& LightSceneInfo->Proxy
 			// Band-aid fix for extremely rare case that light scene proxy contains NaNs.
 			&& !LightSceneInfo->Proxy->GetDirection().ContainsNaN()
 			&& LightSceneInfo->ShouldRenderLightViewIndependent()
 			&& LightSceneInfo->ShouldRenderLight(View))
 		{
-			if (DirectionalLightSceneInfo == NULL && LightSceneInfo->Proxy == SelectedForwardDirectionalLightProxy &&
+			if (OutDirectionalLightSceneInfo == NULL && LightSceneInfo->Proxy == SelectedForwardDirectionalLightProxy &&
 				LightSceneInfo->Proxy->GetLightType() == LightType_Directional)
 			{
 				// We only take the first directional light into account
 				bOutUseDirectionalLightShadowing = LightSceneInfo->Proxy->CastsVolumetricShadow();
 
-				if (CheckForLightFunction(LightSceneInfo))
-				{
-					DirectionalLightSceneInfo = LightSceneInfo;
-				}
-			}
-			else if (LocalLightLighFunctionsEnabled() && CheckForLightFunction(LightSceneInfo))
-			{
-				const FMaterialRenderProxy* MaterialProxy = LightSceneInfo->Proxy->GetLightFunctionMaterial();
-				if (!MaterialProxy || !MaterialProxy->GetIncompleteMaterialWithFallback(Scene->GetFeatureLevel()).IsLightFunction())
-				{
-					continue;
-				}
-
-				uint8 LightType = LightSceneInfo->Proxy->GetLightType();
-				if (LightType == LightType_Spot
-					|| LightType == LightType_Point
-					|| LightType == LightType_Rect)
+				if (CheckForLightFunction(LightSceneInfo) && ViewFamily.EngineShowFlags.LightFunctions)
 				{
-					LocalLightsToEvaluate.Add(LightSceneInfo);
+					OutDirectionalLightSceneInfo = LightSceneInfo;
+					break;	// we only deal with the directional light function so we directly stop now
 				}
 			}
 		}
 	}
 
-
-
-	// Now bake the light function into a 2d transient texture for the special single directional light we have selected
-	if (DirectionalLightSceneInfo)
+	// Now bake the directional light function into a 2d transient texture for the special single directional light we have selected
+	if (OutDirectionalLightSceneInfo)
 	{
 		RDG_EVENT_SCOPE(GraphBuilder, "DirectionalLightFunction");
 
@@ -212,7 +167,7 @@ void FSceneRenderer::RenderLightFunctionForVolumetricFog(
 			const FVector ViewUp = View.ViewMatrices.GetOverriddenTranslatedViewMatrix().GetColumn(1);
 			const FVector ViewRight = View.ViewMatrices.GetOverriddenTranslatedViewMatrix().GetColumn(0);
 
-			const FVector LightDirection = DirectionalLightSceneInfo->Proxy->GetDirection().GetSafeNormal();
+			const FVector LightDirection = OutDirectionalLightSceneInfo->Proxy->GetDirection().GetSafeNormal();
 			FVector AxisWeights;
 			AxisWeights.X = FMath::Abs(LightDirection | ViewRight) * VolumetricFogGridSize.X;
 			AxisWeights.Y = FMath::Abs(LightDirection | ViewUp) * VolumetricFogGridSize.Y;
@@ -248,7 +203,7 @@ void FSceneRenderer::RenderLightFunctionForVolumetricFog(
 			FWholeSceneProjectedShadowInitializer ShadowInitializer;
 
 			check(VolumetricFogMaxDistance > 0);
-			FSphere Bounds = DirectionalLightSceneInfo->Proxy->GetShadowSplitBoundsDepthRange(View, View.ViewMatrices.GetViewOrigin(), 0, VolumetricFogMaxDistance, NULL);
+			FSphere Bounds = OutDirectionalLightSceneInfo->Proxy->GetShadowSplitBoundsDepthRange(View, View.ViewMatrices.GetViewOrigin(), 0, VolumetricFogMaxDistance, NULL);
 			check(Bounds.W > 0);
 
 			const float ShadowExtent = Bounds.W / FMath::Sqrt(3.0f);
@@ -266,7 +221,7 @@ void FSceneRenderer::RenderLightFunctionForVolumetricFog(
 			ShadowInitializer.CascadeSettings.bFarShadowCascade = false;
 
 			ProjectedShadowInfo.SetupWholeSceneProjection(
-				DirectionalLightSceneInfo,
+				OutDirectionalLightSceneInfo,
 				&View,
 				ShadowInitializer,
 				LightFunctionResolution.X,
@@ -287,7 +242,7 @@ void FSceneRenderer::RenderLightFunctionForVolumetricFog(
 			FRDGTexture* LightFunctionTexture = GraphBuilder.CreateTexture(LightFunctionTextureDesc, TEXT("VolumetricFog.LightFunction"));
 			OutDirectionalLightFunctionTexture = LightFunctionTexture;
 
-			const FMaterialRenderProxy* MaterialProxyForRendering = DirectionalLightSceneInfo->Proxy->GetLightFunctionMaterial();
+			const FMaterialRenderProxy* MaterialProxyForRendering = OutDirectionalLightSceneInfo->Proxy->GetLightFunctionMaterial();
 			const FMaterial& Material = MaterialProxyForRendering->GetMaterialWithFallback(Scene->GetFeatureLevel(), MaterialProxyForRendering);
 
 			FVolumetricFogLightFunctionParameters* PassParameters = GraphBuilder.AllocParameters<FVolumetricFogLightFunctionParameters>();
@@ -299,7 +254,7 @@ void FSceneRenderer::RenderLightFunctionForVolumetricFog(
 				RDG_EVENT_NAME("LightFunction %ux%u Material=%s", LightFunctionResolution.X, LightFunctionResolution.Y, *(Material.GetFriendlyName())),
 				PassParameters,
 				ERDGPassFlags::Raster,
-				[PassParameters, &View, MaterialProxyForRendering, &Material, LightFunctionResolution, DirectionalLightSceneInfo, LightFunctionTranslatedWorldToShadowMatrix](FRHICommandList& RHICmdList)
+				[PassParameters, &View, MaterialProxyForRendering, &Material, LightFunctionResolution, OutDirectionalLightSceneInfo, LightFunctionTranslatedWorldToShadowMatrix](FRHICommandList& RHICmdList)
 				{
 					RHICmdList.SetViewport(0.f, 0.f, 0.f, LightFunctionResolution.X, LightFunctionResolution.Y, 1.f);
 
@@ -313,7 +268,7 @@ void FSceneRenderer::RenderLightFunctionForVolumetricFog(
 					const FMaterialShaderMap* MaterialShaderMap = Material.GetRenderingThreadShaderMap();
 					TShaderMapRef<FPostProcessVS> VertexShader(View.ShaderMap);
 
-					check(DirectionalLightSceneInfo->Proxy->GetLightType() == LightType_Directional)
+					check(OutDirectionalLightSceneInfo->Proxy->GetLightType() == LightType_Directional)
 					FVolumetricFogLightFunctionPS::FPermutationDomain PermutationVector;
 					PermutationVector.Set<FVolumetricFogLightFunctionPS::FLightType>(3);
 					TShaderRef<FVolumetricFogLightFunctionPS> PixelShader = MaterialShaderMap->GetShader<FVolumetricFogLightFunctionPS>(PermutationVector);
@@ -324,7 +279,7 @@ void FSceneRenderer::RenderLightFunctionForVolumetricFog(
 
 					SetGraphicsPipelineState(RHICmdList, GraphicsPSOInit, 0);
 
-					FVolumetricFogLightFunctionPS::FParameters PS = PixelShader->GetParameters(View, DirectionalLightSceneInfo, FVector2D(1.0f / LightFunctionResolution.X, 1.0f / LightFunctionResolution.Y), LightFunctionTranslatedWorldToShadowMatrix.Inverse());
+					FVolumetricFogLightFunctionPS::FParameters PS = PixelShader->GetParameters(View, OutDirectionalLightSceneInfo, FVector2D(1.0f / LightFunctionResolution.X, 1.0f / LightFunctionResolution.Y), LightFunctionTranslatedWorldToShadowMatrix.Inverse());
 
 					ClearUnusedGraphResources(PixelShader, &PS);
 					SetShaderParametersMixedPS(RHICmdList, PixelShader, PS, View, MaterialProxyForRendering);
@@ -342,197 +297,4 @@ void FSceneRenderer::RenderLightFunctionForVolumetricFog(
 			);
 		}
 	}
-
-	if (!LocalLightsToEvaluate.Num())
-	{
-		return; // Do not run atlas operation if no light needs it
-	}
-
-	// We are going to render local light LightFunction so we allocate the atlas.
-	View.VolumetricFogResources.TransientLightFunctionTextureAtlas = new (GraphBuilder.Alloc(sizeof(FTransientLightFunctionTextureAtlas), alignof(FTransientLightFunctionTextureAtlas))) FTransientLightFunctionTextureAtlas(GraphBuilder);
-
-	// Compute all the data required by the lights
-	const FIntPoint LightFunctionResolution = GetVolumetricFogLightFunctionResolution();
-	const FMatrix TranslatedWorldToWorld = FTranslationMatrix(-View.ViewMatrices.GetPreViewTranslation());
-	for (FLightSceneInfo* LightSceneInfo : LocalLightsToEvaluate)
-	{
-		uint8 LightType = LightSceneInfo->Proxy->GetLightType();
-
-		FVolumetricFogLocalLightFunctionInfo& LightFunctionData = LocalLightFunctionData.Add(LightSceneInfo);
-		LightFunctionData.AtlasTile = View.VolumetricFogResources.TransientLightFunctionTextureAtlas->AllocateAtlasTile();
-
-		// Spotlights needs to utilize their shadow projection when rendering the light function 
-		if (LightType == LightType_Spot)
-		{
-			TArray<FWholeSceneProjectedShadowInitializer, TInlineAllocator<6>> Initializers;
-			LightSceneInfo->Proxy->GetWholeSceneProjectedShadowInitializer(*View.Family, Initializers);
-
-			FWholeSceneProjectedShadowInitializer ShadowInitializer = Initializers[0];
-			FProjectedShadowInfo ProjectedShadowInfo;
-			ProjectedShadowInfo.SetupWholeSceneProjection(
-				LightSceneInfo,
-				&View,
-				ShadowInitializer,
-				LightFunctionResolution.X,
-				LightFunctionResolution.Y,
-				LightFunctionResolution.X,
-				LightFunctionResolution.Y,
-				0
-			);
-
-			FVector4f ShadowmapMinMaxValue;
-			LightFunctionData.LightFunctionTranslatedWorldToLightMatrix = FMatrix44f(TranslatedWorldToWorld * ProjectedShadowInfo.GetWorldToShadowMatrix(ShadowmapMinMaxValue, &LightFunctionResolution));
-		}
-		else if (LightType == LightType_Point)
-		{
-			LightFunctionData.LightFunctionTranslatedWorldToLightMatrix = FMatrix44f(TranslatedWorldToWorld * LightSceneInfo->Proxy->GetWorldToLight());
-		}
-		else if (LightType == LightType_Rect)
-		{
-			const FVector Scale = LightSceneInfo->Proxy->GetLightFunctionScale();
-			// Switch x and z so that z of the user specified scale affects the distance along the light direction
-			const FVector InverseScale = FVector(1.0 / Scale.Z, 1.0 / Scale.Y, 1.0 / Scale.X);
-			const FMatrix WorldToLight = LightSceneInfo->Proxy->GetWorldToLight() * FScaleMatrix(InverseScale);
-
-			LightFunctionData.LightFunctionTranslatedWorldToLightMatrix = FMatrix44f(TranslatedWorldToWorld * WorldToLight);
-		}
-	}
-
-	// And now process all the local lights, allocated atlas tile and generate the tile as needed
-	FVolumetricFogLightFunctionParameters* PassParameters = GraphBuilder.AllocParameters<FVolumetricFogLightFunctionParameters>();
-	PassParameters->RenderTargets[0] = FRenderTargetBinding(View.VolumetricFogResources.TransientLightFunctionTextureAtlas->GetTransientLightFunctionAtlasTexture(), ERenderTargetLoadAction::ENoAction);
-	PassParameters->SceneTextures = SceneTextures.UniformBuffer;
-	GraphBuilder.AddPass(
-		RDG_EVENT_NAME("Local Lights LightFunctions"),
-		PassParameters,
-		ERDGPassFlags::Raster,
-		[PassParameters, &View, &Scene = Scene](FRHICommandList& RHICmdList)
-		{
-			uint32 AtlasTextureWidth = View.VolumetricFogResources.TransientLightFunctionTextureAtlas->GetAtlasTextureWidth();
-			uint32 LightFunctionResolution = GetVolumetricFogLightFunctionResolution();
-			for (auto& Itr : View.VolumetricFogResources.LocalLightFunctionData)
-			{
-				FLightSceneInfo* LightSceneInfo = Itr.Key;
-				FVolumetricFogLocalLightFunctionInfo& LightFunctionInfo = Itr.Value;
-
-				if (!LightFunctionInfo.AtlasTile.bIsDefault)
-				{
-					const FMaterialRenderProxy* MaterialProxyForRendering = LightSceneInfo->Proxy->GetLightFunctionMaterial();
-					const FMaterial& Material = MaterialProxyForRendering->GetMaterialWithFallback(Scene->GetFeatureLevel(), MaterialProxyForRendering);
-
-					RHICmdList.SetViewport(0.f, 0.f, 0.f, AtlasTextureWidth, AtlasTextureWidth, 1.f);
-
-					FGraphicsPipelineStateInitializer GraphicsPSOInit;
-					RHICmdList.ApplyCachedRenderTargets(GraphicsPSOInit);
-					GraphicsPSOInit.RasterizerState = TStaticRasterizerState<FM_Solid, CM_None>::GetRHI();
-					GraphicsPSOInit.BlendState = TStaticBlendState<>::GetRHI();
-					GraphicsPSOInit.DepthStencilState = TStaticDepthStencilState<false, CF_Always>::GetRHI();
-					GraphicsPSOInit.PrimitiveType = PT_TriangleList;
-
-					const FMaterialShaderMap* MaterialShaderMap = Material.GetRenderingThreadShaderMap();
-					TShaderMapRef<FPostProcessVS> VertexShader(View.ShaderMap);
-
-					FVolumetricFogLightFunctionPS::FPermutationDomain PermutationVector;
-					switch (LightSceneInfo->Proxy->GetLightType())
-					{
-					case LightType_Point:
-						PermutationVector.Set<FVolumetricFogLightFunctionPS::FLightType>(0);
-						break;
-					case LightType_Spot:
-						PermutationVector.Set<FVolumetricFogLightFunctionPS::FLightType>(1);
-						break;
-					case LightType_Rect:
-						PermutationVector.Set<FVolumetricFogLightFunctionPS::FLightType>(2);
-						break;
-					default:
-						check(false)
-					}
-					TShaderRef<FVolumetricFogLightFunctionPS> PixelShader = MaterialShaderMap->GetShader<FVolumetricFogLightFunctionPS>(PermutationVector);
-
-					GraphicsPSOInit.BoundShaderState.VertexDeclarationRHI = GFilterVertexDeclaration.VertexDeclarationRHI;
-					GraphicsPSOInit.BoundShaderState.VertexShaderRHI = VertexShader.GetVertexShader();
-					GraphicsPSOInit.BoundShaderState.PixelShaderRHI = PixelShader.GetPixelShader();
-
-					SetGraphicsPipelineState(RHICmdList, GraphicsPSOInit, 0);
-
-					FVolumetricFogLightFunctionPS::FParameters PS = 
-						PixelShader->GetParameters(View, LightSceneInfo, 
-							FVector2D(1.0f / LightFunctionResolution, 1.0f / LightFunctionResolution),
-							LightFunctionInfo.LightFunctionTranslatedWorldToLightMatrix.Inverse());
-
-					ClearUnusedGraphResources(PixelShader, &PS);
-					SetShaderParametersMixedPS(RHICmdList, PixelShader, PS, View, MaterialProxyForRendering);
-
-					FIntPoint RectSize = LightFunctionInfo.AtlasTile.RectBound.Size();
-					DrawRectangle(
-						RHICmdList,
-						LightFunctionInfo.AtlasTile.RectBound.Min.X, LightFunctionInfo.AtlasTile.RectBound.Min.Y,
-						RectSize.X, RectSize.Y,
-						LightFunctionInfo.AtlasTile.RectBound.Min.X, LightFunctionInfo.AtlasTile.RectBound.Min.Y,
-						RectSize.X, RectSize.Y,
-						AtlasTextureWidth,
-						LightFunctionResolution,
-						VertexShader);
-				}
-			}
-		}
-	);
-}
-
-
-
-FTransientLightFunctionTextureAtlas::FTransientLightFunctionTextureAtlas(FRDGBuilder& GraphBuilder)
-	: AtlasItemWidth(0)
-	, AtlasTextureWidth(0)
-	, AllocatedAtlasTiles(0)
-	, HalfTexelSize(0.0f)
-{
-	AtlasItemWidth = FMath::Clamp(FMath::CeilToInt(FMath::Sqrt(float(GVolumetricFogLightFunctionCount))), 2, 16); // Clamp to reasonable values
-
-	AtlasTextureWidth = AtlasItemWidth * GetVolumetricFogLightFunctionResolution();
-	HalfTexelSize = 0.5f / float(AtlasTextureWidth);
-
-	FRDGTextureDesc LightFunctionTextureDesc = FRDGTextureDesc::Create2D(FIntPoint(AtlasTextureWidth, AtlasTextureWidth), PF_G8, FClearValueBinding::None, TexCreate_ShaderResource | TexCreate_RenderTargetable);
-	TransientLightFunctionAtlasTexture = GraphBuilder.CreateTexture(LightFunctionTextureDesc, TEXT("VolumetricFog.LightFunctionAtlasTexture"));
-	DefaultLightFunctionAtlasItemTexture = GSystemTextures.GetWhiteDummy(GraphBuilder);
-}
-
-FTransientLightFunctionTextureAtlas::~FTransientLightFunctionTextureAtlas()
-{
-	//
-}
-
-FTransientLightFunctionTextureAtlasTile FTransientLightFunctionTextureAtlas::AllocateAtlasTile()
-{
-	FTransientLightFunctionTextureAtlasTile AtlasTile;
-
-	if (AllocatedAtlasTiles < AtlasItemWidth * AtlasItemWidth)
-	{
-		const uint32 LightFunctionResolution = GetVolumetricFogLightFunctionResolution();
-		const FIntPoint TileCoord(AllocatedAtlasTiles % AtlasItemWidth, AllocatedAtlasTiles / AtlasItemWidth);
-		const FIntPoint PixelCoord = TileCoord * LightFunctionResolution;
-
-		AtlasTile.bIsDefault = false;
-		AtlasTile.Texture = TransientLightFunctionAtlasTexture;
-
-		AtlasTile.RectBound = FIntRect(PixelCoord, PixelCoord + FIntPoint(LightFunctionResolution, LightFunctionResolution));
-
-		const float InvAtlasTextureWidthFloat = 1.0f / float(AtlasTextureWidth);
-		const float MinX = float(AtlasTile.RectBound.Min.X) * InvAtlasTextureWidthFloat + HalfTexelSize;
-		const float MinY = float(AtlasTile.RectBound.Min.Y) * InvAtlasTextureWidthFloat + HalfTexelSize;
-		const float MaxX = float(AtlasTile.RectBound.Max.X) * InvAtlasTextureWidthFloat - HalfTexelSize;
-		const float MaxY = float(AtlasTile.RectBound.Max.Y) * InvAtlasTextureWidthFloat - HalfTexelSize;
-		AtlasTile.MinMaxUvBound = FVector4f(MinX, MinY, MaxX, MaxY);
-
-		AllocatedAtlasTiles++;
-	}
-	else
-	{
-		AtlasTile.bIsDefault = true;
-		AtlasTile.Texture = DefaultLightFunctionAtlasItemTexture;
-		AtlasTile.RectBound = FIntRect(0, 0, 1, 1);
-		AtlasTile.MinMaxUvBound = FVector4f(0.0f, 0.0f, 1.0f, 1.0f);
-	}
-
-	return AtlasTile;
 }
