LFAtlas - fixed crash in the mobile shading path.

diff --git a/Engine/Source/Runtime/Renderer/Private/MobileShadingRenderer.cpp b/Engine/Source/Runtime/Renderer/Private/MobileShadingRenderer.cpp
--- a/Engine/Source/Runtime/Renderer/Private/MobileShadingRenderer.cpp
+++ b/Engine/Source/Runtime/Renderer/Private/MobileShadingRenderer.cpp
@@ -1051,6 +1051,8 @@ void FMobileSceneRenderer::Render(FRDGBuilder& GraphBuilder)
 			{
 				ComputeLightGrid(GraphBuilder, bEnableClusteredLocalLights, SortedLightSet);
 			}
+
+			LightFunctionAtlas.RenderLightFunctionAtlas(GraphBuilder, Views);
 		}
 		else
 		{
