Hide editor primitive when visualizing post process stack.

diff --git a/Engine/Source/Runtime/Renderer/Private/SceneRendering.cpp b/Engine/Source/Runtime/Renderer/Private/SceneRendering.cpp
--- a/Engine/Source/Runtime/Renderer/Private/SceneRendering.cpp
+++ b/Engine/Source/Runtime/Renderer/Private/SceneRendering.cpp
@@ -4206,18 +4206,20 @@ FSceneRenderer* FSceneRenderer::CreateSceneRenderer(const FSceneViewFamily* InVi
 
 bool FSceneRenderer::ShouldCompositeEditorPrimitives(const FViewInfo& View)
 {
-	if (View.Family->EngineShowFlags.VisualizeHDR || View.Family->EngineShowFlags.VisualizeSkyLightIlluminance || View.Family->UseDebugViewPS())
+	const FEngineShowFlags& ShowFlags = View.Family->EngineShowFlags;
+	if (ShowFlags.VisualizeHDR || ShowFlags.VisualizeSkyLightIlluminance ||
+		ShowFlags.VisualizePostProcessStack || View.Family->UseDebugViewPS())
 	{
 		// certain visualize modes get obstructed too much
 		return false;
 	}
 
-	if (View.Family->EngineShowFlags.Wireframe)
+	if (ShowFlags.Wireframe)
 	{
 		// We want wireframe view use MSAA if possible.
 		return true;
 	}
-	else if (View.Family->EngineShowFlags.CompositeEditorPrimitives)
+	else if (ShowFlags.CompositeEditorPrimitives)
 	{
 	    // Any elements that needed compositing were drawn then compositing should be done
 	    if (View.ViewMeshElements.Num() 
