Use correct test when checking for UAV texture format support. Without explicit platofrm test.
#tests pc and platofrm preview

diff --git a/Engine/Source/Runtime/Renderer/Private/LocalFogVolumeRendering.cpp b/Engine/Source/Runtime/Renderer/Private/LocalFogVolumeRendering.cpp
--- a/Engine/Source/Runtime/Renderer/Private/LocalFogVolumeRendering.cpp
+++ b/Engine/Source/Runtime/Renderer/Private/LocalFogVolumeRendering.cpp
@@ -516,11 +516,13 @@ void CreateViewLocalFogVolumeBufferSRV(const FScene* Scene, FViewInfo& View, FRD
 	const uint32 TileDataTextureSliceCount					= LocalFogVolumeTileMaxInstanceCount + 1; // +1 because the first slice is the culled instance count
 
 	EPixelFormat TileDataFormat = PF_R8_UINT;	
-	if(IsMobilePlatform(View.GetShaderPlatform()) && IsOpenGLPlatform(View.GetShaderPlatform())) // TODO: !UE::PixelFormat::HasCapabilities(TileDataFormat, EPixelFormatCapabilities::UAV)
+	if(!UE::PixelFormat::HasCapabilities(TileDataFormat, EPixelFormatCapabilities::UAV))
 	{
+		// Some mobile platforms do not support UAV onto R8. A 32bit format is required.
 		TileDataFormat = PF_R8G8B8A8_UINT;
+		check(UE::PixelFormat::HasCapabilities(TileDataFormat, EPixelFormatCapabilities::UAV));
 	}
-	// TODO:  check(UE::PixelFormat::HasCapabilities(TileDataFormat, EPixelFormatCapabilities::UAV));
+
 	FRDGTextureDesc Texture2DArrayDesc(FRDGTextureDesc::Create2DArray(TileDataTextureResolution, TileDataFormat, FClearValueBinding(EClearBinding::ENoneBound), TexCreate_ShaderResource | TexCreate_UAV | TexCreate_ReduceMemoryWithTilingMode | TexCreate_3DTiling, TileDataTextureSliceCount));
 	// LFV TODO, consider FFastVramConfig onto Texture2DArrayDesc
 
