LFAtlas - Fix for crash with simple lights.

diff --git a/Engine/Shaders/Private/DeferredLightPixelShaders.usf b/Engine/Shaders/Private/DeferredLightPixelShaders.usf
--- a/Engine/Shaders/Private/DeferredLightPixelShaders.usf
+++ b/Engine/Shaders/Private/DeferredLightPixelShaders.usf
@@ -55,7 +55,9 @@
 #include "VirtualShadowMaps/VirtualShadowMapMaskBitsCommon.ush"
 #endif
 
+#if USE_LIGHT_FUNCTION_ATLAS
 #include "LightFunctionAtlas/LightFunctionAtlasCommon.usf"
+#endif
 
 // Remap light shape permutation into light type
 #if LIGHT_SOURCE_SHAPE == 0
@@ -337,12 +339,14 @@ void DeferredLightPixelMain(
 		float4 LightAttenuation = GetLightAttenuationFromShadow(InputParams, SceneDepth, DerivedParams.TranslatedWorldPosition);
 		float4 Radiance = GetDynamicLighting(DerivedParams.TranslatedWorldPosition, DerivedParams.CameraVector, ScreenSpaceData.GBuffer, ScreenSpaceData.AmbientOcclusion, ScreenSpaceData.GBuffer.ShadingModelID, LightData, LightAttenuation, Dither, uint2(InputParams.PixelPos), SurfaceShadow);
 
+#if USE_LIGHT_FUNCTION_ATLAS
 		if (bUseLightFunctionAtlas)
 		{
 			Radiance *= GetLocalLightFunctionCommon(
 				DerivedParams.TranslatedWorldPosition,
 				LightData.LightFunctionAtlasLightIndex);
 		}
+#endif
 
 		OutColor += Radiance;
 	}
@@ -450,12 +454,14 @@ void DeferredLightPixelMain(
 
 	const float LocalCoverage = From8bitCoverage(Sample.Coverage8bit);
 
+#if USE_LIGHT_FUNCTION_ATLAS
 	if (bUseLightFunctionAtlas)
 	{
 		Attenuation *= GetLocalLightFunctionCommon(
 			TranslatedWorldPosition,
 			LightData.LightFunctionAtlasLightIndex);
 	}
+#endif
 
 	OutColor.rgb = AtmosphericTransmittance * Radiance.xyz * Attenuation * LocalCoverage * GetExposure();
 	OutColor.a = LocalCoverage;
diff --git a/Engine/Source/Runtime/Renderer/Private/LightRendering.cpp b/Engine/Source/Runtime/Renderer/Private/LightRendering.cpp
--- a/Engine/Source/Runtime/Renderer/Private/LightRendering.cpp
+++ b/Engine/Source/Runtime/Renderer/Private/LightRendering.cpp
@@ -808,6 +808,7 @@ class FDeferredLightPS : public FGlobalShader
 	class FSourceShapeDim		: SHADER_PERMUTATION_ENUM_CLASS("LIGHT_SOURCE_SHAPE", ELightSourceShape);
 	class FSourceTextureDim		: SHADER_PERMUTATION_BOOL("USE_SOURCE_TEXTURE");
 	class FIESProfileDim		: SHADER_PERMUTATION_BOOL("USE_IES_PROFILE");
+	class FLightFunctionAtlasDim: SHADER_PERMUTATION_BOOL("USE_LIGHT_FUNCTION_ATLAS");
 	class FVisualizeCullingDim	: SHADER_PERMUTATION_BOOL("VISUALIZE_LIGHT_CULLING");
 	class FLightingChannelsDim	: SHADER_PERMUTATION_BOOL("USE_LIGHTING_CHANNELS");
 	class FTransmissionDim		: SHADER_PERMUTATION_BOOL("USE_TRANSMISSION");
@@ -823,6 +824,7 @@ class FDeferredLightPS : public FGlobalShader
 		FSourceShapeDim,
 		FSourceTextureDim,
 		FIESProfileDim,
+		FLightFunctionAtlasDim,
 		FVisualizeCullingDim,
 		FLightingChannelsDim,
 		FTransmissionDim,
@@ -870,6 +872,7 @@ class FDeferredLightPS : public FGlobalShader
 		if (PermutationVector.Get< FVisualizeCullingDim >() && (
 			PermutationVector.Get< FSourceTextureDim >() ||
 			PermutationVector.Get< FIESProfileDim >() ||
+			PermutationVector.Get< FLightFunctionAtlasDim >() ||
 			PermutationVector.Get< FTransmissionDim >() ||
 			PermutationVector.Get< FHairLighting >() ||
 			PermutationVector.Get< FHairComplexTransmittance >() ||
@@ -963,6 +966,7 @@ class FDeferredLightPS : public FGlobalShader
 		{
 			PermutationVector.Set< FSourceTextureDim >(false);
 			PermutationVector.Set< FIESProfileDim >(false);
+			PermutationVector.Set< FLightFunctionAtlasDim >(false);
 			PermutationVector.Set< FTransmissionDim >(false);
 			PermutationVector.Set< FHairLighting >(false);
 			PermutationVector.Set< FAtmosphereTransmittance >(false);
@@ -2617,6 +2621,7 @@ static void RenderLight(
 		PermutationVector.Set< FDeferredLightPS::FVirtualShadowMapMask >(bUseVirtualShadowMapMask);
 		PermutationVector.Set< FDeferredLightPS::FSubstrateTileType >(0);
 		PermutationVector.Set< FDeferredLightPS::FHairComplexTransmittance >(bNeedComplexTransmittanceSupport);
+		PermutationVector.Set< FDeferredLightPS::FLightFunctionAtlasDim >(LightSceneInfo->Proxy->GetLightFunctionMaterial()!=nullptr && View.Family->EngineShowFlags.LightFunctions);
 		if (bIsRadial)
 		{
 			PermutationVector.Set< FDeferredLightPS::FSourceShapeDim >(LightProxy->IsRectLight() ? ELightSourceShape::Rect : ELightSourceShape::Capsule);
@@ -2763,6 +2768,7 @@ void FDeferredShadingSceneRenderer::RenderLightForHair(
 	PermutationVector.Set< FDeferredLightPS::FTransmissionDim >(false);
 	PermutationVector.Set< FDeferredLightPS::FHairLighting>(1);
 	PermutationVector.Set< FDeferredLightPS::FHairComplexTransmittance>(true);
+	PermutationVector.Set< FDeferredLightPS::FLightFunctionAtlasDim >(LightSceneInfo->Proxy->GetLightFunctionMaterial() != nullptr && View.Family->EngineShowFlags.LightFunctions);
 	if (bIsDirectional)
 	{
 		PermutationVector.Set< FDeferredLightPS::FSourceShapeDim >(ELightSourceShape::Directional);
@@ -2947,6 +2953,7 @@ static void InternalRenderSimpleLightsStandardDeferred(
 	FDeferredLightPS::FPermutationDomain PermutationVector;
 	PermutationVector.Set< FDeferredLightPS::FSourceShapeDim >(ELightSourceShape::Capsule);
 	PermutationVector.Set< FDeferredLightPS::FIESProfileDim >(false);
+	PermutationVector.Set< FDeferredLightPS::FLightFunctionAtlasDim>(false);
 	PermutationVector.Set< FDeferredLightPS::FVisualizeCullingDim >(View.Family->EngineShowFlags.VisualizeLightCulling);
 	PermutationVector.Set< FDeferredLightPS::FLightingChannelsDim >(false);
 	PermutationVector.Set< FDeferredLightPS::FAnistropicMaterials >(false);
