Fixed RDG validation with indirect args being accessed in a debug shader as a structured buffer while indirect args are actually plain buffers
#tests editor

diff --git a/Engine/Shaders/Private/DistanceField/GlobalDistanceFieldDebug.usf b/Engine/Shaders/Private/DistanceField/GlobalDistanceFieldDebug.usf
--- a/Engine/Shaders/Private/DistanceField/GlobalDistanceFieldDebug.usf
+++ b/Engine/Shaders/Private/DistanceField/GlobalDistanceFieldDebug.usf
@@ -21,7 +21,7 @@ Texture3D<uint> PageTableCombinedTexture;
 Texture3D PageAtlasTexture;
 SamplerState DistanceFieldSampler;
 
-StructuredBuffer<uint> PageComposeIndirectArgBuffer;
+Buffer<uint> PageComposeIndirectArgBuffer;
 
 [numthreads(1, 1, 1)]
 void GlobalDistanceFieldAccumulateUpdatedPagesCS(
diff --git a/Engine/Source/Runtime/Renderer/Private/GlobalDistanceField.cpp b/Engine/Source/Runtime/Renderer/Private/GlobalDistanceField.cpp
--- a/Engine/Source/Runtime/Renderer/Private/GlobalDistanceField.cpp
+++ b/Engine/Source/Runtime/Renderer/Private/GlobalDistanceField.cpp
@@ -1890,7 +1890,7 @@ class FGlobalDistanceFieldAccumulateUpdatedPagesCS : public FGlobalShader
 
 	BEGIN_SHADER_PARAMETER_STRUCT(FParameters, )
 		SHADER_PARAMETER_RDG_BUFFER_UAV(RWStructuredBuffer<uint>, RWPageStatsBuffer)
-		SHADER_PARAMETER_RDG_BUFFER_SRV(StructuredBuffer<uint>, PageComposeIndirectArgBuffer)
+		SHADER_PARAMETER_RDG_BUFFER_SRV(Buffer<uint>, PageComposeIndirectArgBuffer)
 	END_SHADER_PARAMETER_STRUCT()
 
 	static bool ShouldCompilePermutation(const FGlobalShaderPermutationParameters& Parameters)
@@ -3305,4 +3305,4 @@ void UpdateGlobalDistanceFieldVolume(
 		View.ViewState->GlobalDistanceFieldData->bFirstFrame = false;
 		View.ViewState->GlobalDistanceFieldData->UpdateFrame = View.Family->FrameNumber;
 	}
-}
\ No newline at end of file
+}
