Fixed compilation issue on android due to too many srvs on mobile.

diff --git a/Engine/Shaders/Private/VolumetricFog.usf b/Engine/Shaders/Private/VolumetricFog.usf
--- a/Engine/Shaders/Private/VolumetricFog.usf
+++ b/Engine/Shaders/Private/VolumetricFog.usf
@@ -326,9 +326,11 @@ int VirtualShadowMapId;
 #define HISTORY_MISS_SUPER_SAMPLE_COUNT 1
 #endif
 
+#if USE_LIGHT_FUNCTION_ATLAS || USE_LIGHT_FUNCTION
 #include "LightFunctionAtlas/LightFunctionAtlasCommon.usf"
+#endif
 
-#if	USE_LIGHT_FUNCTION
+#if	USE_LIGHT_FUNCTION // Used when injecting light with with shadow, unbatched
 
 /** Matrix can be either WorldToShadow or WorldToLight depending on light type. */
 float4x4 LocalLightFunctionMatrix;
@@ -670,7 +672,7 @@ uint DirectionalLightFunctionAtlasLightIndex;
 
 float GetDirectionalLightFunction(float3 TranslatedWorldPosition) // directional light
 {
-#if USE_LIGHT_FUNCTION_ATLAS
+#if USE_LIGHT_FUNCTION_ATLAS  // Used when injecting lights from the light grid
 	if (DirectionalApplyLightFunctionFromAtlas > 0)
 	{
 		return GetLocalLightFunctionCommon(TranslatedWorldPosition, DirectionalLightFunctionAtlasLightIndex);
diff --git a/Engine/Source/Runtime/Renderer/Private/VolumetricFog.cpp b/Engine/Source/Runtime/Renderer/Private/VolumetricFog.cpp
--- a/Engine/Source/Runtime/Renderer/Private/VolumetricFog.cpp
+++ b/Engine/Source/Runtime/Renderer/Private/VolumetricFog.cpp
@@ -1024,6 +1024,7 @@ class FVolumetricFogLightScatteringCS : public FGlobalShader
 			PermutationVector.Set<FDistanceFieldSkyOcclusion>(false);
 			PermutationVector.Set<FCloudTransmittance>(false);
 			PermutationVector.Set<FTemporalReprojection>(false);
+			PermutationVector.Set<FSampleLightFunctionAtlas>(false);
 		}
 
 		if (!FDataDrivenShaderPlatformInfo::GetSupportsLumenGI(ShaderPlatform))
@@ -1305,7 +1306,6 @@ void FSceneRenderer::ComputeVolumetricFog(FRDGBuilder& GraphBuilder,
 	}
 
 	const FExponentialHeightFogSceneInfo& FogInfo = Scene->ExponentialFogs[0];
-	const bool bUseLightFunctionAtlas = Scene->LightFunctionAtlasSceneData.GetVolumetricFogUsesLightFunctionAtlas();
 
 	TRACE_CPUPROFILER_EVENT_SCOPE(FSceneRenderer::ComputeVolumetricFog);
 	QUICK_SCOPE_CYCLE_COUNTER(STAT_VolumetricFog);
@@ -1335,6 +1335,9 @@ void FSceneRenderer::ComputeVolumetricFog(FRDGBuilder& GraphBuilder,
 			IntegrationData.FrameJitterOffsetValues[FrameOffsetIndex] = VolumetricFogTemporalRandom(View.Family->FrameNumber - FrameOffsetIndex);
 		}
 
+		// Mobile has limited capacities with SRV binding so do not enable atlas sampling on there.
+		const bool bUseLightFunctionAtlas = Scene->LightFunctionAtlasSceneData.GetVolumetricFogUsesLightFunctionAtlas() && !IsMobilePlatform(View.GetShaderPlatform());
+
 		const bool bUseTemporalReprojection =
 			GVolumetricFogTemporalReprojection
 			&& View.ViewState
