Introducing Scene Extensions, which provide a means of adding optional, persistent data for features that live alongside the scene without coupling the scene and feature code. Extending this interface also provides a means of hooking the various stages of scene updates, as well as hooks into scene rendering without requiring any change to the core infrastructure.

As a proof of concept, spline mesh scene resources have been adapted to use this new interface.

Also, as a bit of code clean-up, RVT rendering has been updated to use the renderer module's scene primitive rendering interface, and this interface updated to optionally not be tied to a single view family to accommodate RVT. This should help reduce the risk of breakage of RVT due to unifying the interface it employs with that of other such sources of rendering required that lie outside of standard scene rendering.

diff --git a/Engine/Source/Runtime/RenderCore/Public/RendererInterface.h b/Engine/Source/Runtime/RenderCore/Public/RendererInterface.h
--- a/Engine/Source/Runtime/RenderCore/Public/RendererInterface.h
+++ b/Engine/Source/Runtime/RenderCore/Public/RendererInterface.h
@@ -40,6 +40,7 @@ class FViewInfo;
 template<typename ShaderType, typename PointerTableType> class TShaderRefBase;
 class FSceneUniformBuffer;
 class FBatchedPrimitiveParameters;
+class ISceneRenderer;
 
 namespace Nanite
 {
@@ -650,9 +651,8 @@ public:
 class IScenePrimitiveRenderingContext
 {
 public:
-	virtual ~IScenePrimitiveRenderingContext()
-	{
-	}
+	virtual ~IScenePrimitiveRenderingContext() {}
+	virtual ISceneRenderer* GetSceneRenderer() = 0;
 };
 
 struct FScenePrimitiveRenderingContextScopeHelper
@@ -852,6 +852,7 @@ public:
 	 * The intended use is for stand-alone rendering that involves Scene proxies (that then may need the machinery to render GPU-Scene aware primitives.
 	 */
 	virtual IScenePrimitiveRenderingContext* BeginScenePrimitiveRendering(FRDGBuilder& GraphBuilder, FSceneViewFamily* ViewFamily) = 0;
+	virtual IScenePrimitiveRenderingContext* BeginScenePrimitiveRendering(FRDGBuilder& GraphBuilder, FSceneInterface& Scene) = 0;
 
 	/** Mark all the current scenes as needing to restart path tracer accumulation */
 	virtual void InvalidatePathTracedOutput() = 0;
diff --git a/Engine/Source/Runtime/Renderer/Private/DeferredShadingRenderer.cpp b/Engine/Source/Runtime/Renderer/Private/DeferredShadingRenderer.cpp
--- a/Engine/Source/Runtime/Renderer/Private/DeferredShadingRenderer.cpp
+++ b/Engine/Source/Runtime/Renderer/Private/DeferredShadingRenderer.cpp
@@ -96,7 +96,6 @@
 #include "Lumen/LumenHardwareRayTracingCommon.h"
 #include "SparseVolumeTexture/ISparseVolumeTextureStreamingManager.h"
 #include "WaterInfoTextureRendering.h"
-#include "SplineMeshSceneResources.h"
 #include "PostProcess/DebugAlphaChannel.h"
 #include "StochasticShadows/StochasticShadows.h"
 #include "Rendering/CustomRenderPass.h"
@@ -3065,6 +3064,7 @@ void FDeferredShadingSceneRenderer::Render(FRDGBuilder& GraphBuilder)
 	if (RendererOutput == ERendererOutput::FinalSceneColor)
 	{
 		// Notify the FX system that the scene is about to be rendered.
+		// TODO: These should probably be moved to scene extensions
 		if (FXSystem && Views.IsValidIndex(0))
 		{
 			SCOPE_CYCLE_COUNTER(STAT_FDeferredShadingSceneRenderer_FXSystem_PreRender);
@@ -3105,6 +3105,9 @@ void FDeferredShadingSceneRenderer::Render(FRDGBuilder& GraphBuilder)
 		}
 	}
 
+	// Allow scene extensions to affect the scene uniform buffer after GPU scene has fully updated
+	GetSceneExtensionsRenderer().UpdateSceneUniformBuffer(GraphBuilder, GetSceneUniforms());
+
 	const bool bUseGBuffer = IsUsingGBuffers(ShaderPlatform);
 	const bool bShouldRenderVolumetricFog = ShouldRenderVolumetricFog();
 	const bool bShouldRenderLocalFogVolume = ShouldRenderLocalFogVolume(Scene, ViewFamily);
@@ -3181,13 +3184,9 @@ void FDeferredShadingSceneRenderer::Render(FRDGBuilder& GraphBuilder)
 	// Sanity check - Note: Nanite forces a Z prepass in ShouldForceFullDepthPass()
 	check(!UseNanite(ShaderPlatform) || bNeedsPrePass);
 
+	GetSceneExtensionsRenderer().PreRender(GraphBuilder);
 	GEngine->GetPreRenderDelegateEx().Broadcast(GraphBuilder);
 
-	if (Scene->SplineMeshSceneResources)
-	{
-		Scene->SplineMeshSceneResources->Update(GraphBuilder, GetSceneUniforms());
-	}
-
 	if (DepthPass.IsComputeStencilDitherEnabled())
 	{
 		AddDitheredStencilFillPass(GraphBuilder, Views, SceneTextures.Depth.Target, DepthPass);
@@ -4604,6 +4603,7 @@ void FDeferredShadingSceneRenderer::Render(FRDGBuilder& GraphBuilder)
 		}
 
 		GEngine->GetPostRenderDelegateEx().Broadcast(GraphBuilder);
+		GetSceneExtensionsRenderer().PostRender(GraphBuilder);
 
 #if RHI_RAYTRACING
 		ReleaseRaytracingResources(GraphBuilder, Views, Scene->RayTracingScene, bIsLastSceneRenderer);
diff --git a/Engine/Source/Runtime/Renderer/Private/MobileShadingRenderer.cpp b/Engine/Source/Runtime/Renderer/Private/MobileShadingRenderer.cpp
--- a/Engine/Source/Runtime/Renderer/Private/MobileShadingRenderer.cpp
+++ b/Engine/Source/Runtime/Renderer/Private/MobileShadingRenderer.cpp
@@ -1008,8 +1008,11 @@ void FMobileSceneRenderer::Render(FRDGBuilder& GraphBuilder)
 		GraphBuilder.SetFlushResourcesRHI();
 	}
 
-	GEngine->GetPreRenderDelegateEx().Broadcast(GraphBuilder);
+	// Allow scene extensions to affect the scene uniform buffer
+	GetSceneExtensionsRenderer().UpdateSceneUniformBuffer(GraphBuilder, GetSceneUniforms());
 
+	GetSceneExtensionsRenderer().PreRender(GraphBuilder);
+	GEngine->GetPreRenderDelegateEx().Broadcast(GraphBuilder);
 	
 	GraphBuilder.SetCommandListStat(GET_STATID(STAT_CLMM_SceneSim));
 
@@ -1060,6 +1063,7 @@ void FMobileSceneRenderer::Render(FRDGBuilder& GraphBuilder)
 		}
 
 		// Notify the FX system that the scene is about to be rendered.
+		// TODO: These should probably be moved to scene extensions
 		if (FXSystem)
 		{
 			FXSystem->PreRender(GraphBuilder, GetSceneViews(), GetSceneUniforms(), true /*bAllowGPUParticleUpdate*/);
@@ -1327,6 +1331,7 @@ void FMobileSceneRenderer::Render(FRDGBuilder& GraphBuilder)
 	}
 
 	GEngine->GetPostRenderDelegateEx().Broadcast(GraphBuilder);
+	GetSceneExtensionsRenderer().PostRender(GraphBuilder);
 
 	GraphBuilder.SetCommandListStat(GET_STATID(STAT_CLMM_SceneEnd));
 
diff --git a/Engine/Source/Runtime/Renderer/Private/RendererModule.h b/Engine/Source/Runtime/Renderer/Private/RendererModule.h
--- a/Engine/Source/Runtime/Renderer/Private/RendererModule.h
+++ b/Engine/Source/Runtime/Renderer/Private/RendererModule.h
@@ -136,6 +136,7 @@ public:
 	void RenderPostResolvedSceneColorExtension(FRDGBuilder& GraphBuilder, const FSceneTextures& SceneTextures);
 
 	virtual IScenePrimitiveRenderingContext* BeginScenePrimitiveRendering(FRDGBuilder& GraphBuilder, FSceneViewFamily* ViewFamily) override;
+	virtual IScenePrimitiveRenderingContext* BeginScenePrimitiveRendering(FRDGBuilder& GraphBuilder, FSceneInterface& Scene) override;
 
 	virtual void InvalidatePathTracedOutput() override;
 
diff --git a/Engine/Source/Runtime/Renderer/Private/RendererScene.cpp b/Engine/Source/Runtime/Renderer/Private/RendererScene.cpp
--- a/Engine/Source/Runtime/Renderer/Private/RendererScene.cpp
+++ b/Engine/Source/Runtime/Renderer/Private/RendererScene.cpp
@@ -98,8 +98,6 @@
 #include "Rendering/StaticLightingSystemInterface.h"
 #endif
 
-#include "SplineMeshSceneResources.h"
-
 #define VALIDATE_PRIMITIVE_PACKED_INDEX 0
 
 /** Affects BasePassPixelShader.usf so must relaunch editor to recompile shaders. */
@@ -1638,7 +1636,6 @@ FScene::FScene(UWorld* InWorld, bool bInRequiresHitProxies, bool bInIsEditorScen
 ,	RayTracingDynamicGeometryCollection(nullptr)
 ,	RayTracingSkinnedGeometryUpdateQueue(nullptr)
 #endif
-,	SplineMeshSceneResources(nullptr)
 ,	NumVisibleLights_GameThread(0)
 ,	NumEnabledSkylights_GameThread(0)
 ,	SceneFrameNumber(0)
@@ -1718,10 +1715,8 @@ FScene::FScene(UWorld* InWorld, bool bInRequiresHitProxies, bool bInIsEditorScen
 	// Allocate the shadow scene, it is always present but we use a pointer such that it can be forward declared.
 	ShadowScene = new FShadowScene(*this);
 
-	if (UseSplineMeshSceneResources(GetFeatureLevelShaderPlatform(InFeatureLevel)))
-	{
-		SplineMeshSceneResources = new FSplineMeshSceneResources(*this);
-	}
+	// Make sure we initialize the SceneRenderExtensions last, when the rest of the scene is initialized
+	SceneExtensions.Init(*this);
 }
 
 FScene::~FScene()
@@ -1800,8 +1795,6 @@ FScene::~FScene()
 	checkf(RemovedPrimitiveSceneInfos.Num() == 0, TEXT("Leaking %d FPrimitiveSceneInfo instances."), RemovedPrimitiveSceneInfos.Num()); // Ensure UpdateAllPrimitiveSceneInfos() is called before destruction.
 
 	delete SceneLightInfoUpdates;
-
-	delete SplineMeshSceneResources;
 }
 
 // Helpers for internal templates
@@ -5493,20 +5486,19 @@ void FScene::UpdateAllPrimitiveSceneInfos(FRDGBuilder& GraphBuilder, EUpdateAllP
 	RemovedPrimitiveSceneInfos.Empty();
 
 	RemovedLocalPrimitiveSceneInfos.Sort(FPrimitiveArraySortKey());
-	auto &SceneCullingUpdater = SceneCulling->BeginUpdate(GraphBuilder);
+	GPUScene.OnPreSceneUpdate(GraphBuilder, SceneUpdateChangeSetStorage.GetPreUpdateSet());
 
+	// TODO: Move this to a scene extension?
+	auto &SceneCullingUpdater = SceneCulling->BeginUpdate(GraphBuilder);
 	SceneCullingUpdater.OnPreSceneUpdate(GraphBuilder, SceneUpdateChangeSetStorage.GetPreUpdateSet());
 	
-	if (SplineMeshSceneResources)
-	{
-		SplineMeshSceneResources->PreSceneUpdate(GraphBuilder, SceneUpdateChangeSetStorage.GetPreUpdateSet());
-	}
-	GPUScene.OnPreSceneUpdate(GraphBuilder, SceneUpdateChangeSetStorage.GetPreUpdateSet());
-
 	// Create a SceneUB that permits access to the scene for invalidation processing.
 	FSceneUniformBuffer SceneUB;
 	GPUScene.FillSceneUniformBuffer(GraphBuilder, SceneUB);
 
+	auto& SceneExtensionsUpdater = *GraphBuilder.AllocObject<FSceneExtensionsUpdater>(*this);
+	SceneExtensionsUpdater.PreSceneUpdate(GraphBuilder, SceneUpdateChangeSetStorage.GetPreUpdateSet());
+
 	{
 		SCOPED_NAMED_EVENT(FScene_VirtualShadowCacheUpdate, FColor::Orange);
 		FVirtualShadowMapArrayCacheManager* CacheManager = GetVirtualShadowMapCache();
@@ -6310,15 +6302,13 @@ void FScene::UpdateAllPrimitiveSceneInfos(FRDGBuilder& GraphBuilder, EUpdateAllP
 
 	SceneCullingUpdater.OnPostSceneUpdate(GraphBuilder, SceneUpdateChangeSetStorage.GetPostUpdateSet());
 	
-	if (SplineMeshSceneResources)
-	{
-		SplineMeshSceneResources->PostSceneUpdate(GraphBuilder, SceneUpdateChangeSetStorage.GetPostUpdateSet());
-	}
-	GPUScene.OnPostSceneUpdate(GraphBuilder, SceneUpdateChangeSetStorage.GetPostUpdateSet());
+	GPUScene.OnPostSceneUpdate(GraphBuilder, SceneUpdateChangeSetStorage.GetPostUpdateSet());	
 
 	UpdateCachedShadowState(SceneUpdateChangeSetStorage.GetPreUpdateSet(), SceneUpdateChangeSetStorage.GetPostUpdateSet());
 	ShadowScene->PostSceneUpdate(SceneUpdateChangeSetStorage.GetPreUpdateSet(), SceneUpdateChangeSetStorage.GetPostUpdateSet());
 
+	SceneExtensionsUpdater.PostSceneUpdate(GraphBuilder, SceneUpdateChangeSetStorage.GetPostUpdateSet());
+
 	UpdateUniformExpressionsTask.Wait();
 
 	if (SceneInfosWithStaticDrawListUpdate.Num() > 0)
@@ -6520,10 +6510,12 @@ void FScene::UpdateAllPrimitiveSceneInfos(FRDGBuilder& GraphBuilder, EUpdateAllP
 
 		// Process GPU scene prior to visibility to maximize overlap.
 		GPUScene.Update(GraphBuilder, SceneUB, ExternalAccessQueue, nullptr);
-
+	
 		ExternalAccessQueue.Submit(GraphBuilder);
 	}
 
+	SceneExtensionsUpdater.PostGPUSceneUpdate(GraphBuilder, SceneUB);
+
 	// Need to do this here since we delete the proxy next 
 	// TODO: should refactor to add this as a dependency for a proxy-deletion task instead.
 	SceneCullingUpdater.GetAsyncProxyUseTaskHandle().Wait();
diff --git a/Engine/Source/Runtime/Renderer/Private/SceneExtensions.cpp b/Engine/Source/Runtime/Renderer/Private/SceneExtensions.cpp
new file mode 100644
--- /dev/null
+++ b/Engine/Source/Runtime/Renderer/Private/SceneExtensions.cpp
@@ -0,0 +1,132 @@
+// Copyright Epic Games, Inc. All Rights Reserved.
+
+#include "SceneExtensions.h"
+#include "ScenePrivate.h"
+#include "SceneRendering.h"
+#include "RenderGraphBuilder.h"
+#include "Misc/ScopeLock.h"
+
+FSceneExtensionRegistry* FSceneExtensionRegistry::GlobalRegistry = nullptr;
+
+void FSceneExtensionRegistry::InitRegistry()
+{
+	if (!GlobalRegistry)
+	{
+		GlobalRegistry = new FSceneExtensionRegistry;
+	}
+}
+
+void FSceneExtensionRegistry::Register(ISceneExtensionFactory& Factory)
+{
+	Factory.ExtensionID = Factories.Num();
+	Factories.Add(&Factory);
+}
+
+TArray<ISceneExtension*> FSceneExtensionRegistry::CreateExtensions(FScene& Scene)
+{
+	TArray<ISceneExtension*> Extensions;
+	Extensions.Reserve(Factories.Num());
+
+	int32 ExtensionID = 0;
+	for (auto* Factory : Factories)
+	{
+		checkSlow(Factory->GetExtensionID() == ExtensionID); // sanity check
+		if (auto Extension = Factory->CreateInstance(Scene))
+		{
+			Extensions.Add(Extension);
+		}
+		++ExtensionID;
+	}
+
+	return Extensions;
+}
+
+void FSceneExtensions::Init(FScene& Scene)
+{
+	Extensions = FSceneExtensionRegistry::Get().CreateExtensions(Scene);
+	for (auto Extension : Extensions)
+	{
+		Extension->InitExtension(Scene);
+	}
+}
+
+void FSceneExtensions::Reset()
+{
+	for (auto Extension : Extensions)
+	{
+		delete Extension;
+	}
+	Extensions.Reset();
+}
+
+void FSceneExtensions::CreateUpdaters(FUpdaterList& OutUpdaters)
+{
+	OutUpdaters.Reserve(OutUpdaters.Num() + Extensions.Num());
+	for (auto Extension : Extensions)
+	{
+		if (auto Updater = Extension->CreateUpdater())
+		{
+			OutUpdaters.Add(Updater);
+		}
+	}	
+}
+
+void FSceneExtensions::CreateRenderers(FRendererList& OutRenderers)
+{
+	OutRenderers.Reserve(OutRenderers.Num() + Extensions.Num());
+	for (auto Extension : Extensions)
+	{
+		if (auto Renderer = Extension->CreateRenderer())
+		{
+			OutRenderers.Add(Renderer);
+		}
+	}
+}
+
+void FSceneExtensionsUpdater::Begin(FScene& InScene)
+{
+	checkf(!IsUpdating(), TEXT("Detected FSceneExtensionsUpdater Begin() without matching End()"));
+	
+	Scene = &InScene;
+	Scene->SceneExtensions.CreateUpdaters(Updaters);
+	for (auto Updater : Updaters)
+	{
+		Updater->Begin(InScene);
+	}
+}
+
+void FSceneExtensionsUpdater::End()
+{
+	for (auto Updater : Updaters)
+	{
+		Updater->End();
+		delete Updater;
+	}
+	Updaters.Reset();
+	Scene = nullptr;
+}
+
+void FSceneExtensionsRenderer::Begin(FSceneRendererBase& InSceneRenderer)
+{
+	checkf(!IsRendering(), TEXT("Detected FSceneExtensionsRenderer Begin() without matching End()"));
+
+	SceneRenderer = &InSceneRenderer;
+	check(SceneRenderer->Scene != nullptr);
+
+	SceneRenderer->Scene->SceneExtensions.CreateRenderers(Renderers);
+	for (auto Renderer : Renderers)
+	{
+		Renderer->Begin(InSceneRenderer);
+	}
+}
+
+void FSceneExtensionsRenderer::End()
+{
+	for (auto Renderer : Renderers)
+	{
+		Renderer->End();
+		delete Renderer;
+	}
+	Renderers.Reset();
+	SceneRenderer = nullptr;
+}
diff --git a/Engine/Source/Runtime/Renderer/Private/SceneExtensions.h b/Engine/Source/Runtime/Renderer/Private/SceneExtensions.h
new file mode 100644
--- /dev/null
+++ b/Engine/Source/Runtime/Renderer/Private/SceneExtensions.h
@@ -0,0 +1,252 @@
+// Copyright Epic Games, Inc. All Rights Reserved.
+
+#pragma once
+
+#include "CoreMinimal.h"
+#include "RendererInterface.h"
+#include "Tasks/Task.h"
+
+class FScene;
+class FSceneRendererBase;
+class FRDGBuilder;
+class FSceneUniformBuffer;
+class FScenePreUpdateChangeSet;
+class FScenePostUpdateChangeSet;
+class FSceneExtensionRegistry;
+class ISceneExtensionUpdater;
+class ISceneExtensionRenderer;
+
+/** Abstract interface for an extension to the persistent data of a scene */
+class ISceneExtension
+{
+public:
+	// default fallback static method that can be overridden in child classes to predicate the creation of the extension
+	static bool ShouldCreateExtension(FScene& Scene) { return true; }
+
+	virtual ~ISceneExtension() {}
+
+	virtual void InitExtension(FScene& InScene) = 0;
+	virtual ISceneExtensionUpdater* CreateUpdater() { return nullptr; }
+	virtual ISceneExtensionRenderer* CreateRenderer() { return nullptr; }
+};
+
+/** Abstract interface to receive change sets to perform updates based on scene primitive data. */
+class ISceneExtensionUpdater
+{
+public:
+	virtual ~ISceneExtensionUpdater() {}
+
+	virtual void Begin(FScene& InScene) {}
+	virtual void End() {}
+	virtual void PreSceneUpdate(FRDGBuilder& GraphBuilder, const FScenePreUpdateChangeSet& ChangeSet) {}
+	virtual void PostSceneUpdate(FRDGBuilder& GraphBuilder, const FScenePostUpdateChangeSet& ChangeSet) {}
+	virtual void PostGPUSceneUpdate(FRDGBuilder& GraphBuilder, FSceneUniformBuffer& SceneUniforms) {}
+};
+
+/** Abstract interface for an extension to the scene renderer */
+class ISceneExtensionRenderer
+{
+public:
+	virtual ~ISceneExtensionRenderer() {}
+
+	virtual void Begin(FSceneRendererBase& InRenderer) {}
+	virtual void End() {}
+	virtual void UpdateSceneUniformBuffer(FRDGBuilder& GraphBuilder, FSceneUniformBuffer& SceneUniforms) {}
+	virtual void PreRender(FRDGBuilder& GraphBuilder) {}
+	virtual void PostRender(FRDGBuilder& GraphBuilder) {}
+};
+
+/** Abstract interface for creating an instance of a scene extension */
+class ISceneExtensionFactory
+{
+	friend class FSceneExtensionRegistry;
+
+public:
+	virtual ~ISceneExtensionFactory() {}
+	virtual ISceneExtension* CreateInstance(FScene& Scene) = 0;
+
+	const int32 GetExtensionID() const { return ExtensionID; }
+
+private:
+	int32 ExtensionID = INDEX_NONE;
+};
+
+/** Static class used to store a global registry of extension types */
+class FSceneExtensionRegistry
+{
+public:
+	static FSceneExtensionRegistry& Get()
+	{
+		InitRegistry();
+		return *GlobalRegistry;
+	}
+
+	void Register(ISceneExtensionFactory& Factory);
+	TArray<ISceneExtension*> CreateExtensions(FScene& Scene);
+
+private:
+	static void InitRegistry();
+
+	TArray<ISceneExtensionFactory*> Factories;
+
+	static FSceneExtensionRegistry* GlobalRegistry;
+};
+
+/** A collection of scene extensions */
+class FSceneExtensions
+{	
+public:
+	using FUpdaterList = TArray<ISceneExtensionUpdater*, FSceneRenderingArrayAllocator>;
+	using FRendererList = TArray<ISceneExtensionRenderer*, FSceneRenderingArrayAllocator>;
+
+	~FSceneExtensions() { Reset(); }
+
+	void Init(FScene& Scene);
+	void Reset();
+	void CreateUpdaters(FUpdaterList& OutUpdaters);
+	void CreateRenderers(FRendererList& OutRenderers);
+
+	template<typename TDerivedExtension>
+	TDerivedExtension* GetExtension()
+	{
+		const int32 Index = TDerivedExtension::GetExtensionID();
+		if (ensure(Extensions.IsValidIndex(Index)))
+		{
+			return static_cast<TDerivedExtension*>(Extensions[Index]);
+		}
+		return nullptr;
+	}
+
+	template<typename TDerivedExtension>
+	const TDerivedExtension* GetExtension() const
+	{
+		return const_cast<FSceneExtensions*>(this)->GetExtension<TDerivedExtension>();
+	}
+	
+	template<typename TDerivedExtension>
+	TDerivedExtension& GetExtensionChecked()
+	{
+		TDerivedExtension* Extension = this->GetExtension<TDerivedExtension>();
+		check(Extension != nullptr);
+		return *Extension;
+	}
+	
+	template<typename TDerivedExtension>
+	const TDerivedExtension& GetExtensionChecked() const
+	{
+		return const_cast<FSceneExtensions*>(this)->GetExtensionChecked<TDerivedExtension>();
+	}
+
+	template<typename TFunc>
+	void ForEachExtension(const TFunc& F)
+	{
+		for(auto* Ext : Extensions)
+		{
+			F(Ext);
+		}
+	}
+
+private:
+	TArray<ISceneExtension*> Extensions;
+};
+
+/** Performs updates for the given scene extensions */
+class FSceneExtensionsUpdater
+{
+	friend class FSceneExtensions;
+
+public:
+	FSceneExtensionsUpdater() {}
+	explicit FSceneExtensionsUpdater(FScene& InScene) { Begin(InScene); }
+	~FSceneExtensionsUpdater() { End(); }
+
+	void Begin(FScene& InScene);
+	void End();
+	bool IsUpdating() const { return Scene != nullptr; }
+
+	void PreSceneUpdate(FRDGBuilder& GraphBuilder, const FScenePreUpdateChangeSet& ChangeSet)
+	{
+		for (auto Updater : Updaters) { Updater->PreSceneUpdate(GraphBuilder, ChangeSet); }
+	}
+
+	void PostSceneUpdate(FRDGBuilder& GraphBuilder, const FScenePostUpdateChangeSet& ChangeSet)
+	{
+		for (auto Updater : Updaters) { Updater->PostSceneUpdate(GraphBuilder, ChangeSet); }
+	}
+
+	void PostGPUSceneUpdate(FRDGBuilder& GraphBuilder, FSceneUniformBuffer& SceneUniforms)
+	{
+		for (auto Updater : Updaters) { Updater->PostGPUSceneUpdate(GraphBuilder, SceneUniforms); }
+	}
+
+private:
+	FScene* Scene = nullptr;
+	FSceneExtensions::FUpdaterList Updaters;
+};
+
+/** Performs rendering for the given scene extensions */
+class FSceneExtensionsRenderer
+{
+	friend class FSceneExtensions;
+
+public:
+	FSceneExtensionsRenderer() {}
+	FSceneExtensionsRenderer(FSceneRendererBase& InSceneRenderer) { Begin(InSceneRenderer); }
+	~FSceneExtensionsRenderer() { End(); }
+	
+	void Begin(FSceneRendererBase& InSceneRenderer);
+	void End();
+	bool IsRendering() const { return SceneRenderer != nullptr; }
+
+	void UpdateSceneUniformBuffer(FRDGBuilder& GraphBuilder, FSceneUniformBuffer& SceneUniforms)
+	{
+		for (auto Renderer : Renderers) { Renderer->UpdateSceneUniformBuffer(GraphBuilder, SceneUniforms); }
+	}
+	
+	void PreRender(FRDGBuilder& GraphBuilder)
+	{
+		for (auto Renderer : Renderers) { Renderer->PreRender(GraphBuilder); }
+	}
+	
+	void PostRender(FRDGBuilder& GraphBuilder)
+	{
+		for (auto Renderer : Renderers) { Renderer->PostRender(GraphBuilder); }
+	}
+
+private:
+	FSceneRendererBase* SceneRenderer = nullptr;
+	FSceneExtensions::FRendererList Renderers;
+};
+
+/** Helper to automatically register/unregister a factory implementation for a given ISceneExtension implementation */
+template<typename TDerivedExtension>
+class TSceneExtensionRegistration : public ISceneExtensionFactory
+{
+public:
+	TSceneExtensionRegistration()
+	{
+		FSceneExtensionRegistry::Get().Register(*this);
+	}
+
+	virtual ~TSceneExtensionRegistration() {}
+
+	virtual ISceneExtension* CreateInstance(FScene& Scene) override
+	{
+		if (!TDerivedExtension::ShouldCreateExtension(Scene))
+		{
+			return nullptr;
+		}
+		return new TDerivedExtension();
+	}
+};
+
+/** Use this macros in the class definition of your extension. */
+#define DECLARE_SCENE_EXTENSION(ClassName) \
+	public: \
+		static int32 GetExtensionID() { return ExtensionRegistration.GetExtensionID(); } \
+	private: \
+		static TSceneExtensionRegistration<ClassName> ExtensionRegistration
+
+/** Use this macros in the implementation source file of your extension. */
+#define IMPLEMENT_SCENE_EXTENSION(ClassName) \
+	TSceneExtensionRegistration<ClassName> ClassName::ExtensionRegistration
diff --git a/Engine/Source/Runtime/Renderer/Private/SceneHitProxyRendering.cpp b/Engine/Source/Runtime/Renderer/Private/SceneHitProxyRendering.cpp
--- a/Engine/Source/Runtime/Renderer/Private/SceneHitProxyRendering.cpp
+++ b/Engine/Source/Runtime/Renderer/Private/SceneHitProxyRendering.cpp
@@ -35,7 +35,6 @@
 #include "HairStrands/HairStrandsData.h"
 #include "SimpleMeshDrawCommandPass.h"
 #include "StaticMeshSceneProxy.h"
-#include "SplineMeshSceneResources.h"
 
 class FHitProxyShaderElementData : public FMeshMaterialShaderElementData
 {
@@ -575,7 +574,10 @@ void FMobileSceneRenderer::RenderHitProxies(FRDGBuilder& GraphBuilder)
 
 	// Find the visible primitives.
 	InitViews(GraphBuilder, SceneTexturesConfig, InstanceCullingManager, nullptr, InitViewTaskDatas);
+	
+	GetSceneExtensionsRenderer().UpdateSceneUniformBuffer(GraphBuilder, GetSceneUniforms());
 
+	GetSceneExtensionsRenderer().PreRender(GraphBuilder);
 	GEngine->GetPreRenderDelegateEx().Broadcast(GraphBuilder);
 
 	InstanceCullingManager.FlushRegisteredViews(GraphBuilder);
@@ -584,6 +586,8 @@ void FMobileSceneRenderer::RenderHitProxies(FRDGBuilder& GraphBuilder)
 	::DoRenderHitProxies(GraphBuilder, this, HitProxyTexture, HitProxyDepthTexture, NaniteRasterResults, InstanceCullingManager);
 
 	GEngine->GetPostRenderDelegateEx().Broadcast(GraphBuilder);
+	GetSceneExtensionsRenderer().PostRender(GraphBuilder);
+
 #endif
 
 	OnRenderFinish(GraphBuilder, nullptr);
@@ -641,8 +645,11 @@ void FDeferredShadingSceneRenderer::RenderHitProxies(FRDGBuilder& GraphBuilder)
 		Scene->GPUScene.UploadDynamicPrimitiveShaderDataForView(GraphBuilder, Views[ViewIndex]);
 	}
 
+
 	EndInitViews(GraphBuilder, LumenFrameTemporaries, InstanceCullingManager, InitViewTaskDatas);
 
+	GetSceneExtensionsRenderer().UpdateSceneUniformBuffer(GraphBuilder, GetSceneUniforms());
+
 	ExternalAccessQueue.Submit(GraphBuilder);
 
 	InstanceCullingManager.FlushRegisteredViews(GraphBuilder);
@@ -654,14 +661,11 @@ void FDeferredShadingSceneRenderer::RenderHitProxies(FRDGBuilder& GraphBuilder)
 		Nanite::GStreamingManager.EndAsyncUpdate(GraphBuilder);
 	}
 
+	GetSceneExtensionsRenderer().PreRender(GraphBuilder);
 	GEngine->GetPreRenderDelegateEx().Broadcast(GraphBuilder);
 
-	if (Scene->SplineMeshSceneResources)
-	{
-		Scene->SplineMeshSceneResources->Update(GraphBuilder, GetSceneUniforms());
-	}
-
 	// Notify the FX system that the scene is about to be rendered.
+	// TODO: These should probably be moved to scene extensions
 	if (FXSystem && Views.IsValidIndex(0))
 	{
 		FGPUSortManager* GPUSortManager = FXSystem->GetGPUSortManager();
@@ -732,6 +736,7 @@ void FDeferredShadingSceneRenderer::RenderHitProxies(FRDGBuilder& GraphBuilder)
 	ShaderPrint::EndViews(Views);
 
 	GEngine->GetPostRenderDelegateEx().Broadcast(GraphBuilder);
+	GetSceneExtensionsRenderer().PostRender(GraphBuilder);
 
 	AddDispatchToRHIThreadPass(GraphBuilder);
 #endif
diff --git a/Engine/Source/Runtime/Renderer/Private/ScenePrivate.h b/Engine/Source/Runtime/Renderer/Private/ScenePrivate.h
--- a/Engine/Source/Runtime/Renderer/Private/ScenePrivate.h
+++ b/Engine/Source/Runtime/Renderer/Private/ScenePrivate.h
@@ -56,6 +56,7 @@
 #include "Algo/RemoveIf.h"
 #include "UObject/Package.h"
 #include "LightFunctionAtlas.h"
+#include "SceneExtensions.h"
 
 /** Factor by which to grow occlusion tests **/
 #define OCCLUSION_SLOP (1.0f)
@@ -93,7 +94,6 @@ class FStaticMeshBatch;
 class FShadowScene;
 class FSceneLightInfoUpdates;
 class FSceneCulling;
-class FSplineMeshSceneResources;
 
 /**
  * Describes all light modifications to the scene by recording the light scene IDs.
@@ -3206,8 +3206,8 @@ public:
 	class FRayTracingSkinnedGeometryUpdateQueue* RayTracingSkinnedGeometryUpdateQueue;
 #endif
 
-	/** Spline mesh data shared between all views. */
-	FSplineMeshSceneResources* SplineMeshSceneResources;
+	/** Collection of scene render extensions. */
+	FSceneExtensions SceneExtensions;
 
 	/** Initialization constructor. */
 	FScene(UWorld* InWorld, bool bInRequiresHitProxies,bool bInIsEditorScene, bool bCreateFXSystem, ERHIFeatureLevel::Type InFeatureLevel);
diff --git a/Engine/Source/Runtime/Renderer/Private/SceneRendering.cpp b/Engine/Source/Runtime/Renderer/Private/SceneRendering.cpp
--- a/Engine/Source/Runtime/Renderer/Private/SceneRendering.cpp
+++ b/Engine/Source/Runtime/Renderer/Private/SceneRendering.cpp
@@ -2503,9 +2503,9 @@ TGlobalResource<FGlobalDynamicReadBuffer> FSceneRenderer::DynamicReadBufferForSh
 	FSceneRenderer
 -----------------------------------------------------------------------------*/
 FSceneRenderer::FSceneRenderer(const FSceneViewFamily* InViewFamily, FHitProxyConsumer* HitProxyConsumer)
-:	Scene(CheckPointer(InViewFamily->Scene)->GetRenderScene())
-,	ViewFamily(*CheckPointer(InViewFamily))
-,	VirtualShadowMapArray(*CheckPointer(Scene))
+:	FSceneRendererBase(*CheckPointer(CheckPointer(CheckPointer(InViewFamily)->Scene)->GetRenderScene()))
+,	ViewFamily(*InViewFamily)
+,	VirtualShadowMapArray(*Scene)
 ,	bHasRequestedToggleFreeze(false)
 ,	bUsedPrecomputedVisibility(false)
 ,	bGPUMasksComputed(false)
@@ -2812,6 +2812,9 @@ FSceneRenderer::FSceneRenderer(const FSceneViewFamily* InViewFamily, FHitProxyCo
 
 	bDumpMeshDrawCommandInstancingStats = !!GDumpInstancingStats;
 	GDumpInstancingStats = 0;
+
+	// Initialize scene renderer extensions here, after the rest of the renderer has been initialized
+	InitSceneExtensionsRenderer();
 }
 
 // static
@@ -5118,13 +5121,18 @@ void FRendererModule::RenderPostResolvedSceneColorExtension(FRDGBuilder& GraphBu
 class FScenePrimitiveRenderingContext : public IScenePrimitiveRenderingContext
 {
 public:
-	FScenePrimitiveRenderingContext(FRDGBuilder& GraphBuilder, FScene& Scene, FSceneViewFamily& ViewFamily) :
+	FScenePrimitiveRenderingContext(FRDGBuilder& GraphBuilder, FScene& Scene, FSceneViewFamily* InViewFamily = nullptr) :
 		Renderer(),
 		GPUScene(Scene.GPUScene),
 		GPUSceneDynamicContext(GPUScene),
-		ViewFamily(ViewFamily)
+		ViewFamily(InViewFamily)
 	{
-		ViewFamily.SetSceneRenderer(&Renderer);
+		Renderer.Scene = &Scene;
+		Renderer.InitSceneExtensionsRenderer();
+		if (ViewFamily)
+		{
+			ViewFamily->SetSceneRenderer(&Renderer);
+		}
 
 		Scene.UpdateAllPrimitiveSceneInfos(GraphBuilder);
 		GPUScene.BeginRender(GraphBuilder, GPUSceneDynamicContext);
@@ -5133,13 +5141,21 @@ public:
 	virtual ~FScenePrimitiveRenderingContext()
 	{
 		GPUScene.EndRender();
-		ViewFamily.SetSceneRenderer(nullptr);
+		if (ViewFamily)
+		{
+			ViewFamily->SetSceneRenderer(nullptr);
+		}
+	}
+
+	virtual ISceneRenderer* GetSceneRenderer()
+	{
+		return &Renderer;
 	}
 
 	FSceneRendererBase Renderer;
 	FGPUScene& GPUScene;
 	FGPUSceneDynamicContext GPUSceneDynamicContext;
-	FSceneViewFamily& ViewFamily;
+	FSceneViewFamily* ViewFamily;
 };
 
 
@@ -5150,7 +5166,17 @@ IScenePrimitiveRenderingContext* FRendererModule::BeginScenePrimitiveRendering(F
 	FScene* Scene = ViewFamily->Scene->GetRenderScene();
 	check(Scene);
 
-	FScenePrimitiveRenderingContext* ScenePrimitiveRenderingContext = new FScenePrimitiveRenderingContext(GraphBuilder, *Scene, *ViewFamily);
+	FScenePrimitiveRenderingContext* ScenePrimitiveRenderingContext = new FScenePrimitiveRenderingContext(GraphBuilder, *Scene, ViewFamily);
+
+	return ScenePrimitiveRenderingContext;
+}
+
+IScenePrimitiveRenderingContext* FRendererModule::BeginScenePrimitiveRendering(FRDGBuilder &GraphBuilder, FSceneInterface& InScene)
+{
+	FScene* Scene = InScene.GetRenderScene();
+	check(Scene);
+
+	FScenePrimitiveRenderingContext* ScenePrimitiveRenderingContext = new FScenePrimitiveRenderingContext(GraphBuilder, *Scene);
 
 	return ScenePrimitiveRenderingContext;
 }
diff --git a/Engine/Source/Runtime/Renderer/Private/SceneRendering.h b/Engine/Source/Runtime/Renderer/Private/SceneRendering.h
--- a/Engine/Source/Runtime/Renderer/Private/SceneRendering.h
+++ b/Engine/Source/Runtime/Renderer/Private/SceneRendering.h
@@ -49,6 +49,7 @@
 #include "LocalFogVolumeRendering.h"
 #include "Nanite/NaniteShared.h"
 #include "LightFunctionAtlas.h"
+#include "SceneExtensions.h"
 
 #if RHI_RAYTRACING
 #include "RayTracingInstanceBufferUtil.h"
@@ -1978,6 +1979,12 @@ struct FComputeLightGridOutput
 class FSceneRendererBase : public ISceneRenderer
 {
 public:
+	/** The scene being rendered. */
+	FScene* Scene = nullptr;
+
+	FSceneRendererBase() {}
+	FSceneRendererBase(FScene& InScene) : Scene(&InScene) {}
+
 	// ISceneRenderer interface
 	const FSceneUniformBuffer& GetSceneUniforms() const final override { return SceneUniforms; }
 	FSceneUniformBuffer& GetSceneUniforms() final override { return SceneUniforms; }
@@ -1987,8 +1994,16 @@ public:
 		return GetSceneUniforms().GetBuffer(GraphBuilder);
 	}
 
+	void InitSceneExtensionsRenderer()
+	{
+		SceneExtensionsRenderer.Begin(*this);
+	}
+	FSceneExtensionsRenderer& GetSceneExtensionsRenderer() { return SceneExtensionsRenderer; }
+	const FSceneExtensionsRenderer& GetSceneExtensionsRenderer() const { return SceneExtensionsRenderer; }
+
 private:
 	FSceneUniformBuffer SceneUniforms;
+	FSceneExtensionsRenderer SceneExtensionsRenderer;
 };
 
 /**
@@ -2002,9 +2017,6 @@ public:
 	/** Linear bulk allocator with a lifetime tied to the scene renderer. */
 	FSceneRenderingBulkObjectAllocator Allocator;
 
-	/** The scene being rendered. */
-	FScene* Scene;
-
 	/** The view family being rendered.  This references the Views array. */
 	FViewFamilyInfo ViewFamily;
 
diff --git a/Engine/Source/Runtime/Renderer/Private/SplineMeshSceneResources.cpp b/Engine/Source/Runtime/Renderer/Private/SplineMeshSceneResources.cpp
--- a/Engine/Source/Runtime/Renderer/Private/SplineMeshSceneResources.cpp
+++ b/Engine/Source/Runtime/Renderer/Private/SplineMeshSceneResources.cpp
@@ -124,36 +124,217 @@
 };
 IMPLEMENT_GLOBAL_SHADER(FSplineMeshTextureFillCS, "/Engine/Private/SplineMeshSceneTexture.usf", "FillTexture", SF_Compute);
 
-FSplineMeshSceneResources::FSplineMeshSceneResources(FScene& InScene)
-	: Scene(InScene)
+IMPLEMENT_SCENE_EXTENSION(FSplineMeshSceneExtension);
+
+bool FSplineMeshSceneExtension::ShouldCreateExtension(FScene& Scene)
 {
-	check(UseSplineMeshSceneResources(GetFeatureLevelShaderPlatform(Scene.GetFeatureLevel())));
+	return UseSplineMeshSceneResources(GetFeatureLevelShaderPlatform(Scene.GetFeatureLevel()));
 }
 
-void FSplineMeshSceneResources::PreSceneUpdate(FRDGBuilder& GraphBuilder, const FScenePreUpdateChangeSet& ChangeSet)
+ISceneExtensionUpdater* FSplineMeshSceneExtension::CreateUpdater()
 {
+	return new FSplineMeshSceneUpdater(*this);
+}
+
+ISceneExtensionRenderer* FSplineMeshSceneExtension::CreateRenderer()
+{
+	return new FSplineMeshSceneRenderer(*this);	
+}
+
+FSplineMeshSceneExtension::FPrimitiveSlot& FSplineMeshSceneExtension::Register(const FPrimitiveSceneInfo& PrimitiveSceneInfo)
+{	
+	FPrimitiveSlot& Slot = RegisteredPrimitives.FindOrAdd(&PrimitiveSceneInfo);
+	if (ensureMsgf(Slot.NumSplines == 0, TEXT("This primitive was already registered!")))
+	{
+		// Alloc space for the new splines and ensure they are included in the next update
+		AllocTextureSpace(PrimitiveSceneInfo, GetNumSplines(PrimitiveSceneInfo), Slot);
+	}
+
+	return Slot;
+}
+
+void FSplineMeshSceneExtension::Unregister(const FPrimitiveSceneInfo& PrimitiveSceneInfo)
+{
+	FPrimitiveSlot* Slot = RegisteredPrimitives.Find(&PrimitiveSceneInfo);
+	if (!ensure(Slot != nullptr))
+	{
+		return;
+	}
+
+	SlotAllocator.Free(Slot->FirstSplineIndex, Slot->NumSplines);
+
+	// Clear the instance IDs of the newly freed range
+	for (uint32 i = 0; i < Slot->NumSplines; ++i)
+	{
+		RegisteredInstanceIds[Slot->FirstSplineIndex + i] = INDEX_NONE;
+	}
+	
+	// shorten this look-up if the allocator shrinks
+	RegisteredInstanceIds.SetNumUninitialized(FMath::Min(SlotAllocator.GetMaxSize(), RegisteredInstanceIds.Num()));
+
+	RegisteredPrimitives.Remove(&PrimitiveSceneInfo);
+
+	// Ensure we repopulate the instance ID lookup next update
+	bInstanceLookupDirty = true;
+}
+
+void FSplineMeshSceneExtension::AllocTextureSpace(const FPrimitiveSceneInfo& PrimitiveSceneInfo, uint32 NumSplines, FPrimitiveSlot& OutSlot)
+{
+	check(NumSplines > 0);
+
+	OutSlot.NumSplines = NumSplines;
+	OutSlot.FirstSplineIndex = SlotAllocator.Allocate(NumSplines);
+
+	// If we don't have the space for these instances in the registered ID list, allocate them
+	if (RegisteredInstanceIds.Num() < SlotAllocator.GetMaxSize())
+	{
+		RegisteredInstanceIds.Reserve(SlotAllocator.GetMaxSize());
+		for (uint32 i = RegisteredInstanceIds.Num(); i < OutSlot.FirstSplineIndex + NumSplines; ++i)
+		{
+			RegisteredInstanceIds.Add(INDEX_NONE);
+		}
+	}
+
+	// Store the instance ID for each spline slot allocated (will be used to fill texels with spline data)
+	// NOTE: This is assuming a spline per instance. If that changes, this look-up will need to as well
+	const int32 InstanceSceneDataOffset = PrimitiveSceneInfo.GetInstanceSceneDataOffset();
+	for (uint32 i = 0; i < NumSplines; ++i)
+	{
+		const uint32 SplineIndex = OutSlot.FirstSplineIndex + i;
+		check(RegisteredInstanceIds[SplineIndex] == INDEX_NONE); // sanity check it has been cleared
+		RegisteredInstanceIds[SplineIndex] = InstanceSceneDataOffset + i;
+	}
+
+	bInstanceLookupDirty = true;
+
+	// Give the scene proxy the coordinates allocated so it can place it in its scene data.
+	AssignCoordinates(PrimitiveSceneInfo, OutSlot);
+}
+
+uint32 FSplineMeshSceneExtension::GetNumSplines(const FPrimitiveSceneInfo& SceneInfo)
+{
+	// We only support spline mesh scene proxies currently, and we assume the number of splines is equal to the
+	// number of instance scene data entries. Support could be added later for other scene proxy types that need
+	// baked down splines.
+	check(SceneInfo.Proxy->IsSplineMesh());
+
+	return SceneInfo.GetNumInstanceSceneDataEntries();
+}
+
+void FSplineMeshSceneExtension::AssignCoordinates(const FPrimitiveSceneInfo& SceneInfo, const FPrimitiveSlot& Slot)
+{
+	check(SceneInfo.Proxy->IsSplineMesh()); // sanity check
+
+	if (SceneInfo.Proxy->IsNaniteMesh())
+	{
+		AssignCoordinates(static_cast<FNaniteSplineMeshSceneProxy*>(SceneInfo.Proxy), Slot);
+	}
+	else
+	{
+		AssignCoordinates(static_cast<FSplineMeshSceneProxy*>(SceneInfo.Proxy), Slot);
+	}
+}
+
+template<typename TSplineMeshSceneProxy>
+void FSplineMeshSceneExtension::AssignCoordinates(TSplineMeshSceneProxy* SceneProxy, const FPrimitiveSlot& Slot)
+{
+	for (uint32 i = 0; i < Slot.NumSplines; ++i)
+	{
+		SceneProxy->SetSplineTextureCoord_RenderThread(i, SplineMesh::CalcTextureCoord(Slot.FirstSplineIndex + i));
+	}
+}
+
+void FSplineMeshSceneExtension::DefragTexture()
+{
+	// NOTE: Currently not attempting to reduce motions to a minimal set, we're just ditching our cache
+	// and re-assigning space in the new texture that will be created next update
+	SavedPosTexture = nullptr;
+	SavedRotTexture = nullptr;
+	SlotAllocator.Reset();
+	RegisteredInstanceIds.Reset();
+	for (auto& Pair : RegisteredPrimitives)
+	{
+		AllocTextureSpace(*Pair.Key, Pair.Value.NumSplines, Pair.Value);
+	}
+
+	// Sanity check defragmentation results
+	check(SlotAllocator.GetMaxSize() == SlotAllocator.GetSparselyAllocatedSize());
+}
+
+// Check to replace or update the cached Instance ID lookup buffer
+FRDGBufferSRVRef FSplineMeshSceneExtension::GetInstanceIdLookupSRV(FRDGBuilder& GraphBuilder, bool bForceUpdate)
+{
+	const uint32 InstanceIdLookupSize = RegisteredInstanceIds.Num();
+	const uint32 CurInstanceIdLookupSize = SavedIdLookup.IsValid() ? SavedIdLookup->Desc.NumElements : 0u;
+	const bool bNeedsResize = bForceUpdate || InstanceIdLookupSize != CurInstanceIdLookupSize;
+	const bool bNeedsUpload = bForceUpdate || !SavedIdLookup.IsValid() || bInstanceLookupDirty;
+
+	FRDGBufferRef InstanceIdLookup = nullptr;
+	if (bNeedsResize)
+	{
+		InstanceIdLookup = GraphBuilder.CreateBuffer(
+			FRDGBufferDesc::CreateStructuredDesc(sizeof(uint32), InstanceIdLookupSize),
+			TEXT("SplineMesh.InstanceIdLookup")
+		);
+
+		// Only persist the buffer if we're not force-updating every frame (keeps it in transient memory, otherwise)
+		SavedIdLookup = bForceUpdate ? nullptr : GraphBuilder.ConvertToExternalBuffer(InstanceIdLookup);
+	}
+	else
+	{
+		InstanceIdLookup = GraphBuilder.RegisterExternalBuffer(SavedIdLookup, TEXT("SplineMesh.InstanceIdLookup"));
+	}
+
+	if (bNeedsUpload)
+	{
+		// Upload the contents
+		GraphBuilder.QueueBufferUpload<uint32>(InstanceIdLookup, RegisteredInstanceIds, ERDGInitialDataFlags::NoCopy);
+	}
+
+	bInstanceLookupDirty = false;
+	return GraphBuilder.CreateSRV(InstanceIdLookup);
+}
+
+void FSplineMeshSceneExtension::ClearAllCache()
+{
+	SavedPosTexture = nullptr;
+	SavedRotTexture = nullptr;
+	SavedIdLookup = nullptr;
+}
+
+
+void FSplineMeshSceneUpdater::PreSceneUpdate(FRDGBuilder& GraphBuilder, const FScenePreUpdateChangeSet& ChangeSet)
+{
 	for (FPrimitiveSceneInfo* PrimitiveSceneInfo : ChangeSet.RemovedPrimitiveSceneInfos)
 	{
 		if (PrimitiveSceneInfo->Proxy->IsSplineMesh())
 		{
-			Unregister(*PrimitiveSceneInfo);
+			SceneData->Unregister(*PrimitiveSceneInfo);
 		}
 	}
 
 	// Consolidate free spans on the allocator after a batch of frees.
 	// NOTE: This is important to ensure the allocator trims down its max size naturally when freeing space off the end.
 	// This allows us to downsize the texture without defragging, which is a much heavier operation.
-	SlotAllocator.Consolidate();
+	SceneData->SlotAllocator.Consolidate();
 }
 
-void FSplineMeshSceneResources::PostSceneUpdate(FRDGBuilder& GraphBuilder, const FScenePostUpdateChangeSet& ChangeSet)
+void FSplineMeshSceneUpdater::PostSceneUpdate(FRDGBuilder& GraphBuilder, const FScenePostUpdateChangeSet& ChangeSet)
 {
+	auto RequestUpdate = [this](const FSplineMeshSceneExtension::FPrimitiveSlot& Slot)
+	{
+		for (uint32 i = 0; i < Slot.NumSplines; ++i)
+		{
+			UpdateRequests.AddUnique(Slot.FirstSplineIndex + i);
+		}
+	};
+
 	// First, register any new primitives
 	for (FPrimitiveSceneInfo* PrimitiveSceneInfo : ChangeSet.AddedPrimitiveSceneInfos)
 	{
 		if (PrimitiveSceneInfo->Proxy->IsSplineMesh())
 		{
-			Register(*PrimitiveSceneInfo);
+			RequestUpdate(SceneData->Register(*PrimitiveSceneInfo));
 		}
 	}
 
@@ -162,60 +343,58 @@
 	{
 		if (PrimitiveSceneInfo->Proxy->IsSplineMesh())
 		{
-			RequestUpdate(RegisteredPrimitives.FindChecked(PrimitiveSceneInfo));
+			RequestUpdate(SceneData->RegisteredPrimitives.FindChecked(PrimitiveSceneInfo));
 		}
 	}
 
-	if (NumRegisteredPrimitives() > 0)
+	if (SceneData->NumRegisteredPrimitives() > 0)
 	{
 		// Check to defrag the texture when we could halve the dimensions of the texture just by doing so. Must do this
 		// before GPU scene updates to give the primitives a chance to re-upload their new texcoord assignments.
-		const uint32 CurSize = SlotAllocator.GetMaxSize();
-		const uint32 DefraggedSize = SlotAllocator.GetSparselyAllocatedSize();
+		const uint32 CurSize = SceneData->SlotAllocator.GetMaxSize();
+		const uint32 DefraggedSize = SceneData->SlotAllocator.GetSparselyAllocatedSize();
 		if (SplineMesh::CalcTextureSize(DefraggedSize) < SplineMesh::CalcTextureSize(CurSize))
 		{
-			DefragTexture();
+			SceneData->DefragTexture();
 		}
 	}
 	else
 	{
-		// No active splines, clear all cache and bail
-		SavedPosTexture = nullptr;
-		SavedRotTexture = nullptr;
-		SavedIdLookup = nullptr;
-		UpdateRequests.Reset();
+		// No active splines, just clear all cache
+		SceneData->ClearAllCache();
 	}
 }
 
-void FSplineMeshSceneResources::Update(FRDGBuilder& GraphBuilder, FSceneUniformBuffer& SceneUniforms)
+void FSplineMeshSceneUpdater::PostGPUSceneUpdate(FRDGBuilder& GraphBuilder, FSceneUniformBuffer& SceneUniforms)
 {
-	if (NumRegisteredPrimitives() == 0)
+	if (SceneData->NumRegisteredPrimitives() == 0)
 	{
 		return; // nothing to do
 	}
 
-	// Check if we need to re-size the texture
-	uint32 NeededSize = SplineMesh::CalcTextureSize(SlotAllocator.GetMaxSize());
-	const uint32 CurSize = SavedPosTexture.IsValid() ? SavedPosTexture->GetDesc().Extent.X : 0;
+	// Check if we need to re-size the cached texture
+	uint32 NeededSize = SplineMesh::CalcTextureSize(SceneData->SlotAllocator.GetMaxSize());
+	const uint32 CurSize = SceneData->SavedPosTexture.IsValid() ? SceneData->SavedPosTexture->GetDesc().Extent.X : 0;
 
 	// Clamp to the max dimension and check to report an error about over-sizing the spline mesh texture
 	if (NeededSize > SPLINE_MESH_TEXTURE_MAX_DIMENSION)
 	{
 		NeededSize = SPLINE_MESH_TEXTURE_MAX_DIMENSION;
-		if (!bOverflowError)
+		if (!SceneData->bOverflowError)
 		{
 			UE_LOG(LogRenderer, Error,
 				TEXT("Too many spline meshes have been registered with the scene. The spline mesh texture has grown ")
 				TEXT("to its max size (%dx%d - see r.SplineMesh.BakeToTexture.MaxDimension) and has ran out of space. ")
 				TEXT("Expect some spline meshes to render incorrectly."),
-				NeededSize, NeededSize);
-			bOverflowError = true;
+				NeededSize, NeededSize
+			);
+			SceneData->bOverflowError = true;
 		}
 	}
 
 	// Check if we are forcing a full update because we have no cache or are debugging
 	const bool bForceUpdate = CVarSplineMeshSceneTexturesForceUpdate.GetValueOnRenderThread() != 0;
-	bool bFullUpdate = !SavedPosTexture.IsValid() || bForceUpdate;
+	bool bFullUpdate = !SceneData->SavedPosTexture.IsValid() || bForceUpdate;
 
 	// Register or create the spline texture
 	FRDGTextureRef PosTexture = nullptr;
@@ -247,7 +426,7 @@
 			const uint32 CopyExtent = FMath::Min(NeededSize, CurSize);
 
 			FRDGTextureRef CopySrc = GraphBuilder.RegisterExternalTexture(
-				SavedPosTexture,
+				SceneData->SavedPosTexture,
 				TEXT("SplineMesh.PrevSplinePosTexture")
 			);
 			AddCopyTexturePass(
@@ -260,7 +439,7 @@
 			);
 			
 			CopySrc = GraphBuilder.RegisterExternalTexture(
-				SavedRotTexture,
+				SceneData->SavedRotTexture,
 				TEXT("SplineMesh.PrevSplineRotTexture")
 			);
 			AddCopyTexturePass(
@@ -274,16 +453,16 @@
 		}
 
 		// Don't store off the texture if we're updating every frame (keeps it transient, otherwise)
-		SavedPosTexture = bForceUpdate ? nullptr : GraphBuilder.ConvertToExternalTexture(PosTexture);
-		SavedRotTexture = bForceUpdate ? nullptr : GraphBuilder.ConvertToExternalTexture(RotTexture);
+		SceneData->SavedPosTexture = bForceUpdate ? nullptr : GraphBuilder.ConvertToExternalTexture(PosTexture);
+		SceneData->SavedRotTexture = bForceUpdate ? nullptr : GraphBuilder.ConvertToExternalTexture(RotTexture);
 	}
 	else
 	{
-		check(SavedPosTexture.IsValid());
-		PosTexture = GraphBuilder.RegisterExternalTexture(SavedPosTexture, TEXT("SplineMesh.SplinePosTexture"));
+		check(SceneData->SavedPosTexture.IsValid());
+		PosTexture = GraphBuilder.RegisterExternalTexture(SceneData->SavedPosTexture, TEXT("SplineMesh.SplinePosTexture"));
 
-		check(SavedRotTexture.IsValid());
-		RotTexture = GraphBuilder.RegisterExternalTexture(SavedRotTexture, TEXT("SplineMesh.SplineRotTexture"));
+		check(SceneData->SavedRotTexture.IsValid());
+		RotTexture = GraphBuilder.RegisterExternalTexture(SceneData->SavedRotTexture, TEXT("SplineMesh.SplineRotTexture"));
 	}
 
 	// Perform the update and clear pending requests
@@ -291,25 +470,24 @@
 	const FVector2f InvExtent = FVector2f(1.0f / Extent.X, 1.0f / Extent.Y);
 	if (bFullUpdate || UpdateRequests.Num() > 0)
 	{
-		AddUpdatePass(GraphBuilder, PosTexture, RotTexture, SceneUniforms, Extent, InvExtent, bFullUpdate, bForceUpdate);
+		AddUpdatePass(
+			GraphBuilder,
+			SceneUniforms,
+			PosTexture,
+			RotTexture,
+			Extent,
+			InvExtent,
+			bFullUpdate,
+			bForceUpdate
+		);
 	}
-	UpdateRequests.Reset();
-
-	// Lastly, set up the scene uniforms for spline meshes
-	FSplineMeshSceneResourceParameters ShaderParams;
-	ShaderParams.SplinePosTexture = GraphBuilder.CreateSRV(PosTexture);
-	ShaderParams.SplineRotTexture = GraphBuilder.CreateSRV(RotTexture);
-	ShaderParams.SplineSampler = TStaticSamplerState<SF_Bilinear, AM_Clamp, AM_Clamp, AM_Clamp>::GetRHI();
-	ShaderParams.SplineTextureInvExtent = InvExtent;
-
-	SceneUniforms.Set(SceneUB::SplineMesh, ShaderParams);
 }
 
-void FSplineMeshSceneResources::AddUpdatePass(
+void FSplineMeshSceneUpdater::AddUpdatePass(
 	FRDGBuilder& GraphBuilder,
+	FSceneUniformBuffer& SceneUniforms,
 	FRDGTextureRef PosTexture,
 	FRDGTextureRef RotTexture,
-	FSceneUniformBuffer& SceneUniforms,
 	FVector2f Extent,
 	FVector2f InvExtent,
 	bool bFullUpdate,
@@ -351,11 +529,13 @@
 		UpdateRequestBuffer = GSystemTextures.GetDefaultStructuredBuffer<uint32>(GraphBuilder);
 	}
 
+	FScene& Scene = *SceneData->Scene;
 	auto* PassParameters = GraphBuilder.AllocParameters<FSplineMeshTextureFillCS::FParameters>();
+
 	PassParameters->Scene = SceneUniforms.GetBuffer(GraphBuilder);
 	PassParameters->SplinePosTextureOut = PosTextureUAV;
 	PassParameters->SplineRotTextureOut = RotTextureUAV;
-	PassParameters->InstanceIdLookup = GetInstanceIdLookupSRV(GraphBuilder, bForceUpdate);
+	PassParameters->InstanceIdLookup = SceneData->GetInstanceIdLookupSRV(GraphBuilder, bForceUpdate);
 	PassParameters->UpdateRequests = GraphBuilder.CreateSRV(UpdateRequestBuffer);
 	PassParameters->NumUpdateRequests = NumUpdateRequests;
 	PassParameters->InstanceDataSOAStride = Scene.GPUScene.InstanceSceneDataSOAStride;
@@ -364,7 +544,7 @@
 
 	auto ComputeShader = GetGlobalShaderMap(Scene.GetFeatureLevel())->GetShader<FSplineMeshTextureFillCS>();
 
-	const uint32 NumThreadGroups = NumUpdateRequests ? NumUpdateRequests : SlotAllocator.GetMaxSize();
+	const uint32 NumThreadGroups = NumUpdateRequests ? NumUpdateRequests : SceneData->SlotAllocator.GetMaxSize();
 	FComputeShaderUtils::AddPass(
 		GraphBuilder,
 		RDG_EVENT_NAME("SplineMeshTextureUpdate"),
@@ -374,167 +554,24 @@
 	);
 }
 
-void FSplineMeshSceneResources::Register(const FPrimitiveSceneInfo& PrimitiveSceneInfo)
-{	
-	FPrimitiveSlot& Slot = RegisteredPrimitives.FindOrAdd(&PrimitiveSceneInfo);
-	if (!ensureMsgf(Slot.NumSplines == 0, TEXT("This primitive was already registered!")))
+
+void FSplineMeshSceneRenderer::UpdateSceneUniformBuffer(FRDGBuilder& GraphBuilder, FSceneUniformBuffer& SceneUniforms)
+{
+	if (SceneData->NumRegisteredPrimitives() > 0)
 	{
-		return;
-	}
+		FRDGTextureRef PosTexture = GraphBuilder.RegisterExternalTexture(SceneData->SavedPosTexture);
+		FRDGTextureRef RotTexture = GraphBuilder.RegisterExternalTexture(SceneData->SavedRotTexture);
 
-	// Alloc space for the new splines and ensure they are included in the next update
-	AllocTextureSpace(PrimitiveSceneInfo, GetNumSplines(PrimitiveSceneInfo), Slot);
-	RequestUpdate(Slot);
-}
-
-void FSplineMeshSceneResources::Unregister(const FPrimitiveSceneInfo& PrimitiveSceneInfo)
-{
-	FPrimitiveSlot* Slot = RegisteredPrimitives.Find(&PrimitiveSceneInfo);
-	if (!ensure(Slot != nullptr))
-	{
-		return;
-	}
-
-	SlotAllocator.Free(Slot->FirstSplineIndex, Slot->NumSplines);
-
-	// Clear the instance IDs of the newly freed range
-	for (uint32 i = 0; i < Slot->NumSplines; ++i)
-	{
-		RegisteredInstanceIds[Slot->FirstSplineIndex + i] = INDEX_NONE;
-	}
+		const FIntPoint Extent = SceneData->SavedPosTexture->GetDesc().Extent;
+		const FVector2f InvExtent = FVector2f(1.0f / Extent.X, 1.0f / Extent.Y);
 	
-	// shorten this look-up if the allocator shrinks
-	RegisteredInstanceIds.SetNumUninitialized(FMath::Min(SlotAllocator.GetMaxSize(), RegisteredInstanceIds.Num()));
-
-	RegisteredPrimitives.Remove(&PrimitiveSceneInfo);
-
-	// Ensure we repopulate the instance ID lookup next update
-	bInstanceLookupDirty = true;
-}
-
-void FSplineMeshSceneResources::AllocTextureSpace(const FPrimitiveSceneInfo& PrimitiveSceneInfo, uint32 NumSplines, FPrimitiveSlot& OutSlot)
-{
-	check(NumSplines > 0);
-
-	OutSlot.NumSplines = NumSplines;
-	OutSlot.FirstSplineIndex = SlotAllocator.Allocate(NumSplines);
-
-	// If we don't have the space for these instances in the registered ID list, allocate them
-	if (RegisteredInstanceIds.Num() < SlotAllocator.GetMaxSize())
-	{
-		RegisteredInstanceIds.Reserve(SlotAllocator.GetMaxSize());
-		for (uint32 i = RegisteredInstanceIds.Num(); i < OutSlot.FirstSplineIndex + NumSplines; ++i)
-		{
-			RegisteredInstanceIds.Add(INDEX_NONE);
-		}
-	}
-
-	// Store the instance ID for each spline slot allocated (will be used to fill texels with spline data)
-	// NOTE: This is assuming a spline per instance. If that changes, this look-up will need to as well
-	const int32 InstanceSceneDataOffset = PrimitiveSceneInfo.GetInstanceSceneDataOffset();
-	for (uint32 i = 0; i < NumSplines; ++i)
-	{
-		const uint32 SplineIndex = OutSlot.FirstSplineIndex + i;
-		check(RegisteredInstanceIds[SplineIndex] == INDEX_NONE); // sanity check it has been cleared
-		RegisteredInstanceIds[SplineIndex] = InstanceSceneDataOffset + i;
-	}
-
-	bInstanceLookupDirty = true;
-
-	// Give the scene proxy the coordinates allocated so it can place it in its scene data.
-	AssignCoordinates(PrimitiveSceneInfo, OutSlot);
-}
-
-uint32 FSplineMeshSceneResources::GetNumSplines(const FPrimitiveSceneInfo& SceneInfo)
-{
-	// We only support spline mesh scene proxies currently, and we assume the number of splines is equal to the
-	// number of instance scene data entries. Support could be added later for other scene proxy types that need
-	// baked down splines.
-	check(SceneInfo.Proxy->IsSplineMesh());
-
-	return SceneInfo.GetNumInstanceSceneDataEntries();
-}
-
-void FSplineMeshSceneResources::AssignCoordinates(const FPrimitiveSceneInfo& SceneInfo, const FPrimitiveSlot& Slot)
-{
-	check(SceneInfo.Proxy->IsSplineMesh()); // sanity check
-
-	if (SceneInfo.Proxy->IsNaniteMesh())
-	{
-		AssignCoordinates(static_cast<FNaniteSplineMeshSceneProxy*>(SceneInfo.Proxy), Slot);
-	}
-	else
-	{
-		AssignCoordinates(static_cast<FSplineMeshSceneProxy*>(SceneInfo.Proxy), Slot);
+		// Lastly, set up the scene uniforms for spline meshes
+		FSplineMeshSceneResourceParameters ShaderParams;
+		ShaderParams.SplinePosTexture = GraphBuilder.CreateSRV(PosTexture);
+		ShaderParams.SplineRotTexture = GraphBuilder.CreateSRV(RotTexture);
+		ShaderParams.SplineSampler = TStaticSamplerState<SF_Bilinear, AM_Clamp, AM_Clamp, AM_Clamp>::GetRHI();
+		ShaderParams.SplineTextureInvExtent = InvExtent;
+	
+		SceneUniforms.Set(SceneUB::SplineMesh, ShaderParams);
 	}
 }
-
-template<typename TSplineMeshSceneProxy>
-void FSplineMeshSceneResources::AssignCoordinates(TSplineMeshSceneProxy* SceneProxy, const FPrimitiveSlot& Slot)
-{
-	for (uint32 i = 0; i < Slot.NumSplines; ++i)
-	{
-		SceneProxy->SetSplineTextureCoord_RenderThread(i, SplineMesh::CalcTextureCoord(Slot.FirstSplineIndex + i));
-	}
-}
-
-void FSplineMeshSceneResources::RequestUpdate(const FPrimitiveSlot& Slot)
-{
-	for (uint32 i = 0; i < Slot.NumSplines; ++i)
-	{
-		const uint32 SplineIndex = Slot.FirstSplineIndex + i;
-		check(RegisteredInstanceIds[SplineIndex] != INDEX_NONE); // Sanity check
-		UpdateRequests.AddUnique(SplineIndex);
-	}
-}
-
-void FSplineMeshSceneResources::DefragTexture()
-{
-	// NOTE: Currently not attempting to reduce motions to a minimal set, we're just ditching our cache
-	// and re-assigning space in the new texture that will be created next update
-	SavedPosTexture = nullptr;
-	SavedRotTexture = nullptr;
-	SlotAllocator.Reset();
-	RegisteredInstanceIds.Reset();
-	for (auto& Pair : RegisteredPrimitives)
-	{
-		AllocTextureSpace(*Pair.Key, Pair.Value.NumSplines, Pair.Value);
-	}
-
-	// Sanity check defragmentation results
-	check(SlotAllocator.GetMaxSize() == SlotAllocator.GetSparselyAllocatedSize());
-}
-
-// Check to replace or update the cached Instance ID lookup buffer
-FRDGBufferSRVRef FSplineMeshSceneResources::GetInstanceIdLookupSRV(FRDGBuilder& GraphBuilder, bool bForceUpdate)
-{
-	const uint32 InstanceIdLookupSize = RegisteredInstanceIds.Num();
-	const uint32 CurInstanceIdLookupSize = SavedIdLookup.IsValid() ? SavedIdLookup->Desc.NumElements : 0u;
-	const bool bNeedsResize = bForceUpdate || InstanceIdLookupSize != CurInstanceIdLookupSize;
-	const bool bNeedsUpload = bForceUpdate || !SavedIdLookup.IsValid() || bInstanceLookupDirty;
-
-	FRDGBufferRef InstanceIdLookup = nullptr;
-	if (bNeedsResize)
-	{
-		InstanceIdLookup = GraphBuilder.CreateBuffer(
-			FRDGBufferDesc::CreateStructuredDesc(sizeof(uint32), InstanceIdLookupSize),
-			TEXT("SplineMesh.InstanceIdLookup")
-		);
-
-		// Only persist the buffer if we're not force-updating every frame (keeps it in transient memory, otherwise)
-		SavedIdLookup = bForceUpdate ? nullptr : GraphBuilder.ConvertToExternalBuffer(InstanceIdLookup);
-	}
-	else
-	{
-		InstanceIdLookup = GraphBuilder.RegisterExternalBuffer(SavedIdLookup, TEXT("SplineMesh.InstanceIdLookup"));
-	}
-
-	if (bNeedsUpload)
-	{
-		// Upload the contents
-		GraphBuilder.QueueBufferUpload<uint32>(InstanceIdLookup, RegisteredInstanceIds, ERDGInitialDataFlags::NoCopy);
-	}
-
-	bInstanceLookupDirty = false;
-	return GraphBuilder.CreateSRV(InstanceIdLookup);
-}
diff --git a/Engine/Source/Runtime/Renderer/Private/SplineMeshSceneResources.h b/Engine/Source/Runtime/Renderer/Private/SplineMeshSceneResources.h
--- a/Engine/Source/Runtime/Renderer/Private/SplineMeshSceneResources.h
+++ b/Engine/Source/Runtime/Renderer/Private/SplineMeshSceneResources.h
@@ -8,6 +8,8 @@
 #include "RHIShaderPlatform.h"
 #include "RendererInterface.h"
 #include "Containers/Map.h"
+#include "RenderGraphResources.h"
+#include "SceneExtensions.h"
 
 class FRDGBuilder;
 class FScene;
@@ -16,9 +18,10 @@ class FScenePreUpdateChangeSet;
 class FScenePostUpdateChangeSet;
 class FPrimitiveSceneInfo;
 struct IPooledRenderTarget;
+class FSceneRendererBase;
 
 /**
- *	FSplineMeshSceneResources
+ *	FSplineMeshSceneExtension
  *
  *	This class manages a texture that is used to bake down spline position values at a fixed number of texels
  *	to lower the cost of sampling these splines when deforming spline mesh vertices. Each spline in the scene is
@@ -32,14 +35,19 @@ struct IPooledRenderTarget;
  *	- Defragmentation of the texture will occur if the texture could be 1/4 the size of the current texture
  *	  size when tightly allocated.
  */
-class FSplineMeshSceneResources
+class FSplineMeshSceneExtension : public ISceneExtension
 {
-public:
-	FSplineMeshSceneResources(FScene& InScene);
+	friend class FSplineMeshSceneUpdater;
+	friend class FSplineMeshSceneRenderer;
+
+	DECLARE_SCENE_EXTENSION(FSplineMeshSceneExtension);
 
-	void PreSceneUpdate(FRDGBuilder& GraphBuilder, const FScenePreUpdateChangeSet& ChangeSet);
-	void PostSceneUpdate(FRDGBuilder& GraphBuilder, const FScenePostUpdateChangeSet& ChangeSet);
-	void Update(FRDGBuilder& GraphBuilder, FSceneUniformBuffer& SceneUniforms);
+public:
+	static bool ShouldCreateExtension(FScene& InScene);
+	
+	virtual void InitExtension(FScene& InScene) override { Scene = &InScene; }
+	virtual ISceneExtensionUpdater* CreateUpdater() override;
+	virtual ISceneExtensionRenderer* CreateRenderer() override;
 
 	uint32 NumRegisteredPrimitives() const { return RegisteredPrimitives.Num(); }
 
@@ -51,32 +59,21 @@ private:
 	};
 	using FPrimitiveSlotMap = TMap<const FPrimitiveSceneInfo*, FPrimitiveSlot>;
 
-	void AddUpdatePass(
-		FRDGBuilder& GraphBuilder,
-		FRDGTextureRef PosTexture,
-		FRDGTextureRef RotTexture,
-		FSceneUniformBuffer& SceneUniforms,
-		FVector2f Extent,
-		FVector2f InvExtent,
-		bool bFullUpdate,
-		bool bForceUpdate);
-
-	void Register(const FPrimitiveSceneInfo& PrimitiveSceneInfo);
+	FPrimitiveSlot& Register(const FPrimitiveSceneInfo& PrimitiveSceneInfo);
 	void Unregister(const FPrimitiveSceneInfo& PrimitiveSceneInfo);
 	void AllocTextureSpace(const FPrimitiveSceneInfo& PrimitiveSceneInfo, uint32 NumSplines, FPrimitiveSlot& OutSlot);
 	static uint32 GetNumSplines(const FPrimitiveSceneInfo& SceneInfo);
 	void AssignCoordinates(const FPrimitiveSceneInfo& SceneInfo, const FPrimitiveSlot& Slot);
 	template<typename TSplineMeshSceneProxy>
 	void AssignCoordinates(TSplineMeshSceneProxy* SceneProxy, const FPrimitiveSlot& Slot);
-	void RequestUpdate(const FPrimitiveSlot& Slot);
 	void DefragTexture();
 	FRDGBufferSRVRef GetInstanceIdLookupSRV(FRDGBuilder& GraphBuilder, bool bForceUpdate);
+	void ClearAllCache();
 
 private:
-	FScene& Scene;
+	FScene* Scene = nullptr;
 	FPrimitiveSlotMap RegisteredPrimitives;
 	TArray<uint32> RegisteredInstanceIds;
-	TArray<uint32> UpdateRequests;
 	FSpanAllocator SlotAllocator;
 	TRefCountPtr<IPooledRenderTarget> SavedPosTexture;
 	TRefCountPtr<IPooledRenderTarget> SavedRotTexture;
@@ -84,3 +81,40 @@ private:
 	bool bInstanceLookupDirty = true;
 	bool bOverflowError = false;
 };
+
+/** This class performs updates to the persistent spline mesh resources. */
+class FSplineMeshSceneUpdater : public ISceneExtensionUpdater
+{
+public:
+	FSplineMeshSceneUpdater(FSplineMeshSceneExtension& InSceneData) : SceneData(&InSceneData) {}
+
+	virtual void PreSceneUpdate(FRDGBuilder& GraphBuilder, const FScenePreUpdateChangeSet& ChangeSet) override;
+	virtual void PostSceneUpdate(FRDGBuilder& GraphBuilder, const FScenePostUpdateChangeSet& ChangeSet) override;
+	virtual void PostGPUSceneUpdate(FRDGBuilder& GraphBuilder, FSceneUniformBuffer& SceneUniforms) override;
+
+private:
+	void AddUpdatePass(
+		FRDGBuilder& GraphBuilder,
+		FSceneUniformBuffer& SceneUniforms,
+		FRDGTextureRef PosTexture,
+		FRDGTextureRef RotTexture,
+		FVector2f Extent,
+		FVector2f InvExtent,
+		bool bFullUpdate,
+		bool bForceUpdate
+	);
+
+	FSplineMeshSceneExtension* SceneData = nullptr;
+	TArray<uint32, FSceneRenderingArrayAllocator> UpdateRequests;
+};
+
+/** This class is used to insert the spline mesh scene uniforms into the scene uniform buffer for any given renderer */
+class FSplineMeshSceneRenderer : public ISceneExtensionRenderer
+{
+public:
+	FSplineMeshSceneRenderer(FSplineMeshSceneExtension& InSceneData) : SceneData(&InSceneData) {}
+	virtual void UpdateSceneUniformBuffer(FRDGBuilder& GraphBuilder, FSceneUniformBuffer& SceneUniforms) override;
+
+private:
+	FSplineMeshSceneExtension* SceneData = nullptr;
+};
\ No newline at end of file
diff --git a/Engine/Source/Runtime/Renderer/Private/VT/RuntimeVirtualTextureRender.cpp b/Engine/Source/Runtime/Renderer/Private/VT/RuntimeVirtualTextureRender.cpp
--- a/Engine/Source/Runtime/Renderer/Private/VT/RuntimeVirtualTextureRender.cpp
+++ b/Engine/Source/Runtime/Renderer/Private/VT/RuntimeVirtualTextureRender.cpp
@@ -21,7 +21,8 @@
 #include "RenderCaptureInterface.h"
 #include "SimpleMeshDrawCommandPass.h"
 #include "StaticMeshBatch.h"
-#include "SplineMeshSceneResources.h"
+#include "SceneRendering.h"
+#include "EngineModule.h"
 
 namespace RuntimeVirtualTexture
 {
@@ -1353,7 +1354,7 @@ namespace RuntimeVirtualTexture
 	void RenderPage(
 		FRDGBuilder& GraphBuilder,
 		FScene* Scene,
-		FSceneUniformBuffer& SceneUB,
+		ISceneRenderer* SceneRenderer,
 		uint32 RuntimeVirtualTextureMask,
 		ERuntimeVirtualTextureMaterialType MaterialType,
 		bool bClearTextures,
@@ -1381,7 +1382,8 @@ namespace RuntimeVirtualTexture
 		//todo[vt]: Have specific shader variations and setup for different output texture configs
 		FSceneViewFamily::ConstructionValues ViewFamilyInit(nullptr, Scene, FEngineShowFlags(ESFIM_Game));
 		ViewFamilyInit.SetTime(FGameTime());
-		FSceneViewFamily& ViewFamily = *GraphBuilder.AllocObject<FViewFamilyInfo>(ViewFamilyInit);
+		FViewFamilyInfo& ViewFamily = *GraphBuilder.AllocObject<FViewFamilyInfo>(ViewFamilyInit);
+		ViewFamily.SetSceneRenderer(SceneRenderer);
 
 		FSceneViewInitOptions ViewInitOptions;
 		ViewInitOptions.ViewFamily = &ViewFamily;
@@ -1450,7 +1452,7 @@ namespace RuntimeVirtualTexture
 			ERenderTargetLoadAction LoadAction = bClearTextures ? ERenderTargetLoadAction::EClear : ERenderTargetLoadAction::ENoAction;
 			FShader_VirtualTextureMaterialDraw::FParameters* PassParameters = GraphBuilder.AllocParameters<FShader_VirtualTextureMaterialDraw::FParameters>();
 			PassParameters->View = View->ViewUniformBuffer;
-			PassParameters->Scene = SceneUB.GetBuffer(GraphBuilder);
+			PassParameters->Scene = SceneRenderer->GetSceneUniformBufferRef(GraphBuilder);
 			PassParameters->RenderTargets[0] = GraphSetup.RenderTexture0 ? FRenderTargetBinding(GraphSetup.RenderTexture0, LoadAction) : FRenderTargetBinding();
 			PassParameters->RenderTargets[1] = GraphSetup.RenderTexture1 ? FRenderTargetBinding(GraphSetup.RenderTexture1, LoadAction) : FRenderTargetBinding();
 			PassParameters->RenderTargets[2] = GraphSetup.RenderTexture2 ? FRenderTargetBinding(GraphSetup.RenderTexture2, LoadAction) : FRenderTargetBinding();
@@ -1527,7 +1529,7 @@ namespace RuntimeVirtualTexture
 		}
 	}
 
-	void RenderPagesInternal(FRDGBuilder& GraphBuilder, FRenderPageBatchDesc const& InDesc, FSceneUniformBuffer& SceneUB)
+	void RenderPagesInternal(FRDGBuilder& GraphBuilder, FRenderPageBatchDesc const& InDesc, ISceneRenderer* SceneRenderer)
 	{
 		check(InDesc.NumPageDescs <= EMaxRenderPageBatch);
 
@@ -1543,7 +1545,7 @@ namespace RuntimeVirtualTexture
 				RenderPage(
 					GraphBuilder,
 					InDesc.Scene,
-					SceneUB,
+					SceneRenderer,
 					InDesc.RuntimeVirtualTextureMask,
 					InDesc.MaterialType,
 					InDesc.bClearTextures,
@@ -1563,34 +1565,31 @@ namespace RuntimeVirtualTexture
 
 	void RenderPagesStandAlone(FRDGBuilder& GraphBuilder, FRenderPageBatchDesc const& InDesc)
 	{
-		InDesc.Scene->UpdateAllPrimitiveSceneInfos(GraphBuilder);
-
-		// This is required to collect dynamic primitives from the views (not used here, but we must provide one).
-		FGPUSceneDynamicContext GPUSceneDynamicContext(InDesc.Scene->GPUScene);
-		// Call to let GPU-Scene determine if it is active and record scene primitive count
-		FGPUSceneScopeBeginEndHelper GPUSceneScopeBeginEndHelper(GraphBuilder, InDesc.Scene->GPUScene, GPUSceneDynamicContext);
-
-		FSceneUniformBuffer SceneUB {};
-		InDesc.Scene->GPUScene.FillSceneUniformBuffer(GraphBuilder, SceneUB);
-		if (InDesc.Scene->SplineMeshSceneResources)
-		{
-			InDesc.Scene->SplineMeshSceneResources->Update(GraphBuilder, SceneUB);
-		}
-		RenderPagesInternal(GraphBuilder, InDesc, SceneUB);
+		check(InDesc.Scene != nullptr);
+		FScenePrimitiveRenderingContextScopeHelper RenderingScope(GetRendererModule().BeginScenePrimitiveRendering(GraphBuilder, *InDesc.Scene));
+		RenderPagesInternal(GraphBuilder, InDesc, RenderingScope.ScenePrimitiveRenderingContext->GetSceneRenderer());
 	}
 
 	void RenderPages(FRDGBuilder& GraphBuilder, FRenderPageBatchDesc const& InDesc)
 	{
+		check(InDesc.Scene != nullptr);
 		if (InDesc.Scene->GPUScene.IsRendering())
 		{
 			// TODO: this should be replaced by piping through a reference to the scene renderer rather than just the scene, such that we can get at the already populated scene UB.
-			FSceneUniformBuffer SceneUB{};
-			InDesc.Scene->GPUScene.FillSceneUniformBuffer(GraphBuilder, SceneUB);
-			if (InDesc.Scene->SplineMeshSceneResources)
+			class FSimpleRVTRenderer : public FSceneRendererBase
 			{
-				InDesc.Scene->SplineMeshSceneResources->Update(GraphBuilder, SceneUB);
+			public:
+				FSimpleRVTRenderer(FRDGBuilder& GraphBuilder, FRenderPageBatchDesc const& InDesc)
+					: FSceneRendererBase(*InDesc.Scene)
+				{
+					InitSceneExtensionsRenderer();
+
+					InDesc.Scene->GPUScene.FillSceneUniformBuffer(GraphBuilder, GetSceneUniforms());
+					GetSceneExtensionsRenderer().UpdateSceneUniformBuffer(GraphBuilder, GetSceneUniforms());
 			}
-			RenderPagesInternal(GraphBuilder, InDesc, SceneUB);
+			};
+			FSimpleRVTRenderer SimpleRenderer(GraphBuilder, InDesc);
+			RenderPagesInternal(GraphBuilder, InDesc, &SimpleRenderer);
 		}
 		else
 		{
