LightFunctionAtlas - added debug visualization pass.

diff --git a/Engine/Source/Runtime/Engine/Public/ShowFlags.h b/Engine/Source/Runtime/Engine/Public/ShowFlags.h
--- a/Engine/Source/Runtime/Engine/Public/ShowFlags.h
+++ b/Engine/Source/Runtime/Engine/Public/ShowFlags.h
@@ -476,6 +476,7 @@ private:
 		SetPathTracing(false);
 		SetRayTracingDebug(false);
 		SetVisualizeSkyAtmosphere(false);
+		SetVisualizeLightFunctionAtlas(false);
 		SetVisualizeCalibrationColor(false);
 		SetVisualizeCalibrationGrayscale(false);
 		SetVisualizeCalibrationCustom(false);
diff --git a/Engine/Source/Runtime/Engine/Public/ShowFlagsValues.inl b/Engine/Source/Runtime/Engine/Public/ShowFlagsValues.inl
--- a/Engine/Source/Runtime/Engine/Public/ShowFlagsValues.inl
+++ b/Engine/Source/Runtime/Engine/Public/ShowFlagsValues.inl
@@ -395,6 +395,8 @@ SHOWFLAG_FIXED_IN_SHIPPING(0, RayTracingDebug, SFG_Hidden, NSLOCTEXT("UnrealEd",
 
 /** Enable the SkyAtmosphere visualization to be drawn on screen */
 SHOWFLAG_FIXED_IN_SHIPPING(0, VisualizeSkyAtmosphere, SFG_Visualize, NSLOCTEXT("UnrealEd", "VisualizeSkyAtmosphereSF", "Sky Atmosphere"))
+/** Enable the light function atlas debug visualization to be drawn on screen */
+SHOWFLAG_FIXED_IN_SHIPPING(0, VisualizeLightFunctionAtlas, SFG_Visualize, NSLOCTEXT("UnrealEd", "VisualizeLightFunctionAtlasSF", "Visualize Light Function Atlas"))
 /** Shows a full-screen calibration color post process material defined in renderer config  */
 SHOWFLAG_ALWAYS_ACCESSIBLE(VisualizeCalibrationColor, SFG_Developer, NSLOCTEXT("UnrealEd", "VisualizeCalibrationColorSF", "Visualize Calibration Color"))
 /** Shows a full-screen calibration grayscale post process material defined in renderer config  */
diff --git a/Engine/Source/Runtime/Renderer/Private/LightFunctionAtlas.cpp b/Engine/Source/Runtime/Renderer/Private/LightFunctionAtlas.cpp
--- a/Engine/Source/Runtime/Renderer/Private/LightFunctionAtlas.cpp
+++ b/Engine/Source/Runtime/Renderer/Private/LightFunctionAtlas.cpp
@@ -14,6 +14,9 @@
 #include "PostProcess/PostProcessing.h"
 #include "PostProcess/SceneFilterRendering.h"
 #include "ShadowRendering.h"
+#include "CanvasItem.h"
+#include "CanvasTypes.h"
+#include "Containers/HashTable.h"
 
 DECLARE_GPU_STAT(LightFunctionAtlasGeneration);
 
@@ -241,9 +244,12 @@ void FLightFunctionAtlas::BeginSceneFrame(FViewFamilyInfo& ViewFamily, TArray<FV
 	}
 
 	// We propagate bLightFunctionAtlasEnabled to all the views to ease later shader parameter decision and binding for lighting, shadow or volumetric fog for instance (avoid sending lots of parameters all over the place)
-	LightFunctionAtlasSceneData.bLightFunctionAtlasEnabled				= bLightFunctionAtlasEnabled;
-	LightFunctionAtlasSceneData.bVolumetricFogUsesLightFunctionAtlas	= bLightFunctionAtlasEnabled && VolumetricFogRequestsLightFunctionAtlas;
-	LightFunctionAtlasSceneData.bDeferredlightingUsesLightFunctionAtlas	= bLightFunctionAtlasEnabled && DeferredlightingRequestsLightFunctionAtlas;
+	LightFunctionAtlasSceneData.SetData(
+		this,
+		bLightFunctionAtlasEnabled,
+		bLightFunctionAtlasEnabled && VolumetricFogRequestsLightFunctionAtlas,
+		bLightFunctionAtlasEnabled && DeferredlightingRequestsLightFunctionAtlas);
+
 	for (auto& View : Views)
 	{
 		View.LightFunctionAtlasViewData = FLightFunctionAtlasViewData(&LightFunctionAtlasSceneData);
@@ -667,3 +673,101 @@ void FLightFunctionAtlas::RenderAtlasSlots(FRDGBuilder& GraphBuilder, const TArr
 		}
 	);
 }
+
+FScreenPassTexture FLightFunctionAtlas::AddDebugVisualizationPasses(FRDGBuilder& GraphBuilder, const FViewInfo& View, FScreenPassTexture& ScreenPassSceneColor) const
+{
+#if WITH_EDITOR
+
+	if (!IsLightFunctionAtlasEnabled())
+	{
+		return ScreenPassSceneColor;
+	}
+
+	uint32 AtlasSlotResolution = GetAtlasSlotResolution();
+	uint32 AtlasEdgeSize = GetAtlasEdgeSize();
+	uint32 AtlasResolution = AtlasSlotResolution * AtlasEdgeSize;
+
+	const FIntPoint SrcPoint = FIntPoint::ZeroValue;
+	const FIntPoint SrcSize  = RDGAtlasTexture2D->Desc.Extent;
+	const FIntPoint DstPoint = FIntPoint(100, 100);
+	const FIntPoint DstSize  = FIntPoint(512, 512);
+
+	AddDrawTexturePass(
+		GraphBuilder,
+		View,
+		RDGAtlasTexture2D,
+		ScreenPassSceneColor.Texture,
+		SrcPoint,
+		SrcSize,
+		DstPoint,
+		DstSize);
+
+	// Now debug print
+	AddDrawCanvasPass(GraphBuilder, {}, View, FScreenPassRenderTarget(ScreenPassSceneColor, ERenderTargetLoadAction::ELoad),
+		[&View, this, DstPoint, DstSize, AtlasResolution](FCanvas& Canvas)
+	{
+		FString Text;
+
+		const float ViewPortWidth = float(View.ViewRect.Width());
+		const float ViewPortHeight = float(View.ViewRect.Height());
+		float DrawPosX = float(DstPoint.X + DstSize.X) + 30.0f;
+		float DrawPosY = float(DstPoint.Y) + 10.0f;
+
+		float DisplayResolutionRatio = float(DstSize.X) / float(AtlasResolution);
+
+		Canvas.DrawShadowedString(DstPoint.X + 180.0f, DstPoint.Y - 20.0f, TEXT("LIGHT FUNCTION ATLAS"), GEngine->GetLargeFont(), FLinearColor::White);
+
+		Text = FString::Printf(TEXT("Light Functions in atlas:         %d"), EffectiveLightFunctionSlotArray.Num());
+		Canvas.DrawShadowedString(DrawPosX, DrawPosY, *Text, GEngine->GetLargeFont(), FLinearColor::White);
+		DrawPosY += 20.0f;
+
+		Text = FString::Printf(TEXT("Local Lights sampling atlas: %d"), EffectiveLocalLightSlotArray.Num());
+		Canvas.DrawShadowedString(DrawPosX, DrawPosY, *Text, GEngine->GetLargeFont(), FLinearColor::White);
+		DrawPosY += 40.0f;
+
+		uint32 LightFunctionAtlasSlotIndex = 0;
+		for (auto& AtlasSlot : EffectiveLightFunctionSlotArray)
+		{
+			const FMaterialRenderProxy* LightFunctionMaterial = AtlasSlot.LightFunctionMaterial;
+
+			const FString& MaterialName = LightFunctionMaterial->GetMaterialName();
+			uint32 LFMaterialUniqueID = LightFunctionMaterial->GetMaterialInterface() ? MurmurFinalize32(MurmurFinalize32(LightFunctionMaterial->GetMaterialInterface()->GetUniqueID())) : 0xFFFFFFFF;
+
+			FLinearColor MaterialColor(FColor(LFMaterialUniqueID & 0xFF, (LFMaterialUniqueID >> 8) & 0xFF, (LFMaterialUniqueID >> 16) & 0xFF));
+
+			uint32 LightCountUsingThisMaterial = 0;
+			for (auto& LocalLight : EffectiveLocalLightSlotArray)
+			{
+				LightCountUsingThisMaterial += LocalLight.LightFunctionAtlasSlotIndex == LightFunctionAtlasSlotIndex ? 1 : 0;
+			}
+
+			// Draw the light function material
+			Canvas.DrawTile(DrawPosX - 10.0f, DrawPosY, 15.0f, 15.0f, 0.0f, 0.0f, 1.0f, 1.0f, MaterialColor, nullptr, false);
+
+			Text = FString::Printf(TEXT("%3d lights - %s"), LightCountUsingThisMaterial , *MaterialName);
+			Canvas.DrawShadowedString(DrawPosX, DrawPosY, *Text, GEngine->GetLargeFont(), MaterialColor);
+
+			// Draw a line around the corresponding atlas tile
+			auto OutLineAtlasSlot = [&](float X0, float Y0, float X1, float Y1)
+			{
+				FCanvasLineItem LineItem(
+					FIntPoint(DstPoint.X + X0 * DisplayResolutionRatio, DstPoint.Y + Y0 * DisplayResolutionRatio),
+					FIntPoint(DstPoint.X + X1 * DisplayResolutionRatio, DstPoint.Y + Y1 * DisplayResolutionRatio));
+				LineItem.LineThickness = 2.0f;
+				LineItem.SetColor(MaterialColor);
+				Canvas.DrawItem(LineItem);
+			};
+			OutLineAtlasSlot(AtlasSlot.Min.X + 1, AtlasSlot.Min.Y + 1, AtlasSlot.Max.X - 1, AtlasSlot.Min.Y + 1);
+			OutLineAtlasSlot(AtlasSlot.Max.X - 1, AtlasSlot.Min.Y + 1, AtlasSlot.Max.X - 1, AtlasSlot.Max.Y - 1);
+			OutLineAtlasSlot(AtlasSlot.Max.X - 1, AtlasSlot.Max.Y - 1, AtlasSlot.Min.X + 1, AtlasSlot.Max.Y - 1);
+			OutLineAtlasSlot(AtlasSlot.Min.X + 1, AtlasSlot.Max.Y - 1, AtlasSlot.Min.X + 1, AtlasSlot.Min.Y + 1);
+
+			DrawPosY += 20.0f;
+			LightFunctionAtlasSlotIndex++;
+		}
+	});
+
+#endif // WITH_EDITOR
+
+	return MoveTemp(ScreenPassSceneColor);
+}
diff --git a/Engine/Source/Runtime/Renderer/Private/LightFunctionAtlas.h b/Engine/Source/Runtime/Renderer/Private/LightFunctionAtlas.h
--- a/Engine/Source/Runtime/Renderer/Private/LightFunctionAtlas.h
+++ b/Engine/Source/Runtime/Renderer/Private/LightFunctionAtlas.h
@@ -15,20 +15,33 @@ class FRDGBuilder;
 class FLightSceneInfo;
 class FRDGTexture;
 
+struct FScreenPassTexture;
 struct FSortedLightSetSceneInfo;
 struct IPooledRenderTarget;
+struct FLightFunctionAtlas;
 
 
 
 struct FLightFunctionAtlasSceneData
 {
-	bool bLightFunctionAtlasEnabled = false;
-	bool bVolumetricFogUsesLightFunctionAtlas = false;
-	bool bDeferredlightingUsesLightFunctionAtlas = false;
+	void SetData(FLightFunctionAtlas* InLightFunctionAtlas, bool bInLightFunctionAtlasEnabled, bool bInVolumetricFogUsesLightFunctionAtlas, bool bInDeferredlightingUsesLightFunctionAtlas)
+	{
+		LightFunctionAtlas = InLightFunctionAtlas;
+		bLightFunctionAtlasEnabled = bInLightFunctionAtlasEnabled;
+		bVolumetricFogUsesLightFunctionAtlas = bInVolumetricFogUsesLightFunctionAtlas;
+		bDeferredlightingUsesLightFunctionAtlas = bInDeferredlightingUsesLightFunctionAtlas;
+	}
 
+	const FLightFunctionAtlas* GetLightFunctionAtlas()  const { return LightFunctionAtlas; }
 	bool GetLightFunctionAtlasEnabled()					const { return bLightFunctionAtlasEnabled; }
 	bool GetVolumetricFogUsesLightFunctionAtlas()		const { return bVolumetricFogUsesLightFunctionAtlas; }
 	bool GetDeferredlightingUsesLightFunctionAtlas()	const { return bDeferredlightingUsesLightFunctionAtlas; }
+
+private:
+	FLightFunctionAtlas* LightFunctionAtlas = nullptr;
+	bool bLightFunctionAtlasEnabled = false;
+	bool bVolumetricFogUsesLightFunctionAtlas = false;
+	bool bDeferredlightingUsesLightFunctionAtlas = false;
 };
 
 struct FLightFunctionAtlasViewData
@@ -36,9 +49,9 @@ struct FLightFunctionAtlasViewData
 	FLightFunctionAtlasViewData() : SceneData(nullptr) {}
 	FLightFunctionAtlasViewData(FLightFunctionAtlasSceneData* InSceneData) : SceneData(InSceneData) {}
 
-	bool GetLightFunctionAtlasEnabled()					const { return SceneData ? SceneData->bLightFunctionAtlasEnabled				: false; }
-	bool GetVolumetricFogUsesLightFunctionAtlas()		const { return SceneData ? SceneData->bVolumetricFogUsesLightFunctionAtlas	: false; }
-	bool GetDeferredlightingUsesLightFunctionAtlas()	const { return SceneData ? SceneData->bDeferredlightingUsesLightFunctionAtlas	: false; }
+	bool GetLightFunctionAtlasEnabled()					const { return SceneData ? SceneData->GetLightFunctionAtlasEnabled()				: false; }
+	bool GetVolumetricFogUsesLightFunctionAtlas()		const { return SceneData ? SceneData->GetVolumetricFogUsesLightFunctionAtlas()		: false; }
+	bool GetDeferredlightingUsesLightFunctionAtlas()	const { return SceneData ? SceneData->GetDeferredlightingUsesLightFunctionAtlas()	: false; }
 
 private:
 	FLightFunctionAtlasSceneData* SceneData;
@@ -112,6 +125,9 @@ struct FLightFunctionAtlas
 
 	void RenderLightFunctionAtlas(FRDGBuilder& GraphBuilder, TArray<FViewInfo>& Views);
 
+
+	FScreenPassTexture AddDebugVisualizationPasses(FRDGBuilder& GraphBuilder, const FViewInfo& View, FScreenPassTexture& ScreenPassSceneColor)  const;
+
 	TRDGUniformBufferRef<FLightFunctionAtlasGlobalParameters> GetLightFunctionAtlasGlobalParameters(uint32 ViewIndex, FRDGBuilder& GraphBuilder);
 	TRDGUniformBufferRef<FLightFunctionAtlasGlobalParameters> GetDefaultLightFunctionAtlasGlobalParameters(FRDGBuilder& GraphBuilder);
 
diff --git a/Engine/Source/Runtime/Renderer/Private/PostProcess/PostProcessing.cpp b/Engine/Source/Runtime/Renderer/Private/PostProcess/PostProcessing.cpp
--- a/Engine/Source/Runtime/Renderer/Private/PostProcess/PostProcessing.cpp
+++ b/Engine/Source/Runtime/Renderer/Private/PostProcess/PostProcessing.cpp
@@ -375,6 +375,7 @@ void AddPostProcessingPasses(
 		VisualizeSubstrate,
 		VisualizeLightGrid,
 		VisualizeSkyAtmosphere,
+		VisualizeLightFunctionAtlas,
 		VisualizeLevelInstance,
 		SelectionOutline,
 		EditorPrimitive,
@@ -427,6 +428,7 @@ void AddPostProcessingPasses(
 		TEXT("VisualizeSubstrate"),
 		TEXT("VisualizeLightGrid"),
 		TEXT("VisualizeSkyAtmosphere"),
+		TEXT("VisualizeLightFunctionAtlas"),
 		TEXT("VisualizeLevelInstance"),
 		TEXT("SelectionOutline"),
 		TEXT("EditorPrimitive"),
@@ -465,12 +467,14 @@ void AddPostProcessingPasses(
 	PassSequence.SetEnabled(EPass::VisualizeLightGrid, ShouldVisualizeLightGrid());
 
 #if WITH_EDITOR
-	PassSequence.SetEnabled(EPass::VisualizeSkyAtmosphere, Scene && View.Family && View.Family->EngineShowFlags.VisualizeSkyAtmosphere&& ShouldRenderSkyAtmosphereDebugPasses(Scene, View.Family->EngineShowFlags));
+	PassSequence.SetEnabled(EPass::VisualizeSkyAtmosphere, Scene&& View.Family && View.Family->EngineShowFlags.VisualizeSkyAtmosphere && ShouldRenderSkyAtmosphereDebugPasses(Scene, View.Family->EngineShowFlags));
+	PassSequence.SetEnabled(EPass::VisualizeLightFunctionAtlas, Scene && Scene->LightFunctionAtlasSceneData.GetLightFunctionAtlasEnabled() && View.Family && View.Family->EngineShowFlags.VisualizeLightFunctionAtlas);
 	PassSequence.SetEnabled(EPass::VisualizeLevelInstance, GIsEditor && EngineShowFlags.EditingLevelInstance && EngineShowFlags.VisualizeLevelInstanceEditing && !bVisualizeHDR);
 	PassSequence.SetEnabled(EPass::SelectionOutline, GIsEditor && EngineShowFlags.Selection && EngineShowFlags.SelectionOutline && !EngineShowFlags.Wireframe && !bVisualizeHDR && !IStereoRendering::IsStereoEyeView(View));
 	PassSequence.SetEnabled(EPass::EditorPrimitive, FSceneRenderer::ShouldCompositeEditorPrimitives(View));
 #else
 	PassSequence.SetEnabled(EPass::VisualizeSkyAtmosphere, false);
+	PassSequence.SetEnabled(EPass::VisualizeLightFunctionAtlas, false);
 	PassSequence.SetEnabled(EPass::VisualizeLevelInstance, false);
 	PassSequence.SetEnabled(EPass::SelectionOutline, false);
 	PassSequence.SetEnabled(EPass::EditorPrimitive, false);
@@ -1363,6 +1367,16 @@ void AddPostProcessingPasses(
 		SceneColor = AddSkyAtmosphereDebugPasses(GraphBuilder, Scene, *View.Family, View, SceneColor);
 	}
 
+	if (PassSequence.IsEnabled(EPass::VisualizeLightFunctionAtlas))
+	{
+		FScreenPassRenderTarget OverrideOutput;
+		PassSequence.AcceptOverrideIfLastPass(EPass::VisualizeLightFunctionAtlas, OverrideOutput);
+		if (Scene->LightFunctionAtlasSceneData.GetLightFunctionAtlas())
+		{
+			SceneColor = Scene->LightFunctionAtlasSceneData.GetLightFunctionAtlas()->AddDebugVisualizationPasses(GraphBuilder, View, SceneColor);
+		}
+	}
+
 	if (PassSequence.IsEnabled(EPass::VisualizeLevelInstance))
 	{
 		FVisualizeLevelInstanceInputs PassInputs;
diff --git a/Engine/Source/Runtime/Renderer/Private/ScreenPass.cpp b/Engine/Source/Runtime/Renderer/Private/ScreenPass.cpp
--- a/Engine/Source/Runtime/Renderer/Private/ScreenPass.cpp
+++ b/Engine/Source/Runtime/Renderer/Private/ScreenPass.cpp
@@ -232,31 +232,32 @@ void AddDrawTexturePass(
 	FRDGTextureRef InputTexture,
 	FRDGTextureRef OutputTexture,
 	FIntPoint InputPosition,
+	FIntPoint InputSize,
 	FIntPoint OutputPosition,
-	FIntPoint Size)
+	FIntPoint OutputSize)
 {
 	const FRDGTextureDesc& InputDesc = InputTexture->Desc;
 	const FRDGTextureDesc& OutputDesc = OutputTexture->Desc;
 
-	// Use a hardware copy if formats match.
-	if (InputDesc.Format == OutputDesc.Format)
+	// Use a hardware copy if formats and sizes match.
+	if (InputDesc.Format == OutputDesc.Format && InputSize == OutputSize)
 	{
-		return AddCopyTexturePass(GraphBuilder, InputTexture, OutputTexture, InputPosition, OutputPosition, Size);
+		return AddCopyTexturePass(GraphBuilder, InputTexture, OutputTexture, InputPosition, OutputPosition, InputSize);
 	}
 
-	if (Size == FIntPoint::ZeroValue)
+	if (InputSize == FIntPoint::ZeroValue)
 	{
 		// Copy entire input texture to output texture.
-		Size = InputTexture->Desc.Extent;
+		InputSize = InputTexture->Desc.Extent;
 	}
 
 	// Don't prime color data if the whole texture is being overwritten.
-	const ERenderTargetLoadAction LoadAction = (OutputPosition == FIntPoint::ZeroValue && Size == OutputDesc.Extent)
+	const ERenderTargetLoadAction LoadAction = (OutputPosition == FIntPoint::ZeroValue && InputSize == OutputDesc.Extent)
 		? ERenderTargetLoadAction::ENoAction
 		: ERenderTargetLoadAction::ELoad;
 
-	const FScreenPassTextureViewport InputViewport(InputDesc.Extent, FIntRect(InputPosition, InputPosition + Size));
-	const FScreenPassTextureViewport OutputViewport(OutputDesc.Extent, FIntRect(OutputPosition, OutputPosition + Size));
+	const FScreenPassTextureViewport InputViewport(InputDesc.Extent, FIntRect(InputPosition, InputPosition + InputSize));
+	const FScreenPassTextureViewport OutputViewport(OutputDesc.Extent, FIntRect(OutputPosition, OutputPosition + OutputSize));
 
 	TShaderMapRef<FCopyRectPS> PixelShader(static_cast<const FViewInfo&>(View).ShaderMap);
 
@@ -268,6 +269,26 @@ void AddDrawTexturePass(
 	AddDrawScreenPass(GraphBuilder, RDG_EVENT_NAME("DrawTexture"), View, OutputViewport, InputViewport, PixelShader, Parameters);
 }
 
+void AddDrawTexturePass(
+	FRDGBuilder& GraphBuilder,
+	const FSceneView& View,
+	FRDGTextureRef InputTexture,
+	FRDGTextureRef OutputTexture,
+	FIntPoint InputPosition,
+	FIntPoint OutputPosition,
+	FIntPoint Size)
+{
+	AddDrawTexturePass(
+		GraphBuilder,
+		View,
+		InputTexture,
+		OutputTexture,
+		InputPosition,
+		Size,
+		OutputPosition,
+		Size);
+}
+
 void AddDrawTexturePass(
 	FRDGBuilder& GraphBuilder,
 	const FSceneView& View,
diff --git a/Engine/Source/Runtime/Renderer/Public/ScreenPass.h b/Engine/Source/Runtime/Renderer/Public/ScreenPass.h
--- a/Engine/Source/Runtime/Renderer/Public/ScreenPass.h
+++ b/Engine/Source/Runtime/Renderer/Public/ScreenPass.h
@@ -638,6 +638,16 @@ void RENDERER_API AddDrawTexturePass(
 	FIntPoint OutputPosition = FIntPoint::ZeroValue,
 	FIntPoint Size = FIntPoint::ZeroValue);
 
+void RENDERER_API AddDrawTexturePass(
+	FRDGBuilder& GraphBuilder,
+	const FSceneView& View,
+	FRDGTextureRef InputTexture,
+	FRDGTextureRef OutputTexture,
+	FIntPoint InputPosition,
+	FIntPoint InputSize,
+	FIntPoint OutputPosition,
+	FIntPoint OutputSize);
+
 /** Helper variant which takes a shared viewport instead of unique input / output positions. */
 FORCEINLINE void AddDrawTexturePass(
 	FRDGBuilder& GraphBuilder,
