Changed GPUScene back to always perform a full UpdateGPULights as the existing light change tracking does not catch all updates (notably rect light textures). Uses the scene frame number to avoid redundant uploads in a single frame.

#rnx

diff --git a/Engine/Source/Runtime/Renderer/Private/GPUScene.cpp b/Engine/Source/Runtime/Renderer/Private/GPUScene.cpp
--- a/Engine/Source/Runtime/Renderer/Private/GPUScene.cpp
+++ b/Engine/Source/Runtime/Renderer/Private/GPUScene.cpp
@@ -2170,8 +2170,13 @@ void FGPUScene::OnPostSceneUpdate(FRDGBuilder& GraphBuilder, const FScenePostUpd
 
 void FGPUScene::OnPostLightSceneInfoUpdate(FRDGBuilder& GraphBuilder, const FLightSceneChangeSet& LightsPostUpdateData)
 {
-	const bool bAnythingChanged = !(LightsPostUpdateData.AddedLightIds.IsEmpty() && LightsPostUpdateData.RemovedLightIds.IsEmpty() && LightsPostUpdateData.TransformUpdatedLightIds.IsEmpty() && LightsPostUpdateData.ColorUpdatedLightIds.IsEmpty());
-	if (bAnythingChanged || !LightDataBuffer.IsValid())
+	// TODO: implement proper dirty tracking such that we can actually do incremental updates without loosing information
+	// const bool bAnythingChanged = !(LightsPostUpdateData.AddedLightIds.IsEmpty() && LightsPostUpdateData.RemovedLightIds.IsEmpty() && LightsPostUpdateData.TransformUpdatedLightIds.IsEmpty() && LightsPostUpdateData.ColorUpdatedLightIds.IsEmpty());
+	
+	// SceneFrameNumber is updated in UpdateInternal so if it is the same, we have already uploaded the lights this "frame" - this is not 100% robust since 
+	// if UpdateAllPrimitiveSceneInfos is called multiple times without updating the FScene frame number it will skip light uploads. The real solution is 
+	// to implement functional change tracking for lights so we can actually know when there are relevant changes.
+	if (SceneFrameNumber != Scene.GetFrameNumberRenderThread() || !LightDataBuffer.IsValid())
 	{
 		UpdateGPULights(GraphBuilder, UE::Tasks::FTask{});
 	}
