* Change voxel allocation to be batch better by avoiding pass dependency
* Made indirect args building for all the macrogroup to run in a single pass.

diff --git a/Engine/Shaders/Private/HairStrands/HairStrandsVoxelRasterCompute.usf b/Engine/Shaders/Private/HairStrands/HairStrandsVoxelRasterCompute.usf
--- a/Engine/Shaders/Private/HairStrands/HairStrandsVoxelRasterCompute.usf
+++ b/Engine/Shaders/Private/HairStrands/HairStrandsVoxelRasterCompute.usf
@@ -1,6 +1,6 @@
 // Copyright Epic Games, Inc. All Rights Reserved.
 
-#define HAIR_STRANDS_PARAMETERS 1
+#define HAIR_STRANDS_PARAMETERS 1 
 
 #include "../Common.ush"
 #include "../Common.ush"
@@ -16,6 +16,9 @@
 #include "HairStrandsVoxelPageCommon.ush" 
 #include "HairStrandsDeepShadowCommonStruct.ush"
 
+#define GLOBAL_PAGE_COUNTER_INDEX 0
+#define GROUPS_PAGE_COUNTER_INDEX 1
+
 ///////////////////////////////////////////////////////////////////////////////////////////////////
 // Page allocation
 
@@ -36,12 +39,12 @@ float4				CPU_TranslatedWorldMaxAABB[MAX_HAIR_MACROGROUP_COUNT];
 int4				CPU_PageIndexResolution[MAX_HAIR_MACROGROUP_COUNT];
 uint				CPU_bUseCPUData;
 
-Buffer<float>								GPUVoxelWorldSize;
-Buffer<float>								MacroGroupVoxelSizeBuffer;
-Buffer<int>									MacroGroupAABBBuffer;
-RWBuffer<int>								MacroGroupVoxelAlignedAABBBuffer;
-RWBuffer<uint4>								OutPageIndexResolutionAndOffsetBuffer;
-RWBuffer<uint>								OutPageIndexAllocationIndirectBufferArgs;
+Buffer<float>		GPUVoxelWorldSize;
+Buffer<float>		MacroGroupVoxelSizeBuffer;
+Buffer<int>			MacroGroupAABBBuffer;
+RWBuffer<int>		MacroGroupVoxelAlignedAABBBuffer;
+RWBuffer<uint4>		OutPageIndexResolutionAndOffsetBuffer;
+RWBuffer<uint>		OutPageIndexAllocationIndirectBufferArgs;
 
 
 #if GROUP_SIZE != MAX_HAIR_MACROGROUP_COUNT
@@ -310,20 +313,21 @@ uint3 CPU_PageIndexResolution;
 Buffer<uint4> PageIndexResolutionAndOffsetBuffer;
 Buffer<int> IndirectBufferArgs;
 
-RWBuffer<uint> PageIndexGlobalCounter;
-RWBuffer<uint> PageIndexBuffer;
-RWBuffer<uint> PageToPageIndexBuffer;
-RWBuffer<uint4> PageIndexCoordBuffer;
+RWBuffer<uint> RWPageIndexGlobalCounter;
+RWBuffer<uint> RWPageIndexBuffer;
+RWBuffer<uint> RWPageToPageIndexBuffer;
+RWBuffer<uint4> RWPageIndexCoordBuffer;
 
 groupshared uint LocalCounter;
-groupshared uint GroupBase;
+groupshared uint GroupBase[2];
 
 [numthreads(GROUP_SIZE, 1, 1)]
 void AllocateCS(uint GroupIndex : SV_GroupIndex, uint3 DispatchThreadId : SV_DispatchThreadID)
 {
 	if (GroupIndex == 0)
 	{
-		GroupBase = 0;
+		GroupBase[0] = 0;
+		GroupBase[1] = 0;
 		LocalCounter = 0;
 	}
 	GroupMemoryBarrierWithGroupSync();
@@ -343,7 +347,7 @@ void AllocateCS(uint GroupIndex : SV_GroupIndex, uint3 DispatchThreadId : SV_Dis
 	bool bIsValid = false;
 	if (DispatchThreadId.x < PageIndexCount)
 	{
-		bIsValid = PageIndexBuffer[GridIndex] > 0;
+		bIsValid = RWPageIndexBuffer[GridIndex] > 0;
 	}
 
 	uint Offset = 0;
@@ -355,33 +359,36 @@ void AllocateCS(uint GroupIndex : SV_GroupIndex, uint3 DispatchThreadId : SV_Dis
 
 	if (GroupIndex == 0)
 	{
-		InterlockedAdd(PageIndexGlobalCounter[0], LocalCounter, GroupBase);
+		// * Add page count to global counter for global tacking
+		// * Add page count to the group for per group work
+		InterlockedAdd(RWPageIndexGlobalCounter[GLOBAL_PAGE_COUNTER_INDEX], LocalCounter, GroupBase[0]);
+		InterlockedAdd(RWPageIndexGlobalCounter[GROUPS_PAGE_COUNTER_INDEX + MacroGroupId], LocalCounter, GroupBase[1]);
 	}
 	GroupMemoryBarrierWithGroupSync();
 
 	if (bIsValid)
 	{
-		const uint PageIndex = GroupBase + Offset;
-		const bool bIsAllocationValid = PageIndex < PageCount;
+		const uint PageIndex0 = GroupBase[0] + Offset; // Global page index
+		const uint PageIndex1 = GroupBase[1] + Offset; // Group  page index
+		const bool bIsAllocationValid = PageIndex0 < PageCount;
 
-		PageIndexBuffer[GridIndex] = bIsAllocationValid ? PageIndex : INVALID_VOXEL_PAGE_INDEX;
+		RWPageIndexBuffer[GridIndex] = bIsAllocationValid ? PageIndex0 : INVALID_VOXEL_PAGE_INDEX;
 		if (bIsAllocationValid)
 		{
-			PageToPageIndexBuffer[PageIndex] = GridIndex;
+			RWPageToPageIndexBuffer[PageIndex0] = GridIndex;
 		}
 
 		// Output the coordinates of the allocated page for indirect dispatch usage
 		// If the allocated failed (run out of page), then we mark the IndexCoord with a invalid GroupID
 		const uint LinearIndex = DispatchThreadId.x;
 		const uint3 PageIndexCoord = IndexToCoord(LinearIndex, PageIndexResolution);
-		const uint IndexRelativeToMacroGroupIndexOffset = PageIndex - PageIndexGlobalCounter[1];
-		PageIndexCoordBuffer[PageIndexOffset + IndexRelativeToMacroGroupIndexOffset] = uint4(PageIndexCoord, bIsAllocationValid ? MacroGroupId : INVALID_MACRO_GROUP_ID);
+		RWPageIndexCoordBuffer[PageIndexOffset + PageIndex1] = uint4(PageIndexCoord, bIsAllocationValid ? MacroGroupId : INVALID_MACRO_GROUP_ID);
 	}
 	// Mark page index as invalid
 	// Insure that even if write more (due to larger dispatch count that needed), we do not stomp other instance group page index
 	else if (DispatchThreadId.x < PageIndexCount)
 	{
-		PageIndexBuffer[GridIndex] = INVALID_VOXEL_PAGE_INDEX;
+		RWPageIndexBuffer[GridIndex] = INVALID_VOXEL_PAGE_INDEX;
 	}
 }
 #endif
@@ -438,54 +445,35 @@ void AddDescCS(uint GroupIndex : SV_GroupIndex, uint3 DispatchThreadId : SV_Disp
 #if SHADER_ADDINDIRECTBUFFER
 
 uint PageResolution;
-uint MacroGroupId;
+uint MacroGroupCount;
 int3 IndirectGroupSize;
 
-RWBuffer<uint> OutPageIndexGlobalCounter;
+Buffer<uint> PageIndexGlobalCounter;
 RWBuffer<uint> OutIndirectArgsBuffer;
-RWBuffer<uint> OutTotalRequestedPageAllocationBuffer;
 
-[numthreads(1, 1, 1)]
-void AddIndirectBufferCS(uint GroupIndex : SV_GroupIndex, uint3 DispatchThreadId : SV_DispatchThreadID)
+void WriteArgs(uint WriteIndex, uint AllocatedPageCount)
 {
-	const uint TotalAllocatedPageCount = OutPageIndexGlobalCounter[0];
-	const uint PrevTotalAllocatedPageCount = OutPageIndexGlobalCounter[1];
-	const uint AllocatedPageCount = TotalAllocatedPageCount - PrevTotalAllocatedPageCount;
 	const uint VoxelCountPerPage = PageResolution * PageResolution * PageResolution;
-
 	const uint DispatchCountX = DivideAndRoundUp(VoxelCountPerPage, IndirectGroupSize.x);
 	const uint DispatchCountZ = DivideAndRoundUp(AllocatedPageCount, IndirectGroupSize.z);
-	WriteDispatchIndirectArgs(OutIndirectArgsBuffer, MacroGroupId, DispatchCountX, 1, DispatchCountZ);
-
-	// Copy current total to previous total
-	OutPageIndexGlobalCounter[1] = TotalAllocatedPageCount;
-
-	// The value will be written several times, but only the value written during the last macro group will give the final page count allocated/requested
-	OutTotalRequestedPageAllocationBuffer[0] = TotalAllocatedPageCount;
+	WriteDispatchIndirectArgs(OutIndirectArgsBuffer, WriteIndex, DispatchCountX, 1, DispatchCountZ);
 }
 
-#endif
-
-///////////////////////////////////////////////////////////////////////////////////////////////////
-// Prepare indirect args for indirect clear
-
-#if SHADER_INDPAGECLEARBUFFERGEN
-
-uint PageResolution;
-int3 GroupSize;
-Buffer<uint> PageIndexGlobalCounter;
-RWBuffer<uint> OutIndirectArgsBuffer;
-
-[numthreads(1, 1, 1)]
-void VoxelIndPageClearBufferGenCS(uint GroupIndex : SV_GroupIndex, uint3 DispatchThreadId : SV_DispatchThreadID)
+[numthreads(GROUP_SIZE, 1, 1)]
+void AddIndirectBufferCS(uint GroupIndex : SV_GroupIndex, uint3 DispatchThreadId : SV_DispatchThreadID)
 {
-	const uint TotalAllocatedPageCount = PageIndexGlobalCounter[0];
-	const uint VoxelCountPerPage = PageResolution * PageResolution * PageResolution;
-
-	const uint DispatchX = DivideAndRoundUp(VoxelCountPerPage, GroupSize.x);
-	const uint DispatchZ = DivideAndRoundUp(TotalAllocatedPageCount, GroupSize.z);
+	// Total pages allocated across *all* macro groups
+	if (DispatchThreadId.x == 0)
+	{
+		WriteArgs(GLOBAL_PAGE_COUNTER_INDEX, PageIndexGlobalCounter[GLOBAL_PAGE_COUNTER_INDEX]);
+	}
 
-	WriteDispatchIndirectArgs(OutIndirectArgsBuffer, 0, DispatchX, 1, DispatchZ);
+	// Pages allocated for a particular macro group
+	const uint MacroGroupId = DispatchThreadId.x;
+	if (MacroGroupId < MacroGroupCount)
+	{
+		WriteArgs(GROUPS_PAGE_COUNTER_INDEX + MacroGroupId, PageIndexGlobalCounter[GROUPS_PAGE_COUNTER_INDEX + MacroGroupId]);
+	}
 }
 
 #endif
@@ -505,7 +493,7 @@ RWTexture3D<uint>	OutPageTexture;
 [numthreads(GROUP_SIZE_X, 1, GROUP_SIZE_Z)]
 void VoxelIndPageClearCS(uint3 DispatchThreadId : SV_DispatchThreadID, uint3 GroupId : SV_GroupID, uint3 GroupThreadId : SV_GroupThreadID)
 {
-	const uint TotalAllocatedPageCount = PageIndexGlobalCounter[0];
+	const uint TotalAllocatedPageCount = PageIndexGlobalCounter[GLOBAL_PAGE_COUNTER_INDEX];
 	const uint VoxelCountPerPage = VirtualVoxelParams_PageResolution * VirtualVoxelParams_PageResolution * VirtualVoxelParams_PageResolution;
 
 	const uint LinearVoxelCoord = DispatchThreadId.x;
@@ -541,7 +529,7 @@ float CPUMinVoxelWorldSize;
 float AdaptiveCorrectionThreshold;
 float AdaptiveCorrectionSpeed;
 
-Buffer<uint> TotalRequestedPageAllocationBuffer;
+Buffer<uint> PageIndexGlobalCounter;
 Buffer<float> CurrGPUMinVoxelWorldSize;
 RWBuffer<float> NextGPUMinVoxelWorldSize;
 
@@ -564,7 +552,7 @@ void FeedbackCS(uint3 DispatchThreadId : SV_DispatchThreadID)
 	//    MaxPage          MaxRes^3
 
 	// Ratio used for predicting voxel size increase
-	const uint GPUAllocatedPageCount = TotalRequestedPageAllocationBuffer[0];
+	const uint GPUAllocatedPageCount = PageIndexGlobalCounter[GLOBAL_PAGE_COUNTER_INDEX];
 	const float VolumeRatio = float(GPUAllocatedPageCount) / float(CPUAllocatedPageCount);
 	const float LinearRatio = pow(VolumeRatio, 1.f / 3.f);
 
diff --git a/Engine/Source/Runtime/Renderer/Private/HairStrands/HairStrandsVoxelization.cpp b/Engine/Source/Runtime/Renderer/Private/HairStrands/HairStrandsVoxelization.cpp
--- a/Engine/Source/Runtime/Renderer/Private/HairStrands/HairStrandsVoxelization.cpp
+++ b/Engine/Source/Runtime/Renderer/Private/HairStrands/HairStrandsVoxelization.cpp
@@ -322,10 +322,10 @@ class FVoxelAllocateVoxelPageCS : public FGlobalShader
 		SHADER_PARAMETER(uint32, CPU_PageIndexOffset)
 		RDG_BUFFER_ACCESS(IndirectBufferArgs, ERHIAccess::IndirectArgs)
 		SHADER_PARAMETER_RDG_BUFFER_SRV(Buffer, PageIndexResolutionAndOffsetBuffer)
-		SHADER_PARAMETER_RDG_BUFFER_UAV(RWBuffer, PageIndexGlobalCounter)
-		SHADER_PARAMETER_RDG_BUFFER_UAV(RWBuffer, PageIndexBuffer)
-		SHADER_PARAMETER_RDG_BUFFER_UAV(RWBuffer, PageToPageIndexBuffer)
-		SHADER_PARAMETER_RDG_BUFFER_UAV(RWBuffer, PageIndexCoordBuffer)
+		SHADER_PARAMETER_RDG_BUFFER_UAV(RWBuffer, RWPageIndexGlobalCounter)
+		SHADER_PARAMETER_RDG_BUFFER_UAV(RWBuffer, RWPageIndexBuffer)
+		SHADER_PARAMETER_RDG_BUFFER_UAV(RWBuffer, RWPageToPageIndexBuffer)
+		SHADER_PARAMETER_RDG_BUFFER_UAV(RWBuffer, RWPageIndexCoordBuffer)
 	END_SHADER_PARAMETER_STRUCT()
 
 public:
@@ -379,43 +379,21 @@ class FVoxelAddIndirectBufferCS : public FGlobalShader
 
 	BEGIN_SHADER_PARAMETER_STRUCT(FParameters, )
 		SHADER_PARAMETER(uint32, MacroGroupId)
+		SHADER_PARAMETER(uint32, MacroGroupCount)
 		SHADER_PARAMETER(uint32, PageResolution)
 		SHADER_PARAMETER(FIntVector, IndirectGroupSize)
-		SHADER_PARAMETER_RDG_BUFFER_UAV(RWBuffer, OutPageIndexGlobalCounter)
-		SHADER_PARAMETER_RDG_BUFFER_UAV(RWBuffer, OutIndirectArgsBuffer)
-		SHADER_PARAMETER_RDG_BUFFER_UAV(RWBuffer, OutTotalRequestedPageAllocationBuffer)
-	END_SHADER_PARAMETER_STRUCT()
-
-public:
-	static uint32 GetGroupSize() { return 1; }
-	static bool ShouldCompilePermutation(const FGlobalShaderPermutationParameters& Parameters) { return IsHairStrandsSupported(EHairStrandsShaderType::Strands, Parameters.Platform); }
-	static void ModifyCompilationEnvironment(const FGlobalShaderPermutationParameters& Parameters, FShaderCompilerEnvironment& OutEnvironment)
-	{
-		FGlobalShader::ModifyCompilationEnvironment(Parameters, OutEnvironment);
-		OutEnvironment.SetDefine(TEXT("SHADER_ADDINDIRECTBUFFER"), 1);
-	}
-};
-
-
-class FVoxelIndPageClearBufferGenCS : public FGlobalShader
-{
-	DECLARE_GLOBAL_SHADER(FVoxelIndPageClearBufferGenCS);
-	SHADER_USE_PARAMETER_STRUCT(FVoxelIndPageClearBufferGenCS, FGlobalShader);
-
-	BEGIN_SHADER_PARAMETER_STRUCT(FParameters, )
-		SHADER_PARAMETER(FIntVector, GroupSize)
-		SHADER_PARAMETER(uint32, PageResolution)
 		SHADER_PARAMETER_RDG_BUFFER_SRV(Buffer, PageIndexGlobalCounter)
 		SHADER_PARAMETER_RDG_BUFFER_UAV(RWBuffer, OutIndirectArgsBuffer)
 	END_SHADER_PARAMETER_STRUCT()
 
 public:
-	static uint32 GetGroupSize() { return 1; }
+	static uint32 GetGroupSize() { return 32; }
 	static bool ShouldCompilePermutation(const FGlobalShaderPermutationParameters& Parameters) { return IsHairStrandsSupported(EHairStrandsShaderType::Strands, Parameters.Platform); }
 	static void ModifyCompilationEnvironment(const FGlobalShaderPermutationParameters& Parameters, FShaderCompilerEnvironment& OutEnvironment)
 	{
 		FGlobalShader::ModifyCompilationEnvironment(Parameters, OutEnvironment);
-		OutEnvironment.SetDefine(TEXT("SHADER_INDPAGECLEARBUFFERGEN"), 1);
+		OutEnvironment.SetDefine(TEXT("SHADER_ADDINDIRECTBUFFER"), 1);
+		OutEnvironment.SetDefine(TEXT("GROUP_SIZE"), GetGroupSize());
 	}
 };
 
@@ -455,7 +433,7 @@ class FVoxelAdaptiveFeedbackCS : public FGlobalShader
 		SHADER_PARAMETER(float,  AdaptiveCorrectionThreshold)
 		SHADER_PARAMETER(float,  AdaptiveCorrectionSpeed)
 
-		SHADER_PARAMETER_RDG_BUFFER_SRV(Buffer, TotalRequestedPageAllocationBuffer)
+		SHADER_PARAMETER_RDG_BUFFER_SRV(Buffer, PageIndexGlobalCounter)
 		SHADER_PARAMETER_RDG_BUFFER_SRV(Buffer, CurrGPUMinVoxelWorldSize)
 		SHADER_PARAMETER_RDG_BUFFER_UAV(RWBuffer, NextGPUMinVoxelWorldSize)
 
@@ -477,7 +455,6 @@ IMPLEMENT_GLOBAL_SHADER(FVoxelAllocatePageIndexCS, "/Engine/Private/HairStrands/
 IMPLEMENT_GLOBAL_SHADER(FVoxelAllocateVoxelPageCS, "/Engine/Private/HairStrands/HairStrandsVoxelRasterCompute.usf", "AllocateCS", SF_Compute);
 IMPLEMENT_GLOBAL_SHADER(FVoxelAddNodeDescCS, "/Engine/Private/HairStrands/HairStrandsVoxelRasterCompute.usf", "AddDescCS", SF_Compute);
 IMPLEMENT_GLOBAL_SHADER(FVoxelAddIndirectBufferCS, "/Engine/Private/HairStrands/HairStrandsVoxelRasterCompute.usf", "AddIndirectBufferCS", SF_Compute);
-IMPLEMENT_GLOBAL_SHADER(FVoxelIndPageClearBufferGenCS, "/Engine/Private/HairStrands/HairStrandsVoxelRasterCompute.usf", "VoxelIndPageClearBufferGenCS", SF_Compute);
 IMPLEMENT_GLOBAL_SHADER(FVoxelIndPageClearCS, "/Engine/Private/HairStrands/HairStrandsVoxelRasterCompute.usf", "VoxelIndPageClearCS", SF_Compute);
 IMPLEMENT_GLOBAL_SHADER(FVoxelAdaptiveFeedbackCS, "/Engine/Private/HairStrands/HairStrandsVoxelRasterCompute.usf", "FeedbackCS", SF_Compute);
 
@@ -563,29 +540,20 @@ static void AddAllocateVoxelPagesPass(
 	
 	FRDGBufferRef PageIndexBuffer = GraphBuilder.CreateBuffer(FRDGBufferDesc::CreateBufferDesc(sizeof(uint32), OutTotalPageIndexCount), TEXT("Hair.PageIndexBuffer"));
 	FRDGBufferRef PageIndexCoordBuffer = GraphBuilder.CreateBuffer(FRDGBufferDesc::CreateBufferDesc(sizeof(uint32), OutTotalPageIndexCount), TEXT("Hair.PageIndexCoordBuffer"));
-	FRDGBufferRef PageIndexGlobalCounter = GraphBuilder.CreateBuffer(FRDGBufferDesc::CreateBufferDesc(sizeof(uint32), 2), TEXT("Hair.PageIndexGlobalCounter"));
+	FRDGBufferRef PageIndexGlobalCounter = GraphBuilder.CreateBuffer(FRDGBufferDesc::CreateBufferDesc(sizeof(uint32), MacroGroupCount + 1), TEXT("Hair.PageIndexGlobalCounter")); // First entry store global count
 	FRDGBufferRef NodeDescBuffer = GraphBuilder.CreateBuffer(FRDGBufferDesc::CreateStructuredDesc(sizeof(FPackedVirtualVoxelNodeDesc), MacroGroupCount), TEXT("Hair.VirtualVoxelNodeDescBuffer"));
-	FRDGBufferRef IndirectArgsBuffer = GraphBuilder.CreateBuffer(FRDGBufferDesc::CreateIndirectDesc<FRHIDispatchIndirectParameters>(MacroGroupCount), TEXT("Hair.VirtualVoxelIndirectArgsBuffer"));
+	FRDGBufferRef IndirectArgsBuffer = GraphBuilder.CreateBuffer(FRDGBufferDesc::CreateIndirectDesc<FRHIDispatchIndirectParameters>(MacroGroupCount+1), TEXT("Hair.VirtualVoxelIndirectArgsBuffer")); // First entry store global args
 
 	const uint32 TotalPageCount = PageCountResolution.X * PageCountResolution.Y * PageCountResolution.Z;
 	FRDGBufferRef PageToPageIndexBuffer = GraphBuilder.CreateBuffer(FRDGBufferDesc::CreateBufferDesc(sizeof(uint32), TotalPageCount), TEXT("Hair.PageToPageIndexBuffer"));
 
 	FRDGBufferRef ReadBackBuffer = GraphBuilder.CreateBuffer(FRDGBufferDesc::CreateBufferDesc(sizeof(uint32), 1), TEXT("ReadBackAllocations"));
 
-	FRDGBufferUAVRef PageIndexBufferUAV = GraphBuilder.CreateUAV(PageIndexBuffer, PF_R32_UINT);
-	FRDGBufferUAVRef PageIndexBufferUAVSkipBarrier = GraphBuilder.CreateUAV(PageIndexBuffer, PF_R32_UINT, ERDGUnorderedAccessViewFlags::SkipBarrier);
-	FRDGBufferUAVRef PageIndexGlobalCounterUAV = GraphBuilder.CreateUAV(PageIndexGlobalCounter, PF_R32_UINT);
-	
 	FRDGBufferRef PageIndexResolutionBuffer = GraphBuilder.CreateBuffer(FRDGBufferDesc::CreateBufferDesc(MacroGroupCount * 4 * sizeof(uint32), OutTotalPageIndexCount), TEXT("Hair.PageIndexResolutionBuffer"));
-	FRDGBufferRef PageIndexAllocationIndirectBufferArgs = GraphBuilder.CreateBuffer(FRDGBufferDesc::CreateIndirectDesc<FRHIDispatchIndirectParameters>(MacroGroupCount), TEXT("Hair.PageIndexAllocationIndirectBufferArgs"));
-
-	// Store the total requested page allocation (for feedback purpose)
-	FRDGBufferRef TotalRequestedPageAllocationBuffer = GraphBuilder.CreateBuffer(FRDGBufferDesc::CreateBufferDesc(sizeof(uint32), 1), TEXT("Hair.TotalRequestedPageAllocationBuffer"));
-	FRDGBufferUAVRef TotalRequestedPageAllocationBufferUAV = GraphBuilder.CreateUAV(TotalRequestedPageAllocationBuffer, PF_R32_UINT);
-	AddClearUAVPass(GraphBuilder, TotalRequestedPageAllocationBufferUAV, 0u);
+	FRDGBufferRef PageIndexAllocationIndirectBufferArgs = GraphBuilder.CreateBuffer(FRDGBufferDesc::CreateIndirectDesc<FRHIDispatchIndirectParameters>(MacroGroupCount+1), TEXT("Hair.PageIndexAllocationIndirectBufferArgs"));
 	
-	AddClearUAVPass(GraphBuilder, PageIndexBufferUAV, 0u);
-	AddClearUAVPass(GraphBuilder, PageIndexGlobalCounterUAV, 0u);
+	AddClearUAVPass(GraphBuilder, GraphBuilder.CreateUAV(PageIndexBuffer, PF_R32_UINT), 0u);
+	AddClearUAVPass(GraphBuilder, GraphBuilder.CreateUAV(PageIndexGlobalCounter, PF_R32_UINT), 0u);
 
 	const FRDGBufferSRVRef CurrGPUMinVoxelSizeSRV = GraphBuilder.CreateSRV(CurrGPUMinVoxelSize, PF_R16F);
 	const bool bAdaptiveVoxelEnable = IsHairStrandsAdaptiveVoxelAllocationEnable();
@@ -648,11 +616,12 @@ static void AddAllocateVoxelPagesPass(
 	FRDGBufferSRVRef PageIndexResolutionAndOffsetBufferSRV = GraphBuilder.CreateSRV(PageIndexResolutionBuffer, PF_R32G32B32A32_UINT);
 
 	// Mark valid page index
-	for (uint32 MacroGroupIt = 0; MacroGroupIt < MacroGroupCount; ++MacroGroupIt)
 	{
-		const FHairStrandsMacroGroupData& MacroGroup = MacroGroupDatas[MacroGroupIt];
-		FCPUMacroGroupAllocation& CPUAllocationDesc = CPUAllocationDescs[MacroGroupIt];
+		FRDGBufferUAVRef PageIndexBufferUAVSkipBarrier = GraphBuilder.CreateUAV(PageIndexBuffer, PF_R32_UINT, ERDGUnorderedAccessViewFlags::SkipBarrier);
+		for (uint32 MacroGroupIt = 0; MacroGroupIt < MacroGroupCount; ++MacroGroupIt)
 		{
+			const FHairStrandsMacroGroupData& MacroGroup = MacroGroupDatas[MacroGroupIt];
+			FCPUMacroGroupAllocation& CPUAllocationDesc = CPUAllocationDescs[MacroGroupIt];
 			for (const FHairStrandsMacroGroupData::PrimitiveInfo& PrimitiveInfo : MacroGroup.PrimitivesInfos)
 			{
 				FHairGroupPublicData* HairGroupData = PrimitiveInfo.PublicDataPtr;
@@ -709,64 +678,69 @@ static void AddAllocateVoxelPagesPass(
 		}
 	}
 
-	// Do this into a single pass by creatin a bufer with CPU AABB page res
-	FRDGBufferUAVRef NodeDescUAVSkipBarrier = GraphBuilder.CreateUAV(NodeDescBuffer, ERDGUnorderedAccessViewFlags::SkipBarrier);
-	for (uint32 MacroGroupIt = 0; MacroGroupIt < MacroGroupCount; ++MacroGroupIt)
+	// Fill in hair-macro-group information.
 	{
-		const FHairStrandsMacroGroupData& MacroGroup = MacroGroupDatas[MacroGroupIt];
-		FCPUMacroGroupAllocation& CPUAllocationDesc = CPUAllocationDescs[MacroGroupIt];
-
-		// Fill in hair-macro-group information.
-		// Note: This need to happen before the allocation as we copy the index global count. This global index is 
-		// used as an offset, and thus refers to the previous pass
+		FRDGBufferUAVRef NodeDescUAVSkipBarrier = GraphBuilder.CreateUAV(NodeDescBuffer, ERDGUnorderedAccessViewFlags::SkipBarrier);
+		for (uint32 MacroGroupIt = 0; MacroGroupIt < MacroGroupCount; ++MacroGroupIt)
 		{
-			check(MacroGroup.MacroGroupId < MacroGroupCount);
-
-			FVoxelAddNodeDescCS::FParameters* Parameters = GraphBuilder.AllocParameters<FVoxelAddNodeDescCS::FParameters>();
-			Parameters->MacroGroupId = MacroGroup.MacroGroupId;
-			Parameters->CPU_TranslatedWorldMinAABB = (FVector3f)CPUAllocationDesc.TranslatedMinAABB;
-			Parameters->CPU_TranslatedWorldMaxAABB = (FVector3f)CPUAllocationDesc.TranslatedMaxAABB;
-			Parameters->CPU_PageIndexResolution = CPUAllocationDesc.PageIndexResolution;
-			Parameters->CPU_PageIndexOffset = CPUAllocationDesc.PageIndexOffset;
-			Parameters->CPU_VoxelWorldSize = CPUMinVoxelWorldSize;
-			Parameters->GPU_VoxelWorldSize = CurrGPUMinVoxelSizeSRV;
-			Parameters->bUseCPUVoxelWorldSize = bAdaptiveVoxelEnable ? 0u : 1u;
-			Parameters->MacroGroupVoxelSizeBuffer = GraphBuilder.CreateSRV(MacroGroupResources.MacroGroupVoxelSizeBuffer, PF_R16F);
-			Parameters->OutNodeDescBuffer = NodeDescUAVSkipBarrier;
-
-			if (bIsGPUDriven)
+			const FHairStrandsMacroGroupData& MacroGroup = MacroGroupDatas[MacroGroupIt];
+			FCPUMacroGroupAllocation& CPUAllocationDesc = CPUAllocationDescs[MacroGroupIt];
+	
+			// Note: This need to happen before the allocation as we copy the index global count. This global index is 
+			// used as an offset, and thus refers to the previous pass
 			{
-				Parameters->MacroGroupVoxelAlignedAABBBuffer = GraphBuilder.CreateSRV(MacroGroupResources.MacroGroupVoxelAlignedAABBsBuffer, PF_R32_SINT);
-				Parameters->PageIndexResolutionAndOffsetBuffer = PageIndexResolutionAndOffsetBufferSRV;
+				check(MacroGroup.MacroGroupId < MacroGroupCount);
+	
+				FVoxelAddNodeDescCS::FParameters* Parameters = GraphBuilder.AllocParameters<FVoxelAddNodeDescCS::FParameters>();
+				Parameters->MacroGroupId = MacroGroup.MacroGroupId;
+				Parameters->CPU_TranslatedWorldMinAABB = (FVector3f)CPUAllocationDesc.TranslatedMinAABB;
+				Parameters->CPU_TranslatedWorldMaxAABB = (FVector3f)CPUAllocationDesc.TranslatedMaxAABB;
+				Parameters->CPU_PageIndexResolution = CPUAllocationDesc.PageIndexResolution;
+				Parameters->CPU_PageIndexOffset = CPUAllocationDesc.PageIndexOffset;
+				Parameters->CPU_VoxelWorldSize = CPUMinVoxelWorldSize;
+				Parameters->GPU_VoxelWorldSize = CurrGPUMinVoxelSizeSRV;
+				Parameters->bUseCPUVoxelWorldSize = bAdaptiveVoxelEnable ? 0u : 1u;
+				Parameters->MacroGroupVoxelSizeBuffer = GraphBuilder.CreateSRV(MacroGroupResources.MacroGroupVoxelSizeBuffer, PF_R16F);
+				Parameters->OutNodeDescBuffer = NodeDescUAVSkipBarrier;
+	
+				if (bIsGPUDriven)
+				{
+					Parameters->MacroGroupVoxelAlignedAABBBuffer = GraphBuilder.CreateSRV(MacroGroupResources.MacroGroupVoxelAlignedAABBsBuffer, PF_R32_SINT);
+					Parameters->PageIndexResolutionAndOffsetBuffer = PageIndexResolutionAndOffsetBufferSRV;
+				}
+	
+				FVoxelAddNodeDescCS::FPermutationDomain PermutationVector;
+				PermutationVector.Set<FVoxelAddNodeDescCS::FGPUDriven>(bIsGPUDriven ? 1 : 0);
+	
+				const FIntVector DispatchCount(1, 1, 1);
+				TShaderMapRef<FVoxelAddNodeDescCS> ComputeShader(View.ShaderMap, PermutationVector);
+				FComputeShaderUtils::AddPass(GraphBuilder, RDG_EVENT_NAME("HairStrands::AddNodeDesc"), ComputeShader, Parameters, DispatchCount);
 			}
-
-			FVoxelAddNodeDescCS::FPermutationDomain PermutationVector;
-			PermutationVector.Set<FVoxelAddNodeDescCS::FGPUDriven>(bIsGPUDriven ? 1 : 0);
-
-			const FIntVector DispatchCount(1, 1, 1);
-			TShaderMapRef<FVoxelAddNodeDescCS> ComputeShader(View.ShaderMap, PermutationVector);
-			FComputeShaderUtils::AddPass(GraphBuilder, RDG_EVENT_NAME("HairStrands::AddNodeDesc"), ComputeShader, Parameters, DispatchCount);
 		}
 	}
 
-	// Need to run serial, as PageIndexGlobalCounterUAV is used pass to pass 
-	for (uint32 MacroGroupIt = 0; MacroGroupIt < MacroGroupCount; ++MacroGroupIt)
+	// Allocate pages
 	{
-		const FHairStrandsMacroGroupData& MacroGroup = MacroGroupDatas[MacroGroupIt];
-		FCPUMacroGroupAllocation& CPUAllocationDesc = CPUAllocationDescs[MacroGroupIt];
+		FRDGBufferUAVRef PageIndexGlobalCounterUAVSkipBarrier = GraphBuilder.CreateUAV(PageIndexGlobalCounter, PF_R32_UINT, ERDGUnorderedAccessViewFlags::SkipBarrier);
+		FRDGBufferUAVRef PageToPageIndexUAVSkipBarrier = GraphBuilder.CreateUAV(PageToPageIndexBuffer, PF_R32_UINT, ERDGUnorderedAccessViewFlags::SkipBarrier);
+		FRDGBufferUAVRef PageIndexCoordUAVSkipBarrier = GraphBuilder.CreateUAV(PageIndexCoordBuffer, PF_R8G8B8A8_UINT, ERDGUnorderedAccessViewFlags::SkipBarrier);
+		FRDGBufferUAVRef PageIndexBufferUAVSkipBarrier = GraphBuilder.CreateUAV(PageIndexBuffer, PF_R32_UINT, ERDGUnorderedAccessViewFlags::SkipBarrier);
 
-		// Allocate pages
+		for (uint32 MacroGroupIt = 0; MacroGroupIt < MacroGroupCount; ++MacroGroupIt)
 		{
+			const FHairStrandsMacroGroupData& MacroGroup = MacroGroupDatas[MacroGroupIt];
+			FCPUMacroGroupAllocation& CPUAllocationDesc = CPUAllocationDescs[MacroGroupIt];
+
 			FVoxelAllocateVoxelPageCS::FParameters* Parameters = GraphBuilder.AllocParameters<FVoxelAllocateVoxelPageCS::FParameters>();
 			Parameters->MacroGroupId = MacroGroup.MacroGroupId;
 			Parameters->PageCount = PageCount;
 			Parameters->CPU_PageIndexCount = CPUAllocationDesc.PageIndexCount;
 			Parameters->CPU_PageIndexResolution = CPUAllocationDesc.PageIndexResolution;
 			Parameters->CPU_PageIndexOffset = CPUAllocationDesc.PageIndexOffset;
-			Parameters->PageIndexGlobalCounter = PageIndexGlobalCounterUAV;
-			Parameters->PageIndexBuffer = PageIndexBufferUAV;
-			Parameters->PageToPageIndexBuffer = GraphBuilder.CreateUAV(PageToPageIndexBuffer, PF_R32_UINT);
-			Parameters->PageIndexCoordBuffer = GraphBuilder.CreateUAV(PageIndexCoordBuffer, PF_R8G8B8A8_UINT);
+			Parameters->RWPageIndexGlobalCounter = PageIndexGlobalCounterUAVSkipBarrier;
+			Parameters->RWPageIndexBuffer = PageIndexBufferUAVSkipBarrier;
+			Parameters->RWPageToPageIndexBuffer = PageToPageIndexUAVSkipBarrier;
+			Parameters->RWPageIndexCoordBuffer = PageIndexCoordUAVSkipBarrier;
 
 			FVoxelAllocateVoxelPageCS::FPermutationDomain PermutationVector;
 			PermutationVector.Set<FVoxelAllocateVoxelPageCS::FGPUDriven>(bIsGPUDriven ? 1 : 0);
@@ -779,7 +753,7 @@ static void AddAllocateVoxelPagesPass(
 
 				const uint32 ArgsOffset = sizeof(FRHIDispatchIndirectParameters) * MacroGroup.MacroGroupId;
 
-				FComputeShaderUtils::AddPass(GraphBuilder, RDG_EVENT_NAME("HairStrands::AllocateVoxelPage"), ComputeShader, Parameters, 
+				FComputeShaderUtils::AddPass(GraphBuilder, RDG_EVENT_NAME("HairStrands::AllocateVoxelPage"), ComputeShader, Parameters,
 					PageIndexAllocationIndirectBufferArgs,
 					ArgsOffset);
 			}
@@ -790,22 +764,19 @@ static void AddAllocateVoxelPagesPass(
 				FComputeShaderUtils::AddPass(GraphBuilder, RDG_EVENT_NAME("HairStrands::AllocateVoxelPage"), ComputeShader, Parameters, DispatchCount);
 			}
 		}
+	}
 
-		// Prepare indirect dispatch buffers
-		{
-			check(MacroGroup.MacroGroupId < MacroGroupCount);
+	// Prepare indirect dispatch buffers for all pages, and for each macro group
+	{
+		FVoxelAddIndirectBufferCS::FParameters* Parameters = GraphBuilder.AllocParameters<FVoxelAddIndirectBufferCS::FParameters>();
+		Parameters->MacroGroupCount = MacroGroupCount;
+		Parameters->PageResolution = PageResolution;
+		Parameters->IndirectGroupSize = FVirtualVoxelInjectOpaqueCS::GetGroupSize();
+		Parameters->PageIndexGlobalCounter = GraphBuilder.CreateSRV(PageIndexGlobalCounter, PF_R32_UINT);
+		Parameters->OutIndirectArgsBuffer = GraphBuilder.CreateUAV(IndirectArgsBuffer);
 
-			FVoxelAddIndirectBufferCS::FParameters* Parameters = GraphBuilder.AllocParameters<FVoxelAddIndirectBufferCS::FParameters>();
-			Parameters->MacroGroupId = MacroGroup.MacroGroupId;
-			Parameters->PageResolution = PageResolution;
-			Parameters->IndirectGroupSize = FVirtualVoxelInjectOpaqueCS::GetGroupSize();
-			Parameters->OutPageIndexGlobalCounter = GraphBuilder.CreateUAV(PageIndexGlobalCounter, PF_R32_UINT);
-			Parameters->OutIndirectArgsBuffer = GraphBuilder.CreateUAV(IndirectArgsBuffer);
-			Parameters->OutTotalRequestedPageAllocationBuffer = TotalRequestedPageAllocationBufferUAV;
-
-			TShaderMapRef<FVoxelAddIndirectBufferCS> ComputeShader(View.ShaderMap);
-			FComputeShaderUtils::AddPass(GraphBuilder, RDG_EVENT_NAME("HairStrands::BuildVoxelIndirectArgs"), ComputeShader, Parameters, FIntVector(1, 1, 1));
-		}
+		TShaderMapRef<FVoxelAddIndirectBufferCS> ComputeShader(View.ShaderMap);
+		FComputeShaderUtils::AddPass(GraphBuilder, RDG_EVENT_NAME("HairStrands::BuildVoxelIndirectArgs"), ComputeShader, Parameters, FIntVector(FMath::DivideAndRoundUp(MacroGroupCount, FVoxelAddIndirectBufferCS::GetGroupSize()), 1, 1));
 	}
 
 	// Feedback allocation for next frame
@@ -820,7 +791,7 @@ static void AddAllocateVoxelPagesPass(
 		Parameters->CPUAllocatedPageCount = TotalPageCount;
 		Parameters->AdaptiveCorrectionThreshold = FMath::Clamp(GHairVirtualVoxelAdaptive_CorrectionThreshold, 0.f, 1.f);
 		Parameters->AdaptiveCorrectionSpeed = FMath::Clamp(GHairVirtualVoxelAdaptive_CorrectionSpeed, 0.f, 1.f);
-		Parameters->TotalRequestedPageAllocationBuffer = GraphBuilder.CreateSRV(TotalRequestedPageAllocationBuffer, PF_R32_UINT);
+		Parameters->PageIndexGlobalCounter = GraphBuilder.CreateSRV(PageIndexGlobalCounter, PF_R32_UINT);
 		Parameters->CurrGPUMinVoxelWorldSize = CurrGPUMinVoxelSizeSRV;
 		Parameters->NextGPUMinVoxelWorldSize = GraphBuilder.CreateUAV(NextGPUMinVoxelSize, Format);
 		ShaderPrint::SetParameters(GraphBuilder, View.ShaderPrintData, Parameters->ShaderPrintUniformBuffer);
@@ -946,7 +917,7 @@ static FHairStrandsVoxelResources AllocateVirtualVoxelResources(
 			Out.Parameters.Common.PageTextureResolution,
 			PF_R32_UINT, 
 			FClearValueBinding::Black, 
-			TexCreate_UAV | TexCreate_ShaderResource,
+			TexCreate_UAV | TexCreate_ShaderResource | TexCreate_NoFastClear,
 			MipCount);
 		Out.PageTexture = GraphBuilder.CreateTexture(Desc, TEXT("Hair.VoxelPageTexture"));
 	}
@@ -1037,7 +1008,7 @@ static FHairStrandsVoxelResources AllocateDummyVirtualVoxelResources(
 	return Out;
 }
 
-static FRDGBufferRef IndirectVoxelPageClear(
+static void IndirectVoxelPageClear(
 	FRDGBuilder& GraphBuilder,
 	const FViewInfo& ViewInfo,
 	FHairStrandsVoxelResources& VoxelResources)
@@ -1046,45 +1017,20 @@ static FRDGBufferRef IndirectVoxelPageClear(
 	SCOPED_DRAW_EVENT(GraphBuilder.RHICmdList, HairStrandsIndVoxelPageClear);
 	SCOPED_GPU_STAT(GraphBuilder.RHICmdList, HairStrandsIndVoxelPageClear);
 
-	FRDGBufferRef ClearIndArgsBuffer = GraphBuilder.CreateBuffer(FRDGBufferDesc::CreateIndirectDesc<FRHIDispatchIndirectParameters>(1), TEXT("Hair.VirtualVoxelClearIndirectArgsBuffer"));
-	FRDGBufferSRVRef PageIndexGlobalCounter = GraphBuilder.CreateSRV(VoxelResources.PageIndexGlobalCounter, PF_R32_UINT);
-
-	// Generate the indirect buffer required to clear all voxel allocated linearly in the page volume texture, using the global counter.
-	{
-		FVoxelIndPageClearBufferGenCS::FParameters* Parameters = GraphBuilder.AllocParameters<FVoxelIndPageClearBufferGenCS::FParameters>();
-		Parameters->GroupSize = FVoxelIndPageClearCS::GetGroupSize();
-		Parameters->PageResolution = VoxelResources.Parameters.Common.PageResolution;
-		Parameters->OutIndirectArgsBuffer = GraphBuilder.CreateUAV(ClearIndArgsBuffer);
-		Parameters->PageIndexGlobalCounter = PageIndexGlobalCounter;
-
-		TShaderMapRef<FVoxelIndPageClearBufferGenCS> ComputeShader(ViewInfo.ShaderMap);
-		FComputeShaderUtils::AddPass(
-			GraphBuilder,
-			RDG_EVENT_NAME("HairStrands::VoxelGenIndBufferClearCS"),
-			ComputeShader,
-			Parameters,
-			FIntVector(1,1,1));
-	}
-
-	// Now single dispatch to clear all the pages
-	{
-		FVoxelIndPageClearCS::FParameters* Parameters = GraphBuilder.AllocParameters<FVoxelIndPageClearCS::FParameters>();
-		Parameters->VirtualVoxelParams = VoxelResources.Parameters.Common;
-		Parameters->OutPageTexture = GraphBuilder.CreateUAV(FRDGTextureUAVDesc(VoxelResources.PageTexture));
-		Parameters->IndirectDispatchBuffer = ClearIndArgsBuffer;
-		Parameters->PageIndexGlobalCounter = PageIndexGlobalCounter;
-
-		TShaderMapRef<FVoxelIndPageClearCS> ComputeShader(ViewInfo.ShaderMap);
-		FComputeShaderUtils::AddPass(
-			GraphBuilder,
-			RDG_EVENT_NAME("HairStrands::VoxelIndPageClearCS"),
-			ComputeShader,
-			Parameters,
-			ClearIndArgsBuffer,
-			0);
-	}
+	FVoxelIndPageClearCS::FParameters* Parameters = GraphBuilder.AllocParameters<FVoxelIndPageClearCS::FParameters>();
+	Parameters->VirtualVoxelParams = VoxelResources.Parameters.Common;
+	Parameters->OutPageTexture = GraphBuilder.CreateUAV(FRDGTextureUAVDesc(VoxelResources.PageTexture));
+	Parameters->IndirectDispatchBuffer = VoxelResources.IndirectArgsBuffer;
+	Parameters->PageIndexGlobalCounter = GraphBuilder.CreateSRV(VoxelResources.PageIndexGlobalCounter, PF_R32_UINT);
 
-	return ClearIndArgsBuffer;
+	TShaderMapRef<FVoxelIndPageClearCS> ComputeShader(ViewInfo.ShaderMap);
+	FComputeShaderUtils::AddPass(
+		GraphBuilder,
+		RDG_EVENT_NAME("HairStrands::VoxelIndPageClearCS"),
+		ComputeShader,
+		Parameters,
+		Parameters->IndirectDispatchBuffer,
+		0);
 }
 
 ///////////////////////////////////////////////////////////////////////////////////////////////////
@@ -1256,7 +1202,6 @@ static void AddVirtualVoxelGenerateMipPass(
 	FRDGBuilder& GraphBuilder,
 	const FViewInfo& View,
 	FHairStrandsVoxelResources& VoxelResources,
-	FRDGBufferRef IndirectArgsBuffer, 
 	FRDGBufferRef InPageToPageIndexBuffer)
 {
 	if (!VoxelResources.IsValid())
@@ -1280,7 +1225,7 @@ static void AddVirtualVoxelGenerateMipPass(
 		Parameters->PageResolution		= VoxelResources.Parameters.Common.PageResolution;
 		Parameters->TargetMipIndex		= TargetMipIndex;
 		Parameters->DispatchGroupSize	= FVirtualVoxelGenerateMipCS::GetGroupSize();
-		Parameters->InIndirectArgs		= GraphBuilder.CreateSRV(IndirectArgsBuffer);
+		Parameters->InIndirectArgs		= GraphBuilder.CreateSRV(VoxelResources.IndirectArgsBuffer);
 		Parameters->OutIndirectArgs		= GraphBuilder.CreateUAV(MipIndirectArgs);
 
 		TShaderMapRef<FVirtualVoxelIndirectArgMipCS> ComputeShader(View.ShaderMap);
@@ -1418,13 +1363,14 @@ void VoxelizeHairStrands(
 
 	// Force transition to ensure passes to be batched
 	TArray<FRDGBufferSRVRef> Transitions;
-	Transitions.Reserve(InstanceDatas.Num());
+	Transitions.Reserve(InstanceDatas.Num() * 3u);
 	for (const FInstanceData& InstanceData : InstanceDatas)
 	{
 		if (InstanceData.Data->GetCullingResultAvailable())
 		{
 			Transitions.Add(RegisterAsSRV(GraphBuilder, InstanceData.Data->GetDrawIndirectRasterComputeBuffer()));
-
+			Transitions.Add(RegisterAsSRV(GraphBuilder, InstanceData.Data->GetCulledVertexRadiusScaleBuffer()));
+			Transitions.Add(RegisterAsSRV(GraphBuilder, InstanceData.Data->GetCulledVertexIdBuffer()));
 		}
 	}
 	AddTransitionPass(GraphBuilder, View.ShaderMap, View.GetShaderPlatform(), Transitions);
@@ -1433,7 +1379,7 @@ void VoxelizeHairStrands(
 	FHairStrandsViewStateData* HairStrandsViewStateData = View.ViewState ? &View.ViewState->HairStrandsViewStateData : nullptr;
 	VirtualVoxelResources = AllocateVirtualVoxelResources(GraphBuilder, View, PreViewStereoCorrection, *TransientResources, MacroGroupDatas, MacroGroupResources, PageToPageIndexBuffer, HairStrandsViewStateData);
 
-	FRDGBufferRef ClearIndArgsBuffer = IndirectVoxelPageClear(GraphBuilder, View, VirtualVoxelResources);
+	IndirectVoxelPageClear(GraphBuilder, View, VirtualVoxelResources);
 	AddVirtualVoxelizationRasterPass(GraphBuilder, &View, VirtualVoxelResources, InstanceDatas);
 
 	if (GHairVoxelInjectOpaqueDepthEnable > 0)
@@ -1445,5 +1391,5 @@ void VoxelizeHairStrands(
 		}
 	}
 
-	AddVirtualVoxelGenerateMipPass(GraphBuilder, View, VirtualVoxelResources, ClearIndArgsBuffer, PageToPageIndexBuffer);
+	AddVirtualVoxelGenerateMipPass(GraphBuilder, View, VirtualVoxelResources, PageToPageIndexBuffer);
 }
