LFAtlas - made debug vis atlas slot outline thinner to more easiliy link them with material name.

diff --git a/Engine/Source/Runtime/Renderer/Private/LightFunctionAtlas.cpp b/Engine/Source/Runtime/Renderer/Private/LightFunctionAtlas.cpp
--- a/Engine/Source/Runtime/Renderer/Private/LightFunctionAtlas.cpp
+++ b/Engine/Source/Runtime/Renderer/Private/LightFunctionAtlas.cpp
@@ -753,14 +753,14 @@ FScreenPassTexture FLightFunctionAtlas::AddDebugVisualizationPasses(FRDGBuilder&
 				FCanvasLineItem LineItem(
 					FIntPoint(DstPoint.X + X0 * DisplayResolutionRatio, DstPoint.Y + Y0 * DisplayResolutionRatio),
 					FIntPoint(DstPoint.X + X1 * DisplayResolutionRatio, DstPoint.Y + Y1 * DisplayResolutionRatio));
-				LineItem.LineThickness = 2.0f;
+				LineItem.LineThickness = 4.0f;
 				LineItem.SetColor(MaterialColor);
 				Canvas.DrawItem(LineItem);
 			};
-			OutLineAtlasSlot(AtlasSlot.Min.X + 1, AtlasSlot.Min.Y + 1, AtlasSlot.Max.X - 1, AtlasSlot.Min.Y + 1);
-			OutLineAtlasSlot(AtlasSlot.Max.X - 1, AtlasSlot.Min.Y + 1, AtlasSlot.Max.X - 1, AtlasSlot.Max.Y - 1);
-			OutLineAtlasSlot(AtlasSlot.Max.X - 1, AtlasSlot.Max.Y - 1, AtlasSlot.Min.X + 1, AtlasSlot.Max.Y - 1);
-			OutLineAtlasSlot(AtlasSlot.Min.X + 1, AtlasSlot.Max.Y - 1, AtlasSlot.Min.X + 1, AtlasSlot.Min.Y + 1);
+			OutLineAtlasSlot(AtlasSlot.Min.X + 2, AtlasSlot.Min.Y + 2, AtlasSlot.Max.X - 2, AtlasSlot.Min.Y + 2);
+			OutLineAtlasSlot(AtlasSlot.Max.X - 2, AtlasSlot.Min.Y + 2, AtlasSlot.Max.X - 2, AtlasSlot.Max.Y - 2);
+			OutLineAtlasSlot(AtlasSlot.Max.X - 2, AtlasSlot.Max.Y - 2, AtlasSlot.Min.X + 2, AtlasSlot.Max.Y - 2);
+			OutLineAtlasSlot(AtlasSlot.Min.X + 2, AtlasSlot.Max.Y - 2, AtlasSlot.Min.X + 2, AtlasSlot.Min.Y + 2);
 
 			DrawPosY += 20.0f;
 			LightFunctionAtlasSlotIndex++;
