LFAtlas - define unification

diff --git a/Engine/Shaders/Private/VolumetricFog.usf b/Engine/Shaders/Private/VolumetricFog.usf
--- a/Engine/Shaders/Private/VolumetricFog.usf
+++ b/Engine/Shaders/Private/VolumetricFog.usf
@@ -461,7 +461,7 @@ float4 InjectShadowedLocalLightCommon(uint3 GridCoordinate)
 #endif
 	
 			FLightFunctionColor LightFunctionColor = 1.0f;
-#if	USE_LIGHT_FUNCTION
+#if	USE_LIGHT_FUNCTION_ATLAS
 			LightFunctionColor = GetLocalLightFunctionCommon(TranslatedWorldPosition, LightData.LightFunctionAtlasLightIndex);
 #endif
 
diff --git a/Engine/Source/Runtime/Renderer/Private/VolumetricFog.cpp b/Engine/Source/Runtime/Renderer/Private/VolumetricFog.cpp
--- a/Engine/Source/Runtime/Renderer/Private/VolumetricFog.cpp
+++ b/Engine/Source/Runtime/Renderer/Private/VolumetricFog.cpp
@@ -332,17 +332,17 @@ class FInjectShadowedLocalLightPS : public FGlobalShader
 		RENDER_TARGET_BINDING_SLOTS()
 	END_SHADER_PARAMETER_STRUCT()
 
-	class FDynamicallyShadowed	: SHADER_PERMUTATION_BOOL("DYNAMICALLY_SHADOWED");
-	class FTemporalReprojection : SHADER_PERMUTATION_BOOL("USE_TEMPORAL_REPROJECTION");
-	class FLightFunction		: SHADER_PERMUTATION_BOOL("USE_LIGHT_FUNCTION");
-	class FEnableShadows		: SHADER_PERMUTATION_BOOL("ENABLE_SHADOW_COMPUTATION");
-	class FVirtualShadowMap		: SHADER_PERMUTATION_BOOL("VIRTUAL_SHADOW_MAP");
-	class FRectLightTexture		: SHADER_PERMUTATION_BOOL("USE_RECT_LIGHT_TEXTURE");
+	class FDynamicallyShadowed		: SHADER_PERMUTATION_BOOL("DYNAMICALLY_SHADOWED");
+	class FTemporalReprojection		: SHADER_PERMUTATION_BOOL("USE_TEMPORAL_REPROJECTION");
+	class FSampleLightFunctionAtlas	: SHADER_PERMUTATION_BOOL("USE_LIGHT_FUNCTION_ATLAS");
+	class FEnableShadows			: SHADER_PERMUTATION_BOOL("ENABLE_SHADOW_COMPUTATION");
+	class FVirtualShadowMap			: SHADER_PERMUTATION_BOOL("VIRTUAL_SHADOW_MAP");
+	class FRectLightTexture			: SHADER_PERMUTATION_BOOL("USE_RECT_LIGHT_TEXTURE");
 
 	using FPermutationDomain = TShaderPermutationDomain<
 		FDynamicallyShadowed,
 		FTemporalReprojection,
-		FLightFunction,
+		FSampleLightFunctionAtlas,
 		FEnableShadows,
 		FVirtualShadowMap,
 		FRectLightTexture >;
@@ -371,13 +371,13 @@ class FInjectShadowedLocalLightRGS : public FGlobalShader
 		SHADER_PARAMETER(int32, FirstSlice)
 	END_SHADER_PARAMETER_STRUCT()
 
-	class FTemporalReprojection : SHADER_PERMUTATION_BOOL("USE_TEMPORAL_REPROJECTION");
-	class FLightFunction		: SHADER_PERMUTATION_BOOL("USE_LIGHT_FUNCTION");
-	class FRectLightTexture		: SHADER_PERMUTATION_BOOL("USE_RECT_LIGHT_TEXTURE");
+	class FTemporalReprojection		: SHADER_PERMUTATION_BOOL("USE_TEMPORAL_REPROJECTION");
+	class FSampleLightFunctionAtlas : SHADER_PERMUTATION_BOOL("USE_LIGHT_FUNCTION_ATLAS");
+	class FRectLightTexture			: SHADER_PERMUTATION_BOOL("USE_RECT_LIGHT_TEXTURE");
 
 	using FPermutationDomain = TShaderPermutationDomain<
 		FTemporalReprojection,
-		FLightFunction,
+		FSampleLightFunctionAtlas,
 		FRectLightTexture >;
 
 	static bool ShouldCompilePermutation(const FGlobalShaderPermutationParameters& Parameters)
@@ -527,7 +527,7 @@ void FDeferredShadingSceneRenderer::PrepareRayTracingVolumetricFogShadows(const
 			{
 				FInjectShadowedLocalLightRGS::FPermutationDomain PermutationVector;
 				PermutationVector.Set<FInjectShadowedLocalLightRGS::FTemporalReprojection>((bool)TemporalReprojection);
-				PermutationVector.Set<FInjectShadowedLocalLightRGS::FLightFunction>((bool)UseLightFunction);
+				PermutationVector.Set<FInjectShadowedLocalLightRGS::FSampleLightFunctionAtlas>((bool)UseLightFunction);
 				PermutationVector.Set<FInjectShadowedLocalLightRGS::FRectLightTexture>((bool)UseRectLightTexture);
 
 				TShaderMapRef<FInjectShadowedLocalLightRGS> RayGenerationShader(View.ShaderMap, PermutationVector);
@@ -784,7 +784,7 @@ void FSceneRenderer::RenderLocalLightsForVolumetricFog(
 				FInjectShadowedLocalLightPS::FPermutationDomain PermutationVector;
 				PermutationVector.Set< FInjectShadowedLocalLightPS::FDynamicallyShadowed >(bDynamicallyShadowed);
 				PermutationVector.Set< FInjectShadowedLocalLightPS::FTemporalReprojection >(bUseTemporalReprojection);
-				PermutationVector.Set< FInjectShadowedLocalLightPS::FLightFunction >(bUseLightFunctionAtlas);
+				PermutationVector.Set< FInjectShadowedLocalLightPS::FSampleLightFunctionAtlas >(bUseLightFunctionAtlas);
 				PermutationVector.Set< FInjectShadowedLocalLightPS::FEnableShadows >(bIsShadowed);
 				PermutationVector.Set< FInjectShadowedLocalLightPS::FVirtualShadowMap >(bUseVSM);
 				PermutationVector.Set< FInjectShadowedLocalLightPS::FRectLightTexture >(bUsesRectLightTexture);
@@ -891,7 +891,7 @@ void FSceneRenderer::RenderLocalLightsForVolumetricFog(
 
 				FInjectShadowedLocalLightRGS::FPermutationDomain PermutationVector;
 				PermutationVector.Set< FInjectShadowedLocalLightRGS::FTemporalReprojection >(bUseTemporalReprojection);
-				PermutationVector.Set< FInjectShadowedLocalLightRGS::FLightFunction >(bUseLightFunctionAtlas);
+				PermutationVector.Set< FInjectShadowedLocalLightRGS::FSampleLightFunctionAtlas >(bUseLightFunctionAtlas);
 				PermutationVector.Set< FInjectShadowedLocalLightRGS::FRectLightTexture >(bUsesRectLightTexture);
 
 				TShaderMapRef<FInjectShadowedLocalLightRGS> RayGenerationShader(GetGlobalShaderMap(FeatureLevel), PermutationVector);
