Fixed missing GPU-Scene uniform buffer setup after moving GPU-Scene update out of Render.

diff --git a/Engine/Source/Runtime/Renderer/Private/SceneRendering.cpp b/Engine/Source/Runtime/Renderer/Private/SceneRendering.cpp
--- a/Engine/Source/Runtime/Renderer/Private/SceneRendering.cpp
+++ b/Engine/Source/Runtime/Renderer/Private/SceneRendering.cpp
@@ -5136,6 +5136,8 @@ public:
 
 		Scene.UpdateAllPrimitiveSceneInfos(GraphBuilder);
 		GPUScene.BeginRender(GraphBuilder, GPUSceneDynamicContext);
+
+		GPUScene.FillSceneUniformBuffer(GraphBuilder, Renderer.GetSceneUniforms());
 	}
 
 	virtual ~FScenePrimitiveRenderingContext()
