LFAtlas - unified atlas sampling as part of the Dynamiclighthing code. That will be shared with everything then, only need to add a #define when compiling a material. VFog do nto called get DynamicLighting so it still sample the atlas explicitely.

diff --git a/Engine/Shaders/Private/DeferredLightPixelShaders.usf b/Engine/Shaders/Private/DeferredLightPixelShaders.usf
--- a/Engine/Shaders/Private/DeferredLightPixelShaders.usf
+++ b/Engine/Shaders/Private/DeferredLightPixelShaders.usf
@@ -55,8 +55,6 @@
 #include "VirtualShadowMaps/VirtualShadowMapMaskBitsCommon.ush"
 #endif
 
-#include "LightFunctionAtlas/LightFunctionAtlasCommon.usf"
-
 // Remap light shape permutation into light type
 #if LIGHT_SOURCE_SHAPE == 0
 #define CURRENT_LIGHT_TYPE LIGHT_TYPE_DIRECTIONAL
@@ -337,13 +335,6 @@ void DeferredLightPixelMain(
 		float4 LightAttenuation = GetLightAttenuationFromShadow(InputParams, SceneDepth, DerivedParams.TranslatedWorldPosition);
 		float4 Radiance = GetDynamicLighting(DerivedParams.TranslatedWorldPosition, DerivedParams.CameraVector, ScreenSpaceData.GBuffer, ScreenSpaceData.AmbientOcclusion, ScreenSpaceData.GBuffer.ShadingModelID, LightData, LightAttenuation, Dither, uint2(InputParams.PixelPos), SurfaceShadow);
 
-#if USE_LIGHT_FUNCTION_ATLAS
-		if (bUseLightFunctionAtlas)
-		{
-			Radiance.rgb *= GetLocalLightFunctionCommon(DerivedParams.TranslatedWorldPosition, LightData.LightFunctionAtlasLightIndex);
-		}
-#endif
-
 		OutColor += Radiance;
 	}
 
@@ -445,20 +436,12 @@ void DeferredLightPixelMain(
 
 	float SurfaceShadow = 1.0f;
 	const float4 Radiance = GetDynamicLighting(TranslatedWorldPosition, CameraVector, HairScreenSpaceData.GBuffer, HairScreenSpaceData.AmbientOcclusion, HairScreenSpaceData.GBuffer.ShadingModelID, LightData, LightAttenuation, Dither, PixelCoord, SurfaceShadow);
-	float  Attenuation = ComputeLightProfileMultiplier(TranslatedWorldPosition, GetDeferredLightTranslatedWorldPosition(), -DeferredLightUniforms.Direction, DeferredLightUniforms.Tangent, DeferredLightUniforms.IESAtlasIndex);
+	const float  Attenuation = ComputeLightProfileMultiplier(TranslatedWorldPosition, GetDeferredLightTranslatedWorldPosition(), -DeferredLightUniforms.Direction, DeferredLightUniforms.Tangent, DeferredLightUniforms.IESAtlasIndex);
 	const float3 AtmosphericTransmittance = GetAtmosphereAndCloudAttenuation(GetStableTranslatedLargeWorldPosition(ScreenPosition, Sample.Depth), UV);
 
 	const float LocalCoverage = From8bitCoverage(Sample.Coverage8bit);
 
-	FLightFunctionColor LightFunctionColor = 1.0f;
-#if	USE_LIGHT_FUNCTION_ATLAS
-	if (bUseLightFunctionAtlas)
-	{
-		LightFunctionColor = GetLocalLightFunctionCommon(TranslatedWorldPosition, LightData.LightFunctionAtlasLightIndex);
-	}
-#endif
-
-	OutColor.rgb = AtmosphericTransmittance * Radiance.xyz * LightFunctionColor * Attenuation * LocalCoverage * GetExposure();
+	OutColor.rgb = AtmosphericTransmittance * Radiance.xyz * Attenuation * LocalCoverage * GetExposure();
 	OutColor.a = LocalCoverage;
 }
 #endif
diff --git a/Engine/Shaders/Private/DeferredLightingCommon.ush b/Engine/Shaders/Private/DeferredLightingCommon.ush
--- a/Engine/Shaders/Private/DeferredLightingCommon.ush
+++ b/Engine/Shaders/Private/DeferredLightingCommon.ush
@@ -15,6 +15,10 @@
 #include "Substrate/Substrate.ush"
 #include "LightData.ush"
 
+#if USE_LIGHT_FUNCTION_ATLAS
+#include "LightFunctionAtlas/LightFunctionAtlasCommon.usf"
+#endif
+
 // This Substrate include assumes that if inline shading is required, it needs to 
 // be defined prior to DeferredLightingCommong.ush
 // SubstrateStruct.MaterialTextureArray is only available to translucent materials or global shaders computing contact shadows.
@@ -401,7 +405,11 @@ FLightAccumulator AccumulateDynamicLighting(
 			}
 
 			Lighting.Specular *= LightData.SpecularScale;
-				
+
+		#if USE_LIGHT_FUNCTION_ATLAS
+			MaskedLightColor *= GetLocalLightFunctionCommon(TranslatedWorldPosition, LightData.LightFunctionAtlasLightIndex);
+		#endif
+
 			LightAccumulator_AddSplit( LightAccumulator, Lighting.Diffuse, Lighting.Specular, Lighting.Diffuse, MaskedLightColor * Shadow.SurfaceShadow, bNeedsSeparateSubsurfaceLightAccumulation );
 			LightAccumulator_AddSplit( LightAccumulator, Lighting.Transmission, 0.0f, Lighting.Transmission, MaskedLightColor * Shadow.TransmissionShadow, bNeedsSeparateSubsurfaceLightAccumulation );
 
diff --git a/Engine/Shaders/Private/LightFunctionAtlas/LightFunctionAtlasCommon.usf b/Engine/Shaders/Private/LightFunctionAtlas/LightFunctionAtlasCommon.usf
--- a/Engine/Shaders/Private/LightFunctionAtlas/LightFunctionAtlasCommon.usf
+++ b/Engine/Shaders/Private/LightFunctionAtlas/LightFunctionAtlasCommon.usf
@@ -4,7 +4,7 @@
 	LightFunctionAtlasCommon.usf
 =============================================================================*/
 
-#define LIGHT_FUNCTION_ATLAS 1
+#pragma once
 
 #include "../Common.ush"
 #include "/Engine/Shared/LightDefinitions.h"
diff --git a/Engine/Source/Runtime/Renderer/Private/LightRendering.cpp b/Engine/Source/Runtime/Renderer/Private/LightRendering.cpp
--- a/Engine/Source/Runtime/Renderer/Private/LightRendering.cpp
+++ b/Engine/Source/Runtime/Renderer/Private/LightRendering.cpp
@@ -859,7 +859,6 @@ class FDeferredLightPS : public FGlobalShader
 		// For virtual shadow map mask
 		SHADER_PARAMETER_RDG_UNIFORM_BUFFER(FVirtualShadowMapUniformParameters, VirtualShadowMap)
 		SHADER_PARAMETER(int32, VirtualShadowMapId)
-		SHADER_PARAMETER(int32, bUseLightFunctionAtlas)
 		SHADER_PARAMETER_RDG_TEXTURE(Texture2D, ShadowMaskBits)
 		RENDER_TARGET_BINDING_SLOTS()
 	END_SHADER_PARAMETER_STRUCT()
@@ -2327,7 +2326,6 @@ static FDeferredLightPS::FParameters GetDeferredLightPSParameters(
 	Out.ShadowMaskBits = ShadowMaskBits ? ShadowMaskBits : GSystemTextures.GetZeroUIntDummy(GraphBuilder);
 
 	// If the light is not batched, it could be due to shadow, so we still specify light function atlas sampling.
-	Out.bUseLightFunctionAtlas = LightSceneInfo->Proxy->HasValidLightFunctionAtlasSlot() && View.LightFunctionAtlasViewData.GetDeferredlightingUsesLightFunctionAtlas();
 	Out.LightFunctionAtlas = LightFunctionAtlasGlobalParameters;
 
 	// PS - Render Targets
@@ -2621,7 +2619,9 @@ static void RenderLight(
 		PermutationVector.Set< FDeferredLightPS::FVirtualShadowMapMask >(bUseVirtualShadowMapMask);
 		PermutationVector.Set< FDeferredLightPS::FSubstrateTileType >(0);
 		PermutationVector.Set< FDeferredLightPS::FHairComplexTransmittance >(bNeedComplexTransmittanceSupport);
-		PermutationVector.Set< FDeferredLightPS::FLightFunctionAtlasDim >(LightSceneInfo->Proxy->GetLightFunctionMaterial()!=nullptr && View.Family->EngineShowFlags.LightFunctions);
+		PermutationVector.Set< FDeferredLightPS::FLightFunctionAtlasDim >(
+			View.LightFunctionAtlasViewData.GetDeferredlightingUsesLightFunctionAtlas() && LightSceneInfo->Proxy->HasValidLightFunctionAtlasSlot() &&
+			LightSceneInfo->Proxy->GetLightFunctionMaterial() != nullptr && View.Family->EngineShowFlags.LightFunctions && !View.Family->EngineShowFlags.VisualizeLightCulling);
 		if (bIsRadial)
 		{
 			PermutationVector.Set< FDeferredLightPS::FSourceShapeDim >(LightProxy->IsRectLight() ? ELightSourceShape::Rect : ELightSourceShape::Capsule);
@@ -2768,7 +2768,9 @@ void FDeferredShadingSceneRenderer::RenderLightForHair(
 	PermutationVector.Set< FDeferredLightPS::FTransmissionDim >(false);
 	PermutationVector.Set< FDeferredLightPS::FHairLighting>(1);
 	PermutationVector.Set< FDeferredLightPS::FHairComplexTransmittance>(true);
-	PermutationVector.Set< FDeferredLightPS::FLightFunctionAtlasDim >(LightSceneInfo->Proxy->GetLightFunctionMaterial() != nullptr && View.Family->EngineShowFlags.LightFunctions);
+	PermutationVector.Set< FDeferredLightPS::FLightFunctionAtlasDim >(
+		View.LightFunctionAtlasViewData.GetDeferredlightingUsesLightFunctionAtlas() && LightSceneInfo->Proxy->HasValidLightFunctionAtlasSlot() &&
+		LightSceneInfo->Proxy->GetLightFunctionMaterial() != nullptr && View.Family->EngineShowFlags.LightFunctions && !View.Family->EngineShowFlags.VisualizeLightCulling);
 	if (bIsDirectional)
 	{
 		PermutationVector.Set< FDeferredLightPS::FSourceShapeDim >(ELightSourceShape::Directional);
