Adds a new BL_TranslucencyAfterDOF location to process AfterDOF translucencies before TSR

TSR requires to have post DOF translucency fed separatly, to cancel nanite anti-flickering heuristics and avoid to make these ghost as they most often don't draw velocity.

But this create a problem where the BL_BeforeTonemaper that is actually between DOF and TSR can no longer post process translucencies.

This CL adds a BL_TranslucencyAfterDOF for artists to have possibility to apply some post processing also on translucencies before they are composited to scene color by TSR.

To avoid divergence in between TAA and TSR, this change also no longer compose post DOF translucency before the BL_BeforeTonemaper to ensure consistent behavior between TSR and the rest of the anti-aliasing technics.

To improve user experience and understanding of tool tips in material editor, this CL also renames:
* BL_BeforeTranslucency -> BL_SceneColorBeforeDOF
* BL_BeforeTonemapping-> BL_SceneColorAfterDOF
* BL_AfterTonemapping-> BL_SceneColorAfterTonemapping

diff --git a/Engine/Plugins/Compositing/Composure/Source/Composure/Private/ComposurePostProcessPass.cpp b/Engine/Plugins/Compositing/Composure/Source/Composure/Private/ComposurePostProcessPass.cpp
--- a/Engine/Plugins/Compositing/Composure/Source/Composure/Private/ComposurePostProcessPass.cpp
+++ b/Engine/Plugins/Compositing/Composure/Source/Composure/Private/ComposurePostProcessPass.cpp
@@ -52,7 +52,7 @@ void UComposurePostProcessPass::SetSetupMaterial(UMaterialInterface* Material)
 			*BaseMaterial->GetName());
 		return;
 	}
-	else if (BaseMaterial->BlendableLocation != BL_BeforeTranslucency)
+	else if (BaseMaterial->BlendableLocation != BL_SceneColorBeforeDOF)
 	{
 		UE_LOG(Composure, Error,
 			TEXT("Can't set setup material %s: is not at the before translucency post process location."),
diff --git a/Engine/Plugins/Experimental/RemoteSession/Source/RemoteSession/Private/Channels/RemoteSessionARCameraChannel.cpp b/Engine/Plugins/Experimental/RemoteSession/Source/RemoteSession/Private/Channels/RemoteSessionARCameraChannel.cpp
--- a/Engine/Plugins/Experimental/RemoteSession/Source/RemoteSession/Private/Channels/RemoteSessionARCameraChannel.cpp
+++ b/Engine/Plugins/Experimental/RemoteSession/Source/RemoteSession/Private/Channels/RemoteSessionARCameraChannel.cpp
@@ -75,7 +75,7 @@ public:
 	{
 		FMaterialShader::ModifyCompilationEnvironment(Parameters, OutEnvironment);
 		OutEnvironment.SetDefine(TEXT("POST_PROCESS_MATERIAL"), 1);
-		OutEnvironment.SetDefine(TEXT("POST_PROCESS_MATERIAL_BEFORE_TONEMAP"), (Parameters.MaterialParameters.BlendableLocation != BL_AfterTonemapping) ? 1 : 0);
+		OutEnvironment.SetDefine(TEXT("POST_PROCESS_MATERIAL_BEFORE_TONEMAP"), (Parameters.MaterialParameters.BlendableLocation != BL_SceneColorAfterTonemapping) ? 1 : 0);
 		OutEnvironment.SetDefine(TEXT("POST_PROCESS_MATERIAL_SSRINPUT"), (Parameters.MaterialParameters.BlendableLocation == BL_SSRInput) ? 1 : 0);
 	}
 };
diff --git a/Engine/Plugins/Runtime/AR/AppleAR/AppleARKit/Source/AppleARKit/Private/AppleARKitVideoOverlay.cpp b/Engine/Plugins/Runtime/AR/AppleAR/AppleARKit/Source/AppleARKit/Private/AppleARKitVideoOverlay.cpp
--- a/Engine/Plugins/Runtime/AR/AppleAR/AppleARKit/Source/AppleARKit/Private/AppleARKitVideoOverlay.cpp
+++ b/Engine/Plugins/Runtime/AR/AppleAR/AppleARKit/Source/AppleARKit/Private/AppleARKitVideoOverlay.cpp
@@ -153,7 +153,7 @@ public:
 	{
 		FMaterialShader::ModifyCompilationEnvironment(Parameters, OutEnvironment);
 		OutEnvironment.SetDefine(TEXT("POST_PROCESS_MATERIAL"), 1);
-		OutEnvironment.SetDefine(TEXT("POST_PROCESS_MATERIAL_BEFORE_TONEMAP"), (Parameters.MaterialParameters.BlendableLocation != BL_AfterTonemapping) ? 1 : 0);
+		OutEnvironment.SetDefine(TEXT("POST_PROCESS_MATERIAL_BEFORE_TONEMAP"), (Parameters.MaterialParameters.BlendableLocation != BL_SceneColorAfterTonemapping) ? 1 : 0);
 		OutEnvironment.SetDefine(TEXT("POST_PROCESS_MATERIAL_SSRINPUT"), (Parameters.MaterialParameters.BlendableLocation == BL_SSRInput) ? 1 : 0);
 		OutEnvironment.SetDefine(TEXT("POST_PROCESS_AR_PASSTHROUGH"), 1);
 
diff --git a/Engine/Plugins/Runtime/AR/Google/GoogleARCore/Source/GoogleARCoreRendering/Private/GoogleARCorePassthroughCameraRenderer.cpp b/Engine/Plugins/Runtime/AR/Google/GoogleARCore/Source/GoogleARCoreRendering/Private/GoogleARCorePassthroughCameraRenderer.cpp
--- a/Engine/Plugins/Runtime/AR/Google/GoogleARCore/Source/GoogleARCoreRendering/Private/GoogleARCorePassthroughCameraRenderer.cpp
+++ b/Engine/Plugins/Runtime/AR/Google/GoogleARCore/Source/GoogleARCoreRendering/Private/GoogleARCorePassthroughCameraRenderer.cpp
@@ -117,7 +117,7 @@ public:
 	{
 		FMaterialShader::ModifyCompilationEnvironment(Parameters, OutEnvironment);
 		OutEnvironment.SetDefine(TEXT("POST_PROCESS_MATERIAL"), 1);
-		OutEnvironment.SetDefine(TEXT("POST_PROCESS_MATERIAL_BEFORE_TONEMAP"), (Parameters.MaterialParameters.BlendableLocation != BL_AfterTonemapping) ? 1 : 0);
+		OutEnvironment.SetDefine(TEXT("POST_PROCESS_MATERIAL_BEFORE_TONEMAP"), (Parameters.MaterialParameters.BlendableLocation != BL_SceneColorAfterTonemapping) ? 1 : 0);
 		OutEnvironment.SetDefine(TEXT("POST_PROCESS_MATERIAL_SSRINPUT"), (Parameters.MaterialParameters.BlendableLocation == BL_SSRInput) ? 1 : 0);
 	}
 };
diff --git a/Engine/Shaders/Private/MaterialTemplate.ush b/Engine/Shaders/Private/MaterialTemplate.ush
--- a/Engine/Shaders/Private/MaterialTemplate.ush
+++ b/Engine/Shaders/Private/MaterialTemplate.ush
@@ -2376,7 +2376,7 @@ MaterialFloat2 GetLightmapUVs_DDY(FMaterialPixelParameters Parameters)
 }
 
 //The post-process material needs to decode the scene color since it's encoded at PreTonemapMSAA if MSAA enabled on MetalMobilePlatorm 
-//The POST_PROCESS_MATERIAL_BEFORE_TONEMAP is 1 for both BL_BeforeTranslucency and BL_BeforeTonemapping post-process materials
+//The POST_PROCESS_MATERIAL_BEFORE_TONEMAP is 1 for both BL_SceneColorBeforeDOF and BL_SceneColorAfterDOF post-process materials
 #if FEATURE_LEVEL <= FEATURE_LEVEL_ES3_1 && POST_PROCESS_MATERIAL && POST_PROCESS_MATERIAL_BEFORE_TONEMAP && METAL_PROFILE
 uint bMetalMSAAHDRDecode;
 #endif
diff --git a/Engine/Source/Runtime/Engine/Classes/Engine/BlendableInterface.h b/Engine/Source/Runtime/Engine/Classes/Engine/BlendableInterface.h
--- a/Engine/Source/Runtime/Engine/Classes/Engine/BlendableInterface.h
+++ b/Engine/Source/Runtime/Engine/Classes/Engine/BlendableInterface.h
@@ -9,29 +9,63 @@
 class UMaterialInstanceDynamic;
 class UMaterialInterface;
 
-/** Where to place a material node in the post processing graph. */
+/** Where to place a post process material in the post processing chain. */
 UENUM()
 enum EBlendableLocation : int
 {
-	/** Input0:former pass color, Input1:SeparateTranslucency. */
-	BL_AfterTonemapping UMETA(DisplayName="After Tonemapping"),
-	/** Input0:former pass color, Input1:SeparateTranslucency. */
-	BL_BeforeTonemapping UMETA(DisplayName="Before Tonemapping"),
-	/** Input0:former pass color, Input1:SeparateTranslucency. */
-	BL_BeforeTranslucency UMETA(DisplayName="Before Translucency"),
-	/**
-	* Input0:former pass color, Input1:SeparateTranslucency, Input2: BloomOutput
-	*/
-	BL_ReplacingTonemapper UMETA(DisplayName="Replacing the Tonemapper"),
-//	BL_AfterOpaque,
-//	BL_AfterFog,
-//	BL_AfterTranslucency,
-//	BL_AfterPostProcessAA,
-
-	/** Input0:former pass color. */
-	BL_SSRInput UMETA(DisplayName = "SSR Input"),
-
-	BL_MAX,
+	/** Post process material location to modify the scene color, between translucency distortion and DOF.
+	 * Always run at rendering resolution.
+	 *
+	 * Input0:former pass scene color, excluding AfterDOF translucency
+	 * Input1:AfterDOF translucency.
+	 */
+	BL_SceneColorBeforeDOF = 2 UMETA(DisplayName="Scene Color Before DOF"),
+
+	/** Post process material location to modify the scene color, between DOF and AfterDOF translucency.
+	 * Always run at rendering resolution.
+	 *
+	 * Input0:former pass scene color, excluding AfterDOF translucency
+	 * Input1:AfterDOF translucency.
+	 */
+	BL_SceneColorAfterDOF = 1 UMETA(DisplayName = "Scene Color After DOF"),
+
+	/** Post process material location to modify the AfterDOF translucency, before composition into the scene color.
+	 * Always run at rendering resolution.
+	 *
+	 * Input0:scene color without translucency, after DOF,
+	 * Input1:AfterDOF translucency.
+	 */
+	BL_TranslucencyAfterDOF = 5 UMETA(DisplayName = "Translucency After DOF"),
+
+	/** Post process material location to compose a backplate into SSR, between TSR/TAA and next frame's SSR.
+	 * Runs at display resolution with TSR or TAAU, rendering resolution otherwise.
+	 *
+	 * Input0:TAA/TSR output,
+	 */
+	BL_SSRInput = 4 UMETA(DisplayName = "SSR Input"),
+
+	/** Post process material replacing the tone mapper, to modify the scene color.
+	 * Runs at display resolution with TSR or TAAU, rendering resolution otherwise.
+	 *
+	 * Input0:former pass scene color,
+	 * Input1:AfterDOF translucency.
+	 */
+	BL_ReplacingTonemapper = 3 UMETA(DisplayName="Replacing the Tonemapper"),
+
+	/** Post process material location to modify the scene color, after tone mapper.
+	 * Runs at display resolution with TSR or TAAU, rendering resolution otherwise.
+	 *
+	 * Input0:former pass scene color,
+	 * Input1:AfterDOF translucency.
+	 */
+	BL_SceneColorAfterTonemapping = 0 UMETA(DisplayName = "Scene Color After Tonemapping"),
+
+	BL_MAX = 6 UMETA(Hidden),
+
+	// Olds names that needs to be kept forever to ensure asset serialization to work correctly when UENUM() switched from serializing int to names.
+	BL_BeforeTranslucency UE_DEPRECATED(5.4, "Renamed to BL_SceneColorBeforeDOF")        = BL_SceneColorBeforeDOF UMETA(Hidden),
+	BL_BeforeTonemapping  UE_DEPRECATED(5.4, "Renamed to BL_SceneColorAfterDOF")         = BL_SceneColorAfterDOF UMETA(Hidden),
+	BL_AfterTonemapping   UE_DEPRECATED(5.4, "Renamed to BL_SceneColorAfterTonemapping") = BL_SceneColorAfterTonemapping UMETA(Hidden),
 };
 
 /** Dummy class needed to support Cast<IBlendableInterface>(Object). */
diff --git a/Engine/Source/Runtime/Engine/Classes/Materials/Material.h b/Engine/Source/Runtime/Engine/Classes/Materials/Material.h
--- a/Engine/Source/Runtime/Engine/Classes/Materials/Material.h
+++ b/Engine/Source/Runtime/Engine/Classes/Materials/Material.h
@@ -1156,6 +1156,7 @@ public:
 	ENGINE_API virtual bool IsDeferredDecal() const override;
 	ENGINE_API virtual bool IsUIMaterial() const;
 	ENGINE_API virtual bool IsPostProcessMaterial() const;
+	ENGINE_API bool IsPostProcessMaterialOutputingAlpha() const;
 	ENGINE_API virtual USubsurfaceProfile* GetSubsurfaceProfile_Internal() const override;
 	ENGINE_API virtual uint32 NumSpecularProfile_Internal() const override;
 	ENGINE_API virtual USpecularProfile* GetSpecularProfile_Internal(uint32 Index) const override;
diff --git a/Engine/Source/Runtime/Engine/Private/Materials/HLSLMaterialTranslator.cpp b/Engine/Source/Runtime/Engine/Private/Materials/HLSLMaterialTranslator.cpp
--- a/Engine/Source/Runtime/Engine/Private/Materials/HLSLMaterialTranslator.cpp
+++ b/Engine/Source/Runtime/Engine/Private/Materials/HLSLMaterialTranslator.cpp
@@ -8102,7 +8102,7 @@ int32 FHLSLMaterialTranslator::SceneTextureLookup(int32 ViewportUV, uint32 InSce
 	SubstrateCtx.SubstrateMaterialComplexity.bIsSimple = false;
 	SubstrateCtx.SubstrateMaterialComplexity.bIsSingle = false;
 	
-	if (SceneTextureId == PPI_PostProcessInput0 && Material->GetMaterialDomain() == MD_PostProcess && Material->GetBlendableLocation() != BL_AfterTonemapping)
+	if (SceneTextureId == PPI_PostProcessInput0 && Material->GetMaterialDomain() == MD_PostProcess && Material->GetBlendableLocation() != BL_SceneColorAfterTonemapping)
 	{
 		return AddInlinedCodeChunk(MCT_Float4, TEXT("(float4(View.OneOverPreExposure.xxx, 1) * %s)"), *CoerceParameter(LookUp, MCT_Float4));
 	}
diff --git a/Engine/Source/Runtime/Engine/Private/Materials/Material.cpp b/Engine/Source/Runtime/Engine/Private/Materials/Material.cpp
--- a/Engine/Source/Runtime/Engine/Private/Materials/Material.cpp
+++ b/Engine/Source/Runtime/Engine/Private/Materials/Material.cpp
@@ -998,7 +998,7 @@ UMaterial::UMaterial(const FObjectInitializer& ObjectInitializer)
 	bApplyCloudFogging = false;
 	bIsSky = false;
 	bUsedWithWater = false;
-	BlendableLocation = BL_AfterTonemapping;
+	BlendableLocation = BL_SceneColorAfterTonemapping;
 	BlendablePriority = 0;
 	BlendableOutputAlpha = false;
 	bIsBlendable = true;
@@ -3570,7 +3570,7 @@ void UMaterial::ConvertMaterialToSubstrateMaterial()
 			ShadingModels.ClearShadingModels();
 			ShadingModels.AddShadingModel(MSM_Unlit);
 
-			if (MaterialDomain == MD_PostProcess && !BlendableOutputAlpha)
+			if (MaterialDomain == MD_PostProcess && !IsPostProcessMaterialOutputingAlpha())
 			{
 				BlendMode = BLEND_Opaque;
 			}
@@ -4448,11 +4448,11 @@ bool UMaterial::CanEditChange(const FProperty* InProperty) const
 		{
 			if (bSubstrateEnabled)
 			{
-				return ((MaterialDomain != MD_PostProcess && MaterialDomain != MD_LightFunction && MaterialDomain != MD_Volume) || (MaterialDomain == MD_PostProcess && BlendableOutputAlpha));
+				return ((MaterialDomain != MD_PostProcess && MaterialDomain != MD_LightFunction && MaterialDomain != MD_Volume) || IsPostProcessMaterialOutputingAlpha());
 			}
 			else
 			{
-				return (MaterialDomain == MD_DeferredDecal || MaterialDomain == MD_Surface || MaterialDomain == MD_Volume || MaterialDomain == MD_UI || (MaterialDomain == MD_PostProcess && BlendableOutputAlpha));
+				return (MaterialDomain == MD_DeferredDecal || MaterialDomain == MD_Surface || MaterialDomain == MD_Volume || MaterialDomain == MD_UI || IsPostProcessMaterialOutputingAlpha());
 			}
 		}
 	
@@ -5136,7 +5136,7 @@ void UMaterial::RebuildShadingModelField()
 				ShadingModel = MSM_Unlit;
 				// We keep the blend mode resulting from ConvertLegacyToSubstrateBlendMode because post processes can be translucent.
 				// However, we do force opaque mode if blending has been disabled via the post-process specific BlendableOutputAlpha option.
-				if (!BlendableOutputAlpha)
+				if (!IsPostProcessMaterialOutputingAlpha())
 				{
 					BlendMode = BLEND_Opaque;
 				}
@@ -6723,6 +6723,11 @@ bool UMaterial::IsPostProcessMaterial() const
 	return MaterialDomain == MD_PostProcess;
 }
 
+bool UMaterial::IsPostProcessMaterialOutputingAlpha() const
+{
+	return UMaterial::IsPostProcessMaterial() && (BlendableOutputAlpha || BlendableLocation == BL_TranslucencyAfterDOF);
+}
+
 bool UMaterial::WritesToRuntimeVirtualTexture() const
 {
 	return GetCachedExpressionData().bHasRuntimeVirtualTextureOutput;
@@ -7121,7 +7126,7 @@ bool UMaterial::IsPropertyActiveInEditor(EMaterialProperty InProperty) const
 		BlendMode,
 		ShadingModels,
 		TranslucencyLightingMode,
-		BlendableOutputAlpha,
+		IsPostProcessMaterialOutputingAlpha(),
 		bUsesDistortion,
 		IsShadingModelFromMaterialExpression(),
 		IsTranslucencyWritingVelocity(),
@@ -7146,7 +7151,7 @@ bool UMaterial::IsPropertyActiveInDerived(EMaterialProperty InProperty, const UM
 		DerivedMaterial->GetBlendMode(),
 		DerivedMaterial->GetShadingModels(),
 		TranslucencyLightingMode,
-		BlendableOutputAlpha,
+		IsPostProcessMaterialOutputingAlpha(),
 		bUsesDistortion,
 		DerivedMaterial->IsShadingModelFromMaterialExpression(),
 		IsTranslucencyWritingVelocity(),
diff --git a/Engine/Source/Runtime/Engine/Private/Materials/MaterialHLSLTree.cpp b/Engine/Source/Runtime/Engine/Private/Materials/MaterialHLSLTree.cpp
--- a/Engine/Source/Runtime/Engine/Private/Materials/MaterialHLSLTree.cpp
+++ b/Engine/Source/Runtime/Engine/Private/Materials/MaterialHLSLTree.cpp
@@ -11,7 +11,7 @@
 #include "Materials/MaterialAttributeDefinitionMap.h"
 #include "Materials/MaterialFunctionInterface.h"
 #include "Materials/MaterialExpressionVectorNoise.h"
-#include "Engine/BlendableInterface.h" // BL_AfterTonemapping
+#include "Engine/BlendableInterface.h" // BL_SceneColorAfterTonemapping
 #include "VT/VirtualTextureScalability.h"
 #include "VT/RuntimeVirtualTexture.h"
 #include "DataDrivenShaderPlatformInfo.h"
@@ -2312,7 +2312,7 @@ void FExpressionSceneTexture::EmitValueShader(FEmitContext& Context, FEmitScope&
 		EmitLookup = Context.EmitExpression(Scope, Shader::EValueType::Float4, TEXT("MobileSceneTextureLookup(Parameters, %, %)"), (int)SceneTextureId, EmitTexCoord);
 	}
 
-	if (SceneTextureId == PPI_PostProcessInput0 && Context.Material->GetMaterialDomain() == MD_PostProcess && Context.Material->GetBlendableLocation() != BL_AfterTonemapping)
+	if (SceneTextureId == PPI_PostProcessInput0 && Context.Material->GetMaterialDomain() == MD_PostProcess && Context.Material->GetBlendableLocation() != BL_SceneColorAfterTonemapping)
 	{
 		EmitLookup = Context.EmitExpression(Scope, Shader::EValueType::Float4, TEXT("(float4(View.OneOverPreExposure.xxx, 1) * %)"), EmitLookup);
 	}
diff --git a/Engine/Source/Runtime/Engine/Private/Materials/MaterialShared.cpp b/Engine/Source/Runtime/Engine/Private/Materials/MaterialShared.cpp
--- a/Engine/Source/Runtime/Engine/Private/Materials/MaterialShared.cpp
+++ b/Engine/Source/Runtime/Engine/Private/Materials/MaterialShared.cpp
@@ -2229,7 +2229,7 @@ int32 FMaterialResource::GetBlendableLocation() const
 
 bool FMaterialResource::GetBlendableOutputAlpha() const
 {
-	return Material->BlendableOutputAlpha;
+	return Material->IsPostProcessMaterialOutputingAlpha();
 }
 
 bool FMaterialResource::IsStencilTestEnabled() const
diff --git a/Engine/Source/Runtime/Renderer/Private/PostProcess/PostProcessMaterial.cpp b/Engine/Source/Runtime/Renderer/Private/PostProcess/PostProcessMaterial.cpp
--- a/Engine/Source/Runtime/Renderer/Private/PostProcess/PostProcessMaterial.cpp
+++ b/Engine/Source/Runtime/Renderer/Private/PostProcess/PostProcessMaterial.cpp
@@ -80,7 +80,7 @@ static EMaterialCustomDepthPolicy GetMaterialCustomDepthPolicy(const FMaterial*
 		{
 			UE_LOG(LogRenderer, Warning, TEXT("PostProcessMaterial uses stencil test, but stencil not allocated. Set r.CustomDepth to 3 to allocate custom stencil."));
 		}
-		else if (Material->GetBlendableLocation() == BL_AfterTonemapping)
+		else if (Material->GetBlendableLocation() == BL_SceneColorAfterTonemapping)
 		{
 			// We can't support custom stencil after tonemapping due to target size differences
 			UE_LOG(LogRenderer, Warning, TEXT("PostProcessMaterial uses stencil test, but is set to blend After Tonemapping. This is not supported."));
@@ -252,7 +252,7 @@ public:
 		OutEnvironment.SetDefine(TEXT("POST_PROCESS_MATERIAL"), 1);
 
 		EBlendableLocation Location = EBlendableLocation(Parameters.MaterialParameters.BlendableLocation);
-		OutEnvironment.SetDefine(TEXT("POST_PROCESS_MATERIAL_BEFORE_TONEMAP"), (Location == BL_AfterTonemapping || Location == BL_ReplacingTonemapper) ? 0 : 1);
+		OutEnvironment.SetDefine(TEXT("POST_PROCESS_MATERIAL_BEFORE_TONEMAP"), (Location == BL_SceneColorAfterTonemapping || Location == BL_ReplacingTonemapper) ? 0 : 1);
 		// Post process SSR is always rendered at native resolution as if it was after tone mapping, so we need to account for the fact that it is independent from DRS.
 		// SSR input should not be affected by exposure so it should be specified separately from POST_PROCESS_MATERIAL_BEFORE_TONEMAP 
 		// in order to be able to make DRS independent CameraVector and WorldPosition nodes.
@@ -260,7 +260,7 @@ public:
 
 		if (IsMobilePlatform(Parameters.Platform))
 		{
-			OutEnvironment.SetDefine(TEXT("POST_PROCESS_MATERIAL_BEFORE_TONEMAP"), (Parameters.MaterialParameters.BlendableLocation != BL_AfterTonemapping) ? 1 : 0);
+			OutEnvironment.SetDefine(TEXT("POST_PROCESS_MATERIAL_BEFORE_TONEMAP"), (Parameters.MaterialParameters.BlendableLocation != BL_SceneColorAfterTonemapping) ? 1 : 0);
 		}
 
 		// PostProcessMaterial can both read & write Substrate data
@@ -681,8 +681,6 @@ FScreenPassTexture AddPostProcessMaterialPass(
 {
 	Inputs.Validate();
 
-	const FScreenPassTexture SceneColor = Inputs.GetInput(EPostProcessMaterialInput::SceneColor);
-
 	const ERHIFeatureLevel::Type FeatureLevel = View.GetFeatureLevel();
 
 	const FMaterial* Material = nullptr;
@@ -692,6 +690,9 @@ FScreenPassTexture AddPostProcessMaterialPass(
 	TShaderRef<FPostProcessMaterialPS> PixelShader;
 	GetMaterialInfo(MaterialInterface, FeatureLevel, Inputs, Material, MaterialRenderProxy, MaterialShaderMap, VertexShader, PixelShader);
 
+	EBlendableLocation BlendableLocation = EBlendableLocation(Material->GetBlendableLocation());
+	const FScreenPassTexture SceneColor = Inputs.GetInput(BlendableLocation == BL_TranslucencyAfterDOF ? EPostProcessMaterialInput::SeparateTranslucency : EPostProcessMaterialInput::SceneColor);
+
 	check(VertexShader.IsValid());
 	check(PixelShader.IsValid());
 
@@ -984,15 +985,16 @@ FScreenPassTexture AddPostProcessMaterialChain(
 	FRDGBuilder& GraphBuilder,
 	const FViewInfo& View,
 	const FPostProcessMaterialInputs& InputsTemplate,
-	const FPostProcessMaterialChain& Materials)
+	const FPostProcessMaterialChain& Materials,
+	EPostProcessMaterialInput MaterialInput)
 {
-	FScreenPassTexture Outputs = InputsTemplate.GetInput(EPostProcessMaterialInput::SceneColor);
+	FScreenPassTexture Outputs = InputsTemplate.GetInput(MaterialInput);
 
 	bool bFirstMaterialInChain = true;
 	for (const UMaterialInterface* MaterialInterface : Materials)
 	{
 		FPostProcessMaterialInputs Inputs = InputsTemplate;
-		Inputs.SetInput(EPostProcessMaterialInput::SceneColor, Outputs);
+		Inputs.SetInput(MaterialInput, Outputs);
 		
 		// Only the first material in the chain needs to decode the input color
 		Inputs.bMetalMSAAHDRDecode = Inputs.bMetalMSAAHDRDecode && bFirstMaterialInChain;
diff --git a/Engine/Source/Runtime/Renderer/Private/PostProcess/PostProcessMaterial.h b/Engine/Source/Runtime/Renderer/Private/PostProcess/PostProcessMaterial.h
--- a/Engine/Source/Runtime/Renderer/Private/PostProcess/PostProcessMaterial.h
+++ b/Engine/Source/Runtime/Renderer/Private/PostProcess/PostProcessMaterial.h
@@ -49,7 +49,8 @@ FScreenPassTexture AddPostProcessMaterialChain(
 	FRDGBuilder& GraphBuilder,
 	const FViewInfo& View,
 	const FPostProcessMaterialInputs& Inputs,
-	const FPostProcessMaterialChain& MaterialChain);
+	const FPostProcessMaterialChain& MaterialChain,
+	EPostProcessMaterialInput MaterialInput = EPostProcessMaterialInput::SceneColor);
 
 struct FHighResolutionScreenshotMaskInputs
 {
diff --git a/Engine/Source/Runtime/Renderer/Private/PostProcess/PostProcessing.cpp b/Engine/Source/Runtime/Renderer/Private/PostProcess/PostProcessing.cpp
--- a/Engine/Source/Runtime/Renderer/Private/PostProcess/PostProcessing.cpp
+++ b/Engine/Source/Runtime/Renderer/Private/PostProcess/PostProcessing.cpp
@@ -317,7 +317,7 @@ void AddPostProcessingPasses(
 	const FScreenPassTexture Velocity(SceneTextureParameters.GBufferVelocityTexture, PrimaryViewRect);
 	const FScreenPassTexture BlackDummy(GSystemTextures.GetBlackDummy(GraphBuilder));
 	
-	const FTranslucencyPassResources& PostDOFTranslucencyResources = Inputs.TranslucencyViewResourcesMap.Get(ETranslucencyPass::TPT_TranslucencyAfterDOF);
+	FTranslucencyPassResources PostDOFTranslucencyResources = Inputs.TranslucencyViewResourcesMap.Get(ETranslucencyPass::TPT_TranslucencyAfterDOF);
 	const FTranslucencyPassResources& PostMotionBlurTranslucencyResources = Inputs.TranslucencyViewResourcesMap.Get(ETranslucencyPass::TPT_TranslucencyAfterMotionBlur);
 
 	// Whether should process the alpha channel of the scene color.
@@ -626,11 +626,11 @@ void AddPostProcessingPasses(
 		const bool bBloomEnabled = View.FinalPostProcessSettings.BloomIntensity > 0.0f && !bVisualizeTSR;
 
 		// Whether separate translucency is composed in TSR.
-		bool bComposeSeparateTranslucencyInTSR = TAAConfig == EMainTAAPassConfig::TSR && ComposeSeparateTranslucencyInTSR(View);
+		bool bComposeSeparateTranslucencyInTSR = PostDOFTranslucencyResources.IsValid() && TAAConfig == EMainTAAPassConfig::TSR && ComposeSeparateTranslucencyInTSR(View);
 
 		const FIntPoint PostTAAViewSize = (View.PrimaryScreenPercentageMethod == EPrimaryScreenPercentageMethod::TemporalUpscale && TAAConfig != EMainTAAPassConfig::Disabled) ? View.GetSecondaryViewRectSize() : View.ViewRect.Size();
 
-		const FPostProcessMaterialChain PostProcessMaterialAfterTonemappingChain = GetPostProcessMaterialChain(View, BL_AfterTonemapping);
+		const FPostProcessMaterialChain PostProcessMaterialAfterTonemappingChain = GetPostProcessMaterialChain(View, BL_SceneColorAfterTonemapping);
 
 		PassSequence.SetEnabled(EPass::MotionBlur, bVisualizeMotionBlur || bMotionBlurEnabled);
 		PassSequence.SetEnabled(EPass::Tonemap, bTonemapEnabled);
@@ -679,17 +679,19 @@ void AddPostProcessingPasses(
 		const bool bNeedPostMotionBlurHalfRes = !bProcessQuarterResolution || (bFFTBloomEnabled && FFTBloomResolutionFraction > 0.25f && FFTBloomResolutionFraction <= 0.5f) || (ReflectionsMethod == EReflectionsMethod::SSR && !View.bStatePrevViewInfoIsReadOnly && GSSRHalfResSceneColor);
 		const bool bNeedPostMotionBlurQuarterRes = bProcessQuarterResolution || (bFFTBloomEnabled && FFTBloomResolutionFraction <= 0.25f);
 
-		// Post Process Material Chain - Before Translucency
-		{
-			const FPostProcessMaterialChain MaterialChain = GetPostProcessMaterialChain(View, BL_BeforeTranslucency);
 
-			if (MaterialChain.Num())
-			{
-				SceneColor = AddPostProcessMaterialChain(GraphBuilder, View, GetPostProcessMaterialInputs(SceneColor), MaterialChain);
-			}
+		const FPostProcessMaterialChain MaterialChainSceneColorBeforeDOF = GetPostProcessMaterialChain(View, BL_SceneColorBeforeDOF);
+		const FPostProcessMaterialChain MaterialChainSceneColorAfterDOF = GetPostProcessMaterialChain(View, BL_SceneColorAfterDOF);
+		const FPostProcessMaterialChain MaterialChainTranslucencyAfterDOF = GetPostProcessMaterialChain(View, BL_TranslucencyAfterDOF);
+
+		// Post Process Material Chain - BL_SceneColorBeforeDOF
+		if (MaterialChainSceneColorBeforeDOF.Num())
+		{
+			SceneColor = AddPostProcessMaterialChain(GraphBuilder, View, GetPostProcessMaterialInputs(SceneColor), MaterialChainSceneColorBeforeDOF);
 		}
 
 		// Diaphragm Depth of Field
+		bool bSceneColorHasPostDOFTranslucency = false;
 		{
 			FRDGTextureRef InputSceneColorTexture = SceneColor.Texture;
 
@@ -697,16 +699,51 @@ void AddPostProcessingPasses(
 			{
 				FTranslucencyPassResources DummyTranslucency;
 
+				bool bComposeTranslucency = PostDOFTranslucencyResources.IsValid() && !bComposeSeparateTranslucencyInTSR && MaterialChainTranslucencyAfterDOF.Num() == 0;
+
 				SceneColor.Texture = DiaphragmDOF::AddPasses(
 					GraphBuilder,
 					SceneTextureParameters,
 					View,
 					SceneColor.Texture,
-					bComposeSeparateTranslucencyInTSR ? DummyTranslucency : PostDOFTranslucencyResources);
+					bComposeTranslucency ? PostDOFTranslucencyResources : DummyTranslucency);
+				bSceneColorHasPostDOFTranslucency = bComposeTranslucency;
+			}
+
+			if (GetHairStrandsComposition() == EHairStrandsCompositionType::AfterSeparateTranslucent)
+			{
+				RenderHairComposition(GraphBuilder, View, SceneColor.Texture, SceneDepth.Texture, Velocity.Texture);
+			}
+		}
+
+		// Post Process Material Chain - BL_SceneColorAfterDOF
+		if (MaterialChainSceneColorAfterDOF.Num())
+		{
+			SceneColor = AddPostProcessMaterialChain(GraphBuilder, View, GetPostProcessMaterialInputs(SceneColor), MaterialChainSceneColorAfterDOF);
+		}
+
+		// Post Process Material Chain - BL_TranslucencyAfterDOF
+		if (bSceneColorHasPostDOFTranslucency)
+		{
+			ensure(MaterialChainTranslucencyAfterDOF.Num() == 0);
+			ensure(!bComposeSeparateTranslucencyInTSR);
+		}
+		else if (PostDOFTranslucencyResources.IsValid())
+		{
+			if (MaterialChainTranslucencyAfterDOF.Num())
+			{
+				FScreenPassTexture PostDOFTranslucency = AddPostProcessMaterialChain(
+					GraphBuilder, View,
+					GetPostProcessMaterialInputs(SceneColor),
+					MaterialChainTranslucencyAfterDOF,
+					EPostProcessMaterialInput::SeparateTranslucency);
+
+				PostDOFTranslucencyResources.ColorTexture = PostDOFTranslucency.Texture;
+				ensure(PostDOFTranslucencyResources.ViewRect == PostDOFTranslucency.ViewRect);
 			}
 
 			// DOF passes were not added, therefore need to compose Separate translucency manually.
-			if ((SceneColor.Texture == InputSceneColorTexture || bComposeSeparateTranslucencyInTSR) && PostDOFTranslucencyResources.IsValid())
+			if (!bSceneColorHasPostDOFTranslucency)
 			{
 				FTranslucencyComposition TranslucencyComposition;
 				TranslucencyComposition.Operation = FTranslucencyComposition::EOperation::ComposeToNewSceneColor;
@@ -718,24 +755,17 @@ void AddPostProcessingPasses(
 
 				SceneColor = TranslucencyComposition.AddPass(
 					GraphBuilder, View, PostDOFTranslucencyResources);
-			}
 
-			if (GetHairStrandsComposition() == EHairStrandsCompositionType::AfterSeparateTranslucent)
-			{
-				RenderHairComposition(GraphBuilder, View, SceneColor.Texture, SceneDepth.Texture, Velocity.Texture);
+				bSceneColorHasPostDOFTranslucency = !TranslucencyComposition.bApplyModulateOnly;
 			}
 		}
-
-		// Post Process Material Chain - Before Tonemapping
+		else
 		{
-			const FPostProcessMaterialChain MaterialChain = GetPostProcessMaterialChain(View, BL_BeforeTonemapping);
-
-			if (MaterialChain.Num())
-			{
-				SceneColor = AddPostProcessMaterialChain(GraphBuilder, View, GetPostProcessMaterialInputs(SceneColor), MaterialChain);
-			}
+			bSceneColorHasPostDOFTranslucency = true;
 		}
 
+		ensure(bSceneColorHasPostDOFTranslucency != bComposeSeparateTranslucencyInTSR);
+
 		// Allows for the scene color to be the slice of an array between temporal upscaler and tonemaper.
 		FScreenPassTextureSlice SceneColorSlice = FScreenPassTextureSlice::CreateFromScreenPassTexture(GraphBuilder, SceneColor);
 		SceneColor = FScreenPassTexture();
@@ -2136,7 +2166,7 @@ void AddMobilePostProcessingPasses(FRDGBuilder& GraphBuilder, FScene* Scene, con
 	{
 		FPostProcessMaterialInputs PostProcessMaterialInputs;
 
-		if (BlendableLocation == BL_AfterTonemapping && PassSequence.IsEnabled(EPass::PostProcessMaterialAfterTonemapping))
+		if (BlendableLocation == BL_SceneColorAfterTonemapping && PassSequence.IsEnabled(EPass::PostProcessMaterialAfterTonemapping))
 		{
 			PassSequence.AcceptOverrideIfLastPass(EPass::PostProcessMaterialAfterTonemapping, PostProcessMaterialInputs.OverrideOutput);
 		}
@@ -2179,7 +2209,7 @@ void AddMobilePostProcessingPasses(FRDGBuilder& GraphBuilder, FScene* Scene, con
 
 		bool bUseSeparateTranslucency = IsMobileSeparateTranslucencyActive(View);
 
-		const FPostProcessMaterialChain PostProcessMaterialAfterTonemappingChain = GetPostProcessMaterialChain(View, BL_AfterTonemapping);
+		const FPostProcessMaterialChain PostProcessMaterialAfterTonemappingChain = GetPostProcessMaterialChain(View, BL_SceneColorAfterTonemapping);
 
 		PassSequence.SetEnabled(EPass::Distortion, bUseDistortion);
 		PassSequence.SetEnabled(EPass::SunMask, bUseSun || bUseDof);
@@ -2209,7 +2239,7 @@ void AddMobilePostProcessingPasses(FRDGBuilder& GraphBuilder, FScene* Scene, con
 			SceneColor = AddMobileDistortionMergePass(GraphBuilder, View, DistortionMergeInputs);
 		}
 
-		AddPostProcessMaterialPass(BL_BeforeTranslucency, false);
+		AddPostProcessMaterialPass(BL_SceneColorBeforeDOF, false);
 
 		// Optional fixed pass processes
 		if (PassSequence.IsEnabled(EPass::SunMask))
@@ -2513,7 +2543,7 @@ void AddMobilePostProcessingPasses(FRDGBuilder& GraphBuilder, FScene* Scene, con
 			AddMobileSeparateTranslucencyPass(GraphBuilder, Scene, View, SeparateTranslucencyInputs);
 		}
 
-		AddPostProcessMaterialPass(BL_BeforeTonemapping, false);
+		AddPostProcessMaterialPass(BL_SceneColorAfterDOF, false);
 
 		// Temporal Anti-aliasing. Also may perform a temporal upsample from primary to secondary view rect.
 		if (PassSequence.IsEnabled(EPass::TAA))
@@ -2633,7 +2663,7 @@ void AddMobilePostProcessingPasses(FRDGBuilder& GraphBuilder, FScene* Scene, con
 
 	if (IsPostProcessingEnabled(View))
 	{
-		AddPostProcessMaterialPass(BL_AfterTonemapping, true);
+		AddPostProcessMaterialPass(BL_SceneColorAfterTonemapping, true);
 
 		if (PassSequence.IsEnabled(EPass::FXAA))
 		{
