Implement 16bit VALU for TSRDecimateHistory

Submitting on behalf of Hugo Lamarche

diff --git a/Engine/Source/Runtime/Renderer/Private/PostProcess/TemporalSuperResolution.cpp b/Engine/Source/Runtime/Renderer/Private/PostProcess/TemporalSuperResolution.cpp
--- a/Engine/Source/Runtime/Renderer/Private/PostProcess/TemporalSuperResolution.cpp
+++ b/Engine/Source/Runtime/Renderer/Private/PostProcess/TemporalSuperResolution.cpp
@@ -605,7 +605,7 @@ class FTSRDecimateHistoryCS : public FTSRShader
 
 	class FMoireReprojectionDim : SHADER_PERMUTATION_BOOL("DIM_MOIRE_REPROJECTION");
 	class FResurrectionReprojectionDim : SHADER_PERMUTATION_BOOL("DIM_RESURRECTION_REPROJECTION");
-	using FPermutationDomain = TShaderPermutationDomain<FMoireReprojectionDim, FResurrectionReprojectionDim, FTSRShader::FAlphaChannelDim>;
+	using FPermutationDomain = TShaderPermutationDomain<FMoireReprojectionDim, FResurrectionReprojectionDim, FTSRShader::F16BitVALUDim, FTSRShader::FAlphaChannelDim>;
 
 	BEGIN_SHADER_PARAMETER_STRUCT(FParameters, )
 		SHADER_PARAMETER_STRUCT_INCLUDE(FTSRCommonParameters, CommonParameters)
@@ -638,6 +638,34 @@ class FTSRDecimateHistoryCS : public FTSRShader
 
 		SHADER_PARAMETER_RDG_TEXTURE_UAV(RWTexture2DArray, DebugOutput)
 	END_SHADER_PARAMETER_STRUCT()
+
+	static bool ShouldCompilePermutation(const FGlobalShaderPermutationParameters& Parameters)
+	{
+		if (!FTSRShader::ShouldCompilePermutation(Parameters))
+		{
+			return false;
+		}
+
+		FPermutationDomain PermutationVector(Parameters.PermutationId);
+		
+		if (!FTSRShader::ShouldCompile32or16BitPermutation(Parameters.Platform, PermutationVector.Get<FTSRShader::F16BitVALUDim>()))
+		{
+			return false;
+		}
+
+		return true;
+	}
+
+	static void ModifyCompilationEnvironment(const FGlobalShaderPermutationParameters& Parameters, FShaderCompilerEnvironment& OutEnvironment)
+	{
+		FTSRShader::ModifyCompilationEnvironment(Parameters, OutEnvironment);
+
+		FPermutationDomain PermutationVector(Parameters.PermutationId);
+		if (PermutationVector.Get<FTSRShader::F16BitVALUDim>())
+		{
+			OutEnvironment.CompilerFlags.Add(CFLAG_AllowRealTypes);
+		}
+	}
 }; // class FTSRDecimateHistoryCS
 
 class FTSRRejectShadingCS : public FTSRShader
@@ -2041,15 +2069,17 @@ FDefaultTemporalUpscaler::FOutputs AddTemporalSuperResolutionPasses(
 		FTSRDecimateHistoryCS::FPermutationDomain PermutationVector;
 		PermutationVector.Set<FTSRDecimateHistoryCS::FMoireReprojectionDim>(FlickeringFramePeriod > 0.0f);
 		PermutationVector.Set<FTSRDecimateHistoryCS::FResurrectionReprojectionDim>(bCanResurrectHistory);
+		PermutationVector.Set<FTSRShader::F16BitVALUDim>(bUse16BitVALU);
 		PermutationVector.Set<FTSRShader::FAlphaChannelDim>(bSupportsAlpha);
 
 		TShaderMapRef<FTSRDecimateHistoryCS> ComputeShader(View.ShaderMap, PermutationVector);
 		FComputeShaderUtils::AddPass(
 			GraphBuilder,
-			RDG_EVENT_NAME("TSR DecimateHistory(#%d %s%s%s) %dx%d",
+			RDG_EVENT_NAME("TSR DecimateHistory(#%d %s%s%s%s) %dx%d",
 				PermutationVector.ToDimensionValueId(),
 				PermutationVector.Get<FTSRDecimateHistoryCS::FMoireReprojectionDim>() ? TEXT("ReprojectMoire") : TEXT(""),
 				PermutationVector.Get<FTSRDecimateHistoryCS::FResurrectionReprojectionDim>() ? TEXT(" ReprojectResurrection") : TEXT(""),
+				PermutationVector.Get<FTSRShader::F16BitVALUDim>() ? TEXT(" 16bit") : TEXT(""),
 				PermutationVector.Get<FTSRShader::FAlphaChannelDim>() ? TEXT(" AlphaChannel") : TEXT(""),
 				InputRect.Width(), InputRect.Height()),
 			AsyncComputePasses >= 2 ? ERDGPassFlags::AsyncCompute : ERDGPassFlags::Compute,
