LFAtlas - added support for forward translucent and SLW light function sampling from atlas.
LFA parameter can are bound to base pass data. Project settings Cvars controls how to enable those settings.

diff --git a/Engine/Shaders/Private/BasePassPixelShader.usf b/Engine/Shaders/Private/BasePassPixelShader.usf
--- a/Engine/Shaders/Private/BasePassPixelShader.usf
+++ b/Engine/Shaders/Private/BasePassPixelShader.usf
@@ -22,10 +22,12 @@
 	#define UseBasePassSkylight TranslucentBasePass.Shared.UseBasePassSkylight
 	#define SubstrateStruct TranslucentBasePass.Substrate
 	#define BlueNoise TranslucentBasePass.BlueNoise
+	#define LightFunctionAtlasStruct TranslucentBasePass.Shared.LightFunctionAtlas
 #else
 	#define EyeAdaptationStruct OpaqueBasePass
 	#define UseBasePassSkylight OpaqueBasePass.Shared.UseBasePassSkylight
 	#define SubstrateStruct OpaqueBasePass.Substrate
+	#define LightFunctionAtlasStruct OpaqueBasePass.Shared.LightFunctionAtlas
 #endif
 
 // Material setting to allow forward shading (including mobile) to use preintegrated GF lut for simple IBL.
@@ -1481,7 +1483,11 @@ void FPixelShaderInOut_MainPS(
 			Specular, SingleLayerWater.DistortionParams,
 			SunIlluminance, WaterDiffuseIndirectIlluminance, EnvBrdf,
 			CameraIsUnderWater, WaterVisibility, EyeIndex,
-			bSeparateWaterMainDirLightLuminance, SeparatedWaterMainDirLightLuminance).Luminance;
+			bSeparateWaterMainDirLightLuminance, SeparatedWaterMainDirLightLuminance).Luminance
+	#if USE_LIGHT_FUNCTION_ATLAS
+			* GetLocalLightFunctionCommon(SvPositionToResolvedTranslatedWorld(In.SvPosition), GetDirectionalLightData(EyeIndex).LightFunctionAtlasLightIndex)
+	#endif
+		;
 #endif
 	}
 #endif // MATERIAL_SHADINGMODEL_SINGLELAYERWATER
diff --git a/Engine/Shaders/Private/LightData.ush b/Engine/Shaders/Private/LightData.ush
--- a/Engine/Shaders/Private/LightData.ush
+++ b/Engine/Shaders/Private/LightData.ush
@@ -62,6 +62,7 @@ struct FDirectionalLightData
 	float  DirectionalLightSourceRadius;
 	float  DirectionalLightSoftSourceRadius;
 	float  DirectionalLightSpecularScale;
+	uint  LightFunctionAtlasLightIndex;
 };
 
 ///////////////////////////////////////////////////////////////////////////////////////////////////
@@ -434,6 +435,7 @@ FDeferredLightData ConvertToDeferredLight(
 	Out.HairTransmittance	= InitHairTransmittanceData();
 	Out.SourceRadius		= In.DirectionalLightSourceRadius;
 	Out.SoftSourceRadius	= In.DirectionalLightSoftSourceRadius;
+	Out.LightFunctionAtlasLightIndex = In.LightFunctionAtlasLightIndex;
 
 	// Shadow
 	// * Static shadowing uses ShadowMapChannel, 
diff --git a/Engine/Shaders/Private/LightFunctionAtlas/LightFunctionAtlasCommon.usf b/Engine/Shaders/Private/LightFunctionAtlas/LightFunctionAtlasCommon.usf
--- a/Engine/Shaders/Private/LightFunctionAtlas/LightFunctionAtlasCommon.usf
+++ b/Engine/Shaders/Private/LightFunctionAtlas/LightFunctionAtlasCommon.usf
@@ -11,6 +11,12 @@
 
 // In order to use function from that file, one must use FLightFunctionAtlasGlobalParameters
 
+// Trick so that we can sample from LightFunctionAtlas when it is a global UB or when it is a Struct part of another UB.
+#ifdef LightFunctionAtlasStruct
+#undef  LightFunctionAtlas
+#define LightFunctionAtlas LightFunctionAtlasStruct
+#endif
+
 #if COLORED_LIGHT_FUNCTION_ATLAS
 #define FLightFunctionColor float3
 #else
diff --git a/Engine/Shaders/Private/LightGridCommon.ush b/Engine/Shaders/Private/LightGridCommon.ush
--- a/Engine/Shaders/Private/LightGridCommon.ush
+++ b/Engine/Shaders/Private/LightGridCommon.ush
@@ -203,6 +203,7 @@ FDirectionalLightData GetDirectionalLightData(uint EyeIndex)
 		Result.DirectionalLightSourceRadius			= ForwardLightData.DirectionalLightSourceRadius;
 		Result.DirectionalLightSoftSourceRadius		= ForwardLightData.DirectionalLightSoftSourceRadius;
 		Result.DirectionalLightSpecularScale		= ForwardLightData.DirectionalLightSpecularScale;
+		Result.LightFunctionAtlasLightIndex			= ForwardLightData.LightFunctionAtlasLightIndex;
 
 #if (INSTANCED_STEREO || MOBILE_MULTI_VIEW)
 	}
@@ -216,6 +217,7 @@ FDirectionalLightData GetDirectionalLightData(uint EyeIndex)
 		Result.DirectionalLightSourceRadius			= ForwardLightDataISR.DirectionalLightSourceRadius;
 		Result.DirectionalLightSoftSourceRadius		= ForwardLightDataISR.DirectionalLightSoftSourceRadius;
 		Result.DirectionalLightSpecularScale		= ForwardLightDataISR.DirectionalLightSpecularScale;
+		Result.LightFunctionAtlasLightIndex			= ForwardLightDataISR.LightFunctionAtlasLightIndex;
 	}
 #endif
 
diff --git a/Engine/Source/Runtime/Engine/Classes/Engine/RendererSettings.h b/Engine/Source/Runtime/Engine/Classes/Engine/RendererSettings.h
--- a/Engine/Source/Runtime/Engine/Classes/Engine/RendererSettings.h
+++ b/Engine/Source/Runtime/Engine/Classes/Engine/RendererSettings.h
@@ -887,12 +887,30 @@ class URendererSettings : public UDeveloperSettings
 	/**
 	"Select the format of the light function atlas texture."
 	*/
-	UPROPERTY(config, EditAnywhere, Category = Optimizations, meta = (
+	UPROPERTY(config, EditAnywhere, Category = LightFunctionAtlas, meta = (
 		ConsoleVariable = "r.LightFunctionAtlas.Format", DisplayName = "Light Function Atlas Format",
 		ToolTip = "Select the format of the light function atlas texture.",
 		ConfigRestartRequired = true))
 		TEnumAsByte<ELightFunctionAtlasPixelFormat::Type> LightFunctionAtlasPixelFormat;
 
+	/**
+	"Enable support for light function on Single Layer Water when the light function atlas is enabled."
+	*/
+	UPROPERTY(config, EditAnywhere, Category = LightFunctionAtlas, meta = (
+		ConsoleVariable = "r.SingleLayerWater.UsesLightFunctionAtlas", DisplayName = "Single Layer Water Uses Light Function Atlas.",
+		ToolTip = "Enable support for light function on Single Layer Water when the light function atlas is enabled.",
+		ConfigRestartRequired = true))
+		uint32 bSingleLayerWaterUsesLightFunctionAtlas : 1;
+
+	/**
+	"Enable support for light function on Translucent material using Forward Shading mode, when the light function atlas is enabled."
+	*/
+	UPROPERTY(config, EditAnywhere, Category = LightFunctionAtlas, meta = (
+		ConsoleVariable = "r.Translucent.UsesLightFunctionAtlas", DisplayName = "Translucent Uses Light Function Atlas.",
+		ToolTip = "Enable support for light function on Translucent material using Forward Shading mode, when the light function atlas is enabled.",
+		ConfigRestartRequired = true))
+		uint32 bTranslucentUsesLightFunctionAtlas : 1;
+
 	/**
 	"Enable rect light evaluation on translucent materials when using the Forward Shading mode."
 	*/
diff --git a/Engine/Source/Runtime/RenderCore/Private/RenderUtils.cpp b/Engine/Source/Runtime/RenderCore/Private/RenderUtils.cpp
--- a/Engine/Source/Runtime/RenderCore/Private/RenderUtils.cpp
+++ b/Engine/Source/Runtime/RenderCore/Private/RenderUtils.cpp
@@ -1669,11 +1669,33 @@ static TAutoConsoleVariable<int32> CVarLightFunctionAtlasFormat(
 	TEXT("0: grey scale in [0,1]. 1: colored in [0,1]"),
 	ECVF_ReadOnly | ECVF_RenderThreadSafe);
 
+static TAutoConsoleVariable<int32> CVarSingleLayerWaterUsesLightFunctionAtlas(
+	TEXT("r.SingleLayerWater.UsesLightFunctionAtlas"),
+	0,
+	TEXT("Enable sampling of the light function atlas on SingleLAyer Water mateirals."),
+	ECVF_ReadOnly | ECVF_RenderThreadSafe);
+
+static TAutoConsoleVariable<int32> CVarTranslucentUsesLightFunctionAtlas(
+	TEXT("r.Translucent.UsesLightFunctionAtlas"),
+	0,
+	TEXT("Enable sampling of the light function atlas on translucent materials using forward shading."),
+	ECVF_ReadOnly | ECVF_RenderThreadSafe);
+
 int32 GetLightFunctionAtlasFormat()
 {
 	return CVarLightFunctionAtlasFormat.GetValueOnAnyThread();
 }
 
+bool GetSingleLayerWaterUsesLightFunctionAtlas()
+{
+	return CVarSingleLayerWaterUsesLightFunctionAtlas.GetValueOnAnyThread() > 0;
+}
+
+bool GetTranslucentUsesLightFunctionAtlas()
+{
+	return CVarTranslucentUsesLightFunctionAtlas.GetValueOnAnyThread() > 0;
+}
+
 ///////////////////////////////////////////////////////////////////////////////////////////////////
 // Substrate settings interface
 
diff --git a/Engine/Source/Runtime/RenderCore/Private/Shader.cpp b/Engine/Source/Runtime/RenderCore/Private/Shader.cpp
--- a/Engine/Source/Runtime/RenderCore/Private/Shader.cpp
+++ b/Engine/Source/Runtime/RenderCore/Private/Shader.cpp
@@ -2079,6 +2079,19 @@ void ShaderMapAppendKeyString(EShaderPlatform Platform, FString& KeyString)
 		{
 			KeyString += FString::Printf(TEXT("_LFAC%u"), LightFunctionAtlasFormat);
 		}
+
+		bool bSingleLayerWaterUsesLightFunctionAtlas = GetSingleLayerWaterUsesLightFunctionAtlas();
+		if (bSingleLayerWaterUsesLightFunctionAtlas)
+		{
+			KeyString += FString::Printf(TEXT("_SLWLFA"));
+		}
+
+		bool bTranslucentUsesLightFunctionAtlas = GetTranslucentUsesLightFunctionAtlas();
+		if (bTranslucentUsesLightFunctionAtlas)
+		{
+			KeyString += FString::Printf(TEXT("_FWDLFA"));
+		}
+
 	}
 
 	{
diff --git a/Engine/Source/Runtime/RenderCore/Public/RenderUtils.h b/Engine/Source/Runtime/RenderCore/Public/RenderUtils.h
--- a/Engine/Source/Runtime/RenderCore/Public/RenderUtils.h
+++ b/Engine/Source/Runtime/RenderCore/Public/RenderUtils.h
@@ -549,7 +549,11 @@ namespace Substrate
 	RENDERCORE_API uint32 GetShadingQuality(EShaderPlatform InPlatform);
 }
 
+// Light function atlas settings
 RENDERCORE_API int32 GetLightFunctionAtlasFormat();
+// Light function atlas project settings triggering shader compilation
+RENDERCORE_API bool GetSingleLayerWaterUsesLightFunctionAtlas();
+RENDERCORE_API bool GetTranslucentUsesLightFunctionAtlas();
 
 // LuminanceMax is the amount of light that will cause the sensor to saturate at EV100.
 //  See also https://en.wikipedia.org/wiki/Film_speed and https://en.wikipedia.org/wiki/Exposure_value for more info.
diff --git a/Engine/Source/Runtime/Renderer/Private/BasePassRendering.cpp b/Engine/Source/Runtime/Renderer/Private/BasePassRendering.cpp
--- a/Engine/Source/Runtime/Renderer/Private/BasePassRendering.cpp
+++ b/Engine/Source/Runtime/Renderer/Private/BasePassRendering.cpp
@@ -701,6 +701,7 @@ extern void SetupDummyForwardLightUniformParameters(FRDGBuilder& GraphBuilder, F
 void SetupSharedBasePassParameters(
 	FRDGBuilder& GraphBuilder,
 	const FViewInfo& View,
+	const int32 ViewIndex,
 	bool bLumenGIEnabled,
 	FSharedBasePassUniformParameters& SharedParameters)
 {
@@ -736,6 +737,9 @@ void SetupSharedBasePassParameters(
 
 	SharedParameters.LFV = View.LocalFogVolumeViewData.UniformParametersStruct;
 
+	FLightFunctionAtlas* LightFunctionAtlas = View.LightFunctionAtlasViewData.GetLightFunctionAtlas();
+	SharedParameters.LightFunctionAtlas = *LightFunctionAtlas->GetLightFunctionAtlasGlobalParametersStruct(ViewIndex, GraphBuilder);
+
 	const FScene* Scene = View.Family->Scene ? View.Family->Scene->GetRenderScene() : nullptr;
 	const FPlanarReflectionSceneProxy* ReflectionSceneProxy = Scene ? Scene->GetForwardPassGlobalPlanarReflection() : nullptr;
 
@@ -756,7 +760,7 @@ TRDGUniformBufferRef<FOpaqueBasePassUniformParameters> CreateOpaqueBasePassUnifo
 	bool bLumenGIEnabled)
 {
 	FOpaqueBasePassUniformParameters& BasePassParameters = *GraphBuilder.AllocParameters<FOpaqueBasePassUniformParameters>();
-	SetupSharedBasePassParameters(GraphBuilder, View, bLumenGIEnabled, BasePassParameters.Shared);
+	SetupSharedBasePassParameters(GraphBuilder, View, ViewIndex, bLumenGIEnabled, BasePassParameters.Shared);
 
 	const FRDGSystemTextures& SystemTextures = FRDGSystemTextures::Get(GraphBuilder);
 
diff --git a/Engine/Source/Runtime/Renderer/Private/BasePassRendering.h b/Engine/Source/Runtime/Renderer/Private/BasePassRendering.h
--- a/Engine/Source/Runtime/Renderer/Private/BasePassRendering.h
+++ b/Engine/Source/Runtime/Renderer/Private/BasePassRendering.h
@@ -34,6 +34,8 @@
 #include "Nanite/NaniteMaterials.h"
 #include "BlueNoise.h"
 #include "LocalFogVolumeRendering.h"
+#include "LightFunctionAtlas.h"
+#include "RenderUtils.h"
 
 class FScene;
 
@@ -79,6 +81,7 @@ BEGIN_GLOBAL_SHADER_PARAMETER_STRUCT(FSharedBasePassUniformParameters,)
 	SHADER_PARAMETER_STRUCT(FFogUniformParameters, Fog)
 	SHADER_PARAMETER_STRUCT(FFogUniformParameters, FogISR)
 	SHADER_PARAMETER_STRUCT(FLocalFogVolumeUniformParameters, LFV)
+	SHADER_PARAMETER_STRUCT(FLightFunctionAtlasGlobalParameters, LightFunctionAtlas)
 	SHADER_PARAMETER(uint32, UseBasePassSkylight)
 END_GLOBAL_SHADER_PARAMETER_STRUCT()
 
@@ -141,6 +144,7 @@ DECLARE_GPU_STAT_NAMED_EXTERN(NaniteBasePass, TEXT("Nanite BasePass"));
 extern void SetupSharedBasePassParameters(
 	FRDGBuilder& GraphBuilder,
 	const FViewInfo& View,
+	const int32 ViewIndex,
 	bool bLumenGIEnabled,
 	FSharedBasePassUniformParameters& BasePassParameters);
 
@@ -466,6 +470,12 @@ public:
 
 		OutEnvironment.SetDefine(TEXT("COMPUTE_SHADED"), 1);
 
+		const bool bTranslucent = IsTranslucentBlendMode(Parameters.MaterialParameters);
+		const bool bIsSingleLayerWater = Parameters.MaterialParameters.ShadingModels.HasShadingModel(MSM_SingleLayerWater);
+		const bool bSingleLayerWaterUsesLightFunctionAtlas = bIsSingleLayerWater && GetSingleLayerWaterUsesLightFunctionAtlas();
+		const bool bTranslucentUsesLightFunctionAtlas = bTranslucent && GetTranslucentUsesLightFunctionAtlas();
+		OutEnvironment.SetDefine(TEXT("USE_LIGHT_FUNCTION_ATLAS"), (bSingleLayerWaterUsesLightFunctionAtlas || bTranslucentUsesLightFunctionAtlas) ? TEXT("1") : TEXT("0"));
+
 		TBasePassComputeShaderBaseType<LightMapPolicyType>::ModifyCompilationEnvironment(Parameters, OutEnvironment);
 	}
 	
@@ -618,6 +628,11 @@ public:
 			}
 		}
 
+		const bool bTranslucent = IsTranslucentBlendMode(Parameters.MaterialParameters);
+		const bool bSingleLayerWaterUsesLightFunctionAtlas = bIsSingleLayerWater && GetSingleLayerWaterUsesLightFunctionAtlas();
+		const bool bTranslucentUsesLightFunctionAtlas = bTranslucent && GetTranslucentUsesLightFunctionAtlas();
+		OutEnvironment.SetDefine(TEXT("USE_LIGHT_FUNCTION_ATLAS"), (bSingleLayerWaterUsesLightFunctionAtlas || bTranslucentUsesLightFunctionAtlas) ? TEXT("1") : TEXT("0"));
+
 		TBasePassPixelShaderBaseType<LightMapPolicyType>::ModifyCompilationEnvironment(Parameters, OutEnvironment);
 	}
 	
diff --git a/Engine/Source/Runtime/Renderer/Private/LightFunctionAtlas.cpp b/Engine/Source/Runtime/Renderer/Private/LightFunctionAtlas.cpp
--- a/Engine/Source/Runtime/Renderer/Private/LightFunctionAtlas.cpp
+++ b/Engine/Source/Runtime/Renderer/Private/LightFunctionAtlas.cpp
@@ -82,6 +82,7 @@ static TAutoConsoleVariable<int32> CVarLightFunctionAtlasSize(
 //////////////////////////////////////////////////////////////////////////
 
 // The CVars here represent systems that can request the creation/sampling of the light function atlas.
+// They do not require shader recompilation since they are handled via permutations
 
 // Volumetric fog always generate a light function for the directional light.
 // So this alias really only controls the use of the LightFunctionAtlas on the local lights.
@@ -221,14 +222,34 @@ void FLightFunctionAtlas::UpdateRegisterLightSceneInfo(FLightSceneInfo* LightSce
 	}
 }
 
-void FLightFunctionAtlas::BeginSceneFrame(FViewFamilyInfo& ViewFamily, TArray<FViewInfo>& Views, FLightFunctionAtlasSceneData& LightFunctionAtlasSceneData, bool bShouldRenderVolumetricFog)
+void FLightFunctionAtlas::ClearEmptySceneFrame(FViewInfo* View, FLightFunctionAtlasSceneData* LightFunctionAtlasSceneData)
 {
 	RegisteredLights.Empty(64);
-
-	// Reset UBs
+	DefaultLightFunctionAtlasGlobalParameters = nullptr;
 	DefaultLightFunctionAtlasGlobalParametersUB = nullptr;
+	ViewLightFunctionAtlasGlobalParametersArray.Empty(4);
 	ViewLightFunctionAtlasGlobalParametersUBArray.Empty(4);
 
+	bLightFunctionAtlasEnabled = false;
+	if (LightFunctionAtlasSceneData)
+	{
+		LightFunctionAtlasSceneData->SetData(
+			this,
+			bLightFunctionAtlasEnabled,
+			bLightFunctionAtlasEnabled,
+			bLightFunctionAtlasEnabled);
+	}
+
+	if (View && LightFunctionAtlasSceneData)
+	{
+		View->LightFunctionAtlasViewData = FLightFunctionAtlasViewData(LightFunctionAtlasSceneData);
+	}
+}
+
+void FLightFunctionAtlas::BeginSceneFrame(FViewFamilyInfo& ViewFamily, TArray<FViewInfo>& Views, FLightFunctionAtlasSceneData& LightFunctionAtlasSceneData, bool bShouldRenderVolumetricFog)
+{
+	ClearEmptySceneFrame();
+
 	// Now lets check if we need to generate the atlas for this frame
 	bLightFunctionAtlasEnabled = CVarLightFunctionAtlas.GetValueOnRenderThread() > 0 && ViewFamily.EngineShowFlags.LightFunctions > 0;
 
@@ -241,7 +262,9 @@ void FLightFunctionAtlas::BeginSceneFrame(FViewFamilyInfo& ViewFamily, TArray<FV
 
 		DeferredlightingRequestsLightFunctionAtlas = GDeferredUsesLightFunctionAtlas > 0;
 
-		bLightFunctionAtlasEnabled = bLightFunctionAtlasEnabled && (VolumetricFogRequestsLightFunctionAtlas || DeferredlightingRequestsLightFunctionAtlas); 
+		bLightFunctionAtlasEnabled = bLightFunctionAtlasEnabled && 
+			(VolumetricFogRequestsLightFunctionAtlas || DeferredlightingRequestsLightFunctionAtlas ||
+			 GetSingleLayerWaterUsesLightFunctionAtlas() || GetTranslucentUsesLightFunctionAtlas()); 
 	}
 
 	// We propagate bLightFunctionAtlasEnabled to all the views to ease later shader parameter decision and binding for lighting, shadow or volumetric fog for instance (avoid sending lots of parameters all over the place)
@@ -429,6 +452,23 @@ void FLightFunctionAtlas::AllocateAtlasSlots(const TArray<FViewInfo>& Views)
 	// TODO we could do all the constant buffer setup inline above (done in RenderAtlasSlots right now) if we would send a GraphBuilder here.
 }
 
+FLightFunctionAtlasGlobalParameters* FLightFunctionAtlas::GetLightFunctionAtlasGlobalParametersStruct(uint32 ViewIndex, FRDGBuilder& GraphBuilder)
+{
+	if (!IsLightFunctionAtlasEnabled())
+	{
+		return GetDefaultLightFunctionAtlasGlobalParametersStruct(GraphBuilder);
+	}
+
+	const bool bViewIndexIsValid = ViewIndex < uint32(ViewLightFunctionAtlasGlobalParametersUBArray.Num());
+	check(bViewIndexIsValid);
+	if (bViewIndexIsValid)
+	{
+		return ViewLightFunctionAtlasGlobalParametersArray[ViewIndex];
+	}
+
+	return GetDefaultLightFunctionAtlasGlobalParametersStruct(GraphBuilder);
+}
+
 TRDGUniformBufferRef<FLightFunctionAtlasGlobalParameters> FLightFunctionAtlas::GetLightFunctionAtlasGlobalParameters(uint32 ViewIndex, FRDGBuilder& GraphBuilder)
 {
 	if (!IsLightFunctionAtlasEnabled())
@@ -446,15 +486,24 @@ TRDGUniformBufferRef<FLightFunctionAtlasGlobalParameters> FLightFunctionAtlas::G
 	return GetDefaultLightFunctionAtlasGlobalParameters(GraphBuilder);
 }
 
+FLightFunctionAtlasGlobalParameters* FLightFunctionAtlas::GetDefaultLightFunctionAtlasGlobalParametersStruct(FRDGBuilder& GraphBuilder)
+{
+	if (DefaultLightFunctionAtlasGlobalParameters == nullptr) // Only create the default buffer once per frame
+	{
+		DefaultLightFunctionAtlasGlobalParameters = GraphBuilder.AllocParameters<FLightFunctionAtlasGlobalParameters>();
+		DefaultLightFunctionAtlasGlobalParameters->LightFunctionAtlasTexture = GSystemTextures.GetWhiteDummy(GraphBuilder);
+		DefaultLightFunctionAtlasGlobalParameters->LightFunctionAtlasSampler = TStaticSamplerState<SF_Bilinear>::GetRHI();
+		DefaultLightFunctionAtlasGlobalParameters->Slot_UVSize = 1.0f;
+	}
+	return DefaultLightFunctionAtlasGlobalParameters;
+}
+
 TRDGUniformBufferRef<FLightFunctionAtlasGlobalParameters> FLightFunctionAtlas::GetDefaultLightFunctionAtlasGlobalParameters(FRDGBuilder& GraphBuilder)
 {
 	if (DefaultLightFunctionAtlasGlobalParametersUB == nullptr) // Only create the default buffer once per frame
 	{
-		FLightFunctionAtlasGlobalParameters* LightFunctionAtlasGlobalParams = GraphBuilder.AllocParameters<FLightFunctionAtlasGlobalParameters>();
-		LightFunctionAtlasGlobalParams->LightFunctionAtlasTexture = GSystemTextures.GetWhiteDummy(GraphBuilder);
-		LightFunctionAtlasGlobalParams->LightFunctionAtlasSampler = TStaticSamplerState<SF_Bilinear>::GetRHI();
-		LightFunctionAtlasGlobalParams->Slot_UVSize = 1.0f;
-		DefaultLightFunctionAtlasGlobalParametersUB = GraphBuilder.CreateUniformBuffer(LightFunctionAtlasGlobalParams);
+		FLightFunctionAtlasGlobalParameters* DefaultLightFunctionAtlasGlobalParametersStruct = GetDefaultLightFunctionAtlasGlobalParametersStruct(GraphBuilder);
+		DefaultLightFunctionAtlasGlobalParametersUB = GraphBuilder.CreateUniformBuffer(DefaultLightFunctionAtlasGlobalParametersStruct);
 	}
 	return DefaultLightFunctionAtlasGlobalParametersUB;
 }
@@ -571,6 +620,7 @@ void FLightFunctionAtlas::RenderLightFunctionAtlas(FRDGBuilder& GraphBuilder, TA
 			LightIndex++;
 		}
 
+		ViewLightFunctionAtlasGlobalParametersArray.Add(LightFunctionAtlasGlobalParameters);
 		ViewLightFunctionAtlasGlobalParametersUBArray.Add(GraphBuilder.CreateUniformBuffer(LightFunctionAtlasGlobalParameters));
 	}
 }
diff --git a/Engine/Source/Runtime/Renderer/Private/LightFunctionAtlas.h b/Engine/Source/Runtime/Renderer/Private/LightFunctionAtlas.h
--- a/Engine/Source/Runtime/Renderer/Private/LightFunctionAtlas.h
+++ b/Engine/Source/Runtime/Renderer/Private/LightFunctionAtlas.h
@@ -32,7 +32,7 @@ struct FLightFunctionAtlasSceneData
 		bDeferredlightingUsesLightFunctionAtlas = bInDeferredlightingUsesLightFunctionAtlas;
 	}
 
-	const FLightFunctionAtlas* GetLightFunctionAtlas()  const { return LightFunctionAtlas; }
+	FLightFunctionAtlas* GetLightFunctionAtlas()		const { return LightFunctionAtlas; }
 	bool GetLightFunctionAtlasEnabled()					const { return bLightFunctionAtlasEnabled; }
 	bool GetVolumetricFogUsesLightFunctionAtlas()		const { return bVolumetricFogUsesLightFunctionAtlas; }
 	bool GetDeferredlightingUsesLightFunctionAtlas()	const { return bDeferredlightingUsesLightFunctionAtlas; }
@@ -49,7 +49,8 @@ struct FLightFunctionAtlasViewData
 	FLightFunctionAtlasViewData() : SceneData(nullptr) {}
 	FLightFunctionAtlasViewData(FLightFunctionAtlasSceneData* InSceneData) : SceneData(InSceneData) {}
 
-	bool GetLightFunctionAtlasEnabled()					const { return SceneData ? SceneData->GetLightFunctionAtlasEnabled()				: false; }
+	FLightFunctionAtlas* GetLightFunctionAtlas()		const { return SceneData ? SceneData->GetLightFunctionAtlas() : nullptr; }
+	bool GetLightFunctionAtlasEnabled()					const { return SceneData ? SceneData->GetLightFunctionAtlasEnabled() : false; }
 	bool GetVolumetricFogUsesLightFunctionAtlas()		const { return SceneData ? SceneData->GetVolumetricFogUsesLightFunctionAtlas()		: false; }
 	bool GetDeferredlightingUsesLightFunctionAtlas()	const { return SceneData ? SceneData->GetDeferredlightingUsesLightFunctionAtlas()	: false; }
 
@@ -117,6 +118,8 @@ struct FLightFunctionAtlas
 
 	bool IsLightFunctionAtlasEnabled() const { return bLightFunctionAtlasEnabled; }
 
+	void ClearEmptySceneFrame(FViewInfo* View = nullptr, FLightFunctionAtlasSceneData* LightFunctionAtlasSceneData = nullptr);
+
 	void BeginSceneFrame(FViewFamilyInfo& ViewFamily, TArray<FViewInfo>& Views, FLightFunctionAtlasSceneData& LightFunctionAtlasSceneData, bool bShouldRenderVolumetricFog);
 
 	void UpdateRegisterLightSceneInfo(FLightSceneInfo* LightSceneInfo);
@@ -128,8 +131,10 @@ struct FLightFunctionAtlas
 
 	FScreenPassTexture AddDebugVisualizationPasses(FRDGBuilder& GraphBuilder, const FViewInfo& View, FScreenPassTexture& ScreenPassSceneColor)  const;
 
-	TRDGUniformBufferRef<FLightFunctionAtlasGlobalParameters> GetLightFunctionAtlasGlobalParameters(uint32 ViewIndex, FRDGBuilder& GraphBuilder);
-	TRDGUniformBufferRef<FLightFunctionAtlasGlobalParameters> GetDefaultLightFunctionAtlasGlobalParameters(FRDGBuilder& GraphBuilder);
+	FLightFunctionAtlasGlobalParameters*						GetLightFunctionAtlasGlobalParametersStruct(uint32 ViewIndex, FRDGBuilder& GraphBuilder);
+	TRDGUniformBufferRef<FLightFunctionAtlasGlobalParameters>	GetLightFunctionAtlasGlobalParameters(uint32 ViewIndex, FRDGBuilder& GraphBuilder);
+	FLightFunctionAtlasGlobalParameters*						GetDefaultLightFunctionAtlasGlobalParametersStruct(FRDGBuilder& GraphBuilder);
+	TRDGUniformBufferRef<FLightFunctionAtlasGlobalParameters>	GetDefaultLightFunctionAtlasGlobalParameters(FRDGBuilder& GraphBuilder);
 
 private:
 
@@ -167,6 +172,9 @@ private:
 	};
 	TArray<EffectiveLocalLightSlot> EffectiveLocalLightSlotArray;
 
-	TRDGUniformBufferRef<FLightFunctionAtlasGlobalParameters>	DefaultLightFunctionAtlasGlobalParametersUB;
-	TArray<TRDGUniformBufferRef<FLightFunctionAtlasGlobalParameters>> ViewLightFunctionAtlasGlobalParametersUBArray;
+	FLightFunctionAtlasGlobalParameters*								DefaultLightFunctionAtlasGlobalParameters = nullptr;
+	TRDGUniformBufferRef<FLightFunctionAtlasGlobalParameters>			DefaultLightFunctionAtlasGlobalParametersUB;
+
+	TArray<FLightFunctionAtlasGlobalParameters*>							ViewLightFunctionAtlasGlobalParametersArray;
+	TArray<TRDGUniformBufferRef<FLightFunctionAtlasGlobalParameters>>	ViewLightFunctionAtlasGlobalParametersUBArray;
 };
diff --git a/Engine/Source/Runtime/Renderer/Private/LightGridInjection.cpp b/Engine/Source/Runtime/Renderer/Private/LightGridInjection.cpp
--- a/Engine/Source/Runtime/Renderer/Private/LightGridInjection.cpp
+++ b/Engine/Source/Runtime/Renderer/Private/LightGridInjection.cpp
@@ -139,6 +139,8 @@ void SetupDummyForwardLightUniformParameters(FRDGBuilder& GraphBuilder, FForward
 	}
 	ForwardLightData.CulledLightDataGrid32Bit = CulledLightDataGridSRV;
 	ForwardLightData.CulledLightDataGrid16Bit = CulledLightDataGridSRV;
+
+	ForwardLightData.LightFunctionAtlasLightIndex = 0;
 }
 
 TRDGUniformBufferRef<FForwardLightData> CreateDummyForwardLightUniformBuffer(FRDGBuilder& GraphBuilder, EShaderPlatform ShaderPlatform)
@@ -653,6 +655,7 @@ FComputeLightGridOutput FSceneRenderer::ComputeLightGrid(FRDGBuilder& GraphBuild
 							ForwardLightData->DirectionalLightShadowMapChannelMask = LightTypeAndShadowMapChannelMaskPacked;
 							ForwardLightData->DirectionalLightVSM = INDEX_NONE;
 							ForwardLightData->DirectionalLightSMRTSettings = GetVirtualShadowMapSMRTSettings(true);
+							ForwardLightData->LightFunctionAtlasLightIndex = LightParameters.LightFunctionAtlasLightIndex;
 
 							const FVector2D FadeParams = LightProxy->GetDirectionalLightDistanceFadeParameters(View.GetFeatureLevel(), LightSceneInfo->IsPrecomputedLightingValid(), View.MaxShadowCascades);
 
diff --git a/Engine/Source/Runtime/Renderer/Private/Renderer.cpp b/Engine/Source/Runtime/Renderer/Private/Renderer.cpp
--- a/Engine/Source/Runtime/Renderer/Private/Renderer.cpp
+++ b/Engine/Source/Runtime/Renderer/Private/Renderer.cpp
@@ -47,6 +47,7 @@
 #include "LocalFogVolumeRendering.h"
 #include "Rendering/RayTracingGeometryManager.h"
 #include "PathTracing.h"
+#include "LightFunctionAtlas.h"
 
 DEFINE_LOG_CATEGORY(LogRenderer);
 
@@ -337,6 +338,11 @@ void FRendererModule::DrawTileMesh(FCanvasRenderContext& RenderContext, FMeshPas
 		View.ForwardLightingResources.SetUniformBuffer(CreateDummyForwardLightUniformBuffer(GraphBuilder, View.GetShaderPlatform()));
 		SetDummyLocalFogVolumeForView(GraphBuilder, View);
 
+		// Create a disabled LightFunctionAtlas to be able to render base pass.
+		FLightFunctionAtlas LightFunctionAtlas;
+		FLightFunctionAtlasSceneData LightFunctionAtlasSceneData;
+		LightFunctionAtlas.ClearEmptySceneFrame(&View, &LightFunctionAtlasSceneData);
+
 		TUniformBufferRef<FReflectionCaptureShaderData> EmptyReflectionCaptureUniformBuffer;
 
 		{
diff --git a/Engine/Source/Runtime/Renderer/Private/SceneRendering.h b/Engine/Source/Runtime/Renderer/Private/SceneRendering.h
--- a/Engine/Source/Runtime/Renderer/Private/SceneRendering.h
+++ b/Engine/Source/Runtime/Renderer/Private/SceneRendering.h
@@ -753,6 +753,7 @@ BEGIN_GLOBAL_SHADER_PARAMETER_STRUCT_WITH_CONSTRUCTOR(FForwardLightData, )
 	SHADER_PARAMETER(FVector4f, DirectionalLightStaticShadowBufferSize)
 	SHADER_PARAMETER(FMatrix44f, DirectionalLightTranslatedWorldToStaticShadow)
 	SHADER_PARAMETER(uint32, DirectLightingShowFlag)
+	SHADER_PARAMETER(uint32, LightFunctionAtlasLightIndex)
 	SHADER_PARAMETER_STRUCT(FVirtualShadowMapSMRTSettings, DirectionalLightSMRTSettings)
 	SHADER_PARAMETER_RDG_TEXTURE(Texture2D, DirectionalLightShadowmapAtlas)
 	SHADER_PARAMETER_SAMPLER(SamplerState, ShadowmapSampler)
diff --git a/Engine/Source/Runtime/Renderer/Private/TranslucentRendering.cpp b/Engine/Source/Runtime/Renderer/Private/TranslucentRendering.cpp
--- a/Engine/Source/Runtime/Renderer/Private/TranslucentRendering.cpp
+++ b/Engine/Source/Runtime/Renderer/Private/TranslucentRendering.cpp
@@ -823,7 +823,7 @@ TRDGUniformBufferRef<FTranslucentBasePassUniformParameters> CreateTranslucentBas
 		return GraphBuilder.RegisterExternalTexture(PooledRenderTarget, Flags);
 	};
 
-	SetupSharedBasePassParameters(GraphBuilder, View, bLumenGIEnabled, BasePassParameters.Shared);
+	SetupSharedBasePassParameters(GraphBuilder, View, ViewIndex, bLumenGIEnabled, BasePassParameters.Shared);
 	SetupSceneTextureUniformParameters(GraphBuilder, View.GetSceneTexturesChecked(), View.FeatureLevel, SceneTextureSetupMode, BasePassParameters.SceneTextures);
 	Substrate::BindSubstrateForwardPasslUniformParameters(GraphBuilder, View, BasePassParameters.Substrate);
 
