Unified and fixed depth bound test to work on all platofrm for LFV.

diff --git a/Engine/Source/Runtime/Renderer/Private/LocalFogVolumeRendering.cpp b/Engine/Source/Runtime/Renderer/Private/LocalFogVolumeRendering.cpp
--- a/Engine/Source/Runtime/Renderer/Private/LocalFogVolumeRendering.cpp
+++ b/Engine/Source/Runtime/Renderer/Private/LocalFogVolumeRendering.cpp
@@ -720,9 +720,17 @@ BEGIN_SHADER_PARAMETER_STRUCT(FLocalFogVolumeTiledPassParameters, )
 	RENDER_TARGET_BINDING_SLOTS()
 END_SHADER_PARAMETER_STRUCT()
 
-static float ApplyDepthBoundIfNeeded(FGraphicsPipelineStateInitializer& GraphicsPSOInit, FMatrix ViewProjectionMatrix, FMatrix ViewInvProjectionMatrix, FRHICommandList& RHICmdList)
+struct FDepthBoundSetup
 {
-	float FogStartDistance = GetLocalFogVolumeGlobalStartDistance();
+	bool bEnabled = false;
+	float FogClipDeviceZ = 0.0f;
+	float MinDeviceZ = 0.0f;
+	float MaxDeviceZ = 1.0f;
+};
+
+static FDepthBoundSetup GetDepthBoundSetup(float FogStartDistance, FGraphicsPipelineStateInitializer& GraphicsPSOInit, FMatrix ViewProjectionMatrix, FMatrix ViewInvProjectionMatrix)
+{
+	FDepthBoundSetup DepthBoundSetup;
 
 	// Here we compute the nearest z value the fog can start
 	// to skip shader execution on pixels that are closer.
@@ -738,22 +746,21 @@ static float ApplyDepthBoundIfNeeded(FGraphicsPipelineStateInitializer& Graphics
 	FVector ViewSpaceStartFogPoint(0.0f, 0.0f, FogStartDistance * Ratio);
 	FVector4f ClipSpaceMaxDistance = (FVector4f)ViewProjectionMatrix.TransformPosition(ViewSpaceStartFogPoint); // LWC_TODO: precision loss
 	float FogClipSpaceZ = ClipSpaceMaxDistance.Z / ClipSpaceMaxDistance.W;
-	FogClipSpaceZ = FMath::Clamp(FogClipSpaceZ, 0.f, 1.f);
+	DepthBoundSetup.FogClipDeviceZ = FMath::Clamp(FogClipSpaceZ, 0.f, 1.f);
 
-	GraphicsPSOInit.bDepthBounds = GSupportsDepthBoundsTest;
-	if (GraphicsPSOInit.bDepthBounds)
+	DepthBoundSetup.bEnabled = GSupportsDepthBoundsTest;
+	if (bool(ERHIZBuffer::IsInverted))
 	{
-		if (bool(ERHIZBuffer::IsInverted))
-		{
-			RHICmdList.SetDepthBounds(0.0f, FogClipSpaceZ);
-		}
-		else
-		{
-			RHICmdList.SetDepthBounds(FogClipSpaceZ, 1.0f);
-		}
+		DepthBoundSetup.MinDeviceZ = 0.0f;
+		DepthBoundSetup.MaxDeviceZ = FogClipSpaceZ;
+	}
+	else
+	{
+		DepthBoundSetup.MinDeviceZ = FogClipSpaceZ;
+		DepthBoundSetup.MaxDeviceZ = 1.0f;
 	}
 
-	return FogClipSpaceZ;
+	return DepthBoundSetup;
 }
  
 void RenderLocalFogVolume(
@@ -829,10 +836,13 @@ void RenderLocalFogVolume(
 				GraphicsPSOInit.BoundShaderState.VertexShaderRHI = VertexShader.GetVertexShader();
 				GraphicsPSOInit.BoundShaderState.PixelShaderRHI = PixelShader.GetPixelShader();
 
-				PassParameters->VS.StartDepthZ = ApplyDepthBoundIfNeeded(GraphicsPSOInit, ViewProjectionMatrix, ViewInvProjectionMatrix, RHICmdList);
+				FDepthBoundSetup DepthBoundSetup = GetDepthBoundSetup(GetLocalFogVolumeGlobalStartDistance(), GraphicsPSOInit, ViewProjectionMatrix, ViewInvProjectionMatrix);
+				GraphicsPSOInit.bDepthBounds = DepthBoundSetup.bEnabled;
 
 				SetGraphicsPipelineState(RHICmdList, GraphicsPSOInit, 0);
 
+				RHICmdList.SetDepthBounds(DepthBoundSetup.MinDeviceZ, DepthBoundSetup.MaxDeviceZ);
+
 				SetShaderParameters(RHICmdList, VertexShader, VertexShader.GetVertexShader(), PassParameters->VS);
 				SetShaderParameters(RHICmdList, PixelShader, PixelShader.GetPixelShader(), PassParameters->PS);
 
@@ -986,17 +996,19 @@ void RenderLocalFogVolumeMobile(
 	GraphicsPSOInit.BoundShaderState.VertexShaderRHI = VertexShader.GetVertexShader();
 	GraphicsPSOInit.BoundShaderState.PixelShaderRHI = PixelShader.GetPixelShader();
 
-	float StartDepthZ = ApplyDepthBoundIfNeeded(GraphicsPSOInit, View.ViewMatrices.GetProjectionMatrix(), View.ViewMatrices.GetInvProjectionMatrix(), RHICmdList);
+	FDepthBoundSetup DepthBoundSetup = GetDepthBoundSetup(GetLocalFogVolumeGlobalStartDistance(), GraphicsPSOInit, View.ViewMatrices.GetProjectionMatrix(), View.ViewMatrices.GetInvProjectionMatrix());
+	GraphicsPSOInit.bDepthBounds = DepthBoundSetup.bEnabled;
 
 	SetGraphicsPipelineState(RHICmdList, GraphicsPSOInit, 0);
 
+	RHICmdList.SetDepthBounds(DepthBoundSetup.MinDeviceZ, DepthBoundSetup.MaxDeviceZ);
 
 	FMobileLocalFogVolumeTiledRenderVS::FParameters VSParameters;
 	VSParameters.View = View.GetShaderParameters();
 	VSParameters.LFV = View.LocalFogVolumeViewData.UniformParametersStruct;
 	VSParameters.TileDataBuffer = View.LocalFogVolumeViewData.GPUTileDataBufferSRV;
 	VSParameters.TileDrawIndirectBuffer = View.LocalFogVolumeViewData.GPUTileDrawIndirectBuffer;
-	VSParameters.StartDepthZ = StartDepthZ;
+	VSParameters.StartDepthZ = DepthBoundSetup.FogClipDeviceZ;
 	SetShaderParameters(RHICmdList, VertexShader, VertexShader.GetVertexShader(), VSParameters);
 
 	FMobileLocalFogVolumeTiledRenderPS::FParameters PSParameters;
