LFAtlas - added support for light function atlas on clustered shading passes.

diff --git a/Engine/Source/Runtime/Renderer/Private/ClusteredDeferredShadingPass.cpp b/Engine/Source/Runtime/Renderer/Private/ClusteredDeferredShadingPass.cpp
--- a/Engine/Source/Runtime/Renderer/Private/ClusteredDeferredShadingPass.cpp
+++ b/Engine/Source/Runtime/Renderer/Private/ClusteredDeferredShadingPass.cpp
@@ -28,6 +28,7 @@
 #include "PostProcess/SceneFilterRendering.h"
 #include "PostProcess/PostProcessing.h"
 #include "Substrate/Substrate.h"
+#include "LightFunctionAtlas.h"
 
 // This is used to switch on and off the clustered deferred shading implementation, that uses the light grid to perform shading.
 int32 GUseClusteredDeferredShading = 0;
@@ -96,7 +97,8 @@ class FClusteredShadingPS : public FGlobalShader
 	class FVisualizeLightCullingDim : SHADER_PERMUTATION_BOOL("VISUALIZE_LIGHT_CULLING");
 	class FHairStrandsLighting : SHADER_PERMUTATION_BOOL("USE_HAIR_LIGHTING");
 	class FSubstrateTileType : SHADER_PERMUTATION_INT("SUBSTRATE_TILETYPE", 4);
-	using FPermutationDomain = TShaderPermutationDomain<FVisualizeLightCullingDim, FHairStrandsLighting, FSubstrateTileType>;
+	class FLightFunctionAtlasDim : SHADER_PERMUTATION_BOOL("USE_LIGHT_FUNCTION_ATLAS");
+	using FPermutationDomain = TShaderPermutationDomain<FVisualizeLightCullingDim, FHairStrandsLighting, FSubstrateTileType, FLightFunctionAtlasDim>;
 
 	BEGIN_SHADER_PARAMETER_STRUCT(FParameters, )
 		SHADER_PARAMETER_RDG_UNIFORM_BUFFER(FForwardLightData, Forward)
@@ -105,6 +107,7 @@ class FClusteredShadingPS : public FGlobalShader
 		SHADER_PARAMETER_RDG_UNIFORM_BUFFER(FHairStrandsViewUniformParameters, HairStrands)
 		SHADER_PARAMETER_RDG_UNIFORM_BUFFER(FSubstrateGlobalUniformParameters, Substrate)
 		SHADER_PARAMETER_RDG_UNIFORM_BUFFER(FSceneTextureUniformParameters, SceneTextures)
+		SHADER_PARAMETER_RDG_UNIFORM_BUFFER(FLightFunctionAtlasGlobalParameters, LightFunctionAtlas)
 		SHADER_PARAMETER_RDG_TEXTURE(Texture2D, ShadowMaskBits)
 		SHADER_PARAMETER_STRUCT_INCLUDE(FVirtualShadowMapSamplingParameters, VirtualShadowMapSamplingParameters)
 		SHADER_PARAMETER_STRUCT_INCLUDE(Substrate::FSubstrateTilePassVS::FParameters, SubstrateTile)
@@ -145,6 +148,7 @@ enum class EClusterPassInputType : uint8
 
 static void InternalAddClusteredDeferredShadingPass(
 	FRDGBuilder& GraphBuilder,
+	int32 ViewIndex,
 	FViewInfo& View,
 	const FMinimalSceneTextures& SceneTextures,
 	const FSortedLightSetSceneInfo &SortedLightsSet,
@@ -153,7 +157,8 @@ static void InternalAddClusteredDeferredShadingPass(
 	FRDGTextureRef ShadowMaskBits,
 	FVirtualShadowMapArray& VirtualShadowMapArray,
 	FRDGBufferSRVRef HairTransmittanceBuffer,
-	FSubstrateSceneData* SubstrateSceneData)
+	FSubstrateSceneData* SubstrateSceneData,
+	FLightFunctionAtlas& LightFunctionAtlas)
 {
 	check(SortedLightsSet.ClusteredSupportedEnd > 0);
 	const FIntPoint SceneTextureExtent = SceneTextures.Config.Extent;
@@ -169,6 +174,7 @@ static void InternalAddClusteredDeferredShadingPass(
 	PassParameters->VirtualShadowMapSamplingParameters = VirtualShadowMapArray.GetSamplingParameters(GraphBuilder);
 	PassParameters->HairTransmittanceBuffer = HairTransmittanceBuffer;
 	PassParameters->Substrate = Substrate::BindSubstrateGlobalUniformParameters(View);
+	PassParameters->LightFunctionAtlas = LightFunctionAtlas.GetLightFunctionAtlasGlobalParameters(ViewIndex, GraphBuilder);
 	ShaderPrint::SetParameters(GraphBuilder, View.ShaderPrintData, PassParameters->ShaderPrintUniformBuffer);
 
 	PassParameters->RenderTargets[0] = FRenderTargetBinding(SceneTextures.Color.Target, ERenderTargetLoadAction::ELoad);
@@ -206,6 +212,8 @@ static void InternalAddClusteredDeferredShadingPass(
 		PermutationVector.Set<FClusteredShadingPS::FVisualizeLightCullingDim>(View.Family->EngineShowFlags.VisualizeLightCulling);
 		PermutationVector.Set<FClusteredShadingPS::FHairStrandsLighting>(bHairStrands);
 		PermutationVector.Set<FClusteredShadingPS::FSubstrateTileType>(bSubstrate ? TileType : 0);
+		PermutationVector.Set<FClusteredShadingPS::FLightFunctionAtlasDim>(
+			View.LightFunctionAtlasViewData.GetDeferredlightingUsesLightFunctionAtlas() && View.Family->EngineShowFlags.LightFunctions);
 		TShaderMapRef<FClusteredShadingPS> PixelShader(View.ShaderMap, PermutationVector);
 		{
 			FGraphicsPipelineStateInitializer GraphicsPSOInit;
@@ -292,6 +300,7 @@ void FDeferredShadingSceneRenderer::AddClusteredDeferredShadingPass(
 
 				InternalAddClusteredDeferredShadingPass(
 					GraphBuilder,
+					ViewIndex,
 					View,
 					SceneTextures,
 					SortedLightsSet,
@@ -300,10 +309,12 @@ void FDeferredShadingSceneRenderer::AddClusteredDeferredShadingPass(
 					ShadowMaskBits,
 					VirtualShadowMapArray,
 					nullptr,
-					SubstrateSceneData);
+					SubstrateSceneData,
+					LightFunctionAtlas);
 
 				InternalAddClusteredDeferredShadingPass(
 					GraphBuilder,
+					ViewIndex,
 					View,
 					SceneTextures,
 					SortedLightsSet,
@@ -312,10 +323,12 @@ void FDeferredShadingSceneRenderer::AddClusteredDeferredShadingPass(
 					ShadowMaskBits,
 					VirtualShadowMapArray,
 					nullptr,
-					SubstrateSceneData);
+					SubstrateSceneData,
+					LightFunctionAtlas);
 
 				InternalAddClusteredDeferredShadingPass(
 					GraphBuilder,
+					ViewIndex,
 					View,
 					SceneTextures,
 					SortedLightsSet,
@@ -324,10 +337,12 @@ void FDeferredShadingSceneRenderer::AddClusteredDeferredShadingPass(
 					ShadowMaskBits,
 					VirtualShadowMapArray,
 					nullptr,
-					SubstrateSceneData);
+					SubstrateSceneData,
+					LightFunctionAtlas);
 
 				InternalAddClusteredDeferredShadingPass(
 					GraphBuilder,
+					ViewIndex,
 					View,
 					SceneTextures,
 					SortedLightsSet,
@@ -336,12 +351,14 @@ void FDeferredShadingSceneRenderer::AddClusteredDeferredShadingPass(
 					ShadowMaskBits,
 					VirtualShadowMapArray,
 					nullptr, 
-					SubstrateSceneData);
+					SubstrateSceneData,
+					LightFunctionAtlas);
 			}
 			else
 			{
 				InternalAddClusteredDeferredShadingPass(
 					GraphBuilder,
+					ViewIndex,
 					View,
 					SceneTextures,
 					SortedLightsSet,
@@ -350,7 +367,8 @@ void FDeferredShadingSceneRenderer::AddClusteredDeferredShadingPass(
 					ShadowMaskBits,
 					VirtualShadowMapArray,
 					nullptr,
-					SubstrateSceneData);
+					SubstrateSceneData,
+					LightFunctionAtlas);
 			}
 
 			if (HairStrands::HasViewHairStrandsData(View))
@@ -358,6 +376,7 @@ void FDeferredShadingSceneRenderer::AddClusteredDeferredShadingPass(
 				FHairStrandsTransmittanceMaskData TransmittanceMask = RenderHairStrandsOnePassTransmittanceMask(GraphBuilder, View, HairStrandsShadowMaskBits, VirtualShadowMapArray);
 				InternalAddClusteredDeferredShadingPass(
 					GraphBuilder,
+					ViewIndex,
 					View,
 					SceneTextures,
 					SortedLightsSet,
@@ -366,7 +385,8 @@ void FDeferredShadingSceneRenderer::AddClusteredDeferredShadingPass(
 					HairStrandsShadowMaskBits,
 					VirtualShadowMapArray,
 					GraphBuilder.CreateSRV(TransmittanceMask.TransmittanceMask, FHairStrandsTransmittanceMaskData::Format),
-					SubstrateSceneData);
+					SubstrateSceneData,
+					LightFunctionAtlas);
 			}
 		}
 	}
diff --git a/Engine/Source/Runtime/Renderer/Private/LightFunctionAtlas.cpp b/Engine/Source/Runtime/Renderer/Private/LightFunctionAtlas.cpp
--- a/Engine/Source/Runtime/Renderer/Private/LightFunctionAtlas.cpp
+++ b/Engine/Source/Runtime/Renderer/Private/LightFunctionAtlas.cpp
@@ -44,15 +44,10 @@ How to use the Atlas:
  - Call GetLocalLightFunctionCommon(DerivedParams.TranslatedWorldPosition, LightData.LightFunctionAtlasLightIndex); where LightData is a FDeferredLightData recovered from uniform or the light grid.
 
  What is next:
- - Atlas debug
- - RGB option on project + shader changes with defined type
- - Forward shading option
  - Super sample CVAR
  - Convert systems:
     - Lumen
     - Path tracer
-    - Clustered shading
-    - SLW
     - Check those path with Substrate
 
 */
