Deprecated IPersistentViewUniformBufferExtension
- removed calling IPersistentViewUniformBufferExtension::BeginRenderView() from GatherRayTracingWorldInstancesForView()

diff --git a/Engine/Source/Runtime/RenderCore/Public/RendererInterface.h b/Engine/Source/Runtime/RenderCore/Public/RendererInterface.h
--- a/Engine/Source/Runtime/RenderCore/Public/RendererInterface.h
+++ b/Engine/Source/Runtime/RenderCore/Public/RendererInterface.h
@@ -845,6 +845,7 @@ public:
 	/**	Start prefetching streaming data for Nanite resource that will soon be used for rendering. TODO: Implement callback mechanism */
 	virtual void PrefetchNaniteResource(const Nanite::FResources* Resource, uint32 NumFramesUntilRender) = 0;
 
+	UE_DEPRECATED(5.4, "IPersistentViewUniformBufferExtension will be removed in a future version.")
 	virtual void RegisterPersistentViewUniformBufferExtension(IPersistentViewUniformBufferExtension* Extension) = 0;
 
 	/**
diff --git a/Engine/Source/Runtime/Renderer/Private/DeferredShadingRenderer.cpp b/Engine/Source/Runtime/Renderer/Private/DeferredShadingRenderer.cpp
--- a/Engine/Source/Runtime/Renderer/Private/DeferredShadingRenderer.cpp
+++ b/Engine/Source/Runtime/Renderer/Private/DeferredShadingRenderer.cpp
@@ -1089,13 +1089,6 @@ bool FDeferredShadingSceneRenderer::GatherRayTracingWorldInstancesForView(FRDGBu
 	View.DynamicRayTracingMeshCommandStorage.Reserve(Scene->Primitives.Num());
 	View.VisibleRayTracingMeshCommands.Reserve(Scene->Primitives.Num());
 
-	extern TSet<IPersistentViewUniformBufferExtension*> PersistentViewUniformBufferExtensions;
-
-	for (IPersistentViewUniformBufferExtension* Extension : PersistentViewUniformBufferExtensions)
-	{
-		Extension->BeginRenderView(&View);
-	}
-
 	View.RayTracingMeshResourceCollector = MakeUnique<FRayTracingMeshResourceCollector>(Scene->GetFeatureLevel(), Allocator);
 
 	View.RayTracingCullingParameters.Init(View);
