LFAtlas - removed duplicated checks of show flag.

diff --git a/Engine/Source/Runtime/Renderer/Private/ClusteredDeferredShadingPass.cpp b/Engine/Source/Runtime/Renderer/Private/ClusteredDeferredShadingPass.cpp
--- a/Engine/Source/Runtime/Renderer/Private/ClusteredDeferredShadingPass.cpp
+++ b/Engine/Source/Runtime/Renderer/Private/ClusteredDeferredShadingPass.cpp
@@ -212,8 +212,7 @@ static void InternalAddClusteredDeferredShadingPass(
 		PermutationVector.Set<FClusteredShadingPS::FVisualizeLightCullingDim>(View.Family->EngineShowFlags.VisualizeLightCulling);
 		PermutationVector.Set<FClusteredShadingPS::FHairStrandsLighting>(bHairStrands);
 		PermutationVector.Set<FClusteredShadingPS::FSubstrateTileType>(bSubstrate ? TileType : 0);
-		PermutationVector.Set<FClusteredShadingPS::FLightFunctionAtlasDim>(
-			View.LightFunctionAtlasViewData.GetDeferredlightingUsesLightFunctionAtlas() && View.Family->EngineShowFlags.LightFunctions);
+		PermutationVector.Set<FClusteredShadingPS::FLightFunctionAtlasDim>(View.LightFunctionAtlasViewData.GetDeferredlightingUsesLightFunctionAtlas());
 		TShaderMapRef<FClusteredShadingPS> PixelShader(View.ShaderMap, PermutationVector);
 		{
 			FGraphicsPipelineStateInitializer GraphicsPSOInit;
diff --git a/Engine/Source/Runtime/Renderer/Private/LightRendering.cpp b/Engine/Source/Runtime/Renderer/Private/LightRendering.cpp
--- a/Engine/Source/Runtime/Renderer/Private/LightRendering.cpp
+++ b/Engine/Source/Runtime/Renderer/Private/LightRendering.cpp
@@ -2621,7 +2621,7 @@ static void RenderLight(
 		PermutationVector.Set< FDeferredLightPS::FHairComplexTransmittance >(bNeedComplexTransmittanceSupport);
 		PermutationVector.Set< FDeferredLightPS::FLightFunctionAtlasDim >(
 			View.LightFunctionAtlasViewData.GetDeferredlightingUsesLightFunctionAtlas() && LightSceneInfo->Proxy->HasValidLightFunctionAtlasSlot() &&
-			LightSceneInfo->Proxy->GetLightFunctionMaterial() != nullptr && View.Family->EngineShowFlags.LightFunctions && !View.Family->EngineShowFlags.VisualizeLightCulling);
+			LightSceneInfo->Proxy->GetLightFunctionMaterial() != nullptr && !View.Family->EngineShowFlags.VisualizeLightCulling);
 		if (bIsRadial)
 		{
 			PermutationVector.Set< FDeferredLightPS::FSourceShapeDim >(LightProxy->IsRectLight() ? ELightSourceShape::Rect : ELightSourceShape::Capsule);
@@ -2770,7 +2770,7 @@ void FDeferredShadingSceneRenderer::RenderLightForHair(
 	PermutationVector.Set< FDeferredLightPS::FHairComplexTransmittance>(true);
 	PermutationVector.Set< FDeferredLightPS::FLightFunctionAtlasDim >(
 		View.LightFunctionAtlasViewData.GetDeferredlightingUsesLightFunctionAtlas() && LightSceneInfo->Proxy->HasValidLightFunctionAtlasSlot() &&
-		LightSceneInfo->Proxy->GetLightFunctionMaterial() != nullptr && View.Family->EngineShowFlags.LightFunctions && !View.Family->EngineShowFlags.VisualizeLightCulling);
+		LightSceneInfo->Proxy->GetLightFunctionMaterial() != nullptr && !View.Family->EngineShowFlags.VisualizeLightCulling);
 	if (bIsDirectional)
 	{
 		PermutationVector.Set< FDeferredLightPS::FSourceShapeDim >(ELightSourceShape::Directional);
