PR #11130: [Render] SceneUniformBuffer.h - Added missing dll exports, to be used outside Renderer module

#11130

diff --git a/Engine/Source/Runtime/Renderer/Private/SceneUniformBuffer.cpp b/Engine/Source/Runtime/Renderer/Private/SceneUniformBuffer.cpp
--- a/Engine/Source/Runtime/Renderer/Private/SceneUniformBuffer.cpp
+++ b/Engine/Source/Runtime/Renderer/Private/SceneUniformBuffer.cpp
@@ -60,6 +60,11 @@ TArray<FSceneUniformBufferMemberRegistration*>& FSceneUniformBufferMemberRegistr
 	return *GSceneUniformBufferMemberRegistrationInstances;
 }
 
+void FSceneUniformBufferMemberRegistration::Register(FSceneUniformBufferMemberRegistration& Entry)
+{
+	GetInstances().Add(&Entry);
+}
+
 FShaderParametersMetadata* FSceneUniformBufferTypeRegistry::FImpl::BuildStructMetadata()
 {
 	FSceneUniformBufferMemberRegistration::CommitAll();
diff --git a/Engine/Source/Runtime/Renderer/Public/SceneUniformBuffer.h b/Engine/Source/Runtime/Renderer/Public/SceneUniformBuffer.h
--- a/Engine/Source/Runtime/Renderer/Public/SceneUniformBuffer.h
+++ b/Engine/Source/Runtime/Renderer/Public/SceneUniformBuffer.h
@@ -25,7 +25,7 @@ public:
 	class FImpl;
 
 	// Must be called during static initialization
-	static FMemberId RegisterMember(
+	RENDERER_API static FMemberId RegisterMember(
 		const TCHAR* Name,
 		const FShaderParametersMetadata& StructMetadata,
 		// Used to set default values for this member if no value is set explicitly
@@ -46,15 +46,18 @@ public:
 	FShaderParameterStructConstructor DefaultValueConstructor;
 
 protected:
-	static TArray<FSceneUniformBufferMemberRegistration*>& GetInstances();
+	RENDERER_API static void Register(FSceneUniformBufferMemberRegistration& Entry);
 
 	// Called during static initialization
 	FSceneUniformBufferMemberRegistration(const TCHAR* Name)
 		: Name(Name)
 	{
-		GetInstances().Add(this);
+		Register(*this);
 	}
 	virtual ~FSceneUniformBufferMemberRegistration() = default;
+
+private:
+	static TArray<FSceneUniformBufferMemberRegistration*>& GetInstances();
 };
 
 template<typename TMember>
