LFAtlas - fix for potential crash.

diff --git a/Engine/Source/Runtime/Renderer/Private/BasePassRendering.cpp b/Engine/Source/Runtime/Renderer/Private/BasePassRendering.cpp
--- a/Engine/Source/Runtime/Renderer/Private/BasePassRendering.cpp
+++ b/Engine/Source/Runtime/Renderer/Private/BasePassRendering.cpp
@@ -738,7 +738,14 @@ void SetupSharedBasePassParameters(
 	SharedParameters.LFV = View.LocalFogVolumeViewData.UniformParametersStruct;
 
 	FLightFunctionAtlas* LightFunctionAtlas = View.LightFunctionAtlasViewData.GetLightFunctionAtlas();
-	SharedParameters.LightFunctionAtlas = *LightFunctionAtlas->GetLightFunctionAtlasGlobalParametersStruct(ViewIndex, GraphBuilder);
+	if (LightFunctionAtlas)
+	{
+		SharedParameters.LightFunctionAtlas = *LightFunctionAtlas->GetLightFunctionAtlasGlobalParametersStruct(ViewIndex, GraphBuilder);
+	}
+	else
+	{
+		SharedParameters.LightFunctionAtlas = *FLightFunctionAtlas::GetDefaultLightFunctionAtlasGlobalParametersStruct(GraphBuilder);
+	}
 
 	const FScene* Scene = View.Family->Scene ? View.Family->Scene->GetRenderScene() : nullptr;
 	const FPlanarReflectionSceneProxy* ReflectionSceneProxy = Scene ? Scene->GetForwardPassGlobalPlanarReflection() : nullptr;
diff --git a/Engine/Source/Runtime/Renderer/Private/LightFunctionAtlas.cpp b/Engine/Source/Runtime/Renderer/Private/LightFunctionAtlas.cpp
--- a/Engine/Source/Runtime/Renderer/Private/LightFunctionAtlas.cpp
+++ b/Engine/Source/Runtime/Renderer/Private/LightFunctionAtlas.cpp
@@ -483,14 +483,11 @@ TRDGUniformBufferRef<FLightFunctionAtlasGlobalParameters> FLightFunctionAtlas::G
 
 FLightFunctionAtlasGlobalParameters* FLightFunctionAtlas::GetDefaultLightFunctionAtlasGlobalParametersStruct(FRDGBuilder& GraphBuilder)
 {
-	if (DefaultLightFunctionAtlasGlobalParameters == nullptr) // Only create the default buffer once per frame
-	{
-		DefaultLightFunctionAtlasGlobalParameters = GraphBuilder.AllocParameters<FLightFunctionAtlasGlobalParameters>();
-		DefaultLightFunctionAtlasGlobalParameters->LightFunctionAtlasTexture = GSystemTextures.GetWhiteDummy(GraphBuilder);
-		DefaultLightFunctionAtlasGlobalParameters->LightFunctionAtlasSampler = TStaticSamplerState<SF_Bilinear>::GetRHI();
-		DefaultLightFunctionAtlasGlobalParameters->Slot_UVSize = 1.0f;
-	}
-	return DefaultLightFunctionAtlasGlobalParameters;
+	static FLightFunctionAtlasGlobalParameters DefaultLightFunctionAtlasGlobalParameters;
+	DefaultLightFunctionAtlasGlobalParameters.LightFunctionAtlasTexture = GSystemTextures.GetWhiteDummy(GraphBuilder);
+	DefaultLightFunctionAtlasGlobalParameters.LightFunctionAtlasSampler = TStaticSamplerState<SF_Bilinear>::GetRHI();
+	DefaultLightFunctionAtlasGlobalParameters.Slot_UVSize = 1.0f;
+	return &DefaultLightFunctionAtlasGlobalParameters;
 }
 
 TRDGUniformBufferRef<FLightFunctionAtlasGlobalParameters> FLightFunctionAtlas::GetDefaultLightFunctionAtlasGlobalParameters(FRDGBuilder& GraphBuilder)
diff --git a/Engine/Source/Runtime/Renderer/Private/LightFunctionAtlas.h b/Engine/Source/Runtime/Renderer/Private/LightFunctionAtlas.h
--- a/Engine/Source/Runtime/Renderer/Private/LightFunctionAtlas.h
+++ b/Engine/Source/Runtime/Renderer/Private/LightFunctionAtlas.h
@@ -133,7 +133,8 @@ struct FLightFunctionAtlas
 
 	FLightFunctionAtlasGlobalParameters*						GetLightFunctionAtlasGlobalParametersStruct(uint32 ViewIndex, FRDGBuilder& GraphBuilder);
 	TRDGUniformBufferRef<FLightFunctionAtlasGlobalParameters>	GetLightFunctionAtlasGlobalParameters(uint32 ViewIndex, FRDGBuilder& GraphBuilder);
-	FLightFunctionAtlasGlobalParameters*						GetDefaultLightFunctionAtlasGlobalParametersStruct(FRDGBuilder& GraphBuilder);
+	
+	static FLightFunctionAtlasGlobalParameters*					GetDefaultLightFunctionAtlasGlobalParametersStruct(FRDGBuilder& GraphBuilder);
 	TRDGUniformBufferRef<FLightFunctionAtlasGlobalParameters>	GetDefaultLightFunctionAtlasGlobalParameters(FRDGBuilder& GraphBuilder);
 
 private:
