SampledDirectLighting - denoising improvements
* Neighborhood clamp now is applied to both diffuse and specular. r.SampledDirectLighting.Temporal.NeighborhoodClampScale controls neighborhood clamp strictness
* Added num frames acccumulated for speeding up temporal accumulation after disocclusion
* Added lighting 1st and 2nd moment tracking for controling spatial filter in the future
* Added debug ray visualiazation to SWRT path

diff --git a/Engine/Shaders/Private/SampledDirectLighting/SampledDirectLightingDenoising.usf b/Engine/Shaders/Private/SampledDirectLighting/SampledDirectLightingDenoising.usf
deleted file mode 100644
diff --git a/Engine/Shaders/Private/SampledDirectLighting/SampledDirectLightingHardwareRayTracing.usf b/Engine/Shaders/Private/SampledDirectLighting/SampledDirectLightingHardwareRayTracing.usf
--- a/Engine/Shaders/Private/SampledDirectLighting/SampledDirectLightingHardwareRayTracing.usf
+++ b/Engine/Shaders/Private/SampledDirectLighting/SampledDirectLightingHardwareRayTracing.usf
@@ -120,6 +120,7 @@ RAY_TRACING_ENTRY_RAYGEN_OR_INLINE(HardwareRayTraceLightSamples)
 
 		if (Payload.IsHit())
 		{
+			LightSample.bCompleted = true;
 			LightSample.bVisible = false;
 			RWLightSamples[SampleCoord] = PackLightSample(LightSample);
 		}
diff --git a/Engine/Shaders/Private/SampledDirectLighting/SampledDirectLightingShading.usf b/Engine/Shaders/Private/SampledDirectLighting/SampledDirectLightingShading.usf
--- a/Engine/Shaders/Private/SampledDirectLighting/SampledDirectLightingShading.usf
+++ b/Engine/Shaders/Private/SampledDirectLighting/SampledDirectLightingShading.usf
@@ -6,6 +6,7 @@
 
 RWTexture2D<float3> RWDiffuseLighting;
 RWTexture2D<float3> RWSpecularLighting;
+RWTexture2D<float4> RWLuminanceMoments;
 
 Texture2D<float> ShadingTileAtlas;
 Texture2D<uint> ShadingTileGridAllocator;
@@ -19,6 +20,14 @@ StructuredBuffer<uint> TileAllocator;
 StructuredBuffer<uint> TileData;
 uint TileDataStride;
 
+float3 SafeDemodulate(float3 L, float3 C)
+{
+	L.x = C.x >= 0.01f ? L.x / C.x : L.x;
+	L.y = C.y >= 0.01f ? L.y / C.y : L.y;
+	L.z = C.z >= 0.01f ? L.z / C.z : L.z;
+	return L;
+}
+
 /**
  * Apply all lights per pixel to affected tiles using sampled weights from upscaled shading tiles
  */
@@ -76,8 +85,15 @@ void ShadeLightSamplesCS(
 				}
 				#endif
 
+				const float3 N = GBuffer.WorldNormal;
+				const float3 V = normalize(View.TranslatedWorldCameraOrigin - TranslatedWorldPosition);
+				const float NoV = saturate(dot(N, V));
+				const float3 SpecularEnv = EnvBRDF(GBuffer.SpecularColor, GBuffer.Roughness, NoV);
+
 				float3 DiffuseLighting = 0.0f;
 				float3 SpecularLighting = 0.0f;
+				float2 DiffuseLuminanceMoments = 0.0f;
+				float2 SpecularLuminanceMoments = 0.0f;
 
 				for (uint IndexInCell = 0; IndexInCell < NumShadingTilesInCell; ++IndexInCell)
 				{
@@ -127,9 +143,23 @@ void ShadeLightSamplesCS(
 							LightData, LightAttenuation, Dither, ScreenCoord, 
 							SurfaceShadow);
 		
+						// Apply pre-exposure
+						SplitLighting.DiffuseLighting.xyz *= View.PreExposure;
+						SplitLighting.SpecularLighting.xyz *= View.PreExposure;
+
+						// #sdl_todo: demodulate earlier during BRDF evaluation
+						SplitLighting.DiffuseLighting.xyz = SafeDemodulate(SplitLighting.DiffuseLighting.xyz, GBuffer.DiffuseColor);
+						SplitLighting.SpecularLighting.xyz = SafeDemodulate(SplitLighting.SpecularLighting.xyz, SpecularEnv);
+
 						DiffuseLighting += SplitLighting.DiffuseLighting.xyz * SampleWeight;
 						SpecularLighting += SplitLighting.SpecularLighting.xyz * SampleWeight;
 
+						DiffuseLuminanceMoments.x += Luminance(SplitLighting.DiffuseLighting.xyz) * SampleWeight;
+						DiffuseLuminanceMoments.y += Pow2(Luminance(SplitLighting.DiffuseLighting.xyz)) * SampleWeight;
+
+						SpecularLuminanceMoments.x += Luminance(SplitLighting.SpecularLighting.xyz) * SampleWeight;
+						SpecularLuminanceMoments.y += Pow2(Luminance(SplitLighting.SpecularLighting.xyz)) * SampleWeight;
+
 						#if DEBUG_MODE 
 						if (bDebug)
 						{
@@ -140,16 +170,6 @@ void ShadeLightSamplesCS(
 					}
 				}
 
-				// #sdl_todo: demodulate earlier or secure against division by 0
-				// Demodulate diffuse component
-				DiffuseLighting = DiffuseLighting / GBuffer.DiffuseColor;
-
-				// Demodulate specular component
-				const float3 N = GBuffer.WorldNormal;
-				const float3 V = normalize(View.TranslatedWorldCameraOrigin - TranslatedWorldPosition);
-				const float NoV = saturate(dot(N, V));
-				SpecularLighting = SpecularLighting / EnvBRDF(GBuffer.SpecularColor, GBuffer.Roughness, NoV);
-
 				#if DEBUG_MODE
 				if (bDebug)
 				{
@@ -162,8 +182,10 @@ void ShadeLightSamplesCS(
 				}
 				#endif
 
-				RWDiffuseLighting[ScreenCoord] = DiffuseLighting * View.PreExposure;
-				RWSpecularLighting[ScreenCoord] = SpecularLighting * View.PreExposure;
+				// #sdl_todo: can first moment be derived from Diffuse/Specular Lighting?
+				RWDiffuseLighting[ScreenCoord] = DiffuseLighting;
+				RWSpecularLighting[ScreenCoord] = SpecularLighting;
+				RWLuminanceMoments[ScreenCoord] = float4(DiffuseLuminanceMoments, SpecularLuminanceMoments);
 			}
 		}
 	}
diff --git a/Engine/Shaders/Private/SampledDirectLighting/SampledDirectLightingTemporal.usf b/Engine/Shaders/Private/SampledDirectLighting/SampledDirectLightingTemporal.usf
new file mode 100644
--- /dev/null
+++ b/Engine/Shaders/Private/SampledDirectLighting/SampledDirectLightingTemporal.usf
@@ -0,0 +1,230 @@
+// Copyright Epic Games, Inc. All Rights Reserved.
+
+#include "../Common.ush"
+#include "SampledDirectLighting.ush"
+
+struct FNeighborhoodClamp
+{
+	float3 Min;
+	float3 Max;
+};
+
+FNeighborhoodClamp GetNeighborhoodClamp(
+	uint2 ScreenCoord, 
+	Texture2D SampleTexture,
+	float3 CenterSample)
+{
+	float3 SampleSum = CenterSample;
+	float3 SampleSqSum = Pow2(CenterSample);
+
+	// #sdl_todo: optimize using groupshared and cross-lane ops
+	const int KernelSize = 1;
+	for (int SampleOffsetY = -KernelSize; SampleOffsetY <= KernelSize; ++SampleOffsetY)
+	{
+		for (int SampleOffsetX = -KernelSize; SampleOffsetX <= KernelSize; ++SampleOffsetX)
+		{
+			if (!(SampleOffsetX == 0 && SampleOffsetY == 0))
+			{
+				uint2 SampleScreenCoord = clamp(ScreenCoord + int2(SampleOffsetX, SampleOffsetY), 0u, View.ViewRectMinAndSize.zw - 1);
+				float3 Sample = SampleTexture[SampleScreenCoord].xyz;
+				
+				SampleSum += Sample;
+				SampleSqSum += Pow2(Sample);
+			}
+		}
+	}
+
+	float NumSamples = Pow2(2.0f * KernelSize + 1.0f);
+	float3 M1 = SampleSum / NumSamples;
+	float3 M2 = SampleSqSum / NumSamples;
+	float3 Variance = max(M2 - Pow2(M1), 0.0f);
+	float3 StdDev = sqrt(Variance);
+
+	FNeighborhoodClamp NeighborhoodClamp = (FNeighborhoodClamp)0;
+	NeighborhoodClamp.Min = M1 - TemporalNeighborhoodClampScale * StdDev;
+	NeighborhoodClamp.Max = M1 + TemporalNeighborhoodClampScale * StdDev;
+	return NeighborhoodClamp;
+}
+
+Texture2D ResolvedDiffuseLighting;
+Texture2D ResolvedSpecularLighting;
+Texture2D ResolvedLuminanceMoments;
+Texture2D SampledDirectLightingDepthHistory;
+Texture2D DiffuseLightingHistoryTexture;
+Texture2D SpecularLightingHistoryTexture;
+Texture2D LuminanceMomentsHistoryTexture;
+Texture2D<UNORM float> NumFramesAccumulatedHistoryTexture;
+
+float4 HistoryScreenPositionScaleBias;
+float4 HistoryUVMinMax;
+float PrevSceneColorPreExposureCorrection;
+float TemporalMaxFramesAccumulated;
+float TemporalNeighborhoodClampScale;
+
+RWTexture2D<float4> RWDiffuseLighting;
+RWTexture2D<float4> RWSpecularLighting;
+RWTexture2D<float4> RWLuminanceMoments;
+RWTexture2D<UNORM float> RWNumFramesAccumulated;
+RWTexture2D<float> RWSceneDepth;
+RWTexture2D<float4> RWSceneColor;
+
+[numthreads(THREADGROUP_SIZE, THREADGROUP_SIZE, 1)]
+void SampledDirectLightingTemporalAccumulationCS(
+	uint3 GroupId : SV_GroupID,
+	uint3 GroupThreadId : SV_GroupThreadID,
+	uint3 DispatchThreadId : SV_DispatchThreadID)
+{
+	uint2 ScreenCoord = DispatchThreadId.xy;
+
+#if DEBUG_MODE
+	int2 DebugScreenCoord = View.CursorPosition.x >= 0 ? View.CursorPosition * View.ViewResolutionFraction : -1;
+	bool bDebug = all(ScreenCoord == DebugScreenCoord);
+	FShaderPrintContext Context = InitShaderPrintContext(true, float2(0.55, 0.3));
+#endif
+
+	if (all(ScreenCoord < View.ViewRectMinAndSize.zw))
+	{
+		float2 ScreenUV = (ScreenCoord + 0.5f) * View.BufferSizeAndInvSize.zw;
+		FGBufferData GBuffer = GetGBufferData(ScreenUV);
+
+		const float DeviceZ = SceneDepthTexture[ScreenCoord].x;
+		GBuffer.Depth = ConvertFromDeviceZ(DeviceZ);
+
+		if (GBuffer.Depth > 0.0f && GBuffer.ShadingModelID != SHADINGMODELID_UNLIT)
+		{
+			float3 DiffuseLighting = ResolvedDiffuseLighting[ScreenCoord].xyz;
+			float3 SpecularLighting = ResolvedSpecularLighting[ScreenCoord].xyz;
+			float4 LuminanceMoments = ResolvedLuminanceMoments[ScreenCoord].xyzw;
+			float2 DiffuseLumimanceMoments = LuminanceMoments.xy;
+			float2 SpecularLumimanceMoments = LuminanceMoments.zw;
+			float NumFramesAccumulated = 0.0f;
+
+			#if VALID_HISTORY
+			{
+				float2 ScreenPosition = (ScreenUV - View.ScreenPositionScaleBias.wz) / View.ScreenPositionScaleBias.xy;
+				float3 HistoryScreenPosition = GetHistoryScreenPosition(ScreenPosition, ScreenUV, DeviceZ);
+				float2 HistoryScreenUV = HistoryScreenPosition.xy * HistoryScreenPositionScaleBias.xy + HistoryScreenPositionScaleBias.wz;
+
+				bool bHistoryWasOnScreen = all(HistoryScreenUV >= HistoryUVMinMax.xy) && all (HistoryScreenUV <= HistoryUVMinMax.zw);
+				if (bHistoryWasOnScreen)
+				{
+					uint2 HistoryScreenCoord = floor(HistoryScreenUV * View.BufferSizeAndInvSize.xy - 0.5f);
+					float2 HistoryBilinearWeights = frac(HistoryScreenUV * View.BufferSizeAndInvSize.xy - 0.5f);
+					float2 HistoryGatherUV = (HistoryScreenCoord + 1.0f) * View.BufferSizeAndInvSize.zw;
+
+					float4 HistorySampleSceneDepth4 = SampledDirectLightingDepthHistory.GatherRed(GlobalPointClampedSampler, HistoryGatherUV).wzxy;
+
+					// #sdl_todo: plane weighting
+					float ReprojectedSceneDepth = ConvertFromDeviceZ(HistoryScreenPosition.z);
+					float DisocclusionDistanceThreshold = 0.03f;
+					const float4 DistanceToHistoryValue = abs(HistorySampleSceneDepth4 - ReprojectedSceneDepth);
+					float4 DepthWeights = select(DistanceToHistoryValue >= ReprojectedSceneDepth * DisocclusionDistanceThreshold, 0.0f, 1.0f);
+
+					float4 HistoryWeights = float4(
+						(1 - HistoryBilinearWeights.y) * (1 - HistoryBilinearWeights.x),
+						(1 - HistoryBilinearWeights.y) * HistoryBilinearWeights.x,
+						HistoryBilinearWeights.y * (1 - HistoryBilinearWeights.x),
+						HistoryBilinearWeights.y * HistoryBilinearWeights.x);
+
+					HistoryWeights *= DepthWeights;
+
+					#if DEBUG_MODE	
+					if (bDebug)
+					{
+						Print(Context, TEXT("TemporalAccumulation"));
+						Newline(Context);
+						Print(Context, TEXT("ScreenCoord: "));
+						Print(Context, ScreenCoord.x);
+						Print(Context, ScreenCoord.y);
+						Newline(Context);
+						Print(Context, TEXT("DepthWeights: "));
+						Print(Context, DepthWeights);
+						Newline(Context);
+						Print(Context, TEXT("HistoryWeights: "));
+						Print(Context, HistoryWeights);
+					}
+					#endif
+
+					// Use history only if we found at least one valid sample
+					if (any(HistoryWeights > 0.01f))
+					{
+						float3 DiffuseLightingHistory = 0.0f;
+						DiffuseLightingHistory.x = dot(DiffuseLightingHistoryTexture.GatherRed(GlobalPointClampedSampler, HistoryGatherUV).wzxy, HistoryWeights);
+						DiffuseLightingHistory.y = dot(DiffuseLightingHistoryTexture.GatherGreen(GlobalPointClampedSampler, HistoryGatherUV).wzxy, HistoryWeights);
+						DiffuseLightingHistory.z = dot(DiffuseLightingHistoryTexture.GatherBlue(GlobalPointClampedSampler, HistoryGatherUV).wzxy, HistoryWeights);
+						DiffuseLightingHistory = DiffuseLightingHistory / dot(HistoryWeights, 1.0f);
+
+						float3 SpecularLightingHistory = 0.0f;
+						SpecularLightingHistory.x = dot(SpecularLightingHistoryTexture.GatherRed(GlobalPointClampedSampler, HistoryGatherUV).wzxy, HistoryWeights);
+						SpecularLightingHistory.y = dot(SpecularLightingHistoryTexture.GatherGreen(GlobalPointClampedSampler, HistoryGatherUV).wzxy, HistoryWeights);
+						SpecularLightingHistory.z = dot(SpecularLightingHistoryTexture.GatherBlue(GlobalPointClampedSampler, HistoryGatherUV).wzxy, HistoryWeights);
+						SpecularLightingHistory = SpecularLightingHistory / dot(HistoryWeights, 1.0f);
+
+						float4 LuminanceMomentsHistory = 0.0f;
+						LuminanceMomentsHistory.x = dot(LuminanceMomentsHistoryTexture.GatherRed(GlobalPointClampedSampler, HistoryGatherUV).wzxy, HistoryWeights);
+						LuminanceMomentsHistory.y = dot(LuminanceMomentsHistoryTexture.GatherGreen(GlobalPointClampedSampler, HistoryGatherUV).wzxy, HistoryWeights);
+						LuminanceMomentsHistory.z = dot(LuminanceMomentsHistoryTexture.GatherBlue(GlobalPointClampedSampler, HistoryGatherUV).wzxy, HistoryWeights);						
+						LuminanceMomentsHistory.w = dot(LuminanceMomentsHistoryTexture.GatherAlpha(GlobalPointClampedSampler, HistoryGatherUV).wzxy, HistoryWeights);
+						LuminanceMomentsHistory = LuminanceMomentsHistory / dot(HistoryWeights, 1.0f);
+
+						// Reproject and rescale normalized NumFramesAccumulated
+						float NumFramesAccumulatedHistory = 0.0f;
+						NumFramesAccumulatedHistory = dot(NumFramesAccumulatedHistoryTexture.GatherRed(GlobalPointClampedSampler, HistoryGatherUV).wzxy, HistoryWeights) * TemporalMaxFramesAccumulated;
+						NumFramesAccumulatedHistory = NumFramesAccumulatedHistory / dot(HistoryWeights, 1.0f);
+						NumFramesAccumulated = min(NumFramesAccumulatedHistory + 1.0f, TemporalMaxFramesAccumulated);
+
+						float Alpha = 1.0f / (1.0f + NumFramesAccumulated);
+						DiffuseLighting = lerp(DiffuseLightingHistory, DiffuseLighting, Alpha);
+						SpecularLighting = lerp(SpecularLightingHistory, SpecularLighting, Alpha);
+						LuminanceMoments = lerp(LuminanceMomentsHistory, LuminanceMoments, Alpha);
+
+						FNeighborhoodClamp DiffuseNeighborhoodClamp = GetNeighborhoodClamp(ScreenCoord, ResolvedDiffuseLighting, DiffuseLighting);
+						DiffuseLighting = clamp(DiffuseLighting, DiffuseNeighborhoodClamp.Min, DiffuseNeighborhoodClamp.Max);
+
+						FNeighborhoodClamp SpecularNeighborhoodClamp = GetNeighborhoodClamp(ScreenCoord, ResolvedSpecularLighting, SpecularLighting);
+						SpecularLighting = clamp(SpecularLighting, SpecularNeighborhoodClamp.Min, SpecularNeighborhoodClamp.Max);
+					}
+				}
+			}
+			#endif
+
+			DiffuseLighting = MakeFinite(DiffuseLighting);
+			SpecularLighting = MakeFinite(SpecularLighting);
+
+			RWDiffuseLighting[ScreenCoord] = float4(DiffuseLighting, 0.0f);
+			RWSpecularLighting[ScreenCoord] = float4(SpecularLighting, 0.0f);
+			RWLuminanceMoments[ScreenCoord] = float4(DiffuseLumimanceMoments, SpecularLumimanceMoments);
+			RWNumFramesAccumulated[ScreenCoord] = (NumFramesAccumulated + 0.5f) / TemporalMaxFramesAccumulated;
+
+			float3 TranslatedWorldPosition = GetTranslatedWorldPositionFromScreenUV(ScreenUV, GBuffer.Depth);
+			const float3 N = GBuffer.WorldNormal;
+			const float3 V = normalize(View.TranslatedWorldCameraOrigin - TranslatedWorldPosition);
+			const float NoV = saturate(dot(N, V));
+
+			// #sdl_todo: handle alpha channel
+			// Add to scene color and output separately to history
+			float3 Irradiance = GBuffer.DiffuseColor * DiffuseLighting + SpecularLighting * EnvBRDF(GBuffer.SpecularColor, GBuffer.Roughness, NoV);
+
+			float DiffuseVariance = max(DiffuseLumimanceMoments.y - Pow2(DiffuseLumimanceMoments.x), 0.0f);
+			float SpecularVariance = max(SpecularLumimanceMoments.y - Pow2(SpecularLumimanceMoments.x), 0.0f);
+
+			#if DEBUG_MODE
+			if (bDebug)
+			{
+				Newline(Context);
+				Print(Context, TEXT("Irradiance: "));
+				Print(Context, Irradiance);
+				Newline(Context);
+				Print(Context, TEXT("Variance: "));
+				Newline(Context);
+				Print(Context, DiffuseVariance);
+				Print(Context, SpecularVariance);
+			}
+			#endif
+
+			RWSceneColor[ScreenCoord] = RWSceneColor[ScreenCoord] + float4(Irradiance, 0.0f);
+		}
+
+		RWSceneDepth[ScreenCoord] = GBuffer.Depth;
+	}
+}
\ No newline at end of file
diff --git a/Engine/Shaders/Private/SampledDirectLighting/SampledDirectLightingTracing.usf b/Engine/Shaders/Private/SampledDirectLighting/SampledDirectLightingTracing.usf
--- a/Engine/Shaders/Private/SampledDirectLighting/SampledDirectLightingTracing.usf
+++ b/Engine/Shaders/Private/SampledDirectLighting/SampledDirectLightingTracing.usf
@@ -45,8 +45,6 @@ float HistoryDepthTestRelativeThickness;
 uint MinimumTracingThreadOccupancy;
 
 Texture2D<uint> LightSamples;
-RWTexture2D<uint> RWLightSamples;
-RWTexture2D<float> RWLightSampleRayDistance;
 Texture2D<uint> RWTraceTile;
 
 RWStructuredBuffer<uint> RWCompactedTraceTexelData;
@@ -188,6 +186,9 @@ void CompactLightSampleTracesCS(
 	#endif
 }
 
+RWTexture2D<uint> RWLightSamples;
+RWTexture2D<float> RWLightSampleRayDistance;
+
 StructuredBuffer<uint> CompactedTraceTexelData;
 Texture2D<float> DownsampledSceneDepth;
 Texture2D<UNORM float3> DownsampledSceneWorldNormal;
@@ -324,6 +325,12 @@ void SoftwareRayTraceLightSamplesCS(
 		uint2 ScreenCoord = SampleCoordToScreenCoord(SampleCoord);
 		uint2 DownsampledScreenCoord = SampleCoordToDownsampledScreenCoord(SampleCoord);
 
+#if DEBUG_MODE
+		FShaderPrintContext Context = InitShaderPrintContext(true, float2(0.55, 0.45));
+		int2 DebugScreenCoord = View.CursorPosition.x >= 0 ? View.CursorPosition * View.ViewResolutionFraction : -1;
+		bool bDebug = all(DebugScreenCoord / DOWNSAMPLE_FACTOR == DownsampledScreenCoord);
+#endif
+
 		float2 ScreenUV = (ScreenCoord + 0.5f) * View.BufferSizeAndInvSize.zw;
 		float SceneDepth = DownsampledSceneDepth[DownsampledScreenCoord];
 		float3 TranslatedWorldPosition = GetTranslatedWorldPositionFromScreenUV(ScreenUV, SceneDepth);
@@ -345,7 +352,23 @@ void SoftwareRayTraceLightSamplesCS(
 		if (bHit)
 		{
 			LightSample.bVisible = false;
+			LightSample.bCompleted = true;
 			RWLightSamples[SampleCoord] = PackLightSample(LightSample);
 		}
+
+		#if DEBUG_MODE
+		if (bDebug && DebugMode == DEBUG_MODE_VISUALIZE_TRACING)
+		{
+			const FLocalLightData LocalLightData = GetLocalLightDataNonStereo(LightSample.LocalLightIndex);
+			const FLightSceneData LightSceneData = GetLightSceneData(LocalLightData.LightSceneId);
+			const FDeferredLightData LightData = ConvertToDeferredLight(LightSceneData);
+
+			float3 RayStart = TraceInput.TranslatedWorldRayStart + TraceInput.WorldRayDirection * TraceInput.MinTraceDistance;
+			float3 RayEnd = TraceInput.TranslatedWorldRayStart + TraceInput.WorldRayDirection * (bHit ? SDFResult.HitTime : TraceInput.MaxTraceDistance);
+			float4 RayColor = float4(LightData.Color.xyz / Luminance(LightData.Color.xyz), 1.0f);
+
+			AddLineTWS(Context, RayStart, RayEnd, RayColor);
+		}
+		#endif
 	}
 }
\ No newline at end of file
diff --git a/Engine/Source/Runtime/Renderer/Private/SampledDirectLighting/SampledDirectLighting.cpp b/Engine/Source/Runtime/Renderer/Private/SampledDirectLighting/SampledDirectLighting.cpp
--- a/Engine/Source/Runtime/Renderer/Private/SampledDirectLighting/SampledDirectLighting.cpp
+++ b/Engine/Source/Runtime/Renderer/Private/SampledDirectLighting/SampledDirectLighting.cpp
@@ -46,10 +46,10 @@ static TAutoConsoleVariable<int32> CVarSampledDirectLightingTemporalMaxFramesAcc
 	ECVF_Scalability | ECVF_RenderThreadSafe
 );
 
-static TAutoConsoleVariable<float> CVarSampledDirectLightingTemporalStdDevOffset(
-	TEXT("r.SampledDirectLighting.Temporal.StdDevOffset"),
-	0.1f,
-	TEXT("Increases standard deviation in neighborhood clamp. Higher values cause more ghosting, but allow smoother temporal accumulation."),
+static TAutoConsoleVariable<float> CVarSampledDirectLightingTemporalNeighborhoodClampScale(
+	TEXT("r.SampledDirectLighting.Temporal.NeighborhoodClampScale"),
+	1.0f,
+	TEXT("Scales how permissive is neighborhood clamp. Higher values cause more ghosting, but allow smoother temporal accumulation."),
 	ECVF_Scalability | ECVF_RenderThreadSafe
 );
 
@@ -428,6 +428,7 @@ class FShadeLightSamplesCS : public FGlobalShader
 		SHADER_PARAMETER_STRUCT_INCLUDE(FSampledDirectLightingParameters, SampledDirectLightingParameters)
 		SHADER_PARAMETER_RDG_TEXTURE_UAV(RWTexture2D<float3>, RWDiffuseLighting)
 		SHADER_PARAMETER_RDG_TEXTURE_UAV(RWTexture2D<float3>, RWSpecularLighting)
+		SHADER_PARAMETER_RDG_TEXTURE_UAV(RWTexture2D<float4>, RWLuminanceMoments)
 		SHADER_PARAMETER_RDG_BUFFER_SRV(StructuredBuffer<uint>, TileAllocator)
 		SHADER_PARAMETER_RDG_BUFFER_SRV(StructuredBuffer<uint>, TileData)
 		SHADER_PARAMETER_RDG_BUFFER_SRV(StructuredBuffer<uint>, ShadingTileAllocator)
@@ -470,16 +471,21 @@ class FSampledDirectLightingTemporalAccumulationCS : public FGlobalShader
 
 	BEGIN_SHADER_PARAMETER_STRUCT(FParameters, )
 		SHADER_PARAMETER_STRUCT_INCLUDE(FSampledDirectLightingParameters, SampledDirectLightingParameters)
-		SHADER_PARAMETER_RDG_TEXTURE(Texture2D<float4>, DiffuseLightingTexture)
-		SHADER_PARAMETER_RDG_TEXTURE(Texture2D<float4>, SpecularLightingTexture)
-		SHADER_PARAMETER_RDG_TEXTURE(Texture2D<float4>, DiffuseLightingHistoryTexture)
-		SHADER_PARAMETER_RDG_TEXTURE(Texture2D<float4>, SpecularLightingHistoryTexture)
+		SHADER_PARAMETER_RDG_TEXTURE(Texture2D<float3>, ResolvedDiffuseLighting)
+		SHADER_PARAMETER_RDG_TEXTURE(Texture2D<float3>, ResolvedSpecularLighting)
+		SHADER_PARAMETER_RDG_TEXTURE(Texture2D<float4>, ResolvedLuminanceMoments)
+		SHADER_PARAMETER_RDG_TEXTURE(Texture2D<float3>, DiffuseLightingHistoryTexture)
+		SHADER_PARAMETER_RDG_TEXTURE(Texture2D<float3>, SpecularLightingHistoryTexture)
+		SHADER_PARAMETER_RDG_TEXTURE(Texture2D<float4>, LuminanceMomentsHistoryTexture)
+		SHADER_PARAMETER_RDG_TEXTURE(Texture2D<UNORM float>, NumFramesAccumulatedHistoryTexture)
 		SHADER_PARAMETER_RDG_TEXTURE(Texture2D<float>, SampledDirectLightingDepthHistory)
 		SHADER_PARAMETER(FVector4f, HistoryUVMinMax)
 		SHADER_PARAMETER(FVector4f, HistoryScreenPositionScaleBias)
 		SHADER_PARAMETER(float, PrevSceneColorPreExposureCorrection)
-		SHADER_PARAMETER_RDG_TEXTURE_UAV(RWTexture2D<float3>, RWDiffuseLightingTexture)
-		SHADER_PARAMETER_RDG_TEXTURE_UAV(RWTexture2D<float3>, RWSpecularLightingTexture)
+		SHADER_PARAMETER_RDG_TEXTURE_UAV(RWTexture2D<float3>, RWDiffuseLighting)
+		SHADER_PARAMETER_RDG_TEXTURE_UAV(RWTexture2D<float3>, RWSpecularLighting)
+		SHADER_PARAMETER_RDG_TEXTURE_UAV(RWTexture2D<float3>, RWLuminanceMoments)
+		SHADER_PARAMETER_RDG_TEXTURE_UAV(RWTexture2D<UNORM float>, RWNumFramesAccumulated)
 		SHADER_PARAMETER_RDG_TEXTURE_UAV(RWTexture2D<float>, RWSceneDepth)
 		SHADER_PARAMETER_RDG_TEXTURE_UAV(RWTexture2D<float4>, RWSceneColor)
 	END_SHADER_PARAMETER_STRUCT()
@@ -505,7 +511,7 @@ class FSampledDirectLightingTemporalAccumulationCS : public FGlobalShader
 	}
 };
 
-IMPLEMENT_GLOBAL_SHADER(FSampledDirectLightingTemporalAccumulationCS, "/Engine/Private/SampledDirectLighting/SampledDirectLightingDenoising.usf", "SampledDirectLightingTemporalAccumulationCS", SF_Compute);
+IMPLEMENT_GLOBAL_SHADER(FSampledDirectLightingTemporalAccumulationCS, "/Engine/Private/SampledDirectLighting/SampledDirectLightingTemporal.usf", "SampledDirectLightingTemporalAccumulationCS", SF_Compute);
 
 /**
  * Single pass batched light rendering using ray tracing (distance field or triangle) for shadowing.
@@ -584,7 +590,9 @@ void FDeferredShadingSceneRenderer::RenderSampledDirectLighting(FRDGBuilder& Gra
 	FVector4f HistoryUVMinMax = FVector4f(0.0f, 0.0f, 0.0f, 0.0f);
 	FRDGTextureRef DiffuseLightingHistory = nullptr;
 	FRDGTextureRef SpecularLightingHistory = nullptr;
+	FRDGTextureRef LuminanceMomentsHistory = nullptr;
 	FRDGTextureRef SceneDepthHistory = nullptr;
+	FRDGTextureRef NumFramesAccumulatedHistory = nullptr;
 
 	if (View.ViewState)
 	{
@@ -597,14 +605,18 @@ void FDeferredShadingSceneRenderer::RenderSampledDirectLighting(FRDGBuilder& Gra
 
 			if (LightingViewState.DiffuseLightingHistory
 				&& LightingViewState.SpecularLightingHistory
+				&& LightingViewState.LuminanceMomentsHistory
 				&& LightingViewState.SceneDepthHistory
+				&& LightingViewState.NumFramesAccumulatedHistory
 				&& LightingViewState.DiffuseLightingHistory->GetDesc().Extent == View.GetSceneTexturesConfig().Extent
 				&& LightingViewState.SpecularLightingHistory->GetDesc().Extent == View.GetSceneTexturesConfig().Extent
 				&& LightingViewState.SceneDepthHistory->GetDesc().Extent == SceneTextures.Depth.Resolve->Desc.Extent)
 			{
 				DiffuseLightingHistory = GraphBuilder.RegisterExternalTexture(LightingViewState.DiffuseLightingHistory);
 				SpecularLightingHistory = GraphBuilder.RegisterExternalTexture(LightingViewState.SpecularLightingHistory);
+				LuminanceMomentsHistory = GraphBuilder.RegisterExternalTexture(LightingViewState.LuminanceMomentsHistory);
 				SceneDepthHistory = GraphBuilder.RegisterExternalTexture(LightingViewState.SceneDepthHistory);
+				NumFramesAccumulatedHistory = GraphBuilder.RegisterExternalTexture(LightingViewState.NumFramesAccumulatedHistory);
 			}
 		}
 	}
@@ -647,8 +659,8 @@ void FDeferredShadingSceneRenderer::RenderSampledDirectLighting(FRDGBuilder& Gra
 		SampledDirectLightingParameters.SamplingMinWeight = FMath::Max(CVarSampledDirectLightingSamplingMinWeight.GetValueOnRenderThread(), 0.0f);
 		SampledDirectLightingParameters.TileDataStride = TileDataStride;
 		SampledDirectLightingParameters.DownsampledTileDataStride = DownsampledTileDataStride;
-		SampledDirectLightingParameters.TemporalMaxFramesAccumulated = CVarSampledDirectLightingTemporalMaxFramesAccumulated.GetValueOnRenderThread();
-		SampledDirectLightingParameters.TemporalStdDevOffset = CVarSampledDirectLightingTemporalStdDevOffset.GetValueOnRenderThread();
+		SampledDirectLightingParameters.TemporalMaxFramesAccumulated = FMath::Max(CVarSampledDirectLightingTemporalMaxFramesAccumulated.GetValueOnRenderThread(), 0.0f);
+		SampledDirectLightingParameters.TemporalNeighborhoodClampScale = CVarSampledDirectLightingTemporalNeighborhoodClampScale.GetValueOnRenderThread();
 		SampledDirectLightingParameters.DebugMode = SampledDirectLighting::GetDebugMode();
 		SampledDirectLightingParameters.DebugLightId = INDEX_NONE;
 
@@ -863,7 +875,7 @@ void FDeferredShadingSceneRenderer::RenderSampledDirectLighting(FRDGBuilder& Gra
 
 	FRDGTextureRef ShadingTileAtlas = GraphBuilder.CreateTexture(
 		FRDGTextureDesc::Create2D(ShadingTileAtlasSize, PF_R16F, FClearValueBinding::Black, TexCreate_ShaderResource | TexCreate_UAV),
-		TEXT("SampledDirectLighting.ShadingTileAtlas"));
+		TEXT("SampledDirectLighting	.ShadingTileAtlas"));
 
 	AddClearUAVPass(GraphBuilder, GraphBuilder.CreateUAV(ShadingTileAllocator), 0u);
 	AddClearUAVPass(GraphBuilder, GraphBuilder.CreateUAV(ShadingTileGridAllocator), 0u);
@@ -902,6 +914,10 @@ void FDeferredShadingSceneRenderer::RenderSampledDirectLighting(FRDGBuilder& Gra
 		FRDGTextureDesc::Create2D(View.GetSceneTexturesConfig().Extent, PF_FloatRGB, FClearValueBinding::Black, TexCreate_ShaderResource | TexCreate_UAV),
 		TEXT("SampledDirectLighting.ResolvedSpecularLighting"));
 
+	FRDGTextureRef ResolvedLuminanceMoments = GraphBuilder.CreateTexture(
+		FRDGTextureDesc::Create2D(View.GetSceneTexturesConfig().Extent, PF_FloatRGBA, FClearValueBinding::Black, TexCreate_ShaderResource | TexCreate_UAV),
+		TEXT("SampledDirectLighting.ResolvedLuminanceMoments"));
+
 	// Shade light samples
 	{
 		FRDGTextureUAVRef SceneColorUAV = GraphBuilder.CreateUAV(SceneTextures.Color.Target, ERDGUnorderedAccessViewFlags::SkipBarrier);
@@ -911,6 +927,7 @@ void FDeferredShadingSceneRenderer::RenderSampledDirectLighting(FRDGBuilder& Gra
 			FShadeLightSamplesCS::FParameters* PassParameters = GraphBuilder.AllocParameters<FShadeLightSamplesCS::FParameters>();
 			PassParameters->RWDiffuseLighting = GraphBuilder.CreateUAV(ResolvedDiffuseLighting);
 			PassParameters->RWSpecularLighting = GraphBuilder.CreateUAV(ResolvedSpecularLighting);
+			PassParameters->RWLuminanceMoments = GraphBuilder.CreateUAV(ResolvedLuminanceMoments);
 			PassParameters->IndirectArgs = TileIndirectArgs;
 			PassParameters->SampledDirectLightingParameters = SampledDirectLightingParameters;
 			PassParameters->CompositeTileAllocator = GraphBuilder.CreateSRV(CompositeTileAllocator);
@@ -946,24 +963,37 @@ void FDeferredShadingSceneRenderer::RenderSampledDirectLighting(FRDGBuilder& Gra
 		FRDGTextureDesc::Create2D(View.GetSceneTexturesConfig().Extent, PF_FloatRGBA, FClearValueBinding::Black, TexCreate_ShaderResource | TexCreate_UAV),
 		TEXT("SampledDirectLighting.SpecularLighting"));
 
+	FRDGTextureRef LuminanceMoments = GraphBuilder.CreateTexture(
+		FRDGTextureDesc::Create2D(View.GetSceneTexturesConfig().Extent, PF_FloatRGBA, FClearValueBinding::Black, TexCreate_ShaderResource | TexCreate_UAV),
+		TEXT("SampledDirectLighting.LuminanceMoments"));
+
 	FRDGTextureRef SceneDepthCopy = GraphBuilder.CreateTexture(
 		FRDGTextureDesc::Create2D(SceneTextures.Depth.Resolve->Desc.Extent, PF_R32_FLOAT, FClearValueBinding::Black, TexCreate_ShaderResource | TexCreate_UAV),
 		TEXT("SampledDirectLighting.DepthHistory"));
 
+	FRDGTextureRef NumFramesAccumulated = GraphBuilder.CreateTexture(
+		FRDGTextureDesc::Create2D(View.GetSceneTexturesConfig().Extent, PF_G8, FClearValueBinding::Black, TexCreate_ShaderResource | TexCreate_UAV),
+		TEXT("SampledDirectLighting.NumFramesAccumulated"));
+
 	// Temporal accumulation
 	{
 		FSampledDirectLightingTemporalAccumulationCS::FParameters* PassParameters = GraphBuilder.AllocParameters<FSampledDirectLightingTemporalAccumulationCS::FParameters>();
 		PassParameters->SampledDirectLightingParameters = SampledDirectLightingParameters;
-		PassParameters->DiffuseLightingTexture = ResolvedDiffuseLighting;
-		PassParameters->SpecularLightingTexture = ResolvedSpecularLighting;
+		PassParameters->ResolvedDiffuseLighting = ResolvedDiffuseLighting;
+		PassParameters->ResolvedSpecularLighting = ResolvedSpecularLighting;
+		PassParameters->ResolvedLuminanceMoments = ResolvedLuminanceMoments;
 		PassParameters->DiffuseLightingHistoryTexture = DiffuseLightingHistory;
 		PassParameters->SpecularLightingHistoryTexture = SpecularLightingHistory;
+		PassParameters->LuminanceMomentsHistoryTexture = LuminanceMomentsHistory;
+		PassParameters->NumFramesAccumulatedHistoryTexture = NumFramesAccumulatedHistory;
 		PassParameters->SampledDirectLightingDepthHistory = SceneDepthHistory;
 		PassParameters->PrevSceneColorPreExposureCorrection = View.PreExposure / View.PrevViewInfo.SceneColorPreExposure;
 		PassParameters->HistoryScreenPositionScaleBias = HistoryScreenPositionScaleBias;
 		PassParameters->HistoryUVMinMax = HistoryUVMinMax;
-		PassParameters->RWDiffuseLightingTexture = GraphBuilder.CreateUAV(DiffuseLighting);
-		PassParameters->RWSpecularLightingTexture = GraphBuilder.CreateUAV(SpecularLighting);
+		PassParameters->RWDiffuseLighting = GraphBuilder.CreateUAV(DiffuseLighting);
+		PassParameters->RWSpecularLighting = GraphBuilder.CreateUAV(SpecularLighting);
+		PassParameters->RWLuminanceMoments = GraphBuilder.CreateUAV(LuminanceMoments);
+		PassParameters->RWNumFramesAccumulated = GraphBuilder.CreateUAV(NumFramesAccumulated);
 		PassParameters->RWSceneDepth = GraphBuilder.CreateUAV(SceneDepthCopy);
 		PassParameters->RWSceneColor = GraphBuilder.CreateUAV(SceneTextures.Color.Target);
 
@@ -996,17 +1026,21 @@ void FDeferredShadingSceneRenderer::RenderSampledDirectLighting(FRDGBuilder& Gra
 			(View.ViewRect.Max.X - 1.0f) * InvBufferSize.X,
 			(View.ViewRect.Max.Y - 1.0f) * InvBufferSize.Y);
 
-		if (DiffuseLighting && SpecularLighting && SceneDepthCopy && bTemporal)
+		if (DiffuseLighting && SpecularLighting && LuminanceMomentsHistory && SceneDepthCopy && NumFramesAccumulated && bTemporal)
 		{
 			GraphBuilder.QueueTextureExtraction(DiffuseLighting, &LightingViewState.DiffuseLightingHistory);
 			GraphBuilder.QueueTextureExtraction(SpecularLighting, &LightingViewState.SpecularLightingHistory);
+			GraphBuilder.QueueTextureExtraction(LuminanceMomentsHistory, &LightingViewState.LuminanceMomentsHistory);
 			GraphBuilder.QueueTextureExtraction(SceneDepthCopy, &LightingViewState.SceneDepthHistory);
+			GraphBuilder.QueueTextureExtraction(NumFramesAccumulated, &LightingViewState.NumFramesAccumulatedHistory);
 		}
 		else
 		{
 			LightingViewState.DiffuseLightingHistory = nullptr;
 			LightingViewState.SpecularLightingHistory = nullptr;
+			LightingViewState.LuminanceMomentsHistory = nullptr;
 			LightingViewState.SceneDepthHistory = nullptr;
+			LightingViewState.NumFramesAccumulatedHistory = nullptr;
 		}
 	}
 }
\ No newline at end of file
diff --git a/Engine/Source/Runtime/Renderer/Private/SampledDirectLighting/SampledDirectLightingInternal.h b/Engine/Source/Runtime/Renderer/Private/SampledDirectLighting/SampledDirectLightingInternal.h
--- a/Engine/Source/Runtime/Renderer/Private/SampledDirectLighting/SampledDirectLightingInternal.h
+++ b/Engine/Source/Runtime/Renderer/Private/SampledDirectLighting/SampledDirectLightingInternal.h
@@ -30,8 +30,8 @@ BEGIN_SHADER_PARAMETER_STRUCT(FSampledDirectLightingParameters, )
 	SHADER_PARAMETER(float, SamplingMinWeight)
 	SHADER_PARAMETER(int32, TileDataStride)
 	SHADER_PARAMETER(int32, DownsampledTileDataStride)
-	SHADER_PARAMETER(int32, TemporalMaxFramesAccumulated)
-	SHADER_PARAMETER(float, TemporalStdDevOffset)
+	SHADER_PARAMETER(float, TemporalMaxFramesAccumulated)
+	SHADER_PARAMETER(float, TemporalNeighborhoodClampScale)
 	SHADER_PARAMETER(int32, DebugMode)
 	SHADER_PARAMETER(int32, DebugLightId)
 	SHADER_PARAMETER_RDG_TEXTURE(Texture2D<uint>, DownsampledTileMask)
diff --git a/Engine/Source/Runtime/Renderer/Private/SampledDirectLighting/SampledDirectLightingViewState.h b/Engine/Source/Runtime/Renderer/Private/SampledDirectLighting/SampledDirectLightingViewState.h
--- a/Engine/Source/Runtime/Renderer/Private/SampledDirectLighting/SampledDirectLightingViewState.h
+++ b/Engine/Source/Runtime/Renderer/Private/SampledDirectLighting/SampledDirectLightingViewState.h
@@ -7,7 +7,9 @@ class FSampledDirectLightingViewState
 public:
 	TRefCountPtr<IPooledRenderTarget> DiffuseLightingHistory;
 	TRefCountPtr<IPooledRenderTarget> SpecularLightingHistory;
+	TRefCountPtr<IPooledRenderTarget> LuminanceMomentsHistory;
 	TRefCountPtr<IPooledRenderTarget> SceneDepthHistory;
+	TRefCountPtr<IPooledRenderTarget> NumFramesAccumulatedHistory;
 
 	FVector4f HistoryScreenPositionScaleBias = FVector4f(0.0f, 0.0f, 0.0f, 0.0f);
 	FVector4f HistoryUVMinMax = FVector4f(0.0f, 0.0f, 0.0f, 0.0f);
@@ -16,6 +18,8 @@ public:
 	{
 		DiffuseLightingHistory.SafeRelease();
 		SpecularLightingHistory.SafeRelease();
+		LuminanceMomentsHistory.SafeRelease();
 		SceneDepthHistory.SafeRelease();
+		NumFramesAccumulatedHistory.SafeRelease();
 	}
 };
\ No newline at end of file
