Moved raytracing primitive gathering code from DeferredShadingRenderer.cpp to RayTracing.cpp

diff --git a/Engine/Source/Runtime/Renderer/Private/DeferredShadingRenderer.cpp b/Engine/Source/Runtime/Renderer/Private/DeferredShadingRenderer.cpp
--- a/Engine/Source/Runtime/Renderer/Private/DeferredShadingRenderer.cpp
+++ b/Engine/Source/Runtime/Renderer/Private/DeferredShadingRenderer.cpp
@@ -38,6 +38,7 @@
 #include "VT/VirtualTextureSystem.h"
 #include "GPUScene.h"
 #include "PathTracing.h"
+#include "RayTracing/RayTracing.h"
 #include "RayTracing/RayTracingMaterialHitShaders.h"
 #include "RayTracing/RayTracingLighting.h"
 #include "RayTracing/RayTracingDecals.h"
@@ -205,43 +206,6 @@ static TAutoConsoleVariable<int32> CVarRayTracingAllowPipeline(
 	TEXT("Allow use of Ray Tracing pipelines if supported (default=1)."),
 	ECVF_RenderThreadSafe);
 
-static int32 GRayTracingSceneCaptures = -1;
-static FAutoConsoleVariableRef CVarRayTracingSceneCaptures(
-	TEXT("r.RayTracing.SceneCaptures"),
-	GRayTracingSceneCaptures,
-	TEXT("Enable ray tracing in scene captures.\n")
-	TEXT(" -1: Use scene capture settings (default) \n")
-	TEXT(" 0: off \n")
-	TEXT(" 1: on"),
-	ECVF_RenderThreadSafe);
-
-static int32 GRayTracingExcludeDecals = 0;
-static FAutoConsoleVariableRef CRayTracingExcludeDecals(
-	TEXT("r.RayTracing.ExcludeDecals"),
-	GRayTracingExcludeDecals,
-	TEXT("A toggle that modifies the inclusion of decals in the ray tracing BVH.\n")
-	TEXT(" 0: Decals included in the ray tracing BVH (default)\n")
-	TEXT(" 1: Decals excluded from the ray tracing BVH"),
-	ECVF_RenderThreadSafe);
-
-static int32 GRayTracingExcludeTranslucent = 0;
-static FAutoConsoleVariableRef CRayTracingExcludeTranslucent(
-	TEXT("r.RayTracing.ExcludeTranslucent"),
-	GRayTracingExcludeTranslucent,
-	TEXT("A toggle that modifies the inclusion of translucent objects in the ray tracing scene.\n")
-	TEXT(" 0: Translucent objects included in the ray tracing scene (default)\n")
-	TEXT(" 1: Translucent objects excluded from the ray tracing scene"),
-	ECVF_RenderThreadSafe);
-
-static int32 GRayTracingExcludeSky = 1;
-static FAutoConsoleVariableRef CRayTracingExcludeSky(
-	TEXT("r.RayTracing.ExcludeSky"),
-	GRayTracingExcludeSky,
-	TEXT("A toggle that controls inclusion of sky geometry in the ray tracing scene (excluding sky can make ray tracing faster). This setting is ignored for the Path Tracer.\n")
-	TEXT(" 0: Sky objects included in the ray tracing scene\n")
-	TEXT(" 1: Sky objects excluded from the ray tracing scene (default)"),
-	ECVF_RenderThreadSafe);
-
 static TAutoConsoleVariable<int32> CVarRayTracingAsyncBuild(
 	TEXT("r.RayTracing.AsyncBuild"),
 	0,
@@ -249,47 +213,6 @@ static TAutoConsoleVariable<int32> CVarRayTracingAsyncBuild(
 	ECVF_RenderThreadSafe
 );
 
-static int32 GRayTracingParallelMeshBatchSetup = 1;
-static FAutoConsoleVariableRef CRayTracingParallelMeshBatchSetup(
-	TEXT("r.RayTracing.ParallelMeshBatchSetup"),
-	GRayTracingParallelMeshBatchSetup,
-	TEXT("Whether to setup ray tracing materials via parallel jobs."),
-	ECVF_RenderThreadSafe);
-
-static int32 GRayTracingParallelMeshBatchSize = 1024;
-static FAutoConsoleVariableRef CRayTracingParallelMeshBatchSize(
-	TEXT("r.RayTracing.ParallelMeshBatchSize"),
-	GRayTracingParallelMeshBatchSize,
-	TEXT("Batch size for ray tracing materials parallel jobs."),
-	ECVF_RenderThreadSafe);
-
-static TAutoConsoleVariable<float> CVarRayTracingDynamicGeometryLastRenderTimeUpdateDistance(
-	TEXT("r.RayTracing.DynamicGeometryLastRenderTimeUpdateDistance"),
-	5000.0f,
-	TEXT("Dynamic geometries within this distance will have their LastRenderTime updated, so that visibility based ticking (like skeletal mesh) can work when the component is not directly visible in the view (but reflected)."));
-
-static TAutoConsoleVariable<int32> CVarRayTracingAutoInstance(
-	TEXT("r.RayTracing.AutoInstance"),
-	1,
-	TEXT("Whether to auto instance static meshes\n"),
-	ECVF_RenderThreadSafe
-);
-
-static int32 GRayTracingDebugDisableTriangleCull = 0;
-static FAutoConsoleVariableRef CVarRayTracingDebugDisableTriangleCull(
-	TEXT("r.RayTracing.DebugDisableTriangleCull"),
-	GRayTracingDebugDisableTriangleCull,
-	TEXT("Forces all ray tracing geometry instances to be double-sided by disabling back-face culling. This is useful for debugging and profiling. (default = 0)")
-);
-
-
-static int32 GRayTracingDebugForceOpaque = 0;
-static FAutoConsoleVariableRef CVarRayTracingDebugForceOpaque(
-	TEXT("r.RayTracing.DebugForceOpaque"),
-	GRayTracingDebugForceOpaque,
-	TEXT("Forces all ray tracing geometry instances to be opaque, effectively disabling any-hit shaders. This is useful for debugging and profiling. (default = 0)")
-);
-
 static int32 GRayTracingMultiGpuTLASMask = 1;
 static FAutoConsoleVariableRef CVarRayTracingMultiGpuTLASMask(
 	TEXT("r.RayTracing.MultiGpuMaskTLAS"),
@@ -707,7 +630,7 @@ static bool ShouldPrepareRayTracingDecals(const FScene& Scene, const FSceneViewF
 		return false;
 	}
 
-	if (Scene.Decals.Num() == 0 || GRayTracingExcludeDecals)
+	if (Scene.Decals.Num() == 0 || RayTracing::ShouldExcludeDecals())
 	{
 		return false;
 	}
@@ -715,1089 +638,12 @@ static bool ShouldPrepareRayTracingDecals(const FScene& Scene, const FSceneViewF
 	return ViewFamily.EngineShowFlags.PathTracing && PathTracing::UsesDecals(ViewFamily);
 }
 
-static void AddDebugRayTracingInstanceFlags(ERayTracingInstanceFlags& InOutFlags)
-{
-	if (GRayTracingDebugForceOpaque)
-	{
-		InOutFlags |= ERayTracingInstanceFlags::ForceOpaque;
-	}
-	if (GRayTracingDebugDisableTriangleCull)
-	{
-		InOutFlags |= ERayTracingInstanceFlags::TriangleCullDisable;
-	}
-}
-
-struct FRayTracingRelevantPrimitive
-{
-	FRHIRayTracingGeometry* RayTracingGeometryRHI = nullptr;
-	uint64 StateHash = 0;
-	int32 PrimitiveIndex = -1;
-	FPersistentPrimitiveIndex PersistentPrimitiveIndex;
-	int8 LODIndex = -1;
-	uint8 InstanceMask = 0;
-	bool bStatic = false;
-	bool bAllSegmentsOpaque = true;
-	bool bAllSegmentsCastShadow = true;
-	bool bAnySegmentsCastShadow = false;
-	bool bAnySegmentsDecal = false;
-	bool bAllSegmentsDecal = true;
-	bool bTwoSided = false;
-	bool bIsSky = false;
-	bool bAllSegmentsTranslucent = true;
-
-	const FRayTracingGeometryInstance* CachedRayTracingInstance = nullptr;
-	TArrayView<const int32> CachedRayTracingMeshCommandIndices; // Pointer to FPrimitiveSceneInfo::CachedRayTracingMeshCommandIndicesPerLOD data
-
-	uint64 InstancingKey() const
-	{
-		uint64 Key = StateHash;
-		Key ^= uint64(InstanceMask) << 32;
-		Key ^= bAllSegmentsOpaque ? 0x1ull << 40 : 0x0;
-		Key ^= bAllSegmentsCastShadow ? 0x1ull << 41 : 0x0;
-		Key ^= bAnySegmentsCastShadow ? 0x1ull << 42 : 0x0;
-		Key ^= bAnySegmentsDecal ? 0x1ull << 43 : 0x0;
-		Key ^= bAllSegmentsDecal ? 0x1ull << 44 : 0x0;
-		Key ^= bTwoSided ? 0x1ull << 45 : 0x0;
-		Key ^= bIsSky ? 0x1ull << 46 : 0x0;
-		Key ^= bAllSegmentsTranslucent ? 0x1ull << 47 : 0x0;
-		return Key ^ reinterpret_cast<uint64>(RayTracingGeometryRHI);
-	}
-
-	void UpdateMasks(const ERayTracingPrimitiveFlags Flags, ERayTracingViewMaskMode MaskMode)
-	{
-		FRayTracingMeshCommand Command;
-		Command.InstanceMask = InstanceMask;
-		Command.bOpaque = bAllSegmentsOpaque;
-		Command.bCastRayTracedShadows = bAnySegmentsCastShadow;
-		Command.bDecal = bAnySegmentsDecal;
-		Command.bTwoSided = bTwoSided;
-		Command.bIsSky = bIsSky;
-		Command.bIsTranslucent = bAllSegmentsTranslucent;
-
-		UpdateRayTracingMeshCommandMasks(Command, Flags, MaskMode);
-
-		InstanceMask = Command.InstanceMask;
-	}
-};
-
-struct FRayTracingRelevantPrimitiveList
-{
-	// Filtered lists of relevant primitives
-	TChunkedArray<FRayTracingRelevantPrimitive> StaticPrimitives;
-	TChunkedArray<FRayTracingRelevantPrimitive> DynamicPrimitives;
-
-	// Relevant static primitive LODs are computed asynchronously.
-	// This task must complete before accessing StaticPrimitives in FRayTracingSceneAddInstancesTask.
-	FGraphEventRef StaticPrimitiveLODTask;
-
-	// Array of primitives that should update their cached ray tracing instances via FPrimitiveSceneInfo::UpdateCachedRaytracingData()
-	TArray<FPrimitiveSceneInfo*> DirtyCachedRayTracingPrimitives;
-
-	// Used coarse mesh streaming handles during the last TLAS build
-	TArray<Nanite::CoarseMeshStreamingHandle> UsedCoarseMeshStreamingHandles;
-
-	// Indicates that this object has been fully produced (for validation)
-	bool bValid = false;
-};
-
-// Iterates over Scene's PrimitiveSceneProxies and extracts ones that are relevant for ray tracing.
-// This function can run on any thread.
-static void GatherRayTracingRelevantPrimitives(FScene& Scene, const FViewInfo& View, FRayTracingRelevantPrimitiveList& Result)
-{
-	Result.DirtyCachedRayTracingPrimitives.Reserve(Scene.PrimitiveSceneProxies.Num());
-
-	const bool bGameView = View.bIsGameView || View.Family->EngineShowFlags.Game;
-
-	bool bPerformRayTracing = View.State != nullptr && !View.bIsReflectionCapture && View.bAllowRayTracing;
-	if (bPerformRayTracing)
-	{
-		TRACE_CPUPROFILER_EVENT_SCOPE(GatherRayTracingRelevantPrimitives);
-
-		// Index into the TypeOffsetTable, which contains a prefix sum of primitive indices by proxy type
-		int32 BroadIndex = 0;
-
-		for (int PrimitiveIndex = 0; PrimitiveIndex < Scene.PrimitiveSceneProxies.Num(); PrimitiveIndex++)
-		{
-			// Find the next TypeOffsetTable entry that's relevant to this primitive index.
-			while (PrimitiveIndex >= int(Scene.TypeOffsetTable[BroadIndex].Offset))
-			{
-				BroadIndex++;
-			}
-
-			const ERayTracingPrimitiveFlags Flags = Scene.PrimitiveRayTracingFlags[PrimitiveIndex];
-
-			// Skip before dereferencing SceneInfo
-			if (Flags == ERayTracingPrimitiveFlags::UnsupportedProxyType)
-			{
-				// Find the index of a proxy of the next type, skipping over a batch of proxies that are the same type as current.
-				// This assumes that FPrimitiveSceneProxy::IsRayTracingRelevant() is consistent for all proxies of the same type.
-				// I.e. does not depend on members of the particular FPrimitiveSceneProxy implementation.
-				PrimitiveIndex = Scene.TypeOffsetTable[BroadIndex].Offset - 1;
-				continue;
-			}
-
-			// Get primitive visibility state from culling
-			if (!View.PrimitiveRayTracingVisibilityMap[PrimitiveIndex])
-			{
-				continue;
-			}
-
-			check(!EnumHasAnyFlags(Flags, ERayTracingPrimitiveFlags::Excluded));
-
-			const FPrimitiveSceneInfo* SceneInfo = Scene.Primitives[PrimitiveIndex];
-
-			// #dxr_todo: ray tracing in scene captures should re-use the persistent RT scene. (UE-112448)
-			bool bShouldRayTraceSceneCapture = GRayTracingSceneCaptures > 0
-				|| (GRayTracingSceneCaptures == -1 && View.bSceneCaptureUsesRayTracing);
-
-			if (View.bIsSceneCapture && (!bShouldRayTraceSceneCapture || !SceneInfo->bIsVisibleInSceneCaptures))
-			{
-				continue;
-			}
-
-			if (!View.bIsSceneCapture && SceneInfo->bIsVisibleInSceneCapturesOnly)
-			{
-				continue;
-			}
-
-			// Some primitives should only be visible editor mode, however far field geometry 
-			// and hidden shadow casters must still always be added to the RT scene.
-			if (bGameView && !SceneInfo->bDrawInGame && !SceneInfo->bRayTracingFarField)
-			{
-				// Make sure this isn't an object that wants to be hidden to camera but still wants to cast shadows or be visible to indirect
-				checkf(SceneInfo->Proxy != nullptr, TEXT("SceneInfo does not have a valid Proxy object. If this occurs, this object should probably have been filtered out before being added to Scene.Primitives"));
-				if (!SceneInfo->Proxy->CastsHiddenShadow() && !SceneInfo->Proxy->AffectsIndirectLightingWhileHidden())
-				{
-					continue;
-				}
-			}
-
-			// Marked visible and used after point, check if streaming then mark as used in the TLAS (so it can be streamed in)
-			if (EnumHasAnyFlags(Flags, ERayTracingPrimitiveFlags::Streaming))
-			{
-				check(SceneInfo->CoarseMeshStreamingHandle != INDEX_NONE);
-				Result.UsedCoarseMeshStreamingHandles.Add(SceneInfo->CoarseMeshStreamingHandle);
-			}
-
-			// Is the cached data dirty?
-			// eg: mesh was streamed in/out
-			if (SceneInfo->bCachedRaytracingDataDirty)
-			{
-				Result.DirtyCachedRayTracingPrimitives.Add(Scene.Primitives[PrimitiveIndex]);
-			}
-
-			FRayTracingRelevantPrimitive Item;
-			Item.PrimitiveIndex = PrimitiveIndex;
-			Item.PersistentPrimitiveIndex = SceneInfo->GetPersistentIndex();
-
-			if (EnumHasAnyFlags(Flags, ERayTracingPrimitiveFlags::StaticMesh))
-			{
-				if (View.Family->EngineShowFlags.StaticMeshes)
-				{
-					Item.bStatic = true;
-					Result.StaticPrimitives.AddElement(Item);
-				}
-			}
-			else if (View.Family->EngineShowFlags.SkeletalMeshes)
-			{
-				checkf(!EnumHasAllFlags(Flags, ERayTracingPrimitiveFlags::CacheInstances),
-					TEXT("Only static primitives are expected to use CacheInstances flag."));
-
-				Item.bStatic = false;
-				Result.DynamicPrimitives.AddElement(Item);
-			}
-		}
-	}
-
-	FPrimitiveSceneInfo::UpdateCachedRaytracingData(&Scene, Result.DirtyCachedRayTracingPrimitives);
-
-	static const auto ICVarStaticMeshLODDistanceScale = IConsoleManager::Get().FindConsoleVariable(TEXT("r.StaticMeshLODDistanceScale"));
-	const float LODScaleCVarValue = ICVarStaticMeshLODDistanceScale->GetFloat();
-	const int32 ForcedLODLevel = GetCVarForceLOD();
-
-	Result.StaticPrimitiveLODTask = FFunctionGraphTask::CreateAndDispatchWhenReady(
-		[&Result, &Scene, &View, LODScaleCVarValue, ForcedLODLevel]()
-	{
-		TRACE_CPUPROFILER_EVENT_SCOPE(GatherRayTracingWorldInstances_ComputeLOD);
-
-		ParallelFor(TEXT("GatherRayTracingRelevantPrimitives_ComputeLOD"), Result.StaticPrimitives.Num(), 128,
-			[&Result, &Scene, &View, LODScaleCVarValue, ForcedLODLevel](int32 ItemIndex) 
-		{
-			FRayTracingRelevantPrimitive& RelevantPrimitive = Result.StaticPrimitives[ItemIndex];
-
-			const int32 PrimitiveIndex = RelevantPrimitive.PrimitiveIndex;
-			const FPrimitiveSceneInfo* SceneInfo = Scene.Primitives[PrimitiveIndex];
-			const ERayTracingPrimitiveFlags Flags = Scene.PrimitiveRayTracingFlags[PrimitiveIndex];
-
-			int8 LODIndex = 0;
-
-			if (EnumHasAnyFlags(Flags, ERayTracingPrimitiveFlags::ComputeLOD))
-			{
-				const FPrimitiveBounds& Bounds = Scene.PrimitiveBounds[PrimitiveIndex];
-				const FPrimitiveSceneInfo* RESTRICT PrimitiveSceneInfo = Scene.Primitives[PrimitiveIndex];
-
-				FLODMask LODToRender;
-
-				const int8 CurFirstLODIdx = PrimitiveSceneInfo->Proxy->GetCurrentFirstLODIdx_RenderThread();
-				check(CurFirstLODIdx >= 0);
-
-				float MeshScreenSizeSquared = 0;
-				float LODScale = LODScaleCVarValue * View.LODDistanceFactor;
-				LODToRender = ComputeLODForMeshes(SceneInfo->StaticMeshRelevances, View, Bounds.BoxSphereBounds.Origin, Bounds.BoxSphereBounds.SphereRadius, ForcedLODLevel, MeshScreenSizeSquared, CurFirstLODIdx, LODScale, true);
-
-				LODIndex = LODToRender.GetRayTracedLOD();
-			}
-
-
-			if (EnumHasAnyFlags(Flags, ERayTracingPrimitiveFlags::CacheInstances))
-			{
-				FPrimitiveSceneProxy* SceneProxy = Scene.PrimitiveSceneProxies[PrimitiveIndex];
-				const bool bUsingNaniteRayTracing = (Nanite::GetRayTracingMode() != Nanite::ERayTracingMode::Fallback) && SceneProxy->IsNaniteMesh();
-
-				if (!bUsingNaniteRayTracing)
-				{
-					// Currently IsCachedRayTracingGeometryValid() can only be called for non-nanite geometries
-					checkf(SceneInfo->IsCachedRayTracingGeometryValid(), TEXT("Cached ray tracing instance is expected to be valid. Was mesh LOD streamed but cached data was not invalidated?"));
-					checkf(SceneInfo->CachedRayTracingInstance.GeometryRHI, TEXT("Ray tracing instance must have a valid geometry."));
-				}
-
-				RelevantPrimitive.bAnySegmentsDecal = SceneInfo->bCachedRayTracingInstanceAnySegmentsDecal;
-				RelevantPrimitive.bAllSegmentsDecal = SceneInfo->bCachedRayTracingInstanceAllSegmentsDecal;
-				RelevantPrimitive.CachedRayTracingInstance = &SceneInfo->CachedRayTracingInstance;
-
-				// For primitives with ERayTracingPrimitiveFlags::CacheInstances flag we only cache the instance/mesh commands of the current LOD
-				// (see FPrimitiveSceneInfo::UpdateCachedRayTracingInstance(...) and CacheRayTracingPrimitive(...))
-				LODIndex = 0;
-
-				if (SceneInfo->CachedRayTracingMeshCommandIndicesPerLOD.IsValidIndex(LODIndex))
-				{
-					RelevantPrimitive.CachedRayTracingMeshCommandIndices = SceneInfo->CachedRayTracingMeshCommandIndicesPerLOD[LODIndex];
-				}
-			}
-			else
-			{
-				FRHIRayTracingGeometry* RayTracingGeometryInstance = SceneInfo->GetStaticRayTracingGeometryInstance(LODIndex);
-				if (RayTracingGeometryInstance == nullptr)
-				{
-					return;
-				}
-
-				// Sometimes LODIndex is out of range because it is clamped by ClampToFirstLOD, like the requested LOD is being streamed in and hasn't been available
-				// According to InitViews, we should hide the static mesh instance
-				check(EnumHasAnyFlags(Flags, ERayTracingPrimitiveFlags::CacheMeshCommands));
-				if (SceneInfo->CachedRayTracingMeshCommandIndicesPerLOD.IsValidIndex(LODIndex))
-				{
-					RelevantPrimitive.LODIndex = LODIndex;
-					RelevantPrimitive.RayTracingGeometryRHI = RayTracingGeometryInstance;
-
-					RelevantPrimitive.CachedRayTracingMeshCommandIndices = SceneInfo->CachedRayTracingMeshCommandIndicesPerLOD[LODIndex];
-					RelevantPrimitive.StateHash = SceneInfo->CachedRayTracingMeshCommandsHashPerLOD[LODIndex];
-					
-					// TODO: Cache these flags to avoid having to loop over the RayTracingMeshCommands
-					for (int32 CommandIndex : RelevantPrimitive.CachedRayTracingMeshCommandIndices)
-					{
-						if (CommandIndex >= 0)
-						{
-							const FRayTracingMeshCommand& RayTracingMeshCommand = Scene.CachedRayTracingMeshCommands[CommandIndex];
-
-							RelevantPrimitive.InstanceMask |= RayTracingMeshCommand.InstanceMask;
-							RelevantPrimitive.bAllSegmentsOpaque &= RayTracingMeshCommand.bOpaque;
-							RelevantPrimitive.bAllSegmentsCastShadow &= RayTracingMeshCommand.bCastRayTracedShadows;
-							RelevantPrimitive.bAnySegmentsCastShadow |= RayTracingMeshCommand.bCastRayTracedShadows;
-							RelevantPrimitive.bAnySegmentsDecal |= RayTracingMeshCommand.bDecal;
-							RelevantPrimitive.bAllSegmentsDecal &= RayTracingMeshCommand.bDecal;
-							RelevantPrimitive.bTwoSided |= RayTracingMeshCommand.bTwoSided;
-							RelevantPrimitive.bIsSky |= RayTracingMeshCommand.bIsSky;
-							RelevantPrimitive.bAllSegmentsTranslucent &= RayTracingMeshCommand.bIsTranslucent;
-						}
-						else
-						{
-							// CommandIndex == -1 indicates that the mesh batch has been filtered by FRayTracingMeshProcessor (like the shadow depth pass batch)
-							// Do nothing in this case
-						}
-					}
-
-					ERayTracingViewMaskMode MaskMode = static_cast<ERayTracingViewMaskMode>(Scene.CachedRayTracingMeshCommandsMode);
-
-					RelevantPrimitive.UpdateMasks(Flags, MaskMode);
-				}
-			}
-		});
-	}, TStatId(), nullptr, ENamedThreads::AnyThread);
-
-	Result.bValid = true;
-}
-
-// Class to implement build instance mask and flags so that rendering related mask build is maintained in any renderer module.
-// BuildInstanceMaskAndFlags() will be called in the Engine module where it does not know specifics of the ray tracing instance
-// masks used by the renderer (e.g., path tracer mask might be different from raytracing mask).
-struct FDeferredShadingRayTracingMaterialGatheringContext : public FRayTracingMaterialGatheringContext
-{
-	FDeferredShadingRayTracingMaterialGatheringContext(
-		const FScene* InScene,
-		const FSceneView* InReferenceView,
-		const FSceneViewFamily& InReferenceViewFamily,
-		FRDGBuilder& InGraphBuilder,
-		FRayTracingMeshResourceCollector& InRayTracingMeshResourceCollector,
-		FGlobalDynamicReadBuffer& InDynamicReadBuffer)
-		:FRayTracingMaterialGatheringContext(InScene, InReferenceView, InReferenceViewFamily, InGraphBuilder, InRayTracingMeshResourceCollector, InDynamicReadBuffer){}
-
-	virtual FRayTracingMaskAndFlags BuildInstanceMaskAndFlags(const FRayTracingInstance& Instance, const FPrimitiveSceneProxy& ScenePrimitive) override
-	{
-		return BuildRayTracingInstanceMaskAndFlags(Instance, ScenePrimitive, &ReferenceViewFamily);
-	}
-};
-
 struct FRayTracingRelevantPrimitiveTaskData
 {
-	FRayTracingRelevantPrimitiveList List;
+	RayTracing::FRelevantPrimitiveList* List;
 	FGraphEventRef Task;
 };
 
-bool FDeferredShadingSceneRenderer::GatherRayTracingWorldInstancesForView(FRDGBuilder& GraphBuilder, FViewInfo& View, FRayTracingScene& RayTracingScene, FRayTracingRelevantPrimitiveTaskData* RayTracingRelevantPrimitiveTaskData)
-{
-	TRACE_CPUPROFILER_EVENT_SCOPE(FDeferredShadingSceneRenderer::GatherRayTracingWorldInstances);
-	SCOPE_CYCLE_COUNTER(STAT_GatherRayTracingWorldInstances);
-
-	if (!bAnyRayTracingPassEnabled)
-	{
-		return false;
-	}
-
-	check(RayTracingRelevantPrimitiveTaskData)
-
-	// Wait until RayTracingRelevantPrimitiveList is ready
-	if (RayTracingRelevantPrimitiveTaskData->Task.IsValid())
-	{
-		RayTracingRelevantPrimitiveTaskData->Task->Wait();
-		RayTracingRelevantPrimitiveTaskData->Task.SafeRelease();
-	}
-
-	FRayTracingRelevantPrimitiveList& RelevantPrimitiveList = RayTracingRelevantPrimitiveTaskData->List;
-
-	// Prepare ray tracing scene instance list
-	checkf(RelevantPrimitiveList.bValid, TEXT("Ray tracing relevant primitive list is expected to have been created before GatherRayTracingWorldInstancesForView() is called."));
-
-	// Check that any invalidated cached uniform expressions have been updated on the rendering thread.
-	// Normally this work is done through FMaterialRenderProxy::UpdateUniformExpressionCacheIfNeeded,
-	// however ray tracing material processing (FMaterialShader::GetShaderBindings, which accesses UniformExpressionCache)
-	// is done on task threads, therefore all work must be done here up-front as UpdateUniformExpressionCacheIfNeeded is not free-threaded.
-	check(!FMaterialRenderProxy::HasDeferredUniformExpressionCacheRequests());
-
-	FGPUScenePrimitiveCollector DummyDynamicPrimitiveCollector;
-
-	View.DynamicRayTracingMeshCommandStorage.Reserve(Scene->Primitives.Num());
-	View.VisibleRayTracingMeshCommands.Reserve(Scene->Primitives.Num());
-
-	View.RayTracingMeshResourceCollector = MakeUnique<FRayTracingMeshResourceCollector>(Scene->GetFeatureLevel(), Allocator);
-
-	View.RayTracingCullingParameters.Init(View);
-
-	FDeferredShadingRayTracingMaterialGatheringContext MaterialGatheringContext
-	(
-		Scene,
-		&View,
-		ViewFamily,
-		GraphBuilder,
-		*View.RayTracingMeshResourceCollector,
-		DynamicReadBufferForRayTracing
-	);
-
-	const float CurrentWorldTime = View.Family->Time.GetWorldTimeSeconds();
-
-	// Consume output of the relevant primitive gathering task
-	RayTracingScene.UsedCoarseMeshStreamingHandles = MoveTemp(RelevantPrimitiveList.UsedCoarseMeshStreamingHandles);
-
-	// Inform the coarse mesh streaming manager about all the used streamable render assets in the scene
-	Nanite::FCoarseMeshStreamingManager* CoarseMeshSM = IStreamingManager::Get().GetNaniteCoarseMeshStreamingManager();
-	if (CoarseMeshSM)
-	{
-		CoarseMeshSM->AddUsedStreamingHandles(RayTracingScene.UsedCoarseMeshStreamingHandles);
-	}
-
-	INC_DWORD_STAT_BY(STAT_VisibleRayTracingPrimitives, RelevantPrimitiveList.DynamicPrimitives.Num() + RelevantPrimitiveList.StaticPrimitives.Num());
-
-	{
-		TRACE_CPUPROFILER_EVENT_SCOPE(GatherRayTracingWorldInstances_DynamicElements);
-
-		const bool bParallelMeshBatchSetup = GRayTracingParallelMeshBatchSetup && FApp::ShouldUseThreadingForPerformance();
-
-		const int64 SharedBufferGenerationID = Scene->GetRayTracingDynamicGeometryCollection()->BeginUpdate();
-
-		struct FRayTracingMeshBatchWorkItem
-		{
-			const FPrimitiveSceneProxy* SceneProxy = nullptr;
-			TArray<FMeshBatch> MeshBatchesOwned;
-			TArrayView<const FMeshBatch> MeshBatchesView;
-			uint32 InstanceIndex;
-			uint32 DecalInstanceIndex;
-
-			TArrayView<const FMeshBatch> GetMeshBatches() const
-			{
-				if (MeshBatchesOwned.Num())
-				{
-					check(MeshBatchesView.Num() == 0);
-					return TArrayView<const FMeshBatch>(MeshBatchesOwned);
-				}
-				else
-				{
-					check(MeshBatchesOwned.Num() == 0);
-					return MeshBatchesView;
-				}
-			}
-		};
-
-		static constexpr uint32 MaxWorkItemsPerPage = 128; // Try to keep individual pages small to avoid slow-path memory allocations
-		struct FRayTracingMeshBatchTaskPage
-		{
-			FRayTracingMeshBatchWorkItem WorkItems[MaxWorkItemsPerPage];
-			uint32 NumWorkItems = 0;
-			FRayTracingMeshBatchTaskPage* Next = nullptr;
-		};
-
-		FRayTracingMeshBatchTaskPage* MeshBatchTaskHead = nullptr;
-		FRayTracingMeshBatchTaskPage* MeshBatchTaskPage = nullptr;
-		uint32 NumPendingMeshBatches = 0;
-		const uint32 RayTracingParallelMeshBatchSize = GRayTracingParallelMeshBatchSize;
-
-		auto KickRayTracingMeshBatchTask = [&Allocator = Allocator, &View, &MeshBatchTaskHead, &MeshBatchTaskPage, &NumPendingMeshBatches, Scene = this->Scene]()
-		{
-			if (MeshBatchTaskHead)
-			{
-				FDynamicRayTracingMeshCommandStorage* TaskDynamicCommandStorage = Allocator.Create<FDynamicRayTracingMeshCommandStorage>();
-				View.DynamicRayTracingMeshCommandStoragePerTask.Add(TaskDynamicCommandStorage);
-
-				FRayTracingMeshCommandOneFrameArray* TaskVisibleCommands = Allocator.Create<FRayTracingMeshCommandOneFrameArray>();
-				TaskVisibleCommands->Reserve(NumPendingMeshBatches);
-				View.VisibleRayTracingMeshCommandsPerTask.Add(TaskVisibleCommands);
-
-				View.AddRayTracingMeshBatchTaskList.Add(FFunctionGraphTask::CreateAndDispatchWhenReady(
-					[TaskDataHead = MeshBatchTaskHead, &View, Scene, TaskDynamicCommandStorage, TaskVisibleCommands]()
-				{
-					FTaskTagScope TaskTagScope(ETaskTag::EParallelRenderingThread);
-					TRACE_CPUPROFILER_EVENT_SCOPE(RayTracingMeshBatchTask);
-					FRayTracingMeshBatchTaskPage* Page = TaskDataHead;
-					const int32 ExpectedMaxVisibieCommands = TaskVisibleCommands->Max();
-					while (Page)
-					{
-						for (uint32 ItemIndex = 0; ItemIndex < Page->NumWorkItems; ++ItemIndex)
-						{
-							const FRayTracingMeshBatchWorkItem& WorkItem = Page->WorkItems[ItemIndex];
-							TArrayView<const FMeshBatch> MeshBatches = WorkItem.GetMeshBatches();
-							for (int32 SegmentIndex = 0; SegmentIndex < MeshBatches.Num(); SegmentIndex++)
-							{
-								const FMeshBatch& MeshBatch = MeshBatches[SegmentIndex];
-								FDynamicRayTracingMeshCommandContext CommandContext(
-									*TaskDynamicCommandStorage, *TaskVisibleCommands,
-									SegmentIndex, WorkItem.InstanceIndex, WorkItem.DecalInstanceIndex);
-								FMeshPassProcessorRenderState PassDrawRenderState;
-								FRayTracingMeshProcessor RayTracingMeshProcessor(&CommandContext, Scene, &View, PassDrawRenderState, Scene->CachedRayTracingMeshCommandsMode);
-								RayTracingMeshProcessor.AddMeshBatch(MeshBatch, 1, WorkItem.SceneProxy);
-							}
-						}
-						FRayTracingMeshBatchTaskPage* NextPage = Page->Next;
-						Page = NextPage;
-					}
-					check(ExpectedMaxVisibieCommands <= TaskVisibleCommands->Max());
-				}, TStatId(), nullptr, ENamedThreads::AnyThread));
-			}
-
-			MeshBatchTaskHead = nullptr;
-			MeshBatchTaskPage = nullptr;
-			NumPendingMeshBatches = 0;
-		};
-
-		// Local temporary array of instances used for GetDynamicRayTracingInstances()
-		TArray<FRayTracingInstance> TempRayTracingInstances;
-
-		for (const FRayTracingRelevantPrimitive& RelevantPrimitive : RelevantPrimitiveList.DynamicPrimitives)
-		{
-			const FPersistentPrimitiveIndex PersistentPrimitiveIndex = RelevantPrimitive.PersistentPrimitiveIndex;
-			const int32 PrimitiveIndex = RelevantPrimitive.PrimitiveIndex;
-			FPrimitiveSceneProxy* SceneProxy = Scene->PrimitiveSceneProxies[PrimitiveIndex];
-
-			TempRayTracingInstances.Reset();
-			MaterialGatheringContext.DynamicRayTracingGeometriesToUpdate.Reset();
-
-			SceneProxy->GetDynamicRayTracingInstances(MaterialGatheringContext, TempRayTracingInstances);
-
-			for (const FRayTracingDynamicGeometryUpdateParams& DynamicRayTracingGeometryUpdate : MaterialGatheringContext.DynamicRayTracingGeometriesToUpdate)
-			{
-				Scene->GetRayTracingDynamicGeometryCollection()->AddDynamicMeshBatchForGeometryUpdate(
-					GraphBuilder.RHICmdList,
-					Scene,
-					&View,
-					SceneProxy,
-					DynamicRayTracingGeometryUpdate,
-					PersistentPrimitiveIndex.Index
-				);
-			}
-
-			if (TempRayTracingInstances.Num() > 0)
-			{
-				for (FRayTracingInstance& Instance : TempRayTracingInstances)
-				{
-					const FRayTracingGeometry* Geometry = Instance.Geometry;
-
-					if (!ensureMsgf(Geometry->DynamicGeometrySharedBufferGenerationID == FRayTracingGeometry::NonSharedVertexBuffers
-						|| Geometry->DynamicGeometrySharedBufferGenerationID == SharedBufferGenerationID,
-						TEXT("GenerationID %lld, but expected to be %lld or %lld. Geometry debug name: '%s'. ")
-						TEXT("When shared vertex buffers are used, the contents is expected to be written every frame. ")
-						TEXT("Possibly AddDynamicMeshBatchForGeometryUpdate() was not called for this geometry."),
-						Geometry->DynamicGeometrySharedBufferGenerationID, SharedBufferGenerationID, FRayTracingGeometry::NonSharedVertexBuffers,
-						*Geometry->Initializer.DebugName.ToString()))
-					{
-						continue;
-					}
-
-					// If geometry still has pending build request then add to list which requires a force build
-					if (Geometry->HasPendingBuildRequest())
-					{
-						RayTracingScene.GeometriesToBuild.Add(Geometry);
-					}
-
-					// Validate the material/segment counts
-					if (!ensureMsgf(Instance.GetMaterials().Num() == Geometry->Initializer.Segments.Num() ||
-						(Geometry->Initializer.Segments.Num() == 0 && Instance.GetMaterials().Num() == 1),
-						TEXT("Ray tracing material assignment validation failed for geometry '%s'. "
-							"Instance.GetMaterials().Num() = %d, Geometry->Initializer.Segments.Num() = %d, Instance.Mask = 0x%X."),
-						*Geometry->Initializer.DebugName.ToString(), Instance.GetMaterials().Num(),
-						Geometry->Initializer.Segments.Num(), Instance.MaskAndFlags.Mask))
-					{
-						continue;
-					}
-
-					// Autobuild of InstanceMaskAndFlags if the mask and flags are not built
-					UpdateRayTracingInstanceMaskAndFlagsIfNeeded(Instance, *SceneProxy, &ViewFamily);
-
-					// if primitive has mixed decal and non-decal segments we need to have two ray tracing instances
-					// one containing non-decal segments and the other with decal segments
-					// masking of segments is done using "hidden" hitgroups
-					// TODO: Debug Visualization to highlight primitives using this?
-					const bool bNeedSeparateDecalInstance = Instance.MaskAndFlags.bAnySegmentsDecal && !Instance.MaskAndFlags.bAllSegmentsDecal;
-
-					if (GRayTracingExcludeDecals && Instance.MaskAndFlags.bAnySegmentsDecal && !bNeedSeparateDecalInstance)
-					{
-						continue;
-					}
-
-					FRayTracingGeometryInstance RayTracingInstance;
-					RayTracingInstance.GeometryRHI = Geometry->RayTracingGeometryRHI;
-					checkf(RayTracingInstance.GeometryRHI, TEXT("Ray tracing instance must have a valid geometry."));
-					RayTracingInstance.DefaultUserData = PersistentPrimitiveIndex.Index;
-					RayTracingInstance.bApplyLocalBoundsTransform = Instance.bApplyLocalBoundsTransform;
-					RayTracingInstance.LayerIndex = (uint8)(Instance.MaskAndFlags.bAnySegmentsDecal && !bNeedSeparateDecalInstance ? ERayTracingSceneLayer::Decals : ERayTracingSceneLayer::Base);
-					RayTracingInstance.Mask = Instance.MaskAndFlags.Mask;
-
-					if (Instance.MaskAndFlags.bForceOpaque)
-					{
-						RayTracingInstance.Flags |= ERayTracingInstanceFlags::ForceOpaque;
-					}
-					if (Instance.MaskAndFlags.bDoubleSided)
-					{
-						RayTracingInstance.Flags |= ERayTracingInstanceFlags::TriangleCullDisable;
-					}
-					AddDebugRayTracingInstanceFlags(RayTracingInstance.Flags);
-
-					if (Instance.InstanceGPUTransformsSRV.IsValid())
-					{
-						RayTracingInstance.NumTransforms = Instance.NumTransforms;
-						RayTracingInstance.GPUTransformsSRV = Instance.InstanceGPUTransformsSRV;
-					}
-					else 
-					{
-						if (Instance.OwnsTransforms())
-						{
-							// Slow path: copy transforms to the owned storage
-							checkf(Instance.InstanceTransformsView.Num() == 0, TEXT("InstanceTransformsView is expected to be empty if using InstanceTransforms"));
-							TArrayView<FMatrix> SceneOwnedTransforms = RayTracingScene.Allocate<FMatrix>(Instance.InstanceTransforms.Num());
-							FMemory::Memcpy(SceneOwnedTransforms.GetData(), Instance.InstanceTransforms.GetData(), Instance.InstanceTransforms.Num() * sizeof(RayTracingInstance.Transforms[0]));
-							static_assert(std::is_same_v<decltype(SceneOwnedTransforms[0]), decltype(Instance.InstanceTransforms[0])>, "Unexpected transform type");
-
-							RayTracingInstance.NumTransforms = SceneOwnedTransforms.Num();
-							RayTracingInstance.Transforms = SceneOwnedTransforms;
-						}
-						else
-						{
-							// Fast path: just reference persistently-allocated transforms and avoid a copy
-							checkf(Instance.InstanceTransforms.Num() == 0, TEXT("InstanceTransforms is expected to be empty if using InstanceTransformsView"));
-							RayTracingInstance.NumTransforms = Instance.InstanceTransformsView.Num();
-							RayTracingInstance.Transforms = Instance.InstanceTransformsView;
-						}
-					}
-
-					const uint32 InstanceIndex = RayTracingScene.AddInstance(RayTracingInstance, SceneProxy, true);
-
-					uint32 DecalInstanceIndex = INDEX_NONE;
-					if (bNeedSeparateDecalInstance && !GRayTracingExcludeDecals)
-					{
-						FRayTracingGeometryInstance DecalRayTracingInstance = RayTracingInstance;
-						DecalRayTracingInstance.LayerIndex = (uint8)ERayTracingSceneLayer::Decals;
-
-						DecalInstanceIndex = RayTracingScene.AddInstance(MoveTemp(DecalRayTracingInstance), SceneProxy, true);
-					}
-
-					if (bParallelMeshBatchSetup)
-					{
-						if (NumPendingMeshBatches >= RayTracingParallelMeshBatchSize)
-						{
-							KickRayTracingMeshBatchTask();
-						}
-
-						if (MeshBatchTaskPage == nullptr || MeshBatchTaskPage->NumWorkItems == MaxWorkItemsPerPage)
-						{
-							FRayTracingMeshBatchTaskPage* NextPage = Allocator.Create<FRayTracingMeshBatchTaskPage>();
-							if (MeshBatchTaskHead == nullptr)
-							{
-								MeshBatchTaskHead = NextPage;
-							}
-							if (MeshBatchTaskPage)
-							{
-								MeshBatchTaskPage->Next = NextPage;
-							}
-							MeshBatchTaskPage = NextPage;
-						}
-
-						FRayTracingMeshBatchWorkItem& WorkItem = MeshBatchTaskPage->WorkItems[MeshBatchTaskPage->NumWorkItems];
-						MeshBatchTaskPage->NumWorkItems++;
-
-						NumPendingMeshBatches += Instance.GetMaterials().Num();
-
-						if (Instance.OwnsMaterials())
-						{
-							Swap(WorkItem.MeshBatchesOwned, Instance.Materials);
-						}
-						else
-						{
-							WorkItem.MeshBatchesView = Instance.MaterialsView;
-						}
-
-						WorkItem.SceneProxy = SceneProxy;
-						WorkItem.InstanceIndex = InstanceIndex;
-						WorkItem.DecalInstanceIndex = DecalInstanceIndex;
-					}
-					else
-					{
-						TArrayView<const FMeshBatch> InstanceMaterials = Instance.GetMaterials();
-						for (int32 SegmentIndex = 0; SegmentIndex < InstanceMaterials.Num(); SegmentIndex++)
-						{
-							const FMeshBatch& MeshBatch = InstanceMaterials[SegmentIndex];
-							FDynamicRayTracingMeshCommandContext CommandContext(View.DynamicRayTracingMeshCommandStorage, View.VisibleRayTracingMeshCommands, SegmentIndex, InstanceIndex, DecalInstanceIndex);
-							FMeshPassProcessorRenderState PassDrawRenderState;
-							FRayTracingMeshProcessor RayTracingMeshProcessor(&CommandContext, Scene, &View, PassDrawRenderState, Scene->CachedRayTracingMeshCommandsMode);
-							RayTracingMeshProcessor.AddMeshBatch(MeshBatch, 1, SceneProxy);
-						}
-					}
-				}
-
-				if (CVarRayTracingDynamicGeometryLastRenderTimeUpdateDistance.GetValueOnRenderThread() > 0.0f)
-				{
-					if (FVector::Distance(SceneProxy->GetActorPosition(), View.ViewMatrices.GetViewOrigin()) < CVarRayTracingDynamicGeometryLastRenderTimeUpdateDistance.GetValueOnRenderThread())
-					{
-						FPrimitiveSceneInfo* SceneInfo = Scene->Primitives[PrimitiveIndex];
-						// Update LastRenderTime for components so that visibility based ticking (like skeletal meshes) can get updated
-						// We are only doing this for dynamic geometries now
-						SceneInfo->LastRenderTime = CurrentWorldTime;
-						SceneInfo->UpdateComponentLastRenderTime(CurrentWorldTime, /*bUpdateLastRenderTimeOnScreen=*/true);
-					}
-				}
-			}
-		}
-
-		KickRayTracingMeshBatchTask();
-	}
-	
-	// Task to iterate over static ray tracing instances, perform auto-instancing and culling.
-	// This adds final instances to the ray tracing scene and must be done before FRayTracingScene::BuildInitializationData().
-	struct FRayTracingSceneAddInstancesTask
-	{
-		UE_NONCOPYABLE(FRayTracingSceneAddInstancesTask)
-
-		static ESubsequentsMode::Type GetSubsequentsMode() { return ESubsequentsMode::TrackSubsequents; }
-		TStatId                       GetStatId() const    { return TStatId(); }
-		ENamedThreads::Type           GetDesiredThread()   { return ENamedThreads::AnyThread; }
-
-		// Inputs
-
-		const FScene& Scene;
-		TChunkedArray<FRayTracingRelevantPrimitive>& RelevantStaticPrimitives;
-		const FRayTracingCullingParameters& CullingParameters;
-		const bool bIsPathTracing;
-
-		// Outputs
-
-		FRayTracingScene& RayTracingScene; // New instances are added into FRayTracingScene::Instances and FRayTracingScene::Allocator is used for temporary data
-		TArray<FVisibleRayTracingMeshCommand>& VisibleRayTracingMeshCommands; // New elements are added here by this task
-
-		FRayTracingSceneAddInstancesTask(const FScene& InScene,
-											TChunkedArray<FRayTracingRelevantPrimitive>& InRelevantStaticPrimitives,
-											const FRayTracingCullingParameters& InCullingParameters,
-											const bool bInIsPathTracing,
-											FRayTracingScene& InRayTracingScene, TArray<FVisibleRayTracingMeshCommand>& InVisibleRayTracingMeshCommands)
-			: Scene(InScene)
-			, RelevantStaticPrimitives(InRelevantStaticPrimitives)
-			, CullingParameters(InCullingParameters)
-			, bIsPathTracing(bInIsPathTracing)
-			, RayTracingScene(InRayTracingScene)
-			, VisibleRayTracingMeshCommands(InVisibleRayTracingMeshCommands)
-		{
-			VisibleRayTracingMeshCommands.Reserve(RelevantStaticPrimitives.Num());
-		}
-
-		// TODO: Consider moving auto instance batching logic into FRayTracingScene
-
-		struct FAutoInstanceBatch
-		{
-			int32 Index = INDEX_NONE;
-			int32 DecalIndex = INDEX_NONE;
-
-			// Copies the next InstanceSceneDataOffset and user data into the current batch, returns true if arrays were re-allocated.
-			bool Add(FRayTracingScene& InRayTracingScene, uint32 InInstanceSceneDataOffset, uint32 InUserData)
-			{
-				// Adhoc TArray-like resize behavior, in lieu of support for using a custom FMemStackBase in TArray.
-				// Idea for future: if batch becomes large enough, we could actually split it into multiple instances to avoid memory waste.
-
-				const bool bNeedReallocation = Cursor == InstanceSceneDataOffsets.Num();
-
-				if (bNeedReallocation)
-				{
-					int32 PrevCount = InstanceSceneDataOffsets.Num();
-					int32 NextCount = FMath::Max(PrevCount * 2, 1);
-
-					TArrayView<uint32> NewInstanceSceneDataOffsets = InRayTracingScene.Allocate<uint32>(NextCount);
-					if (PrevCount)
-					{
-						FMemory::Memcpy(NewInstanceSceneDataOffsets.GetData(), InstanceSceneDataOffsets.GetData(), InstanceSceneDataOffsets.GetTypeSize() * InstanceSceneDataOffsets.Num());
-					}
-					InstanceSceneDataOffsets = NewInstanceSceneDataOffsets;
-
-					TArrayView<uint32> NewUserData = InRayTracingScene.Allocate<uint32>(NextCount);
-					if (PrevCount)
-					{
-						FMemory::Memcpy(NewUserData.GetData(), UserData.GetData(), UserData.GetTypeSize() * UserData.Num());
-					}
-					UserData = NewUserData;
-				}
-
-				InstanceSceneDataOffsets[Cursor] = InInstanceSceneDataOffset;
-				UserData[Cursor] = InUserData;
-
-				++Cursor;
-
-				return bNeedReallocation;
-			}
-
-			bool IsValid() const
-			{
-				return InstanceSceneDataOffsets.Num() != 0;
-			}
-
-			TArrayView<uint32> InstanceSceneDataOffsets;
-			TArrayView<uint32> UserData;
-			uint32 Cursor = 0;
-		};
-
-		void DoTask(ENamedThreads::Type CurrentThread, const FGraphEventRef& MyCompletionGraphEvent)
-		{
-			FTaskTagScope TaskTagScope(ETaskTag::EParallelRenderingThread);
-
-			TRACE_CPUPROFILER_EVENT_SCOPE(RayTracingSceneStaticInstanceTask);
-
-			FGraphEventArray CullingTasks;
-
-			const bool bAutoInstance = CVarRayTracingAutoInstance.GetValueOnRenderThread() != 0;
-
-			// Instance batches by FRayTracingRelevantPrimitive::InstancingKey()
-			Experimental::TSherwoodMap<uint64, FAutoInstanceBatch> InstanceBatches;
-
-			TArray<FRayTracingCullPrimitiveInstancesClosure> CullInstancesClosures;
-			if (CullingParameters.CullingMode != RayTracing::ECullingMode::Disabled && GetRayTracingCullingPerInstance())
-			{
-				CullInstancesClosures.Reserve(RelevantStaticPrimitives.Num());
-				CullingTasks.Reserve(RelevantStaticPrimitives.Num() / 256 + 1);
-			}
-
-			// scan relevant primitives computing hash data to look for duplicate instances
-			for (const FRayTracingRelevantPrimitive& RelevantPrimitive : RelevantStaticPrimitives)
-			{
-				const int32 PrimitiveIndex = RelevantPrimitive.PrimitiveIndex;
-				FPrimitiveSceneInfo* SceneInfo = Scene.Primitives[PrimitiveIndex];
-				FPrimitiveSceneProxy* SceneProxy = Scene.PrimitiveSceneProxies[PrimitiveIndex];
-				ERayTracingPrimitiveFlags Flags = Scene.PrimitiveRayTracingFlags[PrimitiveIndex];
-				const FPersistentPrimitiveIndex PersistentPrimitiveIndex = RelevantPrimitive.PersistentPrimitiveIndex;
-
-				if (EnumHasAnyFlags(Flags, ERayTracingPrimitiveFlags::CacheInstances))
-				{
-					const bool bUsingNaniteRayTracing = (Nanite::GetRayTracingMode() != Nanite::ERayTracingMode::Fallback) && SceneProxy->IsNaniteMesh();
-
-					if (bUsingNaniteRayTracing)
-					{
-						Nanite::GRayTracingManager.AddVisiblePrimitive(SceneInfo);
-
-						if (RelevantPrimitive.CachedRayTracingInstance->GeometryRHI == nullptr)
-						{
-							// Nanite ray tracing geometry not ready yet, doesn't include primitive in ray tracing scene
-							continue;
-						}
-					}
-					
-					// TODO: Consider requesting a recache of all ray tracing commands during which decals are excluded
-					
-					// if primitive has mixed decal and non-decal segments we need to have two ray tracing instances
-					// one containing non-decal segments and the other with decal segments
-					// masking of segments is done using "hidden" hitgroups
-					// TODO: Debug Visualization to highlight primitives using this?
-					const bool bNeedSeparateDecalInstance = RelevantPrimitive.bAnySegmentsDecal && !RelevantPrimitive.bAllSegmentsDecal;
-
-					if (GRayTracingExcludeDecals && RelevantPrimitive.bAnySegmentsDecal && !bNeedSeparateDecalInstance)
-					{
-						continue;
-					}
-
-					check(RelevantPrimitive.CachedRayTracingInstance);
-
-					const int32 NewInstanceIndex = RayTracingScene.AddInstance(*RelevantPrimitive.CachedRayTracingInstance, SceneProxy, false);
-					uint32 DecalInstanceIndex = INDEX_NONE;
-
-					{
-						FRayTracingGeometryInstance& NewInstance = RayTracingScene.GetInstance(NewInstanceIndex);
-						AddDebugRayTracingInstanceFlags(NewInstance.Flags);
-
-						NewInstance.LayerIndex = (uint8)(RelevantPrimitive.bAnySegmentsDecal && !bNeedSeparateDecalInstance ? ERayTracingSceneLayer::Decals : ERayTracingSceneLayer::Base);
-
-						const Experimental::FHashElementId GroupId = Scene.PrimitiveRayTracingGroupIds[PrimitiveIndex];
-						const bool bUseGroupBounds = CullingParameters.bCullUsingGroupIds && GroupId.IsValid();
-
-						if (CullingParameters.CullingMode != RayTracing::ECullingMode::Disabled && GetRayTracingCullingPerInstance() && RelevantPrimitive.CachedRayTracingInstance->NumTransforms > 1 && !bUseGroupBounds)
-						{
-							const bool bIsFarFieldPrimitive = EnumHasAnyFlags(Flags, ERayTracingPrimitiveFlags::FarField);
-
-							TArrayView<uint32> InstanceActivationMask = RayTracingScene.Allocate<uint32>(FMath::DivideAndRoundUp(NewInstance.NumTransforms, 32u));
-
-							NewInstance.ActivationMask = InstanceActivationMask;
-
-							FRayTracingCullPrimitiveInstancesClosure Closure;
-							Closure.Scene = &Scene;
-							Closure.SceneInfo = SceneInfo;
-							Closure.PrimitiveIndex = PrimitiveIndex;
-							Closure.bIsFarFieldPrimitive = bIsFarFieldPrimitive;
-							Closure.CullingParameters = &CullingParameters;
-							Closure.OutInstanceActivationMask = InstanceActivationMask;
-
-							CullInstancesClosures.Add(MoveTemp(Closure));
-
-							if (CullInstancesClosures.Num() >= 256)
-							{
-								CullingTasks.Add(FFunctionGraphTask::CreateAndDispatchWhenReady([CullInstancesClosures = MoveTemp(CullInstancesClosures)]()
-									{
-										for (auto& Closure : CullInstancesClosures)
-										{
-											Closure();
-										}
-									}, TStatId(), nullptr, ENamedThreads::AnyThread));
-							}
-						}
-
-						if (bNeedSeparateDecalInstance && !GRayTracingExcludeDecals)
-						{
-							FRayTracingGeometryInstance DecalRayTracingInstance = NewInstance;
-							DecalRayTracingInstance.LayerIndex = (uint8)ERayTracingSceneLayer::Decals;
-
-							DecalInstanceIndex = RayTracingScene.AddInstance(MoveTemp(DecalRayTracingInstance), SceneProxy, false);
-						}
-					}
-
-					// At the moment we only support SM & ISMs on this path
-					check(EnumHasAnyFlags(Flags, ERayTracingPrimitiveFlags::CacheMeshCommands));
-					
-					const bool bHasDecalInstanceIndex = DecalInstanceIndex != INDEX_NONE;
-
-					for (int32 CommandIndex : RelevantPrimitive.CachedRayTracingMeshCommandIndices)
-					{
-						const FRayTracingMeshCommand& MeshCommand = Scene.CachedRayTracingMeshCommands[CommandIndex];
-
-						{
-							const bool bHidden = bHasDecalInstanceIndex && MeshCommand.bDecal;
-							FVisibleRayTracingMeshCommand NewVisibleMeshCommand(&MeshCommand, NewInstanceIndex, bHidden);
-							VisibleRayTracingMeshCommands.Add(NewVisibleMeshCommand);
-						}
-
-						if(bHasDecalInstanceIndex)
-						{
-							const bool bHidden = !MeshCommand.bDecal;
-							FVisibleRayTracingMeshCommand NewVisibleMeshCommand(&MeshCommand, DecalInstanceIndex, bHidden);
-							VisibleRayTracingMeshCommands.Add(NewVisibleMeshCommand);
-						}
-					}
-				}
-				else
-				{
-					const int8 LODIndex = RelevantPrimitive.LODIndex;
-
-					if (LODIndex < 0 || !RelevantPrimitive.bStatic)
-					{
-						continue; // skip dynamic primitives and other 
-					}
-
-					// if primitive has mixed decal and non-decal segments we need to have two ray tracing instances
-					// one containing non-decal segments and the other with decal segments
-					// masking of segments is done using "hidden" hitgroups
-					// TODO: Debug Visualization to highlight primitives using this?
-					const bool bNeedSeparateDecalInstance = RelevantPrimitive.bAnySegmentsDecal && !RelevantPrimitive.bAllSegmentsDecal;
-
-					if (GRayTracingExcludeDecals && RelevantPrimitive.bAnySegmentsDecal && !bNeedSeparateDecalInstance)
-					{
-						continue;
-					}
-
-					if ((GRayTracingExcludeDecals && RelevantPrimitive.bAnySegmentsDecal)
-						|| (GRayTracingExcludeTranslucent && RelevantPrimitive.bAllSegmentsTranslucent)
-						|| (GRayTracingExcludeSky && RelevantPrimitive.bIsSky && !bIsPathTracing))
-					{
-						continue;
-					}
-
-					// location if this is a new entry
-					const uint64 InstanceKey = RelevantPrimitive.InstancingKey();
-
-					FAutoInstanceBatch DummyInstanceBatch = { };
-					FAutoInstanceBatch& InstanceBatch = bAutoInstance ? InstanceBatches.FindOrAdd(InstanceKey, DummyInstanceBatch) : DummyInstanceBatch;
-
-					if (InstanceBatch.IsValid())
-					{
-						// Reusing a previous entry, just append to the instance list.
-
-						bool bReallocated = InstanceBatch.Add(RayTracingScene, SceneInfo->GetInstanceSceneDataOffset(), uint32(PersistentPrimitiveIndex.Index));
-
-						check(InstanceBatch.Index != INDEX_NONE);
-						{
-							FRayTracingGeometryInstance& RayTracingInstance = RayTracingScene.GetInstance(InstanceBatch.Index);
-							++RayTracingInstance.NumTransforms;
-							check(RayTracingInstance.NumTransforms == InstanceBatch.Cursor); // sanity check
-
-							if (bReallocated)
-							{
-								RayTracingInstance.InstanceSceneDataOffsets = InstanceBatch.InstanceSceneDataOffsets;
-								RayTracingInstance.UserData = InstanceBatch.UserData;
-							}
-						}
-
-						if(InstanceBatch.DecalIndex != INDEX_NONE)
-						{
-							FRayTracingGeometryInstance& RayTracingInstance = RayTracingScene.GetInstance(InstanceBatch.DecalIndex);
-							++RayTracingInstance.NumTransforms;
-							check(RayTracingInstance.NumTransforms == InstanceBatch.Cursor); // sanity check
-
-							if (bReallocated)
-							{
-								RayTracingInstance.InstanceSceneDataOffsets = InstanceBatch.InstanceSceneDataOffsets;
-								RayTracingInstance.UserData = InstanceBatch.UserData;
-							}
-						}
-					}
-					else
-					{
-						// Starting new instance batch
-
-						InstanceBatch.Add(RayTracingScene, SceneInfo->GetInstanceSceneDataOffset(), uint32(PersistentPrimitiveIndex.Index));
-
-						FRayTracingGeometryInstance RayTracingInstance;
-						RayTracingInstance.GeometryRHI = RelevantPrimitive.RayTracingGeometryRHI;
-						checkf(RayTracingInstance.GeometryRHI, TEXT("Ray tracing instance must have a valid geometry."));
-						RayTracingInstance.InstanceSceneDataOffsets = InstanceBatch.InstanceSceneDataOffsets;
-						RayTracingInstance.UserData = InstanceBatch.UserData;
-						RayTracingInstance.NumTransforms = 1;
-
-						RayTracingInstance.Mask = RelevantPrimitive.InstanceMask; // When no cached command is found, InstanceMask == 0 and the instance is effectively filtered out
-
-						if (RelevantPrimitive.bAllSegmentsOpaque && RelevantPrimitive.bAllSegmentsCastShadow)
-						{
-							RayTracingInstance.Flags |= ERayTracingInstanceFlags::ForceOpaque;
-						}
-						if (RelevantPrimitive.bTwoSided)
-						{
-							RayTracingInstance.Flags |= ERayTracingInstanceFlags::TriangleCullDisable;
-						}
-						AddDebugRayTracingInstanceFlags(RayTracingInstance.Flags);
-
-						RayTracingInstance.LayerIndex = (uint8)(RelevantPrimitive.bAnySegmentsDecal && !bNeedSeparateDecalInstance ? ERayTracingSceneLayer::Decals : ERayTracingSceneLayer::Base);
-
-						InstanceBatch.Index = RayTracingScene.AddInstance(RayTracingInstance, SceneProxy, false);
-
-						if (bNeedSeparateDecalInstance && !GRayTracingExcludeDecals)
-						{
-							FRayTracingGeometryInstance DecalRayTracingInstance = RayTracingInstance;
-							DecalRayTracingInstance.LayerIndex = (uint8)ERayTracingSceneLayer::Decals;
-
-							InstanceBatch.DecalIndex = RayTracingScene.AddInstance(MoveTemp(DecalRayTracingInstance), SceneProxy, false);
-						}
-
-						const bool bHasDecalInstanceIndex = InstanceBatch.DecalIndex != INDEX_NONE;
-
-						for (int32 CommandIndex : RelevantPrimitive.CachedRayTracingMeshCommandIndices)
-						{
-							if (CommandIndex >= 0)
-							{
-								const FRayTracingMeshCommand& MeshCommand = Scene.CachedRayTracingMeshCommands[CommandIndex];
-
-								{
-									const bool bHidden = bHasDecalInstanceIndex && MeshCommand.bDecal;
-									FVisibleRayTracingMeshCommand NewVisibleMeshCommand(&MeshCommand, InstanceBatch.Index, bHidden);
-									VisibleRayTracingMeshCommands.Add(NewVisibleMeshCommand);
-								}
-
-								if (bHasDecalInstanceIndex)
-								{
-									const bool bHidden = !MeshCommand.bDecal;
-									FVisibleRayTracingMeshCommand NewVisibleMeshCommand(&MeshCommand, InstanceBatch.DecalIndex, bHidden);
-									VisibleRayTracingMeshCommands.Add(NewVisibleMeshCommand);
-								}
-							}
-							else
-							{
-								// CommandIndex == -1 indicates that the mesh batch has been filtered by FRayTracingMeshProcessor (like the shadow depth pass batch)
-								// Do nothing in this case
-							}
-						}
-					}
-				}
-			}
-
-			CullingTasks.Add(FFunctionGraphTask::CreateAndDispatchWhenReady([CullInstancesClosures = MoveTemp(CullInstancesClosures)]()
-			{
-				for (auto& Closure : CullInstancesClosures)
-				{
-					Closure();
-				}
-			}, TStatId(), nullptr, ENamedThreads::AnyThread));
-
-			for (FGraphEventRef& CullingTask : CullingTasks)
-			{
-				MyCompletionGraphEvent->DontCompleteUntil(CullingTask);
-			}
-		}
-	};
-
-	FGraphEventArray AddInstancesTaskPrerequisites;
-	AddInstancesTaskPrerequisites.Add(RelevantPrimitiveList.StaticPrimitiveLODTask);
-
-	FGraphEventRef AddInstancesTask = TGraphTask<FRayTracingSceneAddInstancesTask>::CreateTask(&AddInstancesTaskPrerequisites).ConstructAndDispatchWhenReady(
-		*Scene, RelevantPrimitiveList.StaticPrimitives, View.RayTracingCullingParameters, bool(View.Family->EngineShowFlags.PathTracing), // inputs 
-		RayTracingScene, View.VisibleRayTracingMeshCommands // outputs
-	);
-
-	// Scene init task can run only when all pre-init tasks are complete (including culling tasks that are spawned while adding instances)
-	View.RayTracingSceneInitTask = FFunctionGraphTask::CreateAndDispatchWhenReady(
-		[&View, &RayTracingScene]()
-		{
-			TRACE_CPUPROFILER_EVENT_SCOPE(RayTracingSceneInitTask);
-			View.RayTracingSceneInitData = RayTracingScene.BuildInitializationData();
-		},
-		TStatId(), AddInstancesTask, ENamedThreads::AnyThread);
-
-	return true;
-}
-
 static void DeduplicateRayGenerationShaders(TArray< FRHIRayTracingShader*>& RayGenShaders)
 {
 	TSet<FRHIRayTracingShader*> UniqueRayGenShaders;
@@ -2830,11 +1676,12 @@ void FDeferredShadingSceneRenderer::Render(FRDGBuilder& GraphBuilder)
 				FGraphEventRef PrereqTask = CreateCompatibilityGraphEvent(MakeArrayView({ Scene->GetCacheRayTracingPrimitivesTask(), InitViewTaskDatas.VisibilityTaskData->GetFrustumCullTask() }));
 
 				InitViewTaskDatas.RayTracingRelevantPrimitives = Allocator.Create<FRayTracingRelevantPrimitiveTaskData>();
+				InitViewTaskDatas.RayTracingRelevantPrimitives->List = RayTracing::CreateRelevantPrimitiveList(Allocator);
 				InitViewTaskDatas.RayTracingRelevantPrimitives->Task = FFunctionGraphTask::CreateAndDispatchWhenReady(
-					[Scene = Scene, &ReferenceView, &RayTracingRelevantPrimitiveList = InitViewTaskDatas.RayTracingRelevantPrimitives->List]()
+					[Scene = Scene, &ReferenceView, &RayTracingRelevantPrimitiveList = *InitViewTaskDatas.RayTracingRelevantPrimitives->List]()
 					{
 						FTaskTagScope TaskTagScope(ETaskTag::EParallelRenderingThread);
-						GatherRayTracingRelevantPrimitives(*Scene, ReferenceView, RayTracingRelevantPrimitiveList);
+						RayTracing::GatherRelevantPrimitives(*Scene, ReferenceView, RayTracingRelevantPrimitiveList);
 					}, TStatId(), PrereqTask, ENamedThreads::AnyNormalThreadHiPriTask);
 			}
 		}
@@ -3124,7 +1971,19 @@ void FDeferredShadingSceneRenderer::Render(FRDGBuilder& GraphBuilder)
 		}
 
 #if RHI_RAYTRACING
-		GatherRayTracingWorldInstancesForView(GraphBuilder, ReferenceView, RayTracingScene, InitViewTaskDatas.RayTracingRelevantPrimitives);
+		if (bAnyRayTracingPassEnabled)
+		{
+			// Wait until RayTracingRelevantPrimitiveList is ready
+			if (InitViewTaskDatas.RayTracingRelevantPrimitives->Task.IsValid())
+			{
+				InitViewTaskDatas.RayTracingRelevantPrimitives->Task->Wait();
+				InitViewTaskDatas.RayTracingRelevantPrimitives->Task.SafeRelease();
+			}
+
+			checkf(InitViewTaskDatas.RayTracingRelevantPrimitives->List, TEXT("Should have a valid FRelevantPrimitiveList containing the result of RayTracing::GatherRelevantPrimitives here."));
+
+			RayTracing::GatherWorldInstancesForView(GraphBuilder, *Scene, ReferenceView, RayTracingScene, DynamicReadBufferForRayTracing, Allocator, *InitViewTaskDatas.RayTracingRelevantPrimitives->List);
+		}
 #endif // RHI_RAYTRACING
 
 		bool bAnyLumenEnabled = false;
diff --git a/Engine/Source/Runtime/Renderer/Private/DeferredShadingRenderer.h b/Engine/Source/Runtime/Renderer/Private/DeferredShadingRenderer.h
--- a/Engine/Source/Runtime/Renderer/Private/DeferredShadingRenderer.h
+++ b/Engine/Source/Runtime/Renderer/Private/DeferredShadingRenderer.h
@@ -1111,9 +1111,6 @@ private:
 	void RenderRayTracingBarycentrics(FRDGBuilder& GraphBuilder, const FViewInfo& View, FRDGTextureRef SceneColorOutputTexture, bool bVisualizeProceduralPrimitives);
 	void RayTracingDisplayPicking(const FRayTracingPickingFeedback& PickingFeedback, FScreenMessageWriter& Writer);
 
-	/** Fills RayTracingScene instance list for the given View and adds relevant ray tracing data to the view. Does not reset previous scene contents. */
-	bool GatherRayTracingWorldInstancesForView(FRDGBuilder& GraphBuilder, FViewInfo& View, FRayTracingScene& RayTracingScene, FRayTracingRelevantPrimitiveTaskData* RayTracingRelevantPrimitiveTaskData);
-
 	bool SetupRayTracingPipelineStates(FRDGBuilder& GraphBuilder);
 	void SetupRayTracingLightDataForViews(FRDGBuilder& GraphBuilder);
 	bool DispatchRayTracingWorldUpdates(FRDGBuilder& GraphBuilder, FRDGBufferRef& OutDynamicGeometryScratchBuffer);
diff --git a/Engine/Source/Runtime/Renderer/Private/RayTracing/RayTracing.cpp b/Engine/Source/Runtime/Renderer/Private/RayTracing/RayTracing.cpp
new file mode 100644
--- /dev/null
+++ b/Engine/Source/Runtime/Renderer/Private/RayTracing/RayTracing.cpp
@@ -0,0 +1,1175 @@
+// Copyright Epic Games, Inc. All Rights Reserved.
+
+#include "RayTracing.h"
+
+#if RHI_RAYTRACING
+
+#include "RayTracingDynamicGeometryCollection.h"
+#include "RayTracingInstanceMask.h"
+#include "RayTracingInstanceCulling.h"
+#include "RayTracingMaterialHitShaders.h"
+#include "RayTracingScene.h"
+#include "Nanite\NaniteRayTracing.h"
+#include "Rendering\NaniteCoarseMeshStreamingManager.h"
+#include "ScenePrivate.h"
+#include "Materials\MaterialRenderProxy.h"
+#include "Experimental\Containers\SherwoodHashTable.h"
+#include "Async/ParallelFor.h"
+
+static int32 GRayTracingSceneCaptures = -1;
+static FAutoConsoleVariableRef CVarRayTracingSceneCaptures(
+	TEXT("r.RayTracing.SceneCaptures"),
+	GRayTracingSceneCaptures,
+	TEXT("Enable ray tracing in scene captures.\n")
+	TEXT(" -1: Use scene capture settings (default) \n")
+	TEXT(" 0: off \n")
+	TEXT(" 1: on"),
+	ECVF_RenderThreadSafe);
+
+static int32 GRayTracingParallelMeshBatchSetup = 1;
+static FAutoConsoleVariableRef CRayTracingParallelMeshBatchSetup(
+	TEXT("r.RayTracing.ParallelMeshBatchSetup"),
+	GRayTracingParallelMeshBatchSetup,
+	TEXT("Whether to setup ray tracing materials via parallel jobs."),
+	ECVF_RenderThreadSafe);
+
+static int32 GRayTracingParallelMeshBatchSize = 1024;
+static FAutoConsoleVariableRef CRayTracingParallelMeshBatchSize(
+	TEXT("r.RayTracing.ParallelMeshBatchSize"),
+	GRayTracingParallelMeshBatchSize,
+	TEXT("Batch size for ray tracing materials parallel jobs."),
+	ECVF_RenderThreadSafe);
+
+static TAutoConsoleVariable<float> CVarRayTracingDynamicGeometryLastRenderTimeUpdateDistance(
+	TEXT("r.RayTracing.DynamicGeometryLastRenderTimeUpdateDistance"),
+	5000.0f,
+	TEXT("Dynamic geometries within this distance will have their LastRenderTime updated, so that visibility based ticking (like skeletal mesh) can work when the component is not directly visible in the view (but reflected)."));
+
+static TAutoConsoleVariable<int32> CVarRayTracingAutoInstance(
+	TEXT("r.RayTracing.AutoInstance"),
+	1,
+	TEXT("Whether to auto instance static meshes\n"),
+	ECVF_RenderThreadSafe
+);
+
+static int32 GRayTracingExcludeTranslucent = 0;
+static FAutoConsoleVariableRef CRayTracingExcludeTranslucent(
+	TEXT("r.RayTracing.ExcludeTranslucent"),
+	GRayTracingExcludeTranslucent,
+	TEXT("A toggle that modifies the inclusion of translucent objects in the ray tracing scene.\n")
+	TEXT(" 0: Translucent objects included in the ray tracing scene (default)\n")
+	TEXT(" 1: Translucent objects excluded from the ray tracing scene"),
+	ECVF_RenderThreadSafe);
+
+static int32 GRayTracingExcludeSky = 1;
+static FAutoConsoleVariableRef CRayTracingExcludeSky(
+	TEXT("r.RayTracing.ExcludeSky"),
+	GRayTracingExcludeSky,
+	TEXT("A toggle that controls inclusion of sky geometry in the ray tracing scene (excluding sky can make ray tracing faster). This setting is ignored for the Path Tracer.\n")
+	TEXT(" 0: Sky objects included in the ray tracing scene\n")
+	TEXT(" 1: Sky objects excluded from the ray tracing scene (default)"),
+	ECVF_RenderThreadSafe);
+
+static int32 GRayTracingExcludeDecals = 0;
+static FAutoConsoleVariableRef CRayTracingExcludeDecals(
+	TEXT("r.RayTracing.ExcludeDecals"),
+	GRayTracingExcludeDecals,
+	TEXT("A toggle that modifies the inclusion of decals in the ray tracing BVH.\n")
+	TEXT(" 0: Decals included in the ray tracing BVH (default)\n")
+	TEXT(" 1: Decals excluded from the ray tracing BVH"),
+	ECVF_RenderThreadSafe);
+
+static int32 GRayTracingDebugDisableTriangleCull = 0;
+static FAutoConsoleVariableRef CVarRayTracingDebugDisableTriangleCull(
+	TEXT("r.RayTracing.DebugDisableTriangleCull"),
+	GRayTracingDebugDisableTriangleCull,
+	TEXT("Forces all ray tracing geometry instances to be double-sided by disabling back-face culling. This is useful for debugging and profiling. (default = 0)")
+);
+
+static int32 GRayTracingDebugForceOpaque = 0;
+static FAutoConsoleVariableRef CVarRayTracingDebugForceOpaque(
+	TEXT("r.RayTracing.DebugForceOpaque"),
+	GRayTracingDebugForceOpaque,
+	TEXT("Forces all ray tracing geometry instances to be opaque, effectively disabling any-hit shaders. This is useful for debugging and profiling. (default = 0)")
+);
+
+namespace RayTracing
+{
+	struct FRelevantPrimitive
+	{
+		FRHIRayTracingGeometry* RayTracingGeometryRHI = nullptr;
+		uint64 StateHash = 0;
+		int32 PrimitiveIndex = -1;
+		FPersistentPrimitiveIndex PersistentPrimitiveIndex;
+		int8 LODIndex = -1;
+		uint8 InstanceMask = 0;
+		bool bStatic = false;
+		bool bAllSegmentsOpaque = true;
+		bool bAllSegmentsCastShadow = true;
+		bool bAnySegmentsCastShadow = false;
+		bool bAnySegmentsDecal = false;
+		bool bAllSegmentsDecal = true;
+		bool bTwoSided = false;
+		bool bIsSky = false;
+		bool bAllSegmentsTranslucent = true;
+
+		const FRayTracingGeometryInstance* CachedRayTracingInstance = nullptr;
+		TArrayView<const int32> CachedRayTracingMeshCommandIndices; // Pointer to FPrimitiveSceneInfo::CachedRayTracingMeshCommandIndicesPerLOD data
+
+		uint64 InstancingKey() const
+		{
+			uint64 Key = StateHash;
+			Key ^= uint64(InstanceMask) << 32;
+			Key ^= bAllSegmentsOpaque ? 0x1ull << 40 : 0x0;
+			Key ^= bAllSegmentsCastShadow ? 0x1ull << 41 : 0x0;
+			Key ^= bAnySegmentsCastShadow ? 0x1ull << 42 : 0x0;
+			Key ^= bAnySegmentsDecal ? 0x1ull << 43 : 0x0;
+			Key ^= bAllSegmentsDecal ? 0x1ull << 44 : 0x0;
+			Key ^= bTwoSided ? 0x1ull << 45 : 0x0;
+			Key ^= bIsSky ? 0x1ull << 46 : 0x0;
+			Key ^= bAllSegmentsTranslucent ? 0x1ull << 47 : 0x0;
+			return Key ^ reinterpret_cast<uint64>(RayTracingGeometryRHI);
+		}
+
+		void UpdateMasks(const ERayTracingPrimitiveFlags Flags, ERayTracingViewMaskMode MaskMode)
+		{
+			FRayTracingMeshCommand Command;
+			Command.InstanceMask = InstanceMask;
+			Command.bOpaque = bAllSegmentsOpaque;
+			Command.bCastRayTracedShadows = bAnySegmentsCastShadow;
+			Command.bDecal = bAnySegmentsDecal;
+			Command.bTwoSided = bTwoSided;
+			Command.bIsSky = bIsSky;
+			Command.bIsTranslucent = bAllSegmentsTranslucent;
+
+			UpdateRayTracingMeshCommandMasks(Command, Flags, MaskMode);
+
+			InstanceMask = Command.InstanceMask;
+		}
+	};
+
+	struct FRelevantPrimitiveList
+	{
+		// Filtered lists of relevant primitives
+		TChunkedArray<FRelevantPrimitive> StaticPrimitives;
+		TChunkedArray<FRelevantPrimitive> DynamicPrimitives;
+
+		// Relevant static primitive LODs are computed asynchronously.
+		// This task must complete before accessing StaticPrimitives in FRayTracingSceneAddInstancesTask.
+		FGraphEventRef StaticPrimitiveLODTask;
+
+		// Array of primitives that should update their cached ray tracing instances via FPrimitiveSceneInfo::UpdateCachedRaytracingData()
+		TArray<FPrimitiveSceneInfo*> DirtyCachedRayTracingPrimitives;
+
+		// Used coarse mesh streaming handles during the last TLAS build
+		TArray<Nanite::CoarseMeshStreamingHandle> UsedCoarseMeshStreamingHandles;
+
+		// Indicates that this object has been fully produced (for validation)
+		bool bValid = false;
+	};
+
+	FRelevantPrimitiveList* CreateRelevantPrimitiveList(FSceneRenderingBulkObjectAllocator& InAllocator)
+	{
+		return InAllocator.Create<FRelevantPrimitiveList>();
+	}
+
+	void GatherRelevantPrimitives(FScene& Scene, const FViewInfo& View, FRelevantPrimitiveList& Result)
+	{
+		Result.DirtyCachedRayTracingPrimitives.Reserve(Scene.PrimitiveSceneProxies.Num());
+
+		const bool bGameView = View.bIsGameView || View.Family->EngineShowFlags.Game;
+
+		bool bPerformRayTracing = View.State != nullptr && !View.bIsReflectionCapture && View.bAllowRayTracing;
+		if (bPerformRayTracing)
+		{
+			TRACE_CPUPROFILER_EVENT_SCOPE(GatherRayTracingRelevantPrimitives);
+
+			// Index into the TypeOffsetTable, which contains a prefix sum of primitive indices by proxy type
+			int32 BroadIndex = 0;
+
+			for (int PrimitiveIndex = 0; PrimitiveIndex < Scene.PrimitiveSceneProxies.Num(); PrimitiveIndex++)
+			{
+				// Find the next TypeOffsetTable entry that's relevant to this primitive index.
+				while (PrimitiveIndex >= int(Scene.TypeOffsetTable[BroadIndex].Offset))
+				{
+					BroadIndex++;
+				}
+
+				const ERayTracingPrimitiveFlags Flags = Scene.PrimitiveRayTracingFlags[PrimitiveIndex];
+
+				// Skip before dereferencing SceneInfo
+				if (Flags == ERayTracingPrimitiveFlags::UnsupportedProxyType)
+				{
+					// Find the index of a proxy of the next type, skipping over a batch of proxies that are the same type as current.
+					// This assumes that FPrimitiveSceneProxy::IsRayTracingRelevant() is consistent for all proxies of the same type.
+					// I.e. does not depend on members of the particular FPrimitiveSceneProxy implementation.
+					PrimitiveIndex = Scene.TypeOffsetTable[BroadIndex].Offset - 1;
+					continue;
+				}
+
+				// Get primitive visibility state from culling
+				if (!View.PrimitiveRayTracingVisibilityMap[PrimitiveIndex])
+				{
+					continue;
+				}
+
+				check(!EnumHasAnyFlags(Flags, ERayTracingPrimitiveFlags::Excluded));
+
+				const FPrimitiveSceneInfo* SceneInfo = Scene.Primitives[PrimitiveIndex];
+
+				// #dxr_todo: ray tracing in scene captures should re-use the persistent RT scene. (UE-112448)
+				bool bShouldRayTraceSceneCapture = GRayTracingSceneCaptures > 0
+					|| (GRayTracingSceneCaptures == -1 && View.bSceneCaptureUsesRayTracing);
+
+				if (View.bIsSceneCapture && (!bShouldRayTraceSceneCapture || !SceneInfo->bIsVisibleInSceneCaptures))
+				{
+					continue;
+				}
+
+				if (!View.bIsSceneCapture && SceneInfo->bIsVisibleInSceneCapturesOnly)
+				{
+					continue;
+				}
+
+				// Some primitives should only be visible editor mode, however far field geometry 
+				// and hidden shadow casters must still always be added to the RT scene.
+				if (bGameView && !SceneInfo->bDrawInGame && !SceneInfo->bRayTracingFarField)
+				{
+					// Make sure this isn't an object that wants to be hidden to camera but still wants to cast shadows or be visible to indirect
+					checkf(SceneInfo->Proxy != nullptr, TEXT("SceneInfo does not have a valid Proxy object. If this occurs, this object should probably have been filtered out before being added to Scene.Primitives"));
+					if (!SceneInfo->Proxy->CastsHiddenShadow() && !SceneInfo->Proxy->AffectsIndirectLightingWhileHidden())
+					{
+						continue;
+					}
+				}
+
+				// Marked visible and used after point, check if streaming then mark as used in the TLAS (so it can be streamed in)
+				if (EnumHasAnyFlags(Flags, ERayTracingPrimitiveFlags::Streaming))
+				{
+					check(SceneInfo->CoarseMeshStreamingHandle != INDEX_NONE);
+					Result.UsedCoarseMeshStreamingHandles.Add(SceneInfo->CoarseMeshStreamingHandle);
+				}
+
+				// Is the cached data dirty?
+				// eg: mesh was streamed in/out
+				if (SceneInfo->bCachedRaytracingDataDirty)
+				{
+					Result.DirtyCachedRayTracingPrimitives.Add(Scene.Primitives[PrimitiveIndex]);
+				}
+
+				FRelevantPrimitive Item;
+				Item.PrimitiveIndex = PrimitiveIndex;
+				Item.PersistentPrimitiveIndex = SceneInfo->GetPersistentIndex();
+
+				if (EnumHasAnyFlags(Flags, ERayTracingPrimitiveFlags::StaticMesh))
+				{
+					if (View.Family->EngineShowFlags.StaticMeshes)
+					{
+						Item.bStatic = true;
+						Result.StaticPrimitives.AddElement(Item);
+					}
+				}
+				else if (View.Family->EngineShowFlags.SkeletalMeshes)
+				{
+					checkf(!EnumHasAllFlags(Flags, ERayTracingPrimitiveFlags::CacheInstances),
+						TEXT("Only static primitives are expected to use CacheInstances flag."));
+
+					Item.bStatic = false;
+					Result.DynamicPrimitives.AddElement(Item);
+				}
+			}
+		}
+
+		FPrimitiveSceneInfo::UpdateCachedRaytracingData(&Scene, Result.DirtyCachedRayTracingPrimitives);
+
+		static const auto ICVarStaticMeshLODDistanceScale = IConsoleManager::Get().FindConsoleVariable(TEXT("r.StaticMeshLODDistanceScale"));
+		const float LODScaleCVarValue = ICVarStaticMeshLODDistanceScale->GetFloat();
+		const int32 ForcedLODLevel = GetCVarForceLOD();
+
+		Result.StaticPrimitiveLODTask = FFunctionGraphTask::CreateAndDispatchWhenReady(
+			[&Result, &Scene, &View, LODScaleCVarValue, ForcedLODLevel]()
+			{
+				TRACE_CPUPROFILER_EVENT_SCOPE(GatherRayTracingWorldInstances_ComputeLOD);
+
+				ParallelFor(TEXT("GatherRayTracingRelevantPrimitives_ComputeLOD"), Result.StaticPrimitives.Num(), 128,
+					[&Result, &Scene, &View, LODScaleCVarValue, ForcedLODLevel](int32 ItemIndex)
+					{
+						FRelevantPrimitive& RelevantPrimitive = Result.StaticPrimitives[ItemIndex];
+
+						const int32 PrimitiveIndex = RelevantPrimitive.PrimitiveIndex;
+						const FPrimitiveSceneInfo* SceneInfo = Scene.Primitives[PrimitiveIndex];
+						const ERayTracingPrimitiveFlags Flags = Scene.PrimitiveRayTracingFlags[PrimitiveIndex];
+
+						int8 LODIndex = 0;
+
+						if (EnumHasAnyFlags(Flags, ERayTracingPrimitiveFlags::ComputeLOD))
+						{
+							const FPrimitiveBounds& Bounds = Scene.PrimitiveBounds[PrimitiveIndex];
+							const FPrimitiveSceneInfo* RESTRICT PrimitiveSceneInfo = Scene.Primitives[PrimitiveIndex];
+
+							FLODMask LODToRender;
+
+							const int8 CurFirstLODIdx = PrimitiveSceneInfo->Proxy->GetCurrentFirstLODIdx_RenderThread();
+							check(CurFirstLODIdx >= 0);
+
+							float MeshScreenSizeSquared = 0;
+							float LODScale = LODScaleCVarValue * View.LODDistanceFactor;
+							LODToRender = ComputeLODForMeshes(SceneInfo->StaticMeshRelevances, View, Bounds.BoxSphereBounds.Origin, Bounds.BoxSphereBounds.SphereRadius, ForcedLODLevel, MeshScreenSizeSquared, CurFirstLODIdx, LODScale, true);
+
+							LODIndex = LODToRender.GetRayTracedLOD();
+						}
+
+
+						if (EnumHasAnyFlags(Flags, ERayTracingPrimitiveFlags::CacheInstances))
+						{
+							FPrimitiveSceneProxy* SceneProxy = Scene.PrimitiveSceneProxies[PrimitiveIndex];
+							const bool bUsingNaniteRayTracing = (Nanite::GetRayTracingMode() != Nanite::ERayTracingMode::Fallback) && SceneProxy->IsNaniteMesh();
+
+							if (!bUsingNaniteRayTracing)
+							{
+								// Currently IsCachedRayTracingGeometryValid() can only be called for non-nanite geometries
+								checkf(SceneInfo->IsCachedRayTracingGeometryValid(), TEXT("Cached ray tracing instance is expected to be valid. Was mesh LOD streamed but cached data was not invalidated?"));
+								checkf(SceneInfo->CachedRayTracingInstance.GeometryRHI, TEXT("Ray tracing instance must have a valid geometry."));
+							}
+
+							RelevantPrimitive.bAnySegmentsDecal = SceneInfo->bCachedRayTracingInstanceAnySegmentsDecal;
+							RelevantPrimitive.bAllSegmentsDecal = SceneInfo->bCachedRayTracingInstanceAllSegmentsDecal;
+							RelevantPrimitive.CachedRayTracingInstance = &SceneInfo->CachedRayTracingInstance;
+
+							// For primitives with ERayTracingPrimitiveFlags::CacheInstances flag we only cache the instance/mesh commands of the current LOD
+							// (see FPrimitiveSceneInfo::UpdateCachedRayTracingInstance(...) and CacheRayTracingPrimitive(...))
+							LODIndex = 0;
+
+							if (SceneInfo->CachedRayTracingMeshCommandIndicesPerLOD.IsValidIndex(LODIndex))
+							{
+								RelevantPrimitive.CachedRayTracingMeshCommandIndices = SceneInfo->CachedRayTracingMeshCommandIndicesPerLOD[LODIndex];
+							}
+						}
+						else
+						{
+							FRHIRayTracingGeometry* RayTracingGeometryInstance = SceneInfo->GetStaticRayTracingGeometryInstance(LODIndex);
+							if (RayTracingGeometryInstance == nullptr)
+							{
+								return;
+							}
+
+							// Sometimes LODIndex is out of range because it is clamped by ClampToFirstLOD, like the requested LOD is being streamed in and hasn't been available
+							// According to InitViews, we should hide the static mesh instance
+							check(EnumHasAnyFlags(Flags, ERayTracingPrimitiveFlags::CacheMeshCommands));
+							if (SceneInfo->CachedRayTracingMeshCommandIndicesPerLOD.IsValidIndex(LODIndex))
+							{
+								RelevantPrimitive.LODIndex = LODIndex;
+								RelevantPrimitive.RayTracingGeometryRHI = RayTracingGeometryInstance;
+
+								RelevantPrimitive.CachedRayTracingMeshCommandIndices = SceneInfo->CachedRayTracingMeshCommandIndicesPerLOD[LODIndex];
+								RelevantPrimitive.StateHash = SceneInfo->CachedRayTracingMeshCommandsHashPerLOD[LODIndex];
+
+								// TODO: Cache these flags to avoid having to loop over the RayTracingMeshCommands
+								for (int32 CommandIndex : RelevantPrimitive.CachedRayTracingMeshCommandIndices)
+								{
+									if (CommandIndex >= 0)
+									{
+										const FRayTracingMeshCommand& RayTracingMeshCommand = Scene.CachedRayTracingMeshCommands[CommandIndex];
+
+										RelevantPrimitive.InstanceMask |= RayTracingMeshCommand.InstanceMask;
+										RelevantPrimitive.bAllSegmentsOpaque &= RayTracingMeshCommand.bOpaque;
+										RelevantPrimitive.bAllSegmentsCastShadow &= RayTracingMeshCommand.bCastRayTracedShadows;
+										RelevantPrimitive.bAnySegmentsCastShadow |= RayTracingMeshCommand.bCastRayTracedShadows;
+										RelevantPrimitive.bAnySegmentsDecal |= RayTracingMeshCommand.bDecal;
+										RelevantPrimitive.bAllSegmentsDecal &= RayTracingMeshCommand.bDecal;
+										RelevantPrimitive.bTwoSided |= RayTracingMeshCommand.bTwoSided;
+										RelevantPrimitive.bIsSky |= RayTracingMeshCommand.bIsSky;
+										RelevantPrimitive.bAllSegmentsTranslucent &= RayTracingMeshCommand.bIsTranslucent;
+									}
+									else
+									{
+										// CommandIndex == -1 indicates that the mesh batch has been filtered by FRayTracingMeshProcessor (like the shadow depth pass batch)
+										// Do nothing in this case
+									}
+								}
+
+								ERayTracingViewMaskMode MaskMode = static_cast<ERayTracingViewMaskMode>(Scene.CachedRayTracingMeshCommandsMode);
+
+								RelevantPrimitive.UpdateMasks(Flags, MaskMode);
+							}
+						}
+					});
+			}, TStatId(), nullptr, ENamedThreads::AnyThread);
+
+		Result.bValid = true;
+	}
+
+	static void AddDebugRayTracingInstanceFlags(ERayTracingInstanceFlags& InOutFlags)
+	{
+		if (GRayTracingDebugForceOpaque)
+		{
+			InOutFlags |= ERayTracingInstanceFlags::ForceOpaque;
+		}
+		if (GRayTracingDebugDisableTriangleCull)
+		{
+			InOutFlags |= ERayTracingInstanceFlags::TriangleCullDisable;
+		}
+	}
+
+	// Class to implement build instance mask and flags so that rendering related mask build is maintained in any renderer module.
+	// BuildInstanceMaskAndFlags() will be called in the Engine module where it does not know specifics of the ray tracing instance
+	// masks used by the renderer (e.g., path tracer mask might be different from raytracing mask).
+	struct FDeferredShadingRayTracingMaterialGatheringContext : public FRayTracingMaterialGatheringContext
+	{
+		FDeferredShadingRayTracingMaterialGatheringContext(
+			const FScene* InScene,
+			const FSceneView* InReferenceView,
+			const FSceneViewFamily& InReferenceViewFamily,
+			FRDGBuilder& InGraphBuilder,
+			FRayTracingMeshResourceCollector& InRayTracingMeshResourceCollector,
+			FGlobalDynamicReadBuffer& InDynamicReadBuffer)
+			:FRayTracingMaterialGatheringContext(InScene, InReferenceView, InReferenceViewFamily, InGraphBuilder, InRayTracingMeshResourceCollector, InDynamicReadBuffer) {}
+
+		virtual FRayTracingMaskAndFlags BuildInstanceMaskAndFlags(const FRayTracingInstance& Instance, const FPrimitiveSceneProxy& ScenePrimitive) override
+		{
+			return BuildRayTracingInstanceMaskAndFlags(Instance, ScenePrimitive, &ReferenceViewFamily);
+		}
+	};
+
+	bool GatherWorldInstancesForView(
+		FRDGBuilder& GraphBuilder,
+		FScene& Scene,
+		FViewInfo& View,
+		FRayTracingScene& RayTracingScene,
+		FGlobalDynamicReadBuffer& InDynamicReadBuffer,
+		FSceneRenderingBulkObjectAllocator& InBulkAllocator,
+		FRelevantPrimitiveList& RelevantPrimitiveList)
+	{
+		TRACE_CPUPROFILER_EVENT_SCOPE(GatherRayTracingWorldInstances);
+		SCOPE_CYCLE_COUNTER(STAT_GatherRayTracingWorldInstances);
+
+		// Prepare ray tracing scene instance list
+		checkf(RelevantPrimitiveList.bValid, TEXT("Ray tracing relevant primitive list is expected to have been created before GatherRayTracingWorldInstancesForView() is called."));
+
+		// Check that any invalidated cached uniform expressions have been updated on the rendering thread.
+		// Normally this work is done through FMaterialRenderProxy::UpdateUniformExpressionCacheIfNeeded,
+		// however ray tracing material processing (FMaterialShader::GetShaderBindings, which accesses UniformExpressionCache)
+		// is done on task threads, therefore all work must be done here up-front as UpdateUniformExpressionCacheIfNeeded is not free-threaded.
+		check(!FMaterialRenderProxy::HasDeferredUniformExpressionCacheRequests());
+
+		FGPUScenePrimitiveCollector DummyDynamicPrimitiveCollector;
+
+		View.DynamicRayTracingMeshCommandStorage.Reserve(Scene.Primitives.Num());
+		View.VisibleRayTracingMeshCommands.Reserve(Scene.Primitives.Num());
+
+		View.RayTracingMeshResourceCollector = MakeUnique<FRayTracingMeshResourceCollector>(Scene.GetFeatureLevel(), InBulkAllocator);
+
+		View.RayTracingCullingParameters.Init(View);
+
+		FDeferredShadingRayTracingMaterialGatheringContext MaterialGatheringContext
+		(
+			&Scene,
+			&View,
+			*View.Family,
+			GraphBuilder,
+			*View.RayTracingMeshResourceCollector,
+			InDynamicReadBuffer
+		);
+
+		const float CurrentWorldTime = View.Family->Time.GetWorldTimeSeconds();
+
+		// Consume output of the relevant primitive gathering task
+		RayTracingScene.UsedCoarseMeshStreamingHandles = MoveTemp(RelevantPrimitiveList.UsedCoarseMeshStreamingHandles);
+
+		// Inform the coarse mesh streaming manager about all the used streamable render assets in the scene
+		Nanite::FCoarseMeshStreamingManager* CoarseMeshSM = IStreamingManager::Get().GetNaniteCoarseMeshStreamingManager();
+		if (CoarseMeshSM)
+		{
+			CoarseMeshSM->AddUsedStreamingHandles(RayTracingScene.UsedCoarseMeshStreamingHandles);
+		}
+
+		INC_DWORD_STAT_BY(STAT_VisibleRayTracingPrimitives, RelevantPrimitiveList.DynamicPrimitives.Num() + RelevantPrimitiveList.StaticPrimitives.Num());
+
+		{
+			TRACE_CPUPROFILER_EVENT_SCOPE(GatherRayTracingWorldInstances_DynamicElements);
+
+			const bool bParallelMeshBatchSetup = GRayTracingParallelMeshBatchSetup && FApp::ShouldUseThreadingForPerformance();
+
+			const int64 SharedBufferGenerationID = Scene.GetRayTracingDynamicGeometryCollection()->BeginUpdate();
+
+			struct FRayTracingMeshBatchWorkItem
+			{
+				const FPrimitiveSceneProxy* SceneProxy = nullptr;
+				TArray<FMeshBatch> MeshBatchesOwned;
+				TArrayView<const FMeshBatch> MeshBatchesView;
+				uint32 InstanceIndex;
+				uint32 DecalInstanceIndex;
+
+				TArrayView<const FMeshBatch> GetMeshBatches() const
+				{
+					if (MeshBatchesOwned.Num())
+					{
+						check(MeshBatchesView.Num() == 0);
+						return TArrayView<const FMeshBatch>(MeshBatchesOwned);
+					}
+					else
+					{
+						check(MeshBatchesOwned.Num() == 0);
+						return MeshBatchesView;
+					}
+				}
+			};
+
+			static constexpr uint32 MaxWorkItemsPerPage = 128; // Try to keep individual pages small to avoid slow-path memory allocations
+			struct FRayTracingMeshBatchTaskPage
+			{
+				FRayTracingMeshBatchWorkItem WorkItems[MaxWorkItemsPerPage];
+				uint32 NumWorkItems = 0;
+				FRayTracingMeshBatchTaskPage* Next = nullptr;
+			};
+
+			FRayTracingMeshBatchTaskPage* MeshBatchTaskHead = nullptr;
+			FRayTracingMeshBatchTaskPage* MeshBatchTaskPage = nullptr;
+			uint32 NumPendingMeshBatches = 0;
+			const uint32 RayTracingParallelMeshBatchSize = GRayTracingParallelMeshBatchSize;
+
+			auto KickRayTracingMeshBatchTask = [&InBulkAllocator, &View, &Scene, &MeshBatchTaskHead, &MeshBatchTaskPage, &NumPendingMeshBatches]()
+				{
+					if (MeshBatchTaskHead)
+					{
+						FDynamicRayTracingMeshCommandStorage* TaskDynamicCommandStorage = InBulkAllocator.Create<FDynamicRayTracingMeshCommandStorage>();
+						View.DynamicRayTracingMeshCommandStoragePerTask.Add(TaskDynamicCommandStorage);
+
+						FRayTracingMeshCommandOneFrameArray* TaskVisibleCommands = InBulkAllocator.Create<FRayTracingMeshCommandOneFrameArray>();
+						TaskVisibleCommands->Reserve(NumPendingMeshBatches);
+						View.VisibleRayTracingMeshCommandsPerTask.Add(TaskVisibleCommands);
+
+						View.AddRayTracingMeshBatchTaskList.Add(FFunctionGraphTask::CreateAndDispatchWhenReady(
+							[TaskDataHead = MeshBatchTaskHead, &View, &Scene, TaskDynamicCommandStorage, TaskVisibleCommands]()
+							{
+								FTaskTagScope TaskTagScope(ETaskTag::EParallelRenderingThread);
+								TRACE_CPUPROFILER_EVENT_SCOPE(RayTracingMeshBatchTask);
+								FRayTracingMeshBatchTaskPage* Page = TaskDataHead;
+								const int32 ExpectedMaxVisibieCommands = TaskVisibleCommands->Max();
+								while (Page)
+								{
+									for (uint32 ItemIndex = 0; ItemIndex < Page->NumWorkItems; ++ItemIndex)
+									{
+										const FRayTracingMeshBatchWorkItem& WorkItem = Page->WorkItems[ItemIndex];
+										TArrayView<const FMeshBatch> MeshBatches = WorkItem.GetMeshBatches();
+										for (int32 SegmentIndex = 0; SegmentIndex < MeshBatches.Num(); SegmentIndex++)
+										{
+											const FMeshBatch& MeshBatch = MeshBatches[SegmentIndex];
+											FDynamicRayTracingMeshCommandContext CommandContext(
+												*TaskDynamicCommandStorage, *TaskVisibleCommands,
+												SegmentIndex, WorkItem.InstanceIndex, WorkItem.DecalInstanceIndex);
+											FMeshPassProcessorRenderState PassDrawRenderState;
+											FRayTracingMeshProcessor RayTracingMeshProcessor(&CommandContext, &Scene, &View, PassDrawRenderState, Scene.CachedRayTracingMeshCommandsMode);
+											RayTracingMeshProcessor.AddMeshBatch(MeshBatch, 1, WorkItem.SceneProxy);
+										}
+									}
+									FRayTracingMeshBatchTaskPage* NextPage = Page->Next;
+									Page = NextPage;
+								}
+								check(ExpectedMaxVisibieCommands <= TaskVisibleCommands->Max());
+							}, TStatId(), nullptr, ENamedThreads::AnyThread));
+					}
+
+					MeshBatchTaskHead = nullptr;
+					MeshBatchTaskPage = nullptr;
+					NumPendingMeshBatches = 0;
+				};
+
+			// Local temporary array of instances used for GetDynamicRayTracingInstances()
+			TArray<FRayTracingInstance> TempRayTracingInstances;
+
+			for (const FRelevantPrimitive& RelevantPrimitive : RelevantPrimitiveList.DynamicPrimitives)
+			{
+				const FPersistentPrimitiveIndex PersistentPrimitiveIndex = RelevantPrimitive.PersistentPrimitiveIndex;
+				const int32 PrimitiveIndex = RelevantPrimitive.PrimitiveIndex;
+				FPrimitiveSceneProxy* SceneProxy = Scene.PrimitiveSceneProxies[PrimitiveIndex];
+
+				TempRayTracingInstances.Reset();
+				MaterialGatheringContext.DynamicRayTracingGeometriesToUpdate.Reset();
+
+				SceneProxy->GetDynamicRayTracingInstances(MaterialGatheringContext, TempRayTracingInstances);
+
+				for (const FRayTracingDynamicGeometryUpdateParams& DynamicRayTracingGeometryUpdate : MaterialGatheringContext.DynamicRayTracingGeometriesToUpdate)
+				{
+					Scene.GetRayTracingDynamicGeometryCollection()->AddDynamicMeshBatchForGeometryUpdate(
+						GraphBuilder.RHICmdList,
+						&Scene,
+						&View,
+						SceneProxy,
+						DynamicRayTracingGeometryUpdate,
+						PersistentPrimitiveIndex.Index
+					);
+				}
+
+				if (TempRayTracingInstances.Num() > 0)
+				{
+					for (FRayTracingInstance& Instance : TempRayTracingInstances)
+					{
+						const FRayTracingGeometry* Geometry = Instance.Geometry;
+
+						if (!ensureMsgf(Geometry->DynamicGeometrySharedBufferGenerationID == FRayTracingGeometry::NonSharedVertexBuffers
+							|| Geometry->DynamicGeometrySharedBufferGenerationID == SharedBufferGenerationID,
+							TEXT("GenerationID %lld, but expected to be %lld or %lld. Geometry debug name: '%s'. ")
+							TEXT("When shared vertex buffers are used, the contents is expected to be written every frame. ")
+							TEXT("Possibly AddDynamicMeshBatchForGeometryUpdate() was not called for this geometry."),
+							Geometry->DynamicGeometrySharedBufferGenerationID, SharedBufferGenerationID, FRayTracingGeometry::NonSharedVertexBuffers,
+							*Geometry->Initializer.DebugName.ToString()))
+						{
+							continue;
+						}
+
+						// If geometry still has pending build request then add to list which requires a force build
+						if (Geometry->HasPendingBuildRequest())
+						{
+							RayTracingScene.GeometriesToBuild.Add(Geometry);
+						}
+
+						// Validate the material/segment counts
+						if (!ensureMsgf(Instance.GetMaterials().Num() == Geometry->Initializer.Segments.Num() ||
+							(Geometry->Initializer.Segments.Num() == 0 && Instance.GetMaterials().Num() == 1),
+							TEXT("Ray tracing material assignment validation failed for geometry '%s'. "
+								"Instance.GetMaterials().Num() = %d, Geometry->Initializer.Segments.Num() = %d, Instance.Mask = 0x%X."),
+							*Geometry->Initializer.DebugName.ToString(), Instance.GetMaterials().Num(),
+							Geometry->Initializer.Segments.Num(), Instance.MaskAndFlags.Mask))
+						{
+							continue;
+						}
+
+						// Autobuild of InstanceMaskAndFlags if the mask and flags are not built
+						UpdateRayTracingInstanceMaskAndFlagsIfNeeded(Instance, *SceneProxy, View.Family);
+
+						// if primitive has mixed decal and non-decal segments we need to have two ray tracing instances
+						// one containing non-decal segments and the other with decal segments
+						// masking of segments is done using "hidden" hitgroups
+						// TODO: Debug Visualization to highlight primitives using this?
+						const bool bNeedSeparateDecalInstance = Instance.MaskAndFlags.bAnySegmentsDecal && !Instance.MaskAndFlags.bAllSegmentsDecal;
+
+						if (GRayTracingExcludeDecals && Instance.MaskAndFlags.bAnySegmentsDecal && !bNeedSeparateDecalInstance)
+						{
+							continue;
+						}
+
+						FRayTracingGeometryInstance RayTracingInstance;
+						RayTracingInstance.GeometryRHI = Geometry->RayTracingGeometryRHI;
+						checkf(RayTracingInstance.GeometryRHI, TEXT("Ray tracing instance must have a valid geometry."));
+						RayTracingInstance.DefaultUserData = PersistentPrimitiveIndex.Index;
+						RayTracingInstance.bApplyLocalBoundsTransform = Instance.bApplyLocalBoundsTransform;
+						RayTracingInstance.LayerIndex = (uint8)(Instance.MaskAndFlags.bAnySegmentsDecal && !bNeedSeparateDecalInstance ? ERayTracingSceneLayer::Decals : ERayTracingSceneLayer::Base);
+						RayTracingInstance.Mask = Instance.MaskAndFlags.Mask;
+
+						if (Instance.MaskAndFlags.bForceOpaque)
+						{
+							RayTracingInstance.Flags |= ERayTracingInstanceFlags::ForceOpaque;
+						}
+						if (Instance.MaskAndFlags.bDoubleSided)
+						{
+							RayTracingInstance.Flags |= ERayTracingInstanceFlags::TriangleCullDisable;
+						}
+						AddDebugRayTracingInstanceFlags(RayTracingInstance.Flags);
+
+						if (Instance.InstanceGPUTransformsSRV.IsValid())
+						{
+							RayTracingInstance.NumTransforms = Instance.NumTransforms;
+							RayTracingInstance.GPUTransformsSRV = Instance.InstanceGPUTransformsSRV;
+						}
+						else
+						{
+							if (Instance.OwnsTransforms())
+							{
+								// Slow path: copy transforms to the owned storage
+								checkf(Instance.InstanceTransformsView.Num() == 0, TEXT("InstanceTransformsView is expected to be empty if using InstanceTransforms"));
+								TArrayView<FMatrix> SceneOwnedTransforms = RayTracingScene.Allocate<FMatrix>(Instance.InstanceTransforms.Num());
+								FMemory::Memcpy(SceneOwnedTransforms.GetData(), Instance.InstanceTransforms.GetData(), Instance.InstanceTransforms.Num() * sizeof(RayTracingInstance.Transforms[0]));
+								static_assert(std::is_same_v<decltype(SceneOwnedTransforms[0]), decltype(Instance.InstanceTransforms[0])>, "Unexpected transform type");
+
+								RayTracingInstance.NumTransforms = SceneOwnedTransforms.Num();
+								RayTracingInstance.Transforms = SceneOwnedTransforms;
+							}
+							else
+							{
+								// Fast path: just reference persistently-allocated transforms and avoid a copy
+								checkf(Instance.InstanceTransforms.Num() == 0, TEXT("InstanceTransforms is expected to be empty if using InstanceTransformsView"));
+								RayTracingInstance.NumTransforms = Instance.InstanceTransformsView.Num();
+								RayTracingInstance.Transforms = Instance.InstanceTransformsView;
+							}
+						}
+
+						const uint32 InstanceIndex = RayTracingScene.AddInstance(RayTracingInstance, SceneProxy, true);
+
+						uint32 DecalInstanceIndex = INDEX_NONE;
+						if (bNeedSeparateDecalInstance && !GRayTracingExcludeDecals)
+						{
+							FRayTracingGeometryInstance DecalRayTracingInstance = RayTracingInstance;
+							DecalRayTracingInstance.LayerIndex = (uint8)ERayTracingSceneLayer::Decals;
+
+							DecalInstanceIndex = RayTracingScene.AddInstance(MoveTemp(DecalRayTracingInstance), SceneProxy, true);
+						}
+
+						if (bParallelMeshBatchSetup)
+						{
+							if (NumPendingMeshBatches >= RayTracingParallelMeshBatchSize)
+							{
+								KickRayTracingMeshBatchTask();
+							}
+
+							if (MeshBatchTaskPage == nullptr || MeshBatchTaskPage->NumWorkItems == MaxWorkItemsPerPage)
+							{
+								FRayTracingMeshBatchTaskPage* NextPage = InBulkAllocator.Create<FRayTracingMeshBatchTaskPage>();
+								if (MeshBatchTaskHead == nullptr)
+								{
+									MeshBatchTaskHead = NextPage;
+								}
+								if (MeshBatchTaskPage)
+								{
+									MeshBatchTaskPage->Next = NextPage;
+								}
+								MeshBatchTaskPage = NextPage;
+							}
+
+							FRayTracingMeshBatchWorkItem& WorkItem = MeshBatchTaskPage->WorkItems[MeshBatchTaskPage->NumWorkItems];
+							MeshBatchTaskPage->NumWorkItems++;
+
+							NumPendingMeshBatches += Instance.GetMaterials().Num();
+
+							if (Instance.OwnsMaterials())
+							{
+								Swap(WorkItem.MeshBatchesOwned, Instance.Materials);
+							}
+							else
+							{
+								WorkItem.MeshBatchesView = Instance.MaterialsView;
+							}
+
+							WorkItem.SceneProxy = SceneProxy;
+							WorkItem.InstanceIndex = InstanceIndex;
+							WorkItem.DecalInstanceIndex = DecalInstanceIndex;
+						}
+						else
+						{
+							TArrayView<const FMeshBatch> InstanceMaterials = Instance.GetMaterials();
+							for (int32 SegmentIndex = 0; SegmentIndex < InstanceMaterials.Num(); SegmentIndex++)
+							{
+								const FMeshBatch& MeshBatch = InstanceMaterials[SegmentIndex];
+								FDynamicRayTracingMeshCommandContext CommandContext(View.DynamicRayTracingMeshCommandStorage, View.VisibleRayTracingMeshCommands, SegmentIndex, InstanceIndex, DecalInstanceIndex);
+								FMeshPassProcessorRenderState PassDrawRenderState;
+								FRayTracingMeshProcessor RayTracingMeshProcessor(&CommandContext, &Scene, &View, PassDrawRenderState, Scene.CachedRayTracingMeshCommandsMode);
+								RayTracingMeshProcessor.AddMeshBatch(MeshBatch, 1, SceneProxy);
+							}
+						}
+					}
+
+					if (CVarRayTracingDynamicGeometryLastRenderTimeUpdateDistance.GetValueOnRenderThread() > 0.0f)
+					{
+						if (FVector::Distance(SceneProxy->GetActorPosition(), View.ViewMatrices.GetViewOrigin()) < CVarRayTracingDynamicGeometryLastRenderTimeUpdateDistance.GetValueOnRenderThread())
+						{
+							FPrimitiveSceneInfo* SceneInfo = Scene.Primitives[PrimitiveIndex];
+							// Update LastRenderTime for components so that visibility based ticking (like skeletal meshes) can get updated
+							// We are only doing this for dynamic geometries now
+							SceneInfo->LastRenderTime = CurrentWorldTime;
+							SceneInfo->UpdateComponentLastRenderTime(CurrentWorldTime, /*bUpdateLastRenderTimeOnScreen=*/true);
+						}
+					}
+				}
+			}
+
+			KickRayTracingMeshBatchTask();
+		}
+
+		// Task to iterate over static ray tracing instances, perform auto-instancing and culling.
+		// This adds final instances to the ray tracing scene and must be done before FRayTracingScene::BuildInitializationData().
+		struct FRayTracingSceneAddInstancesTask
+		{
+			UE_NONCOPYABLE(FRayTracingSceneAddInstancesTask)
+
+				static ESubsequentsMode::Type GetSubsequentsMode() { return ESubsequentsMode::TrackSubsequents; }
+			TStatId                       GetStatId() const { return TStatId(); }
+			ENamedThreads::Type           GetDesiredThread() { return ENamedThreads::AnyThread; }
+
+			// Inputs
+
+			const FScene& Scene;
+			TChunkedArray<FRelevantPrimitive>& RelevantStaticPrimitives;
+			const FRayTracingCullingParameters& CullingParameters;
+			const bool bIsPathTracing;
+
+			// Outputs
+
+			FRayTracingScene& RayTracingScene; // New instances are added into FRayTracingScene::Instances and FRayTracingScene::Allocator is used for temporary data
+			TArray<FVisibleRayTracingMeshCommand>& VisibleRayTracingMeshCommands; // New elements are added here by this task
+
+			FRayTracingSceneAddInstancesTask(const FScene& InScene,
+				TChunkedArray<FRelevantPrimitive>& InRelevantStaticPrimitives,
+				const FRayTracingCullingParameters& InCullingParameters,
+				const bool bInIsPathTracing,
+				FRayTracingScene& InRayTracingScene, TArray<FVisibleRayTracingMeshCommand>& InVisibleRayTracingMeshCommands)
+				: Scene(InScene)
+				, RelevantStaticPrimitives(InRelevantStaticPrimitives)
+				, CullingParameters(InCullingParameters)
+				, bIsPathTracing(bInIsPathTracing)
+				, RayTracingScene(InRayTracingScene)
+				, VisibleRayTracingMeshCommands(InVisibleRayTracingMeshCommands)
+			{
+				VisibleRayTracingMeshCommands.Reserve(RelevantStaticPrimitives.Num());
+			}
+
+			// TODO: Consider moving auto instance batching logic into FRayTracingScene
+
+			struct FAutoInstanceBatch
+			{
+				int32 Index = INDEX_NONE;
+				int32 DecalIndex = INDEX_NONE;
+
+				// Copies the next InstanceSceneDataOffset and user data into the current batch, returns true if arrays were re-allocated.
+				bool Add(FRayTracingScene& InRayTracingScene, uint32 InInstanceSceneDataOffset, uint32 InUserData)
+				{
+					// Adhoc TArray-like resize behavior, in lieu of support for using a custom FMemStackBase in TArray.
+					// Idea for future: if batch becomes large enough, we could actually split it into multiple instances to avoid memory waste.
+
+					const bool bNeedReallocation = Cursor == InstanceSceneDataOffsets.Num();
+
+					if (bNeedReallocation)
+					{
+						int32 PrevCount = InstanceSceneDataOffsets.Num();
+						int32 NextCount = FMath::Max(PrevCount * 2, 1);
+
+						TArrayView<uint32> NewInstanceSceneDataOffsets = InRayTracingScene.Allocate<uint32>(NextCount);
+						if (PrevCount)
+						{
+							FMemory::Memcpy(NewInstanceSceneDataOffsets.GetData(), InstanceSceneDataOffsets.GetData(), InstanceSceneDataOffsets.GetTypeSize() * InstanceSceneDataOffsets.Num());
+						}
+						InstanceSceneDataOffsets = NewInstanceSceneDataOffsets;
+
+						TArrayView<uint32> NewUserData = InRayTracingScene.Allocate<uint32>(NextCount);
+						if (PrevCount)
+						{
+							FMemory::Memcpy(NewUserData.GetData(), UserData.GetData(), UserData.GetTypeSize() * UserData.Num());
+						}
+						UserData = NewUserData;
+					}
+
+					InstanceSceneDataOffsets[Cursor] = InInstanceSceneDataOffset;
+					UserData[Cursor] = InUserData;
+
+					++Cursor;
+
+					return bNeedReallocation;
+				}
+
+				bool IsValid() const
+				{
+					return InstanceSceneDataOffsets.Num() != 0;
+				}
+
+				TArrayView<uint32> InstanceSceneDataOffsets;
+				TArrayView<uint32> UserData;
+				uint32 Cursor = 0;
+			};
+
+			void DoTask(ENamedThreads::Type CurrentThread, const FGraphEventRef& MyCompletionGraphEvent)
+			{
+				FTaskTagScope TaskTagScope(ETaskTag::EParallelRenderingThread);
+
+				TRACE_CPUPROFILER_EVENT_SCOPE(RayTracingSceneStaticInstanceTask);
+
+				FGraphEventArray CullingTasks;
+
+				const bool bAutoInstance = CVarRayTracingAutoInstance.GetValueOnRenderThread() != 0;
+
+				// Instance batches by FRelevantPrimitive::InstancingKey()
+				Experimental::TSherwoodMap<uint64, FAutoInstanceBatch> InstanceBatches;
+
+				TArray<FRayTracingCullPrimitiveInstancesClosure> CullInstancesClosures;
+				if (CullingParameters.CullingMode != RayTracing::ECullingMode::Disabled && GetRayTracingCullingPerInstance())
+				{
+					CullInstancesClosures.Reserve(RelevantStaticPrimitives.Num());
+					CullingTasks.Reserve(RelevantStaticPrimitives.Num() / 256 + 1);
+				}
+
+				// scan relevant primitives computing hash data to look for duplicate instances
+				for (const FRelevantPrimitive& RelevantPrimitive : RelevantStaticPrimitives)
+				{
+					const int32 PrimitiveIndex = RelevantPrimitive.PrimitiveIndex;
+					FPrimitiveSceneInfo* SceneInfo = Scene.Primitives[PrimitiveIndex];
+					FPrimitiveSceneProxy* SceneProxy = Scene.PrimitiveSceneProxies[PrimitiveIndex];
+					ERayTracingPrimitiveFlags Flags = Scene.PrimitiveRayTracingFlags[PrimitiveIndex];
+					const FPersistentPrimitiveIndex PersistentPrimitiveIndex = RelevantPrimitive.PersistentPrimitiveIndex;
+
+					if (EnumHasAnyFlags(Flags, ERayTracingPrimitiveFlags::CacheInstances))
+					{
+						const bool bUsingNaniteRayTracing = (Nanite::GetRayTracingMode() != Nanite::ERayTracingMode::Fallback) && SceneProxy->IsNaniteMesh();
+
+						if (bUsingNaniteRayTracing)
+						{
+							Nanite::GRayTracingManager.AddVisiblePrimitive(SceneInfo);
+
+							if (RelevantPrimitive.CachedRayTracingInstance->GeometryRHI == nullptr)
+							{
+								// Nanite ray tracing geometry not ready yet, doesn't include primitive in ray tracing scene
+								continue;
+							}
+						}
+
+						// TODO: Consider requesting a recache of all ray tracing commands during which decals are excluded
+
+						// if primitive has mixed decal and non-decal segments we need to have two ray tracing instances
+						// one containing non-decal segments and the other with decal segments
+						// masking of segments is done using "hidden" hitgroups
+						// TODO: Debug Visualization to highlight primitives using this?
+						const bool bNeedSeparateDecalInstance = RelevantPrimitive.bAnySegmentsDecal && !RelevantPrimitive.bAllSegmentsDecal;
+
+						if (GRayTracingExcludeDecals && RelevantPrimitive.bAnySegmentsDecal && !bNeedSeparateDecalInstance)
+						{
+							continue;
+						}
+
+						check(RelevantPrimitive.CachedRayTracingInstance);
+
+						const int32 NewInstanceIndex = RayTracingScene.AddInstance(*RelevantPrimitive.CachedRayTracingInstance, SceneProxy, false);
+						uint32 DecalInstanceIndex = INDEX_NONE;
+
+						{
+							FRayTracingGeometryInstance& NewInstance = RayTracingScene.GetInstance(NewInstanceIndex);
+							AddDebugRayTracingInstanceFlags(NewInstance.Flags);
+
+							NewInstance.LayerIndex = (uint8)(RelevantPrimitive.bAnySegmentsDecal && !bNeedSeparateDecalInstance ? ERayTracingSceneLayer::Decals : ERayTracingSceneLayer::Base);
+
+							const Experimental::FHashElementId GroupId = Scene.PrimitiveRayTracingGroupIds[PrimitiveIndex];
+							const bool bUseGroupBounds = CullingParameters.bCullUsingGroupIds && GroupId.IsValid();
+
+							if (CullingParameters.CullingMode != RayTracing::ECullingMode::Disabled && GetRayTracingCullingPerInstance() && RelevantPrimitive.CachedRayTracingInstance->NumTransforms > 1 && !bUseGroupBounds)
+							{
+								const bool bIsFarFieldPrimitive = EnumHasAnyFlags(Flags, ERayTracingPrimitiveFlags::FarField);
+
+								TArrayView<uint32> InstanceActivationMask = RayTracingScene.Allocate<uint32>(FMath::DivideAndRoundUp(NewInstance.NumTransforms, 32u));
+
+								NewInstance.ActivationMask = InstanceActivationMask;
+
+								FRayTracingCullPrimitiveInstancesClosure Closure;
+								Closure.Scene = &Scene;
+								Closure.SceneInfo = SceneInfo;
+								Closure.PrimitiveIndex = PrimitiveIndex;
+								Closure.bIsFarFieldPrimitive = bIsFarFieldPrimitive;
+								Closure.CullingParameters = &CullingParameters;
+								Closure.OutInstanceActivationMask = InstanceActivationMask;
+
+								CullInstancesClosures.Add(MoveTemp(Closure));
+
+								if (CullInstancesClosures.Num() >= 256)
+								{
+									CullingTasks.Add(FFunctionGraphTask::CreateAndDispatchWhenReady([CullInstancesClosures = MoveTemp(CullInstancesClosures)]()
+										{
+											for (auto& Closure : CullInstancesClosures)
+											{
+												Closure();
+											}
+										}, TStatId(), nullptr, ENamedThreads::AnyThread));
+								}
+							}
+
+							if (bNeedSeparateDecalInstance && !GRayTracingExcludeDecals)
+							{
+								FRayTracingGeometryInstance DecalRayTracingInstance = NewInstance;
+								DecalRayTracingInstance.LayerIndex = (uint8)ERayTracingSceneLayer::Decals;
+
+								DecalInstanceIndex = RayTracingScene.AddInstance(MoveTemp(DecalRayTracingInstance), SceneProxy, false);
+							}
+						}
+
+						// At the moment we only support SM & ISMs on this path
+						check(EnumHasAnyFlags(Flags, ERayTracingPrimitiveFlags::CacheMeshCommands));
+
+						const bool bHasDecalInstanceIndex = DecalInstanceIndex != INDEX_NONE;
+
+						for (int32 CommandIndex : RelevantPrimitive.CachedRayTracingMeshCommandIndices)
+						{
+							const FRayTracingMeshCommand& MeshCommand = Scene.CachedRayTracingMeshCommands[CommandIndex];
+
+							{
+								const bool bHidden = bHasDecalInstanceIndex && MeshCommand.bDecal;
+								FVisibleRayTracingMeshCommand NewVisibleMeshCommand(&MeshCommand, NewInstanceIndex, bHidden);
+								VisibleRayTracingMeshCommands.Add(NewVisibleMeshCommand);
+							}
+
+							if (bHasDecalInstanceIndex)
+							{
+								const bool bHidden = !MeshCommand.bDecal;
+								FVisibleRayTracingMeshCommand NewVisibleMeshCommand(&MeshCommand, DecalInstanceIndex, bHidden);
+								VisibleRayTracingMeshCommands.Add(NewVisibleMeshCommand);
+							}
+						}
+					}
+					else
+					{
+						const int8 LODIndex = RelevantPrimitive.LODIndex;
+
+						if (LODIndex < 0 || !RelevantPrimitive.bStatic)
+						{
+							continue; // skip dynamic primitives and other 
+						}
+
+						// if primitive has mixed decal and non-decal segments we need to have two ray tracing instances
+						// one containing non-decal segments and the other with decal segments
+						// masking of segments is done using "hidden" hitgroups
+						// TODO: Debug Visualization to highlight primitives using this?
+						const bool bNeedSeparateDecalInstance = RelevantPrimitive.bAnySegmentsDecal && !RelevantPrimitive.bAllSegmentsDecal;
+
+						if (GRayTracingExcludeDecals && RelevantPrimitive.bAnySegmentsDecal && !bNeedSeparateDecalInstance)
+						{
+							continue;
+						}
+
+						if ((GRayTracingExcludeDecals && RelevantPrimitive.bAnySegmentsDecal)
+							|| (GRayTracingExcludeTranslucent && RelevantPrimitive.bAllSegmentsTranslucent)
+							|| (GRayTracingExcludeSky && RelevantPrimitive.bIsSky && !bIsPathTracing))
+						{
+							continue;
+						}
+
+						// location if this is a new entry
+						const uint64 InstanceKey = RelevantPrimitive.InstancingKey();
+
+						FAutoInstanceBatch DummyInstanceBatch = { };
+						FAutoInstanceBatch& InstanceBatch = bAutoInstance ? InstanceBatches.FindOrAdd(InstanceKey, DummyInstanceBatch) : DummyInstanceBatch;
+
+						if (InstanceBatch.IsValid())
+						{
+							// Reusing a previous entry, just append to the instance list.
+
+							bool bReallocated = InstanceBatch.Add(RayTracingScene, SceneInfo->GetInstanceSceneDataOffset(), uint32(PersistentPrimitiveIndex.Index));
+
+							check(InstanceBatch.Index != INDEX_NONE);
+							{
+								FRayTracingGeometryInstance& RayTracingInstance = RayTracingScene.GetInstance(InstanceBatch.Index);
+								++RayTracingInstance.NumTransforms;
+								check(RayTracingInstance.NumTransforms == InstanceBatch.Cursor); // sanity check
+
+								if (bReallocated)
+								{
+									RayTracingInstance.InstanceSceneDataOffsets = InstanceBatch.InstanceSceneDataOffsets;
+									RayTracingInstance.UserData = InstanceBatch.UserData;
+								}
+							}
+
+							if (InstanceBatch.DecalIndex != INDEX_NONE)
+							{
+								FRayTracingGeometryInstance& RayTracingInstance = RayTracingScene.GetInstance(InstanceBatch.DecalIndex);
+								++RayTracingInstance.NumTransforms;
+								check(RayTracingInstance.NumTransforms == InstanceBatch.Cursor); // sanity check
+
+								if (bReallocated)
+								{
+									RayTracingInstance.InstanceSceneDataOffsets = InstanceBatch.InstanceSceneDataOffsets;
+									RayTracingInstance.UserData = InstanceBatch.UserData;
+								}
+							}
+						}
+						else
+						{
+							// Starting new instance batch
+
+							InstanceBatch.Add(RayTracingScene, SceneInfo->GetInstanceSceneDataOffset(), uint32(PersistentPrimitiveIndex.Index));
+
+							FRayTracingGeometryInstance RayTracingInstance;
+							RayTracingInstance.GeometryRHI = RelevantPrimitive.RayTracingGeometryRHI;
+							checkf(RayTracingInstance.GeometryRHI, TEXT("Ray tracing instance must have a valid geometry."));
+							RayTracingInstance.InstanceSceneDataOffsets = InstanceBatch.InstanceSceneDataOffsets;
+							RayTracingInstance.UserData = InstanceBatch.UserData;
+							RayTracingInstance.NumTransforms = 1;
+
+							RayTracingInstance.Mask = RelevantPrimitive.InstanceMask; // When no cached command is found, InstanceMask == 0 and the instance is effectively filtered out
+
+							if (RelevantPrimitive.bAllSegmentsOpaque && RelevantPrimitive.bAllSegmentsCastShadow)
+							{
+								RayTracingInstance.Flags |= ERayTracingInstanceFlags::ForceOpaque;
+							}
+							if (RelevantPrimitive.bTwoSided)
+							{
+								RayTracingInstance.Flags |= ERayTracingInstanceFlags::TriangleCullDisable;
+							}
+							AddDebugRayTracingInstanceFlags(RayTracingInstance.Flags);
+
+							RayTracingInstance.LayerIndex = (uint8)(RelevantPrimitive.bAnySegmentsDecal && !bNeedSeparateDecalInstance ? ERayTracingSceneLayer::Decals : ERayTracingSceneLayer::Base);
+
+							InstanceBatch.Index = RayTracingScene.AddInstance(RayTracingInstance, SceneProxy, false);
+
+							if (bNeedSeparateDecalInstance && !GRayTracingExcludeDecals)
+							{
+								FRayTracingGeometryInstance DecalRayTracingInstance = RayTracingInstance;
+								DecalRayTracingInstance.LayerIndex = (uint8)ERayTracingSceneLayer::Decals;
+
+								InstanceBatch.DecalIndex = RayTracingScene.AddInstance(MoveTemp(DecalRayTracingInstance), SceneProxy, false);
+							}
+
+							const bool bHasDecalInstanceIndex = InstanceBatch.DecalIndex != INDEX_NONE;
+
+							for (int32 CommandIndex : RelevantPrimitive.CachedRayTracingMeshCommandIndices)
+							{
+								if (CommandIndex >= 0)
+								{
+									const FRayTracingMeshCommand& MeshCommand = Scene.CachedRayTracingMeshCommands[CommandIndex];
+
+									{
+										const bool bHidden = bHasDecalInstanceIndex && MeshCommand.bDecal;
+										FVisibleRayTracingMeshCommand NewVisibleMeshCommand(&MeshCommand, InstanceBatch.Index, bHidden);
+										VisibleRayTracingMeshCommands.Add(NewVisibleMeshCommand);
+									}
+
+									if (bHasDecalInstanceIndex)
+									{
+										const bool bHidden = !MeshCommand.bDecal;
+										FVisibleRayTracingMeshCommand NewVisibleMeshCommand(&MeshCommand, InstanceBatch.DecalIndex, bHidden);
+										VisibleRayTracingMeshCommands.Add(NewVisibleMeshCommand);
+									}
+								}
+								else
+								{
+									// CommandIndex == -1 indicates that the mesh batch has been filtered by FRayTracingMeshProcessor (like the shadow depth pass batch)
+									// Do nothing in this case
+								}
+							}
+						}
+					}
+				}
+
+				CullingTasks.Add(FFunctionGraphTask::CreateAndDispatchWhenReady([CullInstancesClosures = MoveTemp(CullInstancesClosures)]()
+					{
+						for (auto& Closure : CullInstancesClosures)
+						{
+							Closure();
+						}
+					}, TStatId(), nullptr, ENamedThreads::AnyThread));
+
+				for (FGraphEventRef& CullingTask : CullingTasks)
+				{
+					MyCompletionGraphEvent->DontCompleteUntil(CullingTask);
+				}
+			}
+		};
+
+		FGraphEventArray AddInstancesTaskPrerequisites;
+		AddInstancesTaskPrerequisites.Add(RelevantPrimitiveList.StaticPrimitiveLODTask);
+
+		FGraphEventRef AddInstancesTask = TGraphTask<FRayTracingSceneAddInstancesTask>::CreateTask(&AddInstancesTaskPrerequisites).ConstructAndDispatchWhenReady(
+			Scene, RelevantPrimitiveList.StaticPrimitives, View.RayTracingCullingParameters, bool(View.Family->EngineShowFlags.PathTracing), // inputs 
+			RayTracingScene, View.VisibleRayTracingMeshCommands // outputs
+		);
+
+		// Scene init task can run only when all pre-init tasks are complete (including culling tasks that are spawned while adding instances)
+		View.RayTracingSceneInitTask = FFunctionGraphTask::CreateAndDispatchWhenReady(
+			[&View, &RayTracingScene]()
+			{
+				TRACE_CPUPROFILER_EVENT_SCOPE(RayTracingSceneInitTask);
+				View.RayTracingSceneInitData = RayTracingScene.BuildInitializationData();
+			},
+			TStatId(), AddInstancesTask, ENamedThreads::AnyThread);
+
+		return true;
+	}
+
+	bool ShouldExcludeDecals()
+	{
+		return GRayTracingExcludeDecals != 0;
+	}
+}
+
+#endif //RHI_RAYTRACING
diff --git a/Engine/Source/Runtime/Renderer/Private/RayTracing/RayTracing.h b/Engine/Source/Runtime/Renderer/Private/RayTracing/RayTracing.h
new file mode 100644
--- /dev/null
+++ b/Engine/Source/Runtime/Renderer/Private/RayTracing/RayTracing.h
@@ -0,0 +1,40 @@
+// Copyright Epic Games, Inc. All Rights Reserved.
+
+#pragma once
+
+#include "RHIDefinitions.h"
+
+#if RHI_RAYTRACING
+
+#include "RayTracingInstance.h"
+#include "MeshPassProcessor.h"
+#include "RenderGraphDefinitions.h"
+
+class FRayTracingScene;
+class FGlobalDynamicReadBuffer;
+
+namespace RayTracing
+{
+	struct FRelevantPrimitiveList;
+
+	FRelevantPrimitiveList* CreateRelevantPrimitiveList(FSceneRenderingBulkObjectAllocator& InAllocator);
+
+	// Iterates over Scene's PrimitiveSceneProxies and extracts ones that are relevant for ray tracing.
+	// This function can run on any thread.
+	void GatherRelevantPrimitives(FScene& Scene, const FViewInfo& View, FRelevantPrimitiveList& OutRelevantPrimitiveList);
+
+	// Fills RayTracingScene instance list for the given View and adds relevant ray tracing data to the view. Does not reset previous scene contents.
+	// This function must run on render thread
+	bool GatherWorldInstancesForView(
+		FRDGBuilder& GraphBuilder,
+		FScene& Scene,
+		FViewInfo& View,
+		FRayTracingScene& RayTracingScene,
+		FGlobalDynamicReadBuffer& InDynamicReadBuffer,
+		FSceneRenderingBulkObjectAllocator& InBulkAllocator,
+		FRelevantPrimitiveList& RelevantPrimitiveList);
+
+	bool ShouldExcludeDecals();
+}
+
+#endif // RHI_RAYTRACING
\ No newline at end of file
