Fixed crash due to GDME task running unbounded when number of shadows to gather is 0.

diff --git a/Engine/Source/Runtime/Renderer/Private/ShadowSetup.cpp b/Engine/Source/Runtime/Renderer/Private/ShadowSetup.cpp
--- a/Engine/Source/Runtime/Renderer/Private/ShadowSetup.cpp
+++ b/Engine/Source/Runtime/Renderer/Private/ShadowSetup.cpp
@@ -4528,121 +4528,124 @@ void FSceneRenderer::GatherShadowDynamicMeshElements(FDynamicShadowsTaskData& Ta
 		}
 	}
 
-	// Process all shadows in the serial path when not in multithreaded mode.
-	TaskData.ShadowsToGatherInSerialPass.Init(!TaskData.bMultithreadedGDME, TaskData.ShadowsToGather.Num());
+	const int32 NumShadowTasks = FMath::Min<int32>(GetNumShadowDynamicMeshElementTasks(), TaskData.ShadowsToGather.Num());
 
-	const UE::Tasks::ETaskPriority TaskPriority = UE::Tasks::ETaskPriority::High;
-
-	const UE::Tasks::EExtendedTaskPriority ExtendedTaskPriority = TaskData.bMultithreadedGDME
-		? UE::Tasks::EExtendedTaskPriority::None
-		: UE::Tasks::EExtendedTaskPriority::Inline;
-
-	int32 NumShadowViews = 0;
-	for (FProjectedShadowInfo* ProjectedShadowInfo : TaskData.ShadowsToGather)
-	{
-		NumShadowViews += ProjectedShadowInfo->bOnePassPointLightShadow ? 6 : 1;
-	}
-
-	UE::Tasks::FTask MeshCollectorsTask = UE::Tasks::Launch(UE_SOURCE_LOCATION, [&InstanceCullingManager = TaskData.InstanceCullingManager, NumShadowViews]
+	if (NumShadowTasks > 0)
 	{
-		// Wait to allocate views until after the task event has triggered.
-		InstanceCullingManager.AllocateViews(NumShadowViews);
+		// Process all shadows in the serial path when not in multithreaded mode.
+		TaskData.ShadowsToGatherInSerialPass.Init(!TaskData.bMultithreadedGDME, TaskData.ShadowsToGather.Num());
 
-	}, TaskData.BeginGatherDynamicMeshElementsTask, TaskPriority, ExtendedTaskPriority);
+		const UE::Tasks::ETaskPriority TaskPriority = UE::Tasks::ETaskPriority::High;
 
-	if (TaskData.bMultithreadedGDME)
-	{
-		UE::Tasks::FTaskEvent MeshCollectorsTaskEvent{ UE_SOURCE_LOCATION };
+		const UE::Tasks::EExtendedTaskPriority ExtendedTaskPriority = TaskData.bMultithreadedGDME
+			? UE::Tasks::EExtendedTaskPriority::None
+			: UE::Tasks::EExtendedTaskPriority::Inline;
 
-		const int32 NumTasks = FMath::Min<int32>(GetNumShadowDynamicMeshElementTasks(), TaskData.ShadowsToGather.Num());
+		int32 NumShadowViews = 0;
+		for (FProjectedShadowInfo* ProjectedShadowInfo : TaskData.ShadowsToGather)
+		{
+			NumShadowViews += ProjectedShadowInfo->bOnePassPointLightShadow ? 6 : 1;
+		}
 
-		for (int32 TaskIndex = 0; TaskIndex < NumTasks; ++TaskIndex)
+		UE::Tasks::FTask MeshCollectorsTask = UE::Tasks::Launch(UE_SOURCE_LOCATION, [&InstanceCullingManager = TaskData.InstanceCullingManager, NumShadowViews]
 		{
-			FRHICommandList* RHICmdList = new FRHICommandList(FRHIGPUMask::All());
-			RHICmdList->SwitchPipeline(ERHIPipeline::Graphics);
-			TaskData.CommandLists.Emplace(RHICmdList);
+			// Wait to allocate views until after the task event has triggered.
+			InstanceCullingManager.AllocateViews(NumShadowViews);
 
-			FShadowMeshCollector* MeshCollector = Allocator.Create<FShadowMeshCollector>(*RHICmdList, *this);
-			TaskData.MeshCollectors.Emplace(MeshCollector);
+		}, TaskData.BeginGatherDynamicMeshElementsTask, TaskPriority, ExtendedTaskPriority);
 
-			MeshCollectorsTaskEvent.AddPrerequisites(
-				MeshCollector->GetPipe().Launch(UE_SOURCE_LOCATION,
-					[this, &TaskData, MeshCollector]() mutable
+		if (TaskData.bMultithreadedGDME)
+		{
+			UE::Tasks::FTaskEvent MeshCollectorsTaskEvent{ UE_SOURCE_LOCATION };
+
+			for (int32 TaskIndex = 0; TaskIndex < NumShadowTasks; ++TaskIndex)
 			{
-				FOptionalTaskTagScope Scope(ETaskTag::EParallelRenderingThread);
-				TArray<const FSceneView*> LocalViews;
-				LocalViews.AddZeroed(1);
+				FRHICommandList* RHICmdList = new FRHICommandList(FRHIGPUMask::All());
+				RHICmdList->SwitchPipeline(ERHIPipeline::Graphics);
+				TaskData.CommandLists.Emplace(RHICmdList);
+
+				FShadowMeshCollector* MeshCollector = Allocator.Create<FShadowMeshCollector>(*RHICmdList, *this);
+				TaskData.MeshCollectors.Emplace(MeshCollector);
 
-				while (true)
+				MeshCollectorsTaskEvent.AddPrerequisites(
+					MeshCollector->GetPipe().Launch(UE_SOURCE_LOCATION,
+						[this, &TaskData, MeshCollector]() mutable
 				{
-					// Atomically increment to get the next shadow to process in this task until empty.
-					const int32 ShadowIndex = TaskData.ShadowsToGatherNextIndex.fetch_add(1, std::memory_order_relaxed);
-					if (ShadowIndex >= TaskData.ShadowsToGather.Num())
+					FOptionalTaskTagScope Scope(ETaskTag::EParallelRenderingThread);
+					TArray<const FSceneView*> LocalViews;
+					LocalViews.AddZeroed(1);
+
+					while (true)
 					{
-						break;
-					}
+						// Atomically increment to get the next shadow to process in this task until empty.
+						const int32 ShadowIndex = TaskData.ShadowsToGatherNextIndex.fetch_add(1, std::memory_order_relaxed);
+						if (ShadowIndex >= TaskData.ShadowsToGather.Num())
+						{
+							break;
+						}
 
-					FProjectedShadowInfo& ProjectedShadowInfo = *TaskData.ShadowsToGather[ShadowIndex];
-					FVisibleLightInfo& VisibleLightInfo = VisibleLightInfos[ProjectedShadowInfo.GetLightSceneInfo().Id];
+						FProjectedShadowInfo& ProjectedShadowInfo = *TaskData.ShadowsToGather[ShadowIndex];
+						FVisibleLightInfo& VisibleLightInfo = VisibleLightInfos[ProjectedShadowInfo.GetLightSceneInfo().Id];
 
-					const bool bProcessedAllPrimitives = ProjectedShadowInfo.GatherDynamicMeshElements(MeshCollector->GetCollector(), *this, VisibleLightInfo, LocalViews, FProjectedShadowInfo::EGatherDynamicMeshElementsPass::Parallel);
+						const bool bProcessedAllPrimitives = ProjectedShadowInfo.GatherDynamicMeshElements(MeshCollector->GetCollector(), *this, VisibleLightInfo, LocalViews, FProjectedShadowInfo::EGatherDynamicMeshElementsPass::Parallel);
 
-					// Atomically mark shadows that require a second serial pass.
-					if (!bProcessedAllPrimitives)
-					{
-						TaskData.ShadowsToGatherInSerialPass[ShadowIndex].AtomicSet(true);
-					}
+						// Atomically mark shadows that require a second serial pass.
+						if (!bProcessedAllPrimitives)
+						{
+							TaskData.ShadowsToGatherInSerialPass[ShadowIndex].AtomicSet(true);
+						}
 
-					// Setup mesh draw command passes immediately if we don't have any more primitives to process in the serial pass.
-					if (bProcessedAllPrimitives)
-					{
-						ProjectedShadowInfo.SetupMeshDrawCommandsForProjectionStenciling(*this, TaskData.InstanceCullingManager);
-						ProjectedShadowInfo.SetupMeshDrawCommandsForShadowDepth(*this, TaskData.InstanceCullingManager);
+						// Setup mesh draw command passes immediately if we don't have any more primitives to process in the serial pass.
+						if (bProcessedAllPrimitives)
+						{
+							ProjectedShadowInfo.SetupMeshDrawCommandsForProjectionStenciling(*this, TaskData.InstanceCullingManager);
+							ProjectedShadowInfo.SetupMeshDrawCommandsForShadowDepth(*this, TaskData.InstanceCullingManager);
+						}
 					}
-				}
 
-			}, MeshCollectorsTask, TaskPriority));
-		}
+				}, MeshCollectorsTask, TaskPriority));
+			}
 
-		MeshCollectorsTaskEvent.Trigger();
-		MeshCollectorsTask = MoveTemp(MeshCollectorsTaskEvent);
-	}
+			MeshCollectorsTaskEvent.Trigger();
+			MeshCollectorsTask = MoveTemp(MeshCollectorsTaskEvent);
+		}
 
-	// Process only serial primitives if in multithreaded mode, otherwise process all of them.
-	const FProjectedShadowInfo::EGatherDynamicMeshElementsPass SerialPass = TaskData.bMultithreadedGDME
-		? FProjectedShadowInfo::EGatherDynamicMeshElementsPass::Serial
-		: FProjectedShadowInfo::EGatherDynamicMeshElementsPass::All;
+		// Process only serial primitives if in multithreaded mode, otherwise process all of them.
+		const FProjectedShadowInfo::EGatherDynamicMeshElementsPass SerialPass = TaskData.bMultithreadedGDME
+			? FProjectedShadowInfo::EGatherDynamicMeshElementsPass::Serial
+			: FProjectedShadowInfo::EGatherDynamicMeshElementsPass::All;
 
-	TaskData.SetupMeshPassTask.AddPrerequisites(UE::Tasks::Launch(UE_SOURCE_LOCATION, [this, &TaskData, SerialPass] () mutable
-	{
-		FOptionalTaskTagScope Scope(ETaskTag::EParallelRenderingThread);
-		TRACE_CPUPROFILER_EVENT_SCOPE(SetupShadowMeshPass);
+		TaskData.SetupMeshPassTask.AddPrerequisites(UE::Tasks::Launch(UE_SOURCE_LOCATION, [this, &TaskData, SerialPass] () mutable
+		{
+			FOptionalTaskTagScope Scope(ETaskTag::EParallelRenderingThread);
+			TRACE_CPUPROFILER_EVENT_SCOPE(SetupShadowMeshPass);
 
-		TArray<const FSceneView*> LocalViews;
+			TArray<const FSceneView*> LocalViews;
 
-		for (TConstSetBitIterator<SceneRenderingBitArrayAllocator> BitIt(TaskData.ShadowsToGatherInSerialPass); BitIt; ++BitIt)
-		{
-			if (LocalViews.IsEmpty())
+			for (TConstSetBitIterator<SceneRenderingBitArrayAllocator> BitIt(TaskData.ShadowsToGatherInSerialPass); BitIt; ++BitIt)
 			{
-				LocalViews.AddZeroed(1);
-			}
+				if (LocalViews.IsEmpty())
+				{
+					LocalViews.AddZeroed(1);
+				}
 
-			FShadowMeshCollector& MeshCollector = TaskData.GetSerialMeshCollector();
-			FProjectedShadowInfo* ProjectedShadowInfo = TaskData.ShadowsToGather[BitIt.GetIndex()];
-			FVisibleLightInfo& VisibleLightInfo = VisibleLightInfos[ProjectedShadowInfo->GetLightSceneInfo().Id];
-			ProjectedShadowInfo->GatherDynamicMeshElements(MeshCollector.GetCollector(), *this, VisibleLightInfo, LocalViews, SerialPass);
+				FShadowMeshCollector& MeshCollector = TaskData.GetSerialMeshCollector();
+				FProjectedShadowInfo* ProjectedShadowInfo = TaskData.ShadowsToGather[BitIt.GetIndex()];
+				FVisibleLightInfo& VisibleLightInfo = VisibleLightInfos[ProjectedShadowInfo->GetLightSceneInfo().Id];
+				ProjectedShadowInfo->GatherDynamicMeshElements(MeshCollector.GetCollector(), *this, VisibleLightInfo, LocalViews, SerialPass);
 
-			// Setup mesh draw command passes now that all primitives are processed.
-			ProjectedShadowInfo->SetupMeshDrawCommandsForProjectionStenciling(*this, TaskData.InstanceCullingManager);
-			ProjectedShadowInfo->SetupMeshDrawCommandsForShadowDepth(*this, TaskData.InstanceCullingManager);
-		}
+				// Setup mesh draw command passes now that all primitives are processed.
+				ProjectedShadowInfo->SetupMeshDrawCommandsForProjectionStenciling(*this, TaskData.InstanceCullingManager);
+				ProjectedShadowInfo->SetupMeshDrawCommandsForShadowDepth(*this, TaskData.InstanceCullingManager);
+			}
 
-		for (FShadowMeshCollector* MeshCollector : TaskData.MeshCollectors)
-		{
-			MeshCollector->Finish();
-		}
+			for (FShadowMeshCollector* MeshCollector : TaskData.MeshCollectors)
+			{
+				MeshCollector->Finish();
+			}
 	
-	}, MeshCollectorsTask, TaskPriority, ExtendedTaskPriority));
+		}, MeshCollectorsTask, TaskPriority, ExtendedTaskPriority));
+	}
 
 	TaskData.SetupMeshPassTask.Trigger();
 }
