SampledDirectLighting - fixed temporal accumulation

diff --git a/Engine/Shaders/Private/SampledDirectLighting/SampledDirectLightingTemporal.usf b/Engine/Shaders/Private/SampledDirectLighting/SampledDirectLightingTemporal.usf
--- a/Engine/Shaders/Private/SampledDirectLighting/SampledDirectLightingTemporal.usf
+++ b/Engine/Shaders/Private/SampledDirectLighting/SampledDirectLightingTemporal.usf
@@ -215,9 +215,10 @@ void SampledDirectLightingTemporalAccumulationCS(
 				Print(Context, TEXT("Irradiance: "));
 				Print(Context, Irradiance);
 				Newline(Context);
-				Print(Context, TEXT("Variance: "));
-				Newline(Context);
+				Print(Context, TEXT("DiffuseVariance:  "));
 				Print(Context, DiffuseVariance);
+				Newline(Context);
+				Print(Context, TEXT("SpecularVariance: "));
 				Print(Context, SpecularVariance);
 			}
 			#endif
diff --git a/Engine/Source/Runtime/Renderer/Private/SampledDirectLighting/SampledDirectLighting.cpp b/Engine/Source/Runtime/Renderer/Private/SampledDirectLighting/SampledDirectLighting.cpp
--- a/Engine/Source/Runtime/Renderer/Private/SampledDirectLighting/SampledDirectLighting.cpp
+++ b/Engine/Source/Runtime/Renderer/Private/SampledDirectLighting/SampledDirectLighting.cpp
@@ -1026,11 +1026,11 @@ void FDeferredShadingSceneRenderer::RenderSampledDirectLighting(FRDGBuilder& Gra
 			(View.ViewRect.Max.X - 1.0f) * InvBufferSize.X,
 			(View.ViewRect.Max.Y - 1.0f) * InvBufferSize.Y);
 
-		if (DiffuseLighting && SpecularLighting && LuminanceMomentsHistory && SceneDepthCopy && NumFramesAccumulated && bTemporal)
+		if (DiffuseLighting && SpecularLighting && LuminanceMoments && SceneDepthCopy && NumFramesAccumulated && bTemporal)
 		{
 			GraphBuilder.QueueTextureExtraction(DiffuseLighting, &LightingViewState.DiffuseLightingHistory);
 			GraphBuilder.QueueTextureExtraction(SpecularLighting, &LightingViewState.SpecularLightingHistory);
-			GraphBuilder.QueueTextureExtraction(LuminanceMomentsHistory, &LightingViewState.LuminanceMomentsHistory);
+			GraphBuilder.QueueTextureExtraction(LuminanceMoments, &LightingViewState.LuminanceMomentsHistory);
 			GraphBuilder.QueueTextureExtraction(SceneDepthCopy, &LightingViewState.SceneDepthHistory);
 			GraphBuilder.QueueTextureExtraction(NumFramesAccumulated, &LightingViewState.NumFramesAccumulatedHistory);
 		}
