fix PixelPreProcessors for screenshot not checking pixel format correctly

#rnx

diff --git a/Engine/Source/Runtime/MediaUtils/Private/MediaRecorder.cpp b/Engine/Source/Runtime/MediaUtils/Private/MediaRecorder.cpp
--- a/Engine/Source/Runtime/MediaUtils/Private/MediaRecorder.cpp
+++ b/Engine/Source/Runtime/MediaUtils/Private/MediaRecorder.cpp
@@ -274,11 +274,15 @@ void FMediaRecorder::TickRecording()
 				{
 					check(bUseFMediaImagePixelData);
 
-					if (Sample->GetFormat() == EMediaTextureSampleFormat::FloatRGBA)
+					if ( ImageTask->PixelData->GetType() == EImagePixelType::Float16 )
 					{
 						ImageTask->PixelPreProcessors.Add(TAsyncAlphaWrite<FFloat16Color>(1.f));
 					}
-					else if (Sample->GetFormat() == EMediaTextureSampleFormat::CharBGRA)
+					else if ( ImageTask->PixelData->GetType() == EImagePixelType::Float32 )
+					{
+						ImageTask->PixelPreProcessors.Add(TAsyncAlphaWrite<FLinearColor>(1.f));
+					}
+					else if ( ImageTask->PixelData->GetType() == EImagePixelType::Color )
 					{
 						ImageTask->PixelPreProcessors.Add(TAsyncAlphaWrite<FColor>(255));
 					}
diff --git a/Engine/Source/Runtime/Renderer/Private/PostProcess/PostProcessVisualizeBuffer.cpp b/Engine/Source/Runtime/Renderer/Private/PostProcess/PostProcessVisualizeBuffer.cpp
--- a/Engine/Source/Runtime/Renderer/Private/PostProcess/PostProcessVisualizeBuffer.cpp
+++ b/Engine/Source/Runtime/Renderer/Private/PostProcess/PostProcessVisualizeBuffer.cpp
@@ -267,7 +267,19 @@ void AddDumpToFilePass(FRDGBuilder& GraphBuilder, FScreenPassTexture Input, cons
 		else if(ImageTask->Format == EImageFormat::PNG)
 		{
 			// PNGs can't have 0 alpha or RGB data is destroyed.
-			ImageTask->PixelPreProcessors.Add(TAsyncAlphaWrite<FLinearColor>(1.0f));
+
+			if ( ImageTask->PixelData->GetType() == EImagePixelType::Float32 )
+			{
+				ImageTask->PixelPreProcessors.Add(TAsyncAlphaWrite<FLinearColor>(1.0f));
+			}
+			else if ( ImageTask->PixelData->GetType() == EImagePixelType::Float16 )
+			{
+				ImageTask->PixelPreProcessors.Add(TAsyncAlphaWrite<FFloat16Color>(1.0f));
+			}
+			else
+			{
+				check(false);
+			}
 		}
 
 		HighResScreenshotConfig.ImageWriteQueue->Enqueue(MoveTemp(ImageTask));
