Updated documentation for r.SkyLight.RealTimeReflectionCapture.TimeSlice.SkyCloudCubeFacePerFrame to be clearer.

diff --git a/Engine/Source/Runtime/Renderer/Private/ReflectionEnvironmentRealTimeCapture.cpp b/Engine/Source/Runtime/Renderer/Private/ReflectionEnvironmentRealTimeCapture.cpp
--- a/Engine/Source/Runtime/Renderer/Private/ReflectionEnvironmentRealTimeCapture.cpp
+++ b/Engine/Source/Runtime/Renderer/Private/ReflectionEnvironmentRealTimeCapture.cpp
@@ -32,12 +32,12 @@ DECLARE_GPU_STAT(CaptureConvolveSkyEnvMap);
 
 static TAutoConsoleVariable<int32> CVarRealTimeReflectionCaptureTimeSlicing(
 	TEXT("r.SkyLight.RealTimeReflectionCapture.TimeSlice"), 1,
-	TEXT("When enabled, the real-time sky light capture and convolutions will by distributed over several frames to lower the per-frame cost. Value in [1,6]."),
+	TEXT("When enabled, the real-time sky light capture and convolutions will by distributed over several frames to lower the per-frame cost."),
 	ECVF_RenderThreadSafe);
 
 static TAutoConsoleVariable<int32> CVarRealTimeReflectionCaptureTimeSlicingSkyCloudCubeFacePerFrame(
 	TEXT("r.SkyLight.RealTimeReflectionCapture.TimeSlice.SkyCloudCubeFacePerFrame"), 6,
-	TEXT("When enabled, the real-time sky light capture and convolutions will by distributed over several frames to lower the per-frame cost."),
+	TEXT("When enabled, the real-time sky light capture, when time sliced, will not render cloud in all cube face in a single frame; but one face per frame. That is to distribute the cloud tracing cost even more, but will add latency and potentially can result in lighting discrepancy between faces if the sun is moving fast. Value in [1,6]."),
 	ECVF_RenderThreadSafe);
 
 static TAutoConsoleVariable<int32> CVarRealTimeReflectionCaptureShadowFromOpaque(
