Moves sky light illuminance meter to a proper visualisation pass.

diff --git a/Engine/Source/Runtime/Engine/Private/SceneInterface.cpp b/Engine/Source/Runtime/Engine/Private/SceneInterface.cpp
--- a/Engine/Source/Runtime/Engine/Private/SceneInterface.cpp
+++ b/Engine/Source/Runtime/Engine/Private/SceneInterface.cpp
@@ -19,6 +19,7 @@ void FSceneInterface::UpdateAllPrimitiveSceneInfos(FRHICommandListImmediate& RHI
 
 void FSceneInterface::ProcessAndRenderIlluminanceMeter(FRDGBuilder& GraphBuilder, TArrayView<FViewInfo> Views, FRDGTextureRef SceneColorTexture)
 {
+	// Deprecated
 }
 
 EShaderPlatform FSceneInterface::GetShaderPlatform() const
diff --git a/Engine/Source/Runtime/Engine/Public/SceneInterface.h b/Engine/Source/Runtime/Engine/Public/SceneInterface.h
--- a/Engine/Source/Runtime/Engine/Public/SceneInterface.h
+++ b/Engine/Source/Runtime/Engine/Public/SceneInterface.h
@@ -264,6 +264,7 @@ public:
 	virtual void AllocateAndCaptureFrameSkyEnvMap(FRDGBuilder& GraphBuilder, FSceneRenderer& SceneRenderer, FViewInfo& MainView, bool bShouldRenderSkyAtmosphere, bool bShouldRenderVolumetricCloud, FInstanceCullingManager& InstanceCullingManager, FRDGExternalAccessQueue& ExternalAccessQueue) {}
 	virtual void ValidateSkyLightRealTimeCapture(FRDGBuilder& GraphBuilder, const FViewInfo& View, FRDGTextureRef SceneColorTexture) {}
 
+	UE_DEPRECATED(5.4, "This method has been refactored to be a proper visualization post process. It should have never been added on the FSceneInterface in the first place.")
 	ENGINE_API virtual void ProcessAndRenderIlluminanceMeter(FRDGBuilder& GraphBuilder, TArrayView<FViewInfo> Views, FRDGTextureRef SceneColorTexture);
 
 	virtual void AddPlanarReflection(class UPlanarReflectionComponent* Component) {}
diff --git a/Engine/Source/Runtime/Renderer/Private/DeferredShadingRenderer.cpp b/Engine/Source/Runtime/Renderer/Private/DeferredShadingRenderer.cpp
--- a/Engine/Source/Runtime/Renderer/Private/DeferredShadingRenderer.cpp
+++ b/Engine/Source/Runtime/Renderer/Private/DeferredShadingRenderer.cpp
@@ -3227,8 +3227,6 @@ void FDeferredShadingSceneRenderer::Render(FRDGBuilder& GraphBuilder)
 			}
 		}
 
-		Scene->ProcessAndRenderIlluminanceMeter(GraphBuilder, Views, SceneTextures.Color.Target);
-
 		if (!bHasRayTracedOverlay && ViewFamily.EngineShowFlags.LightShafts)
 		{
 			SCOPE_CYCLE_COUNTER(STAT_FDeferredShadingSceneRenderer_RenderLightShaftBloom);
diff --git a/Engine/Source/Runtime/Renderer/Private/IlluminanceMeter.cpp b/Engine/Source/Runtime/Renderer/Private/IlluminanceMeter.cpp
--- a/Engine/Source/Runtime/Renderer/Private/IlluminanceMeter.cpp
+++ b/Engine/Source/Runtime/Renderer/Private/IlluminanceMeter.cpp
@@ -4,6 +4,7 @@
 	Functionality for a real time illuminance meter reference for the sky light and sun.
 =============================================================================*/
 
+#include "IlluminanceMeter.h"
 #include "DataDrivenShaderPlatformInfo.h"
 #include "ReflectionEnvironmentCapture.h"
 #include "RHIResources.h"
@@ -12,6 +13,7 @@
 #include "SceneRendering.h"
 #include "ScenePrivate.h"
 #include "ShaderCompiler.h"
+#include "PostProcess/PostProcessing.h"
 
 DECLARE_GPU_STAT(IlluminanceMeter);
 
@@ -98,11 +100,17 @@ IMPLEMENT_GLOBAL_SHADER(FPrintIlluminanceMeterPS, "/Engine/Private/ReflectionEnv
 
 void FScene::ProcessAndRenderIlluminanceMeter(FRDGBuilder& GraphBuilder, TArrayView<FViewInfo> Views, FRDGTextureRef SceneColorTexture)
 {
-	FViewInfo& MainView = Views[0];
-	if (!SkyLight || !MainView.Family->EngineShowFlags.VisualizeSkyLightIlluminance || MainView.bIsSceneCapture || MainView.bIsReflectionCapture || MainView.bIsPlanarReflection ||
-		!CanIlluminanceMeterDisplayOnPlatform(MainView.GetShaderPlatform()))
+	// Deprecated
+}
+
+FScreenPassTexture ProcessAndRenderIlluminanceMeter(FRDGBuilder& GraphBuilder, const FViewInfo& View, FScreenPassTexture& ScreenPassSceneColor)
+{
+	const FScene* Scene = (const FScene*)View.Family->Scene;
+	FSkyLightSceneProxy* SkyLight = Scene ? Scene->SkyLight : nullptr;
+	if (!SkyLight || !View.Family->EngineShowFlags.VisualizeSkyLightIlluminance || View.bIsSceneCapture || View.bIsReflectionCapture || View.bIsPlanarReflection ||
+		!CanIlluminanceMeterDisplayOnPlatform(View.GetShaderPlatform()))
 	{
-		return;
+		return MoveTemp(ScreenPassSceneColor);
 	}
 
 	RDG_EVENT_SCOPE(GraphBuilder, "IlluminanceMeter");
@@ -111,10 +119,10 @@ void FScene::ProcessAndRenderIlluminanceMeter(FRDGBuilder& GraphBuilder, TArrayV
 	// 0- Get source cubemap and allocate cube map transient resources
 
 	FTextureRHIRef SkyLightTextureRHI = GBlackTextureCube->TextureRHI;
-	if (SkyLight->bRealTimeCaptureEnabled && ConvolvedSkyRenderTargetReadyIndex >= 0)
+	if (SkyLight->bRealTimeCaptureEnabled && Scene->ConvolvedSkyRenderTargetReadyIndex >= 0)
 	{
 		// Cannot blend with this capture mode as of today.
-		SkyLightTextureRHI = ConvolvedSkyRenderTarget[ConvolvedSkyRenderTargetReadyIndex]->GetRHI();
+		SkyLightTextureRHI = Scene->ConvolvedSkyRenderTarget[Scene->ConvolvedSkyRenderTargetReadyIndex]->GetRHI();
 
 		//
 	}
@@ -133,12 +141,12 @@ void FScene::ProcessAndRenderIlluminanceMeter(FRDGBuilder& GraphBuilder, TArrayV
 	{
 		RDG_EVENT_SCOPE(GraphBuilder, "CopyCubemapToIlluminanceMeterCubemap");
 
-		TShaderMapRef<FCopyCubemapToIlluminanceMeterCubemapPS> PixelShader(MainView.ShaderMap);
+		TShaderMapRef<FCopyCubemapToIlluminanceMeterCubemapPS> PixelShader(View.ShaderMap);
 
 		for (uint32 CubeFace = 0; CubeFace < CubeFace_MAX; CubeFace++)
 		{
 			auto* PassParameters = GraphBuilder.AllocParameters<FCopyCubemapToIlluminanceMeterCubemapPS::FParameters>();
-			PassParameters->ViewUniformBuffer = MainView.ViewUniformBuffer;
+			PassParameters->ViewUniformBuffer = View.ViewUniformBuffer;
 			PassParameters->SourceCubemapSampler = TStaticSamplerState<SF_Point, AM_Clamp, AM_Clamp, AM_Clamp>::GetRHI();
 			PassParameters->SourceCubemapTexture = SkyLightTextureRHI;
 			PassParameters->CubeFace = CubeFace;
@@ -149,7 +157,7 @@ void FScene::ProcessAndRenderIlluminanceMeter(FRDGBuilder& GraphBuilder, TArrayV
 
 			FPixelShaderUtils::AddFullscreenPass(
 				GraphBuilder,
-				MainView.ShaderMap,
+				View.ShaderMap,
 				RDG_EVENT_NAME("FCopyCubemapToIlluminanceMeterCubemap"),
 				PixelShader,
 				PassParameters,
@@ -161,7 +169,7 @@ void FScene::ProcessAndRenderIlluminanceMeter(FRDGBuilder& GraphBuilder, TArrayV
 	{
 		RDG_EVENT_SCOPE(GraphBuilder, "IntegrateCubemapIlluminance");
 
-		TShaderMapRef<FDownsampleIntegrateCubemapIlluminancePS> PixelShader(MainView.ShaderMap);
+		TShaderMapRef<FDownsampleIntegrateCubemapIlluminancePS> PixelShader(View.ShaderMap);
 		const int32 NumMips = IlluminanceMeterCubeTexture->Desc.NumMips;
 
 		for (int32 MipIndex = 1; MipIndex < NumMips; MipIndex++)
@@ -182,7 +190,7 @@ void FScene::ProcessAndRenderIlluminanceMeter(FRDGBuilder& GraphBuilder, TArrayV
 
 				FPixelShaderUtils::AddFullscreenPass(
 					GraphBuilder,
-					MainView.ShaderMap,
+					View.ShaderMap,
 					RDG_EVENT_NAME("CreateCubeMips (Mip: %d, Face: %d)", MipIndex, CubeFace),
 					PixelShader,
 					PassParameters,
@@ -200,19 +208,21 @@ void FScene::ProcessAndRenderIlluminanceMeter(FRDGBuilder& GraphBuilder, TArrayV
 		FRHIBlendState* PreMultipliedColorTransmittanceBlend = TStaticBlendState<CW_RGB, BO_Add, BF_One, BF_SourceAlpha, BO_Add, BF_Zero, BF_One>::GetRHI();
 
 		FPrintIlluminanceMeterPS::FParameters* PassParameters = GraphBuilder.AllocParameters<FPrintIlluminanceMeterPS::FParameters>();
-		PassParameters->ViewUniformBuffer = MainView.ViewUniformBuffer;
+		PassParameters->ViewUniformBuffer = View.ViewUniformBuffer;
 		PassParameters->SourceCubemapTexture = IlluminanceMeterCubeTexture;
 		PassParameters->SourceCubemapSampler = TStaticSamplerState<SF_Trilinear, AM_Clamp, AM_Clamp, AM_Clamp>::GetRHI();
 		PassParameters->SkyLightCubeTexture = SkyLightTextureRHI;
 		PassParameters->SkyLightCubeSampler = TStaticSamplerState<SF_Trilinear, AM_Clamp, AM_Clamp, AM_Clamp>::GetRHI();
 		PassParameters->CubeMipIndexToSampleIlluminance = IlluminanceMeterCubeTexture->Desc.NumMips - 1;
 		PassParameters->SkyLightCaptureWorlPos = FVector3f(SkyLight->CapturePosition);
-		ShaderPrint::SetParameters(GraphBuilder, MainView.ShaderPrintData, PassParameters->ShaderPrintParameters);
-		PassParameters->RenderTargets[0] = FRenderTargetBinding(SceneColorTexture, ERenderTargetLoadAction::ELoad);
+		ShaderPrint::SetParameters(GraphBuilder, View.ShaderPrintData, PassParameters->ShaderPrintParameters);
+		PassParameters->RenderTargets[0] = FRenderTargetBinding(ScreenPassSceneColor.Texture, ERenderTargetLoadAction::ELoad);
 
 		FPrintIlluminanceMeterPS::FPermutationDomain PermutationVector;
-		TShaderMapRef<FPrintIlluminanceMeterPS> PixelShader(MainView.ShaderMap, PermutationVector);
+		TShaderMapRef<FPrintIlluminanceMeterPS> PixelShader(View.ShaderMap, PermutationVector);
 
-		FPixelShaderUtils::AddFullscreenPass<FPrintIlluminanceMeterPS>(GraphBuilder, MainView.ShaderMap, RDG_EVENT_NAME("Substrate::VisualizeMaterial(Draw)"), PixelShader, PassParameters, MainView.ViewRect, PreMultipliedColorTransmittanceBlend);
+		FPixelShaderUtils::AddFullscreenPass<FPrintIlluminanceMeterPS>(GraphBuilder, View.ShaderMap, RDG_EVENT_NAME("Substrate::VisualizeMaterial(Draw)"), PixelShader, PassParameters, View.ViewRect, PreMultipliedColorTransmittanceBlend);
 	}
+
+	return MoveTemp(ScreenPassSceneColor);
 }
diff --git a/Engine/Source/Runtime/Renderer/Private/IlluminanceMeter.h b/Engine/Source/Runtime/Renderer/Private/IlluminanceMeter.h
new file mode 100644
--- /dev/null
+++ b/Engine/Source/Runtime/Renderer/Private/IlluminanceMeter.h
@@ -0,0 +1,10 @@
+// Copyright Epic Games, Inc. All Rights Reserved.
+
+#include "RenderGraphFwd.h"
+
+class FViewInfo;
+class FSkyLightSceneProxy;
+class FRDGBuilder;
+struct FScreenPassTexture;
+
+FScreenPassTexture ProcessAndRenderIlluminanceMeter(FRDGBuilder& GraphBuilder, const FViewInfo& View, FScreenPassTexture& ScreenPassSceneColor);
diff --git a/Engine/Source/Runtime/Renderer/Private/PostProcess/PostProcessing.cpp b/Engine/Source/Runtime/Renderer/Private/PostProcess/PostProcessing.cpp
--- a/Engine/Source/Runtime/Renderer/Private/PostProcess/PostProcessing.cpp
+++ b/Engine/Source/Runtime/Renderer/Private/PostProcess/PostProcessing.cpp
@@ -70,6 +70,7 @@
 #include "RectLightTextureManager.h"
 #include "IESTextureManager.h"
 #include "UnrealEngine.h"
+#include "IlluminanceMeter.h"
 
 bool IsMobileEyeAdaptationEnabled(const FViewInfo& View);
 
@@ -375,6 +376,7 @@ void AddPostProcessingPasses(
 		VisualizeSubstrate,
 		VisualizeLightGrid,
 		VisualizeSkyAtmosphere,
+		VisualizeSkyLightIlluminanceMeter,
 		VisualizeLightFunctionAtlas,
 		VisualizeLevelInstance,
 		SelectionOutline,
@@ -428,6 +430,7 @@ void AddPostProcessingPasses(
 		TEXT("VisualizeSubstrate"),
 		TEXT("VisualizeLightGrid"),
 		TEXT("VisualizeSkyAtmosphere"),
+		TEXT("VisualizeSkyLightIlluminanceMeter"),
 		TEXT("VisualizeLightFunctionAtlas"),
 		TEXT("VisualizeLevelInstance"),
 		TEXT("SelectionOutline"),
@@ -468,12 +471,14 @@ void AddPostProcessingPasses(
 
 #if WITH_EDITOR
 	PassSequence.SetEnabled(EPass::VisualizeSkyAtmosphere, Scene&& View.Family && View.Family->EngineShowFlags.VisualizeSkyAtmosphere && ShouldRenderSkyAtmosphereDebugPasses(Scene, View.Family->EngineShowFlags));
+	PassSequence.SetEnabled(EPass::VisualizeSkyLightIlluminanceMeter, Scene&& Scene->SkyLight && View.Family && View.Family->EngineShowFlags.VisualizeSkyLightIlluminance);
 	PassSequence.SetEnabled(EPass::VisualizeLightFunctionAtlas, Scene && Scene->LightFunctionAtlasSceneData.GetLightFunctionAtlasEnabled() && View.Family && View.Family->EngineShowFlags.VisualizeLightFunctionAtlas);
 	PassSequence.SetEnabled(EPass::VisualizeLevelInstance, GIsEditor && EngineShowFlags.EditingLevelInstance && EngineShowFlags.VisualizeLevelInstanceEditing && !bVisualizeHDR);
 	PassSequence.SetEnabled(EPass::SelectionOutline, GIsEditor && EngineShowFlags.Selection && EngineShowFlags.SelectionOutline && !EngineShowFlags.Wireframe && !bVisualizeHDR && !IStereoRendering::IsStereoEyeView(View));
 	PassSequence.SetEnabled(EPass::EditorPrimitive, FSceneRenderer::ShouldCompositeEditorPrimitives(View));
 #else
 	PassSequence.SetEnabled(EPass::VisualizeSkyAtmosphere, false);
+	PassSequence.SetEnabled(EPass::VisualizeSkyLightIlluminanceMeter, false);
 	PassSequence.SetEnabled(EPass::VisualizeLightFunctionAtlas, false);
 	PassSequence.SetEnabled(EPass::VisualizeLevelInstance, false);
 	PassSequence.SetEnabled(EPass::SelectionOutline, false);
@@ -1397,6 +1402,13 @@ void AddPostProcessingPasses(
 		SceneColor = AddSkyAtmosphereDebugPasses(GraphBuilder, Scene, *View.Family, View, SceneColor);
 	}
 
+	if (PassSequence.IsEnabled(EPass::VisualizeSkyLightIlluminanceMeter))
+	{
+		FScreenPassRenderTarget OverrideOutput;
+		PassSequence.AcceptOverrideIfLastPass(EPass::VisualizeSkyLightIlluminanceMeter, OverrideOutput);
+		SceneColor = ProcessAndRenderIlluminanceMeter(GraphBuilder, View, SceneColor);
+	}
+
 	if (PassSequence.IsEnabled(EPass::VisualizeLightFunctionAtlas))
 	{
 		FScreenPassRenderTarget OverrideOutput;
