Store IESAtlasIndex in GPUScene Lights.

diff --git a/Engine/Shaders/Private/LightData.ush b/Engine/Shaders/Private/LightData.ush
--- a/Engine/Shaders/Private/LightData.ush
+++ b/Engine/Shaders/Private/LightData.ush
@@ -481,8 +481,8 @@ FDeferredLightData ConvertToDeferredLight(FLightSceneData In)
 	Out.RectLightData.AtlasData.AtlasUVOffset = In.RectLightAtlasUVOffset;
 	Out.RectLightData.AtlasData.AtlasUVScale = In.RectLightAtlasUVScale;
 
-	Out.IESAtlasIndex = 0;
-	Out.LightFunctionAtlasLightIndex = 0;
+	Out.IESAtlasIndex					= In.IESAtlasIndex;
+	Out.LightFunctionAtlasLightIndex	= 0;
 
 	// Shadow
 	// * Static shadowing uses ShadowMapChannel, 
diff --git a/Engine/Shaders/Shared/LightSceneData.h b/Engine/Shaders/Shared/LightSceneData.h
--- a/Engine/Shaders/Shared/LightSceneData.h
+++ b/Engine/Shaders/Shared/LightSceneData.h
@@ -65,10 +65,15 @@ struct FLightSceneData
 
 	float InverseExposureBlend;
 
+	// could pack IESAtlasIndex with other data in the future since it doesn't require 32 bits
+	// FLocalLightData packs it in 16 bits (see UnpackLigthIESAtlasIndex(...))
+	// could probably go down to 8 bits (with some logic in GIESTextureManager to warn about overflow)
+	float IESAtlasIndex;
+
 	// Extra fields
 	uint LightTypeAndShadowMapChannelMaskPacked;
 
-	float3 Padding;
+	float2 Padding;
 };
 
 #ifdef __cplusplus
diff --git a/Engine/Source/Runtime/Renderer/Private/GPUScene.cpp b/Engine/Source/Runtime/Renderer/Private/GPUScene.cpp
--- a/Engine/Source/Runtime/Renderer/Private/GPUScene.cpp
+++ b/Engine/Source/Runtime/Renderer/Private/GPUScene.cpp
@@ -745,6 +745,7 @@ void FGPUScene::InitLightData(const FLightSceneInfoCompact& LightInfoCompact, bo
 	DataOut.RectLightAtlasUVScale = LightParams.RectLightAtlasUVScale;
 	DataOut.RectLightAtlasMaxLevel = LightParams.RectLightAtlasMaxLevel;
 	DataOut.InverseExposureBlend = LightParams.InverseExposureBlend;
+	DataOut.IESAtlasIndex = LightParams.IESAtlasIndex;
 	DataOut.LightTypeAndShadowMapChannelMaskPacked = LightInfo.PackLightTypeAndShadowMapChannelMask(bAllowStaticLighting);
 }
 
