Fix crash when the scene is a FNULLSceneInterface.

- Don't try and cast to a `FScene` and simply just call the base class version of UpdateEarlyZPassMode, which does nothing.

todo: audit the rest of the codebase for similar improper usage and casting of `FScene`

diff --git a/Engine/Source/Runtime/Renderer/Private/RendererScene.cpp b/Engine/Source/Runtime/Renderer/Private/RendererScene.cpp
--- a/Engine/Source/Runtime/Renderer/Private/RendererScene.cpp
+++ b/Engine/Source/Runtime/Renderer/Private/RendererScene.cpp
@@ -1375,8 +1375,7 @@ static void UpdateEarlyZPassModeCVarSinkFunction()
 			UWorld* World = *It;
 			if (World && World->Scene)
 			{
-				FScene* Scene = (FScene*)(World->Scene);
-				Scene->UpdateEarlyZPassMode();
+				World->Scene->UpdateEarlyZPassMode();
 			}
 		}
 
