Remove dead code that was casting to a `FScene`.

diff --git a/Engine/Source/Runtime/Renderer/Private/CompositionLighting/CompositionLighting.cpp b/Engine/Source/Runtime/Renderer/Private/CompositionLighting/CompositionLighting.cpp
--- a/Engine/Source/Runtime/Renderer/Private/CompositionLighting/CompositionLighting.cpp
+++ b/Engine/Source/Runtime/Renderer/Private/CompositionLighting/CompositionLighting.cpp
@@ -60,8 +60,6 @@ static TAutoConsoleVariable<int32> CVarGTAOCombined(
 
 bool IsAmbientCubemapPassRequired(const FSceneView& View)
 {
-	FScene* Scene = (FScene*)View.Family->Scene;
-
 	return View.FinalPostProcessSettings.ContributingCubemaps.Num() != 0 && IsUsingGBuffers(View.GetShaderPlatform());
 }
 
diff --git a/Engine/Source/Runtime/Renderer/Private/PostProcess/PostProcessMotionBlur.cpp b/Engine/Source/Runtime/Renderer/Private/PostProcess/PostProcessMotionBlur.cpp
--- a/Engine/Source/Runtime/Renderer/Private/PostProcess/PostProcessMotionBlur.cpp
+++ b/Engine/Source/Runtime/Renderer/Private/PostProcess/PostProcessMotionBlur.cpp
@@ -1166,7 +1166,6 @@ FScreenPassTextureSlice AddVisualizeMotionBlurPass(FRDGBuilder& GraphBuilder, co
 		Canvas.DrawShadowedString(X, Y += YStep, TEXT("Paused, r.VelocityTest, Parallel:"), GetStatsFont(), FLinearColor(1, 1, 0));
 		Canvas.DrawShadowedString(X + ColumnWidth, Y, *Line, GetStatsFont(), FLinearColor(1, 1, 0));
 
-		const FScene* Scene = (const FScene*)View.Family->Scene;
 		const FSceneViewState *SceneViewState = (const FSceneViewState*)View.State;
 
 		Line = FString::Printf(TEXT("View=%.4x PrevView=%.4x"),
