Sampled Direct Lighting
- support IES profiles.

diff --git a/Engine/Shaders/Private/SampledDirectLighting/SampledDirectLighting.ush b/Engine/Shaders/Private/SampledDirectLighting/SampledDirectLighting.ush
--- a/Engine/Shaders/Private/SampledDirectLighting/SampledDirectLighting.ush
+++ b/Engine/Shaders/Private/SampledDirectLighting/SampledDirectLighting.ush
@@ -9,7 +9,6 @@
 #include "../LightGridCommon.ush"
 
 #define SUPPORT_CONTACT_SHADOWS 0
-#define USE_IES_PROFILE 0
 #define NON_DIRECTIONAL_DIRECT_LIGHTING 0
 #include "../DeferredLightingCommon.ush"
 #include "../SceneData.ush"
diff --git a/Engine/Shaders/Private/SampledDirectLighting/SampledDirectLightingSampling.usf b/Engine/Shaders/Private/SampledDirectLighting/SampledDirectLightingSampling.usf
--- a/Engine/Shaders/Private/SampledDirectLighting/SampledDirectLightingSampling.usf
+++ b/Engine/Shaders/Private/SampledDirectLighting/SampledDirectLightingSampling.usf
@@ -68,6 +68,11 @@ float GetLocalLightTargetPDF(FDeferredLightData LightData, float3 TranslatedWorl
 				SurfaceShadow);
 
 			SampleWeight = Luminance(LightRadiance.xyz * View.PreExposure);
+
+			if (SampleWeight > 0.01f && LightData.IESAtlasIndex >= 0)
+			{
+				SampleWeight *= ComputeLightProfileMultiplier(TranslatedWorldPosition, LightData.TranslatedWorldPosition, -LightData.Direction, LightData.Tangent, LightData.IESAtlasIndex);
+			}
 		}
 		#endif
 	}
diff --git a/Engine/Shaders/Private/SampledDirectLighting/SampledDirectLightingShading.usf b/Engine/Shaders/Private/SampledDirectLighting/SampledDirectLightingShading.usf
--- a/Engine/Shaders/Private/SampledDirectLighting/SampledDirectLightingShading.usf
+++ b/Engine/Shaders/Private/SampledDirectLighting/SampledDirectLightingShading.usf
@@ -111,6 +111,12 @@ void ShadeLightSamplesCS(
 						SampleWeight = ShadingTileAtlas[ShadingTileAtlasCoord];
 					}
 
+					if (SampleWeight > 0.01f && LightData.IESAtlasIndex >= 0)
+					{
+						// #sdl_todo: evaluate moving this to CompositeLightSamplesCS
+						SampleWeight *= ComputeLightProfileMultiplier(TranslatedWorldPosition, LightData.TranslatedWorldPosition, -LightData.Direction, LightData.Tangent, LightData.IESAtlasIndex);
+					}
+
 					#if DEBUG_MODE
 					if (bDebug)
 					{
diff --git a/Engine/Source/Runtime/Renderer/Private/SampledDirectLighting/SampledDirectLighting.cpp b/Engine/Source/Runtime/Renderer/Private/SampledDirectLighting/SampledDirectLighting.cpp
--- a/Engine/Source/Runtime/Renderer/Private/SampledDirectLighting/SampledDirectLighting.cpp
+++ b/Engine/Source/Runtime/Renderer/Private/SampledDirectLighting/SampledDirectLighting.cpp
@@ -128,6 +128,13 @@ static TAutoConsoleVariable<int> CVarSampledDirectLightingTexturedRectLights(
 	ECVF_Scalability | ECVF_RenderThreadSafe
 );
 
+static TAutoConsoleVariable<int> CVarSampledDirectLightingIESProfiles(
+	TEXT("r.SampledDirectLighting.IESProfiles"),
+	1,
+	TEXT("Whether to support IES profiles on lights."),
+	ECVF_Scalability | ECVF_RenderThreadSafe
+);
+
 namespace SampledDirectLighting
 {
 	// must match values in SampledDirectLighting.ush
@@ -323,11 +330,12 @@ class FGenerateSamplesCS : public FGlobalShader
 	END_SHADER_PARAMETER_STRUCT()
 
 	class FTileType : SHADER_PERMUTATION_INT("TILE_TYPE", (int32)SampledDirectLighting::ETileType::MAX);
+	class FIESProfile : SHADER_PERMUTATION_BOOL("USE_IES_PROFILE");
 	class FTexturedRectLights : SHADER_PERMUTATION_BOOL("USE_SOURCE_TEXTURE");
 	class FNumSamplesPerPixel1d : SHADER_PERMUTATION_SPARSE_INT("NUM_SAMPLES_PER_PIXEL_1D", 1, 2, 4);
 	class FCandidateLightMask : SHADER_PERMUTATION_BOOL("CANDIDATE_LIGHT_MASK");
 	class FDebugMode : SHADER_PERMUTATION_BOOL("DEBUG_MODE");
-	using FPermutationDomain = TShaderPermutationDomain<FTileType, FTexturedRectLights, FNumSamplesPerPixel1d, FCandidateLightMask, FDebugMode>;
+	using FPermutationDomain = TShaderPermutationDomain<FTileType, FIESProfile, FTexturedRectLights, FNumSamplesPerPixel1d, FCandidateLightMask, FDebugMode>;
 
 	static int32 GetGroupSize()
 	{	
@@ -445,9 +453,10 @@ class FShadeLightSamplesCS : public FGlobalShader
 	}
 
 	class FTileType : SHADER_PERMUTATION_INT("TILE_TYPE", (int32)SampledDirectLighting::ETileType::MAX);
+	class FIESProfile : SHADER_PERMUTATION_BOOL("USE_IES_PROFILE");
 	class FTexturedRectLights : SHADER_PERMUTATION_BOOL("USE_SOURCE_TEXTURE");
 	class FDebugMode : SHADER_PERMUTATION_BOOL("DEBUG_MODE");
-	using FPermutationDomain = TShaderPermutationDomain<FTileType, FTexturedRectLights, FDebugMode>;
+	using FPermutationDomain = TShaderPermutationDomain<FTileType, FIESProfile, FTexturedRectLights, FDebugMode>;
 
 	static bool ShouldCompilePermutation(const FGlobalShaderPermutationParameters& Parameters)
 	{
@@ -794,6 +803,7 @@ void FDeferredShadingSceneRenderer::RenderSampledDirectLighting(FRDGBuilder& Gra
 
 			FGenerateSamplesCS::FPermutationDomain PermutationVector;
 			PermutationVector.Set<FGenerateSamplesCS::FTileType>(TileType);
+			PermutationVector.Set<FGenerateSamplesCS::FIESProfile>(CVarSampledDirectLightingIESProfiles.GetValueOnRenderThread() != 0);
 			PermutationVector.Set<FGenerateSamplesCS::FTexturedRectLights>(CVarSampledDirectLightingTexturedRectLights.GetValueOnRenderThread() != 0);
 			PermutationVector.Set<FGenerateSamplesCS::FNumSamplesPerPixel1d>(NumSamplesPerPixel2d.X * NumSamplesPerPixel2d.Y);
 			PermutationVector.Set<FGenerateSamplesCS::FDebugMode>(bDebug);
@@ -940,6 +950,7 @@ void FDeferredShadingSceneRenderer::RenderSampledDirectLighting(FRDGBuilder& Gra
 
 			FShadeLightSamplesCS::FPermutationDomain PermutationVector;
 			PermutationVector.Set<FShadeLightSamplesCS::FTileType>(TileType);
+			PermutationVector.Set<FShadeLightSamplesCS::FIESProfile>(CVarSampledDirectLightingIESProfiles.GetValueOnRenderThread() != 0);
 			PermutationVector.Set<FShadeLightSamplesCS::FTexturedRectLights>(CVarSampledDirectLightingTexturedRectLights.GetValueOnRenderThread() != 0);
 			PermutationVector.Set<FShadeLightSamplesCS::FDebugMode>(bDebug);
 			auto ComputeShader = View.ShaderMap->GetShader<FShadeLightSamplesCS>(PermutationVector);
