Fixed volumetric fog flickering at the bottom of the screen. Temporal blending weight and uv clamp now are different to reproject the full screen while avoiding reprojecting bad data.

diff --git a/Engine/Shaders/Private/HeightFogCommon.ush b/Engine/Shaders/Private/HeightFogCommon.ush
--- a/Engine/Shaders/Private/HeightFogCommon.ush
+++ b/Engine/Shaders/Private/HeightFogCommon.ush
@@ -294,7 +294,7 @@ float3 ComputeHistoryVolumeUVFromNDC(float4 NDCPosition)
 {
 	NDCPosition.xy /= NDCPosition.w;
 	float3 VolumeUV = float3(NDCPosition.xy * float2(.5f, -.5f) + .5f, ComputeNormalizedZSliceFromDepth(NDCPosition.w));
-	return min(VolumeUV * float3(View.VolumetricFogViewGridUVToPrevViewRectUV.xy * View.VolumetricFogPrevViewGridRectUVToResourceUV.xy, 1.0), float3(View.VolumetricFogPrevUVMax, 1.0));
+	return min(VolumeUV * float3(View.VolumetricFogViewGridUVToPrevViewRectUV.xy * View.VolumetricFogPrevViewGridRectUVToResourceUV.xy, 1.0), float3(View.VolumetricFogPrevUVMaxForTemporalBlend, 1.0));
 }
 float3 ComputeHistoryVolumeUV_DEPRECATED(float3 WorldPosition, float4x4 WorldToClip)
 {
diff --git a/Engine/Shaders/Private/VolumetricFog.usf b/Engine/Shaders/Private/VolumetricFog.usf
--- a/Engine/Shaders/Private/VolumetricFog.usf
+++ b/Engine/Shaders/Private/VolumetricFog.usf
@@ -723,9 +723,9 @@ void LightScatteringCS(
 	float3 HistoryUV = ComputeHistoryVolumeUV_DEPRECATED(ComputeCellTranslatedWorldPosition(GridCoordinate, .5f), UnjitteredPrevTranslatedWorldToClip);
 	float HistoryAlpha = HistoryWeight;
 
-	// We need to test with View.VolumetricFogPrevUVMax for HistoryUV because that is what we clamp with in ComputeHistoryVolumeUV_DEPRECATED.
+	// We need to test with View.VolumetricFogPrevUVMaxForTemporalBlend for HistoryUV because that is what we clamp with in ComputeHistoryVolumeUV_DEPRECATED.
 	FLATTEN
-	if (any(HistoryUV < 0) || any(HistoryUV >= float3(View.VolumetricFogPrevUVMax, 1.0f)) || CellWasOccluded)
+	if (any(HistoryUV < 0) || any(HistoryUV >= float3(View.VolumetricFogPrevUVMaxForTemporalBlend, 1.0f)) || CellWasOccluded)
 	{
 		HistoryAlpha = 0;
 	}
@@ -926,7 +926,7 @@ void LightScatteringCS(
 	BRANCH
 	if (HistoryAlpha > 0)
 	{
-		float4 PreExposedHistoryScatteringAndExtinction = Texture3DSampleLevel(LightScatteringHistory, LightScatteringHistorySampler, HistoryUV, 0);
+		float4 PreExposedHistoryScatteringAndExtinction = Texture3DSampleLevel(LightScatteringHistory, LightScatteringHistorySampler, min(HistoryUV, float3(View.VolumetricFogPrevUVMax, 1.0f)), 0);
 		PreExposedHistoryScatteringAndExtinction.rgb *= LightScatteringHistoryPreExposureAndInv.y * View.PreExposure; // Bring previous frame pre exposed history luminance as current frame pre exposed luminance. Leave extinction untouched! 
 		PreExposedScatteringAndExtinction = lerp(PreExposedScatteringAndExtinction, PreExposedHistoryScatteringAndExtinction, HistoryAlpha);
 	}
diff --git a/Engine/Source/Runtime/Engine/Public/SceneView.h b/Engine/Source/Runtime/Engine/Public/SceneView.h
--- a/Engine/Source/Runtime/Engine/Public/SceneView.h
+++ b/Engine/Source/Runtime/Engine/Public/SceneView.h
@@ -894,6 +894,7 @@ enum ETranslucencyVolumeCascade
 	VIEW_UNIFORM_BUFFER_MEMBER(FVector2f, VolumetricFogViewGridUVToPrevViewRectUV) \
 	VIEW_UNIFORM_BUFFER_MEMBER(FVector2f, VolumetricFogPrevViewGridRectUVToResourceUV) \
 	VIEW_UNIFORM_BUFFER_MEMBER(FVector2f, VolumetricFogPrevUVMax) \
+	VIEW_UNIFORM_BUFFER_MEMBER(FVector2f, VolumetricFogPrevUVMaxForTemporalBlend) \
 	VIEW_UNIFORM_BUFFER_MEMBER(FVector2f, VolumetricFogScreenToResourceUV) \
 	VIEW_UNIFORM_BUFFER_MEMBER(FVector2f, VolumetricFogUVMax) \
 	VIEW_UNIFORM_BUFFER_MEMBER(float, VolumetricFogMaxDistance) \
diff --git a/Engine/Source/Runtime/Renderer/Private/ScenePrivate.h b/Engine/Source/Runtime/Renderer/Private/ScenePrivate.h
--- a/Engine/Source/Runtime/Renderer/Private/ScenePrivate.h
+++ b/Engine/Source/Runtime/Renderer/Private/ScenePrivate.h
@@ -1111,6 +1111,7 @@ public:
 	FVector2f PrevLightScatteringViewGridUVToViewRectVolumeUV;
 	FVector2f VolumetricFogPrevViewGridRectUVToResourceUV;
 	FVector2f VolumetricFogPrevUVMax;
+	FVector2f VolumetricFogPrevUVMaxForTemporalBlend;
 	TRefCountPtr<IPooledRenderTarget> LightScatteringHistory;
 	TRefCountPtr<IPooledRenderTarget> PrevLightScatteringConservativeDepthTexture;
 
diff --git a/Engine/Source/Runtime/Renderer/Private/SceneRendering.cpp b/Engine/Source/Runtime/Renderer/Private/SceneRendering.cpp
--- a/Engine/Source/Runtime/Renderer/Private/SceneRendering.cpp
+++ b/Engine/Source/Runtime/Renderer/Private/SceneRendering.cpp
@@ -1483,6 +1483,7 @@ void FViewInfo::SetupUniformBufferParameters(
 	ViewUniformShaderParameters.VolumetricFogViewGridUVToPrevViewRectUV = FVector2f::One();
 	ViewUniformShaderParameters.VolumetricFogPrevViewGridRectUVToResourceUV = FVector2f::One();
 	ViewUniformShaderParameters.VolumetricFogPrevUVMax = FVector2f::One();
+	ViewUniformShaderParameters.VolumetricFogPrevUVMaxForTemporalBlend = FVector2f::One();
 	int32 VolumetricFogViewGridPixelSize;
 	int32 VolumetricFogResourceGridPixelSize;
 	const FIntVector VolumetricFogResourceGridSize = GetVolumetricFogResourceGridSize(*this, VolumetricFogViewGridPixelSize);
@@ -1499,9 +1500,10 @@ void FViewInfo::SetupUniformBufferParameters(
 
 		ViewUniformShaderParameters.VolumetricFogPrevViewGridRectUVToResourceUV = ViewState->VolumetricFogPrevViewGridRectUVToResourceUV;
 		ViewUniformShaderParameters.VolumetricFogPrevUVMax = ViewState->VolumetricFogPrevUVMax;
+		ViewUniformShaderParameters.VolumetricFogPrevUVMaxForTemporalBlend = ViewState->VolumetricFogPrevUVMaxForTemporalBlend;
 	}
 	ViewUniformShaderParameters.VolumetricFogScreenToResourceUV = ViewRectSize / (FVector2f(VolumetricFogResourceGridSize.X, VolumetricFogResourceGridSize.Y) * VolumetricFogResourceGridPixelSize);
-	ViewUniformShaderParameters.VolumetricFogUVMax = GetVolumetricFogUVMax(ViewRectSize, VolumetricFogResourceGridSize, VolumetricFogResourceGridPixelSize);
+	ViewUniformShaderParameters.VolumetricFogUVMax = GetVolumetricFogUVMaxForSampling(ViewRectSize, VolumetricFogResourceGridSize, VolumetricFogResourceGridPixelSize);
 
 	SetupPrecomputedVolumetricLightmapUniformBufferParameters(Scene, Family->EngineShowFlags, ViewUniformShaderParameters);
 
diff --git a/Engine/Source/Runtime/Renderer/Private/VolumetricFog.cpp b/Engine/Source/Runtime/Renderer/Private/VolumetricFog.cpp
--- a/Engine/Source/Runtime/Renderer/Private/VolumetricFog.cpp
+++ b/Engine/Source/Runtime/Renderer/Private/VolumetricFog.cpp
@@ -1178,13 +1178,20 @@ FIntVector GetVolumetricFogViewGridSize(const FViewInfo& View, int32& OutVolumet
 	return GetVolumetricFogGridSize(View.ViewRect.Size(), OutVolumetricFogGridPixelSize);
 }
 
-FVector2f GetVolumetricFogUVMax(const FVector2f& ViewRectSize, FIntVector VolumetricFogResourceGridSize, int32 VolumetricFogResourceGridPixelSize)
+FVector2f GetVolumetricFogUVMaxForSampling(const FVector2f& ViewRectSize, FIntVector VolumetricFogResourceGridSize, int32 VolumetricFogResourceGridPixelSize)
 {
 	float ViewRectSizeXSafe = FMath::DivideAndRoundUp<int32>(int32(ViewRectSize.X), VolumetricFogResourceGridPixelSize) * VolumetricFogResourceGridPixelSize - (VolumetricFogResourceGridPixelSize / 2 + 1);
 	float ViewRectSizeYSafe = FMath::DivideAndRoundUp<int32>(int32(ViewRectSize.Y), VolumetricFogResourceGridPixelSize) * VolumetricFogResourceGridPixelSize - (VolumetricFogResourceGridPixelSize / 2 + 1);
 	return FVector2f(ViewRectSizeXSafe, ViewRectSizeYSafe) / (FVector2f(VolumetricFogResourceGridSize.X, VolumetricFogResourceGridSize.Y) * VolumetricFogResourceGridPixelSize);
 }
 
+FVector2f GetVolumetricFogPrevUVMaxForTemporalBlend(const FVector2f& ViewRectSize, FIntVector VolumetricFogResourceGridSize, int32 VolumetricFogResourceGridPixelSize)
+{
+	float ViewRectSizeXSafe = FMath::DivideAndRoundUp<int32>(int32(ViewRectSize.X), VolumetricFogResourceGridPixelSize) * VolumetricFogResourceGridPixelSize;
+	float ViewRectSizeYSafe = FMath::DivideAndRoundUp<int32>(int32(ViewRectSize.Y), VolumetricFogResourceGridPixelSize) * VolumetricFogResourceGridPixelSize;
+	return FVector2f(ViewRectSizeXSafe, ViewRectSizeYSafe) / (FVector2f(VolumetricFogResourceGridSize.X, VolumetricFogResourceGridSize.Y) * VolumetricFogResourceGridPixelSize);
+}
+
 FVector2f GetVolumetricFogFroxelToScreenSVPosRatio(const FViewInfo& View)
 {
 	const FIntPoint ViewRectSize = View.ViewRect.Size();
@@ -1292,6 +1299,7 @@ void FSceneRenderer::SetupVolumetricFog()
 				View.ViewState->PrevLightScatteringViewGridUVToViewRectVolumeUV = FVector2f::One();
 				View.ViewState->VolumetricFogPrevViewGridRectUVToResourceUV = FVector2f::One();
 				View.ViewState->VolumetricFogPrevUVMax = FVector2f::One();
+				View.ViewState->VolumetricFogPrevUVMaxForTemporalBlend = FVector2f::One();
 			}
 		}
 	}
@@ -1725,7 +1733,8 @@ void FSceneRenderer::ComputeVolumetricFog(FRDGBuilder& GraphBuilder,
 			View.ViewState->PrevLightScatteringViewGridUVToViewRectVolumeUV = ViewRectSize / (FVector2f(VolumetricFogViewGridSize.X, VolumetricFogViewGridSize.Y) * VolumetricFogGridPixelSize);
 
 			View.ViewState->VolumetricFogPrevViewGridRectUVToResourceUV = FVector2f(VolumetricFogViewGridSize.X, VolumetricFogViewGridSize.Y) / FVector2f(VolumetricFogResourceGridSize.X, VolumetricFogResourceGridSize.Y);
-			View.ViewState->VolumetricFogPrevUVMax = GetVolumetricFogUVMax(ViewRectSize, VolumetricFogResourceGridSize, VolumetricFogGridPixelSize);
+			View.ViewState->VolumetricFogPrevUVMax = GetVolumetricFogUVMaxForSampling(ViewRectSize, VolumetricFogResourceGridSize, VolumetricFogGridPixelSize);
+			View.ViewState->VolumetricFogPrevUVMaxForTemporalBlend = GetVolumetricFogPrevUVMaxForTemporalBlend(ViewRectSize, VolumetricFogResourceGridSize, VolumetricFogGridPixelSize);
 		}
 		else if (View.ViewState)
 		{
@@ -1734,6 +1743,7 @@ void FSceneRenderer::ComputeVolumetricFog(FRDGBuilder& GraphBuilder,
 			View.ViewState->PrevLightScatteringViewGridUVToViewRectVolumeUV = FVector2f::One();
 			View.ViewState->VolumetricFogPrevViewGridRectUVToResourceUV = FVector2f::One();
 			View.ViewState->VolumetricFogPrevUVMax = FVector2f::One();
+			View.ViewState->VolumetricFogPrevUVMaxForTemporalBlend = FVector2f::One();
 		}
 
 		if (bUseTemporalReprojection && GVolumetricFogConservativeDepth > 0)
diff --git a/Engine/Source/Runtime/Renderer/Private/VolumetricFog.h b/Engine/Source/Runtime/Renderer/Private/VolumetricFog.h
--- a/Engine/Source/Runtime/Renderer/Private/VolumetricFog.h
+++ b/Engine/Source/Runtime/Renderer/Private/VolumetricFog.h
@@ -15,7 +15,8 @@ extern FIntVector GetVolumetricFogResourceGridSize(const FViewInfo& View, int32&
 // Grid size for the view rectangle within the allocated resource
 extern FIntVector GetVolumetricFogViewGridSize(const FViewInfo& View, int32& OutVolumetricFogGridPixelSize);
 
-extern FVector2f GetVolumetricFogUVMax(const FVector2f& ViewRectSize, FIntVector VolumetricFogResourceGridSize, int32 VolumetricFogResourceGridPixelSize);
+extern FVector2f GetVolumetricFogUVMaxForSampling(const FVector2f& ViewRectSize, FIntVector VolumetricFogResourceGridSize, int32 VolumetricFogResourceGridPixelSize);
+extern FVector2f GetVolumetricFogPrevUVMaxForTemporalBlend(const FVector2f& ViewRectSize, FIntVector VolumetricFogResourceGridSize, int32 VolumetricFogResourceGridPixelSize);
 
 extern FVector2f GetVolumetricFogFroxelToScreenSVPosRatio(const FViewInfo& View);
 
