Sampled Direct Lighting
- support sampling Light Function Atlas.

diff --git a/Engine/Source/Runtime/Renderer/Private/LightFunctionAtlas.cpp b/Engine/Source/Runtime/Renderer/Private/LightFunctionAtlas.cpp
--- a/Engine/Source/Runtime/Renderer/Private/LightFunctionAtlas.cpp
+++ b/Engine/Source/Runtime/Renderer/Private/LightFunctionAtlas.cpp
@@ -13,6 +13,7 @@
 #include "SystemTextures.h"
 #include "PostProcess/PostProcessing.h"
 #include "PostProcess/SceneFilterRendering.h"
+#include "SampledDirectLighting/SampledDirectLighting.h"
 #include "ShadowRendering.h"
 #include "CanvasItem.h"
 #include "CanvasTypes.h"
@@ -104,7 +105,13 @@ FAutoConsoleVariableRef CVarDeferredLightsUsesLightFunctionAtlas(
 	ECVF_RenderThreadSafe
 );
 
-
+int GSampledDirectLightingUsesLightFunctionAtlas = 1;
+FAutoConsoleVariableRef CVarSampledDirectLightingUsesLightFunctionAtlas(
+	TEXT("r.SampledDirectLighting.UsesLightFunctionAtlas"),
+	GSampledDirectLightingUsesLightFunctionAtlas,
+	TEXT("Whether the light function atlas is sampled when rendering sampled direct lighting."),
+	ECVF_Scalability | ECVF_RenderThreadSafe
+);
 
 //////////////////////////////////////////////////////////////////////////
 
@@ -232,6 +239,7 @@ void FLightFunctionAtlas::ClearEmptySceneFrame(FViewInfo* View, FLightFunctionAt
 			this,
 			bLightFunctionAtlasEnabled,
 			bLightFunctionAtlasEnabled,
+			bLightFunctionAtlasEnabled,
 			bLightFunctionAtlasEnabled);
 	}
 
@@ -249,16 +257,19 @@ void FLightFunctionAtlas::BeginSceneFrame(FViewFamilyInfo& ViewFamily, TArray<FV
 	bLightFunctionAtlasEnabled = CVarLightFunctionAtlas.GetValueOnRenderThread() > 0 && ViewFamily.EngineShowFlags.LightFunctions > 0;
 
 	// But only really enable the atlas generation if a system asks for it
-	bool VolumetricFogRequestsLightFunctionAtlas = false;
-	bool DeferredlightingRequestsLightFunctionAtlas = false;
+	bool bVolumetricFogRequestsLightFunctionAtlas = false;
+	bool bDeferredlightingRequestsLightFunctionAtlas = false;
+	bool bSampledDirectLightingRequestsLightFunctionAtlas = false;
 	if (bLightFunctionAtlasEnabled)
 	{
-		VolumetricFogRequestsLightFunctionAtlas = bShouldRenderVolumetricFog && GVolumetricFogUsesLightFunctionAtlas > 0;
+		bVolumetricFogRequestsLightFunctionAtlas = bShouldRenderVolumetricFog && GVolumetricFogUsesLightFunctionAtlas > 0;
+
+		bDeferredlightingRequestsLightFunctionAtlas = GDeferredUsesLightFunctionAtlas > 0;
 
-		DeferredlightingRequestsLightFunctionAtlas = GDeferredUsesLightFunctionAtlas > 0;
+		bSampledDirectLightingRequestsLightFunctionAtlas = SampledDirectLighting::IsEnabled() && GSampledDirectLightingUsesLightFunctionAtlas > 0;
 
 		bLightFunctionAtlasEnabled = bLightFunctionAtlasEnabled && 
-			(VolumetricFogRequestsLightFunctionAtlas || DeferredlightingRequestsLightFunctionAtlas ||
+			(bVolumetricFogRequestsLightFunctionAtlas || bDeferredlightingRequestsLightFunctionAtlas || bSampledDirectLightingRequestsLightFunctionAtlas ||
 			 GetSingleLayerWaterUsesLightFunctionAtlas() || GetTranslucentUsesLightFunctionAtlas()); 
 	}
 
@@ -266,8 +277,9 @@ void FLightFunctionAtlas::BeginSceneFrame(FViewFamilyInfo& ViewFamily, TArray<FV
 	LightFunctionAtlasSceneData.SetData(
 		this,
 		bLightFunctionAtlasEnabled,
-		bLightFunctionAtlasEnabled && VolumetricFogRequestsLightFunctionAtlas,
-		bLightFunctionAtlasEnabled && DeferredlightingRequestsLightFunctionAtlas);
+		bLightFunctionAtlasEnabled && bVolumetricFogRequestsLightFunctionAtlas,
+		bLightFunctionAtlasEnabled && bDeferredlightingRequestsLightFunctionAtlas,
+		bLightFunctionAtlasEnabled && bSampledDirectLightingRequestsLightFunctionAtlas);
 
 	for (auto& View : Views)
 	{
diff --git a/Engine/Source/Runtime/Renderer/Private/LightFunctionAtlas.h b/Engine/Source/Runtime/Renderer/Private/LightFunctionAtlas.h
--- a/Engine/Source/Runtime/Renderer/Private/LightFunctionAtlas.h
+++ b/Engine/Source/Runtime/Renderer/Private/LightFunctionAtlas.h
@@ -24,24 +24,31 @@ struct FLightFunctionAtlas;
 
 struct FLightFunctionAtlasSceneData
 {
-	void SetData(FLightFunctionAtlas* InLightFunctionAtlas, bool bInLightFunctionAtlasEnabled, bool bInVolumetricFogUsesLightFunctionAtlas, bool bInDeferredlightingUsesLightFunctionAtlas)
+	void SetData(FLightFunctionAtlas* InLightFunctionAtlas,
+		bool bInLightFunctionAtlasEnabled,
+		bool bInVolumetricFogUsesLightFunctionAtlas,
+		bool bInDeferredlightingUsesLightFunctionAtlas,
+		bool bInSampledDirectLightingUsesLightFunctionAtlas)
 	{
 		LightFunctionAtlas = InLightFunctionAtlas;
 		bLightFunctionAtlasEnabled = bInLightFunctionAtlasEnabled;
 		bVolumetricFogUsesLightFunctionAtlas = bInVolumetricFogUsesLightFunctionAtlas;
 		bDeferredlightingUsesLightFunctionAtlas = bInDeferredlightingUsesLightFunctionAtlas;
+		bSampledDirectLightingUsesLightFunctionAtlas = bInSampledDirectLightingUsesLightFunctionAtlas;
 	}
 
-	FLightFunctionAtlas* GetLightFunctionAtlas()		const { return LightFunctionAtlas; }
-	bool GetLightFunctionAtlasEnabled()					const { return bLightFunctionAtlasEnabled; }
-	bool GetVolumetricFogUsesLightFunctionAtlas()		const { return bVolumetricFogUsesLightFunctionAtlas; }
-	bool GetDeferredlightingUsesLightFunctionAtlas()	const { return bDeferredlightingUsesLightFunctionAtlas; }
+	FLightFunctionAtlas* GetLightFunctionAtlas()			const { return LightFunctionAtlas; }
+	bool GetLightFunctionAtlasEnabled()						const { return bLightFunctionAtlasEnabled; }
+	bool GetVolumetricFogUsesLightFunctionAtlas()			const { return bVolumetricFogUsesLightFunctionAtlas; }
+	bool GetDeferredlightingUsesLightFunctionAtlas()		const { return bDeferredlightingUsesLightFunctionAtlas; }
+	bool GetSampledDirectLightingUsesLightFunctionAtlas()	const { return bSampledDirectLightingUsesLightFunctionAtlas; }
 
 private:
 	FLightFunctionAtlas* LightFunctionAtlas = nullptr;
 	bool bLightFunctionAtlasEnabled = false;
 	bool bVolumetricFogUsesLightFunctionAtlas = false;
 	bool bDeferredlightingUsesLightFunctionAtlas = false;
+	bool bSampledDirectLightingUsesLightFunctionAtlas = false;
 };
 
 struct FLightFunctionAtlasViewData
@@ -49,10 +56,11 @@ struct FLightFunctionAtlasViewData
 	FLightFunctionAtlasViewData() : SceneData(nullptr) {}
 	FLightFunctionAtlasViewData(FLightFunctionAtlasSceneData* InSceneData) : SceneData(InSceneData) {}
 
-	FLightFunctionAtlas* GetLightFunctionAtlas()		const { return SceneData ? SceneData->GetLightFunctionAtlas() : nullptr; }
-	bool GetLightFunctionAtlasEnabled()					const { return SceneData ? SceneData->GetLightFunctionAtlasEnabled() : false; }
-	bool GetVolumetricFogUsesLightFunctionAtlas()		const { return SceneData ? SceneData->GetVolumetricFogUsesLightFunctionAtlas()		: false; }
-	bool GetDeferredlightingUsesLightFunctionAtlas()	const { return SceneData ? SceneData->GetDeferredlightingUsesLightFunctionAtlas()	: false; }
+	FLightFunctionAtlas* GetLightFunctionAtlas()			const { return SceneData ? SceneData->GetLightFunctionAtlas() : nullptr; }
+	bool GetLightFunctionAtlasEnabled()						const { return SceneData ? SceneData->GetLightFunctionAtlasEnabled() : false; }
+	bool GetVolumetricFogUsesLightFunctionAtlas()			const { return SceneData ? SceneData->GetVolumetricFogUsesLightFunctionAtlas() : false; }
+	bool GetDeferredlightingUsesLightFunctionAtlas()		const { return SceneData ? SceneData->GetDeferredlightingUsesLightFunctionAtlas() : false; }
+	bool GetSampledDirectLightingUsesLightFunctionAtlas()	const { return SceneData ? SceneData->GetSampledDirectLightingUsesLightFunctionAtlas() : false; }
 
 private:
 	FLightFunctionAtlasSceneData* SceneData;
diff --git a/Engine/Source/Runtime/Renderer/Private/SampledDirectLighting/SampledDirectLighting.cpp b/Engine/Source/Runtime/Renderer/Private/SampledDirectLighting/SampledDirectLighting.cpp
--- a/Engine/Source/Runtime/Renderer/Private/SampledDirectLighting/SampledDirectLighting.cpp
+++ b/Engine/Source/Runtime/Renderer/Private/SampledDirectLighting/SampledDirectLighting.cpp
@@ -331,11 +331,12 @@ class FGenerateSamplesCS : public FGlobalShader
 
 	class FTileType : SHADER_PERMUTATION_INT("TILE_TYPE", (int32)SampledDirectLighting::ETileType::MAX);
 	class FIESProfile : SHADER_PERMUTATION_BOOL("USE_IES_PROFILE");
+	class FLightFunctionAtlas : SHADER_PERMUTATION_BOOL("USE_LIGHT_FUNCTION_ATLAS");
 	class FTexturedRectLights : SHADER_PERMUTATION_BOOL("USE_SOURCE_TEXTURE");
 	class FNumSamplesPerPixel1d : SHADER_PERMUTATION_SPARSE_INT("NUM_SAMPLES_PER_PIXEL_1D", 1, 2, 4);
 	class FCandidateLightMask : SHADER_PERMUTATION_BOOL("CANDIDATE_LIGHT_MASK");
 	class FDebugMode : SHADER_PERMUTATION_BOOL("DEBUG_MODE");
-	using FPermutationDomain = TShaderPermutationDomain<FTileType, FIESProfile, FTexturedRectLights, FNumSamplesPerPixel1d, FCandidateLightMask, FDebugMode>;
+	using FPermutationDomain = TShaderPermutationDomain<FTileType, FIESProfile, FLightFunctionAtlas, FTexturedRectLights, FNumSamplesPerPixel1d, FCandidateLightMask, FDebugMode>;
 
 	static int32 GetGroupSize()
 	{	
@@ -454,9 +455,10 @@ class FShadeLightSamplesCS : public FGlobalShader
 
 	class FTileType : SHADER_PERMUTATION_INT("TILE_TYPE", (int32)SampledDirectLighting::ETileType::MAX);
 	class FIESProfile : SHADER_PERMUTATION_BOOL("USE_IES_PROFILE");
+	class FLightFunctionAtlas : SHADER_PERMUTATION_BOOL("USE_LIGHT_FUNCTION_ATLAS");
 	class FTexturedRectLights : SHADER_PERMUTATION_BOOL("USE_SOURCE_TEXTURE");
 	class FDebugMode : SHADER_PERMUTATION_BOOL("DEBUG_MODE");
-	using FPermutationDomain = TShaderPermutationDomain<FTileType, FIESProfile, FTexturedRectLights, FDebugMode>;
+	using FPermutationDomain = TShaderPermutationDomain<FTileType, FIESProfile, FLightFunctionAtlas, FTexturedRectLights, FDebugMode>;
 
 	static bool ShouldCompilePermutation(const FGlobalShaderPermutationParameters& Parameters)
 	{
@@ -536,7 +538,8 @@ void FDeferredShadingSceneRenderer::RenderSampledDirectLighting(FRDGBuilder& Gra
 
 	RDG_EVENT_SCOPE(GraphBuilder, "SampledDirectLighting");
 
-	const FViewInfo& View = Views[0];
+	const uint32 ViewIndex = 0;
+	const FViewInfo& View = Views[ViewIndex];
 	FBlueNoise BlueNoise = GetBlueNoiseGlobalParameters();
 	TUniformBufferRef<FBlueNoise> BlueNoiseUniformBuffer = CreateUniformBufferImmediate(BlueNoise, EUniformBufferUsage::UniformBuffer_SingleDraw);
 
@@ -630,6 +633,10 @@ void FDeferredShadingSceneRenderer::RenderSampledDirectLighting(FRDGBuilder& Gra
 		}
 	}
 
+	// Setup the light function atlas
+	const bool bUseLightFunctionAtlas = View.LightFunctionAtlasViewData.GetSampledDirectLightingUsesLightFunctionAtlas();
+	TRDGUniformBufferRef<FLightFunctionAtlasGlobalParameters> LightFunctionAtlasGlobalParameters = LightFunctionAtlas.GetLightFunctionAtlasGlobalParameters(ViewIndex, GraphBuilder);
+
 	const FIntPoint ViewSizeInTiles = FIntPoint::DivideAndRoundUp(View.ViewRect.Size(), SampledDirectLighting::TileSize);
 	const int32 TileDataStride = ViewSizeInTiles.X * ViewSizeInTiles.Y;
 
@@ -648,6 +655,7 @@ void FDeferredShadingSceneRenderer::RenderSampledDirectLighting(FRDGBuilder& Gra
 		SampledDirectLightingParameters.SceneTexturesStruct = SceneTextures.UniformBuffer;
 		SampledDirectLightingParameters.Substrate = Substrate::BindSubstrateGlobalUniformParameters(View);
 		SampledDirectLightingParameters.ForwardLightData = View.ForwardLightingResources.ForwardLightUniformBuffer;
+		SampledDirectLightingParameters.LightFunctionAtlas = LightFunctionAtlasGlobalParameters;
 		SampledDirectLightingParameters.BlueNoise = BlueNoiseUniformBuffer;
 		SampledDirectLightingParameters.PreIntegratedGF = GSystemTextures.PreintegratedGF->GetRHI();
 		SampledDirectLightingParameters.PreIntegratedGFSampler = TStaticSamplerState<SF_Bilinear, AM_Clamp, AM_Clamp, AM_Clamp>::GetRHI();
@@ -804,6 +812,7 @@ void FDeferredShadingSceneRenderer::RenderSampledDirectLighting(FRDGBuilder& Gra
 			FGenerateSamplesCS::FPermutationDomain PermutationVector;
 			PermutationVector.Set<FGenerateSamplesCS::FTileType>(TileType);
 			PermutationVector.Set<FGenerateSamplesCS::FIESProfile>(CVarSampledDirectLightingIESProfiles.GetValueOnRenderThread() != 0);
+			PermutationVector.Set<FGenerateSamplesCS::FLightFunctionAtlas>(bUseLightFunctionAtlas);
 			PermutationVector.Set<FGenerateSamplesCS::FTexturedRectLights>(CVarSampledDirectLightingTexturedRectLights.GetValueOnRenderThread() != 0);
 			PermutationVector.Set<FGenerateSamplesCS::FNumSamplesPerPixel1d>(NumSamplesPerPixel2d.X * NumSamplesPerPixel2d.Y);
 			PermutationVector.Set<FGenerateSamplesCS::FDebugMode>(bDebug);
@@ -951,6 +960,7 @@ void FDeferredShadingSceneRenderer::RenderSampledDirectLighting(FRDGBuilder& Gra
 			FShadeLightSamplesCS::FPermutationDomain PermutationVector;
 			PermutationVector.Set<FShadeLightSamplesCS::FTileType>(TileType);
 			PermutationVector.Set<FShadeLightSamplesCS::FIESProfile>(CVarSampledDirectLightingIESProfiles.GetValueOnRenderThread() != 0);
+			PermutationVector.Set<FShadeLightSamplesCS::FLightFunctionAtlas>(bUseLightFunctionAtlas);
 			PermutationVector.Set<FShadeLightSamplesCS::FTexturedRectLights>(CVarSampledDirectLightingTexturedRectLights.GetValueOnRenderThread() != 0);
 			PermutationVector.Set<FShadeLightSamplesCS::FDebugMode>(bDebug);
 			auto ComputeShader = View.ShaderMap->GetShader<FShadeLightSamplesCS>(PermutationVector);
diff --git a/Engine/Source/Runtime/Renderer/Private/SampledDirectLighting/SampledDirectLightingInternal.h b/Engine/Source/Runtime/Renderer/Private/SampledDirectLighting/SampledDirectLightingInternal.h
--- a/Engine/Source/Runtime/Renderer/Private/SampledDirectLighting/SampledDirectLightingInternal.h
+++ b/Engine/Source/Runtime/Renderer/Private/SampledDirectLighting/SampledDirectLightingInternal.h
@@ -14,6 +14,7 @@ BEGIN_SHADER_PARAMETER_STRUCT(FSampledDirectLightingParameters, )
 	SHADER_PARAMETER_RDG_UNIFORM_BUFFER(FSceneTextureUniformParameters, SceneTexturesStruct)
 	SHADER_PARAMETER_RDG_UNIFORM_BUFFER(FSubstrateGlobalUniformParameters, Substrate)
 	SHADER_PARAMETER_RDG_UNIFORM_BUFFER(FForwardLightData, ForwardLightData)
+	SHADER_PARAMETER_RDG_UNIFORM_BUFFER(FLightFunctionAtlasGlobalParameters, LightFunctionAtlas)
 	SHADER_PARAMETER_STRUCT_REF(FBlueNoise, BlueNoise)
 	SHADER_PARAMETER_TEXTURE(Texture2D, PreIntegratedGF)
 	SHADER_PARAMETER_SAMPLER(SamplerState, PreIntegratedGFSampler)
