Fix GPUScene.LightData buffer being created with incorrect stride (BytesPerElement).
- ResizeStructuredBufferIfNeeded() only works when BytesPerElement = 16.

diff --git a/Engine/Source/Runtime/Renderer/Private/GPUScene.cpp b/Engine/Source/Runtime/Renderer/Private/GPUScene.cpp
--- a/Engine/Source/Runtime/Renderer/Private/GPUScene.cpp
+++ b/Engine/Source/Runtime/Renderer/Private/GPUScene.cpp
@@ -615,6 +615,7 @@ FGPUScene::FGPUScene(FScene &InScene)
 	, InstanceSceneDataSOAStride(0)
 	, InstancePayloadDataAllocator(CVarGPUSceneUseGrowOnlyAllocationPolicy.GetValueOnAnyThread() != 0)
 	, LightmapDataAllocator(CVarGPUSceneUseGrowOnlyAllocationPolicy.GetValueOnAnyThread() != 0)
+	, LightDataBuffer(InitialBufferSize, TEXT("GPUScene.LightData"))
 	, Scene(InScene)
 	, InstanceSceneDataAllocator(CVarGPUSceneUseGrowOnlyAllocationPolicy.GetValueOnAnyThread() != 0)
 {
@@ -706,8 +707,7 @@ void FGPUScene::UpdateGPULights(FRDGBuilder& GraphBuilder, const UE::Tasks::FTas
 
 	}, PrerequisiteTask);
 
-	const uint32 LightDataBufferSize = FMath::RoundUpToPowerOfTwo(FMath::Max(Scene.Lights.GetMaxIndex(), InitialBufferSize));
-	FRDGBuffer *LightDataBufferRDG = ResizeStructuredBufferIfNeeded(GraphBuilder, LightDataBuffer, LightDataBufferSize * sizeof(FLightSceneData), TEXT("GPUScene.LightData"));
+	FRDGBuffer* LightDataBufferRDG = LightDataBuffer.ResizeBufferIfNeeded(GraphBuilder, Scene.Lights.GetMaxIndex());
 
 	GraphBuilder.QueueBufferUpload<FLightSceneData>(LightDataBufferRDG, LightData, ERDGInitialDataFlags::NoCopy);
 }
@@ -892,7 +892,7 @@ FGPUScene::FRegisteredBuffers FGPUScene::UpdateBufferAllocations(FRDGBuilder& Gr
 	const uint32 LightMapDataBufferSize = FMath::RoundUpToPowerOfTwo(FMath::Max(LightmapDataAllocator.GetMaxSize(), InitialBufferSize));
 	BufferState.LightmapDataBuffer = ResizeStructuredBufferIfNeeded(GraphBuilder, LightmapDataBuffer, LightMapDataBufferSize * sizeof(FLightmapSceneShaderData::Data), TEXT("GPUScene.LightmapData"));
 
-	BufferState.LightDataBuffer = GraphBuilder.RegisterExternalBuffer(LightDataBuffer);
+	BufferState.LightDataBuffer = LightDataBuffer.Register(GraphBuilder);
 
 	FGPUSceneResourceParameters ShaderParameters;
 	ShaderParameters.GPUSceneInstanceSceneData = GraphBuilder.CreateSRV(BufferState.InstanceSceneDataBuffer);
@@ -928,7 +928,7 @@ FGPUScene::FRegisteredBuffers FGPUScene::RegisterBuffers(FRDGBuilder& GraphBuild
 	Result.InstanceSceneDataBuffer = GraphBuilder.RegisterExternalBuffer(InstanceSceneDataBuffer);
 	Result.InstancePayloadDataBuffer = GraphBuilder.RegisterExternalBuffer(InstancePayloadDataBuffer);
 	Result.LightmapDataBuffer = GraphBuilder.RegisterExternalBuffer(LightmapDataBuffer);
-	Result.LightDataBuffer = GraphBuilder.RegisterExternalBuffer(LightDataBuffer);
+	Result.LightDataBuffer = LightDataBuffer.Register(GraphBuilder);
 
 	return Result;
 }
@@ -2175,7 +2175,7 @@ void FGPUScene::OnPostLightSceneInfoUpdate(FRDGBuilder& GraphBuilder, const FLig
 	// SceneFrameNumber is updated in UpdateInternal so if it is the same, we have already uploaded the lights this "frame" - this is not 100% robust since 
 	// if UpdateAllPrimitiveSceneInfos is called multiple times without updating the FScene frame number it will skip light uploads. The real solution is 
 	// to implement functional change tracking for lights so we can actually know when there are relevant changes.
-	if (SceneFrameNumber != Scene.GetFrameNumberRenderThread() || !LightDataBuffer.IsValid())
+	if (SceneFrameNumber != Scene.GetFrameNumberRenderThread() || !LightDataBuffer.GetPooledBuffer().IsValid())
 	{
 		UpdateGPULights(GraphBuilder, UE::Tasks::FTask{});
 	}
diff --git a/Engine/Source/Runtime/Renderer/Private/GPUScene.h b/Engine/Source/Runtime/Renderer/Private/GPUScene.h
--- a/Engine/Source/Runtime/Renderer/Private/GPUScene.h
+++ b/Engine/Source/Runtime/Renderer/Private/GPUScene.h
@@ -13,6 +13,7 @@
 #include "LightSceneData.h"
 #include "SceneUniformBuffer.h"
 #include "UnifiedBuffer.h"
+#include "RendererPrivateUtils.h"
 
 class FRDGExternalAccessQueue;
 class FRHICommandList;
@@ -333,7 +334,7 @@ public:
 private:
 	static constexpr int32 InitialBufferSize = 256;
 
-	TRefCountPtr<FRDGPooledBuffer> LightDataBuffer;
+	TPersistentStructuredBuffer<FLightSceneData> LightDataBuffer;
 
 	// Buffers used by GPU-Scene, since they can be resized during updates AND the render passes must retain the 
 	// right copy (this is chiefly because the init of shadow views after pre-pass means we need to be able to set 
diff --git a/Engine/Source/Runtime/Renderer/Private/RendererPrivateUtils.h b/Engine/Source/Runtime/Renderer/Private/RendererPrivateUtils.h
--- a/Engine/Source/Runtime/Renderer/Private/RendererPrivateUtils.h
+++ b/Engine/Source/Runtime/Renderer/Private/RendererPrivateUtils.h
@@ -98,7 +98,7 @@ class FPersistentBuffer
 public:
 	FPersistentBuffer(int32 InMinimumNumElementsReserved, const TCHAR *InName, bool bInRoundUpToPOT = true);
 
-	FRDGBuffer* Register(FRDGBuilder& GraphBuilder);
+	FRDGBuffer* Register(FRDGBuilder& GraphBuilder) const;
 
 	void Empty();
 
diff --git a/Engine/Source/Runtime/Renderer/Private/RendererUtils.cpp b/Engine/Source/Runtime/Renderer/Private/RendererUtils.cpp
--- a/Engine/Source/Runtime/Renderer/Private/RendererUtils.cpp
+++ b/Engine/Source/Runtime/Renderer/Private/RendererUtils.cpp
@@ -237,7 +237,7 @@ FPersistentBuffer::FPersistentBuffer(int32 InMinimumNumElementsReserved, const T
 {
 }
 
-FRDGBuffer* FPersistentBuffer::Register(FRDGBuilder& GraphBuilder) 
+FRDGBuffer* FPersistentBuffer::Register(FRDGBuilder& GraphBuilder) const
 { 
 	return GraphBuilder.RegisterExternalBuffer(PooledBuffer); 
 }
