Fixes a TSR crash on Intel(R) UHD Graphics GPUs

diff --git a/Engine/Source/Runtime/Renderer/Private/PostProcess/TemporalSuperResolution.cpp b/Engine/Source/Runtime/Renderer/Private/PostProcess/TemporalSuperResolution.cpp
--- a/Engine/Source/Runtime/Renderer/Private/PostProcess/TemporalSuperResolution.cpp
+++ b/Engine/Source/Runtime/Renderer/Private/PostProcess/TemporalSuperResolution.cpp
@@ -1310,17 +1310,20 @@ FDefaultTemporalUpscaler::FOutputs AddTemporalSuperResolutionPasses(
 
 	// Controls whether to use 16bit ops on per GPU vendor in mean time each driver matures.
 #if PLATFORM_DESKTOP
-	if (IsRHIDeviceAMD())
+	if ((GRHIGlobals.SupportsNative16BitOps && VALU16BitSupport == ERHIFeatureSupport::RuntimeDependent) || VALU16BitSupport == ERHIFeatureSupport::RuntimeGuaranteed)
 	{
-		bUse16BitVALU = CVarTSR16BitVALUOnAMD.GetValueOnRenderThread() != 0;
-	}
-	else if (IsRHIDeviceIntel())
-	{
-		bUse16BitVALU = CVarTSR16BitVALUOnIntel.GetValueOnRenderThread() != 0;
-	}
-	else if (IsRHIDeviceNVIDIA())
-	{
-		bUse16BitVALU = CVarTSR16BitVALUOnNvidia.GetValueOnRenderThread() != 0;
+		if (IsRHIDeviceAMD())
+		{
+			bUse16BitVALU = CVarTSR16BitVALUOnAMD.GetValueOnRenderThread() != 0;
+		}
+		else if (IsRHIDeviceIntel())
+		{
+			bUse16BitVALU = CVarTSR16BitVALUOnIntel.GetValueOnRenderThread() != 0;
+		}
+		else if (IsRHIDeviceNVIDIA())
+		{
+			bUse16BitVALU = CVarTSR16BitVALUOnNvidia.GetValueOnRenderThread() != 0;
+		}
 	}
 #endif // PLATFORM_DESKTOP
 
