Clear FRayTracingDynamicGeometryCollection in EndUpdate()
- also check working arrays are empty at start of BeginUpdate()

diff --git a/Engine/Source/Runtime/Renderer/Private/DeferredShadingRenderer.cpp b/Engine/Source/Runtime/Renderer/Private/DeferredShadingRenderer.cpp
--- a/Engine/Source/Runtime/Renderer/Private/DeferredShadingRenderer.cpp
+++ b/Engine/Source/Runtime/Renderer/Private/DeferredShadingRenderer.cpp
@@ -1615,11 +1615,6 @@ void FDeferredShadingSceneRenderer::Render(FRDGBuilder& GraphBuilder)
 
 		GRayTracingGeometryManager->Tick(GraphBuilder.RHICmdList, bHasRayTracingEnableChanged);
 
-		if ((GetRayTracingMode() == ERayTracingMode::Dynamic) && bHasRayTracingEnableChanged)
-		{
-			Scene->GetRayTracingDynamicGeometryCollection()->Clear();
-		}
-
 		// Now that we have updated all the PrimitiveSceneInfos, update the RayTracing mesh commands cache if needed
 		{
 			const ERayTracingMeshCommandsMode CurrentMode = ViewFamily.EngineShowFlags.PathTracing ? ERayTracingMeshCommandsMode::PATH_TRACING : ERayTracingMeshCommandsMode::RAY_TRACING;
diff --git a/Engine/Source/Runtime/Renderer/Private/RayTracing/RayTracingDynamicGeometry.cpp b/Engine/Source/Runtime/Renderer/Private/RayTracing/RayTracingDynamicGeometry.cpp
--- a/Engine/Source/Runtime/Renderer/Private/RayTracing/RayTracingDynamicGeometry.cpp
+++ b/Engine/Source/Runtime/Renderer/Private/RayTracing/RayTracingDynamicGeometry.cpp
@@ -132,7 +132,10 @@ void FRayTracingDynamicGeometryCollection::Clear()
 
 int64 FRayTracingDynamicGeometryCollection::BeginUpdate()
 {
-	Clear();
+	check(DispatchCommands.IsEmpty());
+	check(BuildParams.IsEmpty());
+	check(Segments.IsEmpty());
+	check(ReferencedUniformBuffers.IsEmpty());
 
 	// Vertex buffer data can be immediatly reused the next frame, because it's already 'consumed' for building the AccelerationStructure data
 	// Garbage collect unused buffers for n generations
@@ -581,6 +584,8 @@ void FRayTracingDynamicGeometryCollection::DispatchUpdates(FRHICommandListImmedi
 void FRayTracingDynamicGeometryCollection::EndUpdate(FRHICommandListImmediate& RHICmdList)
 {
 	ReferencedUniformBuffers.Empty(ReferencedUniformBuffers.Max());
+
+	Clear();
 }
 
 uint32 FRayTracingDynamicGeometryCollection::ComputeScratchBufferSize()
diff --git a/Engine/Source/Runtime/Renderer/Public/RayTracingDynamicGeometryCollection.h b/Engine/Source/Runtime/Renderer/Public/RayTracingDynamicGeometryCollection.h
--- a/Engine/Source/Runtime/Renderer/Public/RayTracingDynamicGeometryCollection.h
+++ b/Engine/Source/Runtime/Renderer/Public/RayTracingDynamicGeometryCollection.h
@@ -43,7 +43,6 @@ public:
 	RENDERER_API void EndUpdate(FRHICommandListImmediate& RHICmdList);
 
 	// Clears the working arrays to not hold any references.
-	// Needs to be called every frame when ray tracing is enabled or once when ray tracing mode has changed.
 	RENDERER_API void Clear();
 
 	RENDERER_API uint32 ComputeScratchBufferSize();
