Ray Tracing instance culling on the GPU:
* Implement instance culling code on the GPU as a part of BuildTLASInstanceBuffer pass.
* Add PRIMITIVE_SCENE_DATA_FLAG_RAYTRACING_HAS_GROUPID to be able to disable per-instance culling in the shader.
* Only culling modes 2 & 3 are supported.
* Controlled with r.RayTracing.Culling.UseGPU (default off)
* Fix incorrect instance center calculation for CPU culling.
* Fix Instance Overlap mode crash - we didn't create a buffer if there were no instances in the scene.

diff --git a/Engine/Plugins/Experimental/GPULightmass/Source/GPULightmass/Private/LightmapRenderer.cpp b/Engine/Plugins/Experimental/GPULightmass/Source/GPULightmass/Private/LightmapRenderer.cpp
--- a/Engine/Plugins/Experimental/GPULightmass/Source/GPULightmass/Private/LightmapRenderer.cpp
+++ b/Engine/Plugins/Experimental/GPULightmass/Source/GPULightmass/Private/LightmapRenderer.cpp
@@ -1213,6 +1213,7 @@ bool FSceneRenderState::SetupRayTracingScene(FRDGBuilder& GraphBuilder, FSceneUn
 				SceneWithGeometryInstances.NumNativeGPUSceneInstances,
 				SceneWithGeometryInstances.NumNativeCPUInstances,
 				{},
+				nullptr,
 				nullptr);
 
 			RHICmdList.BindAccelerationStructureMemory(RayTracingScene, RayTracingSceneBuffer, 0);
diff --git a/Engine/Shaders/Private/RayTracing/RayTracingInstanceBufferUtil.usf b/Engine/Shaders/Private/RayTracing/RayTracingInstanceBufferUtil.usf
--- a/Engine/Shaders/Private/RayTracing/RayTracingInstanceBufferUtil.usf
+++ b/Engine/Shaders/Private/RayTracing/RayTracingInstanceBufferUtil.usf
@@ -44,6 +44,54 @@ StructuredBuffer<float4> InstanceTransforms;
 // Extra per instance debug data
 RWStructuredBuffer<uint> RWDebugInstanceGPUSceneIndices;
 
+#if GPU_CULLING
+
+float CullingRadius;
+float FarFieldCullingRadius;
+float AngleThresholdRatioSq;
+float3 ViewOrigin;
+uint CullingMode;
+
+bool CullInstance(FInstanceSceneData InstanceSceneData, float4x4 LocalToTranslatedWorld4x4, bool bFarFieldInstance)
+{
+	bool bIsVisible = true;
+
+	const float3 LocalBoxCenter = InstanceSceneData.LocalBoundsCenter.xyz;
+	const float3 LocalBoxExtent = InstanceSceneData.LocalBoundsExtent.xyz;
+	const float4x4 LocalToTranslatedWorld = LocalToTranslatedWorld4x4;
+	const float3 NonUniformScale = InstanceSceneData.NonUniformScale.xyz;
+	
+	const float3 CenterTranslatedWorld = mul(float4(LocalBoxCenter, 1.0f), LocalToTranslatedWorld).xyz;
+	const float Radius = length(LocalBoxExtent * NonUniformScale);
+
+	const float InstanceDrawDistSq = length2(CenterTranslatedWorld - ViewOrigin);
+	
+	if (bFarFieldInstance)
+	{
+		bIsVisible = InstanceDrawDistSq <= Square(FarFieldCullingRadius + Radius);
+	}
+	else
+	{
+		const bool bIsVisibleDistance = InstanceDrawDistSq <= Square(CullingRadius + Radius);
+		const bool bIsVisibleAngle = (Square(Radius) / InstanceDrawDistSq) >= AngleThresholdRatioSq;
+		
+		// Only culling modes 2 & 3 are supported for now
+		if (CullingMode == 3)
+		{
+			bIsVisible = bIsVisibleDistance && bIsVisibleAngle;
+		}
+		else
+		{
+			bIsVisible = bIsVisibleDistance || bIsVisibleAngle;	
+		}
+		
+	}	
+
+	return !bIsVisible;
+}
+
+#endif // GPU_CULLING
+
 [numthreads(THREADGROUP_SIZE, 1, 1)]
 void RayTracingBuildInstanceBufferCS(uint3 DispatchThreadId : SV_DispatchThreadID)
 {
@@ -51,6 +99,8 @@ void RayTracingBuildInstanceBufferCS(uint3 DispatchThreadId : SV_DispatchThreadI
 
 	const uint Index = DispatchThreadId.x;
 
+	bool bCulled = false;
+
 	if (Index < NumInstances)
 	{
 		FRayTracingInstanceDescriptorInput InputDesc = InputInstanceDescriptors[InputDescOffset + Index];
@@ -61,7 +111,17 @@ void RayTracingBuildInstanceBufferCS(uint3 DispatchThreadId : SV_DispatchThreadI
 
 		float4x4 LocalToTranslatedWorld4x4 = LWCMultiplyTranslation(InstanceSceneData.LocalToWorld, PreViewTranslation);
 
-		if (PrimitiveData.Flags & PRIMITIVE_SCENE_DATA_FLAG_RAYTRACING_FAR_FIELD)
+		const bool bFarFieldInstance = PrimitiveData.Flags & PRIMITIVE_SCENE_DATA_FLAG_RAYTRACING_FAR_FIELD;
+
+#if GPU_CULLING
+		const bool bHasGroupId = PrimitiveData.Flags & PRIMITIVE_SCENE_DATA_FLAG_RAYTRACING_HAS_GROUPID;
+		if (!bHasGroupId)
+		{
+			bCulled = CullInstance(InstanceSceneData, LocalToTranslatedWorld4x4, bFarFieldInstance);
+		}
+#endif // GPU_CULLING
+
+		if (bFarFieldInstance)
 		{
 			LocalToTranslatedWorld4x4[3].xyz += FarFieldReferencePos;
 		}
@@ -85,6 +145,7 @@ void RayTracingBuildInstanceBufferCS(uint3 DispatchThreadId : SV_DispatchThreadI
 		LocalToTranslatedWorld[2] = LocalToTranslatedWorld4x4[2];
 
 #if OUTPUT_INSTANCE_GPUSCENE_INDEX
+		if (InputDesc.AccelerationStructureIndex != 0xFFFFFFFF && !bCulled)
 		{
 			RWDebugInstanceGPUSceneIndices[InputDesc.OutputDescriptorIndex] = InputDesc.GPUSceneInstanceOrTransformIndex;
 		}
@@ -99,7 +160,7 @@ void RayTracingBuildInstanceBufferCS(uint3 DispatchThreadId : SV_DispatchThreadI
 
 		uint2 BlasAddress = 0;
 
-		if (InputDesc.AccelerationStructureIndex != 0xFFFFFFFF)
+		if (InputDesc.AccelerationStructureIndex != 0xFFFFFFFF && !bCulled)
 		{
 			BlasAddress = AccelerationStructureAddresses.Load2(InputDesc.AccelerationStructureIndex * 8);
 		}
diff --git a/Engine/Shaders/Shared/SceneDefinitions.h b/Engine/Shaders/Shared/SceneDefinitions.h
--- a/Engine/Shaders/Shared/SceneDefinitions.h
+++ b/Engine/Shaders/Shared/SceneDefinitions.h
@@ -53,6 +53,7 @@
 #define PRIMITIVE_SCENE_DATA_FLAG_INSTANCE_CULLING_OCCLUSION_QUERIES	0x1000000 // Whether per-instance occlusion queries are allowed when performing GPUScene instance culling
 #define PRIMITIVE_SCENE_DATA_FLAG_HAS_PIXEL_ANIMATION					0x2000000
 #define PRIMITIVE_SCENE_DATA_FLAG_RAYTRACING_FAR_FIELD					0x4000000
+#define PRIMITIVE_SCENE_DATA_FLAG_RAYTRACING_HAS_GROUPID				0x8000000
 
 
 // Primitive visibility flags
diff --git a/Engine/Source/Runtime/Engine/Private/PrimitiveSceneProxy.cpp b/Engine/Source/Runtime/Engine/Private/PrimitiveSceneProxy.cpp
--- a/Engine/Source/Runtime/Engine/Private/PrimitiveSceneProxy.cpp
+++ b/Engine/Source/Runtime/Engine/Private/PrimitiveSceneProxy.cpp
@@ -847,7 +847,8 @@ void FPrimitiveSceneProxy::BuildUniformShaderParameters(FPrimitiveUniformShaderP
 			.EditorColors(GetWireframeColor(), GetLevelColor())
 			.SplineMesh(IsSplineMesh())
 			.HasPixelAnimation(AnyMaterialHasPixelAnimation())
-			.RayTracingFarField(IsRayTracingFarField());
+			.RayTracingFarField(IsRayTracingFarField())
+			.RayTracingHasGroupId(GetRayTracingGroupId() != FPrimitiveSceneProxy::InvalidRayTracingGroupId);
 
 		if (PrimitiveSceneInfo != nullptr)
 		{
diff --git a/Engine/Source/Runtime/Engine/Public/PrimitiveUniformShaderParametersBuilder.h b/Engine/Source/Runtime/Engine/Public/PrimitiveUniformShaderParametersBuilder.h
--- a/Engine/Source/Runtime/Engine/Public/PrimitiveUniformShaderParametersBuilder.h
+++ b/Engine/Source/Runtime/Engine/Public/PrimitiveUniformShaderParametersBuilder.h
@@ -62,6 +62,7 @@ public:
 		bAllowInstanceCullingOcclusionQueries		= false;
 		bHasPixelAnimation                          = false;
 		bRayTracingFarField							= false;
+		bRayTracingHasGroupId						= false;
 
 		Parameters.MaxWPOExtent						= 0.0f;
 		Parameters.MinMaterialDisplacement			= 0.0f;
@@ -131,6 +132,7 @@ public:
 	PRIMITIVE_UNIFORM_BUILDER_FLAG_METHOD(bool,			HasAlwaysEvaluateWPOMaterials);
 	PRIMITIVE_UNIFORM_BUILDER_FLAG_METHOD(bool,			HasPixelAnimation);
 	PRIMITIVE_UNIFORM_BUILDER_FLAG_METHOD(bool,			RayTracingFarField);
+	PRIMITIVE_UNIFORM_BUILDER_FLAG_METHOD(bool,			RayTracingHasGroupId);
 
 	PRIMITIVE_UNIFORM_BUILDER_METHOD(uint32,			InstanceSceneDataOffset);
 	PRIMITIVE_UNIFORM_BUILDER_METHOD(uint32,			NumInstanceSceneDataEntries);
@@ -402,6 +404,7 @@ public:
 		Parameters.Flags |= bAllowInstanceCullingOcclusionQueries ? PRIMITIVE_SCENE_DATA_FLAG_INSTANCE_CULLING_OCCLUSION_QUERIES: 0u;
 		Parameters.Flags |= bHasPixelAnimation ? PRIMITIVE_SCENE_DATA_FLAG_HAS_PIXEL_ANIMATION : 0u;
 		Parameters.Flags |= bRayTracingFarField ? PRIMITIVE_SCENE_DATA_FLAG_RAYTRACING_FAR_FIELD : 0u;
+		Parameters.Flags |= bRayTracingHasGroupId ? PRIMITIVE_SCENE_DATA_FLAG_RAYTRACING_HAS_GROUPID : 0u;
 		
 		Parameters.VisibilityFlags = 0;
 		Parameters.VisibilityFlags |= bCastHiddenShadow ? PRIMITIVE_VISIBILITY_FLAG_CAST_HIDDEN_SHADOW : 0u;
@@ -467,4 +470,5 @@ private:
 	uint32 bAllowInstanceCullingOcclusionQueries : 1;
 	uint32 bHasPixelAnimation : 1;
 	uint32 bRayTracingFarField : 1;
+	uint32 bRayTracingHasGroupId : 1;
 };
diff --git a/Engine/Source/Runtime/Renderer/Private/RayTracing/RayTracing.cpp b/Engine/Source/Runtime/Renderer/Private/RayTracing/RayTracing.cpp
--- a/Engine/Source/Runtime/Renderer/Private/RayTracing/RayTracing.cpp
+++ b/Engine/Source/Runtime/Renderer/Private/RayTracing/RayTracing.cpp
@@ -939,7 +939,7 @@ namespace RayTracing
 							const Experimental::FHashElementId GroupId = Scene.PrimitiveRayTracingGroupIds[PrimitiveIndex];
 							const bool bUseGroupBounds = CullingParameters.bCullUsingGroupIds && GroupId.IsValid();
 
-							if (CullingParameters.CullingMode != RayTracing::ECullingMode::Disabled && GetRayTracingCullingPerInstance() && RelevantPrimitive.CachedRayTracingInstance->NumTransforms > 1 && !bUseGroupBounds)
+							if (CullingParameters.CullingMode != RayTracing::ECullingMode::Disabled && GetRayTracingCullingPerInstance() && RelevantPrimitive.CachedRayTracingInstance->NumTransforms > 1 && !bUseGroupBounds && !CullingParameters.bUseGPUInstanceCulling)
 							{
 								const bool bIsFarFieldPrimitive = EnumHasAnyFlags(Flags, ERayTracingPrimitiveFlags::FarField);
 
@@ -1131,17 +1131,20 @@ namespace RayTracing
 					}
 				}
 
-				CullingTasks.Add(FFunctionGraphTask::CreateAndDispatchWhenReady([CullInstancesClosures = MoveTemp(CullInstancesClosures)]()
-					{
-						for (auto& Closure : CullInstancesClosures)
+				if (!CullingParameters.bUseGPUInstanceCulling)
+				{
+					CullingTasks.Add(FFunctionGraphTask::CreateAndDispatchWhenReady([CullInstancesClosures = MoveTemp(CullInstancesClosures)]()
 						{
-							Closure();
-						}
-					}, TStatId(), nullptr, ENamedThreads::AnyThread));
+							for (auto& Closure : CullInstancesClosures)
+							{
+								Closure();
+							}
+						}, TStatId(), nullptr, ENamedThreads::AnyThread));
 
-				for (FGraphEventRef& CullingTask : CullingTasks)
-				{
-					MyCompletionGraphEvent->DontCompleteUntil(CullingTask);
+					for (FGraphEventRef& CullingTask : CullingTasks)
+					{
+						MyCompletionGraphEvent->DontCompleteUntil(CullingTask);
+					}
 				}
 			}
 		};
diff --git a/Engine/Source/Runtime/Renderer/Private/RayTracing/RayTracingInstanceBufferUtil.cpp b/Engine/Source/Runtime/Renderer/Private/RayTracing/RayTracingInstanceBufferUtil.cpp
--- a/Engine/Source/Runtime/Renderer/Private/RayTracing/RayTracingInstanceBufferUtil.cpp
+++ b/Engine/Source/Runtime/Renderer/Private/RayTracing/RayTracingInstanceBufferUtil.cpp
@@ -17,6 +17,8 @@
 #include "DataDrivenShaderPlatformInfo.h"
 #include "ShaderCompilerCore.h"
 
+#include "SceneRendering.h"
+
 #include "Async/ParallelFor.h"
 
 #include "Experimental/Containers/SherwoodHashTable.h"
@@ -301,13 +303,21 @@ struct FRayTracingBuildInstanceBufferCS : public FGlobalShader
 		SHADER_PARAMETER(FVector3f, ViewTilePosition)
 		SHADER_PARAMETER(FVector3f, RelativePreViewTranslation)
 
+		// Instance culling params
+		SHADER_PARAMETER(float, CullingRadius)
+		SHADER_PARAMETER(float, FarFieldCullingRadius)
+		SHADER_PARAMETER(float, AngleThresholdRatioSq)
+		SHADER_PARAMETER(FVector3f, ViewOrigin)
+		SHADER_PARAMETER(uint32, CullingMode)
+
 		// Debug parameters
 		SHADER_PARAMETER_UAV(RWStructuredBuffer<uint>, RWDebugInstanceGPUSceneIndices)
 	END_SHADER_PARAMETER_STRUCT()
 
 	class FUseGPUSceneDim : SHADER_PERMUTATION_BOOL("USE_GPUSCENE");
 	class FOutputInstanceGPUSceneIndexDim : SHADER_PERMUTATION_BOOL("OUTPUT_INSTANCE_GPUSCENE_INDEX");
-	using FPermutationDomain = TShaderPermutationDomain<FUseGPUSceneDim, FOutputInstanceGPUSceneIndexDim>;
+	class FGpuCullingDim : SHADER_PERMUTATION_BOOL("GPU_CULLING");
+	using FPermutationDomain = TShaderPermutationDomain<FUseGPUSceneDim, FOutputInstanceGPUSceneIndexDim, FGpuCullingDim>;
 		
 	static constexpr uint32 ThreadGroupSize = 64;
 
@@ -342,6 +352,7 @@ void BuildRayTracingInstanceBuffer(
 	FShaderResourceViewRHIRef InstanceUploadSRV,
 	FShaderResourceViewRHIRef AccelerationStructureAddressesSRV,
 	FShaderResourceViewRHIRef InstanceTransformSRV,
+	const FRayTracingCullingParameters* CullingParameters,
 	FUnorderedAccessViewRHIRef DebugInstanceGPUSceneIndexUAV)
 {
 	FRayTracingBuildInstanceBufferCS::FParameters PassParams;
@@ -363,11 +374,21 @@ void BuildRayTracingInstanceBuffer(
 		PassParams.GPUScenePrimitiveSceneData = GPUScene->PrimitiveBuffer->GetSRV();
 	}
 
+	if (CullingParameters)
+	{
+		PassParams.CullingRadius = CullingParameters->CullingRadius;
+		PassParams.FarFieldCullingRadius = CullingParameters->FarFieldCullingRadius;
+		PassParams.AngleThresholdRatioSq = CullingParameters->AngleThresholdRatioSq;
+		PassParams.ViewOrigin = CullingParameters->TranslatedViewOrigin;
+		PassParams.CullingMode = uint32(CullingParameters->CullingMode);
+	}
+
 	PassParams.RWDebugInstanceGPUSceneIndices = DebugInstanceGPUSceneIndexUAV;
 
 	FRayTracingBuildInstanceBufferCS::FPermutationDomain PermutationVector;
 	PermutationVector.Set<FRayTracingBuildInstanceBufferCS::FUseGPUSceneDim>(InstanceTransformSRV == nullptr);
 	PermutationVector.Set<FRayTracingBuildInstanceBufferCS::FOutputInstanceGPUSceneIndexDim>(DebugInstanceGPUSceneIndexUAV != nullptr);
+	PermutationVector.Set<FRayTracingBuildInstanceBufferCS::FGpuCullingDim>(CullingParameters != nullptr);
 
 	auto ComputeShader = GetGlobalShaderMap(GMaxRHIFeatureLevel)->GetShader<FRayTracingBuildInstanceBufferCS>(PermutationVector);
 	const int32 GroupSize = FMath::DivideAndRoundUp(PassParams.NumInstances, FRayTracingBuildInstanceBufferCS::ThreadGroupSize);
@@ -395,6 +416,7 @@ void BuildRayTracingInstanceBuffer(
 	uint32 NumNativeGPUSceneInstances,
 	uint32 NumNativeCPUInstances,
 	TConstArrayView<FRayTracingGPUInstance> GPUInstances,
+	const FRayTracingCullingParameters* CullingParameters,
 	FUnorderedAccessViewRHIRef DebugInstanceGPUSceneIndexUAV)
 {
 	if (NumNativeGPUSceneInstances > 0)
@@ -410,6 +432,7 @@ void BuildRayTracingInstanceBuffer(
 			InstanceUploadSRV,
 			AccelerationStructureAddressesSRV,
 			nullptr,
+			CullingParameters,
 			DebugInstanceGPUSceneIndexUAV);
 	}
 
@@ -426,6 +449,7 @@ void BuildRayTracingInstanceBuffer(
 			InstanceUploadSRV,
 			AccelerationStructureAddressesSRV,
 			CPUInstanceTransformSRV,
+			nullptr,
 			nullptr);
 	}
 
@@ -445,6 +469,7 @@ void BuildRayTracingInstanceBuffer(
 			InstanceUploadSRV,
 			AccelerationStructureAddressesSRV,
 			GPUInstance.TransformSRV,
+			CullingParameters,
 			DebugInstanceGPUSceneIndexUAV);
 	}
 }
diff --git a/Engine/Source/Runtime/Renderer/Private/RayTracing/RayTracingInstanceCulling.cpp b/Engine/Source/Runtime/Renderer/Private/RayTracing/RayTracingInstanceCulling.cpp
--- a/Engine/Source/Runtime/Renderer/Private/RayTracing/RayTracingInstanceCulling.cpp
+++ b/Engine/Source/Runtime/Renderer/Private/RayTracing/RayTracingInstanceCulling.cpp
@@ -47,6 +47,12 @@ static TAutoConsoleVariable<int32> CVarRayTracingCullingGroupIds(
 	TEXT("Cull using aggregate ray tracing group id bounds when defined instead of primitive or instance bounds."),
 	ECVF_RenderThreadSafe);
 
+static TAutoConsoleVariable<int32> CVarRayTracingCullingUseGPU(
+	TEXT("r.RayTracing.Culling.UseGPU"),
+	0,
+	TEXT("Cull instances using GPU. Requires per instance culling enabled. Experimental (default=0)"),
+	ECVF_RenderThreadSafe);
+
 RayTracing::ECullingMode RayTracing::GetCullingMode(const FEngineShowFlags& ShowFlags)
 {
 	// Disable culling if path tracer is used, so that path tracer matches raster view
@@ -77,12 +83,14 @@ void FRayTracingCullingParameters::Init(FViewInfo& View)
 	AngleThresholdRatio = FMath::Tan(FMath::Min(89.99f, CullAngleThreshold) * PI / 180.0f);
 	AngleThresholdRatioSq = FMath::Square(AngleThresholdRatio);
 	ViewOrigin = View.ViewMatrices.GetViewOrigin();
+	TranslatedViewOrigin = FVector3f(ViewOrigin + View.ViewMatrices.GetPreViewTranslation());
 	ViewDirection = View.GetViewDirection();
 	bCullAllObjects = CullingMode == RayTracing::ECullingMode::DistanceAndSolidAngle || CullingMode == RayTracing::ECullingMode::DistanceOrSolidAngle;
 	bCullByRadiusOrDistance = CullingMode == RayTracing::ECullingMode::DistanceOrSolidAngle;
 	bIsRayTracingFarField = Lumen::UseFarField(*View.Family);
 	bCullUsingGroupIds = CVarRayTracingCullingGroupIds.GetValueOnRenderThread() != 0;
 	bCullMinDrawDistance = CVarRayTracingCullingUseMinDrawDistance.GetValueOnRenderThread() != 0;
+	bUseGPUInstanceCulling = CVarRayTracingCullingUseGPU.GetValueOnRenderThread() != 0 && GetRayTracingCullingPerInstance() && bCullAllObjects;
 }
 
 namespace RayTracing
@@ -122,7 +130,7 @@ template<bool bCullByRadiusOrDistance>
 bool CullBounds(const FRayTracingCullingParameters& CullingParameters, const FBoxSphereBounds& RESTRICT ObjectBounds, float MinDrawDistance, bool bIsFarFieldPrimitive)
 {
 	const float ObjectRadius = ObjectBounds.SphereRadius;
-	const FVector ObjectCenter = ObjectBounds.Origin + 0.5 * ObjectBounds.BoxExtent;
+	const FVector ObjectCenter = ObjectBounds.Origin;
 	float CameraToObjectCenterLengthSq = FVector::DistSquared(ObjectCenter, CullingParameters.ViewOrigin);
 
 	if (bIsFarFieldPrimitive)
diff --git a/Engine/Source/Runtime/Renderer/Private/RayTracing/RayTracingScene.cpp b/Engine/Source/Runtime/Renderer/Private/RayTracing/RayTracingScene.cpp
--- a/Engine/Source/Runtime/Renderer/Private/RayTracing/RayTracingScene.cpp
+++ b/Engine/Source/Runtime/Renderer/Private/RayTracing/RayTracingScene.cpp
@@ -13,6 +13,7 @@
 #include "PrimitiveSceneProxy.h"
 #include "SceneUniformBuffer.h"
 #include "SceneRendering.h"
+#include "RayTracingInstanceCulling.h"
 
 BEGIN_SHADER_PARAMETER_STRUCT(FBuildInstanceBufferPassParams, )
 	SHADER_PARAMETER_RDG_BUFFER_UAV(RWStructuredBuffer, InstanceBuffer)
@@ -164,6 +165,20 @@ void FRayTracingScene::CreateWithInitializationData(FRDGBuilder& GraphBuilder, c
 		}
 	}
 
+	FRDGBufferUAVRef DebugInstanceGPUSceneIndexBufferUAV = nullptr;
+	if (bNeedsDebugInstanceGPUSceneIndexBuffer)
+	{
+		FRDGBufferDesc DebugInstanceGPUSceneIndexBufferDesc;
+		DebugInstanceGPUSceneIndexBufferDesc.Usage = EBufferUsageFlags::UnorderedAccess | EBufferUsageFlags::ShaderResource | EBufferUsageFlags::StructuredBuffer;
+		DebugInstanceGPUSceneIndexBufferDesc.BytesPerElement = sizeof(uint32);
+		DebugInstanceGPUSceneIndexBufferDesc.NumElements = FMath::Max(NumNativeInstances, 1u);
+
+		DebugInstanceGPUSceneIndexBuffer = GraphBuilder.CreateBuffer(DebugInstanceGPUSceneIndexBufferDesc, TEXT("FRayTracingScene::DebugInstanceGPUSceneIndexBuffer"));
+		DebugInstanceGPUSceneIndexBufferUAV = GraphBuilder.CreateUAV(DebugInstanceGPUSceneIndexBuffer);
+
+		AddClearUAVPass(GraphBuilder, DebugInstanceGPUSceneIndexBufferUAV, 0xFFFFFFFF);
+	}
+
 	if (NumNativeInstances > 0)
 	{
 		const uint32 InstanceUploadBytes = NumNativeInstances * sizeof(FRayTracingInstanceDescriptorInput);
@@ -205,22 +220,9 @@ void FRayTracingScene::CreateWithInitializationData(FRDGBuilder& GraphBuilder, c
 
 		FBuildInstanceBufferPassParams* PassParams = GraphBuilder.AllocParameters<FBuildInstanceBufferPassParams>();
 		PassParams->InstanceBuffer = GraphBuilder.CreateUAV(InstanceBuffer);
-		PassParams->DebugInstanceGPUSceneIndexBuffer = nullptr;
+		PassParams->DebugInstanceGPUSceneIndexBuffer = DebugInstanceGPUSceneIndexBufferUAV;
 		PassParams->Scene = View.GetSceneUniforms().GetBuffer(GraphBuilder);
 
-		if (bNeedsDebugInstanceGPUSceneIndexBuffer)
-		{
-			FRDGBufferDesc DebugInstanceGPUSceneIndexBufferDesc;
-			DebugInstanceGPUSceneIndexBufferDesc.Usage = EBufferUsageFlags::UnorderedAccess | EBufferUsageFlags::ShaderResource | EBufferUsageFlags::StructuredBuffer;
-			DebugInstanceGPUSceneIndexBufferDesc.BytesPerElement = sizeof(uint32);
-			DebugInstanceGPUSceneIndexBufferDesc.NumElements = NumNativeInstances;
-
-			DebugInstanceGPUSceneIndexBuffer = GraphBuilder.CreateBuffer(DebugInstanceGPUSceneIndexBufferDesc, TEXT("FRayTracingScene::DebugInstanceGPUSceneIndexBuffer"));
-			PassParams->DebugInstanceGPUSceneIndexBuffer = GraphBuilder.CreateUAV(DebugInstanceGPUSceneIndexBuffer);
-
-			AddClearUAVPass(GraphBuilder, PassParams->DebugInstanceGPUSceneIndexBuffer, 0xFFFFFFFF);
-		}
-
 		GraphBuilder.AddPass(
 			RDG_EVENT_NAME("BuildTLASInstanceBuffer"),
 			PassParams,
@@ -233,7 +235,8 @@ void FRayTracingScene::CreateWithInitializationData(FRDGBuilder& GraphBuilder, c
 			&SceneInitializer,
 			NumNativeGPUSceneInstances = SceneWithGeometryInstances.NumNativeGPUSceneInstances,
 			NumNativeCPUInstances = SceneWithGeometryInstances.NumNativeCPUInstances,
-			GPUInstances = MoveTemp(SceneWithGeometryInstances.GPUInstances)
+			GPUInstances = MoveTemp(SceneWithGeometryInstances.GPUInstances),
+			CullingParameters = View.RayTracingCullingParameters
 			](FRHICommandListImmediate& RHICmdList)
 			{
 				WaitForTasks();
@@ -281,6 +284,7 @@ void FRayTracingScene::CreateWithInitializationData(FRDGBuilder& GraphBuilder, c
 					NumNativeGPUSceneInstances,
 					NumNativeCPUInstances,
 					GPUInstances,
+					CullingParameters.bUseGPUInstanceCulling ? &CullingParameters : nullptr,
 					PassParams->DebugInstanceGPUSceneIndexBuffer ? PassParams->DebugInstanceGPUSceneIndexBuffer->GetRHI() : nullptr);
 			});
 	}
diff --git a/Engine/Source/Runtime/Renderer/Private/SceneRendering.h b/Engine/Source/Runtime/Renderer/Private/SceneRendering.h
--- a/Engine/Source/Runtime/Renderer/Private/SceneRendering.h
+++ b/Engine/Source/Runtime/Renderer/Private/SceneRendering.h
@@ -1168,11 +1168,13 @@ struct FRayTracingCullingParameters
 	float AngleThresholdRatioSq;
 	FVector ViewOrigin;
 	FVector ViewDirection;
+	FVector3f TranslatedViewOrigin;
 	bool bCullAllObjects;
 	bool bCullByRadiusOrDistance;
 	bool bIsRayTracingFarField;
 	bool bCullUsingGroupIds;
 	bool bCullMinDrawDistance;
+	bool bUseGPUInstanceCulling;
 
 	void Init(FViewInfo& View);
 };
diff --git a/Engine/Source/Runtime/Renderer/Private/Tests/RayTracingTestbed.cpp b/Engine/Source/Runtime/Renderer/Private/Tests/RayTracingTestbed.cpp
--- a/Engine/Source/Runtime/Renderer/Private/Tests/RayTracingTestbed.cpp
+++ b/Engine/Source/Runtime/Renderer/Private/Tests/RayTracingTestbed.cpp
@@ -271,6 +271,7 @@ bool RunRayTracingTestbed_RenderThread(const FString& Parameters)
 		RayTracingScene.NumNativeGPUSceneInstances,
 		RayTracingScene.NumNativeCPUInstances,
 		{},
+		nullptr,
 		nullptr);
 
 	RHICmdList.BindAccelerationStructureMemory(RayTracingScene.Scene, SceneBuffer, 0);
diff --git a/Engine/Source/Runtime/Renderer/Public/RayTracingInstanceBufferUtil.h b/Engine/Source/Runtime/Renderer/Public/RayTracingInstanceBufferUtil.h
--- a/Engine/Source/Runtime/Renderer/Public/RayTracingInstanceBufferUtil.h
+++ b/Engine/Source/Runtime/Renderer/Public/RayTracingInstanceBufferUtil.h
@@ -8,6 +8,7 @@
 #if RHI_RAYTRACING
 
 class FGPUScene;
+struct FRayTracingCullingParameters;
 
 /*
 * 
@@ -95,6 +96,7 @@ RENDERER_API void BuildRayTracingInstanceBuffer(
 	uint32 NumNativeGPUSceneInstances,
 	uint32 NumNativeCPUInstances,
 	TConstArrayView<FRayTracingGPUInstance> GPUInstances,
+	const FRayTracingCullingParameters* CullingParameters,
 	FUnorderedAccessViewRHIRef DebugInstanceGPUSceneIndexUAV);
 
 #endif // RHI_RAYTRACING
\ No newline at end of file
