Workaround crash resizing LightDataBuffer due to MemcpyResource only supporting buffers with BytesPerElement = 4,8,16,32.
- switched back to TRefCountPtr<FRDGPooledBuffer>

diff --git a/Engine/Source/Runtime/Renderer/Private/GPUScene.cpp b/Engine/Source/Runtime/Renderer/Private/GPUScene.cpp
--- a/Engine/Source/Runtime/Renderer/Private/GPUScene.cpp
+++ b/Engine/Source/Runtime/Renderer/Private/GPUScene.cpp
@@ -615,7 +615,6 @@ FGPUScene::FGPUScene(FScene &InScene)
 	, InstanceSceneDataSOAStride(0)
 	, InstancePayloadDataAllocator(CVarGPUSceneUseGrowOnlyAllocationPolicy.GetValueOnAnyThread() != 0)
 	, LightmapDataAllocator(CVarGPUSceneUseGrowOnlyAllocationPolicy.GetValueOnAnyThread() != 0)
-	, LightDataBuffer(InitialBufferSize, TEXT("GPUScene.LightData"))
 	, Scene(InScene)
 	, InstanceSceneDataAllocator(CVarGPUSceneUseGrowOnlyAllocationPolicy.GetValueOnAnyThread() != 0)
 {
@@ -707,7 +706,33 @@ void FGPUScene::UpdateGPULights(FRDGBuilder& GraphBuilder, const UE::Tasks::FTas
 
 	}, PrerequisiteTask);
 
-	FRDGBuffer* LightDataBufferRDG = LightDataBuffer.ResizeBufferIfNeeded(GraphBuilder, Scene.Lights.GetMaxIndex());
+	const uint32 LightDataBufferNumElements = FMath::RoundUpToPowerOfTwo(FMath::Max(Scene.Lights.GetMaxIndex(), InitialBufferSize));
+	const FRDGBufferDesc LightDataBufferDesc = FRDGBufferDesc::CreateStructuredDesc(sizeof(FLightSceneData), LightDataBufferNumElements);
+
+	FRDGBuffer* LightDataBufferRDG;
+	if (LightDataBuffer == nullptr)
+	{
+		LightDataBufferRDG = GraphBuilder.CreateBuffer(LightDataBufferDesc, TEXT("GPUScene.LightData"));
+		LightDataBuffer = GraphBuilder.ConvertToExternalBuffer(LightDataBufferRDG);
+	}
+	else
+	{
+		const uint32 BufferSizeNew = LightDataBufferDesc.GetSize();
+		const uint32 BufferSizeOld = LightDataBuffer->GetCommittedSize();
+
+		if (BufferSizeNew == BufferSizeOld)
+		{
+			LightDataBufferRDG = GraphBuilder.RegisterExternalBuffer(LightDataBuffer);
+		}
+		else
+		{
+			LightDataBufferRDG = GraphBuilder.CreateBuffer(LightDataBufferDesc, TEXT("GPUScene.LightData"));
+			LightDataBuffer = GraphBuilder.ConvertToExternalBuffer(LightDataBufferRDG);
+
+			// Currently whole contents of LightDataBuffer are uploaded so don't need to copy previous contents when resizing
+		}
+
+	}
 
 	GraphBuilder.QueueBufferUpload<FLightSceneData>(LightDataBufferRDG, LightData, ERDGInitialDataFlags::NoCopy);
 }
@@ -892,7 +917,7 @@ FGPUScene::FRegisteredBuffers FGPUScene::UpdateBufferAllocations(FRDGBuilder& Gr
 	const uint32 LightMapDataBufferSize = FMath::RoundUpToPowerOfTwo(FMath::Max(LightmapDataAllocator.GetMaxSize(), InitialBufferSize));
 	BufferState.LightmapDataBuffer = ResizeStructuredBufferIfNeeded(GraphBuilder, LightmapDataBuffer, LightMapDataBufferSize * sizeof(FLightmapSceneShaderData::Data), TEXT("GPUScene.LightmapData"));
 
-	BufferState.LightDataBuffer = LightDataBuffer.Register(GraphBuilder);
+	BufferState.LightDataBuffer = GraphBuilder.RegisterExternalBuffer(LightDataBuffer);
 
 	FGPUSceneResourceParameters ShaderParameters;
 	ShaderParameters.GPUSceneInstanceSceneData = GraphBuilder.CreateSRV(BufferState.InstanceSceneDataBuffer);
@@ -928,7 +953,7 @@ FGPUScene::FRegisteredBuffers FGPUScene::RegisterBuffers(FRDGBuilder& GraphBuild
 	Result.InstanceSceneDataBuffer = GraphBuilder.RegisterExternalBuffer(InstanceSceneDataBuffer);
 	Result.InstancePayloadDataBuffer = GraphBuilder.RegisterExternalBuffer(InstancePayloadDataBuffer);
 	Result.LightmapDataBuffer = GraphBuilder.RegisterExternalBuffer(LightmapDataBuffer);
-	Result.LightDataBuffer = LightDataBuffer.Register(GraphBuilder);
+	Result.LightDataBuffer = GraphBuilder.RegisterExternalBuffer(LightDataBuffer);
 
 	return Result;
 }
@@ -2175,7 +2200,7 @@ void FGPUScene::OnPostLightSceneInfoUpdate(FRDGBuilder& GraphBuilder, const FLig
 	// SceneFrameNumber is updated in UpdateInternal so if it is the same, we have already uploaded the lights this "frame" - this is not 100% robust since 
 	// if UpdateAllPrimitiveSceneInfos is called multiple times without updating the FScene frame number it will skip light uploads. The real solution is 
 	// to implement functional change tracking for lights so we can actually know when there are relevant changes.
-	if (SceneFrameNumber != Scene.GetFrameNumberRenderThread() || !LightDataBuffer.GetPooledBuffer().IsValid())
+	if (SceneFrameNumber != Scene.GetFrameNumberRenderThread() || !LightDataBuffer.IsValid())
 	{
 		UpdateGPULights(GraphBuilder, UE::Tasks::FTask{});
 	}
diff --git a/Engine/Source/Runtime/Renderer/Private/GPUScene.h b/Engine/Source/Runtime/Renderer/Private/GPUScene.h
--- a/Engine/Source/Runtime/Renderer/Private/GPUScene.h
+++ b/Engine/Source/Runtime/Renderer/Private/GPUScene.h
@@ -13,7 +13,6 @@
 #include "LightSceneData.h"
 #include "SceneUniformBuffer.h"
 #include "UnifiedBuffer.h"
-#include "RendererPrivateUtils.h"
 
 class FRDGExternalAccessQueue;
 class FRHICommandList;
@@ -334,7 +333,7 @@ public:
 private:
 	static constexpr int32 InitialBufferSize = 256;
 
-	TPersistentStructuredBuffer<FLightSceneData> LightDataBuffer;
+	TRefCountPtr<FRDGPooledBuffer> LightDataBuffer;
 
 	// Buffers used by GPU-Scene, since they can be resized during updates AND the render passes must retain the 
 	// right copy (this is chiefly because the init of shadow views after pre-pass means we need to be able to set 
