LFAtlas - work with substrate for SLW and Translucent.

diff --git a/Engine/Shaders/Private/Substrate/SubstrateLightingCommon.ush b/Engine/Shaders/Private/Substrate/SubstrateLightingCommon.ush
--- a/Engine/Shaders/Private/Substrate/SubstrateLightingCommon.ush
+++ b/Engine/Shaders/Private/Substrate/SubstrateLightingCommon.ush
@@ -10,6 +10,10 @@
 #if SUBSTRATE_ENABLED
 #ifdef USE_SUBSTRATE_FORWARD_LIGHTING_COMMON
 
+#if USE_LIGHT_FUNCTION_ATLAS
+#include "../LightFunctionAtlas/LightFunctionAtlasCommon.usf"
+#endif
+
 // Used by non-mobile path
 half4 GetPrecomputedShadowFactors(FSubstratePixelHeader SubstratePixelHeader, float3 TranslatedWorldPosition)
 {
@@ -136,7 +140,12 @@ float3 SubstrateForwardLightingCommon(
 			{
 				float3 DiffuseLuminance = BSDFEvaluate.IntegratedDiffuseValue;
 				float3 SpecularLuminance = BSDFEvaluate.IntegratedSpecularValue * LightData.SpecularScale;
-				const float3 LightCommonMultiplier = LightData.Color * LightMask;
+				float3 LightCommonMultiplier = LightData.Color * LightMask;
+
+			#if USE_LIGHT_FUNCTION_ATLAS
+				LightCommonMultiplier *= GetLocalLightFunctionCommon(TranslatedWorldPosition, LightData.LightFunctionAtlasLightIndex);
+			#endif
+
 				Color += (DiffuseLuminance + SpecularLuminance) * LightCommonMultiplier * BSDFColoredVisibility;
 			}
 		#endif
diff --git a/Engine/Source/Runtime/Renderer/Private/LightFunctionAtlas.cpp b/Engine/Source/Runtime/Renderer/Private/LightFunctionAtlas.cpp
--- a/Engine/Source/Runtime/Renderer/Private/LightFunctionAtlas.cpp
+++ b/Engine/Source/Runtime/Renderer/Private/LightFunctionAtlas.cpp
@@ -35,10 +35,6 @@ To avoid allocating SRV when using the atlas we use a constant buffer to store a
 	- A single SRV is used, being the atlas Texture2D.
 	- Otherwise a single constant buffer entry is used for each view, storing: AtlasSlotIndex=>{SubUVs} read for a Light LightIndex=>{FadeParams, AtlasSlotIndex, TranslatedWorlViewProjectionMatrix}
 
-Systems currently using the LightFunctionAtlas:
- - Volumetric fog to inject local lights scattering using the light grid (when shadow casting is not enabled)
- - Deferred lighting: applied to all lights of a scene stored in the light grid, preventing them from going the slow un-bactched route.
-
 How to use the Atlas:
  - Add FLightFunctionAtlasGlobalParameters to your shader
  - #include "LightFunctionAtlas/LightFunctionAtlasCommon.usf" 
@@ -47,9 +43,8 @@ How to use the Atlas:
  What is next:
  - Super sample CVAR
  - Convert systems:
-    - Lumen
     - Path tracer
-    - Check those path with Substrate
+	- Ray tracing?
 
 */
 
