LFAtlas - fixed system addedto atlas not being reset when toggling cvars.

diff --git a/Engine/Source/Runtime/Renderer/Private/LightFunctionAtlas.cpp b/Engine/Source/Runtime/Renderer/Private/LightFunctionAtlas.cpp
--- a/Engine/Source/Runtime/Renderer/Private/LightFunctionAtlas.cpp
+++ b/Engine/Source/Runtime/Renderer/Private/LightFunctionAtlas.cpp
@@ -250,7 +250,7 @@ void FLightFunctionAtlas::ClearEmptySceneFrame(FViewInfo* View, uint32 ViewIndex
 
 void FLightFunctionAtlas::BeginSceneFrame(FViewFamilyInfo& ViewFamily, TArray<FViewInfo>& Views, FLightFunctionAtlasSceneData& LightFunctionAtlasSceneData, bool bShouldRenderVolumetricFog)
 {
-	ClearEmptySceneFrame();
+	ClearEmptySceneFrame(nullptr, 0, &LightFunctionAtlasSceneData);
 
 	// Now lets check if we need to generate the atlas for this frame
 	bLightFunctionAtlasEnabled = CVarLightFunctionAtlas.GetValueOnRenderThread() > 0 && ViewFamily.EngineShowFlags.LightFunctions > 0;
