LFAtlas - now using a buffer of float4 instead of constant buffer for unlimited amount of light.

diff --git a/Engine/Shaders/Private/LightFunctionAtlas/LightFunctionAtlasCommon.usf b/Engine/Shaders/Private/LightFunctionAtlas/LightFunctionAtlasCommon.usf
--- a/Engine/Shaders/Private/LightFunctionAtlas/LightFunctionAtlasCommon.usf
+++ b/Engine/Shaders/Private/LightFunctionAtlas/LightFunctionAtlasCommon.usf
@@ -31,7 +31,9 @@ FLightFunctionColor GetLocalLightFunctionCommon(float3 TranslatedWorldPosition,
 		return 1.0f;
 	}
 
-	float4 LightFunctionLightDataParameters = LightFunctionAtlas.LightInfoDataParameters[LightFunctionAtlasLightIndex];
+	const uint BaseBufferIndex = LightFunctionAtlasLightIndex * 5; // See FAtlasLightInfoData
+
+	float4 LightFunctionLightDataParameters = LightFunctionAtlas.LightInfoDataBuffer[BaseBufferIndex + 0];
 	const float  FadeDistance				= LightFunctionLightDataParameters.x;
 	const uint   PackedLightInfoDataParams	= asuint(LightFunctionLightDataParameters.y);
 	const float2 Slot_MinUV					= LightFunctionLightDataParameters.zw;
@@ -39,7 +41,11 @@ FLightFunctionColor GetLocalLightFunctionCommon(float3 TranslatedWorldPosition,
 	const uint   LightType					= PackedLightInfoDataParams & 0xFF;
 	const float  DisabledBrightness			= f16tof32((PackedLightInfoDataParams >> 8) & 0xFFFF);
 
-	float4x4 ActualLightFunctionMatrix = LightFunctionAtlas.LightInfoDataMatrix[LightFunctionAtlasLightIndex];
+	float4x4 ActualLightFunctionMatrix = float4x4(
+		LightFunctionAtlas.LightInfoDataBuffer[BaseBufferIndex + 1],
+		LightFunctionAtlas.LightInfoDataBuffer[BaseBufferIndex + 2],
+		LightFunctionAtlas.LightInfoDataBuffer[BaseBufferIndex + 3],
+		LightFunctionAtlas.LightInfoDataBuffer[BaseBufferIndex + 4]);
 	float4 HomogeneousShadowPosition = mul(float4(TranslatedWorldPosition, 1), ActualLightFunctionMatrix);
 
 	float2 LightFunctionUV;
diff --git a/Engine/Source/Runtime/Renderer/Private/LightFunctionAtlas.cpp b/Engine/Source/Runtime/Renderer/Private/LightFunctionAtlas.cpp
--- a/Engine/Source/Runtime/Renderer/Private/LightFunctionAtlas.cpp
+++ b/Engine/Source/Runtime/Renderer/Private/LightFunctionAtlas.cpp
@@ -68,6 +68,12 @@ static TAutoConsoleVariable<int32> CVarLightFunctionAtlasSize(
 	TEXT("Experimental: The default size in atlas slot count of the edge of the 2D texture atlas."),
 	ECVF_RenderThreadSafe);
 
+static TAutoConsoleVariable<int32> CVarLightFunctionAtlasMaxLightCount(
+	TEXT("r.LightFunctionAtlas.MaxLightCount"),
+	-1,
+	TEXT("Experimental: Clamp the number of lights that can sample light function atlas. -1 means unlimited light count."),
+	ECVF_RenderThreadSafe);
+
 
 
 //////////////////////////////////////////////////////////////////////////
@@ -373,7 +379,7 @@ void FLightFunctionAtlas::AllocateAtlasSlots(const TArray<FViewInfo>& Views)
 	EffectiveLightFunctionSlotArray.Reserve(AtlasMaxLightFunctionCount);
 
 	EffectiveLocalLightSlotArray.Reset();
-	EffectiveLocalLightSlotArray.Reserve(FMath::Min(LIGHT_FUNCTION_ATLAS_MAX_LIGHT_COUNT, SortedRegisteredLights.Num()));
+	EffectiveLocalLightSlotArray.Reserve(SortedRegisteredLights.Num());
 
 	LightFunctionsSet.Reset();
 	LightFunctionsSet.Reserve(AtlasMaxLightFunctionCount);
@@ -413,7 +419,6 @@ void FLightFunctionAtlas::AllocateAtlasSlots(const TArray<FViewInfo>& Views)
 	uint32 LocalLightWithLightFunctionCount = 0;
 	auto AddLightSlot = [&](FLightSceneInfo* LightSceneInfo, uint8 LightFunctionAtlasSlotIndex)
 	{
-		check(LocalLightWithLightFunctionCount < LIGHT_FUNCTION_ATLAS_MAX_LIGHT_COUNT);
 		EffectiveLocalLightSlot& NewLightSlot = EffectiveLocalLightSlotArray.Emplace_GetRef();
 		NewLightSlot.LightSceneInfo = LightSceneInfo;
 		NewLightSlot.LightFunctionAtlasSlotIndex = LightFunctionAtlasSlotIndex;
@@ -425,12 +430,13 @@ void FLightFunctionAtlas::AllocateAtlasSlots(const TArray<FViewInfo>& Views)
 	// Add the default invalid light slot at the beginning
 	AddLightSlot(nullptr, 0);
 
+	int32 MaxLightCount = CVarLightFunctionAtlasMaxLightCount.GetValueOnRenderThread();
 	for (FSortedRegisteredLights& SortedRegisteredLight : SortedRegisteredLights)
 	{
 		FLightSceneInfo* LightSceneInfo = RegisteredLights[SortedRegisteredLight.RegisteredLightIndex];
 		FLightSceneProxy* Proxy = LightSceneInfo->Proxy;
 
-		if (LocalLightWithLightFunctionCount >= LIGHT_FUNCTION_ATLAS_MAX_LIGHT_COUNT)
+		if (MaxLightCount >= 0 && int32(LocalLightWithLightFunctionCount) >= MaxLightCount)
 		{
 			// We cannot register anymore light, so set them to no light function
 			Proxy->SetLightFunctionAtlasIndices(0);
@@ -499,6 +505,7 @@ FLightFunctionAtlasGlobalParameters* FLightFunctionAtlas::GetDefaultLightFunctio
 	static FLightFunctionAtlasGlobalParameters DefaultLightFunctionAtlasGlobalParameters;
 	DefaultLightFunctionAtlasGlobalParameters.LightFunctionAtlasTexture = GSystemTextures.GetWhiteDummy(GraphBuilder);
 	DefaultLightFunctionAtlasGlobalParameters.LightFunctionAtlasSampler = TStaticSamplerState<SF_Bilinear>::GetRHI();
+	DefaultLightFunctionAtlasGlobalParameters.LightInfoDataBuffer = GraphBuilder.CreateSRV(GSystemTextures.GetDefaultStructuredBuffer(GraphBuilder, sizeof(FVector4f) * 1, 0.0f), PF_A32B32G32R32F);
 	DefaultLightFunctionAtlasGlobalParameters.Slot_UVSize = 1.0f;
 	return &DefaultLightFunctionAtlasGlobalParameters;
 }
@@ -547,6 +554,9 @@ void FLightFunctionAtlas::RenderLightFunctionAtlas(FRDGBuilder& GraphBuilder, TA
 		const FMatrix TranslatedWorldToWorld = FTranslationMatrix(-View.ViewMatrices.GetPreViewTranslation());
 
 		// Write the light data needed to rotate and fade the light function in the world.
+		const uint32 InitialLightInfoDataLightCount = FMath::DivideAndRoundUp(uint32(EffectiveLocalLightSlotArray.Num()), 32u) * 32u;// Alloted with 32 lights step to better reuse shared buffers pool.
+		const uint32 InitialLightInfoDataSize = InitialLightInfoDataLightCount * sizeof(FAtlasLightInfoData);
+		FAtlasLightInfoData* LightInfoDataBufferPtr = (FAtlasLightInfoData*)GraphBuilder.Alloc(InitialLightInfoDataSize, 16);
 		uint32 LightIndex = 0;
 		for (EffectiveLocalLightSlot& LightSlot : EffectiveLocalLightSlotArray)
 		{
@@ -555,7 +565,7 @@ void FLightFunctionAtlas::RenderLightFunctionAtlas(FRDGBuilder& GraphBuilder, TA
 
 			if (bIsDefaultSlot)
 			{
-				LightFunctionAtlasGlobalParameters->LightInfoDataParameters[LightIndex] = FVector4f(1.0f, 1.0f, 1.0f, 0.0f);
+				LightInfoDataBufferPtr[LightIndex].Parameters = FVector4f(1.0f, 1.0f, 1.0f, 0.0f);
 				LightIndex++;
 				continue;
 			}
@@ -586,11 +596,11 @@ void FLightFunctionAtlas::RenderLightFunctionAtlas(FRDGBuilder& GraphBuilder, TA
 				//ShadowFadeFraction = ProjectedShadowInfo.FadeAlphas.Num() == 0 ? 1.0f : ProjectedShadowInfo.FadeAlphas[0];
 
 				FVector4f ShadowmapMinMaxValue;
-				LightFunctionAtlasGlobalParameters->LightInfoDataMatrix[LightIndex] = FMatrix44f(TranslatedWorldToWorld * ProjectedShadowInfo.GetWorldToShadowMatrix(ShadowmapMinMaxValue, &LightFunctionResolution));
+				LightInfoDataBufferPtr[LightIndex].Transform = FMatrix44f(TranslatedWorldToWorld * ProjectedShadowInfo.GetWorldToShadowMatrix(ShadowmapMinMaxValue, &LightFunctionResolution));
 			}
 			else if (LightType == LightType_Point)
 			{
-				LightFunctionAtlasGlobalParameters->LightInfoDataMatrix[LightIndex] = FMatrix44f(TranslatedWorldToWorld * Proxy->GetWorldToLight());
+				LightInfoDataBufferPtr[LightIndex].Transform = FMatrix44f(TranslatedWorldToWorld * Proxy->GetWorldToLight());
 			}
 			else if (LightType == LightType_Rect)
 			{
@@ -599,7 +609,7 @@ void FLightFunctionAtlas::RenderLightFunctionAtlas(FRDGBuilder& GraphBuilder, TA
 				const FVector InverseScale = FVector(1.0 / Scale.Z, 1.0 / Scale.Y, 1.0 / Scale.X);
 				const FMatrix WorldToLight = Proxy->GetWorldToLight() * FScaleMatrix(InverseScale);
 
-				LightFunctionAtlasGlobalParameters->LightInfoDataMatrix[LightIndex] = FMatrix44f(TranslatedWorldToWorld * WorldToLight);
+				LightInfoDataBufferPtr[LightIndex].Transform = FMatrix44f(TranslatedWorldToWorld * WorldToLight);
 			}
 			else if (LightType == LightType_Directional)
 			{
@@ -610,7 +620,7 @@ void FLightFunctionAtlas::RenderLightFunctionAtlas(FRDGBuilder& GraphBuilder, TA
 				const FVector InverseScale = FVector(1.0 / Scale.Z, 1.0 / Scale.Y, 1.0 / Scale.X);
 
 				const FMatrix WorldToLight = LightSceneInfo->Proxy->GetWorldToLight() * FScaleMatrix(FVector(InverseScale));
-				LightFunctionAtlasGlobalParameters->LightInfoDataMatrix[LightIndex] = FMatrix44f(TranslatedWorldToWorld * WorldToLight);
+				LightInfoDataBufferPtr[LightIndex].Transform = FMatrix44f(TranslatedWorldToWorld * WorldToLight);
 			}
 			
 			const uint8 LightFunctionAtlasSlotIndex = LightSlot.LightFunctionAtlasSlotIndex;
@@ -620,11 +630,18 @@ void FLightFunctionAtlas::RenderLightFunctionAtlas(FRDGBuilder& GraphBuilder, TA
 			EffectiveLightFunctionSlot& AtlasSlot = EffectiveLightFunctionSlotArray[LightFunctionAtlasSlotIndex];
 
 			// ShadowFadeFraction is unused.
-			LightFunctionAtlasGlobalParameters->LightInfoDataParameters[LightIndex] = FVector4f(Proxy->GetLightFunctionFadeDistance(), FMath::AsFloat(PackedLightInfoDataParams), AtlasSlot.MinU, AtlasSlot.MinV);
+			LightInfoDataBufferPtr[LightIndex].Parameters = FVector4f(Proxy->GetLightFunctionFadeDistance(), FMath::AsFloat(PackedLightInfoDataParams), AtlasSlot.MinU, AtlasSlot.MinV);
 
 			LightIndex++;
 		}
 
+		// Create the light instance data buffer SRV
+		RDGLightInfoDataBuffer = CreateStructuredBuffer(
+			GraphBuilder, TEXT("LightFunctionAtlasLightInfoData"), 
+			sizeof(FAtlasLightInfoData), InitialLightInfoDataLightCount,
+			reinterpret_cast<void*>(LightInfoDataBufferPtr), InitialLightInfoDataSize, ERDGInitialDataFlags::NoCopy);
+		LightFunctionAtlasGlobalParameters->LightInfoDataBuffer = GraphBuilder.CreateSRV(RDGLightInfoDataBuffer, PF_A32B32G32R32F);
+
 		ViewLightFunctionAtlasGlobalParametersArray.Add(LightFunctionAtlasGlobalParameters);
 		ViewLightFunctionAtlasGlobalParametersUBArray.Add(GraphBuilder.CreateUniformBuffer(LightFunctionAtlasGlobalParameters));
 	}
diff --git a/Engine/Source/Runtime/Renderer/Private/LightFunctionAtlas.h b/Engine/Source/Runtime/Renderer/Private/LightFunctionAtlas.h
--- a/Engine/Source/Runtime/Renderer/Private/LightFunctionAtlas.h
+++ b/Engine/Source/Runtime/Renderer/Private/LightFunctionAtlas.h
@@ -103,22 +103,20 @@ struct FLightFunctionSlot
 
 #define LIGHT_FUNCTION_ATLAS_MAX_LIGHT_FUNCTION_COUNT 256
 
-/**
- * On devices using OpenGL, we are limited to 256 vec4s. 
- * Light parameters take 5 vec4s. So we can fit at maximum of 51.2 light info data. We round it down to 48 lights.
- * This is already plenty. And we can move to use a buffer later on.
- */
-#define LIGHT_FUNCTION_ATLAS_MAX_LIGHT_COUNT 48
+struct FAtlasLightInfoData
+{
+	FVector4f  Parameters;
+	FMatrix44f Transform;
+};
 
 // This allows to not have to store more data per GPU light representation on GPU. The light only needs an index into the array.
 // Using a constant buffer also workaround the fact that we would otherwise need another SRV in forward shaders. 
 // The light atlas texture itself already use 1 extra SRV. We could an extra SRV and have LightInfoDataXXX be in a buffer that scale with amount of light in the scene.
 BEGIN_GLOBAL_SHADER_PARAMETER_STRUCT(FLightFunctionAtlasGlobalParameters, )
 	SHADER_PARAMETER_RDG_TEXTURE(Texture2D<float4>, LightFunctionAtlasTexture)
+	SHADER_PARAMETER_RDG_BUFFER_SRV(StructuredBuffer<float4>, LightInfoDataBuffer)
 	SHADER_PARAMETER_SAMPLER(SamplerState, LightFunctionAtlasSampler)
 	SHADER_PARAMETER(float, Slot_UVSize)
-	SHADER_PARAMETER_ARRAY(FMatrix44f, LightInfoDataMatrix, [LIGHT_FUNCTION_ATLAS_MAX_LIGHT_COUNT])		// Light data
-	SHADER_PARAMETER_ARRAY(FVector4f,  LightInfoDataParameters, [LIGHT_FUNCTION_ATLAS_MAX_LIGHT_COUNT]) // Light data
 END_GLOBAL_SHADER_PARAMETER_STRUCT()
 
 
@@ -161,6 +159,7 @@ private:
 	bool bLightFunctionAtlasEnabled = false;
 
 	FRDGTextureRef RDGAtlasTexture2D = nullptr;
+	FRDGBufferRef RDGLightInfoDataBuffer = nullptr;
 
 	// All the lights that wants to sample light functions
 	TArray<FLightSceneInfo*> RegisteredLights;
@@ -189,7 +188,7 @@ private:
 	FLightFunctionAtlasGlobalParameters*								DefaultLightFunctionAtlasGlobalParameters = nullptr;
 	TRDGUniformBufferRef<FLightFunctionAtlasGlobalParameters>			DefaultLightFunctionAtlasGlobalParametersUB;
 
-	TArray<FLightFunctionAtlasGlobalParameters*>							ViewLightFunctionAtlasGlobalParametersArray;
+	TArray<FLightFunctionAtlasGlobalParameters*>						ViewLightFunctionAtlasGlobalParametersArray;
 	TArray<TRDGUniformBufferRef<FLightFunctionAtlasGlobalParameters>>	ViewLightFunctionAtlasGlobalParametersUBArray;
 };
 
