Added cvar to be able to ignore primitive lighting channels in a mobile deferred shading SkyLighting.
This allows to merge shading and skylight pass improving performance at the cost of of wrong lighting on such primitives

diff --git a/Engine/Source/Runtime/Renderer/Private/MobileDeferredShadingPass.cpp b/Engine/Source/Runtime/Renderer/Private/MobileDeferredShadingPass.cpp
--- a/Engine/Source/Runtime/Renderer/Private/MobileDeferredShadingPass.cpp
+++ b/Engine/Source/Runtime/Renderer/Private/MobileDeferredShadingPass.cpp
@@ -34,6 +34,15 @@ static FAutoConsoleVariableRef CVarMobileUseLightStencilCulling(
 	ECVF_RenderThreadSafe
 );
 
+int32 GMobileIgnoreDeferredShadingSkyLightChannels = 0;
+static FAutoConsoleVariableRef CVarMobileIgnoreDeferredShadingSkyLightChannels(
+	TEXT("r.Mobile.IgnoreDeferredShadingSkyLightChannels"),
+	GMobileIgnoreDeferredShadingSkyLightChannels,
+	TEXT("Whether to ignore primitive lighting channels when applying SkyLighting in a mobile deferred shading.\n" 
+		 "This may improve GPU performance at the cost of incorrect lighting for a primitves with non-default lighting channels"),
+	ECVF_RenderThreadSafe
+);
+
 BEGIN_SHADER_PARAMETER_STRUCT(FMobileDeferredPassParameters, )
 	SHADER_PARAMETER_RDG_UNIFORM_BUFFER(FMobileSceneTextureUniformParameters, MobileSceneTextures)
 	RENDER_TARGET_BINDING_SLOTS()
@@ -577,7 +586,8 @@ static void RenderDirectionalLights(FRHICommandList& RHICmdList, const FScene& S
 		NumLights += (Scene.MobileDirectionalLights[ChannelIdx] ? 1 : 0);
 	}
 	// We can merge reflection and skylight pass with a sole directional light pass and if all primitives and the directional light use the default lighting channel
-	const bool bInlineReflectionAndSky = (NumLights == 1) && !View.bUsesLightingChannels && (Scene.MobileDirectionalLights[0] != nullptr);
+	bool bPrimitivesUseLightingChannels = (View.bUsesLightingChannels && GMobileIgnoreDeferredShadingSkyLightChannels == 0);
+	const bool bInlineReflectionAndSky = (NumLights == 1) && !bPrimitivesUseLightingChannels && (Scene.MobileDirectionalLights[0] != nullptr);
 
 	for (uint32 ChannelIdx = 0; ChannelIdx < UE_ARRAY_COUNT(Scene.MobileDirectionalLights); ChannelIdx++)
 	{
