Precache deferred decal pass PSOs

diff --git a/Engine/Source/Runtime/Engine/Classes/Materials/MaterialInstance.h b/Engine/Source/Runtime/Engine/Classes/Materials/MaterialInstance.h
--- a/Engine/Source/Runtime/Engine/Classes/Materials/MaterialInstance.h
+++ b/Engine/Source/Runtime/Engine/Classes/Materials/MaterialInstance.h
@@ -808,6 +808,7 @@ public:
 	ENGINE_API virtual FDisplacementScaling GetDisplacementScaling() const override;
 	ENGINE_API virtual float GetMaxWorldPositionOffsetDisplacement() const override;
 	ENGINE_API virtual bool ShouldAlwaysEvaluateWorldPositionOffset() const override;
+	ENGINE_API virtual bool IsDeferredDecal() const override;
 	ENGINE_API virtual bool HasPixelAnimation() const override;
 	ENGINE_API virtual USubsurfaceProfile* GetSubsurfaceProfile_Internal() const override;
 	ENGINE_API virtual bool CastsRayTracedShadows() const override;
diff --git a/Engine/Source/Runtime/Engine/Private/Materials/Material.cpp b/Engine/Source/Runtime/Engine/Private/Materials/Material.cpp
--- a/Engine/Source/Runtime/Engine/Private/Materials/Material.cpp
+++ b/Engine/Source/Runtime/Engine/Private/Materials/Material.cpp
@@ -86,6 +86,7 @@
 #include "Misc/OutputDeviceArchiveWrapper.h"
 #include "HAL/FileManager.h"
 #include "BuildSettings.h"
+#include "LocalVertexFactory.h"
 
 #if WITH_EDITOR
 #include "Framework/Notifications/NotificationManager.h"
@@ -4195,6 +4196,12 @@ void UMaterial::PostLoad()
 		UpdateLightmassTextureTracking();
 	}
 
+	if (IsDeferredDecal())
+	{
+		FPSOPrecacheParams PSOPrecacheParams;
+		UMaterialInterface::PrecachePSOs(&FLocalVertexFactory::StaticType, PSOPrecacheParams);
+	}
+
 	//DumpDebugInfo(*GLog);
 }
 
diff --git a/Engine/Source/Runtime/Engine/Private/Materials/MaterialInstance.cpp b/Engine/Source/Runtime/Engine/Private/Materials/MaterialInstance.cpp
--- a/Engine/Source/Runtime/Engine/Private/Materials/MaterialInstance.cpp
+++ b/Engine/Source/Runtime/Engine/Private/Materials/MaterialInstance.cpp
@@ -56,6 +56,7 @@
 #include "ComponentRecreateRenderStateContext.h"
 #include "UObject/UE5ReleaseStreamObjectVersion.h"
 #include "VT/RuntimeVirtualTexture.h"
+#include "LocalVertexFactory.h"
 
 #include UE_INLINE_GENERATED_CPP_BY_NAME(MaterialInstance)
 
@@ -3206,6 +3207,13 @@ void UMaterialInstance::PostLoad()
 
 		LightingGuidFixupMap.Add(GetLightingGuid(), this);
 	}
+
+	if (IsDeferredDecal())
+	{
+		FPSOPrecacheParams PSOPrecacheParams;
+		UMaterialInterface::PrecachePSOs(&FLocalVertexFactory::StaticType, PSOPrecacheParams);
+	}
+
 	//DumpDebugInfo(*GLog);
 }
 
@@ -4616,6 +4624,11 @@ bool UMaterialInstance::ShouldAlwaysEvaluateWorldPositionOffset() const
 	return Parent ? Parent->ShouldAlwaysEvaluateWorldPositionOffset() : false;
 }
 
+bool UMaterialInstance::IsDeferredDecal() const
+{
+	return Parent ? Parent->IsDeferredDecal() : false;
+}
+
 bool UMaterialInstance::HasPixelAnimation() const
 {
 	return bHasPixelAnimation;
diff --git a/Engine/Source/Runtime/Renderer/Private/CompositionLighting/PostProcessDeferredDecals.cpp b/Engine/Source/Runtime/Renderer/Private/CompositionLighting/PostProcessDeferredDecals.cpp
--- a/Engine/Source/Runtime/Renderer/Private/CompositionLighting/PostProcessDeferredDecals.cpp
+++ b/Engine/Source/Runtime/Renderer/Private/CompositionLighting/PostProcessDeferredDecals.cpp
@@ -16,6 +16,7 @@
 #include "VisualizeTexture.h"
 #include "RenderCore.h"
 #include "VariableRateShadingImageManager.h"
+#include "PSOPrecacheValidation.h"
 
 static TAutoConsoleVariable<float> CVarStencilSizeThreshold(
 	TEXT("r.Decal.StencilSizeThreshold"),
@@ -488,6 +489,74 @@ static const TCHAR* GetStageName(EDecalRenderStage Stage)
 	return TEXT("<UNKNOWN>");
 }
 
+void CollectDeferredDecalPassPSOInitializers(
+	int32 PSOCollectorIndex,
+	ERHIFeatureLevel::Type FeatureLevel,
+	const FSceneTexturesConfig& SceneTexturesConfig,
+	const FMaterial& Material,
+	EDecalRenderStage DecalRenderStage,
+	TArray<FPSOPrecacheData>& PSOInitializers)
+{
+	const EShaderPlatform ShaderPlatform = GetFeatureLevelShaderPlatform(FeatureLevel);
+	const FDecalBlendDesc DecalBlendDesc = DecalRendering::ComputeDecalBlendDesc(ShaderPlatform, Material);
+	EDecalRenderTargetMode DecalRenderTargetMode = DecalRendering::GetRenderTargetMode(DecalBlendDesc, DecalRenderStage);
+
+	FGraphicsPipelineStateInitializer GraphicsPSOInit;
+	GraphicsPSOInit.PrimitiveType = PT_TriangleList;		
+	GraphicsPSOInit.BlendState = DecalRendering::GetDecalBlendState(DecalBlendDesc, DecalRenderStage, DecalRenderTargetMode);
+
+	DecalRendering::SetupShaderState(FeatureLevel, Material, DecalRenderStage, GraphicsPSOInit.BoundShaderState);
+
+	FGraphicsPipelineRenderTargetsInfo RenderTargetsInfo;
+	RenderTargetsInfo.NumSamples = 1;
+	GetDeferredDecalRenderTargetsInfo(SceneTexturesConfig, ShaderPlatform, DecalRenderTargetMode, RenderTargetsInfo);
+	ApplyTargetsInfo(GraphicsPSOInit, RenderTargetsInfo);
+	
+	const auto AddDeferredDecalPSO = [&](bool bInsideDecal,	bool bReverseHanded, bool bReverseCulling, bool bDecalUsesStencil)
+	{
+		const EDecalRasterizerState DecalRasterizerState = DecalRendering::GetDecalRasterizerState(bInsideDecal, bReverseHanded, bReverseCulling);
+		GraphicsPSOInit.RasterizerState = DecalRendering::GetDecalRasterizerState(DecalRasterizerState);
+
+		uint32 StencilRef = 0;
+		const FDecalDepthState DecalDepthState = ComputeDecalDepthState(DecalRenderStage, bInsideDecal, bDecalUsesStencil);
+		GraphicsPSOInit.DepthStencilState = GetDecalDepthState(StencilRef, DecalDepthState);
+
+		GraphicsPSOInit.StatePrecachePSOHash = RHIComputeStatePrecachePSOHash(GraphicsPSOInit);
+
+		FPSOPrecacheData& PSOPrecacheData = PSOInitializers.Emplace_GetRef();
+		PSOPrecacheData.bRequired = true;
+		PSOPrecacheData.Type = FPSOPrecacheData::EType::Graphics;
+		PSOPrecacheData.GraphicsPSOInitializer = GraphicsPSOInit;
+#if PSO_PRECACHING_VALIDATE
+		PSOPrecacheData.PSOCollectorIndex = PSOCollectorIndex;
+		PSOPrecacheData.VertexFactoryType = nullptr;
+#endif // PSO_PRECACHING_VALIDATE		
+	};
+
+	PSOInitializers.Reserve(FMath::Max(PSOInitializers.Max(), PSOInitializers.Num() + 16));
+
+	const auto AddDeferredDecalPSOInsideOutside = [&](bool bReverseHanded, bool bReverseCulling, bool bDecalUsesStencil)
+	{
+		AddDeferredDecalPSO(false /*bInsideDecal*/, bReverseHanded, bReverseCulling, bDecalUsesStencil);
+		AddDeferredDecalPSO(true /*bInsideDecal*/, bReverseHanded, bReverseCulling, bDecalUsesStencil);
+	};
+
+	const auto AddDeferredDecalPSOReverseHanded = [&](bool bReverseCulling, bool bDecalUsesStencil)
+	{
+		AddDeferredDecalPSOInsideOutside(false /*bReverseHanded*/, bReverseCulling, bDecalUsesStencil);
+		AddDeferredDecalPSOInsideOutside(true /*bReverseHanded*/, bReverseCulling, bDecalUsesStencil);
+	};
+
+	const auto AddDeferredDecalPSOReverseCulling = [&](bool bDecalUsesStencil)
+	{
+		AddDeferredDecalPSOReverseHanded(false /*bReverseCulling*/, bDecalUsesStencil);
+		AddDeferredDecalPSOReverseHanded(true /*bReverseCulling*/, bDecalUsesStencil);
+	};
+
+	AddDeferredDecalPSOReverseCulling(false /*bDecalUsesStencil*/);
+	AddDeferredDecalPSOReverseCulling(true /*bDecalUsesStencil*/);
+}
+
 void AddDeferredDecalPass(
 	FRDGBuilder& GraphBuilder,
 	const FViewInfo& View,
@@ -556,6 +625,10 @@ void AddDeferredDecalPass(
 		{
 			RHICmdList.SetViewport(View.ViewRect.Min.X, View.ViewRect.Min.Y, 0.0f, View.ViewRect.Max.X, View.ViewRect.Max.Y, 1.0f);
 
+#if PSO_PRECACHING_VALIDATE
+			int32 PSOCollectorIndex = FPSOCollectorCreateManager::GetIndex(EShadingPath::Deferred, TEXT("MeshDecal"));
+#endif // PSO_PRECACHING_VALIDATE
+
 			for (uint32 DecalIndex = DecalIndexBegin; DecalIndex < DecalIndexEnd; ++DecalIndex)
 			{
 				const FTransientDecalRenderData& DecalData = (*SortedDecals)[DecalIndex];
@@ -595,6 +668,14 @@ void AddDeferredDecalPass(
 				GraphicsPSOInit.PrimitiveType = PT_TriangleList;
 
 				DecalRendering::SetShader(RHICmdList, GraphicsPSOInit, StencilRef, View, DecalData, DecalRenderStage, FrustumComponentToClip);
+
+#if PSO_PRECACHING_VALIDATE
+				if (PSOCollectorStats::IsFullPrecachingValidationEnabled())
+				{					
+					PSOCollectorStats::CheckFullPipelineStateInCache(GraphicsPSOInit, EPSOPrecacheResult::Unknown, DecalData.MaterialProxy, &FLocalVertexFactory::StaticType, nullptr, PSOCollectorIndex);
+				}
+#endif // PSO_PRECACHING_VALIDATE
+
 				RHICmdList.DrawIndexedPrimitive(GetUnitCubeIndexBuffer(), 0, 0, 8, 0, UE_ARRAY_COUNT(GCubeIndices) / 3, 1);
 			}
 		});
diff --git a/Engine/Source/Runtime/Renderer/Private/CompositionLighting/PostProcessDeferredDecals.h b/Engine/Source/Runtime/Renderer/Private/CompositionLighting/PostProcessDeferredDecals.h
--- a/Engine/Source/Runtime/Renderer/Private/CompositionLighting/PostProcessDeferredDecals.h
+++ b/Engine/Source/Runtime/Renderer/Private/CompositionLighting/PostProcessDeferredDecals.h
@@ -76,6 +76,14 @@ void GetDeferredDecalRenderTargetsInfo(
 	EDecalRenderTargetMode RenderTargetMode,
 	FGraphicsPipelineRenderTargetsInfo& RenderTargetsInfo);
 
+void CollectDeferredDecalPassPSOInitializers(
+	int32 PSOCollectorIndex,
+	ERHIFeatureLevel::Type FeatureLevel,
+	const FSceneTexturesConfig& SceneTexturesConfig,
+	const FMaterial& Material,
+	EDecalRenderStage DecalRenderStage,
+	TArray<FPSOPrecacheData>& PSOInitializers);
+
 void GetDeferredDecalPassParameters(
 	FRDGBuilder& GraphBuilder,
 	const FViewInfo& View,
diff --git a/Engine/Source/Runtime/Renderer/Private/CompositionLighting/PostProcessMeshDecals.cpp b/Engine/Source/Runtime/Renderer/Private/CompositionLighting/PostProcessMeshDecals.cpp
--- a/Engine/Source/Runtime/Renderer/Private/CompositionLighting/PostProcessMeshDecals.cpp
+++ b/Engine/Source/Runtime/Renderer/Private/CompositionLighting/PostProcessMeshDecals.cpp
@@ -159,6 +159,15 @@ private:
 		ERasterizerFillMode MeshFillMode,
 		ERasterizerCullMode MeshCullMode);
 
+	void CollectDeferredDecalMeshPSOInitializers(
+		const FSceneTexturesConfig& SceneTexturesConfig,
+		const FPSOPrecacheVertexFactoryData& VertexFactoryData,
+		const FPSOPrecacheParams& PreCacheParams,
+		const FMaterial& Material,
+		const FDecalBlendDesc DecalBlendDesc,
+		EDecalRenderStage DecalRenderStage,
+		TArray<FPSOPrecacheData>& PSOInitializers);
+
 	FMeshPassProcessorRenderState PassDrawRenderState;
 	const EDecalRenderStage PassDecalStage;
 	const EDecalRenderTargetMode RenderTargetMode;
@@ -324,6 +333,68 @@ bool FMeshDecalMeshProcessor::Process(
 	return true;
 }
 
+void FMeshDecalMeshProcessor::CollectDeferredDecalMeshPSOInitializers(	
+	const FSceneTexturesConfig& SceneTexturesConfig,
+	const FPSOPrecacheVertexFactoryData& VertexFactoryData,
+	const FPSOPrecacheParams& PreCacheParams,
+	const FMaterial& Material,
+	const FDecalBlendDesc DecalBlendDesc,
+	EDecalRenderStage DecalRenderStage,
+	TArray<FPSOPrecacheData>& PSOInitializers)
+{
+	EDecalRenderTargetMode LocalRenderTargetMode = DecalRendering::GetRenderTargetMode(DecalBlendDesc, DecalRenderStage);
+	PassDrawRenderState.SetBlendState(DecalRendering::GetDecalBlendState(DecalBlendDesc, DecalRenderStage, LocalRenderTargetMode));
+
+	const FMeshDrawingPolicyOverrideSettings OverrideSettings = ComputeMeshOverrideSettings(PreCacheParams);
+	ERasterizerFillMode MeshFillMode = ComputeMeshFillMode(Material, OverrideSettings);
+	ERasterizerCullMode MeshCullMode = ComputeMeshCullMode(Material, OverrideSettings);
+
+	FMaterialShaderTypes ShaderTypes;
+	ShaderTypes.AddShaderType<FMeshDecalsVS>();
+	if (DecalRenderStage == EDecalRenderStage::Emissive)
+	{
+		ShaderTypes.AddShaderType<FMeshDecalsEmissivePS>();
+	}
+	else if (DecalRenderStage == EDecalRenderStage::AmbientOcclusion)
+	{
+		ShaderTypes.AddShaderType<FMeshDecalsAmbientOcclusionPS>();
+	}
+	else
+	{
+		ShaderTypes.AddShaderType<FMeshDecalsPS>();
+	}
+
+	FMaterialShaders Shaders;
+	if (!Material.TryGetShaders(ShaderTypes, VertexFactoryData.VertexFactoryType, Shaders))
+	{
+		return;
+	}
+
+	TMeshProcessorShaders<
+		FMeshDecalsVS,
+		FMeshDecalsPS> MeshDecalPassShaders;
+	Shaders.TryGetVertexShader(MeshDecalPassShaders.VertexShader);
+	Shaders.TryGetPixelShader(MeshDecalPassShaders.PixelShader);
+
+	const EShaderPlatform ShaderPlatform = GetFeatureLevelShaderPlatform(FeatureLevel);
+	FGraphicsPipelineRenderTargetsInfo RenderTargetsInfo;
+	RenderTargetsInfo.NumSamples = 1;
+	GetDeferredDecalRenderTargetsInfo(SceneTexturesConfig, ShaderPlatform, LocalRenderTargetMode, RenderTargetsInfo);
+
+	AddGraphicsPipelineStateInitializer(
+		VertexFactoryData,
+		Material,
+		PassDrawRenderState,
+		RenderTargetsInfo,
+		MeshDecalPassShaders,
+		MeshFillMode,
+		MeshCullMode,
+		PT_TriangleList,
+		EMeshPassFeatures::Default,
+		true /*bRequired*/,
+		PSOInitializers);
+}
+
 void FMeshDecalMeshProcessor::CollectPSOInitializers(
 	const FSceneTexturesConfig& SceneTexturesConfig,
 	const FMaterial& Material,
@@ -341,64 +412,19 @@ void FMeshDecalMeshProcessor::CollectPSOInitializers(
 
 	for (uint32 DecalStageIter = 0; DecalStageIter < (uint32)EDecalRenderStage::Num; ++DecalStageIter)
 	{
-		EDecalRenderStage LocalDecalRenderState = EDecalRenderStage(DecalStageIter);
-
-		const bool bShouldRender = DecalRendering::IsCompatibleWithRenderStage(DecalBlendDesc, LocalDecalRenderState);
+		EDecalRenderStage LocalDecalRenderStage = EDecalRenderStage(DecalStageIter);
+				
+		const bool bShouldRender = DecalRendering::IsCompatibleWithRenderStage(DecalBlendDesc, LocalDecalRenderStage);
 		if (!bShouldRender)
 		{
 			continue;
 		}
 		
-		EDecalRenderTargetMode LocalRenderTargetMode = DecalRendering::GetRenderTargetMode(DecalBlendDesc, LocalDecalRenderState);
-		PassDrawRenderState.SetBlendState(DecalRendering::GetDecalBlendState(DecalBlendDesc, LocalDecalRenderState, LocalRenderTargetMode));
-
-		const FMeshDrawingPolicyOverrideSettings OverrideSettings = ComputeMeshOverrideSettings(PreCacheParams);
-		ERasterizerFillMode MeshFillMode = ComputeMeshFillMode(Material, OverrideSettings);
-		ERasterizerCullMode MeshCullMode = ComputeMeshCullMode(Material, OverrideSettings);
+		// Collect decal pass PSOs
+		CollectDeferredDecalPassPSOInitializers(PSOCollectorIndex, FeatureLevel, SceneTexturesConfig, Material, LocalDecalRenderStage, PSOInitializers);
 
-		FMaterialShaderTypes ShaderTypes;
-		ShaderTypes.AddShaderType<FMeshDecalsVS>();
-		if (LocalDecalRenderState == EDecalRenderStage::Emissive)
-		{
-			ShaderTypes.AddShaderType<FMeshDecalsEmissivePS>();
-		}
-		else if (LocalDecalRenderState == EDecalRenderStage::AmbientOcclusion)
-		{
-			ShaderTypes.AddShaderType<FMeshDecalsAmbientOcclusionPS>();
-		}
-		else
-		{
-			ShaderTypes.AddShaderType<FMeshDecalsPS>();
-		}
-
-		FMaterialShaders Shaders;
-		if (!Material.TryGetShaders(ShaderTypes, VertexFactoryData.VertexFactoryType, Shaders))
-		{
-			continue;
-		}
-
-		TMeshProcessorShaders<
-			FMeshDecalsVS,
-			FMeshDecalsPS> MeshDecalPassShaders;
-		Shaders.TryGetVertexShader(MeshDecalPassShaders.VertexShader);
-		Shaders.TryGetPixelShader(MeshDecalPassShaders.PixelShader);
-
-		FGraphicsPipelineRenderTargetsInfo RenderTargetsInfo;
-		RenderTargetsInfo.NumSamples = 1;
-		GetDeferredDecalRenderTargetsInfo(SceneTexturesConfig, ShaderPlatform, LocalRenderTargetMode, RenderTargetsInfo);
-		
-		AddGraphicsPipelineStateInitializer(
-			VertexFactoryData,
-			Material,
-			PassDrawRenderState,
-			RenderTargetsInfo,
-			MeshDecalPassShaders,
-			MeshFillMode,
-			MeshCullMode,
-			PT_TriangleList,
-			EMeshPassFeatures::Default,
-			true /*bRequired*/,
-			PSOInitializers);
+		// Collect decal mesh PSOs
+		CollectDeferredDecalMeshPSOInitializers(SceneTexturesConfig, VertexFactoryData, PreCacheParams, Material, DecalBlendDesc, LocalDecalRenderStage, PSOInitializers);
 	}
 }
 
diff --git a/Engine/Source/Runtime/Renderer/Private/DecalRenderingShared.cpp b/Engine/Source/Runtime/Renderer/Private/DecalRenderingShared.cpp
--- a/Engine/Source/Runtime/Renderer/Private/DecalRenderingShared.cpp
+++ b/Engine/Source/Runtime/Renderer/Private/DecalRenderingShared.cpp
@@ -404,6 +404,26 @@ namespace DecalRendering
 		return OutPixelShader.IsValid();
 	}
 
+	bool SetupShaderState(
+		ERHIFeatureLevel::Type FeatureLevel,
+		const FMaterial& Material, 
+		EDecalRenderStage DecalRenderStage, 
+		FBoundShaderStateInput& OutBoundShaderState)
+	{
+		TShaderRef<FDeferredDecalPS> PixelShader;
+		if (!TryGetDeferredDecalShaders(Material, FeatureLevel, DecalRenderStage, PixelShader))
+		{
+			return false;
+		}
+
+		TShaderMapRef<FDeferredDecalVS> VertexShader(GetGlobalShaderMap(FeatureLevel));
+		OutBoundShaderState.VertexDeclarationRHI = GetVertexDeclarationFVector4();
+		OutBoundShaderState.VertexShaderRHI = VertexShader.GetVertexShader();
+		OutBoundShaderState.PixelShaderRHI = PixelShader.GetPixelShader();
+
+		return true;
+	}
+
 	FMaterialRenderProxy const* TryGetDeferredDecalMaterial(
 		FMaterialRenderProxy const* MaterialProxy, 
 		ERHIFeatureLevel::Type FeatureLevel,
diff --git a/Engine/Source/Runtime/Renderer/Private/DecalRenderingShared.h b/Engine/Source/Runtime/Renderer/Private/DecalRenderingShared.h
--- a/Engine/Source/Runtime/Renderer/Private/DecalRenderingShared.h
+++ b/Engine/Source/Runtime/Renderer/Private/DecalRenderingShared.h
@@ -40,5 +40,6 @@ namespace DecalRendering
 	void SetVertexShaderOnly(FRHICommandList& RHICmdList, FGraphicsPipelineStateInitializer& GraphicsPSOInit, const FViewInfo& View, const FMatrix& FrustumComponentToClip);
 	void SortDecalList(FTransientDecalRenderDataList& Decals);
 	bool BuildVisibleDecalList(const FScene& Scene, const FViewInfo& View, EDecalRenderStage DecalRenderStage, FTransientDecalRenderDataList* OutVisibleDecals);
+	bool SetupShaderState(ERHIFeatureLevel::Type FeatureLevel, const FMaterial& Material, EDecalRenderStage DecalRenderStage, FBoundShaderStateInput& OutBoundShaderState);
 	void SetShader(FRHICommandList& RHICmdList, FGraphicsPipelineStateInitializer& GraphicsPSOInit, uint32 StencilRef, const FViewInfo& View, const FTransientDecalRenderData& DecalData, EDecalRenderStage DecalRenderStage, const FMatrix& FrustumComponentToClip);
 };
