Fix per-object / modulated shadows in mobile forward renderer
- stencil bits [0-2] are used by shadow rendering so can't be written during depth prepass
- moved "Cast Contact Shadow" and "Mobile Local Lights" stencil bits to bits 4-5 since those are only used by mobile deferred.

diff --git a/Engine/Shaders/Private/ScreenSpaceShadows.usf b/Engine/Shaders/Private/ScreenSpaceShadows.usf
--- a/Engine/Shaders/Private/ScreenSpaceShadows.usf
+++ b/Engine/Shaders/Private/ScreenSpaceShadows.usf
@@ -87,7 +87,7 @@ RWTexture2D<float2> RWShadowFactors;
 Texture2D<uint2> StencilTexture;
 
 // Must match C++ (SceneRenderTargets.h)
-#define STENCIL_MOBILE_CAST_CONTACT_SHADOW_MASK 2
+#define STENCIL_MOBILE_CAST_CONTACT_SHADOW_MASK (1 << 4)
 
 float GetContactShadowOcclusion(float2 HitUV)
 {
diff --git a/Engine/Source/Runtime/Renderer/Private/PostProcess/SceneRenderTargets.h b/Engine/Source/Runtime/Renderer/Private/PostProcess/SceneRenderTargets.h
--- a/Engine/Source/Runtime/Renderer/Private/PostProcess/SceneRenderTargets.h
+++ b/Engine/Source/Runtime/Renderer/Private/PostProcess/SceneRenderTargets.h
@@ -49,13 +49,19 @@
 
 #define STENCIL_LIGHTING_CHANNELS_MASK(Value) uint8(((Value) & 0x7) << STENCIL_LIGHTING_CHANNELS_BIT_ID)
 
-// Mobile specific
-// Store shading model into stencil [1-2] bits (only used in mobile deferred)
-#define GET_STENCIL_MOBILE_SM_MASK(Value) uint8(((Value) & 0x3) << 1)
-
-// Cast contact shadow mask - bit 1 (only used in mobile forward) / Must match shader (ScreenSpaceShadows.usf)
-#define STENCIL_MOBILE_CAST_CONTACT_SHADOW_BIT_ID	1
-// Forward local light buffer mask for light function - bit 2
-#define STENCIL_MOBILE_LIGHTFUNCTION_MASK			uint8(1 << 2)
+// [Mobile specific]
+
+// stencil [0-2] bits are used to render per-object shadows (see ShadowStencilMask in ShadowRendering.cpp)
+
 // Sky material mask - bit 3
 #define STENCIL_MOBILE_SKY_MASK						uint8(1 << 3)
+
+// [Mobile Deferred only]
+// Store shading model into stencil [1-2] bits
+#define GET_STENCIL_MOBILE_SM_MASK(Value)			uint8(((Value) & 0x3) << 1)
+
+// [Mobile Forward only]
+// Cast contact shadow mask - bit 4 / Must match shader (ScreenSpaceShadows.usf)
+#define STENCIL_MOBILE_CAST_CONTACT_SHADOW_BIT_ID	4
+// Forward local light buffer mask for light function - bit 5
+#define STENCIL_MOBILE_LIGHTFUNCTION_MASK			uint8(1 << 5)
