Precache PSOs for Hair strand visibility passes

diff --git a/Engine/Plugins/Runtime/HairStrands/Source/HairStrandsCore/Private/HairStrandsVertexFactory.cpp b/Engine/Plugins/Runtime/HairStrands/Source/HairStrandsCore/Private/HairStrandsVertexFactory.cpp
--- a/Engine/Plugins/Runtime/HairStrands/Source/HairStrandsCore/Private/HairStrandsVertexFactory.cpp
+++ b/Engine/Plugins/Runtime/HairStrands/Source/HairStrandsCore/Private/HairStrandsVertexFactory.cpp
@@ -255,6 +255,13 @@ void FHairStrandsVertexFactory::InitResources(FRHICommandListBase& RHICmdList)
 	Data.Instance->Strands.UniformBuffer = FHairStrandsUniformBuffer::CreateUniformBufferImmediate(Parameters, UniformBuffer_MultiFrame);
 }
 
+void FHairStrandsVertexFactory::GetPSOPrecacheVertexFetchElements(EVertexInputStreamType VertexInputStreamType, FVertexDeclarationElementList& Elements)
+{
+#if VF_STRANDS_SUPPORT_GPU_SCENE
+	Elements.Add(FVertexElement(0, 0, VET_UInt, 13, sizeof(uint32), true));
+#endif
+}
+
 IMPLEMENT_VERTEX_FACTORY_PARAMETER_TYPE(FHairStrandsVertexFactory, SF_Vertex,		FHairStrandsVertexFactoryShaderParameters);
 IMPLEMENT_VERTEX_FACTORY_PARAMETER_TYPE(FHairStrandsVertexFactory, SF_Pixel,		FHairStrandsVertexFactoryShaderParameters);
 IMPLEMENT_VERTEX_FACTORY_PARAMETER_TYPE(FHairStrandsVertexFactory, SF_Compute,		FHairStrandsVertexFactoryShaderParameters);
@@ -281,4 +288,5 @@ IMPLEMENT_VERTEX_FACTORY_TYPE(FHairStrandsVertexFactory, "/Engine/Private/HairSt
 	| EVertexFactoryFlags::SupportsRayTracing
 	| (VF_STRANDS_PROCEDURAL_INTERSECTOR ? EVertexFactoryFlags::SupportsRayTracingProceduralPrimitive : EVertexFactoryFlags::None)
 	| EVertexFactoryFlags::SupportsManualVertexFetch
+	| EVertexFactoryFlags::SupportsPSOPrecaching
 );
diff --git a/Engine/Plugins/Runtime/HairStrands/Source/HairStrandsCore/Public/HairStrandsVertexFactory.h b/Engine/Plugins/Runtime/HairStrands/Source/HairStrandsCore/Public/HairStrandsVertexFactory.h
--- a/Engine/Plugins/Runtime/HairStrands/Source/HairStrandsCore/Public/HairStrandsVertexFactory.h
+++ b/Engine/Plugins/Runtime/HairStrands/Source/HairStrandsCore/Public/HairStrandsVertexFactory.h
@@ -42,6 +42,7 @@ public:
 	static bool ShouldCompilePermutation(const FVertexFactoryShaderPermutationParameters& Parameters);
 	static void ModifyCompilationEnvironment(const FVertexFactoryShaderPermutationParameters& Parameters, FShaderCompilerEnvironment& OutEnvironment);
 	static void ValidateCompiledResult(const FVertexFactoryType* Type, EShaderPlatform Platform, const FShaderParameterMap& ParameterMap, TArray<FString>& OutErrors);
+	static void GetPSOPrecacheVertexFetchElements(EVertexInputStreamType VertexInputStreamType, FVertexDeclarationElementList& Elements);
 
 	/**
 	 * An implementation of the interface used by TSynchronizedResource to update the resource with new data from the game thread.
diff --git a/Engine/Source/Runtime/Renderer/Private/HairStrands/HairStrandsVisibility.cpp b/Engine/Source/Runtime/Renderer/Private/HairStrands/HairStrandsVisibility.cpp
--- a/Engine/Source/Runtime/Renderer/Private/HairStrands/HairStrandsVisibility.cpp
+++ b/Engine/Source/Runtime/Renderer/Private/HairStrands/HairStrandsVisibility.cpp
@@ -570,18 +570,39 @@ IMPLEMENT_MATERIAL_SHADER_TYPE(, FHairMaterialPS, TEXT("/Engine/Private/HairStra
 
 /////////////////////////////////////////////////////////////////////////////////////////
 
+enum class EHairMaterialPassFilter : uint8
+{
+	All,
+	EmissiveOnly,
+	NonEmissiveOnly
+};
+
+static bool ShouldRenderHairStrands(ERHIFeatureLevel::Type FeatureLevel, const FMaterial& Material, const FVertexFactoryType* VFType, bool bPrimitiveRenderInPass)
+{
+	static const FVertexFactoryType* CompatibleVF = FVertexFactoryType::GetVFByName(TEXT("FHairStrandsVertexFactory"));
+
+	// Determine the mesh's material and blend mode.
+	const bool bIsCompatible = IsCompatibleWithHairStrands(&Material, FeatureLevel);
+	const bool bIsHairStrandsFactory = CompatibleVF != nullptr && VFType->GetHashedName() == CompatibleVF->GetHashedName();
+
+	return bPrimitiveRenderInPass && bIsCompatible && bIsHairStrandsFactory && ShouldIncludeDomainInMeshPass(Material.GetMaterialDomain());
+}
+
 class FHairMaterialProcessor : public FMeshPassProcessor
 {
 public:
 	FHairMaterialProcessor(
 		const FScene* Scene,
-		const FSceneView* InViewIfDynamicMeshCommand,
-		const FMeshPassProcessorRenderState& InPassDrawRenderState,
-		FDynamicPassMeshDrawListContext* InDrawListContext);
+		ERHIFeatureLevel::Type FeatureLevel,
+		const FSceneView* InViewIfDynamicMeshCommand,				
+		FDynamicPassMeshDrawListContext* InDrawListContext,
+		EHairMaterialPassFilter InFilter);
 
 	virtual void AddMeshBatch(const FMeshBatch& RESTRICT MeshBatch, uint64 BatchElementMask, const FPrimitiveSceneProxy* RESTRICT PrimitiveSceneProxy, int32 StaticMeshId = -1) override final;
 	void AddMeshBatch(const FMeshBatch& RESTRICT MeshBatch, uint64 BatchElementMask, const FPrimitiveSceneProxy* RESTRICT PrimitiveSceneProxy, int32 StaticMeshId, int32 MacroGroupId, int32 HairMaterialId, float HairCoverageScale);
 
+	virtual void CollectPSOInitializers(const FSceneTexturesConfig& SceneTexturesConfig, const FMaterial& Material, const FPSOPrecacheVertexFactoryData& VertexFactoryData, const FPSOPrecacheParams& PreCacheParams, TArray<FPSOPrecacheData>& PSOInitializers) override final;
+
 private:
 	bool TryAddMeshBatch(
 		const FMeshBatch& RESTRICT MeshBatch,
@@ -607,9 +628,25 @@ private:
 		const uint32 HairPrimitiveLightChannelMask,
 		const float HairCoverageScale);
 
+	void SetupDrawRenderState(EHairMaterialPassFilter InFilter);
+
 	FMeshPassProcessorRenderState PassDrawRenderState;
+	EHairMaterialPassFilter Filter;
 };
 
+void FHairMaterialProcessor::SetupDrawRenderState(EHairMaterialPassFilter InFilter)
+{
+	if (InFilter == EHairMaterialPassFilter::All || InFilter == EHairMaterialPassFilter::EmissiveOnly)
+	{
+		PassDrawRenderState.SetBlendState(TStaticBlendState<CW_RGBA, BO_Add, BF_One, BF_One, BO_Add, BF_One, BF_Zero>::GetRHI());
+	}
+	else
+	{
+		PassDrawRenderState.SetBlendState(TStaticBlendState<>::GetRHI());
+	}
+	PassDrawRenderState.SetDepthStencilState(TStaticDepthStencilState <false, CF_Always> ::GetRHI());
+}
+
 void FHairMaterialProcessor::AddMeshBatch(const FMeshBatch& RESTRICT MeshBatch, uint64 BatchElementMask, const FPrimitiveSceneProxy* RESTRICT PrimitiveSceneProxy, int32 StaticMeshId)
 {
 	AddMeshBatch(MeshBatch, BatchElementMask, PrimitiveSceneProxy, StaticMeshId, 0, 0, 1.f);
@@ -644,17 +681,9 @@ bool FHairMaterialProcessor::TryAddMeshBatch(
 	const FMaterialRenderProxy& MaterialRenderProxy,
 	const FMaterial& Material)
 {
-	static const FVertexFactoryType* CompatibleVF = FVertexFactoryType::GetVFByName(TEXT("FHairStrandsVertexFactory"));
-
-	// Determine the mesh's material and blend mode.
-	const bool bIsCompatible = IsCompatibleWithHairStrands(&Material, FeatureLevel);
-	const bool bIsHairStrandsFactory = MeshBatch.VertexFactory->GetType()->GetHashedName() == CompatibleVF->GetHashedName();
-	const bool bShouldRender = (!PrimitiveSceneProxy && MeshBatch.Elements.Num() > 0) || (PrimitiveSceneProxy && PrimitiveSceneProxy->ShouldRenderInMainPass());
-
-	if (bIsCompatible
-		&& bIsHairStrandsFactory
-		&& bShouldRender
-		&& ShouldIncludeDomainInMeshPass(Material.GetMaterialDomain()))
+	const bool bPrimitiveRenderInPass = (!PrimitiveSceneProxy && MeshBatch.Elements.Num() > 0) || (PrimitiveSceneProxy && PrimitiveSceneProxy->ShouldRenderInMainPass());
+	const bool bShouldRender = ShouldRenderHairStrands(FeatureLevel, Material, MeshBatch.VertexFactory->GetType(), bPrimitiveRenderInPass);
+	if (bShouldRender)
 	{
 		// For the mesh patch to be rendered a single triangle triangle to spawn the necessary amount of thread
 		FMeshBatch MeshBatchCopy = MeshBatch;
@@ -734,15 +763,99 @@ bool FHairMaterialProcessor::Process(
 	return true;
 }
 
+void FHairMaterialProcessor::CollectPSOInitializers(
+	const FSceneTexturesConfig& SceneTexturesConfig, 
+	const FMaterial& Material, 
+	const FPSOPrecacheVertexFactoryData& VertexFactoryData, 
+	const FPSOPrecacheParams& PreCacheParams, 
+	TArray<FPSOPrecacheData>& PSOInitializers)
+{
+	const bool bShouldRender = ShouldRenderHairStrands(FeatureLevel, Material, VertexFactoryData.VertexFactoryType, PreCacheParams.bRenderInMainPass);
+	if (!bShouldRender)
+	{
+		return;
+	}
+
+	TMeshProcessorShaders<
+		FHairMaterialVS,
+		FHairMaterialPS> PassShaders;
+	{
+		FMaterialShaderTypes ShaderTypes;
+		ShaderTypes.AddShaderType<FHairMaterialVS>();
+		ShaderTypes.AddShaderType<FHairMaterialPS>();
+
+		FMaterialShaders Shaders;
+		if (!Material.TryGetShaders(ShaderTypes, VertexFactoryData.VertexFactoryType, Shaders))
+		{
+			return;
+		}
+
+		Shaders.TryGetVertexShader(PassShaders.VertexShader);
+		Shaders.TryGetPixelShader(PassShaders.PixelShader);
+	}
+		
+	const auto AddPSOInitializer = [&](EHairMaterialPassFilter InFilter)
+	{
+		FGraphicsPipelineRenderTargetsInfo RenderTargetsInfo;
+		switch (InFilter)
+		{
+		case EHairMaterialPassFilter::All:
+		case EHairMaterialPassFilter::EmissiveOnly:
+		{
+			RenderTargetsInfo.NumSamples = 1;
+			const EPixelFormat Format = FHairLightSampleClearPS::GetHairLightSampleFormat();
+			AddRenderTargetInfo(Format, TexCreate_UAV | TexCreate_ShaderResource | TexCreate_RenderTargetable, RenderTargetsInfo);
+			break;
+		}
+		case EHairMaterialPassFilter::NonEmissiveOnly:
+		{
+			// No render targets
+			break;
+		}
+		}
+
+		SetupDrawRenderState(InFilter);
+
+		// Generate version for planar reflections as well?
+		bool bReverseCulling = false;
+		AddGraphicsPipelineStateInitializer(
+			VertexFactoryData,
+			Material,
+			PassDrawRenderState,
+			RenderTargetsInfo,
+			PassShaders,
+			ERasterizerFillMode::FM_Solid,
+			bReverseCulling ? ERasterizerCullMode::CM_CW : ERasterizerCullMode::CM_CCW,
+			(EPrimitiveType)PreCacheParams.PrimitiveType,
+			EMeshPassFeatures::Default,
+			true /*bRequired*/,
+			PSOInitializers);	
+	};
+
+	AddPSOInitializer(EHairMaterialPassFilter::All);
+	AddPSOInitializer(EHairMaterialPassFilter::EmissiveOnly);
+	AddPSOInitializer(EHairMaterialPassFilter::NonEmissiveOnly);
+}
+
+static const TCHAR* HairMaterialPassName = TEXT("HairMaterial");
+
 FHairMaterialProcessor::FHairMaterialProcessor(
 	const FScene* Scene,
+	ERHIFeatureLevel::Type FeatureLevel,
 	const FSceneView* InViewIfDynamicMeshCommand,
-	const FMeshPassProcessorRenderState& InPassDrawRenderState,
-	FDynamicPassMeshDrawListContext* InDrawListContext)
-	: FMeshPassProcessor(EMeshPass::Num, Scene, Scene->GetFeatureLevel(), InViewIfDynamicMeshCommand, InDrawListContext)
-	, PassDrawRenderState(InPassDrawRenderState)
+	FDynamicPassMeshDrawListContext* InDrawListContext,
+	EHairMaterialPassFilter InFilter)
+	: FMeshPassProcessor(HairMaterialPassName, Scene, FeatureLevel, InViewIfDynamicMeshCommand, InDrawListContext)
+	, Filter(InFilter)
+{
+	SetupDrawRenderState(Filter);
+}
+
+IPSOCollector* CreatePSOCollectorHairMaterial(ERHIFeatureLevel::Type FeatureLevel)
 {
+	return new FHairMaterialProcessor(nullptr, FeatureLevel, nullptr, nullptr, EHairMaterialPassFilter::All);
 }
+FRegisterPSOCollectorCreateFunction RegisterPSOCollectorHairMaterial(&CreatePSOCollectorHairMaterial, EShadingPath::Deferred, HairMaterialPassName);
 
 /////////////////////////////////////////////////////////////////////////////////////////
 
@@ -864,13 +977,6 @@ static FMaterialPassOutput AddHairMaterialPass(
 	const uint32 ResolutionDim = FMath::CeilToInt(FMath::Sqrt(static_cast<float>(MaxNodeCount)));
 	const FIntPoint Resolution(ResolutionDim, ResolutionDim);
 
-	enum class EHairMaterialPassFilter : uint8
-	{
-		All,
-		EmissiveOnly,
-		NonEmissiveOnly
-	};
-
 	const ERHIFeatureLevel::Type FeatureLevel = ViewInfo->FeatureLevel;
 
 	// Find among the mesh batch, if any of them emit emissive data
@@ -935,19 +1041,8 @@ static FMaterialPassOutput AddHairMaterialPass(
 			RDG_EVENT_NAME("HairStrands::MaterialPass(Emissive=%s)", Filter == EHairMaterialPassFilter::All ? TEXT("On/Off") : (Filter == EHairMaterialPassFilter::EmissiveOnly ? TEXT("On") : TEXT("Off"))),
 			FIntRect(0, 0, Resolution.X, Resolution.Y),
 		[PassParameters, Scene = Scene, ViewInfo, &MacroGroupDatas, MaxNodeCount, Resolution, NodeGroupSize, bUpdateSampleCoverage, Filter, FeatureLevel](FDynamicPassMeshDrawListContext* ShadowContext)
-		{
-			FMeshPassProcessorRenderState DrawRenderState;
-			if (Filter == EHairMaterialPassFilter::All || Filter == EHairMaterialPassFilter::EmissiveOnly)
-			{
-				DrawRenderState.SetBlendState(TStaticBlendState<CW_RGBA, BO_Add, BF_One, BF_One, BO_Add, BF_One, BF_Zero>::GetRHI());
-			}
-			else
-			{
-				DrawRenderState.SetBlendState(TStaticBlendState<>::GetRHI());
-			}
-			DrawRenderState.SetDepthStencilState(TStaticDepthStencilState <false, CF_Always> ::GetRHI());
-			FHairMaterialProcessor MeshProcessor(Scene, ViewInfo, DrawRenderState, ShadowContext);
-
+		{			
+			FHairMaterialProcessor MeshProcessor(Scene, Scene->GetFeatureLevel(), ViewInfo, ShadowContext, Filter);
 			for (const FHairStrandsMacroGroupData& MacroGroupData : MacroGroupDatas)
 			{
 				for (const FHairStrandsMacroGroupData::PrimitiveInfo& PrimitiveInfo : MacroGroupData.PrimitivesInfos)
@@ -1393,14 +1488,16 @@ class FHairVisibilityProcessor : public FMeshPassProcessor
 public:
 	FHairVisibilityProcessor(
 		const FScene* Scene,
+		ERHIFeatureLevel::Type FeatureLevel,
 		const FSceneView* InViewIfDynamicMeshCommand,
-		const FMeshPassProcessorRenderState& InPassDrawRenderState,
 		const EHairVisibilityRenderMode InRenderMode,
 		FDynamicPassMeshDrawListContext* InDrawListContext);
 
 	virtual void AddMeshBatch(const FMeshBatch& RESTRICT MeshBatch, uint64 BatchElementMask, const FPrimitiveSceneProxy* RESTRICT PrimitiveSceneProxy, int32 StaticMeshId = -1) override final;
 	void AddMeshBatch(const FMeshBatch& RESTRICT MeshBatch, uint64 BatchElementMask, const FPrimitiveSceneProxy* RESTRICT PrimitiveSceneProxy, int32 StaticMeshId, uint32 HairMacroGroupId, uint32 HairMaterialId, float HairCoverageScale, bool bCullingEnable);
 
+	virtual void CollectPSOInitializers(const FSceneTexturesConfig& SceneTexturesConfig, const FMaterial& Material, const FPSOPrecacheVertexFactoryData& VertexFactoryData, const FPSOPrecacheParams& PreCacheParams, TArray<FPSOPrecacheData>& PSOInitializers) override final;
+
 private:
 	bool TryAddMeshBatch(
 		const FMeshBatch& RESTRICT MeshBatch,
@@ -1429,10 +1526,48 @@ private:
 		ERasterizerFillMode MeshFillMode,
 		ERasterizerCullMode MeshCullMode);
 
+	template<EHairVisibilityRenderMode RenderMode, bool bCullingEnable = true>
+	void AddPSOInitializer(
+		const FSceneTexturesConfig& SceneTexturesConfig,
+		const FMaterial& Material, 
+		const FPSOPrecacheVertexFactoryData& VertexFactoryData,
+		const FPSOPrecacheParams& PreCacheParams,
+		ERasterizerFillMode MeshFillMode,
+		ERasterizerCullMode MeshCullMode,
+		TArray<FPSOPrecacheData>& PSOInitializers);
+
+	void SetupDrawRenderState(EHairVisibilityRenderMode InRenderMode);
+
 	const EHairVisibilityRenderMode RenderMode;
 	FMeshPassProcessorRenderState PassDrawRenderState;
 };
 
+void FHairVisibilityProcessor::SetupDrawRenderState(EHairVisibilityRenderMode InRenderMode)
+{
+	if (InRenderMode == HairVisibilityRenderMode_MSAA_Visibility)
+	{
+		PassDrawRenderState.SetBlendState(TStaticBlendState<CW_RGBA, BO_Add, BF_One, BF_Zero, BO_Add, BF_One, BF_Zero>::GetRHI());
+		PassDrawRenderState.SetDepthStencilState(TStaticDepthStencilState<true, CF_DepthNearOrEqual>::GetRHI());
+	}
+	else if (InRenderMode == HairVisibilityRenderMode_Transmittance)
+	{
+		PassDrawRenderState.SetBlendState(TStaticBlendState<CW_RED, BO_Add, BF_DestColor, BF_Zero, BO_Add, BF_Zero, BF_Zero>::GetRHI());
+		PassDrawRenderState.SetDepthStencilState(TStaticDepthStencilState<false, CF_DepthNearOrEqual>::GetRHI());
+	}
+	else if (InRenderMode == HairVisibilityRenderMode_TransmittanceAndHairCount)
+	{
+		PassDrawRenderState.SetBlendState(TStaticBlendState<
+			CW_RED, BO_Add, BF_DestColor, BF_Zero, BO_Add, BF_Zero, BF_Zero,
+			CW_RG, BO_Add, BF_One, BF_One, BO_Add, BF_Zero, BF_Zero>::GetRHI());
+		PassDrawRenderState.SetDepthStencilState(TStaticDepthStencilState<false, CF_DepthNearOrEqual>::GetRHI());
+	}
+	else if (InRenderMode == HairVisibilityRenderMode_PPLL)
+	{
+		PassDrawRenderState.SetBlendState(TStaticBlendState<>::GetRHI());
+		PassDrawRenderState.SetDepthStencilState(TStaticDepthStencilState<false, CF_DepthNearOrEqual>::GetRHI());
+	}
+}
+
 void FHairVisibilityProcessor::AddMeshBatch(const FMeshBatch& RESTRICT MeshBatch, uint64 BatchElementMask, const FPrimitiveSceneProxy* RESTRICT PrimitiveSceneProxy, int32 StaticMeshId)
 {
 	AddMeshBatch(MeshBatch, BatchElementMask, PrimitiveSceneProxy, StaticMeshId, 0, 0, 1.f, false);
@@ -1468,19 +1603,12 @@ bool FHairVisibilityProcessor::TryAddMeshBatch(
 	const FMaterialRenderProxy& MaterialRenderProxy,
 	const FMaterial& Material)
 {
-	static const FVertexFactoryType* CompatibleVF = FVertexFactoryType::GetVFByName(TEXT("FHairStrandsVertexFactory"));
-
-	// Determine the mesh's material and blend mode.
-	const bool bIsCompatible = IsCompatibleWithHairStrands(&Material, FeatureLevel);
-	const bool bIsHairStrandsFactory = MeshBatch.VertexFactory->GetType()->GetHashedName() == CompatibleVF->GetHashedName();
-	const bool bShouldRender = (!PrimitiveSceneProxy && MeshBatch.Elements.Num() > 0) || (PrimitiveSceneProxy && PrimitiveSceneProxy->ShouldRenderInMainPass());
-	const uint32 LightChannelMask = PrimitiveSceneProxy ? PrimitiveSceneProxy->GetLightingChannelMask() : 0;
-
-	if (bIsCompatible
-		&& bIsHairStrandsFactory
-		&& bShouldRender
-		&& ShouldIncludeDomainInMeshPass(Material.GetMaterialDomain()))
+	const bool bPrimitiveRenderInPass = (!PrimitiveSceneProxy && MeshBatch.Elements.Num() > 0) || (PrimitiveSceneProxy && PrimitiveSceneProxy->ShouldRenderInMainPass());
+	const bool bShouldRender = ShouldRenderHairStrands(FeatureLevel, Material, MeshBatch.VertexFactory->GetType(), bPrimitiveRenderInPass);
+	if (bShouldRender)
 	{
+		const uint32 LightChannelMask = PrimitiveSceneProxy ? PrimitiveSceneProxy->GetLightingChannelMask() : 0;
+
 		const FMeshDrawingPolicyOverrideSettings OverrideSettings = ComputeMeshOverrideSettings(MeshBatch);
 		const ERasterizerFillMode MeshFillMode = ComputeMeshFillMode(Material, OverrideSettings);
 		const ERasterizerCullMode MeshCullMode = ComputeMeshCullMode(Material, OverrideSettings);
@@ -1557,18 +1685,120 @@ bool FHairVisibilityProcessor::Process(
 	return true;
 }
 
+void FHairVisibilityProcessor::CollectPSOInitializers(
+	const FSceneTexturesConfig& SceneTexturesConfig, 
+	const FMaterial& Material, 
+	const FPSOPrecacheVertexFactoryData& VertexFactoryData, 
+	const FPSOPrecacheParams& PreCacheParams, 
+	TArray<FPSOPrecacheData>& PSOInitializers)
+{
+	const bool bShouldRender = ShouldRenderHairStrands(FeatureLevel, Material, VertexFactoryData.VertexFactoryType, PreCacheParams.bRenderInMainPass);
+	if (!bShouldRender)
+	{
+		return;
+	}
+
+	const FMeshDrawingPolicyOverrideSettings OverrideSettings = ComputeMeshOverrideSettings(PreCacheParams);
+	const ERasterizerFillMode MeshFillMode = ComputeMeshFillMode(Material, OverrideSettings);
+	const ERasterizerCullMode MeshCullMode = ComputeMeshCullMode(Material, OverrideSettings);
+
+	AddPSOInitializer<HairVisibilityRenderMode_MSAA_Visibility, true>(SceneTexturesConfig, Material, VertexFactoryData, PreCacheParams, MeshFillMode, MeshCullMode, PSOInitializers);
+	AddPSOInitializer<HairVisibilityRenderMode_MSAA_Visibility, false>(SceneTexturesConfig, Material, VertexFactoryData, PreCacheParams, MeshFillMode, MeshCullMode, PSOInitializers);
+	AddPSOInitializer<HairVisibilityRenderMode_Transmittance>(SceneTexturesConfig, Material, VertexFactoryData, PreCacheParams, MeshFillMode, MeshCullMode, PSOInitializers);
+	AddPSOInitializer<HairVisibilityRenderMode_TransmittanceAndHairCount>(SceneTexturesConfig, Material, VertexFactoryData, PreCacheParams, MeshFillMode, MeshCullMode, PSOInitializers);
+	AddPSOInitializer<HairVisibilityRenderMode_PPLL>(SceneTexturesConfig, Material, VertexFactoryData, PreCacheParams, MeshFillMode, MeshCullMode, PSOInitializers);
+}
+
+template<EHairVisibilityRenderMode RenderMode, bool bCullingEnable>
+void FHairVisibilityProcessor::AddPSOInitializer(
+	const FSceneTexturesConfig& SceneTexturesConfig,
+	const FMaterial& Material,
+	const FPSOPrecacheVertexFactoryData& VertexFactoryData,
+	const FPSOPrecacheParams& PreCacheParams,
+	ERasterizerFillMode MeshFillMode,
+	ERasterizerCullMode MeshCullMode,
+	TArray<FPSOPrecacheData>& PSOInitializers)
+{
+	SetupDrawRenderState(RenderMode);
+
+	TMeshProcessorShaders<
+		FHairVisibilityVS<RenderMode, bCullingEnable>,
+		FHairVisibilityPS<RenderMode>> PassShaders;
+	{
+		FMaterialShaderTypes ShaderTypes;
+		ShaderTypes.AddShaderType<FHairVisibilityVS<RenderMode, bCullingEnable>>();
+		ShaderTypes.AddShaderType<FHairVisibilityPS<RenderMode>>();
+
+		FMaterialShaders Shaders;
+		if (!Material.TryGetShaders(ShaderTypes, VertexFactoryData.VertexFactoryType, Shaders))
+		{
+			return;
+		}
+
+		Shaders.TryGetVertexShader(PassShaders.VertexShader);
+		Shaders.TryGetPixelShader(PassShaders.PixelShader);
+	}
+
+	FGraphicsPipelineRenderTargetsInfo RenderTargetsInfo;
+	RenderTargetsInfo.NumSamples = 1;
+
+	if (RenderMode == HairVisibilityRenderMode_MSAA_Visibility)
+	{
+		const uint32 MSAASampleCount = GetMaxSamplePerPixel(GetFeatureLevelShaderPlatform(FeatureLevel));
+		RenderTargetsInfo.NumSamples = MSAASampleCount;
+
+		AddRenderTargetInfo(PF_R32_UINT, TexCreate_NoFastClear | TexCreate_RenderTargetable | TexCreate_ShaderResource, RenderTargetsInfo);
+		SetupDepthStencilInfo(PF_D24, SceneTexturesConfig.DepthCreateFlags, ERenderTargetLoadAction::ELoad,
+			ERenderTargetLoadAction::ENoAction, FExclusiveDepthStencil::DepthWrite_StencilNop, RenderTargetsInfo);
+	}
+	else if (RenderMode == HairVisibilityRenderMode_Transmittance || RenderMode == HairVisibilityRenderMode_TransmittanceAndHairCount)
+	{
+		AddRenderTargetInfo(PF_R32_FLOAT, TexCreate_RenderTargetable | TexCreate_ShaderResource, RenderTargetsInfo);
+		if (RenderMode == HairVisibilityRenderMode_TransmittanceAndHairCount)
+		{
+			AddRenderTargetInfo(PF_G32R32F, TexCreate_RenderTargetable | TexCreate_ShaderResource, RenderTargetsInfo);
+		}
+
+		SetupDepthStencilInfo(PF_DepthStencil, SceneTexturesConfig.DepthCreateFlags, ERenderTargetLoadAction::ELoad,
+			ERenderTargetLoadAction::ELoad, FExclusiveDepthStencil::DepthRead_StencilRead, RenderTargetsInfo);
+	}
+
+	AddGraphicsPipelineStateInitializer(
+		VertexFactoryData,
+		Material,
+		PassDrawRenderState,
+		RenderTargetsInfo,
+		PassShaders,
+		MeshFillMode,
+		MeshCullMode,
+		(EPrimitiveType)PreCacheParams.PrimitiveType,
+		EMeshPassFeatures::Default,
+		true /*bRequired*/,
+		PSOInitializers);
+}
+
+static const TCHAR* HairVisibilityPassName = TEXT("HairVisibility");
+
 FHairVisibilityProcessor::FHairVisibilityProcessor(
 	const FScene* Scene,
+	ERHIFeatureLevel::Type FeatureLevel,
 	const FSceneView* InViewIfDynamicMeshCommand,
-	const FMeshPassProcessorRenderState& InPassDrawRenderState,
 	const EHairVisibilityRenderMode InRenderMode,
 	FDynamicPassMeshDrawListContext* InDrawListContext)
-	: FMeshPassProcessor(EMeshPass::Num, Scene, Scene->GetFeatureLevel(), InViewIfDynamicMeshCommand, InDrawListContext)
+	: FMeshPassProcessor(HairVisibilityPassName, Scene, FeatureLevel, InViewIfDynamicMeshCommand, InDrawListContext)
 	, RenderMode(InRenderMode)
-	, PassDrawRenderState(InPassDrawRenderState)
 {
+	SetupDrawRenderState(RenderMode);
 }
 
+IPSOCollector* CreatePSOCollectorHairVisibility(ERHIFeatureLevel::Type FeatureLevel)
+{ 
+	const EShaderPlatform ShaderPlatform = GetFeatureLevelShaderPlatform(FeatureLevel);
+	EHairVisibilityRenderMode RenderMode = GetHairVisibilityRenderMode(ShaderPlatform);
+	return new FHairVisibilityProcessor(nullptr, FeatureLevel, nullptr, RenderMode, nullptr);
+} 
+FRegisterPSOCollectorCreateFunction RegisterPSOCollectorHairVisibility(&CreatePSOCollectorHairVisibility, EShadingPath::Deferred, HairVisibilityPassName);
+
 ///////////////////////////////////////////////////////////////////////////////////////////////////
 // Clear uint texture
 class FClearUIntGraphicPS : public FGlobalShader
@@ -2182,44 +2412,16 @@ static void AddHairVisibilityCommonPass(
 	{
 		check(IsInRenderingThread());
 
-		FMeshPassProcessorRenderState DrawRenderState;
-
+		FHairVisibilityProcessor MeshProcessor(Scene, Scene->GetFeatureLevel(),  ViewInfo, RenderMode, ShadowContext);
+		for (const FHairStrandsMacroGroupData& MacroGroupData : MacroGroupDatas)
 		{
-			if (RenderMode == HairVisibilityRenderMode_MSAA_Visibility)
-			{
-				DrawRenderState.SetBlendState(TStaticBlendState<CW_RGBA, BO_Add, BF_One, BF_Zero, BO_Add, BF_One, BF_Zero>::GetRHI());
-				DrawRenderState.SetDepthStencilState(TStaticDepthStencilState<true, CF_DepthNearOrEqual>::GetRHI());
-			}
-			else if (RenderMode == HairVisibilityRenderMode_Transmittance)
-			{
-				DrawRenderState.SetBlendState(TStaticBlendState<CW_RED, BO_Add, BF_DestColor, BF_Zero, BO_Add, BF_Zero, BF_Zero>::GetRHI());
-				DrawRenderState.SetDepthStencilState(TStaticDepthStencilState<false, CF_DepthNearOrEqual>::GetRHI());
-			}
-			else if (RenderMode == HairVisibilityRenderMode_TransmittanceAndHairCount)
+			for (const FHairStrandsMacroGroupData::PrimitiveInfo& PrimitiveInfo : MacroGroupData.PrimitivesInfos)
 			{
-				DrawRenderState.SetBlendState(TStaticBlendState<
-					CW_RED, BO_Add, BF_DestColor, BF_Zero, BO_Add, BF_Zero, BF_Zero,
-					CW_RG, BO_Add, BF_One, BF_One, BO_Add, BF_Zero, BF_Zero>::GetRHI());
-				DrawRenderState.SetDepthStencilState(TStaticDepthStencilState<false, CF_DepthNearOrEqual>::GetRHI());
-			}
-			else if (RenderMode == HairVisibilityRenderMode_PPLL)
-			{
-				DrawRenderState.SetBlendState(TStaticBlendState<>::GetRHI());
-				DrawRenderState.SetDepthStencilState(TStaticDepthStencilState<false, CF_DepthNearOrEqual>::GetRHI());
-			}
-
-			FHairVisibilityProcessor MeshProcessor(Scene, ViewInfo, DrawRenderState, RenderMode, ShadowContext);
-			
-			for (const FHairStrandsMacroGroupData& MacroGroupData : MacroGroupDatas)
-			{
-				for (const FHairStrandsMacroGroupData::PrimitiveInfo& PrimitiveInfo : MacroGroupData.PrimitivesInfos)
+				if (const FMeshBatch* MeshBatch = PrimitiveInfo.Mesh)
 				{
-					if (const FMeshBatch* MeshBatch = PrimitiveInfo.Mesh)
-					{
-						const uint64 BatchElementMask = ~0ull;
-						const float HairCoverageScale = PrimitiveInfo.PublicDataPtr->GetActiveStrandsCoverageScale();
-						MeshProcessor.AddMeshBatch(*MeshBatch, BatchElementMask, PrimitiveInfo.PrimitiveSceneProxy, -1, MacroGroupData.MacroGroupId, PrimitiveInfo.MaterialId, HairCoverageScale, PrimitiveInfo.IsCullingEnable());
-					}
+					const uint64 BatchElementMask = ~0ull;
+					const float HairCoverageScale = PrimitiveInfo.PublicDataPtr->GetActiveStrandsCoverageScale();
+					MeshProcessor.AddMeshBatch(*MeshBatch, BatchElementMask, PrimitiveInfo.PrimitiveSceneProxy, -1, MacroGroupData.MacroGroupId, PrimitiveInfo.MaterialId, HairCoverageScale, PrimitiveInfo.IsCullingEnable());
 				}
 			}
 		}
