Compile fix

diff --git a/Engine/Source/Runtime/Renderer/Private/HairStrands/HairStrandsVisibility.cpp b/Engine/Source/Runtime/Renderer/Private/HairStrands/HairStrandsVisibility.cpp
--- a/Engine/Source/Runtime/Renderer/Private/HairStrands/HairStrandsVisibility.cpp
+++ b/Engine/Source/Runtime/Renderer/Private/HairStrands/HairStrandsVisibility.cpp
@@ -1709,7 +1709,7 @@ void FHairVisibilityProcessor::CollectPSOInitializers(
 	AddPSOInitializer<HairVisibilityRenderMode_PPLL>(SceneTexturesConfig, Material, VertexFactoryData, PreCacheParams, MeshFillMode, MeshCullMode, PSOInitializers);
 }
 
-template<EHairVisibilityRenderMode RenderMode, bool bCullingEnable>
+template<EHairVisibilityRenderMode TRenderMode, bool bCullingEnable>
 void FHairVisibilityProcessor::AddPSOInitializer(
 	const FSceneTexturesConfig& SceneTexturesConfig,
 	const FMaterial& Material,
@@ -1719,15 +1719,15 @@ void FHairVisibilityProcessor::AddPSOInitializer(
 	ERasterizerCullMode MeshCullMode,
 	TArray<FPSOPrecacheData>& PSOInitializers)
 {
-	SetupDrawRenderState(RenderMode);
+	SetupDrawRenderState(TRenderMode);
 
 	TMeshProcessorShaders<
-		FHairVisibilityVS<RenderMode, bCullingEnable>,
-		FHairVisibilityPS<RenderMode>> PassShaders;
+		FHairVisibilityVS<TRenderMode, bCullingEnable>,
+		FHairVisibilityPS<TRenderMode>> PassShaders;
 	{
 		FMaterialShaderTypes ShaderTypes;
-		ShaderTypes.AddShaderType<FHairVisibilityVS<RenderMode, bCullingEnable>>();
-		ShaderTypes.AddShaderType<FHairVisibilityPS<RenderMode>>();
+		ShaderTypes.AddShaderType<FHairVisibilityVS<TRenderMode, bCullingEnable>>();
+		ShaderTypes.AddShaderType<FHairVisibilityPS<TRenderMode>>();
 
 		FMaterialShaders Shaders;
 		if (!Material.TryGetShaders(ShaderTypes, VertexFactoryData.VertexFactoryType, Shaders))
@@ -1742,7 +1742,7 @@ void FHairVisibilityProcessor::AddPSOInitializer(
 	FGraphicsPipelineRenderTargetsInfo RenderTargetsInfo;
 	RenderTargetsInfo.NumSamples = 1;
 
-	if (RenderMode == HairVisibilityRenderMode_MSAA_Visibility)
+	if (TRenderMode == HairVisibilityRenderMode_MSAA_Visibility)
 	{
 		const uint32 MSAASampleCount = GetMaxSamplePerPixel(GetFeatureLevelShaderPlatform(FeatureLevel));
 		RenderTargetsInfo.NumSamples = MSAASampleCount;
@@ -1751,10 +1751,10 @@ void FHairVisibilityProcessor::AddPSOInitializer(
 		SetupDepthStencilInfo(PF_D24, SceneTexturesConfig.DepthCreateFlags, ERenderTargetLoadAction::ELoad,
 			ERenderTargetLoadAction::ENoAction, FExclusiveDepthStencil::DepthWrite_StencilNop, RenderTargetsInfo);
 	}
-	else if (RenderMode == HairVisibilityRenderMode_Transmittance || RenderMode == HairVisibilityRenderMode_TransmittanceAndHairCount)
+	else if (TRenderMode == HairVisibilityRenderMode_Transmittance || TRenderMode == HairVisibilityRenderMode_TransmittanceAndHairCount)
 	{
 		AddRenderTargetInfo(PF_R32_FLOAT, TexCreate_RenderTargetable | TexCreate_ShaderResource, RenderTargetsInfo);
-		if (RenderMode == HairVisibilityRenderMode_TransmittanceAndHairCount)
+		if (TRenderMode == HairVisibilityRenderMode_TransmittanceAndHairCount)
 		{
 			AddRenderTargetInfo(PF_G32R32F, TexCreate_RenderTargetable | TexCreate_ShaderResource, RenderTargetsInfo);
 		}
