PR #9893  Fix r.Visibility.FrustumCull.Enabled 0

It was previously culling all primitives instead of the opposite.

diff --git a/Engine/Source/Runtime/Renderer/Private/SceneVisibility.cpp b/Engine/Source/Runtime/Renderer/Private/SceneVisibility.cpp
--- a/Engine/Source/Runtime/Renderer/Private/SceneVisibility.cpp
+++ b/Engine/Source/Runtime/Renderer/Private/SceneVisibility.cpp
@@ -3370,6 +3370,11 @@ void FVisibilityViewPacket::BeginInitVisibility()
 
 	bool bShouldVisibilityCull = GFrustumCullEnabled;
 
+	if (!bShouldVisibilityCull)
+	{
+		View.PrimitiveVisibilityMap.SetRange(0, View.PrimitiveVisibilityMap.Num(), true); // Mark all primitives as visible when not visibility culling
+	}
+
 	// Development builds sometimes override frustum culling, e.g. dependent views in the editor.
 #if !(UE_BUILD_SHIPPING || UE_BUILD_TEST)
 	if (ViewState)
