Fix missing initialization of view index in light function atlas.

This is no important as, this index is not used yet.

diff --git a/Engine/Source/Runtime/Renderer/Private/LightFunctionAtlas.h b/Engine/Source/Runtime/Renderer/Private/LightFunctionAtlas.h
--- a/Engine/Source/Runtime/Renderer/Private/LightFunctionAtlas.h
+++ b/Engine/Source/Runtime/Renderer/Private/LightFunctionAtlas.h
@@ -60,7 +60,7 @@ private:
 struct FLightFunctionAtlasViewData
 {
 	FLightFunctionAtlasViewData() {}
-	FLightFunctionAtlasViewData(FLightFunctionAtlasSceneData* InSceneData, uint32 ViewIndex) : SceneData(InSceneData), ViewIndex(0) {}
+	FLightFunctionAtlasViewData(FLightFunctionAtlasSceneData* InSceneData, uint32 InViewIndex) : SceneData(InSceneData), ViewIndex(InViewIndex) {}
 
 	FLightFunctionAtlas* GetLightFunctionAtlas()				const { return SceneData ? SceneData->GetLightFunctionAtlas() : nullptr; }
 	bool GetLightFunctionAtlasEnabled()							const { return SceneData ? SceneData->GetLightFunctionAtlasEnabled() : false; }
