Move occlusion buffer clear before shadowmap rendering to allow some GPU work overlap bettwen shadowmap and basepass rendering in a mobile renderer
VKRHI UAV clears insert overly conservative barriers which prevent work overlap

diff --git a/Engine/Source/Runtime/Renderer/Private/MobileShadingRenderer.cpp b/Engine/Source/Runtime/Renderer/Private/MobileShadingRenderer.cpp
--- a/Engine/Source/Runtime/Renderer/Private/MobileShadingRenderer.cpp
+++ b/Engine/Source/Runtime/Renderer/Private/MobileShadingRenderer.cpp
@@ -1003,6 +1003,11 @@ void FMobileSceneRenderer::Render(FRDGBuilder& GraphBuilder)
 
 	GraphBuilder.SetCommandListStat(GET_STATID(STAT_CLMM_AfterInitViews));
 
+	if (bRendererOutputFinalSceneColor)
+	{
+		BeginOcclusionScope(GraphBuilder, Views);
+	}
+	
 	if (bIsFirstSceneRenderer)
 	{
 		GraphBuilder.SetFlushResourcesRHI();
@@ -1086,6 +1091,7 @@ void FMobileSceneRenderer::Render(FRDGBuilder& GraphBuilder)
 
 		GraphBuilder.SetCommandListStat(GET_STATID(STAT_CLMM_Shadows));
 		RenderShadowDepthMaps(GraphBuilder, nullptr, InstanceCullingManager, ExternalAccessQueue);
+		GraphBuilder.AddDispatchHint();
 
 		// Run local fog volume initialization before base pass and volumetric fog for all the culled instance instance data to be ready.
 		InitLocalFogVolumesForViews(Scene, Views, ViewFamily, GraphBuilder, ShouldRenderVolumetricFog(), bUseHalfResLocalFogVolume);
@@ -1097,7 +1103,6 @@ void FMobileSceneRenderer::Render(FRDGBuilder& GraphBuilder)
 		ExternalAccessQueue.Submit(GraphBuilder);
 
 		PollOcclusionQueriesPass(GraphBuilder);
-		GraphBuilder.AddDispatchHint();
 
 		// Custom depth
 		// bShouldRenderCustomDepth has been initialized in InitViews on mobile platform
@@ -1105,8 +1110,6 @@ void FMobileSceneRenderer::Render(FRDGBuilder& GraphBuilder)
 		{
 			RenderCustomDepthPass(GraphBuilder, SceneTextures.CustomDepth, SceneTextures.GetSceneTextureShaderParameters(FeatureLevel), {}, {});
 		}
-
-		BeginOcclusionScope(GraphBuilder, Views);
 	}
 	else
 	{
