Fix validation error due to mismatching stride between CPU/GPU definition.

diff --git a/Engine/Source/Runtime/Renderer/Private/LightFunctionAtlas.cpp b/Engine/Source/Runtime/Renderer/Private/LightFunctionAtlas.cpp
--- a/Engine/Source/Runtime/Renderer/Private/LightFunctionAtlas.cpp
+++ b/Engine/Source/Runtime/Renderer/Private/LightFunctionAtlas.cpp
@@ -637,9 +637,10 @@ void FLightFunctionAtlas::RenderLightFunctionAtlas(FRDGBuilder& GraphBuilder, TA
 		}
 
 		// Create the light instance data buffer SRV
+		const uint32 Float4Count = sizeof(FAtlasLightInfoData) / sizeof(FVector4f);
 		RDGLightInfoDataBuffer = CreateStructuredBuffer(
 			GraphBuilder, TEXT("LightFunctionAtlasLightInfoData"), 
-			sizeof(FAtlasLightInfoData), InitialLightInfoDataLightCount,
+			sizeof(FVector4f), Float4Count * InitialLightInfoDataLightCount,
 			reinterpret_cast<void*>(LightInfoDataBufferPtr), InitialLightInfoDataSize, ERDGInitialDataFlags::NoCopy);
 		LightFunctionAtlasGlobalParameters->LightInfoDataBuffer = GraphBuilder.CreateSRV(RDGLightInfoDataBuffer, PF_A32B32G32R32F);
 
