HWRT: Enable GPU instance culling by default.

diff --git a/Engine/Source/Runtime/Renderer/Private/RayTracing/RayTracingInstanceCulling.cpp b/Engine/Source/Runtime/Renderer/Private/RayTracing/RayTracingInstanceCulling.cpp
--- a/Engine/Source/Runtime/Renderer/Private/RayTracing/RayTracingInstanceCulling.cpp
+++ b/Engine/Source/Runtime/Renderer/Private/RayTracing/RayTracingInstanceCulling.cpp
@@ -49,8 +49,8 @@ static TAutoConsoleVariable<int32> CVarRayTracingCullingGroupIds(
 
 static TAutoConsoleVariable<int32> CVarRayTracingCullingUseGPU(
 	TEXT("r.RayTracing.Culling.UseGPU"),
-	0,
-	TEXT("Cull instances using GPU. Requires per instance culling enabled. Experimental (default=0)"),
+	1,
+	TEXT("Cull instances using GPU. Requires per instance culling enabled (default=1)."),
 	ECVF_RenderThreadSafe);
 
 RayTracing::ECullingMode RayTracing::GetCullingMode(const FEngineShowFlags& ShowFlags)
