VulkanRHI: Match component types and counts in FHairResourceTransitionPass.

#rnx

diff --git a/Engine/Shaders/Private/HairStrands/HairStrandsMesh.usf b/Engine/Shaders/Private/HairStrands/HairStrandsMesh.usf
--- a/Engine/Shaders/Private/HairStrands/HairStrandsMesh.usf
+++ b/Engine/Shaders/Private/HairStrands/HairStrandsMesh.usf
@@ -318,10 +318,15 @@ void MainCS(uint3 DispatchThreadId : SV_DispatchThreadID)
 	TPrefix##Buffers_15##TSuffix;
 
 TEMPLATE_BUFFERS(StructuredBuffer<float4> Structured, )
-TEMPLATE_BUFFERS(Buffer<uint> Vertex, )
+TEMPLATE_BUFFERS(Buffer<uint> VertexUInt, )
+TEMPLATE_BUFFERS(Buffer<uint4> VertexUInt4, )
+TEMPLATE_BUFFERS(Buffer<float4> VertexFloat4, )
 
 RWStructuredBuffer<uint> DummyOutput;
 uint DummyValue;
+uint NumVertexUint;
+uint NumVertexUint4;
+uint NumVertexFloat4;
 
 [numthreads(1, 1, 1)]
 void MainCS(uint3 DispatchThreadId : SV_DispatchThreadID)
@@ -329,10 +334,14 @@ void MainCS(uint3 DispatchThreadId : SV_DispatchThreadID)
 	if (DummyValue > 0)
 	{
 		uint Acc = 0;
-	#if PERMUTATION_STRUCTURED_BUFFER
+	#if PERMUTATION_BUFFER_TYPE == 0
 		TEMPLATE_BUFFERS(Acc += Structured, [0].x)
+	#elif PERMUTATION_BUFFER_TYPE == 1
+		TEMPLATE_BUFFERS(Acc += VertexUInt, [0])
+	#elif PERMUTATION_BUFFER_TYPE == 2
+		TEMPLATE_BUFFERS(Acc += VertexUInt4, [0].x)
 	#else
-		TEMPLATE_BUFFERS(Acc += Vertex, [0])
+		TEMPLATE_BUFFERS(Acc += VertexFloat4, [0].x)
 	#endif
 		DummyOutput[0] = Acc;
 	}
diff --git a/Engine/Source/Runtime/Renderer/Private/HairStrands/HairStrandsUtils.cpp b/Engine/Source/Runtime/Renderer/Private/HairStrands/HairStrandsUtils.cpp
--- a/Engine/Source/Runtime/Renderer/Private/HairStrands/HairStrandsUtils.cpp
+++ b/Engine/Source/Runtime/Renderer/Private/HairStrands/HairStrandsUtils.cpp
@@ -352,21 +352,47 @@ uint32 PackHairRenderInfoBits(
 
 class FHairResourceTransitionPass : public FGlobalShader
 {
+public:
+	static const int32 MaxBufferCount = 16;
+
 private:
 	DECLARE_GLOBAL_SHADER(FHairResourceTransitionPass);
 	SHADER_USE_PARAMETER_STRUCT(FHairResourceTransitionPass, FGlobalShader);
 
-	class FStructuredBuffer : SHADER_PERMUTATION_BOOL("PERMUTATION_STRUCTURED_BUFFER");
-	using FPermutationDomain = TShaderPermutationDomain<FStructuredBuffer>;
+	class FBufferType : SHADER_PERMUTATION_INT("PERMUTATION_BUFFER_TYPE", 4);
+	using FPermutationDomain = TShaderPermutationDomain<FBufferType>;
 
 	BEGIN_SHADER_PARAMETER_STRUCT(FParameters, )
 		SHADER_PARAMETER(uint32, DummyValue)
-		SHADER_PARAMETER_RDG_BUFFER_SRV_ARRAY(Buffer, VertexBuffers, [16])
-		SHADER_PARAMETER_RDG_BUFFER_SRV_ARRAY(StructuredBuffer, StructuredBuffers, [16])
-		SHADER_PARAMETER_RDG_BUFFER_UAV(RWStructuredBuffer, DummyOutput)
-		END_SHADER_PARAMETER_STRUCT()
+		SHADER_PARAMETER_RDG_BUFFER_SRV_ARRAY(Buffer<uint>, VertexUIntBuffers, [MaxBufferCount])
+		SHADER_PARAMETER_RDG_BUFFER_SRV_ARRAY(Buffer<uint4>, VertexUInt4Buffers, [MaxBufferCount])
+		SHADER_PARAMETER_RDG_BUFFER_SRV_ARRAY(Buffer<float4>, VertexFloat4Buffers, [MaxBufferCount])
+		SHADER_PARAMETER_RDG_BUFFER_SRV_ARRAY(StructuredBuffer<float4>, StructuredBuffers, [MaxBufferCount])
+		SHADER_PARAMETER_RDG_BUFFER_UAV(RWStructuredBuffer<uint>, DummyOutput)
+	END_SHADER_PARAMETER_STRUCT()
 
 public:
+	// 0 : StructuredBuffer<float4>
+	// 1 : Buffer<uint>
+	// 2 : Buffer<uint4>
+	// 3 : Buffer<float4>
+	static const int32 PermutationCount = 4;
+	static int32 GetPermutationIndex(bool bStructured, bool bInteger, int32 NumComponents)
+	{
+		if (bStructured)
+		{
+			return 0;
+		}
+		else if (bInteger)
+		{
+			return (NumComponents == 1) ? 1 : 2;
+		}
+		else
+		{
+			return 3;
+		}
+	}
+
 	static bool IsSupported(EShaderPlatform InPlatform)
 	{
 		return IsHairStrandsSupported(EHairStrandsShaderType::Strands, InPlatform);
@@ -390,63 +416,93 @@ void AddTransitionPass(
 	EShaderPlatform InPlatform,
 	const TArray<FRDGBufferSRVRef>& Transitions)
 {
-	const uint32 ResourceCount = Transitions.Num();
+	const int32 ResourceCount = Transitions.Num();
 	if (ResourceCount == 0 || !FHairResourceTransitionPass::IsSupported(InPlatform))
 	{
 		return;
 	}
 
-	FRDGBufferSRVRef DummyVertexInput = GraphBuilder.CreateSRV(GSystemTextures.GetDefaultBuffer(GraphBuilder, 4u, 1u), PF_R32_UINT);
-	FRDGBufferSRVRef DummyStructuredInput= GraphBuilder.CreateSRV(GSystemTextures.GetDefaultStructuredBuffer(GraphBuilder, 16u));
-
-	FRDGBufferRef DummyOutput = GraphBuilder.CreateBuffer(FRDGBufferDesc::CreateStructuredDesc(4,1),TEXT("DummyOutput"));
+	FRDGBufferRef DummyOutput = GraphBuilder.CreateBuffer(FRDGBufferDesc::CreateStructuredDesc(4, 1), TEXT("DummyOutput"));
 	FRDGBufferUAVRef DummyOutputUAV = GraphBuilder.CreateUAV(DummyOutput, PF_R32_UINT, ERDGUnorderedAccessViewFlags::SkipBarrier);
 
-	check(Transitions[0]);
-	check(Transitions[0]->Desc.Buffer);
-	const bool bStructuredBuffer = EnumHasAnyFlags(Transitions[0]->Desc.Buffer->Desc.Usage, EBufferUsageFlags::StructuredBuffer);
+	TStaticArray<FRDGBufferSRVRef, FHairResourceTransitionPass::MaxBufferCount> SortedTransitions[FHairResourceTransitionPass::PermutationCount];
 
-	const uint32 MaxBufferCount = 16;
-	const uint32 PassCount = FMath::DivideAndRoundUp(ResourceCount, MaxBufferCount);
-	for (uint32 PassIt=0; PassIt< PassCount; ++PassIt)
+	FRDGBufferSRVRef DummyInputs[FHairResourceTransitionPass::PermutationCount];
+	DummyInputs[0] = GraphBuilder.CreateSRV(GSystemTextures.GetDefaultStructuredBuffer(GraphBuilder, 16u));
+	DummyInputs[1] = GraphBuilder.CreateSRV(GSystemTextures.GetDefaultBuffer(GraphBuilder, 4u, 1u), PF_R32_UINT);
+	DummyInputs[2] = GraphBuilder.CreateSRV(GSystemTextures.GetDefaultBuffer(GraphBuilder, 16u, 1u), PF_R32G32B32A32_UINT);
+	DummyInputs[3] = GraphBuilder.CreateSRV(GSystemTextures.GetDefaultBuffer(GraphBuilder, 16u, 1u), PF_A32B32G32R32F);
+
+	int32 ArrayCounts[FHairResourceTransitionPass::PermutationCount];
+	ArrayCounts[0] = 0;
+	ArrayCounts[1] = 0;
+	ArrayCounts[2] = 0;
+	ArrayCounts[3] = 0;
+
+	auto FlushArray = [&](int32 PermutationIndex, int32 TransitionCount)
 	{
 		FHairResourceTransitionPass::FParameters* PassParameters = GraphBuilder.AllocParameters<FHairResourceTransitionPass::FParameters>();
 		PassParameters->DummyValue = 0;
 		PassParameters->DummyOutput = DummyOutputUAV;
 
-		const uint32 PassResourceOffset = PassIt * MaxBufferCount;
-		const uint32 PassResourceCount = FMath::Min(int32(MaxBufferCount), int32(ResourceCount) - int32(PassResourceOffset));
-		if (bStructuredBuffer)
+		TShaderResourceParameterArray<FRDGBufferSRV*, FHairResourceTransitionPass::MaxBufferCount>* ParamArray = nullptr;
+		switch (PermutationIndex)
 		{
-			for (uint32 ResourceIt = 0; ResourceIt < PassResourceCount; ++ResourceIt)
-			{
-				PassParameters->StructuredBuffers[ResourceIt] = Transitions[PassResourceOffset + ResourceIt];
-			}
-			for (uint32 ResourceIt = PassResourceCount; ResourceIt < MaxBufferCount; ++ResourceIt)
-			{
-				PassParameters->StructuredBuffers[ResourceIt] = DummyStructuredInput;
-			}
+		case 0: ParamArray = &PassParameters->StructuredBuffers; break;
+		case 1: ParamArray = &PassParameters->VertexUIntBuffers; break;
+		case 2: ParamArray = &PassParameters->VertexUInt4Buffers; break;
+		case 3: ParamArray = &PassParameters->VertexFloat4Buffers; break;
+		default: checkNoEntry();
+		};
+
+		for (int32 ResourceIndex = 0; ResourceIndex < TransitionCount; ++ResourceIndex)
+		{
+			(*ParamArray)[ResourceIndex] = SortedTransitions[PermutationIndex][ResourceIndex];
 		}
-		else
+
+		for (int32 ResourceIndex = TransitionCount; ResourceIndex < FHairResourceTransitionPass::MaxBufferCount; ++ResourceIndex)
 		{
-			for (uint32 ResourceIt = 0; ResourceIt < PassResourceCount; ++ResourceIt)
-			{
-				PassParameters->VertexBuffers[ResourceIt] = Transitions[PassResourceOffset + ResourceIt];
-			}
-			for (uint32 ResourceIt = PassResourceCount; ResourceIt < MaxBufferCount; ++ResourceIt)
-			{
-				PassParameters->VertexBuffers[ResourceIt] = DummyVertexInput;
-			}
+			(*ParamArray)[ResourceIndex] = DummyInputs[PermutationIndex];
 		}
+
 		FHairResourceTransitionPass::FPermutationDomain PermutationVector;
-		PermutationVector.Set<FHairResourceTransitionPass::FStructuredBuffer>(bStructuredBuffer);
+		PermutationVector.Set<FHairResourceTransitionPass::FBufferType>(PermutationIndex);
 		TShaderMapRef<FHairResourceTransitionPass> ComputeShader(ShaderMap, PermutationVector);
+
 		FComputeShaderUtils::AddPass(
 			GraphBuilder,
-			RDG_EVENT_NAME("HairStrands::ResourceTransitions"),
+			RDG_EVENT_NAME("HairStrands::ResourceTransitions(P=%d)", PermutationIndex),
 			ERDGPassFlags::Compute | ERDGPassFlags::NeverCull,
 			ComputeShader,
 			PassParameters,
 			FIntVector(1, 1, 1));
+
+		ArrayCounts[PermutationIndex] = 0;
+	};
+
+	// Sort the transitions into each of the arrays based on the shader permutation required to transition them
+	for (int32 TransitionIndex = 0; TransitionIndex < Transitions.Num(); ++TransitionIndex)
+	{
+		const bool bStructuredBuffer = EnumHasAnyFlags(Transitions[TransitionIndex]->Desc.Buffer->Desc.Usage, EBufferUsageFlags::StructuredBuffer);
+		const EPixelFormat PixelFormat = Transitions[TransitionIndex]->Desc.Format;
+		const bool bIsIntegerFormat = IsInteger(PixelFormat);
+		const int32 NumComponents = GPixelFormats[PixelFormat].NumComponents;
+
+		const int32 PermutationIndex = FHairResourceTransitionPass::GetPermutationIndex(bStructuredBuffer, bIsIntegerFormat, NumComponents);
+
+		SortedTransitions[PermutationIndex][ArrayCounts[PermutationIndex]++] = Transitions[TransitionIndex];
+
+		if (ArrayCounts[PermutationIndex] >= FHairResourceTransitionPass::MaxBufferCount)
+		{
+			FlushArray(PermutationIndex, FHairResourceTransitionPass::MaxBufferCount);
+		}
+	}
+
+	for (int32 PermutationIndex = 0; PermutationIndex < FHairResourceTransitionPass::PermutationCount; ++PermutationIndex)
+	{
+		if (ArrayCounts[PermutationIndex] > 0)
+		{
+			FlushArray(PermutationIndex, ArrayCounts[PermutationIndex]);
+		}
 	}
 }
\ No newline at end of file
