* Add support for directional light's ShadowAmount for hair mask shadow
* Fix incorrect shadow mask channel writting in hair mask shadow.

diff --git a/Engine/Shaders/Private/HairStrands/HairStrandsDeepShadowMask.usf b/Engine/Shaders/Private/HairStrands/HairStrandsDeepShadowMask.usf
--- a/Engine/Shaders/Private/HairStrands/HairStrandsDeepShadowMask.usf
+++ b/Engine/Shaders/Private/HairStrands/HairStrandsDeepShadowMask.usf
@@ -9,12 +9,18 @@
 #include "HairStrandsDeepShadowCommon.ush"
 #include "HairStrandsDeepTransmittanceCommon.ush"
 
-float4 WriteShadowMaskOutput(const float InFadedShadow)
+float4 WriteShadowMaskOutput(const float InFadedShadow, uint InEncodingType)
 {
 	// Depending of if the light is direction (i.e., WholeSceneLight) or a local light (Point/Spot/Rect), 
 	// we change the output layout
 	const float EncodedShadow = EncodeLightAttenuation(InFadedShadow);
-	return float4(EncodedShadow, EncodedShadow, EncodedShadow, EncodedShadow);
+	switch (InEncodingType)
+	{
+		/* Default    */case 0: return EncodedShadow;
+		/* WholeLight */case 1: return float4(EncodedShadow, 1.f, 1.f, 1.f);
+		/* LocalLight */case 2: return float4(1.f, 1.f, EncodedShadow, 1.f);
+	}
+	return EncodedShadow;
 }
 
 ////////////////////////////////////////////////////////////////////////////////////////////////
@@ -46,6 +52,8 @@ uint				DeepShadow_KernelType;
 float				DeepShadow_DepthBiasScale;
 float				DeepShadow_DensityScale;
 float4				DeepShadow_LayerDepths;
+float 				FadeAlpha;
+uint				EncodingType;
 
 Texture2D<float>	DeepShadow_FrontDepthTexture;
 Texture2D<float4>	DeepShadow_DomTexture;
@@ -384,7 +392,9 @@ void MainPS(
 		Visibility = IsCloser(FrontDepth + ShadowBias, LightSpacePosition.z) ? 0 : 1;
 	#endif
 	}
-	OutColor = WriteShadowMaskOutput(Visibility);
+
+	const float FadedShadow = lerp(1.0f, Visibility, FadeAlpha);
+	OutColor = WriteShadowMaskOutput(FadedShadow, EncodingType);
 
 }
 #endif // SHADER_SHADOWMASK_DEEPSHADOW
@@ -405,6 +415,8 @@ uint				Voxel_MacroGroupCount;
 uint				Voxel_MacroGroupId;
 uint				Voxel_RandomType;
 Texture2D<float4>	RayMarchMaskTexture;
+float 				FadeAlpha;
+uint				EncodingType;
 
 // The hair shadow mask is rendered by tracing ray/cone through the voxel structure. The traversing is done with 
 // point sampling and jittering to get the smooth aspect, and avoid expensive trilinear sampling in a sparse 
@@ -440,6 +452,8 @@ void MainPS(
 	FScreenVertexOutput Input,
 	out float4 OutColor : SV_Target0)
 {
+	OutColor = 1.f;
+
 	const float2 UV = Input.UV;
 	const uint2 PixelCoord = floor(Input.Position.xy);
 	const float SceneDepth = ConvertFromDeviceZ(SceneDepthTexture.Load(uint3(PixelCoord, 0)).x);
@@ -459,6 +473,7 @@ void MainPS(
 	// PCF or simple point sampler
 	float Visibility = 1;
 	bool bHasValidOutput = false;
+	bool bWriteResult = false;
 	if (bIsValidOpaquePixel)
 	{
 		const float3 TranslatedLightPosition= GetTranslatedLightPosition(Voxel_TranslatedLightPosition_LightDirection, TranslatedWorldPosition);
@@ -473,7 +488,6 @@ void MainPS(
 			NormalizedDepthBias = (VirtualVoxel.DepthBiasScale_Shadow * LightDirection + PositionBiasScale * (Random * 2 - 1));
 		}
 
-		FHairTraversalResult Result = InitHairTraversalResult();
 
 		FVirtualVoxelCommonDesc CommonDesc;
 		CommonDesc.PageCountResolution	= VirtualVoxel.PageCountResolution;
@@ -512,7 +526,7 @@ void MainPS(
 					TraversalSettings.PixelRadius = SceneDepth * VirtualVoxel.HairCoveragePixelRadiusAtDepth1;
 					TraversalSettings.bCastShadow = true;
 
-					Result = ComputeHairCountVirtualVoxel(
+					FHairTraversalResult Result = ComputeHairCountVirtualVoxel(
 						TranslatedWorldPosition + NormalizedDepthBias * NodeDesc.VoxelWorldSize,
 						TranslatedLightPosition,
 						CommonDesc,
@@ -520,8 +534,9 @@ void MainPS(
 						VirtualVoxel.PageIndexBuffer,
 						VirtualVoxel.PageTexture,
 						TraversalSettings);
-
 					Visibility = min(Visibility, saturate(1 - Result.HairCount));
+
+					bWriteResult = bWriteResult || Result.HairCount > 0;
 					bHasValidOutput = true;
 				}
 			}
@@ -532,7 +547,12 @@ void MainPS(
 	{
 		discard;
 	}
-	OutColor = WriteShadowMaskOutput(Visibility);
+
+	if (bWriteResult)
+	{
+		const float FadedShadow = lerp(1.0f, Visibility, FadeAlpha);
+		OutColor = WriteShadowMaskOutput(FadedShadow, EncodingType);
+	}
 }
 
 #endif // SHADER_SHADOWMASK_VOXEL
\ No newline at end of file
diff --git a/Engine/Source/Runtime/Renderer/Private/HairStrands/HairStrandsTransmittance.cpp b/Engine/Source/Runtime/Renderer/Private/HairStrands/HairStrandsTransmittance.cpp
--- a/Engine/Source/Runtime/Renderer/Private/HairStrands/HairStrandsTransmittance.cpp
+++ b/Engine/Source/Runtime/Renderer/Private/HairStrands/HairStrandsTransmittance.cpp
@@ -452,7 +452,8 @@ class FHairStrandsVoxelShadowMaskPS : public FGlobalShader
 		SHADER_PARAMETER(uint32, Voxel_MacroGroupCount)
 		SHADER_PARAMETER(uint32, Voxel_MacroGroupId)
 		SHADER_PARAMETER(uint32, Voxel_RandomType)
-		SHADER_PARAMETER(uint32, bIsWholeSceneLight)
+		SHADER_PARAMETER(uint32, EncodingType)
+		SHADER_PARAMETER(float, FadeAlpha)
 
 		SHADER_PARAMETER_RDG_TEXTURE(Texture2D, RayMarchMaskTexture)
 		SHADER_PARAMETER_RDG_TEXTURE(Texture2D, SceneDepthTexture)
@@ -483,6 +484,8 @@ static void AddHairStrandsVoxelShadowMaskPass(
 	FRDGTextureRef SceneDepthTexture,
 	const FViewInfo& View,
 	const FLightSceneInfo* LightSceneInfo,
+	const float FadeAlpha,
+	const uint32 EncodingType,
 	const bool bProjectingForForwardShading,
 	FRDGTextureRef& OutShadowMask)
 {
@@ -512,8 +515,7 @@ static void AddHairStrandsVoxelShadowMaskPass(
 	const FIntRect Viewport = View.ViewRect;
 	const uint32 OutputChannel = bProjectingForForwardShading ? LightSceneInfo->GetDynamicShadowMapChannel() : ~0;
 	const FIntPoint Resolution = OutShadowMask->Desc.Extent;
-	const bool bIsWholeSceneLight = LightSceneInfo->Proxy->GetLightType() == LightType_Directional;
-
+	
 	for (int32 GroupIt=0, GroupCount=View.HairStrandsViewData.MacroGroupDatas.Num(); GroupIt < GroupCount; ++GroupIt)
 	{
 		if (bOnePass && GroupIt > 0)
@@ -524,13 +526,14 @@ static void AddHairStrandsVoxelShadowMaskPass(
 		const FHairStrandsMacroGroupData& MacroGroupData = View.HairStrandsViewData.MacroGroupDatas[GroupIt];
 
 		FHairStrandsVoxelShadowMaskPS::FParameters* Parameters = GraphBuilder.AllocParameters<FHairStrandsVoxelShadowMaskPS::FParameters>();
+		Parameters->FadeAlpha = FadeAlpha;
+		Parameters->EncodingType = EncodingType;
 		Parameters->ViewUniformBuffer = View.ViewUniformBuffer;
 		Parameters->SceneDepthTexture = SceneDepthTexture;	
 		Parameters->HairStrands = HairStrands::BindHairStrandsViewUniformParameters(View);
 		Parameters->VirtualVoxel = HairStrands::BindHairStrandsVoxelUniformParameters(View);
 		Parameters->LinearSampler = TStaticSamplerState<SF_Bilinear, AM_Clamp, AM_Clamp, AM_Clamp>::GetRHI();
 		Parameters->ShadowSampler = TStaticSamplerState<SF_Point, AM_Clamp, AM_Clamp, AM_Clamp>::GetRHI();
-		Parameters->bIsWholeSceneLight = bIsWholeSceneLight ? 1 : 0;
 		Parameters->Voxel_TranslatedLightPosition_LightDirection = GetLightTranslatedWorldPositionAndDirection(View, LightSceneInfo);
 		Parameters->Voxel_MacroGroupId	= MacroGroupData.MacroGroupId;
 		Parameters->Voxel_MacroGroupCount = GroupCount;
@@ -602,11 +605,12 @@ class FHairStrandsDeepShadowMaskPS : public FGlobalShader
 		SHADER_PARAMETER(uint32, DeepShadow_SlotIndex)
 		SHADER_PARAMETER(FIntPoint, DeepShadow_SlotResolution)
 		SHADER_PARAMETER(FMatrix44f, DeepShadow_CPUTranslatedWorldToLightTransform)
-		SHADER_PARAMETER(uint32, bIsWholeSceneLight)
 		SHADER_PARAMETER(float, DeepShadow_DepthBiasScale)
 		SHADER_PARAMETER(float, DeepShadow_DensityScale)
 		SHADER_PARAMETER(uint32, DeepShadow_bIsGPUDriven)
 		SHADER_PARAMETER(FVector4f, DeepShadow_LayerDepths)
+		SHADER_PARAMETER(float, FadeAlpha)
+		SHADER_PARAMETER(uint32, EncodingType)
 
 		SHADER_PARAMETER_RDG_TEXTURE(Texture2D, RayMarchMaskTexture)
 
@@ -637,7 +641,6 @@ IMPLEMENT_GLOBAL_SHADER(FHairStrandsDeepShadowMaskPS, "/Engine/Private/HairStran
 
 struct FHairStrandsDeepShadowParams
 {
-	bool			bIsWholeSceneLight = false;
 	uint32			OutputChannel = ~0;
 	FRDGBufferSRVRef DeepShadow_ViewInfoBuffer = nullptr;
 	FMatrix			DeepShadow_CPUTranslatedWorldToLightTransform;
@@ -657,11 +660,15 @@ static void AddHairStrandsDeepShadowMaskPass(
 	FRDGTextureRef SceneDepthTexture,
 	const FViewInfo& View,
 	const FHairStrandsDeepShadowParams& Params,
+	const float FadeAlpha,
+	const uint32 EncodingType,
 	FRDGTextureRef& OutShadowMask)
 {
 	check(OutShadowMask);
 
 	FHairStrandsDeepShadowMaskPS::FParameters* Parameters = GraphBuilder.AllocParameters<FHairStrandsDeepShadowMaskPS::FParameters>();
+	Parameters->FadeAlpha = FadeAlpha;
+	Parameters->EncodingType = EncodingType;
 	Parameters->ViewUniformBuffer = View.ViewUniformBuffer;
 	Parameters->SceneDepthTexture = SceneDepthTexture;
 	Parameters->DeepShadow_CPUTranslatedWorldToLightTransform = FMatrix44f(Params.DeepShadow_CPUTranslatedWorldToLightTransform);
@@ -669,7 +676,6 @@ static void AddHairStrandsDeepShadowMaskPass(
 	Parameters->DeepShadow_DomTexture = Params.DeepShadow_LayerTexture;
 	Parameters->LinearSampler = TStaticSamplerState<SF_Bilinear, AM_Clamp, AM_Clamp, AM_Clamp>::GetRHI();
 	Parameters->ShadowSampler = TStaticSamplerState<SF_Point, AM_Clamp, AM_Clamp, AM_Clamp>::GetRHI();
-	Parameters->bIsWholeSceneLight = Params.bIsWholeSceneLight ? 1 : 0;
 	Parameters->DeepShadow_DepthBiasScale = Params.DeepShadow_DepthBiasScale;
 	Parameters->DeepShadow_DensityScale = Params.DeepShadow_DensityScale;
 	Parameters->DeepShadow_LayerDepths = Params.DeepShadow_LayerDepths;
@@ -919,7 +925,9 @@ FHairStrandsTransmittanceMaskData RenderHairStrandsTransmittanceMask(
 static void InternalRenderHairStrandsShadowMask(
 	FRDGBuilder& GraphBuilder,
 	const FViewInfo& View,
+	const uint32 ViewIndex,
 	const FLightSceneInfo* LightSceneInfo,
+	const TArrayView<FVisibleLightInfo>& VisibleLightInfos,
 	const FHairStrandsVisibilityData& InVisibilityData,
 	const FHairStrandsMacroGroupDatas& InMacroGroupDatas,
 	const FHairStrandsDeepShadowResources& DeepShadowResources,
@@ -938,6 +946,27 @@ static void InternalRenderHairStrandsShadowMask(
 	RDG_GPU_STAT_SCOPE(GraphBuilder, HairStrandsOpaqueMask);
 	const FMinimalSceneTextures& SceneTextures = View.GetSceneTextures();
 
+	const TArray<FProjectedShadowInfo*, SceneRenderingAllocator>& ShadowInfos = VisibleLightInfos[LightSceneInfo->Id].AllProjectedShadows;
+	float FadeAlpha = 1.0f; 
+	bool bIsWholeSceneDirectionalLightShadow = false;
+	bool bIsWholeSceneLocalLightShadow = false;
+	if (ShadowInfos.Num() > 0)
+	{
+		FadeAlpha = ShadowInfos[0]->FadeAlphas[ViewIndex];
+		bIsWholeSceneDirectionalLightShadow = ShadowInfos[0]->IsWholeSceneDirectionalShadow();
+		bIsWholeSceneLocalLightShadow = ShadowInfos[0]->IsWholeScenePointLightShadow();
+	}
+
+	uint32 EncodingType = 0;
+	if (bIsWholeSceneDirectionalLightShadow)
+	{
+		EncodingType = 1;
+	}
+	else if (bIsWholeSceneLocalLightShadow)
+	{
+		EncodingType = 2;
+	}
+
 	bool bHasDeepShadow = false;
 	if (!IsHairStrandsForVoxelTransmittanceAndShadowEnable())
 	{
@@ -958,7 +987,6 @@ static void InternalRenderHairStrandsShadowMask(
 				}
 
 				bHasDeepShadow = true;
-				const bool bIsWholeSceneLight = LightSceneInfo->Proxy->GetLightType() == LightType_Directional;
 
 				FHairStrandsDeepShadowParams Params;
 				Params.DeepShadow_AtlasSlotIndex = DomData.AtlasSlotIndex;
@@ -968,7 +996,6 @@ static void InternalRenderHairStrandsShadowMask(
 				Params.DeepShadow_AtlasRect = DomData.AtlasRect;
 				Params.DeepShadow_FrontDepthTexture = DeepShadowResources.DepthAtlasTexture;
 				Params.DeepShadow_LayerTexture = DeepShadowResources.LayersAtlasTexture;
-				Params.bIsWholeSceneLight = bIsWholeSceneLight;
 				Params.DeepShadow_DepthBiasScale = GetDeepShadowDepthBiasScale();
 				Params.DeepShadow_DensityScale = GetDeepShadowDensityScale();
 				Params.DeepShadow_LayerDepths = ComputeDeepShadowLayerDepths(DomData.LayerDistribution);
@@ -978,6 +1005,8 @@ static void InternalRenderHairStrandsShadowMask(
 					SceneTextures.Depth.Resolve,
 					View,
 					Params,
+					FadeAlpha,
+					EncodingType,
 					OutShadowMask);
 			}
 		}
@@ -991,6 +1020,8 @@ static void InternalRenderHairStrandsShadowMask(
 			SceneTextures.Depth.Resolve,
 			View,
 			LightSceneInfo,
+			FadeAlpha,
+			EncodingType,
 			bProjectingForForwardShading,
 			OutShadowMask);
 	}
@@ -1000,6 +1031,7 @@ void RenderHairStrandsShadowMask(
 	FRDGBuilder& GraphBuilder,
 	const TArray<FViewInfo>& Views,
 	const FLightSceneInfo* LightSceneInfo,
+	const TArrayView<FVisibleLightInfo>& VisibleLightInfos,
 	const bool bProjectingForForwardShading,
 	FRDGTextureRef OutShadowMask)
 {
@@ -1008,6 +1040,7 @@ void RenderHairStrandsShadowMask(
 		return;
 	}
 
+	uint32 ViewIndex = 0;
 	for (const FViewInfo& View : Views)
 	{
 		if (HairStrands::HasViewHairStrandsData(View))
@@ -1016,7 +1049,9 @@ void RenderHairStrandsShadowMask(
 			InternalRenderHairStrandsShadowMask(
 				GraphBuilder,
 				View,
+				ViewIndex++,
 				LightSceneInfo,
+				VisibleLightInfos,
 				View.HairStrandsViewData.VisibilityData,
 				View.HairStrandsViewData.MacroGroupDatas,
 				View.HairStrandsViewData.DeepShadowResources,
@@ -1031,6 +1066,7 @@ void RenderHairStrandsDeepShadowMask(
 	FRDGBuilder& GraphBuilder,
 	const TArray<FViewInfo>& Views,
 	const FLightSceneInfo* LightSceneInfo,
+	const TArrayView<FVisibleLightInfo>& VisibleLightInfos,
 	FRDGTextureRef OutShadowMask)
 {
 	if (Views.Num() == 0 || OutShadowMask == nullptr)
@@ -1044,6 +1080,7 @@ void RenderHairStrandsDeepShadowMask(
 		return;
 	}
 
+	uint32 ViewIndex = 0;
 	for (const FViewInfo& View : Views)
 	{
 		if (HairStrands::HasViewHairStrandsData(View))
@@ -1052,7 +1089,9 @@ void RenderHairStrandsDeepShadowMask(
 			InternalRenderHairStrandsShadowMask(
 				GraphBuilder, 
 				View, 
+				ViewIndex++,
 				LightSceneInfo, 
+				VisibleLightInfos,
 				View.HairStrandsViewData.VisibilityData,
 				View.HairStrandsViewData.MacroGroupDatas,
 				View.HairStrandsViewData.DeepShadowResources,
diff --git a/Engine/Source/Runtime/Renderer/Private/HairStrands/HairStrandsTransmittance.h b/Engine/Source/Runtime/Renderer/Private/HairStrands/HairStrandsTransmittance.h
--- a/Engine/Source/Runtime/Renderer/Private/HairStrands/HairStrandsTransmittance.h
+++ b/Engine/Source/Runtime/Renderer/Private/HairStrands/HairStrandsTransmittance.h
@@ -13,6 +13,7 @@
 class FLightSceneInfo;
 class FViewInfo;
 class FVirtualShadowMapArray;
+class FVisibleLightInfo;
 
 struct FHairStrandsTransmittanceMaskData
 {
@@ -25,6 +26,7 @@ void RenderHairStrandsShadowMask(
 	FRDGBuilder& GraphBuilder,
 	const TArray<FViewInfo>& Views,
 	const FLightSceneInfo* LightSceneInfo,
+	const TArrayView<FVisibleLightInfo>& VisibleLightInfos,
 	const bool bProjectingForForwardShading,
 	FRDGTextureRef ScreenShadowMaskTexture); 
 
@@ -33,6 +35,7 @@ void RenderHairStrandsDeepShadowMask(
 	FRDGBuilder& GraphBuilder,
 	const TArray<FViewInfo>& Views,
 	const FLightSceneInfo* LightSceneInfo,
+	const TArrayView<FVisibleLightInfo>& VisibleLightInfos,
 	FRDGTextureRef OutShadowMask);
 
 /// Output hair transmittance per hair sample for a given light
diff --git a/Engine/Source/Runtime/Renderer/Private/LightGridInjection.cpp b/Engine/Source/Runtime/Renderer/Private/LightGridInjection.cpp
--- a/Engine/Source/Runtime/Renderer/Private/LightGridInjection.cpp
+++ b/Engine/Source/Runtime/Renderer/Private/LightGridInjection.cpp
@@ -1057,7 +1057,7 @@ void FDeferredShadingSceneRenderer::RenderForwardShadowProjections(
 
 				if (bIsHairEnable)
 				{
-					RenderHairStrandsShadowMask(GraphBuilder, Views, LightSceneInfo, bProjectingForForwardShading, ForwardScreenSpaceShadowMask.Target);
+					RenderHairStrandsShadowMask(GraphBuilder, Views, LightSceneInfo, VisibleLightInfos, bProjectingForForwardShading, ForwardScreenSpaceShadowMask.Target);
 				}
 			}
 
diff --git a/Engine/Source/Runtime/Renderer/Private/LightRendering.cpp b/Engine/Source/Runtime/Renderer/Private/LightRendering.cpp
--- a/Engine/Source/Runtime/Renderer/Private/LightRendering.cpp
+++ b/Engine/Source/Runtime/Renderer/Private/LightRendering.cpp
@@ -1821,7 +1821,7 @@ void FDeferredShadingSceneRenderer::RenderLights(
 						// Inject deep shadow mask if the light supports it
 						if (bUseHairDeepShadow)
 						{
-							RenderHairStrandsDeepShadowMask(GraphBuilder, Views, &LightSceneInfo, ScreenShadowMaskTexture);
+							RenderHairStrandsDeepShadowMask(GraphBuilder, Views, &LightSceneInfo, VisibleLightInfos, ScreenShadowMaskTexture);
 						}
 					}
 					else 
@@ -1954,7 +1954,7 @@ void FDeferredShadingSceneRenderer::RenderLights(
 						// Inject deep shadow mask if the light supports it
 						if (bUseHairDeepShadow)
 						{
-							RenderHairStrandsShadowMask(GraphBuilder, Views, &LightSceneInfo, false /*bForwardShading*/, ScreenShadowMaskTexture);
+							RenderHairStrandsShadowMask(GraphBuilder, Views, &LightSceneInfo, VisibleLightInfos, false /*bForwardShading*/, ScreenShadowMaskTexture);
 						}
 					}
 					else // (OcclusionType == FOcclusionType::Shadowmap)
diff --git a/Engine/Source/Runtime/Renderer/Private/ShadowRendering.cpp b/Engine/Source/Runtime/Renderer/Private/ShadowRendering.cpp
--- a/Engine/Source/Runtime/Renderer/Private/ShadowRendering.cpp
+++ b/Engine/Source/Runtime/Renderer/Private/ShadowRendering.cpp
@@ -2324,7 +2324,7 @@ void FDeferredShadingSceneRenderer::RenderDeferredShadowProjections(
 		}
 		if (bNeedHairShadowMaskPass)
 		{
-			RenderHairStrandsShadowMask(GraphBuilder, Views, LightSceneInfo, false /*bForward*/, ScreenShadowMaskTexture);
+			RenderHairStrandsShadowMask(GraphBuilder, Views, LightSceneInfo, VisibleLightInfos, false /*bForward*/, ScreenShadowMaskTexture);
 		}
 	}
 }
