Robomerge fix

diff --git a/Engine/Source/Runtime/Renderer/Private/DeferredShadingRenderer.cpp b/Engine/Source/Runtime/Renderer/Private/DeferredShadingRenderer.cpp
--- a/Engine/Source/Runtime/Renderer/Private/DeferredShadingRenderer.cpp
+++ b/Engine/Source/Runtime/Renderer/Private/DeferredShadingRenderer.cpp
@@ -2186,7 +2186,7 @@ void FDeferredShadingSceneRenderer::Render(FRDGBuilder& GraphBuilder)
 
 				if (bNaniteEnabled && UseNaniteComputeMaterials())
 				{
-					Nanite::BuildShadingCommands(GraphBuilder, *Scene, CustomRenderPassViews, ENaniteMeshPass::BasePass, NaniteBasePassShadingCommands, true);
+					Nanite::BuildShadingCommands(GraphBuilder, *Scene, ENaniteMeshPass::BasePass, NaniteBasePassShadingCommands, true);
 				}
 
 				RenderBasePass(GraphBuilder, CustomRenderPassViews, SceneTextures, DBufferTextures, BasePassDepthStencilAccess, /*ForwardScreenSpaceShadowMaskTexture=*/nullptr, InstanceCullingManager, bNaniteEnabled, NaniteBasePassShadingCommands, NaniteRasterResults);
