Fill in cvar docs

diff --git a/Engine/Source/Runtime/Renderer/Private/VirtualShadowMaps/VirtualShadowMapClipmap.cpp b/Engine/Source/Runtime/Renderer/Private/VirtualShadowMaps/VirtualShadowMapClipmap.cpp
--- a/Engine/Source/Runtime/Renderer/Private/VirtualShadowMaps/VirtualShadowMapClipmap.cpp
+++ b/Engine/Source/Runtime/Renderer/Private/VirtualShadowMaps/VirtualShadowMapClipmap.cpp
@@ -85,7 +85,8 @@ static TAutoConsoleVariable<int32> CVarClipmapWPODisableDistance(
 static TAutoConsoleVariable<int32> CVarClipmapWPODisableDistanceLodBias(
 	TEXT("r.Shadow.Virtual.Clipmap.WPODisableDistance.LodBias"),
 	3,
-	TEXT("0 to disable"),	// TODO
+	TEXT("The number of clipmap levels further than the distance that an instance would be animated to allow shadow animation.\n")
+	TEXT("Typically 2-4 works well but may need to be adjusted for very low light angles with significant WPO movement."),
 	ECVF_Scalability | ECVF_RenderThreadSafe
 );
 
