Speculative fix for GPU page fault that could occur on some platforms due to seemingly corrupt Instance ID lookup buffer.

#localization none
#tests editor

diff --git a/Engine/Source/Runtime/Renderer/Private/SplineMeshSceneResources.cpp b/Engine/Source/Runtime/Renderer/Private/SplineMeshSceneResources.cpp
--- a/Engine/Source/Runtime/Renderer/Private/SplineMeshSceneResources.cpp
+++ b/Engine/Source/Runtime/Renderer/Private/SplineMeshSceneResources.cpp
@@ -36,6 +36,13 @@ static FAutoConsoleVariableRef CVarSplineMeshSceneTexturesCaptureNextUpdate(
 	ECVF_RenderThreadSafe
 );
 
+static TAutoConsoleVariable<bool> CVarSplineMeshSceneTexturesInstanceIDUploadCopy(
+	TEXT("r.SplineMesh.SceneTextures.InstanceIDUploadCopy"),
+	true,
+	TEXT("When true, will make a copy of the registered instance IDs on buffer upload."),
+	ECVF_RenderThreadSafe
+);
+
 BEGIN_SHADER_PARAMETER_STRUCT(FSplineMeshSceneResourceParameters, RENDERER_API)
 	SHADER_PARAMETER(FVector2f, SplineTextureInvExtent)
 	SHADER_PARAMETER_RDG_TEXTURE_SRV(Texture2D<float4>, SplinePosTexture)
@@ -288,7 +295,9 @@ FRDGBufferSRVRef FSplineMeshSceneExtension::GetInstanceIdLookupSRV(FRDGBuilder&
 	if (bNeedsUpload)
 	{
 		// Upload the contents
-		GraphBuilder.QueueBufferUpload<uint32>(InstanceIdLookup, RegisteredInstanceIds, ERDGInitialDataFlags::NoCopy);
+		const ERDGInitialDataFlags Flags = CVarSplineMeshSceneTexturesInstanceIDUploadCopy.GetValueOnRenderThread() ?
+			ERDGInitialDataFlags::None : ERDGInitialDataFlags::NoCopy;
+		GraphBuilder.QueueBufferUpload<uint32>(InstanceIdLookup, RegisteredInstanceIds, Flags);
 	}
 
 	bInstanceLookupDirty = false;
@@ -573,5 +582,5 @@ void FSplineMeshSceneRenderer::UpdateSceneUniformBuffer(FRDGBuilder& GraphBuilde
 		ShaderParams.SplineTextureInvExtent = InvExtent;
 	
 		SceneUniforms.Set(SceneUB::SplineMesh, ShaderParams);
-	}
+		}
 }
\ No newline at end of file
