Add holdout support for grooms.

diff --git a/Engine/Plugins/Runtime/HairStrands/Source/HairStrandsCore/Private/GroomComponent.cpp b/Engine/Plugins/Runtime/HairStrands/Source/HairStrandsCore/Private/GroomComponent.cpp
--- a/Engine/Plugins/Runtime/HairStrands/Source/HairStrandsCore/Private/GroomComponent.cpp
+++ b/Engine/Plugins/Runtime/HairStrands/Source/HairStrandsCore/Private/GroomComponent.cpp
@@ -455,6 +455,7 @@ public:
 			HairInstance->bForceCards = Component->bUseCards;
 			HairInstance->bUpdatePositionOffset = Component->RegisteredMeshComponent != nullptr;
 			HairInstance->bCastShadow = Component->CastShadow;
+			HairInstance->bHoldout = bHoldout;
 
 			if (HairInstance->Strands.IsValid() && HairInstance->Strands.VertexFactory == nullptr)
 			{
diff --git a/Engine/Plugins/Runtime/HairStrands/Source/HairStrandsCore/Private/GroomDebug.cpp b/Engine/Plugins/Runtime/HairStrands/Source/HairStrandsCore/Private/GroomDebug.cpp
--- a/Engine/Plugins/Runtime/HairStrands/Source/HairStrandsCore/Private/GroomDebug.cpp
+++ b/Engine/Plugins/Runtime/HairStrands/Source/HairStrandsCore/Private/GroomDebug.cpp
@@ -21,6 +21,7 @@
 #include "ShaderPrintParameters.h"
 #include "GroomComponent.h"
 #include "DataDrivenShaderPlatformInfo.h"
+#include "HairStrandsDefinitions.h"
 
 ///////////////////////////////////////////////////////////////////////////////////////////////////
 
@@ -833,14 +834,12 @@ static void AddHairDebugPrintInstancePass(
 
 				{
 					D.Data2 = FUintVector4(0);
-					D.Data2.X |= VisibilityType == EHairInstanceVisibilityType::StrandsPrimaryView           ? 0x1u  : 0u;
-					D.Data2.X |= VisibilityType == EHairInstanceVisibilityType::StrandsShadowView            ? 0x2u  : 0u;
-					D.Data2.X |= Instance->HairGroupPublicData->VFInput.Strands.Common.bScatterSceneLighting ? 0x4u  : 0u;
-					D.Data2.X |= Instance->HairGroupPublicData->VFInput.Strands.Common.bRaytracingGeometry   ? 0x8u  : 0u;
-					D.Data2.X |= Instance->HairGroupPublicData->VFInput.Strands.Common.bStableRasterization  ? 0x10u : 0u;
-					D.Data2.X |= Instance->HairGroupPublicData->bSupportVoxelization                         ? 0x20u : 0u;
-					D.Data2.X |= ActiveGroomCacheType == EGroomCacheType::Guides                             ? 0x40u : 0u;
-					D.Data2.X |= ActiveGroomCacheType == EGroomCacheType::Strands                            ? 0x80u : 0u;
+					D.Data2.X |= VisibilityType == EHairInstanceVisibilityType::StrandsPrimaryView           ? 0x1u : 0u;
+					D.Data2.X |= VisibilityType == EHairInstanceVisibilityType::StrandsShadowView            ? 0x2u : 0u;
+					D.Data2.X |= Instance->HairGroupPublicData->bSupportVoxelization                         ? 0x4u : 0u;
+					D.Data2.X |= ActiveGroomCacheType == EGroomCacheType::Guides                             ? 0x8u : 0u;
+					D.Data2.X |= ActiveGroomCacheType == EGroomCacheType::Strands                            ? 0x10u: 0u;
+					D.Data2.X |= Instance->HairGroupPublicData->VFInput.Strands.Common.Flags << 8u;
 					D.Data2.X |= uint32(FFloat16(Instance->HairGroupPublicData->ContinuousLODScreenSize).Encoded) << 16u;
 
 					D.Data2.Y = Instance->HairGroupPublicData->GetActiveStrandsPointCount();
diff --git a/Engine/Plugins/Runtime/HairStrands/Source/HairStrandsCore/Private/HairStrandsInterpolation.cpp b/Engine/Plugins/Runtime/HairStrands/Source/HairStrandsCore/Private/HairStrandsInterpolation.cpp
--- a/Engine/Plugins/Runtime/HairStrands/Source/HairStrandsCore/Private/HairStrandsInterpolation.cpp
+++ b/Engine/Plugins/Runtime/HairStrands/Source/HairStrandsCore/Private/HairStrandsInterpolation.cpp
@@ -1692,6 +1692,15 @@ FHairGroupPublicData::FVertexFactoryInput InternalComputeHairStrandsVertexInputD
 		OutVFInput.Strands.Common.CurveCount 		= Instance->HairGroupPublicData->GetActiveStrandsCurveCount();
 	}
 
+	
+	bool bRaytracingGeometry = false;
+#if RHI_RAYTRACING
+	// Flag bUseRaytracingGeometry only if RT geometry has been allocated for RayTracing view (not for PathTracing view). 
+	// This flag is used later for selecting if voxelization needs to flags voxel has shadow caster or if shadow casting is handled 
+	// by the RT geometry.
+	bRaytracingGeometry = Instance->Strands.RenRaytracingResource != nullptr && (Instance->Strands.ViewRayTracingMask & EHairViewRayTracingMask::RayTracing) != 0;
+#endif
+
 	OutVFInput.Strands.Common.RegisteredIndex = Instance->RegisteredIndex;
 	OutVFInput.Strands.Common.Radius = (GStrandHairWidth > 0 ? GStrandHairWidth : Instance->Strands.Modifier.HairWidth) * 0.5f;
 	OutVFInput.Strands.Common.RootScale = Instance->Strands.Modifier.HairRootScale;
@@ -1699,9 +1708,11 @@ FHairGroupPublicData::FVertexFactoryInput InternalComputeHairStrandsVertexInputD
 	OutVFInput.Strands.Common.RaytracingRadiusScale = (GHairRaytracingRadiusScale > 0 ? GHairRaytracingRadiusScale : Instance->Strands.Modifier.HairRaytracingRadiusScale);
 	OutVFInput.Strands.Common.Length = Instance->Strands.Modifier.HairLength;
 	OutVFInput.Strands.Common.Density = Instance->Strands.Modifier.HairShadowDensity;
-	OutVFInput.Strands.Common.bScatterSceneLighting = Instance->Strands.Modifier.bScatterSceneLighting;
-	OutVFInput.Strands.Common.bStableRasterization = Instance->Strands.Modifier.bUseStableRasterization;
-	OutVFInput.Strands.Common.bRaytracingGeometry = false;
+	OutVFInput.Strands.Common.Flags = 0;
+	OutVFInput.Strands.Common.Flags |= Instance->Strands.Modifier.bScatterSceneLighting ? HAIR_FLAGS_SCATTER_SCENE_LIGHT : 0u;
+	OutVFInput.Strands.Common.Flags |= Instance->Strands.Modifier.bUseStableRasterization ? HAIR_FLAGS_STABLE_RASTER : 0u;
+	OutVFInput.Strands.Common.Flags |= bRaytracingGeometry ? HAIR_FLAGS_RAYTRACING_GEOMETRY : 0u;
+	OutVFInput.Strands.Common.Flags |= Instance->bHoldout ? HAIR_FLAGS_HOLDOUT : 0u;
 	OutVFInput.Strands.Common.RaytracingProceduralSplits = GetHairRaytracingProceduralSplits();
 	OutVFInput.Strands.Common.GroupIndex = Instance->Debug.GroupIndex;
 	OutVFInput.Strands.Common.GroupCount = Instance->Debug.GroupCount;
@@ -1719,13 +1730,6 @@ FHairGroupPublicData::FVertexFactoryInput InternalComputeHairStrandsVertexInputD
 		GetHairStrandsAttributeParameter(*Instance->Strands.Data, OutVFInput.Strands.Common.Attributes);
 	}
 
-#if RHI_RAYTRACING
-	// Flag bUseRaytracingGeometry only if RT geometry has been allocated for RayTracing view (not for PathTracing view). 
-	// This flag is used later for selecting if voxelization needs to flags voxel has shadow caster or if shadow casting is handled 
-	// by the RT geometry.
-	OutVFInput.Strands.Common.bRaytracingGeometry = Instance->Strands.RenRaytracingResource != nullptr && (Instance->Strands.ViewRayTracingMask & EHairViewRayTracingMask::RayTracing) != 0;
-#endif
-
 	return OutVFInput;
 }
 
diff --git a/Engine/Plugins/Runtime/HairStrands/Source/HairStrandsCore/Public/GroomInstance.h b/Engine/Plugins/Runtime/HairStrands/Source/HairStrandsCore/Public/GroomInstance.h
--- a/Engine/Plugins/Runtime/HairStrands/Source/HairStrandsCore/Public/GroomInstance.h
+++ b/Engine/Plugins/Runtime/HairStrands/Source/HairStrandsCore/Public/GroomInstance.h
@@ -262,6 +262,7 @@ struct HAIRSTRANDSCORE_API FHairGroupInstance : public FHairStrandsInstance
 	bool					bUpdatePositionOffset = false;
 	bool					bCastShadow = true;
 	bool 					bSupportStreaming = true;
+	bool					bHoldout = false;
 	
 	// Deformed component to extract the bone buffer 
 	UMeshComponent*	 DeformedComponent = nullptr;
diff --git a/Engine/Shaders/Private/HairStrands/HairStrandsDebug.usf b/Engine/Shaders/Private/HairStrands/HairStrandsDebug.usf
--- a/Engine/Shaders/Private/HairStrands/HairStrandsDebug.usf
+++ b/Engine/Shaders/Private/HairStrands/HairStrandsDebug.usf
@@ -1,6 +1,7 @@
 // Copyright Epic Games, Inc. All Rights Reserved.
 
 #include "../Common.ush" 
+#include "/Engine/Shared/HairStrandsDefinitions.h"
 #include "../HairShadingCommon.ush"
 #include "../ColorMap.ush"
 
@@ -1208,7 +1209,7 @@ void MainCS(uint InGroupIndex : SV_GroupIndex, uint3 DispatchThreadId : SV_Dispa
 //		Print(Context2,(Sample.LightChannelMask & 0x1) ? 1u : 0u);				Newline(Context2);
 //		Print(Context2,(Sample.LightChannelMask & 0x2) ? 1u : 0u);				Newline(Context2);
 //		Print(Context2,(Sample.LightChannelMask & 0x4) ? 1u : 0u);				Newline(Context2);
-		Print(Context2, Sample.bScatterSceneLighting ? 1u : 0u, FontOrange);		Newline(Context2);
+		Print(Context2, HasHairFlags(Sample.Flags, HAIR_FLAGS_SCATTER_SCENE_LIGHT) ? 1u : 0u, FontOrange); Newline(Context2);
 		Print(Context2, Sample.Depth, FontEmerald);								Newline(Context2);
 		Newline(Context2);
 
@@ -1271,13 +1272,11 @@ struct FInstanceInfo
 	float LengthScale;
 	float CoverageScale;
 	bool Voxel;
-	bool Scatter;
-	bool RTGeometry;
-	bool StableRaster;
 	bool CacheRen;
 	bool CacheSim;
 	bool bHasLODSwitched;
 	bool bAutoLOD;
+	uint Flags;
 
 	// Binding
 	uint UniqueTriangle;
@@ -1346,12 +1345,10 @@ FInstanceInfo ReadInstanceInfo(uint Index)
 
 	Out.bIsVisibleInPrimaryView = Data2.x & 0x1;
 	Out.bIsVisibleInShadowView	= Data2.x & 0x2;
-	Out.Scatter					= Data2.x & 0x4;
-	Out.RTGeometry				= Data2.x & 0x8;
-	Out.StableRaster			= Data2.x & 0x10;
-	Out.Voxel					= Data2.x & 0x20;
-	Out.CacheSim				= Data2.x & 0x40;
-	Out.CacheRen				= Data2.x & 0x80;
+	Out.Voxel					= Data2.x & 0x4;
+	Out.CacheSim				= Data2.x & 0x8;
+	Out.CacheRen				= Data2.x & 0x10;
+	Out.Flags					=(Data2.x>>8) & 0xFF;
 	Out.ScreenSize				= f16tof32((Data2.x>>16) & 0xFFFF);
 	Out.ActivePointCount		= Data2.y;
 	Out.ActiveCurveCount		= Data2.z;
@@ -1591,8 +1588,8 @@ void MainCS(uint GroupIndex : SV_GroupIndex, uint3 DispatchThreadId : SV_Dispatc
 			Print(Context, TEXT("    "));
 
 			PrintBool(Context, Info.Voxel); Print(Context, TEXT("  "));
-			PrintBool(Context, Info.RTGeometry); 
-			PrintBool(Context, Info.Scatter); Print(Context, TEXT("    "));
+			PrintBool(Context, HasHairFlags(Info.Flags, HAIR_FLAGS_RAYTRACING_GEOMETRY)); 
+			PrintBool(Context, HasHairFlags(Info.Flags, HAIR_FLAGS_SCATTER_SCENE_LIGHT)); Print(Context, TEXT("    "));
 			Print(Context, Info.Radius, FontYellow, 5, 3); Print(Context, TEXT("  "));
 			Print(Context, Info.RootScale, FontYellow, 5, 3); Print(Context, TEXT(" "));
 			Print(Context, Info.TipScale, FontYellow, 5, 3); Print(Context, TEXT(" "));
diff --git a/Engine/Shaders/Private/HairStrands/HairStrandsEnvironmentLighting.usf b/Engine/Shaders/Private/HairStrands/HairStrandsEnvironmentLighting.usf
--- a/Engine/Shaders/Private/HairStrands/HairStrandsEnvironmentLighting.usf
+++ b/Engine/Shaders/Private/HairStrands/HairStrandsEnvironmentLighting.usf
@@ -42,6 +42,7 @@
 #define USE_HAIR_COMPLEX_TRANSMITTANCE 1
 
 #include "../Common.ush"
+#include "/Engine/Shared/HairStrandsDefinitions.h"
 #include "HairStrandsCommon.ush"
 #include "../DeferredShadingCommon.ush"  
 #include "../BRDF.ush"
@@ -855,10 +856,11 @@ void MainPS(
 	const bool bDebugEnabled = false;
 #endif
 
+	const bool bScatterSceneLighting = HasHairFlags(Sample.Flags, HAIR_FLAGS_SCATTER_SCENE_LIGHT);
 #if PERMUTATION_INTEGRATION_TYPE == INTEGRATION_SCENECOLOR
-	if (Sample.bScatterSceneLighting)
+	if (bScatterSceneLighting)
 #else
-	if (!Sample.bScatterSceneLighting)
+	if (!bScatterSceneLighting)
 #endif
 	{
 		const float3 SkyLighting =
diff --git a/Engine/Shaders/Private/HairStrands/HairStrandsMaterialCommon.ush b/Engine/Shaders/Private/HairStrands/HairStrandsMaterialCommon.ush
--- a/Engine/Shaders/Private/HairStrands/HairStrandsMaterialCommon.ush
+++ b/Engine/Shaders/Private/HairStrands/HairStrandsMaterialCommon.ush
@@ -34,6 +34,7 @@ uint MaterialPass_MacroGroupId;
 uint MaterialPass_MaterialId;  
 uint MaterialPass_PrimitiveId;
 uint MaterialPass_LightChannelMask;
+uint MaterialPass_Flags;
 float MaterialPass_HairCoverageScale;
 
 #ifndef HAIR_MATERIAL_DEBUG_OUTPUT
@@ -197,7 +198,7 @@ FEvaluateOutput Evaluate(
 	OutSample.ControlPointId	= SampleControlPointId;
 	OutSample.MacroGroupId		= MaterialPass_MacroGroupId;
 	OutSample.LightChannelMask	= MaterialPass_LightChannelMask;
-	OutSample.bScatterSceneLighting = UseScatterSceneLighting();
+	OutSample.Flags 			= MaterialPass_Flags;
 
 	// Material computation
 	{
diff --git a/Engine/Shaders/Private/HairStrands/HairStrandsVertexFactory.ush b/Engine/Shaders/Private/HairStrands/HairStrandsVertexFactory.ush
--- a/Engine/Shaders/Private/HairStrands/HairStrandsVertexFactory.ush
+++ b/Engine/Shaders/Private/HairStrands/HairStrandsVertexFactory.ush
@@ -74,12 +74,12 @@
  */
 bool UseStableRasterization()
 {
-	return HairStrandsVF.bStableRasterization > 0;
+	return HasHairFlags(HairStrandsVF.Flags, HAIR_FLAGS_STABLE_RASTER);
 }
 
 bool UseScatterSceneLighting()
 {
-	return HairStrandsVF.bScatterSceneLighting > 0;
+	return HasHairFlags(HairStrandsVF.Flags,HAIR_FLAGS_SCATTER_SCENE_LIGHT);
 }
 
 struct FVertexFactoryInput
diff --git a/Engine/Shaders/Private/HairStrands/HairStrandsVertexFactoryCommon.ush b/Engine/Shaders/Private/HairStrands/HairStrandsVertexFactoryCommon.ush
--- a/Engine/Shaders/Private/HairStrands/HairStrandsVertexFactoryCommon.ush
+++ b/Engine/Shaders/Private/HairStrands/HairStrandsVertexFactoryCommon.ush
@@ -54,7 +54,6 @@ float3 ReadCardPositionOffset(StructuredBuffer<float4> In, uint InIndex, uint In
 #define HAIR_STRANDS_INSTANCE_PARAMETERS(Prefix) \
 	uint						Prefix##_RegisteredIndex; \
 	uint						Prefix##_bCullingEnable; \
-	uint						Prefix##_bRaytracingGeometry; \
 	uint						Prefix##_RaytracingProceduralSplits; \
 	uint						Prefix##_GroupIndex; \
 	uint						Prefix##_GroupCount; \
@@ -63,7 +62,7 @@ float3 ReadCardPositionOffset(StructuredBuffer<float4> In, uint InIndex, uint In
 	float						Prefix##_TipScale; \
 	float						Prefix##_Length; \
 	float						Prefix##_Density; \
-	uint						Prefix##_bStableRasterization; \
+	uint						Prefix##_Flags; \
 	float3						Prefix##_PositionOffset; \
 	uint						Prefix##_PointCount; \
 	uint						Prefix##_CurveCount; \
diff --git a/Engine/Shaders/Private/HairStrands/HairStrandsVisibilityCommon.ush b/Engine/Shaders/Private/HairStrands/HairStrandsVisibilityCommon.ush
--- a/Engine/Shaders/Private/HairStrands/HairStrandsVisibilityCommon.ush
+++ b/Engine/Shaders/Private/HairStrands/HairStrandsVisibilityCommon.ush
@@ -13,6 +13,7 @@
 
 #define HAIR_VISIBILITY_GROUP_COUNT_WIDTH 64
 
+#include "../PackUnpack.ush"
 #include "/Engine/Shared/HairStrandsDefinitions.h"
 #include "HairStrandsVisibilityCommonStruct.ush"
 
@@ -236,8 +237,8 @@ struct FHairSample
 	float	Specular;
 	float	Backlit;
 	uint	LightChannelMask;
-	bool	bScatterSceneLighting;
-	float3  Emissive;
+	float3	Emissive;
+	uint	Flags;
 };
 
 FPackedHairSample PackHairSample(const FHairSample In)
@@ -249,9 +250,7 @@ FPackedHairSample PackHairSample(const FHairSample In)
 	// Coverage is clamped o 255 as we only have 8 bits to store its value
 	const float3 T = saturate(EncodeTangent(In.Tangent));
 	Out.Tangent_Coverage8bit =
-		((0xFF & uint(T.x * 255)))      | 
-		((0xFF & uint(T.y * 255)) << 8) |
-		((0xFF & uint(T.z * 255)) << 16)| 
+		PackRGBA8(float4(T,0)) |
 		(min(uint(0xFF), In.Coverage8bit) << 24);
 
 	// ControlPointId and MacroGroupId
@@ -263,11 +262,11 @@ FPackedHairSample PackHairSample(const FHairSample In)
 
 	Out.Depth					= In.Depth;																			// 32bits float
 	Out.BaseColor_Roughness		= float4ToUint(float4(sqrt(In.BaseColor), In.Roughness));							// 32bits uint
-	Out.Specular_LightChannels_Backlit = 
-		clamp(uint(In.Specular * 0xFF), 0, 0xFF) | 
-		((In.LightChannelMask & 0x7) << 8)  | 
-		((In.bScatterSceneLighting ? 1 : 0) << 12) |
-		(clamp(uint(In.Backlit * 0xFF), 0, 0xFF) << 16);
+	Out.Data 					= 
+		(PackR8(In.Specular)              )| 
+		(PackR8(In.Backlit)          <<  8)| 
+		((In.Flags & 0xFF)           << 16)|
+		((In.LightChannelMask & 0x7) << 24);
 
 	// Emissive is not packed/unpacked
 	return Out;
@@ -278,10 +277,7 @@ FHairSample UnpackHairSample(const FPackedHairSample In)
 	FHairSample Out;
 	Out.Depth = In.Depth;
 
-	Out.Tangent = DecodeTangent(float3(
-		float((In.Tangent_Coverage8bit)       & 0xFF) / 255.f,
-		float((In.Tangent_Coverage8bit >> 8)  & 0xFF) / 255.f,
-		float((In.Tangent_Coverage8bit >> 16) & 0xFF) / 255.f));
+	Out.Tangent 	 = DecodeTangent(UnpackRGBA8(In.Tangent_Coverage8bit).xyz);
 	Out.Coverage8bit = (In.Tangent_Coverage8bit >> 24) & 0xFF;
 
 	Out.ControlPointId  = In.ControlPointID_MacroGroupID & 0x0FFFFFFF;
@@ -290,10 +286,10 @@ FHairSample UnpackHairSample(const FPackedHairSample In)
 	const float4 BaseColorAndRoughness = UintToFloat4(In.BaseColor_Roughness);
 	Out.BaseColor			= BaseColorAndRoughness.xyz * BaseColorAndRoughness.xyz;
 	Out.Roughness			= BaseColorAndRoughness.w;
-	Out.Specular			= float(0xFF & In.Specular_LightChannels_Backlit) / 0xFF;
-	Out.LightChannelMask	= (In.Specular_LightChannels_Backlit >>8) & 0x7;
-	Out.Backlit				= float(0xFF & (In.Specular_LightChannels_Backlit>>16)) / 0xFF;
-	Out.bScatterSceneLighting = (In.Specular_LightChannels_Backlit & (1<<12)) > 0;
+	Out.Specular			= UnpackR8(In.Data & 0xFF);
+	Out.Backlit				= UnpackR8((In.Data>>8) & 0xFF);
+	Out.Flags 				= (In.Data>>16) & 0xFF;
+	Out.LightChannelMask	= (In.Data>>24) & 0x7;
 
 	// Emissive is not packed/unpacked
 	return Out;
diff --git a/Engine/Shaders/Private/HairStrands/HairStrandsVisibilityCommonStruct.ush b/Engine/Shaders/Private/HairStrands/HairStrandsVisibilityCommonStruct.ush
--- a/Engine/Shaders/Private/HairStrands/HairStrandsVisibilityCommonStruct.ush
+++ b/Engine/Shaders/Private/HairStrands/HairStrandsVisibilityCommonStruct.ush
@@ -8,7 +8,7 @@ struct FPackedHairSample
 	uint  ControlPointID_MacroGroupID;
 	uint  Tangent_Coverage8bit;
 	uint  BaseColor_Roughness;
-	uint  Specular_LightChannels_Backlit;
+	uint  Data;
 };
 
 struct FPackedHairVis
diff --git a/Engine/Shaders/Private/HairStrands/HairStrandsVisibilityComposeSubPixelPS.usf b/Engine/Shaders/Private/HairStrands/HairStrandsVisibilityComposeSubPixelPS.usf
--- a/Engine/Shaders/Private/HairStrands/HairStrandsVisibilityComposeSubPixelPS.usf
+++ b/Engine/Shaders/Private/HairStrands/HairStrandsVisibilityComposeSubPixelPS.usf
@@ -31,6 +31,7 @@ void FastResolvePS(in FScreenVertexOutput Input)
 
 int2 OutputResolution;
 uint bComposeDofDepth;
+uint bHasHoldout;
 
 Texture2D<float>	HairDOFDepthTexture;
 Texture2D<float4>	HairLightingSampleBuffer;
@@ -115,17 +116,27 @@ void ComposeSamplePS(
 	const float3 ClosetPointWorldPosition = LWCHackToFloat(SvPositionToWorld(float4(Input.Position.xy, HairDeviceZ, 1.0)));
 	const float4 Volumetric = EvaluateVolumetric(ClosetPointWorldPosition, ConvertFromDeviceZ(HairDeviceZ));
 
+	bool bHoldout = false;
 	float3 LocalAccColor = 0;
 	LOOP
 	for (uint SampleIt = 0; SampleIt < NodeDesc.Count; SampleIt++)
 	{
+		#if POST_PROCESS_ALPHA
+		if (bHasHoldout)
+		{
+			const FHairSample Sample = UnpackHairSample(HairStrands.HairSampleData[NodeDesc.Offset + SampleIt]);
+			bHoldout = bHoldout || HasHairFlags(Sample.Flags, HAIR_FLAGS_HOLDOUT);
+		}
+		#endif
+
 		const uint LocalOffset = NodeDesc.Offset + SampleIt;
 		const uint2 Coord = uint2(LocalOffset % Resolution, LocalOffset / Resolution);
 		LocalAccColor += HairLightingSampleBuffer.Load(uint3(Coord, 0)).xyz;
 	}
 	OutColor.rgb = LocalAccColor * Volumetric.a + Volumetric.rgb;
 
-	OutColor.rgb *= PixelCoverage;
+	// If any sample support holdout, then the lighting contribution is set to 0, to mask out the current hair pixel
+	OutColor.rgb *= bHoldout ? 0 : PixelCoverage;
 	OutColor.a = PixelCoverage;
 #endif // PERMUTATION_TEMPORAL
 
@@ -173,6 +184,66 @@ void ComposeSamplePS(
 
 ////////////////////////////////////////////////////////////////////////////////////////////////
 
+#if SHADER_HOLDOUT
+
+uint bComposeDofDepth;
+Texture2D<float> HairDOFDepthTexture;
+
+void HoldoutPS(
+	in FScreenVertexOutput Input,
+	out float4 OutColor : SV_Target0,
+	out float  OutDepth : SV_DEPTH)
+{
+#if POST_PROCESS_ALPHA
+	OutColor = 0.0f;
+	const uint3 PixelCoord = uint3(floor(Input.Position.xy), 0);
+	
+	const float HairDeviceZ = HairStrands.HairOnlyDepthTexture.Load(PixelCoord).x;
+
+	const FNodeDesc NodeDesc = DecodeNodeDesc(HairStrands.HairSampleOffset.Load(PixelCoord));
+	if (NodeDesc.Count == 0)
+	{
+		discard;
+	}
+
+	const float PixelCoverage = min(HairStrands.HairCoverageTexture.Load(PixelCoord), 1);
+	if (PixelCoverage == 0)
+	{
+		discard;
+	}
+
+	bool bHoldout = false;
+
+	LOOP
+	for (uint SampleIt = 0; SampleIt < NodeDesc.Count; SampleIt++)
+	{
+		const FHairSample Sample = UnpackHairSample(HairStrands.HairSampleData[NodeDesc.Offset + SampleIt]);
+		bHoldout = bHoldout || HasHairFlags(Sample.Flags, HAIR_FLAGS_HOLDOUT);
+	}
+
+	if (!bHoldout)
+	{
+		discard;
+	}
+
+	OutColor.rgb = 0;
+	OutColor.a = 1;  //PixelCoverage;
+	OutDepth     = HairDeviceZ;
+	if (bComposeDofDepth)
+	{
+		OutDepth = HairDOFDepthTexture.Load(PixelCoord);
+	}
+#else
+	OutColor = 0;
+	OutDepth = 0;
+	discard;
+#endif
+}
+
+#endif // SHADER_HOLDOUT
+
+////////////////////////////////////////////////////////////////////////////////////////////////
+
 #if SHADER_DOFDEPTH
 #include "../CircleDOFCommon.ush"
 
diff --git a/Engine/Shaders/Private/HairStrands/HairStrandsVoxelRasterCompute.usf b/Engine/Shaders/Private/HairStrands/HairStrandsVoxelRasterCompute.usf
--- a/Engine/Shaders/Private/HairStrands/HairStrandsVoxelRasterCompute.usf
+++ b/Engine/Shaders/Private/HairStrands/HairStrandsVoxelRasterCompute.usf
@@ -815,7 +815,7 @@ void MainCS(uint2 DispatchThreadID : SV_DispatchThreadID)
 				// Groom having raytraced geometry will cast-shadow on opaque geomtry with their RT-geometry, not with their 
 				// voxelization. To avoid for doubling/incorrect shadowing, we mark voxel with no-shadow casting flag. 
 				// This adds a significant cost when used.
-				if (HairStrandsVF_bRaytracingGeometry)
+				if (HasHairFlags(HairStrandsVF_Flags, HAIR_FLAGS_RAYTRACING_GEOMETRY))
 				{
 					InterlockedOr(OutPageTexture[VoxelPageCoord], VOXEL_CAST_NO_SHADOW_MASK);
 				}
diff --git a/Engine/Shaders/Shared/HairStrandsDefinitions.h b/Engine/Shaders/Shared/HairStrandsDefinitions.h
--- a/Engine/Shaders/Shared/HairStrandsDefinitions.h
+++ b/Engine/Shaders/Shared/HairStrandsDefinitions.h
@@ -98,4 +98,19 @@
 
 #define HAIR_RBF_ENTRY_COUNT(InSampleCount)	((InSampleCount)+4u)
 
-#define HAIR_CARDS_MAX_TEXTURE_COUNT 6
\ No newline at end of file
+#define HAIR_CARDS_MAX_TEXTURE_COUNT 6
+
+// Hair instance flags 
+#define HAIR_FLAGS_SCATTER_SCENE_LIGHT 1u
+#define HAIR_FLAGS_STABLE_RASTER 2u
+#define HAIR_FLAGS_RAYTRACING_GEOMETRY 4u
+#define HAIR_FLAGS_HOLDOUT 8u
+
+#ifndef __cplusplus //HLSL
+bool HasHairFlags(uint In, uint Flags)
+#else
+FORCEINLINE bool HasHairFlags(uint32 In, uint32 Flags)
+#endif
+{
+	return (In & Flags) != 0;
+}
diff --git a/Engine/Source/Runtime/Renderer/Private/HairStrands/HairStrandsCluster.cpp b/Engine/Source/Runtime/Renderer/Private/HairStrands/HairStrandsCluster.cpp
--- a/Engine/Source/Runtime/Renderer/Private/HairStrands/HairStrandsCluster.cpp
+++ b/Engine/Source/Runtime/Renderer/Private/HairStrands/HairStrandsCluster.cpp
@@ -125,10 +125,8 @@ static void InternalUpdateMacroGroup(FHairStrandsMacroGroupData& MacroGroup, int
 {
 	check(HairData);
 
-	if (HairData->VFInput.Strands.Common.bScatterSceneLighting)
-	{
-		MacroGroup.bNeedScatterSceneLighting = true;
-	}
+	// Track if any instance in the current group needs some features (scatter light scene, holdout, ...)
+	MacroGroup.Flags |= HairData->VFInput.Strands.Common.Flags;
 
 	FHairStrandsMacroGroupData::PrimitiveInfo& PrimitiveInfo = MacroGroup.PrimitivesInfos.AddZeroed_GetRef();
 	PrimitiveInfo.Mesh = Mesh;
@@ -136,6 +134,7 @@ static void InternalUpdateMacroGroup(FHairStrandsMacroGroupData& MacroGroup, int
 	PrimitiveInfo.MaterialId = MaterialId++;
 	PrimitiveInfo.ResourceId = Mesh ? reinterpret_cast<uint64>(Mesh->Elements[0].UserData) : ~0u;
 	PrimitiveInfo.GroupIndex = HairData->GetGroupIndex();
+	PrimitiveInfo.Flags = HairData->VFInput.Strands.Common.Flags;
 	PrimitiveInfo.PublicDataPtr = HairData;
 
 	if (HairData->DoesSupportVoxelization())
@@ -294,6 +293,13 @@ void CreateHairStrandsMacroGroups(
 		}
 		MacroGroupResources.MacroGroupCount = MacroGroups.Num();
 	}
+
+	// Aggregate flags accross all instances/macrogroups
+	OutHairStrandsViewData.Flags = 0;
+	for (const FHairStrandsMacroGroupData& MacroGroupData : View.HairStrandsViewData.MacroGroupDatas)
+	{			
+		OutHairStrandsViewData.Flags |= MacroGroupData.Flags;
+	}
 }
 
 bool FHairStrandsMacroGroupData::PrimitiveInfo::IsCullingEnable() const
diff --git a/Engine/Source/Runtime/Renderer/Private/HairStrands/HairStrandsComposition.cpp b/Engine/Source/Runtime/Renderer/Private/HairStrands/HairStrandsComposition.cpp
--- a/Engine/Source/Runtime/Renderer/Private/HairStrands/HairStrandsComposition.cpp
+++ b/Engine/Source/Runtime/Renderer/Private/HairStrands/HairStrandsComposition.cpp
@@ -42,7 +42,8 @@ enum EHairStrandsCommonPassType
 	DOF,
 	TAAFastResolve,
 	GBuffer,
-	Blit
+	Blit,
+	Holdout
 };
 
 template<typename TPassParameter, typename TPixelShader>
@@ -93,6 +94,10 @@ void InternalCommonDrawPass(
 		{
 			GraphicsPSOInit.BlendState = TStaticBlendState<CW_RGBA, BO_Add, BF_One, BF_Zero, BO_Add, BF_One, BF_Zero>::GetRHI();
 		}
+		else if (Type == EHairStrandsCommonPassType::Holdout)
+		{
+			GraphicsPSOInit.BlendState = TStaticBlendState<CW_ALPHA, BO_Add, BF_Zero, BF_One, BO_Max, BF_One, BF_One>::GetRHI();
+		}
 		else
 		{
 			GraphicsPSOInit.BlendState = TStaticBlendState<CW_RGBA, BO_Add, BF_One, BF_Zero, BO_Max, BF_One, BF_One>::GetRHI();
@@ -100,7 +105,7 @@ void InternalCommonDrawPass(
 
 		GraphicsPSOInit.RasterizerState = TStaticRasterizerState<>::GetRHI();
 
-		if (Type == EHairStrandsCommonPassType::Composition)
+		if (Type == EHairStrandsCommonPassType::Composition || Type == EHairStrandsCommonPassType::Holdout)
 		{
 			GraphicsPSOInit.DepthStencilState = TStaticDepthStencilState<true, CF_DepthNearOrEqual>::GetRHI();
 		}
@@ -148,6 +153,69 @@ void InternalCommonDrawPass(
 
 /////////////////////////////////////////////////////////////////////////////////////////
 
+class FHairHoldoutPS : public FGlobalShader
+{
+	DECLARE_GLOBAL_SHADER(FHairHoldoutPS);
+	SHADER_USE_PARAMETER_STRUCT(FHairHoldoutPS, FGlobalShader);
+
+	using FPermutationDomain = TShaderPermutationDomain<>;
+	
+	BEGIN_SHADER_PARAMETER_STRUCT(FParameters, )
+		SHADER_PARAMETER_STRUCT_REF(FViewUniformShaderParameters, ViewUniformBuffer)
+		SHADER_PARAMETER_STRUCT_INCLUDE(FHairStrandsTilePassVS::FParameters, TileData)
+		SHADER_PARAMETER(FIntPoint, OutputResolution)
+		SHADER_PARAMETER(uint32, bComposeDofDepth)
+		SHADER_PARAMETER_RDG_UNIFORM_BUFFER(FHairStrandsViewUniformParameters, HairStrands)
+		SHADER_PARAMETER_RDG_TEXTURE(Texture2D, HairDOFDepthTexture)
+		RENDER_TARGET_BINDING_SLOTS()
+	END_SHADER_PARAMETER_STRUCT()
+
+public:
+	static bool ShouldCompilePermutation(const FGlobalShaderPermutationParameters& Parameters) { return IsHairStrandsSupported(EHairStrandsShaderType::Strands, Parameters.Platform); }
+	static void ModifyCompilationEnvironment(const FGlobalShaderPermutationParameters& Parameters, FShaderCompilerEnvironment& OutEnvironment)
+	{
+		FGlobalShader::ModifyCompilationEnvironment(Parameters, OutEnvironment);
+		OutEnvironment.SetDefine(TEXT("SHADER_HOLDOUT"), 1);
+	}
+};
+
+IMPLEMENT_GLOBAL_SHADER(FHairHoldoutPS, "/Engine/Private/HairStrands/HairStrandsVisibilityComposeSubPixelPS.usf", "HoldoutPS", SF_Pixel);
+
+static void AddHairHoldoutPass(
+	FRDGBuilder& GraphBuilder,
+	const FViewInfo& View,
+	const FHairStrandsVisibilityData& VisibilityData,
+	const FRDGTextureRef& HairDOFDepthTexture,
+	FRDGTextureRef& OutColorTexture,
+	FRDGTextureRef& OutDepthTexture)
+{
+	const bool bDOFEnable = HairDOFDepthTexture != nullptr ? 1 : 0;
+
+	FHairHoldoutPS::FParameters* Parameters = GraphBuilder.AllocParameters<FHairHoldoutPS::FParameters>();
+	Parameters->bComposeDofDepth = bDOFEnable ? 1 : 0;
+	Parameters->HairDOFDepthTexture = bDOFEnable ? HairDOFDepthTexture : GSystemTextures.GetBlackDummy(GraphBuilder);
+	Parameters->OutputResolution = OutColorTexture->Desc.Extent;
+	Parameters->ViewUniformBuffer = View.ViewUniformBuffer;
+	Parameters->HairStrands = View.HairStrandsViewData.UniformBuffer;
+	Parameters->RenderTargets[0] = FRenderTargetBinding(OutColorTexture, ERenderTargetLoadAction::ELoad);
+	Parameters->RenderTargets.DepthStencil = FDepthStencilBinding(OutDepthTexture, ERenderTargetLoadAction::ELoad, ERenderTargetLoadAction::ELoad, FExclusiveDepthStencil::DepthWrite_StencilRead);
+
+	FHairHoldoutPS::FPermutationDomain PermutationVector;
+	TShaderMapRef<FHairHoldoutPS> PixelShader(View.ShaderMap, PermutationVector);
+	InternalCommonDrawPass(
+		GraphBuilder,
+		RDG_EVENT_NAME("HairStrands::Holdout"),
+		View,
+		OutColorTexture->Desc.Extent,
+		EHairStrandsCommonPassType::Holdout,
+		false,
+		VisibilityData.TileData,
+		PixelShader,
+		Parameters);
+}
+
+/////////////////////////////////////////////////////////////////////////////////////////
+
 class FHairVisibilityComposeSamplePS : public FGlobalShader
 {
 	DECLARE_GLOBAL_SHADER(FHairVisibilityComposeSamplePS);
@@ -162,6 +230,7 @@ class FHairVisibilityComposeSamplePS : public FGlobalShader
 		SHADER_PARAMETER_STRUCT_INCLUDE(FHairStrandsTilePassVS::FParameters, TileData)
 		SHADER_PARAMETER(FIntPoint, OutputResolution)
 		SHADER_PARAMETER(uint32, bComposeDofDepth)
+		SHADER_PARAMETER(uint32, bHasHoldout)
 		SHADER_PARAMETER_RDG_UNIFORM_BUFFER(FHairStrandsViewUniformParameters, HairStrands)
 		SHADER_PARAMETER_RDG_TEXTURE(Texture2D, HairLightingSampleBuffer)
 		SHADER_PARAMETER_RDG_TEXTURE(Texture2D, HairDOFDepthTexture)
@@ -170,10 +239,10 @@ class FHairVisibilityComposeSamplePS : public FGlobalShader
 		SHADER_PARAMETER_RDG_TEXTURE_UAV(RWTexture2D<uint>, OutMetaTexture)
 		SHADER_PARAMETER_RDG_UNIFORM_BUFFER(FFogUniformParameters, FogStruct)
 		RENDER_TARGET_BINDING_SLOTS()
-	END_SHADER_PARAMETER_STRUCT()
+		END_SHADER_PARAMETER_STRUCT()
 
 public:
-	static bool ShouldCompilePermutation(const FGlobalShaderPermutationParameters& Parameters) { return IsHairStrandsSupported(EHairStrandsShaderType::Strands, Parameters.Platform); }
+		static bool ShouldCompilePermutation(const FGlobalShaderPermutationParameters& Parameters) { return IsHairStrandsSupported(EHairStrandsShaderType::Strands, Parameters.Platform); }
 	static void ModifyCompilationEnvironment(const FGlobalShaderPermutationParameters& Parameters, FShaderCompilerEnvironment& OutEnvironment)
 	{
 		FGlobalShader::ModifyCompilationEnvironment(Parameters, OutEnvironment);
@@ -187,7 +256,7 @@ static void AddHairVisibilityComposeSamplePass(
 	FRDGBuilder& GraphBuilder,
 	const FViewInfo& View,
 	const FHairStrandsVisibilityData& VisibilityData,
-	const FRDGTextureRef& CategorizationTexture,
+	const bool bHasHoldout,
 	const FRDGTextureRef& HairDOFDepthTexture,
 	const FRDGTextureRef& HairTemporalAccumulationTexture,
 	FRDGTextureRef& OutColorTexture,
@@ -210,6 +279,7 @@ static void AddHairVisibilityComposeSamplePass(
 	Parameters->HairStrands = View.HairStrandsViewData.UniformBuffer;
 	Parameters->FogStruct = FogBuffer;
 	Parameters->TemporalLayerCount = 1;
+	Parameters->bHasHoldout = bHasHoldout ? 1u : 0u;
 	Parameters->HairTemporalAccumulationTexture = bTemporal ? HairTemporalAccumulationTexture : GSystemTextures.GetBlackDummy(GraphBuilder);
 	Parameters->RenderTargets[0] = FRenderTargetBinding(OutColorTexture, bTemporal ? ERenderTargetLoadAction::ELoad : (bTemporalLayeringEnabled ? ERenderTargetLoadAction::EClear : ERenderTargetLoadAction::ELoad) );
 	Parameters->RenderTargets.DepthStencil = FDepthStencilBinding(OutDepthTexture, ERenderTargetLoadAction::ELoad, ERenderTargetLoadAction::ELoad, FExclusiveDepthStencil::DepthWrite_StencilRead);
@@ -231,7 +301,6 @@ static void AddHairVisibilityComposeSamplePass(
 		Parameters);
 }
 
-
 /////////////////////////////////////////////////////////////////////////////////////////
 
 class FHairDOFDepthPS : public FGlobalShader
@@ -515,6 +584,7 @@ static void InternalRenderHairComposition(
 
 		FRDGTextureRef AccumulatedColor = nullptr;
 
+		const bool bHasHoldout = HasHairFlags(View.HairStrandsViewData.Flags, HAIR_FLAGS_HOLDOUT);
 		const bool bTemporalLayeringEnabled = false; // HAIR_TODO: remove
 
 		if (bTemporalLayeringEnabled)
@@ -528,7 +598,7 @@ static void InternalRenderHairComposition(
 				GraphBuilder,
 				View,
 				VisibilityData,
-				VisibilityData.CoverageTexture,
+				bHasHoldout,
 				DOFDepth,
 				nullptr,
 				HairColorInputTexture,
@@ -551,20 +621,30 @@ static void InternalRenderHairComposition(
 				View.PrevViewInfo.HairHistory,
 				&View.ViewState->PrevFrameViewInfo.HairHistory);
 
-			AccumulatedColor = TAAOutputs.SceneColor;			
+			AccumulatedColor = TAAOutputs.SceneColor;
 		}
 
 		AddHairVisibilityComposeSamplePass(
 			GraphBuilder,
 			View,
 			VisibilityData,
-			VisibilityData.CoverageTexture,
+			bHasHoldout,
 			DOFDepth,
 			AccumulatedColor,
 			SceneColorTexture,
 			SceneDepthTexture,
 			bTemporalLayeringEnabled);
 
+		if (bHasHoldout)
+		{
+			AddHairHoldoutPass(
+				GraphBuilder,
+				View,
+				VisibilityData,
+				DOFDepth,
+				SceneColorTexture,
+				SceneDepthTexture);
+		}
 		if (VisibilityData.ResolveMaskTexture)
 		{
 			AddHairVisibilityFastResolveMaskPass(
diff --git a/Engine/Source/Runtime/Renderer/Private/HairStrands/HairStrandsData.h b/Engine/Source/Runtime/Renderer/Private/HairStrands/HairStrandsData.h
--- a/Engine/Source/Runtime/Renderer/Private/HairStrands/HairStrandsData.h
+++ b/Engine/Source/Runtime/Renderer/Private/HairStrands/HairStrandsData.h
@@ -306,9 +306,10 @@ struct FHairStrandsMacroGroupData
 	{
 		const FMeshBatch* Mesh = nullptr;
 		const FPrimitiveSceneProxy* PrimitiveSceneProxy = nullptr;
-		uint32 MaterialId;
-		uint32 ResourceId;
-		uint32 GroupIndex;
+		uint32 MaterialId = 0;
+		uint32 ResourceId = 0;
+		uint32 GroupIndex = 0;
+		uint32 Flags = 0; // Hair instance flags
 		FHairGroupPublicData* PublicDataPtr = nullptr;
 		bool IsCullingEnable() const;
 	};
@@ -318,10 +319,10 @@ struct FHairStrandsMacroGroupData
 	TPrimitiveInfos PrimitivesInfos;
 	FBoxSphereBounds Bounds;
 	FIntRect ScreenRect;
-	uint32 MacroGroupId;
+	uint32 MacroGroupId = 0;
+	uint32 Flags = 0; // Aggregated flags for all instances from the group
 
 	bool bSupportVoxelization = false; // true if at least one of the Primitive requires voxelization
-	bool bNeedScatterSceneLighting = false;
 };
 
 ////////////////////////////////////////////////////////////////////////////////////
@@ -447,6 +448,7 @@ struct FHairStrandsViewData
 	FHairStrandsVoxelResources VirtualVoxelResources;
 	FHairStrandsMacroGroupResources MacroGroupResources;
 	FHairStrandsDebugData DebugData;
+	uint32 Flags = 0;
 
 	// Transient: store all light visible in primary view(s)
 	struct FDirectionalLightCullData { const FLightSceneInfo* LightInfo = nullptr; FConvexVolume ViewFrustumInLightSpace; };
diff --git a/Engine/Source/Runtime/Renderer/Private/HairStrands/HairStrandsEnvironment.cpp b/Engine/Source/Runtime/Renderer/Private/HairStrands/HairStrandsEnvironment.cpp
--- a/Engine/Source/Runtime/Renderer/Private/HairStrands/HairStrandsEnvironment.cpp
+++ b/Engine/Source/Runtime/Renderer/Private/HairStrands/HairStrandsEnvironment.cpp
@@ -472,17 +472,7 @@ void RenderHairStrandsSceneColorScattering(
 		const FHairStrandsVoxelResources& VoxelResources = View.HairStrandsViewData.VirtualVoxelResources;
 		check(VoxelResources.IsValid());
 
-		bool bNeedScatterSceneLighting = false;
-		for (const FHairStrandsMacroGroupData& MacroGroupData : View.HairStrandsViewData.MacroGroupDatas)
-		{
-			if (MacroGroupData.bNeedScatterSceneLighting)
-			{
-				bNeedScatterSceneLighting = true;
-				break;
-			}
-		}
-
-		if (bNeedScatterSceneLighting)
+		if (HasHairFlags(View.HairStrandsViewData.Flags, HAIR_FLAGS_SCATTER_SCENE_LIGHT))
 		{
 			AddHairStrandsEnvironmentLightingPassPS(GraphBuilder, Scene, View, VisibilityData, VoxelResources, SceneColorTexture, EHairLightingSourceType::SceneColor, nullptr);
 		}
diff --git a/Engine/Source/Runtime/Renderer/Private/HairStrands/HairStrandsVisibility.cpp b/Engine/Source/Runtime/Renderer/Private/HairStrands/HairStrandsVisibility.cpp
--- a/Engine/Source/Runtime/Renderer/Private/HairStrands/HairStrandsVisibility.cpp
+++ b/Engine/Source/Runtime/Renderer/Private/HairStrands/HairStrandsVisibility.cpp
@@ -490,11 +490,26 @@ IMPLEMENT_MATERIAL_SHADER_TYPE(, FHairMaterialVS, TEXT("/Engine/Private/HairStra
 class FHairMaterialShaderElementData : public FMeshMaterialShaderElementData
 {
 public:
-	FHairMaterialShaderElementData(int32 MacroGroupId, int32 MaterialId, int32 PrimitiveId, uint32 LightChannelMask, float HairCoverageScale) : MaterialPass_MacroGroupId(MacroGroupId), MaterialPass_MaterialId(MaterialId), MaterialPass_PrimitiveId(PrimitiveId), MaterialPass_LightChannelMask(LightChannelMask), MaterialPass_HairCoverageScale(HairCoverageScale) { }
+	FHairMaterialShaderElementData(
+		int32 MacroGroupId, 
+		int32 MaterialId, 
+		int32 PrimitiveId, 
+		uint32 LightChannelMask, 
+		uint32 Flags, 
+		float HairCoverageScale) 
+	: MaterialPass_MacroGroupId(MacroGroupId)
+	, MaterialPass_MaterialId(MaterialId)
+	, MaterialPass_PrimitiveId(PrimitiveId)
+	, MaterialPass_LightChannelMask(LightChannelMask)
+	, MaterialPass_Flags(Flags)
+	, MaterialPass_HairCoverageScale(HairCoverageScale) 
+	{ }
+
 	uint32 MaterialPass_MacroGroupId;
 	uint32 MaterialPass_MaterialId;
 	uint32 MaterialPass_PrimitiveId;
 	uint32 MaterialPass_LightChannelMask;
+	uint32 MaterialPass_Flags;
 	uint32 MaterialPass_HairCoverageScale;
 };
 
@@ -518,6 +533,7 @@ public:
 		MaterialPass_MaterialId.Bind(Initializer.ParameterMap, TEXT("MaterialPass_MaterialId"));
 		MaterialPass_PrimitiveId.Bind(Initializer.ParameterMap, TEXT("MaterialPass_PrimitiveId"));
 		MaterialPass_LightChannelMask.Bind(Initializer.ParameterMap, TEXT("MaterialPass_LightChannelMask"));
+		MaterialPass_Flags.Bind(Initializer.ParameterMap, TEXT("MaterialPass_Flags"));
 		MaterialPass_HairCoverageScale.Bind(Initializer.ParameterMap, TEXT("MaterialPass_HairCoverageScale"));
 	}
 
@@ -556,6 +572,7 @@ public:
 		ShaderBindings.Add(MaterialPass_MaterialId, ShaderElementData.MaterialPass_MaterialId);
 		ShaderBindings.Add(MaterialPass_PrimitiveId, ShaderElementData.MaterialPass_PrimitiveId);
 		ShaderBindings.Add(MaterialPass_LightChannelMask, ShaderElementData.MaterialPass_LightChannelMask);
+		ShaderBindings.Add(MaterialPass_Flags, ShaderElementData.MaterialPass_Flags);
 		ShaderBindings.Add(MaterialPass_HairCoverageScale, ShaderElementData.MaterialPass_HairCoverageScale);
 	}
 
@@ -564,6 +581,7 @@ private:
 	LAYOUT_FIELD(FShaderParameter, MaterialPass_MaterialId);
 	LAYOUT_FIELD(FShaderParameter, MaterialPass_PrimitiveId);
 	LAYOUT_FIELD(FShaderParameter, MaterialPass_LightChannelMask);
+	LAYOUT_FIELD(FShaderParameter, MaterialPass_Flags);
 	LAYOUT_FIELD(FShaderParameter, MaterialPass_HairCoverageScale);
 };
 IMPLEMENT_MATERIAL_SHADER_TYPE(, FHairMaterialPS, TEXT("/Engine/Private/HairStrands/HairStrandsMaterialPS.usf"), TEXT("Main"), SF_Pixel);
@@ -599,7 +617,7 @@ public:
 		EHairMaterialPassFilter InFilter);
 
 	virtual void AddMeshBatch(const FMeshBatch& RESTRICT MeshBatch, uint64 BatchElementMask, const FPrimitiveSceneProxy* RESTRICT PrimitiveSceneProxy, int32 StaticMeshId = -1) override final;
-	void AddMeshBatch(const FMeshBatch& RESTRICT MeshBatch, uint64 BatchElementMask, const FPrimitiveSceneProxy* RESTRICT PrimitiveSceneProxy, int32 StaticMeshId, int32 MacroGroupId, int32 HairMaterialId, float HairCoverageScale);
+	void AddMeshBatch(const FMeshBatch& RESTRICT MeshBatch, uint64 BatchElementMask, const FPrimitiveSceneProxy* RESTRICT PrimitiveSceneProxy, int32 StaticMeshId, int32 MacroGroupId, int32 HairMaterialId, uint32 HairFlags, float HairCoverageScale);
 
 	virtual void CollectPSOInitializers(const FSceneTexturesConfig& SceneTexturesConfig, const FMaterial& Material, const FPSOPrecacheVertexFactoryData& VertexFactoryData, const FPSOPrecacheParams& PreCacheParams, TArray<FPSOPrecacheData>& PSOInitializers) override final;
 
@@ -611,6 +629,7 @@ private:
 		int32 StaticMeshId,
 		uint32 MacroGroupId,
 		uint32 HairMaterialId,
+		uint32 HairFlags,
 		float HairCoverageScale,
 		const FMaterialRenderProxy& MaterialRenderProxy,
 		const FMaterial& Material);
@@ -626,6 +645,7 @@ private:
 		const int32 HairMaterialId,
 		const int32 HairControlPointId,
 		const uint32 HairPrimitiveLightChannelMask,
+		const uint32 HairFlags,
 		const float HairCoverageScale);
 
 	void SetupDrawRenderState(EHairMaterialPassFilter InFilter);
@@ -649,10 +669,10 @@ void FHairMaterialProcessor::SetupDrawRenderState(EHairMaterialPassFilter InFilt
 
 void FHairMaterialProcessor::AddMeshBatch(const FMeshBatch& RESTRICT MeshBatch, uint64 BatchElementMask, const FPrimitiveSceneProxy* RESTRICT PrimitiveSceneProxy, int32 StaticMeshId)
 {
-	AddMeshBatch(MeshBatch, BatchElementMask, PrimitiveSceneProxy, StaticMeshId, 0, 0, 1.f);
+	AddMeshBatch(MeshBatch, BatchElementMask, PrimitiveSceneProxy, StaticMeshId, 0/*MacroGroupId*/, 0/*HairMaterialId*/, 0/*HairFlags*/, 1.f/*HairCoverageScale*/);
 }
 
-void FHairMaterialProcessor::AddMeshBatch(const FMeshBatch& RESTRICT MeshBatch, uint64 BatchElementMask, const FPrimitiveSceneProxy* RESTRICT PrimitiveSceneProxy, int32 StaticMeshId, int32 MacroGroupId, int32 HairMaterialId, float HairCoverageScale)
+void FHairMaterialProcessor::AddMeshBatch(const FMeshBatch& RESTRICT MeshBatch, uint64 BatchElementMask, const FPrimitiveSceneProxy* RESTRICT PrimitiveSceneProxy, int32 StaticMeshId, int32 MacroGroupId, int32 HairMaterialId, uint32 HairFlags, float HairCoverageScale)
 {
 	const FMaterialRenderProxy* MaterialRenderProxy = MeshBatch.MaterialRenderProxy;
 	while (MaterialRenderProxy)
@@ -660,7 +680,7 @@ void FHairMaterialProcessor::AddMeshBatch(const FMeshBatch& RESTRICT MeshBatch,
 		const FMaterial* Material = MaterialRenderProxy->GetMaterialNoFallback(FeatureLevel);
 		if (Material)
 		{
-			if (TryAddMeshBatch(MeshBatch, BatchElementMask, PrimitiveSceneProxy, StaticMeshId, MacroGroupId, HairMaterialId, HairCoverageScale, *MaterialRenderProxy, *Material))
+			if (TryAddMeshBatch(MeshBatch, BatchElementMask, PrimitiveSceneProxy, StaticMeshId, MacroGroupId, HairMaterialId, HairFlags, HairCoverageScale, *MaterialRenderProxy, *Material))
 			{
 				break;
 			}
@@ -677,6 +697,7 @@ bool FHairMaterialProcessor::TryAddMeshBatch(
 	int32 StaticMeshId,
 	uint32 MacroGroupId,
 	uint32 HairMaterialId,
+	uint32 HairFlags,
 	float HairCoverageScale,
 	const FMaterialRenderProxy& MaterialRenderProxy,
 	const FMaterial& Material)
@@ -700,7 +721,7 @@ bool FHairMaterialProcessor::TryAddMeshBatch(
 		FMeshDrawCommandPrimitiveIdInfo IdInfo = GetDrawCommandPrimitiveId(SceneInfo, MeshBatch.Elements[0]);
 		uint32 LightChannelMask = PrimitiveSceneProxy ? PrimitiveSceneProxy->GetLightingChannelMask() : 0;
 
-		return Process(MeshBatchCopy, BatchElementMask, PrimitiveSceneProxy, StaticMeshId, MaterialRenderProxy, Material, MacroGroupId, HairMaterialId, IdInfo.DrawPrimitiveId, LightChannelMask, HairCoverageScale);
+		return Process(MeshBatchCopy, BatchElementMask, PrimitiveSceneProxy, StaticMeshId, MaterialRenderProxy, Material, MacroGroupId, HairMaterialId, IdInfo.DrawPrimitiveId, LightChannelMask, HairFlags, HairCoverageScale);
 	}
 
 	return true;
@@ -717,6 +738,7 @@ bool FHairMaterialProcessor::Process(
 	const int32 HairMaterialId,
 	const int32 HairControlPointId,
 	const uint32 HairPrimitiveLightChannelMask,
+	const uint32 HairFlags,
 	const float HairCoverageScale)
 {
 	const FVertexFactory* VertexFactory = MeshBatch.VertexFactory;
@@ -742,7 +764,7 @@ bool FHairMaterialProcessor::Process(
 	}
 
 	FMeshPassProcessorRenderState DrawRenderState(PassDrawRenderState);
-	FHairMaterialShaderElementData ShaderElementData(MacroGroupId, HairMaterialId, HairControlPointId, HairPrimitiveLightChannelMask, HairCoverageScale);
+	FHairMaterialShaderElementData ShaderElementData(MacroGroupId, HairMaterialId, HairControlPointId, HairPrimitiveLightChannelMask, HairFlags, HairCoverageScale);
 	ShaderElementData.InitializeMeshMaterialData(ViewIfDynamicMeshCommand, PrimitiveSceneProxy, MeshBatch, StaticMeshId, false);
 
 	const bool bReverseCulling = ViewIfDynamicMeshCommand ? ViewIfDynamicMeshCommand->bIsPlanarReflection : false;
@@ -1059,7 +1081,7 @@ static FMaterialPassOutput AddHairMaterialPass(
 
 						if (bIsCompatible)
 						{
-							MeshProcessor.AddMeshBatch(*MeshBatch, BatchElementMask, PrimitiveInfo.PrimitiveSceneProxy, -1, MacroGroupData.MacroGroupId, PrimitiveInfo.MaterialId, PrimitiveInfo.PublicDataPtr->GetActiveStrandsCoverageScale());
+							MeshProcessor.AddMeshBatch(*MeshBatch, BatchElementMask, PrimitiveInfo.PrimitiveSceneProxy, -1, MacroGroupData.MacroGroupId, PrimitiveInfo.MaterialId, PrimitiveInfo.Flags, PrimitiveInfo.PublicDataPtr->GetActiveStrandsCoverageScale());
 						}
 					}
 				}
diff --git a/Engine/Source/Runtime/Renderer/Public/HairStrandsInterface.h b/Engine/Source/Runtime/Renderer/Public/HairStrandsInterface.h
--- a/Engine/Source/Runtime/Renderer/Public/HairStrandsInterface.h
+++ b/Engine/Source/Runtime/Renderer/Public/HairStrandsInterface.h
@@ -53,9 +53,7 @@ BEGIN_SHADER_PARAMETER_STRUCT(FHairStrandsInstanceCommonParameters, RENDERER_API
 	SHADER_PARAMETER(uint32, PointCount)
 	SHADER_PARAMETER(uint32, CurveCount)
 	SHADER_PARAMETER(uint32, RaytracingProceduralSplits)
-	SHADER_PARAMETER(uint32, bRaytracingGeometry)
-	SHADER_PARAMETER(uint32, bStableRasterization)
-	SHADER_PARAMETER(uint32, bScatterSceneLighting)
+	SHADER_PARAMETER(uint32, Flags)
 	SHADER_PARAMETER(uint32, bSimulation)
 	SHADER_PARAMETER(uint32, bSingleGuide)
 	SHADER_PARAMETER(FVector3f, PositionOffset)
