PSO Precache miss fixes:
- precache PSOs for GeometryCollectionComponent without custom renderer
- precache PSOs for Volumetric fog materials
- precache PSOs for projected shadow rendering using depth pass processor with different depth stencil state and render target
- always precache shadow rendering PSOs for default stream (non Position and PositionAndNormalOnlyStream)

diff --git a/Engine/Source/Runtime/Experimental/GeometryCollectionEngine/Private/GeometryCollection/GeometryCollectionComponent.cpp b/Engine/Source/Runtime/Experimental/GeometryCollectionEngine/Private/GeometryCollection/GeometryCollectionComponent.cpp
--- a/Engine/Source/Runtime/Experimental/GeometryCollectionEngine/Private/GeometryCollection/GeometryCollectionComponent.cpp
+++ b/Engine/Source/Runtime/Experimental/GeometryCollectionEngine/Private/GeometryCollection/GeometryCollectionComponent.cpp
@@ -14,6 +14,7 @@
 #include "Engine/Engine.h"
 #include "Engine/InstancedStaticMesh.h"
 #include "Engine/StaticMeshSocket.h"
+#include "MaterialDomain.h"
 #include "Field/FieldSystemComponent.h"
 #include "GeometryCollection/Facades/CollectionHierarchyFacade.h"
 #include "GeometryCollection/GeometryCollection.h"
@@ -64,6 +65,7 @@
 #include "Chaos/ChaosGameplayEventDispatcher.h"
 
 #include "Rendering/NaniteResources.h"
+#include "NaniteVertexFactory.h"
 #include "PrimitiveSceneInfo.h"
 #include "GeometryCollection/GeometryCollectionEngineRemoval.h"
 #include "GeometryCollection/Facades/CollectionAnchoringFacade.h"
@@ -6889,6 +6891,12 @@ void UGeometryCollectionComponent::PostLoad()
 {
 	Super::PostLoad();
 
+	// If there is a rest collection and no custom renderer then precache PSOs required to render the mesh
+	if (RestCollection && !CanUseCustomRenderer())
+	{
+		PrecachePSOs();
+	}
+
 	//
 	// The UGeometryCollectionComponent::PhysicalMaterial_DEPRECATED needs
 	// to be transferred to the BodyInstance simple material. Going forward
@@ -6910,6 +6918,47 @@ void UGeometryCollectionComponent::PostLoad()
 	}
 }
 
+void UGeometryCollectionComponent::CollectPSOPrecacheData(const FPSOPrecacheParams& BasePrecachePSOParams, FComponentPSOPrecacheParamsList& OutParams)
+{
+	check(RestCollection);
+
+	FPSOPrecacheVertexFactoryDataList VFDataList;
+	const FVertexFactoryType* VFType = nullptr;
+	if (UseNanite(GMaxRHIShaderPlatform) &&
+		RestCollection->EnableNanite &&
+		RestCollection->HasNaniteData() &&
+		GGeometryCollectionNanite != 0)
+	{
+		if (NaniteLegacyMaterialsSupported())
+		{
+			VFDataList.Add(FPSOPrecacheVertexFactoryData(&Nanite::FVertexFactory::StaticType));
+		}
+
+		if (NaniteComputeMaterialsSupported())
+		{
+			VFDataList.Add(FPSOPrecacheVertexFactoryData(&FNaniteVertexFactory::StaticType));
+		}
+	}	
+	else if (RestCollection->HasMeshData())
+	{
+		VFDataList.Add(FPSOPrecacheVertexFactoryData(&FGeometryCollectionVertexFactory::StaticType));
+	}
+
+	const TManagedArray<FGeometryCollectionSection>& SectionsArray = RestCollection->GetGeometryCollection()->Sections;
+	for (int32 SectionIndex = 0; SectionIndex < SectionsArray.Num(); ++SectionIndex)
+	{
+		const FGeometryCollectionSection& MeshSection = SectionsArray[SectionIndex];
+		const bool bValidMeshSection = MeshSection.MaterialID != INDEX_NONE;
+		UMaterialInterface* MaterialInterface = bValidMeshSection ? GetMaterial(MeshSection.MaterialID) : UMaterial::GetDefaultMaterial(MD_Surface);
+
+		FComponentPSOPrecacheParams& ComponentParams = OutParams[OutParams.AddDefaulted()];
+		ComponentParams.Priority = EPSOPrecachePriority::Medium;
+		ComponentParams.MaterialInterface = MaterialInterface;
+		ComponentParams.VertexFactoryDataList = VFDataList;
+		ComponentParams.PSOPrecacheParams = BasePrecachePSOParams;
+	}
+}
+
 void UGeometryCollectionComponent::SetPhysMaterialOverride(UPhysicalMaterial* NewPhysMaterial)
 {
 	Super::SetPhysMaterialOverride(NewPhysMaterial);
diff --git a/Engine/Source/Runtime/Experimental/GeometryCollectionEngine/Public/GeometryCollection/GeometryCollectionComponent.h b/Engine/Source/Runtime/Experimental/GeometryCollectionEngine/Public/GeometryCollection/GeometryCollectionComponent.h
--- a/Engine/Source/Runtime/Experimental/GeometryCollectionEngine/Public/GeometryCollection/GeometryCollectionComponent.h
+++ b/Engine/Source/Runtime/Experimental/GeometryCollectionEngine/Public/GeometryCollection/GeometryCollectionComponent.h
@@ -1461,6 +1461,9 @@ protected:
 	UPROPERTY(Transient)
 	TScriptInterface<IGeometryCollectionExternalRenderInterface> CustomRenderer;
 
+	/** Collect all the PSO precache data used by the geometry collection */
+	GEOMETRYCOLLECTIONENGINE_API virtual void CollectPSOPrecacheData(const FPSOPrecacheParams& BasePrecachePSOParams, FComponentPSOPrecacheParamsList& OutParams) override;
+
 	/** Populate the dynamic particle data for the render thread. */
 	GEOMETRYCOLLECTIONENGINE_API FGeometryCollectionDynamicData* InitDynamicData(bool bInitialization);
 
diff --git a/Engine/Source/Runtime/Renderer/Private/DepthRendering.cpp b/Engine/Source/Runtime/Renderer/Private/DepthRendering.cpp
--- a/Engine/Source/Runtime/Renderer/Private/DepthRendering.cpp
+++ b/Engine/Source/Runtime/Renderer/Private/DepthRendering.cpp
@@ -897,6 +897,33 @@ void FDepthPassMeshProcessor::CollectPSOInitializersInternal(
 		bPositionOnly ? EMeshPassFeatures::PositionOnly : EMeshPassFeatures::Default,
 		true /*bRequired*/,
 		PSOInitializers);
+
+	// Also cache with project shadow depth stencil state (see FProjectedShadowInfo::SetupMeshDrawCommandsForProjectionStenciling)
+	{
+		// Set stencil to one.
+		DrawRenderState.SetDepthStencilState(
+			TStaticDepthStencilState<
+			false, CF_DepthNearOrEqual,
+			true, CF_Always, SO_Keep, SO_Keep, SO_Replace,
+			false, CF_Always, SO_Keep, SO_Keep, SO_Keep,
+			0xff, 0xff
+			>::GetRHI());
+
+		AddRenderTargetInfo(PF_B8G8R8A8, TexCreate_RenderTargetable | TexCreate_ShaderResource, RenderTargetsInfo);
+
+		AddGraphicsPipelineStateInitializer(
+			VertexFactoryData,
+			MaterialResource,
+			DrawRenderState,
+			RenderTargetsInfo,
+			DepthPassShaders,
+			MeshFillMode,
+			MeshCullMode,
+			PrimitiveType,
+			bPositionOnly ? EMeshPassFeatures::PositionOnly : EMeshPassFeatures::Default,
+			true /*bRequired*/,
+			PSOInitializers);
+	}
 }
 
 bool FDepthPassMeshProcessor::UseDefaultMaterial(const FMaterial& Material, bool bMaterialModifiesMeshPosition, bool bSupportPositionOnlyStream, bool bVFTypeSupportsNullPixelShader, bool& bPositionOnly)
diff --git a/Engine/Source/Runtime/Renderer/Private/ShadowDepthRendering.cpp b/Engine/Source/Runtime/Renderer/Private/ShadowDepthRendering.cpp
--- a/Engine/Source/Runtime/Renderer/Private/ShadowDepthRendering.cpp
+++ b/Engine/Source/Runtime/Renderer/Private/ShadowDepthRendering.cpp
@@ -2299,11 +2299,10 @@ void FShadowDepthPassMeshProcessor::CollectPSOInitializersForEachStreamSetup(
 	{
 		CollectPSOInitializersInternal(VertexFactoryData, MaterialResource, InShadowDepthType, MeshFillMode, MeshCullMode, SupportsPositionAndNormalOnlyStream, bRequired, PSOInitializers);
 	}
-	else
-	{
-		SupportsPositionAndNormalOnlyStream = false;
-		CollectPSOInitializersInternal(VertexFactoryData, MaterialResource, InShadowDepthType, MeshFillMode, MeshCullMode, SupportsPositionAndNormalOnlyStream, bRequired, PSOInitializers);
-	}
+	
+	// Always precache with the default streams (non PositionOnly or PositionAndNormalOnly stream)
+	SupportsPositionAndNormalOnlyStream = false;
+	CollectPSOInitializersInternal(VertexFactoryData, MaterialResource, InShadowDepthType, MeshFillMode, MeshCullMode, SupportsPositionAndNormalOnlyStream, bRequired, PSOInitializers);
 }
 
 FShadowDepthPassMeshProcessor::FShadowDepthPassMeshProcessor(
diff --git a/Engine/Source/Runtime/Renderer/Private/VolumetricFog.cpp b/Engine/Source/Runtime/Renderer/Private/VolumetricFog.cpp
--- a/Engine/Source/Runtime/Renderer/Private/VolumetricFog.cpp
+++ b/Engine/Source/Runtime/Renderer/Private/VolumetricFog.cpp
@@ -1207,6 +1207,15 @@ FVector2f GetVolumetricFogFroxelToScreenSVPosRatio(const FViewInfo& View)
 	return ClipRatio;
 }
 
+FRDGTextureDesc GetVolumetricFogRDGTextureDesc(const FIntVector& VolumetricFogResourceGridSize)
+{
+	return FRDGTextureDesc::Create3D(
+		VolumetricFogResourceGridSize,
+		PF_FloatRGBA,
+		FClearValueBinding::Black,
+		TexCreate_ShaderResource | TexCreate_RenderTargetable | TexCreate_UAV | TexCreate_ReduceMemoryWithTilingMode | TexCreate_3DTiling);
+}
+
 void SetupVolumetricFogGlobalData(const FViewInfo& View, FVolumetricFogGlobalData& Parameters)
 {
 	const FScene* Scene = (FScene*)View.Family->Scene;
@@ -1426,11 +1435,7 @@ void FSceneRenderer::ComputeVolumetricFog(FRDGBuilder& GraphBuilder,
 		View.VolumetricFogResources.IntegratedLightScatteringTexture = nullptr;
 		TRDGUniformBufferRef<FFogUniformParameters> FogUniformBuffer = CreateFogUniformBuffer(GraphBuilder, View);
 
-		FRDGTextureDesc VolumeDesc(FRDGTextureDesc::Create3D(
-			VolumetricFogResourceGridSize,
-			PF_FloatRGBA, 
-			FClearValueBinding::Black, 
-			TexCreate_ShaderResource | TexCreate_RenderTargetable | TexCreate_UAV | TexCreate_ReduceMemoryWithTilingMode | TexCreate_3DTiling));
+		FRDGTextureDesc VolumeDesc = GetVolumetricFogRDGTextureDesc(VolumetricFogResourceGridSize);
 
 		FRDGTextureDesc VolumeDescFastVRAM = VolumeDesc;
 		VolumeDescFastVRAM.Flags |= GFastVRamConfig.VolumetricFog;
diff --git a/Engine/Source/Runtime/Renderer/Private/VolumetricFog.h b/Engine/Source/Runtime/Renderer/Private/VolumetricFog.h
--- a/Engine/Source/Runtime/Renderer/Private/VolumetricFog.h
+++ b/Engine/Source/Runtime/Renderer/Private/VolumetricFog.h
@@ -20,6 +20,8 @@ extern FVector2f GetVolumetricFogPrevUVMaxForTemporalBlend(const FVector2f& View
 
 extern FVector2f GetVolumetricFogFroxelToScreenSVPosRatio(const FViewInfo& View);
 
+extern FRDGTextureDesc GetVolumetricFogRDGTextureDesc(const FIntVector& VolumetricFogResourceGridSize);
+
 extern bool DoesPlatformSupportVolumetricFogVoxelization(const FStaticShaderPlatform Platform);
 extern bool ShouldRenderVolumetricFog(const FScene* Scene, const FSceneViewFamily& ViewFamily);
 extern const FProjectedShadowInfo* GetShadowForInjectionIntoVolumetricFog(const FVisibleLightInfo& VisibleLightInfo);
diff --git a/Engine/Source/Runtime/Renderer/Private/VolumetricFogVoxelization.cpp b/Engine/Source/Runtime/Renderer/Private/VolumetricFogVoxelization.cpp
--- a/Engine/Source/Runtime/Renderer/Private/VolumetricFogVoxelization.cpp
+++ b/Engine/Source/Runtime/Renderer/Private/VolumetricFogVoxelization.cpp
@@ -393,6 +393,8 @@ public:
 		checkf(false, TEXT("Default AddMeshBatch can't be used as rendering requires extra parameters per pass."));
 	}
 
+	virtual void CollectPSOInitializers(const FSceneTexturesConfig& SceneTexturesConfig, const FMaterial& Material, const FPSOPrecacheVertexFactoryData& VertexFactoryData, const FPSOPrecacheParams& PreCacheParams, TArray<FPSOPrecacheData>& PSOInitializers) override final;
+
 private:
 	bool TryAddMeshBatch(
 		const FMeshBatch& RESTRICT MeshBatch,
@@ -415,8 +417,10 @@ private:
 	FMeshPassProcessorRenderState PassDrawRenderState;
 };
 
+static const TCHAR* VoxelizeVolumePassName = TEXT("VoxelizeVolume");
+
 FVoxelizeVolumeMeshProcessor::FVoxelizeVolumeMeshProcessor(const FScene* Scene, ERHIFeatureLevel::Type FeatureLevel, const FViewInfo* InViewIfDynamicMeshCommand, FMeshPassDrawListContext* InDrawListContext)
-	: FMeshPassProcessor(EMeshPass::Num, Scene, FeatureLevel, InViewIfDynamicMeshCommand, InDrawListContext)
+	: FMeshPassProcessor(VoxelizeVolumePassName, Scene, FeatureLevel, InViewIfDynamicMeshCommand, InDrawListContext)
 {
 	PassDrawRenderState.SetBlendState(TStaticBlendState<
 		CW_RGBA, BO_Add, BF_One, BF_One, BO_Add, BF_One, BF_One,
@@ -442,6 +446,51 @@ void FVoxelizeVolumeMeshProcessor::AddMeshBatch(const FMeshBatch& RESTRICT MeshB
 	}
 }
 
+static bool GetVoxelizeVolumePassShaders(
+	const FMaterial& Material,
+	const FVertexFactoryType* VertexFactoryType,
+	ERHIFeatureLevel::Type FeatureLevel,
+	bool bUsePrimitiveSphere,
+	TShaderRef<FVoxelizeVolumeVS>& VertexShader,
+	TShaderRef<FVoxelizeVolumeGS>& GeometryShader,
+	TShaderRef<FVoxelizeVolumePS>& PixelShader)
+{
+	FMaterialShaderTypes ShaderTypes;
+	if (bUsePrimitiveSphere)
+	{
+		ShaderTypes.AddShaderType<TVoxelizeVolumeVS<VMode_Primitive_Sphere>>();
+		if (RHISupportsGeometryShaders(GShaderPlatformForFeatureLevel[FeatureLevel]))
+		{
+			ShaderTypes.AddShaderType<TVoxelizeVolumeGS<VMode_Primitive_Sphere>>();
+		}
+		ShaderTypes.AddShaderType<TVoxelizeVolumePS<VMode_Primitive_Sphere>>();
+	}
+	else
+	{
+		ShaderTypes.AddShaderType<TVoxelizeVolumeVS<VMode_Object_Box>>();
+		if (RHISupportsGeometryShaders(GShaderPlatformForFeatureLevel[FeatureLevel]))
+		{
+			ShaderTypes.AddShaderType<TVoxelizeVolumeGS<VMode_Object_Box>>();
+		}
+		ShaderTypes.AddShaderType<TVoxelizeVolumePS<VMode_Object_Box>>();
+	}
+
+	FMaterialShaders Shaders;
+	if (!Material.TryGetShaders(ShaderTypes, VertexFactoryType, Shaders))
+	{
+		return false;
+	}
+
+	Shaders.TryGetVertexShader(VertexShader);
+	if (RHISupportsGeometryShaders(GShaderPlatformForFeatureLevel[FeatureLevel]))
+	{
+		Shaders.TryGetGeometryShader(GeometryShader);
+	}
+	Shaders.TryGetPixelShader(PixelShader);
+
+	return true;
+}
+
 bool FVoxelizeVolumeMeshProcessor::TryAddMeshBatch(
 	const FMeshBatch& RESTRICT MeshBatch,
 	uint64 BatchElementMask,
@@ -469,47 +518,24 @@ bool FVoxelizeVolumeMeshProcessor::Process(
 	ERasterizerCullMode MeshCullMode)
 {
 	const FVertexFactory* VertexFactory = MeshBatch.VertexFactory;
+	const bool bUsePrimitiveSphere = VertexFactory != GQuadMeshVertexFactory;
 
 	TMeshProcessorShaders<
 		FVoxelizeVolumeVS,
 		FVoxelizeVolumePS,
 		FVoxelizeVolumeGS> PassShaders;
-
-	const bool bUsePrimitiveSphere = VertexFactory != GQuadMeshVertexFactory;
-
-	FMaterialShaderTypes ShaderTypes;
-	if (bUsePrimitiveSphere)
-	{
-		ShaderTypes.AddShaderType<TVoxelizeVolumeVS<VMode_Primitive_Sphere>>();
-		if (RHISupportsGeometryShaders(GShaderPlatformForFeatureLevel[FeatureLevel]))
-		{
-			ShaderTypes.AddShaderType<TVoxelizeVolumeGS<VMode_Primitive_Sphere>>();
-		}
-		ShaderTypes.AddShaderType<TVoxelizeVolumePS<VMode_Primitive_Sphere>>();
-	}
-	else
-	{
-		ShaderTypes.AddShaderType<TVoxelizeVolumeVS<VMode_Object_Box>>();
-		if (RHISupportsGeometryShaders(GShaderPlatformForFeatureLevel[FeatureLevel]))
-		{
-			ShaderTypes.AddShaderType<TVoxelizeVolumeGS<VMode_Object_Box>>();
-		}
-		ShaderTypes.AddShaderType<TVoxelizeVolumePS<VMode_Object_Box>>();
-	}
-
-	FMaterialShaders Shaders;
-	if (!MaterialResource.TryGetShaders(ShaderTypes, VertexFactory->GetType(), Shaders))
+	if (!GetVoxelizeVolumePassShaders(
+		MaterialResource,
+		VertexFactory->GetType(),
+		FeatureLevel,
+		bUsePrimitiveSphere,
+		PassShaders.VertexShader,
+		PassShaders.GeometryShader,
+		PassShaders.PixelShader))
 	{
 		return false;
 	}
-
-	Shaders.TryGetVertexShader(PassShaders.VertexShader);
-	if (RHISupportsGeometryShaders(GShaderPlatformForFeatureLevel[FeatureLevel]))
-	{
-		Shaders.TryGetGeometryShader(PassShaders.GeometryShader);
-	}
-	Shaders.TryGetPixelShader(PassShaders.PixelShader);
-
+	
 	const FMeshDrawCommandSortKey SortKey = CalculateMeshStaticSortKey(PassShaders.VertexShader, PassShaders.PixelShader);
 
 	for (int32 VoxelizationPassIndex = 0; VoxelizationPassIndex < NumVoxelizationPasses; VoxelizationPassIndex++)
@@ -538,6 +564,77 @@ bool FVoxelizeVolumeMeshProcessor::Process(
 	return true;
 }
 
+void FVoxelizeVolumeMeshProcessor::CollectPSOInitializers(
+	const FSceneTexturesConfig& SceneTexturesConfig,
+	const FMaterial& Material,
+	const FPSOPrecacheVertexFactoryData& VertexFactoryData,
+	const FPSOPrecacheParams& PreCacheParams,
+	TArray<FPSOPrecacheData>& PSOInitializers)
+{
+	if (Material.GetMaterialDomain() != MD_Volume)
+	{
+		return;
+	}
+
+	const FMeshDrawingPolicyOverrideSettings OverrideSettings = ComputeMeshOverrideSettings(PreCacheParams);
+	ERasterizerFillMode MeshFillMode = ComputeMeshFillMode(Material, OverrideSettings);
+	ERasterizerCullMode MeshCullMode = CM_None;
+
+	FRDGTextureDesc VolumeDesc = GetVolumetricFogRDGTextureDesc(FIntVector(0));
+
+	FGraphicsPipelineRenderTargetsInfo RenderTargetsInfo;
+	RenderTargetsInfo.NumSamples = 1;
+	AddRenderTargetInfo(VolumeDesc.Format, VolumeDesc.Flags, RenderTargetsInfo);
+
+	static const auto CVarVolumetrixFogEmissive = IConsoleManager::Get().FindConsoleVariable(TEXT("r.VolumetricFog.Emissive"));
+	const bool bUseEmissive = (!CVarVolumetrixFogEmissive) || (CVarVolumetrixFogEmissive->GetInt() > 0);
+	if (bUseEmissive)
+	{
+		AddRenderTargetInfo(VolumeDesc.Format, VolumeDesc.Flags, RenderTargetsInfo);
+	}
+
+	const auto AddPSOInitializer = [&](bool bUsePrimitiveSphere)
+	{
+		TMeshProcessorShaders<
+			FVoxelizeVolumeVS,
+			FVoxelizeVolumePS,
+			FVoxelizeVolumeGS> PassShaders;
+		if (!GetVoxelizeVolumePassShaders(
+			Material,
+			VertexFactoryData.VertexFactoryType,
+			FeatureLevel,
+			bUsePrimitiveSphere,
+			PassShaders.VertexShader,
+			PassShaders.GeometryShader,
+			PassShaders.PixelShader))
+		{
+			return;
+		}
+
+		AddGraphicsPipelineStateInitializer(
+			VertexFactoryData,
+			Material,
+			PassDrawRenderState,
+			RenderTargetsInfo,
+			PassShaders,
+			MeshFillMode,
+			MeshCullMode,
+			(EPrimitiveType)PreCacheParams.PrimitiveType,
+			EMeshPassFeatures::Default,
+			true /*bRequired*/,
+			PSOInitializers);
+	};
+
+	AddPSOInitializer(true /*bUsePrimitiveSphere*/);
+	AddPSOInitializer(false /*bUsePrimitiveSphere*/);	
+}
+
+IPSOCollector* CreatePSOCollectorVoxelizeVolume(ERHIFeatureLevel::Type FeatureLevel)
+{
+	return new FVoxelizeVolumeMeshProcessor(nullptr, FeatureLevel, nullptr, nullptr);
+}
+FRegisterPSOCollectorCreateFunction RegisterPSOCollectorVoxelizeVolume(&CreatePSOCollectorVoxelizeVolume, EShadingPath::Deferred, VoxelizeVolumePassName);
+
 void VoxelizeVolumePrimitive(FVoxelizeVolumeMeshProcessor& PassMeshProcessor,
 	FRHICommandListImmediate& RHICmdList,
 	const FViewInfo& View,
