Remove unused temporal hair composition.

diff --git a/Engine/Shaders/Private/HairStrands/HairStrandsVisibilityComposeSubPixelPS.usf b/Engine/Shaders/Private/HairStrands/HairStrandsVisibilityComposeSubPixelPS.usf
--- a/Engine/Shaders/Private/HairStrands/HairStrandsVisibilityComposeSubPixelPS.usf
+++ b/Engine/Shaders/Private/HairStrands/HairStrandsVisibilityComposeSubPixelPS.usf
@@ -36,9 +36,6 @@ uint bHasHoldout;
 Texture2D<float>	HairDOFDepthTexture;
 Texture2D<float4>	HairLightingSampleBuffer;
 
-uint TemporalLayerCount;
-Texture2D HairTemporalAccumulationTexture;
-
 float4 EvaluateVolumetric(float3 WorldPosition, float SceneDepth)
 {
 	const float3 CameraRelative_WorldPosition = WorldPosition - LWCHackToFloat(PrimaryView.WorldCameraOrigin);
@@ -89,15 +86,6 @@ void ComposeSamplePS(
 	
 	const float HairDeviceZ = HairStrands.HairOnlyDepthTexture.Load(PixelCoord).x;
 
-#if PERMUTATION_TEMPORAL == 1
-	OutColor.rgba = HairTemporalAccumulationTexture.Load(PixelCoord);
-
-	// Approximation: convert running temporal average layer coverage (alpha) to average haircount, multiply by number of layers to get haircount sum, then convert haircount sum to final coverage using equation: Coverage = saturate(min(1.0 - pow(2.71828, -HairCount/1000.0), 1.0) / 0.98);
-	float HairCountAverage = (log(1.0 - (OutColor.a * 0.98)) / log(2.71828))*-1000.0;
-	float HairCountSum = HairCountAverage*TemporalLayerCount;
-	OutColor.a = saturate(min(1.0 - pow(2.71828, -HairCountSum/1000.0), 1.0) / 0.98);
-
-#else // PERMUTATION_TEMPORAL
 	const FNodeDesc NodeDesc = DecodeNodeDesc(HairStrands.HairSampleOffset.Load(PixelCoord));
 	if (NodeDesc.Count == 0)
 	{
@@ -138,7 +126,6 @@ void ComposeSamplePS(
 	// If any sample support holdout, then the lighting contribution is set to 0, to mask out the current hair pixel
 	OutColor.rgb *= bHoldout ? 0 : PixelCoverage;
 	OutColor.a = PixelCoverage;
-#endif // PERMUTATION_TEMPORAL
 
 	OutDepth = HairDeviceZ;
 	if (bComposeDofDepth)
diff --git a/Engine/Source/Runtime/Renderer/Private/HairStrands/HairStrandsComposition.cpp b/Engine/Source/Runtime/Renderer/Private/HairStrands/HairStrandsComposition.cpp
--- a/Engine/Source/Runtime/Renderer/Private/HairStrands/HairStrandsComposition.cpp
+++ b/Engine/Source/Runtime/Renderer/Private/HairStrands/HairStrandsComposition.cpp
@@ -222,8 +222,7 @@ class FHairVisibilityComposeSamplePS : public FGlobalShader
 	SHADER_USE_PARAMETER_STRUCT(FHairVisibilityComposeSamplePS, FGlobalShader);
 
 	class FDebug : SHADER_PERMUTATION_BOOL("PERMUTATION_DEBUG");
-	class FTemporal : SHADER_PERMUTATION_BOOL("PERMUTATION_TEMPORAL");
-	using FPermutationDomain = TShaderPermutationDomain<FDebug, FTemporal>;
+	using FPermutationDomain = TShaderPermutationDomain<FDebug>;
 	
 	BEGIN_SHADER_PARAMETER_STRUCT(FParameters, )
 		SHADER_PARAMETER_STRUCT_REF(FViewUniformShaderParameters, ViewUniformBuffer)
@@ -234,8 +233,6 @@ class FHairVisibilityComposeSamplePS : public FGlobalShader
 		SHADER_PARAMETER_RDG_UNIFORM_BUFFER(FHairStrandsViewUniformParameters, HairStrands)
 		SHADER_PARAMETER_RDG_TEXTURE(Texture2D, HairLightingSampleBuffer)
 		SHADER_PARAMETER_RDG_TEXTURE(Texture2D, HairDOFDepthTexture)
-		SHADER_PARAMETER(uint32, TemporalLayerCount)
-		SHADER_PARAMETER_RDG_TEXTURE(Texture2D, HairTemporalAccumulationTexture)
 		SHADER_PARAMETER_RDG_TEXTURE_UAV(RWTexture2D<uint>, OutMetaTexture)
 		SHADER_PARAMETER_RDG_UNIFORM_BUFFER(FFogUniformParameters, FogStruct)
 		RENDER_TARGET_BINDING_SLOTS()
@@ -260,14 +257,11 @@ static void AddHairVisibilityComposeSamplePass(
 	const FRDGTextureRef& HairDOFDepthTexture,
 	const FRDGTextureRef& HairTemporalAccumulationTexture,
 	FRDGTextureRef& OutColorTexture,
-	FRDGTextureRef& OutDepthTexture,
-	const bool bTemporalLayeringEnabled)
+	FRDGTextureRef& OutDepthTexture)
 {
 	check(VisibilityData.SampleLightingTexture);
 	const bool bDOFEnable = HairDOFDepthTexture != nullptr ? 1 : 0;
 
-	const bool bTemporal = false; //HAIR_TODO: remove
-
 	TRDGUniformBufferRef<FFogUniformParameters> FogBuffer = CreateFogUniformBuffer(GraphBuilder, View);
 
 	FHairVisibilityComposeSamplePS::FParameters* Parameters = GraphBuilder.AllocParameters<FHairVisibilityComposeSamplePS::FParameters>();
@@ -278,23 +272,20 @@ static void AddHairVisibilityComposeSamplePass(
 	Parameters->ViewUniformBuffer = View.ViewUniformBuffer;
 	Parameters->HairStrands = View.HairStrandsViewData.UniformBuffer;
 	Parameters->FogStruct = FogBuffer;
-	Parameters->TemporalLayerCount = 1;
 	Parameters->bHasHoldout = bHasHoldout ? 1u : 0u;
-	Parameters->HairTemporalAccumulationTexture = bTemporal ? HairTemporalAccumulationTexture : GSystemTextures.GetBlackDummy(GraphBuilder);
-	Parameters->RenderTargets[0] = FRenderTargetBinding(OutColorTexture, bTemporal ? ERenderTargetLoadAction::ELoad : (bTemporalLayeringEnabled ? ERenderTargetLoadAction::EClear : ERenderTargetLoadAction::ELoad) );
+	Parameters->RenderTargets[0] = FRenderTargetBinding(OutColorTexture, ERenderTargetLoadAction::ELoad);
 	Parameters->RenderTargets.DepthStencil = FDepthStencilBinding(OutDepthTexture, ERenderTargetLoadAction::ELoad, ERenderTargetLoadAction::ELoad, FExclusiveDepthStencil::DepthWrite_StencilRead);
 
 	const bool bDebugComposition = View.Family->EngineShowFlags.LODColoration;
 	FHairVisibilityComposeSamplePS::FPermutationDomain PermutationVector;
 	PermutationVector.Set<FHairVisibilityComposeSamplePS::FDebug>(bDebugComposition);
-	PermutationVector.Set<FHairVisibilityComposeSamplePS::FTemporal>(bTemporal);
 	TShaderMapRef<FHairVisibilityComposeSamplePS> PixelShader(View.ShaderMap, PermutationVector);
 	InternalCommonDrawPass(
 		GraphBuilder,
 		RDG_EVENT_NAME("HairStrands::ComposeSample"),
 		View,
 		OutColorTexture->Desc.Extent,
-		bTemporal ? EHairStrandsCommonPassType::Composition : (bTemporalLayeringEnabled ? EHairStrandsCommonPassType::Blit : EHairStrandsCommonPassType::Composition),
+		EHairStrandsCommonPassType::Composition,
 		false,
 		VisibilityData.TileData,
 		PixelShader,
@@ -587,43 +578,6 @@ static void InternalRenderHairComposition(
 		const bool bHasHoldout = HasHairFlags(View.HairStrandsViewData.Flags, HAIR_FLAGS_HOLDOUT);
 		const bool bTemporalLayeringEnabled = false; // HAIR_TODO: remove
 
-		if (bTemporalLayeringEnabled)
-		{
-			FIntPoint OutputResolution = SceneColorTexture->Desc.Extent;
-
-			FRDGTextureDesc HairColorInputDesc = FRDGTextureDesc::Create2D(OutputResolution, PF_FloatRGBA, FClearValueBinding(FLinearColor(0, 0, 0, 0)), TexCreate_RenderTargetable | TexCreate_ShaderResource | TexCreate_UAV, 1);
-			FRDGTextureRef HairColorInputTexture = GraphBuilder.CreateTexture(HairColorInputDesc, TEXT("Hair.ColorInputTexture"));
-
-			AddHairVisibilityComposeSamplePass(
-				GraphBuilder,
-				View,
-				VisibilityData,
-				bHasHoldout,
-				DOFDepth,
-				nullptr,
-				HairColorInputTexture,
-				SceneDepthTexture,
-				bTemporalLayeringEnabled);
-
-
-			check(View.ViewState);
-			FTAAPassParameters TAASettings(View);
-			TAASettings.SceneDepthTexture = SceneDepthTexture;
-			TAASettings.SceneVelocityTexture = SceneVelocityTexture;
-			TAASettings.Pass = ETAAPassConfig::Hair;
-			TAASettings.SceneColorInput = HairColorInputTexture;
-			TAASettings.bOutputRenderTargetable = true;
-
-			FTAAOutputs TAAOutputs = AddTemporalAAPass(
-				GraphBuilder,
-				View,
-				TAASettings,
-				View.PrevViewInfo.HairHistory,
-				&View.ViewState->PrevFrameViewInfo.HairHistory);
-
-			AccumulatedColor = TAAOutputs.SceneColor;
-		}
-
 		AddHairVisibilityComposeSamplePass(
 			GraphBuilder,
 			View,
@@ -632,8 +586,7 @@ static void InternalRenderHairComposition(
 			DOFDepth,
 			AccumulatedColor,
 			SceneColorTexture,
-			SceneDepthTexture,
-			bTemporalLayeringEnabled);
+			SceneDepthTexture);
 
 		if (bHasHoldout)
 		{
