SampledDirectLighting - work on denoising
* Initial version of a spatial filter guided by tempora variance
* Recompute temporal first moment from color instead of tracking it separately
* Clear light samples in empty tiles instead of using SharedTileMask for skipping traces in order to simplify compact traces shader
* Fixed shading dispatch overlap
* Apply exposure correction factor in order to correct for exposure changes between frames

diff --git a/Engine/Shaders/Private/SampledDirectLighting/SampledDirectLighting.usf b/Engine/Shaders/Private/SampledDirectLighting/SampledDirectLighting.usf
--- a/Engine/Shaders/Private/SampledDirectLighting/SampledDirectLighting.usf
+++ b/Engine/Shaders/Private/SampledDirectLighting/SampledDirectLighting.usf
@@ -7,7 +7,6 @@ groupshared uint SharedTileMode[THREADGROUP_SIZE * THREADGROUP_SIZE];
 
 RWStructuredBuffer<uint> RWTileAllocator;
 RWStructuredBuffer<uint> RWTileData;
-RWTexture2D<uint> RWDownsampledTileMask;
 uint TileDataStride;
 uint DownsampledTileDataStride;
 uint2 DownsampledViewSize;
@@ -44,15 +43,15 @@ void TileClassificationCS(
 
 		if (GBuffer.ShadingModelID == SHADINGMODELID_UNLIT)
 		{
-			// Skip mark
+			SharedTileMode[LinearThreadIndex] = 0x1;
 		}
 		else if (GBuffer.ShadingModelID == SHADINGMODELID_DEFAULT_LIT)
 		{
-			SharedTileMode[LinearThreadIndex] = 0x1;
+			SharedTileMode[LinearThreadIndex] = 0x2;
 		}
 		else
 		{
-			SharedTileMode[LinearThreadIndex] = 0x2;
+			SharedTileMode[LinearThreadIndex] = 0x4;
 		}
 	}
 
@@ -109,24 +108,24 @@ void TileClassificationCS(
 
 		uint2 TileCoord = ScreenCoord / TILE_SIZE;
 
-		#if DOWNSAMPLED_CLASSIFICATION
-		{
-			RWDownsampledTileMask[TileCoord] = SharedTileMode[0] != 0 ? 1 : 0;
-		}
-		#endif
-
-		if (SharedTileMode[0] & 0x2)
+		if (SharedTileMode[0] & 0x4)
 		{
 			uint SampledDirectLightingTileIndex = 0;
 			InterlockedAdd(RWTileAllocator[TILE_MODE_COMPLEX_SHADING], 1, SampledDirectLightingTileIndex);
 			RWTileData[DataStride * TILE_MODE_COMPLEX_SHADING + SampledDirectLightingTileIndex] = PackTile(TileCoord);
 		}
-		else if (SharedTileMode[0] & 0x1)
+		else if (SharedTileMode[0] & 0x2)
 		{
 			uint SampledDirectLightingTileIndex = 0;
 			InterlockedAdd(RWTileAllocator[TILE_MODE_SIMPLE_SHADING], 1, SampledDirectLightingTileIndex);
 			RWTileData[DataStride * TILE_MODE_SIMPLE_SHADING + SampledDirectLightingTileIndex] = PackTile(TileCoord);
 		}
+		else if (SharedTileMode[0] == 0x1)
+		{
+			uint SampledDirectLightingTileIndex = 0;
+			InterlockedAdd(RWTileAllocator[TILE_MODE_EMPTY], 1, SampledDirectLightingTileIndex);
+			RWTileData[DataStride * TILE_MODE_EMPTY + SampledDirectLightingTileIndex] = PackTile(TileCoord);
+		}
 	}
 }
 
diff --git a/Engine/Shaders/Private/SampledDirectLighting/SampledDirectLighting.ush b/Engine/Shaders/Private/SampledDirectLighting/SampledDirectLighting.ush
--- a/Engine/Shaders/Private/SampledDirectLighting/SampledDirectLighting.ush
+++ b/Engine/Shaders/Private/SampledDirectLighting/SampledDirectLighting.ush
@@ -38,7 +38,8 @@
 // Keep in sync with SampledDirectLighting::ETileType
 #define TILE_MODE_SIMPLE_SHADING	0
 #define TILE_MODE_COMPLEX_SHADING	1
-#define TILE_MODE_MAX				2
+#define TILE_MODE_EMPTY				2
+#define TILE_MODE_MAX				3
 
 uint PackTile(uint2 TileCoord)
 {
diff --git a/Engine/Shaders/Private/SampledDirectLighting/SampledDirectLightingComposite.usf b/Engine/Shaders/Private/SampledDirectLighting/SampledDirectLightingComposite.usf
--- a/Engine/Shaders/Private/SampledDirectLighting/SampledDirectLightingComposite.usf
+++ b/Engine/Shaders/Private/SampledDirectLighting/SampledDirectLightingComposite.usf
@@ -113,18 +113,6 @@ void CompositeLightSamplesCS(
 
 	GroupMemoryBarrierWithGroupSync();
 
-	if (View.GeneralPurposeTweak <= 0.5)
-	{
-		// Mark sub-tiles requiring shading tiles
-		if (SharedSampleWeights[GroupThreadId.x + SHARED_TILE_BORDER][GroupThreadId.y + SHARED_TILE_BORDER] > 0.0f)
-		{
-			uint2 ShadingTileOffset = GroupThreadId.xy / 4;
-			SharedShadingTileMask[ShadingTileOffset.x][ShadingTileOffset.y] = 1;
-		}
-	}
-
-	GroupMemoryBarrierWithGroupSync();
-
 	uint2 LocalDownsampledTileCoord = GroupThreadId.xy;
 	uint2 DownsampledScreenCoord = BaseDownsampledScreenCoord + LocalDownsampledTileCoord;
 	uint2 ShadingTileOffset = GroupThreadId.xy;
diff --git a/Engine/Shaders/Private/SampledDirectLighting/SampledDirectLightingHardwareRayTracing.usf b/Engine/Shaders/Private/SampledDirectLighting/SampledDirectLightingHardwareRayTracing.usf
--- a/Engine/Shaders/Private/SampledDirectLighting/SampledDirectLightingHardwareRayTracing.usf
+++ b/Engine/Shaders/Private/SampledDirectLighting/SampledDirectLightingHardwareRayTracing.usf
@@ -10,7 +10,6 @@
 
 float RayTracingBias;
 float RayTracingNormalBias;
-uint SampledDirectLightingStateFrameIndes;
 uint2 SampleViewSize;
 uint DebugMode;
 uint MaxTraversalIterations;
@@ -28,6 +27,7 @@ RaytracingAccelerationStructure TLAS;
 #endif
 
 RWTexture2D<uint> RWLightSamples;
+
 Texture2D<float> LightSampleRayDistance;
 StructuredBuffer<uint> CompactedTraceTexelData;
 StructuredBuffer<uint> CompactedTraceTexelAllocator;
@@ -118,13 +118,6 @@ RAY_TRACING_ENTRY_RAYGEN_OR_INLINE(HardwareRayTraceLightSamples)
 		}
 		#endif
 
-		if (Payload.IsHit())
-		{
-			LightSample.bCompleted = true;
-			LightSample.bVisible = false;
-			RWLightSamples[SampleCoord] = PackLightSample(LightSample);
-		}
-
 		#if DEBUG_MODE
 		if (bDebug && DebugMode == DEBUG_MODE_VISUALIZE_TRACING)
 		{
@@ -136,8 +129,15 @@ RAY_TRACING_ENTRY_RAYGEN_OR_INLINE(HardwareRayTraceLightSamples)
 			float3 RayEnd = Ray.Origin + Ray.Direction * (Payload.IsHit() ? Payload.HitT : Ray.TMax);
 			float4 RayColor = float4(LightData.Color.xyz / Luminance(LightData.Color.xyz), 1.0f);
 
-			AddLineTWS(Context, RayStart, RayEnd, RayColor);
+			AddLineTWS(Context, RayStart, RayEnd, RayColor);	
 		}
 		#endif
+
+		if (Payload.IsHit())
+		{
+			LightSample.bCompleted = true;
+			LightSample.bVisible = false;
+			RWLightSamples[SampleCoord] = PackLightSample(LightSample);
+		}
 	}
 }
\ No newline at end of file
diff --git a/Engine/Shaders/Private/SampledDirectLighting/SampledDirectLightingSampling.usf b/Engine/Shaders/Private/SampledDirectLighting/SampledDirectLightingSampling.usf
--- a/Engine/Shaders/Private/SampledDirectLighting/SampledDirectLightingSampling.usf
+++ b/Engine/Shaders/Private/SampledDirectLighting/SampledDirectLightingSampling.usf
@@ -1,7 +1,11 @@
 // Copyright Epic Games, Inc. All Rights Reserved.
 
 #ifndef NUM_SAMPLES_PER_PIXEL_1D
-	#error NUM_SAMPLES_PER_PIXEL_1D must be defined
+	#define NUM_SAMPLES_PER_PIXEL_1D 1
+#endif
+
+#ifndef TILE_TYPE
+	#define TILE_TYPE TILE_MODE_EMPTY
 #endif
 
 #include "../Common.ush"
@@ -450,4 +454,33 @@ void GenerateSamplesCS(
 			}
 		}
 	}
+}
+
+/**
+ * Clear some data for empty tiles, which won't be processed by GenerateLightSamplesCS
+ */
+[numthreads(THREADGROUP_SIZE, THREADGROUP_SIZE, 1)]
+void ClearLightSamplesCS(
+	uint3 GroupId : SV_GroupID,
+	uint3 GroupThreadId : SV_GroupThreadID,
+	uint3 DispatchThreadId : SV_DispatchThreadID)
+{
+	uint DownsampledTileIndex = GroupId.x;
+	if (DownsampledTileIndex < DownsampledTileAllocator[TILE_MODE_EMPTY])
+	{
+		uint2 DownsampledTileCoord = UnpackTile(DownsampledTileData[DownsampledTileIndex + TILE_MODE_EMPTY * DownsampledTileDataStride]);
+		uint2 DownsampledScreenCoord = DownsampledTileCoord * TILE_SIZE + GroupThreadId.xy;
+
+		if (all(DownsampledScreenCoord < DownsampledViewSize))
+		{
+			for (uint LightSampleY = 0; LightSampleY < NumSamplesPerPixel.y; ++LightSampleY)
+			{
+				for (uint LightSampleX = 0; LightSampleX < NumSamplesPerPixel.x; ++LightSampleX)
+				{
+					FLightSample LightSample = InitLightSample();
+					RWLightSamples[DownsampledScreenCoord * NumSamplesPerPixel + uint2(LightSampleX, LightSampleY)] = PackLightSample(LightSample);
+				}
+			}
+		}
+	}
 }
\ No newline at end of file
diff --git a/Engine/Shaders/Private/SampledDirectLighting/SampledDirectLightingShading.usf b/Engine/Shaders/Private/SampledDirectLighting/SampledDirectLightingShading.usf
--- a/Engine/Shaders/Private/SampledDirectLighting/SampledDirectLightingShading.usf
+++ b/Engine/Shaders/Private/SampledDirectLighting/SampledDirectLightingShading.usf
@@ -4,9 +4,8 @@
 #include "../BlueNoise.ush"
 #include "SampledDirectLighting.ush"
 
-RWTexture2D<float3> RWDiffuseLighting;
-RWTexture2D<float3> RWSpecularLighting;
-RWTexture2D<float4> RWLuminanceMoments;
+RWTexture2D<float3> RWResolvedDiffuseLighting;
+RWTexture2D<float3> RWResolvedSpecularLighting;
 
 Texture2D<float> ShadingTileAtlas;
 Texture2D<uint> ShadingTileGridAllocator;
@@ -54,6 +53,9 @@ void ShadeLightSamplesCS(
 			float2 ScreenUV = (ScreenCoord + 0.5f) * View.BufferSizeAndInvSize.zw;
 			FGBufferData GBuffer = GetGBufferData(ScreenUV);
 
+			float3 DiffuseLighting = 0.0f;
+			float3 SpecularLighting = 0.0f;
+
 			if (GBuffer.Depth > 0.0f)
 			{
 				uint2 ShadingTileCoord = TileCoord;
@@ -85,16 +87,6 @@ void ShadeLightSamplesCS(
 				}
 				#endif
 
-				const float3 N = GBuffer.WorldNormal;
-				const float3 V = normalize(View.TranslatedWorldCameraOrigin - TranslatedWorldPosition);
-				const float NoV = saturate(dot(N, V));
-				const float3 SpecularEnv = EnvBRDF(GBuffer.SpecularColor, GBuffer.Roughness, NoV);
-
-				float3 DiffuseLighting = 0.0f;
-				float3 SpecularLighting = 0.0f;
-				float2 DiffuseLuminanceMoments = 0.0f;
-				float2 SpecularLuminanceMoments = 0.0f;
-
 				for (uint IndexInCell = 0; IndexInCell < NumShadingTilesInCell; ++IndexInCell)
 				{
 					const uint PackedShadingTile = ShadingTileGrid[GetShadingTileGridOffset(ShadingTileCoord, IndexInCell)];
@@ -148,24 +140,10 @@ void ShadeLightSamplesCS(
 							TranslatedWorldPosition, CameraVector, GBuffer, AmbientOcclusion, GBuffer.ShadingModelID, 
 							LightData, LightAttenuation, Dither, ScreenCoord, 
 							SurfaceShadow);
-		
-						// Apply pre-exposure
-						SplitLighting.DiffuseLighting.xyz *= View.PreExposure;
-						SplitLighting.SpecularLighting.xyz *= View.PreExposure;
-
-						// #sdl_todo: demodulate earlier during BRDF evaluation
-						SplitLighting.DiffuseLighting.xyz = SafeDemodulate(SplitLighting.DiffuseLighting.xyz, GBuffer.DiffuseColor);
-						SplitLighting.SpecularLighting.xyz = SafeDemodulate(SplitLighting.SpecularLighting.xyz, SpecularEnv);
 
 						DiffuseLighting += SplitLighting.DiffuseLighting.xyz * SampleWeight;
 						SpecularLighting += SplitLighting.SpecularLighting.xyz * SampleWeight;
 
-						DiffuseLuminanceMoments.x += Luminance(SplitLighting.DiffuseLighting.xyz) * SampleWeight;
-						DiffuseLuminanceMoments.y += Pow2(Luminance(SplitLighting.DiffuseLighting.xyz)) * SampleWeight;
-
-						SpecularLuminanceMoments.x += Luminance(SplitLighting.SpecularLighting.xyz) * SampleWeight;
-						SpecularLuminanceMoments.y += Pow2(Luminance(SplitLighting.SpecularLighting.xyz)) * SampleWeight;
-
 						#if DEBUG_MODE 
 						if (bDebug)
 						{
@@ -176,23 +154,34 @@ void ShadeLightSamplesCS(
 					}
 				}
 
-				#if DEBUG_MODE
-				if (bDebug)
-				{
-					Newline(Context);
-					Print(Context, TEXT("Diffuse: ")); 
-					Print(Context, DiffuseLighting);
-					Newline(Context);
-					Print(Context, TEXT("Specular: ")); 
-					Print(Context, SpecularLighting);
-				}
-				#endif
+				// Apply pre-exposure
+				DiffuseLighting *= View.PreExposure;
+				SpecularLighting *= View.PreExposure;
 
-				// #sdl_todo: can first moment be derived from Diffuse/Specular Lighting?
-				RWDiffuseLighting[ScreenCoord] = DiffuseLighting;
-				RWSpecularLighting[ScreenCoord] = SpecularLighting;
-				RWLuminanceMoments[ScreenCoord] = float4(DiffuseLuminanceMoments, SpecularLuminanceMoments);
+				const float3 N = GBuffer.WorldNormal;
+				const float3 V = normalize(View.TranslatedWorldCameraOrigin - TranslatedWorldPosition);
+				const float NoV = saturate(dot(N, V));
+				const float3 SpecularEnv = EnvBRDF(GBuffer.SpecularColor, GBuffer.Roughness, NoV);
+
+				// #sdl_todo: demodulate earlier during BRDF evaluation
+				DiffuseLighting.xyz = SafeDemodulate(DiffuseLighting.xyz, GBuffer.DiffuseColor);
+				SpecularLighting.xyz = SafeDemodulate(SpecularLighting.xyz, SpecularEnv);
+			}
+
+			#if DEBUG_MODE
+			if (bDebug)
+			{
+				Newline(Context);
+				Print(Context, TEXT("Diffuse: ")); 
+				Print(Context, DiffuseLighting);
+				Newline(Context);
+				Print(Context, TEXT("Specular: ")); 
+				Print(Context, SpecularLighting);
 			}
+			#endif
+
+			RWResolvedDiffuseLighting[ScreenCoord] = DiffuseLighting;
+			RWResolvedSpecularLighting[ScreenCoord] = SpecularLighting;
 		}
 	}
 }
\ No newline at end of file
diff --git a/Engine/Shaders/Private/SampledDirectLighting/SampledDirectLightingSpatial.usf b/Engine/Shaders/Private/SampledDirectLighting/SampledDirectLightingSpatial.usf
new file mode 100644
--- /dev/null
+++ b/Engine/Shaders/Private/SampledDirectLighting/SampledDirectLightingSpatial.usf
@@ -0,0 +1,134 @@
+// Copyright Epic Games, Inc. All Rights Reserved.
+
+#include "../Common.ush"
+#include "SampledDirectLighting.ush"
+
+Texture2D<float4> DiffuseLightingAndSecondMomentTexture;
+Texture2D<float4> SpecularLightingAndSecondMomentTexture;
+RWTexture2D<float4> RWSceneColor;
+
+float SpatialFilterDepthWeightScale;
+
+/**
+ * Run a spatial filter in order to filter out noise based on the temporal variance.
+ */
+[numthreads(THREADGROUP_SIZE, THREADGROUP_SIZE, 1)]
+void SampledDirectLightingSpatialCS(
+	uint3 GroupId : SV_GroupID,
+	uint3 GroupThreadId : SV_GroupThreadID,
+	uint3 DispatchThreadId : SV_DispatchThreadID)
+{
+	uint2 ScreenCoord = DispatchThreadId.xy;
+
+#if DEBUG_MODE
+	int2 DebugScreenCoord = View.CursorPosition.x >= 0 ? View.CursorPosition * View.ViewResolutionFraction : -1;
+	bool bDebug = all(ScreenCoord == DebugScreenCoord);
+	FShaderPrintContext Context = InitShaderPrintContext(true, float2(0.55, 0.45));
+#endif
+
+	if (all(ScreenCoord < View.ViewRectMinAndSize.zw))
+	{
+		float2 ScreenUV = (ScreenCoord + 0.5f) * View.BufferSizeAndInvSize.zw;
+		FGBufferData GBuffer = GetGBufferData(ScreenUV);
+
+		const float DeviceZ = SceneDepthTexture[ScreenCoord].x;
+		GBuffer.Depth = ConvertFromDeviceZ(DeviceZ);
+
+		if (GBuffer.Depth > 0.0f && GBuffer.ShadingModelID != SHADINGMODELID_UNLIT)
+		{
+			float4 CenterDiffuseLightingAndSecondMoment = DiffuseLightingAndSecondMomentTexture[ScreenCoord];
+			float4 CenterSpecularLightingAndSecondMoment = SpecularLightingAndSecondMomentTexture[ScreenCoord];
+
+			float CenterDiffuseStdDev = sqrt(max(CenterDiffuseLightingAndSecondMoment.w - Pow2(Luminance(CenterDiffuseLightingAndSecondMoment.xyz)), 0.0f));
+			float CenterSpecularStdDev = sqrt(max(CenterSpecularLightingAndSecondMoment.w - Pow2(Luminance(CenterSpecularLightingAndSecondMoment.xyz)), 0.0f));
+
+			float3 DiffuseLightingSum = CenterDiffuseLightingAndSecondMoment.xyz;
+			float3 SpecularLightingSum = CenterSpecularLightingAndSecondMoment.xyz;
+			float DiffuseLightingWeightSum = 1.0f;
+			float SpecularLightingWeightSum = 1.0f;
+
+			// #sdl_todo: expose as CVars
+			float KernelRadiusInPixels = 16.0f;
+			uint NumSamples = 8;
+
+			float GuassianNormalize = 2.0f / (KernelRadiusInPixels * KernelRadiusInPixels);
+			uint2 RandomSeed = Rand3DPCG16(int3(ScreenCoord, SampledDirectLightingStateFrameIndex)).xy;
+
+			float3 TranslatedWorldPosition = GetTranslatedWorldPositionFromScreenUV(ScreenUV, GBuffer.Depth);
+			float4 ScenePlane = float4(GBuffer.WorldNormal, dot(TranslatedWorldPosition, GBuffer.WorldNormal));
+
+			#if SPATIAL_FILTER
+			if (KernelRadiusInPixels > 0.0f)
+			{
+				for (uint NeighborIndex = 0; NeighborIndex < NumSamples; ++NeighborIndex)
+				{
+					float2 NeighborOffsetInRect = Hammersley16(NeighborIndex, NumSamples, RandomSeed);
+					float2 NeighborOffset = UniformSampleDiskConcentric(NeighborOffsetInRect) * KernelRadiusInPixels;
+					int2 NeighborCoord = (int2)(ScreenCoord + NeighborOffset);
+
+					if (all(and(NeighborCoord >= View.ViewRectMinAndSize.xy, NeighborCoord < (View.ViewRectMinAndSize.xy + View.ViewRectMinAndSize.zw))))
+					{
+						// Depth weight
+						float2 NeighborScreenUV = (NeighborCoord + 0.5f) * View.BufferSizeAndInvSize.zw;
+						FGBufferData NeighborGBuffer = GetGBufferData(NeighborScreenUV);
+						const float3 NeighborTranslatedWorldPosition = GetTranslatedWorldPositionFromScreenUV(NeighborScreenUV, NeighborGBuffer.Depth);
+						float PlaneDistance = abs(dot(float4(NeighborTranslatedWorldPosition, -1), ScenePlane));
+						float RelativeDepthDifference = PlaneDistance / GBuffer.Depth;
+						float DepthWeight = exp2(-SpatialFilterDepthWeightScale * (RelativeDepthDifference * RelativeDepthDifference));
+
+						// #sdl_todo: separate weight for specular accounting for roughness and expose as CVars
+						// Normal weight
+						float AngleBetweenNormals = acosFast(saturate(dot(ScenePlane.xyz, NeighborGBuffer.WorldNormal)));
+						float NormalWeight = 1.0f - saturate(AngleBetweenNormals);
+
+						// Diffuse
+						{
+							float4 NeighborDiffuseLightingAndSecondMoment = DiffuseLightingAndSecondMomentTexture[NeighborCoord];
+
+							float LuminanceDelta = abs(Luminance(CenterDiffuseLightingAndSecondMoment.xyz) - Luminance(NeighborDiffuseLightingAndSecondMoment.xyz));
+							float LuminanceWeight = exp2(-LuminanceDelta / max(1.0f * CenterDiffuseStdDev, 0.00001f));
+
+							const float DiffuseNeighborWeight = DepthWeight * NormalWeight * LuminanceWeight;
+							DiffuseLightingSum += NeighborDiffuseLightingAndSecondMoment.xyz * DiffuseNeighborWeight;
+							DiffuseLightingWeightSum += DiffuseNeighborWeight;
+						}
+
+						// Specular
+						{
+							float4 NeighborSpecularLightingAndSecondMoment = SpecularLightingAndSecondMomentTexture[NeighborCoord];
+
+							float LuminanceDelta = abs(Luminance(CenterSpecularLightingAndSecondMoment.xyz) - Luminance(NeighborSpecularLightingAndSecondMoment.xyz));
+							float LuminanceWeight = exp2(-LuminanceDelta / max(1.0f * CenterSpecularStdDev, 0.00001f));
+
+							const float SpecularNeighborWeight = DepthWeight * NormalWeight * LuminanceWeight;
+							SpecularLightingSum += NeighborSpecularLightingAndSecondMoment.xyz * SpecularNeighborWeight;
+							SpecularLightingWeightSum += SpecularNeighborWeight;
+						}
+					}
+				}
+			}
+			#endif
+
+			const float3 N = GBuffer.WorldNormal;
+			const float3 V = normalize(View.TranslatedWorldCameraOrigin - TranslatedWorldPosition);
+			const float NoV = saturate(dot(N, V));
+
+			// #sdl_todo: handle alpha channel
+			// Final pass outputs composites irradiance and outputs it to scene color
+			float3 Irradiance = GBuffer.DiffuseColor * (DiffuseLightingSum / DiffuseLightingWeightSum);
+			Irradiance += (SpecularLightingSum / SpecularLightingWeightSum) * EnvBRDF(GBuffer.SpecularColor, GBuffer.Roughness, NoV);
+
+			#if DEBUG_MODE
+			if (bDebug)
+			{
+				Print(Context, TEXT("Spatial pass"));
+				Newline(Context);
+				Print(Context, TEXT("Irradiance: "));
+				Print(Context, Irradiance);
+			}
+			#endif
+
+			RWSceneColor[ScreenCoord] = RWSceneColor[ScreenCoord] + float4(Irradiance, 0.0f);
+		}
+	}
+}
\ No newline at end of file
diff --git a/Engine/Shaders/Private/SampledDirectLighting/SampledDirectLightingTemporal.usf b/Engine/Shaders/Private/SampledDirectLighting/SampledDirectLightingTemporal.usf
--- a/Engine/Shaders/Private/SampledDirectLighting/SampledDirectLightingTemporal.usf
+++ b/Engine/Shaders/Private/SampledDirectLighting/SampledDirectLightingTemporal.usf
@@ -5,11 +5,9 @@
 
 Texture2D ResolvedDiffuseLighting;
 Texture2D ResolvedSpecularLighting;
-Texture2D ResolvedLuminanceMoments;
 Texture2D SampledDirectLightingDepthHistory;
-Texture2D DiffuseLightingHistoryTexture;
-Texture2D SpecularLightingHistoryTexture;
-Texture2D LuminanceMomentsHistoryTexture;
+Texture2D DiffuseLightingAndSecondMomentHistoryTexture;
+Texture2D SpecularLightingAndSecondMomentHistoryTexture;
 Texture2D<UNORM float> NumFramesAccumulatedHistoryTexture;
 
 float4 HistoryScreenPositionScaleBias;
@@ -18,9 +16,8 @@ float PrevSceneColorPreExposureCorrection;
 float TemporalMaxFramesAccumulated;
 float TemporalNeighborhoodClampScale;
 
-RWTexture2D<float4> RWDiffuseLighting;
-RWTexture2D<float4> RWSpecularLighting;
-RWTexture2D<float4> RWLuminanceMoments;
+RWTexture2D<float4> RWDiffuseLightingAndSecondMoment;
+RWTexture2D<float4> RWSpecularLightingAndSecondMoment;
 RWTexture2D<UNORM float> RWNumFramesAccumulated;
 RWTexture2D<float> RWSceneDepth;
 RWTexture2D<float4> RWSceneColor;
@@ -41,17 +38,17 @@ FNeighborhoodClamp GetNeighborhoodClamp(
 
 	// #sdl_todo: optimize using groupshared and cross-lane ops
 	const int KernelSize = 1;
-	for (int SampleOffsetY = -KernelSize; SampleOffsetY <= KernelSize; ++SampleOffsetY)
+	for (int NeigborOffsetY = -KernelSize; NeigborOffsetY <= KernelSize; ++NeigborOffsetY)
 	{
-		for (int SampleOffsetX = -KernelSize; SampleOffsetX <= KernelSize; ++SampleOffsetX)
+		for (int NeigborOffsetX = -KernelSize; NeigborOffsetX <= KernelSize; ++NeigborOffsetX)
 		{
-			if (!(SampleOffsetX == 0 && SampleOffsetY == 0))
+			if (!(NeigborOffsetX == 0 && NeigborOffsetY == 0))
 			{
-				uint2 SampleScreenCoord = clamp(ScreenCoord + int2(SampleOffsetX, SampleOffsetY), 0u, View.ViewRectMinAndSize.zw - 1);
-				float3 Sample = SampleTexture[SampleScreenCoord].xyz;
+				uint2 NeigborScreenCoord = clamp(ScreenCoord + int2(NeigborOffsetX, NeigborOffsetY), 0u, View.ViewRectMinAndSize.zw - 1);
+				float3 NeigborSample = SampleTexture[NeigborScreenCoord].xyz;
 				
-				SampleSum += Sample;
-				SampleSqSum += Pow2(Sample);
+				SampleSum += NeigborSample;
+				SampleSqSum += Pow2(NeigborSample);
 			}
 		}
 	}
@@ -90,14 +87,19 @@ void SampledDirectLightingTemporalAccumulationCS(
 		const float DeviceZ = SceneDepthTexture[ScreenCoord].x;
 		GBuffer.Depth = ConvertFromDeviceZ(DeviceZ);
 
+		float3 DiffuseLighting = 0.0f;
+		float3 SpecularLighting = 0.0f;
+		float DiffuseSecondMoment = 0.0f;
+		float SpecularSecondMoment = 0.0f;
+		float NumFramesAccumulated = 0.0f;
+
 		if (GBuffer.Depth > 0.0f && GBuffer.ShadingModelID != SHADINGMODELID_UNLIT)
 		{
-			float3 DiffuseLighting = ResolvedDiffuseLighting[ScreenCoord].xyz;
-			float3 SpecularLighting = ResolvedSpecularLighting[ScreenCoord].xyz;
-			float4 LuminanceMoments = ResolvedLuminanceMoments[ScreenCoord].xyzw;
-			float2 DiffuseLumimanceMoments = LuminanceMoments.xy;
-			float2 SpecularLumimanceMoments = LuminanceMoments.zw;
-			float NumFramesAccumulated = 0.0f;
+			DiffuseLighting = ResolvedDiffuseLighting[ScreenCoord].xyz;
+			SpecularLighting = ResolvedSpecularLighting[ScreenCoord].xyz;
+
+			DiffuseSecondMoment = Pow2(Luminance(DiffuseLighting));
+			SpecularSecondMoment = Pow2(Luminance(SpecularLighting));
 
 			#if VALID_HISTORY
 			{
@@ -148,24 +150,26 @@ void SampledDirectLightingTemporalAccumulationCS(
 					// Use history only if we found at least one valid sample
 					if (any(HistoryWeights > 0.01f))
 					{
-						float3 DiffuseLightingHistory = 0.0f;
-						DiffuseLightingHistory.x = dot(DiffuseLightingHistoryTexture.GatherRed(GlobalPointClampedSampler, HistoryGatherUV).wzxy, HistoryWeights);
-						DiffuseLightingHistory.y = dot(DiffuseLightingHistoryTexture.GatherGreen(GlobalPointClampedSampler, HistoryGatherUV).wzxy, HistoryWeights);
-						DiffuseLightingHistory.z = dot(DiffuseLightingHistoryTexture.GatherBlue(GlobalPointClampedSampler, HistoryGatherUV).wzxy, HistoryWeights);
-						DiffuseLightingHistory = DiffuseLightingHistory / dot(HistoryWeights, 1.0f);
-
-						float3 SpecularLightingHistory = 0.0f;
-						SpecularLightingHistory.x = dot(SpecularLightingHistoryTexture.GatherRed(GlobalPointClampedSampler, HistoryGatherUV).wzxy, HistoryWeights);
-						SpecularLightingHistory.y = dot(SpecularLightingHistoryTexture.GatherGreen(GlobalPointClampedSampler, HistoryGatherUV).wzxy, HistoryWeights);
-						SpecularLightingHistory.z = dot(SpecularLightingHistoryTexture.GatherBlue(GlobalPointClampedSampler, HistoryGatherUV).wzxy, HistoryWeights);
-						SpecularLightingHistory = SpecularLightingHistory / dot(HistoryWeights, 1.0f);
-
-						float4 LuminanceMomentsHistory = 0.0f;
-						LuminanceMomentsHistory.x = dot(LuminanceMomentsHistoryTexture.GatherRed(GlobalPointClampedSampler, HistoryGatherUV).wzxy, HistoryWeights);
-						LuminanceMomentsHistory.y = dot(LuminanceMomentsHistoryTexture.GatherGreen(GlobalPointClampedSampler, HistoryGatherUV).wzxy, HistoryWeights);
-						LuminanceMomentsHistory.z = dot(LuminanceMomentsHistoryTexture.GatherBlue(GlobalPointClampedSampler, HistoryGatherUV).wzxy, HistoryWeights);						
-						LuminanceMomentsHistory.w = dot(LuminanceMomentsHistoryTexture.GatherAlpha(GlobalPointClampedSampler, HistoryGatherUV).wzxy, HistoryWeights);
-						LuminanceMomentsHistory = LuminanceMomentsHistory / dot(HistoryWeights, 1.0f);
+						float4 DiffuseLightingAndSecondMomentHistory = 0.0f;
+						DiffuseLightingAndSecondMomentHistory.x = dot(DiffuseLightingAndSecondMomentHistoryTexture.GatherRed(GlobalPointClampedSampler, HistoryGatherUV).wzxy, HistoryWeights);
+						DiffuseLightingAndSecondMomentHistory.y = dot(DiffuseLightingAndSecondMomentHistoryTexture.GatherGreen(GlobalPointClampedSampler, HistoryGatherUV).wzxy, HistoryWeights);
+						DiffuseLightingAndSecondMomentHistory.z = dot(DiffuseLightingAndSecondMomentHistoryTexture.GatherBlue(GlobalPointClampedSampler, HistoryGatherUV).wzxy, HistoryWeights);
+						DiffuseLightingAndSecondMomentHistory.w = dot(DiffuseLightingAndSecondMomentHistoryTexture.GatherAlpha(GlobalPointClampedSampler, HistoryGatherUV).wzxy, HistoryWeights);
+						DiffuseLightingAndSecondMomentHistory = DiffuseLightingAndSecondMomentHistory / dot(HistoryWeights, 1.0f);
+
+						float4 SpecularLightingAndSecondMomentHistory = 0.0f;
+						SpecularLightingAndSecondMomentHistory.x = dot(SpecularLightingAndSecondMomentHistoryTexture.GatherRed(GlobalPointClampedSampler, HistoryGatherUV).wzxy, HistoryWeights);
+						SpecularLightingAndSecondMomentHistory.y = dot(SpecularLightingAndSecondMomentHistoryTexture.GatherGreen(GlobalPointClampedSampler, HistoryGatherUV).wzxy, HistoryWeights);
+						SpecularLightingAndSecondMomentHistory.z = dot(SpecularLightingAndSecondMomentHistoryTexture.GatherBlue(GlobalPointClampedSampler, HistoryGatherUV).wzxy, HistoryWeights);
+						SpecularLightingAndSecondMomentHistory.w = dot(SpecularLightingAndSecondMomentHistoryTexture.GatherAlpha(GlobalPointClampedSampler, HistoryGatherUV).wzxy, HistoryWeights);
+						SpecularLightingAndSecondMomentHistory = SpecularLightingAndSecondMomentHistory / dot(HistoryWeights, 1.0f);
+
+						// Correct history for current frame exposure
+						DiffuseLightingAndSecondMomentHistory.xyz *= PrevSceneColorPreExposureCorrection;
+						SpecularLightingAndSecondMomentHistory.xyz *= PrevSceneColorPreExposureCorrection;
+
+						DiffuseLightingAndSecondMomentHistory.w *= Pow2(PrevSceneColorPreExposureCorrection);
+						SpecularLightingAndSecondMomentHistory.w *= Pow2(PrevSceneColorPreExposureCorrection);
 
 						// Reproject and rescale normalized NumFramesAccumulated
 						float NumFramesAccumulatedHistory = 0.0f;
@@ -174,15 +178,23 @@ void SampledDirectLightingTemporalAccumulationCS(
 						NumFramesAccumulated = min(NumFramesAccumulatedHistory + 1.0f, TemporalMaxFramesAccumulated);
 
 						float Alpha = 1.0f / (1.0f + NumFramesAccumulated);
-						DiffuseLighting = lerp(DiffuseLightingHistory, DiffuseLighting, Alpha);
-						SpecularLighting = lerp(SpecularLightingHistory, SpecularLighting, Alpha);
-						LuminanceMoments = lerp(LuminanceMomentsHistory, LuminanceMoments, Alpha);
+						DiffuseLighting = lerp(DiffuseLightingAndSecondMomentHistory.xyz, DiffuseLighting, Alpha);
+						SpecularLighting = lerp(SpecularLightingAndSecondMomentHistory.xyz, SpecularLighting, Alpha);
+						DiffuseSecondMoment = lerp(DiffuseLightingAndSecondMomentHistory.w, DiffuseSecondMoment, Alpha);
+						SpecularSecondMoment = lerp(SpecularLightingAndSecondMomentHistory.w, SpecularSecondMoment, Alpha);
 
 						FNeighborhoodClamp DiffuseNeighborhoodClamp = GetNeighborhoodClamp(ScreenCoord, ResolvedDiffuseLighting, DiffuseLighting);
-						DiffuseLighting = clamp(DiffuseLighting, DiffuseNeighborhoodClamp.Min, DiffuseNeighborhoodClamp.Max);
-
 						FNeighborhoodClamp SpecularNeighborhoodClamp = GetNeighborhoodClamp(ScreenCoord, ResolvedSpecularLighting, SpecularLighting);
-						SpecularLighting = clamp(SpecularLighting, SpecularNeighborhoodClamp.Min, SpecularNeighborhoodClamp.Max);
+
+						float3 ClampedDiffuseLighting = clamp(DiffuseLighting, DiffuseNeighborhoodClamp.Min, DiffuseNeighborhoodClamp.Max);
+						float3 ClampedSpecularLighting = clamp(SpecularLighting, SpecularNeighborhoodClamp.Min, SpecularNeighborhoodClamp.Max);
+
+						// Approximate clamp for second moments in order to reduce ghosting from the spatial pass while still allowing accumulated secondary moments to deviate from accumulated lighting
+						DiffuseSecondMoment += Pow2(Luminance(ClampedDiffuseLighting)) - Pow2(Luminance(DiffuseLighting));
+						SpecularSecondMoment += Pow2(Luminance(ClampedSpecularLighting)) - Pow2(Luminance(SpecularLighting));
+
+						DiffuseLighting = ClampedDiffuseLighting;
+						SpecularLighting = ClampedSpecularLighting;
 					}
 				}
 			}
@@ -191,26 +203,24 @@ void SampledDirectLightingTemporalAccumulationCS(
 			DiffuseLighting = MakeFinite(DiffuseLighting);
 			SpecularLighting = MakeFinite(SpecularLighting);
 
-			RWDiffuseLighting[ScreenCoord] = float4(DiffuseLighting, 0.0f);
-			RWSpecularLighting[ScreenCoord] = float4(SpecularLighting, 0.0f);
-			RWLuminanceMoments[ScreenCoord] = float4(DiffuseLumimanceMoments, SpecularLumimanceMoments);
-			RWNumFramesAccumulated[ScreenCoord] = (NumFramesAccumulated + 0.5f) / TemporalMaxFramesAccumulated;
-
 			float3 TranslatedWorldPosition = GetTranslatedWorldPositionFromScreenUV(ScreenUV, GBuffer.Depth);
 			const float3 N = GBuffer.WorldNormal;
 			const float3 V = normalize(View.TranslatedWorldCameraOrigin - TranslatedWorldPosition);
 			const float NoV = saturate(dot(N, V));
 
 			// #sdl_todo: handle alpha channel
-			// Add to scene color and output separately to history
+			// Reverse demodulation and combine lighting components
 			float3 Irradiance = GBuffer.DiffuseColor * DiffuseLighting + SpecularLighting * EnvBRDF(GBuffer.SpecularColor, GBuffer.Roughness, NoV);
 
-			float DiffuseVariance = max(DiffuseLumimanceMoments.y - Pow2(DiffuseLumimanceMoments.x), 0.0f);
-			float SpecularVariance = max(SpecularLumimanceMoments.y - Pow2(SpecularLumimanceMoments.x), 0.0f);
+			float DiffuseVariance = abs(DiffuseSecondMoment - Pow2(Luminance(DiffuseLighting)));
+			float SpecularVariance = max(SpecularSecondMoment - Pow2(Luminance(SpecularLighting)), 0.0f);
 
 			#if DEBUG_MODE
 			if (bDebug)
 			{
+				Newline(Context);
+				Print(Context, TEXT("NumFramesAccumulated: "));
+				Print(Context, NumFramesAccumulated);
 				Newline(Context);
 				Print(Context, TEXT("Irradiance: "));
 				Print(Context, Irradiance);
@@ -222,10 +232,11 @@ void SampledDirectLightingTemporalAccumulationCS(
 				Print(Context, SpecularVariance);
 			}
 			#endif
-
-			RWSceneColor[ScreenCoord] = RWSceneColor[ScreenCoord] + float4(Irradiance, 0.0f);
 		}
 
+		RWDiffuseLightingAndSecondMoment[ScreenCoord] = float4(DiffuseLighting, DiffuseSecondMoment);
+		RWSpecularLightingAndSecondMoment[ScreenCoord] = float4(SpecularLighting, SpecularSecondMoment);
+		RWNumFramesAccumulated[ScreenCoord] = (NumFramesAccumulated + 0.5f) / TemporalMaxFramesAccumulated;
 		RWSceneDepth[ScreenCoord] = GBuffer.Depth;
 	}
 }
\ No newline at end of file
diff --git a/Engine/Shaders/Private/SampledDirectLighting/SampledDirectLightingTracing.usf b/Engine/Shaders/Private/SampledDirectLighting/SampledDirectLightingTracing.usf
--- a/Engine/Shaders/Private/SampledDirectLighting/SampledDirectLightingTracing.usf
+++ b/Engine/Shaders/Private/SampledDirectLighting/SampledDirectLightingTracing.usf
@@ -34,10 +34,8 @@ void InitCompactedTraceTexelIndirectArgsCS(uint3 DispatchThreadId : SV_DispatchT
 	}
 }
 
-uint SampledDirectLightingStateFrameIndes;
 uint2 SampleViewSize;
 uint DebugMode;
-
 float4 HZBUvFactorAndInvFactor;
 float MaxHierarchicalScreenTraceIterations;
 float RelativeDepthThickness;
@@ -45,8 +43,8 @@ float HistoryDepthTestRelativeThickness;
 uint MinimumTracingThreadOccupancy;
 
 Texture2D<uint> LightSamples;
-Texture2D<uint> RWTraceTile;
 
+Texture2D<uint> RWTraceTile;
 RWStructuredBuffer<uint> RWCompactedTraceTexelData;
 RWStructuredBuffer<uint> RWCompactedTraceTexelAllocator;
 
@@ -63,9 +61,6 @@ groupshared uint SharedGlobalTraceTexelStartOffset;
 	groupshared uint SharedTraceTexels[THREADGROUP_SIZE * THREADGROUP_SIZE];
 #endif
 
-// Allocated for max size, where we have only 1 sample per pixel
-groupshared uint SharedTileMask[THREADGROUP_SIZE / TILE_SIZE][THREADGROUP_SIZE / TILE_SIZE];
-
 #if COMPILER_SUPPORTS_WAVE_SIZE && WAVE_OPS
 	WAVESIZE(32)
 #endif
@@ -88,14 +83,6 @@ void CompactLightSampleTracesCS(
 	}
 	#endif
 
-	uint2 LocalTileMaskSize = SampleCoordToDownsampledScreenCoord(THREADGROUP_SIZE) / TILE_SIZE;
-	if (all(GroupThreadId.xy < LocalTileMaskSize))
-	{
-		// Skip unitialized SampleBuffer tiles
-		uint2 DownsampledTileCoord = SampleCoordToDownsampledScreenCoord(GroupId.xy * THREADGROUP_SIZE) / TILE_SIZE + GroupThreadId.xy;
-		SharedTileMask[GroupThreadId.x][GroupThreadId.y] = DownsampledTileMask[DownsampledTileCoord];
-	}
-
 	GroupMemoryBarrierWithGroupSync();
 
 	uint2 SampleCoord = DispatchThreadId.xy;
@@ -103,7 +90,7 @@ void CompactLightSampleTracesCS(
 	uint TraceTexelForThisThread = 0;
 
 	uint2 LocalTileMaskCoord = SampleCoordToDownsampledScreenCoord(SampleCoord - GroupId.xy * THREADGROUP_SIZE) / TILE_SIZE;
-	if (all(SampleCoord < SampleViewSize) && SharedTileMask[LocalTileMaskCoord.x][LocalTileMaskCoord.y] != 0)
+	if (all(SampleCoord < SampleViewSize))
 	{
 		// #sdl_todo: try to preserve z-order here for improved tracing coherency
 		FLightSample LightSample = UnpackLightSample(LightSamples[SampleCoord]);
diff --git a/Engine/Source/Runtime/Renderer/Private/SampledDirectLighting/SampledDirectLighting.cpp b/Engine/Source/Runtime/Renderer/Private/SampledDirectLighting/SampledDirectLighting.cpp
--- a/Engine/Source/Runtime/Renderer/Private/SampledDirectLighting/SampledDirectLighting.cpp
+++ b/Engine/Source/Runtime/Renderer/Private/SampledDirectLighting/SampledDirectLighting.cpp
@@ -34,6 +34,13 @@ static TAutoConsoleVariable<int32> CVarSampledDirectLightingMaxShadingTilesPerGr
 	ECVF_Scalability | ECVF_RenderThreadSafe
 );
 
+static TAutoConsoleVariable<float> CVarSampledDirectLightingSamplingMinWeight(
+	TEXT("r.SampledDirectLighting.Sampling.MinWeight"),
+	0.002f,
+	TEXT("Determines minimal sample influence on final pixels. Used to skip samples which would have minimal impact to the final image even if light is fully visible."),
+	ECVF_Scalability | ECVF_RenderThreadSafe
+);
+
 static TAutoConsoleVariable<int32> CVarSampledDirectLightingTemporal(
 	TEXT("r.SampledDirectLighting.Temporal"),
 	1,
@@ -55,17 +62,17 @@ static TAutoConsoleVariable<float> CVarSampledDirectLightingTemporalNeighborhood
 	ECVF_Scalability | ECVF_RenderThreadSafe
 );
 
-static TAutoConsoleVariable<float> CVarSampledDirectLightingSamplingMinWeight(
-	TEXT("r.SampledDirectLighting.Sampling.MinWeight"),
-	0.002f,
-	TEXT("Determines minimal sample influence on final pixels. Used to skip samples which would have minimal impact to the final image even if light is fully visible."),
+static TAutoConsoleVariable<int32> CVarSampledDirectLightingSpatial(
+	TEXT("r.SampledDirectLighting.Spatial"),
+	1,
+	TEXT("Whether denoiser should run spatial filter."),
 	ECVF_Scalability | ECVF_RenderThreadSafe
 );
 
-static TAutoConsoleVariable<int32> CVarSampledDirectLightingSpatial(
-	TEXT("r.SampledDirectLighting.Spatial"),
-	0,
-	TEXT("Whether to run spatial shadow mask denoising pass."),
+static TAutoConsoleVariable<float> CVarSampledDirectLightingSpatialDepthWeightScale(
+	TEXT("r.SampledDirectLighting.Spatial.DepthWeightScale"),
+	10000.0f,
+	TEXT("Scales the depth weight of the spatial filter. Smaller values allow for more sample reuse, but also introduce more bluriness between unrelated surfaces."),
 	ECVF_Scalability | ECVF_RenderThreadSafe
 );
 
@@ -212,7 +219,10 @@ namespace SampledDirectLighting
 	{
 		SimpleShading = 0,
 		ComplexShading = 1,
-		MAX = 2
+		SHADING_MAX = 2,
+
+		Empty = 2,
+		MAX = 3
 	};
 }
 
@@ -225,7 +235,6 @@ class FTileClassificationCS : public FGlobalShader
 		SHADER_PARAMETER_STRUCT_INCLUDE(FSampledDirectLightingParameters, SampledDirectLightingParameters)
 		SHADER_PARAMETER_RDG_BUFFER_UAV(RWStructuredBuffer<uint>, RWTileAllocator)
 		SHADER_PARAMETER_RDG_BUFFER_UAV(RWStructuredBuffer<uint>, RWTileData)
-		SHADER_PARAMETER_RDG_TEXTURE_UAV(RWTexture2D<uint>, RWDownsampledTileMask)
 	END_SHADER_PARAMETER_STRUCT()
 
 	class FDownsampledClassification : SHADER_PERMUTATION_BOOL("DOWNSAMPLED_CLASSIFICATION");
@@ -331,7 +340,7 @@ class FGenerateSamplesCS : public FGlobalShader
 		SHADER_PARAMETER(FVector4f, HistoryScreenPositionScaleBias)
 	END_SHADER_PARAMETER_STRUCT()
 
-	class FTileType : SHADER_PERMUTATION_INT("TILE_TYPE", (int32)SampledDirectLighting::ETileType::MAX);
+	class FTileType : SHADER_PERMUTATION_INT("TILE_TYPE", (int32)SampledDirectLighting::ETileType::SHADING_MAX);
 	class FIESProfile : SHADER_PERMUTATION_BOOL("USE_IES_PROFILE");
 	class FLightFunctionAtlas : SHADER_PERMUTATION_BOOL("USE_LIGHT_FUNCTION_ATLAS");
 	class FTexturedRectLights : SHADER_PERMUTATION_BOOL("USE_SOURCE_TEXTURE");
@@ -360,6 +369,44 @@ class FGenerateSamplesCS : public FGlobalShader
 
 IMPLEMENT_GLOBAL_SHADER(FGenerateSamplesCS, "/Engine/Private/SampledDirectLighting/SampledDirectLightingSampling.usf", "GenerateSamplesCS", SF_Compute);
 
+class FClearLightSamplesCS : public FGlobalShader
+{
+	DECLARE_GLOBAL_SHADER(FClearLightSamplesCS)
+	SHADER_USE_PARAMETER_STRUCT(FClearLightSamplesCS, FGlobalShader)
+
+	BEGIN_SHADER_PARAMETER_STRUCT(FParameters, )
+		RDG_BUFFER_ACCESS(IndirectArgs, ERHIAccess::IndirectArgs)
+		SHADER_PARAMETER_STRUCT_INCLUDE(FSampledDirectLightingParameters, SampledDirectLightingParameters)
+		SHADER_PARAMETER_RDG_TEXTURE_UAV(RWTexture2D<float>, RWDownsampledSceneDepth)
+		SHADER_PARAMETER_RDG_TEXTURE_UAV(RWTexture2D<UNORM float3>, RWDownsampledSceneWorldNormal)
+		SHADER_PARAMETER_RDG_TEXTURE_UAV(RWTexture2D<uint>, RWLightSamples)
+		SHADER_PARAMETER_RDG_BUFFER_SRV(StructuredBuffer<uint>, DownsampledTileAllocator)
+		SHADER_PARAMETER_RDG_BUFFER_SRV(StructuredBuffer<uint>, DownsampledTileData)
+	END_SHADER_PARAMETER_STRUCT()
+
+	class FDebugMode : SHADER_PERMUTATION_BOOL("DEBUG_MODE");
+	using FPermutationDomain = TShaderPermutationDomain<FDebugMode>;
+
+	static int32 GetGroupSize()
+	{
+		return 8;
+	}
+
+	static bool ShouldCompilePermutation(const FGlobalShaderPermutationParameters& Parameters)
+	{
+		return SampledDirectLighting::ShouldCompileShaders(Parameters);
+	}
+
+	static void ModifyCompilationEnvironment(const FGlobalShaderPermutationParameters& Parameters, FShaderCompilerEnvironment& OutEnvironment)
+	{
+		FGlobalShader::ModifyCompilationEnvironment(Parameters, OutEnvironment);
+		SampledDirectLighting::ModifyCompilationEnvironment(Parameters.Platform, OutEnvironment);
+		OutEnvironment.SetDefine(TEXT("THREADGROUP_SIZE"), GetGroupSize());
+	}
+};
+
+IMPLEMENT_GLOBAL_SHADER(FClearLightSamplesCS, "/Engine/Private/SampledDirectLighting/SampledDirectLightingSampling.usf", "ClearLightSamplesCS", SF_Compute);
+
 class FInitCompositeUpsampleWeightsCS : public FGlobalShader
 {
 	DECLARE_GLOBAL_SHADER(FInitCompositeUpsampleWeightsCS)
@@ -437,9 +484,8 @@ class FShadeLightSamplesCS : public FGlobalShader
 	BEGIN_SHADER_PARAMETER_STRUCT(FParameters, )
 		RDG_BUFFER_ACCESS(IndirectArgs, ERHIAccess::IndirectArgs)
 		SHADER_PARAMETER_STRUCT_INCLUDE(FSampledDirectLightingParameters, SampledDirectLightingParameters)
-		SHADER_PARAMETER_RDG_TEXTURE_UAV(RWTexture2D<float3>, RWDiffuseLighting)
-		SHADER_PARAMETER_RDG_TEXTURE_UAV(RWTexture2D<float3>, RWSpecularLighting)
-		SHADER_PARAMETER_RDG_TEXTURE_UAV(RWTexture2D<float4>, RWLuminanceMoments)
+		SHADER_PARAMETER_RDG_TEXTURE_UAV(RWTexture2D<float3>, RWResolvedDiffuseLighting)
+		SHADER_PARAMETER_RDG_TEXTURE_UAV(RWTexture2D<float3>, RWResolvedSpecularLighting)
 		SHADER_PARAMETER_RDG_BUFFER_SRV(StructuredBuffer<uint>, TileAllocator)
 		SHADER_PARAMETER_RDG_BUFFER_SRV(StructuredBuffer<uint>, TileData)
 		SHADER_PARAMETER_RDG_BUFFER_SRV(StructuredBuffer<uint>, ShadingTileAllocator)
@@ -455,7 +501,7 @@ class FShadeLightSamplesCS : public FGlobalShader
 		return 8;
 	}
 
-	class FTileType : SHADER_PERMUTATION_INT("TILE_TYPE", (int32)SampledDirectLighting::ETileType::MAX);
+	class FTileType : SHADER_PERMUTATION_INT("TILE_TYPE", (int32)SampledDirectLighting::ETileType::SHADING_MAX);
 	class FIESProfile : SHADER_PERMUTATION_BOOL("USE_IES_PROFILE");
 	class FLightFunctionAtlas : SHADER_PERMUTATION_BOOL("USE_LIGHT_FUNCTION_ATLAS");
 	class FTexturedRectLights : SHADER_PERMUTATION_BOOL("USE_SOURCE_TEXTURE");
@@ -484,20 +530,17 @@ class FSampledDirectLightingTemporalAccumulationCS : public FGlobalShader
 
 	BEGIN_SHADER_PARAMETER_STRUCT(FParameters, )
 		SHADER_PARAMETER_STRUCT_INCLUDE(FSampledDirectLightingParameters, SampledDirectLightingParameters)
-		SHADER_PARAMETER_RDG_TEXTURE(Texture2D<float3>, ResolvedDiffuseLighting)
-		SHADER_PARAMETER_RDG_TEXTURE(Texture2D<float3>, ResolvedSpecularLighting)
-		SHADER_PARAMETER_RDG_TEXTURE(Texture2D<float4>, ResolvedLuminanceMoments)
-		SHADER_PARAMETER_RDG_TEXTURE(Texture2D<float3>, DiffuseLightingHistoryTexture)
-		SHADER_PARAMETER_RDG_TEXTURE(Texture2D<float3>, SpecularLightingHistoryTexture)
-		SHADER_PARAMETER_RDG_TEXTURE(Texture2D<float4>, LuminanceMomentsHistoryTexture)
+		SHADER_PARAMETER_RDG_TEXTURE(Texture2D<float4>, ResolvedDiffuseLighting)
+		SHADER_PARAMETER_RDG_TEXTURE(Texture2D<float4>, ResolvedSpecularLighting)
+		SHADER_PARAMETER_RDG_TEXTURE(Texture2D<float4>, DiffuseLightingAndSecondMomentHistoryTexture)
+		SHADER_PARAMETER_RDG_TEXTURE(Texture2D<float4>, SpecularLightingAndSecondMomentHistoryTexture)
 		SHADER_PARAMETER_RDG_TEXTURE(Texture2D<UNORM float>, NumFramesAccumulatedHistoryTexture)
 		SHADER_PARAMETER_RDG_TEXTURE(Texture2D<float>, SampledDirectLightingDepthHistory)
 		SHADER_PARAMETER(FVector4f, HistoryUVMinMax)
 		SHADER_PARAMETER(FVector4f, HistoryScreenPositionScaleBias)
 		SHADER_PARAMETER(float, PrevSceneColorPreExposureCorrection)
-		SHADER_PARAMETER_RDG_TEXTURE_UAV(RWTexture2D<float3>, RWDiffuseLighting)
-		SHADER_PARAMETER_RDG_TEXTURE_UAV(RWTexture2D<float3>, RWSpecularLighting)
-		SHADER_PARAMETER_RDG_TEXTURE_UAV(RWTexture2D<float3>, RWLuminanceMoments)
+		SHADER_PARAMETER_RDG_TEXTURE_UAV(RWTexture2D<float4>, RWDiffuseLightingAndSecondMoment)
+		SHADER_PARAMETER_RDG_TEXTURE_UAV(RWTexture2D<float4>, RWSpecularLightingAndSecondMoment)
 		SHADER_PARAMETER_RDG_TEXTURE_UAV(RWTexture2D<UNORM float>, RWNumFramesAccumulated)
 		SHADER_PARAMETER_RDG_TEXTURE_UAV(RWTexture2D<float>, RWSceneDepth)
 		SHADER_PARAMETER_RDG_TEXTURE_UAV(RWTexture2D<float4>, RWSceneColor)
@@ -526,6 +569,42 @@ class FSampledDirectLightingTemporalAccumulationCS : public FGlobalShader
 
 IMPLEMENT_GLOBAL_SHADER(FSampledDirectLightingTemporalAccumulationCS, "/Engine/Private/SampledDirectLighting/SampledDirectLightingTemporal.usf", "SampledDirectLightingTemporalAccumulationCS", SF_Compute);
 
+class FSampledDirectLightingSpatialCS : public FGlobalShader
+{
+	DECLARE_GLOBAL_SHADER(FSampledDirectLightingSpatialCS)
+	SHADER_USE_PARAMETER_STRUCT(FSampledDirectLightingSpatialCS, FGlobalShader)
+
+	BEGIN_SHADER_PARAMETER_STRUCT(FParameters, )
+		SHADER_PARAMETER_STRUCT_INCLUDE(FSampledDirectLightingParameters, SampledDirectLightingParameters)
+		SHADER_PARAMETER_RDG_TEXTURE_UAV(RWTexture2D<float4>, RWSceneColor)
+		SHADER_PARAMETER_RDG_TEXTURE(Texture2D<float3>, DiffuseLightingAndSecondMomentTexture)
+		SHADER_PARAMETER_RDG_TEXTURE(Texture2D<float3>, SpecularLightingAndSecondMomentTexture)
+		SHADER_PARAMETER(float, SpatialFilterDepthWeightScale)
+	END_SHADER_PARAMETER_STRUCT()
+
+	class FSpatialFilter : SHADER_PERMUTATION_BOOL("SPATIAL_FILTER");
+	class FDebugMode : SHADER_PERMUTATION_BOOL("DEBUG_MODE");
+	using FPermutationDomain = TShaderPermutationDomain<FSpatialFilter, FDebugMode>;
+
+	static int32 GetGroupSize()
+	{
+		return 8;
+	}
+
+	static bool ShouldCompilePermutation(const FGlobalShaderPermutationParameters& Parameters)
+	{
+		return SampledDirectLighting::ShouldCompileShaders(Parameters);
+	}
+
+	static void ModifyCompilationEnvironment(const FGlobalShaderPermutationParameters& Parameters, FShaderCompilerEnvironment& OutEnvironment)
+	{
+		FGlobalShader::ModifyCompilationEnvironment(Parameters, OutEnvironment);
+		OutEnvironment.SetDefine(TEXT("THREADGROUP_SIZE"), GetGroupSize());
+	}
+};
+
+IMPLEMENT_GLOBAL_SHADER(FSampledDirectLightingSpatialCS, "/Engine/Private/SampledDirectLighting/SampledDirectLightingSpatial.usf", "SampledDirectLightingSpatialCS", SF_Compute);
+
 /**
  * Single pass batched light rendering using ray tracing (distance field or triangle) for shadowing.
  */
@@ -602,9 +681,8 @@ void FDeferredShadingSceneRenderer::RenderSampledDirectLighting(FRDGBuilder& Gra
 	bool bTemporal = CVarSampledDirectLightingTemporal.GetValueOnRenderThread() != 0;
 	FVector4f HistoryScreenPositionScaleBias = FVector4f(0.0f, 0.0f, 0.0f, 0.0f);
 	FVector4f HistoryUVMinMax = FVector4f(0.0f, 0.0f, 0.0f, 0.0f);
-	FRDGTextureRef DiffuseLightingHistory = nullptr;
-	FRDGTextureRef SpecularLightingHistory = nullptr;
-	FRDGTextureRef LuminanceMomentsHistory = nullptr;
+	FRDGTextureRef DiffuseLightingAndSecondMomentHistory = nullptr;
+	FRDGTextureRef SpecularLightingAndSecondMomentHistory = nullptr;
 	FRDGTextureRef SceneDepthHistory = nullptr;
 	FRDGTextureRef NumFramesAccumulatedHistory = nullptr;
 
@@ -617,18 +695,16 @@ void FDeferredShadingSceneRenderer::RenderSampledDirectLighting(FRDGBuilder& Gra
 			HistoryScreenPositionScaleBias = LightingViewState.HistoryScreenPositionScaleBias;
 			HistoryUVMinMax = LightingViewState.HistoryUVMinMax;
 
-			if (LightingViewState.DiffuseLightingHistory
-				&& LightingViewState.SpecularLightingHistory
-				&& LightingViewState.LuminanceMomentsHistory
+			if (LightingViewState.DiffuseLightingAndSecondMomentHistory
+				&& LightingViewState.SpecularLightingAndSecondMomentHistory
 				&& LightingViewState.SceneDepthHistory
 				&& LightingViewState.NumFramesAccumulatedHistory
-				&& LightingViewState.DiffuseLightingHistory->GetDesc().Extent == View.GetSceneTexturesConfig().Extent
-				&& LightingViewState.SpecularLightingHistory->GetDesc().Extent == View.GetSceneTexturesConfig().Extent
+				&& LightingViewState.DiffuseLightingAndSecondMomentHistory->GetDesc().Extent == View.GetSceneTexturesConfig().Extent
+				&& LightingViewState.SpecularLightingAndSecondMomentHistory->GetDesc().Extent == View.GetSceneTexturesConfig().Extent
 				&& LightingViewState.SceneDepthHistory->GetDesc().Extent == SceneTextures.Depth.Resolve->Desc.Extent)
 			{
-				DiffuseLightingHistory = GraphBuilder.RegisterExternalTexture(LightingViewState.DiffuseLightingHistory);
-				SpecularLightingHistory = GraphBuilder.RegisterExternalTexture(LightingViewState.SpecularLightingHistory);
-				LuminanceMomentsHistory = GraphBuilder.RegisterExternalTexture(LightingViewState.LuminanceMomentsHistory);
+				DiffuseLightingAndSecondMomentHistory = GraphBuilder.RegisterExternalTexture(LightingViewState.DiffuseLightingAndSecondMomentHistory);
+				SpecularLightingAndSecondMomentHistory = GraphBuilder.RegisterExternalTexture(LightingViewState.SpecularLightingAndSecondMomentHistory);
 				SceneDepthHistory = GraphBuilder.RegisterExternalTexture(LightingViewState.SceneDepthHistory);
 				NumFramesAccumulatedHistory = GraphBuilder.RegisterExternalTexture(LightingViewState.NumFramesAccumulatedHistory);
 			}
@@ -723,7 +799,6 @@ void FDeferredShadingSceneRenderer::RenderSampledDirectLighting(FRDGBuilder& Gra
 			PassParameters->SampledDirectLightingParameters = SampledDirectLightingParameters;
 			PassParameters->RWTileAllocator = GraphBuilder.CreateUAV(TileAllocator);
 			PassParameters->RWTileData = GraphBuilder.CreateUAV(TileData);
-			PassParameters->RWDownsampledTileMask = nullptr;
 
 			FTileClassificationCS::FPermutationDomain PermutationVector;
 			PermutationVector.Set<FTileClassificationCS::FDownsampledClassification>(false);
@@ -744,7 +819,6 @@ void FDeferredShadingSceneRenderer::RenderSampledDirectLighting(FRDGBuilder& Gra
 			PassParameters->SampledDirectLightingParameters = SampledDirectLightingParameters;
 			PassParameters->RWTileAllocator = GraphBuilder.CreateUAV(DownsampledTileAllocator);
 			PassParameters->RWTileData = GraphBuilder.CreateUAV(DownsampledTileData);
-			PassParameters->RWDownsampledTileMask = GraphBuilder.CreateUAV(DownsampledTileMask);
 
 			FTileClassificationCS::FPermutationDomain PermutationVector;
 			PermutationVector.Set<FTileClassificationCS::FDownsampledClassification>(true);
@@ -795,7 +869,31 @@ void FDeferredShadingSceneRenderer::RenderSampledDirectLighting(FRDGBuilder& Gra
 		FRDGBufferUAVRef CompositeTileDataUAV = GraphBuilder.CreateUAV(CompositeTileData, ERDGUnorderedAccessViewFlags::SkipBarrier);
 		FRDGTextureUAVRef LightSamplesUAV = GraphBuilder.CreateUAV(LightSamples, ERDGUnorderedAccessViewFlags::SkipBarrier);
 
-		for (int32 TileType = 0; TileType < (int32)SampledDirectLighting::ETileType::MAX; ++TileType)
+		// Clear tiles which don't contain any lights or geometry
+		{
+			FClearLightSamplesCS::FParameters* PassParameters = GraphBuilder.AllocParameters<FClearLightSamplesCS::FParameters>();
+			PassParameters->IndirectArgs = DownsampledTileIndirectArgs;
+			PassParameters->SampledDirectLightingParameters = SampledDirectLightingParameters;
+			PassParameters->RWDownsampledSceneDepth = DownsampledSceneDepthUAV;
+			PassParameters->RWDownsampledSceneWorldNormal = DownsampledSceneWorldNormalUAV;
+			PassParameters->RWLightSamples = LightSamplesUAV;
+			PassParameters->DownsampledTileAllocator = GraphBuilder.CreateSRV(DownsampledTileAllocator);
+			PassParameters->DownsampledTileData = GraphBuilder.CreateSRV(DownsampledTileData);
+
+			FClearLightSamplesCS::FPermutationDomain PermutationVector;
+			PermutationVector.Set<FClearLightSamplesCS::FDebugMode>(bDebug);
+			auto ComputeShader = View.ShaderMap->GetShader<FClearLightSamplesCS>(PermutationVector);
+
+			FComputeShaderUtils::AddPass(
+				GraphBuilder,
+				RDG_EVENT_NAME("ClearLightSamples"),
+				ComputeShader,
+				PassParameters,
+				DownsampledTileIndirectArgs,
+				(int32)SampledDirectLighting::ETileType::Empty * sizeof(FRHIDispatchIndirectParameters));
+		}
+
+		for (int32 TileType = 0; TileType < (int32)SampledDirectLighting::ETileType::SHADING_MAX; ++TileType)
 		{
 			FGenerateSamplesCS::FParameters* PassParameters = GraphBuilder.AllocParameters<FGenerateSamplesCS::FParameters>();
 			PassParameters->IndirectArgs = DownsampledTileIndirectArgs;
@@ -934,20 +1032,16 @@ void FDeferredShadingSceneRenderer::RenderSampledDirectLighting(FRDGBuilder& Gra
 		FRDGTextureDesc::Create2D(View.GetSceneTexturesConfig().Extent, PF_FloatRGB, FClearValueBinding::Black, TexCreate_ShaderResource | TexCreate_UAV),
 		TEXT("SampledDirectLighting.ResolvedSpecularLighting"));
 
-	FRDGTextureRef ResolvedLuminanceMoments = GraphBuilder.CreateTexture(
-		FRDGTextureDesc::Create2D(View.GetSceneTexturesConfig().Extent, PF_FloatRGBA, FClearValueBinding::Black, TexCreate_ShaderResource | TexCreate_UAV),
-		TEXT("SampledDirectLighting.ResolvedLuminanceMoments"));
-
 	// Shade light samples
 	{
-		FRDGTextureUAVRef SceneColorUAV = GraphBuilder.CreateUAV(SceneTextures.Color.Target, ERDGUnorderedAccessViewFlags::SkipBarrier);
+		FRDGTextureUAVRef ResolvedDiffuseLightingUAV = GraphBuilder.CreateUAV(ResolvedDiffuseLighting, ERDGUnorderedAccessViewFlags::SkipBarrier);
+		FRDGTextureUAVRef ResolvedSpecularLightingUAV = GraphBuilder.CreateUAV(ResolvedSpecularLighting, ERDGUnorderedAccessViewFlags::SkipBarrier);
 
-		for (int32 TileType = 0; TileType < (int32)SampledDirectLighting::ETileType::MAX; ++TileType)
+		for (int32 TileType = 0; TileType < (int32)SampledDirectLighting::ETileType::SHADING_MAX; ++TileType)
 		{
 			FShadeLightSamplesCS::FParameters* PassParameters = GraphBuilder.AllocParameters<FShadeLightSamplesCS::FParameters>();
-			PassParameters->RWDiffuseLighting = GraphBuilder.CreateUAV(ResolvedDiffuseLighting);
-			PassParameters->RWSpecularLighting = GraphBuilder.CreateUAV(ResolvedSpecularLighting);
-			PassParameters->RWLuminanceMoments = GraphBuilder.CreateUAV(ResolvedLuminanceMoments);
+			PassParameters->RWResolvedDiffuseLighting = ResolvedDiffuseLightingUAV;
+			PassParameters->RWResolvedSpecularLighting = ResolvedSpecularLightingUAV;
 			PassParameters->IndirectArgs = TileIndirectArgs;
 			PassParameters->SampledDirectLightingParameters = SampledDirectLightingParameters;
 			PassParameters->CompositeTileAllocator = GraphBuilder.CreateSRV(CompositeTileAllocator);
@@ -976,18 +1070,15 @@ void FDeferredShadingSceneRenderer::RenderSampledDirectLighting(FRDGBuilder& Gra
 		}
 	}
 
-	// Final diffuse and specular which will be passed to the next frame
-	FRDGTextureRef DiffuseLighting = GraphBuilder.CreateTexture(
-		FRDGTextureDesc::Create2D(View.GetSceneTexturesConfig().Extent, PF_FloatRGBA, FClearValueBinding::Black, TexCreate_ShaderResource | TexCreate_UAV),
-		TEXT("SampledDirectLighting.DiffuseLighting"));
-
-	FRDGTextureRef SpecularLighting = GraphBuilder.CreateTexture(
+	// Demodulated lighting components with second luminance moments stored in alpha channel for temporal variance tracking
+	// This will be passed to the next frame
+	FRDGTextureRef DiffuseLightingAndSecondMoment = GraphBuilder.CreateTexture(
 		FRDGTextureDesc::Create2D(View.GetSceneTexturesConfig().Extent, PF_FloatRGBA, FClearValueBinding::Black, TexCreate_ShaderResource | TexCreate_UAV),
-		TEXT("SampledDirectLighting.SpecularLighting"));
+		TEXT("SampledDirectLighting.DiffuseLightingAndSecondMoment"));
 
-	FRDGTextureRef LuminanceMoments = GraphBuilder.CreateTexture(
+	FRDGTextureRef SpecularLightingAndSecondMoment = GraphBuilder.CreateTexture(
 		FRDGTextureDesc::Create2D(View.GetSceneTexturesConfig().Extent, PF_FloatRGBA, FClearValueBinding::Black, TexCreate_ShaderResource | TexCreate_UAV),
-		TEXT("SampledDirectLighting.LuminanceMoments"));
+		TEXT("SampledDirectLighting.SpecularLightingAndSecondMoment"));
 
 	FRDGTextureRef SceneDepthCopy = GraphBuilder.CreateTexture(
 		FRDGTextureDesc::Create2D(SceneTextures.Depth.Resolve->Desc.Extent, PF_R32_FLOAT, FClearValueBinding::Black, TexCreate_ShaderResource | TexCreate_UAV),
@@ -1003,24 +1094,21 @@ void FDeferredShadingSceneRenderer::RenderSampledDirectLighting(FRDGBuilder& Gra
 		PassParameters->SampledDirectLightingParameters = SampledDirectLightingParameters;
 		PassParameters->ResolvedDiffuseLighting = ResolvedDiffuseLighting;
 		PassParameters->ResolvedSpecularLighting = ResolvedSpecularLighting;
-		PassParameters->ResolvedLuminanceMoments = ResolvedLuminanceMoments;
-		PassParameters->DiffuseLightingHistoryTexture = DiffuseLightingHistory;
-		PassParameters->SpecularLightingHistoryTexture = SpecularLightingHistory;
-		PassParameters->LuminanceMomentsHistoryTexture = LuminanceMomentsHistory;
+		PassParameters->DiffuseLightingAndSecondMomentHistoryTexture = DiffuseLightingAndSecondMomentHistory;
+		PassParameters->SpecularLightingAndSecondMomentHistoryTexture = SpecularLightingAndSecondMomentHistory;
 		PassParameters->NumFramesAccumulatedHistoryTexture = NumFramesAccumulatedHistory;
 		PassParameters->SampledDirectLightingDepthHistory = SceneDepthHistory;
 		PassParameters->PrevSceneColorPreExposureCorrection = View.PreExposure / View.PrevViewInfo.SceneColorPreExposure;
 		PassParameters->HistoryScreenPositionScaleBias = HistoryScreenPositionScaleBias;
 		PassParameters->HistoryUVMinMax = HistoryUVMinMax;
-		PassParameters->RWDiffuseLighting = GraphBuilder.CreateUAV(DiffuseLighting);
-		PassParameters->RWSpecularLighting = GraphBuilder.CreateUAV(SpecularLighting);
-		PassParameters->RWLuminanceMoments = GraphBuilder.CreateUAV(LuminanceMoments);
+		PassParameters->RWDiffuseLightingAndSecondMoment = GraphBuilder.CreateUAV(DiffuseLightingAndSecondMoment);
+		PassParameters->RWSpecularLightingAndSecondMoment = GraphBuilder.CreateUAV(SpecularLightingAndSecondMoment);
 		PassParameters->RWNumFramesAccumulated = GraphBuilder.CreateUAV(NumFramesAccumulated);
 		PassParameters->RWSceneDepth = GraphBuilder.CreateUAV(SceneDepthCopy);
 		PassParameters->RWSceneColor = GraphBuilder.CreateUAV(SceneTextures.Color.Target);
 
 		FSampledDirectLightingTemporalAccumulationCS::FPermutationDomain PermutationVector;
-		PermutationVector.Set<FSampledDirectLightingTemporalAccumulationCS::FValidHistory>(DiffuseLightingHistory != nullptr && bTemporal);
+		PermutationVector.Set<FSampledDirectLightingTemporalAccumulationCS::FValidHistory>(DiffuseLightingAndSecondMomentHistory != nullptr && bTemporal);
 		PermutationVector.Set<FSampledDirectLightingTemporalAccumulationCS::FDebugMode>(bDebug);
 		auto ComputeShader = View.ShaderMap->GetShader<FSampledDirectLightingTemporalAccumulationCS>(PermutationVector);
 
@@ -1034,6 +1122,30 @@ void FDeferredShadingSceneRenderer::RenderSampledDirectLighting(FRDGBuilder& Gra
 			GroupCount);
 	}
 
+	// Spatial filter
+	{
+		FSampledDirectLightingSpatialCS::FParameters* PassParameters = GraphBuilder.AllocParameters<FSampledDirectLightingSpatialCS::FParameters>();
+		PassParameters->SampledDirectLightingParameters = SampledDirectLightingParameters;
+		PassParameters->RWSceneColor = GraphBuilder.CreateUAV(SceneTextures.Color.Target);
+		PassParameters->DiffuseLightingAndSecondMomentTexture = DiffuseLightingAndSecondMoment;
+		PassParameters->SpecularLightingAndSecondMomentTexture = SpecularLightingAndSecondMoment;
+		PassParameters->SpatialFilterDepthWeightScale = CVarSampledDirectLightingSpatialDepthWeightScale.GetValueOnRenderThread();
+
+		FSampledDirectLightingSpatialCS::FPermutationDomain PermutationVector;
+		PermutationVector.Set<FSampledDirectLightingSpatialCS::FSpatialFilter>(CVarSampledDirectLightingSpatial.GetValueOnRenderThread() != 0);
+		PermutationVector.Set<FSampledDirectLightingSpatialCS::FDebugMode>(bDebug);
+		auto ComputeShader = View.ShaderMap->GetShader<FSampledDirectLightingSpatialCS>(PermutationVector);
+
+		const FIntVector GroupCount = FComputeShaderUtils::GetGroupCount(View.ViewRect.Size(), FSampledDirectLightingSpatialCS::GetGroupSize());
+
+		FComputeShaderUtils::AddPass(
+			GraphBuilder,
+			RDG_EVENT_NAME("Spatial"),
+			ComputeShader,
+			PassParameters,
+			GroupCount);
+	}
+
 	if (View.ViewState && !View.bStatePrevViewInfoIsReadOnly)
 	{
 		FSampledDirectLightingViewState& LightingViewState = View.ViewState->SampledDirectLighting;
@@ -1048,19 +1160,17 @@ void FDeferredShadingSceneRenderer::RenderSampledDirectLighting(FRDGBuilder& Gra
 			(View.ViewRect.Max.X - 1.0f) * InvBufferSize.X,
 			(View.ViewRect.Max.Y - 1.0f) * InvBufferSize.Y);
 
-		if (DiffuseLighting && SpecularLighting && LuminanceMoments && SceneDepthCopy && NumFramesAccumulated && bTemporal)
+		if (DiffuseLightingAndSecondMoment && SpecularLightingAndSecondMoment && SceneDepthCopy && NumFramesAccumulated && bTemporal)
 		{
-			GraphBuilder.QueueTextureExtraction(DiffuseLighting, &LightingViewState.DiffuseLightingHistory);
-			GraphBuilder.QueueTextureExtraction(SpecularLighting, &LightingViewState.SpecularLightingHistory);
-			GraphBuilder.QueueTextureExtraction(LuminanceMoments, &LightingViewState.LuminanceMomentsHistory);
+			GraphBuilder.QueueTextureExtraction(DiffuseLightingAndSecondMoment, &LightingViewState.DiffuseLightingAndSecondMomentHistory);
+			GraphBuilder.QueueTextureExtraction(SpecularLightingAndSecondMoment, &LightingViewState.SpecularLightingAndSecondMomentHistory);
 			GraphBuilder.QueueTextureExtraction(SceneDepthCopy, &LightingViewState.SceneDepthHistory);
 			GraphBuilder.QueueTextureExtraction(NumFramesAccumulated, &LightingViewState.NumFramesAccumulatedHistory);
 		}
 		else
 		{
-			LightingViewState.DiffuseLightingHistory = nullptr;
-			LightingViewState.SpecularLightingHistory = nullptr;
-			LightingViewState.LuminanceMomentsHistory = nullptr;
+			LightingViewState.DiffuseLightingAndSecondMomentHistory = nullptr;
+			LightingViewState.SpecularLightingAndSecondMomentHistory = nullptr;
 			LightingViewState.SceneDepthHistory = nullptr;
 			LightingViewState.NumFramesAccumulatedHistory = nullptr;
 		}
diff --git a/Engine/Source/Runtime/Renderer/Private/SampledDirectLighting/SampledDirectLightingViewState.h b/Engine/Source/Runtime/Renderer/Private/SampledDirectLighting/SampledDirectLightingViewState.h
--- a/Engine/Source/Runtime/Renderer/Private/SampledDirectLighting/SampledDirectLightingViewState.h
+++ b/Engine/Source/Runtime/Renderer/Private/SampledDirectLighting/SampledDirectLightingViewState.h
@@ -8,9 +8,8 @@
 class FSampledDirectLightingViewState
 {
 public:
-	TRefCountPtr<IPooledRenderTarget> DiffuseLightingHistory;
-	TRefCountPtr<IPooledRenderTarget> SpecularLightingHistory;
-	TRefCountPtr<IPooledRenderTarget> LuminanceMomentsHistory;
+	TRefCountPtr<IPooledRenderTarget> DiffuseLightingAndSecondMomentHistory;
+	TRefCountPtr<IPooledRenderTarget> SpecularLightingAndSecondMomentHistory;
 	TRefCountPtr<IPooledRenderTarget> SceneDepthHistory;
 	TRefCountPtr<IPooledRenderTarget> NumFramesAccumulatedHistory;
 
@@ -19,9 +18,8 @@ public:
 
 	void SafeRelease()
 	{
-		DiffuseLightingHistory.SafeRelease();
-		SpecularLightingHistory.SafeRelease();
-		LuminanceMomentsHistory.SafeRelease();
+		DiffuseLightingAndSecondMomentHistory.SafeRelease();
+		SpecularLightingAndSecondMomentHistory.SafeRelease();
 		SceneDepthHistory.SafeRelease();
 		NumFramesAccumulatedHistory.SafeRelease();
 	}
