Fix black views other than the last when split screen is enabled. Only override the load action if there is no action.

diff --git a/Engine/Source/Runtime/Renderer/Private/PostProcess/PostProcessTonemap.cpp b/Engine/Source/Runtime/Renderer/Private/PostProcess/PostProcessTonemap.cpp
--- a/Engine/Source/Runtime/Renderer/Private/PostProcess/PostProcessTonemap.cpp
+++ b/Engine/Source/Runtime/Renderer/Private/PostProcess/PostProcessTonemap.cpp
@@ -555,10 +555,11 @@ bool ShouldWriteAlphaChannel(const FViewInfo& View, const FTonemapInputs& Inputs
 	return (Inputs.bWriteAlphaChannel || bIsStereo || bFormatNeedsAlphaWrite);
 }
 
-bool ShouldOverrideOutputLoadActionToFastClear(const FRDGTextureRef Output, bool bShouldWriteAlphaChannel)
+bool ShouldOverrideOutputLoadActionToFastClear(const FScreenPassRenderTarget& Output, bool bShouldWriteAlphaChannel)
 {
 	bool bShouldOverride = false;
-	EPixelFormatChannelFlags OutputFlags = GetPixelFormatValidChannels(Output->Desc.Format);
+	
+	EPixelFormatChannelFlags OutputFlags = GetPixelFormatValidChannels(Output.Texture->Desc.Format);
 	// If we do not write through alpha channel but the output texture has alpha channel
 	// need to override to fast clear load action ERenderTargetLoadAction::Clear, otherwise,
 	// the alpha channel can be garbage data in terms of different driver implementation.
@@ -566,6 +567,13 @@ bool ShouldOverrideOutputLoadActionToFastClear(const FRDGTextureRef Output, bool
 	{
 		bShouldOverride = true;
 	}
+
+	// If the load action is already load, there should be content in the texture. Disable clear overriding.
+	if (Output.LoadAction == ERenderTargetLoadAction::ELoad)
+	{
+		bShouldOverride = false;
+	}
+
 	return bShouldOverride;
 }
 
@@ -641,7 +649,7 @@ FScreenPassTexture AddTonemapPass(FRDGBuilder& GraphBuilder, const FViewInfo& Vi
 	}
 
 	const bool bShouldWriteAlphaChannel = ShouldWriteAlphaChannel(View, Inputs, Output.Texture);
-	if (ShouldOverrideOutputLoadActionToFastClear(Output.Texture, bShouldWriteAlphaChannel))
+	if (ShouldOverrideOutputLoadActionToFastClear(Output, bShouldWriteAlphaChannel))
 	{
 		Output.LoadAction = ERenderTargetLoadAction::EClear;
 	}
