Fix incorrect rect to spot light conversion due to default parameters mismatching between LF atlas and light flags.

This fixes Lumen test failing when using light grid lighting with rect lights.

diff --git a/Engine/Source/Runtime/Renderer/Private/TranslucentLighting.cpp b/Engine/Source/Runtime/Renderer/Private/TranslucentLighting.cpp
--- a/Engine/Source/Runtime/Renderer/Private/TranslucentLighting.cpp
+++ b/Engine/Source/Runtime/Renderer/Private/TranslucentLighting.cpp
@@ -50,6 +50,7 @@
 #include "VolumetricCloudRendering.h"
 #include "RenderCore.h"
 #include "StaticMeshBatch.h"
+#include "LightFunctionAtlas.h"
 
 class FMaterial;
 
@@ -1098,6 +1099,7 @@ void InjectTranslucencyLightingVolume(
 	INC_DWORD_STAT_BY(STAT_NumLightsInjectedIntoTranslucency, LightInjectionData.Num());
 
 	const FVolumetricCloudShadowAOParameters CloudShadowAOParameters = GetCloudShadowAOParameters(GraphBuilder, View, Scene->GetVolumetricCloudSceneInfo());
+	const bool bUseLightFunctionAtlas = View.LightFunctionAtlasViewData.UsesLightFunctionAtlas(LightFunctionAtlas::ELightFunctionAtlasSystem::DeferredLighting);
 
 	FRDGTextureRef TransmittanceLutTexture = GetSkyTransmittanceLutTexture(GraphBuilder, Scene, View);
 
@@ -1142,7 +1144,7 @@ void InjectTranslucencyLightingVolume(
 				PassParameters->PS.ViewUniformBuffer = View.ViewUniformBuffer;
 
 				FDeferredLightUniformStruct* DeferredLightStruct = GraphBuilder.AllocParameters<FDeferredLightUniformStruct>();
-				*DeferredLightStruct = GetDeferredLightParameters(View, *LightSceneInfo, ELightShaderParameterFlags::RectAsSpotLight);
+				*DeferredLightStruct = GetDeferredLightParameters(View, *LightSceneInfo, bUseLightFunctionAtlas, ELightShaderParameterFlags::RectAsSpotLight);
 				PassParameters->PS.DeferredLight = GraphBuilder.CreateUniformBuffer(DeferredLightStruct);
 
 				GetVolumeShadowingShaderParameters(GraphBuilder, View, LightSceneInfo, InjectionData.ProjectedShadowInfo, PassParameters->PS.VolumeShadowingParameters);
