Fix race condition between Lumen scene lighting, and atmosphere transmittance update.

This could cause some flickering when the directional light parameters are animated (which would reset the directional light transmittance value).

diff --git a/Engine/Source/Runtime/Renderer/Private/DeferredShadingRenderer.cpp b/Engine/Source/Runtime/Renderer/Private/DeferredShadingRenderer.cpp
--- a/Engine/Source/Runtime/Renderer/Private/DeferredShadingRenderer.cpp
+++ b/Engine/Source/Runtime/Renderer/Private/DeferredShadingRenderer.cpp
@@ -1523,20 +1523,43 @@ void FDeferredShadingSceneRenderer::Render(FRDGBuilder& GraphBuilder)
 	{
 		if (RendererOutput == ERendererOutput::FinalSceneColor)
 		{
-			InitViewTaskDatas.LumenFrameTemporaries = &LumenFrameTemporaries;
-
-			// Important that this uses consistent logic throughout the frame, so evaluate once and pass in the flag from here
-			// NOTE: Must be done after  system texture initialization
-			// TODO: This doesn't take into account the potential for split screen views with separate shadow caches
-			const bool bEnableVirtualShadowMaps = UseVirtualShadowMaps(ShaderPlatform, FeatureLevel) && ViewFamily.EngineShowFlags.DynamicShadows && !bHasRayTracedOverlay;
-			VirtualShadowMapArray.Initialize(GraphBuilder, Scene->GetVirtualShadowMapCache(), bEnableVirtualShadowMaps, ViewFamily.EngineShowFlags);
-
-			if (InitViewTaskDatas.LumenFrameTemporaries)
+			// 1. Update sky atmosphere
+			// This needs to be done prior to start Lumen scene lighting to ensure directional light color is correct, as the sun color needs atmosphere transmittance
 			{
-				BeginUpdateLumenSceneTasks(GraphBuilder, *InitViewTaskDatas.LumenFrameTemporaries);
+				const bool bPathTracedAtmosphere = ViewFamily.EngineShowFlags.PathTracing && Views.Num() > 0 && Views[0].FinalPostProcessSettings.PathTracingEnableReferenceAtmosphere;
+				if (ShouldRenderSkyAtmosphere(Scene, ViewFamily.EngineShowFlags) && !bPathTracedAtmosphere)
+				{
+					for (int32 LightIndex = 0; LightIndex < NUM_ATMOSPHERE_LIGHTS; ++LightIndex)
+					{
+						if (Scene->AtmosphereLights[LightIndex])
+						{
+							PrepareSunLightProxy(*Scene->GetSkyAtmosphereSceneInfo(),LightIndex, *Scene->AtmosphereLights[LightIndex]);
+						}
+					}
+				}
+				else
+				{
+					Scene->ResetAtmosphereLightsProperties();
+				}
 			}
 
-			BeginGatherLumenLights(InitViewTaskDatas.LumenDirectLighting, InitViewTaskDatas.VisibilityTaskData);
+			// 2. Update lumen scene
+			{
+				InitViewTaskDatas.LumenFrameTemporaries = &LumenFrameTemporaries;
+	
+				// Important that this uses consistent logic throughout the frame, so evaluate once and pass in the flag from here
+				// NOTE: Must be done after  system texture initialization
+				// TODO: This doesn't take into account the potential for split screen views with separate shadow caches
+				const bool bEnableVirtualShadowMaps = UseVirtualShadowMaps(ShaderPlatform, FeatureLevel) && ViewFamily.EngineShowFlags.DynamicShadows && !bHasRayTracedOverlay;
+				VirtualShadowMapArray.Initialize(GraphBuilder, Scene->GetVirtualShadowMapCache(), bEnableVirtualShadowMaps, ViewFamily.EngineShowFlags);
+	
+				if (InitViewTaskDatas.LumenFrameTemporaries)
+				{
+					BeginUpdateLumenSceneTasks(GraphBuilder, *InitViewTaskDatas.LumenFrameTemporaries);
+				}
+	
+				BeginGatherLumenLights(InitViewTaskDatas.LumenDirectLighting, InitViewTaskDatas.VisibilityTaskData);
+			}
 		}
 
 		if (bNaniteEnabled)
@@ -1693,25 +1716,6 @@ void FDeferredShadingSceneRenderer::Render(FRDGBuilder& GraphBuilder)
 
 	CSV_SCOPED_TIMING_STAT_EXCLUSIVE(RenderOther);
 
-	if (RendererOutput == ERendererOutput::FinalSceneColor)
-	{
-		const bool bPathTracedAtmosphere = ViewFamily.EngineShowFlags.PathTracing && Views.Num() > 0 && Views[0].FinalPostProcessSettings.PathTracingEnableReferenceAtmosphere;
-		if (ShouldRenderSkyAtmosphere(Scene, ViewFamily.EngineShowFlags) && !bPathTracedAtmosphere)
-		{
-			for (int32 LightIndex = 0; LightIndex < NUM_ATMOSPHERE_LIGHTS; ++LightIndex)
-			{
-				if (Scene->AtmosphereLights[LightIndex])
-				{
-					PrepareSunLightProxy(*Scene->GetSkyAtmosphereSceneInfo(),LightIndex, *Scene->AtmosphereLights[LightIndex]);
-				}
-			}
-		}
-		else
-		{
-			Scene->ResetAtmosphereLightsProperties();
-		}
-	}
-
 	SCOPED_NAMED_EVENT(FDeferredShadingSceneRenderer_Render, FColor::Emerald);
 
 #if WITH_MGPU
