Fix bug where setting r.LightMaxDrawDistanceScale to 0 (as is done on Low) would cause local light max draw distance to be ignored. Setting this to 0 now results in all local lights  with non-zero cull distances being culled.
Adjusted r.LightMaxDrawDistanceScale on default Low scalability to be the same as Medium (0.5).

diff --git a/Engine/Config/BaseScalability.ini b/Engine/Config/BaseScalability.ini
--- a/Engine/Config/BaseScalability.ini
+++ b/Engine/Config/BaseScalability.ini
@@ -136,7 +136,7 @@ r.Shadow.CSM.TransitionScale=0
 r.Shadow.PreShadowResolutionFactor=0.5
 r.DistanceFieldShadowing=0
 r.VolumetricFog=0
-r.LightMaxDrawDistanceScale=0
+r.LightMaxDrawDistanceScale=0.5
 r.CapsuleShadows=0
 r.Shadow.Virtual.MaxPhysicalPages=512
 r.Shadow.Virtual.ResolutionLodBiasDirectional=0.0
@@ -160,7 +160,7 @@ r.Shadow.CSM.TransitionScale=0.25
 r.Shadow.PreShadowResolutionFactor=0.5
 r.DistanceFieldShadowing=0
 r.VolumetricFog=0
-r.LightMaxDrawDistanceScale=.5
+r.LightMaxDrawDistanceScale=0.5
 r.CapsuleShadows=1
 r.Shadow.Virtual.MaxPhysicalPages=512
 r.Shadow.Virtual.ResolutionLodBiasDirectional=0.0
diff --git a/Engine/Source/Runtime/Renderer/Private/LightRendering.cpp b/Engine/Source/Runtime/Renderer/Private/LightRendering.cpp
--- a/Engine/Source/Runtime/Renderer/Private/LightRendering.cpp
+++ b/Engine/Source/Runtime/Renderer/Private/LightRendering.cpp
@@ -722,9 +722,11 @@ float GetLightFadeFactor(const FSceneView& View, const FLightSceneProxy* Proxy)
 	SizeFade = FMath::Clamp(6.0f - 6.0f * SizeFade, 0.0f, 1.0f);
 
 	extern float GLightMaxDrawDistanceScale;
-	float MaxDist = Proxy->GetMaxDrawDistance() * GLightMaxDrawDistanceScale;
+	float ProxyMaxDist = Proxy->GetMaxDrawDistance();
+	float ScaledMaxDist = ProxyMaxDist * GLightMaxDrawDistanceScale;
+	// NOTE: Feels like we should scale fade range by GLightMaxDrawDistanceScale as well, but would change legacy behavior
 	float Range = Proxy->GetFadeRange();
-	float DistanceFade = MaxDist ? (MaxDist - FMath::Sqrt(DistanceSquared)) / Range : 1.0f;
+	float DistanceFade = ProxyMaxDist > 0.0f ? (ScaledMaxDist - FMath::Sqrt(DistanceSquared)) / Range : 1.0f;
 	DistanceFade = FMath::Clamp(DistanceFade, 0.0f, 1.0f);
 	return SizeFade * DistanceFade;
 }
diff --git a/Engine/Source/Runtime/Renderer/Private/SceneVisibility.cpp b/Engine/Source/Runtime/Renderer/Private/SceneVisibility.cpp
--- a/Engine/Source/Runtime/Renderer/Private/SceneVisibility.cpp
+++ b/Engine/Source/Runtime/Renderer/Private/SceneVisibility.cpp
@@ -5049,9 +5049,10 @@ void FSceneRenderer::ComputeLightVisibility()
 				if (View.IsPerspectiveProjection())
 				{
 					const float DistanceSquared = (BoundingSphere.Center - View.ViewMatrices.GetViewOrigin()).SizeSquared();
-					const float MaxDistSquared = Proxy->GetMaxDrawDistance() * Proxy->GetMaxDrawDistance() * GLightMaxDrawDistanceScale * GLightMaxDrawDistanceScale;
+					const float ProxyMaxDistance = Proxy->GetMaxDrawDistance();
+					const float ScaledMaxDistance = ProxyMaxDistance * GLightMaxDrawDistanceScale;
 					const bool bDrawLight = (FMath::Square(FMath::Min(0.0002f, GMinScreenRadiusForLights / BoundingSphere.W) * View.LODDistanceFactor) * DistanceSquared < 1.0f)
-												&& (MaxDistSquared == 0 || DistanceSquared < MaxDistSquared);
+												&& (ProxyMaxDistance <= 0.0 || DistanceSquared < FMath::Square(ScaledMaxDistance));
 							
 					VisibleLightViewInfo.bInViewFrustum = bDrawLight && bInViewFrustum;
 					VisibleLightViewInfo.bInDrawRange = bDrawLight;
