Make Substrate 8bits tile encoding dynamic.

This fixes carsh on platforms using 8bits tile encoding but using higher resolution than 1080p.

diff --git a/Engine/Shaders/Private/PackUnpack.ush b/Engine/Shaders/Private/PackUnpack.ush
--- a/Engine/Shaders/Private/PackUnpack.ush
+++ b/Engine/Shaders/Private/PackUnpack.ush
@@ -267,6 +267,8 @@ float UnpackR24F(uint In)
 #define PackTileCoordXXbits(In, InMask, InBits)  ((In.x & InMask) | ((In.y & InMask) << InBits))
 #define UnpackTileCoordXXbits(In,InMask, InBits) uint2(In & InMask, (In >> InBits) & InMask)
 
+uint  PackTileCoord8bits(uint2 In)  { return PackTileCoordXXbits  (In, 0xFF,  8u); }
+uint2 UnpackTileCoord8bits(uint In) { return UnpackTileCoordXXbits(In, 0xFF,  8u); }
 uint  PackTileCoord10bits(uint2 In) { return PackTileCoordXXbits  (In, 0x3FF, 10u); }
 uint2 UnpackTileCoord10bits(uint In){ return UnpackTileCoordXXbits(In, 0x3FF, 10u); }
 uint  PackTileCoord12bits(uint2 In) { return PackTileCoordXXbits  (In, 0xFFF, 12u); }
diff --git a/Engine/Shaders/Private/Substrate/SubstrateDBuffer.usf b/Engine/Shaders/Private/Substrate/SubstrateDBuffer.usf
--- a/Engine/Shaders/Private/Substrate/SubstrateDBuffer.usf
+++ b/Engine/Shaders/Private/Substrate/SubstrateDBuffer.usf
@@ -79,6 +79,7 @@ uint FirstSliceStoringSubstrateSSSData;
 RWTexture2D<SUBSTRATE_TOP_LAYER_TYPE> TopLayerTexture;
 RWTexture2DArray<uint> MaterialTextureArrayUAV;
 
+uint TileEncoding;
 uint TileListBufferOffset;
 Buffer<uint> TileListBuffer;
 Buffer<uint> TileIndirectBuffer;
@@ -97,7 +98,7 @@ void MainCS(
 	if (View.GeneralPurposeTweak == 0)
 		return;
 
-	const uint2 TileCoord = SubstrateUnpackTile(TileListBuffer[TileListBufferOffset + GroupId.x]);
+	const uint2 TileCoord = SubstrateUnpackTile(TileListBuffer[TileListBufferOffset + GroupId.x], TileEncoding);
 	const uint2 PixelCoord = TileCoord * SUBSTRATE_TILE_SIZE + GroupThreadId;
 	const bool bIsValid = all(PixelCoord < uint2(View.ViewSizeAndInvSize.xy));
 
diff --git a/Engine/Shaders/Private/Substrate/SubstrateMaterialClassification.usf b/Engine/Shaders/Private/Substrate/SubstrateMaterialClassification.usf
--- a/Engine/Shaders/Private/Substrate/SubstrateMaterialClassification.usf
+++ b/Engine/Shaders/Private/Substrate/SubstrateMaterialClassification.usf
@@ -23,6 +23,7 @@ Texture2D<uint> TopLayerCmaskTexture;
 #endif
 RWTexture2DArray<uint> MaterialTextureArrayUAV;
 
+uint TileEncoding;
 uint4 TileListBufferOffsets[SUBSTRATE_TILE_TYPE_COUNT];
 uint GetTileListBufferOffsets(uint Type)
 {
@@ -275,7 +276,7 @@ void TileMainCS(uint2 DispatchThreadId : SV_DispatchThreadID, uint LinearIndex :
 
 	if (LinearIndex < 1 && bTileContainsSubstrate)
 	{
-		uint EncodedTile = SubstratePackTile(GroupId.xy);
+		uint EncodedTile = SubstratePackTile(GroupId.xy, TileEncoding);
 
 		if (bTileContainsComplexSpecial)
 		{
@@ -416,6 +417,7 @@ Texture2DArray<uint> MaterialTextureArray;
 
 Buffer<uint> TileListBuffer;
 uint TileListBufferOffset;
+uint TileEncoding;
 
 RWTexture2D<uint> RWClosureOffsetTexture;
 RWBuffer<uint>    RWClosureTileCountBuffer;
@@ -431,7 +433,7 @@ WAVESIZE(64) // PERMUTATION_WAVE_OPS is true only when wave>=64 are available
 [numthreads(SUBSTRATE_TILE_SIZE, SUBSTRATE_TILE_SIZE, 1)]
 void ClosureTileMainCS(uint2 GroupThreadId : SV_GroupThreadID, uint2 GroupId : SV_GroupID, uint LinearIndex : SV_GroupIndex)
 {
-	const uint2 TileCoord = SubstrateUnpackTile(TileListBuffer[TileListBufferOffset + GroupId.x]);
+	const uint2 TileCoord = SubstrateUnpackTile(TileListBuffer[TileListBufferOffset + GroupId.x], TileEncoding);
 	uint2 PixelCoord = TileCoord * SUBSTRATE_TILE_SIZE + GroupThreadId;
 	const bool bIsInViewRect = all(PixelCoord < uint2(View.ViewRectMinAndSize.zw));
 	PixelCoord += View.ViewRectMinAndSize.xy;
diff --git a/Engine/Shaders/Private/Substrate/SubstrateTile.usf b/Engine/Shaders/Private/Substrate/SubstrateTile.usf
--- a/Engine/Shaders/Private/Substrate/SubstrateTile.usf
+++ b/Engine/Shaders/Private/Substrate/SubstrateTile.usf
@@ -16,6 +16,7 @@ float4x4 ViewScreenToTranslatedWorld;
 
 Buffer<uint> TileListBuffer;
 uint		 TileListBufferOffset;
+uint		 TileEncoding;
 
 void SubstrateTilePassVS(
 	in uint InVertexId : SV_VertexID,
@@ -30,7 +31,7 @@ void SubstrateTilePassVS(
 	out float4 OutPosition : SV_POSITION)	// Use by StencilMainPS
 {
 #if SUBSTRATE_ENABLED
-	const uint2 TileCoord = SubstrateUnpackTile(TileListBuffer[TileListBufferOffset + InInstanceId]);
+	const uint2 TileCoord = SubstrateUnpackTile(TileListBuffer[TileListBufferOffset + InInstanceId], TileEncoding);
 #else
 	// Not used, but this shader is compiled for non-Substrate path for setup reason, but not used
 	const uint2 TileCoord = 0;
diff --git a/Engine/Shaders/Private/Substrate/SubstrateTile.ush b/Engine/Shaders/Private/Substrate/SubstrateTile.ush
--- a/Engine/Shaders/Private/Substrate/SubstrateTile.ush
+++ b/Engine/Shaders/Private/Substrate/SubstrateTile.ush
@@ -8,33 +8,17 @@
 
 #if SUBSTRATE_ENABLED
 
-#ifndef USE_8BIT_TILE_COORD
-#error USE_8BIT_TILE_COORD needs to be defined
-#define USE_8BIT_TILE_COORD 0
-#endif
 ////////////////////////////////////////////////////////////////////////////////////////////////////////////
 // Substrate material type tile
 
-uint2 SubstrateUnpackTile(uint In)
+uint2 SubstrateUnpackTile(uint In, uint InEncoding)
 {
-#if USE_8BIT_TILE_COORD
-	return uint2(In & 0xFF, In >> 8);
-#else
-	return uint2(In & 0xFFFF, In >> 16);
-#endif
+	return InEncoding == SUBSTRATE_TILE_ENCODING_16BITS ? UnpackTileCoord16bits(In) : UnpackTileCoord8bits(In);
 }
 
-uint SubstratePackTile(uint2 TileCoord)
+uint SubstratePackTile(uint2 TileCoord, uint InEncoding)
 {
-#if USE_8BIT_TILE_COORD
-	return TileCoord.x | (TileCoord.y << 8);
-#else
-  #if COMPILER_SUPPORTS_PACK_INTRINSICS
-	return PackUInt2ToUInt(TileCoord.x, TileCoord.y);
-  #else
-	return TileCoord.x | (TileCoord.y << 16); // assumes 16bit is enough to represent a tiled resolution up to 65,535 :)
-  #endif
-#endif
+	return InEncoding == SUBSTRATE_TILE_ENCODING_16BITS ? PackUInt2ToUInt(TileCoord)/*TODO*/ : PackTileCoord8bits(TileCoord);
 }
 
 ////////////////////////////////////////////////////////////////////////////////////////////////////////////
diff --git a/Engine/Shaders/Shared/SubstrateDefinitions.h b/Engine/Shaders/Shared/SubstrateDefinitions.h
--- a/Engine/Shaders/Shared/SubstrateDefinitions.h
+++ b/Engine/Shaders/Shared/SubstrateDefinitions.h
@@ -76,6 +76,9 @@
 #define SUBSTRATE_MATERIAL_TYPE_COMPLEX					2
 #define SUBSTRATE_MATERIAL_TYPE_COMPLEX_SPECIAL			3
 
+#define SUBSTRATE_TILE_ENCODING_16BITS 					0
+#define SUBSTRATE_TILE_ENCODING_8BITS  					1
+
 // sizeof(FRHIDrawIndirectParameters) = 4 uints = 16 bytes
 #define GetSubstrateTileTypeDrawIndirectArgOffset_Byte(x)  (x * 16)
 #define GetSubstrateTileTypeDrawIndirectArgOffset_DWord(x) (x * 4)
diff --git a/Engine/Source/Runtime/Engine/Private/ShaderCompiler/ShaderCompiler.cpp b/Engine/Source/Runtime/Engine/Private/ShaderCompiler/ShaderCompiler.cpp
--- a/Engine/Source/Runtime/Engine/Private/ShaderCompiler/ShaderCompiler.cpp
+++ b/Engine/Source/Runtime/Engine/Private/ShaderCompiler/ShaderCompiler.cpp
@@ -8265,9 +8265,6 @@ void GlobalBeginCompileShader(
 			const uint32 SubstrateClosurePerPixel = Substrate::GetClosurePerPixel(Target.GetPlatform());
 			SET_SHADER_DEFINE(Input.Environment, SUBSTRATE_MATERIAL_CLOSURE_COUNT, SubstrateClosurePerPixel);
 
-			const bool bTileCoord8Bits = Substrate::Is8bitTileCoordEnabled();
-			SET_SHADER_DEFINE(Input.Environment, USE_8BIT_TILE_COORD, bTileCoord8Bits ? 1 : 0);
-
 			const bool bSubstrateDBufferPass = Substrate::IsDBufferPassEnabled(Target.GetPlatform());
 			SET_SHADER_DEFINE(Input.Environment, SUBSTRATE_USE_DBUFFER_PASS, bSubstrateDBufferPass ? 1 : 0);
 
diff --git a/Engine/Source/Runtime/RenderCore/Private/RenderUtils.cpp b/Engine/Source/Runtime/RenderCore/Private/RenderUtils.cpp
--- a/Engine/Source/Runtime/RenderCore/Private/RenderUtils.cpp
+++ b/Engine/Source/Runtime/RenderCore/Private/RenderUtils.cpp
@@ -1792,12 +1792,6 @@ static TAutoConsoleVariable<int32> CVarSubstrateDebugAdvancedVisualizationShader
 	TEXT("Enable advanced Substrate material debug visualization shaders. Base pass shaders can output such advanced data."),
 	ECVF_ReadOnly | ECVF_RenderThreadSafe);
 
-static TAutoConsoleVariable<int32> CVarSubstrateTileCoord8Bits(
-	TEXT("r.Substrate.TileCoord8bits"),
-	0,
-	TEXT("Format of tile coord. This variable is read-only."),
-	ECVF_ReadOnly | ECVF_RenderThreadSafe);
-
 namespace Substrate
 {
 	bool IsSubstrateEnabled()
@@ -1930,11 +1924,6 @@ namespace Substrate
 		return IsSubstrateEnabled() && CVar.Get(InPlatform) > 0;
 	}
 
-	bool Is8bitTileCoordEnabled()
-	{
-		return CVarSubstrateTileCoord8Bits.GetValueOnAnyThread() > 0 ? 1 : 0;
-	}
-
 	uint32 GetSheenQuality()
 	{
 		return CVarSubstrateSheenQuality.GetValueOnAnyThread();
diff --git a/Engine/Source/Runtime/RenderCore/Public/RenderUtils.h b/Engine/Source/Runtime/RenderCore/Public/RenderUtils.h
--- a/Engine/Source/Runtime/RenderCore/Public/RenderUtils.h
+++ b/Engine/Source/Runtime/RenderCore/Public/RenderUtils.h
@@ -548,7 +548,6 @@ namespace Substrate
 	RENDERCORE_API bool IsDBufferPassEnabled(EShaderPlatform InPlatform);
 	RENDERCORE_API bool IsOpaqueRoughRefractionEnabled();
 	RENDERCORE_API bool IsAdvancedVisualizationEnabled();
-	RENDERCORE_API bool Is8bitTileCoordEnabled();
 
 	RENDERCORE_API uint32 GetRayTracingMaterialPayloadSizeInBytes(bool bFullySimplifiedMaterial);
 
diff --git a/Engine/Source/Runtime/Renderer/Private/Substrate/Substrate.cpp b/Engine/Source/Runtime/Renderer/Private/Substrate/Substrate.cpp
--- a/Engine/Source/Runtime/Renderer/Private/Substrate/Substrate.cpp
+++ b/Engine/Source/Runtime/Renderer/Private/Substrate/Substrate.cpp
@@ -63,6 +63,12 @@ static TAutoConsoleVariable<int32> CVarSubstrateAllocationMode(
 	TEXT("Substrate resource allocation mode. \n 0: Allocate resources based on view requirement, \n 1: Allocate resources based on view requirement, but can only grow over frame to minimize resources reallocation and hitches, \n 2: Allocate resources based on platform settings."),
 	ECVF_RenderThreadSafe);
 
+static TAutoConsoleVariable<int32> CVarSubstrateTileCoord8Bits(
+	TEXT("r.Substrate.TileCoord8bits"),
+	0,
+	TEXT("Format of tile coord. This variable is read-only."),
+	ECVF_RenderThreadSafe);
+
 IMPLEMENT_GLOBAL_SHADER_PARAMETER_STRUCT(FSubstrateGlobalUniformParameters, "Substrate");
 
 void FSubstrateViewData::Reset()
@@ -149,6 +155,11 @@ FIntPoint GetSubstrateTextureResolution(const FViewInfo& View, const FIntPoint&
 	return InResolution;
 }
 
+bool Is8bitTileCoordEnabled()
+{
+	return CVarSubstrateTileCoord8Bits.GetValueOnAnyThread() > 0 ? 1 : 0;
+}
+
 bool GetSubstrateUsesComplexSpecialPath(const FViewInfo& View)
 {
 	if (Substrate::IsSubstrateEnabled())
@@ -172,15 +183,9 @@ bool IsClassificationAsync()
 	return CVarSubstrateAsyncClassification.GetValueOnRenderThread() > 0;
 }
 
-static EPixelFormat GetClassificationTileFormat(const FIntPoint& InResolution)
+static EPixelFormat GetClassificationTileFormat(const FIntPoint& InResolution, uint32 InTileEncoding)
 {
-	// For platform which whose resolution is never above 1080p, use 8bit tile format for performance.
-	const bool bRequest8bit = Is8bitTileCoordEnabled();
-	if (bRequest8bit)
-	{
-		check(InResolution.X <= 2048 && InResolution.Y <= 2048);
-	}
-	return bRequest8bit ? PF_R16_UINT : PF_R32_UINT;
+	return InTileEncoding == SUBSTRATE_TILE_ENCODING_16BITS ? PF_R32_UINT : PF_R16_UINT;
 }
 
 static void InitialiseSubstrateViewData(FRDGBuilder& GraphBuilder, FViewInfo& View, const FSceneTexturesConfig& SceneTexturesConfig, bool bNeedClosureOffets, FSubstrateSceneData& SceneData)
@@ -215,6 +220,10 @@ static void InitialiseSubstrateViewData(FRDGBuilder& GraphBuilder, FViewInfo& Vi
 			const uint32 DecalTileCount = IsDBufferPassEnabled(View.GetShaderPlatform()) ? TileResolution.X * TileResolution.Y : 4;
 			const uint32 RegularTileCount = TileResolution.X * TileResolution.Y;
 
+			// For platforms whose resolution is never above 1080p, use 8bit tile format for performance, if possible
+			const bool bRequest8bit = Substrate::Is8bitTileCoordEnabled() && (TileResolution.X <= 256 && TileResolution.Y <= 256);
+			Out.TileEncoding = bRequest8bit ? SUBSTRATE_TILE_ENCODING_8BITS : SUBSTRATE_TILE_ENCODING_16BITS;
+
 			bool bUsesComplexSpecialRenderPath = SceneData.bUsesComplexSpecialRenderPath; // Use the Scene temporally stable bUsesComplexSpecialRenderPath to reduce buffer reallocation.
 
 			Out.ClassificationTileListBufferOffset[ESubstrateTileType::ESimple]							= 0;
@@ -230,7 +239,7 @@ static void InitialiseSubstrateViewData(FRDGBuilder& GraphBuilder, FViewInfo& Vi
 
 			check(TotalTileCount > 0);
 
-			const EPixelFormat ClassificationTileFormat = GetClassificationTileFormat(DynResIndependentViewSize);
+			const EPixelFormat ClassificationTileFormat = GetClassificationTileFormat(DynResIndependentViewSize, Out.TileEncoding);
 			const uint32 FormatBytes = ClassificationTileFormat == PF_R16_UINT ? sizeof(uint16) : sizeof(uint32);
 
 			Out.ClassificationTileListBuffer = GraphBuilder.CreateBuffer(FRDGBufferDesc::CreateBufferDesc(FormatBytes, TotalTileCount), TEXT("Substrate.TileListBuffer"));
@@ -685,6 +694,7 @@ class FSubstrateClosureTilePassCS : public FGlobalShader
 
 		SHADER_PARAMETER_RDG_BUFFER_SRV(Buffer<uint>, TileListBuffer)
 		SHADER_PARAMETER(uint32, TileListBufferOffset)
+		SHADER_PARAMETER(uint32, TileEncoding)
 		RDG_BUFFER_ACCESS(TileIndirectBuffer, ERHIAccess::IndirectArgs)
 	END_SHADER_PARAMETER_STRUCT()
 
@@ -731,6 +741,7 @@ class FSubstrateMaterialTileClassificationPassCS : public FGlobalShader
 		SHADER_PARAMETER(int32, bRectPrimitive)
 		SHADER_PARAMETER(FIntPoint, ViewResolution)
 		SHADER_PARAMETER(uint32, MaxBytesPerPixel)
+		SHADER_PARAMETER(uint32, TileEncoding)
 		SHADER_PARAMETER_ARRAY(FUintVector4, TileListBufferOffsets, [SUBSTRATE_TILE_TYPE_COUNT])
 		SHADER_PARAMETER_RDG_TEXTURE(Texture2D<SUBSTRATE_TOP_LAYER_TYPE>, TopLayerTexture)
 		SHADER_PARAMETER_RDG_TEXTURE(Texture2D, TopLayerCmaskTexture)
@@ -791,6 +802,7 @@ class FSubstrateDBufferPassCS : public FGlobalShader
 		SHADER_PARAMETER_RDG_TEXTURE_UAV(RWTexture2DArray<uint>, MaterialTextureArrayUAV)
 		SHADER_PARAMETER_RDG_BUFFER_SRV(Buffer<uint>, TileListBuffer)
 		SHADER_PARAMETER(uint32, TileListBufferOffset)
+		SHADER_PARAMETER(uint32, TileEncoding)
 		SHADER_PARAMETER_RDG_TEXTURE_SRV(Texture2D<float>, SceneStencilTexture)
 		RDG_BUFFER_ACCESS(TileIndirectBuffer, ERHIAccess::IndirectArgs)
 	END_SHADER_PARAMETER_STRUCT()
@@ -925,6 +937,7 @@ static FSubstrateTileParameter InternalSetTileParameters(FRDGBuilder* GraphBuild
 	{
 		Out.TileListBuffer = View.SubstrateViewData.ClassificationTileListBufferSRV;
 		Out.TileListBufferOffset = View.SubstrateViewData.ClassificationTileListBufferOffset[TileType];
+		Out.TileEncoding = View.SubstrateViewData.TileEncoding;
 		Out.TileIndirectBuffer = View.SubstrateViewData.ClassificationTileDrawIndirectBuffer;
 	}
 	else if (GraphBuilder)
@@ -933,6 +946,7 @@ static FSubstrateTileParameter InternalSetTileParameters(FRDGBuilder* GraphBuild
 		FRDGBufferSRVRef BufferDummySRV = GraphBuilder->CreateSRV(BufferDummy, PF_R32_UINT);
 		Out.TileListBuffer = BufferDummySRV;
 		Out.TileListBufferOffset = 0;
+		Out.TileEncoding = SUBSTRATE_TILE_ENCODING_16BITS;
 		Out.TileIndirectBuffer = BufferDummy;
 	}
 	return Out;
@@ -953,6 +967,7 @@ FSubstrateTilePassVS::FParameters SetTileParameters(
 	Out.ViewScreenToTranslatedWorld = View.CachedViewUniformShaderParameters->ScreenToTranslatedWorld;
 	Out.TileListBuffer = Temp.TileListBuffer;
 	Out.TileListBufferOffset = Temp.TileListBufferOffset;
+	Out.TileEncoding = Temp.TileEncoding;
 	Out.TileIndirectBuffer = Temp.TileIndirectBuffer;
 	return Out;
 }
@@ -973,6 +988,7 @@ FSubstrateTilePassVS::FParameters SetTileParameters(
 	Out.ViewScreenToTranslatedWorld = View.CachedViewUniformShaderParameters->ScreenToTranslatedWorld;
 	Out.TileListBuffer = Temp.TileListBuffer;
 	Out.TileListBufferOffset = Temp.TileListBufferOffset;
+	Out.TileEncoding = Temp.TileEncoding;
 	Out.TileIndirectBuffer = Temp.TileIndirectBuffer;
 	return Out;
 }
@@ -1286,6 +1302,7 @@ void AddSubstrateMaterialClassificationPass(FRDGBuilder& GraphBuilder, const FMi
 			PassParameters->DBuffer = GetDBufferParameters(GraphBuilder, DBufferTextures, Platform);
 			PassParameters->SceneStencilTexture = SceneTextures.Stencil;
 			PassParameters->TileListBufferUAV = SubstrateViewData->ClassificationTileListBufferUAV;
+			PassParameters->TileEncoding = SubstrateViewData->TileEncoding;
 			for (uint32 TileType = 0; TileType < SUBSTRATE_TILE_TYPE_COUNT; ++TileType)
 			{
 				PassParameters->TileListBufferOffsets[TileType] = FUintVector4(SubstrateViewData->ClassificationTileListBufferOffset[TileType], 0, 0, 0);
@@ -1338,6 +1355,7 @@ void AddSubstrateMaterialClassificationPass(FRDGBuilder& GraphBuilder, const FMi
 				PassParameters->MaterialTextureArray = SubstrateSceneData->MaterialTextureArraySRV;
 				PassParameters->TileListBuffer = SubstrateViewData->ClassificationTileListBufferSRV;
 				PassParameters->TileListBufferOffset = SubstrateViewData->ClassificationTileListBufferOffset[TileType];
+				PassParameters->TileEncoding = SubstrateViewData->TileEncoding;
 				PassParameters->TileIndirectBuffer = SubstrateViewData->ClassificationTileDispatchIndirectBuffer;
 
 				PassParameters->RWClosureOffsetTexture = GraphBuilder.CreateUAV(SubstrateSceneData->ClosureOffsetTexture);
@@ -1445,6 +1463,7 @@ void AddSubstrateDBufferPass(FRDGBuilder& GraphBuilder, const FMinimalSceneTextu
 
 			PassParameters->TileListBuffer = SubstrateViewData->ClassificationTileListBufferSRV;
 			PassParameters->TileListBufferOffset = SubstrateViewData->ClassificationTileListBufferOffset[TileType];
+			PassParameters->TileEncoding = SubstrateViewData->TileEncoding;
 			PassParameters->TileIndirectBuffer = SubstrateViewData->ClassificationTileDispatchIndirectBuffer;
 
 			// Dispatch with tile data
diff --git a/Engine/Source/Runtime/Renderer/Private/Substrate/Substrate.h b/Engine/Source/Runtime/Renderer/Private/Substrate/Substrate.h
--- a/Engine/Source/Runtime/Renderer/Private/Substrate/Substrate.h
+++ b/Engine/Source/Runtime/Renderer/Private/Substrate/Substrate.h
@@ -51,6 +51,7 @@ END_SHADER_PARAMETER_STRUCT()
 BEGIN_SHADER_PARAMETER_STRUCT(FSubstrateTileParameter, )
 	SHADER_PARAMETER_RDG_BUFFER_SRV(Buffer<uint>, TileListBuffer)
 	SHADER_PARAMETER(uint32, TileListBufferOffset)
+	SHADER_PARAMETER(uint32, TileEncoding)
 	RDG_BUFFER_ACCESS(TileIndirectBuffer, ERHIAccess::IndirectArgs)
 END_SHADER_PARAMETER_STRUCT()
 
@@ -151,6 +152,7 @@ struct FSubstrateViewData
 	bool bUsesComplexSpecialRenderPath = 0;
 
 	FIntPoint TileCount  = FIntPoint(0, 0);
+	uint32    TileEncoding = SUBSTRATE_TILE_ENCODING_16BITS;
 	uint32    LayerCount = 0;
 
 	FRDGBufferRef    ClassificationTileListBuffer = nullptr;
@@ -231,6 +233,7 @@ class FSubstrateTilePassVS : public FGlobalShader
 		SHADER_PARAMETER(FMatrix44f, ViewScreenToTranslatedWorld)
 		SHADER_PARAMETER_RDG_BUFFER_SRV(Buffer<uint>, TileListBuffer)
 		SHADER_PARAMETER(uint32, TileListBufferOffset)
+		SHADER_PARAMETER(uint32, TileEncoding)
 		RDG_BUFFER_ACCESS(TileIndirectBuffer, ERHIAccess::IndirectArgs)
 	END_SHADER_PARAMETER_STRUCT()
 
diff --git a/Engine/Source/Runtime/Renderer/Private/Substrate/SubstrateVisualize.cpp b/Engine/Source/Runtime/Renderer/Private/Substrate/SubstrateVisualize.cpp
--- a/Engine/Source/Runtime/Renderer/Private/Substrate/SubstrateVisualize.cpp
+++ b/Engine/Source/Runtime/Renderer/Private/Substrate/SubstrateVisualize.cpp
@@ -281,6 +281,7 @@ bool IsClassificationAsync();
 bool SupportsCMask(const FStaticShaderPlatform InPlatform);
 bool UsesSubstrateClosureCountFromMaterialData();
 uint32 GetMaterialBufferAllocationMode();
+bool Is8bitTileCoordEnabled();
 
 static void AddVisualizeSystemInfoPasses(FRDGBuilder& GraphBuilder, const FViewInfo& View, FScreenPassTexture& ScreenPassSceneColor, EShaderPlatform Platform)
 {
