Fix Substrate view mode on consoles.

diff --git a/Engine/Shaders/Private/Substrate/SubstrateVisualize.usf b/Engine/Shaders/Private/Substrate/SubstrateVisualize.usf
--- a/Engine/Shaders/Private/Substrate/SubstrateVisualize.usf
+++ b/Engine/Shaders/Private/Substrate/SubstrateVisualize.usf
@@ -18,6 +18,13 @@
 #include "/Engine/Private/ShaderPrint.ush"
 #include "../ColorMap.ush"
 
+uint bOverrideCursorPosition;
+int2 GetCursorPosition()
+{
+	const int2 OverrideCursorPosition = (View.ViewSizeAndInvSize.xy + View.ViewRectMin.xy) / 2;
+	return bOverrideCursorPosition > 0 ? OverrideCursorPosition : View.CursorPosition;	
+}
+
 // This is for the fxc compiler to compile faster.
 #define SUBSTRATE_VIZ_MAX_BSDF 2
 
@@ -642,7 +649,7 @@ float3 GetWorldPositionFromPixelPos(uint2 PixelPos)
 [numthreads(1, 1, 1)]
 void MaterialPrintInfoCS(uint3 DispatchThreadId : SV_DispatchThreadID)
 {
-	const uint2 PixelPos = (float2(View.CursorPosition) * View.ViewResolutionFraction);
+	const uint2 PixelPos = (float2(GetCursorPosition()) * View.ViewResolutionFraction);
 	const float3 WorldPosition = GetWorldPositionFromPixelPos(PixelPos);
 	const float3 V = -normalize(WorldPosition - PrimaryView.TranslatedWorldCameraOrigin);
 
@@ -655,6 +662,14 @@ void MaterialPrintInfoCS(uint3 DispatchThreadId : SV_DispatchThreadID)
 	const float3 dPdx = WorldPositionX - WorldPosition;
 	const float3 dPdy = WorldPositionY - WorldPosition;
 	DrawPixelFootprint(WorldPosition, dPdx, dPdy, PixelPos);
+
+	// When overriding cursor position (on platform without mouse), display a small rect to show the pointing position
+	if (bOverrideCursorPosition)
+	{
+		const uint2 CursorPosition = GetCursorPosition();
+		FShaderPrintContext Ctx = InitShaderPrintContext(true, uint2(0, 0));
+		AddQuadSS(Ctx, float2(CursorPosition) - float2(1, 1), float2(CursorPosition) + float2(1, 1), ColorYellow);
+	}
 }
 #endif // SHADER_MATERIALPRINT
 
@@ -755,7 +770,7 @@ void VisualizeMaterialPS(
 		}
 
 		// Cursor
-		if (all(PixelPosDynRes == ShaderPrintData.CursorCoord))
+		if (all(PixelPosDynRes == GetCursorPosition()))
 		{
 			FShaderPrintContext Context = InitShaderPrintContext(true, uint2(250, 50));
 			Print(Context, MaterialBytePerPixel, InitFontColor(BytePerPixelToColor(MaterialBytePerPixel)), 3, 0); 
@@ -801,7 +816,7 @@ void MaterialDebugSubstrateTreeCS(uint3 DispatchThreadId : SV_DispatchThreadID)
 {
 // UE-194712 - Temporary workaround with Vulkan shader compilation
 #if !COMPILER_VULKAN
-	const uint2 PixelPos = (float2(View.CursorPosition) * View.ViewResolutionFraction);
+	const uint2 PixelPos = (float2(GetCursorPosition()) * View.ViewResolutionFraction);
 
 	FSubstrateAddressing SubstrateAddressing = GetSubstratePixelDataByteOffset(PixelPos, uint2(View.BufferSizeAndInvSize.xy), Substrate.MaxBytesPerPixel);
 	FSubstratePixelHeader Header = UnpackSubstrateHeaderIn(Substrate.MaterialTextureArray, SubstrateAddressing, Substrate.TopLayerTexture);
@@ -1407,7 +1422,7 @@ void MaterialDebugSubstrateTreePS(
 {
 	OutColor = float4(0.0f, 0.0f, 0.0f, 1.0f);
 
-	FSubstrateAddressing SubstrateAddressing = GetSubstratePixelDataByteOffset((float2(View.CursorPosition) * View.ViewResolutionFraction), uint2(View.BufferSizeAndInvSize.xy), Substrate.MaxBytesPerPixel);
+	FSubstrateAddressing SubstrateAddressing = GetSubstratePixelDataByteOffset((float2(GetCursorPosition()) * View.ViewResolutionFraction), uint2(View.BufferSizeAndInvSize.xy), Substrate.MaxBytesPerPixel);
 	FSubstratePixelHeader Header = UnpackSubstrateHeaderIn(Substrate.MaterialTextureArray, SubstrateAddressing, Substrate.TopLayerTexture);
 	if (Header.ClosureCount > 0)
 	{
diff --git a/Engine/Source/Runtime/Engine/Private/SubstrateVisualizationData.cpp b/Engine/Source/Runtime/Engine/Private/SubstrateVisualizationData.cpp
--- a/Engine/Source/Runtime/Engine/Private/SubstrateVisualizationData.cpp
+++ b/Engine/Source/Runtime/Engine/Private/SubstrateVisualizationData.cpp
@@ -10,7 +10,7 @@
 
 static FSubstrateVisualizationData GSubstrateVisualizationData;
 
-static FString ConfigureConsoleCommand(FSubstrateVisualizationData::TModeMap& ModeMap)
+static FString ConfigureConsoleCommand(FSubstrateVisualizationData::TModeMap& ModeMap, IConsoleVariable*& OutCVar)
 {
 	FString AvailableVisualizationModes;
 	for (FSubstrateVisualizationData::TModeMap::TConstIterator It = ModeMap.CreateConstIterator(); It; ++It)
@@ -29,7 +29,7 @@ static FString ConfigureConsoleCommand(FSubstrateVisualizationData::TModeMap& Mo
 	Out = TEXT("When the viewport view-mode is set to 'Substrate Visualization', this command specifies which of the various channels to display. Values entered other than the allowed values shown below will be ignored.");
 	Out += AvailableVisualizationModes;
 
-	IConsoleManager::Get().RegisterConsoleVariable(
+	OutCVar = IConsoleManager::Get().RegisterConsoleVariable(
 		FSubstrateVisualizationData::GetVisualizeConsoleCommandName(),
 		0,
 		*Out,
@@ -149,7 +149,7 @@ void FSubstrateVisualizationData::Initialize()
 			LOCTEXT("None", "None"));
 		
 
-		ConsoleDocumentationVisualizationMode = ConfigureConsoleCommand(AllModeMap);
+		ConsoleDocumentationVisualizationMode = ConfigureConsoleCommand(AllModeMap, CVarViewModes);
 
 		// Now only copy the available modes for the menu to not overload it with useless entries.
 		for (auto& Mode : AllModeMap)
@@ -187,6 +187,16 @@ FSubstrateViewMode FSubstrateVisualizationData::GetViewMode(const FName& InModeN
 	}
 }
 
+uint32 FSubstrateVisualizationData::GetViewMode()
+{
+	uint32 OutViewMode = 0;
+	if (GSubstrateVisualizationData.IsInitialized() && GSubstrateVisualizationData.CVarViewModes && GSubstrateVisualizationData.CVarViewModes->AsVariableInt())
+	{
+		OutViewMode = GSubstrateVisualizationData.CVarViewModes->AsVariableInt()->GetValueOnRenderThread();
+	}
+	return OutViewMode;
+}
+
 bool FSubstrateVisualizationData::GetModeDefaultComposited(const FName& InModeName) const
 {
 	if (const FModeRecord* Record = ModeMap.Find(InModeName))
diff --git a/Engine/Source/Runtime/Engine/Public/SubstrateVisualizationData.h b/Engine/Source/Runtime/Engine/Public/SubstrateVisualizationData.h
--- a/Engine/Source/Runtime/Engine/Public/SubstrateVisualizationData.h
+++ b/Engine/Source/Runtime/Engine/Public/SubstrateVisualizationData.h
@@ -57,6 +57,8 @@ public:
 
 	ENGINE_API FSubstrateViewMode GetViewMode(const FName& InModeName) const;
 
+	ENGINE_API static uint32 GetViewMode();
+
 	ENGINE_API bool GetModeDefaultComposited(const FName& InModeName) const;
 
 	inline const TModeMap& GetModeMap() const
@@ -80,6 +82,9 @@ private:
 
 	/** Flag indicating if system is initialized. **/
 	bool bIsInitialized;
+
+	/** CVar for handling view modes. **/
+	IConsoleVariable* CVarViewModes = nullptr;
 };
 
 ENGINE_API FSubstrateVisualizationData& GetSubstrateVisualizationData();
diff --git a/Engine/Source/Runtime/Renderer/Private/PostProcess/PostProcessing.cpp b/Engine/Source/Runtime/Renderer/Private/PostProcess/PostProcessing.cpp
--- a/Engine/Source/Runtime/Renderer/Private/PostProcess/PostProcessing.cpp
+++ b/Engine/Source/Runtime/Renderer/Private/PostProcess/PostProcessing.cpp
@@ -1384,7 +1384,17 @@ void AddPostProcessingPasses(
 	{
 		FScreenPassRenderTarget OverrideOutput;
 		PassSequence.AcceptOverrideIfLastPass(EPass::VisualizeSubstrate, OverrideOutput);
-		SceneColor = Substrate::AddSubstrateDebugPasses(GraphBuilder, View, SceneColor);
+
+		FScreenPassTexture DebugColorOutput = Substrate::AddSubstrateDebugPasses(GraphBuilder, View, SceneColor);
+		if (OverrideOutput.IsValid())
+		{
+			AddDrawTexturePass(GraphBuilder, View, DebugColorOutput, OverrideOutput);
+			SceneColor = OverrideOutput;
+		}
+		else
+		{
+			SceneColor = DebugColorOutput;
+		}
 	}
 
 	if (PassSequence.IsEnabled(EPass::VisualizeLightGrid))
diff --git a/Engine/Source/Runtime/Renderer/Private/ShaderPrint.cpp b/Engine/Source/Runtime/Renderer/Private/ShaderPrint.cpp
--- a/Engine/Source/Runtime/Renderer/Private/ShaderPrint.cpp
+++ b/Engine/Source/Runtime/Renderer/Private/ShaderPrint.cpp
@@ -1171,7 +1171,7 @@ namespace ShaderPrint
 			});
 	}
 
-	void InternalDrawZoom(FRDGBuilder& GraphBuilder, const FShaderPrintData& ShaderPrintData, const FScreenPassTexture& OutputTexture)
+	static void InternalDrawZoom(FRDGBuilder& GraphBuilder, const FShaderPrintData& ShaderPrintData, const FScreenPassTexture& OutputTexture)
 	{
 		FGlobalShaderMap* GlobalShaderMap = GetGlobalShaderMap(GMaxRHIFeatureLevel);
 		FRDGTextureDesc Desc = OutputTexture.Texture->Desc;
@@ -1201,7 +1201,7 @@ namespace ShaderPrint
 		AddCopyTexturePass(GraphBuilder, OutZoomTexture, OutputTexture.Texture);
 	}
 
-	void InternalDrawView(FRDGBuilder& GraphBuilder, const FViewInfo& View, const FShaderPrintData& ShaderPrintData, const FScreenPassTexture& OutputTexture, const FScreenPassTexture& DepthTexture)
+	static void InternalDrawView(FRDGBuilder& GraphBuilder, const FViewInfo& View, const FShaderPrintData& ShaderPrintData, const FScreenPassTexture& OutputTexture, const FScreenPassTexture& DepthTexture)
 	{
 		if (!ensure(OutputTexture.IsValid()))
 		{
diff --git a/Engine/Source/Runtime/Renderer/Private/Substrate/SubstrateVisualize.cpp b/Engine/Source/Runtime/Renderer/Private/Substrate/SubstrateVisualize.cpp
--- a/Engine/Source/Runtime/Renderer/Private/Substrate/SubstrateVisualize.cpp
+++ b/Engine/Source/Runtime/Renderer/Private/Substrate/SubstrateVisualize.cpp
@@ -22,10 +22,17 @@ void AddSubstrateInternalClassificationTilePass(
 	ESubstrateTileType TileMaterialType,
 	const bool bDebug);
 
-static bool SubstrateDebugVisualizationCanRunOnPlatform(EShaderPlatform Platform)
+static bool IsSubstrateDebugVisualizationSupported(EShaderPlatform InPlatform, bool bIsEditorOnly, EShaderPermutationFlags Flags)
 {
-	// On some consoles, this ALU heavy shader (and with optimisation disables for the sake of low compilation time) would spill registers. So only keep it for the editor.
-	return GetMaxSupportedFeatureLevel(Platform) >= ERHIFeatureLevel::SM5 && IsPCPlatform(Platform);
+	return 
+		Substrate::IsSubstrateEnabled() && 
+		GetMaxSupportedFeatureLevel(InPlatform) >= ERHIFeatureLevel::SM5 && 
+		(bIsEditorOnly ? (IsPCPlatform(InPlatform) || EnumHasAllFlags(Flags, EShaderPermutationFlags::HasEditorOnlyData)) : true);
+}
+
+static bool SubstrateDebugVisualizationCanRunOnPlatform(EShaderPlatform InPlatform)
+{
+	return IsSubstrateDebugVisualizationSupported(InPlatform, false, EShaderPermutationFlags::None);
 }
 
 class FMaterialPrintInfoCS : public FGlobalShader
@@ -37,6 +44,7 @@ class FMaterialPrintInfoCS : public FGlobalShader
 
 	BEGIN_SHADER_PARAMETER_STRUCT(FParameters, )
 		SHADER_PARAMETER(uint32, ClosureIndex)
+		SHADER_PARAMETER(uint32, bOverrideCursorPosition)
 		SHADER_PARAMETER_STRUCT_REF(FViewUniformShaderParameters, ViewUniformBuffer)
 		SHADER_PARAMETER_RDG_UNIFORM_BUFFER(FSubstrateGlobalUniformParameters, Substrate)
 		SHADER_PARAMETER_STRUCT_INCLUDE(FSceneTextureParameters, SceneTextures)
@@ -44,6 +52,11 @@ class FMaterialPrintInfoCS : public FGlobalShader
 		SHADER_PARAMETER_RDG_BUFFER_UAV(RWBuffer<uint>, RWPositionOffsetBuffer)
 	END_SHADER_PARAMETER_STRUCT()
 
+	static bool IsSupported(EShaderPlatform InPlatform, EShaderPermutationFlags InFlags=EShaderPermutationFlags::None)
+	{
+		return IsSubstrateDebugVisualizationSupported(InPlatform, false, InFlags);
+	}
+
 	static FPermutationDomain RemapPermutation(FPermutationDomain PermutationVector)
 	{
 		return PermutationVector;
@@ -51,7 +64,7 @@ class FMaterialPrintInfoCS : public FGlobalShader
 
 	static bool ShouldCompilePermutation(const FGlobalShaderPermutationParameters& Parameters)
 	{
-		return Substrate::IsSubstrateEnabled() && SubstrateDebugVisualizationCanRunOnPlatform(Parameters.Platform) && EnumHasAllFlags(Parameters.Flags, EShaderPermutationFlags::HasEditorOnlyData);
+		return IsSupported(Parameters.Platform, Parameters.Flags);
 	}
 
 	static void ModifyCompilationEnvironment(const FGlobalShaderPermutationParameters& Parameters, FShaderCompilerEnvironment& OutEnvironment)
@@ -74,6 +87,7 @@ class FVisualizeMaterialCountPS : public FGlobalShader
 
 	BEGIN_SHADER_PARAMETER_STRUCT(FParameters, )
 		SHADER_PARAMETER(uint32, ViewMode)
+		SHADER_PARAMETER(uint32, bOverrideCursorPosition)
 		SHADER_PARAMETER_STRUCT_REF(FViewUniformShaderParameters, ViewUniformBuffer)
 		SHADER_PARAMETER_RDG_UNIFORM_BUFFER(FSubstrateGlobalUniformParameters, Substrate)
 		SHADER_PARAMETER_STRUCT_INCLUDE(FSceneTextureParameters, SceneTextures)
@@ -81,6 +95,11 @@ class FVisualizeMaterialCountPS : public FGlobalShader
 		RENDER_TARGET_BINDING_SLOTS()
 	END_SHADER_PARAMETER_STRUCT()
 
+	static bool IsSupported(EShaderPlatform InPlatform, EShaderPermutationFlags InFlags=EShaderPermutationFlags::None)
+	{
+		return IsSubstrateDebugVisualizationSupported(InPlatform, false, InFlags);
+	}
+
 	static FPermutationDomain RemapPermutation(FPermutationDomain PermutationVector)
 	{
 		return PermutationVector;
@@ -88,7 +107,7 @@ class FVisualizeMaterialCountPS : public FGlobalShader
 
 	static bool ShouldCompilePermutation(const FGlobalShaderPermutationParameters& Parameters)
 	{
-		return Substrate::IsSubstrateEnabled() && SubstrateDebugVisualizationCanRunOnPlatform(Parameters.Platform) && EnumHasAllFlags(Parameters.Flags, EShaderPermutationFlags::HasEditorOnlyData);
+		return IsSupported(Parameters.Platform, Parameters.Flags);
 	}
 
 	static void ModifyCompilationEnvironment(const FGlobalShaderPermutationParameters& Parameters, FShaderCompilerEnvironment& OutEnvironment)
@@ -134,14 +153,19 @@ class FSubstrateSystemInfoCS : public FGlobalShader
 		SHADER_PARAMETER_STRUCT_INCLUDE(ShaderPrint::FShaderParameters, ShaderPrintParameters)
 	END_SHADER_PARAMETER_STRUCT()
 
-		static FPermutationDomain RemapPermutation(FPermutationDomain PermutationVector)
+	static bool IsSupported(EShaderPlatform InPlatform, EShaderPermutationFlags InFlags=EShaderPermutationFlags::None)
+	{
+		return IsSubstrateDebugVisualizationSupported(InPlatform, false, InFlags);
+	}
+
+	static FPermutationDomain RemapPermutation(FPermutationDomain PermutationVector)
 	{
 		return PermutationVector;
 	}
 
 	static bool ShouldCompilePermutation(const FGlobalShaderPermutationParameters& Parameters)
 	{
-		return Substrate::IsSubstrateEnabled() && SubstrateDebugVisualizationCanRunOnPlatform(Parameters.Platform) && EnumHasAllFlags(Parameters.Flags, EShaderPermutationFlags::HasEditorOnlyData);
+		return IsSupported(Parameters.Platform, Parameters.Flags);
 	}
 
 	static void ModifyCompilationEnvironment(const FGlobalShaderPermutationParameters& Parameters, FShaderCompilerEnvironment& OutEnvironment)
@@ -163,11 +187,17 @@ class FMaterialDebugSubstrateTreeCS : public FGlobalShader
 	using FPermutationDomain = TShaderPermutationDomain<>;
 
 	BEGIN_SHADER_PARAMETER_STRUCT(FParameters, )
+		SHADER_PARAMETER(uint32, bOverrideCursorPosition)
 		SHADER_PARAMETER_STRUCT_REF(FViewUniformShaderParameters, ViewUniformBuffer)
 		SHADER_PARAMETER_RDG_UNIFORM_BUFFER(FSubstrateGlobalUniformParameters, Substrate)
 		SHADER_PARAMETER_STRUCT_INCLUDE(ShaderPrint::FShaderParameters, ShaderPrintParameters)
 	END_SHADER_PARAMETER_STRUCT()
 
+	static bool IsSupported(EShaderPlatform InPlatform, EShaderPermutationFlags InFlags=EShaderPermutationFlags::None)
+	{
+		return IsSubstrateDebugVisualizationSupported(InPlatform, true, InFlags);
+	}
+
 	static FPermutationDomain RemapPermutation(FPermutationDomain PermutationVector)
 	{
 		return PermutationVector;
@@ -175,7 +205,7 @@ class FMaterialDebugSubstrateTreeCS : public FGlobalShader
 
 	static bool ShouldCompilePermutation(const FGlobalShaderPermutationParameters& Parameters)
 	{
-		return Substrate::IsSubstrateEnabled() && SubstrateDebugVisualizationCanRunOnPlatform(Parameters.Platform) && EnumHasAllFlags(Parameters.Flags, EShaderPermutationFlags::HasEditorOnlyData);
+		return IsSupported(Parameters.Platform, Parameters.Flags);
 	}
 
 	static void ModifyCompilationEnvironment(const FGlobalShaderPermutationParameters& Parameters, FShaderCompilerEnvironment& OutEnvironment)
@@ -197,6 +227,7 @@ class FMaterialDebugSubstrateTreePS : public FGlobalShader
 	using FPermutationDomain = TShaderPermutationDomain<>;
 
 	BEGIN_SHADER_PARAMETER_STRUCT(FParameters, )
+		SHADER_PARAMETER(uint32, bOverrideCursorPosition)
 		SHADER_PARAMETER_STRUCT_REF(FViewUniformShaderParameters, ViewUniformBuffer)
 		SHADER_PARAMETER_RDG_UNIFORM_BUFFER(FSubstrateGlobalUniformParameters, Substrate)
 		SHADER_PARAMETER_RDG_UNIFORM_BUFFER(FReflectionUniformParameters, ReflectionStruct)
@@ -206,6 +237,11 @@ class FMaterialDebugSubstrateTreePS : public FGlobalShader
 		RENDER_TARGET_BINDING_SLOTS()
 	END_SHADER_PARAMETER_STRUCT()
 
+	static bool IsSupported(EShaderPlatform InPlatform, EShaderPermutationFlags InFlags=EShaderPermutationFlags::None)
+	{
+		return IsSubstrateDebugVisualizationSupported(InPlatform, true, InFlags);
+	}
+	
 	static FPermutationDomain RemapPermutation(FPermutationDomain PermutationVector)
 	{
 		return PermutationVector;
@@ -213,7 +249,7 @@ class FMaterialDebugSubstrateTreePS : public FGlobalShader
 
 	static bool ShouldCompilePermutation(const FGlobalShaderPermutationParameters& Parameters)
 	{
-		return Substrate::IsSubstrateEnabled() && SubstrateDebugVisualizationCanRunOnPlatform(Parameters.Platform) && EnumHasAllFlags(Parameters.Flags, EShaderPermutationFlags::HasEditorOnlyData);
+		return IsSupported(Parameters.Platform, Parameters.Flags);
 	}
 
 	static void ModifyCompilationEnvironment(const FGlobalShaderPermutationParameters& Parameters, FShaderCompilerEnvironment& OutEnvironment)
@@ -229,6 +265,8 @@ IMPLEMENT_GLOBAL_SHADER(FMaterialDebugSubstrateTreePS, "/Engine/Private/Substrat
 
 static void AddVisualizeMaterialPropertiesPasses(FRDGBuilder& GraphBuilder, const FViewInfo& View, FScreenPassTexture& ScreenPassSceneColor, EShaderPlatform Platform)
 {
+	if (!FMaterialPrintInfoCS::IsSupported(Platform)) return;
+
 	// Force ShaderPrint on.
 	ShaderPrint::SetEnabled(true);
 	ShaderPrint::RequestSpaceForLines(1024);
@@ -243,6 +281,7 @@ static void AddVisualizeMaterialPropertiesPasses(FRDGBuilder& GraphBuilder, cons
 	{
 		FMaterialPrintInfoCS::FParameters* PassParameters = GraphBuilder.AllocParameters<FMaterialPrintInfoCS::FParameters>();
 		PassParameters->ClosureIndex = ClosureIndex;
+		PassParameters->bOverrideCursorPosition = WITH_EDITOR ? 0u : 1u;
 		PassParameters->RWPositionOffsetBuffer = PrintOffsetBufferUAV;
 		PassParameters->ViewUniformBuffer = View.ViewUniformBuffer;
 		PassParameters->Substrate = Substrate::BindSubstrateGlobalUniformParameters(View);
@@ -256,6 +295,8 @@ static void AddVisualizeMaterialPropertiesPasses(FRDGBuilder& GraphBuilder, cons
 
 static void AddVisualizeMaterialCountPasses(FRDGBuilder & GraphBuilder, const FViewInfo & View, FScreenPassTexture & ScreenPassSceneColor, EShaderPlatform Platform, uint32 ViewMode)
 {
+	if (!FVisualizeMaterialCountPS::IsSupported(Platform)) return;
+
 	ShaderPrint::SetEnabled(true);
 	ShaderPrint::RequestSpaceForLines(1024);
 	ShaderPrint::RequestSpaceForCharacters(1024);
@@ -266,6 +307,7 @@ static void AddVisualizeMaterialCountPasses(FRDGBuilder & GraphBuilder, const FV
 	FVisualizeMaterialCountPS::FParameters* PassParameters = GraphBuilder.AllocParameters<FVisualizeMaterialCountPS::FParameters>();
 	PassParameters->ViewUniformBuffer = View.ViewUniformBuffer;
 	PassParameters->ViewMode = FMath::Clamp(ViewMode, 2, 3);
+	PassParameters->bOverrideCursorPosition = WITH_EDITOR ? 0u : 1u;
 	PassParameters->Substrate = Substrate::BindSubstrateGlobalUniformParameters(View);
 	PassParameters->SceneTextures = GetSceneTextureParameters(GraphBuilder, View);
 	PassParameters->RenderTargets[0] = FRenderTargetBinding(SceneColorTexture, ERenderTargetLoadAction::ELoad);
@@ -285,6 +327,8 @@ bool Is8bitTileCoordEnabled();
 
 static void AddVisualizeSystemInfoPasses(FRDGBuilder& GraphBuilder, const FViewInfo& View, FScreenPassTexture& ScreenPassSceneColor, EShaderPlatform Platform)
 {
+	if (!FSubstrateSystemInfoCS::IsSupported(Platform)) return;
+
 	// Force ShaderPrint on.
 	ShaderPrint::SetEnabled(true);
 	ShaderPrint::RequestSpaceForLines(1024);
@@ -322,7 +366,9 @@ static void AddVisualizeSystemInfoPasses(FRDGBuilder& GraphBuilder, const FViewI
 // Draw each material layer independently
 static void AddVisualizeAdvancedMaterialPasses(FRDGBuilder& GraphBuilder, const FViewInfo& View, FScreenPassTexture& ScreenPassSceneColor, EShaderPlatform Platform)
 {
-	if (!IsAdvancedVisualizationEnabled())
+	if (!IsAdvancedVisualizationEnabled() ||
+		!FMaterialDebugSubstrateTreeCS::IsSupported(Platform) || 
+		!FMaterialDebugSubstrateTreePS::IsSupported(Platform))
 	{
 		return;
 	}
@@ -337,6 +383,7 @@ static void AddVisualizeAdvancedMaterialPasses(FRDGBuilder& GraphBuilder, const
 		FMaterialDebugSubstrateTreeCS::FParameters* PassParameters = GraphBuilder.AllocParameters<FMaterialDebugSubstrateTreeCS::FParameters>();
 		PassParameters->ViewUniformBuffer = View.ViewUniformBuffer;
 		PassParameters->Substrate = Substrate::BindSubstrateGlobalUniformParameters(View);
+		PassParameters->bOverrideCursorPosition = WITH_EDITOR ? 0u : 1u;
 		ShaderPrint::SetParameters(GraphBuilder, View.ShaderPrintData, PassParameters->ShaderPrintParameters);
 
 		TShaderMapRef<FMaterialDebugSubstrateTreeCS> ComputeShader(View.ShaderMap);
@@ -347,6 +394,7 @@ static void AddVisualizeAdvancedMaterialPasses(FRDGBuilder& GraphBuilder, const
 		FMaterialDebugSubstrateTreePS::FParameters* PassParameters = GraphBuilder.AllocParameters<FMaterialDebugSubstrateTreePS::FParameters>();
 		PassParameters->ViewUniformBuffer = View.ViewUniformBuffer;
 		PassParameters->Substrate = Substrate::BindSubstrateGlobalUniformParameters(View);
+		PassParameters->bOverrideCursorPosition = WITH_EDITOR ? 0u : 1u;
 		PassParameters->ReflectionStruct = CreateReflectionUniformBuffer(GraphBuilder, View);
 		PassParameters->ReflectionCapture = View.ReflectionCaptureUniformBuffer;
 		PassParameters->ForwardLightData = View.ForwardLightingResources.ForwardLightUniformBuffer;
@@ -371,9 +419,7 @@ static FSubstrateViewMode GetSubstrateVisualizeMode(const FViewInfo & View)
 	FSubstrateViewMode Out = FSubstrateViewMode::None;
 	if (IsSubstrateEnabled() && SubstrateDebugVisualizationCanRunOnPlatform(View.GetShaderPlatform()))
 	{
-		// Variable defined in SubstrateVisualizationData.h/.cpp
-		static const auto CVar = IConsoleManager::Get().FindConsoleVariable(FSubstrateVisualizationData::GetVisualizeConsoleCommandName());
-		const uint32 ViewMode = CVar && CVar->AsVariableInt() ? CVar->AsVariableInt()->GetValueOnRenderThread() : 0;
+		const uint32 ViewMode = FSubstrateVisualizationData::GetViewMode();
 		switch (ViewMode)
 		{
 			case 1: return FSubstrateViewMode::MaterialProperties;
