Slack memory savings for FCachedMeshDrawCommandInfo.  Do an explicit "Shrink" instead of "bAllowShrinking" which shrinks conditionally.  Arrays with 32 or 64 reserved elements weren't ever being shrunk, leading to a few MB worth of slack.  Costs around 1.2% in perf on FPrimitiveSceneInfo::CacheMeshDrawCommands, or around 0.003 ms frame time on the lowest spec platform.

#rnx

diff --git a/Engine/Source/Runtime/Renderer/Private/PrimitiveSceneInfo.cpp b/Engine/Source/Runtime/Renderer/Private/PrimitiveSceneInfo.cpp
--- a/Engine/Source/Runtime/Renderer/Private/PrimitiveSceneInfo.cpp
+++ b/Engine/Source/Runtime/Renderer/Private/PrimitiveSceneInfo.cpp
@@ -557,7 +557,9 @@ void FPrimitiveSceneInfo::CacheMeshDrawCommands(FScene* Scene, TArrayView<FPrimi
 					PrefixSum = NewPrefixSum;
 				}
 			}
-			SceneInfo->StaticMeshCommandInfos.SetNum(PrefixSum, true);
+
+			SceneInfo->StaticMeshCommandInfos.SetNum(PrefixSum, false);		// bAllowShrinking == false, since we explicitly call Shrink on the next line
+			SceneInfo->StaticMeshCommandInfos.Shrink();
 		}
 	};
 
