PSO precache fixes:
- Precache PSOs for materials used by the WidgetComponent
- Just-in-time PSO precaching for PSOs used by canvas rendering - these force use the base pass rendering with specific color format on RT0
- Also precache PSOs with alpha channel on RT0 for base pass materials because these can be used for planar reflection or scene capture (water and sky also use base pass shaders)
- Also precache custom depth PSOs when depth rendering is enabled (custom depth can be enabled without precaching PSOs again - would probably be too late as well if done just in time)
- Change default CVar value of r.PSOPrecache.LightMapPolicyMode to 1 to only precache PSOs for LMP_NO_LIGHTMAP
- Change default DepthStencil mode from DepthWrite_StencilWrite to DepthRead_StencilWrite which is used by default
- ParticleComponent (Niagara and Cascade) precache custom depth passes as well and also precache for both regular and culling and reverse culling because PSO precaching is done at asset level and don't know if component will have negative determinant
- Use spline params of SplineMeshComponent to know if reverse culling is needed
- SplineMeshVertexFactory now used GPU scene and doesn't need special implementation for GetPSOPrecacheVertexFetchElements anymore (can use parent one)

diff --git a/Engine/Source/Runtime/Engine/Private/Components/PrimitiveComponent.cpp b/Engine/Source/Runtime/Engine/Private/Components/PrimitiveComponent.cpp
--- a/Engine/Source/Runtime/Engine/Private/Components/PrimitiveComponent.cpp
+++ b/Engine/Source/Runtime/Engine/Private/Components/PrimitiveComponent.cpp
@@ -4614,7 +4614,7 @@ void UPrimitiveComponent::SetupPrecachePSOParams(FPSOPrecacheParams& Params)
 	Params.bCastShadow = CastShadow;
 	// Custom depth can be toggled at runtime with PSO precache call so assume it might be needed when depth pass is needed
 	// Ideally precache those with lower priority and don't wait on these (UE-174426)
-	Params.bRenderCustomDepth = bRenderCustomDepth;// bRenderInDepthPass;
+	Params.bRenderCustomDepth = bRenderInDepthPass;
 	Params.bCastShadowAsTwoSided = bCastShadowAsTwoSided;
 	Params.SetMobility(Mobility);	
 	Params.SetStencilWriteMask(FRendererStencilMaskEvaluation::ToStencilMask(CustomDepthStencilWriteMask));
diff --git a/Engine/Source/Runtime/Engine/Private/Components/SplineMeshComponent.cpp b/Engine/Source/Runtime/Engine/Private/Components/SplineMeshComponent.cpp
--- a/Engine/Source/Runtime/Engine/Private/Components/SplineMeshComponent.cpp
+++ b/Engine/Source/Runtime/Engine/Private/Components/SplineMeshComponent.cpp
@@ -848,7 +848,10 @@ void USplineMeshComponent::CollectPSOPrecacheData(const FPSOPrecacheParams& Base
 		FLocalVertexFactory::GetVertexElements(GMaxRHIFeatureLevel, EVertexInputStreamType::Default, bSupportsManualVertexFetch, Data, Elements);
 	};
 
-	CollectPSOPrecacheDataImpl(VertexFactoryType, BasePrecachePSOParams, SMC_GetElements, OutParams);
+	FPSOPrecacheParams SplineMeshPSOParams = BasePrecachePSOParams;
+	SplineMeshPSOParams.bReverseCulling ^= (SplineParams.StartScale.X < 0) ^ (SplineParams.StartScale.Y < 0);
+
+	CollectPSOPrecacheDataImpl(VertexFactoryType, SplineMeshPSOParams, SMC_GetElements, OutParams);
 }
 
 FPrimitiveSceneProxy* USplineMeshComponent::CreateStaticMeshSceneProxy(Nanite::FMaterialAudit& NaniteMaterials, bool bCreateNanite)
diff --git a/Engine/Source/Runtime/Engine/Private/Components/StaticMeshComponent.cpp b/Engine/Source/Runtime/Engine/Private/Components/StaticMeshComponent.cpp
--- a/Engine/Source/Runtime/Engine/Private/Components/StaticMeshComponent.cpp
+++ b/Engine/Source/Runtime/Engine/Private/Components/StaticMeshComponent.cpp
@@ -1639,7 +1639,7 @@ void UStaticMeshComponent::CollectPSOPrecacheDataImpl(
 
 	FPSOPrecacheParams PrecachePSOParams = BasePrecachePSOParams;
 	PrecachePSOParams.bCastShadow = bAnySectionCastsShadows;
-	PrecachePSOParams.bReverseCulling = bReverseCulling != bIsLocalToWorldDeterminantNegative;
+	PrecachePSOParams.bReverseCulling = PrecachePSOParams.bReverseCulling || bReverseCulling != bIsLocalToWorldDeterminantNegative;
 	PrecachePSOParams.bForceLODModel = ForcedLodModel > 0;
 
 	for (FPSOPrecacheVertexFactoryDataPerMaterialIndex& VFsPerMaterial : VFTypesPerMaterialIndex)
diff --git a/Engine/Source/Runtime/Engine/Private/Materials/MaterialShader.cpp b/Engine/Source/Runtime/Engine/Private/Materials/MaterialShader.cpp
--- a/Engine/Source/Runtime/Engine/Private/Materials/MaterialShader.cpp
+++ b/Engine/Source/Runtime/Engine/Private/Materials/MaterialShader.cpp
@@ -2758,8 +2758,6 @@ FPSOPrecacheDataArray FMaterialShaderMap::CollectPSOPrecacheData(const FMaterial
 			IPSOCollector* PSOCollector = CreateFunction(PrecacheParams.FeatureLevel);
 			if (PSOCollector != nullptr)
 			{
-				//check(PSOCollector->PSOCollectorIndex == Index);
-
 #if PSO_PRECACHING_VALIDATE
 				ConditionalBreakOnPSOPrecacheMaterial(*PrecacheParams.Material, Index);
 #endif // PSO_PRECACHING_VALIDATE
diff --git a/Engine/Source/Runtime/Engine/Private/PSOPrecacheValidation.cpp b/Engine/Source/Runtime/Engine/Private/PSOPrecacheValidation.cpp
--- a/Engine/Source/Runtime/Engine/Private/PSOPrecacheValidation.cpp
+++ b/Engine/Source/Runtime/Engine/Private/PSOPrecacheValidation.cpp
@@ -767,9 +767,11 @@ void LogMinimalPSOStateMissInfo(
 			CompareStateAndLogChanges(TEXT("DepthBounds"), uint32(PSOPrecacheData.GraphicsPSOInitializer.bDepthBounds), uint32(Initializer.bDepthBounds), StringBuilder);
 			CompareStateAndLogChanges(TEXT("MultiViewCount"), uint32(PSOPrecacheData.GraphicsPSOInitializer.MultiViewCount), uint32(Initializer.MultiViewCount), StringBuilder);
 			CompareStateAndLogChanges(TEXT("HasFragmentDensityAttachment"), uint32(PSOPrecacheData.GraphicsPSOInitializer.bHasFragmentDensityAttachment), uint32(Initializer.bHasFragmentDensityAttachment), StringBuilder);
-			CompareStateAndLogChanges(TEXT("AllowVariableRateShading"), uint32(PSOPrecacheData.GraphicsPSOInitializer.bAllowVariableRateShading), uint32(Initializer.bAllowVariableRateShading), StringBuilder);
 			CompareStateAndLogChanges(TEXT("DrawShadingRate"), uint32(PSOPrecacheData.GraphicsPSOInitializer.ShadingRate), uint32(Initializer.ShadingRate), StringBuilder);
 			CompareStateAndLogChanges(TEXT("PrimitiveType"), uint32(PSOPrecacheData.GraphicsPSOInitializer.PrimitiveType), uint32(Initializer.PrimitiveType), StringBuilder);
+
+			// Ignored for now since it doesn't add to any real PSO and can give different state
+			//CompareStateAndLogChanges(TEXT("AllowVariableRateShading"), uint32(PSOPrecacheData.GraphicsPSOInitializer.bAllowVariableRateShading), uint32(Initializer.bAllowVariableRateShading), StringBuilder);			
 		}
 	}
 	else
diff --git a/Engine/Source/Runtime/Engine/Private/Particles/ParticleComponents.cpp b/Engine/Source/Runtime/Engine/Private/Particles/ParticleComponents.cpp
--- a/Engine/Source/Runtime/Engine/Private/Particles/ParticleComponents.cpp
+++ b/Engine/Source/Runtime/Engine/Private/Particles/ParticleComponents.cpp
@@ -151,6 +151,7 @@ void UFXSystemAsset::LaunchPSOPrecaching(TArrayView<VFsPerMaterialData> VFsPerMa
 {
 	FPSOPrecacheParams PreCachePSOParams;
 	PreCachePSOParams.SetMobility(EComponentMobility::Movable);
+	PreCachePSOParams.bRenderCustomDepth = true;
 
 	FGraphEventArray PrecachePSOsEvents;
 	for (VFsPerMaterialData& VFsPerMaterial : VFsPerMaterials)
@@ -159,7 +160,15 @@ void UFXSystemAsset::LaunchPSOPrecaching(TArrayView<VFsPerMaterialData> VFsPerMa
 		{
 			PreCachePSOParams.PrimitiveType = (EPrimitiveType)VFsPerMaterial.PrimitiveType;
 			PreCachePSOParams.bDisableBackFaceCulling = VFsPerMaterial.bDisableBackfaceCulling;
-			PrecachePSOsEvents.Append(VFsPerMaterial.MaterialInterface->PrecachePSOs(VFsPerMaterial.VertexFactoryData, PreCachePSOParams,EPSOPrecachePriority::Medium, MaterialPSOPrecacheRequestIDs));
+			PreCachePSOParams.bReverseCulling = false;
+			PrecachePSOsEvents.Append(VFsPerMaterial.MaterialInterface->PrecachePSOs(VFsPerMaterial.VertexFactoryData, PreCachePSOParams, EPSOPrecachePriority::Medium, MaterialPSOPrecacheRequestIDs));
+
+			// Also precache with reverse culling if not two sided because we don't know of the component using the asset will have negative determinant
+			if (!PreCachePSOParams.bDisableBackFaceCulling)
+			{
+				PreCachePSOParams.bReverseCulling = true;
+				PrecachePSOsEvents.Append(VFsPerMaterial.MaterialInterface->PrecachePSOs(VFsPerMaterial.VertexFactoryData, PreCachePSOParams, EPSOPrecachePriority::Medium, MaterialPSOPrecacheRequestIDs));
+			}
 		}
 	}
 
diff --git a/Engine/Source/Runtime/Engine/Private/SceneTexturesConfig.cpp b/Engine/Source/Runtime/Engine/Private/SceneTexturesConfig.cpp
--- a/Engine/Source/Runtime/Engine/Private/SceneTexturesConfig.cpp
+++ b/Engine/Source/Runtime/Engine/Private/SceneTexturesConfig.cpp
@@ -130,7 +130,7 @@ static EPixelFormat GetSceneColorFormat(bool bRequiresAlphaChannel)
 	return Format;
 }
 
-static void GetSceneColorFormatAndCreateFlags(ERHIFeatureLevel::Type FeatureLevel, bool bRequiresAlphaChannel, ETextureCreateFlags ExtraSceneColorCreateFlags, uint32 NumSamples, bool bMemorylessMSAA, EPixelFormat& SceneColorFormat, ETextureCreateFlags& SceneColorCreateFlags)
+void GetSceneColorFormatAndCreateFlags(ERHIFeatureLevel::Type FeatureLevel, bool bRequiresAlphaChannel, ETextureCreateFlags ExtraSceneColorCreateFlags, uint32 NumSamples, bool bMemorylessMSAA, EPixelFormat& SceneColorFormat, ETextureCreateFlags& SceneColorCreateFlags)
 {
 	EShadingPath ShadingPath = GetFeatureLevelShadingPath(FeatureLevel);
 	switch (ShadingPath)
diff --git a/Engine/Source/Runtime/Engine/Private/SplineMeshSceneProxy.cpp b/Engine/Source/Runtime/Engine/Private/SplineMeshSceneProxy.cpp
--- a/Engine/Source/Runtime/Engine/Private/SplineMeshSceneProxy.cpp
+++ b/Engine/Source/Runtime/Engine/Private/SplineMeshSceneProxy.cpp
@@ -40,34 +40,7 @@ void FSplineMeshVertexFactory::ModifyCompilationEnvironment(const FVertexFactory
  */
 void FSplineMeshVertexFactory::GetPSOPrecacheVertexFetchElements(EVertexInputStreamType VertexInputStreamType, FVertexDeclarationElementList& Elements)
 {
-	if (VertexInputStreamType == EVertexInputStreamType::PositionOnly || VertexInputStreamType == EVertexInputStreamType::PositionAndNormalOnly)
-	{
-		FLocalVertexFactory::GetPSOPrecacheVertexFetchElements(VertexInputStreamType, Elements);
-	}
-	else
-	{
-		// Position
-		Elements.Add(FVertexElement(0, 0, VET_Float3, 0, 12, false));
-
-		// Normals
-		Elements.Add(FVertexElement(1, 0, VET_PackedNormal, 1, 0, false));
-		Elements.Add(FVertexElement(2, 0, VET_PackedNormal, 2, 0, false));
-
-		// Color
-		Elements.Add(FVertexElement(3, 0, VET_Color, 3, 0, false));
-
-		// Texcoords
-		Elements.Add(FVertexElement(4, 0, VET_Float4, 4, 0, false));
-		Elements.Add(FVertexElement(5, 0, VET_Float4, 5, 0, false));
-		Elements.Add(FVertexElement(6, 0, VET_Float4, 6, 0, false));
-		Elements.Add(FVertexElement(7, 0, VET_Float4, 7, 0, false));
-
-		// Lightmap coords
-		Elements.Add(FVertexElement(8, 0, VET_Float2, 15, 0, false));
-		
-		// Primitive ID
-		Elements.Add(FVertexElement(9, 0, VET_UInt, 13, 0, true));
-	}
+	FLocalVertexFactory::GetPSOPrecacheVertexFetchElements(VertexInputStreamType, Elements);
 }
 
 FSplineMeshSceneProxy::FSplineMeshSceneProxy(USplineMeshComponent* InComponent) :
diff --git a/Engine/Source/Runtime/Engine/Private/UserInterface/Canvas.cpp b/Engine/Source/Runtime/Engine/Private/UserInterface/Canvas.cpp
--- a/Engine/Source/Runtime/Engine/Private/UserInterface/Canvas.cpp
+++ b/Engine/Source/Runtime/Engine/Private/UserInterface/Canvas.cpp
@@ -2067,6 +2067,12 @@ void UCanvas::K2_DrawMaterial(UMaterialInterface* RenderMaterial, FVector2D Scre
 		// This is a user-facing function, so we'd rather make sure that shaders are ready by the time we render, in order to ensure we don't draw with a fallback material :
 		RenderMaterial->EnsureIsComplete();
 
+		// Should be moved earlier
+		FPSOPrecacheParams PSOPrecacheParams;
+		PSOPrecacheParams.bCanvasMaterial = true;
+		PSOPrecacheParams.BasePassPixelFormat = (Canvas->GetRenderTarget() && Canvas->GetRenderTarget()->GetRenderTargetTexture()) ? Canvas->GetRenderTarget()->GetRenderTargetTexture()->GetDesc().Format : PF_Unknown;
+		RenderMaterial->PrecachePSOs(&FLocalVertexFactory::StaticType, PSOPrecacheParams);
+
 		FCanvasTileItem TileItem(ScreenPosition, RenderMaterial->GetRenderProxy(), ScreenSize, CoordinatePosition, CoordinatePosition + CoordinateSize);
 		TileItem.Rotation = FRotator(0, Rotation, 0);
 		TileItem.PivotPoint = PivotPoint;
diff --git a/Engine/Source/Runtime/Engine/Public/PSOPrecache.h b/Engine/Source/Runtime/Engine/Public/PSOPrecache.h
--- a/Engine/Source/Runtime/Engine/Public/PSOPrecache.h
+++ b/Engine/Source/Runtime/Engine/Public/PSOPrecache.h
@@ -44,6 +44,7 @@ struct FPSOPrecacheParams
 		Mobility = (uint8)EComponentMobility::Static;
 		bAnyMaterialHasWorldPositionOffset = false;
 		StencilWriteMask = (uint8)EStencilMask::SM_Default;
+		BasePassPixelFormat = (uint16)PF_Unknown;
 		Unused = 0;
 	}
 
@@ -59,7 +60,7 @@ struct FPSOPrecacheParams
 
 	friend uint32 GetTypeHash(const FPSOPrecacheParams& Params)
 	{
-		return uint32(Params.Data);
+		return GetTypeHash(Params.Data);
 	}
 
 	void SetMobility(EComponentMobility::Type InMobility)
@@ -87,35 +88,48 @@ struct FPSOPrecacheParams
 		return TEnumAsByte<EStencilMask>((EStencilMask)StencilWriteMask);
 	}
 
+	void SetBassPixelFormat(EPixelFormat InBasePassPixelFormat)
+	{
+		BasePassPixelFormat = (uint16)InBasePassPixelFormat;
+	}
+
+	EPixelFormat GetBassPixelFormat() const
+	{
+		return (EPixelFormat)BasePassPixelFormat;
+	}
+
 	union
 	{
 		struct
 		{
-			uint32 PrimitiveType : 6;
-
-			uint32 bDefaultMaterial : 1;
-			uint32 bCanvasMaterial : 1;
-
-			uint32 bRenderInMainPass : 1;
-			uint32 bRenderInDepthPass : 1;
-			uint32 bStaticLighting : 1;
-			uint32 bCastShadow : 1;
-			uint32 bRenderCustomDepth : 1;
-
-			uint32 bAffectDynamicIndirectLighting : 1;
-			uint32 bReverseCulling : 1;
-			uint32 bDisableBackFaceCulling : 1;
-			uint32 bCastShadowAsTwoSided : 1;
-			uint32 bForceLODModel : 1;
-
-			uint32 Mobility : 4;
-			uint32 bAnyMaterialHasWorldPositionOffset : 1;
-			uint32 StencilWriteMask : 4;
-
-			uint32 Unused : 5;
+			uint64 PrimitiveType : 6;
+			
+			uint64 bDefaultMaterial : 1;
+			uint64 bCanvasMaterial : 1;
+			
+			uint64 bRenderInMainPass : 1;
+			uint64 bRenderInDepthPass : 1;
+			uint64 bStaticLighting : 1;
+			uint64 bCastShadow : 1;
+			uint64 bRenderCustomDepth : 1;
+			
+			uint64 bAffectDynamicIndirectLighting : 1;
+			uint64 bReverseCulling : 1;
+			uint64 bDisableBackFaceCulling : 1;
+			uint64 bCastShadowAsTwoSided : 1;
+			uint64 bForceLODModel : 1;
+			
+			uint64 Mobility : 4;
+			uint64 bAnyMaterialHasWorldPositionOffset : 1;
+			uint64 StencilWriteMask : 4;
+
+			uint64 BasePassPixelFormat : 16;
+
+			uint64 Unused : 21;
 		};
-		uint32 Data;
+		uint64 Data;
 	};
+
 };
 
 // Unique ID to find the FVertexDeclarationElementList - these can be shared
@@ -214,7 +228,7 @@ typedef TArray<FPSOPrecacheRequestResult, TInlineAllocator<4> > FPSOPrecacheRequ
 struct FMaterialPSOPrecacheParams
 {
 	ERHIFeatureLevel::Type FeatureLevel = ERHIFeatureLevel::Num;
-	const FMaterial* Material = nullptr;
+	FMaterial* Material = nullptr;
 	FPSOPrecacheVertexFactoryData VertexFactoryData;
 	FPSOPrecacheParams PrecachePSOParams;
 
diff --git a/Engine/Source/Runtime/Engine/Public/SceneTexturesConfig.h b/Engine/Source/Runtime/Engine/Public/SceneTexturesConfig.h
--- a/Engine/Source/Runtime/Engine/Public/SceneTexturesConfig.h
+++ b/Engine/Source/Runtime/Engine/Public/SceneTexturesConfig.h
@@ -73,6 +73,8 @@ extern ENGINE_API FSceneTextureShaderParameters GetSceneTextureShaderParameters(
 
 extern ENGINE_API FSceneTextureShaderParameters GetSceneTextureShaderParameters(TRDGUniformBufferRef<FMobileSceneTextureUniformParameters> UniformBuffer);
 
+extern ENGINE_API void GetSceneColorFormatAndCreateFlags(ERHIFeatureLevel::Type FeatureLevel, bool bRequiresAlphaChannel, ETextureCreateFlags ExtraSceneColorCreateFlags, uint32 NumSamples, bool bMemorylessMSAA, EPixelFormat& SceneColorFormat, ETextureCreateFlags& SceneColorCreateFlags);
+
 enum class ESceneTextureExtracts : uint32
 {
 	/** No textures are extracted from the render graph after execution. */
diff --git a/Engine/Source/Runtime/Renderer/Private/BasePassRendering.cpp b/Engine/Source/Runtime/Renderer/Private/BasePassRendering.cpp
--- a/Engine/Source/Runtime/Renderer/Private/BasePassRendering.cpp
+++ b/Engine/Source/Runtime/Renderer/Private/BasePassRendering.cpp
@@ -96,10 +96,10 @@ static TAutoConsoleVariable<int32> CVarClearGBufferDBeforeBasePass(
 
 static TAutoConsoleVariable<int32> CVarPSOPrecacheLightMapPolicyMode(
 	TEXT("r.PSOPrecache.LightMapPolicyMode"),
-	0,
+	1,
 	TEXT("Defines which light map policies should be checked during PSO precaching of the base pass.\n") \
-	TEXT(" 0: All possible LMP will be checked (default).\n") \
-	TEXT(" 1: Only LMP_NO_LIGHTMAP will be precached.\n"),
+	TEXT(" 0: All possible LMP will be checked.\n") \
+	TEXT(" 1: Only LMP_NO_LIGHTMAP will be precached (default).\n"),
 	ECVF_ReadOnly
 );
 
@@ -1619,6 +1619,7 @@ template<typename LightMapPolicyType>
 void FBasePassMeshProcessor::CollectPSOInitializersForLMPolicy(
 	const FSceneTexturesConfig& SceneTexturesConfig,
 	const FPSOPrecacheVertexFactoryData& VertexFactoryData,
+	const FPSOPrecacheParams& PreCacheParams,
 	const FMaterial& RESTRICT MaterialResource,
 	FMaterialShadingModelField ShadingModels,
 	const bool bRenderSkylight,
@@ -1732,10 +1733,20 @@ void FBasePassMeshProcessor::CollectPSOInitializersForLMPolicy(
 	}
 	else
 	{
-		// Regular base pass with gbuffer bindings
-		SetupGBufferRenderTargetInfo(SceneTexturesConfig, RenderTargetsInfo, true /*bSetupDepthStencil*/);
+		if (PreCacheParams.bCanvasMaterial)
+		{
+			DrawRenderState.SetDepthStencilState(TStaticDepthStencilState<false, CF_Always>::GetRHI());
+			SetTranslucentRenderState(DrawRenderState, MaterialResource, GShaderPlatformForFeatureLevel[FeatureLevel], ETranslucencyPass::TPT_AllTranslucency);
+			AddRenderTargetInfo(PreCacheParams.GetBassPixelFormat() != PF_Unknown ? PreCacheParams.GetBassPixelFormat() : PF_B8G8R8A8, TexCreate_RenderTargetable | TexCreate_ShaderResource, RenderTargetsInfo);
+		}
+		else
+		{
+			// Regular base pass with gbuffer bindings
+			SetupGBufferRenderTargetInfo(SceneTexturesConfig, RenderTargetsInfo, true /*bSetupDepthStencil*/);
+		}
 
-		AddGraphicsPipelineStateInitializer(
+		AddBasePassGraphicsPipelineStateInitializer(
+			FeatureLevel,
 			VertexFactoryData,
 			MaterialResource,
 			DrawRenderState,
@@ -1744,8 +1755,8 @@ void FBasePassMeshProcessor::CollectPSOInitializersForLMPolicy(
 			MeshFillMode,
 			MeshCullMode,
 			PrimitiveType,
-			EMeshPassFeatures::Default,
-			true /*bRequired*/,
+			true /*bPrecacheAlphaColorChannel*/,
+			PSOCollectorIndex,
 			PSOInitializers);
 	}	
 }
@@ -2334,6 +2345,7 @@ void FBasePassMeshProcessor::CollectPSOInitializers(const FSceneTexturesConfig&
 	bool bShouldDraw = ShouldDraw(Material)
 		&& ShouldIncludeDomainInMeshPass(Material.GetMaterialDomain())
 		&& ShouldIncludeMaterialInDefaultOpaquePass(Material);
+	bShouldDraw = bShouldDraw || (PreCacheParams.bCanvasMaterial && MeshPassType == EMeshPass::BasePass);
 	if (!bShouldDraw || !PreCacheParams.bRenderInMainPass)
 	{
 		return;
@@ -2355,15 +2367,16 @@ void FBasePassMeshProcessor::CollectPSOInitializers(const FSceneTexturesConfig&
 
 	{
 		bool bRenderSkyLight = true; // generate for both skylight enabled/disabled? Or can this be known already at this point?
-		CollectPSOInitializersForSkyLight(SceneTexturesConfig, VertexFactoryData, Material, bRenderSkyLight, bDitheredLODTransition, MeshFillMode, MeshCullMode, (EPrimitiveType)PreCacheParams.PrimitiveType, PSOInitializers);
+		CollectPSOInitializersForSkyLight(SceneTexturesConfig, VertexFactoryData, PreCacheParams, Material, bRenderSkyLight, bDitheredLODTransition, MeshFillMode, MeshCullMode, (EPrimitiveType)PreCacheParams.PrimitiveType, PSOInitializers);
 		bRenderSkyLight = false;
-		CollectPSOInitializersForSkyLight(SceneTexturesConfig, VertexFactoryData, Material, bRenderSkyLight, bDitheredLODTransition, MeshFillMode, MeshCullMode, (EPrimitiveType)PreCacheParams.PrimitiveType, PSOInitializers);
+		CollectPSOInitializersForSkyLight(SceneTexturesConfig, VertexFactoryData, PreCacheParams, Material, bRenderSkyLight, bDitheredLODTransition, MeshFillMode, MeshCullMode, (EPrimitiveType)PreCacheParams.PrimitiveType, PSOInitializers);
 	}
 }
 
 void FBasePassMeshProcessor::CollectPSOInitializersForSkyLight(
 	const FSceneTexturesConfig& SceneTexturesConfig,
 	const FPSOPrecacheVertexFactoryData& VertexFactoryData,
+	const FPSOPrecacheParams& PreCacheParams,
 	const FMaterial& RESTRICT Material,
 	const bool bRenderSkyLight,
 	const bool bDitheredLODTransition,
@@ -2391,19 +2404,19 @@ void FBasePassMeshProcessor::CollectPSOInitializersForSkyLight(
 		if (bAllowStaticLighting && bUseVolumetricLightmap)
 		{
 			CollectPSOInitializersForLMPolicy< FSelfShadowedVolumetricLightmapPolicy >(
-				SceneTexturesConfig, VertexFactoryData, Material, ShadingModels, bRenderSkyLight, bDitheredLODTransition,
+				SceneTexturesConfig, VertexFactoryData, PreCacheParams, Material, ShadingModels, bRenderSkyLight, bDitheredLODTransition,
 				FSelfShadowedVolumetricLightmapPolicy(), MeshFillMode, MeshCullMode, PrimitiveType, PSOInitializers);
 		}
 
 		if (IsIndirectLightingCacheAllowed(FeatureLevel) && bAllowIndirectLightingCache)
 		{
 			CollectPSOInitializersForLMPolicy< FSelfShadowedCachedPointIndirectLightingPolicy >(
-				SceneTexturesConfig, VertexFactoryData, Material, ShadingModels, bRenderSkyLight, bDitheredLODTransition,
+				SceneTexturesConfig, VertexFactoryData, PreCacheParams, Material, ShadingModels, bRenderSkyLight, bDitheredLODTransition,
 				FSelfShadowedCachedPointIndirectLightingPolicy(), MeshFillMode, MeshCullMode, PrimitiveType, PSOInitializers);
 		}
 
 		CollectPSOInitializersForLMPolicy< FSelfShadowedTranslucencyPolicy >(
-			SceneTexturesConfig, VertexFactoryData, Material, ShadingModels, bRenderSkyLight, bDitheredLODTransition,
+			SceneTexturesConfig, VertexFactoryData, PreCacheParams, Material, ShadingModels, bRenderSkyLight, bDitheredLODTransition,
 			FSelfShadowedTranslucencyPolicy(), MeshFillMode, MeshCullMode, PrimitiveType, PSOInitializers);
 	}
 
@@ -2411,7 +2424,7 @@ void FBasePassMeshProcessor::CollectPSOInitializersForSkyLight(
 	for (ELightMapPolicyType LightMapPolicyType : UniformLightMapPolicyTypes)
 	{
 		CollectPSOInitializersForLMPolicy< FUniformLightMapPolicy >(
-			SceneTexturesConfig, VertexFactoryData, Material, ShadingModels, bRenderSkyLight, bDitheredLODTransition,
+			SceneTexturesConfig, VertexFactoryData, PreCacheParams, Material, ShadingModels, bRenderSkyLight, bDitheredLODTransition,
 			FUniformLightMapPolicy(LightMapPolicyType), MeshFillMode, MeshCullMode, PrimitiveType, PSOInitializers);
 	}
 }
diff --git a/Engine/Source/Runtime/Renderer/Private/BasePassRendering.h b/Engine/Source/Runtime/Renderer/Private/BasePassRendering.h
--- a/Engine/Source/Runtime/Renderer/Private/BasePassRendering.h
+++ b/Engine/Source/Runtime/Renderer/Private/BasePassRendering.h
@@ -36,6 +36,7 @@
 #include "LocalFogVolumeRendering.h"
 #include "LightFunctionAtlas.h"
 #include "RenderUtils.h"
+#include "SceneTexturesConfig.h"
 
 class FScene;
 
@@ -839,6 +840,21 @@ public:
 	static ELightMapPolicyType GetUniformLightMapPolicyType(ERHIFeatureLevel::Type FeatureLevelconst, const FScene* Scene, const FLightCacheInterface* LCI, const FPrimitiveSceneProxy* RESTRICT PrimitiveSceneProxy, const FMaterial& Material);
 	static TArray<ELightMapPolicyType, TInlineAllocator<2>> GetUniformLightMapPolicyTypeForPSOCollection(ERHIFeatureLevel::Type FeatureLevel, const FMaterial& Material);
 
+	template<typename PassShadersType>
+	static void AddBasePassGraphicsPipelineStateInitializer(
+		ERHIFeatureLevel::Type InFeatureLevel,
+		const FPSOPrecacheVertexFactoryData& VertexFactoryData,
+		const FMaterial& RESTRICT MaterialResource,
+		const FMeshPassProcessorRenderState& RESTRICT DrawRenderState,
+		const FGraphicsPipelineRenderTargetsInfo& RESTRICT RenderTargetsInfo,
+		PassShadersType PassShaders,
+		ERasterizerFillMode MeshFillMode,
+		ERasterizerCullMode MeshCullMode,
+		EPrimitiveType PrimitiveType,
+		bool bPrecacheAlphaColorChannel,
+		int InPSOCollectorIndex,
+		TArray<FPSOPrecacheData>& PSOInitializers);
+
 private:
 
 	bool TryAddMeshBatch(const FMeshBatch& RESTRICT MeshBatch, uint64 BatchElementMask, const FPrimitiveSceneProxy* RESTRICT PrimitiveSceneProxy, int32 StaticMeshId, const FMaterialRenderProxy& MaterialRenderProxy, const FMaterial& Material);
@@ -863,6 +879,7 @@ private:
 	void CollectPSOInitializersForSkyLight(
 		const FSceneTexturesConfig& SceneTexturesConfig,
 		const FPSOPrecacheVertexFactoryData& VertexFactoryData,
+		const FPSOPrecacheParams& PreCacheParams,
 		const FMaterial& RESTRICT MaterialResource,
 		const bool bRenderSkylight,
 		const bool bDitheredLODTransition,
@@ -875,6 +892,7 @@ private:
 	void CollectPSOInitializersForLMPolicy(
 		const FSceneTexturesConfig& SceneTexturesConfig,
 		const FPSOPrecacheVertexFactoryData& VertexFactoryData,
+		const FPSOPrecacheParams& PreCacheParams,
 		const FMaterial& RESTRICT MaterialResource,
 		FMaterialShadingModelField ShadingModels,
 		const bool bRenderSkylight,
@@ -894,6 +912,69 @@ private:
 	float AutoBeforeDOFTranslucencyBoundary = 0.0f;
 };
 
+template<typename PassShadersType>
+void FBasePassMeshProcessor::AddBasePassGraphicsPipelineStateInitializer(
+	ERHIFeatureLevel::Type InFeatureLevel,
+	const FPSOPrecacheVertexFactoryData& VertexFactoryData,
+	const FMaterial& RESTRICT MaterialResource,
+	const FMeshPassProcessorRenderState& RESTRICT DrawRenderState,
+	const FGraphicsPipelineRenderTargetsInfo& RESTRICT RenderTargetsInfo,
+	PassShadersType PassShaders,
+	ERasterizerFillMode MeshFillMode,
+	ERasterizerCullMode MeshCullMode,
+	EPrimitiveType PrimitiveType,
+	bool bPrecacheAlphaColorChannel,
+	int InPSOCollectorIndex,
+	TArray<FPSOPrecacheData>& PSOInitializers)
+{
+	AddGraphicsPipelineStateInitializer(
+		VertexFactoryData,
+		MaterialResource,
+		DrawRenderState,
+		RenderTargetsInfo,
+		PassShaders,
+		MeshFillMode,
+		MeshCullMode,
+		PrimitiveType,
+		EMeshPassFeatures::Default,
+		ESubpassHint::None,
+		0,
+		true /*bRequired*/,
+		InPSOCollectorIndex,
+		PSOInitializers);
+
+	// Planar reflections and scene captures use scene color alpha to keep track of where content has been rendered, for compositing into a different scene later
+	if (bPrecacheAlphaColorChannel)
+	{
+		FGraphicsPipelineRenderTargetsInfo AlphaColorRenderTargetsInfo = RenderTargetsInfo;
+
+		bool bRequiresAlphaChannel = true;
+		ETextureCreateFlags ExtraSceneColorCreateFlags = ETextureCreateFlags::None;
+		EPixelFormat SceneColorFormatWithAlpha;
+		ETextureCreateFlags SceneColorCreateFlagsWithAlpha;
+		GetSceneColorFormatAndCreateFlags(InFeatureLevel, bRequiresAlphaChannel, ExtraSceneColorCreateFlags, RenderTargetsInfo.NumSamples, false /*bMemorylessMSAA*/, SceneColorFormatWithAlpha, SceneColorCreateFlagsWithAlpha);
+
+		AlphaColorRenderTargetsInfo.RenderTargetFormats[0] = SceneColorFormatWithAlpha;
+		AlphaColorRenderTargetsInfo.RenderTargetFlags[0] = SceneColorCreateFlagsWithAlpha;
+
+		AddGraphicsPipelineStateInitializer(
+			VertexFactoryData,
+			MaterialResource,
+			DrawRenderState,
+			AlphaColorRenderTargetsInfo,
+			PassShaders,
+			MeshFillMode,
+			MeshCullMode,
+			PrimitiveType,
+			EMeshPassFeatures::Default,
+			ESubpassHint::None,
+			0,
+			true /*bRequired*/,
+			InPSOCollectorIndex,
+			PSOInitializers);
+	}
+}
+
 ENUM_CLASS_FLAGS(FBasePassMeshProcessor::EFlags);
 
 extern void SetupBasePassState(FExclusiveDepthStencil::Type BasePassDepthStencilAccess, const bool bShaderComplexity, FMeshPassProcessorRenderState& DrawRenderState);
diff --git a/Engine/Source/Runtime/Renderer/Private/MeshPassProcessor.cpp b/Engine/Source/Runtime/Renderer/Private/MeshPassProcessor.cpp
--- a/Engine/Source/Runtime/Renderer/Private/MeshPassProcessor.cpp
+++ b/Engine/Source/Runtime/Renderer/Private/MeshPassProcessor.cpp
@@ -2312,7 +2312,7 @@ void SetupGBufferRenderTargetInfo(const FSceneTexturesConfig& SceneTexturesConfi
 	if (bSetupDepthStencil)
 	{
 		SetupDepthStencilInfo(PF_DepthStencil, SceneTexturesConfig.DepthCreateFlags, ERenderTargetLoadAction::ELoad,
-			ERenderTargetLoadAction::ELoad, FExclusiveDepthStencil::DepthWrite_StencilWrite, RenderTargetsInfo);
+			ERenderTargetLoadAction::ELoad, FExclusiveDepthStencil::DepthRead_StencilWrite, RenderTargetsInfo);
 	}
 }
 
diff --git a/Engine/Source/Runtime/Renderer/Private/MobileBasePass.cpp b/Engine/Source/Runtime/Renderer/Private/MobileBasePass.cpp
--- a/Engine/Source/Runtime/Renderer/Private/MobileBasePass.cpp
+++ b/Engine/Source/Runtime/Renderer/Private/MobileBasePass.cpp
@@ -1051,6 +1051,7 @@ void FMobileBasePassMeshProcessor::CollectPSOInitializersForLMPolicy(
 		SubpassHint,
 		SubpassIndex,
 		true /*bRequired*/,
+		PSOCollectorIndex,
 		PSOInitializers);
 }
 
diff --git a/Engine/Source/Runtime/Renderer/Private/SingleLayerWaterRendering.cpp b/Engine/Source/Runtime/Renderer/Private/SingleLayerWaterRendering.cpp
--- a/Engine/Source/Runtime/Renderer/Private/SingleLayerWaterRendering.cpp
+++ b/Engine/Source/Runtime/Renderer/Private/SingleLayerWaterRendering.cpp
@@ -1532,7 +1532,8 @@ void FSingleLayerWaterPassMeshProcessor::CollectPSOInitializers(const FSceneText
 		const bool bHasDepthPrepass = IsSingleLayerWaterDepthPrepassEnabled(GetFeatureLevelShaderPlatform(FeatureLevel), FeatureLevel);
 		RenderTargetsInfo.DepthStencilAccess = bHasDepthPrepass ? FExclusiveDepthStencil::DepthRead_StencilRead : FExclusiveDepthStencil::DepthRead_StencilNop;
 
-		AddGraphicsPipelineStateInitializer(
+		FBasePassMeshProcessor::AddBasePassGraphicsPipelineStateInitializer(
+			FeatureLevel,
 			VertexFactoryData,
 			Material,
 			PassDrawRenderState,
@@ -1541,8 +1542,8 @@ void FSingleLayerWaterPassMeshProcessor::CollectPSOInitializers(const FSceneText
 			MeshFillMode,
 			MeshCullMode,
 			(EPrimitiveType)PreCacheParams.PrimitiveType,
-			EMeshPassFeatures::Default, 
-			true /*bRequired*/,
+			true /*bPrecacheAlphaColorChannel*/,
+			PSOCollectorIndex,
 			PSOInitializers);
 	}
 }
diff --git a/Engine/Source/Runtime/Renderer/Private/SkyPassRendering.cpp b/Engine/Source/Runtime/Renderer/Private/SkyPassRendering.cpp
--- a/Engine/Source/Runtime/Renderer/Private/SkyPassRendering.cpp
+++ b/Engine/Source/Runtime/Renderer/Private/SkyPassRendering.cpp
@@ -202,7 +202,8 @@ void FSkyPassMeshProcessor::CollectPSOInitializers(const FSceneTexturesConfig& S
 	FGraphicsPipelineRenderTargetsInfo RenderTargetsInfo;
 	SetupGBufferRenderTargetInfo(SceneTexturesConfig, RenderTargetsInfo, true /*bSetupDepthStencil*/);
 	
-	AddGraphicsPipelineStateInitializer(
+	FBasePassMeshProcessor::AddBasePassGraphicsPipelineStateInitializer(
+		FeatureLevel,
 		VertexFactoryData,
 		Material,
 		PassDrawRenderState,
@@ -211,37 +212,39 @@ void FSkyPassMeshProcessor::CollectPSOInitializers(const FSceneTexturesConfig& S
 		MeshFillMode,
 		MeshCullMode,
 		(EPrimitiveType)PreCacheParams.PrimitiveType,
-		EMeshPassFeatures::Default,
-		true /*bRequired*/,
+		true /*bPrecacheAlphaColorChannel*/,
+		PSOCollectorIndex,
 		PSOInitializers);
 
 	// Also generate with depth write which is used during CaptureSkyMeshReflection
-	const FExclusiveDepthStencil::Type SceneBasePassDepthStencilAccess = FScene::GetDefaultBasePassDepthStencilAccess(FeatureLevel);
-	FMeshPassProcessorRenderState SkyCaptureDrawRenderState;
-	FExclusiveDepthStencil::Type BasePassDepthStencilAccess_Sky = FExclusiveDepthStencil::Type(SceneBasePassDepthStencilAccess | FExclusiveDepthStencil::DepthWrite);
-	SetupBasePassState(BasePassDepthStencilAccess_Sky, false, SkyCaptureDrawRenderState);
-	
-	// Also change render target format
-	FRDGTextureDesc SkyCaptureRenderTargetDesc = FSkyPassMeshProcessor::GetCaptureFrameSkyEnvMapTextureDesc(1, 1);
-
-	FGraphicsPipelineRenderTargetsInfo SkyCaptureRenderTargetsInfo;
-	SkyCaptureRenderTargetsInfo.NumSamples = 1;
-	AddRenderTargetInfo(SkyCaptureRenderTargetDesc.Format, SkyCaptureRenderTargetDesc.Flags, SkyCaptureRenderTargetsInfo);
-	SetupDepthStencilInfo(PF_DepthStencil, SceneTexturesConfig.DepthCreateFlags, ERenderTargetLoadAction::ELoad,
-		ERenderTargetLoadAction::ELoad, FExclusiveDepthStencil::DepthWrite_StencilWrite, SkyCaptureRenderTargetsInfo);
-
-	AddGraphicsPipelineStateInitializer(
-		VertexFactoryData,
-		Material,
-		SkyCaptureDrawRenderState,
-		SkyCaptureRenderTargetsInfo,
-		SkyPassShaders,
-		MeshFillMode,
-		MeshCullMode,
-		(EPrimitiveType)PreCacheParams.PrimitiveType,
-		EMeshPassFeatures::Default,
-		true /*bRequired*/,
-		PSOInitializers);
+	{
+		const FExclusiveDepthStencil::Type SceneBasePassDepthStencilAccess = FScene::GetDefaultBasePassDepthStencilAccess(FeatureLevel);
+		FMeshPassProcessorRenderState SkyCaptureDrawRenderState;
+		FExclusiveDepthStencil::Type BasePassDepthStencilAccess_Sky = FExclusiveDepthStencil::Type(SceneBasePassDepthStencilAccess | FExclusiveDepthStencil::DepthWrite);
+		SetupBasePassState(BasePassDepthStencilAccess_Sky, false, SkyCaptureDrawRenderState);
+
+		// Also change render target format
+		FRDGTextureDesc SkyCaptureRenderTargetDesc = FSkyPassMeshProcessor::GetCaptureFrameSkyEnvMapTextureDesc(1, 1);
+
+		FGraphicsPipelineRenderTargetsInfo SkyCaptureRenderTargetsInfo;
+		SkyCaptureRenderTargetsInfo.NumSamples = 1;
+		AddRenderTargetInfo(SkyCaptureRenderTargetDesc.Format, SkyCaptureRenderTargetDesc.Flags, SkyCaptureRenderTargetsInfo);
+		SetupDepthStencilInfo(PF_DepthStencil, SceneTexturesConfig.DepthCreateFlags, ERenderTargetLoadAction::ELoad,
+			ERenderTargetLoadAction::ELoad, FExclusiveDepthStencil::DepthWrite_StencilWrite, SkyCaptureRenderTargetsInfo);
+
+		AddGraphicsPipelineStateInitializer(
+			VertexFactoryData,
+			Material,
+			SkyCaptureDrawRenderState,
+			SkyCaptureRenderTargetsInfo,
+			SkyPassShaders,
+			MeshFillMode,
+			MeshCullMode,
+			(EPrimitiveType)PreCacheParams.PrimitiveType,
+			EMeshPassFeatures::Default,
+			true /*bRequired*/,
+			PSOInitializers);
+	}
 }
 
 FRDGTextureDesc FSkyPassMeshProcessor::GetCaptureFrameSkyEnvMapTextureDesc(uint32 CubeWidth, uint32 CubeMipCount)
diff --git a/Engine/Source/Runtime/Renderer/Public/MeshPassProcessor.h b/Engine/Source/Runtime/Renderer/Public/MeshPassProcessor.h
--- a/Engine/Source/Runtime/Renderer/Public/MeshPassProcessor.h
+++ b/Engine/Source/Runtime/Renderer/Public/MeshPassProcessor.h
@@ -2100,12 +2100,12 @@ public:
 		ERasterizerFillMode MeshFillMode,
 		ERasterizerCullMode MeshCullMode,
 		EPrimitiveType PrimitiveType,
-		EMeshPassFeatures MeshPassFeatures, 
+		EMeshPassFeatures MeshPassFeatures,
 		bool bRequired,
 		TArray<FPSOPrecacheData>& PSOInitializers);
 
 	template<typename PassShadersType>
-	void AddGraphicsPipelineStateInitializer(
+	static void AddGraphicsPipelineStateInitializer(
 		const FPSOPrecacheVertexFactoryData& VertexFactoryData,
 		const FMaterial& RESTRICT MaterialResource,
 		const FMeshPassProcessorRenderState& RESTRICT DrawRenderState,
@@ -2118,6 +2118,7 @@ public:
 		ESubpassHint SubpassHint,
 		uint8 SubpassIndex,
 		bool bRequired,
+		int32 PSOCollectorIndex,
 		TArray<FPSOPrecacheData>& PSOInitializers);
 
 protected:
diff --git a/Engine/Source/Runtime/Renderer/Public/MeshPassProcessor.inl b/Engine/Source/Runtime/Renderer/Public/MeshPassProcessor.inl
--- a/Engine/Source/Runtime/Renderer/Public/MeshPassProcessor.inl
+++ b/Engine/Source/Runtime/Renderer/Public/MeshPassProcessor.inl
@@ -232,6 +232,7 @@ void FMeshPassProcessor::AddGraphicsPipelineStateInitializer(
 		ESubpassHint::None,
 		0,
 		bRequired,
+		PSOCollectorIndex,
 		PSOInitializers);
 }
 
@@ -249,6 +250,7 @@ void FMeshPassProcessor::AddGraphicsPipelineStateInitializer(
 	ESubpassHint SubpassHint,
 	uint8 SubpassIndex,
 	bool bRequired,
+	int32 InPSOCollectorIndex,
 	TArray<FPSOPrecacheData>& PSOInitializers)
 {
 	TRACE_CPUPROFILER_EVENT_SCOPE(FMeshPassProcessor::AddGraphicsPipelineStateInitializer);
@@ -295,9 +297,9 @@ void FMeshPassProcessor::AddGraphicsPipelineStateInitializer(
 	{
 		MinimalPipelineStateInitializer.StatePrecachePSOHash = PipelineStateInitializer.StatePrecachePSOHash;
 		FGraphicsMinimalPipelineStateInitializer ShadersOnlyInitializer = PSOCollectorStats::GetShadersOnlyInitializer(MinimalPipelineStateInitializer);
-		PSOCollectorStats::GetShadersOnlyPSOPrecacheStatsCollector().AddStateToCache(ShadersOnlyInitializer, PSOCollectorStats::GetPSOPrecacheHash, &MaterialResource, PSOCollectorIndex, VertexFactoryData.VertexFactoryType);
+		PSOCollectorStats::GetShadersOnlyPSOPrecacheStatsCollector().AddStateToCache(ShadersOnlyInitializer, PSOCollectorStats::GetPSOPrecacheHash, &MaterialResource, InPSOCollectorIndex, VertexFactoryData.VertexFactoryType);
 		FGraphicsMinimalPipelineStateInitializer PatchedMinimalInitializer = PSOCollectorStats::PatchMinimalPipelineStateToCheck(MinimalPipelineStateInitializer);
-		PSOCollectorStats::GetMinimalPSOPrecacheStatsCollector().AddStateToCache(PatchedMinimalInitializer, PSOCollectorStats::GetPSOPrecacheHash, &MaterialResource, PSOCollectorIndex, VertexFactoryData.VertexFactoryType);
+		PSOCollectorStats::GetMinimalPSOPrecacheStatsCollector().AddStateToCache(PatchedMinimalInitializer, PSOCollectorStats::GetPSOPrecacheHash, &MaterialResource, InPSOCollectorIndex, VertexFactoryData.VertexFactoryType);
 	}
 #endif // PSO_PRECACHING_VALIDATE
 
@@ -310,7 +312,7 @@ void FMeshPassProcessor::AddGraphicsPipelineStateInitializer(
 	PSOPrecacheData.Type = FPSOPrecacheData::EType::Graphics;
 	PSOPrecacheData.GraphicsPSOInitializer = PipelineStateInitializer;
 #if PSO_PRECACHING_VALIDATE
-	PSOPrecacheData.PSOCollectorIndex = PSOCollectorIndex;
+	PSOPrecacheData.PSOCollectorIndex = InPSOCollectorIndex;
 	PSOPrecacheData.VertexFactoryType = VertexFactoryData.VertexFactoryType;
 	if (PSOCollectorStats::IsFullPrecachingValidationEnabled())
 	{
diff --git a/Engine/Source/Runtime/UMG/Private/Components/WidgetComponent.cpp b/Engine/Source/Runtime/UMG/Private/Components/WidgetComponent.cpp
--- a/Engine/Source/Runtime/UMG/Private/Components/WidgetComponent.cpp
+++ b/Engine/Source/Runtime/UMG/Private/Components/WidgetComponent.cpp
@@ -687,6 +687,25 @@ bool UWidgetComponent::CanBeInCluster() const
 	return false;
 }
 
+void UWidgetComponent::PostLoad()
+{
+	Super::PostLoad();
+
+	PrecachePSOs();
+}
+
+void UWidgetComponent::CollectPSOPrecacheData(const FPSOPrecacheParams& BasePrecachePSOParams, FComponentPSOPrecacheParamsList& OutParams)
+{
+	if (MaterialInstance)
+	{
+		FComponentPSOPrecacheParams& ComponentParams = OutParams[OutParams.AddDefaulted()];
+		ComponentParams.Priority = EPSOPrecachePriority::High;
+		ComponentParams.MaterialInterface = MaterialInstance;
+		ComponentParams.VertexFactoryDataList.Add(FPSOPrecacheVertexFactoryData(&FLocalVertexFactory::StaticType));
+		ComponentParams.PSOPrecacheParams = BasePrecachePSOParams;
+	}
+}
+
 void UWidgetComponent::BeginPlay()
 {
 	SetComponentTickEnabled(TickMode != ETickMode::Disabled);
@@ -742,6 +761,8 @@ void UWidgetComponent::UpdateMaterialInstance()
 	}
 	UpdateMaterialInstanceParameters();
 
+	PrecachePSOs();
+
 	MarkRenderStateDirty();
 }
 
diff --git a/Engine/Source/Runtime/UMG/Public/Components/WidgetComponent.h b/Engine/Source/Runtime/UMG/Public/Components/WidgetComponent.h
--- a/Engine/Source/Runtime/UMG/Public/Components/WidgetComponent.h
+++ b/Engine/Source/Runtime/UMG/Public/Components/WidgetComponent.h
@@ -99,6 +99,7 @@ public:
 	//UObject interface
 	UMG_API virtual void Serialize(FArchive& Ar) override;
 	UMG_API virtual bool CanBeInCluster() const override;
+	UMG_API virtual void PostLoad() override;
 	//~ End UObject Interface
 
 	/** UActorComponent Interface */
@@ -116,7 +117,9 @@ public:
 	UMG_API UMaterialInterface* GetMaterial(int32 MaterialIndex) const override;
 	UMG_API virtual void SetMaterial(int32 ElementIndex, UMaterialInterface* Material) override;
 	UMG_API int32 GetNumMaterials() const override;
-	UMG_API virtual void TickComponent(float DeltaTime, enum ELevelTick TickType, FActorComponentTickFunction *ThisTickFunction) override;
+	UMG_API virtual void TickComponent(float DeltaTime, enum ELevelTick TickType, FActorComponentTickFunction* ThisTickFunction) override;/** Collect all the PSO precache data used by the static mesh component */
+	UMG_API virtual void CollectPSOPrecacheData(const FPSOPrecacheParams& BasePrecachePSOParams, FComponentPSOPrecacheParamsList& OutParams) override;
+
 
 	UMG_API virtual TStructOnScope<FActorComponentInstanceData> GetComponentInstanceData() const override;
 	UMG_API void ApplyComponentInstanceData(struct FWidgetComponentInstanceData* ComponentInstanceData);
