Fix CIS error

#rnx

diff --git a/Engine/Source/Runtime/Renderer/Private/SceneCulling/HierarchicalSpatialHashGrid.h b/Engine/Source/Runtime/Renderer/Private/SceneCulling/HierarchicalSpatialHashGrid.h
--- a/Engine/Source/Runtime/Renderer/Private/SceneCulling/HierarchicalSpatialHashGrid.h
+++ b/Engine/Source/Runtime/Renderer/Private/SceneCulling/HierarchicalSpatialHashGrid.h
@@ -167,7 +167,7 @@ template <typename BlockTraitsType>
 class THierarchicalSpatialHashGrid
 {
 public:
-	using FBlockLoc = BlockTraitsType::FBlockLoc;
+	using FBlockLoc = typename BlockTraitsType::FBlockLoc;
 
 	static constexpr int32 CellBlockDimLog2 = BlockTraitsType::CellBlockDimLog2;
 	static constexpr int32 kMaxLevel = 64;
