Change default VSM page marking to 2x2 stride

diff --git a/Engine/Source/Runtime/Renderer/Private/VirtualShadowMaps/VirtualShadowMapArray.cpp b/Engine/Source/Runtime/Renderer/Private/VirtualShadowMaps/VirtualShadowMapArray.cpp
--- a/Engine/Source/Runtime/Renderer/Private/VirtualShadowMaps/VirtualShadowMapArray.cpp
+++ b/Engine/Source/Runtime/Renderer/Private/VirtualShadowMaps/VirtualShadowMapArray.cpp
@@ -336,7 +336,7 @@ static TAutoConsoleVariable<int32> CVarVirtualShadowSinglePassBatched(
 
 static TAutoConsoleVariable<int32> CVarVirtualShadowMapPageMarkingPixelStrideX(
 	TEXT("r.Shadow.Virtual.PageMarkingPixelStrideX"),
-	1,
+	2,
 	TEXT("During page marking, instead of testing every screen pixel, test every Nth pixel.\n")
 	TEXT("Page marking from screen pixels is used to determine which VSM pages are seen from the camera and need to be rendered.\n")
 	TEXT("Increasing this value reduces page-marking costs, but could introduce artifacts due to missing pages.\n")
@@ -345,7 +345,7 @@ static TAutoConsoleVariable<int32> CVarVirtualShadowMapPageMarkingPixelStrideX(
 
 static TAutoConsoleVariable<int32> CVarVirtualShadowMapPageMarkingPixelStrideY(
 	TEXT("r.Shadow.Virtual.PageMarkingPixelStrideY"),
-	1,
+	2,
 	TEXT("Same as PageMarkingPixelStrideX, but on the vertical axis of the screen."),
 	ECVF_RenderThreadSafe);
 
