Fix incorrect feature level checks.

The IsReflectionEnvironmentActive and FVectorFieldAnimatedResource::Update functions were checking if the feature level is SM5 exactly, instead of checking if it's at least SM5.

#rnx

diff --git a/Engine/Source/Runtime/Engine/Private/VectorField.cpp b/Engine/Source/Runtime/Engine/Private/VectorField.cpp
--- a/Engine/Source/Runtime/Engine/Private/VectorField.cpp
+++ b/Engine/Source/Runtime/Engine/Private/VectorField.cpp
@@ -1112,7 +1112,7 @@ public:
 	{
 		check(IsInRenderingThread());
 
-		if (GetFeatureLevel() == ERHIFeatureLevel::SM5 && AnimatedVectorField && AnimatedVectorField->Texture && AnimatedVectorField->Texture->GetResource())
+		if (GetFeatureLevel() >= ERHIFeatureLevel::SM5 && AnimatedVectorField && AnimatedVectorField->Texture && AnimatedVectorField->Texture->GetResource())
 		{
 			SCOPED_DRAW_EVENT(RHICmdList, AnimateVectorField);
 
diff --git a/Engine/Source/Runtime/Renderer/Private/CompositionLighting/CompositionLighting.cpp b/Engine/Source/Runtime/Renderer/Private/CompositionLighting/CompositionLighting.cpp
--- a/Engine/Source/Runtime/Renderer/Private/CompositionLighting/CompositionLighting.cpp
+++ b/Engine/Source/Runtime/Renderer/Private/CompositionLighting/CompositionLighting.cpp
@@ -73,7 +73,7 @@ static bool IsReflectionEnvironmentActive(const FSceneView& View)
 	bool HasReflectionCaptures = (Scene->ReflectionSceneData.RegisteredReflectionCaptures.Num() > 0);
 	bool HasSSR = View.Family->EngineShowFlags.ScreenSpaceReflections;
 
-	return (Scene->GetFeatureLevel() == ERHIFeatureLevel::SM5 && IsReflectingEnvironment && (HasReflectionCaptures || HasSSR) && !IsForwardShadingEnabled(View.GetShaderPlatform()));
+	return (Scene->GetFeatureLevel() >= ERHIFeatureLevel::SM5 && IsReflectingEnvironment && (HasReflectionCaptures || HasSSR) && !IsForwardShadingEnabled(View.GetShaderPlatform()));
 }
 
 static bool IsSkylightActive(const FViewInfo& View)
