Alpha holdout - sky, fog and cloud passes now write transmittance correctly into the alpha channel.
Sky dome mesh with IsSky material always write transmittance = 1. Standalone SkyAtmosphere writes appropriate grey scale transmittance.
The SkyViewLUT is now RGBA when alpha hold out is enabled to store transmittance.
When alpha propagation is disabled, things are unchanged.

diff --git a/Engine/Shaders/Private/BasePassPixelShader.usf b/Engine/Shaders/Private/BasePassPixelShader.usf
--- a/Engine/Shaders/Private/BasePassPixelShader.usf
+++ b/Engine/Shaders/Private/BasePassPixelShader.usf
@@ -2046,6 +2046,14 @@ void FPixelShaderInOut_MainPS(
 	#else
 		Out.MRT[0].rgba *= ViewPreExposure;
 	#endif
+
+	#if MATERIAL_IS_SKY && POST_PROCESS_ALPHA
+		// In this case, alpha will represent transmitance toward the background.
+		// But the sky material is the background. And we have no way torecover transmittance to outer space for now so we simply write transmittance=1.
+		// This matches the path tracer behavior. Note: when the SkyAtmosphere is standalone (no sky dome mesh material), we can write grey scale trnasmitance to outer space.
+		Out.MRT[0].a = 1.0f;
+	#endif
+
 #endif
 
 	// If support OIT, then remove blending contribution and insert OIT sample
diff --git a/Engine/Shaders/Private/SkyAtmosphere.usf b/Engine/Shaders/Private/SkyAtmosphere.usf
--- a/Engine/Shaders/Private/SkyAtmosphere.usf
+++ b/Engine/Shaders/Private/SkyAtmosphere.usf
@@ -944,10 +944,19 @@ void RenderSkyAtmosphereRayMarchingPS(
 		bool IntersectGround = RaySphereIntersectNearest(WorldPosLocal, WorldDirLocal, float3(0, 0, 0), Atmosphere.BottomRadiusKm) >= 0.0f;
 
 		SkyViewLutParamsToUv(IntersectGround, ViewZenithCosAngle, WorldDirLocal, ViewHeight, Atmosphere.BottomRadiusKm, SkyAtmosphere.SkyViewLutSizeAndInvSize, UV);
-		float3 SkyLuminance = SkyViewLutTexture.SampleLevel(SkyViewLutTextureSampler, UV, 0).rgb;
+		float4 SkyLuminanceTransmittance = SkyViewLutTexture.SampleLevel(SkyViewLutTextureSampler, UV, 0);
+		float3 SkyLuminance = SkyLuminanceTransmittance.rgb;
+
+		float3 SkyGreyTransmittance = 1.0f;
+	#if POST_PROCESS_ALPHA
+		if (Atmosphere.bHoldout && View.RenderingReflectionCaptureMask == 0.0f)
+		{
+			SkyGreyTransmittance = SkyLuminanceTransmittance.aaa;
+		}
+	#endif
 
 		PreExposedL += SkyLuminance * LuminanceScale * (ViewOneOverPreExposure * OutputPreExposure);
-		OutLuminance = PrepareOutput(PreExposedL);
+		OutLuminance = PrepareOutput(PreExposedL, SkyGreyTransmittance);
 		return;
 	}
 #endif
@@ -1230,12 +1239,9 @@ void RenderMultiScatteredLuminanceLutCS(uint3 ThreadId : SV_DispatchThreadID)
 // Sky View LUT
 ////////////////////////////////////////////////////////////
 
-// see SkyAtmosphereRendering.cpp GetSkyLutTextureFormat
-#if SHADING_PATH_MOBILE
+// Even thought the texture can have 3 channels only (see SkyAtmosphereRendering.cpp GetSkyLutTextureFormat)
+// We always write 4 channels as the grey scale transmittance can be used when alpha propagation is enabled.
 RWTexture2D<float4> SkyViewLutUAV;
-#else
-RWTexture2D<float3> SkyViewLutUAV;
-#endif
 
 [numthreads(THREADGROUP_SIZE, THREADGROUP_SIZE, 1)]
 void RenderSkyViewLutCS(uint3 ThreadId : SV_DispatchThreadID)
@@ -1292,11 +1298,8 @@ void RenderSkyViewLutCS(uint3 ThreadId : SV_DispatchThreadID)
 		View.AtmosphereLightIlluminanceOuterSpace[0].rgb, View.AtmosphereLightIlluminanceOuterSpace[1].rgb,
 		AerialPespectiveViewDistanceScale);
 
-#if SHADING_PATH_MOBILE
-	SkyViewLutUAV[int2(PixPos)] = float4(ss.L, 0.0f);
-#else
-	SkyViewLutUAV[int2(PixPos)] = ss.L;
-#endif
+	const float Transmittance = dot(ss.Transmittance, float3(1.0f / 3.0f, 1.0f / 3.0f, 1.0f / 3.0f));
+	SkyViewLutUAV[int2(PixPos)] = float4(ss.L, Transmittance);
 }
 
 
diff --git a/Engine/Source/Runtime/Renderer/Private/FogRendering.cpp b/Engine/Source/Runtime/Renderer/Private/FogRendering.cpp
--- a/Engine/Source/Runtime/Renderer/Private/FogRendering.cpp
+++ b/Engine/Source/Runtime/Renderer/Private/FogRendering.cpp
@@ -11,6 +11,7 @@
 #include "SceneCore.h"
 #include "ScreenPass.h"
 #include "TextureResource.h"
+#include "PostProcess/PostProcessing.h" // IsPostProcessingWithAlphaChannelSupported
 
 DECLARE_GPU_DRAWCALL_STAT(Fog);
 
@@ -315,8 +316,16 @@ static void RenderViewFog(
 
 	if (bEnableBlending)
 	{
-		// disable alpha writes in order to preserve scene depth values on PC
-		GraphicsPSOInit.BlendState = TStaticBlendState<CW_RGB, BO_Add, BF_One, BF_SourceAlpha>::GetRHI();
+		const bool bSupportsAlpha = IsPostProcessingWithAlphaChannelSupported();
+		if (bSupportsAlpha)
+		{
+			GraphicsPSOInit.BlendState = TStaticBlendState<CW_RGBA, BO_Add, BF_One, BF_SourceAlpha, BO_Add, BF_Zero, BF_SourceAlpha>::GetRHI();
+		}
+		else
+		{
+			// disable alpha writes in order to preserve scene depth values on PC
+			GraphicsPSOInit.BlendState = TStaticBlendState<CW_RGB, BO_Add, BF_One, BF_SourceAlpha>::GetRHI();
+		}
 	}
 	else
 	{
diff --git a/Engine/Source/Runtime/Renderer/Private/SkyAtmosphereRendering.cpp b/Engine/Source/Runtime/Renderer/Private/SkyAtmosphereRendering.cpp
--- a/Engine/Source/Runtime/Renderer/Private/SkyAtmosphereRendering.cpp
+++ b/Engine/Source/Runtime/Renderer/Private/SkyAtmosphereRendering.cpp
@@ -19,6 +19,7 @@
 #include "RendererUtils.h"
 #include "ScreenPass.h"
 #include "UnrealEngine.h"
+#include "PostProcess/PostProcessing.h" // IsPostProcessingWithAlphaChannelSupported
 
 
 //PRAGMA_DISABLE_OPTIMIZATION
@@ -1111,7 +1112,8 @@ void FSceneRenderer::InitSkyAtmosphereForViews(FRHICommandListImmediate& RHICmdL
 static EPixelFormat GetSkyLutTextureFormat(ERHIFeatureLevel::Type FeatureLevel)
 {
 	EPixelFormat TextureLUTFormat = PF_FloatRGB;
-	if (FeatureLevel <= ERHIFeatureLevel::ES3_1)
+	const bool bSupportsAlpha = IsPostProcessingWithAlphaChannelSupported();
+	if (FeatureLevel <= ERHIFeatureLevel::ES3_1 || bSupportsAlpha)
 	{
 		// OpenGL ES3.1 does not support storing into 3-component images
 		// TODO: check if need this for Metal, Vulkan
@@ -1858,12 +1860,13 @@ void FSceneRenderer::RenderSkyAtmosphereInternal(
 
 		const bool bFastAerialPerspectiveDepthTest = SkyRC.bFastAerialPerspectiveDepthTest;
 		const bool bRenderSkyPixel = SkyRC.bRenderSkyPixel;
+		const bool bSupportsAlpha = IsPostProcessingWithAlphaChannelSupported();
 		FIntRect Viewport = SkyRC.Viewport;
 		GraphBuilder.AddPass(
 			{},
 			PsPassParameters,
 			ERDGPassFlags::Raster,
-			[PsPassParameters, VertexShader, PixelShader, Viewport, bFastAerialPerspectiveDepthTest, bRenderSkyPixel, bDisableBlending, StartDepthZ](FRHICommandList& RHICmdListLambda)
+			[PsPassParameters, VertexShader, PixelShader, Viewport, bFastAerialPerspectiveDepthTest, bRenderSkyPixel, bDisableBlending, StartDepthZ, bSupportsAlpha](FRHICommandList& RHICmdListLambda)
 		{
 			RHICmdListLambda.SetViewport(Viewport.Min.X, Viewport.Min.Y, 0.0f, Viewport.Max.X, Viewport.Max.Y, 1.0f);
 
@@ -1876,7 +1879,14 @@ void FSceneRenderer::RenderSkyAtmosphereInternal(
 			}
 			else
 			{
-				GraphicsPSOInit.BlendState = TStaticBlendState<CW_RGB, BO_Add, BF_One, BF_SourceAlpha, BO_Add, BF_Zero, BF_One>::GetRHI();
+				if (bSupportsAlpha)
+				{
+					GraphicsPSOInit.BlendState = TStaticBlendState<CW_RGBA, BO_Add, BF_One, BF_SourceAlpha, BO_Add, BF_Zero, BF_SourceAlpha>::GetRHI();
+				}
+				else
+				{
+					GraphicsPSOInit.BlendState = TStaticBlendState<CW_RGB, BO_Add, BF_One, BF_SourceAlpha, BO_Add, BF_Zero, BF_One>::GetRHI(); 
+				}
 			}
 			if (bFastAerialPerspectiveDepthTest)
 			{
diff --git a/Engine/Source/Runtime/Renderer/Private/VolumetricRenderTarget.cpp b/Engine/Source/Runtime/Renderer/Private/VolumetricRenderTarget.cpp
--- a/Engine/Source/Runtime/Renderer/Private/VolumetricRenderTarget.cpp
+++ b/Engine/Source/Runtime/Renderer/Private/VolumetricRenderTarget.cpp
@@ -13,6 +13,7 @@
 #include "SingleLayerWaterRendering.h"
 #include "VolumetricCloudRendering.h"
 #include "RendererUtils.h"
+#include "PostProcess/PostProcessing.h" // IsPostProcessingWithAlphaChannelSupported
 
 static TAutoConsoleVariable<int32> CVarVolumetricRenderTarget(
 	TEXT("r.VolumetricRenderTarget"), 1,
@@ -755,7 +756,16 @@ void ComposeVolumetricRenderTargetOverScene(
 	RDG_GPU_STAT_SCOPE(GraphBuilder, VolCloudComposeOverScene);
 	SCOPED_NAMED_EVENT(VolCloudComposeOverScene, FColor::Emerald);
 
-	FRHIBlendState* PreMultipliedColorTransmittanceBlend = TStaticBlendState<CW_RGB, BO_Add, BF_One, BF_SourceAlpha, BO_Add, BF_Zero, BF_One>::GetRHI();
+	FRHIBlendState* PreMultipliedColorTransmittanceBlend;
+	const bool bSupportsAlpha = IsPostProcessingWithAlphaChannelSupported();
+	if (bSupportsAlpha)
+	{
+		PreMultipliedColorTransmittanceBlend = TStaticBlendState<CW_RGBA, BO_Add, BF_One, BF_SourceAlpha, BO_Add, BF_Zero, BF_SourceAlpha>::GetRHI();
+	}
+	else
+	{
+		PreMultipliedColorTransmittanceBlend = TStaticBlendState<CW_RGB, BO_Add, BF_One, BF_SourceAlpha, BO_Add, BF_Zero, BF_One>::GetRHI();
+	}
 
 	for (int32 ViewIndex = 0; ViewIndex < Views.Num(); ViewIndex++)
 	{
