Allows TSR RejectShading pass to modify the motion vector for resurrection and optical-flow

diff --git a/Engine/Shaders/Private/TemporalSuperResolution/TSRMergeRejection.usf b/Engine/Shaders/Private/TemporalSuperResolution/TSRMergeRejection.usf
deleted file mode 100644
diff --git a/Engine/Shaders/Private/TemporalSuperResolution/TSRRejectShading.usf b/Engine/Shaders/Private/TemporalSuperResolution/TSRRejectShading.usf
--- a/Engine/Shaders/Private/TemporalSuperResolution/TSRRejectShading.usf
+++ b/Engine/Shaders/Private/TemporalSuperResolution/TSRRejectShading.usf
@@ -128,6 +128,8 @@
 
 //------------------------------------------------------- TYPEDEFS
 
+#define tsr_simd_float  TLaneVector2D<float, 1, LANE_STRIDE_X, LANE_STRIDE_Y>
+#define tsr_simd_float2 TLaneVector2D<float, 2, LANE_STRIDE_X, LANE_STRIDE_Y>
 #define tsr_simd_half  TLaneVector2D<tsr_half, 1, LANE_STRIDE_X, LANE_STRIDE_Y>
 #define tsr_simd_half2 TLaneVector2D<tsr_half, 2, LANE_STRIDE_X, LANE_STRIDE_Y>
 #define tsr_simd_half3 TLaneVector2D<tsr_half, 3, LANE_STRIDE_X, LANE_STRIDE_Y>
@@ -159,6 +161,7 @@ static const tsr_half FilteringWeight = rcp(tsr_half(1.0 + 4 * 0.5 + 4 * 0.25));
 //------------------------------------------------------- PARAMETERS
 
 FScreenTransform InputPixelPosToTranslucencyTextureUV;
+float4x4 ClipToResurrectionClip;
 float2 TranslucencyTextureUVMin;
 float2 TranslucencyTextureUVMax;
 float3 HistoryGuideQuantizationError;
@@ -180,12 +183,14 @@ Texture2D<tsr_half4> ReprojectedHistoryMoireTexture;
 #endif
 Texture2D<tsr_half2> DecimateMaskTexture;
 Texture2D<tsr_ushort> IsMovingMaskTexture;
+Texture2D<float2> ClosestDepthTexture;
 
 RWTexture2DArray<tsr_half4> HistoryGuideOutput;
 RWTexture2DArray<tsr_half4> HistoryMoireOutput;
-RWTexture2D<tsr_half4> HistoryRejectionOutput;
-RWTexture2D<tsr_halfC> InputSceneColorOutput;
-RWTexture2D<tsr_half> InputSceneColorLdrLumaOutput;
+RWTexture2D<tsr_half4>  HistoryRejectionOutput;
+RWTexture2D<float2>     DilatedVelocityOutput;
+RWTexture2D<tsr_halfC>  InputSceneColorOutput;
+RWTexture2D<tsr_half>   InputSceneColorLdrLumaOutput;
 RWTexture2D<tsr_ushort> AntiAliasMaskOutput;
 
 
@@ -198,6 +203,20 @@ static tsr_simd_halfC Debug[DEBUG_ARRAY_SIZE];
 
 //------------------------------------------------------- INPUT FETCHES
 
+tsr_short2 ComputeElementPixelPos(const uint ElementIndex)
+{
+	tsr_short2 ClampedFetchPixelPos = ComputeElementInputPixelPos(
+		/* LaneStride = */ uint2(LANE_STRIDE_X, LANE_STRIDE_Y),
+		TileOverscan,
+		InputInfo_ViewportMin,
+		InputInfo_ViewportMax,
+		GGroupId,
+		GGroupThreadIndex,
+		ElementIndex);
+
+	return ClampedFetchPixelPos;
+}
+
 tsr_simd_short2 ComputePixelPos()
 {
 	tsr_simd_short2 PixelPos;
@@ -205,16 +224,7 @@ tsr_simd_short2 ComputePixelPos()
 	UNROLL_N(SIMD_SIZE)
 	for (uint ElementIndex = 0; ElementIndex < SIMD_SIZE; ElementIndex++)
 	{
-		tsr_short2 ClampedFetchPixelPos = ComputeElementInputPixelPos(
-			/* LaneStride = */ uint2(LANE_STRIDE_X, LANE_STRIDE_Y),
-			TileOverscan,
-			InputInfo_ViewportMin,
-			InputInfo_ViewportMax,
-			GGroupId,
-			GGroupThreadIndex,
-			ElementIndex);
-
-		PixelPos.SetElement(ElementIndex, ClampedFetchPixelPos);
+		PixelPos.SetElement(ElementIndex, ComputeElementPixelPos(ElementIndex));
 	}
 
 	return PixelPos;
@@ -794,6 +804,42 @@ tsr_simd_bool ShouldResurrectHistory(tsr_simd_half PrevRejectionBlendFinal, tsr_
 	return bResurrectHistory;
 }
 
+// Load the depth for ComputeResurrectionVelocity()
+CALL_SITE_DEBUGLOC
+void FetchDeviceZ(
+	tsr_simd_short2 PixelPos,
+	out tsr_simd_float DeviceZ)
+{
+	UNROLL_N(SIMD_SIZE)
+	for (uint ElementIndex = 0; ElementIndex < SIMD_SIZE; ElementIndex++)
+	{
+		tsr_short2 ClampedFetchPixelPos = PixelPos.GetElement(ElementIndex);
+		
+		DeviceZ.SetElement(ElementIndex, ClosestDepthTexture[ClampedFetchPixelPos].g);
+	}
+} // FetchMoireInput()
+
+// Compute the velocity used by the reprojection of the resurrection.
+CALL_SITE_DEBUGLOC
+tsr_simd_float2 ComputeResurrectionMotionVector(tsr_simd_float DeviceZ)
+{
+	tsr_simd_float2 ResurrectedEncodedScreenVelocity;
+	
+	UNROLL_N(SIMD_SIZE)
+	for (uint ElementIndex = 0; ElementIndex < SIMD_SIZE; ElementIndex++)
+	{
+		float2 ScreenPos = ApplyScreenTransform(float2(ComputeElementPixelPos(ElementIndex)), InputPixelPosToScreenPos);
+
+		float4 ThisClip = float4(ScreenPos, DeviceZ.GetElement(ElementIndex), 1);
+		float4 ResurrectionClip = mul(ThisClip, ClipToResurrectionClip); 
+		float2 ResurrectionScreen = ResurrectionClip.xy / ResurrectionClip.w;
+			
+		ResurrectedEncodedScreenVelocity.SetElement(ElementIndex, EncodeVelocityToTexture(float3(ScreenPos - ResurrectionScreen, 0.0)).xy);
+	}
+
+	return ResurrectedEncodedScreenVelocity;
+}
+
 #endif // CONFIG_HISTORY_RESURRECTION
 
 
@@ -1138,22 +1184,42 @@ void MainCS(
 			// Reload the resurrected history from cache.
 			tsr_simd_halfC ResurrectedHistory;
 			tsr_simd_half  ResurrectedHistoryUncertainty;
-			#if CONFIG_CACHE_REFETCH
+			tsr_simd_float DeviceZ;
+			PLATFORM_SPECIFIC_ISOLATE
 			{
-				FetchHistoryGuide(
-					ResurrectedHistoryGuideTexture,
-					ReprojectedHistoryGuideMetadataTexture,
-					ComputePixelPos(),
-					/* out */ ResurrectedHistory,
-					/* out */ ResurrectedHistoryUncertainty);
+				tsr_simd_short2 PixePos = ComputePixelPos();
+				
+				#if CONFIG_CACHE_REFETCH
+					FetchHistoryGuide(
+						ResurrectedHistoryGuideTexture,
+						ReprojectedHistoryGuideMetadataTexture,
+						PixePos,
+						/* out */ ResurrectedHistory,
+						/* out */ ResurrectedHistoryUncertainty);
+				#else
+					ResurrectedHistory            = OriginalResurrectedHistory;
+					ResurrectedHistoryUncertainty = OriginalResurrectedHistoryUncertainty;
+				#endif
+					
+				FetchDeviceZ(PixePos, /* out */ DeviceZ);
 			}
-			#else
-				ResurrectedHistory            = OriginalResurrectedHistory;
-				ResurrectedHistoryUncertainty = OriginalResurrectedHistoryUncertainty;
-			#endif
 
+			// Apply the resurrection on the guide immediatly to free up registers.
 			History = select(tsr_simd_boolC::Vectorize(bResurrectHistory), ResurrectedHistory, History);
 			HistoryUncertainty = select(bResurrectHistory, ResurrectedHistoryUncertainty, HistoryUncertainty);
+			
+			// Compute the motion vector for the resurrection
+			tsr_simd_float2 ResurrectionMotionVector = ComputeResurrectionMotionVector(DeviceZ);
+			
+			// Amends the motion vector for resurrection areas.
+			UNROLL_N(SIMD_SIZE)
+			for (uint ElementIndex = 0; ElementIndex < SIMD_SIZE; ElementIndex++)
+			{
+				tsr_short2 OutputPixelPos = ComputeElementOutputPixelPos(ElementIndex);
+				OutputPixelPos.x = select(bResurrectHistory.GetElement(ElementIndex), OutputPixelPos.x, tsr_short(-1));
+
+				DilatedVelocityOutput[OutputPixelPos] = ResurrectionMotionVector.GetElement(ElementIndex);
+			}
 		}
 		#endif // CONFIG_HISTORY_RESURRECTION
 
diff --git a/Engine/Source/Runtime/Renderer/Private/PostProcess/TemporalSuperResolution.cpp b/Engine/Source/Runtime/Renderer/Private/PostProcess/TemporalSuperResolution.cpp
--- a/Engine/Source/Runtime/Renderer/Private/PostProcess/TemporalSuperResolution.cpp
+++ b/Engine/Source/Runtime/Renderer/Private/PostProcess/TemporalSuperResolution.cpp
@@ -685,6 +685,7 @@ class FTSRRejectShadingCS : public FTSRShader
 		SHADER_PARAMETER(FScreenTransform, InputPixelPosToTranslucencyTextureUV)
 		SHADER_PARAMETER(FVector2f, TranslucencyTextureUVMin)
 		SHADER_PARAMETER(FVector2f, TranslucencyTextureUVMax)
+		SHADER_PARAMETER(FMatrix44f, ClipToResurrectionClip)
 		SHADER_PARAMETER(FVector3f, HistoryGuideQuantizationError)
 		SHADER_PARAMETER(float, FlickeringFramePeriod)
 		SHADER_PARAMETER(float, TheoricBlendFactor)
@@ -702,10 +703,12 @@ class FTSRRejectShadingCS : public FTSRShader
 		SHADER_PARAMETER_RDG_TEXTURE_SRV(Texture2D, ResurrectedHistoryGuideMetadataTexture)
 		SHADER_PARAMETER_RDG_TEXTURE(Texture2D, DecimateMaskTexture)
 		SHADER_PARAMETER_RDG_TEXTURE_SRV(Texture2D, IsMovingMaskTexture)
+		SHADER_PARAMETER_RDG_TEXTURE(Texture2D, ClosestDepthTexture)
 
 		SHADER_PARAMETER_RDG_TEXTURE_UAV(RWTexture2DArray, HistoryGuideOutput)
 		SHADER_PARAMETER_RDG_TEXTURE_UAV(RWTexture2DArray, HistoryMoireOutput)
 		SHADER_PARAMETER_RDG_TEXTURE_UAV(RWTexture2D, HistoryRejectionOutput)
+		SHADER_PARAMETER_RDG_TEXTURE_UAV(RWTexture2D, DilatedVelocityOutput)
 		SHADER_PARAMETER_RDG_TEXTURE_UAV(RWTexture2D, InputSceneColorOutput)
 		SHADER_PARAMETER_RDG_TEXTURE_UAV(RWTexture2D, InputSceneColorLdrLumaOutput)
 		SHADER_PARAMETER_RDG_TEXTURE_UAV(RWTexture2D, AntiAliasMaskOutput)
@@ -792,23 +795,6 @@ class FTSRRejectShadingCS : public FTSRShader
 	}
 }; // class FTSRRejectShadingCS
 
-class FTSRMergeRejectionCS : public FTSRShader
-{
-	DECLARE_GLOBAL_SHADER(FTSRMergeRejectionCS);
-	SHADER_USE_PARAMETER_STRUCT(FTSRMergeRejectionCS, FTSRShader);
-
-	BEGIN_SHADER_PARAMETER_STRUCT(FParameters, )
-		SHADER_PARAMETER_STRUCT_INCLUDE(FTSRCommonParameters, CommonParameters)
-		SHADER_PARAMETER(FMatrix44f, ClipToResurrectionClip)
-
-		SHADER_PARAMETER_RDG_TEXTURE(Texture2D, ClosestDepthTexture)
-		SHADER_PARAMETER_RDG_TEXTURE(Texture2D, HistoryRejectionTexture)
-
-		SHADER_PARAMETER_RDG_TEXTURE_UAV(RWTexture2D, DilatedVelocityOutput)
-		SHADER_PARAMETER_RDG_TEXTURE_UAV(RWTexture2DArray, DebugOutput)
-	END_SHADER_PARAMETER_STRUCT()
-}; // class FTSRMergeRejectionCS
-
 class FTSRSpatialAntiAliasingCS : public FTSRShader
 {
 	DECLARE_GLOBAL_SHADER(FTSRSpatialAntiAliasingCS);
@@ -1068,7 +1054,6 @@ IMPLEMENT_GLOBAL_SHADER(FTSRForwardScatterDepthCS,   "/Engine/Private/TemporalSu
 IMPLEMENT_GLOBAL_SHADER(FTSRDilateVelocityCS,        "/Engine/Private/TemporalSuperResolution/TSRDilateVelocity.usf",        "MainCS", SF_Compute);
 IMPLEMENT_GLOBAL_SHADER(FTSRDecimateHistoryCS,       "/Engine/Private/TemporalSuperResolution/TSRDecimateHistory.usf",       "MainCS", SF_Compute);
 IMPLEMENT_GLOBAL_SHADER(FTSRRejectShadingCS,         "/Engine/Private/TemporalSuperResolution/TSRRejectShading.usf",         "MainCS", SF_Compute);
-IMPLEMENT_GLOBAL_SHADER(FTSRMergeRejectionCS,        "/Engine/Private/TemporalSuperResolution/TSRMergeRejection.usf",        "MainCS", SF_Compute);
 IMPLEMENT_GLOBAL_SHADER(FTSRSpatialAntiAliasingCS,   "/Engine/Private/TemporalSuperResolution/TSRSpatialAntiAliasing.usf",   "MainCS", SF_Compute);
 IMPLEMENT_GLOBAL_SHADER(FTSRUpdateHistoryCS,         "/Engine/Private/TemporalSuperResolution/TSRUpdateHistory.usf",         "MainCS", SF_Compute);
 IMPLEMENT_GLOBAL_SHADER(FTSRResolveHistoryCS,        "/Engine/Private/TemporalSuperResolution/TSRResolveHistory.usf",        "MainCS", SF_Compute);
@@ -2175,6 +2160,7 @@ FDefaultTemporalUpscaler::FOutputs AddTemporalSuperResolutionPasses(
 			FScreenTransform::ChangeTextureBasisFromTo(TranslucencyViewport, FScreenTransform::ETextureBasis::ViewportUV, FScreenTransform::ETextureBasis::TextureUV);
 		PassParameters->TranslucencyTextureUVMin = GetScreenPassTextureViewportParameters(TranslucencyViewport).UVViewportBilinearMin;
 		PassParameters->TranslucencyTextureUVMax = GetScreenPassTextureViewportParameters(TranslucencyViewport).UVViewportBilinearMax;
+		PassParameters->ClipToResurrectionClip = ClipToResurrectionClip;
 		PassParameters->HistoryGuideQuantizationError = ComputePixelFormatQuantizationError(History.GuideArray->Desc.Format);
 		PassParameters->FlickeringFramePeriod = FlickeringFramePeriod;
 		PassParameters->TheoricBlendFactor = 1.0f / (1.0f + MaxHistorySampleCount / OutputToInputResolutionFractionSquare);
@@ -2214,6 +2200,7 @@ FDefaultTemporalUpscaler::FOutputs AddTemporalSuperResolutionPasses(
 		}
 		PassParameters->DecimateMaskTexture = DecimateMaskTexture;
 		PassParameters->IsMovingMaskTexture = IsMovingMaskTexture;
+		PassParameters->ClosestDepthTexture = ClosestDepthTexture;
 
 		// Outputs
 		{
@@ -2265,6 +2252,9 @@ FDefaultTemporalUpscaler::FOutputs AddTemporalSuperResolutionPasses(
 			// Output how the history should rejected in the HistoryUpdate
 			PassParameters->HistoryRejectionOutput = GraphBuilder.CreateUAV(HistoryRejectionTexture);
 
+			// Amends how the history should be rejected
+			PassParameters->DilatedVelocityOutput = GraphBuilder.CreateUAV(DilatedVelocityTexture);
+
 			// Output the composed translucency and opaque scene color to speed up HistoryUpdate
 			PassParameters->InputSceneColorOutput = bComputeInputSceneColorTexture
 				? GraphBuilder.CreateUAV(InputSceneColorTexture)
@@ -2300,29 +2290,6 @@ FDefaultTemporalUpscaler::FOutputs AddTemporalSuperResolutionPasses(
 			FComputeShaderUtils::GetGroupCount(InputRect.Size(), TileSize));
 	}
 
-	// Commit to resurrection if selected by the shading heuristic
-	if (bCanResurrectHistory)
-	{
-		FTSRMergeRejectionCS::FParameters* PassParameters = GraphBuilder.AllocParameters<FTSRMergeRejectionCS::FParameters>();
-		PassParameters->CommonParameters = CommonParameters;
-		PassParameters->ClipToResurrectionClip = ClipToResurrectionClip;
-
-		PassParameters->ClosestDepthTexture = ClosestDepthTexture;
-		PassParameters->HistoryRejectionTexture = HistoryRejectionTexture;
-
-		PassParameters->DilatedVelocityOutput = GraphBuilder.CreateUAV(DilatedVelocityTexture);
-		PassParameters->DebugOutput = CreateDebugUAV(InputExtent, TEXT("Debug.TSR.MergeRejection"));
-
-		TShaderMapRef<FTSRMergeRejectionCS> ComputeShader(View.ShaderMap);
-		FComputeShaderUtils::AddPass(
-			GraphBuilder,
-			RDG_EVENT_NAME("TSR MergeRejection %dx%d", InputRect.Width(), InputRect.Height()),
-			AsyncComputePasses >= 3 ? ERDGPassFlags::AsyncCompute : ERDGPassFlags::Compute,
-			ComputeShader,
-			PassParameters,
-			FComputeShaderUtils::GetGroupCount(InputRect.Size(), 8));
-	}
-
 	// Spatial anti-aliasing when doing history rejection.
 	FRDGTextureRef AntiAliasingTexture = nullptr;
 	if (RejectionAntiAliasingQuality > 0)
