Cuts the number of permutations of the TSR's RejectShading pass.

diff --git a/Engine/Shaders/Private/TemporalSuperResolution/TSRCommon.ush b/Engine/Shaders/Private/TemporalSuperResolution/TSRCommon.ush
--- a/Engine/Shaders/Private/TemporalSuperResolution/TSRCommon.ush
+++ b/Engine/Shaders/Private/TemporalSuperResolution/TSRCommon.ush
@@ -14,7 +14,7 @@
 
 //------------------------------------------------------- RECOMPILE HASH
 
-#pragma message("UESHADERMETADATA_VERSION 2A42D2EF-728A-4AC9-825F-2DC009BE04B0")
+#pragma message("UESHADERMETADATA_VERSION 65292827-F021-4C62-9790-FC626A1BFE69")
 
 
 //------------------------------------------------------- COMPILER CONFIG
diff --git a/Engine/Shaders/Private/TemporalSuperResolution/TSRRejectShading.usf b/Engine/Shaders/Private/TemporalSuperResolution/TSRRejectShading.usf
--- a/Engine/Shaders/Private/TemporalSuperResolution/TSRRejectShading.usf
+++ b/Engine/Shaders/Private/TemporalSuperResolution/TSRRejectShading.usf
@@ -169,6 +169,8 @@ float FlickeringFramePeriod;
 float TheoricBlendFactor;
 uint TileOverscan;
 float PerceptionAdd;
+uint bEnableResurrection;
+uint bEnableFlickeringHeuristic;
 
 
 Texture2D<tsr_halfC> InputTexture;
@@ -927,20 +929,20 @@ void MainCS(
 	}
 		
 	// Analyses the moire pattern overtime.
-	tsr_simd_halfC OriginalOpaqueInput;
-	tsr_simd_half4 OriginalTranslucencyInput;
-	tsr_simd_half  VelocityEdge;
-	tsr_simd_bool  bIsDisoccluded;
-	tsr_simd_half  MoireError = tsr_simd_half::Const(0.0);
+	tsr_simd_half  MoireError;
 	#if CONFIG_MOIRE_DETECTION
+	BRANCH
+	if (bEnableFlickeringHeuristic)
 	PLATFORM_SPECIFIC_ISOLATE
 	{
 		// Load its own data
-		#if CONFIG_MOIRE_DETECTION
-			tsr_simd_half  MoireInput;
-			tsr_simd_half4 PrevMoireHistory;
-			tsr_simd_half  bIsStatic;
-		#endif
+		tsr_simd_half  VelocityEdge;
+		tsr_simd_bool  bIsDisoccluded;
+		tsr_simd_halfC OriginalOpaqueInput;
+		tsr_simd_half4 OriginalTranslucencyInput;
+		tsr_simd_half  MoireInput;
+		tsr_simd_half4 PrevMoireHistory;
+		tsr_simd_half  bIsStatic;
 		PLATFORM_SPECIFIC_ISOLATE
 		{
 			tsr_simd_short2 PixelPos = ComputePixelPos();
@@ -973,9 +975,17 @@ void MainCS(
 			HistoryMoireOutput[tsr_short3(OutputPixelPos, 0)] = MoireHistory.GetElement(ElementIndex);
 		}
 	}
+	else
 	#endif // CONFIG_MOIRE_DETECTION
+	{
+		MoireError = tsr_simd_half::Const(0.0);
+	}
 	
 	// Load all the history data.
+	tsr_simd_half  VelocityEdge;
+	tsr_simd_bool  bIsDisoccluded;
+	tsr_simd_halfC OriginalOpaqueInput;
+	tsr_simd_half4 OriginalTranslucencyInput;
 	tsr_simd_halfC OriginalHistory;
 	tsr_simd_half  OriginalHistoryUncertainty;
 	#if CONFIG_HISTORY_RESURRECTION
@@ -986,10 +996,8 @@ void MainCS(
 	{
 		tsr_simd_short2 PixelPos = ComputePixelPos();
 
-		#if !CONFIG_MOIRE_DETECTION || 1
-			FetchDecimateMask(PixelPos, /* out */ VelocityEdge, /* out */ bIsDisoccluded);
-			FetchSceneColorAndTranslucency(PixelPos, /* out */ OriginalOpaqueInput, /* out */ OriginalTranslucencyInput);
-		#endif
+		FetchDecimateMask(PixelPos, /* out */ VelocityEdge, /* out */ bIsDisoccluded);
+		FetchSceneColorAndTranslucency(PixelPos, /* out */ OriginalOpaqueInput, /* out */ OriginalTranslucencyInput);
 		FetchHistoryGuide(
 			ReprojectedHistoryGuideTexture,
 			ReprojectedHistoryGuideMetadataTexture,
@@ -1043,7 +1051,7 @@ void MainCS(
 	tsr_simd_half RejectionClampBlend = tsr_simd_half::Const(1.0);
 	tsr_simd_bool bAliasingIsVisible = tsr_simd_bool::Const(false);
 	#if CONFIG_HISTORY_RESURRECTION
-		tsr_simd_bool bResurrectionIsCloser = tsr_simd_bool::Const(false);
+		tsr_simd_bool bResurrectionIsCloser;
 	#endif
 	PLATFORM_SPECIFIC_ISOLATE
 	{
@@ -1060,12 +1068,20 @@ void MainCS(
 			
 		// Check whether the history resurrection is closer to input, to discard and no pay the register pressure of OriginalResurrectedHistory
 		#if CONFIG_HISTORY_RESURRECTION
+		BRANCH
+		if (bEnableResurrection)
+		{
 			tsr_simd_halfC ExposedResurrectedHistory = ChangePerceptionAdd(
 				OriginalResurrectedHistory,
 				/* SourcePerceptionAdd = */ kHistoryGuidePreceptionAdd,
 				/* DestPerceptionAdd   = */ PerceptionAdd);
 				
 			bResurrectionIsCloser = IsResurrectedFrameCloserThanPreviousFrame(ExposedInput, ExposedHistory, ExposedResurrectedHistory);
+		}
+		else
+		{
+			bResurrectionIsCloser = tsr_simd_bool::Const(false);
+		}
 		#endif
 		
 		// Mutual Annihilation
diff --git a/Engine/Source/Runtime/Renderer/Private/PostProcess/TemporalSuperResolution.cpp b/Engine/Source/Runtime/Renderer/Private/PostProcess/TemporalSuperResolution.cpp
--- a/Engine/Source/Runtime/Renderer/Private/PostProcess/TemporalSuperResolution.cpp
+++ b/Engine/Source/Runtime/Renderer/Private/PostProcess/TemporalSuperResolution.cpp
@@ -422,28 +422,14 @@ public:
 
 	static bool ShouldCompile32or16BitPermutation(EShaderPlatform Platform, bool bIs16BitVALUPermutation)
 	{
-		// Always compile 32bit ops on preview platform 
-		if (FDataDrivenShaderPlatformInfo::GetIsPreviewPlatform(Platform) && !bIs16BitVALUPermutation)
+		// Always compile the 32bit permutations for the alpha channel
+		if (!bIs16BitVALUPermutation)
 		{
 			return true;
 		}
 
 		const ERHIFeatureSupport Support = FTSRShader::Supports16BitVALU(Platform);
-
-		if (Support == ERHIFeatureSupport::RuntimeGuaranteed)
-		{
-			// Only compile the 16bit permutation
-			return bIs16BitVALUPermutation;
-		}
-		else if (Support == ERHIFeatureSupport::RuntimeDependent)
-		{
-			// Compile both the 32bit and 16bit permutations
-			return true;
-		}
-		else
-		{
-			return !bIs16BitVALUPermutation;
-		}
+		return Support != ERHIFeatureSupport::Unsupported;
 	}
 
 	static ERHIFeatureSupport SupportsWaveOps(EShaderPlatform Platform)
@@ -691,7 +677,8 @@ class FTSRRejectShadingCS : public FTSRShader
 		SHADER_PARAMETER(float, TheoricBlendFactor)
 		SHADER_PARAMETER(int32, TileOverscan)
 		SHADER_PARAMETER(float, PerceptionAdd)
-		SHADER_PARAMETER(int32, bOverrideWithResurrectionTexture)
+		SHADER_PARAMETER(int32, bEnableResurrection)
+		SHADER_PARAMETER(int32, bEnableFlickeringHeuristic)
 
 		SHADER_PARAMETER_RDG_TEXTURE(Texture2D, InputTexture)
 		SHADER_PARAMETER_RDG_TEXTURE(Texture2D, InputMoireLumaTexture)
@@ -719,12 +706,37 @@ class FTSRRejectShadingCS : public FTSRShader
 	{
 		int32 WaveSize = PermutationVector.Get<FWaveSizeOps>();
 
-		// WaveSize is for Intel Arc GPU which also supports 16bits ops, so compiling WaveSize=16 32bit ops is useless and should instead fall back to WaveSize=0.
+		// WaveSize=16 is for Intel Arc GPU which also supports 16bits ops, so compiling WaveSize=16 32bit ops is useless and should instead fall back to WaveSize=0.
 		if (WaveSize == 16 && !PermutationVector.Get<FTSRShader::F16BitVALUDim>())
 		{
 			PermutationVector.Set<FWaveSizeOps>(0);
 		}
 
+		// Only compile the alpha channel with 32bit ops, as this is mostly targeting enterprise uses on Quadro GPUs
+		if (PermutationVector.Get<FTSRShader::FAlphaChannelDim>())
+		{
+			PermutationVector.Set<FTSRShader::F16BitVALUDim>(false);
+		}
+
+		// Optimising register pressure with 16bit for waveops that is 1 pixel/lane is pointless.
+		if (WaveSize == 0)
+		{
+			PermutationVector.Set<FTSRShader::F16BitVALUDim>(false);
+		}
+
+		// Register pressure is identical between all these permutation with 16bit
+		if (PermutationVector.Get<FTSRShader::F16BitVALUDim>())
+		{
+			PermutationVector.Set<FFlickeringDetectionDim>(true);
+			PermutationVector.Set<FHistoryResurrectionDim>(true);
+		}
+
+		// Flickering detection is on sg.AntiAliasQuality>=2 which also have resurrection.
+		if (PermutationVector.Get<FFlickeringDetectionDim>())
+		{
+			PermutationVector.Set<FHistoryResurrectionDim>(true);
+		}
+
 		return PermutationVector;
 	}
 
@@ -757,13 +769,6 @@ class FTSRRejectShadingCS : public FTSRShader
 				return false;
 			}
 		}
-		else // if (WaveSize == LDS fallback)
-		{
-			if (!FTSRShader::SupportsLDS(Parameters.Platform))
-			{
-				return false;
-			}
-		}
 
 		if (!FTSRShader::ShouldCompile32or16BitPermutation(Parameters.Platform, PermutationVector.Get<FTSRShader::F16BitVALUDim>()))
 		{
@@ -1285,8 +1290,7 @@ FDefaultTemporalUpscaler::FOutputs AddTemporalSuperResolutionPasses(
 
 	// Whether to use wave ops optimizations.
 	const ERHIFeatureSupport WaveOpsSupport = FTSRShader::SupportsWaveOps(View.GetShaderPlatform());
-	const bool bSupportsLDS = FTSRShader::SupportsLDS(View.GetShaderPlatform());
-	const bool bUseWaveOps = (CVarTSRWaveOps.GetValueOnRenderThread() != 0 && GRHISupportsWaveOperations && bSupportsLDS && (WaveOpsSupport == ERHIFeatureSupport::RuntimeDependent || WaveOpsSupport == ERHIFeatureSupport::RuntimeGuaranteed)) || !bSupportsLDS;
+	const bool bUseWaveOps = (CVarTSRWaveOps.GetValueOnRenderThread() != 0 && GRHISupportsWaveOperations && (WaveOpsSupport == ERHIFeatureSupport::RuntimeDependent || WaveOpsSupport == ERHIFeatureSupport::RuntimeGuaranteed));
 	const int32 WaveSizeOverride = bUseWaveOps ? CVarTSRWaveSize.GetValueOnAnyThread() : 0;
 	
 	// Whether to use 16bit VALU
@@ -2004,6 +2008,7 @@ FDefaultTemporalUpscaler::FOutputs AddTemporalSuperResolutionPasses(
 			ReprojectedHistoryGuideTexture = GraphBuilder.CreateTexture(Desc, TEXT("TSR.ReprojectedHistoryGuide"));
 		}
 
+		if (FlickeringFramePeriod > 0.0f)
 		{
 			FRDGTextureDesc Desc = FRDGTextureDesc::Create2DArray(
 				InputExtent,
@@ -2051,7 +2056,10 @@ FDefaultTemporalUpscaler::FOutputs AddTemporalSuperResolutionPasses(
 		PassParameters->ClipToResurrectionClip = ClipToResurrectionClip;
 
 		PassParameters->ReprojectedHistoryGuideOutput = GraphBuilder.CreateUAV(ReprojectedHistoryGuideTexture);
-		PassParameters->ReprojectedHistoryMoireOutput = GraphBuilder.CreateUAV(ReprojectedHistoryMoireTexture);
+		if (ReprojectedHistoryMoireTexture)
+		{
+			PassParameters->ReprojectedHistoryMoireOutput = GraphBuilder.CreateUAV(ReprojectedHistoryMoireTexture);
+		}
 		PassParameters->HoleFilledVelocityOutput = GraphBuilder.CreateUAV(HoleFilledVelocityTexture);
 		PassParameters->DecimateMaskOutput = GraphBuilder.CreateUAV(DecimateMaskTexture);
 		PassParameters->DebugOutput = CreateDebugUAV(InputExtent, TEXT("Debug.TSR.DecimateHistory"));
@@ -2166,6 +2174,8 @@ FDefaultTemporalUpscaler::FOutputs AddTemporalSuperResolutionPasses(
 		PassParameters->TheoricBlendFactor = 1.0f / (1.0f + MaxHistorySampleCount / OutputToInputResolutionFractionSquare);
 		PassParameters->TileOverscan = TileOverscan;
 		PassParameters->PerceptionAdd = FMath::Pow(0.5f, CVarTSRShadingExposureOffset.GetValueOnRenderThread());
+		PassParameters->bEnableResurrection = bCanResurrectHistory;
+		PassParameters->bEnableFlickeringHeuristic = FlickeringFramePeriod > 0.0f;
 
 		PassParameters->InputTexture = PassInputs.SceneColor.Texture;
 		if (PassInputs.FlickeringInputTexture.IsValid())
@@ -2186,7 +2196,7 @@ FDefaultTemporalUpscaler::FOutputs AddTemporalSuperResolutionPasses(
 				ReprojectedHistoryGuideTexture, /* SliceIndex = */ 0 * HistoryColorGuideSliceCountWithoutResurrection + 1));
 		}
 		PassParameters->ReprojectedHistoryMoireTexture = ReprojectedHistoryMoireTexture ? GraphBuilder.CreateSRV(FRDGTextureSRVDesc::CreateForSlice(
-			ReprojectedHistoryMoireTexture, /* SliceIndex = */ 0)) : nullptr;
+			ReprojectedHistoryMoireTexture, /* SliceIndex = */ 0)) : GraphBuilder.CreateSRV(FRDGTextureSRVDesc(BlackDummy));
 		if (bCanResurrectHistory)
 		{
 			PassParameters->ResurrectedHistoryGuideTexture = GraphBuilder.CreateSRV(FRDGTextureSRVDesc::CreateForSlice(
@@ -2198,8 +2208,13 @@ FDefaultTemporalUpscaler::FOutputs AddTemporalSuperResolutionPasses(
 					ReprojectedHistoryGuideTexture, /* SliceIndex = */ 1 * HistoryColorGuideSliceCountWithoutResurrection + 1));
 			}
 		}
+		else
+		{
+			PassParameters->ResurrectedHistoryGuideTexture = GraphBuilder.CreateSRV(FRDGTextureSRVDesc(BlackDummy));
+			PassParameters->ResurrectedHistoryGuideMetadataTexture = GraphBuilder.CreateSRV(FRDGTextureSRVDesc(BlackDummy));
+		}
 		PassParameters->DecimateMaskTexture = DecimateMaskTexture;
-		PassParameters->IsMovingMaskTexture = IsMovingMaskTexture;
+		PassParameters->IsMovingMaskTexture = IsMovingMaskTexture ? IsMovingMaskTexture : GraphBuilder.CreateSRV(FRDGTextureSRVDesc(BlackUintDummy));
 		PassParameters->ClosestDepthTexture = ClosestDepthTexture;
 
 		// Outputs
@@ -2220,7 +2235,7 @@ FDefaultTemporalUpscaler::FOutputs AddTemporalSuperResolutionPasses(
 			// Output history for the anti-flickering heuristic that know how something flicker overtime.
 			if (FlickeringFramePeriod == 0.0f)
 			{
-				// NOP
+				PassParameters->HistoryMoireOutput = CreateDummyUAVArray(GraphBuilder, History.MoireArray->Desc.Format);
 			}
 			else if (View.bStatePrevViewInfoIsReadOnly)
 			{
@@ -2252,7 +2267,7 @@ FDefaultTemporalUpscaler::FOutputs AddTemporalSuperResolutionPasses(
 			// Output how the history should rejected in the HistoryUpdate
 			PassParameters->HistoryRejectionOutput = GraphBuilder.CreateUAV(HistoryRejectionTexture);
 
-			// Amends how the history should be rejected
+			// Amends how the history should be reprojected
 			PassParameters->DilatedVelocityOutput = GraphBuilder.CreateUAV(DilatedVelocityTexture);
 
 			// Output the composed translucency and opaque scene color to speed up HistoryUpdate
@@ -2281,7 +2296,7 @@ FDefaultTemporalUpscaler::FOutputs AddTemporalSuperResolutionPasses(
 				int32(PermutationVector.Get<FTSRRejectShadingCS::FWaveSizeOps>()),
 				PassParameters->FlickeringFramePeriod,
 				PermutationVector.Get<FTSRShader::F16BitVALUDim>() ? TEXT("16bit") : TEXT("32bit"),
-				PermutationVector.Get<FTSRRejectShadingCS::FHistoryResurrectionDim>() ? TEXT(" Resurrection") : TEXT(""),
+				PassParameters->bEnableResurrection ? TEXT(" Resurrection") : TEXT(""),
 				PermutationVector.Get<FTSRShader::FAlphaChannelDim>() ? TEXT(" AlphaChannel") : TEXT(""),
 				InputRect.Width(), InputRect.Height()),
 			AsyncComputePasses >= 3 ? ERDGPassFlags::AsyncCompute : ERDGPassFlags::Compute,
