Path Tracing: Implement Holdout flags for volumes

To avoid having too many parameters, consolidate the path tracer's boolean enable/holdout flags into a single bitmask. Simplify some of the logic for heterogeneous volume lookups to avoid redundant copies of the code in the inner loop.

diff --git a/Engine/Shaders/Private/PathTracing/PathTracingCore.ush b/Engine/Shaders/Private/PathTracing/PathTracingCore.ush
--- a/Engine/Shaders/Private/PathTracing/PathTracingCore.ush
+++ b/Engine/Shaders/Private/PathTracing/PathTracingCore.ush
@@ -387,7 +387,7 @@ FPathTracingPayload TraceTransparentRay(inout FPathState PathState, int Bounce,
 		{
 			// we didn't hit any real geometry, but we did hit the volume's blocker geometry
 			// create a virtual hit with it here
-			PackedPayload = VolumeGetBlockerHit(VolumeIsectList.BlockerID, PathState.Ray.Origin, PathState.Ray.Direction, VolumeIsectList.BlockerHitT);
+			PackedPayload = VolumeGetBlockerHit(VolumeIsectList.BlockerID, PathState.Ray.Origin, PathState.Ray.Direction, VolumeIsectList.BlockerHitT, bIsCameraRay);
 		}
 
 		// Loop over lights to capture their contribution
@@ -452,8 +452,9 @@ FPathTracingPayload TraceTransparentRay(inout FPathState PathState, int Bounce,
 			}
 
 			// NOTE: when using analytic raymarching, we can skip all raymarching on the last bounce because we know we don't need a scattering point
-			const uint UseNullTrackingTransmittance = !UseAnalyticTransmittance;
-			if (UseNullTrackingTransmittance || (UseAnalyticTransmittance && !bLastBounce))
+			const bool bUseAnalyticTransmittance = (VolumeFlags & PATH_TRACER_VOLUME_USE_ANALYTIC_TRANSMITTANCE) != 0;
+			const bool bUseNullTrackingTransmittance = !bUseAnalyticTransmittance;
+			if (bUseNullTrackingTransmittance || (bUseAnalyticTransmittance && !bLastBounce))
 			{
 				// remember the start of our interval before ray marching messes with it
 				float StartVolumeTMin = Interval.VolumeTMin;
@@ -478,13 +479,18 @@ FPathTracingPayload TraceTransparentRay(inout FPathState PathState, int Bounce,
 						float3 Ro = PathState.Ray.Origin;
 						float3 Rd = PathState.Ray.Direction;
 						float3 TranslatedWorldPos = Ro + Interval.VolumeTMin * Rd;
-						FVolumeShadedResult Result = VolumeGetDensity(TranslatedWorldPos, Interval);
+						FVolumeShadedResult Result = VolumeGetDensity(TranslatedWorldPos, Interval, bIsCameraRay);
 
 						// clamp to make sure we never exceed the majorant (should not be the case, but need to avoid any possible numerical issues)
 						float3 SigmaT = min(Result.SigmaT, TrackingResult.SigmaBar);
 						float3 SigmaN = TrackingResult.SigmaBar - SigmaT;
 
-						if (UseAnalyticTransmittance || (UseNullTrackingTransmittance && !bLastBounce)) // check now if this loop is also computing transmittance
+						if (bIsCameraRay)
+						{
+							PathState.BackgroundVisibility += Luminance(PathThroughput * Result.SigmaH);
+						}
+
+						if (bUseAnalyticTransmittance || (bUseNullTrackingTransmittance && !bLastBounce)) // check now if this loop is also computing transmittance
 						{
 							float3 SigmaS = min(Result.SigmaSRayleigh + Result.SigmaSHG, SigmaT);
 							// accumulate a signal for the denoiser
@@ -516,14 +522,14 @@ FPathTracingPayload TraceTransparentRay(inout FPathState PathState, int Bounce,
 
 				// restore interval start now that we are finished with ray marching
 				Interval.VolumeTMin = StartVolumeTMin;
-				if (UseNullTrackingTransmittance)
+				if (bUseNullTrackingTransmittance)
 				{
 					// record the change in transmittance since we already computed it
 					PathState.PathThroughput = PathThroughput;
 				}
 			}
 			
-			if (UseAnalyticTransmittance)
+			if (bUseAnalyticTransmittance)
 			{
 				// analytically handle the transmittance to the next surface (which will be higher quality than the transmittance implicitly computed above and prevent noise with emissive surfaces in volumes)
 				PathState.PathThroughput = VolumeGetTransmittance(PathState.PathThroughput, PathState.Ray.Origin, PathState.Ray.Direction, Interval, PathState.RandSequence);
@@ -621,7 +627,8 @@ FPathTracingPayload TraceTransparentRay(inout FPathState PathState, int Bounce,
 
 		if (HitPayload.HitT == POSITIVE_INFINITY)
 		{
-			// if our hit was against an infinite light, exit now
+			// if our hit was against an infinite light, treat it as having background alpha, and exit
+			PathState.BackgroundVisibility += Luminance(PathState.PathThroughput);
 			break;
 		}
 
@@ -1394,7 +1401,7 @@ bool PathTracingKernel(inout FPathState PathState, int Bounce)
 
 			// find out how much volume exists at the current point
 			const float3 TranslatedWorldPos = Ro + EquiAngularT * Rd;
-			const FVolumeShadedResult Result = VolumeGetDensity(TranslatedWorldPos, VolumeSegment.Interval);
+			const FVolumeShadedResult Result = VolumeGetDensity(TranslatedWorldPos, VolumeSegment.Interval, bIsCameraRay);
 			const float3 SigmaT = Result.SigmaT;
 			const float3 SigmaS = min(Result.SigmaSRayleigh + Result.SigmaSHG, SigmaT);
 
diff --git a/Engine/Shaders/Private/PathTracing/Volume/PathTracingAtmosphere.ush b/Engine/Shaders/Private/PathTracing/Volume/PathTracingAtmosphere.ush
--- a/Engine/Shaders/Private/PathTracing/Volume/PathTracingAtmosphere.ush
+++ b/Engine/Shaders/Private/PathTracing/Volume/PathTracingAtmosphere.ush
@@ -114,11 +114,6 @@ float3 RayPlanetIntersectPrecise(float3 Ro, float3 Rd)
 
 FVolumeIntersection AtmosphereIntersect(float3 Origin, float3 Direction, float TMin, float TMax)
 {
-	// NOTE: only this call needs to be guarded at the moment since the others can't be reached if this returns nothing
-	if (!EnableAtmosphere)
-	{
-		return CreateEmptyVolumeIntersection();
-	}
 	float3 Result = RayPlanetIntersectPrecise(Origin, Direction);
 	// clip to provided ray interval
 	return CreateVolumeIntersection(
@@ -127,7 +122,7 @@ FVolumeIntersection AtmosphereIntersect(float3 Origin, float3 Direction, float T
 		Result.x > TMin && Result.x < TMax ? Result.x : -1.0);
 }
 
-FPackedPathTracingPayload AtmosphereGetBlockerHit(float3 Origin, float3 Direction, float HitT)
+FPackedPathTracingPayload AtmosphereGetBlockerHit(float3 Origin, float3 Direction, float HitT, bool bIsCameraRay)
 {
 	float3 Center = PlanetCenterTranslatedWorldHi + PlanetCenterTranslatedWorldLo;
 	float3 WorldP = Origin + HitT * Direction;
@@ -141,6 +136,11 @@ FPackedPathTracingPayload AtmosphereGetBlockerHit(float3 Origin, float3 Directio
 	Payload.HitT = HitT;
 	Payload.PrimitiveLightingChannelMask = 7; // visible to all lights
 	Payload.SetFrontFace();
+	if (bIsCameraRay && (VolumeFlags & PATH_TRACER_VOLUME_HOLDOUT_ATMOSPHERE) != 0)
+	{
+		Payload.SetHoldout();
+		Payload.ShadingModelID = SHADINGMODELID_UNLIT;
+	}
 	return PackPathTracingPayload(Payload);
 }
 
diff --git a/Engine/Shaders/Private/PathTracing/Volume/PathTracingHeterogeneousVolumes.ush b/Engine/Shaders/Private/PathTracing/Volume/PathTracingHeterogeneousVolumes.ush
--- a/Engine/Shaders/Private/PathTracing/Volume/PathTracingHeterogeneousVolumes.ush
+++ b/Engine/Shaders/Private/PathTracing/Volume/PathTracingHeterogeneousVolumes.ush
@@ -20,34 +20,30 @@ float GetPhaseG()
 
 FVolumeIntersection HeterogeneousVolumesIntersect(float3 Origin, float3 Direction, float TMin, float TMax)
 {
-	if (EnableHeterogeneousVolumes)
+	float3 WorldBoundsMin = FrustumGridUniformBuffer.TopLevelGridWorldBoundsMin;
+	float3 WorldBoundsMax = FrustumGridUniformBuffer.TopLevelGridWorldBoundsMax;
+	if (OrthoGridUniformBuffer.bUseOrthoGrid)
 	{
-		float3 WorldBoundsMin = FrustumGridUniformBuffer.TopLevelGridWorldBoundsMin;
-		float3 WorldBoundsMax = FrustumGridUniformBuffer.TopLevelGridWorldBoundsMax;
-		if (OrthoGridUniformBuffer.bUseOrthoGrid)
-		{
-			WorldBoundsMin = OrthoGridUniformBuffer.TopLevelGridWorldBoundsMin;
-			WorldBoundsMax = OrthoGridUniformBuffer.TopLevelGridWorldBoundsMax;
-		}
-
-		if (OrthoGridUniformBuffer.bUseOrthoGrid || FrustumGridUniformBuffer.bUseFrustumGrid)
-		{
-			float3 TranslatedWorldBoundsMin = GetTranslatedWorldPos(WorldBoundsMin);
-			float3 TranslatedWorldBoundsMax = GetTranslatedWorldPos(WorldBoundsMax);
-
-			float2 RayHitT = IntersectAABB(
-				Origin,
-				Direction,
-				TMin,
-				TMax,
-				TranslatedWorldBoundsMin,
-				TranslatedWorldBoundsMax
-			);
-
-			return CreateVolumeIntersection(RayHitT.x, RayHitT.y);
-		}
+		WorldBoundsMin = OrthoGridUniformBuffer.TopLevelGridWorldBoundsMin;
+		WorldBoundsMax = OrthoGridUniformBuffer.TopLevelGridWorldBoundsMax;
 	}
 
+	if (OrthoGridUniformBuffer.bUseOrthoGrid || FrustumGridUniformBuffer.bUseFrustumGrid)
+	{
+		float3 TranslatedWorldBoundsMin = GetTranslatedWorldPos(WorldBoundsMin);
+		float3 TranslatedWorldBoundsMax = GetTranslatedWorldPos(WorldBoundsMax);
+
+		float2 RayHitT = IntersectAABB(
+			Origin,
+			Direction,
+			TMin,
+			TMax,
+			TranslatedWorldBoundsMin,
+			TranslatedWorldBoundsMax
+		);
+
+		return CreateVolumeIntersection(RayHitT.x, RayHitT.y);
+	}
 	return CreateEmptyVolumeIntersection();
 }
 
@@ -175,29 +171,15 @@ FVolumeShadedResult HeterogeneousVolumesGetDensity(float3 TranslatedWorldPos)
 {
 	FVolumeShadedResult Result = (FVolumeShadedResult) 0;
 
-	if (EnableHeterogeneousVolumes)
+	bool bInFrustum = false;
+	if (FrustumGridUniformBuffer.bUseFrustumGrid)
 	{
-		bool bUseBothGrids = OrthoGridUniformBuffer.bUseOrthoGrid && FrustumGridUniformBuffer.bUseFrustumGrid;
-		if (bUseBothGrids)
-		{
-			bool bInFrustum = false;
-			Result = GetFrustumVoxelGridDensity(TranslatedWorldPos, bInFrustum);
-			if (!bInFrustum)
-			{
-				Result = GetOrthoVoxelGridDensity(TranslatedWorldPos);
-			}
-		}
-		else if (FrustumGridUniformBuffer.bUseFrustumGrid)
-		{
-			bool bInFrustum = false;
-			Result = GetFrustumVoxelGridDensity(TranslatedWorldPos, bInFrustum);
-		}
-		else if (OrthoGridUniformBuffer.bUseOrthoGrid)
-		{
-			Result = GetOrthoVoxelGridDensity(TranslatedWorldPos);
-		}
+		Result = GetFrustumVoxelGridDensity(TranslatedWorldPos, bInFrustum);
+	}
+	if (!bInFrustum && OrthoGridUniformBuffer.bUseOrthoGrid)
+	{
+		Result = GetOrthoVoxelGridDensity(TranslatedWorldPos);
 	}
-
 	return Result;
 }
 
diff --git a/Engine/Shaders/Private/PathTracing/Volume/PathTracingVolume.ush b/Engine/Shaders/Private/PathTracing/Volume/PathTracingVolume.ush
--- a/Engine/Shaders/Private/PathTracing/Volume/PathTracingVolume.ush
+++ b/Engine/Shaders/Private/PathTracing/Volume/PathTracingVolume.ush
@@ -8,37 +8,34 @@
 #include "PathTracingHeterogeneousVolumes.ush"
 #include "PathTracingVolumeSampling.ush"
 
-uint UseAnalyticTransmittance;
-
-uint EnableAtmosphere;
-uint EnableFog;
-uint EnableHeterogeneousVolumes;
+uint VolumeFlags;
 int MaxRaymarchSteps;
 
 
 // Given an input ray, figure out where the provided volume intersects it
 FVolumeIntersectionList VolumeIntersect(float3 Origin, float3 Direction, float TMin, float TMax)
 {
+	// NOTE: only this call needs to be guarded by enable flags at the moment since the others can't be reached if this returns nothing
 	FVolumeIntersectionList Result = CreateEmptyVolumeIntersectionList();
-	if (EnableAtmosphere)
+	if (VolumeFlags & PATH_TRACER_VOLUME_ENABLE_ATMOSPHERE)
 	{
 		Result.Add(VOLUMEID_ATMOSPHERE, AtmosphereIntersect(Origin, Direction, TMin, TMax));
 	}
-	if (EnableFog)
+	if (VolumeFlags & PATH_TRACER_VOLUME_ENABLE_FOG)
 	{
 		Result.Add(VOLUMEID_FOG, FogIntersect(Origin, Direction, TMin, TMax));
 	}
-	if (EnableHeterogeneousVolumes)
+	if (VolumeFlags & PATH_TRACER_VOLUME_ENABLE_HETEROGENEOUS_VOLUMES)
 	{
 		Result.Add(VOLUMEID_HETEROGENEOUS_VOLUMES, HeterogeneousVolumesIntersect(Origin, Direction, TMin, TMax));
 	}
 	return Result;
 }
 
-FPackedPathTracingPayload VolumeGetBlockerHit(uint VolumeID, float3 Origin, float3 Direction, float HitT)
+FPackedPathTracingPayload VolumeGetBlockerHit(uint VolumeID, float3 Origin, float3 Direction, float HitT, bool bIsCameraRay)
 {
 	// NOTE: only atmosphere supports blockers, so if this gets called it is safe to assume its for the atmosphere volume
-	return AtmosphereGetBlockerHit(Origin, Direction, HitT);
+	return AtmosphereGetBlockerHit(Origin, Direction, HitT, bIsCameraRay);
 }
 
 // Given an input ray, figure out the bounds on density. The T interval may be smaller than the one returned by VolumeIntersect
@@ -58,17 +55,25 @@ FVolumeDensityBounds VolumeGetDensityBounds(float3 Origin, float3 Direction, FVo
 }
 
 // Given a point in world space, return the amount of volume and its scattering properties
-FVolumeShadedResult VolumeGetDensity(float3 TranslatedWorldPos, FVolumeIntersectionInterval Interval)
+FVolumeShadedResult VolumeGetDensity(float3 TranslatedWorldPos, FVolumeIntersectionInterval Interval, bool bIsCameraRay)
 {
 	FVolumeShadedResult Result = (FVolumeShadedResult)0;
 	for (int Index = 0; Index < Interval.Num; Index++)
 	{
+		FVolumeShadedResult LocalResult = (FVolumeShadedResult)0;
 		switch (Interval.VolumeID[Index])
 		{
-			case VOLUMEID_ATMOSPHERE:            MergeVolumeShadedResult(Result, AtmosphereGetDensity(TranslatedWorldPos)); break;
-			case VOLUMEID_FOG:                   MergeVolumeShadedResult(Result, FogGetDensity(TranslatedWorldPos)); break;
-			case VOLUMEID_HETEROGENEOUS_VOLUMES: MergeVolumeShadedResult(Result, HeterogeneousVolumesGetDensity(TranslatedWorldPos)); break;
+			case VOLUMEID_ATMOSPHERE:            LocalResult = AtmosphereGetDensity(TranslatedWorldPos); break;
+			case VOLUMEID_FOG:                   LocalResult = FogGetDensity(TranslatedWorldPos); break;
+			case VOLUMEID_HETEROGENEOUS_VOLUMES: LocalResult = HeterogeneousVolumesGetDensity(TranslatedWorldPos); break;
+		}
+		if (bIsCameraRay && (VolumeFlags & PATH_TRACER_VOLUME_HOLDOUT_BIT << Interval.VolumeID[Index]) != 0)
+		{
+			// remove scattering and set holdout density
+			LocalResult.SigmaH = LocalResult.SigmaT;
+			LocalResult.SigmaSHG = LocalResult.SigmaSRayleigh = 0;
 		}
+		MergeVolumeShadedResult(Result, LocalResult);
 	}
 	return Result;
 }
@@ -213,7 +218,7 @@ float3 VolumeGetTransmittanceNullTracking(float3 StartThroughput, float3 Origin,
 
 float3 VolumeGetTransmittance(float3 StartThroughput, float3 Origin, float3 Direction, FVolumeIntersectionInterval Interval, inout RandomSequence RandSequence)
 {
-	if (UseAnalyticTransmittance)
+	if (VolumeFlags & PATH_TRACER_VOLUME_USE_ANALYTIC_TRANSMITTANCE)
 	{
 		return VolumeGetTransmittanceAnalytical(StartThroughput, Origin, Direction, Interval, RandSequence);
 	}
@@ -223,7 +228,7 @@ float3 VolumeGetTransmittance(float3 StartThroughput, float3 Origin, float3 Dire
 
 float3 VolumeGetTransmittance(float3 StartThroughput, float3 Origin, float3 Direction, FVolumeIntersectionList IntersectionList, inout RandomSequence RandSequence)
 {
-	if (UseAnalyticTransmittance)
+	if (VolumeFlags & PATH_TRACER_VOLUME_USE_ANALYTIC_TRANSMITTANCE)
 	{
 		return VolumeGetTransmittanceAnalytical(StartThroughput, Origin, Direction, IntersectionList, RandSequence);
 	}
@@ -235,7 +240,7 @@ FVolumeTrackingResult VolumeSampleDistance(float3 PathThroughput, float3 Origin,
 {
 	FVolumeDensityBounds VolumeDensityBounds = VolumeGetDensityBounds(Origin, Direction, Interval);
 	float3 SigmaBar = VolumeDensityBounds.SigmaMax;
-	if (!UseAnalyticTransmittance)
+	if ((VolumeFlags & PATH_TRACER_VOLUME_USE_ANALYTIC_TRANSMITTANCE) == 0)
 	{
 		// if we are not using analytical transmittance, a tight majorant could prevent us from "seeing" hits that match the density exactly
 		// leading to heavy noise on bright objects embedded in the volume
diff --git a/Engine/Shaders/Private/PathTracing/Volume/PathTracingVolumeCommon.ush b/Engine/Shaders/Private/PathTracing/Volume/PathTracingVolumeCommon.ush
--- a/Engine/Shaders/Private/PathTracing/Volume/PathTracingVolumeCommon.ush
+++ b/Engine/Shaders/Private/PathTracing/Volume/PathTracingVolumeCommon.ush
@@ -33,10 +33,6 @@ FVolumeIntersection CreateEmptyVolumeIntersection()
 	return (FVolumeIntersection)0;
 }
 
-#define VOLUMEID_ATMOSPHERE				0
-#define VOLUMEID_FOG					1
-#define VOLUMEID_HETEROGENEOUS_VOLUMES	2
-#define PATH_TRACER_MAX_VOLUMES			3
 
 // representes an interval along the ray over which the specified volumes are known to exist
 struct FVolumeIntersectionInterval
@@ -188,6 +184,7 @@ void MergeVolumeDensityBounds(inout FVolumeDensityBounds Merged, FVolumeDensityB
 struct FVolumeShadedResult
 {
 	float3 SigmaT;         // extinction
+	float3 SigmaH;         // holdout density
 	float3 SigmaSRayleigh; // Rayleight scattering coefficient
 	float3 SigmaSHG;       // Henyey-Greenstein scattering coefficient
 	float3 Emission;
@@ -215,6 +212,7 @@ FPathTracingPayload CreateMediumHitPayload(float HitT, float3 TranslatedWorldPos
 void MergeVolumeShadedResult(inout FVolumeShadedResult Merged, FVolumeShadedResult Other)
 {
 	Merged.SigmaT += Other.SigmaT;
+	Merged.SigmaH += Other.SigmaH;
 	Merged.SigmaSRayleigh += Other.SigmaSRayleigh;
 	Merged.SigmaSHG += Other.SigmaSHG;
 	Merged.Emission += Other.Emission;
@@ -224,23 +222,6 @@ void MergeVolumeShadedResult(inout FVolumeShadedResult Merged, FVolumeShadedResu
 	Merged.PhaseG = lerp(Merged.PhaseG, Other.PhaseG, Qs > 0 ? Qa / Qs : 0.0);
 }
 
-struct FVolumeSampleResult
-{
-	float3 SurfaceThroughput;  // the path throughput beyond the volume (for the surface hit behind usually)
-	float3 VolumeThroughput;   // the path throughput up to the volume sample (if we sampled one)
-	float  VolumeT;            // distance to volume sample along the ray (will be <= 0 if none was sampled)
-	FVolumeShadedResult VolumeSample; // actual shaded volume hit (if we have one)
-
-	bool VolumeHit() { return VolumeT > 0.0; }
-};
-
-FVolumeSampleResult CreateEmptyVolumeSampleResult()
-{
-	FVolumeSampleResult Result = (FVolumeSampleResult)0;
-	Result.VolumeT = -1.0;
-	return Result;
-}
-
 // represent a stochastically chosen volume element that is guaranteed to lie between 
 struct FVolumeSegment
 {
diff --git a/Engine/Shaders/Shared/PathTracingDefinitions.h b/Engine/Shaders/Shared/PathTracingDefinitions.h
--- a/Engine/Shaders/Shared/PathTracingDefinitions.h
+++ b/Engine/Shaders/Shared/PathTracingDefinitions.h
@@ -48,3 +48,21 @@
 
 // Constants for the energy conservation texture sizes
 #define PATHTRACER_ENERGY_TABLE_RESOLUTION				32
+
+// Constants related to volumetric support
+#define VOLUMEID_ATMOSPHERE				0
+#define VOLUMEID_FOG					1
+#define VOLUMEID_HETEROGENEOUS_VOLUMES	2
+#define PATH_TRACER_MAX_VOLUMES			3 // NOTE: Adjust _BIT constants below if this ever grows
+
+#define PATH_TRACER_VOLUME_ENABLE_BIT						(1u)
+#define PATH_TRACER_VOLUME_ENABLE_ATMOSPHERE  				(PATH_TRACER_VOLUME_ENABLE_BIT << VOLUMEID_ATMOSPHERE)
+#define PATH_TRACER_VOLUME_ENABLE_FOG               		(PATH_TRACER_VOLUME_ENABLE_BIT << VOLUMEID_FOG)
+#define PATH_TRACER_VOLUME_ENABLE_HETEROGENEOUS_VOLUMES 	(PATH_TRACER_VOLUME_ENABLE_BIT << VOLUMEID_HETEROGENEOUS_VOLUMES)
+
+#define PATH_TRACER_VOLUME_HOLDOUT_BIT						(8u)
+#define PATH_TRACER_VOLUME_HOLDOUT_ATMOSPHERE  				(PATH_TRACER_VOLUME_HOLDOUT_BIT << VOLUMEID_ATMOSPHERE)
+#define PATH_TRACER_VOLUME_HOLDOUT_FOG               		(PATH_TRACER_VOLUME_HOLDOUT_BIT << VOLUMEID_FOG)
+#define PATH_TRACER_VOLUME_HOLDOUT_HETEROGENEOUS_VOLUMES	(PATH_TRACER_VOLUME_HOLDOUT_BIT << VOLUMEID_HETEROGENEOUS_VOLUMES)
+
+#define PATH_TRACER_VOLUME_USE_ANALYTIC_TRANSMITTANCE       (64u)
diff --git a/Engine/Source/Runtime/Renderer/Private/PathTracing.cpp b/Engine/Source/Runtime/Renderer/Private/PathTracing.cpp
--- a/Engine/Source/Runtime/Renderer/Private/PathTracing.cpp
+++ b/Engine/Source/Runtime/Renderer/Private/PathTracing.cpp
@@ -433,13 +433,10 @@ BEGIN_SHADER_PARAMETER_STRUCT(FPathTracingData, )
 	SHADER_PARAMETER(uint32, VisualizeLightGrid)
 	SHADER_PARAMETER(uint32, VisualizeDecalGrid)
 	SHADER_PARAMETER(uint32, EnableDBuffer)
-	SHADER_PARAMETER(uint32, EnableAtmosphere)
-	SHADER_PARAMETER(uint32, EnableFog)
-	SHADER_PARAMETER(uint32, EnableHeterogeneousVolumes)
+	SHADER_PARAMETER(uint32, VolumeFlags)
 	SHADER_PARAMETER(uint32, EnabledDirectLightingContributions)   // PATHTRACER_CONTRIBUTION_*
 	SHADER_PARAMETER(uint32, EnabledIndirectLightingContributions) // PATHTRACER_CONTRIBUTION_*
 	SHADER_PARAMETER(uint32, ApplyDiffuseSpecularOverrides)
-	SHADER_PARAMETER(uint32, UseAnalyticTransmittance)
 	SHADER_PARAMETER(int32, MaxRaymarchSteps)
 	SHADER_PARAMETER(float, MaxPathIntensity)
 	SHADER_PARAMETER(float, MaxNormalBias)
@@ -487,10 +484,7 @@ struct FPathTracingConfig
 			PathTracingData.EnableDBuffer != Other.PathTracingData.EnableDBuffer ||
 			PathTracingData.MaxPathIntensity != Other.PathTracingData.MaxPathIntensity ||
 			PathTracingData.FilterWidth != Other.PathTracingData.FilterWidth ||
-			PathTracingData.EnableAtmosphere != Other.PathTracingData.EnableAtmosphere ||
-			PathTracingData.EnableFog != Other.PathTracingData.EnableFog ||
-			PathTracingData.EnableHeterogeneousVolumes != Other.PathTracingData.EnableHeterogeneousVolumes ||
-			PathTracingData.UseAnalyticTransmittance != Other.PathTracingData.UseAnalyticTransmittance ||
+			PathTracingData.VolumeFlags != Other.PathTracingData.VolumeFlags ||
 			PathTracingData.ApplyDiffuseSpecularOverrides != Other.PathTracingData.ApplyDiffuseSpecularOverrides ||
 			PathTracingData.EnabledDirectLightingContributions != Other.PathTracingData.EnabledDirectLightingContributions ||
 			PathTracingData.EnabledIndirectLightingContributions != Other.PathTracingData.EnabledIndirectLightingContributions ||
@@ -603,7 +597,7 @@ namespace PathTracing
 	}
 }
 
-static uint32 EvalUseAnalyticTransmittance(const FViewInfo& View)
+static bool EvalUseAnalyticTransmittance(const FViewInfo& View)
 {
 	int32 UseAnalyticTransmittance = CVarPathTracingUseAnalyticTransmittance.GetValueOnRenderThread();
 	if (UseAnalyticTransmittance < 0)
@@ -611,7 +605,7 @@ static uint32 EvalUseAnalyticTransmittance(const FViewInfo& View)
 		UseAnalyticTransmittance = !ShouldRenderHeterogeneousVolumesForView(View);
 	}
 
-	return uint32(UseAnalyticTransmittance);
+	return UseAnalyticTransmittance != 0;
 }
 
 // This function prepares the portion of shader arguments that may involve invalidating the path traced state
@@ -656,21 +650,25 @@ static void PreparePathTracingData(const FScene* Scene, const FViewInfo& View, F
 		PathTracingData.CameraLensRadius.Y = 0.5f * FocalLengthInCM / PPV.DepthOfFieldFstop;
 		PathTracingData.CameraLensRadius.X = PathTracingData.CameraLensRadius.Y / FMath::Clamp(PPV.DepthOfFieldSqueezeFactor, 1.0f, 2.0f);
 	}
-	PathTracingData.EnableAtmosphere =
+
+	// Merge all volume flags into one uint
+	PathTracingData.VolumeFlags = 0;
+	PathTracingData.VolumeFlags |=
 		ShouldRenderSkyAtmosphere(Scene, ShowFlags) && 
 		View.SkyAtmosphereUniformShaderParameters != nullptr &&
-		PPV.PathTracingEnableReferenceAtmosphere != 0;
-
-	PathTracingData.EnableFog = ShouldRenderFog(*View.Family)
+		PPV.PathTracingEnableReferenceAtmosphere != 0 ? PATH_TRACER_VOLUME_ENABLE_ATMOSPHERE : 0;
+	PathTracingData.VolumeFlags |= ShouldRenderFog(*View.Family)
 		&& Scene->ExponentialFogs.Num() > 0
 		&& Scene->ExponentialFogs[0].bEnableVolumetricFog
 		&& Scene->ExponentialFogs[0].VolumetricFogDistance > 0
 		&& Scene->ExponentialFogs[0].VolumetricFogExtinctionScale > 0
 		&& (Scene->ExponentialFogs[0].FogData[0].Density > 0 ||
-			Scene->ExponentialFogs[0].FogData[1].Density > 0);
+			Scene->ExponentialFogs[0].FogData[1].Density > 0) ? PATH_TRACER_VOLUME_ENABLE_FOG : 0;
+	PathTracingData.VolumeFlags |= ShouldRenderHeterogeneousVolumesForView(View) ? PATH_TRACER_VOLUME_ENABLE_HETEROGENEOUS_VOLUMES : 0;
+	PathTracingData.VolumeFlags |= View.SkyAtmosphereUniformShaderParameters != nullptr &&  View.SkyAtmosphereUniformShaderParameters->bHoldout ? PATH_TRACER_VOLUME_HOLDOUT_ATMOSPHERE : 0;
+	PathTracingData.VolumeFlags |= Scene->ExponentialFogs.Num() > 0 && Scene->ExponentialFogs[0].bHoldout ? PATH_TRACER_VOLUME_HOLDOUT_FOG : 0;
+	PathTracingData.VolumeFlags |= EvalUseAnalyticTransmittance(View) ? PATH_TRACER_VOLUME_USE_ANALYTIC_TRANSMITTANCE : 0;
 
-	PathTracingData.EnableHeterogeneousVolumes = ShouldRenderHeterogeneousVolumesForView(View);
-	PathTracingData.UseAnalyticTransmittance = EvalUseAnalyticTransmittance(View);
 	PathTracingData.EnableDBuffer = CVarPathTracingUseDBuffer.GetValueOnRenderThread();
 
 	PathTracingData.DecalRoughnessCutoff = PathTracing::UsesDecals(*View.Family) && View.bHasRayTracingDecals ? CVarPathTracingDecalRoughnessCutoff.GetValueOnRenderThread() : -1.0f;
@@ -2055,7 +2053,7 @@ void SetLightParameters(FRDGBuilder& GraphBuilder, FPathTracingRG::FParameters*
 
 	// Prepend SkyLight to light buffer since it is not part of the regular light list
 	// skylight should be excluded if we are using the reference atmosphere calculation (don't bother checking again if an atmosphere is present)
-	const bool bUseAtmosphere = PassParameters->PathTracingData.EnableAtmosphere != 0;
+	const bool bUseAtmosphere = (PassParameters->PathTracingData.VolumeFlags & PATH_TRACER_VOLUME_ENABLE_ATMOSPHERE) != 0;
 	const bool bEnableSkydome = !bUseAtmosphere;
 	if (PrepareSkyTexture(GraphBuilder, Scene, View, bEnableSkydome, UseMISCompensation, &PassParameters->SkylightParameters))
 	{
@@ -2597,7 +2595,7 @@ void FDeferredShadingSceneRenderer::RenderPathTracing(
 
 	// prepare atmosphere optical depth lookup texture (if needed)
 	FRDGTexture* AtmosphereOpticalDepthLUT = nullptr;
-	if (Config.PathTracingData.EnableAtmosphere)
+	if ((Config.PathTracingData.VolumeFlags & PATH_TRACER_VOLUME_ENABLE_ATMOSPHERE) != 0)
 	{
 		check(Scene->GetSkyAtmosphereSceneInfo() != nullptr);
 		check(Scene->GetSkyAtmosphereSceneInfo()->GetAtmosphereShaderParameters() != nullptr);
@@ -2917,7 +2915,7 @@ void FDeferredShadingSceneRenderer::RenderPathTracing(
 								PassParameters->PlanetCenterTranslatedWorldHi = PreviousPassParameters->PlanetCenterTranslatedWorldHi;
 								PassParameters->PlanetCenterTranslatedWorldLo = PreviousPassParameters->PlanetCenterTranslatedWorldLo;
 							}
-							else if (Config.PathTracingData.EnableAtmosphere)
+							else if ((Config.PathTracingData.VolumeFlags & PATH_TRACER_VOLUME_ENABLE_ATMOSPHERE) != 0)
 							{
 								PassParameters->Atmosphere = Scene->GetSkyAtmosphereSceneInfo()->GetAtmosphereUniformBuffer();
 								FVector PlanetCenterTranslatedWorld = Scene->GetSkyAtmosphereSceneInfo()->GetSkyAtmosphereSceneProxy().GetAtmosphereSetup().PlanetCenterKm * double(FAtmosphereSetup::SkyUnitToCm) + View.ViewMatrices.GetPreViewTranslation();
@@ -2935,7 +2933,7 @@ void FDeferredShadingSceneRenderer::RenderPathTracing(
 							PassParameters->AtmosphereOpticalDepthLUT = AtmosphereOpticalDepthLUT;
 							PassParameters->AtmosphereOpticalDepthLUTSampler = TStaticSamplerState<SF_Bilinear, AM_Clamp, AM_Clamp, AM_Clamp>::GetRHI();
 
-							if (Config.PathTracingData.EnableFog)
+							if ((Config.PathTracingData.VolumeFlags & PATH_TRACER_VOLUME_ENABLE_FOG) != 0)
 							{
 								PassParameters->FogParameters = PrepareFogParameters(View, Scene->ExponentialFogs[0]);
 							}
