Fixed CIS static analysis warning

diff --git a/Engine/Source/Runtime/Renderer/Private/LightGridInjection.cpp b/Engine/Source/Runtime/Renderer/Private/LightGridInjection.cpp
--- a/Engine/Source/Runtime/Renderer/Private/LightGridInjection.cpp
+++ b/Engine/Source/Runtime/Renderer/Private/LightGridInjection.cpp
@@ -971,7 +971,7 @@ FComputeLightGridOutput FDeferredShadingSceneRenderer::GatherLightsAndComputeLig
 	
 	const bool bCullLightsToGrid = GLightCullingQuality 
 		&& (IsForwardShadingEnabled(ShaderPlatform) || bAnyViewUsesForwardLighting || IsRayTracingEnabled() || ShouldUseClusteredDeferredShading() ||
-			bAnyViewUsesLumen || ViewFamily.EngineShowFlags.VisualizeMeshDistanceFields || VirtualShadowMapArray.IsEnabled() || StochasticShadows::IsEnabled() || StochasticShadows::IsEnabled());
+			bAnyViewUsesLumen || ViewFamily.EngineShowFlags.VisualizeMeshDistanceFields || VirtualShadowMapArray.IsEnabled() || StochasticShadows::IsEnabled());
 
 	// Store this flag if lights are injected in the grids, check with 'AreLightsInLightGrid()'
 	bAreLightsInLightGrid = bCullLightsToGrid;
