Fixed CIS static analysis warning

diff --git a/Engine/Source/Runtime/Renderer/Private/LightRendering.cpp b/Engine/Source/Runtime/Renderer/Private/LightRendering.cpp
--- a/Engine/Source/Runtime/Renderer/Private/LightRendering.cpp
+++ b/Engine/Source/Runtime/Renderer/Private/LightRendering.cpp
@@ -1187,8 +1187,7 @@ void FSceneRenderer::GatherAndSortLights(FSortedLightSetSceneInfo& OutSortedLigh
 					SortedLightInfo->SortKey.Fields.bIsNotSimpleLight = 1;
 
 					// Lights handled by Stochastic Shadows
-					const bool bHandledByStochasticShadows = StochasticShadows::IsLightSupported(LightSceneInfoCompact.LightType, LightSceneInfoCompact.CastRaytracedShadow)
-						|| StochasticShadows::IsLightSupported(LightSceneInfoCompact.LightType, LightSceneInfoCompact.CastRaytracedShadow);
+					const bool bHandledByStochasticShadows = StochasticShadows::IsLightSupported(LightSceneInfoCompact.LightType, LightSceneInfoCompact.CastRaytracedShadow);
 
 					// tiled and clustered deferred lighting only supported for certain lights that don't use any additional features
 					// And also that are not directional (mostly because it doesn't make so much sense to insert them into every grid cell in the universe)
