Add light channel support for cluster shading.

This allows to get more light going through the cluster shading path.

XSX 24 lights on screen on LumenTest 6.0ms -> 2.7ms

diff --git a/Engine/Shaders/Private/ClusteredDeferredShadingPixelShader.usf b/Engine/Shaders/Private/ClusteredDeferredShadingPixelShader.usf
--- a/Engine/Shaders/Private/ClusteredDeferredShadingPixelShader.usf
+++ b/Engine/Shaders/Private/ClusteredDeferredShadingPixelShader.usf
@@ -70,7 +70,9 @@
 // Enable cluster debug visualization
 #define GRID_DEBUG_ENABLE 0
 
-Texture2D< uint4 > ShadowMaskBits;
+uint bHasLightChannels;
+Texture2D<uint> LightingChannelsTexture;
+Texture2D<uint4> ShadowMaskBits;
 
 #if USE_HAIR_LIGHTING
 uint				HairTransmittanceBufferMaxCount;
@@ -240,6 +242,11 @@ float4 GetLightGridLocalLighting(
 
 	int2 PixelPos = int2( ScreenUV.xy * View.BufferSizeAndInvSize.xy );
 	uint4 ShadowMask = ShadowMaskBits[ PixelPos ];
+	uint LightChannelMask = LIGHTING_CHANNEL_MASK;
+	if (bHasLightChannels)
+	{
+		LightChannelMask = LightingChannelsTexture.Load(uint3(PixelPos, 0)).x;
+	}
 	
 	LOOP
 	for (uint GridLightListIndex = FirstNonSimpleLightIndex; GridLightListIndex < NumLightsInGridCell; GridLightListIndex++)
@@ -258,6 +265,11 @@ float4 GetLightGridLocalLighting(
 			continue;
 		}
 
+		if ((LightChannelMask & UnpackLightingChannelMask(LocalLight)) == 0)
+		{
+			continue;
+		}
+
 		FDeferredLightData LightData = ConvertToDeferredLight(LocalLight);
 		LightData.ShadowedBits = 1;
 		LightData.bRectLight = false;
diff --git a/Engine/Source/Runtime/Renderer/Private/ClusteredDeferredShadingPass.cpp b/Engine/Source/Runtime/Renderer/Private/ClusteredDeferredShadingPass.cpp
--- a/Engine/Source/Runtime/Renderer/Private/ClusteredDeferredShadingPass.cpp
+++ b/Engine/Source/Runtime/Renderer/Private/ClusteredDeferredShadingPass.cpp
@@ -101,6 +101,7 @@ class FClusteredShadingPS : public FGlobalShader
 	using FPermutationDomain = TShaderPermutationDomain<FVisualizeLightCullingDim, FHairStrandsLighting, FSubstrateTileType, FLightFunctionAtlasDim>;
 
 	BEGIN_SHADER_PARAMETER_STRUCT(FParameters, )
+		SHADER_PARAMETER(uint32, bHasLightChannels)
 		SHADER_PARAMETER_RDG_UNIFORM_BUFFER(FForwardLightData, Forward)
 		SHADER_PARAMETER_STRUCT_REF(FViewUniformShaderParameters, View)
 		SHADER_PARAMETER_STRUCT_INCLUDE(ShaderPrint::FShaderParameters, ShaderPrintUniformBuffer)
@@ -109,6 +110,7 @@ class FClusteredShadingPS : public FGlobalShader
 		SHADER_PARAMETER_RDG_UNIFORM_BUFFER(FSceneTextureUniformParameters, SceneTextures)
 		SHADER_PARAMETER_RDG_UNIFORM_BUFFER(FLightFunctionAtlasGlobalParameters, LightFunctionAtlas)
 		SHADER_PARAMETER_RDG_TEXTURE(Texture2D, ShadowMaskBits)
+		SHADER_PARAMETER_RDG_TEXTURE(Texture2D<uint>, LightingChannelsTexture)
 		SHADER_PARAMETER_STRUCT_INCLUDE(FVirtualShadowMapSamplingParameters, VirtualShadowMapSamplingParameters)
 		SHADER_PARAMETER_STRUCT_INCLUDE(Substrate::FSubstrateTilePassVS::FParameters, SubstrateTile)
 		SHADER_PARAMETER_RDG_BUFFER_SRV(Buffer<uint>, HairTransmittanceBuffer)
@@ -151,9 +153,10 @@ static void InternalAddClusteredDeferredShadingPass(
 	int32 ViewIndex,
 	FViewInfo& View,
 	const FMinimalSceneTextures& SceneTextures,
-	const FSortedLightSetSceneInfo &SortedLightsSet,
+	const FSortedLightSetSceneInfo& SortedLightsSet,
 	EClusterPassInputType InputType,
 	ESubstrateTileType TileType,
+	FRDGTextureRef LightingChannelsTexture,
 	FRDGTextureRef ShadowMaskBits,
 	FVirtualShadowMapArray& VirtualShadowMapArray,
 	FRDGBufferSRVRef HairTransmittanceBuffer,
@@ -165,11 +168,13 @@ static void InternalAddClusteredDeferredShadingPass(
 	const bool bSubstrate = Substrate::IsSubstrateEnabled() && !bHairStrands;
 	
 	FClusteredShadingPS::FParameters *PassParameters = GraphBuilder.AllocParameters<FClusteredShadingPS::FParameters>();
+	PassParameters->bHasLightChannels = SortedLightsSet.bHasLightChannels;
 	PassParameters->View = View.ViewUniformBuffer;
 	PassParameters->HairStrands = HairStrands::BindHairStrandsViewUniformParameters(View);
 	PassParameters->Forward = View.ForwardLightingResources.ForwardLightUniformBuffer;
 	PassParameters->SceneTextures = SceneTextures.UniformBuffer;
 	PassParameters->ShadowMaskBits = ShadowMaskBits ? ShadowMaskBits : GSystemTextures.GetZeroUIntDummy(GraphBuilder);
+	PassParameters->LightingChannelsTexture = LightingChannelsTexture ? LightingChannelsTexture : GSystemTextures.GetZeroUIntDummy(GraphBuilder);
 	PassParameters->VirtualShadowMapSamplingParameters = VirtualShadowMapArray.GetSamplingParameters(GraphBuilder);
 	PassParameters->HairTransmittanceBuffer = HairTransmittanceBuffer;
 	PassParameters->Substrate = Substrate::BindSubstrateGlobalUniformParameters(View);
@@ -275,7 +280,8 @@ void FDeferredShadingSceneRenderer::AddClusteredDeferredShadingPass(
 	const FMinimalSceneTextures& SceneTextures,
 	const FSortedLightSetSceneInfo &SortedLightsSet,
 	FRDGTextureRef ShadowMaskBits,
-	FRDGTextureRef HairStrandsShadowMaskBits)
+	FRDGTextureRef HairStrandsShadowMaskBits,
+	FRDGTextureRef LightingChannelsTexture)
 {
 	check(GUseClusteredDeferredShading);
 
@@ -304,6 +310,7 @@ void FDeferredShadingSceneRenderer::AddClusteredDeferredShadingPass(
 					SortedLightsSet,
 					EClusterPassInputType::Substrate,
 					ESubstrateTileType::EComplexSpecial,
+					LightingChannelsTexture,
 					ShadowMaskBits,
 					VirtualShadowMapArray,
 					nullptr,
@@ -317,6 +324,7 @@ void FDeferredShadingSceneRenderer::AddClusteredDeferredShadingPass(
 					SortedLightsSet,
 					EClusterPassInputType::Substrate,
 					ESubstrateTileType::EComplex,
+					LightingChannelsTexture,
 					ShadowMaskBits,
 					VirtualShadowMapArray,
 					nullptr,
@@ -330,6 +338,7 @@ void FDeferredShadingSceneRenderer::AddClusteredDeferredShadingPass(
 					SortedLightsSet,
 					EClusterPassInputType::Substrate,
 					ESubstrateTileType::ESingle,
+					LightingChannelsTexture,
 					ShadowMaskBits,
 					VirtualShadowMapArray,
 					nullptr,
@@ -343,6 +352,7 @@ void FDeferredShadingSceneRenderer::AddClusteredDeferredShadingPass(
 					SortedLightsSet,
 					EClusterPassInputType::Substrate,
 					ESubstrateTileType::ESimple,
+					LightingChannelsTexture,
 					ShadowMaskBits,
 					VirtualShadowMapArray,
 					nullptr, 
@@ -358,6 +368,7 @@ void FDeferredShadingSceneRenderer::AddClusteredDeferredShadingPass(
 					SortedLightsSet,
 					EClusterPassInputType::GBuffer,
 					ESubstrateTileType::ECount,
+					LightingChannelsTexture,
 					ShadowMaskBits,
 					VirtualShadowMapArray,
 					nullptr,
@@ -375,6 +386,7 @@ void FDeferredShadingSceneRenderer::AddClusteredDeferredShadingPass(
 					SortedLightsSet,
 					EClusterPassInputType::HairStrands,
 					ESubstrateTileType::ECount,
+					LightingChannelsTexture,
 					HairStrandsShadowMaskBits,
 					VirtualShadowMapArray,
 					GraphBuilder.CreateSRV(TransmittanceMask.TransmittanceMask, FHairStrandsTransmittanceMaskData::Format),
diff --git a/Engine/Source/Runtime/Renderer/Private/DeferredShadingRenderer.h b/Engine/Source/Runtime/Renderer/Private/DeferredShadingRenderer.h
--- a/Engine/Source/Runtime/Renderer/Private/DeferredShadingRenderer.h
+++ b/Engine/Source/Runtime/Renderer/Private/DeferredShadingRenderer.h
@@ -822,7 +822,8 @@ private:
 		const FMinimalSceneTextures& SceneTextures,
 		const FSortedLightSetSceneInfo& SortedLightsSet,
 		FRDGTextureRef ShadowMaskBits,
-		FRDGTextureRef HairStrandsShadowMaskBits);
+		FRDGTextureRef HairStrandsShadowMaskBits, 
+		FRDGTextureRef ShadowMaskBitsLightingChannelsTexture);
 
 	/** Renders the scene's lighting. */
 	void RenderLights(
diff --git a/Engine/Source/Runtime/Renderer/Private/LightRendering.cpp b/Engine/Source/Runtime/Renderer/Private/LightRendering.cpp
--- a/Engine/Source/Runtime/Renderer/Private/LightRendering.cpp
+++ b/Engine/Source/Runtime/Renderer/Private/LightRendering.cpp
@@ -1178,7 +1178,6 @@ void FSceneRenderer::GatherAndSortLights(FSortedLightSetSceneInfo& OutSortedLigh
 
 					// Check for shadows and light functions.
 					SortedLightInfo->SortKey.Fields.LightType = LightSceneInfoCompact.LightType;
-					SortedLightInfo->SortKey.Fields.bTextureProfile = ViewFamily.EngineShowFlags.TexturedLightProfiles && LightSceneInfo->Proxy->GetIESTextureResource();
 					SortedLightInfo->SortKey.Fields.bShadowed = bDynamicShadows && CheckForProjectedShadows(LightSceneInfo);
 					SortedLightInfo->SortKey.Fields.bLightFunction = ViewFamily.EngineShowFlags.LightFunctions && CheckForLightFunction(LightSceneInfo);
 					SortedLightInfo->SortKey.Fields.bUsesLightingChannels = Views[ViewIndex].bUsesLightingChannels && LightSceneInfo->Proxy->GetLightingChannelMask() != GetDefaultLightingChannelMask();
@@ -1196,12 +1195,16 @@ void FSceneRenderer::GatherAndSortLights(FSortedLightSetSceneInfo& OutSortedLigh
 					// Rect lights are not supported as the performance impact is significant even if not used, for now, left for trad. deferred.
 					const bool bClusteredDeferredSupported =
 						(!SortedLightInfo->SortKey.Fields.bShadowed || bShadowedLightsInClustered) &&
-						(!SortedLightInfo->SortKey.Fields.bLightFunction || bUseLightFunctionAtlas) &&
-						!SortedLightInfo->SortKey.Fields.bUsesLightingChannels
+						(!SortedLightInfo->SortKey.Fields.bLightFunction || bUseLightFunctionAtlas)
 						&& LightSceneInfoCompact.LightType != LightType_Directional
 						&& LightSceneInfoCompact.LightType != LightType_Rect
 						&& !bHandledByStochasticShadows;
 
+					// Track feature available accross all lights
+					if (SortedLightInfo->SortKey.Fields.LightType == LightType_Rect)	{ OutSortedLights.bHasRectLights = true; }
+					if (SortedLightInfo->SortKey.Fields.bUsesLightingChannels)			{ OutSortedLights.bHasLightChannels = true; }
+					if (SortedLightInfo->SortKey.Fields.bLightFunction) 				{ OutSortedLights.bHasLightFunctions = true; }
+
 					SortedLightInfo->SortKey.Fields.bClusteredDeferredNotSupported = !bClusteredDeferredSupported;
 					SortedLightInfo->SortKey.Fields.bHandledByStochasticShadows = bHandledByStochasticShadows;
 					break;
@@ -1225,7 +1228,6 @@ void FSceneRenderer::GatherAndSortLights(FSortedLightSetSceneInfo& OutSortedLigh
 
 		FSortedLightSceneInfo* SortedLightInfo = new(SortedLights) FSortedLightSceneInfo(SimpleLightIndex);
 		SortedLightInfo->SortKey.Fields.LightType = LightType_Point;
-		SortedLightInfo->SortKey.Fields.bTextureProfile = 0;
 		SortedLightInfo->SortKey.Fields.bShadowed = 0;
 		SortedLightInfo->SortKey.Fields.bLightFunction = 0;
 		SortedLightInfo->SortKey.Fields.bUsesLightingChannels = 0;
@@ -1259,7 +1261,6 @@ void FSceneRenderer::GatherAndSortLights(FSortedLightSetSceneInfo& OutSortedLigh
 	{
 		const FSortedLightSceneInfo& SortedLightInfo = SortedLights[LightIndex];
 		const bool bDrawShadows = SortedLightInfo.SortKey.Fields.bShadowed;
-		const bool bTextureLightProfile = SortedLightInfo.SortKey.Fields.bTextureProfile;
 		const bool bLightingChannels = SortedLightInfo.SortKey.Fields.bUsesLightingChannels;
 
 		// Do not schedule unbatched lights if the atlas is used and enabled
@@ -1431,7 +1432,7 @@ void FDeferredShadingSceneRenderer::RenderLights(
 				// Tell the trad. deferred that the simple lights are spoken for.
 				bRenderSimpleLightsStandardDeferred = false;
 
-				AddClusteredDeferredShadingPass(GraphBuilder, SceneTextures, SortedLightSet, ShadowSceneRenderer->VirtualShadowMapMaskBits, ShadowSceneRenderer->VirtualShadowMapMaskBitsHairStrands);
+				AddClusteredDeferredShadingPass(GraphBuilder, SceneTextures, SortedLightSet, ShadowSceneRenderer->VirtualShadowMapMaskBits, ShadowSceneRenderer->VirtualShadowMapMaskBitsHairStrands, LightingChannelsTexture);
 			}
 
 			if (bRenderSimpleLightsStandardDeferred)
diff --git a/Engine/Source/Runtime/Renderer/Private/LightSceneInfo.h b/Engine/Source/Runtime/Renderer/Private/LightSceneInfo.h
--- a/Engine/Source/Runtime/Renderer/Private/LightSceneInfo.h
+++ b/Engine/Source/Runtime/Renderer/Private/LightSceneInfo.h
@@ -84,8 +84,6 @@ struct FSortedLightSceneInfo
 			// Currently bHandledByLumen is the MSB and LightType is LSB
 			/** The type of light. */
 			uint32 LightType : LightType_NumBits;
-			/** Whether the light has a texture profile. */
-			uint32 bTextureProfile : 1;
 			/** Whether the light uses a light function. */
 			uint32 bLightFunction : 1;
 			/** Whether the light uses lighting channels. */
@@ -144,6 +142,10 @@ struct FSortedLightSetSceneInfo
 	// First light handled by Stochastic Shadows
 	int32 StochasticShadowsLightStart;
 
+	bool bHasRectLights = false;
+	bool bHasLightFunctions = false;
+	bool bHasLightChannels = false;
+
 	FSimpleLightArray SimpleLights;
 	TArray<FSortedLightSceneInfo, SceneRenderingAllocator> SortedLights;
 };
