Add rect light support for cluster shading.

diff --git a/Engine/Shaders/Private/ClusteredDeferredShadingPixelShader.usf b/Engine/Shaders/Private/ClusteredDeferredShadingPixelShader.usf
--- a/Engine/Shaders/Private/ClusteredDeferredShadingPixelShader.usf
+++ b/Engine/Shaders/Private/ClusteredDeferredShadingPixelShader.usf
@@ -159,7 +159,7 @@ float4 GetSimpleLightLighting(
 	#if SUBSTRATE_ENABLED && USE_HAIR_LIGHTING == 0
 		FDeferredLightData LightData = ConvertToDeferredLight(LocalLight);
 		LightData.ShadowedBits = 1;
-		LightData.bRectLight = false;
+		LightData.bRectLight = LightData.bRectLight && USE_RECT_LIGHT;
 
 		FSubstrateShadowTermInputParameters SubstrateShadowTermInputParameters	= GetInitialisedSubstrateShadowTermInputParameters();
 		SubstrateShadowTermInputParameters.bEvaluateShadowTerm					= false;
@@ -272,7 +272,7 @@ float4 GetLightGridLocalLighting(
 
 		FDeferredLightData LightData = ConvertToDeferredLight(LocalLight);
 		LightData.ShadowedBits = 1;
-		LightData.bRectLight = false;
+		LightData.bRectLight = LightData.bRectLight && USE_RECT_LIGHT;
 		#if USE_IES_PROFILE
 		LightData.Color *= ComputeLightProfileMultiplier(TranslatedWorldPosition, LightData.TranslatedWorldPosition, -LightData.Direction, LightData.Tangent, LightData.IESAtlasIndex);
 		#endif
diff --git a/Engine/Source/Runtime/Renderer/Private/ClusteredDeferredShadingPass.cpp b/Engine/Source/Runtime/Renderer/Private/ClusteredDeferredShadingPass.cpp
--- a/Engine/Source/Runtime/Renderer/Private/ClusteredDeferredShadingPass.cpp
+++ b/Engine/Source/Runtime/Renderer/Private/ClusteredDeferredShadingPass.cpp
@@ -98,7 +98,8 @@ class FClusteredShadingPS : public FGlobalShader
 	class FHairStrandsLighting : SHADER_PERMUTATION_BOOL("USE_HAIR_LIGHTING");
 	class FSubstrateTileType : SHADER_PERMUTATION_INT("SUBSTRATE_TILETYPE", 4);
 	class FLightFunctionAtlasDim : SHADER_PERMUTATION_BOOL("USE_LIGHT_FUNCTION_ATLAS");
-	using FPermutationDomain = TShaderPermutationDomain<FVisualizeLightCullingDim, FHairStrandsLighting, FSubstrateTileType, FLightFunctionAtlasDim>;
+	class FRectLight : SHADER_PERMUTATION_BOOL("USE_RECT_LIGHT");
+	using FPermutationDomain = TShaderPermutationDomain<FVisualizeLightCullingDim, FHairStrandsLighting, FSubstrateTileType, FLightFunctionAtlasDim, FRectLight>;
 
 	BEGIN_SHADER_PARAMETER_STRUCT(FParameters, )
 		SHADER_PARAMETER(uint32, bHasLightChannels)
@@ -166,7 +167,9 @@ static void InternalAddClusteredDeferredShadingPass(
 	const FIntPoint SceneTextureExtent = SceneTextures.Config.Extent;
 	const bool bHairStrands = InputType == EClusterPassInputType::HairStrands;
 	const bool bSubstrate = Substrate::IsSubstrateEnabled() && !bHairStrands;
-	
+	const bool bHasRectLights = SortedLightsSet.bHasRectLights;
+	const bool bLightFunctionAtlas = LightFunctionAtlas::IsEnabled(View, ELightFunctionAtlasSystem::DeferredLighting);
+
 	FClusteredShadingPS::FParameters *PassParameters = GraphBuilder.AllocParameters<FClusteredShadingPS::FParameters>();
 	PassParameters->bHasLightChannels = SortedLightsSet.bHasLightChannels;
 	PassParameters->View = View.ViewUniformBuffer;
@@ -199,10 +202,10 @@ static void InternalAddClusteredDeferredShadingPass(
 	
 	const TCHAR* TileTypeName = ToString(TileType);
 	GraphBuilder.AddPass(
-		RDG_EVENT_NAME("Light::ClusteredDeferredShading(%s,Lights:%d,Tile:%s)", bHairStrands ? TEXT("HairStrands") : (bSubstrate ? TEXT("Substrate") : TEXT("GBuffer")), SortedLightsSet.ClusteredSupportedEnd, TileTypeName),
+		RDG_EVENT_NAME("Light::ClusteredDeferredShading(%s,Lights:%d%s%s%s%s)", bHairStrands ? TEXT("HairStrands") : (bSubstrate ? TEXT("Substrate") : TEXT("GBuffer")), SortedLightsSet.ClusteredSupportedEnd, bSubstrate ? TEXT(",Tile:") : TEXT(""), bSubstrate ? TileTypeName : TEXT(""), bLightFunctionAtlas ? TEXT(",LFAtlas") : TEXT(""), bHasRectLights ? TEXT(",RectLight") : TEXT("")),
 		PassParameters,
 		ERDGPassFlags::Raster,
-		[PassParameters, &View, SceneTextureExtent, bHairStrands, bSubstrate, TileType, PrimitiveType](FRHICommandListImmediate& InRHICmdList)
+		[PassParameters, &View, SceneTextureExtent, bHasRectLights, bLightFunctionAtlas, bHairStrands, bSubstrate, TileType, PrimitiveType](FRHICommandListImmediate& InRHICmdList)
 	{
 		TShaderMapRef<FClusteredShadingVS> HairVertexShader(View.ShaderMap);
 		TShaderMapRef<FPostProcessVS> VertexShader(View.ShaderMap);
@@ -216,7 +219,8 @@ static void InternalAddClusteredDeferredShadingPass(
 		PermutationVector.Set<FClusteredShadingPS::FVisualizeLightCullingDim>(View.Family->EngineShowFlags.VisualizeLightCulling);
 		PermutationVector.Set<FClusteredShadingPS::FHairStrandsLighting>(bHairStrands);
 		PermutationVector.Set<FClusteredShadingPS::FSubstrateTileType>(bSubstrate ? TileType : 0);
-		PermutationVector.Set<FClusteredShadingPS::FLightFunctionAtlasDim>(LightFunctionAtlas::IsEnabled(View, ELightFunctionAtlasSystem::DeferredLighting));
+		PermutationVector.Set<FClusteredShadingPS::FRectLight>(bHasRectLights);
+		PermutationVector.Set<FClusteredShadingPS::FLightFunctionAtlasDim>(bLightFunctionAtlas);
 		TShaderMapRef<FClusteredShadingPS> PixelShader(View.ShaderMap, PermutationVector);
 		{
 			FGraphicsPipelineStateInitializer GraphicsPSOInit;
diff --git a/Engine/Source/Runtime/Renderer/Private/LightRendering.cpp b/Engine/Source/Runtime/Renderer/Private/LightRendering.cpp
--- a/Engine/Source/Runtime/Renderer/Private/LightRendering.cpp
+++ b/Engine/Source/Runtime/Renderer/Private/LightRendering.cpp
@@ -1197,7 +1197,6 @@ void FSceneRenderer::GatherAndSortLights(FSortedLightSetSceneInfo& OutSortedLigh
 						(!SortedLightInfo->SortKey.Fields.bShadowed || bShadowedLightsInClustered) &&
 						(!SortedLightInfo->SortKey.Fields.bLightFunction || bUseLightFunctionAtlas)
 						&& LightSceneInfoCompact.LightType != LightType_Directional
-						&& LightSceneInfoCompact.LightType != LightType_Rect
 						&& !bHandledByStochasticShadows;
 
 					// Track feature available accross all lights
