Stochastic Direct Ligting - restore alongside Stochastic Shadows

#tests 6585a1bc30458dc570da4dc3

diff --git a/Engine/Shaders/Private/StochasticDirectLighting/StochasticDirectLighting.usf b/Engine/Shaders/Private/StochasticDirectLighting/StochasticDirectLighting.usf
new file mode 100644
--- /dev/null
+++ b/Engine/Shaders/Private/StochasticDirectLighting/StochasticDirectLighting.usf
@@ -0,0 +1,234 @@
+// Copyright Epic Games, Inc. All Rights Reserved.
+
+#include "../Common.ush"
+#include "StochasticDirectLighting.ush"
+
+groupshared uint SharedTileMode[THREADGROUP_SIZE * THREADGROUP_SIZE];
+
+RWStructuredBuffer<uint> RWTileAllocator;
+RWStructuredBuffer<uint> RWTileData;
+uint TileDataStride;
+uint DownsampledTileDataStride;
+uint2 DownsampledViewSize;
+
+/**
+ * Run tile classification to generate tiles for each subsequent pass
+ */
+[numthreads(THREADGROUP_SIZE, THREADGROUP_SIZE, 1)]
+void TileClassificationCS(
+	uint3 GroupId : SV_GroupID,
+	uint3 GroupThreadId : SV_GroupThreadID,
+	uint3 DispatchThreadId : SV_DispatchThreadID)
+{
+	uint LinearThreadIndex = GroupThreadId.x + THREADGROUP_SIZE * GroupThreadId.y;
+
+	SharedTileMode[LinearThreadIndex] = 0;
+	
+	GroupMemoryBarrierWithGroupSync();
+
+	uint2 ScreenCoord = DispatchThreadId.xy;
+
+#if DOWNSAMPLED_CLASSIFICATION
+	if (all(ScreenCoord < DownsampledViewSize))
+#else
+	if (all(ScreenCoord < View.ViewRectMinAndSize.zw))
+#endif
+	{
+#if DOWNSAMPLED_CLASSIFICATION
+		float2 ScreenUV = DownsampledScreenCoordToScreenUV(ScreenCoord);
+#else
+		float2 ScreenUV = (ScreenCoord + 0.5f) * View.BufferSizeAndInvSize.zw;
+#endif
+		FGBufferData GBuffer = GetGBufferData(ScreenUV);
+
+		if (GBuffer.ShadingModelID == SHADINGMODELID_UNLIT)
+		{
+			SharedTileMode[LinearThreadIndex] = 0x1;
+		}
+		else if (GBuffer.ShadingModelID == SHADINGMODELID_DEFAULT_LIT)
+		{
+			SharedTileMode[LinearThreadIndex] = 0x2;
+		}
+		else
+		{
+			SharedTileMode[LinearThreadIndex] = 0x4;
+		}
+	}
+
+	GroupMemoryBarrierWithGroupSync();
+
+	// GroupShared reduction
+	if (LinearThreadIndex < 32)
+	{
+		SharedTileMode[LinearThreadIndex] = SharedTileMode[LinearThreadIndex] | SharedTileMode[LinearThreadIndex + 32];
+	}
+
+	GroupMemoryBarrierWithGroupSync();
+
+	if (LinearThreadIndex < 16)
+	{
+		SharedTileMode[LinearThreadIndex] = SharedTileMode[LinearThreadIndex] | SharedTileMode[LinearThreadIndex + 16];
+	}
+
+	GroupMemoryBarrierWithGroupSync();
+
+	if (LinearThreadIndex < 8)
+	{
+		SharedTileMode[LinearThreadIndex] = SharedTileMode[LinearThreadIndex] | SharedTileMode[LinearThreadIndex + 8];
+	}
+
+	GroupMemoryBarrierWithGroupSync();
+
+	if (LinearThreadIndex < 4)
+	{
+		SharedTileMode[LinearThreadIndex] = SharedTileMode[LinearThreadIndex] | SharedTileMode[LinearThreadIndex + 4];
+	}
+
+	GroupMemoryBarrierWithGroupSync();
+
+	if (LinearThreadIndex < 2)
+	{
+		SharedTileMode[LinearThreadIndex] = SharedTileMode[LinearThreadIndex] | SharedTileMode[LinearThreadIndex + 2];
+	}
+
+	GroupMemoryBarrierWithGroupSync();
+
+	if (LinearThreadIndex < 1)
+	{
+		SharedTileMode[LinearThreadIndex] = SharedTileMode[LinearThreadIndex] | SharedTileMode[LinearThreadIndex + 1];
+	}
+
+	if (LinearThreadIndex == 0)
+	{
+#if DOWNSAMPLED_CLASSIFICATION
+		const uint DataStride = DownsampledTileDataStride;
+#else
+		const uint DataStride = TileDataStride;
+#endif
+
+		uint2 TileCoord = ScreenCoord / TILE_SIZE;
+
+		if (SharedTileMode[0] & 0x4)
+		{
+			uint StochasticDirectLightingTileIndex = 0;
+			InterlockedAdd(RWTileAllocator[TILE_MODE_COMPLEX_SHADING], 1, StochasticDirectLightingTileIndex);
+			RWTileData[DataStride * TILE_MODE_COMPLEX_SHADING + StochasticDirectLightingTileIndex] = PackTile(TileCoord);
+		}
+		else if (SharedTileMode[0] & 0x2)
+		{
+			uint StochasticDirectLightingTileIndex = 0;
+			InterlockedAdd(RWTileAllocator[TILE_MODE_SIMPLE_SHADING], 1, StochasticDirectLightingTileIndex);
+			RWTileData[DataStride * TILE_MODE_SIMPLE_SHADING + StochasticDirectLightingTileIndex] = PackTile(TileCoord);
+		}
+		else if (SharedTileMode[0] == 0x1)
+		{
+			uint StochasticDirectLightingTileIndex = 0;
+			InterlockedAdd(RWTileAllocator[TILE_MODE_EMPTY], 1, StochasticDirectLightingTileIndex);
+			RWTileData[DataStride * TILE_MODE_EMPTY + StochasticDirectLightingTileIndex] = PackTile(TileCoord);
+		}
+	}
+}
+
+RWBuffer<uint> RWTileIndirectArgs;
+RWBuffer<uint> RWDownsampledTileIndirectArgs;
+StructuredBuffer<uint> TileAllocator;
+StructuredBuffer<uint> DownsampledTileAllocator;
+
+[numthreads(THREADGROUP_SIZE, 1, 1)]
+void InitTileIndirectArgsCS(uint3 DispatchThreadId : SV_DispatchThreadID)
+{
+	uint TileMode = DispatchThreadId.x;
+	if (TileMode < TILE_MODE_MAX)
+	{
+		WriteDispatchIndirectArgs(RWTileIndirectArgs, TileMode, TileAllocator[TileMode], 1, 1);
+		WriteDispatchIndirectArgs(RWDownsampledTileIndirectArgs, TileMode, DownsampledTileAllocator[TileMode], 1, 1);
+	}
+}
+
+RWBuffer<uint> RWIndirectArgs;
+StructuredBuffer<uint> CompositeTileAllocator;
+uint MaxCompositeTiles;
+
+[numthreads(THREADGROUP_SIZE, 1, 1)]
+void InitCompositeIndirectArgsCS(uint3 DispatchThreadId : SV_DispatchThreadID)
+{
+	if (all(DispatchThreadId == 0))
+	{
+		// Dispatch one group per shadow mask tile
+		uint NumGroups = min(CompositeTileAllocator[0], MaxCompositeTiles);
+		WriteDispatchIndirectArgs(RWIndirectArgs, 0, NumGroups, 1, 1);
+	}
+}
+
+RWTexture2D<float4> RWCompositeUpsampleWeights;
+
+Texture2D<float> DownsampledSceneDepth;
+float2 DownsampledBufferInvSize;
+
+/**
+ * Prepare weights for upsampling in order to pull out common per pixel composite work into a separate dispatch.
+ */
+[numthreads(THREADGROUP_SIZE, THREADGROUP_SIZE, 1)]
+void InitCompositeUpsampleWeightsCS(
+	uint3 GroupId : SV_GroupID,
+	uint3 GroupThreadId : SV_GroupThreadID,
+	uint3 DispatchThreadId : SV_DispatchThreadID)
+{
+	uint2 ScreenCoord = DispatchThreadId.xy;
+
+	if (all(ScreenCoord < View.ViewRectMinAndSize.zw))
+	{
+		float2 ScreenUV = (ScreenCoord + 0.5f) * View.BufferSizeAndInvSize.zw;
+		FGBufferData GBuffer = GetGBufferData(ScreenUV);
+		float3 TranslatedWorldPosition = GetTranslatedWorldPositionFromScreenUV(ScreenUV, GBuffer.Depth);		
+
+		uint2 DownsampledScreenCoord = floor(ScreenUV * View.BufferSizeAndInvSize.xy / DOWNSAMPLE_FACTOR - 0.5f);
+		float2 DownsampledBilinearWeights = frac(ScreenUV * View.BufferSizeAndInvSize.xy / DOWNSAMPLE_FACTOR - 0.5f);
+		float2 DownsampledGatherUV = (DownsampledScreenCoord + 1.0f) * DownsampledBufferInvSize;
+
+		float4 CornerDepths = DownsampledSceneDepth.GatherRed(GlobalPointClampedSampler, DownsampledGatherUV).wzxy;
+
+		float4 InterpolationWeights = float4(
+			(1 - DownsampledBilinearWeights.y) * (1 - DownsampledBilinearWeights.x),
+			(1 - DownsampledBilinearWeights.y) * DownsampledBilinearWeights.x,
+			DownsampledBilinearWeights.y * (1 - DownsampledBilinearWeights.x),
+			DownsampledBilinearWeights.y * DownsampledBilinearWeights.x);
+
+		float4 DepthWeights = 0.0f;	
+
+		#define PLANE_WEIGHTING 1
+		#if PLANE_WEIGHTING
+		{
+			float4 ScenePlane = float4(GBuffer.WorldNormal, dot(TranslatedWorldPosition, GBuffer.WorldNormal));
+
+			float3 Position00 = GetTranslatedWorldPositionFromScreenUV(DownsampledScreenCoordToScreenUV(DownsampledScreenCoord + uint2(0, 0)), CornerDepths.x);
+			float3 Position10 = GetTranslatedWorldPositionFromScreenUV(DownsampledScreenCoordToScreenUV(DownsampledScreenCoord + uint2(1, 0)), CornerDepths.y);
+			float3 Position01 = GetTranslatedWorldPositionFromScreenUV(DownsampledScreenCoordToScreenUV(DownsampledScreenCoord + uint2(0, 1)), CornerDepths.z);
+			float3 Position11 = GetTranslatedWorldPositionFromScreenUV(DownsampledScreenCoordToScreenUV(DownsampledScreenCoord + uint2(1, 1)), CornerDepths.w);
+
+			float4 PlaneDistances;
+			PlaneDistances.x = abs(dot(float4(Position00, -1), ScenePlane));
+			PlaneDistances.y = abs(dot(float4(Position10, -1), ScenePlane));
+			PlaneDistances.z = abs(dot(float4(Position01, -1), ScenePlane));
+			PlaneDistances.w = abs(dot(float4(Position11, -1), ScenePlane));
+
+			float4 RelativeDepthDifference = PlaneDistances / GBuffer.Depth;
+
+			DepthWeights = select(CornerDepths > 0.0f, exp2(-10000.0f * (RelativeDepthDifference * RelativeDepthDifference)), 0.0f);
+		}
+		#else
+		{
+			float4 DepthDifference = abs(CornerDepths - GBuffer.Depth.xxxx);
+			float4 RelativeDepthDifference = DepthDifference / GBuffer.Depth;	
+
+			DepthWeights = select(CornerDepths > 0.0f, exp2(-100.0f * (RelativeDepthDifference * RelativeDepthDifference)), 0.0f);
+		}
+		#endif
+
+		float Epsilon = 0.001f;
+		InterpolationWeights *= DepthWeights;
+		InterpolationWeights /= max(dot(InterpolationWeights, 1), Epsilon);
+
+		RWCompositeUpsampleWeights[ScreenCoord] = InterpolationWeights;
+	}
+}
\ No newline at end of file
diff --git a/Engine/Shaders/Private/StochasticDirectLighting/StochasticDirectLighting.ush b/Engine/Shaders/Private/StochasticDirectLighting/StochasticDirectLighting.ush
new file mode 100644
--- /dev/null
+++ b/Engine/Shaders/Private/StochasticDirectLighting/StochasticDirectLighting.ush
@@ -0,0 +1,189 @@
+// Copyright Epic Games, Inc. All Rights Reserved.
+
+#pragma once
+
+// Change this to force recompilation of all Stochastic Direct Lighting shaders
+#pragma message("UESHADERMETADATA_VERSION 1A63AC4B-2335-2AD1-2CAB-82295C72CAFC")
+
+#define STOCHASTIC_SHADOWS 1
+#include "../LightGridCommon.ush"
+
+#define SUPPORT_CONTACT_SHADOWS 0
+#define NON_DIRECTIONAL_DIRECT_LIGHTING 0
+#include "../DeferredLightingCommon.ush"
+#include "../SceneData.ush"
+#include "../Hash.ush"
+
+#if DEBUG_MODE
+	#include "../ShaderPrint.ush"
+#endif
+
+#include "../Lumen/LumenPosition.ush"
+
+#define DEBUG_MODE_VISUALIZE_SAMPLING 1
+#define DEBUG_MODE_VISUALIZE_TRACING  2
+
+// Limited by PackLightSample
+#define LOCAL_LIGHT_INDEX_BITS		(14)
+#define LOCAL_LIGHT_INDEX_MASK		((1 << LOCAL_LIGHT_INDEX_BITS) - 1)
+#define INVALID_LOCAL_LIGHT_INDEX	(LOCAL_LIGHT_INDEX_MASK)
+
+// Keep in sync with StochasticDirectLighting.cpp
+#define SHADING_ATLAS_SIZE_IN_TILES 512
+#define DOWNSAMPLE_FACTOR			2
+#define TILE_SIZE					8
+#define MAX_LOCAL_LIGHT_INDEX_XY	16
+#define MAX_LOCAL_LIGHT_INDEX		(MAX_LOCAL_LIGHT_INDEX_XY * MAX_LOCAL_LIGHT_INDEX_XY)
+
+#define SHARED_LIGHT_MASK_SIZE		(MAX_LOCAL_LIGHT_INDEX / 32)
+
+// Keep in sync with StochasticDirectLighting::ETileType
+#define TILE_MODE_SIMPLE_SHADING	0
+#define TILE_MODE_COMPLEX_SHADING	1
+#define TILE_MODE_EMPTY				2
+#define TILE_MODE_MAX				3
+
+uint PackTile(uint2 TileCoord)
+{
+	return TileCoord.x | (TileCoord.y << 16);
+}
+
+uint2 UnpackTile(uint PackedTile)
+{
+	return uint2(PackedTile & 0xFFFF, PackedTile >> 16);
+}
+
+uint2 GetShadingTileAtlasCoord(uint ShadingTileIndex, uint2 CoordInTile)
+{
+	uint2 AtlasCoord;
+	AtlasCoord.x = (ShadingTileIndex % SHADING_ATLAS_SIZE_IN_TILES) * TILE_SIZE + CoordInTile.x;
+	AtlasCoord.y = (ShadingTileIndex / SHADING_ATLAS_SIZE_IN_TILES) * TILE_SIZE + CoordInTile.y;
+	return AtlasCoord;
+}
+
+struct FLightSample
+{
+	uint LocalLightIndex;
+	float Weight;
+	bool bVisible;
+	bool bCompleted; // Whether tracing was completed
+};
+
+FLightSample InitLightSample()
+{
+	FLightSample LightSample;
+	LightSample.LocalLightIndex = INVALID_LOCAL_LIGHT_INDEX;
+	LightSample.Weight = 0.0f;
+	LightSample.bVisible = false;
+	LightSample.bCompleted = true;
+	return LightSample;
+}
+
+uint PackLightSample(FLightSample LightSample)
+{
+	uint PackedLightSample = LightSample.LocalLightIndex & LOCAL_LIGHT_INDEX_MASK;
+	PackedLightSample |= LightSample.bVisible ? 0x8000 : 0;
+	PackedLightSample |= LightSample.bCompleted ? 0x4000 : 0;
+	PackedLightSample |= f32tof16(LightSample.Weight) << 16;
+	return PackedLightSample;
+}
+
+FLightSample UnpackLightSample(uint PackedLightSample)
+{
+	FLightSample LightSample;
+	LightSample.LocalLightIndex = PackedLightSample & 0x3FFF;
+	LightSample.bVisible = PackedLightSample & 0x8000 ? true : false;
+	LightSample.bCompleted = PackedLightSample & 0x4000 ? true : false;
+	LightSample.Weight = f16tof32(PackedLightSample >> 16);
+	return LightSample;
+}
+
+uint StochasticDirectLightingStateFrameIndex;
+
+/** 
+ * Returns sample jitter offset in the range [0, DOWNSAMPLE_FACTOR - 1]
+ */
+uint2 GetSampleScreenCoordJitter()
+{
+	return Hammersley16(StochasticDirectLightingStateFrameIndex, 8, 0) * DOWNSAMPLE_FACTOR;
+}
+
+/** 
+ * Returns sample jitter offset in the range [0, DOWNSAMPLE_FACTOR - 1] from the last frame
+ */
+uint2 GetHistorySampleScreenCoordJitter()
+{
+	return Hammersley16(StochasticDirectLightingStateFrameIndex + 7, 8, 0) * DOWNSAMPLE_FACTOR;
+}
+
+uint2 DownsampledScreenCoordToScreenCoord(uint2 DownsampledScreenCoord)
+{
+	return DownsampledScreenCoord * DOWNSAMPLE_FACTOR + GetSampleScreenCoordJitter();
+}
+
+float2 DownsampledScreenCoordToScreenUV(uint2 DownsampledScreenCoord)
+{
+	uint2 ScreenCoord = DownsampledScreenCoordToScreenCoord(DownsampledScreenCoord);
+	float2 ScreenUV = (ScreenCoord + 0.5f) * View.BufferSizeAndInvSize.zw;
+	return ScreenUV;
+}
+
+uint2 PackCompositeTile(uint LocalLightIndex, uint2 DownsampledScreenCoord)
+{
+	DownsampledScreenCoord /= TILE_SIZE;
+
+	uint2 PackedTile;
+	PackedTile.x = LocalLightIndex;
+	PackedTile.y = (DownsampledScreenCoord.x & 0xFFFF) | ((DownsampledScreenCoord.y & 0xFFFF) << 16);
+	return PackedTile;
+}
+
+void UnpackCompositeTile(uint2 PackedTile, out uint LocalLightIndex, out uint2 DownsampledScreenCoord)
+{
+	LocalLightIndex = PackedTile.x;
+	DownsampledScreenCoord.x = PackedTile.y & 0xFFFF;
+	DownsampledScreenCoord.y = PackedTile.y >> 16;
+
+	DownsampledScreenCoord *= TILE_SIZE;
+}
+
+struct FShadingTile
+{
+	uint LocalLightIndex;
+	uint ShadingTileIndex;
+};
+
+uint PackShadingTile(uint LocalLightIndex, uint ShadingTileIndex)
+{
+	uint PackedShadingTile = LocalLightIndex & 0xFFF;
+	PackedShadingTile |= ShadingTileIndex << 12;
+	return PackedShadingTile;
+}
+
+FShadingTile UnpackShadingTile(uint PackedShadingTile)
+{
+	FShadingTile ShadingTile;
+	ShadingTile.LocalLightIndex = PackedShadingTile & 0xFFF;
+	ShadingTile.ShadingTileIndex = PackedShadingTile >> 12;
+	return ShadingTile;
+}
+
+uint2 ShadingTileGridSize;
+uint MaxShadingTilesPerGridCell;
+
+uint GetShadingTileGridOffset(uint2 ShadingTileCoord, uint IndexInCell)
+{
+	return (ShadingTileCoord.x + ShadingTileCoord.y * ShadingTileGridSize.x) * MaxShadingTilesPerGridCell + IndexInCell;
+}
+
+// Remove texture detail before denoising
+float3 SafeDemodulate(float3 Lighting, float3 Texture)
+{
+	return Lighting / max(Texture, 0.001f);
+}
+
+// Restore texture detail after denoising
+float3 SafeModulate(float3 Lighting, float3 Texture)
+{
+	return Lighting * max(Texture, 0.001f);
+}
\ No newline at end of file
diff --git a/Engine/Shaders/Private/StochasticDirectLighting/StochasticDirectLightingComposite.usf b/Engine/Shaders/Private/StochasticDirectLighting/StochasticDirectLightingComposite.usf
new file mode 100644
--- /dev/null
+++ b/Engine/Shaders/Private/StochasticDirectLighting/StochasticDirectLightingComposite.usf
@@ -0,0 +1,190 @@
+// Copyright Epic Games, Inc. All Rights Reserved.
+
+#include "StochasticDirectLighting.ush"
+
+RWStructuredBuffer<uint> RWShadingTiles;
+RWStructuredBuffer<uint> RWShadingTileAllocator;
+RWTexture2D<uint> RWShadingTileGridAllocator;
+RWStructuredBuffer<uint> RWShadingTileGrid;
+RWTexture2D<UNORM float> RWShadingTileAtlas;
+
+StructuredBuffer<uint2> CompositeTileData;
+Texture2D<uint> LightSamples;
+Texture2D<float> DownsampledSceneDepth;
+Texture2D<UNORM float3> DownsampledSceneWorldNormal;
+Texture2D<float4> CompositeUpsampleWeights;
+
+uint2 NumSamplesPerPixel;
+int2 DownsampledViewSize;
+uint MaxShadingTiles;
+uint DebugLightId;
+
+#define SHARED_TILE_BORDER 1
+#define SHARED_TILE_SIZE (8 + 2 * SHARED_TILE_BORDER)
+groupshared float SharedSampleWeights[SHARED_TILE_SIZE][SHARED_TILE_SIZE];
+
+// 2x2 tile mask. 1 means that we need to spawn a shading tile for that region.
+groupshared uint SharedShadingTileMask[2][2];
+groupshared uint SharedShadingTileIndex[2][2];
+
+void AccumulateLightSample(uint2 SharedTileCoord, float UpsampleWeight, inout float SampleWeightSum)
+{
+	float SampleWeight = SharedSampleWeights[SharedTileCoord.x][SharedTileCoord.y];
+	SampleWeightSum += SampleWeight * UpsampleWeight;
+}
+
+/**
+ * Run one thread per local region and composite current frame's traces.
+ * Then create upsampled shading tiles from the local region.
+ */
+[numthreads(THREADGROUP_SIZE, THREADGROUP_SIZE, 1)]
+void CompositeLightSamplesCS(
+	uint3 GroupId : SV_GroupID,
+	uint3 GroupThreadId : SV_GroupThreadID,
+	uint3 DispatchThreadId : SV_DispatchThreadID)
+{
+	uint CompositeTileIndex = GroupId.x;
+
+	uint LocalLightIndex;
+	uint2 BaseDownsampledScreenCoord;
+	UnpackCompositeTile(CompositeTileData[CompositeTileIndex], LocalLightIndex, BaseDownsampledScreenCoord);
+
+	if (all(GroupThreadId < 2))
+	{
+		SharedShadingTileMask[GroupThreadId.x][GroupThreadId.y] = 0;
+		SharedShadingTileIndex[GroupThreadId.x][GroupThreadId.y] = 0xFFFFFFFF;
+	}
+
+	GroupMemoryBarrierWithGroupSync();
+
+	const FLocalLightData LocalLightData = GetLocalLightDataNonStereo(LocalLightIndex);
+	
+	// Composite traces into SharedSampleWeights
+	for (uint SharedCoordY = GroupThreadId.y; SharedCoordY < SHARED_TILE_SIZE; SharedCoordY += THREADGROUP_SIZE)
+	{
+		for (uint SharedCoordX = GroupThreadId.x; SharedCoordX < SHARED_TILE_SIZE; SharedCoordX += THREADGROUP_SIZE)
+		{
+			const uint2 SharedCoord = uint2(SharedCoordX, SharedCoordY);
+			const uint2 DownsampledScreenCoord = clamp(int2(BaseDownsampledScreenCoord + SharedCoord) - SHARED_TILE_BORDER, 0, DownsampledViewSize - 1);
+			
+			float WeightSum = 0.0f;
+
+			for (uint LocalSampleY = 0; LocalSampleY < NumSamplesPerPixel.y; ++LocalSampleY)
+			{
+				for (uint LocalSampleX = 0; LocalSampleX < NumSamplesPerPixel.x; ++LocalSampleX)
+				{
+					uint2 SampleCoord = DownsampledScreenCoord * NumSamplesPerPixel + uint2(LocalSampleX, LocalSampleY);
+					FLightSample LightSample = UnpackLightSample(LightSamples[SampleCoord]);
+
+					if (LightSample.LocalLightIndex == LocalLightIndex && LightSample.bVisible)
+					{
+						WeightSum += LightSample.Weight;
+					}
+				}
+			}
+
+			// Mark sub-tiles requiring shading tiles, accounting for a bilinear interpolation footprint
+			if (WeightSum > 0.001f)
+			{
+				uint LitBitMask = 0;
+
+				for (uint ShadingTileOffsetY = 0; ShadingTileOffsetY < 2; ++ShadingTileOffsetY)
+				{
+					for (uint ShadingTileOffsetX = 0; ShadingTileOffsetX < 2; ++ShadingTileOffsetX)
+					{
+						const uint2 ShadingTileOffset = uint2(ShadingTileOffsetX, ShadingTileOffsetY);
+
+						// Bilinear footprint
+						uint2 ShadingTileMin = ShadingTileOffset * 4 + (SHARED_TILE_BORDER - 1);
+						uint2 ShadingTileMax = ShadingTileMin + 5;
+
+						if (and(all(SharedCoord >= ShadingTileMin), all(SharedCoord <= ShadingTileMax)))
+						{
+							// #sdl_todo: don't write to the spot from so many threads
+							SharedShadingTileMask[ShadingTileOffset.x][ShadingTileOffset.y] = 1;
+						}
+					}
+				}
+			}
+
+			SharedSampleWeights[SharedCoordX][SharedCoordY] = WeightSum;
+		}
+	}
+
+	GroupMemoryBarrierWithGroupSync();
+
+	uint2 LocalDownsampledTileCoord = GroupThreadId.xy;
+	uint2 DownsampledScreenCoord = BaseDownsampledScreenCoord + LocalDownsampledTileCoord;
+	uint2 ShadingTileOffset = GroupThreadId.xy;
+
+	// Allocate shading tiles
+	if (all(ShadingTileOffset < 2))
+	{
+		if (SharedShadingTileMask[ShadingTileOffset.x][ShadingTileOffset.y] != 0)
+		{	
+			const uint2 ShadingTileCoord = (BaseDownsampledScreenCoord * DOWNSAMPLE_FACTOR) / TILE_SIZE + ShadingTileOffset;
+
+			if (all(ShadingTileCoord < ShadingTileGridSize))
+			{
+				uint ShadingTileIndex = MaxShadingTiles;
+				InterlockedAdd(RWShadingTileAllocator[0], 1u, ShadingTileIndex);
+
+				if (ShadingTileIndex < MaxShadingTiles)
+				{
+					uint ShadingTileIndexInCell = 0;
+					InterlockedAdd(RWShadingTileGridAllocator[ShadingTileCoord], 1u, ShadingTileIndexInCell);
+
+					if (ShadingTileIndexInCell < MaxShadingTilesPerGridCell)
+					{
+						RWShadingTileGrid[GetShadingTileGridOffset(ShadingTileCoord, ShadingTileIndexInCell)] = PackShadingTile(LocalLightIndex, ShadingTileIndex);
+						SharedShadingTileIndex[ShadingTileOffset.x][ShadingTileOffset.y] = ShadingTileIndex;
+					}
+				}
+			}
+		}
+	}
+
+	GroupMemoryBarrierWithGroupSync();
+
+	// Composite shading tiles
+	for (uint ShadingTileOffsetY = 0; ShadingTileOffsetY < 2; ++ShadingTileOffsetY)
+	{
+		for (uint ShadingTileOffsetX = 0; ShadingTileOffsetX < 2; ++ShadingTileOffsetX)
+		{
+			uint ShadingTileIndex = SharedShadingTileIndex[ShadingTileOffsetX][ShadingTileOffsetY];
+
+			if (ShadingTileIndex < MaxShadingTiles)
+			{
+				const uint2 CoordInShadingTile = GroupThreadId.xy;
+				const uint2 ShadingTileCoord = (BaseDownsampledScreenCoord * DOWNSAMPLE_FACTOR) / TILE_SIZE + uint2(ShadingTileOffsetX, ShadingTileOffsetY);
+				uint2 ScreenCoord = ShadingTileCoord * TILE_SIZE + CoordInShadingTile;
+				float4 DownsampledScreenWeights = CompositeUpsampleWeights[ScreenCoord];
+
+				const uint2 SharedTileCoord00 = (uint2(ShadingTileOffsetX, ShadingTileOffsetY) * TILE_SIZE + CoordInShadingTile + 1) / DOWNSAMPLE_FACTOR + (SHARED_TILE_BORDER - 1);
+
+				float SampleWeightSum = 0.0f;
+				AccumulateLightSample(SharedTileCoord00 + uint2(0, 0), DownsampledScreenWeights.x, SampleWeightSum);
+				AccumulateLightSample(SharedTileCoord00 + uint2(1, 0), DownsampledScreenWeights.y, SampleWeightSum);
+				AccumulateLightSample(SharedTileCoord00 + uint2(0, 1), DownsampledScreenWeights.z, SampleWeightSum);
+				AccumulateLightSample(SharedTileCoord00 + uint2(1, 1), DownsampledScreenWeights.w, SampleWeightSum);
+
+				#if DEBUG_MODE
+					int2 DebugScreenCoord = View.CursorPosition.x >= 0 ? View.CursorPosition * View.ViewResolutionFraction : -1;
+					bool bDebug = all(ScreenCoord == DebugScreenCoord) && LocalLightData.LightSceneId == DebugLightId;
+					FShaderPrintContext Context = InitShaderPrintContext(true, float2(0.05, 0.3));
+
+					if (bDebug)
+					{
+						Print(Context, TEXT("CompositeLightSamples"));
+						Newline(Context);
+						Print(Context, TEXT("SampleWeight: "));
+						Print(Context, SampleWeightSum);
+					}
+				#endif
+
+				uint2 ShadingTileAtlasCoord = GetShadingTileAtlasCoord(ShadingTileIndex, CoordInShadingTile);
+				RWShadingTileAtlas[ShadingTileAtlasCoord] = SampleWeightSum;
+			}
+		}
+	}
+}
diff --git a/Engine/Shaders/Private/StochasticDirectLighting/StochasticDirectLightingHardwareRayTracing.usf b/Engine/Shaders/Private/StochasticDirectLighting/StochasticDirectLightingHardwareRayTracing.usf
new file mode 100644
--- /dev/null
+++ b/Engine/Shaders/Private/StochasticDirectLighting/StochasticDirectLightingHardwareRayTracing.usf
@@ -0,0 +1,155 @@
+// Copyright Epic Games, Inc. All Rights Reserved.
+
+#include "../Common.ush"
+#include "../BlueNoise.ush"
+#include "../RayTracing/RayTracingCommon.ush"
+#include "../SceneTexturesCommon.ush"
+#include "StochasticDirectLighting.ush"
+#include "StochasticDirectLightingTracing.ush"
+#include "../Lumen/LumenHardwareRayTracingPayloadCommon.ush"
+#if HAIR_VOXEL_TRACES
+	#include "../HairStrands/HairStrandsRaytracing.ush"
+#endif
+
+float RayTracingBias;
+float RayTracingNormalBias;
+uint2 SampleViewSize;
+uint DebugMode;
+uint MaxTraversalIterations;
+RaytracingAccelerationStructure TLAS;
+
+#if LUMEN_HARDWARE_INLINE_RAYTRACING
+	struct FHitGroupRootConstants
+	{
+		uint BaseInstanceIndex;
+		uint UserData;
+	};
+
+	StructuredBuffer<FHitGroupRootConstants> HitGroupData;
+	StructuredBuffer<FRayTracingSceneMetadataRecord> RayTracingSceneMetadata;
+#endif
+
+RWTexture2D<uint> RWLightSamples;
+
+Texture2D<float> LightSampleRayDistance;
+StructuredBuffer<uint> CompactedTraceTexelData;
+StructuredBuffer<uint> CompactedTraceTexelAllocator;
+Texture2D<float> DownsampledSceneDepth;
+Texture2D<UNORM float3> DownsampledSceneWorldNormal;
+
+RAY_TRACING_ENTRY_RAYGEN_OR_INLINE(HardwareRayTraceLightSamples)
+{
+	uint ThreadIndex = DispatchThreadIndex.x;
+	uint GroupIndex = DispatchThreadIndex.y;
+	uint TraceTexelIndex = GroupIndex * 64 + ThreadIndex;
+
+	if (TraceTexelIndex < CompactedTraceTexelAllocator[0])
+	{
+		uint2 SampleCoord = UnpackTraceTexel(CompactedTraceTexelData[TraceTexelIndex]);
+		uint2 ScreenCoord = SampleCoordToScreenCoord(SampleCoord);
+		uint2 DownsampledScreenCoord = SampleCoordToDownsampledScreenCoord(SampleCoord);
+
+#if DEBUG_MODE
+		FShaderPrintContext Context = InitShaderPrintContext(true, float2(0.55, 0.45));
+		int2 DebugScreenCoord = View.CursorPosition.x >= 0 ? View.CursorPosition * View.ViewResolutionFraction : -1;
+		bool bDebug = all(DebugScreenCoord / DOWNSAMPLE_FACTOR == DownsampledScreenCoord);
+#endif
+
+		float2 ScreenUV = (ScreenCoord + 0.5f) * View.BufferSizeAndInvSize.zw;
+		float SceneDepth = DownsampledSceneDepth[DownsampledScreenCoord];
+		float3 SceneWorldNormal = normalize(DecodeNormal(DownsampledSceneWorldNormal[DownsampledScreenCoord]));
+		float3 TranslatedWorldPosition = GetTranslatedWorldPositionFromScreenUV(ScreenUV, SceneDepth);
+
+		FLightSample LightSample = UnpackLightSample(RWLightSamples[SampleCoord]);
+		FLightSampleTrace LightSampleTrace = GetLightSampleTrace(TranslatedWorldPosition, LightSample.LocalLightIndex, SampleCoord);
+		const float LightSampleRayMinT = LightSampleRayDistance[SampleCoord];
+
+		FRayDesc Ray = (FRayDesc)0;
+		Ray.Origin = TranslatedWorldPosition;
+		Ray.Direction = LightSampleTrace.Direction;
+		Ray.TMin = max(LightSampleRayMinT, RayTracingBias);
+		Ray.TMax = LightSampleTrace.Distance;
+
+		ApplyPositionBias(Ray.Origin, Ray.Direction, SceneWorldNormal, RayTracingNormalBias);
+
+		uint RayFlags = RAY_FLAG_ACCEPT_FIRST_HIT_AND_END_SEARCH;
+		uint RayMask = RAY_TRACING_MASK_OPAQUE;
+
+		FLumenMinimalPayload Payload = (FLumenMinimalPayload)0;
+
+		// #sdl_todo: Share it with LumenHardwareRayTracingCommon.ush
+		#if LUMEN_HARDWARE_INLINE_RAYTRACING
+		{
+			FTraceRayInlineContext TraceRayInlineContext = CreateTraceRayInlineContext();
+			TraceRayInlineContext.MaxIterations = MaxTraversalIterations;
+
+			FTraceRayInlineResult TraceResult = TraceRayInline(TLAS, RayFlags, RayMask, Ray.GetNativeDesc(), TraceRayInlineContext);
+
+			if (TraceResult.IsHit())
+			{
+				Payload.HitT = TraceResult.HitT;
+
+				// Our SBTs always have 2 slots per material. Therefore InstanceContributionToHitGroupIndex has a RAY_TRACING_NUM_SHADER_SLOTS multiplier.
+				// But because we only encode 1 material into LumenHardwareRayTracingHitDataBuffer we want to get actual material index.
+				uint InstanceContributionIndex = TraceResult.InstanceContributionToHitGroupIndex / RAY_TRACING_NUM_SHADER_SLOTS;
+				uint HitGroupIndex = TraceResult.GeometryIndex + InstanceContributionIndex;
+
+				FHitGroupRootConstants HitData = HitGroupData[HitGroupIndex];
+
+				const int Primitive = TraceResult.InstanceID;
+				const int PrimitiveInstanceIndex = TraceResult.InstanceIndex - HitData.BaseInstanceIndex;
+				const uint GPUSceneInstanceIndex = GetGPUSceneInstanceIndex(Primitive, PrimitiveInstanceIndex);
+				Payload.SetGPUSceneInstanceIndex(GPUSceneInstanceIndex);
+
+				uint MaterialShaderIndex = HitData.UserData & 0x3FFFFFFF;
+				Payload.SetMaterialShaderIndex(MaterialShaderIndex);
+
+				uint bIsTranslucent = (HitData.UserData >> 31) & 0x1;
+				Payload.SetIsTranslucent(bIsTranslucent);
+
+				uint bIsTwoSided = (HitData.UserData >> 30) & 0x1;
+				Payload.SetIsTwoSided(bIsTwoSided);
+			}
+			else
+			{
+				Payload.SetMiss();
+			}
+		}
+		#else
+		{
+			TraceRay(TLAS, RayFlags, RayMask, /*ShaderSlotIndex*/ 0, RAY_TRACING_NUM_SHADER_SLOTS, 0, Ray.GetNativeDesc(), Payload);
+		}
+		#endif
+
+		#if HAIR_VOXEL_TRACES
+		{
+			// #sdl_todo: replace with spatiotemporal blue noise
+			RandomSequence RandSequence;
+			RandomSequence_Initialize(RandSequence, SampleCoord, 0, StochasticDirectLightingStateFrameIndex, 1);
+			Payload.HitT = TraverseHair(ScreenCoord, RandSequence, Ray.Origin, Ray.Direction, Payload.HitT, VirtualVoxel.Raytracing_ShadowOcclusionThreshold);
+		}
+		#endif
+
+		#if DEBUG_MODE
+		if (bDebug && DebugMode == DEBUG_MODE_VISUALIZE_TRACING)
+		{
+			const FLocalLightData LocalLightData = GetLocalLightDataNonStereo(LightSample.LocalLightIndex);
+			const FLightSceneData LightSceneData = GetLightSceneData(LocalLightData.LightSceneId);
+			const FDeferredLightData LightData = ConvertToDeferredLight(LightSceneData);
+
+			float3 RayStart = Ray.Origin + Ray.Direction * Ray.TMin;
+			float3 RayEnd = Ray.Origin + Ray.Direction * (Payload.IsHit() ? Payload.HitT : Ray.TMax);
+			float4 RayColor = float4(LightData.Color.xyz / Luminance(LightData.Color.xyz), 1.0f);
+
+			AddLineTWS(Context, RayStart, RayEnd, RayColor);	
+		}
+		#endif
+
+		if (Payload.IsHit())
+		{
+			LightSample.bCompleted = true;
+			LightSample.bVisible = false;
+			RWLightSamples[SampleCoord] = PackLightSample(LightSample);
+		}
+	}
+}
\ No newline at end of file
diff --git a/Engine/Shaders/Private/StochasticDirectLighting/StochasticDirectLightingSampling.usf b/Engine/Shaders/Private/StochasticDirectLighting/StochasticDirectLightingSampling.usf
new file mode 100644
--- /dev/null
+++ b/Engine/Shaders/Private/StochasticDirectLighting/StochasticDirectLightingSampling.usf
@@ -0,0 +1,490 @@
+// Copyright Epic Games, Inc. All Rights Reserved.
+
+#ifndef NUM_SAMPLES_PER_PIXEL_1D
+	#define NUM_SAMPLES_PER_PIXEL_1D 1
+#endif
+
+#ifndef TILE_TYPE
+	#define TILE_TYPE TILE_MODE_EMPTY
+#endif
+
+#include "../Common.ush"
+#include "../BlueNoise.ush"
+#include "StochasticDirectLighting.ush"
+#include "StochasticDirectLightingTracing.ush"
+
+float GetLocalLightTargetPDF(FDeferredLightData LightData, float3 TranslatedWorldPosition, FGBufferData GBuffer, uint2 ScreenCoord)
+{
+	float SampleWeight = 0.0f;
+
+	float3 L = LightData.Direction;
+	float3 ToLight = L;
+
+	float LightMask = GetLocalLightAttenuation(TranslatedWorldPosition, LightData, ToLight, L);
+
+	if (LightMask > 0.0f)
+	{
+		float Attenuation = 0.0f;
+
+		if (LightData.bRectLight)
+		{
+			FRect Rect = GetRect(ToLight, LightData);
+			Attenuation = IntegrateLight(Rect);
+		}
+		else
+		{
+			FCapsuleLight Capsule = GetCapsule(ToLight, LightData);
+			Attenuation = IntegrateLight(Capsule, LightData.bInverseSquared);
+		}
+
+		// #sdl_todo: figure out cheap sampling
+		SampleWeight = Attenuation * LightMask * View.PreExposure * Luminance(LightData.Color);
+
+		#define REFERENCE_SAMPLING 1
+		#if REFERENCE_SAMPLING
+		{
+			float3 CameraVector = normalize(TranslatedWorldPosition - View.TranslatedWorldCameraOrigin);
+			GBuffer.Roughness = max(GBuffer.Roughness, View.MinRoughness);
+
+			float4 LightAttenuation = 1.0f;
+			float Dither = 0.5f;
+			float SurfaceShadow = 1.0f;
+			float AmbientOcclusion = 1.0f;
+
+			// #sdl_todo: Move more shading models to TILE_MODE_SIMPLE_SHADING
+			#if TILE_TYPE == TILE_MODE_SIMPLE_SHADING
+			{
+				GBuffer.ShadingModelID = SHADINGMODELID_DEFAULT_LIT;
+			}
+			#endif
+			LightData.ShadowedBits = 0;
+
+			float4 LightRadiance = GetDynamicLighting(
+				TranslatedWorldPosition,
+				CameraVector,
+				GBuffer, 
+				AmbientOcclusion,
+				GBuffer.ShadingModelID,
+				LightData,
+				LightAttenuation,
+				Dither,
+				ScreenCoord, 
+				SurfaceShadow);
+
+			SampleWeight = Luminance(LightRadiance.xyz * View.PreExposure);
+
+			if (SampleWeight > 0.01f && LightData.IESAtlasIndex >= 0)
+			{
+				SampleWeight *= ComputeLightProfileMultiplier(TranslatedWorldPosition, LightData.TranslatedWorldPosition, -LightData.Direction, LightData.Tangent, LightData.IESAtlasIndex);
+			}
+		}
+		#endif
+	}
+
+	// Simulate tonemapping
+	SampleWeight = log2(max(SampleWeight + 1.0f, 1.0f));
+
+	return SampleWeight;
+}
+
+uint2 DownsampledViewSize;
+uint MaxCompositeTiles;
+float SamplingMinWeight;
+uint2 NumSamplesPerPixel;
+int DebugMode;
+
+RWTexture2D<float> RWDownsampledSceneDepth;
+RWTexture2D<UNORM float3> RWDownsampledSceneWorldNormal;
+RWStructuredBuffer<uint> RWCompositeTileAllocator;
+RWStructuredBuffer<uint2> RWCompositeTileData;
+RWTexture2D<uint> RWLightSamples;
+
+float4 HistoryScreenPositionScaleBias;
+float4 HistoryUVMinMax;
+
+groupshared uint SharedLightMask[SHARED_LIGHT_MASK_SIZE];
+
+#if CANDIDATE_LIGHT_MASK
+	groupshared uint SharedCandidateLightMask[SHARED_LIGHT_MASK_SIZE];
+#endif
+
+void MarkLocalLightIndex(uint LocalLightIndex)
+{
+	// Re-distribute nearby ids to different bytes in order to balance gather workload
+	uint DWORDIndex = LocalLightIndex % SHARED_LIGHT_MASK_SIZE;
+	uint BitMask = 1u << (LocalLightIndex / SHARED_LIGHT_MASK_SIZE);
+	InterlockedOr(SharedLightMask[DWORDIndex], BitMask);
+}
+
+StructuredBuffer<uint> DownsampledTileAllocator;
+StructuredBuffer<uint> DownsampledTileData;
+uint DownsampledTileDataStride;
+float2 DownsampledBufferInvSize;
+
+/**
+ * Run one thread per sample and generate new light samples for tracing
+ */
+[numthreads(THREADGROUP_SIZE, THREADGROUP_SIZE, 1)]
+void GenerateSamplesCS(
+	uint3 GroupId : SV_GroupID,
+	uint3 GroupThreadId : SV_GroupThreadID,
+	uint3 DispatchThreadId : SV_DispatchThreadID)
+{
+	uint LinearThreadIndex = GroupThreadId.y * THREADGROUP_SIZE + GroupThreadId.x;
+
+	if (LinearThreadIndex < SHARED_LIGHT_MASK_SIZE)
+	{
+		SharedLightMask[LinearThreadIndex] = 0;
+	}
+
+#if CANDIDATE_LIGHT_MASK
+	if (LinearThreadIndex < SHARED_LIGHT_MASK_SIZE)
+	{
+		SharedCandidateLightMask[LinearThreadIndex] = 0;
+	}
+#endif
+
+	GroupMemoryBarrierWithGroupSync();
+
+	uint DownsampledTileIndex = GroupId.x;
+	if (DownsampledTileIndex < DownsampledTileAllocator[TILE_TYPE])
+	{
+		uint2 DownsampledTileCoord = UnpackTile(DownsampledTileData[DownsampledTileIndex + TILE_TYPE * DownsampledTileDataStride]);
+		uint2 DownsampledScreenCoord = DownsampledTileCoord * TILE_SIZE + GroupThreadId.xy;
+
+		#if CANDIDATE_LIGHT_MASK
+		if (all(DownsampledScreenCoord < DownsampledViewSize))
+		{
+			uint2 ScreenCoord = DownsampledScreenCoordToScreenCoord(DownsampledScreenCoord);	
+			const float2 ScreenUV = (ScreenCoord + 0.5f) * View.BufferSizeAndInvSize.zw;
+			const FGBufferData GBuffer = GetGBufferData(ScreenUV);
+			const float SceneDepth = GBuffer.Depth;
+
+			float3 TranslatedWorldPosition = GetTranslatedWorldPositionFromScreenUV(ScreenUV, SceneDepth);		
+
+			const uint EyeIndex = 0;
+			const uint GridIndex = ComputeLightGridCellIndex(ScreenCoord, SceneDepth, EyeIndex);
+			const FCulledLightsGridData CulledLightGridData = GetCulledLightsGrid(GridIndex, EyeIndex);
+			const uint NumLightsInGridCell = min(CulledLightGridData.NumLights, GetMaxLightsPerCell(EyeIndex));
+
+			for (uint GridLightIndex = 0; GridLightIndex < NumLightsInGridCell; ++GridLightIndex)
+			{
+				const uint LocalLightIndex = GetCulledLightDataGrid(CulledLightGridData.DataStartIndex + GridLightIndex);
+				const FLocalLightData LocalLightData = GetLocalLightDataNonStereo(LocalLightIndex);
+
+				if (LocalLightData.bStochasticShadowsSupported && LocalLightIndex < MAX_LOCAL_LIGHT_INDEX)
+				{
+					uint DWORDIndex = LocalLightIndex / 32;
+					uint BitMask = 1u << (LocalLightIndex % 32);
+					InterlockedOr(SharedCandidateLightMask[DWORDIndex], BitMask);
+				}
+			}
+		}
+		#endif
+	
+		GroupMemoryBarrierWithGroupSync();
+
+		if (all(DownsampledScreenCoord < DownsampledViewSize))
+		{
+			uint2 ScreenCoord = DownsampledScreenCoordToScreenCoord(DownsampledScreenCoord);
+
+	#if DEBUG_MODE
+			int2 DebugScreenCoord = View.CursorPosition.x >= 0 ? View.CursorPosition * View.ViewResolutionFraction : -1;
+			bool bDebug = all(DownsampledScreenCoord == DebugScreenCoord / DOWNSAMPLE_FACTOR);
+			FShaderPrintContext Context = InitShaderPrintContext(true, float2(0.05, 0.05));
+	#endif
+
+			FLightSample LightSamples[NUM_SAMPLES_PER_PIXEL_1D];
+			for (uint LightSampleIndex = 0; LightSampleIndex < NUM_SAMPLES_PER_PIXEL_1D; ++LightSampleIndex)
+			{
+				LightSamples[LightSampleIndex] = InitLightSample();
+			}
+
+			const uint2 LightSampleOffsets[4] = 
+			{
+				uint2(0, 0),
+				uint2(1, 0),
+				uint2(0, 1),
+				uint2(1, 1),
+			};
+
+			const float2 ScreenUV = (ScreenCoord + 0.5f) * View.BufferSizeAndInvSize.zw;
+			const FGBufferData GBuffer = GetGBufferData(ScreenUV);
+			const float SceneDepth = GBuffer.Depth;
+			float WeightSum = 0.0f;
+
+			if (SceneDepth > 0)
+			{
+				float3 TranslatedWorldPosition = GetTranslatedWorldPositionFromScreenUV(ScreenUV, SceneDepth);		
+
+				const uint EyeIndex = 0;
+				const uint GridIndex = ComputeLightGridCellIndex(ScreenCoord, SceneDepth, EyeIndex);
+				const FCulledLightsGridData CulledLightGridData = GetCulledLightsGrid(GridIndex, EyeIndex);
+				const uint NumLightsInGridCell = min(CulledLightGridData.NumLights, GetMaxLightsPerCell(EyeIndex));
+				const uint NumLocalLights = GetNumLocalLights(EyeIndex);
+
+				// Initialize random variables using spatiotemporal Blue Noise
+				float LightIndexRandom[NUM_SAMPLES_PER_PIXEL_1D];
+				{
+					float GoldenRatioConjugate = 0.6180339887f;
+					float RandomScalar = BlueNoiseScalar(DownsampledScreenCoord, StochasticDirectLightingStateFrameIndex);
+
+					for (uint LightSampleIndex = 0; LightSampleIndex < NUM_SAMPLES_PER_PIXEL_1D; ++LightSampleIndex)
+					{
+						LightIndexRandom[LightSampleIndex] = frac(RandomScalar + GoldenRatioConjugate * LightSampleIndex);
+					}
+				}
+
+				#if DEBUG_MODE
+				if (bDebug)
+				{
+					Print(Context, TEXT("GenerateSamplesCS"));
+					Newline(Context);
+					Print(Context, TEXT("ScreenCoord: "));
+					Print(Context, ScreenCoord.x);
+					Print(Context, ScreenCoord.y);
+					Newline(Context);
+					Print(Context, TEXT("TileType: "));
+					Print(Context, uint(TILE_TYPE));
+					Print(Context, DownsampledTileAllocator[TILE_TYPE]);
+					Newline(Context);
+					Print(Context, TEXT("View.PreExposure: "));
+					Print(Context, View.PreExposure);
+					Newline(Context);
+					Print(Context, TEXT("NumLightsInGridCell: "));
+					Print(Context, NumLightsInGridCell);
+					Print(Context, TEXT("NumLocalLights: "));
+					if (NumLocalLights <= MAX_LOCAL_LIGHT_INDEX)
+					{
+						Print(Context, NumLocalLights);
+					}
+					else
+					{
+						Print(Context, NumLocalLights, FontRed);
+					}
+					AddCrossTWS(Context, TranslatedWorldPosition, 5.0f, float4(1, 1, 0, 1));
+					Newline(Context);
+					Print(Context, TEXT("LightId  | SampleWeight"));
+				}
+				#endif
+
+	#if CANDIDATE_LIGHT_MASK
+				for (uint MaskIndex = 0; MaskIndex < SHARED_LIGHT_MASK_SIZE; ++MaskIndex)
+				{
+					uint CandidateLightMask = SharedCandidateLightMask[MaskIndex];
+					while (CandidateLightMask != 0)
+					{
+						const uint NextBitIndex = firstbitlow(CandidateLightMask);
+						const uint NextBitMask = 1u << NextBitIndex;
+						CandidateLightMask ^= NextBitMask;
+
+						const uint LocalLightIndex = MaskIndex * 32 + NextBitIndex;
+						const FLocalLightData LocalLightData = GetLocalLightDataNonStereo(LocalLightIndex);
+						const FLightSceneData LightSceneData = GetLightSceneData(LocalLightData.LightSceneId);
+						FDeferredLightData LightData = ConvertToDeferredLight(LightSceneData);
+					
+						float SampleWeight = GetLocalLightTargetPDF(LightData, TranslatedWorldPosition, GBuffer, ScreenCoord);
+
+						#if DEBUG_MODE
+						if (bDebug)
+						{
+							Newline(Context);
+							Print(Context, LocalLightData.LightSceneId);
+							Print(Context, SampleWeight);
+						}
+						#endif
+
+						if (SampleWeight > SamplingMinWeight)
+						{
+							float Tau = WeightSum / (WeightSum + SampleWeight);
+							WeightSum += SampleWeight;
+
+							for (uint LightSampleIndex = 0; LightSampleIndex < NUM_SAMPLES_PER_PIXEL_1D; ++LightSampleIndex)
+							{
+								if (LightIndexRandom[LightSampleIndex] < Tau)
+								{
+									LightIndexRandom[LightSampleIndex] /= Tau;
+								}
+								else
+								{
+									// Select this sample
+									LightIndexRandom[LightSampleIndex] = (LightIndexRandom[LightSampleIndex] - Tau) / (1.0f - Tau);
+									LightSamples[LightSampleIndex].LocalLightIndex = LocalLightIndex;
+									LightSamples[LightSampleIndex].Weight = SampleWeight;
+								}
+
+								LightIndexRandom[LightSampleIndex] = clamp(LightIndexRandom[LightSampleIndex], 0, 0.9999f);
+							}
+						}
+					}
+				}
+
+	#else
+				for (uint GridLightIndex = 0; GridLightIndex < NumLightsInGridCell; ++GridLightIndex)
+				{
+					const uint LocalLightIndex = GetCulledLightDataGrid(CulledLightGridData.DataStartIndex + GridLightIndex);
+					const FLocalLightData LocalLightData = GetLocalLightDataNonStereo(LocalLightIndex);
+
+					if (LocalLightData.bStochasticShadowsSupported)
+					{
+						const FLightSceneData LightSceneData = GetLightSceneData(LocalLightData.LightSceneId);
+						FDeferredLightData LightData = ConvertToDeferredLight(LightSceneData);
+						float SampleWeight = GetLocalLightTargetPDF(LightData, TranslatedWorldPosition, GBuffer, ScreenCoord);
+
+						#if DEBUG_MODE
+						if (bDebug)
+						{
+							Newline(Context);
+							Print(Context, LocalLightData.LightSceneId);
+							Print(Context, SampleWeight);
+						}
+						#endif
+
+						if (SampleWeight > SamplingMinWeight)
+						{
+							float Tau = WeightSum / (WeightSum + SampleWeight);
+							WeightSum += SampleWeight;
+
+							for (uint LightSampleIndex = 0; LightSampleIndex < NUM_SAMPLES_PER_PIXEL_1D; ++LightSampleIndex)
+							{
+								if (LightIndexRandom[LightSampleIndex] < Tau)
+								{
+									LightIndexRandom[LightSampleIndex] /= Tau;
+								}
+								else
+								{
+									// Select this sample
+									LightIndexRandom[LightSampleIndex] = (LightIndexRandom[LightSampleIndex] - Tau) / (1.0f - Tau);
+									LightSamples[LightSampleIndex].LocalLightIndex = LocalLightIndex;
+									LightSamples[LightSampleIndex].Weight = SampleWeight;
+								}
+
+								LightIndexRandom[LightSampleIndex] = clamp(LightIndexRandom[LightSampleIndex], 0, 0.9999f);
+							}
+						}
+					}
+				}
+	#endif
+
+				#if DEBUG_MODE
+				if (bDebug)
+				{
+					Newline(Context);
+					Print(Context, TEXT("Weight sum: "));
+					Print(Context, WeightSum);
+					Newline(Context);
+					Print(Context, TEXT("Selected lights: "));
+
+					for (uint LightSampleIndex = 0; LightSampleIndex < NUM_SAMPLES_PER_PIXEL_1D; ++LightSampleIndex)
+					{
+						FLightSample LightSample = LightSamples[LightSampleIndex];
+						
+						Newline(Context);
+						Print(Context, LightSample.LocalLightIndex); 
+						Print(Context, LightSample.Weight);
+						Print(Context, WeightSum /  LightSample.Weight);
+
+						if (DebugMode == DEBUG_MODE_VISUALIZE_SAMPLING)
+						{
+							const uint2 SampleCoord = DownsampledScreenCoord * DOWNSAMPLE_FACTOR + LightSampleOffsets[LightSampleIndex];
+
+							const FLocalLightData LocalLightData = GetLocalLightDataNonStereo(LightSample.LocalLightIndex);
+							const FLightSceneData LightSceneData = GetLightSceneData(LocalLightData.LightSceneId);
+							const FLightSampleTrace LightSampleTrace = GetLightSampleTrace(TranslatedWorldPosition, LightSample.LocalLightIndex, SampleCoord);
+
+							AddLineTWS(Context, TranslatedWorldPosition, TranslatedWorldPosition + LightSampleTrace.Direction * LightSampleTrace.Distance, float4(LightSceneData.Color.xyz / Luminance(LightSceneData.Color.xyz), 1.0f));
+						}
+					}
+				}
+				#endif
+			}
+
+			RWDownsampledSceneDepth[DownsampledScreenCoord] = SceneDepth;
+			RWDownsampledSceneWorldNormal[DownsampledScreenCoord] = EncodeNormal(GBuffer.WorldNormal);
+
+			for (uint LightSampleIndex = 0; LightSampleIndex < NUM_SAMPLES_PER_PIXEL_1D; ++LightSampleIndex)
+			{
+				FLightSample LightSample = LightSamples[LightSampleIndex];
+
+				if (LightSample.LocalLightIndex != INVALID_LOCAL_LIGHT_INDEX)
+				{
+					const FLocalLightData LocalLightData = GetLocalLightDataNonStereo(LightSample.LocalLightIndex);
+					const FLightSceneData LightSceneData = GetLightSceneData(LocalLightData.LightSceneId);
+					const bool bCastShadows = UnpackCastShadow(LightSceneData.LightTypeAndShadowMapChannelMaskPacked);
+
+					LightSample.bVisible = true;
+					LightSample.bCompleted = bCastShadows ? false : true;
+					LightSample.Weight = WeightSum / (NUM_SAMPLES_PER_PIXEL_1D * LightSample.Weight);
+				}
+
+				RWLightSamples[DownsampledScreenCoord * NumSamplesPerPixel + LightSampleOffsets[LightSampleIndex]] = PackLightSample(LightSample);
+
+				// Mark tiles which need to be processed
+				if (LightSample.LocalLightIndex != INVALID_LOCAL_LIGHT_INDEX)
+				{
+					MarkLocalLightIndex(LightSample.LocalLightIndex);
+				}
+			}
+		}
+
+		GroupMemoryBarrierWithGroupSync();
+
+		// Write out allocated shadow mask tiles
+		// #sdl_todo: process part of the list per tread
+		if (LinearThreadIndex < SHARED_LIGHT_MASK_SIZE)
+		{
+			const uint LightIndex = LinearThreadIndex;
+			const uint BitMask = 1u << LightIndex;
+
+			uint2 DownsampledScreenCoord = DownsampledTileCoord * TILE_SIZE;
+	
+			uint LightMask = SharedLightMask[LinearThreadIndex];
+			while (LightMask != 0)
+			{
+				const uint NextBitIndex = firstbitlow(LightMask);
+				const uint NextBitMask = 1u << NextBitIndex;
+				LightMask ^= NextBitMask;
+
+				const uint LocalLightIndex = LinearThreadIndex + NextBitIndex * SHARED_LIGHT_MASK_SIZE;
+
+				// #sdl_todo: merge allocations through groupshared and do only one atomic per group, or just replace entire thing with a scalar loop over SelectedSamples?
+				uint CompositeTileIndex = 0;
+				InterlockedAdd(RWCompositeTileAllocator[0], 1, CompositeTileIndex);
+
+				if (CompositeTileIndex < MaxCompositeTiles)
+				{
+					RWCompositeTileData[CompositeTileIndex] = PackCompositeTile(LocalLightIndex, DownsampledScreenCoord);
+				}
+			}
+		}
+	}
+}
+
+/**
+ * Clear some data for empty tiles, which won't be processed by GenerateLightSamplesCS
+ */
+[numthreads(THREADGROUP_SIZE, THREADGROUP_SIZE, 1)]
+void ClearLightSamplesCS(
+	uint3 GroupId : SV_GroupID,
+	uint3 GroupThreadId : SV_GroupThreadID,
+	uint3 DispatchThreadId : SV_DispatchThreadID)
+{
+	uint DownsampledTileIndex = GroupId.x;
+	if (DownsampledTileIndex < DownsampledTileAllocator[TILE_MODE_EMPTY])
+	{
+		uint2 DownsampledTileCoord = UnpackTile(DownsampledTileData[DownsampledTileIndex + TILE_MODE_EMPTY * DownsampledTileDataStride]);
+		uint2 DownsampledScreenCoord = DownsampledTileCoord * TILE_SIZE + GroupThreadId.xy;
+
+		if (all(DownsampledScreenCoord < DownsampledViewSize))
+		{
+			for (uint LightSampleY = 0; LightSampleY < NumSamplesPerPixel.y; ++LightSampleY)
+			{
+				for (uint LightSampleX = 0; LightSampleX < NumSamplesPerPixel.x; ++LightSampleX)
+				{
+					FLightSample LightSample = InitLightSample();
+					RWLightSamples[DownsampledScreenCoord * NumSamplesPerPixel + uint2(LightSampleX, LightSampleY)] = PackLightSample(LightSample);
+				}
+			}
+		}
+	}
+}
\ No newline at end of file
diff --git a/Engine/Shaders/Private/StochasticDirectLighting/StochasticDirectLightingShading.usf b/Engine/Shaders/Private/StochasticDirectLighting/StochasticDirectLightingShading.usf
new file mode 100644
--- /dev/null
+++ b/Engine/Shaders/Private/StochasticDirectLighting/StochasticDirectLightingShading.usf
@@ -0,0 +1,181 @@
+// Copyright Epic Games, Inc. All Rights Reserved.
+
+#include "../Common.ush"
+#include "../BlueNoise.ush"
+#include "StochasticDirectLighting.ush"
+
+RWTexture2D<float3> RWResolvedDiffuseLighting;
+RWTexture2D<float3> RWResolvedSpecularLighting;
+
+Texture2D<float> ShadingTileAtlas;
+Texture2D<uint> ShadingTileGridAllocator;
+StructuredBuffer<uint> ShadingTileGrid;
+StructuredBuffer<uint> CompositeTileAllocator;
+StructuredBuffer<uint> ShadingTileAllocator;
+uint MaxCompositeTiles;
+uint MaxShadingTiles;
+
+StructuredBuffer<uint> TileAllocator;
+StructuredBuffer<uint> TileData;
+uint TileDataStride;
+
+/**
+ * Apply all lights per pixel to affected tiles using sampled weights from upscaled shading tiles
+ */
+[numthreads(THREADGROUP_SIZE, THREADGROUP_SIZE, 1)]
+void ShadeLightSamplesCS(
+	uint3 GroupId : SV_GroupID,
+	uint3 GroupThreadId : SV_GroupThreadID,
+	uint3 DispatchThreadId : SV_DispatchThreadID)
+{
+	uint TileIndex = GroupId.x;
+	if (TileIndex < TileAllocator[TILE_TYPE])
+	{
+		uint2 TileCoord = UnpackTile(TileData[TileIndex + TILE_TYPE * TileDataStride]);
+		uint2 ScreenCoord = TileCoord * TILE_SIZE + GroupThreadId.xy;
+
+		if (all(ScreenCoord < View.ViewRectMinAndSize.zw))
+		{
+	#if DEBUG_MODE
+			int2 DebugScreenCoord = View.CursorPosition.x >= 0 ? View.CursorPosition * View.ViewResolutionFraction : -1;
+			bool bDebug = all(ScreenCoord == DebugScreenCoord);
+			FShaderPrintContext Context = InitShaderPrintContext(true, float2(0.55, 0.05));
+	#endif
+
+			float2 ScreenUV = (ScreenCoord + 0.5f) * View.BufferSizeAndInvSize.zw;
+			FGBufferData GBuffer = GetGBufferData(ScreenUV);
+
+			// #sdl_todo: handle this path when shadows are disabled
+			float ScreenSpaceAO = 1.0f;//Texture2DSampleLevel(SceneTexturesStruct.ScreenSpaceAOTexture, SceneTexturesStruct_ScreenSpaceAOTextureSampler, ScreenUV, 0).x;
+
+			float3 DiffuseLighting = 0.0f;
+			float3 SpecularLighting = 0.0f;
+
+			if (GBuffer.Depth > 0.0f)
+			{
+				uint2 ShadingTileCoord = TileCoord;
+				uint2 CoordInShadingTile = GroupThreadId.xy;
+
+				float3 TranslatedWorldPosition = GetTranslatedWorldPositionFromScreenUV(ScreenUV, GBuffer.Depth);		
+				float3 CameraVector = normalize(TranslatedWorldPosition - View.TranslatedWorldCameraOrigin);
+
+				uint NumShadingTilesInCell = ShadingTileGridAllocator.Load(int3(ShadingTileCoord, 0));
+
+				#if DEBUG_MODE
+				if (bDebug)
+				{
+					Print(Context, TEXT("ShadeSamplesCS"));
+					Newline(Context);
+					Print(Context, TEXT("ScreenCoord: "));
+					Print(Context, ScreenCoord.x);
+					Print(Context, ScreenCoord.y);
+					Newline(Context);
+					Print(Context, TEXT("NumShadingTilesInCell: "));
+					Print(Context, NumShadingTilesInCell);
+					Newline(Context);
+					Print(Context, TEXT("ShadingTiles:    "));
+					Print(Context, ShadingTileAllocator[0]);
+					Print(Context, MaxShadingTiles);
+					Print(Context, ShadingTileAllocator[0] / float(MaxShadingTiles));
+					Newline(Context);
+					Print(Context, TEXT("LightId   | Weight    | Diffuse   | Specular"));
+				}
+				#endif
+
+				for (uint IndexInCell = 0; IndexInCell < NumShadingTilesInCell; ++IndexInCell)
+				{
+					const uint PackedShadingTile = ShadingTileGrid[GetShadingTileGridOffset(ShadingTileCoord, IndexInCell)];
+					const FShadingTile ShadingTile = UnpackShadingTile(PackedShadingTile);
+					const FLocalLightData LocalLightData = GetLocalLightDataNonStereo(ShadingTile.LocalLightIndex);
+					const FLightSceneData LightSceneData = GetLightSceneData(LocalLightData.LightSceneId);
+					FDeferredLightData LightData = ConvertToDeferredLight(LightSceneData);
+
+					float SampleWeight = 0.0f;
+
+					if (ShadingTile.ShadingTileIndex < MaxShadingTiles)
+					{
+						uint2 ShadingTileAtlasCoord = GetShadingTileAtlasCoord(ShadingTile.ShadingTileIndex, CoordInShadingTile);
+						SampleWeight = ShadingTileAtlas[ShadingTileAtlasCoord];
+					}
+
+					if (SampleWeight > 0.01f && LightData.IESAtlasIndex >= 0)
+					{
+						// #sdl_todo: evaluate moving this to CompositeLightSamplesCS
+						SampleWeight *= ComputeLightProfileMultiplier(TranslatedWorldPosition, LightData.TranslatedWorldPosition, -LightData.Direction, LightData.Tangent, LightData.IESAtlasIndex);
+					}
+
+					#if DEBUG_MODE
+					if (bDebug)
+					{
+						Newline(Context);
+						Print(Context, LocalLightData.LightSceneId);
+						Print(Context, SampleWeight);
+					}
+					#endif
+
+					if (SampleWeight > 0.0f)
+					{
+						GBuffer.Roughness = max(GBuffer.Roughness, View.MinRoughness);
+
+						float4 LightAttenuation = 1.0f;
+						float Dither = 0.5f;
+						float SurfaceShadow = 1;
+						float AmbientOcclusion = ScreenSpaceAO;
+
+						// #sdl_todo: Move more shading models to TILE_MODE_SIMPLE_SHADING
+						#if TILE_TYPE == TILE_MODE_SIMPLE_SHADING
+						{
+							GBuffer.ShadingModelID = SHADINGMODELID_DEFAULT_LIT;
+						}
+						#endif
+						LightData.ShadowedBits = 0;
+
+						FDeferredLightingSplit SplitLighting = GetDynamicLightingSplit(
+							TranslatedWorldPosition, CameraVector, GBuffer, AmbientOcclusion, GBuffer.ShadingModelID, 
+							LightData, LightAttenuation, Dither, ScreenCoord, 
+							SurfaceShadow);
+
+						DiffuseLighting += SplitLighting.DiffuseLighting.xyz * SampleWeight;
+						SpecularLighting += SplitLighting.SpecularLighting.xyz * SampleWeight;
+
+						#if DEBUG_MODE 
+						if (bDebug)
+						{
+							Print(Context, SplitLighting.DiffuseLighting.xyz);
+							Print(Context, SplitLighting.SpecularLighting.xyz);
+						}
+						#endif
+					}
+				}
+
+				// Apply pre-exposure
+				DiffuseLighting *= View.PreExposure;
+				SpecularLighting *= View.PreExposure;
+
+				const float3 N = GBuffer.WorldNormal;
+				const float3 V = normalize(View.TranslatedWorldCameraOrigin - TranslatedWorldPosition);
+				const float NoV = saturate(dot(N, V));
+				const float3 SpecularEnv = EnvBRDF(GBuffer.SpecularColor, GBuffer.Roughness, NoV);
+
+				// #sdl_todo: demodulate earlier during BRDF evaluation
+				DiffuseLighting = SafeDemodulate(DiffuseLighting, GBuffer.DiffuseColor);
+				SpecularLighting = SafeDemodulate(SpecularLighting, SpecularEnv);
+			}
+
+			#if DEBUG_MODE
+			if (bDebug)
+			{
+				Newline(Context);
+				Print(Context, TEXT("Diffuse: ")); 
+				Print(Context, DiffuseLighting);
+				Newline(Context);
+				Print(Context, TEXT("Specular: ")); 
+				Print(Context, SpecularLighting);
+			}
+			#endif
+
+			RWResolvedDiffuseLighting[ScreenCoord] = DiffuseLighting;
+			RWResolvedSpecularLighting[ScreenCoord] = SpecularLighting;
+		}
+	}
+}
\ No newline at end of file
diff --git a/Engine/Shaders/Private/StochasticDirectLighting/StochasticDirectLightingSpatial.usf b/Engine/Shaders/Private/StochasticDirectLighting/StochasticDirectLightingSpatial.usf
new file mode 100644
--- /dev/null
+++ b/Engine/Shaders/Private/StochasticDirectLighting/StochasticDirectLightingSpatial.usf
@@ -0,0 +1,140 @@
+// Copyright Epic Games, Inc. All Rights Reserved.
+
+#include "../Common.ush"
+#include "StochasticDirectLighting.ush"
+
+Texture2D<float4> DiffuseLightingAndSecondMomentTexture;
+Texture2D<float4> SpecularLightingAndSecondMomentTexture;
+RWTexture2D<float4> RWSceneColor;
+
+float SpatialFilterDepthWeightScale;
+
+/**
+ * Run a spatial filter in order to filter out noise based on the temporal variance.
+ */
+[numthreads(THREADGROUP_SIZE, THREADGROUP_SIZE, 1)]
+void SpatialFilterCS(
+	uint3 GroupId : SV_GroupID,
+	uint3 GroupThreadId : SV_GroupThreadID,
+	uint3 DispatchThreadId : SV_DispatchThreadID)
+{
+	uint2 ScreenCoord = DispatchThreadId.xy;
+
+#if DEBUG_MODE
+	int2 DebugScreenCoord = View.CursorPosition.x >= 0 ? View.CursorPosition * View.ViewResolutionFraction : -1;
+	bool bDebug = all(ScreenCoord == DebugScreenCoord);
+	FShaderPrintContext Context = InitShaderPrintContext(true, float2(0.55, 0.45));
+#endif
+
+	if (all(ScreenCoord < View.ViewRectMinAndSize.zw))
+	{
+		float2 ScreenUV = (ScreenCoord + 0.5f) * View.BufferSizeAndInvSize.zw;
+		FGBufferData GBuffer = GetGBufferData(ScreenUV);
+
+		const float DeviceZ = SceneDepthTexture[ScreenCoord].x;
+		GBuffer.Depth = ConvertFromDeviceZ(DeviceZ);
+
+		if (GBuffer.Depth > 0.0f && GBuffer.ShadingModelID != SHADINGMODELID_UNLIT)
+		{
+			float4 CenterDiffuseLightingAndSecondMoment = DiffuseLightingAndSecondMomentTexture[ScreenCoord];
+			float4 CenterSpecularLightingAndSecondMoment = SpecularLightingAndSecondMomentTexture[ScreenCoord];
+
+			float CenterDiffuseStdDev = sqrt(max(CenterDiffuseLightingAndSecondMoment.w - Pow2(Luminance(CenterDiffuseLightingAndSecondMoment.xyz)), 0.0f));
+			float CenterSpecularStdDev = sqrt(max(CenterSpecularLightingAndSecondMoment.w - Pow2(Luminance(CenterSpecularLightingAndSecondMoment.xyz)), 0.0f));
+
+			float3 DiffuseLightingSum = CenterDiffuseLightingAndSecondMoment.xyz;
+			float3 SpecularLightingSum = CenterSpecularLightingAndSecondMoment.xyz;
+			float DiffuseLightingWeightSum = 1.0f;
+			float SpecularLightingWeightSum = 1.0f;
+
+			// #sdl_todo: expose as CVars
+			float KernelRadiusInPixels = 16.0f;
+			uint NumSamples = 8;
+
+			float GuassianNormalize = 2.0f / (KernelRadiusInPixels * KernelRadiusInPixels);
+			uint2 RandomSeed = Rand3DPCG16(int3(ScreenCoord, StochasticDirectLightingStateFrameIndex)).xy;
+
+			float3 TranslatedWorldPosition = GetTranslatedWorldPositionFromScreenUV(ScreenUV, GBuffer.Depth);
+			float4 ScenePlane = float4(GBuffer.WorldNormal, dot(TranslatedWorldPosition, GBuffer.WorldNormal));
+
+			#if SPATIAL_FILTER
+			if (KernelRadiusInPixels > 0.0f)
+			{
+				for (uint NeighborIndex = 0; NeighborIndex < NumSamples; ++NeighborIndex)
+				{
+					float2 NeighborOffsetInRect = Hammersley16(NeighborIndex, NumSamples, RandomSeed);
+					float2 NeighborOffset = UniformSampleDiskConcentric(NeighborOffsetInRect) * KernelRadiusInPixels;
+					int2 NeighborCoord = (int2)(ScreenCoord + NeighborOffset);
+
+					if (all(and(NeighborCoord >= View.ViewRectMinAndSize.xy, NeighborCoord < (View.ViewRectMinAndSize.xy + View.ViewRectMinAndSize.zw))))
+					{
+						// Depth weight
+						float2 NeighborScreenUV = (NeighborCoord + 0.5f) * View.BufferSizeAndInvSize.zw;
+						FGBufferData NeighborGBuffer = GetGBufferData(NeighborScreenUV);
+						const float3 NeighborTranslatedWorldPosition = GetTranslatedWorldPositionFromScreenUV(NeighborScreenUV, NeighborGBuffer.Depth);
+						float PlaneDistance = abs(dot(float4(NeighborTranslatedWorldPosition, -1), ScenePlane));
+						float RelativeDepthDifference = PlaneDistance / GBuffer.Depth;
+						float DepthWeight = exp2(-SpatialFilterDepthWeightScale * (RelativeDepthDifference * RelativeDepthDifference));
+
+						// #sdl_todo: separate weight for specular accounting for roughness and expose as CVars
+						// Normal weight
+						float AngleBetweenNormals = acosFast(saturate(dot(ScenePlane.xyz, NeighborGBuffer.WorldNormal)));
+						float NormalWeight = 1.0f - saturate(AngleBetweenNormals);
+
+						// Diffuse
+						{
+							float4 NeighborDiffuseLightingAndSecondMoment = DiffuseLightingAndSecondMomentTexture[NeighborCoord];
+
+							float LuminanceDelta = abs(Luminance(CenterDiffuseLightingAndSecondMoment.xyz) - Luminance(NeighborDiffuseLightingAndSecondMoment.xyz));
+							float LuminanceWeight = exp2(-LuminanceDelta / max(1.0f * CenterDiffuseStdDev, 0.00001f));
+
+							const float DiffuseNeighborWeight = DepthWeight * NormalWeight * LuminanceWeight;
+							DiffuseLightingSum += NeighborDiffuseLightingAndSecondMoment.xyz * DiffuseNeighborWeight;
+							DiffuseLightingWeightSum += DiffuseNeighborWeight;
+						}
+
+						// Specular
+						{
+							float4 NeighborSpecularLightingAndSecondMoment = SpecularLightingAndSecondMomentTexture[NeighborCoord];
+
+							float LuminanceDelta = abs(Luminance(CenterSpecularLightingAndSecondMoment.xyz) - Luminance(NeighborSpecularLightingAndSecondMoment.xyz));
+							float LuminanceWeight = exp2(-LuminanceDelta / max(1.0f * CenterSpecularStdDev, 0.00001f));
+
+							const float SpecularNeighborWeight = DepthWeight * NormalWeight * LuminanceWeight;
+							SpecularLightingSum += NeighborSpecularLightingAndSecondMoment.xyz * SpecularNeighborWeight;
+							SpecularLightingWeightSum += SpecularNeighborWeight;
+						}
+					}
+				}
+			}
+			#endif
+
+			const float3 N = GBuffer.WorldNormal;
+			const float3 V = normalize(View.TranslatedWorldCameraOrigin - TranslatedWorldPosition);
+			const float NoV = saturate(dot(N, V));
+			const float3 SpecularEnv = EnvBRDF(GBuffer.SpecularColor, GBuffer.Roughness, NoV);
+
+			// Final pass outputs composites irradiance and outputs it to scene color
+			float3 DiffuseLighting = SafeModulate(DiffuseLightingSum / DiffuseLightingWeightSum, GBuffer.DiffuseColor);
+			float3 SpecularLighting = SafeModulate(SpecularLightingSum / SpecularLightingWeightSum, SpecularEnv);
+
+			// Composite diffuse and specular into scene color
+			const bool bNeedsSeparateSubsurfaceLightAccumulation = UseSubsurfaceProfile(GBuffer.ShadingModelID);
+			FLightAccumulator LightAccumulator = (FLightAccumulator)0;
+			LightAccumulator_AddSplit(LightAccumulator, DiffuseLighting, SpecularLighting, /*ScatterableLight*/ DiffuseLighting, /*CommonMultiplier*/ 1.0f, bNeedsSeparateSubsurfaceLightAccumulation);
+			float4 AccumulatedSceneColor = LightAccumulator_GetResult(LightAccumulator);
+
+			#if DEBUG_MODE
+			if (bDebug)
+			{
+				Print(Context, TEXT("Spatial pass"));
+				Newline(Context);
+				Print(Context, TEXT("AccumulatedSceneColor: "));
+				Print(Context, AccumulatedSceneColor);
+			}
+			#endif
+
+			RWSceneColor[ScreenCoord] = RWSceneColor[ScreenCoord] + AccumulatedSceneColor;
+		}
+	}
+}
\ No newline at end of file
diff --git a/Engine/Shaders/Private/StochasticDirectLighting/StochasticDirectLightingTemporal.usf b/Engine/Shaders/Private/StochasticDirectLighting/StochasticDirectLightingTemporal.usf
new file mode 100644
--- /dev/null
+++ b/Engine/Shaders/Private/StochasticDirectLighting/StochasticDirectLightingTemporal.usf
@@ -0,0 +1,230 @@
+// Copyright Epic Games, Inc. All Rights Reserved.
+
+#include "../Common.ush"
+#include "StochasticDirectLighting.ush"
+
+Texture2D ResolvedDiffuseLighting;
+Texture2D ResolvedSpecularLighting;
+Texture2D StochasticDirectLightingDepthHistory;
+Texture2D DiffuseLightingAndSecondMomentHistoryTexture;
+Texture2D SpecularLightingAndSecondMomentHistoryTexture;
+Texture2D<UNORM float> NumFramesAccumulatedHistoryTexture;
+
+float4 HistoryScreenPositionScaleBias;
+float4 HistoryUVMinMax;
+float PrevSceneColorPreExposureCorrection;
+float TemporalMaxFramesAccumulated;
+float TemporalNeighborhoodClampScale;
+
+RWTexture2D<float4> RWDiffuseLightingAndSecondMoment;
+RWTexture2D<float4> RWSpecularLightingAndSecondMoment;
+RWTexture2D<UNORM float> RWNumFramesAccumulated;
+RWTexture2D<float> RWSceneDepth;
+RWTexture2D<float4> RWSceneColor;
+
+struct FNeighborhoodClamp
+{
+	float3 Min;
+	float3 Max;
+};
+
+FNeighborhoodClamp GetNeighborhoodClamp(
+	uint2 ScreenCoord, 
+	Texture2D SampleTexture,
+	float3 CenterSample)
+{
+	float3 SampleSum = CenterSample;
+	float3 SampleSqSum = Pow2(CenterSample);
+
+	// #sdl_todo: optimize using groupshared and cross-lane ops
+	const int KernelSize = 2;
+	for (int NeigborOffsetY = -KernelSize; NeigborOffsetY <= KernelSize; ++NeigborOffsetY)
+	{
+		for (int NeigborOffsetX = -KernelSize; NeigborOffsetX <= KernelSize; ++NeigborOffsetX)
+		{
+			if (!(NeigborOffsetX == 0 && NeigborOffsetY == 0))
+			{
+				uint2 NeigborScreenCoord = clamp(ScreenCoord + int2(NeigborOffsetX, NeigborOffsetY), 0u, View.ViewRectMinAndSize.zw - 1);
+				float3 NeigborSample = SampleTexture[NeigborScreenCoord].xyz;
+				
+				SampleSum += NeigborSample;
+				SampleSqSum += Pow2(NeigborSample);
+			}
+		}
+	}
+
+	float NumSamples = Pow2(2.0f * KernelSize + 1.0f);
+	float3 M1 = SampleSum / NumSamples;
+	float3 M2 = SampleSqSum / NumSamples;
+	float3 Variance = max(M2 - Pow2(M1), 0.0f);
+	float3 StdDev = sqrt(Variance);
+
+	FNeighborhoodClamp NeighborhoodClamp = (FNeighborhoodClamp)0;
+	NeighborhoodClamp.Min = M1 - TemporalNeighborhoodClampScale * StdDev;
+	NeighborhoodClamp.Max = M1 + TemporalNeighborhoodClampScale * StdDev;
+	return NeighborhoodClamp;
+}
+
+[numthreads(THREADGROUP_SIZE, THREADGROUP_SIZE, 1)]
+void TemporalAccumulationCS(
+	uint3 GroupId : SV_GroupID,
+	uint3 GroupThreadId : SV_GroupThreadID,
+	uint3 DispatchThreadId : SV_DispatchThreadID)
+{
+	uint2 ScreenCoord = DispatchThreadId.xy;
+
+#if DEBUG_MODE
+	int2 DebugScreenCoord = View.CursorPosition.x >= 0 ? View.CursorPosition * View.ViewResolutionFraction : -1;
+	bool bDebug = all(ScreenCoord == DebugScreenCoord);
+	FShaderPrintContext Context = InitShaderPrintContext(true, float2(0.55, 0.3));
+#endif
+
+	if (all(ScreenCoord < View.ViewRectMinAndSize.zw))
+	{
+		float2 ScreenUV = (ScreenCoord + 0.5f) * View.BufferSizeAndInvSize.zw;
+		FGBufferData GBuffer = GetGBufferData(ScreenUV);
+
+		const float DeviceZ = SceneDepthTexture[ScreenCoord].x;
+		GBuffer.Depth = ConvertFromDeviceZ(DeviceZ);
+
+		float3 DiffuseLighting = 0.0f;
+		float3 SpecularLighting = 0.0f;
+		float DiffuseSecondMoment = 0.0f;
+		float SpecularSecondMoment = 0.0f;
+		float NumFramesAccumulated = 0.0f;
+
+		if (GBuffer.Depth > 0.0f && GBuffer.ShadingModelID != SHADINGMODELID_UNLIT)
+		{
+			DiffuseLighting = ResolvedDiffuseLighting[ScreenCoord].xyz;
+			SpecularLighting = ResolvedSpecularLighting[ScreenCoord].xyz;
+
+			DiffuseSecondMoment = Pow2(Luminance(DiffuseLighting));
+			SpecularSecondMoment = Pow2(Luminance(SpecularLighting));
+
+			#if VALID_HISTORY
+			{
+				float2 ScreenPosition = (ScreenUV - View.ScreenPositionScaleBias.wz) / View.ScreenPositionScaleBias.xy;
+				float3 HistoryScreenPosition = GetHistoryScreenPosition(ScreenPosition, ScreenUV, DeviceZ);
+				float2 HistoryScreenUV = HistoryScreenPosition.xy * HistoryScreenPositionScaleBias.xy + HistoryScreenPositionScaleBias.wz;
+
+				bool bHistoryWasOnScreen = all(HistoryScreenUV >= HistoryUVMinMax.xy) && all (HistoryScreenUV <= HistoryUVMinMax.zw);
+				if (bHistoryWasOnScreen)
+				{
+					uint2 HistoryScreenCoord = floor(HistoryScreenUV * View.BufferSizeAndInvSize.xy - 0.5f);
+					float2 HistoryBilinearWeights = frac(HistoryScreenUV * View.BufferSizeAndInvSize.xy - 0.5f);
+					float2 HistoryGatherUV = (HistoryScreenCoord + 1.0f) * View.BufferSizeAndInvSize.zw;
+
+					float4 HistorySampleSceneDepth4 = StochasticDirectLightingDepthHistory.GatherRed(GlobalPointClampedSampler, HistoryGatherUV).wzxy;
+
+					// #sdl_todo: plane weighting
+					float ReprojectedSceneDepth = ConvertFromDeviceZ(HistoryScreenPosition.z);
+					float DisocclusionDistanceThreshold = 0.03f;
+					const float4 DistanceToHistoryValue = abs(HistorySampleSceneDepth4 - ReprojectedSceneDepth);
+					float4 DepthWeights = select(DistanceToHistoryValue >= ReprojectedSceneDepth * DisocclusionDistanceThreshold, 0.0f, 1.0f);
+
+					float4 HistoryWeights = float4(
+						(1 - HistoryBilinearWeights.y) * (1 - HistoryBilinearWeights.x),
+						(1 - HistoryBilinearWeights.y) * HistoryBilinearWeights.x,
+						HistoryBilinearWeights.y * (1 - HistoryBilinearWeights.x),
+						HistoryBilinearWeights.y * HistoryBilinearWeights.x);
+
+					HistoryWeights *= DepthWeights;
+
+					#if DEBUG_MODE	
+					if (bDebug)
+					{
+						Print(Context, TEXT("TemporalAccumulation"));
+						Newline(Context);
+						Print(Context, TEXT("ScreenCoord: "));
+						Print(Context, ScreenCoord.x);
+						Print(Context, ScreenCoord.y);
+						Newline(Context);
+						Print(Context, TEXT("DepthWeights: "));
+						Print(Context, DepthWeights);
+						Newline(Context);
+						Print(Context, TEXT("HistoryWeights: "));
+						Print(Context, HistoryWeights);
+					}
+					#endif
+
+					// Use history only if we found at least one valid sample
+					if (any(HistoryWeights > 0.01f))
+					{
+						float4 DiffuseLightingAndSecondMomentHistory = 0.0f;
+						DiffuseLightingAndSecondMomentHistory.x = dot(DiffuseLightingAndSecondMomentHistoryTexture.GatherRed(GlobalPointClampedSampler, HistoryGatherUV).wzxy, HistoryWeights);
+						DiffuseLightingAndSecondMomentHistory.y = dot(DiffuseLightingAndSecondMomentHistoryTexture.GatherGreen(GlobalPointClampedSampler, HistoryGatherUV).wzxy, HistoryWeights);
+						DiffuseLightingAndSecondMomentHistory.z = dot(DiffuseLightingAndSecondMomentHistoryTexture.GatherBlue(GlobalPointClampedSampler, HistoryGatherUV).wzxy, HistoryWeights);
+						DiffuseLightingAndSecondMomentHistory.w = dot(DiffuseLightingAndSecondMomentHistoryTexture.GatherAlpha(GlobalPointClampedSampler, HistoryGatherUV).wzxy, HistoryWeights);
+						DiffuseLightingAndSecondMomentHistory = DiffuseLightingAndSecondMomentHistory / dot(HistoryWeights, 1.0f);
+
+						float4 SpecularLightingAndSecondMomentHistory = 0.0f;
+						SpecularLightingAndSecondMomentHistory.x = dot(SpecularLightingAndSecondMomentHistoryTexture.GatherRed(GlobalPointClampedSampler, HistoryGatherUV).wzxy, HistoryWeights);
+						SpecularLightingAndSecondMomentHistory.y = dot(SpecularLightingAndSecondMomentHistoryTexture.GatherGreen(GlobalPointClampedSampler, HistoryGatherUV).wzxy, HistoryWeights);
+						SpecularLightingAndSecondMomentHistory.z = dot(SpecularLightingAndSecondMomentHistoryTexture.GatherBlue(GlobalPointClampedSampler, HistoryGatherUV).wzxy, HistoryWeights);
+						SpecularLightingAndSecondMomentHistory.w = dot(SpecularLightingAndSecondMomentHistoryTexture.GatherAlpha(GlobalPointClampedSampler, HistoryGatherUV).wzxy, HistoryWeights);
+						SpecularLightingAndSecondMomentHistory = SpecularLightingAndSecondMomentHistory / dot(HistoryWeights, 1.0f);
+
+						// Correct history for current frame exposure
+						DiffuseLightingAndSecondMomentHistory.xyz *= PrevSceneColorPreExposureCorrection;
+						SpecularLightingAndSecondMomentHistory.xyz *= PrevSceneColorPreExposureCorrection;
+
+						DiffuseLightingAndSecondMomentHistory.w *= Pow2(PrevSceneColorPreExposureCorrection);
+						SpecularLightingAndSecondMomentHistory.w *= Pow2(PrevSceneColorPreExposureCorrection);
+
+						// Reproject and rescale normalized NumFramesAccumulated
+						float NumFramesAccumulatedHistory = 0.0f;
+						NumFramesAccumulatedHistory = dot(NumFramesAccumulatedHistoryTexture.GatherRed(GlobalPointClampedSampler, HistoryGatherUV).wzxy, HistoryWeights) * TemporalMaxFramesAccumulated;
+						NumFramesAccumulatedHistory = NumFramesAccumulatedHistory / dot(HistoryWeights, 1.0f);
+						NumFramesAccumulated = min(NumFramesAccumulatedHistory + 1.0f, TemporalMaxFramesAccumulated);
+
+						float Alpha = 1.0f / (1.0f + NumFramesAccumulated);
+						DiffuseLighting = lerp(DiffuseLightingAndSecondMomentHistory.xyz, DiffuseLighting, Alpha);
+						SpecularLighting = lerp(SpecularLightingAndSecondMomentHistory.xyz, SpecularLighting, Alpha);
+						DiffuseSecondMoment = lerp(DiffuseLightingAndSecondMomentHistory.w, DiffuseSecondMoment, Alpha);
+						SpecularSecondMoment = lerp(SpecularLightingAndSecondMomentHistory.w, SpecularSecondMoment, Alpha);
+
+						FNeighborhoodClamp DiffuseNeighborhoodClamp = GetNeighborhoodClamp(ScreenCoord, ResolvedDiffuseLighting, DiffuseLighting);
+						FNeighborhoodClamp SpecularNeighborhoodClamp = GetNeighborhoodClamp(ScreenCoord, ResolvedSpecularLighting, SpecularLighting);
+
+						float3 ClampedDiffuseLighting = clamp(DiffuseLighting, DiffuseNeighborhoodClamp.Min, DiffuseNeighborhoodClamp.Max);
+						float3 ClampedSpecularLighting = clamp(SpecularLighting, SpecularNeighborhoodClamp.Min, SpecularNeighborhoodClamp.Max);
+
+						// Approximate clamp for second moments in order to reduce ghosting from the spatial pass while still allowing accumulated secondary moments to deviate from accumulated lighting
+						DiffuseSecondMoment += Pow2(Luminance(ClampedDiffuseLighting)) - Pow2(Luminance(DiffuseLighting));
+						SpecularSecondMoment += Pow2(Luminance(ClampedSpecularLighting)) - Pow2(Luminance(SpecularLighting));
+
+						DiffuseLighting = ClampedDiffuseLighting;
+						SpecularLighting = ClampedSpecularLighting;
+					}
+				}
+			}
+			#endif
+
+			DiffuseLighting = MakeFinite(DiffuseLighting);
+			SpecularLighting = MakeFinite(SpecularLighting);
+
+			#if DEBUG_MODE
+			if (bDebug)
+			{
+				float DiffuseVariance = max(DiffuseSecondMoment - Pow2(Luminance(DiffuseLighting)), 0.0f);
+				float SpecularVariance = max(SpecularSecondMoment - Pow2(Luminance(SpecularLighting)), 0.0f);
+
+				Newline(Context);
+				Print(Context, TEXT("NumFramesAccumulated: "));
+				Print(Context, NumFramesAccumulated);
+				Newline(Context);
+				Print(Context, TEXT("DiffuseVariance:  "));
+				Print(Context, DiffuseVariance);
+				Newline(Context);
+				Print(Context, TEXT("SpecularVariance: "));
+				Print(Context, SpecularVariance);
+			}
+			#endif
+		}
+
+		RWDiffuseLightingAndSecondMoment[ScreenCoord] = float4(DiffuseLighting, DiffuseSecondMoment);
+		RWSpecularLightingAndSecondMoment[ScreenCoord] = float4(SpecularLighting, SpecularSecondMoment);
+		RWNumFramesAccumulated[ScreenCoord] = (NumFramesAccumulated + 0.5f) / TemporalMaxFramesAccumulated;
+		RWSceneDepth[ScreenCoord] = GBuffer.Depth;
+	}
+}
\ No newline at end of file
diff --git a/Engine/Shaders/Private/StochasticDirectLighting/StochasticDirectLightingTracing.usf b/Engine/Shaders/Private/StochasticDirectLighting/StochasticDirectLightingTracing.usf
new file mode 100644
--- /dev/null
+++ b/Engine/Shaders/Private/StochasticDirectLighting/StochasticDirectLightingTracing.usf
@@ -0,0 +1,373 @@
+// Copyright Epic Games, Inc. All Rights Reserved.
+
+#include "../Common.ush"
+#define DISTANCE_FIELD_IN_VIEW_UB 1
+#include "../BlueNoise.ush"
+#include "../DistanceField/GlobalDistanceFieldShared.ush"
+#include "../DistanceField/GlobalDistanceFieldUtils.ush"
+ #define HZB_TRACE_INCLUDE_FULL_RES_DEPTH 1
+#include "../Lumen/LumenScreenTracing.ush"
+#include "../SceneTexturesCommon.ush"
+#include "StochasticDirectLighting.ush"
+#include "StochasticDirectLightingTracing.ush"
+#if HAIR_VOXEL_TRACES
+	#include "../HairStrands/HairStrandsRaytracing.ush"
+#endif
+
+StructuredBuffer<uint> CompactedTraceTexelAllocator;
+RWBuffer<uint> RWIndirectArgs;
+Texture2D<uint> DownsampledTileMask;
+
+[numthreads(THREADGROUP_SIZE, 1, 1)]
+void InitCompactedTraceTexelIndirectArgsCS(uint3 DispatchThreadId : SV_DispatchThreadID)
+{
+	if (all(DispatchThreadId == 0))
+	{
+		// NumTracesDiv64
+		uint NumGroups = (CompactedTraceTexelAllocator[0] + 63) / 64;
+		WriteDispatchIndirectArgs(RWIndirectArgs, 0, NumGroups, 1, 1);
+
+		// NumTracesDiv32
+		NumGroups = (CompactedTraceTexelAllocator[0] + 31) / 32;
+		WriteDispatchIndirectArgs(RWIndirectArgs, 1, NumGroups, 1, 1);
+
+		// NumTraces
+		NumGroups = CompactedTraceTexelAllocator[0];
+		WriteDispatchIndirectArgs(RWIndirectArgs, 2, NumGroups, 1, 1);
+	}
+}
+
+uint2 SampleViewSize;
+uint DebugMode;
+float4 HZBUvFactorAndInvFactor;
+float MaxHierarchicalScreenTraceIterations;
+float RelativeDepthThickness;
+float HistoryDepthTestRelativeThickness;
+uint MinimumTracingThreadOccupancy;
+
+Texture2D<uint> LightSamples;
+
+Texture2D<uint> RWTraceTile;
+RWStructuredBuffer<uint> RWCompactedTraceTexelData;
+RWStructuredBuffer<uint> RWCompactedTraceTexelAllocator;
+
+groupshared uint SharedGlobalTraceTexelStartOffset;
+
+#if WAVE_OPS
+	#define MIN_SUPPORTED_WAVE_SIZE 32
+	#define NUM_WAVE_SUMS ((THREADGROUP_SIZE * THREADGROUP_SIZE + MIN_SUPPORTED_WAVE_SIZE - 1) / MIN_SUPPORTED_WAVE_SIZE)
+	groupshared uint SharedWaveSums[NUM_WAVE_SUMS];
+	groupshared uint SharedTexelValid[THREADGROUP_SIZE * THREADGROUP_SIZE];
+#else
+	//@todo - ordered compaction for non-wave ops path
+	groupshared uint SharedTraceTexelAllocator;
+	groupshared uint SharedTraceTexels[THREADGROUP_SIZE * THREADGROUP_SIZE];
+#endif
+
+#if COMPILER_SUPPORTS_WAVE_SIZE && WAVE_OPS
+	WAVESIZE(32)
+#endif
+[numthreads(THREADGROUP_SIZE, THREADGROUP_SIZE, 1)]
+void CompactLightSampleTracesCS(
+	uint3 GroupId : SV_GroupID,
+	uint3 GroupThreadId : SV_GroupThreadID,
+	uint3 DispatchThreadId : SV_DispatchThreadID)
+{
+	const uint LinearThreadIndex = GroupThreadId.y * THREADGROUP_SIZE + GroupThreadId.x;
+
+	#if WAVE_OPS
+	{
+		SharedTexelValid[LinearThreadIndex] = 0;
+	}
+	#else
+	if (LinearThreadIndex == 0)
+	{
+		SharedTraceTexelAllocator = 0;
+	}
+	#endif
+
+	GroupMemoryBarrierWithGroupSync();
+
+	uint2 SampleCoord = DispatchThreadId.xy;
+	bool bTraceValid = false;
+	uint TraceTexelForThisThread = 0;
+
+	uint2 LocalTileMaskCoord = SampleCoordToDownsampledScreenCoord(SampleCoord - GroupId.xy * THREADGROUP_SIZE) / TILE_SIZE;
+	if (all(SampleCoord < SampleViewSize))
+	{
+		// #sdl_todo: try to preserve z-order here for improved tracing coherency
+		FLightSample LightSample = UnpackLightSample(LightSamples[SampleCoord]);
+
+		if (LightSample.LocalLightIndex != INVALID_LOCAL_LIGHT_INDEX && !LightSample.bCompleted)
+		{
+			#if WAVE_OPS
+			{
+				bTraceValid = true;
+				SharedTexelValid[LinearThreadIndex] = 1;
+				TraceTexelForThisThread = PackTraceTexel(SampleCoord);
+			}
+			#else
+			{
+				uint SharedTexelOffset;
+				InterlockedAdd(SharedTraceTexelAllocator, 1, SharedTexelOffset);
+				SharedTraceTexels[SharedTexelOffset] = PackTraceTexel(SampleCoord);
+			}
+			#endif
+		}
+	}
+
+	GroupMemoryBarrierWithGroupSync();
+
+	#if WAVE_OPS
+	{
+		const uint WaveIndex = LinearThreadIndex / WaveGetLaneCount();
+		const uint LaneIndex = WaveGetLaneIndex();
+
+		uint Value = SharedTexelValid[LinearThreadIndex];
+		uint ThisWaveSum = WaveActiveSum(Value);
+
+		if (LaneIndex == 0)
+		{
+			SharedWaveSums[WaveIndex] = ThisWaveSum;
+		}
+
+		GroupMemoryBarrierWithGroupSync();
+
+		if (WaveIndex == 0 && LaneIndex < NUM_WAVE_SUMS)
+		{
+			SharedWaveSums[LaneIndex] = WavePrefixSum(SharedWaveSums[LaneIndex]);
+		}
+
+		GroupMemoryBarrierWithGroupSync();
+
+		uint PayloadGroupSize = SharedWaveSums[WaveIndex] + WavePrefixSum(Value) + Value;
+
+		uint WaveLaneCount = WaveGetLaneCount();
+		uint LastWaveIndex = (THREADGROUP_SIZE * THREADGROUP_SIZE - 1) / WaveLaneCount;
+
+		// Last lane of last wave has the max index used by the group
+		// Allocate this group's compacted traces from the global allocator
+		if (WaveIndex == LastWaveIndex && LaneIndex == WaveLaneCount - 1 && PayloadGroupSize > 0)
+		{
+			InterlockedAdd(RWCompactedTraceTexelAllocator[0], PayloadGroupSize, SharedGlobalTraceTexelStartOffset);
+		}
+
+		GroupMemoryBarrierWithGroupSync();
+
+		if (bTraceValid)
+		{
+			RWCompactedTraceTexelData[SharedGlobalTraceTexelStartOffset + PayloadGroupSize - 1] = TraceTexelForThisThread;
+		}
+	}
+	#else
+	{
+		if (LinearThreadIndex == 0)
+		{
+			InterlockedAdd(RWCompactedTraceTexelAllocator[0], SharedTraceTexelAllocator, SharedGlobalTraceTexelStartOffset);
+		}
+
+		GroupMemoryBarrierWithGroupSync();
+
+		if (LinearThreadIndex < SharedTraceTexelAllocator)
+		{
+			RWCompactedTraceTexelData[SharedGlobalTraceTexelStartOffset + LinearThreadIndex] = SharedTraceTexels[LinearThreadIndex];
+		}
+	}
+	#endif
+}
+
+RWTexture2D<uint> RWLightSamples;
+RWTexture2D<float> RWLightSampleRayDistance;
+
+StructuredBuffer<uint> CompactedTraceTexelData;
+Texture2D<float> DownsampledSceneDepth;
+Texture2D<UNORM float3> DownsampledSceneWorldNormal;
+
+[numthreads(THREADGROUP_SIZE, 1, 1)]
+void ScreenSpaceRayTraceLightSamplesCS(
+	uint3 GroupId : SV_GroupID,
+	uint3 GroupThreadId : SV_GroupThreadID,
+	uint3 DispatchThreadId : SV_DispatchThreadID)
+{
+#if DEBUG_MODE
+	FShaderPrintContext Context = InitShaderPrintContext(true, float2(0.05, 0.9));
+	if (all(DispatchThreadId == 0))
+	{
+		//Print(Context, TEXT("ScreenSpaceRayTraceLightSamplesCS"));
+		//Newline(Context);
+		//Print(Context, TEXT("CompactedTraceTexelAllocator: "));
+		//Print(Context, CompactedTraceTexelAllocator[0]);
+	}
+#endif
+
+	uint TraceTexelIndex = DispatchThreadId.x;
+
+	if (TraceTexelIndex < CompactedTraceTexelAllocator[0])
+	{
+		uint2 SampleCoord = UnpackTraceTexel(CompactedTraceTexelData[TraceTexelIndex]);
+		uint2 ScreenCoord = SampleCoordToScreenCoord(SampleCoord);
+		uint2 DownsampledScreenCoord = SampleCoordToDownsampledScreenCoord(SampleCoord);		
+
+		float2 ScreenUV = (ScreenCoord + 0.5f) * View.BufferSizeAndInvSize.zw;
+		float SceneDepth = DownsampledSceneDepth[DownsampledScreenCoord];
+		float3 SceneWorldNormal = normalize(DecodeNormal(DownsampledSceneWorldNormal[DownsampledScreenCoord]));
+		float3 TranslatedWorldPosition = GetTranslatedWorldPositionFromScreenUV(ScreenUV, SceneDepth);
+
+		FLightSample LightSample = UnpackLightSample(RWLightSamples[SampleCoord]);
+		FLightSampleTrace LightSampleTrace = GetLightSampleTrace(TranslatedWorldPosition, LightSample.LocalLightIndex, SampleCoord);
+
+		#if DEBUG_MODE
+			int2 DebugScreenCoord = View.CursorPosition.x >= 0 ? View.CursorPosition * View.ViewResolutionFraction : -1;
+			bool bDebug = all(DebugScreenCoord / DOWNSAMPLE_FACTOR == DownsampledScreenCoord);
+		#endif
+
+		#if !STRATA_ENABLED
+		// Bias along the normal to avoid self-intersecting camera facing scene depth texels
+		// HZB traversal uses point filtering, so scene depth texels stick out from the plane they are representing
+		{
+			float2 CornerScreenUV = ScreenUV + .5f * View.BufferSizeAndInvSize.zw;
+			float2 CornerScreenPosition = (CornerScreenUV - View.ScreenPositionScaleBias.wz) / View.ScreenPositionScaleBias.xy;
+			float3 CornerTranslatedWorldPosition = mul(float4(GetScreenPositionForProjectionType(CornerScreenPosition, SceneDepth), SceneDepth, 1), View.ScreenToTranslatedWorld).xyz;
+
+			// #sdl_todo: fixup bias and incorporate it into Lumen
+			// Project texel corner onto normal
+			float NormalBias = abs(dot(CornerTranslatedWorldPosition - TranslatedWorldPosition, SceneWorldNormal)) * 2.0f;
+			TranslatedWorldPosition += NormalBias * SceneWorldNormal + 0.05f * SceneWorldNormal;
+
+			float4 RayStartClip = mul(float4(TranslatedWorldPosition, 1.0f), View.TranslatedWorldToClip);
+
+			// Skip screen traces if the normal bias pushed our starting point off-screen, as TraceScreen() doesn't handle this
+			if (any(RayStartClip.xy < -RayStartClip.w) || any(RayStartClip.xy > RayStartClip.w))
+			{
+				return;
+			}
+		}
+		#endif
+
+		bool bHit;
+		bool bUncertain;
+		float3 HitUVz;
+		float3 LastVisibleHitUVz;
+		float HitTileZ;
+
+		TraceScreen(
+			TranslatedWorldPosition,
+			LightSampleTrace.Direction,
+			LightSampleTrace.Distance,
+			HZBUvFactorAndInvFactor,
+			MaxHierarchicalScreenTraceIterations, 
+			RelativeDepthThickness,
+			0,
+			MinimumTracingThreadOccupancy,
+			bHit,
+			bUncertain,
+			HitUVz,
+			LastVisibleHitUVz,
+			HitTileZ);
+
+		bHit = bHit && !bUncertain;
+
+		const float RayT = sqrt(ComputeRayHitSqrDistance(TranslatedWorldPosition, HitUVz));
+
+		if (bHit)
+		{
+			LightSample.bVisible = false;
+			LightSample.bCompleted = true;
+			RWLightSamples[SampleCoord] = PackLightSample(LightSample);
+		}
+		else
+		{
+			// Output ray distance for ray continuation in next pass
+			RWLightSampleRayDistance[SampleCoord] = RayT;
+		}
+
+		#if DEBUG_MODE
+		if (bDebug && DebugMode == DEBUG_MODE_VISUALIZE_TRACING)
+		{
+			const FLocalLightData LocalLightData = GetLocalLightDataNonStereo(LightSample.LocalLightIndex);
+			const FLightSceneData LightSceneData = GetLightSceneData(LocalLightData.LightSceneId);
+			const FDeferredLightData LightData = ConvertToDeferredLight(LightSceneData);
+
+			float3 RayStart = TranslatedWorldPosition;
+			float3 RayEnd = TranslatedWorldPosition + LightSampleTrace.Direction * min(LightSampleTrace.Distance, RayT);
+			float4 RayColor = float4(LightData.Color.xyz / Luminance(LightData.Color.xyz), 1.0f);
+
+			AddLineTWS(Context, RayStart, RayEnd, RayColor, 1.0f);
+			AddCrossTWS(Context, RayEnd, 2.0f, RayColor);
+		}
+		#endif
+	}
+}
+
+Texture2D<float> LightSampleRayDistance;
+
+[numthreads(THREADGROUP_SIZE, 1, 1)]
+void SoftwareRayTraceLightSamplesCS(
+	uint3 GroupId : SV_GroupID,
+	uint3 GroupThreadId : SV_GroupThreadID,
+	uint3 DispatchThreadId : SV_DispatchThreadID)
+{
+	uint TraceTexelIndex = DispatchThreadId.x;
+
+	if (TraceTexelIndex < CompactedTraceTexelAllocator[0])
+	{
+		uint2 SampleCoord = UnpackTraceTexel(CompactedTraceTexelData[TraceTexelIndex]);
+		uint2 ScreenCoord = SampleCoordToScreenCoord(SampleCoord);
+		uint2 DownsampledScreenCoord = SampleCoordToDownsampledScreenCoord(SampleCoord);
+
+#if DEBUG_MODE
+		FShaderPrintContext Context = InitShaderPrintContext(true, float2(0.55, 0.45));
+		int2 DebugScreenCoord = View.CursorPosition.x >= 0 ? View.CursorPosition * View.ViewResolutionFraction : -1;
+		bool bDebug = all(DebugScreenCoord / DOWNSAMPLE_FACTOR == DownsampledScreenCoord);
+#endif
+
+		float2 ScreenUV = (ScreenCoord + 0.5f) * View.BufferSizeAndInvSize.zw;
+		float SceneDepth = DownsampledSceneDepth[DownsampledScreenCoord];
+		float3 TranslatedWorldPosition = GetTranslatedWorldPositionFromScreenUV(ScreenUV, SceneDepth);
+
+		FLightSample LightSample = UnpackLightSample(RWLightSamples[SampleCoord]);
+		const float LightSampleRayMinT = LightSampleRayDistance[SampleCoord];
+		FLightSampleTrace LightSampleTrace = GetLightSampleTrace(TranslatedWorldPosition, LightSample.LocalLightIndex, SampleCoord);
+
+		FGlobalSDFTraceInput TraceInput = SetupGlobalSDFTraceInput(TranslatedWorldPosition, LightSampleTrace.Direction, LightSampleRayMinT, LightSampleTrace.Distance, 1.0f, 1.0f);
+
+		// #sdl_todo: setup bias
+		TraceInput.VoxelSizeRelativeBias = 0.5f;
+		TraceInput.VoxelSizeRelativeRayEndBias = 0.5f;
+		TraceInput.DitherScreenCoord = SampleCoord;
+
+		FGlobalSDFTraceResult SDFResult = RayTraceGlobalDistanceField(TraceInput);
+
+		#if HAIR_VOXEL_TRACES
+		{
+			// #sdl_todo: replace with spatiotemporal blue noise
+			RandomSequence RandSequence;
+			RandomSequence_Initialize(RandSequence, SampleCoord, 0, StochasticDirectLightingStateFrameIndex, 1);
+			SDFResult.HitTime = TraverseHair(ScreenCoord, RandSequence, TraceInput.TranslatedWorldRayStart, TraceInput.WorldRayDirection, SDFResult.HitTime, VirtualVoxel.Raytracing_ShadowOcclusionThreshold);
+		}
+		#endif
+
+		const bool bHit = GlobalSDFTraceResultIsHit(SDFResult);
+		if (bHit)
+		{
+			LightSample.bVisible = false;
+			LightSample.bCompleted = true;
+			RWLightSamples[SampleCoord] = PackLightSample(LightSample);
+		}
+
+		#if DEBUG_MODE
+		if (bDebug && DebugMode == DEBUG_MODE_VISUALIZE_TRACING)
+		{
+			const FLocalLightData LocalLightData = GetLocalLightDataNonStereo(LightSample.LocalLightIndex);
+			const FLightSceneData LightSceneData = GetLightSceneData(LocalLightData.LightSceneId);
+			const FDeferredLightData LightData = ConvertToDeferredLight(LightSceneData);
+
+			float3 RayStart = TraceInput.TranslatedWorldRayStart + TraceInput.WorldRayDirection * TraceInput.MinTraceDistance;
+			float3 RayEnd = TraceInput.TranslatedWorldRayStart + TraceInput.WorldRayDirection * (bHit ? SDFResult.HitTime : TraceInput.MaxTraceDistance);
+			float4 RayColor = float4(LightData.Color.xyz / Luminance(LightData.Color.xyz), 1.0f);
+
+			AddLineTWS(Context, RayStart, RayEnd, RayColor);
+		}
+		#endif
+	}
+}
\ No newline at end of file
diff --git a/Engine/Shaders/Private/StochasticDirectLighting/StochasticDirectLightingTracing.ush b/Engine/Shaders/Private/StochasticDirectLighting/StochasticDirectLightingTracing.ush
new file mode 100644
--- /dev/null
+++ b/Engine/Shaders/Private/StochasticDirectLighting/StochasticDirectLightingTracing.ush
@@ -0,0 +1,131 @@
+// Copyright Epic Games, Inc. All Rights Reserved.
+
+#pragma once
+
+#include "../RayTracing/RayTracingCapsuleLight.ush"
+#include "../RayTracing/RayTracingSphereLight.ush"
+#include "../RayTracing/RayTracingRectLight.ush"
+
+uint2 NumSamplesPerPixelDivideShift;
+
+uint2 SampleCoordToDownsampledScreenCoord(uint2 SampleCoord)
+{
+	// SampleCoord / NumSamplesPerPixel
+	return SampleCoord >> NumSamplesPerPixelDivideShift;
+}
+
+uint2 SampleCoordToScreenCoord(uint2 SampleCoord)
+{
+	uint2 DownsampledScreenCoord = SampleCoordToDownsampledScreenCoord(SampleCoord);
+	return DownsampledScreenCoord * DOWNSAMPLE_FACTOR + GetSampleScreenCoordJitter();
+}
+
+bool GenerateShadowRay(
+	FLightShaderParameters LightParameters,
+	bool bSpotLight,
+	bool bRectLight,
+	float3 TranslatedWorldPosition,
+	float3 WorldNormal,
+	float2 RandSample,
+	out float3 RayOrigin,
+	out float3 RayDirection,
+	out float RayTMin,
+	out float RayTMax)
+{
+	float RayPdf = 0.0;
+	if (bRectLight)
+	{
+		return GenerateRectLightOcclusionRay(
+			LightParameters,
+			TranslatedWorldPosition, WorldNormal,
+			RandSample,
+			/* out */ RayOrigin,
+			/* out */ RayDirection,
+			/* out */ RayTMin,
+			/* out */ RayTMax,
+			/* out */ RayPdf);
+	}
+	else
+	{
+		// point or spot light
+		if (LightParameters.SourceLength > 0.0)
+		{
+			return GenerateCapsuleLightOcclusionRayWithSolidAngleSampling(
+				LightParameters,
+				TranslatedWorldPosition, WorldNormal,
+				RandSample,
+				/* out */ RayOrigin,
+				/* out */ RayDirection,
+				/* out */ RayTMin,
+				/* out */ RayTMax,
+				/* out */ RayPdf);
+		}
+		return GenerateSphereLightOcclusionRayWithSolidAngleSampling(
+			LightParameters,
+			TranslatedWorldPosition, WorldNormal,
+			RandSample,
+			/* out */ RayOrigin,
+			/* out */ RayDirection,
+			/* out */ RayTMin,
+			/* out */ RayTMax,
+			/* out */ RayPdf);
+	}
+}
+
+struct FLightSampleTrace
+{
+	float3 Direction;
+	float Distance;
+};
+
+/**
+ * Compute shadow ray direction and distance for a given reservoir.
+ */
+FLightSampleTrace GetLightSampleTrace(float3 TranslatedWorldPosition, uint LocalLightIndex, uint2 SampleCoord)
+{
+	const float2 LightSampleRandom = BlueNoiseVec2(SampleCoord, StochasticDirectLightingStateFrameIndex);
+
+	FLightSampleTrace LightSampleTrace;
+	LightSampleTrace.Direction = 0.0f;
+	LightSampleTrace.Distance = 0.0f;
+
+	if (LocalLightIndex != INVALID_LOCAL_LIGHT_INDEX)
+	{
+		const FLocalLightData LocalLightData = GetLocalLightDataNonStereo(LocalLightIndex);
+		const FLightSceneData LightSceneData = GetLightSceneData(LocalLightData.LightSceneId);
+		const FDeferredLightData LightData = ConvertToDeferredLight(LightSceneData);
+		const FLightShaderParameters LightParameters = ConvertToLightShaderParameters(LightSceneData);
+
+		float3 Unused;
+		float Unused2;
+
+		bool bValid = GenerateShadowRay(
+			LightParameters,
+			LightData.bSpotLight,
+			LightData.bRectLight,
+			TranslatedWorldPosition,
+			float3(0, 0, 0),
+			LightSampleRandom,
+			Unused,
+			LightSampleTrace.Direction,
+			Unused2,
+			LightSampleTrace.Distance);
+	}
+
+	return LightSampleTrace;
+}
+
+uint PackTraceTexel(uint2 SampleCoord)
+{
+	uint PackedTraceTexel = SampleCoord.y & 0xFFFF;
+	PackedTraceTexel |= (SampleCoord.x & 0xFFFF) << 16;
+	return PackedTraceTexel;
+}
+
+uint2 UnpackTraceTexel(uint PackedTraceTexel)
+{
+	uint2 SampleCoord;
+	SampleCoord.x = PackedTraceTexel >> 16;
+	SampleCoord.y = PackedTraceTexel & 0xFFFF;
+	return SampleCoord;
+}
\ No newline at end of file
diff --git a/Engine/Shaders/Private/StochasticShadows/StochasticShadowsSpatial.usf b/Engine/Shaders/Private/StochasticShadows/StochasticShadowsSpatial.usf
--- a/Engine/Shaders/Private/StochasticShadows/StochasticShadowsSpatial.usf
+++ b/Engine/Shaders/Private/StochasticShadows/StochasticShadowsSpatial.usf
@@ -202,7 +202,7 @@ void ShadowMaskSpatialPassCS(
 					uint2 NeighborShadowMaskPageTableCoord = GetShadowMaskPageTableCoord(LocalLightIndex, NeighborCoord);
 					uint NeighborShadowMaskTileIndex = ShadowMaskPageTable[NeighborShadowMaskPageTableCoord];
 
-					// #kris_todo: CVar and normalize tags!
+					// #ss_todo: CVar and normalize tags!
 					float SpatialFilterDepthWeightScale = 10000.0f;
 
 					// Depth weight
diff --git a/Engine/Source/Runtime/Renderer/Private/DeferredShadingRenderer.cpp b/Engine/Source/Runtime/Renderer/Private/DeferredShadingRenderer.cpp
--- a/Engine/Source/Runtime/Renderer/Private/DeferredShadingRenderer.cpp
+++ b/Engine/Source/Runtime/Renderer/Private/DeferredShadingRenderer.cpp
@@ -99,6 +99,7 @@
 #include "WaterInfoTextureRendering.h"
 #include "PostProcess/DebugAlphaChannel.h"
 #include "StochasticShadows/StochasticShadows.h"
+#include "StochasticDirectLighting/StochasticDirectLighting.h"
 #include "Rendering/CustomRenderPass.h"
 
 #if !UE_BUILD_SHIPPING
@@ -764,6 +765,7 @@ bool FDeferredShadingSceneRenderer::SetupRayTracingPipelineStates(FRDGBuilder& G
 		for (const FViewInfo& View : Views)
 		{
 			PrepareStochasticShadowsLumenMaterial(View, LumenHardwareRayTracingRayGenShaders);
+			PrepareStochasticDirectLightingLumenMaterial(View, LumenHardwareRayTracingRayGenShaders);
 		}
 
 		DeduplicateRayGenerationShaders(LumenHardwareRayTracingRayGenShaders);
@@ -1017,7 +1019,8 @@ void FDeferredShadingSceneRenderer::WaitForRayTracingScene(FRDGBuilder& GraphBui
 	for (const FViewInfo& View : Views)
 	{
 		if (Lumen::AnyLumenHardwareInlineRayTracingPassEnabled(Scene, View) 
-			|| StochasticShadows::UseInlineHardwareRayTracing())
+			|| StochasticShadows::UseInlineHardwareRayTracing()
+			|| StochasticDirectLighting::UseInlineHardwareRayTracing())
 		{
 			bAnyLumenHardwareInlineRayTracingPassEnabled = true;
 		}
@@ -1028,7 +1031,9 @@ void FDeferredShadingSceneRenderer::WaitForRayTracingScene(FRDGBuilder& GraphBui
 		SetupLumenHardwareRayTracingHitGroupBuffer(GraphBuilder, ReferenceView);
 	}
 
-	if (Lumen::UseHardwareRayTracing(ViewFamily) || StochasticShadows::UseHardwareRayTracing())
+	if (Lumen::UseHardwareRayTracing(ViewFamily) 
+		|| StochasticShadows::UseHardwareRayTracing()
+		|| StochasticDirectLighting::UseHardwareRayTracing())
 	{
 		SetupLumenHardwareRayTracingUniformBuffer(GraphBuilder, ReferenceView);
 	}
@@ -1199,7 +1204,7 @@ void FDeferredShadingSceneRenderer::CommitFinalPipelineState()
 				bHasSSGI || bUseLumen);
 
 			ViewPipelineState.Set(&FPerViewPipelineState::bClosestHZB, 
-				bHasSSGI || bUseLumen || StochasticShadows::IsUsingClosestHZB());
+				bHasSSGI || bUseLumen || StochasticShadows::IsUsingClosestHZB() || StochasticDirectLighting::IsUsingClosestHZB());
 		}
 	}
 
@@ -2832,6 +2837,10 @@ void FDeferredShadingSceneRenderer::Render(FRDGBuilder& GraphBuilder)
 				RenderStochasticShadows(
 					GraphBuilder,
 					SceneTextures);
+
+				RenderStochasticDirectLighting(
+					GraphBuilder,
+					SceneTextures);
 			}
 
 			InjectTranslucencyLightingVolumeAmbientCubemap(GraphBuilder, Views, TranslucencyLightingVolumeTextures);
diff --git a/Engine/Source/Runtime/Renderer/Private/DeferredShadingRenderer.h b/Engine/Source/Runtime/Renderer/Private/DeferredShadingRenderer.h
--- a/Engine/Source/Runtime/Renderer/Private/DeferredShadingRenderer.h
+++ b/Engine/Source/Runtime/Renderer/Private/DeferredShadingRenderer.h
@@ -722,6 +722,10 @@ private:
 		FRDGBuilder& GraphBuilder,
 		const FSceneTextures& SceneTextures);
 
+	void RenderStochasticDirectLighting(
+		FRDGBuilder& GraphBuilder,
+		const FSceneTextures& SceneTextures);
+
 	FSSDSignalTextures RenderLumenFinalGather(
 		FRDGBuilder& GraphBuilder,
 		const FSceneTextures& SceneTextures,
@@ -1153,6 +1157,7 @@ private:
 	static void PrepareLumenHardwareRayTracingRadiosityLumenMaterial(const FViewInfo& View, TArray<FRHIRayTracingShader*>& OutRayGenShaders);
 	static void PrepareLumenHardwareRayTracingDirectLightingLumenMaterial(const FViewInfo& View, TArray<FRHIRayTracingShader*>& OutRayGenShaders);
 	static void PrepareStochasticShadowsLumenMaterial(const FViewInfo& View, TArray<FRHIRayTracingShader*>& OutRayGenShaders);
+	static void PrepareStochasticDirectLightingLumenMaterial(const FViewInfo& View, TArray<FRHIRayTracingShader*>& OutRayGenShaders);
 #endif // RHI_RAYTRACING
 
 
diff --git a/Engine/Source/Runtime/Renderer/Private/DistanceFieldAmbientOcclusion.cpp b/Engine/Source/Runtime/Renderer/Private/DistanceFieldAmbientOcclusion.cpp
--- a/Engine/Source/Runtime/Renderer/Private/DistanceFieldAmbientOcclusion.cpp
+++ b/Engine/Source/Runtime/Renderer/Private/DistanceFieldAmbientOcclusion.cpp
@@ -17,6 +17,7 @@
 #include "RayTracing/RaytracingOptions.h"
 #include "Lumen/Lumen.h"
 #include "StochasticShadows/StochasticShadows.h"
+#include "StochasticDirectLighting/StochasticDirectLighting.h"
 #include "ScenePrivate.h"
 #include "Substrate/Substrate.h"
 
@@ -780,10 +781,11 @@ bool FSceneRenderer::ShouldPrepareGlobalDistanceField() const
 
 	// TODO: Should check any light actually uses Stochastic Shadows
 	const bool bShouldPrepareForStochasticShadows = StochasticShadows::IsUsingGlobalSDF();
+	const bool bShouldPrepareForStochasticDirectLighting = StochasticDirectLighting::IsUsingGlobalSDF();
 
 	const bool bShouldPrepareForVisualization = ViewFamily.EngineShowFlags.VisualizeGlobalDistanceField;
 
-	return (bShouldPrepareForAO || bShouldPrepareForLumen || bShouldPrepareForStochasticShadows || bShouldPrepareForVisualization) && UseGlobalDistanceField();
+	return (bShouldPrepareForAO || bShouldPrepareForLumen || bShouldPrepareForStochasticShadows || bShouldPrepareForStochasticDirectLighting || bShouldPrepareForVisualization) && UseGlobalDistanceField();
 }
 
 void FDeferredShadingSceneRenderer::RenderDFAOAsIndirectShadowing(
diff --git a/Engine/Source/Runtime/Renderer/Private/LightFunctionAtlas.cpp b/Engine/Source/Runtime/Renderer/Private/LightFunctionAtlas.cpp
--- a/Engine/Source/Runtime/Renderer/Private/LightFunctionAtlas.cpp
+++ b/Engine/Source/Runtime/Renderer/Private/LightFunctionAtlas.cpp
@@ -14,6 +14,7 @@
 #include "PostProcess/PostProcessing.h"
 #include "PostProcess/SceneFilterRendering.h"
 #include "StochasticShadows/StochasticShadows.h"
+#include "StochasticDirectLighting/StochasticDirectLighting.h"
 #include "ShadowRendering.h"
 #include "CanvasItem.h"
 #include "CanvasTypes.h"
@@ -266,7 +267,7 @@ void FLightFunctionAtlas::BeginSceneFrame(const FViewFamilyInfo& ViewFamily, TAr
 	{
 		bVolumetricFogRequestsLF 		= bShouldRenderVolumetricFog && GVolumetricFogUsesLightFunctionAtlas > 0;
 		bDeferredlightingRequestsLF		= GDeferredUsesLightFunctionAtlas > 0;
-		bStochasticShadowsRequestsLF	= StochasticShadows::IsUsingLightFunctions();
+		bStochasticShadowsRequestsLF	= StochasticShadows::IsUsingLightFunctions() || StochasticDirectLighting::IsUsingLightFunctions();
 		bLumenRequestsLF 				= GLumenUsesLightFunctionAtlas > 0;// && IsLumenTranslucencyGIEnabled();// GLumenScene enabled ...;
 
 		bLightFunctionAtlasEnabled = bLightFunctionAtlasEnabled && 
diff --git a/Engine/Source/Runtime/Renderer/Private/LightGridInjection.cpp b/Engine/Source/Runtime/Renderer/Private/LightGridInjection.cpp
--- a/Engine/Source/Runtime/Renderer/Private/LightGridInjection.cpp
+++ b/Engine/Source/Runtime/Renderer/Private/LightGridInjection.cpp
@@ -33,6 +33,7 @@ LightGridInjection.cpp
 #include "ShaderPrintParameters.h"
 #include "RenderUtils.h"
 #include "StochasticShadows/StochasticShadows.h"
+#include "StochasticDirectLighting/StochasticDirectLighting.h"
 
 int32 GLightGridPixelSize = 64;
 FAutoConsoleVariableRef CVarLightGridPixelSize(
@@ -971,7 +972,7 @@ FComputeLightGridOutput FDeferredShadingSceneRenderer::GatherLightsAndComputeLig
 	
 	const bool bCullLightsToGrid = GLightCullingQuality 
 		&& (IsForwardShadingEnabled(ShaderPlatform) || bAnyViewUsesForwardLighting || IsRayTracingEnabled() || ShouldUseClusteredDeferredShading() ||
-			bAnyViewUsesLumen || ViewFamily.EngineShowFlags.VisualizeMeshDistanceFields || VirtualShadowMapArray.IsEnabled() || StochasticShadows::IsEnabled());
+			bAnyViewUsesLumen || ViewFamily.EngineShowFlags.VisualizeMeshDistanceFields || VirtualShadowMapArray.IsEnabled() || StochasticShadows::IsEnabled() || StochasticDirectLighting::IsEnabled());
 
 	// Store this flag if lights are injected in the grids, check with 'AreLightsInLightGrid()'
 	bAreLightsInLightGrid = bCullLightsToGrid;
diff --git a/Engine/Source/Runtime/Renderer/Private/LightRendering.cpp b/Engine/Source/Runtime/Renderer/Private/LightRendering.cpp
--- a/Engine/Source/Runtime/Renderer/Private/LightRendering.cpp
+++ b/Engine/Source/Runtime/Renderer/Private/LightRendering.cpp
@@ -31,6 +31,7 @@
 #include "MobileBasePassRendering.h"
 #include "TranslucentLighting.h"
 #include "StochasticShadows/StochasticShadows.h"
+#include "StochasticDirectLighting/StochasticDirectLighting.h"
 #include "LightFunctionAtlas.h"
 
 using namespace LightFunctionAtlas;
@@ -684,6 +685,11 @@ FLightOcclusionType GetLightOcclusionType(const FLightSceneProxy& Proxy)
 		return FLightOcclusionType::StochasticShadows;
 	}
 
+	if (StochasticDirectLighting::IsLightSupported(Proxy.GetLightType(), Proxy.CastsRaytracedShadow()))
+	{
+		return FLightOcclusionType::StochasticShadows;
+	}
+
 #if RHI_RAYTRACING
 	if (ShouldRenderRayTracingShadowsForLight(Proxy))
 	{
@@ -701,6 +707,11 @@ FLightOcclusionType GetLightOcclusionType(const FLightSceneInfoCompact& LightInf
 		return FLightOcclusionType::StochasticShadows;
 	}
 
+	if (StochasticDirectLighting::IsLightSupported(LightInfo.LightType, LightInfo.CastRaytracedShadow))
+	{
+		return FLightOcclusionType::StochasticShadows;
+	}
+
 #if RHI_RAYTRACING
 	if (ShouldRenderRayTracingShadowsForLight(LightInfo))
 	{
@@ -1186,7 +1197,8 @@ void FSceneRenderer::GatherAndSortLights(FSortedLightSetSceneInfo& OutSortedLigh
 					SortedLightInfo->SortKey.Fields.bIsNotSimpleLight = 1;
 
 					// Lights handled by Stochastic Shadows
-					const bool bHandledByStochasticShadows = StochasticShadows::IsLightSupported(LightSceneInfoCompact.LightType, LightSceneInfoCompact.CastRaytracedShadow);
+					const bool bHandledByStochasticShadows = StochasticShadows::IsLightSupported(LightSceneInfoCompact.LightType, LightSceneInfoCompact.CastRaytracedShadow)
+						|| StochasticDirectLighting::IsLightSupported(LightSceneInfoCompact.LightType, LightSceneInfoCompact.CastRaytracedShadow);
 
 					// tiled and clustered deferred lighting only supported for certain lights that don't use any additional features
 					// And also that are not directional (mostly because it doesn't make so much sense to insert them into every grid cell in the universe)
diff --git a/Engine/Source/Runtime/Renderer/Private/ScenePrivate.h b/Engine/Source/Runtime/Renderer/Private/ScenePrivate.h
--- a/Engine/Source/Runtime/Renderer/Private/ScenePrivate.h
+++ b/Engine/Source/Runtime/Renderer/Private/ScenePrivate.h
@@ -44,6 +44,7 @@
 #include "Nanite/Nanite.h"
 #include "Lumen/LumenViewState.h"
 #include "StochasticShadows/StochasticShadowsViewState.h"
+#include "StochasticDirectLighting/StochasticDirectLightingViewState.h"
 #include "VolumetricRenderTargetViewStateData.h"
 #include "GPUScene.h"
 #include "DynamicBVH.h"
@@ -1019,6 +1020,7 @@ public:
 
 	FLumenViewState Lumen;
 	FStochasticShadowsViewState StochasticShadows;
+	FStochasticDirectLightingViewState StochasticDirectLighting;
 
 	// Map from Light ID in GPU Scene to index in the View's ForwardLightData array
 	// This is stored in ViewState so we can access previous frame mapping
diff --git a/Engine/Source/Runtime/Renderer/Private/StochasticDirectLighting/StochasticDirectLighting.cpp b/Engine/Source/Runtime/Renderer/Private/StochasticDirectLighting/StochasticDirectLighting.cpp
new file mode 100644
--- /dev/null
+++ b/Engine/Source/Runtime/Renderer/Private/StochasticDirectLighting/StochasticDirectLighting.cpp
@@ -0,0 +1,1191 @@
+// Copyright Epic Games, Inc. All Rights Reserved.
+
+#include "StochasticDirectLighting.h"
+#include "StochasticDirectLightingInternal.h"
+#include "RendererPrivate.h"
+#include "PixelShaderUtils.h"
+#include "BasePassRendering.h"
+
+namespace StochasticDirectLighting
+{
+
+static TAutoConsoleVariable<int32> CVarStochasticDirectLighting(
+	TEXT("r.StochasticDirectLighting"),
+	0,
+	TEXT("Whether to enable Stochastic Direct Lighting. Experimental feature stochastically sampling analytical light in a single pass by leveraging ray tracing.\n")
+	TEXT("1 - all lights using ray tracing shadows will be stochastically sampled\n")
+	TEXT("2 - all lights will be stochastically sampled"),
+	ECVF_Scalability | ECVF_RenderThreadSafe
+);
+
+static TAutoConsoleVariable<int32> CVarStochasticDirectLightingNumSamplesPerPixel(
+	TEXT("r.StochasticDirectLighting.NumSamplesPerPixel"),
+	4,
+	TEXT("Number of samples (shadow rays) per half-res pixel.\n")
+	TEXT("1 - 0.25 trace per pixel\n")
+	TEXT("2 - 0.5 trace per pixel\n")
+	TEXT("4 - 1 trace per pixel"),
+	ECVF_Scalability | ECVF_RenderThreadSafe
+);
+
+static TAutoConsoleVariable<int32> CVarStochasticDirectLightingMaxShadingTilesPerGridCell(
+	TEXT("r.StochasticDirectLighting.MaxShadingTilesPerGridCell"),
+	32,
+	TEXT("Maximum number of shading tiles per grid cell."),
+	ECVF_Scalability | ECVF_RenderThreadSafe
+);
+
+static TAutoConsoleVariable<float> CVarStochasticDirectLightingSamplingMinWeight(
+	TEXT("r.StochasticDirectLighting.Sampling.MinWeight"),
+	0.001f,
+	TEXT("Determines minimal sample influence on final pixels. Used to skip samples which would have minimal impact to the final image even if light is fully visible."),
+	ECVF_Scalability | ECVF_RenderThreadSafe
+);
+
+static TAutoConsoleVariable<int32> CVarStochasticDirectLightingTemporal(
+	TEXT("r.StochasticDirectLighting.Temporal"),
+	1,
+	TEXT("Whether to use temporal accumulation for shadow mask."),
+	ECVF_Scalability | ECVF_RenderThreadSafe
+);
+
+static TAutoConsoleVariable<int32> CVarStochasticDirectLightingTemporalMaxFramesAccumulated(
+	TEXT("r.StochasticDirectLighting.Temporal.MaxFramesAccumulated"),
+	8,
+	TEXT("Max history length when accumulating frames. Lower values have less ghosting, but more noise."),
+	ECVF_Scalability | ECVF_RenderThreadSafe
+);
+
+static TAutoConsoleVariable<float> CVarStochasticDirectLightingTemporalNeighborhoodClampScale(
+	TEXT("r.StochasticDirectLighting.Temporal.NeighborhoodClampScale"),
+	2.0f,
+	TEXT("Scales how permissive is neighborhood clamp. Higher values cause more ghosting, but allow smoother temporal accumulation."),
+	ECVF_Scalability | ECVF_RenderThreadSafe
+);
+
+static TAutoConsoleVariable<int32> CVarStochasticDirectLightingSpatial(
+	TEXT("r.StochasticDirectLighting.Spatial"),
+	1,
+	TEXT("Whether denoiser should run spatial filter."),
+	ECVF_Scalability | ECVF_RenderThreadSafe
+);
+
+static TAutoConsoleVariable<float> CVarStochasticDirectLightingSpatialDepthWeightScale(
+	TEXT("r.StochasticDirectLighting.Spatial.DepthWeightScale"),
+	10000.0f,
+	TEXT("Scales the depth weight of the spatial filter. Smaller values allow for more sample reuse, but also introduce more bluriness between unrelated surfaces."),
+	ECVF_Scalability | ECVF_RenderThreadSafe
+);
+
+static TAutoConsoleVariable<int32> CVarStochasticDirectLightingWaveOps(
+	TEXT("r.StochasticDirectLighting.WaveOps"),
+	1,
+	TEXT("Whether to use wave ops. Useful for debugging."),
+	ECVF_Scalability | ECVF_RenderThreadSafe
+);
+
+static TAutoConsoleVariable<int32> CVarStochasticDirectLightingDebug(
+	TEXT("r.StochasticDirectLighting.Debug"),
+	0,
+	TEXT("Whether to enabled debug mode, which prints various extra debug information from shaders.")
+	TEXT("0 - Disable\n")
+	TEXT("1 - Visualize sampling\n")
+	TEXT("2 - Visualize tracing\n"),
+	ECVF_RenderThreadSafe
+);
+
+static TAutoConsoleVariable<int32> CVarStochasticDirectLightingDebugLightId(
+	TEXT("r.StochasticDirectLighting.Debug.LightId"),
+	-1,
+	TEXT("Which light to show debug info for. When set to -1, uses the currently selected light in editor."),
+	ECVF_RenderThreadSafe
+);
+
+int32 GStochasticDirectLightingReset = 0;
+FAutoConsoleVariableRef CVarStochasticDirectLightingReset(
+	TEXT("r.StochasticDirectLighting.Reset"),
+	GStochasticDirectLightingReset,
+	TEXT("Reset history for debugging."),
+	ECVF_RenderThreadSafe
+);
+
+int32 GStochasticDirectLightingResetEveryNthFrame = 0;
+	FAutoConsoleVariableRef CVarStochasticDirectLightingResetEveryNthFrame(
+	TEXT("r.StochasticDirectLighting.ResetEveryNthFrame"),
+		GStochasticDirectLightingResetEveryNthFrame,
+	TEXT("Reset history every Nth frame for debugging."),
+	ECVF_RenderThreadSafe
+);
+
+static TAutoConsoleVariable<int> CVarStochasticDirectLightingFixedStateFrameIndex(
+	TEXT("r.StochasticDirectLighting.FixedStateFrameIndex"),
+	-1,
+	TEXT("Whether to override View.StateFrameIndex for debugging."),
+	ECVF_RenderThreadSafe
+);
+
+static TAutoConsoleVariable<int> CVarStochasticDirectLightingCandidateLightMask(
+	TEXT("r.StochasticDirectLighting.CandidateLightMask"),
+	1,
+	TEXT("#sdl_todo: finish and pick one shader path."),
+	ECVF_Scalability | ECVF_RenderThreadSafe
+);
+
+static TAutoConsoleVariable<int> CVarStochasticDirectLightingTexturedRectLights(
+	TEXT("r.StochasticDirectLighting.TexturedRectLights"),
+	0,
+	TEXT("Whether to support textured rect lights."),
+	ECVF_Scalability | ECVF_RenderThreadSafe
+);
+
+static TAutoConsoleVariable<int> CVarStochasticDirectLightingLightFunctions(
+	TEXT("r.StochasticDirectLighting.LightFunctions"),
+	0,
+	TEXT("Whether to support light functions."),
+	ECVF_Scalability | ECVF_RenderThreadSafe
+);
+
+static TAutoConsoleVariable<int> CVarStochasticDirectLightingIESProfiles(
+	TEXT("r.StochasticDirectLighting.IESProfiles"),
+	1,
+	TEXT("Whether to support IES profiles on lights."),
+	ECVF_Scalability | ECVF_RenderThreadSafe
+);
+
+// must match values in StochasticDirectLighting.ush
+constexpr int32 TileSize = 8;
+constexpr int32 MaxLocalLightIndexXY = 16; // 16 * 16 = 256
+constexpr uint32 ShadingTileIndexUnshadowed = 0xFFFFF; // limited by PackShadingTile()
+constexpr uint32 ShadingAtlasSizeInTiles = 512;
+
+bool IsEnabled()
+{
+	return CVarStochasticDirectLighting.GetValueOnRenderThread() != 0;
+}
+
+bool IsUsingLightFunctions()
+{
+	return IsEnabled() && CVarStochasticDirectLightingLightFunctions.GetValueOnRenderThread() != 0;
+}
+
+bool IsLightSupported(uint8 LightType, ECastRayTracedShadow::Type CastRayTracedShadow)
+{
+	if (StochasticDirectLighting::IsEnabled() && LightType != LightType_Directional)
+	{
+		const bool bRayTracedShadows = (CastRayTracedShadow == ECastRayTracedShadow::Enabled || (ShouldRenderRayTracingShadows() && CastRayTracedShadow == ECastRayTracedShadow::UseProjectSetting));
+		return CVarStochasticDirectLighting.GetValueOnRenderThread() == 2 || bRayTracedShadows;
+	}
+
+	return false;
+}
+
+bool ShouldCompileShaders(const FGlobalShaderPermutationParameters& Parameters)
+{
+	if (IsMobilePlatform(Parameters.Platform))
+	{
+		return false;
+	}
+
+	// SM6 because it uses typed loads to accumulate lights
+	return IsFeatureLevelSupported(Parameters.Platform, ERHIFeatureLevel::SM6);
+}
+
+uint32 GetStateFrameIndex(FSceneViewState* ViewState)
+{
+	uint32 StateFrameIndex = ViewState ? ViewState->GetFrameIndex() : 0;
+
+	if (CVarStochasticDirectLightingFixedStateFrameIndex.GetValueOnRenderThread() >= 0)
+	{
+		StateFrameIndex = CVarStochasticDirectLightingFixedStateFrameIndex.GetValueOnRenderThread();
+	}
+		
+	return StateFrameIndex;
+}
+
+FIntPoint GetNumSamplesPerPixel2d()
+{
+	const uint32 NumSamplesPerPixel1d = FMath::RoundUpToPowerOfTwo(FMath::Clamp(CVarStochasticDirectLightingNumSamplesPerPixel.GetValueOnRenderThread(), 1, 4));
+	return NumSamplesPerPixel1d == 4 ? FIntPoint(2, 2) : (NumSamplesPerPixel1d == 2 ? FIntPoint(2, 1) : FIntPoint(1, 1));
+}
+
+int32 GetDebugMode()
+{
+	return CVarStochasticDirectLightingDebug.GetValueOnRenderThread();
+}
+
+bool UseWaveOps(EShaderPlatform ShaderPlatform)
+{
+	return CVarStochasticDirectLightingWaveOps.GetValueOnRenderThread() != 0 
+		&& GRHISupportsWaveOperations
+		&& RHISupportsWaveOperations(ShaderPlatform);
+}
+
+void ModifyCompilationEnvironment(EShaderPlatform Platform, FShaderCompilerEnvironment& OutEnvironment)
+{
+	FForwardLightingParameters::ModifyCompilationEnvironment(Platform, OutEnvironment);
+}
+
+// Keep in sync with TILE_TYPE_* in shaders
+enum class ETileType : uint8
+{
+	SimpleShading = 0,
+	ComplexShading = 1,
+	SHADING_MAX = 2,
+
+	Empty = 2,
+	MAX = 3
+};
+
+class FTileClassificationCS : public FGlobalShader
+{
+	DECLARE_GLOBAL_SHADER(FTileClassificationCS)
+	SHADER_USE_PARAMETER_STRUCT(FTileClassificationCS, FGlobalShader)
+
+	BEGIN_SHADER_PARAMETER_STRUCT(FParameters, )
+		SHADER_PARAMETER_STRUCT_INCLUDE(FStochasticDirectLightingParameters, StochasticDirectLightingParameters)
+		SHADER_PARAMETER_RDG_BUFFER_UAV(RWStructuredBuffer<uint>, RWTileAllocator)
+		SHADER_PARAMETER_RDG_BUFFER_UAV(RWStructuredBuffer<uint>, RWTileData)
+	END_SHADER_PARAMETER_STRUCT()
+
+	class FDownsampledClassification : SHADER_PERMUTATION_BOOL("DOWNSAMPLED_CLASSIFICATION");
+	using FPermutationDomain = TShaderPermutationDomain<FDownsampledClassification>;
+
+	static int32 GetGroupSize()
+	{	
+		return 8;
+	}
+
+	static bool ShouldCompilePermutation(const FGlobalShaderPermutationParameters& Parameters)
+	{
+		return StochasticDirectLighting::ShouldCompileShaders(Parameters);
+	}
+
+	static void ModifyCompilationEnvironment(const FGlobalShaderPermutationParameters& Parameters, FShaderCompilerEnvironment& OutEnvironment)
+	{
+		FGlobalShader::ModifyCompilationEnvironment(Parameters, OutEnvironment);
+		OutEnvironment.SetDefine(TEXT("THREADGROUP_SIZE"), GetGroupSize());
+	}
+};
+
+IMPLEMENT_GLOBAL_SHADER(FTileClassificationCS, "/Engine/Private/StochasticDirectLighting/StochasticDirectLighting.usf", "TileClassificationCS", SF_Compute);
+
+class FInitTileIndirectArgsCS : public FGlobalShader
+{
+	DECLARE_GLOBAL_SHADER(FInitTileIndirectArgsCS)
+	SHADER_USE_PARAMETER_STRUCT(FInitTileIndirectArgsCS, FGlobalShader);
+
+	BEGIN_SHADER_PARAMETER_STRUCT(FParameters, )
+		SHADER_PARAMETER_STRUCT_INCLUDE(FStochasticDirectLightingParameters, StochasticDirectLightingParameters)
+		SHADER_PARAMETER_RDG_BUFFER_UAV(RWBuffer<uint>, RWTileIndirectArgs)
+		SHADER_PARAMETER_RDG_BUFFER_UAV(RWBuffer<uint>, RWDownsampledTileIndirectArgs)
+		SHADER_PARAMETER_RDG_BUFFER_SRV(StructuredBuffer<uint>, TileAllocator)
+		SHADER_PARAMETER_RDG_BUFFER_SRV(StructuredBuffer<uint>, DownsampledTileAllocator)
+	END_SHADER_PARAMETER_STRUCT()
+
+	static bool ShouldCompilePermutation(const FGlobalShaderPermutationParameters& Parameters)
+	{
+		return StochasticDirectLighting::ShouldCompileShaders(Parameters);
+	}
+
+	static void ModifyCompilationEnvironment(const FGlobalShaderPermutationParameters& Parameters, FShaderCompilerEnvironment& OutEnvironment)
+	{
+		FGlobalShader::ModifyCompilationEnvironment(Parameters, OutEnvironment);
+		OutEnvironment.SetDefine(TEXT("THREADGROUP_SIZE"), GetGroupSize());
+	}
+
+	static int32 GetGroupSize()
+	{
+		return 64;
+	}
+};
+
+IMPLEMENT_GLOBAL_SHADER(FInitTileIndirectArgsCS, "/Engine/Private/StochasticDirectLighting/StochasticDirectLighting.usf", "InitTileIndirectArgsCS", SF_Compute);
+
+class FInitCompositeIndirectArgsCS : public FGlobalShader
+{
+	DECLARE_GLOBAL_SHADER(FInitCompositeIndirectArgsCS)
+	SHADER_USE_PARAMETER_STRUCT(FInitCompositeIndirectArgsCS, FGlobalShader);
+
+	BEGIN_SHADER_PARAMETER_STRUCT(FParameters, )
+		SHADER_PARAMETER_STRUCT_INCLUDE(FStochasticDirectLightingParameters, StochasticDirectLightingParameters)
+		SHADER_PARAMETER_RDG_BUFFER_UAV(RWBuffer<uint>, RWIndirectArgs)
+		SHADER_PARAMETER_RDG_BUFFER_SRV(StructuredBuffer<uint>, CompositeTileAllocator)
+	END_SHADER_PARAMETER_STRUCT()
+
+	static bool ShouldCompilePermutation(const FGlobalShaderPermutationParameters& Parameters)
+	{
+		return StochasticDirectLighting::ShouldCompileShaders(Parameters);
+	}
+
+	static void ModifyCompilationEnvironment(const FGlobalShaderPermutationParameters& Parameters, FShaderCompilerEnvironment& OutEnvironment)
+	{
+		FGlobalShader::ModifyCompilationEnvironment(Parameters, OutEnvironment);
+		OutEnvironment.SetDefine(TEXT("THREADGROUP_SIZE"), GetGroupSize());
+	}
+
+	static int32 GetGroupSize()
+	{
+		return 64;
+	}
+};
+
+IMPLEMENT_GLOBAL_SHADER(FInitCompositeIndirectArgsCS, "/Engine/Private/StochasticDirectLighting/StochasticDirectLighting.usf", "InitCompositeIndirectArgsCS", SF_Compute);
+
+class FGenerateSamplesCS : public FGlobalShader
+{
+	DECLARE_GLOBAL_SHADER(FGenerateSamplesCS)
+	SHADER_USE_PARAMETER_STRUCT(FGenerateSamplesCS, FGlobalShader)
+
+	BEGIN_SHADER_PARAMETER_STRUCT(FParameters, )
+		RDG_BUFFER_ACCESS(IndirectArgs, ERHIAccess::IndirectArgs)
+		SHADER_PARAMETER_STRUCT_INCLUDE(FStochasticDirectLightingParameters, StochasticDirectLightingParameters)
+		SHADER_PARAMETER_RDG_TEXTURE_UAV(RWTexture2D<float>, RWDownsampledSceneDepth)
+		SHADER_PARAMETER_RDG_TEXTURE_UAV(RWTexture2D<UNORM float3>, RWDownsampledSceneWorldNormal)
+		SHADER_PARAMETER_RDG_BUFFER_UAV(RWStructuredBuffer<uint>, RWCompositeTileAllocator)
+		SHADER_PARAMETER_RDG_BUFFER_UAV(RWStructuredBuffer<uint2>, RWCompositeTileData)
+		SHADER_PARAMETER_RDG_TEXTURE_UAV(RWTexture2D<uint>, RWLightSamples)
+		SHADER_PARAMETER_RDG_BUFFER_SRV(StructuredBuffer<uint>, DownsampledTileAllocator)
+		SHADER_PARAMETER_RDG_BUFFER_SRV(StructuredBuffer<uint>, DownsampledTileData)
+		SHADER_PARAMETER(FVector4f, HistoryUVMinMax)
+		SHADER_PARAMETER(FVector4f, HistoryScreenPositionScaleBias)
+	END_SHADER_PARAMETER_STRUCT()
+
+	class FTileType : SHADER_PERMUTATION_INT("TILE_TYPE", (int32)StochasticDirectLighting::ETileType::SHADING_MAX);
+	class FIESProfile : SHADER_PERMUTATION_BOOL("USE_IES_PROFILE");
+	class FLightFunctionAtlas : SHADER_PERMUTATION_BOOL("USE_LIGHT_FUNCTION_ATLAS");
+	class FTexturedRectLights : SHADER_PERMUTATION_BOOL("USE_SOURCE_TEXTURE");
+	class FNumSamplesPerPixel1d : SHADER_PERMUTATION_SPARSE_INT("NUM_SAMPLES_PER_PIXEL_1D", 1, 2, 4);
+	class FCandidateLightMask : SHADER_PERMUTATION_BOOL("CANDIDATE_LIGHT_MASK");
+	class FDebugMode : SHADER_PERMUTATION_BOOL("DEBUG_MODE");
+	using FPermutationDomain = TShaderPermutationDomain<FTileType, FIESProfile, FLightFunctionAtlas, FTexturedRectLights, FNumSamplesPerPixel1d, FCandidateLightMask, FDebugMode>;
+
+	static int32 GetGroupSize()
+	{	
+		return 8;
+	}
+
+	static bool ShouldCompilePermutation(const FGlobalShaderPermutationParameters& Parameters)
+	{
+		return StochasticDirectLighting::ShouldCompileShaders(Parameters);
+	}
+
+	static void ModifyCompilationEnvironment(const FGlobalShaderPermutationParameters& Parameters, FShaderCompilerEnvironment& OutEnvironment)
+	{
+		FGlobalShader::ModifyCompilationEnvironment(Parameters, OutEnvironment);
+		StochasticDirectLighting::ModifyCompilationEnvironment(Parameters.Platform, OutEnvironment);
+		OutEnvironment.SetDefine(TEXT("THREADGROUP_SIZE"), GetGroupSize());
+	}
+};
+
+IMPLEMENT_GLOBAL_SHADER(FGenerateSamplesCS, "/Engine/Private/StochasticDirectLighting/StochasticDirectLightingSampling.usf", "GenerateSamplesCS", SF_Compute);
+
+class FClearLightSamplesCS : public FGlobalShader
+{
+	DECLARE_GLOBAL_SHADER(FClearLightSamplesCS)
+	SHADER_USE_PARAMETER_STRUCT(FClearLightSamplesCS, FGlobalShader)
+
+	BEGIN_SHADER_PARAMETER_STRUCT(FParameters, )
+		RDG_BUFFER_ACCESS(IndirectArgs, ERHIAccess::IndirectArgs)
+		SHADER_PARAMETER_STRUCT_INCLUDE(FStochasticDirectLightingParameters, StochasticDirectLightingParameters)
+		SHADER_PARAMETER_RDG_TEXTURE_UAV(RWTexture2D<float>, RWDownsampledSceneDepth)
+		SHADER_PARAMETER_RDG_TEXTURE_UAV(RWTexture2D<UNORM float3>, RWDownsampledSceneWorldNormal)
+		SHADER_PARAMETER_RDG_TEXTURE_UAV(RWTexture2D<uint>, RWLightSamples)
+		SHADER_PARAMETER_RDG_BUFFER_SRV(StructuredBuffer<uint>, DownsampledTileAllocator)
+		SHADER_PARAMETER_RDG_BUFFER_SRV(StructuredBuffer<uint>, DownsampledTileData)
+	END_SHADER_PARAMETER_STRUCT()
+
+	class FDebugMode : SHADER_PERMUTATION_BOOL("DEBUG_MODE");
+	using FPermutationDomain = TShaderPermutationDomain<FDebugMode>;
+
+	static int32 GetGroupSize()
+	{
+		return 8;
+	}
+
+	static bool ShouldCompilePermutation(const FGlobalShaderPermutationParameters& Parameters)
+	{
+		return StochasticDirectLighting::ShouldCompileShaders(Parameters);
+	}
+
+	static void ModifyCompilationEnvironment(const FGlobalShaderPermutationParameters& Parameters, FShaderCompilerEnvironment& OutEnvironment)
+	{
+		FGlobalShader::ModifyCompilationEnvironment(Parameters, OutEnvironment);
+		StochasticDirectLighting::ModifyCompilationEnvironment(Parameters.Platform, OutEnvironment);
+		OutEnvironment.SetDefine(TEXT("THREADGROUP_SIZE"), GetGroupSize());
+	}
+};
+
+IMPLEMENT_GLOBAL_SHADER(FClearLightSamplesCS, "/Engine/Private/StochasticDirectLighting/StochasticDirectLightingSampling.usf", "ClearLightSamplesCS", SF_Compute);
+
+class FInitCompositeUpsampleWeightsCS : public FGlobalShader
+{
+	DECLARE_GLOBAL_SHADER(FInitCompositeUpsampleWeightsCS)
+	SHADER_USE_PARAMETER_STRUCT(FInitCompositeUpsampleWeightsCS, FGlobalShader)
+
+	BEGIN_SHADER_PARAMETER_STRUCT(FParameters, )
+		SHADER_PARAMETER_STRUCT_INCLUDE(FStochasticDirectLightingParameters, StochasticDirectLightingParameters)
+		SHADER_PARAMETER_RDG_TEXTURE_UAV(RWTexture2D<float4>, RWCompositeUpsampleWeights)
+	END_SHADER_PARAMETER_STRUCT()
+
+	static int32 GetGroupSize()
+	{
+		return 8;
+	}
+
+	static bool ShouldCompilePermutation(const FGlobalShaderPermutationParameters& Parameters)
+	{
+		return StochasticDirectLighting::ShouldCompileShaders(Parameters);
+	}
+
+	static void ModifyCompilationEnvironment(const FGlobalShaderPermutationParameters& Parameters, FShaderCompilerEnvironment& OutEnvironment)
+	{
+		FGlobalShader::ModifyCompilationEnvironment(Parameters, OutEnvironment);
+		OutEnvironment.SetDefine(TEXT("THREADGROUP_SIZE"), GetGroupSize());
+	}
+};
+
+IMPLEMENT_GLOBAL_SHADER(FInitCompositeUpsampleWeightsCS, "/Engine/Private/StochasticDirectLighting/StochasticDirectLighting.usf", "InitCompositeUpsampleWeightsCS", SF_Compute);
+
+class FCompositeLightSamplesCS : public FGlobalShader
+{
+	DECLARE_GLOBAL_SHADER(FCompositeLightSamplesCS)
+	SHADER_USE_PARAMETER_STRUCT(FCompositeLightSamplesCS, FGlobalShader)
+
+	BEGIN_SHADER_PARAMETER_STRUCT(FParameters, )
+		SHADER_PARAMETER_STRUCT_INCLUDE(FStochasticDirectLightingParameters, StochasticDirectLightingParameters)
+		RDG_BUFFER_ACCESS(IndirectArgs, ERHIAccess::IndirectArgs)
+		SHADER_PARAMETER_RDG_BUFFER_UAV(RWStructuredBuffer<uint>, RWShadingTileAllocator)
+		SHADER_PARAMETER_RDG_TEXTURE_UAV(RWTexture2D<uint>, RWShadingTileGridAllocator)
+		SHADER_PARAMETER_RDG_BUFFER_UAV(RWStructuredBuffer<uint>, RWShadingTileGrid)
+		SHADER_PARAMETER_RDG_TEXTURE_UAV(RWTexture2D<float>, RWShadingTileAtlas)
+		SHADER_PARAMETER_RDG_BUFFER_SRV(StructuredBuffer<uint2>, CompositeTileData)
+		SHADER_PARAMETER_RDG_TEXTURE(Texture2D<flaot4>, CompositeUpsampleWeights)
+		SHADER_PARAMETER_RDG_TEXTURE(Texture2D<uint>, LightSamples)
+	END_SHADER_PARAMETER_STRUCT()
+
+	static int32 GetGroupSize()
+	{
+		return 8;
+	}
+
+	class FDebugMode : SHADER_PERMUTATION_BOOL("DEBUG_MODE");
+	using FPermutationDomain = TShaderPermutationDomain<FDebugMode>;
+
+	static bool ShouldCompilePermutation(const FGlobalShaderPermutationParameters& Parameters)
+	{
+		return StochasticDirectLighting::ShouldCompileShaders(Parameters);
+	}
+
+	static void ModifyCompilationEnvironment(const FGlobalShaderPermutationParameters& Parameters, FShaderCompilerEnvironment& OutEnvironment)
+	{
+		FGlobalShader::ModifyCompilationEnvironment(Parameters, OutEnvironment);
+		StochasticDirectLighting::ModifyCompilationEnvironment(Parameters.Platform, OutEnvironment);
+		OutEnvironment.SetDefine(TEXT("THREADGROUP_SIZE"), GetGroupSize());
+	}
+};
+
+IMPLEMENT_GLOBAL_SHADER(FCompositeLightSamplesCS, "/Engine/Private/StochasticDirectLighting/StochasticDirectLightingComposite.usf", "CompositeLightSamplesCS", SF_Compute);
+
+class FShadeLightSamplesCS : public FGlobalShader
+{
+	DECLARE_GLOBAL_SHADER(FShadeLightSamplesCS)
+	SHADER_USE_PARAMETER_STRUCT(FShadeLightSamplesCS, FGlobalShader)
+
+	BEGIN_SHADER_PARAMETER_STRUCT(FParameters, )
+		RDG_BUFFER_ACCESS(IndirectArgs, ERHIAccess::IndirectArgs)
+		SHADER_PARAMETER_STRUCT_INCLUDE(FStochasticDirectLightingParameters, StochasticDirectLightingParameters)
+		SHADER_PARAMETER_RDG_TEXTURE_UAV(RWTexture2D<float3>, RWResolvedDiffuseLighting)
+		SHADER_PARAMETER_RDG_TEXTURE_UAV(RWTexture2D<float3>, RWResolvedSpecularLighting)
+		SHADER_PARAMETER_RDG_BUFFER_SRV(StructuredBuffer<uint>, TileAllocator)
+		SHADER_PARAMETER_RDG_BUFFER_SRV(StructuredBuffer<uint>, TileData)
+		SHADER_PARAMETER_RDG_BUFFER_SRV(StructuredBuffer<uint>, ShadingTileAllocator)
+		SHADER_PARAMETER_RDG_BUFFER_SRV(StructuredBuffer<uint>, CompositeTileAllocator)
+		SHADER_PARAMETER_RDG_TEXTURE(Texture2D<uint>, ShadingTileGridAllocator)
+		SHADER_PARAMETER_RDG_BUFFER_SRV(StructuredBuffer<uint>, ShadingTileGrid)
+		SHADER_PARAMETER_RDG_BUFFER_SRV(StructuredBuffer<uint>, ShadingTiles)
+		SHADER_PARAMETER_RDG_TEXTURE(Texture2D<float>, ShadingTileAtlas)
+	END_SHADER_PARAMETER_STRUCT()
+
+	static int32 GetGroupSize()
+	{
+		return 8;
+	}
+
+	class FTileType : SHADER_PERMUTATION_INT("TILE_TYPE", (int32)StochasticDirectLighting::ETileType::SHADING_MAX);
+	class FIESProfile : SHADER_PERMUTATION_BOOL("USE_IES_PROFILE");
+	class FLightFunctionAtlas : SHADER_PERMUTATION_BOOL("USE_LIGHT_FUNCTION_ATLAS");
+	class FTexturedRectLights : SHADER_PERMUTATION_BOOL("USE_SOURCE_TEXTURE");
+	class FDebugMode : SHADER_PERMUTATION_BOOL("DEBUG_MODE");
+	using FPermutationDomain = TShaderPermutationDomain<FTileType, FIESProfile, FLightFunctionAtlas, FTexturedRectLights, FDebugMode>;
+
+	static bool ShouldCompilePermutation(const FGlobalShaderPermutationParameters& Parameters)
+	{
+		return StochasticDirectLighting::ShouldCompileShaders(Parameters);
+	}
+
+	static void ModifyCompilationEnvironment(const FGlobalShaderPermutationParameters& Parameters, FShaderCompilerEnvironment& OutEnvironment)
+	{
+		FGlobalShader::ModifyCompilationEnvironment(Parameters, OutEnvironment);
+		StochasticDirectLighting::ModifyCompilationEnvironment(Parameters.Platform, OutEnvironment);
+		OutEnvironment.SetDefine(TEXT("THREADGROUP_SIZE"), GetGroupSize());
+	}
+};
+
+IMPLEMENT_GLOBAL_SHADER(FShadeLightSamplesCS, "/Engine/Private/StochasticDirectLighting/StochasticDirectLightingShading.usf", "ShadeLightSamplesCS", SF_Compute);
+
+class FTemporalAccumulationCS : public FGlobalShader
+{
+	DECLARE_GLOBAL_SHADER(FTemporalAccumulationCS)
+	SHADER_USE_PARAMETER_STRUCT(FTemporalAccumulationCS, FGlobalShader)
+
+	BEGIN_SHADER_PARAMETER_STRUCT(FParameters, )
+		SHADER_PARAMETER_STRUCT_INCLUDE(FStochasticDirectLightingParameters, StochasticDirectLightingParameters)
+		SHADER_PARAMETER_RDG_TEXTURE(Texture2D<float4>, ResolvedDiffuseLighting)
+		SHADER_PARAMETER_RDG_TEXTURE(Texture2D<float4>, ResolvedSpecularLighting)
+		SHADER_PARAMETER_RDG_TEXTURE(Texture2D<float4>, DiffuseLightingAndSecondMomentHistoryTexture)
+		SHADER_PARAMETER_RDG_TEXTURE(Texture2D<float4>, SpecularLightingAndSecondMomentHistoryTexture)
+		SHADER_PARAMETER_RDG_TEXTURE(Texture2D<UNORM float>, NumFramesAccumulatedHistoryTexture)
+		SHADER_PARAMETER_RDG_TEXTURE(Texture2D<float>, StochasticDirectLightingDepthHistory)
+		SHADER_PARAMETER(FVector4f, HistoryUVMinMax)
+		SHADER_PARAMETER(FVector4f, HistoryScreenPositionScaleBias)
+		SHADER_PARAMETER(float, PrevSceneColorPreExposureCorrection)
+		SHADER_PARAMETER_RDG_TEXTURE_UAV(RWTexture2D<float4>, RWDiffuseLightingAndSecondMoment)
+		SHADER_PARAMETER_RDG_TEXTURE_UAV(RWTexture2D<float4>, RWSpecularLightingAndSecondMoment)
+		SHADER_PARAMETER_RDG_TEXTURE_UAV(RWTexture2D<UNORM float>, RWNumFramesAccumulated)
+		SHADER_PARAMETER_RDG_TEXTURE_UAV(RWTexture2D<float>, RWSceneDepth)
+		SHADER_PARAMETER_RDG_TEXTURE_UAV(RWTexture2D<float4>, RWSceneColor)
+	END_SHADER_PARAMETER_STRUCT()
+
+	class FValidHistory : SHADER_PERMUTATION_BOOL("VALID_HISTORY");
+	class FDebugMode : SHADER_PERMUTATION_BOOL("DEBUG_MODE");
+	using FPermutationDomain = TShaderPermutationDomain<FValidHistory, FDebugMode>;
+
+	static int32 GetGroupSize()
+	{
+		return 8;
+	}
+
+	static bool ShouldCompilePermutation(const FGlobalShaderPermutationParameters& Parameters)
+	{
+		return StochasticDirectLighting::ShouldCompileShaders(Parameters);
+	}
+
+	static void ModifyCompilationEnvironment(const FGlobalShaderPermutationParameters& Parameters, FShaderCompilerEnvironment& OutEnvironment)
+	{
+		FGlobalShader::ModifyCompilationEnvironment(Parameters, OutEnvironment);
+		OutEnvironment.SetDefine(TEXT("THREADGROUP_SIZE"), GetGroupSize());
+	}
+};
+
+IMPLEMENT_GLOBAL_SHADER(FTemporalAccumulationCS, "/Engine/Private/StochasticDirectLighting/StochasticDirectLightingTemporal.usf", "TemporalAccumulationCS", SF_Compute);
+
+class FSpatialFilterCS : public FGlobalShader
+{
+	DECLARE_GLOBAL_SHADER(FSpatialFilterCS)
+	SHADER_USE_PARAMETER_STRUCT(FSpatialFilterCS, FGlobalShader)
+
+	BEGIN_SHADER_PARAMETER_STRUCT(FParameters, )
+		SHADER_PARAMETER_STRUCT_INCLUDE(FStochasticDirectLightingParameters, StochasticDirectLightingParameters)
+		SHADER_PARAMETER_RDG_TEXTURE_UAV(RWTexture2D<float4>, RWSceneColor)
+		SHADER_PARAMETER_RDG_TEXTURE(Texture2D<float3>, DiffuseLightingAndSecondMomentTexture)
+		SHADER_PARAMETER_RDG_TEXTURE(Texture2D<float3>, SpecularLightingAndSecondMomentTexture)
+		SHADER_PARAMETER(float, SpatialFilterDepthWeightScale)
+	END_SHADER_PARAMETER_STRUCT()
+
+	class FSpatialFilter : SHADER_PERMUTATION_BOOL("SPATIAL_FILTER");
+	class FDebugMode : SHADER_PERMUTATION_BOOL("DEBUG_MODE");
+	using FPermutationDomain = TShaderPermutationDomain<FSpatialFilter, FDebugMode>;
+
+	static int32 GetGroupSize()
+	{
+		return 8;
+	}
+
+	static bool ShouldCompilePermutation(const FGlobalShaderPermutationParameters& Parameters)
+	{
+		return StochasticDirectLighting::ShouldCompileShaders(Parameters);
+	}
+
+	static void ModifyCompilationEnvironment(const FGlobalShaderPermutationParameters& Parameters, FShaderCompilerEnvironment& OutEnvironment)
+	{
+		FGlobalShader::ModifyCompilationEnvironment(Parameters, OutEnvironment);
+		OutEnvironment.SetDefine(TEXT("THREADGROUP_SIZE"), GetGroupSize());
+	}
+};
+
+IMPLEMENT_GLOBAL_SHADER(FSpatialFilterCS, "/Engine/Private/StochasticDirectLighting/StochasticDirectLightingSpatial.usf", "SpatialFilterCS", SF_Compute);
+
+} // namespace StochasticDirectLighting
+
+/**
+ * Single pass batched light rendering using ray tracing (distance field or triangle) for shadowing.
+ */
+void FDeferredShadingSceneRenderer::RenderStochasticDirectLighting(FRDGBuilder& GraphBuilder, const FSceneTextures& SceneTextures)
+{
+	using namespace StochasticDirectLighting;
+	if (!StochasticDirectLighting::IsEnabled())
+	{
+		return;
+	}
+
+	check(AreLightsInLightGrid());
+
+	RDG_EVENT_SCOPE(GraphBuilder, "StochasticDirectLighting");
+
+	const uint32 ViewIndex = 0;
+	const FViewInfo& View = Views[ViewIndex];
+	FBlueNoise BlueNoise = GetBlueNoiseGlobalParameters();
+	TUniformBufferRef<FBlueNoise> BlueNoiseUniformBuffer = CreateUniformBufferImmediate(BlueNoise, EUniformBufferUsage::UniformBuffer_SingleDraw);
+
+	const bool bDebug = StochasticDirectLighting::GetDebugMode() != 0;
+	const bool bWaveOps = StochasticDirectLighting::UseWaveOps(View.GetShaderPlatform())
+		&& GRHIMinimumWaveSize <= 32
+		&& GRHIMaximumWaveSize >= 32;
+
+	// History reset for debugging purposes
+	bool bResetHistory = false;
+
+	if (GStochasticDirectLightingResetEveryNthFrame > 0 && (ViewFamily.FrameNumber % (uint32)GStochasticDirectLightingResetEveryNthFrame) == 0)
+	{
+		bResetHistory = true;
+	}
+
+	if (GStochasticDirectLightingReset != 0)
+	{
+		GStochasticDirectLightingReset = 0;
+		bResetHistory = true;
+	}
+
+	const FIntPoint NumSamplesPerPixel2d = StochasticDirectLighting::GetNumSamplesPerPixel2d();
+
+	const uint32 DownsampleFactor = 2;
+	const FIntPoint DownsampledViewSize = FIntPoint::DivideAndRoundUp(View.ViewRect.Size(), DownsampleFactor);
+	const FIntPoint SampleViewSize = DownsampledViewSize * NumSamplesPerPixel2d;
+	const FIntPoint DownsampledBufferSize = FIntPoint::DivideAndRoundUp(SceneTextures.Config.Extent, DownsampleFactor);
+	const FIntPoint SampleBufferSize = DownsampledBufferSize * NumSamplesPerPixel2d;
+	const FIntPoint DonwnsampledSampleBufferSize = DownsampledBufferSize * NumSamplesPerPixel2d;
+
+	// #sdl_todo: make atlas size based on the resolution
+	const FIntPoint ShadingTileGridSize = FIntPoint::DivideAndRoundUp(SceneTextures.Config.Extent, StochasticDirectLighting::TileSize);
+	const FIntPoint ShadingTileAtlasSize = StochasticDirectLighting::TileSize * StochasticDirectLighting::ShadingAtlasSizeInTiles;
+	const int32 MaxShadingTiles = StochasticDirectLighting::ShadingAtlasSizeInTiles * StochasticDirectLighting::ShadingAtlasSizeInTiles;
+	check(MaxShadingTiles == (ShadingTileAtlasSize.X * ShadingTileAtlasSize.Y) / (StochasticDirectLighting::TileSize * StochasticDirectLighting::TileSize));
+	check(MaxShadingTiles < StochasticDirectLighting::ShadingTileIndexUnshadowed);
+	const int32 MaxCompositeTiles = MaxShadingTiles;
+
+	const int32 MaxShadingTilesPerGridCell = FMath::Max(CVarStochasticDirectLightingMaxShadingTilesPerGridCell.GetValueOnRenderThread(), 0);
+
+	FRDGTextureRef DownsampledSceneDepth = GraphBuilder.CreateTexture(
+		FRDGTextureDesc::Create2D(DownsampledBufferSize, PF_R32_FLOAT, FClearValueBinding::Black, TexCreate_ShaderResource | TexCreate_UAV),
+		TEXT("StochasticDirectLighting.DownsampledSceneDepth"));
+
+	FRDGTextureRef DownsampledSceneWorldNormal = GraphBuilder.CreateTexture(
+		FRDGTextureDesc::Create2D(DownsampledBufferSize, PF_A2B10G10R10, FClearValueBinding::Black, TexCreate_ShaderResource | TexCreate_UAV),
+		TEXT("StochasticDirectLighting.DownsampledSceneWorldNormal"));
+
+	FRDGTextureRef LightSamples = GraphBuilder.CreateTexture(
+		FRDGTextureDesc::Create2D(DonwnsampledSampleBufferSize, PF_R32_UINT, FClearValueBinding::Black, TexCreate_ShaderResource | TexCreate_UAV),
+		TEXT("StochasticDirectLighting.LightSamples"));
+
+	FRDGTextureRef LightSampleRayDistance = GraphBuilder.CreateTexture(
+		FRDGTextureDesc::Create2D(DonwnsampledSampleBufferSize, PF_R16F, FClearValueBinding::Black, TexCreate_ShaderResource | TexCreate_UAV),
+		TEXT("StochasticDirectLighting.LightSampleRayDistance"));
+	
+	bool bTemporal = CVarStochasticDirectLightingTemporal.GetValueOnRenderThread() != 0;
+	FVector4f HistoryScreenPositionScaleBias = FVector4f(0.0f, 0.0f, 0.0f, 0.0f);
+	FVector4f HistoryUVMinMax = FVector4f(0.0f, 0.0f, 0.0f, 0.0f);
+	FRDGTextureRef DiffuseLightingAndSecondMomentHistory = nullptr;
+	FRDGTextureRef SpecularLightingAndSecondMomentHistory = nullptr;
+	FRDGTextureRef SceneDepthHistory = nullptr;
+	FRDGTextureRef NumFramesAccumulatedHistory = nullptr;
+
+	if (View.ViewState)
+	{
+		const FStochasticDirectLightingViewState& LightingViewState = View.ViewState->StochasticDirectLighting;
+
+		if (!View.bCameraCut && !bResetHistory && bTemporal)
+		{
+			HistoryScreenPositionScaleBias = LightingViewState.HistoryScreenPositionScaleBias;
+			HistoryUVMinMax = LightingViewState.HistoryUVMinMax;
+
+			if (LightingViewState.DiffuseLightingAndSecondMomentHistory
+				&& LightingViewState.SpecularLightingAndSecondMomentHistory
+				&& LightingViewState.SceneDepthHistory
+				&& LightingViewState.NumFramesAccumulatedHistory
+				&& LightingViewState.DiffuseLightingAndSecondMomentHistory->GetDesc().Extent == View.GetSceneTexturesConfig().Extent
+				&& LightingViewState.SpecularLightingAndSecondMomentHistory->GetDesc().Extent == View.GetSceneTexturesConfig().Extent
+				&& LightingViewState.SceneDepthHistory->GetDesc().Extent == SceneTextures.Depth.Resolve->Desc.Extent)
+			{
+				DiffuseLightingAndSecondMomentHistory = GraphBuilder.RegisterExternalTexture(LightingViewState.DiffuseLightingAndSecondMomentHistory);
+				SpecularLightingAndSecondMomentHistory = GraphBuilder.RegisterExternalTexture(LightingViewState.SpecularLightingAndSecondMomentHistory);
+				SceneDepthHistory = GraphBuilder.RegisterExternalTexture(LightingViewState.SceneDepthHistory);
+				NumFramesAccumulatedHistory = GraphBuilder.RegisterExternalTexture(LightingViewState.NumFramesAccumulatedHistory);
+			}
+		}
+	}
+
+	// Setup the light function atlas
+	const bool bUseLightFunctionAtlas = LightFunctionAtlas::IsEnabled(View, LightFunctionAtlas::ELightFunctionAtlasSystem::StochasticShadows);
+
+	const FIntPoint ViewSizeInTiles = FIntPoint::DivideAndRoundUp(View.ViewRect.Size(), StochasticDirectLighting::TileSize);
+	const int32 TileDataStride = ViewSizeInTiles.X * ViewSizeInTiles.Y;
+
+	const FIntPoint DownsampledViewSizeInTiles = FIntPoint::DivideAndRoundUp(DownsampledViewSize, StochasticDirectLighting::TileSize);
+	const int32 DownsampledTileDataStride = DownsampledViewSizeInTiles.X * DownsampledViewSizeInTiles.Y;
+
+	FRDGTextureRef DownsampledTileMask = GraphBuilder.CreateTexture(
+		FRDGTextureDesc::Create2D(FMath::DivideAndRoundUp<FIntPoint>(DownsampledBufferSize, StochasticDirectLighting::TileSize), PF_R8_UINT, FClearValueBinding::Black, TexCreate_ShaderResource | TexCreate_UAV),
+		TEXT("StochasticDirectLighting.DownsampledTileMask"));
+
+	FStochasticDirectLightingParameters StochasticDirectLightingParameters;
+	{
+		StochasticDirectLightingParameters.ViewUniformBuffer = View.ViewUniformBuffer;
+		StochasticDirectLightingParameters.Scene = View.GetSceneUniforms().GetBuffer(GraphBuilder);
+		StochasticDirectLightingParameters.SceneTextures = GetSceneTextureParameters(GraphBuilder, SceneTextures.UniformBuffer);
+		StochasticDirectLightingParameters.SceneTexturesStruct = SceneTextures.UniformBuffer;
+		StochasticDirectLightingParameters.Substrate = Substrate::BindSubstrateGlobalUniformParameters(View);
+		StochasticDirectLightingParameters.ForwardLightData = View.ForwardLightingResources.ForwardLightUniformBuffer;
+		StochasticDirectLightingParameters.LightFunctionAtlas = LightFunctionAtlas::BindGlobalParameters(GraphBuilder, View);
+		StochasticDirectLightingParameters.BlueNoise = BlueNoiseUniformBuffer;
+		StochasticDirectLightingParameters.PreIntegratedGF = GSystemTextures.PreintegratedGF->GetRHI();
+		StochasticDirectLightingParameters.PreIntegratedGFSampler = TStaticSamplerState<SF_Bilinear, AM_Clamp, AM_Clamp, AM_Clamp>::GetRHI();
+		StochasticDirectLightingParameters.DownsampledViewSize = DownsampledViewSize;
+		StochasticDirectLightingParameters.SampleViewSize = SampleViewSize;
+		StochasticDirectLightingParameters.NumSamplesPerPixel = NumSamplesPerPixel2d;
+		StochasticDirectLightingParameters.NumSamplesPerPixelDivideShift.X = FMath::FloorLog2(NumSamplesPerPixel2d.X);
+		StochasticDirectLightingParameters.NumSamplesPerPixelDivideShift.Y = FMath::FloorLog2(NumSamplesPerPixel2d.Y);
+		StochasticDirectLightingParameters.StochasticDirectLightingStateFrameIndex = StochasticDirectLighting::GetStateFrameIndex(View.ViewState);
+		StochasticDirectLightingParameters.DownsampledTileMask = DownsampledTileMask;
+		StochasticDirectLightingParameters.DownsampledSceneDepth = DownsampledSceneDepth;
+		StochasticDirectLightingParameters.DownsampledSceneWorldNormal = DownsampledSceneWorldNormal;
+		StochasticDirectLightingParameters.MaxCompositeTiles = MaxCompositeTiles;
+		StochasticDirectLightingParameters.MaxShadingTiles = MaxShadingTiles;
+		StochasticDirectLightingParameters.MaxShadingTilesPerGridCell = MaxShadingTilesPerGridCell;
+		StochasticDirectLightingParameters.ShadingTileGridSize = ShadingTileGridSize;
+		StochasticDirectLightingParameters.DownsampledBufferInvSize = FVector2f(1.0f) / DownsampledBufferSize;
+		StochasticDirectLightingParameters.SamplingMinWeight = FMath::Max(CVarStochasticDirectLightingSamplingMinWeight.GetValueOnRenderThread(), 0.0f);
+		StochasticDirectLightingParameters.TileDataStride = TileDataStride;
+		StochasticDirectLightingParameters.DownsampledTileDataStride = DownsampledTileDataStride;
+		StochasticDirectLightingParameters.TemporalMaxFramesAccumulated = FMath::Max(CVarStochasticDirectLightingTemporalMaxFramesAccumulated.GetValueOnRenderThread(), 0.0f);
+		StochasticDirectLightingParameters.TemporalNeighborhoodClampScale = CVarStochasticDirectLightingTemporalNeighborhoodClampScale.GetValueOnRenderThread();
+		StochasticDirectLightingParameters.DebugMode = StochasticDirectLighting::GetDebugMode();
+		StochasticDirectLightingParameters.DebugLightId = INDEX_NONE;
+
+		if (bDebug)
+		{
+			ShaderPrint::SetEnabled(true);
+			ShaderPrint::RequestSpaceForLines(1024);
+			ShaderPrint::SetParameters(GraphBuilder, View.ShaderPrintData, StochasticDirectLightingParameters.ShaderPrintUniformBuffer);
+
+			StochasticDirectLightingParameters.DebugLightId = CVarStochasticDirectLightingDebugLightId.GetValueOnRenderThread();
+
+			if (StochasticDirectLightingParameters.DebugLightId < 0)
+			{
+				for (auto LightIt = Scene->Lights.CreateConstIterator(); LightIt; ++LightIt)
+				{
+					const FLightSceneInfoCompact& LightSceneInfoCompact = *LightIt;
+					const FLightSceneInfo* const LightSceneInfo = LightSceneInfoCompact.LightSceneInfo;
+
+					if (LightSceneInfo->Proxy->IsSelected())
+					{
+						StochasticDirectLightingParameters.DebugLightId = LightSceneInfo->Id;
+						break;
+					}
+				}
+			}
+		}
+	}
+
+	FRDGBufferRef TileAllocator = GraphBuilder.CreateBuffer(FRDGBufferDesc::CreateStructuredDesc(sizeof(uint32), (int32)StochasticDirectLighting::ETileType::MAX), TEXT("StochasticDirectLighting.TileAllocator"));
+	FRDGBufferRef TileData = GraphBuilder.CreateBuffer(FRDGBufferDesc::CreateStructuredDesc(sizeof(uint32), TileDataStride * (int32)StochasticDirectLighting::ETileType::MAX), TEXT("StochasticDirectLighting.TileData"));
+	AddClearUAVPass(GraphBuilder, GraphBuilder.CreateUAV(TileAllocator), 0);
+
+	FRDGBufferRef DownsampledTileAllocator = GraphBuilder.CreateBuffer(FRDGBufferDesc::CreateStructuredDesc(sizeof(uint32), (int32)StochasticDirectLighting::ETileType::MAX), TEXT("StochasticDirectLighting.DownsampledTileAllocator"));
+	FRDGBufferRef DownsampledTileData = GraphBuilder.CreateBuffer(FRDGBufferDesc::CreateStructuredDesc(sizeof(uint32), DownsampledTileDataStride * (int32)StochasticDirectLighting::ETileType::MAX), TEXT("StochasticDirectLighting.DownsampledTileData"));
+	AddClearUAVPass(GraphBuilder, GraphBuilder.CreateUAV(DownsampledTileAllocator), 0);
+
+	// #sdl_todo: merge classification passes or reuse downsampled one to create full res tiles
+	// Run tile classification to generate tiles for the subsequent passes
+	{
+		{
+			FTileClassificationCS::FParameters* PassParameters = GraphBuilder.AllocParameters<FTileClassificationCS::FParameters>();
+			PassParameters->StochasticDirectLightingParameters = StochasticDirectLightingParameters;
+			PassParameters->RWTileAllocator = GraphBuilder.CreateUAV(TileAllocator);
+			PassParameters->RWTileData = GraphBuilder.CreateUAV(TileData);
+
+			FTileClassificationCS::FPermutationDomain PermutationVector;
+			PermutationVector.Set<FTileClassificationCS::FDownsampledClassification>(false);
+			auto ComputeShader = View.ShaderMap->GetShader<FTileClassificationCS>(PermutationVector);
+
+			const FIntVector GroupCount = FComputeShaderUtils::GetGroupCount(View.ViewRect.Size(), FTileClassificationCS::GetGroupSize());
+
+			FComputeShaderUtils::AddPass(
+				GraphBuilder,
+				RDG_EVENT_NAME("TileClassification %dx%d", View.ViewRect.Size().X, View.ViewRect.Size().Y),
+				ComputeShader,
+				PassParameters,
+				GroupCount);
+		}
+
+		{
+			FTileClassificationCS::FParameters* PassParameters = GraphBuilder.AllocParameters<FTileClassificationCS::FParameters>();
+			PassParameters->StochasticDirectLightingParameters = StochasticDirectLightingParameters;
+			PassParameters->RWTileAllocator = GraphBuilder.CreateUAV(DownsampledTileAllocator);
+			PassParameters->RWTileData = GraphBuilder.CreateUAV(DownsampledTileData);
+
+			FTileClassificationCS::FPermutationDomain PermutationVector;
+			PermutationVector.Set<FTileClassificationCS::FDownsampledClassification>(true);
+			auto ComputeShader = View.ShaderMap->GetShader<FTileClassificationCS>(PermutationVector);
+
+			const FIntVector GroupCount = FComputeShaderUtils::GetGroupCount(View.ViewRect.Size(), FTileClassificationCS::GetGroupSize());
+
+			FComputeShaderUtils::AddPass(
+				GraphBuilder,
+				RDG_EVENT_NAME("DownsampledTileClassification %dx%d", DownsampledViewSize.X, DownsampledViewSize.Y),
+				ComputeShader,
+				PassParameters,
+				GroupCount);
+		}
+	}
+
+	FRDGBufferRef TileIndirectArgs = GraphBuilder.CreateBuffer(FRDGBufferDesc::CreateIndirectDesc<FRHIDispatchIndirectParameters>((int32)StochasticDirectLighting::ETileType::MAX), TEXT("StochasticDirectLighting.TileIndirectArgs"));
+	FRDGBufferRef DownsampledTileIndirectArgs = GraphBuilder.CreateBuffer(FRDGBufferDesc::CreateIndirectDesc<FRHIDispatchIndirectParameters>((int32)StochasticDirectLighting::ETileType::MAX), TEXT("StochasticDirectLighting.DownsampledTileIndirectArgs"));
+
+	// Setup indirect args for classified tiles
+	{
+		FInitTileIndirectArgsCS::FParameters* PassParameters = GraphBuilder.AllocParameters<FInitTileIndirectArgsCS::FParameters>();
+		PassParameters->StochasticDirectLightingParameters = StochasticDirectLightingParameters;
+		PassParameters->RWTileIndirectArgs = GraphBuilder.CreateUAV(TileIndirectArgs);
+		PassParameters->RWDownsampledTileIndirectArgs = GraphBuilder.CreateUAV(DownsampledTileIndirectArgs);
+		PassParameters->TileAllocator = GraphBuilder.CreateSRV(TileAllocator);
+		PassParameters->DownsampledTileAllocator = GraphBuilder.CreateSRV(DownsampledTileAllocator);
+
+		auto ComputeShader = View.ShaderMap->GetShader<FInitTileIndirectArgsCS>();
+
+		FComputeShaderUtils::AddPass(
+			GraphBuilder,
+			RDG_EVENT_NAME("InitTileIndirectArgs"),
+			ComputeShader,
+			PassParameters,
+			FIntVector(1, 1, 1));
+	}
+
+	FRDGBufferRef CompositeTileAllocator = GraphBuilder.CreateBuffer(FRDGBufferDesc::CreateStructuredDesc(sizeof(uint32), 1), TEXT("StochasticDirectLighting.CompositeTileAllocator"));
+	FRDGBufferRef CompositeTileData = GraphBuilder.CreateBuffer(FRDGBufferDesc::CreateStructuredDesc(2 * sizeof(uint32), MaxCompositeTiles), TEXT("StochasticDirectLighting.CompositeTileData"));
+	AddClearUAVPass(GraphBuilder, GraphBuilder.CreateUAV(CompositeTileAllocator), 0);
+	
+	// Generate new candidate light samples
+	{
+		FRDGTextureUAVRef DownsampledSceneDepthUAV = GraphBuilder.CreateUAV(DownsampledSceneDepth, ERDGUnorderedAccessViewFlags::SkipBarrier);
+		FRDGTextureUAVRef DownsampledSceneWorldNormalUAV = GraphBuilder.CreateUAV(DownsampledSceneWorldNormal, ERDGUnorderedAccessViewFlags::SkipBarrier);
+		FRDGBufferUAVRef CompositeTileAllocatorUAV = GraphBuilder.CreateUAV(CompositeTileAllocator, ERDGUnorderedAccessViewFlags::SkipBarrier);
+		FRDGBufferUAVRef CompositeTileDataUAV = GraphBuilder.CreateUAV(CompositeTileData, ERDGUnorderedAccessViewFlags::SkipBarrier);
+		FRDGTextureUAVRef LightSamplesUAV = GraphBuilder.CreateUAV(LightSamples, ERDGUnorderedAccessViewFlags::SkipBarrier);
+
+		// Clear tiles which don't contain any lights or geometry
+		{
+			FClearLightSamplesCS::FParameters* PassParameters = GraphBuilder.AllocParameters<FClearLightSamplesCS::FParameters>();
+			PassParameters->IndirectArgs = DownsampledTileIndirectArgs;
+			PassParameters->StochasticDirectLightingParameters = StochasticDirectLightingParameters;
+			PassParameters->RWDownsampledSceneDepth = DownsampledSceneDepthUAV;
+			PassParameters->RWDownsampledSceneWorldNormal = DownsampledSceneWorldNormalUAV;
+			PassParameters->RWLightSamples = LightSamplesUAV;
+			PassParameters->DownsampledTileAllocator = GraphBuilder.CreateSRV(DownsampledTileAllocator);
+			PassParameters->DownsampledTileData = GraphBuilder.CreateSRV(DownsampledTileData);
+
+			FClearLightSamplesCS::FPermutationDomain PermutationVector;
+			PermutationVector.Set<FClearLightSamplesCS::FDebugMode>(bDebug);
+			auto ComputeShader = View.ShaderMap->GetShader<FClearLightSamplesCS>(PermutationVector);
+
+			FComputeShaderUtils::AddPass(
+				GraphBuilder,
+				RDG_EVENT_NAME("ClearLightSamples"),
+				ComputeShader,
+				PassParameters,
+				DownsampledTileIndirectArgs,
+				(int32)StochasticDirectLighting::ETileType::Empty * sizeof(FRHIDispatchIndirectParameters));
+		}
+
+		for (int32 TileType = 0; TileType < (int32)StochasticDirectLighting::ETileType::SHADING_MAX; ++TileType)
+		{
+			FGenerateSamplesCS::FParameters* PassParameters = GraphBuilder.AllocParameters<FGenerateSamplesCS::FParameters>();
+			PassParameters->IndirectArgs = DownsampledTileIndirectArgs;
+			PassParameters->StochasticDirectLightingParameters = StochasticDirectLightingParameters;
+			PassParameters->RWDownsampledSceneDepth = DownsampledSceneDepthUAV;
+			PassParameters->RWDownsampledSceneWorldNormal = DownsampledSceneWorldNormalUAV;
+			PassParameters->RWLightSamples = LightSamplesUAV;
+			PassParameters->RWCompositeTileAllocator = CompositeTileAllocatorUAV;
+			PassParameters->RWCompositeTileData = CompositeTileDataUAV;
+			PassParameters->DownsampledTileAllocator = GraphBuilder.CreateSRV(DownsampledTileAllocator);
+			PassParameters->DownsampledTileData = GraphBuilder.CreateSRV(DownsampledTileData);
+			PassParameters->HistoryScreenPositionScaleBias = HistoryScreenPositionScaleBias;
+			PassParameters->HistoryUVMinMax = HistoryUVMinMax;
+
+			FGenerateSamplesCS::FPermutationDomain PermutationVector;
+			PermutationVector.Set<FGenerateSamplesCS::FTileType>(TileType);
+			PermutationVector.Set<FGenerateSamplesCS::FIESProfile>(CVarStochasticDirectLightingIESProfiles.GetValueOnRenderThread() != 0);
+			PermutationVector.Set<FGenerateSamplesCS::FLightFunctionAtlas>(bUseLightFunctionAtlas);
+			PermutationVector.Set<FGenerateSamplesCS::FTexturedRectLights>(CVarStochasticDirectLightingTexturedRectLights.GetValueOnRenderThread() != 0);
+			PermutationVector.Set<FGenerateSamplesCS::FNumSamplesPerPixel1d>(NumSamplesPerPixel2d.X * NumSamplesPerPixel2d.Y);
+			PermutationVector.Set<FGenerateSamplesCS::FDebugMode>(bDebug);
+			PermutationVector.Set<FGenerateSamplesCS::FCandidateLightMask>(CVarStochasticDirectLightingCandidateLightMask.GetValueOnRenderThread() != 0);
+			auto ComputeShader = View.ShaderMap->GetShader<FGenerateSamplesCS>(PermutationVector);
+
+			FComputeShaderUtils::AddPass(
+				GraphBuilder,
+				RDG_EVENT_NAME("GenerateSamples SamplesPerPixel:%dx%d TileType:%d", NumSamplesPerPixel2d.X, NumSamplesPerPixel2d.Y, TileType),
+				ComputeShader,
+				PassParameters,
+				DownsampledTileIndirectArgs,
+				TileType * sizeof(FRHIDispatchIndirectParameters));
+		}
+	}
+
+	FRDGBufferRef CompositeTileIndirectArgs = GraphBuilder.CreateBuffer(FRDGBufferDesc::CreateIndirectDesc<FRHIDispatchIndirectParameters>(1), TEXT("StochasticDirectLighting.CompositeTileIndirectArgs"));
+
+	// Setup indirect args for shadow mask tile updates
+	{
+		FInitCompositeIndirectArgsCS::FParameters* PassParameters = GraphBuilder.AllocParameters<FInitCompositeIndirectArgsCS::FParameters>();
+		PassParameters->StochasticDirectLightingParameters = StochasticDirectLightingParameters;
+		PassParameters->RWIndirectArgs = GraphBuilder.CreateUAV(CompositeTileIndirectArgs);
+		PassParameters->CompositeTileAllocator = GraphBuilder.CreateSRV(CompositeTileAllocator);
+
+		auto ComputeShader = View.ShaderMap->GetShader<FInitCompositeIndirectArgsCS>();
+
+		FComputeShaderUtils::AddPass(
+			GraphBuilder,
+			RDG_EVENT_NAME("InitCompositeIndirectArgs"),
+			ComputeShader,
+			PassParameters,
+			FIntVector(1, 1, 1));
+	}
+
+	StochasticDirectLighting::RayTraceLightSamples(
+		View,
+		GraphBuilder, 
+		SceneTextures,
+		SampleBufferSize,
+		LightSamples,
+		LightSampleRayDistance,
+		StochasticDirectLightingParameters
+	);
+
+	FRDGTextureRef CompositeUpsampleWeights = GraphBuilder.CreateTexture(
+		FRDGTextureDesc::Create2D(SceneTextures.Config.Extent, PF_FloatRGBA, FClearValueBinding::Black, TexCreate_ShaderResource | TexCreate_UAV),
+		TEXT("StochasticDirectLighting.CompositeUpsampleWeights"));
+
+	// Init composite upsample weights
+	{
+		FInitCompositeUpsampleWeightsCS::FParameters* PassParameters = GraphBuilder.AllocParameters<FInitCompositeUpsampleWeightsCS::FParameters>();
+		PassParameters->StochasticDirectLightingParameters = StochasticDirectLightingParameters;
+		PassParameters->RWCompositeUpsampleWeights = GraphBuilder.CreateUAV(CompositeUpsampleWeights);
+
+		auto ComputeShader = View.ShaderMap->GetShader<FInitCompositeUpsampleWeightsCS>();
+
+		const FIntVector GroupCount = FComputeShaderUtils::GetGroupCount(View.ViewRect.Size(), FInitCompositeUpsampleWeightsCS::GetGroupSize());
+
+		FComputeShaderUtils::AddPass(
+			GraphBuilder,
+			RDG_EVENT_NAME("InitCompositeUpsampleWeights"),
+			ComputeShader,
+			PassParameters,
+			GroupCount);
+	}
+
+	FRDGBufferRef ShadingTileAllocator = GraphBuilder.CreateBuffer(
+		FRDGBufferDesc::CreateStructuredDesc(sizeof(uint32), 1),
+		TEXT("StochasticDirectLightingParameters.ShadingTileAllocator"));
+
+	FRDGTextureRef ShadingTileGridAllocator = GraphBuilder.CreateTexture(
+		FRDGTextureDesc::Create2D(ShadingTileGridSize, PF_R32_UINT, FClearValueBinding::Black, TexCreate_ShaderResource | TexCreate_UAV),
+		TEXT("StochasticDirectLighting.ShadingTileGridAllocator"));
+
+	FRDGBufferRef ShadingTileGrid = GraphBuilder.CreateBuffer(
+		FRDGBufferDesc::CreateStructuredDesc(sizeof(uint32), ShadingTileGridSize.X * ShadingTileGridSize.Y * MaxShadingTilesPerGridCell),
+		TEXT("StochasticDirectLightingParameters.ShadingTileGrid"));
+
+	FRDGTextureRef ShadingTileAtlas = GraphBuilder.CreateTexture(
+		FRDGTextureDesc::Create2D(ShadingTileAtlasSize, PF_R16F, FClearValueBinding::Black, TexCreate_ShaderResource | TexCreate_UAV),
+		TEXT("StochasticDirectLighting	.ShadingTileAtlas"));
+
+	AddClearUAVPass(GraphBuilder, GraphBuilder.CreateUAV(ShadingTileAllocator), 0u);
+	AddClearUAVPass(GraphBuilder, GraphBuilder.CreateUAV(ShadingTileGridAllocator), 0u);
+
+	// Composite shadow masks traces
+	{
+		FCompositeLightSamplesCS::FParameters* PassParameters = GraphBuilder.AllocParameters<FCompositeLightSamplesCS::FParameters>();
+		PassParameters->StochasticDirectLightingParameters = StochasticDirectLightingParameters;
+		PassParameters->IndirectArgs = CompositeTileIndirectArgs;
+		PassParameters->RWShadingTileAllocator = GraphBuilder.CreateUAV(ShadingTileAllocator);
+		PassParameters->RWShadingTileGridAllocator = GraphBuilder.CreateUAV(ShadingTileGridAllocator);
+		PassParameters->RWShadingTileGrid = GraphBuilder.CreateUAV(ShadingTileGrid);
+		PassParameters->RWShadingTileAtlas = GraphBuilder.CreateUAV(ShadingTileAtlas);
+		PassParameters->CompositeTileData = GraphBuilder.CreateSRV(CompositeTileData);
+		PassParameters->CompositeUpsampleWeights = CompositeUpsampleWeights;
+		PassParameters->LightSamples = LightSamples;
+
+		FCompositeLightSamplesCS::FPermutationDomain PermutationVector;
+		PermutationVector.Set<FCompositeLightSamplesCS::FDebugMode>(bDebug);
+		auto ComputeShader = View.ShaderMap->GetShader<FCompositeLightSamplesCS>(PermutationVector);
+
+		FComputeShaderUtils::AddPass(
+			GraphBuilder,
+			RDG_EVENT_NAME("CompositeLightSamples"),
+			ComputeShader,
+			PassParameters,
+			CompositeTileIndirectArgs,
+			0);
+	}
+
+	FRDGTextureRef ResolvedDiffuseLighting = GraphBuilder.CreateTexture(
+		FRDGTextureDesc::Create2D(View.GetSceneTexturesConfig().Extent, PF_FloatRGB, FClearValueBinding::Black, TexCreate_ShaderResource | TexCreate_UAV),
+		TEXT("StochasticDirectLighting.ResolvedDiffuseLighting"));
+
+	FRDGTextureRef ResolvedSpecularLighting = GraphBuilder.CreateTexture(
+		FRDGTextureDesc::Create2D(View.GetSceneTexturesConfig().Extent, PF_FloatRGB, FClearValueBinding::Black, TexCreate_ShaderResource | TexCreate_UAV),
+		TEXT("StochasticDirectLighting.ResolvedSpecularLighting"));
+
+	// Shade light samples
+	{
+		FRDGTextureUAVRef ResolvedDiffuseLightingUAV = GraphBuilder.CreateUAV(ResolvedDiffuseLighting, ERDGUnorderedAccessViewFlags::SkipBarrier);
+		FRDGTextureUAVRef ResolvedSpecularLightingUAV = GraphBuilder.CreateUAV(ResolvedSpecularLighting, ERDGUnorderedAccessViewFlags::SkipBarrier);
+
+		for (int32 TileType = 0; TileType < (int32)StochasticDirectLighting::ETileType::SHADING_MAX; ++TileType)
+		{
+			FShadeLightSamplesCS::FParameters* PassParameters = GraphBuilder.AllocParameters<FShadeLightSamplesCS::FParameters>();
+			PassParameters->RWResolvedDiffuseLighting = ResolvedDiffuseLightingUAV;
+			PassParameters->RWResolvedSpecularLighting = ResolvedSpecularLightingUAV;
+			PassParameters->IndirectArgs = TileIndirectArgs;
+			PassParameters->StochasticDirectLightingParameters = StochasticDirectLightingParameters;
+			PassParameters->CompositeTileAllocator = GraphBuilder.CreateSRV(CompositeTileAllocator);
+			PassParameters->TileAllocator = GraphBuilder.CreateSRV(TileAllocator);
+			PassParameters->TileData = GraphBuilder.CreateSRV(TileData);
+			PassParameters->ShadingTileAllocator = GraphBuilder.CreateSRV(ShadingTileAllocator);
+			PassParameters->ShadingTileGridAllocator = ShadingTileGridAllocator;
+			PassParameters->ShadingTileGrid = GraphBuilder.CreateSRV(ShadingTileGrid);
+			PassParameters->ShadingTileAtlas = ShadingTileAtlas;
+
+			FShadeLightSamplesCS::FPermutationDomain PermutationVector;
+			PermutationVector.Set<FShadeLightSamplesCS::FTileType>(TileType);
+			PermutationVector.Set<FShadeLightSamplesCS::FIESProfile>(CVarStochasticDirectLightingIESProfiles.GetValueOnRenderThread() != 0);
+			PermutationVector.Set<FShadeLightSamplesCS::FLightFunctionAtlas>(bUseLightFunctionAtlas);
+			PermutationVector.Set<FShadeLightSamplesCS::FTexturedRectLights>(CVarStochasticDirectLightingTexturedRectLights.GetValueOnRenderThread() != 0);
+			PermutationVector.Set<FShadeLightSamplesCS::FDebugMode>(bDebug);
+			auto ComputeShader = View.ShaderMap->GetShader<FShadeLightSamplesCS>(PermutationVector);
+
+			FComputeShaderUtils::AddPass(
+				GraphBuilder,
+				RDG_EVENT_NAME("ShadeLightSamples TileType:%d", TileType),
+				ComputeShader,
+				PassParameters,
+				TileIndirectArgs,
+				TileType * sizeof(FRHIDispatchIndirectParameters));
+		}
+	}
+
+	// Demodulated lighting components with second luminance moments stored in alpha channel for temporal variance tracking
+	// This will be passed to the next frame
+	FRDGTextureRef DiffuseLightingAndSecondMoment = GraphBuilder.CreateTexture(
+		FRDGTextureDesc::Create2D(View.GetSceneTexturesConfig().Extent, PF_FloatRGBA, FClearValueBinding::Black, TexCreate_ShaderResource | TexCreate_UAV),
+		TEXT("StochasticDirectLighting.DiffuseLightingAndSecondMoment"));
+
+	FRDGTextureRef SpecularLightingAndSecondMoment = GraphBuilder.CreateTexture(
+		FRDGTextureDesc::Create2D(View.GetSceneTexturesConfig().Extent, PF_FloatRGBA, FClearValueBinding::Black, TexCreate_ShaderResource | TexCreate_UAV),
+		TEXT("StochasticDirectLighting.SpecularLightingAndSecondMoment"));
+
+	FRDGTextureRef SceneDepthCopy = GraphBuilder.CreateTexture(
+		FRDGTextureDesc::Create2D(SceneTextures.Depth.Resolve->Desc.Extent, PF_R32_FLOAT, FClearValueBinding::Black, TexCreate_ShaderResource | TexCreate_UAV),
+		TEXT("StochasticDirectLighting.DepthHistory"));
+
+	FRDGTextureRef NumFramesAccumulated = GraphBuilder.CreateTexture(
+		FRDGTextureDesc::Create2D(View.GetSceneTexturesConfig().Extent, PF_G8, FClearValueBinding::Black, TexCreate_ShaderResource | TexCreate_UAV),
+		TEXT("StochasticDirectLighting.NumFramesAccumulated"));
+
+	// Temporal accumulation
+	{
+		FTemporalAccumulationCS::FParameters* PassParameters = GraphBuilder.AllocParameters<FTemporalAccumulationCS::FParameters>();
+		PassParameters->StochasticDirectLightingParameters = StochasticDirectLightingParameters;
+		PassParameters->ResolvedDiffuseLighting = ResolvedDiffuseLighting;
+		PassParameters->ResolvedSpecularLighting = ResolvedSpecularLighting;
+		PassParameters->DiffuseLightingAndSecondMomentHistoryTexture = DiffuseLightingAndSecondMomentHistory;
+		PassParameters->SpecularLightingAndSecondMomentHistoryTexture = SpecularLightingAndSecondMomentHistory;
+		PassParameters->NumFramesAccumulatedHistoryTexture = NumFramesAccumulatedHistory;
+		PassParameters->StochasticDirectLightingDepthHistory = SceneDepthHistory;
+		PassParameters->PrevSceneColorPreExposureCorrection = View.PreExposure / View.PrevViewInfo.SceneColorPreExposure;
+		PassParameters->HistoryScreenPositionScaleBias = HistoryScreenPositionScaleBias;
+		PassParameters->HistoryUVMinMax = HistoryUVMinMax;
+		PassParameters->RWDiffuseLightingAndSecondMoment = GraphBuilder.CreateUAV(DiffuseLightingAndSecondMoment);
+		PassParameters->RWSpecularLightingAndSecondMoment = GraphBuilder.CreateUAV(SpecularLightingAndSecondMoment);
+		PassParameters->RWNumFramesAccumulated = GraphBuilder.CreateUAV(NumFramesAccumulated);
+		PassParameters->RWSceneDepth = GraphBuilder.CreateUAV(SceneDepthCopy);
+		PassParameters->RWSceneColor = GraphBuilder.CreateUAV(SceneTextures.Color.Target);
+
+		FTemporalAccumulationCS::FPermutationDomain PermutationVector;
+		PermutationVector.Set<FTemporalAccumulationCS::FValidHistory>(DiffuseLightingAndSecondMomentHistory != nullptr && bTemporal);
+		PermutationVector.Set<FTemporalAccumulationCS::FDebugMode>(bDebug);
+		auto ComputeShader = View.ShaderMap->GetShader<FTemporalAccumulationCS>(PermutationVector);
+
+		const FIntVector GroupCount = FComputeShaderUtils::GetGroupCount(View.ViewRect.Size(), FTemporalAccumulationCS::GetGroupSize());
+
+		FComputeShaderUtils::AddPass(
+			GraphBuilder,
+			RDG_EVENT_NAME("TemporalAccumulation"),
+			ComputeShader,
+			PassParameters,
+			GroupCount);
+	}
+
+	// Spatial filter
+	{
+		FSpatialFilterCS::FParameters* PassParameters = GraphBuilder.AllocParameters<FSpatialFilterCS::FParameters>();
+		PassParameters->StochasticDirectLightingParameters = StochasticDirectLightingParameters;
+		PassParameters->RWSceneColor = GraphBuilder.CreateUAV(SceneTextures.Color.Target);
+		PassParameters->DiffuseLightingAndSecondMomentTexture = DiffuseLightingAndSecondMoment;
+		PassParameters->SpecularLightingAndSecondMomentTexture = SpecularLightingAndSecondMoment;
+		PassParameters->SpatialFilterDepthWeightScale = CVarStochasticDirectLightingSpatialDepthWeightScale.GetValueOnRenderThread();
+
+		FSpatialFilterCS::FPermutationDomain PermutationVector;
+		PermutationVector.Set<FSpatialFilterCS::FSpatialFilter>(CVarStochasticDirectLightingSpatial.GetValueOnRenderThread() != 0);
+		PermutationVector.Set<FSpatialFilterCS::FDebugMode>(bDebug);
+		auto ComputeShader = View.ShaderMap->GetShader<FSpatialFilterCS>(PermutationVector);
+
+		const FIntVector GroupCount = FComputeShaderUtils::GetGroupCount(View.ViewRect.Size(), FSpatialFilterCS::GetGroupSize());
+
+		FComputeShaderUtils::AddPass(
+			GraphBuilder,
+			RDG_EVENT_NAME("Spatial"),
+			ComputeShader,
+			PassParameters,
+			GroupCount);
+	}
+
+	if (View.ViewState && !View.bStatePrevViewInfoIsReadOnly)
+	{
+		FStochasticDirectLightingViewState& LightingViewState = View.ViewState->StochasticDirectLighting;
+
+		LightingViewState.HistoryScreenPositionScaleBias = View.GetScreenPositionScaleBias(View.GetSceneTexturesConfig().Extent, View.ViewRect);
+
+		// Pull in the max UV to exclude the region which will read outside the viewport due to bilinear filtering
+		const FVector2D InvBufferSize(1.0f / SceneTextures.Config.Extent.X, 1.0f / SceneTextures.Config.Extent.Y);
+		LightingViewState.HistoryUVMinMax = FVector4f(
+			(View.ViewRect.Min.X + 0.5f) * InvBufferSize.X,
+			(View.ViewRect.Min.Y + 0.5f) * InvBufferSize.Y,
+			(View.ViewRect.Max.X - 1.0f) * InvBufferSize.X,
+			(View.ViewRect.Max.Y - 1.0f) * InvBufferSize.Y);
+
+		if (DiffuseLightingAndSecondMoment && SpecularLightingAndSecondMoment && SceneDepthCopy && NumFramesAccumulated && bTemporal)
+		{
+			GraphBuilder.QueueTextureExtraction(DiffuseLightingAndSecondMoment, &LightingViewState.DiffuseLightingAndSecondMomentHistory);
+			GraphBuilder.QueueTextureExtraction(SpecularLightingAndSecondMoment, &LightingViewState.SpecularLightingAndSecondMomentHistory);
+			GraphBuilder.QueueTextureExtraction(SceneDepthCopy, &LightingViewState.SceneDepthHistory);
+			GraphBuilder.QueueTextureExtraction(NumFramesAccumulated, &LightingViewState.NumFramesAccumulatedHistory);
+		}
+		else
+		{
+			LightingViewState.DiffuseLightingAndSecondMomentHistory = nullptr;
+			LightingViewState.SpecularLightingAndSecondMomentHistory = nullptr;
+			LightingViewState.SceneDepthHistory = nullptr;
+			LightingViewState.NumFramesAccumulatedHistory = nullptr;
+		}
+	}
+}
\ No newline at end of file
diff --git a/Engine/Source/Runtime/Renderer/Private/StochasticDirectLighting/StochasticDirectLighting.h b/Engine/Source/Runtime/Renderer/Private/StochasticDirectLighting/StochasticDirectLighting.h
new file mode 100644
--- /dev/null
+++ b/Engine/Source/Runtime/Renderer/Private/StochasticDirectLighting/StochasticDirectLighting.h
@@ -0,0 +1,24 @@
+// Copyright Epic Games, Inc. All Rights Reserved.
+
+#pragma once
+
+#include "CoreMinimal.h"
+
+namespace ECastRayTracedShadow
+{
+	enum Type : int;
+};
+
+// Public StochasticDirectLighting interface
+namespace StochasticDirectLighting
+{
+	bool IsEnabled();
+
+	bool IsUsingClosestHZB();
+	bool IsUsingGlobalSDF();
+	bool IsUsingLightFunctions();
+
+	bool IsLightSupported(uint8 LightType, ECastRayTracedShadow::Type CastRayTracedShadow);
+	bool UseHardwareRayTracing();
+	bool UseInlineHardwareRayTracing();
+};
\ No newline at end of file
diff --git a/Engine/Source/Runtime/Renderer/Private/StochasticDirectLighting/StochasticDirectLightingInternal.h b/Engine/Source/Runtime/Renderer/Private/StochasticDirectLighting/StochasticDirectLightingInternal.h
new file mode 100644
--- /dev/null
+++ b/Engine/Source/Runtime/Renderer/Private/StochasticDirectLighting/StochasticDirectLightingInternal.h
@@ -0,0 +1,61 @@
+// Copyright Epic Games, Inc. All Rights Reserved.
+
+#pragma once
+
+#include "CoreMinimal.h"
+#include "RendererPrivate.h"
+#include "BlueNoise.h"
+
+BEGIN_SHADER_PARAMETER_STRUCT(FStochasticDirectLightingParameters, )
+	SHADER_PARAMETER_STRUCT_REF(FViewUniformShaderParameters, ViewUniformBuffer)
+	SHADER_PARAMETER_STRUCT_INCLUDE(ShaderPrint::FShaderParameters, ShaderPrintUniformBuffer)
+	SHADER_PARAMETER_STRUCT_INCLUDE(FSceneTextureParameters, SceneTextures)
+	SHADER_PARAMETER_RDG_UNIFORM_BUFFER(FSceneUniformParameters, Scene)
+	SHADER_PARAMETER_RDG_UNIFORM_BUFFER(FSceneTextureUniformParameters, SceneTexturesStruct)
+	SHADER_PARAMETER_RDG_UNIFORM_BUFFER(FSubstrateGlobalUniformParameters, Substrate)
+	SHADER_PARAMETER_RDG_UNIFORM_BUFFER(FForwardLightData, ForwardLightData)
+	SHADER_PARAMETER_RDG_UNIFORM_BUFFER(LightFunctionAtlas::FLightFunctionAtlasGlobalParameters, LightFunctionAtlas)
+	SHADER_PARAMETER_STRUCT_REF(FBlueNoise, BlueNoise)
+	SHADER_PARAMETER_TEXTURE(Texture2D, PreIntegratedGF)
+	SHADER_PARAMETER_SAMPLER(SamplerState, PreIntegratedGFSampler)
+	SHADER_PARAMETER(FIntPoint, SampleViewSize)
+	SHADER_PARAMETER(FIntPoint, DownsampledViewSize)
+	SHADER_PARAMETER(FIntPoint, NumSamplesPerPixel)
+	SHADER_PARAMETER(FIntPoint, NumSamplesPerPixelDivideShift)
+	SHADER_PARAMETER(uint32, StochasticDirectLightingStateFrameIndex)
+	SHADER_PARAMETER(uint32, MaxCompositeTiles)
+	SHADER_PARAMETER(uint32, MaxShadingTiles)
+	SHADER_PARAMETER(uint32, MaxShadingTilesPerGridCell)
+	SHADER_PARAMETER(FIntPoint, ShadingTileGridSize)
+	SHADER_PARAMETER(FVector2f, DownsampledBufferInvSize)
+	SHADER_PARAMETER(float, SamplingMinWeight)
+	SHADER_PARAMETER(int32, TileDataStride)
+	SHADER_PARAMETER(int32, DownsampledTileDataStride)
+	SHADER_PARAMETER(float, TemporalMaxFramesAccumulated)
+	SHADER_PARAMETER(float, TemporalNeighborhoodClampScale)
+	SHADER_PARAMETER(int32, DebugMode)
+	SHADER_PARAMETER(int32, DebugLightId)
+	SHADER_PARAMETER_RDG_TEXTURE(Texture2D<uint>, DownsampledTileMask)
+	SHADER_PARAMETER_RDG_TEXTURE(Texture2D<float>, DownsampledSceneDepth)
+	SHADER_PARAMETER_RDG_TEXTURE(Texture2D<UNORM float3>, DownsampledSceneWorldNormal)
+END_SHADER_PARAMETER_STRUCT()
+
+// Internal functions, don't use outside of the StochasticDirectLighting
+namespace StochasticDirectLighting
+{
+	void RayTraceLightSamples(
+		const FViewInfo& View,
+		FRDGBuilder& GraphBuilder,
+		const FSceneTextures& SceneTextures,
+		const FIntPoint SampleBufferSize,
+		FRDGTextureRef LightSamples,
+		FRDGTextureRef LightSampleRayDistance,
+		const FStochasticDirectLightingParameters& StochasticDirectLightingParameters
+	);
+
+	bool ShouldCompileShaders(const FGlobalShaderPermutationParameters& Parameters);
+	bool UseWaveOps(EShaderPlatform ShaderPlatform);
+	int32 GetDebugMode();
+
+	void ModifyCompilationEnvironment(EShaderPlatform Platform, FShaderCompilerEnvironment& OutEnvironment);
+};
\ No newline at end of file
diff --git a/Engine/Source/Runtime/Renderer/Private/StochasticDirectLighting/StochasticDirectLightingRayTracing.cpp b/Engine/Source/Runtime/Renderer/Private/StochasticDirectLighting/StochasticDirectLightingRayTracing.cpp
new file mode 100644
--- /dev/null
+++ b/Engine/Source/Runtime/Renderer/Private/StochasticDirectLighting/StochasticDirectLightingRayTracing.cpp
@@ -0,0 +1,627 @@
+// Copyright Epic Games, Inc. All Rights Reserved.
+
+#include "StochasticDirectLighting.h"
+#include "StochasticDirectLightingInternal.h"
+#include "Lumen/LumenTracingUtils.h"
+#include "Lumen/LumenHardwareRayTracingCommon.h"
+#include "BasePassRendering.h"
+
+namespace StochasticDirectLighting
+{
+
+static TAutoConsoleVariable<int32> CVarStochasticDirectLightingScreenTraces(
+	TEXT("r.StochasticDirectLighting.ScreenTraces"),
+	0,
+	TEXT("Whether to use screen space tracing for shadow rays."),
+	ECVF_Scalability | ECVF_RenderThreadSafe
+);
+
+static TAutoConsoleVariable<int32> CVarStochasticDirectLightingScreenTracesMaxIterations(
+	TEXT("r.StochasticDirectLighting.ScreenTraces.MaxIterations"),
+	50,
+	TEXT("Max iterations for HZB tracing."),
+	ECVF_Scalability | ECVF_RenderThreadSafe
+);
+
+static TAutoConsoleVariable<int32> CVarStochasticDirectLightingScreenTracesMinimumOccupancy(
+	TEXT("r.StochasticDirectLighting.ScreenTraces.MinimumOccupancy"),
+	0,
+	TEXT("Minimum number of threads still tracing before aborting the trace. Can be used for scalability to abandon traces that have a disproportionate cost."),
+	ECVF_Scalability | ECVF_RenderThreadSafe
+);
+
+static TAutoConsoleVariable<float> CVarStochasticDirectLightingScreenTraceRelativeDepthThreshold(
+	TEXT("r.StochasticDirectLighting.ScreenTraces.RelativeDepthThickness"),
+	0.005f,
+	TEXT("Determines depth thickness of objects hit by HZB tracing, as a relative depth threshold."),
+	ECVF_Scalability | ECVF_RenderThreadSafe
+);
+
+static TAutoConsoleVariable<int32> CVarStochasticDirectLightingWorldSpaceTraces(
+	TEXT("r.StochasticDirectLighting.WorldSpaceTraces"),
+	1,
+	TEXT("Whether to trace world space shadow rays for samples. Useful for debugging."),
+	ECVF_Scalability | ECVF_RenderThreadSafe
+);
+
+static TAutoConsoleVariable<int32> CVarStochasticDirectLightingHardwareRayTracing(
+	TEXT("r.StochasticDirectLighting.HardwareRayTracing"),
+	1,
+	TEXT("Whether to use hardware ray tracing for shadow rays."),
+	ECVF_Scalability | ECVF_RenderThreadSafe
+);
+
+static TAutoConsoleVariable<int32> CVarStochasticDirectLightingHardwareRayTracingInline(
+	TEXT("r.StochasticDirectLighting.HardwareRayTracing.Inline"),
+	1,
+	TEXT("Uses hardware inline ray tracing for ray traced lighting, when available."),
+	ECVF_RenderThreadSafe | ECVF_Scalability
+);
+
+static TAutoConsoleVariable<float> CVarStochasticDirectLightingHardwareRayTracingBias(
+	TEXT("r.StochasticDirectLighting.HardwareRayTracing.Bias"),
+	1.0f,
+	TEXT("Constant bias for hardware ray traced shadow rays."),
+	ECVF_Scalability | ECVF_RenderThreadSafe
+);
+
+static TAutoConsoleVariable<float> CVarStochasticDirectLightingHardwareRayTracingNormalBias(
+	TEXT("r.StochasticDirectLighting.HardwareRayTracing.NormalBias"),
+	0.1f,
+	TEXT("Normal bias for hardware ray traced shadow rays."),
+	ECVF_Scalability | ECVF_RenderThreadSafe
+);
+
+static TAutoConsoleVariable<int32> CVarStochasticDirectLightingHardwareRayTracingMaxIterations(
+	TEXT("r.StochasticDirectLighting.HardwareRayTracing.MaxIterations"),
+	8192,
+	TEXT("Limit number of ray tracing traversal iterations on supported platfoms. Improves performance, but may add over-occlusion."),
+	ECVF_Scalability | ECVF_RenderThreadSafe
+);
+
+static TAutoConsoleVariable<int32> CVarStochasticDirectLightingHairVoxelTraces(
+	TEXT("r.StochasticDirectLighting.HairVoxelTraces"),
+	1,
+	TEXT("Whether to trace hair voxels."),
+	ECVF_Scalability | ECVF_RenderThreadSafe
+);
+
+bool UseHardwareRayTracing()
+{
+	#if RHI_RAYTRACING
+	{
+		return IsRayTracingEnabled() 
+			&& CVarStochasticDirectLightingHardwareRayTracing.GetValueOnRenderThread() != 0;
+	}
+	#else
+	{
+		return false;
+	}
+	#endif
+}
+
+bool UseInlineHardwareRayTracing()
+{
+	#if RHI_RAYTRACING
+	{
+		return UseHardwareRayTracing()
+			&& GRHISupportsInlineRayTracing
+			&& CVarStochasticDirectLightingHardwareRayTracingInline.GetValueOnRenderThread() != 0;
+	}
+	#else
+	{
+		return false;
+	}
+	#endif
+}
+
+bool IsUsingClosestHZB()
+{
+	return IsEnabled() && CVarStochasticDirectLightingScreenTraces.GetValueOnRenderThread() != 0;
+}
+
+bool IsUsingGlobalSDF()
+{
+	return IsEnabled() && CVarStochasticDirectLightingWorldSpaceTraces.GetValueOnRenderThread() != 0 && !UseHardwareRayTracing();
+}
+
+BEGIN_SHADER_PARAMETER_STRUCT(FHairVoxelTraceParameters, )
+	SHADER_PARAMETER_RDG_UNIFORM_BUFFER(FHairStrandsViewUniformParameters, HairStrands)
+	SHADER_PARAMETER_RDG_UNIFORM_BUFFER(FVirtualVoxelParameters, VirtualVoxel)
+END_SHADER_PARAMETER_STRUCT()
+
+BEGIN_SHADER_PARAMETER_STRUCT(FCompactedTraceParameters, )
+	RDG_BUFFER_ACCESS(IndirectArgs, ERHIAccess::IndirectArgs)
+	SHADER_PARAMETER_RDG_BUFFER_SRV(StructuredBuffer<uint>, CompactedTraceTexelData)
+	SHADER_PARAMETER_RDG_BUFFER_SRV(StructuredBuffer<uint>, CompactedTraceTexelAllocator)
+END_SHADER_PARAMETER_STRUCT()
+
+enum class ECompactedTraceIndirectArgs
+{
+	NumTracesDiv64 = 0 * sizeof(FRHIDispatchIndirectParameters),
+	NumTracesDiv32 = 1 * sizeof(FRHIDispatchIndirectParameters),
+	NumTraces = 2 * sizeof(FRHIDispatchIndirectParameters),
+	MAX = 3
+};
+
+class FCompactLightSampleTracesCS : public FGlobalShader
+{
+	DECLARE_GLOBAL_SHADER(FCompactLightSampleTracesCS)
+	SHADER_USE_PARAMETER_STRUCT(FCompactLightSampleTracesCS, FGlobalShader);
+
+	BEGIN_SHADER_PARAMETER_STRUCT(FParameters, )
+		SHADER_PARAMETER_STRUCT_INCLUDE(FStochasticDirectLightingParameters, StochasticDirectLightingParameters)
+		SHADER_PARAMETER_RDG_BUFFER_UAV(RWStructuredBuffer<uint>, RWCompactedTraceTexelData)
+		SHADER_PARAMETER_RDG_BUFFER_UAV(RWStructuredBuffer<uint>, RWCompactedTraceTexelAllocator)
+		SHADER_PARAMETER_RDG_TEXTURE(Texture2D<uint>, LightSamples)
+	END_SHADER_PARAMETER_STRUCT()
+
+	static int32 GetGroupSize()
+	{
+		return 16;
+	}
+
+	class FWaveOps : SHADER_PERMUTATION_BOOL("WAVE_OPS");
+	using FPermutationDomain = TShaderPermutationDomain<FWaveOps>;
+
+	static bool ShouldCompilePermutation(const FGlobalShaderPermutationParameters& Parameters)
+	{
+		return StochasticDirectLighting::ShouldCompileShaders(Parameters);
+	}
+
+	FORCENOINLINE static void ModifyCompilationEnvironment(const FGlobalShaderPermutationParameters& Parameters, FShaderCompilerEnvironment& OutEnvironment)
+	{
+		FGlobalShader::ModifyCompilationEnvironment(Parameters, OutEnvironment);
+		OutEnvironment.SetDefine(TEXT("THREADGROUP_SIZE"), GetGroupSize());
+
+		FPermutationDomain PermutationVector(Parameters.PermutationId);
+		if (PermutationVector.Get<FWaveOps>())
+		{
+			OutEnvironment.CompilerFlags.Add(CFLAG_WaveOperations);
+		}
+	}
+};
+
+IMPLEMENT_GLOBAL_SHADER(FCompactLightSampleTracesCS, "/Engine/Private/StochasticDirectLighting/StochasticDirectLightingTracing.usf", "CompactLightSampleTracesCS", SF_Compute);
+
+class FInitCompactedTraceTexelIndirectArgsCS : public FGlobalShader
+{
+	DECLARE_GLOBAL_SHADER(FInitCompactedTraceTexelIndirectArgsCS)
+	SHADER_USE_PARAMETER_STRUCT(FInitCompactedTraceTexelIndirectArgsCS, FGlobalShader);
+
+	BEGIN_SHADER_PARAMETER_STRUCT(FParameters, )
+		SHADER_PARAMETER_STRUCT_INCLUDE(FStochasticDirectLightingParameters, StochasticDirectLightingParameters)
+		SHADER_PARAMETER_RDG_BUFFER_UAV(RWBuffer<uint>, RWIndirectArgs)
+		SHADER_PARAMETER_RDG_BUFFER_SRV(StructuredBuffer<uint>, CompactedTraceTexelAllocator)
+	END_SHADER_PARAMETER_STRUCT()
+
+	static bool ShouldCompilePermutation(const FGlobalShaderPermutationParameters& Parameters)
+	{
+		return StochasticDirectLighting::ShouldCompileShaders(Parameters);
+	}
+
+	static void ModifyCompilationEnvironment(const FGlobalShaderPermutationParameters& Parameters, FShaderCompilerEnvironment& OutEnvironment)
+	{
+		FGlobalShader::ModifyCompilationEnvironment(Parameters, OutEnvironment);
+		OutEnvironment.SetDefine(TEXT("THREADGROUP_SIZE"), GetGroupSize());
+	}
+
+	static int32 GetGroupSize()
+	{
+		return 64;
+	}
+};
+
+IMPLEMENT_GLOBAL_SHADER(FInitCompactedTraceTexelIndirectArgsCS, "/Engine/Private/StochasticDirectLighting/StochasticDirectLightingTracing.usf", "InitCompactedTraceTexelIndirectArgsCS", SF_Compute);
+
+#if RHI_RAYTRACING
+
+class FHardwareRayTraceLightSamples : public FLumenHardwareRayTracingShaderBase
+{
+	DECLARE_LUMEN_RAYTRACING_SHADER(FHardwareRayTraceLightSamples, Lumen::ERayTracingShaderDispatchSize::DispatchSize1D)
+
+	BEGIN_SHADER_PARAMETER_STRUCT(FParameters, )
+		SHADER_PARAMETER_STRUCT_INCLUDE(StochasticDirectLighting::FCompactedTraceParameters, CompactedTraceParameters)
+		SHADER_PARAMETER_STRUCT_INCLUDE(FStochasticDirectLightingParameters, StochasticDirectLightingParameters)
+		SHADER_PARAMETER_STRUCT_INCLUDE(FHairVoxelTraceParameters, HairVoxelTraceParameters)
+		SHADER_PARAMETER_RDG_TEXTURE_UAV(RWTexture2D<uint>, RWLightSamples)
+		SHADER_PARAMETER_RDG_TEXTURE(Texture2D<float>, LightSampleRayDistance)
+		SHADER_PARAMETER(float, RayTracingBias)
+		SHADER_PARAMETER(float, RayTracingNormalBias)
+		// Ray Tracing
+		SHADER_PARAMETER(uint32, MaxTraversalIterations)
+		SHADER_PARAMETER_RDG_BUFFER_SRV(RaytracingAccelerationStructure, TLAS)
+		SHADER_PARAMETER_SRV(StructuredBuffer, RayTracingSceneMetadata)
+		// Inline Ray Tracing
+		SHADER_PARAMETER_RDG_BUFFER_SRV(StructuredBuffer<Lumen::FHitGroupRootConstants>, HitGroupData)
+		SHADER_PARAMETER_RDG_UNIFORM_BUFFER(FLumenHardwareRayTracingUniformBufferParameters, LumenHardwareRayTracingUniformBuffer)
+	END_SHADER_PARAMETER_STRUCT()
+
+	class FHairVoxelTraces : SHADER_PERMUTATION_BOOL("HAIR_VOXEL_TRACES");
+	class FDebugMode : SHADER_PERMUTATION_BOOL("DEBUG_MODE");
+	using FPermutationDomain = TShaderPermutationDomain<FHairVoxelTraces, FDebugMode>;
+
+	static bool ShouldCompilePermutation(const FGlobalShaderPermutationParameters& Parameters, Lumen::ERayTracingShaderDispatchType ShaderDispatchType)
+	{
+		return StochasticDirectLighting::ShouldCompileShaders(Parameters)  
+			&& FLumenHardwareRayTracingShaderBase::ShouldCompilePermutation(Parameters, ShaderDispatchType);
+	}
+
+	static void ModifyCompilationEnvironment(const FGlobalShaderPermutationParameters& Parameters, Lumen::ERayTracingShaderDispatchType ShaderDispatchType, FShaderCompilerEnvironment& OutEnvironment)
+	{
+		FLumenHardwareRayTracingShaderBase::ModifyCompilationEnvironment(Parameters, ShaderDispatchType, Lumen::ESurfaceCacheSampling::AlwaysResidentPagesWithoutFeedback, OutEnvironment);
+		StochasticDirectLighting::ModifyCompilationEnvironment(Parameters.Platform, OutEnvironment);
+	}
+
+	static ERayTracingPayloadType GetRayTracingPayloadType(const int32 PermutationId)
+	{
+		return ERayTracingPayloadType::LumenMinimal;
+	}
+};
+
+IMPLEMENT_LUMEN_RAYGEN_AND_COMPUTE_RAYTRACING_SHADERS(FHardwareRayTraceLightSamples)
+
+IMPLEMENT_GLOBAL_SHADER(FHardwareRayTraceLightSamplesCS, "/Engine/Private/StochasticDirectLighting/StochasticDirectLightingHardwareRayTracing.usf", "HardwareRayTraceLightSamplesCS", SF_Compute);
+IMPLEMENT_GLOBAL_SHADER(FHardwareRayTraceLightSamplesRGS, "/Engine/Private/StochasticDirectLighting/StochasticDirectLightingHardwareRayTracing.usf", "HardwareRayTraceLightSamplesRGS", SF_RayGen);
+
+#endif // RHI_RAYTRACING
+
+class FSoftwareRayTraceLightSamplesCS : public FGlobalShader
+{
+	DECLARE_GLOBAL_SHADER(FSoftwareRayTraceLightSamplesCS)
+	SHADER_USE_PARAMETER_STRUCT(FSoftwareRayTraceLightSamplesCS, FGlobalShader);
+
+	BEGIN_SHADER_PARAMETER_STRUCT(FParameters, )
+		SHADER_PARAMETER_STRUCT_INCLUDE(StochasticDirectLighting::FCompactedTraceParameters, CompactedTraceParameters)
+		SHADER_PARAMETER_STRUCT_INCLUDE(FStochasticDirectLightingParameters, StochasticDirectLightingParameters)
+		SHADER_PARAMETER_STRUCT_INCLUDE(FHairVoxelTraceParameters, HairVoxelTraceParameters)
+		SHADER_PARAMETER_RDG_TEXTURE_UAV(RWTexture2D<uint>, RWLightSamples)
+		SHADER_PARAMETER_RDG_TEXTURE(Texture2D<float>, LightSampleRayDistance)
+	END_SHADER_PARAMETER_STRUCT()
+
+	static int32 GetGroupSize()
+	{
+		return 64;
+	}
+
+	class FHairVoxelTraces : SHADER_PERMUTATION_BOOL("HAIR_VOXEL_TRACES");
+	class FDebugMode : SHADER_PERMUTATION_BOOL("DEBUG_MODE");
+	using FPermutationDomain = TShaderPermutationDomain<FHairVoxelTraces, FDebugMode>;
+
+	static bool ShouldCompilePermutation(const FGlobalShaderPermutationParameters& Parameters)
+	{
+		return StochasticDirectLighting::ShouldCompileShaders(Parameters);
+	}
+
+	FORCENOINLINE static void ModifyCompilationEnvironment(const FGlobalShaderPermutationParameters& Parameters, FShaderCompilerEnvironment& OutEnvironment)
+	{
+		FGlobalShader::ModifyCompilationEnvironment(Parameters, OutEnvironment);
+		StochasticDirectLighting::ModifyCompilationEnvironment(Parameters.Platform, OutEnvironment);
+		OutEnvironment.CompilerFlags.Add(CFLAG_Wave32);
+		OutEnvironment.SetDefine(TEXT("THREADGROUP_SIZE"), GetGroupSize());
+
+		// GPU Scene definitions
+		OutEnvironment.SetDefine(TEXT("VF_SUPPORTS_PRIMITIVE_SCENE_DATA"), 1);
+	}
+};
+
+IMPLEMENT_GLOBAL_SHADER(FSoftwareRayTraceLightSamplesCS, "/Engine/Private/StochasticDirectLighting/StochasticDirectLightingTracing.usf", "SoftwareRayTraceLightSamplesCS", SF_Compute);
+
+class FScreenSpaceRayTraceLightSamplesCS : public FGlobalShader
+{
+	DECLARE_GLOBAL_SHADER(FScreenSpaceRayTraceLightSamplesCS)
+	SHADER_USE_PARAMETER_STRUCT(FScreenSpaceRayTraceLightSamplesCS, FGlobalShader);
+
+	BEGIN_SHADER_PARAMETER_STRUCT(FParameters, )
+		SHADER_PARAMETER_STRUCT_INCLUDE(StochasticDirectLighting::FCompactedTraceParameters, CompactedTraceParameters)
+		SHADER_PARAMETER_STRUCT_INCLUDE(FStochasticDirectLightingParameters, StochasticDirectLightingParameters)
+		SHADER_PARAMETER_RDG_TEXTURE_UAV(RWTexture2D<uint>, RWLightSamples)
+		SHADER_PARAMETER_RDG_TEXTURE_UAV(RWTexture2D<float>, RWLightSampleRayDistance)
+		SHADER_PARAMETER_STRUCT_INCLUDE(FLumenHZBScreenTraceParameters, HZBScreenTraceParameters)
+		SHADER_PARAMETER(float, MaxHierarchicalScreenTraceIterations)
+		SHADER_PARAMETER(float, RelativeDepthThickness)
+		SHADER_PARAMETER(float, HistoryDepthTestRelativeThickness)
+		SHADER_PARAMETER(uint32, MinimumTracingThreadOccupancy)
+	END_SHADER_PARAMETER_STRUCT()
+
+	static int32 GetGroupSize()
+	{
+		return 64;
+	}
+
+	class FDebugMode : SHADER_PERMUTATION_BOOL("DEBUG_MODE");
+	using FPermutationDomain = TShaderPermutationDomain<FDebugMode>;
+
+	static bool ShouldCompilePermutation(const FGlobalShaderPermutationParameters& Parameters)
+	{
+		return StochasticDirectLighting::ShouldCompileShaders(Parameters);
+	}
+
+	FORCENOINLINE static void ModifyCompilationEnvironment(const FGlobalShaderPermutationParameters& Parameters, FShaderCompilerEnvironment& OutEnvironment)
+	{
+		FGlobalShader::ModifyCompilationEnvironment(Parameters, OutEnvironment);
+		StochasticDirectLighting::ModifyCompilationEnvironment(Parameters.Platform, OutEnvironment);
+		OutEnvironment.CompilerFlags.Add(CFLAG_Wave32);
+		OutEnvironment.SetDefine(TEXT("THREADGROUP_SIZE"), GetGroupSize());
+	}
+};
+
+IMPLEMENT_GLOBAL_SHADER(FScreenSpaceRayTraceLightSamplesCS, "/Engine/Private/StochasticDirectLighting/StochasticDirectLightingTracing.usf", "ScreenSpaceRayTraceLightSamplesCS", SF_Compute);
+
+#if RHI_RAYTRACING
+void SetHardwareRayTracingPassParameters(
+	const FViewInfo& View,
+	FRDGBuilder& GraphBuilder,
+	const FCompactedTraceParameters& CompactedTraceParameters,
+	const FStochasticDirectLightingParameters& StochasticDirectLightingParameters,
+	FRDGTextureRef LightSamples,
+	FRDGTextureRef LightSampleRayDistance,
+	FHardwareRayTraceLightSamples::FParameters* PassParameters)
+{
+	PassParameters->CompactedTraceParameters = CompactedTraceParameters;
+	PassParameters->StochasticDirectLightingParameters = StochasticDirectLightingParameters;
+	PassParameters->RWLightSamples = GraphBuilder.CreateUAV(LightSamples);
+	PassParameters->LightSampleRayDistance = LightSampleRayDistance;
+	PassParameters->RayTracingBias = CVarStochasticDirectLightingHardwareRayTracingBias.GetValueOnRenderThread();
+	PassParameters->RayTracingNormalBias = CVarStochasticDirectLightingHardwareRayTracingNormalBias.GetValueOnRenderThread();
+
+	checkf(View.HasRayTracingScene(), TEXT("TLAS does not exist. Verify that the current pass is represented in Lumen::AnyLumenHardwareRayTracingPassEnabled()."));
+	PassParameters->TLAS = View.GetRayTracingSceneLayerViewChecked(ERayTracingSceneLayer::Base);
+	PassParameters->MaxTraversalIterations = FMath::Max(CVarStochasticDirectLightingHardwareRayTracingMaxIterations.GetValueOnRenderThread(), 1);
+
+	// Inline
+	PassParameters->HitGroupData = View.GetPrimaryView()->LumenHardwareRayTracingHitDataBuffer ? GraphBuilder.CreateSRV(View.GetPrimaryView()->LumenHardwareRayTracingHitDataBuffer) : nullptr;
+	PassParameters->LumenHardwareRayTracingUniformBuffer = View.GetPrimaryView()->LumenHardwareRayTracingUniformBuffer ? View.GetPrimaryView()->LumenHardwareRayTracingUniformBuffer : nullptr;
+	checkf(View.RayTracingSceneInitTask == nullptr, TEXT("RayTracingSceneInitTask must be completed before creating SRV for RayTracingSceneMetadata."));
+	PassParameters->RayTracingSceneMetadata = View.GetRayTracingSceneChecked()->GetOrCreateMetadataBufferSRV(GraphBuilder.RHICmdList);
+}
+
+#endif
+
+FCompactedTraceParameters CompactStochasticDirectLightingTraces(
+	const FViewInfo& View,
+	FRDGBuilder& GraphBuilder,
+	const FIntPoint SampleBufferSize,
+	FRDGTextureRef LightSamples,
+	const FStochasticDirectLightingParameters& StochasticDirectLightingParameters)
+{
+	FRDGBufferRef CompactedTraceTexelData = GraphBuilder.CreateBuffer(FRDGBufferDesc::CreateStructuredDesc(sizeof(uint32), SampleBufferSize.X * SampleBufferSize.Y),
+		TEXT("StochasticDirectLightingParameters.CompactedTraceTexelData"));
+
+	FRDGBufferRef CompactedTraceTexelAllocator = GraphBuilder.CreateBuffer(FRDGBufferDesc::CreateStructuredDesc(sizeof(uint32), 1),
+		TEXT("StochasticDirectLightingParameters.CompactedTraceTexelAllocator"));
+
+	FRDGBufferRef CompactedTraceTexelIndirectArgs = GraphBuilder.CreateBuffer(
+		FRDGBufferDesc::CreateIndirectDesc<FRHIDispatchIndirectParameters>((int32)ECompactedTraceIndirectArgs::MAX),
+		TEXT("StochasticDirectLighting.CompactedTraceTexelIndirectArgs"));
+
+	AddClearUAVPass(GraphBuilder, GraphBuilder.CreateUAV(CompactedTraceTexelAllocator), 0);
+
+	// Compact light sample traces before tracing
+	{
+		FCompactLightSampleTracesCS::FParameters* PassParameters = GraphBuilder.AllocParameters<FCompactLightSampleTracesCS::FParameters>();
+		PassParameters->RWCompactedTraceTexelData = GraphBuilder.CreateUAV(CompactedTraceTexelData);
+		PassParameters->RWCompactedTraceTexelAllocator = GraphBuilder.CreateUAV(CompactedTraceTexelAllocator);
+		PassParameters->StochasticDirectLightingParameters = StochasticDirectLightingParameters;
+		PassParameters->LightSamples = LightSamples;
+
+		const bool bWaveOps = StochasticDirectLighting::UseWaveOps(View.GetShaderPlatform())
+			&& GRHIMinimumWaveSize <= 32
+			&& GRHIMaximumWaveSize >= 32;
+
+		FCompactLightSampleTracesCS::FPermutationDomain PermutationVector;
+		PermutationVector.Set<FCompactLightSampleTracesCS::FWaveOps>(bWaveOps);
+		auto ComputeShader = View.ShaderMap->GetShader<FCompactLightSampleTracesCS>(PermutationVector);
+
+		const FIntVector GroupCount = FComputeShaderUtils::GetGroupCount(StochasticDirectLightingParameters.SampleViewSize, FCompactLightSampleTracesCS::GetGroupSize());
+
+		FComputeShaderUtils::AddPass(
+			GraphBuilder,
+			RDG_EVENT_NAME("CompactLightSampleTraces"),
+			ComputeShader,
+			PassParameters,
+			GroupCount);
+	}
+
+	// Setup indirect args for tracing
+	{
+		FInitCompactedTraceTexelIndirectArgsCS::FParameters* PassParameters = GraphBuilder.AllocParameters<FInitCompactedTraceTexelIndirectArgsCS::FParameters>();
+		PassParameters->StochasticDirectLightingParameters = StochasticDirectLightingParameters;
+		PassParameters->RWIndirectArgs = GraphBuilder.CreateUAV(CompactedTraceTexelIndirectArgs);
+		PassParameters->CompactedTraceTexelAllocator = GraphBuilder.CreateSRV(CompactedTraceTexelAllocator);
+
+		auto ComputeShader = View.ShaderMap->GetShader<FInitCompactedTraceTexelIndirectArgsCS>();
+
+		FComputeShaderUtils::AddPass(
+			GraphBuilder,
+			RDG_EVENT_NAME("InitCompactedTraceTexelIndirectArgs"),
+			ComputeShader,
+			PassParameters,
+			FIntVector(1, 1, 1));
+	}
+
+	FCompactedTraceParameters Parameters;
+	Parameters.CompactedTraceTexelAllocator = GraphBuilder.CreateSRV(CompactedTraceTexelAllocator);
+	Parameters.CompactedTraceTexelData = GraphBuilder.CreateSRV(CompactedTraceTexelData);
+	Parameters.IndirectArgs = CompactedTraceTexelIndirectArgs;
+	return Parameters;
+}
+
+/**
+ * Ray trace light samples using a variety of tracing methods depending on the feature configuration.
+ */
+void RayTraceLightSamples(
+	const FViewInfo& View,
+	FRDGBuilder& GraphBuilder,
+	const FSceneTextures& SceneTextures,
+	const FIntPoint SampleBufferSize,
+	FRDGTextureRef LightSamples,
+	FRDGTextureRef LightSampleRayDistance,
+	const FStochasticDirectLightingParameters& StochasticDirectLightingParameters)
+{
+	const bool bDebug = StochasticDirectLighting::GetDebugMode() != 0;
+
+	if (CVarStochasticDirectLightingScreenTraces.GetValueOnRenderThread() != 0)
+	{
+		FCompactedTraceParameters CompactedTraceParameters = CompactStochasticDirectLightingTraces(
+			View,
+			GraphBuilder,
+			SampleBufferSize,
+			LightSamples,
+			StochasticDirectLightingParameters);
+
+		FScreenSpaceRayTraceLightSamplesCS::FParameters* PassParameters = GraphBuilder.AllocParameters<FScreenSpaceRayTraceLightSamplesCS::FParameters>();
+		PassParameters->CompactedTraceParameters = CompactedTraceParameters;
+		PassParameters->StochasticDirectLightingParameters = StochasticDirectLightingParameters;
+		PassParameters->RWLightSamples = GraphBuilder.CreateUAV(LightSamples);
+		PassParameters->RWLightSampleRayDistance = GraphBuilder.CreateUAV(LightSampleRayDistance);
+		PassParameters->HZBScreenTraceParameters = SetupHZBScreenTraceParameters(GraphBuilder, View, SceneTextures, /*bBindLumenHistory*/ false);
+		PassParameters->MaxHierarchicalScreenTraceIterations = CVarStochasticDirectLightingScreenTracesMaxIterations.GetValueOnRenderThread();
+		PassParameters->RelativeDepthThickness = CVarStochasticDirectLightingScreenTraceRelativeDepthThreshold.GetValueOnRenderThread() * View.ViewMatrices.GetPerProjectionDepthThicknessScale();
+		PassParameters->HistoryDepthTestRelativeThickness = 0.0f;
+		PassParameters->MinimumTracingThreadOccupancy = CVarStochasticDirectLightingScreenTracesMinimumOccupancy.GetValueOnRenderThread();
+
+		FScreenSpaceRayTraceLightSamplesCS::FPermutationDomain PermutationVector;
+		PermutationVector.Set<FScreenSpaceRayTraceLightSamplesCS::FDebugMode>(bDebug);
+		auto ComputeShader = View.ShaderMap->GetShader<FScreenSpaceRayTraceLightSamplesCS>(PermutationVector);
+
+		FComputeShaderUtils::AddPass(
+			GraphBuilder,
+			RDG_EVENT_NAME("ScreenSpaceRayTraceLightSamples"),
+			ComputeShader,
+			PassParameters,
+			CompactedTraceParameters.IndirectArgs,
+			(int32)StochasticDirectLighting::ECompactedTraceIndirectArgs::NumTracesDiv64);
+	}
+	else
+	{
+		AddClearUAVPass(GraphBuilder, GraphBuilder.CreateUAV(LightSampleRayDistance), 0.0f);
+	}
+
+	const bool bHairVoxelTraces = HairStrands::HasViewHairStrandsData(View)
+		&& HairStrands::HasViewHairStrandsVoxelData(View)
+		&& CVarStochasticDirectLightingHairVoxelTraces.GetValueOnRenderThread() != 0;
+
+	FHairVoxelTraceParameters HairVoxelTraceParameters;
+	if (bHairVoxelTraces)
+	{
+		HairVoxelTraceParameters.HairStrands = HairStrands::BindHairStrandsViewUniformParameters(View);
+		HairVoxelTraceParameters.VirtualVoxel = HairStrands::BindHairStrandsVoxelUniformParameters(View);
+	}
+
+	if (CVarStochasticDirectLightingWorldSpaceTraces.GetValueOnRenderThread() != 0)
+	{
+		FCompactedTraceParameters CompactedTraceParameters = CompactStochasticDirectLightingTraces(
+			View,
+			GraphBuilder,
+			SampleBufferSize,
+			LightSamples,
+			StochasticDirectLightingParameters);
+
+		if (StochasticDirectLighting::UseHardwareRayTracing())
+		{
+			#if RHI_RAYTRACING
+			if (StochasticDirectLighting::UseInlineHardwareRayTracing())
+			{
+				FHardwareRayTraceLightSamplesCS::FParameters* PassParameters = GraphBuilder.AllocParameters<FHardwareRayTraceLightSamplesCS::FParameters>();
+				StochasticDirectLighting::SetHardwareRayTracingPassParameters(
+					View,
+					GraphBuilder,
+					CompactedTraceParameters,
+					StochasticDirectLightingParameters,
+					LightSamples,
+					LightSampleRayDistance,
+					PassParameters);
+				PassParameters->HairVoxelTraceParameters = HairVoxelTraceParameters;
+
+				FHardwareRayTraceLightSamplesCS::FPermutationDomain PermutationVector;
+				PermutationVector.Set<FHardwareRayTraceLightSamplesCS::FHairVoxelTraces>(bHairVoxelTraces);
+				PermutationVector.Set<FHardwareRayTraceLightSamplesCS::FDebugMode>(bDebug);
+				auto ComputeShader = View.ShaderMap->GetShader<FHardwareRayTraceLightSamplesCS>(PermutationVector);
+
+				FComputeShaderUtils::AddPass(
+					GraphBuilder,
+					RDG_EVENT_NAME("HardwareRayTraceLightSamples Inline"),
+					ERDGPassFlags::Compute,
+					ComputeShader,
+					PassParameters,
+					CompactedTraceParameters.IndirectArgs,
+					(int32)StochasticDirectLighting::ECompactedTraceIndirectArgs::NumTracesDiv32);
+			}
+			else
+			{
+				FHardwareRayTraceLightSamplesRGS::FParameters* PassParameters = GraphBuilder.AllocParameters<FHardwareRayTraceLightSamplesRGS::FParameters>();
+				StochasticDirectLighting::SetHardwareRayTracingPassParameters(
+					View,
+					GraphBuilder,
+					CompactedTraceParameters,
+					StochasticDirectLightingParameters,
+					LightSamples,
+					LightSampleRayDistance,
+					PassParameters);
+				PassParameters->HairVoxelTraceParameters = HairVoxelTraceParameters;
+
+				FHardwareRayTraceLightSamplesRGS::FPermutationDomain PermutationVector;
+				PermutationVector.Set<FHardwareRayTraceLightSamplesRGS::FHairVoxelTraces>(bHairVoxelTraces);
+				PermutationVector.Set<FHardwareRayTraceLightSamplesRGS::FDebugMode>(bDebug);
+				auto RayGenerationShader = View.ShaderMap->GetShader<FHardwareRayTraceLightSamplesRGS>(PermutationVector);
+
+				AddLumenRayTraceDispatchIndirectPass(
+					GraphBuilder,
+					RDG_EVENT_NAME("HardwareRayTraceLightSamples RayGen"),
+					RayGenerationShader,
+					PassParameters,
+					PassParameters->CompactedTraceParameters.IndirectArgs,
+					(int32)StochasticDirectLighting::ECompactedTraceIndirectArgs::NumTraces,
+					View,
+					/*bUseMinimalPayload*/ true);
+			}
+			#endif
+		}
+		else
+		{
+			ensure(StochasticDirectLighting::IsUsingGlobalSDF());
+
+			FSoftwareRayTraceLightSamplesCS::FParameters* PassParameters = GraphBuilder.AllocParameters<FSoftwareRayTraceLightSamplesCS::FParameters>();
+			PassParameters->CompactedTraceParameters = CompactedTraceParameters;
+			PassParameters->StochasticDirectLightingParameters = StochasticDirectLightingParameters;
+			PassParameters->HairVoxelTraceParameters = HairVoxelTraceParameters;
+			PassParameters->RWLightSamples = GraphBuilder.CreateUAV(LightSamples);
+			PassParameters->LightSampleRayDistance = LightSampleRayDistance;
+
+			FSoftwareRayTraceLightSamplesCS::FPermutationDomain PermutationVector;
+			PermutationVector.Set<FSoftwareRayTraceLightSamplesCS::FHairVoxelTraces>(bHairVoxelTraces);
+			PermutationVector.Set<FSoftwareRayTraceLightSamplesCS::FDebugMode>(bDebug);
+			auto ComputeShader = View.ShaderMap->GetShader<FSoftwareRayTraceLightSamplesCS>(PermutationVector);
+
+			FComputeShaderUtils::AddPass(
+				GraphBuilder,
+				RDG_EVENT_NAME("SoftwareRayTraceLightSamples"),
+				ComputeShader,
+				PassParameters,
+				CompactedTraceParameters.IndirectArgs,
+				0);
+		}
+	}
+}
+
+} // namespace StochasticDirectLighting
+
+#if RHI_RAYTRACING
+void FDeferredShadingSceneRenderer::PrepareStochasticDirectLightingLumenMaterial(const FViewInfo& View, TArray<FRHIRayTracingShader*>& OutRayGenShaders)
+{
+	using namespace StochasticDirectLighting;
+
+	if (StochasticDirectLighting::IsEnabled() && StochasticDirectLighting::UseHardwareRayTracing())
+	{
+		for (int32 HairVoxelTraces = 0; HairVoxelTraces < 2; ++HairVoxelTraces)
+		{
+			FHardwareRayTraceLightSamplesRGS::FPermutationDomain PermutationVector;
+			PermutationVector.Set<FHardwareRayTraceLightSamplesRGS::FHairVoxelTraces>(HairVoxelTraces != 0);
+			PermutationVector.Set<FHardwareRayTraceLightSamplesRGS::FDebugMode>(StochasticDirectLighting::GetDebugMode() != 0);
+			TShaderRef<FHardwareRayTraceLightSamplesRGS> RayGenerationShader = View.ShaderMap->GetShader<FHardwareRayTraceLightSamplesRGS>(PermutationVector);
+			OutRayGenShaders.Add(RayGenerationShader.GetRayTracingShader());
+		}
+	}
+}
+#endif
\ No newline at end of file
diff --git a/Engine/Source/Runtime/Renderer/Private/StochasticDirectLighting/StochasticDirectLightingViewState.h b/Engine/Source/Runtime/Renderer/Private/StochasticDirectLighting/StochasticDirectLightingViewState.h
new file mode 100644
--- /dev/null
+++ b/Engine/Source/Runtime/Renderer/Private/StochasticDirectLighting/StochasticDirectLightingViewState.h
@@ -0,0 +1,26 @@
+// Copyright Epic Games, Inc. All Rights Reserved.
+
+#pragma once
+
+#include "Math/Vector4.h"
+#include "RendererInterface.h"
+
+class FStochasticDirectLightingViewState
+{
+public:
+	TRefCountPtr<IPooledRenderTarget> DiffuseLightingAndSecondMomentHistory;
+	TRefCountPtr<IPooledRenderTarget> SpecularLightingAndSecondMomentHistory;
+	TRefCountPtr<IPooledRenderTarget> SceneDepthHistory;
+	TRefCountPtr<IPooledRenderTarget> NumFramesAccumulatedHistory;
+
+	FVector4f HistoryScreenPositionScaleBias = FVector4f(0.0f, 0.0f, 0.0f, 0.0f);
+	FVector4f HistoryUVMinMax = FVector4f(0.0f, 0.0f, 0.0f, 0.0f);
+
+	void SafeRelease()
+	{
+		DiffuseLightingAndSecondMomentHistory.SafeRelease();
+		SpecularLightingAndSecondMomentHistory.SafeRelease();
+		SceneDepthHistory.SafeRelease();
+		NumFramesAccumulatedHistory.SafeRelease();
+	}
+};
\ No newline at end of file
diff --git a/Engine/Source/Runtime/Renderer/Private/StochasticShadows/StochasticShadows.cpp b/Engine/Source/Runtime/Renderer/Private/StochasticShadows/StochasticShadows.cpp
--- a/Engine/Source/Runtime/Renderer/Private/StochasticShadows/StochasticShadows.cpp
+++ b/Engine/Source/Runtime/Renderer/Private/StochasticShadows/StochasticShadows.cpp
@@ -5,6 +5,7 @@
 #include "RendererPrivate.h"
 #include "PixelShaderUtils.h"
 #include "BasePassRendering.h"
+#include "StochasticDirectLighting/StochasticDirectLighting.h"
 
 static TAutoConsoleVariable<int32> CVarStochasticShadows(
 	TEXT("r.StochasticShadows"),
@@ -161,6 +162,11 @@ namespace StochasticShadows
 
 	bool IsEnabled()
 	{
+		if (StochasticDirectLighting::IsEnabled())
+		{
+			return false;
+		}
+
 		return CVarStochasticShadows.GetValueOnRenderThread() != 0;
 	}
 
