Make hair sample texture to be snap for better allocation
* Change allocator counter to a structured buffer
* Add common function to compute atlas resolution and transform linear coord into 2D coord
* Round hair sample texture to improve texture reused on platform using pooled resouces as transient resources

diff --git a/Engine/Shaders/Private/ClusteredDeferredShadingVertexShader.usf b/Engine/Shaders/Private/ClusteredDeferredShadingVertexShader.usf
--- a/Engine/Shaders/Private/ClusteredDeferredShadingVertexShader.usf
+++ b/Engine/Shaders/Private/ClusteredDeferredShadingVertexShader.usf
@@ -6,6 +6,7 @@
 
 #include "HairStrands/HairStrandsVisibilityCommonStruct.ush" 
 #include "Common.ush"
+#include "HairStrands/HairStrandsVisibilityCommon.ush"
 
 struct FHairStrandsSamplesDesc
 {
@@ -20,9 +21,8 @@ void ClusteredShadingVertexShader(
 	nointerpolation out uint OutNodeCount : DISPATCH_NODECOUNT,
 	nointerpolation out uint2 OutResolution : DISPATCH_RESOLUTION)
 {
-	OutNodeCount = HairStrands.HairSampleCount.Load(uint3(0,0,0));
-	OutResolution.x = ceil(sqrt(OutNodeCount));
-	OutResolution.y = OutResolution.x;
+	OutNodeCount = HairStrands.HairSampleCount[0];
+	OutResolution = GetHairSampleResolution(OutNodeCount);
 
 	const float2 ClipCoord = ((OutResolution + 0.5f) / float2(HairStrands.HairSampleViewportResolution)) * 2;
 
diff --git a/Engine/Shaders/Private/DeferredLightVertexShaders.usf b/Engine/Shaders/Private/DeferredLightVertexShaders.usf
--- a/Engine/Shaders/Private/DeferredLightVertexShaders.usf
+++ b/Engine/Shaders/Private/DeferredLightVertexShaders.usf
@@ -117,15 +117,16 @@ void VertexMain(
 
 #if defined(SHADER_HAIR) && SHADER_HAIR == 1
 
+#include "HairStrands/HairStrandsVisibilityCommon.ush"
+
 void HairVertexMain(
 	in uint InVertexId : SV_VertexID,
 	out float4 OutPosition : SV_POSITION,
 	nointerpolation out uint OutNodeCount : DISPATCH_NODECOUNT,
 	nointerpolation out uint2 OutResolution : DISPATCH_RESOLUTION)
 {
-	OutNodeCount = HairStrands.HairSampleCount.Load(uint3(0,0,0));
-	OutResolution.x = ceil(sqrt(OutNodeCount));
-	OutResolution.y = OutResolution.x;
+	OutNodeCount = HairStrands.HairSampleCount[0];
+	OutResolution = GetHairSampleResolution(OutNodeCount);
 
 	const float2 ClipCoord = ((OutResolution + 0.5f) / float2(HairStrands.HairSampleViewportResolution)) * 2;
 
diff --git a/Engine/Shaders/Private/HairStrands/HairStrandsComposition.usf b/Engine/Shaders/Private/HairStrands/HairStrandsComposition.usf
--- a/Engine/Shaders/Private/HairStrands/HairStrandsComposition.usf
+++ b/Engine/Shaders/Private/HairStrands/HairStrandsComposition.usf
@@ -98,8 +98,8 @@ void ComposeSamplePS(
 		discard;
 	}
 
-	const uint TotalNodeCount = HairStrands.HairSampleCount.Load(uint3(0, 0, 0));
-	const uint Resolution = uint(ceil(sqrt(float(TotalNodeCount))));
+	const uint TotalNodeCount = HairStrands.HairSampleCount[0];
+	const uint2 Resolution = GetHairSampleResolution(TotalNodeCount);
 
 	const float3 ClosetPointWorldPosition = LWCHackToFloat(SvPositionToWorld(float4(Input.Position.xy, HairDeviceZ, 1.0)));
 	const float4 Volumetric = EvaluateVolumetric(ClosetPointWorldPosition, ConvertFromDeviceZ(HairDeviceZ));
@@ -117,8 +117,8 @@ void ComposeSamplePS(
 		}
 		#endif
 
-		const uint LocalOffset = NodeDesc.Offset + SampleIt;
-		const uint2 Coord = uint2(LocalOffset % Resolution, LocalOffset / Resolution);
+		const uint LinearCoord = NodeDesc.Offset + SampleIt;
+		const uint2 Coord = GetHairSampleCoord(LinearCoord, Resolution);
 		LocalAccColor += HairLightingSampleBuffer.Load(uint3(Coord, 0)).xyz;
 	}
 	OutColor.rgb = LocalAccColor * Volumetric.a + Volumetric.rgb;
diff --git a/Engine/Shaders/Private/HairStrands/HairStrandsDebug.usf b/Engine/Shaders/Private/HairStrands/HairStrandsDebug.usf
--- a/Engine/Shaders/Private/HairStrands/HairStrandsDebug.usf
+++ b/Engine/Shaders/Private/HairStrands/HairStrandsDebug.usf
@@ -1040,10 +1040,11 @@ void MainCS(uint InGroupIndex : SV_GroupIndex, uint3 DispatchThreadId : SV_Dispa
 
 		const uint DispatchX = HairVisibilityIndirectArgsBuffer[0];
 		const uint DispatchY = HairVisibilityIndirectArgsBuffer[1];
-		const uint UsedSampleCount = HairStrands.HairSampleCount.Load(uint3(0, 0, 0));
+		const uint UsedSampleCount = HairStrands.HairSampleCount[0];
 		const uint DispatchedSampleCount = DispatchX * DispatchY * HairVisibilityNodeGroupSize;
 
-		const uint2 EquivalentResolution = ceil(sqrt(UsedSampleCount));
+		const uint2 EquivalentResolution = GetHairSampleResolution(UsedSampleCount);
+		const uint2 AllocatedResolution = HairStrands.HairSampleViewportResolution;
 	
 		// Pixel coord
 		FShaderPrintContext Context = InitShaderPrintContext(true, uint2(450, 100));
@@ -1062,6 +1063,10 @@ void MainCS(uint InGroupIndex : SV_GroupIndex, uint3 DispatchThreadId : SV_Dispa
 		Print(Context, EquivalentResolution.x, FontCyan, 4);
 		Print(Context, TEXT("x"), FontSilver);
 		Print(Context, EquivalentResolution.y, FontCyan, 4);
+		Print(Context, TEXT(" | "), FontSilver);
+		Print(Context, AllocatedResolution.x, FontGrey, 4);
+		Print(Context, TEXT("x"), FontSilver);
+		Print(Context, AllocatedResolution.y, FontGrey, 4);
 		Newline(Context);
 
 		Print(Context, TEXT("MacroGroup Count           "), FontYellow);
diff --git a/Engine/Shaders/Private/HairStrands/HairStrandsDeepTransmittanceMask.usf b/Engine/Shaders/Private/HairStrands/HairStrandsDeepTransmittanceMask.usf
--- a/Engine/Shaders/Private/HairStrands/HairStrandsDeepTransmittanceMask.usf
+++ b/Engine/Shaders/Private/HairStrands/HairStrandsDeepTransmittanceMask.usf
@@ -163,7 +163,7 @@ void MainCS(uint2 DispatchThreadId : SV_DispatchThreadID)
 	// Note: Use a fixed group count width (HAIR_VISIBILITY_GROUP_COUNT_WIDTH) for avoiding loading the indirect args (to avoid dep. memory fetch)
 	const uint SampleIndex = DispatchThreadId.x + DispatchThreadId.y * TRANSMITTANCE_GROUP_SIZE * HAIR_VISIBILITY_GROUP_COUNT_WIDTH;
 
-	const uint MaxVisibilityNodeCount = HairStrands.HairSampleCount.Load(uint3(0, 0, 0));
+	const uint MaxVisibilityNodeCount = HairStrands.HairSampleCount[0];
 	if (SampleIndex >= MaxVisibilityNodeCount)
 	{
 		return;
@@ -355,7 +355,7 @@ void MainCS(uint2 DispatchThreadId : SV_DispatchThreadID)
 {
 	const uint SampleIndex = DispatchThreadId.x + DispatchThreadId.y * TRANSMITTANCE_GROUP_SIZE * HAIR_VISIBILITY_GROUP_COUNT_WIDTH;
 
-	const uint MaxVisibilityNodeCount = HairStrands.HairSampleCount.Load(uint3(0, 0, 0));
+	const uint MaxVisibilityNodeCount = HairStrands.HairSampleCount[0];
 	if (SampleIndex >= MaxVisibilityNodeCount)
 	{
 		return;
@@ -440,7 +440,7 @@ void MainCS(uint2 DispatchThreadId : SV_DispatchThreadID)
 {
 	const uint SampleIndex = DispatchThreadId.x + DispatchThreadId.y * TRANSMITTANCE_GROUP_SIZE * HAIR_VISIBILITY_GROUP_COUNT_WIDTH;
 
-	const uint MaxVisibilityNodeCount = HairStrands.HairSampleCount.Load(uint3(0, 0, 0));
+	const uint MaxVisibilityNodeCount = HairStrands.HairSampleCount[0];
 	if (SampleIndex >= MaxVisibilityNodeCount)
 	{
 		return;
diff --git a/Engine/Shaders/Private/HairStrands/HairStrandsEnvironmentLighting.usf b/Engine/Shaders/Private/HairStrands/HairStrandsEnvironmentLighting.usf
--- a/Engine/Shaders/Private/HairStrands/HairStrandsEnvironmentLighting.usf
+++ b/Engine/Shaders/Private/HairStrands/HairStrandsEnvironmentLighting.usf
@@ -808,9 +808,8 @@ void MainVS(
 	nointerpolation out uint OutNodeCount : DISPATCH_NODECOUNT,
 	nointerpolation out uint2 OutResolution : DISPATCH_RESOLUTION)
 {
-	OutNodeCount = HairStrands.HairSampleCount.Load(uint3(0, 0, 0));
-	OutResolution.x = ceil(sqrt(OutNodeCount));
-	OutResolution.y = OutResolution.x;
+	OutNodeCount = HairStrands.HairSampleCount[0];
+	OutResolution = GetHairSampleResolution(OutNodeCount);
 
 	const float2 ClipCoord = ((OutResolution + 0.5f) / float2(HairStrands.HairSampleViewportResolution)) * 2;
 
@@ -835,7 +834,7 @@ void MainPS(
 
 	const uint2 InCoord = uint2(SVPos.xy);
 	const uint LocalOffset = InCoord.x + InCoord.y * Resolution.x;
-	const uint MaxVisibilityNodeCount = HairStrands.HairSampleCount.Load(uint3(0, 0, 0));
+	const uint MaxVisibilityNodeCount = HairStrands.HairSampleCount[0];
 	if (LocalOffset >= NodeCount)
 	{
 		return;
diff --git a/Engine/Shaders/Private/HairStrands/HairStrandsForwardCulling.usf b/Engine/Shaders/Private/HairStrands/HairStrandsForwardCulling.usf
--- a/Engine/Shaders/Private/HairStrands/HairStrandsForwardCulling.usf
+++ b/Engine/Shaders/Private/HairStrands/HairStrandsForwardCulling.usf
@@ -29,7 +29,7 @@ uint				CurveCount;
 RWBuffer<float4>					RWPoints;
 RWStructuredBuffer<FPackedHairVis>	RWHairVis;
 RWBuffer<uint2>						RWCoord;
-RWTexture2D<uint>					RWCounter;
+RWStructuredBuffer<uint>			RWCounter;
 
 ///////////////////////////////////////////////////////////////////////////////////////////////////
 // Utils
@@ -158,7 +158,7 @@ void CSMain(uint DispatchThreadID : SV_DispatchThreadID, uint GroupThreadID : SV
 	// Alloc
 	if (GroupThreadID == 0)
 	{
-		InterlockedAdd(RWCounter[uint2(0,0)], s_ValidSegmentCount, s_WriteOffset);
+		InterlockedAdd(RWCounter[0], s_ValidSegmentCount, s_WriteOffset);
 	}
 	GroupMemoryBarrierWithGroupSync();
 
@@ -308,7 +308,7 @@ void CSMain(uint DispatchThreadID : SV_DispatchThreadID, uint GroupThreadID : SV
 	// Alloc
 	if (GroupThreadID == 0)
 	{
-		InterlockedAdd(RWCounter[uint2(0,0)], s_ValidSegmentCount, s_WriteOffset);
+		InterlockedAdd(RWCounter[0], s_ValidSegmentCount, s_WriteOffset);
 	}
 	GroupMemoryBarrierWithGroupSync();
 
diff --git a/Engine/Shaders/Private/HairStrands/HairStrandsForwardRaster.usf b/Engine/Shaders/Private/HairStrands/HairStrandsForwardRaster.usf
--- a/Engine/Shaders/Private/HairStrands/HairStrandsForwardRaster.usf
+++ b/Engine/Shaders/Private/HairStrands/HairStrandsForwardRaster.usf
@@ -214,12 +214,12 @@ float	RadiusAtDepth1;
 ///////////////////////////////////////////////////////////////////////////
 // Control points
 Buffer<float4> ControlPoints;
-Texture2D<uint> ControlPointCount;
+StructuredBuffer<uint> ControlPointCount;
 uint MaxControlPointCount;
 
 uint GetControlPointCount()
 {
-	return ControlPointCount.Load(uint3(0,0,0));
+	return ControlPointCount[0];
 }
 
 // Custom encoding for forward rasterizer. Position are 32bits float. This is temporary.
@@ -1189,9 +1189,7 @@ float4 UnpackVelocityData(uint2 In)
 
 float3 LoadSampleColor(uint InPrimId, uint2 InSampleResolution)
 {
-	uint2 SampleCoord = 0;
-	SampleCoord.x = InPrimId % InSampleResolution.x;
-	SampleCoord.y = InPrimId / InSampleResolution.x;
+	const uint2 SampleCoord = GetHairSampleCoord(InPrimId, InSampleResolution);
 	return SampleLightingTexture.Load(uint3(SampleCoord, 0)).xyz;
 }
 
@@ -1248,12 +1246,7 @@ void RasterCS(
 	uint   Thread_LoopCountToFullCoverage = 0;
 	uint2  Thread_PixelCoord = 0;
 
-	uint2 SampleLightingEffectiveResolution = 0;
-	{
-		const uint EffectivePointCount = ControlPointCount.Load(uint3(0, 0, 0));
-		SampleLightingEffectiveResolution.x = ceil(sqrt(EffectivePointCount));
-		SampleLightingEffectiveResolution.y = SampleLightingEffectiveResolution.x;
-	}
+	const uint2 SampleLightingEffectiveResolution = GetHairSampleResolution(ControlPointCount[0]);
 
 	LOOP
 	for (uint WorkIndex = 0; WorkIndex < MAX_WORK_COUNT; WorkIndex++)
diff --git a/Engine/Shaders/Private/HairStrands/HairStrandsLightSample.usf b/Engine/Shaders/Private/HairStrands/HairStrandsLightSample.usf
--- a/Engine/Shaders/Private/HairStrands/HairStrandsLightSample.usf
+++ b/Engine/Shaders/Private/HairStrands/HairStrandsLightSample.usf
@@ -1,13 +1,14 @@
 // Copyright Epic Games, Inc. All Rights Reserved.
 
 #include "../Common.ush"
+#include "HairStrandsVisibilityCommon.ush"
 
 ///////////////////////////////////////////////////////////////////////////
 
 #if SHADER_VERTEX
 
 int2 MaxViewportResolution;
-Texture2D<uint> HairNodeCountTexture;
+StructuredBuffer<uint> HairNodeCountBuffer;
 
 void MainVS(
 	in uint InVertexId : SV_VertexID,
@@ -15,9 +16,8 @@ void MainVS(
 	nointerpolation out uint OutNodeCount : DISPATCH_NODECOUNT,
 	nointerpolation out uint2 OutResolution : DISPATCH_RESOLUTION)
 {
-	OutNodeCount = HairNodeCountTexture.Load(uint3(0, 0, 0));
-	OutResolution.x = ceil(sqrt(OutNodeCount));
-	OutResolution.y = OutResolution.x;
+	OutNodeCount = HairNodeCountBuffer[0];
+	OutResolution = GetHairSampleResolution(OutNodeCount);
 
 	const float2 ClipCoord = ((OutResolution + 0.5f) / float2(MaxViewportResolution)) * 2;
 
diff --git a/Engine/Shaders/Private/HairStrands/HairStrandsMaterialVS.usf b/Engine/Shaders/Private/HairStrands/HairStrandsMaterialVS.usf
--- a/Engine/Shaders/Private/HairStrands/HairStrandsMaterialVS.usf
+++ b/Engine/Shaders/Private/HairStrands/HairStrandsMaterialVS.usf
@@ -15,9 +15,8 @@ void Main(
 	nointerpolation out uint  OutTotalNodeCount : DISPATCH_NODE_COUNT,
 	nointerpolation out uint  OutPrimitiveId : PRIMITIVE_ID)
 {
-	OutTotalNodeCount = MaterialPassParameters.TotalNodeCounter.Load(uint3(0,0,0));
-	OutResolution.x = ceil(sqrt(OutTotalNodeCount));
-	OutResolution.y = OutResolution.x;
+	OutTotalNodeCount = MaterialPassParameters.TotalNodeCounter[0];
+	OutResolution = GetHairSampleResolution(OutTotalNodeCount);
 
 	const float2 ClipCoord = ((OutResolution +0.5f) / float2(MaterialPassParameters.MaxResolution)) * 2;
 
diff --git a/Engine/Shaders/Private/HairStrands/HairStrandsVisibilityCommon.ush b/Engine/Shaders/Private/HairStrands/HairStrandsVisibilityCommon.ush
--- a/Engine/Shaders/Private/HairStrands/HairStrandsVisibilityCommon.ush
+++ b/Engine/Shaders/Private/HairStrands/HairStrandsVisibilityCommon.ush
@@ -305,6 +305,17 @@ uint GetPackedHairSampleCoverage8bit(in FPackedHairSample In)
 	return 0xFF & (In.Tangent_Coverage8bit>>24);
 }
 
+uint2 GetHairSampleResolution(uint InNodeCount)
+{
+	uint Resolution = ceil(sqrt(InNodeCount));
+	return uint2(Resolution, Resolution);
+}
+
+uint2 GetHairSampleCoord(uint InLinearCoord, uint2 InResolution)
+{
+	return uint2(InLinearCoord % InResolution.x, InLinearCoord / InResolution.x);
+}
+
 ////////////////////////////////////////////////////////////////////////////////////////////////////////////////////////////////
 // Hair velocity
 
diff --git a/Engine/Shaders/Private/HairStrands/HairStrandsVisibilityCompaction.usf b/Engine/Shaders/Private/HairStrands/HairStrandsVisibilityCompaction.usf
--- a/Engine/Shaders/Private/HairStrands/HairStrandsVisibilityCompaction.usf
+++ b/Engine/Shaders/Private/HairStrands/HairStrandsVisibilityCompaction.usf
@@ -25,7 +25,7 @@ float CosTangentThreshold;
 float CoverageThreshold; // Allow to scale the transmittance to compensate its asymptotic behavior 
 uint VelocityType;
 
-RWTexture2D<uint> OutCompactNodeCounter;
+RWStructuredBuffer<uint> OutCompactNodeCounter;
 RWTexture2D<uint> OutCompactNodeIndex;
 RWStructuredBuffer<FPackedHairVis> OutCompactNodeVis;
 RWBuffer<uint2> OutCompactNodeCoord;
@@ -136,7 +136,7 @@ void MainCS(uint GroupIndex : SV_GroupIndex, uint3 DispatchThreadId : SV_Dispatc
 	// Now allocate all the nodes for this group contiguously in memory
 	if (GroupIndex == 0 && AllocationNodeCount > 0)
 	{
-		InterlockedAdd(OutCompactNodeCounter[uint2(0, 0)], AllocationNodeCount, AllocationNodeOffset);
+		InterlockedAdd(OutCompactNodeCounter[0], AllocationNodeCount, AllocationNodeOffset);
 	}
 	GroupMemoryBarrierWithGroupSync();
 	uint OffsetInStorage = AllocationNodeOffset + PixelDataSubOffsetInStorage;
@@ -403,7 +403,7 @@ void MainCS(
 	GroupMemoryBarrierWithGroupSync();
 	if (GroupIndex == 0 && AllocationNodeCount > 0)
 	{
-		InterlockedAdd(OutCompactNodeCounter[uint2(0, 0)], AllocationNodeCount, AllocationNodeOffset);
+		InterlockedAdd(OutCompactNodeCounter[0], AllocationNodeCount, AllocationNodeOffset);
 	}
 	GroupMemoryBarrierWithGroupSync();
 
diff --git a/Engine/Shaders/Private/HairStrands/HairStrandsVisibilityCompactionComputeRaster.usf b/Engine/Shaders/Private/HairStrands/HairStrandsVisibilityCompactionComputeRaster.usf
--- a/Engine/Shaders/Private/HairStrands/HairStrandsVisibilityCompactionComputeRaster.usf
+++ b/Engine/Shaders/Private/HairStrands/HairStrandsVisibilityCompactionComputeRaster.usf
@@ -45,7 +45,7 @@ uint SamplerPerPixel;
 float CoverageThreshold;
 uint bSortSampleByDepth;
 
-RWTexture2D<uint> 					OutCompactNodeCounter;
+RWStructuredBuffer<uint> 			OutCompactNodeCounter;
 RWTexture2D<uint> 					OutCompactNodeIndex;
 RWStructuredBuffer<FPackedHairVis> 	OutCompactNodeVis;
 RWBuffer<uint2>						OutCompactNodeCoord;
@@ -204,7 +204,7 @@ void MainCS(uint GroupIndex : SV_GroupIndex, uint3 GroupId : SV_GroupID, uint3 G
 	GroupMemoryBarrierWithGroupSync();
 	if (GroupIndex == 0 && AllocationNodeCount > 0)
 	{
-		InterlockedAdd(OutCompactNodeCounter[uint2(0, 0)], AllocationNodeCount, AllocationNodeOffset);
+		InterlockedAdd(OutCompactNodeCounter[0], AllocationNodeCount, AllocationNodeOffset);
 	}
 	GroupMemoryBarrierWithGroupSync();
 
diff --git a/Engine/Shaders/Private/HairStrands/HairStrandsVisibilityCopyIndirectArg.usf b/Engine/Shaders/Private/HairStrands/HairStrandsVisibilityCopyIndirectArg.usf
--- a/Engine/Shaders/Private/HairStrands/HairStrandsVisibilityCopyIndirectArg.usf
+++ b/Engine/Shaders/Private/HairStrands/HairStrandsVisibilityCopyIndirectArg.usf
@@ -5,7 +5,7 @@
 #include "../Common.ush"
 #include "HairStrandsVisibilityCommon.ush"
 
-Texture2D<uint> CounterTexture;
+StructuredBuffer<uint> CounterBuffer;
 RWBuffer<uint> OutArgBuffer;
 uint ThreadGroupSize;
 uint ItemCountPerGroup;
@@ -13,7 +13,7 @@ uint ItemCountPerGroup;
 [numthreads(1, 1, 1)]
 void CopyCS(uint GroupIndex : SV_GroupIndex, uint3 DispatchThreadId : SV_DispatchThreadID)
 {
-	const uint NodeCount  = CounterTexture.Load(uint3(0,0,0)) * ItemCountPerGroup;
+	const uint NodeCount  = CounterBuffer[0] * ItemCountPerGroup;
 	const uint TotalGroup = (NodeCount + ThreadGroupSize-1) / ThreadGroupSize;
 	uint3 OutArgs;
 #if 0
diff --git a/Engine/Source/Runtime/Renderer/Private/HairStrands/HairStrandsData.h b/Engine/Source/Runtime/Renderer/Private/HairStrands/HairStrandsData.h
--- a/Engine/Source/Runtime/Renderer/Private/HairStrands/HairStrandsData.h
+++ b/Engine/Source/Runtime/Renderer/Private/HairStrands/HairStrandsData.h
@@ -40,7 +40,7 @@ BEGIN_GLOBAL_SHADER_PARAMETER_STRUCT(FHairStrandsViewUniformParameters, )
 	SHADER_PARAMETER_RDG_TEXTURE(Texture2D, HairOnlyDepthFurthestHZBTexture)			// HZB furthest depth texture containing only hair depth (not strongly typed for legacy shader code reason)
 	SHADER_PARAMETER_SAMPLER(SamplerState, HairOnlyDepthHZBSampler)						// HZB sampler
 	SHADER_PARAMETER_RDG_TEXTURE(Texture2D<uint>, HairSampleOffset)						// Offset & count, for accessing pixel's samples, based on screen pixel position
-	SHADER_PARAMETER_RDG_TEXTURE(Texture2D<uint>, HairSampleCount)						// Total count of hair sample, in sample space
+	SHADER_PARAMETER_RDG_BUFFER_SRV(StructuredBuffer<uint>, HairSampleCount)			// Total count of hair sample, in sample space
 	SHADER_PARAMETER_RDG_BUFFER_SRV(StructuredBuffer<FPackedHairSample>, HairSampleData)// Sample data (coverage, tangent, base color, ...), in sample space // HAIRSTRANDS_TODO: change this to be a uint4 so that we don't have to include the type for generated contant buffer
 	SHADER_PARAMETER_RDG_BUFFER_SRV(Buffer<uint2>, HairSampleCoords)					// Screen pixel coordinate of each sample, in sample space
 	SHADER_PARAMETER_RDG_BUFFER_SRV(Buffer<uint2>, HairTileData)						// Tile coords (RG16F)
@@ -112,7 +112,7 @@ struct FHairStrandsVisibilityData
 
 	uint32			MaxSampleCount = 8; // Sample count per pixel
 	uint32			MaxNodeCount = 0;	// Total sample count
-	FRDGTextureRef	NodeCount = nullptr;
+	FRDGBufferRef	NodeCount = nullptr;
 	FRDGTextureRef	NodeIndex = nullptr;
 	FRDGBufferRef	NodeData = nullptr;
 	FRDGBufferRef	NodeVisData = nullptr;
@@ -124,7 +124,7 @@ struct FHairStrandsVisibilityData
 	uint32				RasterizedInstanceCount = 0;
 	uint32				MaxControlPointCount = 0;
 	FRDGBufferSRVRef	ControlPointsSRV = nullptr;
-	FRDGTextureRef		ControlPointCount = nullptr;
+	FRDGBufferSRVRef	ControlPointCount = nullptr;
 	FRDGBufferSRVRef	ControlPointVelocitySRV = nullptr;
 
 	FHairStrandsTiles TileData;
diff --git a/Engine/Source/Runtime/Renderer/Private/HairStrands/HairStrandsForwardRaster.cpp b/Engine/Source/Runtime/Renderer/Private/HairStrands/HairStrandsForwardRaster.cpp
--- a/Engine/Source/Runtime/Renderer/Private/HairStrands/HairStrandsForwardRaster.cpp
+++ b/Engine/Source/Runtime/Renderer/Private/HairStrands/HairStrandsForwardRaster.cpp
@@ -49,7 +49,7 @@ class FHairCullSegmentCS : public FGlobalShader
 		SHADER_PARAMETER_RDG_TEXTURE(Texture2D, SceneDepthTexture)
 		SHADER_PARAMETER_RDG_BUFFER_UAV(RWStructuredBuffer<FPackedHairVis>, RWHairVis)
 		SHADER_PARAMETER_RDG_BUFFER_UAV(RWBuffer<uint2>, RWCoord)
-		SHADER_PARAMETER_RDG_TEXTURE_UAV(RWTexture2D<uint>, RWCounter)
+		SHADER_PARAMETER_RDG_BUFFER_UAV(RWStructuredBuffer<uint>, RWCounter)
 		SHADER_PARAMETER_RDG_BUFFER_UAV(RWBuffer<float4>, RWPoints)
 
 		END_SHADER_PARAMETER_STRUCT()
@@ -85,7 +85,7 @@ BEGIN_SHADER_PARAMETER_STRUCT(FRasterComputeForwardCommonParameters, )
 	SHADER_PARAMETER(uint32,    MaxControlPointCount)
 
 	SHADER_PARAMETER_RDG_BUFFER_SRV(Buffer, ControlPoints)
-	SHADER_PARAMETER_RDG_TEXTURE(Texture2D<uint>, ControlPointCount)
+	SHADER_PARAMETER_RDG_BUFFER_SRV(StructuredBuffer<uint>, ControlPointCount)
 END_SHADER_PARAMETER_STRUCT()
 
 ///////////////////////////////////////////////////////////////////////////////////////////////////
@@ -326,13 +326,13 @@ FRasterForwardCullingOutput AddHairStrandsForwardCullingPass(
 	Out.NodeCoord		= GraphBuilder.CreateBuffer(FRDGBufferDesc::CreateBufferDesc(sizeof(uint32), MaxNumPrimIDs), TEXT("Hair.VisibilityNodeCoord"));
 	Out.Points			= GraphBuilder.CreateBuffer(FRDGBufferDesc::CreateBufferDesc(sizeof(uint32)*4u, MaxNumPrimIDs), TEXT("Hair.ControlPoints"));
 	Out.PointsSRV		= GraphBuilder.CreateSRV(Out.Points, ControlPointFormat);
-	Out.PointCount		= GraphBuilder.CreateTexture(FRDGTextureDesc::Create2D(FIntPoint(1,1), PF_R32_UINT, FClearValueBinding::None, TexCreate_UAV | TexCreate_ShaderResource), TEXT("Hair.ControlPointCount"));
+	Out.PointCount		= GraphBuilder.CreateBuffer(FRDGBufferDesc::CreateStructuredDesc(sizeof(uint32), 1), TEXT("Hair.ControlPointCount"));
 	Out.RasterizedInstanceCount = 0;
 
 	FRDGBufferUAVRef PointsUAV = GraphBuilder.CreateUAV(Out.Points, ControlPointFormat);
 
 	const uint32 ClearValues[4] = { 0u,0u,0u,0u };
-	AddClearUAVPass(GraphBuilder, GraphBuilder.CreateUAV(Out.PointCount), ClearValues);
+	AddClearUAVPass(GraphBuilder, GraphBuilder.CreateUAV(Out.PointCount), 0u);
 	AddClearUAVPass(GraphBuilder, GraphBuilder.CreateUAV(Out.NodeIndex), ClearValues);
 
 	// Create and set the uniform buffer
diff --git a/Engine/Source/Runtime/Renderer/Private/HairStrands/HairStrandsForwardRaster.h b/Engine/Source/Runtime/Renderer/Private/HairStrands/HairStrandsForwardRaster.h
--- a/Engine/Source/Runtime/Renderer/Private/HairStrands/HairStrandsForwardRaster.h
+++ b/Engine/Source/Runtime/Renderer/Private/HairStrands/HairStrandsForwardRaster.h
@@ -21,7 +21,7 @@ struct FRasterForwardCullingOutput
 	FRDGTextureRef		NodeIndex = nullptr;
 	FRDGBufferRef		NodeCoord = nullptr;
 	FRDGBufferRef		NodeVis = nullptr;
-	FRDGTextureRef		PointCount = nullptr;
+	FRDGBufferRef		PointCount = nullptr;
 	FRDGBufferRef		Points = nullptr;
 	FRDGBufferSRVRef	PointsSRV = nullptr;
 };
diff --git a/Engine/Source/Runtime/Renderer/Private/HairStrands/HairStrandsRendering.cpp b/Engine/Source/Runtime/Renderer/Private/HairStrands/HairStrandsRendering.cpp
--- a/Engine/Source/Runtime/Renderer/Private/HairStrands/HairStrandsRendering.cpp
+++ b/Engine/Source/Runtime/Renderer/Private/HairStrands/HairStrandsRendering.cpp
@@ -25,7 +25,7 @@ static TRDGUniformBufferRef<FHairStrandsViewUniformParameters> InternalCreateHai
 		Parameters->HairSampleOffset = In->NodeIndex;
 		Parameters->HairSampleData = GraphBuilder.CreateSRV(In->NodeData);
 		Parameters->HairSampleCoords = GraphBuilder.CreateSRV(In->NodeCoord, FHairStrandsVisibilityData::NodeCoordFormat);
-		Parameters->HairSampleCount = In->NodeCount;
+		Parameters->HairSampleCount = GraphBuilder.CreateSRV(In->NodeCount);
 		Parameters->HairSampleViewportResolution = In->SampleLightingViewportResolution;
 		Parameters->MaxSamplePerPixelCount = In->MaxSampleCount;
 
@@ -62,7 +62,7 @@ static TRDGUniformBufferRef<FHairStrandsViewUniformParameters> InternalCreateHai
 		Parameters->HairOnlyDepthClosestHZBTexture = Parameters->HairOnlyDepthFurthestHZBTexture;
 		Parameters->HairOnlyDepthHZBSampler = TStaticSamplerState<SF_Point>::GetRHI();
 		Parameters->HairCoverageTexture = BlackTexture;
-		Parameters->HairSampleCount = ZeroR32_UINT;
+		Parameters->HairSampleCount = DummyBufferR32SRV;
 		Parameters->HairSampleOffset = ZeroR32_UINT;
 		Parameters->HairSampleCoords = DummyBufferRG16SRV;
 		Parameters->HairSampleData	 = GraphBuilder.CreateSRV(DummyNodeBuffer);
diff --git a/Engine/Source/Runtime/Renderer/Private/HairStrands/HairStrandsVisibility.cpp b/Engine/Source/Runtime/Renderer/Private/HairStrands/HairStrandsVisibility.cpp
--- a/Engine/Source/Runtime/Renderer/Private/HairStrands/HairStrandsVisibility.cpp
+++ b/Engine/Source/Runtime/Renderer/Private/HairStrands/HairStrandsVisibility.cpp
@@ -343,7 +343,7 @@ class FHairLightSampleClearVS : public FGlobalShader
 
 	BEGIN_SHADER_PARAMETER_STRUCT(FParameters, )
 		SHADER_PARAMETER(FIntPoint, MaxViewportResolution)
-		SHADER_PARAMETER_RDG_TEXTURE(Texture2D, HairNodeCountTexture)
+		SHADER_PARAMETER_RDG_BUFFER_SRV(StructuredBuffer<uint>, HairNodeCountBuffer)
 	END_SHADER_PARAMETER_STRUCT()
 
 	static bool ShouldCompilePermutation(const FGlobalShaderPermutationParameters& Parameters) { return IsHairStrandsSupported(EHairStrandsShaderType::Strands, Parameters.Platform); }
@@ -364,7 +364,7 @@ class FHairLightSampleClearPS : public FGlobalShader
 
 	BEGIN_SHADER_PARAMETER_STRUCT(FParameters, )
 		SHADER_PARAMETER(FIntPoint, MaxViewportResolution)
-		SHADER_PARAMETER_RDG_TEXTURE(Texture2D, HairNodeCountTexture)
+		SHADER_PARAMETER_RDG_BUFFER_SRV(StructuredBuffer<uint>, HairNodeCountBuffer)
 		RENDER_TARGET_BINDING_SLOTS()
 	END_SHADER_PARAMETER_STRUCT()
 
@@ -403,17 +403,20 @@ static FRDGTextureRef AddClearLightSamplePass(
 	FRDGBuilder& GraphBuilder,
 	const FViewInfo* View,
 	const uint32 MaxNodeCount,
-	const FRDGTextureRef NodeCounter)
+	const FRDGBufferSRVRef NodeCounterBuffer)
 {	
 	const EPixelFormat Format = FHairLightSampleClearPS::GetHairLightSampleFormat();
 
-	const uint32 SampleTextureResolution = FMath::CeilToInt(FMath::Sqrt(static_cast<float>(MaxNodeCount)));
+	// Compute the target texture resolution and round it up 128
+	uint32 SampleTextureResolution = FMath::CeilToInt(FMath::Sqrt(static_cast<float>(MaxNodeCount)));
+	SampleTextureResolution = FMath::DivideAndRoundUp(SampleTextureResolution, 128u) * 128u;
+
 	FRDGTextureDesc Desc = FRDGTextureDesc::Create2D(FIntPoint(SampleTextureResolution, SampleTextureResolution), Format, FClearValueBinding::Black, TexCreate_UAV | TexCreate_ShaderResource | TexCreate_RenderTargetable);
 	FRDGTextureRef Output = GraphBuilder.CreateTexture(Desc, TEXT("Hair.LightSample"));
 
 	FHairLightSampleClearPS::FParameters* ParametersPS = GraphBuilder.AllocParameters<FHairLightSampleClearPS::FParameters>();
 	ParametersPS->MaxViewportResolution = Desc.Extent;
-	ParametersPS->HairNodeCountTexture = NodeCounter;
+	ParametersPS->HairNodeCountBuffer = NodeCounterBuffer;
 	
 	FHairLightSampleClearPS::FPermutationDomain PermutationVector;
 	PermutationVector.Set<FHairLightSampleClearPS::FOutputFormat>(Format == PF_FloatR11G11B10 ? 1 : 0);
@@ -432,7 +435,7 @@ static FRDGTextureRef AddClearLightSamplePass(
 	{
 		FHairLightSampleClearVS::FParameters ParametersVS;
 		ParametersVS.MaxViewportResolution = ParametersPS->MaxViewportResolution;
-		ParametersVS.HairNodeCountTexture = ParametersPS->HairNodeCountTexture;
+		ParametersVS.HairNodeCountBuffer = ParametersPS->HairNodeCountBuffer;
 
 		FGraphicsPipelineStateInitializer GraphicsPSOInit;
 		RHICmdList.ApplyCachedRenderTargets(GraphicsPSOInit);
@@ -888,7 +891,7 @@ BEGIN_GLOBAL_SHADER_PARAMETER_STRUCT(FVisibilityMaterialPassUniformParameters, )
 	SHADER_PARAMETER(uint32, bUpdateSampleCoverage)
 	SHADER_PARAMETER(uint32, bInterpolationEnabled)
 	SHADER_PARAMETER_RDG_TEXTURE(Texture2D<uint>, NodeIndex)
-	SHADER_PARAMETER_RDG_TEXTURE(Texture2D<uint>, TotalNodeCounter)
+	SHADER_PARAMETER_RDG_BUFFER_SRV(StructuredBuffer<uint>, TotalNodeCounter)
 	SHADER_PARAMETER_RDG_BUFFER_SRV(Buffer<uint2>, NodeCoord)
 	SHADER_PARAMETER_RDG_BUFFER_SRV(StructuredBuffer<FPackedHairVis>, NodeVis)
 	SHADER_PARAMETER_RDG_BUFFER_SRV(Buffer<uint>, IndirectArgs)
@@ -978,7 +981,7 @@ static FMaterialPassOutput AddHairMaterialPass(
 	FRDGTextureRef CompactNodeIndex,
 	FRDGBufferRef CompactNodeVis,
 	FRDGBufferRef CompactNodeCoord,
-	FRDGTextureRef CompactNodeCounter,
+	FRDGBufferRef CompactNodeCounter,
 	FRDGBufferRef IndirectArgBuffer)
 {
 	if (!CompactNodeVis || !CompactNodeIndex)
@@ -990,10 +993,12 @@ static FMaterialPassOutput AddHairMaterialPass(
 	const EPixelFormat VelocityFormat = FVelocityRendering::GetFormat(ViewInfo->GetShaderPlatform());
 	check(VelocityFormat == PF_A16B16G16R16 || VelocityFormat == PF_G16R16);
 		
+	FRDGBufferSRVRef CompactNodeCounterSRV = GraphBuilder.CreateSRV(CompactNodeCounter);
+
 	FMaterialPassOutput Output;
 	Output.NodeData				 = GraphBuilder.CreateBuffer(FRDGBufferDesc::CreateStructuredDesc(sizeof(HairStrandsVisibilityInternal::NodeData), MaxNodeCount), TEXT("Hair.CompactNodeData"));
 	Output.NodeVelocity			 = GraphBuilder.CreateBuffer(FRDGBufferDesc::CreateBufferDesc(VelocityFormat == PF_G16R16 ? 4 : 8, CompactNodeVis->Desc.NumElements), TEXT("Hair.CompactNodeVelocity"));
-	Output.SampleLightingTexture = AddClearLightSamplePass(GraphBuilder, ViewInfo, MaxNodeCount, CompactNodeCounter);
+	Output.SampleLightingTexture = AddClearLightSamplePass(GraphBuilder, ViewInfo, MaxNodeCount, CompactNodeCounterSRV);
 	Output.NodeVelocitySRV		 = GraphBuilder.CreateSRV(Output.NodeVelocity, VelocityFormat);
 
 	const uint32 ResolutionDim = FMath::CeilToInt(FMath::Sqrt(static_cast<float>(MaxNodeCount)));
@@ -1032,7 +1037,7 @@ static FMaterialPassOutput AddHairMaterialPass(
 			UniformParameters->MaxResolution = Resolution;
 			UniformParameters->NodeGroupSize = NodeGroupSize;
 			UniformParameters->MaxSampleCount = MaxNodeCount;
-			UniformParameters->TotalNodeCounter = CompactNodeCounter;
+			UniformParameters->TotalNodeCounter = CompactNodeCounterSRV;
 			UniformParameters->NodeIndex = CompactNodeIndex;
 			UniformParameters->NodeVis = GraphBuilder.CreateSRV(CompactNodeVis);
 			UniformParameters->NodeCoord = GraphBuilder.CreateSRV(CompactNodeCoord, FHairStrandsVisibilityData::NodeCoordFormat);
@@ -1934,7 +1939,7 @@ class FCopyIndirectBufferCS : public FGlobalShader
 	BEGIN_SHADER_PARAMETER_STRUCT(FParameters, )
 		SHADER_PARAMETER(uint32, ThreadGroupSize)
 		SHADER_PARAMETER(uint32, ItemCountPerGroup)
-		SHADER_PARAMETER_RDG_TEXTURE(Texture2D, CounterTexture)
+		SHADER_PARAMETER_RDG_BUFFER_SRV(StructuredBuffer<uint>, CounterBuffer)
 		SHADER_PARAMETER_RDG_BUFFER_UAV(RWBuffer, OutArgBuffer)
 	END_SHADER_PARAMETER_STRUCT()
 
@@ -1949,16 +1954,16 @@ static FRDGBufferRef AddCopyIndirectArgPass(
 	const FViewInfo* View,
 	const uint32 ThreadGroupSize,
 	const uint32 ItemCountPerGroup,
-	FRDGTextureRef CounterTexture)
+	FRDGBufferSRVRef CounterBuffer)
 {
-	check(CounterTexture);
+	check(CounterBuffer);
 
 	FRDGBufferRef OutBuffer = GraphBuilder.CreateBuffer(FRDGBufferDesc::CreateIndirectDesc<FRHIDispatchIndirectParameters>(), TEXT("Hair.VisibilityIndirectArgBuffer"));
 
 	FCopyIndirectBufferCS::FParameters* Parameters = GraphBuilder.AllocParameters<FCopyIndirectBufferCS::FParameters>();
 	Parameters->ThreadGroupSize = ThreadGroupSize;
 	Parameters->ItemCountPerGroup = ItemCountPerGroup;
-	Parameters->CounterTexture = CounterTexture;
+	Parameters->CounterBuffer = CounterBuffer;
 	Parameters->OutArgBuffer = GraphBuilder.CreateUAV(OutBuffer);
 
 	TShaderMapRef<FCopyIndirectBufferCS> ComputeShader(View->ShaderMap);
@@ -2011,7 +2016,7 @@ class FHairVisibilityControlPointIdCompactionCS : public FGlobalShader
 
 		SHADER_PARAMETER_RDG_TEXTURE(Texture2D, SceneDepthTexture)
 
-		SHADER_PARAMETER_RDG_TEXTURE_UAV(RWTexture2D, OutCompactNodeCounter)
+		SHADER_PARAMETER_RDG_BUFFER_UAV(RWStructuredBuffer<uint>, OutCompactNodeCounter)
 		SHADER_PARAMETER_RDG_TEXTURE_UAV(RWTexture2D, OutCompactNodeIndex)
 		SHADER_PARAMETER_RDG_TEXTURE_UAV(RWTexture2D, OutCoverageTexture)
 		SHADER_PARAMETER_RDG_BUFFER_UAV(RWStructuredBuffer, OutCompactNodeVis)
@@ -2058,7 +2063,7 @@ static void AddHairVisibilityControlPointIdCompactionPass(
 	const uint32 NodeGroupSize,
 	const FHairStrandsTiles& TileData,
 	FHairVisibilityControlPointIdCompactionCS::FParameters* PassParameters,
-	FRDGTextureRef& OutCompactCounter,
+	FRDGBufferRef& OutCompactCounterBuffer,
 	FRDGTextureRef& OutCompactNodeIndex,
 	FRDGBufferRef& OutCompactNodeVis,
 	FRDGBufferRef& OutCompactNodeCoord,
@@ -2084,8 +2089,7 @@ static void AddHairVisibilityControlPointIdCompactionPass(
 	}
 
 	{
-		FRDGTextureDesc Desc = FRDGTextureDesc::Create2D(FIntPoint(1, 1), PF_R32_UINT, FClearValueBinding::None, TexCreate_UAV | TexCreate_ShaderResource | TexCreate_AtomicCompatible);
-		OutCompactCounter = GraphBuilder.CreateTexture(Desc, TEXT("Hair.VisibilityCompactCounter"));
+		OutCompactCounterBuffer = GraphBuilder.CreateBuffer(FRDGBufferDesc::CreateStructuredDesc(sizeof(uint32), 1), TEXT("Hair.VisibilityCompactCounter"));
 	}
 
 	{
@@ -2099,7 +2103,7 @@ static void AddHairVisibilityControlPointIdCompactionPass(
 	}
 
 	const uint32 ClearValues[4] = { 0u,0u,0u,0u };
-	AddClearUAVPass(GraphBuilder, GraphBuilder.CreateUAV(OutCompactCounter), ClearValues);
+	AddClearUAVPass(GraphBuilder, GraphBuilder.CreateUAV(OutCompactCounterBuffer), 0u);
 	AddClearGraphicPass(GraphBuilder, RDG_EVENT_NAME("HairStrands::NodeOffsetAndCount"), &View, 0, TileData, OutCompactNodeIndex);
 	AddClearGraphicPass(GraphBuilder, RDG_EVENT_NAME("HairStrands::CoverageTexture"), &View, 0, TileData, OutCoverageTexture);
 
@@ -2132,7 +2136,7 @@ static void AddHairVisibilityControlPointIdCompactionPass(
 	PassParameters->CosTangentThreshold = FMath::Cos(FMath::DegreesToRadians(FMath::Clamp(GHairStrandsMaterialCompactionTangentThreshold, 0.f, 90.f)));
 	PassParameters->SceneDepthTexture = SceneDepthTexture;
 	PassParameters->ViewUniformBuffer = View.ViewUniformBuffer;
-	PassParameters->OutCompactNodeCounter = GraphBuilder.CreateUAV(OutCompactCounter);
+	PassParameters->OutCompactNodeCounter = GraphBuilder.CreateUAV(OutCompactCounterBuffer);
 	PassParameters->OutCompactNodeIndex = GraphBuilder.CreateUAV(OutCompactNodeIndex);
 	PassParameters->OutCompactNodeVis = GraphBuilder.CreateUAV(OutCompactNodeVis);
 	PassParameters->OutCompactNodeCoord = GraphBuilder.CreateUAV(OutCompactNodeCoord, FHairStrandsVisibilityData::NodeCoordFormat);
@@ -2154,7 +2158,7 @@ static void AddHairVisibilityControlPointIdCompactionPass(
 		TileData.TileIndirectDispatchBuffer,
 		FHairStrandsTiles::GetIndirectDispatchArgOffset(TileType));
 
-	OutIndirectArgsBuffer = AddCopyIndirectArgPass(GraphBuilder, &View, NodeGroupSize, 1, OutCompactCounter);
+	OutIndirectArgsBuffer = AddCopyIndirectArgPass(GraphBuilder, &View, NodeGroupSize, 1, GraphBuilder.CreateSRV(OutCompactCounterBuffer));
 	OutMaxRenderNodeCount = MaxRenderNodeCount;
 }
 
@@ -2182,7 +2186,7 @@ class FHairVisibilityCompactionComputeRasterCS : public FGlobalShader
 		SHADER_PARAMETER_RDG_TEXTURE(Texture2D<uint>, PrimMatTexture)
 		SHADER_PARAMETER_RDG_TEXTURE(Texture2D<uint>, HairCountTexture)
 
-		SHADER_PARAMETER_RDG_TEXTURE_UAV(RWTexture2D, OutCompactNodeCounter)
+		SHADER_PARAMETER_RDG_BUFFER_UAV(RWStructuredBuffer<uint>, OutCompactNodeCounter)
 		SHADER_PARAMETER_RDG_TEXTURE_UAV(RWTexture2D, OutCompactNodeIndex)
 		SHADER_PARAMETER_RDG_TEXTURE_UAV(RWTexture2D, OutCoverageTexture)
 		SHADER_PARAMETER_RDG_BUFFER_UAV(RWStructuredBuffer, OutCompactNodeVis)
@@ -2218,7 +2222,7 @@ static void AddHairVisibilityCompactionComputeRasterPass(
 	const uint32 SamplesPerPixel,
 	const FRasterComputeOutput& RasterComputeData,
 	const FHairStrandsTiles& TileData,
-	FRDGTextureRef& OutCompactCounter,
+	FRDGBufferRef& OutCompactCounter,
 	FRDGTextureRef& OutCompactNodeIndex,
 	FRDGBufferRef&  OutCompactNodeVis,
 	FRDGBufferRef&  OutCompactNodeCoord,
@@ -2229,8 +2233,7 @@ static void AddHairVisibilityCompactionComputeRasterPass(
 	FIntPoint Resolution = RasterComputeData.DepthCovTexture->Desc.Extent;
 
 	{
-		FRDGTextureDesc Desc = FRDGTextureDesc::Create2D(FIntPoint(1,1), PF_R32_UINT, FClearValueBinding::None, TexCreate_UAV | TexCreate_AtomicCompatible);
-		OutCompactCounter = GraphBuilder.CreateTexture(Desc, TEXT("Hair.VisibilityCompactCounter"));
+		OutCompactCounter = GraphBuilder.CreateBuffer(FRDGBufferDesc::CreateStructuredDesc(sizeof(uint32),1), TEXT("Hair.VisibilityCompactCounter"));
 	}
 
 	{
@@ -2244,7 +2247,7 @@ static void AddHairVisibilityCompactionComputeRasterPass(
 	}
 
 	const uint32 ClearValues[4] = { 0u,0u,0u,0u };
-	AddClearUAVPass(GraphBuilder, GraphBuilder.CreateUAV(OutCompactCounter), ClearValues);
+	AddClearUAVPass(GraphBuilder, GraphBuilder.CreateUAV(OutCompactCounter), 0u);
 	AddClearUAVPass(GraphBuilder, GraphBuilder.CreateUAV(OutCompactNodeIndex), ClearValues);
 	AddClearUAVPass(GraphBuilder, GraphBuilder.CreateUAV(OutCoverageTexture), 0.f);
 
@@ -2296,7 +2299,7 @@ static void AddHairVisibilityCompactionComputeRasterPass(
 		TileData.TileIndirectDispatchBuffer,
 		FHairStrandsTiles::GetIndirectDispatchArgOffset(TileType));
 
-	OutIndirectArgsBuffer = AddCopyIndirectArgPass(GraphBuilder, &View, NodeGroupSize, 1, OutCompactCounter);
+	OutIndirectArgsBuffer = AddCopyIndirectArgPass(GraphBuilder, &View, NodeGroupSize, 1, GraphBuilder.CreateSRV(OutCompactCounter));
 	OutMaxRenderNodeCount = MaxRenderNodeCount;
 }
 
@@ -4402,7 +4405,7 @@ void DrawHitProxies(
 
 	const uint32 NodeGroupSize = GetVendorOptimalGroupSize1D();
 	const uint32 MaxSampleCount = 4;
-	FRDGTextureRef NodeCounter = nullptr;
+	FRDGBufferRef  NodeCounter = nullptr;
 	FRDGTextureRef VisNodeIndex = nullptr;
 	FRDGBufferRef  VisNodeData = nullptr;
 	FRDGBufferRef  VisNodeCoord = nullptr;
@@ -4641,7 +4644,7 @@ void RenderHairStrandsVisibilityBuffer(
 			FRDGTextureRef CompactNodeIndex = nullptr;
 			FRDGBufferRef  CompactNodeData = nullptr;
 			FRDGBufferRef  CompactNodeVis = nullptr;
-			FRDGTextureRef NodeCounter = nullptr;
+			FRDGBufferRef  NodeCounter = nullptr;
 
 			if (RenderMode == HairVisibilityRenderMode_ComputeRasterForward)
 			{
@@ -4662,7 +4665,7 @@ void RenderHairStrandsVisibilityBuffer(
 
 				VisibilityData.MaxControlPointCount = Out.NodeCoord->Desc.NumElements;
 				VisibilityData.ControlPointsSRV = Out.PointsSRV;
-				VisibilityData.ControlPointCount = Out.PointCount;
+				VisibilityData.ControlPointCount = GraphBuilder.CreateSRV(Out.PointCount);
 				VisibilityData.RasterizedInstanceCount = Out.RasterizedInstanceCount;
 
 				CompactNodeIndex = Out.NodeIndex;
@@ -4670,7 +4673,7 @@ void RenderHairStrandsVisibilityBuffer(
 				CompactNodeVis = Out.NodeVis;
 				FRDGBufferRef CompactNodeCoord = Out.NodeCoord;
 
-				FRDGBufferRef IndirectArgsBuffer = AddCopyIndirectArgPass(GraphBuilder, &View, VisibilityData.NodeGroupSize, 1, Out.PointCount);
+				FRDGBufferRef IndirectArgsBuffer = AddCopyIndirectArgPass(GraphBuilder, &View, VisibilityData.NodeGroupSize, 1, GraphBuilder.CreateSRV(Out.PointCount));
 
 				// Generate Tile data
 				VisibilityData.TileData = AddHairStrandsGenerateTilesPass(GraphBuilder, View,Out.Resolution); // Tile for all pixels == no tile accelaration
