Fix unity build error from duplicate function definition

diff --git a/Engine/Source/Runtime/Renderer/Private/RayTracing/RayTracingDecals.cpp b/Engine/Source/Runtime/Renderer/Private/RayTracing/RayTracingDecals.cpp
--- a/Engine/Source/Runtime/Renderer/Private/RayTracing/RayTracingDecals.cpp
+++ b/Engine/Source/Runtime/Renderer/Private/RayTracing/RayTracingDecals.cpp
@@ -230,7 +230,7 @@ IMPLEMENT_RT_PAYLOAD_TYPE(ERayTracingPayloadType::Decals, 48);
 
 IMPLEMENT_SHADER_TYPE(, FRayTracingDecalMaterialShader, TEXT("/Engine/Private/RayTracing/RayTracingDecalMaterialShader.usf"), TEXT("RayTracingDecalMaterialShader"), SF_RayCallable);
 
-static bool NeedsAnyHitShader(EBlendMode BlendMode)
+static bool DecalNeedsAnyHitShader(EBlendMode BlendMode)
 {
 	return BlendMode != BLEND_Opaque;
 }
@@ -267,7 +267,7 @@ public:
 		{
 			return false;
 		}
-		if (NeedsAnyHitShader(Parameters.MaterialParameters.BlendMode) != bUseAnyHitShader)
+		if (DecalNeedsAnyHitShader(Parameters.MaterialParameters.BlendMode) != bUseAnyHitShader)
 		{
 			// the anyhit permutation is only required if the material is masked or has a non-opaque blend mode
 			return false;
@@ -324,7 +324,7 @@ IMPLEMENT_MATERIAL_SHADER_TYPE(template<>, FRayTracingDecalMaterialCHS_AHS, TEXT
 
 FShaderType* GetRayTracingDecalMaterialShaderType(EBlendMode BlendMode)
 {
-	if (NeedsAnyHitShader(BlendMode))
+	if (DecalNeedsAnyHitShader(BlendMode))
 	{
 		return &FRayTracingDecalMaterialCHS_AHS::GetStaticType();
 	}
