Fix bloom when local exposure is enabled but bloom threshold is not used
- bloom threshold and local exposure are applied in the same function (BloomSetupCommon(...)) which caused an issue when local exposure is enabled but PPV is not using threshold however some threshold logic was still being applied.
- prevent issues by using different permutations of the shader depending on which features are enabled (USE_LOCAL_EXPOSURE / USE_THRESHOLD)

diff --git a/Engine/Shaders/Private/PostProcessBloom.usf b/Engine/Shaders/Private/PostProcessBloom.usf
--- a/Engine/Shaders/Private/PostProcessBloom.usf
+++ b/Engine/Shaders/Private/PostProcessBloom.usf
@@ -55,9 +55,13 @@ float4 BloomSetupCommon(float2 UV, float2 ViewportUV, float2 ExposureScaleMiddle
 
 	half3 LinearColor = SceneColor.rgb;
 
+#if USE_THRESHOLD
 	half TotalLuminance = Luminance(LinearColor) * ExposureScale;
 	half BloomLuminance = TotalLuminance - BloomThreshold;
 	half BloomAmount = saturate(BloomLuminance * 0.5f);
+#else
+	half BloomAmount = 1.0f;
+#endif
 
 	return float4(BloomAmount * LinearColor, 0) * View.PreExposure;
 }
diff --git a/Engine/Source/Runtime/Renderer/Private/PostProcess/PostProcessBloomSetup.cpp b/Engine/Source/Runtime/Renderer/Private/PostProcess/PostProcessBloomSetup.cpp
--- a/Engine/Source/Runtime/Renderer/Private/PostProcess/PostProcessBloomSetup.cpp
+++ b/Engine/Source/Runtime/Renderer/Private/PostProcess/PostProcessBloomSetup.cpp
@@ -69,7 +69,8 @@ public:
 	SHADER_USE_PARAMETER_STRUCT(FBloomSetupPS, FGlobalShader);
 
 	class FLocalExposureDim : SHADER_PERMUTATION_BOOL("USE_LOCAL_EXPOSURE");
-	using FPermutationDomain = TShaderPermutationDomain<FLocalExposureDim>;
+	class FThresholdDim : SHADER_PERMUTATION_BOOL("USE_THRESHOLD");
+	using FPermutationDomain = TShaderPermutationDomain<FLocalExposureDim, FThresholdDim>;
 
 	BEGIN_SHADER_PARAMETER_STRUCT(FParameters, )
 		SHADER_PARAMETER_STRUCT_INCLUDE(FBloomSetupParameters, BloomSetup)
@@ -92,7 +93,8 @@ public:
 	SHADER_USE_PARAMETER_STRUCT(FBloomSetupCS, FGlobalShader);
 
 	class FLocalExposureDim : SHADER_PERMUTATION_BOOL("USE_LOCAL_EXPOSURE");
-	using FPermutationDomain = TShaderPermutationDomain<FLocalExposureDim>;
+	class FThresholdDim : SHADER_PERMUTATION_BOOL("USE_THRESHOLD");
+	using FPermutationDomain = TShaderPermutationDomain<FLocalExposureDim, FThresholdDim>;
 
 	BEGIN_SHADER_PARAMETER_STRUCT(FParameters, )
 		SHADER_PARAMETER_STRUCT_INCLUDE(FBloomSetupParameters, BloomSetup)
@@ -119,10 +121,12 @@ FScreenPassTexture AddBloomSetupPass(FRDGBuilder& GraphBuilder, const FViewInfo&
 {
 	check(Inputs.SceneColor.IsValid());
 	check(Inputs.EyeAdaptationBuffer);
-	check(Inputs.Threshold > -1.0f || Inputs.EyeAdaptationParameters != nullptr);
 
 	const bool bIsComputePass = View.bUseComputePasses;
 	const bool bLocalExposureEnabled = Inputs.LocalExposureTexture != nullptr;
+	const bool bThresholdEnabled = Inputs.Threshold > -1.0f;
+
+	check(bLocalExposureEnabled || bThresholdEnabled);
 
 	FRDGTextureDesc OutputDesc = Inputs.SceneColor.Texture->Desc;
 	OutputDesc.Reset();
@@ -139,6 +143,7 @@ FScreenPassTexture AddBloomSetupPass(FRDGBuilder& GraphBuilder, const FViewInfo&
 
 		FBloomSetupCS::FPermutationDomain PermutationVector;
 		PermutationVector.Set<FBloomSetupCS::FLocalExposureDim>(bLocalExposureEnabled);
+		PermutationVector.Set<FBloomSetupCS::FThresholdDim>(bThresholdEnabled);
 
 		auto ComputeShader = View.ShaderMap->GetShader<FBloomSetupCS>(PermutationVector);
 
@@ -160,6 +165,7 @@ FScreenPassTexture AddBloomSetupPass(FRDGBuilder& GraphBuilder, const FViewInfo&
 
 		FBloomSetupPS::FPermutationDomain PermutationVector;
 		PermutationVector.Set<FBloomSetupPS::FLocalExposureDim>(bLocalExposureEnabled);
+		PermutationVector.Set<FBloomSetupPS::FThresholdDim>(bThresholdEnabled);
 
 		auto PixelShader = View.ShaderMap->GetShader<FBloomSetupPS>(PermutationVector);
 
