Only allocate virtual texture physical texture when there are allocated virtual textures using the physical space.
This removes a common case of unnecessary memory allocation when a streaming virtual texture has been loaded but is unused.

diff --git a/Engine/Source/Runtime/Renderer/Private/VT/AllocatedVirtualTexture.cpp b/Engine/Source/Runtime/Renderer/Private/VT/AllocatedVirtualTexture.cpp
--- a/Engine/Source/Runtime/Renderer/Private/VT/AllocatedVirtualTexture.cpp
+++ b/Engine/Source/Runtime/Renderer/Private/VT/AllocatedVirtualTexture.cpp
@@ -8,7 +8,8 @@
 #include "VT/VirtualTexturePhysicalSpace.h"
 #include "Misc/StringBuilder.h"
 
-FAllocatedVirtualTexture::FAllocatedVirtualTexture(FRHICommandListBase& RHICmdList,
+FAllocatedVirtualTexture::FAllocatedVirtualTexture(
+	FRHICommandListBase& RHICmdList,
 	FVirtualTextureSystem* InSystem,
 	uint32 InFrame,
 	const FAllocatedVTDescription& InDesc,
@@ -40,7 +41,7 @@ FAllocatedVirtualTexture::FAllocatedVirtualTexture(FRHICommandListBase& RHICmdLi
 				ProducerPhysicalGroupIndex = Producer->GetPhysicalGroupIndexForTextureLayer(ProducerLayerIndex);
 				PhysicalSpace = Producer->GetPhysicalSpaceForPhysicalGroup(ProducerPhysicalGroupIndex);
 			}
-			const uint32 UniquePhysicalSpaceIndex = AddUniquePhysicalSpace(PhysicalSpace, UniqueProducerIndex, ProducerPhysicalGroupIndex);
+			const uint32 UniquePhysicalSpaceIndex = AddUniquePhysicalSpace(RHICmdList, PhysicalSpace, UniqueProducerIndex, ProducerPhysicalGroupIndex);
 			UniquePageTableLayers[UniquePhysicalSpaceIndex].ProducerTextureLayerMask |= 1 << ProducerLayerIndex;
 			const uint8 PageTableLayerLocalIndex = UniquePageTableLayers[UniquePhysicalSpaceIndex].TextureLayerCount++;
 			
@@ -140,6 +141,10 @@ void FAllocatedVirtualTexture::Destroy(FVirtualTextureSystem* System)
 
 		for (int32 PageTableIndex = 0u; PageTableIndex < UniquePageTableLayers.Num(); ++PageTableIndex)
 		{
+			if (UniquePageTableLayers[PageTableIndex].PhysicalSpace->ReleaseResourceRef() == 0)
+			{
+				UniquePageTableLayers[PageTableIndex].PhysicalSpace->ReleaseResource();
+			}
 			UniquePageTableLayers[PageTableIndex].PhysicalSpace.SafeRelease();
 		}
 
@@ -360,7 +365,7 @@ uint32 FAllocatedVirtualTexture::AddUniqueProducer(FVirtualTextureProducerHandle
 	return Index;
 }
 
-uint32 FAllocatedVirtualTexture::AddUniquePhysicalSpace(FVirtualTexturePhysicalSpace* InPhysicalSpace, uint32 InUniqueProducerIndex, uint32 InProducerPhysicalSpaceIndex)
+uint32 FAllocatedVirtualTexture::AddUniquePhysicalSpace(FRHICommandListBase& InRHICmdList, FVirtualTexturePhysicalSpace* InPhysicalSpace, uint32 InUniqueProducerIndex, uint32 InProducerPhysicalSpaceIndex)
 {
 	if (Description.bShareDuplicateLayers)
 	{
@@ -383,6 +388,11 @@ uint32 FAllocatedVirtualTexture::AddUniquePhysicalSpace(FVirtualTexturePhysicalS
 	UniquePageTableLayers[Index].ProducerTextureLayerMask = 0;
 	UniquePageTableLayers[Index].TextureLayerCount = 0;
 
+	if (InPhysicalSpace->AddResourceRef() == 1)
+	{
+		InPhysicalSpace->InitResource(InRHICmdList);
+	}
+
 	return Index;
 }
 
diff --git a/Engine/Source/Runtime/Renderer/Private/VT/AllocatedVirtualTexture.h b/Engine/Source/Runtime/Renderer/Private/VT/AllocatedVirtualTexture.h
--- a/Engine/Source/Runtime/Renderer/Private/VT/AllocatedVirtualTexture.h
+++ b/Engine/Source/Runtime/Renderer/Private/VT/AllocatedVirtualTexture.h
@@ -66,7 +66,7 @@ public:
 private:
 	uint32 CalculatePersistentHash(FAllocatedVTDescription const& InDesc, FVirtualTextureProducer* const* InProducers) const;
 	uint32 AddUniqueProducer(FVirtualTextureProducerHandle const& InHandle, const FVirtualTextureProducer* InProducer);
-	uint32 AddUniquePhysicalSpace(FVirtualTexturePhysicalSpace* InPhysicalSpace, uint32 InUniqueProducerIndex, uint32 InProducerPhysicalSpaceIndex);
+	uint32 AddUniquePhysicalSpace(FRHICommandListBase& RHICmdList, FVirtualTexturePhysicalSpace* InPhysicalSpace, uint32 InUniqueProducerIndex, uint32 InProducerPhysicalSpaceIndex);
 
 	uint32 FrameAllocated;
 	uint32 PersistentHash;
diff --git a/Engine/Source/Runtime/Renderer/Private/VT/VirtualTexturePhysicalSpace.cpp b/Engine/Source/Runtime/Renderer/Private/VT/VirtualTexturePhysicalSpace.cpp
--- a/Engine/Source/Runtime/Renderer/Private/VT/VirtualTexturePhysicalSpace.cpp
+++ b/Engine/Source/Runtime/Renderer/Private/VT/VirtualTexturePhysicalSpace.cpp
@@ -14,6 +14,8 @@
 #include "RHIUtilities.h"
 #include "RenderGraphBuilder.h"
 
+DECLARE_MEMORY_STAT_POOL(TEXT("Total Physical Memory"), STAT_TotalPhysicalMemory, STATGROUP_VirtualTextureMemory, FPlatformMemory::MCR_GPU);
+
 CSV_DECLARE_CATEGORY_EXTERN(VirtualTexturing);
 
 static TAutoConsoleVariable<float> CVarVTResidencyMaxMipMapBias(
@@ -58,6 +60,7 @@ FVirtualTexturePhysicalSpace::FVirtualTexturePhysicalSpace(uint16 InID, const FV
 	, DescriptionExt(InDescExt)
 	, ID(InID)
 	, NumRefs(0u)
+	, NumResourceRefs(0u)
 	, ResidencyMipMapBias(0.0f)
 	, LastFrameOversubscribed(0)
 #if !UE_BUILD_SHIPPING
@@ -161,10 +164,14 @@ void FVirtualTexturePhysicalSpace::InitRHI(FRHICommandListBase& RHICmdList)
 		SRVCreateInfo.SRGBOverride = SRGBO_Default;
 		TextureSRV_SRGB[Layer] = RHICmdList.CreateShaderResourceView(TextureRHI, SRVCreateInfo);
 	}
+
+	INC_MEMORY_STAT_BY(STAT_TotalPhysicalMemory, GetSizeInBytes());
 }
 
 void FVirtualTexturePhysicalSpace::ReleaseRHI()
 {
+	DEC_MEMORY_STAT_BY(STAT_TotalPhysicalMemory, GetSizeInBytes());
+
 	for (int32 Layer = 0; Layer < Description.NumLayers; ++Layer)
 	{
 		GRenderTargetPool.FreeUnusedResource(PooledRenderTarget[Layer]);
diff --git a/Engine/Source/Runtime/Renderer/Private/VT/VirtualTexturePhysicalSpace.h b/Engine/Source/Runtime/Renderer/Private/VT/VirtualTexturePhysicalSpace.h
--- a/Engine/Source/Runtime/Renderer/Private/VT/VirtualTexturePhysicalSpace.h
+++ b/Engine/Source/Runtime/Renderer/Private/VT/VirtualTexturePhysicalSpace.h
@@ -107,10 +107,18 @@ public:
 	virtual void InitRHI(FRHICommandListBase& RHICmdList) override;
 	virtual void ReleaseRHI() override;
 
+	// RefCount for TRefCountPtr management of physical space.
+	// Both producers and allocated VT keep references to keep a physical space alive.
 	inline uint32 AddRef() { return ++NumRefs; }
 	inline uint32 Release() { check(NumRefs > 0u); return --NumRefs; }
 	inline uint32 GetRefCount() const { return NumRefs; }
 
+	// RefCount for tracking allocated VT references only.
+	// We only allocate underlying physical texture when referenced from an allocated VT.
+	inline uint32 AddResourceRef() { return ++NumResourceRefs; }
+	inline uint32 ReleaseResourceRef() { check(NumResourceRefs > 0u); return --NumResourceRefs; }
+	inline uint32 GetResourceRefCount() const { return NumResourceRefs; }
+
 	FRHITexture* GetPhysicalTexture(int32 Layer) const
 	{
 		check(PooledRenderTarget[Layer].IsValid());
@@ -152,6 +160,7 @@ private:
 
 	uint16 ID;
 	uint32 NumRefs;
+	uint32 NumResourceRefs;
 	
 	float ResidencyMipMapBias;
 	uint32 LastFrameOversubscribed;
diff --git a/Engine/Source/Runtime/Renderer/Private/VT/VirtualTextureSystem.cpp b/Engine/Source/Runtime/Renderer/Private/VT/VirtualTextureSystem.cpp
--- a/Engine/Source/Runtime/Renderer/Private/VT/VirtualTextureSystem.cpp
+++ b/Engine/Source/Runtime/Renderer/Private/VT/VirtualTextureSystem.cpp
@@ -67,7 +67,6 @@ DECLARE_DWORD_COUNTER_STAT(TEXT("Num stacks produced"), STAT_NumStacksProduced,
 
 DECLARE_DWORD_COUNTER_STAT(TEXT("Num flush caches"), STAT_NumFlushCache, STATGROUP_VirtualTexturing);
 
-DECLARE_MEMORY_STAT_POOL(TEXT("Total Physical Memory"), STAT_TotalPhysicalMemory, STATGROUP_VirtualTextureMemory, FPlatformMemory::MCR_GPU);
 DECLARE_MEMORY_STAT_POOL(TEXT("Total Pagetable Memory"), STAT_TotalPagetableMemory, STATGROUP_VirtualTextureMemory, FPlatformMemory::MCR_GPU);
 
 DECLARE_GPU_STAT(VirtualTexture);
@@ -379,7 +378,6 @@ FVirtualTextureSystem::~FVirtualTextureSystem()
 		if (PhysicalSpace)
 		{
 			check(PhysicalSpace->GetRefCount() == 0u);
-			DEC_MEMORY_STAT_BY(STAT_TotalPhysicalMemory, PhysicalSpace->GetSizeInBytes());
 			BeginReleaseResource(PhysicalSpace);
 		}
 	}
@@ -473,7 +471,7 @@ void FVirtualTextureSystem::ListPhysicalPoolsFromConsole()
 	uint64 TotalPhysicalMemory = 0u;
 	for(int32 i = 0; i < PhysicalSpaces.Num(); ++i)
 	{
-		if (PhysicalSpaces[i])
+		if (PhysicalSpaces[i] && PhysicalSpaces[i]->IsInitialized())
 		{
 			const FVirtualTexturePhysicalSpace& PhysicalSpace = *PhysicalSpaces[i];
 			const FVTPhysicalSpaceDescription& Desc = PhysicalSpace.GetDescription();
@@ -1046,9 +1044,6 @@ FVirtualTexturePhysicalSpace* FVirtualTextureSystem::AcquirePhysicalSpace(FRHICo
 	FVirtualTexturePhysicalSpace* PhysicalSpace = new FVirtualTexturePhysicalSpace(ID, InDesc, DescExt);
 	PhysicalSpaces[ID] = PhysicalSpace;
 
-	INC_MEMORY_STAT_BY(STAT_TotalPhysicalMemory, PhysicalSpace->GetSizeInBytes());
-	PhysicalSpace->InitResource(RHICmdList);
-
 	return PhysicalSpace;
 }
 
@@ -1061,8 +1056,6 @@ void FVirtualTextureSystem::ReleasePendingSpaces()
 		FVirtualTexturePhysicalSpace* PhysicalSpace = PhysicalSpaces[Id];
 		if ((bool)PhysicalSpace && PhysicalSpace->GetRefCount() == 0u)
 		{
-			DEC_MEMORY_STAT_BY(STAT_TotalPhysicalMemory, PhysicalSpace->GetSizeInBytes());
-
 			const FTexturePagePool& PagePool = PhysicalSpace->GetPagePool();
 			check(PagePool.GetNumMappedPages() == 0u);
 			check(PagePool.GetNumLockedPages() == 0u);
@@ -2995,7 +2988,7 @@ void FVirtualTextureSystem::DrawResidencyHud(UCanvas* InCanvas, APlayerControlle
 	for (int32 SpaceIndex = 0; SpaceIndex < PhysicalSpaces.Num(); ++SpaceIndex)
 	{
 		FVirtualTexturePhysicalSpace* PhysicalSpace = PhysicalSpaces[SpaceIndex];
-		if (PhysicalSpace)
+		if (PhysicalSpace && PhysicalSpace->IsInitialized())
 		{
 			int32 GraphX = GraphIndex % NumGraphsInRow;
 			int32 GraphY = GraphIndex / NumGraphsInRow;
