- Disable low latency on none surface domain materials
- This fixes an issue with volume fog sprites which is rendered before simulations are enabled if they use things like depth buffer reads

diff --git a/Engine/Plugins/FX/Niagara/Source/Niagara/Private/NiagaraRendererSprites.cpp b/Engine/Plugins/FX/Niagara/Source/Niagara/Private/NiagaraRendererSprites.cpp
--- a/Engine/Plugins/FX/Niagara/Source/Niagara/Private/NiagaraRendererSprites.cpp
+++ b/Engine/Plugins/FX/Niagara/Source/Niagara/Private/NiagaraRendererSprites.cpp
@@ -252,6 +252,7 @@ void FNiagaraRendererSprites::PrepareParticleSpriteRenderData(FParticleSpriteRen
 	bool bLowLatencyTranslucencyEnabled =
 		ParticleSpriteRenderData.bHasTranslucentMaterials &&
 		bGpuLowLatencyTranslucency &&
+		Material.GetMaterialDomain() == MD_Surface &&
 		GpuReadyTickStage >= CurrentParticleData->GetGPUDataReadyStage() &&
 		!SceneProxy->CastsVolumetricTranslucentShadow() &&
 		ViewFamilySupportLowLatencyTranslucency(ViewFamily);
diff --git a/Engine/Source/Runtime/Renderer/Private/FXRenderingUtils.cpp b/Engine/Source/Runtime/Renderer/Private/FXRenderingUtils.cpp
--- a/Engine/Source/Runtime/Renderer/Private/FXRenderingUtils.cpp
+++ b/Engine/Source/Runtime/Renderer/Private/FXRenderingUtils.cpp
@@ -27,8 +27,8 @@ bool UE::FXRenderingUtils::CanMaterialRenderBeforeFXPostOpaque(
 	const FPrimitiveSceneProxy& SceneProxy,
 	const FMaterial& Material)
 {
-	// Opaque materials & translucent that write custom depth will always need to render before FFXSystemInterface::PostOpaqueRender
-	if (!IsTranslucentBlendMode(Material) || (SceneProxy.ShouldRenderCustomDepth() && Material.IsTranslucencyWritingCustomDepth()))
+	// Opaque materials, none surface materials & translucent that write custom depth will always need to render before FFXSystemInterface::PostOpaqueRender
+	if (!IsTranslucentBlendMode(Material) || Material.GetMaterialDomain() != MD_Surface || (SceneProxy.ShouldRenderCustomDepth() && Material.IsTranslucencyWritingCustomDepth()))
 	{
 		return true;
 	}
