Use compute path for IES atlas updates on mobile platforms when possible

diff --git a/Engine/Plugins/Experimental/GPULightmass/Source/GPULightmass/Private/LightmapRenderer.cpp b/Engine/Plugins/Experimental/GPULightmass/Source/GPULightmass/Private/LightmapRenderer.cpp
--- a/Engine/Plugins/Experimental/GPULightmass/Source/GPULightmass/Private/LightmapRenderer.cpp
+++ b/Engine/Plugins/Experimental/GPULightmass/Source/GPULightmass/Private/LightmapRenderer.cpp
@@ -2051,7 +2051,7 @@ void FLightmapRenderer::Finalize(FRDGBuilder& GraphBuilder)
 	}
 
 	RectLightAtlas::UpdateAtlasTexture(GraphBuilder, Scene->FeatureLevel);
-	IESAtlas::UpdateAtlasTexture(GraphBuilder, Scene->FeatureLevel);
+	IESAtlas::UpdateAtlasTexture(GraphBuilder, GetFeatureLevelShaderPlatform(Scene->FeatureLevel));
 
 	FSceneUniformBuffer SceneUniforms;
 
diff --git a/Engine/Plugins/Experimental/GPULightmass/Source/GPULightmass/Private/VolumetricLightmap.cpp b/Engine/Plugins/Experimental/GPULightmass/Source/GPULightmass/Private/VolumetricLightmap.cpp
--- a/Engine/Plugins/Experimental/GPULightmass/Source/GPULightmass/Private/VolumetricLightmap.cpp
+++ b/Engine/Plugins/Experimental/GPULightmass/Source/GPULightmass/Private/VolumetricLightmap.cpp
@@ -750,7 +750,7 @@ void FVolumetricLightmapRenderer::BackgroundTick()
 	OutputBrickDataRDG.DebugSetupInvariants();
 	
 	RectLightAtlas::UpdateAtlasTexture(GraphBuilder, Scene->FeatureLevel);
-	IESAtlas::UpdateAtlasTexture(GraphBuilder, Scene->FeatureLevel);
+	IESAtlas::UpdateAtlasTexture(GraphBuilder, GetFeatureLevelShaderPlatform(Scene->FeatureLevel));
 
 	FGlobalShaderMap* GlobalShaderMap = GetGlobalShaderMap(Scene->FeatureLevel);
 
diff --git a/Engine/Source/Runtime/Renderer/Private/DeferredShadingRenderer.cpp b/Engine/Source/Runtime/Renderer/Private/DeferredShadingRenderer.cpp
--- a/Engine/Source/Runtime/Renderer/Private/DeferredShadingRenderer.cpp
+++ b/Engine/Source/Runtime/Renderer/Private/DeferredShadingRenderer.cpp
@@ -1739,7 +1739,7 @@ void FDeferredShadingSceneRenderer::Render(FRDGBuilder& GraphBuilder)
 
 		// Force the rect light texture & IES texture to be updated.
 		RectLightAtlas::UpdateAtlasTexture(GraphBuilder, FeatureLevel);
-		IESAtlas::UpdateAtlasTexture(GraphBuilder, FeatureLevel);
+		IESAtlas::UpdateAtlasTexture(GraphBuilder, ShaderPlatform);
 	}
 
 	FSceneTexturesConfig& SceneTexturesConfig = GetActiveSceneTexturesConfig();
diff --git a/Engine/Source/Runtime/Renderer/Private/IESTextureManager.cpp b/Engine/Source/Runtime/Renderer/Private/IESTextureManager.cpp
--- a/Engine/Source/Runtime/Renderer/Private/IESTextureManager.cpp
+++ b/Engine/Source/Runtime/Renderer/Private/IESTextureManager.cpp
@@ -1,4 +1,4 @@
-// Copyright Epic Games, Inc. All Rights Reserved.
+﻿// Copyright Epic Games, Inc. All Rights Reserved.
 
 #include "IESTextureManager.h"
 #include "Shader.h"
@@ -192,10 +192,10 @@ static FRDGTextureRef AddIESDebugPass(
 ///////////////////////////////////////////////////////////////////////////////////////////////////
 // Add pass
 
-static bool UseComputeForAtlasUpdate(ERHIFeatureLevel::Type FeatureLevel)
+static bool UseComputeForAtlasUpdate(const FStaticShaderPlatform ShaderPlatform)
 {
-	// Prefer raster version on mobile, as OpenGL ES does not support writing to R16F images from compute
-	return (FeatureLevel > ERHIFeatureLevel::ES3_1);
+	// OpenGL ES does not support writing to R16F images from compute
+	return !IsOpenGLPlatform(ShaderPlatform);
 }
 
 class FIESAtlasAddTextureCS : public FGlobalShader
@@ -216,7 +216,7 @@ class FIESAtlasAddTextureCS : public FGlobalShader
 public:
 	static bool ShouldCompilePermutation(const FGlobalShaderPermutationParameters& Parameters) 
 	{ 
-		return UseComputeForAtlasUpdate(GetMaxSupportedFeatureLevel(Parameters.Platform)); 
+		return UseComputeForAtlasUpdate(Parameters.Platform); 
 	}
 	static void ModifyCompilationEnvironment(const FGlobalShaderPermutationParameters& Parameters, FShaderCompilerEnvironment& OutEnvironment)
 	{
@@ -286,7 +286,7 @@ class FIESAtlasAddTexturePS : public FGlobalShader
 public:
 	static bool ShouldCompilePermutation(const FGlobalShaderPermutationParameters& Parameters) 
 	{ 
-		return !UseComputeForAtlasUpdate(GetMaxSupportedFeatureLevel(Parameters.Platform)); 
+		return !UseComputeForAtlasUpdate(Parameters.Platform); 
 	}
 	static void ModifyCompilationEnvironment(const FGlobalShaderPermutationParameters& Parameters, FShaderCompilerEnvironment& OutEnvironment)
 	{
@@ -330,12 +330,12 @@ static void AddSlotsPassPS(
 
 static void AddSlotsPass(
 	FRDGBuilder& GraphBuilder,
-	ERHIFeatureLevel::Type FeatureLevel,
+	const FStaticShaderPlatform ShaderPlatform,
 	FGlobalShaderMap* ShaderMap,
 	const TArray<FAtlasSlot>& Slots,
 	FRDGTextureRef& OutAtlas)
 {
-	if (UseComputeForAtlasUpdate(FeatureLevel))
+	if (UseComputeForAtlasUpdate(ShaderPlatform))
 	{
 		AddSlotsPassCS(GraphBuilder, ShaderMap, Slots, OutAtlas);
 	}
@@ -426,9 +426,9 @@ void RemoveTexture(uint32 InSlotIndex)
 	}
 }
 
-static FRDGTextureRef CreateAtlasTexture(FRDGBuilder& GraphBuilder, ERHIFeatureLevel::Type FeatureLevel, const FIntPoint& Resolution, uint32 SliceCount)
+static FRDGTextureRef CreateAtlasTexture(FRDGBuilder& GraphBuilder, const FStaticShaderPlatform ShaderPlatform, const FIntPoint& Resolution, uint32 SliceCount)
 {
-	ETextureCreateFlags CreateAddFlags = UseComputeForAtlasUpdate(FeatureLevel) ? ETextureCreateFlags::UAV : ETextureCreateFlags::TargetArraySlicesIndependently;
+	ETextureCreateFlags CreateAddFlags = UseComputeForAtlasUpdate(ShaderPlatform) ? ETextureCreateFlags::UAV : ETextureCreateFlags::TargetArraySlicesIndependently;
 	if (GIsEditor)
 	{
 		// Make sure UAV flag is always present in Editor environment as target can be used by both compute and raster
@@ -488,7 +488,7 @@ static uint32 GetRequestedSlotCount(const FIESTextureManager& In)
 	return OutCount;
 }
 
-void UpdateAtlasTexture(FRDGBuilder& GraphBuilder, const ERHIFeatureLevel::Type FeatureLevel)
+void UpdateAtlasTexture(FRDGBuilder& GraphBuilder, const FStaticShaderPlatform ShaderPlatform)
 {
 	// Force update by resetting the atlas layout
 	static uint32 CachedAtlasResolution = CVarIESTextureResolution.GetValueOnRenderThread();
@@ -540,7 +540,7 @@ void UpdateAtlasTexture(FRDGBuilder& GraphBuilder, const ERHIFeatureLevel::Type
 	// 3. Process new atlas entries (texture allocate, slice allocation, ...) or refresh slots
 	if (GIESTextureManager.bHasPendingAdds)
 	{
-		FGlobalShaderMap* ShaderMap = GetGlobalShaderMap(FeatureLevel);
+		FGlobalShaderMap* ShaderMap = GetGlobalShaderMap(ShaderPlatform);
 
 		// 3.1. Allocate new slice
 		TArray<FAtlasSlot> NewSlots;
@@ -564,7 +564,7 @@ void UpdateAtlasTexture(FRDGBuilder& GraphBuilder, const ERHIFeatureLevel::Type
 			bool bNeedExtraction = false;
 			if (bForceRecreate)
 			{
-				AtlasTexture = CreateAtlasTexture(GraphBuilder, FeatureLevel, CachedAtlasResolution, RequestedSliceCount);
+				AtlasTexture = CreateAtlasTexture(GraphBuilder, ShaderPlatform, CachedAtlasResolution, RequestedSliceCount);
 				bNeedExtraction = true;
 			}
 			else
@@ -575,7 +575,7 @@ void UpdateAtlasTexture(FRDGBuilder& GraphBuilder, const ERHIFeatureLevel::Type
 			// 3.2.2 Insert new slots into the atlas texture
 			if (NewSlots.Num() > 0)
 			{
-				AddSlotsPass(GraphBuilder, FeatureLevel, ShaderMap, NewSlots, AtlasTexture);
+				AddSlotsPass(GraphBuilder, ShaderPlatform, ShaderMap, NewSlots, AtlasTexture);
 			}
 
 			if (bNeedExtraction)
@@ -600,9 +600,9 @@ void UpdateAtlasTexture(FRDGBuilder& GraphBuilder, const ERHIFeatureLevel::Type
 			}
 		}
 
-		FGlobalShaderMap* ShaderMap = GetGlobalShaderMap(FeatureLevel);
+		FGlobalShaderMap* ShaderMap = GetGlobalShaderMap(ShaderPlatform);
 		FRDGTextureRef AtlasTexture = GraphBuilder.RegisterExternalTexture(GIESTextureManager.AtlasTexture);
-		AddSlotsPass(GraphBuilder, FeatureLevel, ShaderMap, RefreshSlots, AtlasTexture);
+		AddSlotsPass(GraphBuilder, ShaderPlatform, ShaderMap, RefreshSlots, AtlasTexture);
 		GIESTextureManager.bHasPendingRefreshes = false;
 	}
 }
diff --git a/Engine/Source/Runtime/Renderer/Private/IESTextureManager.h b/Engine/Source/Runtime/Renderer/Private/IESTextureManager.h
--- a/Engine/Source/Runtime/Renderer/Private/IESTextureManager.h
+++ b/Engine/Source/Runtime/Renderer/Private/IESTextureManager.h
@@ -25,7 +25,7 @@ RENDERER_API float GetAtlasSlot(uint32 InSlotId);
 RENDERER_API FRHITexture* GetAtlasTexture();
 
 // Update the rect light atlas texture
-RENDERER_API void UpdateAtlasTexture(FRDGBuilder& GraphBuilder, const ERHIFeatureLevel::Type FeatureLevel);
+RENDERER_API void UpdateAtlasTexture(FRDGBuilder& GraphBuilder, const FStaticShaderPlatform ShaderPlatform);
 
 // Return the rect light atlas debug pass
 RENDERER_API void AddDebugPass(FRDGBuilder& GraphBuilder, const FViewInfo& View, FRDGTextureRef OutputTexture);
diff --git a/Engine/Source/Runtime/Renderer/Private/MobileShadingRenderer.cpp b/Engine/Source/Runtime/Renderer/Private/MobileShadingRenderer.cpp
--- a/Engine/Source/Runtime/Renderer/Private/MobileShadingRenderer.cpp
+++ b/Engine/Source/Runtime/Renderer/Private/MobileShadingRenderer.cpp
@@ -985,7 +985,7 @@ void FMobileSceneRenderer::Render(FRDGBuilder& GraphBuilder)
 		if (bDeferredShading)
 		{
 			RectLightAtlas::UpdateAtlasTexture(GraphBuilder, FeatureLevel);
-			IESAtlas::UpdateAtlasTexture(GraphBuilder, FeatureLevel);
+			IESAtlas::UpdateAtlasTexture(GraphBuilder, ShaderPlatform);
 		}
 
 		// Important that this uses consistent logic throughout the frame, so evaluate once and pass in the flag from here
