Reimplement properly hair contact shadow support for hair strands.

* This fixes engine test failing
* This allows to toggle contact shadow on hair strands as expected.

diff --git a/Engine/Plugins/Runtime/HairStrands/Source/HairStrandsCore/Private/HairStrandsInterpolation.cpp b/Engine/Plugins/Runtime/HairStrands/Source/HairStrandsCore/Private/HairStrandsInterpolation.cpp
--- a/Engine/Plugins/Runtime/HairStrands/Source/HairStrandsCore/Private/HairStrandsInterpolation.cpp
+++ b/Engine/Plugins/Runtime/HairStrands/Source/HairStrandsCore/Private/HairStrandsInterpolation.cpp
@@ -1737,6 +1737,7 @@ FHairGroupPublicData::FVertexFactoryInput InternalComputeHairStrandsVertexInputD
 	OutVFInput.Strands.Common.Flags |= Instance->Strands.Modifier.bUseStableRasterization ? HAIR_FLAGS_STABLE_RASTER : 0u;
 	OutVFInput.Strands.Common.Flags |= bRaytracingGeometry ? HAIR_FLAGS_RAYTRACING_GEOMETRY : 0u;
 	OutVFInput.Strands.Common.Flags |= Instance->Debug.Proxy->Holdout() ? HAIR_FLAGS_HOLDOUT : 0u;
+	OutVFInput.Strands.Common.Flags |= Instance->Debug.Proxy->CastsContactShadow() ? HAIR_FLAGS_CAST_CONTACT_SHADOW : 0u;
 	OutVFInput.Strands.Common.RaytracingProceduralSplits = GetHairRaytracingProceduralSplits();
 	OutVFInput.Strands.Common.GroupIndex = Instance->Debug.GroupIndex;
 	OutVFInput.Strands.Common.GroupCount = Instance->Debug.GroupCount;
diff --git a/Engine/Shaders/Private/DeferredLightingCommon.ush b/Engine/Shaders/Private/DeferredLightingCommon.ush
--- a/Engine/Shaders/Private/DeferredLightingCommon.ush
+++ b/Engine/Shaders/Private/DeferredLightingCommon.ush
@@ -63,11 +63,11 @@ float ShadowRayCast(
 		uint2 PixelPos = View.ViewRectMin.xy + View.ViewSizeAndInvSize.xy * HitUV;
 		FSubstrateAddressing SubstrateAddressing = GetSubstratePixelDataByteOffset(PixelPos, uint2(View.BufferSizeAndInvSize.xy), SubstrateRayCast.MaxBytesPerPixel);
 		FSubstratePixelHeader SubstratePixelHeader = UnpackSubstrateHeaderIn(SubstrateRayCast.MaterialTextureArray, SubstrateAddressing, SubstrateRayCast.TopLayerTexture);
-		bOutHitCastContactShadow = bHairNoShadowLight || SubstratePixelHeader.DoesCastContactShadow();
+		bOutHitCastContactShadow = SubstratePixelHeader.DoesCastContactShadow();
 #endif
 #else
 		FGBufferData SampleGBuffer = GetGBufferData(HitUV);
-		bOutHitCastContactShadow = bHairNoShadowLight || CastContactShadow(SampleGBuffer);
+		bOutHitCastContactShadow = CastContactShadow(SampleGBuffer);
 #endif
 	}
 
diff --git a/Engine/Shaders/Private/HairStrands/HairStrandsVisibilityDepthPS.usf b/Engine/Shaders/Private/HairStrands/HairStrandsVisibilityDepthPS.usf
--- a/Engine/Shaders/Private/HairStrands/HairStrandsVisibilityDepthPS.usf
+++ b/Engine/Shaders/Private/HairStrands/HairStrandsVisibilityDepthPS.usf
@@ -25,6 +25,7 @@ struct FHairSampleData
 {
 	float DeviceZ;
 	uint LightChannelMask;
+	bool CastContactShadow;
 };
 
 FHairSampleData ExtractSampleData(uint3 PixelCoord)
@@ -39,6 +40,7 @@ FHairSampleData ExtractSampleData(uint3 PixelCoord)
 		const FHairSample Sample = UnpackHairSample(HairSampleData[SampleIndex]);
 		Out.DeviceZ = max(Out.DeviceZ, Sample.Depth); // Inverse-Z
 		Out.LightChannelMask |= Sample.LightChannelMask;
+		Out.CastContactShadow |= HasHairFlags(Sample.Flags, HAIR_FLAGS_CAST_CONTACT_SHADOW);
 	}
 	return Out;
 }
@@ -52,11 +54,15 @@ FHairSampleData ExtractSampleData(uint3 PixelCoord)
 #endif
 #if PERMUTATION_OUTPUT_TYPE == OUTPUT_MATERIALDATA || PERMUTATION_OUTPUT_TYPE == OUTPUT_MATERIALDATA_LIGHTCHANNEL
 void MainPS(
-	in FScreenVertexOutput Input,
-	out float SvDepth			: SV_DEPTH,
-	out OUTPUT_FORMAT Output0	: SV_Target0,
-	out OUTPUT_FORMAT Output1	: SV_Target1,
-	out float4 OutColor			: SV_Target2)
+	in FScreenVertexOutput Input
+	, out float SvDepth			: SV_DEPTH
+	, out OUTPUT_FORMAT Output0	: SV_Target0
+	, out OUTPUT_FORMAT Output1	: SV_Target1
+	, out float4 OutColor		: SV_Target2
+#if !SUBSTRATE_ENABLED
+	, out OUTPUT_FORMAT Output2	: SV_Target3
+#endif
+)
 {
 	const uint3 PixelCoord = uint3(Input.Position.xy, 0);
 	const float Coverage = CoverageTexture.Load(PixelCoord);
@@ -76,9 +82,15 @@ void MainPS(
 #if SUBSTRATE_ENABLED
 	Output0 = 0; // MaterialData
 	Output1 = 0; // TopLayerData
+
+	// Force PerObjectData contact shadow to respect hair settings, in order to have correct screen cast later
+	Output0 = SetCastContactShadow(SampleData.CastContactShadow);
 #else
 	Output0 = float4(0, 0, 0, 0);
 	Output1 = 0;
+
+	// Force PerObjectData contact shadow to respect hair settings, in order to have correct screen cast later
+	Output2 = float4(0,0,0,SampleData.CastContactShadow ? 1 : 0);
 #endif
 	OutColor = 0;
 
diff --git a/Engine/Shaders/Private/HairStrands/HairStrandsVisibilityUtils.ush b/Engine/Shaders/Private/HairStrands/HairStrandsVisibilityUtils.ush
--- a/Engine/Shaders/Private/HairStrands/HairStrandsVisibilityUtils.ush
+++ b/Engine/Shaders/Private/HairStrands/HairStrandsVisibilityUtils.ush
@@ -19,6 +19,6 @@ FGBufferData HairSampleToGBufferData(in FHairSample In, in float InDualScatterin
 	Out.DiffuseColor = 0;
 	Out.SpecularColor = 0;
 	Out.PrecomputedShadowFactors = 1;
-	Out.PerObjectGBufferData = 0;
+	Out.PerObjectGBufferData = HasHairFlags(In.Flags, HAIR_FLAGS_CAST_CONTACT_SHADOW) ? 1u : 0u;
 	return Out;
 }
\ No newline at end of file
diff --git a/Engine/Shaders/Shared/HairStrandsDefinitions.h b/Engine/Shaders/Shared/HairStrandsDefinitions.h
--- a/Engine/Shaders/Shared/HairStrandsDefinitions.h
+++ b/Engine/Shaders/Shared/HairStrandsDefinitions.h
@@ -105,6 +105,7 @@
 #define HAIR_FLAGS_STABLE_RASTER 2u
 #define HAIR_FLAGS_RAYTRACING_GEOMETRY 4u
 #define HAIR_FLAGS_HOLDOUT 8u
+#define HAIR_FLAGS_CAST_CONTACT_SHADOW 10u
 
 #ifndef __cplusplus //HLSL
 bool HasHairFlags(uint In, uint Flags)
diff --git a/Engine/Source/Runtime/Renderer/Private/HairStrands/HairStrandsVisibility.cpp b/Engine/Source/Runtime/Renderer/Private/HairStrands/HairStrandsVisibility.cpp
--- a/Engine/Source/Runtime/Renderer/Private/HairStrands/HairStrandsVisibility.cpp
+++ b/Engine/Source/Runtime/Renderer/Private/HairStrands/HairStrandsVisibility.cpp
@@ -2745,6 +2745,7 @@ static void AddHairAuxilaryPass(
 	const FRDGBufferRef& HairSampleData,
 	FRDGTextureRef Output0,
 	FRDGTextureRef Output1,
+	FRDGTextureRef Output2,
 	FRDGTextureRef OutColorTexture,
 	FRDGTextureRef OutDepthTexture,
 	FRDGTextureRef OutLightChannelMaskTexture)
@@ -2777,6 +2778,7 @@ static void AddHairAuxilaryPass(
 		Parameters->RenderTargets[0] = FRenderTargetBinding(Output0, ERenderTargetLoadAction::ELoad, 0 /*Mip-0*/, 0/*First slice, which contains state/counter*/);
 		Parameters->RenderTargets[1] = FRenderTargetBinding(Output1, ERenderTargetLoadAction::ELoad);
 		Parameters->RenderTargets[2] = FRenderTargetBinding(OutColorTexture, ERenderTargetLoadAction::ELoad);
+		Parameters->RenderTargets[3] = FRenderTargetBinding(Output2, ERenderTargetLoadAction::ELoad);
 	}
 
 	if (PassType == EHairAuxilaryPassType::MaterialData_LightChannelMask || PassType == EHairAuxilaryPassType::LightChannelMask)
@@ -2865,6 +2867,7 @@ static FRDGTextureRef AddHairLightChannelMaskPass(
 		nullptr,
 		nullptr,
 		nullptr,
+		nullptr,
 		SceneDepthTexture,
 		OutLightChannelMask);
 	return OutLightChannelMask;
@@ -2881,6 +2884,7 @@ static void AddHairMaterialDataPatchPass(
 	const FRDGBufferRef& HairSampleData,
 	FRDGTextureRef& OutMaterial0,
 	FRDGTextureRef& OutMaterial1,
+	FRDGTextureRef& OutMaterial2,
 	FRDGTextureRef& OutColorTexture,
 	FRDGTextureRef& OutDepthTexture,
 	FRDGTextureRef& OutLightChannelMask)
@@ -2917,6 +2921,7 @@ static void AddHairMaterialDataPatchPass(
 		HairSampleData,
 		OutMaterial0,
 		OutMaterial1,
+		OutMaterial2,
 		OutColorTexture,
 		OutDepthTexture,
 		OutLightChannelMask);
@@ -2951,6 +2956,7 @@ static void AddHairOnlyDepthPass(
 			nullptr,
 			nullptr,
 			nullptr,
+			nullptr,
 			OutDepthTexture,
 			nullptr);
 	}
@@ -2967,6 +2973,7 @@ static void AddHairOnlyDepthPass(
 		nullptr,
 		nullptr,
 		nullptr,
+		nullptr,
 		OutDepthTexture,
 		nullptr);
 }
@@ -4606,15 +4613,18 @@ void RenderHairStrandsVisibilityBuffer(
 
 	FRDGTextureRef SceneMaterial0 = nullptr;
 	FRDGTextureRef SceneMaterial1 = nullptr;
+	FRDGTextureRef SceneMaterial2 = nullptr;
 	if (Substrate::IsSubstrateEnabled())
 	{
 		SceneMaterial0 = View.SubstrateViewData.SceneData->MaterialTextureArray;
 		SceneMaterial1 = View.SubstrateViewData.SceneData->TopLayerTexture;
+		SceneMaterial2 = nullptr;
 	}
 	else
 	{
 		SceneMaterial0 = SceneGBufferBTexture;
 		SceneMaterial1 = SceneGBufferCTexture;
+		SceneMaterial2 = SceneGBufferATexture;
 	}
 
 	{
@@ -4728,6 +4738,7 @@ void RenderHairStrandsVisibilityBuffer(
 						CompactNodeData,
 						SceneMaterial0,
 						SceneMaterial1,
+						SceneMaterial2,
 						SceneColorTexture,
 						SceneDepthTexture,
 						VisibilityData.LightChannelMaskTexture);
@@ -4844,6 +4855,7 @@ void RenderHairStrandsVisibilityBuffer(
 						CompactNodeData,
 						SceneMaterial0,
 						SceneMaterial1,
+						SceneMaterial2,
 						SceneColorTexture,
 						SceneDepthTexture,
 						VisibilityData.LightChannelMaskTexture);
@@ -5040,6 +5052,7 @@ void RenderHairStrandsVisibilityBuffer(
 							CompactNodeData,
 							SceneMaterial0,
 							SceneMaterial1,
+							SceneMaterial2,
 							SceneColorTexture,
 							SceneDepthTexture,
 							VisibilityData.LightChannelMaskTexture);
@@ -5179,6 +5192,7 @@ void RenderHairStrandsVisibilityBuffer(
 						CompactNodeData,
 						SceneMaterial0,
 						SceneMaterial1,
+						SceneMaterial2,
 						SceneColorTexture,
 						SceneDepthTexture,
 						VisibilityData.LightChannelMaskTexture);
