Add defensive code for nullptr.

diff --git a/Engine/Source/Runtime/Renderer/Private/VT/AllocatedVirtualTexture.cpp b/Engine/Source/Runtime/Renderer/Private/VT/AllocatedVirtualTexture.cpp
--- a/Engine/Source/Runtime/Renderer/Private/VT/AllocatedVirtualTexture.cpp
+++ b/Engine/Source/Runtime/Renderer/Private/VT/AllocatedVirtualTexture.cpp
@@ -141,11 +141,14 @@ void FAllocatedVirtualTexture::Destroy(FVirtualTextureSystem* System)
 
 		for (int32 PageTableIndex = 0u; PageTableIndex < UniquePageTableLayers.Num(); ++PageTableIndex)
 		{
-			if (UniquePageTableLayers[PageTableIndex].PhysicalSpace->ReleaseResourceRef() == 0)
+			if (UniquePageTableLayers[PageTableIndex].PhysicalSpace)
 			{
-				UniquePageTableLayers[PageTableIndex].PhysicalSpace->ReleaseResource();
+				if (UniquePageTableLayers[PageTableIndex].PhysicalSpace->ReleaseResourceRef() == 0)
+				{
+					UniquePageTableLayers[PageTableIndex].PhysicalSpace->ReleaseResource();
+				}
+				UniquePageTableLayers[PageTableIndex].PhysicalSpace.SafeRelease();
 			}
-			UniquePageTableLayers[PageTableIndex].PhysicalSpace.SafeRelease();
 		}
 
 #if DO_CHECK
@@ -388,7 +391,7 @@ uint32 FAllocatedVirtualTexture::AddUniquePhysicalSpace(FRHICommandListBase& InR
 	UniquePageTableLayers[Index].ProducerTextureLayerMask = 0;
 	UniquePageTableLayers[Index].TextureLayerCount = 0;
 
-	if (InPhysicalSpace->AddResourceRef() == 1)
+	if (InPhysicalSpace && InPhysicalSpace->AddResourceRef() == 1)
 	{
 		InPhysicalSpace->InitResource(InRHICmdList);
 	}
