Skip calculating DF shadows if downsampled scissor rect is null.
- prevents assert that triggered because of CS being dispatched with 0 groups.

diff --git a/Engine/Source/Runtime/Renderer/Private/DistanceFieldShadowing.cpp b/Engine/Source/Runtime/Renderer/Private/DistanceFieldShadowing.cpp
--- a/Engine/Source/Runtime/Renderer/Private/DistanceFieldShadowing.cpp
+++ b/Engine/Source/Runtime/Renderer/Private/DistanceFieldShadowing.cpp
@@ -146,6 +146,11 @@ FIntPoint GetBufferSizeForDFShadows(const FViewInfo& View)
 	return FIntPoint::DivideAndRoundDown(View.GetSceneTexturesConfig().Extent, GetDFShadowDownsampleFactor());
 }
 
+FIntRect GetScissorRectForDFShadows(FIntRect ScissorRect)
+{
+	return ScissorRect / GetDFShadowDownsampleFactor();
+}
+
 class FCullObjectsForShadowCS : public FGlobalShader
 {
 	DECLARE_GLOBAL_SHADER(FCullObjectsForShadowCS);
@@ -733,8 +738,10 @@ void RayTraceShadows(
 	FRDGBuilder& GraphBuilder,
 	bool bAsyncCompute,
 	const FMinimalSceneTextures& SceneTextures,
-	FRDGTextureRef RayTracedShadowsTexture,
+	FRDGTextureRef OutputTexture,
 	const FViewInfo& View,
+	const FIntRect& ScissorRect,
+	const FIntRect& DownsampledScissorRect,
 	const FDistanceFieldSceneData& DistanceFieldSceneData,
 	const FProjectedShadowInfo* ProjectedShadowInfo,
 	EDistanceFieldPrimitiveType PrimitiveType,
@@ -744,12 +751,6 @@ void RayTraceShadows(
 	const FDistanceFieldCulledObjectBufferParameters& CulledObjectBufferParameters,
 	const FLightTileIntersectionParameters& LightTileIntersectionParameters)
 {
-	FIntRect ScissorRect;
-	if (!ProjectedShadowInfo->GetLightSceneInfo().Proxy->GetScissorRect(ScissorRect, View, View.ViewRect))
-	{
-		ScissorRect = View.ViewRect;
-	}
-
 	const int32 DFShadowQuality = (PrimitiveType == DFPT_HeightField ? GetHFShadowQuality() : GetDFShadowQuality()) - 1;
 	check(DFShadowQuality >= 0);
 
@@ -777,7 +778,7 @@ void RayTraceShadows(
 	{
 		FDistanceFieldShadowingCS::FParameters* PassParameters = GraphBuilder.AllocParameters<FDistanceFieldShadowingCS::FParameters>();
 			
-		PassParameters->RWShadowFactors = GraphBuilder.CreateUAV(RayTracedShadowsTexture);
+		PassParameters->RWShadowFactors = GraphBuilder.CreateUAV(OutputTexture);
 		PassParameters->View = View.ViewUniformBuffer;
 		PassParameters->SceneTextures = SceneTextures.GetSceneTextureShaderParameters(View.GetFeatureLevel());
 
@@ -820,7 +821,7 @@ void RayTraceShadows(
 		}
 
 		PassParameters->DownsampleFactor = GetDFShadowDownsampleFactor();
-		const FIntPoint OutputBufferSize = GetBufferSizeForDFShadows(View);
+		const FIntPoint OutputBufferSize = OutputTexture->Desc.Extent;
 		PassParameters->InvOutputBufferSize = FVector2f(1.f / OutputBufferSize.X, 1.f / OutputBufferSize.Y);
 		PassParameters->ShadowFactorsTexture = PrevOutputTexture;
 		PassParameters->ShadowFactorsSampler = TStaticSamplerState<>::GetRHI();
@@ -838,8 +839,8 @@ void RayTraceShadows(
 
 		auto ComputeShader = View.ShaderMap->GetShader< FDistanceFieldShadowingCS >(PermutationVector);
 
-		uint32 GroupSizeX = FMath::DivideAndRoundUp(ScissorRect.Size().X / GetDFShadowDownsampleFactor(), GDistanceFieldShadowTileSizeX);
-		uint32 GroupSizeY = FMath::DivideAndRoundUp(ScissorRect.Size().Y / GetDFShadowDownsampleFactor(), GDistanceFieldShadowTileSizeY);
+		uint32 GroupSizeX = FMath::DivideAndRoundUp(DownsampledScissorRect.Size().X, GDistanceFieldShadowTileSizeX);
+		uint32 GroupSizeY = FMath::DivideAndRoundUp(DownsampledScissorRect.Size().Y, GDistanceFieldShadowTileSizeY);
 		PassParameters->NumGroups = FVector2f(GroupSizeX, GroupSizeY);
 
 		FComputeShaderUtils::AddPass(
@@ -856,7 +857,8 @@ FRDGTextureRef FProjectedShadowInfo::RenderRayTracedDistanceFieldProjection(
 	FRDGBuilder& GraphBuilder,
 	bool bAsyncCompute, 
 	const FMinimalSceneTextures& SceneTextures,
-	const FViewInfo& View)
+	const FViewInfo& View,
+	const FIntRect& ScissorRect)
 {
 	DistanceFieldShadowViewGPUData& SDFShadowViewGPUData = CachedDistanceFieldShadowViewGPUData.FindOrAdd(&View);
 
@@ -866,6 +868,14 @@ FRDGTextureRef FProjectedShadowInfo::RenderRayTracedDistanceFieldProjection(
 		return SDFShadowViewGPUData.RayTracedShadowsTexture;
 	}
 
+	const FIntRect DownsampledScissorRect = GetScissorRectForDFShadows(ScissorRect);
+
+	if (DownsampledScissorRect.Area() <= 0)
+	{
+		// skip calculating DF shadows
+		return nullptr;
+	}
+
 	const bool bDFShadowSupported = SupportsDistanceFieldShadows(View.GetFeatureLevel(), View.GetShaderPlatform());
 	const bool bHFShadowSupported = SupportsHeightFieldShadows(View.GetFeatureLevel(), View.GetShaderPlatform());
 	const FScene* Scene = (const FScene*)(View.Family->Scene);
@@ -938,7 +948,22 @@ FRDGTextureRef FProjectedShadowInfo::RenderRayTracedDistanceFieldProjection(
 				SDFShadowViewGPUData.RayTracedShadowsTexture = GraphBuilder.CreateTexture(Desc, TEXT("RayTracedShadows"));
 			}
 
-			RayTraceShadows(GraphBuilder, bAsyncCompute, SceneTextures, SDFShadowViewGPUData.RayTracedShadowsTexture, View, Scene->DistanceFieldSceneData, this, DFPT_SignedDistanceField, false, nullptr, ObjectBufferParameters, *SDFShadowViewGPUData.SDFCulledObjectBufferParameters, *SDFShadowViewGPUData.SDFLightTileIntersectionParameters);
+			RayTraceShadows(
+				GraphBuilder,
+				bAsyncCompute,
+				SceneTextures,
+				SDFShadowViewGPUData.RayTracedShadowsTexture,
+				View,
+				ScissorRect,
+				DownsampledScissorRect,
+				Scene->DistanceFieldSceneData,
+				this,
+				DFPT_SignedDistanceField,
+				false,
+				nullptr,
+				ObjectBufferParameters,
+				*SDFShadowViewGPUData.SDFCulledObjectBufferParameters,
+				*SDFShadowViewGPUData.SDFLightTileIntersectionParameters);
 		}
 	}
 
@@ -1008,7 +1033,22 @@ FRDGTextureRef FProjectedShadowInfo::RenderRayTracedDistanceFieldProjection(
 			SDFShadowViewGPUData.RayTracedShadowsTexture = GraphBuilder.CreateTexture(Desc, TEXT("RayTracedShadows"));
 		}
 
-		RayTraceShadows(GraphBuilder, bAsyncCompute, SceneTextures, SDFShadowViewGPUData.RayTracedShadowsTexture, View, Scene->DistanceFieldSceneData, this, DFPT_HeightField, bHasPrevOutput, PrevOutputTexture, ObjectBufferParameters, *SDFShadowViewGPUData.HeightFieldCulledObjectBufferParameters, *SDFShadowViewGPUData.HeightFieldLightTileIntersectionParameters);
+		RayTraceShadows(
+			GraphBuilder,
+			bAsyncCompute,
+			SceneTextures,
+			SDFShadowViewGPUData.RayTracedShadowsTexture,
+			View,
+			ScissorRect,
+			DownsampledScissorRect,
+			Scene->DistanceFieldSceneData,
+			this,
+			DFPT_HeightField,
+			bHasPrevOutput,
+			PrevOutputTexture,
+			ObjectBufferParameters,
+			*SDFShadowViewGPUData.HeightFieldCulledObjectBufferParameters,
+			*SDFShadowViewGPUData.HeightFieldLightTileIntersectionParameters);
 	}
 
 	return SDFShadowViewGPUData.RayTracedShadowsTexture;
@@ -1034,7 +1074,7 @@ void FProjectedShadowInfo::RenderRayTracedDistanceFieldProjection(
 	check(ScissorRect.Area() > 0);
 	const bool bRunTiled = UseShadowIndirectDraw(View.GetShaderPlatform()) && TiledShadowRendering != nullptr;
 
-	FRDGTextureRef RayTracedShadowsTexture = RenderRayTracedDistanceFieldProjection(GraphBuilder, false, SceneTextures, View);
+	FRDGTextureRef RayTracedShadowsTexture = RenderRayTracedDistanceFieldProjection(GraphBuilder, false, SceneTextures, View, ScissorRect);
 
 	if (RayTracedShadowsTexture)
 	{
diff --git a/Engine/Source/Runtime/Renderer/Private/ShadowRendering.cpp b/Engine/Source/Runtime/Renderer/Private/ShadowRendering.cpp
--- a/Engine/Source/Runtime/Renderer/Private/ShadowRendering.cpp
+++ b/Engine/Source/Runtime/Renderer/Private/ShadowRendering.cpp
@@ -2225,7 +2225,7 @@ void FSceneRenderer::BeginAsyncDistanceFieldShadowProjections(FRDGBuilder& Graph
 				{
 					// Kick off distance field shadow calculation in async compute
 					// Don't need store result reference because it is internally cached by FProjectedShadowInfo
-					ProjectedShadowInfo->RenderRayTracedDistanceFieldProjection(GraphBuilder, true, SceneTextures, View);
+					ProjectedShadowInfo->RenderRayTracedDistanceFieldProjection(GraphBuilder, true, SceneTextures, View, ScissorRect);
 				}
 			}
 		}
diff --git a/Engine/Source/Runtime/Renderer/Private/ShadowRendering.h b/Engine/Source/Runtime/Renderer/Private/ShadowRendering.h
--- a/Engine/Source/Runtime/Renderer/Private/ShadowRendering.h
+++ b/Engine/Source/Runtime/Renderer/Private/ShadowRendering.h
@@ -632,7 +632,8 @@ public:
 		FRDGBuilder& GraphBuilder,
 		bool bAsyncCompute,
 		const FMinimalSceneTextures& SceneTextures,
-		const FViewInfo& View);
+		const FViewInfo& View,
+		const FIntRect& ScissorRect);
 
 	/** 
 	* Renders ray traced distance field shadows into an existing texture.
