Cull small instances from Distance Field Scene to prevent precision issues causing LocalToRelativeWorld to not be invertible.

diff --git a/Engine/Source/Runtime/Renderer/Private/DistanceFieldObjectManagement.cpp b/Engine/Source/Runtime/Renderer/Private/DistanceFieldObjectManagement.cpp
--- a/Engine/Source/Runtime/Renderer/Private/DistanceFieldObjectManagement.cpp
+++ b/Engine/Source/Runtime/Renderer/Private/DistanceFieldObjectManagement.cpp
@@ -390,7 +390,7 @@ void ProcessPrimitiveUpdate(
 
 				const FMatrix::FReal MinScale = LocalToWorld.GetMinimumAxisScale();
 
-				if (bIsAddOperation && (MinScale <= 0 || bInstanceCountOverflow))
+				if (bIsAddOperation && (MinScale < 0.0001f || bInstanceCountOverflow))
 				{
 					// Skip degenerate instances or when instance count limit is reached
 					PrimitiveSceneInfo->DistanceFieldInstanceIndices[TransformIndex] = -1;
