Fixed assert due to check whether occlusion task is complete.

diff --git a/Engine/Source/Runtime/Renderer/Private/SceneVisibilityPrivate.h b/Engine/Source/Runtime/Renderer/Private/SceneVisibilityPrivate.h
--- a/Engine/Source/Runtime/Renderer/Private/SceneVisibilityPrivate.h
+++ b/Engine/Source/Runtime/Renderer/Private/SceneVisibilityPrivate.h
@@ -408,13 +408,6 @@ class FVisibilityViewPacket
 public:
 	FVisibilityViewPacket(FVisibilityTaskData& TaskData, FScene& InScene, FViewInfo& InView, int32 ViewIndex);
 
-	~FVisibilityViewPacket()
-	{
-		check(Tasks.FrustumCull.IsCompleted());
-		check(Tasks.OcclusionCull.IsCompleted());
-		check(Tasks.ComputeRelevance.IsCompleted());
-	}
-
 	FVisibilityTaskData& TaskData;
 	FVisibilityTaskConfig& TaskConfig;
 	FScene& Scene;
