Change Editor viewport orthographic debug views to utilise the camera view origin update paths, allowing Lit modes to be used in editor.
Also fixed a couple of minor bugs with projection matrix calculations and improved shadow LOD resolves in various orthographic views (perspective remains unnaffected).

There are some limitations to this fix due to near clip plane locations, but also introduced a CVar to allow users to adjust the near clip plane for debug purposes. Does not affect non-editor builds at all.

diff --git a/Engine/Source/Editor/UnrealEd/Private/EditorComponents.cpp b/Engine/Source/Editor/UnrealEd/Private/EditorComponents.cpp
--- a/Engine/Source/Editor/UnrealEd/Private/EditorComponents.cpp
+++ b/Engine/Source/Editor/UnrealEd/Private/EditorComponents.cpp
@@ -262,7 +262,7 @@ void FGridWidget::DrawNewGrid(const FSceneView* View, FPrimitiveDrawInterface* P
 
 	FMatrix ObjectToWorld = FMatrix::Identity;
 
-	FVector CameraPos = View->ViewMatrices.GetViewOrigin();
+	FVector CameraPos = View->ViewMatrices.GetLODViewOrigin();
 
 	FVector UVCameraPos = FVector(CameraPos.X, CameraPos.Y, 0);
 
diff --git a/Engine/Source/Editor/UnrealEd/Private/EditorViewportClient.cpp b/Engine/Source/Editor/UnrealEd/Private/EditorViewportClient.cpp
--- a/Engine/Source/Editor/UnrealEd/Private/EditorViewportClient.cpp
+++ b/Engine/Source/Editor/UnrealEd/Private/EditorViewportClient.cpp
@@ -98,6 +98,15 @@ static TAutoConsoleVariable<int32> CVarEditorViewportTest(
 	TEXT("1..7: Various Configuations"),
 	ECVF_RenderThreadSafe);
 
+static TAutoConsoleVariable<float> CVarOrthoEditorDebugClipPlaneScale(
+	TEXT("r.Ortho.EditorDebugClipPlaneScale"),
+	1.0f,
+	TEXT("Only affects the editor ortho viewports in Lit modes.\n")
+	TEXT("Set the scale to proportionally alter the near plane based on current Ortho width that is set.\n")
+	TEXT("This changes when geometry clips in the scene as the Orthozoom is changed. Helpful for varying mesh sizes.\n")
+	TEXT("Other light artefacts may appear when this value changes, this is unavoidable for now.\n"),
+	ECVF_RenderThreadSafe);
+
 static bool GetDefaultLowDPIPreviewValue()
 {
 	static auto CVarEditorViewportHighDPIPtr = IConsoleManager::Get().FindConsoleVariable(TEXT("r.Editor.Viewport.HighDPI"));
@@ -1121,8 +1130,8 @@ FSceneView* FEditorViewportClient::CalcSceneView(FSceneViewFamily* ViewFamily, c
 		else
 		{
 			static_assert((bool)ERHIZBuffer::IsInverted, "Check all the Rotation Matrix transformations!");
-			FMatrix::FReal ZScale = 0.5f / UE_OLD_HALF_WORLD_MAX;	// LWC_TODO: WORLD_MAX misuse?
-			FMatrix::FReal ZOffset = UE_OLD_HALF_WORLD_MAX;
+			FMatrix::FReal ZScale = 0.5f / UE_OLD_WORLD_MAX;
+			FMatrix::FReal ZOffset = UE_OLD_WORLD_MAX;
 
 			//The divisor for the matrix needs to match the translation code.
 			const float Zoom = GetOrthoUnitsPerPixel(Viewport);
@@ -1198,6 +1207,19 @@ FSceneView* FEditorViewportClient::CalcSceneView(FSceneViewFamily* ViewFamily, c
 				ZScale,
 				ZOffset
 				);
+
+			if(!ViewFamily->EngineShowFlags.Wireframe)
+			{
+				/**
+				* Update the ortho near plane and view origin to a position that is proportional to the OrthoWidth
+				* and a User specified multiplier which allows adjusting of the NearPlane location. 
+				* This ensures lighting works more appropriately on meshes in Lit mode, but introduces more abrupt NearPlane clipping.
+				*
+				* Not needed for Wireframe view; NearPlane clipping is unnecessary outside of Lit modes.
+				*/
+				float OrthoNearPlane = OrthoWidth * -CVarOrthoEditorDebugClipPlaneScale.GetValueOnAnyThread();
+				ViewInitOptions.UpdateOrthoNearPlane(OrthoNearPlane, true);
+			}
 		}
 
 		if (bConstrainAspectRatio)
diff --git a/Engine/Source/Runtime/Engine/Private/SceneView.cpp b/Engine/Source/Runtime/Engine/Private/SceneView.cpp
--- a/Engine/Source/Runtime/Engine/Private/SceneView.cpp
+++ b/Engine/Source/Runtime/Engine/Private/SceneView.cpp
@@ -605,7 +605,6 @@ bool FSceneViewProjectionData::UpdateOrthoNearPlane(FSceneViewProjectionData* In
 
 	//OrthoNearClipPlane is negative at this point + we are moving the theoretical camera position backwards.
 	InOutProjectionData->ViewOrigin += ViewForward * NearPlane;
-	NearPlane = GDefaultUpdateOrthoNearPlane;
 
 	//If required, recalculate the new projection matrix from the old one using the new NearClip value
 	if (bUpdateOrthoProjectionMatrix)
@@ -615,6 +614,7 @@ bool FSceneViewProjectionData::UpdateOrthoNearPlane(FSceneViewProjectionData* In
 		const float OrthoHeight = InvProjectionMatrix.M[1][1];
 
 		const float FarPlane = NearPlane - (InvProjectionMatrix.M[2][2]);
+		NearPlane = GDefaultUpdateOrthoNearPlane;
 		const float ZScale = 1.0f / (FarPlane - NearPlane);
 
 		InOutProjectionData->ProjectionMatrix = FReversedZOrthoMatrix(
@@ -624,6 +624,10 @@ bool FSceneViewProjectionData::UpdateOrthoNearPlane(FSceneViewProjectionData* In
 			-NearPlane
 		);
 	}
+	else
+	{
+		NearPlane = GDefaultUpdateOrthoNearPlane;
+	}
 
 	return true;
 }
diff --git a/Engine/Source/Runtime/Renderer/Private/ShadowSetup.cpp b/Engine/Source/Runtime/Renderer/Private/ShadowSetup.cpp
--- a/Engine/Source/Runtime/Renderer/Private/ShadowSetup.cpp
+++ b/Engine/Source/Runtime/Renderer/Private/ShadowSetup.cpp
@@ -1935,7 +1935,7 @@ bool FProjectedShadowInfo::AddSubjectPrimitive(FDynamicShadowsTaskData& TaskData
 				if( CurrentView.IsPerspectiveProjection() )
 				{
 					// Compute the distance between the view and the primitive.
-					float DistanceSquared = (Proxy->GetBounds().Origin - CurrentView.ShadowViewMatrices.GetViewOrigin()).SizeSquared();
+					float DistanceSquared = (Proxy->GetBounds().Origin - CurrentView.ShadowViewMatrices.GetLODViewOrigin()).SizeSquared();
 
 					bool bIsDistanceCulled = CurrentView.IsDistanceCulled(
 						DistanceSquared,
@@ -1994,7 +1994,7 @@ bool FProjectedShadowInfo::AddSubjectPrimitive(FDynamicShadowsTaskData& TaskData
 					// Note: skip small-mesh culling for VSM since it needs it drawn for GPU-side caching.
 					if (bWholeSceneShadow && CacheMode != SDCM_StaticPrimitivesOnly && MeshPassTargetType != EMeshPass::VSMShadowDepth)
 					{
-						const float DistanceSquared = ( Bounds.Origin - CurrentView.ShadowViewMatrices.GetViewOrigin() ).SizeSquared();
+						const float DistanceSquared = ( Bounds.Origin - CurrentView.ShadowViewMatrices.GetLODViewOrigin() ).SizeSquared();
 						const float LODScaleSquared = FMath::Square(CurrentView.LODDistanceFactor);
 						const bool bDrawShadowDepth = FMath::Square(Bounds.SphereRadius) > FMath::Square(GMinScreenRadiusForShadowCaster) * DistanceSquared * LODScaleSquared;
 						if( !bDrawShadowDepth )
@@ -2144,7 +2144,7 @@ uint64 FProjectedShadowInfo::AddSubjectPrimitive_AnyThread(
 			{
 				bool bFadingIn;
 				// Compute the distance between the view and the primitive.
-				const float DistanceSquared = (PrimitiveSceneInfoCompact.Bounds.Origin - CurrentView->ShadowViewMatrices.GetViewOrigin()).SizeSquared();
+				const float DistanceSquared = (PrimitiveSceneInfoCompact.Bounds.Origin - CurrentView->ShadowViewMatrices.GetLODViewOrigin()).SizeSquared();
 
 				if (CurrentView->IsDistanceCulled_AnyThread(
 					DistanceSquared,
@@ -2208,7 +2208,7 @@ uint64 FProjectedShadowInfo::AddSubjectPrimitive_AnyThread(
 				if (bWholeSceneShadow)
 				{
 					const FCompactBoxSphereBounds& Bounds = PrimitiveSceneInfoCompact.Bounds;
-					const float DistanceSquared = (Bounds.Origin - CurrentView->ShadowViewMatrices.GetViewOrigin()).SizeSquared();
+					const float DistanceSquared = (Bounds.Origin - CurrentView->ShadowViewMatrices.GetLODViewOrigin()).SizeSquared();
 					const float LODScaleSquared = FMath::Square(CurrentView->LODDistanceFactor);
 					const bool bDrawShadowDepth = FMath::Square(Bounds.SphereRadius) > FMath::Square(GMinScreenRadiusForShadowCaster) * DistanceSquared * LODScaleSquared;
 					if (!bDrawShadowDepth)
@@ -3461,7 +3461,7 @@ void FSceneRenderer::CreatePerObjectProjectedShadow(
 				const float DistanceFromShadowCenterSquared = (WholeSceneShadow->ShadowBounds.Center - Bounds.Origin).SizeSquared();
 				//@todo - if view dependent whole scene shadows are ever supported in splitscreen, 
 				// We can only disable the preshadow at this point if it is inside a whole scene shadow for all views
-				const float DistanceFromViewSquared = ((FVector)WholeSceneShadow->DependentView->ShadowViewMatrices.GetViewOrigin() - Bounds.Origin).SizeSquared();
+				const float DistanceFromViewSquared = ((FVector)WholeSceneShadow->DependentView->ShadowViewMatrices.GetLODViewOrigin() - Bounds.Origin).SizeSquared();
 				// Mark the preshadow as inside the whole scene shadow if its bounding sphere is inside the near fade distance
 				if (DistanceFromShadowCenterSquared < FMath::Square(FMath::Max(WholeSceneShadow->ShadowBounds.W - Bounds.SphereRadius, 0.0f))
 					//@todo - why is this extra threshold required?
@@ -4834,7 +4834,7 @@ struct FGatherShadowPrimitivesPacket
 				check(ProjectedShadowInfo->DependentView);
 
 				{
-					const float DistanceSquared = (PrimitiveBounds.Origin - ProjectedShadowInfo->DependentView->ShadowViewMatrices.GetViewOrigin()).SizeSquared();
+					const float DistanceSquared = (PrimitiveBounds.Origin - ProjectedShadowInfo->DependentView->ShadowViewMatrices.GetLODViewOrigin()).SizeSquared();
 					const float LODScaleSquared = FMath::Square(ProjectedShadowInfo->DependentView->LODDistanceFactor);
 					bScreenSpaceSizeCulled = FMath::Square(PrimitiveBounds.SphereRadius) < FMath::Square(MinScreenRadiusForShadowCaster) * DistanceSquared * LODScaleSquared;
 				}
