Fixed missing shadows on mobile due to skipping shadow GDME if no tasks are used.

diff --git a/Engine/Source/Runtime/Renderer/Private/ShadowSetup.cpp b/Engine/Source/Runtime/Renderer/Private/ShadowSetup.cpp
--- a/Engine/Source/Runtime/Renderer/Private/ShadowSetup.cpp
+++ b/Engine/Source/Runtime/Renderer/Private/ShadowSetup.cpp
@@ -4529,9 +4529,7 @@ void FSceneRenderer::GatherShadowDynamicMeshElements(FDynamicShadowsTaskData& Ta
 		}
 	}
 
-	const int32 NumShadowTasks = FMath::Min<int32>(GetNumShadowDynamicMeshElementTasks(), TaskData.ShadowsToGather.Num());
-
-	if (NumShadowTasks > 0)
+	if (!TaskData.ShadowsToGather.IsEmpty())
 	{
 		// Process all shadows in the serial path when not in multithreaded mode.
 		TaskData.ShadowsToGatherInSerialPass.Init(!TaskData.bMultithreadedGDME, TaskData.ShadowsToGather.Num());
@@ -4558,6 +4556,8 @@ void FSceneRenderer::GatherShadowDynamicMeshElements(FDynamicShadowsTaskData& Ta
 		if (TaskData.bMultithreadedGDME)
 		{
 			UE::Tasks::FTaskEvent MeshCollectorsTaskEvent{ UE_SOURCE_LOCATION };
+			
+			const int32 NumShadowTasks = FMath::Min<int32>(GetNumShadowDynamicMeshElementTasks(), TaskData.ShadowsToGather.Num());
 
 			for (int32 TaskIndex = 0; TaskIndex < NumShadowTasks; ++TaskIndex)
 			{
