Deletes r.TSR.Subpixel has HistoryUncertainty method replaces it well and is faster and covering other edge cases.

diff --git a/Engine/Shaders/Private/TemporalSuperResolution/TSRClearPrevTextures.usf b/Engine/Shaders/Private/TemporalSuperResolution/TSRClearPrevTextures.usf
--- a/Engine/Shaders/Private/TemporalSuperResolution/TSRClearPrevTextures.usf
+++ b/Engine/Shaders/Private/TemporalSuperResolution/TSRClearPrevTextures.usf
@@ -14,7 +14,6 @@
 
 RWTexture2D<uint> PrevUseCountOutput;
 RWTexture2D<uint> PrevClosestDepthOutput;
-RWTexture2D<uint> SubpixelDepthOutput;
 
 
 //------------------------------------------------------- ENTRY POINT
@@ -48,7 +47,6 @@ void MainCS(
 
 			PrevUseCountOutput[PrevUseCountPixelPos] = 0;
 			PrevClosestDepthOutput[InputPixelPos] = 0;
-			SubpixelDepthOutput[InputPixelPos] = 0;
 		}
 	}
 }
diff --git a/Engine/Shaders/Private/TemporalSuperResolution/TSRCommon.ush b/Engine/Shaders/Private/TemporalSuperResolution/TSRCommon.ush
--- a/Engine/Shaders/Private/TemporalSuperResolution/TSRCommon.ush
+++ b/Engine/Shaders/Private/TemporalSuperResolution/TSRCommon.ush
@@ -79,17 +79,6 @@
 
 #define MAX_PARALLAX_FACTOR 8.0
 
-// Size and payload size of the subpixel grid in the history
-#define SUB_PIXEL_GRID_SIZE 2
-#define SUB_PIXEL_COUNT (SUB_PIXEL_GRID_SIZE * SUB_PIXEL_GRID_SIZE)
-
-#define SUB_PIXEL_BIT_COUNT 4
-#define SUB_PIXEL_BIT_MASK ((1 << SUB_PIXEL_BIT_COUNT) - 1)
-
-#define SUB_PIXEL_PARALLAX_FACTOR_BIT_COUNT 4
-#define SUB_PIXEL_PARALLAX_FACTOR_BIT_OFFSET 0
-#define SUB_PIXEL_PARALLAX_FACTOR_BIT_MASK ((1 << SUB_PIXEL_PARALLAX_FACTOR_BIT_COUNT) - 1)
-
 // Minimum number of sample of each texel when reprojecting the high frequency history.
 // Value too low causes area of low sample count to have pixel dominating too much in reprojection causing input res aliasing on texture details
 #define MIN_REPROJECTION_SAMPLE_COUNT 2.0
@@ -1118,54 +1107,6 @@ tsr_half2 RejectionFactorToBlendFactor(tsr_half2 RejectionFactor)
 	return tsr_half(1.0) - RejectionFactor;
 }
 
-/** Payload of a subpixel depth information in the history */
-struct FTSRSubpixelDepth
-{
-	/** Depth of the pixel. */
-	float DeviceZ;
-	
-	/** Offset of the depth from the center of pixel. */
-	float2 Offset;
-	
-	/** Number of reprojected frame remaining before this depth is to be discarded. */
-	uint RemainingLifetime;
-};
-
-uint CompressSubpixelDepth(FTSRSubpixelDepth SubpixelDepth)
-{
-	uint CompressedSubpixelDepth = 0u;
-	CompressedSubpixelDepth |= f32tof16(SubpixelDepth.DeviceZ) << 17u;
-	CompressedSubpixelDepth |= uint(fastClamp(int(round(SubpixelDepth.Offset.x * 16.00 + 16.0)), 0, 31)) << 12u;
-	CompressedSubpixelDepth |= uint(fastClamp(int(round(SubpixelDepth.Offset.y * 16.00 + 16.0)), 0, 31)) <<  7u;
-	CompressedSubpixelDepth |= (SubpixelDepth.RemainingLifetime) << 0u;
-	return CompressedSubpixelDepth;
-}
-
-FTSRSubpixelDepth UncompressSubpixelDepth(uint CompressedSubpixelDepth)
-{
-	FTSRSubpixelDepth SubpixelDepth;
-	SubpixelDepth.DeviceZ  = f16tof32(CompressedSubpixelDepth >> 17u);
-	SubpixelDepth.Offset.x = float((CompressedSubpixelDepth >> 12u) & 0x1Fu) * rcp(16.0) - 1.0;
-	SubpixelDepth.Offset.y = float((CompressedSubpixelDepth >>  7u) & 0x1Fu) * rcp(16.0) - 1.0;
-	SubpixelDepth.RemainingLifetime = CompressedSubpixelDepth & 0x7Fu;
-	return SubpixelDepth;
-}
-
-void ForwardReprojectStaticDepth(float2 PrevScreenPosition, float PrevDeviceZ, out float2 ScreenPos, out float SceneDepth)
-{
-	float3 PreViewTranslationOffset = LWCToFloat(LWCSubtract(PrimaryView.PreViewTranslation, PrimaryView.PrevPreViewTranslation));
- 
-	// Compute position of the history in world space.
-	float PrevSceneDepth = ConvertFromDeviceZ(PrevDeviceZ);
-	float3 TranslatedWorldPosition = mul(float4(PrevScreenPosition * PrevSceneDepth, PrevSceneDepth, 1), View.PrevScreenToTranslatedWorld).xyz + PreViewTranslationOffset;
-	
-	// Reproject the world position into new viewport.
-	float4 ClipPos = mul(float4(TranslatedWorldPosition, 1.0), View.TranslatedWorldToClip);
-	
-	ScreenPos = ClipPos.xy / ClipPos.w;
-	SceneDepth = ClipPos.w;
-}
-
 FCatmullRomSamples GetBicubic2DCatmullRomSamples_Stubbe(float2 UV, float2 Size, in float2 InvSize)
 {
 	float2 PixelCoord = UV * Size;
diff --git a/Engine/Shaders/Private/TemporalSuperResolution/TSRDilateVelocity.usf b/Engine/Shaders/Private/TemporalSuperResolution/TSRDilateVelocity.usf
--- a/Engine/Shaders/Private/TemporalSuperResolution/TSRDilateVelocity.usf
+++ b/Engine/Shaders/Private/TemporalSuperResolution/TSRDilateVelocity.usf
@@ -24,8 +24,6 @@
 
 #define CONFIG_DEPTH_EDGE_DETECTION 1
 
-#define CONFIG_SUBPIXEL_DEPTH (DIM_SUBPIXEL_DEPTH)
-
 #define CONFIG_VELOCITY_EXTRAPOLATION 0
 
 
@@ -47,19 +45,15 @@ float2 PrevOutputBufferUVMax;
 float VelocityExtrapolationMultiplier;
 float InvFlickeringMaxParralaxVelocity;
 uint bOutputIsMovingTexture;
-uint bIncludeDynamicDepthInSubpixelDetails;
-uint SubpixelDepthLifetime;
 
 Texture2D<float> SceneDepthTexture;
 Texture2D SceneVelocityTexture;
-Texture2D<uint> PrevScatteredSubpixelDepthTexture;
 
 RWTexture2D<float2> DilatedVelocityOutput;
 RWTexture2D<float2> ClosestDepthOutput;
 globallycoherent RWTexture2D<uint> PrevUseCountOutput;
 globallycoherent RWTexture2D<uint> PrevClosestDepthOutput;
 RWTexture2DArray<uint> R8Output;
-RWTexture2D<uint> SubpixelDepthOutput;
 
 RWTexture2D<float3>	VelocityFlattenOutput;
 RWTexture2DArray<float4> VelocityTileArrayOutput;
@@ -79,10 +73,6 @@ void UnpackDepthAndOffset(uint Packed, out int2 Offset, out float DeviceZ)
 	Offset.y = int((Packed >> 2) & 0x3) - 1;
 }
 
-float MaxDeviceZWithSubpixel(float RenderedDeviceZ, float SubpixelDeviceZ)
-{
-	return max(RenderedDeviceZ, SubpixelDeviceZ);
-}
 
 //------------------------------------------------------- ENTRY POINT
 
@@ -109,14 +99,8 @@ void MainCS(
 
 	// Do all depth memory fetches.
 	float DeviceZ;
-	#if CONFIG_SUBPIXEL_DEPTH
-		uint PrevCompressedSubpixelDepth;
-	#endif
 	#if CONFIG_AA_CROSS_OPTIMISATION
 		float NeighborDeviceZArray[3];
-		#if CONFIG_SUBPIXEL_DEPTH
-			uint NeighborPrevCompressedSubpixelDepthArray[3];
-		#endif
 	#endif
 	ISOLATE
 	{
@@ -124,9 +108,6 @@ void MainCS(
 			int2 NeighborPixelPos = clamp(int2(InputPixelPos), InputPixelPosMin, InputPixelPosMax);
 				
 			DeviceZ = SceneDepthTexture[NeighborPixelPos];
-			#if CONFIG_SUBPIXEL_DEPTH
-				PrevCompressedSubpixelDepth = PrevScatteredSubpixelDepthTexture[NeighborPixelPos];
-			#endif
 		}
 		#if CONFIG_AA_CROSS_OPTIMISATION
 			UNROLL_N(3)
@@ -136,16 +117,9 @@ void MainCS(
 				NeighborPixelPos = clamp(NeighborPixelPos, InputPixelPosMin, InputPixelPosMax);
 
 				NeighborDeviceZArray[i] = SceneDepthTexture[NeighborPixelPos];
-				#if CONFIG_SUBPIXEL_DEPTH
-					NeighborPrevCompressedSubpixelDepthArray[i] = PrevScatteredSubpixelDepthTexture[NeighborPixelPos];
-				#endif
 			}
 		#endif
 	}
-
-	#if CONFIG_SUBPIXEL_DEPTH
-		float PrevScatteredDeviceZ = UncompressSubpixelDepth(PrevCompressedSubpixelDepth).DeviceZ;
-	#endif
 	
 	float ClosestDeviceZ = DeviceZ;
 	int2 VelocityPixelOffset = int2(0, 0);
@@ -166,19 +140,14 @@ void MainCS(
 		#endif
 
 
-		float CenterSubpixelDeviceZ = DeviceZ;
-		#if CONFIG_SUBPIXEL_DEPTH
-			CenterSubpixelDeviceZ = MaxDeviceZWithSubpixel(DeviceZ, PrevScatteredDeviceZ);
-		#endif
-
-		uint PackedDepthOffset = PackDepthAndOffset(/* Offset = */ int2(0, 0), CenterSubpixelDeviceZ);
+		uint PackedDepthOffset = PackDepthAndOffset(/* Offset = */ int2(0, 0), DeviceZ);
 
 		#if CONFIG_DEPTH_EDGE_DETECTION
 		float DeviceZError;
 		{
-			float Depth = ConvertFromDeviceZ(CenterSubpixelDeviceZ);
+			float Depth = ConvertFromDeviceZ(DeviceZ);
 			float ErrorCorrectedDeviceZ = ConvertToDeviceZ(GetDepthPixelRadiusForProjectionType(Depth) * 2.0f + Depth);
-			DeviceZError = abs(ErrorCorrectedDeviceZ - CenterSubpixelDeviceZ);
+			DeviceZError = abs(ErrorCorrectedDeviceZ - DeviceZ);
 		}
 		#endif
 
@@ -195,14 +164,6 @@ void MainCS(
 
 				float SampleDeviceZ0 = i == 3 ? WaveBroadcast(Horizontal, NeighborDeviceZArray[2]) : NeighborDeviceZArray[i];
 				float SampleDeviceZ1 = i == 3 ? WaveBroadcast(Vertical  , NeighborDeviceZArray[0]) : WaveBroadcast(CenterBroadcastArray[i], DeviceZ);
-				
-				#if CONFIG_SUBPIXEL_DEPTH
-					uint SamplePrevCompressedSubpixel0 = i == 3 ? WaveBroadcast(Horizontal, NeighborPrevCompressedSubpixelDepthArray[2]) : NeighborPrevCompressedSubpixelDepthArray[i];
-					uint SamplePrevCompressedSubpixel1 = i == 3 ? WaveBroadcast(Vertical  , NeighborPrevCompressedSubpixelDepthArray[0]) : WaveBroadcast(CenterBroadcastArray[i], PrevCompressedSubpixelDepth);
-					
-					float SampleSubpixelDeviceZ0 = UncompressSubpixelDepth(SamplePrevCompressedSubpixel0).DeviceZ;
-					float SampleSubpixelDeviceZ1 = UncompressSubpixelDepth(SamplePrevCompressedSubpixel1).DeviceZ;
-				#endif
 			#else // !CONFIG_AA_CROSS_OPTIMISATION
 				const int2 Offset = kOffsets[i];
 				
@@ -213,37 +174,24 @@ void MainCS(
 
 				float SampleDeviceZ0 = SceneDepthTexture[NeighborPixelPos0];
 				float SampleDeviceZ1 = SceneDepthTexture[NeighborPixelPos1];
-			
-				#if CONFIG_SUBPIXEL_DEPTH
-					float SampleSubpixelDeviceZ0 = UncompressSubpixelDepth(PrevScatteredSubpixelDepthTexture[NeighborPixelPos0]).DeviceZ;
-					float SampleSubpixelDeviceZ1 = UncompressSubpixelDepth(PrevScatteredSubpixelDepthTexture[NeighborPixelPos1]).DeviceZ;
-				#endif
 			#endif // !CONFIG_AA_CROSS_OPTIMISATION
 
-			#if CONFIG_SUBPIXEL_DEPTH
-				float CombinedSampleSubpixelDeviceZ0 = MaxDeviceZWithSubpixel(SampleDeviceZ0, SampleSubpixelDeviceZ0);
-				float CombinedSampleSubpixelDeviceZ1 = MaxDeviceZWithSubpixel(SampleDeviceZ1, SampleSubpixelDeviceZ1);
-			#else
-				float CombinedSampleSubpixelDeviceZ0 = SampleDeviceZ0;
-				float CombinedSampleSubpixelDeviceZ1 = SampleDeviceZ1;
-			#endif
-
 			#if CONFIG_DEPTH_EDGE_DETECTION
 			{
-				float DepthDiff = abs(CombinedSampleSubpixelDeviceZ1 - CombinedSampleSubpixelDeviceZ0);
-				float DepthVariation = (CombinedSampleSubpixelDeviceZ0 + CombinedSampleSubpixelDeviceZ1) * 0.5 - CenterSubpixelDeviceZ;
+				float DepthDiff = abs(SampleDeviceZ1 - SampleDeviceZ0);
+				float DepthVariation = (SampleDeviceZ0 + SampleDeviceZ1) * 0.5 - DeviceZ;
 
 				FLATTEN
 				if (DepthVariation > max(DepthDiff * 0.25, DeviceZError))
 				{
-					PackedDepthOffset = max(PackedDepthOffset, PackDepthAndOffset(+Offset, CombinedSampleSubpixelDeviceZ0));
-					PackedDepthOffset = max(PackedDepthOffset, PackDepthAndOffset(-Offset, CombinedSampleSubpixelDeviceZ1));
+					PackedDepthOffset = max(PackedDepthOffset, PackDepthAndOffset(+Offset, SampleDeviceZ0));
+					PackedDepthOffset = max(PackedDepthOffset, PackDepthAndOffset(-Offset, SampleDeviceZ1));
 				}
 			}
 			#else
 			{
-				PackedDepthOffset = max(PackedDepthOffset, PackDepthAndOffset(+Offset, CombinedSampleSubpixelDeviceZ0));
-				PackedDepthOffset = max(PackedDepthOffset, PackDepthAndOffset(-Offset, CombinedSampleSubpixelDeviceZ1));
+				PackedDepthOffset = max(PackedDepthOffset, PackDepthAndOffset(+Offset, SampleDeviceZ0));
+				PackedDepthOffset = max(PackedDepthOffset, PackDepthAndOffset(-Offset, SampleDeviceZ1));
 			}
 			#endif
 		}
@@ -276,11 +224,6 @@ void MainCS(
 
 		float ClosestDeviceZ1 = ClosestDeviceZ;
 		float ClosestDeviceZ2 = SceneDepthTexture[PixelCoord2];
-		#if CONFIG_SUBPIXEL_DEPTH
-			float ClosestSubpixelDeviceZ2 = UncompressSubpixelDepth(PrevScatteredSubpixelDepthTexture[PixelCoord2]).DeviceZ;
-			
-			ClosestDeviceZ2 = MaxDeviceZWithSubpixel(ClosestDeviceZ2, ClosestSubpixelDeviceZ2, SampleDeviceZError);
-		#endif
 
 		float4 EncodedVelocity1 = SceneVelocityTexture[PixelCoord1];
 		float4 EncodedVelocity2 = SceneVelocityTexture[PixelCoord2];
@@ -364,19 +307,6 @@ void MainCS(
 		IsMovingMask = select(HasPixelAnimation, tsr_half(1.0), IsMovingMask);
 	}
 
-	#if CONFIG_SUBPIXEL_DEPTH
-	uint CompressedSubpixelDepth = PrevCompressedSubpixelDepth;
-	{
-		FTSRSubpixelDepth PrevSubpixelDepth = UncompressSubpixelDepth(PrevCompressedSubpixelDepth);
-
-		FLATTEN
-		if (DeviceZ > PrevSubpixelDepth.DeviceZ || (bIsMovingPixel && !bIncludeDynamicDepthInSubpixelDetails))
-		{
-			CompressedSubpixelDepth = 0;
-		}
-	}
-	#endif
-
 	// Do motion blur velocity flatten.
 	#if DIM_MOTION_BLUR_DIRECTIONS > 0
 		FVelocityRange VelocityPolarRange;
@@ -437,11 +367,6 @@ void MainCS(
 		//DilatedVelocityOutput[OutputPixelPos] = EncodedVelocity.xy; //EncodeVelocityToTexture(float3(FinalScreenVelocity.xy, 0.0)).xy;
 		DilatedVelocityOutput[OutputPixelPos] = EncodeVelocityToTexture(float3(FinalScreenVelocity.xy, 0.0)).xy;
 		ClosestDepthOutput[OutputPixelPos] = float2(PrevClosestDeviceZ, ClosestDeviceZ);
-		#if CONFIG_SUBPIXEL_DEPTH
-		{
-			SubpixelDepthOutput[OutputPixelPos] = CompressedSubpixelDepth;
-		}
-		#endif
 
 		BRANCH
 		if (bOutputIsMovingTexture)
diff --git a/Engine/Shaders/Private/TemporalSuperResolution/TSRForwardScatterDepth.usf b/Engine/Shaders/Private/TemporalSuperResolution/TSRForwardScatterDepth.usf
deleted file mode 100644
diff --git a/Engine/Source/Runtime/Renderer/Private/PostProcess/TemporalSuperResolution.cpp b/Engine/Source/Runtime/Renderer/Private/PostProcess/TemporalSuperResolution.cpp
--- a/Engine/Source/Runtime/Renderer/Private/PostProcess/TemporalSuperResolution.cpp
+++ b/Engine/Source/Runtime/Renderer/Private/PostProcess/TemporalSuperResolution.cpp
@@ -283,28 +283,6 @@ TAutoConsoleVariable<float> CVarTSRWeightClampingPixelSpeed(
 //	TEXT("Defines how much the velocity should be extrapolated on geometric discontinuities (Default = 1.0f)."),
 //	ECVF_Scalability | ECVF_RenderThreadSafe);
 
-TAutoConsoleVariable<int32> CVarTSRSubpixelMethod(
-	TEXT("r.TSR.Subpixel.Method"), 0,
-	TEXT("One particular challenge of Nanite amount of details is that sometimes these details can be thiner than a rendering pixel in which case ")
-	TEXT("they only render in some frames. When that happens, it means neither depth or velocity buffer to be able to reproject them. This is for instance ")
-	TEXT("visible with the `vis SceneDepthZ` command.")
-	TEXT("\n")
-	TEXT("This settings control the method to reproject and/or discard the subpixel details.\n")
-	TEXT(" 0: disable subpixel details accumulation in history which means all this these subpixel features may ghost; \n")
-	TEXT(" 2: accumulate subpixel details' closest depth to be able to reproject them even when not drawing in depth/velocity buffer which works great for static geometry but not so much for any moving geometry (default)\n"),
-	ECVF_RenderThreadSafe);
-
-TAutoConsoleVariable<int32> CVarTSRSubpixelDepthMaxAge(
-	TEXT("r.TSR.Subpixel.DepthMaxAge"), 3,
-	TEXT("Maximum age in frames of subpixel's depth kept in history for their self reprojection (default to 3 frames)."),
-	ECVF_RenderThreadSafe);
-
-TAutoConsoleVariable<int32> CVarTSRSubpixelIncludeMovingDepth(
-	TEXT("r.TSR.Subpixel.IncludeMovingDepth"), 0,
-	TEXT("Whether the depth of moving subpixel detail should also be included in the subpixel depth history for their reprojection. This is a really bad idea to turn this on ")
-	TEXT("because it is impossible how a moving object's velocity involves overtime when it's only occasionally drawing its velocity. (disabled by default)."),
-	ECVF_RenderThreadSafe);
-
 #if !UE_BUILD_OPTIMIZED_SHOWFLAGS
 
 TAutoConsoleVariable<int32> CVarTSRVisualize(
@@ -359,14 +337,11 @@ BEGIN_SHADER_PARAMETER_STRUCT(FTSRHistoryTextures, )
 	SHADER_PARAMETER_RDG_TEXTURE(Texture2DArray, MetadataArray)
 	SHADER_PARAMETER_RDG_TEXTURE(Texture2DArray, GuideArray)
 	SHADER_PARAMETER_RDG_TEXTURE(Texture2DArray, MoireArray)
-	SHADER_PARAMETER_RDG_TEXTURE(Texture2D, SubpixelDepth)
 END_SHADER_PARAMETER_STRUCT()
 
 BEGIN_SHADER_PARAMETER_STRUCT(FTSRPrevHistoryParameters, )
 	SHADER_PARAMETER_STRUCT(FScreenPassTextureViewportParameters, PrevHistoryInfo)
 	SHADER_PARAMETER(FScreenTransform, ScreenPosToPrevHistoryBufferUV)
-	SHADER_PARAMETER(FScreenTransform, ScreenPosToPrevSubpixelDetails)
-	SHADER_PARAMETER(FVector2f, PrevSubpixelDetailsExtent)
 	SHADER_PARAMETER(float, HistoryPreExposureCorrection)
 	SHADER_PARAMETER(float, ResurrectionPreExposureCorrection)
 END_SHADER_PARAMETER_STRUCT()
@@ -379,13 +354,6 @@ enum class ETSRHistoryFormatBits : uint32
 };
 ENUM_CLASS_FLAGS(ETSRHistoryFormatBits);
 
-enum class ETSRSubpixelMethod
-{
-	Disabled,
-	ParallaxFactor,
-	ClosestDepth,
-};
-
 bool IsOutputDifferentThanHighFrequency(ETSRHistoryFormatBits HistoryFormatBits)
 {
 	return false;
@@ -501,37 +469,16 @@ class FTSRClearPrevTexturesCS : public FTSRShader
 
 		SHADER_PARAMETER_RDG_TEXTURE_UAV(RWTexture2D, PrevUseCountOutput)
 		SHADER_PARAMETER_RDG_TEXTURE_UAV(RWTexture2D, PrevClosestDepthOutput)
-		SHADER_PARAMETER_RDG_TEXTURE_UAV(RWTexture2D, SubpixelDepthOutput)
 	END_SHADER_PARAMETER_STRUCT()
 }; // class FTSRClearPrevTexturesCS
 
-class FTSRForwardScatterDepthCS : public FTSRShader
-{
-	DECLARE_GLOBAL_SHADER(FTSRForwardScatterDepthCS);
-	SHADER_USE_PARAMETER_STRUCT(FTSRForwardScatterDepthCS, FTSRShader);
-
-	BEGIN_SHADER_PARAMETER_STRUCT(FParameters, )
-		SHADER_PARAMETER_STRUCT_INCLUDE(FTSRCommonParameters, CommonParameters)
-		SHADER_PARAMETER(FIntPoint, PrevSubpixelDepthViewportMin)
-		SHADER_PARAMETER(FIntPoint, PrevSubpixelDepthViewportMax)
-		SHADER_PARAMETER(FScreenTransform, InputPixelPosToPrevScreenPosition)
-		SHADER_PARAMETER(FScreenTransform, ScreenPosToOutputPixelPos)
-
-		SHADER_PARAMETER_RDG_TEXTURE(Texture2D, PrevSubpixelDepthTexture)
-
-		SHADER_PARAMETER_RDG_TEXTURE_UAV(RWTexture2D, SubpixelDepthOutput)
-		SHADER_PARAMETER_RDG_TEXTURE_UAV(RWTexture2DArray, DebugOutput)
-	END_SHADER_PARAMETER_STRUCT()
-}; // class FTSRForwardScatterDepthCS
-
 class FTSRDilateVelocityCS : public FTSRShader
 {
 	DECLARE_GLOBAL_SHADER(FTSRDilateVelocityCS);
 	SHADER_USE_PARAMETER_STRUCT(FTSRDilateVelocityCS, FTSRShader);
 
 	class FMotionBlurDirectionsDim : SHADER_PERMUTATION_INT("DIM_MOTION_BLUR_DIRECTIONS", 3);
-	class FSubpixelDepthDim : SHADER_PERMUTATION_BOOL("DIM_SUBPIXEL_DEPTH");
-	using FPermutationDomain = TShaderPermutationDomain<FMotionBlurDirectionsDim, FSubpixelDepthDim>;
+	using FPermutationDomain = TShaderPermutationDomain<FMotionBlurDirectionsDim>;
 
 	BEGIN_SHADER_PARAMETER_STRUCT(FParameters, )
 		SHADER_PARAMETER_STRUCT_INCLUDE(FTSRCommonParameters, CommonParameters)
@@ -543,19 +490,15 @@ class FTSRDilateVelocityCS : public FTSRShader
 		SHADER_PARAMETER(float, VelocityExtrapolationMultiplier)
 		SHADER_PARAMETER(float, InvFlickeringMaxParralaxVelocity)
 		SHADER_PARAMETER(int32, bOutputIsMovingTexture)
-		SHADER_PARAMETER(int32, bIncludeDynamicDepthInSubpixelDetails)
-		SHADER_PARAMETER(int32, SubpixelDepthLifetime)
 
 		SHADER_PARAMETER_RDG_TEXTURE(Texture2D, SceneDepthTexture)
 		SHADER_PARAMETER_RDG_TEXTURE(Texture2D, SceneVelocityTexture)
-		SHADER_PARAMETER_RDG_TEXTURE(Texture2D, PrevScatteredSubpixelDepthTexture)
 
 		SHADER_PARAMETER_RDG_TEXTURE_UAV(RWTexture2D, DilatedVelocityOutput)
 		SHADER_PARAMETER_RDG_TEXTURE_UAV(RWTexture2D, ClosestDepthOutput)
 		SHADER_PARAMETER_RDG_TEXTURE_UAV(RWTexture2D, PrevUseCountOutput)
 		SHADER_PARAMETER_RDG_TEXTURE_UAV(RWTexture2D, PrevClosestDepthOutput)
 		SHADER_PARAMETER_RDG_TEXTURE_UAV(RWTexture2DArray, R8Output)
-		SHADER_PARAMETER_RDG_TEXTURE_UAV(RWTexture2DArray, SubpixelDepthOutput)
 
 		SHADER_PARAMETER_RDG_TEXTURE_UAV(RWTexture2D, VelocityFlattenOutput)
 		SHADER_PARAMETER_RDG_TEXTURE_UAV(RWTexture2DArray, VelocityTileArrayOutput)
@@ -1059,7 +1002,6 @@ class FTSRVisualizeCS : public FTSRShader
 
 IMPLEMENT_GLOBAL_SHADER(FTSRMeasureFlickeringLumaCS, "/Engine/Private/TemporalSuperResolution/TSRMeasureFlickeringLuma.usf", "MainCS", SF_Compute);
 IMPLEMENT_GLOBAL_SHADER(FTSRClearPrevTexturesCS,     "/Engine/Private/TemporalSuperResolution/TSRClearPrevTextures.usf",     "MainCS", SF_Compute);
-IMPLEMENT_GLOBAL_SHADER(FTSRForwardScatterDepthCS,   "/Engine/Private/TemporalSuperResolution/TSRForwardScatterDepth.usf",   "MainCS", SF_Compute);
 IMPLEMENT_GLOBAL_SHADER(FTSRDilateVelocityCS,        "/Engine/Private/TemporalSuperResolution/TSRDilateVelocity.usf",        "MainCS", SF_Compute);
 IMPLEMENT_GLOBAL_SHADER(FTSRDecimateHistoryCS,       "/Engine/Private/TemporalSuperResolution/TSRDecimateHistory.usf",       "MainCS", SF_Compute);
 IMPLEMENT_GLOBAL_SHADER(FTSRRejectShadingCS,         "/Engine/Private/TemporalSuperResolution/TSRRejectShading.usf",         "MainCS", SF_Compute);
@@ -1277,9 +1219,6 @@ FDefaultTemporalUpscaler::FOutputs AddTemporalSuperResolutionPasses(
 {
 	const FTSRHistory& InputHistory = View.PrevViewInfo.TSRHistory;
 
-	// Maximum depth lifetime encodable in FTSRSubpixelDepth::RemainingLifetime
-	const int32 kMaxDepthLifetime = (1 << 7) - 1;
-
 	// Number of frames stored in the history.
 	FTSRHistorySliceSequence HistorySliceSequence;
 	if (CVarTSRResurrectionEnable.GetValueOnRenderThread())
@@ -1320,9 +1259,6 @@ FDefaultTemporalUpscaler::FOutputs AddTemporalSuperResolutionPasses(
 	}
 #endif // PLATFORM_DESKTOP
 
-	// Whether should accumulate sub pixel depth.
-	const ETSRSubpixelMethod SubpixelMethod = ETSRSubpixelMethod(FMath::Clamp(CVarTSRSubpixelMethod.GetValueOnRenderThread(), 0, 2));
-
 	// Whether alpha channel is supported.
 	const bool bSupportsAlpha = CVarTSRAlphaChannel.GetValueOnRenderThread() >= 0 ? (CVarTSRAlphaChannel.GetValueOnRenderThread() > 0) : IsPostProcessingWithAlphaChannelSupported();
 
@@ -1658,7 +1594,6 @@ FDefaultTemporalUpscaler::FOutputs AddTemporalSuperResolutionPasses(
 		PrevHistory.MetadataArray = BlackArrayDummy;
 		PrevHistory.GuideArray = BlackArrayDummy;
 		PrevHistory.MoireArray = BlackArrayDummy;
-		PrevHistory.SubpixelDepth = nullptr;
 
 		if (HistorySliceSequence.GetRollingIndexCount() > 1)
 		{
@@ -1685,11 +1620,6 @@ FDefaultTemporalUpscaler::FOutputs AddTemporalSuperResolutionPasses(
 			? GraphBuilder.RegisterExternalTexture(InputHistory.MoireArray)
 			: BlackDummy;
 
-		// Register non-filterable history
-		PrevHistory.SubpixelDepth = InputHistory.SubpixelDepth.IsValid()
-			? GraphBuilder.RegisterExternalTexture(InputHistory.SubpixelDepth, TEXT("TSR.PrevHistory.SubpixelDepth"))
-			: nullptr;
-
 		int32 ResurrectionFrameRollingIndex = 0;
 		int32 PrevFrameRollingIndex = 0;
 		if (PrevHistorySliceSequence.GetRollingIndexCount() == 1)
@@ -1743,17 +1673,6 @@ FDefaultTemporalUpscaler::FOutputs AddTemporalSuperResolutionPasses(
 		ResurrectionGuideViewport = FScreenPassTextureViewport(PrevHistory.GuideArray->Desc.Extent, InputHistory.InputViewportRects[ResurrectionFrameSliceIndex]);
 	}
 
-	// Allocate non-reusuable history resources
-	if (SubpixelMethod == ETSRSubpixelMethod::ClosestDepth)
-	{
-		FRDGTextureDesc SubpixelDepthDesc = FRDGTextureDesc::Create2D(
-			InputExtent,
-			PF_R32_UINT,
-			FClearValueBinding::None,
-			TexCreate_ShaderResource | TexCreate_UAV);
-		History.SubpixelDepth = GraphBuilder.CreateTexture(SubpixelDepthDesc, TEXT("TSR.History.SubpixelDepth"));
-	}
-
 	// Setup the shader parameters for previous frame history
 	FTSRPrevHistoryParameters PrevHistoryParameters;
 	{
@@ -1775,7 +1694,6 @@ FDefaultTemporalUpscaler::FOutputs AddTemporalSuperResolutionPasses(
 	// Clear atomic scattered texture.
 	FRDGTextureRef PrevUseCountTexture;
 	FRDGTextureRef PrevClosestDepthTexture;
-	FRDGTextureRef PrevScatteredSubpixelDepthTexture = nullptr;
 	{
 		{
 			FRDGTextureDesc Desc = FRDGTextureDesc::Create2D(
@@ -1792,18 +1710,6 @@ FDefaultTemporalUpscaler::FOutputs AddTemporalSuperResolutionPasses(
 		PassParameters->CommonParameters = CommonParameters;
 		PassParameters->PrevUseCountOutput = GraphBuilder.CreateUAV(PrevUseCountTexture);
 		PassParameters->PrevClosestDepthOutput = GraphBuilder.CreateUAV(PrevClosestDepthTexture);
-		if (History.SubpixelDepth && PrevHistory.SubpixelDepth)
-		{
-			FRDGTextureDesc Desc = History.SubpixelDepth->Desc;
-			Desc.Flags |= TexCreate_AtomicCompatible;
-
-			PrevScatteredSubpixelDepthTexture = GraphBuilder.CreateTexture(Desc, TEXT("TSR.PrevScatteredSubpixelDepth"));
-			PassParameters->SubpixelDepthOutput = GraphBuilder.CreateUAV(PrevScatteredSubpixelDepthTexture);
-		}
-		else
-		{
-			PassParameters->SubpixelDepthOutput = CreateDummyUAV(GraphBuilder, PF_R32_UINT);
-		}
 
 		TShaderMapRef<FTSRClearPrevTexturesCS> ComputeShader(View.ShaderMap);
 		FComputeShaderUtils::AddPass(
@@ -1815,36 +1721,6 @@ FDefaultTemporalUpscaler::FOutputs AddTemporalSuperResolutionPasses(
 			FComputeShaderUtils::GetGroupCount(InputRect.Size(), 8 * 2));
 	}
 
-	// Forward scatter previous frame sub pixel depths.
-	if (PrevScatteredSubpixelDepthTexture)
-	{
-		FIntRect ScatterDepthViewport = InputHistory.InputViewportRect;
-
-		FTSRForwardScatterDepthCS::FParameters* PassParameters = GraphBuilder.AllocParameters<FTSRForwardScatterDepthCS::FParameters>();
-		PassParameters->CommonParameters = CommonParameters;
-		PassParameters->PrevSubpixelDepthViewportMin = ScatterDepthViewport.Min;
-		PassParameters->PrevSubpixelDepthViewportMax = ScatterDepthViewport.Max - 1;
-		PassParameters->InputPixelPosToPrevScreenPosition = (FScreenTransform::Identity + 0.5f) * FScreenTransform::ChangeTextureBasisFromTo(
-			PrevHistory.SubpixelDepth->Desc.Extent, ScatterDepthViewport,
-			FScreenTransform::ETextureBasis::TexelPosition, FScreenTransform::ETextureBasis::ScreenPosition);
-		PassParameters->ScreenPosToOutputPixelPos = FScreenTransform::ChangeTextureBasisFromTo(
-			InputExtent, InputRect,
-			FScreenTransform::ETextureBasis::ScreenPosition, FScreenTransform::ETextureBasis::TexelPosition);
-
-		PassParameters->PrevSubpixelDepthTexture = PrevHistory.SubpixelDepth;
-		PassParameters->SubpixelDepthOutput = GraphBuilder.CreateUAV(PrevScatteredSubpixelDepthTexture);
-		PassParameters->DebugOutput = CreateDebugUAV(InputExtent, TEXT("Debug.TSR.ForwardScatterDepth"));
-
-		TShaderMapRef<FTSRForwardScatterDepthCS> ComputeShader(View.ShaderMap);
-		FComputeShaderUtils::AddPass(
-			GraphBuilder,
-			RDG_EVENT_NAME("TSR ForwardScatterDepth %dx%d", ScatterDepthViewport.Width(), ScatterDepthViewport.Height()),
-			AsyncComputePasses >= 1 ? ERDGPassFlags::AsyncCompute : ERDGPassFlags::Compute,
-			ComputeShader,
-			PassParameters,
-			FComputeShaderUtils::GetGroupCount(ScatterDepthViewport.Size(), 8));
-	}
-
 	// Dilate the velocity texture & scatter reprojection into previous frame
 	FRDGTextureRef DilatedVelocityTexture;
 	FRDGTextureRef ClosestDepthTexture;
@@ -1892,7 +1768,6 @@ FDefaultTemporalUpscaler::FOutputs AddTemporalSuperResolutionPasses(
 
 		int32 TileSize = 8;
 		FTSRDilateVelocityCS::FPermutationDomain PermutationVector;
-		PermutationVector.Set<FTSRDilateVelocityCS::FSubpixelDepthDim>(SubpixelMethod == ETSRSubpixelMethod::ClosestDepth);
 
 		FTSRDilateVelocityCS::FParameters* PassParameters = GraphBuilder.AllocParameters<FTSRDilateVelocityCS::FParameters>();
 		PassParameters->CommonParameters = CommonParameters;
@@ -1905,29 +1780,15 @@ FDefaultTemporalUpscaler::FOutputs AddTemporalSuperResolutionPasses(
 			PassParameters->InvFlickeringMaxParralaxVelocity = 1.0f / FlickeringMaxParralaxVelocity;
 		}
 		PassParameters->bOutputIsMovingTexture = bOutputIsMovingTexture;
-		PassParameters->bIncludeDynamicDepthInSubpixelDetails = SubpixelMethod == ETSRSubpixelMethod::ClosestDepth && CVarTSRSubpixelIncludeMovingDepth.GetValueOnRenderThread();
-		PassParameters->SubpixelDepthLifetime = FMath::Clamp(CVarTSRSubpixelDepthMaxAge.GetValueOnRenderThread(), 1, kMaxDepthLifetime);
-
+		
 		PassParameters->SceneDepthTexture = PassInputs.SceneDepth.Texture;
 		PassParameters->SceneVelocityTexture = PassInputs.SceneVelocity.Texture;
-		if (PrevScatteredSubpixelDepthTexture)
-		{
-			PassParameters->PrevScatteredSubpixelDepthTexture = PrevScatteredSubpixelDepthTexture;
-		}
-		else
-		{
-			PassParameters->PrevScatteredSubpixelDepthTexture = BlackUintDummy;
-		}
 
 		PassParameters->DilatedVelocityOutput = GraphBuilder.CreateUAV(DilatedVelocityTexture);
 		PassParameters->ClosestDepthOutput = GraphBuilder.CreateUAV(ClosestDepthTexture);
 		PassParameters->PrevUseCountOutput = GraphBuilder.CreateUAV(PrevUseCountTexture);
 		PassParameters->PrevClosestDepthOutput = GraphBuilder.CreateUAV(PrevClosestDepthTexture);
 		PassParameters->R8Output = GraphBuilder.CreateUAV(R8OutputTexture);
-		if (SubpixelMethod == ETSRSubpixelMethod::ClosestDepth)
-		{
-			PassParameters->SubpixelDepthOutput = GraphBuilder.CreateUAV(History.SubpixelDepth);
-		}
 
 		// Setup up the motion blur's velocity flatten pass.
 		if (PassInputs.bGenerateVelocityFlattenTextures)
@@ -1970,11 +1831,10 @@ FDefaultTemporalUpscaler::FOutputs AddTemporalSuperResolutionPasses(
 		TShaderMapRef<FTSRDilateVelocityCS> ComputeShader(View.ShaderMap, PermutationVector);
 		FComputeShaderUtils::AddPass(
 			GraphBuilder,
-			RDG_EVENT_NAME("TSR DilateVelocity(#%d MotionBlurDirections=%d%s%s) %dx%d",
+			RDG_EVENT_NAME("TSR DilateVelocity(#%d MotionBlurDirections=%d%s) %dx%d",
 				PermutationVector.ToDimensionValueId(),
 				int32(PermutationVector.Get<FTSRDilateVelocityCS::FMotionBlurDirectionsDim>()),
 				bOutputIsMovingTexture ? TEXT(" OutputIsMoving") : TEXT(""),
-				PermutationVector.Get<FTSRDilateVelocityCS::FSubpixelDepthDim>() ? TEXT(" SubpixelDepth") : TEXT(""),
 				InputRect.Width(), InputRect.Height()),
 			AsyncComputePasses >= 2 ? ERDGPassFlags::AsyncCompute : ERDGPassFlags::Compute,
 			ComputeShader,
@@ -2603,12 +2463,6 @@ FDefaultTemporalUpscaler::FOutputs AddTemporalSuperResolutionPasses(
 		GraphBuilder.QueueTextureExtraction(History.ColorArray, &OutputHistory.ColorArray);
 		GraphBuilder.QueueTextureExtraction(History.MetadataArray, &OutputHistory.MetadataArray);
 
-		// Extract non-filterable history
-		if (History.SubpixelDepth)
-		{
-			GraphBuilder.QueueTextureExtraction(History.SubpixelDepth, &OutputHistory.SubpixelDepth);
-		}
-
 		// Extract history guide
 		GraphBuilder.QueueTextureExtraction(History.GuideArray, &OutputHistory.GuideArray);
 
diff --git a/Engine/Source/Runtime/Renderer/Private/RendererScene.cpp b/Engine/Source/Runtime/Renderer/Private/RendererScene.cpp
--- a/Engine/Source/Runtime/Renderer/Private/RendererScene.cpp
+++ b/Engine/Source/Runtime/Renderer/Private/RendererScene.cpp
@@ -822,7 +822,6 @@ uint64 FTSRHistory::GetGPUSizeBytes(bool bLogSizes) const
 	return
 		GetRenderTargetGPUSizeBytes(ColorArray, bLogSizes) +
 		GetRenderTargetGPUSizeBytes(MetadataArray, bLogSizes) +
-		GetRenderTargetGPUSizeBytes(SubpixelDepth, bLogSizes) +
 		GetRenderTargetGPUSizeBytes(GuideArray, bLogSizes) +
 		GetRenderTargetGPUSizeBytes(MoireArray, bLogSizes);
 }
diff --git a/Engine/Source/Runtime/Renderer/Private/SceneRendering.h b/Engine/Source/Runtime/Renderer/Private/SceneRendering.h
--- a/Engine/Source/Runtime/Renderer/Private/SceneRendering.h
+++ b/Engine/Source/Runtime/Renderer/Private/SceneRendering.h
@@ -938,7 +938,6 @@ struct FTSRHistory
 	TRefCountPtr<IPooledRenderTarget> MetadataArray;
 
 	// Input resolution representation of the output
-	TRefCountPtr<IPooledRenderTarget> SubpixelDepth;
 	TRefCountPtr<IPooledRenderTarget> GuideArray;
 	TRefCountPtr<IPooledRenderTarget> MoireArray;
 
