HWRT: Deprecate CPU instance culling
* Remove code related to CPU culling when possible
* r.RayTracing.Culling.UseGPU has been removed and now instance culling can only happen on the GPU which is driven by r.RayTracing.Culling.PerInstance

diff --git a/Engine/Source/Runtime/RHI/Public/RHIResources.h b/Engine/Source/Runtime/RHI/Public/RHIResources.h
--- a/Engine/Source/Runtime/RHI/Public/RHIResources.h
+++ b/Engine/Source/Runtime/RHI/Public/RHIResources.h
@@ -2847,7 +2847,16 @@ class FRHIRayTracingGeometry;
 * All instances covered by this descriptor will share shader bindings, but may have different transforms and user data.
 */
 struct FRayTracingGeometryInstance
-{	
+{
+	PRAGMA_DISABLE_DEPRECATION_WARNINGS
+	FRayTracingGeometryInstance() = default;
+	FRayTracingGeometryInstance(const FRayTracingGeometryInstance&) = default;
+	FRayTracingGeometryInstance& operator=(const FRayTracingGeometryInstance&) = default;
+	FRayTracingGeometryInstance(FRayTracingGeometryInstance&&) = default;
+	FRayTracingGeometryInstance& operator=(FRayTracingGeometryInstance&&) = default;
+	~FRayTracingGeometryInstance() = default;
+	PRAGMA_ENABLE_DEPRECATION_WARNINGS
+
 	FRHIRayTracingGeometry* GeometryRHI = nullptr;
 
 	// A single physical mesh may be duplicated many times in the scene with different transforms and user data.
@@ -2875,6 +2884,7 @@ struct FRayTracingGeometryInstance
 	TArrayView<const uint32> UserData;
 
 	// Each geometry copy can have one bit to make it individually deactivated (removed from TLAS while maintaining hit group indexing). Useful for culling.
+	UE_DEPRECATED(5.4, "ActivationMask has been deprecated.")
 	TArrayView<const uint32> ActivationMask;
 
 	// Whether local bounds scale and center translation should be applied to the instance transform.
diff --git a/Engine/Source/Runtime/Renderer/Private/PrimitiveSceneInfo.cpp b/Engine/Source/Runtime/Renderer/Private/PrimitiveSceneInfo.cpp
--- a/Engine/Source/Runtime/Renderer/Private/PrimitiveSceneInfo.cpp
+++ b/Engine/Source/Runtime/Renderer/Private/PrimitiveSceneInfo.cpp
@@ -1170,11 +1170,6 @@ void FPrimitiveSceneInfo::UpdateCachedRayTracingInstance(FPrimitiveSceneInfo* Sc
 		SceneInfo->CachedRayTracingInstance.NumTransforms = CachedRayTracingInstance.NumTransforms;
 		SceneInfo->CachedRayTracingInstance.BaseInstanceSceneDataOffset = SceneInfo->GetInstanceSceneDataOffset();
 
-		SceneInfo->CachedRayTracingInstanceWorldBounds.Empty();
-		SceneInfo->CachedRayTracingInstanceWorldBounds.AddUninitialized(CachedRayTracingInstance.NumTransforms);
-
-		SceneInfo->UpdateCachedRayTracingInstanceWorldBounds(SceneProxy->GetLocalToWorld());
-
 		SceneInfo->CachedRayTracingGeometry = CachedRayTracingInstance.Geometry;
 
 		if (Nanite::GetRayTracingMode() != Nanite::ERayTracingMode::Fallback && SceneProxy->IsNaniteMesh())
@@ -1233,36 +1228,6 @@ void FPrimitiveSceneInfo::RemoveCachedRayTracingPrimitives()
 		CachedRayTracingMeshCommandsHashPerLOD.Empty();
 	}
 }
-
-void FPrimitiveSceneInfo::UpdateCachedRayTracingInstanceWorldBounds(const FMatrix& NewPrimitiveLocalToWorld)
-{
-	QUICK_SCOPE_CYCLE_COUNTER(STAT_UpdateCachedRayTracingInstanceWorldBounds);
-	TRACE_CPUPROFILER_EVENT_SCOPE(UpdateCachedRayTracingInstanceWorldBounds);
-
-	if (CachedRayTracingInstanceWorldBounds.IsEmpty())
-	{
-		return;
-	}
-	
-	SmallestRayTracingInstanceWorldBoundsIndex = 0;
-
-	const FInstanceSceneDataBuffers *InstanceSceneDataBuffers = GetInstanceSceneDataBuffers();
-	if (InstanceSceneDataBuffers && InstanceSceneDataBuffers->GetNumInstances() > 0)
-	{
-		for (int32 Index = 0; Index < CachedRayTracingInstanceWorldBounds.Num(); Index++)
-		{
-			FBoxSphereBounds WorldSpaceBounds = InstanceSceneDataBuffers->GetInstanceWorldBounds(Index);
-			CachedRayTracingInstanceWorldBounds[Index] = WorldSpaceBounds;
-			SmallestRayTracingInstanceWorldBoundsIndex = WorldSpaceBounds.SphereRadius < CachedRayTracingInstanceWorldBounds[SmallestRayTracingInstanceWorldBoundsIndex].SphereRadius ? Index : SmallestRayTracingInstanceWorldBoundsIndex;
-		}
-	}
-	else
-	{
-		check(CachedRayTracingInstanceWorldBounds.Num() == 1);
-		CachedRayTracingInstanceWorldBounds[0] = Proxy->GetBounds();
-		SmallestRayTracingInstanceWorldBoundsIndex = 0;
-	}
-}
 #endif
 
 void FPrimitiveSceneInfo::AddStaticMeshes(FScene* Scene, TArrayView<FPrimitiveSceneInfo*> SceneInfos, bool bCacheMeshDrawCommands)
diff --git a/Engine/Source/Runtime/Renderer/Private/RayTracing/RayTracing.cpp b/Engine/Source/Runtime/Renderer/Private/RayTracing/RayTracing.cpp
--- a/Engine/Source/Runtime/Renderer/Private/RayTracing/RayTracing.cpp
+++ b/Engine/Source/Runtime/Renderer/Private/RayTracing/RayTracing.cpp
@@ -874,20 +874,11 @@ namespace RayTracing
 
 				TRACE_CPUPROFILER_EVENT_SCOPE(RayTracingSceneStaticInstanceTask);
 
-				FGraphEventArray CullingTasks;
-
 				const bool bAutoInstance = CVarRayTracingAutoInstance.GetValueOnRenderThread() != 0;
 
 				// Instance batches by FRelevantPrimitive::InstancingKey()
 				Experimental::TSherwoodMap<uint64, FAutoInstanceBatch> InstanceBatches;
 
-				TArray<FRayTracingCullPrimitiveInstancesClosure> CullInstancesClosures;
-				if (CullingParameters.CullingMode != RayTracing::ECullingMode::Disabled && GetRayTracingCullingPerInstance())
-				{
-					CullInstancesClosures.Reserve(RelevantStaticPrimitives.Num());
-					CullingTasks.Reserve(RelevantStaticPrimitives.Num() / 256 + 1);
-				}
-
 				// scan relevant primitives computing hash data to look for duplicate instances
 				for (const FRelevantPrimitive& RelevantPrimitive : RelevantStaticPrimitives)
 				{
@@ -936,39 +927,6 @@ namespace RayTracing
 
 							NewInstance.LayerIndex = (uint8)(RelevantPrimitive.bAnySegmentsDecal && !bNeedSeparateDecalInstance ? ERayTracingSceneLayer::Decals : ERayTracingSceneLayer::Base);
 
-							const Experimental::FHashElementId GroupId = Scene.PrimitiveRayTracingGroupIds[PrimitiveIndex];
-							const bool bUseGroupBounds = CullingParameters.bCullUsingGroupIds && GroupId.IsValid();
-
-							if (CullingParameters.CullingMode != RayTracing::ECullingMode::Disabled && GetRayTracingCullingPerInstance() && RelevantPrimitive.CachedRayTracingInstance->NumTransforms > 1 && !bUseGroupBounds && !CullingParameters.bUseGPUInstanceCulling)
-							{
-								const bool bIsFarFieldPrimitive = EnumHasAnyFlags(Flags, ERayTracingPrimitiveFlags::FarField);
-
-								TArrayView<uint32> InstanceActivationMask = RayTracingScene.Allocate<uint32>(FMath::DivideAndRoundUp(NewInstance.NumTransforms, 32u));
-
-								NewInstance.ActivationMask = InstanceActivationMask;
-
-								FRayTracingCullPrimitiveInstancesClosure Closure;
-								Closure.Scene = &Scene;
-								Closure.SceneInfo = SceneInfo;
-								Closure.PrimitiveIndex = PrimitiveIndex;
-								Closure.bIsFarFieldPrimitive = bIsFarFieldPrimitive;
-								Closure.CullingParameters = &CullingParameters;
-								Closure.OutInstanceActivationMask = InstanceActivationMask;
-
-								CullInstancesClosures.Add(MoveTemp(Closure));
-
-								if (CullInstancesClosures.Num() >= 256)
-								{
-									CullingTasks.Add(FFunctionGraphTask::CreateAndDispatchWhenReady([CullInstancesClosures = MoveTemp(CullInstancesClosures)]()
-										{
-											for (auto& Closure : CullInstancesClosures)
-											{
-												Closure();
-											}
-										}, TStatId(), nullptr, ENamedThreads::AnyThread));
-								}
-							}
-
 							if (bNeedSeparateDecalInstance && !GRayTracingExcludeDecals)
 							{
 								FRayTracingGeometryInstance DecalRayTracingInstance = NewInstance;
@@ -1130,22 +1088,6 @@ namespace RayTracing
 						}
 					}
 				}
-
-				if (!CullingParameters.bUseGPUInstanceCulling)
-				{
-					CullingTasks.Add(FFunctionGraphTask::CreateAndDispatchWhenReady([CullInstancesClosures = MoveTemp(CullInstancesClosures)]()
-						{
-							for (auto& Closure : CullInstancesClosures)
-							{
-								Closure();
-							}
-						}, TStatId(), nullptr, ENamedThreads::AnyThread));
-
-					for (FGraphEventRef& CullingTask : CullingTasks)
-					{
-						MyCompletionGraphEvent->DontCompleteUntil(CullingTask);
-					}
-				}
 			}
 		};
 
diff --git a/Engine/Source/Runtime/Renderer/Private/RayTracing/RayTracingInstanceBufferUtil.cpp b/Engine/Source/Runtime/Renderer/Private/RayTracing/RayTracingInstanceBufferUtil.cpp
--- a/Engine/Source/Runtime/Renderer/Private/RayTracing/RayTracingInstanceBufferUtil.cpp
+++ b/Engine/Source/Runtime/Renderer/Private/RayTracing/RayTracingInstanceBufferUtil.cpp
@@ -182,9 +182,7 @@ void FillRayTracingInstanceUploadBuffer(
 			const uint32 NumTransforms = SceneInstance.NumTransforms;
 
 			checkf(SceneInstance.UserData.Num() == 0 || SceneInstance.UserData.Num() >= int32(NumTransforms),
-				TEXT("User data array must be either be empty (Instance.DefaultUserData is used), or contain one entry per entry in Transforms array."));
-
-			check(SceneInstance.ActivationMask.IsEmpty() || SceneInstance.ActivationMask.Num() * 32 >= int32(NumTransforms));
+				TEXT("User data array must be either be empty (Instance.DefaultUserData is used), or contain one entry per entry in Transforms array."));			
 
 			const bool bUseUniqueUserData = SceneInstance.UserData.Num() != 0;
 
@@ -243,12 +241,7 @@ void FillRayTracingInstanceUploadBuffer(
 				checkf(InstanceDesc.InstanceId <= 0xFFFFFF, TEXT("InstanceId must fit in 24 bits."));
 				checkf(InstanceDesc.InstanceContributionToHitGroupIndex <= 0xFFFFFF, TEXT("InstanceContributionToHitGroupIndex must fit in 24 bits."));
 
-				if (!SceneInstance.ActivationMask.IsEmpty() && (SceneInstance.ActivationMask[TransformIndex / 32] & (1 << (TransformIndex % 32))) == 0)
-				{
-					// Set flag for deactivated instances
-					InstanceDesc.AccelerationStructureIndex = 0xFFFFFFFF;
-					NumInactiveNativeInstancesThisSceneInstance++;
-				} else if (bCpuInstance)
+				if (bCpuInstance)
 				{
 					const uint32 TransformDataOffset = InstanceDesc.GPUSceneInstanceOrTransformIndex * 3;
 					FMatrix LocalToTranslatedWorld = SceneInstance.Transforms[TransformIndex].ConcatTranslation(PreViewTranslation);
diff --git a/Engine/Source/Runtime/Renderer/Private/RayTracing/RayTracingInstanceCulling.cpp b/Engine/Source/Runtime/Renderer/Private/RayTracing/RayTracingInstanceCulling.cpp
--- a/Engine/Source/Runtime/Renderer/Private/RayTracing/RayTracingInstanceCulling.cpp
+++ b/Engine/Source/Runtime/Renderer/Private/RayTracing/RayTracingInstanceCulling.cpp
@@ -47,12 +47,6 @@ static TAutoConsoleVariable<int32> CVarRayTracingCullingGroupIds(
 	TEXT("Cull using aggregate ray tracing group id bounds when defined instead of primitive or instance bounds."),
 	ECVF_RenderThreadSafe);
 
-static TAutoConsoleVariable<int32> CVarRayTracingCullingUseGPU(
-	TEXT("r.RayTracing.Culling.UseGPU"),
-	1,
-	TEXT("Cull instances using GPU. Requires per instance culling enabled (default=1)."),
-	ECVF_RenderThreadSafe);
-
 RayTracing::ECullingMode RayTracing::GetCullingMode(const FEngineShowFlags& ShowFlags)
 {
 	// Disable culling if path tracer is used, so that path tracer matches raster view
@@ -69,11 +63,6 @@ float GetRayTracingCullingRadius()
 	return CVarRayTracingCullingRadius.GetValueOnRenderThread();
 }
 
-int32 GetRayTracingCullingPerInstance()
-{
-	return CVarRayTracingCullingPerInstance.GetValueOnRenderThread();
-}
-
 void FRayTracingCullingParameters::Init(FViewInfo& View)
 {
 	CullingMode = RayTracing::GetCullingMode(View.Family->EngineShowFlags);
@@ -90,7 +79,7 @@ void FRayTracingCullingParameters::Init(FViewInfo& View)
 	bIsRayTracingFarField = Lumen::UseFarField(*View.Family);
 	bCullUsingGroupIds = CVarRayTracingCullingGroupIds.GetValueOnRenderThread() != 0;
 	bCullMinDrawDistance = CVarRayTracingCullingUseMinDrawDistance.GetValueOnRenderThread() != 0;
-	bUseGPUInstanceCulling = CVarRayTracingCullingUseGPU.GetValueOnRenderThread() != 0 && GetRayTracingCullingPerInstance() && bCullAllObjects;
+	bUseInstanceCulling = CVarRayTracingCullingPerInstance.GetValueOnRenderThread() != 0 && bCullAllObjects;
 }
 
 namespace RayTracing
@@ -193,78 +182,6 @@ bool ShouldCullBounds(const FRayTracingCullingParameters& CullingParameters, con
 	return false;
 }
 
-bool ShouldSkipPerInstanceCullingForPrimitive(const FRayTracingCullingParameters& CullingParameters, FBoxSphereBounds ObjectBounds, FBoxSphereBounds SmallestInstanceBounds, bool bIsFarFieldPrimitive)
-{
-	bool bSkipCulling = false;
-
-	const float ObjectRadius = ObjectBounds.SphereRadius;
-	const FVector ObjectCenter = ObjectBounds.Origin;
-	const FVector CameraToObjectCenter = FVector(ObjectCenter - CullingParameters.ViewOrigin);
-
-	const FVector CameraToFurthestInstanceCenter = CameraToObjectCenter * (CameraToObjectCenter.Size() + ObjectRadius + SmallestInstanceBounds.SphereRadius) / CameraToObjectCenter.Size();
-
-	const bool bConsiderCulling = CullingParameters.bCullAllObjects || FVector::DotProduct(CullingParameters.ViewDirection, CameraToObjectCenter) < -ObjectRadius;
-
-	if (bConsiderCulling)
-	{
-		const float CameraToObjectCenterLength = CameraToObjectCenter.Size();
-
-		if (bIsFarFieldPrimitive)
-		{
-			if (CameraToObjectCenterLength < (CullingParameters.FarFieldCullingRadius - ObjectRadius))
-			{
-				bSkipCulling = true;
-			}
-		}
-		else
-		{
-			const bool bSkipDistanceCulling = CameraToObjectCenterLength < (CullingParameters.CullingRadius - ObjectRadius);
-
-			// Cull by solid angle: check the radius of bounding sphere against angle threshold
-			const bool bSkipAngleCulling = FMath::IsFinite(SmallestInstanceBounds.SphereRadius / CameraToFurthestInstanceCenter.Size()) && SmallestInstanceBounds.SphereRadius / CameraToFurthestInstanceCenter.Size() >= CullingParameters.AngleThresholdRatio;
-
-			if (CullingParameters.bCullByRadiusOrDistance)
-			{
-				if (bSkipDistanceCulling && bSkipAngleCulling)
-				{
-					bSkipCulling = true;
-				}
-			}
-			else if (bSkipDistanceCulling || bSkipAngleCulling)
-			{
-				bSkipCulling = true;
-			}
-		}
-	}
-	else
-	{
-		bSkipCulling = true;
-	}
-
-	return bSkipCulling;
-}
-
 } // namspace RayTracing
 
-void FRayTracingCullPrimitiveInstancesClosure::operator()() const
-{
-	FMemory::Memset(OutInstanceActivationMask.GetData(), 0xFF, OutInstanceActivationMask.Num() * 4);
-
-	checkf(!CullingParameters->bCullUsingGroupIds || !Scene->PrimitiveRayTracingGroupIds[PrimitiveIndex].IsValid(), TEXT("Shouldn't do instance level culling of primitives in raytracing groups."));
-
-	const FPrimitiveBounds& PrimitiveBounds = Scene->PrimitiveBounds[PrimitiveIndex];
-
-	if (!RayTracing::ShouldSkipPerInstanceCullingForPrimitive(*CullingParameters, PrimitiveBounds.BoxSphereBounds, SceneInfo->CachedRayTracingInstanceWorldBounds[SceneInfo->SmallestRayTracingInstanceWorldBoundsIndex], bIsFarFieldPrimitive))
-	{
-		for (int32 InstanceIndex = 0; InstanceIndex < SceneInfo->CachedRayTracingInstanceWorldBounds.Num(); InstanceIndex++)
-		{
-			const FBoxSphereBounds& InstanceBounds = SceneInfo->CachedRayTracingInstanceWorldBounds[InstanceIndex];
-			if (RayTracing::ShouldCullBounds(*CullingParameters, InstanceBounds, PrimitiveBounds.MinDrawDistance, bIsFarFieldPrimitive))
-			{
-				OutInstanceActivationMask[InstanceIndex / 32] &= ~(1 << (InstanceIndex % 32));
-			}
-		}
-	}
-}
-
 #endif
diff --git a/Engine/Source/Runtime/Renderer/Private/RayTracing/RayTracingInstanceCulling.h b/Engine/Source/Runtime/Renderer/Private/RayTracing/RayTracingInstanceCulling.h
--- a/Engine/Source/Runtime/Renderer/Private/RayTracing/RayTracingInstanceCulling.h
+++ b/Engine/Source/Runtime/Renderer/Private/RayTracing/RayTracingInstanceCulling.h
@@ -16,7 +16,6 @@ struct FEngineShowFlags;
 struct FRayTracingCullingParameters;
 
 float GetRayTracingCullingRadius();
-int32 GetRayTracingCullingPerInstance();
 
 namespace RayTracing
 {
@@ -33,21 +32,6 @@ bool CullPrimitiveByFlags(const FRayTracingCullingParameters& CullingParameters,
 
 // Completely test if the bounds should be culled for ray tracing. This includes all configuration values.
 bool ShouldCullBounds(const FRayTracingCullingParameters& CullingParameters, const FBoxSphereBounds& ObjectBounds, float MinDrawDistance, bool bIsFarFieldPrimitive);
-
-bool ShouldSkipPerInstanceCullingForPrimitive(const FRayTracingCullingParameters& CullingParameters, FBoxSphereBounds ObjectBounds, FBoxSphereBounds SmallestInstanceBounds, bool bIsFarFieldPrimitive);
 }
 
-struct FRayTracingCullPrimitiveInstancesClosure
-{
-	const FScene* Scene;
-	int32 PrimitiveIndex;
-	const FPrimitiveSceneInfo* SceneInfo;
-	bool bIsFarFieldPrimitive;
-	TArrayView<uint32> OutInstanceActivationMask;
-
-	const FRayTracingCullingParameters* CullingParameters;
-
-	void operator()() const;
-};
-
 #endif
diff --git a/Engine/Source/Runtime/Renderer/Private/RayTracing/RayTracingScene.cpp b/Engine/Source/Runtime/Renderer/Private/RayTracing/RayTracingScene.cpp
--- a/Engine/Source/Runtime/Renderer/Private/RayTracing/RayTracingScene.cpp
+++ b/Engine/Source/Runtime/Renderer/Private/RayTracing/RayTracingScene.cpp
@@ -284,7 +284,7 @@ void FRayTracingScene::CreateWithInitializationData(FRDGBuilder& GraphBuilder, c
 					NumNativeGPUSceneInstances,
 					NumNativeCPUInstances,
 					GPUInstances,
-					CullingParameters.bUseGPUInstanceCulling ? &CullingParameters : nullptr,
+					CullingParameters.bUseInstanceCulling ? &CullingParameters : nullptr,
 					PassParams->DebugInstanceGPUSceneIndexBuffer ? PassParams->DebugInstanceGPUSceneIndexBuffer->GetRHI() : nullptr);
 			});
 	}
diff --git a/Engine/Source/Runtime/Renderer/Private/RendererScene.cpp b/Engine/Source/Runtime/Renderer/Private/RendererScene.cpp
--- a/Engine/Source/Runtime/Renderer/Private/RendererScene.cpp
+++ b/Engine/Source/Runtime/Renderer/Private/RendererScene.cpp
@@ -6154,13 +6154,6 @@ void FScene::UpdateAllPrimitiveSceneInfos(FRDGBuilder& GraphBuilder, EUpdateAllP
 
 			DistanceFieldSceneData.UpdatePrimitive(PrimitiveSceneInfo);
 			LumenUpdatePrimitive(PrimitiveSceneInfo);
-		#if RHI_RAYTRACING
-			// Don't do this for things with an instance update coming!
-			if (!UpdatedInstances.Find(PrimitiveSceneInfo->Proxy))
-			{
-				PrimitiveSceneInfo->UpdateCachedRayTracingInstanceWorldBounds(LocalToWorld);
-			}
-		#endif
 
 			// If the primitive has static mesh elements, it should have returned true from ShouldRecreateProxyOnUpdateTransform!
 			check(!(bUpdateStaticDrawLists && PrimitiveSceneInfo->StaticMeshes.Num()));
diff --git a/Engine/Source/Runtime/Renderer/Private/SceneRendering.h b/Engine/Source/Runtime/Renderer/Private/SceneRendering.h
--- a/Engine/Source/Runtime/Renderer/Private/SceneRendering.h
+++ b/Engine/Source/Runtime/Renderer/Private/SceneRendering.h
@@ -1173,7 +1173,7 @@ struct FRayTracingCullingParameters
 	bool bIsRayTracingFarField;
 	bool bCullUsingGroupIds;
 	bool bCullMinDrawDistance;
-	bool bUseGPUInstanceCulling;
+	bool bUseInstanceCulling;
 
 	void Init(FViewInfo& View);
 };
diff --git a/Engine/Source/Runtime/Renderer/Public/PrimitiveSceneInfo.h b/Engine/Source/Runtime/Renderer/Public/PrimitiveSceneInfo.h
--- a/Engine/Source/Runtime/Renderer/Public/PrimitiveSceneInfo.h
+++ b/Engine/Source/Runtime/Renderer/Public/PrimitiveSceneInfo.h
@@ -690,8 +690,6 @@ public:
 	TArray<uint64> CachedRayTracingMeshCommandsHashPerLOD;
 	// TODO: this should be placed in FRayTracingScene and we have a pointer/handle here. It's here for now for PoC
 	FRayTracingGeometryInstance CachedRayTracingInstance;
-	TArray<FBoxSphereBounds> CachedRayTracingInstanceWorldBounds;
-	int32 SmallestRayTracingInstanceWorldBoundsIndex;
 #endif
 
 private:
@@ -757,7 +755,8 @@ private:
 	void RemoveCachedRayTracingPrimitives();
 
 	/** Updates cached world bounds in CachedRayTracingInstance */
-	void UpdateCachedRayTracingInstanceWorldBounds(const FMatrix& NewPrimitiveLocalToWorld);
+	UE_DEPRECATED(5.4, "UpdateCachedRayTracingInstanceWorldBounds has been deprecated.")
+	void UpdateCachedRayTracingInstanceWorldBounds(const FMatrix& NewPrimitiveLocalToWorld) {};
 
 	/** Updates cached ray tracing instances. Utility closely mirrors CacheRayTracingPrimitives(..) */
 	static void UpdateCachedRayTracingInstances(FScene* Scene, const TArrayView<FPrimitiveSceneInfo*>& SceneInfos);
