Fix for ShaderPrint not working with OpenGL RHI. Uses of StructuredBuffer in vertex shaders were converted into SSBO accesses, which is essentially a UAV access. This is not supported on all platforms. Fixed by using typed buffers instead.

#rnx

diff --git a/Engine/Shaders/Private/ShaderPrintCommon.ush b/Engine/Shaders/Private/ShaderPrintCommon.ush
--- a/Engine/Shaders/Private/ShaderPrintCommon.ush
+++ b/Engine/Shaders/Private/ShaderPrintCommon.ush
@@ -135,6 +135,15 @@ void ClearCounters(RWStructuredBuffer<uint> InRWBuffer)
 	InRWBuffer[SHADER_PRINT_COUNTER_OFFSET_FREE] = 0;
 }
 
+void ClearCounters(RWBuffer<uint> InRWBuffer)
+{
+	// Symbol/print Counter
+	InRWBuffer[SHADER_PRINT_COUNTER_OFFSET_LINE] = 0;
+	InRWBuffer[SHADER_PRINT_COUNTER_OFFSET_TRIANGLE] = 0;
+	InRWBuffer[SHADER_PRINT_COUNTER_OFFSET_SYMBOL] = 0;
+	InRWBuffer[SHADER_PRINT_COUNTER_OFFSET_FREE] = 0;
+}
+
 ///////////////////////////////////////////////////////////////////////////////////////////////////
 // State 
 #define SHADER_PRINT_STATE_COUNT_OFFSET 0
@@ -265,7 +274,7 @@ FShaderPrintItem UnpackShaderPrintItem(FPackedShaderPrintItem In)
 }
 
 // /!\ If you update this function, update the WriteSymbol() version without explicit buffer, below in this file
-void WriteSymbol(uint Offset, FShaderPrintItem In, RWStructuredBuffer<uint> InRWBuffer)
+void WriteSymbol(uint Offset, FShaderPrintItem In, RWBuffer<uint> InRWBuffer)
 {
 	const uint Offset4 = GetSymbolOffset(Offset);
 
@@ -289,6 +298,19 @@ FShaderPrintItem ReadSymbol(uint Offset, StructuredBuffer<uint> InBuffer)
 	return UnpackShaderPrintItem(Packed);
 }
 
+FShaderPrintItem ReadSymbol(uint Offset, Buffer<uint> InBuffer)
+{
+	const uint Offset4 = GetSymbolOffset(Offset);
+
+	FPackedShaderPrintItem Packed = (FPackedShaderPrintItem)0;
+	Packed.ScreenPos16bits	= InBuffer[Offset4 + 0];
+	Packed.Value			= InBuffer[Offset4 + 1];
+	Packed.TypeAndColor		= InBuffer[Offset4 + 2];
+	Packed.Metadata			= InBuffer[Offset4 + 3];
+
+	return UnpackShaderPrintItem(Packed);
+}
+
 ///////////////////////////////////////////////////////////////////////////////////////////////////
 // Context structures.
 // Stores the current dynamic state, and the constant data describing config and buffer bindings.
@@ -745,7 +767,7 @@ FPackedTriangleElement PackTriangleElement(FTriangleElement In)
 	return Out;
 }
 
-FTriangleElement UnpackTriangleElement(StructuredBuffer<uint> InPrimitiveBuffer, uint InIndex, uint MaxCharacterCount, uint MaxLineCount)
+FTriangleElement UnpackTriangleElement(Buffer<uint> InPrimitiveBuffer, uint InIndex, uint MaxCharacterCount, uint MaxLineCount)
 {
 	const uint Index12 = GetPrimitiveTriangleOffset(InIndex, MaxCharacterCount, MaxLineCount);
 
@@ -949,7 +971,7 @@ FPackedLineElement PackLineElement(FLineElement In)
 	return Out;
 }
 
-FLineElement UnpackLineElement(StructuredBuffer<uint> InPrimitiveBuffer, uint InIndex, uint MaxCharacterCount)
+FLineElement UnpackLineElement(Buffer<uint> InPrimitiveBuffer, uint InIndex, uint MaxCharacterCount)
 {
 	const uint Index8 = GetPrimitiveLineOffset(InIndex, MaxCharacterCount);
 
diff --git a/Engine/Shaders/Private/ShaderPrintDraw.usf b/Engine/Shaders/Private/ShaderPrintDraw.usf
--- a/Engine/Shaders/Private/ShaderPrintDraw.usf
+++ b/Engine/Shaders/Private/ShaderPrintDraw.usf
@@ -15,10 +15,10 @@ uint FrameIndex;
 uint FrameThreshold;
 
 StructuredBuffer<uint> ValuesBuffer;
-StructuredBuffer<uint> SymbolsBuffer;
+Buffer<uint> SymbolsBuffer;
 
 RWStructuredBuffer<uint> RWValuesBuffer;
-RWStructuredBuffer<uint> RWSymbolsBuffer;
+RWBuffer<uint> RWSymbolsBuffer;
 RWStructuredBuffer<uint> RWStateBuffer;
 
 RWBuffer<uint> RWIndirectDispatchArgsBuffer;
diff --git a/Engine/Shaders/Private/ShaderPrintDrawPrimitive.usf b/Engine/Shaders/Private/ShaderPrintDrawPrimitive.usf
--- a/Engine/Shaders/Private/ShaderPrintDrawPrimitive.usf
+++ b/Engine/Shaders/Private/ShaderPrintDrawPrimitive.usf
@@ -9,7 +9,7 @@ uint bCheckerboardEnabled;
 uint bDrawOccludedLines;
 
 // The shader used to render the debug element on view
-StructuredBuffer<uint> ShaderDrawDebugPrimitive;
+Buffer<uint> ShaderDrawDebugPrimitive;
 
 Texture2D<float> DepthTexture;
 float2 OutputInvResolution; 
@@ -110,7 +110,7 @@ void ShaderDrawDebugPS(
 
 #if GPU_DEBUG_RENDERING_COPY_CS
 
-StructuredBuffer<uint> ElementBuffer;
+Buffer<uint> ElementBuffer;
 RWBuffer<uint> RWIndirectArgs;
 uint PrimitiveType;
 
diff --git a/Engine/Source/Runtime/Renderer/Private/ShaderPrint.cpp b/Engine/Source/Runtime/Renderer/Private/ShaderPrint.cpp
--- a/Engine/Source/Runtime/Renderer/Private/ShaderPrint.cpp
+++ b/Engine/Source/Runtime/Renderer/Private/ShaderPrint.cpp
@@ -452,7 +452,7 @@ namespace ShaderPrint
 		BEGIN_SHADER_PARAMETER_STRUCT(FParameters, )
 			SHADER_PARAMETER_STRUCT_REF(FShaderPrintCommonParameters, Common)
 			SHADER_PARAMETER_RDG_BUFFER_SRV(StructuredBuffer<uint>, ValuesBuffer)
-			SHADER_PARAMETER_RDG_BUFFER_UAV(RWStructuredBuffer<uint>, RWSymbolsBuffer)
+			SHADER_PARAMETER_RDG_BUFFER_UAV(RWBuffer<uint>, RWSymbolsBuffer)
 			SHADER_PARAMETER_RDG_BUFFER_UAV(RWBuffer<uint>, RWIndirectDispatchArgsBuffer)
 		END_SHADER_PARAMETER_STRUCT()
 
@@ -497,7 +497,7 @@ namespace ShaderPrint
 			SHADER_PARAMETER(uint32, FrameIndex)
 			SHADER_PARAMETER_STRUCT_REF(FShaderPrintCommonParameters, Common)
 			SHADER_PARAMETER_RDG_BUFFER_SRV(StructuredBuffer<uint>, ValuesBuffer)
-			SHADER_PARAMETER_RDG_BUFFER_UAV(RWStructuredBuffer<uint>, RWSymbolsBuffer)
+			SHADER_PARAMETER_RDG_BUFFER_UAV(RWBuffer<uint>, RWSymbolsBuffer)
 			SHADER_PARAMETER_RDG_BUFFER_UAV(RWStructuredBuffer<uint>, RWStateBuffer)
 			RDG_BUFFER_ACCESS(IndirectDispatchArgsBuffer, ERHIAccess::IndirectArgs)
 		END_SHADER_PARAMETER_STRUCT()
@@ -519,7 +519,7 @@ namespace ShaderPrint
 
 		BEGIN_SHADER_PARAMETER_STRUCT(FParameters, )
 			SHADER_PARAMETER_STRUCT_REF(FShaderPrintCommonParameters, Common)
-			SHADER_PARAMETER_RDG_BUFFER_SRV(StructuredBuffer<uint>, SymbolsBuffer)
+			SHADER_PARAMETER_RDG_BUFFER_SRV(Buffer<uint>, SymbolsBuffer)
 			SHADER_PARAMETER_RDG_BUFFER_UAV(RWBuffer<uint>, RWIndirectDrawArgsBuffer)
 		END_SHADER_PARAMETER_STRUCT()
 
@@ -546,7 +546,7 @@ namespace ShaderPrint
 			RENDER_TARGET_BINDING_SLOTS()
 			SHADER_PARAMETER_STRUCT_REF(FShaderPrintCommonParameters, Common)
 			SHADER_PARAMETER_TEXTURE(Texture2D, MiniFontTexture)
-			SHADER_PARAMETER_RDG_BUFFER_SRV(StructuredBuffer<uint>, SymbolsBuffer)
+			SHADER_PARAMETER_RDG_BUFFER_SRV(Buffer<uint>, SymbolsBuffer)
 			RDG_BUFFER_ACCESS(IndirectDrawArgsBuffer, ERHIAccess::IndirectArgs)
 		END_SHADER_PARAMETER_STRUCT()
 
@@ -581,7 +581,7 @@ namespace ShaderPrint
 		SHADER_USE_PARAMETER_STRUCT(FShaderDrawDebugCopyCS, FGlobalShader);
 
 		BEGIN_SHADER_PARAMETER_STRUCT(FParameters, )
-			SHADER_PARAMETER_RDG_BUFFER_SRV(StructuredBuffer, ElementBuffer)
+			SHADER_PARAMETER_RDG_BUFFER_SRV(Buffer, ElementBuffer)
 			SHADER_PARAMETER_RDG_BUFFER_UAV(RWBuffer, RWIndirectArgs)
 			SHADER_PARAMETER(uint32, PrimitiveType)
 			SHADER_PARAMETER_STRUCT_REF(FShaderPrintCommonParameters, ShaderPrintData)
@@ -612,8 +612,7 @@ namespace ShaderPrint
 			SHADER_PARAMETER(FVector3f, TranslatedWorldOffsetConversion)
 			SHADER_PARAMETER(FMatrix44f, TranslatedWorldToClip)
 			SHADER_PARAMETER_STRUCT_REF(FViewUniformShaderParameters, View)
-			SHADER_PARAMETER_SRV(StructuredBuffer, LockedShaderDrawDebugPrimitive)
-			SHADER_PARAMETER_RDG_BUFFER_SRV(StructuredBuffer, ShaderDrawDebugPrimitive)
+			SHADER_PARAMETER_RDG_BUFFER_SRV(Buffer<uint>, ShaderDrawDebugPrimitive)
 			RDG_BUFFER_ACCESS(IndirectBuffer, ERHIAccess::IndirectArgs)
 		END_SHADER_PARAMETER_STRUCT()
 
@@ -949,9 +948,11 @@ namespace ShaderPrint
 	{
 		// Initialize graph managed resources
 		const uint32 UintElementCount = GetCountersUintSize() + GetPackedSymbolUintSize() * GetMaxSymbolCountFromCharacterCount(ShaderPrintData.Setup.MaxCharacterCount);
-		FRDGBufferRef SymbolBuffer = GraphBuilder.CreateBuffer(FRDGBufferDesc::CreateStructuredDesc(4, UintElementCount), TEXT("ShaderPrint.SymbolBuffer"));
+		FRDGBufferRef SymbolBuffer = GraphBuilder.CreateBuffer(FRDGBufferDesc::CreateBufferDesc(4, UintElementCount), TEXT("ShaderPrint.SymbolBuffer"));
 		FRDGBufferRef IndirectDispatchArgsBuffer = GraphBuilder.CreateBuffer(FRDGBufferDesc::CreateIndirectDesc<FRHIDispatchIndirectParameters>(1), TEXT("ShaderPrint.IndirectDispatchArgs"));
 		FRDGBufferRef IndirectDrawArgsBuffer = GraphBuilder.CreateBuffer(FRDGBufferDesc::CreateIndirectDesc(5), TEXT("ShaderPrint.IndirectDrawArgs"));
+		FRDGBufferUAVRef SymbolBufferUAV = GraphBuilder.CreateUAV(SymbolBuffer, PF_R32_UINT);
+		FRDGBufferSRVRef SymbolBufferSRV = GraphBuilder.CreateSRV(SymbolBuffer, PF_R32_UINT);
 
 		// Non graph managed resources
 		FRDGBufferSRVRef ValueBuffer = GraphBuilder.CreateSRV(ShaderPrintData.ShaderPrintEntryBuffer);
@@ -968,8 +969,8 @@ namespace ShaderPrint
 			SHADER::FParameters* PassParameters = GraphBuilder.AllocParameters<SHADER::FParameters>();
 			PassParameters->Common = ShaderPrintData.UniformBuffer;
 			PassParameters->ValuesBuffer = ValueBuffer;
-			PassParameters->RWSymbolsBuffer = GraphBuilder.CreateUAV(SymbolBuffer);
-			PassParameters->RWIndirectDispatchArgsBuffer = GraphBuilder.CreateUAV(IndirectDispatchArgsBuffer, EPixelFormat::PF_R32_UINT);
+			PassParameters->RWSymbolsBuffer = SymbolBufferUAV;
+			PassParameters->RWIndirectDispatchArgsBuffer = GraphBuilder.CreateUAV(IndirectDispatchArgsBuffer, PF_R32_UINT);
 
 			FComputeShaderUtils::AddPass(
 				GraphBuilder, 
@@ -987,7 +988,7 @@ namespace ShaderPrint
 			PassParameters->FrameIndex = FrameNumber;
 			PassParameters->Common = ShaderPrintData.UniformBuffer;
 			PassParameters->ValuesBuffer = ValueBuffer;
-			PassParameters->RWSymbolsBuffer = GraphBuilder.CreateUAV(SymbolBuffer);
+			PassParameters->RWSymbolsBuffer = SymbolBufferUAV;
 			PassParameters->RWStateBuffer = GraphBuilder.CreateUAV(ShaderPrintData.ShaderPrintStateBuffer);
 			PassParameters->IndirectDispatchArgsBuffer = IndirectDispatchArgsBuffer;
 
@@ -1024,8 +1025,8 @@ namespace ShaderPrint
 
 			SHADER::FParameters* PassParameters = GraphBuilder.AllocParameters<SHADER::FParameters>();
 			PassParameters->Common = ShaderPrintData.UniformBuffer;
-			PassParameters->SymbolsBuffer = GraphBuilder.CreateSRV(SymbolBuffer);
-			PassParameters->RWIndirectDrawArgsBuffer = GraphBuilder.CreateUAV(IndirectDrawArgsBuffer, EPixelFormat::PF_R32_UINT);
+			PassParameters->SymbolsBuffer = SymbolBufferSRV;
+			PassParameters->RWIndirectDrawArgsBuffer = GraphBuilder.CreateUAV(IndirectDrawArgsBuffer, PF_R32_UINT);
 
 			FComputeShaderUtils::AddPass(
 				GraphBuilder,
@@ -1044,7 +1045,7 @@ namespace ShaderPrint
 			PassParameters->RenderTargets[0] = FRenderTargetBinding(OutputTexture.Texture, ERenderTargetLoadAction::ELoad);
 			PassParameters->Common = ShaderPrintData.UniformBuffer;
 			PassParameters->MiniFontTexture = FontTexture;
-			PassParameters->SymbolsBuffer = GraphBuilder.CreateSRV(SymbolBuffer);
+			PassParameters->SymbolsBuffer = SymbolBufferSRV;
 			PassParameters->IndirectDrawArgsBuffer = IndirectDrawArgsBuffer;
 
 			GraphBuilder.AddPass(
@@ -1077,7 +1078,7 @@ namespace ShaderPrint
 	static void InternalDrawView_Primitives(
 		FRDGBuilder& GraphBuilder,
 		const FShaderPrintData& ShaderPrintData,
-		FRDGBufferRef ShaderPrintPrimitiveBuffer,
+		FRDGBufferSRVRef ShaderPrintPrimitiveBufferSRV,
 		const FIntRect& ViewRect,
 		const FIntRect& UnscaledViewRect,
 		const FMatrix & TranslatedWorldToClip,
@@ -1092,7 +1093,7 @@ namespace ShaderPrint
 		FRDGBufferRef IndirectBuffer = GraphBuilder.CreateBuffer(FRDGBufferDesc::CreateIndirectDesc<FRHIDrawIndirectParameters>(1), TEXT("ShaderDraw.IndirectBuffer"), ERDGBufferFlags::None);
 		{
 			FShaderDrawDebugCopyCS::FParameters* Parameters = GraphBuilder.AllocParameters<FShaderDrawDebugCopyCS::FParameters>();
-			Parameters->ElementBuffer = GraphBuilder.CreateSRV(ShaderPrintPrimitiveBuffer);
+			Parameters->ElementBuffer = ShaderPrintPrimitiveBufferSRV;
 			Parameters->RWIndirectArgs = GraphBuilder.CreateUAV(IndirectBuffer, PF_R32_UINT);
 			Parameters->ShaderPrintData = ShaderPrintData.UniformBuffer;
 			Parameters->PrimitiveType = bLines ? 0u : 1u;
@@ -1130,7 +1131,7 @@ namespace ShaderPrint
 		PassParameters->PS.DepthSampler = TStaticSamplerState<SF_Point, AM_Clamp, AM_Clamp, AM_Clamp>::GetRHI();
 		PassParameters->PS.bCheckerboardEnabled = bLines ? 1u : 0u;
 		PassParameters->PS.bDrawOccludedLines = CVarDrawOccludedLines.GetValueOnRenderThread() != 0 ? 1 : 0;
-		PassParameters->VS.ShaderDrawDebugPrimitive = GraphBuilder.CreateSRV(ShaderPrintPrimitiveBuffer);
+		PassParameters->VS.ShaderDrawDebugPrimitive = ShaderPrintPrimitiveBufferSRV;
 		PassParameters->VS.IndirectBuffer = IndirectBuffer;
 		PassParameters->VS.Common = ShaderPrintData.UniformBuffer;
 
@@ -1210,6 +1211,15 @@ namespace ShaderPrint
 
 		RDG_EVENT_SCOPE(GraphBuilder, "ShaderPrint::DrawView");
 
+		// ShaderPrintPrimitiveBuffer is a StructuredBuffer, but we need to read its content in a vertex shader. For certain platforms, this results in UAV reads, which is not always supported in a vertex shader.
+		// To work around this issue, we instead create a typed buffer here (readable in VS on all platforms) and copy ShaderPrintPrimitiveBuffer into it.
+		auto CreateTypedBufferFromStructured = [](FRDGBuilder& GraphBuilder, FRDGBufferRef Input, const TCHAR* ResultBufferName) -> FRDGBufferRef
+		{
+			FRDGBufferRef Result = GraphBuilder.CreateBuffer(FRDGBufferDesc::CreateBufferDesc(Input->Desc.BytesPerElement, Input->Desc.NumElements), ResultBufferName);
+			AddCopyBufferPass(GraphBuilder, Result, Input);
+			return Result;
+		};
+
 		const FIntRect SourceViewRect = View.ViewRect;
 		const FIntRect OutputViewRect = OutputTexture.ViewRect;
 		const FVector PreViewTranslation = View.ViewMatrices.GetPreViewTranslation() - ShaderPrintData.Setup.PreViewTranslation;
@@ -1219,15 +1229,18 @@ namespace ShaderPrint
 			InternalCopyParameters(GraphBuilder, View.ShaderPrintData);
 		}
 
+		FRDGBufferRef ShaderPrintEntryTypedBuffer = CreateTypedBufferFromStructured(GraphBuilder, ShaderPrintData.ShaderPrintEntryBuffer, TEXT("ShaderDraw.EntryBufferTyped"));
+		FRDGBufferSRVRef ShaderPrintEntryTypedBufferSRV = GraphBuilder.CreateSRV(ShaderPrintEntryTypedBuffer, PF_R32_UINT);
+
 		// Lines
 		{
-			FRDGBufferRef DataBuffer = ShaderPrintData.ShaderPrintEntryBuffer;
+			FRDGBufferSRVRef DataBuffer = ShaderPrintEntryTypedBufferSRV;
 			InternalDrawView_Primitives(GraphBuilder, ShaderPrintData, DataBuffer, SourceViewRect, OutputViewRect, View.ViewMatrices.GetTranslatedViewProjectionMatrix(), PreViewTranslation, true /*bLines*/, false /*bLocked*/, OutputTexture.Texture, DepthTexture.Texture);
 		}
 
 		// Triangles
 		{
-			FRDGBufferRef DataBuffer = ShaderPrintData.ShaderPrintEntryBuffer;
+			FRDGBufferSRVRef DataBuffer = ShaderPrintEntryTypedBufferSRV;
 			InternalDrawView_Primitives(GraphBuilder, ShaderPrintData, DataBuffer, SourceViewRect, OutputViewRect, View.ViewMatrices.GetTranslatedViewProjectionMatrix(), PreViewTranslation, false /*bLines*/, false /*bLocked*/, OutputTexture.Texture, DepthTexture.Texture);
 		}
 
@@ -1235,7 +1248,9 @@ namespace ShaderPrint
 		if (View.ViewState && View.ViewState->ShaderPrintStateData.bIsLocked)
 		{
 			const FVector LockedPreViewTranslation = View.ViewMatrices.GetPreViewTranslation() - View.ViewState->ShaderPrintStateData.PreViewTranslation;
-			FRDGBufferRef DataBuffer = GraphBuilder.RegisterExternalBuffer(View.ViewState->ShaderPrintStateData.EntryBuffer);
+			FRDGBufferRef DataBufferStructured = GraphBuilder.RegisterExternalBuffer(View.ViewState->ShaderPrintStateData.EntryBuffer);
+			FRDGBufferRef DataBufferTyped = CreateTypedBufferFromStructured(GraphBuilder, DataBufferStructured, TEXT("ShaderDraw.LockedEntryBufferTyped"));
+			FRDGBufferSRVRef DataBuffer = GraphBuilder.CreateSRV(DataBufferTyped, PF_R32_UINT);
 			InternalDrawView_Primitives(GraphBuilder, ShaderPrintData, DataBuffer, SourceViewRect, OutputViewRect, View.ViewMatrices.GetTranslatedViewProjectionMatrix(), LockedPreViewTranslation, true  /*bLines*/, true/*bLocked*/, OutputTexture.Texture, DepthTexture.Texture);
 			InternalDrawView_Primitives(GraphBuilder, ShaderPrintData, DataBuffer, SourceViewRect, OutputViewRect, View.ViewMatrices.GetTranslatedViewProjectionMatrix(), LockedPreViewTranslation, false /*bLines*/, true/*bLocked*/, OutputTexture.Texture, DepthTexture.Texture);
 		}
