Fix crash when using light complexity visualization with hair strands as the pernutation is not buit.

diff --git a/Engine/Source/Runtime/Renderer/Private/LightRendering.cpp b/Engine/Source/Runtime/Renderer/Private/LightRendering.cpp
--- a/Engine/Source/Runtime/Renderer/Private/LightRendering.cpp
+++ b/Engine/Source/Runtime/Renderer/Private/LightRendering.cpp
@@ -2699,7 +2699,7 @@ void FDeferredShadingSceneRenderer::RenderLightForHair(
 {
 	// Ensure the light is valid for this view
 	const bool bHairRenderingEnabled = HairStrands::HasViewHairStrandsData(View);
-	if (!bHairRenderingEnabled || !LightSceneInfo->ShouldRenderLight(View) || View.HairStrandsViewData.VisibilityData.SampleLightingTexture == nullptr)
+	if (!bHairRenderingEnabled || !LightSceneInfo->ShouldRenderLight(View) || View.HairStrandsViewData.VisibilityData.SampleLightingTexture == nullptr || View.Family->EngineShowFlags.VisualizeLightCulling)
 	{
 		return;
 	}
