Moved refreshing cached ray tracing data before visibility tasks are launched.
- Scene.RefreshRayTracingMeshCommandCache() and Scene.RefreshRayTracingInstances() can result in changes to PrimitiveRayTracingFlags, which are read in visibility tasks. However, those functions were called on rendering thread and could end up causing race conditions / invalid state since visibility tasks are running on worker threads.

diff --git a/Engine/Source/Runtime/Renderer/Private/DeferredShadingRenderer.cpp b/Engine/Source/Runtime/Renderer/Private/DeferredShadingRenderer.cpp
--- a/Engine/Source/Runtime/Renderer/Private/DeferredShadingRenderer.cpp
+++ b/Engine/Source/Runtime/Renderer/Private/DeferredShadingRenderer.cpp
@@ -605,6 +605,54 @@ static void RenderOpaqueFX(
 
 #if RHI_RAYTRACING
 
+static void RefreshCachedRayTracingState(FScene& Scene, const FSceneViewFamily& ViewFamily, TArrayView<FViewInfo> Views)
+{
+	const ERayTracingMeshCommandsMode CurrentMode = ViewFamily.EngineShowFlags.PathTracing ? ERayTracingMeshCommandsMode::PATH_TRACING : ERayTracingMeshCommandsMode::RAY_TRACING;
+	bool bNaniteCoarseMeshStreamingModeChanged = false;
+#if WITH_EDITOR
+	bNaniteCoarseMeshStreamingModeChanged = Nanite::FCoarseMeshStreamingManager::CheckStreamingMode();
+#endif // WITH_EDITOR
+	const bool bNaniteRayTracingModeChanged = Nanite::GRayTracingManager.CheckModeChanged();
+
+	if (CurrentMode != Scene.CachedRayTracingMeshCommandsMode
+		|| bNaniteCoarseMeshStreamingModeChanged
+		|| bNaniteRayTracingModeChanged
+		|| bUpdateCachedRayTracingState)
+	{
+		Scene.WaitForCacheRayTracingPrimitivesTask();
+
+		// In some situations, we need to refresh the cached ray tracing mesh commands because they contain data about the currently bound shader. 
+		// This operation is a bit expensive but only happens once as we transition between modes which should be rare.
+		Scene.CachedRayTracingMeshCommandsMode = CurrentMode;
+		Scene.RefreshRayTracingMeshCommandCache();
+		bUpdateCachedRayTracingState = false;
+	}
+
+	if (bRefreshRayTracingInstances)
+	{
+		Scene.WaitForCacheRayTracingPrimitivesTask();
+
+		// In some situations, we need to refresh the cached ray tracing instance.
+		// This assumes that cached instances will keep using the same LOD since CachedRayTracingMeshCommands is not recalculated
+		// eg: Need to update PrimitiveRayTracingFlags
+		// This operation is a bit expensive but only happens once as we transition between modes which should be rare.
+		Scene.RefreshRayTracingInstances();
+		bRefreshRayTracingInstances = false;
+	}
+
+	if (bNaniteRayTracingModeChanged)
+	{
+		for (FViewInfo& View : Views)
+		{
+			if (View.ViewState != nullptr && !View.bIsOfflineRender)
+			{
+				// don't invalidate in the offline case because we only get one attempt at rendering each sample
+				View.ViewState->PathTracingInvalidate();
+			}
+		}
+	}
+}
+
 static bool ShouldPrepareRayTracingDecals(const FScene& Scene, const FSceneViewFamily& ViewFamily)
 {
 	if (!IsRayTracingEnabled() || !RHISupportsRayTracingCallableShaders(ViewFamily.GetShaderPlatform()))
@@ -1494,6 +1542,8 @@ void FDeferredShadingSceneRenderer::Render(FRDGBuilder& GraphBuilder)
 		SCOPED_GPU_STAT(GraphBuilder.RHICmdList, RayTracingGeometry);
 		GRayTracingGeometryManager->ProcessBuildRequests(GraphBuilder.RHICmdList);
 	}
+
+	RefreshCachedRayTracingState(*Scene, ViewFamily, Views);
 #endif
 
 	FInitViewTaskDatas InitViewTaskDatas = OnRenderBegin(GraphBuilder);
@@ -1619,68 +1669,20 @@ void FDeferredShadingSceneRenderer::Render(FRDGBuilder& GraphBuilder)
 		// Initialize ray tracing flags, in case they weren't initialized in the CreateSceneRenderers code path
 		InitializeRayTracingFlags_RenderThread();
 
-		// Now that we have updated all the PrimitiveSceneInfos, update the RayTracing mesh commands cache if needed
+		if (bAnyRayTracingPassEnabled)
 		{
-			const ERayTracingMeshCommandsMode CurrentMode = ViewFamily.EngineShowFlags.PathTracing ? ERayTracingMeshCommandsMode::PATH_TRACING : ERayTracingMeshCommandsMode::RAY_TRACING;
-			bool bNaniteCoarseMeshStreamingModeChanged = false;
-#if WITH_EDITOR
-			bNaniteCoarseMeshStreamingModeChanged = Nanite::FCoarseMeshStreamingManager::CheckStreamingMode();
-#endif // WITH_EDITOR
-			const bool bNaniteRayTracingModeChanged = Nanite::GRayTracingManager.CheckModeChanged();
-
-			if (CurrentMode != Scene->CachedRayTracingMeshCommandsMode
-				|| bNaniteCoarseMeshStreamingModeChanged
-				|| bNaniteRayTracingModeChanged
-				|| bUpdateCachedRayTracingState)
-			{
-				Scene->WaitForCacheRayTracingPrimitivesTask();
-
-				// In some situations, we need to refresh the cached ray tracing mesh commands because they contain data about the currently bound shader. 
-				// This operation is a bit expensive but only happens once as we transition between modes which should be rare.
-				Scene->CachedRayTracingMeshCommandsMode = CurrentMode;
-				Scene->RefreshRayTracingMeshCommandCache();
-				bUpdateCachedRayTracingState = false;
-			}
-
-			if (bRefreshRayTracingInstances)
-			{
-				Scene->WaitForCacheRayTracingPrimitivesTask();
-
-				// In some situations, we need to refresh the cached ray tracing instance.
-				// This assumes that cached instances will keep using the same LOD since CachedRayTracingMeshCommands is not recalculated
-				// eg: Need to update PrimitiveRayTracingFlags
-				// This operation is a bit expensive but only happens once as we transition between modes which should be rare.
-				Scene->RefreshRayTracingInstances();
-				bRefreshRayTracingInstances = false;
-			}
+			const int32 ReferenceViewIndex = 0;
+			FViewInfo& ReferenceView = Views[ReferenceViewIndex];
+			FGraphEventRef PrereqTask = CreateCompatibilityGraphEvent(MakeArrayView({ Scene->GetCacheRayTracingPrimitivesTask(), InitViewTaskDatas.VisibilityTaskData->GetFrustumCullTask() }));
 
-			if (bNaniteRayTracingModeChanged)
-			{
-				for (FViewInfo& View : Views)
+			InitViewTaskDatas.RayTracingRelevantPrimitives = Allocator.Create<FRayTracingRelevantPrimitiveTaskData>();
+			InitViewTaskDatas.RayTracingRelevantPrimitives->List = RayTracing::CreateRelevantPrimitiveList(Allocator);
+			InitViewTaskDatas.RayTracingRelevantPrimitives->Task = FFunctionGraphTask::CreateAndDispatchWhenReady(
+				[Scene = Scene, &ReferenceView, &RayTracingRelevantPrimitiveList = *InitViewTaskDatas.RayTracingRelevantPrimitives->List]()
 				{
-					if (View.ViewState != nullptr && !View.bIsOfflineRender)
-					{
-						// don't invalidate in the offline case because we only get one attempt at rendering each sample
-						View.ViewState->PathTracingInvalidate();
-					}
-				}
-			}
-
-			if (bAnyRayTracingPassEnabled)
-			{
-				const int32 ReferenceViewIndex = 0;
-				FViewInfo& ReferenceView = Views[ReferenceViewIndex];
-				FGraphEventRef PrereqTask = CreateCompatibilityGraphEvent(MakeArrayView({ Scene->GetCacheRayTracingPrimitivesTask(), InitViewTaskDatas.VisibilityTaskData->GetFrustumCullTask() }));
-
-				InitViewTaskDatas.RayTracingRelevantPrimitives = Allocator.Create<FRayTracingRelevantPrimitiveTaskData>();
-				InitViewTaskDatas.RayTracingRelevantPrimitives->List = RayTracing::CreateRelevantPrimitiveList(Allocator);
-				InitViewTaskDatas.RayTracingRelevantPrimitives->Task = FFunctionGraphTask::CreateAndDispatchWhenReady(
-					[Scene = Scene, &ReferenceView, &RayTracingRelevantPrimitiveList = *InitViewTaskDatas.RayTracingRelevantPrimitives->List]()
-					{
-						FTaskTagScope TaskTagScope(ETaskTag::EParallelRenderingThread);
-						RayTracing::GatherRelevantPrimitives(*Scene, ReferenceView, RayTracingRelevantPrimitiveList);
-					}, TStatId(), PrereqTask, ENamedThreads::AnyNormalThreadHiPriTask);
-			}
+					FTaskTagScope TaskTagScope(ETaskTag::EParallelRenderingThread);
+					RayTracing::GatherRelevantPrimitives(*Scene, ReferenceView, RayTracingRelevantPrimitiveList);
+				}, TStatId(), PrereqTask, ENamedThreads::AnyNormalThreadHiPriTask);
 		}
 #endif
 	}
