Global Distance Field optimizations
* Run single loop over all DF Scene objects when culling to clipmaps. This halves required time, which helps a lot in larger scenes where most of the objects are outside of the clipmap bounds.
* Split DF streaming reache readback into mobility layers in order not to reacache the entire static layer when only a movable DF was streamed in

diff --git a/Engine/Shaders/Private/DistanceField/GlobalDistanceField.usf b/Engine/Shaders/Private/DistanceField/GlobalDistanceField.usf
--- a/Engine/Shaders/Private/DistanceField/GlobalDistanceField.usf
+++ b/Engine/Shaders/Private/DistanceField/GlobalDistanceField.usf
@@ -26,6 +26,46 @@
 	#define THREADGROUP_SIZE_Z 1
 #endif
 
+struct FPackedClipmap
+{
+	float3 TranslatedWorldCenter;
+	float3 WorldExtent;
+	float3 ViewTilePosition;
+	float3 RelativePreViewTranslation;
+
+	bool bReadbackHasPendingStreaming;
+	uint ClipmapIndex; // Global index
+	float MeshSDFRadiusThreshold;
+	float InfluenceRadiusSq;
+};
+
+// Keep in sync with GlobalDistanceField::PackedClipmapBufferStride
+#define PACKED_CLIPMAP_BUFFER_STRIDE 4
+StructuredBuffer<float4> PackedClipmapBuffer;
+
+FPackedClipmap GetPackedClipmap(uint PackedClipmapIndex)
+{
+	uint BaseOffset = PACKED_CLIPMAP_BUFFER_STRIDE * PackedClipmapIndex;
+
+	FPackedClipmap PackedClipmap;
+
+	const float4 V0 = PackedClipmapBuffer[BaseOffset + 0];
+	const float4 V1 = PackedClipmapBuffer[BaseOffset + 1];
+	const float4 V2 = PackedClipmapBuffer[BaseOffset + 2];
+	const float4 V3 = PackedClipmapBuffer[BaseOffset + 3];
+
+	PackedClipmap.TranslatedWorldCenter = V0.xyz;
+	PackedClipmap.MeshSDFRadiusThreshold = V0.w;
+	PackedClipmap.WorldExtent = V1.xyz;
+	PackedClipmap.InfluenceRadiusSq = V1.w;
+	PackedClipmap.ViewTilePosition = V2.xyz;
+	PackedClipmap.bReadbackHasPendingStreaming = V2.w > 0.0f;
+	PackedClipmap.RelativePreViewTranslation = V3.xyz;
+	PackedClipmap.ClipmapIndex = V3.w;
+
+	return PackedClipmap;
+}
+
 RWStructuredBuffer<uint> RWCullGridAllocator;
 RWStructuredBuffer<uint> RWCullGridObjectHeader;
 RWStructuredBuffer<uint> RWCullGridObjectArray;
@@ -45,48 +85,56 @@ RWStructuredBuffer<uint> RWHasPendingStreaming;
 float3 ClipmapTranslatedWorldCenter;
 float3 ClipmapWorldExtent;
 float InfluenceRadius;
-float InfluenceRadiusSq;
 // 0: accept only mostly static objects, 1: accept only often moving objects, 2: accept all objects
 uint AcceptOftenMovingObjectsOnly;
-float MeshSDFRadiusThreshold;
-
-float3 ViewTilePosition;
-float3 RelativePreViewTranslation;
+uint NumPackedClipmaps;
+uint ObjectIndexBufferStride;
 
 [numthreads(CULLOBJECTS_THREADGROUP_SIZE, 1, 1)]
 void CullObjectsToClipmapCS(
 	uint GroupIndex : SV_GroupIndex,
 	uint3 GroupId : SV_GroupID)
 {
-	const FLWCVector3 PreViewTranslation = MakeLWCVector3(-ViewTilePosition, RelativePreViewTranslation);
-
 	const uint ObjectIndex = GetUnWrappedDispatchThreadId(GroupId, GroupIndex, CULLOBJECTS_THREADGROUP_SIZE);
 	if (ObjectIndex < NumSceneObjects)
 	{
 		const FDFObjectBounds DFObjectBounds = LoadDFObjectBounds(ObjectIndex);
-		const float3 TranslatedBoundsCenter = LWCToFloat(LWCAdd(DFObjectBounds.Center, PreViewTranslation));
 
-		BRANCH
-		if ((DFObjectBounds.bVisible || DFObjectBounds.bAffectIndirectLightingWhileHidden)
-			&& (DFObjectBounds.SphereRadius > MeshSDFRadiusThreshold || DFObjectBounds.bEmissiveLightSource)
-			&& (AcceptOftenMovingObjectsOnly == 2 || DFObjectBounds.OftenMoving == AcceptOftenMovingObjectsOnly))
+		const bool bMatchingMobility = AcceptOftenMovingObjectsOnly == 2 || DFObjectBounds.OftenMoving == AcceptOftenMovingObjectsOnly;
+		const bool bVisible = DFObjectBounds.bVisible || DFObjectBounds.bAffectIndirectLightingWhileHidden;
+
+		if (bMatchingMobility && bVisible)
 		{
-			float DistanceSq = ComputeSquaredDistanceBetweenAABBs(ClipmapTranslatedWorldCenter, ClipmapWorldExtent, TranslatedBoundsCenter, DFObjectBounds.BoxExtent);
-			if (DistanceSq < InfluenceRadiusSq)
+			for (uint PackedClipmapIndex = 0; PackedClipmapIndex < NumPackedClipmaps; ++PackedClipmapIndex)
 			{
-				uint DestIndex;
-				InterlockedAdd(RWObjectIndexNumBuffer[0], 1U, DestIndex);
-				RWObjectIndexBuffer[DestIndex] = ObjectIndex;
-
-#if READBACK_HAS_PENDING_STREAMING
-				FDFObjectData DFObjectData = LoadDFObjectData(ObjectIndex);
-				uint NumMips = LoadDFAssetData(DFObjectData.AssetIndex, 0).NumMips;
+				FPackedClipmap Clipmap = GetPackedClipmap(PackedClipmapIndex);
 
-				if (NumMips == 1)
+				if (DFObjectBounds.SphereRadius > Clipmap.MeshSDFRadiusThreshold || DFObjectBounds.bEmissiveLightSource)
 				{
-					RWHasPendingStreaming[0] = 1;
+					const FLWCVector3 PreViewTranslation = MakeLWCVector3(-Clipmap.ViewTilePosition, Clipmap.RelativePreViewTranslation);
+					const float3 TranslatedBoundsCenter = LWCToFloat(LWCAdd(DFObjectBounds.Center, PreViewTranslation));
+					float DistanceSq = ComputeSquaredDistanceBetweenAABBs(Clipmap.TranslatedWorldCenter, Clipmap.WorldExtent, TranslatedBoundsCenter, DFObjectBounds.BoxExtent);
+
+					if (DistanceSq < Clipmap.InfluenceRadiusSq)
+					{
+						uint DestIndex;
+						InterlockedAdd(RWObjectIndexNumBuffer[PackedClipmapIndex], 1U, DestIndex);
+						RWObjectIndexBuffer[DestIndex + PackedClipmapIndex * ObjectIndexBufferStride] = ObjectIndex;
+
+						#if READBACK_HAS_PENDING_STREAMING
+						if (Clipmap.bReadbackHasPendingStreaming)
+						{
+							FDFObjectData DFObjectData = LoadDFObjectData(ObjectIndex);
+							uint NumMips = LoadDFAssetData(DFObjectData.AssetIndex, 0).NumMips;
+
+							if (NumMips == 1)
+							{
+								RWHasPendingStreaming[Clipmap.ClipmapIndex] = 1;
+							}
+						}
+						#endif
+					}
 				}
-#endif
 			}
 		}
 	}
@@ -97,13 +145,13 @@ RWBuffer<uint> RWPageUpdateIndirectArgBuffer;
 RWBuffer<uint> RWCullGridUpdateIndirectArgBuffer;
 RWBuffer<uint> RWBuildObjectGridIndirectArgBuffer;
 
-[numthreads(1, 1, 1)]
+[numthreads(THREADGROUP_SIZE, 1, 1)]
 void ClearIndirectArgBufferCS(
 	uint3 GroupId : SV_GroupID,
 	uint3 DispatchThreadId : SV_DispatchThreadID,
 	uint3 GroupThreadId : SV_GroupThreadID)
 {
-	if (DispatchThreadId.x == 0)
+	if (all(DispatchThreadId) == 0)
 	{
 		const uint3 ClearValue = uint3(0, 1, 1);
 		WriteDispatchIndirectArgs(RWPageComposeIndirectArgBuffer, 0, ClearValue);
@@ -121,6 +169,7 @@ RWBuffer<uint> RWCullGridIndirectArgBuffer;
 
 Buffer<float4> UpdateBoundsBuffer;
 uint NumUpdateBounds;
+float InfluenceRadiusSq;
 
 uint3 CullGridResolution;
 float3 CullGridCoordToTranslatedWorldCenterScale;
@@ -300,6 +349,8 @@ float3 InvPageGridResolution;
 float3 PageWorldExtent;
 float PageWorldRadius;
 float ClipmapInfluenceRadius;
+float3 ViewTilePosition;
+float3 RelativePreViewTranslation;
 
 RWStructuredBuffer<uint> RWPageComposeTileBuffer;
 
@@ -525,6 +576,7 @@ StructuredBuffer<uint> ObjectIndexBuffer;
 StructuredBuffer<uint> ObjectIndexNumBuffer;
 StructuredBuffer<uint> CullGridTileBuffer;
 float3 CullTileWorldExtent;
+uint PackedClipmapIndex;
 
 [numthreads(GRID_CULL_THREADGROUP_TOTALSIZE, 1, 1)]
 void CullObjectsToGridCS(
@@ -551,9 +603,9 @@ void CullObjectsToGridCS(
 
 	if (all(CullGridCoord < CullGridResolution))
 	{
-		for (uint IndexInObjectIndexBuffer = ThreadIndex; IndexInObjectIndexBuffer < ObjectIndexNumBuffer[0]; IndexInObjectIndexBuffer += GRID_CULL_THREADGROUP_TOTALSIZE)
+		for (uint IndexInObjectIndexBuffer = ThreadIndex; IndexInObjectIndexBuffer < ObjectIndexNumBuffer[PackedClipmapIndex]; IndexInObjectIndexBuffer += GRID_CULL_THREADGROUP_TOTALSIZE)
 		{
-			uint ObjectIndex = ObjectIndexBuffer[IndexInObjectIndexBuffer];
+			uint ObjectIndex = ObjectIndexBuffer[PackedClipmapIndex * ObjectIndexBufferStride + IndexInObjectIndexBuffer];
 
 			const FDFObjectBounds DFObjectBounds = LoadDFObjectBounds(ObjectIndex);
 			const float3 TranslatedBoundsCenter = LWCToFloat(LWCAdd(DFObjectBounds.Center, PreViewTranslation));
@@ -581,9 +633,9 @@ void CullObjectsToGridCS(
 
 	if (all(CullGridCoord < CullGridResolution))
 	{
-		for (uint IndexInObjectIndexBuffer = ThreadIndex; IndexInObjectIndexBuffer < ObjectIndexNumBuffer[0]; IndexInObjectIndexBuffer += GRID_CULL_THREADGROUP_TOTALSIZE)
+		for (uint IndexInObjectIndexBuffer = ThreadIndex; IndexInObjectIndexBuffer < ObjectIndexNumBuffer[PackedClipmapIndex]; IndexInObjectIndexBuffer += GRID_CULL_THREADGROUP_TOTALSIZE)
 		{
-			uint ObjectIndex = ObjectIndexBuffer[IndexInObjectIndexBuffer];
+			uint ObjectIndex = ObjectIndexBuffer[PackedClipmapIndex * ObjectIndexBufferStride + IndexInObjectIndexBuffer];
 
 			const FDFObjectBounds DFObjectBounds = LoadDFObjectBounds(ObjectIndex);
 			const float3 TranslatedBoundsCenter = LWCToFloat(LWCAdd(DFObjectBounds.Center, PreViewTranslation));
diff --git a/Engine/Source/Runtime/Renderer/Private/GlobalDistanceField.cpp b/Engine/Source/Runtime/Renderer/Private/GlobalDistanceField.cpp
--- a/Engine/Source/Runtime/Renderer/Private/GlobalDistanceField.cpp
+++ b/Engine/Source/Runtime/Renderer/Private/GlobalDistanceField.cpp
@@ -410,6 +410,9 @@ namespace GlobalDistanceField
 	// Every distance field page stores 4^3 object grid cells with 4 * uint32 elements per cell
 	const int32 ObjectGridPageBufferStride = 4 * sizeof(uint32);
 	const int32 ObjectGridPageBufferNumElementsPerPage = 4 * 4 * 4;
+
+	// Keep in sync with PACKED_CLIPMAP_BUFFER_STRIDE
+	const int32 PackedClipmapBufferStride = 4;
 }
 const int32 GGlobalDistanceFieldPageResolutionInAtlas = 8; // Includes 0.5 texel trilinear filter margin
 const int32 GGlobalDistanceFieldCoveragePageResolutionInAtlas = 4; // Includes 0.5 texel trilinear filter margin
@@ -769,60 +772,6 @@ FVector GetGlobalDistanceFieldViewOrigin(const FViewInfo& View, int32 ClipmapInd
 	return CameraOrigin;
 }
 
-void RecaptureClipmapForMeshSDFStreamingIfNeeded(
-	const FViewInfo& View,
-	const FScene* Scene,
-	FPersistentGlobalDistanceFieldData& GlobalDistanceFieldData, 
-	FGlobalDistanceFieldClipmapState& ClipmapViewState,
-	FGlobalDistanceFieldClipmap& Clipmap,
-	int32 ClipmapIndex,
-	uint32 CacheType,
-	const FBox& ClipmapBounds)
-{
-	FRHIGPUBufferReadback* LatestReadbackBuffer = nullptr;
-
-	// Find latest buffer that is ready
-	while (ClipmapViewState.ReadbackBuffersNumPending > 0)
-	{
-		uint32 Index = (ClipmapViewState.ReadbackBuffersWriteIndex + ClipmapViewState.MaxPendingStreamingReadbackBuffers - ClipmapViewState.ReadbackBuffersNumPending) % ClipmapViewState.MaxPendingStreamingReadbackBuffers;
-		if (ClipmapViewState.HasPendingStreamingReadbackBuffers[Index]->IsReady())
-		{
-			ClipmapViewState.ReadbackBuffersNumPending--;
-			LatestReadbackBuffer = ClipmapViewState.HasPendingStreamingReadbackBuffers[Index].Get();
-		}
-		else
-		{
-			break;
-		}
-	}
-
-	bool bPreviousCaptureHadPendingStreaming = false;
-
-	// Readback whether the last CullObjectsToClipmap for this clipmap detected Mesh SDFs that have not streamed in (NumMips == 1)
-	if (LatestReadbackBuffer)
-	{
-		const uint32* LatestReadbackBufferPtr = (const uint32*)LatestReadbackBuffer->Lock(1 * sizeof(uint32));
-		bPreviousCaptureHadPendingStreaming = LatestReadbackBufferPtr[0] != 0;
-		LatestReadbackBuffer->Unlock();
-	}
-
-	TArray<int32>& DeferredUpdatesForMeshSDFStreaming = GlobalDistanceFieldData.DeferredUpdatesForMeshSDFStreaming[CacheType];
-
-	if (bPreviousCaptureHadPendingStreaming)
-	{
-		// Add a new deferred update for when the streaming is finished
-		DeferredUpdatesForMeshSDFStreaming.AddUnique(ClipmapIndex);
-	}
-	// Mesh SDFs are done streaming, recapture the clipmaps
-	else if (!Scene->DistanceFieldSceneData.HasPendingStreaming() && DeferredUpdatesForMeshSDFStreaming.Remove(ClipmapIndex))
-	{
-		// Push full update
-		Clipmap.UpdateBounds.Reset();
-		Clipmap.UpdateBounds.Add(FClipmapUpdateBounds(ClipmapBounds.GetCenter(), ClipmapBounds.GetExtent(), false));
-		Clipmap.FullRecaptureReason = EGlobalSDFFullRecaptureReason::MeshSDFStreaming;
-	}
-}
-
 static void ComputeUpdateRegionsAndUpdateViewState(
 	FRHICommandListImmediate& RHICmdList, 
 	FViewInfo& View, 
@@ -1078,6 +1027,46 @@ static void ComputeUpdateRegionsAndUpdateViewState(
 			GlobalDistanceFieldInfo.PageTableLayerTextures[CacheType] = PageTableTexture;
 		}
 
+		// Process mesh SDF streaming readback in order to trigger full clipmap update when streaming is done
+		for (uint32 CacheType = 0; CacheType < GDF_Num; CacheType++)
+		{
+			FGlobalDistanceFieldStreamingReadback& StreamingReadback = GlobalDistanceFieldData.StreamingReadback[CacheType];
+
+			FRHIGPUBufferReadback* LatestReadbackBuffer = nullptr;
+
+			// Find latest buffer that is ready
+			while (StreamingReadback.ReadbackBuffersNumPending > 0)
+			{
+				uint32 Index = (StreamingReadback.ReadbackBuffersWriteIndex + StreamingReadback.MaxPendingStreamingReadbackBuffers - StreamingReadback.ReadbackBuffersNumPending) % StreamingReadback.MaxPendingStreamingReadbackBuffers;
+				if (StreamingReadback.PendingStreamingReadbackBuffers[Index]->IsReady())
+				{
+					StreamingReadback.ReadbackBuffersNumPending--;
+					LatestReadbackBuffer = StreamingReadback.PendingStreamingReadbackBuffers[Index].Get();
+				}
+				else
+				{
+					break;
+				}
+			}
+
+			// Readback whether the last CullObjectsToClipmap for this clipmap detected Mesh SDFs that have not streamed in (NumMips == 1)
+			if (LatestReadbackBuffer)
+			{
+				const uint32* LatestReadbackBufferPtr = (const uint32*)LatestReadbackBuffer->Lock(GlobalDistanceField::MaxClipmaps * sizeof(uint32));
+
+				for (uint32 ClipmapIndex = 0; ClipmapIndex < GlobalDistanceField::MaxClipmaps; ++ClipmapIndex)
+				{
+					if (LatestReadbackBufferPtr[ClipmapIndex] != 0)
+					{
+						// Add a new deferred update for when the streaming is finished
+						GlobalDistanceFieldData.DeferredUpdatesForMeshSDFStreaming[CacheType].AddUnique(ClipmapIndex);
+					}
+				}
+
+				LatestReadbackBuffer->Unlock();
+			}
+		}
+
 		for (int32 ClipmapIndex = 0; ClipmapIndex < NumClipmaps; ClipmapIndex++)
 		{
 			FGlobalDistanceFieldClipmapState& ClipmapViewState = GlobalDistanceFieldData.ClipmapState[ClipmapIndex];
@@ -1234,7 +1223,14 @@ static void ComputeUpdateRegionsAndUpdateViewState(
 						Clipmap.FullRecaptureReason = EGlobalSDFFullRecaptureReason::HeightfieldStreaming;
 					}
 
-					RecaptureClipmapForMeshSDFStreamingIfNeeded(View, Scene, GlobalDistanceFieldData, ClipmapViewState, Clipmap, ClipmapIndex, CacheType, ClipmapBounds);
+					// Push full update when mesh SDF streaming is done
+					if (!Scene->DistanceFieldSceneData.HasPendingStreaming() && GlobalDistanceFieldData.DeferredUpdatesForMeshSDFStreaming[CacheType].Remove(ClipmapIndex))
+					{
+						// Push full update
+						Clipmap.UpdateBounds.Reset();
+						Clipmap.UpdateBounds.Add(FClipmapUpdateBounds(ClipmapBounds.GetCenter(), ClipmapBounds.GetExtent(), false));
+						Clipmap.FullRecaptureReason = EGlobalSDFFullRecaptureReason::MeshSDFStreaming;
+					}
 
 					ClipmapViewState.Cache[CacheType].PrimitiveModifiedBounds.Empty(DistanceField::MinPrimitiveModifiedBoundsAllocation);
 				}
@@ -1412,13 +1408,10 @@ class FCullObjectsToClipmapCS : public FGlobalShader
 		SHADER_PARAMETER_RDG_BUFFER_UAV(RWStructuredBuffer<uint>, RWHasPendingStreaming)
 		SHADER_PARAMETER_STRUCT_INCLUDE(FDistanceFieldObjectBufferParameters, DistanceFieldObjectBuffers)
 		SHADER_PARAMETER_STRUCT_INCLUDE(FDistanceFieldAtlasParameters, DistanceFieldAtlasParameters)
-		SHADER_PARAMETER(FVector3f, ClipmapTranslatedWorldCenter)
-		SHADER_PARAMETER(FVector3f, ClipmapWorldExtent)
+		SHADER_PARAMETER_RDG_BUFFER_SRV(StructuredBuffer<float4>, PackedClipmapBuffer)
 		SHADER_PARAMETER(uint32, AcceptOftenMovingObjectsOnly)
-		SHADER_PARAMETER(float, MeshSDFRadiusThreshold)
-		SHADER_PARAMETER(float, InfluenceRadiusSq)
-		SHADER_PARAMETER(FVector3f, ViewTilePosition)
-		SHADER_PARAMETER(FVector3f, RelativePreViewTranslation)
+		SHADER_PARAMETER(uint32, NumPackedClipmaps)
+		SHADER_PARAMETER(uint32, ObjectIndexBufferStride)
 	END_SHADER_PARAMETER_STRUCT()
 
 	class FReadbackHasPendingStreaming : SHADER_PERMUTATION_BOOL("READBACK_HAS_PENDING_STREAMING");
@@ -1462,7 +1455,7 @@ class FClearIndirectArgBufferCS : public FGlobalShader
 
 	static int32 GetGroupSize()
 	{
-		return 1;
+		return 64;
 	}
 
 	static void ModifyCompilationEnvironment(const FGlobalShaderPermutationParameters& Parameters, FShaderCompilerEnvironment& OutEnvironment)
@@ -1538,6 +1531,8 @@ class FCullObjectsToGridCS : public FGlobalShader
 		SHADER_PARAMETER(float, InfluenceRadiusSq)
 		SHADER_PARAMETER(FVector3f, ViewTilePosition)
 		SHADER_PARAMETER(FVector3f, RelativePreViewTranslation)
+		SHADER_PARAMETER(uint32, PackedClipmapIndex)
+		SHADER_PARAMETER(uint32, ObjectIndexBufferStride)
 	END_SHADER_PARAMETER_STRUCT()
 
 	static bool ShouldCompilePermutation(const FGlobalShaderPermutationParameters& Parameters)
@@ -1631,8 +1626,6 @@ class FCompositeObjectsIntoPagesCS : public FGlobalShader
 		SHADER_PARAMETER_RDG_TEXTURE(Texture3D<uint>, ParentPageTableLayerTexture)
 		SHADER_PARAMETER_RDG_BUFFER_SRV(StructuredBuffer<uint>, CullGridObjectHeader)
 		SHADER_PARAMETER_RDG_BUFFER_SRV(StructuredBuffer<uint>, CullGridObjectArray)
-		SHADER_PARAMETER_RDG_BUFFER_SRV(StructuredBuffer<uint>, ObjectIndexNumBuffer)
-		SHADER_PARAMETER_RDG_BUFFER_SRV(StructuredBuffer<uint>, ObjectIndexBuffer)
 		SHADER_PARAMETER_STRUCT_INCLUDE(FDistanceFieldObjectBufferParameters, DistanceFieldObjectBuffers)
 		SHADER_PARAMETER_STRUCT_INCLUDE(FDistanceFieldAtlasParameters, DistanceFieldAtlas)
 		SHADER_PARAMETER_STRUCT_INCLUDE(FGlobalDistanceFieldUpdateParameters, GlobalDistanceFieldUpdateParameters)
@@ -2122,9 +2115,6 @@ struct FGlobalDistanceFieldPackedClipmap
 	int32 NumUpdateBounds = 0;
 	FRDGBufferRef UpdateBoundsBuffer = nullptr;
 	FHeightfieldDescription UpdateRegionHeightfield;
-	FRDGBufferRef ObjectIndexBuffer = nullptr;
-	FRDGBufferRef ObjectIndexNumBuffer = nullptr;
-	FRDGBufferRef HasPendingStreamingReadbackBuffer = nullptr;
 
 	bool bRecacheClipmapsWithPendingStreaming = false;
 
@@ -2325,15 +2315,13 @@ void UpdateGlobalDistanceFieldCache(
 		FDistanceFieldAtlasParameters DistanceFieldAtlas = DistanceField::SetupAtlasParameters(GraphBuilder, DistanceFieldSceneData);
 
 		// Allocate buffers for objects culled to clipmaps
-		for (FGlobalDistanceFieldPackedClipmap& PackedClipmap : PackedClipmaps)
-		{
-			const int32 MaxSDFMeshObjects = FMath::RoundUpToPowerOfTwo(DistanceFieldSceneData.NumObjectsInBuffer);
-			PackedClipmap.ObjectIndexBuffer = GraphBuilder.CreateBuffer(FRDGBufferDesc::CreateStructuredDesc(sizeof(uint32), MaxSDFMeshObjects), TEXT("GlobalDistanceField.ObjectIndices"));
-			PackedClipmap.ObjectIndexNumBuffer = GraphBuilder.CreateBuffer(FRDGBufferDesc::CreateStructuredDesc(sizeof(uint32), 1), TEXT("GlobalDistanceField.ObjectIndexNum"));
-			AddClearUAVPass(GraphBuilder, GraphBuilder.CreateUAV(PackedClipmap.ObjectIndexNumBuffer, PF_R32_UINT), 0);
-		}
+		const uint32 ObjectIndexBufferStride = FMath::RoundUpToPowerOfTwo(DistanceFieldSceneData.NumObjectsInBuffer);
+		FRDGBufferRef ObjectIndexBuffer = GraphBuilder.CreateBuffer(FRDGBufferDesc::CreateStructuredDesc(sizeof(uint32), FMath::Max(PackedClipmaps.Num(), 1) * ObjectIndexBufferStride), TEXT("GlobalDistanceField.ObjectIndices"));
+		FRDGBufferRef ObjectIndexNumBuffer = GraphBuilder.CreateBuffer(FRDGBufferDesc::CreateStructuredDesc(sizeof(uint32), FMath::Max(PackedClipmaps.Num(), 1)), TEXT("GlobalDistanceField.ObjectIndexNum"));
+		AddClearUAVPass(GraphBuilder, GraphBuilder.CreateUAV(ObjectIndexNumBuffer, PF_R32_UINT), 0);
 
 		// Prepare re-cache buffers
+		bool bAnyClipmapHasPendingStreamingReadback = false;
 		for (FGlobalDistanceFieldPackedClipmap& PackedClipmap : PackedClipmaps)
 		{
 			if (Scene->DistanceFieldSceneData.NumObjectsInBuffer > 0)
@@ -2344,106 +2332,120 @@ void UpdateGlobalDistanceFieldCache(
 					CacheType == GDF_MostlyStatic &&
 					PackedClipmap.Index < 2;
 
-				PackedClipmap.HasPendingStreamingReadbackBuffer = nullptr;
-
 				if (PackedClipmap.bRecacheClipmapsWithPendingStreaming)
 				{
-					const FGlobalDistanceFieldClipmapState& ClipmapViewState = View.ViewState->GlobalDistanceFieldData->ClipmapState[PackedClipmap.Index];
+					const FGlobalDistanceFieldStreamingReadback& StreamingReadback = View.ViewState->GlobalDistanceFieldData->StreamingReadback[CacheType];
 
 					// It is not safe to EnqueueCopy on a buffer that already has a pending copy
-					PackedClipmap.bRecacheClipmapsWithPendingStreaming = PackedClipmap.bRecacheClipmapsWithPendingStreaming && ClipmapViewState.ReadbackBuffersNumPending < ClipmapViewState.MaxPendingStreamingReadbackBuffers;
+					PackedClipmap.bRecacheClipmapsWithPendingStreaming = PackedClipmap.bRecacheClipmapsWithPendingStreaming 
+						&& StreamingReadback.ReadbackBuffersNumPending < StreamingReadback.MaxPendingStreamingReadbackBuffers;
 
 					if (PackedClipmap.bRecacheClipmapsWithPendingStreaming)
 					{
-						FRDGBufferDesc HasPendingStreamingReadbackDesc = FRDGBufferDesc::CreateStructuredDesc(sizeof(uint32), 1);
-						HasPendingStreamingReadbackDesc.Usage = EBufferUsageFlags(HasPendingStreamingReadbackDesc.Usage | BUF_SourceCopy);
-						PackedClipmap.HasPendingStreamingReadbackBuffer = GraphBuilder.CreateBuffer(HasPendingStreamingReadbackDesc, TEXT("GlobalDistanceField.HasPendingStreamingReadback"));
-						AddClearUAVPass(GraphBuilder, GraphBuilder.CreateUAV(PackedClipmap.HasPendingStreamingReadbackBuffer, PF_R32_UINT), 0);
+						bAnyClipmapHasPendingStreamingReadback = true;
 					}
 				}
 			}
 		}
 
-		// Cull the global objects to the update regions
-		for (FGlobalDistanceFieldPackedClipmap& PackedClipmap : PackedClipmaps)
+		// Prepare pending mesh SDF streaming readback buffer
+		FRDGBufferRef PendingStreamingReadbackBuffer = nullptr;
+		if (bAnyClipmapHasPendingStreamingReadback)
 		{
-			const FGlobalDistanceFieldClipmap& Clipmap = Clipmaps[PackedClipmap.Index];
+			FRDGBufferDesc HasPendingStreamingReadbackDesc = FRDGBufferDesc::CreateStructuredDesc(GlobalDistanceField::MaxClipmaps * sizeof(uint32), GlobalDistanceField::MaxClipmaps);
+			HasPendingStreamingReadbackDesc.Usage = EBufferUsageFlags(HasPendingStreamingReadbackDesc.Usage | BUF_SourceCopy);
+			PendingStreamingReadbackBuffer = GraphBuilder.CreateBuffer(HasPendingStreamingReadbackDesc, TEXT("GlobalDistanceField.HasPendingStreamingReadback"));
+			AddClearUAVPass(GraphBuilder, GraphBuilder.CreateUAV(PendingStreamingReadbackBuffer, PF_R32_UINT), 0);
+		}
 
-			if (Scene->DistanceFieldSceneData.NumObjectsInBuffer > 0)
+		// Upload packed clipmap data to GPU
+		FRDGBufferRef PackedClipmapBuffer = nullptr;
+		if (PackedClipmaps.Num() > 0)
+		{
+			TArray<FVector4f> UploadData;
+			UploadData.SetNum(GlobalDistanceField::PackedClipmapBufferStride * PackedClipmaps.Num());
+
+			for (int32 PackedClipmapIndex = 0; PackedClipmapIndex < PackedClipmaps.Num(); ++PackedClipmapIndex)
 			{
-				uint32 AcceptOftenMovingObjectsOnlyValue = 0;
+				const FGlobalDistanceFieldPackedClipmap& PackedClipmap = PackedClipmaps[PackedClipmapIndex];
 
-				if (!GAOGlobalDistanceFieldCacheMostlyStaticSeparately)
-				{
-					AcceptOftenMovingObjectsOnlyValue = 2;
-				}
-				else if (CacheType == GDF_Full)
-				{
-					// First cache is for mostly static, second contains both, inheriting static objects distance fields with a lookup
-					// So only composite often moving objects into the full global distance field
-					AcceptOftenMovingObjectsOnlyValue = 1;
-				}
+				const float RadiusThresholdScale = bLumenEnabled ? 1.0f / FMath::Clamp(View.FinalPostProcessSettings.LumenSceneDetail, .01f, 100.0f) : 1.0f;
+				const float MeshSDFRadiusThreshold = GetMinMeshSDFRadius(PackedClipmap.VoxelSize.X) * RadiusThresholdScale;
 
-				{
-					FCullObjectsToClipmapCS::FParameters* PassParameters = GraphBuilder.AllocParameters<FCullObjectsToClipmapCS::FParameters>();
-					PassParameters->RWObjectIndexBuffer = GraphBuilder.CreateUAV(PackedClipmap.ObjectIndexBuffer, PF_R32_UINT);
-					PassParameters->RWObjectIndexNumBuffer = GraphBuilder.CreateUAV(PackedClipmap.ObjectIndexNumBuffer, PF_R32_UINT);
+				UploadData[PackedClipmapIndex * GlobalDistanceField::PackedClipmapBufferStride + 0] = FVector4f((FVector3f)PackedClipmap.TranslatedBounds.GetCenter(), MeshSDFRadiusThreshold);
+				UploadData[PackedClipmapIndex * GlobalDistanceField::PackedClipmapBufferStride + 1] = FVector4f((FVector3f)PackedClipmap.TranslatedBounds.GetExtent(), PackedClipmap.InfluenceRadius * PackedClipmap.InfluenceRadius);
+				UploadData[PackedClipmapIndex * GlobalDistanceField::PackedClipmapBufferStride + 2] = FVector4f(PackedClipmap.ViewTilePosition, PackedClipmap.bRecacheClipmapsWithPendingStreaming ? 1.0f : 0.0f);
+				UploadData[PackedClipmapIndex * GlobalDistanceField::PackedClipmapBufferStride + 3] = FVector4f(PackedClipmap.RelativePreViewTranslation, PackedClipmap.Index);
+			}
 
-					PassParameters->RWHasPendingStreaming = PackedClipmap.HasPendingStreamingReadbackBuffer ? GraphBuilder.CreateUAV(PackedClipmap.HasPendingStreamingReadbackBuffer, PF_R32_UINT) : nullptr;
+			PackedClipmapBuffer =
+				CreateStructuredUploadBuffer(GraphBuilder,
+					TEXT("GlobalDistanceField.PackedClipmapBuffer"),
+					UploadData);
+		}
 
-					PassParameters->DistanceFieldObjectBuffers = DistanceFieldObjectBuffers;
-					PassParameters->DistanceFieldAtlasParameters = DistanceFieldAtlas;
+		// Cull the global objects to the update regions
+		if (Scene->DistanceFieldSceneData.NumObjectsInBuffer > 0)
+		{
+			uint32 AcceptOftenMovingObjectsOnlyValue = 0;
 
-					PassParameters->ClipmapTranslatedWorldCenter = (FVector3f)PackedClipmap.TranslatedBounds.GetCenter();
-					PassParameters->ClipmapWorldExtent = (FVector3f)PackedClipmap.TranslatedBounds.GetExtent();
-					PassParameters->AcceptOftenMovingObjectsOnly = AcceptOftenMovingObjectsOnlyValue;
-					const float RadiusThresholdScale = bLumenEnabled ? 1.0f / FMath::Clamp(View.FinalPostProcessSettings.LumenSceneDetail, .01f, 100.0f) : 1.0f;
-					PassParameters->MeshSDFRadiusThreshold = GetMinMeshSDFRadius(PackedClipmap.VoxelSize.X) * RadiusThresholdScale;
-					PassParameters->InfluenceRadiusSq = PackedClipmap.InfluenceRadius * PackedClipmap.InfluenceRadius;
+			if (!GAOGlobalDistanceFieldCacheMostlyStaticSeparately)
+			{
+				AcceptOftenMovingObjectsOnlyValue = 2;
+			}
+			else if (CacheType == GDF_Full)
+			{
+				// First cache is for mostly static, second contains both, inheriting static objects distance fields with a lookup
+				// So only composite often moving objects into the full global distance field
+				AcceptOftenMovingObjectsOnlyValue = 1;
+			}
 
-					PassParameters->ViewTilePosition = PackedClipmap.ViewTilePosition;
-					PassParameters->RelativePreViewTranslation = PackedClipmap.RelativePreViewTranslation;
+			FCullObjectsToClipmapCS::FParameters* PassParameters = GraphBuilder.AllocParameters<FCullObjectsToClipmapCS::FParameters>();
+			PassParameters->RWObjectIndexBuffer = GraphBuilder.CreateUAV(ObjectIndexBuffer, PF_R32_UINT);
+			PassParameters->RWObjectIndexNumBuffer = GraphBuilder.CreateUAV(ObjectIndexNumBuffer, PF_R32_UINT);
+			PassParameters->RWHasPendingStreaming = PendingStreamingReadbackBuffer ? GraphBuilder.CreateUAV(PendingStreamingReadbackBuffer, PF_R32_UINT) : nullptr;
+			PassParameters->PackedClipmapBuffer = GraphBuilder.CreateSRV(PackedClipmapBuffer);
+			PassParameters->AcceptOftenMovingObjectsOnly = AcceptOftenMovingObjectsOnlyValue;
+			PassParameters->NumPackedClipmaps = PackedClipmaps.Num();
+			PassParameters->ObjectIndexBufferStride = ObjectIndexBufferStride;
+			PassParameters->DistanceFieldObjectBuffers = DistanceFieldObjectBuffers;
+			PassParameters->DistanceFieldAtlasParameters = DistanceFieldAtlas;
 
-					FCullObjectsToClipmapCS::FPermutationDomain PermutationVector;
-					PermutationVector.Set<FCullObjectsToClipmapCS::FReadbackHasPendingStreaming>(PackedClipmap.bRecacheClipmapsWithPendingStreaming);
-					auto ComputeShader = View.ShaderMap->GetShader<FCullObjectsToClipmapCS>(PermutationVector);
+			FCullObjectsToClipmapCS::FPermutationDomain PermutationVector;
+			PermutationVector.Set<FCullObjectsToClipmapCS::FReadbackHasPendingStreaming>(bAnyClipmapHasPendingStreamingReadback);
+			auto ComputeShader = View.ShaderMap->GetShader<FCullObjectsToClipmapCS>(PermutationVector);
 
-					const FIntVector GroupCount = FComputeShaderUtils::GetGroupCountWrapped(DistanceFieldSceneData.NumObjectsInBuffer, FCullObjectsToClipmapCS::GetGroupSize());
+			const FIntVector GroupCount = FComputeShaderUtils::GetGroupCountWrapped(DistanceFieldSceneData.NumObjectsInBuffer, FCullObjectsToClipmapCS::GetGroupSize());
 
-					FComputeShaderUtils::AddPass(
-						GraphBuilder,
-						RDG_EVENT_NAME("CullToClipmap %d", PackedClipmap.Index),
-						ComputeShader,
-						PassParameters,
-						GroupCount);
-				}
-			}
+			FComputeShaderUtils::AddPass(
+				GraphBuilder,
+				RDG_EVENT_NAME("CullToClipmaps"),
+				ComputeShader,
+				PassParameters,
+				GroupCount);
 		}
 
 		// Readback re-cache requests
-		for (const FGlobalDistanceFieldPackedClipmap& PackedClipmap : PackedClipmaps)
+		if (PendingStreamingReadbackBuffer)
 		{
-			if (PackedClipmap.bRecacheClipmapsWithPendingStreaming)
-			{
-				FGlobalDistanceFieldClipmapState& ClipmapViewState = View.ViewState->GlobalDistanceFieldData->ClipmapState[PackedClipmap.Index];
+			FGlobalDistanceFieldStreamingReadback& StreamingReadback = View.ViewState->GlobalDistanceFieldData->StreamingReadback[CacheType];
 
-				if (!ClipmapViewState.HasPendingStreamingReadbackBuffers[ClipmapViewState.ReadbackBuffersWriteIndex].IsValid())
-				{
-					ClipmapViewState.HasPendingStreamingReadbackBuffers[ClipmapViewState.ReadbackBuffersWriteIndex] =
-						MakeUnique<FRHIGPUBufferReadback>(TEXT("GlobalDistanceField.HasPendingStreamingReadback"));
-				}
+			if (!StreamingReadback.PendingStreamingReadbackBuffers[StreamingReadback.ReadbackBuffersWriteIndex].IsValid())
+			{
+				StreamingReadback.PendingStreamingReadbackBuffers[StreamingReadback.ReadbackBuffersWriteIndex] =
+					MakeUnique<FRHIGPUBufferReadback>(TEXT("GlobalDistanceField.PendingStreamingReadback"));
+			}
 
-				FRHIGPUBufferReadback* ReadbackBuffer = ClipmapViewState.HasPendingStreamingReadbackBuffers[ClipmapViewState.ReadbackBuffersWriteIndex].Get();
+			FRHIGPUBufferReadback* ReadbackBuffer = StreamingReadback.PendingStreamingReadbackBuffers[StreamingReadback.ReadbackBuffersWriteIndex].Get();
 
-				AddReadbackBufferPass(GraphBuilder, RDG_EVENT_NAME("GlobalDistanceField.HasPendingStreamingReadback"), PackedClipmap.HasPendingStreamingReadbackBuffer,
-					[ReadbackBuffer, HasPendingStreamingReadbackBuffer = PackedClipmap.HasPendingStreamingReadbackBuffer](FRHICommandList& RHICmdList)
-					{
-						ReadbackBuffer->EnqueueCopy(RHICmdList, HasPendingStreamingReadbackBuffer->GetRHI(), 0u);
-					});
+			AddReadbackBufferPass(GraphBuilder, RDG_EVENT_NAME("GlobalDistanceField.HasPendingStreamingReadback"), PendingStreamingReadbackBuffer,
+				[ReadbackBuffer, PendingStreamingReadbackBuffer = PendingStreamingReadbackBuffer](FRHICommandList& RHICmdList)
+				{
+					ReadbackBuffer->EnqueueCopy(RHICmdList, PendingStreamingReadbackBuffer->GetRHI(), 0u);
+				});
 
-				ClipmapViewState.ReadbackBuffersWriteIndex = (ClipmapViewState.ReadbackBuffersWriteIndex + 1u) % ClipmapViewState.MaxPendingStreamingReadbackBuffers;
-				ClipmapViewState.ReadbackBuffersNumPending = FMath::Min(ClipmapViewState.ReadbackBuffersNumPending + 1u, ClipmapViewState.MaxPendingStreamingReadbackBuffers);
-			}
+			StreamingReadback.ReadbackBuffersWriteIndex = (StreamingReadback.ReadbackBuffersWriteIndex + 1u) % StreamingReadback.MaxPendingStreamingReadbackBuffers;
+			StreamingReadback.ReadbackBuffersNumPending = FMath::Min(StreamingReadback.ReadbackBuffersNumPending + 1u, StreamingReadback.MaxPendingStreamingReadbackBuffers);
 		}
 
 		// Clear indirect dispatch arguments
@@ -2635,16 +2637,20 @@ void UpdateGlobalDistanceFieldCache(
 		// Cull objects into a cull grid
 		if (Scene->DistanceFieldSceneData.NumObjectsInBuffer > 0)
 		{
-			for (FGlobalDistanceFieldPackedClipmap& PackedClipmap : PackedClipmaps)
+			for (int32 PackedClipmapIndex = 0; PackedClipmapIndex < PackedClipmaps.Num(); ++PackedClipmapIndex)
 			{
+				const FGlobalDistanceFieldPackedClipmap& PackedClipmap = PackedClipmaps[PackedClipmapIndex];
+
 				FCullObjectsToGridCS::FParameters* PassParameters = GraphBuilder.AllocParameters<FCullObjectsToGridCS::FParameters>();
 				PassParameters->RWCullGridAllocator = GraphBuilder.CreateUAV(PackedClipmap.CullGridAllocator, PF_R32_UINT);
 				PassParameters->RWCullGridObjectHeader = GraphBuilder.CreateUAV(PackedClipmap.CullGridObjectHeader, PF_R32_UINT);
 				PassParameters->RWCullGridObjectArray = GraphBuilder.CreateUAV(PackedClipmap.CullGridObjectArray, PF_R32_UINT);
 				PassParameters->CullGridIndirectArgBuffer = PackedClipmap.CullGridUpdateIndirectArgBuffer;
 				PassParameters->CullGridTileBuffer = GraphBuilder.CreateSRV(PackedClipmap.CullGridUpdateTileBuffer, PF_R32_UINT);
-				PassParameters->ObjectIndexBuffer = GraphBuilder.CreateSRV(PackedClipmap.ObjectIndexBuffer, PF_R32_UINT);
-				PassParameters->ObjectIndexNumBuffer = GraphBuilder.CreateSRV(PackedClipmap.ObjectIndexNumBuffer, PF_R32_UINT);
+				PassParameters->ObjectIndexBuffer = GraphBuilder.CreateSRV(ObjectIndexBuffer, PF_R32_UINT);
+				PassParameters->ObjectIndexNumBuffer = GraphBuilder.CreateSRV(ObjectIndexNumBuffer, PF_R32_UINT);
+				PassParameters->PackedClipmapIndex = PackedClipmapIndex;
+				PassParameters->ObjectIndexBufferStride = ObjectIndexBufferStride;
 				PassParameters->DistanceFieldObjectBuffers = DistanceFieldObjectBuffers;
 				PassParameters->CullGridResolution = PackedClipmap.CullGridResolution;
 				PassParameters->CullGridCoordToTranslatedWorldCenterScale = (FVector3f)PackedClipmap.CullGridCoordToTranslatedWorldCenterScale;
@@ -2809,8 +2815,6 @@ void UpdateGlobalDistanceFieldCache(
 					PassParameters->ParentPageTableLayerTexture = ParentPageTableLayerTexture;
 					PassParameters->CullGridObjectHeader = GraphBuilder.CreateSRV(PackedClipmap.CullGridObjectHeader, PF_R32_UINT);
 					PassParameters->CullGridObjectArray = GraphBuilder.CreateSRV(PackedClipmap.CullGridObjectArray, PF_R32_UINT);
-					PassParameters->ObjectIndexBuffer = GraphBuilder.CreateSRV(PackedClipmap.ObjectIndexBuffer, PF_R32_UINT);
-					PassParameters->ObjectIndexNumBuffer = GraphBuilder.CreateSRV(PackedClipmap.ObjectIndexNumBuffer, PF_R32_UINT);
 					PassParameters->DistanceFieldObjectBuffers = DistanceFieldObjectBuffers;
 					PassParameters->DistanceFieldAtlas = DistanceFieldAtlas;
 					PassParameters->GlobalDistanceFieldUpdateParameters = PackedClipmap.UpdateParameters;
diff --git a/Engine/Source/Runtime/Renderer/Private/ScenePrivate.h b/Engine/Source/Runtime/Renderer/Private/ScenePrivate.h
--- a/Engine/Source/Runtime/Renderer/Private/ScenePrivate.h
+++ b/Engine/Source/Runtime/Renderer/Private/ScenePrivate.h
@@ -484,8 +484,6 @@ public:
 		CacheClipmapInfluenceRadius = 0.0f;
 		CacheMostlyStaticSeparately = 1;
 		LastUsedSceneDataForFullUpdate = nullptr;
-
-		HasPendingStreamingReadbackBuffers.SetNum(MaxPendingStreamingReadbackBuffers);
 	}
 
 	FInt64Vector FullUpdateOriginInPages;
@@ -496,11 +494,6 @@ public:
 	float CachedClipmapExtent;
 	float CacheClipmapInfluenceRadius;
 
-	TArray<TUniquePtr<FRHIGPUBufferReadback>> HasPendingStreamingReadbackBuffers;
-	uint32 MaxPendingStreamingReadbackBuffers = 4;
-	uint32 ReadbackBuffersWriteIndex = 0;
-	uint32 ReadbackBuffersNumPending = 0;
-
 	FGlobalDistanceFieldCacheTypeState Cache[GDF_Num];
 
 	// Used to perform a full update of the clip map when the scene data changes
@@ -670,6 +663,19 @@ private:
 	bool bSeparatedAtmosphereMieRayLeigh;
 };
 
+struct FGlobalDistanceFieldStreamingReadback
+{
+	FGlobalDistanceFieldStreamingReadback()
+	{
+		PendingStreamingReadbackBuffers.SetNum(MaxPendingStreamingReadbackBuffers);
+	}
+
+	TArray<TUniquePtr<FRHIGPUBufferReadback>> PendingStreamingReadbackBuffers;
+	uint32 MaxPendingStreamingReadbackBuffers = 4;
+	uint32 ReadbackBuffersWriteIndex = 0;
+	uint32 ReadbackBuffersNumPending = 0;
+};
+
 struct FPersistentGlobalDistanceFieldData : public FThreadSafeRefCountedObject
 {
 	// Array of ClipmapIndex
@@ -690,6 +696,8 @@ struct FPersistentGlobalDistanceFieldData : public FThreadSafeRefCountedObject
 	FRHIGPUMask LastGPUMask;
 #endif
 
+	FGlobalDistanceFieldStreamingReadback StreamingReadback[GDF_Num];
+
 	TRefCountPtr<FRDGPooledBuffer> PageFreeListAllocatorBuffer;
 	TRefCountPtr<FRDGPooledBuffer> PageFreeListBuffer;
 	TRefCountPtr<IPooledRenderTarget> PageAtlasTexture;
