Global Distance Field - fixed crash when updating clipmaps without any update bounds

diff --git a/Engine/Source/Runtime/Renderer/Private/GlobalDistanceField.cpp b/Engine/Source/Runtime/Renderer/Private/GlobalDistanceField.cpp
--- a/Engine/Source/Runtime/Renderer/Private/GlobalDistanceField.cpp
+++ b/Engine/Source/Runtime/Renderer/Private/GlobalDistanceField.cpp
@@ -2360,10 +2360,18 @@ void UpdateGlobalDistanceFieldCache(
 
 		// Upload packed clipmap data to GPU
 		FRDGBufferRef PackedClipmapBuffer = nullptr;
-		if (PackedClipmaps.Num() > 0)
 		{
 			TArray<FVector4f> UploadData;
-			UploadData.SetNum(GlobalDistanceField::PackedClipmapBufferStride * PackedClipmaps.Num());
+
+			if (PackedClipmaps.Num() > 0)
+			{
+				UploadData.SetNum(GlobalDistanceField::PackedClipmapBufferStride * PackedClipmaps.Num());
+			}
+			else
+			{
+				// Create a dummy entry filled with 0 in order to be able to create a dummy PackedClipmapBuffer resource
+				UploadData.SetNumZeroed(GlobalDistanceField::PackedClipmapBufferStride);
+			}
 
 			for (int32 PackedClipmapIndex = 0; PackedClipmapIndex < PackedClipmaps.Num(); ++PackedClipmapIndex)
 			{
