Fixes a D3D11 validation failure binding a Texture2DArray into a RWTexture2D with TSR, sg.AntiAliasingQuality=0, sg.MotionBlurQuality=0

TSR's UpdateHistory now always render half or quarter res in a Texture2DArray of size 1, to not introduce shader permutation or dynamic branch on SceneColorOutputMip1 that could create dependent image_store on some hardware.

diff --git a/Engine/Shaders/Private/TemporalSuperResolution/TSRUpdateHistory.usf b/Engine/Shaders/Private/TemporalSuperResolution/TSRUpdateHistory.usf
--- a/Engine/Shaders/Private/TemporalSuperResolution/TSRUpdateHistory.usf
+++ b/Engine/Shaders/Private/TemporalSuperResolution/TSRUpdateHistory.usf
@@ -114,12 +114,11 @@ uint HistoryArrayIndices_HighFrequency;
 Texture2DArray<tsr_halfC> PrevHistoryColorTexture;
 Texture2DArray<tsr_halfM> PrevHistoryMetadataTexture;
 
-RWTexture2D<tsr_halfC> SceneColorOutputMip0;
-RWTexture2D<tsr_halfC> SceneColorOutputMip1;
-
 RWTexture2DArray<tsr_halfC> HistoryColorOutput;
 RWTexture2DArray<tsr_halfM> HistoryMetadataOutput;
 
+RWTexture2DArray<tsr_halfC> SceneColorOutputMip1;
+
 
 //------------------------------------------------------- LDS
 
@@ -1108,10 +1107,10 @@ void MainCS(
 				}
 			}
 
-			SceneColorOutputMip1[dpv_lo(HalfResOutputPixelPos)] = dpv_lo(HalfResOutput);
+			SceneColorOutputMip1[tsr_short3(dpv_lo(HalfResOutputPixelPos), 0)] = dpv_lo(HalfResOutput);
 			
 			#if CONFIG_ENABLE_DUAL_PIXEL_VECTORIZATION
-				SceneColorOutputMip1[dpv_hi(HalfResOutputPixelPos)] = dpv_hi(HalfResOutput);
+				SceneColorOutputMip1[tsr_short3(dpv_hi(HalfResOutputPixelPos), 0)] = dpv_hi(HalfResOutput);
 			#endif
 		}
 		
diff --git a/Engine/Source/Runtime/Renderer/Private/PostProcess/PostProcessBloomSetup.cpp b/Engine/Source/Runtime/Renderer/Private/PostProcess/PostProcessBloomSetup.cpp
--- a/Engine/Source/Runtime/Renderer/Private/PostProcess/PostProcessBloomSetup.cpp
+++ b/Engine/Source/Runtime/Renderer/Private/PostProcess/PostProcessBloomSetup.cpp
@@ -28,7 +28,7 @@ TAutoConsoleVariable<float> CVarBloomCross(
 BEGIN_SHADER_PARAMETER_STRUCT(FBloomSetupParameters, )
 	SHADER_PARAMETER_STRUCT_REF(FViewUniformShaderParameters, View)
 	SHADER_PARAMETER_STRUCT(FScreenPassTextureViewportParameters, Input)
-	SHADER_PARAMETER_RDG_TEXTURE(Texture2D, InputTexture)
+	SHADER_PARAMETER_RDG_TEXTURE_SRV(Texture2D, InputTexture)
 	SHADER_PARAMETER_SAMPLER(SamplerState, InputSampler)
 	SHADER_PARAMETER_RDG_TEXTURE(Texture2D, LumBilateralGrid)
 	SHADER_PARAMETER_SAMPLER(SamplerState, LumBilateralGridSampler)
@@ -49,7 +49,7 @@ FBloomSetupParameters GetBloomSetupParameters(
 	FBloomSetupParameters Parameters;
 	Parameters.View = View.ViewUniformBuffer;
 	Parameters.Input = GetScreenPassTextureViewportParameters(InputViewport);
-	Parameters.InputTexture = Inputs.SceneColor.Texture;
+	Parameters.InputTexture = Inputs.SceneColor.TextureSRV;
 	Parameters.InputSampler = TStaticSamplerState<SF_Bilinear, AM_Clamp, AM_Clamp, AM_Clamp>::GetRHI();
 	Parameters.LumBilateralGrid = Inputs.LocalExposureTexture;
 	Parameters.LumBilateralGridSampler = TStaticSamplerState<SF_Bilinear, AM_Clamp, AM_Clamp, AM_Clamp>::GetRHI();
@@ -128,9 +128,12 @@ FScreenPassTexture AddBloomSetupPass(FRDGBuilder& GraphBuilder, const FViewInfo&
 
 	check(bLocalExposureEnabled || bThresholdEnabled);
 
-	FRDGTextureDesc OutputDesc = Inputs.SceneColor.Texture->Desc;
-	OutputDesc.Reset();
-	OutputDesc.Flags |= bIsComputePass ? TexCreate_UAV : TexCreate_RenderTargetable;
+	const FRDGTextureDesc& InputDesc = Inputs.SceneColor.TextureSRV->Desc.Texture->Desc;
+	FRDGTextureDesc OutputDesc = FRDGTextureDesc::Create2D(
+		InputDesc.Extent,
+		InputDesc.Format,
+		FClearValueBinding::None,
+		/* InFlags = */ TexCreate_ShaderResource | (bIsComputePass ? TexCreate_UAV : TexCreate_RenderTargetable));
 
 	const FScreenPassTextureViewport Viewport(Inputs.SceneColor);
 	const FScreenPassRenderTarget Output(GraphBuilder.CreateTexture(OutputDesc, TEXT("BloomSetup")), Viewport.Rect, ERenderTargetLoadAction::ENoAction);
diff --git a/Engine/Source/Runtime/Renderer/Private/PostProcess/PostProcessBloomSetup.h b/Engine/Source/Runtime/Renderer/Private/PostProcess/PostProcessBloomSetup.h
--- a/Engine/Source/Runtime/Renderer/Private/PostProcess/PostProcessBloomSetup.h
+++ b/Engine/Source/Runtime/Renderer/Private/PostProcess/PostProcessBloomSetup.h
@@ -26,7 +26,7 @@ FScreenPassTexture AddGaussianBloomPasses(FRDGBuilder& GraphBuilder, const FView
 struct FBloomSetupInputs
 {
 	// [Required]: The intermediate scene color being processed.
-	FScreenPassTexture SceneColor;
+	FScreenPassTextureSlice SceneColor;
 
 	// [Required]: The scene eye adaptation buffer.
 	FRDGBufferRef EyeAdaptationBuffer = nullptr;
diff --git a/Engine/Source/Runtime/Renderer/Private/PostProcess/PostProcessDownsample.cpp b/Engine/Source/Runtime/Renderer/Private/PostProcess/PostProcessDownsample.cpp
--- a/Engine/Source/Runtime/Renderer/Private/PostProcess/PostProcessDownsample.cpp
+++ b/Engine/Source/Runtime/Renderer/Private/PostProcess/PostProcessDownsample.cpp
@@ -108,20 +108,15 @@ FScreenPassTexture AddDownsamplePass(
 
 	// Construct the output texture to be half resolution (rounded up to even) with an optional format override.
 	{
-		FRDGTextureDesc Desc = Inputs.SceneColor.Texture->Desc;
-		Desc.Reset();
-		Desc.Extent = FIntPoint::DivideAndRoundUp(Desc.Extent, 2);
+		const FRDGTextureDesc& InputDesc = Inputs.SceneColor.TextureSRV->Desc.Texture->Desc;
+		FRDGTextureDesc Desc = FRDGTextureDesc::Create2D(
+			FIntPoint::DivideAndRoundUp(InputDesc.Extent, 2),
+			Inputs.FormatOverride != PF_Unknown ? Inputs.FormatOverride  : InputDesc.Format,
+			FClearValueBinding(FLinearColor(0, 0, 0, 0)),
+			/* InFlags = */ TexCreate_ShaderResource | GFastVRamConfig.Downsample | (bIsComputePass ? TexCreate_UAV : (TexCreate_RenderTargetable | TexCreate_NoFastClear)));
+
 		Desc.Extent.X = FMath::Max(1, Desc.Extent.X);
 		Desc.Extent.Y = FMath::Max(1, Desc.Extent.Y);
-		Desc.Flags &= ~(TexCreate_RenderTargetable | TexCreate_UAV | TexCreate_Presentable);
-		Desc.Flags |= bIsComputePass ? TexCreate_UAV : (TexCreate_RenderTargetable | TexCreate_NoFastClear);
-		Desc.Flags |= GFastVRamConfig.Downsample;
-		Desc.ClearValue = FClearValueBinding(FLinearColor(0, 0, 0, 0));
-
-		if (Inputs.FormatOverride != PF_Unknown)
-		{
-			Desc.Format = Inputs.FormatOverride;
-		}
 
 		if (Inputs.UserSuppliedOutput && Translate(Inputs.UserSuppliedOutput->GetDesc()) == Desc)
 		{
@@ -145,7 +140,7 @@ FScreenPassTexture AddDownsamplePass(
 		PermutationVector.Set<FDownsampleQualityDimension>(Inputs.Quality);
 
 		FDownsamplePS::FParameters* PassParameters = GraphBuilder.AllocParameters<FDownsamplePS::FParameters>();
-		PassParameters->Common = GetDownsampleParameters(View, Output, FScreenPassTextureSlice::CreateFromScreenPassTexture(GraphBuilder, Inputs.SceneColor), Inputs.Quality);
+		PassParameters->Common = GetDownsampleParameters(View, Output, Inputs.SceneColor, Inputs.Quality);
 		PassParameters->RenderTargets[0] = Output.GetRenderTargetBinding();
 
 		TShaderMapRef<FDownsamplePS> PixelShader(View.ShaderMap, PermutationVector);
@@ -200,7 +195,7 @@ void FSceneDownsampleChain::Init(
 	FRDGBuilder& GraphBuilder,
 	const FViewInfo& View,
 	const FEyeAdaptationParameters& EyeAdaptationParameters,
-	FScreenPassTexture HalfResolutionSceneColor,
+	FScreenPassTextureSlice HalfResolutionSceneColor,
 	EDownsampleQuality DownsampleQuality,
 	bool bLogLumaInAlpha)
 {
@@ -231,13 +226,13 @@ void FSceneDownsampleChain::Init(
 		PassInputs.SceneColor = Textures[PreviousStageIndex];
 		PassInputs.Quality = DownsampleQuality;
 
-		Textures[StageIndex] = AddDownsamplePass(GraphBuilder, View, PassInputs);
+		Textures[StageIndex] = FScreenPassTextureSlice::CreateFromScreenPassTexture(GraphBuilder, AddDownsamplePass(GraphBuilder, View, PassInputs));
 
 		if (bLogLumaInAlpha)
 		{
 			bLogLumaInAlpha = false;
 
-			Textures[StageIndex] = AddBasicEyeAdaptationSetupPass(GraphBuilder, View, EyeAdaptationParameters, Textures[StageIndex]);
+			Textures[StageIndex] = FScreenPassTextureSlice::CreateFromScreenPassTexture(GraphBuilder, AddBasicEyeAdaptationSetupPass(GraphBuilder, View, EyeAdaptationParameters, FScreenPassTexture(Textures[StageIndex])));
 		}
 	}
 
diff --git a/Engine/Source/Runtime/Renderer/Private/PostProcess/PostProcessDownsample.h b/Engine/Source/Runtime/Renderer/Private/PostProcess/PostProcessDownsample.h
--- a/Engine/Source/Runtime/Renderer/Private/PostProcess/PostProcessDownsample.h
+++ b/Engine/Source/Runtime/Renderer/Private/PostProcess/PostProcessDownsample.h
@@ -38,7 +38,7 @@ struct FDownsamplePassInputs
 	IPooledRenderTarget* UserSuppliedOutput = nullptr;
 
 	// Input scene color RDG texture / view rect. Must not be null.
-	FScreenPassTexture SceneColor;
+	FScreenPassTextureSlice SceneColor;
 
 	// The downsample method to use.
 	EDownsampleQuality Quality = EDownsampleQuality::Low;
@@ -68,7 +68,7 @@ public:
 		FRDGBuilder& GraphBuilder,
 		const FViewInfo& View,
 		const FEyeAdaptationParameters& EyeAdaptationParameters,
-		FScreenPassTexture HalfResolutionSceneColor,
+		FScreenPassTextureSlice HalfResolutionSceneColor,
 		EDownsampleQuality DownsampleQuality,
 		bool bLogLumaInAlpha);
 
@@ -77,22 +77,22 @@ public:
 		return bInitialized;
 	}
 
-	FScreenPassTexture GetTexture(uint32 StageIndex) const
+	FScreenPassTextureSlice GetTexture(uint32 StageIndex) const
 	{
 		return Textures[StageIndex];
 	}
 
-	FScreenPassTexture GetFirstTexture() const
+	FScreenPassTextureSlice GetFirstTexture() const
 	{
 		return Textures[0];
 	}
 
-	FScreenPassTexture GetLastTexture() const
+	FScreenPassTextureSlice GetLastTexture() const
 	{
 		return Textures[StageCount - 1];
 	}
 
 private:
-	TStaticArray<FScreenPassTexture, StageCount> Textures;
+	TStaticArray<FScreenPassTextureSlice, StageCount> Textures;
 	bool bInitialized = false;
 };
\ No newline at end of file
diff --git a/Engine/Source/Runtime/Renderer/Private/PostProcess/PostProcessEyeAdaptation.cpp b/Engine/Source/Runtime/Renderer/Private/PostProcess/PostProcessEyeAdaptation.cpp
--- a/Engine/Source/Runtime/Renderer/Private/PostProcess/PostProcessEyeAdaptation.cpp
+++ b/Engine/Source/Runtime/Renderer/Private/PostProcess/PostProcessEyeAdaptation.cpp
@@ -944,7 +944,7 @@ public:
 		SHADER_PARAMETER_STRUCT(FEyeAdaptationParameters, EyeAdaptation)
 		SHADER_PARAMETER_STRUCT(FLocalExposureParameters, LocalExposure)
 		SHADER_PARAMETER_STRUCT(FScreenPassTextureViewportParameters, Color)
-		SHADER_PARAMETER_RDG_TEXTURE(Texture2D, ColorTexture)
+		SHADER_PARAMETER_RDG_TEXTURE_SRV(Texture2D, ColorTexture)
 		SHADER_PARAMETER_RDG_BUFFER_SRV(StructuredBuffer<float4>, EyeAdaptationBuffer)
 		SHADER_PARAMETER_RDG_BUFFER_UAV(RWStructuredBuffer<float4>, RWEyeAdaptationBuffer)
 	END_SHADER_PARAMETER_STRUCT()
@@ -965,7 +965,7 @@ FRDGBufferRef AddBasicEyeAdaptationPass(
 	const FViewInfo& View,
 	const FEyeAdaptationParameters& EyeAdaptationParameters,
 	const FLocalExposureParameters& LocalExposureParameters,
-	FScreenPassTexture SceneColor,
+	FScreenPassTextureSlice SceneColor,
 	FRDGBufferRef EyeAdaptationBuffer,
 	bool bComputeAverageLocalExposure)
 {
@@ -981,7 +981,7 @@ FRDGBufferRef AddBasicEyeAdaptationPass(
 	PassParameters->EyeAdaptation = EyeAdaptationParameters;
 	PassParameters->LocalExposure = LocalExposureParameters;
 	PassParameters->Color = GetScreenPassTextureViewportParameters(SceneColorViewport);
-	PassParameters->ColorTexture = SceneColor.Texture;
+	PassParameters->ColorTexture = SceneColor.TextureSRV;
 	PassParameters->EyeAdaptationBuffer = GraphBuilder.CreateSRV(EyeAdaptationBuffer);
 	PassParameters->RWEyeAdaptationBuffer = GraphBuilder.CreateUAV(OutputBuffer);
 
diff --git a/Engine/Source/Runtime/Renderer/Private/PostProcess/PostProcessEyeAdaptation.h b/Engine/Source/Runtime/Renderer/Private/PostProcess/PostProcessEyeAdaptation.h
--- a/Engine/Source/Runtime/Renderer/Private/PostProcess/PostProcessEyeAdaptation.h
+++ b/Engine/Source/Runtime/Renderer/Private/PostProcess/PostProcessEyeAdaptation.h
@@ -117,7 +117,7 @@ FRDGBufferRef AddBasicEyeAdaptationPass(
 	const FViewInfo& View,
 	const FEyeAdaptationParameters& EyeAdaptationParameters,
 	const FLocalExposureParameters& LocalExposureParameters,
-	FScreenPassTexture SceneColor,
+	FScreenPassTextureSlice SceneColor,
 	FRDGBufferRef EyeAdaptationBuffer,
 	bool bComputeAverageLocalExposure);
 
diff --git a/Engine/Source/Runtime/Renderer/Private/PostProcess/PostProcessHistogram.cpp b/Engine/Source/Runtime/Renderer/Private/PostProcess/PostProcessHistogram.cpp
--- a/Engine/Source/Runtime/Renderer/Private/PostProcess/PostProcessHistogram.cpp
+++ b/Engine/Source/Runtime/Renderer/Private/PostProcess/PostProcessHistogram.cpp
@@ -38,7 +38,7 @@ public:
 		SHADER_PARAMETER_STRUCT(FScreenPassTextureViewportParameters, Input)
 		SHADER_PARAMETER_STRUCT(FEyeAdaptationParameters, EyeAdaptation)
 		SHADER_PARAMETER_SAMPLER(SamplerState, InputSampler)
-		SHADER_PARAMETER_RDG_TEXTURE(Texture2D, InputTexture)
+		SHADER_PARAMETER_RDG_TEXTURE_SRV(Texture2D, InputTexture)
 		SHADER_PARAMETER_RDG_TEXTURE_UAV(RWTexture2D<float4>, HistogramRWTexture)
 		SHADER_PARAMETER_RDG_TEXTURE_UAV(RWTexture3D<float2>, BilateralGridRWTexture)
 		SHADER_PARAMETER_RDG_TEXTURE_UAV(RWTexture2D<float>, DebugOutput)
@@ -164,7 +164,7 @@ public:
 		SHADER_PARAMETER_STRUCT_INCLUDE(FSceneTextureParameters, SceneTextures)
 		SHADER_PARAMETER_STRUCT(FScreenPassTextureViewportParameters, Input)
 		SHADER_PARAMETER_STRUCT(FEyeAdaptationParameters, EyeAdaptation)
-		SHADER_PARAMETER_RDG_TEXTURE(Texture2D, InputTexture)
+		SHADER_PARAMETER_RDG_TEXTURE_SRV(Texture2D, InputTexture)
 		SHADER_PARAMETER_RDG_BUFFER_SRV(StructuredBuffer<float4>, EyeAdaptationBuffer)
 		SHADER_PARAMETER(FIntPoint, ThreadGroupCount)
 		SHADER_PARAMETER_RDG_TEXTURE_UAV(RWTexture2D<float4>, HistogramScatter64Output)
@@ -240,7 +240,7 @@ static FRDGTextureRef AddHistogramLegacyPass(
 	FRDGBuilder& GraphBuilder,
 	const FViewInfo& View,
 	const FEyeAdaptationParameters& EyeAdaptationParameters,
-	const FScreenPassTexture& SceneColor,
+	const FScreenPassTextureSlice& SceneColor,
 	FRDGBufferRef EyeAdaptationBuffer)
 {
 	check(SceneColor.IsValid());
@@ -270,7 +270,7 @@ static FRDGTextureRef AddHistogramLegacyPass(
 		PassParameters->View = View.ViewUniformBuffer;
 		PassParameters->Input = GetScreenPassTextureViewportParameters(FScreenPassTextureViewport(SceneColor));
 		PassParameters->InputSampler = TStaticSamplerState<SF_Point, AM_Clamp, AM_Clamp, AM_Clamp>::GetRHI();
-		PassParameters->InputTexture = SceneColor.Texture;
+		PassParameters->InputTexture = SceneColor.TextureSRV;
 		PassParameters->HistogramRWTexture = GraphBuilder.CreateUAV(HistogramTexture);
 		PassParameters->ThreadGroupCount = HistogramThreadGroupCount;
 		PassParameters->EyeAdaptation = EyeAdaptationParameters;
@@ -338,7 +338,7 @@ static FRDGTextureRef AddHistogramAtomicPass(
 	FRDGBuilder& GraphBuilder,
 	const FViewInfo& View,
 	const FEyeAdaptationParameters& EyeAdaptationParameters,
-	const FScreenPassTexture& SceneColor,
+	const FScreenPassTextureSlice& SceneColor,
 	const FSceneTextureParameters& SceneTextures,
 	FRDGBufferRef EyeAdaptationBuffer)
 {
@@ -374,7 +374,7 @@ static FRDGTextureRef AddHistogramAtomicPass(
 		PassParameters->View = View.ViewUniformBuffer;
 		PassParameters->SceneTextures = SceneTextures;
 		PassParameters->Input = GetScreenPassTextureViewportParameters(FScreenPassTextureViewport(SceneColor));
-		PassParameters->InputTexture = SceneColor.Texture;
+		PassParameters->InputTexture = SceneColor.TextureSRV;
 		PassParameters->EyeAdaptationBuffer = GraphBuilder.CreateSRV(EyeAdaptationBuffer);
 		PassParameters->ThreadGroupCount = FIntPoint(SceneColor.ViewRect.Size().Y, 1);
 		PassParameters->EyeAdaptation = EyeAdaptationParameters;
@@ -448,7 +448,7 @@ FRDGTextureRef AddHistogramPass(
 	FRDGBuilder& GraphBuilder,
 	const FViewInfo& View,
 	const FEyeAdaptationParameters& EyeAdaptationParameters,
-	FScreenPassTexture SceneColor,
+	FScreenPassTextureSlice SceneColor,
 	const FSceneTextureParameters& SceneTextures,
 	FRDGBufferRef EyeAdaptationBuffer)
 {
@@ -486,7 +486,7 @@ FRDGTextureRef AddLocalExposurePass(
 	FRDGBuilder& GraphBuilder,
 	const FViewInfo& View,
 	const FEyeAdaptationParameters& EyeAdaptationParameters,
-	FScreenPassTexture SceneColor)
+	FScreenPassTextureSlice SceneColor)
 {
 	check(SceneColor.IsValid());
 
@@ -514,7 +514,7 @@ FRDGTextureRef AddLocalExposurePass(
 		PassParameters->EyeAdaptation = EyeAdaptationParameters;
 		PassParameters->Input = GetScreenPassTextureViewportParameters(FScreenPassTextureViewport(SceneColor));
 		PassParameters->InputSampler = TStaticSamplerState<SF_Point, AM_Clamp, AM_Clamp, AM_Clamp>::GetRHI();
-		PassParameters->InputTexture = SceneColor.Texture;
+		PassParameters->InputTexture = SceneColor.TextureSRV;
 		PassParameters->BilateralGridRWTexture = GraphBuilder.CreateUAV(LocalExposureTexture);
 		PassParameters->ThreadGroupCount = ThreadGroupCount;
 
diff --git a/Engine/Source/Runtime/Renderer/Private/PostProcess/PostProcessHistogram.h b/Engine/Source/Runtime/Renderer/Private/PostProcess/PostProcessHistogram.h
--- a/Engine/Source/Runtime/Renderer/Private/PostProcess/PostProcessHistogram.h
+++ b/Engine/Source/Runtime/Renderer/Private/PostProcess/PostProcessHistogram.h
@@ -7,13 +7,13 @@
 class FEyeAdaptationParameters;
 class FSceneTextureParameters;
 class FViewInfo;
-struct FScreenPassTexture;
+struct FScreenPassTextureSlice;
 
 FRDGTextureRef AddHistogramPass(
 	FRDGBuilder& GraphBuilder,
 	const FViewInfo& View,
 	const FEyeAdaptationParameters& EyeAdaptationParameters,
-	FScreenPassTexture SceneColor,
+	FScreenPassTextureSlice SceneColor,
 	const FSceneTextureParameters& SceneTextures,
 	FRDGBufferRef EyeAdaptationBuffer);
 
@@ -21,4 +21,4 @@ FRDGTextureRef AddLocalExposurePass(
 	FRDGBuilder& GraphBuilder,
 	const FViewInfo& View,
 	const FEyeAdaptationParameters& EyeAdaptationParameters,
-	FScreenPassTexture SceneColor);
\ No newline at end of file
+	FScreenPassTextureSlice SceneColor);
\ No newline at end of file
diff --git a/Engine/Source/Runtime/Renderer/Private/PostProcess/PostProcessLensFlares.cpp b/Engine/Source/Runtime/Renderer/Private/PostProcess/PostProcessLensFlares.cpp
--- a/Engine/Source/Runtime/Renderer/Private/PostProcess/PostProcessLensFlares.cpp
+++ b/Engine/Source/Runtime/Renderer/Private/PostProcess/PostProcessLensFlares.cpp
@@ -27,7 +27,7 @@ TAutoConsoleVariable<int32> CVarLensFlareQuality(
 
 // The RDG inputs shared by all lens flare passes.
 BEGIN_SHADER_PARAMETER_STRUCT(FLensFlarePassParameters, )
-	SHADER_PARAMETER_RDG_TEXTURE(Texture2D, InputTexture)
+	SHADER_PARAMETER_RDG_TEXTURE_SRV(Texture2D, InputTexture)
 	RENDER_TARGET_BINDING_SLOTS()
 END_SHADER_PARAMETER_STRUCT()
 
@@ -163,12 +163,12 @@ FScreenPassTexture AddLensFlaresPass(
 
 	// Initialize the blur output texture.
 	{
-		FRDGTextureDesc BlurOutputDesc = Inputs.Flare.Texture->Desc;
-		BlurOutputDesc.Reset();
-
-		// More precision for additive blend.
-		BlurOutputDesc.Format = PF_FloatRGBA;
-		BlurOutputDesc.ClearValue = FClearValueBinding(FLinearColor::Transparent);
+		const FRDGTextureDesc& InputDesc = Inputs.Flare.TextureSRV->Desc.Texture->Desc;
+		FRDGTextureDesc BlurOutputDesc = FRDGTextureDesc::Create2D(
+			InputDesc.Extent,
+			PF_FloatRGBA,
+			FClearValueBinding(FLinearColor::Transparent),
+			/* InFlags = */ TexCreate_ShaderResource | TexCreate_RenderTargetable);
 
 		BlurOutputTexture = GraphBuilder.CreateTexture(BlurOutputDesc, TEXT("LensFlareBlur"));
 	}
@@ -182,7 +182,7 @@ FScreenPassTexture AddLensFlaresPass(
 		const float KernelSize = (Inputs.BokehSizePercent * static_cast<float>(FlareViewSize.X)) * PercentToScale;
 
 		FLensFlarePassParameters* PassParameters = GraphBuilder.AllocParameters<FLensFlarePassParameters>();
-		PassParameters->InputTexture = Inputs.Flare.Texture;
+		PassParameters->InputTexture = Inputs.Flare.TextureSRV;
 		PassParameters->RenderTargets[0] = FRenderTargetBinding(BlurOutputTexture, ERenderTargetLoadAction::EClear);
 
 		// Setup vertex shader parameters.
@@ -245,10 +245,12 @@ FScreenPassTexture AddLensFlaresPass(
 
 	// Initialize the lens flare output texture.
 	{
-		FRDGTextureDesc LensFlareTextureDesc = Inputs.Bloom.Texture->Desc;
-		LensFlareTextureDesc.Flags |= TexCreate_DisableDCC; // DCC will cause errors on this resource when running splitscreen so disable
-		LensFlareTextureDesc.Reset();
-		LensFlareTextureDesc.ClearValue = FClearValueBinding(FLinearColor::Transparent);
+		const FRDGTextureDesc& InputDesc = Inputs.Bloom.TextureSRV->Desc.Texture->Desc;
+		FRDGTextureDesc LensFlareTextureDesc = FRDGTextureDesc::Create2D(
+			InputDesc.Extent,
+			InputDesc.Format,
+			FClearValueBinding(FLinearColor::Transparent),
+			/* InFlags = */ TexCreate_ShaderResource | TexCreate_RenderTargetable | TexCreate_DisableDCC);
 
 		LensFlareTexture = GraphBuilder.CreateTexture(LensFlareTextureDesc, TEXT("LensFlareTexture"));
 	}
@@ -258,7 +260,7 @@ FScreenPassTexture AddLensFlaresPass(
 	if (Inputs.bCompositeWithBloom)
 	{
 		FLensFlareCopyBloomPS::FParameters* CopyPassParameters = GraphBuilder.AllocParameters<FLensFlareCopyBloomPS::FParameters>();
-		CopyPassParameters->Pass.InputTexture = Inputs.Bloom.Texture;
+		CopyPassParameters->Pass.InputTexture = Inputs.Bloom.TextureSRV;
 		CopyPassParameters->Pass.RenderTargets[0] = FRenderTargetBinding(LensFlareTexture, ERenderTargetLoadAction::ENoAction);
 
 		FPixelShaderUtils::AddFullscreenPass(
@@ -294,7 +296,7 @@ FScreenPassTexture AddLensFlaresPass(
 		const FLinearColor LensFlareTint = Inputs.TintColorsPerFlare[LensFlareIndex];
 		
 		FLensFlareCompositePS::FParameters* PassParameters = GraphBuilder.AllocParameters<FLensFlareCompositePS::FParameters>();
-		PassParameters->Pass.InputTexture = BlurOutputTexture;
+		PassParameters->Pass.InputTexture = GraphBuilder.CreateSRV(FRDGTextureSRVDesc(BlurOutputTexture));
 		PassParameters->Pass.RenderTargets[0] = FRenderTargetBinding(LensFlareTexture, LensFlareLoadAction);
 		PassParameters->InputSampler = BilinearClampSampler;
 		PassParameters->FlareColor = LensFlareHDRColor * LensFlareTint;
@@ -363,8 +365,8 @@ FScreenPassTexture AddLensFlaresPass(
 	FRDGBuilder& GraphBuilder,
 	const FViewInfo& View,
 	FScreenPassTexture Bloom,
-	FScreenPassTexture QualitySceneDownsample,
-	FScreenPassTexture DefaultSceneDownsample)
+	FScreenPassTextureSlice QualitySceneDownsample,
+	FScreenPassTextureSlice DefaultSceneDownsample)
 {
 	ensure(IsLensFlaresEnabled(View));
 
@@ -395,7 +397,7 @@ FScreenPassTexture AddLensFlaresPass(
 	}
 
 	FLensFlareInputs LensFlareInputs;
-	LensFlareInputs.Bloom = Bloom;
+	LensFlareInputs.Bloom = FScreenPassTextureSlice::CreateFromScreenPassTexture(GraphBuilder, Bloom);
 	LensFlareInputs.Flare = QualitySceneDownsample;
 	LensFlareInputs.BokehShapeTexture = BokehTextureRHI;
 	LensFlareInputs.TintColorsPerFlare = Settings.LensFlareTints;
diff --git a/Engine/Source/Runtime/Renderer/Private/PostProcess/PostProcessLensFlares.h b/Engine/Source/Runtime/Renderer/Private/PostProcess/PostProcessLensFlares.h
--- a/Engine/Source/Runtime/Renderer/Private/PostProcess/PostProcessLensFlares.h
+++ b/Engine/Source/Runtime/Renderer/Private/PostProcess/PostProcessLensFlares.h
@@ -24,11 +24,11 @@ struct FLensFlareInputs
 	// [Required] The bloom convolution texture. If enabled, this will be composited with lens flares. Otherwise,
 	// a transparent black texture is used instead. Either way, the final output texture will use the this texture
 	// descriptor and viewport.
-	FScreenPassTexture Bloom;
+	FScreenPassTextureSlice Bloom;
 
 	// [Required] The scene color input, before bloom, which is used as the source of lens flares.
 	// This can be a downsampled input based on the desired quality level.
-	FScreenPassTexture Flare;
+	FScreenPassTextureSlice Flare;
 
 	// [Required] The bokeh shape texture to use to blur the lens flares.
 	FRHITexture* BokehShapeTexture;
@@ -65,5 +65,5 @@ FScreenPassTexture AddLensFlaresPass(
 	FRDGBuilder& GraphBuilder,
 	const FViewInfo& View,
 	FScreenPassTexture Bloom,
-	FScreenPassTexture QualitySceneDownsample,
-	FScreenPassTexture DefaultSceneDownsample);
\ No newline at end of file
+	FScreenPassTextureSlice QualitySceneDownsample,
+	FScreenPassTextureSlice DefaultSceneDownsample);
\ No newline at end of file
diff --git a/Engine/Source/Runtime/Renderer/Private/PostProcess/PostProcessLocalExposure.cpp b/Engine/Source/Runtime/Renderer/Private/PostProcess/PostProcessLocalExposure.cpp
--- a/Engine/Source/Runtime/Renderer/Private/PostProcess/PostProcessLocalExposure.cpp
+++ b/Engine/Source/Runtime/Renderer/Private/PostProcess/PostProcessLocalExposure.cpp
@@ -29,7 +29,7 @@ public:
 		SHADER_PARAMETER_STRUCT_REF(FViewUniformShaderParameters, View)
 		SHADER_PARAMETER_STRUCT(FEyeAdaptationParameters, EyeAdaptation)
 		SHADER_PARAMETER_STRUCT(FScreenPassTextureViewportParameters, Input)
-		SHADER_PARAMETER_RDG_TEXTURE(Texture2D, InputTexture)
+		SHADER_PARAMETER_RDG_TEXTURE_SRV(Texture2D, InputTexture)
 		SHADER_PARAMETER_RDG_TEXTURE_UAV(RWTexture2D<float>, OutputFloat)
 	END_SHADER_PARAMETER_STRUCT()
 
@@ -148,7 +148,7 @@ FRDGTextureRef AddLocalExposureBlurredLogLuminancePass(
 	FRDGBuilder& GraphBuilder,
 	const FViewInfo& View,
 	const FEyeAdaptationParameters& EyeAdaptationParameters,
-	FScreenPassTexture InputTexture)
+	FScreenPassTextureSlice InputTexture)
 {
 	check(InputTexture.IsValid());
 
@@ -170,7 +170,7 @@ FRDGTextureRef AddLocalExposureBlurredLogLuminancePass(
 		PassParameters->View = View.ViewUniformBuffer;
 		PassParameters->EyeAdaptation = EyeAdaptationParameters;
 		PassParameters->Input = GetScreenPassTextureViewportParameters(FScreenPassTextureViewport(InputTexture));
-		PassParameters->InputTexture = InputTexture.Texture;
+		PassParameters->InputTexture = InputTexture.TextureSRV;
 		PassParameters->OutputFloat = GraphBuilder.CreateUAV(GaussianLumSetupTexture);
 
 		FComputeShaderUtils::AddPass(
@@ -188,7 +188,7 @@ FRDGTextureRef AddLocalExposureBlurredLogLuminancePass(
 		FGaussianBlurInputs GaussianBlurInputs;
 		GaussianBlurInputs.NameX = TEXT("LocalExposureGaussianX");
 		GaussianBlurInputs.NameY = TEXT("LocalExposureGaussianY");
-		GaussianBlurInputs.Filter = FScreenPassTexture(GaussianLumSetupTexture);
+		GaussianBlurInputs.Filter = FScreenPassTextureSlice::CreateFromScreenPassTexture(GraphBuilder, FScreenPassTexture(GaussianLumSetupTexture));
 		GaussianBlurInputs.TintColor = FLinearColor::White;
 		GaussianBlurInputs.CrossCenterWeight = FVector2f::ZeroVector;
 		GaussianBlurInputs.KernelSizePercent = View.FinalPostProcessSettings.LocalExposureBlurredLuminanceKernelSizePercent;
diff --git a/Engine/Source/Runtime/Renderer/Private/PostProcess/PostProcessLocalExposure.h b/Engine/Source/Runtime/Renderer/Private/PostProcess/PostProcessLocalExposure.h
--- a/Engine/Source/Runtime/Renderer/Private/PostProcess/PostProcessLocalExposure.h
+++ b/Engine/Source/Runtime/Renderer/Private/PostProcess/PostProcessLocalExposure.h
@@ -23,7 +23,7 @@ FRDGTextureRef AddLocalExposureBlurredLogLuminancePass(
 	FRDGBuilder& GraphBuilder,
 	const FViewInfo& View,
 	const FEyeAdaptationParameters& EyeAdaptationParameters,
-	FScreenPassTexture InputTexture);
+	FScreenPassTextureSlice InputTexture);
 
 void AddApplyLocalExposurePass(
 	FRDGBuilder& GraphBuilder,
diff --git a/Engine/Source/Runtime/Renderer/Private/PostProcess/PostProcessVisualizeBuffer.cpp b/Engine/Source/Runtime/Renderer/Private/PostProcess/PostProcessVisualizeBuffer.cpp
--- a/Engine/Source/Runtime/Renderer/Private/PostProcess/PostProcessVisualizeBuffer.cpp
+++ b/Engine/Source/Runtime/Renderer/Private/PostProcess/PostProcessVisualizeBuffer.cpp
@@ -406,14 +406,14 @@ FScreenPassTexture AddVisualizeGBufferOverviewPass(
 		{
 			FDownsamplePassInputs DownsampleInputs;
 			DownsampleInputs.Name = TEXT("MaterialHalfSize");
-			DownsampleInputs.SceneColor = Output;
+			DownsampleInputs.SceneColor = FScreenPassTextureSlice::CreateFromScreenPassTexture(GraphBuilder, Output);
 			DownsampleInputs.Flags = EDownsampleFlags::ForceRaster;
 			DownsampleInputs.Quality = EDownsampleQuality::Low;
 
 			FScreenPassTexture HalfSize = AddDownsamplePass(GraphBuilder, View, DownsampleInputs);
 
 			DownsampleInputs.Name = TEXT("MaterialQuarterSize");
-			DownsampleInputs.SceneColor = HalfSize;
+			DownsampleInputs.SceneColor = FScreenPassTextureSlice::CreateFromScreenPassTexture(GraphBuilder, HalfSize);
 
 			FVisualizeBufferTile Tile;
 			Tile.Input = AddDownsamplePass(GraphBuilder, View, DownsampleInputs);
diff --git a/Engine/Source/Runtime/Renderer/Private/PostProcess/PostProcessWeightedSampleSum.cpp b/Engine/Source/Runtime/Renderer/Private/PostProcess/PostProcessWeightedSampleSum.cpp
--- a/Engine/Source/Runtime/Renderer/Private/PostProcess/PostProcessWeightedSampleSum.cpp
+++ b/Engine/Source/Runtime/Renderer/Private/PostProcess/PostProcessWeightedSampleSum.cpp
@@ -53,7 +53,7 @@ TAutoConsoleVariable<float> CVarFastBlurThreshold(
 	ECVF_Scalability | ECVF_RenderThreadSafe);
 
 BEGIN_SHADER_PARAMETER_STRUCT(FFilterParameters, )
-	SHADER_PARAMETER_STRUCT(FScreenPassTextureInput, Filter)
+	SHADER_PARAMETER_STRUCT(FScreenPassTextureSliceInput, Filter)
 	SHADER_PARAMETER_STRUCT(FScreenPassTextureInput, Additive)
 	SHADER_PARAMETER_ARRAY(FVector4f, SampleOffsets, [MAX_PACKED_SAMPLES_OFFSET])
 	SHADER_PARAMETER_ARRAY(FLinearColor, SampleWeights, [MAX_FILTER_SAMPLES])
@@ -62,7 +62,7 @@ END_SHADER_PARAMETER_STRUCT()
 
 void GetFilterParameters(
 	FFilterParameters& OutParameters,
-	const FScreenPassTextureInput& Filter,
+	const FScreenPassTextureSliceInput& Filter,
 	const FScreenPassTextureInput& Additive,
 	TArrayView<const FVector2f> SampleOffsets,
 	TArrayView<const FLinearColor> SampleWeights)
@@ -342,7 +342,7 @@ FScreenPassTexture AddGaussianBlurPass(
 	const FViewInfo& View,
 	const TCHAR* Name,
 	FScreenPassTextureViewport OutputViewport,
-	FScreenPassTexture Filter,
+	FScreenPassTextureSlice Filter,
 	FScreenPassTexture Additive,
 	TArrayView<const FVector2f> SampleOffsets,
 	TArrayView<const FLinearColor> SampleWeights,
@@ -357,7 +357,7 @@ FScreenPassTexture AddGaussianBlurPass(
 
 	const bool bCombineAdditive = Additive.IsValid();
 
-	const bool bManualUVBorder = Filter.ViewRect.Min != FIntPoint::ZeroValue || Filter.ViewRect.Max != Filter.Texture->Desc.Extent;
+	const bool bManualUVBorder = Filter.ViewRect.Min != FIntPoint::ZeroValue || Filter.ViewRect.Max != Filter.TextureSRV->Desc.Texture->Desc.Extent;
 
 	const uint32 SampleCount = SampleOffsets.Num();
 
@@ -374,7 +374,7 @@ FScreenPassTexture AddGaussianBlurPass(
 		SamplerState = TStaticSamplerState<SF_Bilinear, AM_Border, AM_Border, AM_Clamp>::GetRHI();
 	}
 
-	const FScreenPassTextureInput FilterInput = GetScreenPassTextureInput(Filter, SamplerState);
+	const FScreenPassTextureSliceInput FilterInput = GetScreenPassTextureInput(Filter, SamplerState);
 
 	FScreenPassTextureInput AdditiveInput;
 
@@ -383,11 +383,12 @@ FScreenPassTexture AddGaussianBlurPass(
 		AdditiveInput = GetScreenPassTextureInput(Additive, SamplerState);
 	}
 
-	FRDGTextureDesc OutputDesc = Filter.Texture->Desc;
-	OutputDesc.Reset();
-	OutputDesc.Extent = OutputViewport.Extent;
-	OutputDesc.Flags |= bIsComputePass ? TexCreate_UAV : TexCreate_NoFastClear;
-	OutputDesc.ClearValue = FClearValueBinding(FLinearColor::Transparent);
+	const FRDGTextureDesc& InputDesc = Filter.TextureSRV->Desc.Texture->Desc;
+	FRDGTextureDesc OutputDesc = FRDGTextureDesc::Create2D(
+		OutputViewport.Extent,
+		InputDesc.Format,
+		FClearValueBinding(FLinearColor::Transparent),
+		/* InFlags = */ TexCreate_ShaderResource | (bIsComputePass ? TexCreate_UAV : (TexCreate_RenderTargetable | TexCreate_NoFastClear)));
 
 	const FScreenPassRenderTarget Output(GraphBuilder.CreateTexture(OutputDesc, Name), OutputViewport.Rect, ERenderTargetLoadAction::ENoAction);
 
@@ -555,7 +556,7 @@ FScreenPassTexture AddGaussianBlurPass(
 			View,
 			Inputs.NameY,
 			FilterViewport,
-			HorizontalOutput,
+			FScreenPassTextureSlice::CreateFromScreenPassTexture(GraphBuilder, HorizontalOutput),
 			Inputs.Additive,
 			TArrayView<const FVector2f>(SampleOffsets, SampleCount),
 			TArrayView<const FLinearColor>(SampleWeights, SampleCount),
diff --git a/Engine/Source/Runtime/Renderer/Private/PostProcess/PostProcessWeightedSampleSum.h b/Engine/Source/Runtime/Renderer/Private/PostProcess/PostProcessWeightedSampleSum.h
--- a/Engine/Source/Runtime/Renderer/Private/PostProcess/PostProcessWeightedSampleSum.h
+++ b/Engine/Source/Runtime/Renderer/Private/PostProcess/PostProcessWeightedSampleSum.h
@@ -11,7 +11,7 @@ struct FGaussianBlurInputs
 	const TCHAR* NameY = nullptr;
 
 	// The input texture to be filtered.
-	FScreenPassTexture Filter;
+	FScreenPassTextureSlice Filter;
 
 	// The input texture to be added after filtering.
 	FScreenPassTexture Additive;
diff --git a/Engine/Source/Runtime/Renderer/Private/PostProcess/PostProcessing.cpp b/Engine/Source/Runtime/Renderer/Private/PostProcess/PostProcessing.cpp
--- a/Engine/Source/Runtime/Renderer/Private/PostProcess/PostProcessing.cpp
+++ b/Engine/Source/Runtime/Renderer/Private/PostProcess/PostProcessing.cpp
@@ -775,8 +775,8 @@ void AddPostProcessingPasses(
 		FScreenPassTextureSlice SceneColorSlice = FScreenPassTextureSlice::CreateFromScreenPassTexture(GraphBuilder, SceneColor);
 		SceneColor = FScreenPassTexture();
 
-		FScreenPassTexture HalfResSceneColor;
-		FScreenPassTexture QuarterResSceneColor;
+		FScreenPassTextureSlice HalfResSceneColor;
+		FScreenPassTextureSlice QuarterResSceneColor;
 		FVelocityFlattenTextures VelocityFlattenTextures;
 		if (TAAConfig != EMainTAAPassConfig::Disabled)
 		{
@@ -932,8 +932,8 @@ void AddPostProcessingPasses(
 			{
 				FMotionBlurOutputs PassOutputs = AddMotionBlurPass(GraphBuilder, View, PassInputs);
 				SceneColorSlice = PassOutputs.FullRes;
-				HalfResSceneColor = PassOutputs.HalfRes;
-				QuarterResSceneColor = PassOutputs.QuarterRes;
+				HalfResSceneColor = FScreenPassTextureSlice::CreateFromScreenPassTexture(GraphBuilder, PassOutputs.HalfRes);
+				QuarterResSceneColor = FScreenPassTextureSlice::CreateFromScreenPassTexture(GraphBuilder, PassOutputs.QuarterRes);
 			}
 		}
 		else if (PostMotionBlurTranslucencyResources.IsValid())
@@ -955,8 +955,8 @@ void AddPostProcessingPasses(
 			// Invalidate half and quarter res.
 			if (NewSceneColorSlice != SceneColorSlice)
 			{
-				HalfResSceneColor = FScreenPassTexture();
-				QuarterResSceneColor = FScreenPassTexture();
+				HalfResSceneColor = FScreenPassTextureSlice();
+				QuarterResSceneColor = FScreenPassTextureSlice();
 			}
 
 			SceneColorSlice = NewSceneColorSlice;
@@ -969,12 +969,11 @@ void AddPostProcessingPasses(
 			{
 				FDownsamplePassInputs PassInputs;
 				PassInputs.Name = TEXT("PostProcessing.SceneColor.HalfRes");
-				PassInputs.SceneColor = FScreenPassTexture(SceneColorSlice);
+				PassInputs.SceneColor = SceneColorSlice;
 				PassInputs.Quality = DownsampleQuality;
 				PassInputs.FormatOverride = DownsampleOverrideFormat;
-				PassInputs.UserSuppliedOutput = View.PrevViewInfo.HalfResTemporalAAHistory;
 
-				HalfResSceneColor = AddDownsamplePass(GraphBuilder, View, PassInputs);
+				HalfResSceneColor = FScreenPassTextureSlice::CreateFromScreenPassTexture(GraphBuilder, AddDownsamplePass(GraphBuilder, View, PassInputs));
 			}
 
 			if (bNeedPostMotionBlurQuarterRes && !QuarterResSceneColor.IsValid())
@@ -984,21 +983,21 @@ void AddPostProcessingPasses(
 				PassInputs.SceneColor = HalfResSceneColor;
 				PassInputs.Quality = DownsampleQuality;
 
-				QuarterResSceneColor = AddDownsamplePass(GraphBuilder, View, PassInputs);
+				QuarterResSceneColor = FScreenPassTextureSlice::CreateFromScreenPassTexture(GraphBuilder, AddDownsamplePass(GraphBuilder, View, PassInputs));
 			}
 		}
 
 
 
 		// Store half res scene color in the history
-		if (ReflectionsMethod == EReflectionsMethod::SSR && !View.bStatePrevViewInfoIsReadOnly && GSSRHalfResSceneColor)
+		if (ReflectionsMethod == EReflectionsMethod::SSR && !View.bStatePrevViewInfoIsReadOnly && GSSRHalfResSceneColor && HalfResSceneColor.IsValid())
 		{
 			check(View.ViewState);
-			GraphBuilder.QueueTextureExtraction(HalfResSceneColor.Texture, &View.ViewState->PrevFrameViewInfo.HalfResTemporalAAHistory);
+			GraphBuilder.QueueTextureExtraction(HalfResSceneColor.TextureSRV->Desc.Texture, &View.ViewState->PrevFrameViewInfo.HalfResTemporalAAHistory);
 		}
 
 		{
-			FScreenPassTexture LocalExposureSceneColor = bProcessQuarterResolution ? QuarterResSceneColor : HalfResSceneColor;
+			FScreenPassTextureSlice LocalExposureSceneColor = bProcessQuarterResolution ? QuarterResSceneColor : HalfResSceneColor;
 
 			if (bLocalExposureEnabled)
 			{
@@ -1013,7 +1012,7 @@ void AddPostProcessingPasses(
 
 		if (bHistogramEnabled)
 		{
-			FScreenPassTexture HistogramSceneColor = bProcessQuarterResolution ? QuarterResSceneColor : HalfResSceneColor;
+			FScreenPassTextureSlice HistogramSceneColor = bProcessQuarterResolution ? QuarterResSceneColor : HalfResSceneColor;
 
 			if (IsAutoExposureUsingIlluminanceEnabled(View))
 			{
@@ -1021,11 +1020,11 @@ void AddPostProcessingPasses(
 				{
 					const FIntRect IlluminanceRect = GetDownscaledRect(PrimaryViewRect, GetAutoExposureIlluminanceDownscaleFactor());
 
-					HistogramSceneColor = FScreenPassTexture(Inputs.ExposureIlluminance, IlluminanceRect);
+					HistogramSceneColor = FScreenPassTextureSlice::CreateFromScreenPassTexture(GraphBuilder, FScreenPassTexture(Inputs.ExposureIlluminance, IlluminanceRect));
 				}
 				else
 				{
-					HistogramSceneColor = OriginalSceneColor;
+					HistogramSceneColor = FScreenPassTextureSlice::CreateFromScreenPassTexture(GraphBuilder, OriginalSceneColor);
 				}
 			}
 
@@ -1096,12 +1095,12 @@ void AddPostProcessingPasses(
 
 				if (FFTBloomResolutionFraction <= 0.25f)
 				{
-					InputSceneColor = FScreenPassTextureSlice::CreateFromScreenPassTexture(GraphBuilder, QuarterResSceneColor);
+					InputSceneColor = QuarterResSceneColor;
 					InputResolutionFraction = 0.25f;
 				}
 				else if (FFTBloomResolutionFraction <= 0.5f)
 				{
-					InputSceneColor = FScreenPassTextureSlice::CreateFromScreenPassTexture(GraphBuilder, HalfResSceneColor);
+					InputSceneColor = HalfResSceneColor;
 					InputResolutionFraction = 0.5f;
 				}
 				else
@@ -1135,7 +1134,7 @@ void AddPostProcessingPasses(
 				}
 				else
 				{
-					FScreenPassTexture DownsampleInput = bProcessQuarterResolution ? QuarterResSceneColor : HalfResSceneColor;
+					FScreenPassTextureSlice DownsampleInput = bProcessQuarterResolution ? QuarterResSceneColor : HalfResSceneColor;
 
 					if (bBloomSetupRequiredEnabled)
 					{
@@ -1150,7 +1149,7 @@ void AddPostProcessingPasses(
 						SetupPassInputs.BlurredLogLuminanceTexture = LocalExposureBlurredLogLumTexture;
 						SetupPassInputs.Threshold = BloomThreshold;
 
-						DownsampleInput = AddBloomSetupPass(GraphBuilder, View, SetupPassInputs);
+						DownsampleInput = FScreenPassTextureSlice::CreateFromScreenPassTexture(GraphBuilder, AddBloomSetupPass(GraphBuilder, View, SetupPassInputs));
 					}
 
 					const bool bLogLumaInAlpha = false;
@@ -1166,7 +1165,9 @@ void AddPostProcessingPasses(
 			{
 				const ELensFlareQuality LensFlareQuality = GetLensFlareQuality();
 				const uint32 LensFlareDownsampleStageIndex = static_cast<uint32>(ELensFlareQuality::MAX) - static_cast<uint32>(LensFlareQuality) - 1;
-				Bloom = AddLensFlaresPass(GraphBuilder, View, Bloom, LensFlareSceneDownsampleChain->GetTexture(LensFlareDownsampleStageIndex), LensFlareSceneDownsampleChain->GetFirstTexture());
+				Bloom = AddLensFlaresPass(GraphBuilder, View, Bloom,
+					LensFlareSceneDownsampleChain->GetTexture(LensFlareDownsampleStageIndex),
+					LensFlareSceneDownsampleChain->GetFirstTexture());
 			}
 		}
 
@@ -2376,7 +2377,7 @@ void AddMobilePostProcessingPasses(FRDGBuilder& GraphBuilder, FScene* Scene, con
 							FGaussianBlurInputs GaussianBlurInputs;
 							GaussianBlurInputs.NameX = BlurDebugX;
 							GaussianBlurInputs.NameY = BlurDebugY;
-							GaussianBlurInputs.Filter = DOFSetup;
+							GaussianBlurInputs.Filter = FScreenPassTextureSlice::CreateFromScreenPassTexture(GraphBuilder, DOFSetup);
 							GaussianBlurInputs.TintColor = FLinearColor::White;
 							GaussianBlurInputs.CrossCenterWeight = FVector2f::ZeroVector;
 							GaussianBlurInputs.KernelSizePercent = KernelSizePercent;
@@ -2489,7 +2490,9 @@ void AddMobilePostProcessingPasses(FRDGBuilder& GraphBuilder, FScene* Scene, con
 			{
 				const ELensFlareQuality LensFlareQuality = GetLensFlareQuality();
 				const uint32 LensFlareDownsampleStageIndex = static_cast<uint32>(ELensFlareQuality::MAX) - static_cast<uint32>(LensFlareQuality) - 1;
-				BloomUpOutputs = AddLensFlaresPass(GraphBuilder, View, BloomUpOutputs, PostProcessDownsample_Bloom[LensFlareDownsampleStageIndex], PostProcessDownsample_Bloom[0]);
+				BloomUpOutputs = AddLensFlaresPass(GraphBuilder, View, BloomUpOutputs,
+					FScreenPassTextureSlice::CreateFromScreenPassTexture(GraphBuilder, PostProcessDownsample_Bloom[LensFlareDownsampleStageIndex]),
+					FScreenPassTextureSlice::CreateFromScreenPassTexture(GraphBuilder, PostProcessDownsample_Bloom[0]));
 			}
 		}
 
diff --git a/Engine/Source/Runtime/Renderer/Private/PostProcess/TemporalAA.cpp b/Engine/Source/Runtime/Renderer/Private/PostProcess/TemporalAA.cpp
--- a/Engine/Source/Runtime/Renderer/Private/PostProcess/TemporalAA.cpp
+++ b/Engine/Source/Runtime/Renderer/Private/PostProcess/TemporalAA.cpp
@@ -654,16 +654,8 @@ FTAAOutputs AddTemporalAAPass(
 				Inputs.DownsampleOverrideFormat != PF_Unknown ? Inputs.DownsampleOverrideFormat : Inputs.SceneColorInput->Desc.Format,
 				FClearValueBinding::Black,
 				TexCreate_ShaderResource | TexCreate_UAV | GFastVRamConfig.Downsample);
-			const TRefCountPtr<IPooledRenderTarget>& PrevFrameHalfResTAAHistory = View.PrevViewInfo.HalfResTemporalAAHistory;
 
-			if (PrevFrameHalfResTAAHistory && Translate(PrevFrameHalfResTAAHistory->GetDesc()) == HalfResSceneColorDesc)
-			{
-				Outputs.DownsampledSceneColor = GraphBuilder.RegisterExternalTexture(PrevFrameHalfResTAAHistory);
-			}
-			else
-			{
-				Outputs.DownsampledSceneColor = GraphBuilder.CreateTexture(HalfResSceneColorDesc, TEXT("SceneColorHalfRes"));
-			}
+			Outputs.DownsampledSceneColor = GraphBuilder.CreateTexture(HalfResSceneColorDesc, TEXT("SceneColorHalfRes"));
 		}
 	}
 
@@ -1068,8 +1060,11 @@ FDefaultTemporalUpscaler::FOutputs AddGen4MainTemporalAAPasses(
 
 	FDefaultTemporalUpscaler::FOutputs Outputs;
 	Outputs.FullRes = FScreenPassTextureSlice::CreateFromScreenPassTexture(GraphBuilder, FScreenPassTexture(SceneColorTexture, SecondaryViewRect));
-	Outputs.HalfRes.Texture = TAAOutputs.DownsampledSceneColor;
-	Outputs.HalfRes.ViewRect = FIntRect::DivideAndRoundUp(SecondaryViewRect, 2);
+	if (TAAOutputs.DownsampledSceneColor)
+	{
+		Outputs.HalfRes.TextureSRV = GraphBuilder.CreateSRV(FRDGTextureSRVDesc(TAAOutputs.DownsampledSceneColor));
+		Outputs.HalfRes.ViewRect = FIntRect::DivideAndRoundUp(SecondaryViewRect, 2);
+	}
 	return Outputs;
 } // AddGen4MainTemporalAAPasses()
 
diff --git a/Engine/Source/Runtime/Renderer/Private/PostProcess/TemporalAA.h b/Engine/Source/Runtime/Renderer/Private/PostProcess/TemporalAA.h
--- a/Engine/Source/Runtime/Renderer/Private/PostProcess/TemporalAA.h
+++ b/Engine/Source/Runtime/Renderer/Private/PostProcess/TemporalAA.h
@@ -206,8 +206,8 @@ struct FDefaultTemporalUpscaler
 	struct FOutputs
 	{
 		FScreenPassTextureSlice FullRes;
-		FScreenPassTexture HalfRes;
-		FScreenPassTexture QuarterRes;
+		FScreenPassTextureSlice HalfRes;
+		FScreenPassTextureSlice QuarterRes;
 		FVelocityFlattenTextures VelocityFlattenTextures;
 	};
 };
diff --git a/Engine/Source/Runtime/Renderer/Private/PostProcess/TemporalSuperResolution.cpp b/Engine/Source/Runtime/Renderer/Private/PostProcess/TemporalSuperResolution.cpp
--- a/Engine/Source/Runtime/Renderer/Private/PostProcess/TemporalSuperResolution.cpp
+++ b/Engine/Source/Runtime/Renderer/Private/PostProcess/TemporalSuperResolution.cpp
@@ -354,11 +354,6 @@ enum class ETSRHistoryFormatBits : uint32
 };
 ENUM_CLASS_FLAGS(ETSRHistoryFormatBits);
 
-bool IsOutputDifferentThanHighFrequency(ETSRHistoryFormatBits HistoryFormatBits)
-{
-	return false;
-}
-
 FTSRHistoryArrayIndices TranslateHistoryFormatBitsToArrayIndices(ETSRHistoryFormatBits HistoryFormatBits)
 {
 	FTSRHistoryArrayIndices ArrayIndices;
@@ -834,10 +829,9 @@ class FTSRUpdateHistoryCS : public FTSRShader
 		SHADER_PARAMETER_RDG_TEXTURE_SRV(Texture2DArray, PrevHistoryColorTexture)
 		SHADER_PARAMETER_RDG_TEXTURE_SRV(Texture2DArray, PrevHistoryMetadataTexture)
 
-		SHADER_PARAMETER_RDG_TEXTURE_UAV(RWTexture2D, SceneColorOutputMip0)
-		SHADER_PARAMETER_RDG_TEXTURE_UAV(RWTexture2D, SceneColorOutputMip1)
 		SHADER_PARAMETER_RDG_TEXTURE_UAV(RWTexture2DArray, HistoryColorOutput)
 		SHADER_PARAMETER_RDG_TEXTURE_UAV(RWTexture2DArray, HistoryMetadataOutput)
+		SHADER_PARAMETER_RDG_TEXTURE_UAV(RWTexture2DArray, SceneColorOutputMip1)
 		SHADER_PARAMETER_RDG_TEXTURE_UAV(RWTexture2DArray, DebugOutput)
 	END_SHADER_PARAMETER_STRUCT()
 
@@ -1288,8 +1282,6 @@ FDefaultTemporalUpscaler::FOutputs AddTemporalSuperResolutionPasses(
 	}
 	FTSRHistoryArrayIndices HistoryArrayIndices = TranslateHistoryFormatBitsToArrayIndices(HistoryFormatBits);
 
-	const bool bIsOutputDifferentThanHighFrequency = IsOutputDifferentThanHighFrequency(HistoryFormatBits);
-
 	FTSRUpdateHistoryCS::EQuality UpdateHistoryQuality = FTSRUpdateHistoryCS::EQuality(FMath::Clamp(CVarTSRHistoryUpdateQuality.GetValueOnRenderThread(), 0, int32(FTSRUpdateHistoryCS::EQuality::MAX) - 1));
 
 	bool bIsSeperateTranslucyTexturesValid = PassInputs.PostDOFTranslucencyResources.IsValid();
@@ -1471,59 +1463,12 @@ FDefaultTemporalUpscaler::FOutputs AddTemporalSuperResolutionPasses(
 		return WaveSizeOps;
 	};
 
-	// Output
-	FRDGTextureRef UpdateHistoryOutputTexture = nullptr;
-	FRDGTextureRef SceneColorOutputTexture = nullptr;
-	FRDGTextureRef SceneColorOutputHalfResTexture = nullptr;
-	FRDGTextureRef SceneColorOutputQuarterResTexture = nullptr;
-	{
-		check(!(PassInputs.bGenerateOutputMip1 && (PassInputs.bGenerateSceneColorHalfRes || PassInputs.bGenerateSceneColorQuarterRes)));
-		FRDGTextureDesc OutputDesc = FRDGTextureDesc::Create2D(
-			OutputExtent,
-			ColorFormat,
-			FClearValueBinding::None,
-			/* InFlags = */ TexCreate_ShaderResource | TexCreate_UAV,
-			/* NumMips = */ PassInputs.bGenerateOutputMip1 ? 2 : 1);
-
-		if (OutputRect.Size() != HistorySize)
-		{
-			SceneColorOutputTexture = GraphBuilder.CreateTexture(OutputDesc, TEXT("TSR.Output"));
-
-			if (bIsOutputDifferentThanHighFrequency)
-			{
-				FRDGTextureDesc HistoryOutputDesc = OutputDesc;
-				HistoryOutputDesc.Extent = HistoryExtent;
-				HistoryOutputDesc.NumMips = 1;
-				UpdateHistoryOutputTexture = GraphBuilder.CreateTexture(OutputDesc, TEXT("TSR.HistoryUpdateOutput"));
-			}
-		}
-		else
-		{
-			UpdateHistoryOutputTexture = GraphBuilder.CreateTexture(OutputDesc, TEXT("TSR.Output"));
-		}
-
-		// Generate quarter res output if only one needed, otherwise only output half res and let caller downscale to quarter res. 
-		// This is motivated to saves UAV slots on FTSRUpdateHistoryCS
-		if (PassInputs.bGenerateSceneColorQuarterRes && !PassInputs.bGenerateSceneColorHalfRes && OutputRect.Size() == HistorySize)
-		{
-			FRDGTextureDesc QuarterResDesc = OutputDesc;
-			QuarterResDesc.Extent = OutputExtent / 4;
-			SceneColorOutputQuarterResTexture = GraphBuilder.CreateTexture(QuarterResDesc, TEXT("TSR.QuarterResOutput"));
-		}
-		else if (PassInputs.bGenerateSceneColorHalfRes || PassInputs.bGenerateSceneColorQuarterRes)
-		{
-			FRDGTextureDesc HalfResDesc = OutputDesc;
-			HalfResDesc.Extent = OutputExtent / 2;
-			SceneColorOutputHalfResTexture = GraphBuilder.CreateTexture(HalfResDesc, TEXT("TSR.HalfResOutput"));
-		}
-	}
-
 	// Allocate a new history
 	FTSRHistoryTextures History;
 	const int32 HistoryColorGuideSliceCountWithoutResurrection = bSupportsAlpha ? 2 : 1;
 	{
 		{
-			bool bRequires2Mips = HistorySize == OutputRect.Size() && !bIsOutputDifferentThanHighFrequency && PassInputs.bGenerateOutputMip1;
+			bool bRequires2Mips = HistorySize == OutputRect.Size() && PassInputs.bGenerateOutputMip1;
 
 			FRDGTextureDesc ArrayDesc = FRDGTextureDesc::Create2DArray(
 				HistoryExtent,
@@ -2208,6 +2153,8 @@ FDefaultTemporalUpscaler::FOutputs AddTemporalSuperResolutionPasses(
 
 	// Update temporal history.
 	FRDGTextureSRVRef UpdateHistoryTextureSRV = nullptr;
+	FRDGTextureSRVRef SceneColorOutputHalfResTextureSRV = nullptr;
+	FRDGTextureSRVRef SceneColorOutputQuarterResTextureSRV = nullptr;
 	{
 		static const TCHAR* const kUpdateQualityNames[] = {
 			TEXT("Low"),
@@ -2301,50 +2248,53 @@ FDefaultTemporalUpscaler::FOutputs AddTemporalSuperResolutionPasses(
 			PassParameters->HistoryColorOutput = GraphBuilder.CreateUAV(ColorUAVDesc);
 			PassParameters->HistoryMetadataOutput = GraphBuilder.CreateUAV(MetadataUAVDesc);
 
-			if (bIsOutputDifferentThanHighFrequency)
-			{
-				PassParameters->SceneColorOutputMip0 = GraphBuilder.CreateUAV(FRDGTextureUAVDesc(UpdateHistoryOutputTexture, /* InMipLevel = */ 0));
-				UpdateHistoryTextureSRV = GraphBuilder.CreateSRV(FRDGTextureSRVDesc(UpdateHistoryOutputTexture));
-			}
-			else
-			{
-				PassParameters->SceneColorOutputMip0 = CreateDummyUAV(GraphBuilder, PF_FloatR11G11B10);
-				UpdateHistoryTextureSRV = GraphBuilder.CreateSRV(FRDGTextureSRVDesc::CreateForSlice(
-					History.ColorArray, ColorUAVDesc.FirstArraySlice + HistoryArrayIndices.HighFrequency));
-			}
+			UpdateHistoryTextureSRV = GraphBuilder.CreateSRV(FRDGTextureSRVDesc::CreateForSlice(
+				History.ColorArray, ColorUAVDesc.FirstArraySlice + HistoryArrayIndices.HighFrequency));
 		}
 
 		if (PassInputs.bGenerateOutputMip1 && HistorySize == OutputRect.Size())
 		{
-			PassParameters->bGenerateOutputMip1 = true;
-			if (bIsOutputDifferentThanHighFrequency)
-			{
-				PassParameters->SceneColorOutputMip1 = GraphBuilder.CreateUAV(FRDGTextureUAVDesc(UpdateHistoryOutputTexture, /* InMipLevel = */ 1));
-			}
-			else
-			{
-				FRDGTextureUAVDesc Mip1Desc(History.ColorArray);
-				Mip1Desc.MipLevel = 1;
-				Mip1Desc.FirstArraySlice = UpdateHistoryTextureSRV->Desc.FirstArraySlice;
-				Mip1Desc.NumArraySlices = 1;
-				Mip1Desc.DimensionOverride = ETextureDimension::Texture2D;
+			FRDGTextureUAVDesc Mip1Desc(History.ColorArray);
+			Mip1Desc.MipLevel = 1;
+			Mip1Desc.FirstArraySlice = UpdateHistoryTextureSRV->Desc.FirstArraySlice;
+			Mip1Desc.NumArraySlices = 1;
 
-				PassParameters->SceneColorOutputMip1 = GraphBuilder.CreateUAV(Mip1Desc);
-			}
+			PassParameters->bGenerateOutputMip1 = true;
+			PassParameters->SceneColorOutputMip1 = GraphBuilder.CreateUAV(Mip1Desc);
 		}
-		else if (SceneColorOutputHalfResTexture && HistorySize == OutputRect.Size())
+		else if (PassInputs.bGenerateSceneColorHalfRes && HistorySize == OutputRect.Size())
 		{
+			FRDGTextureDesc HalfResDesc = FRDGTextureDesc::Create2DArray(
+				OutputExtent / 2,
+				ColorFormat,
+				FClearValueBinding::None,
+				/* InFlags = */ TexCreate_ShaderResource | TexCreate_UAV,
+				/* ArraySize = */ 1);
+			FRDGTextureRef SceneColorOutputHalfResTexture = GraphBuilder.CreateTexture(HalfResDesc, TEXT("TSR.HalfResOutput"));
+
 			PassParameters->bGenerateOutputMip1 = true;
 			PassParameters->SceneColorOutputMip1 = GraphBuilder.CreateUAV(FRDGTextureUAVDesc(SceneColorOutputHalfResTexture));
+
+			SceneColorOutputHalfResTextureSRV = GraphBuilder.CreateSRV(FRDGTextureSRVDesc::CreateForSlice(SceneColorOutputHalfResTexture, /* SliceIndex = */ 0));
 		}
-		else if (SceneColorOutputQuarterResTexture && HistorySize == OutputRect.Size())
+		else if (PassInputs.bGenerateSceneColorQuarterRes && HistorySize == OutputRect.Size())
 		{
+			FRDGTextureDesc QuarterResDesc = FRDGTextureDesc::Create2DArray(
+				OutputExtent / 4,
+				ColorFormat,
+				FClearValueBinding::None,
+				/* InFlags = */ TexCreate_ShaderResource | TexCreate_UAV,
+				/* ArraySize = */ 1);
+			FRDGTextureRef SceneColorOutputQuarterResTexture = GraphBuilder.CreateTexture(QuarterResDesc, TEXT("TSR.QuarterResOutput"));
+
 			PassParameters->bGenerateOutputMip2 = true;
 			PassParameters->SceneColorOutputMip1 = GraphBuilder.CreateUAV(FRDGTextureUAVDesc(SceneColorOutputQuarterResTexture));
+
+			SceneColorOutputQuarterResTextureSRV = GraphBuilder.CreateSRV(FRDGTextureSRVDesc::CreateForSlice(SceneColorOutputQuarterResTexture, /* SliceIndex = */ 0));
 		}
 		else
 		{
-			PassParameters->SceneColorOutputMip1 = CreateDummyUAV(GraphBuilder, PF_FloatR11G11B10);
+			PassParameters->SceneColorOutputMip1 = CreateDummyUAVArray(GraphBuilder, PF_FloatR11G11B10);
 		}
 		PassParameters->DebugOutput = CreateDebugUAV(HistoryExtent, TEXT("Debug.TSR.UpdateHistory"));
 
@@ -2374,8 +2324,9 @@ FDefaultTemporalUpscaler::FOutputs AddTemporalSuperResolutionPasses(
 	FRDGTextureSRVRef SceneColorOutputTextureSRV = UpdateHistoryTextureSRV;
 	if (HistorySize != OutputRect.Size())
 	{
-		check(!SceneColorOutputQuarterResTexture);
-
+		check(!SceneColorOutputHalfResTextureSRV);
+		check(!SceneColorOutputQuarterResTextureSRV);
+		
 		bool bNyquistHistory = HistorySize.X == 2 * OutputRect.Width() && HistorySize.Y == 2 * OutputRect.Height();
 
 		FTSRResolveHistoryCS::FParameters* PassParameters = GraphBuilder.AllocParameters<FTSRResolveHistoryCS::FParameters>();
@@ -2392,16 +2343,38 @@ FDefaultTemporalUpscaler::FOutputs AddTemporalSuperResolutionPasses(
 
 		PassParameters->UpdateHistoryOutputTexture = UpdateHistoryTextureSRV;
 		
-		PassParameters->SceneColorOutputMip0 = GraphBuilder.CreateUAV(FRDGTextureUAVDesc(SceneColorOutputTexture, /* InMipLevel = */ 0));
+		FRDGTextureRef SceneColorOutputTexture;
+		{
+			FRDGTextureDesc OutputDesc = FRDGTextureDesc::Create2D(
+				OutputExtent,
+				ColorFormat,
+				FClearValueBinding::None,
+				/* InFlags = */ TexCreate_ShaderResource | TexCreate_UAV,
+				/* NumMips = */ PassInputs.bGenerateOutputMip1 ? 2 : 1);
+			SceneColorOutputTexture = GraphBuilder.CreateTexture(OutputDesc, TEXT("TSR.Output"));
+
+			PassParameters->SceneColorOutputMip0 = GraphBuilder.CreateUAV(FRDGTextureUAVDesc(SceneColorOutputTexture, /* InMipLevel = */ 0));
+			SceneColorOutputTextureSRV = GraphBuilder.CreateSRV(FRDGTextureSRVDesc(SceneColorOutputTexture));
+		}
+
 		if (PassInputs.bGenerateOutputMip1)
 		{
 			PassParameters->bGenerateOutputMip1 = true;
 			PassParameters->SceneColorOutputMip1 = GraphBuilder.CreateUAV(FRDGTextureUAVDesc(SceneColorOutputTexture, /* InMipLevel = */ 1));
 		}
-		else if (SceneColorOutputHalfResTexture)
+		else if (PassInputs.bGenerateSceneColorHalfRes)
 		{
+			FRDGTextureDesc HalfResDesc = FRDGTextureDesc::Create2D(
+				OutputExtent / 2,
+				ColorFormat,
+				FClearValueBinding::None,
+				/* InFlags = */ TexCreate_ShaderResource | TexCreate_UAV);
+			FRDGTextureRef SceneColorOutputHalfResTexture = GraphBuilder.CreateTexture(HalfResDesc, TEXT("TSR.HalfResOutput"));
+
 			PassParameters->bGenerateOutputMip1 = true;
 			PassParameters->SceneColorOutputMip1 = GraphBuilder.CreateUAV(FRDGTextureUAVDesc(SceneColorOutputHalfResTexture));
+
+			SceneColorOutputHalfResTextureSRV = GraphBuilder.CreateSRV(FRDGTextureSRVDesc(SceneColorOutputHalfResTexture));
 		}
 		else
 		{
@@ -2640,15 +2613,15 @@ FDefaultTemporalUpscaler::FOutputs AddTemporalSuperResolutionPasses(
 
 	FDefaultTemporalUpscaler::FOutputs Outputs;
 	Outputs.FullRes = FScreenPassTextureSlice(SceneColorOutputTextureSRV, OutputRect);
-	if (SceneColorOutputHalfResTexture)
+	if (SceneColorOutputHalfResTextureSRV)
 	{
-		Outputs.HalfRes.Texture = SceneColorOutputHalfResTexture;
+		Outputs.HalfRes.TextureSRV = SceneColorOutputHalfResTextureSRV;
 		Outputs.HalfRes.ViewRect.Min = OutputRect.Min / 2;
 		Outputs.HalfRes.ViewRect.Max = Outputs.HalfRes.ViewRect.Min + FIntPoint::DivideAndRoundUp(OutputRect.Size(), 2);
 	}
-	if (SceneColorOutputQuarterResTexture)
+	if (SceneColorOutputQuarterResTextureSRV)
 	{
-		Outputs.QuarterRes.Texture = SceneColorOutputQuarterResTexture;
+		Outputs.QuarterRes.TextureSRV = SceneColorOutputQuarterResTextureSRV;
 		Outputs.QuarterRes.ViewRect.Min = OutputRect.Min / 4;
 		Outputs.QuarterRes.ViewRect.Max = Outputs.HalfRes.ViewRect.Min + FIntPoint::DivideAndRoundUp(OutputRect.Size(), 4);
 	}
diff --git a/Engine/Source/Runtime/Renderer/Private/ScreenSpaceRayTracing.cpp b/Engine/Source/Runtime/Renderer/Private/ScreenSpaceRayTracing.cpp
--- a/Engine/Source/Runtime/Renderer/Private/ScreenSpaceRayTracing.cpp
+++ b/Engine/Source/Runtime/Renderer/Private/ScreenSpaceRayTracing.cpp
@@ -1007,8 +1007,16 @@ void RenderScreenSpaceReflections(
 		}
 		else if (GSSRHalfResSceneColor && View.PrevViewInfo.HalfResTemporalAAHistory.IsValid())
 		{
-			InputColor = GraphBuilder.CreateSRV(FRDGTextureSRVDesc(
-				GraphBuilder.RegisterExternalTexture(View.PrevViewInfo.HalfResTemporalAAHistory)));
+			FRDGTextureRef HalfResTemporalAAHistory = GraphBuilder.RegisterExternalTexture(View.PrevViewInfo.HalfResTemporalAAHistory);
+
+			if (HalfResTemporalAAHistory->Desc.Dimension == ETextureDimension::Texture2DArray)
+			{
+				InputColor = GraphBuilder.CreateSRV(FRDGTextureSRVDesc::CreateForSlice(HalfResTemporalAAHistory, /* SliceIndex = */ 0));
+			}
+			else
+			{
+				InputColor = GraphBuilder.CreateSRV(FRDGTextureSRVDesc(HalfResTemporalAAHistory));
+			}
 		}
 		else if (View.PrevViewInfo.TemporalAAHistory.IsValid())
 		{
diff --git a/Engine/Source/Runtime/Renderer/Public/ScreenPass.h b/Engine/Source/Runtime/Renderer/Public/ScreenPass.h
--- a/Engine/Source/Runtime/Renderer/Public/ScreenPass.h
+++ b/Engine/Source/Runtime/Renderer/Public/ScreenPass.h
@@ -346,7 +346,14 @@ BEGIN_SHADER_PARAMETER_STRUCT(FScreenPassTextureInput, )
 	SHADER_PARAMETER_SAMPLER(SamplerState, Sampler)
 END_SHADER_PARAMETER_STRUCT()
 
+BEGIN_SHADER_PARAMETER_STRUCT(FScreenPassTextureSliceInput, )
+	SHADER_PARAMETER_STRUCT_INCLUDE(FScreenPassTextureViewportParameters, Viewport)
+	SHADER_PARAMETER_RDG_TEXTURE_SRV(Texture2D, Texture)
+	SHADER_PARAMETER_SAMPLER(SamplerState, Sampler)
+END_SHADER_PARAMETER_STRUCT()
+
 FScreenPassTextureInput GetScreenPassTextureInput(FScreenPassTexture Input, FRHISamplerState* Sampler);
+FScreenPassTextureSliceInput GetScreenPassTextureInput(FScreenPassTextureSlice Input, FRHISamplerState* Sampler);
 
 /** Draw information for the more advanced DrawScreenPass variant. Allows customizing the blend / depth stencil state,
  *  providing a custom vertex shader, and more fine-grained control of the underlying draw call.
diff --git a/Engine/Source/Runtime/Renderer/Public/ScreenPass.inl b/Engine/Source/Runtime/Renderer/Public/ScreenPass.inl
--- a/Engine/Source/Runtime/Renderer/Public/ScreenPass.inl
+++ b/Engine/Source/Runtime/Renderer/Public/ScreenPass.inl
@@ -390,4 +390,13 @@ inline FScreenPassTextureInput GetScreenPassTextureInput(FScreenPassTexture Text
 	Input.Texture = TexturePair.Texture;
 	Input.Sampler = Sampler;
 	return Input;
-}
\ No newline at end of file
+}
+
+inline FScreenPassTextureSliceInput GetScreenPassTextureInput(FScreenPassTextureSlice TexturePair, FRHISamplerState* Sampler)
+{
+	FScreenPassTextureSliceInput Input;
+	Input.Viewport = GetScreenPassTextureViewportParameters(FScreenPassTextureViewport(TexturePair));
+	Input.Texture = TexturePair.TextureSRV;
+	Input.Sampler = Sampler;
+	return Input;
+}
