Fix missing ShaderPrint ModifyEnvironement() calls in various shader declaration, causing them to fail compilation on non-supported platforms.

diff --git a/Engine/Plugins/Runtime/HairStrands/Source/HairStrandsCore/Private/HairStrandsInterpolation.cpp b/Engine/Plugins/Runtime/HairStrands/Source/HairStrandsCore/Private/HairStrandsInterpolation.cpp
--- a/Engine/Plugins/Runtime/HairStrands/Source/HairStrandsCore/Private/HairStrandsInterpolation.cpp
+++ b/Engine/Plugins/Runtime/HairStrands/Source/HairStrandsCore/Private/HairStrandsInterpolation.cpp
@@ -513,6 +513,7 @@ public:
 	static void ModifyCompilationEnvironment(const FGlobalShaderPermutationParameters& Parameters, FShaderCompilerEnvironment& OutEnvironment)
 	{
 		FGlobalShader::ModifyCompilationEnvironment(Parameters, OutEnvironment);
+		ShaderPrint::ModifyCompilationEnvironment(Parameters, OutEnvironment);
 		OutEnvironment.SetDefine(TEXT("SHADER_HAIRINTERPOLATION"), 1);
 		OutEnvironment.SetDefine(TEXT("GROUP_SIZE"), GetGroupSize());
 		OutEnvironment.CompilerFlags.Add(CFLAG_Wave32);
diff --git a/Engine/Source/Runtime/Renderer/Private/ShaderPrint.cpp b/Engine/Source/Runtime/Renderer/Private/ShaderPrint.cpp
--- a/Engine/Source/Runtime/Renderer/Private/ShaderPrint.cpp
+++ b/Engine/Source/Runtime/Renderer/Private/ShaderPrint.cpp
@@ -298,10 +298,10 @@ namespace ShaderPrint
 		return true;
 	}
 
-	void ModifyCompilationEnvironment(const FGlobalShaderPermutationParameters& Parameters, FShaderCompilerEnvironment& OutEnvironment)
+	void ModifyCompilationEnvironment(const EShaderPlatform Platform, FShaderCompilerEnvironment& OutEnvironment)
 	{
 		// Work around issues with HLSLcc by switching to DXC
-		if (IsHlslccShaderPlatform(Parameters.Platform))
+		if (IsHlslccShaderPlatform(Platform))
 		{
 			OutEnvironment.CompilerFlags.Add(CFLAG_ForceDXC);
 		}
@@ -309,6 +309,11 @@ namespace ShaderPrint
 		OutEnvironment.SetDefine(TEXT("SHADER_PRINT_EXPLICIT_BINDING"), 1);
 	}
 
+	void ModifyCompilationEnvironment(const FGlobalShaderPermutationParameters& Parameters, FShaderCompilerEnvironment& OutEnvironment)
+	{
+		ModifyCompilationEnvironment(Parameters.Platform, OutEnvironment);
+	}
+
 	bool IsEnabled()
 	{
 		return GEnabled != 0;
diff --git a/Engine/Source/Runtime/Renderer/Private/StochasticDirectLighting/StochasticDirectLighting.cpp b/Engine/Source/Runtime/Renderer/Private/StochasticDirectLighting/StochasticDirectLighting.cpp
--- a/Engine/Source/Runtime/Renderer/Private/StochasticDirectLighting/StochasticDirectLighting.cpp
+++ b/Engine/Source/Runtime/Renderer/Private/StochasticDirectLighting/StochasticDirectLighting.cpp
@@ -215,6 +215,7 @@ namespace StochasticDirectLighting
 	void ModifyCompilationEnvironment(EShaderPlatform Platform, FShaderCompilerEnvironment& OutEnvironment)
 	{
 		FForwardLightingParameters::ModifyCompilationEnvironment(Platform, OutEnvironment);
+		ShaderPrint::ModifyCompilationEnvironment(Platform, OutEnvironment);
 	}
 
 	// Keep in sync with TILE_TYPE_* in shaders
diff --git a/Engine/Source/Runtime/Renderer/Private/Substrate/SubstrateVisualize.cpp b/Engine/Source/Runtime/Renderer/Private/Substrate/SubstrateVisualize.cpp
--- a/Engine/Source/Runtime/Renderer/Private/Substrate/SubstrateVisualize.cpp
+++ b/Engine/Source/Runtime/Renderer/Private/Substrate/SubstrateVisualize.cpp
@@ -70,6 +70,7 @@ class FMaterialPrintInfoCS : public FGlobalShader
 	static void ModifyCompilationEnvironment(const FGlobalShaderPermutationParameters& Parameters, FShaderCompilerEnvironment& OutEnvironment)
 	{
 		FGlobalShader::ModifyCompilationEnvironment(Parameters, OutEnvironment);
+		ShaderPrint::ModifyCompilationEnvironment(Parameters, OutEnvironment);
 
 		// Stay debug and skip optimizations to reduce compilation time on this long shader.
 		OutEnvironment.CompilerFlags.Add(CFLAG_Debug);
@@ -113,6 +114,7 @@ class FVisualizeMaterialCountPS : public FGlobalShader
 	static void ModifyCompilationEnvironment(const FGlobalShaderPermutationParameters& Parameters, FShaderCompilerEnvironment& OutEnvironment)
 	{
 		FGlobalShader::ModifyCompilationEnvironment(Parameters, OutEnvironment);
+		ShaderPrint::ModifyCompilationEnvironment(Parameters, OutEnvironment);
 
 		// Stay debug and skip optimizations to reduce compilation time on this long shader.
 		OutEnvironment.CompilerFlags.Add(CFLAG_Debug);
@@ -171,6 +173,7 @@ class FSubstrateSystemInfoCS : public FGlobalShader
 	static void ModifyCompilationEnvironment(const FGlobalShaderPermutationParameters& Parameters, FShaderCompilerEnvironment& OutEnvironment)
 	{
 		FGlobalShader::ModifyCompilationEnvironment(Parameters, OutEnvironment);
+		ShaderPrint::ModifyCompilationEnvironment(Parameters, OutEnvironment);
 
 		// Stay debug and skip optimizations to reduce compilation time on this long shader.
 		OutEnvironment.CompilerFlags.Add(CFLAG_Debug);
@@ -211,6 +214,7 @@ class FMaterialDebugSubstrateTreeCS : public FGlobalShader
 	static void ModifyCompilationEnvironment(const FGlobalShaderPermutationParameters& Parameters, FShaderCompilerEnvironment& OutEnvironment)
 	{
 		FGlobalShader::ModifyCompilationEnvironment(Parameters, OutEnvironment);
+		ShaderPrint::ModifyCompilationEnvironment(Parameters, OutEnvironment);
 
 		// Stay debug and skip optimizations to reduce compilation time on this long shader.
 		OutEnvironment.CompilerFlags.Add(CFLAG_Debug);
@@ -255,6 +259,7 @@ class FMaterialDebugSubstrateTreePS : public FGlobalShader
 	static void ModifyCompilationEnvironment(const FGlobalShaderPermutationParameters& Parameters, FShaderCompilerEnvironment& OutEnvironment)
 	{
 		FGlobalShader::ModifyCompilationEnvironment(Parameters, OutEnvironment);
+		ShaderPrint::ModifyCompilationEnvironment(Parameters, OutEnvironment);
 
 		// Stay debug and skip optimizations to reduce compilation time on this long shader.
 		OutEnvironment.CompilerFlags.Add(CFLAG_Debug);
diff --git a/Engine/Source/Runtime/Renderer/Public/ShaderPrintParameters.h b/Engine/Source/Runtime/Renderer/Public/ShaderPrintParameters.h
--- a/Engine/Source/Runtime/Renderer/Public/ShaderPrintParameters.h
+++ b/Engine/Source/Runtime/Renderer/Public/ShaderPrintParameters.h
@@ -52,6 +52,7 @@ namespace ShaderPrint
 	RENDERER_API bool IsSupported(EShaderPlatform Platform);
 
 	// Set any flags or defines needed when using ShaderPrint
+	RENDERER_API void ModifyCompilationEnvironment(const EShaderPlatform Platform, FShaderCompilerEnvironment& OutEnvironment);
 	RENDERER_API void ModifyCompilationEnvironment(const FGlobalShaderPermutationParameters& Parameters, FShaderCompilerEnvironment& OutEnvironment);
 
 	// Have we enabled the ShaderPrint system?
