Fix stride mismatch for RWHasPendingStreaming which is declared as RWStructuredBuffer<uint> in HLSL

diff --git a/Engine/Source/Runtime/Renderer/Private/GlobalDistanceField.cpp b/Engine/Source/Runtime/Renderer/Private/GlobalDistanceField.cpp
--- a/Engine/Source/Runtime/Renderer/Private/GlobalDistanceField.cpp
+++ b/Engine/Source/Runtime/Renderer/Private/GlobalDistanceField.cpp
@@ -2352,7 +2352,7 @@ void UpdateGlobalDistanceFieldCache(
 		FRDGBufferRef PendingStreamingReadbackBuffer = nullptr;
 		if (bAnyClipmapHasPendingStreamingReadback)
 		{
-			FRDGBufferDesc HasPendingStreamingReadbackDesc = FRDGBufferDesc::CreateStructuredDesc(GlobalDistanceField::MaxClipmaps * sizeof(uint32), GlobalDistanceField::MaxClipmaps);
+			FRDGBufferDesc HasPendingStreamingReadbackDesc = FRDGBufferDesc::CreateStructuredDesc(sizeof(uint32), GlobalDistanceField::MaxClipmaps);
 			HasPendingStreamingReadbackDesc.Usage = EBufferUsageFlags(HasPendingStreamingReadbackDesc.Usage | BUF_SourceCopy);
 			PendingStreamingReadbackBuffer = GraphBuilder.CreateBuffer(HasPendingStreamingReadbackDesc, TEXT("GlobalDistanceField.HasPendingStreamingReadback"));
 			AddClearUAVPass(GraphBuilder, GraphBuilder.CreateUAV(PendingStreamingReadbackBuffer, PF_R32_UINT), 0);
