Remove the invalidating instance list from GPUScene (where it never belonged) and instead add to the VSM cache manager, via an interface to collect the dynamic instances where needed.

diff --git a/Engine/Source/Runtime/Renderer/Private/GPUScene.cpp b/Engine/Source/Runtime/Renderer/Private/GPUScene.cpp
--- a/Engine/Source/Runtime/Renderer/Private/GPUScene.cpp
+++ b/Engine/Source/Runtime/Renderer/Private/GPUScene.cpp
@@ -29,6 +29,7 @@
 #include "InstanceCulling/InstanceCullingOcclusionQuery.h"
 #include "PrimitiveUniformShaderParametersBuilder.h"
 #include "InstanceDataSceneProxy.h"
+#include "SceneRendererInterface.h"
 
 // Useful for debugging
 #define FORCEINLINE_GPUSCENE FORCEINLINE
@@ -1571,7 +1572,7 @@ struct FUploadDataSourceAdapterDynamicPrimitives
 	TArray<uint32, SceneRenderingAllocator> PrimitivesIds;
 };
 
-void FGPUScene::UploadDynamicPrimitiveShaderDataForViewInternal(FRDGBuilder& GraphBuilder, FViewInfo& View, bool bIsShadowView)
+void FGPUScene::UploadDynamicPrimitiveShaderDataForViewInternal(FRDGBuilder& GraphBuilder, FViewInfo& View, UE::Renderer::Private::IShadowInvalidatingInstances *ShadowInvalidatingInstances)
 {
 	LLM_SCOPE_BYTAG(GPUScene);
 
@@ -1610,20 +1611,14 @@ void FGPUScene::UploadDynamicPrimitiveShaderDataForViewInternal(FRDGBuilder& Gra
 		ensure(UploadIdStart < DynamicPrimitivesOffset);
 		ensure(InstanceIdStart != INDEX_NONE);
 
-		if (bIsShadowView)
+		if (ShadowInvalidatingInstances)
 		{
 			// Enqueue cache invalidations for all dynamic primitives' instances, as they will be removed this frame and are not associated
 			// with any particular FPrimitiveSceneInfo. Will occur on the next call to UpdateAllPrimitiveSceneInfos
 			for (const FGPUScenePrimitiveCollector::FPrimitiveData& PrimitiveData : Collector.UploadData->PrimitiveData)
 			{
-				ensure(PrimitiveData.LocalInstanceSceneDataOffset != INDEX_NONE);				
-				DynamicPrimitiveInstancesToInvalidate.Add(
-					FInstanceRange
-					{
-						PrimitiveData.LocalInstanceSceneDataOffset + InstanceIdStart,
-						PrimitiveData.NumInstances
-					}
-				);
+				check(PrimitiveData.LocalInstanceSceneDataOffset != INDEX_NONE);
+				ShadowInvalidatingInstances->AddInstanceRange(PrimitiveData.LocalInstanceSceneDataOffset + uint32(InstanceIdStart), PrimitiveData.NumInstances);
 			}
 		}
 
@@ -1787,13 +1782,13 @@ void FGPUScene::Update(FRDGBuilder& GraphBuilder, FSceneUniformBuffer& SceneUB,
 	}
 }
 
-void FGPUScene::UploadDynamicPrimitiveShaderDataForView(FRDGBuilder& GraphBuilder, FViewInfo& View, bool bIsShadowView)
+void FGPUScene::UploadDynamicPrimitiveShaderDataForView(FRDGBuilder& GraphBuilder, FViewInfo& View, UE::Renderer::Private::IShadowInvalidatingInstances *ShadowInvalidatingInstances)
 {
 	if (bIsEnabled)
 	{
 		RDG_GPU_MASK_SCOPE(GraphBuilder, FRHIGPUMask::All());
 
-		UploadDynamicPrimitiveShaderDataForViewInternal(GraphBuilder, View, bIsShadowView);
+		UploadDynamicPrimitiveShaderDataForViewInternal(GraphBuilder, View, ShadowInvalidatingInstances);
 	}
 }
 
@@ -2167,7 +2162,6 @@ SIZE_T FGPUScene::GetAllocatedSize() const
 	return PrimitivesToUpdate.GetAllocatedSize()
 		+ InstanceRangesToClear.GetAllocatedSize()
 		+ PrimitiveDirtyState.GetAllocatedSize()
-		+ DynamicPrimitiveInstancesToInvalidate.GetAllocatedSize()
 		+ DeferredGPUWritePassDelegates[uint32(EGPUSceneGPUWritePass::PostOpaqueRendering)].GetAllocatedSize();
 }
 
diff --git a/Engine/Source/Runtime/Renderer/Private/GPUScene.h b/Engine/Source/Runtime/Renderer/Private/GPUScene.h
--- a/Engine/Source/Runtime/Renderer/Private/GPUScene.h
+++ b/Engine/Source/Runtime/Renderer/Private/GPUScene.h
@@ -28,6 +28,10 @@ class FScenePostUpdateChangeSet;
 class FPrimitiveSceneProxy;
 class IVisibilityTaskData;
 struct FLightSceneChangeSet;
+namespace UE::Renderer::Private
+{
+	class IShadowInvalidatingInstances;
+}
 
 DECLARE_GPU_STAT_NAMED_EXTERN(GPUSceneUpdate, TEXT("GPUSceneUpdate"))
 
@@ -217,7 +221,7 @@ public:
 	/**
 	 * Upload primitives from View.DynamicPrimitiveCollector.
 	 */
-	void UploadDynamicPrimitiveShaderDataForView(FRDGBuilder& GraphBuilder, FViewInfo& View, bool bIsShadowView = false);
+	void UploadDynamicPrimitiveShaderDataForView(FRDGBuilder& GraphBuilder, FViewInfo& View, UE::Renderer::Private::IShadowInvalidatingInstances *ShadowInvalidatingInstances = nullptr);
 
 	/**
 	 * Modifies the GPUScene specific scene UB parameters to the current versions. Returns true if any of the parameters changed.
@@ -321,9 +325,6 @@ public:
 	FRDGAsyncScatterUploadBuffer   LightmapUploadBuffer;
 
 	using FInstanceRange = FGPUSceneInstanceRange;
-
-	TArray<FInstanceRange> DynamicPrimitiveInstancesToInvalidate;
-
 	using FInstanceGPULoadBalancer = TInstanceCullingLoadBalancer<SceneRenderingAllocator>;
 
 	inline const FScene &GetScene() const { return Scene; }
@@ -428,7 +429,7 @@ private:
 
 	static void InitLightData(const FLightSceneInfoCompact& LightInfoCompact, bool bAllowStaticLighting, FLightSceneData& DataOut);
 
-	void UploadDynamicPrimitiveShaderDataForViewInternal(FRDGBuilder& GraphBuilder, FViewInfo& View, bool bIsShadowView);
+	void UploadDynamicPrimitiveShaderDataForViewInternal(FRDGBuilder& GraphBuilder, FViewInfo& View, UE::Renderer::Private::IShadowInvalidatingInstances *ShadowInvalidatingInstances);
 
 	void UpdateInternal(FRDGBuilder& GraphBuilder, FSceneUniformBuffer& SceneUB, FRDGExternalAccessQueue& ExternalAccessQueue, IVisibilityTaskData* VisibilityTaskData);
 
diff --git a/Engine/Source/Runtime/Renderer/Private/SceneRendering.cpp b/Engine/Source/Runtime/Renderer/Private/SceneRendering.cpp
--- a/Engine/Source/Runtime/Renderer/Private/SceneRendering.cpp
+++ b/Engine/Source/Runtime/Renderer/Private/SceneRendering.cpp
@@ -4520,6 +4520,15 @@ void FSceneRenderer::RenderThreadEnd(FRHICommandListImmediate& RHICmdList, const
 	}
 }
 
+UE::Renderer::Private::IShadowInvalidatingInstances *FSceneRenderer::GetShadowInvalidatingInstancesInterface(const FSceneView *SceneView)
+{
+	if (ShadowSceneRenderer)
+	{
+		return ShadowSceneRenderer->GetInvalidatingInstancesInterface(SceneView);
+	}
+	return nullptr;
+}
+
 void FSceneRenderer::CleanUp(FRHICommandListImmediate& RHICmdList)
 {
 	if (GSceneRenderCleanUpState.CompletionMode == ESceneRenderCleanUpMode::Immediate || GSceneRenderCleanUpState.Renderers.IsEmpty())
diff --git a/Engine/Source/Runtime/Renderer/Private/SceneRendering.h b/Engine/Source/Runtime/Renderer/Private/SceneRendering.h
--- a/Engine/Source/Runtime/Renderer/Private/SceneRendering.h
+++ b/Engine/Source/Runtime/Renderer/Private/SceneRendering.h
@@ -2102,8 +2102,10 @@ public:
 	RENDERER_API static void RenderThreadBegin(FRHICommandListImmediate& RHICmdList, const TArray<FSceneRenderer*>& SceneRenderers);
 	RENDERER_API static void RenderThreadEnd(FRHICommandListImmediate& RHICmdList, const TArray<FSceneRenderer*>& SceneRenderers);
 
-	// FSceneRenderer interface
+	// ISceneRenderer interface
+	virtual UE::Renderer::Private::IShadowInvalidatingInstances *GetShadowInvalidatingInstancesInterface(const FSceneView *SceneView) override;
 
+	// FSceneRenderer interface
 	virtual void Render(FRDGBuilder& GraphBuilder) = 0;
 	virtual void RenderHitProxies(FRDGBuilder& GraphBuilder) {}
 	virtual bool ShouldRenderVelocities() const { return false; }
diff --git a/Engine/Source/Runtime/Renderer/Private/ShadowSetup.cpp b/Engine/Source/Runtime/Renderer/Private/ShadowSetup.cpp
--- a/Engine/Source/Runtime/Renderer/Private/ShadowSetup.cpp
+++ b/Engine/Source/Runtime/Renderer/Private/ShadowSetup.cpp
@@ -5534,6 +5534,8 @@ void FSceneRenderer::AddViewDependentWholeSceneShadowsForView(
 					VisibleLightInfo.AllProjectedShadows.Add(ProjectedShadowInfo);
 					ShadowInfosThatNeedCulling.Add(ProjectedShadowInfo);
 
+					ProjectedShadowInfo->VirtualShadowMapPerLightCacheEntry = VirtualShadowMapClipmap->GetCacheEntry();
+
 					// TODO: disentangle from legacy setup alltogether
 					// TODO2: This needs to not depend on the non-nanite shadows!
 					ShadowSceneRenderer->AddDirectionalLightShadow(ProjectedShadowInfo);
@@ -6506,32 +6508,32 @@ void FSceneRenderer::FinishDynamicShadowMeshPassSetup(FRDGBuilder& GraphBuilder,
 	{
 		for (FProjectedShadowInfo* ProjectedShadowInfo : ShadowMapAtlas.Shadows)
 		{
-			Scene->GPUScene.UploadDynamicPrimitiveShaderDataForView(GraphBuilder, *ProjectedShadowInfo->ShadowDepthView, true);
+			Scene->GPUScene.UploadDynamicPrimitiveShaderDataForView(GraphBuilder, *ProjectedShadowInfo->ShadowDepthView);
 		}
 	}
 	for (FSortedShadowMapAtlas& ShadowMap : SortedShadowsForShadowDepthPass.ShadowMapCubemaps)
 	{
 		check(ShadowMap.Shadows.Num() == 1);
 		FProjectedShadowInfo* ProjectedShadowInfo = ShadowMap.Shadows[0];
-		Scene->GPUScene.UploadDynamicPrimitiveShaderDataForView(GraphBuilder, *ProjectedShadowInfo->ShadowDepthView, true);
+		Scene->GPUScene.UploadDynamicPrimitiveShaderDataForView(GraphBuilder, *ProjectedShadowInfo->ShadowDepthView);
 	}
 	for (FProjectedShadowInfo* ProjectedShadowInfo : SortedShadowsForShadowDepthPass.PreshadowCache.Shadows)
 	{
 		if (!ProjectedShadowInfo->bDepthsCached)
 		{
-			Scene->GPUScene.UploadDynamicPrimitiveShaderDataForView(GraphBuilder, *ProjectedShadowInfo->ShadowDepthView, true);
+			Scene->GPUScene.UploadDynamicPrimitiveShaderDataForView(GraphBuilder, *ProjectedShadowInfo->ShadowDepthView);
 		}
 	}
 	for (const FSortedShadowMapAtlas& ShadowMapAtlas : SortedShadowsForShadowDepthPass.TranslucencyShadowMapAtlases)
 	{
 		for (FProjectedShadowInfo* ProjectedShadowInfo : ShadowMapAtlas.Shadows)
 		{
-			Scene->GPUScene.UploadDynamicPrimitiveShaderDataForView(GraphBuilder, *ProjectedShadowInfo->ShadowDepthView, true);
+			Scene->GPUScene.UploadDynamicPrimitiveShaderDataForView(GraphBuilder, *ProjectedShadowInfo->ShadowDepthView);
 		}
 	}
 	for (FProjectedShadowInfo* ProjectedShadowInfo : SortedShadowsForShadowDepthPass.VirtualShadowMapShadows)
 	{
-		Scene->GPUScene.UploadDynamicPrimitiveShaderDataForView(GraphBuilder, *ProjectedShadowInfo->ShadowDepthView, true);
+		Scene->GPUScene.UploadDynamicPrimitiveShaderDataForView(GraphBuilder, *ProjectedShadowInfo->ShadowDepthView, GetShadowInvalidatingInstancesInterface(ProjectedShadowInfo->DependentView) );
 	}
 
 	DynamicReadBufferForShadows.Commit(GraphBuilder.RHICmdList);
diff --git a/Engine/Source/Runtime/Renderer/Private/Shadows/ShadowSceneRenderer.cpp b/Engine/Source/Runtime/Renderer/Private/Shadows/ShadowSceneRenderer.cpp
--- a/Engine/Source/Runtime/Renderer/Private/Shadows/ShadowSceneRenderer.cpp
+++ b/Engine/Source/Runtime/Renderer/Private/Shadows/ShadowSceneRenderer.cpp
@@ -170,6 +170,18 @@ void FShadowSceneRenderer::BeginRender(FRDGBuilder& GraphBuilder)
 	});
 }
 
+UE::Renderer::Private::IShadowInvalidatingInstances *FShadowSceneRenderer::GetInvalidatingInstancesInterface(const FSceneView *SceneView)
+{
+	// No need to collect invalidations if there is nothing to invalidate.
+	if (Scene.VirtualShadowMapCache && Scene.VirtualShadowMapCache->IsCacheDataAvailable())
+	{
+		// TODO: Make use of the SceneView parameter to register invalidations for view-dependent shadows appropriately.
+		return Scene.VirtualShadowMapCache->GetInvalidatingInstancesInterface();
+	}
+	return nullptr;
+}
+
+
 static float GetResolutionLODBiasLocal(float LightMobilityFactor)
 {
 	return FVirtualShadowMapArray::InterpolateResolutionBias(
diff --git a/Engine/Source/Runtime/Renderer/Private/Shadows/ShadowSceneRenderer.h b/Engine/Source/Runtime/Renderer/Private/Shadows/ShadowSceneRenderer.h
--- a/Engine/Source/Runtime/Renderer/Private/Shadows/ShadowSceneRenderer.h
+++ b/Engine/Source/Runtime/Renderer/Private/Shadows/ShadowSceneRenderer.h
@@ -16,6 +16,10 @@ class FRDGBuilder;
 class FVirtualShadowMapPerLightCacheEntry;
 struct FNaniteVisibilityQuery;
 class FShadowScene;
+namespace UE::Renderer::Private
+{
+	class IShadowInvalidatingInstances;
+}
 
 /**
  * Transient scope for per-frame rendering resources for the shadow rendering.
@@ -96,6 +100,8 @@ public:
 		return RendererSetupTask;
 	}
 
+	UE::Renderer::Private::IShadowInvalidatingInstances *GetInvalidatingInstancesInterface(const FSceneView *SceneView);
+
 private:
 	UE::Tasks::FTask RendererSetupTask;
 
diff --git a/Engine/Source/Runtime/Renderer/Private/VirtualShadowMaps/VirtualShadowMapCacheManager.cpp b/Engine/Source/Runtime/Renderer/Private/VirtualShadowMaps/VirtualShadowMapCacheManager.cpp
--- a/Engine/Source/Runtime/Renderer/Private/VirtualShadowMaps/VirtualShadowMapCacheManager.cpp
+++ b/Engine/Source/Runtime/Renderer/Private/VirtualShadowMaps/VirtualShadowMapCacheManager.cpp
@@ -298,6 +298,16 @@ bool FVirtualShadowMapPerLightCacheEntry::UpdateLocal(const FProjectedShadowInit
 	return Prev.RenderedFrameNumber >= 0;
 }
 
+void FVirtualShadowMapArrayCacheManager::FShadowInvalidatingInstancesImplementation::AddPrimitive(const FPrimitiveSceneInfo* PrimitiveSceneInfo)
+{
+	AddInstanceRange(PrimitiveSceneInfo->GetInstanceSceneDataOffset(), PrimitiveSceneInfo->GetNumInstanceSceneDataEntries());
+}
+
+void FVirtualShadowMapArrayCacheManager::FShadowInvalidatingInstancesImplementation::AddInstanceRange(uint32 InstanceSceneDataOffset, uint32 NumInstanceSceneDataEntries)
+{
+	PrimitiveInstancesToInvalidate.Add(FVirtualShadowMapInstanceRange{int32(InstanceSceneDataOffset), int32(NumInstanceSceneDataEntries)});
+}		
+
 void FVirtualShadowMapPerLightCacheEntry::Invalidate()
 {
 	Prev.RenderedFrameNumber = -1;
@@ -312,19 +322,18 @@ FVirtualShadowMapArrayCacheManager::FInvalidatingPrimitiveCollector::FInvalidati
 
 void FVirtualShadowMapArrayCacheManager::FInvalidatingPrimitiveCollector::AddDynamicAndGPUPrimitives()
 {
-	// Add and clear pending invalidations enqueued on the GPU Scene from dynamic primitives added since last invalidation
-	for (const FGPUScene::FInstanceRange& Range : GPUScene.DynamicPrimitiveInstancesToInvalidate)
+	// Add and clear pending invalidations enqueued since the last invalidation (scene update)
+	// This includes "dynamic primitives"
+	for (const FVirtualShadowMapInstanceRange& Range : Manager.ShadowInvalidatingInstancesImplementation.PrimitiveInstancesToInvalidate)
 	{
-		// Dynamic primitives are never cached as static; see FUploadDataSourceAdapterDynamicPrimitives::GetPrimitiveInfo
-		// TODO: Do we ever need to invalidate these "post" update?
 		Instances.Add(Range.InstanceSceneDataOffset, Range.NumInstanceSceneDataEntries, 0);
 	}
 
-	GPUScene.DynamicPrimitiveInstancesToInvalidate.Reset();
+	Manager.ShadowInvalidatingInstancesImplementation.PrimitiveInstancesToInvalidate.Reset();
 
 	for (auto& CacheEntry : Manager.CacheEntries)
 	{
-		for (const FVirtualShadowMapPerLightCacheEntry::FInstanceRange& Range : CacheEntry.Value->PrimitiveInstancesToInvalidate)
+		for (const FVirtualShadowMapInstanceRange& Range : CacheEntry.Value->PrimitiveInstancesToInvalidate)
 		{
 			// TODO: Do we ever need to invalidate these "post" update?
 			Instances.Add(Range.InstanceSceneDataOffset, Range.NumInstanceSceneDataEntries, 0);
@@ -464,6 +473,7 @@ FVirtualShadowMapFeedback::FReadbackInfo FVirtualShadowMapFeedback::GetLatestRea
 
 FVirtualShadowMapArrayCacheManager::FVirtualShadowMapArrayCacheManager(FScene* InScene) 
 	: Scene(InScene)
+	, ShadowInvalidatingInstancesImplementation(*this)
 {
 	// Handle message with status sent back from GPU
 	StatusFeedbackSocket = GPUMessage::RegisterHandler(TEXT("Shadow.Virtual.StatusFeedback"), [this](GPUMessage::FReader Message)
@@ -842,10 +852,7 @@ void FVirtualShadowMapPerLightCacheEntry::OnPrimitiveRendered(const FPrimitiveSc
 		// Skip if the invalidation mode is NOT auto (because Always will do it elsewhere & the others should prevent this).
 		if (PrimitiveSceneInfo->Proxy->HasDeformableMesh() && PrimitiveSceneInfo->Proxy->GetShadowCacheInvalidationBehavior() == EShadowCacheInvalidationBehavior::Auto)
 		{
-			PrimitiveInstancesToInvalidate.Add(FInstanceRange{ 
-				PrimitiveSceneInfo->GetInstanceSceneDataOffset(),
-				PrimitiveSceneInfo->GetNumInstanceSceneDataEntries()
-			});
+			PrimitiveInstancesToInvalidate.Add(FVirtualShadowMapInstanceRange{ PrimitiveSceneInfo->GetInstanceSceneDataOffset(), PrimitiveSceneInfo->GetNumInstanceSceneDataEntries() });
 		}
 	}
 }
diff --git a/Engine/Source/Runtime/Renderer/Private/VirtualShadowMaps/VirtualShadowMapCacheManager.h b/Engine/Source/Runtime/Renderer/Private/VirtualShadowMaps/VirtualShadowMapCacheManager.h
--- a/Engine/Source/Runtime/Renderer/Private/VirtualShadowMaps/VirtualShadowMapCacheManager.h
+++ b/Engine/Source/Runtime/Renderer/Private/VirtualShadowMaps/VirtualShadowMapCacheManager.h
@@ -8,6 +8,7 @@
 #include "InstanceCulling/InstanceCullingLoadBalancer.h"
 #include "GPUScene.h"
 #include "GPUMessaging.h"
+#include "SceneRendererInterface.h"
 
 class FRHIGPUBufferReadback;
 class FGPUScene;
@@ -16,6 +17,12 @@ class FInvalidatePagesParameters;
 
 namespace Nanite { struct FPackedViewParams; }
 
+struct FVirtualShadowMapInstanceRange
+{
+	int32 InstanceSceneDataOffset;
+	int32 NumInstanceSceneDataEntries;
+};
+
 #define VSM_LOG_INVALIDATIONS 0
 
 class FVirtualShadowMapCacheEntry
@@ -127,14 +134,7 @@ public:
 	// One entry represents the cached state of a given shadow map in the set of either a clipmap(N), one cube map(6) or a regular VSM (1)
 	TArray<FVirtualShadowMapCacheEntry> ShadowMapEntries;
 
-	// TODO: refactor this to not ne stored in the cache entry when we move (some) invalidaitons to the end of frame rather than in the scene primitive updates.
-	struct FInstanceRange
-	{
-		int32 InstanceSceneDataOffset;
-		int32 NumInstanceSceneDataEntries;
-	};
-
-	TArray<FInstanceRange> PrimitiveInstancesToInvalidate;
+	TArray<FVirtualShadowMapInstanceRange> PrimitiveInstancesToInvalidate;
 
 private:
 	FProjectedShadowInitializer LocalCacheKey;
@@ -345,7 +345,23 @@ public:
 		return CacheEntries.CreateConstIterator();
 	}
 
+	UE::Renderer::Private::IShadowInvalidatingInstances *GetInvalidatingInstancesInterface() { return &ShadowInvalidatingInstancesImplementation; }
+
 private:
+
+	/** 
+	 */
+	class FShadowInvalidatingInstancesImplementation : public UE::Renderer::Private::IShadowInvalidatingInstances
+	{
+	public:
+		FShadowInvalidatingInstancesImplementation(FVirtualShadowMapArrayCacheManager &InCacheManager) : CacheManager(InCacheManager) {}
+		virtual void AddPrimitive(const FPrimitiveSceneInfo *PrimitiveSceneInfo);
+		virtual void AddInstanceRange(uint32 InstanceSceneDataOffset, uint32 NumInstanceSceneDataEntries);
+
+		FVirtualShadowMapArrayCacheManager &CacheManager;
+		TArray<FVirtualShadowMapInstanceRange> PrimitiveInstancesToInvalidate;
+	};
+
 	// Invalidate the cache for all shadows, causing any pages to be rerendered
 	void Invalidate(FRDGBuilder& GraphBuilder);
 
@@ -402,8 +418,6 @@ private:
 	float GlobalResolutionLodBias = 0.0f;
 	uint32 LastFrameOverPageAllocationBudget = 0;
 	
-	
-
 	// Debug stuff
 #if !UE_BUILD_SHIPPING
 	FDelegateHandle ScreenMessageDelegate;
@@ -426,4 +440,5 @@ private:
 #endif // UE_BUILD_SHIPPING
 
 	FScene* Scene;
+	FShadowInvalidatingInstancesImplementation ShadowInvalidatingInstancesImplementation;
 };
diff --git a/Engine/Source/Runtime/Renderer/Public/SceneRendererInterface.h b/Engine/Source/Runtime/Renderer/Public/SceneRendererInterface.h
--- a/Engine/Source/Runtime/Renderer/Public/SceneRendererInterface.h
+++ b/Engine/Source/Runtime/Renderer/Public/SceneRendererInterface.h
@@ -6,9 +6,38 @@
 #include "ShaderParameterStructDeclaration.h"
 
 class FSceneUniformBuffer;
+class FPrimitiveSceneInfo;
 
 DECLARE_UNIFORM_BUFFER_STRUCT(FSceneUniformParameters, RENDERER_API)
 
+namespace UE::Renderer::Private
+{
+
+/** 
+ * Experimental interface to invalidate (VSM) shadows by marking instances during a scene render. The collected objects are invalidated at the next scene 
+ * update time, before GPU-Scene update, so this can't be used to trigger invalidation for the current frame.
+ * NOTE: marked as Renderer::Private as it is experimental and may be subject to change.
+ */
+class IShadowInvalidatingInstances
+{
+public:
+	RENDERER_API virtual ~IShadowInvalidatingInstances() = default;
+
+	/**
+	 * Add primitive for invalidation. Will add all the instances to the queue. 
+	 * Should be called from the render thread.
+	 */
+	virtual void AddPrimitive(const FPrimitiveSceneInfo *PrimitiveSceneInfo) = 0;
+
+	/**
+	 * Add a range of instances for invalidation. 
+	 * Should be called from the render thread.
+	 */
+	virtual void AddInstanceRange(uint32 InstanceSceneDataOffset, uint32 NumInstanceSceneDataEntries) = 0;
+};
+
+}
+
 /**
  * Public interface into FSceneRenderer. Used as the scope for scene rendering functions.
  */
@@ -21,6 +50,13 @@ public:
 	virtual FSceneUniformBuffer& GetSceneUniforms() = 0;
 
 	virtual TRDGUniformBufferRef<FSceneUniformParameters> GetSceneUniformBufferRef(FRDGBuilder& GraphBuilder) = 0;
+
+	/** 
+	 * If supported, will return an interface to mark primitives as invalidating the shadows (which will be effected post-frame).
+	 * The returned interface will mark primitives for view-dependent shadows according to the supplied SceneView.
+	 * Returns null if there is no need to invalidate anything.
+	 */
+	RENDERER_API virtual UE::Renderer::Private::IShadowInvalidatingInstances *GetShadowInvalidatingInstancesInterface(const FSceneView *SceneView) { return nullptr; }
 };
 
 inline TRDGUniformBufferRef<FSceneUniformParameters> GetSceneUniformBufferRef(FRDGBuilder& GraphBuilder, const FSceneView& View)
