Fix incorrect assert condition that would trigger when trying to allocate the last atlas slot.

diff --git a/Engine/Source/Runtime/Renderer/Private/LightFunctionAtlas.cpp b/Engine/Source/Runtime/Renderer/Private/LightFunctionAtlas.cpp
--- a/Engine/Source/Runtime/Renderer/Private/LightFunctionAtlas.cpp
+++ b/Engine/Source/Runtime/Renderer/Private/LightFunctionAtlas.cpp
@@ -389,7 +389,7 @@ void FLightFunctionAtlas::AllocateAtlasSlots(const TArray<FViewInfo>& Views)
 	uint32 NextAtlasSlotY = 0;
 	auto AddAtlasSlot = [&](FLightFunctionSlotKey Key, const FMaterialRenderProxy* LightFunctionMaterial)
 	{
-		check(uint32(LightFunctionsSet.Num()) < (AtlasMaxLightFunctionCount - 1));
+		check(uint32(LightFunctionsSet.Num()) < AtlasMaxLightFunctionCount);
 		const uint32 NewSlotIndex = EffectiveLightFunctionSlotArray.Num();
 
 		Key.EffectiveLightFunctionSlotIndex = NewSlotIndex;
