Prevent custom render passes from using screen percentage.

#rnx

diff --git a/Engine/Source/Runtime/Renderer/Private/SceneRendering.cpp b/Engine/Source/Runtime/Renderer/Private/SceneRendering.cpp
--- a/Engine/Source/Runtime/Renderer/Private/SceneRendering.cpp
+++ b/Engine/Source/Runtime/Renderer/Private/SceneRendering.cpp
@@ -2880,10 +2880,13 @@ FIntPoint FSceneRenderer::GetDesiredInternalBufferSize(const FSceneViewFamily& V
 
 	for (const FSceneView* View : ViewFamily.AllViews)
 	{
-		FIntPoint ViewSize = ApplyResolutionFraction(ViewFamily, View->UnconstrainedViewRect.Size(), ResolutionFractionUpperBound);
+		// Note: This ensures that custom passes (rendered with the main renderer) ignore screen percentage, like regular scene captures.
+		const float AdjustedResolutionFractionUpperBounds = View->CustomRenderPass ? 1.0f : ResolutionFractionUpperBound;
+
+		FIntPoint ViewSize = ApplyResolutionFraction(ViewFamily, View->UnconstrainedViewRect.Size(), AdjustedResolutionFractionUpperBounds);
 		FIntPoint ViewRectMin = QuantizeViewRectMin(FIntPoint(
-			FMath::CeilToInt(View->UnconstrainedViewRect.Min.X * ResolutionFractionUpperBound),
-			FMath::CeilToInt(View->UnconstrainedViewRect.Min.Y * ResolutionFractionUpperBound)));
+			FMath::CeilToInt(View->UnconstrainedViewRect.Min.X * AdjustedResolutionFractionUpperBounds),
+			FMath::CeilToInt(View->UnconstrainedViewRect.Min.Y * AdjustedResolutionFractionUpperBounds)));
 
 		FamilySizeUpperBound.X = FMath::Max(FamilySizeUpperBound.X, ViewRectMin.X + ViewSize.X);
 		FamilySizeUpperBound.Y = FMath::Max(FamilySizeUpperBound.Y, ViewRectMin.Y + ViewSize.Y);
