Removes all references to POST_PROCESS_ALPHA in post processing chain except the tonemapper

POST_PROCESS_ALPHA was made using a readonly project settings because back then there was no shader permutation API, yet do not introduce perf regression on consoles that didn't need alpha channel cost.

Problems we currently run into nowdays, is ease of testing, lack of EngineTest, as well as non converage of our shader compilation reflight in these scopes.

This changes adds shader permutations throughout DOF, MotionBlur, TAA, Mitchell Netravali, Upscale, FXAA, and adds shader parameter in more trivial shader like FFT bloom's clamp pass and editor primitives.

Tonemapper has too many permutation to do it at the moment. Remaining POST_PROCESS_ALPHA outside the post processing chain are left as an exercies to the respective system owners.

diff --git a/Engine/Shaders/Private/Bloom/BloomClampKernel.usf b/Engine/Shaders/Private/Bloom/BloomClampKernel.usf
--- a/Engine/Shaders/Private/Bloom/BloomClampKernel.usf
+++ b/Engine/Shaders/Private/Bloom/BloomClampKernel.usf
@@ -10,6 +10,8 @@
 
 //------------------------------------------------------- PARAMETERS
 
+uint bProcessAlpha;
+
 Texture2D KernelSpatialTexture;
 StructuredBuffer<FBloomKernelInfo> KernelConstantsBuffer;
 RWTexture2D<float4> ClampedKernelSpatialOutput;
@@ -28,11 +30,12 @@ void MainCS(
 
 	float4 KernelTexel = KernelSpatialTexture[PixelCoord];
 
-#if POST_PROCESS_ALPHA > 0
 	float4 ClampedKernel = clamp(KernelTexel, 0.0, OriginalKernelInfo.MaxScatterDispersionEnergy);
-#else
-	float4 ClampedKernel = clamp(KernelTexel.rgbb, 0.0, OriginalKernelInfo.MaxScatterDispersionEnergy.rgbb);
-#endif
+	FLATTEN
+	if (!bProcessAlpha)
+	{
+		ClampedKernel.a = ClampedKernel.b;
+	}
 
 	ClampedKernelSpatialOutput[PixelCoord] = ClampedKernel;
 }
diff --git a/Engine/Shaders/Private/Common.ush b/Engine/Shaders/Private/Common.ush
--- a/Engine/Shaders/Private/Common.ush
+++ b/Engine/Shaders/Private/Common.ush
@@ -119,16 +119,6 @@ const static float Max11BitsFloat = 65024.0f;
 const static float Max10BitsFloat = 64512.0f;
 const static float3 Max111110BitsFloat3 = float3(Max11BitsFloat, Max11BitsFloat, Max10BitsFloat);
 
-#if POST_PROCESS_ALPHA
-#define SceneColorLayout float4
-#define CastFloat4ToSceneColorLayout(x) (x)
-#define SetSceneColorLayoutToFloat4(dest,value) dest = (value)
-#else
-#define SceneColorLayout float3
-#define CastFloat4ToSceneColorLayout(x) ((x).rgb)
-#define SetSceneColorLayoutToFloat4(dest,value) dest.rgb = (value).rgb
-#endif
-
 #define SUPPORTS_TEXTURE_EXTERNAL		(COMPILER_GLSL_ES3_1)
 
 #if !SUPPORTS_TEXTURE_EXTERNAL
diff --git a/Engine/Shaders/Private/DiaphragmDOF/DOFCommon.ush b/Engine/Shaders/Private/DiaphragmDOF/DOFCommon.ush
--- a/Engine/Shaders/Private/DiaphragmDOF/DOFCommon.ush
+++ b/Engine/Shaders/Private/DiaphragmDOF/DOFCommon.ush
@@ -13,6 +13,11 @@
 #include "../CircleDOFCommon.ush"
 
 
+//------------------------------------------------------- RECOMPILE HASH
+
+#pragma message("UESHADERMETADATA_VERSION 65292827-F021-4C62-9790-FC626A1BFE69")
+
+
 //------------------------------------------------------- COMPILER CONFIG
 
 // Generate vector truncation warnings to errors.
@@ -104,7 +109,11 @@
 #define BOKEH_LUT_SIZE 32
 
 /** Whether alpha channel should be processed or not. */
-#define CONFIG_DOF_ALPHA (POST_PROCESS_ALPHA)
+#ifdef DIM_ALPHA_CHANNEL
+	#define CONFIG_DOF_ALPHA (DIM_ALPHA_CHANNEL)
+#else
+	#define CONFIG_DOF_ALPHA 0
+#endif
 
 /** Amount of error tolerated to quick fast gathering in. */
 #define FAST_GATHERING_COC_ERROR 0.05
diff --git a/Engine/Shaders/Private/FXAAShader.usf b/Engine/Shaders/Private/FXAAShader.usf
--- a/Engine/Shaders/Private/FXAAShader.usf
+++ b/Engine/Shaders/Private/FXAAShader.usf
@@ -98,7 +98,7 @@ SamplerState Input_Sampler;
 			fxaaConsoleEdgeThresholdMin,
 			fxaaConsole360ConstDir);
 
-		#if (POST_PROCESS_ALPHA != 2)
+		#if DIM_ALPHA_CHANNEL == 0
 			OutColor.a = 1.0;
 		#endif
 	}
diff --git a/Engine/Shaders/Private/MotionBlur/MotionBlurApply.usf b/Engine/Shaders/Private/MotionBlur/MotionBlurApply.usf
--- a/Engine/Shaders/Private/MotionBlur/MotionBlurApply.usf
+++ b/Engine/Shaders/Private/MotionBlur/MotionBlurApply.usf
@@ -65,6 +65,8 @@
 	#define CONFIG_SKIP_CENTER 0
 #endif
 
+#define CONFIG_SCENE_COLOR_ALPHA (DIM_ALPHA_CHANNEL)
+
 
 //------------------------------------------------------- CONSTANTS
 
@@ -132,7 +134,7 @@ RWTexture2D<float4> DebugOutput;
 
 groupshared mb_half4 SharedArray0[TILE_SIZE * TILE_SIZE];
 
-#if POST_PROCESS_ALPHA
+#if CONFIG_SCENE_COLOR_ALPHA
 groupshared mb_half  SharedArray1[TILE_SIZE * TILE_SIZE];
 #endif
 
@@ -821,7 +823,7 @@ void MainCS(
 		}
 		#elif CONFIG_SHUFFLE_HALF_RES && !PLATFORM_SUPPORTS_WAVE_BROADCAST
 		{
-			#if POST_PROCESS_ALPHA
+			#if CONFIG_SCENE_COLOR_ALPHA
 				SharedArray0[GroupThreadIndex] = MotionBlurColor;
 				SharedArray1[GroupThreadIndex] = FullResBlend;
 			#else
@@ -833,7 +835,7 @@ void MainCS(
 			MotionBlurColor2 = SharedArray0[(GroupThreadIndex & (~0x3)) | 0x2];
 			MotionBlurColor3 = SharedArray0[(GroupThreadIndex & (~0x3)) | 0x3];
 		
-			#if POST_PROCESS_ALPHA
+			#if CONFIG_SCENE_COLOR_ALPHA
 				FullResBlend0 = SharedArray1[(GroupThreadIndex & (~0x3)) | 0x0];
 				FullResBlend1 = SharedArray1[(GroupThreadIndex & (~0x3)) | 0x1];
 				FullResBlend2 = SharedArray1[(GroupThreadIndex & (~0x3)) | 0x2];
@@ -887,7 +889,7 @@ void MainCS(
 		
 		// Ensure that alpha values that are expected to be opaque (but are only close to opaque) are forced to be opaque.
 		// (0.995 chosen to accommodate handling of 254/255)
-		#if POST_PROCESS_ALPHA
+		#if CONFIG_SCENE_COLOR_ALPHA
 		{
 			OutputColor0[3] = select(OutputColor0[3] > mb_half(0.995), mb_half(1.0), OutputColor0[3]);
 			OutputColor1[3] = select(OutputColor1[3] > mb_half(0.995), mb_half(1.0), OutputColor1[3]);
@@ -1011,7 +1013,7 @@ void MainCS(
 
 		// Ensure that alpha values that are expected to be opaque (but are only close to opaque) are forced to be opaque.
 		// (0.995 chosen to accommodate handling of 254/255)
-		#if POST_PROCESS_ALPHA
+		#if CONFIG_SCENE_COLOR_ALPHA
 		{
 			OutputColor[3] = select(OutputColor[3] > mb_half(0.995), mb_half(1.0), OutputColor[3]);
 			OutputColor[3] = select(OutputColor[3] < mb_half(0.005), mb_half(0.0), OutputColor[3]);
diff --git a/Engine/Shaders/Private/PostProcessCompositePrimitives.usf b/Engine/Shaders/Private/PostProcessCompositePrimitives.usf
--- a/Engine/Shaders/Private/PostProcessCompositePrimitives.usf
+++ b/Engine/Shaders/Private/PostProcessCompositePrimitives.usf
@@ -97,6 +97,7 @@ FScreenTransform ColorToDepth;
 
 float2 DepthTextureJitter;
 uint bCameraCut;
+uint bProcessAlpha;
 
 float4 SampleOffsetArray[MSAA_SAMPLE_COUNT];
 
@@ -402,9 +403,13 @@ void MainCompositeEditorPrimitivesPS(
 	// Blend editor ptimitives with scene color.
 	OutColor.rgb = SceneColor.rgb * (1 - EditorPrimitiveColor.a) + EditorPrimitiveColor.rgb * (DarkenMask * HiddenMask);
 
-	#if POST_PROCESS_ALPHA
+	FLATTEN
+	if (bProcessAlpha)
+	{
 		OutColor.a = lerp(SceneColor.a, 1, EditorPrimitiveColor.a * DarkenMask * HiddenMask);
-	#else
+	}
+	else
+	{
 		OutColor.a = 0;
-	#endif
+	}
 }
diff --git a/Engine/Shaders/Private/PostProcessMitchellNetravali.usf b/Engine/Shaders/Private/PostProcessMitchellNetravali.usf
--- a/Engine/Shaders/Private/PostProcessMitchellNetravali.usf
+++ b/Engine/Shaders/Private/PostProcessMitchellNetravali.usf
@@ -144,7 +144,12 @@ void DownsampleMainCS(uint2 DispatchThreadId : SV_DispatchThreadID)
 	}
 
 	// Remove unecessary ALU work for alpha channel if unsupported on project.
-	#if POST_PROCESS_ALPHA == 0
+	#if DIM_ALPHA_CHANNEL
+	{
+		OutColor[3] = select(OutColor[3] > 0.995, 1.0, OutColor[3]);
+		OutColor[3] = select(OutColor[3] < 0.005, 0.0, OutColor[3]);
+	}
+	#else
 		OutColor.a = 0;
 	#endif
 
diff --git a/Engine/Shaders/Private/PostProcessUpscale.usf b/Engine/Shaders/Private/PostProcessUpscale.usf
--- a/Engine/Shaders/Private/PostProcessUpscale.usf
+++ b/Engine/Shaders/Private/PostProcessUpscale.usf
@@ -290,7 +290,7 @@ void MainPS(noperspective float4 UVAndScreenPos : TEXCOORD0, float4 SvPosition :
 	}
 #endif
 
-#if POST_PROCESS_ALPHA
+#if defined(DIM_ALPHA_CHANNEL) && DIM_ALPHA_CHANNEL == 1
 	FLATTEN
 	if (OutColor.a > 0.999)
 	{
diff --git a/Engine/Shaders/Private/TemporalAA.usf b/Engine/Shaders/Private/TemporalAA.usf
--- a/Engine/Shaders/Private/TemporalAA.usf
+++ b/Engine/Shaders/Private/TemporalAA.usf
@@ -67,7 +67,7 @@
 #endif
 
 #if TAA_PASS_CONFIG == 0 // Main
-	#if POST_PROCESS_ALPHA
+	#if TAA_ALPHA_CHANNEL
 		#define AA_HISTORY_PAYLOAD (HISTORY_PAYLOAD_RGB_TRANSLUCENCY)
 	#else
 		#define AA_HISTORY_PAYLOAD (HISTORY_PAYLOAD_RGB)
@@ -126,7 +126,7 @@
 	#define AA_LOWER_RESOLUTION 1
 
 #elif TAA_PASS_CONFIG == 1 || TAA_PASS_CONFIG == 2 // MainUpsampling & MainSuperSampling
-	#if POST_PROCESS_ALPHA
+	#if TAA_ALPHA_CHANNEL
 		#define AA_HISTORY_PAYLOAD (HISTORY_PAYLOAD_RGB_TRANSLUCENCY)
 	#else
 		#define AA_HISTORY_PAYLOAD (HISTORY_PAYLOAD_RGB)
@@ -187,7 +187,7 @@
 	#endif
 	
 #elif TAA_PASS_CONFIG == 5 || TAA_PASS_CONFIG == 6 // Used for diaphragm DOF pre-filtering.
-	#if POST_PROCESS_ALPHA
+	#if TAA_ALPHA_CHANNEL
 		#define AA_HISTORY_PAYLOAD (HISTORY_PAYLOAD_RGB_OPACITY_COC)
 	#else
 		#define AA_HISTORY_PAYLOAD (HISTORY_PAYLOAD_RGB_COC)
diff --git a/Engine/Source/Runtime/Renderer/Private/PostProcess/DiaphragmDOF.cpp b/Engine/Source/Runtime/Renderer/Private/PostProcess/DiaphragmDOF.cpp
--- a/Engine/Source/Runtime/Renderer/Private/PostProcess/DiaphragmDOF.cpp
+++ b/Engine/Source/Runtime/Renderer/Private/PostProcess/DiaphragmDOF.cpp
@@ -420,6 +420,8 @@ const FVector2f kContantlyBlockingAffineTransformation(0, 0);
 class FDiaphragmDOFShader : public FGlobalShader
 {
 public:
+	class FAlphaChannelDim : SHADER_PERMUTATION_BOOL("DIM_ALPHA_CHANNEL");
+
 	static bool ShouldCompilePermutation(const FGlobalShaderPermutationParameters& Parameters)
 	{
 		return DiaphragmDOF::IsSupported(Parameters.Platform);
@@ -722,7 +724,7 @@ class FDiaphragmDOFSetupCS : public FDiaphragmDOFShader
 
 	class FOutputResDivisor : SHADER_PERMUTATION_INT("DIM_OUTPUT_RES_DIVISOR", 3);
 
-	using FPermutationDomain = TShaderPermutationDomain<FOutputResDivisor>;
+	using FPermutationDomain = TShaderPermutationDomain<FDiaphragmDOFShader::FAlphaChannelDim, FOutputResDivisor>;
 
 	BEGIN_SHADER_PARAMETER_STRUCT(FParameters, )
 		SHADER_PARAMETER_STRUCT_INCLUDE(FDOFCommonShaderParameters, CommonParameters)
@@ -747,7 +749,7 @@ class FDiaphragmDOFCocFlattenCS : public FDiaphragmDOFShader
 
 	class FDoCocGather4 : SHADER_PERMUTATION_BOOL("DIM_DO_COC_GATHER4");
 
-	using FPermutationDomain = TShaderPermutationDomain<FDoCocGather4>;
+	using FPermutationDomain = TShaderPermutationDomain<FDiaphragmDOFShader::FAlphaChannelDim, FDoCocGather4>;
 	
 	BEGIN_SHADER_PARAMETER_STRUCT(FParameters, )
 		SHADER_PARAMETER(FIntRect, ViewportRect)
@@ -799,7 +801,9 @@ class FDiaphragmDOFDownsampleCS : public FDiaphragmDOFShader
 		}
 		return FDiaphragmDOFShader::ShouldCompilePermutation(Parameters);
 	}
-	
+
+	using FPermutationDomain = TShaderPermutationDomain<FDiaphragmDOFShader::FAlphaChannelDim>;
+
 	BEGIN_SHADER_PARAMETER_STRUCT(FParameters, )
 		SHADER_PARAMETER(FIntRect, ViewportRect)
 		SHADER_PARAMETER(FVector2f, MaxBufferUV)
@@ -824,6 +828,7 @@ class FDiaphragmDOFReduceCS : public FDiaphragmDOFShader
 	class FDisableOutputMip0 : SHADER_PERMUTATION_BOOL("DIM_DISABLE_OUTPUT_MIP0");
 
 	using FPermutationDomain = TShaderPermutationDomain<
+		FDiaphragmDOFShader::FAlphaChannelDim,
 		FReduceMipCount,
 		FHybridScatterForeground,
 		FHybridScatterBackground,
@@ -976,6 +981,7 @@ class FDiaphragmDOFGatherCS : public FDiaphragmDOFShader
 	SHADER_USE_PARAMETER_STRUCT(FDiaphragmDOFGatherCS, FDiaphragmDOFShader);
 
 	using FPermutationDomain = TShaderPermutationDomain<
+		FDiaphragmDOFShader::FAlphaChannelDim,
 		FDDOFLayerProcessingDim,
 		FDDOFGatherRingCountDim,
 		FDDOFBokehSimulationDim,
@@ -1141,7 +1147,11 @@ class FDiaphragmDOFPostfilterCS : public FDiaphragmDOFShader
 
 	class FTileOptimization : SHADER_PERMUTATION_BOOL("DIM_TILE_PERMUTATION");
 
-	using FPermutationDomain = TShaderPermutationDomain<FDDOFLayerProcessingDim, FDDOFPostfilterMethodDim, FTileOptimization>;
+	using FPermutationDomain = TShaderPermutationDomain<
+		FDiaphragmDOFShader::FAlphaChannelDim,
+		FDDOFLayerProcessingDim,
+		FDDOFPostfilterMethodDim,
+		FTileOptimization>;
 
 	static FPermutationDomain RemapPermutationVector(FPermutationDomain PermutationVector)
 	{
@@ -1267,7 +1277,11 @@ class FDiaphragmDOFRecombineCS : public FDiaphragmDOFShader
 
 	class FQualityDim : SHADER_PERMUTATION_INT("DIM_QUALITY", 3);
 
-	using FPermutationDomain = TShaderPermutationDomain<FDDOFLayerProcessingDim, FDDOFBokehSimulationDim, FQualityDim>;
+	using FPermutationDomain = TShaderPermutationDomain<
+		FDiaphragmDOFShader::FAlphaChannelDim,
+		FDDOFLayerProcessingDim,
+		FDDOFBokehSimulationDim,
+		FQualityDim>;
 
 	static bool ShouldCompilePermutation(const FGlobalShaderPermutationParameters& Parameters)
 	{
@@ -1575,6 +1589,7 @@ FRDGTextureRef DiaphragmDOF::AddPasses(
 		bool bOutputHalfResolution = PrefilteringResolutionDivisor == 2;
 		
 		FDiaphragmDOFSetupCS::FPermutationDomain PermutationVector;
+		PermutationVector.Set<FDiaphragmDOFShader::FAlphaChannelDim>(bProcessSceneAlpha);
 	
 		FIntPoint PassViewSize = FullResViewSize;
 		FIntPoint GroupSize(kDefaultGroupSize, kDefaultGroupSize);
@@ -1717,6 +1732,7 @@ FRDGTextureRef DiaphragmDOF::AddPasses(
 			PassParameters->TileOutput = CreateUAVs(GraphBuilder, FlattenedTileClassificationTextures);
 			
 			FDiaphragmDOFCocFlattenCS::FPermutationDomain PermutationVector;
+			PermutationVector.Set<FDiaphragmDOFShader::FAlphaChannelDim>(bProcessSceneAlpha);
 			PermutationVector.Set<FDiaphragmDOFCocFlattenCS::FDoCocGather4>(PreprocessViewSize != GatheringViewSize);
 
 			TShaderMapRef<FDiaphragmDOFCocFlattenCS> ComputeShader(View.ShaderMap, PermutationVector);
@@ -1932,6 +1948,9 @@ FRDGTextureRef DiaphragmDOF::AddPasses(
 
 			FIntPoint PassViewSize = FIntPoint::DivideAndRoundUp(PreprocessViewSize, 2);
 
+			FDiaphragmDOFDownsampleCS::FPermutationDomain PermutationVector;
+			PermutationVector.Set<FDiaphragmDOFShader::FAlphaChannelDim>(bProcessSceneAlpha);
+
 			FDiaphragmDOFDownsampleCS::FParameters* PassParameters = GraphBuilder.AllocParameters<FDiaphragmDOFDownsampleCS::FParameters>();
 			PassParameters->CommonParameters = CommonParameters;
 			PassParameters->ViewportRect = FIntRect(0, 0, PassViewSize.X, PassViewSize.Y);
@@ -1945,7 +1964,7 @@ FRDGTextureRef DiaphragmDOF::AddPasses(
 
 			PassParameters->OutDownsampledGatherInput = CreateUAVs(GraphBuilder, QuarterResGatherInputTextures);
 			
-			TShaderMapRef<FDiaphragmDOFDownsampleCS> ComputeShader(View.ShaderMap);
+			TShaderMapRef<FDiaphragmDOFDownsampleCS> ComputeShader(View.ShaderMap, PermutationVector);
 			FComputeShaderUtils::AddPass(
 				GraphBuilder,
 				RDG_EVENT_NAME("DOF Downsample %dx%d", PassViewSize.X, PassViewSize.Y),
@@ -1976,6 +1995,7 @@ FRDGTextureRef DiaphragmDOF::AddPasses(
 			FIntPoint PassViewSize = PreprocessViewSize;
 
 			FDiaphragmDOFReduceCS::FPermutationDomain PermutationVector;
+			PermutationVector.Set<FDiaphragmDOFShader::FAlphaChannelDim>(bProcessSceneAlpha);
 			PermutationVector.Set<FDiaphragmDOFReduceCS::FReduceMipCount>(ProcessingMipLevelCount);
 			PermutationVector.Set<FDiaphragmDOFReduceCS::FHybridScatterForeground>(bForegroundHybridScattering);
 			PermutationVector.Set<FDiaphragmDOFReduceCS::FHybridScatterBackground>(bBackgroundHybridScattering);
@@ -2259,6 +2279,7 @@ FRDGTextureRef DiaphragmDOF::AddPasses(
 			FIntPoint SrcSize = ReducedGatherInputTextures.SceneColor->Desc.Extent;
 
 			FDiaphragmDOFGatherCS::FPermutationDomain PermutationVector;
+			PermutationVector.Set<FDiaphragmDOFShader::FAlphaChannelDim>(bProcessSceneAlpha);
 			PermutationVector.Set<FDDOFLayerProcessingDim>(ConvolutionSettings.LayerProcessing);
 			PermutationVector.Set<FDDOFGatherRingCountDim>(ConvolutionSettings.RingCount);
 			PermutationVector.Set<FDDOFGatherQualityDim>(ConvolutionSettings.QualityConfig);
@@ -2396,6 +2417,7 @@ FRDGTextureRef DiaphragmDOF::AddPasses(
 			}
 
 			FDiaphragmDOFPostfilterCS::FPermutationDomain PermutationVector;
+			PermutationVector.Set<FDiaphragmDOFShader::FAlphaChannelDim>(bProcessSceneAlpha);
 			PermutationVector.Set<FDDOFLayerProcessingDim>(ConvolutionSettings.LayerProcessing);
 			PermutationVector.Set<FDDOFPostfilterMethodDim>(ConvolutionSettings.PostfilterMethod);
 			PermutationVector.Set<FDiaphragmDOFPostfilterCS::FTileOptimization>(true); // TODO
@@ -2645,6 +2667,7 @@ FRDGTextureRef DiaphragmDOF::AddPasses(
 		FIntRect PassViewRect = View.ViewRect;
 
 		FDiaphragmDOFRecombineCS::FPermutationDomain PermutationVector;
+		PermutationVector.Set<FDiaphragmDOFShader::FAlphaChannelDim>(bProcessSceneAlpha);
 		if (bGatherForeground && bGatherBackground)
 		{
 			PermutationVector.Set<FDDOFLayerProcessingDim>(EDiaphragmDOFLayerProcessing::ForegroundAndBackground);
diff --git a/Engine/Source/Runtime/Renderer/Private/PostProcess/PostProcessAA.cpp b/Engine/Source/Runtime/Renderer/Private/PostProcess/PostProcessAA.cpp
--- a/Engine/Source/Runtime/Renderer/Private/PostProcess/PostProcessAA.cpp
+++ b/Engine/Source/Runtime/Renderer/Private/PostProcess/PostProcessAA.cpp
@@ -26,8 +26,9 @@ public:
 		return true;
 	}
 
+	class FAlphaChannelDim : SHADER_PERMUTATION_BOOL("DIM_ALPHA_CHANNEL");
 	class FQualityDimension : SHADER_PERMUTATION_ENUM_CLASS("FXAA_PRESET", EFXAAQuality);
-	using FPermutationDomain = TShaderPermutationDomain<FQualityDimension>;
+	using FPermutationDomain = TShaderPermutationDomain<FAlphaChannelDim, FQualityDimension>;
 
 	BEGIN_SHADER_PARAMETER_STRUCT(FParameters, )
 		SHADER_PARAMETER_STRUCT(FScreenPassTextureInput, Input)
@@ -98,6 +99,7 @@ FScreenPassTexture AddFXAAPass(FRDGBuilder& GraphBuilder, const FSceneView& InSc
 	PassParameters->fxaaConsoleEdgeThresholdMin = 0.05f;
 
 	FFXAAPS::FPermutationDomain PixelPermutationVector;
+	PixelPermutationVector.Set<FFXAAPS::FAlphaChannelDim>(IsPostProcessingWithAlphaChannelSupported());
 	PixelPermutationVector.Set<FFXAAPS::FQualityDimension>(Inputs.Quality);
 
 	TShaderMapRef<FFXAAPS> PixelShader(View.ShaderMap, PixelPermutationVector);
@@ -105,8 +107,11 @@ FScreenPassTexture AddFXAAPass(FRDGBuilder& GraphBuilder, const FSceneView& InSc
 	FPixelShaderUtils::AddFullscreenPass(
 		GraphBuilder,
 		View.ShaderMap,
-		RDG_EVENT_NAME("FXAA(Quality=%d) %dx%d PS",
-			Inputs.Quality, Output.ViewRect.Width(), Output.ViewRect.Height()),
+		RDG_EVENT_NAME("FXAA(Quality=%d%s) %dx%d PS",
+			Inputs.Quality,
+			PixelPermutationVector.Get<FFXAAPS::FAlphaChannelDim>() ? TEXT(" Alpha") : TEXT(""),
+			Output.ViewRect.Width(),
+			Output.ViewRect.Height()),
 		PixelShader,
 		PassParameters,
 		Output.ViewRect);
diff --git a/Engine/Source/Runtime/Renderer/Private/PostProcess/PostProcessCompositeEditorPrimitives.cpp b/Engine/Source/Runtime/Renderer/Private/PostProcess/PostProcessCompositeEditorPrimitives.cpp
--- a/Engine/Source/Runtime/Renderer/Private/PostProcess/PostProcessCompositeEditorPrimitives.cpp
+++ b/Engine/Source/Runtime/Renderer/Private/PostProcess/PostProcessCompositeEditorPrimitives.cpp
@@ -9,6 +9,7 @@
 #include "MobileBasePassRendering.h"
 #include "PixelShaderUtils.h"
 #include "Substrate/Substrate.h"
+#include "PostProcessing.h"
 
 namespace
 {
@@ -33,10 +34,11 @@ namespace
 			SHADER_PARAMETER(uint32, bOpaqueEditorGizmo)
 			SHADER_PARAMETER(uint32, bCompositeAnyNonNullDepth)
 			SHADER_PARAMETER(FVector2f, DepthTextureJitter)
+			SHADER_PARAMETER(uint32, bProcessAlpha)
 			RENDER_TARGET_BINDING_SLOTS()
-			END_SHADER_PARAMETER_STRUCT()
+		END_SHADER_PARAMETER_STRUCT()
 
-			static void ModifyCompilationEnvironment(const FGlobalShaderPermutationParameters& Parameters, FShaderCompilerEnvironment& OutEnvironment)
+		static void ModifyCompilationEnvironment(const FGlobalShaderPermutationParameters& Parameters, FShaderCompilerEnvironment& OutEnvironment)
 		{
 			FGlobalShader::ModifyCompilationEnvironment(Parameters, OutEnvironment);
 		}
@@ -388,6 +390,7 @@ FScreenPassTexture AddEditorPrimitivePass(
 		PassParameters->bOpaqueEditorGizmo = bOpaqueEditorGizmo;
 		PassParameters->bCompositeAnyNonNullDepth = bProducedByPriorPass;
 		PassParameters->DepthTextureJitter = SceneDepthJitter;
+		PassParameters->bProcessAlpha = IsPostProcessingWithAlphaChannelSupported();
 
 		for (int32 i = 0; i < int32(NumMSAASamples); i++)
 		{
diff --git a/Engine/Source/Runtime/Renderer/Private/PostProcess/PostProcessFFTBloom.cpp b/Engine/Source/Runtime/Renderer/Private/PostProcess/PostProcessFFTBloom.cpp
--- a/Engine/Source/Runtime/Renderer/Private/PostProcess/PostProcessFFTBloom.cpp
+++ b/Engine/Source/Runtime/Renderer/Private/PostProcess/PostProcessFFTBloom.cpp
@@ -9,6 +9,7 @@
 #include "RendererModule.h"
 #include "Rendering/Texture2DResource.h"
 #include "ScenePrivate.h"
+#include "PostProcessing.h"
 
 namespace
 {
@@ -153,6 +154,7 @@ public:
 	SHADER_USE_PARAMETER_STRUCT(FBloomClampKernelCS, FFFTBloomShader);
 
 	BEGIN_SHADER_PARAMETER_STRUCT(FParameters, )
+		SHADER_PARAMETER(int32, bProcessAlpha)
 		SHADER_PARAMETER_RDG_TEXTURE(Texture2D, KernelSpatialTexture)
 		SHADER_PARAMETER_RDG_BUFFER_SRV(StructuredBuffer<uint>, KernelConstantsBuffer)
 		SHADER_PARAMETER_RDG_TEXTURE_UAV(RWTexture2D<float4>, ClampedKernelSpatialOutput)
@@ -582,6 +584,7 @@ void InitDomainAndGetKernel(
 			}
 
 			FBloomClampKernelCS::FParameters* PassParameters = GraphBuilder.AllocParameters<FBloomClampKernelCS::FParameters>();
+			PassParameters->bProcessAlpha = IsPostProcessingWithAlphaChannelSupported();
 			PassParameters->KernelSpatialTexture = SpatialKernelTexture;
 			PassParameters->KernelConstantsBuffer = GraphBuilder.CreateSRV(KernelConstantsBuffer);
 			PassParameters->ClampedKernelSpatialOutput = GraphBuilder.CreateUAV(ClampedKernelTexture);
diff --git a/Engine/Source/Runtime/Renderer/Private/PostProcess/PostProcessMitchellNetravali.cpp b/Engine/Source/Runtime/Renderer/Private/PostProcess/PostProcessMitchellNetravali.cpp
--- a/Engine/Source/Runtime/Renderer/Private/PostProcess/PostProcessMitchellNetravali.cpp
+++ b/Engine/Source/Runtime/Renderer/Private/PostProcess/PostProcessMitchellNetravali.cpp
@@ -6,6 +6,7 @@
 #include "SceneRendering.h"
 #include "ScenePrivate.h"
 #include "SceneTextureParameters.h"
+#include "PostProcessing.h"
 
 class FMitchellNetravaliDownsampleCS : public FGlobalShader
 {
@@ -13,6 +14,9 @@ public:
 	DECLARE_GLOBAL_SHADER(FMitchellNetravaliDownsampleCS);
 	SHADER_USE_PARAMETER_STRUCT(FMitchellNetravaliDownsampleCS, FGlobalShader);
 
+	class FAlphaChannelDim : SHADER_PERMUTATION_BOOL("DIM_ALPHA_CHANNEL");
+	using FPermutationDomain = TShaderPermutationDomain<FAlphaChannelDim>;
+
 	static bool ShouldCompilePermutation(const FGlobalShaderPermutationParameters& Parameters)
 	{
 		return IsFeatureLevelSupported(Parameters.Platform, ERHIFeatureLevel::SM5) && !IsOpenGLPlatform(Parameters.Platform);
@@ -60,11 +64,19 @@ FRDGTextureRef ComputeMitchellNetravaliDownsample(
 	PassParameters->DispatchThreadToInputUVBias.X = PassParameters->DispatchThreadToInputUVScale.X * (0.5f + Input.ViewRect.Min.X);
 	PassParameters->DispatchThreadToInputUVBias.Y = PassParameters->DispatchThreadToInputUVScale.Y * (0.5f + Input.ViewRect.Min.Y);
 
-	TShaderMapRef<FMitchellNetravaliDownsampleCS> ComputeShader(View.ShaderMap);
+	FMitchellNetravaliDownsampleCS::FPermutationDomain PermutationVector;
+	PermutationVector.Set<FMitchellNetravaliDownsampleCS::FAlphaChannelDim>(IsPostProcessingWithAlphaChannelSupported());
+
+	TShaderMapRef<FMitchellNetravaliDownsampleCS> ComputeShader(View.ShaderMap, PermutationVector);
 
 	FComputeShaderUtils::AddPass(
 		GraphBuilder,
-		RDG_EVENT_NAME("MitchellNetravaliDownsample %dx%d -> %dx%d", Input.ViewRect.Width(), Input.ViewRect.Height(), OutputViewport.Rect.Width(), OutputViewport.Rect.Height()),
+		RDG_EVENT_NAME("MitchellNetravaliDownsample(%s) %dx%d -> %dx%d",
+			PermutationVector.Get<FMitchellNetravaliDownsampleCS::FAlphaChannelDim>() ? TEXT("Alpha") : TEXT(""),
+			Input.ViewRect.Width(),
+			Input.ViewRect.Height(),
+			OutputViewport.Rect.Width(),
+			OutputViewport.Rect.Height()),
 		ComputeShader,
 		PassParameters,
 		FComputeShaderUtils::GetGroupCount(OutputViewport.Rect.Size(), FComputeShaderUtils::kGolden2DGroupSize));
diff --git a/Engine/Source/Runtime/Renderer/Private/PostProcess/PostProcessMotionBlur.cpp b/Engine/Source/Runtime/Renderer/Private/PostProcess/PostProcessMotionBlur.cpp
--- a/Engine/Source/Runtime/Renderer/Private/PostProcess/PostProcessMotionBlur.cpp
+++ b/Engine/Source/Runtime/Renderer/Private/PostProcess/PostProcessMotionBlur.cpp
@@ -428,10 +428,11 @@ class FMotionBlurFilterCS : public FMotionBlurShader
 	{
 		return TileClassification == ETileClassification::GatherHalfRes || TileClassification == ETileClassification::ScatterAsGatherOneVelocityHalfRes;
 	}
-	
+
+	class FAlphaChannelDim : SHADER_PERMUTATION_BOOL("DIM_ALPHA_CHANNEL");
 	class FTileClassificationDim : SHADER_PERMUTATION_ENUM_CLASS("DIM_TILE_CLASSIFICATION", ETileClassification);
 
-	using FPermutationDomain = TShaderPermutationDomain<FMotionBlurDirections, FTileClassificationDim>;
+	using FPermutationDomain = TShaderPermutationDomain<FMotionBlurDirections, FTileClassificationDim, FAlphaChannelDim>;
 
 	BEGIN_SHADER_PARAMETER_STRUCT(FParameters, )
 		SHADER_PARAMETER_STRUCT(FScreenPassTextureViewportParameters, Color)
@@ -869,6 +870,8 @@ FMotionBlurOutputs AddMotionBlurFilterPass(
 
 	const FIntVector FilterTileCount = FComputeShaderUtils::GetGroupCount(Viewports.Color.Rect.Size(), kMotionBlurFilterTileSize);
 
+	const bool bAlphaChannel = IsPostProcessingWithAlphaChannelSupported();
+	
 	int32 TileListMaxSize = FilterTileCount.X * FilterTileCount.Y;
 
 	// Tile classify the filtering
@@ -995,7 +998,7 @@ FMotionBlurOutputs AddMotionBlurFilterPass(
 		{
 			FRDGTextureDesc Desc = FRDGTextureDesc::Create2D(
 				ColorTexture->Desc.Texture->Desc.Extent,
-				IsPostProcessingWithAlphaChannelSupported() ? PF_FloatRGBA : PF_FloatRGB,
+				bAlphaChannel ? PF_FloatRGBA : PF_FloatRGB,
 				FClearValueBinding::None,
 				TexCreate_UAV | TexCreate_ShaderResource | GFastVRamConfig.MotionBlur);
 
@@ -1060,8 +1063,9 @@ FMotionBlurOutputs AddMotionBlurFilterPass(
 
 		OriginalPassParameters.DebugOutput = CreateDebugUAV(GraphBuilder, Output.FullRes.TextureSRV->Desc.Texture->Desc.Extent, TEXT("Debug.MotionBlur.Filter"));
 
-		RDG_EVENT_SCOPE(GraphBuilder, "MotionBlur FullResFilter(BlurDirections=%d MaxSamples=%d%s%s%s) %dx%d",
+		RDG_EVENT_SCOPE(GraphBuilder, "MotionBlur FullResFilter(BlurDirections=%d MaxSamples=%d%s%s%s%s) %dx%d",
 			BlurDirections, OriginalPassParameters.MaxSampleCount,
+			bAlphaChannel ? TEXT(" AlphaChannel") : TEXT(""),
 			PostMotionBlurTranslucency ? TEXT(" ComposeTranslucency") : TEXT(""),
 			OriginalPassParameters.OutputMip1 ? TEXT(" OutputMip1") : TEXT(""),
 			OriginalPassParameters.OutputMip2 ? TEXT(" OutputMip2") : TEXT(""),
@@ -1074,6 +1078,7 @@ FMotionBlurOutputs AddMotionBlurFilterPass(
 			PassParameters->TileListOffset = TileListMaxSize * TileClassifcation;
 
 			FMotionBlurFilterCS::FPermutationDomain PermutationVector;
+			PermutationVector.Set<FMotionBlurFilterCS::FAlphaChannelDim>(bAlphaChannel);
 			PermutationVector.Set<FMotionBlurFilterCS::FTileClassificationDim>(FMotionBlurFilterCS::ETileClassification(TileClassifcation));
 			PermutationVector.Set<FMotionBlurDirections>(BlurDirections);
 			PermutationVector = FMotionBlurFilterCS::RemapPermutation(PermutationVector);
diff --git a/Engine/Source/Runtime/Renderer/Private/PostProcess/PostProcessUpscale.cpp b/Engine/Source/Runtime/Renderer/Private/PostProcess/PostProcessUpscale.cpp
--- a/Engine/Source/Runtime/Renderer/Private/PostProcess/PostProcessUpscale.cpp
+++ b/Engine/Source/Runtime/Renderer/Private/PostProcess/PostProcessUpscale.cpp
@@ -4,6 +4,7 @@
 #include "PostProcess/SceneFilterRendering.h"
 #include "DataDrivenShaderPlatformInfo.h"
 #include "SceneRendering.h"
+#include "PostProcessing.h"
 
 namespace
 {
@@ -123,8 +124,9 @@ public:
 	SHADER_USE_PARAMETER_STRUCT(FUpscalePS, FGlobalShader);
 	using FParameters = FUpscaleParameters;
 
+	class FAlphaChannelDim : SHADER_PERMUTATION_BOOL("DIM_ALPHA_CHANNEL");
 	class FMethodDimension : SHADER_PERMUTATION_ENUM_CLASS("METHOD", EUpscaleMethod);
-	using FPermutationDomain = TShaderPermutationDomain<FMethodDimension>;
+	using FPermutationDomain = TShaderPermutationDomain<FAlphaChannelDim, FMethodDimension>;
 
 	static bool ShouldCompilePermutation(const FGlobalShaderPermutationParameters& Parameters)
 	{
@@ -228,6 +230,7 @@ FScreenPassTexture ISpatialUpscaler::AddDefaultUpscalePass(
 	PassParameters->View = View.ViewUniformBuffer;
 
 	FUpscalePS::FPermutationDomain PixelPermutationVector;
+	PixelPermutationVector.Set<FUpscalePS::FAlphaChannelDim>(IsPostProcessingWithAlphaChannelSupported());
 	PixelPermutationVector.Set<FUpscalePS::FMethodDimension>(Method);
 	TShaderMapRef<FUpscalePS> PixelShader(View.ShaderMap, PixelPermutationVector);
 
@@ -236,7 +239,14 @@ FScreenPassTexture ISpatialUpscaler::AddDefaultUpscalePass(
 	const TCHAR* StageName = StageNames[static_cast<uint32>(Inputs.Stage)];
 
 	GraphBuilder.AddPass(
-		RDG_EVENT_NAME("Upscale(%s Method=%d) %dx%d -> %dx%d", StageName, int32(Method), Inputs.SceneColor.ViewRect.Width(), Inputs.SceneColor.ViewRect.Height(), Output.ViewRect.Width(), Output.ViewRect.Height()),
+		RDG_EVENT_NAME("Upscale(%s Method=%d%s) %dx%d -> %dx%d",
+			StageName,
+			int32(Method),
+			PixelPermutationVector.Get<FUpscalePS::FAlphaChannelDim>() ? TEXT(" Alpha") : TEXT(""),
+			Inputs.SceneColor.ViewRect.Width(),
+			Inputs.SceneColor.ViewRect.Height(),
+			Output.ViewRect.Width(),
+			Output.ViewRect.Height()),
 		PassParameters,
 		ERDGPassFlags::Raster,
 		[&View, bUsePaniniProjection, PixelShader, PassParameters, InputViewport, OutputViewport](FRHICommandList& RHICmdList)
diff --git a/Engine/Source/Runtime/Renderer/Private/PostProcess/TemporalAA.cpp b/Engine/Source/Runtime/Renderer/Private/PostProcess/TemporalAA.cpp
--- a/Engine/Source/Runtime/Renderer/Private/PostProcess/TemporalAA.cpp
+++ b/Engine/Source/Runtime/Renderer/Private/PostProcess/TemporalAA.cpp
@@ -98,6 +98,7 @@ inline bool DoesPlatformSupportTemporalHistoryUpscale(EShaderPlatform Platform)
 class FTemporalAA : public FGlobalShader
 {
 public:
+	class FAlphaChannelDim : SHADER_PERMUTATION_BOOL("TAA_ALPHA_CHANNEL");
 	class FTAAPassConfigDim : SHADER_PERMUTATION_ENUM_CLASS("TAA_PASS_CONFIG", ETAAPassConfig);
 	class FTAAQualityDim : SHADER_PERMUTATION_ENUM_CLASS("TAA_QUALITY", ETAAQuality);
 	class FTAAScreenPercentageDim : SHADER_PERMUTATION_INT("TAA_SCREEN_PERCENTAGE_RANGE", 4);
@@ -178,6 +179,7 @@ class FTemporalAAPS : public FTemporalAA
 	SHADER_USE_PARAMETER_STRUCT(FTemporalAAPS, FTemporalAA);
 
 	using FPermutationDomain = TShaderPermutationDomain<
+		FTemporalAA::FAlphaChannelDim,
 		FTemporalAA::FTAAPassConfigDim,
 		FTemporalAA::FTAAQualityDim,
 		FTemporalAA::FTAAScreenPercentageDim>;
@@ -186,11 +188,36 @@ class FTemporalAAPS : public FTemporalAA
 		SHADER_PARAMETER_STRUCT_INCLUDE(FTemporalAA::FParameters, Common)
 		RENDER_TARGET_BINDING_SLOTS()
 	END_SHADER_PARAMETER_STRUCT()
+		
+	static FPermutationDomain RemapPermutation(FPermutationDomain PermutationVector)
+	{
+		if (PermutationVector.Get<FTemporalAA::FTAAPassConfigDim>() == ETAAPassConfig::Main ||
+			PermutationVector.Get<FTemporalAA::FTAAPassConfigDim>() == ETAAPassConfig::MainUpsampling ||
+			PermutationVector.Get<FTemporalAA::FTAAPassConfigDim>() == ETAAPassConfig::MainSuperSampling ||
+			PermutationVector.Get<FTemporalAA::FTAAPassConfigDim>() == ETAAPassConfig::DiaphragmDOF ||
+			PermutationVector.Get<FTemporalAA::FTAAPassConfigDim>() == ETAAPassConfig::DiaphragmDOFUpsampling)
+		{
+			// NOP
+		}
+		else
+		{
+			// Alpha channel is not supported on these permutations
+			PermutationVector.Set<FAlphaChannelDim>(false);
+		}
+
+		return PermutationVector;
+	}
 
 	static bool ShouldCompilePermutation(const FGlobalShaderPermutationParameters& Parameters)
 	{
 		FPermutationDomain PermutationVector(Parameters.PermutationId);
 
+		// Don't compile the shader permutation if gets remaped at runtime.
+		if (PermutationVector != RemapPermutation(PermutationVector))
+		{
+			return false;
+		}
+
 		// Pixel shader is only used on mobile to utilize the hardware frame buffer compression to save bandwidth
 		if (!IsMobilePlatform(Parameters.Platform))
 		{
@@ -244,6 +271,7 @@ class FTemporalAACS : public FTemporalAA
 	class FTAADownsampleDim : SHADER_PERMUTATION_BOOL("TAA_DOWNSAMPLE");
 
 	using FPermutationDomain = TShaderPermutationDomain<
+		FTemporalAA::FAlphaChannelDim,
 		FTemporalAA::FTAAPassConfigDim,
 		FTemporalAA::FTAAQualityDim,
 		FTemporalAA::FTAAScreenPercentageDim,
@@ -288,6 +316,20 @@ class FTemporalAACS : public FTemporalAA
 			PermutationVector.Set<FTAADownsampleDim>(false);
 		}
 
+		if (PermutationVector.Get<FTemporalAA::FTAAPassConfigDim>() == ETAAPassConfig::Main ||
+			PermutationVector.Get<FTemporalAA::FTAAPassConfigDim>() == ETAAPassConfig::MainUpsampling ||
+			PermutationVector.Get<FTemporalAA::FTAAPassConfigDim>() == ETAAPassConfig::MainSuperSampling ||
+			PermutationVector.Get<FTemporalAA::FTAAPassConfigDim>() == ETAAPassConfig::DiaphragmDOF ||
+			PermutationVector.Get<FTemporalAA::FTAAPassConfigDim>() == ETAAPassConfig::DiaphragmDOFUpsampling)
+		{
+			// NOP
+		}
+		else
+		{
+			// Alpha channel is not supported on these permutations
+			PermutationVector.Set<FAlphaChannelDim>(false);
+		}
+
 		return PermutationVector;
 	}
 
@@ -827,6 +869,7 @@ FTAAOutputs AddTemporalAAPass(
 	if (bIsComputePass)
 	{
 		FTemporalAACS::FPermutationDomain PermutationVector;
+		PermutationVector.Set<FTemporalAA::FAlphaChannelDim>(bSupportsAlpha);
 		PermutationVector.Set<FTemporalAA::FTAAPassConfigDim>(Inputs.Pass);
 		PermutationVector.Set<FTemporalAA::FTAAQualityDim>(Inputs.Quality);
 		PermutationVector.Set<FTemporalAACS::FTAADownsampleDim>(Inputs.bDownsample);
@@ -895,9 +938,10 @@ FTAAOutputs AddTemporalAAPass(
 
 		FComputeShaderUtils::AddPass(
 			GraphBuilder,
-			RDG_EVENT_NAME("TAA(%s Quality=%s) %dx%d -> %dx%d",
+			RDG_EVENT_NAME("TAA(%s Quality=%s%s) %dx%d -> %dx%d",
 				PassName,
 				kTAAQualityNames[int32(PermutationVector.Get<FTemporalAA::FTAAQualityDim>())],
+				PermutationVector.Get<FTemporalAA::FAlphaChannelDim>() ? TEXT(" Alpha") : TEXT(""),
 				PracticableSrcRect.Width(), PracticableSrcRect.Height(),
 				PracticableDestRect.Width(), PracticableDestRect.Height()),
 			ComputeShader,
@@ -909,6 +953,7 @@ FTAAOutputs AddTemporalAAPass(
 		check(IsMobilePlatform(View.GetShaderPlatform()));
 
 		FTemporalAAPS::FPermutationDomain PermutationVector;
+		PermutationVector.Set<FTemporalAA::FAlphaChannelDim>(bSupportsAlpha);
 		PermutationVector.Set<FTemporalAA::FTAAPassConfigDim>(Inputs.Pass);
 		PermutationVector.Set<FTemporalAA::FTAAQualityDim>(Inputs.Quality);
 
@@ -935,6 +980,8 @@ FTAAOutputs AddTemporalAAPass(
 			}
 		}
 
+		PermutationVector = FTemporalAAPS::RemapPermutation(PermutationVector);
+
 		FTemporalAAPS::FParameters* PassParameters = GraphBuilder.AllocParameters<FTemporalAAPS::FParameters>();
 
 		SetupTemporalAACommonPassParameters(&PassParameters->Common);
@@ -948,9 +995,10 @@ FTAAOutputs AddTemporalAAPass(
 		FPixelShaderUtils::AddFullscreenPass(
 			GraphBuilder,
 			View.ShaderMap,
-			RDG_EVENT_NAME("TAA(%s Quality=%s) %dx%d -> %dx%d",
+			RDG_EVENT_NAME("TAA(%s Quality=%s%s) %dx%d -> %dx%d",
 				PassName,
 				kTAAQualityNames[int32(PermutationVector.Get<FTemporalAA::FTAAQualityDim>())],
+				PermutationVector.Get<FTemporalAA::FAlphaChannelDim>() ? TEXT(" Alpha") : TEXT(""),
 				PracticableSrcRect.Width(), PracticableSrcRect.Height(),
 				PracticableDestRect.Width(), PracticableDestRect.Height()),
 			PixelShader,
