Global Distance Field - optimize culling to clipmap by sharing translated world space (camera-centric) between all clipmaps.

diff --git a/Engine/Shaders/Private/DistanceField/GlobalDistanceField.usf b/Engine/Shaders/Private/DistanceField/GlobalDistanceField.usf
--- a/Engine/Shaders/Private/DistanceField/GlobalDistanceField.usf
+++ b/Engine/Shaders/Private/DistanceField/GlobalDistanceField.usf
@@ -26,42 +26,49 @@
 	#define THREADGROUP_SIZE_Z 1
 #endif
 
+// Must match FPackedClipmapData in GlobalDistanceField.cpp
+struct FPackedClipmapData
+{
+	float3 TranslatedWorldCenter; // Camera-centered world space
+	float MeshSDFRadiusThreshold;	
+
+	float3 WorldExtent;
+	float InfluenceRadiusSq;
+
+	uint ClipmapIndex; 
+	uint Flags;
+	uint Padding0;
+	uint Padding1;	
+};
+
 struct FPackedClipmap
 {
-	float3 TranslatedWorldCenter;
+	float3 TranslatedWorldCenter; // Camera-centered world space
 	float3 WorldExtent;
-	float3 ViewTilePosition;
-	float3 RelativePreViewTranslation;
 
-	bool bReadbackHasPendingStreaming;
-	uint ClipmapIndex; // Global index
 	float MeshSDFRadiusThreshold;
 	float InfluenceRadiusSq;
+
+	uint ClipmapIndex; // Global index
+	bool bReadbackHasPendingStreaming;
 };
 
-// Keep in sync with GlobalDistanceField::PackedClipmapBufferStride
-#define PACKED_CLIPMAP_BUFFER_STRIDE 4
-StructuredBuffer<float4> PackedClipmapBuffer;
+StructuredBuffer<FPackedClipmapData> PackedClipmapBuffer;
 
 FPackedClipmap GetPackedClipmap(uint PackedClipmapIndex)
 {
-	uint BaseOffset = PACKED_CLIPMAP_BUFFER_STRIDE * PackedClipmapIndex;
+	FPackedClipmapData ClipmapData = PackedClipmapBuffer[PackedClipmapIndex];
 
 	FPackedClipmap PackedClipmap;
 
-	const float4 V0 = PackedClipmapBuffer[BaseOffset + 0];
-	const float4 V1 = PackedClipmapBuffer[BaseOffset + 1];
-	const float4 V2 = PackedClipmapBuffer[BaseOffset + 2];
-	const float4 V3 = PackedClipmapBuffer[BaseOffset + 3];
+	PackedClipmap.TranslatedWorldCenter = ClipmapData.TranslatedWorldCenter;
+	PackedClipmap.WorldExtent = ClipmapData.WorldExtent;
+
+	PackedClipmap.MeshSDFRadiusThreshold = ClipmapData.MeshSDFRadiusThreshold;
+	PackedClipmap.InfluenceRadiusSq = ClipmapData.InfluenceRadiusSq;
 
-	PackedClipmap.TranslatedWorldCenter = V0.xyz;
-	PackedClipmap.MeshSDFRadiusThreshold = V0.w;
-	PackedClipmap.WorldExtent = V1.xyz;
-	PackedClipmap.InfluenceRadiusSq = V1.w;
-	PackedClipmap.ViewTilePosition = V2.xyz;
-	PackedClipmap.bReadbackHasPendingStreaming = V2.w > 0.0f;
-	PackedClipmap.RelativePreViewTranslation = V3.xyz;
-	PackedClipmap.ClipmapIndex = V3.w;
+	PackedClipmap.ClipmapIndex = ClipmapData.ClipmapIndex;
+	PackedClipmap.bReadbackHasPendingStreaming = ClipmapData.Flags != 0;
 
 	return PackedClipmap;
 }
@@ -82,9 +89,10 @@ RWStructuredBuffer<uint> RWObjectIndexNumBuffer;
 
 RWStructuredBuffer<uint> RWHasPendingStreaming;
 
-float3 ClipmapTranslatedWorldCenter;
-float3 ClipmapWorldExtent;
-float InfluenceRadius;
+
+float3 ViewTilePosition;
+float3 RelativePreViewTranslation;
+
 // 0: accept only mostly static objects, 1: accept only often moving objects, 2: accept all objects
 uint AcceptOftenMovingObjectsOnly;
 uint NumPackedClipmaps;
@@ -105,15 +113,16 @@ void CullObjectsToClipmapCS(
 
 		if (bMatchingMobility && bVisible)
 		{
+			const FLWCVector3 PreViewTranslation = MakeLWCVector3(-ViewTilePosition, RelativePreViewTranslation);
+			const float3 TranslatedWorldObjectBoundsCenter = LWCToFloat(LWCAdd(DFObjectBounds.Center, PreViewTranslation));
+
 			for (uint PackedClipmapIndex = 0; PackedClipmapIndex < NumPackedClipmaps; ++PackedClipmapIndex)
 			{
 				FPackedClipmap Clipmap = GetPackedClipmap(PackedClipmapIndex);
 
 				if (DFObjectBounds.SphereRadius > Clipmap.MeshSDFRadiusThreshold || DFObjectBounds.bEmissiveLightSource)
 				{
-					const FLWCVector3 PreViewTranslation = MakeLWCVector3(-Clipmap.ViewTilePosition, Clipmap.RelativePreViewTranslation);
-					const float3 TranslatedBoundsCenter = LWCToFloat(LWCAdd(DFObjectBounds.Center, PreViewTranslation));
-					float DistanceSq = ComputeSquaredDistanceBetweenAABBs(Clipmap.TranslatedWorldCenter, Clipmap.WorldExtent, TranslatedBoundsCenter, DFObjectBounds.BoxExtent);
+					float DistanceSq = ComputeSquaredDistanceBetweenAABBs(Clipmap.TranslatedWorldCenter, Clipmap.WorldExtent, TranslatedWorldObjectBoundsCenter, DFObjectBounds.BoxExtent);
 
 					if (DistanceSq < Clipmap.InfluenceRadiusSq)
 					{
@@ -349,8 +358,6 @@ float3 InvPageGridResolution;
 float3 PageWorldExtent;
 float PageWorldRadius;
 float ClipmapInfluenceRadius;
-float3 ViewTilePosition;
-float3 RelativePreViewTranslation;
 
 RWStructuredBuffer<uint> RWPageComposeTileBuffer;
 
diff --git a/Engine/Source/Runtime/Renderer/Private/GlobalDistanceField.cpp b/Engine/Source/Runtime/Renderer/Private/GlobalDistanceField.cpp
--- a/Engine/Source/Runtime/Renderer/Private/GlobalDistanceField.cpp
+++ b/Engine/Source/Runtime/Renderer/Private/GlobalDistanceField.cpp
@@ -410,9 +410,6 @@ namespace GlobalDistanceField
 	// Every distance field page stores 4^3 object grid cells with 4 * uint32 elements per cell
 	const int32 ObjectGridPageBufferStride = 4 * sizeof(uint32);
 	const int32 ObjectGridPageBufferNumElementsPerPage = 4 * 4 * 4;
-
-	// Keep in sync with PACKED_CLIPMAP_BUFFER_STRIDE
-	const int32 PackedClipmapBufferStride = 4;
 }
 const int32 GGlobalDistanceFieldPageResolutionInAtlas = 8; // Includes 0.5 texel trilinear filter margin
 const int32 GGlobalDistanceFieldCoveragePageResolutionInAtlas = 4; // Includes 0.5 texel trilinear filter margin
@@ -1412,6 +1409,8 @@ class FCullObjectsToClipmapCS : public FGlobalShader
 		SHADER_PARAMETER(uint32, AcceptOftenMovingObjectsOnly)
 		SHADER_PARAMETER(uint32, NumPackedClipmaps)
 		SHADER_PARAMETER(uint32, ObjectIndexBufferStride)
+		SHADER_PARAMETER(FVector3f, ViewTilePosition)
+		SHADER_PARAMETER(FVector3f, RelativePreViewTranslation)
 	END_SHADER_PARAMETER_STRUCT()
 
 	class FReadbackHasPendingStreaming : SHADER_PERMUTATION_BOOL("READBACK_HAS_PENDING_STREAMING");
@@ -2361,17 +2360,23 @@ void UpdateGlobalDistanceFieldCache(
 		// Upload packed clipmap data to GPU
 		FRDGBufferRef PackedClipmapBuffer = nullptr;
 		{
-			TArray<FVector4f> UploadData;
-
-			if (PackedClipmaps.Num() > 0)
+			// Must match FPackedClipmapData in GlobalDistanceField.usf
+			struct FPackedClipmapData
 			{
-				UploadData.SetNum(GlobalDistanceField::PackedClipmapBufferStride * PackedClipmaps.Num());
-			}
-			else
-			{
-				// Create a dummy entry filled with 0 in order to be able to create a dummy PackedClipmapBuffer resource
-				UploadData.SetNumZeroed(GlobalDistanceField::PackedClipmapBufferStride);
-			}
+				FVector3f TranslatedWorldCenter; // Camera-centered world space
+				float MeshSDFRadiusThreshold;
+
+				FVector3f WorldExtent;
+				float InfluenceRadiusSq;
+
+				uint32 ClipmapIndex;
+				uint32 Flags;
+				uint32 Padding0;
+				uint32 Padding1;
+			};
+
+			TArray<FPackedClipmapData> UploadData;
+			UploadData.SetNum(FMath::Max(PackedClipmaps.Num(), 1));
 
 			for (int32 PackedClipmapIndex = 0; PackedClipmapIndex < PackedClipmaps.Num(); ++PackedClipmapIndex)
 			{
@@ -2380,10 +2385,22 @@ void UpdateGlobalDistanceFieldCache(
 				const float RadiusThresholdScale = bLumenEnabled ? 1.0f / FMath::Clamp(View.FinalPostProcessSettings.LumenSceneDetail, .01f, 100.0f) : 1.0f;
 				const float MeshSDFRadiusThreshold = GetMinMeshSDFRadius(PackedClipmap.VoxelSize.X) * RadiusThresholdScale;
 
-				UploadData[PackedClipmapIndex * GlobalDistanceField::PackedClipmapBufferStride + 0] = FVector4f((FVector3f)PackedClipmap.TranslatedBounds.GetCenter(), MeshSDFRadiusThreshold);
-				UploadData[PackedClipmapIndex * GlobalDistanceField::PackedClipmapBufferStride + 1] = FVector4f((FVector3f)PackedClipmap.TranslatedBounds.GetExtent(), PackedClipmap.InfluenceRadius * PackedClipmap.InfluenceRadius);
-				UploadData[PackedClipmapIndex * GlobalDistanceField::PackedClipmapBufferStride + 2] = FVector4f(PackedClipmap.ViewTilePosition, PackedClipmap.bRecacheClipmapsWithPendingStreaming ? 1.0f : 0.0f);
-				UploadData[PackedClipmapIndex * GlobalDistanceField::PackedClipmapBufferStride + 3] = FVector4f(PackedClipmap.RelativePreViewTranslation, PackedClipmap.Index);
+				// Camera-centered world space
+				const FGlobalDistanceFieldClipmap& Clipmap = Clipmaps[PackedClipmap.Index];
+				const FVector TranslatedWorldCenter = Clipmap.Bounds.GetCenter() + View.ViewMatrices.GetPreViewTranslation();
+
+				FPackedClipmapData& ClipmapData = UploadData[PackedClipmapIndex];
+
+				ClipmapData.TranslatedWorldCenter = (FVector3f)TranslatedWorldCenter;
+				ClipmapData.MeshSDFRadiusThreshold = MeshSDFRadiusThreshold;
+
+				ClipmapData.WorldExtent = (FVector3f)PackedClipmap.TranslatedBounds.GetExtent();
+				ClipmapData.InfluenceRadiusSq = PackedClipmap.InfluenceRadius * PackedClipmap.InfluenceRadius;
+
+				ClipmapData.ClipmapIndex = PackedClipmap.Index;
+				ClipmapData.Flags = PackedClipmap.bRecacheClipmapsWithPendingStreaming ? 1u : 0u;
+				ClipmapData.Padding0 = 0;
+				ClipmapData.Padding1 = 0;
 			}
 
 			PackedClipmapBuffer =
@@ -2419,6 +2436,11 @@ void UpdateGlobalDistanceFieldCache(
 			PassParameters->DistanceFieldObjectBuffers = DistanceFieldObjectBuffers;
 			PassParameters->DistanceFieldAtlasParameters = DistanceFieldAtlas;
 
+			const FLargeWorldRenderPosition AbsoluteViewOrigin(View.ViewMatrices.GetViewOrigin());
+			const FVector ViewTileOffset = AbsoluteViewOrigin.GetTileOffset();
+			PassParameters->RelativePreViewTranslation = FVector3f(View.ViewMatrices.GetPreViewTranslation() + ViewTileOffset);
+			PassParameters->ViewTilePosition = AbsoluteViewOrigin.GetTile();
+
 			FCullObjectsToClipmapCS::FPermutationDomain PermutationVector;
 			PermutationVector.Set<FCullObjectsToClipmapCS::FReadbackHasPendingStreaming>(bAnyClipmapHasPendingStreamingReadback);
 			auto ComputeShader = View.ShaderMap->GetShader<FCullObjectsToClipmapCS>(PermutationVector);
