Refactor function arguments in DistanceFieldObjectManagement.cpp

diff --git a/Engine/Source/Runtime/Renderer/Private/DistanceFieldObjectManagement.cpp b/Engine/Source/Runtime/Renderer/Private/DistanceFieldObjectManagement.cpp
--- a/Engine/Source/Runtime/Renderer/Private/DistanceFieldObjectManagement.cpp
+++ b/Engine/Source/Runtime/Renderer/Private/DistanceFieldObjectManagement.cpp
@@ -209,16 +209,13 @@ static int32 PartitionUpdateRangesDFO(FParallelUpdateRangesDFO& Ranges, int32 It
 	return Ranges.Range[3].ItemCount > 0 ? 4 : 3;
 }
 
-void AddModifiedBounds(FScene* Scene, FGlobalDFCacheType CacheType, const FBox& Bounds)
+void AddModifiedBounds(FDistanceFieldSceneData& DistanceFieldSceneData, FGlobalDFCacheType CacheType, const FBox& Bounds)
 {
-	FDistanceFieldSceneData& DistanceFieldData = Scene->DistanceFieldSceneData;
-	DistanceFieldData.PrimitiveModifiedBounds[CacheType].Add(Bounds);
+	DistanceFieldSceneData.PrimitiveModifiedBounds[CacheType].Add(Bounds);
 }
 
-void ProcessDistanceFieldObjectRemoves(FScene* Scene, TArray<FSetElementId>& DistanceFieldAssetRemoves)
+void ProcessDistanceFieldObjectRemoves(FDistanceFieldSceneData& DistanceFieldSceneData, TArray<FSetElementId>& DistanceFieldAssetRemoves)
 {
-	FDistanceFieldSceneData& DistanceFieldSceneData = Scene->DistanceFieldSceneData;
-
 	if (DistanceFieldSceneData.PendingRemoveOperations.Num() > 0)
 	{
 		TArray<int32, SceneRenderingAllocator> PendingRemoveOperations;
@@ -249,7 +246,7 @@ void ProcessDistanceFieldObjectRemoves(FScene* Scene, TArray<FSetElementId>& Dis
 				{
 					// Mark region covered by instance in global distance field as modified
 					FGlobalDFCacheType CacheType = PrimitiveRemoveInfo.bOftenMoving ? GDF_Full : GDF_MostlyStatic;
-					AddModifiedBounds(Scene, CacheType, DistanceFieldSceneData.PrimitiveInstanceMapping[InstanceIndex].GetWorldBounds());
+					AddModifiedBounds(DistanceFieldSceneData, CacheType, DistanceFieldSceneData.PrimitiveInstanceMapping[InstanceIndex].GetWorldBounds());
 
 					// Add individual instances to temporary array for processing in the next pass
 					PendingRemoveOperations.Add(InstanceIndex);
@@ -309,17 +306,14 @@ void LogDistanceFieldUpdate(FPrimitiveSceneInfo const* PrimitiveSceneInfo, float
 /** Gathers the information needed to represent a single object's distance field and appends it to the upload buffers. */
 void ProcessPrimitiveUpdate(
 	bool bIsAddOperation,
-	FScene* Scene,
+	FDistanceFieldSceneData& DistanceFieldSceneData,
 	FPrimitiveSceneInfo* PrimitiveSceneInfo,
 	TArray<FMatrix>& InstanceLocalToWorldTmpStorage,
-	TArray<int32>& IndicesToUpdateInObjectBuffers, 
 	TArray<FDistanceFieldAssetMipId>& DistanceFieldAssetAdds,
 	TArray<FSetElementId>& DistanceFieldAssetRemoves)
 {
 	const FPrimitiveSceneProxy* Proxy = PrimitiveSceneInfo->Proxy;
 
-	FDistanceFieldSceneData& DistanceFieldSceneData = Scene->DistanceFieldSceneData;
-
 	InstanceLocalToWorldTmpStorage.Reset();
 
 	const FDistanceFieldVolumeData* DistanceFieldData = nullptr;
@@ -409,7 +403,7 @@ void ProcessPrimitiveUpdate(
 					UploadIndex = PrimitiveSceneInfo->DistanceFieldInstanceIndices[TransformIndex];
 				}
 
-				IndicesToUpdateInObjectBuffers.Add(UploadIndex);
+				DistanceFieldSceneData.IndicesToUpdateInObjectBuffers.Add(UploadIndex);
 
 				const FBox WorldBounds = ((FBox)DistanceFieldData->LocalSpaceMeshBounds).TransformBy(LocalToWorld);
 
@@ -419,7 +413,7 @@ void ProcessPrimitiveUpdate(
 					const int32 MappingIndex = DistanceFieldSceneData.PrimitiveInstanceMapping.Add(FPrimitiveAndInstance(LocalToWorld, WorldBounds, PrimitiveSceneInfo, TransformIndex));
 					PrimitiveSceneInfo->DistanceFieldInstanceIndices[TransformIndex] = AddIndex;
 
-					AddModifiedBounds(Scene, CacheType, WorldBounds);
+					AddModifiedBounds(DistanceFieldSceneData, CacheType, WorldBounds);
 					LogDistanceFieldUpdate(PrimitiveSceneInfo, BoundingRadius, bIsAddOperation);
 				}
 				else 
@@ -442,12 +436,12 @@ void ProcessPrimitiveUpdate(
 							const FVector MergedExtentIncrease = MergedBounds.GetExtent() - PrevWorldBounds.GetExtent() - WorldBounds.GetExtent();
 							if (MergedExtentIncrease.GetMax() < 100.0f)
 							{
-								AddModifiedBounds(Scene, CacheType, MergedBounds);
+								AddModifiedBounds(DistanceFieldSceneData, CacheType, MergedBounds);
 							}
 							else
 							{
-								AddModifiedBounds(Scene, CacheType, PrevWorldBounds);
-								AddModifiedBounds(Scene, CacheType, WorldBounds);
+								AddModifiedBounds(DistanceFieldSceneData, CacheType, PrevWorldBounds);
+								AddModifiedBounds(DistanceFieldSceneData, CacheType, WorldBounds);
 							}
 							LogDistanceFieldUpdate(PrimitiveSceneInfo, BoundingRadius, bIsAddOperation);
 
@@ -490,7 +484,7 @@ void FDistanceFieldSceneData::UpdateDistanceFieldObjectBuffers(
 
 		// Process removes before adds, as the adds will overwrite primitive allocation info
 		// This also prevents re-uploading distance fields on render state recreation
-		ProcessDistanceFieldObjectRemoves(Scene, DistanceFieldAssetRemoves);
+		ProcessDistanceFieldObjectRemoves(*this, DistanceFieldAssetRemoves);
 
 		if ((PendingAddOperations.Num() > 0 || PendingUpdateOperations.Num() > 0) && GDFReverseAtlasAllocationOrder == GDFPreviousReverseAtlasAllocationOrder)
 		{
@@ -500,10 +494,9 @@ void FDistanceFieldSceneData::UpdateDistanceFieldObjectBuffers(
 			{
 				ProcessPrimitiveUpdate(
 					true,
-					Scene,
+					*this,
 					PrimitiveSceneInfo,
 					InstanceLocalToPrimitiveTransforms,
-					IndicesToUpdateInObjectBuffers,
 					DistanceFieldAssetAdds,
 					DistanceFieldAssetRemoves);
 			}
@@ -512,10 +505,9 @@ void FDistanceFieldSceneData::UpdateDistanceFieldObjectBuffers(
 			{
 				ProcessPrimitiveUpdate(
 					false,
-					Scene,
+					*this,
 					PrimitiveSceneInfo,
 					InstanceLocalToPrimitiveTransforms,
-					IndicesToUpdateInObjectBuffers,
 					DistanceFieldAssetAdds,
 					DistanceFieldAssetRemoves);
 			}
@@ -964,7 +956,7 @@ void FSceneRenderer::ProcessPendingHeightFieldPrimitiveAddAndRemoveOps(TArray<ui
 		PendingRemoveIndices.Add(RemoveInfo.DistanceFieldInstanceIndices[0]);
 
 		const FGlobalDFCacheType CacheType = RemoveInfo.bOftenMoving ? GDF_Full : GDF_MostlyStatic;
-		AddModifiedBounds(Scene, CacheType, RemoveInfo.WorldBounds);
+		AddModifiedBounds(SceneData, CacheType, RemoveInfo.WorldBounds);
 	}
 
 	SceneData.PendingHeightFieldRemoveOps.Reset();
@@ -1015,7 +1007,7 @@ void FSceneRenderer::ProcessPendingHeightFieldPrimitiveAddAndRemoveOps(TArray<ui
 
 		const FGlobalDFCacheType CacheType = Primitive->Proxy->IsOftenMoving() ? GDF_Full : GDF_MostlyStatic;
 		const FBoxSphereBounds& Bounds = Primitive->Proxy->GetBounds();
-		AddModifiedBounds(Scene, CacheType, Bounds.GetBox());
+		AddModifiedBounds(SceneData, CacheType, Bounds.GetBox());
 	}
 
 	SceneData.PendingHeightFieldAddOps.Reset();
