Prevent duplicate uploads to DistanceFields.DFObjectDataUploadBuffer

diff --git a/Engine/Source/Runtime/Renderer/Private/DistanceFieldObjectManagement.cpp b/Engine/Source/Runtime/Renderer/Private/DistanceFieldObjectManagement.cpp
--- a/Engine/Source/Runtime/Renderer/Private/DistanceFieldObjectManagement.cpp
+++ b/Engine/Source/Runtime/Renderer/Private/DistanceFieldObjectManagement.cpp
@@ -278,7 +278,11 @@ void ProcessDistanceFieldObjectRemoves(FDistanceFieldSceneData& DistanceFieldSce
 
 				DistanceFieldSceneData.PrimitiveInstanceMapping.RemoveAtSwap(RemoveIndex, 1, false);
 
-				DistanceFieldSceneData.IndicesToUpdateInObjectBuffers.Add(RemoveIndex);
+				if (!DistanceFieldSceneData.IndicesToUpdateInObjectBuffersSet.Contains(RemoveIndex))
+				{
+					DistanceFieldSceneData.IndicesToUpdateInObjectBuffers.Add(RemoveIndex);
+					DistanceFieldSceneData.IndicesToUpdateInObjectBuffersSet.Add(RemoveIndex);
+				}
 			}
 
 			PendingRemoveOperations.Reset();
@@ -403,15 +407,18 @@ void ProcessPrimitiveUpdate(
 					UploadIndex = PrimitiveSceneInfo->DistanceFieldInstanceIndices[TransformIndex];
 				}
 
-				DistanceFieldSceneData.IndicesToUpdateInObjectBuffers.Add(UploadIndex);
+				if (!DistanceFieldSceneData.IndicesToUpdateInObjectBuffersSet.Contains(UploadIndex))
+				{
+					DistanceFieldSceneData.IndicesToUpdateInObjectBuffers.Add(UploadIndex);
+					DistanceFieldSceneData.IndicesToUpdateInObjectBuffersSet.Add(UploadIndex);
+				}
 
 				const FBox WorldBounds = ((FBox)DistanceFieldData->LocalSpaceMeshBounds).TransformBy(LocalToWorld);
 
 				if (bIsAddOperation)
 				{
-					const int32 AddIndex = UploadIndex;
 					const int32 MappingIndex = DistanceFieldSceneData.PrimitiveInstanceMapping.Add(FPrimitiveAndInstance(LocalToWorld, WorldBounds, PrimitiveSceneInfo, TransformIndex));
-					PrimitiveSceneInfo->DistanceFieldInstanceIndices[TransformIndex] = AddIndex;
+					PrimitiveSceneInfo->DistanceFieldInstanceIndices[TransformIndex] = UploadIndex;
 
 					AddModifiedBounds(DistanceFieldSceneData, CacheType, WorldBounds);
 					LogDistanceFieldUpdate(PrimitiveSceneInfo, BoundingRadius, bIsAddOperation);
@@ -539,6 +546,8 @@ void FDistanceFieldSceneData::UpdateDistanceFieldObjectBuffers(
 			// This is not expected to be hit during gameplay.
 			static const int32 MAX_NUM_DISTANCE_FIELD_OBJECT_UPLOADS = (2 << 20);
 
+			check(IndicesToUpdateInObjectBuffers.Num() == IndicesToUpdateInObjectBuffersSet.Num());
+
 			const int32 NumDFObjectUploads = FMath::Min(IndicesToUpdateInObjectBuffers.Num(), MAX_NUM_DISTANCE_FIELD_OBJECT_UPLOADS);
 
 			static FCriticalSection DFUpdateCS;
@@ -688,12 +697,21 @@ void FDistanceFieldSceneData::UpdateDistanceFieldObjectBuffers(
 				ExternalAccessQueue.Add(DFObjectDataBuffer, ERHIAccess::SRVMask, ERHIPipeline::All);
 				ExternalAccessQueue.Add(DFObjectBoundsBuffer, ERHIAccess::SRVMask, ERHIPipeline::All);
 
+				IndicesToUpdateInObjectBuffersSet.Reset();
+
 				if (IndicesToUpdateInObjectBuffers.Num() > NumDFObjectUploads)
 				{
 					// this is not expected to happen frequently since we can perform up to MAX_NUM_DISTANCE_FIELD_OBJECT_UPLOADS per frame
 					// RemoveAtSwap would be more efficient but could potentially result in starvation
 					const bool bAllowShrinking = true; // allow array to shrink since getting into this code path means array is very large
 					IndicesToUpdateInObjectBuffers.RemoveAt(0, NumDFObjectUploads, bAllowShrinking);
+
+					for (int32 Index : IndicesToUpdateInObjectBuffers)
+					{
+						IndicesToUpdateInObjectBuffersSet.Add(Index);
+					}
+
+					check(IndicesToUpdateInObjectBuffers.Num() == IndicesToUpdateInObjectBuffersSet.Num());
 				}
 				else
 				{
diff --git a/Engine/Source/Runtime/Renderer/Private/ScenePrivate.h b/Engine/Source/Runtime/Renderer/Private/ScenePrivate.h
--- a/Engine/Source/Runtime/Renderer/Private/ScenePrivate.h
+++ b/Engine/Source/Runtime/Renderer/Private/ScenePrivate.h
@@ -2173,8 +2173,12 @@ public:
 	FRDGScatterUploadBuffer UploadHeightFieldBoundsBuffer;
 	FRDGScatterUploadBuffer UploadDistanceFieldDataBuffer;
 	FRDGScatterUploadBuffer UploadDistanceFieldBoundsBuffer;
-	
+
+	// track indices that need to be updated using both an array and a set
+	// array is used for fast iteration and support ParallelFor
+	// set is used to prevent duplicate indices
 	TArray<int32> IndicesToUpdateInObjectBuffers;
+	TSet<int32> IndicesToUpdateInObjectBuffersSet;
 
 	TSet<FDistanceFieldAssetState, TFDistanceFieldAssetStateFuncs> AssetStateArray;
 	TRefCountPtr<FRDGPooledBuffer> AssetDataBuffer;
