Changes r.TSR.16BitVALU.AMD's default to 1

diff --git a/Engine/Source/Runtime/Renderer/Private/PostProcess/TemporalSuperResolution.cpp b/Engine/Source/Runtime/Renderer/Private/PostProcess/TemporalSuperResolution.cpp
--- a/Engine/Source/Runtime/Renderer/Private/PostProcess/TemporalSuperResolution.cpp
+++ b/Engine/Source/Runtime/Renderer/Private/PostProcess/TemporalSuperResolution.cpp
@@ -87,7 +87,7 @@ TAutoConsoleVariable<int32> CVarTSR16BitVALU(
 #if PLATFORM_DESKTOP
 
 TAutoConsoleVariable<int32> CVarTSR16BitVALUOnAMD(
-	TEXT("r.TSR.16BitVALU.AMD"), 0,
+	TEXT("r.TSR.16BitVALU.AMD"), 1,
 	TEXT("Overrides whether to use 16bit VALU on AMD desktop GPUs"),
 	ECVF_RenderThreadSafe);
 
