Fix issues in distance field scene managment when instance scale changes

Context:
- Degenerate instances (MinScale < 0.0001f) are not added to DF scene.
- However this filtering was applied when primitive was created and not updated if instances scale changed.
- This caused issues with missing instances in the case they started with scale = 0 but were then updated to have valid scale, or having degenerate instances in the scene if they were created with valid scale but subsequently modified with scale = 0.

Fix:
- When transforms change, remove instances that become degenerate from DF scene and add those that were previously degenerate but become valid.

diff --git a/Engine/Source/Runtime/Renderer/Private/DistanceFieldObjectManagement.cpp b/Engine/Source/Runtime/Renderer/Private/DistanceFieldObjectManagement.cpp
--- a/Engine/Source/Runtime/Renderer/Private/DistanceFieldObjectManagement.cpp
+++ b/Engine/Source/Runtime/Renderer/Private/DistanceFieldObjectManagement.cpp
@@ -214,6 +214,30 @@ void AddModifiedBounds(FDistanceFieldSceneData& DistanceFieldSceneData, FGlobalD
 	DistanceFieldSceneData.PrimitiveModifiedBounds[CacheType].Add(Bounds);
 }
 
+static void RemoveDistanceFieldInstance(int32 RemoveIndex, FDistanceFieldSceneData& DistanceFieldSceneData)
+{
+	--DistanceFieldSceneData.NumObjectsInBuffer;
+
+	if (RemoveIndex < DistanceFieldSceneData.NumObjectsInBuffer)
+	{
+		const int32 MoveFromIndex = DistanceFieldSceneData.NumObjectsInBuffer;
+
+		FPrimitiveAndInstance& PrimitiveAndInstanceBeingMoved = DistanceFieldSceneData.PrimitiveInstanceMapping[MoveFromIndex];
+
+		// Fixup indices of the primitive that is being moved
+		check(PrimitiveAndInstanceBeingMoved.Primitive && PrimitiveAndInstanceBeingMoved.Primitive->DistanceFieldInstanceIndices.Num() > 0);
+		PrimitiveAndInstanceBeingMoved.Primitive->DistanceFieldInstanceIndices[PrimitiveAndInstanceBeingMoved.InstanceIndex] = RemoveIndex;
+	}
+
+	DistanceFieldSceneData.PrimitiveInstanceMapping.RemoveAtSwap(RemoveIndex, 1, false);
+
+	if(!DistanceFieldSceneData.IndicesToUpdateInObjectBuffersSet.Contains(RemoveIndex))
+	{
+		DistanceFieldSceneData.IndicesToUpdateInObjectBuffers.Add(RemoveIndex);
+		DistanceFieldSceneData.IndicesToUpdateInObjectBuffersSet.Add(RemoveIndex);
+	}
+}
+
 void ProcessDistanceFieldObjectRemoves(FDistanceFieldSceneData& DistanceFieldSceneData, TArray<FSetElementId>& DistanceFieldAssetRemoves)
 {
 	if (DistanceFieldSceneData.PendingRemoveOperations.Num() > 0)
@@ -265,24 +289,7 @@ void ProcessDistanceFieldObjectRemoves(FDistanceFieldSceneData& DistanceFieldSce
 
 			for (int32 RemoveIndex : PendingRemoveOperations)
 			{
-				--DistanceFieldSceneData.NumObjectsInBuffer;
-				const int32 MoveFromIndex = DistanceFieldSceneData.NumObjectsInBuffer;
-
-				FPrimitiveAndInstance& PrimitiveAndInstanceBeingMoved = DistanceFieldSceneData.PrimitiveInstanceMapping[MoveFromIndex];
-				if (RemoveIndex < DistanceFieldSceneData.NumObjectsInBuffer)
-				{
-					// Fixup indices of the primitive that is being moved
-					check(PrimitiveAndInstanceBeingMoved.Primitive && PrimitiveAndInstanceBeingMoved.Primitive->DistanceFieldInstanceIndices.Num() > 0);
-					PrimitiveAndInstanceBeingMoved.Primitive->DistanceFieldInstanceIndices[PrimitiveAndInstanceBeingMoved.InstanceIndex] = RemoveIndex;
-				}
-
-				DistanceFieldSceneData.PrimitiveInstanceMapping.RemoveAtSwap(RemoveIndex, 1, false);
-
-				if (!DistanceFieldSceneData.IndicesToUpdateInObjectBuffersSet.Contains(RemoveIndex))
-				{
-					DistanceFieldSceneData.IndicesToUpdateInObjectBuffers.Add(RemoveIndex);
-					DistanceFieldSceneData.IndicesToUpdateInObjectBuffersSet.Add(RemoveIndex);
-				}
+				RemoveDistanceFieldInstance(RemoveIndex, DistanceFieldSceneData);
 			}
 
 			PendingRemoveOperations.Reset();
@@ -375,7 +382,9 @@ void ProcessPrimitiveUpdate(
 
 			for (int32 TransformIndex = 0; TransformIndex < InstanceLocalToWorldTransforms.Num(); TransformIndex++)
 			{
-				const bool bInstanceCountOverflow = bIsAddOperation && (DistanceFieldSceneData.NumObjectsInBuffer + 1 > MAX_INSTANCE_ID);
+				const int32 bNewInstance = bIsAddOperation || (PrimitiveSceneInfo->DistanceFieldInstanceIndices[TransformIndex] == -1);
+
+				const bool bInstanceCountOverflow = bNewInstance && (DistanceFieldSceneData.NumObjectsInBuffer + 1 > MAX_INSTANCE_ID);
 
 				static bool bWarnOnce = true;
 				if (bInstanceCountOverflow && bWarnOnce)
@@ -388,16 +397,23 @@ void ProcessPrimitiveUpdate(
 
 				const FMatrix::FReal MinScale = LocalToWorld.GetMinimumAxisScale();
 
-				if (bIsAddOperation && (MinScale < 0.0001f || bInstanceCountOverflow))
+				// Don't include degenerate instances or when instance count limit is reached
+				if (MinScale < 0.0001f || bInstanceCountOverflow)
 				{
-					// Skip degenerate instances or when instance count limit is reached
+					if (!bNewInstance)
+					{
+						// remove existing instance
+						const int32 RemoveIndex = PrimitiveSceneInfo->DistanceFieldInstanceIndices[TransformIndex];
+						RemoveDistanceFieldInstance(RemoveIndex, DistanceFieldSceneData);
+					}
+
 					PrimitiveSceneInfo->DistanceFieldInstanceIndices[TransformIndex] = -1;
 					continue;
 				}
 
 				uint32 UploadIndex;
 
-				if (bIsAddOperation)
+				if (bNewInstance)
 				{
 					UploadIndex = DistanceFieldSceneData.NumObjectsInBuffer;
 					++DistanceFieldSceneData.NumObjectsInBuffer;
@@ -415,7 +431,7 @@ void ProcessPrimitiveUpdate(
 
 				const FBox WorldBounds = ((FBox)DistanceFieldData->LocalSpaceMeshBounds).TransformBy(LocalToWorld);
 
-				if (bIsAddOperation)
+				if (bNewInstance)
 				{
 					const int32 MappingIndex = DistanceFieldSceneData.PrimitiveInstanceMapping.Add(FPrimitiveAndInstance(LocalToWorld, WorldBounds, PrimitiveSceneInfo, TransformIndex));
 					PrimitiveSceneInfo->DistanceFieldInstanceIndices[TransformIndex] = UploadIndex;
@@ -425,35 +441,33 @@ void ProcessPrimitiveUpdate(
 				}
 				else 
 				{
-					// InstanceIndex will be -1 with zero scale meshes
 					const int32 InstanceIndex = PrimitiveSceneInfo->DistanceFieldInstanceIndices[TransformIndex];
-					if (InstanceIndex >= 0)
-					{
-						FPrimitiveAndInstance& Mapping = DistanceFieldSceneData.PrimitiveInstanceMapping[InstanceIndex];
+					check(InstanceIndex >= 0);
 
-						const FMatrix PrevLocalToWorld = Mapping.GetLocalToWorld();
-						const FBox PrevWorldBounds = Mapping.GetWorldBounds();
+					FPrimitiveAndInstance& Mapping = DistanceFieldSceneData.PrimitiveInstanceMapping[InstanceIndex];
 
-						// Filter out global distance field updates which were too small
-						if (!PrevWorldBounds.GetExtent().Equals(WorldBounds.GetExtent(), 0.01f)
-							|| !PrevLocalToWorld.Equals(LocalToWorld, 0.01f))
-						{
-							// decide if we want to make a single global distance field update or two updates for large movement (teleport) case
-							const FBox MergedBounds = PrevWorldBounds + WorldBounds;
-							const FVector MergedExtentIncrease = MergedBounds.GetExtent() - PrevWorldBounds.GetExtent() - WorldBounds.GetExtent();
-							if (MergedExtentIncrease.GetMax() < 100.0f)
-							{
-								AddModifiedBounds(DistanceFieldSceneData, CacheType, MergedBounds);
-							}
-							else
-							{
-								AddModifiedBounds(DistanceFieldSceneData, CacheType, PrevWorldBounds);
-								AddModifiedBounds(DistanceFieldSceneData, CacheType, WorldBounds);
-							}
-							LogDistanceFieldUpdate(PrimitiveSceneInfo, BoundingRadius, bIsAddOperation);
+					const FMatrix PrevLocalToWorld = Mapping.GetLocalToWorld();
+					const FBox PrevWorldBounds = Mapping.GetWorldBounds();
 
-							Mapping.SetTransformAndBounds(LocalToWorld, WorldBounds);
+					// Filter out global distance field updates which were too small
+					if (!PrevWorldBounds.GetExtent().Equals(WorldBounds.GetExtent(), 0.01f)
+						|| !PrevLocalToWorld.Equals(LocalToWorld, 0.01f))
+					{
+						// decide if we want to make a single global distance field update or two updates for large movement (teleport) case
+						const FBox MergedBounds = PrevWorldBounds + WorldBounds;
+						const FVector MergedExtentIncrease = MergedBounds.GetExtent() - PrevWorldBounds.GetExtent() - WorldBounds.GetExtent();
+						if (MergedExtentIncrease.GetMax() < 100.0f)
+						{
+							AddModifiedBounds(DistanceFieldSceneData, CacheType, MergedBounds);
 						}
+						else
+						{
+							AddModifiedBounds(DistanceFieldSceneData, CacheType, PrevWorldBounds);
+							AddModifiedBounds(DistanceFieldSceneData, CacheType, WorldBounds);
+						}
+						LogDistanceFieldUpdate(PrimitiveSceneInfo, BoundingRadius, bIsAddOperation);
+
+						Mapping.SetTransformAndBounds(LocalToWorld, WorldBounds);
 					}
 				}
 			}
@@ -569,7 +583,9 @@ void FDistanceFieldSceneData::UpdateDistanceFieldObjectBuffers(
 						for (int32 ItemIndex = ParallelRanges.Range[RangeIndex].ItemStart; ItemIndex < ParallelRanges.Range[RangeIndex].ItemStart + ParallelRanges.Range[RangeIndex].ItemCount; ++ItemIndex)
 						{
 							const int32 Index = IndicesToUpdateInObjectBuffers[ItemIndex];
-							if (Index >= 0 && Index < PrimitiveInstanceMapping.Num())
+							checkf(Index >= 0, TEXT("Invalid instances should've been skipped in ProcessPrimitiveUpdate(...)"));
+
+							if (Index < PrimitiveInstanceMapping.Num())
 							{
 								const FPrimitiveAndInstance& PrimAndInst = PrimitiveInstanceMapping[Index];
 								const FPrimitiveSceneProxy* PrimitiveSceneProxy = PrimAndInst.Primitive->Proxy;
