Moved scene occlusion allocation to a setup task.

diff --git a/Engine/Source/Runtime/Renderer/Private/SceneOcclusion.cpp b/Engine/Source/Runtime/Renderer/Private/SceneOcclusion.cpp
--- a/Engine/Source/Runtime/Renderer/Private/SceneOcclusion.cpp
+++ b/Engine/Source/Runtime/Renderer/Private/SceneOcclusion.cpp
@@ -1064,7 +1064,7 @@ static void TrimAllOcclusionHistory(TArrayView<FViewInfo> Views)
 	}
 }
 
-static FViewOcclusionQueriesPerView AllocateOcclusionTests(const FScene* Scene, TArrayView<const FVisibleLightInfo> VisibleLightInfos, TArrayView<FViewInfo> Views)
+static void AllocateOcclusionTests(FViewOcclusionQueriesPerView& QueriesPerView, const FScene* Scene, TArrayView<const FVisibleLightInfo> VisibleLightInfos, TArrayView<FViewInfo> Views)
 {
 	SCOPED_NAMED_EVENT(FSceneRenderer_AllocateOcclusionTestsOcclusionTests, FColor::Emerald);
 
@@ -1073,7 +1073,6 @@ static FViewOcclusionQueriesPerView AllocateOcclusionTests(const FScene* Scene,
 
 	bool bBatchedQueries = false;
 
-	FViewOcclusionQueriesPerView QueriesPerView;
 	QueriesPerView.AddDefaulted(Views.Num());
 
 	// Perform occlusion queries for each view
@@ -1217,7 +1216,6 @@ static FViewOcclusionQueriesPerView AllocateOcclusionTests(const FScene* Scene,
 	{
 		QueriesPerView.Empty();
 	}
-	return MoveTemp(QueriesPerView);
 }
 
 static void BeginOcclusionTests(
@@ -1414,15 +1412,18 @@ void FDeferredShadingSceneRenderer::RenderOcclusion(
 			}
 		}
 
-		// Issue occlusion queries. This is done after the downsampled depth buffer is created so that it can be used for issuing queries.
-		FViewOcclusionQueriesPerView QueriesPerView = AllocateOcclusionTests(Scene, VisibleLightInfos, Views);
+		auto& QueriesPerView = *GraphBuilder.AllocObject<FViewOcclusionQueriesPerView>();
+		auto* PassParameters = GraphBuilder.AllocParameters<FRenderTargetParameters>();
 
-		if (QueriesPerView.Num())
+		GraphBuilder.AddSetupTask([this, PassParameters, &QueriesPerView]
 		{
+			// Issue occlusion queries. This is done after the downsampled depth buffer is created so that it can be used for issuing queries.
+			AllocateOcclusionTests(QueriesPerView, Scene, VisibleLightInfos, Views);
+
 			int32 NumQueriesForBatch = 0;
 
 			{
-				for (int32 ViewIndex = 0; ViewIndex < Views.Num(); ViewIndex++)
+				for (int32 ViewIndex = 0; ViewIndex < QueriesPerView.Num(); ViewIndex++)
 				{
 					const FViewOcclusionQueries& ViewQuery = QueriesPerView[ViewIndex];
 					NumQueriesForBatch += ViewQuery.LocalLightQueries.Num();
@@ -1442,20 +1443,25 @@ void FDeferredShadingSceneRenderer::RenderOcclusion(
 				}
 			}
 
-			auto* PassParameters = GraphBuilder.AllocParameters<FRenderTargetParameters>();
-			PassParameters->RenderTargets.DepthStencil = FDepthStencilBinding(OcclusionDepthTexture, ERenderTargetLoadAction::ELoad, ERenderTargetLoadAction::ELoad, FExclusiveDepthStencil::DepthRead_StencilWrite);
 			PassParameters->RenderTargets.NumOcclusionQueries = NumQueriesForBatch;
+		});
+
+		PassParameters->RenderTargets.DepthStencil = FDepthStencilBinding(OcclusionDepthTexture, ERenderTargetLoadAction::ELoad, ERenderTargetLoadAction::ELoad, FExclusiveDepthStencil::DepthRead_StencilWrite);
+
+		RDG_GPU_STAT_SCOPE(GraphBuilder, BeginOcclusionTests);
 
-			RDG_GPU_STAT_SCOPE(GraphBuilder, BeginOcclusionTests);
-			GraphBuilder.AddPass(
-				RDG_EVENT_NAME("BeginOcclusionTests"),
-				PassParameters,
-				ERDGPassFlags::Raster | ERDGPassFlags::NeverCull,
-				[this, LocalQueriesPerView = MoveTemp(QueriesPerView), DownsampleFactor](FRHICommandList& RHICmdList)
+		GraphBuilder.AddPass(
+			RDG_EVENT_NAME("BeginOcclusionTests"),
+			PassParameters,
+			ERDGPassFlags::Raster | ERDGPassFlags::NeverCull,
+			[this, &QueriesPerView, DownsampleFactor](FRHICommandList& RHICmdList)
 			{
-				BeginOcclusionTests(RHICmdList, Views, FeatureLevel, LocalQueriesPerView, DownsampleFactor);
+				if (!QueriesPerView.IsEmpty())
+				{
+					BeginOcclusionTests(RHICmdList, Views, FeatureLevel, QueriesPerView, DownsampleFactor);
+					QueriesPerView.Empty();
+				}
 			});
-		}
 	}
 	else
 	{
@@ -1490,7 +1496,8 @@ void FMobileSceneRenderer::RenderOcclusion(FRHICommandList& RHICmdList)
 
 	{
 		SCOPED_NAMED_EVENT(FMobileSceneRenderer_BeginOcclusionTests, FColor::Emerald);
-		const FViewOcclusionQueriesPerView QueriesPerView = AllocateOcclusionTests(Scene, VisibleLightInfos, Views);
+		FViewOcclusionQueriesPerView QueriesPerView;
+		AllocateOcclusionTests(QueriesPerView, Scene, VisibleLightInfos, Views);
 
 		if (QueriesPerView.Num())
 		{
