SceneCapture - Disable near plane correction when a custom projection matrix is used

diff --git a/Engine/Source/Runtime/Renderer/Private/SceneCaptureRendering.cpp b/Engine/Source/Runtime/Renderer/Private/SceneCaptureRendering.cpp
--- a/Engine/Source/Runtime/Renderer/Private/SceneCaptureRendering.cpp
+++ b/Engine/Source/Runtime/Renderer/Private/SceneCaptureRendering.cpp
@@ -791,9 +791,9 @@ static FSceneRenderer* CreateSceneRendererForSceneCapture(
 	SceneCaptureViewInfo.ViewRect = FIntRect(0, 0, RenderTargetSize.X, RenderTargetSize.Y);
 
 	// Use camera position correction for ortho scene captures
-	if(!SceneCaptureViewInfo.IsPerspectiveProjection())
+	if(USceneCaptureComponent2D * SceneCaptureComponent2D = Cast<USceneCaptureComponent2D>(SceneCaptureComponent))
 	{
-		if(USceneCaptureComponent2D * SceneCaptureComponent2D = Cast<USceneCaptureComponent2D>(SceneCaptureComponent))
+		if (!SceneCaptureViewInfo.IsPerspectiveProjection() && !SceneCaptureComponent2D->bUseCustomProjectionMatrix)
 		{
 			float NearPlane = SceneCaptureViewInfo.GetNearPlaneFromProjectionMatrix();
 			SceneCaptureViewInfo.UpdateOrthoNearPlane(NearPlane, true);
