Move some of the invalidation tracking back to per-primitive and on the CPU since we ended up doing WPO stuff slightly differently.
This lets us get rid of the GPU static->dynamic pass that was non-trivially expensive on scenes with many instances, and is currently unnecessary.
Behavior should be the same as before, just with the logic moved back to the CPU and per-primitive (which that part of it functionally still was).

Swap invaliation dispatch back to single per group and go back to expanding on the CPU to avoid redundantly hitting cases that were already CPU culled to specific lights. More of these will be present with Ola's recent changes.

diff --git a/Engine/Shaders/Private/VirtualShadowMaps/VirtualShadowMapCacheLoadBalancer.usf b/Engine/Shaders/Private/VirtualShadowMaps/VirtualShadowMapCacheLoadBalancer.usf
--- a/Engine/Shaders/Private/VirtualShadowMaps/VirtualShadowMapCacheLoadBalancer.usf
+++ b/Engine/Shaders/Private/VirtualShadowMaps/VirtualShadowMapCacheLoadBalancer.usf
@@ -7,6 +7,20 @@
 #include "../WaveOpUtil.ush"
 #include "../InstanceCulling/InstanceCullingLoadBalancer.ush"
 
+struct FInstanceInvalidationPayload
+{
+	int VirtualShadowMapId;
+	uint Flags;
+};
+
+FInstanceInvalidationPayload DecodeInstanceInvalidationPayload(uint Payload)
+{
+	FInstanceInvalidationPayload Result;
+	Result.VirtualShadowMapId = Payload >> VSM_INVALIDATION_PAYLOAD_FLAG_BITS;
+	Result.Flags = Payload & ((1U<<VSM_INVALIDATION_PAYLOAD_FLAG_BITS)-1U);
+	return Result;
+}
+
 [numthreads(CS_1D_GROUP_SIZE_X, 1, 1)]
 void InvalidateInstancePagesLoadBalancerCS(
 	uint3 GroupId : SV_GroupID,
@@ -29,15 +43,8 @@ void InvalidateInstancePagesLoadBalancerCS(
 	checkSlow(ShadowMapIndex < VirtualShadowMap.NumFullShadowMaps);
 	const int VirtualShadowMapId = VSM_MAX_SINGLE_PAGE_SHADOW_MAPS + ShadowMapIndex;
 
-	if (ShadowMapIndex == 0)
-	{
-		// Mark the instance as having caused an invalidation for future consideration of dynamic/static caching
-		// We want to do this even if it doesn't invalidate any pages this frame, as this is an indicator of an
-		// instance that is "doing things that may cause invalidations", and thus we want it in the dynamic pool.
-		VirtualShadowMapMarkInstanceAsInvalidated(InstanceId);
-
-		// TODO: If we want to change static->dynamic, we should probably just swap it right here
-	}
+	FInstanceInvalidationPayload Payload = DecodeInstanceInvalidationPayload(WorkSetup.Item.Payload);
+	bool bForceStatic = (Payload.Flags & VSM_INVALIDATION_PAYLOAD_FLAG_FORCE_STATIC) != 0;
 
-	InvalidateInstancePages(VirtualShadowMapId, InstanceId, false);
+	InvalidateInstancePages(Payload.VirtualShadowMapId, InstanceId, bForceStatic);
 }
diff --git a/Engine/Shaders/Private/VirtualShadowMaps/VirtualShadowMapCacheManagement.usf b/Engine/Shaders/Private/VirtualShadowMaps/VirtualShadowMapCacheManagement.usf
deleted file mode 100644
diff --git a/Engine/Shaders/Private/VirtualShadowMaps/VirtualShadowMapPageCacheCommon.ush b/Engine/Shaders/Private/VirtualShadowMaps/VirtualShadowMapPageCacheCommon.ush
--- a/Engine/Shaders/Private/VirtualShadowMaps/VirtualShadowMapPageCacheCommon.ush
+++ b/Engine/Shaders/Private/VirtualShadowMaps/VirtualShadowMapPageCacheCommon.ush
@@ -8,17 +8,6 @@ VirtualShadowMapPageCacheCommon.ush:
 #include "../Nanite/NaniteDataDecode.ush"
 #include "VirtualShadowMapProjectionStructs.ush"
 
-#define VSM_CACHE_AS_STATIC 1
-#define VSM_CACHE_AS_DYNAMIC 0
-
-RWStructuredBuffer<uint> OutVirtualShadowMapLastInstanceInvalidatedFrame;
-
-void VirtualShadowMapMarkInstanceAsInvalidated(uint InstanceId)
-{
-	checkStructuredBufferAccessSlow(OutVirtualShadowMapLastInstanceInvalidatedFrame, InstanceId);
-	OutVirtualShadowMapLastInstanceInvalidatedFrame[InstanceId] = GetGPUSceneFrameNumber();
-}
-
 bool ShouldMaterialInvalidateShadowCache(FPrimitiveSceneData PrimitiveData, bool bEnableWPO)
 {
 	// Determine if this instance will cause page invalidations.
@@ -54,6 +43,13 @@ bool VirtualShadowMapIsWPOAllowed(FPrimitiveSceneData PrimitiveData, int Virtual
 	return bEvaluateWPO;
 }
 
+bool GetCachePrimitiveAsDynamic(uint PersistentPrimitiveIndex)
+{
+	uint PrimitiveIdWordOffset = PersistentPrimitiveIndex / 32U;
+	uint PrimitiveIdWordMask = 1U << (PersistentPrimitiveIndex % 32U);
+	return (VirtualShadowMap.CachePrimitiveAsDynamic[PrimitiveIdWordOffset] & PrimitiveIdWordMask) != 0;
+}
+
 bool ShouldCacheInstanceAsStatic(uint InstanceId, bool bViewUncached, bool bAllowWPO)
 {
 	if (bViewUncached || bAllowWPO)
@@ -61,6 +57,7 @@ bool ShouldCacheInstanceAsStatic(uint InstanceId, bool bViewUncached, bool bAllo
 		return false;
 	}
 
-	uint CacheFlag = VirtualShadowMap.CacheInstanceAsStatic[InstanceId];
-	return (CacheFlag == VSM_CACHE_AS_STATIC);
+	FInstanceSceneData InstanceData = GetInstanceSceneData(InstanceId, false);
+	FPrimitiveSceneData PrimitiveData = GetPrimitiveData(InstanceData.PrimitiveId);
+	return !GetCachePrimitiveAsDynamic(PrimitiveData.PersistentPrimitiveIndex);
 }
diff --git a/Engine/Shaders/Shared/VirtualShadowMapDefinitions.h b/Engine/Shaders/Shared/VirtualShadowMapDefinitions.h
--- a/Engine/Shaders/Shared/VirtualShadowMapDefinitions.h
+++ b/Engine/Shaders/Shared/VirtualShadowMapDefinitions.h
@@ -28,6 +28,11 @@ VirtualShadowMapDefinitions.h: used in virtual shadow map shaders and C++ code t
 // Hard limit for max distant lights supported 8k for now - we may revise later. We need to keep them in a fixed range for now to make allocation easy and minimize overhead for indexing.
 #define VSM_MAX_SINGLE_PAGE_SHADOW_MAPS (1024U * 8U)
 
+#define VSM_INVALIDATION_PAYLOAD_FLAG_NONE                      0
+#define VSM_INVALIDATION_PAYLOAD_FLAG_FORCE_STATIC              (1 << 0)
+// 8 bit flags, 24 bit VSM ID
+#define VSM_INVALIDATION_PAYLOAD_FLAG_BITS                      8
+
 #ifdef __cplusplus
 #include "HLSLTypeAliases.h"
 
diff --git a/Engine/Source/Runtime/Renderer/Private/RendererScene.cpp b/Engine/Source/Runtime/Renderer/Private/RendererScene.cpp
--- a/Engine/Source/Runtime/Renderer/Private/RendererScene.cpp
+++ b/Engine/Source/Runtime/Renderer/Private/RendererScene.cpp
@@ -5537,7 +5537,6 @@ void FScene::UpdateAllPrimitiveSceneInfos(FRDGBuilder& GraphBuilder, EUpdateAllP
 		FVirtualShadowMapArrayCacheManager* CacheManager = GetVirtualShadowMapCache();
 
 		FVirtualShadowMapArrayCacheManager::FInvalidatingPrimitiveCollector InvalidatingPrimitiveCollector(CacheManager);
-		InvalidatingPrimitiveCollector.AddDynamicAndGPUPrimitives();
 
 		// Primitives that are tracked as always invalidating shadows, pipe through as transform updates
 		for (FPrimitiveSceneInfo* PrimitiveSceneInfo : ShadowScene->GetAlwaysInvalidatingPrimitives())
@@ -5565,8 +5564,6 @@ void FScene::UpdateAllPrimitiveSceneInfos(FRDGBuilder& GraphBuilder, EUpdateAllP
 		{
 			InvalidatingPrimitiveCollector.UpdatedTransform(CullDistance.Key->GetPrimitiveSceneInfo());
 		}
-		
-		InvalidatingPrimitiveCollector.Finalize();
 
 		CacheManager->ProcessInvalidations(GraphBuilder, SceneUB, InvalidatingPrimitiveCollector);
 	}
diff --git a/Engine/Source/Runtime/Renderer/Private/VirtualShadowMaps/VirtualShadowMapArray.cpp b/Engine/Source/Runtime/Renderer/Private/VirtualShadowMaps/VirtualShadowMapArray.cpp
--- a/Engine/Source/Runtime/Renderer/Private/VirtualShadowMaps/VirtualShadowMapArray.cpp
+++ b/Engine/Source/Runtime/Renderer/Private/VirtualShadowMaps/VirtualShadowMapArray.cpp
@@ -449,7 +449,7 @@ void FVirtualShadowMapArray::Initialize(
 	UniformParameters.PageRectBounds = GraphBuilder.CreateSRV(GSystemTextures.GetDefaultStructuredBuffer(GraphBuilder, sizeof(FIntVector4)));
 	UniformParameters.LightGridData = GraphBuilder.CreateSRV(GSystemTextures.GetDefaultStructuredBuffer(GraphBuilder, sizeof(uint32)));
 	UniformParameters.NumCulledLightsGrid = GraphBuilder.CreateSRV(GSystemTextures.GetDefaultStructuredBuffer(GraphBuilder, sizeof(uint32)));
-	UniformParameters.CacheInstanceAsStatic = GraphBuilder.CreateSRV(GSystemTextures.GetDefaultStructuredBuffer(GraphBuilder, sizeof(uint32)));
+	UniformParameters.CachePrimitiveAsDynamic = GraphBuilder.CreateSRV(GSystemTextures.GetDefaultStructuredBuffer(GraphBuilder, sizeof(uint32)));
 
 	if (bEnabled)
 	{
@@ -518,9 +518,7 @@ void FVirtualShadowMapArray::Initialize(
 		PhysicalPageMetaDataRDG = GraphBuilder.RegisterExternalBuffer(CacheManager->GetPhysicalPageMetaData());
 		UniformParameters.PhysicalPagePool = PhysicalPagePoolRDG;
 
-		CacheInstanceAsStaticRDG = GraphBuilder.RegisterExternalBuffer(CacheManager->GetCacheInstanceAsStatic());
-		LastInstanceInvalidatedFrameRDG = GraphBuilder.RegisterExternalBuffer(CacheManager->GetLastInstanceInvalidatedFrame());
-		UniformParameters.CacheInstanceAsStatic = GraphBuilder.CreateSRV(CacheInstanceAsStaticRDG);
+		UniformParameters.CachePrimitiveAsDynamic = GraphBuilder.CreateSRV(CacheManager->UploadCachePrimitiveAsDynamic(GraphBuilder));
 	}
 	else
 	{
diff --git a/Engine/Source/Runtime/Renderer/Private/VirtualShadowMaps/VirtualShadowMapArray.h b/Engine/Source/Runtime/Renderer/Private/VirtualShadowMaps/VirtualShadowMapArray.h
--- a/Engine/Source/Runtime/Renderer/Private/VirtualShadowMaps/VirtualShadowMapArray.h
+++ b/Engine/Source/Runtime/Renderer/Private/VirtualShadowMaps/VirtualShadowMapArray.h
@@ -155,7 +155,7 @@ BEGIN_GLOBAL_SHADER_PARAMETER_STRUCT(FVirtualShadowMapUniformParameters, )
 	SHADER_PARAMETER_RDG_BUFFER_SRV(StructuredBuffer<uint4>, PageRectBounds)
 	SHADER_PARAMETER_RDG_TEXTURE(Texture2DArray<uint>, PhysicalPagePool)
 
-	SHADER_PARAMETER_RDG_BUFFER_SRV(StructuredBuffer<uint>, CacheInstanceAsStatic)
+	SHADER_PARAMETER_RDG_BUFFER_SRV(StructuredBuffer<uint>, CachePrimitiveAsDynamic)
 
 	// Light grid with only the lights that have VSMs present
 	// Still references the original indices from the global light grid
@@ -374,9 +374,6 @@ public:
 	FRDGTextureRef HZBPhysicalRDG = nullptr;
 	FRDGBufferRef PhysicalPageMetaDataRDG = nullptr;
 
-	FRDGBufferRef CacheInstanceAsStaticRDG = nullptr;
-	FRDGBufferRef LastInstanceInvalidatedFrameRDG = nullptr;
-
 	// Buffer that serves as the page table for all virtual shadow maps
 	FRDGBufferRef PageTableRDG = nullptr;
 		
diff --git a/Engine/Source/Runtime/Renderer/Private/VirtualShadowMaps/VirtualShadowMapCacheManager.cpp b/Engine/Source/Runtime/Renderer/Private/VirtualShadowMaps/VirtualShadowMapCacheManager.cpp
--- a/Engine/Source/Runtime/Renderer/Private/VirtualShadowMaps/VirtualShadowMapCacheManager.cpp
+++ b/Engine/Source/Runtime/Renderer/Private/VirtualShadowMaps/VirtualShadowMapCacheManager.cpp
@@ -305,7 +305,7 @@ void FVirtualShadowMapArrayCacheManager::FShadowInvalidatingInstancesImplementat
 
 void FVirtualShadowMapArrayCacheManager::FShadowInvalidatingInstancesImplementation::AddInstanceRange(uint32 InstanceSceneDataOffset, uint32 NumInstanceSceneDataEntries)
 {
-	PrimitiveInstancesToInvalidate.Add(FVirtualShadowMapInstanceRange{int32(InstanceSceneDataOffset), int32(NumInstanceSceneDataEntries)});
+	PrimitiveInstancesToInvalidate.Add(FVirtualShadowMapInstanceRange{int32(InstanceSceneDataOffset), int32(NumInstanceSceneDataEntries), false});
 }		
 
 void FVirtualShadowMapPerLightCacheEntry::Invalidate()
@@ -313,31 +313,46 @@ void FVirtualShadowMapPerLightCacheEntry::Invalidate()
 	Prev.RenderedFrameNumber = -1;
 }
 
+static uint32 EncodeInstanceInvalidationPayload(int32 VirtualShadowMapId, uint32 Flags = VSM_INVALIDATION_PAYLOAD_FLAG_NONE)
+{
+	check(VirtualShadowMapId >= 0);		// Should not be INDEX_NONE by this point
+
+	uint32 Payload = Flags;
+	Payload = Payload | (uint32(VirtualShadowMapId) << VSM_INVALIDATION_PAYLOAD_FLAG_BITS);
+	return Payload;
+}
+
 FVirtualShadowMapArrayCacheManager::FInvalidatingPrimitiveCollector::FInvalidatingPrimitiveCollector(FVirtualShadowMapArrayCacheManager* InVirtualShadowMapArrayCacheManager)
 	: Scene(*InVirtualShadowMapArrayCacheManager->Scene)
 	, GPUScene(InVirtualShadowMapArrayCacheManager->Scene->GPUScene)
 	, Manager(*InVirtualShadowMapArrayCacheManager)
 {
+	uint32 Num = Manager.CachePrimitiveAsDynamic.Num();
+	InvalidatedPrimitives.SetNum(Num, false);
+	RemovedPrimitives.SetNum(Num, false);
 }
 
-void FVirtualShadowMapArrayCacheManager::FInvalidatingPrimitiveCollector::AddDynamicAndGPUPrimitives()
-{
-	// Add and clear pending invalidations enqueued since the last invalidation (scene update)
-	// This includes "dynamic primitives"
-	for (const FVirtualShadowMapInstanceRange& Range : Manager.ShadowInvalidatingInstancesImplementation.PrimitiveInstancesToInvalidate)
-	{
-		Instances.Add(Range.InstanceSceneDataOffset, Range.NumInstanceSceneDataEntries, 0);
-	}
-
-	Manager.ShadowInvalidatingInstancesImplementation.PrimitiveInstancesToInvalidate.Reset();
-
+void FVirtualShadowMapArrayCacheManager::FInvalidatingPrimitiveCollector::AddPrimitivesToInvalidate()
+{	
 	for (auto& CacheEntry : Manager.CacheEntries)
 	{
-		for (const FVirtualShadowMapInstanceRange& Range : CacheEntry.Value->PrimitiveInstancesToInvalidate)
+		for (const auto& SmCacheEntry : CacheEntry.Value->ShadowMapEntries)
 		{
-			// TODO: Do we ever need to invalidate these "post" update?
-			Instances.Add(Range.InstanceSceneDataOffset, Range.NumInstanceSceneDataEntries, 0);
+			const uint32 Payload = EncodeInstanceInvalidationPayload(SmCacheEntry.CurrentVirtualShadowMapId);
+			const uint32 PayloadForceStatic = EncodeInstanceInvalidationPayload(SmCacheEntry.CurrentVirtualShadowMapId, VSM_INVALIDATION_PAYLOAD_FLAG_FORCE_STATIC);
+
+			for (const FVirtualShadowMapInstanceRange& Range : Manager.ShadowInvalidatingInstancesImplementation.PrimitiveInstancesToInvalidate)
+			{
+				Instances.Add(Range.InstanceSceneDataOffset, Range.NumInstanceSceneDataEntries, Range.bForceInvalidateStatic ? PayloadForceStatic : Payload);
+			}
+
+			for (const FVirtualShadowMapInstanceRange& Range : CacheEntry.Value->PrimitiveInstancesToInvalidate)
+			{
+				Instances.Add(Range.InstanceSceneDataOffset, Range.NumInstanceSceneDataEntries, Range.bForceInvalidateStatic ? PayloadForceStatic : Payload);
+			}
 		}
+
+		Manager.ShadowInvalidatingInstancesImplementation.PrimitiveInstancesToInvalidate.Reset();
 		CacheEntry.Value->PrimitiveInstancesToInvalidate.Reset();
 	}
 }
@@ -356,6 +371,11 @@ void FVirtualShadowMapArrayCacheManager::FInvalidatingPrimitiveCollector::AddInv
 	const FPersistentPrimitiveIndex PersistentPrimitiveIndex = PrimitiveSceneInfo->GetPersistentIndex();
 	const EPrimitiveDirtyState DirtyState = GPUScene.GetPrimitiveDirtyState(PersistentPrimitiveIndex);
 
+	if (bRemovedPrimitive)
+	{
+		RemovedPrimitives[PersistentPrimitiveIndex.Index] = true;
+	}
+
 	// TODO: Early out on stuff that doesn't cast shadows on the CPU?
 
 	// Suppress invalidations from moved primitives that are marked to behave as if they were static.
@@ -378,6 +398,7 @@ void FVirtualShadowMapArrayCacheManager::FInvalidatingPrimitiveCollector::AddInv
 	}
 
 	//UE_LOG(LogRenderer, Warning, TEXT("Invalidating instances: %u, %u"), InstanceSceneDataOffset, NumInstanceSceneDataEntries);
+	InvalidatedPrimitives[PersistentPrimitiveIndex.Index] = true;
 
 	// Nanite meshes need special handling because they don't get culled on CPU, thus always process invalidations for those
 	const bool bIsNaniteMesh = Scene.PrimitiveFlagsCompact[PrimitiveID].bIsNaniteMesh;
@@ -393,16 +414,15 @@ void FVirtualShadowMapArrayCacheManager::FInvalidatingPrimitiveCollector::AddInv
 				CachedPrimitives[PersistentPrimitiveIndex.Index] = false;
 			}
 
-			Instances.Add(InstanceSceneDataOffset, NumInstanceSceneDataEntries, 0);
+			// Add item for each shadow map explicitly, inflates host data but improves load balancing,
+			for (const auto& SmCacheEntry : CacheEntry.Value->ShadowMapEntries)
+			{
+				Instances.Add(InstanceSceneDataOffset, NumInstanceSceneDataEntries, EncodeInstanceInvalidationPayload(SmCacheEntry.CurrentVirtualShadowMapId));
+			}
 		}
 	}
 }
 
-void FVirtualShadowMapArrayCacheManager::FInvalidatingPrimitiveCollector::Finalize()
-{
-	Instances.FinalizeBatches();
-}
-
 
 FVirtualShadowMapFeedback::FVirtualShadowMapFeedback()
 {
@@ -675,47 +695,6 @@ void FVirtualShadowMapArrayCacheManager::SetPhysicalPoolSize(FRDGBuilder& GraphB
 		bInvalidateCache = true;
 	}
 
-	// TODO: Probably better entry point for this, but works here for now	
-	const uint32 TargetMaxInstanceCount = FMath::RoundUpToPowerOfTwo(FMath::Max(1024u * 128u, Scene->GPUScene.GetInstanceIdUpperBoundGPU()));	
-	
-	if (!CacheInstanceAsStatic || CacheInstanceAsStatic->Desc.NumElements < TargetMaxInstanceCount ||
-		!LastInstanceInvalidatedFrame || LastInstanceInvalidatedFrame->Desc.NumElements < TargetMaxInstanceCount)
-	{
-		// TODO: See if there's a better way to do a buffer resize
-		FRDGBufferRef NewCacheInstanceAsStaticRDG = GraphBuilder.CreateBuffer(
-			FRDGBufferDesc::CreateStructuredDesc(sizeof(uint32), TargetMaxInstanceCount),
-			TEXT("Shadow.Virtual.CacheInstanceAsStatic"));
-
-		FRDGBufferRef NewLastInstanceInvalidatedFrameRDG = GraphBuilder.CreateBuffer(
-			FRDGBufferDesc::CreateStructuredDesc(sizeof(uint32), TargetMaxInstanceCount),
-			TEXT("Shadow.Virtual.LastInstanceInvalidatedFrame"));
-
-		// TODO: Could almost certainly optimize this into one pass with the copy below but we're not meant to be
-		// doing this frequently at the moment.
-		AddClearUAVPass(GraphBuilder, GraphBuilder.CreateUAV(NewCacheInstanceAsStaticRDG), 0);
-		AddClearUAVPass(GraphBuilder, GraphBuilder.CreateUAV(NewLastInstanceInvalidatedFrameRDG), MAX_uint32);
-		 
-		if (CacheInstanceAsStatic)
-		{
-			UE_LOG(LogRenderer, Display, TEXT("Reallocating Shadow.Virtual.CacheInstanceAsStatic due to instance count change."));
-
-			FRDGBufferRef CacheInstanceAsStaticRDG = GraphBuilder.RegisterExternalBuffer(CacheInstanceAsStatic);
-			const uint64 NumBytes = CacheInstanceAsStatic->Desc.NumElements * CacheInstanceAsStatic->Desc.BytesPerElement;			
-			AddCopyBufferPass(GraphBuilder, NewCacheInstanceAsStaticRDG, 0, CacheInstanceAsStaticRDG, 0, NumBytes);
-		}
-
-		if (LastInstanceInvalidatedFrame)
-		{
-			FRDGBufferRef InstanceInvalidatedRecentlyRDG = GraphBuilder.RegisterExternalBuffer(LastInstanceInvalidatedFrame);
-			const uint64 NumBytes = LastInstanceInvalidatedFrame->Desc.NumElements * LastInstanceInvalidatedFrame->Desc.BytesPerElement;			
-			AddCopyBufferPass(GraphBuilder, NewLastInstanceInvalidatedFrameRDG, 0, InstanceInvalidatedRecentlyRDG, 0, NumBytes);
-		}
-		
-		// Persistent, so we extract it immediately (over top of old versions)
-		CacheInstanceAsStatic = GraphBuilder.ConvertToExternalBuffer(NewCacheInstanceAsStaticRDG);
-		LastInstanceInvalidatedFrame = GraphBuilder.ConvertToExternalBuffer(NewLastInstanceInvalidatedFrameRDG);
-	}
-
 	if (bInvalidateCache)
 	{
 		Invalidate(GraphBuilder);
@@ -728,8 +707,6 @@ void FVirtualShadowMapArrayCacheManager::FreePhysicalPool(FRDGBuilder& GraphBuil
 	{
 		PhysicalPagePool = nullptr;
 		PhysicalPageMetaData = nullptr;
-		CacheInstanceAsStatic = nullptr;
-		LastInstanceInvalidatedFrame = nullptr;
 		Invalidate(GraphBuilder);
 	}
 }
@@ -806,6 +783,23 @@ bool FVirtualShadowMapArrayCacheManager::IsHZBDataAvailable()
 	return HZBPhysicalPagePool && PrevBuffers.PageTable && PrevBuffers.PageFlags;
 }
 
+FRDGBufferRef FVirtualShadowMapArrayCacheManager::UploadCachePrimitiveAsDynamic(FRDGBuilder& GraphBuilder) const
+{
+	const uint32 NumElements = FMath::Max(1, FMath::DivideAndRoundUp(CachePrimitiveAsDynamic.Num(), 32));
+	//GraphBuilder.CreateBuffer(FRDGBufferDesc::CreateStructuredDesc(sizeof(uint32), BufferElements), 
+	
+	FRDGBufferRef CachePrimitiveAsDynamicRDG = CreateStructuredBuffer(
+		GraphBuilder,
+		TEXT("CachePrimitiveAsDynamic"),
+		sizeof(uint32),
+		NumElements,
+		CachePrimitiveAsDynamic.GetData(),
+		FMath::DivideAndRoundUp(CachePrimitiveAsDynamic.Num(), 8)	// Size in bytes of initial data
+	);
+
+	return CachePrimitiveAsDynamicRDG;
+}
+
 TSharedPtr<FVirtualShadowMapPerLightCacheEntry> FVirtualShadowMapArrayCacheManager::FindCreateLightCacheEntry(
 	int32 LightSceneId, uint32 ViewUniqueID, uint32 NumShadowMaps)
 {
@@ -852,7 +846,11 @@ void FVirtualShadowMapPerLightCacheEntry::OnPrimitiveRendered(const FPrimitiveSc
 		// Skip if the invalidation mode is NOT auto (because Always will do it elsewhere & the others should prevent this).
 		if (PrimitiveSceneInfo->Proxy->HasDeformableMesh() && PrimitiveSceneInfo->Proxy->GetShadowCacheInvalidationBehavior() == EShadowCacheInvalidationBehavior::Auto)
 		{
-			PrimitiveInstancesToInvalidate.Add(FVirtualShadowMapInstanceRange{ PrimitiveSceneInfo->GetInstanceSceneDataOffset(), PrimitiveSceneInfo->GetNumInstanceSceneDataEntries() });
+			PrimitiveInstancesToInvalidate.Add(FVirtualShadowMapInstanceRange{ 
+				PrimitiveSceneInfo->GetInstanceSceneDataOffset(),
+				PrimitiveSceneInfo->GetNumInstanceSceneDataEntries(),
+				false
+			});
 		}
 	}
 }
@@ -987,7 +985,7 @@ void FVirtualShadowMapArrayCacheManager::ExtractFrameData(
 			PrevUniformParameters.PageFlags = nullptr;
 			PrevUniformParameters.LightGridData = nullptr;
 			PrevUniformParameters.NumCulledLightsGrid = nullptr;
-			PrevUniformParameters.CacheInstanceAsStatic = nullptr;
+			PrevUniformParameters.CachePrimitiveAsDynamic = nullptr;
 		}
 
 		// propagate current-frame primitive state to cache entry
@@ -1190,6 +1188,19 @@ void FVirtualShadowMapArrayCacheManager::OnSceneChange()
 		CacheEntry.Value->RenderedPrimitives.SetNum(MaxPersistentPrimitiveIndex, false);
 	}
 
+	// TODO: Initialize new primitives based on their mobility; need a way to know which ones are newly created though
+	CachePrimitiveAsDynamic.SetNum(MaxPersistentPrimitiveIndex, false);	
+	if (MaxPersistentPrimitiveIndex > LastPrimitiveInvalidatedFrame.Num())
+	{
+		const uint32 OldSize = LastPrimitiveInvalidatedFrame.Num();
+		LastPrimitiveInvalidatedFrame.SetNumUninitialized(MaxPersistentPrimitiveIndex);
+		for (int32 It = OldSize; It < MaxPersistentPrimitiveIndex; ++It)
+		{
+			// Unknown last invalidation
+			LastPrimitiveInvalidatedFrame[It] = 0xFFFFFFFF;
+		}
+	}
+
 	// Do instance-based GPU allocations here too? For now we do them lazily each frame when the FVirtualShadowMapArray gets constructed
 }
 
@@ -1218,7 +1229,6 @@ public:
 	BEGIN_SHADER_PARAMETER_STRUCT(FParameters, )
 		SHADER_PARAMETER_STRUCT_INCLUDE(FInvalidatePagesParameters, InvalidatePagesParameters)
 		SHADER_PARAMETER_STRUCT_INCLUDE(FGPUScene::FInstanceGPULoadBalancer::FShaderParameters, LoadBalancerParameters)
-		SHADER_PARAMETER_RDG_BUFFER_UAV(RWStructuredBuffer<uint>, OutVirtualShadowMapLastInstanceInvalidatedFrame)
 	END_SHADER_PARAMETER_STRUCT()
 
 	// This is probably fine even in instance list mode
@@ -1241,106 +1251,105 @@ public:
 };
 IMPLEMENT_GLOBAL_SHADER(FInvalidateInstancePagesLoadBalancerCS, "/Engine/Private/VirtualShadowMaps/VirtualShadowMapCacheLoadBalancer.usf", "InvalidateInstancePagesLoadBalancerCS", SF_Compute);
 
-class FInvalidateInstancePagesListCS : public FVirtualShadowMapPageManagementShader
+void FVirtualShadowMapArrayCacheManager::UpdateCachePrimitiveAsDynamic(FInvalidatingPrimitiveCollector& InvalidatingPrimitiveCollector)
 {
-	DECLARE_GLOBAL_SHADER(FInvalidateInstancePagesListCS);
-	SHADER_USE_PARAMETER_STRUCT(FInvalidateInstancePagesListCS, FVirtualShadowMapPageManagementShader)
-
-	class FDebugDim : SHADER_PERMUTATION_BOOL("ENABLE_DEBUG_MODE");
-	class FUseHzbDim : SHADER_PERMUTATION_BOOL("USE_HZB_OCCLUSION");
-	using FPermutationDomain = TShaderPermutationDomain<FUseHzbDim, FDebugDim>;
+	//UE_LOG(LogRenderer, Display, TEXT("VirtualShadowMapCacheManager: UpdateCachePrimitiveAsDynamic"));
 
-public:
-	BEGIN_SHADER_PARAMETER_STRUCT(FParameters, )
-		SHADER_PARAMETER_STRUCT_INCLUDE(FInvalidatePagesParameters, InvalidatePagesParameters)
-		SHADER_PARAMETER_RDG_BUFFER_SRV(StructuredBuffer<uint>, InstanceInvalidationList)
-		SHADER_PARAMETER(uint32, MaxInstanceInvalidations)
-		RDG_BUFFER_ACCESS(IndirectArgs, ERHIAccess::IndirectArgs)
-	END_SHADER_PARAMETER_STRUCT()
-};
-IMPLEMENT_GLOBAL_SHADER(FInvalidateInstancePagesListCS, "/Engine/Private/VirtualShadowMaps/VirtualShadowMapCacheManagement.usf", "InvalidateInstancePagesListCS", SF_Compute);
+	const uint32 SceneFrameNumber = Scene->GetFrameNumberRenderThread();
+	const uint32 FramesStaticThreshold = CVarFramesStaticThreshold.GetValueOnRenderThread();
 
+	// Update the cache states of things that are being invalidated
+	for (TConstSetBitIterator<> BitIt(InvalidatingPrimitiveCollector.InvalidatedPrimitives); BitIt; ++BitIt)
+	{
+		int32 PersistentPrimitiveIndex = BitIt.GetIndex();
+		// Any invalidations mean we set this primitive to dynamic. We already added an invalidation otherwise
+		// we wouldn't be here, so no need to add another.
+		CachePrimitiveAsDynamic[PersistentPrimitiveIndex] = true;
+		LastPrimitiveInvalidatedFrame[PersistentPrimitiveIndex] = SceneFrameNumber;
+	}
 
-class FCacheInstanceAsStaticCS : public FVirtualShadowMapPageManagementShader
-{
-	DECLARE_GLOBAL_SHADER(FCacheInstanceAsStaticCS);
-	SHADER_USE_PARAMETER_STRUCT(FCacheInstanceAsStaticCS, FVirtualShadowMapPageManagementShader)
+	// Zero out anything that was being removed
+	// NOTE: This will be redundant with the invalidated stuff, but shouldn't be a big deal
+	for (TConstSetBitIterator<> BitIt(InvalidatingPrimitiveCollector.RemovedPrimitives); BitIt; ++BitIt)
+	{
+		int32 PersistentPrimitiveIndex = BitIt.GetIndex();
+		// TODO: We probably want to start new primitives as dynamic by default instead, but we don't want to have
+		// to loop over all of them and try and get their PrimitiveSceneInfo every frame for invalid ones
+		CachePrimitiveAsDynamic[PersistentPrimitiveIndex] = false;
+		LastPrimitiveInvalidatedFrame[PersistentPrimitiveIndex] = 0xFFFFFFFF;
 
-	BEGIN_SHADER_PARAMETER_STRUCT(FParameters, )
-		SHADER_PARAMETER_RDG_UNIFORM_BUFFER(FVirtualShadowMapUniformParameters, VirtualShadowMap)
-		SHADER_PARAMETER_RDG_UNIFORM_BUFFER(FSceneUniformParameters, Scene)
-		SHADER_PARAMETER_RDG_BUFFER_UAV(RWStructuredBuffer<uint>, OutVirtualShadowMapLastInstanceInvalidatedFrame)
-		SHADER_PARAMETER_RDG_BUFFER_UAV(RWStructuredBuffer<uint>, OutInstanceInvalidationList)
-		SHADER_PARAMETER_RDG_BUFFER_UAV(RWStructuredBuffer<uint>, OutCacheInstanceAsStatic)
-		SHADER_PARAMETER(uint32, MaxInstanceInvalidations)
-		SHADER_PARAMETER(uint32, FramesStaticThreshold)
-	END_SHADER_PARAMETER_STRUCT()
-};
-IMPLEMENT_GLOBAL_SHADER(FCacheInstanceAsStaticCS, "/Engine/Private/VirtualShadowMaps/VirtualShadowMapCacheManagement.usf", "UpdateCacheInstanceAsStaticCS", SF_Compute);
+		//UE_LOG(LogRenderer, Display, TEXT("VirtualShadowMapCacheManager: Removing primitive %d!"), PersistentPrimitiveIndex);
+	}
 
-class FSetInstanceInvalidationListIndirectArgsCS : public FVirtualShadowMapPageManagementShader
-{
-	DECLARE_GLOBAL_SHADER(FSetInstanceInvalidationListIndirectArgsCS);
-	SHADER_USE_PARAMETER_STRUCT(FSetInstanceInvalidationListIndirectArgsCS, FVirtualShadowMapPageManagementShader)
+	// Finally check anything that is currently dynamic to see if it has not invalidated for long enough that
+	// we should move it back to static
+	for (TConstSetBitIterator<> BitIt(CachePrimitiveAsDynamic); BitIt; ++BitIt)
+	{
+		int32 PersistentPrimitiveIndex = BitIt.GetIndex();
 
-	BEGIN_SHADER_PARAMETER_STRUCT(FParameters, )
-		SHADER_PARAMETER_RDG_UNIFORM_BUFFER(FVirtualShadowMapUniformParameters, VirtualShadowMap)
-		SHADER_PARAMETER_RDG_BUFFER_SRV(StructuredBuffer<uint>, InstanceInvalidationList)
-		SHADER_PARAMETER_RDG_BUFFER_UAV(RWBuffer<uint>, OutIndirectArgsBuffer)
-		SHADER_PARAMETER(uint32, InvalidateInstancePagesCSGroupSize)
-		SHADER_PARAMETER(uint32, MaxInstanceInvalidations)
-	END_SHADER_PARAMETER_STRUCT()
-};
-IMPLEMENT_GLOBAL_SHADER(FSetInstanceInvalidationListIndirectArgsCS, "/Engine/Private/VirtualShadowMaps/VirtualShadowMapCacheManagement.usf", "SetInstanceInvalidationListIndirectArgs", SF_Compute);
+		const uint32 LastInvalidationFrame = LastPrimitiveInvalidatedFrame[PersistentPrimitiveIndex];
+		// Note: cleared to MAX_uint32; treated as "unknown/no invalidations"
+		const uint32 InvalidationAge = 
+			SceneFrameNumber >= LastInvalidationFrame ?
+			(SceneFrameNumber - LastInvalidationFrame) :
+			0xFFFFFFFF;
 
+		const bool bWantStatic = InvalidationAge > FramesStaticThreshold;
+		if (bWantStatic)
+		{
+			// Add invalidation and swap it to static
+			FPersistentPrimitiveIndex WrappedIndex;
+			WrappedIndex.Index = PersistentPrimitiveIndex;
+			FPrimitiveSceneInfo* PrimitiveSceneInfo = Scene->GetPrimitiveSceneInfo(WrappedIndex);
+			if (PrimitiveSceneInfo)
+			{
+				FVirtualShadowMapInstanceRange InstanceRange{
+					PrimitiveSceneInfo->GetInstanceSceneDataOffset(),
+					PrimitiveSceneInfo->GetNumInstanceSceneDataEntries(),
+					true		// Always invalidate static since we are swapping between the two
+				};
+
+				// Add it to the invalidation list for every light
+				// NOTE: These will get added to the invalidation list by AddDynamicAndGPUPrimitives
+				for (auto& CacheEntry : CacheEntries)
+				{
+					CacheEntry.Value->PrimitiveInstancesToInvalidate.Add(InstanceRange);
+				}
+			}
+			else
+			{
+				// This seems to still happen very occasionally... presumably a remove gets "missed" somehow and thus we try and transition
+				// something that is no longer valid back to static. This could also potentially mean we incorrect transition a new thing that
+				// grabbed this slot back to static, but that is less likely as the addition would trigger a separate invalidation.
+				// Not much we can do here currently other than ignore it and move on
+				UE_LOG(LogRenderer, Display, TEXT("VirtualShadowMapCacheManager: Invalid persistent primitive index %d, age %u!"), PersistentPrimitiveIndex, InvalidationAge);
+				LastPrimitiveInvalidatedFrame[PersistentPrimitiveIndex] = 0xFFFFFFFF;
+			}
+			// NOTE: This is safe with the current set bit iterator, but should maybe use a temp array for future safety?
+			CachePrimitiveAsDynamic[PersistentPrimitiveIndex] = false;
+		}
+	}
+}
 
-void FVirtualShadowMapArrayCacheManager::ProcessInvalidations(FRDGBuilder& GraphBuilder, FSceneUniformBuffer &SceneUniformBuffer, const FInvalidatingPrimitiveCollector& InvalidatingPrimitiveCollector)
+void FVirtualShadowMapArrayCacheManager::ProcessInvalidations(FRDGBuilder& GraphBuilder, FSceneUniformBuffer &SceneUniformBuffer, FInvalidatingPrimitiveCollector& InvalidatingPrimitiveCollector)
 {
 	if (IsCacheDataAvailable() && PrevUniformParameters.NumFullShadowMaps > 0)
 	{
 		RDG_EVENT_SCOPE(GraphBuilder, "Shadow.Virtual.ProcessInvalidations");
 
+		// NOTE: Important that we get some of these parameters (ex. CachePrimitiveAsDynamic) before
+		// we update them as the shader needs to know the previous cache states for invalidation.
 		FInvalidationPassCommon InvalidationPassCommon = GetUniformParametersForInvalidation(GraphBuilder, SceneUniformBuffer);
 
+		// Check whether we want to swap any cache states and add any invalidations to that end as well
+		UpdateCachePrimitiveAsDynamic(InvalidatingPrimitiveCollector);
+		InvalidatingPrimitiveCollector.AddPrimitivesToInvalidate();
+		InvalidatingPrimitiveCollector.Instances.FinalizeBatches();
+
 		if (!InvalidatingPrimitiveCollector.Instances.IsEmpty())
 		{
 			ProcessInvalidations(GraphBuilder, InvalidationPassCommon, InvalidatingPrimitiveCollector.Instances);
 		}
-
-		// TODO: Some of this only needs to run with separate static caching enabled
-		{
-			const uint32 MaxInstanceInvalidations = CVarMaxInstanceTransitionsPerFrame.GetValueOnRenderThread();
-			const uint32 InstanceCount = CacheInstanceAsStatic->Desc.NumElements;
-
-			// [0] is the counter
-			FRDGBufferRef InstanceInvalidationList = GraphBuilder.CreateBuffer(FRDGBufferDesc::CreateStructuredDesc(sizeof(int32), MaxInstanceInvalidations + 1), TEXT("Shadow.Virtual.InstanceInvalidationList"));
-			// Actually only need to zero out the first element so can probably chain this on to a different pass
-			AddClearUAVPass(GraphBuilder, GraphBuilder.CreateUAV(InstanceInvalidationList), 0);
-
-			{
-				FCacheInstanceAsStaticCS::FParameters* PassParameters = GraphBuilder.AllocParameters<FCacheInstanceAsStaticCS::FParameters>();
-				PassParameters->VirtualShadowMap = InvalidationPassCommon.VirtualShadowMapUniformBuffer;
-				PassParameters->Scene = InvalidationPassCommon.SceneUniformBuffer;
-				PassParameters->OutVirtualShadowMapLastInstanceInvalidatedFrame = GraphBuilder.CreateUAV(GraphBuilder.RegisterExternalBuffer(LastInstanceInvalidatedFrame));
-				PassParameters->OutInstanceInvalidationList = GraphBuilder.CreateUAV(InstanceInvalidationList);
-				PassParameters->OutCacheInstanceAsStatic = GraphBuilder.CreateUAV(GraphBuilder.RegisterExternalBuffer(CacheInstanceAsStatic));
-				PassParameters->MaxInstanceInvalidations = MaxInstanceInvalidations;
-				PassParameters->FramesStaticThreshold = CVarFramesStaticThreshold.GetValueOnRenderThread();
-
-				auto ComputeShader = GetGlobalShaderMap(Scene->GetFeatureLevel())->GetShader<FCacheInstanceAsStaticCS>();
-
-				FComputeShaderUtils::AddPass(
-					GraphBuilder,
-					RDG_EVENT_NAME("UpdateCacheInstanceAsStaticCS"),
-					ComputeShader,
-					PassParameters,
-					// TODO: Do we need to do the 3d grid thing for this?
-					FIntVector(FMath::DivideAndRoundUp(InstanceCount, FCacheInstanceAsStaticCS::DefaultCSGroupX), 1, 1)
-				);
-			}
-
-			ProcessInvalidations(GraphBuilder, InvalidationPassCommon, InstanceInvalidationList, MaxInstanceInvalidations);
-		}
 	}
 }
 
@@ -1367,8 +1376,7 @@ FVirtualShadowMapArrayCacheManager::FInvalidationPassCommon FVirtualShadowMapArr
 		UniformParameters->PageTable = RegExtCreateSrv(PrevBuffers.PageTable, TEXT("Shadow.Virtual.PrevPageTable"));
 		UniformParameters->PageFlags = RegExtCreateSrv(PrevBuffers.PageFlags, TEXT("Shadow.Virtual.PrevPageFlags"));
 		UniformParameters->PageRectBounds = RegExtCreateSrv(PrevBuffers.PageRectBounds, TEXT("Shadow.Virtual.PrevPageRectBounds"));
-		// TODO: Unclear what we want here currently
-		UniformParameters->CacheInstanceAsStatic = RegExtCreateSrv(CacheInstanceAsStatic, TEXT("Shadow.Virtual.CacheInstanceAsStatic"));
+		UniformParameters->CachePrimitiveAsDynamic = GraphBuilder.CreateSRV(UploadCachePrimitiveAsDynamic(GraphBuilder));
 
 		// Unused in this path... may be a better way to handle this
 		UniformParameters->PhysicalPagePool = GSystemTextures.GetZeroUIntArrayAtomicCompatDummy(GraphBuilder);
@@ -1418,15 +1426,11 @@ void FVirtualShadowMapArrayCacheManager::ProcessInvalidations(
 	FInvalidateInstancePagesLoadBalancerCS::FParameters* PassParameters = GraphBuilder.AllocParameters<FInvalidateInstancePagesLoadBalancerCS::FParameters>();
 
 	SetInvalidateInstancePagesParameters(GraphBuilder, InvalidationPassCommon, &PassParameters->InvalidatePagesParameters);
-	PassParameters->OutVirtualShadowMapLastInstanceInvalidatedFrame = GraphBuilder.CreateUAV(GraphBuilder.RegisterExternalBuffer(LastInstanceInvalidatedFrame));
+	Instances.UploadFinalized(GraphBuilder).GetShaderParameters(GraphBuilder, PassParameters->LoadBalancerParameters);
 
 	FInvalidateInstancePagesLoadBalancerCS::FPermutationDomain PermutationVector;
 	PermutationVector.Set<FInvalidateInstancePagesLoadBalancerCS::FUseHzbDim>(PassParameters->InvalidatePagesParameters.HZBTexture != nullptr);
-		
-	// Launch a copy of each group for each VSM to improve load balance
-	const int32 NumGroupsPerBatch = int32(PrevUniformParameters.NumFullShadowMaps);
-	Instances.UploadFinalized(GraphBuilder, FInstanceGPULoadBalancer::DefaultRDGInitialDataFlags, NumGroupsPerBatch).GetShaderParameters(GraphBuilder, PassParameters->LoadBalancerParameters);
-
+	
 	auto ComputeShader = GetGlobalShaderMap(Scene->GetFeatureLevel())->GetShader<FInvalidateInstancePagesLoadBalancerCS>(PermutationVector);
 
 	FComputeShaderUtils::AddPass(
@@ -1434,66 +1438,10 @@ void FVirtualShadowMapArrayCacheManager::ProcessInvalidations(
 		RDG_EVENT_NAME("InvalidateInstancePagesLoadBalancerCS (%d batches)", Instances.GetBatches().Num()),
 		ComputeShader,
 		PassParameters,
-		Instances.GetWrappedCsGroupCount(NumGroupsPerBatch)
+		Instances.GetWrappedCsGroupCount()
 	);
 }
 
-void FVirtualShadowMapArrayCacheManager::ProcessInvalidations(
-	FRDGBuilder& GraphBuilder,
-	const FInvalidationPassCommon& InvalidationPassCommon,
-	FRDGBufferRef InstanceInvalidationList,
-	uint32 MaxInstanceInvalidations) const
-{
-	check(InvalidationPassCommon.UniformParameters->NumFullShadowMaps > 0);
-	check(InstanceInvalidationList);
-
-	// Set indirect args
-	FRDGBufferRef IndirectArgs = GraphBuilder.CreateBuffer(FRDGBufferDesc::CreateIndirectDesc<FRHIDispatchIndirectParameters>(), TEXT("Shadow.Virtual.InstanceInvalidationIndirectArgs"));
-	{
-		FSetInstanceInvalidationListIndirectArgsCS::FParameters* PassParameters = GraphBuilder.AllocParameters<FSetInstanceInvalidationListIndirectArgsCS::FParameters>();
-		PassParameters->VirtualShadowMap = InvalidationPassCommon.VirtualShadowMapUniformBuffer;
-		PassParameters->OutIndirectArgsBuffer = GraphBuilder.CreateUAV(IndirectArgs);
-		PassParameters->InstanceInvalidationList = GraphBuilder.CreateSRV(InstanceInvalidationList);
-		PassParameters->InvalidateInstancePagesCSGroupSize = FInvalidateInstancePagesListCS::DefaultCSGroupX;
-		PassParameters->MaxInstanceInvalidations = MaxInstanceInvalidations;
-
-		auto ComputeShader = GetGlobalShaderMap(Scene->GetFeatureLevel())->GetShader<FSetInstanceInvalidationListIndirectArgsCS>();
-
-		FComputeShaderUtils::AddPass(
-			GraphBuilder,
-			RDG_EVENT_NAME("SetInstanceInvalidationListIndirectArgsCS"),
-			ComputeShader,
-			PassParameters,
-			FIntVector(1, 1, 1)
-		);
-	}
-
-	// Invalidate GPU instances
-	{
-		FInvalidateInstancePagesListCS::FParameters* PassParameters = GraphBuilder.AllocParameters<FInvalidateInstancePagesListCS::FParameters>();
-
-		SetInvalidateInstancePagesParameters(GraphBuilder, InvalidationPassCommon, &PassParameters->InvalidatePagesParameters);
-		PassParameters->InstanceInvalidationList = GraphBuilder.CreateSRV(InstanceInvalidationList);
-		PassParameters->IndirectArgs = IndirectArgs;
-		PassParameters->MaxInstanceInvalidations = MaxInstanceInvalidations;
-
-		FInvalidateInstancePagesListCS::FPermutationDomain PermutationVector;
-		PermutationVector.Set<FInvalidateInstancePagesListCS::FUseHzbDim>(PassParameters->InvalidatePagesParameters.HZBTexture != nullptr);
-
-		auto ComputeShader = GetGlobalShaderMap(Scene->GetFeatureLevel())->GetShader<FInvalidateInstancePagesListCS>(PermutationVector);
-
-		FComputeShaderUtils::AddPass(
-			GraphBuilder,
-			RDG_EVENT_NAME("InvalidateInstancePagesListCS"),
-			ComputeShader,
-			PassParameters,
-			PassParameters->IndirectArgs,
-			0
-		);
-	}
-}
-
-
 // Remove old info used to track logging.
 void FVirtualShadowMapArrayCacheManager::TrimLoggingInfo()
 {
diff --git a/Engine/Source/Runtime/Renderer/Private/VirtualShadowMaps/VirtualShadowMapCacheManager.h b/Engine/Source/Runtime/Renderer/Private/VirtualShadowMaps/VirtualShadowMapCacheManager.h
--- a/Engine/Source/Runtime/Renderer/Private/VirtualShadowMaps/VirtualShadowMapCacheManager.h
+++ b/Engine/Source/Runtime/Renderer/Private/VirtualShadowMaps/VirtualShadowMapCacheManager.h
@@ -21,6 +21,7 @@ struct FVirtualShadowMapInstanceRange
 {
 	int32 InstanceSceneDataOffset;
 	int32 NumInstanceSceneDataEntries;
+	bool bForceInvalidateStatic;
 };
 
 #define VSM_LOG_INVALIDATIONS 0
@@ -220,8 +221,6 @@ public:
 	void FreePhysicalPool(FRDGBuilder& GraphBuilder);
 	TRefCountPtr<IPooledRenderTarget> GetPhysicalPagePool() const { return PhysicalPagePool; }
 	TRefCountPtr<FRDGPooledBuffer> GetPhysicalPageMetaData() const { return PhysicalPageMetaData; }
-	TRefCountPtr<FRDGPooledBuffer> GetCacheInstanceAsStatic() const { return CacheInstanceAsStatic; }
-	TRefCountPtr<FRDGPooledBuffer> GetLastInstanceInvalidatedFrame() const { return LastInstanceInvalidatedFrame; }
 
 	// Called by VirtualShadowMapArray to potentially resize the HZB physical pool
 	TRefCountPtr<IPooledRenderTarget> SetHZBPhysicalPoolSize(FRDGBuilder& GraphBuilder, FIntPoint RequestedSize, const EPixelFormat Format);
@@ -271,7 +270,7 @@ public:
 	public:
 		FInvalidatingPrimitiveCollector(FVirtualShadowMapArrayCacheManager* InVirtualShadowMapArrayCacheManager);
 
-		void AddDynamicAndGPUPrimitives();
+		void AddPrimitivesToInvalidate();
 
 		/**
 		 * All of these functions filters redundant primitive adds, and thus expects valid IDs (so can't be called for primitives that have not yet been added)
@@ -296,9 +295,9 @@ public:
 			AddInvalidation(PrimitiveSceneInfo, false);
 		}
 
-		void Finalize();
-
 		FInstanceGPULoadBalancer Instances;
+		TBitArray<> InvalidatedPrimitives;
+		TBitArray<> RemovedPrimitives;
 
 	private:
 		void AddInvalidation(FPrimitiveSceneInfo* PrimitiveSceneInfo, bool bRemovedPrimitive);
@@ -311,7 +310,7 @@ public:
 	void ProcessInvalidations(
 		FRDGBuilder& GraphBuilder,
 		FSceneUniformBuffer &SceneUniformBuffer,
-		const FInvalidatingPrimitiveCollector& InvalidatingPrimitiveCollector);
+		FInvalidatingPrimitiveCollector& InvalidatingPrimitiveCollector);
 
 	/**
 	 * Allow the cache manager to track scene changes, in particular track resizing of primitive tracking data.
@@ -346,7 +345,7 @@ public:
 	}
 
 	UE::Renderer::Private::IShadowInvalidatingInstances *GetInvalidatingInstancesInterface() { return &ShadowInvalidatingInstancesImplementation; }
-
+	FRDGBufferRef UploadCachePrimitiveAsDynamic(FRDGBuilder& GraphBuilder) const;
 private:
 
 	/** 
@@ -379,10 +378,10 @@ private:
 		const FInvalidationPassCommon& InvalidationPassCommon,
 		FInvalidatePagesParameters* PassParameters) const;
 
+	void UpdateCachePrimitiveAsDynamic(FInvalidatingPrimitiveCollector& InvalidatingPrimitiveCollector);
+
 	// Invalidate instances based on CPU instance ranges. This is used for CPU-based updates like object transform changes, etc.
 	void ProcessInvalidations(FRDGBuilder& GraphBuilder, const FInvalidationPassCommon& InvalidationPassCommon, const FInstanceGPULoadBalancer& Instances) const;
-		// Invalidate instances based on a GPU instance list. This is currently used for static<->dynamic cache mode transitions only.
-	void ProcessInvalidations(FRDGBuilder& GraphBuilder, const FInvalidationPassCommon& InvalidationPassCommon, FRDGBufferRef InstanceInvalidationList, uint32 MaxInstanceInvalidations) const;
 	
 	void ExtractStats(FRDGBuilder& GraphBuilder, FVirtualShadowMapArray &VirtualShadowMapArray);
 
@@ -400,13 +399,16 @@ private:
 	ETextureCreateFlags PhysicalPagePoolCreateFlags = TexCreate_None;
 	TRefCountPtr<FRDGPooledBuffer> PhysicalPageMetaData;
 	uint32 MaxPhysicalPages = 0;
-	// For now these only grow
-	TRefCountPtr<FRDGPooledBuffer> CacheInstanceAsStatic;
-	TRefCountPtr<FRDGPooledBuffer> LastInstanceInvalidatedFrame;
 
 	// Index the Cache entries by the light ID
 	FEntryMap CacheEntries;
 
+	// Store the last time a primitive caused an invalidation for dynamic/static caching purposes
+	// NOTE: Set bits as dynamic since the container makes it easier to iterate those
+	TBitArray<> CachePrimitiveAsDynamic;
+	// Indexed by PersistentPrimitiveIndex
+	TArray<uint32> LastPrimitiveInvalidatedFrame;
+
 	// Stores stats over frames when activated.
 	TRefCountPtr<FRDGPooledBuffer> AccumulatedStatsBuffer;
 	bool bAccumulatingStats = false;
