Fixed wrong stencil settings for sky mesh in a mobile rendering
#11106

diff --git a/Engine/Source/Runtime/Renderer/Private/MobileFogRendering.cpp b/Engine/Source/Runtime/Renderer/Private/MobileFogRendering.cpp
--- a/Engine/Source/Runtime/Renderer/Private/MobileFogRendering.cpp
+++ b/Engine/Source/Runtime/Renderer/Private/MobileFogRendering.cpp
@@ -172,7 +172,7 @@ void FMobileSceneRenderer::RenderFog(FRHICommandList& RHICmdList, const FViewInf
 	{
 		GraphicsPSOInit.DepthStencilState = TStaticDepthStencilState<
 			false, CF_DepthNearOrEqual,
-			true, CF_Equal, SO_Keep, SO_Keep, SO_Keep,
+			true, CF_NotEqual, SO_Keep, SO_Keep, SO_Keep,
 			false, CF_Always, SO_Keep, SO_Keep, SO_Keep,
 			STENCIL_MOBILE_SKY_MASK, 0x00>::GetRHI();
 	}
@@ -218,7 +218,14 @@ void FMobileSceneRenderer::RenderFog(FRHICommandList& RHICmdList, const FViewInf
 	GraphicsPSOInit.BoundShaderState.VertexShaderRHI = VertexShader.GetVertexShader();
 	GraphicsPSOInit.BoundShaderState.PixelShaderRHI = PixelShader.GetPixelShader();
 	
-	SetGraphicsPipelineState(RHICmdList, GraphicsPSOInit, 0);
+	if (bUseDepthTest)
+	{
+		SetGraphicsPipelineState(RHICmdList, GraphicsPSOInit, STENCIL_MOBILE_SKY_MASK);
+	}
+	else
+	{
+		SetGraphicsPipelineState(RHICmdList, GraphicsPSOInit, 0);
+	}
 	
 	// Use height fog start distance by default and fallback to AP distance
 	float FogStartDistance = GetFogDefaultStartDistance();
diff --git a/Engine/Source/Runtime/Renderer/Private/SkyPassRendering.cpp b/Engine/Source/Runtime/Renderer/Private/SkyPassRendering.cpp
--- a/Engine/Source/Runtime/Renderer/Private/SkyPassRendering.cpp
+++ b/Engine/Source/Runtime/Renderer/Private/SkyPassRendering.cpp
@@ -133,7 +133,7 @@ bool FSkyPassMeshProcessor::Process(
 			false, CF_Always, SO_Keep, SO_Keep, SO_Keep,
 			0x00, STENCIL_MOBILE_SKY_MASK>::GetRHI());
 		
-		PassDrawRenderState.SetStencilRef(1); 
+		PassDrawRenderState.SetStencilRef(STENCIL_MOBILE_SKY_MASK); 
 		
 		TMobileBasePassShaderElementData<LightMapPolicyType> ShaderElementData(nullptr, false);
 		ShaderElementData.InitializeMeshMaterialData(ViewIfDynamicMeshCommand, PrimitiveSceneProxy, MeshBatch, StaticMeshId, false);
