Implements a reusable FTSRConvolutionNetworkShader

diff --git a/Engine/Shaders/Private/TemporalSuperResolution/TSRCommon.ush b/Engine/Shaders/Private/TemporalSuperResolution/TSRCommon.ush
--- a/Engine/Shaders/Private/TemporalSuperResolution/TSRCommon.ush
+++ b/Engine/Shaders/Private/TemporalSuperResolution/TSRCommon.ush
@@ -14,7 +14,7 @@
 
 //------------------------------------------------------- RECOMPILE HASH
 
-#pragma message("UESHADERMETADATA_VERSION 65292827-F021-4C62-9790-FC626A1BFE69")
+#pragma message("UESHADERMETADATA_VERSION D9BD9048-6528-48DC-922C-82BAF2992B38")
 
 
 //------------------------------------------------------- COMPILER CONFIG
diff --git a/Engine/Source/Runtime/Renderer/Private/PostProcess/TemporalSuperResolution.cpp b/Engine/Source/Runtime/Renderer/Private/PostProcess/TemporalSuperResolution.cpp
--- a/Engine/Source/Runtime/Renderer/Private/PostProcess/TemporalSuperResolution.cpp
+++ b/Engine/Source/Runtime/Renderer/Private/PostProcess/TemporalSuperResolution.cpp
@@ -441,6 +441,97 @@
 	}
 }; // class FTemporalSuperResolutionShader
 
+class FTSRConvolutionNetworkShader : public FTSRShader
+{
+public:
+	class FWaveSizeOps : SHADER_PERMUTATION_SPARSE_INT("DIM_WAVE_SIZE", 0, 16, 32, 64);
+
+	using FPermutationDomain = TShaderPermutationDomain<FWaveSizeOps, FTSRShader::F16BitVALUDim, FTSRShader::FAlphaChannelDim>;
+
+	FTSRConvolutionNetworkShader(const ShaderMetaType::CompiledShaderInitializerType& Initializer)
+		: FTSRShader(Initializer)
+	{ }
+
+	FTSRConvolutionNetworkShader()
+	{ }
+
+	static FPermutationDomain RemapPermutation(FPermutationDomain PermutationVector)
+	{
+		int32 WaveSize = PermutationVector.Get<FWaveSizeOps>();
+
+		// WaveSize=16 is for Intel Arc GPU which also supports 16bits ops, so compiling WaveSize=16 32bit ops is useless and should instead fall back to WaveSize=0.
+		if (WaveSize == 16 && !PermutationVector.Get<FTSRShader::F16BitVALUDim>())
+		{
+			PermutationVector.Set<FWaveSizeOps>(0);
+		}
+
+		// Only compile the alpha channel with 32bit ops, as this is mostly targeting enterprise uses on Quadro GPUs
+		if (PermutationVector.Get<FTSRShader::FAlphaChannelDim>())
+		{
+			PermutationVector.Set<FTSRShader::F16BitVALUDim>(false);
+		}
+
+		// Optimising register pressure with 16bit for waveops that is 1 pixel/lane is pointless.
+		if (WaveSize == 0)
+		{
+			PermutationVector.Set<FTSRShader::F16BitVALUDim>(false);
+		}
+
+		return PermutationVector;
+	}
+
+	static bool ShouldCompilePermutation(const FGlobalShaderPermutationParameters& Parameters, FPermutationDomain PermutationVector)
+	{
+		if (!FTSRShader::ShouldCompilePermutation(Parameters))
+		{
+			return false;
+		}
+
+		int32 WaveSize = PermutationVector.Get<FWaveSizeOps>();
+
+		ERHIFeatureSupport WaveOpsSupport = FTSRShader::SupportsWaveOps(Parameters.Platform);
+		if (WaveSize != 0)
+		{
+			if (WaveOpsSupport == ERHIFeatureSupport::Unsupported)
+			{
+				return false;
+			}
+
+			if (WaveSize < int32(FDataDrivenShaderPlatformInfo::GetMinimumWaveSize(Parameters.Platform)) ||
+				WaveSize > int32(FDataDrivenShaderPlatformInfo::GetMaximumWaveSize(Parameters.Platform)))
+			{
+				return false;
+			}
+		}
+
+		if (!FTSRShader::ShouldCompile32or16BitPermutation(Parameters.Platform, PermutationVector.Get<FTSRShader::F16BitVALUDim>()))
+		{
+			return false;
+		}
+
+		return true;
+	}
+
+	static void ModifyCompilationEnvironment(const FGlobalShaderPermutationParameters& Parameters, const FPermutationDomain& PermutationVector, FShaderCompilerEnvironment& OutEnvironment)
+	{
+		FTSRShader::ModifyCompilationEnvironment(Parameters, OutEnvironment);
+
+		if (PermutationVector.Get<FWaveSizeOps>() != 0)
+		{
+			if (PermutationVector.Get<FWaveSizeOps>() == 32)
+			{
+				OutEnvironment.CompilerFlags.Add(CFLAG_Wave32);
+			}
+			OutEnvironment.CompilerFlags.Add(CFLAG_WaveOperations);
+		}
+
+		if (PermutationVector.Get<FTSRShader::F16BitVALUDim>())
+		{
+			OutEnvironment.CompilerFlags.Add(CFLAG_AllowRealTypes);
+		}
+	}
+}; // FTSRConvolutionNetworkShader
+
 class FTSRMeasureFlickeringLumaCS : public FTSRShader
 {
 	DECLARE_GLOBAL_SHADER(FTSRMeasureFlickeringLumaCS);
@@ -593,16 +684,18 @@
 	}
 }; // class FTSRDecimateHistoryCS
 
-class FTSRRejectShadingCS : public FTSRShader
+class FTSRRejectShadingCS : public FTSRConvolutionNetworkShader
 {
 	DECLARE_GLOBAL_SHADER(FTSRRejectShadingCS);
-	SHADER_USE_PARAMETER_STRUCT(FTSRRejectShadingCS, FTSRShader);
+	SHADER_USE_PARAMETER_STRUCT(FTSRRejectShadingCS, FTSRConvolutionNetworkShader);
 
-	class FWaveSizeOps : SHADER_PERMUTATION_SPARSE_INT("DIM_WAVE_SIZE", 0, 16, 32, 64);
 	class FFlickeringDetectionDim : SHADER_PERMUTATION_BOOL("DIM_FLICKERING_DETECTION");
 	class FHistoryResurrectionDim : SHADER_PERMUTATION_BOOL("DIM_HISTORY_RESURRECTION");
 
-	using FPermutationDomain = TShaderPermutationDomain<FWaveSizeOps, FFlickeringDetectionDim, FHistoryResurrectionDim, FTSRShader::F16BitVALUDim, FTSRShader::FAlphaChannelDim>;
+	using FPermutationDomain = TShaderPermutationDomain<
+		FTSRConvolutionNetworkShader::FPermutationDomain,
+		FFlickeringDetectionDim,
+		FHistoryResurrectionDim>;
 
 	BEGIN_SHADER_PARAMETER_STRUCT(FParameters, )
 		SHADER_PARAMETER_STRUCT_INCLUDE(FTSRCommonParameters, CommonParameters)
@@ -642,28 +735,12 @@
 
 	static FPermutationDomain RemapPermutation(FPermutationDomain PermutationVector)
 	{
-		int32 WaveSize = PermutationVector.Get<FWaveSizeOps>();
+		// Remap redondant convolution permutations.
+		PermutationVector.Set<FTSRConvolutionNetworkShader::FPermutationDomain>(
+			FTSRConvolutionNetworkShader::RemapPermutation(PermutationVector.Get<FTSRConvolutionNetworkShader::FPermutationDomain>()));
 
-		// WaveSize=16 is for Intel Arc GPU which also supports 16bits ops, so compiling WaveSize=16 32bit ops is useless and should instead fall back to WaveSize=0.
-		if (WaveSize == 16 && !PermutationVector.Get<FTSRShader::F16BitVALUDim>())
-		{
-			PermutationVector.Set<FWaveSizeOps>(0);
-		}
-
-		// Only compile the alpha channel with 32bit ops, as this is mostly targeting enterprise uses on Quadro GPUs
-		if (PermutationVector.Get<FTSRShader::FAlphaChannelDim>())
-		{
-			PermutationVector.Set<FTSRShader::F16BitVALUDim>(false);
-		}
-
-		// Optimising register pressure with 16bit for waveops that is 1 pixel/lane is pointless.
-		if (WaveSize == 0)
-		{
-			PermutationVector.Set<FTSRShader::F16BitVALUDim>(false);
-		}
-
 		// Register pressure is identical between all these permutation with 16bit
-		if (PermutationVector.Get<FTSRShader::F16BitVALUDim>())
+		if (PermutationVector.Get<FTSRConvolutionNetworkShader::FPermutationDomain>().Get<FTSRShader::F16BitVALUDim>())
 		{
 			PermutationVector.Set<FFlickeringDetectionDim>(true);
 			PermutationVector.Set<FHistoryResurrectionDim>(true);
@@ -680,35 +757,13 @@
 
 	static bool ShouldCompilePermutation(const FGlobalShaderPermutationParameters& Parameters)
 	{
-		if (!FTSRShader::ShouldCompilePermutation(Parameters))
-		{
-			return false;
-		}
-
 		FPermutationDomain PermutationVector(Parameters.PermutationId);
 		if (PermutationVector != RemapPermutation(PermutationVector))
 		{
 			return false;
 		}
 
-		int32 WaveSize = PermutationVector.Get<FWaveSizeOps>();
-
-		ERHIFeatureSupport WaveOpsSupport = FTSRShader::SupportsWaveOps(Parameters.Platform);
-		if (WaveSize != 0)
-		{
-			if (WaveOpsSupport == ERHIFeatureSupport::Unsupported)
-			{
-				return false;
-			}
-
-			if (WaveSize < int32(FDataDrivenShaderPlatformInfo::GetMinimumWaveSize(Parameters.Platform)) ||
-				WaveSize > int32(FDataDrivenShaderPlatformInfo::GetMaximumWaveSize(Parameters.Platform)))
-			{
-				return false;
-			}
-		}
-
-		if (!FTSRShader::ShouldCompile32or16BitPermutation(Parameters.Platform, PermutationVector.Get<FTSRShader::F16BitVALUDim>()))
+		if (!FTSRConvolutionNetworkShader::ShouldCompilePermutation(Parameters, PermutationVector.Get<FTSRConvolutionNetworkShader::FPermutationDomain>()))
 		{
 			return false;
 		}
@@ -720,21 +775,10 @@
 	{
 		FPermutationDomain PermutationVector(Parameters.PermutationId);
 
-		FTSRShader::ModifyCompilationEnvironment(Parameters, OutEnvironment);
-
-		if (PermutationVector.Get<FWaveSizeOps>() != 0)
-		{
-			if (PermutationVector.Get<FWaveSizeOps>() == 32)
-			{
-				OutEnvironment.CompilerFlags.Add(CFLAG_Wave32);
-			}
-			OutEnvironment.CompilerFlags.Add(CFLAG_WaveOperations);
-		}
-
-		if (PermutationVector.Get<FTSRShader::F16BitVALUDim>())
-		{
-			OutEnvironment.CompilerFlags.Add(CFLAG_AllowRealTypes);
-		}
+		FTSRConvolutionNetworkShader::ModifyCompilationEnvironment(
+			Parameters,
+			PermutationVector.Get<FTSRConvolutionNetworkShader::FPermutationDomain>(),
+			OutEnvironment);
 	}
 }; // class FTSRRejectShadingCS
 
@@ -1958,12 +2002,15 @@
 		FScreenPassTextureViewport TranslucencyViewport(
 			SeparateTranslucencyTexture->Desc.Extent, SeparateTranslucencyRect);
 
+		FTSRConvolutionNetworkShader::FPermutationDomain ConvolutionNetworkPermutationVector;
+		ConvolutionNetworkPermutationVector.Set<FTSRConvolutionNetworkShader::FWaveSizeOps>(SelectWaveSize({ 16, 32, 64 }));
+		ConvolutionNetworkPermutationVector.Set<FTSRShader::F16BitVALUDim>(bUse16BitVALU);
+		ConvolutionNetworkPermutationVector.Set<FTSRShader::FAlphaChannelDim>(bSupportsAlpha);
+
 		FTSRRejectShadingCS::FPermutationDomain PermutationVector;
-		PermutationVector.Set<FTSRRejectShadingCS::FWaveSizeOps>(SelectWaveSize({ 16, 32, 64 }));
+		PermutationVector.Set<FTSRConvolutionNetworkShader::FPermutationDomain>(ConvolutionNetworkPermutationVector);
 		PermutationVector.Set<FTSRRejectShadingCS::FFlickeringDetectionDim>(FlickeringFramePeriod > 0.0f);
 		PermutationVector.Set<FTSRRejectShadingCS::FHistoryResurrectionDim>(bCanResurrectHistory);
-		PermutationVector.Set<FTSRShader::F16BitVALUDim>(bUse16BitVALU);
-		PermutationVector.Set<FTSRShader::FAlphaChannelDim>(bSupportsAlpha);
 		PermutationVector = FTSRRejectShadingCS::RemapPermutation(PermutationVector);
 
 		const int32 GroupTileSize = 32;
@@ -2098,15 +2145,15 @@
 		TShaderMapRef<FTSRRejectShadingCS> ComputeShader(View.ShaderMap, PermutationVector);
 		FComputeShaderUtils::AddPass(
 			GraphBuilder,
-			RDG_EVENT_NAME("TSR RejectShading(#%d TileSize=%d PaddingCostMultiplier=%1.1f WaveSize=%d FlickeringFramePeriod=%f VALU=%s%s%s) %dx%d",
+			RDG_EVENT_NAME("TSR RejectShading(#%d TileSize=%d PaddingCostMultiplier=%1.1f WaveSize=%d VALU=%s%s FlickeringFramePeriod=%f%s) %dx%d",
 				PermutationVector.ToDimensionValueId(),
 				TileSize,
 				FMath::Pow(float(GroupTileSize) / float(TileSize), 2),
-				int32(PermutationVector.Get<FTSRRejectShadingCS::FWaveSizeOps>()),
+				int32(PermutationVector.Get<FTSRConvolutionNetworkShader::FPermutationDomain>().Get<FTSRRejectShadingCS::FWaveSizeOps>()),
+				PermutationVector.Get<FTSRConvolutionNetworkShader::FPermutationDomain>().Get<FTSRShader::F16BitVALUDim>() ? TEXT("16bit") : TEXT("32bit"),
+				PermutationVector.Get<FTSRConvolutionNetworkShader::FPermutationDomain>().Get<FTSRShader::FAlphaChannelDim>() ? TEXT(" AlphaChannel") : TEXT(""),
 				PassParameters->FlickeringFramePeriod,
-				PermutationVector.Get<FTSRShader::F16BitVALUDim>() ? TEXT("16bit") : TEXT("32bit"),
 				PassParameters->bEnableResurrection ? TEXT(" Resurrection") : TEXT(""),
-				PermutationVector.Get<FTSRShader::FAlphaChannelDim>() ? TEXT(" AlphaChannel") : TEXT(""),
 				InputRect.Width(), InputRect.Height()),
 			AsyncComputePasses >= 3 ? ERDGPassFlags::AsyncCompute : ERDGPassFlags::Compute,
 			ComputeShader,
