Implements a SceneColorBeforeBloom post process material location.

Artists don't actually have a post process material location in linear color space that has all the geometry composed in in linear color space before the tonemaper. Therefor this CL implements a post process material location in between motion blur and bloom. With TAAU & TSR, this means running display resolution, as well as requiring dedicated downsample passes running after.

diff --git a/Engine/Plugins/Compositing/OpenColorIO/Source/OpenColorIO/Private/OpenColorIODisplayExtension.cpp b/Engine/Plugins/Compositing/OpenColorIO/Source/OpenColorIO/Private/OpenColorIODisplayExtension.cpp
--- a/Engine/Plugins/Compositing/OpenColorIO/Source/OpenColorIO/Private/OpenColorIODisplayExtension.cpp
+++ b/Engine/Plugins/Compositing/OpenColorIO/Source/OpenColorIO/Private/OpenColorIODisplayExtension.cpp
@@ -72,7 +72,7 @@ void FOpenColorIODisplayExtension::SubscribeToPostProcessingPass(EPostProcessing
 
 FScreenPassTexture FOpenColorIODisplayExtension::PostProcessPassAfterTonemap_RenderThread(FRDGBuilder& GraphBuilder, const FSceneView& View, const FPostProcessMaterialInputs& InOutInputs)
 {
-	const FScreenPassTexture& SceneColor = InOutInputs.GetInput(EPostProcessMaterialInput::SceneColor);
+	const FScreenPassTexture& SceneColor = FScreenPassTexture::CopyFromSlice(GraphBuilder, InOutInputs.GetInput(EPostProcessMaterialInput::SceneColor));
 	check(SceneColor.IsValid());
 
 	FScreenPassRenderTarget Output = InOutInputs.OverrideOutput;
diff --git a/Engine/Plugins/Experimental/CineCameraSceneCapture/Source/CineCameraSceneCapture/Private/CineCameraSceneCaptureComponent.cpp b/Engine/Plugins/Experimental/CineCameraSceneCapture/Source/CineCameraSceneCapture/Private/CineCameraSceneCaptureComponent.cpp
--- a/Engine/Plugins/Experimental/CineCameraSceneCapture/Source/CineCameraSceneCapture/Private/CineCameraSceneCaptureComponent.cpp
+++ b/Engine/Plugins/Experimental/CineCameraSceneCapture/Source/CineCameraSceneCapture/Private/CineCameraSceneCaptureComponent.cpp
@@ -104,7 +104,7 @@ public:
 
 	FScreenPassTexture PostProcessPassAfterTonemap_RenderThread(FRDGBuilder& GraphBuilder, const FSceneView& View, const FPostProcessMaterialInputs& InOutInputs)
 	{
-		const FScreenPassTexture& SceneColor = InOutInputs.GetInput(EPostProcessMaterialInput::SceneColor);
+		const FScreenPassTexture& SceneColor = FScreenPassTexture::CopyFromSlice(GraphBuilder, InOutInputs.GetInput(EPostProcessMaterialInput::SceneColor));
 		check(SceneColor.IsValid());
 
 		FScreenPassRenderTarget Output = InOutInputs.OverrideOutput;
diff --git a/Engine/Plugins/Media/MediaIOFramework/Source/MediaIOCore/Private/MediaCaptureSceneViewExtension.h b/Engine/Plugins/Media/MediaIOFramework/Source/MediaIOCore/Private/MediaCaptureSceneViewExtension.h
--- a/Engine/Plugins/Media/MediaIOFramework/Source/MediaIOCore/Private/MediaCaptureSceneViewExtension.h
+++ b/Engine/Plugins/Media/MediaIOFramework/Source/MediaIOCore/Private/MediaCaptureSceneViewExtension.h
@@ -94,7 +94,7 @@ public:
 	{
 		TRACE_CPUPROFILER_EVENT_SCOPE(MediaCaptureExtensionCallback);
 
-		FScreenPassTexture SceneColor = InOutInputs.GetInput(EPostProcessMaterialInput::SceneColor);
+		FScreenPassTexture SceneColor = FScreenPassTexture::CopyFromSlice(GraphBuilder, InOutInputs.GetInput(EPostProcessMaterialInput::SceneColor));
 		check(SceneColor.IsValid());
 
 		if (FRDGTextureRef TextureRef = SceneColor.Texture)
@@ -106,24 +106,7 @@ public:
 			}
 		}
 
-		return ReturnUntouchedSceneColorForPostProcessing(InOutInputs);
-	}
-	
-	/** 
-	 * A helper function that extracts the right scene color texture, untouched, to be used further in post processing. 
-	 */
-	FScreenPassTexture ReturnUntouchedSceneColorForPostProcessing(const FPostProcessMaterialInputs& InOutInputs)
-	{
-		if (InOutInputs.OverrideOutput.IsValid())
-		{
-			return InOutInputs.OverrideOutput;
-		}
-		else
-		{
-			/** We don't want to modify scene texture in any way. We just want it to be passed back onto the next stage. */
-			FScreenPassTexture SceneTexture = const_cast<FScreenPassTexture&>(InOutInputs.Textures[(uint32)EPostProcessMaterialInput::SceneColor]);
-			return SceneTexture;
-		}
+		return InOutInputs.ReturnUntouchedSceneColorForPostProcessing(GraphBuilder);
 	}
 
 	bool IsValid() const
diff --git a/Engine/Plugins/Runtime/nDisplay/Source/DisplayCluster/Private/Render/DisplayDevice/Proxy/DisplayClusterDisplayDeviceProxy_OpenColorIO.cpp b/Engine/Plugins/Runtime/nDisplay/Source/DisplayCluster/Private/Render/DisplayDevice/Proxy/DisplayClusterDisplayDeviceProxy_OpenColorIO.cpp
--- a/Engine/Plugins/Runtime/nDisplay/Source/DisplayCluster/Private/Render/DisplayDevice/Proxy/DisplayClusterDisplayDeviceProxy_OpenColorIO.cpp
+++ b/Engine/Plugins/Runtime/nDisplay/Source/DisplayCluster/Private/Render/DisplayDevice/Proxy/DisplayClusterDisplayDeviceProxy_OpenColorIO.cpp
@@ -56,20 +56,3 @@ bool FDisplayClusterDisplayDeviceProxy_OpenColorIO::AddFinalPass_RenderThread(FR
 
 	return true;
 }
-
-/**
-* A helper function that extracts the right scene color texture, untouched, to be used further in post processing.
-*/
-FScreenPassTexture IDisplayClusterDisplayDeviceProxy::ReturnUntouchedSceneColorForPostProcessing(const FPostProcessMaterialInputs& InOutInputs)
-{
-	if (InOutInputs.OverrideOutput.IsValid())
-	{
-		return InOutInputs.OverrideOutput;
-	}
-	else
-	{
-		/** We don't want to modify scene texture in any way. We just want it to be passed back onto the next stage. */
-		FScreenPassTexture SceneTexture = const_cast<FScreenPassTexture&>(InOutInputs.Textures[(uint32)EPostProcessMaterialInput::SceneColor]);
-		return SceneTexture;
-	}
-}
diff --git a/Engine/Plugins/Runtime/nDisplay/Source/DisplayCluster/Private/Render/Viewport/DisplayClusterViewportFrameStatsViewExtension.cpp b/Engine/Plugins/Runtime/nDisplay/Source/DisplayCluster/Private/Render/Viewport/DisplayClusterViewportFrameStatsViewExtension.cpp
--- a/Engine/Plugins/Runtime/nDisplay/Source/DisplayCluster/Private/Render/Viewport/DisplayClusterViewportFrameStatsViewExtension.cpp
+++ b/Engine/Plugins/Runtime/nDisplay/Source/DisplayCluster/Private/Render/Viewport/DisplayClusterViewportFrameStatsViewExtension.cpp
@@ -102,7 +102,7 @@ void FDisplayClusterViewportFrameStatsViewExtension::BeginRenderViewFamily(FScen
 
 FScreenPassTexture FDisplayClusterViewportFrameStatsViewExtension::PostProcessPassAfterTonemap_RenderThread(FRDGBuilder& GraphBuilder, const FSceneView& View, const FPostProcessMaterialInputs& Inputs)
 {
-	const FScreenPassTexture& SceneColor = Inputs.GetInput(EPostProcessMaterialInput::SceneColor);
+	const FScreenPassTexture& SceneColor = FScreenPassTexture::CopyFromSlice(GraphBuilder, Inputs.GetInput(EPostProcessMaterialInput::SceneColor));
 	check(SceneColor.IsValid());
 
 	FGlobalShaderMap* ShaderMap = GetGlobalShaderMap(View.GetFeatureLevel());
diff --git a/Engine/Plugins/Runtime/nDisplay/Source/DisplayCluster/Private/Render/Viewport/DisplayClusterViewportManagerViewExtension.cpp b/Engine/Plugins/Runtime/nDisplay/Source/DisplayCluster/Private/Render/Viewport/DisplayClusterViewportManagerViewExtension.cpp
--- a/Engine/Plugins/Runtime/nDisplay/Source/DisplayCluster/Private/Render/Viewport/DisplayClusterViewportManagerViewExtension.cpp
+++ b/Engine/Plugins/Runtime/nDisplay/Source/DisplayCluster/Private/Render/Viewport/DisplayClusterViewportManagerViewExtension.cpp
@@ -119,7 +119,7 @@ FScreenPassTexture FDisplayClusterViewportManagerViewExtension::PostProcessPassA
 {
 	if (!IsActive())
 	{
-		return ReturnUntouchedSceneColorForPostProcessing(Inputs);
+		return Inputs.ReturnUntouchedSceneColorForPostProcessing(GraphBuilder);
 	}
 
 	if (const FDisplayClusterViewportManagerProxy* ViewportManagerProxy = Configuration->Proxy->GetViewportManagerProxyImpl())
@@ -134,14 +134,14 @@ FScreenPassTexture FDisplayClusterViewportManagerViewExtension::PostProcessPassA
 		}
 	}
 
-	return ReturnUntouchedSceneColorForPostProcessing(Inputs);
+	return Inputs.ReturnUntouchedSceneColorForPostProcessing(GraphBuilder);
 }
 
 FScreenPassTexture FDisplayClusterViewportManagerViewExtension::PostProcessPassAfterSSRInput_RenderThread(FRDGBuilder& GraphBuilder, const FSceneView& View, const FPostProcessMaterialInputs& Inputs)
 {
 	if (!IsActive())
 	{
-		return ReturnUntouchedSceneColorForPostProcessing(Inputs);
+		return Inputs.ReturnUntouchedSceneColorForPostProcessing(GraphBuilder);
 	}
 
 	if (const FDisplayClusterViewportManagerProxy* ViewportManagerProxy = Configuration->Proxy->GetViewportManagerProxyImpl())
@@ -156,14 +156,14 @@ FScreenPassTexture FDisplayClusterViewportManagerViewExtension::PostProcessPassA
 		}
 	}
 
-	return ReturnUntouchedSceneColorForPostProcessing(Inputs);
+	return Inputs.ReturnUntouchedSceneColorForPostProcessing(GraphBuilder);
 }
 
 FScreenPassTexture FDisplayClusterViewportManagerViewExtension::PostProcessPassAfterTonemap_RenderThread(FRDGBuilder& GraphBuilder, const FSceneView& View, const FPostProcessMaterialInputs& Inputs)
 {
 	if (!IsActive())
 	{
-		return ReturnUntouchedSceneColorForPostProcessing(Inputs);
+		return Inputs.ReturnUntouchedSceneColorForPostProcessing(GraphBuilder);
 	}
 
 	if (const FDisplayClusterViewportManagerProxy* ViewportManagerProxy = Configuration->Proxy->GetViewportManagerProxyImpl())
@@ -178,24 +178,7 @@ FScreenPassTexture FDisplayClusterViewportManagerViewExtension::PostProcessPassA
 		}
 	}
 
-	return ReturnUntouchedSceneColorForPostProcessing(Inputs);
-}
-
-/**
-* A helper function that extracts the right scene color texture, untouched, to be used further in post processing.
-*/
-FScreenPassTexture FDisplayClusterViewportManagerViewExtension::ReturnUntouchedSceneColorForPostProcessing(const FPostProcessMaterialInputs& InOutInputs)
-{
-	if (InOutInputs.OverrideOutput.IsValid())
-	{
-		return InOutInputs.OverrideOutput;
-	}
-	else
-	{
-		/** We don't want to modify scene texture in any way. We just want it to be passed back onto the next stage. */
-		FScreenPassTexture SceneTexture = const_cast<FScreenPassTexture&>(InOutInputs.Textures[(uint32)EPostProcessMaterialInput::SceneColor]);
-		return SceneTexture;
-	}
+	return Inputs.ReturnUntouchedSceneColorForPostProcessing(GraphBuilder);
 }
 
 bool FDisplayClusterViewportManagerViewExtension::IsActiveThisFrame_Internal(const FSceneViewExtensionContext& Context) const
diff --git a/Engine/Plugins/Runtime/nDisplay/Source/DisplayCluster/Private/Render/Viewport/DisplayClusterViewportManagerViewExtension.h b/Engine/Plugins/Runtime/nDisplay/Source/DisplayCluster/Private/Render/Viewport/DisplayClusterViewportManagerViewExtension.h
--- a/Engine/Plugins/Runtime/nDisplay/Source/DisplayCluster/Private/Render/Viewport/DisplayClusterViewportManagerViewExtension.h
+++ b/Engine/Plugins/Runtime/nDisplay/Source/DisplayCluster/Private/Render/Viewport/DisplayClusterViewportManagerViewExtension.h
@@ -44,14 +44,6 @@ public:
 	virtual void PostRenderViewFamily_RenderThread(FRDGBuilder& GraphBuilder, FSceneViewFamily& InViewFamily) override;
 	//~End ISceneVIewExtension interface
 
-	/** PP Helper: Get output PP texture from inputs (OverrideOutput or SceneColor).
-	 *
-	 * @param InOutInputs - PP material inputs
-	 *
-	 * @return - Screen pass texture
-	 */
-	static FScreenPassTexture ReturnUntouchedSceneColorForPostProcessing(const FPostProcessMaterialInputs& InOutInputs);
-
 	/** Release from render thread.*/
 	void Release_RenderThread();
 
diff --git a/Engine/Plugins/Runtime/nDisplay/Source/DisplayCluster/Private/Render/Viewport/DisplayClusterViewportProxy.cpp b/Engine/Plugins/Runtime/nDisplay/Source/DisplayCluster/Private/Render/Viewport/DisplayClusterViewportProxy.cpp
--- a/Engine/Plugins/Runtime/nDisplay/Source/DisplayCluster/Private/Render/Viewport/DisplayClusterViewportProxy.cpp
+++ b/Engine/Plugins/Runtime/nDisplay/Source/DisplayCluster/Private/Render/Viewport/DisplayClusterViewportProxy.cpp
@@ -394,7 +394,7 @@ void FDisplayClusterViewportProxy::OnResolvedSceneColor_RenderThread(FRDGBuilder
 
 FScreenPassTexture FDisplayClusterViewportProxy::OnPostProcessPassAfterSSRInput_RenderThread(FRDGBuilder& GraphBuilder, const FSceneView& View, const FPostProcessMaterialInputs& Inputs, const uint32 ContextNum)
 {
-	FScreenPassTexture OutScreenPassTexture = FDisplayClusterViewportManagerViewExtension::ReturnUntouchedSceneColorForPostProcessing(Inputs);
+	FScreenPassTexture OutScreenPassTexture = Inputs.ReturnUntouchedSceneColorForPostProcessing(GraphBuilder);
 	if (OutScreenPassTexture.IsValid())
 	{
 		// Copy alpha channel to 'InputShaderResource'
@@ -407,7 +407,7 @@ FScreenPassTexture FDisplayClusterViewportProxy::OnPostProcessPassAfterSSRInput_
 
 FScreenPassTexture FDisplayClusterViewportProxy::OnPostProcessPassAfterFXAA_RenderThread(FRDGBuilder& GraphBuilder, const FSceneView& View, const FPostProcessMaterialInputs& Inputs, const uint32 ContextNum)
 {
-	FScreenPassTexture OutScreenPassTexture = FDisplayClusterViewportManagerViewExtension::ReturnUntouchedSceneColorForPostProcessing(Inputs);
+	FScreenPassTexture OutScreenPassTexture = Inputs.ReturnUntouchedSceneColorForPostProcessing(GraphBuilder);
 	if (OutScreenPassTexture.IsValid())
 	{
 		// Restore alpha channel after OCIO
@@ -427,7 +427,7 @@ FScreenPassTexture FDisplayClusterViewportProxy::OnPostProcessPassAfterTonemap_R
 		return OpenColorIO->PostProcessPassAfterTonemap_RenderThread(GraphBuilder, GetContexts_RenderThread()[ContextNum], View, Inputs);
 	}
 
-	return FDisplayClusterViewportManagerViewExtension::ReturnUntouchedSceneColorForPostProcessing(Inputs);
+	return Inputs.ReturnUntouchedSceneColorForPostProcessing(GraphBuilder);
 }
 
 void FDisplayClusterViewportProxy::OnPostRenderViewFamily_RenderThread(FRDGBuilder& GraphBuilder, FSceneViewFamily& InViewFamily, const FSceneView& InSceneView, const FDisplayClusterViewportProxy_Context& InProxyContext)
diff --git a/Engine/Plugins/Runtime/nDisplay/Source/DisplayCluster/Private/Render/Viewport/DisplayClusterViewport_OpenColorIO.cpp b/Engine/Plugins/Runtime/nDisplay/Source/DisplayCluster/Private/Render/Viewport/DisplayClusterViewport_OpenColorIO.cpp
--- a/Engine/Plugins/Runtime/nDisplay/Source/DisplayCluster/Private/Render/Viewport/DisplayClusterViewport_OpenColorIO.cpp
+++ b/Engine/Plugins/Runtime/nDisplay/Source/DisplayCluster/Private/Render/Viewport/DisplayClusterViewport_OpenColorIO.cpp
@@ -91,7 +91,7 @@ bool FDisplayClusterViewport_OpenColorIO::AddPass_RenderThread(FRDGBuilder& Grap
 // This is a copy of FOpenColorIODisplayExtension::PostProcessPassAfterTonemap_RenderThread()
 FScreenPassTexture FDisplayClusterViewport_OpenColorIO::PostProcessPassAfterTonemap_RenderThread(FRDGBuilder& GraphBuilder, const FDisplayClusterViewport_Context& InViewportContext, const FSceneView& View, const FPostProcessMaterialInputs& InOutInputs)
 {
-	const FScreenPassTexture& SceneColor = InOutInputs.GetInput(EPostProcessMaterialInput::SceneColor);
+	const FScreenPassTexture& SceneColor = FScreenPassTexture::CopyFromSlice(GraphBuilder, InOutInputs.GetInput(EPostProcessMaterialInput::SceneColor));
 	check(SceneColor.IsValid());
 	FScreenPassRenderTarget Output = InOutInputs.OverrideOutput;
 
diff --git a/Engine/Plugins/Runtime/nDisplay/Source/DisplayCluster/Public/Render/DisplayDevice/IDisplayClusterDisplayDeviceProxy.h b/Engine/Plugins/Runtime/nDisplay/Source/DisplayCluster/Public/Render/DisplayDevice/IDisplayClusterDisplayDeviceProxy.h
--- a/Engine/Plugins/Runtime/nDisplay/Source/DisplayCluster/Public/Render/DisplayDevice/IDisplayClusterDisplayDeviceProxy.h
+++ b/Engine/Plugins/Runtime/nDisplay/Source/DisplayCluster/Public/Render/DisplayDevice/IDisplayClusterDisplayDeviceProxy.h
@@ -8,6 +8,7 @@
 #include "EngineUtils.h"
 #include "ScreenRendering.h"
 #include "ScreenPass.h"
+#include "PostProcess/PostProcessMaterialInputs.h"
 
 class FDisplayClusterViewport_Context;
 class FRDGBuilder;
@@ -21,8 +22,6 @@ struct FScreenPassTexture;
  */
 class DISPLAYCLUSTER_API IDisplayClusterDisplayDeviceProxy
 {
-	static FScreenPassTexture ReturnUntouchedSceneColorForPostProcessing(const FPostProcessMaterialInputs& Inputs);
-
 public:
 	virtual ~IDisplayClusterDisplayDeviceProxy() = default;
 
@@ -63,7 +62,7 @@ public:
 	 * @return - Screen pass texture
 	 */
 	virtual FScreenPassTexture OnPostProcessPassAfterFXAA_RenderThread(FRDGBuilder& GraphBuilder, const FSceneView& View, const FPostProcessMaterialInputs& Inputs, const uint32 ContextNum)
-	{ return ReturnUntouchedSceneColorForPostProcessing(Inputs); }
+	{ return Inputs.ReturnUntouchedSceneColorForPostProcessing(GraphBuilder); }
 
 	/** Allow callback OnPostProcessPassAfterSSRInput. */
 	virtual bool ShouldUsePostProcessPassAfterSSRInput() const
@@ -79,7 +78,7 @@ public:
 	 * @return - Screen pass texture
 	 */
 	virtual FScreenPassTexture OnPostProcessPassAfterSSRInput_RenderThread(FRDGBuilder& GraphBuilder, const FSceneView& View, const FPostProcessMaterialInputs& Inputs, const uint32 ContextNum)
-	{ return ReturnUntouchedSceneColorForPostProcessing(Inputs); }
+	{ return Inputs.ReturnUntouchedSceneColorForPostProcessing(GraphBuilder); }
 
 	/** Allow callback OnPostProcessPassAfterTonemap. */
 	virtual bool ShouldUsePostProcessPassTonemap() const
@@ -95,5 +94,5 @@ public:
 	 * @return - Screen pass texture
 	 */
 	virtual FScreenPassTexture OnPostProcessPassAfterTonemap_RenderThread(FRDGBuilder& GraphBuilder, const FSceneView& View, const FPostProcessMaterialInputs& Inputs, const uint32 ContextNum)
-	{ return ReturnUntouchedSceneColorForPostProcessing(Inputs); }
+	{ return Inputs.ReturnUntouchedSceneColorForPostProcessing(GraphBuilder); }
 };
diff --git a/Engine/Shaders/Private/PostProcessMaterialShaders.usf b/Engine/Shaders/Private/PostProcessMaterialShaders.usf
--- a/Engine/Shaders/Private/PostProcessMaterialShaders.usf
+++ b/Engine/Shaders/Private/PostProcessMaterialShaders.usf
@@ -7,7 +7,7 @@
 #include "Common.ush"
 #include "ScreenPass.ush"
 
-#pragma message("UESHADERMETADATA_VERSION C79EC12D-4DE8-4D26-B0C7-BD44503987C5")
+#pragma message("UESHADERMETADATA_VERSION 02E0E449-2223-45B4-8CD0-93AB087FBEE3")
 
 #ifndef POST_PROCESS_MATERIAL
 #define POST_PROCESS_MATERIAL 0
diff --git a/Engine/Source/Editor/PixelInspector/PixelInspectorModule.Build.cs b/Engine/Source/Editor/PixelInspector/PixelInspectorModule.Build.cs
--- a/Engine/Source/Editor/PixelInspector/PixelInspectorModule.Build.cs
+++ b/Engine/Source/Editor/PixelInspector/PixelInspectorModule.Build.cs
@@ -22,10 +22,11 @@ public class PixelInspectorModule : ModuleRules
         );
 
 		PrivateDependencyModuleNames.AddRange(
-             new string[] {
-					"ColorManagement"
-				}
-         );
+			new string[] {
+				"ColorManagement",
+				"Renderer",
+			}
+		);
 
 		PrivateIncludePathModuleNames.AddRange(
 			new string[] {
diff --git a/Engine/Source/Editor/PixelInspector/Private/PixelInspector.cpp b/Engine/Source/Editor/PixelInspector/Private/PixelInspector.cpp
--- a/Engine/Source/Editor/PixelInspector/Private/PixelInspector.cpp
+++ b/Engine/Source/Editor/PixelInspector/Private/PixelInspector.cpp
@@ -36,23 +36,6 @@
 
 namespace PixelInspector
 {
-	/** 
-	* A helper function that extracts the right scene color texture, untouched, to be used further in post processing. 
-	*/
-	FScreenPassTexture ReturnUntouchedSceneColorForPostProcessing(const FPostProcessMaterialInputs& InOutInputs)
-	{
-		if (InOutInputs.OverrideOutput.IsValid())
-		{
-			return InOutInputs.OverrideOutput;
-		}
-		else
-		{
-			/** We don't want to modify scene texture in any way. We just want it to be passed back onto the next stage. */
-			FScreenPassTexture SceneTexture = const_cast<FScreenPassTexture&>(InOutInputs.Textures[(uint32)EPostProcessMaterialInput::SceneColor]);
-			return SceneTexture;
-		}
-	}
-
 	SPixelInspector::SPixelInspector()
 	{
 		DisplayResult = nullptr;
@@ -889,16 +872,16 @@ namespace PixelInspector
 
 	FScreenPassTexture FPixelInspectorSceneViewExtension::PostProcessPassAfterFxaa_RenderThread(FRDGBuilder& GraphBuilder, const FSceneView& View, const FPostProcessMaterialInputs& InOutInputs)
 	{
-		FinalColorPixelFormat = InOutInputs.Textures[(uint32)EPostProcessMaterialInput::SceneColor].Texture->Desc.Format;
+		FinalColorPixelFormat = InOutInputs.Textures[(uint32)EPostProcessMaterialInput::SceneColor].TextureSRV->Desc.Texture->Desc.Format;
 		// Don't need to modify anything, just return the untouched scene color texture back to post processing.
-		return ReturnUntouchedSceneColorForPostProcessing(InOutInputs);
+		return InOutInputs.ReturnUntouchedSceneColorForPostProcessing(GraphBuilder);
 	}
 	
 	FScreenPassTexture FPixelInspectorSceneViewExtension::PostProcessPassAfterMotionBlur_RenderThread(FRDGBuilder& GraphBuilder, const FSceneView& View, const FPostProcessMaterialInputs& InOutInputs)
 	{
-		HDRPixelFormat = InOutInputs.Textures[(uint32)EPostProcessMaterialInput::SceneColor].Texture->Desc.Format;
+		HDRPixelFormat = InOutInputs.Textures[(uint32)EPostProcessMaterialInput::SceneColor].TextureSRV->Desc.Texture->Desc.Format;
 		// Don't need to modify anything, just return the untouched scene color texture back to post processing.
-		return ReturnUntouchedSceneColorForPostProcessing(InOutInputs);
+		return InOutInputs.ReturnUntouchedSceneColorForPostProcessing(GraphBuilder);
 	}
 
 };
diff --git a/Engine/Source/Runtime/Engine/Classes/Engine/BlendableInterface.h b/Engine/Source/Runtime/Engine/Classes/Engine/BlendableInterface.h
--- a/Engine/Source/Runtime/Engine/Classes/Engine/BlendableInterface.h
+++ b/Engine/Source/Runtime/Engine/Classes/Engine/BlendableInterface.h
@@ -15,6 +15,7 @@ enum EBlendableLocation : int
 {
 	/** Post process material location to modify the scene color, between translucency distortion and DOF.
 	 * Always run at rendering resolution.
+	 * Inputs and output always in linear color space.
 	 *
 	 * Input0:former pass scene color, excluding AfterDOF translucency
 	 * Input1:AfterDOF translucency.
@@ -23,6 +24,7 @@ enum EBlendableLocation : int
 
 	/** Post process material location to modify the scene color, between DOF and AfterDOF translucency.
 	 * Always run at rendering resolution.
+	 * Inputs and output always in linear color space.
 	 *
 	 * Input0:former pass scene color, excluding AfterDOF translucency
 	 * Input1:AfterDOF translucency.
@@ -31,6 +33,7 @@ enum EBlendableLocation : int
 
 	/** Post process material location to modify the AfterDOF translucency, before composition into the scene color.
 	 * Always run at rendering resolution.
+	 * Inputs and output always in linear color space.
 	 *
 	 * Input0:scene color without translucency, after DOF,
 	 * Input1:AfterDOF translucency.
@@ -39,13 +42,23 @@ enum EBlendableLocation : int
 
 	/** Post process material location to compose a backplate into SSR, between TSR/TAA and next frame's SSR.
 	 * Runs at display resolution with TSR or TAAU, rendering resolution otherwise.
+	 * Inputs and output always in linear color space.
 	 *
 	 * Input0:TAA/TSR output,
 	 */
 	BL_SSRInput = 4 UMETA(DisplayName = "SSR Input"),
 
+	/** Post process material location to modify the scene color, before bloom.
+	 * Runs at display resolution with TSR or TAAU, rendering resolution otherwise.
+	 * Inputs and output always in linear color space.
+	 *
+	 * Input0:former pass scene color,
+	 */
+	BL_SceneColorBeforeBloom = 6 UMETA(DisplayName = "Scene Color Before Bloom"),
+
 	/** Post process material replacing the tone mapper, to modify the scene color.
 	 * Runs at display resolution with TSR or TAAU, rendering resolution otherwise.
+	 * Inputs are always linear color space.
 	 *
 	 * Input0:former pass scene color,
 	 * Input1:AfterDOF translucency.
@@ -54,13 +67,14 @@ enum EBlendableLocation : int
 
 	/** Post process material location to modify the scene color, after tone mapper.
 	 * Runs at display resolution with TSR or TAAU, rendering resolution otherwise.
+	 * Inputs and output in different color spaces, based rendering settings for instance (sRGB/Rec709, HDR or even Linear Color).
 	 *
 	 * Input0:former pass scene color,
 	 * Input1:AfterDOF translucency.
 	 */
 	BL_SceneColorAfterTonemapping = 0 UMETA(DisplayName = "Scene Color After Tonemapping"),
 
-	BL_MAX = 6 UMETA(Hidden),
+	BL_MAX = 7 UMETA(Hidden),
 
 	// Olds names that needs to be kept forever to ensure asset serialization to work correctly when UENUM() switched from serializing int to names.
 	BL_BeforeTranslucency UE_DEPRECATED(5.4, "Renamed to BL_SceneColorBeforeDOF")        = BL_SceneColorBeforeDOF UMETA(Hidden),
diff --git a/Engine/Source/Runtime/Renderer/Private/PostProcess/PostProcessMaterial.cpp b/Engine/Source/Runtime/Renderer/Private/PostProcess/PostProcessMaterial.cpp
--- a/Engine/Source/Runtime/Renderer/Private/PostProcess/PostProcessMaterial.cpp
+++ b/Engine/Source/Runtime/Renderer/Private/PostProcess/PostProcessMaterial.cpp
@@ -460,12 +460,12 @@ FPostProcessMaterialParameters* GetPostProcessMaterialParameters(
 
 	for (uint32 InputIndex = 0; InputIndex < kPostProcessMaterialInputCountMax; ++InputIndex)
 	{
-		FScreenPassTexture Input = Inputs.GetInput((EPostProcessMaterialInput)InputIndex);
+		FScreenPassTextureSlice Input = Inputs.GetInput((EPostProcessMaterialInput)InputIndex);
 
 		// Need to provide valid textures for when shader compilation doesn't cull unused parameters.
-		if (!Input.Texture || !MaterialShaderMap->UsesSceneTexture(PPI_PostProcessInput0 + InputIndex))
+		if (!Input.IsValid() || !MaterialShaderMap->UsesSceneTexture(PPI_PostProcessInput0 + InputIndex))
 		{
-			Input = BlackDummy;
+			Input = FScreenPassTextureSlice::CreateFromScreenPassTexture(GraphBuilder, BlackDummy);
 		}
 
 		PostProcessMaterialParameters->PostProcessInput[InputIndex] = GetScreenPassTextureInput(Input, PointClampSampler);
@@ -530,7 +530,7 @@ void AddNeuralPostProcessPass(
 {
 	Inputs.Validate();
 
-	const FScreenPassTexture SceneColor = Inputs.GetInput(EPostProcessMaterialInput::SceneColor);
+	const FScreenPassTexture SceneColor = FScreenPassTexture::CopyFromSlice(GraphBuilder, Inputs.GetInput(EPostProcessMaterialInput::SceneColor));
 
 	const ERHIFeatureLevel::Type FeatureLevel = View.GetFeatureLevel();
 
@@ -691,7 +691,7 @@ FScreenPassTexture AddPostProcessMaterialPass(
 	GetMaterialInfo(MaterialInterface, FeatureLevel, Inputs, Material, MaterialRenderProxy, MaterialShaderMap, VertexShader, PixelShader);
 
 	EBlendableLocation BlendableLocation = EBlendableLocation(Material->GetBlendableLocation());
-	const FScreenPassTexture SceneColor = Inputs.GetInput(BlendableLocation == BL_TranslucencyAfterDOF ? EPostProcessMaterialInput::SeparateTranslucency : EPostProcessMaterialInput::SceneColor);
+	const FScreenPassTexture SceneColor = FScreenPassTexture::CopyFromSlice(GraphBuilder, Inputs.GetInput(BlendableLocation == BL_TranslucencyAfterDOF ? EPostProcessMaterialInput::SeparateTranslucency : EPostProcessMaterialInput::SceneColor));
 
 	check(VertexShader.IsValid());
 	check(PixelShader.IsValid());
@@ -902,7 +902,7 @@ FScreenPassTexture AddPostProcessMaterialPass(
 {
 	if (!ensureMsgf(View.bIsViewInfo, TEXT("AddPostProcessMaterialPass requires that its View parameter is an FViewInfo.")))
 	{
-		return Inputs.GetInput(EPostProcessMaterialInput::SceneColor);
+		return FScreenPassTexture::CopyFromSlice(GraphBuilder, Inputs.GetInput(EPostProcessMaterialInput::SceneColor));
 	}
 
 	return AddPostProcessMaterialPass(GraphBuilder, static_cast<const FViewInfo&>(View), Inputs, MaterialInterface);
@@ -988,13 +988,14 @@ FScreenPassTexture AddPostProcessMaterialChain(
 	const FPostProcessMaterialChain& Materials,
 	EPostProcessMaterialInput MaterialInput)
 {
-	FScreenPassTexture Outputs = InputsTemplate.GetInput(MaterialInput);
+	FScreenPassTextureSlice CurrentInput = InputsTemplate.GetInput(MaterialInput);
+	FScreenPassTexture Outputs;
 
 	bool bFirstMaterialInChain = true;
 	for (const UMaterialInterface* MaterialInterface : Materials)
 	{
 		FPostProcessMaterialInputs Inputs = InputsTemplate;
-		Inputs.SetInput(MaterialInput, Outputs);
+		Inputs.SetInput(MaterialInput, CurrentInput);
 		
 		// Only the first material in the chain needs to decode the input color
 		Inputs.bMetalMSAAHDRDecode = Inputs.bMetalMSAAHDRDecode && bFirstMaterialInChain;
@@ -1007,6 +1008,12 @@ FScreenPassTexture AddPostProcessMaterialChain(
 		}
 
 		Outputs = AddPostProcessMaterialPass(GraphBuilder, View, Inputs, MaterialInterface);
+		CurrentInput = FScreenPassTextureSlice::CreateFromScreenPassTexture(GraphBuilder, Outputs);
+	}
+
+	if (!Outputs.IsValid())
+	{
+		Outputs = FScreenPassTexture::CopyFromSlice(GraphBuilder, CurrentInput);
 	}
 
 	return Outputs;
@@ -1114,7 +1121,7 @@ FScreenPassTexture AddHighResolutionScreenshotMaskPass(
 	{
 		FPostProcessMaterialInputs PassInputs;
 		PassSequence.AcceptOverrideIfLastPass(EPass::Material, PassInputs.OverrideOutput);
-		PassInputs.SetInput(EPostProcessMaterialInput::SceneColor, Output);
+		PassInputs.SetInput(GraphBuilder, EPostProcessMaterialInput::SceneColor, Output);
 		PassInputs.SceneTextures = Inputs.SceneTextures;
 
 		Output = AddPostProcessMaterialPass(GraphBuilder, View, PassInputs, Inputs.Material);
@@ -1125,7 +1132,7 @@ FScreenPassTexture AddHighResolutionScreenshotMaskPass(
 		PassSequence.AcceptPass(EPass::MaskMaterial);
 
 		FPostProcessMaterialInputs PassInputs;
-		PassInputs.SetInput(EPostProcessMaterialInput::SceneColor, Output);
+		PassInputs.SetInput(GraphBuilder, EPostProcessMaterialInput::SceneColor, Output);
 		PassInputs.SceneTextures = Inputs.SceneTextures;
 
 		// Explicitly allocate the render target to match the FSceneView extents and rect, so the output pixel arrangement matches
@@ -1159,7 +1166,7 @@ FScreenPassTexture AddHighResolutionScreenshotMaskPass(
 	{
 		FPostProcessMaterialInputs PassInputs;
 		PassSequence.AcceptOverrideIfLastPass(EPass::CaptureRegionMaterial, PassInputs.OverrideOutput);
-		PassInputs.SetInput(EPostProcessMaterialInput::SceneColor, Output);
+		PassInputs.SetInput(GraphBuilder, EPostProcessMaterialInput::SceneColor, Output);
 		PassInputs.SceneTextures = Inputs.SceneTextures;
 
 		Output = AddPostProcessMaterialPass(GraphBuilder, View, PassInputs, Inputs.CaptureRegionMaterial);
diff --git a/Engine/Source/Runtime/Renderer/Private/PostProcess/PostProcessMaterial.h b/Engine/Source/Runtime/Renderer/Private/PostProcess/PostProcessMaterial.h
--- a/Engine/Source/Runtime/Renderer/Private/PostProcess/PostProcessMaterial.h
+++ b/Engine/Source/Runtime/Renderer/Private/PostProcess/PostProcessMaterial.h
@@ -21,7 +21,7 @@ BEGIN_SHADER_PARAMETER_STRUCT(FPostProcessMaterialParameters, )
 	SHADER_PARAMETER_RDG_UNIFORM_BUFFER(FSubstrateGlobalUniformParameters, Substrate)
 	SHADER_PARAMETER_STRUCT_INCLUDE(FNeuralPostProcessShaderParameters, NeuralPostProcessParameters)
 	SHADER_PARAMETER_STRUCT(FScreenPassTextureViewportParameters, PostProcessOutput)
-	SHADER_PARAMETER_STRUCT_ARRAY(FScreenPassTextureInput, PostProcessInput, [kPostProcessMaterialInputCountMax])
+	SHADER_PARAMETER_STRUCT_ARRAY(FScreenPassTextureSliceInput, PostProcessInput, [kPostProcessMaterialInputCountMax])
 	SHADER_PARAMETER_STRUCT_ARRAY(FScreenPassTextureInput, PathTracingPostProcessInput, [kPathTracingPostProcessMaterialInputCountMax])
 	SHADER_PARAMETER_SAMPLER(SamplerState, PostProcessInput_BilinearSampler)
 	SHADER_PARAMETER_RDG_BUFFER_SRV(StructuredBuffer<float4>, EyeAdaptationBuffer)
diff --git a/Engine/Source/Runtime/Renderer/Private/PostProcess/PostProcessVisualizeBuffer.cpp b/Engine/Source/Runtime/Renderer/Private/PostProcess/PostProcessVisualizeBuffer.cpp
--- a/Engine/Source/Runtime/Renderer/Private/PostProcess/PostProcessVisualizeBuffer.cpp
+++ b/Engine/Source/Runtime/Renderer/Private/PostProcess/PostProcessVisualizeBuffer.cpp
@@ -347,11 +347,11 @@ FScreenPassTexture AddVisualizeGBufferOverviewPass(
 		RDG_EVENT_SCOPE(GraphBuilder, "%s", *MaterialName);
 
 		FPostProcessMaterialInputs PostProcessMaterialInputs;
-		PostProcessMaterialInputs.SetInput(EPostProcessMaterialInput::SceneColor, Inputs.SceneColor);
-		PostProcessMaterialInputs.SetInput(EPostProcessMaterialInput::SeparateTranslucency, Inputs.SeparateTranslucency);
-		PostProcessMaterialInputs.SetInput(EPostProcessMaterialInput::PreTonemapHDRColor, Inputs.SceneColorBeforeTonemap);
-		PostProcessMaterialInputs.SetInput(EPostProcessMaterialInput::PostTonemapHDRColor, Inputs.SceneColorAfterTonemap);
-		PostProcessMaterialInputs.SetInput(EPostProcessMaterialInput::Velocity, Inputs.Velocity);
+		PostProcessMaterialInputs.SetInput(GraphBuilder, EPostProcessMaterialInput::SceneColor, Inputs.SceneColor);
+		PostProcessMaterialInputs.SetInput(GraphBuilder, EPostProcessMaterialInput::SeparateTranslucency, Inputs.SeparateTranslucency);
+		PostProcessMaterialInputs.SetInput(GraphBuilder, EPostProcessMaterialInput::PreTonemapHDRColor, Inputs.SceneColorBeforeTonemap);
+		PostProcessMaterialInputs.SetInput(GraphBuilder, EPostProcessMaterialInput::PostTonemapHDRColor, Inputs.SceneColorAfterTonemap);
+		PostProcessMaterialInputs.SetInput(GraphBuilder, EPostProcessMaterialInput::Velocity, Inputs.Velocity);
 
 		if (View.Family->EngineShowFlags.PathTracing && Inputs.PathTracingResources->bPostProcessEnabled)
 		{
diff --git a/Engine/Source/Runtime/Renderer/Private/PostProcess/PostProcessing.cpp b/Engine/Source/Runtime/Renderer/Private/PostProcess/PostProcessing.cpp
--- a/Engine/Source/Runtime/Renderer/Private/PostProcess/PostProcessing.cpp
+++ b/Engine/Source/Runtime/Renderer/Private/PostProcess/PostProcessing.cpp
@@ -365,6 +365,7 @@ void AddPostProcessingPasses(
 	enum class EPass : uint32
 	{
 		MotionBlur,
+		PostProcessMaterialBeforeBloom,
 		Tonemap,
 		FXAA,
 		PostProcessMaterialAfterTonemapping,
@@ -419,9 +420,10 @@ void AddPostProcessingPasses(
 	const TCHAR* PassNames[] =
 	{
 		TEXT("MotionBlur"),
+		TEXT("PostProcessMaterial (SceneColorBeforeBloom)"),
 		TEXT("Tonemap"),
 		TEXT("FXAA"),
-		TEXT("PostProcessMaterial (AfterTonemapping)"),
+		TEXT("PostProcessMaterial (SceneColorAfterTonemapping)"),
 		TEXT("VisualizeLumenScene"),
 		TEXT("VisualizeDepthOfField"),
 		TEXT("VisualizeStationaryLightOverlap"),
@@ -509,7 +511,7 @@ void AddPostProcessingPasses(
 	{ 
 		FPostProcessMaterialInputs PostProcessMaterialInputs;
 
-		PostProcessMaterialInputs.SetInput(EPostProcessMaterialInput::SceneColor, InSceneColor);
+		PostProcessMaterialInputs.SetInput(GraphBuilder, EPostProcessMaterialInput::SceneColor, InSceneColor);
 
 		FIntRect ViewRect{ 0, 0, 1, 1 };
 
@@ -530,8 +532,8 @@ void AddPostProcessingPasses(
 			ViewRect = PostDOFTranslucencyResources.ViewRect;
 		}
 
-		PostProcessMaterialInputs.SetInput(EPostProcessMaterialInput::SeparateTranslucency, FScreenPassTexture(PostDOFTranslucencyResources.GetColorForRead(GraphBuilder), ViewRect));
-		PostProcessMaterialInputs.SetInput(EPostProcessMaterialInput::Velocity, Velocity);
+		PostProcessMaterialInputs.SetInput(GraphBuilder, EPostProcessMaterialInput::SeparateTranslucency, FScreenPassTexture(PostDOFTranslucencyResources.GetColorForRead(GraphBuilder), ViewRect));
+		PostProcessMaterialInputs.SetInput(GraphBuilder, EPostProcessMaterialInput::Velocity, Velocity);
 		PostProcessMaterialInputs.SceneTextures = GetSceneTextureShaderParameters(Inputs.SceneTextures);
 		PostProcessMaterialInputs.CustomDepthTexture = CustomDepth.Texture;
 		PostProcessMaterialInputs.bManualStencilTest = Inputs.bSeparateCustomStencil;
@@ -552,7 +554,7 @@ void AddPostProcessingPasses(
 
 			for (int32 AfterPassCallbackIndex = 0; AfterPassCallbackIndex < PassCallbacks.Num(); AfterPassCallbackIndex++)
 			{
-				InOutPostProcessAfterPassInputs.SetInput(EPostProcessMaterialInput::SceneColor, InSceneColor);
+				InOutPostProcessAfterPassInputs.SetInput(GraphBuilder, EPostProcessMaterialInput::SceneColor, InSceneColor);
 
 				FAfterPassCallbackDelegate& AfterPassCallback = PassCallbacks[AfterPassCallbackIndex];
 				PassSequence.AcceptOverrideIfLastPass(InPass, InOutPostProcessAfterPassInputs.OverrideOutput, AfterPassCallbackIndex);
@@ -635,9 +637,11 @@ void AddPostProcessingPasses(
 
 		const FIntPoint PostTAAViewSize = (View.PrimaryScreenPercentageMethod == EPrimaryScreenPercentageMethod::TemporalUpscale && TAAConfig != EMainTAAPassConfig::Disabled) ? View.GetSecondaryViewRectSize() : View.ViewRect.Size();
 
+		const FPostProcessMaterialChain PostProcessMaterialBeforeBloomChain = GetPostProcessMaterialChain(View, BL_SceneColorBeforeBloom);
 		const FPostProcessMaterialChain PostProcessMaterialAfterTonemappingChain = GetPostProcessMaterialChain(View, BL_SceneColorAfterTonemapping);
 
 		PassSequence.SetEnabled(EPass::MotionBlur, bVisualizeMotionBlur || bMotionBlurEnabled);
+		PassSequence.SetEnabled(EPass::PostProcessMaterialBeforeBloom, PostProcessMaterialBeforeBloomChain.Num() != 0);
 		PassSequence.SetEnabled(EPass::Tonemap, bTonemapEnabled);
 		PassSequence.SetEnabled(EPass::FXAA, AntiAliasingMethod == AAM_FXAA);
 		PassSequence.SetEnabled(EPass::PostProcessMaterialAfterTonemapping, PostProcessMaterialAfterTonemappingChain.Num() != 0);
@@ -681,8 +685,8 @@ void AddPostProcessingPasses(
 			(bLensFlareEnabled && bFFTBloomEnabled) ||
 			bLocalExposureBlurredLum);
 		extern int32 GSSRHalfResSceneColor;
-		const bool bNeedPostMotionBlurHalfRes = !bProcessQuarterResolution || (bFFTBloomEnabled && FFTBloomResolutionFraction > 0.25f && FFTBloomResolutionFraction <= 0.5f) || (ReflectionsMethod == EReflectionsMethod::SSR && !View.bStatePrevViewInfoIsReadOnly && GSSRHalfResSceneColor);
-		const bool bNeedPostMotionBlurQuarterRes = bProcessQuarterResolution || (bFFTBloomEnabled && FFTBloomResolutionFraction <= 0.25f);
+		const bool bNeedBeforeBloomHalfRes    = !bProcessQuarterResolution || (bFFTBloomEnabled && FFTBloomResolutionFraction > 0.25f && FFTBloomResolutionFraction <= 0.5f) || (ReflectionsMethod == EReflectionsMethod::SSR && !View.bStatePrevViewInfoIsReadOnly && GSSRHalfResSceneColor);
+		const bool bNeedBeforeBloomQuarterRes = bProcessQuarterResolution || (bFFTBloomEnabled && FFTBloomResolutionFraction <= 0.25f);
 
 
 		const FPostProcessMaterialChain MaterialChainSceneColorBeforeDOF = GetPostProcessMaterialChain(View, BL_SceneColorBeforeDOF);
@@ -796,14 +800,18 @@ void AddPostProcessingPasses(
 					UpscalerPassInputs.bGenerateVelocityFlattenTextures = FVelocityFlattenTextures::AllowExternal(View) && !bVisualizeMotionBlur;
 				}
 			}
+			else if (PostProcessMaterialBeforeBloomChain.Num() > 0)
+			{
+				// NOP
+			}
 			else
 			{
 				UpscalerPassInputs.bGenerateSceneColorHalfRes =
-					bNeedPostMotionBlurHalfRes &&
+					bNeedBeforeBloomHalfRes &&
 					DownsampleQuality == EDownsampleQuality::Low;
 
 				UpscalerPassInputs.bGenerateSceneColorQuarterRes =
-					bNeedPostMotionBlurQuarterRes &&
+					bNeedBeforeBloomQuarterRes &&
 					DownsampleQuality == EDownsampleQuality::Low;
 			}
 			UpscalerPassInputs.bAllowFullResSlice = PassSequence.IsEnabled(EPass::MotionBlur) || PassSequence.IsEnabled(EPass::Tonemap);
@@ -890,11 +898,11 @@ void AddPostProcessingPasses(
 			{
 				if (PassOutput.IsValid())
 				{
-					InOutPassInputs.SetInput(EPostProcessMaterialInput::SceneColor, PassOutput);
+					InOutPassInputs.SetInput(GraphBuilder, EPostProcessMaterialInput::SceneColor, PassOutput);
 				}
 				else
 				{
-					InOutPassInputs.SetInput(EPostProcessMaterialInput::SceneColor, FScreenPassTexture::CopyFromSlice(GraphBuilder, SceneColorSlice));
+					InOutPassInputs.SetInput(GraphBuilder, EPostProcessMaterialInput::SceneColor, FScreenPassTexture::CopyFromSlice(GraphBuilder, SceneColorSlice));
 				}
 				PassOutput = PassCallback.Execute(GraphBuilder, View, InOutPassInputs);
 			}
@@ -913,8 +921,8 @@ void AddPostProcessingPasses(
 		{
 			FMotionBlurInputs PassInputs;
 			PassSequence.AcceptOverrideIfLastPass(EPass::MotionBlur, PassInputs.OverrideOutput);
-			PassInputs.bOutputHalfRes = bNeedPostMotionBlurHalfRes && DownsampleQuality == EDownsampleQuality::Low;
-			PassInputs.bOutputQuarterRes = bNeedPostMotionBlurQuarterRes && DownsampleQuality == EDownsampleQuality::Low;
+			PassInputs.bOutputHalfRes = PostProcessMaterialBeforeBloomChain.Num() == 0 && bNeedBeforeBloomHalfRes && DownsampleQuality == EDownsampleQuality::Low;
+			PassInputs.bOutputQuarterRes = bNeedBeforeBloomQuarterRes && DownsampleQuality == EDownsampleQuality::Low;
 			PassInputs.SceneColor = SceneColorSlice;
 			PassInputs.SceneDepth = SceneDepth;
 			PassInputs.SceneVelocity = Velocity;
@@ -962,10 +970,20 @@ void AddPostProcessingPasses(
 			SceneColorSlice = NewSceneColorSlice;
 		}
 
-		// Generate post motion blur lower res scene color if they have not been generated.
+		// Post Process Material Chain - Before Bloom
+		if (PassSequence.IsEnabled(EPass::PostProcessMaterialBeforeBloom))
+		{
+			FPostProcessMaterialInputs PostProcessMaterialInputs = GetPostProcessMaterialInputs(FScreenPassTexture());
+			PassSequence.AcceptOverrideIfLastPass(EPass::PostProcessMaterialBeforeBloom, PostProcessMaterialInputs.OverrideOutput);
+			PostProcessMaterialInputs.SetInput(EPostProcessMaterialInput::SceneColor, SceneColorSlice);
+
+			SceneColorSlice = FScreenPassTextureSlice::CreateFromScreenPassTexture(GraphBuilder, AddPostProcessMaterialChain(GraphBuilder, View, PostProcessMaterialInputs, PostProcessMaterialBeforeBloomChain));
+		}
+
+		// Generate before bloom lower res scene color if they have not been generated.
 		{
-			if ((bNeedPostMotionBlurHalfRes && !HalfResSceneColor.IsValid()) ||
-				(bNeedPostMotionBlurQuarterRes && !QuarterResSceneColor.IsValid() && !HalfResSceneColor.IsValid()))
+			if ((bNeedBeforeBloomHalfRes && !HalfResSceneColor.IsValid()) ||
+				(bNeedBeforeBloomQuarterRes && !QuarterResSceneColor.IsValid() && !HalfResSceneColor.IsValid()))
 			{
 				FDownsamplePassInputs PassInputs;
 				PassInputs.Name = TEXT("PostProcessing.SceneColor.HalfRes");
@@ -976,7 +994,7 @@ void AddPostProcessingPasses(
 				HalfResSceneColor = FScreenPassTextureSlice::CreateFromScreenPassTexture(GraphBuilder, AddDownsamplePass(GraphBuilder, View, PassInputs));
 			}
 
-			if (bNeedPostMotionBlurQuarterRes && !QuarterResSceneColor.IsValid())
+			if (bNeedBeforeBloomQuarterRes && !QuarterResSceneColor.IsValid())
 			{
 				FDownsamplePassInputs PassInputs;
 				PassInputs.Name = TEXT("PostProcessing.SceneColor.QuarterRes");
@@ -1183,9 +1201,9 @@ void AddPostProcessingPasses(
 
 				FPostProcessMaterialInputs PassInputs;
 				PassSequence.AcceptOverrideIfLastPass(EPass::Tonemap, PassInputs.OverrideOutput);
-				PassInputs.SetInput(EPostProcessMaterialInput::SceneColor, FScreenPassTexture::CopyFromSlice(GraphBuilder, SceneColorSlice));
+				PassInputs.SetInput(GraphBuilder, EPostProcessMaterialInput::SceneColor, FScreenPassTexture::CopyFromSlice(GraphBuilder, SceneColorSlice));
 				//PassInputs.SetInput(EPostProcessMaterialInput::SeparateTranslucency, SeparateTranslucency);
-				PassInputs.SetInput(EPostProcessMaterialInput::CombinedBloom, Bloom);
+				PassInputs.SetInput(GraphBuilder, EPostProcessMaterialInput::CombinedBloom, Bloom);
 				PassInputs.SceneTextures = GetSceneTextureShaderParameters(Inputs.SceneTextures);
 				PassInputs.CustomDepthTexture = CustomDepth.Texture;
 				PassInputs.bManualStencilTest = Inputs.bSeparateCustomStencil;
@@ -1254,8 +1272,8 @@ void AddPostProcessingPasses(
 		{
 			FPostProcessMaterialInputs PassInputs = GetPostProcessMaterialInputs(SceneColor);
 			PassSequence.AcceptOverrideIfLastPass(EPass::PostProcessMaterialAfterTonemapping, PassInputs.OverrideOutput);
-			PassInputs.SetInput(EPostProcessMaterialInput::PreTonemapHDRColor, FScreenPassTexture::CopyFromSlice(GraphBuilder, SceneColorBeforeTonemapSlice));
-			PassInputs.SetInput(EPostProcessMaterialInput::PostTonemapHDRColor, SceneColorAfterTonemap);
+			PassInputs.SetInput(GraphBuilder, EPostProcessMaterialInput::PreTonemapHDRColor, FScreenPassTexture::CopyFromSlice(GraphBuilder, SceneColorBeforeTonemapSlice));
+			PassInputs.SetInput(GraphBuilder, EPostProcessMaterialInput::PostTonemapHDRColor, SceneColorAfterTonemap);
 			PassInputs.SceneTextures = GetSceneTextureShaderParameters(Inputs.SceneTextures);
 
 			SceneColor = AddPostProcessMaterialChain(GraphBuilder, View, PassInputs, PostProcessMaterialAfterTonemappingChain);
@@ -2008,7 +2026,7 @@ void AddVisualizeCalibrationMaterialPostProcessingPasses(FRDGBuilder& GraphBuild
 
 	// Post Process Material - Before Color Correction
 	FPostProcessMaterialInputs PostProcessMaterialInputs;
-	PostProcessMaterialInputs.SetInput(EPostProcessMaterialInput::SceneColor, SceneColor);
+	PostProcessMaterialInputs.SetInput(GraphBuilder, EPostProcessMaterialInput::SceneColor, SceneColor);
 	PostProcessMaterialInputs.SceneTextures = GetSceneTextureShaderParameters(Inputs.SceneTextures);
 
 	SceneColor = AddPostProcessMaterialPass(GraphBuilder, View, PostProcessMaterialInputs, InMaterialInterface);
@@ -2194,7 +2212,7 @@ void AddMobilePostProcessingPasses(FRDGBuilder& GraphBuilder, FScene* Scene, con
 			PassSequence.AcceptOverrideIfLastPass(EPass::PostProcessMaterialAfterTonemapping, PostProcessMaterialInputs.OverrideOutput);
 		}
 
-		PostProcessMaterialInputs.SetInput(EPostProcessMaterialInput::SceneColor, SceneColor);
+		PostProcessMaterialInputs.SetInput(GraphBuilder, EPostProcessMaterialInput::SceneColor, SceneColor);
 
 		PostProcessMaterialInputs.CustomDepthTexture = CustomDepth.Texture;
 
diff --git a/Engine/Source/Runtime/Renderer/Public/PostProcess/PostProcessMaterialInputs.h b/Engine/Source/Runtime/Renderer/Public/PostProcess/PostProcessMaterialInputs.h
--- a/Engine/Source/Runtime/Renderer/Public/PostProcess/PostProcessMaterialInputs.h
+++ b/Engine/Source/Runtime/Renderer/Public/PostProcess/PostProcessMaterialInputs.h
@@ -43,12 +43,17 @@ enum class EPathTracingPostProcessMaterialInput : uint32
 
 struct FPostProcessMaterialInputs
 {
-	inline void SetInput(EPostProcessMaterialInput Input, FScreenPassTexture Texture)
+	inline void SetInput(FRDGBuilder& GraphBuilder, EPostProcessMaterialInput Input, FScreenPassTexture Texture)
+	{
+		SetInput(Input, FScreenPassTextureSlice::CreateFromScreenPassTexture(GraphBuilder, Texture));
+	}
+
+	inline void SetInput(EPostProcessMaterialInput Input, FScreenPassTextureSlice Texture)
 	{
 		Textures[(uint32)Input] = Texture;
 	}
 
-	inline FScreenPassTexture GetInput(EPostProcessMaterialInput Input) const
+	inline FScreenPassTextureSlice GetInput(EPostProcessMaterialInput Input) const
 	{
 		return Textures[(uint32)Input];
 	}
@@ -83,17 +88,34 @@ struct FPostProcessMaterialInputs
 
 	inline void ValidateInputExists(EPostProcessMaterialInput Input) const
 	{
-		const FScreenPassTexture Texture = GetInput(EPostProcessMaterialInput::SceneColor);
+		const FScreenPassTextureSlice Texture = GetInput(EPostProcessMaterialInput::SceneColor);
 		check(Texture.IsValid());
 	}
 
+	/**
+	* A helper function that extracts the right scene color texture, untouched, to be used further in post processing.
+	*/
+	inline FScreenPassTexture ReturnUntouchedSceneColorForPostProcessing(FRDGBuilder& GraphBuilder) const
+	{
+		if (OverrideOutput.IsValid())
+		{
+			return OverrideOutput;
+		}
+		else
+		{
+			/** We don't want to modify scene texture in any way. We just want it to be passed back onto the next stage. */
+			FScreenPassTextureSlice SceneTexture = const_cast<FScreenPassTextureSlice&>(Textures[(uint32)EPostProcessMaterialInput::SceneColor]);
+			return FScreenPassTexture::CopyFromSlice(GraphBuilder, SceneTexture);
+		}
+	}
+
 	// [Optional] Render to the specified output. If invalid, a new texture is created and returned.
 	FScreenPassRenderTarget OverrideOutput;
 
 	/** Array of input textures bound to the material. The first element represents the output from
 	 *  the previous post process and is required. All other inputs are optional.
 	 */
-	TStaticArray<FScreenPassTexture, kPostProcessMaterialInputCountMax> Textures;
+	TStaticArray<FScreenPassTextureSlice, kPostProcessMaterialInputCountMax> Textures;
 
 	/**
 	*	Array of input textures bound to the material from path tracing. All inputs are optional
