Speculatively try fixing FastVRAM on bloom setup and gaussian blur passes that could have been relying on the previous passes' flag before

diff --git a/Engine/Source/Runtime/Renderer/Private/PostProcess/PostProcessBloomSetup.cpp b/Engine/Source/Runtime/Renderer/Private/PostProcess/PostProcessBloomSetup.cpp
--- a/Engine/Source/Runtime/Renderer/Private/PostProcess/PostProcessBloomSetup.cpp
+++ b/Engine/Source/Runtime/Renderer/Private/PostProcess/PostProcessBloomSetup.cpp
@@ -133,7 +133,7 @@ FScreenPassTexture AddBloomSetupPass(FRDGBuilder& GraphBuilder, const FViewInfo&
 		InputDesc.Extent,
 		InputDesc.Format,
 		FClearValueBinding::None,
-		/* InFlags = */ TexCreate_ShaderResource | (bIsComputePass ? TexCreate_UAV : TexCreate_RenderTargetable));
+		/* InFlags = */ TexCreate_ShaderResource | (bIsComputePass ? TexCreate_UAV : TexCreate_RenderTargetable) | (InputDesc.Flags & (TexCreate_FastVRAM | TexCreate_FastVRAMPartialAlloc)));
 
 	const FScreenPassTextureViewport Viewport(Inputs.SceneColor);
 	const FScreenPassRenderTarget Output(GraphBuilder.CreateTexture(OutputDesc, TEXT("BloomSetup")), Viewport.Rect, ERenderTargetLoadAction::ENoAction);
diff --git a/Engine/Source/Runtime/Renderer/Private/PostProcess/PostProcessWeightedSampleSum.cpp b/Engine/Source/Runtime/Renderer/Private/PostProcess/PostProcessWeightedSampleSum.cpp
--- a/Engine/Source/Runtime/Renderer/Private/PostProcess/PostProcessWeightedSampleSum.cpp
+++ b/Engine/Source/Runtime/Renderer/Private/PostProcess/PostProcessWeightedSampleSum.cpp
@@ -388,7 +388,7 @@ FScreenPassTexture AddGaussianBlurPass(
 		OutputViewport.Extent,
 		InputDesc.Format,
 		FClearValueBinding(FLinearColor::Transparent),
-		/* InFlags = */ TexCreate_ShaderResource | (bIsComputePass ? TexCreate_UAV : (TexCreate_RenderTargetable | TexCreate_NoFastClear)));
+		/* InFlags = */ TexCreate_ShaderResource | (bIsComputePass ? TexCreate_UAV : (TexCreate_RenderTargetable | TexCreate_NoFastClear)) | (InputDesc.Flags & (TexCreate_FastVRAM | TexCreate_FastVRAMPartialAlloc)));
 
 	const FScreenPassRenderTarget Output(GraphBuilder.CreateTexture(OutputDesc, Name), OutputViewport.Rect, ERenderTargetLoadAction::ENoAction);
 
