Don't need to update PrimitiveRayTracingGroupIds and CoarseMeshStreamingHandle when updating cached raytracing primitive/instance data.

diff --git a/Engine/Source/Runtime/Renderer/Private/PrimitiveSceneInfo.cpp b/Engine/Source/Runtime/Renderer/Private/PrimitiveSceneInfo.cpp
--- a/Engine/Source/Runtime/Renderer/Private/PrimitiveSceneInfo.cpp
+++ b/Engine/Source/Runtime/Renderer/Private/PrimitiveSceneInfo.cpp
@@ -312,7 +312,7 @@ FPrimitiveSceneInfo::FPrimitiveSceneInfo(const FPrimitiveSceneInfoAdapter& InAda
 	bIsRayTracingStaticRelevant(InAdapter.SceneProxy->IsRayTracingStaticRelevant()),
 	bIsVisibleInRayTracing(InAdapter.SceneProxy->IsVisibleInRayTracing()),
 	bCachedRaytracingDataDirty(false),
-	CoarseMeshStreamingHandle(INDEX_NONE),
+	CoarseMeshStreamingHandle(InAdapter.SceneProxy->GetCoarseMeshStreamingHandle()),
 #endif
 	// We want the unsynchronized access here, as the responsibility passes to the primitive scene info.
 	InstanceSceneDataBuffersInternal(InAdapter.SceneProxy->GetInstanceSceneDataBuffers(FPrimitiveSceneProxy::EInstanceBufferAccessFlags::UnsynchronizedAndUnsafe)),
@@ -925,18 +925,24 @@ void FPrimitiveSceneInfo::UpdateCachedRayTracingInstances(FScene* Scene, const T
 
 		for (FPrimitiveSceneInfo* SceneInfo : SceneInfos)
 		{
-			// Write group id
-			const int32 RayTracingGroupId = SceneInfo->Proxy->GetRayTracingGroupId();
-			if (RayTracingGroupId != -1)
+			const int32 PrimitiveIndex = SceneInfo->GetIndex();
+
+#if DO_CHECK
 			{
-				Scene->PrimitiveRayTracingGroupIds[SceneInfo->GetIndex()] = Scene->PrimitiveRayTracingGroups.FindId(RayTracingGroupId);
+				Experimental::FHashElementId SceneRayTracingGroupId;
+				const int32 RayTracingGroupId = SceneInfo->Proxy->GetRayTracingGroupId();
+				if (RayTracingGroupId != -1)
+				{
+					SceneRayTracingGroupId = Scene->PrimitiveRayTracingGroups.FindId(RayTracingGroupId);
+				}
+
+				check(Scene->PrimitiveRayTracingGroupIds[PrimitiveIndex] == SceneRayTracingGroupId);
+				check(SceneInfo->CoarseMeshStreamingHandle == SceneInfo->Proxy->GetCoarseMeshStreamingHandle());
 			}
+#endif
 
 			FRayTracingInstance CachedRayTracingInstance;
-			ERayTracingPrimitiveFlags& Flags = Scene->PrimitiveRayTracingFlags[SceneInfo->GetIndex()];
-
-			// Cache the coarse mesh streaming handle
-			SceneInfo->CoarseMeshStreamingHandle = SceneInfo->Proxy->GetCoarseMeshStreamingHandle();
+			ERayTracingPrimitiveFlags& Flags = Scene->PrimitiveRayTracingFlags[PrimitiveIndex];
 
 			// Write flags
 			Flags = SceneInfo->Proxy->GetCachedRayTracingInstance(CachedRayTracingInstance);
@@ -1025,20 +1031,24 @@ void CacheRayTracingPrimitive(
 	FRayTracingInstance& OutCachedRayTracingInstance, 
 	ERayTracingPrimitiveFlags& OutFlags)
 {
-	// Write group id
-	const int32 RayTracingGroupId = SceneInfo->Proxy->GetRayTracingGroupId();
-	if (RayTracingGroupId != -1)
+#if DO_CHECK
 	{
-		Scene->PrimitiveRayTracingGroupIds[SceneInfo->GetIndex()] = Scene->PrimitiveRayTracingGroups.FindId(RayTracingGroupId);
+		Experimental::FHashElementId SceneRayTracingGroupId;
+		const int32 RayTracingGroupId = SceneInfo->Proxy->GetRayTracingGroupId();
+		if (RayTracingGroupId != -1)
+		{
+			SceneRayTracingGroupId = Scene->PrimitiveRayTracingGroups.FindId(RayTracingGroupId);
+		}
+
+		check(Scene->PrimitiveRayTracingGroupIds[SceneInfo->GetIndex()] == SceneRayTracingGroupId);
+		check(SceneInfo->CoarseMeshStreamingHandle == SceneInfo->Proxy->GetCoarseMeshStreamingHandle());
 	}
+#endif
 
 	// Write flags
 	OutFlags = SceneInfo->Proxy->GetCachedRayTracingInstance(OutCachedRayTracingInstance);
 	UpdateRayTracingInstanceMaskAndFlagsIfNeeded(OutCachedRayTracingInstance, *(SceneInfo->Proxy), nullptr);
 
-	// Cache the coarse mesh streaming handle
-	SceneInfo->CoarseMeshStreamingHandle = SceneInfo->Proxy->GetCoarseMeshStreamingHandle();
-
 	// the following flags cause ray tracing mesh command caching to be disabled
 	static const ERayTracingPrimitiveFlags DisableCacheMeshCommandsFlags = ERayTracingPrimitiveFlags::Dynamic
 		| ERayTracingPrimitiveFlags::Exclude
