Remove unused NumCameraStaticFrames and RayTracingNumIterations from FSceneViewState

diff --git a/Engine/Source/Runtime/Renderer/Private/ScenePrivate.h b/Engine/Source/Runtime/Renderer/Private/ScenePrivate.h
--- a/Engine/Source/Runtime/Renderer/Private/ScenePrivate.h
+++ b/Engine/Source/Runtime/Renderer/Private/ScenePrivate.h
@@ -849,12 +849,6 @@ public:
 	FMatrix44f	PrevViewMatrixForOcclusionQuery;
 	FVector		PrevViewOriginForOcclusionQuery;
 
-#if RHI_RAYTRACING
-	/** Number of consecutive frames the camera is static */
-	uint32 NumCameraStaticFrames;
-	int32 RayTracingNumIterations;
-#endif
-
 	// A counter incremented once each time this view is rendered.
 	uint32 OcclusionFrameCounter;
 
diff --git a/Engine/Source/Runtime/Renderer/Private/SceneVisibility.cpp b/Engine/Source/Runtime/Renderer/Private/SceneVisibility.cpp
--- a/Engine/Source/Runtime/Renderer/Private/SceneVisibility.cpp
+++ b/Engine/Source/Runtime/Renderer/Private/SceneVisibility.cpp
@@ -4937,20 +4937,6 @@ void FSceneRenderer::PrepareViewStateForVisibility(const FSceneTexturesConfig& S
 
 			ViewState->PrevViewMatrixForOcclusionQuery = FMatrix44f(View.ViewMatrices.GetViewMatrix());	// LWC_TODO: Precision loss
 			ViewState->PrevViewOriginForOcclusionQuery = View.ViewMatrices.GetViewOrigin();
-				
-			// store old view matrix and detect conditions where we should reset motion blur 
-#if RHI_RAYTRACING
-			{
-				if (bResetCamera || IsLargeCameraMovement(View, ViewState->PrevFrameViewInfo.ViewMatrices.GetViewMatrix(), ViewState->PrevFrameViewInfo.ViewMatrices.GetViewOrigin(), 0.1f, 0.1f))
-				{
-					ViewState->RayTracingNumIterations = 1;
-				}
-				else
-				{
-					ViewState->RayTracingNumIterations++;
-				}
-			}
-#endif // RHI_RAYTRACING
 
 			// we don't use DeltaTime as it can be 0 (in editor) and is computed by subtracting floats (loses precision over time)
 			// Clamp DeltaWorldTime to reasonable values for the purposes of motion blur, things like TimeDilation can make it very small
