Add Vulkan Tonemap Subpass.

Provides a path for Quest devices to support Mobile HDR  in light of a future possible removal of Mobile LDR.

diff --git a/Engine/Shaders/Private/PostProcessTonemap.usf b/Engine/Shaders/Private/PostProcessTonemap.usf
--- a/Engine/Shaders/Private/PostProcessTonemap.usf
+++ b/Engine/Shaders/Private/PostProcessTonemap.usf
@@ -624,3 +624,50 @@ void MainCS(uint2 DispatchThreadId : SV_DispatchThreadID)
 }
 #endif
 
+#ifdef SUBPASS_MSAA_SAMPLES
+	#if VULKAN_PROFILE
+		#include "/Engine/Public/Platform/Vulkan/VulkanSubpassSupport.ush"
+	#endif
+#endif
+
+half4 FetchAndResolveSceneColor(float2 UV)
+{
+	half4 SceneColor = half4(1.0, 1.0, 1.0, 1.0);
+
+#if VULKAN_PROFILE && SUBPASS_MSAA_SAMPLES > 0	
+	#if SUBPASS_MSAA_SAMPLES == 8
+		SceneColor = half4(VulkanSubpassFetch0(0) + VulkanSubpassFetch0(1) + VulkanSubpassFetch0(2) + VulkanSubpassFetch0(3) +
+			VulkanSubpassFetch0(4) + VulkanSubpassFetch0(5) + VulkanSubpassFetch0(6) + VulkanSubpassFetch0(7)) * 0.125;
+	#elif SUBPASS_MSAA_SAMPLES == 4
+		SceneColor = half4(VulkanSubpassFetch0(0) + VulkanSubpassFetch0(1) + VulkanSubpassFetch0(2) + VulkanSubpassFetch0(3)) * 0.25;
+	#elif SUBPASS_MSAA_SAMPLES == 2
+		SceneColor = half4(VulkanSubpassFetch0(0) + VulkanSubpassFetch0(1)) * 0.5;
+	#elif SUBPASS_MSAA_SAMPLES == 1
+		SceneColor = half4(VulkanSubpassFetch0());
+	#endif
+#else
+	SceneColor = Texture2DSample(ColorTexture, ColorSampler, UV);
+#endif
+	return SceneColor;
+}
+
+void MobileCustomResolve_MainPS(
+	in noperspective float2 UV						: TEXCOORD0,
+	in				 float4 SvPosition				: SV_POSITION,		// after all interpolators
+	out				 float4 OutColor				: SV_Target0
+)
+{
+	ResolvedView = ResolveView();
+	const float OneOverPreExposure = ResolvedView.OneOverPreExposure;
+
+	half4 SceneColor = FetchAndResolveSceneColor(UV);
+	SceneColor.a = 1.0;
+
+	half3 LinearColor = SceneColor.rgb * OneOverPreExposure;
+	LinearColor *= ColorScale0.rgb; // SceneColorTint
+	half3 DeviceColor = ColorLookupTable(LinearColor);
+
+	OutColor.rgb = DeviceColor;
+	OutColor.a = 1.0;
+}
+
diff --git a/Engine/Shaders/Public/Platform/Vulkan/VulkanSubpassSupport.ush b/Engine/Shaders/Public/Platform/Vulkan/VulkanSubpassSupport.ush
--- a/Engine/Shaders/Public/Platform/Vulkan/VulkanSubpassSupport.ush
+++ b/Engine/Shaders/Public/Platform/Vulkan/VulkanSubpassSupport.ush
@@ -7,6 +7,12 @@
 
 #if COMPILER_DXC && PIXELSHADER
 
+	#ifdef SUBPASS_MSAA_SAMPLES
+		#if SUBPASS_MSAA_SAMPLES > 1
+			#define USE_MULTI_SAMPLE 1
+		#endif
+	#endif
+
 	// Note: Do not rename subpass input attachment declarations, or adjust your changes in VulkanBackend.cpp (see VULKAN_SUBPASS_FETCH_VAR_W)
 
 	[[vk::input_attachment_index(0)]]
@@ -14,8 +20,12 @@
 	#define VulkanSubpassDepthFetch() GENERATED_SubpassDepthFetchAttachment.SubpassLoad( )
 
 	[[vk::input_attachment_index(1)]]
-	SubpassInput<float4> GENERATED_SubpassFetchAttachment0;
-	#define VulkanSubpassFetch0() GENERATED_SubpassFetchAttachment0.SubpassLoad()
+	#ifndef USE_MULTI_SAMPLE
+		SubpassInput<float4> GENERATED_SubpassFetchAttachment0;
+	#else
+		SubpassInputMS<float4> GENERATED_SubpassFetchAttachment0;
+	#endif
+	#define VulkanSubpassFetch0 GENERATED_SubpassFetchAttachment0.SubpassLoad
 
 	[[vk::input_attachment_index(2)]]
 	SubpassInput<float4> GENERATED_SubpassFetchAttachment1;
diff --git a/Engine/Source/Runtime/Engine/Classes/Engine/RendererSettings.h b/Engine/Source/Runtime/Engine/Classes/Engine/RendererSettings.h
--- a/Engine/Source/Runtime/Engine/Classes/Engine/RendererSettings.h
+++ b/Engine/Source/Runtime/Engine/Classes/Engine/RendererSettings.h
@@ -965,7 +965,6 @@ class URendererSettings : public UDeveloperSettings
 	uint32 bMobilePostProcessing:1;
 
 	UPROPERTY(config, EditAnywhere, Category = VR, meta = (
-		EditCondition = "!bMobilePostProcessing",
 		ConsoleVariable = "vr.MobileMultiView", DisplayName = "Mobile Multi-View",
 		ToolTip = "Enable single-pass stereoscopic rendering on mobile platforms.",
 		ConfigRestartRequired = true))
diff --git a/Engine/Source/Runtime/Engine/Private/GameViewportClient.cpp b/Engine/Source/Runtime/Engine/Private/GameViewportClient.cpp
--- a/Engine/Source/Runtime/Engine/Private/GameViewportClient.cpp
+++ b/Engine/Source/Runtime/Engine/Private/GameViewportClient.cpp
@@ -1379,9 +1379,8 @@ void UGameViewportClient::Draw(FViewport* InViewport, FCanvas* SceneCanvas)
 	if (GEngine->IsStereoscopic3D())
 	{
 		static const auto MobileMultiViewCVar = IConsoleManager::Get().FindTConsoleVariableDataInt(TEXT("vr.MobileMultiView"));
-		const bool bSkipPostprocessing = !IsMobileHDR();
 		const bool bUsingMobileRenderer = GetFeatureLevelShadingPath(MyWorld->Scene->GetFeatureLevel()) == EShadingPath::Mobile;
-		bRequireMultiView = (GSupportsMobileMultiView || GRHISupportsArrayIndexFromAnyShader) && bUsingMobileRenderer && bSkipPostprocessing && (MobileMultiViewCVar && MobileMultiViewCVar->GetValueOnAnyThread() != 0);
+		bRequireMultiView = (GSupportsMobileMultiView || GRHISupportsArrayIndexFromAnyShader) && bUsingMobileRenderer && (MobileMultiViewCVar && MobileMultiViewCVar->GetValueOnAnyThread() != 0);
 	}
 
 	FSceneViewFamilyContext ViewFamily(FSceneViewFamily::ConstructionValues(
diff --git a/Engine/Source/Runtime/Engine/Public/SceneTexturesConfig.h b/Engine/Source/Runtime/Engine/Public/SceneTexturesConfig.h
--- a/Engine/Source/Runtime/Engine/Public/SceneTexturesConfig.h
+++ b/Engine/Source/Runtime/Engine/Public/SceneTexturesConfig.h
@@ -123,6 +123,7 @@ struct FSceneTexturesConfig
 		: bRequireMultiView{}
 		, bIsUsingGBuffers{}
 		, bKeepDepthContent{ 1 }
+		, bRequiresDepthAux{}
 		, bPreciseDepthAux{}
 		, bSamplesCustomStencil{}
 		, bMemorylessMSAA{}
@@ -190,6 +191,9 @@ struct FSceneTexturesConfig
 	// (Mobile) True if the platform should write depth content back to memory.
 	uint32 bKeepDepthContent : 1;
 
+	// (Mobile) True if platform requires SceneDepthAux target
+	uint32 bRequiresDepthAux : 1;
+	
 	// (Mobile) True if SceneDepthAux should use a precise pixel format
 	uint32 bPreciseDepthAux : 1;
 
diff --git a/Engine/Source/Runtime/RHI/Private/RHI.cpp b/Engine/Source/Runtime/RHI/Private/RHI.cpp
--- a/Engine/Source/Runtime/RHI/Private/RHI.cpp
+++ b/Engine/Source/Runtime/RHI/Private/RHI.cpp
@@ -1536,7 +1536,7 @@ void FRHIRenderPassInfo::ConvertToRenderTargetsInfo(FRHISetRenderTargetsInfo& Ou
 
 		OutRTInfo.bClearColor |= (LoadAction == ERenderTargetLoadAction::EClear);
 
-		ensure(!OutRTInfo.bHasResolveAttachments || ColorRenderTargets[Index].ResolveTarget);
+		ensure(!OutRTInfo.bHasResolveAttachments || ColorRenderTargets[Index].RenderTarget->GetNumSamples() == 1);
 		if (ColorRenderTargets[Index].ResolveTarget)
 		{
 			OutRTInfo.bHasResolveAttachments = true;
@@ -1583,7 +1583,8 @@ void FRHIRenderPassInfo::Validate() const
 			}
 			else
 			{
-				ensureMsgf(Entry.RenderTarget->GetNumSamples() == NumSamples, TEXT("RenderTarget have inconsistent NumSamples: first %d, then %d"), NumSamples, Entry.RenderTarget->GetNumSamples());
+				// CustomResolveSubpass can have targets with a different NumSamples
+				ensureMsgf(Entry.RenderTarget->GetNumSamples() == NumSamples || SubpassHint == ESubpassHint::CustomResolveSubpass, TEXT("RenderTarget have inconsistent NumSamples: first %d, then %d"), NumSamples, Entry.RenderTarget->GetNumSamples());
 			}
 
 			ERenderTargetStoreAction Store = GetStoreAction(Entry.Action);
@@ -1650,7 +1651,7 @@ void FRHIRenderPassInfo::Validate() const
 			//ensure(StencilStore == ERenderTargetStoreAction::EStore);
 		}
 		
-		if (SubpassHint == ESubpassHint::DepthReadSubpass)
+		if (SubpassHint == ESubpassHint::DepthReadSubpass || SubpassHint == ESubpassHint::CustomResolveSubpass)
 		{
 			// for depth read sub-pass
 			// 1. render pass must have depth target
@@ -1663,6 +1664,7 @@ void FRHIRenderPassInfo::Validate() const
 		ensure(DepthStencilRenderTarget.Action == EDepthStencilTargetActions::DontLoad_DontStore);
 		ensure(DepthStencilRenderTarget.ExclusiveDepthStencil == FExclusiveDepthStencil::DepthNop_StencilNop);
 		ensure(SubpassHint != ESubpassHint::DepthReadSubpass);
+		ensure(SubpassHint != ESubpassHint::CustomResolveSubpass);
 	}
 }
 #endif
diff --git a/Engine/Source/Runtime/RHI/Public/RHIGlobals.h b/Engine/Source/Runtime/RHI/Public/RHIGlobals.h
--- a/Engine/Source/Runtime/RHI/Public/RHIGlobals.h
+++ b/Engine/Source/Runtime/RHI/Public/RHIGlobals.h
@@ -639,6 +639,9 @@ struct FRHIGlobals
 
 	/** True if HDR requires vendor specific extensions */
 	bool HDRNeedsVendorExtensions = false;
+
+	/** True if RHI supports MSAA resolve with a custom shader */
+	bool SupportsMSAAShaderResolve = false;
 };
 
 extern RHI_API FRHIGlobals GRHIGlobals;
@@ -806,6 +809,7 @@ extern RHI_API FRHIGlobals GRHIGlobals;
 #define GShaderPlatformForFeatureLevel                         GRHIGlobals.ShaderPlatformForFeatureLevel
 #define GRHIIsDebugLayerEnabled                                GRHIGlobals.IsDebugLayerEnabled
 #define GRHIHDRNeedsVendorExtensions						   GRHIGlobals.HDRNeedsVendorExtensions
+#define GRHISupportsMSAAShaderResolve						   GRHIGlobals.SupportsMSAAShaderResolve
 
 // Utility Getters
 
diff --git a/Engine/Source/Runtime/RHI/Public/RHIResources.h b/Engine/Source/Runtime/RHI/Public/RHIResources.h
--- a/Engine/Source/Runtime/RHI/Public/RHIResources.h
+++ b/Engine/Source/Runtime/RHI/Public/RHIResources.h
@@ -3680,6 +3680,9 @@ enum class ESubpassHint : uint8
 
 	// Mobile defferred shading subpass
 	DeferredShadingSubpass,
+
+	// Mobile MSAA custom resolve subpass. Includes DepthReadSubpass.
+	CustomResolveSubpass,
 };
 
 enum class EConservativeRasterization : uint8
diff --git a/Engine/Source/Runtime/RenderCore/Private/StereoRenderUtils.cpp b/Engine/Source/Runtime/RenderCore/Private/StereoRenderUtils.cpp
--- a/Engine/Source/Runtime/RenderCore/Private/StereoRenderUtils.cpp
+++ b/Engine/Source/Runtime/RenderCore/Private/StereoRenderUtils.cpp
@@ -78,7 +78,7 @@ RENDERCORE_API FStereoShaderAspects::FStereoShaderAspects(EShaderPlatform Platfo
 	UE_DEBUG_SSA_LOG_BOOL(bInstancedStereoNative);
 	UE_DEBUG_SSA_LOG(Log, TEXT("---"));
 
-	const bool bMobileMultiViewCoreSupport = bMobilePlatform && bMobileMultiView && !bMobilePostprocessing;
+	const bool bMobileMultiViewCoreSupport = bMobilePlatform && bMobileMultiView; 
 	if (bMobileMultiViewCoreSupport)
 	{
 		UE_DEBUG_SSA_LOG(Log, TEXT("RHISupportsMobileMultiView(%s) = %d."), *LexToString(Platform), RHISupportsMobileMultiView(Platform));
diff --git a/Engine/Source/Runtime/Renderer/Private/MobileShadingRenderer.cpp b/Engine/Source/Runtime/Renderer/Private/MobileShadingRenderer.cpp
--- a/Engine/Source/Runtime/Renderer/Private/MobileShadingRenderer.cpp
+++ b/Engine/Source/Runtime/Renderer/Private/MobileShadingRenderer.cpp
@@ -33,6 +33,8 @@
 #include "PostProcess/PostProcessHMD.h"
 #include "PostProcess/PostProcessPixelProjectedReflectionMobile.h"
 #include "PostProcess/PostProcessAmbientOcclusionMobile.h"
+#include "PostProcess/PostProcessCombineLUTs.h"
+#include "PostProcess/PostProcessTonemap.h"
 #include "IHeadMountedDisplay.h"
 #include "IXRTrackingSystem.h"
 #include "SceneViewExtension.h"
@@ -97,6 +99,20 @@ static TAutoConsoleVariable<int32> CVarMobileCustomDepthForTranslucency(
 	TEXT(" 1 = On [default]"),
 	ECVF_Scalability | ECVF_RenderThreadSafe);
 
+static TAutoConsoleVariable<int32> CVarMobileTonemapSubpass(
+	TEXT("r.Mobile.TonemapSubpass"),
+	0,
+	TEXT(" Whether to enable mobile tonemap subpass \n")
+	TEXT(" 0 = Off [default]\n")
+	TEXT(" 1 = On"),
+	ECVF_Scalability | ECVF_RenderThreadSafe);
+
+static bool IsMobileTonemapSubpassEnabled(const FStaticShaderPlatform Platform)
+{
+	static auto* MobileTonemapSubpassPathCvar = IConsoleManager::Get().FindTConsoleVariableDataInt(TEXT("r.Mobile.TonemapSubpass"));
+	return (MobileTonemapSubpassPathCvar && (MobileTonemapSubpassPathCvar->GetValueOnAnyThread() == 1)) && IsMobileHDR() && !IsMobileDeferredShadingEnabled(Platform);
+}
+
 DECLARE_GPU_STAT_NAMED(MobileSceneRender, TEXT("Mobile Scene Render"));
 
 DECLARE_CYCLE_STAT(TEXT("SceneStart"), STAT_CLMM_SceneStart, STATGROUP_CommandListMarkers);
@@ -200,6 +216,7 @@ BEGIN_SHADER_PARAMETER_STRUCT(FMobileRenderPassParameters, RENDERER_API)
 	SHADER_PARAMETER_RDG_BUFFER_SRV(Buffer<uint>, LocalFogVolumeTileDataBuffer)
 	SHADER_PARAMETER_RDG_TEXTURE_SRV(Texture2D<float4>, HalfResLocalFogVolumeViewSRV)
 	SHADER_PARAMETER_RDG_TEXTURE_SRV(Texture2D<float>, HalfResLocalFogVolumeDepthSRV)
+	RDG_TEXTURE_ACCESS(ColorGradingLUT, ERHIAccess::SRVGraphics)
 	RENDER_TARGET_BINDING_SLOTS()
 END_SHADER_PARAMETER_STRUCT()
 
@@ -284,7 +301,6 @@ FMobileSceneRenderer::FMobileSceneRenderer(const FSceneViewFamily* InViewFamily,
 	bRequiresShadowProjections = false;
 	bIsFullDepthPrepassEnabled = Scene->EarlyZPassMode == DDM_AllOpaque;
 	bIsMaskedOnlyDepthPrepassEnabled = Scene->EarlyZPassMode == DDM_MaskedOnly;
-	bRequiresSceneDepthAux = MobileRequiresSceneDepthAux(ShaderPlatform);
 	bEnableClusteredLocalLights = MobileForwardEnableLocalLights(ShaderPlatform);
 	bEnableClusteredReflections = MobileForwardEnableClusteredReflections(ShaderPlatform);
 	
@@ -302,6 +318,10 @@ FMobileSceneRenderer::FMobileSceneRenderer(const FSceneViewFamily* InViewFamily,
 	}
 
 	NumMSAASamples = GetDefaultMSAACount(ERHIFeatureLevel::ES3_1);
+	// As of UE 5.4 only vulkan supports inline (single pass) tonemap
+	bTonemapSubpass = IsMobileTonemapSubpassEnabled(ShaderPlatform) && ViewFamily.bResolveScene && GetRendererOutput() == FSceneRenderer::ERendererOutput::FinalSceneColor;
+	bTonemapSubpassInline = bTonemapSubpass && IsVulkanPlatform(ShaderPlatform) && (GRHISupportsMSAAShaderResolve || NumMSAASamples == 1);
+	bRequiresSceneDepthAux = MobileRequiresSceneDepthAux(ShaderPlatform) && !bTonemapSubpass;
 }
 
 class FMobileDirLightShaderParamsRenderResource : public FRenderResource
@@ -585,6 +605,7 @@ void FMobileSceneRenderer::InitViews(
 	// If we render in a single pass MSAA targets can be memoryless
     SceneTexturesConfig.bMemorylessMSAA = !(bRequiresMultiPass || bShouldCompositeEditorPrimitives || bRequireSeparateViewPass);
     SceneTexturesConfig.NumSamples = NumMSAASamples;
+	SceneTexturesConfig.ExtraSceneColorCreateFlags |= (bTonemapSubpassInline ? TexCreate_InputAttachmentRead : TexCreate_None);
     SceneTexturesConfig.BuildSceneColorAndDepthFlags();
 	if (bDeferredShading) 
 	{
@@ -601,6 +622,7 @@ void FMobileSceneRenderer::InitViews(
 		SceneTexturesConfig.MobilePixelProjectedReflectionExtent = FIntPoint::ZeroValue;
 	}
 
+	SceneTexturesConfig.bRequiresDepthAux = bRequiresSceneDepthAux;
 	// When we capturing scene depth, use a more precise format for SceneDepthAux as it will be used as a source DepthTexture
 	if (bSceneDepthCapture)
 	{
@@ -1316,7 +1338,7 @@ void FMobileSceneRenderer::Render(FRDGBuilder& GraphBuilder)
 	
 		if (ViewFamily.bResolveScene)
 		{
-			if (bRenderToSceneColor)
+			if (bRenderToSceneColor && !bTonemapSubpassInline)
 			{
 				// Finish rendering for each view, or the full stereo buffer if enabled
 				{
@@ -1331,8 +1353,18 @@ void FMobileSceneRenderer::Render(FRDGBuilder& GraphBuilder)
 
 					for (int32 ViewIndex = 0; ViewIndex < Views.Num(); ViewIndex++)
 					{
-						RDG_EVENT_SCOPE_CONDITIONAL(GraphBuilder, Views.Num() > 1, "View%d", ViewIndex);
-						AddMobilePostProcessingPasses(GraphBuilder, Scene, Views[ViewIndex], GetSceneUniforms(), PostProcessingInputs, InstanceCullingManager);
+						if (Views[ViewIndex].ShouldRenderView())
+						{
+							RDG_EVENT_SCOPE_CONDITIONAL(GraphBuilder, Views.Num() > 1, "View%d", ViewIndex);
+							if (bTonemapSubpass)
+							{
+								AddMobileCustomResolvePass(GraphBuilder, Views[ViewIndex], SceneTextures, ViewFamilyTexture);
+							}
+							else
+							{ 
+								AddMobilePostProcessingPasses(GraphBuilder, Scene, Views[ViewIndex], GetSceneUniforms(), PostProcessingInputs, InstanceCullingManager);
+							}
+						}
 					}
 				}
 			}
@@ -1411,6 +1443,15 @@ FRenderTargetBindingSlots FMobileSceneRenderer::InitRenderTargetBindings_Forward
 	{
 		BasePassRenderTargets[1] = FRenderTargetBinding(SceneTextures.DepthAux.Target, SceneTextures.DepthAux.Resolve, ERenderTargetLoadAction::EClear);
 	}
+		
+	if (bTonemapSubpassInline)
+	{
+		// DepthAux is not used with tonemap subpass, since there are no post-processing passes
+		// Backbuffer surface provided as a second render target instead of resolve target.
+		BasePassRenderTargets[0].SetResolveTexture(nullptr);
+		BasePassRenderTargets[1] = FRenderTargetBinding(ViewFamilyTexture, nullptr, ERenderTargetLoadAction::EClear);
+	}
+	
 	BasePassRenderTargets.DepthStencil = bIsFullDepthPrepassEnabled ? 
 		FDepthStencilBinding(SceneDepth, ERenderTargetLoadAction::ELoad, ERenderTargetLoadAction::ELoad, FExclusiveDepthStencil::DepthRead_StencilWrite) : 
 		FDepthStencilBinding(SceneDepth, ERenderTargetLoadAction::EClear, ERenderTargetLoadAction::EClear, FExclusiveDepthStencil::DepthWrite_StencilWrite);
@@ -1497,8 +1538,13 @@ void FMobileSceneRenderer::RenderForward(FRDGBuilder& GraphBuilder, FRDGTextureR
 
 void FMobileSceneRenderer::RenderForwardSinglePass(FRDGBuilder& GraphBuilder, FMobileRenderPassParameters* PassParameters, FRenderViewContext& ViewContext, FSceneTextures& SceneTextures)
 {
-	PassParameters->RenderTargets.SubpassHint = ESubpassHint::DepthReadSubpass;
-	
+	if (bTonemapSubpassInline)
+	{
+		// tonemapping LUT pass before we start main render pass. The texture is needed by the custom resolve pass which does tonemapping
+		PassParameters->ColorGradingLUT = AddCombineLUTPass(GraphBuilder, *ViewContext.ViewInfo);
+	}
+		
+	PassParameters->RenderTargets.SubpassHint = bTonemapSubpassInline ? ESubpassHint::CustomResolveSubpass : ESubpassHint::DepthReadSubpass;
 	const bool bDoOcclusionQueries = (!bIsFullDepthPrepassEnabled && ViewContext.bIsLastView && DoOcclusionQueries());
 	PassParameters->RenderTargets.NumOcclusionQueries = bDoOcclusionQueries ? ComputeNumOcclusionQueriesToBatch() : 0u;
 	
@@ -1552,6 +1598,11 @@ void FMobileSceneRenderer::RenderForwardSinglePass(FRDGBuilder& GraphBuilder, FM
 
 		// Pre-tonemap before MSAA resolve (iOS only)
 		PreTonemapMSAA(RHICmdList, SceneTextures);
+		if (bTonemapSubpassInline)
+		{
+			RHICmdList.NextSubpass();
+			RenderMobileCustomResolve(RHICmdList, View, NumMSAASamples, SceneTextures);
+		}
 	});
 	
 	// resolve MSAA depth
diff --git a/Engine/Source/Runtime/Renderer/Private/PostProcess/PostProcessTonemap.cpp b/Engine/Source/Runtime/Renderer/Private/PostProcess/PostProcessTonemap.cpp
--- a/Engine/Source/Runtime/Renderer/Private/PostProcess/PostProcessTonemap.cpp
+++ b/Engine/Source/Runtime/Renderer/Private/PostProcess/PostProcessTonemap.cpp
@@ -1035,3 +1035,127 @@ FScreenPassTexture AddTonemapPass(FRDGBuilder& GraphBuilder, const FViewInfo& Vi
 	return MoveTemp(Output);
 }
 
+
+// MSAA custom resolve shader that does tonemapping 
+class FMobileCustomResolvePS : public FGlobalShader
+{
+	DECLARE_GLOBAL_SHADER(FMobileCustomResolvePS);
+
+	SHADER_USE_PARAMETER_STRUCT(FMobileCustomResolvePS, FGlobalShader);
+	BEGIN_SHADER_PARAMETER_STRUCT(FParameters, )
+		SHADER_PARAMETER_STRUCT_INCLUDE(FViewShaderParameters, View)
+		SHADER_PARAMETER_RDG_TEXTURE(Texture2D, ColorTexture)
+		SHADER_PARAMETER_SAMPLER(SamplerState, ColorSampler)
+		SHADER_PARAMETER(FVector4f, ColorScale0)
+		SHADER_PARAMETER_TEXTURE(Texture2D, ColorGradingLUT)
+		SHADER_PARAMETER_SAMPLER(SamplerState, ColorGradingLUTSampler)
+		SHADER_PARAMETER(float, LUTSize)
+		SHADER_PARAMETER(float, InvLUTSize)
+		SHADER_PARAMETER(float, LUTScale)
+		SHADER_PARAMETER(float, LUTOffset)
+	END_SHADER_PARAMETER_STRUCT()
+
+	class FTonemapperSubpassMsaaDim : SHADER_PERMUTATION_SPARSE_INT("SUBPASS_MSAA_SAMPLES", 0, 1, 2, 4, 8);
+	using FPermutationDomain = TShaderPermutationDomain<FTonemapperSubpassMsaaDim>;
+
+	static bool ShouldCompilePermutation(const FGlobalShaderPermutationParameters& Parameters)
+	{
+		return IsMobilePlatform(Parameters.Platform);
+	}
+
+	static void ModifyCompilationEnvironment(const FGlobalShaderPermutationParameters& Parameters, FShaderCompilerEnvironment& OutEnvironment)
+	{
+		const int UseVolumeLut = PipelineVolumeTextureLUTSupportGuaranteedAtRuntime(Parameters.Platform) ? 1 : 0;
+		OutEnvironment.SetDefine(TEXT("USE_VOLUME_LUT"), UseVolumeLut);
+	}
+};
+
+IMPLEMENT_GLOBAL_SHADER(FMobileCustomResolvePS, "/Engine/Private/PostProcessTonemap.usf", "MobileCustomResolve_MainPS", SF_Pixel);
+
+void RenderMobileCustomResolve(FRHICommandList& RHICmdList, const FViewInfo& View, const int32 SubpassMSAASamples, FSceneTextures& SceneTextures)
+{
+	// Part of scene rendering pass
+	check(RHICmdList.IsInsideRenderPass());
+	SCOPED_DRAW_EVENT(RHICmdList, MobileTonemapSubpass);
+
+	IPooledRenderTarget* ColorGradingLUT = View.GetTonemappingLUT();
+	const FIntPoint TargetSize = SceneTextures.Color.Target->Desc.Extent;
+	const float LUTSize = ColorGradingLUT ? (float)ColorGradingLUT->GetDesc().GetSize().Y : /* unused (default): */ 32.0f;
+	const FPostProcessSettings& Settings = View.FinalPostProcessSettings;
+	
+	TShaderMapRef<FScreenVS> VertexShader(View.ShaderMap);
+	
+	FMobileCustomResolvePS::FPermutationDomain PermutationVector;
+	PermutationVector.Set<FMobileCustomResolvePS::FTonemapperSubpassMsaaDim>(SubpassMSAASamples);
+	TShaderMapRef<FMobileCustomResolvePS> PixelShader(View.ShaderMap, PermutationVector);
+
+	FMobileCustomResolvePS::FParameters PSShaderParameters;
+	PSShaderParameters.View = View.GetShaderParameters();
+	PSShaderParameters.ColorScale0 = FVector4f(Settings.SceneColorTint.R, Settings.SceneColorTint.G, Settings.SceneColorTint.B, 0);
+	PSShaderParameters.ColorGradingLUT = ColorGradingLUT ? ColorGradingLUT->GetRHI() : GBlackTexture->TextureRHI.GetReference();
+	PSShaderParameters.ColorGradingLUTSampler = TStaticSamplerState<SF_Bilinear, AM_Clamp, AM_Clamp, AM_Clamp>::GetRHI();
+	PSShaderParameters.LUTSize = LUTSize;
+	PSShaderParameters.InvLUTSize = 1.0f / LUTSize;
+	PSShaderParameters.LUTScale = (LUTSize - 1.0f) / LUTSize;
+	PSShaderParameters.LUTOffset = 0.5f / LUTSize;
+	if (SubpassMSAASamples == 0u)
+	{
+		PSShaderParameters.ColorTexture = SceneTextures.Color.Resolve;
+		PSShaderParameters.ColorSampler = TStaticSamplerState<SF_Point, AM_Clamp, AM_Clamp, AM_Clamp>::GetRHI();
+	}
+
+	FGraphicsPipelineStateInitializer GraphicsPSOInit;
+	RHICmdList.ApplyCachedRenderTargets(GraphicsPSOInit);
+
+	GraphicsPSOInit.BlendState = TStaticBlendState<>::GetRHI();
+	GraphicsPSOInit.RasterizerState = TStaticRasterizerState<>::GetRHI();
+	GraphicsPSOInit.DepthStencilState = TStaticDepthStencilState<false, CF_Always>::GetRHI();
+
+	GraphicsPSOInit.BoundShaderState.VertexDeclarationRHI = GFilterVertexDeclaration.VertexDeclarationRHI;
+	GraphicsPSOInit.BoundShaderState.VertexShaderRHI = VertexShader.GetVertexShader();
+	GraphicsPSOInit.BoundShaderState.PixelShaderRHI = PixelShader.GetPixelShader();
+	GraphicsPSOInit.PrimitiveType = PT_TriangleList;
+
+	SetGraphicsPipelineState(RHICmdList, GraphicsPSOInit, 0);
+	SetShaderParameters(RHICmdList, PixelShader, PixelShader.GetPixelShader(), PSShaderParameters);
+	RHICmdList.SetViewport(0, 0, 0.0f, TargetSize.X, TargetSize.Y, 1.0f);
+
+	DrawRectangle(
+		RHICmdList,
+		0, 0,
+		TargetSize.X, TargetSize.Y,
+		0, 0,
+		TargetSize.X, TargetSize.Y,
+		TargetSize,
+		TargetSize,
+		VertexShader,
+		EDRF_UseTriangleOptimization);
+}
+
+BEGIN_SHADER_PARAMETER_STRUCT(FMobileCustomResolveParameters, )
+	SHADER_PARAMETER_STRUCT_INCLUDE(FViewShaderParameters, View)
+	SHADER_PARAMETER_RDG_TEXTURE(Texture2D, ColorTexture)
+	RDG_TEXTURE_ACCESS(ColorGradingLUT, ERHIAccess::SRVGraphics)
+	RENDER_TARGET_BINDING_SLOTS()
+END_SHADER_PARAMETER_STRUCT()
+
+void AddMobileCustomResolvePass(FRDGBuilder& GraphBuilder, const FViewInfo& View, FSceneTextures& SceneTextures, FRDGTextureRef ViewFamilyTexture)
+{
+	FRDGTextureRef ColorGradingLUT = AddCombineLUTPass(GraphBuilder, View);
+
+	FMobileCustomResolveParameters* PassParameters = GraphBuilder.AllocParameters<FMobileCustomResolveParameters>();
+	PassParameters->View = View.GetShaderParameters();
+	PassParameters->ColorTexture = SceneTextures.Color.Resolve;
+	PassParameters->ColorGradingLUT = ColorGradingLUT;
+	PassParameters->RenderTargets[0] = FRenderTargetBinding(ViewFamilyTexture, ERenderTargetLoadAction::EClear);
+
+	GraphBuilder.AddPass(
+		RDG_EVENT_NAME("MobileCustomResolvePass"),
+		PassParameters,
+		ERDGPassFlags::Raster,
+		[&View, &SceneTextures](FRHICommandListImmediate& RHICmdList)
+		{
+			const uint32 SubpassMSAASamples = 0u; // not using subpass resolve
+			RenderMobileCustomResolve(RHICmdList, View, SubpassMSAASamples, SceneTextures);
+		});
+}
diff --git a/Engine/Source/Runtime/Renderer/Private/PostProcess/PostProcessTonemap.h b/Engine/Source/Runtime/Renderer/Private/PostProcess/PostProcessTonemap.h
--- a/Engine/Source/Runtime/Renderer/Private/PostProcess/PostProcessTonemap.h
+++ b/Engine/Source/Runtime/Renderer/Private/PostProcess/PostProcessTonemap.h
@@ -75,3 +75,5 @@ struct FTonemapInputs
 };
 
 FScreenPassTexture AddTonemapPass(FRDGBuilder& GraphBuilder, const FViewInfo& View, const FTonemapInputs& Inputs);
+void RenderMobileCustomResolve(FRHICommandList& RHICmdList, const FViewInfo& View, const int32 SubpassMSAASamples, FSceneTextures& SceneTextures);
+void AddMobileCustomResolvePass(FRDGBuilder& GraphBuilder, const FViewInfo& View, FSceneTextures& SceneTextures, FRDGTextureRef ViewFamilyTexture);
diff --git a/Engine/Source/Runtime/Renderer/Private/SceneRendering.h b/Engine/Source/Runtime/Renderer/Private/SceneRendering.h
--- a/Engine/Source/Runtime/Renderer/Private/SceneRendering.h
+++ b/Engine/Source/Runtime/Renderer/Private/SceneRendering.h
@@ -2745,6 +2745,8 @@ private:
 	const bool bDeferredShading;
 	const bool bRequiresDBufferDecals;
 	const bool bUseVirtualTexturing;
+	bool bTonemapSubpass;
+	bool bTonemapSubpassInline;
 	int32 NumMSAASamples;
 	bool bRenderToSceneColor;
 	bool bRequiresMultiPass;
diff --git a/Engine/Source/Runtime/Renderer/Private/SceneTextures.cpp b/Engine/Source/Runtime/Renderer/Private/SceneTextures.cpp
--- a/Engine/Source/Runtime/Renderer/Private/SceneTextures.cpp
+++ b/Engine/Source/Runtime/Renderer/Private/SceneTextures.cpp
@@ -587,8 +587,7 @@ void FSceneTextures::InitializeViewFamily(FRDGBuilder& GraphBuilder, FViewFamily
 		}
 	}
 
-
-	if (Config.ShadingPath == EShadingPath::Mobile && MobileRequiresSceneDepthAux(Config.ShaderPlatform))
+	if (Config.bRequiresDepthAux)
 	{
 		const float FarDepth = (float)ERHIZBuffer::FarPlane;
 		const FLinearColor FarDepthColor(FarDepth, FarDepth, FarDepth, FarDepth);
diff --git a/Engine/Source/Runtime/VulkanRHI/Private/VulkanDevice.h b/Engine/Source/Runtime/VulkanRHI/Private/VulkanDevice.h
--- a/Engine/Source/Runtime/VulkanRHI/Private/VulkanDevice.h
+++ b/Engine/Source/Runtime/VulkanRHI/Private/VulkanDevice.h
@@ -77,6 +77,7 @@ struct FOptionalVulkanDeviceExtensions
 			uint64 HasEXTShaderViewportIndexLayer : 1;
 			uint64 HasSeparateDepthStencilLayouts : 1;
 			uint64 HasEXTHostQueryReset : 1;
+			uint64 HasQcomRenderPassShaderResolve : 1;
 
 			// Promoted to 1.3
 			uint64 HasEXTTextureCompressionASTCHDR : 1;
diff --git a/Engine/Source/Runtime/VulkanRHI/Private/VulkanExtensions.cpp b/Engine/Source/Runtime/VulkanRHI/Private/VulkanExtensions.cpp
--- a/Engine/Source/Runtime/VulkanRHI/Private/VulkanExtensions.cpp
+++ b/Engine/Source/Runtime/VulkanRHI/Private/VulkanExtensions.cpp
@@ -1353,7 +1353,7 @@ FVulkanDeviceExtensionArray FVulkanDeviceExtension::GetUESupportedDeviceExtensio
 	ADD_SIMPLE_EXTENSION(VK_KHR_SHADER_FLOAT_CONTROLS_EXTENSION_NAME,    VULKAN_RHI_RAYTRACING,                VK_API_VERSION_1_2,            DEVICE_EXT_FLAG_SETTER(HasShaderFloatControls));
 	ADD_SIMPLE_EXTENSION(VK_KHR_IMAGE_FORMAT_LIST_EXTENSION_NAME,        VULKAN_EXTENSION_ENABLED,             VK_API_VERSION_1_2,            DEVICE_EXT_FLAG_SETTER(HasKHRImageFormatList));
 	ADD_SIMPLE_EXTENSION(VK_EXT_VALIDATION_CACHE_EXTENSION_NAME,         VULKAN_SUPPORTS_VALIDATION_CACHE,     VULKAN_EXTENSION_NOT_PROMOTED, DEVICE_EXT_FLAG_SETTER(HasEXTValidationCache));
-
+	ADD_SIMPLE_EXTENSION(VK_QCOM_RENDER_PASS_SHADER_RESOLVE_EXTENSION_NAME, VULKAN_SUPPORTS_QCOM_RENDERPASS_SHADER_RESOLVE, VULKAN_EXTENSION_NOT_PROMOTED, DEVICE_EXT_FLAG_SETTER(HasQcomRenderPassShaderResolve));
 
 	// Externally activated extensions (supported by the engine, but enabled externally by plugin or other) :
 
diff --git a/Engine/Source/Runtime/VulkanRHI/Private/VulkanPendingState.cpp b/Engine/Source/Runtime/VulkanRHI/Private/VulkanPendingState.cpp
--- a/Engine/Source/Runtime/VulkanRHI/Private/VulkanPendingState.cpp
+++ b/Engine/Source/Runtime/VulkanRHI/Private/VulkanPendingState.cpp
@@ -549,8 +549,6 @@ void FVulkanPendingGfxState::UpdateInputAttachments(FVulkanFramebuffer* Framebuf
 		switch (AttachmentData.Type)
 		{
 		case FVulkanShaderHeader::EAttachmentType::Color0:
-			CurrentState->SetInputAttachment(AttachmentData.DescriptorSet, AttachmentData.BindingIndex, Framebuffer->AttachmentTextureViews[0]->GetTextureView(), VK_IMAGE_LAYOUT_GENERAL);
-			break;
 		case FVulkanShaderHeader::EAttachmentType::Color1:
 		case FVulkanShaderHeader::EAttachmentType::Color2:
 		case FVulkanShaderHeader::EAttachmentType::Color3:
diff --git a/Engine/Source/Runtime/VulkanRHI/Private/VulkanPipeline.cpp b/Engine/Source/Runtime/VulkanRHI/Private/VulkanPipeline.cpp
--- a/Engine/Source/Runtime/VulkanRHI/Private/VulkanPipeline.cpp
+++ b/Engine/Source/Runtime/VulkanRHI/Private/VulkanPipeline.cpp
@@ -1822,6 +1822,16 @@ void FVulkanPipelineStateCacheManager::CreateGfxEntry(const FGraphicsPipelineSta
 		NumRenderTargets = 1;
 	}
 
+	if (PSOInitializer.SubpassHint == ESubpassHint::CustomResolveSubpass)
+	{
+		NumRenderTargets = 1; // This applies to base and depth passes as well. One render target for base and depth, another one for custom resolve.
+		if (PSOInitializer.SubpassIndex >= 2)
+		{ 
+			// the resolve subpass renders to a non MSAA surface
+			OutGfxEntry->RasterizationSamples = 1;
+		}
+	}
+
 	OutGfxEntry->ColorAttachmentStates.AddUninitialized(NumRenderTargets);
 	for (int32 Index = 0; Index < OutGfxEntry->ColorAttachmentStates.Num(); ++Index)
 	{
diff --git a/Engine/Source/Runtime/VulkanRHI/Private/VulkanRHI.cpp b/Engine/Source/Runtime/VulkanRHI/Private/VulkanRHI.cpp
--- a/Engine/Source/Runtime/VulkanRHI/Private/VulkanRHI.cpp
+++ b/Engine/Source/Runtime/VulkanRHI/Private/VulkanRHI.cpp
@@ -929,6 +929,7 @@ void FVulkanDynamicRHI::InitInstance()
 		GRHISupportsTextureStreaming = true;
 		GSupportsTimestampRenderQueries = FVulkanPlatform::SupportsTimestampRenderQueries();
 		GSupportsMobileMultiView = Device->GetOptionalExtensions().HasKHRMultiview ? true : false;
+		GRHISupportsMSAAShaderResolve = Device->GetOptionalExtensions().HasQcomRenderPassShaderResolve ? true : false;
 #if VULKAN_RHI_RAYTRACING
 		GRHISupportsRayTracing = RHISupportsRayTracing(GMaxRHIShaderPlatform) && Device->GetOptionalExtensions().HasRaytracingExtensions();
 
diff --git a/Engine/Source/Runtime/VulkanRHI/Private/VulkanRenderTarget.cpp b/Engine/Source/Runtime/VulkanRHI/Private/VulkanRenderTarget.cpp
--- a/Engine/Source/Runtime/VulkanRHI/Private/VulkanRenderTarget.cpp
+++ b/Engine/Source/Runtime/VulkanRHI/Private/VulkanRenderTarget.cpp
@@ -1038,14 +1038,17 @@ FVulkanRenderTargetLayout::FVulkanRenderTargetLayout(FVulkanDevice& InDevice, co
 			Extent.Extent3D.depth = TextureDesc.Depth;
 		}
 
-		ensure(!NumSamples || NumSamples == ColorEntry.RenderTarget->GetNumSamples());
+		// CustomResolveSubpass can have targets with a different NumSamples
+		ensure(!NumSamples || NumSamples == ColorEntry.RenderTarget->GetNumSamples() || RPInfo.SubpassHint == ESubpassHint::CustomResolveSubpass);
 		NumSamples = ColorEntry.RenderTarget->GetNumSamples();
 
 		ensure(!GetIsMultiView() || !bMultiviewRenderTargets || Texture->GetNumberOfArrayLevels() > 1);
 		bMultiviewRenderTargets = Texture->GetNumberOfArrayLevels() > 1;
+		// With a CustomResolveSubpass last color attachment is a resolve target
+		bool bCustomResolveAttachment = (Index == (NumColorRenderTargets - 1)) && RPInfo.SubpassHint == ESubpassHint::CustomResolveSubpass;
 
 		VkAttachmentDescription& CurrDesc = Desc[NumAttachmentDescriptions];
-		CurrDesc.samples = static_cast<VkSampleCountFlagBits>(NumSamples);
+		CurrDesc.samples = bCustomResolveAttachment ? VK_SAMPLE_COUNT_1_BIT : static_cast<VkSampleCountFlagBits>(NumSamples);
 		CurrDesc.format = UEToVkTextureFormat(ColorEntry.RenderTarget->GetFormat(), EnumHasAllFlags(Texture->GetDesc().Flags, TexCreate_SRGB));
 		CurrDesc.loadOp = RenderTargetLoadActionToVulkan(GetLoadAction(ColorEntry.Action));
 		bFoundClearOp = bFoundClearOp || (CurrDesc.loadOp == VK_ATTACHMENT_LOAD_OP_CLEAR);
@@ -1097,7 +1100,8 @@ FVulkanRenderTargetLayout::FVulkanRenderTargetLayout(FVulkanDevice& InDevice, co
 		const FRHITextureDesc& TextureDesc = Texture->GetDesc();
 
 		CurrDesc.samples = static_cast<VkSampleCountFlagBits>(RPInfo.DepthStencilRenderTarget.DepthStencilTarget->GetNumSamples());
-		ensure(!NumSamples || CurrDesc.samples == NumSamples);
+		// CustomResolveSubpass can have targets with a different NumSamples
+		ensure(!NumSamples || CurrDesc.samples == NumSamples || RPInfo.SubpassHint == ESubpassHint::CustomResolveSubpass);
 		NumSamples = CurrDesc.samples;
 		CurrDesc.format = UEToVkTextureFormat(RPInfo.DepthStencilRenderTarget.DepthStencilTarget->GetFormat(), false);
 		CurrDesc.loadOp = RenderTargetLoadActionToVulkan(GetLoadAction(GetDepthActions(RPInfo.DepthStencilRenderTarget.Action)));
@@ -1266,8 +1270,11 @@ FVulkanRenderTargetLayout::FVulkanRenderTargetLayout(const FGraphicsPipelineStat
 		EPixelFormat UEFormat = (EPixelFormat)Initializer.RenderTargetFormats[Index];
 		if (UEFormat != PF_Unknown)
 		{
+			// With a CustomResolveSubpass last color attachment is a resolve target
+			bool bCustomResolveAttachment = (Index == (Initializer.RenderTargetsEnabled - 1)) && Initializer.SubpassHint == ESubpassHint::CustomResolveSubpass;
+			
 			VkAttachmentDescription& CurrDesc = Desc[NumAttachmentDescriptions];
-			CurrDesc.samples = static_cast<VkSampleCountFlagBits>(NumSamples);
+			CurrDesc.samples = bCustomResolveAttachment ? VK_SAMPLE_COUNT_1_BIT : static_cast<VkSampleCountFlagBits>(NumSamples);
 			CurrDesc.format = UEToVkTextureFormat(UEFormat, EnumHasAllFlags(Initializer.RenderTargetFlags[Index], TexCreate_SRGB));
 			CurrDesc.loadOp = VK_ATTACHMENT_LOAD_OP_DONT_CARE;
 			CurrDesc.storeOp = VK_ATTACHMENT_STORE_OP_DONT_CARE;
diff --git a/Engine/Source/Runtime/VulkanRHI/Private/VulkanRenderpass.h b/Engine/Source/Runtime/VulkanRHI/Private/VulkanRenderpass.h
--- a/Engine/Source/Runtime/VulkanRHI/Private/VulkanRenderpass.h
+++ b/Engine/Source/Runtime/VulkanRHI/Private/VulkanRenderpass.h
@@ -424,13 +424,14 @@ public:
 		uint32 MultiviewMask = (0b1 << RTLayout.GetMultiViewCount()) - 1;
 
 		const bool bDeferredShadingSubpass = RTLayout.GetSubpassHint() == ESubpassHint::DeferredShadingSubpass;
-		const bool bDepthReadSubpass = RTLayout.GetSubpassHint() == ESubpassHint::DepthReadSubpass;
 		const bool bApplyFragmentShadingRate =  GRHISupportsAttachmentVariableRateShading 
 												&& GRHIVariableRateShadingEnabled 
 												&& GRHIAttachmentVariableRateShadingEnabled 
 												&& RTLayout.GetFragmentDensityAttachmentReference() != nullptr
 												&& Device.GetOptionalExtensions().HasKHRFragmentShadingRate 
 												&& Device.GetOptionalExtensionProperties().FragmentShadingRateFeatures.attachmentFragmentShadingRate == VK_TRUE;
+		const bool bCustomResolveSubpass = RTLayout.GetSubpassHint() == ESubpassHint::CustomResolveSubpass;
+		const bool bDepthReadSubpass = bCustomResolveSubpass || (RTLayout.GetSubpassHint() == ESubpassHint::DepthReadSubpass);
 		const bool bHasDepthStencilAttachmentReference = (RTLayout.GetDepthAttachmentReference() != nullptr);
 
 		if (bApplyFragmentShadingRate)
@@ -440,7 +441,8 @@ public:
 		}
 
 		// Grab (and optionally convert) attachment references.
-		for (uint32 ColorAttachment = 0; ColorAttachment < RTLayout.GetNumColorAttachments(); ++ColorAttachment)
+		uint32 NumColorAttachments = RTLayout.GetNumColorAttachments();
+		for (uint32 ColorAttachment = 0; ColorAttachment < NumColorAttachments; ++ColorAttachment)
 		{
 			ColorAttachmentReferences.Add(TAttachmentReferenceClass(RTLayout.GetColorAttachmentReferences()[ColorAttachment], 0));
 			if (RTLayout.GetResolveAttachmentReferences() != nullptr)
@@ -449,6 +451,12 @@ public:
 			}
 		}
 
+		// CustomResolveSubpass has an additional color attachment that should not be used by main and depth subpasses
+		if (bCustomResolveSubpass && (NumColorAttachments > 1))
+		{
+			NumColorAttachments--;
+		}
+
 		uint32_t DepthInputAttachment = VK_ATTACHMENT_UNUSED;
 		VkImageLayout DepthInputAttachmentLayout = VK_IMAGE_LAYOUT_UNDEFINED;
 		VkImageAspectFlags DepthInputAspectMask = 0;
@@ -483,12 +491,8 @@ public:
 		{
 			TSubpassDescriptionClass& SubpassDesc = SubpassDescriptions[NumSubpasses++];
 
-			SubpassDesc.SetColorAttachments(ColorAttachmentReferences);
-			if (!bDepthReadSubpass)
-			{
-				// only set resolve attachment on the last subpass
-				SubpassDesc.SetResolveAttachments(ResolveAttachmentReferences);
-			}
+			SubpassDesc.SetColorAttachments(ColorAttachmentReferences, NumColorAttachments);
+
 			if (bHasDepthStencilAttachmentReference)
 			{
 				SubpassDesc.SetDepthStencilAttachment(&DepthStencilAttachmentReference);
@@ -506,8 +510,7 @@ public:
 		{
 			TSubpassDescriptionClass& SubpassDesc = SubpassDescriptions[NumSubpasses++];
 
-			SubpassDesc.SetColorAttachments(ColorAttachmentReferences);
-			SubpassDesc.SetResolveAttachments(ResolveAttachmentReferences);
+			SubpassDesc.SetColorAttachments(ColorAttachmentReferences, NumColorAttachments);
 
 			check(RTLayout.GetDepthAttachmentReference());
 
@@ -617,6 +620,45 @@ public:
 			}
 		}
 
+		// Custom resolve subpass
+		if (bCustomResolveSubpass)
+		{
+			TSubpassDescriptionClass& SubpassDesc = SubpassDescriptions[NumSubpasses++];
+			ColorAttachments3[0].attachment = ColorAttachmentReferences[1].attachment;
+			ColorAttachments3[0].layout = VK_IMAGE_LAYOUT_COLOR_ATTACHMENT_OPTIMAL;
+			ColorAttachments3[0].SetAspect(VK_IMAGE_ASPECT_COLOR_BIT);
+			
+			InputAttachments3[0].attachment = VK_ATTACHMENT_UNUSED; // The subpass fetch logic expects depth in first attachment.
+			InputAttachments3[0].layout = VK_IMAGE_LAYOUT_UNDEFINED;
+			InputAttachments3[0].SetAspect(0);
+			
+			InputAttachments3[1].attachment = ColorAttachmentReferences[0].attachment; // SceneColor as input
+			InputAttachments3[1].layout = VK_IMAGE_LAYOUT_SHADER_READ_ONLY_OPTIMAL;
+			InputAttachments3[1].SetAspect(VK_IMAGE_ASPECT_COLOR_BIT);
+		
+			SubpassDesc.SetInputAttachments(InputAttachments3, 2);
+			SubpassDesc.colorAttachmentCount = 1;
+			SubpassDesc.pColorAttachments = ColorAttachments3;
+
+			SubpassDesc.SetMultiViewMask(MultiviewMask);
+
+			TSubpassDependencyClass& SubpassDep = SubpassDependencies[NumDependencies++];
+			SubpassDep.srcSubpass = 1;
+			SubpassDep.dstSubpass = 2;
+			SubpassDep.srcStageMask = VK_PIPELINE_STAGE_COLOR_ATTACHMENT_OUTPUT_BIT;
+			SubpassDep.dstStageMask = VK_PIPELINE_STAGE_FRAGMENT_SHADER_BIT;
+			SubpassDep.srcAccessMask = VK_ACCESS_COLOR_ATTACHMENT_WRITE_BIT;
+			SubpassDep.dstAccessMask = VK_ACCESS_INPUT_ATTACHMENT_READ_BIT;
+			SubpassDep.dependencyFlags = VK_DEPENDENCY_BY_REGION_BIT;
+		}
+
+		// Only set resolve attachment on the last subpass
+		// CustomResolveSubpass does a custom resolve into second color target and does not need resolve attachments
+		if (!bCustomResolveSubpass)
+		{
+			SubpassDescriptions[NumSubpasses - 1].SetResolveAttachments(ResolveAttachmentReferences);
+		}
+
 		for (uint32 Attachment = 0; Attachment < RTLayout.GetNumAttachmentDescriptions(); ++Attachment)
 		{
 			if (bHasDepthStencilAttachmentReference && (Attachment == DepthStencilAttachmentReference.attachment))
@@ -717,6 +759,10 @@ private:
 	TAttachmentReferenceClass DepthStencilAttachmentReference;
 	TArray<TAttachmentDescriptionClass> AttachmentDescriptions;
 
+	// Tonemap subpass
+	TAttachmentReferenceClass InputAttachments3[MaxSimultaneousRenderTargets + 1];
+	TAttachmentReferenceClass ColorAttachments3[MaxSimultaneousRenderTargets + 1];
+
 	FVulkanAttachmentReference<VkAttachmentReference2> ShadingRateAttachmentReference;
 	FVulkanFragmentShadingRateAttachmentInfo FragmentShadingRateAttachmentInfo;
 
diff --git a/Engine/Source/Runtime/VulkanRHI/Public/VulkanConfiguration.h b/Engine/Source/Runtime/VulkanRHI/Public/VulkanConfiguration.h
--- a/Engine/Source/Runtime/VulkanRHI/Public/VulkanConfiguration.h
+++ b/Engine/Source/Runtime/VulkanRHI/Public/VulkanConfiguration.h
@@ -191,6 +191,14 @@ struct VkAllocationCallbacks;
 	#define VULKAN_SUPPORTS_QCOM_RENDERPASS_TRANSFORM			0
 #endif
 
+#ifndef VULKAN_SUPPORTS_QCOM_RENDERPASS_SHADER_RESOLVE
+	#ifdef VK_QCOM_render_pass_shader_resolve
+		#define VULKAN_SUPPORTS_QCOM_RENDERPASS_SHADER_RESOLVE	1
+	#else
+		#define VULKAN_SUPPORTS_QCOM_RENDERPASS_SHADER_RESOLVE	0
+	#endif
+#endif
+
 #ifndef VULKAN_SUPPORTS_FULLSCREEN_EXCLUSIVE
 	#ifdef VK_EXT_full_screen_exclusive
 		#define VULKAN_SUPPORTS_FULLSCREEN_EXCLUSIVE			1
