Moved TScenePrimitiveArray::Remove over to using EAllowShrinking.
It isn't passing on its argument to the containers - that existing behavior has been retained.

diff --git a/Engine/Source/Runtime/Renderer/Private/ScenePrivateBase.h b/Engine/Source/Runtime/Renderer/Private/ScenePrivateBase.h
--- a/Engine/Source/Runtime/Renderer/Private/ScenePrivateBase.h
+++ b/Engine/Source/Runtime/Renderer/Private/ScenePrivateBase.h
@@ -77,7 +77,7 @@ public:
 		return Chunk.Last();
 	}
 
-	void Remove(uint32 Count, bool bAllowShrinking)
+	void Remove(uint32 Count, EAllowShrinking AllowShrinking)
 	{
 		check(Count <= NumElements);
 		const uint32 NumElementsNew = NumElements - Count;
@@ -86,14 +86,19 @@ public:
 			--NumElements;
 
 			ChunkType& Chunk = *Chunks.Last();
-			Chunk.Pop(false);
+			Chunk.Pop(EAllowShrinking::No);
 
 			if (Chunk.IsEmpty())
 			{
-				Chunks.Pop(false);
+				Chunks.Pop(EAllowShrinking::No);
 			}
 		}
 	}
+	UE_ALLOWSHRINKING_BOOL_DEPRECATED("Remove")
+	FORCEINLINE void Remove(uint32 Count, bool bAllowShrinking)
+	{
+		Remove(Count, bAllowShrinking ? EAllowShrinking::Yes : EAllowShrinking::No);
+	}
 
 	void Reserve(int32 Count)
 	{
