Fix error caching ray tracing mesh commands before FMaterialRenderProxy::UpdateDeferredCachedUniformExpressions(...)
- Context: caching mesh commands requires UniformExpressionCache to be up to date.
- Fix:
  - moved handling renderer changes that require updating cached raytracing data (previously done in RefreshCachedRayTracingState(...)) to new function RayTracing::OnRenderBegin(...)
  - call RayTracing::OnRenderBegin(...) after Scene->UpdateAllPrimitiveSceneInfos(...) which internally updates UniformExpressionCache.

diff --git a/Engine/Source/Runtime/Renderer/Private/DeferredShadingRenderer.cpp b/Engine/Source/Runtime/Renderer/Private/DeferredShadingRenderer.cpp
--- a/Engine/Source/Runtime/Renderer/Private/DeferredShadingRenderer.cpp
+++ b/Engine/Source/Runtime/Renderer/Private/DeferredShadingRenderer.cpp
@@ -228,57 +228,6 @@ static FAutoConsoleCommand RecreateRenderStateContextCmd(
 	TEXT("Recreate render state."),
 	FConsoleCommandDelegate::CreateStatic([] { FGlobalComponentRecreateRenderStateContext Context; }));
 
-#if RHI_RAYTRACING
-
-static bool bUpdateCachedRayTracingState = false;
-
-static FAutoConsoleCommand UpdateCachedRayTracingStateCmd(
-	TEXT("r.RayTracing.UpdateCachedState"),
-	TEXT("Update cached ray tracing state (mesh commands and instances)."),
-	FConsoleCommandDelegate::CreateStatic([] { bUpdateCachedRayTracingState = true; }));
-
-static bool bRefreshRayTracingInstances = false;
-
-static void RefreshRayTracingInstancesSinkFunction()
-{
-	static const auto RayTracingStaticMeshesCVar = IConsoleManager::Get().FindTConsoleVariableDataInt(TEXT("r.RayTracing.Geometry.StaticMeshes"));
-	static const auto RayTracingHISMCVar = IConsoleManager::Get().FindTConsoleVariableDataInt(TEXT("r.RayTracing.Geometry.HierarchicalInstancedStaticMesh"));
-	static const auto RayTracingNaniteProxiesCVar = IConsoleManager::Get().FindTConsoleVariableDataInt(TEXT("r.RayTracing.Geometry.NaniteProxies"));
-	static const auto RayTracingLandscapeGrassCVar = IConsoleManager::Get().FindTConsoleVariableDataInt(TEXT("r.RayTracing.Geometry.LandscapeGrass"));
-
-	static int32 CachedRayTracingStaticMeshes = RayTracingStaticMeshesCVar->GetValueOnGameThread();
-	static int32 CachedRayTracingHISM = RayTracingHISMCVar->GetValueOnGameThread();
-	static int32 CachedRayTracingNaniteProxies = RayTracingNaniteProxiesCVar->GetValueOnGameThread();
-	static int32 CachedRayTracingLandscapeGrass = RayTracingLandscapeGrassCVar->GetValueOnGameThread();
-
-	const int32 RayTracingStaticMeshes = RayTracingStaticMeshesCVar->GetValueOnGameThread();
-	const int32 RayTracingHISM = RayTracingHISMCVar->GetValueOnGameThread();
-	const int32 RayTracingNaniteProxies = RayTracingNaniteProxiesCVar->GetValueOnGameThread();
-	const int32 RayTracingLandscapeGrass = RayTracingLandscapeGrassCVar->GetValueOnGameThread();
-
-	if (RayTracingStaticMeshes != CachedRayTracingStaticMeshes
-		|| RayTracingHISM != CachedRayTracingHISM
-		|| RayTracingNaniteProxies != CachedRayTracingNaniteProxies
-		|| RayTracingLandscapeGrass != CachedRayTracingLandscapeGrass)
-	{
-		ENQUEUE_RENDER_COMMAND(RefreshRayTracingInstancesCmd)(
-			[](FRHICommandListImmediate&)
-			{
-				bRefreshRayTracingInstances = true;
-			}
-		);
-
-		CachedRayTracingStaticMeshes = RayTracingStaticMeshes;
-		CachedRayTracingHISM = RayTracingHISM;
-		CachedRayTracingNaniteProxies = RayTracingNaniteProxies;
-		CachedRayTracingLandscapeGrass = RayTracingLandscapeGrass;
-	}
-}
-
-static FAutoConsoleVariableSink CVarRefreshRayTracingInstancesSink(FConsoleCommandDelegate::CreateStatic(&RefreshRayTracingInstancesSinkFunction));
-
-#endif // RHI_RAYTRACING
-
 #if !UE_BUILD_SHIPPING
 static TAutoConsoleVariable<int32> CVarForceBlackVelocityBuffer(
 	TEXT("r.Test.ForceBlackVelocityBuffer"), 0,
@@ -605,54 +554,6 @@ static void RenderOpaqueFX(
 
 #if RHI_RAYTRACING
 
-static void RefreshCachedRayTracingState(FScene& Scene, const FSceneViewFamily& ViewFamily, TArrayView<FViewInfo> Views)
-{
-	const ERayTracingMeshCommandsMode CurrentMode = ViewFamily.EngineShowFlags.PathTracing ? ERayTracingMeshCommandsMode::PATH_TRACING : ERayTracingMeshCommandsMode::RAY_TRACING;
-	bool bNaniteCoarseMeshStreamingModeChanged = false;
-#if WITH_EDITOR
-	bNaniteCoarseMeshStreamingModeChanged = Nanite::FCoarseMeshStreamingManager::CheckStreamingMode();
-#endif // WITH_EDITOR
-	const bool bNaniteRayTracingModeChanged = Nanite::GRayTracingManager.CheckModeChanged();
-
-	if (CurrentMode != Scene.CachedRayTracingMeshCommandsMode
-		|| bNaniteCoarseMeshStreamingModeChanged
-		|| bNaniteRayTracingModeChanged
-		|| bUpdateCachedRayTracingState)
-	{
-		Scene.WaitForCacheRayTracingPrimitivesTask();
-
-		// In some situations, we need to refresh the cached ray tracing mesh commands because they contain data about the currently bound shader. 
-		// This operation is a bit expensive but only happens once as we transition between modes which should be rare.
-		Scene.CachedRayTracingMeshCommandsMode = CurrentMode;
-		Scene.RefreshRayTracingMeshCommandCache();
-		bUpdateCachedRayTracingState = false;
-	}
-
-	if (bRefreshRayTracingInstances)
-	{
-		Scene.WaitForCacheRayTracingPrimitivesTask();
-
-		// In some situations, we need to refresh the cached ray tracing instance.
-		// This assumes that cached instances will keep using the same LOD since CachedRayTracingMeshCommands is not recalculated
-		// eg: Need to update PrimitiveRayTracingFlags
-		// This operation is a bit expensive but only happens once as we transition between modes which should be rare.
-		Scene.RefreshRayTracingInstances();
-		bRefreshRayTracingInstances = false;
-	}
-
-	if (bNaniteRayTracingModeChanged)
-	{
-		for (FViewInfo& View : Views)
-		{
-			if (View.ViewState != nullptr && !View.bIsOfflineRender)
-			{
-				// don't invalidate in the offline case because we only get one attempt at rendering each sample
-				View.ViewState->PathTracingInvalidate();
-			}
-		}
-	}
-}
-
 static bool ShouldPrepareRayTracingDecals(const FScene& Scene, const FSceneViewFamily& ViewFamily)
 {
 	if (!IsRayTracingEnabled() || !RHISupportsRayTracingCallableShaders(ViewFamily.GetShaderPlatform()))
@@ -1542,8 +1443,6 @@ void FDeferredShadingSceneRenderer::Render(FRDGBuilder& GraphBuilder)
 		SCOPED_GPU_STAT(GraphBuilder.RHICmdList, RayTracingGeometry);
 		GRayTracingGeometryManager->ProcessBuildRequests(GraphBuilder.RHICmdList);
 	}
-
-	RefreshCachedRayTracingState(*Scene, ViewFamily, Views);
 #endif
 
 	FInitViewTaskDatas InitViewTaskDatas = OnRenderBegin(GraphBuilder);
diff --git a/Engine/Source/Runtime/Renderer/Private/RayTracing/RayTracing.cpp b/Engine/Source/Runtime/Renderer/Private/RayTracing/RayTracing.cpp
--- a/Engine/Source/Runtime/Renderer/Private/RayTracing/RayTracing.cpp
+++ b/Engine/Source/Runtime/Renderer/Private/RayTracing/RayTracing.cpp
@@ -93,6 +93,53 @@ static FAutoConsoleVariableRef CVarRayTracingDebugForceOpaque(
 	TEXT("Forces all ray tracing geometry instances to be opaque, effectively disabling any-hit shaders. This is useful for debugging and profiling. (default = 0)")
 );
 
+static bool bUpdateCachedRayTracingState = false;
+
+static FAutoConsoleCommand UpdateCachedRayTracingStateCmd(
+	TEXT("r.RayTracing.UpdateCachedState"),
+	TEXT("Update cached ray tracing state (mesh commands and instances)."),
+	FConsoleCommandDelegate::CreateStatic([] { bUpdateCachedRayTracingState = true; }));
+
+static bool bRefreshRayTracingInstances = false;
+
+static void RefreshRayTracingInstancesSinkFunction()
+{
+	static const auto RayTracingStaticMeshesCVar = IConsoleManager::Get().FindTConsoleVariableDataInt(TEXT("r.RayTracing.Geometry.StaticMeshes"));
+	static const auto RayTracingHISMCVar = IConsoleManager::Get().FindTConsoleVariableDataInt(TEXT("r.RayTracing.Geometry.HierarchicalInstancedStaticMesh"));
+	static const auto RayTracingNaniteProxiesCVar = IConsoleManager::Get().FindTConsoleVariableDataInt(TEXT("r.RayTracing.Geometry.NaniteProxies"));
+	static const auto RayTracingLandscapeGrassCVar = IConsoleManager::Get().FindTConsoleVariableDataInt(TEXT("r.RayTracing.Geometry.LandscapeGrass"));
+
+	static int32 CachedRayTracingStaticMeshes = RayTracingStaticMeshesCVar->GetValueOnGameThread();
+	static int32 CachedRayTracingHISM = RayTracingHISMCVar->GetValueOnGameThread();
+	static int32 CachedRayTracingNaniteProxies = RayTracingNaniteProxiesCVar->GetValueOnGameThread();
+	static int32 CachedRayTracingLandscapeGrass = RayTracingLandscapeGrassCVar->GetValueOnGameThread();
+
+	const int32 RayTracingStaticMeshes = RayTracingStaticMeshesCVar->GetValueOnGameThread();
+	const int32 RayTracingHISM = RayTracingHISMCVar->GetValueOnGameThread();
+	const int32 RayTracingNaniteProxies = RayTracingNaniteProxiesCVar->GetValueOnGameThread();
+	const int32 RayTracingLandscapeGrass = RayTracingLandscapeGrassCVar->GetValueOnGameThread();
+
+	if (RayTracingStaticMeshes != CachedRayTracingStaticMeshes
+		|| RayTracingHISM != CachedRayTracingHISM
+		|| RayTracingNaniteProxies != CachedRayTracingNaniteProxies
+		|| RayTracingLandscapeGrass != CachedRayTracingLandscapeGrass)
+	{
+		ENQUEUE_RENDER_COMMAND(RefreshRayTracingInstancesCmd)(
+			[](FRHICommandListImmediate&)
+			{
+				bRefreshRayTracingInstances = true;
+			}
+		);
+
+		CachedRayTracingStaticMeshes = RayTracingStaticMeshes;
+		CachedRayTracingHISM = RayTracingHISM;
+		CachedRayTracingNaniteProxies = RayTracingNaniteProxies;
+		CachedRayTracingLandscapeGrass = RayTracingLandscapeGrass;
+	}
+}
+
+static FAutoConsoleVariableSink CVarRefreshRayTracingInstancesSink(FConsoleCommandDelegate::CreateStatic(&RefreshRayTracingInstancesSinkFunction));
+
 namespace RayTracing
 {
 	struct FRelevantPrimitive
@@ -168,6 +215,54 @@ namespace RayTracing
 		bool bValid = false;
 	};
 
+	void OnRenderBegin(FScene& Scene, TArray<FViewInfo>& Views, const FViewFamilyInfo& ViewFamily)
+	{
+		const ERayTracingMeshCommandsMode CurrentMode = ViewFamily.EngineShowFlags.PathTracing ? ERayTracingMeshCommandsMode::PATH_TRACING : ERayTracingMeshCommandsMode::RAY_TRACING;
+		bool bNaniteCoarseMeshStreamingModeChanged = false;
+#if WITH_EDITOR
+		bNaniteCoarseMeshStreamingModeChanged = Nanite::FCoarseMeshStreamingManager::CheckStreamingMode();
+#endif // WITH_EDITOR
+		const bool bNaniteRayTracingModeChanged = Nanite::GRayTracingManager.CheckModeChanged();
+
+		if (CurrentMode != Scene.CachedRayTracingMeshCommandsMode
+			|| bNaniteCoarseMeshStreamingModeChanged
+			|| bNaniteRayTracingModeChanged
+			|| bUpdateCachedRayTracingState)
+		{
+			Scene.WaitForCacheRayTracingPrimitivesTask();
+
+			// In some situations, we need to refresh the cached ray tracing mesh commands because they contain data about the currently bound shader. 
+			// This operation is a bit expensive but only happens once as we transition between modes which should be rare.
+			Scene.CachedRayTracingMeshCommandsMode = CurrentMode;
+			Scene.RefreshRayTracingMeshCommandCache();
+			bUpdateCachedRayTracingState = false;
+		}
+
+		if (bRefreshRayTracingInstances)
+		{
+			Scene.WaitForCacheRayTracingPrimitivesTask();
+
+			// In some situations, we need to refresh the cached ray tracing instance.
+			// This assumes that cached instances will keep using the same LOD since CachedRayTracingMeshCommands is not recalculated
+			// eg: Need to update PrimitiveRayTracingFlags
+			// This operation is a bit expensive but only happens once as we transition between modes which should be rare.
+			Scene.RefreshRayTracingInstances();
+			bRefreshRayTracingInstances = false;
+		}
+
+		if (bNaniteRayTracingModeChanged)
+		{
+			for (FViewInfo& View : Views)
+			{
+				if (View.ViewState != nullptr && !View.bIsOfflineRender)
+				{
+					// don't invalidate in the offline case because we only get one attempt at rendering each sample
+					View.ViewState->PathTracingInvalidate();
+				}
+			}
+		}
+	}
+
 	FRelevantPrimitiveList* CreateRelevantPrimitiveList(FSceneRenderingBulkObjectAllocator& InAllocator)
 	{
 		return InAllocator.Create<FRelevantPrimitiveList>();
diff --git a/Engine/Source/Runtime/Renderer/Private/RayTracing/RayTracing.h b/Engine/Source/Runtime/Renderer/Private/RayTracing/RayTracing.h
--- a/Engine/Source/Runtime/Renderer/Private/RayTracing/RayTracing.h
+++ b/Engine/Source/Runtime/Renderer/Private/RayTracing/RayTracing.h
@@ -11,12 +11,15 @@
 #include "RenderGraphDefinitions.h"
 
 class FRayTracingScene;
+class FViewFamilyInfo;
 class FGlobalDynamicReadBuffer;
 
 namespace RayTracing
 {
 	struct FRelevantPrimitiveList;
 
+	void OnRenderBegin(FScene& Scene, TArray<FViewInfo>& Views, const FViewFamilyInfo& ViewFamily);
+
 	FRelevantPrimitiveList* CreateRelevantPrimitiveList(FSceneRenderingBulkObjectAllocator& InAllocator);
 
 	// Iterates over Scene's PrimitiveSceneProxies and extracts ones that are relevant for ray tracing.
diff --git a/Engine/Source/Runtime/Renderer/Private/SceneRendering.cpp b/Engine/Source/Runtime/Renderer/Private/SceneRendering.cpp
--- a/Engine/Source/Runtime/Renderer/Private/SceneRendering.cpp
+++ b/Engine/Source/Runtime/Renderer/Private/SceneRendering.cpp
@@ -69,6 +69,7 @@
 #if WITH_EDITOR
 #include "Rendering/StaticLightingSystemInterface.h"
 #endif
+#include "RayTracing/RayTracing.h"
 #include "RayTracing/RayTracingScene.h"
 #include "Rendering/RayTracingGeometryManager.h"
 #include "FXSystem.h"
@@ -3464,6 +3465,10 @@ IVisibilityTaskData* FSceneRenderer::OnRenderBegin(FRDGBuilder& GraphBuilder)
 
 	Scene->UpdateAllPrimitiveSceneInfos(GraphBuilder, AsyncOps);
 
+#if RHI_RAYTRACING
+	RayTracing::OnRenderBegin(*Scene, Views, ViewFamily);
+#endif
+
 	if (!ViewFamily.ViewExtensions.IsEmpty())
 	{
 		{
