My previous change to MeshDrawCommandStats (allowing grouping stats by 'collections') didn't take into account the TotalPrimitives metric, this was left as a TODO.  This CL finally circles back and allows the user to specificy the primitive budget per collection.

An example of this: A collection might only use the basepass, the display should have a primitive budget number relevant for the basepass, not all passes.

diff --git a/Engine/Source/Runtime/Engine/Public/MeshDrawCommandStatsSettings.h b/Engine/Source/Runtime/Engine/Public/MeshDrawCommandStatsSettings.h
--- a/Engine/Source/Runtime/Engine/Public/MeshDrawCommandStatsSettings.h
+++ b/Engine/Source/Runtime/Engine/Public/MeshDrawCommandStatsSettings.h
@@ -30,6 +30,20 @@ struct FMeshDrawCommandStatsBudget
 	TArray<FName> Passes;
 };
 
+/** Budget totals for each MeshDrawCommandStats collection. */
+USTRUCT()
+struct FMeshDrawCommandStatsBudgetTotals
+{
+	GENERATED_BODY()
+
+	/** The total amount of primitives budgeted for this collection. */
+	UPROPERTY(config, EditAnywhere, Category = Engine)
+	int32 PrimitiveBudget = 0;
+	/** The collection this applies to. */
+	UPROPERTY(config, EditAnywhere, Category = Engine)
+	int32 Collection = 0;
+};
+
 /** User settings used by the MeshDrawCommandStats system. */
 UCLASS(Config=Engine, defaultconfig, perplatformconfig, meta = (DisplayName = "Mesh Stats"))
 class ENGINE_API UMeshDrawCommandStatsSettings : public UDeveloperSettings
@@ -40,9 +54,9 @@ public:
 	/** Budgets used by r.MeshDrawCommands.Stats */
 	UPROPERTY(config, EditAnywhere, Category = Engine)
 	TArray<FMeshDrawCommandStatsBudget> Budgets;
-	/** The total primitive budget. This is the maximimum triangles expected, post-culling, summed across all passes. */
+	/** The total primitive budget for a collection. */
 	UPROPERTY(config, EditAnywhere, Category = Engine)
-	int32 TotalPrimitiveBudget = 0;
+	TArray<FMeshDrawCommandStatsBudgetTotals> BudgetTotals;
 	/** Which collection to export to CSV */
 	UPROPERTY(config, EditAnywhere, Category = Engine)
 	int32 CollectionForCsvProfiler = 1;
diff --git a/Engine/Source/Runtime/Renderer/Private/MeshDrawCommandStats.cpp b/Engine/Source/Runtime/Renderer/Private/MeshDrawCommandStats.cpp
--- a/Engine/Source/Runtime/Renderer/Private/MeshDrawCommandStats.cpp
+++ b/Engine/Source/Runtime/Renderer/Private/MeshDrawCommandStats.cpp
@@ -146,6 +146,9 @@ FMeshDrawCommandStatsManager::FMeshDrawCommandStatsManager()
 	// Is it fine to keep the screen message delegate always registered even if we are not showing anything?
 	ScreenMessageDelegate = FRendererOnScreenNotification::Get().AddLambda([this](TMultiMap<FCoreDelegates::EOnScreenMessageSeverity, FText >& OutMessages)
 		{	
+			int32 TotalPrimitivesTracked = 0;
+			int32 TotalPrimitivesUntracked = 0;
+
 			const bool bShowStats = CVarMeshDrawCommandStats->GetInt() != (int)MeshDrawStatsCollection::None;
 			if (bShowStats)
 			{
@@ -170,6 +173,8 @@ FMeshDrawCommandStatsManager::FMeshDrawCommandStatsManager()
 									*(CategoryBudget.CategoryName.ToString()),
 									*PassFriendlyNames
 							)));
+
+							TotalPrimitivesTracked += *PrimitiveCount;
 						}
 					}
 				}
@@ -184,14 +189,22 @@ FMeshDrawCommandStatsManager::FMeshDrawCommandStatsManager()
 					{
 						OutMessages.Add(FCoreDelegates::EOnScreenMessageSeverity::Info, FText::FromString(FString::Printf(TEXT("\t%5dK - %s"), PrimitiveCount / 1000, *(Name.ToString()))));
 					}
+
+					TotalPrimitivesUntracked += PrimitiveCount;
 				}
 
 				// Show total budget.
-				const int32 TotalBudget = Settings != 0 ? Settings->TotalPrimitiveBudget : 0;
-				if (TotalBudget > 0)
+				FStatCollection* StatCollection = StatCollections.Find(Collection);
+				const int32 PrimitiveBudget = StatCollection != nullptr ? StatCollection->PrimitiveBudget : 0;
+				if (PrimitiveBudget > 0)
 				{
-					FCoreDelegates::EOnScreenMessageSeverity Severity = TotalBudget < Stats.TotalPrimitives ? FCoreDelegates::EOnScreenMessageSeverity::Warning : FCoreDelegates::EOnScreenMessageSeverity::Info;
-					OutMessages.Add(Severity, FText::FromString(FString::Printf(TEXT("%5dK / %5dK - TOTAL (All Passes)"), Stats.TotalPrimitives / 1000, TotalBudget / 1000)));
+					FCoreDelegates::EOnScreenMessageSeverity Severity = PrimitiveBudget < Stats.TotalPrimitives ? FCoreDelegates::EOnScreenMessageSeverity::Warning : FCoreDelegates::EOnScreenMessageSeverity::Info;
+					OutMessages.Add(Severity, FText::FromString(FString::Printf(TEXT("%5dK / %5dK - Total Budgeted"), TotalPrimitivesTracked / 1000, PrimitiveBudget / 1000)));
+					
+					if (TotalPrimitivesUntracked)
+					{
+						OutMessages.Add(Severity, FText::FromString(FString::Printf(TEXT("%5dK - Total Untracked"), TotalPrimitivesUntracked)));
+					}
 				}
 				else
 				{
@@ -453,6 +466,15 @@ void FMeshDrawCommandStatsManager::Update()
 				}
 			}
 			
+			for (const FMeshDrawCommandStatsBudgetTotals& BudgetTotal : Settings->BudgetTotals)
+			{
+				FStatCollection* Collection = StatCollections.Find(BudgetTotal.Collection);
+				if (Collection)
+				{
+					Collection->PrimitiveBudget = BudgetTotal.PrimitiveBudget;
+				}
+			}
+			
 			for (TPair<int32, FStatCollection>& Pair : StatCollections)
 			{
 				Pair.Value.Finish();
diff --git a/Engine/Source/Runtime/Renderer/Private/MeshDrawCommandStats.h b/Engine/Source/Runtime/Renderer/Private/MeshDrawCommandStats.h
--- a/Engine/Source/Runtime/Renderer/Private/MeshDrawCommandStats.h
+++ b/Engine/Source/Runtime/Renderer/Private/MeshDrawCommandStats.h
@@ -194,6 +194,7 @@ private:
 		TArray<FCollectionCategory> Categories;
 		TMap<FName, TArray<int>> StatToCategoryIndices;
 		TMap<FName, FString> CategoryPassFriendlyNames; //< Pipe delimited passes tracked by this category, used for on screen display
+		int32 PrimitiveBudget;
 	};
 
 	using FStatCollectionMap = TMap<int32, FStatCollection>;
