Many Lights - run spatial filter until we accumulate enough frames to estimated temporal variance. Also don't advance the frame counter when history is paused (read only).

diff --git a/Engine/Shaders/Private/ManyLights/ManyLightsDenoiserSpatial.usf b/Engine/Shaders/Private/ManyLights/ManyLightsDenoiserSpatial.usf
--- a/Engine/Shaders/Private/ManyLights/ManyLightsDenoiserSpatial.usf
+++ b/Engine/Shaders/Private/ManyLights/ManyLightsDenoiserSpatial.usf
@@ -5,9 +5,11 @@
 
 Texture2D<float4> DiffuseLightingAndSecondMomentTexture;
 Texture2D<float4> SpecularLightingAndSecondMomentTexture;
+Texture2D<UNORM float> NumFramesAccumulatedTexture;
 RWTexture2D<float4> RWSceneColor;
 
 float SpatialFilterDepthWeightScale;
+float TemporalMaxFramesAccumulated;
 
 /**
  * Run a spatial filter in order to filter out noise based on the temporal variance.
@@ -41,6 +43,10 @@ void DenoiserSpatialCS(
 
 			float CenterDiffuseStdDev = sqrt(max(CenterDiffuseLightingAndSecondMoment.w - Pow2(Luminance(CenterDiffuseLightingAndSecondMoment.xyz)), 0.0f));
 			float CenterSpecularStdDev = sqrt(max(CenterSpecularLightingAndSecondMoment.w - Pow2(Luminance(CenterSpecularLightingAndSecondMoment.xyz)), 0.0f));
+	
+			// Run spatial filter until we accumulate enough frames to estimated temporal variance
+			float NumFramesAccumulated = NumFramesAccumulatedTexture[ScreenCoord] * TemporalMaxFramesAccumulated;
+			CenterDiffuseStdDev = lerp(1.0f, CenterDiffuseStdDev, saturate(NumFramesAccumulated / min(5.0f, TemporalMaxFramesAccumulated)));
 
 			float3 DiffuseLightingSum = CenterDiffuseLightingAndSecondMoment.xyz;
 			float3 SpecularLightingSum = CenterSpecularLightingAndSecondMoment.xyz;
diff --git a/Engine/Shaders/Private/ManyLights/ManyLightsDenoiserTemporal.usf b/Engine/Shaders/Private/ManyLights/ManyLightsDenoiserTemporal.usf
--- a/Engine/Shaders/Private/ManyLights/ManyLightsDenoiserTemporal.usf
+++ b/Engine/Shaders/Private/ManyLights/ManyLightsDenoiserTemporal.usf
@@ -17,6 +17,7 @@ float4 HistoryUVMinMax;
 float PrevSceneColorPreExposureCorrection;
 float TemporalMaxFramesAccumulated;
 float TemporalNeighborhoodClampScale;
+uint TemporalAdvanceFrame;
 uint2 DownsampledViewSize;
 
 RWTexture2D<float4> RWDiffuseLightingAndSecondMoment;
@@ -207,7 +208,13 @@ void DenoiserTemporalCS(
 						float NumFramesAccumulatedHistory = 0.0f;
 						NumFramesAccumulatedHistory = dot(NumFramesAccumulatedHistoryTexture.GatherRed(GlobalPointClampedSampler, HistoryGatherUV).wzxy, HistoryWeights) * TemporalMaxFramesAccumulated;
 						NumFramesAccumulatedHistory = NumFramesAccumulatedHistory / dot(HistoryWeights, 1.0f);
-						NumFramesAccumulated = min(NumFramesAccumulatedHistory + 1.0f, TemporalMaxFramesAccumulated);
+
+						// Advance the frame counter
+						if (TemporalAdvanceFrame != 0)
+						{
+							NumFramesAccumulatedHistory += 1.0f;
+						}
+						NumFramesAccumulated = min(NumFramesAccumulatedHistory, TemporalMaxFramesAccumulated);
 
 						// Clamp history to current neighborhood
 						FNeighborhood DiffuseNeighborhood = GetNeighborhood(ScreenCoord, ResolvedDiffuseLighting, DiffuseLighting);
diff --git a/Engine/Source/Runtime/Renderer/Private/ManyLights/ManyLights.cpp b/Engine/Source/Runtime/Renderer/Private/ManyLights/ManyLights.cpp
--- a/Engine/Source/Runtime/Renderer/Private/ManyLights/ManyLights.cpp
+++ b/Engine/Source/Runtime/Renderer/Private/ManyLights/ManyLights.cpp
@@ -587,6 +587,7 @@ class FDenoiserSpatialCS : public FGlobalShader
 		SHADER_PARAMETER_RDG_TEXTURE_UAV(RWTexture2D<float4>, RWSceneColor)
 		SHADER_PARAMETER_RDG_TEXTURE(Texture2D<float3>, DiffuseLightingAndSecondMomentTexture)
 		SHADER_PARAMETER_RDG_TEXTURE(Texture2D<float3>, SpecularLightingAndSecondMomentTexture)
+		SHADER_PARAMETER_RDG_TEXTURE(Texture2D<UNORM float>, NumFramesAccumulatedTexture)
 		SHADER_PARAMETER(float, SpatialFilterDepthWeightScale)
 	END_SHADER_PARAMETER_STRUCT()
 
@@ -763,6 +764,7 @@ void FDeferredShadingSceneRenderer::RenderManyLights(FRDGBuilder& GraphBuilder,
 		ManyLightsParameters.DownsampledTileDataStride = DownsampledTileDataStride;
 		ManyLightsParameters.TemporalMaxFramesAccumulated = FMath::Max(CVarManyLightsTemporalMaxFramesAccumulated.GetValueOnRenderThread(), 0.0f);
 		ManyLightsParameters.TemporalNeighborhoodClampScale = CVarManyLightsTemporalNeighborhoodClampScale.GetValueOnRenderThread();
+		ManyLightsParameters.TemporalAdvanceFrame = View.ViewState && !View.bStatePrevViewInfoIsReadOnly ? 1 : 0;
 		ManyLightsParameters.DebugMode = ManyLights::GetDebugMode();
 		ManyLightsParameters.DebugLightId = INDEX_NONE;
 
@@ -1145,6 +1147,7 @@ void FDeferredShadingSceneRenderer::RenderManyLights(FRDGBuilder& GraphBuilder,
 		PassParameters->RWSceneColor = GraphBuilder.CreateUAV(SceneTextures.Color.Target);
 		PassParameters->DiffuseLightingAndSecondMomentTexture = DiffuseLightingAndSecondMoment;
 		PassParameters->SpecularLightingAndSecondMomentTexture = SpecularLightingAndSecondMoment;
+		PassParameters->NumFramesAccumulatedTexture = NumFramesAccumulated;
 		PassParameters->SpatialFilterDepthWeightScale = CVarManyLightsSpatialDepthWeightScale.GetValueOnRenderThread();
 
 		FDenoiserSpatialCS::FPermutationDomain PermutationVector;
diff --git a/Engine/Source/Runtime/Renderer/Private/ManyLights/ManyLightsInternal.h b/Engine/Source/Runtime/Renderer/Private/ManyLights/ManyLightsInternal.h
--- a/Engine/Source/Runtime/Renderer/Private/ManyLights/ManyLightsInternal.h
+++ b/Engine/Source/Runtime/Renderer/Private/ManyLights/ManyLightsInternal.h
@@ -33,6 +33,7 @@ BEGIN_SHADER_PARAMETER_STRUCT(FManyLightsParameters, )
 	SHADER_PARAMETER(int32, DownsampledTileDataStride)
 	SHADER_PARAMETER(float, TemporalMaxFramesAccumulated)
 	SHADER_PARAMETER(float, TemporalNeighborhoodClampScale)
+	SHADER_PARAMETER(int32, TemporalAdvanceFrame)
 	SHADER_PARAMETER(int32, DebugMode)
 	SHADER_PARAMETER(int32, DebugLightId)
 	SHADER_PARAMETER_RDG_TEXTURE(Texture2D<uint>, DownsampledTileMask)
