Working on visibility overlap with gpu scene update.

diff --git a/Engine/Source/Runtime/Renderer/Private/GPUScene.cpp b/Engine/Source/Runtime/Renderer/Private/GPUScene.cpp
--- a/Engine/Source/Runtime/Renderer/Private/GPUScene.cpp
+++ b/Engine/Source/Runtime/Renderer/Private/GPUScene.cpp
@@ -773,7 +773,7 @@ void FGPUScene::InitLightData(const FLightSceneInfoCompact& LightInfoCompact, bo
 	DataOut.LightTypeAndShadowMapChannelMaskPacked = LightInfo.PackLightTypeAndShadowMapChannelMask(bAllowStaticLighting);
 }
 
-void FGPUScene::UpdateInternal(FRDGBuilder& GraphBuilder, FSceneUniformBuffer& SceneUB, FRDGExternalAccessQueue& ExternalAccessQueue, IVisibilityTaskData* VisibilityTaskData)
+void FGPUScene::UpdateInternal(FRDGBuilder& GraphBuilder, FSceneUniformBuffer& SceneUB, FRDGExternalAccessQueue& ExternalAccessQueue, const UE::Tasks::FTask& UpdateTaskPrerequisites)
 {
 	LLM_SCOPE_BYTAG(GPUScene);
 
@@ -857,20 +857,12 @@ void FGPUScene::UpdateInternal(FRDGBuilder& GraphBuilder, FSceneUniformBuffer& S
 	PrimitivesToUpdate.Reset();
 	PrimitiveDirtyState.Init(EPrimitiveDirtyState::None, PrimitiveDirtyState.Num());
 
-	UE::Tasks::FTask PrerequisiteAsyncTask;
-
-	// An optimization to schedule the async task to occur after dynamic mesh elements are computed in order to reduce contention in the visibility task graph.
-	if (GraphBuilder.IsParallelSetupEnabled() && VisibilityTaskData)
-	{
-		PrerequisiteAsyncTask = VisibilityTaskData->GetComputeRelevanceTask();
-	}
-
 	{
 		SCOPED_NAMED_EVENT(UpdateGPUScene, FColor::Green);
 		QUICK_SCOPE_CYCLE_COUNTER(STAT_UpdateGPUScene);
 		SCOPE_CYCLE_COUNTER(STAT_UpdateGPUSceneTime);
 
-		UploadGeneral<FUploadDataSourceAdapterScenePrimitives>(GraphBuilder, BufferState, &ExternalAccessQueue, Adapter, PrerequisiteAsyncTask);
+		UploadGeneral<FUploadDataSourceAdapterScenePrimitives>(GraphBuilder, BufferState, &ExternalAccessQueue, Adapter, UpdateTaskPrerequisites);
 	}
 }
 
@@ -1764,7 +1756,7 @@ void FGPUScene::AddPrimitiveToUpdate(FPersistentPrimitiveIndex PersistentPrimiti
 }
 
 
-void FGPUScene::Update(FRDGBuilder& GraphBuilder, FSceneUniformBuffer& SceneUB, FRDGExternalAccessQueue& ExternalAccessQueue, IVisibilityTaskData* VisibilityTaskData)
+void FGPUScene::Update(FRDGBuilder& GraphBuilder, FSceneUniformBuffer& SceneUB, FRDGExternalAccessQueue& ExternalAccessQueue, const UE::Tasks::FTask& UpdateTaskPrerequisites)
 {
 	if (bIsEnabled)
 	{
@@ -1778,7 +1770,7 @@ void FGPUScene::Update(FRDGBuilder& GraphBuilder, FSceneUniformBuffer& SceneUB,
 		// Default state when updated (no "dynamic primitives" pushed)
 		DynamicPrimitivesOffset = Scene.GetMaxPersistentPrimitiveIndex();
 
-		UpdateInternal(GraphBuilder, SceneUB, ExternalAccessQueue, VisibilityTaskData);
+		UpdateInternal(GraphBuilder, SceneUB, ExternalAccessQueue, UpdateTaskPrerequisites);
 	}
 }
 
diff --git a/Engine/Source/Runtime/Renderer/Private/GPUScene.h b/Engine/Source/Runtime/Renderer/Private/GPUScene.h
--- a/Engine/Source/Runtime/Renderer/Private/GPUScene.h
+++ b/Engine/Source/Runtime/Renderer/Private/GPUScene.h
@@ -26,7 +26,6 @@ class FInstanceCullingOcclusionQueryRenderer;
 class FScenePreUpdateChangeSet;
 class FScenePostUpdateChangeSet;
 class FPrimitiveSceneProxy;
-class IVisibilityTaskData;
 struct FLightSceneChangeSet;
 namespace UE::Renderer::Private
 {
@@ -231,7 +230,7 @@ public:
 	/**
 	 * Pull all pending updates from Scene and upload primitive & instance data.
 	 */
-	void Update(FRDGBuilder& GraphBuilder, FSceneUniformBuffer& SceneUB, FRDGExternalAccessQueue& ExternalAccessQueue, IVisibilityTaskData* VisibilityTaskData = nullptr);
+	void Update(FRDGBuilder& GraphBuilder, FSceneUniformBuffer& SceneUB, FRDGExternalAccessQueue& ExternalAccessQueue, const UE::Tasks::FTask& UpdateTaskPrerequisites = {});
 
 	/**
 	 * Queue the given primitive for upload to GPU at next call to Update.
@@ -431,7 +430,7 @@ private:
 
 	void UploadDynamicPrimitiveShaderDataForViewInternal(FRDGBuilder& GraphBuilder, FViewInfo& View, UE::Renderer::Private::IShadowInvalidatingInstances *ShadowInvalidatingInstances);
 
-	void UpdateInternal(FRDGBuilder& GraphBuilder, FSceneUniformBuffer& SceneUB, FRDGExternalAccessQueue& ExternalAccessQueue, IVisibilityTaskData* VisibilityTaskData);
+	void UpdateInternal(FRDGBuilder& GraphBuilder, FSceneUniformBuffer& SceneUB, FRDGExternalAccessQueue& ExternalAccessQueue, const UE::Tasks::FTask& UpdateTaskPrerequisites);
 
 	void AddUpdatePrimitiveIdsPass(FRDGBuilder& GraphBuilder, FInstanceGPULoadBalancer& IdOnlyUpdateItems);
 
diff --git a/Engine/Source/Runtime/Renderer/Private/PrimitiveSceneInfo.cpp b/Engine/Source/Runtime/Renderer/Private/PrimitiveSceneInfo.cpp
--- a/Engine/Source/Runtime/Renderer/Private/PrimitiveSceneInfo.cpp
+++ b/Engine/Source/Runtime/Renderer/Private/PrimitiveSceneInfo.cpp
@@ -1240,7 +1240,7 @@ void FPrimitiveSceneInfo::RemoveCachedRayTracingPrimitives()
 }
 #endif
 
-void FPrimitiveSceneInfo::AddStaticMeshes(FScene* Scene, TArrayView<FPrimitiveSceneInfo*> SceneInfos, bool bCacheMeshDrawCommands)
+void FPrimitiveSceneInfo::AddStaticMeshes(FRHICommandListBase& RHICmdList, FScene* Scene, TArrayView<FPrimitiveSceneInfo*> SceneInfos, bool bCacheMeshDrawCommands)
 {
 	LLM_SCOPE(ELLMTag::StaticMesh);
 
@@ -1264,7 +1264,6 @@ void FPrimitiveSceneInfo::AddStaticMeshes(FScene* Scene, TArrayView<FPrimitiveSc
 
 	{
 		const ERHIFeatureLevel::Type FeatureLevel = Scene->GetFeatureLevel();
-		FRHICommandListBase& RHICmdList = FRHICommandListImmediate::Get();
 
 		SCOPED_NAMED_EVENT(FPrimitiveSceneInfo_AddStaticMeshes_UpdateSceneArrays, FColor::Blue);
 		for (FPrimitiveSceneInfo* SceneInfo : SceneInfos)
diff --git a/Engine/Source/Runtime/Renderer/Private/RendererScene.cpp b/Engine/Source/Runtime/Renderer/Private/RendererScene.cpp
--- a/Engine/Source/Runtime/Renderer/Private/RendererScene.cpp
+++ b/Engine/Source/Runtime/Renderer/Private/RendererScene.cpp
@@ -4519,7 +4519,7 @@ void FScene::UpdateStaticDrawLists_RenderThread(FRHICommandListImmediate& RHICmd
 		Primitive->RemoveStaticMeshes();
 	}
 
-	FPrimitiveSceneInfo::AddStaticMeshes(this, Primitives);
+	FPrimitiveSceneInfo::AddStaticMeshes(RHICmdList, this, Primitives);
 }
 
 void FScene::UpdateStaticDrawLists()
@@ -5412,8 +5412,14 @@ struct FSceneUpdateChangeSetStorage
 	}
 };
 
-
 void FScene::UpdateAllPrimitiveSceneInfos(FRDGBuilder& GraphBuilder, EUpdateAllPrimitiveSceneInfosAsyncOps AsyncOps)
+{
+	FUpdateParameters Parameters;
+	Parameters.AsyncOps = AsyncOps;
+	Update(GraphBuilder, Parameters);
+}
+
+void FScene::Update(FRDGBuilder& GraphBuilder, const FUpdateParameters& Parameters)
 {
 	TRACE_CPUPROFILER_EVENT_SCOPE(Scene::UpdateAllPrimitiveSceneInfos);
 	SCOPED_NAMED_EVENT(FScene_UpdateAllPrimitiveSceneInfos, FColor::Orange);
@@ -5441,7 +5447,7 @@ void FScene::UpdateAllPrimitiveSceneInfos(FRDGBuilder& GraphBuilder, EUpdateAllP
 	FSceneRenderer::WaitForCleanUpTasks(GraphBuilder.RHICmdList);
 
 	UE::Tasks::FTask UpdateUniformExpressionsTask;
-	FMaterialRenderProxy::UpdateDeferredCachedUniformExpressions(GraphBuilder.RHICmdList, EnumHasAnyFlags(AsyncOps, EUpdateAllPrimitiveSceneInfosAsyncOps::CacheMaterialUniformExpressions) ? &UpdateUniformExpressionsTask : nullptr);
+	FMaterialRenderProxy::UpdateDeferredCachedUniformExpressions(GraphBuilder.RHICmdList, EnumHasAnyFlags(Parameters.AsyncOps, EUpdateAllPrimitiveSceneInfosAsyncOps::CacheMaterialUniformExpressions) ? &UpdateUniformExpressionsTask : nullptr);
 
 	RDG_EVENT_SCOPE(GraphBuilder, "UpdateAllPrimitiveSceneInfos");
 
@@ -5564,7 +5570,7 @@ void FScene::UpdateAllPrimitiveSceneInfos(FRDGBuilder& GraphBuilder, EUpdateAllP
 		{
 			InvalidatingPrimitiveCollector.UpdatedTransform(Transform.Key->GetPrimitiveSceneInfo());
 		}
-		
+
 		for (const auto& CullDistance : UpdatedInstanceCullDistance)
 		{
 			InvalidatingPrimitiveCollector.UpdatedTransform(CullDistance.Key->GetPrimitiveSceneInfo());
@@ -6348,44 +6354,53 @@ void FScene::UpdateAllPrimitiveSceneInfos(FRDGBuilder& GraphBuilder, EUpdateAllP
 
 	SceneExtensionsUpdaters.PostSceneUpdate(GraphBuilder, SceneUpdateChangeSetStorage.GetPostUpdateSet());
 
-	UpdateUniformExpressionsTask.Wait();
-
-	if (SceneInfosWithStaticDrawListUpdate.Num() > 0)
+	CreateLightPrimitiveInteractionsTask = GraphBuilder.AddSetupTask([this, &SceneInfosWithAddToScene]
 	{
-		FPrimitiveSceneInfo::AddStaticMeshes(this, SceneInfosWithStaticDrawListUpdate, false);
-	}
+		SCOPED_NAMED_EVENT(CreateLightPrimitiveInteractions, FColor::Emerald);
+
+		for (FPrimitiveSceneInfo* SceneInfo : SceneInfosWithAddToScene)
+		{
+			CreateLightPrimitiveInteractionsForPrimitive(SceneInfo);
+		}
 
-	FPrimitiveSceneInfo::UpdateVirtualTextures(this, SceneInfosWithAddToScene);
+	}, EnumHasAnyFlags(Parameters.AsyncOps, EUpdateAllPrimitiveSceneInfosAsyncOps::CreateLightPrimitiveInteractions));
 
-	UpdateReflectionSceneData(this);
+	const bool bAsyncCacheMeshDrawCommands = EnumHasAnyFlags(Parameters.AsyncOps, EUpdateAllPrimitiveSceneInfosAsyncOps::CacheMeshDrawCommands) && GRHISupportsMultithreadedShaderCreation;
 
-	if (bScenesPrimitivesNeedStaticMeshElementUpdate || CachedDefaultBasePassDepthStencilAccess != DefaultBasePassDepthStencilAccess)
+	UE::Tasks::FTask AddStaticMeshesTask = GraphBuilder.AddCommandListSetupTask(
+		[this, AddStaticMeshes = CopyTemp(SceneInfosWithStaticDrawListUpdate), SceneInfosWithFlushVirtualTexture = MoveTemp(SceneInfosWithFlushVirtualTexture), &SceneInfosWithAddToScene]
+			(FRHICommandListBase& RHICmdList) mutable
 	{
-		SceneInfosWithStaticDrawListUpdate.Reset();
+		SCOPED_NAMED_EVENT(StaticMeshUpdate, FColor::Emerald);
 
-		// Mark all primitives as needing an update
-		// Note: Only visible primitives will actually update their static mesh elements
-		for (int32 PrimitiveIndex = 0; PrimitiveIndex < Primitives.Num(); PrimitiveIndex++)
+		if (AddStaticMeshes.Num() > 0)
 		{
-			FPrimitiveSceneInfo* PrimitiveSceneInfo = Primitives[PrimitiveIndex];
-			PrimitivesNeedingStaticMeshUpdate[PrimitiveSceneInfo->PackedIndex] = true;
+			FPrimitiveSceneInfo::AddStaticMeshes(RHICmdList, this, AddStaticMeshes, false);
+		}
 
-			// HACK: Update Nanite primitives that need re-caching in GPU Scene
-			// TODO: Should be able to remove this after the move to compute materials.
-			if (bScenesPrimitivesNeedStaticMeshElementUpdate &&
-				PrimitiveSceneInfo->Proxy &&
-				PrimitiveSceneInfo->Proxy->IsNaniteMesh())
-			{
-				GPUScene.AddPrimitiveToUpdate(PrimitiveSceneInfo->GetPersistentIndex(), EPrimitiveDirtyState::ChangedOther);
-			}
+		FPrimitiveSceneInfo::UpdateVirtualTextures(this, SceneInfosWithAddToScene);
+
+		for (FPrimitiveSceneInfo* PrimitiveSceneInfo : SceneInfosWithFlushVirtualTexture)
+		{
+			PrimitiveSceneInfo->FlushRuntimeVirtualTexture();
+			PrimitiveSceneInfo->bPendingFlushVirtualTexture = false;
 		}
 
-		bScenesPrimitivesNeedStaticMeshElementUpdate = false;
-		CachedDefaultBasePassDepthStencilAccess = DefaultBasePassDepthStencilAccess;
-	}
+	}, UpdateUniformExpressionsTask, UE::Tasks::ETaskPriority::High, bAsyncCacheMeshDrawCommands);
+
+	UpdateReflectionSceneData(this);
 
 	{
 		SCOPED_NAMED_EVENT(UpdateStaticMeshes, FColor::Emerald);
+		
+		if (bScenesPrimitivesNeedStaticMeshElementUpdate || CachedDefaultBasePassDepthStencilAccess != DefaultBasePassDepthStencilAccess)
+		{
+			// Mark all primitives as needing an update
+			PrimitivesNeedingStaticMeshUpdate.Init(true, PrimitivesNeedingStaticMeshUpdate.Num());
+
+			bScenesPrimitivesNeedStaticMeshElementUpdate = false;
+			CachedDefaultBasePassDepthStencilAccess = DefaultBasePassDepthStencilAccess;
+		}
 
 		for (FPrimitiveSceneInfo* PrimitiveSceneInfo : SceneInfosWithStaticDrawListUpdate)
 		{
@@ -6407,79 +6422,57 @@ void FScene::UpdateAllPrimitiveSceneInfos(FRDGBuilder& GraphBuilder, EUpdateAllP
 		}
 	}
 
-	for (const auto& CustomParams : UpdatedCustomPrimitiveParams)
-	{
-		FPrimitiveSceneProxy* PrimitiveSceneProxy = CustomParams.Key;
-		if (DeletedPrimitiveSceneInfos.Contains(PrimitiveSceneProxy->GetPrimitiveSceneInfo()))
-		{
-			continue;
-		}
-
-		FScopeCycleCounter Context(PrimitiveSceneProxy->GetStatId());
-		PrimitiveSceneProxy->CustomPrimitiveData = CustomParams.Value;
-
-		// Ensure an update of primitive data before rendering
-		PrimitiveSceneProxy->GetPrimitiveSceneInfo()->MarkGPUStateDirty(EPrimitiveDirtyState::ChangedOther);
-	}
-
+	if (bScenesPrimitivesNeedStaticMeshElementUpdate)
 	{
-		SCOPED_NAMED_EVENT(UpdateUniformBuffers, FColor::Emerald);
-		TArray<FPrimitiveSceneProxy*, SceneRenderingAllocator> ProxiesToUpdate;
-
-		for (TConstSetBitIterator<> BitIt(PrimitivesNeedingUniformBufferUpdate); BitIt; ++BitIt)
-		{
-			const int32 Index = BitIt.GetIndex();
-			FPrimitiveSceneInfo* Primitive = Primitives[Index];
-			PrimitivesNeedingUniformBufferUpdate[Index] = false;
-			ProxiesToUpdate.Emplace(Primitive->Proxy);
-			GPUScene.AddPrimitiveToUpdate(Primitive->GetPersistentIndex(), EPrimitiveDirtyState::ChangedAll);
-		}
-
-		GraphBuilder.AddCommandListSetupTask([this, ProxiesToUpdate = MoveTemp(ProxiesToUpdate)](FRHICommandList& RHICmdList)
+		for (int32 PrimitiveIndex = 0; PrimitiveIndex < Primitives.Num(); PrimitiveIndex++)
 		{
-			SCOPED_NAMED_EVENT(AsyncUpdateUniformBuffers, FColor::Emerald);
-
-			for (FPrimitiveSceneProxy* Proxy : ProxiesToUpdate)
+			// HACK: Update Nanite primitives that need re-caching in GPU Scene
+			// TODO: Should be able to remove this after the move to compute materials.
+			if (PrimitiveSceneProxies[PrimitiveIndex] && PrimitiveSceneProxies[PrimitiveIndex]->IsNaniteMesh())
 			{
-				Proxy->UpdateUniformBuffer(RHICmdList);
+				GPUScene.AddPrimitiveToUpdate(Primitives[PrimitiveIndex]->GetPersistentIndex(), EPrimitiveDirtyState::ChangedOther);
 			}
-		});
-	}
-	
-	CreateLightPrimitiveInteractionsTask = GraphBuilder.AddSetupTask([this, &SceneInfosWithAddToScene]
-	{
-		SCOPED_NAMED_EVENT(CreateLightPrimitiveInteractions, FColor::Emerald);
-
-		for (FPrimitiveSceneInfo* SceneInfo : SceneInfosWithAddToScene)
-		{
-			CreateLightPrimitiveInteractionsForPrimitive(SceneInfo);
 		}
-
-	}, EnumHasAnyFlags(AsyncOps, EUpdateAllPrimitiveSceneInfosAsyncOps::CreateLightPrimitiveInteractions));
+	}
 
 	if (SceneInfosWithStaticDrawListUpdate.Num() > 0)
 	{
-		const bool bLaunchAsyncTask = EnumHasAnyFlags(AsyncOps, EUpdateAllPrimitiveSceneInfosAsyncOps::CacheMeshDrawCommands) && GRHISupportsMultithreadedShaderCreation;
-
-		CacheMeshDrawCommandsTask = GraphBuilder.AddSetupTask([this, &SceneInfosWithStaticDrawListUpdate, bLaunchAsyncTask]()
+		CacheMeshDrawCommandsTask = GraphBuilder.AddSetupTask([this, &SceneInfosWithStaticDrawListUpdate]
 		{
 			FPrimitiveSceneInfo::CacheMeshDrawCommands(this, SceneInfosWithStaticDrawListUpdate);
 
-		}, IsMobilePlatform(GetShaderPlatform()) ? CreateLightPrimitiveInteractionsTask : UE::Tasks::FTask(), UE::Tasks::ETaskPriority::High, bLaunchAsyncTask);
+		}, MakeArrayView({ AddStaticMeshesTask, IsMobilePlatform(GetShaderPlatform()) ? CreateLightPrimitiveInteractionsTask : UE::Tasks::FTask() }), UE::Tasks::ETaskPriority::Normal, bAsyncCacheMeshDrawCommands);
 
-		CacheNaniteMaterialBinsTask = GraphBuilder.AddSetupTask([this, &SceneInfosWithStaticDrawListUpdate, bLaunchAsyncTask]()
+		CacheNaniteMaterialBinsTask = GraphBuilder.AddSetupTask([this, &SceneInfosWithStaticDrawListUpdate]
 		{
 			FPrimitiveSceneInfo::CacheNaniteMaterialBins(this, SceneInfosWithStaticDrawListUpdate);
-		}, bLaunchAsyncTask);
+
+		}, AddStaticMeshesTask, UE::Tasks::ETaskPriority::Normal, bAsyncCacheMeshDrawCommands);
 
 #if RHI_RAYTRACING
-		CacheRayTracingPrimitivesTask = GraphBuilder.AddSetupTask([this, &SceneInfosWithStaticDrawListUpdate]()
+		CacheRayTracingPrimitivesTask = GraphBuilder.AddSetupTask([this, &SceneInfosWithStaticDrawListUpdate]
 		{
 			FPrimitiveSceneInfo::CacheRayTracingPrimitives(this, SceneInfosWithStaticDrawListUpdate);
-		}, bLaunchAsyncTask);
+
+		}, AddStaticMeshesTask, UE::Tasks::ETaskPriority::Normal, bAsyncCacheMeshDrawCommands);
 #endif
 	}
 
+	for (const auto& CustomParams : UpdatedCustomPrimitiveParams)
+	{
+		FPrimitiveSceneProxy* PrimitiveSceneProxy = CustomParams.Key;
+		if (DeletedPrimitiveSceneInfos.Contains(PrimitiveSceneProxy->GetPrimitiveSceneInfo()))
+		{
+			continue;
+		}
+
+		FScopeCycleCounter Context(PrimitiveSceneProxy->GetStatId());
+		PrimitiveSceneProxy->CustomPrimitiveData = CustomParams.Value;
+
+		// Ensure an update of primitive data before rendering
+		PrimitiveSceneProxy->GetPrimitiveSceneInfo()->MarkGPUStateDirty(EPrimitiveDirtyState::ChangedOther);
+	}
+
 	if (auto NaniteMaterialsUpdater = SceneExtensionsUpdaters.GetUpdaterPtr<Nanite::FMaterialsSceneExtension::FUpdater>())
 	{
 		NaniteMaterialsUpdater->PostCacheNaniteMaterialBins(GraphBuilder, SceneInfosWithStaticDrawListUpdate);
@@ -6494,12 +6487,6 @@ void FScene::UpdateAllPrimitiveSceneInfos(FRDGBuilder& GraphBuilder, EUpdateAllP
 		SceneLODHierarchy.UpdateNodeSceneInfo(PrimitiveSceneInfo->PrimitiveComponentId, PrimitiveSceneInfo);
 	}
 
-	for (FPrimitiveSceneInfo* PrimitiveSceneInfo : SceneInfosWithFlushVirtualTexture)
-	{
-		PrimitiveSceneInfo->FlushRuntimeVirtualTexture();
-		PrimitiveSceneInfo->bPendingFlushVirtualTexture = false;
-	}
-
 	for (const auto& Attachments : UpdatedAttachmentRoots)
 	{
 		FPrimitiveSceneInfo* PrimitiveSceneInfo = Attachments.Key;
@@ -6596,6 +6583,35 @@ void FScene::UpdateAllPrimitiveSceneInfos(FRDGBuilder& GraphBuilder, EUpdateAllP
 
 		DistanceFieldSceneData.UpdatePrimitive(SceneInfo);
 	}
+
+	if (Parameters.Callbacks.PostStaticMeshUpdate)
+	{
+		Parameters.Callbacks.PostStaticMeshUpdate(AddStaticMeshesTask);
+	}
+
+	{
+		SCOPED_NAMED_EVENT(UpdateUniformBuffers, FColor::Emerald);
+		TArray<FPrimitiveSceneProxy*, SceneRenderingAllocator> ProxiesToUpdate;
+
+		for (TConstSetBitIterator<> BitIt(PrimitivesNeedingUniformBufferUpdate); BitIt; ++BitIt)
+		{
+			const int32 Index = BitIt.GetIndex();
+			FPrimitiveSceneInfo* Primitive = Primitives[Index];
+			PrimitivesNeedingUniformBufferUpdate[Index] = false;
+			ProxiesToUpdate.Emplace(Primitive->Proxy);
+			GPUScene.AddPrimitiveToUpdate(Primitive->GetPersistentIndex(), EPrimitiveDirtyState::ChangedAll);
+		}
+
+		GraphBuilder.AddCommandListSetupTask([this, ProxiesToUpdate = MoveTemp(ProxiesToUpdate)](FRHICommandList& RHICmdList)
+		{
+			SCOPED_NAMED_EVENT(AsyncUpdateUniformBuffers, FColor::Emerald);
+
+			for (FPrimitiveSceneProxy* Proxy : ProxiesToUpdate)
+			{
+				Proxy->UpdateUniformBuffer(RHICmdList);
+			}
+		});
+	}
 	
 	{
 		RDG_CSV_STAT_EXCLUSIVE_SCOPE(GraphBuilder, UpdateGPUScene);
@@ -6603,8 +6619,7 @@ void FScene::UpdateAllPrimitiveSceneInfos(FRDGBuilder& GraphBuilder, EUpdateAllP
 
 		FRDGExternalAccessQueue ExternalAccessQueue;
 
-		// Process GPU scene prior to visibility to maximize overlap.
-		GPUScene.Update(GraphBuilder, SceneUB, ExternalAccessQueue, nullptr);
+		GPUScene.Update(GraphBuilder, SceneUB, ExternalAccessQueue, Parameters.GPUSceneUpdateTaskPrerequisites);
 	
 		ExternalAccessQueue.Submit(GraphBuilder);
 	}
@@ -6632,6 +6647,8 @@ void FScene::UpdateAllPrimitiveSceneInfos(FRDGBuilder& GraphBuilder, EUpdateAllP
 		}
 	});
 
+	AddStaticMeshesTask.Wait();
+
 	if (bNeedPathTracedInvalidation)
 	{
 		InvalidatePathTracedOutput();
diff --git a/Engine/Source/Runtime/Renderer/Private/ScenePrivate.h b/Engine/Source/Runtime/Renderer/Private/ScenePrivate.h
--- a/Engine/Source/Runtime/Renderer/Private/ScenePrivate.h
+++ b/Engine/Source/Runtime/Renderer/Private/ScenePrivate.h
@@ -3224,6 +3224,20 @@ public:
 
 	using FSceneInterface::UpdateAllPrimitiveSceneInfos;
 
+	struct FUpdateParameters
+	{
+		EUpdateAllPrimitiveSceneInfosAsyncOps AsyncOps = EUpdateAllPrimitiveSceneInfosAsyncOps::None;
+		UE::Tasks::FTask GPUSceneUpdateTaskPrerequisites;
+
+		struct
+		{
+			TFunction<void(const UE::Tasks::FTask&)> PostStaticMeshUpdate;
+
+		} Callbacks;
+	};
+
+	void Update(FRDGBuilder& GraphBuilder, const FUpdateParameters& Parameters);
+
 	// FSceneInterface interface.
 	virtual void AddPrimitive(UPrimitiveComponent* Primitive) override;
 	virtual void RemovePrimitive(UPrimitiveComponent* Primitive) override;
diff --git a/Engine/Source/Runtime/Renderer/Private/SceneRendering.cpp b/Engine/Source/Runtime/Renderer/Private/SceneRendering.cpp
--- a/Engine/Source/Runtime/Renderer/Private/SceneRendering.cpp
+++ b/Engine/Source/Runtime/Renderer/Private/SceneRendering.cpp
@@ -2934,6 +2934,7 @@ FSceneRenderer::ERendererOutput FSceneRenderer::GetRendererOutput() const
 void FSceneRenderer::PrepareViewRectsForRendering(FRHICommandListImmediate& RHICmdList)
 {
 	check(IsInRenderingThread());
+	TRACE_CPUPROFILER_EVENT_SCOPE(PrepareViewRectsForRendering);
 
 	// Read the resolution data.
 	{
@@ -3476,61 +3477,80 @@ IVisibilityTaskData* FSceneRenderer::OnRenderBegin(FRDGBuilder& GraphBuilder)
 		AsyncOps |= EUpdateAllPrimitiveSceneInfosAsyncOps::CacheMaterialUniformExpressions;
 	}
 
-	Scene->UpdateAllPrimitiveSceneInfos(GraphBuilder, AsyncOps);
+	IVisibilityTaskData* VisibilityTaskData = nullptr;
 
+	FScene::FUpdateParameters SceneUpdateParameters;
+	SceneUpdateParameters.AsyncOps = AsyncOps;
+
+	UE::Tasks::FTaskEvent GPUSceneUpdateTaskPrerequisites{ UE_SOURCE_LOCATION };
+	SceneUpdateParameters.GPUSceneUpdateTaskPrerequisites = GPUSceneUpdateTaskPrerequisites;
+
+	SceneUpdateParameters.Callbacks.PostStaticMeshUpdate = [&] (const UE::Tasks::FTask& StaticMeshUpdateTask)
+	{
 #if RHI_RAYTRACING
-	RayTracing::OnRenderBegin(*Scene, Views, ViewFamily);
+		RayTracing::OnRenderBegin(*Scene, Views, ViewFamily);
 #endif
 
-	if (!ViewFamily.ViewExtensions.IsEmpty())
-	{
+		if (!ViewFamily.ViewExtensions.IsEmpty())
 		{
-			RDG_CSV_STAT_EXCLUSIVE_SCOPE(GraphBuilder, PreRender);
-			SCOPE_CYCLE_COUNTER(STAT_FDeferredShadingSceneRenderer_ViewExtensionPreRenderView);
-
-			for (int32 ViewExt = 0; ViewExt < ViewFamily.ViewExtensions.Num(); ViewExt++)
 			{
-				ViewFamily.ViewExtensions[ViewExt]->PreRenderViewFamily_RenderThread(GraphBuilder, ViewFamily);
-				for (int32 ViewIndex = 0; ViewIndex < AllViews.Num(); ViewIndex++)
+				RDG_CSV_STAT_EXCLUSIVE_SCOPE(GraphBuilder, PreRender);
+				SCOPE_CYCLE_COUNTER(STAT_FDeferredShadingSceneRenderer_ViewExtensionPreRenderView);
+	
+				for (int32 ViewExt = 0; ViewExt < ViewFamily.ViewExtensions.Num(); ViewExt++)
 				{
-					ViewFamily.ViewExtensions[ViewExt]->PreRenderView_RenderThread(GraphBuilder, *AllViews[ViewIndex]);
+					ViewFamily.ViewExtensions[ViewExt]->PreRenderViewFamily_RenderThread(GraphBuilder, ViewFamily);
+					for (int32 ViewIndex = 0; ViewIndex < AllViews.Num(); ViewIndex++)
+					{
+						ViewFamily.ViewExtensions[ViewExt]->PreRenderView_RenderThread(GraphBuilder, *AllViews[ViewIndex]);
+					}
 				}
 			}
 		}
-	}
-
-	PrepareViewRectsForRendering(GraphBuilder.RHICmdList);
-
-	InitializeSceneTexturesConfig(ViewFamily.SceneTexturesConfig, ViewFamily);
-	FSceneTexturesConfig& SceneTexturesConfig = GetActiveSceneTexturesConfig();
-	FSceneTexturesConfig::Set(SceneTexturesConfig);
-
-	PrepareViewStateForVisibility(SceneTexturesConfig);
-
-	// Run Groom LOD selection prior to visibility for selecting appropriate LOD & geometry type
-	if (IsGroomEnabled())
-	{
-		if (Views.Num() > 0 && !ViewFamily.EngineShowFlags.HitProxies)
+	
+		PrepareViewRectsForRendering(GraphBuilder.RHICmdList);
+	
+		InitializeSceneTexturesConfig(ViewFamily.SceneTexturesConfig, ViewFamily);
+		FSceneTexturesConfig& SceneTexturesConfig = GetActiveSceneTexturesConfig();
+		FSceneTexturesConfig::Set(SceneTexturesConfig);
+	
+		PrepareViewStateForVisibility(SceneTexturesConfig);
+	
+		// Run Groom LOD selection prior to visibility for selecting appropriate LOD & geometry type
+		if (IsGroomEnabled())
 		{
-			FHairStrandsBookmarkParameters Parameters;
-			CreateHairStrandsBookmarkParameters(Scene, Views, AllFamilyViews, Parameters, false/*bComputeVisibleInstances*/);
-			if (Parameters.HasInstances())
+			if (Views.Num() > 0 && !ViewFamily.EngineShowFlags.HitProxies)
 			{
-				Scene->WaitForGPUSkinCacheTask();
-
-				// 1. Select appropriate LOD & geometry type
-				RunHairStrandsBookmark(GraphBuilder, EHairStrandsBookmark::ProcessLODSelection, Parameters);
+				FHairStrandsBookmarkParameters Parameters;
+				CreateHairStrandsBookmarkParameters(Scene, Views, AllFamilyViews, Parameters, false/*bComputeVisibleInstances*/);
+				if (Parameters.HasInstances())
+				{
+					Scene->WaitForGPUSkinCacheTask();
+	
+					// 1. Select appropriate LOD & geometry type
+					RunHairStrandsBookmark(GraphBuilder, EHairStrandsBookmark::ProcessLODSelection, Parameters);
+				}
 			}
 		}
-	}
 	
-	LightFunctionAtlas::OnRenderBegin(LightFunctionAtlas, *Scene, Views, ViewFamily);
+		LightFunctionAtlas::OnRenderBegin(LightFunctionAtlas, *Scene, Views, ViewFamily);
+	
+		FVisualizeTexturePresent::OnStartRender(Views[0]);
+	
+		GraphBuilder.RHICmdList.BeginScene();
+
+		VisibilityTaskData = LaunchVisibilityTasks(GraphBuilder.RHICmdList, *this, StaticMeshUpdateTask);
 
-	FVisualizeTexturePresent::OnStartRender(Views[0]);
+		if (GraphBuilder.IsParallelSetupEnabled())
+		{
+			GPUSceneUpdateTaskPrerequisites.AddPrerequisites(VisibilityTaskData->GetComputeRelevanceTask());
+		}
+		GPUSceneUpdateTaskPrerequisites.Trigger();
+	};
 
-	GraphBuilder.RHICmdList.BeginScene();
+	Scene->Update(GraphBuilder, SceneUpdateParameters);
 
-	return LaunchVisibilityTasks(GraphBuilder.RHICmdList, *this);
+	return VisibilityTaskData;
 }
 
 /** 
diff --git a/Engine/Source/Runtime/Renderer/Private/SceneVisibility.cpp b/Engine/Source/Runtime/Renderer/Private/SceneVisibility.cpp
--- a/Engine/Source/Runtime/Renderer/Private/SceneVisibility.cpp
+++ b/Engine/Source/Runtime/Renderer/Private/SceneVisibility.cpp
@@ -386,10 +386,10 @@ TRACE_DECLARE_INT_COUNTER(Scene_Visibility_Relevance_NumPrimitivesProcessed, TEX
 
 ///////////////////////////////////////////////////////////////////////////////
 
-IVisibilityTaskData* LaunchVisibilityTasks(FRHICommandListImmediate& RHICmdList, FSceneRenderer& SceneRenderer)
+IVisibilityTaskData* LaunchVisibilityTasks(FRHICommandListImmediate& RHICmdList, FSceneRenderer& SceneRenderer, const UE::Tasks::FTask& BeginInitVisibilityPrerequisites)
 {
 	FVisibilityTaskData* TaskData = SceneRenderer.Allocator.Create<FVisibilityTaskData>(RHICmdList, SceneRenderer);
-	TaskData->LaunchVisibilityTasks();
+	TaskData->LaunchVisibilityTasks(BeginInitVisibilityPrerequisites);
 	return TaskData;
 }
 
@@ -3139,8 +3139,8 @@ FVisibilityTaskConfig::FVisibilityTaskConfig(const FScene& Scene, TConstArrayVie
 
 	// These values tune the task granularity based on number of primitives in the scene and the number of worker tasks available.
 	const uint32 NumFrustumCullTasksPerThread   = 2;
-	const uint32 NumOcclusionCullTasksPerThread = 1;
-	const uint32 NumRelevanceTasksPerThread     = 16;
+	const uint32 NumOcclusionCullTasksPerThread = 2;
+	const uint32 NumRelevanceTasksPerThread     = 32;
 
 	// Frustum Cull
 	{
@@ -3910,7 +3910,7 @@ FVisibilityTaskData::FVisibilityTaskData(FRHICommandListImmediate& InRHICmdList,
 	Tasks.bWaitingAllowed = TaskConfig.Schedule == EVisibilityTaskSchedule::Parallel;
 }
 
-void FVisibilityTaskData::LaunchVisibilityTasks()
+void FVisibilityTaskData::LaunchVisibilityTasks(const UE::Tasks::FTask& BeginInitVisibilityPrerequisites)
 {
 	SCOPED_NAMED_EVENT(LaunchVisibilityTasks, FColor::Magenta);
 	ViewPackets.Reserve(Views.Num());
@@ -3925,6 +3925,12 @@ void FVisibilityTaskData::LaunchVisibilityTasks()
 		Tasks.FrustumCull.AddPrerequisites(ViewPacket.Tasks.FrustumCull);
 		Tasks.OcclusionCull.AddPrerequisites(ViewPacket.Tasks.OcclusionCull);
 		Tasks.ComputeRelevance.AddPrerequisites(ViewPacket.Tasks.ComputeRelevance);
+
+		if (ViewPacket.ViewState)
+		{
+			SCOPE_CYCLE_COUNTER(STAT_DecompressPrecomputedOcclusion);
+			ViewPacket.View.PrecomputedVisibilityData = ViewPacket.ViewState->GetPrecomputedVisibilityData(ViewPacket.View, &Scene);
+		}
 	}
 
 	// Each relevance task should have this as a prerequisite, but in case there aren't any tasks we make it explicit.
@@ -4050,7 +4056,7 @@ void FVisibilityTaskData::LaunchVisibilityTasks()
 				FOptionalTaskTagScope Scope(ETaskTag::EParallelRenderingThread);
 				ViewPacket.BeginInitVisibility();
 
-			}, TaskConfig.TaskPriority));
+			}, BeginInitVisibilityPrerequisites, TaskConfig.TaskPriority));
 		}
 
 		// Static relevance is finalized for ALL views after each view completes static mesh filtering tasks.
@@ -4541,6 +4547,8 @@ void FSceneRenderer::PreVisibilityFrameSetup(FRDGBuilder& GraphBuilder)
 
 void FSceneRenderer::PrepareViewStateForVisibility(const FSceneTexturesConfig& SceneTexturesConfig)
 {
+	TRACE_CPUPROFILER_EVENT_SCOPE(PrepareViewStateForVisibility);
+
 #if UE_BUILD_SHIPPING
 	const bool bFreezeTemporalHistories = false;
 	const bool bFreezeTemporalSequences = false;
diff --git a/Engine/Source/Runtime/Renderer/Private/SceneVisibility.h b/Engine/Source/Runtime/Renderer/Private/SceneVisibility.h
--- a/Engine/Source/Runtime/Renderer/Private/SceneVisibility.h
+++ b/Engine/Source/Runtime/Renderer/Private/SceneVisibility.h
@@ -67,4 +67,4 @@ public:
 	//////////////////////////////////////////////////////////////////////////////
 };
 
-extern IVisibilityTaskData* LaunchVisibilityTasks(FRHICommandListImmediate& RHICmdList, FSceneRenderer& SceneRenderer);
\ No newline at end of file
+extern IVisibilityTaskData* LaunchVisibilityTasks(FRHICommandListImmediate& RHICmdList, FSceneRenderer& SceneRenderer, const UE::Tasks::FTask& BeginInitVisibilityTaskPrerequisites);
\ No newline at end of file
diff --git a/Engine/Source/Runtime/Renderer/Private/SceneVisibilityPrivate.h b/Engine/Source/Runtime/Renderer/Private/SceneVisibilityPrivate.h
--- a/Engine/Source/Runtime/Renderer/Private/SceneVisibilityPrivate.h
+++ b/Engine/Source/Runtime/Renderer/Private/SceneVisibilityPrivate.h
@@ -492,7 +492,7 @@ public:
 		check(bFinished);
 	}
 
-	void LaunchVisibilityTasks();
+	void LaunchVisibilityTasks(const UE::Tasks::FTask& BeginInitVisibilityPrerequisites);
 
 	void ProcessRenderThreadTasks() override;
 
diff --git a/Engine/Source/Runtime/Renderer/Public/PrimitiveSceneInfo.h b/Engine/Source/Runtime/Renderer/Public/PrimitiveSceneInfo.h
--- a/Engine/Source/Runtime/Renderer/Public/PrimitiveSceneInfo.h
+++ b/Engine/Source/Runtime/Renderer/Public/PrimitiveSceneInfo.h
@@ -414,7 +414,7 @@ public:
 	RENDERER_API bool RequestUniformBufferUpdate();
 
 	/** Adds the primitive's static meshes to the scene. */
-	static void AddStaticMeshes(FScene* Scene, TArrayView<FPrimitiveSceneInfo*> SceneInfos, bool bCacheMeshDrawCommands = true);
+	static void AddStaticMeshes(FRHICommandListBase& RHICmdList, FScene* Scene, TArrayView<FPrimitiveSceneInfo*> SceneInfos, bool bCacheMeshDrawCommands = true);
 
 	/** Removes the primitive's static meshes from the scene. */
 	void RemoveStaticMeshes();
