Fixes RHI assertion failure trying to create a SRV on backbuffer in AfterTonemapper post process material

diff --git a/Engine/Source/Runtime/Renderer/Private/PostProcess/PostProcessMaterial.cpp b/Engine/Source/Runtime/Renderer/Private/PostProcess/PostProcessMaterial.cpp
--- a/Engine/Source/Runtime/Renderer/Private/PostProcess/PostProcessMaterial.cpp
+++ b/Engine/Source/Runtime/Renderer/Private/PostProcess/PostProcessMaterial.cpp
@@ -1008,7 +1008,12 @@ FScreenPassTexture AddPostProcessMaterialChain(
 		}
 
 		Outputs = AddPostProcessMaterialPass(GraphBuilder, View, Inputs, MaterialInterface);
-		CurrentInput = FScreenPassTextureSlice::CreateFromScreenPassTexture(GraphBuilder, Outputs);
+
+		// Don't create the CurrentInput out of Outputs of the last material as this could possibly be the back buffer for AfterTonemap post process material
+		if (MaterialInterface != Materials.Last())
+		{
+			CurrentInput = FScreenPassTextureSlice::CreateFromScreenPassTexture(GraphBuilder, Outputs);
+		}
 	}
 
 	if (!Outputs.IsValid())
