Fix OIT sorted triangle index buffer crashing on creation due to data not being streamed-in at creation time.

When the data are not ready yet, this CL postpone sorted index buffer initialization later, until the data are ready.

diff --git a/Engine/Source/Runtime/Engine/Public/MeshBatch.h b/Engine/Source/Runtime/Engine/Public/MeshBatch.h
--- a/Engine/Source/Runtime/Engine/Public/MeshBatch.h
+++ b/Engine/Source/Runtime/Engine/Public/MeshBatch.h
@@ -77,7 +77,7 @@ struct FMeshBatchElementDynamicIndexBuffer
 	/** Returns true if the allocation is valid. */
 	FORCEINLINE bool IsValid() const
 	{
-		return IndexBuffer != NULL;
+		return IndexBuffer != NULL && IndexBuffer->IsInitialized();
 	}
 };
 
diff --git a/Engine/Source/Runtime/Renderer/Private/OIT/OIT.cpp b/Engine/Source/Runtime/Renderer/Private/OIT/OIT.cpp
--- a/Engine/Source/Runtime/Renderer/Private/OIT/OIT.cpp
+++ b/Engine/Source/Runtime/Renderer/Private/OIT/OIT.cpp
@@ -250,23 +250,28 @@ class FSortedIndexBuffer : public FIndexBuffer
 public:
 	static const uint32 SliceCount = 256u;
 
-	FSortedIndexBuffer(uint32 InId, const FBufferRHIRef& InSourceIndexBuffer, uint32 InNumIndices, const TCHAR* InDebugName)
+	FSortedIndexBuffer(uint32 InId, const FIndexBuffer* InSourceIndexBuffer, uint32 InNumIndices, const TCHAR* InDebugName)
 	: SourceIndexBuffer(InSourceIndexBuffer)
 	, NumIndices(InNumIndices)
 	, Id(InId)
 	, DebugName(InDebugName) { }
 
+	bool CanBeInitialized() const
+	{
+		return SourceIndexBuffer && SourceIndexBuffer->IndexBufferRHI && SourceIndexBuffer->IndexBufferRHI->GetStride() > 0;
+	}
+
 	virtual void InitRHI(FRHICommandListBase& RHICmdList) override
 	{
-		check(SourceIndexBuffer);
-		const uint32 BytesPerElement = SourceIndexBuffer->GetStride();
+		check(SourceIndexBuffer && SourceIndexBuffer->IndexBufferRHI);
+		const uint32 BytesPerElement = SourceIndexBuffer->IndexBufferRHI->GetStride();
 		check(BytesPerElement == 2 || BytesPerElement == 4);
 		const EPixelFormat Format = BytesPerElement == 2 ? PF_R16_UINT : PF_R32_UINT;
 
 		FRHIResourceCreateInfo CreateInfo(DebugName);
 		IndexBufferRHI = RHICmdList.CreateBuffer(NumIndices * BytesPerElement, BUF_UnorderedAccess | BUF_ShaderResource | BUF_IndexBuffer, BytesPerElement /*Stride*/, ERHIAccess::VertexOrIndexBuffer, CreateInfo);
 		SortedIndexUAV = RHICmdList.CreateUnorderedAccessView(IndexBufferRHI, Format);
-		SourceIndexSRV = RHICmdList.CreateShaderResourceView(SourceIndexBuffer, BytesPerElement, Format);
+		SourceIndexSRV = RHICmdList.CreateShaderResourceView(SourceIndexBuffer->IndexBufferRHI, BytesPerElement, Format);
 	}
 
 	virtual void ReleaseRHI() override
@@ -278,7 +283,7 @@ public:
 
 	static constexpr uint32 InvalidId = ~0;
 
-	FBufferRHIRef SourceIndexBuffer = nullptr;
+	const FIndexBuffer* SourceIndexBuffer = nullptr;
 	uint32 NumIndices = 0;
 	uint32 Id = FSortedIndexBuffer::InvalidId;
 	uint32 LastUsedFrameId = 0;
@@ -311,7 +316,11 @@ static void TrimSortedIndexBuffers(TArray<FSortedIndexBuffer*>& FreeBuffers, TQu
 	}
 }
 
-FSortedTriangleData FOITSceneData::Allocate(FRHICommandListBase& RHICmdList, EPrimitiveType InPrimitiveType, const FMeshBatchElement& InMeshElement)
+void FOITSceneData::Allocate(
+	FRHICommandListBase& RHICmdList, 
+	EPrimitiveType InPrimitiveType, 
+	const FMeshBatchElement& InMeshElement, 
+	FMeshBatchElementDynamicIndexBuffer& OutMeshElement)
 {
 	check(InMeshElement.IndexBuffer && InMeshElement.IndexBuffer->IndexBufferRHI);
 	check(InPrimitiveType == PT_TriangleList || InPrimitiveType == PT_TriangleStrip);
@@ -338,7 +347,8 @@ FSortedTriangleData FOITSceneData::Allocate(FRHICommandListBase& RHICmdList, EPr
 		for (uint32 FreeIt=0,FreeCount=FreeBuffers.Num(); FreeIt<FreeCount; ++FreeIt)
 		{		
 			FSortedIndexBuffer* FreeBuffer = FreeBuffers[FreeIt];
-			if (FreeBuffer != nullptr && FreeBuffer->NumIndices >= NumIndices && FreeBuffer->Id == FSortedIndexBuffer::InvalidId)
+			// TODO: check for format as well + more robust comparison
+			if (FreeBuffer != nullptr && FreeBuffer->NumIndices >= NumIndices && FreeBuffer->Id == FSortedIndexBuffer::InvalidId) 
 			{			
 				OITIndexBuffer = FreeBuffer;
 				OITIndexBuffer->Id = FreeSlot;
@@ -352,8 +362,13 @@ FSortedTriangleData FOITSceneData::Allocate(FRHICommandListBase& RHICmdList, EPr
 	// Otherwise create a new one
 	if (OITIndexBuffer == nullptr)
 	{
-		OITIndexBuffer = new FSortedIndexBuffer(FreeSlot, InMeshElement.IndexBuffer->IndexBufferRHI, NumIndices, TEXT("OIT::SortedIndexBuffer"));
-		OITIndexBuffer->InitResource(RHICmdList);	
+		OITIndexBuffer = new FSortedIndexBuffer(FreeSlot, InMeshElement.IndexBuffer, NumIndices, TEXT("OIT::SortedIndexBuffer"));
+
+		// It's possible the index buffer isn't ready yet (data not streamed-in yet). In such case we post-pone OITIndexBuffer creation.
+		if (OITIndexBuffer->CanBeInitialized())
+		{
+			OITIndexBuffer->InitResource(RHICmdList);
+		}
 	}
 	Out->NumPrimitives = InMeshElement.NumPrimitives;
 	Out->NumIndices = NumIndices;
@@ -366,39 +381,34 @@ FSortedTriangleData FOITSceneData::Allocate(FRHICommandListBase& RHICmdList, EPr
 	Out->SortedPrimitiveType = PT_TriangleList;
 	Out->SourceIndexBuffer = InMeshElement.IndexBuffer;
 	Out->SortedIndexBuffer = OITIndexBuffer;
-	Out->SortedIndexUAV = OITIndexBuffer->SortedIndexUAV;
-	Out->SourceIndexSRV = OITIndexBuffer->SourceIndexSRV;
 
-	return *Out;
+	OutMeshElement.IndexBuffer 	= Out->SortedIndexBuffer;
+	OutMeshElement.FirstIndex 	= Out->SortedFirstIndex;
+	OutMeshElement.PrimitiveType= Out->SortedPrimitiveType;
 }
 
-void FOITSceneData::Deallocate(FIndexBuffer* InIndexBuffer)
+void FOITSceneData::Deallocate(FMeshBatchElement& OutMeshElement)
 {
-	if (InIndexBuffer == nullptr)
-	{
-		return;
-	}
-
-	FSortedIndexBuffer* OITIndexBuffer = (FSortedIndexBuffer*)InIndexBuffer;
-	const uint32 Slot = OITIndexBuffer->Id;
-	if (Slot < uint32(Allocations.Num()))
+	if (FSortedIndexBuffer* OITIndexBuffer = (FSortedIndexBuffer*)OutMeshElement.DynamicIndexBuffer.IndexBuffer)
 	{
-		if (CVarOIT_SortedTriangles_Pool.GetValueOnAnyThread() > 0)
+		const uint32 Slot = OITIndexBuffer->Id;
+		if (Slot < uint32(Allocations.Num()))
 		{
-			OITIndexBuffer->Id = FSortedIndexBuffer::InvalidId;
-			OITIndexBuffer->LastUsedFrameId = FrameIndex;
-			FreeBuffers.Add(OITIndexBuffer);
-			Allocations[Slot] = FSortedTriangleData();
-		}
-		else
-		{
-			FSortedTriangleData& In = Allocations[Slot];
-			In.SortedIndexUAV = nullptr;
-			In.SourceIndexSRV = nullptr;
-			PendingDeletes.Enqueue((FSortedIndexBuffer*)In.SortedIndexBuffer);
-			In = FSortedTriangleData();
+			if (CVarOIT_SortedTriangles_Pool.GetValueOnAnyThread() > 0)
+			{
+				OITIndexBuffer->Id = FSortedIndexBuffer::InvalidId;
+				OITIndexBuffer->LastUsedFrameId = FrameIndex;
+				FreeBuffers.Add(OITIndexBuffer);
+				Allocations[Slot] = FSortedTriangleData();
+			}
+			else
+			{
+				FSortedTriangleData& In = Allocations[Slot];
+				PendingDeletes.Enqueue((FSortedIndexBuffer*)In.SortedIndexBuffer);
+				In = FSortedTriangleData();
+			}
+			FreeSlots.Enqueue(Slot);
 		}
-		FreeSlots.Enqueue(Slot);
 	}
 }
 
@@ -613,6 +623,18 @@ static void AddOITSortTriangleIndexPass(
 		return;
 	}
 
+	// If the index buffer hasn't been initialized yet (e.g., data are not streamed in yet), run initialization
+	FSortedIndexBuffer* SortedIndexBuffer = (FSortedIndexBuffer*)MeshBatch.Mesh->Elements[0].DynamicIndexBuffer.IndexBuffer;
+	if (!SortedIndexBuffer->IsInitialized())
+	{
+		if (!SortedIndexBuffer->CanBeInitialized())
+		{
+			// Data are still not ready yet
+			return;
+		}
+		SortedIndexBuffer->InitResource(GraphBuilder.RHICmdList);
+	}
+
 	const FLocalVertexFactory* VF = (const FLocalVertexFactory*)MeshBatch.Mesh->VertexFactory;
 	if (!VF) { return; }
 	const FShaderResourceViewRHIRef VertexPosition = VF->GetPositionsSRV();
@@ -659,8 +681,8 @@ static void AddOITSortTriangleIndexPass(
 		Parameters->SourceBaseVertexIndex	= Allocation.SourceBaseVertexIndex;
 		Parameters->SourceMinVertexIndex	= Allocation.SourceMinVertexIndex;
 		Parameters->SourceMaxVertexIndex	= Allocation.SourceMaxVertexIndex;
-		Parameters->IndexBuffer				= Allocation.SourceIndexSRV;
-		Parameters->OutIndexBuffer			= Allocation.SortedIndexUAV;
+		Parameters->IndexBuffer				= Allocation.SortedIndexBuffer->SourceIndexSRV;
+		Parameters->OutIndexBuffer			= Allocation.SortedIndexBuffer->SortedIndexUAV;
 		Parameters->OutSliceCounterBuffer	= SliceCounterUAV;
 		Parameters->OutPrimitiveSliceBuffer = GraphBuilder.CreateUAV(PrimitiveSliceBuffer, PackedFormat);
 
@@ -727,8 +749,8 @@ static void AddOITSortTriangleIndexPass(
 		Parameters->SliceOffsetsBuffer		= GraphBuilder.CreateSRV(SliceOffsetsBuffer, PF_R32_UINT);
 		Parameters->PrimitiveSliceBuffer	= GraphBuilder.CreateSRV(PrimitiveSliceBuffer, PackedFormat);
 
-		Parameters->IndexBuffer				= Allocation.SourceIndexSRV;
-		Parameters->OutIndexBuffer			= Allocation.SortedIndexUAV;
+		Parameters->IndexBuffer				= Allocation.SortedIndexBuffer->SourceIndexSRV;
+		Parameters->OutIndexBuffer			= Allocation.SortedIndexBuffer->SortedIndexUAV;
 
 		TShaderMapRef<FOITSortTriangleIndex_WriteOutCS> ComputeShader(View.ShaderMap);
 
@@ -891,14 +913,6 @@ namespace OIT
 		}
 	}
 
-	void ConvertSortedIndexToDynamicIndex(FSortedTriangleData* In, FMeshBatchElementDynamicIndexBuffer* Out)
-	{
-		check(In && Out);
-		Out->IndexBuffer = In->SortedIndexBuffer;
-		Out->FirstIndex = In->SortedFirstIndex;
-		Out->PrimitiveType = In->SortedPrimitiveType;
-	}
-
 	FOITData CreateOITData(FRDGBuilder& GraphBuilder, const FViewInfo& View, EOITPassType PassType)
 	{
 		const bool bOIT = IsSortedPixelsEnabled(View);
@@ -975,9 +989,12 @@ namespace OIT
 		FSortedIndexBuffer* Buffer = nullptr;
 		while (OITSceneData.PendingDeletes.Dequeue(Buffer))
 		{
-			Buffer->ReleaseResource();
-			delete Buffer;
-			Buffer = nullptr;
+			if (Buffer)
+			{
+				Buffer->ReleaseResource();
+				delete Buffer;
+				Buffer = nullptr;
+			}
 		}
 	}
 }
\ No newline at end of file
diff --git a/Engine/Source/Runtime/Renderer/Private/OIT/OIT.h b/Engine/Source/Runtime/Renderer/Private/OIT/OIT.h
--- a/Engine/Source/Runtime/Renderer/Private/OIT/OIT.h
+++ b/Engine/Source/Runtime/Renderer/Private/OIT/OIT.h
@@ -30,10 +30,7 @@ enum class FTriangleSortingOrder
 struct FSortedTriangleData
 {
 	const FIndexBuffer* SourceIndexBuffer = nullptr;
-	FIndexBuffer* SortedIndexBuffer = nullptr;
-
-	FShaderResourceViewRHIRef  SourceIndexSRV = nullptr;
-	FUnorderedAccessViewRHIRef SortedIndexUAV = nullptr;
+	FSortedIndexBuffer* SortedIndexBuffer = nullptr;
 
 	uint32 SortedFirstIndex = 0;
 	uint32 SourceFirstIndex	= 0;
@@ -67,10 +64,10 @@ struct FOITData
 struct FOITSceneData
 {
 	/* Allocate sorted-triangle data for a instance */
-	FSortedTriangleData Allocate(FRHICommandListBase& RHICmdList, EPrimitiveType PrimitiveType, const FMeshBatchElement& InMeshElement);
+	void Allocate(FRHICommandListBase& RHICmdList, EPrimitiveType PrimitiveType, const FMeshBatchElement& InMeshElement, FMeshBatchElementDynamicIndexBuffer& OutMeshElement);
 
 	/* Deallocate sorted-triangle data */
-	void Deallocate(FIndexBuffer* IndexBuffer);
+	void Deallocate(FMeshBatchElement& OutMeshElement);
 
 	TArray<FSortedTriangleData> Allocations;
 	TArray<FSortedIndexBuffer*> FreeBuffers;
@@ -100,9 +97,6 @@ namespace OIT
 	/* Sort triangles of all instances whose has the sorted triangle option enabled */
 	void AddSortTrianglesPass(FRDGBuilder& GraphBuilder, const FViewInfo& View, FOITSceneData& OITSceneData, FTriangleSortingOrder SortType);
 
-	/* Convert FSortedTriangleData into FMeshBatchElementDynamicIndexBuffer */
-	void ConvertSortedIndexToDynamicIndex(FSortedTriangleData* In, FMeshBatchElementDynamicIndexBuffer* Out);
-
 	/* Create OIT data for translucent pass */
 	FOITData CreateOITData(FRDGBuilder& GraphBuilder, const FViewInfo& View, EOITPassType PassType);
 
diff --git a/Engine/Source/Runtime/Renderer/Private/PrimitiveSceneInfo.cpp b/Engine/Source/Runtime/Renderer/Private/PrimitiveSceneInfo.cpp
--- a/Engine/Source/Runtime/Renderer/Private/PrimitiveSceneInfo.cpp
+++ b/Engine/Source/Runtime/Renderer/Private/PrimitiveSceneInfo.cpp
@@ -1284,8 +1284,7 @@ void FPrimitiveSceneInfo::AddStaticMeshes(FRHICommandListBase& RHICmdList, FScen
 
 				if (bAllocateSortedTriangles && OIT::IsCompatible(Mesh, FeatureLevel))
 				{
-					FSortedTriangleData Allocation = Scene->OITSceneData.Allocate(RHICmdList, EPrimitiveType(Mesh.Type), Mesh.Elements[0]);
-					OIT::ConvertSortedIndexToDynamicIndex(&Allocation, &Mesh.Elements[0].DynamicIndexBuffer);
+					Scene->OITSceneData.Allocate(RHICmdList, EPrimitiveType(Mesh.Type), Mesh.Elements[0], Mesh.Elements[0].DynamicIndexBuffer);
 				}
 			}
 		}
@@ -1693,9 +1692,9 @@ void FPrimitiveSceneInfo::RemoveStaticMeshes()
 		for (int32 MeshIndex = 0; MeshIndex < StaticMeshes.Num(); MeshIndex++)
 		{
 			FStaticMeshBatch& Mesh = StaticMeshes[MeshIndex];
-			if (Mesh.Elements.Num() > 0 && Mesh.Elements[0].DynamicIndexBuffer.IsValid())
+			if (Mesh.Elements.Num() > 0)
 			{
-				Scene->OITSceneData.Deallocate(Mesh.Elements[0].DynamicIndexBuffer.IndexBuffer);
+				Scene->OITSceneData.Deallocate(Mesh.Elements[0]);
 			}
 		}
 	}
