Fix CIS non-unity build issue.

diff --git a/Engine/Source/Runtime/Engine/Public/MeshBatch.h b/Engine/Source/Runtime/Engine/Public/MeshBatch.h
--- a/Engine/Source/Runtime/Engine/Public/MeshBatch.h
+++ b/Engine/Source/Runtime/Engine/Public/MeshBatch.h
@@ -77,7 +77,7 @@ struct FMeshBatchElementDynamicIndexBuffer
 	/** Returns true if the allocation is valid. */
 	FORCEINLINE bool IsValid() const
 	{
-		return IndexBuffer != NULL && IndexBuffer->IsInitialized();
+		return IndexBuffer != NULL;
 	}
 };
 
diff --git a/Engine/Source/Runtime/Renderer/Private/MeshPassProcessor.cpp b/Engine/Source/Runtime/Renderer/Private/MeshPassProcessor.cpp
--- a/Engine/Source/Runtime/Renderer/Private/MeshPassProcessor.cpp
+++ b/Engine/Source/Runtime/Renderer/Private/MeshPassProcessor.cpp
@@ -1060,7 +1060,7 @@ void FMeshDrawCommand::SetDrawParametersAndFinalize(
 	NumInstances = BatchElement.NumInstances;
 
 	// If the mesh batch has a valid dynamic index buffer, use it instead
-	if (BatchElement.DynamicIndexBuffer.IsValid())
+	if (BatchElement.DynamicIndexBuffer.IsValid() &&  BatchElement.DynamicIndexBuffer.IndexBuffer->IsInitialized())
 	{
 		check(!BatchElement.DynamicIndexBuffer.IndexBuffer || (BatchElement.DynamicIndexBuffer.IndexBuffer && BatchElement.DynamicIndexBuffer.IndexBuffer->IsInitialized() && BatchElement.DynamicIndexBuffer.IndexBuffer->IndexBufferRHI));
 		IndexBuffer = BatchElement.DynamicIndexBuffer.IndexBuffer ? BatchElement.DynamicIndexBuffer.IndexBuffer->IndexBufferRHI.GetReference() : nullptr;
