Adds eighth resolution post processing support with r.PostProcessing.QuarterResolutionDownsample=2

This combined with r.Bloom.ScreenPercentage=12.5 allows to run the entire post processing chain at one eighth of the resolution, which TSR with sg.AntiAliasingQuality < 3 can produce.

diff --git a/Engine/Shaders/Private/TemporalSuperResolution/TSRUpdateHistory.usf b/Engine/Shaders/Private/TemporalSuperResolution/TSRUpdateHistory.usf
--- a/Engine/Shaders/Private/TemporalSuperResolution/TSRUpdateHistory.usf
+++ b/Engine/Shaders/Private/TemporalSuperResolution/TSRUpdateHistory.usf
@@ -108,6 +108,7 @@ float ResurrectionFrameIndex;
 float PrevFrameIndex;
 uint bGenerateOutputMip1;
 uint bGenerateOutputMip2;
+uint bGenerateOutputMip3;
 uint bHasSeparateTranslucency;
 uint HistoryArrayIndices_HighFrequency;
 
@@ -315,7 +316,7 @@ T BilinearSampleColorHistory(Texture2D<T> Texture, float2 SampleUV)
 }
 #endif
 
-tsr_halfCx2 DownsampleSceneColor(tsr_halfCx2 SceneColor, uint LocalGroupThreadIndex, const uint XorButterFly)
+tsr_halfCx2 DownsampleSceneColor(tsr_halfCx2 SceneColor, uint LocalGroupThreadIndex, const uint2 XorButterFly)
 {
 	SceneColor = SceneColor * tsr_half(0.25);
 	// Forces the * tsr_half(0.25) to be applied before to avoid turning bright pixels to +inf in the adds below.
@@ -328,14 +329,11 @@ tsr_halfCx2 DownsampleSceneColor(tsr_halfCx2 SceneColor, uint LocalGroupThreadIn
 		SceneColor = dpv_min(SceneColor, ColorMax10BitsFloat);
 	#endif
 
-	#if PLATFORM_SUPPORTS_WAVE_BROADCAST && CONFIG_ENABLE_DUAL_PIXEL_VECTORIZATION
-	if (uint(XorButterFly * 4) <= WaveGetLaneCount())
+	#if PLATFORM_SUPPORTS_WAVE_BROADCAST
+	if (uint(XorButterFly.x * 2) <= WaveGetLaneCount())
 	{
-		FWaveBroadcastSettings Horizontal = InitWaveXorButterfly(XorButterFly * (XorButterFly * 0x1));
-		FWaveBroadcastSettings Vertical   = InitWaveXorButterfly(XorButterFly * (XorButterFly * 0x2));
-
+		FWaveBroadcastSettings Horizontal = InitWaveXorButterfly(XorButterFly.x);
 		SceneColor += WaveBroadcast(Horizontal, SceneColor);
-		SceneColor += WaveBroadcast(Vertical, SceneColor);
 	}
 	else
 	#endif
@@ -352,59 +350,99 @@ tsr_halfCx2 DownsampleSceneColor(tsr_halfCx2 SceneColor, uint LocalGroupThreadIn
 			#endif
 		#endif
 		
-		if (XorButterFly * 0x2 <= 0x8)
-		{
-			// NOP
-		}
 		#if PLATFORM_SUPPORTS_WAVE_BROADCAST
-		else if (XorButterFly * 0x2 <= WaveGetLaneCount())
-		{
-			// NOP
-		}
+			const bool bBarrier = XorButterFly.x * 0x2 > WaveGetLaneCount();
+		#else
+			const bool bBarrier = XorButterFly.x * 0x2 > 0x8;
 		#endif
-		else
+
+		if (bBarrier)
 		{
 			GroupMemoryBarrierWithGroupSync();
 		}
-
-		SharedArray0[LocalGroupThreadIndex] += SharedArray0[LocalGroupThreadIndex ^ (XorButterFly * 0x1)];
-		#if CONFIG_ENABLE_DUAL_PIXEL_VECTORIZATION
-			SharedArray1[LocalGroupThreadIndex] += SharedArray1[LocalGroupThreadIndex ^ (XorButterFly * 0x1)];
+		
+		uint OtherGroupThreadIndex = LocalGroupThreadIndex ^ (XorButterFly.x % GROUP_SHARED_SIZE);
+		tsr_halfCx2 OtherSceneColor;
+		#if CONFIG_SCENE_COLOR_ALPHA
+			#if CONFIG_ENABLE_DUAL_PIXEL_VECTORIZATION
+				OtherSceneColor = dpv_interleave_registers(SharedArray0[OtherGroupThreadIndex], SharedArray1[OtherGroupThreadIndex]);
+			#else
+				OtherSceneColor = dpv_interleave_mono_registers(SharedArray0[OtherGroupThreadIndex]);
+			#endif
+		#else
+			#if CONFIG_ENABLE_DUAL_PIXEL_VECTORIZATION
+				OtherSceneColor = dpv_interleave_registers(SharedArray0[OtherGroupThreadIndex].rgb, SharedArray1[OtherGroupThreadIndex].rgb);
+			#else
+				OtherSceneColor = dpv_interleave_mono_registers(SharedArray0[OtherGroupThreadIndex].rgb);
+			#endif
 		#endif
-			
-		if (XorButterFly * 0x4 <= 0x8)
+		
+		if (bBarrier)
 		{
-			// NOP
+			GroupMemoryBarrierWithGroupSync();
 		}
+		
+		SceneColor += OtherSceneColor;
+	}
+	
+	if (XorButterFly.y == 0x20 && CONFIG_ENABLE_DUAL_PIXEL_VECTORIZATION)
+	{
+		SceneColor = dpv_interleave_mono_registers(dpv_lo(SceneColor) + dpv_hi(SceneColor));
+	}
+	#if PLATFORM_SUPPORTS_WAVE_BROADCAST
+	else if (uint(XorButterFly.y * 2) <= WaveGetLaneCount())
+	{
+		FWaveBroadcastSettings Vertical = InitWaveXorButterfly(XorButterFly.y);
+		SceneColor += WaveBroadcast(Vertical, SceneColor);
+	}
+	else
+	#endif
+	{
+		#if CONFIG_SCENE_COLOR_ALPHA
+			SharedArray0[LocalGroupThreadIndex] = dpv_lo(SceneColor);
+			#if CONFIG_ENABLE_DUAL_PIXEL_VECTORIZATION
+				SharedArray1[LocalGroupThreadIndex] = dpv_hi(SceneColor);
+			#endif
+		#else
+			SharedArray0[LocalGroupThreadIndex] = tsr_half4(dpv_lo(SceneColor), 0.0);
+			#if CONFIG_ENABLE_DUAL_PIXEL_VECTORIZATION
+				SharedArray1[LocalGroupThreadIndex] = tsr_half4(dpv_hi(SceneColor), 0.0);
+			#endif
+		#endif
+		
 		#if PLATFORM_SUPPORTS_WAVE_BROADCAST
-		else if (XorButterFly * 0x4 <= WaveGetLaneCount())
-		{
-			// NOP
-		}
+			const bool bBarrier = XorButterFly.y * 0x2 > WaveGetLaneCount();
+		#else
+			const bool bBarrier = XorButterFly.y * 0x2 > 0x8;
 		#endif
-		else
+		
+		if (bBarrier)
 		{
 			GroupMemoryBarrierWithGroupSync();
 		}
 		
-		SharedArray0[LocalGroupThreadIndex] += SharedArray0[LocalGroupThreadIndex ^ (XorButterFly * 0x2)];
-		#if CONFIG_ENABLE_DUAL_PIXEL_VECTORIZATION
-			SharedArray1[LocalGroupThreadIndex] += SharedArray1[LocalGroupThreadIndex ^ (XorButterFly * 0x2)];
-		#endif
-					
+		uint OtherGroupThreadIndex = LocalGroupThreadIndex ^ (XorButterFly.y % GROUP_SHARED_SIZE);
+		tsr_halfCx2 OtherSceneColor;
 		#if CONFIG_SCENE_COLOR_ALPHA
 			#if CONFIG_ENABLE_DUAL_PIXEL_VECTORIZATION
-				SceneColor = dpv_interleave_registers(SharedArray0[LocalGroupThreadIndex], SharedArray1[LocalGroupThreadIndex]);
+				OtherSceneColor = dpv_interleave_registers(SharedArray0[OtherGroupThreadIndex], SharedArray1[OtherGroupThreadIndex]);
 			#else
-				SceneColor = dpv_interleave_mono_registers(SharedArray0[LocalGroupThreadIndex]);
+				OtherSceneColor = dpv_interleave_mono_registers(SharedArray0[OtherGroupThreadIndex]);
 			#endif
 		#else
 			#if CONFIG_ENABLE_DUAL_PIXEL_VECTORIZATION
-				SceneColor = dpv_interleave_registers(SharedArray0[LocalGroupThreadIndex].rgb, SharedArray1[LocalGroupThreadIndex].rgb);
+				OtherSceneColor = dpv_interleave_registers(SharedArray0[OtherGroupThreadIndex].rgb, SharedArray1[OtherGroupThreadIndex].rgb);
 			#else
-				SceneColor = dpv_interleave_mono_registers(SharedArray0[LocalGroupThreadIndex].rgb);
+				OtherSceneColor = dpv_interleave_mono_registers(SharedArray0[OtherGroupThreadIndex].rgb);
 			#endif
 		#endif
+		
+		if (bBarrier)
+		{
+			GroupMemoryBarrierWithGroupSync();
+		}
+		
+		SceneColor += OtherSceneColor;
 	}
 	return SceneColor;
 }
@@ -1088,14 +1126,22 @@ void MainCS(
 			tsr_short2x2 HalfResOutputPixelPos = dpv_interleave_mono_registers(tsr_short(-1).xx);
 
 			BRANCH
-			if (bGenerateOutputMip1 || bGenerateOutputMip2)
+			if (bGenerateOutputMip1 || bGenerateOutputMip2 || bGenerateOutputMip3)
 			{
-				HalfResOutput = DownsampleSceneColor(HalfResOutput, LocalGroupThreadIndex, /* XorButterFly = */ 0x1);
-
+				HalfResOutput = DownsampleSceneColor(HalfResOutput, LocalGroupThreadIndex, /* XorButterFly = */ uint2(0x01, 0x02));
+				
 				BRANCH
-				if (bGenerateOutputMip2)
+				if (bGenerateOutputMip3)
+				{
+					HalfResOutput = DownsampleSceneColor(HalfResOutput, LocalGroupThreadIndex, /* XorButterFly = */ uint2(0x04, 0x10));
+					HalfResOutput = DownsampleSceneColor(HalfResOutput, LocalGroupThreadIndex, /* XorButterFly = */ uint2(0x08, 0x20));
+
+					HalfResOutputPixelPos[0] = (LocalHistoryPixelPos[0] >> tsr_short(3)) | (((LocalHistoryPixelPos[0] & tsr_short(0x7))) * tsr_short(~0));
+					HalfResOutputPixelPos[1] = (LocalHistoryPixelPos[1] >> tsr_short(3)) | (((LocalHistoryPixelPos[1] & tsr_short(0x7))) * tsr_short(~0));
+				}
+				else if (bGenerateOutputMip2)
 				{
-					HalfResOutput = DownsampleSceneColor(HalfResOutput, LocalGroupThreadIndex, /* XorButterFly = */ 0x4);
+					HalfResOutput = DownsampleSceneColor(HalfResOutput, LocalGroupThreadIndex, /* XorButterFly = */ uint2(0x04, 0x10));
 
 					HalfResOutputPixelPos[0] = (LocalHistoryPixelPos[0] >> tsr_short(2)) | (((LocalHistoryPixelPos[0] & tsr_short(0x3))) * tsr_short(~0));
 					HalfResOutputPixelPos[1] = (LocalHistoryPixelPos[1] >> tsr_short(2)) | (((LocalHistoryPixelPos[1] & tsr_short(0x3))) * tsr_short(~0));
diff --git a/Engine/Source/Runtime/Renderer/Private/PostProcess/PostProcessing.cpp b/Engine/Source/Runtime/Renderer/Private/PostProcess/PostProcessing.cpp
--- a/Engine/Source/Runtime/Renderer/Private/PostProcess/PostProcessing.cpp
+++ b/Engine/Source/Runtime/Renderer/Private/PostProcess/PostProcessing.cpp
@@ -193,11 +193,6 @@ bool IsPostProcessingEnabled(const FViewInfo& View)
 	}
 }
 
-bool IsPostProcessingQuarterResolutionDownsampleEnabled()
-{
-	return CVarPostProcessingQuarterResolutionDownsample.GetValueOnRenderThread() != 0;
-}
-
 bool IsPostProcessingWithAlphaChannelSupported()
 {
 	return CVarPostProcessingPropagateAlpha.GetValueOnAnyThread() != 0;
@@ -674,7 +669,8 @@ void AddPostProcessingPasses(
 		const bool bBasicEyeAdaptationEnabled = bEyeAdaptationEnabled && (AutoExposureMethod == EAutoExposureMethod::AEM_Basic);
 		const bool bLocalExposureBlurredLum = bLocalExposureEnabled && View.FinalPostProcessSettings.LocalExposureBlurredLuminanceBlend > 0.0f;
 
-		const bool bProcessQuarterResolution = IsPostProcessingQuarterResolutionDownsampleEnabled();
+		const bool bProcessQuarterResolution = CVarPostProcessingQuarterResolutionDownsample.GetValueOnRenderThread() == 1;
+		const bool bProcessEighthResolution = CVarPostProcessingQuarterResolutionDownsample.GetValueOnRenderThread() == 2;
 		const bool bMotionBlurNeedsHalfResInput = PassSequence.IsEnabled(EPass::MotionBlur) && DoesMotionBlurNeedsHalfResInput() && !bVisualizeMotionBlur;
 
 		const float FFTBloomResolutionFraction = GetFFTBloomResolutionFraction(PostTAAViewSize);
@@ -685,8 +681,9 @@ void AddPostProcessingPasses(
 			(bLensFlareEnabled && bFFTBloomEnabled) ||
 			bLocalExposureBlurredLum);
 		extern int32 GSSRHalfResSceneColor;
-		const bool bNeedBeforeBloomHalfRes    = !bProcessQuarterResolution || (bFFTBloomEnabled && FFTBloomResolutionFraction > 0.25f && FFTBloomResolutionFraction <= 0.5f) || (ReflectionsMethod == EReflectionsMethod::SSR && !View.bStatePrevViewInfoIsReadOnly && GSSRHalfResSceneColor);
-		const bool bNeedBeforeBloomQuarterRes = bProcessQuarterResolution || (bFFTBloomEnabled && FFTBloomResolutionFraction <= 0.25f);
+		const bool bNeedBeforeBloomHalfRes    = (!bProcessQuarterResolution && !bProcessEighthResolution) || (bFFTBloomEnabled && FFTBloomResolutionFraction > 0.25f && FFTBloomResolutionFraction <= 0.5f) || (ReflectionsMethod == EReflectionsMethod::SSR && !View.bStatePrevViewInfoIsReadOnly && GSSRHalfResSceneColor);
+		const bool bNeedBeforeBloomQuarterRes = bProcessQuarterResolution || (bFFTBloomEnabled && FFTBloomResolutionFraction > 0.125f && FFTBloomResolutionFraction <= 0.25f);
+		const bool bNeedBeforeBloomEighthRes  = bProcessEighthResolution || (bFFTBloomEnabled && FFTBloomResolutionFraction <= 0.125f);
 
 
 		const FPostProcessMaterialChain MaterialChainSceneColorBeforeDOF = GetPostProcessMaterialChain(View, BL_SceneColorBeforeDOF);
@@ -781,6 +778,7 @@ void AddPostProcessingPasses(
 
 		FScreenPassTextureSlice HalfResSceneColor;
 		FScreenPassTextureSlice QuarterResSceneColor;
+		FScreenPassTextureSlice EighthResSceneColor;
 		FVelocityFlattenTextures VelocityFlattenTextures;
 		if (TAAConfig != EMainTAAPassConfig::Disabled)
 		{
@@ -813,6 +811,10 @@ void AddPostProcessingPasses(
 				UpscalerPassInputs.bGenerateSceneColorQuarterRes =
 					bNeedBeforeBloomQuarterRes &&
 					DownsampleQuality == EDownsampleQuality::Low;
+
+				UpscalerPassInputs.bGenerateSceneColorEighthRes =
+					bNeedBeforeBloomEighthRes &&
+					DownsampleQuality == EDownsampleQuality::Low;
 			}
 			UpscalerPassInputs.bAllowFullResSlice = PassSequence.IsEnabled(EPass::MotionBlur) || PassSequence.IsEnabled(EPass::Tonemap);
 			UpscalerPassInputs.DownsampleOverrideFormat = DownsampleOverrideFormat;
@@ -850,6 +852,7 @@ void AddPostProcessingPasses(
 			SceneColorSlice = Outputs.FullRes;
 			HalfResSceneColor = Outputs.HalfRes;
 			QuarterResSceneColor = Outputs.QuarterRes;
+			EighthResSceneColor = Outputs.EighthRes;
 			VelocityFlattenTextures = Outputs.VelocityFlattenTextures;
 
 			if (PassSequence.IsEnabled(EPass::VisualizeTemporalUpscaler))
@@ -922,7 +925,7 @@ void AddPostProcessingPasses(
 			FMotionBlurInputs PassInputs;
 			PassSequence.AcceptOverrideIfLastPass(EPass::MotionBlur, PassInputs.OverrideOutput);
 			PassInputs.bOutputHalfRes = PostProcessMaterialBeforeBloomChain.Num() == 0 && bNeedBeforeBloomHalfRes && DownsampleQuality == EDownsampleQuality::Low;
-			PassInputs.bOutputQuarterRes = bNeedBeforeBloomQuarterRes && DownsampleQuality == EDownsampleQuality::Low;
+			PassInputs.bOutputQuarterRes = (bNeedBeforeBloomQuarterRes || bNeedBeforeBloomEighthRes) && DownsampleQuality == EDownsampleQuality::Low;
 			PassInputs.SceneColor = SceneColorSlice;
 			PassInputs.SceneDepth = SceneDepth;
 			PassInputs.SceneVelocity = Velocity;
@@ -965,6 +968,7 @@ void AddPostProcessingPasses(
 			{
 				HalfResSceneColor = FScreenPassTextureSlice();
 				QuarterResSceneColor = FScreenPassTextureSlice();
+				EighthResSceneColor = FScreenPassTextureSlice();
 			}
 
 			SceneColorSlice = NewSceneColorSlice;
@@ -983,7 +987,8 @@ void AddPostProcessingPasses(
 		// Generate before bloom lower res scene color if they have not been generated.
 		{
 			if ((bNeedBeforeBloomHalfRes && !HalfResSceneColor.IsValid()) ||
-				(bNeedBeforeBloomQuarterRes && !QuarterResSceneColor.IsValid() && !HalfResSceneColor.IsValid()))
+				(bNeedBeforeBloomQuarterRes && !QuarterResSceneColor.IsValid() && !HalfResSceneColor.IsValid()) ||
+				(bNeedBeforeBloomEighthRes && !EighthResSceneColor.IsValid() && !QuarterResSceneColor.IsValid() && !HalfResSceneColor.IsValid()))
 			{
 				FDownsamplePassInputs PassInputs;
 				PassInputs.Name = TEXT("PostProcessing.SceneColor.HalfRes");
@@ -994,7 +999,8 @@ void AddPostProcessingPasses(
 				HalfResSceneColor = FScreenPassTextureSlice::CreateFromScreenPassTexture(GraphBuilder, AddDownsamplePass(GraphBuilder, View, PassInputs));
 			}
 
-			if (bNeedBeforeBloomQuarterRes && !QuarterResSceneColor.IsValid())
+			if ((bNeedBeforeBloomQuarterRes && !QuarterResSceneColor.IsValid()) ||
+				(bNeedBeforeBloomEighthRes && !EighthResSceneColor.IsValid() && !QuarterResSceneColor.IsValid()))
 			{
 				FDownsamplePassInputs PassInputs;
 				PassInputs.Name = TEXT("PostProcessing.SceneColor.QuarterRes");
@@ -1003,6 +1009,16 @@ void AddPostProcessingPasses(
 
 				QuarterResSceneColor = FScreenPassTextureSlice::CreateFromScreenPassTexture(GraphBuilder, AddDownsamplePass(GraphBuilder, View, PassInputs));
 			}
+
+			if (bNeedBeforeBloomEighthRes && !EighthResSceneColor.IsValid())
+			{
+				FDownsamplePassInputs PassInputs;
+				PassInputs.Name = TEXT("PostProcessing.SceneColor.EighthRes");
+				PassInputs.SceneColor = QuarterResSceneColor;
+				PassInputs.Quality = DownsampleQuality;
+
+				EighthResSceneColor = FScreenPassTextureSlice::CreateFromScreenPassTexture(GraphBuilder, AddDownsamplePass(GraphBuilder, View, PassInputs));
+			}
 		}
 
 
@@ -1015,7 +1031,7 @@ void AddPostProcessingPasses(
 		}
 
 		{
-			FScreenPassTextureSlice LocalExposureSceneColor = bProcessQuarterResolution ? QuarterResSceneColor : HalfResSceneColor;
+			FScreenPassTextureSlice LocalExposureSceneColor = bProcessEighthResolution ? EighthResSceneColor : (bProcessQuarterResolution ? QuarterResSceneColor : HalfResSceneColor);
 
 			if (bLocalExposureEnabled)
 			{
@@ -1030,7 +1046,7 @@ void AddPostProcessingPasses(
 
 		if (bHistogramEnabled)
 		{
-			FScreenPassTextureSlice HistogramSceneColor = bProcessQuarterResolution ? QuarterResSceneColor : HalfResSceneColor;
+			FScreenPassTextureSlice HistogramSceneColor = bProcessEighthResolution ? EighthResSceneColor : (bProcessQuarterResolution ? QuarterResSceneColor : HalfResSceneColor);
 
 			if (IsAutoExposureUsingIlluminanceEnabled(View))
 			{
@@ -1061,14 +1077,14 @@ void AddPostProcessingPasses(
 			SceneDownsampleChain.Init(
 				GraphBuilder, View,
 				EyeAdaptationParameters,
-				bProcessQuarterResolution ? QuarterResSceneColor : HalfResSceneColor,
+				bProcessEighthResolution ? EighthResSceneColor : (bProcessQuarterResolution ? QuarterResSceneColor : HalfResSceneColor),
 				DownsampleChainQuality,
 				bLogLumaInAlpha);
 		}
 
 		if (bLocalExposureBlurredLum)
 		{
-			const uint32 BlurredLumMip = bProcessQuarterResolution ? 3 : 4;
+			const uint32 BlurredLumMip = bProcessEighthResolution ? 2 : (bProcessQuarterResolution ? 3 : 4);
 			LocalExposureBlurredLogLumTexture = AddLocalExposureBlurredLogLuminancePass(
 				GraphBuilder, View,
 				EyeAdaptationParameters, SceneDownsampleChain.GetTexture(BlurredLumMip));
@@ -1111,7 +1127,12 @@ void AddPostProcessingPasses(
 				float InputResolutionFraction;
 				FScreenPassTextureSlice InputSceneColor;
 
-				if (FFTBloomResolutionFraction <= 0.25f)
+				if (FFTBloomResolutionFraction <= 0.125f)
+				{
+					InputSceneColor = EighthResSceneColor;
+					InputResolutionFraction = 0.125f;
+				}
+				else if (FFTBloomResolutionFraction <= 0.25f)
 				{
 					InputSceneColor = QuarterResSceneColor;
 					InputResolutionFraction = 0.25f;
@@ -1152,7 +1173,7 @@ void AddPostProcessingPasses(
 				}
 				else
 				{
-					FScreenPassTextureSlice DownsampleInput = bProcessQuarterResolution ? QuarterResSceneColor : HalfResSceneColor;
+					FScreenPassTextureSlice DownsampleInput = bProcessEighthResolution ? EighthResSceneColor : (bProcessQuarterResolution ? QuarterResSceneColor : HalfResSceneColor);
 
 					if (bBloomSetupRequiredEnabled)
 					{
diff --git a/Engine/Source/Runtime/Renderer/Private/PostProcess/TemporalAA.h b/Engine/Source/Runtime/Renderer/Private/PostProcess/TemporalAA.h
--- a/Engine/Source/Runtime/Renderer/Private/PostProcess/TemporalAA.h
+++ b/Engine/Source/Runtime/Renderer/Private/PostProcess/TemporalAA.h
@@ -193,6 +193,7 @@ struct FDefaultTemporalUpscaler
 		bool bAllowFullResSlice = false;
 		bool bGenerateSceneColorHalfRes = false;
 		bool bGenerateSceneColorQuarterRes = false;
+		bool bGenerateSceneColorEighthRes = false;
 		bool bGenerateOutputMip1 = false;
 		bool bGenerateVelocityFlattenTextures = false;
 		EPixelFormat DownsampleOverrideFormat;
@@ -208,6 +209,7 @@ struct FDefaultTemporalUpscaler
 		FScreenPassTextureSlice FullRes;
 		FScreenPassTextureSlice HalfRes;
 		FScreenPassTextureSlice QuarterRes;
+		FScreenPassTextureSlice EighthRes;
 		FVelocityFlattenTextures VelocityFlattenTextures;
 	};
 };
diff --git a/Engine/Source/Runtime/Renderer/Private/PostProcess/TemporalSuperResolution.cpp b/Engine/Source/Runtime/Renderer/Private/PostProcess/TemporalSuperResolution.cpp
--- a/Engine/Source/Runtime/Renderer/Private/PostProcess/TemporalSuperResolution.cpp
+++ b/Engine/Source/Runtime/Renderer/Private/PostProcess/TemporalSuperResolution.cpp
@@ -866,6 +866,7 @@ class FTSRUpdateHistoryCS : public FTSRShader
 		SHADER_PARAMETER(float, PrevFrameIndex)
 		SHADER_PARAMETER(int32, bGenerateOutputMip1)
 		SHADER_PARAMETER(int32, bGenerateOutputMip2)
+		SHADER_PARAMETER(int32, bGenerateOutputMip3)
 		SHADER_PARAMETER(int32, bHasSeparateTranslucency)
 
 		SHADER_PARAMETER_STRUCT(FTSRHistoryArrayIndices, HistoryArrayIndices)
@@ -2202,6 +2203,7 @@ FDefaultTemporalUpscaler::FOutputs AddTemporalSuperResolutionPasses(
 	FRDGTextureSRVRef UpdateHistoryTextureSRV = nullptr;
 	FRDGTextureSRVRef SceneColorOutputHalfResTextureSRV = nullptr;
 	FRDGTextureSRVRef SceneColorOutputQuarterResTextureSRV = nullptr;
+	FRDGTextureSRVRef SceneColorOutputEighthResTextureSRV = nullptr;
 	{
 		static const TCHAR* const kUpdateQualityNames[] = {
 			TEXT("Low"),
@@ -2240,6 +2242,7 @@ FDefaultTemporalUpscaler::FOutputs AddTemporalSuperResolutionPasses(
 		PassParameters->InputContributionMultiplier = OutputToHistoryResolutionFractionSquare; 
 		PassParameters->bGenerateOutputMip1 = false;
 		PassParameters->bGenerateOutputMip2 = false;
+		PassParameters->bGenerateOutputMip3 = false;
 		PassParameters->bHasSeparateTranslucency = bHasSeparateTranslucency;
 
 		PassParameters->HistoryArrayIndices = HistoryArrayIndices;
@@ -2339,6 +2342,21 @@ FDefaultTemporalUpscaler::FOutputs AddTemporalSuperResolutionPasses(
 
 			SceneColorOutputQuarterResTextureSRV = GraphBuilder.CreateSRV(FRDGTextureSRVDesc::CreateForSlice(SceneColorOutputQuarterResTexture, /* SliceIndex = */ 0));
 		}
+		else if (PassInputs.bGenerateSceneColorEighthRes && HistorySize == OutputRect.Size())
+		{
+			FRDGTextureDesc QuarterResDesc = FRDGTextureDesc::Create2DArray(
+				FIntPoint::DivideAndRoundUp(OutputExtent, 8),
+				ColorFormat,
+				FClearValueBinding::None,
+				/* InFlags = */ TexCreate_ShaderResource | TexCreate_UAV,
+				/* ArraySize = */ 1);
+			FRDGTextureRef SceneColorOutputEighthResTexture = GraphBuilder.CreateTexture(QuarterResDesc, TEXT("TSR.EighthResOutput"));
+
+			PassParameters->bGenerateOutputMip3 = true;
+			PassParameters->SceneColorOutputMip1 = GraphBuilder.CreateUAV(FRDGTextureUAVDesc(SceneColorOutputEighthResTexture));
+
+			SceneColorOutputEighthResTextureSRV = GraphBuilder.CreateSRV(FRDGTextureSRVDesc::CreateForSlice(SceneColorOutputEighthResTexture, /* SliceIndex = */ 0));
+		}
 		else
 		{
 			PassParameters->SceneColorOutputMip1 = CreateDummyUAVArray(GraphBuilder, PF_FloatR11G11B10);
@@ -2359,7 +2377,7 @@ FDefaultTemporalUpscaler::FOutputs AddTemporalSuperResolutionPasses(
 				PermutationVector.Get<FTSRShader::F16BitVALUDim>() ? TEXT(" 16bit") : TEXT(""),
 				PermutationVector.Get<FTSRShader::FAlphaChannelDim>() ? TEXT(" AlphaChannel") : TEXT(""),
 				HistoryColorFormat == PF_FloatR11G11B10 ? TEXT(" R11G11B10") : TEXT(""),
-				PassParameters->bGenerateOutputMip2 ? TEXT(" OutputMip2") : (PassParameters->bGenerateOutputMip1 ? TEXT(" OutputMip1") : TEXT("")),
+				PassParameters->bGenerateOutputMip3 ? TEXT(" OutputMip3") : (PassParameters->bGenerateOutputMip2 ? TEXT(" OutputMip2") : (PassParameters->bGenerateOutputMip1 ? TEXT(" OutputMip1") : TEXT(""))),
 				HistorySize.X, HistorySize.Y),
 			AsyncComputePasses >= 3 ? ERDGPassFlags::AsyncCompute : ERDGPassFlags::Compute,
 			ComputeShader,
@@ -2409,7 +2427,7 @@ FDefaultTemporalUpscaler::FOutputs AddTemporalSuperResolutionPasses(
 			PassParameters->bGenerateOutputMip1 = true;
 			PassParameters->SceneColorOutputMip1 = GraphBuilder.CreateUAV(FRDGTextureUAVDesc(SceneColorOutputTexture, /* InMipLevel = */ 1));
 		}
-		else if (PassInputs.bGenerateSceneColorHalfRes)
+		else if (PassInputs.bGenerateSceneColorHalfRes || PassInputs.bGenerateSceneColorQuarterRes || PassInputs.bGenerateSceneColorEighthRes)
 		{
 			FRDGTextureDesc HalfResDesc = FRDGTextureDesc::Create2D(
 				OutputExtent / 2,
@@ -2438,11 +2456,12 @@ FDefaultTemporalUpscaler::FOutputs AddTemporalSuperResolutionPasses(
 		TShaderMapRef<FTSRResolveHistoryCS> ComputeShader(View.ShaderMap, PermutationVector);
 		FComputeShaderUtils::AddPass(
 			GraphBuilder,
-			RDG_EVENT_NAME("TSR ResolveHistory(#%d WaveSize=%d%s%s) %dx%d", //-V510
+			RDG_EVENT_NAME("TSR ResolveHistory(#%d WaveSize=%d%s%s%s) %dx%d", //-V510
 				PermutationVector.ToDimensionValueId(),
 				PermutationVector.Get<FTSRResolveHistoryCS::FNyquistDim>(),
 				PermutationVector.Get<FTSRShader::F16BitVALUDim>() ? TEXT(" 16bit") : TEXT(""),
 				PermutationVector.Get<FTSRShader::FAlphaChannelDim>() ? TEXT(" AlphaChannel") : TEXT(""),
+				PassParameters->bGenerateOutputMip1 ? TEXT(" OutputMip1") : TEXT(""),
 				OutputRect.Width(), OutputRect.Height()),
 			AsyncComputePasses >= 3 ? ERDGPassFlags::AsyncCompute : ERDGPassFlags::Compute,
 			ComputeShader,
@@ -2672,6 +2691,12 @@ FDefaultTemporalUpscaler::FOutputs AddTemporalSuperResolutionPasses(
 		Outputs.QuarterRes.ViewRect.Min = OutputRect.Min / 4;
 		Outputs.QuarterRes.ViewRect.Max = Outputs.HalfRes.ViewRect.Min + FIntPoint::DivideAndRoundUp(OutputRect.Size(), 4);
 	}
+	if (SceneColorOutputEighthResTextureSRV)
+	{
+		Outputs.EighthRes.TextureSRV = SceneColorOutputEighthResTextureSRV;
+		Outputs.EighthRes.ViewRect.Min = FIntPoint::DivideAndRoundUp(OutputRect.Min, 8);
+		Outputs.EighthRes.ViewRect.Max = Outputs.EighthRes.ViewRect.Min + FIntPoint::DivideAndRoundUp(OutputRect.Size(), 8);
+	}
 	Outputs.VelocityFlattenTextures = VelocityFlattenTextures;
 	return Outputs;
 } // AddTemporalSuperResolutionPasses()
diff --git a/Engine/Source/Runtime/Renderer/Private/SceneVisibility.cpp b/Engine/Source/Runtime/Renderer/Private/SceneVisibility.cpp
--- a/Engine/Source/Runtime/Renderer/Private/SceneVisibility.cpp
+++ b/Engine/Source/Runtime/Renderer/Private/SceneVisibility.cpp
@@ -4684,11 +4684,13 @@ void FSceneRenderer::PrepareViewStateForVisibility(const FSceneTexturesConfig& S
 				// other random signal that are power of two of View.StateFrameIndex
 				if (TAAConfig == EMainTAAPassConfig::TSR)
 				{
-					static const int8 kFirstPrimeNumbers[25] = {
+					static const uint8 kFirstPrimeNumbers[] = {
 						2, 3, 5, 7, 11, 13, 17, 19, 23, 29, 31, 37, 41, 43, 47, 53, 59, 61, 67, 71, 73, 79, 83, 89, 97,
+						101, 103, 107, 109, 113, 127, 131, 137, 139, 149, 151, 157, 163, 167, 173, 179, 181, 191, 193, 197, 199,
+						211, 223, 227, 229, 233, 239, 241, 251,
 					};
 
-					for (int32 PrimeNumberId = 4; PrimeNumberId < UE_ARRAY_COUNT(kFirstPrimeNumbers); PrimeNumberId++)
+					for (int32 PrimeNumberId = FMath::Max(4, (TemporalAASamples - 1) / 5); PrimeNumberId < UE_ARRAY_COUNT(kFirstPrimeNumbers); PrimeNumberId++)
 					{
 						if (int32(kFirstPrimeNumbers[PrimeNumberId]) >= TemporalAASamples)
 						{
