Fix DoubleFloat issues affecting Heterogeneous Volumes

diff --git a/Engine/Shaders/Private/HeterogeneousVolumes/HeterogeneousVolumesLiveShadingPipeline.usf b/Engine/Shaders/Private/HeterogeneousVolumes/HeterogeneousVolumesLiveShadingPipeline.usf
--- a/Engine/Shaders/Private/HeterogeneousVolumes/HeterogeneousVolumesLiveShadingPipeline.usf
+++ b/Engine/Shaders/Private/HeterogeneousVolumes/HeterogeneousVolumesLiveShadingPipeline.usf
@@ -176,6 +176,9 @@ FVolumeSampleContext CreateVolumeSampleContext(float3 LocalPosition, float3 Worl
 	//VolumeSampleContext.MaterialParameters.PrimitiveId = PrimitiveId;
 	VolumeSampleContext.MaterialParameters.AbsoluteWorldPosition = DFPromote(WorldPosition);
 	VolumeSampleContext.MaterialParameters.LWCData.AbsoluteWorldPosition = WSPromote(WorldPosition);
+	// TODO: Add object centroid to LWC.ObjectWorldPosition
+	VolumeSampleContext.MaterialParameters.LWCData.LocalToWorld = WSPromote(GetLocalToWorld());
+	VolumeSampleContext.MaterialParameters.LWCData.WorldToLocal = WSPromoteInverse(GetWorldToLocal());
 
 	// Evaluate Material graph
 	CalcPixelMaterialInputs(VolumeSampleContext.MaterialParameters, VolumeSampleContext.PixelMaterialInputs);
diff --git a/Engine/Shaders/Private/HeterogeneousVolumes/HeterogeneousVolumesMaterialBakingPipeline.usf b/Engine/Shaders/Private/HeterogeneousVolumes/HeterogeneousVolumesMaterialBakingPipeline.usf
--- a/Engine/Shaders/Private/HeterogeneousVolumes/HeterogeneousVolumesMaterialBakingPipeline.usf
+++ b/Engine/Shaders/Private/HeterogeneousVolumes/HeterogeneousVolumesMaterialBakingPipeline.usf
@@ -62,6 +62,9 @@ void HeterogeneousVolumesBakeMaterialCS(
 		MaterialParameters.PrimitiveId = PrimitiveId;
 		MaterialParameters.AbsoluteWorldPosition = DFPromote(WorldPosition);
 		MaterialParameters.LWCData.AbsoluteWorldPosition = WSPromote(WorldPosition);
+		// TODO: Add object centroid to LWC.ObjectWorldPosition
+		MaterialParameters.LWCData.LocalToWorld = WSPromote(LocalToWorld);
+		MaterialParameters.LWCData.WorldToLocal = WSPromoteInverse(WorldToLocal);
 
 		// Evaluate material graph
 		FPixelMaterialInputs PixelMaterialInputs;
@@ -265,6 +268,9 @@ void RasterizeBottomLevelOrthoGridCS(
 				MaterialParameters.PrimitiveId = PrimitiveId;
 				MaterialParameters.AbsoluteWorldPosition = DFPromote(WorldPosition);
 				MaterialParameters.LWCData.AbsoluteWorldPosition = WSPromote(WorldPosition);
+				// TODO: Add object centroid to LWC.ObjectWorldPosition
+				MaterialParameters.LWCData.LocalToWorld = WSPromote(LocalToWorld);
+				MaterialParameters.LWCData.WorldToLocal = WSPromoteInverse(WorldToLocal);
 
 				// Evaluate material graph
 				FPixelMaterialInputs PixelMaterialInputs;
@@ -471,6 +477,9 @@ void RasterizeBottomLevelFrustumGridCS(
 		MaterialParameters.PrimitiveId = PrimitiveId;
 		MaterialParameters.AbsoluteWorldPosition = DFPromote(WorldPosition);
 		MaterialParameters.LWCData.AbsoluteWorldPosition = WSPromote(WorldPosition);
+		// TODO: Add object centroid to LWC.ObjectWorldPosition
+		MaterialParameters.LWCData.LocalToWorld = WSPromote(LocalToWorld);
+		MaterialParameters.LWCData.WorldToLocal = WSPromoteInverse(WorldToLocal);
 
 		// Evaluate material graph
 		FPixelMaterialInputs PixelMaterialInputs;
diff --git a/Engine/Source/Runtime/Renderer/Private/HeterogeneousVolumes/HeterogeneousVolumesLiveShadingPipeline.cpp b/Engine/Source/Runtime/Renderer/Private/HeterogeneousVolumes/HeterogeneousVolumesLiveShadingPipeline.cpp
--- a/Engine/Source/Runtime/Renderer/Private/HeterogeneousVolumes/HeterogeneousVolumesLiveShadingPipeline.cpp
+++ b/Engine/Source/Runtime/Renderer/Private/HeterogeneousVolumes/HeterogeneousVolumesLiveShadingPipeline.cpp
@@ -426,10 +426,12 @@ static void RenderLightingCacheWithLiveShading(
 		PassParameters->VolumetricScatteringIntensity = LightSceneInfo->Proxy->GetVolumetricScatteringIntensity();
 
 		// Object data
-		FVector3f ViewOriginHigh = FDFVector3(View.ViewMatrices.GetViewOrigin()).High;
-		FMatrix44f RelativeLocalToWorld = FDFMatrix::MakeToRelativeWorldMatrix(ViewOriginHigh, HeterogeneousVolumeInterface->GetLocalToWorld()).M;
-		PassParameters->LocalToWorld = RelativeLocalToWorld;
-		PassParameters->WorldToLocal = RelativeLocalToWorld.Inverse();
+		// TODO: Convert to relative-local space
+		//FVector3f ViewOriginHigh = FDFVector3(View.ViewMatrices.GetViewOrigin()).High;
+		//FMatrix44f RelativeLocalToWorld = FDFMatrix::MakeToRelativeWorldMatrix(ViewOriginHigh, HeterogeneousVolumeInterface->GetLocalToWorld()).M;
+		FMatrix44f LocalToWorld = FMatrix44f(HeterogeneousVolumeInterface->GetLocalToWorld());
+		PassParameters->LocalToWorld = LocalToWorld;
+		PassParameters->WorldToLocal = LocalToWorld.Inverse();
 		PassParameters->LocalBoundsOrigin = FVector3f(LocalBoxSphereBounds.Origin);
 		PassParameters->LocalBoundsExtent = FVector3f(LocalBoxSphereBounds.BoxExtent);
 		PassParameters->PrimitiveId = PersistentPrimitiveIndex.Index;
@@ -577,10 +579,11 @@ static void RenderSingleScatteringWithLiveShading(
 		PassParameters->ShadowStepFactor = HeterogeneousVolumeInterface->GetShadowStepFactor();
 
 		// Object data
-		FVector3f ViewOriginHigh = FDFVector3(View.ViewMatrices.GetViewOrigin()).High;
-		FMatrix44f RelativeLocalToWorld = FDFMatrix::MakeToRelativeWorldMatrix(ViewOriginHigh, HeterogeneousVolumeInterface->GetLocalToWorld()).M;
-		PassParameters->LocalToWorld = RelativeLocalToWorld;
-		PassParameters->WorldToLocal = RelativeLocalToWorld.Inverse();
+		// TODO: Convert to relative-local space
+		//FVector3f ViewOriginHigh = FDFVector3(View.ViewMatrices.GetViewOrigin()).High;
+		//FMatrix44f RelativeLocalToWorld = FDFMatrix::MakeToRelativeWorldMatrix(ViewOriginHigh, HeterogeneousVolumeInterface->GetLocalToWorld()).M;
+		PassParameters->LocalToWorld = FMatrix44f(HeterogeneousVolumeInterface->GetLocalToWorld());
+		PassParameters->WorldToLocal = FMatrix44f(HeterogeneousVolumeInterface->GetLocalToWorld().Inverse());
 		PassParameters->LocalBoundsOrigin = FVector3f(LocalBoxSphereBounds.Origin);
 		PassParameters->LocalBoundsExtent = FVector3f(LocalBoxSphereBounds.BoxExtent);
 		PassParameters->PrimitiveId = PersistentPrimitiveIndex.Index;
diff --git a/Engine/Source/Runtime/Renderer/Private/HeterogeneousVolumes/HeterogeneousVolumesMaterialBakingPipeline.cpp b/Engine/Source/Runtime/Renderer/Private/HeterogeneousVolumes/HeterogeneousVolumesMaterialBakingPipeline.cpp
--- a/Engine/Source/Runtime/Renderer/Private/HeterogeneousVolumes/HeterogeneousVolumesMaterialBakingPipeline.cpp
+++ b/Engine/Source/Runtime/Renderer/Private/HeterogeneousVolumes/HeterogeneousVolumesMaterialBakingPipeline.cpp
@@ -138,10 +138,12 @@ void ComputeHeterogeneousVolumeBakeMaterial(
 		PassParameters->Scene = View.GetSceneUniforms().GetBuffer(GraphBuilder);
 
 		// Object data
+		// TODO: Convert to relative-local space
 		FVector3f ViewOriginHigh = FDFVector3(View.ViewMatrices.GetViewOrigin()).High;
 		FMatrix44f RelativeLocalToWorld = FDFMatrix::MakeToRelativeWorldMatrix(ViewOriginHigh, HeterogeneousVolumeInterface->GetLocalToWorld()).M;
-		PassParameters->LocalToWorld = RelativeLocalToWorld;
-		PassParameters->WorldToLocal = RelativeLocalToWorld.Inverse();
+		FMatrix44f LocalToWorld = FMatrix44f(HeterogeneousVolumeInterface->GetLocalToWorld());
+		PassParameters->LocalToWorld = LocalToWorld;
+		PassParameters->WorldToLocal = LocalToWorld.Inverse();
 		PassParameters->LocalBoundsOrigin = FVector3f(LocalBoxSphereBounds.Origin);
 		PassParameters->LocalBoundsExtent = FVector3f(LocalBoxSphereBounds.BoxExtent);
 		PassParameters->PrimitiveId = PersistentPrimitiveIndex.Index;
diff --git a/Engine/Source/Runtime/Renderer/Private/HeterogeneousVolumes/HeterogeneousVolumesPreshadingPipeline.cpp b/Engine/Source/Runtime/Renderer/Private/HeterogeneousVolumes/HeterogeneousVolumesPreshadingPipeline.cpp
--- a/Engine/Source/Runtime/Renderer/Private/HeterogeneousVolumes/HeterogeneousVolumesPreshadingPipeline.cpp
+++ b/Engine/Source/Runtime/Renderer/Private/HeterogeneousVolumes/HeterogeneousVolumesPreshadingPipeline.cpp
@@ -1593,10 +1593,12 @@ void RenderWithPreshading(
 	FSparseVoxelUniformBufferParameters* SparseVoxelUniformBufferParameters = GraphBuilder.AllocParameters<FSparseVoxelUniformBufferParameters>();
 	{
 		// Object data
-		FVector3f ViewOriginHigh = FDFVector3(View.ViewMatrices.GetViewOrigin()).High;
-		FMatrix44f RelativeLocalToWorld = FDFMatrix::MakeToRelativeWorldMatrix(ViewOriginHigh, HeterogeneousVolumeInterface->GetLocalToWorld()).M;
-		SparseVoxelUniformBufferParameters->LocalToWorld = RelativeLocalToWorld;
-		SparseVoxelUniformBufferParameters->WorldToLocal = RelativeLocalToWorld.Inverse();
+		// TODO: Convert to relative-local space
+		//FVector3f ViewOriginHigh = FDFVector3(View.ViewMatrices.GetViewOrigin()).High;
+		//FMatrix44f RelativeLocalToWorld = FDFMatrix::MakeToRelativeWorldMatrix(ViewOriginHigh, HeterogeneousVolumeInterface->GetLocalToWorld()).M;
+		FMatrix44f LocalToWorld = FMatrix44f(HeterogeneousVolumeInterface->GetLocalToWorld());
+		SparseVoxelUniformBufferParameters->LocalToWorld = LocalToWorld;
+		SparseVoxelUniformBufferParameters->WorldToLocal = LocalToWorld.Inverse();
 		SparseVoxelUniformBufferParameters->LocalBoundsOrigin = FVector3f(LocalBoxSphereBounds.Origin);
 		SparseVoxelUniformBufferParameters->LocalBoundsExtent = FVector3f(LocalBoxSphereBounds.BoxExtent);
 
