Mark all CSM PSOs as required during precache to fix possible runtime hitches

diff --git a/Engine/Source/Runtime/Renderer/Private/ShadowDepthRendering.cpp b/Engine/Source/Runtime/Renderer/Private/ShadowDepthRendering.cpp
--- a/Engine/Source/Runtime/Renderer/Private/ShadowDepthRendering.cpp
+++ b/Engine/Source/Runtime/Renderer/Private/ShadowDepthRendering.cpp
@@ -2256,9 +2256,6 @@ void FShadowDepthPassMeshProcessor::CollectPSOInitializersForEachShadowDepthType
 		CollectPSOInitializersForEachStreamSetup(VertexFactoryData, Material, LocalShadowDepthType, MeshFillMode, FinalCullMode, bRequired, PSOInitializers);
 	}
 
-	// Only mark directional shadows are required to wait for proxy creation
-	bRequired = false;
-
 	// Collect for non-directional one pass point light shadows
 	{
 		LocalShadowDepthType.bDirectionalLight = false;
