Refactored GatherRayTracingRelevantPrimitives to use ParallelFor
- new ParallelFor loop builds array of indices of relevant static/dynamic primitives
- modified GatherRayTracingRelevantPrimitives_ComputeLOD to loop over indices and setup FRelevantPrimitive of each static primitive
- relevant dynamic primitives don't use FRelevantPrimitive at all since most of the fields are not relevant for dynamic primitives (only index is relevant)
- using ParallelFor is not compatible with the optimization that skipped ERayTracingPrimitiveFlags::UnsupportedProxyType using the TypeOffsetTable however ParallelFor should be faster in most projects (for example CitySample has fewer than 100 unsupported proxies)

Performance:
- City Sample (console): 0.72ms -> 0.34ms
- Electric Dreams Sample (PC with 64 cores): 2.4ms -> 0.15ms

diff --git a/Engine/Source/Runtime/Renderer/Private/RayTracing/RayTracing.cpp b/Engine/Source/Runtime/Renderer/Private/RayTracing/RayTracing.cpp
--- a/Engine/Source/Runtime/Renderer/Private/RayTracing/RayTracing.cpp
+++ b/Engine/Source/Runtime/Renderer/Private/RayTracing/RayTracing.cpp
@@ -150,15 +150,14 @@ namespace RayTracing
 		FPersistentPrimitiveIndex PersistentPrimitiveIndex;
 		int8 LODIndex = -1;
 		uint8 InstanceMask = 0;
-		bool bStatic = false;
-		bool bAllSegmentsOpaque = true;
-		bool bAllSegmentsCastShadow = true;
-		bool bAnySegmentsCastShadow = false;
-		bool bAnySegmentsDecal = false;
-		bool bAllSegmentsDecal = true;
-		bool bTwoSided = false;
-		bool bIsSky = false;
-		bool bAllSegmentsTranslucent = true;
+		bool bAllSegmentsOpaque : 1 = true;
+		bool bAllSegmentsCastShadow : 1 = true;
+		bool bAnySegmentsCastShadow : 1 = false;
+		bool bAnySegmentsDecal : 1 = false;
+		bool bAllSegmentsDecal : 1 = true;
+		bool bTwoSided : 1 = false;
+		bool bIsSky : 1 = false;
+		bool bAllSegmentsTranslucent : 1 = true;
 
 		const FRayTracingGeometryInstance* CachedRayTracingInstance = nullptr;
 		TArrayView<const int32> CachedRayTracingMeshCommandIndices; // Pointer to FPrimitiveSceneInfo::CachedRayTracingMeshCommandIndicesPerLOD data
@@ -198,8 +197,8 @@ namespace RayTracing
 	struct FRelevantPrimitiveList
 	{
 		// Filtered lists of relevant primitives
-		TChunkedArray<FRelevantPrimitive> StaticPrimitives;
-		TChunkedArray<FRelevantPrimitive> DynamicPrimitives;
+		TArray<FRelevantPrimitive> StaticPrimitives;
+		TArray<int32> DynamicPrimitives;
 
 		// Relevant static primitive LODs are computed asynchronously.
 		// This task must complete before accessing StaticPrimitives in FRayTracingSceneAddInstancesTask.
@@ -270,7 +269,7 @@ namespace RayTracing
 
 	void GatherRelevantPrimitives(FScene& Scene, const FViewInfo& View, FRelevantPrimitiveList& Result)
 	{
-		Result.DirtyCachedRayTracingPrimitives.Reserve(Scene.PrimitiveSceneProxies.Num());
+		TArray<int32> StaticPrimitives;
 
 		const bool bGameView = View.bIsGameView || View.Family->EngineShowFlags.Game;
 
@@ -279,35 +278,31 @@ namespace RayTracing
 		{
 			TRACE_CPUPROFILER_EVENT_SCOPE(GatherRayTracingRelevantPrimitives);
 
-			// Index into the TypeOffsetTable, which contains a prefix sum of primitive indices by proxy type
-			int32 BroadIndex = 0;
-
-			for (int PrimitiveIndex = 0; PrimitiveIndex < Scene.PrimitiveSceneProxies.Num(); PrimitiveIndex++)
+			struct FGatherRelevantPrimitivesContext
 			{
-				// Find the next TypeOffsetTable entry that's relevant to this primitive index.
-				while (PrimitiveIndex >= int(Scene.TypeOffsetTable[BroadIndex].Offset))
-				{
-					BroadIndex++;
-				}
-
-				const ERayTracingPrimitiveFlags Flags = Scene.PrimitiveRayTracingFlags[PrimitiveIndex];
-
-				// Skip before dereferencing SceneInfo
-				if (Flags == ERayTracingPrimitiveFlags::UnsupportedProxyType)
-				{
-					// Find the index of a proxy of the next type, skipping over a batch of proxies that are the same type as current.
-					// This assumes that FPrimitiveSceneProxy::IsRayTracingRelevant() is consistent for all proxies of the same type.
-					// I.e. does not depend on members of the particular FPrimitiveSceneProxy implementation.
-					PrimitiveIndex = Scene.TypeOffsetTable[BroadIndex].Offset - 1;
-					continue;
-				}
+				TChunkedArray<int32> StaticPrimitives;
+				TChunkedArray<int32> DynamicPrimitives;
+				TChunkedArray<Nanite::CoarseMeshStreamingHandle> UsedCoarseMeshStreamingHandles;
+				TChunkedArray<FPrimitiveSceneInfo*> DirtyCachedRayTracingPrimitives;
+			};
 
+			TArray<FGatherRelevantPrimitivesContext> Contexts;
+			const int32 MinBatchSize = 128;
+			ParallelForWithTaskContext(
+				TEXT("GatherRayTracingRelevantPrimitives_Parallel"),
+				Contexts,
+				Scene.PrimitiveSceneProxies.Num(),
+				MinBatchSize,
+				[&Scene, &View, bGameView](FGatherRelevantPrimitivesContext& Context, int32 PrimitiveIndex)
+			{
 				// Get primitive visibility state from culling
 				if (!View.PrimitiveRayTracingVisibilityMap[PrimitiveIndex])
 				{
-					continue;
+					return;
 				}
 
+				const ERayTracingPrimitiveFlags Flags = Scene.PrimitiveRayTracingFlags[PrimitiveIndex];
+
 				check(!EnumHasAnyFlags(Flags, ERayTracingPrimitiveFlags::Exclude));
 
 				const FPrimitiveSceneInfo* SceneInfo = Scene.Primitives[PrimitiveIndex];
@@ -318,12 +313,12 @@ namespace RayTracing
 
 				if (View.bIsSceneCapture && (!bShouldRayTraceSceneCapture || !SceneInfo->bIsVisibleInSceneCaptures))
 				{
-					continue;
+					return;
 				}
 
 				if (!View.bIsSceneCapture && SceneInfo->bIsVisibleInSceneCapturesOnly)
 				{
-					continue;
+					return;
 				}
 
 				// Some primitives should only be visible editor mode, however far field geometry 
@@ -334,7 +329,7 @@ namespace RayTracing
 					checkf(SceneInfo->Proxy != nullptr, TEXT("SceneInfo does not have a valid Proxy object. If this occurs, this object should probably have been filtered out before being added to Scene.Primitives"));
 					if (!SceneInfo->Proxy->CastsHiddenShadow() && !SceneInfo->Proxy->AffectsIndirectLightingWhileHidden())
 					{
-						continue;
+						return;
 					}
 				}
 
@@ -342,39 +337,64 @@ namespace RayTracing
 				if (EnumHasAnyFlags(Flags, ERayTracingPrimitiveFlags::Streaming))
 				{
 					check(SceneInfo->CoarseMeshStreamingHandle != INDEX_NONE);
-					Result.UsedCoarseMeshStreamingHandles.Add(SceneInfo->CoarseMeshStreamingHandle);
+					Context.UsedCoarseMeshStreamingHandles.AddElement(SceneInfo->CoarseMeshStreamingHandle);
 				}
 
 				// Is the cached data dirty?
 				// eg: mesh was streamed in/out
 				if (SceneInfo->bCachedRaytracingDataDirty)
 				{
-					Result.DirtyCachedRayTracingPrimitives.Add(Scene.Primitives[PrimitiveIndex]);
+					Context.DirtyCachedRayTracingPrimitives.AddElement(Scene.Primitives[PrimitiveIndex]);
 				}
 
 				if (EnumHasAnyFlags(Flags, ERayTracingPrimitiveFlags::Skip))
 				{
-					continue;
+					return;
 				}
 
-				FRelevantPrimitive Item;
-				Item.PrimitiveIndex = PrimitiveIndex;
-				Item.PersistentPrimitiveIndex = SceneInfo->GetPersistentIndex();
-
 				if (EnumHasAnyFlags(Flags, ERayTracingPrimitiveFlags::Dynamic))
 				{
 					checkf(!EnumHasAllFlags(Flags, ERayTracingPrimitiveFlags::CacheInstances), TEXT("Only static primitives are expected to use CacheInstances flag."));
 
 					if (View.Family->EngineShowFlags.SkeletalMeshes) // TODO: Fix this check
 					{
-						Item.bStatic = false;
-						Result.DynamicPrimitives.AddElement(Item);
+						Context.DynamicPrimitives.AddElement(PrimitiveIndex);
 					}
 				}
 				else if (View.Family->EngineShowFlags.StaticMeshes)
 				{
-					Item.bStatic = true;
-					Result.StaticPrimitives.AddElement(Item);
+					Context.StaticPrimitives.AddElement(PrimitiveIndex);
+				}
+			});
+
+			if (Contexts.Num() > 0)
+			{
+				SCOPED_NAMED_EVENT(GatherRayTracingRelevantPrimitives_Merge, FColor::Emerald);
+
+				int32 NumStaticPrimitives = 0;
+				int32 NumDynamicPrimitives = 0;
+				int32 NumUsedCoarseMeshStreamingHandles = 0;
+				int32 NumDirtyCachedRayTracingPrimitives = 0;
+
+				for (auto& Context : Contexts)
+				{
+					NumStaticPrimitives += Context.StaticPrimitives.Num();
+					NumDynamicPrimitives += Context.DynamicPrimitives.Num();
+					NumUsedCoarseMeshStreamingHandles += Context.UsedCoarseMeshStreamingHandles.Num();
+					NumDirtyCachedRayTracingPrimitives += Context.DirtyCachedRayTracingPrimitives.Num();
+				}
+
+				StaticPrimitives.Reserve(NumStaticPrimitives);
+				Result.DynamicPrimitives.Reserve(NumDynamicPrimitives);
+				Result.UsedCoarseMeshStreamingHandles.Reserve(NumUsedCoarseMeshStreamingHandles);
+				Result.DirtyCachedRayTracingPrimitives.Reserve(NumDirtyCachedRayTracingPrimitives);
+
+				for (auto& Context : Contexts)
+				{
+					Context.StaticPrimitives.CopyToLinearArray(StaticPrimitives);
+					Context.DynamicPrimitives.CopyToLinearArray(Result.DynamicPrimitives);
+					Context.UsedCoarseMeshStreamingHandles.CopyToLinearArray(Result.UsedCoarseMeshStreamingHandles);
+					Context.DirtyCachedRayTracingPrimitives.CopyToLinearArray(Result.DirtyCachedRayTracingPrimitives);
 				}
 			}
 		}
@@ -386,39 +406,42 @@ namespace RayTracing
 		const int32 ForcedLODLevel = GetCVarForceLOD();
 
 		Result.StaticPrimitiveLODTask = FFunctionGraphTask::CreateAndDispatchWhenReady(
-			[&Result, &Scene, &View, LODScaleCVarValue, ForcedLODLevel]()
+			[&Result, &Scene, &View, LODScaleCVarValue, ForcedLODLevel, StaticPrimitiveIndices = MoveTemp(StaticPrimitives)]()
 			{
-				TRACE_CPUPROFILER_EVENT_SCOPE(GatherRayTracingWorldInstances_ComputeLOD);
+				TRACE_CPUPROFILER_EVENT_SCOPE(GatherRayTracingRelevantPrimitives_ComputeLOD);
+
+				Result.StaticPrimitives.SetNumUninitialized(StaticPrimitiveIndices.Num());
 
-				ParallelFor(TEXT("GatherRayTracingRelevantPrimitives_ComputeLOD"), Result.StaticPrimitives.Num(), 128,
-					[&Result, &Scene, &View, LODScaleCVarValue, ForcedLODLevel](int32 ItemIndex)
+				const int32 MinBatchSize = 128;
+				ParallelFor(TEXT("GatherRayTracingRelevantPrimitives_ComputeLOD_Parallel"), StaticPrimitiveIndices.Num(), MinBatchSize,
+					[&Result, &Scene, &View, LODScaleCVarValue, ForcedLODLevel, &StaticPrimitiveIndices](int32 ItemIndex)
 					{
-						FRelevantPrimitive& RelevantPrimitive = Result.StaticPrimitives[ItemIndex];
+						const int32 PrimitiveIndex = StaticPrimitiveIndices[ItemIndex];
 
-						const int32 PrimitiveIndex = RelevantPrimitive.PrimitiveIndex;
 						const FPrimitiveSceneInfo* SceneInfo = Scene.Primitives[PrimitiveIndex];
 						const ERayTracingPrimitiveFlags Flags = Scene.PrimitiveRayTracingFlags[PrimitiveIndex];
 
+						FRelevantPrimitive& RelevantPrimitive = Result.StaticPrimitives[ItemIndex];
+						RelevantPrimitive = {};
+						RelevantPrimitive.PrimitiveIndex = PrimitiveIndex;
+						RelevantPrimitive.PersistentPrimitiveIndex = SceneInfo->GetPersistentIndex();
+
 						int8 LODIndex = 0;
 
 						if (EnumHasAnyFlags(Flags, ERayTracingPrimitiveFlags::ComputeLOD))
 						{
 							const FPrimitiveBounds& Bounds = Scene.PrimitiveBounds[PrimitiveIndex];
-							const FPrimitiveSceneInfo* RESTRICT PrimitiveSceneInfo = Scene.Primitives[PrimitiveIndex];
-
-							FLODMask LODToRender;
 
-							const int8 CurFirstLODIdx = PrimitiveSceneInfo->Proxy->GetCurrentFirstLODIdx_RenderThread();
+							const int8 CurFirstLODIdx = SceneInfo->Proxy->GetCurrentFirstLODIdx_RenderThread();
 							check(CurFirstLODIdx >= 0);
 
 							float MeshScreenSizeSquared = 0;
 							float LODScale = LODScaleCVarValue * View.LODDistanceFactor;
-							LODToRender = ComputeLODForMeshes(SceneInfo->StaticMeshRelevances, View, Bounds.BoxSphereBounds.Origin, Bounds.BoxSphereBounds.SphereRadius, ForcedLODLevel, MeshScreenSizeSquared, CurFirstLODIdx, LODScale, true);
+							FLODMask LODToRender = ComputeLODForMeshes(SceneInfo->StaticMeshRelevances, View, Bounds.BoxSphereBounds.Origin, Bounds.BoxSphereBounds.SphereRadius, ForcedLODLevel, MeshScreenSizeSquared, CurFirstLODIdx, LODScale, true);
 
 							LODIndex = LODToRender.GetRayTracedLOD();
 						}
 
-
 						if (EnumHasAnyFlags(Flags, ERayTracingPrimitiveFlags::CacheInstances))
 						{
 							FPrimitiveSceneProxy* SceneProxy = Scene.PrimitiveSceneProxies[PrimitiveIndex];
@@ -675,11 +698,11 @@ namespace RayTracing
 			// Local temporary array of instances used for GetDynamicRayTracingInstances()
 			TArray<FRayTracingInstance> TempRayTracingInstances;
 
-			for (const FRelevantPrimitive& RelevantPrimitive : RelevantPrimitiveList.DynamicPrimitives)
+			for (int32 PrimitiveIndex : RelevantPrimitiveList.DynamicPrimitives)
 			{
-				const FPersistentPrimitiveIndex PersistentPrimitiveIndex = RelevantPrimitive.PersistentPrimitiveIndex;
-				const int32 PrimitiveIndex = RelevantPrimitive.PrimitiveIndex;
 				FPrimitiveSceneProxy* SceneProxy = Scene.PrimitiveSceneProxies[PrimitiveIndex];
+				FPrimitiveSceneInfo* SceneInfo = Scene.Primitives[PrimitiveIndex];
+				const FPersistentPrimitiveIndex PersistentPrimitiveIndex = SceneInfo->GetPersistentIndex();
 
 				TempRayTracingInstances.Reset();
 				MaterialGatheringContext.DynamicRayTracingGeometriesToUpdate.Reset();
@@ -858,7 +881,6 @@ namespace RayTracing
 					{
 						if (FVector::Distance(SceneProxy->GetActorPosition(), View.ViewMatrices.GetViewOrigin()) < CVarRayTracingDynamicGeometryLastRenderTimeUpdateDistance.GetValueOnRenderThread())
 						{
-							FPrimitiveSceneInfo* SceneInfo = Scene.Primitives[PrimitiveIndex];
 							// Update LastRenderTime for components so that visibility based ticking (like skeletal meshes) can get updated
 							// We are only doing this for dynamic geometries now
 							SceneInfo->LastRenderTime = CurrentWorldTime;
@@ -877,14 +899,14 @@ namespace RayTracing
 		{
 			UE_NONCOPYABLE(FRayTracingSceneAddInstancesTask)
 
-				static ESubsequentsMode::Type GetSubsequentsMode() { return ESubsequentsMode::TrackSubsequents; }
+			static ESubsequentsMode::Type GetSubsequentsMode() { return ESubsequentsMode::TrackSubsequents; }
 			TStatId                       GetStatId() const { return TStatId(); }
 			ENamedThreads::Type           GetDesiredThread() { return ENamedThreads::AnyThread; }
 
 			// Inputs
 
 			const FScene& Scene;
-			TChunkedArray<FRelevantPrimitive>& RelevantStaticPrimitives;
+			TArray<FRelevantPrimitive>& RelevantStaticPrimitives;
 			const FRayTracingCullingParameters& CullingParameters;
 			const bool bIsPathTracing;
 
@@ -894,7 +916,7 @@ namespace RayTracing
 			TArray<FVisibleRayTracingMeshCommand>& VisibleRayTracingMeshCommands; // New elements are added here by this task
 
 			FRayTracingSceneAddInstancesTask(const FScene& InScene,
-				TChunkedArray<FRelevantPrimitive>& InRelevantStaticPrimitives,
+				TArray<FRelevantPrimitive>& InRelevantStaticPrimitives,
 				const FRayTracingCullingParameters& InCullingParameters,
 				const bool bInIsPathTracing,
 				FRayTracingScene& InRayTracingScene, TArray<FVisibleRayTracingMeshCommand>& InVisibleRayTracingMeshCommands)
@@ -1053,9 +1075,10 @@ namespace RayTracing
 					{
 						const int8 LODIndex = RelevantPrimitive.LODIndex;
 
-						if (LODIndex < 0 || !RelevantPrimitive.bStatic)
+						if (LODIndex < 0)
 						{
-							continue; // skip dynamic primitives and other 
+							// TODO: Filter these primitives earlier
+							continue; 
 						}
 
 						// if primitive has mixed decal and non-decal segments we need to have two ray tracing instances
diff --git a/Engine/Source/Runtime/Renderer/Private/RayTracing/RayTracingInstanceCulling.cpp b/Engine/Source/Runtime/Renderer/Private/RayTracing/RayTracingInstanceCulling.cpp
--- a/Engine/Source/Runtime/Renderer/Private/RayTracing/RayTracingInstanceCulling.cpp
+++ b/Engine/Source/Runtime/Renderer/Private/RayTracing/RayTracingInstanceCulling.cpp
@@ -94,8 +94,9 @@ bool ShouldConsiderCulling(const FRayTracingCullingParameters& CullingParameters
 
 bool CullPrimitiveByFlags(const FRayTracingCullingParameters& CullingParameters, const FScene* RESTRICT Scene, int32 PrimitiveIndex)
 {
-	if (Scene->PrimitiveRayTracingFlags[PrimitiveIndex] == ERayTracingPrimitiveFlags::UnsupportedProxyType)
+	if (EnumHasAnyFlags(Scene->PrimitiveRayTracingFlags[PrimitiveIndex], ERayTracingPrimitiveFlags::UnsupportedProxyType))
 	{
+		checkf(Scene->PrimitiveRayTracingFlags[PrimitiveIndex] == ERayTracingPrimitiveFlags::UnsupportedProxyType, TEXT("ERayTracingPrimitiveFlags::UnsupportedProxyType should not be combined with other flags"));
 		return true;
 	}
 
