Remove declaration of unused & unimplemented function RenderVirtualShadowMapProjections
#rnx

diff --git a/Engine/Source/Runtime/Renderer/Private/DeferredShadingRenderer.h b/Engine/Source/Runtime/Renderer/Private/DeferredShadingRenderer.h
--- a/Engine/Source/Runtime/Renderer/Private/DeferredShadingRenderer.h
+++ b/Engine/Source/Runtime/Renderer/Private/DeferredShadingRenderer.h
@@ -899,13 +899,6 @@ private:
 	/** Renders indirect shadows from capsules modulated onto scene color. */
 	void RenderIndirectCapsuleShadows(FRDGBuilder& GraphBuilder, const FSceneTextures& SceneTextures) const;
 
-	void RenderVirtualShadowMapProjections(
-		FRDGBuilder& GraphBuilder,
-		const FMinimalSceneTextures& SceneTextures,
-		FRDGTextureRef ScreenShadowMaskTexture,
-		FRDGTextureRef ScreenShadowMaskSubPixelTexture,
-		const FLightSceneInfo* LightSceneInfo);
-
 	/** Renders capsule shadows for movable skylights, using the cone of visibility (bent normal) from DFAO. */
 	void RenderCapsuleShadowsForMovableSkylight(
 		FRDGBuilder& GraphBuilder,
