Fixed bug in hierarchical instance culling that caused the footprint-based test to never be selected.

diff --git a/Engine/Source/Runtime/Renderer/Private/SceneCulling/SceneCulling.cpp b/Engine/Source/Runtime/Renderer/Private/SceneCulling/SceneCulling.cpp
--- a/Engine/Source/Runtime/Renderer/Private/SceneCulling/SceneCulling.cpp
+++ b/Engine/Source/Runtime/Renderer/Private/SceneCulling/SceneCulling.cpp
@@ -676,8 +676,8 @@ void FSceneCulling::Test(const FCullingVolume& CullingVolume, TArray<FCellDraw,
 
 	if (CullingVolume.Sphere.W > 0.0f)
 	{
-		const float Level0CellSize = SpatialHash.GetCellSize(0);
-		FFootprint64 LightFootprint = SpatialHash.CalcFootprintSphere(0, CullingVolume.Sphere.Center, CullingVolume.Sphere.W + (Level0CellSize * 0.5f));
+		const float Level0CellSize = SpatialHash.GetCellSize(SpatialHash.GetFirstLevel());
+		FFootprint64 LightFootprint = SpatialHash.CalcFootprintSphere(SpatialHash.GetFirstLevel(), CullingVolume.Sphere.Center, CullingVolume.Sphere.W + (Level0CellSize * 0.5f));
 
 		// Diagonal length
 		if ((LightFootprint.Max - LightFootprint.Min).Size() < SmallFootprintCellCountThreshold)
@@ -2228,9 +2228,6 @@ public:
 			BUILDER_LOG("ItemChunkDataUploader %u -> %d", ChunkId, ChunkId);
 		}
 
-		SceneCulling.BlockLevelOccupancyMask.Reset();
-		SceneCulling.BlockLevelOccupancyMask.SetNum(FSpatialHash::kMaxLevel, false);
-
 		BlockDataUploader.Reserve(NumDirtyBlocks);
 		SceneCulling.CellBlockData.SetNum(SpatialHashMap.GetMaxIndex());
 		for (TConstSetBitIterator<SceneRenderingAllocator> BitIt(DirtyBlocks); BitIt; ++BitIt)
@@ -2264,7 +2261,6 @@ public:
 				FSceneCulling::FBlockLoc BlockLoc = BlockItem.Key;
 				FVector3d BlockWorldPos = SpatialHash.CalcBlockWorldPosition(BlockLoc);
 
-				SceneCulling.BlockLevelOccupancyMask[BlockLoc.GetLevel()] = true;
 				BlockData.WorldPos = FDFVector3{ BlockWorldPos };
 				BlockData.LevelCellSize = SpatialHash.GetCellSize(BlockLoc.GetLevel() - FSpatialHash::CellBlockDimLog2);
 			}
@@ -2281,6 +2277,13 @@ public:
 			}
 		}
 
+		// Mark used levels to be able to skip empty ones when querying. Could be skipped if no blocks were added/removed.
+		SceneCulling.BlockLevelOccupancyMask.Init(false, FSpatialHash::kMaxLevel);
+		for (const auto &Item : SpatialHashMap)
+		{
+			SceneCulling.BlockLevelOccupancyMask[Item.Key.GetLevel()] = true;
+		}
+
 		// Make sure to release any references to data allocated through SceneRenderingAllocator
 		// TODO: maybe change to default allocator? Might make sense to keep the data allocated anyway.
 		DirtyChunks.Empty();
