Heterogeneous Volumes: Initial shadow-casting support and integration with translucency pass.

diff --git a/Engine/Shaders/Private/BasePassPixelShader.usf b/Engine/Shaders/Private/BasePassPixelShader.usf
--- a/Engine/Shaders/Private/BasePassPixelShader.usf
+++ b/Engine/Shaders/Private/BasePassPixelShader.usf
@@ -8,8 +8,18 @@
 #define UE_DF_FORCE_FP32_OPS 1
 #endif
 
+#include "HeterogeneousVolumes/HeterogeneousVolumesAdaptiveVolumetricShadowMapTypes.ush" 
 #include "Common.ush"
 
+#ifndef ADAPTIVE_VOLUMETRIC_SHADOW_MAP
+#define ADAPTIVE_VOLUMETRIC_SHADOW_MAP 0
+#endif // ADAPTIVE_VOLUMETRIC_SHADOW_MAP
+
+#if ADAPTIVE_VOLUMETRIC_SHADOW_MAP
+#define AVSM TranslucentBasePass.AVSM
+#include "HeterogeneousVolumes/HeterogeneousVolumesAdaptiveVolumetricShadowMapSampling.ush"
+#endif // ADAPTIVE_VOLUMETRIC_SHADOW_MAP
+
 // Reroute SceneTexturesStruct uniform buffer references to the appropriate base pass uniform buffer
 #if MATERIALBLENDING_TRANSLUCENT || MATERIALBLENDING_ADDITIVE || MATERIALBLENDING_MODULATE
 	#define SceneTexturesStruct TranslucentBasePass.SceneTextures
@@ -1334,6 +1344,13 @@ void FPixelShaderInOut_MainPS(
 	}
 #endif
 
+#if ADAPTIVE_VOLUMETRIC_SHADOW_MAP
+	float3 TranslatedWorldPosition = SvPositionToResolvedTranslatedWorld(In.SvPosition);
+	float Transmittance = saturate(AVSM_SampleTransmittance(TranslatedWorldPosition, ResolvedView.TranslatedWorldCameraOrigin));
+	Opacity *= Transmittance;
+
+#endif // ADAPTIVE_VOLUMETRIC_SHADOW_MAP
+
 #if NEEDS_BASEPASS_PIXEL_FOGGING
 	const float OneOverPreExposure = ResolvedView.OneOverPreExposure;
 	float4 NDCPosition = mul(float4(MaterialParameters.WorldPosition_CamRelative.xyz, 1.0f), ResolvedView.TranslatedWorldToClip);
diff --git a/Engine/Shaders/Private/DeferredLightPixelShaders.usf b/Engine/Shaders/Private/DeferredLightPixelShaders.usf
--- a/Engine/Shaders/Private/DeferredLightPixelShaders.usf
+++ b/Engine/Shaders/Private/DeferredLightPixelShaders.usf
@@ -22,6 +22,9 @@
   #endif
 #endif
 
+
+#include "HeterogeneousVolumes/HeterogeneousVolumesAdaptiveVolumetricShadowMapTypes.ush" 
+//#include "HeterogeneousVolumes/HeterogeneousVolumesVoxelGridTypes.ush"
 #include "HairStrands/HairStrandsVisibilityCommonStruct.ush" 
 #include "Common.ush"
 
@@ -55,6 +58,13 @@
 #include "VirtualShadowMaps/VirtualShadowMapMaskBitsCommon.ush"
 #endif
 
+#if USE_ADAPTIVE_VOLUMETRIC_SHADOW_MAP
+#include "HeterogeneousVolumes/HeterogeneousVolumesAdaptiveVolumetricShadowMapSampling.ush"
+
+//#include "HeterogeneousVolumes/HeterogeneousVolumesVoxelGridUtils2.ush"
+//#include "HeterogeneousVolumes/HeterogeneousVolumesRayMarchingUtils.ush"
+#endif // USE_ADAPTIVE_VOLUMETRIC_SHADOW_MAP
+
 // Remap light shape permutation into light type
 #if LIGHT_SOURCE_SHAPE == 0
 #define CURRENT_LIGHT_TYPE LIGHT_TYPE_DIRECTIONAL
@@ -284,6 +294,10 @@ void DeferredLightPixelMain(
 			SubstrateShadowTermInputParameters.PrecomputedShadowFactors				= SubstrateReadPrecomputedShadowFactors(SubstratePixelHeader, PixelPos, SceneTexturesStruct.GBufferETexture);
 			SubstrateShadowTermInputParameters.TranslatedWorldPosition				= DerivedParams.TranslatedWorldPosition;
 			SubstrateShadowTermInputParameters.LightAttenuation						= GetLightAttenuationFromShadow(InputParams, SceneDepth, DerivedParams.TranslatedWorldPosition);
+#if USE_ADAPTIVE_VOLUMETRIC_SHADOW_MAP
+			//SubstrateShadowTermInputParameters.LightAttenuation *= Luminance(ComputeTransmittance(DerivedParams.TranslatedWorldPosition, LightData.TranslatedWorldPosition, 256));
+			SubstrateShadowTermInputParameters.LightAttenuation *= AVSM_SampleTransmittance(DerivedParams.TranslatedWorldPosition, LightData.TranslatedWorldPosition);
+#endif // USE_ADAPTIVE_VOLUMETRIC_SHADOW_MAP
 			SubstrateShadowTermInputParameters.Dither								= InterleavedGradientNoise(InputParams.PixelPos, View.StateFrameIndexMod8);
 
 			FSubstrateDeferredLighting SubstrateLighting = SubstrateDeferredLighting(
@@ -333,6 +347,24 @@ void DeferredLightPixelMain(
 		float SurfaceShadow = 1.0f;
 		
 		float4 LightAttenuation = GetLightAttenuationFromShadow(InputParams, SceneDepth, DerivedParams.TranslatedWorldPosition);
+#if USE_ADAPTIVE_VOLUMETRIC_SHADOW_MAP
+		float3 WorldPosition = DFFastSubtractDemote(DerivedParams.TranslatedWorldPosition, PrimaryView.PreViewTranslation);
+		float3 WorldToLight = LightData.Direction * 10000;
+		if (LightData.bSpotLight || LightData.bRectLight)
+		{
+			float3 WorldLightPosition = DFFastSubtractDemote(LightData.TranslatedWorldPosition, PrimaryView.PreViewTranslation);
+			WorldToLight = WorldLightPosition - WorldPosition;
+		}
+
+		if (!AVSM.bIsEmpty)
+		{
+			LightAttenuation *= AVSM_SampleTransmittance(DerivedParams.TranslatedWorldPosition, LightData.TranslatedWorldPosition);
+		}
+		else
+		{
+			//LightAttenuation *= Luminance(ComputeTransmittance(WorldPosition, WorldToLight, 512));
+		}
+#endif // USE_ADAPTIVE_VOLUMETRIC_SHADOW_MAP
 		float4 Radiance = GetDynamicLighting(DerivedParams.TranslatedWorldPosition, DerivedParams.CameraVector, ScreenSpaceData.GBuffer, ScreenSpaceData.AmbientOcclusion, ScreenSpaceData.GBuffer.ShadingModelID, LightData, LightAttenuation, Dither, uint2(InputParams.PixelPos), SurfaceShadow);
 
 		OutColor += Radiance;
diff --git a/Engine/Shaders/Private/DeferredLightingCommon.ush b/Engine/Shaders/Private/DeferredLightingCommon.ush
--- a/Engine/Shaders/Private/DeferredLightingCommon.ush
+++ b/Engine/Shaders/Private/DeferredLightingCommon.ush
@@ -19,6 +19,17 @@
 #include "LightFunctionAtlas/LightFunctionAtlasCommon.usf"
 #endif
 
+#ifndef ADAPTIVE_VOLUMETRIC_SHADOW_MAP
+#define ADAPTIVE_VOLUMETRIC_SHADOW_MAP 0
+#endif
+
+#if ADAPTIVE_VOLUMETRIC_SHADOW_MAP
+#include "HeterogeneousVolumes/HeterogeneousVolumesAdaptiveVolumetricShadowMapSampling.ush"
+
+//#include "HeterogeneousVolumes/HeterogeneousVolumesVoxelGridUtils2.ush"
+//#include "HeterogeneousVolumes/HeterogeneousVolumesRayMarchingUtils.ush"
+#endif // ADAPTIVE_VOLUMETRIC_SHADOW_MAP
+
 // This Substrate include assumes that if inline shading is required, it needs to 
 // be defined prior to DeferredLightingCommong.ush
 // SubstrateStruct.MaterialTextureArray is only available to translucent materials or global shaders computing contact shadows.
@@ -320,6 +331,10 @@ FLightAccumulator AccumulateDynamicLighting(
 	if (LightData.bRadialLight)
 	{
 		LightMask = GetLocalLightAttenuation( TranslatedWorldPosition, LightData, ToLight, L );
+#if ADAPTIVE_VOLUMETRIC_SHADOW_MAP
+		//LightAttenuation *= ComputeTransmittance(DerivedParams.TranslatedWorldPosition, LightData.TranslatedWorldPosition, 256);
+		LightAttenuation *= AVSM_SampleTransmittance(TranslatedWorldPosition, LightData.TranslatedWorldPosition);
+#endif // ADAPTIVE_VOLUMETRIC_SHADOW_MAP
 		MaskedLightColor *= LightMask;
 	}
 
diff --git a/Engine/Shaders/Private/HeterogeneousVolumes/HeterogeneousVolumesAdaptiveVolumetricShadowMapSampling.ush b/Engine/Shaders/Private/HeterogeneousVolumes/HeterogeneousVolumesAdaptiveVolumetricShadowMapSampling.ush
new file mode 100644
--- /dev/null
+++ b/Engine/Shaders/Private/HeterogeneousVolumes/HeterogeneousVolumesAdaptiveVolumetricShadowMapSampling.ush
@@ -0,0 +1,117 @@
+// Copyright Epic Games, Inc. All Rights Reserved.
+
+#pragma once
+
+#include "HeterogeneousVolumesAdaptiveVolumetricShadowMapUtils.ush"
+
+uint GetMyLinearIndex(uint Face, uint2 PixelCoord)
+{
+	uint LinearIndex = Face * AVSM.Resolution.x * AVSM.Resolution.y + PixelCoord.y * AVSM.Resolution.x + PixelCoord.x;
+	return LinearIndex;
+}
+
+float AVSM_Sample(int Face, int2 Pixel, float WorldHitT)
+{
+	uint LinearIndex = GetMyLinearIndex(Face, Pixel);
+	FAVSMIndirectionData IndirectionData = ASVM_UnpackIndirectionData(AVSM.IndirectionBuffer[LinearIndex]);
+
+	FAVSMSampleData ShadowData[2] = {
+		AVSM_CreateSampleData(0.0, 1.0),
+		AVSM_CreateSampleData(0.0, 1.0)
+	};
+
+	if (IndirectionData.SampleCount > 0)
+	{
+		ShadowData[0] = AVSM_UnpackSampleData(AVSM.SampleBuffer[IndirectionData.PixelOffset]);
+		for (int ElementIndex = 1; ElementIndex < IndirectionData.SampleCount; ++ElementIndex)
+		{
+			ShadowData[1] = AVSM_UnpackSampleData(AVSM.SampleBuffer[IndirectionData.PixelOffset + ElementIndex]);
+			if (WorldHitT < ShadowData[1].X)
+			{
+				break;
+			}
+
+			ShadowData[0] = ShadowData[1];
+		}
+	}
+
+	float Transmittance = 1.0;
+	if (WorldHitT > ShadowData[0].X)
+	{
+		float BlendWeight = saturate((WorldHitT - ShadowData[0].X) * rcp(ShadowData[1].X - ShadowData[0].X));
+		Transmittance = lerp(ShadowData[0].Tau, ShadowData[1].Tau, BlendWeight);
+	}
+	return Transmittance;
+}
+
+float AVSM_SampleTransmittance(float3 TranslatedWorldPosition, float3 LightTranslatedWorldPosition)
+{
+	if (any(AVSM.Resolution == 0))
+	{
+		return 1.0;
+	}
+
+	float Transmittance = 1.0;
+
+	int Face = 0;
+	if (AVSM.NumShadowMatrices > 1)
+	{
+		float3 LightDir = TranslatedWorldPosition - AVSM.TranslatedWorldOrigin;
+		float3 LightDirAbs = abs(LightDir);
+
+		Face = LightDirAbs[1] > LightDirAbs[0] ? 1 : 0;
+		Face = LightDirAbs[2] > LightDirAbs[Face] ? 2 : Face;
+
+		int FaceOffset = (sign(LightDir[Face]) > 0) ? 1 : 0;
+		Face = Face * 2 + FaceOffset;
+	}
+
+	float4 ShadowPositionAsFloat4 = 0;
+	float3 ShadowPosition = 0;
+	{
+		ShadowPositionAsFloat4 = mul(float4(TranslatedWorldPosition, 1), AVSM.TranslatedWorldToShadow[Face]);
+		if (ShadowPositionAsFloat4.w > 0)
+		{
+			ShadowPosition = ShadowPositionAsFloat4.xyz / ShadowPositionAsFloat4.w;
+			if (all(ShadowPosition.xyz > 0) || all(ShadowPosition.xy < 1))
+			{
+				float2 ShadowMapCoord = clamp(ShadowPosition.xy * AVSM.Resolution, 0, AVSM.Resolution - 1);
+				float ShadowBias = 0.5;
+				float ShadowZ = length(TranslatedWorldPosition - AVSM.TranslatedWorldOrigin) - ShadowBias;
+				if (AVSM.bIsDirectionalLight)
+				{
+					ShadowZ = dot(AVSM.TranslatedWorldPlane.xyz, TranslatedWorldPosition) + AVSM.TranslatedWorldPlane.w;
+				}
+
+			#define BILINEAR_INTERPOLATION 1
+			#if BILINEAR_INTERPOLATION
+				// Bilinear interpolation of shadow data
+				int2 ShadowMapCoordX[4] =
+				{
+					clamp(ShadowMapCoord, 0, AVSM.Resolution - 1),
+					clamp(ShadowMapCoord + int2(0, 1), 0, AVSM.Resolution - 1),
+					clamp(ShadowMapCoord + int2(1, 0), 0, AVSM.Resolution - 1),
+					clamp(ShadowMapCoord + int2(1, 1), 0, AVSM.Resolution - 1)
+				};
+
+				float TransmittanceX[4] =
+				{
+					AVSM_Sample(Face, ShadowMapCoordX[0], ShadowZ),
+					AVSM_Sample(Face, ShadowMapCoordX[1], ShadowZ),
+					AVSM_Sample(Face, ShadowMapCoordX[2], ShadowZ),
+					AVSM_Sample(Face, ShadowMapCoordX[3], ShadowZ)
+				};
+
+				float2 Weight = frac(ShadowMapCoord);
+				float TransmittanceY0 = lerp(TransmittanceX[0], TransmittanceX[2], Weight.x);
+				float TransmittanceY1 = lerp(TransmittanceX[1], TransmittanceX[3], Weight.x);
+				Transmittance = lerp(TransmittanceY0, TransmittanceY1, Weight.y);
+			#else
+				Transmittance = AVSM_Sample(Face, ShadowMapCoord, ShadowZ);
+			#endif
+			}
+		}
+	}
+
+	return Transmittance;
+}
diff --git a/Engine/Shaders/Private/HeterogeneousVolumes/HeterogeneousVolumesAdaptiveVolumetricShadowMapTypes.ush b/Engine/Shaders/Private/HeterogeneousVolumes/HeterogeneousVolumesAdaptiveVolumetricShadowMapTypes.ush
new file mode 100644
--- /dev/null
+++ b/Engine/Shaders/Private/HeterogeneousVolumes/HeterogeneousVolumesAdaptiveVolumetricShadowMapTypes.ush
@@ -0,0 +1,21 @@
+// Copyright Epic Games, Inc. All Rights Reserved.
+
+#pragma once
+
+#define AVSM_NULL_PTR 0xFF
+
+struct FAVSMLinkedListPackedData
+{
+	int Data;
+	int Ptr;
+};
+
+struct FAVSMSamplePackedData
+{
+	uint Data;
+};
+
+struct FAVSMIndirectionPackedData
+{
+	uint Data[2];
+};
diff --git a/Engine/Shaders/Private/HeterogeneousVolumes/HeterogeneousVolumesAdaptiveVolumetricShadowMapUtils.ush b/Engine/Shaders/Private/HeterogeneousVolumes/HeterogeneousVolumesAdaptiveVolumetricShadowMapUtils.ush
new file mode 100644
--- /dev/null
+++ b/Engine/Shaders/Private/HeterogeneousVolumes/HeterogeneousVolumesAdaptiveVolumetricShadowMapUtils.ush
@@ -0,0 +1,124 @@
+// Copyright Epic Games, Inc. All Rights Reserved.
+
+#pragma once
+
+#include "HeterogeneousVolumesAdaptiveVolumetricShadowMapTypes.ush"
+
+struct FAVSMSampleData
+{
+	float X;
+	float Tau;
+};
+
+/* AVSM Sample Data API */
+FAVSMSampleData AVSM_CreateSampleData(float X, float Tau)
+{
+	FAVSMSampleData SampleData;
+	SampleData.X = X;
+	SampleData.Tau = Tau;
+	return SampleData;
+}
+
+int AVSM_PackSampleData(float X, float3 Tau)
+{
+	int PackedData = f32tof16(X) | (f32tof16(Luminance(Tau)) << 16);
+	return PackedData;
+}
+
+int AVSM_PackSampleData(FAVSMSampleData SampleData)
+{
+	int PackedData = 0;
+	PackedData |= f32tof16(SampleData.X);
+	PackedData |= f32tof16(SampleData.Tau) << 16;
+	return PackedData;
+}
+
+FAVSMSampleData AVSM_UnpackSampleData(int PackedData)
+{
+	FAVSMSampleData SampleData;
+	SampleData.X = f16tof32(PackedData & 0xFFFF);
+	SampleData.Tau = f16tof32(PackedData >> 16);
+	return SampleData;
+}
+
+/* AVSM Linked List API */
+
+FAVSMSampleData AVSM_UnpackSampleData(int2 PackedData)
+{
+	return AVSM_UnpackSampleData(PackedData.x);
+}
+
+int AVSM_GetPrevPtr(int2 PackedData)
+{
+	return (PackedData.y >> 8) & AVSM_NULL_PTR;
+}
+
+int AVSM_GetNextPtr(int2 PackedData)
+{
+	return PackedData.y & AVSM_NULL_PTR;
+}
+
+void AVSM_SetPrevPtr(inout int2 PackedData, int PrevPtr)
+{
+	PackedData.y = (PackedData.y & 0xFFFF00FF) |
+		((PrevPtr & AVSM_NULL_PTR) << 8);
+}
+
+void AVSM_SetNextPtr(inout int2 PackedData, int NextPtr)
+{
+	PackedData.y = (PackedData.y & 0xFFFFFF00) |
+		(NextPtr & AVSM_NULL_PTR);
+}
+
+struct FAVSMIndirectionDataContext
+{
+	// TODO: Store begin and length
+	FAVSMSampleData Begin;
+	FAVSMSampleData End;
+};
+
+FAVSMIndirectionDataContext AVSM_CreateIndirectionDataContext(FAVSMSampleData Begin, FAVSMSampleData End)
+{
+	FAVSMIndirectionDataContext Context;
+	Context.Begin = Begin;
+	Context.End = End;
+	return Context;
+}
+
+#define AVSM_INVALID_PIXEL_OFFSET 0x3FFFFFF
+
+struct FAVSMIndirectionData
+{
+	uint PixelOffset;
+	int SampleCount;
+	float BeginX;
+	float EndX;
+};
+
+FAVSMIndirectionData AVSM_CreateIndirectionData(FAVSMIndirectionDataContext Context, uint PixelOffset, int SampleCount)
+{
+	FAVSMIndirectionData IndirectionData;
+	IndirectionData.PixelOffset = PixelOffset;
+	IndirectionData.SampleCount = SampleCount;
+	IndirectionData.BeginX = Context.Begin.X;
+	IndirectionData.EndX = Context.End.X;
+	return IndirectionData;
+}
+
+uint2 AVSM_PackIndirectionData(FAVSMIndirectionData IndirectionData)
+{
+	uint2 PackedData;
+	PackedData.x = ((IndirectionData.PixelOffset & AVSM_INVALID_PIXEL_OFFSET) << 6) | (IndirectionData.SampleCount & 0x3f);
+	PackedData.y = (f32tof16(IndirectionData.BeginX) << 16) | f32tof16(IndirectionData.EndX);
+	return PackedData;
+}
+
+FAVSMIndirectionData ASVM_UnpackIndirectionData(uint2 PackedData)
+{
+	FAVSMIndirectionData IndirectionData;
+	IndirectionData.PixelOffset = (PackedData.x >> 6) & AVSM_INVALID_PIXEL_OFFSET;
+	IndirectionData.SampleCount = PackedData.x & 0x3f;
+	IndirectionData.BeginX = f16tof32(PackedData.y >> 16);
+	IndirectionData.EndX = f16tof32(PackedData.y);
+	return IndirectionData;
+}
diff --git a/Engine/Shaders/Private/HeterogeneousVolumes/HeterogeneousVolumesFrustumVoxelGridUtils.ush b/Engine/Shaders/Private/HeterogeneousVolumes/HeterogeneousVolumesFrustumVoxelGridUtils.ush
--- a/Engine/Shaders/Private/HeterogeneousVolumes/HeterogeneousVolumesFrustumVoxelGridUtils.ush
+++ b/Engine/Shaders/Private/HeterogeneousVolumes/HeterogeneousVolumesFrustumVoxelGridUtils.ush
@@ -1,5 +1,7 @@
 // Copyright Epic Games, Inc. All Rights Reserved.
 
+#pragma once
+
 float Exponential(float sigma, float x)
 {
 	// X = 1 - e^(-sigma * x)
diff --git a/Engine/Shaders/Private/HeterogeneousVolumes/HeterogeneousVolumesLiveShadingPipeline.usf b/Engine/Shaders/Private/HeterogeneousVolumes/HeterogeneousVolumesLiveShadingPipeline.usf
--- a/Engine/Shaders/Private/HeterogeneousVolumes/HeterogeneousVolumesLiveShadingPipeline.usf
+++ b/Engine/Shaders/Private/HeterogeneousVolumes/HeterogeneousVolumesLiveShadingPipeline.usf
@@ -1,5 +1,6 @@
 // Copyright Epic Games, Inc. All Rights Reserved.
 
+#include "HeterogeneousVolumesAdaptiveVolumetricShadowMapTypes.ush"
 #include "HeterogeneousVolumesSparseVoxelUtils.ush"
 #include "../Common.ush"
 
diff --git a/Engine/Shaders/Private/HeterogeneousVolumes/HeterogeneousVolumesMaterialBakingPipeline.usf b/Engine/Shaders/Private/HeterogeneousVolumes/HeterogeneousVolumesMaterialBakingPipeline.usf
--- a/Engine/Shaders/Private/HeterogeneousVolumes/HeterogeneousVolumesMaterialBakingPipeline.usf
+++ b/Engine/Shaders/Private/HeterogeneousVolumes/HeterogeneousVolumesMaterialBakingPipeline.usf
@@ -149,6 +149,7 @@ int FixedBottomLevelResolution;
 //RWBuffer<uint> RWHashToVoxelBuffer;
 
 // Sampling mode
+int bJitter;
 int bSampleAtVertices;
 
 float HomogeneousThreshold;
@@ -241,8 +242,11 @@ void RasterizeBottomLevelOrthoGridCS(
 		int2 BottomLevelVoxelIndexAs2D = int2(BottomLevelVoxelIndex.y * BottomLevelVoxelResolution.x + BottomLevelVoxelIndex.x, BottomLevelVoxelIndex.z);
 
 		float3 Jitter = 0.5;
-		Jitter.xy = BlueNoiseVec2(BottomLevelVoxelIndexAs2D, View.StateFrameIndex);
-		Jitter.z = BlueNoiseScalar(BottomLevelVoxelIndexAs2D, View.StateFrameIndex);
+		if (bJitter)
+		{
+			Jitter.xy = BlueNoiseVec2(BottomLevelVoxelIndexAs2D, View.StateFrameIndex);
+			Jitter.z = BlueNoiseScalar(BottomLevelVoxelIndexAs2D, View.StateFrameIndex);
+		}
 
 		float3 Extinction = 0;
 		float3 Emission = 0;
@@ -456,10 +460,11 @@ void RasterizeBottomLevelFrustumGridCS(
 	float3 TopLevelVoxelPos = GetVoxelIndex(TopLevelGridLinearIndex, TopLevelGridResolution);
 	float3 BottomLevelVoxelPos = TopLevelVoxelPos * BottomLevelVoxelResolution + GroupThreadId;
 
-	float2 JitterXY = 0.5;
-	//float JitterZ = 0.5;
-	float JitterZ = BlueNoiseScalar(BottomLevelVoxelPos.xy, View.StateFrameIndex);
-	float3 Jitter = float3(JitterXY, JitterZ);
+	float3 Jitter = 0.5;
+	if (bJitter)
+	{
+		Jitter.z = BlueNoiseScalar(BottomLevelVoxelPos.xy, View.StateFrameIndex);
+	}
 
 	float3 LocalVoxelPos = GroupThreadId + Jitter;
 	float3 VoxelPosition = TopLevelVoxelPos + LocalVoxelPos / BottomLevelVoxelResolution;
diff --git a/Engine/Shaders/Private/HeterogeneousVolumes/HeterogeneousVolumesRayMarchingTypes.ush b/Engine/Shaders/Private/HeterogeneousVolumes/HeterogeneousVolumesRayMarchingTypes.ush
new file mode 100644
--- /dev/null
+++ b/Engine/Shaders/Private/HeterogeneousVolumes/HeterogeneousVolumesRayMarchingTypes.ush
@@ -0,0 +1,127 @@
+// Copyright Epic Games, Inc. All Rights Reserved.
+
+#include "../PathTracing/Utilities/PathTracingRandomSequence.ush"
+
+struct FRayMarchingContext
+{
+	// Local-space
+	float3 LocalRayOrigin;
+	float3 LocalRayDirection;
+	float LocalRayTMin;
+	float LocalRayTMax;
+
+	// World-space
+	float3 WorldRayOrigin;
+	float3 WorldRayDirection;
+	float LocalToWorldScale;
+
+	// Ray-step attributes
+	float Jitter;
+	float StepSize;
+	uint MaxStepCount;
+
+	float MipLevel;
+
+	// Lighting
+	int bApplyEmissionAndTransmittance;
+	int bApplyDirectLighting;
+	int bApplyShadowTransmittance;
+
+	// TEMPORARY STORAGE!!
+	RandomSequence RandSequence;
+};
+
+FRayMarchingContext CreateRayMarchingContext(
+	// Local-space
+	float3 LocalRayOrigin,
+	float3 LocalRayDirection,
+	float LocalRayTMin,
+	float LocalRayTMax,
+	// World-space
+	float3 WorldRayOrigin,
+	float3 WorldRayDirection,
+	// Ray-step attributes
+	float Jitter,
+	float StepSize,
+	uint MaxStepCount,
+	// Lighting
+	int bApplyEmissionAndTransmittance,
+	int bApplyDirectLighting,
+	int bApplyShadowTransmittance
+)
+{
+	FRayMarchingContext Context = (FRayMarchingContext)0;
+
+	// Local-space
+	Context.LocalRayOrigin = LocalRayOrigin;
+	Context.LocalRayDirection = LocalRayDirection;
+	Context.LocalRayTMin = LocalRayTMin;
+	Context.LocalRayTMax = LocalRayTMax;
+	// World-space
+	Context.WorldRayOrigin = WorldRayOrigin;
+	Context.WorldRayDirection = WorldRayDirection;
+	// Ray-step attributes
+	Context.Jitter = Jitter;
+	Context.StepSize = StepSize;
+	Context.MaxStepCount = MaxStepCount;
+	// Lighting
+	Context.bApplyEmissionAndTransmittance = bApplyEmissionAndTransmittance;
+	Context.bApplyDirectLighting = bApplyDirectLighting;
+	Context.bApplyShadowTransmittance = bApplyShadowTransmittance;
+
+	// Transform local-scale properties into world-scale
+	float3 LocalRayHit = LocalRayOrigin + LocalRayDirection * LocalRayTMax;
+	float3 WorldRayHit = mul(float4(LocalRayHit, 1.0), GetLocalToWorld()).xyz;
+	float WorldRayTMax = length(WorldRayHit - WorldRayOrigin);
+	Context.LocalToWorldScale = WorldRayTMax / LocalRayTMax;
+
+	Context.MipLevel = 0.0f;
+
+	return Context;
+}
+
+float CalcStepSize(float3 LocalRayDirection, float LocalStepSize, float LocalStepFactor)
+{
+	if (LocalStepSize <= 0.0f)
+	{
+		// Intersect aligned voxel to determine candidate t-values and use the minimum distance
+		float3 VoxelSize = GetLocalBoundsExtent() / float3(GetVolumeResolution());
+		float3 tValues = abs(VoxelSize / LocalRayDirection);
+		float VoxelStepSize = min(tValues.x, min(tValues.y, tValues.z));
+		LocalStepSize = VoxelStepSize * LocalStepFactor;
+	}
+	return max(LocalStepSize, 0.01f);
+}
+
+float CalcStepSize(float3 LocalRayDirection)
+{
+	float LocalStepSize = GetStepSize();
+	float LocalStepFactor = GetStepFactor();
+	return CalcStepSize(LocalRayDirection, LocalStepSize, LocalStepFactor);
+}
+
+float CalcShadowStepSize(float3 LocalRayDirection)
+{
+	float LocalStepSize = GetShadowStepSize();
+	float LocalStepFactor = GetShadowStepFactor();
+	return CalcStepSize(LocalRayDirection, LocalStepSize, LocalStepFactor);
+}
+
+float RecalculateStepSize(inout FRayMarchingContext Context)
+{
+	// Recalculate step-size if the bounding-box diagonal is greater than the maximum step-count
+	float DiagonalLength = length(2.0 * GetLocalBoundsExtent());
+	uint MaxStepCount = max(DiagonalLength / Context.StepSize, 1);
+
+	float StepSize = (MaxStepCount > Context.MaxStepCount) ? DiagonalLength / Context.MaxStepCount : Context.StepSize;
+	return StepSize;
+}
+
+uint CalcStepCount(inout FRayMarchingContext Context)
+{
+	Context.StepSize = RecalculateStepSize(Context);
+
+	float HitSpan = length(Context.LocalRayTMax - Context.LocalRayTMin);
+	uint StepCount = max(ceil(HitSpan / Context.StepSize), 1);
+	return StepCount;
+}
\ No newline at end of file
diff --git a/Engine/Shaders/Private/HeterogeneousVolumes/HeterogeneousVolumesRayMarchingUtils.ush b/Engine/Shaders/Private/HeterogeneousVolumes/HeterogeneousVolumesRayMarchingUtils.ush
--- a/Engine/Shaders/Private/HeterogeneousVolumes/HeterogeneousVolumesRayMarchingUtils.ush
+++ b/Engine/Shaders/Private/HeterogeneousVolumes/HeterogeneousVolumesRayMarchingUtils.ush
@@ -1,5 +1,7 @@
 // Copyright Epic Games, Inc. All Rights Reserved.
 
+#include "HeterogeneousVolumesRayMarchingTypes.ush"
+#include "HeterogeneousVolumesTracingUtils.ush"
 #include "HeterogeneousVolumesTransmittanceVolumeUtils.ush"
 
 #include "../ForwardShadowingCommon.ush"
@@ -10,6 +12,10 @@
 #define RadianceCacheInterpolation LumenGIVolumeStruct
 #include "../Lumen/LumenTranslucencyVolumeShared.ush"
 
+#ifndef CUSTOM_TRANSMITTANCE_FUNCTION
+#define CUSTOM_TRANSMITTANCE_FUNCTION 0
+#endif // CUSTOM_TRANSMITTANCE_FUNCTION
+
 #ifndef HARDWARE_RAY_TRACING
 #define HARDWARE_RAY_TRACING 0
 #endif // HARDWARE_RAY_TRACING
@@ -18,107 +24,13 @@
 #define HARD_SURFACE_SHADOWING 1
 #endif // HARD_SURFACE_SHADOWING
 
-struct FRayMarchingContext
-{
-	// Local-space
-	float3 LocalRayOrigin;
-	float3 LocalRayDirection;
-	float LocalRayTMin;
-	float LocalRayTMax;
-
-	// World-space
-	float3 WorldRayOrigin;
-	float3 WorldRayDirection;
-	float LocalToWorldScale;
-
-	// Ray-step attributes
-	float Jitter;
-	float StepSize;
-	uint MaxStepCount;
-
-	float MipLevel;
-
-	// Lighting
-	int bApplyEmissionAndTransmittance;
-	int bApplyDirectLighting;
-	int bApplyShadowTransmittance;
-};
-
-FRayMarchingContext CreateRayMarchingContext(
-	// Local-space
-	float3 LocalRayOrigin,
-	float3 LocalRayDirection,
-	float LocalRayTMin,
-	float LocalRayTMax,
-	// World-space
-	float3 WorldRayOrigin,
-	float3 WorldRayDirection,
-	// Ray-step attributes
-	float Jitter,
-	float StepSize,
-	uint MaxStepCount,
-	// Lighting
-	int bApplyEmissionAndTransmittance,
-	int bApplyDirectLighting,
-	int bApplyShadowTransmittance
-)
-{
-	FRayMarchingContext Context = (FRayMarchingContext)0;
-
-	// Local-space
-	Context.LocalRayOrigin = LocalRayOrigin;
-	Context.LocalRayDirection = LocalRayDirection;
-	Context.LocalRayTMin = LocalRayTMin;
-	Context.LocalRayTMax = LocalRayTMax;
-	// World-space
-	Context.WorldRayOrigin = WorldRayOrigin;
-	Context.WorldRayDirection = WorldRayDirection;
-	// Ray-step attributes
-	Context.Jitter = Jitter;
-	Context.StepSize = StepSize;
-	Context.MaxStepCount = MaxStepCount;
-	// Lighting
-	Context.bApplyEmissionAndTransmittance = bApplyEmissionAndTransmittance;
-	Context.bApplyDirectLighting = bApplyDirectLighting;
-	Context.bApplyShadowTransmittance = bApplyShadowTransmittance;
-
-	// Transform local-scale properties into world-scale
-	float3 LocalRayHit = LocalRayOrigin + LocalRayDirection * LocalRayTMax;
-	float3 WorldRayHit = mul(float4(LocalRayHit, 1.0), GetLocalToWorld()).xyz;
-	float WorldRayTMax = length(WorldRayHit - WorldRayOrigin);
-	Context.LocalToWorldScale = WorldRayTMax / LocalRayTMax;
-
-	Context.MipLevel = 0.0f;
-
-	return Context;
-}
+#ifndef DIM_USE_ADAPTIVE_VOLUMETRIC_SHADOW_MAP
+#define DIM_USE_ADAPTIVE_VOLUMETRIC_SHADOW_MAP 0
+#endif // DIM_USE_ADAPTIVE_VOLUMETRIC_SHADOW_MAP
 
-float CalcStepSize(float3 LocalRayDirection, float LocalStepSize, float LocalStepFactor)
-{
-	if (LocalStepSize <= 0.0f)
-	{
-		// Intersect aligned voxel to determine candidate t-values and use the minimum distance
-		float3 VoxelSize = GetLocalBoundsExtent() / float3(GetVolumeResolution());
-		float3 tValues = abs(VoxelSize / LocalRayDirection);
-		float VoxelStepSize = min(tValues.x, min(tValues.y, tValues.z));
-		LocalStepSize = VoxelStepSize * LocalStepFactor;
-	}
-	return max(LocalStepSize, 0.01f);
-}
-
-float CalcStepSize(float3 LocalRayDirection)
-{
-	float LocalStepSize = GetStepSize();
-	float LocalStepFactor = GetStepFactor();
-	return CalcStepSize(LocalRayDirection, LocalStepSize, LocalStepFactor);
-}
-
-float CalcShadowStepSize(float3 LocalRayDirection)
-{
-	float LocalStepSize = GetShadowStepSize();
-	float LocalStepFactor = GetShadowStepFactor();
-	return CalcStepSize(LocalRayDirection, LocalStepSize, LocalStepFactor);
-}
+#if DIM_USE_ADAPTIVE_VOLUMETRIC_SHADOW_MAP
+#include "HeterogeneousVolumesAdaptiveVolumetricShadowMapSampling.ush"
+#endif // DIM_USE_ADAPTIVE_VOLUMETRIC_SHADOW_MAP
 
 float CalcShadowBias()
 {
@@ -135,25 +47,6 @@ float CalcShadowBias()
 	return WorldShadowBias;
 }
 
-float RecalculateStepSize(inout FRayMarchingContext Context)
-{
-	// Recalculate step-size if the bounding-box diagonal is greater than the maximum step-count
-	float DiagonalLength = length(2.0 * GetLocalBoundsExtent());
-	uint MaxStepCount = max(DiagonalLength / Context.StepSize, 1);
-
-	float StepSize = (MaxStepCount > Context.MaxStepCount) ? DiagonalLength / Context.MaxStepCount : Context.StepSize;
-	return StepSize;
-}
-
-uint CalcStepCount(inout FRayMarchingContext Context)
-{
-	Context.StepSize = RecalculateStepSize(Context);
-
-	float HitSpan = length(Context.LocalRayTMax - Context.LocalRayTMin);
-	uint StepCount = max(ceil(HitSpan / Context.StepSize), 1);
-	return StepCount;
-}
-
 float ComputeHardSurfaceShadowFactor(
 	float3 TranslatedWorldPosition,
 	FDeferredLightData LightData,
@@ -228,7 +121,11 @@ float3 ComputeTransmittance(
 	inout FRayMarchingContext RayMarchingContext
 )
 {
-#if HARDWARE_RAY_TRACING
+#if CUSTOM_TRANSMITTANCE_FUNCTION
+	float WorldRayTMin = RayMarchingContext.LocalRayTMin * RayMarchingContext.LocalToWorldScale;
+	float WorldRayTMax = RayMarchingContext.LocalRayTMax * RayMarchingContext.LocalToWorldScale;
+	return TransmittanceFunction(RayMarchingContext.WorldRayOrigin, RayMarchingContext.WorldRayDirection, WorldRayTMin, WorldRayTMax);
+#elif HARDWARE_RAY_TRACING
 	//  TODO: Incorporpate ray jitter
 	return ComputeTransmittanceHardwareRayTracing(RayMarchingContext.WorldRayOrigin, RayMarchingContext.WorldRayDirection, RayMarchingContext.LocalRayTMin, RayMarchingContext.LocalRayTMax);
 #else // HARDWARE_RAY_TRACING
@@ -260,51 +157,62 @@ float3 ComputeTransmittance(
 	uint MaxStepCount
 )
 {
+	float3 Transmittance = 1.0;
+
 #if DIM_USE_TRANSMITTANCE_VOLUME
 	float3 LocalRayOrigin = mul(float4(WorldRayOrigin, 1.0), GetWorldToLocal()).xyz;
 	float3 LocalBoundsMin = GetLocalBoundsOrigin() - GetLocalBoundsExtent();
 	float3 TransmittanceUVW = saturate((LocalRayOrigin - LocalBoundsMin) / (2.0 * GetLocalBoundsExtent()));
 	float MipLevel = 0;
 
-	float3 Transmittance = SampleLightingCache(TransmittanceUVW, MipLevel);
-
-#else // DIM_USE_TRANSMITTANCE_VOLUME
-
-	float3 WorldRayEnd = WorldRayOrigin + ToLight;
-	float3 WorldRayDirection = normalize(ToLight);
-
-	float3 LocalRayOrigin = mul(float4(WorldRayOrigin, 1.0), GetWorldToLocal()).xyz;
-	float3 LocalRayEnd = mul(float4(WorldRayEnd, 1.0), GetWorldToLocal()).xyz;
-	float3 LocalRayDirection = LocalRayEnd - LocalRayOrigin;
-	float LocalRayTMin = 0.0;
-	float LocalRayTMax = length(LocalRayDirection);
-	LocalRayDirection /= LocalRayTMax;
-
-	float ShadowBias = 0.5;
-	float ShadowStepSize = CalcShadowStepSize(LocalRayDirection);
-	int bApplyEmissionAndTransmittance = 0;
-	int bApplyDirectLighting = 0;
-	int bApplyShadowTransmittance = 0;
-
-	FRayMarchingContext ShadowRayMarchingContext = CreateRayMarchingContext(
-		LocalRayOrigin,
-		LocalRayDirection,
-		LocalRayTMin,
-		LocalRayTMax,
-		WorldRayOrigin,
-		WorldRayDirection,
-		ShadowBias,
-		ShadowStepSize,
-		MaxStepCount,
-		bApplyEmissionAndTransmittance,
-		bApplyDirectLighting,
-		bApplyShadowTransmittance
-		//RayMarchingContext.MaxShadowTraceDistance
-	);
-
-	float3 Transmittance = ComputeTransmittance(ShadowRayMarchingContext);
-#endif // DIM_USE_INSCATTERING_VOLUME
+	Transmittance = SampleLightingCache(TransmittanceUVW, MipLevel);
 
+#else
+	#if DIM_USE_ADAPTIVE_VOLUMETRIC_SHADOW_MAP
+		if (!AVSM.bIsEmpty)
+		{
+			float3 TranslatedWorldPosition = DFHackToFloat(DFFastAdd(WorldRayOrigin, PrimaryView.PreViewTranslation));
+			float3 LightTranslatedWorldPosition = TranslatedWorldPosition + ToLight;
+			Transmittance = AVSM_SampleTransmittance(TranslatedWorldPosition, LightTranslatedWorldPosition);
+		}
+		else
+	#endif
+		{
+			float3 WorldRayEnd = WorldRayOrigin + ToLight;
+			float3 WorldRayDirection = normalize(ToLight);
+
+			float3 LocalRayOrigin = mul(float4(WorldRayOrigin, 1.0), GetWorldToLocal()).xyz;
+			float3 LocalRayEnd = mul(float4(WorldRayEnd, 1.0), GetWorldToLocal()).xyz;
+			float3 LocalRayDirection = LocalRayEnd - LocalRayOrigin;
+			float LocalRayTMin = 0.0;
+			float LocalRayTMax = length(LocalRayDirection);
+			LocalRayDirection /= LocalRayTMax;
+
+			float ShadowBias = 0.5;
+			float ShadowStepSize = CalcShadowStepSize(LocalRayDirection);
+			int bApplyEmissionAndTransmittance = 0;
+			int bApplyDirectLighting = 0;
+			int bApplyShadowTransmittance = 0;
+
+			FRayMarchingContext ShadowRayMarchingContext = CreateRayMarchingContext(
+				LocalRayOrigin,
+				LocalRayDirection,
+				LocalRayTMin,
+				LocalRayTMax,
+				WorldRayOrigin,
+				WorldRayDirection,
+				ShadowBias,
+				ShadowStepSize,
+				MaxStepCount,
+				bApplyEmissionAndTransmittance,
+				bApplyDirectLighting,
+				bApplyShadowTransmittance
+				//RayMarchingContext.MaxShadowTraceDistance
+			);
+
+			Transmittance = ComputeTransmittance(ShadowRayMarchingContext);
+		}
+#endif
 	return Transmittance;
 }
 
@@ -410,14 +318,13 @@ void RayMarchSingleScattering(
 			Radiance += Emission * RayMarchingContext.StepSize * Transmittance;
 
 #if DIM_USE_LUMEN_GI
-			float3 TranslucencyEvaluationPosition = WorldPosition;
-			FTwoBandSHVectorRGB TranslucencyGISH = GetTranslucencyGIVolumeLighting(DFPromote(TranslucencyEvaluationPosition), PrimaryView.WorldToClip, true);
+			FTwoBandSHVectorRGB TranslucencyGISH = GetTranslucencyGIVolumeLighting(DFPromote(WorldPosition), PrimaryView.WorldToClip, true);
 
 			FTwoBandSHVector RotatedHGZonalHarmonic;
 			float3 CameraVector = normalize(RayMarchingContext.WorldRayDirection);
 			float PhaseG = 0.0;
 			RotatedHGZonalHarmonic.V = float4(1.0f, CameraVector.y, CameraVector.z, CameraVector.x) * float4(1.0f, PhaseG, PhaseG, PhaseG);
-			float3 IndirectInscattering = max(DotSH(TranslucencyGISH, RotatedHGZonalHarmonic), 0);
+			float3 IndirectInscattering = max(DotSH(TranslucencyGISH, RotatedHGZonalHarmonic), 0) / PI;
 
 			float3 Extinction = SampleExtinction(SampleContext);
 			float3 Albedo = SampleAlbedo(SampleContext);
diff --git a/Engine/Shaders/Private/HeterogeneousVolumes/HeterogeneousVolumesTracingUtils.ush b/Engine/Shaders/Private/HeterogeneousVolumes/HeterogeneousVolumesTracingUtils.ush
--- a/Engine/Shaders/Private/HeterogeneousVolumes/HeterogeneousVolumesTracingUtils.ush
+++ b/Engine/Shaders/Private/HeterogeneousVolumes/HeterogeneousVolumesTracingUtils.ush
@@ -1,5 +1,7 @@
 // Copyright Epic Games, Inc. All Rights Reserved.
 
+#pragma once
+
 // Modified LineBoxIntersect() to operate at consistent RayDirection scale, within some [TMin, TMax] span
 float2 IntersectAABB(float3 RayOrigin, float3 RayDirection, float RayTMin, float RayTMax, float3 BoxMin, float3 BoxMax)
 {
diff --git a/Engine/Shaders/Private/HeterogeneousVolumes/HeterogeneousVolumesVoxelGridRendering.usf b/Engine/Shaders/Private/HeterogeneousVolumes/HeterogeneousVolumesVoxelGridRendering.usf
--- a/Engine/Shaders/Private/HeterogeneousVolumes/HeterogeneousVolumesVoxelGridRendering.usf
+++ b/Engine/Shaders/Private/HeterogeneousVolumes/HeterogeneousVolumesVoxelGridRendering.usf
@@ -1,12 +1,14 @@
 // Copyright Epic Games, Inc. All Rights Reserved.
 
 #include "HeterogeneousVolumesVoxelGridTypes.ush"
+#include "HeterogeneousVolumesAdaptiveVolumetricShadowMapTypes.ush"
 
 #include "../Common.ush"
 #include "../DeferredShadingCommon.ush"
 #include "../ComputeShaderUtils.ush"
 #include "HeterogeneousVolumesTracingUtils.ush"
 #include "HeterogeneousVolumesVoxelGridTraversal.ush"
+#include "HeterogeneousVolumesVoxelGridRenderingUtils.ush"
 
 #include "../PathTracing/Utilities/PathTracingRandomSequence.ush"
 
@@ -40,13 +42,16 @@
 
 struct FRenderDebugData
 {
+	float3 Estimate;
 	float3 RayOrigin;
 	float3 RayDirection;
 	float TMin;
 	float TMax;
-	float3 Estimate;
 	float Pdf;
 	float Distance;
+
+	float4x4 ScreenToTranslatedWorld;
+	float4x4 ClipToTranslatedWorld;
 };
 
 // Ray data
@@ -143,68 +148,7 @@ FMajorantData EvalMajorant(float3 TranslatedWorldPos)
 	return MajorantData;
 }
 
-float3 EvalExtinction(float3 TranslatedWorldPos, inout int LinearBottomLevelVoxelPos)
-{
-	float3 Extinction = 0;
-
-	float3 TranslatedWorldBoundsMin = GetTranslatedWorldPos(OrthoGridUniformBuffer.TopLevelGridWorldBoundsMin);
-	float3 TranslatedWorldBoundsMax = GetTranslatedWorldPos(OrthoGridUniformBuffer.TopLevelGridWorldBoundsMax);
-	float3 TranslatedWorldBoundsExtent = TranslatedWorldBoundsMax - TranslatedWorldBoundsMin;
-	float3 GridUV = (TranslatedWorldPos - TranslatedWorldBoundsMin) / TranslatedWorldBoundsExtent;
-	if (all(GridUV >= 0.0) && all(GridUV <= 1.0))
-	{
-		float3 TopLevelVoxelPos = GridUV * OrthoGridUniformBuffer.TopLevelGridResolution;
-		uint LinearTopLevelVoxelPos = GetLinearIndex(TopLevelVoxelPos, OrthoGridUniformBuffer.TopLevelGridResolution);
-
-		if (OrthoGridUniformBuffer.bEnableIndirectionGrid)
-		{
-			FTopLevelGridData TopLevelGridData = OrthoGridUniformBuffer.TopLevelGridBuffer[LinearTopLevelVoxelPos];
-			if (IsBottomLevelAllocated(TopLevelGridData))
-			{
-				uint IndirectionGridIndex = GetBottomLevelIndex(TopLevelGridData);
-				uint3 IndirectionGridVoxelResolution = GetBottomLevelVoxelResolution(TopLevelGridData);
-
-				// Convert to Indirection voxel-space
-				float3 IndirectionGridVoxelPos = frac(TopLevelVoxelPos) * IndirectionGridVoxelResolution;
-				int3 IndirectionGridVoxelPosAsInt = clamp(IndirectionGridVoxelPos, 0, IndirectionGridVoxelResolution - 1);
-				FTopLevelGridData IndirectionData = OrthoGridUniformBuffer.IndirectionGridBuffer[IndirectionGridIndex + MortonEncode3(IndirectionGridVoxelPosAsInt)];
-				if (IsBottomLevelAllocated(IndirectionData))
-				{
-					uint BottomLevelIndex = GetBottomLevelIndex(IndirectionData);
-					uint3 BottomLevelVoxelResolution = GetBottomLevelVoxelResolution(IndirectionData);
-
-					// Constant interpolation
-					float3 BottomLevelVoxelPos = frac(IndirectionGridVoxelPos) * BottomLevelVoxelResolution;
-					int3 BottomLevelVoxelPosAsInt = clamp(BottomLevelVoxelPos, 0, BottomLevelVoxelResolution - 1);
-					LinearBottomLevelVoxelPos = BottomLevelIndex + MortonEncode3(BottomLevelVoxelPosAsInt);
-					Extinction = GetExtinction(OrthoGridUniformBuffer.ExtinctionGridBuffer[LinearBottomLevelVoxelPos]);
-					//Albedo = GetAlbedo(OrthoGridUniformBuffer.ScatteringGridBuffer[LinearBottomLevelVoxelPos]);
-				}
-			}
-		}
-		else
-		{
-			FTopLevelGridData TopLevelData = OrthoGridUniformBuffer.TopLevelGridBuffer[LinearTopLevelVoxelPos];
-			if (IsBottomLevelAllocated(TopLevelData))
-			{
-				uint BottomLevelIndex = GetBottomLevelIndex(OrthoGridUniformBuffer.TopLevelGridBuffer[LinearTopLevelVoxelPos]);
-				uint3 BottomLevelVoxelResolution = GetBottomLevelVoxelResolution(OrthoGridUniformBuffer.TopLevelGridBuffer[LinearTopLevelVoxelPos]);
-
-				// Constant Interpolation
-				float3 BottomLevelVoxelPos = frac(TopLevelVoxelPos) * BottomLevelVoxelResolution;
-				int3 BottomLevelVoxelPosAsInt = clamp(BottomLevelVoxelPos, 0, BottomLevelVoxelResolution - 1);
-				LinearBottomLevelVoxelPos = BottomLevelIndex + MortonEncode3(BottomLevelVoxelPosAsInt);
-
-				Extinction = GetExtinction(OrthoGridUniformBuffer.ExtinctionGridBuffer[LinearBottomLevelVoxelPos]);
-				//Albedo = GetAlbedo(OrthoGridUniformBuffer.ScatteringGridBuffer[LinearBottomLevelVoxelPos]);
-			}
-		}
-	}
-
-	return Extinction;
-}
-
-float CalcStepSize(float3 WorldDirection)
+float CalcStepSizeLocal(float3 WorldDirection)
 {
 	float3 TranslatedWorldBoundsMin = GetTranslatedWorldPos(OrthoGridUniformBuffer.TopLevelGridWorldBoundsMin);
 	float3 TranslatedWorldBoundsMax = GetTranslatedWorldPos(OrthoGridUniformBuffer.TopLevelGridWorldBoundsMax);
@@ -230,7 +174,7 @@ FTrackingResult SampleBasedMarch(float3 WorldRayOrigin, float3 WorldRayDirection
 		float3 WorldPosition = WorldRayOrigin + WorldRayDirection * HitT;
 
 		int LinearVoxelPos = -1;
-		float3 Extinction = EvalExtinction(WorldPosition, LinearVoxelPos);
+		float3 Extinction = EvalExtinctionOutOfFrustum(WorldPosition, LinearVoxelPos);
 		if (any(Extinction > 0))
 		{
 			Result.Radiance += GetEmission(OrthoGridUniformBuffer.EmissionGridBuffer[LinearVoxelPos]) * Result.Throughput;
@@ -312,7 +256,7 @@ FTrackingResult DeltaTracking(float3 WorldRayOrigin, float3 WorldRayDirection, f
 		{
 			float3 WorldPosition = WorldRayOrigin + WorldRayDirection * Result.Distance;
 			int LinearVoxelPos = -1;
-			float3 Extinction = EvalExtinction(WorldPosition, LinearVoxelPos);
+			float3 Extinction = EvalExtinctionOutOfFrustum(WorldPosition, LinearVoxelPos);
 			float SigmaT = Luminance(Extinction);
 
 			//float AcceptanceProbability = MaxComponent(SigmaT) / SigmaBar;
@@ -348,7 +292,7 @@ FTrackingResult SpectralWeightedDeltaTracking(float3 WorldRayOrigin, float3 Worl
 		{
 			float3 WorldPosition = WorldRayOrigin + WorldRayDirection * Result.Distance;
 			int LinearVoxelPos = -1;
-			float3 Extinction = EvalExtinction(WorldPosition, LinearVoxelPos);
+			float3 Extinction = EvalExtinctionOutOfFrustum(WorldPosition, LinearVoxelPos);
 
 			// Pick a volume coefficient, based on proportional weight
 			float3 Albedo = GetAlbedo(OrthoGridUniformBuffer.ScatteringGridBuffer[LinearVoxelPos]) / Extinction;
@@ -438,7 +382,7 @@ FTrackingResult RatioTracking(float3 WorldRayOrigin, float3 WorldRayDirection, f
 		{
 			float3 WorldPosition = WorldRayOrigin + WorldRayDirection * Result.Distance;
 			int LinearVoxelPos = -1;
-			float3 Extinction = EvalExtinction(WorldPosition, LinearVoxelPos);
+			float3 Extinction = EvalExtinctionOutOfFrustum(WorldPosition, LinearVoxelPos);
 			float SigmaT = Luminance(Extinction);
 
 			float AcceptanceProbability = min(SigmaT / SigmaBar, 1);
@@ -473,7 +417,7 @@ FTrackingResult SpectralRatioTracking(float3 WorldRayOrigin, float3 WorldRayDire
 		{
 			float3 WorldPosition = WorldRayOrigin + WorldRayDirection * Result.Distance;
 			int LinearVoxelPos = -1;
-			float3 Extinction = EvalExtinction(WorldPosition, LinearVoxelPos);
+			float3 Extinction = EvalExtinctionOutOfFrustum(WorldPosition, LinearVoxelPos);
 
 			// Pick a component, based on proportional weight
 			float3 SpectralCdf = float3(Extinction.x, Extinction.x + Extinction.y, Extinction.x + Extinction.y + Extinction.z);
@@ -508,8 +452,8 @@ FTrackingResult DDABasedMajorantMarch(float3 WorldRayOrigin, float3 Direction, f
 	float WorldRayTMax = length(WorldRayEnd - WorldRayBegin);
 
 	// Transform to voxel-space
-	float3 WorldBoundsMin = GetTranslatedWorldPos(OrthoGridUniformBuffer.TopLevelGridWorldBoundsMin);
-	float3 WorldBoundsMax = GetTranslatedWorldPos(OrthoGridUniformBuffer.TopLevelGridWorldBoundsMax);
+	float3 WorldBoundsMin = OrthoGridUniformBuffer.TopLevelGridWorldBoundsMin;
+	float3 WorldBoundsMax = OrthoGridUniformBuffer.TopLevelGridWorldBoundsMax;
 	float3 TopLevelGridWorldBoundsExtent = WorldBoundsMax - WorldBoundsMin;
 	float3 VoxelRayBegin = (WorldRayBegin - WorldBoundsMin) / TopLevelGridWorldBoundsExtent * OrthoGridUniformBuffer.TopLevelGridResolution;
 	float3 VoxelRayEnd = (WorldRayEnd - WorldBoundsMin) / TopLevelGridWorldBoundsExtent * OrthoGridUniformBuffer.TopLevelGridResolution;
@@ -584,7 +528,7 @@ FTrackingResult DDABasedMajorantMarch(float3 WorldRayOrigin, float3 Direction, f
 		float3 TranslatedWorldPos = WorldRayBegin + Direction * VoxelRayMarchT * VoxelToWorldScale;
 		if (DistanceSamplingMethod == DEBUG_MODE_SAMPLE_BASED_MARCHING)
 		{
-			float StepSize = CalcStepSize(Direction);
+			float StepSize = CalcStepSizeLocal(Direction);
 			FTrackingResult Sample = SampleBasedMarch(TranslatedWorldPos, Direction, 0.0, VoxelRayMarchDeltaT * VoxelToWorldScale, StepSize);
 			TrackingResult = ConcatenateTrackingResult(TrackingResult, Sample);
 		}
@@ -595,7 +539,8 @@ FTrackingResult DDABasedMajorantMarch(float3 WorldRayOrigin, float3 Direction, f
 		}
 		else if (DistanceSamplingMethod == DEBUG_MODE_TRACKING_BASED_DELTA_TRACKING)
 		{
-			FTrackingResult Sample = SpectralWeightedDeltaTracking(TranslatedWorldPos, Direction, 0.0, VoxelRayMarchDeltaT * VoxelToWorldScale, MajorantData.Majorant, CollisionSamplingMethod, RandSequence);
+			//FTrackingResult Sample = SpectralWeightedDeltaTracking(TranslatedWorldPos, Direction, 0.0, VoxelRayMarchDeltaT * VoxelToWorldScale, MajorantData.Majorant, CollisionSamplingMethod, RandSequence);
+			FTrackingResult Sample = DeltaTracking(TranslatedWorldPos, Direction, 0.0, VoxelRayMarchDeltaT * VoxelToWorldScale, MajorantData.Majorant, RandSequence);
 			TrackingResult = ConcatenateTrackingResult(TrackingResult, Sample);
 		}
 		else
@@ -660,7 +605,7 @@ void RenderTransmittanceWithVoxelGridCS(
 #if DEBUG_MODE_SAMPLING_BASED
 	if (ClipHitT.x < ClipHitT.y)
 	{
-		float StepSize = CalcStepSize(RayDirection);
+		float StepSize = CalcStepSizeLocal(RayDirection);
 		Jitter *= StepSize;
 #if DEBUG_MODE == DEBUG_MODE_SAMPLE_BASED_MARCHING
 		FTrackingResult Result = SampleBasedMarch(RayOrigin, RayDirection, ClipHitT.x + Jitter, ClipHitT.y, StepSize);
@@ -821,4 +766,4 @@ void RenderTransmittanceWithVoxelGridCS(
 	// Output..
 	float3 Opacity = 1.0 - Transmittance;
 	RWLightingTexture[PixelCoord] += float4(Radiance, Luminance(Opacity));
-}
\ No newline at end of file
+}
diff --git a/Engine/Shaders/Private/HeterogeneousVolumes/HeterogeneousVolumesVoxelGridRenderingUtils.ush b/Engine/Shaders/Private/HeterogeneousVolumes/HeterogeneousVolumesVoxelGridRenderingUtils.ush
new file mode 100644
--- /dev/null
+++ b/Engine/Shaders/Private/HeterogeneousVolumes/HeterogeneousVolumesVoxelGridRenderingUtils.ush
@@ -0,0 +1,317 @@
+// Copyright Epic Games, Inc. All Rights Reserved.
+
+#include "HeterogeneousVolumesVoxelGridUtils.ush"
+#include "HeterogeneousVolumesFrustumVoxelGridUtils.ush"
+
+float3 EvalExtinctionOutOfFrustum(float3 WorldPos, inout int LinearBottomLevelVoxelPos)
+{
+	float3 Extinction = 0;
+
+	float3 WorldBoundsMin = OrthoGridUniformBuffer.TopLevelGridWorldBoundsMin;
+	float3 WorldBoundsMax = OrthoGridUniformBuffer.TopLevelGridWorldBoundsMax;
+	float3 WorldBoundsExtent = WorldBoundsMax - WorldBoundsMin;
+	float3 GridUV = (WorldPos - WorldBoundsMin) / WorldBoundsExtent;
+	if (all(GridUV >= 0.0) && all(GridUV <= 1.0))
+	{
+		float3 TopLevelVoxelPos = GridUV * OrthoGridUniformBuffer.TopLevelGridResolution;
+		uint LinearTopLevelVoxelPos = GetLinearIndex(TopLevelVoxelPos, OrthoGridUniformBuffer.TopLevelGridResolution);
+
+		if (OrthoGridUniformBuffer.bEnableIndirectionGrid)
+		{
+			FTopLevelGridData TopLevelGridData = OrthoGridUniformBuffer.TopLevelGridBuffer[LinearTopLevelVoxelPos];
+			if (IsBottomLevelAllocated(TopLevelGridData))
+			{
+				uint IndirectionGridIndex = GetBottomLevelIndex(TopLevelGridData);
+				uint3 IndirectionGridVoxelResolution = GetBottomLevelVoxelResolution(TopLevelGridData);
+
+				// Convert to Indirection voxel-space
+				float3 IndirectionGridVoxelPos = frac(TopLevelVoxelPos) * IndirectionGridVoxelResolution;
+				int3 IndirectionGridVoxelPosAsInt = clamp(IndirectionGridVoxelPos, 0, IndirectionGridVoxelResolution - 1);
+				FTopLevelGridData IndirectionData = OrthoGridUniformBuffer.IndirectionGridBuffer[IndirectionGridIndex + MortonEncode3(IndirectionGridVoxelPosAsInt)];
+				if (IsBottomLevelAllocated(IndirectionData))
+				{
+					uint BottomLevelIndex = GetBottomLevelIndex(IndirectionData);
+					uint3 BottomLevelVoxelResolution = GetBottomLevelVoxelResolution(IndirectionData);
+
+					// Constant interpolation
+					float3 BottomLevelVoxelPos = frac(IndirectionGridVoxelPos) * BottomLevelVoxelResolution;
+					int3 BottomLevelVoxelPosAsInt = clamp(BottomLevelVoxelPos, 0, BottomLevelVoxelResolution - 1);
+					LinearBottomLevelVoxelPos = BottomLevelIndex + MortonEncode3(BottomLevelVoxelPosAsInt);
+					Extinction = GetExtinction(OrthoGridUniformBuffer.ExtinctionGridBuffer[LinearBottomLevelVoxelPos]);
+					//Albedo = GetAlbedo(OrthoGridUniformBuffer.ScatteringGridBuffer[LinearBottomLevelVoxelPos]);
+				}
+			}
+		}
+		else
+		{
+			FTopLevelGridData TopLevelData = OrthoGridUniformBuffer.TopLevelGridBuffer[LinearTopLevelVoxelPos];
+			if (IsBottomLevelAllocated(TopLevelData))
+			{
+				uint BottomLevelIndex = GetBottomLevelIndex(OrthoGridUniformBuffer.TopLevelGridBuffer[LinearTopLevelVoxelPos]);
+				uint3 BottomLevelVoxelResolution = GetBottomLevelVoxelResolution(OrthoGridUniformBuffer.TopLevelGridBuffer[LinearTopLevelVoxelPos]);
+
+				// Constant Interpolation
+				float3 BottomLevelVoxelPos = frac(TopLevelVoxelPos) * BottomLevelVoxelResolution;
+				int3 BottomLevelVoxelPosAsInt = clamp(BottomLevelVoxelPos, 0, BottomLevelVoxelResolution - 1);
+				LinearBottomLevelVoxelPos = BottomLevelIndex + MortonEncode3(BottomLevelVoxelPosAsInt);
+
+				Extinction = GetExtinction(OrthoGridUniformBuffer.ExtinctionGridBuffer[LinearBottomLevelVoxelPos]);
+				//Albedo = GetAlbedo(OrthoGridUniformBuffer.ScatteringGridBuffer[LinearBottomLevelVoxelPos]);
+			}
+		}
+	}
+
+	return Extinction;
+}
+
+float3 EvalExtinctionInFrustum(float3 WorldPos, inout bool bInFrustum)
+{
+	float3 Extinction = 0;
+
+	// Convert WorldPos to voxel space
+	float3 ViewPos = mul(float4(WorldPos, 1), FrustumGridUniformBuffer.WorldToView).xyz;
+
+	int3 VoxelDimensions = FrustumGridUniformBuffer.VoxelDimensions;
+	float NearPlaneDepth = FrustumGridUniformBuffer.NearPlaneDepth;
+	float FarPlaneDepth = FrustumGridUniformBuffer.FarPlaneDepth;
+	float TanHalfFOV = FrustumGridUniformBuffer.TanHalfFOV;
+	float3 VoxelPos = ViewToVoxel(ViewPos, VoxelDimensions, NearPlaneDepth, FarPlaneDepth, TanHalfFOV);
+
+	bInFrustum = all(VoxelPos > 0) && all(VoxelPos < FrustumGridUniformBuffer.TopLevelFroxelGridResolution);
+	if (bInFrustum)
+	{
+		uint LinearTopLevelVoxelPos = GetLinearIndex(VoxelPos, FrustumGridUniformBuffer.TopLevelFroxelGridResolution);
+		FTopLevelGridData TopLevelGridData = FrustumGridUniformBuffer.TopLevelFroxelGridBuffer[LinearTopLevelVoxelPos];
+		if (IsBottomLevelAllocated(TopLevelGridData))
+		{
+			uint BottomLevelIndex = GetBottomLevelIndex(TopLevelGridData);
+			uint3 BottomLevelVoxelResolution = GetBottomLevelVoxelResolution(TopLevelGridData);
+
+			float3 BottomLevelVoxelPos = frac(VoxelPos) * BottomLevelVoxelResolution;
+			uint LinearBottomLevelVoxelPos = BottomLevelIndex + MortonEncode3(uint3(BottomLevelVoxelPos));
+			Extinction = GetExtinction(FrustumGridUniformBuffer.ExtinctionFroxelGridBuffer[LinearBottomLevelVoxelPos]);
+		}
+	}
+
+	return Extinction;
+}
+
+float3 EvalExtinction(float3 WorldPos)
+{
+	bool bInFrustum = false;
+	float3 Extinction = EvalExtinctionInFrustum(WorldPos, bInFrustum);
+	if (!bInFrustum)
+	{
+		int LinearVoxelPos;
+		Extinction = EvalExtinctionOutOfFrustum(WorldPos, LinearVoxelPos);
+	}
+	return Extinction;
+}
+
+float3 EvalAlbedoOutOfFrustum(float3 WorldPos, inout int LinearBottomLevelVoxelPos)
+{
+	float3 Albedo = 0;
+
+	float3 WorldBoundsMin = OrthoGridUniformBuffer.TopLevelGridWorldBoundsMin;
+	float3 WorldBoundsMax = OrthoGridUniformBuffer.TopLevelGridWorldBoundsMax;
+	float3 WorldBoundsExtent = WorldBoundsMax - WorldBoundsMin;
+	float3 GridUV = (WorldPos - WorldBoundsMin) / WorldBoundsExtent;
+	if (all(GridUV >= 0.0) && all(GridUV <= 1.0))
+	{
+		float3 TopLevelVoxelPos = GridUV * OrthoGridUniformBuffer.TopLevelGridResolution;
+		uint LinearTopLevelVoxelPos = GetLinearIndex(TopLevelVoxelPos, OrthoGridUniformBuffer.TopLevelGridResolution);
+
+		if (OrthoGridUniformBuffer.bEnableIndirectionGrid)
+		{
+			FTopLevelGridData TopLevelGridData = OrthoGridUniformBuffer.TopLevelGridBuffer[LinearTopLevelVoxelPos];
+			if (IsBottomLevelAllocated(TopLevelGridData))
+			{
+				uint IndirectionGridIndex = GetBottomLevelIndex(TopLevelGridData);
+				uint3 IndirectionGridVoxelResolution = GetBottomLevelVoxelResolution(TopLevelGridData);
+
+				// Convert to Indirection voxel-space
+				float3 IndirectionGridVoxelPos = frac(TopLevelVoxelPos) * IndirectionGridVoxelResolution;
+				int3 IndirectionGridVoxelPosAsInt = clamp(IndirectionGridVoxelPos, 0, IndirectionGridVoxelResolution - 1);
+				FTopLevelGridData IndirectionData = OrthoGridUniformBuffer.IndirectionGridBuffer[IndirectionGridIndex + MortonEncode3(IndirectionGridVoxelPosAsInt)];
+				if (IsBottomLevelAllocated(IndirectionData))
+				{
+					uint BottomLevelIndex = GetBottomLevelIndex(IndirectionData);
+					uint3 BottomLevelVoxelResolution = GetBottomLevelVoxelResolution(IndirectionData);
+
+					// Constant interpolation
+					float3 BottomLevelVoxelPos = frac(IndirectionGridVoxelPos) * BottomLevelVoxelResolution;
+					int3 BottomLevelVoxelPosAsInt = clamp(BottomLevelVoxelPos, 0, BottomLevelVoxelResolution - 1);
+					LinearBottomLevelVoxelPos = BottomLevelIndex + MortonEncode3(BottomLevelVoxelPosAsInt);
+					Albedo = GetAlbedo(OrthoGridUniformBuffer.ScatteringGridBuffer[LinearBottomLevelVoxelPos]);
+					//Albedo = GetAlbedo(OrthoGridUniformBuffer.ScatteringGridBuffer[LinearBottomLevelVoxelPos]);
+				}
+			}
+		}
+		else
+		{
+			FTopLevelGridData TopLevelData = OrthoGridUniformBuffer.TopLevelGridBuffer[LinearTopLevelVoxelPos];
+			if (IsBottomLevelAllocated(TopLevelData))
+			{
+				uint BottomLevelIndex = GetBottomLevelIndex(OrthoGridUniformBuffer.TopLevelGridBuffer[LinearTopLevelVoxelPos]);
+				uint3 BottomLevelVoxelResolution = GetBottomLevelVoxelResolution(OrthoGridUniformBuffer.TopLevelGridBuffer[LinearTopLevelVoxelPos]);
+
+				// Constant Interpolation
+				float3 BottomLevelVoxelPos = frac(TopLevelVoxelPos) * BottomLevelVoxelResolution;
+				int3 BottomLevelVoxelPosAsInt = clamp(BottomLevelVoxelPos, 0, BottomLevelVoxelResolution - 1);
+				LinearBottomLevelVoxelPos = BottomLevelIndex + MortonEncode3(BottomLevelVoxelPosAsInt);
+
+				Albedo = GetAlbedo(OrthoGridUniformBuffer.ScatteringGridBuffer[LinearBottomLevelVoxelPos]);
+				//Albedo = GetAlbedo(OrthoGridUniformBuffer.ScatteringGridBuffer[LinearBottomLevelVoxelPos]);
+			}
+		}
+	}
+
+	return Albedo;
+}
+
+float3 EvalAlbedoInFrustum(float3 WorldPos, inout bool bInFrustum)
+{
+	float3 Albedo = 0;
+
+	// Convert WorldPos to voxel space
+	float3 ViewPos = mul(float4(WorldPos, 1), FrustumGridUniformBuffer.WorldToView).xyz;
+
+	int3 VoxelDimensions = FrustumGridUniformBuffer.VoxelDimensions;
+	float NearPlaneDepth = FrustumGridUniformBuffer.NearPlaneDepth;
+	float FarPlaneDepth = FrustumGridUniformBuffer.FarPlaneDepth;
+	float TanHalfFOV = FrustumGridUniformBuffer.TanHalfFOV;
+	float3 VoxelPos = ViewToVoxel(ViewPos, VoxelDimensions, NearPlaneDepth, FarPlaneDepth, TanHalfFOV);
+
+	bInFrustum = all(VoxelPos > 0) && all(VoxelPos < FrustumGridUniformBuffer.TopLevelFroxelGridResolution);
+	if (bInFrustum)
+	{
+		uint LinearTopLevelVoxelPos = GetLinearIndex(VoxelPos, FrustumGridUniformBuffer.TopLevelFroxelGridResolution);
+		FTopLevelGridData TopLevelGridData = FrustumGridUniformBuffer.TopLevelFroxelGridBuffer[LinearTopLevelVoxelPos];
+		if (IsBottomLevelAllocated(TopLevelGridData))
+		{
+			uint BottomLevelIndex = GetBottomLevelIndex(TopLevelGridData);
+			uint3 BottomLevelVoxelResolution = GetBottomLevelVoxelResolution(TopLevelGridData);
+
+			float3 BottomLevelVoxelPos = frac(VoxelPos) * BottomLevelVoxelResolution;
+			uint LinearBottomLevelVoxelPos = BottomLevelIndex + MortonEncode3(uint3(BottomLevelVoxelPos));
+			Albedo = GetAlbedo(FrustumGridUniformBuffer.ScatteringFroxelGridBuffer[LinearBottomLevelVoxelPos]);
+		}
+	}
+
+	return Albedo;
+}
+
+float3 EvalAlbedo(float3 WorldPos)
+{
+	bool bInFrustum = false;
+	float3 Albedo = EvalAlbedoInFrustum(WorldPos, bInFrustum);
+	if (!bInFrustum)
+	{
+		int LinearVoxelPos;
+		Albedo = EvalAlbedoOutOfFrustum(WorldPos, LinearVoxelPos);
+	}
+	return Albedo;
+}
+
+float3 EvalEmissionOutOfFrustum(float3 WorldPos, inout int LinearBottomLevelVoxelPos)
+{
+	float3 Emission = 0;
+
+	float3 WorldBoundsMin = OrthoGridUniformBuffer.TopLevelGridWorldBoundsMin;
+	float3 WorldBoundsMax = OrthoGridUniformBuffer.TopLevelGridWorldBoundsMax;
+	float3 WorldBoundsExtent = WorldBoundsMax - WorldBoundsMin;
+	float3 GridUV = (WorldPos - WorldBoundsMin) / WorldBoundsExtent;
+	if (all(GridUV >= 0.0) && all(GridUV <= 1.0))
+	{
+		float3 TopLevelVoxelPos = GridUV * OrthoGridUniformBuffer.TopLevelGridResolution;
+		uint LinearTopLevelVoxelPos = GetLinearIndex(TopLevelVoxelPos, OrthoGridUniformBuffer.TopLevelGridResolution);
+
+		if (OrthoGridUniformBuffer.bEnableIndirectionGrid)
+		{
+			FTopLevelGridData TopLevelGridData = OrthoGridUniformBuffer.TopLevelGridBuffer[LinearTopLevelVoxelPos];
+			if (IsBottomLevelAllocated(TopLevelGridData))
+			{
+				uint IndirectionGridIndex = GetBottomLevelIndex(TopLevelGridData);
+				uint3 IndirectionGridVoxelResolution = GetBottomLevelVoxelResolution(TopLevelGridData);
+
+				// Convert to Indirection voxel-space
+				float3 IndirectionGridVoxelPos = frac(TopLevelVoxelPos) * IndirectionGridVoxelResolution;
+				int3 IndirectionGridVoxelPosAsInt = clamp(IndirectionGridVoxelPos, 0, IndirectionGridVoxelResolution - 1);
+				FTopLevelGridData IndirectionData = OrthoGridUniformBuffer.IndirectionGridBuffer[IndirectionGridIndex + MortonEncode3(IndirectionGridVoxelPosAsInt)];
+				if (IsBottomLevelAllocated(IndirectionData))
+				{
+					uint BottomLevelIndex = GetBottomLevelIndex(IndirectionData);
+					uint3 BottomLevelVoxelResolution = GetBottomLevelVoxelResolution(IndirectionData);
+
+					// Constant interpolation
+					float3 BottomLevelVoxelPos = frac(IndirectionGridVoxelPos) * BottomLevelVoxelResolution;
+					int3 BottomLevelVoxelPosAsInt = clamp(BottomLevelVoxelPos, 0, BottomLevelVoxelResolution - 1);
+					LinearBottomLevelVoxelPos = BottomLevelIndex + MortonEncode3(BottomLevelVoxelPosAsInt);
+					Emission = GetEmission(OrthoGridUniformBuffer.EmissionGridBuffer[LinearBottomLevelVoxelPos]);
+				}
+			}
+		}
+		else
+		{
+			FTopLevelGridData TopLevelData = OrthoGridUniformBuffer.TopLevelGridBuffer[LinearTopLevelVoxelPos];
+			if (IsBottomLevelAllocated(TopLevelData))
+			{
+				uint BottomLevelIndex = GetBottomLevelIndex(OrthoGridUniformBuffer.TopLevelGridBuffer[LinearTopLevelVoxelPos]);
+				uint3 BottomLevelVoxelResolution = GetBottomLevelVoxelResolution(OrthoGridUniformBuffer.TopLevelGridBuffer[LinearTopLevelVoxelPos]);
+
+				// Constant Interpolation
+				float3 BottomLevelVoxelPos = frac(TopLevelVoxelPos) * BottomLevelVoxelResolution;
+				int3 BottomLevelVoxelPosAsInt = clamp(BottomLevelVoxelPos, 0, BottomLevelVoxelResolution - 1);
+				LinearBottomLevelVoxelPos = BottomLevelIndex + MortonEncode3(BottomLevelVoxelPosAsInt);
+
+				Emission = GetEmission(OrthoGridUniformBuffer.EmissionGridBuffer[LinearBottomLevelVoxelPos]);
+			}
+		}
+	}
+
+	return Emission;
+}
+
+float3 EvalEmissionInFrustum(float3 WorldPos, inout bool bInFrustum)
+{
+	float3 Emission = 0;
+
+	// Convert WorldPos to voxel space
+	float3 ViewPos = mul(float4(WorldPos, 1), FrustumGridUniformBuffer.WorldToView).xyz;
+
+	int3 VoxelDimensions = FrustumGridUniformBuffer.VoxelDimensions;
+	float NearPlaneDepth = FrustumGridUniformBuffer.NearPlaneDepth;
+	float FarPlaneDepth = FrustumGridUniformBuffer.FarPlaneDepth;
+	float TanHalfFOV = FrustumGridUniformBuffer.TanHalfFOV;
+	float3 VoxelPos = ViewToVoxel(ViewPos, VoxelDimensions, NearPlaneDepth, FarPlaneDepth, TanHalfFOV);
+
+	bInFrustum = all(VoxelPos > 0) && all(VoxelPos < FrustumGridUniformBuffer.TopLevelFroxelGridResolution);
+	if (bInFrustum)
+	{
+		uint LinearTopLevelVoxelPos = GetLinearIndex(VoxelPos, FrustumGridUniformBuffer.TopLevelFroxelGridResolution);
+		FTopLevelGridData TopLevelGridData = FrustumGridUniformBuffer.TopLevelFroxelGridBuffer[LinearTopLevelVoxelPos];
+		if (IsBottomLevelAllocated(TopLevelGridData))
+		{
+			uint BottomLevelIndex = GetBottomLevelIndex(TopLevelGridData);
+			uint3 BottomLevelVoxelResolution = GetBottomLevelVoxelResolution(TopLevelGridData);
+
+			float3 BottomLevelVoxelPos = frac(VoxelPos) * BottomLevelVoxelResolution;
+			uint LinearBottomLevelVoxelPos = BottomLevelIndex + MortonEncode3(uint3(BottomLevelVoxelPos));
+			Emission = GetEmission(FrustumGridUniformBuffer.EmissionFroxelGridBuffer[LinearBottomLevelVoxelPos]);
+		}
+	}
+
+	return Emission;
+}
+
+float3 EvalEmission(float3 WorldPos)
+{
+	bool bInFrustum = false;
+	float3 Emission = EvalEmissionInFrustum(WorldPos, bInFrustum);
+	if (!bInFrustum)
+	{
+		int LinearVoxelPos;
+		Emission = EvalEmissionOutOfFrustum(WorldPos, LinearVoxelPos);
+	}
+	return Emission;
+}
\ No newline at end of file
diff --git a/Engine/Shaders/Private/HeterogeneousVolumes/HeterogeneousVolumesVoxelGridShadows.usf b/Engine/Shaders/Private/HeterogeneousVolumes/HeterogeneousVolumesVoxelGridShadows.usf
new file mode 100644
--- /dev/null
+++ b/Engine/Shaders/Private/HeterogeneousVolumes/HeterogeneousVolumesVoxelGridShadows.usf
@@ -0,0 +1,542 @@
+// Copyright Epic Games, Inc. All Rights Reserved.
+
+#include "HeterogeneousVolumesVoxelGridTypes.ush"
+#include "HeterogeneousVolumesAdaptiveVolumetricShadowMapTypes.ush"
+
+#include "../Common.ush"
+#include "../DeferredShadingCommon.ush"
+#include "../ComputeShaderUtils.ush"
+#include "HeterogeneousVolumesTracingUtils.ush"
+#include "HeterogeneousVolumesVoxelGridTraversal.ush"
+#include "HeterogeneousVolumesVoxelGridRenderingUtils.ush"
+
+#include "../PathTracing/Utilities/PathTracingRandomSequence.ush"
+
+// Ray data
+int bJitter;
+
+// Dispatch data
+int3 GroupCount;
+int DebugMode;
+
+int MaxStepCount;
+
+int NumShadowMatrices;
+float4x4 TranslatedWorldToShadow[6];
+float4x4 ShadowToTranslatedWorld[6];
+float3 TranslatedWorldOrigin;
+
+// Lighting
+int bApplyEmissionAndTransmittance;
+int bApplyDirectLighting;
+int bApplyShadowTransmittance;
+int LightType;
+float VolumetricScatteringIntensity;
+
+// Beer shadow map
+Texture2D<float4> BeerShadowMapTexture;
+int2 ShadowResolution;
+
+RWTexture2D<float4> RWBeerShadowMapTexture;
+
+struct FVolumeSampleContext
+{
+	float3 WorldPosition;
+};
+
+float4x4 IdentityMatrix()
+{
+	return float4x4(
+		1, 0, 0, 0,
+		0, 1, 0, 0,
+		0, 0, 1, 0,
+		0, 0, 0, 1);
+}
+
+float4x4 GetLocalToWorld()
+{
+	return IdentityMatrix();
+}
+
+float4x4 GetWorldToLocal()
+{
+	return IdentityMatrix();
+}
+
+float3 GetLocalBoundsMin()
+{
+	float3 LocalBoundsMin = 0;
+	if (OrthoGridUniformBuffer.bUseOrthoGrid)
+	{
+		LocalBoundsMin = OrthoGridUniformBuffer.TopLevelGridWorldBoundsMin;
+	}
+	else if (FrustumGridUniformBuffer.bUseFrustumGrid)
+	{
+		LocalBoundsMin = FrustumGridUniformBuffer.TopLevelGridWorldBoundsMin;
+	}
+	return LocalBoundsMin;
+}
+
+float3 GetLocalBoundsMax()
+{
+	float3 LocalBoundsMax = 0;
+	if (OrthoGridUniformBuffer.bUseOrthoGrid)
+	{
+		LocalBoundsMax = OrthoGridUniformBuffer.TopLevelGridWorldBoundsMax;
+	}
+	else if (FrustumGridUniformBuffer.bUseFrustumGrid)
+	{
+		LocalBoundsMax = FrustumGridUniformBuffer.TopLevelGridWorldBoundsMax;
+	}
+	return LocalBoundsMax;
+}
+
+float3 GetLocalBoundsOrigin()
+{
+	return (GetLocalBoundsMin() + GetLocalBoundsMax()) / 2.0;
+}
+
+float3 GetLocalBoundsExtent()
+{
+	return (GetLocalBoundsMax() - GetLocalBoundsMin()) / 2.0;
+}
+
+float GetMaxTraceDistance()
+{
+	return 30000;
+}
+
+float GetMaxShadowTraceDistance()
+{
+	return 30000;
+}
+
+float GetStepSize()
+{
+	return 1.0;
+}
+
+float GetStepFactor()
+{
+	return 1.0;
+}
+
+float GetShadowStepSize()
+{
+	return 1.0;
+}
+
+float GetShadowStepFactor()
+{
+	return 1.0;
+}
+
+uint3 GetVolumeResolution()
+{
+	return OrthoGridUniformBuffer.TopLevelGridResolution;
+}
+
+uint GetNumVoxels()
+{
+	return GetVolumeResolution().x * GetVolumeResolution().y * GetVolumeResolution().z;
+}
+
+FVolumeSampleContext CreateVolumeSampleContext(float3 LocalPosition, float3 WorldPosition, float MipLevel)
+{
+	FVolumeSampleContext VolumeSampleContext;
+	VolumeSampleContext.WorldPosition = WorldPosition;
+
+	return VolumeSampleContext;
+}
+
+float3 SampleExtinction(inout FVolumeSampleContext Context)
+{
+	return EvalExtinction(Context.WorldPosition);
+}
+
+float3 SampleEmission(inout FVolumeSampleContext Context)
+{
+	return EvalEmission(Context.WorldPosition);
+}
+
+float3 SampleAlbedo(inout FVolumeSampleContext Context)
+{
+	float3 Albedo = 0;
+
+	float3 Extinction = EvalExtinction(Context.WorldPosition);
+	if (any(Extinction > 0))
+	{
+		float3 Scattering = EvalAlbedo(Context.WorldPosition);
+		Albedo = Scattering / Extinction;
+	}
+
+	return Albedo;
+}
+
+#define SUPPORT_CONTACT_SHADOWS 0
+
+#define DYNAMICALLY_SHADOWED 1
+#define TREAT_MAXDEPTH_UNSHADOWED 1
+
+#define SHADOW_QUALITY 2
+#define NO_TRANSLUCENCY_AVAILABLE
+
+#include "HeterogeneousVolumesLightingUtils.ush"
+#include "HeterogeneousVolumesRayMarchingUtils.ush"
+#include "HeterogeneousVolumesAdaptiveVolumetricShadowMapUtils.ush"
+
+int MaxSampleCount;
+float AbsoluteErrorThreshold;
+float RelativeErrorThreshold;
+
+RWBuffer<int> RWVolumetricShadowLinkedListAllocatorBuffer;
+RWStructuredBuffer<int2> RWVolumetricShadowLinkedListBuffer;
+
+// Debug
+struct FVolumetricShadowMapDebugData
+{
+	float3 LightRayStart;
+	float3 LightRayEnd;
+	float3 RayOrigin;
+	float3 RayEnd;
+	float2 HitSpan;
+};
+
+RWStructuredBuffer<FVolumetricShadowMapDebugData> RWDebugBuffer;
+
+uint GetLinearIndex(uint Face, uint2 PixelCoord)
+{
+	uint LinearIndex = Face * ShadowResolution.x * ShadowResolution.y + PixelCoord.y * ShadowResolution.x + PixelCoord.x;
+	return LinearIndex;
+}
+
+[numthreads(THREADGROUP_SIZE_2D, THREADGROUP_SIZE_2D, 1)]
+void RenderVolumetricShadowMapForLightWithVoxelGridCS(
+	uint2 GroupThreadId : SV_GroupThreadID,
+	uint3 DispatchThreadId : SV_DispatchThreadID
+)
+{
+	// Create screen ray
+	if (any(DispatchThreadId.xy >= ShadowResolution))
+	{
+		return;
+	}
+	uint Face = DispatchThreadId.z;
+	uint2 PixelCoord = DispatchThreadId.xy;
+	uint LinearIndex = GetLinearIndex(Face, PixelCoord);
+
+	// Generate a light ray
+	float2 ShadowUV = float2(PixelCoord + 0.5f) / float2(ShadowResolution);
+	const float NearScreenZ = 1.0;
+	const float FarScreenZ = 0.0;
+
+	float4 LightRayStartAsFloat4 = mul(float4(ShadowUV, NearScreenZ, 1), ShadowToTranslatedWorld[Face]);
+	float4 LightRayEndAsFloat4 = mul(float4(ShadowUV, FarScreenZ, 1), ShadowToTranslatedWorld[Face]);
+	float3 LightRayStart = LightRayStartAsFloat4.xyz * rcp(LightRayStartAsFloat4.w);
+	float3 LightRayEnd = LightRayEndAsFloat4.xyz * rcp(LightRayEndAsFloat4.w);
+
+	// Intersect ray with bounding volume
+	// TODO: LWC integration..
+	float3 WorldRayOrigin = DFFastSubtractDemote(LightRayStart, PrimaryView.PreViewTranslation);
+	float3 WorldRayEnd = DFFastSubtractDemote(LightRayEnd, PrimaryView.PreViewTranslation);
+
+	float3 WorldRayDirection = WorldRayEnd - WorldRayOrigin;
+	float WorldRayLength = length(WorldRayDirection);
+	WorldRayDirection /= WorldRayLength;
+
+	float3 WorldBoundsMin = FrustumGridUniformBuffer.TopLevelGridWorldBoundsMin;
+	float3 WorldBoundsMax = FrustumGridUniformBuffer.TopLevelGridWorldBoundsMax;
+	if (OrthoGridUniformBuffer.bUseOrthoGrid)
+	{
+		WorldBoundsMin = OrthoGridUniformBuffer.TopLevelGridWorldBoundsMin;
+		WorldBoundsMax = OrthoGridUniformBuffer.TopLevelGridWorldBoundsMax;
+	}
+
+	float2 RayHitT = IntersectAABB(
+		WorldRayOrigin,
+		WorldRayDirection,
+		0.0,
+		WorldRayLength,
+		WorldBoundsMin,
+		WorldBoundsMax
+	);
+
+	float LocalHitT = RayHitT.x;
+	float WorldHitT = RayHitT.x;
+	float3 Transmittance = 1;
+	int PixelIndex = LinearIndex * MaxSampleCount;
+
+	float HitSpan = RayHitT.y - RayHitT.x;
+	float Depth = 1.0;
+	if (HitSpan > 0.0)
+	{
+		float Jitter = bJitter ? InterleavedGradientNoise(PixelCoord, View.StateFrameIndexMod8) : 0.0;
+
+		FRayMarchingContext RayMarchingContext = CreateRayMarchingContext(
+			// Local-space
+			WorldRayOrigin,
+			WorldRayDirection,
+			RayHitT.x,
+			RayHitT.y,
+			// World-space
+			WorldRayOrigin,
+			WorldRayDirection,
+			// Ray-step attributes
+			Jitter,
+			CalcStepSize(WorldRayDirection),
+			MaxStepCount,
+			//MaxSampleCount, // Replace max-steps with max-samples..
+			bApplyEmissionAndTransmittance,
+			bApplyDirectLighting,
+			bApplyShadowTransmittance
+		);
+
+		RandomSequence RandSequence;
+		RandomSequence_Initialize(RandSequence, LinearIndex, View.StateFrameIndex);
+		RayMarchingContext.RandSequence = RandSequence;
+		uint StepCount = CalcStepCount(RayMarchingContext);
+		int EntryCount = 1;
+
+		float3 ExpectedExtinction = 0.0;
+		float3 ExpectedTransmittance = 1.0;
+		int PrevElementIndex = AVSM_NULL_PTR;
+		int ElementIndex = 0;
+
+		for (uint StepIndex = 0; StepIndex < StepCount; ++StepIndex)
+		{
+			LocalHitT = min(RayMarchingContext.LocalRayTMin + RayMarchingContext.StepSize * (RayMarchingContext.Jitter + StepIndex), RayMarchingContext.LocalRayTMax);
+			WorldHitT = LocalHitT * RayMarchingContext.LocalToWorldScale;
+			float3 LocalPosition = RayMarchingContext.LocalRayOrigin + RayMarchingContext.LocalRayDirection * LocalHitT;
+			float3 WorldPosition = RayMarchingContext.WorldRayOrigin + RayMarchingContext.WorldRayDirection * WorldHitT;
+
+			FVolumeSampleContext SampleContext = CreateVolumeSampleContext(LocalPosition, WorldPosition, RayMarchingContext.MipLevel);
+			float3 Extinction = SampleExtinction(SampleContext);
+			Transmittance *= exp(-Extinction * RayMarchingContext.StepSize);
+
+			bool bFirstEntry = EntryCount == 1;
+			if (bFirstEntry && any(Extinction > 0.0))
+			{
+				RWVolumetricShadowLinkedListBuffer[PixelIndex + ElementIndex].x = AVSM_PackSampleData(WorldHitT, Transmittance);
+
+				// Allocate next entry
+				int NextElementIndex = ElementIndex + 1;
+				AVSM_SetPrevPtr(RWVolumetricShadowLinkedListBuffer[PixelIndex + ElementIndex], PrevElementIndex);
+				AVSM_SetNextPtr(RWVolumetricShadowLinkedListBuffer[PixelIndex + ElementIndex], NextElementIndex);
+
+				PrevElementIndex = ElementIndex;
+				ElementIndex = NextElementIndex;
+				EntryCount++;
+
+				// Reset expected values
+				ExpectedExtinction = Extinction;
+				ExpectedTransmittance = Transmittance;
+			}
+			else if (EntryCount > 1 && any(Extinction > 0.0))
+			{
+				// If extrapolated transmittance is beyond relative error tolerance
+				ExpectedTransmittance *= exp(-ExpectedExtinction * RayMarchingContext.StepSize);
+				float AbsoluteError = length(Transmittance - ExpectedTransmittance);
+				float RelativeError = length((Transmittance - ExpectedTransmittance) / ExpectedTransmittance);
+				if ((AbsoluteError > AbsoluteErrorThreshold) && (RelativeError > RelativeErrorThreshold))
+				{
+					// Create vertex at current location
+					if (EntryCount < MaxSampleCount)
+					{
+						// Allocate next entry
+						int NextElementIndex = ElementIndex + 1;
+						RWVolumetricShadowLinkedListBuffer[PixelIndex + ElementIndex].x = AVSM_PackSampleData(WorldHitT, Transmittance);
+						AVSM_SetPrevPtr(RWVolumetricShadowLinkedListBuffer[PixelIndex + ElementIndex], PrevElementIndex);
+						AVSM_SetNextPtr(RWVolumetricShadowLinkedListBuffer[PixelIndex + ElementIndex], NextElementIndex);
+
+						PrevElementIndex = ElementIndex;
+						ElementIndex = NextElementIndex;
+						EntryCount++;
+
+						// Reset expected values
+						ExpectedExtinction = Extinction;
+						ExpectedTransmittance = Transmittance;
+					}
+#if USE_AVSM_COMPRESSION
+					// Or replace an existing vertex that contributes the lowest error to the approximation
+					else
+					{
+						int LowestErrorIndex = 0;
+						float LowestError = 1.0e6;
+						for (int Index = 1; Index < MaxSampleCount - 1; ++Index)
+						{
+							int PrevIndex = AVSM_GetPrevPtr(RWVolumetricShadowLinkedListBuffer[PixelIndex + Index]);
+							int NextIndex = AVSM_GetNextPtr(RWVolumetricShadowLinkedListBuffer[PixelIndex + Index]);
+
+							FAVSMSampleData Data[3] = {
+								AVSM_UnpackSampleData(RWVolumetricShadowLinkedListBuffer[PixelIndex + PrevIndex]),
+								AVSM_UnpackSampleData(RWVolumetricShadowLinkedListBuffer[PixelIndex + Index]),
+								AVSM_UnpackSampleData(RWVolumetricShadowLinkedListBuffer[PixelIndex + NextIndex])
+							};
+
+							float Error = abs(
+								(Data[2].X - Data[0].X) * (Data[2].Tau - Data[1].Tau) -
+								(Data[2].X - Data[1].X) * (Data[2].Tau - Data[0].Tau)
+							);
+
+							if (Error < LowestError)
+							{
+								LowestError = Error;
+								LowestErrorIndex = Index;
+							}
+						}
+
+						// Remove control point and replace with current control point
+						if (LowestErrorIndex > 0)
+						{
+							// Remove
+							int PrevIndex = AVSM_GetPrevPtr(RWVolumetricShadowLinkedListBuffer[PixelIndex + LowestErrorIndex]);
+							int NextIndex = AVSM_GetNextPtr(RWVolumetricShadowLinkedListBuffer[PixelIndex + LowestErrorIndex]);
+
+							AVSM_SetNextPtr(RWVolumetricShadowLinkedListBuffer[PixelIndex + PrevIndex], NextIndex);
+							AVSM_SetPrevPtr(RWVolumetricShadowLinkedListBuffer[PixelIndex + NextIndex], PrevIndex);
+
+							// Replace
+							RWVolumetricShadowLinkedListBuffer[PixelIndex + LowestErrorIndex].x = AVSM_PackSampleData(WorldHitT, Transmittance);
+							AVSM_SetPrevPtr(RWVolumetricShadowLinkedListBuffer[PixelIndex + LowestErrorIndex], PrevElementIndex);
+							AVSM_SetNextPtr(RWVolumetricShadowLinkedListBuffer[PixelIndex + LowestErrorIndex], MaxSampleCount - 1);
+
+							// Update neighborhood
+							AVSM_SetPrevPtr(RWVolumetricShadowLinkedListBuffer[PixelIndex + MaxSampleCount - 1], LowestErrorIndex);
+							AVSM_SetNextPtr(RWVolumetricShadowLinkedListBuffer[PixelIndex + PrevElementIndex], LowestErrorIndex);
+
+							PrevElementIndex = LowestErrorIndex;
+							ElementIndex = MaxSampleCount - 1;
+
+							// Reset expected values
+							ExpectedExtinction = Extinction;
+							ExpectedTransmittance = Transmittance;
+						}
+					}
+#endif // USE_AVSM_COMPRESSION
+				}
+			}
+
+			// Early termination
+			float Epsilon = 1.0e-7;
+			if (all(Transmittance < Epsilon))
+			{
+				Transmittance = 0.0;
+				break;
+			}
+		}
+
+		RWVolumetricShadowLinkedListBuffer[PixelIndex + ElementIndex].x = AVSM_PackSampleData(WorldHitT, Transmittance);
+		AVSM_SetPrevPtr(RWVolumetricShadowLinkedListBuffer[PixelIndex + ElementIndex], PrevElementIndex);
+		AVSM_SetNextPtr(RWVolumetricShadowLinkedListBuffer[PixelIndex + ElementIndex], AVSM_NULL_PTR);
+	}
+	else
+	{
+		RWVolumetricShadowLinkedListBuffer[PixelIndex].x = AVSM_PackSampleData(0.0, Transmittance);
+		AVSM_SetPrevPtr(RWVolumetricShadowLinkedListBuffer[PixelIndex], AVSM_NULL_PTR);
+		AVSM_SetNextPtr(RWVolumetricShadowLinkedListBuffer[PixelIndex], AVSM_NULL_PTR);
+	}
+
+#if 0
+	RWDebugBuffer[LinearIndex].LightRayStart = LightRayStart;
+	RWDebugBuffer[LinearIndex].LightRayEnd = LightRayEnd;
+	RWDebugBuffer[LinearIndex].RayOrigin = WorldRayOrigin;
+	RWDebugBuffer[LinearIndex].RayEnd = WorldRayEnd;
+	RWDebugBuffer[LinearIndex].HitSpan = float2(HitSpan, Depth);
+#endif
+	if (Face == View.GeneralPurposeTweak)
+	{
+		RWBeerShadowMapTexture[PixelCoord] = float4(Transmittance, 1);
+	}
+}
+
+StructuredBuffer<int2> VolumetricShadowLinkedListBuffer;
+
+RWBuffer<int> RWVolumetricShadowIndirectionAllocatorBuffer;
+RWStructuredBuffer<uint2> RWVolumetricShadowIndirectionBuffer;
+RWStructuredBuffer<uint> RWVolumetricShadowTransmittanceBuffer;
+
+groupshared int GSWarpSampleCount[THREADGROUP_SIZE_2D * THREADGROUP_SIZE_2D];
+groupshared int GSWarpPrefixSum[THREADGROUP_SIZE_2D * THREADGROUP_SIZE_2D];
+groupshared int GSWarpOffset;
+
+[numthreads(THREADGROUP_SIZE_2D, THREADGROUP_SIZE_2D, 1)]
+void CompressVolumetricShadowMapCS(
+	uint2 GroupThreadId : SV_GroupThreadID,
+	uint3 DispatchThreadId : SV_DispatchThreadID,
+	uint GroupIndex : SV_GroupIndex
+)
+{
+	GSWarpSampleCount[GroupIndex] = 0;
+	GSWarpOffset = AVSM_INVALID_PIXEL_OFFSET;
+
+	// TODO: Early return, dependent on parallelizing pre-fix sum
+#if 0
+	if (any(DispatchThreadId >= ShadowResolution))
+	{
+		return;
+	}
+#endif
+	uint Face = DispatchThreadId.z;
+	uint2 PixelCoord = DispatchThreadId.xy;
+	uint LinearIndex = GetLinearIndex(Face, PixelCoord);
+	int PixelIndex = LinearIndex * MaxSampleCount;
+
+	// Calculate sample count
+	int SampleCount = 0;
+	FAVSMSampleData DataBegin, DataEnd;
+	if (all(PixelCoord < ShadowResolution))
+	{
+		DataBegin = AVSM_UnpackSampleData(VolumetricShadowLinkedListBuffer[PixelIndex]);
+		DataEnd = DataBegin;
+		int NextPtr = AVSM_GetNextPtr(VolumetricShadowLinkedListBuffer[PixelIndex]);
+		SampleCount = (NextPtr == AVSM_NULL_PTR) ? 0 : 1;
+		while (NextPtr != AVSM_NULL_PTR)
+		{
+			SampleCount++;
+			DataEnd = AVSM_UnpackSampleData(VolumetricShadowLinkedListBuffer[PixelIndex + NextPtr]);
+			NextPtr = AVSM_GetNextPtr(VolumetricShadowLinkedListBuffer[PixelIndex + NextPtr]);
+		}
+	}
+	GSWarpSampleCount[GroupIndex] = SampleCount;
+	GroupMemoryBarrierWithGroupSync();
+
+	// TODO: Parallel Prefix sum
+	int PrefixSum = 0;
+	for (int ThreadIndex = 0; ThreadIndex < GroupIndex; ++ThreadIndex)
+	{
+		PrefixSum += GSWarpSampleCount[ThreadIndex];
+	}
+	GSWarpPrefixSum[GroupIndex] = PrefixSum;
+	GroupMemoryBarrierWithGroupSync();
+
+	// Allocate for entire threadgroup
+	uint LastIndex = THREADGROUP_SIZE_2D * THREADGROUP_SIZE_2D - 1;
+	uint TotalSampleCount = GSWarpPrefixSum[LastIndex] + GSWarpSampleCount[LastIndex];
+	if ((GroupIndex == 0) && (TotalSampleCount > 0))
+	{
+		InterlockedAdd(RWVolumetricShadowIndirectionAllocatorBuffer[0], TotalSampleCount, GSWarpOffset);
+	}
+	GroupMemoryBarrierWithGroupSync();
+
+	// Pack and fill data per thread
+	FAVSMIndirectionDataContext Context = AVSM_CreateIndirectionDataContext(DataBegin, DataEnd);
+	int NextPtr = 0;
+	for (int SampleIndex = 0; SampleIndex < SampleCount; ++SampleIndex)
+	{
+		FAVSMSampleData ShadowData = AVSM_UnpackSampleData(VolumetricShadowLinkedListBuffer[PixelIndex + NextPtr]);
+		NextPtr = AVSM_GetNextPtr(VolumetricShadowLinkedListBuffer[PixelIndex + NextPtr]);
+
+		RWVolumetricShadowTransmittanceBuffer[GSWarpOffset + GSWarpPrefixSum[GroupIndex] + SampleIndex] = AVSM_PackSampleData(ShadowData);
+	}
+
+	// Write indirection entry
+	if (all(PixelCoord < ShadowResolution))
+	{
+		FAVSMIndirectionData IndirectionData = AVSM_CreateIndirectionData(
+			Context,
+			GSWarpOffset + GSWarpPrefixSum[GroupIndex],
+			SampleCount
+		);
+		RWVolumetricShadowIndirectionBuffer[LinearIndex] = AVSM_PackIndirectionData(IndirectionData);
+	}
+}
diff --git a/Engine/Shaders/Private/HeterogeneousVolumes/HeterogeneousVolumesVoxelGridUtils.ush b/Engine/Shaders/Private/HeterogeneousVolumes/HeterogeneousVolumesVoxelGridUtils.ush
--- a/Engine/Shaders/Private/HeterogeneousVolumes/HeterogeneousVolumesVoxelGridUtils.ush
+++ b/Engine/Shaders/Private/HeterogeneousVolumes/HeterogeneousVolumesVoxelGridUtils.ush
@@ -1,5 +1,7 @@
 // Copyright Epic Games, Inc. All Rights Reserved.
 
+#pragma once
+
 #include "../MortonCode.ush"
 
 #define EMPTY_VOXEL_INDEX 0x1FFFFFFF
diff --git a/Engine/Shaders/Private/TranslucentLightInjectionShaders.usf b/Engine/Shaders/Private/TranslucentLightInjectionShaders.usf
--- a/Engine/Shaders/Private/TranslucentLightInjectionShaders.usf
+++ b/Engine/Shaders/Private/TranslucentLightInjectionShaders.usf
@@ -4,6 +4,7 @@
  * TranslucentLightInjectionShaders.usf: Shaders for calculating lighting in a volume to use on translucency
  */
 
+#include "HeterogeneousVolumes/HeterogeneousVolumesAdaptiveVolumetricShadowMapTypes.ush"
 #include "Common.ush"
 #include "SHCommon.ush"
 
@@ -33,6 +34,14 @@ uint AtmospherePerPixelTransmittanceEnabled;
 #include "VirtualShadowMaps/VirtualShadowMapProjectionCommon.ush"
 #endif
 
+#ifndef ADAPTIVE_VOLUMETRIC_SHADOW_MAP
+#define ADAPTIVE_VOLUMETRIC_SHADOW_MAP 0
+#endif // ADAPTIVE_VOLUMETRIC_SHADOW_MAP
+
+#if ADAPTIVE_VOLUMETRIC_SHADOW_MAP
+#include "HeterogeneousVolumes/HeterogeneousVolumesAdaptiveVolumetricShadowMapSampling.ush"
+#endif // ADAPTIVE_VOLUMETRIC_SHADOW_MAP
+
 float CalcSimpleLightAttenuation(float3 TranslatedWorldPosition)
 {
 	float3 WorldLightVector = SimpleLightPositionAndRadius.xyz - TranslatedWorldPosition;
@@ -204,6 +213,10 @@ void InjectMainPS(
 		ShadowFactor *= VirtualShadowMapSample.ShadowFactor;
 #endif
 
+#if ADAPTIVE_VOLUMETRIC_SHADOW_MAP
+		ShadowFactor *= AVSM_SampleTransmittance(TranslatedWorldPositionForLighting, DeferredLightUniforms.TranslatedWorldPosition);
+#endif // ADAPTIVE_VOLUMETRIC_SHADOW_MAP
+
 		// Apply light function (it will also fade out at the edge of the volume so overall dimming of light using light function is not a recommended workflow)
 		ShadowFactor *= GetLightFunctionShadowFactor(TranslatedWorldPositionForLighting);
 
diff --git a/Engine/Source/Runtime/Engine/Classes/Engine/RendererSettings.h b/Engine/Source/Runtime/Engine/Classes/Engine/RendererSettings.h
--- a/Engine/Source/Runtime/Engine/Classes/Engine/RendererSettings.h
+++ b/Engine/Source/Runtime/Engine/Classes/Engine/RendererSettings.h
@@ -994,6 +994,18 @@ class URendererSettings : public UDeveloperSettings
 		ConfigRestartRequired = false))
 	uint32 bEnableHeterogeneousVolumes : 1;
 
+	UPROPERTY(config, EditAnywhere, Category = Experimental, meta = (
+		ConsoleVariable = "r.Translucency.HeterogeneousVolumes", DisplayName = "Composite Heterogeneous Volumes with Translucency",
+		ToolTip = "Enable compositing with heterogeneous volumes when rendering translucency.",
+		ConfigRestartRequired = true))
+	uint32 bCompositeHeterogeneousVolumesWithTranslucency : 1;
+
+	UPROPERTY(config, EditAnywhere, Category = Experimental, meta = (
+		ConsoleVariable = "r.HeterogeneousVolumes.Shadows", DisplayName = "Enable shadow-casting with Heterogeneous Volumes",
+		ToolTip = "Enable heterogeneous volumes to cast shadows onto the environment.",
+		ConfigRestartRequired = true))
+		uint32 bShouldHeterogeneousVolumesCastShadows : 1;
+
 	UPROPERTY(config, EditAnywhere, Category=Editor, meta=(
 		ConsoleVariable="r.WireframeCullThreshold",DisplayName="Wireframe Cull Threshold",
 		ToolTip="Screen radius at which wireframe objects are culled. Larger values can improve performance when viewing a scene in wireframe."))
diff --git a/Engine/Source/Runtime/Engine/Private/Components/HeterogeneousVolumeComponent.cpp b/Engine/Source/Runtime/Engine/Private/Components/HeterogeneousVolumeComponent.cpp
--- a/Engine/Source/Runtime/Engine/Private/Components/HeterogeneousVolumeComponent.cpp
+++ b/Engine/Source/Runtime/Engine/Private/Components/HeterogeneousVolumeComponent.cpp
@@ -65,7 +65,11 @@ FHeterogeneousVolumeSceneProxy::FHeterogeneousVolumeSceneProxy(UHeterogeneousVol
 	: FPrimitiveSceneProxy(InComponent)
 	, MaterialInterface(InComponent->GetMaterial(0))
 	, VertexFactory(GetScene().GetFeatureLevel(), "FHeterogeneousVolumeSceneProxy")
+#if ACTOR_HAS_LABELS
+	, HeterogeneousVolumeData(this, InComponent->GetReadableName())
+#else
 	, HeterogeneousVolumeData(this)
+#endif
 {
 	bIsHeterogeneousVolume = true;
 
@@ -164,6 +168,11 @@ void FHeterogeneousVolumeSceneProxy::GetDynamicMeshElements(
 	uint32 VisibilityMap,
 	FMeshElementCollector& Collector) const
 {
+	if (Views.IsEmpty())
+	{
+		return;
+	}
+
 	// Create a dummy MeshBatch to make the system happy..
 	if (MaterialInterface)
 	{
@@ -597,6 +606,7 @@ AHeterogeneousVolume::AHeterogeneousVolume(const FObjectInitializer& ObjectIniti
 
 	PrimaryActorTick.bCanEverTick = true;
 	SetHidden(false);
+	HeterogeneousVolumeComponent->bCastDynamicShadow = true;
 }
 
 #undef LOCTEXT_NAMESPACE
\ No newline at end of file
diff --git a/Engine/Source/Runtime/Engine/Public/HeterogeneousVolumeInterface.h b/Engine/Source/Runtime/Engine/Public/HeterogeneousVolumeInterface.h
--- a/Engine/Source/Runtime/Engine/Public/HeterogeneousVolumeInterface.h
+++ b/Engine/Source/Runtime/Engine/Public/HeterogeneousVolumeInterface.h
@@ -2,6 +2,7 @@
 
 #pragma once
 
+#include "GameFramework/Actor.h"
 #include "CoreMinimal.h"
 #include "SceneManagement.h"
 #include "PrimitiveSceneProxy.h"
@@ -30,6 +31,9 @@ public:
 	virtual float GetShadowStepFactor() const = 0;
 	virtual float GetShadowBiasFactor() const = 0;
 	virtual float GetLightingDownsampleFactor() const = 0;
+
+	// Debug
+	virtual FString GetReadableName() const = 0;
 };
 
 class FPrimitiveSceneProxy;
@@ -47,6 +51,20 @@ public:
 		, ShadowBiasFactor(0.0)
 		, LightingDownsampleFactor(1.0)
 	{}
+
+	FHeterogeneousVolumeData(const FPrimitiveSceneProxy* SceneProxy, FString Name)
+		: PrimitiveSceneProxy(SceneProxy)
+		, InstanceToLocal(FMatrix::Identity)
+		, VoxelResolution(FIntVector::ZeroValue)
+		, MinimumVoxelSize(0.1)
+		, StepFactor(1.0)
+		, ShadowStepFactor(8.0)
+		, ShadowBiasFactor(0.0)
+		, LightingDownsampleFactor(1.0)
+#if ACTOR_HAS_LABELS
+		, ReadableName(Name)
+#endif // ACTOR_HAS_LABELS
+	{}
 	virtual ~FHeterogeneousVolumeData() {}
 
 	virtual const FPrimitiveSceneProxy* GetPrimitiveSceneProxy() const { return PrimitiveSceneProxy; }
@@ -76,4 +94,11 @@ public:
 	float ShadowStepFactor;
 	float ShadowBiasFactor;
 	float LightingDownsampleFactor;
+
+#if ACTOR_HAS_LABELS
+	FString ReadableName;
+	virtual FString GetReadableName() const { return ReadableName; }
+#else
+	virtual FString GetReadableName() const { return PrimitiveSceneProxy->GetResourceName().ToString(); }
+#endif // ACTOR_HAS_LABELS
 };
diff --git a/Engine/Source/Runtime/RenderCore/Private/RenderUtils.cpp b/Engine/Source/Runtime/RenderCore/Private/RenderUtils.cpp
--- a/Engine/Source/Runtime/RenderCore/Private/RenderUtils.cpp
+++ b/Engine/Source/Runtime/RenderCore/Private/RenderUtils.cpp
@@ -1398,6 +1398,14 @@ bool IsWaterVirtualShadowMapFilteringEnabled(const FStaticShaderPlatform Platfor
 	return !IsForwardShadingEnabled(Platform) && bWaterVSMFilteringSupported && bVirtualShadowMapsSupported;
 }
 
+bool DoesPlatformSupportHeterogeneousVolumes(EShaderPlatform Platform)
+{
+	return IsFeatureLevelSupported(Platform, ERHIFeatureLevel::SM5)
+		// TODO:
+		// && FDataDrivenShaderPlatformInfo::GetSupportsHeterogeneousVolumes(Platform)
+		&& !IsForwardShadingEnabled(Platform);
+}
+
 bool IsSingleLayerWaterDepthPrepassEnabled(const FStaticShaderPlatform Platform, const FStaticFeatureLevel FeatureLevel)
 {
 	static const auto CVarWaterSingleLayerDepthPrepass = IConsoleManager::Get().FindConsoleVariable(TEXT("r.Water.SingleLayer.DepthPrepass"));
diff --git a/Engine/Source/Runtime/RenderCore/Private/Shader.cpp b/Engine/Source/Runtime/RenderCore/Private/Shader.cpp
--- a/Engine/Source/Runtime/RenderCore/Private/Shader.cpp
+++ b/Engine/Source/Runtime/RenderCore/Private/Shader.cpp
@@ -2279,6 +2279,21 @@ void ShaderMapAppendKeyString(EShaderPlatform Platform, FString& KeyString)
 			CVarTextureLod && CVarTextureLod->GetBool() ? 1 : 0);
 	}
 
+	if (DoesPlatformSupportHeterogeneousVolumes(Platform))
+	{
+		static const auto ShadowCVar = IConsoleManager::Get().FindTConsoleVariableDataInt(TEXT("r.HeterogeneousVolumes.Shadows"));
+		if (ShadowCVar && ShadowCVar->GetValueOnAnyThread() != 0)
+		{
+			KeyString += TEXT("_HVSHADOW");
+		}
+
+		static const auto CompTranslucencyCVar = IConsoleManager::Get().FindTConsoleVariableDataInt(TEXT("r.Translucency.HeterogeneousVolumes"));
+		if (CompTranslucencyCVar && CompTranslucencyCVar->GetValueOnAnyThread() != 0)
+		{
+			KeyString += TEXT("_HVCOMPTRANSL");
+		}
+	}
+
 	if (ForceSimpleSkyDiffuse(Platform))
 	{
 		KeyString += TEXT("_SSD");
diff --git a/Engine/Source/Runtime/RenderCore/Public/RenderUtils.h b/Engine/Source/Runtime/RenderCore/Public/RenderUtils.h
--- a/Engine/Source/Runtime/RenderCore/Public/RenderUtils.h
+++ b/Engine/Source/Runtime/RenderCore/Public/RenderUtils.h
@@ -433,6 +433,11 @@ RENDERCORE_API bool UseNonNaniteVirtualShadowMaps(EShaderPlatform ShaderPlatform
 /** Returns if water should evaluate virtual shadow maps a second time for the water surface. This is for a platform so can be used at cook time. */
 RENDERCORE_API bool IsWaterVirtualShadowMapFilteringEnabled(const FStaticShaderPlatform Platform);
 
+/**
+ * Returns true if Heterogeneous Volumes should be used for the given shader platform.
+ */
+RENDERCORE_API bool DoesPlatformSupportHeterogeneousVolumes(EShaderPlatform Platform);
+
 /**
 *	(Non-runtime) Checks if the depth prepass for single layer water is enabled. This also depends on virtual shadow maps to be supported on the platform.
 */
diff --git a/Engine/Source/Runtime/Renderer/Private/BasePassRendering.h b/Engine/Source/Runtime/Renderer/Private/BasePassRendering.h
--- a/Engine/Source/Runtime/Renderer/Private/BasePassRendering.h
+++ b/Engine/Source/Runtime/Renderer/Private/BasePassRendering.h
@@ -37,6 +37,7 @@
 #include "LightFunctionAtlas.h"
 #include "RenderUtils.h"
 #include "SceneTexturesConfig.h"
+#include "HeterogeneousVolumes/HeterogeneousVolumes.h"
 
 class FScene;
 
@@ -137,6 +138,7 @@ BEGIN_GLOBAL_SHADER_PARAMETER_STRUCT(FTranslucentBasePassUniformParameters,)
 	SHADER_PARAMETER_RDG_TEXTURE(Texture2D, SceneColorCopyTexture)
 	SHADER_PARAMETER_SAMPLER(SamplerState, SceneColorCopySampler)
 	SHADER_PARAMETER_STRUCT(FBlueNoiseParameters, BlueNoise)
+	SHADER_PARAMETER_STRUCT(FAdaptiveVolumetricShadowMapUniformBufferParameters, AVSM)
 END_GLOBAL_SHADER_PARAMETER_STRUCT()
 
 DECLARE_GPU_DRAWCALL_STAT_EXTERN(Basepass);
@@ -631,6 +633,12 @@ public:
 		const bool bTranslucentUsesLightFunctionAtlas = bTranslucent && GetTranslucentUsesLightFunctionAtlas();
 		OutEnvironment.SetDefine(TEXT("USE_LIGHT_FUNCTION_ATLAS"), (bSingleLayerWaterUsesLightFunctionAtlas || bTranslucentUsesLightFunctionAtlas) ? TEXT("1") : TEXT("0"));
 
+		const bool bIsTranslucent = IsTranslucentBlendMode(Parameters.MaterialParameters);
+		if (bIsTranslucent && DoesPlatformSupportHeterogeneousVolumes(Parameters.Platform) && ShouldCompositeHeterogeneousVolumesWithTranslucency())
+		{
+			OutEnvironment.SetDefine(TEXT("ADAPTIVE_VOLUMETRIC_SHADOW_MAP"), TEXT("1"));
+		}
+
 		TBasePassPixelShaderBaseType<LightMapPolicyType>::ModifyCompilationEnvironment(Parameters, OutEnvironment);
 	}
 	
diff --git a/Engine/Source/Runtime/Renderer/Private/DeferredShadingRenderer.cpp b/Engine/Source/Runtime/Renderer/Private/DeferredShadingRenderer.cpp
--- a/Engine/Source/Runtime/Renderer/Private/DeferredShadingRenderer.cpp
+++ b/Engine/Source/Runtime/Renderer/Private/DeferredShadingRenderer.cpp
@@ -2520,6 +2520,11 @@ void FDeferredShadingSceneRenderer::Render(FRDGBuilder& GraphBuilder)
 			}
 		}
 
+		if (ShouldRenderHeterogeneousVolumes(Scene) && !bHasRayTracedOverlay)
+		{
+			RenderHeterogeneousVolumeShadows(GraphBuilder, SceneTextures);
+		}
+
 		// Post base pass for material classification
 		// This needs to run before virtual shadow map, in order to have ready&cleared classified SSS data
 		if (Substrate::IsSubstrateEnabled() && !bHasRayTracedOverlay)
@@ -3377,6 +3382,7 @@ void FDeferredShadingSceneRenderer::Render(FRDGBuilder& GraphBuilder)
 
 	::Substrate::PostRender(*Scene);
 	HairStrands::PostRender(*Scene);
+	HeterogeneousVolumes::PostRender(*Scene, Views);
 
 	// Release the view's previous frame histories so that their memory can be reused at the graph's execution.
 	for (int32 ViewIndex = 0; ViewIndex < Views.Num(); ViewIndex++)
diff --git a/Engine/Source/Runtime/Renderer/Private/DeferredShadingRenderer.h b/Engine/Source/Runtime/Renderer/Private/DeferredShadingRenderer.h
--- a/Engine/Source/Runtime/Renderer/Private/DeferredShadingRenderer.h
+++ b/Engine/Source/Runtime/Renderer/Private/DeferredShadingRenderer.h
@@ -980,6 +980,7 @@ private:
 		FRDGTextureSRVRef SceneStencilTexture,
 		const TArrayView<FRDGTextureRef> NaniteResolveTextures);
 
+	void RenderHeterogeneousVolumeShadows(FRDGBuilder& GraphBuilder, const FSceneTextures& SceneTextures);
 	void RenderHeterogeneousVolumes(FRDGBuilder& GraphBuilder, const FSceneTextures& SceneTextures);
 	void CompositeHeterogeneousVolumes(FRDGBuilder& GraphBuilder, const FSceneTextures& SceneTextures);
 
diff --git a/Engine/Source/Runtime/Renderer/Private/HeterogeneousVolumes/HeterogeneousVolumes.cpp b/Engine/Source/Runtime/Renderer/Private/HeterogeneousVolumes/HeterogeneousVolumes.cpp
--- a/Engine/Source/Runtime/Renderer/Private/HeterogeneousVolumes/HeterogeneousVolumes.cpp
+++ b/Engine/Source/Runtime/Renderer/Private/HeterogeneousVolumes/HeterogeneousVolumes.cpp
@@ -18,6 +18,19 @@ static TAutoConsoleVariable<int32> CVarHeterogeneousVolumes(
 	ECVF_RenderThreadSafe
 );
 
+static TAutoConsoleVariable<int32> CVarHeterogeneousVolumesShadows(
+	TEXT("r.HeterogeneousVolumes.Shadows"),
+	0,
+	TEXT("Enables heterogeneous volume-casting shadows (default = 0)"),
+	ECVF_RenderThreadSafe | ECVF_ReadOnly);
+
+static TAutoConsoleVariable<int32> CVarTranslucencyHeterogeneousVolumes(
+	TEXT("r.Translucency.HeterogeneousVolumes"),
+	0,
+	TEXT("Enables composting with heterogeneous volumes when rendering translucency (Default = 0)\n"),
+	ECVF_RenderThreadSafe | ECVF_ReadOnly
+);
+
 static TAutoConsoleVariable<int32> CVarHeterogeneousVolumesDebug(
 	TEXT("r.HeterogeneousVolumes.Debug"),
 	0,
@@ -160,6 +173,13 @@ static TAutoConsoleVariable<int32> CVarHeterogeneousVolumesLightingCache(
 	ECVF_RenderThreadSafe
 );
 
+static TAutoConsoleVariable<int32> CVarHeterogeneousVolumesLightingCacheUseAVSM(
+	TEXT("r.HeterogeneousVolumes.LightingCache.UseAVSM"),
+	0,
+	TEXT("Enables use of AVSMs when evaluating self-shadowing (Default = 0)"),
+	ECVF_RenderThreadSafe
+);
+
 static TAutoConsoleVariable<int32> CVarHeterogeneousVolumesLightingCacheDownsampleFactor(
 	TEXT("r.HeterogeneousVolumes.LightingCache.DownsampleFactor"),
 	0,
@@ -176,6 +196,7 @@ static TAutoConsoleVariable<int32> CVarHeterogeneousVolumesDepthSort(
 	ECVF_RenderThreadSafe
 );
 
+DECLARE_GPU_STAT_NAMED(HeterogeneousVolumeShadowsStat, TEXT("HeterogeneousVolumeShadows"));
 DECLARE_GPU_STAT_NAMED(HeterogeneousVolumesStat, TEXT("HeterogeneousVolumes"));
 
 static bool IsHeterogeneousVolumesEnabled()
@@ -183,6 +204,16 @@ static bool IsHeterogeneousVolumesEnabled()
 	return CVarHeterogeneousVolumes.GetValueOnRenderThread() != 0;
 }
 
+bool ShouldHeterogeneousVolumesCastShadows()
+{
+	return CVarHeterogeneousVolumesShadows.GetValueOnAnyThread() != 0;
+}
+
+bool ShouldCompositeHeterogeneousVolumesWithTranslucency()
+{
+	return CVarTranslucencyHeterogeneousVolumes.GetValueOnAnyThread() != 0;
+}
+
 bool ShouldRenderHeterogeneousVolumes(
 	const FScene* Scene
 )
@@ -214,14 +245,6 @@ bool ShouldRenderHeterogeneousVolumesForView(
 		&& !View.bIsReflectionCapture;
 }
 
-bool DoesPlatformSupportHeterogeneousVolumes(EShaderPlatform Platform)
-{
-	return IsFeatureLevelSupported(Platform, ERHIFeatureLevel::SM5)
-		// TODO:
-		// && FDataDrivenShaderPlatformInfo::GetSupportsHeterogeneousVolumes(Platform)
-		&& !IsForwardShadingEnabled(Platform);
-}
-
 bool DoesMaterialShaderSupportHeterogeneousVolumes(const FMaterialShaderParameters& MaterialShaderParameters)
 {
 	return (MaterialShaderParameters.MaterialDomain == MD_Volume)
@@ -335,6 +358,11 @@ namespace HeterogeneousVolumes
 		return CVarHeterogeneousVolumesLightingCache.GetValueOnRenderThread();
 	}
 
+	bool UseAdaptiveVolumetricShadowMapForSelfShadowing()
+	{
+		return ShouldHeterogeneousVolumesCastShadows() && CVarHeterogeneousVolumesLightingCacheUseAVSM.GetValueOnRenderThread();
+	}
+
 	bool UseLightingCacheForInscattering()
 	{
 		return CVarHeterogeneousVolumesLightingCache.GetValueOnRenderThread() == 2;
@@ -395,24 +423,131 @@ namespace HeterogeneousVolumes
 	}
 }
 
-void FDeferredShadingSceneRenderer::RenderHeterogeneousVolumes(
+bool ShouldBuildVoxelGrids(const FScene* Scene)
+{
+	// TODO: Build the light list once
+	if (ShouldHeterogeneousVolumesCastShadows())
+	{
+		for (auto LightIt = Scene->Lights.CreateConstIterator(); LightIt; ++LightIt)
+		{
+			if (LightIt->LightSceneInfo->Proxy->CastsVolumetricShadow())
+			{
+				return true;
+			}
+		}
+	}
+
+	if (ShouldCompositeHeterogeneousVolumesWithTranslucency())
+	{
+		return true;
+	}
+
+	if (HeterogeneousVolumes::GetDebugMode() != 0)
+	{
+		return true;
+	}
+
+	return false;
+}
+
+bool ShouldCacheVoxelGrids(const FScene* Scene, FSceneViewState* ViewState)
+{
+	// If the caching structure exists
+	if (ViewState == nullptr)
+	{
+		return false;
+	}
+
+	if (HeterogeneousVolumes::GetDebugMode() != 0)
+	{
+		return true;
+	}
+	
+	// TODO: If any light supports ray tracing
+
+	return false;
+}
+
+void FDeferredShadingSceneRenderer::RenderHeterogeneousVolumeShadows(
 	FRDGBuilder& GraphBuilder,
 	const FSceneTextures& SceneTextures
 )
 {
-	RDG_EVENT_SCOPE(GraphBuilder, "HeterogeneousVolumes");
-	RDG_GPU_STAT_SCOPE(GraphBuilder, HeterogeneousVolumesStat);
+	if (!ShouldBuildVoxelGrids(Scene))
+	{
+		return;
+	}
+
+	RDG_EVENT_SCOPE(GraphBuilder, "HeterogeneousVolumeShadows");
+	RDG_GPU_STAT_SCOPE(GraphBuilder, HeterogeneousVolumeShadowsStat);
 	SCOPED_NAMED_EVENT(HeterogeneousVolumes, FColor::Emerald);
 
-	// Build global transmittance structure
-	TRDGUniformBufferRef<FOrthoVoxelGridUniformBufferParameters> OrthoGridUniformBuffer;
-	TRDGUniformBufferRef<FFrustumVoxelGridUniformBufferParameters> FrustumGridUniformBuffer;
+	FVoxelGridBuildOptions BuildOptions;
+	BuildOptions.bJitter = false;
+	BuildOptions.MinimumVoxelSizeOutsideFrustum = 2.0;
 
-	if (HeterogeneousVolumes::GetDebugMode() != 0)
+	TRDGUniformBufferRef<FOrthoVoxelGridUniformBufferParameters> OrthoGridUniformBuffer = nullptr;
+	TRDGUniformBufferRef<FFrustumVoxelGridUniformBufferParameters> FrustumGridUniformBuffer = nullptr;
+	BuildOrthoVoxelGrid(GraphBuilder, Scene, Views, VisibleLightInfos, BuildOptions, OrthoGridUniformBuffer);
+	BuildFrustumVoxelGrid(GraphBuilder, Scene, Views[0], BuildOptions, FrustumGridUniformBuffer);
+
+	for (int32 ViewIndex = 0; ViewIndex < Views.Num(); ViewIndex++)
 	{
-		BuildOrthoVoxelGrid(GraphBuilder, Scene, Views, OrthoGridUniformBuffer);
-		BuildFrustumVoxelGrid(GraphBuilder, Scene, Views[0], FrustumGridUniformBuffer);
+		FViewInfo& View = Views[ViewIndex];
+
+		if (ShouldCompositeHeterogeneousVolumesWithTranslucency())
+		{
+			RenderAdaptiveVolumetricCameraMapWithVoxelGrid(
+				GraphBuilder,
+				// Scene data
+				SceneTextures,
+				Scene,
+				ViewFamily,
+				View,
+				// Volume data
+				OrthoGridUniformBuffer,
+				FrustumGridUniformBuffer
+			);
+		}
+
+		if (ShouldHeterogeneousVolumesCastShadows())
+		{
+			RenderAdaptiveVolumetricShadowMapWithVoxelGrid(
+				GraphBuilder,
+				// Scene data
+				SceneTextures,
+				Scene,
+				ViewFamily,
+				View,
+				// Shadow Data
+				VisibleLightInfos,
+				VirtualShadowMapArray,
+				// Volume data
+				OrthoGridUniformBuffer,
+				FrustumGridUniformBuffer
+			);
+		}
+	}
+
+	FSceneViewState* ViewState = Views[0].ViewState;
+	if (ShouldCacheVoxelGrids(Scene, ViewState))
+	{
+		ViewState->OrthoVoxelGridUniformBuffer = OrthoGridUniformBuffer;
+		ViewState->FrustumVoxelGridUniformBuffer = FrustumGridUniformBuffer;
 	}
+}
+
+void FDeferredShadingSceneRenderer::RenderHeterogeneousVolumes(
+	FRDGBuilder& GraphBuilder,
+	const FSceneTextures& SceneTextures
+)
+{
+	RDG_EVENT_SCOPE(GraphBuilder, "HeterogeneousVolumes");
+	RDG_GPU_STAT_SCOPE(GraphBuilder, HeterogeneousVolumesStat);
+	SCOPED_NAMED_EVENT(HeterogeneousVolumes, FColor::Emerald);
+
+	TRDGUniformBufferRef<FOrthoVoxelGridUniformBufferParameters> OrthoGridUniformBuffer = HeterogeneousVolumes::GetOrthoVoxelGridUniformBuffer(GraphBuilder, Views[0].ViewState);
+	TRDGUniformBufferRef<FFrustumVoxelGridUniformBufferParameters> FrustumGridUniformBuffer = HeterogeneousVolumes::GetFrustumVoxelGridUniformBuffer(GraphBuilder, Views[0].ViewState);
 
 	for (int32 ViewIndex = 0; ViewIndex < Views.Num(); ViewIndex++)
 	{
@@ -429,6 +564,7 @@ void FDeferredShadingSceneRenderer::RenderHeterogeneousVolumes(
 		{
 			if (HeterogeneousVolumes::GetDebugMode() != 0)
 			{
+				// TODO: Replace with single-scattering voxel grid implementation.
 				RenderTransmittanceWithVoxelGrid(
 					GraphBuilder,
 					SceneTextures,
@@ -512,21 +648,25 @@ void FDeferredShadingSceneRenderer::RenderHeterogeneousVolumes(
 						const FPersistentPrimitiveIndex PrimitiveId = PrimitiveSceneInfo->GetPersistentIndex();
 						const FBoxSphereBounds LocalBoxSphereBounds = HeterogeneousVolume->GetLocalBounds();
 
-						RDG_EVENT_SCOPE(GraphBuilder, "%s [%d]", *PrimitiveSceneProxy->GetResourceName().ToString(), VolumeIndex);
+						RDG_EVENT_SCOPE(GraphBuilder, "%s [%d]", *HeterogeneousVolume->GetReadableName(), VolumeIndex);
 
 						// Allocate transmittance volume
-						// TODO: Allow option for scalar transmittance to conserve bandwidth
-						FIntVector LightingCacheResolution = HeterogeneousVolumes::GetLightingCacheResolution(HeterogeneousVolume);
-						uint32 NumMips = FMath::Log2(float(FMath::Min(FMath::Min(LightingCacheResolution.X, LightingCacheResolution.Y), LightingCacheResolution.Z))) + 1;
-						FRDGTextureDesc LightingCacheDesc = FRDGTextureDesc::Create3D(
-							LightingCacheResolution,
-                            !IsMetalPlatform(GShaderPlatformForFeatureLevel[View.FeatureLevel]) ? PF_FloatR11G11B10 : PF_FloatRGBA,
-							FClearValueBinding::Black,
-							TexCreate_ShaderResource | TexCreate_UAV | TexCreate_3DTiling,
-							NumMips
-						);
-						FRDGTextureRef LightingCacheTexture = GraphBuilder.CreateTexture(LightingCacheDesc, TEXT("HeterogeneousVolumes.LightingCacheTexture"));
-						AddClearUAVPass(GraphBuilder, GraphBuilder.CreateUAV(LightingCacheTexture), FLinearColor::Black);
+						FRDGTextureRef LightingCacheTexture = GSystemTextures.GetBlackDummy(GraphBuilder);
+						if (HeterogeneousVolumes::GetLightingCacheMode() != 0)
+						{
+							// TODO: Allow option for scalar transmittance to conserve bandwidth
+							FIntVector LightingCacheResolution = HeterogeneousVolumes::GetLightingCacheResolution(HeterogeneousVolume);
+							uint32 NumMips = FMath::Log2(float(FMath::Min(FMath::Min(LightingCacheResolution.X, LightingCacheResolution.Y), LightingCacheResolution.Z))) + 1;
+							FRDGTextureDesc LightingCacheDesc = FRDGTextureDesc::Create3D(
+								LightingCacheResolution,
+								!IsMetalPlatform(GShaderPlatformForFeatureLevel[View.FeatureLevel]) ? PF_FloatR11G11B10 : PF_FloatRGBA,
+								FClearValueBinding::Black,
+								TexCreate_ShaderResource | TexCreate_UAV | TexCreate_3DTiling,
+								NumMips
+							);
+							LightingCacheTexture = GraphBuilder.CreateTexture(LightingCacheDesc, TEXT("HeterogeneousVolumes.LightingCacheTexture"));
+							AddClearUAVPass(GraphBuilder, GraphBuilder.CreateUAV(LightingCacheTexture), FLinearColor::Black);
+						}
 
 						// Material baking executes a pre-shading pipeline
 						if (CVarHeterogeneousVolumesPreshading.GetValueOnRenderThread())
@@ -670,3 +810,24 @@ void FDeferredShadingSceneRenderer::CompositeHeterogeneousVolumes(
 		}
 	}
 }
+
+namespace HeterogeneousVolumes
+{
+	void PostRender(FScene& Scene, TArray<FViewInfo>& Views)
+	{
+		for (int32 ViewIndex = 0; ViewIndex < Views.Num(); ViewIndex++)
+		{
+			FSceneViewState* ViewState = Views[ViewIndex].ViewState;
+			if (ViewState)
+			{
+				DestroyAdaptiveVolumetricShadowMapUniformBuffer(ViewState->AdaptiveVolumetricCameraMapUniformBuffer);
+
+				for (auto Itr = ViewState->AdaptiveVolumetricShadowMapUniformBufferMap.begin(); Itr != ViewState->AdaptiveVolumetricShadowMapUniformBufferMap.end(); ++Itr)
+				{
+					DestroyAdaptiveVolumetricShadowMapUniformBuffer(Itr->Value);
+				}
+				ViewState->AdaptiveVolumetricShadowMapUniformBufferMap.Empty();
+			}
+		}
+	}
+}
\ No newline at end of file
diff --git a/Engine/Source/Runtime/Renderer/Private/HeterogeneousVolumes/HeterogeneousVolumes.h b/Engine/Source/Runtime/Renderer/Private/HeterogeneousVolumes/HeterogeneousVolumes.h
--- a/Engine/Source/Runtime/Renderer/Private/HeterogeneousVolumes/HeterogeneousVolumes.h
+++ b/Engine/Source/Runtime/Renderer/Private/HeterogeneousVolumes/HeterogeneousVolumes.h
@@ -13,6 +13,7 @@ class FRayTracingScene;
 class FRDGBuilder;
 class FScene;
 class FSceneView;
+class FSceneViewState;
 class FViewInfo;
 class FVirtualShadowMapArray;
 class FVisibleLightInfo;
@@ -30,7 +31,6 @@ struct FPersistentPrimitiveIndex;
 bool ShouldRenderHeterogeneousVolumes(const FScene* Scene);
 bool ShouldRenderHeterogeneousVolumesForAnyView(const TArrayView<FViewInfo>& Views);
 bool ShouldRenderHeterogeneousVolumesForView(const FViewInfo& View);
-bool DoesPlatformSupportHeterogeneousVolumes(EShaderPlatform Platform);
 bool DoesMaterialShaderSupportHeterogeneousVolumes(const FMaterialShaderParameters& Parameters);
 bool DoesMaterialShaderSupportHeterogeneousVolumes(const FMaterial& Material);
 bool ShouldRenderMeshBatchWithHeterogeneousVolumes(
@@ -39,6 +39,9 @@ bool ShouldRenderMeshBatchWithHeterogeneousVolumes(
 	ERHIFeatureLevel::Type FeatureLevel
 );
 
+bool ShouldCompositeHeterogeneousVolumesWithTranslucency();
+bool ShouldHeterogeneousVolumesCastShadows();
+
 //
 // Internal API
 //
@@ -72,6 +75,7 @@ namespace HeterogeneousVolumes
 	bool UseSparseVoxelPerTileCulling();
 	bool UseLightingCacheForInscattering();
 	bool UseLightingCacheForTransmittance();
+	bool UseAdaptiveVolumetricShadowMapForSelfShadowing();
 
 	bool EnableIndirectionGrid();
 	bool EnableLinearInterpolation();
@@ -197,10 +201,18 @@ END_UNIFORM_BUFFER_STRUCT()
 
 // Render specializations
 
+struct FVoxelGridBuildOptions
+{
+	bool bJitter = HeterogeneousVolumes::ShouldJitter();
+	float MinimumVoxelSizeOutsideFrustum = HeterogeneousVolumes::GetMinimumVoxelSizeOutsideFrustum();
+};
+
 void BuildOrthoVoxelGrid(
 	FRDGBuilder& GraphBuilder,
 	const FScene* Scene,
 	/*const*/ TArray<FViewInfo>& Views,
+	const TArray<FVisibleLightInfo, SceneRenderingAllocator>& VisibleLightInfos,
+	const FVoxelGridBuildOptions& BuildOptions,
 	TRDGUniformBufferRef<FOrthoVoxelGridUniformBufferParameters>& OrthoVoxelGridUniformBuffer
 );
 
@@ -208,6 +220,7 @@ void BuildFrustumVoxelGrid(
 	FRDGBuilder& GraphBuilder,
 	const FScene* Scene,
 	const FViewInfo& View,
+	const FVoxelGridBuildOptions& BuildOptions,
 	TRDGUniformBufferRef<FFrustumVoxelGridUniformBufferParameters>& FrustumVoxelGridUniformBuffer
 );
 
@@ -227,7 +240,7 @@ struct FAdaptiveFrustumGridParameterCache
 	FIntVector TopLevelGridResolution;
 	FIntVector VoxelDimensions;
 
-	int32 bUseFrustumGrid;
+	int32 bUseFrustumGrid = false;
 
 	float NearPlaneDepth;
 	float FarPlaneDepth;
@@ -246,7 +259,7 @@ struct FAdaptiveOrthoGridParameterCache
 	FVector3f TopLevelGridWorldBoundsMin;
 	FVector3f TopLevelGridWorldBoundsMax;
 	FIntVector TopLevelGridResolution;
-	int32 bUseOrthoGrid;
+	int32 bUseOrthoGrid = false;
 	int32 bUseMajorantGrid;
 	int32 bEnableIndirectionGrid;
 
@@ -285,6 +298,102 @@ void RegisterExternalOrthoVoxelGridUniformBuffer(
 	TRDGUniformBufferRef<FOrthoVoxelGridUniformBufferParameters>& OrthoGridUniformBuffer
 );
 
+struct FAVSMLinkedListPackedData
+{
+	int32 Data[2];
+};
+
+struct FAVSMIndirectionPackedData
+{
+	uint32 Data[2];
+};
+
+struct FAVSMSamplePackedData
+{
+	uint32 Data;
+};
+
+BEGIN_UNIFORM_BUFFER_STRUCT(FAdaptiveVolumetricShadowMapUniformBufferParameters, )
+	SHADER_PARAMETER_ARRAY(FMatrix44f, TranslatedWorldToShadow, [6])
+	SHADER_PARAMETER(FVector3f, TranslatedWorldOrigin)
+	SHADER_PARAMETER(FVector4f, TranslatedWorldPlane)
+
+	SHADER_PARAMETER(FIntPoint, Resolution)
+	SHADER_PARAMETER(int32, NumShadowMatrices)
+	SHADER_PARAMETER(int32, MaxSampleCount)
+	SHADER_PARAMETER(int32, bIsEmpty)
+	SHADER_PARAMETER(int32, bIsDirectionalLight)
+
+	SHADER_PARAMETER_RDG_BUFFER_SRV(StructuredBuffer<int2>, LinkedListBuffer)
+	SHADER_PARAMETER_RDG_BUFFER_SRV(StructuredBuffer<uint2>, IndirectionBuffer)
+	SHADER_PARAMETER_RDG_BUFFER_SRV(StructuredBuffer<uint>, SampleBuffer)
+END_UNIFORM_BUFFER_STRUCT()
+
+namespace HeterogeneousVolumes {
+	struct FAdaptiveVolumetricShadowMapParameterCache
+	{
+		FMatrix44f TranslatedWorldToShadow[6] =
+		{
+			FMatrix44f::Identity,
+			FMatrix44f::Identity,
+			FMatrix44f::Identity,
+			FMatrix44f::Identity,
+			FMatrix44f::Identity,
+			FMatrix44f::Identity
+		};
+		FVector3f TranslatedWorldOrigin = FVector3f::ZeroVector;
+		FVector4f TranslatedWorldPlane = FVector4f::Zero();
+		FIntPoint Resolution = FIntPoint::ZeroValue;
+		int32 NumShadowMatrices = 1;
+		int32 MaxSampleCount = 0;
+		bool bIsEmpty = true;
+		bool bIsDirectionalLight = false;
+
+		TRefCountPtr<FRDGPooledBuffer> LinkedListBuffer;
+		TRefCountPtr<FRDGPooledBuffer> IndirectionBuffer;
+		TRefCountPtr<FRDGPooledBuffer> SampleBuffer;
+	};
+
+	struct FAdaptiveVolumetricShadowMapState
+	{
+		TRDGUniformBufferRef<FAdaptiveVolumetricShadowMapUniformBufferParameters> UniformBuffer = nullptr;
+		FAdaptiveVolumetricShadowMapParameterCache PrevFrameParameterCache;
+	};
+
+	TRDGUniformBufferRef<FAdaptiveVolumetricShadowMapUniformBufferParameters> GetAdaptiveVolumetricShadowMapUniformBuffer(
+		FRDGBuilder& GraphBuilder,
+		FSceneViewState* ViewState,
+		const FLightSceneInfo* LightSceneInfo
+	);
+
+	void DestroyAdaptiveVolumetricShadowMapUniformBuffer(
+		TRDGUniformBufferRef<FAdaptiveVolumetricShadowMapUniformBufferParameters>& AdaptiveVolumetricShadowMapUniformBuffer
+	);
+
+	TRDGUniformBufferRef<FAdaptiveVolumetricShadowMapUniformBufferParameters> GetAdaptiveVolumetricCameraMapUniformBuffer(
+		FRDGBuilder& GraphBuilder,
+		FSceneViewState* ViewState
+	);
+
+	FAdaptiveVolumetricShadowMapUniformBufferParameters GetAdaptiveVolumetricCameraMapParameters(
+		FRDGBuilder& GraphBuilder,
+		FSceneViewState* ViewState
+	);
+
+	TRDGUniformBufferRef<FFrustumVoxelGridUniformBufferParameters> GetFrustumVoxelGridUniformBuffer(
+		FRDGBuilder& GraphBuilder,
+		FSceneViewState* ViewState
+	);
+
+	TRDGUniformBufferRef<FOrthoVoxelGridUniformBufferParameters> GetOrthoVoxelGridUniformBuffer(
+		FRDGBuilder& GraphBuilder,
+		FSceneViewState* ViewState
+	);
+
+	void PostRender(FScene& Scene, TArray<FViewInfo>& Views);
+
+} // namespace HeterogeneousVolumes
+
 void RenderWithLiveShading(
 	FRDGBuilder& GraphBuilder,
 	const FSceneTextures& SceneTextures,
@@ -338,6 +447,50 @@ void RenderTransmittanceWithVoxelGrid(
 	FRDGTextureRef& HeterogeneousVolumeRadiance
 );
 
+void RenderAdaptiveVolumetricShadowMapWithVoxelGrid(
+	FRDGBuilder& GraphBuilder,
+	// Scene data
+	const FSceneTextures& SceneTextures,
+	FScene* Scene,
+	const FSceneViewFamily& ViewFamily,
+	FViewInfo& View,
+	// Shadow data
+	TArray<FVisibleLightInfo, SceneRenderingAllocator>& VisibleLightInfos,
+	const FVirtualShadowMapArray& VirtualShadowMapArray,
+	// Volume data
+	const TRDGUniformBufferRef<FOrthoVoxelGridUniformBufferParameters>& OrthoGridUniformBuffer,
+	const TRDGUniformBufferRef<FFrustumVoxelGridUniformBufferParameters>& FrustumGridUniformBuffer
+);
+
+void RenderAdaptiveVolumetricCameraMapWithVoxelGrid(
+	FRDGBuilder& GraphBuilder,
+	// Scene data
+	const FSceneTextures& SceneTextures,
+	FScene* Scene,
+	const FSceneViewFamily& ViewFamily,
+	FViewInfo& View,
+	// Volume data
+	const TRDGUniformBufferRef<FOrthoVoxelGridUniformBufferParameters>& OrthoGridUniformBuffer,
+	const TRDGUniformBufferRef<FFrustumVoxelGridUniformBufferParameters>& FrustumGridUniformBuffer
+);
+
+void RenderSingleScatteringWithVoxelGrid(
+	FRDGBuilder& GraphBuilder,
+	// Scene data
+	const FSceneTextures& SceneTextures,
+	FScene* Scene,
+	const FSceneViewFamily& ViewFamily,
+	FViewInfo& View,
+	// Shadow data
+	TArray<FVisibleLightInfo, SceneRenderingAllocator>& VisibleLightInfos,
+	const FVirtualShadowMapArray& VirtualShadowMapArray,
+	// Volume data
+	const TRDGUniformBufferRef<FOrthoVoxelGridUniformBufferParameters>& OrthoGridUniformBuffer,
+	const TRDGUniformBufferRef<FFrustumVoxelGridUniformBufferParameters>& FrustumGridUniformBuffer,
+	// Output
+	FRDGTextureRef& HeterogeneousVolumeRadiance
+);
+
 // Preshading Pipeline
 void ComputeHeterogeneousVolumeBakeMaterial(
 	FRDGBuilder& GraphBuilder,
diff --git a/Engine/Source/Runtime/Renderer/Private/HeterogeneousVolumes/HeterogeneousVolumesLiveShadingPipeline.cpp b/Engine/Source/Runtime/Renderer/Private/HeterogeneousVolumes/HeterogeneousVolumesLiveShadingPipeline.cpp
--- a/Engine/Source/Runtime/Renderer/Private/HeterogeneousVolumes/HeterogeneousVolumesLiveShadingPipeline.cpp
+++ b/Engine/Source/Runtime/Renderer/Private/HeterogeneousVolumes/HeterogeneousVolumesLiveShadingPipeline.cpp
@@ -95,7 +95,8 @@ class FRenderLightingCacheWithLiveShadingCS : public FMeshMaterialShader
 	DECLARE_SHADER_TYPE(FRenderLightingCacheWithLiveShadingCS, MeshMaterial);
 
 	class FLightingCacheMode : SHADER_PERMUTATION_INT("DIM_LIGHTING_CACHE_MODE", 2);
-	using FPermutationDomain = TShaderPermutationDomain<FLightingCacheMode>;
+	class FUseAdaptiveVolumetricShadowMap : SHADER_PERMUTATION_BOOL("DIM_USE_ADAPTIVE_VOLUMETRIC_SHADOW_MAP");
+	using FPermutationDomain = TShaderPermutationDomain<FLightingCacheMode, FUseAdaptiveVolumetricShadowMap>;
 
 	BEGIN_SHADER_PARAMETER_STRUCT(FParameters, )
 		// Scene data
@@ -118,6 +119,7 @@ class FRenderLightingCacheWithLiveShadingCS : public FMeshMaterialShader
 		SHADER_PARAMETER_STRUCT_INCLUDE(FVolumeShadowingShaderParameters, VolumeShadowingShaderParameters)
 		SHADER_PARAMETER_STRUCT_INCLUDE(FVirtualShadowMapSamplingParameters, VirtualShadowMapSamplingParameters)
 		SHADER_PARAMETER(int32, VirtualShadowMapId)
+		SHADER_PARAMETER_RDG_UNIFORM_BUFFER(FAdaptiveVolumetricShadowMapUniformBufferParameters, AVSM)
 
 		// Object data
 		SHADER_PARAMETER(FMatrix44f, LocalToWorld)
@@ -213,7 +215,8 @@ class FRenderSingleScatteringWithLiveShadingCS : public FMeshMaterialShader
 	class FUseTransmittanceVolume : SHADER_PERMUTATION_BOOL("DIM_USE_TRANSMITTANCE_VOLUME");
 	class FUseInscatteringVolume : SHADER_PERMUTATION_BOOL("DIM_USE_INSCATTERING_VOLUME");
 	class FUseLumenGI : SHADER_PERMUTATION_BOOL("DIM_USE_LUMEN_GI");
-	using FPermutationDomain = TShaderPermutationDomain<FUseTransmittanceVolume, FUseInscatteringVolume, FUseLumenGI>;
+	class FUseAdaptiveVolumetricShadowMap : SHADER_PERMUTATION_BOOL("DIM_USE_ADAPTIVE_VOLUMETRIC_SHADOW_MAP");
+	using FPermutationDomain = TShaderPermutationDomain<FUseTransmittanceVolume, FUseInscatteringVolume, FUseLumenGI, FUseAdaptiveVolumetricShadowMap>;
 
 	BEGIN_SHADER_PARAMETER_STRUCT(FParameters, )
 		// Scene data
@@ -236,6 +239,7 @@ class FRenderSingleScatteringWithLiveShadingCS : public FMeshMaterialShader
 		SHADER_PARAMETER_STRUCT_INCLUDE(FVolumeShadowingShaderParameters, VolumeShadowingShaderParameters)
 		SHADER_PARAMETER_STRUCT_INCLUDE(FVirtualShadowMapSamplingParameters, VirtualShadowMapSamplingParameters)
 		SHADER_PARAMETER(int32, VirtualShadowMapId)
+		SHADER_PARAMETER_RDG_UNIFORM_BUFFER(FAdaptiveVolumetricShadowMapUniformBufferParameters, AVSM)
 
 		// Indirect Lighting
 		SHADER_PARAMETER_RDG_UNIFORM_BUFFER(FLumenTranslucencyLightingUniforms, LumenGIVolumeStruct)
@@ -475,6 +479,7 @@ static void RenderLightingCacheWithLiveShading(
 			PassParameters->VirtualShadowMapId = -1;
 		}
 		PassParameters->VirtualShadowMapSamplingParameters = VirtualShadowMapArray.GetSamplingParameters(GraphBuilder);
+		PassParameters->AVSM = HeterogeneousVolumes::GetAdaptiveVolumetricShadowMapUniformBuffer(GraphBuilder, View.ViewState, LightSceneInfo);
 
 		// Output
 		PassParameters->RWLightingCacheTexture = GraphBuilder.CreateUAV(LightingCacheTexture);
@@ -501,6 +506,7 @@ static void RenderLightingCacheWithLiveShading(
 
 	FRenderLightingCacheWithLiveShadingCS::FPermutationDomain PermutationVector;
 	PermutationVector.Set<FRenderLightingCacheWithLiveShadingCS::FLightingCacheMode>(HeterogeneousVolumes::GetLightingCacheMode() - 1);
+	PermutationVector.Set<FRenderLightingCacheWithLiveShadingCS::FUseAdaptiveVolumetricShadowMap>(HeterogeneousVolumes::UseAdaptiveVolumetricShadowMapForSelfShadowing());
 	TShaderRef<FRenderLightingCacheWithLiveShadingCS> ComputeShader = Material.GetShader<FRenderLightingCacheWithLiveShadingCS>(&FLocalVertexFactory::StaticType, PermutationVector, false);
 	if (!ComputeShader.IsNull())
 	{
@@ -573,6 +579,10 @@ static void RenderSingleScatteringWithLiveShading(
 		{
 			PassParameters->VolumetricScatteringIntensity = LightSceneInfo->Proxy->GetVolumetricScatteringIntensity();
 		}
+		else
+		{
+			PassParameters->VolumetricScatteringIntensity = 1.0f;
+		}
 		PassParameters->DeferredLight = DeferredLightUB;
 		PassParameters->LightType = LightType;
 		PassParameters->ShadowStepSize = HeterogeneousVolumes::GetShadowStepSize();
@@ -619,6 +629,7 @@ static void RenderSingleScatteringWithLiveShading(
 			SetVolumeShadowingDefaultShaderParametersGlobal(GraphBuilder, PassParameters->VolumeShadowingShaderParameters);
 		}
 		PassParameters->VirtualShadowMapSamplingParameters = VirtualShadowMapArray.GetSamplingParameters(GraphBuilder);
+		PassParameters->AVSM = HeterogeneousVolumes::GetAdaptiveVolumetricShadowMapUniformBuffer(GraphBuilder, View.ViewState, LightSceneInfo);
 
 		// Indirect lighting data
 		auto* LumenUniforms = GraphBuilder.AllocParameters<FLumenTranslucencyLightingUniforms>();
@@ -664,6 +675,7 @@ static void RenderSingleScatteringWithLiveShading(
 	PermutationVector.Set<FRenderSingleScatteringWithLiveShadingCS::FUseTransmittanceVolume>(HeterogeneousVolumes::UseLightingCacheForTransmittance() && PassParameters->bApplyShadowTransmittance);
 	PermutationVector.Set<FRenderSingleScatteringWithLiveShadingCS::FUseInscatteringVolume>(HeterogeneousVolumes::UseLightingCacheForInscattering());
 	PermutationVector.Set<FRenderSingleScatteringWithLiveShadingCS::FUseLumenGI>(HeterogeneousVolumes::UseIndirectLighting() && View.GetLumenTranslucencyGIVolume().Texture0 != nullptr);
+	PermutationVector.Set<FRenderSingleScatteringWithLiveShadingCS::FUseAdaptiveVolumetricShadowMap>(HeterogeneousVolumes::UseAdaptiveVolumetricShadowMapForSelfShadowing());
 	TShaderRef<FRenderSingleScatteringWithLiveShadingCS> ComputeShader = Material.GetShader<FRenderSingleScatteringWithLiveShadingCS>(&FLocalVertexFactory::StaticType, PermutationVector, false);
 	if (!ComputeShader.IsNull())
 	{
diff --git a/Engine/Source/Runtime/Renderer/Private/HeterogeneousVolumes/HeterogeneousVolumesVoxelGridPipeline.cpp b/Engine/Source/Runtime/Renderer/Private/HeterogeneousVolumes/HeterogeneousVolumesVoxelGridPipeline.cpp
--- a/Engine/Source/Runtime/Renderer/Private/HeterogeneousVolumes/HeterogeneousVolumesVoxelGridPipeline.cpp
+++ b/Engine/Source/Runtime/Renderer/Private/HeterogeneousVolumes/HeterogeneousVolumesVoxelGridPipeline.cpp
@@ -13,6 +13,8 @@
 #include "ScenePrivate.h"
 #include "SceneManagement.h"
 #include "BlueNoise.h"
+#include "VolumeLighting.h"
+#include "VolumetricFog.h"
 
 static TAutoConsoleVariable<int32> CVarHeterogeneousVolumesEnableFrustumVoxelGrid(
 	TEXT("r.HeterogeneousVolumes.FrustumGrid"),
@@ -88,6 +90,13 @@ static TAutoConsoleVariable<int32> CVarHeterogeneousVolumesMarchingMode(
 	ECVF_RenderThreadSafe
 );
 
+static TAutoConsoleVariable<int32> CVarHeterogeneousVolumesTessellationJitter(
+	TEXT("r.HeterogeneousVolumes.Tessellation.Jitter"),
+	1,
+	TEXT("Enables jittering when tessellating the acceleration grids (Default = 1)"),
+	ECVF_RenderThreadSafe
+);
+
 static TAutoConsoleVariable<int32> CVarHeterogeneousVolumesBottomLevelGridResolution(
 	TEXT("r.HeterogeneousVolumes.Tessellation.BottomLevelGrid.Resolution"),
 	4,
@@ -179,8 +188,57 @@ static TAutoConsoleVariable<int32> CVarHeterogeneousVolumesEnableMajorantGridMax
 	ECVF_RenderThreadSafe
 );
 
+static TAutoConsoleVariable<int32> CVarHeterogeneousVolumesShadowCameraDownsampleFactor(
+	TEXT("r.HeterogeneousVolumes.Shadows.CameraDownsampleFactor"),
+	2,
+	TEXT("Controls downsample factor for camera volumetric shadow map (default = 2)"),
+	ECVF_RenderThreadSafe);
+
+static TAutoConsoleVariable<int32> CVarHeterogeneousVolumesShadowResolution(
+	TEXT("r.HeterogeneousVolumes.Shadows.Resolution"),
+	256,
+	TEXT("Resolution when building volumetric shadow map (Default = 256)\n"),
+	ECVF_RenderThreadSafe
+);
+
+static TAutoConsoleVariable<int32> CVarHeterogeneousVolumesShadowMaxSampleCount(
+	TEXT("r.HeterogeneousVolumes.Shadows.MaxSampleCount"),
+	32,
+	TEXT("Maximum sample count when building volumetric shadow map (Default = 32)\n"),
+	ECVF_RenderThreadSafe
+);
+
+static TAutoConsoleVariable<float> CVarHeterogeneousVolumesShadowAbsoluteErrorThreshold(
+	TEXT("r.HeterogeneousVolumes.Shadows.AbsoluteErrorThreshold"),
+	0.0,
+	TEXT("Absolute error threshold for volume shadow compression (Default = 0.0)\n"),
+	ECVF_RenderThreadSafe
+);
+
+static TAutoConsoleVariable<float> CVarHeterogeneousVolumesShadowRelativeErrorThreshold(
+	TEXT("r.HeterogeneousVolumes.Shadows.RelativeErrorThreshold"),
+	0.05,
+	TEXT("Relative error threshold for volume shadow compression (Default = 0.05)\n"),
+	ECVF_RenderThreadSafe
+);
+
+static TAutoConsoleVariable<int32> CVarHeterogeneousVolumesShadowUseAVSMCompression(
+	TEXT("r.HeterogeneousVolumes.Shadows.UseAVSMCompression"),
+	1,
+	TEXT("Enables AVSM compression (Default = 1)\n"),
+	ECVF_RenderThreadSafe
+);
+
+static TAutoConsoleVariable<int32> CVarHeterogeneousVolumesShadowDebugTweak(
+	TEXT("r.HeterogeneousVolumes.Shadows.DebugTweak"),
+	0,
+	TEXT("Debug tweak value (Default = 0)\n"),
+	ECVF_RenderThreadSafe
+);
+
 IMPLEMENT_UNIFORM_BUFFER_STRUCT(FOrthoVoxelGridUniformBufferParameters, "OrthoGridUniformBuffer");
 IMPLEMENT_UNIFORM_BUFFER_STRUCT(FFrustumVoxelGridUniformBufferParameters, "FrustumGridUniformBuffer");
+IMPLEMENT_UNIFORM_BUFFER_STRUCT(FAdaptiveVolumetricShadowMapUniformBufferParameters, "AVSM");
 
 struct FRasterTileData
 {
@@ -281,6 +339,11 @@ namespace HeterogeneousVolumes
 		return CVarHeterogeneousVolumesIndirectionGrid.GetValueOnRenderThread() != 0;
 	}
 
+	bool EnableJitter()
+	{
+		return CVarHeterogeneousVolumesTessellationJitter.GetValueOnRenderThread() != 0;
+	}
+
 	int32 GetBottomLevelGridResolution()
 	{
 		return FMath::Clamp(CVarHeterogeneousVolumesBottomLevelGridResolution.GetValueOnRenderThread(), 1, 4);
@@ -310,6 +373,36 @@ namespace HeterogeneousVolumes
 	{
 		return FMath::Tan(FMath::DegreesToRadians(FOVInDegrees * 0.5));
 	}
+
+	FIntPoint GetShadowMapResolution()
+	{
+		return FIntPoint(FMath::Clamp(CVarHeterogeneousVolumesShadowResolution.GetValueOnRenderThread(), 1, 1024));
+	}
+
+	uint32 GetShadowMaxSampleCount()
+	{
+		return FMath::Clamp(CVarHeterogeneousVolumesShadowMaxSampleCount.GetValueOnRenderThread(), 2, 64);
+	}
+
+	float GetShadowAbsoluteErrorThreshold()
+	{
+		return FMath::Max(CVarHeterogeneousVolumesShadowAbsoluteErrorThreshold.GetValueOnRenderThread(), 0.0);
+	}
+
+	float GetShadowRelativeErrorThreshold()
+	{
+		return FMath::Max(CVarHeterogeneousVolumesShadowRelativeErrorThreshold.GetValueOnRenderThread(), 0.0);
+	}
+
+	bool UseAVSMCompression()
+	{
+		return CVarHeterogeneousVolumesShadowUseAVSMCompression.GetValueOnRenderThread() != 0;
+	}
+
+	float GetCameraDownsampleFactor()
+	{
+		return FMath::Max(CVarHeterogeneousVolumesShadowCameraDownsampleFactor.GetValueOnRenderThread(), 1);
+	}
 }
 
 class FMarkTopLevelGridVoxelsForFrustumGrid : public FGlobalShader
@@ -391,6 +484,7 @@ class FRasterizeBottomLevelFrustumGridCS : public FMeshMaterialShader
 		SHADER_PARAMETER(int, BottomLevelGridBufferSize)
 
 		// Sampling data
+		SHADER_PARAMETER(int, bJitter)
 		SHADER_PARAMETER_STRUCT_REF(FBlueNoise, BlueNoise)
 
 		// Raster tile data
@@ -687,6 +781,7 @@ class FRasterizeBottomLevelOrthoGridCS : public FMeshMaterialShader
 		SHADER_PARAMETER_STRUCT_REF(FBlueNoise, BlueNoise)
 
 		// Volume sample mode
+		SHADER_PARAMETER(int, bJitter)
 		SHADER_PARAMETER(int, bSampleAtVertices)
 
 		// Raster tile data
@@ -1114,6 +1209,9 @@ struct FRenderDebugData
 	float TMax;
 	float Distance;
 	FVector4f EstimateAndPdf;
+
+	FMatrix44f ScreenToTranslatedWorld;
+	FMatrix44f ClipToTranslatedWorld;
 };
 
 class FRenderTransmittanceWithVoxelGridCS : public FGlobalShader
@@ -1178,12 +1276,13 @@ IMPLEMENT_GLOBAL_SHADER(FRenderTransmittanceWithVoxelGridCS, "/Engine/Private/He
 void CalcViewBoundsAndMinimumVoxelSize(
 	FRDGBuilder& GraphBuilder,
 	const FViewInfo& View,
+	const FVoxelGridBuildOptions& BuildOptions,
 	FBoxSphereBounds& TopLevelGridBounds,
 	float& MinimumVoxelSize
 )
 {
 	TopLevelGridBounds = FBoxSphereBounds(ForceInit);
-	MinimumVoxelSize = HeterogeneousVolumes::GetMinimumVoxelSizeOutsideFrustum();
+	MinimumVoxelSize = BuildOptions.MinimumVoxelSizeOutsideFrustum;
 
 	// Build view bounds
 	FVector WorldCameraOrigin = View.ViewMatrices.GetViewOrigin();
@@ -1236,17 +1335,19 @@ void CalcViewBoundsAndMinimumVoxelSize(
 void CalcGlobalBoundsAndMinimumVoxelSize(
 	FRDGBuilder& GraphBuilder,
 	const TArray<FViewInfo>& Views,
+	const TSet<FVolumetricMeshBatch>& HeterogeneousVolumesMeshBatches,
+	const FVoxelGridBuildOptions& BuildOptions,
 	FBoxSphereBounds& TopLevelGridBounds,
 	float& GlobalMinimumVoxelSize
 )
 {
 	TopLevelGridBounds = FBoxSphereBounds(ForceInit);
-	GlobalMinimumVoxelSize = HeterogeneousVolumes::GetMinimumVoxelSizeOutsideFrustum();
+	GlobalMinimumVoxelSize = BuildOptions.MinimumVoxelSizeOutsideFrustum;
 
 	// Cycle through all Volume PrimitiveSceneProxies to collect bounds information and minimum voxel-size
 	for (int32 ViewIndex = 0; ViewIndex < Views.Num(); ViewIndex++)
 	{
-		float ViewMinimumVoxelSize = HeterogeneousVolumes::GetMinimumVoxelSizeOutsideFrustum();
+		float ViewMinimumVoxelSize = BuildOptions.MinimumVoxelSizeOutsideFrustum;
 		FBoxSphereBounds AggregatePrimitiveBounds = FBoxSphereBounds(ForceInit);
 
 		// Build view bounds
@@ -1258,11 +1359,12 @@ void CalcGlobalBoundsAndMinimumVoxelSize(
 		int32 HalfWidth = View.ViewRect.Width() * 0.5;
 		float PixelWidth = TanHalfFOV / HalfWidth;
 
-		for (int32 MeshBatchIndex = 0; MeshBatchIndex < View.HeterogeneousVolumesMeshBatches.Num(); ++MeshBatchIndex)
+		for (auto MeshBatchIt = HeterogeneousVolumesMeshBatches.begin(); MeshBatchIt != HeterogeneousVolumesMeshBatches.end(); ++MeshBatchIt)
 		{
-			// Only Niagara mesh particles bound to volume materials
-			const FMeshBatch* Mesh = View.HeterogeneousVolumesMeshBatches[MeshBatchIndex].Mesh;
-			const FPrimitiveSceneProxy* PrimitiveSceneProxy = View.HeterogeneousVolumesMeshBatches[MeshBatchIndex].Proxy;
+			const FVolumetricMeshBatch& MeshBatch = *MeshBatchIt;
+			const FMeshBatch* Mesh = MeshBatch.Mesh;
+			const FPrimitiveSceneProxy* PrimitiveSceneProxy = MeshBatch.Proxy;
+
 			if (!ShouldRenderMeshBatchWithHeterogeneousVolumes(Mesh, PrimitiveSceneProxy, View.GetFeatureLevel()))
 			{
 				continue;
@@ -1278,16 +1380,17 @@ void CalcGlobalBoundsAndMinimumVoxelSize(
 				}
 				// Only incorporate the primitive if it intersects with the canera bounding sphere where radius=MaxTraceDistance
 				const FBoxSphereBounds& PrimitiveBounds = HeterogeneousVolume->GetBounds();
-				if (View.ViewFrustum.IntersectBox(PrimitiveBounds.Origin, PrimitiveBounds.BoxExtent))
+
 				{
 					AggregatePrimitiveBounds = Union(AggregatePrimitiveBounds, PrimitiveBounds);
 
-					if (View.ViewFrustum.IntersectBox(AggregatePrimitiveBounds.Origin, AggregatePrimitiveBounds.BoxExtent))
+					if (View.ViewFrustum.IntersectBox(PrimitiveBounds.Origin, PrimitiveBounds.BoxExtent))
 					{
 						// Bandlimit minimum voxel size request with projected voxel size, based on shading rate
 						FVector VoxelCenter = PrimitiveBounds.Origin;
 						float Distance = FMath::Max(FVector(PrimitiveBounds.Origin - WorldCameraOrigin).Length() - AggregatePrimitiveBounds.BoxExtent.Length(), 0.0);
 						float VoxelWidth = Distance * PixelWidth * HeterogeneousVolumes::GetShadingRateForOrthoGrid();
+						//float VoxelWidth = HeterogeneousVolumes::GetShadingRateForOrthoGrid();
 
 						float PerVolumeMinimumVoxelSize = FMath::Max(VoxelWidth, HeterogeneousVolume->GetMinimumVoxelSize());
 						ViewMinimumVoxelSize = FMath::Min(PerVolumeMinimumVoxelSize, ViewMinimumVoxelSize);
@@ -1386,9 +1489,8 @@ void RegisterExternalFrustumVoxelGridUniformBuffer(
 	FrustumGridUniformBuffer = GraphBuilder.CreateUniformBuffer(UniformBufferParameters);
 }
 
-void CreateEmptyFrustumVoxelGridUniformBuffer(
-	FRDGBuilder& GraphBuilder,
-	TRDGUniformBufferRef<FFrustumVoxelGridUniformBufferParameters>& FrustumGridUniformBuffer
+TRDGUniformBufferRef <FFrustumVoxelGridUniformBufferParameters> CreateEmptyFrustumVoxelGridUniformBuffer(
+	FRDGBuilder& GraphBuilder
 )
 {
 	FFrustumVoxelGridUniformBufferParameters* UniformBufferParameters = GraphBuilder.AllocParameters<FFrustumVoxelGridUniformBufferParameters>();
@@ -1414,7 +1516,7 @@ void CreateEmptyFrustumVoxelGridUniformBuffer(
 		UniformBufferParameters->FarPlaneDepth = 0.0;
 		UniformBufferParameters->TanHalfFOV = 1.0;
 	}
-	FrustumGridUniformBuffer = GraphBuilder.CreateUniformBuffer(UniformBufferParameters);
+	return GraphBuilder.CreateUniformBuffer(UniformBufferParameters);
 }
 
 void ClipNearFarDistances(const FViewInfo& View, const FBoxSphereBounds& TopLevelGridBounds, float& NearPlaneDistance, float& FarPlaneDistance)
@@ -1476,6 +1578,8 @@ void CalculateVoxelSize(
 	FRDGBuilder& GraphBuilder,
 	const FScene* Scene,
 	/*const*/ TArray<FViewInfo>& Views,
+	const TSet<FVolumetricMeshBatch>& HeterogeneousVolumesMeshBatches,
+	const FVoxelGridBuildOptions& BuildOptions,
 	FBoxSphereBounds TopLevelGridBounds,
 	FIntVector TopLevelGridResolution,
 	FRDGBufferRef& TopLevelGridBuffer
@@ -1494,10 +1598,11 @@ void CalculateVoxelSize(
 		FVector WorldCameraOrigin = View.ViewMatrices.GetViewOrigin();
 		FBoxSphereBounds WorldCameraBounds(FSphere(WorldCameraOrigin, HeterogeneousVolumes::GetMaxTraceDistance()));
 
-		for (int32 MeshBatchIndex = 0; MeshBatchIndex < View.HeterogeneousVolumesMeshBatches.Num(); ++MeshBatchIndex)
+		for (auto MeshBatchIt = HeterogeneousVolumesMeshBatches.begin(); MeshBatchIt != HeterogeneousVolumesMeshBatches.end(); ++MeshBatchIt)
 		{
-			const FMeshBatch* Mesh = View.HeterogeneousVolumesMeshBatches[MeshBatchIndex].Mesh;
-			const FPrimitiveSceneProxy* PrimitiveSceneProxy = View.HeterogeneousVolumesMeshBatches[MeshBatchIndex].Proxy;
+			const FVolumetricMeshBatch& MeshBatch = *MeshBatchIt;
+			const FMeshBatch* Mesh = MeshBatch.Mesh;
+			const FPrimitiveSceneProxy* PrimitiveSceneProxy = MeshBatch.Proxy;
 
 			for (int32 VolumeIndex = 0; VolumeIndex < Mesh->Elements.Num(); ++VolumeIndex)
 			{
@@ -1521,8 +1626,8 @@ void CalculateVoxelSize(
 					PassParameters->PrimitiveWorldBoundsMax = FVector3f(PrimitiveBounds.Origin + PrimitiveBounds.BoxExtent);
 
 					PassParameters->ShadingRate = HeterogeneousVolumes::GetShadingRateForOrthoGrid();
-						PassParameters->MinVoxelSizeInFrustum = FMath::Max(HeterogeneousVolume->GetMinimumVoxelSize(), HeterogeneousVolumes::GetMinimumVoxelSizeInFrustum());
-					PassParameters->MinVoxelSizeOutOfFrustum = HeterogeneousVolumes::GetMinimumVoxelSizeOutsideFrustum();
+					PassParameters->MinVoxelSizeInFrustum = FMath::Max(HeterogeneousVolume->GetMinimumVoxelSize(), HeterogeneousVolumes::GetMinimumVoxelSizeInFrustum());
+					PassParameters->MinVoxelSizeOutOfFrustum = BuildOptions.MinimumVoxelSizeOutsideFrustum;
 
 					PassParameters->RWTopLevelGridBuffer = GraphBuilder.CreateUAV(TopLevelGridBuffer);
 				}
@@ -1780,6 +1885,7 @@ void RasterizeVolumesIntoFrustumVoxelGrid(
 	FRDGBuilder& GraphBuilder,
 	const FScene* Scene,
 	const FViewInfo& View,
+	const FVoxelGridBuildOptions& BuildOptions,
 	// Transform data
 	FMatrix& ViewToWorld,
 	float NearPlaneDistance,
@@ -1886,6 +1992,7 @@ void RasterizeVolumesIntoFrustumVoxelGrid(
 			PassParameters->FarPlaneDepth = FarPlaneDistance;
 
 			// Sampling data
+			PassParameters->bJitter = BuildOptions.bJitter;
 			FBlueNoise BlueNoise = GetBlueNoiseGlobalParameters();
 			PassParameters->BlueNoise = CreateUniformBufferImmediate(BlueNoise, EUniformBufferUsage::UniformBuffer_SingleDraw);
 
@@ -1934,12 +2041,13 @@ void BuildFrustumVoxelGrid(
 	FRDGBuilder& GraphBuilder,
 	const FScene* Scene,
 	const FViewInfo& View,
+	const FVoxelGridBuildOptions& BuildOptions,
 	TRDGUniformBufferRef<FFrustumVoxelGridUniformBufferParameters>& FrustumGridUniformBuffer
 )
 {
 	if (!ShouldRenderHeterogeneousVolumesForView(View) || !HeterogeneousVolumes::EnableFrustumVoxelGrid())
 	{
-		CreateEmptyFrustumVoxelGridUniformBuffer(GraphBuilder, FrustumGridUniformBuffer);
+		FrustumGridUniformBuffer = CreateEmptyFrustumVoxelGridUniformBuffer(GraphBuilder);
 		return;
 	}
 
@@ -1948,11 +2056,11 @@ void BuildFrustumVoxelGrid(
 	// Determine the minimum voxel size for the scene, based on screen projection or user-defined minima
 	FBoxSphereBounds TopLevelGridBounds;
 	float MinimumVoxelSize;
-	CalcViewBoundsAndMinimumVoxelSize(GraphBuilder, View, TopLevelGridBounds, MinimumVoxelSize);
+	CalcViewBoundsAndMinimumVoxelSize(GraphBuilder, View, BuildOptions, TopLevelGridBounds, MinimumVoxelSize);
 
 	if (TopLevelGridBounds.SphereRadius == 0)
 	{
-		CreateEmptyFrustumVoxelGridUniformBuffer(GraphBuilder, FrustumGridUniformBuffer);
+		FrustumGridUniformBuffer = CreateEmptyFrustumVoxelGridUniformBuffer(GraphBuilder);
 		return;
 	}
 
@@ -1973,7 +2081,7 @@ void BuildFrustumVoxelGrid(
 	int32 TopLevelVoxelCount = TopLevelGridResolution.X * TopLevelGridResolution.Y * TopLevelGridResolution.Z;
 	if (TopLevelVoxelCount == 0)
 	{
-		CreateEmptyFrustumVoxelGridUniformBuffer(GraphBuilder, FrustumGridUniformBuffer);
+		FrustumGridUniformBuffer = CreateEmptyFrustumVoxelGridUniformBuffer(GraphBuilder);
 		return;
 	}
 
@@ -2015,6 +2123,7 @@ void BuildFrustumVoxelGrid(
 		GraphBuilder,
 		Scene,
 		View,
+		BuildOptions,
 		ViewToWorld,
 		NearPlaneDistance,
 		FarPlaneDistance,
@@ -2096,6 +2205,7 @@ void BuildFrustumVoxelGrid(
 		UniformBufferParameters->EmissionFroxelGridBuffer = GraphBuilder.CreateSRV(EmissionGridBuffer);
 		UniformBufferParameters->ScatteringFroxelGridBuffer = GraphBuilder.CreateSRV(ScatteringGridBuffer);
 	}
+
 	FrustumGridUniformBuffer = GraphBuilder.CreateUniformBuffer(UniformBufferParameters);
 }
 
@@ -2152,9 +2262,8 @@ void RegisterExternalOrthoVoxelGridUniformBuffer(
 	OrthoGridUniformBuffer = GraphBuilder.CreateUniformBuffer(UniformBufferParameters);
 }
 
-void CreateEmptyOrthoVoxelGridUniformBuffer(
-	FRDGBuilder& GraphBuilder,
-	TRDGUniformBufferRef<FOrthoVoxelGridUniformBufferParameters>& OrthoGridUniformBuffer
+TRDGUniformBufferRef<FOrthoVoxelGridUniformBufferParameters> CreateEmptyOrthoVoxelGridUniformBuffer(
+	FRDGBuilder& GraphBuilder
 )
 {
 	FOrthoVoxelGridUniformBufferParameters* OrthoGridUniformBufferParameters = GraphBuilder.AllocParameters<FOrthoVoxelGridUniformBufferParameters>();
@@ -2175,13 +2284,44 @@ void CreateEmptyOrthoVoxelGridUniformBuffer(
 		OrthoGridUniformBufferParameters->bEnableIndirectionGrid = false;
 		OrthoGridUniformBufferParameters->MajorantGridBuffer = GraphBuilder.CreateSRV(GSystemTextures.GetDefaultStructuredBuffer(GraphBuilder, sizeof(FScalarGridData)));
 	}
-	OrthoGridUniformBuffer = GraphBuilder.CreateUniformBuffer(OrthoGridUniformBufferParameters);
+	return GraphBuilder.CreateUniformBuffer(OrthoGridUniformBufferParameters);
+}
+
+namespace HeterogeneousVolumes
+{
+	TRDGUniformBufferRef<FOrthoVoxelGridUniformBufferParameters> GetOrthoVoxelGridUniformBuffer(
+		FRDGBuilder& GraphBuilder,
+		FSceneViewState* ViewState
+	)
+	{
+		if (ViewState && ViewState->OrthoVoxelGridUniformBuffer)
+		{
+			return ViewState->OrthoVoxelGridUniformBuffer;
+		}
+
+		return CreateEmptyOrthoVoxelGridUniformBuffer(GraphBuilder);
+	}
+
+	TRDGUniformBufferRef<FFrustumVoxelGridUniformBufferParameters> GetFrustumVoxelGridUniformBuffer(
+		FRDGBuilder& GraphBuilder,
+		FSceneViewState* ViewState
+	)
+	{
+		if (ViewState && ViewState->FrustumVoxelGridUniformBuffer)
+		{
+			return ViewState->FrustumVoxelGridUniformBuffer;
+		}
+
+		return CreateEmptyFrustumVoxelGridUniformBuffer(GraphBuilder);
+	}
 }
 
 void RasterizeVolumesIntoOrthoVoxelGrid(
 	FRDGBuilder& GraphBuilder,
 	const FScene* Scene,
 	const TArray<FViewInfo>& Views,
+	const TSet<FVolumetricMeshBatch>& HeterogeneousVolumesMeshBatches,
+	const FVoxelGridBuildOptions BuildOptions,
 	// Raster tile
 	FRDGBufferRef RasterTileBuffer,
 	FRDGBufferRef RasterTileAllocatorBuffer,
@@ -2272,10 +2412,11 @@ void RasterizeVolumesIntoOrthoVoxelGrid(
 		FVector WorldCameraOrigin = View.ViewMatrices.GetViewOrigin();
 		FBoxSphereBounds WorldCameraBounds(FSphere(WorldCameraOrigin, HeterogeneousVolumes::GetMaxTraceDistance()));
 
-		for (int32 MeshBatchIndex = 0; MeshBatchIndex < View.HeterogeneousVolumesMeshBatches.Num(); ++MeshBatchIndex)
+		for (auto MeshBatchIt = HeterogeneousVolumesMeshBatches.begin(); MeshBatchIt != HeterogeneousVolumesMeshBatches.end(); ++MeshBatchIt)
 		{
-			const FMeshBatch* Mesh = View.HeterogeneousVolumesMeshBatches[MeshBatchIndex].Mesh;
-			const FPrimitiveSceneProxy* PrimitiveSceneProxy = View.HeterogeneousVolumesMeshBatches[MeshBatchIndex].Proxy;
+			FVolumetricMeshBatch VolumetricMeshBatch = *MeshBatchIt;
+			const FMeshBatch* Mesh = VolumetricMeshBatch.Mesh;
+			const FPrimitiveSceneProxy* PrimitiveSceneProxy = VolumetricMeshBatch.Proxy;
 			const FMaterialRenderProxy* MaterialRenderProxy = Mesh->MaterialRenderProxy;
 			if (!ShouldRenderMeshBatchWithHeterogeneousVolumes(Mesh, PrimitiveSceneProxy, View.GetFeatureLevel()))
 			{
@@ -2328,6 +2469,7 @@ void RasterizeVolumesIntoOrthoVoxelGrid(
 				PassParameters->IndirectArgs = RasterizeBottomLevelGridIndirectArgsBuffer;
 
 				// Sampling mode
+				PassParameters->bJitter = BuildOptions.bJitter;
 				PassParameters->bSampleAtVertices = HeterogeneousVolumes::EnableLinearInterpolation();
 
 				// Grid data
@@ -2381,35 +2523,126 @@ void RasterizeVolumesIntoOrthoVoxelGrid(
 	}
 }
 
+uint32 GetTypeHash(const FVolumetricMeshBatch& MeshBatch)
+{
+	return HashCombineFast(GetTypeHash(MeshBatch.Mesh), GetTypeHash(MeshBatch.Proxy));
+}
+
+const FProjectedShadowInfo* GetProjectedShadowInfo(const FVisibleLightInfo* VisibleLightInfo)
+{
+	if (VisibleLightInfo)
+	{
+		for (int32 ShadowIndex = 0; ShadowIndex < VisibleLightInfo->ShadowsToProject.Num(); ShadowIndex++)
+		{
+			FProjectedShadowInfo* ProjectedShadowInfo = VisibleLightInfo->ShadowsToProject[ShadowIndex];
+			if (ProjectedShadowInfo)
+			{
+				return ProjectedShadowInfo;
+			}
+		}
+	}
+
+	return nullptr;
+}
+
+void CollectHeterogeneousVolumeMeshBatches(
+	FRDGBuilder& GraphBuilder,
+	const FScene* Scene,
+	const TArray<FViewInfo>& Views,
+	const TArray<FVisibleLightInfo, SceneRenderingAllocator> VisibleLightInfos,
+	TSet<FVolumetricMeshBatch>& HeterogeneousVolumesMeshBatches
+)
+{
+	HeterogeneousVolumesMeshBatches.Reset();
+
+	for (int32 ViewIndex = 0; ViewIndex < Views.Num(); ViewIndex++)
+	{
+		const FViewInfo& View = Views[ViewIndex];
+		for (int32 MeshBatchIndex = 0; MeshBatchIndex < View.HeterogeneousVolumesMeshBatches.Num(); ++MeshBatchIndex)
+		{
+			HeterogeneousVolumesMeshBatches.FindOrAdd(View.HeterogeneousVolumesMeshBatches[MeshBatchIndex]);
+		}
+	}
+
+	TArray<FLightSceneInfoCompact, TInlineAllocator<64>> LightSceneInfoCompact;
+	for (auto LightIt = Scene->Lights.CreateConstIterator(); LightIt; ++LightIt)
+	{
+		// TODO: Use global bounds information..
+		//if (LightIt->AffectsPrimitive(HeterogeneousVolumeInterface->GetBounds(), HeterogeneousVolumeInterface->GetPrimitiveSceneProxy()))
+		{
+			LightSceneInfoCompact.Add(*LightIt);
+		}
+	}
+
+	int32 NumPasses = LightSceneInfoCompact.Num();
+	for (int32 PassIndex = 0; PassIndex < NumPasses; ++PassIndex)
+	{
+		const FLightSceneInfo* LightSceneInfo = LightSceneInfoCompact[PassIndex].LightSceneInfo;
+		if (LightSceneInfo->Proxy->CastsVolumetricShadow())
+		{
+			const FVisibleLightInfo* VisibleLightInfo = &VisibleLightInfos[LightSceneInfo->Id];
+			const FProjectedShadowInfo* ProjectedShadowInfo = GetProjectedShadowInfo(VisibleLightInfo);
+			if (ProjectedShadowInfo != nullptr)
+			{
+				const TArray<FMeshBatchAndRelevance, SceneRenderingAllocator>& MeshBatches = ProjectedShadowInfo->GetDynamicSubjectHeterogeneousVolumeMeshElements();
+				for (int32 MeshBatchIndex = 0; MeshBatchIndex < MeshBatches.Num(); ++MeshBatchIndex)
+				{
+					//HeterogeneousVolumesMeshBatches.FindOrAdd(MeshBatches[MeshBatchIndex]);
+					HeterogeneousVolumesMeshBatches.FindOrAdd(FVolumetricMeshBatch(MeshBatches[MeshBatchIndex].Mesh, MeshBatches[MeshBatchIndex].PrimitiveSceneProxy));
+				}
+			}
+		}
+	}
+}
+
 void BuildOrthoVoxelGrid(
 	FRDGBuilder& GraphBuilder,
 	const FScene* Scene,
 	/*const*/ TArray<FViewInfo>& Views,
+	const TArray<FVisibleLightInfo, SceneRenderingAllocator>& VisibleLightInfos,
+	const FVoxelGridBuildOptions& BuildOptions,
 	TRDGUniformBufferRef<FOrthoVoxelGridUniformBufferParameters>& OrthoGridUniformBuffer
 )
 {
-	if (!ShouldRenderHeterogeneousVolumes(Scene) || !ShouldRenderHeterogeneousVolumesForAnyView(Views) || !HeterogeneousVolumes::EnableOrthoVoxelGrid())
+	if (!ShouldRenderHeterogeneousVolumes(Scene) || !HeterogeneousVolumes::EnableOrthoVoxelGrid())
 	{
-		CreateEmptyOrthoVoxelGridUniformBuffer(GraphBuilder, OrthoGridUniformBuffer);
+		OrthoGridUniformBuffer = CreateEmptyOrthoVoxelGridUniformBuffer(GraphBuilder);
 		return;
 	}
 	check(!Views.IsEmpty());
 
 	RDG_EVENT_SCOPE(GraphBuilder, "Ortho Grid Build");
 
+	TSet<FVolumetricMeshBatch> HeterogeneousVolumesMeshBatches;
+	CollectHeterogeneousVolumeMeshBatches(
+		GraphBuilder,
+		Scene,
+		Views,
+		VisibleLightInfos,
+		HeterogeneousVolumesMeshBatches
+	);
+
+	if (HeterogeneousVolumesMeshBatches.IsEmpty())
+	{
+		OrthoGridUniformBuffer = CreateEmptyOrthoVoxelGridUniformBuffer(GraphBuilder);
+		return;
+	}
+
 	// Collect global bounds
 	FBoxSphereBounds TopLevelGridBounds;
 	float GlobalMinimumVoxelSize;
 	CalcGlobalBoundsAndMinimumVoxelSize(
 		GraphBuilder,
 		Views,
+		HeterogeneousVolumesMeshBatches,
+		BuildOptions,
 		TopLevelGridBounds,
 		GlobalMinimumVoxelSize
 	);
 
 	if (TopLevelGridBounds.SphereRadius == 0)
 	{
-		CreateEmptyOrthoVoxelGridUniformBuffer(GraphBuilder, OrthoGridUniformBuffer);
+		OrthoGridUniformBuffer = CreateEmptyOrthoVoxelGridUniformBuffer(GraphBuilder);
 		return;
 	}
 
@@ -2427,6 +2660,8 @@ void BuildOrthoVoxelGrid(
 		GraphBuilder,
 		Scene,
 		Views,
+		HeterogeneousVolumesMeshBatches,
+		BuildOptions,
 		TopLevelGridBounds,
 		TopLevelGridResolution,
 		TopLevelGridBuffer
@@ -2467,6 +2702,8 @@ void BuildOrthoVoxelGrid(
 		GraphBuilder,
 		Scene,
 		Views,
+		HeterogeneousVolumesMeshBatches,
+		BuildOptions,
 		// Tile data
 		RasterTileBuffer,
 		RasterTileAllocatorBuffer,
@@ -2505,6 +2742,7 @@ void BuildOrthoVoxelGrid(
 		OrthoGridUniformBufferParameters->bEnableIndirectionGrid = HeterogeneousVolumes::EnableIndirectionGrid();
 		OrthoGridUniformBufferParameters->MajorantGridBuffer = GraphBuilder.CreateSRV(MajorantGridBuffer);
 	}
+
 	OrthoGridUniformBuffer = GraphBuilder.CreateUniformBuffer(OrthoGridUniformBufferParameters);
 }
 
@@ -2565,4 +2803,971 @@ void RenderTransmittanceWithVoxelGrid(
 		PassParameters,
 		GroupCount
 	);
-}
\ No newline at end of file
+}
+
+struct FVolumetricShadowMapDebugData
+{
+	FVector3f LightRayStart;
+	FVector3f LightRayEnd;
+	FVector3f RayOrigin;
+	FVector3f RayEnd;
+	float HitSpan[2];
+};
+
+class FRenderVolumetricShadowMapForLightWithVoxelGridCS : public FGlobalShader
+{
+	DECLARE_GLOBAL_SHADER(FRenderVolumetricShadowMapForLightWithVoxelGridCS);
+	SHADER_USE_PARAMETER_STRUCT(FRenderVolumetricShadowMapForLightWithVoxelGridCS, FGlobalShader);
+
+	class FUseAVSMCompression : SHADER_PERMUTATION_BOOL("USE_AVSM_COMPRESSION");
+	using FPermutationDomain = TShaderPermutationDomain<FUseAVSMCompression>;
+
+	BEGIN_SHADER_PARAMETER_STRUCT(FParameters, )
+		// Scene data
+		SHADER_PARAMETER_STRUCT_REF(FViewUniformShaderParameters, View)
+		SHADER_PARAMETER_STRUCT_INCLUDE(FSceneTextureParameters, SceneTextures)
+
+		// Shadow data
+		SHADER_PARAMETER(float, ShadowStepSize)
+		SHADER_PARAMETER(float, ShadowStepFactor)
+
+		// Volumetric Shadow Map data
+		SHADER_PARAMETER(FVector3f, TranslatedWorldOrigin)
+		SHADER_PARAMETER(FIntPoint, ShadowResolution)
+		SHADER_PARAMETER(int, MaxSampleCount)
+		SHADER_PARAMETER(float, AbsoluteErrorThreshold)
+		SHADER_PARAMETER(float, RelativeErrorThreshold)
+
+		SHADER_PARAMETER(int, NumShadowMatrices)
+		SHADER_PARAMETER_ARRAY(FMatrix44f, TranslatedWorldToShadow, [6])
+		SHADER_PARAMETER_ARRAY(FMatrix44f, ShadowToTranslatedWorld, [6])
+
+		// Volume data
+		SHADER_PARAMETER_RDG_UNIFORM_BUFFER(FOrthoVoxelGridUniformBufferParameters, OrthoGridUniformBuffer)
+		SHADER_PARAMETER_RDG_UNIFORM_BUFFER(FFrustumVoxelGridUniformBufferParameters, FrustumGridUniformBuffer)
+
+		// Ray data
+		SHADER_PARAMETER(float, MaxTraceDistance)
+		SHADER_PARAMETER(int, MaxStepCount)
+		SHADER_PARAMETER(int, bJitter)
+
+		// Dispatch data
+		SHADER_PARAMETER(FIntVector, GroupCount)
+
+		// Output
+		SHADER_PARAMETER_RDG_BUFFER_UAV(RWBuffer<int>, RWVolumetricShadowLinkedListAllocatorBuffer)
+		SHADER_PARAMETER_RDG_BUFFER_UAV(RWStructuredBuffer<int2>, RWVolumetricShadowLinkedListBuffer)
+		SHADER_PARAMETER_RDG_TEXTURE_UAV(RWTexture2D<float4>, RWBeerShadowMapTexture)
+
+		// Debug
+		SHADER_PARAMETER_RDG_BUFFER_UAV(RWStructuredBuffer<FVolumetricShadowMapDebugData>, RWDebugBuffer)
+	END_SHADER_PARAMETER_STRUCT()
+
+	static bool ShouldCompilePermutation(
+		const FGlobalShaderPermutationParameters& Parameters
+	)
+	{
+		return DoesPlatformSupportHeterogeneousVolumes(Parameters.Platform);
+	}
+
+	static void ModifyCompilationEnvironment(
+		const FGlobalShaderPermutationParameters& Parameters,
+		FShaderCompilerEnvironment& OutEnvironment
+	)
+	{
+		FMaterialShader::ModifyCompilationEnvironment(Parameters, OutEnvironment);
+		OutEnvironment.SetDefine(TEXT("THREADGROUP_SIZE_1D"), GetThreadGroupSize1D());
+		OutEnvironment.SetDefine(TEXT("THREADGROUP_SIZE_2D"), GetThreadGroupSize2D());
+		OutEnvironment.SetDefine(TEXT("THREADGROUP_SIZE_3D"), GetThreadGroupSize3D());
+
+		// Temporary disabling..
+		OutEnvironment.SetDefine(TEXT("DIM_USE_TRANSMITTANCE_VOLUME"), 0);
+		OutEnvironment.SetDefine(TEXT("DIM_USE_INSCATTERING_VOLUME"), 0);
+		OutEnvironment.SetDefine(TEXT("DIM_USE_LUMEN_GI"), 0);
+
+		OutEnvironment.CompilerFlags.Add(CFLAG_AllowTypedUAVLoads);
+	}
+
+	static int32 GetThreadGroupSize1D() { return GetThreadGroupSize3D() * GetThreadGroupSize3D() * GetThreadGroupSize3D(); }
+	static int32 GetThreadGroupSize2D() { return 8; }
+	static int32 GetThreadGroupSize3D() { return 4; }
+};
+
+IMPLEMENT_GLOBAL_SHADER(FRenderVolumetricShadowMapForLightWithVoxelGridCS, "/Engine/Private/HeterogeneousVolumes/HeterogeneousVolumesVoxelGridShadows.usf", "RenderVolumetricShadowMapForLightWithVoxelGridCS", SF_Compute);
+
+
+class FCompressVolumetricShadowMapCS : public FGlobalShader
+{
+	DECLARE_GLOBAL_SHADER(FCompressVolumetricShadowMapCS);
+	SHADER_USE_PARAMETER_STRUCT(FCompressVolumetricShadowMapCS, FGlobalShader);
+
+	class FUseAVSMCompression : SHADER_PERMUTATION_BOOL("USE_AVSM_COMPRESSION");
+	using FPermutationDomain = TShaderPermutationDomain<FUseAVSMCompression>;
+
+	BEGIN_SHADER_PARAMETER_STRUCT(FParameters, )
+		// Input
+		SHADER_PARAMETER(FIntPoint, ShadowResolution)
+		SHADER_PARAMETER(int, MaxSampleCount)
+		SHADER_PARAMETER_RDG_BUFFER_SRV(StructuredBuffer<int2>, VolumetricShadowLinkedListBuffer)
+
+		// Output
+		SHADER_PARAMETER_RDG_BUFFER_UAV(RWBuffer<int>, RWVolumetricShadowIndirectionAllocatorBuffer)
+		SHADER_PARAMETER_RDG_BUFFER_UAV(RWStructuredBuffer<FAVSMIndirectionPackedData>, RWVolumetricShadowIndirectionBuffer)
+		SHADER_PARAMETER_RDG_BUFFER_UAV(RWStructuredBuffer<FAVSMSamplePackedData>, RWVolumetricShadowTransmittanceBuffer)
+	END_SHADER_PARAMETER_STRUCT()
+
+	static bool ShouldCompilePermutation(
+		const FGlobalShaderPermutationParameters& Parameters
+	)
+	{
+		return DoesPlatformSupportHeterogeneousVolumes(Parameters.Platform);
+	}
+
+	static void ModifyCompilationEnvironment(
+		const FGlobalShaderPermutationParameters& Parameters,
+		FShaderCompilerEnvironment& OutEnvironment
+	)
+	{
+		FMaterialShader::ModifyCompilationEnvironment(Parameters, OutEnvironment);
+		OutEnvironment.SetDefine(TEXT("THREADGROUP_SIZE_1D"), GetThreadGroupSize1D());
+		OutEnvironment.SetDefine(TEXT("THREADGROUP_SIZE_2D"), GetThreadGroupSize2D());
+		OutEnvironment.SetDefine(TEXT("THREADGROUP_SIZE_3D"), GetThreadGroupSize3D());
+
+		// Temporary disabling..
+		OutEnvironment.SetDefine(TEXT("DIM_USE_TRANSMITTANCE_VOLUME"), 0);
+		OutEnvironment.SetDefine(TEXT("DIM_USE_INSCATTERING_VOLUME"), 0);
+		OutEnvironment.SetDefine(TEXT("DIM_USE_LUMEN_GI"), 0);
+
+		OutEnvironment.CompilerFlags.Add(CFLAG_AllowTypedUAVLoads);
+	}
+
+	static int32 GetThreadGroupSize1D() { return GetThreadGroupSize3D() * GetThreadGroupSize3D() * GetThreadGroupSize3D(); }
+	static int32 GetThreadGroupSize2D() { return 8; }
+	static int32 GetThreadGroupSize3D() { return 4; }
+};
+
+IMPLEMENT_GLOBAL_SHADER(FCompressVolumetricShadowMapCS, "/Engine/Private/HeterogeneousVolumes/HeterogeneousVolumesVoxelGridShadows.usf", "CompressVolumetricShadowMapCS", SF_Compute);
+
+void CreateAdaptiveVolumetricShadowMapUniformBuffer(
+	FRDGBuilder& GraphBuilder,
+	const FVector3f& TranslatedWorldOrigin,
+	const FVector4f& TranslatedWorldPlane,
+	const FMatrix44f* TranslatedWorldToShadow,
+	FIntPoint VolumetricShadowMapResolution,
+	int32 NumShadowMatrices,
+	uint32 VolumetricShadowMapMaxSampleCount,
+	bool bIsDirectionalLight,
+	FRDGBufferRef VolumetricShadowMapLinkedListBuffer,
+	FRDGBufferRef VolumetricShadowMapIndirectionBuffer,
+	FRDGBufferRef VolumetricShadowMapSampleBuffer,
+	TRDGUniformBufferRef<FAdaptiveVolumetricShadowMapUniformBufferParameters>& AdaptiveVolumetricShadowMapUniformBuffer
+)
+{
+	FAdaptiveVolumetricShadowMapUniformBufferParameters* UniformBufferParameters = GraphBuilder.AllocParameters<FAdaptiveVolumetricShadowMapUniformBufferParameters>();
+	{
+		UniformBufferParameters->NumShadowMatrices = NumShadowMatrices;
+		for (int32 i = 0; i < NumShadowMatrices; ++i)
+		{
+			UniformBufferParameters->TranslatedWorldToShadow[i] = TranslatedWorldToShadow[i];
+		}
+		UniformBufferParameters->TranslatedWorldOrigin = TranslatedWorldOrigin;
+		UniformBufferParameters->TranslatedWorldPlane = TranslatedWorldPlane;
+		UniformBufferParameters->Resolution = VolumetricShadowMapResolution;
+		UniformBufferParameters->MaxSampleCount = VolumetricShadowMapMaxSampleCount;
+		UniformBufferParameters->bIsEmpty = (VolumetricShadowMapResolution == FIntPoint::ZeroValue);
+		UniformBufferParameters->bIsDirectionalLight = bIsDirectionalLight;
+		UniformBufferParameters->LinkedListBuffer = GraphBuilder.CreateSRV(VolumetricShadowMapLinkedListBuffer);
+		UniformBufferParameters->IndirectionBuffer = GraphBuilder.CreateSRV(VolumetricShadowMapIndirectionBuffer);
+		UniformBufferParameters->SampleBuffer = GraphBuilder.CreateSRV(VolumetricShadowMapSampleBuffer);
+	}
+
+	AdaptiveVolumetricShadowMapUniformBuffer = GraphBuilder.CreateUniformBuffer(UniformBufferParameters);
+}
+
+namespace HeterogeneousVolumes {
+
+	TRDGUniformBufferRef<FAdaptiveVolumetricShadowMapUniformBufferParameters> CreateEmptyAdaptiveVolumetricShadowMapUniformBuffer(
+		FRDGBuilder& GraphBuilder
+	)
+	{
+		FAdaptiveVolumetricShadowMapUniformBufferParameters* UniformBufferParameters = GraphBuilder.AllocParameters<FAdaptiveVolumetricShadowMapUniformBufferParameters>();
+		{
+			UniformBufferParameters->NumShadowMatrices = 1;
+			for (int32 i = 0; i < UniformBufferParameters->NumShadowMatrices; ++i)
+			{
+				UniformBufferParameters->TranslatedWorldToShadow[i] = FMatrix44f::Identity;
+			}
+			UniformBufferParameters->TranslatedWorldOrigin = FVector3f::ZeroVector;
+			UniformBufferParameters->TranslatedWorldPlane = FVector4f::Zero();
+			UniformBufferParameters->Resolution = FIntPoint::ZeroValue;
+			UniformBufferParameters->MaxSampleCount = 0;
+			UniformBufferParameters->bIsEmpty = true;
+			UniformBufferParameters->bIsDirectionalLight = false;
+			UniformBufferParameters->LinkedListBuffer = GraphBuilder.CreateSRV(GSystemTextures.GetDefaultStructuredBuffer(GraphBuilder, sizeof(FAVSMLinkedListPackedData)));
+			UniformBufferParameters->IndirectionBuffer = GraphBuilder.CreateSRV(GSystemTextures.GetDefaultStructuredBuffer(GraphBuilder, sizeof(FAVSMIndirectionPackedData)));
+			UniformBufferParameters->SampleBuffer = GraphBuilder.CreateSRV(GSystemTextures.GetDefaultStructuredBuffer(GraphBuilder, sizeof(FAVSMSamplePackedData)));
+		}
+
+		TRDGUniformBufferRef<FAdaptiveVolumetricShadowMapUniformBufferParameters> AdaptiveVolumetricShadowMapUniformBuffer = GraphBuilder.CreateUniformBuffer(UniformBufferParameters);
+		return AdaptiveVolumetricShadowMapUniformBuffer;
+	}
+
+	void DestroyAdaptiveVolumetricShadowMapUniformBuffer(
+		TRDGUniformBufferRef<FAdaptiveVolumetricShadowMapUniformBufferParameters>& AdaptiveVolumetricShadowMapUniformBuffer
+	)
+	{
+		AdaptiveVolumetricShadowMapUniformBuffer = nullptr;
+	}
+
+	TRDGUniformBufferRef<FAdaptiveVolumetricShadowMapUniformBufferParameters> GetAdaptiveVolumetricShadowMapUniformBuffer(
+		FRDGBuilder& GraphBuilder,
+		FSceneViewState* ViewState,
+		const FLightSceneInfo* LightSceneInfo
+	)
+	{
+		if (ViewState && LightSceneInfo)
+		{
+			TRDGUniformBufferRef<FAdaptiveVolumetricShadowMapUniformBufferParameters>* AdaptiveVolumetricShadowMapUniformBuffer = ViewState->AdaptiveVolumetricShadowMapUniformBufferMap.Find(LightSceneInfo->Id);
+			if (AdaptiveVolumetricShadowMapUniformBuffer != nullptr)
+			{
+				return *AdaptiveVolumetricShadowMapUniformBuffer;
+			}
+		}
+		
+		return CreateEmptyAdaptiveVolumetricShadowMapUniformBuffer(GraphBuilder);
+	}
+
+	TRDGUniformBufferRef<FAdaptiveVolumetricShadowMapUniformBufferParameters> GetAdaptiveVolumetricCameraMapUniformBuffer(
+		FRDGBuilder& GraphBuilder,
+		FSceneViewState* ViewState
+	)
+	{
+		if (ViewState)
+		{
+			return ViewState->AdaptiveVolumetricCameraMapUniformBuffer;
+		}
+		
+		return CreateEmptyAdaptiveVolumetricShadowMapUniformBuffer(GraphBuilder);
+	}
+
+	FAdaptiveVolumetricShadowMapUniformBufferParameters GetAdaptiveVolumetricCameraMapParameters(
+		FRDGBuilder& GraphBuilder,
+		FSceneViewState* ViewState)
+	{
+		FAdaptiveVolumetricShadowMapUniformBufferParameters Parameters;
+
+		if (ViewState && ViewState->AdaptiveVolumetricCameraMapUniformBuffer)
+		{
+			TRDGUniformBufferRef<FAdaptiveVolumetricShadowMapUniformBufferParameters> UniformBuffer = ViewState->AdaptiveVolumetricCameraMapUniformBuffer;
+
+			Parameters.NumShadowMatrices = UniformBuffer->GetParameters()->NumShadowMatrices;
+			for (int32 i = 0; i < Parameters.NumShadowMatrices; ++i)
+			{
+				Parameters.TranslatedWorldToShadow[i] = UniformBuffer->GetParameters()->TranslatedWorldToShadow[i];
+			}
+			Parameters.TranslatedWorldOrigin = UniformBuffer->GetParameters()->TranslatedWorldOrigin;
+			Parameters.TranslatedWorldPlane = UniformBuffer->GetParameters()->TranslatedWorldPlane;
+			Parameters.Resolution = UniformBuffer->GetParameters()->Resolution;
+			Parameters.MaxSampleCount = UniformBuffer->GetParameters()->MaxSampleCount;
+			Parameters.bIsEmpty = UniformBuffer->GetParameters()->bIsEmpty;
+			Parameters.bIsDirectionalLight = UniformBuffer->GetParameters()->bIsDirectionalLight;
+			Parameters.LinkedListBuffer = UniformBuffer->GetParameters()->LinkedListBuffer;
+			Parameters.IndirectionBuffer = UniformBuffer->GetParameters()->IndirectionBuffer;
+			Parameters.SampleBuffer = UniformBuffer->GetParameters()->SampleBuffer;
+		}
+		else
+		{
+			Parameters.NumShadowMatrices = 1;
+			for (int32 i = 0; i < Parameters.NumShadowMatrices; ++i)
+			{
+				Parameters.TranslatedWorldToShadow[i] = FMatrix44f::Identity;
+			}
+			Parameters.TranslatedWorldOrigin = FVector3f::ZeroVector;
+			Parameters.TranslatedWorldPlane = FVector4f::Zero();
+			Parameters.Resolution = FIntPoint::ZeroValue;
+			Parameters.MaxSampleCount = 0;
+			Parameters.bIsEmpty = true;
+			Parameters.bIsDirectionalLight = false;
+			Parameters.LinkedListBuffer = GraphBuilder.CreateSRV(GSystemTextures.GetDefaultStructuredBuffer(GraphBuilder, sizeof(FAVSMLinkedListPackedData)));
+			Parameters.IndirectionBuffer = GraphBuilder.CreateSRV(GSystemTextures.GetDefaultStructuredBuffer(GraphBuilder, sizeof(FAVSMIndirectionPackedData)));
+			Parameters.SampleBuffer = GraphBuilder.CreateSRV(GSystemTextures.GetDefaultStructuredBuffer(GraphBuilder, sizeof(FAVSMSamplePackedData)));
+		}
+
+		return Parameters;
+	}
+
+} // namespace HeterogeneousVolumes
+
+void CompressVolumetricShadowMap(
+	FRDGBuilder& GraphBuilder,
+	FViewInfo& View,
+	FIntVector GroupCount,
+	// Input
+	FIntPoint ShadowMapResolution,
+	uint32 MaxSampleCount,
+	FRDGBufferRef VolumetricShadowLinkedListBuffer,
+	// Output
+	FRDGBufferRef& VolumetricShadowIndirectionBuffer,
+	FRDGBufferRef& VolumetricShadowTransmittanceBuffer
+)
+{
+	int32 VolumetricShadowPixelCount = ShadowMapResolution.X * ShadowMapResolution.Y * GroupCount.Z;
+	FRDGBufferRef VolumetricShadowIndirectionAllocatorBuffer = GraphBuilder.CreateBuffer(
+		FRDGBufferDesc::CreateBufferDesc(sizeof(uint32), 1),
+		TEXT("HeterogeneousVolume.VolumetricShadowIndirectionAllocatorBuffer")
+	);
+	AddClearUAVPass(GraphBuilder, GraphBuilder.CreateUAV(VolumetricShadowIndirectionAllocatorBuffer, PF_R32_UINT), 0);
+
+	VolumetricShadowIndirectionBuffer = GraphBuilder.CreateBuffer(
+		FRDGBufferDesc::CreateStructuredDesc(sizeof(FAVSMIndirectionPackedData), VolumetricShadowPixelCount),
+		TEXT("HeterogeneousVolume.VolumetricShadowIndirectionBuffer")
+	);
+
+	int32 FVolumetricShadowTransmittanceMaxCount = VolumetricShadowPixelCount * MaxSampleCount;
+	VolumetricShadowTransmittanceBuffer = GraphBuilder.CreateBuffer(
+		FRDGBufferDesc::CreateStructuredDesc(sizeof(FAVSMSamplePackedData), FVolumetricShadowTransmittanceMaxCount),
+		TEXT("HeterogeneousVolume.VolumetricShadowTransmittanceBuffer")
+	);
+
+	FCompressVolumetricShadowMapCS::FParameters* PassParameters = GraphBuilder.AllocParameters<FCompressVolumetricShadowMapCS::FParameters>();
+	{
+		// Input
+		PassParameters->ShadowResolution = ShadowMapResolution;
+		PassParameters->MaxSampleCount = MaxSampleCount;
+		PassParameters->VolumetricShadowLinkedListBuffer = GraphBuilder.CreateSRV(VolumetricShadowLinkedListBuffer);
+
+		// Output
+		PassParameters->RWVolumetricShadowIndirectionAllocatorBuffer = GraphBuilder.CreateUAV(VolumetricShadowIndirectionAllocatorBuffer, PF_R32_UINT);
+		PassParameters->RWVolumetricShadowIndirectionBuffer = GraphBuilder.CreateUAV(VolumetricShadowIndirectionBuffer);
+		PassParameters->RWVolumetricShadowTransmittanceBuffer = GraphBuilder.CreateUAV(VolumetricShadowTransmittanceBuffer);
+	}
+
+	FCompressVolumetricShadowMapCS::FPermutationDomain PermutationVector;
+	TShaderRef<FCompressVolumetricShadowMapCS> ComputeShader = View.ShaderMap->GetShader<FCompressVolumetricShadowMapCS>(PermutationVector);
+	FComputeShaderUtils::AddPass(
+		GraphBuilder,
+		// Add Light name..
+		RDG_EVENT_NAME("CompressVolumetricShadowMapCS"),
+		ERDGPassFlags::Compute | ERDGPassFlags::NeverCull,
+		ComputeShader,
+		PassParameters,
+		GroupCount
+	);
+}
+
+void RenderVolumetricShadowMapForLightWithVoxelGrid(
+	FRDGBuilder& GraphBuilder,
+	// Scene data
+	const FSceneTextures& SceneTextures,
+	FScene* Scene,
+	const FSceneViewFamily& ViewFamily,
+	FViewInfo& View,
+	// Light data
+	bool bApplyEmissionAndTransmittance,
+	bool bApplyDirectLighting,
+	bool bApplyShadowTransmittance,
+	uint32 LightType,
+	const FLightSceneInfo* LightSceneInfo,
+	// Shadow data
+	const FVisibleLightInfo* VisibleLightInfo,
+	const FVirtualShadowMapArray& VirtualShadowMapArray,
+	// Volume data
+	const TRDGUniformBufferRef<FOrthoVoxelGridUniformBufferParameters>& OrthoGridUniformBuffer,
+	const TRDGUniformBufferRef<FFrustumVoxelGridUniformBufferParameters>& FrustumGridUniformBuffer,
+	// Output
+	bool& bIsDirectionalLight,
+	FVector3f& TranslatedWorldOrigin,
+	FVector4f& TranslatedWorldPlane,
+	FMatrix44f* TranslatedWorldToShadow,
+	FIntVector& GroupCount,
+	int32& NumShadowMatrices,
+	FIntPoint& ShadowMapResolution,
+	uint32& MaxSampleCount,
+	FRDGTextureRef& BeerShadowMapTexture,
+	FRDGBufferRef& VolumetricShadowLinkedListBuffer
+)
+{
+	check(LightSceneInfo);
+	check(VisibleLightInfo);
+
+	// TODO: Push ProjectedShadowInfo from calling function
+	const FProjectedShadowInfo* ProjectedShadowInfo = GetProjectedShadowInfo(VisibleLightInfo);
+	check(ProjectedShadowInfo != NULL)
+
+	ShadowMapResolution = HeterogeneousVolumes::GetShadowMapResolution();
+
+	bool bIsMultiProjection = (LightType == LightType_Point) || (LightType == LightType_Rect);
+	GroupCount = FIntVector(1);
+	GroupCount.X = FMath::DivideAndRoundUp(ShadowMapResolution.X, FRenderVolumetricShadowMapForLightWithVoxelGridCS::GetThreadGroupSize2D());
+	GroupCount.Y = FMath::DivideAndRoundUp(ShadowMapResolution.Y, FRenderVolumetricShadowMapForLightWithVoxelGridCS::GetThreadGroupSize2D());
+	GroupCount.Z = bIsMultiProjection ? 6 : 1;
+
+	// TODO: Allocate debug data
+	int32 BufferSize = ShadowMapResolution.X * ShadowMapResolution.Y;
+	if (bIsMultiProjection)
+	{
+		BufferSize *= 6;
+	}
+	FRDGBufferRef DebugBuffer = GraphBuilder.CreateBuffer(
+		FRDGBufferDesc::CreateStructuredDesc(sizeof(FVolumetricShadowMapDebugData), BufferSize),
+		TEXT("HeterogeneousVolume.RenderDebugData")
+	);
+
+	FRDGBufferRef VolumetricShadowLinkedListAllocatorBuffer = GraphBuilder.CreateBuffer(
+		FRDGBufferDesc::CreateBufferDesc(sizeof(uint32), 1),
+		TEXT("HeterogeneousVolume.VolumetricShadowLinkedListAllocatorBuffer")
+	);
+
+	MaxSampleCount = HeterogeneousVolumes::GetShadowMaxSampleCount();
+	int32 VolumetricShadowLinkedListElementCount = ShadowMapResolution.X * ShadowMapResolution.Y * MaxSampleCount;
+	if (bIsMultiProjection)
+	{
+		VolumetricShadowLinkedListElementCount *= 6;
+	}
+	VolumetricShadowLinkedListBuffer = GraphBuilder.CreateBuffer(
+		FRDGBufferDesc::CreateStructuredDesc(sizeof(FAVSMLinkedListPackedData), VolumetricShadowLinkedListElementCount),
+		TEXT("HeterogeneousVolume.VolumetricShadowLinkedListBuffer")
+	);
+
+	const float RelativeErrorThreshold = HeterogeneousVolumes::GetShadowRelativeErrorThreshold();
+
+	// Initialize allocator to contain 1-spp
+	AddClearUAVPass(GraphBuilder, GraphBuilder.CreateUAV(VolumetricShadowLinkedListAllocatorBuffer, PF_R32_UINT), ShadowMapResolution.X * ShadowMapResolution.Y);
+
+	FDeferredLightUniformStruct DeferredLightUniform;
+	if (bApplyDirectLighting && (LightSceneInfo != nullptr))
+	{
+		DeferredLightUniform = GetDeferredLightParameters(View, *LightSceneInfo);
+	}
+	TUniformBufferRef<FDeferredLightUniformStruct> DeferredLightUB = CreateUniformBufferImmediate(DeferredLightUniform, UniformBuffer_SingleDraw);
+
+	FRenderVolumetricShadowMapForLightWithVoxelGridCS::FParameters* PassParameters = GraphBuilder.AllocParameters<FRenderVolumetricShadowMapForLightWithVoxelGridCS::FParameters>();
+	{
+		// Scene data
+		PassParameters->View = View.ViewUniformBuffer;
+		PassParameters->SceneTextures = GetSceneTextureParameters(GraphBuilder, SceneTextures);
+
+		// Ray Data
+		PassParameters->ShadowStepSize = HeterogeneousVolumes::GetShadowStepSize();
+		// TODO:
+		//PassParameters->ShadowStepFactor = HeterogeneousVolumes::GetShadowStepFactor();
+		PassParameters->ShadowStepFactor = 1.0;
+		PassParameters->MaxTraceDistance = HeterogeneousVolumes::GetMaxTraceDistance();
+		PassParameters->MaxStepCount = HeterogeneousVolumes::GetMaxStepCount();
+		PassParameters->bJitter = HeterogeneousVolumes::ShouldJitter();
+
+		PassParameters->NumShadowMatrices = ProjectedShadowInfo->OnePassShadowViewProjectionMatrices.Num();
+		if (PassParameters->NumShadowMatrices > 0)
+		{
+			FVector PreViewTranslation = View.ViewMatrices.GetPreViewTranslation();
+			FMatrix TranslatedWorldToWorldMatrix = FTranslationMatrix(-PreViewTranslation);
+			FVector LightPosition = LightSceneInfo->Proxy->GetPosition();
+			FMatrix WorldToLightMatrix = LightSceneInfo->Proxy->GetWorldToLight();
+
+			// Remove light rotation when building the RectLight projections..
+			FMatrix RotationalAdjustmentMatrix = FMatrix::Identity;
+			if (LightType == LIGHT_TYPE_RECT)
+			{
+				FVector LightDirection = LightSceneInfo->Proxy->GetDirection().GetSafeNormal();
+				RotationalAdjustmentMatrix = FRotationMatrix(LightDirection.Rotation());
+			}
+
+			FMatrix ViewMatrix[] = {
+				FLookFromMatrix(FVector::Zero(), FVector(-1, 0, 0), FVector(0, 0, 1)),
+				FLookFromMatrix(FVector::Zero(), FVector(1, 0, 0), FVector(0, 0, 1)),
+				FLookFromMatrix(FVector::Zero(), FVector(0, -1, 0), FVector(0, 0, 1)),
+				FLookFromMatrix(FVector::Zero(), FVector(0, 1, 0), FVector(0, 0, 1)),
+				FLookFromMatrix(FVector::Zero(), FVector(0, 0, -1), FVector(1, 0, 0)),
+				FLookFromMatrix(FVector::Zero(), FVector(0, 0, 1), FVector(1, 0, 0))
+			};
+
+			FMatrix PerspectiveMatrix = FPerspectiveMatrix(
+				PI / 4.0f,
+				ShadowMapResolution.X,
+				ShadowMapResolution.Y,
+				1.0,
+				LightSceneInfo->Proxy->GetRadius()
+			);
+
+			FMatrix ScreenMatrix = FScaleMatrix(FVector(0.5, -0.5, -0.5)) * FTranslationMatrix(FVector(0.5, 0.5, 0.5));
+
+			for (int32 i = 0; i < PassParameters->NumShadowMatrices; ++i)
+			{
+				FMatrix WorldToShadowMatrix = WorldToLightMatrix * RotationalAdjustmentMatrix * ViewMatrix[i] * PerspectiveMatrix * ScreenMatrix;
+				PassParameters->TranslatedWorldToShadow[i] = FMatrix44f(TranslatedWorldToWorldMatrix * WorldToShadowMatrix);
+				PassParameters->ShadowToTranslatedWorld[i] = PassParameters->TranslatedWorldToShadow[i].Inverse();
+			}
+			PassParameters->TranslatedWorldOrigin = FVector3f(PreViewTranslation + LightPosition);
+		}
+		else if (LightType == LightType_Directional)
+		{
+			// Build orthographic projection centered around volume..
+			FVector PreViewTranslation = View.ViewMatrices.GetPreViewTranslation();
+			FMatrix TranslatedWorldToWorldMatrix = FTranslationMatrix(-PreViewTranslation);
+
+			FVector WorldBoundsMin = FVector(OrthoGridUniformBuffer->GetParameters()->TopLevelGridWorldBoundsMin);
+			FVector WorldBoundsMax = FVector(OrthoGridUniformBuffer->GetParameters()->TopLevelGridWorldBoundsMax);
+			FBoxSphereBounds VolumeBounds(FBox(WorldBoundsMin, WorldBoundsMax));
+			FMatrix TranslationMatrix = FTranslationMatrix(-VolumeBounds.Origin);
+
+			FVector LightDirection = LightSceneInfo->Proxy->GetDirection().GetSafeNormal();
+			FMatrix RotationMatrix = FInverseRotationMatrix(LightDirection.Rotation());
+			FMatrix ScaleMatrix = FScaleMatrix(FVector(1.0 / VolumeBounds.SphereRadius));
+
+			const FMatrix FaceMatrix(
+				FPlane(0, 0, 1, 0),
+				FPlane(0, 1, 0, 0),
+				FPlane(-1, 0, 0, 0),
+				FPlane(0, 0, 0, 1));
+
+			// Invert Z to match reverse-Z for the rest of the shadow types!
+			FMatrix ScreenMatrix = FScaleMatrix(FVector(0.5, -0.5, -0.5)) * FTranslationMatrix(FVector(0.5, 0.5, 0.5));
+			FMatrix WorldToShadowMatrix = TranslationMatrix * RotationMatrix * ScaleMatrix * FaceMatrix * ScreenMatrix;
+			FMatrix44f TranslatedWorldToShadowMatrix = FMatrix44f(TranslatedWorldToWorldMatrix * WorldToShadowMatrix);
+
+			PassParameters->NumShadowMatrices = 1;
+			PassParameters->TranslatedWorldToShadow[0] = TranslatedWorldToShadowMatrix;
+			PassParameters->ShadowToTranslatedWorld[0] = TranslatedWorldToShadowMatrix.Inverse();
+			PassParameters->TranslatedWorldOrigin = FVector3f(PreViewTranslation + VolumeBounds.Origin - LightDirection * VolumeBounds.SphereRadius);
+		}
+		else
+		{
+			FVector PreViewTranslation = View.ViewMatrices.GetPreViewTranslation();
+			FMatrix TranslatedWorldToWorldMatrix = FTranslationMatrix(-PreViewTranslation);
+			FVector4f ShadowmapMinMax = FVector4f::Zero();
+			FMatrix WorldToShadowMatrix = ProjectedShadowInfo->GetWorldToShadowMatrix(ShadowmapMinMax);
+			FMatrix44f TranslatedWorldToShadowMatrix = FMatrix44f(TranslatedWorldToWorldMatrix * WorldToShadowMatrix);
+
+			PassParameters->NumShadowMatrices = 1;
+			PassParameters->TranslatedWorldToShadow[0] = TranslatedWorldToShadowMatrix;
+			PassParameters->ShadowToTranslatedWorld[0] = TranslatedWorldToShadowMatrix.Inverse();
+			PassParameters->TranslatedWorldOrigin = FVector3f(View.ViewMatrices.GetPreViewTranslation() - ProjectedShadowInfo->PreShadowTranslation);
+		}
+
+		TranslatedWorldOrigin = PassParameters->TranslatedWorldOrigin;
+		NumShadowMatrices = PassParameters->NumShadowMatrices;
+		for (int32 i = 0; i < PassParameters->NumShadowMatrices; ++i)
+		{
+			TranslatedWorldToShadow[i] = PassParameters->TranslatedWorldToShadow[i];
+		}
+
+		FVector LightDirection = LightSceneInfo->Proxy->GetDirection().GetSafeNormal();
+		float W = -FVector3f::DotProduct(TranslatedWorldOrigin, FVector3f(LightDirection));
+		TranslatedWorldPlane = FVector4f(LightDirection.X, LightDirection.Y, LightDirection.Z, W);
+
+		PassParameters->ShadowResolution = ShadowMapResolution;
+		PassParameters->MaxSampleCount = MaxSampleCount;
+		PassParameters->AbsoluteErrorThreshold = HeterogeneousVolumes::GetShadowAbsoluteErrorThreshold();
+		PassParameters->RelativeErrorThreshold = RelativeErrorThreshold;
+
+		// Volume data
+		PassParameters->OrthoGridUniformBuffer = OrthoGridUniformBuffer;
+		PassParameters->FrustumGridUniformBuffer = FrustumGridUniformBuffer;
+
+		// Dispatch data
+		PassParameters->GroupCount = GroupCount;
+
+		// Output
+		PassParameters->RWVolumetricShadowLinkedListAllocatorBuffer = GraphBuilder.CreateUAV(VolumetricShadowLinkedListAllocatorBuffer, PF_R32_UINT);
+		PassParameters->RWVolumetricShadowLinkedListBuffer = GraphBuilder.CreateUAV(VolumetricShadowLinkedListBuffer);
+		PassParameters->RWBeerShadowMapTexture = GraphBuilder.CreateUAV(BeerShadowMapTexture);
+		PassParameters->RWDebugBuffer = GraphBuilder.CreateUAV(DebugBuffer);
+	}
+
+	FRenderVolumetricShadowMapForLightWithVoxelGridCS::FPermutationDomain PermutationVector;
+	PermutationVector.Set<FRenderVolumetricShadowMapForLightWithVoxelGridCS::FUseAVSMCompression>(HeterogeneousVolumes::UseAVSMCompression());
+	TShaderRef<FRenderVolumetricShadowMapForLightWithVoxelGridCS> ComputeShader = View.ShaderMap->GetShader<FRenderVolumetricShadowMapForLightWithVoxelGridCS>(PermutationVector);
+	FComputeShaderUtils::AddPass(
+		GraphBuilder,
+		// Add Light name..
+		RDG_EVENT_NAME("RenderVolumetricShadowMapForLightWithVoxelGridCS"),
+		ERDGPassFlags::Compute | ERDGPassFlags::NeverCull,
+		ComputeShader,
+		PassParameters,
+		GroupCount
+	);
+}
+
+void RenderVolumetricShadowMapForCameraWithVoxelGrid(
+	FRDGBuilder& GraphBuilder,
+	// Scene data
+	const FSceneTextures& SceneTextures,
+	FScene* Scene,
+	const FSceneViewFamily& ViewFamily,
+	FViewInfo& View,
+	// Volume data
+	const TRDGUniformBufferRef<FOrthoVoxelGridUniformBufferParameters>& OrthoGridUniformBuffer,
+	const TRDGUniformBufferRef<FFrustumVoxelGridUniformBufferParameters>& FrustumGridUniformBuffer,
+	// Output
+	FVector3f& TranslatedWorldOrigin,
+	FMatrix44f& TranslatedWorldToShadow,
+	FIntVector& GroupCount,
+	FIntPoint& ShadowMapResolution,
+	uint32& MaxSampleCount,
+	FRDGTextureRef& BeerShadowMapTexture,
+	FRDGBufferRef& VolumetricShadowLinkedListBuffer
+)
+{
+	ShadowMapResolution = FIntPoint(View.ViewRect.Width(), View.ViewRect.Height());
+
+	float DownsampleFactor = HeterogeneousVolumes::GetCameraDownsampleFactor();
+	ShadowMapResolution.X = FMath::Max(ShadowMapResolution.X / DownsampleFactor, 1);
+	ShadowMapResolution.Y = FMath::Max(ShadowMapResolution.Y / DownsampleFactor, 1);
+	//ShadowMapResolution = HeterogeneousVolumes::GetShadowMapResolution();
+
+	GroupCount = FIntVector(1);
+	GroupCount.X = FMath::DivideAndRoundUp(ShadowMapResolution.X, FRenderVolumetricShadowMapForLightWithVoxelGridCS::GetThreadGroupSize2D());
+	GroupCount.Y = FMath::DivideAndRoundUp(ShadowMapResolution.Y, FRenderVolumetricShadowMapForLightWithVoxelGridCS::GetThreadGroupSize2D());
+
+	// TODO: Allocate debug data
+	int32 BufferSize = ShadowMapResolution.X * ShadowMapResolution.Y;
+	FRDGBufferRef DebugBuffer = GraphBuilder.CreateBuffer(
+		FRDGBufferDesc::CreateStructuredDesc(sizeof(FVolumetricShadowMapDebugData), BufferSize),
+		TEXT("HeterogeneousVolume.RenderDebugData")
+	);
+
+	FRDGBufferRef VolumetricShadowLinkedListAllocatorBuffer = GraphBuilder.CreateBuffer(
+		FRDGBufferDesc::CreateBufferDesc(sizeof(uint32), 1),
+		TEXT("HeterogeneousVolume.CameraAVSM.LinkedListAllocatorBuffer")
+	);
+
+	MaxSampleCount = HeterogeneousVolumes::GetShadowMaxSampleCount();
+	int32 VolumetricShadowLinkedListElementCount = ShadowMapResolution.X * ShadowMapResolution.Y * MaxSampleCount;
+	VolumetricShadowLinkedListBuffer = GraphBuilder.CreateBuffer(
+		FRDGBufferDesc::CreateStructuredDesc(sizeof(FAVSMLinkedListPackedData), VolumetricShadowLinkedListElementCount),
+		TEXT("HeterogeneousVolume.CameraAVSM.LinkedListBuffer")
+	);
+
+	const float RelativeErrorThreshold = HeterogeneousVolumes::GetShadowRelativeErrorThreshold();
+
+	// Initialize allocator to contain 1-spp
+	AddClearUAVPass(GraphBuilder, GraphBuilder.CreateUAV(VolumetricShadowLinkedListAllocatorBuffer, PF_R32_UINT), ShadowMapResolution.X * ShadowMapResolution.Y);
+
+	// TODO: Use the frustum grid bounds instead of re-based trace-distance bounds
+	// Intersect TopLevelGridBounds with light ray to get appropriate culling distances.
+
+	FMatrix ViewToClip = FPerspectiveMatrix(
+		FMath::DegreesToRadians(View.FOV * 0.5),
+		ShadowMapResolution.X,
+		ShadowMapResolution.Y,
+		1.0,
+		HeterogeneousVolumes::GetMaxTraceDistance()
+	);
+	FMatrix ClipToView = ViewToClip.Inverse();
+	FMatrix ScreenMatrix = FScaleMatrix(FVector(0.5, -0.5, -0.5)) * FTranslationMatrix(FVector(0.5, 0.5, 0.5));
+	TranslatedWorldToShadow = FMatrix44f(View.ViewMatrices.GetTranslatedViewMatrix() * ViewToClip * ScreenMatrix);
+	FMatrix44f ShadowToTranslatedWorld = TranslatedWorldToShadow.Inverse();
+	TranslatedWorldOrigin = ShadowToTranslatedWorld.GetOrigin();
+
+	FRenderVolumetricShadowMapForLightWithVoxelGridCS::FParameters* PassParameters = GraphBuilder.AllocParameters<FRenderVolumetricShadowMapForLightWithVoxelGridCS::FParameters>();
+	{
+		// Scene data
+		PassParameters->View = View.ViewUniformBuffer;
+		PassParameters->SceneTextures = GetSceneTextureParameters(GraphBuilder, SceneTextures);
+
+		// Ray Data
+		PassParameters->ShadowStepSize = HeterogeneousVolumes::GetShadowStepSize();
+		//PassParameters->ShadowStepFactor = HeterogeneousVolumeInterface->GetShadowStepFactor();
+		PassParameters->ShadowStepFactor = 1.0;
+		PassParameters->MaxTraceDistance = HeterogeneousVolumes::GetMaxTraceDistance();
+		PassParameters->MaxStepCount = HeterogeneousVolumes::GetMaxStepCount();
+		PassParameters->bJitter = HeterogeneousVolumes::ShouldJitter();
+
+		PassParameters->ShadowResolution = ShadowMapResolution;
+		PassParameters->NumShadowMatrices = 1;
+		PassParameters->TranslatedWorldToShadow[0] = TranslatedWorldToShadow;
+		PassParameters->ShadowToTranslatedWorld[0] = ShadowToTranslatedWorld;
+		
+		PassParameters->MaxSampleCount = MaxSampleCount;
+		PassParameters->AbsoluteErrorThreshold = HeterogeneousVolumes::GetShadowAbsoluteErrorThreshold();
+		PassParameters->RelativeErrorThreshold = RelativeErrorThreshold;
+
+		// Volume data
+		PassParameters->OrthoGridUniformBuffer = OrthoGridUniformBuffer;
+		PassParameters->FrustumGridUniformBuffer = FrustumGridUniformBuffer;
+
+		// Dispatch data
+		PassParameters->GroupCount = GroupCount;
+
+		// Output
+		PassParameters->RWVolumetricShadowLinkedListAllocatorBuffer = GraphBuilder.CreateUAV(VolumetricShadowLinkedListAllocatorBuffer, PF_R32_UINT);
+		PassParameters->RWVolumetricShadowLinkedListBuffer = GraphBuilder.CreateUAV(VolumetricShadowLinkedListBuffer);
+		PassParameters->RWBeerShadowMapTexture = GraphBuilder.CreateUAV(BeerShadowMapTexture);
+		PassParameters->RWDebugBuffer = GraphBuilder.CreateUAV(DebugBuffer);
+	}
+
+	bool bUseCustomProjection = false;
+	FRenderVolumetricShadowMapForLightWithVoxelGridCS::FPermutationDomain PermutationVector;
+	PermutationVector.Set<FRenderVolumetricShadowMapForLightWithVoxelGridCS::FUseAVSMCompression>(HeterogeneousVolumes::UseAVSMCompression());
+	TShaderRef<FRenderVolumetricShadowMapForLightWithVoxelGridCS> ComputeShader = View.ShaderMap->GetShader<FRenderVolumetricShadowMapForLightWithVoxelGridCS>(PermutationVector);
+	FComputeShaderUtils::AddPass(
+		GraphBuilder,
+		// Add Light name..
+		RDG_EVENT_NAME("RenderVolumetricShadowMapForCameraWithVoxelGridCS"),
+		ERDGPassFlags::Compute | ERDGPassFlags::NeverCull,
+		ComputeShader,
+		PassParameters,
+		GroupCount
+	);
+}
+
+void RenderAdaptiveVolumetricShadowMapWithVoxelGrid(
+	FRDGBuilder& GraphBuilder,
+	// Scene data
+	const FSceneTextures& SceneTextures,
+	FScene* Scene,
+	const FSceneViewFamily& ViewFamily,
+	FViewInfo& View,
+	// Shadow data
+	TArray<FVisibleLightInfo, SceneRenderingAllocator>& VisibleLightInfos,
+	const FVirtualShadowMapArray& VirtualShadowMapArray,
+	// Volume data
+	const TRDGUniformBufferRef<FOrthoVoxelGridUniformBufferParameters>& OrthoGridUniformBuffer,
+	const TRDGUniformBufferRef<FFrustumVoxelGridUniformBufferParameters>& FrustumGridUniformBuffer
+)
+{
+	RDG_EVENT_SCOPE(GraphBuilder, "Adaptive Volumetric Shadow Maps");
+	bool bShouldRenderShadowMaps = !View.ViewRect.IsEmpty() &&
+		(OrthoGridUniformBuffer->GetParameters()->bUseOrthoGrid || FrustumGridUniformBuffer->GetParameters()->bUseFrustumGrid);
+
+	// Light culling
+	TArray<FLightSceneInfoCompact, TInlineAllocator<64>> LightSceneInfoCompact;
+	for (auto LightIt = Scene->Lights.CreateConstIterator(); LightIt; ++LightIt)
+	{
+		// TODO: Use global bounds information..
+		//if (LightIt->AffectsPrimitive(HeterogeneousVolumeInterface->GetBounds(), HeterogeneousVolumeInterface->GetPrimitiveSceneProxy()))
+		{
+			LightSceneInfoCompact.Add(*LightIt);
+		}
+	}
+
+	// Light loop:
+	int32 NumPasses = LightSceneInfoCompact.Num();
+	for (int32 PassIndex = 0; PassIndex < NumPasses; ++PassIndex)
+	{
+		bool bApplyEmissionAndTransmittance = PassIndex == 0;
+		bool bApplyDirectLighting = !LightSceneInfoCompact.IsEmpty();
+		bool bApplyShadowTransmittance = false;
+
+		uint32 LightType = 0;
+		FLightSceneInfo* LightSceneInfo = nullptr;
+		const FVisibleLightInfo* VisibleLightInfo = nullptr;
+		if (bApplyDirectLighting)
+		{
+			LightType = LightSceneInfoCompact[PassIndex].LightType;
+			LightSceneInfo = LightSceneInfoCompact[PassIndex].LightSceneInfo;
+			check(LightSceneInfo != nullptr);
+
+			bApplyDirectLighting = (LightSceneInfo != nullptr);
+			bool bDynamicallyShadowed = false;
+			if (LightSceneInfo)
+			{
+				VisibleLightInfo = &VisibleLightInfos[LightSceneInfo->Id];
+				bApplyShadowTransmittance = LightSceneInfo->Proxy && LightSceneInfo->Proxy->CastsVolumetricShadow();
+				const FProjectedShadowInfo* ProjectedShadowInfo = GetProjectedShadowInfo(VisibleLightInfo);
+				bDynamicallyShadowed = ProjectedShadowInfo != nullptr;
+			}
+
+			TRDGUniformBufferRef<FAdaptiveVolumetricShadowMapUniformBufferParameters> AdaptiveVolumetricShadowMapUniformBuffer;
+			bool bCreateShadowMap = bShouldRenderShadowMaps && bApplyShadowTransmittance && bDynamicallyShadowed && !ShouldRenderRayTracingShadowsForLight(LightSceneInfoCompact[PassIndex]);
+			if (bCreateShadowMap)
+			{
+				FString LightName;
+				FSceneRenderer::GetLightNameForDrawEvent(LightSceneInfo->Proxy, LightName);
+				RDG_EVENT_SCOPE(GraphBuilder, "%s", *LightName);
+
+				FRDGTextureDesc Desc = SceneTextures.Color.Target->Desc;
+				Desc.Format = PF_FloatRGBA;
+				Desc.Flags &= ~(TexCreate_FastVRAM);
+				FRDGTextureRef BeerShadowMapTexture = GraphBuilder.CreateTexture(Desc, TEXT("BeerShadowMapTexture"));
+				AddClearUAVPass(GraphBuilder, GraphBuilder.CreateUAV(BeerShadowMapTexture), FLinearColor::Transparent);
+
+				bool bIsDirectionalLight = false;
+				FVector3f TranslatedWorldOrigin = FVector3f::Zero();
+				FVector4f TranslatedWorldPlane = FVector4f::Zero();
+				FMatrix44f TranslatedWorldToShadow[] =
+				{
+					FMatrix44f::Identity,
+					FMatrix44f::Identity,
+					FMatrix44f::Identity,
+					FMatrix44f::Identity,
+					FMatrix44f::Identity,
+					FMatrix44f::Identity
+				};
+				FIntVector GroupCount = FIntVector::ZeroValue;
+				int32 NumShadowMatrices = 0;
+				FIntPoint VolumetricShadowMapResolution = FIntPoint::NoneValue;
+				uint32 VolumetricShadowMapMaxSampleCount = 0;
+				FRDGBufferRef VolumetricShadowMapLinkedListBuffer;
+				RenderVolumetricShadowMapForLightWithVoxelGrid(
+					GraphBuilder,
+					// Scene data
+					SceneTextures,
+					Scene,
+					ViewFamily,
+					View,
+					// Light data
+					bApplyEmissionAndTransmittance,
+					bApplyDirectLighting,
+					bApplyShadowTransmittance,
+					LightType,
+					LightSceneInfo,
+					// Shadow data
+					VisibleLightInfo,
+					VirtualShadowMapArray,
+					// Volume data
+					OrthoGridUniformBuffer,
+					FrustumGridUniformBuffer,
+					// Output
+					bIsDirectionalLight,
+					TranslatedWorldOrigin,
+					TranslatedWorldPlane,
+					TranslatedWorldToShadow,
+					GroupCount,
+					NumShadowMatrices,
+					VolumetricShadowMapResolution,
+					VolumetricShadowMapMaxSampleCount,
+					BeerShadowMapTexture,
+					VolumetricShadowMapLinkedListBuffer
+				);
+
+				FRDGBufferRef VolumetricShadowMapIndirectionBuffer;
+				FRDGBufferRef VolumetricShadowMapSampleBuffer;
+				CompressVolumetricShadowMap(
+					GraphBuilder,
+					View,
+					GroupCount,
+					VolumetricShadowMapResolution,
+					VolumetricShadowMapMaxSampleCount,
+					VolumetricShadowMapLinkedListBuffer,
+					VolumetricShadowMapIndirectionBuffer,
+					VolumetricShadowMapSampleBuffer
+				);
+
+				CreateAdaptiveVolumetricShadowMapUniformBuffer(
+					GraphBuilder,
+					TranslatedWorldOrigin,
+					TranslatedWorldPlane,
+					TranslatedWorldToShadow,
+					VolumetricShadowMapResolution,
+					NumShadowMatrices,
+					VolumetricShadowMapMaxSampleCount,
+					bIsDirectionalLight,
+					VolumetricShadowMapLinkedListBuffer,
+					VolumetricShadowMapIndirectionBuffer,
+					VolumetricShadowMapSampleBuffer,
+					AdaptiveVolumetricShadowMapUniformBuffer
+				);
+			}
+			else
+			{
+				AdaptiveVolumetricShadowMapUniformBuffer = HeterogeneousVolumes::CreateEmptyAdaptiveVolumetricShadowMapUniformBuffer(GraphBuilder);
+			}
+
+			if (View.ViewState)
+			{
+				TRDGUniformBufferRef<FAdaptiveVolumetricShadowMapUniformBufferParameters>& AdaptiveVolumetricShadowMap = View.ViewState->AdaptiveVolumetricShadowMapUniformBufferMap.FindOrAdd(LightSceneInfo->Id);
+				AdaptiveVolumetricShadowMap = AdaptiveVolumetricShadowMapUniformBuffer;
+			}
+		}
+	}
+}
+
+void RenderAdaptiveVolumetricCameraMapWithVoxelGrid(
+	FRDGBuilder& GraphBuilder,
+	// Scene data
+	const FSceneTextures& SceneTextures,
+	FScene* Scene,
+	const FSceneViewFamily& ViewFamily,
+	FViewInfo& View,
+	// Volume data
+	const TRDGUniformBufferRef<FOrthoVoxelGridUniformBufferParameters>& OrthoGridUniformBuffer,
+	const TRDGUniformBufferRef<FFrustumVoxelGridUniformBufferParameters>& FrustumGridUniformBuffer
+)
+{
+	if (View.ViewState == nullptr)
+	{
+		return;
+	}
+	RDG_EVENT_SCOPE(GraphBuilder, "Adaptive Volumetric Camera Map");
+
+	FVector3f TranslatedWorldOrigin = FVector3f::ZeroVector;
+	int32 NumShadowMatrices = 1;
+	FMatrix44f TranslatedWorldToShadow[] =
+	{
+		FMatrix44f::Identity
+	};
+	FIntPoint VolumetricShadowMapResolution = FIntPoint::NoneValue;
+	uint32 VolumetricShadowMapMaxSampleCount = 0;
+	FRDGBufferRef VolumetricShadowMapLinkedListBuffer = GSystemTextures.GetDefaultStructuredBuffer(GraphBuilder, sizeof(FAVSMLinkedListPackedData));;
+	FRDGBufferRef VolumetricShadowMapIndirectionBuffer = GSystemTextures.GetDefaultStructuredBuffer(GraphBuilder, sizeof(FAVSMIndirectionPackedData));
+	FRDGBufferRef VolumetricShadowMapSampleBuffer = GSystemTextures.GetDefaultStructuredBuffer(GraphBuilder, sizeof(FAVSMSamplePackedData));
+
+	bool bShouldRenderCameraMap = !View.ViewRect.IsEmpty()
+		&& (OrthoGridUniformBuffer->GetParameters()->bUseOrthoGrid || FrustumGridUniformBuffer->GetParameters()->bUseFrustumGrid);
+
+	TRDGUniformBufferRef<FAdaptiveVolumetricShadowMapUniformBufferParameters> AdaptiveVolumetricShadowMapUniformBuffer;
+	if (bShouldRenderCameraMap)
+	{
+		FRDGTextureDesc Desc = SceneTextures.Color.Target->Desc;
+		Desc.Format = PF_FloatRGBA;
+		Desc.Flags &= ~(TexCreate_FastVRAM);
+		FRDGTextureRef CameraShadowMapTexture = GraphBuilder.CreateTexture(Desc, TEXT("CameraShadowMapTexture"));
+		AddClearUAVPass(GraphBuilder, GraphBuilder.CreateUAV(CameraShadowMapTexture), FLinearColor::Transparent);
+
+		FIntVector GroupCount = FIntVector::ZeroValue;
+		RenderVolumetricShadowMapForCameraWithVoxelGrid(
+			GraphBuilder,
+			// Scene data
+			SceneTextures,
+			Scene,
+			ViewFamily,
+			View,
+			// Volume data
+			OrthoGridUniformBuffer,
+			FrustumGridUniformBuffer,
+			// Output
+			TranslatedWorldOrigin,
+			TranslatedWorldToShadow[0],
+			GroupCount,
+			VolumetricShadowMapResolution,
+			VolumetricShadowMapMaxSampleCount,
+			CameraShadowMapTexture,
+			VolumetricShadowMapLinkedListBuffer
+		);
+
+		CompressVolumetricShadowMap(
+			GraphBuilder,
+			View,
+			GroupCount,
+			VolumetricShadowMapResolution,
+			VolumetricShadowMapMaxSampleCount,
+			VolumetricShadowMapLinkedListBuffer,
+			VolumetricShadowMapIndirectionBuffer,
+			VolumetricShadowMapSampleBuffer
+		);
+
+		FVector4f TranslatedWorldPlane = FVector4f::Zero();
+		bool bIsDirectionalLight = false;
+		CreateAdaptiveVolumetricShadowMapUniformBuffer(
+			GraphBuilder,
+			TranslatedWorldOrigin,
+			TranslatedWorldPlane,
+			TranslatedWorldToShadow,
+			VolumetricShadowMapResolution,
+			NumShadowMatrices,
+			VolumetricShadowMapMaxSampleCount,
+			bIsDirectionalLight,
+			VolumetricShadowMapLinkedListBuffer,
+			VolumetricShadowMapIndirectionBuffer,
+			VolumetricShadowMapSampleBuffer,
+			AdaptiveVolumetricShadowMapUniformBuffer
+		);
+	}
+	else
+	{
+		AdaptiveVolumetricShadowMapUniformBuffer = HeterogeneousVolumes::CreateEmptyAdaptiveVolumetricShadowMapUniformBuffer(GraphBuilder);
+	}
+
+	View.ViewState->AdaptiveVolumetricCameraMapUniformBuffer = AdaptiveVolumetricShadowMapUniformBuffer;
+}
diff --git a/Engine/Source/Runtime/Renderer/Private/LightRendering.cpp b/Engine/Source/Runtime/Renderer/Private/LightRendering.cpp
--- a/Engine/Source/Runtime/Renderer/Private/LightRendering.cpp
+++ b/Engine/Source/Runtime/Renderer/Private/LightRendering.cpp
@@ -32,6 +32,7 @@
 #include "TranslucentLighting.h"
 #include "ManyLights/ManyLights.h"
 #include "LightFunctionAtlas.h"
+#include "HeterogeneousVolumes/HeterogeneousVolumes.h"
 
 using namespace LightFunctionAtlas;
 
@@ -866,6 +867,10 @@ class FDeferredLightPS : public FGlobalShader
 		SHADER_PARAMETER_RDG_UNIFORM_BUFFER(FVirtualShadowMapUniformParameters, VirtualShadowMap)
 		SHADER_PARAMETER(int32, VirtualShadowMapId)
 		SHADER_PARAMETER_RDG_TEXTURE(Texture2D, ShadowMaskBits)
+		// Heterogeneous Volume data
+		//SHADER_PARAMETER_RDG_UNIFORM_BUFFER(FOrthoVoxelGridUniformBufferParameters, OrthoGridUniformBuffer)
+		//SHADER_PARAMETER_RDG_UNIFORM_BUFFER(FFrustumVoxelGridUniformBufferParameters, FrustumGridUniformBuffer)
+		SHADER_PARAMETER_RDG_UNIFORM_BUFFER(FAdaptiveVolumetricShadowMapUniformBufferParameters, AVSM)
 		RENDER_TARGET_BINDING_SLOTS()
 	END_SHADER_PARAMETER_STRUCT()
 
@@ -993,6 +998,8 @@ class FDeferredLightPS : public FGlobalShader
 			FVirtualShadowMapArray::SetShaderDefines(OutEnvironment);
 			FForwardLightingParameters::ModifyCompilationEnvironment(Parameters.Platform, OutEnvironment);
 		}
+
+		OutEnvironment.SetDefine(TEXT("USE_ADAPTIVE_VOLUMETRIC_SHADOW_MAP"), ShouldHeterogeneousVolumesCastShadows());
 	}
 
 	void SetParameters(FRHIBatchedShaderParameters& BatchedParameters, const FDeferredLightUniformStruct& DeferredLightUniformsValue)
@@ -2606,6 +2613,11 @@ static void RenderLight(
 			PassParameters->SubstrateTileComplex = Substrate::SetTileParameters(GraphBuilder, View, ESubstrateTileType::EComplex);
 			PassParameters->SubstrateTileSpectialComplex = Substrate::SetTileParameters(GraphBuilder, View, ESubstrateTileType::EComplexSpecial);
 		}
+#if 0
+		PassParameters->PS.OrthoGridUniformBuffer = HeterogeneousVolumes::GetOrthoVoxelGridUniformBuffer(GraphBuilder, View.ViewState);
+		PassParameters->PS.FrustumGridUniformBuffer = HeterogeneousVolumes::GetFrustumVoxelGridUniformBuffer(GraphBuilder, View.ViewState);
+#endif
+		PassParameters->PS.AVSM = HeterogeneousVolumes::GetAdaptiveVolumetricShadowMapUniformBuffer(GraphBuilder, View.ViewState, LightSceneInfo);
 
 		FDeferredLightPS::FPermutationDomain PermutationVector;
 		PermutationVector.Set< FDeferredLightPS::FTransmissionDim >(bTransmission);
@@ -2618,6 +2630,7 @@ static void RenderLight(
 		PermutationVector.Set< FDeferredLightPS::FLightFunctionAtlasDim >(
 			LightFunctionAtlas::IsEnabled(View, ELightFunctionAtlasSystem::DeferredLighting) && LightSceneInfo->Proxy->HasValidLightFunctionAtlasSlot() &&
 			LightSceneInfo->Proxy->GetLightFunctionMaterial() != nullptr && !View.Family->EngineShowFlags.VisualizeLightCulling);
+
 		if (bIsRadial)
 		{
 			PermutationVector.Set< FDeferredLightPS::FSourceShapeDim >(LightProxy->IsRectLight() ? ELightSourceShape::Rect : ELightSourceShape::Capsule);
diff --git a/Engine/Source/Runtime/Renderer/Private/PathTracing.cpp b/Engine/Source/Runtime/Renderer/Private/PathTracing.cpp
--- a/Engine/Source/Runtime/Renderer/Private/PathTracing.cpp
+++ b/Engine/Source/Runtime/Renderer/Private/PathTracing.cpp
@@ -2754,8 +2754,9 @@ void FDeferredShadingSceneRenderer::RenderPathTracing(
 				!PathTracingState->AdaptiveOrthoGridParameterCache.TopLevelGridBuffer;
 			if (bCreateVolumeGrids)
 			{
-				BuildOrthoVoxelGrid(GraphBuilder, Scene, Views, OrthoGridUniformBuffer);
-				BuildFrustumVoxelGrid(GraphBuilder, Scene, Views[0], FrustumGridUniformBuffer);
+				FVoxelGridBuildOptions BuildOptions;
+				BuildOrthoVoxelGrid(GraphBuilder, Scene, Views, VisibleLightInfos, BuildOptions, OrthoGridUniformBuffer);
+				BuildFrustumVoxelGrid(GraphBuilder, Scene, Views[0], BuildOptions, FrustumGridUniformBuffer);
 			}
 			else
 			{
diff --git a/Engine/Source/Runtime/Renderer/Private/ScenePrivate.h b/Engine/Source/Runtime/Renderer/Private/ScenePrivate.h
--- a/Engine/Source/Runtime/Renderer/Private/ScenePrivate.h
+++ b/Engine/Source/Runtime/Renderer/Private/ScenePrivate.h
@@ -57,6 +57,7 @@
 #include "UObject/Package.h"
 #include "LightFunctionAtlas.h"
 #include "SceneExtensions.h"
+#include "HeterogeneousVolumes/HeterogeneousVolumes.h"
 
 /** Factor by which to grow occlusion tests **/
 #define OCCLUSION_SLOP (1.0f)
@@ -1022,6 +1023,13 @@ public:
 	FLumenViewState Lumen;
 	FManyLightsViewState ManyLights;
 
+	// Heterogeneous Volumes cached data stores
+	TRDGUniformBufferRef<FOrthoVoxelGridUniformBufferParameters> OrthoVoxelGridUniformBuffer = nullptr;
+	TRDGUniformBufferRef<FFrustumVoxelGridUniformBufferParameters> FrustumVoxelGridUniformBuffer = nullptr;
+
+	TRDGUniformBufferRef<FAdaptiveVolumetricShadowMapUniformBufferParameters> AdaptiveVolumetricCameraMapUniformBuffer = nullptr;
+	TMap<int32, TRDGUniformBufferRef<FAdaptiveVolumetricShadowMapUniformBufferParameters>> AdaptiveVolumetricShadowMapUniformBufferMap;
+
 	// Map from Light ID in GPU Scene to index in the View's ForwardLightData array
 	// This is stored in ViewState so we can access previous frame mapping
 	TMap<int32, int32> LightSceneIdToLocalLightIndex;
diff --git a/Engine/Source/Runtime/Renderer/Private/SceneRendering.h b/Engine/Source/Runtime/Renderer/Private/SceneRendering.h
--- a/Engine/Source/Runtime/Renderer/Private/SceneRendering.h
+++ b/Engine/Source/Runtime/Renderer/Private/SceneRendering.h
@@ -838,6 +838,16 @@ struct FVolumetricMeshBatch
 {
 	const FMeshBatch* Mesh;
 	const FPrimitiveSceneProxy* Proxy;
+
+	FVolumetricMeshBatch(const FMeshBatch* M, const FPrimitiveSceneProxy* P)
+		: Mesh(M)
+		, Proxy(P)
+	{}
+
+	FORCEINLINE bool operator==(const FVolumetricMeshBatch& rhs) const
+	{
+		return (Mesh->MeshIdInPrimitive == rhs.Mesh->MeshIdInPrimitive) && (Proxy == rhs.Proxy);
+	}
 };
 
 struct FSkyMeshBatch
diff --git a/Engine/Source/Runtime/Renderer/Private/ShadowRendering.h b/Engine/Source/Runtime/Renderer/Private/ShadowRendering.h
--- a/Engine/Source/Runtime/Renderer/Private/ShadowRendering.h
+++ b/Engine/Source/Runtime/Renderer/Private/ShadowRendering.h
@@ -804,6 +804,8 @@ public:
 	/** Check if we need to set the scissor rect to exclude portion of the CSM slices outside the view frustum */
 	bool ShouldUseCSMScissorOptim() const;
 
+	const TArray<FMeshBatchAndRelevance, SceneRenderingAllocator>& GetDynamicSubjectHeterogeneousVolumeMeshElements() const { return DynamicSubjectHeterogeneousVolumeMeshElements; }
+
 private:
 	// 0 if Setup...() wasn't called yet
 	FLightSceneInfo* LightSceneInfo;
@@ -821,11 +823,15 @@ private:
 	PrimitiveArrayType ReceiverPrimitives;
 	/** Subject primitives with translucent relevance. */
 	PrimitiveArrayType SubjectTranslucentPrimitives;
+	/** Subject primitives for heterogeneous volume shadows. */
+	PrimitiveArrayType SubjectHeterogeneousVolumePrimitives;
 
 	/** Dynamic mesh elements for subject primitives. */
 	TArray<FMeshBatchAndRelevance,SceneRenderingAllocator> DynamicSubjectMeshElements;
 	/** Dynamic mesh elements for translucent subject primitives. */
 	TArray<FMeshBatchAndRelevance,SceneRenderingAllocator> DynamicSubjectTranslucentMeshElements;
+	/** Dynamic mesh elements for heterogeneous volume primitives. */
+	TArray<FMeshBatchAndRelevance, SceneRenderingAllocator> DynamicSubjectHeterogeneousVolumeMeshElements;
 
 	TArray<const FStaticMeshBatch*, SceneRenderingAllocator> SubjectMeshCommandBuildRequests;
 	TArray<EMeshDrawCommandCullingPayloadFlags, SceneRenderingAllocator> SubjectMeshCommandBuildFlags;
@@ -953,6 +959,16 @@ private:
 		int32& OutNumDynamicSubjectMeshElements,
 		EGatherDynamicMeshElementsPass Pass);
 
+	bool GatherDynamicHeterogeneousVolumeMeshElementsArray(
+		FSceneRenderer& Renderer,
+		FMeshElementCollector& Collector,
+		const PrimitiveArrayType& PrimitiveArray,
+		const TArray<const FSceneView*>& Views,
+		const FSceneViewFamily& ViewFamily,
+		TArray<FMeshBatchAndRelevance, SceneRenderingAllocator>& OutDynamicMeshElements,
+		int32& OutNumDynamicSubjectMeshElements,
+		EGatherDynamicMeshElementsPass Pass);
+
 	void SetupFrustumForProjection(const FViewInfo* View, TArray<FVector4f, TInlineAllocator<8>>& OutFrustumVertices, bool& bOutCameraInsideShadowFrustum, FPlane* OutPlanes) const;
 
 	void SetupProjectionStencilMask(
diff --git a/Engine/Source/Runtime/Renderer/Private/ShadowSetup.cpp b/Engine/Source/Runtime/Renderer/Private/ShadowSetup.cpp
--- a/Engine/Source/Runtime/Renderer/Private/ShadowSetup.cpp
+++ b/Engine/Source/Runtime/Renderer/Private/ShadowSetup.cpp
@@ -1447,6 +1447,7 @@ struct FAddSubjectPrimitiveResult
 			uint32 bOverflowed : 1;
 			uint32 bDynamicSubjectPrimitive : 1;
 			uint32 bTranslucentSubjectPrimitive : 1;
+			uint32 bHeterogeneousVolumeSubjectPrimitive : 1;
 			uint32 bNeedUniformBufferUpdate : 1;
 			uint32 bNeedPrimitiveFadingStateUpdate : 1;
 			uint32 bFadingIn : 1;
@@ -1609,6 +1610,7 @@ struct FAddSubjectPrimitiveStats
 	int32 NumMDCBuildRequests;
 	int32 NumDynamicSubs;
 	int32 NumTranslucentSubs;
+	int32 NumHeterogeneousVolumeSubs;
 	int32 NumDeferredPrimitives;
 
 	FAddSubjectPrimitiveStats()
@@ -1616,6 +1618,7 @@ struct FAddSubjectPrimitiveStats
 		, NumMDCBuildRequests(0)
 		, NumDynamicSubs(0)
 		, NumTranslucentSubs(0)
+		, NumHeterogeneousVolumeSubs(0)
 		, NumDeferredPrimitives(0)
 	{}
 
@@ -1637,6 +1640,10 @@ struct FAddSubjectPrimitiveStats
 		{
 			FPlatformAtomics::InterlockedAdd(&NumTranslucentSubs, Other.NumTranslucentSubs);
 		}
+		if (Other.NumHeterogeneousVolumeSubs > 0)
+		{
+			FPlatformAtomics::InterlockedAdd(&NumHeterogeneousVolumeSubs, Other.NumHeterogeneousVolumeSubs);
+		}
 		if (Other.NumDeferredPrimitives > 0)
 		{
 			FPlatformAtomics::InterlockedAdd(&NumDeferredPrimitives, Other.NumDeferredPrimitives);
@@ -1923,6 +1930,7 @@ bool FProjectedShadowInfo::AddSubjectPrimitive(FDynamicShadowsTaskData& TaskData
 		bool bTranslucentRelevance = false;
 		bool bShadowRelevance = false;
 		bool bDynamicRelevance = false;
+		bool bHeterogeneousVolumeRelevance = false;
 
 		int32 PrimitiveId = PrimitiveSceneInfo->GetIndex();
 
@@ -1973,6 +1981,7 @@ bool FProjectedShadowInfo::AddSubjectPrimitive(FDynamicShadowsTaskData& TaskData
 			bTranslucentRelevance |= ViewRelevance.HasTranslucency();
 			bShadowRelevance |= ViewRelevance.bShadowRelevance;
 			bDynamicRelevance = bDynamicRelevance || ViewRelevance.bDynamicRelevance;
+			bHeterogeneousVolumeRelevance = ViewRelevance.bHasVolumeMaterialDomain && PrimitiveSceneInfo->Proxy->IsHeterogeneousVolume();
 		}
 
 		if (bShadowRelevance)
@@ -2041,6 +2050,12 @@ bool FProjectedShadowInfo::AddSubjectPrimitive(FDynamicShadowsTaskData& TaskData
 			SubjectTranslucentPrimitives.Add(PrimitiveSceneInfo);
 			bWasAdded = true;
 		}
+
+		if (bHeterogeneousVolumeRelevance)
+		{
+			SubjectHeterogeneousVolumePrimitives.Add(PrimitiveSceneInfo);
+			bWasAdded = true;
+		}
 	}
 	return bWasAdded;
 }
@@ -2134,6 +2149,7 @@ uint64 FProjectedShadowInfo::AddSubjectPrimitive_AnyThread(
 		bool bStaticRelevance = false;
 		bool bMayBeFading = false;
 		bool bDynamicRelevance = false;
+		bool bHeterogeneousVolumeRelevance = false;
 
 		const int32 PrimitiveId = PrimitiveSceneInfo->GetIndex();
 		FPrimitiveViewRelevance& ViewRelevance = CurrentView->PrimitiveViewRelevanceMap[PrimitiveId];
@@ -2189,6 +2205,7 @@ uint64 FProjectedShadowInfo::AddSubjectPrimitive_AnyThread(
 		bShadowRelevance = ViewRelevance.bShadowRelevance;
 		bStaticRelevance = ViewRelevance.bStaticRelevance;
 		bDynamicRelevance = ViewRelevance.bDynamicRelevance;
+		bHeterogeneousVolumeRelevance = ViewRelevance.bHasVolumeMaterialDomain && PrimitiveSceneInfo->Proxy->IsHeterogeneousVolume();
 
 		if (!bShadowRelevance)
 		{
@@ -2245,6 +2262,12 @@ uint64 FProjectedShadowInfo::AddSubjectPrimitive_AnyThread(
 			Result.bTranslucentSubjectPrimitive = true;
 			++OutStats.NumTranslucentSubs;
 		}
+
+		if (bHeterogeneousVolumeRelevance)
+		{
+			Result.bHeterogeneousVolumeSubjectPrimitive = true;
+			++OutStats.NumHeterogeneousVolumeSubs;
+		}
 	}
 
 	return Result.Qword;
@@ -2257,6 +2280,7 @@ void FProjectedShadowInfo::PresizeSubjectPrimitiveArrays(const FAddSubjectPrimit
 	SubjectMeshCommandBuildFlags.Reserve(SubjectMeshCommandBuildFlags.Num() + Stats.NumMDCBuildRequests);
 	DynamicSubjectPrimitives.Reserve(DynamicSubjectPrimitives.Num() + Stats.NumDeferredPrimitives + Stats.NumDynamicSubs);
 	SubjectTranslucentPrimitives.Reserve(SubjectTranslucentPrimitives.Num() + Stats.NumTranslucentSubs);
+	SubjectHeterogeneousVolumePrimitives.Reserve(SubjectHeterogeneousVolumePrimitives.Num() + Stats.NumHeterogeneousVolumeSubs);
 }
 
 void FProjectedShadowInfo::FinalizeAddSubjectPrimitive(
@@ -2376,12 +2400,18 @@ void FProjectedShadowInfo::FinalizeAddSubjectPrimitive(
 	{
 		SubjectTranslucentPrimitives.Add(PrimitiveSceneInfo);
 	}
+
+	if (Result.bHeterogeneousVolumeSubjectPrimitive)
+	{
+		SubjectHeterogeneousVolumePrimitives.Add(PrimitiveSceneInfo);
+	}
 }
 
 bool FProjectedShadowInfo::HasSubjectPrims() const
 {
 	return DynamicSubjectPrimitives.Num() > 0
 		|| SubjectTranslucentPrimitives.Num() > 0
+		|| SubjectHeterogeneousVolumePrimitives.Num() > 0
 		|| ShadowDepthPass.HasAnyDraw()
 		|| SubjectMeshCommandBuildRequests.Num() > 0
 		|| ShadowDepthPassVisibleCommands.Num() > 0
@@ -2634,7 +2664,7 @@ bool FProjectedShadowInfo::GatherDynamicMeshElements(
 
 	bool bProcessedAllPrimitives = true;
 
-	if (DynamicSubjectPrimitives.Num() > 0 || ReceiverPrimitives.Num() > 0 || SubjectTranslucentPrimitives.Num() > 0)
+	if (DynamicSubjectPrimitives.Num() > 0 || ReceiverPrimitives.Num() > 0 || SubjectTranslucentPrimitives.Num() > 0 || SubjectHeterogeneousVolumePrimitives.Num() > 0)
 	{
 		ReusedViewsArray[0] = ShadowDepthView;
 
@@ -2667,6 +2697,11 @@ bool FProjectedShadowInfo::GatherDynamicMeshElements(
 		int32 NumDynamicSubjectTranslucentMeshElements = 0;
 		ShadowDepthView->SetDynamicMeshElementsShadowCullFrustum(&CasterOuterFrustum);
 		bProcessedAllPrimitives &= GatherDynamicMeshElementsArray(MeshCollector, SubjectTranslucentPrimitives, ReusedViewsArray, Renderer.ViewFamily, DynamicSubjectTranslucentMeshElements, NumDynamicSubjectTranslucentMeshElements, Pass);
+
+		int32 NumDynamicSubjectHeterogeneousVolumeMeshElements = 0;
+		bProcessedAllPrimitives &= GatherDynamicHeterogeneousVolumeMeshElementsArray(Renderer, MeshCollector, SubjectHeterogeneousVolumePrimitives, ReusedViewsArray, Renderer.ViewFamily, DynamicSubjectHeterogeneousVolumeMeshElements, NumDynamicSubjectHeterogeneousVolumeMeshElements, Pass);
+
+		MeshCollector.ClearViewMeshArrays();
 	}
 
 	return bProcessedAllPrimitives;
@@ -2731,7 +2766,48 @@ bool FProjectedShadowInfo::GatherDynamicMeshElementsArray(
 	}
 
 	OutNumDynamicSubjectMeshElements += MeshCollector.GetMeshElementCount(0);
-	MeshCollector.ClearViewMeshArrays();
+	return bProcessedAllPrimitives;
+}
+
+bool FProjectedShadowInfo::GatherDynamicHeterogeneousVolumeMeshElementsArray(
+	FSceneRenderer& Renderer,
+	FMeshElementCollector& MeshCollector,
+	const PrimitiveArrayType& Primitives,
+	const TArray<const FSceneView*>& Views,
+	const FSceneViewFamily& ViewFamily,
+	TArray<FMeshBatchAndRelevance, SceneRenderingAllocator>& OutDynamicMeshElements,
+	int32& OutNumDynamicSubjectMeshElements,
+	EGatherDynamicMeshElementsPass Pass)
+{
+	int32 Zero = 0;
+	TArray<int32> MeshElementsPerView(&Zero, Views.Num());
+	for (int32 ViewIndex = 0; ViewIndex < Views.Num(); ++ViewIndex)
+	{
+		MeshElementsPerView[ViewIndex] = MeshCollector.GetMeshElementCount(ViewIndex);
+	}
+
+	bool bProcessedAllPrimitives = GatherDynamicMeshElementsArray(
+		MeshCollector,
+		Primitives,
+		Views,
+		ViewFamily,
+		OutDynamicMeshElements,
+		OutNumDynamicSubjectMeshElements,
+		Pass
+	);
+
+	for (int32 ViewIndex = 0; ViewIndex < Views.Num(); ++ViewIndex)
+	{
+		uint32 MeshElementCount = MeshCollector.GetMeshElementCount(ViewIndex);
+		for (uint32 Index = MeshElementsPerView[ViewIndex]; Index < MeshElementCount; ++Index)
+		{
+			FMeshBatchAndRelevance MeshBatchAndRelevance = (*MeshCollector.MeshBatches[ViewIndex])[Index];
+			OutDynamicMeshElements.Add(MeshBatchAndRelevance);
+		}
+
+		OutNumDynamicSubjectMeshElements += MeshElementCount - MeshElementsPerView[ViewIndex];
+	}
+
 	return bProcessedAllPrimitives;
 }
 
@@ -2763,10 +2839,12 @@ void FProjectedShadowInfo::ClearTransientArrays()
 	NumSubjectMeshCommandBuildRequestElements = 0;
 
 	SubjectTranslucentPrimitives.Empty();
+	SubjectHeterogeneousVolumePrimitives.Empty();
 	DynamicSubjectPrimitives.Empty();
 	ReceiverPrimitives.Empty();
 	DynamicSubjectMeshElements.Empty();
 	DynamicSubjectTranslucentMeshElements.Empty();
+	DynamicSubjectHeterogeneousVolumeMeshElements.Empty();
 
 	ShadowDepthPassVisibleCommands.Empty();
 	ShadowDepthPass.WaitForTasksAndEmpty();
diff --git a/Engine/Source/Runtime/Renderer/Private/TranslucentLighting.cpp b/Engine/Source/Runtime/Renderer/Private/TranslucentLighting.cpp
--- a/Engine/Source/Runtime/Renderer/Private/TranslucentLighting.cpp
+++ b/Engine/Source/Runtime/Renderer/Private/TranslucentLighting.cpp
@@ -51,6 +51,7 @@
 #include "RenderCore.h"
 #include "StaticMeshBatch.h"
 #include "LightFunctionAtlas.h"
+#include "HeterogeneousVolumes/HeterogeneousVolumes.h"
 
 class FMaterial;
 
@@ -671,6 +672,7 @@ class FTranslucentLightingInjectPS : public FMaterialShader
 		SHADER_PARAMETER_STRUCT_INCLUDE(FVolumeShadowingShaderParameters, VolumeShadowingParameters)
 		SHADER_PARAMETER_STRUCT_INCLUDE(FVirtualShadowMapSamplingParameters, VirtualShadowMapSamplingParameters)
 		SHADER_PARAMETER_STRUCT_INCLUDE(FLightCloudTransmittanceParameters, LightCloudTransmittanceParameters)
+		SHADER_PARAMETER_RDG_UNIFORM_BUFFER(FAdaptiveVolumetricShadowMapUniformBufferParameters, AVSM)
 		SHADER_PARAMETER(FMatrix44f, LightFunctionTranslatedWorldToLight)
 		SHADER_PARAMETER(FVector4f, LightFunctionParameters)
 		SHADER_PARAMETER(float, SpotlightMask)
@@ -684,12 +686,14 @@ class FTranslucentLightingInjectPS : public FMaterialShader
 	class FDynamicallyShadowed	: SHADER_PERMUTATION_BOOL("DYNAMICALLY_SHADOWED");
 	class FLightFunction		: SHADER_PERMUTATION_BOOL("APPLY_LIGHT_FUNCTION");
 	class FVirtualShadowMap		: SHADER_PERMUTATION_BOOL("VIRTUAL_SHADOW_MAP");
+	class FAdaptiveVolumetricShadowMap : SHADER_PERMUTATION_BOOL("ADAPTIVE_VOLUMETRIC_SHADOW_MAP");
 
 	using FPermutationDomain = TShaderPermutationDomain<
 		FRadialAttenuation,
 		FDynamicallyShadowed,
 		FLightFunction,
-		FVirtualShadowMap >;
+		FVirtualShadowMap,
+		FAdaptiveVolumetricShadowMap >;
 
 public:
 
@@ -714,6 +718,11 @@ public:
 			return false;
 		}
 
+		if (!DoesPlatformSupportHeterogeneousVolumes(Parameters.Platform) && PermutationVector.Get<FAdaptiveVolumetricShadowMap>() != 0)
+		{
+			return false;
+		}
+
 		return (Parameters.MaterialParameters.MaterialDomain == MD_LightFunction || Parameters.MaterialParameters.bIsSpecialEngineMaterial) &&
 			(IsFeatureLevelSupported(Parameters.Platform, ERHIFeatureLevel::SM5) &&
 			(RHISupportsGeometryShaders(Parameters.Platform) || RHISupportsVertexShaderLayer(Parameters.Platform)));
@@ -1118,6 +1127,7 @@ void InjectTranslucencyLightingVolume(
 			const bool bInverseSquared = LightSceneInfo->Proxy->IsInverseSquared();
 			const bool bDirectionalLight = LightSceneInfo->Proxy->GetLightType() == LightType_Directional;
 			bool bUseVSM = Renderer.VirtualShadowMapArray.IsAllocated();
+			const bool bUseAdaptiveVolumetricShadowMap = LightSceneInfo->Proxy->CastsVolumetricShadow() && ShouldRenderHeterogeneousVolumes(Scene) && ShouldHeterogeneousVolumesCastShadows();
 
 			const FVolumeBounds VolumeBounds = CalculateLightVolumeBounds(LightSceneInfo->Proxy->GetBoundingSphere(), View, VolumeCascadeIndex, bDirectionalLight);
 			if (VolumeBounds.IsValid())
@@ -1156,6 +1166,7 @@ void InjectTranslucencyLightingVolume(
 				PassParameters->PS.VirtualShadowMapId = VirtualShadowMapId;
 				PassParameters->PS.LightFunctionParameters = FLightFunctionSharedParameters::GetLightFunctionSharedParameters(LightSceneInfo, 1.0f);
 				PassParameters->PS.VolumeCascadeIndex = VolumeCascadeIndex;
+				PassParameters->PS.AVSM = HeterogeneousVolumes::GetAdaptiveVolumetricShadowMapUniformBuffer(GraphBuilder, View.ViewState, LightSceneInfo);
 
 				bool bIsSpotlight = LightSceneInfo->Proxy->GetLightType() == LightType_Spot;
 				PassParameters->PS.SpotlightMask = bIsSpotlight ? 1.0f : 0.0f; //@todo - needs to be a permutation to reduce shadow filtering work
@@ -1183,7 +1194,7 @@ void InjectTranslucencyLightingVolume(
 						InjectionData.bApplyLightFunction ? TEXT(",LightFunction") : TEXT("")),
 					PassParameters,
 					ERDGPassFlags::Raster,
-					[PassParameters, VertexShader, GeometryShader, &View, &Renderer, &InjectionData, LightSceneInfo, bDirectionalLight, bUseVSM, VolumeBounds, VolumeCascadeIndex](FRHICommandList& RHICmdList)
+					[PassParameters, VertexShader, GeometryShader, &View, &Renderer, &InjectionData, LightSceneInfo, bDirectionalLight, bUseVSM, VolumeBounds, VolumeCascadeIndex, bUseAdaptiveVolumetricShadowMap](FRHICommandList& RHICmdList)
 				{
 					FGraphicsPipelineStateInitializer GraphicsPSOInit;
 					RHICmdList.ApplyCachedRenderTargets(GraphicsPSOInit);
@@ -1217,6 +1228,7 @@ void InjectTranslucencyLightingVolume(
 					PermutationVector.Set< FTranslucentLightingInjectPS::FDynamicallyShadowed >( InjectionData.ProjectedShadowInfo != nullptr );
 					PermutationVector.Set< FTranslucentLightingInjectPS::FLightFunction >( InjectionData.bApplyLightFunction );
 					PermutationVector.Set< FTranslucentLightingInjectPS::FVirtualShadowMap >( bUseVSM );
+					PermutationVector.Set< FTranslucentLightingInjectPS::FAdaptiveVolumetricShadowMap >(bUseAdaptiveVolumetricShadowMap);
 
 					auto PixelShader = MaterialShaderMap->GetShader< FTranslucentLightingInjectPS >( PermutationVector );
 	
diff --git a/Engine/Source/Runtime/Renderer/Private/TranslucentRendering.cpp b/Engine/Source/Runtime/Renderer/Private/TranslucentRendering.cpp
--- a/Engine/Source/Runtime/Renderer/Private/TranslucentRendering.cpp
+++ b/Engine/Source/Runtime/Renderer/Private/TranslucentRendering.cpp
@@ -1006,6 +1006,8 @@ TRDGUniformBufferRef<FTranslucentBasePassUniformParameters> CreateTranslucentBas
 		BasePassParameters.BlueNoise = GetBlueNoiseDummyParameters();
 	}
 
+	BasePassParameters.AVSM = HeterogeneousVolumes::GetAdaptiveVolumetricCameraMapParameters(GraphBuilder, View.ViewState);
+
 	return GraphBuilder.CreateUniformBuffer(&BasePassParameters);
 }
 
