PostProcessing: Removed duplicate definitions of FPostProcessMaterialInputs.

diff --git a/Engine/Source/Runtime/Renderer/Private/PostProcess/PostProcessing.h b/Engine/Source/Runtime/Renderer/Private/PostProcess/PostProcessing.h
--- a/Engine/Source/Runtime/Renderer/Private/PostProcess/PostProcessing.h
+++ b/Engine/Source/Runtime/Renderer/Private/PostProcess/PostProcessing.h
@@ -6,11 +6,10 @@
 #include "RHIFwd.h"
 #include "TranslucentRendering.h"
 #include "PathTracing.h"
+#include "PostProcess/PostProcessInputs.h"
 
 enum class EReflectionsMethod;
 
-class FSceneTextureParameters;
-class FSceneTextureUniformParameters;
 class FScreenPassVS;
 class FViewInfo;
 class FVirtualShadowMapArray;
@@ -31,25 +30,6 @@ bool IsPostProcessingWithAlphaChannelSupported();
 
 using FPostProcessVS = FScreenPassVS;
 
-struct FPostProcessingInputs
-{
-	TRDGUniformBufferRef<FSceneTextureUniformParameters> SceneTextures = nullptr;
-	FRDGTextureRef ViewFamilyTexture = nullptr;
-	FRDGTextureRef CustomDepthTexture = nullptr;
-	FRDGTextureRef ExposureIlluminance = nullptr;
-	FTranslucencyViewResourcesMap TranslucencyViewResourcesMap;
-	FPathTracingResources PathTracingResources;
-
-	bool bSeparateCustomStencil = false;
-
-	void Validate() const
-	{
-		check(SceneTextures);
-		check(ViewFamilyTexture);
-		check(TranslucencyViewResourcesMap.IsValid());
-	}
-};
-
 void AddPostProcessingPasses(
 	FRDGBuilder& GraphBuilder,
 	const FViewInfo& View, int32 ViewIndex,
