Fixes a bug in mobile post processing chain where AddPostProcessMaterialPass(BL_SceneColorAfterTonemapping) could be called even when PassSequence.SetEnabled(EPass::PostProcessMaterialAfterTonemapping, false); with bIsPostProcessingEnabled=false

diff --git a/Engine/Source/Runtime/Renderer/Private/PostProcess/PostProcessing.cpp b/Engine/Source/Runtime/Renderer/Private/PostProcess/PostProcessing.cpp
--- a/Engine/Source/Runtime/Renderer/Private/PostProcess/PostProcessing.cpp
+++ b/Engine/Source/Runtime/Renderer/Private/PostProcess/PostProcessing.cpp
@@ -2728,7 +2728,10 @@ void AddMobilePostProcessingPasses(FRDGBuilder& GraphBuilder, FScene* Scene, con
 
 	if (IsPostProcessingEnabled(View))
 	{
-		AddPostProcessMaterialPass(BL_SceneColorAfterTonemapping, true);
+		if (PassSequence.IsEnabled(EPass::PostProcessMaterialAfterTonemapping))
+		{
+			AddPostProcessMaterialPass(BL_SceneColorAfterTonemapping, true);
+		}
 
 		if (PassSequence.IsEnabled(EPass::FXAA))
 		{
