* Allowed custom render passes to specify a persistent "FSceneViewStateInterface" in order to properly create a scene view / scene view key. This allows custom render passes to eventually use persistent scene view data (e.g. temporal)
* Added validation that a scene view state, if used by a FViewInfo, cannot be shared by 2 different FViewInfo, since this some systems like visibility don't support (and it doesn't sound like a good idea either)
* Custom render passes now have an option to specify whether they are a scene capture (so that non custom render passes made for other purposes than scene captures are not considered as such)

diff --git a/Engine/Source/Runtime/Engine/Public/SceneInterface.h b/Engine/Source/Runtime/Engine/Public/SceneInterface.h
--- a/Engine/Source/Runtime/Engine/Public/SceneInterface.h
+++ b/Engine/Source/Runtime/Engine/Public/SceneInterface.h
@@ -701,6 +701,8 @@ public:
 		TSet<FPrimitiveComponentId> HiddenPrimitives;
 		TOptional<TSet<FPrimitiveComponentId>> ShowOnlyPrimitives;
 		const AActor* ViewActor = nullptr;
+		class FSceneViewStateInterface* ViewStateInterface = nullptr;
+		bool bIsSceneCapture = false;
 
 		class FCustomRenderPassBase* CustomRenderPass = nullptr;
 	};
diff --git a/Engine/Source/Runtime/Renderer/Private/SceneCaptureRendering.cpp b/Engine/Source/Runtime/Renderer/Private/SceneCaptureRendering.cpp
--- a/Engine/Source/Runtime/Renderer/Private/SceneCaptureRendering.cpp
+++ b/Engine/Source/Runtime/Renderer/Private/SceneCaptureRendering.cpp
@@ -922,6 +922,7 @@ void FScene::UpdateSceneCaptureContents(USceneCaptureComponent2D* CaptureCompone
 			PassInput.ViewRotationMatrix = ViewRotationMatrix;
 			PassInput.ProjectionMatrix = ProjectionMatrix;
 			PassInput.ViewActor = CaptureComponent->GetViewOwner();
+			PassInput.bIsSceneCapture = true;
 
 			FString DebugName = CaptureComponent->CaptureSource == ESceneCaptureSource::SCS_SceneDepth ? TEXT("SceneCapturePass_SceneDepth") : TEXT("SceneCapturePass_DeviceDepth");
 			FCustomRenderPassBase::ERenderOutput RenderOutput = CaptureComponent->CaptureSource == ESceneCaptureSource::SCS_SceneDepth ? FCustomRenderPassBase::ERenderOutput::SceneDepth : FCustomRenderPassBase::ERenderOutput::DeviceDepth;
diff --git a/Engine/Source/Runtime/Renderer/Private/SceneRendering.cpp b/Engine/Source/Runtime/Renderer/Private/SceneRendering.cpp
--- a/Engine/Source/Runtime/Renderer/Private/SceneRendering.cpp
+++ b/Engine/Source/Runtime/Renderer/Private/SceneRendering.cpp
@@ -2772,10 +2772,11 @@ FSceneRenderer::FSceneRenderer(const FSceneViewFamily* InViewFamily, FHitProxyCo
 		}
 	}
 
+	// Prepare custom render passes and their views:
 	CustomRenderPassInfos.Empty(Scene->CustomRenderPassRendererInputs.Num());
 	CustomRenderPassInfos.AddDefaulted(Scene->CustomRenderPassRendererInputs.Num());
 
-	int32 NumSceneCaptureViews = 0;
+	int32 NumAdditionalViews = 0;
 	for (int32 i = 0; i < Scene->CustomRenderPassRendererInputs.Num(); i++)
 	{
 		const FScene::FCustomRenderPassRendererInput& PassInput = Scene->CustomRenderPassRendererInputs[i];
@@ -2784,11 +2785,12 @@ FSceneRenderer::FSceneRenderer(const FSceneViewFamily* InViewFamily, FHitProxyCo
 		CustomRenderPassInfos[i].CustomRenderPass = CustomRenderPass;
 
 		FSceneViewInitOptions ViewInitOptions;
+		ViewInitOptions.SceneViewStateInterface = PassInput.ViewStateInterface;
 		ViewInitOptions.SetViewRectangle(FIntRect(0, 0, CustomRenderPass->GetRenderTargetSize().X, CustomRenderPass->GetRenderTargetSize().Y));
 		ViewInitOptions.ViewOrigin = PassInput.ViewLocation;
 		ViewInitOptions.ViewRotationMatrix = PassInput.ViewRotationMatrix;
 		ViewInitOptions.ProjectionMatrix = PassInput.ProjectionMatrix;
-		ViewInitOptions.bIsSceneCapture = true;
+		ViewInitOptions.bIsSceneCapture = PassInput.bIsSceneCapture;
 		ViewInitOptions.ViewFamily = &ViewFamily;
 		ViewInitOptions.ViewActor = PassInput.ViewActor;
 		ViewInitOptions.ShowOnlyPrimitives = PassInput.ShowOnlyPrimitives;
@@ -2803,10 +2805,10 @@ FSceneRenderer::FSceneRenderer(const FSceneViewFamily* InViewFamily, FHitProxyCo
 		ViewInfo->CustomRenderPass = CustomRenderPass;
 		CustomRenderPass->Views.Add(ViewInfo);
 
-		NumSceneCaptureViews++;
+		NumAdditionalViews++;
 	}
 
-	AllViews.Empty(Views.Num() + NumSceneCaptureViews);
+	AllViews.Empty(Views.Num() + NumAdditionalViews);
 	for (int32 i = 0; i < Views.Num(); ++i)
 	{
 		AllViews.Add(&Views[i]);
@@ -2819,6 +2821,19 @@ FSceneRenderer::FSceneRenderer(const FSceneViewFamily* InViewFamily, FHitProxyCo
 		}
 	}
 
+#if !UE_BUILD_SHIPPING
+	// Validate the views
+	TSet<FSceneViewStateInterface*> UniqueViewStates;
+	for (FViewInfo* View : AllViews)
+	{
+		if (View->State != nullptr)
+		{
+			checkf(!UniqueViewStates.Contains(View->State), TEXT("2 views sharing a view state is currently forbidden, please make sure each FViewInfo is using a separate FSceneViewStateInterface or none at all"));
+			UniqueViewStates.Add(View->State);
+		}
+	}
+#endif // !UE_BUILD_SHIPPING
+
 	check(!ViewFamily.AllViews.Num());
 	ViewFamily.AllViews.Append(AllViews);
 
