HeterogeneousVolumes: Fix regression in directional light shadow generation.

diff --git a/Engine/Source/Runtime/Renderer/Private/HeterogeneousVolumes/HeterogeneousVolumesVoxelGridPipeline.cpp b/Engine/Source/Runtime/Renderer/Private/HeterogeneousVolumes/HeterogeneousVolumesVoxelGridPipeline.cpp
--- a/Engine/Source/Runtime/Renderer/Private/HeterogeneousVolumes/HeterogeneousVolumesVoxelGridPipeline.cpp
+++ b/Engine/Source/Runtime/Renderer/Private/HeterogeneousVolumes/HeterogeneousVolumesVoxelGridPipeline.cpp
@@ -3301,6 +3301,7 @@ void RenderVolumetricShadowMapForLightWithVoxelGrid(
 		}
 		else if (LightType == LightType_Directional)
 		{
+			bIsDirectionalLight = true;
 			// Build orthographic projection centered around volume..
 			FVector PreViewTranslation = View.ViewMatrices.GetPreViewTranslation();
 			FMatrix TranslatedWorldToWorldMatrix = FTranslationMatrix(-PreViewTranslation);
