Scene Capture:  Skip RHI resource flush when rendering multiple scene captures in a frame (similar to what's done for nDisplay multi-view).  In a test project with 7 scene captures, saves 3 ms (8% of the frame).

#rnx

diff --git a/Engine/Source/Runtime/Engine/Classes/Components/SceneCaptureComponent.h b/Engine/Source/Runtime/Engine/Classes/Components/SceneCaptureComponent.h
--- a/Engine/Source/Runtime/Engine/Classes/Components/SceneCaptureComponent.h
+++ b/Engine/Source/Runtime/Engine/Classes/Components/SceneCaptureComponent.h
@@ -210,6 +210,9 @@ public:
 	/** Whether this component is a USceneCaptureComponentCube */
 	virtual bool IsCube() const { return false; }
 
+	/** Updates "FrameUpdated" field, returns "true" if this is one of multiple scene captures this frame.  Called from FScene::UpdateSceneCaptureContents. */
+	ENGINE_API bool SetFrameUpdated();
+
 protected:
 	/** Update the show flags from our show flags settings (ideally, you'd be able to set this more directly, but currently unable to make FEngineShowFlags a UStruct to use it as a FProperty...) */
 	ENGINE_API void UpdateShowFlags();
@@ -229,6 +232,9 @@ protected:
 	 */
 	TIndirectArray<FSceneViewStateReference> ViewStates;
 
+	/** Frame this scene capture was last updated */
+	uint64 FrameUpdated;
+
 #if WITH_EDITORONLY_DATA
 	/** The mesh used by ProxyMeshComponent */
 	UPROPERTY(transient)
diff --git a/Engine/Source/Runtime/Engine/Private/Components/SceneCaptureComponent.cpp b/Engine/Source/Runtime/Engine/Private/Components/SceneCaptureComponent.cpp
--- a/Engine/Source/Runtime/Engine/Private/Components/SceneCaptureComponent.cpp
+++ b/Engine/Source/Runtime/Engine/Private/Components/SceneCaptureComponent.cpp
@@ -44,6 +44,9 @@
 static TMultiMap<TWeakObjectPtr<UWorld>, TWeakObjectPtr<USceneCaptureComponent> > SceneCapturesToUpdateMap;
 static FCriticalSection SceneCapturesToUpdateMapCS;
 
+// Tracks the latest frame a scene capture was rendered
+static uint64 GSceneCaptureLatestRenderedFrame = INDEX_NONE;
+
 static TAutoConsoleVariable<bool> CVarSCOverrideOrthographicTilingValues(
 	TEXT("r.SceneCapture.OverrideOrthographicTilingValues"),
 	false,
@@ -426,6 +429,17 @@ void USceneCaptureComponent::UpdateShowFlags()
 	}
 }
 
+bool USceneCaptureComponent::SetFrameUpdated()
+{
+	bool bFirstCaptureThisFrame = GFrameCounter != GSceneCaptureLatestRenderedFrame;
+	bool bRepeatOfThisCapture = GFrameCounter == FrameUpdated;
+	GSceneCaptureLatestRenderedFrame = GFrameCounter;
+	FrameUpdated = GFrameCounter;
+
+	// We consider this a "multiple" capture if it's not the first capture this frame, and not a repeat of the same capture.
+	return !bFirstCaptureThisFrame && !bRepeatOfThisCapture;
+}
+
 #if WITH_EDITOR
 
 bool USceneCaptureComponent::CanEditChange(const FProperty* InProperty) const
diff --git a/Engine/Source/Runtime/Renderer/Private/SceneCaptureRendering.cpp b/Engine/Source/Runtime/Renderer/Private/SceneCaptureRendering.cpp
--- a/Engine/Source/Runtime/Renderer/Private/SceneCaptureRendering.cpp
+++ b/Engine/Source/Runtime/Renderer/Private/SceneCaptureRendering.cpp
@@ -952,6 +952,10 @@ void FScene::UpdateSceneCaptureContents(USceneCaptureComponent2D* CaptureCompone
 
 		check(SceneRenderer != nullptr);
 
+		// When bIsMultipleSceneCapture is true, set bIsFirstSceneRenderer to false, which tells the scene renderer it can skip RHI resource flush, saving performance
+		bool bIsMultipleSceneCapture = CaptureComponent->SetFrameUpdated();
+		SceneRenderer->bIsFirstSceneRenderer = !bIsMultipleSceneCapture;
+
 		SceneRenderer->Views[0].bFogOnlyOnRenderedOpaque = CaptureComponent->bConsiderUnrenderedOpaquePixelAsFullyTranslucent;
 
 		SceneRenderer->ViewFamily.SceneCaptureCompositeMode = CaptureComponent->CompositeMode;
@@ -1204,6 +1208,8 @@ void FScene::UpdateSceneCaptureContents(USceneCaptureComponentCube* CaptureCompo
 
 	if (TextureTarget)
 	{
+		bool bIsMultipleSceneCapture = CaptureComponent->SetFrameUpdated();
+
 		const float FOV = 90 * (float)PI / 360.0f;
 		for (int32 faceidx = 0; faceidx < (int32)ECubeFace::CubeFace_MAX; faceidx++)
 		{
@@ -1232,6 +1238,10 @@ void FScene::UpdateSceneCaptureContents(USceneCaptureComponentCube* CaptureCompo
 				Location, ProjectionMatrix, CaptureComponent->MaxViewDistanceOverride,
 				bCaptureSceneColor, &PostProcessSettings, 0, CaptureComponent->GetViewOwner(), faceidx);
 
+			// When bIsMultipleSceneCapture is true, set bIsFirstSceneRenderer to false, which tells the scene renderer it can skip RHI resource flush, saving performance.
+			// We can also skip RHI resource flush on faces after the first.
+			SceneRenderer->bIsFirstSceneRenderer = (faceidx == 0) && !bIsMultipleSceneCapture;
+
 			for (const FSceneViewExtensionRef& Extension : SceneRenderer->ViewFamily.ViewExtensions)
 			{
 				Extension->SetupViewFamily(SceneRenderer->ViewFamily);
