Fix non-atomic concurrent writes to visibility map.

While different tasks always operate on different primitive indices, the visibility map being a bit array means that concurrent accesses are not safe if within the same dword.

diff --git a/Engine/Source/Runtime/Renderer/Private/SceneVisibility.cpp b/Engine/Source/Runtime/Renderer/Private/SceneVisibility.cpp
--- a/Engine/Source/Runtime/Renderer/Private/SceneVisibility.cpp
+++ b/Engine/Source/Runtime/Renderer/Private/SceneVisibility.cpp
@@ -2355,7 +2355,7 @@ FGPUOcclusionPacket::FGPUOcclusionPacket(FVisibilityViewPacket& InViewPacket, co
 	, bNewlyConsideredBBoxExpandActive(GExpandNewlyOcclusionTestedBBoxesAmount > 0.0f && GFramesToExpandNewlyOcclusionTestedBBoxes > 0 && GFramesNotOcclusionTestedToExpandBBoxes > 0)
 {}
 
-template <typename VisitorType>
+template <bool bIsParallel, typename VisitorType>
 bool FGPUOcclusionPacket::OcclusionCullPrimitive(VisitorType& Visitor, FOcclusionCullResult& Result, int32 Index)
 {
 	const uint8 OcclusionFlags = Scene.PrimitiveOcclusionFlags[Index];
@@ -2373,7 +2373,14 @@ bool FGPUOcclusionPacket::OcclusionCullPrimitive(VisitorType& Visitor, FOcclusio
 		bSubQueries = true;
 		if (!NumSubQueries)
 		{
-			View.PrimitiveVisibilityMap[Index] = false;
+			if constexpr (bIsParallel)
+			{
+				View.PrimitiveVisibilityMap[Index].AtomicSet(false);
+			}
+			else
+			{
+				View.PrimitiveVisibilityMap[Index] = false;
+			}
 			return false;
 		}
 
@@ -2644,7 +2651,14 @@ bool FGPUOcclusionPacket::OcclusionCullPrimitive(VisitorType& Visitor, FOcclusio
 		{
 			if (bIsOccluded)
 			{
-				View.PrimitiveVisibilityMap[Index] = false;
+				if constexpr (bIsParallel)
+				{
+					View.PrimitiveVisibilityMap[Index].AtomicSet(false);
+				}
+				else
+				{
+					View.PrimitiveVisibilityMap[Index] = false;
+				}
 				bIsVisible = false;
 				Result.NumCulledPrimitives++;
 			}
@@ -2658,7 +2672,14 @@ bool FGPUOcclusionPacket::OcclusionCullPrimitive(VisitorType& Visitor, FOcclusio
 
 		if (bAllSubOccluded)
 		{
-			View.PrimitiveVisibilityMap[Index] = false;
+			if constexpr (bIsParallel)
+			{
+				View.PrimitiveVisibilityMap[Index].AtomicSet(false);
+			}
+			else
+			{
+				View.PrimitiveVisibilityMap[Index] = false;
+			}
 			bIsVisible = false;
 			Result.NumCulledPrimitives++;
 		}
@@ -2931,7 +2952,7 @@ FOcclusionCullResult FGPUOcclusionParallelPacket::OcclusionCullTask(FPrimitiveIn
 
 	for (int32 Index : Input)
 	{
-		if (OcclusionCullPrimitive(RecordVisitor, Result, Index))
+		if (OcclusionCullPrimitive<true>(RecordVisitor, Result, Index))
 		{
 			PrimitiveIndexList.Emplace(Index);
 		}
@@ -3056,7 +3077,7 @@ void FGPUOcclusionSerial::AddPrimitives(FPrimitiveRange PrimitiveRange)
 		EOcclusionFlags::Type OcclusionFlags;
 		if (Packet.CanBeOccluded(BitIt.GetIndex(), OcclusionFlags))
 		{
-			Packet.OcclusionCullPrimitive(ProcessVisitor, OcclusionCullResult, BitIt.GetIndex());
+			Packet.OcclusionCullPrimitive<false>(ProcessVisitor, OcclusionCullResult, BitIt.GetIndex());
 		}
 	}
 }
diff --git a/Engine/Source/Runtime/Renderer/Private/SceneVisibilityPrivate.h b/Engine/Source/Runtime/Renderer/Private/SceneVisibilityPrivate.h
--- a/Engine/Source/Runtime/Renderer/Private/SceneVisibilityPrivate.h
+++ b/Engine/Source/Runtime/Renderer/Private/SceneVisibilityPrivate.h
@@ -924,7 +924,7 @@ public:
 		TaskConfig.OcclusionCull.NumTestedQueries.fetch_add(Result.NumTestedQueries, std::memory_order_relaxed);
 	}
 
-	template <typename VisitorType>
+	template <bool bIsParallel, typename VisitorType>
 	bool OcclusionCullPrimitive(VisitorType& Visitor, FOcclusionCullResult& Result, int32 Index);
 
 	//////////////////////////////////////////////////////////////////////////////
