Fixed crash in mobile renderer due to LPI invalidating static meshes that were also added to the scene.

diff --git a/Engine/Source/Runtime/Renderer/Private/RendererScene.cpp b/Engine/Source/Runtime/Renderer/Private/RendererScene.cpp
--- a/Engine/Source/Runtime/Renderer/Private/RendererScene.cpp
+++ b/Engine/Source/Runtime/Renderer/Private/RendererScene.cpp
@@ -6354,17 +6354,6 @@ void FScene::Update(FRDGBuilder& GraphBuilder, const FUpdateParameters& Paramete
 
 	SceneExtensionsUpdaters.PostSceneUpdate(GraphBuilder, SceneUpdateChangeSetStorage.GetPostUpdateSet());
 
-	CreateLightPrimitiveInteractionsTask = GraphBuilder.AddSetupTask([this, &SceneInfosWithAddToScene]
-	{
-		SCOPED_NAMED_EVENT(CreateLightPrimitiveInteractions, FColor::Emerald);
-
-		for (FPrimitiveSceneInfo* SceneInfo : SceneInfosWithAddToScene)
-		{
-			CreateLightPrimitiveInteractionsForPrimitive(SceneInfo);
-		}
-
-	}, EnumHasAnyFlags(Parameters.AsyncOps, EUpdateAllPrimitiveSceneInfosAsyncOps::CreateLightPrimitiveInteractions));
-
 	const bool bAsyncCacheMeshDrawCommands = EnumHasAnyFlags(Parameters.AsyncOps, EUpdateAllPrimitiveSceneInfosAsyncOps::CacheMeshDrawCommands) && GRHISupportsMultithreadedShaderCreation;
 
 	UE::Tasks::FTask AddStaticMeshesTask = GraphBuilder.AddCommandListSetupTask(
@@ -6422,6 +6411,18 @@ void FScene::Update(FRDGBuilder& GraphBuilder, const FUpdateParameters& Paramete
 		}
 	}
 
+	// LPI creation needs to launch after the static mesh update as it can call RequestStaticMeshUpdate() which modifies PrimitivesNeedingStaticMeshUpdate.
+	CreateLightPrimitiveInteractionsTask = GraphBuilder.AddSetupTask([this, &SceneInfosWithAddToScene]
+	{
+		SCOPED_NAMED_EVENT(CreateLightPrimitiveInteractions, FColor::Emerald);
+
+		for (FPrimitiveSceneInfo* SceneInfo : SceneInfosWithAddToScene)
+		{
+			CreateLightPrimitiveInteractionsForPrimitive(SceneInfo);
+		}
+
+	}, EnumHasAnyFlags(Parameters.AsyncOps, EUpdateAllPrimitiveSceneInfosAsyncOps::CreateLightPrimitiveInteractions));
+
 	if (bScenesPrimitivesNeedStaticMeshElementUpdate)
 	{
 		for (int32 PrimitiveIndex = 0; PrimitiveIndex < Primitives.Num(); PrimitiveIndex++)
