Force Lowerst Lod for Sky Capture on Android

diff --git a/Engine/Config/BaseDeviceProfiles.ini b/Engine/Config/BaseDeviceProfiles.ini
--- a/Engine/Config/BaseDeviceProfiles.ini
+++ b/Engine/Config/BaseDeviceProfiles.ini
@@ -784,6 +784,8 @@ bIsVisibleForAssets=True
 +CVars=r.DistanceFields=0
 ; LightGrid Settings
 +CVars=r.Forward.LightLinkedListCulling=0
+; Scale the LODDistanceFactor when doing a sky light capture
++CVars=r.SkylightCapture.LODDistanceScale=2
 ; PreviewAllowlistCVars and PreviewDenyListCVars are arrays of cvars (or prefixes) that are included or excluded from being applied in editor preview.
 ; If any PreviewAllowlistCVars is set, cvars are denied by default.
 PreviewAllowlistCVars=sg.
diff --git a/Engine/Source/Runtime/Renderer/Private/ReflectionEnvironmentCapture.cpp b/Engine/Source/Runtime/Renderer/Private/ReflectionEnvironmentCapture.cpp
--- a/Engine/Source/Runtime/Renderer/Private/ReflectionEnvironmentCapture.cpp
+++ b/Engine/Source/Runtime/Renderer/Private/ReflectionEnvironmentCapture.cpp
@@ -66,6 +66,14 @@ static FAutoConsoleVariableRef CVarGSupersampleCaptureFactor(
 	ECVF_RenderThreadSafe
 	);
 
+float GSkylightCaptureLODDistanceScale = 1.f;
+static FAutoConsoleVariableRef CVarSkylightCaptureLODDistanceScale(
+	TEXT("r.SkylightCapture.LODDistanceScale"),
+	GSkylightCaptureLODDistanceScale,
+	TEXT("LODDistanceScale for the Sky Light Capture. Default is 1")
+	TEXT("Negative values will be clamped to 1"),
+	ECVF_Scalability);
+
 /** 
  * Mip map used by a Roughness of 0, counting down from the lowest resolution mip (MipCount - 1).  
  * This has been tweaked along with ReflectionCaptureRoughnessMipScale to make good use of the resolution in each mip, especially the highest resolution mips.
@@ -1328,6 +1336,11 @@ void CaptureSceneIntoScratchCubemap(
 		View->StartFinalPostprocessSettings(CapturePosition);
 		View->EndFinalPostprocessSettings(ViewInitOptions);
 
+		if (bCapturingForSkyLight)
+		{
+			const float SkylightCaptureLODDistanceScale = GSkylightCaptureLODDistanceScale > 0.f ? GSkylightCaptureLODDistanceScale : 1.f;
+			View->LODDistanceFactor *= SkylightCaptureLODDistanceScale;
+		}
 		ViewFamily.Views.Add(View);
 
 		ViewFamily.SetScreenPercentageInterface(new FLegacyScreenPercentageDriver(
