Remove unnecessary UAV clear when Heterogeneous Volumes are not present in the view.

diff --git a/Engine/Source/Runtime/Renderer/Private/HeterogeneousVolumes/HeterogeneousVolumes.cpp b/Engine/Source/Runtime/Renderer/Private/HeterogeneousVolumes/HeterogeneousVolumes.cpp
--- a/Engine/Source/Runtime/Renderer/Private/HeterogeneousVolumes/HeterogeneousVolumes.cpp
+++ b/Engine/Source/Runtime/Renderer/Private/HeterogeneousVolumes/HeterogeneousVolumes.cpp
@@ -552,16 +552,14 @@ void FDeferredShadingSceneRenderer::RenderHeterogeneousVolumes(
 	for (int32 ViewIndex = 0; ViewIndex < Views.Num(); ViewIndex++)
 	{
 		FViewInfo& View = Views[ViewIndex];
-
-		// Per-view??
-		FRDGTextureDesc Desc = SceneTextures.Color.Target->Desc;
-		Desc.Format = PF_FloatRGBA;
-		Desc.Flags &= ~(TexCreate_FastVRAM);
-		FRDGTextureRef HeterogeneousVolumeRadiance = GraphBuilder.CreateTexture(Desc, TEXT("HeterogeneousVolumes"));
-		AddClearUAVPass(GraphBuilder, GraphBuilder.CreateUAV(HeterogeneousVolumeRadiance), FLinearColor::Transparent);
-
 		if (ShouldRenderHeterogeneousVolumesForView(View))
 		{
+			FRDGTextureDesc Desc = SceneTextures.Color.Target->Desc;
+			Desc.Format = PF_FloatRGBA;
+			Desc.Flags &= ~(TexCreate_FastVRAM);
+			FRDGTextureRef HeterogeneousVolumeRadiance = GraphBuilder.CreateTexture(Desc, TEXT("HeterogeneousVolumes"));
+			AddClearUAVPass(GraphBuilder, GraphBuilder.CreateUAV(HeterogeneousVolumeRadiance), FLinearColor::Transparent);
+
 			if (HeterogeneousVolumes::GetDebugMode() != 0)
 			{
 				// TODO: Replace with single-scattering voxel grid implementation.
@@ -570,7 +568,7 @@ void FDeferredShadingSceneRenderer::RenderHeterogeneousVolumes(
 					SceneTextures,
 					Scene,
 					ViewFamily,
-					Views[ViewIndex],
+					View,
 					OrthoGridUniformBuffer,
 					FrustumGridUniformBuffer,
 					HeterogeneousVolumeRadiance
@@ -640,83 +638,81 @@ void FDeferredShadingSceneRenderer::RenderHeterogeneousVolumes(
 
 				for (int32 VolumeIndex = 0; VolumeIndex < VolumeMeshes.Num(); ++VolumeIndex)
 				{
+					const IHeterogeneousVolumeInterface* HeterogeneousVolume = VolumeMeshes[VolumeIndex].Volume;
+					const FMaterialRenderProxy* MaterialRenderProxy = VolumeMeshes[VolumeIndex].MaterialRenderProxy;
+					const FPrimitiveSceneProxy* PrimitiveSceneProxy = HeterogeneousVolume->GetPrimitiveSceneProxy();
+					const FPrimitiveSceneInfo* PrimitiveSceneInfo = PrimitiveSceneProxy->GetPrimitiveSceneInfo();
+					const FPersistentPrimitiveIndex PrimitiveId = PrimitiveSceneInfo->GetPersistentIndex();
+					const FBoxSphereBounds LocalBoxSphereBounds = HeterogeneousVolume->GetLocalBounds();
+
+					RDG_EVENT_SCOPE(GraphBuilder, "%s [%d]", *HeterogeneousVolume->GetReadableName(), VolumeIndex);
+
+					// Allocate transmittance volume
+					FRDGTextureRef LightingCacheTexture = GSystemTextures.GetBlackDummy(GraphBuilder);
+					if (HeterogeneousVolumes::GetLightingCacheMode() != 0)
 					{
-						const IHeterogeneousVolumeInterface* HeterogeneousVolume = VolumeMeshes[VolumeIndex].Volume;
-						const FMaterialRenderProxy* MaterialRenderProxy = VolumeMeshes[VolumeIndex].MaterialRenderProxy;
-						const FPrimitiveSceneProxy* PrimitiveSceneProxy = HeterogeneousVolume->GetPrimitiveSceneProxy();
-						const FPrimitiveSceneInfo* PrimitiveSceneInfo = PrimitiveSceneProxy->GetPrimitiveSceneInfo();
-						const FPersistentPrimitiveIndex PrimitiveId = PrimitiveSceneInfo->GetPersistentIndex();
-						const FBoxSphereBounds LocalBoxSphereBounds = HeterogeneousVolume->GetLocalBounds();
-
-						RDG_EVENT_SCOPE(GraphBuilder, "%s [%d]", *HeterogeneousVolume->GetReadableName(), VolumeIndex);
-
-						// Allocate transmittance volume
-						FRDGTextureRef LightingCacheTexture = GSystemTextures.GetBlackDummy(GraphBuilder);
-						if (HeterogeneousVolumes::GetLightingCacheMode() != 0)
-						{
-							// TODO: Allow option for scalar transmittance to conserve bandwidth
-							FIntVector LightingCacheResolution = HeterogeneousVolumes::GetLightingCacheResolution(HeterogeneousVolume);
-							uint32 NumMips = FMath::Log2(float(FMath::Min(FMath::Min(LightingCacheResolution.X, LightingCacheResolution.Y), LightingCacheResolution.Z))) + 1;
-							FRDGTextureDesc LightingCacheDesc = FRDGTextureDesc::Create3D(
-								LightingCacheResolution,
-								!IsMetalPlatform(GShaderPlatformForFeatureLevel[View.FeatureLevel]) ? PF_FloatR11G11B10 : PF_FloatRGBA,
-								FClearValueBinding::Black,
-								TexCreate_ShaderResource | TexCreate_UAV | TexCreate_3DTiling,
-								NumMips
-							);
-							LightingCacheTexture = GraphBuilder.CreateTexture(LightingCacheDesc, TEXT("HeterogeneousVolumes.LightingCacheTexture"));
-							AddClearUAVPass(GraphBuilder, GraphBuilder.CreateUAV(LightingCacheTexture), FLinearColor::Black);
-						}
+						// TODO: Allow option for scalar transmittance to conserve bandwidth
+						FIntVector LightingCacheResolution = HeterogeneousVolumes::GetLightingCacheResolution(HeterogeneousVolume);
+						uint32 NumMips = FMath::Log2(float(FMath::Min(FMath::Min(LightingCacheResolution.X, LightingCacheResolution.Y), LightingCacheResolution.Z))) + 1;
+						FRDGTextureDesc LightingCacheDesc = FRDGTextureDesc::Create3D(
+							LightingCacheResolution,
+							!IsMetalPlatform(GShaderPlatformForFeatureLevel[View.FeatureLevel]) ? PF_FloatR11G11B10 : PF_FloatRGBA,
+							FClearValueBinding::Black,
+							TexCreate_ShaderResource | TexCreate_UAV | TexCreate_3DTiling,
+							NumMips
+						);
+						LightingCacheTexture = GraphBuilder.CreateTexture(LightingCacheDesc, TEXT("HeterogeneousVolumes.LightingCacheTexture"));
+						AddClearUAVPass(GraphBuilder, GraphBuilder.CreateUAV(LightingCacheTexture), FLinearColor::Black);
+					}
 
-						// Material baking executes a pre-shading pipeline
-						if (CVarHeterogeneousVolumesPreshading.GetValueOnRenderThread())
-						{
-							RenderWithPreshading(
-								GraphBuilder,
-								SceneTextures,
-								Scene,
-								ViewFamily,
-								View,
-								// Shadow Data
-								VisibleLightInfos,
-								VirtualShadowMapArray,
-								// Object data
-								HeterogeneousVolume,
-								MaterialRenderProxy,
-								PrimitiveId,
-								LocalBoxSphereBounds,
-								// Transmittance accleration
-								LightingCacheTexture,
-								// Output
-								HeterogeneousVolumeRadiance
-							);
-						}
-						// Otherwise execute a live-shading pipeline
-						else
-						{
-							RenderWithLiveShading(
-								GraphBuilder,
-								SceneTextures,
-								Scene,
-								View,
-								// Shadow Data
-								VisibleLightInfos,
-								VirtualShadowMapArray,
-								// Object Data
-								HeterogeneousVolume,
-								MaterialRenderProxy,
-								PrimitiveId,
-								LocalBoxSphereBounds,
-								// Transmittance accleration
-								LightingCacheTexture,
-								// Output
-								HeterogeneousVolumeRadiance
-							);
-						}
+					// Material baking executes a pre-shading pipeline
+					if (CVarHeterogeneousVolumesPreshading.GetValueOnRenderThread())
+					{
+						RenderWithPreshading(
+							GraphBuilder,
+							SceneTextures,
+							Scene,
+							ViewFamily,
+							View,
+							// Shadow Data
+							VisibleLightInfos,
+							VirtualShadowMapArray,
+							// Object data
+							HeterogeneousVolume,
+							MaterialRenderProxy,
+							PrimitiveId,
+							LocalBoxSphereBounds,
+							// Transmittance accleration
+							LightingCacheTexture,
+							// Output
+							HeterogeneousVolumeRadiance
+						);
+					}
+					// Otherwise execute a live-shading pipeline
+					else
+					{
+						RenderWithLiveShading(
+							GraphBuilder,
+							SceneTextures,
+							Scene,
+							View,
+							// Shadow Data
+							VisibleLightInfos,
+							VirtualShadowMapArray,
+							// Object Data
+							HeterogeneousVolume,
+							MaterialRenderProxy,
+							PrimitiveId,
+							LocalBoxSphereBounds,
+							// Transmittance accleration
+							LightingCacheTexture,
+							// Output
+							HeterogeneousVolumeRadiance
+						);
 					}
 				}
-
 			}
+
 			View.HeterogeneousVolumeRadiance = HeterogeneousVolumeRadiance;
 		}
 	}
