Fixes memory pressure on split-screen and VR due TSR overallocating its history extent

diff --git a/Engine/Shaders/Private/TemporalSuperResolution/TSRRejectShading.usf b/Engine/Shaders/Private/TemporalSuperResolution/TSRRejectShading.usf
--- a/Engine/Shaders/Private/TemporalSuperResolution/TSRRejectShading.usf
+++ b/Engine/Shaders/Private/TemporalSuperResolution/TSRRejectShading.usf
@@ -451,8 +451,9 @@ void MainCS(
 		for (uint ElementIndex = 0; ElementIndex < SIMD_SIZE; ElementIndex++)
 		{
 			const tsr_short2 OutputPixelPos = ComputeElementOutputPixelPos(ElementIndex);
+			const tsr_short2 OutputHistoryPixelPos = OutputPixelPos - tsr_short2(InputInfo_ViewportMin);
 
-			HistoryMoireOutput[tsr_short3(OutputPixelPos, 0)] = MoireHistory.GetElement(ElementIndex);
+			HistoryMoireOutput[tsr_short3(OutputHistoryPixelPos, 0)] = MoireHistory.GetElement(ElementIndex);
 		}
 	}
 	else
@@ -781,6 +782,7 @@ void MainCS(
 		{
 			const tsr_short2 LaneSimdPixelOffset = GetElementPixelOffsetInTile(LANE_STRIDE_X, LANE_STRIDE_Y, GGroupThreadIndex, ElementIndex);
 			const tsr_short2 OutputPixelPos = ComputeElementOutputPixelPos(ElementIndex);
+			const tsr_short2 OutputHistoryPixelPos = OutputPixelPos - tsr_short2(InputInfo_ViewportMin);
 
 			tsr_halfC FinalGuideColor = FinalGuide.GetElement(ElementIndex);
 			
@@ -798,10 +800,10 @@ void MainCS(
 			#endif
 
 			#if CONFIG_SCENE_COLOR_ALPHA
-				HistoryGuideOutput[tsr_short3(OutputPixelPos, 0)] = FinalGuideColor;
-				HistoryGuideOutput[tsr_short3(OutputPixelPos, 1)] = tsr_half4(InputUncertainty.GetElement(ElementIndex), 0, 0, 0);
+				HistoryGuideOutput[tsr_short3(OutputHistoryPixelPos, 0)] = FinalGuideColor;
+				HistoryGuideOutput[tsr_short3(OutputHistoryPixelPos, 1)] = tsr_half4(InputUncertainty.GetElement(ElementIndex), 0, 0, 0);
 			#else
-				HistoryGuideOutput[tsr_short3(OutputPixelPos, 0)] = tsr_half4(FinalGuideColor.xyz, InputUncertainty.GetElement(ElementIndex));
+				HistoryGuideOutput[tsr_short3(OutputHistoryPixelPos, 0)] = tsr_half4(FinalGuideColor.xyz, InputUncertainty.GetElement(ElementIndex));
 			#endif
 			HistoryRejectionOutput[OutputPixelPos] = tsr_half4(
 				RejectionBlendFinal.GetElement(ElementIndex) - tsr_half(0.5 * rcp(256.0)),
diff --git a/Engine/Shaders/Private/TemporalSuperResolution/TSRVisualize.usf b/Engine/Shaders/Private/TemporalSuperResolution/TSRVisualize.usf
--- a/Engine/Shaders/Private/TemporalSuperResolution/TSRVisualize.usf
+++ b/Engine/Shaders/Private/TemporalSuperResolution/TSRVisualize.usf
@@ -25,6 +25,9 @@ FScreenTransform OutputPixelPosToScreenPos;
 FScreenTransform ScreenPosToHistoryUV;
 FScreenTransform ScreenPosToInputPixelPos;
 FScreenTransform ScreenPosToInputUV;
+FScreenTransform ScreenPosToMoireHistoryUV;
+float2 MoireHistoryUVBilinearMin;
+float2 MoireHistoryUVBilinearMax;
 float4x4 ClipToResurrectionClip;
 uint2 OutputViewRectMin;
 uint2 OutputViewRectMax;
@@ -210,11 +213,15 @@ tsr_half3 DisplayAntiFlickering(FContext Ctx)
 	// Compute the threshold at which the variation count should quick in
 	const tsr_half TotalVariationCountThreshold = tsr_half(1.0 / (1.0 - pow(1.0 - BlendFinal, tsr_half(FlickeringFramePeriod))));
 
+	// Texture coordinate
+	float2 MoireHistoryUV = ApplyScreenTransform(Ctx.ScreenPos, ScreenPosToMoireHistoryUV);
+	MoireHistoryUV = fastClamp(MoireHistoryUV, MoireHistoryUVBilinearMin, MoireHistoryUVBilinearMax);
+
 	// Texture fetches
 	tsr_half3 OriginalOpaqueInput = InputTexture.SampleLevel(GlobalPointClampedSampler, Ctx.InputUV, 0);
 	tsr_half  MoireInput = InputMoireLumaTexture.SampleLevel(GlobalPointClampedSampler, Ctx.InputUV, 0);
 	tsr_half4 OriginalTranslucencyInput = InputSceneTranslucencyTexture.SampleLevel(GlobalPointClampedSampler, Ctx.InputUV, 0);
-	tsr_half4 MoireHistory = MoireHistoryTexture.SampleLevel(GlobalPointClampedSampler, Ctx.InputUV, 0);
+	tsr_half4 MoireHistory = MoireHistoryTexture.SampleLevel(GlobalPointClampedSampler, MoireHistoryUV, 0);
 	tsr_half DisableAntiFlickering = tsr_half(1.0) - tsr_half(IsMovingMaskTexture[Ctx.InputPixelPos]) * rcp(tsr_half(255.0));
 
 	// Unpack history
diff --git a/Engine/Source/Runtime/Renderer/Private/PostProcess/TemporalSuperResolution.cpp b/Engine/Source/Runtime/Renderer/Private/PostProcess/TemporalSuperResolution.cpp
--- a/Engine/Source/Runtime/Renderer/Private/PostProcess/TemporalSuperResolution.cpp
+++ b/Engine/Source/Runtime/Renderer/Private/PostProcess/TemporalSuperResolution.cpp
@@ -13,6 +13,7 @@
 #include "RendererModule.h"
 #include "ShaderPlatformCachedIniValue.h"
 #include "PostProcess/PostProcessVisualizeBuffer.h"
+#include "DynamicResolutionState.h"
 
 #define COMPILE_TSR_DEBUG_PASSES (!UE_BUILD_SHIPPING)
 
@@ -1010,6 +1011,9 @@ class FTSRVisualizeCS : public FTSRShader
 		SHADER_PARAMETER(FScreenTransform, ScreenPosToHistoryUV)
 		SHADER_PARAMETER(FScreenTransform, ScreenPosToInputPixelPos)
 		SHADER_PARAMETER(FScreenTransform, ScreenPosToInputUV)
+		SHADER_PARAMETER(FScreenTransform, ScreenPosToMoireHistoryUV)
+		SHADER_PARAMETER(FVector2f, MoireHistoryUVBilinearMin)
+		SHADER_PARAMETER(FVector2f, MoireHistoryUVBilinearMax)
 		SHADER_PARAMETER(FMatrix44f, ClipToResurrectionClip)
 		SHADER_PARAMETER(FIntPoint, OutputViewRectMin)
 		SHADER_PARAMETER(FIntPoint, OutputViewRectMax)
@@ -1353,9 +1357,16 @@ FDefaultTemporalUpscaler::FOutputs AddTemporalSuperResolutionPasses(
 		FIntPoint QuantizedPrimaryUpscaleViewSize;
 		QuantizeSceneBufferSize(OutputRect.Max, QuantizedPrimaryUpscaleViewSize);
 
-		OutputExtent = FIntPoint(
-			FMath::Max(InputExtent.X, QuantizedPrimaryUpscaleViewSize.X),
-			FMath::Max(InputExtent.Y, QuantizedPrimaryUpscaleViewSize.Y));
+		if (GIsEditor)
+		{
+			OutputExtent = FIntPoint(
+				FMath::Max(InputExtent.X, QuantizedPrimaryUpscaleViewSize.X),
+				FMath::Max(InputExtent.Y, QuantizedPrimaryUpscaleViewSize.Y));
+		}
+		else
+		{
+			OutputExtent = QuantizedPrimaryUpscaleViewSize;
+		}
 	}
 	else
 	{
@@ -1364,6 +1375,34 @@ FDefaultTemporalUpscaler::FOutputs AddTemporalSuperResolutionPasses(
 		OutputExtent = InputExtent;
 	}
 
+	FIntPoint HistoryGuideExtent;
+	{
+		// Compute final resolution fraction uper bound.
+		float ResolutionFractionUpperBound = 1.f;
+		if (ISceneViewFamilyScreenPercentage const* ScreenPercentageInterface = View.Family->GetScreenPercentageInterface())
+		{
+			DynamicRenderScaling::TMap<float> DynamicResolutionUpperBounds = ScreenPercentageInterface->GetResolutionFractionsUpperBound();
+			const float PrimaryResolutionFractionUpperBound = DynamicResolutionUpperBounds[GDynamicPrimaryResolutionFraction];
+			ResolutionFractionUpperBound = PrimaryResolutionFractionUpperBound * View.Family->SecondaryViewFraction;
+		}
+
+		FIntPoint MaxRenderingViewSize = FSceneRenderer::ApplyResolutionFraction(*View.Family, View.UnconstrainedViewRect.Size(), ResolutionFractionUpperBound);
+
+		FIntPoint QuantizedMaxGuideSize;
+		QuantizeSceneBufferSize(MaxRenderingViewSize, QuantizedMaxGuideSize);
+
+		if (GIsEditor)
+		{
+			HistoryGuideExtent = FIntPoint(
+				FMath::Max(InputExtent.X, QuantizedMaxGuideSize.X),
+				FMath::Max(InputExtent.Y, QuantizedMaxGuideSize.Y));
+		}
+		else
+		{
+			HistoryGuideExtent = QuantizedMaxGuideSize;
+		}
+	}
+
 	FIntPoint HistoryExtent;
 	FIntPoint HistorySize;
 	{
@@ -1382,9 +1421,16 @@ FDefaultTemporalUpscaler::FOutputs AddTemporalSuperResolutionPasses(
 		FIntPoint QuantizedHistoryViewSize;
 		QuantizeSceneBufferSize(HistorySize, QuantizedHistoryViewSize);
 
-		HistoryExtent = FIntPoint(
-			FMath::Max(InputExtent.X, QuantizedHistoryViewSize.X),
-			FMath::Max(InputExtent.Y, QuantizedHistoryViewSize.Y));
+		if (GIsEditor)
+		{
+			HistoryExtent = FIntPoint(
+				FMath::Max(InputExtent.X, QuantizedHistoryViewSize.X),
+				FMath::Max(InputExtent.Y, QuantizedHistoryViewSize.Y));
+		}
+		else
+		{
+			HistoryExtent = QuantizedHistoryViewSize;
+		}
 	}
 	float OutputToHistoryResolutionFraction = float(HistorySize.X) / float(OutputRect.Width());
 	float OutputToHistoryResolutionFractionSquare = OutputToHistoryResolutionFraction * OutputToHistoryResolutionFraction;
@@ -1537,7 +1583,7 @@ FDefaultTemporalUpscaler::FOutputs AddTemporalSuperResolutionPasses(
 
 		{
 			FRDGTextureDesc Desc = FRDGTextureDesc::Create2DArray(
-				InputExtent,
+				HistoryGuideExtent,
 				bSupportsAlpha ? PF_FloatRGBA : PF_A2B10G10R10,
 				FClearValueBinding::None,
 				TexCreate_ShaderResource | TexCreate_UAV,
@@ -1547,7 +1593,7 @@ FDefaultTemporalUpscaler::FOutputs AddTemporalSuperResolutionPasses(
 
 		{
 			FRDGTextureDesc Desc = FRDGTextureDesc::Create2DArray(
-				InputExtent,
+				HistoryGuideExtent,
 				PF_R8G8B8A8,
 				FClearValueBinding::None,
 				TexCreate_ShaderResource | TexCreate_UAV,
@@ -1661,6 +1707,7 @@ FDefaultTemporalUpscaler::FOutputs AddTemporalSuperResolutionPasses(
 
 		ClipToResurrectionClip = FMatrix44f(InvViewProj * PrevViewProj);
 		ResurrectionGuideViewport = FScreenPassTextureViewport(PrevHistory.GuideArray->Desc.Extent, InputHistory.InputViewportRects[ResurrectionFrameSliceIndex]);
+		ResurrectionGuideViewport.Rect = ResurrectionGuideViewport.Rect - ResurrectionGuideViewport.Rect.Min;
 	}
 
 	// Setup the shader parameters for previous frame history
@@ -1886,13 +1933,13 @@ FDefaultTemporalUpscaler::FOutputs AddTemporalSuperResolutionPasses(
 		PassParameters->PrevHistoryParameters = PrevHistoryParameters;
 
 		{
-			FScreenPassTextureViewport PrevHistoryColorViewport(PrevHistory.GuideArray->Desc.Extent, InputHistory.InputViewportRect);
+			FScreenPassTextureViewport PrevHistoryGuideViewport(PrevHistory.GuideArray->Desc.Extent, InputHistory.InputViewportRect - InputHistory.InputViewportRect.Min);
 			PassParameters->PrevHistoryGuide = PrevHistory.GuideArray;
 			PassParameters->PrevHistoryMoire = PrevHistory.MoireArray;
-			PassParameters->PrevGuideInfo = GetScreenPassTextureViewportParameters(PrevHistoryColorViewport);
+			PassParameters->PrevGuideInfo = GetScreenPassTextureViewportParameters(PrevHistoryGuideViewport);
 			PassParameters->InputPixelPosToReprojectScreenPos = ((FScreenTransform::Identity - InputRect.Min + 0.5f) / InputRect.Size()) * FScreenTransform::ViewportUVToScreenPos;
 			PassParameters->ScreenPosToPrevHistoryGuideBufferUV = FScreenTransform::ChangeTextureBasisFromTo(
-				PrevHistoryColorViewport,
+				PrevHistoryGuideViewport,
 				FScreenTransform::ETextureBasis::ScreenPosition,
 				FScreenTransform::ETextureBasis::TextureUV);
 			PassParameters->ScreenPosToResurrectionGuideBufferUV = FScreenTransform::ChangeTextureBasisFromTo(
@@ -2574,6 +2621,13 @@ FDefaultTemporalUpscaler::FOutputs AddTemporalSuperResolutionPasses(
 			PassParameters->ScreenPosToHistoryUV = FScreenTransform::ChangeTextureBasisFromTo(HistoryExtent, FIntRect(FIntPoint::ZeroValue, HistorySize), FScreenTransform::ETextureBasis::ScreenPosition, FScreenTransform::ETextureBasis::TextureUV);
 			PassParameters->ScreenPosToInputPixelPos = FScreenTransform::ChangeTextureBasisFromTo(InputExtent, InputRect, FScreenTransform::ETextureBasis::ScreenPosition, FScreenTransform::ETextureBasis::TexelPosition);
 			PassParameters->ScreenPosToInputUV = FScreenTransform::ChangeTextureBasisFromTo(InputExtent, InputRect, FScreenTransform::ETextureBasis::ScreenPosition, FScreenTransform::ETextureBasis::TextureUV);
+			{
+				FScreenPassTextureViewport PrevHistoryGuideViewport(History.GuideArray->Desc.Extent, InputHistory.InputViewportRect - InputHistory.InputViewportRect.Min);
+				PassParameters->ScreenPosToMoireHistoryUV = FScreenTransform::ChangeTextureBasisFromTo(PrevHistoryGuideViewport, FScreenTransform::ETextureBasis::ScreenPosition, FScreenTransform::ETextureBasis::TextureUV);
+				PassParameters->MoireHistoryUVBilinearMin = GetScreenPassTextureViewportParameters(PrevHistoryGuideViewport).UVViewportBilinearMin;
+				PassParameters->MoireHistoryUVBilinearMax = GetScreenPassTextureViewportParameters(PrevHistoryGuideViewport).UVViewportBilinearMax;
+			}
+
 			PassParameters->ClipToResurrectionClip = ClipToResurrectionClip;
 			PassParameters->OutputViewRectMin = VisualizeRect.Min;
 			PassParameters->OutputViewRectMax = VisualizeRect.Max;
