Reduces TSR history's memory pressure by keeping only one TSR.History.Moire as it is unpracticable to resurrect with it

diff --git a/Engine/Shaders/Private/TemporalSuperResolution/TSRDecimateHistory.usf b/Engine/Shaders/Private/TemporalSuperResolution/TSRDecimateHistory.usf
--- a/Engine/Shaders/Private/TemporalSuperResolution/TSRDecimateHistory.usf
+++ b/Engine/Shaders/Private/TemporalSuperResolution/TSRDecimateHistory.usf
@@ -381,12 +381,12 @@ tsr_halfC CorrectGuideColorExposure(tsr_halfC ReprojectedHistoryGuideColor, cons
 
 #if CONFIG_MOIRE_REPROJECTION
 // Sample previous moire history
-tsr_half4 ReprojectPrevHistoryMoire(tsr_ushort2 InputPixelPos, float2 TextureUV, float FrameIndex)
+tsr_half4 ReprojectPrevHistoryMoire(tsr_ushort2 InputPixelPos, float2 TextureUV)
 {
 	uint2 Random = Rand3DPCG16(int3(InputPixelPos - InputInfo_ViewportMin, View.StateFrameIndexMod8)).xy;
 	float2 E = Hammersley16(0, 1, Random).xy;
 
-	return PrevHistoryMoire.SampleLevel(GlobalPointClampedSampler, float3(TextureUV + (E - 0.5) * PrevGuideInfo_ExtentInverse, FrameIndex), 0);
+	return PrevHistoryMoire.SampleLevel(GlobalPointClampedSampler, float3(TextureUV + (E - 0.5) * PrevGuideInfo_ExtentInverse, 0), 0);
 }
 
 // Correct pre exposure of the moire history.
@@ -560,7 +560,7 @@ void MainCS(
 
 			#if CONFIG_MOIRE_REPROJECTION
 			{
-				ReprojectedHistoryMoireColor = ReprojectPrevHistoryMoire(InputPixelPos, KernelHistoryBufferUV, PrevFrameIndex);
+				ReprojectedHistoryMoireColor = ReprojectPrevHistoryMoire(InputPixelPos, KernelHistoryBufferUV);
 			}
 			#endif
 		}
diff --git a/Engine/Source/Runtime/Renderer/Private/PostProcess/TemporalSuperResolution.cpp b/Engine/Source/Runtime/Renderer/Private/PostProcess/TemporalSuperResolution.cpp
--- a/Engine/Source/Runtime/Renderer/Private/PostProcess/TemporalSuperResolution.cpp
+++ b/Engine/Source/Runtime/Renderer/Private/PostProcess/TemporalSuperResolution.cpp
@@ -1597,7 +1597,7 @@ FDefaultTemporalUpscaler::FOutputs AddTemporalSuperResolutionPasses(
 				PF_R8G8B8A8,
 				FClearValueBinding::None,
 				TexCreate_ShaderResource | TexCreate_UAV,
-				HistorySliceSequence.FrameStorageCount);
+				/* ArraySize = */ 1);
 			History.MoireArray = GraphBuilder.CreateTexture(Desc, TEXT("TSR.History.Moire"));
 		}
 	}
@@ -1654,7 +1654,7 @@ FDefaultTemporalUpscaler::FOutputs AddTemporalSuperResolutionPasses(
 		PrevHistory.GuideArray = GraphBuilder.RegisterExternalTexture(InputHistory.GuideArray);
 		PrevHistory.MoireArray = InputHistory.MoireArray.IsValid()
 			? GraphBuilder.RegisterExternalTexture(InputHistory.MoireArray)
-			: BlackDummy;
+			: BlackArrayDummy;
 
 		int32 ResurrectionFrameRollingIndex = 0;
 		int32 PrevFrameRollingIndex = 0;
@@ -2159,13 +2159,9 @@ FDefaultTemporalUpscaler::FOutputs AddTemporalSuperResolutionPasses(
 			}
 			else
 			{
-				FRDGTextureUAVDesc MoireUAVDesc(History.MoireArray);
-				MoireUAVDesc.FirstArraySlice = CurrentFrameSliceIndex;
-				MoireUAVDesc.NumArraySlices = 1;
-
-				PassParameters->HistoryMoireOutput = GraphBuilder.CreateUAV(MoireUAVDesc);
+				PassParameters->HistoryMoireOutput = GraphBuilder.CreateUAV(FRDGTextureUAVDesc(History.MoireArray));
 
-				MoireHistoryTexture = GraphBuilder.CreateSRV(FRDGTextureSRVDesc::CreateForSlice(History.MoireArray, CurrentFrameSliceIndex));
+				MoireHistoryTexture = GraphBuilder.CreateSRV(FRDGTextureSRVDesc::CreateForSlice(History.MoireArray, /* SliceIndex = */ 0));
 			}
 
 			// Output how the history should rejected in the HistoryUpdate
