Updates documentation of TSR's cvars

diff --git a/Engine/Source/Runtime/Renderer/Private/PostProcess/TemporalSuperResolution.cpp b/Engine/Source/Runtime/Renderer/Private/PostProcess/TemporalSuperResolution.cpp
--- a/Engine/Source/Runtime/Renderer/Private/PostProcess/TemporalSuperResolution.cpp
+++ b/Engine/Source/Runtime/Renderer/Private/PostProcess/TemporalSuperResolution.cpp
@@ -32,7 +32,9 @@ TAutoConsoleVariable<float> CVarTSRHistorySampleCount(
 	TEXT("r.TSR.History.SampleCount"), 16.0f,
 	TEXT("Maximum number sample for each output pixel in the history. Higher values means more stability on highlights on static images, ")
 	TEXT("but may introduce additional ghosting on firefliers style of VFX. Minimum value supported is 8.0 as TSR was in 5.0 and 5.1. ")
-	TEXT("Maximum value possible due to the encoding of the TSR.History.Metadata is 32.0. Defaults to 16.0."),
+	TEXT("Maximum value possible due to the encoding of the TSR.History.Metadata is 32.0. Defaults to 16.0.\n")
+	TEXT("\n")
+	TEXT("Use \"r.TSR.Visualize 0\" command to see how many samples where accumulated in TSR history on areas of the screen."),
 	ECVF_RenderThreadSafe);
 
 TAutoConsoleVariable<float> CVarTSRHistorySP(
@@ -122,11 +124,13 @@ TAutoConsoleVariable<int32> CVarTSRFlickeringEnable(
 	TEXT("how can the history be identical to rendered frame if the amount of details you have in the rendered frame is not in history? ")
 	TEXT("how can the history accumulate details if the history is too different from the rendered frame?\n")
 	TEXT("\n")
-	TEXT("When enabled, this heuristic monitor how the luminance of the scene right before any translucency drawing stored in the ")
-	TEXT("TSR.Moire.Luma resource how it involves over successive frames. And if it is detected to constantly flicker regularily above a certain ")
+	TEXT("When enabled, this flickering temporal analysis monitor how the luminance of the scene right before any translucency drawing stored in the ")
+	TEXT("TSR.Flickering.Luminance resource how it involves over successive frames. And if it is detected to constantly flicker regularily above a certain ")
 	TEXT("threshold defined with this r.TSR.ShadingRejection.Flickering.* cvars, the heuristic attempts to stabilize the image by letting ghost within ")
 	TEXT("luminance boundary tied to the amplititude of flickering.\n")
 	TEXT("\n")
+	TEXT("Use \"r.TSR.Visualize 7\" command to see on screen where this heuristic quicks in orange and red. Pink is where it is disabled.\n")
+	TEXT("\n")
 	TEXT("One particular caveat of this heuristic is that any opaque geometry with incorrect motion vector can make a pixel look identically flickery ")
 	TEXT("quicking this heuristic in and leaving undesired ghosting effects on the said geometry. When that happens, it is highly encourage to ")
 	TEXT("verify the motion vector through the VisualizeMotionBlur show flag and how these motion vectors are able to reproject previous frame ")
@@ -180,10 +184,9 @@ TAutoConsoleVariable<float> CVarTSRFlickeringMaxParralaxVelocity(
 
 TAutoConsoleVariable<int32> CVarTSRShadingTileOverscan(
 	TEXT("r.TSR.ShadingRejection.TileOverscan"), 3,
-	TEXT("The shading rejection run a network of convolutions on the GPU all in single 16x16 without roundtrip to main memory. ")
+	TEXT("The shading rejection run a network of convolutions on the GPU all in single 32x32 without roundtrip to main video memory. ")
 	TEXT("However chaining many convlutions in this tiles means that some convolutions on the edge arround are becoming corrupted ")
-	TEXT("and therefor need to overlap the tile by couple of padding to hide it. This controls is controled by default in the ")
-	TEXT(" anti-aliasing scalability settings."),
+	TEXT("and therefor need to overlap the tile by couple of padding to hide it. Higher means less prones to tiling artifacts, but performance loss."),
 	ECVF_RenderThreadSafe);
 
 TAutoConsoleVariable<float> CVarTSRShadingExposureOffset(
@@ -198,7 +201,7 @@ TAutoConsoleVariable<float> CVarTSRShadingExposureOffset(
 	TEXT("lifts the shadows's LDR intensity, meaning MeasureBackbufferLDRQuantizationError() is decreased in these shadows and increased in ")
 	TEXT("the highlights, control directly.\n")
 	TEXT("\n")
-	TEXT("The best TSR internal buffer to verify this is TSR.Moire.Luma but must be verified in DumpGPU ")
+	TEXT("The best TSR internal buffer to verify this is TSR.Flickering.Luminance, either with the \"show VisualizeTemporalUpscaler\" command or in DumpGPU ")
 	TEXT("with the RGB Linear[0;1] source color space against the Tonemaper's output in sRGB source color space.\n"),
 	ECVF_RenderThreadSafe);
 
@@ -208,8 +211,10 @@ TAutoConsoleVariable<int32> CVarTSRRejectionAntiAliasingQuality(
 	TEXT("While this may not be critical when the rendering resolution is not much lowered than display resolution, ")
 	TEXT("this technic however becomes essential to hide lower rendering resolution rendering because of two reasons:\n")
 	TEXT(" - the screen space size of aliasing is inverse proportional to rendering resolution;\n")
-	TEXT(" - rendering at lower resolution means need more frame to reach at least 1 rendered pixel per display pixel.")
-	TEXT("")
+	TEXT(" - rendering at lower resolution means need more frame to reach at least 1 rendered pixel per display pixel.\n")
+	TEXT("\n")
+	TEXT("Use \"r.TSR.Visualize 6\" command to see on screen where the spatial anti-aliaser quicks in green.\n")
+	TEXT("\n")
 	TEXT("By default, it is only disabled by default in the low anti-aliasing scalability group."),
 	ECVF_Scalability | ECVF_RenderThreadSafe);
 
@@ -226,8 +231,8 @@ TAutoConsoleVariable<int32> CVarTSRResurrectionEnable(
 	TEXT("again (no matter through parallax disocclusion, shading changes, translucent VFX moving) which will have the advantage ")
 	TEXT("bypass the need to newly accumulate a second time by simply resurrected the previously accumulated details.\n")
 	TEXT("\n")
-	TEXT("Command \"vis TSR.HistoryResurrectionMask\" parts of the screen is being resurrected by TSR.\n")
-	TEXT("Command \"vis TSR.History.Color INDEX1\" one of the persistent frame being recorded and later used for history resurrection.\n")
+	TEXT("Command \"r.TSR.Visualize 4\" too see parts of the screen is being resurrected by TSR in green.\n")
+	TEXT("Command \"r.TSR.Visualize 5\" too see the oldest frame being possibly resurrected.\n")
 	TEXT("\n")
 	TEXT("Currently experimental and disabled by default."),
 	ECVF_Scalability | ECVF_RenderThreadSafe);
@@ -236,13 +241,13 @@ TAutoConsoleVariable<int32> CVarTSRResurrectionPersistentFrameCount(
 	TEXT("r.TSR.Resurrection.PersistentFrameCount"), 2,
 	TEXT("Configures the number of persistent frame to record in history for futur history resurrection. ")
 	TEXT("This will increase the memory footprint of the entire TSR history. ")
-	TEXT("Must be an odd number greater or equal to 2. (default=2)"),
+	TEXT("Must be an even number greater or equal to 2. (default=2)"),
 	ECVF_RenderThreadSafe);
 
 TAutoConsoleVariable<int32> CVarTSRResurrectionPersistentFrameInterval(
 	TEXT("r.TSR.Resurrection.PersistentFrameInterval"), 31,
 	TEXT("Configures in number of frames how often persistent frame should be recorded in history for futur history resurrection. ")
-	TEXT("This has no implication on memory footprint of the TSR history. Must be an even number greater or equal to 1. ")
+	TEXT("This has no implication on memory footprint of the TSR history. Must be an odd number greater or equal to 1. ")
 	TEXT("Uses the VisualizeTSR show flag and r.TSR.Visualize=5 to tune this parameter to your content. ")
 	TEXT("(default=31)"),
 	ECVF_RenderThreadSafe);
@@ -250,9 +255,9 @@ TAutoConsoleVariable<int32> CVarTSRResurrectionPersistentFrameInterval(
 TAutoConsoleVariable<int32> CVarTSRAsyncCompute(
 	TEXT("r.TSR.AsyncCompute"), 2,
 	TEXT("Controls how TSR run on async compute. Some TSR passes can overlap with previous passes.\n")
-	TEXT(" 0: Disabled (default);\n")
+	TEXT(" 0: Disabled;\n")
 	TEXT(" 1: Run on async compute only passes that are completly independent from any intermediary resource of this frame, namely ClearPrevTextures and ForwardScatterDepth passes;\n")
-	TEXT(" 2: Run on async compute only passes that are completly independent or only dependent on the depth and velocity buffer which can overlap for instance with translucency or DOF. Any passes on critical path remains on the graphics queue;\n")
+	TEXT(" 2: Run on async compute only passes that are completly independent or only dependent on the depth and velocity buffer which can overlap for instance with translucency or DOF. Any passes on critical path remains on the graphics queue (default);\n")
 	TEXT(" 3: Run all passes on async compute;"),
 	ECVF_RenderThreadSafe);
 
@@ -261,7 +266,7 @@ TAutoConsoleVariable<float> CVarTSRWeightClampingSampleCount(
 	TEXT("Number of sample to count to in history pixel to clamp history to when output pixel velocity reach r.TSR.Velocity.WeightClampingPixelSpeed. ")
 	TEXT("Higher value means higher stability on movement, but at the expense of additional blur due to successive convolution of each history reprojection.\n")
 	TEXT("\n")
-	TEXT("It is possible to visualize the number of sample in TSR history with the console command `vis TSR.History.Metadata`.\n")
+	TEXT("Use \"r.TSR.Visualize 0\" command to see how many samples where accumulated in TSR history on areas of the screen.\n")
 	TEXT("\n")
 	TEXT("Please note this clamp the sample count in history pixel, not output pixel, and therefore lower values are by designed less ")
 	TEXT("noticeable with higher r.TSR.History.ScreenPercentage. This is done so such that increasing r.TSR.History.ScreenPercentage uniterally & automatically ")
@@ -298,7 +303,7 @@ TAutoConsoleVariable<int32> CVarTSRVisualize(
 	TEXT("  4: Mask where the history is resurrected (with r.TSR.Resurrection=1);\n")
 	TEXT("  5: Mask where the history is resurrected in the resurrected frame (with r.TSR.Resurrection=1), particularily interesting to tune r.TSR.Resurrection.PersistentFrameInterval;\n")
 	TEXT("  6: Mask where spatial anti-aliasing is being computed;\n")
-	TEXT("  7: Mask where the anti-flickering heuristic is taking effects (with r.TSR.ShadingRejection.Flickering=1);\n"),
+	TEXT("  7: Mask where the flickering temporal analysis heuristic is taking effects (with r.TSR.ShadingRejection.Flickering=1);\n"),
 	ECVF_RenderThreadSafe);
 
 #endif
