HeterogeneousVolumes: Fixes to ortho grid builder

1) Clamp top-level grid size to (128, 128, 256)
2) Iterate through all projected shadows when collecting volume mesh batches
3) Change out-of-frustum minimum voxel size to 10

diff --git a/Engine/Source/Runtime/Renderer/Private/HeterogeneousVolumes/HeterogeneousVolumes.cpp b/Engine/Source/Runtime/Renderer/Private/HeterogeneousVolumes/HeterogeneousVolumes.cpp
--- a/Engine/Source/Runtime/Renderer/Private/HeterogeneousVolumes/HeterogeneousVolumes.cpp
+++ b/Engine/Source/Runtime/Renderer/Private/HeterogeneousVolumes/HeterogeneousVolumes.cpp
@@ -484,7 +484,7 @@ void FDeferredShadingSceneRenderer::RenderHeterogeneousVolumeShadows(
 
 	FVoxelGridBuildOptions BuildOptions;
 	BuildOptions.bJitter = false;
-	BuildOptions.MinimumVoxelSizeOutsideFrustum = 2.0;
+	BuildOptions.MinimumVoxelSizeOutsideFrustum = 10.0;
 
 	TRDGUniformBufferRef<FOrthoVoxelGridUniformBufferParameters> OrthoGridUniformBuffer = nullptr;
 	TRDGUniformBufferRef<FFrustumVoxelGridUniformBufferParameters> FrustumGridUniformBuffer = nullptr;
diff --git a/Engine/Source/Runtime/Renderer/Private/HeterogeneousVolumes/HeterogeneousVolumesVoxelGridPipeline.cpp b/Engine/Source/Runtime/Renderer/Private/HeterogeneousVolumes/HeterogeneousVolumesVoxelGridPipeline.cpp
--- a/Engine/Source/Runtime/Renderer/Private/HeterogeneousVolumes/HeterogeneousVolumesVoxelGridPipeline.cpp
+++ b/Engine/Source/Runtime/Renderer/Private/HeterogeneousVolumes/HeterogeneousVolumesVoxelGridPipeline.cpp
@@ -1569,9 +1569,9 @@ void CalculateTopLevelGridResolution(
 	TopLevelGridResolution.Z = FMath::CeilToInt(TopLevelGridResolutionAsFloat.Z);
 
 	// Clamp to a moderate limit to also handle indirection grid allocation
-	TopLevelGridResolution.X = FMath::Clamp(TopLevelGridResolution.X, 1, 512);
-	TopLevelGridResolution.Y = FMath::Clamp(TopLevelGridResolution.Y, 1, 512);
-	TopLevelGridResolution.Z = FMath::Clamp(TopLevelGridResolution.Z, 1, 512);
+	TopLevelGridResolution.X = FMath::Clamp(TopLevelGridResolution.X, 1, 128);
+	TopLevelGridResolution.Y = FMath::Clamp(TopLevelGridResolution.Y, 1, 128);
+	TopLevelGridResolution.Z = FMath::Clamp(TopLevelGridResolution.Z, 1, 256);
 }
 
 void CalculateVoxelSize(
@@ -2528,18 +2528,20 @@ uint32 GetTypeHash(const FVolumetricMeshBatch& MeshBatch)
 	return HashCombineFast(GetTypeHash(MeshBatch.Mesh), GetTypeHash(MeshBatch.Proxy));
 }
 
-const FProjectedShadowInfo* GetProjectedShadowInfo(const FVisibleLightInfo* VisibleLightInfo)
+bool IsDynamicShadow(const FVisibleLightInfo* VisibleLightInfo)
 {
-	if (VisibleLightInfo)
+	check(VisibleLightInfo);
+	return !VisibleLightInfo->ShadowsToProject.IsEmpty();
+}
+
+const FProjectedShadowInfo* GetProjectedShadowInfo(const FVisibleLightInfo* VisibleLightInfo, int32 ShadowIndex)
+{
+	check(VisibleLightInfo);
+	check(ShadowIndex < VisibleLightInfo->ShadowsToProject.Num());
+
+	if (VisibleLightInfo && ShadowIndex < VisibleLightInfo->ShadowsToProject.Num())
 	{
-		for (int32 ShadowIndex = 0; ShadowIndex < VisibleLightInfo->ShadowsToProject.Num(); ShadowIndex++)
-		{
-			FProjectedShadowInfo* ProjectedShadowInfo = VisibleLightInfo->ShadowsToProject[ShadowIndex];
-			if (ProjectedShadowInfo)
-			{
-				return ProjectedShadowInfo;
-			}
-		}
+		return VisibleLightInfo->ShadowsToProject[ShadowIndex];
 	}
 
 	return nullptr;
@@ -2581,19 +2583,22 @@ void CollectHeterogeneousVolumeMeshBatches(
 		if (LightSceneInfo->Proxy->CastsVolumetricShadow())
 		{
 			const FVisibleLightInfo* VisibleLightInfo = &VisibleLightInfos[LightSceneInfo->Id];
-			const FProjectedShadowInfo* ProjectedShadowInfo = GetProjectedShadowInfo(VisibleLightInfo);
+			check(VisibleLightInfo);
+			for (int32 ShadowIndex = 0; ShadowIndex < VisibleLightInfo->ShadowsToProject.Num(); ++ShadowIndex)
+			{
+				const FProjectedShadowInfo* ProjectedShadowInfo = GetProjectedShadowInfo(VisibleLightInfo, ShadowIndex);
 			if (ProjectedShadowInfo != nullptr)
 			{
 				const TArray<FMeshBatchAndRelevance, SceneRenderingAllocator>& MeshBatches = ProjectedShadowInfo->GetDynamicSubjectHeterogeneousVolumeMeshElements();
 				for (int32 MeshBatchIndex = 0; MeshBatchIndex < MeshBatches.Num(); ++MeshBatchIndex)
 				{
-					//HeterogeneousVolumesMeshBatches.FindOrAdd(MeshBatches[MeshBatchIndex]);
 					HeterogeneousVolumesMeshBatches.FindOrAdd(FVolumetricMeshBatch(MeshBatches[MeshBatchIndex].Mesh, MeshBatches[MeshBatchIndex].PrimitiveSceneProxy));
 				}
 			}
 		}
 	}
 }
+}
 
 void BuildOrthoVoxelGrid(
 	FRDGBuilder& GraphBuilder,
@@ -3190,8 +3195,7 @@ void RenderVolumetricShadowMapForLightWithVoxelGrid(
 	check(LightSceneInfo);
 	check(VisibleLightInfo);
 
-	// TODO: Push ProjectedShadowInfo from calling function
-	const FProjectedShadowInfo* ProjectedShadowInfo = GetProjectedShadowInfo(VisibleLightInfo);
+	const FProjectedShadowInfo* ProjectedShadowInfo = GetProjectedShadowInfo(VisibleLightInfo, 0);
 	check(ProjectedShadowInfo != NULL)
 
 	ShadowMapResolution = HeterogeneousVolumes::GetShadowMapResolution();
@@ -3546,7 +3550,7 @@ void RenderAdaptiveVolumetricShadowMapWithVoxelGrid(
 	int32 NumPasses = LightSceneInfoCompact.Num();
 	for (int32 PassIndex = 0; PassIndex < NumPasses; ++PassIndex)
 	{
-		bool bApplyEmissionAndTransmittance = PassIndex == 0;
+		bool bApplyEmissionAndTransmittance = (PassIndex == (NumPasses - 1));
 		bool bApplyDirectLighting = !LightSceneInfoCompact.IsEmpty();
 		bool bApplyShadowTransmittance = false;
 
@@ -3565,8 +3569,7 @@ void RenderAdaptiveVolumetricShadowMapWithVoxelGrid(
 			{
 				VisibleLightInfo = &VisibleLightInfos[LightSceneInfo->Id];
 				bApplyShadowTransmittance = LightSceneInfo->Proxy && LightSceneInfo->Proxy->CastsVolumetricShadow();
-				const FProjectedShadowInfo* ProjectedShadowInfo = GetProjectedShadowInfo(VisibleLightInfo);
-				bDynamicallyShadowed = ProjectedShadowInfo != nullptr;
+				bDynamicallyShadowed = IsDynamicShadow(VisibleLightInfo);
 			}
 
 			TRDGUniformBufferRef<FAdaptiveVolumetricShadowMapUniformBufferParameters> AdaptiveVolumetricShadowMapUniformBuffer;
