Fix Hair velocity output for the second / offset view.

diff --git a/Engine/Shaders/Private/HairStrands/HairStrandsVelocity.usf b/Engine/Shaders/Private/HairStrands/HairStrandsVelocity.usf
--- a/Engine/Shaders/Private/HairStrands/HairStrandsVelocity.usf
+++ b/Engine/Shaders/Private/HairStrands/HairStrandsVelocity.usf
@@ -38,6 +38,7 @@ Texture2D<float> CoverageTexture;
 
 uint TileSize;
 int2 TileCountXY;
+int2 ViewRectMin;
 Buffer<uint> TileCountBuffer;
 Buffer<uint2> TileDataBuffer;
 
@@ -59,7 +60,7 @@ void MainCS(uint3 DispatchThreadId : SV_DispatchThreadID, uint2 InGroupId : SV_G
 		return;
 	}
 	const uint2 GroupId = TileDataBuffer[TileIndex1D];
-	const uint2 PixelCoord = GroupId * TileSize + InGroupThreadId;
+	const uint2 PixelCoord = uint2(ViewRectMin) + GroupId * TileSize + InGroupThreadId;
 
 	const FNodeDesc NodeDesc = DecodeNodeDesc(NodeIndex.Load(uint3(PixelCoord, 0)));
 	bool bNeedFastResolve = false;
diff --git a/Engine/Shaders/Private/HairStrands/HairStrandsVisibilityTile.usf b/Engine/Shaders/Private/HairStrands/HairStrandsVisibilityTile.usf
--- a/Engine/Shaders/Private/HairStrands/HairStrandsVisibilityTile.usf
+++ b/Engine/Shaders/Private/HairStrands/HairStrandsVisibilityTile.usf
@@ -394,6 +394,7 @@ void MainVS(
 uint TileSize;
 uint TileType;
 int2 TileCountXY;
+int2 ViewRectMin;
 int2 Resolution;
 Buffer<uint>  TileCountBuffer;
 Buffer<uint2> TileDataBuffer;
@@ -410,7 +411,7 @@ void TileMainCS(uint2 InGroupId : SV_GroupID, uint2 InGroupThreadId : SV_GroupTh
 		return;
 	}
 	const uint2 GroupId = TileDataBuffer[TileIndex1D];
-	const uint2 PixelCoord = GroupId * TileSize + InGroupThreadId;
+	const uint2 PixelCoord = uint2(ViewRectMin) + GroupId * TileSize + InGroupThreadId;
 
 	if (all(PixelCoord < uint2(Resolution)))
 	{
diff --git a/Engine/Source/Runtime/Renderer/Private/HairStrands/HairStrandsTile.cpp b/Engine/Source/Runtime/Renderer/Private/HairStrands/HairStrandsTile.cpp
--- a/Engine/Source/Runtime/Renderer/Private/HairStrands/HairStrandsTile.cpp
+++ b/Engine/Source/Runtime/Renderer/Private/HairStrands/HairStrandsTile.cpp
@@ -441,6 +441,7 @@ class FHairTileClearCS : public FGlobalShader
 	BEGIN_SHADER_PARAMETER_STRUCT(FParameters, )
 		SHADER_PARAMETER(FIntPoint, Resolution)
 		SHADER_PARAMETER(FIntPoint, TileCountXY)
+		SHADER_PARAMETER(FIntPoint, ViewRectMin)
 		SHADER_PARAMETER(uint32, TileSize)
 		SHADER_PARAMETER(uint32, TileType)
 		SHADER_PARAMETER_RDG_BUFFER_SRV(Buffer<uint2>, TileDataBuffer)
@@ -473,12 +474,15 @@ void AddHairStrandsTileClearPass(
 	FHairStrandsTiles::ETileType TileType,
 	FRDGTextureRef OutTexture)
 {
-	if (!OutTexture || !TileData.IsValid() || TileType == FHairStrandsTiles::ETileType::Count) return;
-	
+	if (!OutTexture || !TileData.IsValid() || TileType == FHairStrandsTiles::ETileType::Count)
+	{
+		return;
+	}
 
 	FHairTileClearCS::FParameters* PassParameters = GraphBuilder.AllocParameters<FHairTileClearCS::FParameters>();
 	PassParameters->Resolution = OutTexture->Desc.Extent;
 	PassParameters->TileCountXY = TileData.TileCountXY;
+	PassParameters->ViewRectMin = View.ViewRect.Min;
 	PassParameters->TileSize = TileData.TileSize;
 	PassParameters->TileType = ToIndex(TileType);
 	PassParameters->TileCountBuffer = TileData.TileCountSRV;
diff --git a/Engine/Source/Runtime/Renderer/Private/HairStrands/HairStrandsVisibility.cpp b/Engine/Source/Runtime/Renderer/Private/HairStrands/HairStrandsVisibility.cpp
--- a/Engine/Source/Runtime/Renderer/Private/HairStrands/HairStrandsVisibility.cpp
+++ b/Engine/Source/Runtime/Renderer/Private/HairStrands/HairStrandsVisibility.cpp
@@ -1242,10 +1242,11 @@ class FHairVelocityCS: public FGlobalShader
 		SHADER_PARAMETER_RDG_BUFFER_SRV(StructuredBuffer<FPackedHairVis>, NodeVis)
 		SHADER_PARAMETER_RDG_TEXTURE_UAV(RWTexture2D, OutVelocityTexture)
 		SHADER_PARAMETER_RDG_TEXTURE_UAV(RWTexture2D, OutResolveMaskTexture)
-		SHADER_PARAMETER_STRUCT_REF(FViewUniformShaderParameters, ViewUniformBuffer)
+		SHADER_PARAMETER_STRUCT_INCLUDE(FViewShaderParameters, View)
 
 		SHADER_PARAMETER(FIntPoint, TileCountXY)
 		SHADER_PARAMETER(uint32, TileSize)
+		SHADER_PARAMETER(FIntPoint, ViewRectMin)
 		SHADER_PARAMETER_RDG_BUFFER_SRV(Buffer<uint2>, TileDataBuffer)
 		SHADER_PARAMETER_RDG_BUFFER_SRV(Buffer<uint>, TileCountBuffer)
 		SHADER_PARAMETER_RDG_BUFFER_SRV(Buffer<uint>, HairTileCount)
@@ -1302,7 +1303,7 @@ static void AddHairVelocityPass(
 
 	FHairVelocityCS::FParameters* PassParameters = GraphBuilder.AllocParameters<FHairVelocityCS::FParameters>();
 	PassParameters->bNeedClear = bNeedClear ? 1u : 0u;
-	PassParameters->ViewUniformBuffer = View.ViewUniformBuffer;
+	PassParameters->View = View.GetShaderParameters();
 	PassParameters->VelocityThreshold = GetHairFastResolveVelocityThreshold(Resolution);
 	PassParameters->CoverageThreshold = GetHairWriteVelocityCoverageThreshold();
 	PassParameters->NodeIndex = NodeIndex;
@@ -1318,6 +1319,7 @@ static void AddHairVelocityPass(
 	PassParameters->Resolution			= Resolution;
 	PassParameters->TileCountXY			= TileData.TileCountXY;
 	PassParameters->TileSize			= TileData.TileSize;
+	PassParameters->ViewRectMin			= View.ViewRect.Min;
 	PassParameters->TileCountBuffer		= TileData.TileCountSRV;
 	PassParameters->TileDataBuffer		= TileData.GetTileBufferSRV(TileType);
 	PassParameters->TileIndirectArgs	= TileData.TileIndirectDispatchBuffer;
