Applying HeightFog and VolumetricFog in-scattering and attenuation to Heterogeneous Volumes.

diff --git a/Engine/Shaders/Private/HeightFogCommon.ush b/Engine/Shaders/Private/HeightFogCommon.ush
--- a/Engine/Shaders/Private/HeightFogCommon.ush
+++ b/Engine/Shaders/Private/HeightFogCommon.ush
@@ -26,6 +26,10 @@
 #define PERMUTATION_SUPPORT_FOG_SECOND_TERM 1
 #endif
 
+#ifndef FOG_MATERIALBLENDING_OVERRIDE
+#define FOG_MATERIALBLENDING_OVERRIDE 0
+#endif
+
 #if INSTANCED_STEREO
 	#if MATERIALBLENDING_ANY_TRANSLUCENT
 		#define FogStructISR TranslucentBasePass.Shared.FogISR
@@ -199,7 +203,7 @@ half4 GetExponentialHeightFog(float3 WorldPositionRelativeToCamera, float Exclud
 	}
 
 	// Fog color is unused when additive / modulate blend modes are active.
-	#if (MATERIALBLENDING_ADDITIVE || MATERIALBLENDING_MODULATE)
+	#if (MATERIALBLENDING_ADDITIVE || MATERIALBLENDING_MODULATE) && !FOG_MATERIALBLENDING_OVERRIDE
 		half3 FogColor = 0.0;
 	#else
 		half3 FogColor = (InscatteringColor) * (1 - ExpFogFactor) + DirectionalInscattering;
diff --git a/Engine/Shaders/Private/HeterogeneousVolumes/HeterogeneousVolumesRayMarchingUtils.ush b/Engine/Shaders/Private/HeterogeneousVolumes/HeterogeneousVolumesRayMarchingUtils.ush
--- a/Engine/Shaders/Private/HeterogeneousVolumes/HeterogeneousVolumesRayMarchingUtils.ush
+++ b/Engine/Shaders/Private/HeterogeneousVolumes/HeterogeneousVolumesRayMarchingUtils.ush
@@ -28,10 +28,18 @@
 #define DIM_USE_ADAPTIVE_VOLUMETRIC_SHADOW_MAP 0
 #endif // DIM_USE_ADAPTIVE_VOLUMETRIC_SHADOW_MAP
 
+#ifndef APPLY_FOG_INSCATTERING
+#define APPLY_FOG_INSCATTERING 0
+#endif // APPLY_FOG_INSCATTERING
+
 #if DIM_USE_ADAPTIVE_VOLUMETRIC_SHADOW_MAP
 #include "HeterogeneousVolumesAdaptiveVolumetricShadowMapSampling.ush"
 #endif // DIM_USE_ADAPTIVE_VOLUMETRIC_SHADOW_MAP
 
+#if APPLY_FOG_INSCATTERING
+#include "../HeightFogCommon.ush"
+#endif // APPLY_FOG_INSCATTERING
+
 float CalcShadowBias()
 {
 	float WorldShadowBias = 0.0f;
@@ -359,4 +367,34 @@ void RayMarchSingleScattering(
 			break;
 		}
 	}
+
+#if APPLY_FOG_INSCATTERING
+	if (any(Transmittance < 1))
+	{
+		// TODO: Change in-scattering distance to be based on the first scattering event, instead of the
+		// intersection distance of the bounding box near-plane.
+		float RayT = RayMarchingContext.LocalRayTMin * RayMarchingContext.LocalToWorldScale;
+		float3 CameraRelative_WorldPosition = RayMarchingContext.WorldRayDirection * RayT;
+		float4 HeightFog = float4(0, 0, 0, 1);
+		if (ShouldApplyHeightFog())
+		{
+			HeightFog = CalculateHeightFog(CameraRelative_WorldPosition);
+		}
+
+		float4 CombinedFog = HeightFog;
+		if (ShouldApplyVolumetricFog() && FogStruct.ApplyVolumetricFog > 0)
+		{
+			float3 WorldPosition = CameraRelative_WorldPosition + RayMarchingContext.WorldRayOrigin;
+			const uint EyeIndex = 0;
+			float3 VolumeUV = ComputeVolumeUV_DEPRECATED(WorldPosition, DFHackToFloat(PrimaryView.WorldToClip));
+			CombinedFog = CombineVolumetricFog(HeightFog, VolumeUV, EyeIndex, RayT);
+		}
+
+		// Fogging in-scattering contribution is not currently attenuated by Heterogeneous Volumes
+		// for elements behind the subject. To prevent doubled contribution of in-scattering for 
+		// regions of low density (eg: silhouette edges), blending is interpreted as alpha-masking instead.
+		float Opacity = 1.0 - Transmittance;
+		Radiance = Radiance * CombinedFog.a + CombinedFog.rgb * Opacity;
+	}
+#endif // APPLY_FOG_INSCATTERING
 }
diff --git a/Engine/Shaders/Private/HeterogeneousVolumes/HeterogeneousVolumesSparseVoxelUniformBufferUtils.ush b/Engine/Shaders/Private/HeterogeneousVolumes/HeterogeneousVolumesSparseVoxelUniformBufferUtils.ush
--- a/Engine/Shaders/Private/HeterogeneousVolumes/HeterogeneousVolumesSparseVoxelUniformBufferUtils.ush
+++ b/Engine/Shaders/Private/HeterogeneousVolumes/HeterogeneousVolumesSparseVoxelUniformBufferUtils.ush
@@ -50,6 +50,16 @@ float GetShadowStepFactor()
 	return SparseVoxelUniformBuffer.ShadowStepFactor;
 }
 
+int ShouldApplyHeightFog()
+{
+	return SparseVoxelUniformBuffer.bApplyHeightFog;
+}
+
+int ShouldApplyVolumetricFog()
+{
+	return SparseVoxelUniformBuffer.bApplyVolumetricFog;
+}
+
 uint3 GetVolumeResolution()
 {
 	return SparseVoxelUniformBuffer.VolumeResolution;
diff --git a/Engine/Source/Runtime/Renderer/Private/HeterogeneousVolumes/HeterogeneousVolumes.cpp b/Engine/Source/Runtime/Renderer/Private/HeterogeneousVolumes/HeterogeneousVolumes.cpp
--- a/Engine/Source/Runtime/Renderer/Private/HeterogeneousVolumes/HeterogeneousVolumes.cpp
+++ b/Engine/Source/Runtime/Renderer/Private/HeterogeneousVolumes/HeterogeneousVolumes.cpp
@@ -196,6 +196,20 @@ static TAutoConsoleVariable<int32> CVarHeterogeneousVolumesDepthSort(
 	ECVF_RenderThreadSafe
 );
 
+static TAutoConsoleVariable<int32> CVarHeterogeneousVolumesApplyHeightFog(
+	TEXT("r.HeterogeneousVolumes.HeightFog"),
+	1,
+	TEXT("Applies height fog to Heterogeneous Volumes (Default = 1)"),
+	ECVF_RenderThreadSafe
+);
+
+static TAutoConsoleVariable<int32> CVarHeterogeneousVolumesApplyVolumetricFog(
+	TEXT("r.HeterogeneousVolumes.VolumetricFog"),
+	1,
+	TEXT("Applies volumetric fog to Heterogeneous Volumes (Default = 1)"),
+	ECVF_RenderThreadSafe
+);
+
 DECLARE_GPU_STAT_NAMED(HeterogeneousVolumeShadowsStat, TEXT("HeterogeneousVolumeShadows"));
 DECLARE_GPU_STAT_NAMED(HeterogeneousVolumesStat, TEXT("HeterogeneousVolumes"));
 
@@ -389,6 +403,16 @@ namespace HeterogeneousVolumes
 		return CVarHeterogeneousVolumesIndirectLighting.GetValueOnRenderThread() != 0;
 	}
 
+	bool ShouldApplyHeightFog()
+	{
+		return CVarHeterogeneousVolumesApplyHeightFog.GetValueOnRenderThread() != 0;
+	}
+
+	bool ShouldApplyVolumetricFog()
+	{
+		return CVarHeterogeneousVolumesApplyVolumetricFog.GetValueOnRenderThread() != 0;
+	}
+
 	// Convenience Utils
 	int GetVoxelCount(FIntVector VolumeResolution)
 	{
diff --git a/Engine/Source/Runtime/Renderer/Private/HeterogeneousVolumes/HeterogeneousVolumes.h b/Engine/Source/Runtime/Renderer/Private/HeterogeneousVolumes/HeterogeneousVolumes.h
--- a/Engine/Source/Runtime/Renderer/Private/HeterogeneousVolumes/HeterogeneousVolumes.h
+++ b/Engine/Source/Runtime/Renderer/Private/HeterogeneousVolumes/HeterogeneousVolumes.h
@@ -76,6 +76,8 @@ namespace HeterogeneousVolumes
 	bool UseLightingCacheForInscattering();
 	bool UseLightingCacheForTransmittance();
 	bool UseAdaptiveVolumetricShadowMapForSelfShadowing();
+	bool ShouldApplyHeightFog();
+	bool ShouldApplyVolumetricFog();
 
 	bool EnableIndirectionGrid();
 	bool EnableLinearInterpolation();
@@ -121,6 +123,8 @@ BEGIN_UNIFORM_BUFFER_STRUCT(FSparseVoxelUniformBufferParameters, )
 	SHADER_PARAMETER(float, StepFactor)
 	SHADER_PARAMETER(float, ShadowStepSize)
 	SHADER_PARAMETER(float, ShadowStepFactor)
+	SHADER_PARAMETER(int, bApplyHeightFog)
+	SHADER_PARAMETER(int, bApplyVolumetricFog)
 END_UNIFORM_BUFFER_STRUCT()
 
 BEGIN_SHADER_PARAMETER_STRUCT(FLightingCacheParameters, )
diff --git a/Engine/Source/Runtime/Renderer/Private/HeterogeneousVolumes/HeterogeneousVolumesLiveShadingPipeline.cpp b/Engine/Source/Runtime/Renderer/Private/HeterogeneousVolumes/HeterogeneousVolumesLiveShadingPipeline.cpp
--- a/Engine/Source/Runtime/Renderer/Private/HeterogeneousVolumes/HeterogeneousVolumesLiveShadingPipeline.cpp
+++ b/Engine/Source/Runtime/Renderer/Private/HeterogeneousVolumes/HeterogeneousVolumesLiveShadingPipeline.cpp
@@ -241,6 +241,11 @@ class FRenderSingleScatteringWithLiveShadingCS : public FMeshMaterialShader
 		SHADER_PARAMETER(int32, VirtualShadowMapId)
 		SHADER_PARAMETER_RDG_UNIFORM_BUFFER(FAdaptiveVolumetricShadowMapUniformBufferParameters, AVSM)
 
+		// Atmosphere
+		SHADER_PARAMETER_RDG_UNIFORM_BUFFER(FFogUniformParameters, FogStruct)
+		SHADER_PARAMETER(int, bApplyHeightFog)
+		SHADER_PARAMETER(int, bApplyVolumetricFog)
+
 		// Indirect Lighting
 		SHADER_PARAMETER_RDG_UNIFORM_BUFFER(FLumenTranslucencyLightingUniforms, LumenGIVolumeStruct)
 
@@ -309,6 +314,8 @@ class FRenderSingleScatteringWithLiveShadingCS : public FMeshMaterialShader
 		FMaterialShader::ModifyCompilationEnvironment(Parameters, OutEnvironment);
 		OutEnvironment.SetDefine(TEXT("THREADGROUP_SIZE_1D"), GetThreadGroupSize1D());
 		OutEnvironment.SetDefine(TEXT("THREADGROUP_SIZE_2D"), GetThreadGroupSize2D());
+		OutEnvironment.SetDefine(TEXT("APPLY_FOG_INSCATTERING"), 1);
+		OutEnvironment.SetDefine(TEXT("FOG_MATERIALBLENDING_OVERRIDE"), 1);
 
 		bool bSupportVirtualShadowMap = IsFeatureLevelSupported(Parameters.Platform, ERHIFeatureLevel::SM5);
 		if (bSupportVirtualShadowMap)
@@ -631,6 +638,11 @@ static void RenderSingleScatteringWithLiveShading(
 		PassParameters->VirtualShadowMapSamplingParameters = VirtualShadowMapArray.GetSamplingParameters(GraphBuilder);
 		PassParameters->AVSM = HeterogeneousVolumes::GetAdaptiveVolumetricShadowMapUniformBuffer(GraphBuilder, View.ViewState, LightSceneInfo);
 
+		TRDGUniformBufferRef<FFogUniformParameters> FogBuffer = CreateFogUniformBuffer(GraphBuilder, View);
+		PassParameters->FogStruct = FogBuffer;
+		PassParameters->bApplyHeightFog = HeterogeneousVolumes::ShouldApplyHeightFog();
+		PassParameters->bApplyVolumetricFog = HeterogeneousVolumes::ShouldApplyVolumetricFog();
+
 		// Indirect lighting data
 		auto* LumenUniforms = GraphBuilder.AllocParameters<FLumenTranslucencyLightingUniforms>();
 		LumenUniforms->Parameters = GetLumenTranslucencyLightingParameters(GraphBuilder, View.GetLumenTranslucencyGIVolume(), View.LumenFrontLayerTranslucency);
diff --git a/Engine/Source/Runtime/Renderer/Private/HeterogeneousVolumes/HeterogeneousVolumesPreshadingPipeline.cpp b/Engine/Source/Runtime/Renderer/Private/HeterogeneousVolumes/HeterogeneousVolumesPreshadingPipeline.cpp
--- a/Engine/Source/Runtime/Renderer/Private/HeterogeneousVolumes/HeterogeneousVolumesPreshadingPipeline.cpp
+++ b/Engine/Source/Runtime/Renderer/Private/HeterogeneousVolumes/HeterogeneousVolumesPreshadingPipeline.cpp
@@ -315,6 +315,9 @@ class FRenderSingleScatteringWithPreshadingCS : public FGlobalShader
 		SHADER_PARAMETER_STRUCT_INCLUDE(FVirtualShadowMapSamplingParameters, VirtualShadowMapSamplingParameters)
 		SHADER_PARAMETER(int32, VirtualShadowMapId)
 
+		// Atmosphere
+		SHADER_PARAMETER_RDG_UNIFORM_BUFFER(FFogUniformParameters, FogStruct)
+
 		// Indirect Lighting
 		SHADER_PARAMETER_RDG_UNIFORM_BUFFER(FLumenTranslucencyLightingUniforms, LumenGIVolumeStruct)
 
@@ -358,6 +361,7 @@ class FRenderSingleScatteringWithPreshadingCS : public FGlobalShader
 		FMaterialShader::ModifyCompilationEnvironment(Parameters, OutEnvironment);
 		OutEnvironment.SetDefine(TEXT("THREADGROUP_SIZE_1D"), GetThreadGroupSize1D());
 		OutEnvironment.SetDefine(TEXT("THREADGROUP_SIZE_2D"), GetThreadGroupSize2D());
+		OutEnvironment.SetDefine("APPLY_FOG_INSCATTERING", 1);
 
 		bool bSupportVirtualShadowMap = IsFeatureLevelSupported(Parameters.Platform, ERHIFeatureLevel::SM5);
 		if (bSupportVirtualShadowMap)
@@ -579,6 +583,9 @@ void RenderSingleScatteringWithPreshadingCompute(
 		}
 		PassParameters->VirtualShadowMapSamplingParameters = VirtualShadowMapArray.GetSamplingParameters(GraphBuilder);
 
+		TRDGUniformBufferRef<FFogUniformParameters> FogBuffer = CreateFogUniformBuffer(GraphBuilder, View);
+		PassParameters->FogStruct = FogBuffer;
+
 		// Indirect lighting data
 		auto* LumenUniforms = GraphBuilder.AllocParameters<FLumenTranslucencyLightingUniforms>();
 		LumenUniforms->Parameters = GetLumenTranslucencyLightingParameters(GraphBuilder, View.GetLumenTranslucencyGIVolume(), View.LumenFrontLayerTranslucency);
@@ -1621,6 +1628,8 @@ void RenderWithPreshading(
 		SparseVoxelUniformBufferParameters->StepFactor = HeterogeneousVolumeInterface->GetStepFactor();
 		SparseVoxelUniformBufferParameters->ShadowStepSize = HeterogeneousVolumes::GetShadowStepSize();
 		SparseVoxelUniformBufferParameters->ShadowStepFactor = HeterogeneousVolumeInterface->GetShadowStepFactor();
+		SparseVoxelUniformBufferParameters->bApplyHeightFog = HeterogeneousVolumes::ShouldApplyHeightFog();
+		SparseVoxelUniformBufferParameters->bApplyVolumetricFog = HeterogeneousVolumes::ShouldApplyVolumetricFog();
 	}
 	TRDGUniformBufferRef<FSparseVoxelUniformBufferParameters> SparseVoxelUniformBuffer = GraphBuilder.CreateUniformBuffer(SparseVoxelUniformBufferParameters);
 
