Fix static analysis warning

diff --git a/Engine/Source/Runtime/Renderer/Private/HeterogeneousVolumes/HeterogeneousVolumesVoxelGridPipeline.cpp b/Engine/Source/Runtime/Renderer/Private/HeterogeneousVolumes/HeterogeneousVolumesVoxelGridPipeline.cpp
--- a/Engine/Source/Runtime/Renderer/Private/HeterogeneousVolumes/HeterogeneousVolumesVoxelGridPipeline.cpp
+++ b/Engine/Source/Runtime/Renderer/Private/HeterogeneousVolumes/HeterogeneousVolumesVoxelGridPipeline.cpp
@@ -2530,15 +2530,18 @@ uint32 GetTypeHash(const FVolumetricMeshBatch& MeshBatch)
 
 bool IsDynamicShadow(const FVisibleLightInfo* VisibleLightInfo)
 {
-	check(VisibleLightInfo);
-	return !VisibleLightInfo->ShadowsToProject.IsEmpty();
+	check(VisibleLightInfo != nullptr);
+	if (VisibleLightInfo)
+	{
+		return !VisibleLightInfo->ShadowsToProject.IsEmpty();
+	}
+
+	return false;
 }
 
 const FProjectedShadowInfo* GetProjectedShadowInfo(const FVisibleLightInfo* VisibleLightInfo, int32 ShadowIndex)
 {
-	check(VisibleLightInfo);
-	check(ShadowIndex < VisibleLightInfo->ShadowsToProject.Num());
-
+	check(VisibleLightInfo != nullptr);
 	if (VisibleLightInfo && ShadowIndex < VisibleLightInfo->ShadowsToProject.Num())
 	{
 		return VisibleLightInfo->ShadowsToProject[ShadowIndex];
