Implements r.TSR.ForceSeparateTranslucency to fix ghosting issue on low EffectsQuality.

diff --git a/Engine/Source/Runtime/Engine/Private/SceneView.cpp b/Engine/Source/Runtime/Engine/Private/SceneView.cpp
--- a/Engine/Source/Runtime/Engine/Private/SceneView.cpp
+++ b/Engine/Source/Runtime/Engine/Private/SceneView.cpp
@@ -308,23 +308,6 @@ static TAutoConsoleVariable<float> CVarRoughnessMax(
 	ECVF_Cheat | ECVF_RenderThreadSafe
 	);
 
-static TAutoConsoleVariable<int32> CVarAllowTranslucencyAfterDOF(
-	TEXT("r.SeparateTranslucency"),
-	1,
-	TEXT("Allows to disable the separate translucency feature (all translucency is rendered in separate RT and composited\n")
-	TEXT("after DOF, if not specified otherwise in the material).\n")
-	TEXT(" 0: off (translucency is affected by depth of field)\n")
-	TEXT(" 1: on costs GPU performance and memory but keeps translucency unaffected by Depth of Field. (default)"),
-	ECVF_Scalability | ECVF_RenderThreadSafe);
-
-static TAutoConsoleVariable<int32> CVarTranslucencyStandardSeparated(
-	TEXT("r.Translucency.StandardSeparated"),
-	0,
-	TEXT("Render translucent meshes in separate buffer from the scene color.\n")
-	TEXT("This prevent those meshes from self refracting and leaking scnee color behind over edges when it should be affect by colored transmittance.\n")
-	TEXT("Forced disabled when r.SeparateTranslucency is 0.\n"),
-	ECVF_RenderThreadSafe | ECVF_Default);
-
 static TAutoConsoleVariable<int32> CVarEnableTemporalUpsample(
 	TEXT("r.TemporalAA.Upsampling"),
 	1,
@@ -2854,27 +2837,6 @@ PRAGMA_ENABLE_DEPRECATION_WARNINGS
 	{
 		EngineShowFlags.ScreenPercentage = false;
 	}
-
-	// Check if the translucency are allowed to be rendered after DOF, if not, translucency after DOF will be rendered in standard translucency.
-	{
-		bool SeparateTranslucencyEnabled = EngineShowFlags.PostProcessing // Used for reflection captures.
-										&& !UseDebugViewPS()
-										&& EngineShowFlags.SeparateTranslucency;
-
-		const bool bIsMobile = GetFeatureLevel() == ERHIFeatureLevel::ES3_1;
-		if (bIsMobile)
-		{
-			const bool bMobileMSAA = GetDefaultMSAACount(ERHIFeatureLevel::ES3_1) > 1;
-			SeparateTranslucencyEnabled &= (IsMobileHDR() && !bMobileMSAA); // on <= ES3_1 separate translucency requires HDR on and MSAA off
-		}
-
-		bAllowTranslucencyAfterDOF = SeparateTranslucencyEnabled && CVarAllowTranslucencyAfterDOF.GetValueOnAnyThread() != 0;
-
-		// We do not allow separated translucency on mobile
-		// When MSAA sample count is >1 it works, but hair has not been properly tested so far due to other issues, so MSAA cannot use separted standard translucent for now.
-		uint32 MSAASampleCount = GetDefaultMSAACount(GetFeatureLevel());
-		bAllowStandardTranslucencySeparated = SeparateTranslucencyEnabled && MSAASampleCount == 1 && !bIsMobile && CVarTranslucencyStandardSeparated.GetValueOnAnyThread() != 0;
-	}
 }
 
 PRAGMA_DISABLE_DEPRECATION_WARNINGS // TOptional can't be deprecated without emitting warnings in destructor
diff --git a/Engine/Source/Runtime/Engine/Public/SceneView.h b/Engine/Source/Runtime/Engine/Public/SceneView.h
--- a/Engine/Source/Runtime/Engine/Public/SceneView.h
+++ b/Engine/Source/Runtime/Engine/Public/SceneView.h
@@ -2382,10 +2382,10 @@ private:
 	TArray<TSharedRef<class ISceneViewFamilyExtentionData, ESPMode::ThreadSafe> > ViewExtentionDatas;
 
 	/** whether the translucency are allowed to render after DOF, if not they will be rendered in standard translucency. */
-	bool bAllowTranslucencyAfterDOF;
+	bool bAllowTranslucencyAfterDOF = false;
 
 	/** whether the pre DOF translucency are allowed to be rendered in separated target from scene to allow for better composition with distortion.*/
-	bool bAllowStandardTranslucencySeparated;
+	bool bAllowStandardTranslucencySeparated = false;
 
 	/** True if this view is the current editing view or the active game view */
 	bool bIsInFocus = true;
diff --git a/Engine/Source/Runtime/Renderer/Private/SceneRendering.cpp b/Engine/Source/Runtime/Renderer/Private/SceneRendering.cpp
--- a/Engine/Source/Runtime/Renderer/Private/SceneRendering.cpp
+++ b/Engine/Source/Runtime/Renderer/Private/SceneRendering.cpp
@@ -240,6 +240,28 @@ static TAutoConsoleVariable<int32> CVarRayTracingSceneUpdateOnce(
 	ECVF_RenderThreadSafe
 );
 
+static TAutoConsoleVariable<int32> CVarAllowTranslucencyAfterDOF(
+	TEXT("r.SeparateTranslucency"),
+	1,
+	TEXT("Allows to disable the separate translucency feature (all translucency is rendered in separate RT and composited\n")
+	TEXT("after DOF, if not specified otherwise in the material).\n")
+	TEXT(" 0: off (translucency is affected by depth of field)\n")
+	TEXT(" 1: on costs GPU performance and memory but keeps translucency unaffected by Depth of Field. (default)"),
+	ECVF_Scalability | ECVF_RenderThreadSafe);
+
+static TAutoConsoleVariable<int32> CVarTranslucencyStandardSeparated(
+	TEXT("r.Translucency.StandardSeparated"),
+	0,
+	TEXT("Render translucent meshes in separate buffer from the scene color.\n")
+	TEXT("This prevent those meshes from self refracting and leaking scnee color behind over edges when it should be affect by colored transmittance.\n")
+	TEXT("Forced disabled when r.SeparateTranslucency is 0.\n"),
+	ECVF_RenderThreadSafe | ECVF_Default);
+
+static TAutoConsoleVariable<int32> CVarTSRForceSeparateTranslucency(
+	TEXT("r.TSR.ForceSeparateTranslucency"), 1,
+	TEXT("Overrides r.SeparateTranslucency whenever TSR is enabled (enabled by default).\n"),
+	ECVF_RenderThreadSafe);
+
 static TAutoConsoleVariable<int32> CVarViewHasTileOffsetData(
 	TEXT("r.ViewHasTileOffsetData"),
 	1,
@@ -2765,6 +2787,40 @@ FSceneRenderer::FSceneRenderer(const FSceneViewFamily* InViewFamily, FHitProxyCo
 	}
 #endif // !UE_BUILD_SHIPPING
 
+
+	// Check if the translucency are allowed to be rendered after DOF, if not, translucency after DOF will be rendered in standard translucency.
+	{
+		bool SeparateTranslucencyEnabled = ViewFamily.EngineShowFlags.PostProcessing // Used for reflection captures.
+			&& !ViewFamily.UseDebugViewPS()
+			&& ViewFamily.EngineShowFlags.SeparateTranslucency;
+
+		const bool bIsMobile = ViewFamily.GetFeatureLevel() == ERHIFeatureLevel::ES3_1;
+		if (bIsMobile)
+		{
+			const bool bMobileMSAA = GetDefaultMSAACount(ERHIFeatureLevel::ES3_1) > 1;
+			SeparateTranslucencyEnabled &= (IsMobileHDR() && !bMobileMSAA); // on <= ES3_1 separate translucency requires HDR on and MSAA off
+		}
+
+		ViewFamily.bAllowTranslucencyAfterDOF = SeparateTranslucencyEnabled && CVarAllowTranslucencyAfterDOF.GetValueOnAnyThread() != 0;
+
+		if (!ViewFamily.bAllowTranslucencyAfterDOF && !bIsMobile && CVarTSRForceSeparateTranslucency.GetValueOnAnyThread() != 0)
+		{
+			for (FViewInfo* View : AllViews)
+			{
+				if (View->AntiAliasingMethod == AAM_TSR)
+				{
+					ViewFamily.bAllowTranslucencyAfterDOF = true;
+					break;
+				}
+			}
+		}
+
+		// We do not allow separated translucency on mobile
+		// When MSAA sample count is >1 it works, but hair has not been properly tested so far due to other issues, so MSAA cannot use separted standard translucent for now.
+		uint32 MSAASampleCount = GetDefaultMSAACount(ViewFamily.GetFeatureLevel());
+		ViewFamily.bAllowStandardTranslucencySeparated = SeparateTranslucencyEnabled && MSAASampleCount == 1 && !bIsMobile && CVarTranslucencyStandardSeparated.GetValueOnAnyThread() != 0;
+	}
+
 	check(!ViewFamily.AllViews.Num());
 	ViewFamily.AllViews.Append(AllViews);
 
