Many Lights - Support AVOID_SELF_INTERSECTIONS permutation in HWRT path to prevent incorrect self-shadowing when ray tracing representation doesn't match rasterized geometry.

diff --git a/Engine/Shaders/Private/ManyLights/ManyLightsHardwareRayTracing.usf b/Engine/Shaders/Private/ManyLights/ManyLightsHardwareRayTracing.usf
--- a/Engine/Shaders/Private/ManyLights/ManyLightsHardwareRayTracing.usf
+++ b/Engine/Shaders/Private/ManyLights/ManyLightsHardwareRayTracing.usf
@@ -7,6 +7,7 @@
 #include "ManyLights.ush"
 #include "ManyLightsRayTracing.ush"
 #include "../Lumen/LumenHardwareRayTracingPayloadCommon.ush"
+#include "../Lumen/LumenHardwareRayTracingCommon.ush"
 #if HAIR_VOXEL_TRACES
 	#include "../HairStrands/HairStrandsRaytracing.ush"
 #endif
@@ -19,12 +20,6 @@ uint MaxTraversalIterations;
 RaytracingAccelerationStructure TLAS;
 
 #if LUMEN_HARDWARE_INLINE_RAYTRACING
-	struct FHitGroupRootConstants
-	{
-		uint BaseInstanceIndex;
-		uint UserData;
-	};
-
 	StructuredBuffer<FHitGroupRootConstants> HitGroupData;
 	StructuredBuffer<FRayTracingSceneMetadataRecord> RayTracingSceneMetadata;
 #endif
@@ -72,61 +67,33 @@ RAY_TRACING_ENTRY_RAYGEN_OR_INLINE(HardwareRayTraceLightSamples)
 
 		ApplyPositionBias(Ray.Origin, Ray.Direction, SceneWorldNormal, RayTracingNormalBias);
 
-		uint RayFlags = RAY_FLAG_ACCEPT_FIRST_HIT_AND_END_SEARCH;
-		uint RayMask = RAY_TRACING_MASK_OPAQUE;
-
-		FLumenMinimalPayload Payload = (FLumenMinimalPayload)0;
-
-		// #ml_todo: Share it with LumenHardwareRayTracingCommon.ush
-		#if LUMEN_HARDWARE_INLINE_RAYTRACING
-		{
-			FTraceRayInlineContext TraceRayInlineContext = CreateTraceRayInlineContext();
-			TraceRayInlineContext.MaxIterations = MaxTraversalIterations;
-
-			FTraceRayInlineResult TraceResult = TraceRayInline(TLAS, RayFlags, RayMask, Ray.GetNativeDesc(), TraceRayInlineContext);
-
-			if (TraceResult.IsHit())
-			{
-				Payload.HitT = TraceResult.HitT;
+		FRayCone RayCone = (FRayCone)0;
+		FRayTracedLightingContext LightingContext = CreateRayTracedLightingContext(
+			RayCone,
+			0,
+			0, // dummy coordinate
+			RAY_FLAG_NONE,
+			MaxTraversalIterations);
 
-				// Our SBTs always have 2 slots per material. Therefore InstanceContributionToHitGroupIndex has a RAY_TRACING_NUM_SHADER_SLOTS multiplier.
-				// But because we only encode 1 material into LumenHardwareRayTracingHitDataBuffer we want to get actual material index.
-				uint InstanceContributionIndex = TraceResult.InstanceContributionToHitGroupIndex / RAY_TRACING_NUM_SHADER_SLOTS;
-				uint HitGroupIndex = TraceResult.GeometryIndex + InstanceContributionIndex;
+		// Shadows don't need closest hit distance
+		LightingContext.bAcceptFirstHitAndEndSearch = true;
 
-				FHitGroupRootConstants HitData = HitGroupData[HitGroupIndex];
-
-				const int Primitive = TraceResult.InstanceID;
-				const int PrimitiveInstanceIndex = TraceResult.InstanceIndex - HitData.BaseInstanceIndex;
-				const uint GPUSceneInstanceIndex = GetGPUSceneInstanceIndex(Primitive, PrimitiveInstanceIndex);
-				Payload.SetGPUSceneInstanceIndex(GPUSceneInstanceIndex);
-
-				uint MaterialShaderIndex = HitData.UserData & 0x3FFFFFFF;
-				Payload.SetMaterialShaderIndex(MaterialShaderIndex);
+#if LUMEN_HARDWARE_INLINE_RAYTRACING
+		LightingContext.HitGroupData = HitGroupData;
+		LightingContext.RayTracingSceneMetadata = RayTracingSceneMetadata;
+#endif // LUMEN_HARDWARE_INLINE_RAYTRACING
 
-				uint bIsTranslucent = (HitData.UserData >> 31) & 0x1;
-				Payload.SetIsTranslucent(bIsTranslucent);
+		LightingContext.InstanceMask = RAY_TRACING_MASK_OPAQUE;
 
-				uint bIsTwoSided = (HitData.UserData >> 30) & 0x1;
-				Payload.SetIsTwoSided(bIsTwoSided);
-			}
-			else
-			{
-				Payload.SetMiss();
-			}
-		}
-		#else
-		{
-			TraceRay(TLAS, RayFlags, RayMask, /*ShaderSlotIndex*/ 0, RAY_TRACING_NUM_SHADER_SLOTS, 0, Ray.GetNativeDesc(), Payload);
-		}
-		#endif
+		FLumenMinimalRayResult TraceResult = TraceLumenMinimalRay(TLAS, Ray, LightingContext);
 
 		#if HAIR_VOXEL_TRACES
 		{
 			// #ml_todo: replace with spatiotemporal blue noise
 			RandomSequence RandSequence;
 			RandomSequence_Initialize(RandSequence, SampleCoord, 0, ManyLightsStateFrameIndex, 1);
-			Payload.HitT = TraverseHair(ScreenCoord, RandSequence, Ray.Origin, Ray.Direction, Payload.HitT, VirtualVoxel.Raytracing_ShadowOcclusionThreshold);
+			TraceResult.HitT = TraverseHair(ScreenCoord, RandSequence, Ray.Origin, Ray.Direction, TraceResult.HitT, VirtualVoxel.Raytracing_ShadowOcclusionThreshold);
+			TraceResult.bHit = TraceResult.HitT > 0;
 		}
 		#endif
 
@@ -137,14 +104,14 @@ RAY_TRACING_ENTRY_RAYGEN_OR_INLINE(HardwareRayTraceLightSamples)
 			const FDeferredLightData LightData = ConvertToDeferredLight(LocalLightData);
 
 			float3 RayStart = Ray.Origin + Ray.Direction * Ray.TMin;
-			float3 RayEnd = Ray.Origin + Ray.Direction * (Payload.IsHit() ? Payload.HitT : Ray.TMax);
+			float3 RayEnd = Ray.Origin + Ray.Direction * (TraceResult.bHit ? TraceResult.HitT : Ray.TMax);
 			float4 RayColor = float4(LightData.Color.xyz / Luminance(LightData.Color.xyz), 1.0f);
 
 			AddLineTWS(Context, RayStart, RayEnd, RayColor);	
 		}
 		#endif
 
-		if (Payload.IsHit())
+		if (TraceResult.bHit)
 		{
 			LightSample.bCompleted = true;
 			LightSample.bVisible = false;
diff --git a/Engine/Source/Runtime/Renderer/Private/ManyLights/ManyLightsRayTracing.cpp b/Engine/Source/Runtime/Renderer/Private/ManyLights/ManyLightsRayTracing.cpp
--- a/Engine/Source/Runtime/Renderer/Private/ManyLights/ManyLightsRayTracing.cpp
+++ b/Engine/Source/Runtime/Renderer/Private/ManyLights/ManyLightsRayTracing.cpp
@@ -76,6 +76,17 @@ static TAutoConsoleVariable<int32> CVarManyLightsHardwareRayTracingMaxIterations
 	ECVF_Scalability | ECVF_RenderThreadSafe
 );
 
+// #ml_todo: Separate config cvars from Lumen once we support multiple SBT with same RayTracingPipeline or Global Uniform Buffers in Ray Tracing
+static TAutoConsoleVariable<bool> CVarManyLightsHardwareRayTracingAvoidSelfIntersections(
+	TEXT("r.ManyLights.HardwareRayTracing.AvoidSelfIntersections"),
+	true,
+	TEXT("Whether to avoid self-intersections.\n")
+	TEXT("Currently shares config with Lumen:\n")
+	TEXT("- r.Lumen.HardwareRayTracing.SkipBackFaceHitDistance\n")
+	TEXT("- r.Lumen.HardwareRayTracing.SkipTwoSidedHitDistance\n"),
+	ECVF_Scalability | ECVF_RenderThreadSafe
+);
+
 static TAutoConsoleVariable<int32> CVarManyLightsHairVoxelTraces(
 	TEXT("r.ManyLights.HairVoxelTraces"),
 	1,
@@ -244,9 +255,10 @@ class FHardwareRayTraceLightSamples : public FLumenHardwareRayTracingShaderBase
 		SHADER_PARAMETER_RDG_UNIFORM_BUFFER(FLumenHardwareRayTracingUniformBufferParameters, LumenHardwareRayTracingUniformBuffer)
 	END_SHADER_PARAMETER_STRUCT()
 
+	class FAvoidSelfIntersections : SHADER_PERMUTATION_BOOL("AVOID_SELF_INTERSECTIONS");
 	class FHairVoxelTraces : SHADER_PERMUTATION_BOOL("HAIR_VOXEL_TRACES");
 	class FDebugMode : SHADER_PERMUTATION_BOOL("DEBUG_MODE");
-	using FPermutationDomain = TShaderPermutationDomain<FHairVoxelTraces, FDebugMode>;
+	using FPermutationDomain = TShaderPermutationDomain<FAvoidSelfIntersections, FHairVoxelTraces, FDebugMode>;
 
 	static bool ShouldCompilePermutation(const FGlobalShaderPermutationParameters& Parameters, Lumen::ERayTracingShaderDispatchType ShaderDispatchType)
 	{
@@ -365,6 +377,7 @@ void FDeferredShadingSceneRenderer::PrepareManyLightsLumenMaterial(const FViewIn
 		for (int32 HairVoxelTraces = 0; HairVoxelTraces < 2; ++HairVoxelTraces)
 		{
 			FHardwareRayTraceLightSamplesRGS::FPermutationDomain PermutationVector;
+			PermutationVector.Set<FHardwareRayTraceLightSamplesRGS::FAvoidSelfIntersections>(CVarManyLightsHardwareRayTracingAvoidSelfIntersections.GetValueOnRenderThread());
 			PermutationVector.Set<FHardwareRayTraceLightSamplesRGS::FHairVoxelTraces>(HairVoxelTraces != 0);
 			PermutationVector.Set<FHardwareRayTraceLightSamplesRGS::FDebugMode>(ManyLights::GetDebugMode() != 0);
 			TShaderRef<FHardwareRayTraceLightSamplesRGS> RayGenerationShader = View.ShaderMap->GetShader<FHardwareRayTraceLightSamplesRGS>(PermutationVector);
@@ -380,6 +393,7 @@ namespace ManyLights
 		FRDGBuilder& GraphBuilder,
 		const FCompactedTraceParameters& CompactedTraceParameters,
 		const FManyLightsParameters& ManyLightsParameters,
+		const FHairVoxelTraceParameters& HairVoxelTraceParameters,
 		FRDGTextureRef LightSamples,
 		FRDGTextureRef LightSampleRayDistance,
 		FHardwareRayTraceLightSamples::FParameters* PassParameters);
@@ -390,12 +404,14 @@ void ManyLights::SetHardwareRayTracingPassParameters(
 	FRDGBuilder& GraphBuilder,
 	const ManyLights::FCompactedTraceParameters& CompactedTraceParameters,
 	const FManyLightsParameters& ManyLightsParameters,
+	const FHairVoxelTraceParameters& HairVoxelTraceParameters,
 	FRDGTextureRef LightSamples,
 	FRDGTextureRef LightSampleRayDistance,
 	FHardwareRayTraceLightSamples::FParameters* PassParameters)
 {
 	PassParameters->CompactedTraceParameters = CompactedTraceParameters;
 	PassParameters->ManyLightsParameters = ManyLightsParameters;
+	PassParameters->HairVoxelTraceParameters = HairVoxelTraceParameters;
 	PassParameters->RWLightSamples = GraphBuilder.CreateUAV(LightSamples);
 	PassParameters->LightSampleRayDistance = LightSampleRayDistance;
 	PassParameters->RayTracingBias = CVarManyLightsHardwareRayTracingBias.GetValueOnRenderThread();
@@ -555,23 +571,24 @@ void ManyLights::RayTraceLightSamples(
 
 		if (ManyLights::UseHardwareRayTracing())
 		{
-			#if RHI_RAYTRACING
+#if RHI_RAYTRACING
+			FHardwareRayTraceLightSamples::FParameters* PassParameters = GraphBuilder.AllocParameters<FHardwareRayTraceLightSamples::FParameters>();
+			ManyLights::SetHardwareRayTracingPassParameters(
+				View,
+				GraphBuilder,
+				CompactedTraceParameters,
+				ManyLightsParameters,
+				HairVoxelTraceParameters,
+				LightSamples,
+				LightSampleRayDistance,
+				PassParameters);
+
+			FHardwareRayTraceLightSamples::FPermutationDomain PermutationVector;
+			PermutationVector.Set<FHardwareRayTraceLightSamples::FAvoidSelfIntersections>(CVarManyLightsHardwareRayTracingAvoidSelfIntersections.GetValueOnRenderThread());
+			PermutationVector.Set<FHardwareRayTraceLightSamples::FHairVoxelTraces>(bHairVoxelTraces);
+			PermutationVector.Set<FHardwareRayTraceLightSamples::FDebugMode>(bDebug);
 			if (ManyLights::UseInlineHardwareRayTracing())
 			{
-				FHardwareRayTraceLightSamplesCS::FParameters* PassParameters = GraphBuilder.AllocParameters<FHardwareRayTraceLightSamplesCS::FParameters>();
-				ManyLights::SetHardwareRayTracingPassParameters(
-					View,
-					GraphBuilder,
-					CompactedTraceParameters,
-					ManyLightsParameters,
-					LightSamples,
-					LightSampleRayDistance,
-					PassParameters);
-				PassParameters->HairVoxelTraceParameters = HairVoxelTraceParameters;
-
-				FHardwareRayTraceLightSamplesCS::FPermutationDomain PermutationVector;
-				PermutationVector.Set<FHardwareRayTraceLightSamplesCS::FHairVoxelTraces>(bHairVoxelTraces);
-				PermutationVector.Set<FHardwareRayTraceLightSamplesCS::FDebugMode>(bDebug);
 				auto ComputeShader = View.ShaderMap->GetShader<FHardwareRayTraceLightSamplesCS>(PermutationVector);
 
 				FComputeShaderUtils::AddPass(
@@ -585,20 +602,6 @@ void ManyLights::RayTraceLightSamples(
 			}
 			else
 			{
-				FHardwareRayTraceLightSamplesRGS::FParameters* PassParameters = GraphBuilder.AllocParameters<FHardwareRayTraceLightSamplesRGS::FParameters>();
-				ManyLights::SetHardwareRayTracingPassParameters(
-					View,
-					GraphBuilder,
-					CompactedTraceParameters,
-					ManyLightsParameters,
-					LightSamples,
-					LightSampleRayDistance,
-					PassParameters);
-				PassParameters->HairVoxelTraceParameters = HairVoxelTraceParameters;
-
-				FHardwareRayTraceLightSamplesRGS::FPermutationDomain PermutationVector;
-				PermutationVector.Set<FHardwareRayTraceLightSamplesRGS::FHairVoxelTraces>(bHairVoxelTraces);
-				PermutationVector.Set<FHardwareRayTraceLightSamplesRGS::FDebugMode>(bDebug);
 				auto RayGenerationShader = View.ShaderMap->GetShader<FHardwareRayTraceLightSamplesRGS>(PermutationVector);
 
 				AddLumenRayTraceDispatchIndirectPass(
