Fix for DDM_NonMaskedOnly during depth rendering after wrong refactor

diff --git a/Engine/Source/Runtime/Renderer/Private/DepthRendering.cpp b/Engine/Source/Runtime/Renderer/Private/DepthRendering.cpp
--- a/Engine/Source/Runtime/Renderer/Private/DepthRendering.cpp
+++ b/Engine/Source/Runtime/Renderer/Private/DepthRendering.cpp
@@ -910,9 +910,11 @@ void FDepthPassMeshProcessor::CollectPSOInitializersInternal(
 	}
 }
 
-bool FDepthPassMeshProcessor::UseDefaultMaterial(const FMaterial& Material, bool bMaterialModifiesMeshPosition, bool bSupportPositionOnlyStream, bool bVFTypeSupportsNullPixelShader, bool& bPositionOnly)
+bool FDepthPassMeshProcessor::ShouldRender(const FMaterial& Material, bool bMaterialModifiesMeshPosition, bool bSupportPositionOnlyStream, bool bVFTypeSupportsNullPixelShader, bool& bUseDefaultMaterial, bool& bPositionOnly)
 {
-	bool bUseDefaultMaterial = false;
+	bool bShouldRender = false;
+	bUseDefaultMaterial = false;
+	bPositionOnly = false;
 
 	if (IsOpaqueBlendMode(Material)
 		&& EarlyZPassMode != DDM_MaskedOnly
@@ -920,6 +922,7 @@ bool FDepthPassMeshProcessor::UseDefaultMaterial(const FMaterial& Material, bool
 		&& !bMaterialModifiesMeshPosition
 		&& Material.WritesEveryPixel(false, bVFTypeSupportsNullPixelShader))
 	{
+		bShouldRender = true;
 		bUseDefaultMaterial = true;
 		bPositionOnly = true;
 	}
@@ -929,6 +932,8 @@ bool FDepthPassMeshProcessor::UseDefaultMaterial(const FMaterial& Material, bool
 		const bool bMaterialMasked = !Material.WritesEveryPixel(false, bVFTypeSupportsNullPixelShader) || Material.IsTranslucencyWritingCustomDepth();
 		if ((!bMaterialMasked && EarlyZPassMode != DDM_MaskedOnly) || (bMaterialMasked && EarlyZPassMode != DDM_NonMaskedOnly))
 		{
+			bShouldRender = true;
+
 			if (!bMaterialMasked && !bMaterialModifiesMeshPosition)
 			{
 				bUseDefaultMaterial = true;
@@ -937,7 +942,7 @@ bool FDepthPassMeshProcessor::UseDefaultMaterial(const FMaterial& Material, bool
 		}
 	}
 
-	return bUseDefaultMaterial;
+	return bShouldRender;
 }
 
 bool FDepthPassMeshProcessor::TryAddMeshBatch(const FMeshBatch& RESTRICT MeshBatch, uint64 BatchElementMask, const FPrimitiveSceneProxy* RESTRICT PrimitiveSceneProxy, int32 StaticMeshId, const FMaterialRenderProxy& MaterialRenderProxy, const FMaterial& Material)
@@ -956,29 +961,31 @@ bool FDepthPassMeshProcessor::TryAddMeshBatch(const FMeshBatch& RESTRICT MeshBat
 		const bool bEvaluateWPO = Material.MaterialModifiesMeshPosition_RenderThread()
 			&& (!ShouldOptimizedWPOAffectNonNaniteShaderSelection() || !PrimitiveSceneProxy || PrimitiveSceneProxy->EvaluateWorldPositionOffset());
 		bool bPositionOnly = false;
-		bool bUseDefaultMaterial = UseDefaultMaterial(Material, bEvaluateWPO, bSupportPositionOnlyStream, bVFTypeSupportsNullPixelShader, bPositionOnly);
-
-		const FMaterialRenderProxy* EffectiveMaterialRenderProxy = &MaterialRenderProxy;
-		const FMaterial* EffectiveMaterial = &Material;
-		if (bUseDefaultMaterial)
+		bool bUseDefaultMaterial = false;
+		if (ShouldRender(Material, bEvaluateWPO, bSupportPositionOnlyStream, bVFTypeSupportsNullPixelShader, bUseDefaultMaterial, bPositionOnly))
 		{
-			// Override with the default material
-			EffectiveMaterialRenderProxy = UMaterial::GetDefaultMaterial(MD_Surface)->GetRenderProxy();
-			EffectiveMaterial = EffectiveMaterialRenderProxy->GetMaterialNoFallback(FeatureLevel);
-			check(EffectiveMaterial);
-		}
+			const FMaterialRenderProxy* EffectiveMaterialRenderProxy = &MaterialRenderProxy;
+			const FMaterial* EffectiveMaterial = &Material;
+			if (bUseDefaultMaterial)
+			{
+				// Override with the default material
+				EffectiveMaterialRenderProxy = UMaterial::GetDefaultMaterial(MD_Surface)->GetRenderProxy();
+				EffectiveMaterial = EffectiveMaterialRenderProxy->GetMaterialNoFallback(FeatureLevel);
+				check(EffectiveMaterial);
+			}
 
-		const FMeshDrawingPolicyOverrideSettings OverrideSettings = ComputeMeshOverrideSettings(MeshBatch);
-		const ERasterizerFillMode MeshFillMode = ComputeMeshFillMode(Material, OverrideSettings);
-		const ERasterizerCullMode MeshCullMode = ComputeMeshCullMode(Material, OverrideSettings);
+			const FMeshDrawingPolicyOverrideSettings OverrideSettings = ComputeMeshOverrideSettings(MeshBatch);
+			const ERasterizerFillMode MeshFillMode = ComputeMeshFillMode(Material, OverrideSettings);
+			const ERasterizerCullMode MeshCullMode = ComputeMeshCullMode(Material, OverrideSettings);
 
-		if (bPositionOnly)
-		{
-			bResult = Process<true>(MeshBatch, BatchElementMask, StaticMeshId, PrimitiveSceneProxy, *EffectiveMaterialRenderProxy, *EffectiveMaterial, MeshFillMode, MeshCullMode);
-		}
-		else
-		{
-			bResult = Process<false>(MeshBatch, BatchElementMask, StaticMeshId, PrimitiveSceneProxy, *EffectiveMaterialRenderProxy, *EffectiveMaterial, MeshFillMode, MeshCullMode);
+			if (bPositionOnly)
+			{
+				bResult = Process<true>(MeshBatch, BatchElementMask, StaticMeshId, PrimitiveSceneProxy, *EffectiveMaterialRenderProxy, *EffectiveMaterial, MeshFillMode, MeshCullMode);
+			}
+			else
+			{
+				bResult = Process<false>(MeshBatch, BatchElementMask, StaticMeshId, PrimitiveSceneProxy, *EffectiveMaterialRenderProxy, *EffectiveMaterial, MeshFillMode, MeshCullMode);
+			}
 		}
 	}
 
@@ -1090,33 +1097,38 @@ void FDepthPassMeshProcessor::CollectPSOInitializers(const FSceneTexturesConfig&
 	const bool bSupportPositionOnlyStream = VertexFactoryData.VertexFactoryType->SupportsPositionOnly();  
 	const bool bVFTypeSupportsNullPixelShader = VertexFactoryData.VertexFactoryType->SupportsNullPixelShader();
 	bool bPositionOnly = false;
-	bool bUseDefaultMaterial = UseDefaultMaterial(Material, Material.MaterialModifiesMeshPosition_GameThread(), bSupportPositionOnlyStream, bVFTypeSupportsNullPixelShader, bPositionOnly);
-
-	const FMaterial* EffectiveMaterial = &Material;
-	if (bUseDefaultMaterial && !bSupportPositionOnlyStream && VertexFactoryData.CustomDefaultVertexDeclaration)
+	bool bUseDefaultMaterial = false;
+	if (ShouldRender(Material, Material.MaterialModifiesMeshPosition_GameThread(), bSupportPositionOnlyStream, bVFTypeSupportsNullPixelShader, bUseDefaultMaterial, bPositionOnly))
 	{
-		EMaterialQualityLevel::Type ActiveQualityLevel = GetCachedScalabilityCVars().MaterialQualityLevel;
-		EffectiveMaterial = UMaterial::GetDefaultMaterial(MD_Surface)->GetMaterialResource(FeatureLevel, ActiveQualityLevel);
-		bUseDefaultMaterial = false;
-	}
+		bool bCollectPSOs = !bUseDefaultMaterial;
 
-	if (!bUseDefaultMaterial)
-	{
-		check(!bPositionOnly);
+		// Collect PSOs for default material if there is a custom vertex declaration
+		const FMaterial* EffectiveMaterial = &Material;
+		if (bUseDefaultMaterial && !bSupportPositionOnlyStream && VertexFactoryData.CustomDefaultVertexDeclaration)
+		{
+			EMaterialQualityLevel::Type ActiveQualityLevel = GetCachedScalabilityCVars().MaterialQualityLevel;
+			EffectiveMaterial = UMaterial::GetDefaultMaterial(MD_Surface)->GetMaterialResource(FeatureLevel, ActiveQualityLevel);
+			bCollectPSOs = true;
+		}
 
-		const FMeshDrawingPolicyOverrideSettings OverrideSettings = ComputeMeshOverrideSettings(PreCacheParams);
-		const ERasterizerFillMode MeshFillMode = ComputeMeshFillMode(Material, OverrideSettings);
-		const ERasterizerCullMode MeshCullMode = ComputeMeshCullMode(Material, OverrideSettings);
+		if (bCollectPSOs)
+		{
+			check(!bPositionOnly);
 
-		bool bIsMoveable = PreCacheParams.IsMoveable();
-		const bool bAllowDitheredLODTransition = !bIsMoveable && Material.IsDitheredLODTransition();
+			const FMeshDrawingPolicyOverrideSettings OverrideSettings = ComputeMeshOverrideSettings(PreCacheParams);
+			const ERasterizerFillMode MeshFillMode = ComputeMeshFillMode(Material, OverrideSettings);
+			const ERasterizerCullMode MeshCullMode = ComputeMeshCullMode(Material, OverrideSettings);
 
-		bool bDitheredLODTransition = false;
-		CollectPSOInitializersInternal<false>(SceneTexturesConfig, VertexFactoryData, *EffectiveMaterial, MeshFillMode, MeshCullMode, bDitheredLODTransition, (EPrimitiveType)PreCacheParams.PrimitiveType, PSOInitializers);
-		if (bAllowDitheredLODTransition)
-		{
-			bDitheredLODTransition = true;
+			bool bIsMoveable = PreCacheParams.IsMoveable();
+			const bool bAllowDitheredLODTransition = !bIsMoveable && Material.IsDitheredLODTransition();
+
+			bool bDitheredLODTransition = false;
 			CollectPSOInitializersInternal<false>(SceneTexturesConfig, VertexFactoryData, *EffectiveMaterial, MeshFillMode, MeshCullMode, bDitheredLODTransition, (EPrimitiveType)PreCacheParams.PrimitiveType, PSOInitializers);
+			if (bAllowDitheredLODTransition)
+			{
+				bDitheredLODTransition = true;
+				CollectPSOInitializersInternal<false>(SceneTexturesConfig, VertexFactoryData, *EffectiveMaterial, MeshFillMode, MeshCullMode, bDitheredLODTransition, (EPrimitiveType)PreCacheParams.PrimitiveType, PSOInitializers);
+			}
 		}
 	}
 }
diff --git a/Engine/Source/Runtime/Renderer/Private/DepthRendering.h b/Engine/Source/Runtime/Renderer/Private/DepthRendering.h
--- a/Engine/Source/Runtime/Renderer/Private/DepthRendering.h
+++ b/Engine/Source/Runtime/Renderer/Private/DepthRendering.h
@@ -210,7 +210,7 @@ private:
 		ERasterizerFillMode MeshFillMode,
 		ERasterizerCullMode MeshCullMode);
 
-	bool UseDefaultMaterial(const FMaterial& Material, bool bMaterialModifiesMeshPosition, bool bSupportPositionOnlyStream, bool bVFTypeSupportsNullPixelShader, bool& bPositionOnly);
+	bool ShouldRender(const FMaterial& Material, bool bMaterialModifiesMeshPosition, bool bSupportPositionOnlyStream, bool bVFTypeSupportsNullPixelShader, bool& bUseDefaultMaterial, bool& bPositionOnly);
 
 	void CollectDefaultMaterialPSOInitializers(
 		const FSceneTexturesConfig& SceneTexturesConfig, 
