LFV - fixed culling and depth test flickering issues.

#changelist validated
#virtualized

diff --git a/Engine/Source/Runtime/Renderer/Private/LocalFogVolumeRendering.cpp b/Engine/Source/Runtime/Renderer/Private/LocalFogVolumeRendering.cpp
--- a/Engine/Source/Runtime/Renderer/Private/LocalFogVolumeRendering.cpp
+++ b/Engine/Source/Runtime/Renderer/Private/LocalFogVolumeRendering.cpp
@@ -137,7 +137,7 @@ bool ShouldRenderLocalFogVolumeInVolumetricFog(const FScene* Scene, const FScene
 
 float GetLocalFogVolumeGlobalStartDistance()
 {
-	return FMath::Max(0.0f, CVarLocalFogVolumeGlobalStartDistance.GetValueOnRenderThread());
+	return FMath::Max(10.0f, CVarLocalFogVolumeGlobalStartDistance.GetValueOnRenderThread());
 }
 
 bool IsLocalFogVolumeHalfResolution()
@@ -728,7 +728,7 @@ struct FDepthBoundSetup
 	float MaxDeviceZ = 1.0f;
 };
 
-static FDepthBoundSetup GetDepthBoundSetup(float FogStartDistance, FGraphicsPipelineStateInitializer& GraphicsPSOInit, FMatrix ViewProjectionMatrix, FMatrix ViewInvProjectionMatrix)
+static FDepthBoundSetup GetDepthBoundSetup(float FogStartDistance, FMatrix ViewProjectionMatrix, FMatrix ViewInvProjectionMatrix)
 {
 	FDepthBoundSetup DepthBoundSetup;
 
@@ -752,11 +752,11 @@ static FDepthBoundSetup GetDepthBoundSetup(float FogStartDistance, FGraphicsPipe
 	if (bool(ERHIZBuffer::IsInverted))
 	{
 		DepthBoundSetup.MinDeviceZ = 0.0f;
-		DepthBoundSetup.MaxDeviceZ = FogClipSpaceZ;
+		DepthBoundSetup.MaxDeviceZ = DepthBoundSetup.FogClipDeviceZ;
 	}
 	else
 	{
-		DepthBoundSetup.MinDeviceZ = FogClipSpaceZ;
+		DepthBoundSetup.MinDeviceZ = DepthBoundSetup.FogClipDeviceZ;
 		DepthBoundSetup.MaxDeviceZ = 1.0f;
 	}
 
@@ -785,11 +785,17 @@ void RenderLocalFogVolume(
 				continue;
 			}
 
+			const FIntRect ViewRect = View.ViewRect;
+			const FMatrix ViewProjectionMatrix = View.ViewMatrices.GetProjectionMatrix();
+			const FMatrix ViewInvProjectionMatrix = View.ViewMatrices.GetInvProjectionMatrix();
+			FDepthBoundSetup DepthBoundSetup = GetDepthBoundSetup(GetLocalFogVolumeGlobalStartDistance(), ViewProjectionMatrix, ViewInvProjectionMatrix);
+
 			FLocalFogVolumeTiledPassParameters* PassParameters = GraphBuilder.AllocParameters<FLocalFogVolumeTiledPassParameters>();
 
 			PassParameters->VS.View = GetShaderBinding(View.ViewUniformBuffer);
 			PassParameters->VS.LFV = View.LocalFogVolumeViewData.UniformParametersStruct;
 			PassParameters->VS.TileDataBuffer = View.LocalFogVolumeViewData.GPUTileDataBufferSRV;
+			PassParameters->VS.StartDepthZ = DepthBoundSetup.FogClipDeviceZ;
 
 			PassParameters->PS.View = GetShaderBinding(View.ViewUniformBuffer);
 			PassParameters->PS.LFV = View.LocalFogVolumeViewData.UniformParametersStruct;
@@ -806,10 +812,6 @@ void RenderLocalFogVolume(
 			FLocalFogVolumeTiledRenderPS::FPermutationDomain PsPermutationVector;
 			auto PixelShader = View.ShaderMap->GetShader< FLocalFogVolumeTiledRenderPS >(PsPermutationVector);
 
-			const FIntRect ViewRect = View.ViewRect;
-			const FMatrix ViewProjectionMatrix = View.ViewMatrices.GetProjectionMatrix();
-			const FMatrix ViewInvProjectionMatrix = View.ViewMatrices.GetInvProjectionMatrix();
-
 			ClearUnusedGraphResources(VertexShader, &PassParameters->VS);
 			ClearUnusedGraphResources(PixelShader, &PassParameters->PS);
 
@@ -819,7 +821,7 @@ void RenderLocalFogVolume(
 				RDG_EVENT_NAME("LocalFogVolume.Tiled (%u X %u)", LocalFogVolumeTileDataTextureResolution.X, LocalFogVolumeTileDataTextureResolution.Y),
 				PassParameters,
 				ERDGPassFlags::Raster,
-				[VertexShader, PixelShader, PassParameters, ViewRect, ViewProjectionMatrix, ViewInvProjectionMatrix](FRHICommandList& RHICmdList)
+				[VertexShader, PixelShader, PassParameters, ViewRect, DepthBoundSetup](FRHICommandList& RHICmdList)
 			{
 				FGraphicsPipelineStateInitializer GraphicsPSOInit;
 				RHICmdList.ApplyCachedRenderTargets(GraphicsPSOInit);
@@ -836,7 +838,6 @@ void RenderLocalFogVolume(
 				GraphicsPSOInit.BoundShaderState.VertexShaderRHI = VertexShader.GetVertexShader();
 				GraphicsPSOInit.BoundShaderState.PixelShaderRHI = PixelShader.GetPixelShader();
 
-				FDepthBoundSetup DepthBoundSetup = GetDepthBoundSetup(GetLocalFogVolumeGlobalStartDistance(), GraphicsPSOInit, ViewProjectionMatrix, ViewInvProjectionMatrix);
 				GraphicsPSOInit.bDepthBounds = DepthBoundSetup.bEnabled;
 
 				SetGraphicsPipelineState(RHICmdList, GraphicsPSOInit, 0);
@@ -996,7 +997,7 @@ void RenderLocalFogVolumeMobile(
 	GraphicsPSOInit.BoundShaderState.VertexShaderRHI = VertexShader.GetVertexShader();
 	GraphicsPSOInit.BoundShaderState.PixelShaderRHI = PixelShader.GetPixelShader();
 
-	FDepthBoundSetup DepthBoundSetup = GetDepthBoundSetup(GetLocalFogVolumeGlobalStartDistance(), GraphicsPSOInit, View.ViewMatrices.GetProjectionMatrix(), View.ViewMatrices.GetInvProjectionMatrix());
+	FDepthBoundSetup DepthBoundSetup = GetDepthBoundSetup(GetLocalFogVolumeGlobalStartDistance(), View.ViewMatrices.GetProjectionMatrix(), View.ViewMatrices.GetInvProjectionMatrix());
 	GraphicsPSOInit.bDepthBounds = DepthBoundSetup.bEnabled;
 
 	SetGraphicsPipelineState(RHICmdList, GraphicsPSOInit, 0);
diff --git a/Engine/Source/Runtime/Renderer/Private/LocalFogVolumeRendering.h b/Engine/Source/Runtime/Renderer/Private/LocalFogVolumeRendering.h
--- a/Engine/Source/Runtime/Renderer/Private/LocalFogVolumeRendering.h
+++ b/Engine/Source/Runtime/Renderer/Private/LocalFogVolumeRendering.h
@@ -51,7 +51,7 @@ public:
 
 	float GetUniformScale() 
 	{
-		union { float F; uint32 U; } FU = { Data0[3] };
+		union { uint32 U; float F; } FU = { Data0[3] };
 		return FU.F;
 	}
 };
