Provide depth RT for draw tile mesh renders on mobile.
 - workaround for an odd driver crash seen on some android adreno platforms.

#changelist validated
#virtualized

diff --git a/Engine/Source/Runtime/Renderer/Private/Renderer.cpp b/Engine/Source/Runtime/Renderer/Private/Renderer.cpp
--- a/Engine/Source/Runtime/Renderer/Private/Renderer.cpp
+++ b/Engine/Source/Runtime/Renderer/Private/Renderer.cpp
@@ -359,6 +359,14 @@ void FRendererModule::DrawTileMesh(FCanvasRenderContext& RenderContext, FMeshPas
 		PassParameters->InstanceCullingDrawParams.InstanceCulling = FInstanceCullingContext::CreateDummyInstanceCullingUniformBuffer(GraphBuilder);
 		PassParameters->ReflectionCapture = EmptyReflectionCaptureUniformBuffer;
 
+		if (ShadingPath == EShadingPath::Mobile && MobileRequiresSceneDepthAux(ViewFamily->GetShaderPlatform()))
+		{
+			FSceneTexturesConfig& Config = ViewFamily->SceneTexturesConfig;
+			FRDGTextureDesc DepthAuxDesc = FRDGTextureDesc::Create2D(Config.Extent, PF_R16F, FClearValueBinding(FLinearColor::Transparent), TexCreate_RenderTargetable | TexCreate_ShaderResource | TexCreate_InputAttachmentRead);
+			FRDGTextureRef DepthAux = GraphBuilder.CreateTexture(DepthAuxDesc, TEXT("SceneDepthAux"));
+			PassParameters->RenderTargets[1] = FRenderTargetBinding(DepthAux, ERenderTargetLoadAction::EClear);
+		}
+
 		// handle translucent material blend modes, not relevant in MaterialTexCoordScalesAnalysis since it outputs the scales.
 		if (ViewFamily->GetDebugViewShaderMode() == DVSM_OutputMaterialTextureScales)
 		{
