Remove unused cvar from previous commit

#changelist validated
#virtualized

diff --git a/Engine/Source/Runtime/Renderer/Private/VirtualShadowMaps/VirtualShadowMapCacheManager.cpp b/Engine/Source/Runtime/Renderer/Private/VirtualShadowMaps/VirtualShadowMapCacheManager.cpp
--- a/Engine/Source/Runtime/Renderer/Private/VirtualShadowMaps/VirtualShadowMapCacheManager.cpp
+++ b/Engine/Source/Runtime/Renderer/Private/VirtualShadowMaps/VirtualShadowMapCacheManager.cpp
@@ -78,13 +78,6 @@ FAutoConsoleVariableRef CVarVSMMaxPageAgeSinceLastRequest(
 	ECVF_RenderThreadSafe
 );
 
-static TAutoConsoleVariable<int32> CVarMaxInstanceTransitionsPerFrame(
-	TEXT("r.Shadow.Virtual.Cache.StaticSeparate.MaxInstanceTransitionsPerFrame"),
-	50000,
-	TEXT("Maximum number of static<->dynamic instance transitions per frame.\n")
-	TEXT("These transitions cause invalidations but setting this too low will result in scene transitions taking too long, causing too many redundant transitions."),
-	ECVF_RenderThreadSafe);
-
 static TAutoConsoleVariable<int32> CVarFramesStaticThreshold(
 	TEXT("r.Shadow.Virtual.Cache.StaticSeparate.FramesStaticThreshold"),
 	100,
