Yet another fix for a fix of Hair view offset (argh).

- The third time is the charm.

#changelist validated
#virtualized

diff --git a/Engine/Shaders/Private/HairStrands/HairStrandsVisibilityTile.usf b/Engine/Shaders/Private/HairStrands/HairStrandsVisibilityTile.usf
--- a/Engine/Shaders/Private/HairStrands/HairStrandsVisibilityTile.usf
+++ b/Engine/Shaders/Private/HairStrands/HairStrandsVisibilityTile.usf
@@ -394,6 +394,7 @@ void MainVS(
 uint TileSize;
 uint TileType;
 int2 TileCountXY;
+int2 ViewRectMin;
 int2 Resolution;
 Buffer<uint>  TileCountBuffer;
 Buffer<uint2> TileDataBuffer;
@@ -410,7 +411,7 @@ void TileMainCS(uint2 InGroupId : SV_GroupID, uint2 InGroupThreadId : SV_GroupTh
 		return;
 	}
 	const uint2 GroupId = TileDataBuffer[TileIndex1D];
-	const uint2 PixelCoord = uint2(View.ViewRectMin.xy) + GroupId * TileSize + InGroupThreadId;
+	const uint2 PixelCoord = uint2(ViewRectMin) + GroupId * TileSize + InGroupThreadId;
 
 	if (all(PixelCoord < uint2(Resolution)))
 	{
diff --git a/Engine/Source/Runtime/Renderer/Private/HairStrands/HairStrandsTile.cpp b/Engine/Source/Runtime/Renderer/Private/HairStrands/HairStrandsTile.cpp
--- a/Engine/Source/Runtime/Renderer/Private/HairStrands/HairStrandsTile.cpp
+++ b/Engine/Source/Runtime/Renderer/Private/HairStrands/HairStrandsTile.cpp
@@ -441,6 +441,7 @@ class FHairTileClearCS : public FGlobalShader
 	BEGIN_SHADER_PARAMETER_STRUCT(FParameters, )
 		SHADER_PARAMETER(FIntPoint, Resolution)
 		SHADER_PARAMETER(FIntPoint, TileCountXY)
+		SHADER_PARAMETER(FIntPoint, ViewRectMin)
 		SHADER_PARAMETER(uint32, TileSize)
 		SHADER_PARAMETER(uint32, TileType)
 		SHADER_PARAMETER_RDG_BUFFER_SRV(Buffer<uint2>, TileDataBuffer)
@@ -481,6 +482,7 @@ void AddHairStrandsTileClearPass(
 	FHairTileClearCS::FParameters* PassParameters = GraphBuilder.AllocParameters<FHairTileClearCS::FParameters>();
 	PassParameters->Resolution = OutTexture->Desc.Extent;
 	PassParameters->TileCountXY = TileData.TileCountXY;
+	PassParameters->ViewRectMin = View.ViewRect.Min;
 	PassParameters->TileSize = TileData.TileSize;
 	PassParameters->TileType = ToIndex(TileType);
 	PassParameters->TileCountBuffer = TileData.TileCountSRV;
