Heterogeneous Volumes: Add GuildBuild options to enable or disable specific grids. Disable frustum grid when building for shadows.

#changelist validated
#virtualized

diff --git a/Engine/Source/Runtime/Renderer/Private/HeterogeneousVolumes/HeterogeneousVolumes.cpp b/Engine/Source/Runtime/Renderer/Private/HeterogeneousVolumes/HeterogeneousVolumes.cpp
--- a/Engine/Source/Runtime/Renderer/Private/HeterogeneousVolumes/HeterogeneousVolumes.cpp
+++ b/Engine/Source/Runtime/Renderer/Private/HeterogeneousVolumes/HeterogeneousVolumes.cpp
@@ -507,8 +507,10 @@ void FDeferredShadingSceneRenderer::RenderHeterogeneousVolumeShadows(
 	SCOPED_NAMED_EVENT(HeterogeneousVolumes, FColor::Emerald);
 
 	FVoxelGridBuildOptions BuildOptions;
-	BuildOptions.bJitter = false;
 	BuildOptions.MinimumVoxelSizeOutsideFrustum = 10.0;
+	BuildOptions.bBuildOrthoGrid = true;
+	BuildOptions.bBuildFrustumGrid = false;
+	BuildOptions.bJitter = false;
 
 	TRDGUniformBufferRef<FOrthoVoxelGridUniformBufferParameters> OrthoGridUniformBuffer = nullptr;
 	TRDGUniformBufferRef<FFrustumVoxelGridUniformBufferParameters> FrustumGridUniformBuffer = nullptr;
diff --git a/Engine/Source/Runtime/Renderer/Private/HeterogeneousVolumes/HeterogeneousVolumes.h b/Engine/Source/Runtime/Renderer/Private/HeterogeneousVolumes/HeterogeneousVolumes.h
--- a/Engine/Source/Runtime/Renderer/Private/HeterogeneousVolumes/HeterogeneousVolumes.h
+++ b/Engine/Source/Runtime/Renderer/Private/HeterogeneousVolumes/HeterogeneousVolumes.h
@@ -207,8 +207,11 @@ END_UNIFORM_BUFFER_STRUCT()
 
 struct FVoxelGridBuildOptions
 {
-	bool bJitter = HeterogeneousVolumes::ShouldJitter();
 	float MinimumVoxelSizeOutsideFrustum = HeterogeneousVolumes::GetMinimumVoxelSizeOutsideFrustum();
+
+	bool bBuildOrthoGrid = true;
+	bool bBuildFrustumGrid = true;
+	bool bJitter = HeterogeneousVolumes::ShouldJitter();
 };
 
 void BuildOrthoVoxelGrid(
diff --git a/Engine/Source/Runtime/Renderer/Private/HeterogeneousVolumes/HeterogeneousVolumesVoxelGridPipeline.cpp b/Engine/Source/Runtime/Renderer/Private/HeterogeneousVolumes/HeterogeneousVolumesVoxelGridPipeline.cpp
--- a/Engine/Source/Runtime/Renderer/Private/HeterogeneousVolumes/HeterogeneousVolumesVoxelGridPipeline.cpp
+++ b/Engine/Source/Runtime/Renderer/Private/HeterogeneousVolumes/HeterogeneousVolumesVoxelGridPipeline.cpp
@@ -2045,7 +2045,7 @@ void BuildFrustumVoxelGrid(
 	TRDGUniformBufferRef<FFrustumVoxelGridUniformBufferParameters>& FrustumGridUniformBuffer
 )
 {
-	if (!ShouldRenderHeterogeneousVolumesForView(View) || !HeterogeneousVolumes::EnableFrustumVoxelGrid())
+	if (!ShouldRenderHeterogeneousVolumesForView(View) || !HeterogeneousVolumes::EnableFrustumVoxelGrid() || !BuildOptions.bBuildFrustumGrid)
 	{
 		FrustumGridUniformBuffer = CreateEmptyFrustumVoxelGridUniformBuffer(GraphBuilder);
 		return;
@@ -2612,7 +2612,7 @@ void BuildOrthoVoxelGrid(
 	TRDGUniformBufferRef<FOrthoVoxelGridUniformBufferParameters>& OrthoGridUniformBuffer
 )
 {
-	if (!ShouldRenderHeterogeneousVolumes(Scene) || !HeterogeneousVolumes::EnableOrthoVoxelGrid())
+	if (!ShouldRenderHeterogeneousVolumes(Scene) || !HeterogeneousVolumes::EnableOrthoVoxelGrid() || !BuildOptions.bBuildOrthoGrid)
 	{
 		OrthoGridUniformBuffer = CreateEmptyOrthoVoxelGridUniformBuffer(GraphBuilder);
 		return;
