Heterogeneous Volumes: Apply bounds culling to trim active voxels when computing the lighting cache.

#changelist validated
#virtualized

diff --git a/Engine/Shaders/Private/HeterogeneousVolumes/HeterogeneousVolumesLiveShadingPipeline.usf b/Engine/Shaders/Private/HeterogeneousVolumes/HeterogeneousVolumesLiveShadingPipeline.usf
--- a/Engine/Shaders/Private/HeterogeneousVolumes/HeterogeneousVolumesLiveShadingPipeline.usf
+++ b/Engine/Shaders/Private/HeterogeneousVolumes/HeterogeneousVolumesLiveShadingPipeline.usf
@@ -42,6 +42,8 @@ int PrimitiveId;
 
 // Volume data
 int3 VoxelResolution;
+int3 VoxelMin;
+int3 VoxelMax;
 
 // Lighting
 int bApplyEmissionAndTransmittance;
@@ -229,8 +231,8 @@ void RenderLightingCacheWithLiveShadingCS(
 	uint3 DispatchThreadId : SV_DispatchThreadID
 )
 {
-	float3 VoxelIndex = DispatchThreadId;
-	if (any(DispatchThreadId > LightingCacheResolution))
+	float3 VoxelIndex = DispatchThreadId + VoxelMin;
+	if (any(DispatchThreadId > VoxelMax))
 	{
 		return;
 	}
diff --git a/Engine/Source/Runtime/Renderer/Private/HeterogeneousVolumes/HeterogeneousVolumesLiveShadingPipeline.cpp b/Engine/Source/Runtime/Renderer/Private/HeterogeneousVolumes/HeterogeneousVolumesLiveShadingPipeline.cpp
--- a/Engine/Source/Runtime/Renderer/Private/HeterogeneousVolumes/HeterogeneousVolumesLiveShadingPipeline.cpp
+++ b/Engine/Source/Runtime/Renderer/Private/HeterogeneousVolumes/HeterogeneousVolumesLiveShadingPipeline.cpp
@@ -16,6 +16,21 @@
 #include "VolumeLighting.h"
 #include "VolumetricFog.h"
 
+static TAutoConsoleVariable<int32> CVarHeterogeneousLightingCacheBoundsCulling(
+	TEXT("r.HeterogeneousVolumes.LightingCache.BoundsCulling"),
+	1,
+	TEXT("Enables bounds culling when populating the lighting cache (Default = 1)"),
+	ECVF_RenderThreadSafe
+);
+
+namespace HeterogeneousVolumes
+{
+	bool ShouldBoundsCull()
+	{
+		return CVarHeterogeneousLightingCacheBoundsCulling.GetValueOnRenderThread() != 0;
+	}
+}
+
 //-OPT: Remove duplicate bindings
 // At the moment we need to bind the mesh draw parameters as they will be applied and on some RHIs this will crash if the texture is nullptr
 // We have the same parameters in the loose FParameters shader structure that are applied after the mesh draw.
@@ -139,6 +154,8 @@ class FRenderLightingCacheWithLiveShadingCS : public FMeshMaterialShader
 		// Volume data
 		SHADER_PARAMETER(FIntVector, VoxelResolution)
 		SHADER_PARAMETER_STRUCT_INCLUDE(FLightingCacheParameters, LightingCache)
+		SHADER_PARAMETER(FIntVector, VoxelMin)
+		SHADER_PARAMETER(FIntVector, VoxelMax)
 
 		// Output
 		SHADER_PARAMETER_RDG_TEXTURE_UAV(RWTexture3D<float>, RWLightingCacheTexture)
@@ -506,7 +523,59 @@ static void RenderLightingCacheWithLiveShading(
 	}
 #endif // WANTS_DRAW_MESH_EVENTS
 
-	FIntVector GroupCount = HeterogeneousVolumes::GetLightingCacheResolution(HeterogeneousVolumeInterface);
+	PassParameters->VoxelMin = FIntVector::ZeroValue;
+	PassParameters->VoxelMax = HeterogeneousVolumes::GetLightingCacheResolution(HeterogeneousVolumeInterface) - FIntVector(1);
+	
+	bool bShouldBoundsCull = HeterogeneousVolumes::ShouldBoundsCull();
+	if (LightType != LightType_Directional && bShouldBoundsCull)
+	{
+		auto ToFVector3f = [](FVector4 V4f) {
+			return FVector3f(V4f.X, V4f.Y, V4f.Z);
+		};
+
+		auto FloorVector = [](const FVector& V) {
+			return FVector(
+				FMath::FloorToFloat(V.X),
+				FMath::FloorToFloat(V.Y),
+				FMath::FloorToFloat(V.Z)
+			);
+		};
+
+		auto CeilVector = [](const FVector& V) {
+			return FVector(
+				FMath::CeilToFloat(V.X),
+				FMath::CeilToFloat(V.Y),
+				FMath::CeilToFloat(V.Z)
+			);
+		};
+
+		auto ClampVector = [](const FVector& V, const FIntVector& Min, const FIntVector& Max) {
+			FIntVector IntV;
+			IntV.X = FMath::Clamp(V.X, Min.X, Max.X);
+			IntV.Y = FMath::Clamp(V.Y, Min.Y, Max.Y);
+			IntV.Z = FMath::Clamp(V.Z, Min.Z, Max.Z);
+			return IntV;
+		};
+
+		FSphere WorldLightBoundingSphere = LightSceneInfo->Proxy->GetBoundingSphere();
+		FVector LocalLightCenter = FVector(PassParameters->WorldToLocal.TransformPosition(ToFVector3f(WorldLightBoundingSphere.Center)));
+		FVector3f ScalingTerm = PassParameters->WorldToLocal.GetScaleVector();
+		FVector LocalLightExtent = FVector(ScalingTerm) * WorldLightBoundingSphere.W;
+		FVector LocalLightMin = LocalLightCenter - LocalLightExtent;
+		FVector LocalLightMax = LocalLightCenter + LocalLightExtent;
+
+		FVector LightingCacheMin = LocalBoxSphereBounds.Origin - LocalBoxSphereBounds.BoxExtent;
+		FVector LightingCacheMax = LocalBoxSphereBounds.Origin + LocalBoxSphereBounds.BoxExtent;
+
+		FVector LocalLightMinUV = (LocalLightMin - LightingCacheMin) / (LightingCacheMax - LightingCacheMin);
+		FVector LocalLightMaxUV = (LocalLightMax - LightingCacheMin) / (LightingCacheMax - LightingCacheMin);
+		FVector LightingCacheResolution = FVector(PassParameters->LightingCache.LightingCacheResolution);
+		PassParameters->VoxelMin = ClampVector(FloorVector(LocalLightMinUV * LightingCacheResolution), FIntVector::ZeroValue, PassParameters->VoxelMax);
+		PassParameters->VoxelMax = ClampVector(CeilVector(LocalLightMaxUV * LightingCacheResolution), FIntVector::ZeroValue, PassParameters->VoxelMax);
+	}
+
+	FIntVector GroupCount = PassParameters->VoxelMax - PassParameters->VoxelMin + FIntVector(1);
+	check(GroupCount.X > 0 && GroupCount.Y > 0 && GroupCount.Z > 0);
 	GroupCount.X = FMath::DivideAndRoundUp(GroupCount.X, FRenderLightingCacheWithLiveShadingCS::GetThreadGroupSize3D());
 	GroupCount.Y = FMath::DivideAndRoundUp(GroupCount.Y, FRenderLightingCacheWithLiveShadingCS::GetThreadGroupSize3D());
 	GroupCount.Z = FMath::DivideAndRoundUp(GroupCount.Z, FRenderLightingCacheWithLiveShadingCS::GetThreadGroupSize3D());
