Change ShaderPrint underlying buffer to use typed buffer in order to address SPIR-V translation issue on certain platforms.

diff --git a/Engine/Plugins/Animation/DeformerGraph/Shaders/Private/DataInterfaceDebugDraw.ush b/Engine/Plugins/Animation/DeformerGraph/Shaders/Private/DataInterfaceDebugDraw.ush
--- a/Engine/Plugins/Animation/DeformerGraph/Shaders/Private/DataInterfaceDebugDraw.ush
+++ b/Engine/Plugins/Animation/DeformerGraph/Shaders/Private/DataInterfaceDebugDraw.ush
@@ -53,8 +53,8 @@ uint {DataInterfaceName}_MaxSymbolCount;
 uint {DataInterfaceName}_MaxStateCount;
 uint {DataInterfaceName}_MaxLineCount;
 uint {DataInterfaceName}_MaxTriangleCount;
-StructuredBuffer<uint> {DataInterfaceName}_StateBuffer;
-RWStructuredBuffer<uint> {DataInterfaceName}_RWEntryBuffer;
+Buffer<uint> {DataInterfaceName}_StateBuffer;
+RWBuffer<uint> {DataInterfaceName}_RWEntryBuffer;
 
 FDebugDraw ReadDebugDraw_{DataInterfaceName}()
 {
diff --git a/Engine/Plugins/Animation/DeformerGraph/Source/OptimusCore/Private/DataInterfaces/OptimusDataInterfaceDebugDraw.cpp b/Engine/Plugins/Animation/DeformerGraph/Source/OptimusCore/Private/DataInterfaces/OptimusDataInterfaceDebugDraw.cpp
--- a/Engine/Plugins/Animation/DeformerGraph/Source/OptimusCore/Private/DataInterfaces/OptimusDataInterfaceDebugDraw.cpp
+++ b/Engine/Plugins/Animation/DeformerGraph/Source/OptimusCore/Private/DataInterfaces/OptimusDataInterfaceDebugDraw.cpp
@@ -75,8 +75,8 @@ BEGIN_SHADER_PARAMETER_STRUCT(FDebugDrawDataInterfaceParameters, )
 	SHADER_PARAMETER(uint32, MaxStateCount)
 	SHADER_PARAMETER(uint32, MaxLineCount)
 	SHADER_PARAMETER(uint32, MaxTriangleCount)
-	SHADER_PARAMETER_RDG_BUFFER_SRV(StructuredBuffer<uint2>, StateBuffer)
-	SHADER_PARAMETER_RDG_BUFFER_UAV(RWStructuredBuffer<ShaderPrintItem>, RWEntryBuffer)
+	SHADER_PARAMETER_RDG_BUFFER_SRV(Buffer<uint>, StateBuffer)
+	SHADER_PARAMETER_RDG_BUFFER_UAV(RWBuffer<uint>, RWEntryBuffer)
 END_SHADER_PARAMETER_STRUCT()
 
 void UOptimusDebugDrawDataInterface::GetShaderParameters(TCHAR const* UID, FShaderParametersMetadataBuilder& InOutBuilder, FShaderParametersMetadataAllocations& InOutAllocations) const
diff --git a/Engine/Shaders/Private/ShaderPrint.ush b/Engine/Shaders/Private/ShaderPrint.ush
--- a/Engine/Shaders/Private/ShaderPrint.ush
+++ b/Engine/Shaders/Private/ShaderPrint.ush
@@ -42,14 +42,14 @@
 	// When using Diagnostic buffer as output, the regular shader entry output won't be used, and the pass will be culled. 
 	// To prevent this we add a dummy write, so that the pass does get culled
 	#if SHADER_PRINT_EXPLICIT_BINDING
-	RWStructuredBuffer<uint> ShaderPrint_RWEntryBuffer;
+	RWBuffer<uint> ShaderPrint_RWEntryBuffer;
 	#endif
 
 	#define SHADER_PRINT_RWENTRYBUFFER(Ctx, Idx) UEDiagnosticBuffer[SHADER_PRINT_DIAGNOSTIC_BUFFER_ENTRY_OFFSET + (Idx)]
 	#define SHADER_PRINT_RWSTATEBUFFER(Ctx, Idx) UEDiagnosticBuffer[SHADER_PRINT_DIAGNOSTIC_BUFFER_ENTRY_OFFSET + (Idx)]
 #else
-	StructuredBuffer<uint> ShaderPrint_StateBuffer;
-	RWStructuredBuffer<uint> ShaderPrint_RWEntryBuffer;
+	Buffer<uint> ShaderPrint_StateBuffer;
+	RWBuffer<uint> ShaderPrint_RWEntryBuffer;
 
 	// Use global buffers directly instead of using FShaderPrintContext::Buffers.
 	#if SHADER_PRINT_USE_GLOBAL_RESOURCE
diff --git a/Engine/Shaders/Private/ShaderPrintCommon.ush b/Engine/Shaders/Private/ShaderPrintCommon.ush
--- a/Engine/Shaders/Private/ShaderPrintCommon.ush
+++ b/Engine/Shaders/Private/ShaderPrintCommon.ush
@@ -331,8 +331,8 @@ struct FShaderPrintConfig
 
 struct FShaderPrintBuffers
 {
-	StructuredBuffer<uint> StateBuffer;
-	RWStructuredBuffer<uint> RWEntryBuffer;
+	Buffer<uint> StateBuffer;
+	RWBuffer<uint> RWEntryBuffer;
 };
 
 struct FShaderPrintContext
diff --git a/Engine/Shaders/Private/ShaderPrintDraw.usf b/Engine/Shaders/Private/ShaderPrintDraw.usf
--- a/Engine/Shaders/Private/ShaderPrintDraw.usf
+++ b/Engine/Shaders/Private/ShaderPrintDraw.usf
@@ -14,12 +14,12 @@ uint2 ViewportResolution;
 uint FrameIndex;
 uint FrameThreshold;
 
-StructuredBuffer<uint> ValuesBuffer;
+Buffer<uint> ValuesBuffer;
 Buffer<uint> SymbolsBuffer;
 
-RWStructuredBuffer<uint> RWValuesBuffer;
+RWBuffer<uint> RWValuesBuffer;
 RWBuffer<uint> RWSymbolsBuffer;
-RWStructuredBuffer<uint> RWStateBuffer;
+RWBuffer<uint> RWStateBuffer;
 
 RWBuffer<uint> RWIndirectDispatchArgsBuffer;
 RWBuffer<uint> RWIndirectDrawArgsBuffer;
diff --git a/Engine/Source/Runtime/Renderer/Private/ShaderPrint.cpp b/Engine/Source/Runtime/Renderer/Private/ShaderPrint.cpp
--- a/Engine/Source/Runtime/Renderer/Private/ShaderPrint.cpp
+++ b/Engine/Source/Runtime/Renderer/Private/ShaderPrint.cpp
@@ -270,8 +270,8 @@ namespace ShaderPrint
 	void SetParameters(FRDGBuilder& GraphBuilder, const FShaderPrintData& InData, FShaderParameters& OutParameters)
 	{
 		OutParameters.Common = InData.UniformBuffer;
-		OutParameters.ShaderPrint_StateBuffer = GraphBuilder.CreateSRV(InData.ShaderPrintStateBuffer);
-		OutParameters.ShaderPrint_RWEntryBuffer = GraphBuilder.CreateUAV(InData.ShaderPrintEntryBuffer);
+		OutParameters.ShaderPrint_StateBuffer = GraphBuilder.CreateSRV(InData.ShaderPrintStateBuffer, PF_R32_UINT);
+		OutParameters.ShaderPrint_RWEntryBuffer = GraphBuilder.CreateUAV(InData.ShaderPrintEntryBuffer, PF_R32_UINT);
 	}
 
 	void SetParameters(FRDGBuilder& GraphBuilder, FShaderParameters& OutParameters)
@@ -411,7 +411,7 @@ namespace ShaderPrint
 
 		BEGIN_SHADER_PARAMETER_STRUCT(FParameters, )
 			SHADER_PARAMETER_STRUCT_REF(FShaderPrintCommonParameters, Common)
-			SHADER_PARAMETER_RDG_BUFFER_UAV(RWStructuredBuffer<uint>, RWValuesBuffer)
+			SHADER_PARAMETER_RDG_BUFFER_UAV(RWBuffer<uint>, RWValuesBuffer)
 		END_SHADER_PARAMETER_STRUCT()
 
 		static bool ShouldCompilePermutation(FGlobalShaderPermutationParameters const& Parameters)
@@ -438,7 +438,7 @@ namespace ShaderPrint
 		SHADER_USE_PARAMETER_STRUCT(FShaderPrintClearCounterCS, FGlobalShader);
 
 		BEGIN_SHADER_PARAMETER_STRUCT(FParameters, )
-			SHADER_PARAMETER_RDG_BUFFER_UAV(RWStructuredBuffer<uint>, RWValuesBuffer)
+			SHADER_PARAMETER_RDG_BUFFER_UAV(RWBuffer<uint>, RWValuesBuffer)
 		END_SHADER_PARAMETER_STRUCT()
 
 		static bool ShouldCompilePermutation(FGlobalShaderPermutationParameters const& Parameters)
@@ -458,7 +458,7 @@ namespace ShaderPrint
 
 		BEGIN_SHADER_PARAMETER_STRUCT(FParameters, )
 			SHADER_PARAMETER_STRUCT_REF(FShaderPrintCommonParameters, Common)
-			SHADER_PARAMETER_RDG_BUFFER_SRV(StructuredBuffer<uint>, ValuesBuffer)
+			SHADER_PARAMETER_RDG_BUFFER_SRV(Buffer<uint>, ValuesBuffer)
 			SHADER_PARAMETER_RDG_BUFFER_UAV(RWBuffer<uint>, RWSymbolsBuffer)
 			SHADER_PARAMETER_RDG_BUFFER_UAV(RWBuffer<uint>, RWIndirectDispatchArgsBuffer)
 		END_SHADER_PARAMETER_STRUCT()
@@ -482,7 +482,7 @@ namespace ShaderPrint
 			SHADER_PARAMETER(uint32, FrameIndex)
 			SHADER_PARAMETER(uint32, FrameThreshold)
 			SHADER_PARAMETER_STRUCT_REF(FShaderPrintCommonParameters, Common)
-			SHADER_PARAMETER_RDG_BUFFER_UAV(RWStructuredBuffer<uint>, RWStateBuffer)
+			SHADER_PARAMETER_RDG_BUFFER_UAV(RWBuffer<uint>, RWStateBuffer)
 		END_SHADER_PARAMETER_STRUCT()
 
 		static bool ShouldCompilePermutation(FGlobalShaderPermutationParameters const& Parameters)
@@ -503,9 +503,9 @@ namespace ShaderPrint
 		BEGIN_SHADER_PARAMETER_STRUCT(FParameters, )
 			SHADER_PARAMETER(uint32, FrameIndex)
 			SHADER_PARAMETER_STRUCT_REF(FShaderPrintCommonParameters, Common)
-			SHADER_PARAMETER_RDG_BUFFER_SRV(StructuredBuffer<uint>, ValuesBuffer)
+			SHADER_PARAMETER_RDG_BUFFER_SRV(Buffer<uint>, ValuesBuffer)
 			SHADER_PARAMETER_RDG_BUFFER_UAV(RWBuffer<uint>, RWSymbolsBuffer)
-			SHADER_PARAMETER_RDG_BUFFER_UAV(RWStructuredBuffer<uint>, RWStateBuffer)
+			SHADER_PARAMETER_RDG_BUFFER_UAV(RWBuffer<uint>, RWStateBuffer)
 			RDG_BUFFER_ACCESS(IndirectDispatchArgsBuffer, ERHIAccess::IndirectArgs)
 		END_SHADER_PARAMETER_STRUCT()
 
@@ -789,7 +789,7 @@ namespace ShaderPrint
 			ERDGBufferFlags Flags = bLockBufferThisFrame ? ERDGBufferFlags::MultiFrame : ERDGBufferFlags::None;
 
 			const uint32 UintElementCount = GetRequestedEntryBufferSizeInUint(InSetup);
-			ShaderPrintData.ShaderPrintEntryBuffer = GraphBuilder.CreateBuffer(FRDGBufferDesc::CreateStructuredDesc(4, UintElementCount), TEXT("ShaderPrint.EntryBuffer"), Flags);
+			ShaderPrintData.ShaderPrintEntryBuffer = GraphBuilder.CreateBuffer(FRDGBufferDesc::CreateBufferDesc(4, UintElementCount), TEXT("ShaderPrint.EntryBuffer"), Flags);
 
 			// State buffer is retrieved from the view state, or created if it does not exist
 			if (InViewState != nullptr)
@@ -800,7 +800,7 @@ namespace ShaderPrint
 				}
 				else
 				{
-					ShaderPrintData.ShaderPrintStateBuffer = GraphBuilder.CreateBuffer(FRDGBufferDesc::CreateStructuredDesc(4, (3 * InSetup.MaxStateCount) + 1), TEXT("ShaderPrint.StateBuffer"));
+					ShaderPrintData.ShaderPrintStateBuffer = GraphBuilder.CreateBuffer(FRDGBufferDesc::CreateBufferDesc(4, (3 * InSetup.MaxStateCount) + 1), TEXT("ShaderPrint.StateBuffer"));
 					AddClearUAVPass(GraphBuilder, GraphBuilder.CreateUAV(ShaderPrintData.ShaderPrintStateBuffer, PF_R32_UINT), 0u);
 					InViewState->ShaderPrintStateData.StateBuffer = GraphBuilder.ConvertToExternalBuffer(ShaderPrintData.ShaderPrintStateBuffer);
 				}
@@ -816,7 +816,7 @@ namespace ShaderPrint
 				TShaderMapRef<FShaderPrintClearCounterCS> ComputeShader(GlobalShaderMap);
 
 				FShaderPrintClearCounterCS::FParameters* PassParameters = GraphBuilder.AllocParameters<FShaderPrintClearCounterCS::FParameters>();
-				PassParameters->RWValuesBuffer = GraphBuilder.CreateUAV(ShaderPrintData.ShaderPrintEntryBuffer);
+				PassParameters->RWValuesBuffer = GraphBuilder.CreateUAV(ShaderPrintData.ShaderPrintEntryBuffer, PF_R32_UINT);
 				ClearUnusedGraphResources(ComputeShader, PassParameters);
 
 				GraphBuilder.AddPass(
@@ -902,7 +902,7 @@ namespace ShaderPrint
 		TShaderMapRef<FShaderPrintCopyCS> ComputeShader(GlobalShaderMap);
 		FShaderPrintCopyCS::FParameters* PassParameters = GraphBuilder.AllocParameters<FShaderPrintCopyCS::FParameters>();
 		PassParameters->Common = ShaderPrintData.UniformBuffer;
-		PassParameters->RWValuesBuffer = GraphBuilder.CreateUAV(ShaderPrintData.ShaderPrintEntryBuffer);;
+		PassParameters->RWValuesBuffer = GraphBuilder.CreateUAV(ShaderPrintData.ShaderPrintEntryBuffer, PF_R32_UINT);
 		FComputeShaderUtils::AddPass(
 			GraphBuilder,
 			RDG_EVENT_NAME("ShaderPrint::CopyParameters"),
@@ -962,8 +962,8 @@ namespace ShaderPrint
 		FRDGBufferSRVRef SymbolBufferSRV = GraphBuilder.CreateSRV(SymbolBuffer, PF_R32_UINT);
 
 		// Non graph managed resources
-		FRDGBufferSRVRef ValueBuffer = GraphBuilder.CreateSRV(ShaderPrintData.ShaderPrintEntryBuffer);
-		FRDGBufferSRVRef StateBuffer = GraphBuilder.CreateSRV(ShaderPrintData.ShaderPrintStateBuffer);
+		FRDGBufferSRVRef ValueBuffer = GraphBuilder.CreateSRV(ShaderPrintData.ShaderPrintEntryBuffer, PF_R32_UINT);
+		FRDGBufferSRVRef StateBuffer = GraphBuilder.CreateSRV(ShaderPrintData.ShaderPrintStateBuffer, PF_R32_UINT);
 		FTextureRHIRef FontTexture = GSystemTextures.AsciiTexture->GetRHI();
 
 		FGlobalShaderMap* GlobalShaderMap = GetGlobalShaderMap(GMaxRHIFeatureLevel);
@@ -996,7 +996,7 @@ namespace ShaderPrint
 			PassParameters->Common = ShaderPrintData.UniformBuffer;
 			PassParameters->ValuesBuffer = ValueBuffer;
 			PassParameters->RWSymbolsBuffer = SymbolBufferUAV;
-			PassParameters->RWStateBuffer = GraphBuilder.CreateUAV(ShaderPrintData.ShaderPrintStateBuffer);
+			PassParameters->RWStateBuffer = GraphBuilder.CreateUAV(ShaderPrintData.ShaderPrintStateBuffer, PF_R32_UINT);
 			PassParameters->IndirectDispatchArgsBuffer = IndirectDispatchArgsBuffer;
 
 			FComputeShaderUtils::AddPass(
@@ -1016,7 +1016,7 @@ namespace ShaderPrint
 			PassParameters->FrameIndex = FrameNumber;
 			PassParameters->FrameThreshold = 300u;
 			PassParameters->Common = ShaderPrintData.UniformBuffer;
-			PassParameters->RWStateBuffer = GraphBuilder.CreateUAV(ShaderPrintData.ShaderPrintStateBuffer);
+			PassParameters->RWStateBuffer = GraphBuilder.CreateUAV(ShaderPrintData.ShaderPrintStateBuffer, PF_R32_UINT);
 
 			FComputeShaderUtils::AddPass(
 				GraphBuilder,
diff --git a/Engine/Source/Runtime/Renderer/Public/ShaderPrintParameters.h b/Engine/Source/Runtime/Renderer/Public/ShaderPrintParameters.h
--- a/Engine/Source/Runtime/Renderer/Public/ShaderPrintParameters.h
+++ b/Engine/Source/Runtime/Renderer/Public/ShaderPrintParameters.h
@@ -43,8 +43,8 @@ namespace ShaderPrint
 	// ShaderPrint parameter struct declaration
 	BEGIN_SHADER_PARAMETER_STRUCT(FShaderParameters, )
 		SHADER_PARAMETER_STRUCT_REF(FShaderPrintCommonParameters, Common)
-		SHADER_PARAMETER_RDG_BUFFER_SRV(StructuredBuffer<uint2>, ShaderPrint_StateBuffer)
-		SHADER_PARAMETER_RDG_BUFFER_UAV(RWStructuredBuffer<ShaderPrintItem>, ShaderPrint_RWEntryBuffer)
+		SHADER_PARAMETER_RDG_BUFFER_SRV(Buffer<uint>, ShaderPrint_StateBuffer)
+		SHADER_PARAMETER_RDG_BUFFER_UAV(RWBuffer<uint>, ShaderPrint_RWEntryBuffer)
 	END_SHADER_PARAMETER_STRUCT()
 
 	// Does the platform support the ShaderPrint system?
