Refactor handling of meshes with some decal segments in ray tracing.
- don't try to reuse main instance for decals when bAllSegmentsDecal, instead add regular or decal instances to scene as necessary.
- fixes issues with incorrect SBT entries running with r.RayTracing.ExcludeDecals.

diff --git a/Engine/Source/Runtime/Renderer/Private/RayTracing/RayTracing.cpp b/Engine/Source/Runtime/Renderer/Private/RayTracing/RayTracing.cpp
--- a/Engine/Source/Runtime/Renderer/Private/RayTracing/RayTracing.cpp
+++ b/Engine/Source/Runtime/Renderer/Private/RayTracing/RayTracing.cpp
@@ -762,9 +762,9 @@ namespace RayTracing
 						// one containing non-decal segments and the other with decal segments
 						// masking of segments is done using "hidden" hitgroups
 						// TODO: Debug Visualization to highlight primitives using this?
-						const bool bNeedSeparateDecalInstance = Instance.MaskAndFlags.bAnySegmentsDecal && !Instance.MaskAndFlags.bAllSegmentsDecal;
+						const bool bNeedDecalInstance = Instance.MaskAndFlags.bAnySegmentsDecal && !ShouldExcludeDecals();
 
-						if (GRayTracingExcludeDecals && Instance.MaskAndFlags.bAnySegmentsDecal && !bNeedSeparateDecalInstance)
+						if (ShouldExcludeDecals() && Instance.MaskAndFlags.bAllSegmentsDecal)
 						{
 							continue;
 						}
@@ -774,7 +774,7 @@ namespace RayTracing
 						checkf(RayTracingInstance.GeometryRHI, TEXT("Ray tracing instance must have a valid geometry."));
 						RayTracingInstance.DefaultUserData = PersistentPrimitiveIndex.Index;
 						RayTracingInstance.bApplyLocalBoundsTransform = Instance.bApplyLocalBoundsTransform;
-						RayTracingInstance.LayerIndex = (uint8)(Instance.MaskAndFlags.bAnySegmentsDecal && !bNeedSeparateDecalInstance ? ERayTracingSceneLayer::Decals : ERayTracingSceneLayer::Base);
+						RayTracingInstance.LayerIndex = (uint8)ERayTracingSceneLayer::Base;
 						RayTracingInstance.Mask = Instance.MaskAndFlags.Mask;
 
 						if (Instance.MaskAndFlags.bForceOpaque)
@@ -814,10 +814,14 @@ namespace RayTracing
 							}
 						}
 
-						const uint32 InstanceIndex = RayTracingScene.AddInstance(RayTracingInstance, SceneProxy, true);
+						uint32 InstanceIndex = INDEX_NONE;
+						if (!Instance.MaskAndFlags.bAllSegmentsDecal)
+						{
+							InstanceIndex = RayTracingScene.AddInstance(RayTracingInstance, SceneProxy, true);
+						}
 
 						uint32 DecalInstanceIndex = INDEX_NONE;
-						if (bNeedSeparateDecalInstance && !GRayTracingExcludeDecals)
+						if (bNeedDecalInstance)
 						{
 							FRayTracingGeometryInstance DecalRayTracingInstance = RayTracingInstance;
 							DecalRayTracingInstance.LayerIndex = (uint8)ERayTracingSceneLayer::Decals;
@@ -1024,46 +1028,47 @@ namespace RayTracing
 						// one containing non-decal segments and the other with decal segments
 						// masking of segments is done using "hidden" hitgroups
 						// TODO: Debug Visualization to highlight primitives using this?
-						const bool bNeedSeparateDecalInstance = RelevantPrimitive.bAnySegmentsDecal && !RelevantPrimitive.bAllSegmentsDecal;
+						const bool bNeedDecalInstance = RelevantPrimitive.bAnySegmentsDecal && !ShouldExcludeDecals();
 
-						if (GRayTracingExcludeDecals && RelevantPrimitive.bAnySegmentsDecal && !bNeedSeparateDecalInstance)
+						if (ShouldExcludeDecals() && RelevantPrimitive.bAllSegmentsDecal)
 						{
 							continue;
 						}
 
 						check(RelevantPrimitive.CachedRayTracingInstance);
 
-						const int32 NewInstanceIndex = RayTracingScene.AddInstance(*RelevantPrimitive.CachedRayTracingInstance, SceneProxy, false);
-						uint32 DecalInstanceIndex = INDEX_NONE;
-
+						int32 InstanceIndex = INDEX_NONE;
+						if (!RelevantPrimitive.bAllSegmentsDecal)
 						{
-							FRayTracingGeometryInstance& NewInstance = RayTracingScene.GetInstance(NewInstanceIndex);
-							AddDebugRayTracingInstanceFlags(NewInstance.Flags);
+							FRayTracingGeometryInstance RayTracingInstance = *RelevantPrimitive.CachedRayTracingInstance;
+							RayTracingInstance.LayerIndex = (uint8)ERayTracingSceneLayer::Base;
+							AddDebugRayTracingInstanceFlags(RayTracingInstance.Flags);
 
-							NewInstance.LayerIndex = (uint8)(RelevantPrimitive.bAnySegmentsDecal && !bNeedSeparateDecalInstance ? ERayTracingSceneLayer::Decals : ERayTracingSceneLayer::Base);
+							InstanceIndex = RayTracingScene.AddInstance(MoveTemp(RayTracingInstance), SceneProxy, false);
+						}
 
-							if (bNeedSeparateDecalInstance && !GRayTracingExcludeDecals)
-							{
-								FRayTracingGeometryInstance DecalRayTracingInstance = NewInstance;
-								DecalRayTracingInstance.LayerIndex = (uint8)ERayTracingSceneLayer::Decals;
+						uint32 DecalInstanceIndex = INDEX_NONE;
+						if (bNeedDecalInstance)
+						{
+							FRayTracingGeometryInstance DecalRayTracingInstance = *RelevantPrimitive.CachedRayTracingInstance;
+							DecalRayTracingInstance.LayerIndex = (uint8)ERayTracingSceneLayer::Decals;
+							AddDebugRayTracingInstanceFlags(DecalRayTracingInstance.Flags);
 
-								DecalInstanceIndex = RayTracingScene.AddInstance(MoveTemp(DecalRayTracingInstance), SceneProxy, false);
-							}
+							DecalInstanceIndex = RayTracingScene.AddInstance(MoveTemp(DecalRayTracingInstance), SceneProxy, false);
 						}
 
-						const bool bHasDecalInstanceIndex = DecalInstanceIndex != INDEX_NONE;
-
 						for (int32 CommandIndex : RelevantPrimitive.CachedRayTracingMeshCommandIndices)
 						{
 							const FRayTracingMeshCommand& MeshCommand = Scene.CachedRayTracingMeshCommands[CommandIndex];
 
+							if(InstanceIndex != INDEX_NONE)
 							{
-								const bool bHidden = bHasDecalInstanceIndex && MeshCommand.bDecal;
-								FVisibleRayTracingMeshCommand NewVisibleMeshCommand(&MeshCommand, NewInstanceIndex, bHidden);
+								const bool bHidden = MeshCommand.bDecal;
+								FVisibleRayTracingMeshCommand NewVisibleMeshCommand(&MeshCommand, InstanceIndex, bHidden);
 								VisibleRayTracingMeshCommands.Add(NewVisibleMeshCommand);
 							}
 
-							if (bHasDecalInstanceIndex)
+							if (DecalInstanceIndex != INDEX_NONE)
 							{
 								const bool bHidden = !MeshCommand.bDecal;
 								FVisibleRayTracingMeshCommand NewVisibleMeshCommand(&MeshCommand, DecalInstanceIndex, bHidden);
@@ -1085,15 +1090,14 @@ namespace RayTracing
 						// one containing non-decal segments and the other with decal segments
 						// masking of segments is done using "hidden" hitgroups
 						// TODO: Debug Visualization to highlight primitives using this?
-						const bool bNeedSeparateDecalInstance = RelevantPrimitive.bAnySegmentsDecal && !RelevantPrimitive.bAllSegmentsDecal;
+						const bool bNeedDecalInstance = RelevantPrimitive.bAnySegmentsDecal && !ShouldExcludeDecals();
 
-						if (GRayTracingExcludeDecals && RelevantPrimitive.bAnySegmentsDecal && !bNeedSeparateDecalInstance)
+						if (ShouldExcludeDecals() && RelevantPrimitive.bAllSegmentsDecal)
 						{
 							continue;
 						}
 
-						if ((GRayTracingExcludeDecals && RelevantPrimitive.bAnySegmentsDecal)
-							|| (GRayTracingExcludeTranslucent && RelevantPrimitive.bAllSegmentsTranslucent)
+						if ((GRayTracingExcludeTranslucent && RelevantPrimitive.bAllSegmentsTranslucent)
 							|| (GRayTracingExcludeSky && RelevantPrimitive.bIsSky && !bIsPathTracing))
 						{
 							continue;
@@ -1111,7 +1115,7 @@ namespace RayTracing
 
 							bool bReallocated = InstanceBatch.Add(RayTracingScene, SceneInfo->GetInstanceSceneDataOffset(), uint32(PersistentPrimitiveIndex.Index));
 
-							check(InstanceBatch.Index != INDEX_NONE);
+							if(InstanceBatch.Index != INDEX_NONE)
 							{
 								FRayTracingGeometryInstance& RayTracingInstance = RayTracingScene.GetInstance(InstanceBatch.Index);
 								++RayTracingInstance.NumTransforms;
@@ -1162,11 +1166,16 @@ namespace RayTracing
 							}
 							AddDebugRayTracingInstanceFlags(RayTracingInstance.Flags);
 
-							RayTracingInstance.LayerIndex = (uint8)(RelevantPrimitive.bAnySegmentsDecal && !bNeedSeparateDecalInstance ? ERayTracingSceneLayer::Decals : ERayTracingSceneLayer::Base);
+							RayTracingInstance.LayerIndex = (uint8)ERayTracingSceneLayer::Base;
 
-							InstanceBatch.Index = RayTracingScene.AddInstance(RayTracingInstance, SceneProxy, false);
+							InstanceBatch.Index = INDEX_NONE;
+							if (!RelevantPrimitive.bAllSegmentsDecal)
+							{
+								InstanceBatch.Index = RayTracingScene.AddInstance(RayTracingInstance, SceneProxy, false);
+							}
 
-							if (bNeedSeparateDecalInstance && !GRayTracingExcludeDecals)
+							InstanceBatch.DecalIndex = INDEX_NONE;
+							if (bNeedDecalInstance)
 							{
 								FRayTracingGeometryInstance DecalRayTracingInstance = RayTracingInstance;
 								DecalRayTracingInstance.LayerIndex = (uint8)ERayTracingSceneLayer::Decals;
@@ -1174,21 +1183,20 @@ namespace RayTracing
 								InstanceBatch.DecalIndex = RayTracingScene.AddInstance(MoveTemp(DecalRayTracingInstance), SceneProxy, false);
 							}
 
-							const bool bHasDecalInstanceIndex = InstanceBatch.DecalIndex != INDEX_NONE;
-
 							for (int32 CommandIndex : RelevantPrimitive.CachedRayTracingMeshCommandIndices)
 							{
 								if (CommandIndex >= 0)
 								{
 									const FRayTracingMeshCommand& MeshCommand = Scene.CachedRayTracingMeshCommands[CommandIndex];
 
+									if(InstanceBatch.Index != INDEX_NONE)
 									{
-										const bool bHidden = bHasDecalInstanceIndex && MeshCommand.bDecal;
+										const bool bHidden = MeshCommand.bDecal;
 										FVisibleRayTracingMeshCommand NewVisibleMeshCommand(&MeshCommand, InstanceBatch.Index, bHidden);
 										VisibleRayTracingMeshCommands.Add(NewVisibleMeshCommand);
 									}
 
-									if (bHasDecalInstanceIndex)
+									if (InstanceBatch.DecalIndex != INDEX_NONE)
 									{
 										const bool bHidden = !MeshCommand.bDecal;
 										FVisibleRayTracingMeshCommand NewVisibleMeshCommand(&MeshCommand, InstanceBatch.DecalIndex, bHidden);
diff --git a/Engine/Source/Runtime/Renderer/Public/MeshPassProcessor.h b/Engine/Source/Runtime/Renderer/Public/MeshPassProcessor.h
--- a/Engine/Source/Runtime/Renderer/Public/MeshPassProcessor.h
+++ b/Engine/Source/Runtime/Renderer/Public/MeshPassProcessor.h
@@ -2453,15 +2453,14 @@ public:
 
 	virtual void FinalizeCommand(FRayTracingMeshCommand& RayTracingMeshCommand) override final
 	{
-		const bool bHasDecalInstanceIndex = RayTracingDecalInstanceIndex != INDEX_NONE;
-
+		if(RayTracingInstanceIndex != INDEX_NONE)
 		{
-			const bool bHidden = bHasDecalInstanceIndex && RayTracingMeshCommand.bDecal;
+			const bool bHidden = RayTracingMeshCommand.bDecal;
 			FVisibleRayTracingMeshCommand NewVisibleMeshCommand(&RayTracingMeshCommand, RayTracingInstanceIndex, bHidden);
 			VisibleCommands.Add(NewVisibleMeshCommand);
 		}
 
-		if (bHasDecalInstanceIndex)
+		if (RayTracingDecalInstanceIndex != INDEX_NONE)
 		{
 			const bool bHidden = !RayTracingMeshCommand.bDecal;
 			FVisibleRayTracingMeshCommand NewVisibleMeshCommand(&RayTracingMeshCommand, RayTracingDecalInstanceIndex, bHidden);
