Fix tile artifact on hair groom.

Primary visibility memory allocation is based on groom CPU bound. When using long hair with physic, the CPU bound might not be conservative enough, resulting into too low hair visibility memory allocation. This CL increase the CPU bound size for this allocation proportionaly to the longest curve of a groom.

diff --git a/Engine/Source/Runtime/Renderer/Private/HairStrands/HairStrandsCluster.cpp b/Engine/Source/Runtime/Renderer/Private/HairStrands/HairStrandsCluster.cpp
--- a/Engine/Source/Runtime/Renderer/Private/HairStrands/HairStrandsCluster.cpp
+++ b/Engine/Source/Runtime/Renderer/Private/HairStrands/HairStrandsCluster.cpp
@@ -172,7 +172,12 @@ void CreateHairStrandsMacroGroups(
 		if (!bIsValid)
 			return;
 
-		const FBoxSphereBounds& PrimitiveBounds = Proxy ? Proxy->GetBounds() : Bounds;
+		const FBoxSphereBounds& OriginalPrimitiveBounds = Proxy ? Proxy->GetBounds() : Bounds;
+
+		// Expand hair bound by (half) the groom's max hair-length, to be sure that the bounds are large enough. 
+		// This is important as the primary visibility memory allocation is based on the screen-projection of this CPU bound.
+		// If the bound is too small, the allocation won't be enough, resulting in tile artifacts.
+		FBoxSphereBounds PrimitiveBounds = OriginalPrimitiveBounds.ExpandBy(HairData->VFInput.Strands.Common.Length * 0.5f);
 
 		bool bFound = false;
 		float MinDistance = FLT_MAX;
