Fix ensure when running with -onethread.

The GPUScene primitive upload lambda runs on the main thread in that case, so it shouldn't override the task tag if it's already present.

#rnx

diff --git a/Engine/Source/Runtime/Renderer/Private/GPUScene.cpp b/Engine/Source/Runtime/Renderer/Private/GPUScene.cpp
--- a/Engine/Source/Runtime/Renderer/Private/GPUScene.cpp
+++ b/Engine/Source/Runtime/Renderer/Private/GPUScene.cpp
@@ -1167,7 +1167,7 @@ void FGPUScene::UploadGeneral(FRDGBuilder& GraphBuilder, const FRegisteredBuffer
 
 			ParallelForTemplate(TEXT("GPUScene Upload Primitives Task"), TaskContext.NumPrimitiveDataUploads, 1, [&TaskContext, &UploadDataSourceAdapter](int32 ItemIndex)
 			{
-				FTaskTagScope TaskTagScope(ETaskTag::EParallelRenderingThread);
+				FOptionalTaskTagScope TaskTagScope(ETaskTag::EParallelRenderingThread);
 
 				FVector4f* DstData = static_cast<FVector4f*>(TaskContext.PrimitiveUploader->GetRef(ItemIndex));
 				UploadDataSourceAdapter.GetPrimitiveShaderData(ItemIndex, DstData);
