Add instance culling to decal passes.
This fixes HISM decals.
Support is still needed for ISM LOD culling flags, but that involves more plumbing so deferring to a second pass.

diff --git a/Engine/Source/Runtime/Renderer/Private/CompositionLighting/CompositionLighting.cpp b/Engine/Source/Runtime/Renderer/Private/CompositionLighting/CompositionLighting.cpp
--- a/Engine/Source/Runtime/Renderer/Private/CompositionLighting/CompositionLighting.cpp
+++ b/Engine/Source/Runtime/Renderer/Private/CompositionLighting/CompositionLighting.cpp
@@ -548,7 +548,7 @@ void FCompositionLighting::TryInit()
 	bInitialized = true;
 }
 
-void FCompositionLighting::ProcessBeforeBasePass(FRDGBuilder& GraphBuilder, FDBufferTextures& DBufferTextures)
+void FCompositionLighting::ProcessBeforeBasePass(FRDGBuilder& GraphBuilder, FDBufferTextures& DBufferTextures, FInstanceCullingManager& InstanceCullingManager)
 {
 	if (HasRayTracedOverlay(ViewFamily))
 	{
@@ -581,7 +581,7 @@ void FCompositionLighting::ProcessBeforeBasePass(FRDGBuilder& GraphBuilder, FDBu
 		if (bEnableDBuffer)
 		{
 			FDeferredDecalPassTextures DecalPassTextures = GetDeferredDecalPassTextures(GraphBuilder, View, SceneTextures, &DBufferTextures);
-			AddDeferredDecalPass(GraphBuilder, View, DecalPassTextures, EDecalRenderStage::BeforeBasePass);
+			AddDeferredDecalPass(GraphBuilder, View, DecalPassTextures, InstanceCullingManager, EDecalRenderStage::BeforeBasePass);
 		}
 
 		if (bEnableSSAO)
@@ -598,7 +598,7 @@ void FCompositionLighting::ProcessBeforeBasePass(FRDGBuilder& GraphBuilder, FDBu
 	}
 }
 
-void FCompositionLighting::ProcessAfterBasePass(FRDGBuilder& GraphBuilder, EProcessAfterBasePassMode Mode)
+void FCompositionLighting::ProcessAfterBasePass(FRDGBuilder& GraphBuilder, FInstanceCullingManager& InstanceCullingManager, EProcessAfterBasePassMode Mode)
 {
 	if (HasRayTracedOverlay(ViewFamily))
 	{
@@ -626,13 +626,13 @@ void FCompositionLighting::ProcessAfterBasePass(FRDGBuilder& GraphBuilder, EProc
 		{
 			// We can disable this pass if using DBuffer decals
 			// Decals are before AmbientOcclusion so the decal can output a normal that AO is affected by
-			AddDeferredDecalPass(GraphBuilder, View, DecalPassTextures, EDecalRenderStage::BeforeLighting);
+			AddDeferredDecalPass(GraphBuilder, View, DecalPassTextures, InstanceCullingManager, EDecalRenderStage::BeforeLighting);
 		}
 
 		if (bEnableDecals && Mode != EProcessAfterBasePassMode::OnlyBeforeLightingDecals)
 		{
 			// DBuffer decals with emissive component
-			AddDeferredDecalPass(GraphBuilder, View, DecalPassTextures, EDecalRenderStage::Emissive);
+			AddDeferredDecalPass(GraphBuilder, View, DecalPassTextures, InstanceCullingManager, EDecalRenderStage::Emissive);
 		}
 
 		// Forward shading SSAO is applied before the base pass using only the depth buffer.
@@ -675,7 +675,7 @@ void FCompositionLighting::ProcessAfterBasePass(FRDGBuilder& GraphBuilder, EProc
 					if (bEnableDecals)
 					{
 						DecalPassTextures.ScreenSpaceAO = AmbientOcclusion.Texture;
-						AddDeferredDecalPass(GraphBuilder, View, DecalPassTextures, EDecalRenderStage::AmbientOcclusion);
+						AddDeferredDecalPass(GraphBuilder, View, DecalPassTextures, InstanceCullingManager, EDecalRenderStage::AmbientOcclusion);
 					}
 				}
 			}
diff --git a/Engine/Source/Runtime/Renderer/Private/CompositionLighting/CompositionLighting.h b/Engine/Source/Runtime/Renderer/Private/CompositionLighting/CompositionLighting.h
--- a/Engine/Source/Runtime/Renderer/Private/CompositionLighting/CompositionLighting.h
+++ b/Engine/Source/Runtime/Renderer/Private/CompositionLighting/CompositionLighting.h
@@ -21,7 +21,7 @@ public:
 
 	void ProcessAfterOcclusion(FRDGBuilder& GraphBuilder);
 
-	void ProcessBeforeBasePass(FRDGBuilder& GraphBuilder, FDBufferTextures& DBufferTextures);
+	void ProcessBeforeBasePass(FRDGBuilder& GraphBuilder, FDBufferTextures& DBufferTextures, FInstanceCullingManager& InstanceCullingManager);
 
 	enum class EProcessAfterBasePassMode
 	{
@@ -30,7 +30,7 @@ public:
 		All
 	};
 
-	void ProcessAfterBasePass(FRDGBuilder& GraphBuilder, EProcessAfterBasePassMode Mode);
+	void ProcessAfterBasePass(FRDGBuilder& GraphBuilder, FInstanceCullingManager& InstanceCullingManager, EProcessAfterBasePassMode Mode);
 
 private:
 	void TryInit();
diff --git a/Engine/Source/Runtime/Renderer/Private/CompositionLighting/PostProcessDeferredDecals.cpp b/Engine/Source/Runtime/Renderer/Private/CompositionLighting/PostProcessDeferredDecals.cpp
--- a/Engine/Source/Runtime/Renderer/Private/CompositionLighting/PostProcessDeferredDecals.cpp
+++ b/Engine/Source/Runtime/Renderer/Private/CompositionLighting/PostProcessDeferredDecals.cpp
@@ -170,9 +170,6 @@ void GetDeferredDecalPassParameters(
 	PassParameters.DeferredDecal = CreateDeferredDecalUniformBuffer(View);
 	PassParameters.DecalPass = Textures.DecalPassUniformBuffer;
 	
-	// TODO: hook up to instance culling manager and convert to simple mesh pass.
-	PassParameters.InstanceCulling = FInstanceCullingContext::CreateDummyInstanceCullingUniformBuffer(GraphBuilder);
-
 	FRDGTextureRef DepthTexture = Textures.Depth.Target;
 
 	FRenderTargetBindingSlots& RenderTargets = PassParameters.RenderTargets;
@@ -561,6 +558,7 @@ void AddDeferredDecalPass(
 	FRDGBuilder& GraphBuilder,
 	const FViewInfo& View,
 	const FDeferredDecalPassTextures& PassTextures,
+	FInstanceCullingManager& InstanceCullingManager,
 	EDecalRenderStage DecalRenderStage)
 {
 	check(PassTextures.Depth.IsValid());
@@ -687,7 +685,7 @@ void AddDeferredDecalPass(
 
 		if (MeshDecalCount > 0 && (DecalRenderStage == EDecalRenderStage::BeforeBasePass || DecalRenderStage == EDecalRenderStage::BeforeLighting || DecalRenderStage == EDecalRenderStage::Emissive || DecalRenderStage == EDecalRenderStage::AmbientOcclusion))
 		{
-			RenderMeshDecals(GraphBuilder, View, PassTextures, DecalRenderStage);
+			RenderMeshDecals(GraphBuilder, Scene, View, PassTextures, InstanceCullingManager, DecalRenderStage);
 		}
 
 		if (SortedDecalCount > 0)
diff --git a/Engine/Source/Runtime/Renderer/Private/CompositionLighting/PostProcessDeferredDecals.h b/Engine/Source/Runtime/Renderer/Private/CompositionLighting/PostProcessDeferredDecals.h
--- a/Engine/Source/Runtime/Renderer/Private/CompositionLighting/PostProcessDeferredDecals.h
+++ b/Engine/Source/Runtime/Renderer/Private/CompositionLighting/PostProcessDeferredDecals.h
@@ -48,6 +48,7 @@ void AddDeferredDecalPass(
 	FRDGBuilder& GraphBuilder,
 	const FViewInfo& ViewInfo,
 	const FDeferredDecalPassTextures& Textures,
+	FInstanceCullingManager& InstanceCullingManager,
 	EDecalRenderStage RenderStage);
 
 BEGIN_GLOBAL_SHADER_PARAMETER_STRUCT(FDeferredDecalUniformParameters, )
@@ -66,7 +67,7 @@ BEGIN_SHADER_PARAMETER_STRUCT(FDeferredDecalPassParameters, )
 	SHADER_PARAMETER_RDG_UNIFORM_BUFFER(FSceneUniformParameters, Scene)
 	SHADER_PARAMETER_STRUCT_REF(FDeferredDecalUniformParameters, DeferredDecal)
 	SHADER_PARAMETER_RDG_UNIFORM_BUFFER(FDecalPassUniformParameters, DecalPass)
-	SHADER_PARAMETER_RDG_UNIFORM_BUFFER(FInstanceCullingGlobalUniforms, InstanceCulling)
+	SHADER_PARAMETER_STRUCT_INCLUDE(FInstanceCullingDrawParams, InstanceCullingDrawParams)
 	RENDER_TARGET_BINDING_SLOTS()
 END_SHADER_PARAMETER_STRUCT()
 
@@ -93,8 +94,10 @@ void GetDeferredDecalPassParameters(
 
 void RenderMeshDecals(
 	FRDGBuilder& GraphBuilder,
+	const FScene& Scene,
 	const FViewInfo& View,
 	const FDeferredDecalPassTextures& DecalPassTextures,
+	FInstanceCullingManager& InstanceCullingManager, 
 	EDecalRenderStage DecalRenderStage);
 
 void ExtractNormalsForNextFrameReprojection(
diff --git a/Engine/Source/Runtime/Renderer/Private/CompositionLighting/PostProcessMeshDecals.cpp b/Engine/Source/Runtime/Renderer/Private/CompositionLighting/PostProcessMeshDecals.cpp
--- a/Engine/Source/Runtime/Renderer/Private/CompositionLighting/PostProcessMeshDecals.cpp
+++ b/Engine/Source/Runtime/Renderer/Private/CompositionLighting/PostProcessMeshDecals.cpp
@@ -20,6 +20,7 @@
 #include "UnrealEngine.h"
 #include "DebugViewModeRendering.h"
 #include "MeshPassProcessor.inl"
+#include "SimpleMeshDrawCommandPass.h"
 
 class FMeshDecalsVS : public FMeshMaterialShader
 {
@@ -445,19 +446,23 @@ FRegisterPSOCollectorCreateFunction RegisterPSOCollectorMeshDecal(&CreateMeshDec
 
 void DrawDecalMeshCommands(
 	FRDGBuilder& GraphBuilder,
+	const FScene& Scene,
 	const FViewInfo& View,
 	const FDeferredDecalPassTextures& DecalPassTextures,
+	FInstanceCullingManager& InstanceCullingManager,
 	EDecalRenderStage DecalRenderStage,
 	EDecalRenderTargetMode RenderTargetMode)
 {
 	auto* PassParameters = GraphBuilder.AllocParameters<FDeferredDecalPassParameters>();
 	GetDeferredDecalPassParameters(GraphBuilder, View, DecalPassTextures, RenderTargetMode, *PassParameters);
 
-	AddDrawDynamicMeshPass(
-		GraphBuilder,
-		RDG_EVENT_NAME("MeshDecals"),
-		PassParameters,
-		View,
+	AddSimpleMeshPass(
+		GraphBuilder, 
+		PassParameters, 
+		&Scene, 
+		View, 
+		&InstanceCullingManager, 
+		RDG_EVENT_NAME("MeshDecals"), 
 		View.ViewRect,
 		[&View, DecalRenderStage, RenderTargetMode](FDynamicPassMeshDrawListContext* DynamicMeshPassContext)
 	{
@@ -479,14 +484,15 @@ void DrawDecalMeshCommands(
 
 			PassMeshProcessor.AddMeshBatch(*Mesh, DefaultBatchElementMask, PrimitiveSceneProxy);
 		}
-
-	}, true);
+	});
 }
 
 void RenderMeshDecals(
 	FRDGBuilder& GraphBuilder,
+	const FScene& Scene,
 	const FViewInfo& View,
 	const FDeferredDecalPassTextures& DecalPassTextures,
+	FInstanceCullingManager& InstanceCullingManager,
 	EDecalRenderStage DecalRenderStage)
 {
 	QUICK_SCOPE_CYCLE_COUNTER(STAT_FSceneRenderer_RenderMeshDecals);
@@ -494,28 +500,28 @@ void RenderMeshDecals(
 	switch (DecalRenderStage)
 	{
 	case EDecalRenderStage::BeforeBasePass:
-		DrawDecalMeshCommands(GraphBuilder, View, DecalPassTextures, DecalRenderStage, EDecalRenderTargetMode::DBuffer);
+		DrawDecalMeshCommands(GraphBuilder, Scene, View, DecalPassTextures, InstanceCullingManager, DecalRenderStage, EDecalRenderTargetMode::DBuffer);
 		break;
 
 	case EDecalRenderStage::BeforeLighting:
-		DrawDecalMeshCommands(GraphBuilder, View, DecalPassTextures, DecalRenderStage, EDecalRenderTargetMode::SceneColorAndGBuffer);
-		DrawDecalMeshCommands(GraphBuilder, View, DecalPassTextures, DecalRenderStage, EDecalRenderTargetMode::SceneColorAndGBufferNoNormal);
+		DrawDecalMeshCommands(GraphBuilder, Scene, View, DecalPassTextures, InstanceCullingManager, DecalRenderStage, EDecalRenderTargetMode::SceneColorAndGBuffer);
+		DrawDecalMeshCommands(GraphBuilder, Scene, View, DecalPassTextures, InstanceCullingManager, DecalRenderStage, EDecalRenderTargetMode::SceneColorAndGBufferNoNormal);
 		break;
 
 	case EDecalRenderStage::Mobile:
-		DrawDecalMeshCommands(GraphBuilder, View, DecalPassTextures, DecalRenderStage, EDecalRenderTargetMode::SceneColor);
+		DrawDecalMeshCommands(GraphBuilder, Scene, View, DecalPassTextures, InstanceCullingManager, DecalRenderStage, EDecalRenderTargetMode::SceneColor);
 		break;
 
 	case EDecalRenderStage::MobileBeforeLighting:
-		DrawDecalMeshCommands(GraphBuilder, View, DecalPassTextures, DecalRenderStage, EDecalRenderTargetMode::SceneColorAndGBuffer);
+		DrawDecalMeshCommands(GraphBuilder, Scene, View, DecalPassTextures, InstanceCullingManager, DecalRenderStage, EDecalRenderTargetMode::SceneColorAndGBuffer);
 		break;
 
 	case EDecalRenderStage::Emissive:
-		DrawDecalMeshCommands(GraphBuilder, View, DecalPassTextures, DecalRenderStage, EDecalRenderTargetMode::SceneColor);
+		DrawDecalMeshCommands(GraphBuilder, Scene, View, DecalPassTextures, InstanceCullingManager, DecalRenderStage, EDecalRenderTargetMode::SceneColor);
 		break;
 
 	case EDecalRenderStage::AmbientOcclusion:
-		DrawDecalMeshCommands(GraphBuilder, View, DecalPassTextures, DecalRenderStage, EDecalRenderTargetMode::AmbientOcclusion);
+		DrawDecalMeshCommands(GraphBuilder, Scene, View, DecalPassTextures, InstanceCullingManager, DecalRenderStage, EDecalRenderTargetMode::AmbientOcclusion);
 		break;
 	}
 }
diff --git a/Engine/Source/Runtime/Renderer/Private/DeferredShadingRenderer.cpp b/Engine/Source/Runtime/Renderer/Private/DeferredShadingRenderer.cpp
--- a/Engine/Source/Runtime/Renderer/Private/DeferredShadingRenderer.cpp
+++ b/Engine/Source/Runtime/Renderer/Private/DeferredShadingRenderer.cpp
@@ -2334,7 +2334,7 @@ void FDeferredShadingSceneRenderer::Render(FRDGBuilder& GraphBuilder)
 		{
 			RDG_CSV_STAT_EXCLUSIVE_SCOPE(GraphBuilder, DeferredShadingSceneRenderer_DBuffer);
 			SCOPE_CYCLE_COUNTER(STAT_FDeferredShadingSceneRenderer_DBuffer);
-			CompositionLighting.ProcessBeforeBasePass(GraphBuilder, DBufferTextures);
+			CompositionLighting.ProcessBeforeBasePass(GraphBuilder, DBufferTextures, InstanceCullingManager);
 		}
 		
 		if (IsForwardShadingEnabled(ShaderPlatform))
@@ -2544,6 +2544,7 @@ void FDeferredShadingSceneRenderer::Render(FRDGBuilder& GraphBuilder)
 				GraphBuilder,
 				CompositionLighting,
 				SceneTextures,
+				InstanceCullingManager,
 				LumenFrameTemporaries,
 				LightingChannelsTexture,
 				/*bHasLumenLights*/ false,
@@ -2651,7 +2652,7 @@ void FDeferredShadingSceneRenderer::Render(FRDGBuilder& GraphBuilder)
 			const FCompositionLighting::EProcessAfterBasePassMode Mode = AsyncLumenIndirectLightingOutputs.bHasDrawnBeforeLightingDecals ?
 				FCompositionLighting::EProcessAfterBasePassMode::SkipBeforeLightingDecals : FCompositionLighting::EProcessAfterBasePassMode::All;
 
-			CompositionLighting.ProcessAfterBasePass(GraphBuilder, Mode);
+			CompositionLighting.ProcessAfterBasePass(GraphBuilder, InstanceCullingManager, Mode);
 		}
 
 		// Rebuild scene textures to include velocity, custom depth, and SSAO.
diff --git a/Engine/Source/Runtime/Renderer/Private/DeferredShadingRenderer.h b/Engine/Source/Runtime/Renderer/Private/DeferredShadingRenderer.h
--- a/Engine/Source/Runtime/Renderer/Private/DeferredShadingRenderer.h
+++ b/Engine/Source/Runtime/Renderer/Private/DeferredShadingRenderer.h
@@ -659,6 +659,7 @@ private:
 		FRDGBuilder& GraphBuilder,
 		class FCompositionLighting& CompositionLighting,
 		FSceneTextures& SceneTextures,
+		FInstanceCullingManager& InstanceCullingManager,
 		const FLumenSceneFrameTemporaries& LumenFrameTemporaries,
 		FRDGTextureRef LightingChannelsTexture,
 		bool bHasLumenLights,
diff --git a/Engine/Source/Runtime/Renderer/Private/IndirectLightRendering.cpp b/Engine/Source/Runtime/Renderer/Private/IndirectLightRendering.cpp
--- a/Engine/Source/Runtime/Renderer/Private/IndirectLightRendering.cpp
+++ b/Engine/Source/Runtime/Renderer/Private/IndirectLightRendering.cpp
@@ -859,6 +859,7 @@ void FDeferredShadingSceneRenderer::DispatchAsyncLumenIndirectLightingWork(
 	FRDGBuilder& GraphBuilder,
 	FCompositionLighting& CompositionLighting,
 	FSceneTextures& SceneTextures,
+	FInstanceCullingManager& InstanceCullingManager,
 	const FLumenSceneFrameTemporaries& LumenFrameTemporaries,
 	FRDGTextureRef LightingChannelsTexture,
 	bool bHasLumenLights,
@@ -879,7 +880,7 @@ void FDeferredShadingSceneRenderer::DispatchAsyncLumenIndirectLightingWork(
 	}
 
 	// Decals may modify GBuffers so they need to be done first. Can decals read velocities and/or custom depth? If so, they need to be rendered earlier too.
-	CompositionLighting.ProcessAfterBasePass(GraphBuilder, FCompositionLighting::EProcessAfterBasePassMode::OnlyBeforeLightingDecals);
+	CompositionLighting.ProcessAfterBasePass(GraphBuilder, InstanceCullingManager, FCompositionLighting::EProcessAfterBasePassMode::OnlyBeforeLightingDecals);
 	Outputs.bHasDrawnBeforeLightingDecals = true;
 
 	LLM_SCOPE_BYTAG(Lumen);
diff --git a/Engine/Source/Runtime/Renderer/Private/MobileDecalRendering.cpp b/Engine/Source/Runtime/Renderer/Private/MobileDecalRendering.cpp
--- a/Engine/Source/Runtime/Renderer/Private/MobileDecalRendering.cpp
+++ b/Engine/Source/Runtime/Renderer/Private/MobileDecalRendering.cpp
@@ -149,7 +149,7 @@ void RenderDeferredDecalsMobile(FRHICommandList& RHICmdList, const FScene& Scene
 	}
 }
 
-void FMobileSceneRenderer::RenderDBuffer(FRDGBuilder& GraphBuilder, FSceneTextures& SceneTextures, FDBufferTextures& DBufferTextures)
+void FMobileSceneRenderer::RenderDBuffer(FRDGBuilder& GraphBuilder, FSceneTextures& SceneTextures, FDBufferTextures& DBufferTextures, FInstanceCullingManager& InstanceCullingManager)
 {
 	RDG_EVENT_SCOPE(GraphBuilder, "RenderDBuffer");
 	QUICK_SCOPE_CYCLE_COUNTER(STAT_RenderDBuffer);
@@ -166,6 +166,6 @@ void FMobileSceneRenderer::RenderDBuffer(FRDGBuilder& GraphBuilder, FSceneTextur
 		}
 
 		FDeferredDecalPassTextures DecalPassTextures = GetDeferredDecalPassTextures(GraphBuilder, View, SceneTextures, &DBufferTextures);
-		AddDeferredDecalPass(GraphBuilder, View, DecalPassTextures, EDecalRenderStage::BeforeBasePass);
+		AddDeferredDecalPass(GraphBuilder, View, DecalPassTextures, InstanceCullingManager, EDecalRenderStage::BeforeBasePass);
 	}
 }
diff --git a/Engine/Source/Runtime/Renderer/Private/MobileShadingRenderer.cpp b/Engine/Source/Runtime/Renderer/Private/MobileShadingRenderer.cpp
--- a/Engine/Source/Runtime/Renderer/Private/MobileShadingRenderer.cpp
+++ b/Engine/Source/Runtime/Renderer/Private/MobileShadingRenderer.cpp
@@ -1253,7 +1253,7 @@ void FMobileSceneRenderer::Render(FRDGBuilder& GraphBuilder)
 			if (bRequiresDBufferDecals)
 			{
 				DBufferTextures = CreateDBufferTextures(GraphBuilder, SceneTextures.Config.Extent, ShaderPlatform);
-				RenderDBuffer(GraphBuilder, SceneTextures, DBufferTextures);
+				RenderDBuffer(GraphBuilder, SceneTextures, DBufferTextures, InstanceCullingManager);
 			}
 		}
 
diff --git a/Engine/Source/Runtime/Renderer/Private/SceneRendering.h b/Engine/Source/Runtime/Renderer/Private/SceneRendering.h
--- a/Engine/Source/Runtime/Renderer/Private/SceneRendering.h
+++ b/Engine/Source/Runtime/Renderer/Private/SceneRendering.h
@@ -2716,7 +2716,7 @@ protected:
 
 	/** Renders decals. */
 	void RenderDecals(FRHICommandList& RHICmdList, const FViewInfo& View, const FInstanceCullingDrawParams* InstanceCullingDrawParams);
-	void RenderDBuffer(FRDGBuilder& GraphBuilder, FSceneTextures& SceneTextures, FDBufferTextures& DBufferTextures);
+	void RenderDBuffer(FRDGBuilder& GraphBuilder, FSceneTextures& SceneTextures, FDBufferTextures& DBufferTextures, FInstanceCullingManager& InstanceCullingManager);
 
 	/** Renders the atmospheric and height fog */
 	void RenderFog(FRHICommandList& RHICmdList, const FViewInfo& View);
