Adding missing RayTracing ShaderPlatform checks.

diff --git a/Engine/Source/Runtime/Renderer/Private/DeferredShadingRenderer.cpp b/Engine/Source/Runtime/Renderer/Private/DeferredShadingRenderer.cpp
--- a/Engine/Source/Runtime/Renderer/Private/DeferredShadingRenderer.cpp
+++ b/Engine/Source/Runtime/Renderer/Private/DeferredShadingRenderer.cpp
@@ -3359,7 +3359,7 @@ void FDeferredShadingSceneRenderer::Render(FRDGBuilder& GraphBuilder)
 
 bool AnyRayTracingPassEnabled(const FScene* Scene, const FViewInfo& View)
 {
-	if (!IsRayTracingEnabled() || Scene == nullptr)
+	if (!IsRayTracingEnabled(View.GetShaderPlatform()) || Scene == nullptr)
 	{
 		return false;
 	}
@@ -3377,7 +3377,7 @@ bool AnyRayTracingPassEnabled(const FScene* Scene, const FViewInfo& View)
 
 bool ShouldRenderRayTracingEffect(bool bEffectEnabled, ERayTracingPipelineCompatibilityFlags CompatibilityFlags, const FSceneView* View)
 {
-	if (!IsRayTracingEnabled() || (View && !View->bAllowRayTracing))
+	if ((View && !IsRayTracingEnabled(View->GetShaderPlatform())) || (!View && !IsRayTracingEnabled()) || (View && !View->bAllowRayTracing))
 	{
 		return false;
 	}
