Selectively add ISR support for the mesh pass portion of the editor outline post-processing effect. Fixes editor outlines failing to appear in the right eye in VR scouting when instanced stereo is enabled.

diff --git a/Engine/Shaders/Private/HitProxyPixelShader.usf b/Engine/Shaders/Private/HitProxyPixelShader.usf
--- a/Engine/Shaders/Private/HitProxyPixelShader.usf
+++ b/Engine/Shaders/Private/HitProxyPixelShader.usf
@@ -31,7 +31,11 @@ void Main(
 	,out float4 OutColor		: SV_Target0
 	)
 {
-	ResolvedView = ResolveView();
+	#if INSTANCED_STEREO
+		ResolvedView = ResolveView(Interpolants.EyeIndex);
+	#else
+		ResolvedView = ResolveView();
+	#endif
 
 	FMaterialPixelParameters MaterialParameters = GetMaterialPixelParameters(Interpolants, SvPosition);
 	FPixelMaterialInputs PixelMaterialInputs;
diff --git a/Engine/Shaders/Private/HitProxyVertexShader.usf b/Engine/Shaders/Private/HitProxyVertexShader.usf
--- a/Engine/Shaders/Private/HitProxyVertexShader.usf
+++ b/Engine/Shaders/Private/HitProxyVertexShader.usf
@@ -51,10 +51,17 @@ void Main(
 	out FHitProxyVSOutput Output
 #if USE_GLOBAL_CLIP_PLANE
 	, out float OutGlobalClipPlaneDistance : SV_ClipDistance
+#endif
+#if INSTANCED_STEREO
+	, out uint ViewportIndex : SV_ViewPortArrayIndex
 #endif
 	)
 {
-	ResolvedView = ResolveView();
+#if INSTANCED_STEREO
+	uint EyeIndex = GetEyeIndexFromVF(Input);
+	ViewportIndex = EyeIndex;
+#endif
+	ResolvedView = ResolveViewFromVF(Input);
 
 	FVertexFactoryIntermediates VFIntermediates = GetVertexFactoryIntermediates(Input);
 	float4 WorldPosition = VertexFactoryGetWorldPosition(Input, VFIntermediates);
@@ -76,7 +83,7 @@ void Main(
 
 	{
 		float4 RasterizedWorldPosition = VertexFactoryGetRasterizedWorldPosition(Input, VFIntermediates, WorldPosition);
-		Output.Position = INVARIANT(mul(RasterizedWorldPosition, View.TranslatedWorldToClip));
+		Output.Position = INVARIANT(mul(RasterizedWorldPosition, ResolvedView.TranslatedWorldToClip));
 	}
 
 #if USE_GLOBAL_CLIP_PLANE
diff --git a/Engine/Shaders/Private/PostProcessSelectionOutline.usf b/Engine/Shaders/Private/PostProcessSelectionOutline.usf
--- a/Engine/Shaders/Private/PostProcessSelectionOutline.usf
+++ b/Engine/Shaders/Private/PostProcessSelectionOutline.usf
@@ -29,6 +29,7 @@ int OutlineColorIndexBits;
 float SelectionHighlightIntensity;
 float BSPSelectionIntensity;
 float UILuminanceAndIsSCRGB;
+float SecondaryViewportOffset;
 
 struct FPixelInfo
 {
@@ -123,10 +124,10 @@ float4 PerSample(int2 PixelPos, int SampleID, FPixelInfo CenterPixelInfo, float3
 	PixelInfo[0] = CenterPixelInfo;
 
 	// Diagonal cross is thicker than vertical/horizontal cross.
-	PixelInfo[1] = GetPixelInfo(PixelPos, SampleID,  1,  1, DeviceZPlane, DeviceZMinMax);
-	PixelInfo[2] = GetPixelInfo(PixelPos, SampleID, -1, -1, DeviceZPlane, DeviceZMinMax);
-	PixelInfo[3] = GetPixelInfo(PixelPos, SampleID,  1, -1, DeviceZPlane, DeviceZMinMax);
-	PixelInfo[4] = GetPixelInfo(PixelPos, SampleID, -1,  1, DeviceZPlane, DeviceZMinMax);
+	PixelInfo[1] = GetPixelInfo(PixelPos, SampleID, SecondaryViewportOffset + 1,  1, DeviceZPlane, DeviceZMinMax);
+	PixelInfo[2] = GetPixelInfo(PixelPos, SampleID, SecondaryViewportOffset - 1, -1, DeviceZPlane, DeviceZMinMax);
+	PixelInfo[3] = GetPixelInfo(PixelPos, SampleID, SecondaryViewportOffset + 1, -1, DeviceZPlane, DeviceZMinMax);
+	PixelInfo[4] = GetPixelInfo(PixelPos, SampleID, SecondaryViewportOffset - 1,  1, DeviceZPlane, DeviceZMinMax);
 
 	// With (.x != .y) we can detect a border around and between each object.
 	int2 BorderMinMax = int2(255,0);
@@ -218,17 +219,17 @@ void MainPS(
 
 #if MSAA_SAMPLE_COUNT > 1
 	const int4 StencilScout = int4(
-		EditorPrimitivesStencil.Load(ColorPixelPos + int2( ScoutStep,  0), 0) STENCIL_COMPONENT_SWIZZLE,
-		EditorPrimitivesStencil.Load(ColorPixelPos + int2(-ScoutStep,  0), 0) STENCIL_COMPONENT_SWIZZLE,
-		EditorPrimitivesStencil.Load(ColorPixelPos + int2( 0,  ScoutStep), 0) STENCIL_COMPONENT_SWIZZLE,
-		EditorPrimitivesStencil.Load(ColorPixelPos + int2( 0, -ScoutStep), 0) STENCIL_COMPONENT_SWIZZLE
+		EditorPrimitivesStencil.Load(ColorPixelPos + int2( ScoutStep,  0) + int2(SecondaryViewportOffset, 0), 0) STENCIL_COMPONENT_SWIZZLE,
+		EditorPrimitivesStencil.Load(ColorPixelPos + int2(-ScoutStep,  0) + int2(SecondaryViewportOffset, 0), 0) STENCIL_COMPONENT_SWIZZLE,
+		EditorPrimitivesStencil.Load(ColorPixelPos + int2( 0,  ScoutStep) + int2(SecondaryViewportOffset, 0), 0) STENCIL_COMPONENT_SWIZZLE,
+		EditorPrimitivesStencil.Load(ColorPixelPos + int2( 0, -ScoutStep) + int2(SecondaryViewportOffset, 0), 0) STENCIL_COMPONENT_SWIZZLE
 	);
 #else
 	const int4 StencilScout = int4(
-		EditorPrimitivesStencil.Load(int3(ColorPixelPos + int2( ScoutStep,  0), 0)) STENCIL_COMPONENT_SWIZZLE,
-		EditorPrimitivesStencil.Load(int3(ColorPixelPos + int2(-ScoutStep,  0), 0)) STENCIL_COMPONENT_SWIZZLE,
-		EditorPrimitivesStencil.Load(int3(ColorPixelPos + int2( 0,  ScoutStep), 0)) STENCIL_COMPONENT_SWIZZLE,
-		EditorPrimitivesStencil.Load(int3(ColorPixelPos + int2( 0, -ScoutStep), 0)) STENCIL_COMPONENT_SWIZZLE
+		EditorPrimitivesStencil.Load(int3(ColorPixelPos + int2( ScoutStep,  0) + int2(SecondaryViewportOffset, 0), 0)) STENCIL_COMPONENT_SWIZZLE,
+		EditorPrimitivesStencil.Load(int3(ColorPixelPos + int2(-ScoutStep,  0) + int2(SecondaryViewportOffset, 0), 0)) STENCIL_COMPONENT_SWIZZLE,
+		EditorPrimitivesStencil.Load(int3(ColorPixelPos + int2( 0,  ScoutStep) + int2(SecondaryViewportOffset, 0), 0)) STENCIL_COMPONENT_SWIZZLE,
+		EditorPrimitivesStencil.Load(int3(ColorPixelPos + int2( 0, -ScoutStep) + int2(SecondaryViewportOffset, 0), 0)) STENCIL_COMPONENT_SWIZZLE
 	);
 #endif
 
@@ -256,7 +257,7 @@ void MainPS(
 		DeviceZMinMax.x = min(Center, min(min(Left, Right), min(Top, Bottom)));
 		DeviceZMinMax.y = max(Center, max(max(Left, Right), max(Top, Bottom)));
 
-		FPixelInfo CenterPixelInfo = GetPixelInfo(ColorPixelPos, 0, 0,  0, DeviceZPlane, DeviceZMinMax);
+		FPixelInfo CenterPixelInfo = GetPixelInfo(ColorPixelPos, 0, SecondaryViewportOffset, 0, DeviceZPlane, DeviceZMinMax);
 
 		float4 Sum = 0;
 #if COMPILER_GLSL && FEATURE_LEVEL >= FEATURE_LEVEL_SM4
diff --git a/Engine/Source/Runtime/Renderer/Private/DeferredShadingRenderer.cpp b/Engine/Source/Runtime/Renderer/Private/DeferredShadingRenderer.cpp
--- a/Engine/Source/Runtime/Renderer/Private/DeferredShadingRenderer.cpp
+++ b/Engine/Source/Runtime/Renderer/Private/DeferredShadingRenderer.cpp
@@ -3226,6 +3226,8 @@ void FDeferredShadingSceneRenderer::Render(FRDGBuilder& GraphBuilder)
 			PostProcessingInputs.bSeparateCustomStencil = SceneTextures.CustomDepth.bSeparateStencilBuffer;
 			PostProcessingInputs.PathTracingResources = PathTracingResources;
 
+			FRDGTextureRef InstancedEditorDepthTexture = nullptr; // Used to pass instanced stereo depth data from primary to secondary views
+
 			GraphBuilder.FlushSetupQueue();
 
 			if (ViewFamily.UseDebugViewPS())
@@ -3233,7 +3235,7 @@ void FDeferredShadingSceneRenderer::Render(FRDGBuilder& GraphBuilder)
 				for (int32 ViewIndex = 0; ViewIndex < Views.Num(); ViewIndex++)
 				{
 					const FViewInfo& View = Views[ViewIndex];
-   					const Nanite::FRasterResults* NaniteResults = bNaniteEnabled ? &NaniteRasterResults[ViewIndex] : nullptr;
+					const Nanite::FRasterResults* NaniteResults = bNaniteEnabled ? &NaniteRasterResults[ViewIndex] : nullptr;
 					RDG_GPU_MASK_SCOPE(GraphBuilder, View.GPUMask);
 					RDG_EVENT_SCOPE_CONDITIONAL(GraphBuilder, Views.Num() > 1, "View%d", ViewIndex);
 					PostProcessingInputs.TranslucencyViewResourcesMap = FTranslucencyViewResourcesMap(TranslucencyResourceMap, ViewIndex);
@@ -3294,7 +3296,8 @@ void FDeferredShadingSceneRenderer::Render(FRDGBuilder& GraphBuilder)
 							&VirtualShadowMapArray,
 							LumenFrameTemporaries,
 							SceneWithoutWaterTextures,
-							TSRFlickeringInput);
+							TSRFlickeringInput,
+							InstancedEditorDepthTexture);
 					}
 				}
 			}
diff --git a/Engine/Source/Runtime/Renderer/Private/PostProcess/PostProcessSelectionOutline.cpp b/Engine/Source/Runtime/Renderer/Private/PostProcess/PostProcessSelectionOutline.cpp
--- a/Engine/Source/Runtime/Renderer/Private/PostProcess/PostProcessSelectionOutline.cpp
+++ b/Engine/Source/Runtime/Renderer/Private/PostProcess/PostProcessSelectionOutline.cpp
@@ -37,6 +37,7 @@ public:
 		SHADER_PARAMETER(float, SelectionHighlightIntensity)
 		SHADER_PARAMETER(float, BSPSelectionIntensity)
 		SHADER_PARAMETER(float, UILuminanceAndIsSCRGB)
+		SHADER_PARAMETER(float, SecondaryViewportOffset)
 		RENDER_TARGET_BINDING_SLOTS()
 	END_SHADER_PARAMETER_STRUCT()
 
@@ -62,7 +63,7 @@ IMPLEMENT_GLOBAL_SHADER(FSelectionOutlinePS, "/Engine/Private/PostProcessSelecti
 } //! namespace
 
 BEGIN_SHADER_PARAMETER_STRUCT(FSelectionOutlinePassParameters, )
-	SHADER_PARAMETER_STRUCT_REF(FViewUniformShaderParameters, View)
+	SHADER_PARAMETER_STRUCT_INCLUDE(FViewShaderParameters, View)
 	SHADER_PARAMETER_STRUCT_INCLUDE(FSceneTextureShaderParameters, SceneTextures)
 	SHADER_PARAMETER_STRUCT_INCLUDE(FInstanceCullingDrawParams, InstanceCullingDrawParams)
 	RENDER_TARGET_BINDING_SLOTS()
@@ -77,7 +78,8 @@ FScreenPassTexture AddSelectionOutlinePass(
 	const FViewInfo& View,
 	FSceneUniformBuffer &SceneUniformBuffer,
 	const FSelectionOutlineInputs& Inputs,
-	const Nanite::FRasterResults* NaniteRasterResults)
+	const Nanite::FRasterResults* NaniteRasterResults,
+	FRDGTextureRef& DepthStencilTexture)
 {
 	check(Inputs.SceneColor.IsValid());
 	check(Inputs.SceneDepth.IsValid());
@@ -90,69 +92,99 @@ FScreenPassTexture AddSelectionOutlinePass(
 	const uint32 NumSamples = View.GetSceneTexturesConfig().NumSamples;
 
 	// Patch uniform buffers with updated state for rendering the outline mesh draw commands.
-	const FViewInfo* EditorView = CreateCompositePrimitiveView(View, Inputs.SceneColor.ViewRect, NumSamples);
-
-	FRDGTextureRef DepthStencilTexture = nullptr;
+	const bool bIsInstancedStereoView = View.bIsInstancedStereoEnabled && IStereoRendering::IsStereoEyePass(View.StereoPass);
+	const FViewInfo* EditorView = CreateCompositePrimitiveView(View, bIsInstancedStereoView ? View.ViewRect : Inputs.SceneColor.ViewRect, NumSamples);
 
 	// Generate custom depth / stencil for outline shapes.
 	{
-		const FScreenPassTextureViewport SceneColorViewport(Inputs.SceneColor);
-		RDG_EVENT_SCOPE(GraphBuilder, "OutlineDepth %dx%d", SceneColorViewport.Rect.Width(), SceneColorViewport.Rect.Height());
+		FScreenPassTextureViewport DepthStencilViewport(Inputs.SceneColor);
+		RDG_EVENT_SCOPE(GraphBuilder, "OutlineDepth %dx%d", DepthStencilViewport.Rect.Width(), DepthStencilViewport.Rect.Height());
 
 		FScene* Scene = View.Family->Scene->GetRenderScene();
 
+		if (View.ShouldRenderView() || !DepthStencilTexture)
 		{
-			FRDGTextureDesc DepthStencilDesc = Inputs.SceneColor.Texture->Desc;
-			DepthStencilDesc.Reset();
-			DepthStencilDesc.Format = PF_DepthStencil;
-			// This is a reversed Z depth surface, so 0.0f is the far plane.
-			DepthStencilDesc.ClearValue = FClearValueBinding((float)ERHIZBuffer::FarPlane, 0);
-			DepthStencilDesc.Flags = TexCreate_DepthStencilTargetable | TexCreate_ShaderResource;
-			DepthStencilDesc.NumSamples = NumSamples;
-
-			DepthStencilTexture = GraphBuilder.CreateTexture(DepthStencilDesc, TEXT("Editor.SelectionOutline"));
-		}
-
-		{
-			auto* PassParameters = GraphBuilder.AllocParameters<FSelectionOutlinePassParameters>();
 
-			PassParameters->View = EditorView->ViewUniformBuffer;
-			PassParameters->SceneTextures = Inputs.SceneTextures;
-			PassParameters->RenderTargets.DepthStencil = FDepthStencilBinding(
-				DepthStencilTexture,
-				ERenderTargetLoadAction::EClear,
-				ERenderTargetLoadAction::EClear,
-				FExclusiveDepthStencil::DepthWrite_StencilWrite);
+			// If instanced stereo is enabled, we can only do traditional draws from the primary view, so we need to generate the secondary view's depth data at the same time.
+			// To hold that secondary view depth data, we patch the depth stencil texture and its viewport to be double-wide.
+			if (bIsInstancedStereoView /* Only primary ISR views have View.ShouldRenderView() == true */)
+			{
+				DepthStencilViewport = FScreenPassTextureViewport(View.ViewRectWithSecondaryViews);
+			}
 
-			const_cast<FViewInfo&>(View).ParallelMeshDrawCommandPasses[EMeshPass::EditorSelection].BuildRenderingCommands(GraphBuilder, Scene->GPUScene, PassParameters->InstanceCullingDrawParams);
+			{
+				FRDGTextureDesc DepthStencilDesc = Inputs.SceneColor.Texture->Desc;
+				DepthStencilDesc.Reset();
+				DepthStencilDesc.Format = PF_DepthStencil;
+				// This is a reversed Z depth surface, so 0.0f is the far plane.
+				DepthStencilDesc.ClearValue = FClearValueBinding((float)ERHIZBuffer::FarPlane, 0);
+				DepthStencilDesc.Flags = TexCreate_DepthStencilTargetable | TexCreate_ShaderResource;
+				DepthStencilDesc.NumSamples = NumSamples;
+				DepthStencilDesc.Extent = DepthStencilViewport.Extent;
+
+				DepthStencilTexture = GraphBuilder.CreateTexture(DepthStencilDesc, TEXT("Editor.SelectionOutline"));
+			}
 
-			GraphBuilder.AddPass(
-				RDG_EVENT_NAME("EditorSelectionDepth"),
-				PassParameters,
-				ERDGPassFlags::Raster,
-				[&View, SceneColorViewport, DepthStencilTexture, NaniteRasterResults, PassParameters, bNaniteEnabled](FRHICommandListImmediate& RHICmdList)
-				{
-					RHICmdList.SetViewport(SceneColorViewport.Rect.Min.X, SceneColorViewport.Rect.Min.Y, 0.0f, SceneColorViewport.Rect.Max.X, SceneColorViewport.Rect.Max.Y, 1.0f);
+			{
+				auto* PassParameters = GraphBuilder.AllocParameters<FSelectionOutlinePassParameters>();
+
+				PassParameters->View = EditorView->GetShaderParameters();
+				PassParameters->SceneTextures = Inputs.SceneTextures;
+				PassParameters->RenderTargets.DepthStencil = FDepthStencilBinding(
+					DepthStencilTexture,
+					ERenderTargetLoadAction::EClear,
+					ERenderTargetLoadAction::EClear,
+					FExclusiveDepthStencil::DepthWrite_StencilWrite);
+
+				const float ViewportScale = (float)Inputs.SceneColor.ViewRect.Width() / (float)View.ViewRect.Width(); // Outline selection can be before or after the primary upscale
+
+				const_cast<FViewInfo&>(View).ParallelMeshDrawCommandPasses[EMeshPass::EditorSelection].BuildRenderingCommands(GraphBuilder, Scene->GPUScene, PassParameters->InstanceCullingDrawParams);
+
+				GraphBuilder.AddPass(
+					RDG_EVENT_NAME("EditorSelectionDepth"),
+					PassParameters,
+					ERDGPassFlags::Raster,
+					[&View, DepthStencilViewport, PassParameters, ViewportScale](FRHICommandListImmediate& RHICmdList)
+					{
+						if (View.bIsInstancedStereoEnabled && View.StereoPass == EStereoscopicPass::eSSP_PRIMARY)
+						{
+							FSceneRenderer::SetStereoViewport(RHICmdList, View, ViewportScale);
+						}
+						else
+						{
+							RHICmdList.SetViewport(DepthStencilViewport.Rect.Min.X, DepthStencilViewport.Rect.Min.Y, 0.0f, DepthStencilViewport.Rect.Max.X, DepthStencilViewport.Rect.Max.Y, 1.0f);
+						}
+						
+						// Run selection pass on static elements
+						View.ParallelMeshDrawCommandPasses[EMeshPass::EditorSelection].DispatchDraw(nullptr, RHICmdList, &PassParameters->InstanceCullingDrawParams);
+					}
+				);
+			}
+		}
 
-					// Run selection pass on static elements
-					View.ParallelMeshDrawCommandPasses[EMeshPass::EditorSelection].DispatchDraw(nullptr, RHICmdList, &PassParameters->InstanceCullingDrawParams);
-				}
-			);
+		// If the depth stencil is double-wide, retarget the view rect to the appropriate side for remaining steps.
+		if (bIsInstancedStereoView)
+		{
+			// Use View.ViewRect rather than SceneColor.ViewRect, since SceneColor may be single-wide (thus always 0,0 locked).
+			DepthStencilViewport = FScreenPassTextureViewport(DepthStencilTexture->Desc.Extent, View.ViewRect);
 		}
 
 		// Render Nanite mesh outlines after regular mesh outline, but before borders
 		if (bNaniteEnabled)
 		{
+			// Update editor view to true target view rect
 			Nanite::DrawEditorSelection(GraphBuilder, DepthStencilTexture, *Scene, View, *EditorView, SceneUniformBuffer, NaniteRasterResults);
 		}
 
 		// Render HairStrands outlines
 		if (HairStrands::HasViewHairStrandsData(View))
 		{
-			HairStrands::DrawEditorSelection(GraphBuilder, View, SceneColorViewport.Rect, DepthStencilTexture);
+			HairStrands::DrawEditorSelection(GraphBuilder, View, DepthStencilViewport.Rect, DepthStencilTexture);
 		}
 
 		// Clear the borders to get an outline around the objects if it's partly outside of the screen
+		// Unnecessary for stereo (VR editor) views since the edge of the screen is not visible
+		if (!IStereoRendering::IsStereoEyePass(View.StereoPass))
 		{
 			auto* PassParameters = GraphBuilder.AllocParameters<FSelectionOutlineClearBorderParameters>();
 
@@ -166,27 +198,27 @@ FScreenPassTexture AddSelectionOutlinePass(
 				RDG_EVENT_NAME("DrawOutlineBorder"),
 				PassParameters,
 				ERDGPassFlags::Raster,
-				[SceneColorViewport](FRHICommandListImmediate& RHICmdList)
+				[DepthStencilViewport](FRHICommandListImmediate& RHICmdList)
 				{
-                    RHICmdList.SetViewport(SceneColorViewport.Rect.Min.X, SceneColorViewport.Rect.Min.Y, 0.0f, SceneColorViewport.Rect.Max.X, SceneColorViewport.Rect.Max.Y, 1.0f);
-                    
-					FIntRect InnerRect = SceneColorViewport.Rect;
+					RHICmdList.SetViewport(DepthStencilViewport.Rect.Min.X, DepthStencilViewport.Rect.Min.Y, 0.0f, DepthStencilViewport.Rect.Max.X, DepthStencilViewport.Rect.Max.Y, 1.0f);
+
+					FIntRect InnerRect = DepthStencilViewport.Rect;
 
 					// 1 as we have an outline that is that thick
 					InnerRect.InflateRect(-1);
 
 					// top
-					RHICmdList.SetScissorRect(true, SceneColorViewport.Rect.Min.X, SceneColorViewport.Rect.Min.Y, SceneColorViewport.Rect.Max.X, InnerRect.Min.Y);
-					DrawClearQuad(RHICmdList, false, FLinearColor(), true, (float)ERHIZBuffer::FarPlane, true, 0, SceneColorViewport.Extent, FIntRect());
+					RHICmdList.SetScissorRect(true, DepthStencilViewport.Rect.Min.X, DepthStencilViewport.Rect.Min.Y, DepthStencilViewport.Rect.Max.X, InnerRect.Min.Y);
+					DrawClearQuad(RHICmdList, false, FLinearColor(), true, (float)ERHIZBuffer::FarPlane, true, 0, DepthStencilViewport.Extent, FIntRect());
 					// bottom
-					RHICmdList.SetScissorRect(true, SceneColorViewport.Rect.Min.X, InnerRect.Max.Y, SceneColorViewport.Rect.Max.X, SceneColorViewport.Rect.Max.Y);
-					DrawClearQuad(RHICmdList, false, FLinearColor(), true, (float)ERHIZBuffer::FarPlane, true, 0, SceneColorViewport.Extent, FIntRect());
+					RHICmdList.SetScissorRect(true, DepthStencilViewport.Rect.Min.X, InnerRect.Max.Y, DepthStencilViewport.Rect.Max.X, DepthStencilViewport.Rect.Max.Y);
+					DrawClearQuad(RHICmdList, false, FLinearColor(), true, (float)ERHIZBuffer::FarPlane, true, 0, DepthStencilViewport.Extent, FIntRect());
 					// left
-					RHICmdList.SetScissorRect(true, SceneColorViewport.Rect.Min.X, SceneColorViewport.Rect.Min.Y, InnerRect.Min.X, SceneColorViewport.Rect.Max.Y);
-					DrawClearQuad(RHICmdList, false, FLinearColor(), true, (float)ERHIZBuffer::FarPlane, true, 0, SceneColorViewport.Extent, FIntRect());
+					RHICmdList.SetScissorRect(true, DepthStencilViewport.Rect.Min.X, DepthStencilViewport.Rect.Min.Y, InnerRect.Min.X, DepthStencilViewport.Rect.Max.Y);
+					DrawClearQuad(RHICmdList, false, FLinearColor(), true, (float)ERHIZBuffer::FarPlane, true, 0, DepthStencilViewport.Extent, FIntRect());
 					// right
-					RHICmdList.SetScissorRect(true, InnerRect.Max.X, SceneColorViewport.Rect.Min.Y, SceneColorViewport.Rect.Max.X, SceneColorViewport.Rect.Max.Y);
-					DrawClearQuad(RHICmdList, false, FLinearColor(), true, (float)ERHIZBuffer::FarPlane, true, 0, SceneColorViewport.Extent, FIntRect());
+					RHICmdList.SetScissorRect(true, InnerRect.Max.X, DepthStencilViewport.Rect.Min.Y, DepthStencilViewport.Rect.Max.X, DepthStencilViewport.Rect.Max.Y);
+					DrawClearQuad(RHICmdList, false, FLinearColor(), true, (float)ERHIZBuffer::FarPlane, true, 0, DepthStencilViewport.Extent, FIntRect());
 
 					RHICmdList.SetScissorRect(false, 0, 0, 0, 0);
 				}
@@ -209,6 +241,12 @@ FScreenPassTexture AddSelectionOutlinePass(
 
 		FRHISamplerState* PointClampSampler = TStaticSamplerState<SF_Point, AM_Clamp, AM_Clamp, AM_Clamp>::GetRHI();
 
+		// In side-by-side stereo, the SceneColor texture can come in as single- or double-wide depending on what post-processing passes are enabled, and it can also be pre- or post-primary upscale
+		// If it's double-wide, the left eye will always be on the left side of the texture, but the right eye may be on the left (0,0 aligned) or the right depending on preceding passes
+		// If an instanced secondary view is on the left, it's depth stencil data will still be on the right, so we need to apply an offset to the outline shader when using instanced stencil data
+		const bool bIsInstancedSecondaryView = bIsInstancedStereoView && View.StereoPass == EStereoscopicPass::eSSP_SECONDARY;
+		const float SecondaryViewOffset = (bIsInstancedSecondaryView && Inputs.SceneColor.ViewRect.Min.X == 0) ? View.ViewRect.Min.X : 0.0f;
+
 		FSelectionOutlinePS::FParameters* PassParameters = GraphBuilder.AllocParameters<FSelectionOutlinePS::FParameters>();
 		PassParameters->RenderTargets[0] = Output.GetRenderTargetBinding();
 		PassParameters->View = View.ViewUniformBuffer;
@@ -230,6 +268,7 @@ FScreenPassTexture AddSelectionOutlinePass(
 		PassParameters->OutlineColorIndexBits = 3;
 		PassParameters->SelectionHighlightIntensity = GEngine->SelectionHighlightIntensity;
 		PassParameters->BSPSelectionIntensity = GEngine->BSPSelectionHighlightIntensity;
+		PassParameters->SecondaryViewportOffset = SecondaryViewOffset;
 
 		EDisplayOutputFormat DisplayOutputFormat = View.Family->RenderTarget->GetDisplayOutputFormat();
 		bool bIsHDR = false;
diff --git a/Engine/Source/Runtime/Renderer/Private/PostProcess/PostProcessSelectionOutline.h b/Engine/Source/Runtime/Renderer/Private/PostProcess/PostProcessSelectionOutline.h
--- a/Engine/Source/Runtime/Renderer/Private/PostProcess/PostProcessSelectionOutline.h
+++ b/Engine/Source/Runtime/Renderer/Private/PostProcess/PostProcessSelectionOutline.h
@@ -33,6 +33,7 @@ FScreenPassTexture AddSelectionOutlinePass(
 	const FViewInfo& View,
 	FSceneUniformBuffer &SceneUniformBuffer,
 	const FSelectionOutlineInputs& Inputs,
-	const Nanite::FRasterResults* NaniteRasterResults);
+	const Nanite::FRasterResults* NaniteRasterResults,
+	FRDGTextureRef& DepthStencilTexture);
 
 #endif
diff --git a/Engine/Source/Runtime/Renderer/Private/PostProcess/PostProcessing.cpp b/Engine/Source/Runtime/Renderer/Private/PostProcess/PostProcessing.cpp
--- a/Engine/Source/Runtime/Renderer/Private/PostProcess/PostProcessing.cpp
+++ b/Engine/Source/Runtime/Renderer/Private/PostProcess/PostProcessing.cpp
@@ -292,7 +292,8 @@ void AddPostProcessingPasses(
 	FVirtualShadowMapArray* VirtualShadowMapArray, 
 	FLumenSceneFrameTemporaries& LumenFrameTemporaries,
 	const FSceneWithoutWaterTextures& SceneWithoutWaterTextures,
-	FScreenPassTexture TSRFlickeringInput)
+	FScreenPassTexture TSRFlickeringInput,
+	FRDGTextureRef& InstancedEditorDepthTexture)
 {
 	RDG_CSV_STAT_EXCLUSIVE_SCOPE(GraphBuilder, RenderPostProcessing);
 	QUICK_SCOPE_CYCLE_COUNTER(STAT_PostProcessing_Process);
@@ -471,7 +472,7 @@ void AddPostProcessingPasses(
 	PassSequence.SetEnabled(EPass::VisualizeSkyLightIlluminanceMeter, Scene&& Scene->SkyLight && View.Family && View.Family->EngineShowFlags.VisualizeSkyLightIlluminance);
 	PassSequence.SetEnabled(EPass::VisualizeLightFunctionAtlas, Scene && Scene->LightFunctionAtlasSceneData.GetLightFunctionAtlasEnabled() && View.Family && View.Family->EngineShowFlags.VisualizeLightFunctionAtlas);
 	PassSequence.SetEnabled(EPass::VisualizeLevelInstance, GIsEditor && EngineShowFlags.EditingLevelInstance && EngineShowFlags.VisualizeLevelInstanceEditing && !bVisualizeHDR);
-	PassSequence.SetEnabled(EPass::SelectionOutline, GIsEditor && EngineShowFlags.Selection && EngineShowFlags.SelectionOutline && !EngineShowFlags.Wireframe && !bVisualizeHDR && !IStereoRendering::IsStereoEyeView(View));
+	PassSequence.SetEnabled(EPass::SelectionOutline, GIsEditor && EngineShowFlags.Selection && EngineShowFlags.SelectionOutline && !EngineShowFlags.Wireframe && !bVisualizeHDR);
 	PassSequence.SetEnabled(EPass::EditorPrimitive, FSceneRenderer::ShouldCompositeEditorPrimitives(View));
 #else
 	PassSequence.SetEnabled(EPass::VisualizeSkyAtmosphere, false);
@@ -1489,7 +1490,7 @@ void AddPostProcessingPasses(
 		PassInputs.SceneDepth = SceneDepth;
 		PassInputs.SceneTextures.SceneTextures = Inputs.SceneTextures;
 
-		SceneColor = AddSelectionOutlinePass(GraphBuilder, View, SceneUniformBuffer, PassInputs, NaniteRasterResults);
+		SceneColor = AddSelectionOutlinePass(GraphBuilder, View, SceneUniformBuffer, PassInputs, NaniteRasterResults, InstancedEditorDepthTexture);
 	}
 
 	if (PassSequence.IsEnabled(EPass::EditorPrimitive))
@@ -1954,7 +1955,8 @@ void AddDebugViewPostProcessingPasses(FRDGBuilder& GraphBuilder, const FViewInfo
 		PassInputs.SceneDepth = SceneDepth;
 		PassInputs.SceneTextures.SceneTextures = Inputs.SceneTextures;
 
-		SceneColor = AddSelectionOutlinePass(GraphBuilder, View, SceneUniformBuffer, PassInputs, NaniteRasterResults);
+		FRDGTextureRef DummyStencilTexture = nullptr;
+		SceneColor = AddSelectionOutlinePass(GraphBuilder, View, SceneUniformBuffer, PassInputs, NaniteRasterResults, DummyStencilTexture);
 	}
 #endif
 
@@ -2772,7 +2774,8 @@ void AddMobilePostProcessingPasses(FRDGBuilder& GraphBuilder, FScene* Scene, con
 		PassInputs.OverrideOutput.LoadAction = View.IsFirstInFamily() ? ERenderTargetLoadAction::EClear : ERenderTargetLoadAction::ELoad;
 
 		// TODO: Nanite - pipe through results
-		SceneColor = AddSelectionOutlinePass(GraphBuilder, View, SceneUniformBuffer, PassInputs, nullptr);
+		FRDGTextureRef DummyStencilTexture = nullptr;
+		SceneColor = AddSelectionOutlinePass(GraphBuilder, View, SceneUniformBuffer, PassInputs, nullptr, DummyStencilTexture);
 	}
 
 	if (PassSequence.IsEnabled(EPass::EditorPrimitive))
diff --git a/Engine/Source/Runtime/Renderer/Private/PostProcess/PostProcessing.h b/Engine/Source/Runtime/Renderer/Private/PostProcess/PostProcessing.h
--- a/Engine/Source/Runtime/Renderer/Private/PostProcess/PostProcessing.h
+++ b/Engine/Source/Runtime/Renderer/Private/PostProcess/PostProcessing.h
@@ -43,7 +43,8 @@ void AddPostProcessingPasses(
 	FVirtualShadowMapArray* VirtualShadowMapArray,
 	struct FLumenSceneFrameTemporaries& LumenFrameTemporaries,
 	const FSceneWithoutWaterTextures& SceneWithoutWaterTextures,
-	FScreenPassTexture TSRFlickeringInput);
+	FScreenPassTexture TSRFlickeringInput,
+	FRDGTextureRef& InstancedEditorDepthTexture);
 
 void AddDebugViewPostProcessingPasses(FRDGBuilder& GraphBuilder, const FViewInfo& View, FSceneUniformBuffer &SceneUniformBuffer, const FPostProcessingInputs& Inputs, const Nanite::FRasterResults* NaniteRasterResults);
 
