Screen Space Denoise split screen fixes.  Ray traced sky light View loop needs to wait until the last view to update the output reference, because the variable is used elsewhere in the loop.  Denoiser had a math bug causing screen coordinates to be computed incorrectly for any viewport that didn't start at the origin, breaking world position coordinate calculations, finally leading to artifacts where noise fails to be averaged away due to neighbor samples having incorrectly computed weights (the broken world position calculation makes the samples appear far away).

#rnx

diff --git a/Engine/Source/Runtime/Renderer/Private/RayTracing/RaytracingSkylight.cpp b/Engine/Source/Runtime/Renderer/Private/RayTracing/RaytracingSkylight.cpp
--- a/Engine/Source/Runtime/Renderer/Private/RayTracing/RaytracingSkylight.cpp
+++ b/Engine/Source/Runtime/Renderer/Private/RayTracing/RaytracingSkylight.cpp
@@ -427,6 +427,7 @@ void FDeferredShadingSceneRenderer::RenderRayTracingSkyLight(
 	FSceneTextureParameters SceneTextures = GetSceneTextureParameters(GraphBuilder, Views[0]);
 
 	int32 ViewIndex = 0;
+	int32 LastViewIndex = Views.Num() - 1;
 	for (FViewInfo& View : Views)
 	{
 		RDG_GPU_MASK_SCOPE(GraphBuilder, View.GPUMask);
@@ -528,7 +529,11 @@ void FDeferredShadingSceneRenderer::RenderRayTracingSkyLight(
 					DenoiserInputs,
 					RayTracingConfig);
 
-				OutSkyLightTexture = DenoiserOutputs.Color;
+				// Need to set output used by the caller on the last iteration of the view loop
+				if (ViewIndex == LastViewIndex)
+				{
+					OutSkyLightTexture = DenoiserOutputs.Color;
+				}
 			}
 		}
 
diff --git a/Engine/Source/Runtime/Renderer/Private/ScreenSpaceDenoise.cpp b/Engine/Source/Runtime/Renderer/Private/ScreenSpaceDenoise.cpp
--- a/Engine/Source/Runtime/Renderer/Private/ScreenSpaceDenoise.cpp
+++ b/Engine/Source/Runtime/Renderer/Private/ScreenSpaceDenoise.cpp
@@ -1288,8 +1288,8 @@ void Denoiser::SetupCommonShaderParameters(
 
 	OutPublicCommonParameters->SceneBufferUVToScreenPosition.X = float(FullResBufferExtent.X) / float(View.ViewRect.Width()) * 2.0f;
 	OutPublicCommonParameters->SceneBufferUVToScreenPosition.Y = -float(FullResBufferExtent.Y) / float(View.ViewRect.Height()) * 2.0f;
-	OutPublicCommonParameters->SceneBufferUVToScreenPosition.Z = float(View.ViewRect.Min.X) / float(View.ViewRect.Width()) * 2.0f - 1.0f;
-	OutPublicCommonParameters->SceneBufferUVToScreenPosition.W = -float(View.ViewRect.Min.Y) / float(View.ViewRect.Height()) * 2.0f + 1.0f;
+	OutPublicCommonParameters->SceneBufferUVToScreenPosition.Z = -float(View.ViewRect.Min.X) / float(View.ViewRect.Width()) * 2.0f - 1.0f;
+	OutPublicCommonParameters->SceneBufferUVToScreenPosition.W = float(View.ViewRect.Min.Y) / float(View.ViewRect.Height()) * 2.0f + 1.0f;
 
 	OutPublicCommonParameters->DenoiserBufferBilinearUVMinMax = FVector4f(
 		float(DenoiserViewport.Min.X + 0.5f) / float(DenoiserBufferExtent.X),
@@ -1922,8 +1922,8 @@ static void DenoiseSignalAtConstantPixelDensity(
 
 			PassParameters->PrevSceneBufferUVToScreenPosition.X = float(PrevFrameBufferExtent.X) / float(ViewportExtent.X) * 2.0f;
 			PassParameters->PrevSceneBufferUVToScreenPosition.Y = -float(PrevFrameBufferExtent.Y) / float(ViewportExtent.Y) * 2.0f;
-			PassParameters->PrevSceneBufferUVToScreenPosition.Z = float(ViewportOffset.X) / float(ViewportExtent.X) * 2.0f - 1.0f;
-			PassParameters->PrevSceneBufferUVToScreenPosition.W = -float(ViewportOffset.Y) / float(ViewportExtent.Y) * 2.0f + 1.0f;
+			PassParameters->PrevSceneBufferUVToScreenPosition.Z = -float(ViewportOffset.X) / float(ViewportExtent.X) * 2.0f - 1.0f;
+			PassParameters->PrevSceneBufferUVToScreenPosition.W = float(ViewportOffset.Y) / float(ViewportExtent.Y) * 2.0f + 1.0f;
 		}
 
 		if (bGlobalCameraCut)
@@ -1966,8 +1966,8 @@ static void DenoiseSignalAtConstantPixelDensity(
 
 			PassParameters->PrevSceneBufferUVToScreenPosition.X = float(PrevFrameBufferExtent.X) / float(ViewportExtent.X) * 2.0f;
 			PassParameters->PrevSceneBufferUVToScreenPosition.Y = -float(PrevFrameBufferExtent.Y) / float(ViewportExtent.Y) * 2.0f;
-			PassParameters->PrevSceneBufferUVToScreenPosition.Z = float(ViewportOffset.X) / float(ViewportExtent.X) * 2.0f - 1.0f;
-			PassParameters->PrevSceneBufferUVToScreenPosition.W = -float(ViewportOffset.Y) / float(ViewportExtent.Y) * 2.0f + 1.0f;
+			PassParameters->PrevSceneBufferUVToScreenPosition.Z = -float(ViewportOffset.X) / float(ViewportExtent.X) * 2.0f - 1.0f;
+			PassParameters->PrevSceneBufferUVToScreenPosition.W = float(ViewportOffset.Y) / float(ViewportExtent.Y) * 2.0f + 1.0f;
 		}
 
 		FScreenSpaceDenoiserHistory DummyPrevFrameHistory;
