Reenable depth bound test when computing shadow mask for hair shadow.

This reduce shadow cost considerably when hair does not cover the entire screen (0.6ms -> 0.07ms on XSX).

diff --git a/Engine/Source/Runtime/Renderer/Private/ShadowRendering.cpp b/Engine/Source/Runtime/Renderer/Private/ShadowRendering.cpp
--- a/Engine/Source/Runtime/Renderer/Private/ShadowRendering.cpp
+++ b/Engine/Source/Runtime/Renderer/Private/ShadowRendering.cpp
@@ -1290,7 +1290,7 @@ void FProjectedShadowInfo::RenderProjectionInternal(
 
 	const bool bSubPixelSupport = HairStrandsUniformBuffer != nullptr;// HairStrands::HasViewHairStrandsData(*View);
 	const bool bStencilTestEnabled = !bSubPixelSupport && GShadowStencilCulling;
-	const bool bDepthBoundsTestEnabled = IsWholeSceneDirectionalShadow() && GSupportsDepthBoundsTest && CVarCSMDepthBoundsTest.GetValueOnRenderThread() != 0 && !bSubPixelSupport;
+	const bool bDepthBoundsTestEnabled = IsWholeSceneDirectionalShadow() && GSupportsDepthBoundsTest && CVarCSMDepthBoundsTest.GetValueOnRenderThread() != 0;// && !bSubPixelSupport;
 	const uint32 StencilRef = bSubPixelSupport && !IsWholeSceneDirectionalShadow() && !bCameraInsideShadowFrustum ? 1u : 0u;
 
 	if (!bDepthBoundsTestEnabled && bStencilTestEnabled)
@@ -1299,7 +1299,7 @@ void FProjectedShadowInfo::RenderProjectionInternal(
 	}
 
 	// Mark stencil so that only hair pixel within volume bound will be affected by the pre-shadow mask
-	if (bSubPixelSupport)
+	if (bSubPixelSupport && !bDepthBoundsTestEnabled)
 	{
 		DrawClearQuad(RHICmdList, false, FLinearColor::Transparent, false, 0, true, 0);
 		if (!IsWholeSceneDirectionalShadow())
@@ -1358,7 +1358,7 @@ void FProjectedShadowInfo::RenderProjectionInternal(
 	}
 	else
 	{
-		if (bSubPixelSupport)
+		if (bSubPixelSupport && !bDepthBoundsTestEnabled)
 		{
 			GraphicsPSOInit.DepthStencilState = TStaticDepthStencilState<false, CF_DepthFartherOrEqual, true, CF_Equal, SO_Keep, SO_Keep, SO_Keep, true, CF_Equal, SO_Keep, SO_Keep, SO_Keep, 0xFF, 0xFF>::GetRHI();
 		}
