Change hair voxel texture tilling mode.

diff --git a/Engine/Source/Runtime/Renderer/Private/HairStrands/HairStrandsVoxelization.cpp b/Engine/Source/Runtime/Renderer/Private/HairStrands/HairStrandsVoxelization.cpp
--- a/Engine/Source/Runtime/Renderer/Private/HairStrands/HairStrandsVoxelization.cpp
+++ b/Engine/Source/Runtime/Renderer/Private/HairStrands/HairStrandsVoxelization.cpp
@@ -917,7 +917,7 @@ static FHairStrandsVoxelResources AllocateVirtualVoxelResources(
 			Out.Parameters.Common.PageTextureResolution,
 			PF_R32_UINT, 
 			FClearValueBinding::Black, 
-			TexCreate_UAV | TexCreate_ShaderResource | TexCreate_NoFastClear,
+			TexCreate_UAV | TexCreate_ShaderResource | TexCreate_NoFastClear | TexCreate_3DTiling,
 			MipCount);
 		Out.PageTexture = GraphBuilder.CreateTexture(Desc, TEXT("Hair.VoxelPageTexture"));
 	}
