Enabled hierarchical instance culling & precomputation by default.

This improves instance culling performance for dense scenes like CitySample significantly, in particular when there are many local lights. Top level controlling cvar is r.SceneCulling.

diff --git a/Engine/Source/Runtime/Renderer/Private/SceneCulling/SceneCulling.cpp b/Engine/Source/Runtime/Renderer/Private/SceneCulling/SceneCulling.cpp
--- a/Engine/Source/Runtime/Renderer/Private/SceneCulling/SceneCulling.cpp
+++ b/Engine/Source/Runtime/Renderer/Private/SceneCulling/SceneCulling.cpp
@@ -105,7 +105,7 @@ LLM_DEFINE_TAG(SceneCulling, NAME_None, NAME_None, GET_STATFNAME(STAT_SceneCulli
 
 static TAutoConsoleVariable<int32> CVarSceneCulling(
 	TEXT("r.SceneCulling"), 
-	0, 
+	1, 
 	TEXT("Enable/Disable scene culling.\n")
 	TEXT("  While enabled, it will only build the instance hierarchy if used by any system - currently that corresponds to Nanite being enabled.\n")
 	TEXT("  Forces a recreate of all render state since (at present) there is only an incremental update path."),
@@ -117,7 +117,7 @@ static TAutoConsoleVariable<int32> CVarSceneCulling(
 
 static TAutoConsoleVariable<int32> CVarSceneCullingPrecomputed(
 	TEXT("r.SceneCulling.Precomputed"), 
-	0, 
+	1, 
 	TEXT("Enable/Disable precomputed spatial hashes for scene culling."),
 	ECVF_RenderThreadSafe | ECVF_ReadOnly);
 
