Heterogeneous Volumes: Adding CVars for volumetric shadow map creation and changing defaults

Default interpretation of the ortho grid builder has changed from pixel-space units to world-space units.
This can be reverted to the legacy behavior by toggling the following CVar:
r.HeterogeneousVolumes.OrthoGrid.VoxelizationMode 0

diff --git a/Engine/Shaders/Private/HeterogeneousVolumes/HeterogeneousVolumesVoxelGridShadows.usf b/Engine/Shaders/Private/HeterogeneousVolumes/HeterogeneousVolumesVoxelGridShadows.usf
--- a/Engine/Shaders/Private/HeterogeneousVolumes/HeterogeneousVolumesVoxelGridShadows.usf
+++ b/Engine/Shaders/Private/HeterogeneousVolumes/HeterogeneousVolumesVoxelGridShadows.usf
@@ -13,10 +13,11 @@
 
 // Ray data
 int bJitter;
+float ShadowStepSize;
 
 // Dispatch data
 int3 GroupCount;
-int DebugMode;
+int ShadowDebugTweak;
 
 int MaxStepCount;
 
@@ -122,7 +123,7 @@ float GetStepFactor()
 
 float GetShadowStepSize()
 {
-	return 1.0;
+	return ShadowStepSize;
 }
 
 float GetShadowStepFactor()
@@ -282,7 +283,7 @@ void RenderVolumetricShadowMapForLightWithVoxelGridCS(
 			WorldRayDirection,
 			// Ray-step attributes
 			Jitter,
-			CalcStepSize(WorldRayDirection),
+			CalcShadowStepSize(WorldRayDirection),
 			MaxStepCount,
 			//MaxSampleCount, // Replace max-steps with max-samples..
 			bApplyEmissionAndTransmittance,
@@ -442,7 +443,7 @@ void RenderVolumetricShadowMapForLightWithVoxelGridCS(
 	RWDebugBuffer[LinearIndex].RayEnd = WorldRayEnd;
 	RWDebugBuffer[LinearIndex].HitSpan = float2(HitSpan, Depth);
 #endif
-	if (Face == View.GeneralPurposeTweak)
+	if (Face == ShadowDebugTweak)
 	{
 		RWBeerShadowMapTexture[PixelCoord] = float4(Transmittance, 1);
 	}
diff --git a/Engine/Source/Runtime/Renderer/Private/HeterogeneousVolumes/HeterogeneousVolumes.cpp b/Engine/Source/Runtime/Renderer/Private/HeterogeneousVolumes/HeterogeneousVolumes.cpp
--- a/Engine/Source/Runtime/Renderer/Private/HeterogeneousVolumes/HeterogeneousVolumes.cpp
+++ b/Engine/Source/Runtime/Renderer/Private/HeterogeneousVolumes/HeterogeneousVolumes.cpp
@@ -507,7 +507,8 @@ void FDeferredShadingSceneRenderer::RenderHeterogeneousVolumeShadows(
 	SCOPED_NAMED_EVENT(HeterogeneousVolumes, FColor::Emerald);
 
 	FVoxelGridBuildOptions BuildOptions;
-	BuildOptions.MinimumVoxelSizeOutsideFrustum = 10.0;
+	BuildOptions.MinimumVoxelSizeOutsideFrustum = HeterogeneousVolumes::GetOutOfFrustumShadingRateForShadows();
+	BuildOptions.MinimumVoxelSizeInFrustum = HeterogeneousVolumes::GetShadingRateForShadows();
 	BuildOptions.bBuildOrthoGrid = true;
 	BuildOptions.bBuildFrustumGrid = false;
 	BuildOptions.bJitter = false;
diff --git a/Engine/Source/Runtime/Renderer/Private/HeterogeneousVolumes/HeterogeneousVolumes.h b/Engine/Source/Runtime/Renderer/Private/HeterogeneousVolumes/HeterogeneousVolumes.h
--- a/Engine/Source/Runtime/Renderer/Private/HeterogeneousVolumes/HeterogeneousVolumes.h
+++ b/Engine/Source/Runtime/Renderer/Private/HeterogeneousVolumes/HeterogeneousVolumes.h
@@ -59,6 +59,8 @@ namespace HeterogeneousVolumes
 	float GetMaxStepCount();
 	float GetMinimumVoxelSizeInFrustum();
 	float GetMinimumVoxelSizeOutsideFrustum();
+	float GetShadingRateForShadows();
+	float GetOutOfFrustumShadingRateForShadows();
 
 	int32 GetMipLevel();
 	int32 GetDebugMode();
@@ -208,6 +210,7 @@ END_UNIFORM_BUFFER_STRUCT()
 struct FVoxelGridBuildOptions
 {
 	float MinimumVoxelSizeOutsideFrustum = HeterogeneousVolumes::GetMinimumVoxelSizeOutsideFrustum();
+	float MinimumVoxelSizeInFrustum = HeterogeneousVolumes::GetMinimumVoxelSizeInFrustum();
 
 	bool bBuildOrthoGrid = true;
 	bool bBuildFrustumGrid = true;
diff --git a/Engine/Source/Runtime/Renderer/Private/HeterogeneousVolumes/HeterogeneousVolumesVoxelGridPipeline.cpp b/Engine/Source/Runtime/Renderer/Private/HeterogeneousVolumes/HeterogeneousVolumesVoxelGridPipeline.cpp
--- a/Engine/Source/Runtime/Renderer/Private/HeterogeneousVolumes/HeterogeneousVolumesVoxelGridPipeline.cpp
+++ b/Engine/Source/Runtime/Renderer/Private/HeterogeneousVolumes/HeterogeneousVolumesVoxelGridPipeline.cpp
@@ -68,7 +68,16 @@ static TAutoConsoleVariable<int32> CVarHeterogeneousVolumesEnableOrthoVoxelGrid(
 static TAutoConsoleVariable<float> CVarHeterogeneousVolumesOrthoGridShadingRate(
 	TEXT("r.HeterogeneousVolumes.OrthoGrid.ShadingRate"),
 	4.0,
-	TEXT("The voxel tessellation rate, in pixel-space (Default = 4.0)"),
+	TEXT("The voxelization rate (Default = 4.0)"),
+	ECVF_RenderThreadSafe
+);
+
+static TAutoConsoleVariable<int32> CVarHeterogeneousVolumesOrthoGridVoxelizationMode(
+	TEXT("r.HeterogeneousVolumes.OrthoGrid.VoxelizationMode"),
+	1,
+	TEXT("Voxelization mode (Default = 1)\n")
+	TEXT("0: Screen-space voxel size (legacy behavior)\n")
+	TEXT("1: World-space voxel size\n"),
 	ECVF_RenderThreadSafe
 );
 
@@ -196,15 +205,22 @@ static TAutoConsoleVariable<int32> CVarHeterogeneousVolumesShadowCameraDownsampl
 
 static TAutoConsoleVariable<int32> CVarHeterogeneousVolumesShadowResolution(
 	TEXT("r.HeterogeneousVolumes.Shadows.Resolution"),
-	256,
-	TEXT("Resolution when building volumetric shadow map (Default = 256)\n"),
+	512,
+	TEXT("Resolution when building volumetric shadow map (Default = 512)\n"),
+	ECVF_RenderThreadSafe
+);
+
+static TAutoConsoleVariable<float> CVarHeterogeneousVolumesStepSizeForShadows(
+	TEXT("r.HeterogeneousVolumes.Shadows.StepSize"),
+	2.0,
+	TEXT("Ray marching step size when building volumetric shadow map (Default = 2.0)\n"),
 	ECVF_RenderThreadSafe
 );
 
 static TAutoConsoleVariable<int32> CVarHeterogeneousVolumesShadowMaxSampleCount(
 	TEXT("r.HeterogeneousVolumes.Shadows.MaxSampleCount"),
-	32,
-	TEXT("Maximum sample count when building volumetric shadow map (Default = 32)\n"),
+	16,
+	TEXT("Maximum sample count when building volumetric shadow map (Default = 16)\n"),
 	ECVF_RenderThreadSafe
 );
 
@@ -236,6 +252,20 @@ static TAutoConsoleVariable<int32> CVarHeterogeneousVolumesShadowDebugTweak(
 	ECVF_RenderThreadSafe
 );
 
+static TAutoConsoleVariable<int32> CVarHeterogeneousVolumesShadowShadingRate(
+	TEXT("r.HeterogeneousVolumes.Shadows.ShadingRate"),
+	2,
+	TEXT("Debug tweak value (Default = 0)\n"),
+	ECVF_RenderThreadSafe
+);
+
+static TAutoConsoleVariable<int32> CVarHeterogeneousVolumesShadowOutOfFrustumShadingRate(
+	TEXT("r.HeterogeneousVolumes.Shadows.OutOfFrustumShadingRate"),
+	2,
+	TEXT("Debug tweak value (Default = 0)\n"),
+	ECVF_RenderThreadSafe
+);
+
 IMPLEMENT_UNIFORM_BUFFER_STRUCT(FOrthoVoxelGridUniformBufferParameters, "OrthoGridUniformBuffer");
 IMPLEMENT_UNIFORM_BUFFER_STRUCT(FFrustumVoxelGridUniformBufferParameters, "FrustumGridUniformBuffer");
 IMPLEMENT_UNIFORM_BUFFER_STRUCT(FAdaptiveVolumetricShadowMapUniformBufferParameters, "AVSM");
@@ -379,6 +409,11 @@ namespace HeterogeneousVolumes
 		return FIntPoint(FMath::Clamp(CVarHeterogeneousVolumesShadowResolution.GetValueOnRenderThread(), 1, 1024));
 	}
 
+	float GetStepSizeForShadows()
+	{
+		return FMath::Max(CVarHeterogeneousVolumesStepSizeForShadows.GetValueOnRenderThread(), 0.01f);
+	}
+
 	uint32 GetShadowMaxSampleCount()
 	{
 		return FMath::Clamp(CVarHeterogeneousVolumesShadowMaxSampleCount.GetValueOnRenderThread(), 2, 64);
@@ -403,6 +438,16 @@ namespace HeterogeneousVolumes
 	{
 		return FMath::Max(CVarHeterogeneousVolumesShadowCameraDownsampleFactor.GetValueOnRenderThread(), 1);
 	}
+
+	float GetShadingRateForShadows()
+	{
+		return FMath::Max(CVarHeterogeneousVolumesShadowShadingRate.GetValueOnRenderThread(), 0.1);
+	}
+
+	float GetOutOfFrustumShadingRateForShadows()
+	{
+		return FMath::Max(CVarHeterogeneousVolumesShadowOutOfFrustumShadingRate.GetValueOnRenderThread(), 0.1);
+	}
 }
 
 class FMarkTopLevelGridVoxelsForFrustumGrid : public FGlobalShader
@@ -1341,6 +1386,7 @@ void CalcGlobalBoundsAndMinimumVoxelSize(
 	float& GlobalMinimumVoxelSize
 )
 {
+	const bool bConvertToPixelSpace = (CVarHeterogeneousVolumesOrthoGridVoxelizationMode.GetValueOnAnyThread() == 0);
 	TopLevelGridBounds = FBoxSphereBounds(ForceInit);
 	GlobalMinimumVoxelSize = BuildOptions.MinimumVoxelSizeOutsideFrustum;
 
@@ -1386,11 +1432,14 @@ void CalcGlobalBoundsAndMinimumVoxelSize(
 
 					if (View.ViewFrustum.IntersectBox(PrimitiveBounds.Origin, PrimitiveBounds.BoxExtent))
 					{
-						// Bandlimit minimum voxel size request with projected voxel size, based on shading rate
-						FVector VoxelCenter = PrimitiveBounds.Origin;
-						float Distance = FMath::Max(FVector(PrimitiveBounds.Origin - WorldCameraOrigin).Length() - AggregatePrimitiveBounds.BoxExtent.Length(), 0.0);
-						float VoxelWidth = Distance * PixelWidth * HeterogeneousVolumes::GetShadingRateForOrthoGrid();
-						//float VoxelWidth = HeterogeneousVolumes::GetShadingRateForOrthoGrid();
+						float VoxelWidth = BuildOptions.MinimumVoxelSizeInFrustum;
+						// Legacy behavior converts world-space units to pixel-space units
+						if (bConvertToPixelSpace)
+						{
+							FVector VoxelCenter = PrimitiveBounds.Origin;
+							float Distance = FMath::Max(FVector(PrimitiveBounds.Origin - WorldCameraOrigin).Length() - AggregatePrimitiveBounds.BoxExtent.Length(), 0.0);
+							VoxelWidth *= Distance * PixelWidth;
+						}
 
 						float PerVolumeMinimumVoxelSize = FMath::Max(VoxelWidth, HeterogeneousVolume->GetMinimumVoxelSize());
 						ViewMinimumVoxelSize = FMath::Min(PerVolumeMinimumVoxelSize, ViewMinimumVoxelSize);
@@ -1401,8 +1450,8 @@ void CalcGlobalBoundsAndMinimumVoxelSize(
 			}
 		}
 
-		// Clamp per-view minimum voxel-size to in-frustum maximum
-		if (View.ViewFrustum.IntersectBox(AggregatePrimitiveBounds.Origin, AggregatePrimitiveBounds.BoxExtent))
+		// When converting to pixel-space units, clamp per-view minimum voxel-size to in-frustum minimum
+		if (bConvertToPixelSpace && View.ViewFrustum.IntersectBox(AggregatePrimitiveBounds.Origin, AggregatePrimitiveBounds.BoxExtent))
 		{
 			ViewMinimumVoxelSize = FMath::Max(ViewMinimumVoxelSize, HeterogeneousVolumes::GetMinimumVoxelSizeInFrustum());
 		}
@@ -2861,6 +2910,7 @@ class FRenderVolumetricShadowMapForLightWithVoxelGridCS : public FGlobalShader
 
 		// Dispatch data
 		SHADER_PARAMETER(FIntVector, GroupCount)
+		SHADER_PARAMETER(int, ShadowDebugTweak)
 
 		// Output
 		SHADER_PARAMETER_RDG_BUFFER_UAV(RWBuffer<int>, RWVolumetricShadowLinkedListAllocatorBuffer)
@@ -3255,9 +3305,7 @@ void RenderVolumetricShadowMapForLightWithVoxelGrid(
 		PassParameters->SceneTextures = GetSceneTextureParameters(GraphBuilder, SceneTextures);
 
 		// Ray Data
-		PassParameters->ShadowStepSize = HeterogeneousVolumes::GetShadowStepSize();
-		// TODO:
-		//PassParameters->ShadowStepFactor = HeterogeneousVolumes::GetShadowStepFactor();
+		PassParameters->ShadowStepSize = HeterogeneousVolumes::GetStepSizeForShadows();
 		PassParameters->ShadowStepFactor = 1.0;
 		PassParameters->MaxTraceDistance = HeterogeneousVolumes::GetMaxTraceDistance();
 		PassParameters->MaxStepCount = HeterogeneousVolumes::GetMaxStepCount();
@@ -3374,6 +3422,7 @@ void RenderVolumetricShadowMapForLightWithVoxelGrid(
 
 		// Dispatch data
 		PassParameters->GroupCount = GroupCount;
+		PassParameters->ShadowDebugTweak = CVarHeterogeneousVolumesShadowDebugTweak.GetValueOnRenderThread();
 
 		// Output
 		PassParameters->RWVolumetricShadowLinkedListAllocatorBuffer = GraphBuilder.CreateUAV(VolumetricShadowLinkedListAllocatorBuffer, PF_R32_UINT);
@@ -3474,8 +3523,7 @@ void RenderVolumetricShadowMapForCameraWithVoxelGrid(
 		PassParameters->SceneTextures = GetSceneTextureParameters(GraphBuilder, SceneTextures);
 
 		// Ray Data
-		PassParameters->ShadowStepSize = HeterogeneousVolumes::GetShadowStepSize();
-		//PassParameters->ShadowStepFactor = HeterogeneousVolumeInterface->GetShadowStepFactor();
+		PassParameters->ShadowStepSize = HeterogeneousVolumes::GetStepSizeForShadows();
 		PassParameters->ShadowStepFactor = 1.0;
 		PassParameters->MaxTraceDistance = HeterogeneousVolumes::GetMaxTraceDistance();
 		PassParameters->MaxStepCount = HeterogeneousVolumes::GetMaxStepCount();
