Heterogeneous Volumes: Respect CastShadows flag in when building shadow acceleration structure

diff --git a/Engine/Source/Runtime/Engine/Private/Components/HeterogeneousVolumeComponent.cpp b/Engine/Source/Runtime/Engine/Private/Components/HeterogeneousVolumeComponent.cpp
--- a/Engine/Source/Runtime/Engine/Private/Components/HeterogeneousVolumeComponent.cpp
+++ b/Engine/Source/Runtime/Engine/Private/Components/HeterogeneousVolumeComponent.cpp
@@ -72,6 +72,7 @@ FHeterogeneousVolumeSceneProxy::FHeterogeneousVolumeSceneProxy(UHeterogeneousVol
 #endif
 {
 	bIsHeterogeneousVolume = true;
+	bCastDynamicShadow = InComponent->CastShadow;
 
 	HeterogeneousVolumeData.VoxelResolution = InComponent->VolumeResolution;
 
@@ -183,7 +184,7 @@ void FHeterogeneousVolumeSceneProxy::GetDynamicMeshElements(
 		Mesh.MaterialRenderProxy = MaterialInterface->GetRenderProxy();
 		Mesh.LCI = NULL;
 		Mesh.ReverseCulling = IsLocalToWorldDeterminantNegative() ? true : false;
-		Mesh.CastShadow = false;
+		Mesh.CastShadow = CastsDynamicShadow();
 		//Mesh.DepthPriorityGroup = (ESceneDepthPriorityGroup)GetDepthPriorityGroup(View);
 		Mesh.Type = PT_TriangleStrip;
 		Mesh.bDisableBackfaceCulling = true;
@@ -606,7 +607,6 @@ AHeterogeneousVolume::AHeterogeneousVolume(const FObjectInitializer& ObjectIniti
 
 	PrimaryActorTick.bCanEverTick = true;
 	SetHidden(false);
-	HeterogeneousVolumeComponent->bCastDynamicShadow = true;
 }
 
 #undef LOCTEXT_NAMESPACE
\ No newline at end of file
diff --git a/Engine/Source/Runtime/Renderer/Private/HeterogeneousVolumes/HeterogeneousVolumes.cpp b/Engine/Source/Runtime/Renderer/Private/HeterogeneousVolumes/HeterogeneousVolumes.cpp
--- a/Engine/Source/Runtime/Renderer/Private/HeterogeneousVolumes/HeterogeneousVolumes.cpp
+++ b/Engine/Source/Runtime/Renderer/Private/HeterogeneousVolumes/HeterogeneousVolumes.cpp
@@ -507,6 +507,7 @@ void FDeferredShadingSceneRenderer::RenderHeterogeneousVolumeShadows(
 	SCOPED_NAMED_EVENT(HeterogeneousVolumes, FColor::Emerald);
 
 	FVoxelGridBuildOptions BuildOptions;
+	BuildOptions.VoxelGridBuildMode = EVoxelGridBuildMode::Shadows;
 	BuildOptions.MinimumVoxelSizeOutsideFrustum = HeterogeneousVolumes::GetOutOfFrustumShadingRateForShadows();
 	BuildOptions.MinimumVoxelSizeInFrustum = HeterogeneousVolumes::GetShadingRateForShadows();
 	BuildOptions.bBuildOrthoGrid = true;
diff --git a/Engine/Source/Runtime/Renderer/Private/HeterogeneousVolumes/HeterogeneousVolumes.h b/Engine/Source/Runtime/Renderer/Private/HeterogeneousVolumes/HeterogeneousVolumes.h
--- a/Engine/Source/Runtime/Renderer/Private/HeterogeneousVolumes/HeterogeneousVolumes.h
+++ b/Engine/Source/Runtime/Renderer/Private/HeterogeneousVolumes/HeterogeneousVolumes.h
@@ -207,8 +207,15 @@ END_UNIFORM_BUFFER_STRUCT()
 
 // Render specializations
 
+enum class EVoxelGridBuildMode
+{
+	PathTracing,
+	Shadows,
+};
+
 struct FVoxelGridBuildOptions
 {
+	EVoxelGridBuildMode VoxelGridBuildMode = EVoxelGridBuildMode::PathTracing;
 	float MinimumVoxelSizeOutsideFrustum = HeterogeneousVolumes::GetMinimumVoxelSizeOutsideFrustum();
 	float MinimumVoxelSizeInFrustum = HeterogeneousVolumes::GetMinimumVoxelSizeInFrustum();
 
diff --git a/Engine/Source/Runtime/Renderer/Private/HeterogeneousVolumes/HeterogeneousVolumesVoxelGridPipeline.cpp b/Engine/Source/Runtime/Renderer/Private/HeterogeneousVolumes/HeterogeneousVolumesVoxelGridPipeline.cpp
--- a/Engine/Source/Runtime/Renderer/Private/HeterogeneousVolumes/HeterogeneousVolumesVoxelGridPipeline.cpp
+++ b/Engine/Source/Runtime/Renderer/Private/HeterogeneousVolumes/HeterogeneousVolumesVoxelGridPipeline.cpp
@@ -2599,8 +2599,21 @@ const FProjectedShadowInfo* GetProjectedShadowInfo(const FVisibleLightInfo* Visi
 	return nullptr;
 }
 
+bool ShouldAddToVoxelGrid(const FPrimitiveSceneProxy* Proxy, const FViewInfo& View, const FVoxelGridBuildOptions& BuildOptions)
+{
+	switch (BuildOptions.VoxelGridBuildMode)
+	{
+		default:
+		case EVoxelGridBuildMode::PathTracing:
+			return Proxy->IsShown(&View);
+		case EVoxelGridBuildMode::Shadows:
+			return Proxy->IsShadowCast(&View);
+	}
+}
+
 void CollectHeterogeneousVolumeMeshBatches(
 	FRDGBuilder& GraphBuilder,
+	const FVoxelGridBuildOptions& BuildOptions,
 	const FScene* Scene,
 	const TArray<FViewInfo>& Views,
 	const TArray<FVisibleLightInfo, SceneRenderingAllocator> VisibleLightInfos,
@@ -2614,7 +2627,11 @@ void CollectHeterogeneousVolumeMeshBatches(
 		const FViewInfo& View = Views[ViewIndex];
 		for (int32 MeshBatchIndex = 0; MeshBatchIndex < View.HeterogeneousVolumesMeshBatches.Num(); ++MeshBatchIndex)
 		{
-			HeterogeneousVolumesMeshBatches.FindOrAdd(View.HeterogeneousVolumesMeshBatches[MeshBatchIndex]);
+			const FVolumetricMeshBatch& MeshBatch = View.HeterogeneousVolumesMeshBatches[MeshBatchIndex];
+			if (ShouldAddToVoxelGrid(MeshBatch.Proxy, View, BuildOptions))
+			{
+				HeterogeneousVolumesMeshBatches.FindOrAdd(MeshBatch);
+			}
 		}
 	}
 
@@ -2644,7 +2661,11 @@ void CollectHeterogeneousVolumeMeshBatches(
 				const TArray<FMeshBatchAndRelevance, SceneRenderingAllocator>& MeshBatches = ProjectedShadowInfo->GetDynamicSubjectHeterogeneousVolumeMeshElements();
 				for (int32 MeshBatchIndex = 0; MeshBatchIndex < MeshBatches.Num(); ++MeshBatchIndex)
 				{
-					HeterogeneousVolumesMeshBatches.FindOrAdd(FVolumetricMeshBatch(MeshBatches[MeshBatchIndex].Mesh, MeshBatches[MeshBatchIndex].PrimitiveSceneProxy));
+					const FMeshBatchAndRelevance& MeshBatch = MeshBatches[MeshBatchIndex];
+					if (ShouldAddToVoxelGrid(MeshBatch.PrimitiveSceneProxy, *ProjectedShadowInfo->ShadowDepthView, BuildOptions))
+					{
+						HeterogeneousVolumesMeshBatches.FindOrAdd(FVolumetricMeshBatch(MeshBatch.Mesh, MeshBatch.PrimitiveSceneProxy));
+					}
 				}
 			}
 		}
@@ -2673,6 +2694,7 @@ void BuildOrthoVoxelGrid(
 	TSet<FVolumetricMeshBatch> HeterogeneousVolumesMeshBatches;
 	CollectHeterogeneousVolumeMeshBatches(
 		GraphBuilder,
+		BuildOptions,
 		Scene,
 		Views,
 		VisibleLightInfos,
