Rename HairStrandsCluster.h/.cpp into HairStrandsMacroGroup.h/.cpp

diff --git a/Engine/Source/Runtime/Renderer/Private/HairStrands/HairStrandsCluster.cpp b/Engine/Source/Runtime/Renderer/Private/HairStrands/HairStrandsCluster.cpp
deleted file mode 100644
diff --git a/Engine/Source/Runtime/Renderer/Private/HairStrands/HairStrandsCluster.h b/Engine/Source/Runtime/Renderer/Private/HairStrands/HairStrandsCluster.h
deleted file mode 100644
diff --git a/Engine/Source/Runtime/Renderer/Private/HairStrands/HairStrandsDebug.cpp b/Engine/Source/Runtime/Renderer/Private/HairStrands/HairStrandsDebug.cpp
--- a/Engine/Source/Runtime/Renderer/Private/HairStrands/HairStrandsDebug.cpp
+++ b/Engine/Source/Runtime/Renderer/Private/HairStrands/HairStrandsDebug.cpp
@@ -2,7 +2,6 @@
 
 #include "HairStrandsDebug.h"
 #include "HairStrandsInterface.h"
-#include "HairStrandsCluster.h"
 #include "HairStrandsDeepShadow.h"
 #include "HairStrandsUtils.h"
 #include "HairStrandsVoxelization.h"
diff --git a/Engine/Source/Runtime/Renderer/Private/HairStrands/HairStrandsDeepShadow.cpp b/Engine/Source/Runtime/Renderer/Private/HairStrands/HairStrandsDeepShadow.cpp
--- a/Engine/Source/Runtime/Renderer/Private/HairStrands/HairStrandsDeepShadow.cpp
+++ b/Engine/Source/Runtime/Renderer/Private/HairStrands/HairStrandsDeepShadow.cpp
@@ -2,7 +2,6 @@
 
 #include "HairStrandsDeepShadow.h"
 #include "HairStrandsRasterCommon.h"
-#include "HairStrandsCluster.h"
 #include "HairStrandsUtils.h"
 #include "HairStrandsData.h"
 #include "LightSceneInfo.h"
@@ -221,6 +220,7 @@ public:
 IMPLEMENT_GLOBAL_SHADER(FDeepShadowCreateViewInfoCS, "/Engine/Private/HairStrands/HairStrandsDeepShadowAllocation.usf", "CreateViewInfo", SF_Compute);
 ///////////////////////////////////////////////////////////////////////////////////////////////////
 
+bool IsHairStrandsForVoxelTransmittanceAndShadowEnable();
 float GetDeepShadowMaxFovAngle();
 float GetDeepShadowRasterizationScale();
 float GetDeepShadowAABBScale();
diff --git a/Engine/Source/Runtime/Renderer/Private/HairStrands/HairStrandsMacroGroup.cpp b/Engine/Source/Runtime/Renderer/Private/HairStrands/HairStrandsMacroGroup.cpp
new file mode 100644
--- /dev/null
+++ b/Engine/Source/Runtime/Renderer/Private/HairStrands/HairStrandsMacroGroup.cpp
@@ -0,0 +1,314 @@
+// Copyright Epic Games, Inc. All Rights Reserved.
+
+
+#include "HairStrandsMacroGroup.h"
+#include "HairStrandsUtils.h"
+#include "HairStrandsData.h"
+#include "SceneRendering.h"
+#include "SceneManagement.h"
+#include "RendererInterface.h"
+#include "Shader.h"
+#include "GlobalShader.h"
+#include "ShaderParameters.h"
+#include "ShaderParameterStruct.h"
+#include "ScenePrivate.h"
+
+static int32 GHairVirtualVoxel_NumPixelPerVoxel = 1;
+static FAutoConsoleVariableRef CVarHairVirtualVoxel_NumPixelPerVoxel(TEXT("r.HairStrands.Voxelization.VoxelSizeInPixel"), GHairVirtualVoxel_NumPixelPerVoxel, TEXT("Target size of voxel size in pixels"), ECVF_RenderThreadSafe);
+
+class FHairMacroGroupAABBCS : public FGlobalShader
+{
+	DECLARE_GLOBAL_SHADER(FHairMacroGroupAABBCS);
+	SHADER_USE_PARAMETER_STRUCT(FHairMacroGroupAABBCS, FGlobalShader);
+
+	BEGIN_SHADER_PARAMETER_STRUCT(FParameters, )
+		SHADER_PARAMETER(uint32, PassCount)
+		SHADER_PARAMETER(uint32, MacroGroupId)
+		SHADER_PARAMETER(uint32, MacroGroupCount)
+
+		SHADER_PARAMETER(float, PixelSizeAtDepth1)
+		SHADER_PARAMETER(float, NumPixelPerVoxel)
+		SHADER_PARAMETER(uint32, VoxelPageResolution)
+
+		SHADER_PARAMETER_RDG_BUFFER_SRV(Buffer<uint>, RegisteredIndexBuffer)
+		SHADER_PARAMETER_RDG_BUFFER_SRV(Buffer<int>, InGroupAABBBuffer)
+		SHADER_PARAMETER_RDG_BUFFER_UAV(RWBuffer, OutMacroGroupAABBBuffer)
+		SHADER_PARAMETER_RDG_BUFFER_UAV(RWBuffer, OutMacroGroupVoxelSizeBuffer)
+		SHADER_PARAMETER_STRUCT_REF(FViewUniformShaderParameters, View)
+	END_SHADER_PARAMETER_STRUCT()
+
+public:
+	static uint32 GetGroupSize() { return 32; }
+	static bool ShouldCompilePermutation(const FGlobalShaderPermutationParameters& Parameters) { return IsHairStrandsSupported(EHairStrandsShaderType::Strands, Parameters.Platform); }
+	static void ModifyCompilationEnvironment(const FGlobalShaderPermutationParameters& Parameters, FShaderCompilerEnvironment& OutEnvironment)
+	{
+		FGlobalShader::ModifyCompilationEnvironment(Parameters, OutEnvironment);
+		OutEnvironment.SetDefine(TEXT("SHADER_AABBUPDATE"), 1);
+		OutEnvironment.SetDefine(TEXT("GROUP_SIZE"), GetGroupSize());
+	}
+};
+
+IMPLEMENT_GLOBAL_SHADER(FHairMacroGroupAABBCS, "/Engine/Private/HairStrands/HairStrandsAABB.usf", "Main", SF_Compute);
+
+void GetVoxelPageResolution(uint32& OutPageResolution, uint32& OutPageResolutionLog2);
+
+static void AddHairMacroGroupAABBPass(
+	FRDGBuilder& GraphBuilder,
+	const FViewInfo& View, 
+	FHairTransientResources& TransientResources,
+	FHairStrandsMacroGroupData& MacroGroup,
+	FRDGBufferUAVRef& OutHairMacroGroupAABBBufferUAV,
+	FRDGBufferUAVRef& OutMacroGroupVoxelSizeBufferUAV)
+{
+	const uint32 PrimitiveCount = MacroGroup.PrimitivesInfos.Num();
+	if (PrimitiveCount == 0)
+		return;
+	
+	// Compute the average pixel size at a distance of 1 units
+	const FIntPoint Resolution(View.ViewRect.Width(), View.ViewRect.Height());
+	const float vFOV = FMath::DegreesToRadians(View.FOV);
+	const float PixelSizeAtDepth1 = FMath::Tan(vFOV * 0.5f) / (0.5f * Resolution.Y);
+	const uint32 MacroGroupCount = MacroGroup.PrimitivesInfos.Num();
+	
+	uint32 VoxelPageResolution = 0;
+	uint32 VoxelPageResolutionLog2 = 1;
+	GetVoxelPageResolution(VoxelPageResolution, VoxelPageResolutionLog2);
+
+	TArray<uint32> RegisteredIndices;
+	RegisteredIndices.Reserve(MacroGroupCount);
+	for (const FHairStrandsMacroGroupData::PrimitiveInfo& PrimitiveInfo : MacroGroup.PrimitivesInfos)
+	{
+		RegisteredIndices.Add(PrimitiveInfo.PublicDataPtr->Instance->RegisteredIndex);
+	}
+	FRDGBufferRef RegisteredIndexBuffer = CreateVertexBuffer(GraphBuilder, TEXT("Hair.RegisteredIndexBuffer"), FRDGBufferDesc::CreateBufferDesc(4, RegisteredIndices.Num()), RegisteredIndices.GetData(), 4u * RegisteredIndices.Num());
+
+	const float NumPixelPerVoxel = FMath::Clamp(GHairVirtualVoxel_NumPixelPerVoxel, 1.f, 50.f);
+
+	FHairMacroGroupAABBCS::FParameters* Parameters = GraphBuilder.AllocParameters<FHairMacroGroupAABBCS::FParameters>();
+	Parameters->MacroGroupId 					= MacroGroup.MacroGroupId;
+	Parameters->RegisteredIndexBuffer			= GraphBuilder.CreateSRV(RegisteredIndexBuffer, PF_R32_UINT);
+	Parameters->InGroupAABBBuffer				= TransientResources.GroupAABBSRV;
+	Parameters->OutMacroGroupAABBBuffer 		= OutHairMacroGroupAABBBufferUAV;
+	Parameters->OutMacroGroupVoxelSizeBuffer 	= OutMacroGroupVoxelSizeBufferUAV;
+	Parameters->PixelSizeAtDepth1   			= PixelSizeAtDepth1;
+	Parameters->NumPixelPerVoxel    			= NumPixelPerVoxel;
+	Parameters->VoxelPageResolution 			= VoxelPageResolution;
+	Parameters->View			    			= View.ViewUniformBuffer;
+	Parameters->MacroGroupCount					= MacroGroupCount;
+	Parameters->PassCount						= FMath::DivideAndRoundUp(MacroGroupCount, FHairMacroGroupAABBCS::GetGroupSize());
+	
+	TShaderMapRef<FHairMacroGroupAABBCS> ComputeShader(View.ShaderMap);
+	FComputeShaderUtils::AddPass(
+		GraphBuilder,
+		RDG_EVENT_NAME("HairStrands::MacroGroupAABBUpdate"),
+		ComputeShader,
+		Parameters,
+		FIntVector(1, 1, 1));
+}
+
+static bool DoesGroupExists(uint32 ResourceId, uint32 GroupIndex, const FHairStrandsMacroGroupData::TPrimitiveInfos& PrimitivesGroups)
+{
+	// Simple linear search as the expected number of groups is supposed to be low (<16, see FHairStrandsMacroGroupData::MaxMacroGroupCount)
+	for (const FHairStrandsMacroGroupData::PrimitiveInfo& Group : PrimitivesGroups)
+	{
+		if (Group.GroupIndex == GroupIndex && Group.ResourceId == ResourceId)
+		{
+			return true;
+		}
+	}
+	return false;
+}
+
+bool IsHairStrandsNonVisibleShadowCastingEnable();
+
+static void InternalUpdateMacroGroup(FHairStrandsMacroGroupData& MacroGroup, int32& MaterialId, FHairGroupPublicData* HairData, const FMeshBatch* Mesh, const FPrimitiveSceneProxy* Proxy)
+{
+	check(HairData);
+
+	// Track if any instance in the current group needs some features (scatter light scene, holdout, ...)
+	MacroGroup.Flags |= HairData->VFInput.Strands.Common.Flags;
+
+	FHairStrandsMacroGroupData::PrimitiveInfo& PrimitiveInfo = MacroGroup.PrimitivesInfos.AddZeroed_GetRef();
+	PrimitiveInfo.Mesh = Mesh;
+	PrimitiveInfo.PrimitiveSceneProxy = Proxy;
+	PrimitiveInfo.MaterialId = MaterialId++;
+	PrimitiveInfo.ResourceId = Mesh ? reinterpret_cast<uint64>(Mesh->Elements[0].UserData) : ~0u;
+	PrimitiveInfo.GroupIndex = HairData->GetGroupIndex();
+	PrimitiveInfo.Flags = HairData->VFInput.Strands.Common.Flags;
+	PrimitiveInfo.PublicDataPtr = HairData;
+
+	if (HairData->DoesSupportVoxelization())
+	{
+		MacroGroup.bSupportVoxelization = true;
+	}
+}
+
+void CreateHairStrandsMacroGroups(
+	FRDGBuilder& GraphBuilder,
+	const FScene* Scene,
+	const FViewInfo& View, 
+	const TArray<EHairInstanceVisibilityType>& InstancesVisibilityType,
+	FHairStrandsViewData& OutHairStrandsViewData,
+	bool bBuildGPUAABB)
+{
+	const bool bHasHairStrandsElements = View.HairStrandsMeshElements.Num() != 0 || Scene->HairStrandsSceneData.RegisteredProxies.Num() != 0;
+	if (!View.Family || !bHasHairStrandsElements || View.bIsReflectionCapture)
+	{
+		return;
+	}
+
+	TArray<FHairStrandsMacroGroupData, SceneRenderingAllocator>& MacroGroups = OutHairStrandsViewData.MacroGroupDatas;
+
+	int32 MaterialId = 0;
+
+	// Aggregate all hair primitives within the same area into macro groups, for allocating/rendering DOM/voxel
+	uint32 MacroGroupId = 0;
+	auto UpdateMacroGroup = [&MacroGroups, &MacroGroupId, &MaterialId](FHairGroupPublicData* HairData, const FMeshBatch* Mesh,  const FPrimitiveSceneProxy* Proxy, const FBoxSphereBounds& Bounds)
+	{
+		check(HairData);
+
+		// Ensure that the element has been initialized
+		const bool bIsValid = HairData->VFInput.Strands.PositionBufferRHISRV != nullptr;
+		if (!bIsValid)
+			return;
+
+		const FBoxSphereBounds& OriginalPrimitiveBounds = Proxy ? Proxy->GetBounds() : Bounds;
+
+		// Expand hair bound by (half) the groom's max hair-length, to be sure that the bounds are large enough. 
+		// This is important as the primary visibility memory allocation is based on the screen-projection of this CPU bound.
+		// If the bound is too small, the allocation won't be enough, resulting in tile artifacts.
+		FBoxSphereBounds PrimitiveBounds = OriginalPrimitiveBounds.ExpandBy(HairData->VFInput.Strands.Common.Length * 0.5f);
+
+		bool bFound = false;
+		float MinDistance = FLT_MAX;
+		uint32 ClosestMacroGroupId = ~0u;
+		for (FHairStrandsMacroGroupData& MacroGroup : MacroGroups)
+		{
+			const FSphere MacroSphere = MacroGroup.Bounds.GetSphere();
+			const FSphere PrimSphere = PrimitiveBounds.GetSphere();
+
+			const float DistCenters = (MacroSphere.Center - PrimSphere.Center).Size();
+			const float AccumRadius = FMath::Max(0.f, MacroSphere.W + PrimSphere.W);
+			const bool bIntersect = DistCenters <= AccumRadius;
+			
+			if (bIntersect)
+			{
+				MacroGroup.Bounds = Union(MacroGroup.Bounds, PrimitiveBounds);
+
+				InternalUpdateMacroGroup(MacroGroup, MaterialId, HairData, Mesh, Proxy);
+				bFound = true;
+				break;
+			}
+
+			const float MacroToPrimDistance = DistCenters - AccumRadius;
+			if (MacroToPrimDistance < MinDistance)
+			{
+				MinDistance = MacroToPrimDistance;
+				ClosestMacroGroupId = MacroGroup.MacroGroupId;
+			}
+		}
+
+		if (!bFound)
+		{
+			// If we have reached the max number of macro group (MAX_HAIR_MACROGROUP_COUNT), then merge the current one with the closest one.
+			if (MacroGroups.Num() == FHairStrandsMacroGroupData::MaxMacroGroupCount)
+			{
+				check(ClosestMacroGroupId != ~0u);
+				FHairStrandsMacroGroupData& MacroGroup = MacroGroups[ClosestMacroGroupId];
+				check(MacroGroup.MacroGroupId == ClosestMacroGroupId);
+				MacroGroup.Bounds = Union(MacroGroup.Bounds, PrimitiveBounds);
+				InternalUpdateMacroGroup(MacroGroup, MaterialId, HairData, Mesh, Proxy);
+			}
+			else
+			{
+				FHairStrandsMacroGroupData MacroGroup;
+				MacroGroup.MacroGroupId = MacroGroupId++;
+				InternalUpdateMacroGroup(MacroGroup, MaterialId, HairData, Mesh, Proxy);
+				MacroGroup.Bounds = PrimitiveBounds;
+				MacroGroups.Add(MacroGroup);
+			}
+		}
+	};
+
+	// 0. Pre-sort visible instances by register index to get stable macro-groups
+	struct FVisibleBatch
+	{
+		const FMeshBatchAndRelevance* Batch = nullptr;
+		FHairGroupPublicData* HairData = nullptr;
+	};
+	TArray<FVisibleBatch> VisibleBatches;
+	VisibleBatches.Reserve(View.HairStrandsMeshElements.Num());
+	for (const FMeshBatchAndRelevance& MeshBatchAndRelevance : View.HairStrandsMeshElements)
+	{
+		if (HairStrands::IsHairStrandsVF(MeshBatchAndRelevance.Mesh))
+		{
+			if (FHairGroupPublicData* HairData = HairStrands::GetHairData(MeshBatchAndRelevance.Mesh))
+			{
+				VisibleBatches.Add( { &MeshBatchAndRelevance, HairData });
+			}
+		}
+	}
+	VisibleBatches.Sort([](const FVisibleBatch& A, const FVisibleBatch& B) { return A.HairData->Instance->RegisteredIndex < B.HairData->Instance->RegisteredIndex; });
+
+	// 1. Add all visible hair-strands instances
+	static FBoxSphereBounds EmptyBound(ForceInit);
+	for (FVisibleBatch& VisibleBatch : VisibleBatches)
+	{
+		UpdateMacroGroup(VisibleBatch.HairData, VisibleBatch.Batch->Mesh, VisibleBatch.Batch->PrimitiveSceneProxy, EmptyBound);
+	}
+
+	// 2. Add all hair-strands instances which are non-visible in primary view(s) but visible in shadow view(s)
+	if (IsHairStrandsNonVisibleShadowCastingEnable())
+	{
+		for (FHairStrandsInstance* Instance : Scene->HairStrandsSceneData.RegisteredProxies)
+		{
+			if (Instance && InstancesVisibilityType.IsValidIndex(Instance->RegisteredIndex) && InstancesVisibilityType[Instance->RegisteredIndex] == EHairInstanceVisibilityType::StrandsShadowView)
+			{
+				UpdateMacroGroup(const_cast<FHairGroupPublicData*>(Instance->GetHairData()), nullptr, nullptr, Instance->GetBounds());
+			}
+		}
+	}
+
+	// Compute the screen size of macro group projection, for allocation purpose
+	for (FHairStrandsMacroGroupData& MacroGroup : MacroGroups)
+	{
+		MacroGroup.ScreenRect = ComputeProjectedScreenRect(MacroGroup.Bounds.GetBox(), View);
+	}
+	// Sanity check
+	check(MacroGroups.Num() <= FHairStrandsMacroGroupData::MaxMacroGroupCount);
+	check(Scene->HairStrandsSceneData.TransientResources);
+
+	// Build hair macro group AABBB
+	FHairStrandsMacroGroupResources& MacroGroupResources = OutHairStrandsViewData.MacroGroupResources;
+	const uint32 MacroGroupCount = MacroGroups.Num();
+	if (MacroGroupCount > 0 && bBuildGPUAABB)
+	{
+		DECLARE_GPU_STAT(HairStrandsAABB);
+		RDG_EVENT_SCOPE(GraphBuilder, "HairStrandsAABB");
+		RDG_GPU_STAT_SCOPE(GraphBuilder, HairStrandsAABB);
+
+		MacroGroupResources.MacroGroupAABBsBuffer = GraphBuilder.CreateBuffer(FRDGBufferDesc::CreateBufferDesc(4, 6 * MacroGroupCount), TEXT("Hair.MacroGroupAABBBuffer"));
+		MacroGroupResources.MacroGroupVoxelAlignedAABBsBuffer = GraphBuilder.CreateBuffer(FRDGBufferDesc::CreateBufferDesc(4, 6 * MacroGroupCount), TEXT("Hair.MacroGroupVoxelAlignedAABBBuffer"));
+		MacroGroupResources.MacroGroupVoxelSizeBuffer = GraphBuilder.CreateBuffer(FRDGBufferDesc::CreateBufferDesc(2, MacroGroupCount), TEXT("Hair.MacroGroupVoxelSize"));
+		FRDGBufferUAVRef MacroGroupAABBBufferUAV = GraphBuilder.CreateUAV(MacroGroupResources.MacroGroupAABBsBuffer, PF_R32_SINT, ERDGUnorderedAccessViewFlags::SkipBarrier);
+		FRDGBufferUAVRef MacroGroupVoxelSizeBufferUAV = GraphBuilder.CreateUAV(MacroGroupResources.MacroGroupVoxelSizeBuffer, PF_R16F, ERDGUnorderedAccessViewFlags::SkipBarrier);
+		for (FHairStrandsMacroGroupData& MacroGroup : MacroGroups)
+		{				
+			AddHairMacroGroupAABBPass(GraphBuilder, View, *Scene->HairStrandsSceneData.TransientResources, MacroGroup, MacroGroupAABBBufferUAV, MacroGroupVoxelSizeBufferUAV);
+		}
+		MacroGroupResources.MacroGroupCount = MacroGroups.Num();
+	}
+
+	// Aggregate flags accross all instances/macrogroups
+	OutHairStrandsViewData.Flags = 0;
+	for (const FHairStrandsMacroGroupData& MacroGroupData : View.HairStrandsViewData.MacroGroupDatas)
+	{			
+		OutHairStrandsViewData.Flags |= MacroGroupData.Flags;
+	}
+}
+
+bool FHairStrandsMacroGroupData::PrimitiveInfo::IsCullingEnable() const
+{
+	const FHairGroupPublicData* HairData = HairStrands::GetHairData(Mesh);
+	return HairData->GetCullingResultAvailable();
+}
diff --git a/Engine/Source/Runtime/Renderer/Private/HairStrands/HairStrandsMacroGroup.h b/Engine/Source/Runtime/Renderer/Private/HairStrands/HairStrandsMacroGroup.h
new file mode 100644
--- /dev/null
+++ b/Engine/Source/Runtime/Renderer/Private/HairStrands/HairStrandsMacroGroup.h
@@ -0,0 +1,23 @@
+// Copyright Epic Games, Inc. All Rights Reserved.
+
+/*=============================================================================
+	HairStrandsCluster.h: Hair strands macro group computation implementation.
+=============================================================================*/
+
+#pragma once
+
+#include "CoreMinimal.h"
+#include "RendererInterface.h"
+#include "HairStrandsInterface.h"
+
+class FScene;
+class FViewInfo;
+struct FHairStrandsViewData;
+
+void CreateHairStrandsMacroGroups(
+	FRDGBuilder& GraphBuilder,
+	const FScene* Scene,
+	const FViewInfo& View, 
+	const TArray<EHairInstanceVisibilityType>& InstancesVisibilityType,
+	FHairStrandsViewData& OutHairStrandsViewData,
+	bool bBuildGPUAABB=true);
\ No newline at end of file
diff --git a/Engine/Source/Runtime/Renderer/Private/HairStrands/HairStrandsRasterCommon.h b/Engine/Source/Runtime/Renderer/Private/HairStrands/HairStrandsRasterCommon.h
--- a/Engine/Source/Runtime/Renderer/Private/HairStrands/HairStrandsRasterCommon.h
+++ b/Engine/Source/Runtime/Renderer/Private/HairStrands/HairStrandsRasterCommon.h
@@ -8,9 +8,11 @@
 
 #include "CoreMinimal.h"
 #include "RendererInterface.h"
-#include "HairStrandsCluster.h"
+#include "InstanceCulling/InstanceCullingContext.h"
+#include "HairStrands/HairStrandsData.h"
 
 class FViewInfo;
+class FInstanceCullingManager;
 
 enum class EHairStrandsRasterPassType : uint8
 {
diff --git a/Engine/Source/Runtime/Renderer/Private/HairStrands/HairStrandsRendering.cpp b/Engine/Source/Runtime/Renderer/Private/HairStrands/HairStrandsRendering.cpp
--- a/Engine/Source/Runtime/Renderer/Private/HairStrands/HairStrandsRendering.cpp
+++ b/Engine/Source/Runtime/Renderer/Private/HairStrands/HairStrandsRendering.cpp
@@ -6,7 +6,7 @@
 #include "ScenePrivate.h"
 #include "SystemTextures.h"
 #include "RenderGraphUtils.h"
-#include "HairStrands/HairStrandsCluster.h"
+#include "HairStrands/HairStrandsMacroGroup.h"
 
 static TRDGUniformBufferRef<FHairStrandsViewUniformParameters> InternalCreateHairStrandsViewUniformBuffer(
 	FRDGBuilder& GraphBuilder, 
diff --git a/Engine/Source/Runtime/Renderer/Private/HairStrands/HairStrandsTransmittance.cpp b/Engine/Source/Runtime/Renderer/Private/HairStrands/HairStrandsTransmittance.cpp
--- a/Engine/Source/Runtime/Renderer/Private/HairStrands/HairStrandsTransmittance.cpp
+++ b/Engine/Source/Runtime/Renderer/Private/HairStrands/HairStrandsTransmittance.cpp
@@ -1,7 +1,6 @@
 // Copyright Epic Games, Inc. All Rights Reserved.
 
 #include "HairStrandsTransmittance.h"
-#include "HairStrandsCluster.h"
 #include "HairStrandsLUT.h"
 #include "HairStrandsDeepShadow.h"
 #include "HairStrandsVoxelization.h"
diff --git a/Engine/Source/Runtime/Renderer/Private/HairStrands/HairStrandsUtils.cpp b/Engine/Source/Runtime/Renderer/Private/HairStrands/HairStrandsUtils.cpp
--- a/Engine/Source/Runtime/Renderer/Private/HairStrands/HairStrandsUtils.cpp
+++ b/Engine/Source/Runtime/Renderer/Private/HairStrands/HairStrandsUtils.cpp
@@ -3,7 +3,6 @@
 #include "HairStrandsUtils.h"
 #include "LightSceneProxy.h"
 #include "ScenePrivate.h"
-#include "HairStrandsCluster.h"
 #include "Rendering/SkeletalMeshRenderData.h"
 #include "HairStrandsData.h"
 #include "SystemTextures.h"
diff --git a/Engine/Source/Runtime/Renderer/Private/HairStrands/HairStrandsVisibility.cpp b/Engine/Source/Runtime/Renderer/Private/HairStrands/HairStrandsVisibility.cpp
--- a/Engine/Source/Runtime/Renderer/Private/HairStrands/HairStrandsVisibility.cpp
+++ b/Engine/Source/Runtime/Renderer/Private/HairStrands/HairStrandsVisibility.cpp
@@ -1,7 +1,7 @@
 // Copyright Epic Games, Inc. All Rights Reserved.
 
 #include "HairStrandsVisibility.h"
-#include "HairStrandsCluster.h"
+#include "HairStrandsMacroGroup.h"
 #include "HairStrandsUtils.h"
 #include "HairStrandsInterface.h"
 #include "HairStrandsLUT.h"
diff --git a/Engine/Source/Runtime/Renderer/Private/HairStrands/HairStrandsVoxelization.cpp b/Engine/Source/Runtime/Renderer/Private/HairStrands/HairStrandsVoxelization.cpp
--- a/Engine/Source/Runtime/Renderer/Private/HairStrands/HairStrandsVoxelization.cpp
+++ b/Engine/Source/Runtime/Renderer/Private/HairStrands/HairStrandsVoxelization.cpp
@@ -2,7 +2,6 @@
 
 #include "HairStrandsVoxelization.h"
 #include "HairStrandsRasterCommon.h"
-#include "HairStrandsCluster.h"
 #include "HairStrandsUtils.h"
 #include "HairStrandsLUT.h"
 #include "Shader.h"
