Fix compilation warning due to accessing deprecated member with copy constructor.

diff --git a/Engine/Source/Runtime/Renderer/Private/HairStrands/HairStrandsDeepShadow.cpp b/Engine/Source/Runtime/Renderer/Private/HairStrands/HairStrandsDeepShadow.cpp
--- a/Engine/Source/Runtime/Renderer/Private/HairStrands/HairStrandsDeepShadow.cpp
+++ b/Engine/Source/Runtime/Renderer/Private/HairStrands/HairStrandsDeepShadow.cpp
@@ -459,7 +459,8 @@ void RenderHairStrandsDeepShadows(
 						UniformParameters->AtlasRect = DomData.AtlasRect;
 						UniformParameters->AtlasSlotIndex = DomData.AtlasSlotIndex;
 						UniformParameters->LayerDepths = LayerDepths;
-						UniformParameters->ViewportResolution = AtlasSlotResolution;
+						UniformParameters->ViewportResolution[0] = AtlasSlotResolution[0];
+						UniformParameters->ViewportResolution[1] = AtlasSlotResolution[1];
 						UniformParameters->FrontDepthTexture = SystemTextures.DepthDummy;
 						UniformParameters->DeepShadowViewInfoBuffer = DeepShadowViewInfoBufferSRV;
 
@@ -499,7 +500,8 @@ void RenderHairStrandsDeepShadows(
 						UniformParameters->AtlasRect = DomData.AtlasRect;
 						UniformParameters->AtlasSlotIndex = DomData.AtlasSlotIndex;
 						UniformParameters->LayerDepths = LayerDepths;
-						UniformParameters->ViewportResolution = AtlasSlotResolution;
+						UniformParameters->ViewportResolution[0] = AtlasSlotResolution[0];
+						UniformParameters->ViewportResolution[1] = AtlasSlotResolution[1];
 						UniformParameters->FrontDepthTexture = FrontDepthAtlasTexture;
 						UniformParameters->DeepShadowViewInfoBuffer = DeepShadowViewInfoBufferSRV;
 
