Store Decals in TArray and use RemoveAtSwap instead of TSparseArray.
- enables faster iteration over arrays.
- also makes it easier to parallelize loops in follow up CLs.

diff --git a/Engine/Source/Runtime/Renderer/Private/RayTracing/RayTracingDecals.cpp b/Engine/Source/Runtime/Renderer/Private/RayTracing/RayTracingDecals.cpp
--- a/Engine/Source/Runtime/Renderer/Private/RayTracing/RayTracingDecals.cpp
+++ b/Engine/Source/Runtime/Renderer/Private/RayTracing/RayTracingDecals.cpp
@@ -471,7 +471,7 @@ void BuildDecalGrid(FRDGBuilder& GraphBuilder, uint32 NumDecals, FRDGBufferSRVRe
 	}
 }
 
-FTransientDecalRenderDataList GetSortedDecals(const TSparseArray<FDeferredDecalProxy*>& Decals, FScene& Scene, const FViewInfo& View)
+FTransientDecalRenderDataList GetSortedDecals(TConstArrayView<FDeferredDecalProxy*> Decals, FScene& Scene, const FViewInfo& View)
 {
 	const float FadeMultiplier = DecalRendering::GetDecalFadeScreenSizeMultiplier();
 	const bool bIsPerspectiveProjection = View.IsPerspectiveProjection();
diff --git a/Engine/Source/Runtime/Renderer/Private/RendererScene.cpp b/Engine/Source/Runtime/Renderer/Private/RendererScene.cpp
--- a/Engine/Source/Runtime/Renderer/Private/RendererScene.cpp
+++ b/Engine/Source/Runtime/Renderer/Private/RendererScene.cpp
@@ -2868,7 +2868,7 @@ bool FScene::HasAtmosphereLightRequiringLightingBuild() const
 
 void FScene::AddOrRemoveDecal_RenderThread(FDeferredDecalProxy* Proxy, bool bAdd)
 {
-	if(bAdd)
+	if (bAdd)
 	{
 		Decals.Add(Proxy);
 		InvalidatePathTracedOutput();
@@ -2876,15 +2876,13 @@ void FScene::AddOrRemoveDecal_RenderThread(FDeferredDecalProxy* Proxy, bool bAdd
 	else
 	{
 		// can be optimized
-		for(TSparseArray<FDeferredDecalProxy*>::TIterator It(Decals); It; ++It)
+		for (int32 Index = 0; Index < Decals.Num(); ++Index)
 		{
-			FDeferredDecalProxy* CurrentProxy = *It;
-
-			if (CurrentProxy == Proxy)
+			if (Decals[Index] == Proxy)
 			{
 				InvalidatePathTracedOutput();
-				It.RemoveCurrent();
-				delete CurrentProxy;
+				Decals.RemoveAtSwap(Index, 1, EAllowShrinking::No);
+				delete Proxy;
 				break;
 			}
 		}
diff --git a/Engine/Source/Runtime/Renderer/Private/ScenePrivate.h b/Engine/Source/Runtime/Renderer/Private/ScenePrivate.h
--- a/Engine/Source/Runtime/Renderer/Private/ScenePrivate.h
+++ b/Engine/Source/Runtime/Renderer/Private/ScenePrivate.h
@@ -3044,7 +3044,7 @@ public:
 	TArray<FLightSceneInfo*, TInlineAllocator<4>> DirectionalLights;
 
 	/** The decals in the scene. */
-	TSparseArray<FDeferredDecalProxy*> Decals;
+	TArray<FDeferredDecalProxy*> Decals;
 
 	/** Potential capsule shadow casters registered to the scene. */
 	TArray<FPrimitiveSceneInfo*> DynamicIndirectCasterPrimitives; 
