Fix CIS issue.

diff --git a/Engine/Source/Runtime/Renderer/Private/HairStrands/HairStrandsDeepShadow.cpp b/Engine/Source/Runtime/Renderer/Private/HairStrands/HairStrandsDeepShadow.cpp
--- a/Engine/Source/Runtime/Renderer/Private/HairStrands/HairStrandsDeepShadow.cpp
+++ b/Engine/Source/Runtime/Renderer/Private/HairStrands/HairStrandsDeepShadow.cpp
@@ -459,8 +459,6 @@ void RenderHairStrandsDeepShadows(
 						UniformParameters->AtlasRect = DomData.AtlasRect;
 						UniformParameters->AtlasSlotIndex = DomData.AtlasSlotIndex;
 						UniformParameters->LayerDepths = LayerDepths;
-						UniformParameters->ViewportResolution[0] = AtlasSlotResolution[0];
-						UniformParameters->ViewportResolution[1] = AtlasSlotResolution[1];
 						UniformParameters->FrontDepthTexture = SystemTextures.DepthDummy;
 						UniformParameters->DeepShadowViewInfoBuffer = DeepShadowViewInfoBufferSRV;
 
@@ -500,8 +498,6 @@ void RenderHairStrandsDeepShadows(
 						UniformParameters->AtlasRect = DomData.AtlasRect;
 						UniformParameters->AtlasSlotIndex = DomData.AtlasSlotIndex;
 						UniformParameters->LayerDepths = LayerDepths;
-						UniformParameters->ViewportResolution[0] = AtlasSlotResolution[0];
-						UniformParameters->ViewportResolution[1] = AtlasSlotResolution[1];
 						UniformParameters->FrontDepthTexture = FrontDepthAtlasTexture;
 						UniformParameters->DeepShadowViewInfoBuffer = DeepShadowViewInfoBufferSRV;
 
diff --git a/Engine/Source/Runtime/Renderer/Private/HairStrands/HairStrandsRasterCommon.h b/Engine/Source/Runtime/Renderer/Private/HairStrands/HairStrandsRasterCommon.h
--- a/Engine/Source/Runtime/Renderer/Private/HairStrands/HairStrandsRasterCommon.h
+++ b/Engine/Source/Runtime/Renderer/Private/HairStrands/HairStrandsRasterCommon.h
@@ -27,7 +27,6 @@ BEGIN_GLOBAL_SHADER_PARAMETER_STRUCT(FHairDeepShadowRasterUniformParameters, REN
 	SHADER_PARAMETER(FMatrix44f, CPU_TranslatedWorldToClipMatrix)
 	SHADER_PARAMETER(FVector4f, SliceValue)
 	SHADER_PARAMETER(FIntRect, AtlasRect)
-	SHADER_PARAMETER(FIntPoint, ViewportResolution)
 	SHADER_PARAMETER(uint32, AtlasSlotIndex)
 	SHADER_PARAMETER(FVector4f, LayerDepths)
 	SHADER_PARAMETER_RDG_TEXTURE(Texture2D<float>, FrontDepthTexture)
