Refactored logic that calculates which deferred decals are relevant into BuildRelevantDecalList(...) and BuildRelevantDecalList(...)
- cache and reuse result of BuildRelevantDecalList(...) when rendering the different stages

Level containing ~2000 decals running on console: 4.4ms -> 3ms

diff --git a/Engine/Source/Runtime/Renderer/Private/CompositionLighting/CompositionLighting.cpp b/Engine/Source/Runtime/Renderer/Private/CompositionLighting/CompositionLighting.cpp
--- a/Engine/Source/Runtime/Renderer/Private/CompositionLighting/CompositionLighting.cpp
+++ b/Engine/Source/Runtime/Renderer/Private/CompositionLighting/CompositionLighting.cpp
@@ -500,11 +500,23 @@ FCompositionLighting::FCompositionLighting(TArrayView<const FViewInfo> InViews,
 	, bEnableDBuffer(IsDBufferEnabled(ViewFamily, SceneTextures.Config.ShaderPlatform))
 	, bEnableDecals(ViewFamily.EngineShowFlags.Decals && !ViewFamily.EngineShowFlags.VisualizeLightCulling)
 {
+	const FScene& Scene = *(FScene*)ViewFamily.Scene;
+
 	ViewAOConfigs.SetNum(Views.Num());
 
+	if (bEnableDecals)
+	{
+		VisibleDecals.SetNum(Views.Num());
+	}
+
 	for (int32 Index = 0; Index < Views.Num(); ++Index)
 	{
 		ViewAOConfigs[Index].bRequested = RequestSSAOFunction(Index);
+
+		if (bEnableDecals)
+		{
+			VisibleDecals[Index] = DecalRendering::BuildVisibleDecalList(Scene.Decals, Views[Index]);
+		}
 	}
 }
 
@@ -581,7 +593,7 @@ void FCompositionLighting::ProcessBeforeBasePass(FRDGBuilder& GraphBuilder, FDBu
 		if (bEnableDBuffer)
 		{
 			FDeferredDecalPassTextures DecalPassTextures = GetDeferredDecalPassTextures(GraphBuilder, View, SceneTextures, &DBufferTextures);
-			AddDeferredDecalPass(GraphBuilder, View, DecalPassTextures, InstanceCullingManager, EDecalRenderStage::BeforeBasePass);
+			AddDeferredDecalPass(GraphBuilder, View, VisibleDecals[ViewIndex], DecalPassTextures, InstanceCullingManager, EDecalRenderStage::BeforeBasePass);
 		}
 
 		if (bEnableSSAO)
@@ -626,13 +638,13 @@ void FCompositionLighting::ProcessAfterBasePass(FRDGBuilder& GraphBuilder, FInst
 		{
 			// We can disable this pass if using DBuffer decals
 			// Decals are before AmbientOcclusion so the decal can output a normal that AO is affected by
-			AddDeferredDecalPass(GraphBuilder, View, DecalPassTextures, InstanceCullingManager, EDecalRenderStage::BeforeLighting);
+			AddDeferredDecalPass(GraphBuilder, View, VisibleDecals[ViewIndex], DecalPassTextures, InstanceCullingManager, EDecalRenderStage::BeforeLighting);
 		}
 
 		if (bEnableDecals && Mode != EProcessAfterBasePassMode::OnlyBeforeLightingDecals)
 		{
 			// DBuffer decals with emissive component
-			AddDeferredDecalPass(GraphBuilder, View, DecalPassTextures, InstanceCullingManager, EDecalRenderStage::Emissive);
+			AddDeferredDecalPass(GraphBuilder, View, VisibleDecals[ViewIndex], DecalPassTextures, InstanceCullingManager, EDecalRenderStage::Emissive);
 		}
 
 		// Forward shading SSAO is applied before the base pass using only the depth buffer.
@@ -656,7 +668,7 @@ void FCompositionLighting::ProcessAfterBasePass(FRDGBuilder& GraphBuilder, FInst
 					AmbientOcclusion = AddPostProcessingGTAOPostAsync(GraphBuilder, View, Parameters, GTAOHorizons, FinalTarget);
 
 					ensureMsgf(
-						DecalRendering::BuildVisibleDecalList(*(FScene*)View.Family->Scene, View, EDecalRenderStage::AmbientOcclusion, nullptr) == false,
+						DecalRendering::BuildRelevantDecalList(VisibleDecals[ViewIndex], EDecalRenderStage::AmbientOcclusion, nullptr) == false,
 						TEXT("Ambient occlusion decals are not supported with Async compute SSAO."));
 				}
 				else
@@ -675,7 +687,7 @@ void FCompositionLighting::ProcessAfterBasePass(FRDGBuilder& GraphBuilder, FInst
 					if (bEnableDecals)
 					{
 						DecalPassTextures.ScreenSpaceAO = AmbientOcclusion.Texture;
-						AddDeferredDecalPass(GraphBuilder, View, DecalPassTextures, InstanceCullingManager, EDecalRenderStage::AmbientOcclusion);
+						AddDeferredDecalPass(GraphBuilder, View, VisibleDecals[ViewIndex], DecalPassTextures, InstanceCullingManager, EDecalRenderStage::AmbientOcclusion);
 					}
 				}
 			}
diff --git a/Engine/Source/Runtime/Renderer/Private/CompositionLighting/CompositionLighting.h b/Engine/Source/Runtime/Renderer/Private/CompositionLighting/CompositionLighting.h
--- a/Engine/Source/Runtime/Renderer/Private/CompositionLighting/CompositionLighting.h
+++ b/Engine/Source/Runtime/Renderer/Private/CompositionLighting/CompositionLighting.h
@@ -10,6 +10,7 @@
 #include "SceneRendering.h"
 #include "CompositionLighting/PostProcessDeferredDecals.h"
 #include "CompositionLighting/PostProcessAmbientOcclusion.h"
+#include "DecalRenderingShared.h"
 
 /**
  * The center for all screen space processing activities (e.g. G-buffer manipulation, lighting).
@@ -59,6 +60,7 @@ private:
 	};
 
 	TArray<FAOConfig, TInlineAllocator<8>> ViewAOConfigs;
+	TArray<FTransientDecalRenderDataList, TInlineAllocator<8>> VisibleDecals;
 	FRDGTextureRef HorizonsTexture = nullptr;
 	bool bInitialized = false;
 };
diff --git a/Engine/Source/Runtime/Renderer/Private/CompositionLighting/PostProcessDeferredDecals.cpp b/Engine/Source/Runtime/Renderer/Private/CompositionLighting/PostProcessDeferredDecals.cpp
--- a/Engine/Source/Runtime/Renderer/Private/CompositionLighting/PostProcessDeferredDecals.cpp
+++ b/Engine/Source/Runtime/Renderer/Private/CompositionLighting/PostProcessDeferredDecals.cpp
@@ -557,6 +557,7 @@ void CollectDeferredDecalPassPSOInitializers(
 void AddDeferredDecalPass(
 	FRDGBuilder& GraphBuilder,
 	const FViewInfo& View,
+	TConstArrayView<FTransientDecalRenderData> VisibleDecals,
 	const FDeferredDecalPassTextures& PassTextures,
 	FInstanceCullingManager& InstanceCullingManager,
 	EDecalRenderStage DecalRenderStage)
@@ -584,10 +585,10 @@ void AddDeferredDecalPass(
 	checkf(DecalRenderStage != EDecalRenderStage::AmbientOcclusion || PassTextures.ScreenSpaceAO, TEXT("Attepting to render AO decals without SSAO having emitted a valid render target."));
 	checkf(DecalRenderStage != EDecalRenderStage::BeforeBasePass || IsUsingDBuffers(ShaderPlatform), TEXT("Only DBuffer decals are supported before the base pass."));
 
-	if (DecalCount)
+	if (!VisibleDecals.IsEmpty())
 	{
 		SortedDecals = GraphBuilder.AllocObject<FTransientDecalRenderDataList>();
-		DecalRendering::BuildVisibleDecalList(Scene, View, DecalRenderStage, SortedDecals);
+		DecalRendering::BuildRelevantDecalList(VisibleDecals, DecalRenderStage, SortedDecals);
 		SortedDecalCount = SortedDecals->Num();
 
 		INC_DWORD_STAT_BY(STAT_Decals, SortedDecalCount);
@@ -690,7 +691,7 @@ void AddDeferredDecalPass(
 
 		if (SortedDecalCount > 0)
 		{
-			RDG_EVENT_SCOPE(GraphBuilder, "Decals (Visible %d, Total: %d)", SortedDecalCount, DecalCount);
+			RDG_EVENT_SCOPE(GraphBuilder, "Decals (Relevant: %d, Visible: %d, Total: %d)", SortedDecalCount, VisibleDecals.Num(), DecalCount);
 
 			uint32 SortedDecalIndex = 1;
 			uint32 LastSortedDecalIndex = 0;
diff --git a/Engine/Source/Runtime/Renderer/Private/CompositionLighting/PostProcessDeferredDecals.h b/Engine/Source/Runtime/Renderer/Private/CompositionLighting/PostProcessDeferredDecals.h
--- a/Engine/Source/Runtime/Renderer/Private/CompositionLighting/PostProcessDeferredDecals.h
+++ b/Engine/Source/Runtime/Renderer/Private/CompositionLighting/PostProcessDeferredDecals.h
@@ -9,6 +9,7 @@
 
 enum class EDecalRenderStage : uint8;
 enum class EDecalRenderTargetMode : uint8;
+struct FTransientDecalRenderData;
 struct FDBufferTextures;
 struct FSceneTextures;
 class FViewInfo;
@@ -47,6 +48,7 @@ FDeferredDecalPassTextures GetDeferredDecalPassTextures(
 void AddDeferredDecalPass(
 	FRDGBuilder& GraphBuilder,
 	const FViewInfo& ViewInfo,
+	TConstArrayView<FTransientDecalRenderData> VisibleDecals,
 	const FDeferredDecalPassTextures& Textures,
 	FInstanceCullingManager& InstanceCullingManager,
 	EDecalRenderStage RenderStage);
diff --git a/Engine/Source/Runtime/Renderer/Private/DecalRenderingShared.cpp b/Engine/Source/Runtime/Renderer/Private/DecalRenderingShared.cpp
--- a/Engine/Source/Runtime/Renderer/Private/DecalRenderingShared.cpp
+++ b/Engine/Source/Runtime/Renderer/Private/DecalRenderingShared.cpp
@@ -23,7 +23,7 @@ static TAutoConsoleVariable<float> CVarDecalFadeScreenSizeMultiplier(
 	TEXT("  Smaller means decals fade less aggressively.")
 	);
 
-FTransientDecalRenderData::FTransientDecalRenderData(const FScene& InScene, const FDeferredDecalProxy& InDecalProxy, float InConservativeRadius)
+FTransientDecalRenderData::FTransientDecalRenderData(const FDeferredDecalProxy& InDecalProxy, float InConservativeRadius, EShaderPlatform ShaderPlatform, ERHIFeatureLevel::Type FeatureLevel)
 	: Proxy(InDecalProxy)
 	, MaterialProxy(InDecalProxy.DecalMaterial->GetRenderProxy())
 	, ConservativeRadius(InConservativeRadius)
@@ -32,8 +32,8 @@ FTransientDecalRenderData::FTransientDecalRenderData(const FScene& InScene, cons
 {
 	// Build BlendDesc from a potentially incomplete material.
 	// If our shader isn't compiled yet then we will potentially render later with a different fallback material.
-	FMaterial const& MaterialResource = MaterialProxy->GetIncompleteMaterialWithFallback(InScene.GetFeatureLevel());
-	BlendDesc = DecalRendering::ComputeDecalBlendDesc(InScene.GetShaderPlatform(), MaterialResource);
+	FMaterial const& MaterialResource = MaterialProxy->GetIncompleteMaterialWithFallback(FeatureLevel);
+	BlendDesc = DecalRendering::ComputeDecalBlendDesc(ShaderPlatform, MaterialResource);
 }
 
 /**
@@ -277,29 +277,26 @@ namespace DecalRendering
 		Decals.Sort(FCompareFTransientDecalRenderData());
 	}
 
-	bool BuildVisibleDecalList(const FScene& Scene, const FViewInfo& View, EDecalRenderStage DecalRenderStage, FTransientDecalRenderDataList* OutVisibleDecals)
+	FTransientDecalRenderDataList BuildVisibleDecalList(TConstArrayView<FDeferredDecalProxy*> Decals, const FViewInfo& View)
 	{
 		QUICK_SCOPE_CYCLE_COUNTER(BuildVisibleDecalList);
 
-		if (OutVisibleDecals)
-		{
-			OutVisibleDecals->Empty(Scene.Decals.Num());
-		}
-
 		// Don't draw for shader complexity mode.
 		// todo: Handle shader complexity mode for deferred decal.
 		if (View.Family->EngineShowFlags.ShaderComplexity)
 		{
-			return false;
+			return {};
 		}
 
+		FTransientDecalRenderDataList OutVisibleDecals;
+
 		const float FadeMultiplier = GetDecalFadeScreenSizeMultiplier();
 		const EShaderPlatform ShaderPlatform = View.GetShaderPlatform();
 
 		const bool bIsPerspectiveProjection = View.IsPerspectiveProjection();
 
 		// Build a list of decals that need to be rendered for this view in SortedDecals
-		for (const FDeferredDecalProxy* DecalProxy : Scene.Decals)
+		for (const FDeferredDecalProxy* DecalProxy : Decals)
 		{
 			if (!DecalProxy->DecalMaterial || !DecalProxy->DecalMaterial->IsValidLowLevelFast())
 			{
@@ -317,41 +314,68 @@ namespace DecalRendering
 
 			// can be optimized as we test against a sphere around the box instead of the box itself
 			const float ConservativeRadius = FMath::Sqrt(
-					ComponentToWorldMatrix.GetScaledAxis(EAxis::X).SizeSquared() +
-					ComponentToWorldMatrix.GetScaledAxis(EAxis::Y).SizeSquared() +
-					ComponentToWorldMatrix.GetScaledAxis(EAxis::Z).SizeSquared());
+				ComponentToWorldMatrix.GetScaledAxis(EAxis::X).SizeSquared() +
+				ComponentToWorldMatrix.GetScaledAxis(EAxis::Y).SizeSquared() +
+				ComponentToWorldMatrix.GetScaledAxis(EAxis::Z).SizeSquared());
 
 			// can be optimized as the test is too conservative (sphere instead of OBB)
-			if(ConservativeRadius < SMALL_NUMBER || !View.ViewFrustum.IntersectSphere(ComponentToWorldMatrix.GetOrigin(), ConservativeRadius))
+			if (ConservativeRadius < SMALL_NUMBER || !View.ViewFrustum.IntersectSphere(ComponentToWorldMatrix.GetOrigin(), ConservativeRadius))
 			{
 				bIsShown = false;
 			}
 
 			if (bIsShown)
 			{
-				FTransientDecalRenderData Data(Scene, *DecalProxy, ConservativeRadius);
-			
-				if (IsCompatibleWithRenderStage(Data.BlendDesc, DecalRenderStage))
+				FTransientDecalRenderData Data(*DecalProxy, ConservativeRadius, ShaderPlatform, View.GetFeatureLevel());
+
+				if (bIsPerspectiveProjection && Data.Proxy.FadeScreenSize != 0.0f)
+				{
+					Data.FadeAlpha = CalculateDecalFadeAlpha(Data.Proxy.FadeScreenSize, ComponentToWorldMatrix, View, FadeMultiplier);
+				}
+
+				const bool bShouldRender = Data.FadeAlpha > 0.0f;
+
+				if (bShouldRender)
 				{
-					if (bIsPerspectiveProjection && Data.Proxy.FadeScreenSize != 0.0f)
-					{
-						Data.FadeAlpha = CalculateDecalFadeAlpha(Data.Proxy.FadeScreenSize, ComponentToWorldMatrix, View, FadeMultiplier);
-					}
-
-					const bool bShouldRender = Data.FadeAlpha > 0.0f;
-
-					if (bShouldRender)
-					{
-						if (!OutVisibleDecals)
-						{
-							return true;
-						}
-						OutVisibleDecals->Add(Data);
-					}
+					OutVisibleDecals.Add(Data);
 				}
 			}
 		}
 
+		if (OutVisibleDecals.Num() > 0)
+		{
+			SortDecalList(OutVisibleDecals);
+		}
+
+		return OutVisibleDecals;
+	}
+
+	bool BuildRelevantDecalList(TConstArrayView<FTransientDecalRenderData> Decals, EDecalRenderStage DecalRenderStage, FTransientDecalRenderDataList* OutVisibleDecals)
+	{
+		QUICK_SCOPE_CYCLE_COUNTER(BuildRelevantDecalList);
+
+		if (OutVisibleDecals)
+		{
+			OutVisibleDecals->Empty(Decals.Num());
+		}
+
+		// Build a list of decals that need to be rendered for this stage in SortedDecals
+		for (const FTransientDecalRenderData& DecalRenderData : Decals)
+		{
+			checkf(DecalRenderData.Proxy.DecalMaterial && DecalRenderData.Proxy.DecalMaterial->IsValidLowLevelFast(),
+				TEXT("Decals should've been filtered earlier in BuildVisibleDecalList"));
+
+			if (IsCompatibleWithRenderStage(DecalRenderData.BlendDesc, DecalRenderStage))
+			{
+				if (!OutVisibleDecals)
+				{
+					return true;
+				}
+
+				OutVisibleDecals->Add(DecalRenderData);
+			}
+		}
+
 		if (!OutVisibleDecals)
 		{
 			return false;
diff --git a/Engine/Source/Runtime/Renderer/Private/DecalRenderingShared.h b/Engine/Source/Runtime/Renderer/Private/DecalRenderingShared.h
--- a/Engine/Source/Runtime/Renderer/Private/DecalRenderingShared.h
+++ b/Engine/Source/Runtime/Renderer/Private/DecalRenderingShared.h
@@ -24,7 +24,7 @@ struct FTransientDecalRenderData
 	float FadeAlpha;
 	FLinearColor DecalColor;
 
-	FTransientDecalRenderData(const FScene& InScene, const FDeferredDecalProxy& InDecalProxy, float InConservativeRadius);
+	FTransientDecalRenderData(const FDeferredDecalProxy& InDecalProxy, float InConservativeRadius, EShaderPlatform ShaderPlatform, ERHIFeatureLevel::Type FeatureLevel);
 };
 	
 typedef TArray<FTransientDecalRenderData, SceneRenderingAllocator> FTransientDecalRenderDataList;
@@ -39,7 +39,8 @@ namespace DecalRendering
 	FMatrix ComputeComponentToClipMatrix(const FViewInfo& View, const FMatrix& DecalComponentToWorld);
 	void SetVertexShaderOnly(FRHICommandList& RHICmdList, FGraphicsPipelineStateInitializer& GraphicsPSOInit, const FViewInfo& View, const FMatrix& FrustumComponentToClip);
 	void SortDecalList(FTransientDecalRenderDataList& Decals);
-	bool BuildVisibleDecalList(const FScene& Scene, const FViewInfo& View, EDecalRenderStage DecalRenderStage, FTransientDecalRenderDataList* OutVisibleDecals);
+	FTransientDecalRenderDataList BuildVisibleDecalList(TConstArrayView<FDeferredDecalProxy*> Decals, const FViewInfo& View);
+	bool BuildRelevantDecalList(TConstArrayView<FTransientDecalRenderData> Decals, EDecalRenderStage DecalRenderStage, FTransientDecalRenderDataList* OutVisibleDecals);
 	bool SetupShaderState(ERHIFeatureLevel::Type FeatureLevel, const FMaterial& Material, EDecalRenderStage DecalRenderStage, FBoundShaderStateInput& OutBoundShaderState);
 	void SetShader(FRHICommandList& RHICmdList, FGraphicsPipelineStateInitializer& GraphicsPSOInit, uint32 StencilRef, const FViewInfo& View, const FTransientDecalRenderData& DecalData, EDecalRenderStage DecalRenderStage, const FMatrix& FrustumComponentToClip);
 };
diff --git a/Engine/Source/Runtime/Renderer/Private/MobileDecalRendering.cpp b/Engine/Source/Runtime/Renderer/Private/MobileDecalRendering.cpp
--- a/Engine/Source/Runtime/Renderer/Private/MobileDecalRendering.cpp
+++ b/Engine/Source/Runtime/Renderer/Private/MobileDecalRendering.cpp
@@ -84,14 +84,15 @@ void FMobileSceneRenderer::RenderDecals(FRHICommandList& RHICmdList, const FView
 
 void RenderDeferredDecalsMobile(FRHICommandList& RHICmdList, const FScene& Scene, const FViewInfo& View, EDecalRenderStage DecalRenderStage, EDecalRenderTargetMode RenderTargetMode)
 {
-	const uint32 DecalCount = Scene.Decals.Num();
 	int32 SortedDecalCount = 0;
 	FTransientDecalRenderDataList SortedDecals;
 
-	if (DecalCount > 0)
+	if (!Scene.Decals.IsEmpty())
 	{
+		FTransientDecalRenderDataList VisibleDecals = DecalRendering::BuildVisibleDecalList(Scene.Decals, View);
+
 		// Build a list of decals that need to be rendered for this view
-		DecalRendering::BuildVisibleDecalList(Scene, View, DecalRenderStage, &SortedDecals);
+		DecalRendering::BuildRelevantDecalList(VisibleDecals, DecalRenderStage, &SortedDecals);
 		SortedDecalCount = SortedDecals.Num();
 		INC_DWORD_STAT_BY(STAT_Decals, SortedDecalCount);
 	}
@@ -165,7 +166,9 @@ void FMobileSceneRenderer::RenderDBuffer(FRDGBuilder& GraphBuilder, FSceneTextur
 			continue;
 		}
 
+		FTransientDecalRenderDataList VisibleDecals = DecalRendering::BuildVisibleDecalList(Scene->Decals, View);
+
 		FDeferredDecalPassTextures DecalPassTextures = GetDeferredDecalPassTextures(GraphBuilder, View, SceneTextures, &DBufferTextures);
-		AddDeferredDecalPass(GraphBuilder, View, DecalPassTextures, InstanceCullingManager, EDecalRenderStage::BeforeBasePass);
+		AddDeferredDecalPass(GraphBuilder, View, VisibleDecals, DecalPassTextures, InstanceCullingManager, EDecalRenderStage::BeforeBasePass);
 	}
 }
diff --git a/Engine/Source/Runtime/Renderer/Private/RayTracing/RayTracingDecals.cpp b/Engine/Source/Runtime/Renderer/Private/RayTracing/RayTracingDecals.cpp
--- a/Engine/Source/Runtime/Renderer/Private/RayTracing/RayTracingDecals.cpp
+++ b/Engine/Source/Runtime/Renderer/Private/RayTracing/RayTracingDecals.cpp
@@ -495,7 +495,7 @@ FTransientDecalRenderDataList GetSortedDecals(TConstArrayView<FDeferredDecalProx
 
 		if (bIsShown)
 		{
-			FTransientDecalRenderData Data(Scene, *DecalProxy, 0.0f);
+			FTransientDecalRenderData Data(*DecalProxy, 0.0f, Scene.GetShaderPlatform(), Scene.GetFeatureLevel());
 
 			if (bIsPerspectiveProjection && Data.Proxy.FadeScreenSize != 0.0f)
 			{
