Only collect HV primitives if the shadow is volumetric

diff --git a/Engine/Source/Runtime/Renderer/Private/ShadowRendering.h b/Engine/Source/Runtime/Renderer/Private/ShadowRendering.h
--- a/Engine/Source/Runtime/Renderer/Private/ShadowRendering.h
+++ b/Engine/Source/Runtime/Renderer/Private/ShadowRendering.h
@@ -436,6 +436,9 @@ public:
 	/** Whether the the shadow overlaps any nanite primitives */
 	uint32 bContainsNaniteSubjects : 1;
 	uint32 bShouldRenderVSM : 1;
+
+	/** Whether this shadow should support casting shadows from volumetric surfaces. */
+	uint32 bVolumetricShadow : 1;
 	
 	/** Used to fetch the correct cached static mesh draw commands */
 	EMeshPass::Type MeshPassTargetType = EMeshPass::CSMShadowDepth;
diff --git a/Engine/Source/Runtime/Renderer/Private/ShadowSetup.cpp b/Engine/Source/Runtime/Renderer/Private/ShadowSetup.cpp
--- a/Engine/Source/Runtime/Renderer/Private/ShadowSetup.cpp
+++ b/Engine/Source/Runtime/Renderer/Private/ShadowSetup.cpp
@@ -750,6 +750,7 @@ FProjectedShadowInfo::FProjectedShadowInfo()
 	, bNaniteGeometry(true)
 	, bContainsNaniteSubjects(false)
 	, bShouldRenderVSM(true)
+	, bVolumetricShadow(false)
 	, PerObjectShadowFadeStart(UE_OLD_WORLD_MAX)		// LWC_TODO: Upgrade to double? Should have been HALF_WORLD_MAX to match precision.
 	, InvPerObjectShadowFadeLength(0.0f)
 	, OverlappedUVOnCachedShadowMap(-1.0f, -1.0f, -1.0f, -1.0f)
@@ -803,6 +804,7 @@ bool FProjectedShadowInfo::SetupPerObjectProjection(
 	bSelfShadowOnly = InParentSceneInfo->Proxy->CastsSelfShadowOnly();
 	bTransmission = InLightSceneInfo->Proxy->Transmission();
 	bHairStrandsDeepShadow = InLightSceneInfo->Proxy->CastsHairStrandsDeepShadow();
+	bVolumetricShadow = InLightSceneInfo->Proxy->CastsVolumetricShadow();
 	
 	check(!bRayTracedDistanceField);
 	
@@ -940,6 +942,7 @@ void FProjectedShadowInfo::SetupWholeSceneProjection(
 	bWholeSceneShadow = true;
 	bTransmission = InLightSceneInfo->Proxy->Transmission();
 	bHairStrandsDeepShadow = InLightSceneInfo->Proxy->CastsHairStrandsDeepShadow();
+	bVolumetricShadow = InLightSceneInfo->Proxy->CastsVolumetricShadow();
 	BorderSize = InBorderSize;
 
 	// Scale matrix to also include the border
@@ -2051,7 +2054,7 @@ bool FProjectedShadowInfo::AddSubjectPrimitive(FDynamicShadowsTaskData& TaskData
 			bWasAdded = true;
 		}
 
-		if (bHeterogeneousVolumeRelevance)
+		if (bVolumetricShadow && bHeterogeneousVolumeRelevance)
 		{
 			SubjectHeterogeneousVolumePrimitives.Add(PrimitiveSceneInfo);
 			bWasAdded = true;
@@ -2263,7 +2266,7 @@ uint64 FProjectedShadowInfo::AddSubjectPrimitive_AnyThread(
 			++OutStats.NumTranslucentSubs;
 		}
 
-		if (bHeterogeneousVolumeRelevance)
+		if (bVolumetricShadow && bHeterogeneousVolumeRelevance)
 		{
 			Result.bHeterogeneousVolumeSubjectPrimitive = true;
 			++OutStats.NumHeterogeneousVolumeSubs;
