Path  Tracing: Reduce size of lights

Remove the redundant parameter storing the bitangent as this can easily be reconstructed when needed (mainly for rectlights).

Also fix the type for IESAtlasIndex to float like in the rest of the code

diff --git a/Engine/Plugins/Experimental/GPULightmass/Source/GPULightmass/Private/PathTracingLightParameters.inl b/Engine/Plugins/Experimental/GPULightmass/Source/GPULightmass/Private/PathTracingLightParameters.inl
--- a/Engine/Plugins/Experimental/GPULightmass/Source/GPULightmass/Private/PathTracingLightParameters.inl
+++ b/Engine/Plugins/Experimental/GPULightmass/Source/GPULightmass/Private/PathTracingLightParameters.inl
@@ -80,8 +80,7 @@ void SetupPathTracingLightParameters(
 		DestLight.TranslatedWorldPosition = FVector3f(Light.Position + View.ViewMatrices.GetPreViewTranslation());
 		DestLight.Color = (FVector3f)(Light.Color);
 		DestLight.Normal = (FVector3f)Light.Direction;
-		DestLight.dPdu = (FVector3f)FVector::CrossProduct(Light.Tangent, Light.Direction);
-		DestLight.dPdv = (FVector3f)Light.Tangent;
+		DestLight.Tangent = (FVector3f)Light.Tangent;
 
 		DestLight.Dimensions = FVector2f(Light.SourceRadius, Light.SourceLength);
 		DestLight.Attenuation = 1.0f / Light.AttenuationRadius;
@@ -102,8 +101,7 @@ void SetupPathTracingLightParameters(
 
 		DestLight.TranslatedWorldPosition = FVector3f(Light.Position + View.ViewMatrices.GetPreViewTranslation());
 		DestLight.Normal = (FVector3f)Light.Direction;
-		DestLight.dPdu = (FVector3f)FVector::CrossProduct(Light.Tangent, Light.Direction);
-		DestLight.dPdv = (FVector3f)Light.Tangent;
+		DestLight.Tangent = (FVector3f)Light.Tangent;
 		DestLight.Color = FVector3f(Light.Color);
 		DestLight.Dimensions = FVector2f(Light.SourceRadius, Light.SourceLength);
 		DestLight.Shaping = FVector2f(Light.SpotAngles);
@@ -125,8 +123,7 @@ void SetupPathTracingLightParameters(
 
 		DestLight.TranslatedWorldPosition = FVector3f(Light.Position + View.ViewMatrices.GetPreViewTranslation());
 		DestLight.Normal = (FVector3f)Light.Direction;
-		DestLight.dPdu = (FVector3f)FVector::CrossProduct(Light.Tangent, -Light.Direction);
-		DestLight.dPdv = (FVector3f)Light.Tangent;
+		DestLight.Tangent = (FVector3f)Light.Tangent;
 
 		FLinearColor LightColor = Light.Color;
 		LightColor /= 0.5f * Light.SourceWidth * Light.SourceHeight;
diff --git a/Engine/Shaders/Private/PathTracing/Light/PathTracingLightCommon.ush b/Engine/Shaders/Private/PathTracing/Light/PathTracingLightCommon.ush
--- a/Engine/Shaders/Private/PathTracing/Light/PathTracingLightCommon.ush
+++ b/Engine/Shaders/Private/PathTracing/Light/PathTracingLightCommon.ush
@@ -184,9 +184,6 @@ bool IsPhysicalLight(int LightId)
 
 float3 GetTranslatedPosition(int LightId)
 {
-	//const FDFVector3 WorldPosition = DFFromTileOffset_Hack(MakeLWCVector3(SceneLights[LightId].TilePosition, SceneLights[LightId].RelativeWorldPosition));
-	//return DFFastToTranslatedWorld(WorldPosition, PreViewTranslation);
-
 	return SceneLights[LightId].TranslatedWorldPosition;
 }
 
@@ -197,12 +194,12 @@ float3 GetNormal(int LightId)
 
 float3 GetdPdu(int LightId)
 {
-	return SceneLights[LightId].dPdu;
+	return cross(SceneLights[LightId].Normal, SceneLights[LightId].Tangent);
 }
 
 float3 GetdPdv(int LightId)
 {
-	return SceneLights[LightId].dPdv;
+	return SceneLights[LightId].Tangent;
 }
 
 float GetWidth(int LightId)
diff --git a/Engine/Shaders/Shared/RayTracingTypes.h b/Engine/Shaders/Shared/RayTracingTypes.h
--- a/Engine/Shaders/Shared/RayTracingTypes.h
+++ b/Engine/Shaders/Shared/RayTracingTypes.h
@@ -42,27 +42,28 @@ struct FRTLightingData
 HLSL_STATIC_ASSERT(sizeof(FRTLightingData) == 128, "Ray tracing light structure should be kept as small as possible");
 
 // #dxr_todo: Unify this with FRTLightingData ?
-struct FPathTracingLight {
+struct FPathTracingLight
+{
 	float3  TranslatedWorldPosition;
 	float3  Normal;
-	float3  dPdu;
-	float3  dPdv;
+	float3  Tangent;
 	float3  Color;
 	float2  Dimensions; // Radius,Length or RectWidth,RectHeight or Sin(Angle/2),0 depending on light type
-	float2  Shaping;    // Barndoor controls for RectLights, Cone angles for spots lights
+	float2  Shaping;    // Barndoor controls for RectLights, Cone angles for spot lights
 	float   SpecularScale;
 	float   Attenuation;
 	float   FalloffExponent; // for non-inverse square decay lights only
 	float   VolumetricScatteringIntensity;  // scale for volume contributions
-	int     IESAtlasIndex;
+	float   IESAtlasIndex;
 	uint    Flags; // see defines PATHTRACER_FLAG_*
 	uint    MissShaderIndex;  // used to implement light functions
 	float2  RectLightAtlasUVScale;  // Rect. light atlas UV transformation
 	float2  RectLightAtlasUVOffset; // Rect. light atlas UV transformation
 };
-HLSL_STATIC_ASSERT(sizeof(FPathTracingLight) == 120, "Path tracing light structure should be kept as small as possible");
+HLSL_STATIC_ASSERT(sizeof(FPathTracingLight) == 108, "Path tracing light structure should be kept as small as possible");
 
-struct FPathTracingPackedPathState {
+struct FPathTracingPackedPathState
+{
 	uint      RandSeqSampleIndex;
 	uint      RandSeqSampleSeed;
 	float3    Radiance;
diff --git a/Engine/Source/Runtime/Renderer/Private/PathTracing.cpp b/Engine/Source/Runtime/Renderer/Private/PathTracing.cpp
--- a/Engine/Source/Runtime/Renderer/Private/PathTracing.cpp
+++ b/Engine/Source/Runtime/Renderer/Private/PathTracing.cpp
@@ -1696,7 +1696,7 @@ bool FRayTracingMeshProcessor::ProcessPathTracing(
 		return false;
 	}
 
-	TBasePassShaderElementData<FUniformLightMapPolicy> ShaderElementData(MeshBatch.LCI);
+	TBasePassShaderElementData<FUniformLightMapPolicy> ShaderElementData(nullptr);
 	ShaderElementData.InitializeMeshMaterialData(ViewIfDynamicMeshCommand, PrimitiveSceneProxy, MeshBatch, -1, true);
 
 	BuildRayTracingMeshCommands(
@@ -2108,8 +2108,7 @@ void SetLightParameters(FRDGBuilder& GraphBuilder, FPathTracingRG::FParameters*
 			DestLight.Color = FVector3f(LightParameters.Color) * LightParameters.GetLightExposureScale(View.GetLastEyeAdaptationExposure());
 			DestLight.TranslatedWorldPosition = FVector3f(LightParameters.WorldPosition + View.ViewMatrices.GetPreViewTranslation());
 			DestLight.Normal = -LightParameters.Direction;
-			DestLight.dPdu = FVector3f::CrossProduct(LightParameters.Tangent, LightParameters.Direction);
-			DestLight.dPdv = LightParameters.Tangent;
+			DestLight.Tangent = LightParameters.Tangent;
 			DestLight.Shaping = FVector2f(0.0f, 0.0f);
 			DestLight.SpecularScale = LightParameters.SpecularScale;
 			DestLight.Attenuation = LightParameters.InvRadius;
@@ -2173,8 +2172,7 @@ void SetLightParameters(FRDGBuilder& GraphBuilder, FPathTracingRG::FParameters*
 		DestLight.Color = FVector3f(LightParameters.Color) * LightParameters.GetLightExposureScale(View.GetLastEyeAdaptationExposure());
 		DestLight.TranslatedWorldPosition = FVector3f(LightParameters.WorldPosition + View.ViewMatrices.GetPreViewTranslation());
 		DestLight.Normal = -LightParameters.Direction;
-		DestLight.dPdu = FVector3f::CrossProduct(LightParameters.Tangent, LightParameters.Direction);
-		DestLight.dPdv = LightParameters.Tangent;
+		DestLight.Tangent = LightParameters.Tangent;
 		DestLight.Shaping = FVector2f(0.0f, 0.0f);
 		DestLight.SpecularScale = LightParameters.SpecularScale;
 		DestLight.Attenuation = LightParameters.InvRadius;
