Fix Reflection capture removal causes the primitive GPU state become dirty.

#tests maps provided by the UDN reporter.

diff --git a/Engine/Source/Runtime/Renderer/Private/PrimitiveSceneInfo.cpp b/Engine/Source/Runtime/Renderer/Private/PrimitiveSceneInfo.cpp
--- a/Engine/Source/Runtime/Renderer/Private/PrimitiveSceneInfo.cpp
+++ b/Engine/Source/Runtime/Renderer/Private/PrimitiveSceneInfo.cpp
@@ -2293,7 +2293,6 @@ void FPrimitiveSceneInfo::RemoveCachedReflectionCaptures()
 	CachedPlanarReflectionProxy = nullptr;
 	FMemory::Memzero(CachedReflectionCaptureProxies);
 	bNeedsCachedReflectionCaptureUpdate = true;
-	MarkGPUStateDirty(EPrimitiveDirtyState::ChangedAll);
 }
 
 void FPrimitiveSceneInfo::UpdateComponentLastRenderTime(float CurrentWorldTime, bool bUpdateLastRenderTimeOnScreen) const
