Fix post process ambient occlusion border filtering out of range. Clamp the sampling UV to ViewportMin/Max bounds. Mobile ignores this clamp as it does not use half resolution normals and Downsampled AO created in the setup.

diff --git a/Engine/Shaders/Private/PostProcessAmbientOcclusionCommon.ush b/Engine/Shaders/Private/PostProcessAmbientOcclusionCommon.ush
--- a/Engine/Shaders/Private/PostProcessAmbientOcclusionCommon.ush
+++ b/Engine/Shaders/Private/PostProcessAmbientOcclusionCommon.ush
@@ -47,6 +47,8 @@ Texture2D SSAO_DownsampledAO;
 SamplerState SSAO_Sampler;
 
 float2 SSAO_DownsampledAOInverseSize;
+float2 SSAO_DownsampledAOUVViewportMin;
+float2 SSAO_DownsampledAOUVViewportMax;
 float2 SSAO_SvPositionScaleBias;
 
 // could be moved to a more central spot
@@ -167,11 +169,16 @@ float4 ComputeUpsampleContribution(float SceneDepth, float2 InUV, float3 CenterW
 
 	UNROLL for(int i = 0; i < SampleCount; ++i)
 	{
-		float4 SampleValue = Texture2DSample(SSAO_DownsampledAO, SSAO_Sampler, UV[i]);
+#if SHADING_PATH_MOBILE
+		float2 ClampedUV = UV[i];
+#else
+		float2 ClampedUV = clamp(UV[i], SSAO_DownsampledAOUVViewportMin, SSAO_DownsampledAOUVViewportMax);
+#endif
+		float4 SampleValue = Texture2DSample(SSAO_DownsampledAO, SSAO_Sampler, ClampedUV);
 
 		MinIteration = min(MinIteration, SampleValue.g);
 
-		float4 NormalAndSampleDepth = Texture2DSample(SSAO_NormalsTexture, SSAO_Sampler, UV[i]);
+		float4 NormalAndSampleDepth = Texture2DSample(SSAO_NormalsTexture, SSAO_Sampler, ClampedUV);
 		float SampleDepth = NormalAndSampleDepth.a * Constant_Float16F_Scale;
 
 		// when tweaking this constant look for crawling pattern at edges
diff --git a/Engine/Source/Runtime/Renderer/Private/CompositionLighting/PostProcessAmbientOcclusion.cpp b/Engine/Source/Runtime/Renderer/Private/CompositionLighting/PostProcessAmbientOcclusion.cpp
--- a/Engine/Source/Runtime/Renderer/Private/CompositionLighting/PostProcessAmbientOcclusion.cpp
+++ b/Engine/Source/Runtime/Renderer/Private/CompositionLighting/PostProcessAmbientOcclusion.cpp
@@ -666,6 +666,8 @@ BEGIN_SHADER_PARAMETER_STRUCT(FAmbientOcclusionParameters, )
 	SHADER_PARAMETER_STRUCT_INCLUDE(FSSAOShaderParameters, SSAOParameters)
 
 	SHADER_PARAMETER(FVector2f, SSAO_DownsampledAOInverseSize)
+	SHADER_PARAMETER(FVector2f, SSAO_DownsampledAOUVViewportMin)
+	SHADER_PARAMETER(FVector2f, SSAO_DownsampledAOUVViewportMax)
 	SHADER_PARAMETER(FVector2f, SSAO_SvPositionScaleBias)
 
 	SHADER_PARAMETER_RDG_TEXTURE(Texture2D, SSAO_SetupTexture)
@@ -822,11 +824,19 @@ void AddAmbientOcclusionPass(
 	{
 		SharedParameters.SSAO_DownsampledAO = DownsampledAO.Texture;
 		SharedParameters.SSAO_DownsampledAOInverseSize = FVector2f(1.0f) / FVector2f(DownsampledAO.Texture->Desc.Extent);
+
+		const FVector2f ViewportMin(DownsampledAO.ViewRect.Min.X, DownsampledAO.ViewRect.Min.Y);
+		const FVector2f ViewportMax(DownsampledAO.ViewRect.Max.X, DownsampledAO.ViewRect.Max.Y);
+
+		SharedParameters.SSAO_DownsampledAOUVViewportMin = ViewportMin * SharedParameters.SSAO_DownsampledAOInverseSize;
+		SharedParameters.SSAO_DownsampledAOUVViewportMax = ViewportMax * SharedParameters.SSAO_DownsampledAOInverseSize;
 	}
 	else
 	{
 		SharedParameters.SSAO_DownsampledAO = SystemTextures.Black;
 		SharedParameters.SSAO_DownsampledAOInverseSize = FVector2f(1.0f, 1.0f);
+		SharedParameters.SSAO_DownsampledAOUVViewportMin = FVector2f(0.0f, 0.0f);
+		SharedParameters.SSAO_DownsampledAOUVViewportMax = FVector2f(1.0f, 1.0f);
 	}
 
 	SharedParameters.SSAO_SvPositionScaleBias = FVector2f(1, 0);
