Add profiler scope to RenderLightFunctionAtlas.

diff --git a/Engine/Source/Runtime/Renderer/Private/LightFunctionAtlas.cpp b/Engine/Source/Runtime/Renderer/Private/LightFunctionAtlas.cpp
--- a/Engine/Source/Runtime/Renderer/Private/LightFunctionAtlas.cpp
+++ b/Engine/Source/Runtime/Renderer/Private/LightFunctionAtlas.cpp
@@ -525,6 +525,8 @@ void FLightFunctionAtlas::RenderLightFunctionAtlas(FRDGBuilder& GraphBuilder, TA
 		return;
 	}
 
+	QUICK_SCOPE_CYCLE_COUNTER(RenderLightFunctionAtlas);
+
 	//
 	// Render the atlas
 	//
