Parallelize cached static raytracing primitives processing in FRayTracingSceneAddInstancesTask.

Context:
- FRayTracingSceneAddInstancesTask processed cached and non-cached raytracing primitives in a single-threaded loop
- Non-cached primitives cannot be efficiently parallelized due to auto-instancing logic.
- On the other hand, cached primitives processing doesn't have dependencies between primitives which means we can efficiently split it in parallel tasks.

Change:
- Modified GatherRayTracingRelevantPrimitives_ComputeLOD to output separate arrays on cached and non-cached static raytracing primitives.
- Also moved logic that filtered out cached static primitives depending on ShowFlags/cvars/etc out of RayTracingSceneStaticInstanceTask and into GatherRayTracingRelevantPrimitives_ComputeLOD.
- Preallocate range of instances in RayTracingScene and VisibleMeshCommands for cached static instances.
- Fill RayTracingScene instances and VisibleMeshCommands for cached static instances in parallel.

diff --git a/Engine/Source/Runtime/Core/Public/Async/ParallelFor.h b/Engine/Source/Runtime/Core/Public/Async/ParallelFor.h
--- a/Engine/Source/Runtime/Core/Public/Async/ParallelFor.h
+++ b/Engine/Source/Runtime/Core/Public/Async/ParallelFor.h
@@ -666,6 +666,7 @@ inline void ParallelForWithPreWorkWithExistingTaskContext(
 	*   "workspace" that can be mutated without need for synchronization primitives. For this variant, the user provides a
 	* 	callable to construct each context element.
 	*	@param OutContexts; Array that will hold the user-defined, task-level context objects (allocated per parallel task)
+	*   @param DebugName; ProfilingScope and Debugname
 	*	@param Num; number of calls of Body; Body(0), Body(1)....Body(Num - 1)
 	* 	@param ContextConstructor; Function to call to initialize each task context allocated for the operation
 	*	@param Body; Function to call from multiple threads
@@ -673,7 +674,7 @@ inline void ParallelForWithPreWorkWithExistingTaskContext(
 	*	Notes: Please add stats around to calls to parallel for and within your lambda as appropriate. Do not clog the task graph with long running tasks or tasks that block.
 **/
 template <typename ContextType, typename ContextAllocatorType, typename ContextConstructorType, typename FunctionType>
-inline void ParallelForWithTaskContext(TArray<ContextType, ContextAllocatorType>& OutContexts, int32 Num, const ContextConstructorType& ContextConstructor, const FunctionType& Body, EParallelForFlags Flags = EParallelForFlags::None)
+inline void ParallelForWithTaskContext(const TCHAR* DebugName, TArray<ContextType, ContextAllocatorType>& OutContexts, int32 Num, const ContextConstructorType& ContextConstructor, const FunctionType& Body, EParallelForFlags Flags = EParallelForFlags::None)
 {
 	if (Num > 0)
 	{
@@ -684,10 +685,28 @@ inline void ParallelForWithTaskContext(TArray<ContextType, ContextAllocatorType>
 		{
 			new(&OutContexts[ContextIndex]) ContextType(ContextConstructor(ContextIndex, NumContexts));
 		}
-		ParallelForImpl::ParallelForInternal(TEXT("ParallelFor Task"), Num, 1, Body, [](){}, Flags, TArrayView<ContextType>(OutContexts));
+		ParallelForImpl::ParallelForInternal(DebugName, Num, 1, Body, [](){}, Flags, TArrayView<ContextType>(OutContexts));
 	}
 }
 
+/** 
+	*	General purpose parallel for that uses the taskgraph. This variant constructs for the caller a user-defined context
+	* 	object for each task that may get spawned to do work, and passes it on to the loop body to give it a task-local
+	*   "workspace" that can be mutated without need for synchronization primitives. For this variant, the user provides a
+	* 	callable to construct each context element.
+	*	@param OutContexts; Array that will hold the user-defined, task-level context objects (allocated per parallel task)
+	*	@param Num; number of calls of Body; Body(0), Body(1)....Body(Num - 1)
+	* 	@param ContextConstructor; Function to call to initialize each task context allocated for the operation
+	*	@param Body; Function to call from multiple threads
+	*	@param Flags; Used to customize the behavior of the ParallelFor if needed.
+	*	Notes: Please add stats around to calls to parallel for and within your lambda as appropriate. Do not clog the task graph with long running tasks or tasks that block.
+**/
+template <typename ContextType, typename ContextAllocatorType, typename ContextConstructorType, typename FunctionType>
+inline void ParallelForWithTaskContext(TArray<ContextType, ContextAllocatorType>& OutContexts, int32 Num, const ContextConstructorType& ContextConstructor, const FunctionType& Body, EParallelForFlags Flags = EParallelForFlags::None)
+{
+	ParallelForWithTaskContext(TEXT("ParallelFor Task"), OutContexts, Num, ContextConstructor, Body, Flags);
+}
+
 /** 
 	*	General purpose parallel for that uses the taskgraph. This variant constructs for the caller a user-defined context
 	* 	object for each task that may get spawned to do work, and passes it on to the loop body to give it a task-local
diff --git a/Engine/Source/Runtime/Renderer/Private/RayTracing/RayTracing.cpp b/Engine/Source/Runtime/Renderer/Private/RayTracing/RayTracing.cpp
--- a/Engine/Source/Runtime/Renderer/Private/RayTracing/RayTracing.cpp
+++ b/Engine/Source/Runtime/Renderer/Private/RayTracing/RayTracing.cpp
@@ -162,6 +162,11 @@ namespace RayTracing
 		const FRayTracingGeometryInstance* CachedRayTracingInstance = nullptr;
 		TArrayView<const int32> CachedRayTracingMeshCommandIndices; // Pointer to FPrimitiveSceneInfo::CachedRayTracingMeshCommandIndicesPerLOD data
 
+		// Offsets relative to FRelevantPrimitiveContext offsets
+		int32 RelativeSceneInstanceOffset = INDEX_NONE;
+		int32 RelativeVisibleMeshCommandOffset = INDEX_NONE;
+		int32 ContextIndex = INDEX_NONE;
+
 		uint64 InstancingKey() const
 		{
 			uint64 Key = StateHash;
@@ -194,22 +199,34 @@ namespace RayTracing
 		}
 	};
 
+	struct FRelevantPrimitiveGatherContext
+	{
+		int32 SceneInstanceOffset = -1;
+		int32 VisibleMeshCommandOffset = -1;
+	};
+
 	struct FRelevantPrimitiveList
 	{
 		// Filtered lists of relevant primitives
 		TArray<FRelevantPrimitive> StaticPrimitives;
+		TArray<FRelevantPrimitive> CachedStaticPrimitives;
 		TArray<int32> DynamicPrimitives;
 
+		TArray<FRelevantPrimitiveGatherContext> GatherContexts;
+
 		// Relevant static primitive LODs are computed asynchronously.
-		// This task must complete before accessing StaticPrimitives in FRayTracingSceneAddInstancesTask.
+		// This task must complete before accessing StaticPrimitives/CachedStaticPrimitives in FRayTracingSceneAddInstancesTask.
 		FGraphEventRef StaticPrimitiveLODTask;
 
 		// Array of primitives that should update their cached ray tracing instances via FPrimitiveSceneInfo::UpdateCachedRaytracingData()
-		TArray<FPrimitiveSceneInfo*> DirtyCachedRayTracingPrimitives;
+		TArray<FPrimitiveSceneInfo*> DirtyCachedRayTracingPrimitives; // TODO: remove this since it seems to be transient
 
 		// Used coarse mesh streaming handles during the last TLAS build
 		TArray<Nanite::CoarseMeshStreamingHandle> UsedCoarseMeshStreamingHandles;
 
+		int32 NumCachedStaticSceneInstances = 0;
+		int32 NumCachedStaticVisibleMeshCommands = 0;
+
 		// Indicates that this object has been fully produced (for validation)
 		bool bValid = false;
 	};
@@ -399,6 +416,7 @@ namespace RayTracing
 			}
 		}
 
+		// TODO: check whether it's ok to do this on a parallel task
 		FPrimitiveSceneInfo::UpdateCachedRaytracingData(&Scene, Result.DirtyCachedRayTracingPrimitives);
 
 		static const auto ICVarStaticMeshLODDistanceScale = IConsoleManager::Get().FindConsoleVariable(TEXT("r.StaticMeshLODDistanceScale"));
@@ -408,23 +426,44 @@ namespace RayTracing
 		Result.StaticPrimitiveLODTask = FFunctionGraphTask::CreateAndDispatchWhenReady(
 			[&Result, &Scene, &View, LODScaleCVarValue, ForcedLODLevel, StaticPrimitiveIndices = MoveTemp(StaticPrimitives)]()
 			{
+				FTaskTagScope TaskTagScope(ETaskTag::EParallelRenderingThread);
+
 				TRACE_CPUPROFILER_EVENT_SCOPE(GatherRayTracingRelevantPrimitives_ComputeLOD);
 
-				Result.StaticPrimitives.SetNumUninitialized(StaticPrimitiveIndices.Num());
+				struct FRelevantStaticPrimitivesContext
+				{
+					FRelevantStaticPrimitivesContext(int32 InContextIndex) : ContextIndex(InContextIndex) {}
+
+					TChunkedArray<FRelevantPrimitive> StaticPrimitives;
+					TChunkedArray<FRelevantPrimitive> CachedStaticPrimitives;
+					TChunkedArray<const FPrimitiveSceneInfo*> VisibleNaniteRayTracingPrimitives;
+
+					int32 NumCachedStaticSceneInstances = 0;
+					int32 NumCachedStaticVisibleMeshCommands = 0;
 
-				const int32 MinBatchSize = 128;
-				ParallelFor(TEXT("GatherRayTracingRelevantPrimitives_ComputeLOD_Parallel"), StaticPrimitiveIndices.Num(), MinBatchSize,
-					[&Result, &Scene, &View, LODScaleCVarValue, ForcedLODLevel, &StaticPrimitiveIndices](int32 ItemIndex)
+					int32 ContextIndex = INDEX_NONE;
+				};
+
+				TArray<FRelevantStaticPrimitivesContext> Contexts;
+				ParallelForWithTaskContext(
+					TEXT("GatherRayTracingRelevantPrimitives_ComputeLOD_Parallel"),
+					Contexts,
+					StaticPrimitiveIndices.Num(),
+					[](int32 ContextIndex, int32 NumContexts) { return ContextIndex; },
+					[&Scene, &View, LODScaleCVarValue, ForcedLODLevel, &StaticPrimitiveIndices](FRelevantStaticPrimitivesContext& Context, int32 ItemIndex)
 					{
 						const int32 PrimitiveIndex = StaticPrimitiveIndices[ItemIndex];
 
+						const FPrimitiveSceneProxy* SceneProxy = Scene.PrimitiveSceneProxies[PrimitiveIndex];
 						const FPrimitiveSceneInfo* SceneInfo = Scene.Primitives[PrimitiveIndex];
 						const ERayTracingPrimitiveFlags Flags = Scene.PrimitiveRayTracingFlags[PrimitiveIndex];
 
-						FRelevantPrimitive& RelevantPrimitive = Result.StaticPrimitives[ItemIndex];
-						RelevantPrimitive = {};
-						RelevantPrimitive.PrimitiveIndex = PrimitiveIndex;
-						RelevantPrimitive.PersistentPrimitiveIndex = SceneInfo->GetPersistentIndex();
+						const bool bUsingNaniteRayTracing = (Nanite::GetRayTracingMode() != Nanite::ERayTracingMode::Fallback) && SceneProxy->IsNaniteMesh();
+
+						if (bUsingNaniteRayTracing)
+						{
+							Context.VisibleNaniteRayTracingPrimitives.AddElement(SceneInfo);
+						}
 
 						int8 LODIndex = 0;
 
@@ -432,7 +471,7 @@ namespace RayTracing
 						{
 							const FPrimitiveBounds& Bounds = Scene.PrimitiveBounds[PrimitiveIndex];
 
-							const int8 CurFirstLODIdx = SceneInfo->Proxy->GetCurrentFirstLODIdx_RenderThread();
+							const int8 CurFirstLODIdx = SceneProxy->GetCurrentFirstLODIdx_RenderThread();
 							check(CurFirstLODIdx >= 0);
 
 							float MeshScreenSizeSquared = 0;
@@ -444,28 +483,66 @@ namespace RayTracing
 
 						if (EnumHasAnyFlags(Flags, ERayTracingPrimitiveFlags::CacheInstances))
 						{
-							FPrimitiveSceneProxy* SceneProxy = Scene.PrimitiveSceneProxies[PrimitiveIndex];
-							const bool bUsingNaniteRayTracing = (Nanite::GetRayTracingMode() != Nanite::ERayTracingMode::Fallback) && SceneProxy->IsNaniteMesh();
-
-							if (!bUsingNaniteRayTracing)
+							if (bUsingNaniteRayTracing)
+							{
+								if (SceneInfo->CachedRayTracingInstance.GeometryRHI == nullptr)
+								{
+									// Nanite ray tracing geometry not ready yet, doesn't include primitive in ray tracing scene
+									return;
+								}
+							}
+							else
 							{
 								// Currently IsCachedRayTracingGeometryValid() can only be called for non-nanite geometries
 								checkf(SceneInfo->IsCachedRayTracingGeometryValid(), TEXT("Cached ray tracing instance is expected to be valid. Was mesh LOD streamed but cached data was not invalidated?"));
 								checkf(SceneInfo->CachedRayTracingInstance.GeometryRHI, TEXT("Ray tracing instance must have a valid geometry."));
 							}
 
-							RelevantPrimitive.bAnySegmentsDecal = SceneInfo->bCachedRayTracingInstanceAnySegmentsDecal;
-							RelevantPrimitive.bAllSegmentsDecal = SceneInfo->bCachedRayTracingInstanceAllSegmentsDecal;
-							RelevantPrimitive.CachedRayTracingInstance = &SceneInfo->CachedRayTracingInstance;
+							if (ShouldExcludeDecals() && SceneInfo->bCachedRayTracingInstanceAllSegmentsDecal)
+							{
+								return;
+							}
+
+							checkf(SceneInfo->CachedRayTracingInstance.GeometryRHI, TEXT("Ray tracing instance must have a valid geometry."));
 
 							// For primitives with ERayTracingPrimitiveFlags::CacheInstances flag we only cache the instance/mesh commands of the current LOD
 							// (see FPrimitiveSceneInfo::UpdateCachedRayTracingInstance(...) and CacheRayTracingPrimitive(...))
+							check(!EnumHasAnyFlags(Flags, ERayTracingPrimitiveFlags::ComputeLOD));
 							LODIndex = 0;
 
+							FRelevantPrimitive* RelevantPrimitive = new (Context.CachedStaticPrimitives) FRelevantPrimitive();
+							RelevantPrimitive->PrimitiveIndex = PrimitiveIndex;
+							RelevantPrimitive->PersistentPrimitiveIndex = SceneInfo->GetPersistentIndex();
+
+							RelevantPrimitive->CachedRayTracingInstance = &SceneInfo->CachedRayTracingInstance;
+							RelevantPrimitive->bAnySegmentsDecal = SceneInfo->bCachedRayTracingInstanceAnySegmentsDecal;
+							RelevantPrimitive->bAllSegmentsDecal = SceneInfo->bCachedRayTracingInstanceAllSegmentsDecal;
+
 							if (SceneInfo->CachedRayTracingMeshCommandIndicesPerLOD.IsValidIndex(LODIndex))
 							{
-								RelevantPrimitive.CachedRayTracingMeshCommandIndices = SceneInfo->CachedRayTracingMeshCommandIndicesPerLOD[LODIndex];
+								RelevantPrimitive->CachedRayTracingMeshCommandIndices = SceneInfo->CachedRayTracingMeshCommandIndicesPerLOD[LODIndex];
+							}
+							else
+							{
+								// TODO: check if this ever happens. should probably skip primitive if so
 							}
+
+							// if primitive has mixed decal and non-decal segments we need to have two ray tracing instances
+							// one containing non-decal segments and the other with decal segments
+							// masking of segments is done using "hidden" hitgroups
+							// TODO: Debug Visualization to highlight primitives using this?
+							const bool bNeedDecalInstance = RelevantPrimitive->bAnySegmentsDecal && !ShouldExcludeDecals();
+
+							const uint32 NumTLASInstances = bNeedDecalInstance && !RelevantPrimitive->bAllSegmentsDecal ? 2 : 1;
+
+							// For now store offsets relative to current context
+							// Will be patched later to be a global offset
+							RelevantPrimitive->RelativeSceneInstanceOffset = Context.NumCachedStaticSceneInstances;
+							RelevantPrimitive->RelativeVisibleMeshCommandOffset = Context.NumCachedStaticVisibleMeshCommands;
+							RelevantPrimitive->ContextIndex = Context.ContextIndex;
+
+							Context.NumCachedStaticSceneInstances += NumTLASInstances;
+							Context.NumCachedStaticVisibleMeshCommands += RelevantPrimitive->CachedRayTracingMeshCommandIndices.Num() * NumTLASInstances;
 						}
 						else
 						{
@@ -479,28 +556,32 @@ namespace RayTracing
 							// According to InitViews, we should hide the static mesh instance
 							if (SceneInfo->CachedRayTracingMeshCommandIndicesPerLOD.IsValidIndex(LODIndex))
 							{
-								RelevantPrimitive.LODIndex = LODIndex;
-								RelevantPrimitive.RayTracingGeometryRHI = RayTracingGeometryInstance;
+								FRelevantPrimitive* RelevantPrimitive = new (Context.StaticPrimitives) FRelevantPrimitive();
+								RelevantPrimitive->PrimitiveIndex = PrimitiveIndex;
+								RelevantPrimitive->PersistentPrimitiveIndex = SceneInfo->GetPersistentIndex();
+
+								RelevantPrimitive->LODIndex = LODIndex;
+								RelevantPrimitive->RayTracingGeometryRHI = RayTracingGeometryInstance;
 
-								RelevantPrimitive.CachedRayTracingMeshCommandIndices = SceneInfo->CachedRayTracingMeshCommandIndicesPerLOD[LODIndex];
-								RelevantPrimitive.StateHash = SceneInfo->CachedRayTracingMeshCommandsHashPerLOD[LODIndex];
+								RelevantPrimitive->CachedRayTracingMeshCommandIndices = SceneInfo->CachedRayTracingMeshCommandIndicesPerLOD[LODIndex];
+								RelevantPrimitive->StateHash = SceneInfo->CachedRayTracingMeshCommandsHashPerLOD[LODIndex];
 
 								// TODO: Cache these flags to avoid having to loop over the RayTracingMeshCommands
-								for (int32 CommandIndex : RelevantPrimitive.CachedRayTracingMeshCommandIndices)
+								for (int32 CommandIndex : RelevantPrimitive->CachedRayTracingMeshCommandIndices)
 								{
 									if (CommandIndex >= 0)
 									{
 										const FRayTracingMeshCommand& RayTracingMeshCommand = Scene.CachedRayTracingMeshCommands[CommandIndex];
 
-										RelevantPrimitive.InstanceMask |= RayTracingMeshCommand.InstanceMask;
-										RelevantPrimitive.bAllSegmentsOpaque &= RayTracingMeshCommand.bOpaque;
-										RelevantPrimitive.bAllSegmentsCastShadow &= RayTracingMeshCommand.bCastRayTracedShadows;
-										RelevantPrimitive.bAnySegmentsCastShadow |= RayTracingMeshCommand.bCastRayTracedShadows;
-										RelevantPrimitive.bAnySegmentsDecal |= RayTracingMeshCommand.bDecal;
-										RelevantPrimitive.bAllSegmentsDecal &= RayTracingMeshCommand.bDecal;
-										RelevantPrimitive.bTwoSided |= RayTracingMeshCommand.bTwoSided;
-										RelevantPrimitive.bIsSky |= RayTracingMeshCommand.bIsSky;
-										RelevantPrimitive.bAllSegmentsTranslucent &= RayTracingMeshCommand.bIsTranslucent;
+										RelevantPrimitive->InstanceMask |= RayTracingMeshCommand.InstanceMask;
+										RelevantPrimitive->bAllSegmentsOpaque &= RayTracingMeshCommand.bOpaque;
+										RelevantPrimitive->bAllSegmentsCastShadow &= RayTracingMeshCommand.bCastRayTracedShadows;
+										RelevantPrimitive->bAnySegmentsCastShadow |= RayTracingMeshCommand.bCastRayTracedShadows;
+										RelevantPrimitive->bAnySegmentsDecal |= RayTracingMeshCommand.bDecal;
+										RelevantPrimitive->bAllSegmentsDecal &= RayTracingMeshCommand.bDecal;
+										RelevantPrimitive->bTwoSided |= RayTracingMeshCommand.bTwoSided;
+										RelevantPrimitive->bIsSky |= RayTracingMeshCommand.bIsSky;
+										RelevantPrimitive->bAllSegmentsTranslucent &= RayTracingMeshCommand.bIsTranslucent;
 									}
 									else
 									{
@@ -511,10 +592,49 @@ namespace RayTracing
 
 								ERayTracingViewMaskMode MaskMode = static_cast<ERayTracingViewMaskMode>(Scene.CachedRayTracingMeshCommandsMode);
 
-								RelevantPrimitive.UpdateMasks(Flags, MaskMode);
+								RelevantPrimitive->UpdateMasks(Flags, MaskMode);
 							}
 						}
 					});
+
+					if (Contexts.Num() > 0)
+					{
+						SCOPED_NAMED_EVENT(GatherRayTracingRelevantPrimitives_ComputeLOD_Merge, FColor::Emerald);
+
+						uint32 NumStaticPrimitives = 0;
+						uint32 NumCachedStaticPrimitives = 0;
+
+						for (auto& Context : Contexts)
+						{
+							NumStaticPrimitives += Context.StaticPrimitives.Num();
+							NumCachedStaticPrimitives += Context.CachedStaticPrimitives.Num();
+						}
+
+						Result.StaticPrimitives.Reserve(NumStaticPrimitives);
+						Result.CachedStaticPrimitives.Reserve(NumCachedStaticPrimitives);
+
+						Result.GatherContexts.SetNum(Contexts.Num());
+
+						for (int32 ContextIndex = 0; ContextIndex < Contexts.Num(); ++ContextIndex)
+						{
+							FRelevantStaticPrimitivesContext& Context = Contexts[ContextIndex];
+							FRelevantPrimitiveGatherContext& GatherContext = Result.GatherContexts[ContextIndex];
+
+							Context.StaticPrimitives.CopyToLinearArray(Result.StaticPrimitives);
+							Context.CachedStaticPrimitives.CopyToLinearArray(Result.CachedStaticPrimitives);
+
+							GatherContext.SceneInstanceOffset = Result.NumCachedStaticSceneInstances;
+							GatherContext.VisibleMeshCommandOffset = Result.NumCachedStaticVisibleMeshCommands;
+
+							Result.NumCachedStaticSceneInstances += Context.NumCachedStaticSceneInstances;
+							Result.NumCachedStaticVisibleMeshCommands += Context.NumCachedStaticVisibleMeshCommands;
+
+							for (const FPrimitiveSceneInfo* SceneInfo : Context.VisibleNaniteRayTracingPrimitives)
+							{
+								Nanite::GRayTracingManager.AddVisiblePrimitive(SceneInfo);
+							}
+						}
+					}
 			}, TStatId(), nullptr, ENamedThreads::AnyThread);
 
 		Result.bValid = true;
@@ -911,9 +1031,14 @@ namespace RayTracing
 
 			const FScene& Scene;
 			TArray<FRelevantPrimitive>& RelevantStaticPrimitives;
+			TArray<FRelevantPrimitive>& RelevantCachedStaticPrimitives;
+			TArray<FRelevantPrimitiveGatherContext>& GatherContexts;
 			const FRayTracingCullingParameters& CullingParameters;
 			const bool bIsPathTracing;
 
+			const int32& NumCachedStaticSceneInstances;
+			const int32& NumCachedStaticVisibleMeshCommands;
+
 			// Outputs
 
 			FRayTracingScene& RayTracingScene; // New instances are added into FRayTracingScene::Instances and FRayTracingScene::Allocator is used for temporary data
@@ -921,17 +1046,25 @@ namespace RayTracing
 
 			FRayTracingSceneAddInstancesTask(const FScene& InScene,
 				TArray<FRelevantPrimitive>& InRelevantStaticPrimitives,
+				TArray<FRelevantPrimitive>& InRelevantCachedStaticPrimitives,
+				TArray<FRelevantPrimitiveGatherContext>& InGatherContexts,
 				const FRayTracingCullingParameters& InCullingParameters,
 				const bool bInIsPathTracing,
+				const int32& InNumCachedStaticSceneInstances,
+				const int32& InNumCachedStaticVisibleMeshCommands,
 				FRayTracingScene& InRayTracingScene, TArray<FVisibleRayTracingMeshCommand>& InVisibleRayTracingMeshCommands)
 				: Scene(InScene)
 				, RelevantStaticPrimitives(InRelevantStaticPrimitives)
+				, RelevantCachedStaticPrimitives(InRelevantCachedStaticPrimitives)
+				, GatherContexts(InGatherContexts)
 				, CullingParameters(InCullingParameters)
 				, bIsPathTracing(bInIsPathTracing)
+				, NumCachedStaticSceneInstances(InNumCachedStaticSceneInstances)
+				, NumCachedStaticVisibleMeshCommands(InNumCachedStaticVisibleMeshCommands)
 				, RayTracingScene(InRayTracingScene)
 				, VisibleRayTracingMeshCommands(InVisibleRayTracingMeshCommands)
 			{
-				VisibleRayTracingMeshCommands.Reserve(RelevantStaticPrimitives.Num());
+				VisibleRayTracingMeshCommands.Reserve(RelevantStaticPrimitives.Num() + RelevantCachedStaticPrimitives.Num());
 			}
 
 			// TODO: Consider moving auto instance batching logic into FRayTracingScene
@@ -992,92 +1125,26 @@ namespace RayTracing
 				FTaskTagScope TaskTagScope(ETaskTag::EParallelRenderingThread);
 
 				TRACE_CPUPROFILER_EVENT_SCOPE(RayTracingSceneStaticInstanceTask);
-
-				const bool bAutoInstance = CVarRayTracingAutoInstance.GetValueOnRenderThread() != 0;
-
-				// Instance batches by FRelevantPrimitive::InstancingKey()
-				Experimental::TSherwoodMap<uint64, FAutoInstanceBatch> InstanceBatches;
-
-				// scan relevant primitives computing hash data to look for duplicate instances
-				for (const FRelevantPrimitive& RelevantPrimitive : RelevantStaticPrimitives)
+				
 				{
-					const int32 PrimitiveIndex = RelevantPrimitive.PrimitiveIndex;
-					FPrimitiveSceneInfo* SceneInfo = Scene.Primitives[PrimitiveIndex];
-					FPrimitiveSceneProxy* SceneProxy = Scene.PrimitiveSceneProxies[PrimitiveIndex];
-					ERayTracingPrimitiveFlags Flags = Scene.PrimitiveRayTracingFlags[PrimitiveIndex];
-					const FPersistentPrimitiveIndex PersistentPrimitiveIndex = RelevantPrimitive.PersistentPrimitiveIndex;
+					TRACE_CPUPROFILER_EVENT_SCOPE(RayTracingScene_AddStaticInstances);
 
-					if (EnumHasAnyFlags(Flags, ERayTracingPrimitiveFlags::CacheInstances))
-					{
-						const bool bUsingNaniteRayTracing = (Nanite::GetRayTracingMode() != Nanite::ERayTracingMode::Fallback) && SceneProxy->IsNaniteMesh();
+					const bool bAutoInstance = CVarRayTracingAutoInstance.GetValueOnRenderThread() != 0;
 
-						if (bUsingNaniteRayTracing)
-						{
-							Nanite::GRayTracingManager.AddVisiblePrimitive(SceneInfo);
+					// Instance batches by FRelevantPrimitive::InstancingKey()
+					Experimental::TSherwoodMap<uint64, FAutoInstanceBatch> InstanceBatches;
 
-							if (RelevantPrimitive.CachedRayTracingInstance->GeometryRHI == nullptr)
-							{
-								// Nanite ray tracing geometry not ready yet, doesn't include primitive in ray tracing scene
-								continue;
-							}
-						}
-
-						// TODO: Consider requesting a recache of all ray tracing commands during which decals are excluded
-
-						// if primitive has mixed decal and non-decal segments we need to have two ray tracing instances
-						// one containing non-decal segments and the other with decal segments
-						// masking of segments is done using "hidden" hitgroups
-						// TODO: Debug Visualization to highlight primitives using this?
-						const bool bNeedDecalInstance = RelevantPrimitive.bAnySegmentsDecal && !ShouldExcludeDecals();
-
-						if (ShouldExcludeDecals() && RelevantPrimitive.bAllSegmentsDecal)
-						{
-							continue;
-						}
-
-						check(RelevantPrimitive.CachedRayTracingInstance);
-
-						int32 InstanceIndex = INDEX_NONE;
-						if (!RelevantPrimitive.bAllSegmentsDecal)
-						{
-							FRayTracingGeometryInstance RayTracingInstance = *RelevantPrimitive.CachedRayTracingInstance;
-							RayTracingInstance.LayerIndex = (uint8)ERayTracingSceneLayer::Base;
-							AddDebugRayTracingInstanceFlags(RayTracingInstance.Flags);
-
-							InstanceIndex = RayTracingScene.AddInstance(MoveTemp(RayTracingInstance), SceneProxy, false);
-						}
-
-						uint32 DecalInstanceIndex = INDEX_NONE;
-						if (bNeedDecalInstance)
-						{
-							FRayTracingGeometryInstance DecalRayTracingInstance = *RelevantPrimitive.CachedRayTracingInstance;
-							DecalRayTracingInstance.LayerIndex = (uint8)ERayTracingSceneLayer::Decals;
-							AddDebugRayTracingInstanceFlags(DecalRayTracingInstance.Flags);
-
-							DecalInstanceIndex = RayTracingScene.AddInstance(MoveTemp(DecalRayTracingInstance), SceneProxy, false);
-						}
+					// scan relevant primitives computing hash data to look for duplicate instances
+					for (const FRelevantPrimitive& RelevantPrimitive : RelevantStaticPrimitives)
+					{
+						const int32 PrimitiveIndex = RelevantPrimitive.PrimitiveIndex;
+						FPrimitiveSceneInfo* SceneInfo = Scene.Primitives[PrimitiveIndex];
+						FPrimitiveSceneProxy* SceneProxy = Scene.PrimitiveSceneProxies[PrimitiveIndex];
+						ERayTracingPrimitiveFlags Flags = Scene.PrimitiveRayTracingFlags[PrimitiveIndex];
+						const FPersistentPrimitiveIndex PersistentPrimitiveIndex = RelevantPrimitive.PersistentPrimitiveIndex;
 
-						for (int32 CommandIndex : RelevantPrimitive.CachedRayTracingMeshCommandIndices)
-						{
-							const FRayTracingMeshCommand& MeshCommand = Scene.CachedRayTracingMeshCommands[CommandIndex];
+						check(!EnumHasAnyFlags(Flags, ERayTracingPrimitiveFlags::CacheInstances));
 
-							if(InstanceIndex != INDEX_NONE)
-							{
-								const bool bHidden = MeshCommand.bDecal;
-								FVisibleRayTracingMeshCommand NewVisibleMeshCommand(&MeshCommand, InstanceIndex, bHidden);
-								VisibleRayTracingMeshCommands.Add(NewVisibleMeshCommand);
-							}
-
-							if (DecalInstanceIndex != INDEX_NONE)
-							{
-								const bool bHidden = !MeshCommand.bDecal;
-								FVisibleRayTracingMeshCommand NewVisibleMeshCommand(&MeshCommand, DecalInstanceIndex, bHidden);
-								VisibleRayTracingMeshCommands.Add(NewVisibleMeshCommand);
-							}
-						}
-					}
-					else
-					{
 						const int8 LODIndex = RelevantPrimitive.LODIndex;
 
 						if (LODIndex < 0)
@@ -1212,6 +1279,95 @@ namespace RayTracing
 						}
 					}
 				}
+
+				{
+					TRACE_CPUPROFILER_EVENT_SCOPE(RayTracingScene_AddCachedStaticInstances);
+
+					const uint32 BaseCachedStaticInstanceIndex = RayTracingScene.AddInstancesUninitialized(NumCachedStaticSceneInstances);
+
+					const uint32 BaseCachedVisibleMeshCommandsIndex = VisibleRayTracingMeshCommands.AddUninitialized(NumCachedStaticVisibleMeshCommands);
+					TArrayView<FVisibleRayTracingMeshCommand> CachedStaticVisibleRayTracingMeshCommands = TArrayView<FVisibleRayTracingMeshCommand>(VisibleRayTracingMeshCommands.GetData() + BaseCachedVisibleMeshCommandsIndex, NumCachedStaticVisibleMeshCommands);
+
+					const int32 MinBatchSize = 128;
+					ParallelFor(
+						TEXT("RayTracingScene_AddCachedStaticInstances_ParallelFor"),
+						RelevantCachedStaticPrimitives.Num(),
+						MinBatchSize,
+						[this, BaseCachedStaticInstanceIndex, CachedStaticVisibleRayTracingMeshCommands](int32 Index)
+					{
+						const FRelevantPrimitive& RelevantPrimitive = RelevantCachedStaticPrimitives[Index];
+						const int32 PrimitiveIndex = RelevantPrimitive.PrimitiveIndex;
+						FPrimitiveSceneInfo* SceneInfo = Scene.Primitives[PrimitiveIndex];
+						FPrimitiveSceneProxy* SceneProxy = Scene.PrimitiveSceneProxies[PrimitiveIndex];
+						ERayTracingPrimitiveFlags Flags = Scene.PrimitiveRayTracingFlags[PrimitiveIndex];
+						const FPersistentPrimitiveIndex PersistentPrimitiveIndex = RelevantPrimitive.PersistentPrimitiveIndex;
+
+						check(EnumHasAnyFlags(Flags, ERayTracingPrimitiveFlags::CacheInstances));
+
+						const bool bUsingNaniteRayTracing = (Nanite::GetRayTracingMode() != Nanite::ERayTracingMode::Fallback) && SceneProxy->IsNaniteMesh();
+
+						if (bUsingNaniteRayTracing)
+						{
+							check(RelevantPrimitive.CachedRayTracingInstance->GeometryRHI != nullptr);
+						}
+
+						// if primitive has mixed decal and non-decal segments we need to have two ray tracing instances
+						// one containing non-decal segments and the other with decal segments
+						// masking of segments is done using "hidden" hitgroups
+						// TODO: Debug Visualization to highlight primitives using this?
+						const bool bNeedDecalInstance = RelevantPrimitive.bAnySegmentsDecal && !RelevantPrimitive.bAllSegmentsDecal && !ShouldExcludeDecals();
+
+						check(!ShouldExcludeDecals() || !RelevantPrimitive.bAllSegmentsDecal);
+
+						check(RelevantPrimitive.CachedRayTracingInstance);
+
+						int32 TmpInstanceIndex = BaseCachedStaticInstanceIndex + GatherContexts[RelevantPrimitive.ContextIndex].SceneInstanceOffset + RelevantPrimitive.RelativeSceneInstanceOffset;
+
+						int32 InstanceIndex = INDEX_NONE;
+						if (!RelevantPrimitive.bAllSegmentsDecal)
+						{
+							FRayTracingGeometryInstance RayTracingInstance = *RelevantPrimitive.CachedRayTracingInstance;
+							RayTracingInstance.LayerIndex = (uint8)ERayTracingSceneLayer::Base;
+							AddDebugRayTracingInstanceFlags(RayTracingInstance.Flags);
+
+							InstanceIndex = TmpInstanceIndex++;
+							RayTracingScene.SetInstance(InstanceIndex, MoveTemp(RayTracingInstance), SceneProxy, false);
+						}
+
+						uint32 DecalInstanceIndex = INDEX_NONE;
+						if (bNeedDecalInstance)
+						{
+							FRayTracingGeometryInstance DecalRayTracingInstance = *RelevantPrimitive.CachedRayTracingInstance;
+							DecalRayTracingInstance.LayerIndex = (uint8)ERayTracingSceneLayer::Decals;
+							AddDebugRayTracingInstanceFlags(DecalRayTracingInstance.Flags);
+
+							DecalInstanceIndex = TmpInstanceIndex++;
+							RayTracingScene.SetInstance(DecalInstanceIndex, MoveTemp(DecalRayTracingInstance), SceneProxy, false);
+						}
+
+						const int32 VisibleMeshCommandOffset = GatherContexts[RelevantPrimitive.ContextIndex].VisibleMeshCommandOffset + RelevantPrimitive.RelativeVisibleMeshCommandOffset;
+						int32 CommandCount = 0;
+
+						for (int32 CommandIndex : RelevantPrimitive.CachedRayTracingMeshCommandIndices)
+						{
+							const FRayTracingMeshCommand& MeshCommand = Scene.CachedRayTracingMeshCommands[CommandIndex];
+
+							if (InstanceIndex != INDEX_NONE)
+							{
+								const bool bHidden = MeshCommand.bDecal;
+								CachedStaticVisibleRayTracingMeshCommands[VisibleMeshCommandOffset + CommandCount] = FVisibleRayTracingMeshCommand(&MeshCommand, InstanceIndex, bHidden);
+								++CommandCount;
+							}
+
+							if (DecalInstanceIndex != INDEX_NONE)
+							{
+								const bool bHidden = !MeshCommand.bDecal;
+								CachedStaticVisibleRayTracingMeshCommands[VisibleMeshCommandOffset + CommandCount] = FVisibleRayTracingMeshCommand(&MeshCommand, DecalInstanceIndex, bHidden);
+								++CommandCount;
+							}
+						}
+					});
+				}
 			}
 		};
 
@@ -1219,8 +1375,18 @@ namespace RayTracing
 		AddInstancesTaskPrerequisites.Add(RelevantPrimitiveList.StaticPrimitiveLODTask);
 
 		FGraphEventRef AddInstancesTask = TGraphTask<FRayTracingSceneAddInstancesTask>::CreateTask(&AddInstancesTaskPrerequisites).ConstructAndDispatchWhenReady(
-			Scene, RelevantPrimitiveList.StaticPrimitives, View.RayTracingCullingParameters, bool(View.Family->EngineShowFlags.PathTracing), // inputs 
-			RayTracingScene, View.VisibleRayTracingMeshCommands // outputs
+			// inputs
+			Scene,
+			RelevantPrimitiveList.StaticPrimitives,
+			RelevantPrimitiveList.CachedStaticPrimitives,
+			RelevantPrimitiveList.GatherContexts,
+			View.RayTracingCullingParameters,
+			bool(View.Family->EngineShowFlags.PathTracing),
+			RelevantPrimitiveList.NumCachedStaticSceneInstances,
+			RelevantPrimitiveList.NumCachedStaticVisibleMeshCommands,
+			// outputs
+			RayTracingScene,
+			View.VisibleRayTracingMeshCommands
 		);
 
 		// Scene init task can run only when all pre-init tasks are complete (including culling tasks that are spawned while adding instances)
@@ -1241,4 +1407,7 @@ namespace RayTracing
 	}
 }
 
+static_assert(TIsTriviallyDestructible<RayTracing::FRelevantPrimitive>::Value == true, "FRelevantPrimitive must be trivially destructible");
+template <> struct TIsPODType<RayTracing::FRelevantPrimitive> { enum { Value = true }; }; // Necessary to use TChunkedArray::CopyToLinearArray
+
 #endif //RHI_RAYTRACING
diff --git a/Engine/Source/Runtime/Renderer/Private/RayTracing/RayTracingScene.cpp b/Engine/Source/Runtime/Renderer/Private/RayTracing/RayTracingScene.cpp
--- a/Engine/Source/Runtime/Renderer/Private/RayTracing/RayTracingScene.cpp
+++ b/Engine/Source/Runtime/Renderer/Private/RayTracing/RayTracingScene.cpp
@@ -348,11 +348,49 @@ uint32 FRayTracingScene::AddInstance(FRayTracingGeometryInstance Instance, const
 		{
 			InstanceDebugData.ProxyHash = Proxy->GetTypeHash();
 		}
+
+		check(Instances.Num() == InstancesDebugData.Num());
 	}
 
 	return InstanceIndex;
 }
 
+uint32 FRayTracingScene::AddInstancesUninitialized(uint32 NumInstances)
+{
+	const uint32 OldNum = Instances.AddUninitialized(NumInstances);
+
+	if (bInstanceDebugDataEnabled)
+	{
+		InstancesDebugData.AddUninitialized(NumInstances);
+
+		check(Instances.Num() == InstancesDebugData.Num());
+	}
+
+	return OldNum;
+}
+
+void FRayTracingScene::SetInstance(uint32 InstanceIndex, FRayTracingGeometryInstance InInstance, const FPrimitiveSceneProxy* Proxy, bool bDynamic)
+{
+	FRHIRayTracingGeometry* GeometryRHI = InInstance.GeometryRHI;
+
+	FRayTracingGeometryInstance* Instance = &Instances[InstanceIndex];
+	new (Instance) FRayTracingGeometryInstance(MoveTemp(InInstance));
+
+	if (bInstanceDebugDataEnabled)
+	{
+		FRayTracingInstanceDebugData InstanceDebugData;
+		InstanceDebugData.Flags = bDynamic ? 1 : 0;
+		InstanceDebugData.GeometryAddress = uint64(GeometryRHI);
+
+		if (Proxy)
+		{
+			InstanceDebugData.ProxyHash = Proxy->GetTypeHash();
+		}
+
+		InstancesDebugData[InstanceIndex] = InstanceDebugData;
+	}
+}
+
 void FRayTracingScene::Reset(bool bInInstanceDebugDataEnabled)
 {
 	WaitForTasks();
diff --git a/Engine/Source/Runtime/Renderer/Private/RayTracing/RayTracingScene.h b/Engine/Source/Runtime/Renderer/Private/RayTracing/RayTracingScene.h
--- a/Engine/Source/Runtime/Renderer/Private/RayTracing/RayTracingScene.h
+++ b/Engine/Source/Runtime/Renderer/Private/RayTracing/RayTracingScene.h
@@ -51,6 +51,10 @@ public:
 
 	uint32 AddInstance(FRayTracingGeometryInstance Instance, const FPrimitiveSceneProxy* Proxy = nullptr, bool bDynamic = false);
 
+	uint32 AddInstancesUninitialized(uint32 NumInstances);
+
+	void SetInstance(uint32 InstanceIndex, FRayTracingGeometryInstance Instance, const FPrimitiveSceneProxy* Proxy = nullptr, bool bDynamic = false);
+
 	// Allocates RayTracingSceneRHI and builds various metadata required to create the final scene.
 	FRayTracingSceneWithGeometryInstances BuildInitializationData() const;
 
