Add runtime checks for HasEditorOnlyData() in DebugProbeRendering.cpp
- in some cases builds can be compiled with WITH_EDITORONLY_DATA set but have editor only data stripped during cook.

diff --git a/Engine/Source/Runtime/Renderer/Private/DebugProbeRendering.cpp b/Engine/Source/Runtime/Renderer/Private/DebugProbeRendering.cpp
--- a/Engine/Source/Runtime/Renderer/Private/DebugProbeRendering.cpp
+++ b/Engine/Source/Runtime/Renderer/Private/DebugProbeRendering.cpp
@@ -98,6 +98,8 @@ static void CommonStampDeferredDebugProbeDrawCall(
 	int32 RenderPass,
 	bool bIlluminanceMeter)
 {
+	check(FPlatformProperties::HasEditorOnlyData());
+
 	PassParameters->ViewUniformBuffer = View.ViewUniformBuffer;
 	PassParameters->MaterialTextureArrayUAV = View.SubstrateViewData.SceneData->MaterialTextureArrayUAVWithoutRTs;
 	PassParameters->MaxBytesPerPixel = View.SubstrateViewData.SceneData->EffectiveMaxBytesPerPixel;
@@ -125,6 +127,11 @@ void StampDeferredDebugProbeDepthPS(
 	const FRDGTextureRef SceneDepthTexture)
 {
 #if DEBUG_PROBE_ENABLED
+	if (FPlatformProperties::HasEditorOnlyData())
+	{
+		return;
+	}
+
 	RDG_EVENT_SCOPE(GraphBuilder, "StampDeferredDebugProbeDepth");
 	RDG_GPU_STAT_SCOPE(GraphBuilder, StampDeferredDebugProbe);
 
@@ -154,6 +161,11 @@ void StampDeferredDebugProbeMaterialPS(
 	const FMinimalSceneTextures& SceneTextures)
 {
 #if DEBUG_PROBE_ENABLED
+	if (FPlatformProperties::HasEditorOnlyData())
+	{
+		return;
+	}
+
 	RDG_EVENT_SCOPE(GraphBuilder, "StampDeferredDebugProbeMaterial");
 	RDG_GPU_STAT_SCOPE(GraphBuilder, StampDeferredDebugProbe);
 
@@ -214,6 +226,11 @@ void StampDeferredDebugProbeVelocityPS(
 	const FRenderTargetBindingSlots& BasePassRenderTargets)
 {
 #if DEBUG_PROBE_ENABLED
+	if (FPlatformProperties::HasEditorOnlyData())
+	{
+		return;
+	}
+
 	RDG_EVENT_SCOPE(GraphBuilder, "StampDeferredDebugProbeVelocity");
 	RDG_GPU_STAT_SCOPE(GraphBuilder, StampDeferredDebugProbe);
 
