Avoid binding the RT lighting miss shaders if we are not going to be using them

This avoid a crash from invalid buffers being bound when ray tracing is enabled.

diff --git a/Engine/Source/Runtime/Renderer/Private/DeferredShadingRenderer.cpp b/Engine/Source/Runtime/Renderer/Private/DeferredShadingRenderer.cpp
--- a/Engine/Source/Runtime/Renderer/Private/DeferredShadingRenderer.cpp
+++ b/Engine/Source/Runtime/Renderer/Private/DeferredShadingRenderer.cpp
@@ -1027,9 +1027,15 @@ void FDeferredShadingSceneRenderer::WaitForRayTracingScene(FRDGBuilder& GraphBui
 			else
 			{
 				SetupRayTracingDefaultMissShader(RHICmdList, ReferenceView);
-				SetupRayTracingLightingMissShader(RHICmdList, ReferenceView);
+
+				// only a few passes use the light miss shaders (those that do direct lighting)
+				// simply skip setting these up if we haven't built the RT lighting data
+				if (ReferenceView.RayTracingLightGridUniformBuffer != nullptr)
+				{
+					SetupRayTracingLightingMissShader(RHICmdList, ReferenceView);
 				
-				BindLightFunctionShaders(RHICmdList, Scene, RayTracingLightFunctionMap, ReferenceView);
+					BindLightFunctionShaders(RHICmdList, Scene, RayTracingLightFunctionMap, ReferenceView);
+				}
 			}
 		}
 
