Improvements to deferred decal CPU performance
- Use ParallelFor in BuildVisibleDecalList(...)
  - using TChunkedArray to store intermediate results required some changes to FTransientDecalRenderData
- Only sort decals in BuildRelevantDecalList(...)
  - removed unecessary sort in BuildVisibleDecalList(...)
  - in the future we could move the sort earlier in the frame and run on async task (eg: kick-off at end of BuildVisibleDecalList(...) and wait in BuildRelevantDecalList(...)

Level containing ~2000 decals running on console: 3ms -> 1.5ms

diff --git a/Engine/Source/Runtime/Renderer/Private/CompositionLighting/PostProcessDeferredDecals.cpp b/Engine/Source/Runtime/Renderer/Private/CompositionLighting/PostProcessDeferredDecals.cpp
--- a/Engine/Source/Runtime/Renderer/Private/CompositionLighting/PostProcessDeferredDecals.cpp
+++ b/Engine/Source/Runtime/Renderer/Private/CompositionLighting/PostProcessDeferredDecals.cpp
@@ -631,7 +631,7 @@ void AddDeferredDecalPass(
 			for (uint32 DecalIndex = DecalIndexBegin; DecalIndex < DecalIndexEnd; ++DecalIndex)
 			{
 				const FTransientDecalRenderData& DecalData = (*SortedDecals)[DecalIndex];
-				const FDeferredDecalProxy& DecalProxy = DecalData.Proxy;
+				const FDeferredDecalProxy& DecalProxy = *DecalData.Proxy;
 				const FMatrix ComponentToWorldMatrix = DecalProxy.ComponentTrans.ToMatrixWithScale();
 				const FMatrix FrustumComponentToClip = DecalRendering::ComputeComponentToClipMatrix(View, ComponentToWorldMatrix);
 				const bool bStencilThisDecal = IsStencilOptimizationAvailable(DecalRenderStage);
diff --git a/Engine/Source/Runtime/Renderer/Private/DecalRenderingShared.cpp b/Engine/Source/Runtime/Renderer/Private/DecalRenderingShared.cpp
--- a/Engine/Source/Runtime/Renderer/Private/DecalRenderingShared.cpp
+++ b/Engine/Source/Runtime/Renderer/Private/DecalRenderingShared.cpp
@@ -15,6 +15,7 @@
 #include "ScenePrivate.h"
 #include "PipelineStateCache.h"
 #include "MobileBasePassRendering.h"
+#include "Async/ParallelFor.h"
 
 static TAutoConsoleVariable<float> CVarDecalFadeScreenSizeMultiplier(
 	TEXT("r.Decal.FadeScreenSizeMult"),
@@ -23,11 +24,19 @@ static TAutoConsoleVariable<float> CVarDecalFadeScreenSizeMultiplier(
 	TEXT("  Smaller means decals fade less aggressively.")
 	);
 
-FTransientDecalRenderData::FTransientDecalRenderData(const FDeferredDecalProxy& InDecalProxy, float InConservativeRadius, EShaderPlatform ShaderPlatform, ERHIFeatureLevel::Type FeatureLevel)
-	: Proxy(InDecalProxy)
+static TAutoConsoleVariable<bool> CVarDecalVisibilityMultithreaded(
+	TEXT("r.Decal.Visibility.Multithreaded"),
+	true,
+	TEXT("Whether to build visible decal list using multithreading. 0=disabled, 1=enabled (default)"),
+	ECVF_RenderThreadSafe);
+
+template<> struct TIsPODType<FTransientDecalRenderData> { enum { Value = true }; };
+
+FTransientDecalRenderData::FTransientDecalRenderData(const FDeferredDecalProxy& InDecalProxy, float InConservativeRadius, float InFadeAlpha, EShaderPlatform ShaderPlatform, ERHIFeatureLevel::Type FeatureLevel)
+	: Proxy(&InDecalProxy)
 	, MaterialProxy(InDecalProxy.DecalMaterial->GetRenderProxy())
 	, ConservativeRadius(InConservativeRadius)
-	, FadeAlpha(1.0f)
+	, FadeAlpha(InFadeAlpha)
 	, DecalColor(InDecalProxy.DecalColor)
 {
 	// Build BlendDesc from a potentially incomplete material.
@@ -249,9 +258,9 @@ namespace DecalRendering
 		{
 			FORCEINLINE bool operator()(const FTransientDecalRenderData& A, const FTransientDecalRenderData& B) const
 			{
-				if (B.Proxy.SortOrder != A.Proxy.SortOrder)
+				if (B.Proxy->SortOrder != A.Proxy->SortOrder)
 				{ 
-					return A.Proxy.SortOrder < B.Proxy.SortOrder;
+					return A.Proxy->SortOrder < B.Proxy->SortOrder;
 				}
 				if (B.BlendDesc.bWriteNormal != A.BlendDesc.bWriteNormal)
 				{
@@ -269,7 +278,7 @@ namespace DecalRendering
 					// Batch decals with the same material together
 					return B.MaterialProxy < A.MaterialProxy;
 				}
-				return (PTRINT)B.Proxy.Component < (PTRINT)A.Proxy.Component;
+				return (PTRINT)B.Proxy->Component < (PTRINT)A.Proxy->Component;
 			}
 		};
 
@@ -277,6 +286,51 @@ namespace DecalRendering
 		Decals.Sort(FCompareFTransientDecalRenderData());
 	}
 
+	static bool ProcessDecal(const FDeferredDecalProxy* DecalProxy, const FViewInfo& View, float FadeMultiplier, EShaderPlatform ShaderPlatform, bool bIsPerspectiveProjection, FTransientDecalRenderData& OutData)
+	{
+		if (!DecalProxy->DecalMaterial || !DecalProxy->DecalMaterial->IsValidLowLevelFast())
+		{
+			return false;
+		}
+
+		if (!DecalProxy->IsShown(&View))
+		{
+			return false;
+		}
+
+		const FMatrix ComponentToWorldMatrix = DecalProxy->ComponentTrans.ToMatrixWithScale();
+
+		// can be optimized as we test against a sphere around the box instead of the box itself
+		const float ConservativeRadius = FMath::Sqrt(
+			ComponentToWorldMatrix.GetScaledAxis(EAxis::X).SizeSquared() +
+			ComponentToWorldMatrix.GetScaledAxis(EAxis::Y).SizeSquared() +
+			ComponentToWorldMatrix.GetScaledAxis(EAxis::Z).SizeSquared());
+
+		// can be optimized as the test is too conservative (sphere instead of OBB)
+		if (ConservativeRadius < SMALL_NUMBER || !View.ViewFrustum.IntersectSphere(ComponentToWorldMatrix.GetOrigin(), ConservativeRadius))
+		{
+			return false;
+		}
+
+		float FadeAlpha = 1.0f;
+
+		if (bIsPerspectiveProjection && DecalProxy->FadeScreenSize != 0.0f)
+		{
+			FadeAlpha = CalculateDecalFadeAlpha(DecalProxy->FadeScreenSize, ComponentToWorldMatrix, View, FadeMultiplier);
+		}
+
+		const bool bShouldRender = FadeAlpha > 0.0f;
+
+		if (!bShouldRender)
+		{
+			return false;
+		}
+
+		OutData = FTransientDecalRenderData(*DecalProxy, ConservativeRadius, FadeAlpha, ShaderPlatform, View.GetFeatureLevel());
+
+		return true;
+	};
+
 	FTransientDecalRenderDataList BuildVisibleDecalList(TConstArrayView<FDeferredDecalProxy*> Decals, const FViewInfo& View)
 	{
 		QUICK_SCOPE_CYCLE_COUNTER(BuildVisibleDecalList);
@@ -295,56 +349,66 @@ namespace DecalRendering
 
 		const bool bIsPerspectiveProjection = View.IsPerspectiveProjection();
 
-		// Build a list of decals that need to be rendered for this view in SortedDecals
-		for (const FDeferredDecalProxy* DecalProxy : Decals)
+		// Build a list of decals that need to be rendered for this view
+
+		if (CVarDecalVisibilityMultithreaded.GetValueOnRenderThread() && FApp::ShouldUseThreadingForPerformance())
 		{
-			if (!DecalProxy->DecalMaterial || !DecalProxy->DecalMaterial->IsValidLowLevelFast())
+			struct FVisibleDecalListContext
 			{
-				continue;
-			}
-
-			bool bIsShown = true;
+				TChunkedArray<FTransientDecalRenderData> VisibleDecals;
+			};
+
+			TArray<FVisibleDecalListContext> Contexts;
+			const int32 MinBatchSize = 128;
+			ParallelForWithTaskContext(
+				TEXT("BuildVisibleDecalList_Parallel"),
+				Contexts,
+				Decals.Num(),
+				MinBatchSize,
+				[Decals, &View, ShaderPlatform, bIsPerspectiveProjection, FadeMultiplier](FVisibleDecalListContext& Context, int32 ItemIndex)
+				{
+					FTaskTagScope TaskTagScope(ETaskTag::EParallelRenderingThread);
 
-			if (!DecalProxy->IsShown(&View))
-			{
-				bIsShown = false;
-			}
+					const FDeferredDecalProxy* DecalProxy = Decals[ItemIndex];
 
-			const FMatrix ComponentToWorldMatrix = DecalProxy->ComponentTrans.ToMatrixWithScale();
+					FTransientDecalRenderData Data;
 
-			// can be optimized as we test against a sphere around the box instead of the box itself
-			const float ConservativeRadius = FMath::Sqrt(
-				ComponentToWorldMatrix.GetScaledAxis(EAxis::X).SizeSquared() +
-				ComponentToWorldMatrix.GetScaledAxis(EAxis::Y).SizeSquared() +
-				ComponentToWorldMatrix.GetScaledAxis(EAxis::Z).SizeSquared());
+					if (ProcessDecal(DecalProxy, View, FadeMultiplier, ShaderPlatform, bIsPerspectiveProjection, Data))
+					{
+						Context.VisibleDecals.AddElement(MoveTemp(Data));
+					}
+				});
 
-			// can be optimized as the test is too conservative (sphere instead of OBB)
-			if (ConservativeRadius < SMALL_NUMBER || !View.ViewFrustum.IntersectSphere(ComponentToWorldMatrix.GetOrigin(), ConservativeRadius))
+			if (Contexts.Num() > 0)
 			{
-				bIsShown = false;
-			}
+				SCOPED_NAMED_EVENT(BuildVisibleDecalList_Parallel_Merge, FColor::Emerald);
 
-			if (bIsShown)
-			{
-				FTransientDecalRenderData Data(*DecalProxy, ConservativeRadius, ShaderPlatform, View.GetFeatureLevel());
+				int32 NumVisibleDecals = 0;
 
-				if (bIsPerspectiveProjection && Data.Proxy.FadeScreenSize != 0.0f)
+				for (auto& Context : Contexts)
 				{
-					Data.FadeAlpha = CalculateDecalFadeAlpha(Data.Proxy.FadeScreenSize, ComponentToWorldMatrix, View, FadeMultiplier);
+					NumVisibleDecals += Context.VisibleDecals.Num();
 				}
 
-				const bool bShouldRender = Data.FadeAlpha > 0.0f;
+				OutVisibleDecals.Empty(NumVisibleDecals);
 
-				if (bShouldRender)
+				for (auto& Context : Contexts)
 				{
-					OutVisibleDecals.Add(Data);
+					Context.VisibleDecals.CopyToLinearArray(OutVisibleDecals);
 				}
 			}
 		}
-
-		if (OutVisibleDecals.Num() > 0)
+		else
 		{
-			SortDecalList(OutVisibleDecals);
+			for (const FDeferredDecalProxy* DecalProxy : Decals)
+			{
+				FTransientDecalRenderData Data;
+
+				if (ProcessDecal(DecalProxy, View, FadeMultiplier, ShaderPlatform, bIsPerspectiveProjection, Data))
+				{
+					OutVisibleDecals.Add(MoveTemp(Data));
+				}
+			}
 		}
 
 		return OutVisibleDecals;
@@ -362,8 +426,8 @@ namespace DecalRendering
 		// Build a list of decals that need to be rendered for this stage in SortedDecals
 		for (const FTransientDecalRenderData& DecalRenderData : Decals)
 		{
-			checkf(DecalRenderData.Proxy.DecalMaterial && DecalRenderData.Proxy.DecalMaterial->IsValidLowLevelFast(),
-				TEXT("Decals should've been filtered earlier in BuildVisibleDecalList"));
+			checkf(DecalRenderData.Proxy->DecalMaterial && DecalRenderData.Proxy->DecalMaterial->IsValidLowLevelFast(),
+				TEXT("Decals should've been filtered earlier in BuildVisibleDecalList(...)"));
 
 			if (IsCompatibleWithRenderStage(DecalRenderData.BlendDesc, DecalRenderStage))
 			{
@@ -383,6 +447,7 @@ namespace DecalRendering
 
 		if (OutVisibleDecals->Num() > 0)
 		{
+			// Could potentially run sort of auxiliary thread (kicked-off in BuildVisibleDecalList(...))
 			SortDecalList(*OutVisibleDecals);
 
 			return true;
@@ -499,7 +564,7 @@ namespace DecalRendering
 
 		// Set pixel shader parameters.
 		{
-			SetShaderParametersLegacyPS(RHICmdList, PixelShader, View, DecalData.Proxy, MaterialProxy, MaterialResource, DecalData.FadeAlpha);
+			SetShaderParametersLegacyPS(RHICmdList, PixelShader, View, *DecalData.Proxy, MaterialProxy, MaterialResource, DecalData.FadeAlpha);
 		}
 
 		// Set stream source after updating cached strides
diff --git a/Engine/Source/Runtime/Renderer/Private/DecalRenderingShared.h b/Engine/Source/Runtime/Renderer/Private/DecalRenderingShared.h
--- a/Engine/Source/Runtime/Renderer/Private/DecalRenderingShared.h
+++ b/Engine/Source/Runtime/Renderer/Private/DecalRenderingShared.h
@@ -17,14 +17,15 @@ class FViewInfo;
  */
 struct FTransientDecalRenderData
 {
-	const FDeferredDecalProxy& Proxy;
+	const FDeferredDecalProxy* Proxy;
 	const FMaterialRenderProxy* MaterialProxy;
 	FDecalBlendDesc BlendDesc;
 	float ConservativeRadius;
 	float FadeAlpha;
 	FLinearColor DecalColor;
 
-	FTransientDecalRenderData(const FDeferredDecalProxy& InDecalProxy, float InConservativeRadius, EShaderPlatform ShaderPlatform, ERHIFeatureLevel::Type FeatureLevel);
+	FTransientDecalRenderData() = default; // required to support TChunkedArray
+	FTransientDecalRenderData(const FDeferredDecalProxy& InDecalProxy, float InConservativeRadius, float InFadeAlpha, EShaderPlatform ShaderPlatform, ERHIFeatureLevel::Type FeatureLevel);
 };
 	
 typedef TArray<FTransientDecalRenderData, SceneRenderingAllocator> FTransientDecalRenderDataList;
diff --git a/Engine/Source/Runtime/Renderer/Private/MobileDecalRendering.cpp b/Engine/Source/Runtime/Renderer/Private/MobileDecalRendering.cpp
--- a/Engine/Source/Runtime/Renderer/Private/MobileDecalRendering.cpp
+++ b/Engine/Source/Runtime/Renderer/Private/MobileDecalRendering.cpp
@@ -108,7 +108,7 @@ void RenderDeferredDecalsMobile(FRHICommandList& RHICmdList, const FScene& Scene
 		for (int32 DecalIndex = 0; DecalIndex < SortedDecalCount; DecalIndex++)
 		{
 			const FTransientDecalRenderData& DecalData = SortedDecals[DecalIndex];
-			const FDeferredDecalProxy& DecalProxy = DecalData.Proxy;
+			const FDeferredDecalProxy& DecalProxy = *DecalData.Proxy;
 			const FMatrix ComponentToWorldMatrix = DecalProxy.ComponentTrans.ToMatrixWithScale();
 			const FMatrix FrustumComponentToClip = DecalRendering::ComputeComponentToClipMatrix(View, ComponentToWorldMatrix);
 
diff --git a/Engine/Source/Runtime/Renderer/Private/RayTracing/RayTracingDecals.cpp b/Engine/Source/Runtime/Renderer/Private/RayTracing/RayTracingDecals.cpp
--- a/Engine/Source/Runtime/Renderer/Private/RayTracing/RayTracingDecals.cpp
+++ b/Engine/Source/Runtime/Renderer/Private/RayTracing/RayTracingDecals.cpp
@@ -72,14 +72,16 @@ static TUniformBufferRef<FDecalParametersRayTracing> CreateDecalParametersBuffer
 	const FTransientDecalRenderData& DecalData,
 	EUniformBufferUsage Usage)
 {
+	const FDeferredDecalProxy& DecalProxy = *DecalData.Proxy;
+
 	FDecalParametersRayTracing Parameters;
 
-	const FMatrix DecalToWorldMatrix = DecalData.Proxy.ComponentTrans.ToMatrixWithScale();
-	const FMatrix WorldToDecalMatrix = DecalData.Proxy.ComponentTrans.ToInverseMatrixWithScale();
+	const FMatrix DecalToWorldMatrix = DecalProxy.ComponentTrans.ToMatrixWithScale();
+	const FMatrix WorldToDecalMatrix = DecalProxy.ComponentTrans.ToInverseMatrixWithScale();
 	const FDFVector3 AbsoluteOrigin(DecalToWorldMatrix.GetOrigin());
 	const FVector3f PositionHigh = AbsoluteOrigin.High;
 	const FMatrix44f RelativeDecalToWorldMatrix = FDFMatrix::MakeToRelativeWorldMatrix(PositionHigh, DecalToWorldMatrix).M;
-	const FVector3f OrientationVector = (FVector3f)DecalData.Proxy.ComponentTrans.GetUnitAxis(EAxis::X);
+	const FVector3f OrientationVector = (FVector3f)DecalProxy.ComponentTrans.GetUnitAxis(EAxis::X);
 
 	Parameters.DecalPositionHigh = PositionHigh;
 	Parameters.WorldToDecal = FMatrix44f(FTranslationMatrix(-View.ViewMatrices.GetPreViewTranslation()) * WorldToDecalMatrix);
@@ -93,7 +95,7 @@ static TUniformBufferRef<FDecalParametersRayTracing> CreateDecalParametersBuffer
 	// Certain engine captures (e.g. environment reflection) don't have a tick. Default to fully opaque.
 	if (View.Family->Time.GetWorldTimeSeconds())
 	{
-		LifetimeAlpha = FMath::Clamp(FMath::Min(View.Family->Time.GetWorldTimeSeconds() * -DecalData.Proxy.InvFadeDuration + DecalData.Proxy.FadeStartDelayNormalized, View.Family->Time.GetWorldTimeSeconds() * DecalData.Proxy.InvFadeInDuration + DecalData.Proxy.FadeInStartDelayNormalized), 0.0f, 1.0f);
+		LifetimeAlpha = FMath::Clamp(FMath::Min(View.Family->Time.GetWorldTimeSeconds() * -DecalProxy.InvFadeDuration + DecalProxy.FadeStartDelayNormalized, View.Family->Time.GetWorldTimeSeconds() * DecalProxy.InvFadeInDuration + DecalProxy.FadeInStartDelayNormalized), 0.0f, 1.0f);
 	}
 
 	Parameters.DecalParams = FVector2f(DecalData.FadeAlpha, LifetimeAlpha);
@@ -495,19 +497,20 @@ FTransientDecalRenderDataList GetSortedDecals(TConstArrayView<FDeferredDecalProx
 
 		if (bIsShown)
 		{
-			FTransientDecalRenderData Data(*DecalProxy, 0.0f, Scene.GetShaderPlatform(), Scene.GetFeatureLevel());
+			float FadeAlpha = 1.0f;
 
-			if (bIsPerspectiveProjection && Data.Proxy.FadeScreenSize != 0.0f)
+			if (bIsPerspectiveProjection && DecalProxy->FadeScreenSize != 0.0f)
 			{
-				const FMatrix ComponentToWorldMatrix = Data.Proxy.ComponentTrans.ToMatrixWithScale();
-				Data.FadeAlpha = DecalRendering::CalculateDecalFadeAlpha(Data.Proxy.FadeScreenSize, ComponentToWorldMatrix, View, FadeMultiplier);
+				const FMatrix ComponentToWorldMatrix = DecalProxy->ComponentTrans.ToMatrixWithScale();
+				FadeAlpha = DecalRendering::CalculateDecalFadeAlpha(DecalProxy->FadeScreenSize, ComponentToWorldMatrix, View, FadeMultiplier);
 			}
 
-			const bool bShouldRender = Data.FadeAlpha > 0.0f;
+			const bool bShouldRender = FadeAlpha > 0.0f;
 
 			if (bShouldRender)
 			{
-				SortedDecals.Add(Data);
+				FTransientDecalRenderData Data(*DecalProxy, 0.0f, FadeAlpha, Scene.GetShaderPlatform(), Scene.GetFeatureLevel());
+				SortedDecals.Add(MoveTemp(Data));
 			}
 		}
 	}
@@ -569,7 +572,7 @@ TRDGUniformBufferRef<FRayTracingDecals> CreateRayTracingDecalData(FRDGBuilder& G
 		FMeshDrawSingleShaderBindings SingleShaderBindings = Command.ShaderBindings.GetSingleShaderBindings(SF_RayCallable);
 		CallableShader->GetShaderBindings(&Scene, Scene.GetFeatureLevel(), *MaterialProxy, *MaterialResource, View, DecalParameters, SingleShaderBindings);
 
-		const FBox BoxBounds = DecalData.Proxy.GetBounds().GetBox();
+		const FBox BoxBounds = DecalData.Proxy->GetBounds().GetBox();
 
 		FRayTracingDecal& DestDecal = RayTracingDecals.AddDefaulted_GetRef();
 		DestDecal.TranslatedBoundMin = FVector3f(BoxBounds.Min + View.ViewMatrices.GetPreViewTranslation());
