Holdout - fixed height fog to match expectations.

diff --git a/Engine/Shaders/Private/HeightFogPixelShader.usf b/Engine/Shaders/Private/HeightFogPixelShader.usf
--- a/Engine/Shaders/Private/HeightFogPixelShader.usf
+++ b/Engine/Shaders/Private/HeightFogPixelShader.usf
@@ -120,14 +120,20 @@ void ExponentialPixelMain(
 	// Coverage of the cloud layer itself does not change.
 	OutColor.a   = CloudLuminanceTransmittance.a;
 
-#else
+#else // PERMUTATION_SAMPLE_FOG_ON_CLOUDS
 
 	if (bOnlyOnRenderedOpaque > 0.0 && !bIsRendered)
 	{
 		FogInscatteringAndOpacity.rgb = 0;
 		FogInscatteringAndOpacity.a = 1;
 	}
-
+	
+#if POST_PROCESS_ALPHA
+	// We need to output coverage so that we can potentially apply alpha hold out blending. See RenderViewFog.
+	OutColor = RETURN_COLOR(float4(FogInscatteringAndOpacity.rgb, 1.0 - FogInscatteringAndOpacity.w));
+#else
 	OutColor = RETURN_COLOR(float4(FogInscatteringAndOpacity.rgb, FogInscatteringAndOpacity.w));
 #endif
+	
+#endif // PERMUTATION_SAMPLE_FOG_ON_CLOUDS
 }
diff --git a/Engine/Source/Runtime/Renderer/Private/FogRendering.cpp b/Engine/Source/Runtime/Renderer/Private/FogRendering.cpp
--- a/Engine/Source/Runtime/Renderer/Private/FogRendering.cpp
+++ b/Engine/Source/Runtime/Renderer/Private/FogRendering.cpp
@@ -319,11 +319,22 @@ static void RenderViewFog(
 		const bool bSupportsAlpha = IsPostProcessingWithAlphaChannelSupported();
 		if (bSupportsAlpha)
 		{
-			GraphicsPSOInit.BlendState = TStaticBlendState<CW_RGBA, BO_Add, BF_One, BF_SourceAlpha, BO_Add, BF_Zero, BF_SourceAlpha>::GetRHI();
+			// Coverage is the alpha output of the shader in this case.
+			if (View.bExponentialFogHoldout && View.CachedViewUniformShaderParameters->RenderingReflectionCaptureMask == 0.0f)
+			{
+				// Alpha holdout: apply only when requested and when not rendering reflections. We want to punch a hole according to the Coverage. (black as throughput=0 should become brighter for see throught)
+				// SceneAlpha = Coverage*1 + (1.0-Coverage)*(SceneThroughput)
+				GraphicsPSOInit.BlendState = TStaticBlendState<CW_RGBA, BO_Add, BF_One, BF_InverseSourceAlpha, BO_Add, BF_One, BF_InverseSourceAlpha>::GetRHI();
+			}
+			else
+			{
+				// Need to store throughput into alpha so we multiply with transmittance=1-Coverage.
+				GraphicsPSOInit.BlendState = TStaticBlendState<CW_RGBA, BO_Add, BF_One, BF_InverseSourceAlpha, BO_Add, BF_Zero, BF_InverseSourceAlpha>::GetRHI();
+			}
 		}
 		else
 		{
-			// disable alpha writes in order to preserve scene depth values on PC
+			// Disable alpha writes in order to preserve scene depth values on PC
 			GraphicsPSOInit.BlendState = TStaticBlendState<CW_RGB, BO_Add, BF_One, BF_SourceAlpha>::GetRHI();
 		}
 	}
