Fix crash accessing VisibleDecals in FCompositionLighting.
- DBuffer passes were not disabled when using light complexity visualization
- Added check to detect incorrect state earlier (DBuffer should only be enabled when Decals are enabled)

diff --git a/Engine/Source/Runtime/Renderer/Private/CompositionLighting/CompositionLighting.cpp b/Engine/Source/Runtime/Renderer/Private/CompositionLighting/CompositionLighting.cpp
--- a/Engine/Source/Runtime/Renderer/Private/CompositionLighting/CompositionLighting.cpp
+++ b/Engine/Source/Runtime/Renderer/Private/CompositionLighting/CompositionLighting.cpp
@@ -498,8 +498,10 @@ FCompositionLighting::FCompositionLighting(TArrayView<const FViewInfo> InViews,
 	, ViewFamily(*InViews[0].Family)
 	, SceneTextures(InSceneTextures)
 	, bEnableDBuffer(IsDBufferEnabled(ViewFamily, SceneTextures.Config.ShaderPlatform))
-	, bEnableDecals(ViewFamily.EngineShowFlags.Decals && !ViewFamily.EngineShowFlags.VisualizeLightCulling)
+	, bEnableDecals(AreDecalsEnabled(ViewFamily))
 {
+	checkf(bEnableDecals || !bEnableDBuffer, TEXT("DBuffer should only be enabled when Decals are enabled."));
+
 	const FScene& Scene = *(FScene*)ViewFamily.Scene;
 
 	ViewAOConfigs.SetNum(Views.Num());
diff --git a/Engine/Source/Runtime/Renderer/Private/CompositionLighting/PostProcessDeferredDecals.cpp b/Engine/Source/Runtime/Renderer/Private/CompositionLighting/PostProcessDeferredDecals.cpp
--- a/Engine/Source/Runtime/Renderer/Private/CompositionLighting/PostProcessDeferredDecals.cpp
+++ b/Engine/Source/Runtime/Renderer/Private/CompositionLighting/PostProcessDeferredDecals.cpp
@@ -54,9 +54,16 @@ static TAutoConsoleVariable<bool> CVarDBufferDecalNormalReprojectionEnabled(
 	TEXT("in depth prepass is enabled as well (r.VelocityOutputPass=0). Otherwise the fallback is the normal extracted from the depth buffer."),
 	ECVF_RenderThreadSafe);
 
+bool AreDecalsEnabled(const FSceneViewFamily& ViewFamily)
+{
+	return ViewFamily.EngineShowFlags.Decals && !ViewFamily.EngineShowFlags.VisualizeLightCulling;
+}
+
 bool IsDBufferEnabled(const FSceneViewFamily& ViewFamily, EShaderPlatform ShaderPlatform)
 {
-	return !ViewFamily.EngineShowFlags.ShaderComplexity && ViewFamily.EngineShowFlags.Decals && IsUsingDBuffers(ShaderPlatform);
+	return IsUsingDBuffers(ShaderPlatform)
+		&& AreDecalsEnabled(ViewFamily)
+		&& !ViewFamily.EngineShowFlags.ShaderComplexity;
 }
 
 IMPLEMENT_STATIC_UNIFORM_BUFFER_STRUCT(FDecalPassUniformParameters, "DecalPass", SceneTextures);
diff --git a/Engine/Source/Runtime/Renderer/Private/CompositionLighting/PostProcessDeferredDecals.h b/Engine/Source/Runtime/Renderer/Private/CompositionLighting/PostProcessDeferredDecals.h
--- a/Engine/Source/Runtime/Renderer/Private/CompositionLighting/PostProcessDeferredDecals.h
+++ b/Engine/Source/Runtime/Renderer/Private/CompositionLighting/PostProcessDeferredDecals.h
@@ -16,6 +16,7 @@ class FViewInfo;
 
 DECLARE_UNIFORM_BUFFER_STRUCT(FSceneUniformParameters, RENDERER_API)
 
+bool AreDecalsEnabled(const FSceneViewFamily& ViewFamily);
 bool IsDBufferEnabled(const FSceneViewFamily& ViewFamily, EShaderPlatform ShaderPlatform);
 
 BEGIN_GLOBAL_SHADER_PARAMETER_STRUCT(FDecalPassUniformParameters, )
