HeterogeneousVolumeComponent: integration with SVT frame-varying transform

diff --git a/Engine/Source/Runtime/Engine/Classes/Components/HeterogeneousVolumeComponent.h b/Engine/Source/Runtime/Engine/Classes/Components/HeterogeneousVolumeComponent.h
--- a/Engine/Source/Runtime/Engine/Classes/Components/HeterogeneousVolumeComponent.h
+++ b/Engine/Source/Runtime/Engine/Classes/Components/HeterogeneousVolumeComponent.h
@@ -24,6 +24,10 @@ class UHeterogeneousVolumeComponent : public UMeshComponent
 	UPROPERTY(EditAnywhere, BlueprintReadOnly, Category = Volume)
 	FIntVector VolumeResolution;
 
+	//UPROPERTY(EditAnywhere, BlueprintReadWrite, Category = Animation, AdvancedDisplay, meta = (EditCondition = "false"))
+	UPROPERTY(VisibleAnywhere, BlueprintReadWrite, Category = Animation, AdvancedDisplay)
+	FTransform FrameTransform;
+
 	UPROPERTY(EditAnywhere, BlueprintReadWrite, Category = Animation, meta = (EditCondition = "bPlaying == false"))
 	float Frame;
 
@@ -48,6 +52,9 @@ class UHeterogeneousVolumeComponent : public UMeshComponent
 	UPROPERTY(EditAnywhere, Category = SparseVolumeTextureStreaming)
 	uint32 bIssueBlockingRequests : 1;
 
+	UPROPERTY(EditAnywhere, Category = Volume)
+	uint32 bPivotAtCentroid : 1;
+
 	UPROPERTY(EditAnywhere, Category = Lighting)
 	float StepFactor;
 
diff --git a/Engine/Source/Runtime/Engine/Private/Components/HeterogeneousVolumeComponent.cpp b/Engine/Source/Runtime/Engine/Private/Components/HeterogeneousVolumeComponent.cpp
--- a/Engine/Source/Runtime/Engine/Private/Components/HeterogeneousVolumeComponent.cpp
+++ b/Engine/Source/Runtime/Engine/Private/Components/HeterogeneousVolumeComponent.cpp
@@ -75,6 +75,7 @@ FHeterogeneousVolumeSceneProxy::FHeterogeneousVolumeSceneProxy(UHeterogeneousVol
 	bCastDynamicShadow = InComponent->CastShadow;
 
 	HeterogeneousVolumeData.VoxelResolution = InComponent->VolumeResolution;
+	HeterogeneousVolumeData.InstanceToLocal = InComponent->FrameTransform.ToMatrixWithScale();
 
 	// Infer minimum voxel size from bounds and resolution
 	FVector VoxelSize = 2.0 * InComponent->Bounds.BoxExtent;
@@ -92,6 +93,7 @@ FHeterogeneousVolumeSceneProxy::FHeterogeneousVolumeSceneProxy(UHeterogeneousVol
 	HeterogeneousVolumeData.ShadowStepFactor = InComponent->ShadowStepFactor;
 	HeterogeneousVolumeData.ShadowBiasFactor = InComponent->ShadowBiasFactor;
 	HeterogeneousVolumeData.LightingDownsampleFactor = InComponent->LightingDownsampleFactor;
+	HeterogeneousVolumeData.bPivotAtCentroid = InComponent->bPivotAtCentroid;
 
 	// Initialize vertex buffer data for a quad
 	StaticMeshVertexBuffers.PositionVertexBuffer.Init(4);
@@ -249,6 +251,7 @@ UHeterogeneousVolumeComponent::UHeterogeneousVolumeComponent(const FObjectInitia
 	LightingDownsampleFactor = 2.0f;
 	MipLevel = 0;
 	bIssueBlockingRequests = false;
+	bPivotAtCentroid = false;
 	PreviousSVT = nullptr;
 }
 
@@ -330,11 +333,19 @@ FPrimitiveSceneProxy* UHeterogeneousVolumeComponent::CreateSceneProxy()
 FBoxSphereBounds UHeterogeneousVolumeComponent::CalcBounds(const FTransform& LocalToWorld) const
 {
 	FBoxSphereBounds NewBounds;
-	NewBounds.Origin = FVector::ZeroVector;
-	NewBounds.BoxExtent = FVector(VolumeResolution) * 0.5;
-	NewBounds.SphereRadius = NewBounds.BoxExtent.Length();
 
-	return NewBounds.TransformBy(LocalToWorld);
+	FVector HalfVolumeResolution = FVector(VolumeResolution) * 0.5;
+	if (bPivotAtCentroid)
+	{
+		NewBounds.Origin = FVector::ZeroVector;
+	}
+	else
+	{
+		NewBounds.Origin = HalfVolumeResolution;
+	}
+	NewBounds.BoxExtent = HalfVolumeResolution;
+	NewBounds.SphereRadius = NewBounds.BoxExtent.Length();
+	return NewBounds.TransformBy(FrameTransform * LocalToWorld);
 }
 
 void UHeterogeneousVolumeComponent::PostLoad()
@@ -563,6 +574,13 @@ void UHeterogeneousVolumeComponent::TickComponent(float DeltaTime, ELevelTick Ti
 				{
 					MarkRenderTransformDirty();
 				}
+
+				FTransform PerFrameTransform = SparseVolumeTextureFrame->GetFrameTransform();
+				if (!PerFrameTransform.Equals(FrameTransform))
+				{
+					FrameTransform = PerFrameTransform;
+					MarkRenderStateDirty();
+				}
 			}
 
 			MaterialInstanceDynamic->SetSparseVolumeTextureParameterValue(SVTParameterName, SparseVolumeTextureFrame);
diff --git a/Engine/Source/Runtime/Engine/Public/HeterogeneousVolumeInterface.h b/Engine/Source/Runtime/Engine/Public/HeterogeneousVolumeInterface.h
--- a/Engine/Source/Runtime/Engine/Public/HeterogeneousVolumeInterface.h
+++ b/Engine/Source/Runtime/Engine/Public/HeterogeneousVolumeInterface.h
@@ -25,6 +25,7 @@ public:
 	// Volume
 	virtual FIntVector GetVoxelResolution() const = 0;
 	virtual float GetMinimumVoxelSize() const = 0;
+	virtual bool IsPivotAtCentroid() const = 0;
 
 	// Lighting
 	virtual float GetStepFactor() const = 0;
@@ -61,6 +62,7 @@ public:
 		, ShadowStepFactor(8.0)
 		, ShadowBiasFactor(0.0)
 		, LightingDownsampleFactor(1.0)
+		, bPivotAtCentroid(false)
 #if ACTOR_HAS_LABELS
 		, ReadableName(Name)
 #endif // ACTOR_HAS_LABELS
@@ -79,6 +81,7 @@ public:
 	// Volume
 	virtual FIntVector GetVoxelResolution() const { return VoxelResolution; }
 	virtual float GetMinimumVoxelSize() const { return MinimumVoxelSize; }
+	virtual bool IsPivotAtCentroid() const { return bPivotAtCentroid; }
 
 	// Lighting
 	virtual float GetStepFactor() const { return StepFactor; }
@@ -94,6 +97,7 @@ public:
 	float ShadowStepFactor;
 	float ShadowBiasFactor;
 	float LightingDownsampleFactor;
+	bool bPivotAtCentroid;
 
 #if ACTOR_HAS_LABELS
 	FString ReadableName;
diff --git a/Engine/Source/Runtime/Renderer/Private/HeterogeneousVolumes/HeterogeneousVolumesLiveShadingPipeline.cpp b/Engine/Source/Runtime/Renderer/Private/HeterogeneousVolumes/HeterogeneousVolumesLiveShadingPipeline.cpp
--- a/Engine/Source/Runtime/Renderer/Private/HeterogeneousVolumes/HeterogeneousVolumesLiveShadingPipeline.cpp
+++ b/Engine/Source/Runtime/Renderer/Private/HeterogeneousVolumes/HeterogeneousVolumesLiveShadingPipeline.cpp
@@ -457,11 +457,15 @@ static void RenderLightingCacheWithLiveShading(
 		// TODO: Convert to relative-local space
 		//FVector3f ViewOriginHigh = FDFVector3(View.ViewMatrices.GetViewOrigin()).High;
 		//FMatrix44f RelativeLocalToWorld = FDFMatrix::MakeToRelativeWorldMatrix(ViewOriginHigh, HeterogeneousVolumeInterface->GetLocalToWorld()).M;
-		FMatrix44f LocalToWorld = FMatrix44f(HeterogeneousVolumeInterface->GetLocalToWorld());
-		PassParameters->LocalToWorld = LocalToWorld;
-		PassParameters->WorldToLocal = LocalToWorld.Inverse();
-		PassParameters->LocalBoundsOrigin = FVector3f(LocalBoxSphereBounds.Origin);
-		PassParameters->LocalBoundsExtent = FVector3f(LocalBoxSphereBounds.BoxExtent);
+		FMatrix InstanceToLocal = HeterogeneousVolumeInterface->GetInstanceToLocal();
+		FMatrix LocalToWorld = HeterogeneousVolumeInterface->GetLocalToWorld();
+		PassParameters->LocalToWorld = FMatrix44f(InstanceToLocal * LocalToWorld);
+		PassParameters->WorldToLocal = PassParameters->LocalToWorld.Inverse();
+
+		FMatrix LocalToInstance = InstanceToLocal.Inverse();
+		FBoxSphereBounds InstanceBoxSphereBounds = LocalBoxSphereBounds.TransformBy(LocalToInstance);
+		PassParameters->LocalBoundsOrigin = FVector3f(InstanceBoxSphereBounds.Origin);
+		PassParameters->LocalBoundsExtent = FVector3f(InstanceBoxSphereBounds.BoxExtent);
 		PassParameters->PrimitiveId = PersistentPrimitiveIndex.Index;
 
 		// Transmittance volume
@@ -668,10 +672,15 @@ static void RenderSingleScatteringWithLiveShading(
 		// TODO: Convert to relative-local space
 		//FVector3f ViewOriginHigh = FDFVector3(View.ViewMatrices.GetViewOrigin()).High;
 		//FMatrix44f RelativeLocalToWorld = FDFMatrix::MakeToRelativeWorldMatrix(ViewOriginHigh, HeterogeneousVolumeInterface->GetLocalToWorld()).M;
-		PassParameters->LocalToWorld = FMatrix44f(HeterogeneousVolumeInterface->GetLocalToWorld());
-		PassParameters->WorldToLocal = FMatrix44f(HeterogeneousVolumeInterface->GetLocalToWorld().Inverse());
-		PassParameters->LocalBoundsOrigin = FVector3f(LocalBoxSphereBounds.Origin);
-		PassParameters->LocalBoundsExtent = FVector3f(LocalBoxSphereBounds.BoxExtent);
+		FMatrix InstanceToLocal = HeterogeneousVolumeInterface->GetInstanceToLocal();
+		FMatrix LocalToWorld = HeterogeneousVolumeInterface->GetLocalToWorld();
+		PassParameters->LocalToWorld = FMatrix44f(InstanceToLocal * LocalToWorld);
+		PassParameters->WorldToLocal = PassParameters->LocalToWorld.Inverse();
+
+		FMatrix LocalToInstance = InstanceToLocal.Inverse();
+		FBoxSphereBounds InstanceBoxSphereBounds = LocalBoxSphereBounds.TransformBy(LocalToInstance);
+		PassParameters->LocalBoundsOrigin = FVector3f(InstanceBoxSphereBounds.Origin);
+		PassParameters->LocalBoundsExtent = FVector3f(InstanceBoxSphereBounds.BoxExtent);
 		PassParameters->PrimitiveId = PersistentPrimitiveIndex.Index;
 
 		// Volume data
diff --git a/Engine/Source/Runtime/Renderer/Private/HeterogeneousVolumes/HeterogeneousVolumesMaterialBakingPipeline.cpp b/Engine/Source/Runtime/Renderer/Private/HeterogeneousVolumes/HeterogeneousVolumesMaterialBakingPipeline.cpp
--- a/Engine/Source/Runtime/Renderer/Private/HeterogeneousVolumes/HeterogeneousVolumesMaterialBakingPipeline.cpp
+++ b/Engine/Source/Runtime/Renderer/Private/HeterogeneousVolumes/HeterogeneousVolumesMaterialBakingPipeline.cpp
@@ -139,13 +139,17 @@ void ComputeHeterogeneousVolumeBakeMaterial(
 
 		// Object data
 		// TODO: Convert to relative-local space
-		FVector3f ViewOriginHigh = FDFVector3(View.ViewMatrices.GetViewOrigin()).High;
-		FMatrix44f RelativeLocalToWorld = FDFMatrix::MakeToRelativeWorldMatrix(ViewOriginHigh, HeterogeneousVolumeInterface->GetLocalToWorld()).M;
-		FMatrix44f LocalToWorld = FMatrix44f(HeterogeneousVolumeInterface->GetLocalToWorld());
-		PassParameters->LocalToWorld = LocalToWorld;
-		PassParameters->WorldToLocal = LocalToWorld.Inverse();
-		PassParameters->LocalBoundsOrigin = FVector3f(LocalBoxSphereBounds.Origin);
-		PassParameters->LocalBoundsExtent = FVector3f(LocalBoxSphereBounds.BoxExtent);
+		//FVector3f ViewOriginHigh = FDFVector3(View.ViewMatrices.GetViewOrigin()).High;
+		//FMatrix44f RelativeLocalToWorld = FDFMatrix::MakeToRelativeWorldMatrix(ViewOriginHigh, HeterogeneousVolumeInterface->GetLocalToWorld()).M;
+		FMatrix InstanceToLocal = HeterogeneousVolumeInterface->GetInstanceToLocal();
+		FMatrix LocalToWorld = HeterogeneousVolumeInterface->GetLocalToWorld();
+		PassParameters->LocalToWorld = FMatrix44f(InstanceToLocal * LocalToWorld);
+		PassParameters->WorldToLocal = PassParameters->LocalToWorld.Inverse();
+
+		FMatrix LocalToInstance = InstanceToLocal.Inverse();
+		FBoxSphereBounds InstanceBoxSphereBounds = LocalBoxSphereBounds.TransformBy(LocalToInstance);
+		PassParameters->LocalBoundsOrigin = FVector3f(InstanceBoxSphereBounds.Origin);
+		PassParameters->LocalBoundsExtent = FVector3f(InstanceBoxSphereBounds.BoxExtent);
 		PassParameters->PrimitiveId = PersistentPrimitiveIndex.Index;
 
 		// Volume data
diff --git a/Engine/Source/Runtime/Renderer/Private/HeterogeneousVolumes/HeterogeneousVolumesPreshadingPipeline.cpp b/Engine/Source/Runtime/Renderer/Private/HeterogeneousVolumes/HeterogeneousVolumesPreshadingPipeline.cpp
--- a/Engine/Source/Runtime/Renderer/Private/HeterogeneousVolumes/HeterogeneousVolumesPreshadingPipeline.cpp
+++ b/Engine/Source/Runtime/Renderer/Private/HeterogeneousVolumes/HeterogeneousVolumesPreshadingPipeline.cpp
@@ -1603,11 +1603,15 @@ void RenderWithPreshading(
 		// TODO: Convert to relative-local space
 		//FVector3f ViewOriginHigh = FDFVector3(View.ViewMatrices.GetViewOrigin()).High;
 		//FMatrix44f RelativeLocalToWorld = FDFMatrix::MakeToRelativeWorldMatrix(ViewOriginHigh, HeterogeneousVolumeInterface->GetLocalToWorld()).M;
-		FMatrix44f LocalToWorld = FMatrix44f(HeterogeneousVolumeInterface->GetLocalToWorld());
-		SparseVoxelUniformBufferParameters->LocalToWorld = LocalToWorld;
-		SparseVoxelUniformBufferParameters->WorldToLocal = LocalToWorld.Inverse();
-		SparseVoxelUniformBufferParameters->LocalBoundsOrigin = FVector3f(LocalBoxSphereBounds.Origin);
-		SparseVoxelUniformBufferParameters->LocalBoundsExtent = FVector3f(LocalBoxSphereBounds.BoxExtent);
+		FMatrix InstanceToLocal = HeterogeneousVolumeInterface->GetInstanceToLocal();
+		FMatrix LocalToWorld = HeterogeneousVolumeInterface->GetLocalToWorld();
+		SparseVoxelUniformBufferParameters->LocalToWorld = FMatrix44f(InstanceToLocal * LocalToWorld);
+		SparseVoxelUniformBufferParameters->WorldToLocal = SparseVoxelUniformBufferParameters->LocalToWorld.Inverse();
+
+		FMatrix LocalToInstance = InstanceToLocal.Inverse();
+		FBoxSphereBounds InstanceBoxSphereBounds = LocalBoxSphereBounds.TransformBy(LocalToInstance);
+		SparseVoxelUniformBufferParameters->LocalBoundsOrigin = FVector3f(InstanceBoxSphereBounds.Origin);
+		SparseVoxelUniformBufferParameters->LocalBoundsExtent = FVector3f(InstanceBoxSphereBounds.BoxExtent);
 
 		// Volume data
 		SparseVoxelUniformBufferParameters->VolumeResolution = VolumeResolution;
diff --git a/Engine/Source/Runtime/Renderer/Private/HeterogeneousVolumes/HeterogeneousVolumesVoxelGridPipeline.cpp b/Engine/Source/Runtime/Renderer/Private/HeterogeneousVolumes/HeterogeneousVolumesVoxelGridPipeline.cpp
--- a/Engine/Source/Runtime/Renderer/Private/HeterogeneousVolumes/HeterogeneousVolumesVoxelGridPipeline.cpp
+++ b/Engine/Source/Runtime/Renderer/Private/HeterogeneousVolumes/HeterogeneousVolumesVoxelGridPipeline.cpp
@@ -2010,79 +2010,97 @@ void RasterizeVolumesIntoFrustumVoxelGrid(
 			continue;
 		}
 
-		const FPrimitiveSceneInfo* PrimitiveSceneInfo = PrimitiveSceneProxy->GetPrimitiveSceneInfo();
-		const int32 PrimitiveId = PrimitiveSceneInfo->GetIndex();
-		const FBoxSphereBounds LocalBoxSphereBounds = PrimitiveSceneProxy->GetLocalBounds();
-		const FBoxSphereBounds PrimitiveBounds = PrimitiveSceneProxy->GetBounds();
-		const FMaterial& Material = MaterialRenderProxy->GetMaterialWithFallback(View.GetFeatureLevel(), MaterialRenderProxy);
-
-		FRasterizeBottomLevelFrustumGridCS::FParameters* PassParameters = GraphBuilder.AllocParameters<FRasterizeBottomLevelFrustumGridCS::FParameters>();
+		for (int32 VolumeIndex = 0; VolumeIndex < Mesh->Elements.Num(); ++VolumeIndex)
 		{
-			// Scene data
-			PassParameters->View = View.ViewUniformBuffer;
-			PassParameters->Scene = View.GetSceneUniforms().GetBuffer(GraphBuilder);
-
-			// Primitive data
-			PassParameters->PrimitiveWorldBoundsMin = FVector3f(PrimitiveBounds.Origin - PrimitiveBounds.BoxExtent);
-			PassParameters->PrimitiveWorldBoundsMax = FVector3f(PrimitiveBounds.Origin + PrimitiveBounds.BoxExtent);
-			FMatrix44f LocalToWorld = FMatrix44f(PrimitiveSceneProxy->GetLocalToWorld());
-			PassParameters->LocalToWorld = LocalToWorld;
-			PassParameters->WorldToLocal = LocalToWorld.Inverse();
-			PassParameters->LocalBoundsOrigin = FVector3f(LocalBoxSphereBounds.Origin);
-			PassParameters->LocalBoundsExtent = FVector3f(LocalBoxSphereBounds.BoxExtent);
-			PassParameters->PrimitiveId = PrimitiveId;
+			const IHeterogeneousVolumeInterface* HeterogeneousVolumeInterface = (IHeterogeneousVolumeInterface*)Mesh->Elements[VolumeIndex].UserData;
+			//check(HeterogeneousVolumeInterface != nullptr);
+			if (HeterogeneousVolumeInterface == nullptr)
+			{
+				continue;
+			}
 
-			// Volume data
-			PassParameters->TopLevelGridResolution = TopLevelGridResolution;
-			PassParameters->VoxelDimensions = TopLevelGridResolution;
-			PassParameters->ViewToWorld = FMatrix44f(ViewToWorld);
-			PassParameters->TanHalfFOV = HeterogeneousVolumes::CalcTanHalfFOV(View.FOV);
-			PassParameters->NearPlaneDepth = NearPlaneDistance;
-			PassParameters->FarPlaneDepth = FarPlaneDistance;
-
-			// Sampling data
-			PassParameters->bJitter = BuildOptions.bJitter;
-			FBlueNoise BlueNoise = GetBlueNoiseGlobalParameters();
-			PassParameters->BlueNoise = CreateUniformBufferImmediate(BlueNoise, EUniformBufferUsage::UniformBuffer_SingleDraw);
-
-			// Raster tile data
-			PassParameters->RasterTileAllocatorBuffer = GraphBuilder.CreateSRV(RasterTileAllocatorBuffer, PF_R32_UINT);
-			PassParameters->RasterTileBuffer = GraphBuilder.CreateSRV(RasterTileBuffer);
-
-			// Indirect args
-			PassParameters->IndirectArgs = RasterizeBottomLevelGridIndirectArgsBuffer;
-
-			// Grid data
-			PassParameters->RWBottomLevelGridAllocatorBuffer = GraphBuilder.CreateUAV(BottomLevelGridAllocatorBuffer, PF_R32_UINT);
-			PassParameters->RWTopLevelGridBuffer = GraphBuilder.CreateUAV(TopLevelGridBuffer);
-
-			PassParameters->RWExtinctionGridBuffer = GraphBuilder.CreateUAV(ExtinctionGridBuffer);
-			PassParameters->RWEmissionGridBuffer = GraphBuilder.CreateUAV(EmissionGridBuffer);
-			PassParameters->RWScatteringGridBuffer = GraphBuilder.CreateUAV(ScatteringGridBuffer);
-			PassParameters->BottomLevelGridBufferSize = BottomLevelGridBufferSize;
-		}
+			const FPrimitiveSceneInfo* PrimitiveSceneInfo = PrimitiveSceneProxy->GetPrimitiveSceneInfo();
+			const int32 PrimitiveId = PrimitiveSceneInfo->GetIndex();
+			const FBoxSphereBounds LocalBoxSphereBounds = PrimitiveSceneProxy->GetLocalBounds();
+			const FBoxSphereBounds PrimitiveBounds = PrimitiveSceneProxy->GetBounds();
+			const FMaterial& Material = MaterialRenderProxy->GetMaterialWithFallback(View.GetFeatureLevel(), MaterialRenderProxy);
 
-		GraphBuilder.AddPass(
-			RDG_EVENT_NAME("FrustumGrid.RasterizeBottomLevelGrid"),
-			PassParameters,
-			ERDGPassFlags::Compute,
-			// Why is scene explicitly copied??
-			[PassParameters, LocalScene = Scene, &View, MaterialRenderProxy, &Material](FRHIComputeCommandList& RHICmdList)
+			FRasterizeBottomLevelFrustumGridCS::FParameters* PassParameters = GraphBuilder.AllocParameters<FRasterizeBottomLevelFrustumGridCS::FParameters>();
 			{
-				FRasterizeBottomLevelFrustumGridCS::FPermutationDomain PermutationVector;
-				TShaderRef<FRasterizeBottomLevelFrustumGridCS> ComputeShader = Material.GetShader<FRasterizeBottomLevelFrustumGridCS>(&FLocalVertexFactory::StaticType, PermutationVector, false);
+				// Scene data
+				PassParameters->View = View.ViewUniformBuffer;
+				PassParameters->Scene = View.GetSceneUniforms().GetBuffer(GraphBuilder);
+
+				// Primitive data
+				PassParameters->PrimitiveWorldBoundsMin = FVector3f(PrimitiveBounds.Origin - PrimitiveBounds.BoxExtent);
+				PassParameters->PrimitiveWorldBoundsMax = FVector3f(PrimitiveBounds.Origin + PrimitiveBounds.BoxExtent);
+				
+				// TODO: Convert to relative-local space
+				//FVector3f ViewOriginHigh = FDFVector3(View.ViewMatrices.GetViewOrigin()).High;
+				//FMatrix44f RelativeLocalToWorld = FDFMatrix::MakeToRelativeWorldMatrix(ViewOriginHigh, HeterogeneousVolumeInterface->GetLocalToWorld()).M;
+				FMatrix InstanceToLocal = HeterogeneousVolumeInterface->GetInstanceToLocal();
+				FMatrix LocalToWorld = HeterogeneousVolumeInterface->GetLocalToWorld();
+				PassParameters->LocalToWorld = FMatrix44f(InstanceToLocal * LocalToWorld);
+				PassParameters->WorldToLocal = FMatrix44f(PassParameters->LocalToWorld.Inverse());
+
+				FMatrix LocalToInstance = InstanceToLocal.Inverse();
+				FBoxSphereBounds InstanceBoxSphereBounds = LocalBoxSphereBounds.TransformBy(FTransform(LocalToInstance));
+				PassParameters->LocalBoundsOrigin = FVector3f(InstanceBoxSphereBounds.Origin);
+				PassParameters->LocalBoundsExtent = FVector3f(InstanceBoxSphereBounds.BoxExtent);
+				PassParameters->PrimitiveId = PrimitiveId;
+
+				// Volume data
+				PassParameters->TopLevelGridResolution = TopLevelGridResolution;
+				PassParameters->VoxelDimensions = TopLevelGridResolution;
+				PassParameters->ViewToWorld = FMatrix44f(ViewToWorld);
+				PassParameters->TanHalfFOV = HeterogeneousVolumes::CalcTanHalfFOV(View.FOV);
+				PassParameters->NearPlaneDepth = NearPlaneDistance;
+				PassParameters->FarPlaneDepth = FarPlaneDistance;
+
+				// Sampling data
+				PassParameters->bJitter = BuildOptions.bJitter;
+				FBlueNoise BlueNoise = GetBlueNoiseGlobalParameters();
+				PassParameters->BlueNoise = CreateUniformBufferImmediate(BlueNoise, EUniformBufferUsage::UniformBuffer_SingleDraw);
+
+				// Raster tile data
+				PassParameters->RasterTileAllocatorBuffer = GraphBuilder.CreateSRV(RasterTileAllocatorBuffer, PF_R32_UINT);
+				PassParameters->RasterTileBuffer = GraphBuilder.CreateSRV(RasterTileBuffer);
+
+				// Indirect args
+				PassParameters->IndirectArgs = RasterizeBottomLevelGridIndirectArgsBuffer;
+
+				// Grid data
+				PassParameters->RWBottomLevelGridAllocatorBuffer = GraphBuilder.CreateUAV(BottomLevelGridAllocatorBuffer, PF_R32_UINT);
+				PassParameters->RWTopLevelGridBuffer = GraphBuilder.CreateUAV(TopLevelGridBuffer);
 
-				if (!ComputeShader.IsNull())
+				PassParameters->RWExtinctionGridBuffer = GraphBuilder.CreateUAV(ExtinctionGridBuffer);
+				PassParameters->RWEmissionGridBuffer = GraphBuilder.CreateUAV(EmissionGridBuffer);
+				PassParameters->RWScatteringGridBuffer = GraphBuilder.CreateUAV(ScatteringGridBuffer);
+				PassParameters->BottomLevelGridBufferSize = BottomLevelGridBufferSize;
+			}
+
+			GraphBuilder.AddPass(
+				RDG_EVENT_NAME("FrustumGrid.RasterizeBottomLevelGrid"),
+				PassParameters,
+				ERDGPassFlags::Compute,
+				// Why is scene explicitly copied??
+				[PassParameters, LocalScene = Scene, &View, MaterialRenderProxy, &Material](FRHIComputeCommandList& RHICmdList)
 				{
-					ClearUnusedGraphResources(ComputeShader, PassParameters);
+					FRasterizeBottomLevelFrustumGridCS::FPermutationDomain PermutationVector;
+					TShaderRef<FRasterizeBottomLevelFrustumGridCS> ComputeShader = Material.GetShader<FRasterizeBottomLevelFrustumGridCS>(&FLocalVertexFactory::StaticType, PermutationVector, false);
 
-					FMeshDrawShaderBindings ShaderBindings;
-					UE::MeshPassUtils::SetupComputeBindings(ComputeShader, LocalScene, LocalScene->GetFeatureLevel(), nullptr, *MaterialRenderProxy, Material, ShaderBindings);
+					if (!ComputeShader.IsNull())
+					{
+						ClearUnusedGraphResources(ComputeShader, PassParameters);
 
-					UE::MeshPassUtils::DispatchIndirect(RHICmdList, ComputeShader, ShaderBindings, *PassParameters, PassParameters->IndirectArgs->GetIndirectRHICallBuffer(), 0);
+						FMeshDrawShaderBindings ShaderBindings;
+						UE::MeshPassUtils::SetupComputeBindings(ComputeShader, LocalScene, LocalScene->GetFeatureLevel(), nullptr, *MaterialRenderProxy, Material, ShaderBindings);
+
+						UE::MeshPassUtils::DispatchIndirect(RHICmdList, ComputeShader, ShaderBindings, *PassParameters, PassParameters->IndirectArgs->GetIndirectRHICallBuffer(), 0);
+					}
 				}
-			}
-		);
+			);
+		}
 	}
 }
 
@@ -2472,102 +2490,119 @@ void RasterizeVolumesIntoOrthoVoxelGrid(
 				continue;
 			}
 
-			const FPrimitiveSceneInfo* PrimitiveSceneInfo = PrimitiveSceneProxy->GetPrimitiveSceneInfo();
-			const int32 PrimitiveId = PrimitiveSceneInfo->GetIndex();
-			const FBoxSphereBounds LocalBoxSphereBounds = PrimitiveSceneProxy->GetLocalBounds();
-			const FBoxSphereBounds PrimitiveBounds = PrimitiveSceneProxy->GetBounds();
-			const FMaterial& Material = MaterialRenderProxy->GetMaterialWithFallback(View.GetFeatureLevel(), MaterialRenderProxy);
-
-			FRasterizeBottomLevelOrthoGridCS::FParameters* PassParameters = GraphBuilder.AllocParameters<FRasterizeBottomLevelOrthoGridCS::FParameters>();
+			for (int32 VolumeIndex = 0; VolumeIndex < Mesh->Elements.Num(); ++VolumeIndex)
 			{
-				// Scene data
-				PassParameters->View = View.ViewUniformBuffer;
-				PassParameters->Scene = View.GetSceneUniforms().GetBuffer(GraphBuilder);
+				const IHeterogeneousVolumeInterface* HeterogeneousVolumeInterface = (IHeterogeneousVolumeInterface*)Mesh->Elements[VolumeIndex].UserData;
+				//check(HeterogeneousVolumeInterface != nullptr);
+				if (HeterogeneousVolumeInterface == nullptr)
+				{
+					continue;
+				}
 
-				// Primitive data
-				PassParameters->PrimitiveWorldBoundsMin = FVector3f(PrimitiveBounds.Origin - PrimitiveBounds.BoxExtent);
-				PassParameters->PrimitiveWorldBoundsMax = FVector3f(PrimitiveBounds.Origin + PrimitiveBounds.BoxExtent);
+				const FPrimitiveSceneInfo* PrimitiveSceneInfo = PrimitiveSceneProxy->GetPrimitiveSceneInfo();
+				const int32 PrimitiveId = PrimitiveSceneInfo->GetIndex();
+				const FBoxSphereBounds LocalBoxSphereBounds = PrimitiveSceneProxy->GetLocalBounds();
+				const FBoxSphereBounds PrimitiveBounds = PrimitiveSceneProxy->GetBounds();
+				const FMaterial& Material = MaterialRenderProxy->GetMaterialWithFallback(View.GetFeatureLevel(), MaterialRenderProxy);
 
-				FMatrix44f LocalToWorld = FMatrix44f(PrimitiveSceneProxy->GetLocalToWorld());
-				PassParameters->LocalToWorld = LocalToWorld;
-				PassParameters->WorldToLocal = LocalToWorld.Inverse();
-				PassParameters->LocalBoundsOrigin = FVector3f(LocalBoxSphereBounds.Origin);
-				PassParameters->LocalBoundsExtent = FVector3f(LocalBoxSphereBounds.BoxExtent);
-				PassParameters->PrimitiveId = PrimitiveId;
+				FRasterizeBottomLevelOrthoGridCS::FParameters* PassParameters = GraphBuilder.AllocParameters<FRasterizeBottomLevelOrthoGridCS::FParameters>();
+				{
+					// Scene data
+					PassParameters->View = View.ViewUniformBuffer;
+					PassParameters->Scene = View.GetSceneUniforms().GetBuffer(GraphBuilder);
 
-				// Volume data
-				PassParameters->TopLevelGridResolution = TopLevelGridResolution;
-				PassParameters->TopLevelGridWorldBoundsMin = FVector3f(TopLevelGridBounds.Origin - TopLevelGridBounds.BoxExtent);
-				PassParameters->TopLevelGridWorldBoundsMax = FVector3f(TopLevelGridBounds.Origin + TopLevelGridBounds.BoxExtent);
+					// Primitive data
+					PassParameters->PrimitiveWorldBoundsMin = FVector3f(PrimitiveBounds.Origin - PrimitiveBounds.BoxExtent);
+					PassParameters->PrimitiveWorldBoundsMax = FVector3f(PrimitiveBounds.Origin + PrimitiveBounds.BoxExtent);
 
-				// Sampling data
-				FBlueNoise BlueNoise = GetBlueNoiseGlobalParameters();
-				PassParameters->BlueNoise = CreateUniformBufferImmediate(BlueNoise, EUniformBufferUsage::UniformBuffer_SingleDraw);
+					// TODO: Convert to relative-local space
+					//FVector3f ViewOriginHigh = FDFVector3(View.ViewMatrices.GetViewOrigin()).High;
+					//FMatrix44f RelativeLocalToWorld = FDFMatrix::MakeToRelativeWorldMatrix(ViewOriginHigh, HeterogeneousVolumeInterface->GetLocalToWorld()).M;
+					FMatrix InstanceToLocal = HeterogeneousVolumeInterface->GetInstanceToLocal();
+					FMatrix LocalToWorld = HeterogeneousVolumeInterface->GetLocalToWorld();
+					PassParameters->LocalToWorld = FMatrix44f(InstanceToLocal * LocalToWorld);
+					PassParameters->WorldToLocal = FMatrix44f(PassParameters->LocalToWorld.Inverse());
 
-				// Unify with "object" definition??
-				PassParameters->PrimitiveWorldBoundsMin = FVector3f(PrimitiveBounds.Origin - PrimitiveBounds.BoxExtent);
-				PassParameters->PrimitiveWorldBoundsMax = FVector3f(PrimitiveBounds.Origin + PrimitiveBounds.BoxExtent);
+					FMatrix LocalToInstance = InstanceToLocal.Inverse();
+					FBoxSphereBounds InstanceBoxSphereBounds = LocalBoxSphereBounds.TransformBy(LocalToInstance);
+					PassParameters->LocalBoundsOrigin = FVector3f(InstanceBoxSphereBounds.Origin);
+					PassParameters->LocalBoundsExtent = FVector3f(InstanceBoxSphereBounds.BoxExtent);
+					PassParameters->PrimitiveId = PrimitiveId;
 
-				PassParameters->BottomLevelGridBufferSize = BottomLevelGridBufferSize;
+					// Volume data
+					PassParameters->TopLevelGridResolution = TopLevelGridResolution;
+					PassParameters->TopLevelGridWorldBoundsMin = FVector3f(TopLevelGridBounds.Origin - TopLevelGridBounds.BoxExtent);
+					PassParameters->TopLevelGridWorldBoundsMax = FVector3f(TopLevelGridBounds.Origin + TopLevelGridBounds.BoxExtent);
 
-				// Raster tile data
-				PassParameters->RasterTileAllocatorBuffer = GraphBuilder.CreateSRV(RasterTileAllocatorBuffer, PF_R32_UINT);
-				PassParameters->RasterTileBuffer = GraphBuilder.CreateSRV(RasterTileBuffer);
+					// Sampling data
+					FBlueNoise BlueNoise = GetBlueNoiseGlobalParameters();
+					PassParameters->BlueNoise = CreateUniformBufferImmediate(BlueNoise, EUniformBufferUsage::UniformBuffer_SingleDraw);
 
-				// Indirect args
-				PassParameters->IndirectArgs = RasterizeBottomLevelGridIndirectArgsBuffer;
+					// Unify with "object" definition??
+					PassParameters->PrimitiveWorldBoundsMin = FVector3f(PrimitiveBounds.Origin - PrimitiveBounds.BoxExtent);
+					PassParameters->PrimitiveWorldBoundsMax = FVector3f(PrimitiveBounds.Origin + PrimitiveBounds.BoxExtent);
 
-				// Sampling mode
-				PassParameters->bJitter = BuildOptions.bJitter;
-				PassParameters->bSampleAtVertices = HeterogeneousVolumes::EnableLinearInterpolation();
+					PassParameters->BottomLevelGridBufferSize = BottomLevelGridBufferSize;
 
-				// Grid data
-				PassParameters->TopLevelGridBuffer = GraphBuilder.CreateSRV(TopLevelGridBuffer);
-				PassParameters->RWBottomLevelGridAllocatorBuffer = GraphBuilder.CreateUAV(BottomLevelGridAllocatorBuffer, PF_R32_UINT);
-				PassParameters->RWTopLevelGridBuffer = GraphBuilder.CreateUAV(TopLevelGridBuffer);
-				PassParameters->RWExtinctionGridBuffer = GraphBuilder.CreateUAV(ExtinctionGridBuffer);
-				PassParameters->RWEmissionGridBuffer = GraphBuilder.CreateUAV(EmissionGridBuffer);
-				PassParameters->RWScatteringGridBuffer = GraphBuilder.CreateUAV(ScatteringGridBuffer);
+					// Raster tile data
+					PassParameters->RasterTileAllocatorBuffer = GraphBuilder.CreateSRV(RasterTileAllocatorBuffer, PF_R32_UINT);
+					PassParameters->RasterTileBuffer = GraphBuilder.CreateSRV(RasterTileBuffer);
 
-				// Indirection Grid
-				PassParameters->FixedBottomLevelResolution = HeterogeneousVolumes::GetBottomLevelGridResolution();
-				PassParameters->RWIndirectionGridBuffer = GraphBuilder.CreateUAV(IndirectionGridBuffer);
+					// Indirect args
+					PassParameters->IndirectArgs = RasterizeBottomLevelGridIndirectArgsBuffer;
 
-#if 0
-				// Hash table
-				if (HeterogeneousVolumes::EnableVoxelHashing())
-				{
-					PassParameters->RWHashTable = GraphBuilder.CreateUAV(HashTableBuffer);
-					PassParameters->RWHashToVoxelBuffer = GraphBuilder.CreateUAV(HashToVoxelBuffer);
-					PassParameters->HashTableSize = HashTableBufferSize;
-				}
-#endif
-				PassParameters->HomogeneousThreshold = HeterogeneousVolumes::GetHomogeneousAggregationThreshold();
-			}
+					// Sampling mode
+					PassParameters->bJitter = BuildOptions.bJitter;
+					PassParameters->bSampleAtVertices = HeterogeneousVolumes::EnableLinearInterpolation();
 
-			GraphBuilder.AddPass(
-				RDG_EVENT_NAME("RasterizeBottomLevelGrid"),
-				PassParameters,
-				ERDGPassFlags::Compute,
-				// Why is scene explicitly copied?
-				[PassParameters, LocalScene = Scene, &View, MaterialRenderProxy, &Material](FRHIComputeCommandList& RHICmdList)
-				{
-					FRasterizeBottomLevelOrthoGridCS::FPermutationDomain PermutationVector;
-					PermutationVector.Set<FRasterizeBottomLevelOrthoGridCS::FEnableIndirectionGrid>(HeterogeneousVolumes::EnableIndirectionGrid());
-					PermutationVector.Set<FRasterizeBottomLevelOrthoGridCS::FEnableHomogeneousAggregation>(HeterogeneousVolumes::EnableHomogeneousAggregation());
-					TShaderRef<FRasterizeBottomLevelOrthoGridCS> ComputeShader = Material.GetShader<FRasterizeBottomLevelOrthoGridCS>(&FLocalVertexFactory::StaticType, PermutationVector, false);
+					// Grid data
+					PassParameters->TopLevelGridBuffer = GraphBuilder.CreateSRV(TopLevelGridBuffer);
+					PassParameters->RWBottomLevelGridAllocatorBuffer = GraphBuilder.CreateUAV(BottomLevelGridAllocatorBuffer, PF_R32_UINT);
+					PassParameters->RWTopLevelGridBuffer = GraphBuilder.CreateUAV(TopLevelGridBuffer);
+					PassParameters->RWExtinctionGridBuffer = GraphBuilder.CreateUAV(ExtinctionGridBuffer);
+					PassParameters->RWEmissionGridBuffer = GraphBuilder.CreateUAV(EmissionGridBuffer);
+					PassParameters->RWScatteringGridBuffer = GraphBuilder.CreateUAV(ScatteringGridBuffer);
 
-					if (!ComputeShader.IsNull())
+					// Indirection Grid
+					PassParameters->FixedBottomLevelResolution = HeterogeneousVolumes::GetBottomLevelGridResolution();
+					PassParameters->RWIndirectionGridBuffer = GraphBuilder.CreateUAV(IndirectionGridBuffer);
+
+	#if 0
+					// Hash table
+					if (HeterogeneousVolumes::EnableVoxelHashing())
 					{
-						ClearUnusedGraphResources(ComputeShader, PassParameters);
+						PassParameters->RWHashTable = GraphBuilder.CreateUAV(HashTableBuffer);
+						PassParameters->RWHashToVoxelBuffer = GraphBuilder.CreateUAV(HashToVoxelBuffer);
+						PassParameters->HashTableSize = HashTableBufferSize;
+					}
+	#endif
+					PassParameters->HomogeneousThreshold = HeterogeneousVolumes::GetHomogeneousAggregationThreshold();
+				}
 
-						FMeshDrawShaderBindings ShaderBindings;
-						UE::MeshPassUtils::SetupComputeBindings(ComputeShader, LocalScene, LocalScene->GetFeatureLevel(), nullptr, *MaterialRenderProxy, Material, ShaderBindings);
+				GraphBuilder.AddPass(
+					RDG_EVENT_NAME("RasterizeBottomLevelGrid"),
+					PassParameters,
+					ERDGPassFlags::Compute,
+					// Why is scene explicitly copied?
+					[PassParameters, LocalScene = Scene, &View, MaterialRenderProxy, &Material](FRHIComputeCommandList& RHICmdList)
+					{
+						FRasterizeBottomLevelOrthoGridCS::FPermutationDomain PermutationVector;
+						PermutationVector.Set<FRasterizeBottomLevelOrthoGridCS::FEnableIndirectionGrid>(HeterogeneousVolumes::EnableIndirectionGrid());
+						PermutationVector.Set<FRasterizeBottomLevelOrthoGridCS::FEnableHomogeneousAggregation>(HeterogeneousVolumes::EnableHomogeneousAggregation());
+						TShaderRef<FRasterizeBottomLevelOrthoGridCS> ComputeShader = Material.GetShader<FRasterizeBottomLevelOrthoGridCS>(&FLocalVertexFactory::StaticType, PermutationVector, false);
 
-						UE::MeshPassUtils::DispatchIndirect(RHICmdList, ComputeShader, ShaderBindings, *PassParameters, PassParameters->IndirectArgs->GetIndirectRHICallBuffer(), 0);
+						if (!ComputeShader.IsNull())
+						{
+							ClearUnusedGraphResources(ComputeShader, PassParameters);
+
+							FMeshDrawShaderBindings ShaderBindings;
+							UE::MeshPassUtils::SetupComputeBindings(ComputeShader, LocalScene, LocalScene->GetFeatureLevel(), nullptr, *MaterialRenderProxy, Material, ShaderBindings);
+
+							UE::MeshPassUtils::DispatchIndirect(RHICmdList, ComputeShader, ShaderBindings, *PassParameters, PassParameters->IndirectArgs->GetIndirectRHICallBuffer(), 0);
+						}
 					}
-				}
-			);
+				);
+			}
 		}
 	}
 }
