Add VSMResolutionLodBias to light scene proxy. Controls the resolution of each light's shadow map allowing certain key lights to be higher-res while other unimportant lights lower-res to fine-tune performance and memory usage

diff --git a/Engine/Source/Runtime/Engine/Private/LightSceneProxy.cpp b/Engine/Source/Runtime/Engine/Private/LightSceneProxy.cpp
--- a/Engine/Source/Runtime/Engine/Private/LightSceneProxy.cpp
+++ b/Engine/Source/Runtime/Engine/Private/LightSceneProxy.cpp
@@ -140,6 +140,7 @@ FLightSceneProxy::FLightSceneProxy(const ULightComponent* InLightComponent)
 		bTransmission = false;
 	}
 	VSMTexelDitherScale = 1.0f;
+	VSMResolutionLodBias = 0.0f;
 }
 
 FLightSceneProxy::~FLightSceneProxy() = default;
diff --git a/Engine/Source/Runtime/Engine/Public/LightSceneProxy.h b/Engine/Source/Runtime/Engine/Public/LightSceneProxy.h
--- a/Engine/Source/Runtime/Engine/Public/LightSceneProxy.h
+++ b/Engine/Source/Runtime/Engine/Public/LightSceneProxy.h
@@ -115,6 +115,11 @@ public:
 		return VSMTexelDitherScale;
 	}
 
+	ENGINE_API float GetVSMResolutionLodBias() const
+	{
+		return VSMResolutionLodBias;
+	}
+	
 	/** Returns the number of view dependent shadows this light will create, not counting distance field shadow cascades. */
 	virtual uint32 GetNumViewDependentWholeSceneShadows(const FSceneView& View, bool bPrecomputedLightingIsValid) const { return 0; }
 
@@ -464,6 +469,8 @@ protected:
 
 	/** Used to control the amount of additional dither filtering applied to shadows for each light. */
 	float VSMTexelDitherScale;
+	/** Used to control shadow resolution for each light. */
+	float VSMResolutionLodBias;
 
 	/** Only for whole scene directional lights, if FarShadowCascadeCount > 0 and FarShadowDistance >= WholeSceneDynamicShadowRadius, where far shadow cascade should end. */
 	float FarShadowDistance;
diff --git a/Engine/Source/Runtime/Renderer/Private/VirtualShadowMaps/VirtualShadowMapClipmap.cpp b/Engine/Source/Runtime/Renderer/Private/VirtualShadowMaps/VirtualShadowMapClipmap.cpp
--- a/Engine/Source/Runtime/Renderer/Private/VirtualShadowMaps/VirtualShadowMapClipmap.cpp
+++ b/Engine/Source/Runtime/Renderer/Private/VirtualShadowMaps/VirtualShadowMapClipmap.cpp
@@ -138,6 +138,7 @@ FVirtualShadowMapClipmap::FVirtualShadowMapClipmap(
 	// just resizing the virtual shadow maps for each clipmap, but convenient for now. This means we need to additionally bias
 	// which levels are present.
 	ResolutionLodBias = FVirtualShadowMapArray::InterpolateResolutionBias(CVarVirtualShadowMapResolutionLodBiasDirectional.GetValueOnRenderThread(), CVarVirtualShadowMapResolutionLodBiasDirectionalMoving.GetValueOnRenderThread(), LightMobilityFactor) + FMath::Log2(LodScale);
+	ResolutionLodBias += GetLightSceneInfo().Proxy->GetVSMResolutionLodBias();
 	// Clamp negative absolute resolution biases as they would exceed the maximum resolution/ranges allocated
 	ResolutionLodBias = FMath::Max(0.0f, ResolutionLodBias);
 
