Added environment component hold out flags.
This is useful for holdout and later for IsVisibleInMainView controls for SkyAtm, VolCloud and ExpFog components.

diff --git a/Engine/Shaders/Private/HeightFogCommon.ush b/Engine/Shaders/Private/HeightFogCommon.ush
--- a/Engine/Shaders/Private/HeightFogCommon.ush
+++ b/Engine/Shaders/Private/HeightFogCommon.ush
@@ -6,6 +6,8 @@
 
 #pragma once
 
+#include "/Engine/Shared/EnvironmentComponentsFlags.h"
+
 #ifndef PERMUTATION_SUPPORT_FOG_INSCATTERING_TEXTURE
 #define PERMUTATION_SUPPORT_FOG_INSCATTERING_TEXTURE 1
 #endif
@@ -210,7 +212,7 @@ half4 GetExponentialHeightFog(float3 WorldPositionRelativeToCamera, float Exclud
 	#endif
 
 #if POST_PROCESS_ALPHA
-	if (FogStruct.bHoldout && View.RenderingReflectionCaptureMask == 0.0f)
+	if (IsExponentialFogHoldout(View.EnvironmentComponentsFlags) && View.RenderingReflectionCaptureMask == 0.0f)
 	{
 		FogColor.rgb *= 0;
 	}
@@ -277,7 +279,7 @@ float4 CombineVolumetricFog(float4 GlobalFog, float3 VolumeUV, uint EyeIndex, fl
 	//VolumetricFogLookup.rgb += .1f * frac(min(ZSlice, 1.0f) / View.VolumetricFogInvGridSize.z);
 
 #if POST_PROCESS_ALPHA
-	if (FogStruct.bHoldout && View.RenderingReflectionCaptureMask == 0.0f)
+	if (IsExponentialFogHoldout(View.EnvironmentComponentsFlags) && View.RenderingReflectionCaptureMask == 0.0f)
 	{
 		VolumetricFogLookup.rgb *= 0;
 		GlobalFog.rgb *= 0;
diff --git a/Engine/Shaders/Private/SkyAtmosphere.usf b/Engine/Shaders/Private/SkyAtmosphere.usf
--- a/Engine/Shaders/Private/SkyAtmosphere.usf
+++ b/Engine/Shaders/Private/SkyAtmosphere.usf
@@ -8,6 +8,7 @@
 #include "Common.ush"
 #include "SceneTexturesCommon.ush"
 #include "SkyAtmosphereCommon.ush"
+#include "/Engine/Shared/EnvironmentComponentsFlags.h"
 
 #ifndef COLORED_TRANSMITTANCE_ENABLED	// Never used, UE4 does not supports dual blending
 #define COLORED_TRANSMITTANCE_ENABLED 0 
@@ -825,7 +826,7 @@ float4 PrepareOutput(float3 PreExposedLuminance, float3 Transmittance = float3(1
 	float3 SafePreExposedLuminance = min(PreExposedLuminance, Max10BitsFloat.xxx * 0.5f);
 
 #if POST_PROCESS_ALPHA
-	if (Atmosphere.bHoldout && View.RenderingReflectionCaptureMask == 0.0f)
+	if (IsSkyAtmosphereHoldout(View.EnvironmentComponentsFlags) && View.RenderingReflectionCaptureMask == 0.0f)
 	{
 		SafePreExposedLuminance *= 0;
 	}
@@ -949,7 +950,7 @@ void RenderSkyAtmosphereRayMarchingPS(
 
 		float3 SkyGreyTransmittance = 1.0f;
 	#if POST_PROCESS_ALPHA
-		if (Atmosphere.bHoldout && View.RenderingReflectionCaptureMask == 0.0f)
+		if (IsSkyAtmosphereHoldout(View.EnvironmentComponentsFlags) && View.RenderingReflectionCaptureMask == 0.0f)
 		{
 			SkyGreyTransmittance = SkyLuminanceTransmittance.aaa;
 		}
@@ -1581,7 +1582,7 @@ void RenderCameraAerialPerspectiveVolumeCS(uint3 ThreadId : SV_DispatchThreadID)
 		tMaxMax);
 
 #if POST_PROCESS_ALPHA
-	if (Atmosphere.bHoldout && !RealTimeReflection360Mode)
+	if (IsSkyAtmosphereHoldout(View.EnvironmentComponentsFlags) && !RealTimeReflection360Mode)
 	{
 		ss.L *= 0;
 		ss.LMieOnly *= 0;
diff --git a/Engine/Shaders/Private/VolumetricCloud.usf b/Engine/Shaders/Private/VolumetricCloud.usf
--- a/Engine/Shaders/Private/VolumetricCloud.usf
+++ b/Engine/Shaders/Private/VolumetricCloud.usf
@@ -14,6 +14,7 @@
 
 #include "Random.ush"
 #include "VolumetricCloudCommon.ush"
+#include "/Engine/Shared/EnvironmentComponentsFlags.h"
 
 #if SUBSTRATE_ENABLED && !MATERIAL_IS_SUBSTRATE
 #undef SUBSTRATE_ENABLED
@@ -1408,7 +1409,7 @@ void MainCommon(in FMaterialPixelParameters MaterialParameters, in float4 SvPosi
 	const float OutputPreExposure = (ResolvedView.RealTimeReflectionCapture ? ResolvedView.RealTimeReflectionCapturePreExposure : ResolvedView.PreExposure);
 
 #if POST_PROCESS_ALPHA
-	if (RenderVolumetricCloudParameters.bHoldout2 && View.RenderingReflectionCaptureMask == 0.0f)
+	if (IsVolumetricCloudHoldout(View.EnvironmentComponentsFlags) && View.RenderingReflectionCaptureMask == 0.0f)
 	{
 		Luminance *= 0;
 	#if CLOUD_MIN_AND_MAX_DEPTH
diff --git a/Engine/Shaders/Shared/EnvironmentComponentsFlags.h b/Engine/Shaders/Shared/EnvironmentComponentsFlags.h
new file mode 100644
--- /dev/null
+++ b/Engine/Shaders/Shared/EnvironmentComponentsFlags.h
@@ -0,0 +1,25 @@
+// Copyright Epic Games, Inc. All Rights Reserved.
+
+/*================================================================================================
+	!!! Changing this file requires recompilation of the engine !!!
+=================================================================================================*/
+
+#pragma once
+
+
+#ifndef __cplusplus
+// Change this to force recompilation of all shaders refeencing this file.
+#pragma message("UESHADERMETADATA_VERSION 8F92177A-03B9-4AA9-A983-3CA04895E84A")
+#endif
+
+
+#define ENVCOMP_FLAG_SKYATMOSPHERE_HOLDOUT 0x1
+#define ENVCOMP_FLAG_VOLUMETRICCLOUD_HOLDOUT 0x2
+#define ENVCOMP_FLAG_EXPONENTIALFOG_HOLDOUT 0x4
+
+
+// Using macro functions to avoid having to move the definition in cpp to avoid the multiple definition in different obj files.
+#define IsSkyAtmosphereHoldout(EnvironmentComponentsFlags)   ((EnvironmentComponentsFlags & ENVCOMP_FLAG_SKYATMOSPHERE_HOLDOUT)   > 0)
+#define IsVolumetricCloudHoldout(EnvironmentComponentsFlags) ((EnvironmentComponentsFlags & ENVCOMP_FLAG_VOLUMETRICCLOUD_HOLDOUT) > 0)
+#define IsExponentialFogHoldout(EnvironmentComponentsFlags)  ((EnvironmentComponentsFlags & ENVCOMP_FLAG_EXPONENTIALFOG_HOLDOUT)  > 0)
+
diff --git a/Engine/Source/Runtime/Engine/Private/Components/SkyAtmosphereComponent.cpp b/Engine/Source/Runtime/Engine/Private/Components/SkyAtmosphereComponent.cpp
--- a/Engine/Source/Runtime/Engine/Private/Components/SkyAtmosphereComponent.cpp
+++ b/Engine/Source/Runtime/Engine/Private/Components/SkyAtmosphereComponent.cpp
@@ -479,6 +479,7 @@ ASkyAtmosphere::ASkyAtmosphere(const FObjectInitializer& ObjectInitializer)
 FSkyAtmosphereSceneProxy::FSkyAtmosphereSceneProxy(const USkyAtmosphereComponent* InComponent)
 	: bStaticLightingBuilt(false)
 	, AtmosphereSetup(*InComponent)
+	, bHoldout(InComponent->bHoldout > 0)
 {
 	SkyLuminanceFactor = InComponent->SkyLuminanceFactor;
 	AerialPespectiveViewDistanceScale = InComponent->AerialPespectiveViewDistanceScale;
diff --git a/Engine/Source/Runtime/Engine/Public/Rendering/SkyAtmosphereCommonData.cpp b/Engine/Source/Runtime/Engine/Public/Rendering/SkyAtmosphereCommonData.cpp
--- a/Engine/Source/Runtime/Engine/Public/Rendering/SkyAtmosphereCommonData.cpp
+++ b/Engine/Source/Runtime/Engine/Public/Rendering/SkyAtmosphereCommonData.cpp
@@ -104,7 +104,6 @@ FAtmosphereSetup::FAtmosphereSetup(const USkyAtmosphereComponent& SkyAtmosphereC
 	}
 
 	TransmittanceMinLightElevationAngle = SkyAtmosphereComponent.TransmittanceMinLightElevationAngle;
-	bHoldout = SkyAtmosphereComponent.bHoldout ? 1 : 0;
 
 	UpdateTransform(SkyAtmosphereComponent.GetComponentTransform(), uint8(SkyAtmosphereComponent.TransformMode));
 }
diff --git a/Engine/Source/Runtime/Engine/Public/Rendering/SkyAtmosphereCommonData.h b/Engine/Source/Runtime/Engine/Public/Rendering/SkyAtmosphereCommonData.h
--- a/Engine/Source/Runtime/Engine/Public/Rendering/SkyAtmosphereCommonData.h
+++ b/Engine/Source/Runtime/Engine/Public/Rendering/SkyAtmosphereCommonData.h
@@ -43,8 +43,6 @@ struct FAtmosphereSetup
 
 	FLinearColor GroundAlbedo;
 
-	uint32 bHoldout;
-
 	float TransmittanceMinLightElevationAngle;
 
 	ENGINE_API FAtmosphereSetup(const USkyAtmosphereComponent& SkyAtmosphereComponent);
diff --git a/Engine/Source/Runtime/Engine/Public/SceneManagement.h b/Engine/Source/Runtime/Engine/Public/SceneManagement.h
--- a/Engine/Source/Runtime/Engine/Public/SceneManagement.h
+++ b/Engine/Source/Runtime/Engine/Public/SceneManagement.h
@@ -1198,6 +1198,8 @@ public:
 
 	const FAtmosphereSetup& GetAtmosphereSetup() const { return AtmosphereSetup; }
 
+	bool IsHoldout() const { return bHoldout; }
+
 	void UpdateTransform(const FTransform& ComponentTransform, uint8 TranformMode) { AtmosphereSetup.UpdateTransform(ComponentTransform, TranformMode); }
 	void ApplyWorldOffset(const FVector3f& InOffset) { AtmosphereSetup.ApplyWorldOffset((FVector)InOffset); }
 
@@ -1214,6 +1216,7 @@ private:
 	float HeightFogContribution;
 	float AerialPerspectiveStartDepthKm;
 	float TraceSampleCountScale;
+	bool bHoldout;
 
 	bool OverrideAtmosphericLight[NUM_ATMOSPHERE_LIGHTS];
 	FVector OverrideAtmosphericLightDirection[NUM_ATMOSPHERE_LIGHTS];
diff --git a/Engine/Source/Runtime/Engine/Public/SceneView.h b/Engine/Source/Runtime/Engine/Public/SceneView.h
--- a/Engine/Source/Runtime/Engine/Public/SceneView.h
+++ b/Engine/Source/Runtime/Engine/Public/SceneView.h
@@ -977,6 +977,7 @@ enum ETranslucencyVolumeCascade
 	VIEW_UNIFORM_BUFFER_MEMBER_PER_VIEW(FVector4f, ScreenRayLengthMultiplier) \
 	VIEW_UNIFORM_BUFFER_MEMBER_PER_VIEW(FVector4f, GlintLUTParameters0) \
 	VIEW_UNIFORM_BUFFER_MEMBER_PER_VIEW(FVector4f, GlintLUTParameters1) \
+	VIEW_UNIFORM_BUFFER_MEMBER(int, EnvironmentComponentsFlags) \
 
 /** The uniform shader parameters associated with a view. */
 #define VIEW_UNIFORM_BUFFER_MEMBER(type, identifier) \
diff --git a/Engine/Source/Runtime/Renderer/Private/FogRendering.cpp b/Engine/Source/Runtime/Renderer/Private/FogRendering.cpp
--- a/Engine/Source/Runtime/Renderer/Private/FogRendering.cpp
+++ b/Engine/Source/Runtime/Renderer/Private/FogRendering.cpp
@@ -12,6 +12,7 @@
 #include "ScreenPass.h"
 #include "TextureResource.h"
 #include "PostProcess/PostProcessing.h" // IsPostProcessingWithAlphaChannelSupported
+#include "EnvironmentComponentsFlags.h"
 
 DECLARE_GPU_DRAWCALL_STAT(Fog);
 
@@ -85,7 +86,6 @@ void SetupFogUniformParameters(FRDGBuilder& GraphBuilder, const FViewInfo& View,
 		OutParameters.IntegratedLightScatteringSampler = TStaticSamplerState<SF_Bilinear, AM_Clamp, AM_Clamp, AM_Clamp>::GetRHI();
 		OutParameters.VolumetricFogStartDistance = View.VolumetricFogStartDistance;
 		OutParameters.VolumetricFogNearFadeInDistanceInv = View.VolumetricFogNearFadeInDistanceInv;
-		OutParameters.bHoldout = View.bExponentialFogHoldout;
 	}
 }
 
@@ -248,7 +248,6 @@ void FSceneRenderer::InitFogConstants()
 				View.bUseDirectionalInscattering = SunLight != nullptr;
 				View.bEnableVolumetricFog = FogInfo.bEnableVolumetricFog;
 				View.VolumetricFogStartDistance = FogInfo.VolumetricFogStartDistance;
-				View.bExponentialFogHoldout = FogInfo.bHoldout;
 				View.VolumetricFogNearFadeInDistanceInv = FogInfo.VolumetricFogNearFadeInDistance > 0.0f ? (1.0f / FogInfo.VolumetricFogNearFadeInDistance) : 100000000.0f;
 			}
 		}
@@ -320,7 +319,7 @@ static void RenderViewFog(
 		if (bSupportsAlpha)
 		{
 			// Coverage is the alpha output of the shader in this case.
-			if (View.bExponentialFogHoldout && View.CachedViewUniformShaderParameters->RenderingReflectionCaptureMask == 0.0f)
+			if (IsExponentialFogHoldout(View.CachedViewUniformShaderParameters->EnvironmentComponentsFlags) && View.CachedViewUniformShaderParameters->RenderingReflectionCaptureMask == 0.0f)
 			{
 				// Alpha holdout: apply only when requested and when not rendering reflections. We want to punch a hole according to the Coverage. (black as throughput=0 should become brighter for see throught)
 				// SceneAlpha = Coverage*1 + (1.0-Coverage)*(SceneThroughput)
diff --git a/Engine/Source/Runtime/Renderer/Private/FogRendering.h b/Engine/Source/Runtime/Renderer/Private/FogRendering.h
--- a/Engine/Source/Runtime/Renderer/Private/FogRendering.h
+++ b/Engine/Source/Runtime/Renderer/Private/FogRendering.h
@@ -26,7 +26,6 @@ BEGIN_GLOBAL_SHADER_PARAMETER_STRUCT(FFogUniformParameters,)
 	SHADER_PARAMETER(float, ApplyVolumetricFog)
 	SHADER_PARAMETER(float, VolumetricFogStartDistance)
 	SHADER_PARAMETER(float, VolumetricFogNearFadeInDistanceInv)
-	SHADER_PARAMETER(int32, bHoldout)
 	SHADER_PARAMETER_TEXTURE(TextureCube, FogInscatteringColorCubemap)
 	SHADER_PARAMETER_SAMPLER(SamplerState, FogInscatteringColorSampler)
 	SHADER_PARAMETER_RDG_TEXTURE(Texture3D, IntegratedLightScattering)
diff --git a/Engine/Source/Runtime/Renderer/Private/PathTracing.cpp b/Engine/Source/Runtime/Renderer/Private/PathTracing.cpp
--- a/Engine/Source/Runtime/Renderer/Private/PathTracing.cpp
+++ b/Engine/Source/Runtime/Renderer/Private/PathTracing.cpp
@@ -34,6 +34,7 @@ TAutoConsoleVariable<int32> CVarPathTracing(
 #include "SkyAtmosphereRendering.h"
 #include <limits>
 #include "PathTracingSpatialTemporalDenoising.h"
+#include "EnvironmentComponentsFlags.h"
 
 TAutoConsoleVariable<int32> CVarPathTracingExperimental(
 	TEXT("r.PathTracing.Experimental"),
@@ -665,8 +666,8 @@ static void PreparePathTracingData(const FScene* Scene, const FViewInfo& View, F
 		&& (Scene->ExponentialFogs[0].FogData[0].Density > 0 ||
 			Scene->ExponentialFogs[0].FogData[1].Density > 0) ? PATH_TRACER_VOLUME_ENABLE_FOG : 0;
 	PathTracingData.VolumeFlags |= ShouldRenderHeterogeneousVolumesForView(View) ? PATH_TRACER_VOLUME_ENABLE_HETEROGENEOUS_VOLUMES : 0;
-	PathTracingData.VolumeFlags |= View.SkyAtmosphereUniformShaderParameters != nullptr &&  View.SkyAtmosphereUniformShaderParameters->bHoldout ? PATH_TRACER_VOLUME_HOLDOUT_ATMOSPHERE : 0;
-	PathTracingData.VolumeFlags |= Scene->ExponentialFogs.Num() > 0 && Scene->ExponentialFogs[0].bHoldout ? PATH_TRACER_VOLUME_HOLDOUT_FOG : 0;
+	PathTracingData.VolumeFlags |= View.SkyAtmosphereUniformShaderParameters != nullptr && IsSkyAtmosphereHoldout(View.CachedViewUniformShaderParameters->EnvironmentComponentsFlags) ? PATH_TRACER_VOLUME_HOLDOUT_ATMOSPHERE : 0;
+	PathTracingData.VolumeFlags |= Scene->ExponentialFogs.Num() > 0 && IsExponentialFogHoldout(View.CachedViewUniformShaderParameters->EnvironmentComponentsFlags) ? PATH_TRACER_VOLUME_HOLDOUT_FOG : 0;
 	PathTracingData.VolumeFlags |= EvalUseAnalyticTransmittance(View) ? PATH_TRACER_VOLUME_USE_ANALYTIC_TRANSMITTANCE : 0;
 
 	PathTracingData.EnableDBuffer = CVarPathTracingUseDBuffer.GetValueOnRenderThread();
@@ -2062,7 +2063,7 @@ void SetLightParameters(FRDGBuilder& GraphBuilder, FPathTracingRG::FParameters*
 		DestLight.VolumetricScatteringIntensity = Scene->SkyLight->VolumetricScatteringIntensity;
 		DestLight.IESAtlasIndex = INDEX_NONE;
 		DestLight.MissShaderIndex = 0;
-		if ((Scene->SkyLight->bRealTimeCaptureEnabled && (View.SkyAtmosphereUniformShaderParameters == nullptr ||  !View.SkyAtmosphereUniformShaderParameters->bHoldout)) || CVarPathTracingVisibleLights.GetValueOnRenderThread() == 2)
+		if ((Scene->SkyLight->bRealTimeCaptureEnabled && (View.SkyAtmosphereUniformShaderParameters == nullptr ||  !IsSkyAtmosphereHoldout(View.CachedViewUniformShaderParameters->EnvironmentComponentsFlags))) || CVarPathTracingVisibleLights.GetValueOnRenderThread() == 2)
 		{
 			// When using the realtime capture system, always make the skylight visible
 			// because this is our only way of "seeing" the atmo/clouds at the moment
diff --git a/Engine/Source/Runtime/Renderer/Private/SceneRendering.cpp b/Engine/Source/Runtime/Renderer/Private/SceneRendering.cpp
--- a/Engine/Source/Runtime/Renderer/Private/SceneRendering.cpp
+++ b/Engine/Source/Runtime/Renderer/Private/SceneRendering.cpp
@@ -99,6 +99,8 @@
 #include "OIT/OIT.h"
 #include "Rendering/CustomRenderPass.h"
 #include "LightFunctionAtlas.h"
+#include "EnvironmentComponentsFlags.h"
+#include "VolumetricCloudProxy.h"
 
 /*-----------------------------------------------------------------------------
 	Globals
@@ -1443,6 +1445,35 @@ void FViewInfo::SetupUniformBufferParameters(
 	ViewUniformShaderParameters.AtmosphereIrradianceTextureSampler = TStaticSamplerState<SF_Bilinear>::GetRHI();
 	ViewUniformShaderParameters.AtmosphereInscatterTextureSampler = TStaticSamplerState<SF_Bilinear>::GetRHI();
 
+	// Upload environment holdout flags
+	if (Scene)
+	{
+		ViewUniformShaderParameters.EnvironmentComponentsFlags = 0;
+
+		if (ShouldRenderSkyAtmosphere(Scene, Family->EngineShowFlags))
+		{
+			FSkyAtmosphereRenderSceneInfo* SkyAtmosphere = Scene->SkyAtmosphere;
+			const FSkyAtmosphereSceneProxy& SkyAtmosphereSceneProxy = SkyAtmosphere->GetSkyAtmosphereSceneProxy();
+
+			ViewUniformShaderParameters.EnvironmentComponentsFlags |= SkyAtmosphereSceneProxy.IsHoldout() ? ENVCOMP_FLAG_SKYATMOSPHERE_HOLDOUT : 0;
+		}
+
+		if (ShouldRenderVolumetricCloud(Scene, Family->EngineShowFlags))
+		{
+			FVolumetricCloudRenderSceneInfo* VolumetricCloud = Scene->VolumetricCloud;
+			const FVolumetricCloudSceneProxy& VolumetricCloudSceneProxy = VolumetricCloud->GetVolumetricCloudSceneProxy();
+
+			ViewUniformShaderParameters.EnvironmentComponentsFlags |= VolumetricCloudSceneProxy.bHoldout ? ENVCOMP_FLAG_VOLUMETRICCLOUD_HOLDOUT : 0;
+		}
+
+		if (Scene->ExponentialFogs.Num() > 0)
+		{
+			FExponentialHeightFogSceneInfo& Fog = Scene->ExponentialFogs[0];
+
+			ViewUniformShaderParameters.EnvironmentComponentsFlags |= Fog.bHoldout ? ENVCOMP_FLAG_EXPONENTIALFOG_HOLDOUT : 0;
+		}
+	}
+
 	// This should probably be in SetupCommonViewUniformBufferParameters, but drags in too many dependencies
 	UpdateNoiseTextureParameters(ViewUniformShaderParameters);
 
diff --git a/Engine/Source/Runtime/Renderer/Private/SceneRendering.h b/Engine/Source/Runtime/Renderer/Private/SceneRendering.h
--- a/Engine/Source/Runtime/Renderer/Private/SceneRendering.h
+++ b/Engine/Source/Runtime/Renderer/Private/SceneRendering.h
@@ -1395,7 +1395,6 @@ public:
 	FVector4f ExponentialFogParameters3;
 	FVector4f SkyAtmosphereAmbientContributionColorScale;
 	bool bEnableVolumetricFog;
-	bool bExponentialFogHoldout;
 	float VolumetricFogStartDistance;
 	float VolumetricFogNearFadeInDistanceInv;
 	FVector2f SinCosInscatteringColorCubemapRotation;
diff --git a/Engine/Source/Runtime/Renderer/Private/SkyAtmosphereRendering.cpp b/Engine/Source/Runtime/Renderer/Private/SkyAtmosphereRendering.cpp
--- a/Engine/Source/Runtime/Renderer/Private/SkyAtmosphereRendering.cpp
+++ b/Engine/Source/Runtime/Renderer/Private/SkyAtmosphereRendering.cpp
@@ -432,7 +432,6 @@ static void CopyAtmosphereSetupToUniformShaderParameters(FAtmosphereUniformShade
 	COPYMACRO(AbsorptionDensity1LinearTerm);
 	COPYMACRO(AbsorptionExtinction);
 	COPYMACRO(GroundAlbedo);
-	COPYMACRO(bHoldout);
 #undef COPYMACRO
 }
 
diff --git a/Engine/Source/Runtime/Renderer/Private/SkyAtmosphereRendering.h b/Engine/Source/Runtime/Renderer/Private/SkyAtmosphereRendering.h
--- a/Engine/Source/Runtime/Renderer/Private/SkyAtmosphereRendering.h
+++ b/Engine/Source/Runtime/Renderer/Private/SkyAtmosphereRendering.h
@@ -53,7 +53,6 @@ BEGIN_GLOBAL_SHADER_PARAMETER_STRUCT(FAtmosphereUniformShaderParameters, )
 	SHADER_PARAMETER(float, AbsorptionDensity1LinearTerm)
 	SHADER_PARAMETER(FLinearColor, AbsorptionExtinction)
 	SHADER_PARAMETER(FLinearColor, GroundAlbedo)
-	SHADER_PARAMETER(int32, bHoldout)
 END_GLOBAL_SHADER_PARAMETER_STRUCT()
 
 // These parameters are shared on the view global uniform buffer and are dynamically changed with cvars.
diff --git a/Engine/Source/Runtime/Renderer/Private/VolumetricCloudRendering.cpp b/Engine/Source/Runtime/Renderer/Private/VolumetricCloudRendering.cpp
--- a/Engine/Source/Runtime/Renderer/Private/VolumetricCloudRendering.cpp
+++ b/Engine/Source/Runtime/Renderer/Private/VolumetricCloudRendering.cpp
@@ -1611,7 +1611,6 @@ void FSceneRenderer::InitVolumetricCloudsForViews(FRDGBuilder& GraphBuilder, boo
 			CloudGlobalShaderParams.ShadowTracingMaxDistance			= KilometersToCentimeters * FMath::Max(0.1f, CloudProxy.ShadowTracingDistance);
 			CloudGlobalShaderParams.InvDistanceToSampleCountMax			= 1.0f / FMath::Max(1.0f, KilometersToCentimeters * CVarVolumetricCloudDistanceToSampleMaxCount.GetValueOnAnyThread());
 			CloudGlobalShaderParams.StopTracingTransmittanceThreshold	= FMath::Clamp(CloudProxy.StopTracingTransmittanceThreshold, 0.0f, 1.0f);
-			CloudGlobalShaderParams.bHoldout2							= CloudProxy.bHoldout ? 1 : 0;
 
 			auto PrepareCloudShadowMapLightData = [&](FLightSceneProxy* AtmosphericLight, int LightIndex)
 			{
diff --git a/Engine/Source/Runtime/Renderer/Private/VolumetricCloudRendering.h b/Engine/Source/Runtime/Renderer/Private/VolumetricCloudRendering.h
--- a/Engine/Source/Runtime/Renderer/Private/VolumetricCloudRendering.h
+++ b/Engine/Source/Runtime/Renderer/Private/VolumetricCloudRendering.h
@@ -38,7 +38,6 @@ BEGIN_SHADER_PARAMETER_STRUCT(FVolumetricCloudCommonShaderParameters, )
 	SHADER_PARAMETER(int32, SampleCountMax)
 	SHADER_PARAMETER(float, InvDistanceToSampleCountMax)
 	SHADER_PARAMETER(int32, ShadowSampleCountMax)
-	SHADER_PARAMETER(int32, bHoldout2)
 	SHADER_PARAMETER(float, ShadowTracingMaxDistance)
 	SHADER_PARAMETER(float, StopTracingTransmittanceThreshold)
 	SHADER_PARAMETER(float, SkyLightCloudBottomVisibility)
