Delete deprecated functions in RayTracingInstance.h.

diff --git a/Engine/Source/Runtime/Renderer/Private/RayTracing/RayTracingInstance.cpp b/Engine/Source/Runtime/Renderer/Private/RayTracing/RayTracingInstance.cpp
--- a/Engine/Source/Runtime/Renderer/Private/RayTracing/RayTracingInstance.cpp
+++ b/Engine/Source/Runtime/Renderer/Private/RayTracing/RayTracingInstance.cpp
@@ -8,39 +8,4 @@
 
 #if RHI_RAYTRACING
 
-#include "Materials/MaterialRenderProxy.h"
-#include "MaterialShared.h"
-#include "RayTracing/RayTracingInstanceMask.h"
-#include "MeshPassProcessor.h"
-
-
-void FRayTracingInstance::BuildInstanceMaskAndFlags(ERHIFeatureLevel::Type FeatureLevel)
-{
-	FSceneProxyRayTracingMaskInfo MaskInfo;
-	MaskInfo.MaskMode = ERayTracingViewMaskMode::RayTracing; // Deprecated function only supports RayTracing
-
-	uint8 ExtraMask = bThinGeometry ? ComputeRayTracingInstanceMask(ERayTracingInstanceMaskType::HairStrands, MaskInfo.MaskMode) : 0;
-
-	TArrayView<const FMeshBatch> MeshBatches = GetMaterials();
-	MaskAndFlags = BuildRayTracingInstanceMaskAndFlags(MeshBatches, FeatureLevel, MaskInfo, InstanceLayer, ExtraMask);
-}
-
-FRayTracingMaskAndFlags BuildRayTracingInstanceMaskAndFlags(TArrayView<const FMeshBatch> MeshBatches, ERHIFeatureLevel::Type FeatureLevel, ERayTracingInstanceLayer InstanceLayer, uint8 ExtraMask)
-{
-	FSceneProxyRayTracingMaskInfo MaskInfo;
-	MaskInfo.MaskMode = ERayTracingViewMaskMode::RayTracing; // Deprecated function only supports RayTracing
-
-	return BuildRayTracingInstanceMaskAndFlags(
-		MeshBatches,
-		FeatureLevel,
-		MaskInfo,
-		InstanceLayer,
-		ExtraMask);
-}	
-
-uint8 ComputeBlendModeMask(const EBlendMode BlendMode)
-{
-	return BlendModeToRayTracingInstanceMask(BlendMode, ERayTracingViewMaskMode::RayTracing);
-}
-
 #endif
diff --git a/Engine/Source/Runtime/Renderer/Public/RayTracingInstance.h b/Engine/Source/Runtime/Renderer/Public/RayTracingInstance.h
--- a/Engine/Source/Runtime/Renderer/Public/RayTracingInstance.h
+++ b/Engine/Source/Runtime/Renderer/Public/RayTracingInstance.h
@@ -126,17 +126,6 @@ struct FRayTracingInstance
 
 	/** When instance transforms are only available in GPU, this SRV holds them. */
 	FShaderResourceViewRHIRef InstanceGPUTransformsSRV;
-
-	/** Build mask and flags based on materials specified in Materials. You can still override Mask after calling this function.*/
-	UE_DEPRECATED(5.2, "Use BuildInstanceMaskAndFlags() with PrimitiveSceneProxy instead. Calling this function leads to incorrect path tracing result")
-	void BuildInstanceMaskAndFlags(ERHIFeatureLevel::Type FeatureLevel);
 };
 
-
-/** Build mask and flags based on materials specified in Materials. You can still override Mask after calling this function. */
-UE_DEPRECATED(5.2, "Use BuildRayTracingInstanceMaskAndFlags() with FSceneProxyRayTracingMaskInfo instead. Calling this function leads to incorrect path tracing result")
-FRayTracingMaskAndFlags BuildRayTracingInstanceMaskAndFlags(TArrayView<const FMeshBatch> MeshBatches, ERHIFeatureLevel::Type FeatureLevel, ERayTracingInstanceLayer InstanceLayer = ERayTracingInstanceLayer::NearField, uint8 ExtraMask = 0);
-
-UE_DEPRECATED(5.2, "Use BlendModeToRayTracingInstanceMask() instead. Calling this function leads to incorrect path tracing result.")
-uint8 ComputeBlendModeMask(const EBlendMode BlendMode);
 #endif
