Call UpdateSceneUniformBuffer() on scene extensions.
Fixes missing splines from baked RVT streaming mips.

diff --git a/Engine/Source/Runtime/Renderer/Private/SceneRendering.cpp b/Engine/Source/Runtime/Renderer/Private/SceneRendering.cpp
--- a/Engine/Source/Runtime/Renderer/Private/SceneRendering.cpp
+++ b/Engine/Source/Runtime/Renderer/Private/SceneRendering.cpp
@@ -5239,6 +5239,7 @@ public:
 		GPUScene.BeginRender(GraphBuilder, GPUSceneDynamicContext);
 
 		GPUScene.FillSceneUniformBuffer(GraphBuilder, Renderer.GetSceneUniforms());
+		Renderer.GetSceneExtensionsRenderers().UpdateSceneUniformBuffer(GraphBuilder, Renderer.GetSceneUniforms());
 	}
 
 	virtual ~FScenePrimitiveRenderingContext()
