Fix hair strands not being rendered when r.InstanceCulling.ForceInstanceCulling is enabled.

When r.InstanceCulling.ForceInstanceCulling, the MeshCommand indirect args buffer is always overriden by the instance culling indirect args. This is incorrect for hair as the indirect buffer is output by the hair culling system, not the instance culling. This results into hair not being drawn at all. This changes add an option to allow or not the indirect args buffer when using the SimpleMeshProcessor.

diff --git a/Engine/Source/Runtime/Renderer/Private/HairStrands/HairStrandsRasterCommon.cpp b/Engine/Source/Runtime/Renderer/Private/HairStrands/HairStrandsRasterCommon.cpp
--- a/Engine/Source/Runtime/Renderer/Private/HairStrands/HairStrandsRasterCommon.cpp
+++ b/Engine/Source/Runtime/Renderer/Private/HairStrands/HairStrandsRasterCommon.cpp
@@ -332,7 +332,7 @@ void AddHairStrandsRasterPass(
 		ViewInfo->CachedViewUniformShaderParameters->ViewForward = SavedViewForward;
 	}
 
-	AddSimpleMeshPass(GraphBuilder, PassParameters, Scene, *ViewInfo, &InstanceCullingManager, GetPassName(RasterPassType), ViewportRect,
+	AddSimpleMeshPass(GraphBuilder, PassParameters, Scene, *ViewInfo, &InstanceCullingManager, GetPassName(RasterPassType), ViewportRect, false /*bAllowOverrideIndirectArgs*/,
 		[PassParameters, Scene = Scene, ViewInfo, RasterPassType, &PrimitiveSceneInfos, HairRenderInfo, HairRenderInfoBits, RasterDirection](FDynamicPassMeshDrawListContext* ShadowContext)
 	{
 		SCOPE_CYCLE_COUNTER(STAT_RenderPerObjectShadowDepthsTime);
diff --git a/Engine/Source/Runtime/Renderer/Private/HairStrands/HairStrandsVisibility.cpp b/Engine/Source/Runtime/Renderer/Private/HairStrands/HairStrandsVisibility.cpp
--- a/Engine/Source/Runtime/Renderer/Private/HairStrands/HairStrandsVisibility.cpp
+++ b/Engine/Source/Runtime/Renderer/Private/HairStrands/HairStrandsVisibility.cpp
@@ -1066,6 +1066,7 @@ static FMaterialPassOutput AddHairMaterialPass(
 			&InstanceCullingManager, 
 			RDG_EVENT_NAME("HairStrands::MaterialPass(Emissive=%s)", Filter == EHairMaterialPassFilter::All ? TEXT("On/Off") : (Filter == EHairMaterialPassFilter::EmissiveOnly ? TEXT("On") : TEXT("Off"))),
 			FIntRect(0, 0, Resolution.X, Resolution.Y),
+			false /*bAllowOverrideIndirectArgs*/,
 		[PassParameters, Scene = Scene, ViewInfo, &MacroGroupDatas, MaxNodeCount, Resolution, NodeGroupSize, bUpdateSampleCoverage, Filter, FeatureLevel](FDynamicPassMeshDrawListContext* ShadowContext)
 		{			
 			FHairMaterialProcessor MeshProcessor(Scene, Scene->GetFeatureLevel(), ViewInfo, ShadowContext, Filter);
@@ -2430,7 +2431,7 @@ static void AddHairVisibilityCommonPass(
 		PassParameters->View = ViewInfo->ViewUniformBuffer;
 	}
 
-	AddSimpleMeshPass(GraphBuilder, PassParameters, Scene, *ViewInfo, &InstanceCullingManager, GetPassName(), ViewInfo->ViewRect,
+	AddSimpleMeshPass(GraphBuilder, PassParameters, Scene, *ViewInfo, &InstanceCullingManager, GetPassName(), ViewInfo->ViewRect, false /*bAllowOverrideIndirectArgs*/,
 		[PassParameters, Scene = Scene, ViewInfo, &MacroGroupDatas, RenderMode](FDynamicPassMeshDrawListContext* ShadowContext)
 	{
 		check(IsInRenderingThread());
diff --git a/Engine/Source/Runtime/Renderer/Public/SimpleMeshDrawCommandPass.h b/Engine/Source/Runtime/Renderer/Public/SimpleMeshDrawCommandPass.h
--- a/Engine/Source/Runtime/Renderer/Public/SimpleMeshDrawCommandPass.h
+++ b/Engine/Source/Runtime/Renderer/Public/SimpleMeshDrawCommandPass.h
@@ -66,7 +66,8 @@ template <typename PassParametersType, typename AddMeshBatchesCallbackLambdaType
 void AddSimpleMeshPass(FRDGBuilder& GraphBuilder, PassParametersType* PassParametersIn, const FScene* Scene, const FSceneView &View, FInstanceCullingManager* InstanceCullingManager, FRDGEventName&& PassName,
 	const ERDGPassFlags& PassFlags,
 	AddMeshBatchesCallbackLambdaType AddMeshBatchesCallback,
-	PassPrologueLambdaType PassPrologueCallback)
+	PassPrologueLambdaType PassPrologueCallback,
+	bool bAllowIndirectArgsOverride=true)
 {
 	TRACE_CPUPROFILER_EVENT_SCOPE(AddSimpleMeshPass);
 
@@ -86,6 +87,11 @@ void AddSimpleMeshPass(FRDGBuilder& GraphBuilder, PassParametersType* PassParame
 		SimpleMeshDrawCommandPass->BuildRenderingCommands(GraphBuilder, View, *Scene, PassParameters->InstanceCullingDrawParams);
 	}
 
+	if (!bAllowIndirectArgsOverride)
+	{
+		PassParameters->InstanceCullingDrawParams.DrawIndirectArgsBuffer = nullptr;
+	}
+
 	GraphBuilder.AddPass(
 		MoveTemp(PassName),
 		PassParameters,
@@ -113,6 +119,20 @@ void AddSimpleMeshPass(FRDGBuilder& GraphBuilder, PassParametersType* PassParame
 	);
 }
 
+template <typename PassParametersType, typename AddMeshBatchesCallbackLambdaType>
+void AddSimpleMeshPass(FRDGBuilder& GraphBuilder, PassParametersType* PassParameters, const FScene* Scene, const FSceneView& View, FInstanceCullingManager *InstanceCullingManager, FRDGEventName&& PassName, const FIntRect& ViewPortRect, bool bAllowIndirectArgsOverride,
+	AddMeshBatchesCallbackLambdaType AddMeshBatchesCallback)
+{
+	AddSimpleMeshPass(GraphBuilder, PassParameters, Scene, View, InstanceCullingManager, MoveTemp(PassName), ERDGPassFlags::Raster, AddMeshBatchesCallback,
+		[ViewPortRect](FRHICommandList& RHICmdList)
+		{
+			RHICmdList.SetViewport(static_cast<float>(ViewPortRect.Min.X), static_cast<float>(ViewPortRect.Min.Y), 0.0f,
+			                       static_cast<float>(ViewPortRect.Max.X), static_cast<float>(ViewPortRect.Max.Y), 1.0f);
+		},
+		bAllowIndirectArgsOverride
+	);
+}
+
 // GPUCULL_TODO: Write documentation
 template <typename PassParametersType, typename AddMeshBatchesCallbackLambdaType>
 void AddSimpleMeshPass(FRDGBuilder& GraphBuilder, PassParametersType* PassParameters, const FScene* Scene, const FSceneView& View, FInstanceCullingManager *InstanceCullingManager, FRDGEventName&& PassName, 
