Make support for virtual texture 16bit page tables configurable using the cvar r.VT.Support16BitPageTable.
Defaults to on which is the existing behavior.
In some cases we might want to disable and pay a higher page table memory cost to reduce the overall pressure on the number of unique page table formats which is limited to 16.

diff --git a/Engine/Source/Runtime/Renderer/Private/VT/AllocatedVirtualTexture.cpp b/Engine/Source/Runtime/Renderer/Private/VT/AllocatedVirtualTexture.cpp
--- a/Engine/Source/Runtime/Renderer/Private/VT/AllocatedVirtualTexture.cpp
+++ b/Engine/Source/Runtime/Renderer/Private/VT/AllocatedVirtualTexture.cpp
@@ -2,11 +2,21 @@
 
 #include "AllocatedVirtualTexture.h"
 
+#include "Misc/StringBuilder.h"
 #include "VT/VirtualTextureScalability.h"
 #include "VT/VirtualTextureSystem.h"
 #include "VT/VirtualTextureSpace.h"
 #include "VT/VirtualTexturePhysicalSpace.h"
-#include "Misc/StringBuilder.h"
+
+bool GSupport16BitPageTable = true;
+static FAutoConsoleVariableRef CVarVTSupport16BitPageTable(
+	TEXT("r.VT.Support16BitPageTable"),
+	GSupport16BitPageTable,
+	TEXT("Enable support for 16 bit page table entries.\n")
+	TEXT("This can reduce page table memory when only 16bit addressing is needed.\n")
+	TEXT("But this can increase the number of page table spaces required when a mixture of 16bit and 32bit addressing is needed.\n")
+	TEXT("Defaults on.\n"),
+	ECVF_ReadOnly);
 
 FAllocatedVirtualTexture::FAllocatedVirtualTexture(
 	FRHICommandListBase& RHICmdList,
@@ -74,7 +84,7 @@ FAllocatedVirtualTexture::FAllocatedVirtualTexture(
 	LockOrUnlockTiles(InSystem, true);
 
 	// Use 16bit page table entries if all physical spaces are small enough
-	bool bSupport16BitPageTable = true;
+	bool bSupport16BitPageTable = GSupport16BitPageTable;
 	for (int32 Index = 0; Index < UniquePageTableLayers.Num(); ++Index)
 	{
 		const FVirtualTexturePhysicalSpace* PhysicalSpace = UniquePageTableLayers[Index].PhysicalSpace;
diff --git a/Engine/Source/Runtime/Renderer/Private/VT/VirtualTextureSystem.cpp b/Engine/Source/Runtime/Renderer/Private/VT/VirtualTextureSystem.cpp
--- a/Engine/Source/Runtime/Renderer/Private/VT/VirtualTextureSystem.cpp
+++ b/Engine/Source/Runtime/Renderer/Private/VT/VirtualTextureSystem.cpp
@@ -624,7 +624,6 @@ IAllocatedVirtualTexture* FVirtualTextureSystem::AllocateVirtualTexture(FRHIComm
 	uint32 WidthInBlocks = 0u;
 	uint32 HeightInBlocks = 0u;
 	uint32 DepthInTiles = 0u;
-	bool bSupport16BitPageTable = true;
 	FVirtualTextureProducer* ProducerForLayer[VIRTUALTEXTURE_SPACE_MAXLAYERS] = { nullptr };
 	bool bAnyLayerProducerWantsPersistentHighestMip = false;
 	for (uint32 LayerIndex = 0u; LayerIndex < Desc.NumTextureLayers; ++LayerIndex)
@@ -643,10 +642,6 @@ IAllocatedVirtualTexture* FVirtualTextureSystem::AllocateVirtualTexture(FRHIComm
 			uint32 ProducerLayerIndex = Desc.ProducerLayerIndex[LayerIndex];
 			uint32 ProducerPhysicalGroup = Producer->GetPhysicalGroupIndexForTextureLayer(ProducerLayerIndex);
 			FVirtualTexturePhysicalSpace* PhysicalSpace = Producer->GetPhysicalSpaceForPhysicalGroup(ProducerPhysicalGroup);
-			if (!PhysicalSpace->DoesSupport16BitPageTable())
-			{
-				bSupport16BitPageTable = false;
-			}
 			bAnyLayerProducerWantsPersistentHighestMip |= Producer->GetDescription().bPersistentHighestMip;
 		}
 	}
