Clear CachedRayTracingInstance.Flags before setting necessary flags in UpdateCachedRayTracingInstance(...)

diff --git a/Engine/Source/Runtime/Renderer/Private/PrimitiveSceneInfo.cpp b/Engine/Source/Runtime/Renderer/Private/PrimitiveSceneInfo.cpp
--- a/Engine/Source/Runtime/Renderer/Private/PrimitiveSceneInfo.cpp
+++ b/Engine/Source/Runtime/Renderer/Private/PrimitiveSceneInfo.cpp
@@ -1203,6 +1203,8 @@ void FPrimitiveSceneInfo::UpdateCachedRayTracingInstance(FPrimitiveSceneInfo* Sc
 
 		SceneInfo->CachedRayTracingInstance.bApplyLocalBoundsTransform = CachedRayTracingInstance.bApplyLocalBoundsTransform;
 
+		SceneInfo->CachedRayTracingInstance.Flags = ERayTracingInstanceFlags::None;
+
 		if (CachedRayTracingInstance.MaskAndFlags.bForceOpaque)
 		{
 			SceneInfo->CachedRayTracingInstance.Flags |= ERayTracingInstanceFlags::ForceOpaque;
