Clear bCachedRaytracingDataDirty once the data is updated instead of clearing when update is scheduled.
- Also initizile some missing properties in FPrimitiveSceneInfo constructor.

diff --git a/Engine/Source/Runtime/Renderer/Private/PrimitiveSceneInfo.cpp b/Engine/Source/Runtime/Renderer/Private/PrimitiveSceneInfo.cpp
--- a/Engine/Source/Runtime/Renderer/Private/PrimitiveSceneInfo.cpp
+++ b/Engine/Source/Runtime/Renderer/Private/PrimitiveSceneInfo.cpp
@@ -311,7 +311,9 @@ FPrimitiveSceneInfo::FPrimitiveSceneInfo(const FPrimitiveSceneInfoAdapter& InAda
 	bIsRayTracingRelevant(InAdapter.SceneProxy->IsRayTracingRelevant()),
 	bIsRayTracingStaticRelevant(InAdapter.SceneProxy->IsRayTracingStaticRelevant()),
 	bIsVisibleInRayTracing(InAdapter.SceneProxy->IsVisibleInRayTracing()),
-	bCachedRaytracingDataDirty(false),
+	bCachedRaytracingDataDirty(true),
+	bCachedRayTracingInstanceAnySegmentsDecal(false),
+	bCachedRayTracingInstanceAllSegmentsDecal(false),
 	CoarseMeshStreamingHandle(InAdapter.SceneProxy->GetCoarseMeshStreamingHandle()),
 #endif
 	// We want the unsynchronized access here, as the responsibility passes to the primitive scene info.
@@ -1131,6 +1133,7 @@ void FPrimitiveSceneInfo::CacheRayTracingPrimitives(FScene* Scene, const TArrayV
 					ERayTracingPrimitiveFlags& Flags = Scene->PrimitiveRayTracingFlags[SceneInfo->GetIndex()];
 					CacheRayTracingPrimitive<true>(Scene, SceneInfo, Context.Commands, Context.CommandContext, Context.RayTracingMeshProcessor, &Context.DeferredMeshLODCommandIndices, CachedInstance, Flags);
 					UpdateCachedRayTracingInstance(SceneInfo, CachedInstance, Flags);
+					SceneInfo->bCachedRaytracingDataDirty = false;
 				}
 			);
 
@@ -1162,6 +1165,7 @@ void FPrimitiveSceneInfo::CacheRayTracingPrimitives(FScene* Scene, const TArrayV
 				ERayTracingPrimitiveFlags& Flags = Scene->PrimitiveRayTracingFlags[SceneInfo->GetIndex()];
 				CacheRayTracingPrimitive<false>(Scene, SceneInfo, CachedRayTracingMeshCommands, CommandContext, RayTracingMeshProcessor, nullptr, CachedRayTracingInstance, Flags);
 				UpdateCachedRayTracingInstance(SceneInfo, CachedRayTracingInstance, Flags);
+				SceneInfo->bCachedRaytracingDataDirty = false;
 			}
 		}
 	}
@@ -1882,7 +1886,6 @@ void FPrimitiveSceneInfo::UpdateCachedRaytracingData(FScene* Scene, const TArray
 			// scene info is being updated here
 			check(SceneInfo->bCachedRaytracingDataDirty);
 			SceneInfo->RemoveCachedRayTracingPrimitives();
-			SceneInfo->bCachedRaytracingDataDirty = false;
 		}
 
 		CacheRayTracingPrimitives(Scene, SceneInfos);
diff --git a/Engine/Source/Runtime/Renderer/Private/RayTracing/RayTracing.cpp b/Engine/Source/Runtime/Renderer/Private/RayTracing/RayTracing.cpp
--- a/Engine/Source/Runtime/Renderer/Private/RayTracing/RayTracing.cpp
+++ b/Engine/Source/Runtime/Renderer/Private/RayTracing/RayTracing.cpp
@@ -514,6 +514,8 @@ namespace RayTracing
 							RelevantPrimitive->PrimitiveIndex = PrimitiveIndex;
 							RelevantPrimitive->PersistentPrimitiveIndex = SceneInfo->GetPersistentIndex();
 
+							ensureMsgf(!SceneInfo->bCachedRaytracingDataDirty, TEXT("Cached ray tracing instances must be up-to-date at this point"));
+
 							RelevantPrimitive->CachedRayTracingInstance = &SceneInfo->CachedRayTracingInstance;
 							RelevantPrimitive->bAnySegmentsDecal = SceneInfo->bCachedRayTracingInstanceAnySegmentsDecal;
 							RelevantPrimitive->bAllSegmentsDecal = SceneInfo->bCachedRayTracingInstanceAllSegmentsDecal;
