Ray Tracing - fixed crash with RayTracingLightGridUniformBuffer by binding dummy UB when RT light grid isn't used.

diff --git a/Engine/Source/Runtime/Renderer/Private/DeferredShadingRenderer.cpp b/Engine/Source/Runtime/Renderer/Private/DeferredShadingRenderer.cpp
--- a/Engine/Source/Runtime/Renderer/Private/DeferredShadingRenderer.cpp
+++ b/Engine/Source/Runtime/Renderer/Private/DeferredShadingRenderer.cpp
@@ -735,7 +735,8 @@ void FDeferredShadingSceneRenderer::SetupRayTracingLightDataForViews(FRDGBuilder
 	for (int32 ViewIndex = 0; ViewIndex < Views.Num(); ++ViewIndex)
 	{
 		FViewInfo& View = Views[ViewIndex];
-		View.RayTracingLightGridUniformBuffer = nullptr;
+
+		bool bBuildLightGrid = false;
 
 		// Path Tracing currently uses its own code to manage lights, so doesn't need to run this.
 		if (!bPathTracingEnabled)
@@ -746,10 +747,12 @@ void FDeferredShadingSceneRenderer::SetupRayTracingLightDataForViews(FRDGBuilder
 				|| GetRayTracingTranslucencyOptions(View).bEnabled
 				|| ViewFamily.EngineShowFlags.RayTracingDebug)
 			{
-				// The light data is built in TranslatedWorld space so must be built per view
-				View.RayTracingLightGridUniformBuffer = CreateRayTracingLightData(GraphBuilder, Scene, View, View.ShaderMap);
+				bBuildLightGrid = true;
 			}
 		}
+
+		// The light data is built in TranslatedWorld space so must be built per view
+		View.RayTracingLightGridUniformBuffer = CreateRayTracingLightData(GraphBuilder, Scene, View, View.ShaderMap, bBuildLightGrid);
 	}
 }
 
@@ -1027,15 +1030,9 @@ void FDeferredShadingSceneRenderer::WaitForRayTracingScene(FRDGBuilder& GraphBui
 			else
 			{
 				SetupRayTracingDefaultMissShader(RHICmdList, ReferenceView);
+				SetupRayTracingLightingMissShader(RHICmdList, ReferenceView);
 
-				// only a few passes use the light miss shaders (those that do direct lighting)
-				// simply skip setting these up if we haven't built the RT lighting data
-				if (ReferenceView.RayTracingLightGridUniformBuffer != nullptr)
-				{
-					SetupRayTracingLightingMissShader(RHICmdList, ReferenceView);
-				
-					BindLightFunctionShaders(RHICmdList, Scene, RayTracingLightFunctionMap, ReferenceView);
-				}
+				BindLightFunctionShaders(RHICmdList, Scene, RayTracingLightFunctionMap, ReferenceView);
 			}
 		}
 
diff --git a/Engine/Source/Runtime/Renderer/Private/RayTracing/RayTracingLighting.cpp b/Engine/Source/Runtime/Renderer/Private/RayTracing/RayTracingLighting.cpp
--- a/Engine/Source/Runtime/Renderer/Private/RayTracing/RayTracingLighting.cpp
+++ b/Engine/Source/Runtime/Renderer/Private/RayTracing/RayTracingLighting.cpp
@@ -248,108 +248,126 @@ TRDGUniformBufferRef<FRayTracingLightGrid> CreateRayTracingLightData(
 	FRDGBuilder& GraphBuilder,
 	const FScene* Scene,
 	const FSceneView& View,
-	FGlobalShaderMap* ShaderMap)
+	FGlobalShaderMap* ShaderMap,
+	bool bBuildLightGrid)
 {
 	FRayTracingLightGrid* LightGridParameters = GraphBuilder.AllocParameters<FRayTracingLightGrid>();
-	const FScene::FLightSceneInfoCompactSparseArray& Lights = Scene->Lights;
 
-	// Count the number of lights we want to include by type
-	int NumLightsByType[LightType_MAX] = {};
-	for (const FLightSceneInfoCompact& Light : Lights)
+	if (bBuildLightGrid)
 	{
-		if (!ShouldIncludeRayTracingLight(Light))
-			continue;
-		check(Light.LightType < LightType_MAX);
-		NumLightsByType[Light.LightType]++;
-	}
-
-	// Figure out offset in the target light buffer where each light type will start
-	int LightTypeOffsets[LightType_MAX + 1];
-	LightTypeOffsets[0] = 0;
-	for (int TypeIndex = 1; TypeIndex <= LightType_MAX; TypeIndex++)
-	{
-		LightTypeOffsets[TypeIndex] = LightTypeOffsets[TypeIndex - 1] + NumLightsByType[TypeIndex - 1];
-	}
-
-	LightGridParameters->SceneLightCount = LightTypeOffsets[LightType_MAX];
-
-	FRDGUploadData<FRTLightingData> LightDataArray(GraphBuilder, LightGridParameters->SceneLightCount);
-
-	const FRayTracingLightFunctionMap* RayTracingLightFunctionMap = GraphBuilder.Blackboard.Get<FRayTracingLightFunctionMap>();
-	for (const FLightSceneInfoCompact& Light : Lights)
-	{
-		if (!ShouldIncludeRayTracingLight(Light))
-			continue;
-
-		FLightRenderParameters LightParameters;
-		Light.LightSceneInfo->Proxy->GetLightShaderParameters(LightParameters);
+		const FScene::FLightSceneInfoCompactSparseArray& Lights = Scene->Lights;
 
-		if (Light.LightSceneInfo->Proxy->IsInverseSquared())
+		// Count the number of lights we want to include by type
+		int NumLightsByType[LightType_MAX] = {};
+		for (const FLightSceneInfoCompact& Light : Lights)
 		{
-			LightParameters.FalloffExponent = 0;
+			if (!ShouldIncludeRayTracingLight(Light))
+				continue;
+			check(Light.LightType < LightType_MAX);
+			NumLightsByType[Light.LightType]++;
 		}
 
-		// Figure out where in the target light array this light goes (so that all lights will be sorted by type)
-		int32 Offset = LightTypeOffsets[Light.LightType];
-		LightTypeOffsets[Light.LightType]++; // increment offset for next light
-		
-
-		FRTLightingData& LightDataElement = LightDataArray[Offset];
-
-		LightDataElement.Type = Light.LightType;
-
-		LightDataElement.Direction = LightParameters.Direction;
-		LightDataElement.TranslatedLightPosition = FVector3f(LightParameters.WorldPosition + View.ViewMatrices.GetPreViewTranslation());
-		LightDataElement.LightColor = FVector3f(LightParameters.Color) * LightParameters.GetLightExposureScale(View.GetLastEyeAdaptationExposure());
-		LightDataElement.Tangent = LightParameters.Tangent;
-
-		// Ray tracing should compute fade parameters ignoring lightmaps
-		const FVector2D FadeParams = Light.LightSceneInfo->Proxy->GetDirectionalLightDistanceFadeParameters(View.GetFeatureLevel(), false, View.MaxShadowCascades);
-		const FVector2D DistanceFadeMAD = { FadeParams.Y, -FadeParams.X * FadeParams.Y };
-
-		LightDataElement.SpotAngles = LightParameters.SpotAngles;
-		LightDataElement.DistanceFadeMAD = FVector2f(DistanceFadeMAD);
-
-		LightDataElement.InvRadius = LightParameters.InvRadius;
-		LightDataElement.SpecularScale = LightParameters.SpecularScale;
-		LightDataElement.FalloffExponent = LightParameters.FalloffExponent;
-		LightDataElement.SourceRadius = LightParameters.SourceRadius;
-		LightDataElement.SourceLength = LightParameters.SourceLength;
-		LightDataElement.SoftSourceRadius = LightParameters.SoftSourceRadius;
-		LightDataElement.RectLightBarnCosAngle = LightParameters.RectLightBarnCosAngle;
-		LightDataElement.RectLightBarnLength = LightParameters.RectLightBarnLength;
-		LightDataElement.IESAtlasIndex = LightParameters.IESAtlasIndex;
-		LightDataElement.RectLightAtlasUVOffset[0] = LightParameters.RectLightAtlasUVOffset.X;
-		LightDataElement.RectLightAtlasUVOffset[1] = LightParameters.RectLightAtlasUVOffset.Y;
-		LightDataElement.RectLightAtlasUVScale[0] = LightParameters.RectLightAtlasUVScale.X;
-		LightDataElement.RectLightAtlasUVScale[1] = LightParameters.RectLightAtlasUVScale.Y;
-		LightDataElement.RectLightAtlasMaxLevel = LightParameters.RectLightAtlasMaxLevel;
-		LightDataElement.LightMissShaderIndex = RAY_TRACING_MISS_SHADER_SLOT_LIGHTING;
-
-		// Stuff directional light's shadow angle factor into a RectLight parameter
-		if (Light.LightType == LightType_Directional)
+		// Figure out offset in the target light buffer where each light type will start
+		int LightTypeOffsets[LightType_MAX + 1];
+		LightTypeOffsets[0] = 0;
+		for (int TypeIndex = 1; TypeIndex <= LightType_MAX; TypeIndex++)
 		{
-			LightDataElement.RectLightBarnCosAngle = Light.LightSceneInfo->Proxy->GetShadowSourceAngleFactor();
+			LightTypeOffsets[TypeIndex] = LightTypeOffsets[TypeIndex - 1] + NumLightsByType[TypeIndex - 1];
 		}
 
-		// NOTE: This map will be empty if the light functions are disabled for some reason
-		if (RayTracingLightFunctionMap)
+		LightGridParameters->SceneLightCount = LightTypeOffsets[LightType_MAX];
+
+		FRDGUploadData<FRTLightingData> LightDataArray(GraphBuilder, LightGridParameters->SceneLightCount);
+
+		const FRayTracingLightFunctionMap* RayTracingLightFunctionMap = GraphBuilder.Blackboard.Get<FRayTracingLightFunctionMap>();
+		for (const FLightSceneInfoCompact& Light : Lights)
 		{
-			const int32* LightFunctionIndex = RayTracingLightFunctionMap->Find(Light.LightSceneInfo);
-			if (LightFunctionIndex)
+			if (!ShouldIncludeRayTracingLight(Light))
+				continue;
+
+			FLightRenderParameters LightParameters;
+			Light.LightSceneInfo->Proxy->GetLightShaderParameters(LightParameters);
+
+			if (Light.LightSceneInfo->Proxy->IsInverseSquared())
 			{
-				check(uint32(*LightFunctionIndex) > RAY_TRACING_MISS_SHADER_SLOT_LIGHTING);
-				check(uint32(*LightFunctionIndex) < Scene->RayTracingScene.NumMissShaderSlots);
-				LightDataElement.LightMissShaderIndex = *LightFunctionIndex;
+				LightParameters.FalloffExponent = 0;
+			}
+
+			// Figure out where in the target light array this light goes (so that all lights will be sorted by type)
+			int32 Offset = LightTypeOffsets[Light.LightType];
+			LightTypeOffsets[Light.LightType]++; // increment offset for next light
+
+
+			FRTLightingData& LightDataElement = LightDataArray[Offset];
+
+			LightDataElement.Type = Light.LightType;
+
+			LightDataElement.Direction = LightParameters.Direction;
+			LightDataElement.TranslatedLightPosition = FVector3f(LightParameters.WorldPosition + View.ViewMatrices.GetPreViewTranslation());
+			LightDataElement.LightColor = FVector3f(LightParameters.Color) * LightParameters.GetLightExposureScale(View.GetLastEyeAdaptationExposure());
+			LightDataElement.Tangent = LightParameters.Tangent;
+
+			// Ray tracing should compute fade parameters ignoring lightmaps
+			const FVector2D FadeParams = Light.LightSceneInfo->Proxy->GetDirectionalLightDistanceFadeParameters(View.GetFeatureLevel(), false, View.MaxShadowCascades);
+			const FVector2D DistanceFadeMAD = { FadeParams.Y, -FadeParams.X * FadeParams.Y };
+
+			LightDataElement.SpotAngles = LightParameters.SpotAngles;
+			LightDataElement.DistanceFadeMAD = FVector2f(DistanceFadeMAD);
+
+			LightDataElement.InvRadius = LightParameters.InvRadius;
+			LightDataElement.SpecularScale = LightParameters.SpecularScale;
+			LightDataElement.FalloffExponent = LightParameters.FalloffExponent;
+			LightDataElement.SourceRadius = LightParameters.SourceRadius;
+			LightDataElement.SourceLength = LightParameters.SourceLength;
+			LightDataElement.SoftSourceRadius = LightParameters.SoftSourceRadius;
+			LightDataElement.RectLightBarnCosAngle = LightParameters.RectLightBarnCosAngle;
+			LightDataElement.RectLightBarnLength = LightParameters.RectLightBarnLength;
+			LightDataElement.IESAtlasIndex = LightParameters.IESAtlasIndex;
+			LightDataElement.RectLightAtlasUVOffset[0] = LightParameters.RectLightAtlasUVOffset.X;
+			LightDataElement.RectLightAtlasUVOffset[1] = LightParameters.RectLightAtlasUVOffset.Y;
+			LightDataElement.RectLightAtlasUVScale[0] = LightParameters.RectLightAtlasUVScale.X;
+			LightDataElement.RectLightAtlasUVScale[1] = LightParameters.RectLightAtlasUVScale.Y;
+			LightDataElement.RectLightAtlasMaxLevel = LightParameters.RectLightAtlasMaxLevel;
+			LightDataElement.LightMissShaderIndex = RAY_TRACING_MISS_SHADER_SLOT_LIGHTING;
+
+			// Stuff directional light's shadow angle factor into a RectLight parameter
+			if (Light.LightType == LightType_Directional)
+			{
+				LightDataElement.RectLightBarnCosAngle = Light.LightSceneInfo->Proxy->GetShadowSourceAngleFactor();
+			}
+
+			// NOTE: This map will be empty if the light functions are disabled for some reason
+			if (RayTracingLightFunctionMap)
+			{
+				const int32* LightFunctionIndex = RayTracingLightFunctionMap->Find(Light.LightSceneInfo);
+				if (LightFunctionIndex)
+				{
+					check(uint32(*LightFunctionIndex) > RAY_TRACING_MISS_SHADER_SLOT_LIGHTING);
+					check(uint32(*LightFunctionIndex) < Scene->RayTracingScene.NumMissShaderSlots);
+					LightDataElement.LightMissShaderIndex = *LightFunctionIndex;
+				}
 			}
 		}
-	}
-	// last light type should not match the total scene light count
-	check(LightGridParameters->SceneLightCount == LightTypeOffsets[LightType_MAX - 1]);
+		// last light type should not match the total scene light count
+		check(LightGridParameters->SceneLightCount == LightTypeOffsets[LightType_MAX - 1]);
 
-	LightGridParameters->SceneLights = GraphBuilder.CreateSRV(CreateStructuredBuffer(GraphBuilder, TEXT("LightBuffer"), LightDataArray));
-	LightGridParameters->SceneInfiniteLightCount = NumLightsByType[LightType_Directional];
-	PrepareLightGrid(GraphBuilder, ShaderMap, LightGridParameters, LightDataArray.GetData());
+		LightGridParameters->SceneLights = GraphBuilder.CreateSRV(CreateStructuredBuffer(GraphBuilder, TEXT("LightBuffer"), LightDataArray));
+		LightGridParameters->SceneInfiniteLightCount = NumLightsByType[LightType_Directional];
+		PrepareLightGrid(GraphBuilder, ShaderMap, LightGridParameters, LightDataArray.GetData());
+	}
+	else
+	{
+		LightGridParameters->SceneLightCount = 0;
+		LightGridParameters->SceneInfiniteLightCount = 0;
+		LightGridParameters->SceneLightsTranslatedBoundMin = FVector3f::ZeroVector;
+		LightGridParameters->SceneLightsTranslatedBoundMax = FVector3f::ZeroVector;
+		LightGridParameters->SceneLights = GraphBuilder.CreateSRV(GSystemTextures.GetDefaultBuffer(GraphBuilder, sizeof(uint32), 0u), PF_R32_UINT);
+		LightGridParameters->LightGrid = GSystemTextures.GetDefaultTexture2D(GraphBuilder, PF_R32_UINT, 0u);
+		LightGridParameters->LightGridData = GraphBuilder.CreateSRV(GSystemTextures.GetDefaultStructuredBuffer(GraphBuilder, sizeof(uint32), 0u), PF_R32_UINT);
+		LightGridParameters->LightGridResolution = 0;
+		LightGridParameters->LightGridMaxCount = 0;
+		LightGridParameters->LightGridAxis = 0;
+	}
 
 	return GraphBuilder.CreateUniformBuffer(LightGridParameters);
 }
diff --git a/Engine/Source/Runtime/Renderer/Private/RayTracing/RayTracingLighting.h b/Engine/Source/Runtime/Renderer/Private/RayTracing/RayTracingLighting.h
--- a/Engine/Source/Runtime/Renderer/Private/RayTracing/RayTracingLighting.h
+++ b/Engine/Source/Runtime/Renderer/Private/RayTracing/RayTracingLighting.h
@@ -41,7 +41,8 @@ TRDGUniformBufferRef<FRayTracingLightGrid> CreateRayTracingLightData(
 	FRDGBuilder& GraphBuilder,
 	const FScene* Scene,
 	const FSceneView& View,
-	FGlobalShaderMap* ShaderMap);
+	FGlobalShaderMap* ShaderMap,
+	bool bBuildLightGrid);
 
 void BindLightFunctionShaders(
 	FRHICommandList& RHICmdList,
