Constructing velocity buffer for heterogeneous volumes

diff --git a/Engine/Shaders/Private/HeterogeneousVolumes/HeterogeneousVolumesHardwareRayTracing.usf b/Engine/Shaders/Private/HeterogeneousVolumes/HeterogeneousVolumesHardwareRayTracing.usf
--- a/Engine/Shaders/Private/HeterogeneousVolumes/HeterogeneousVolumesHardwareRayTracing.usf
+++ b/Engine/Shaders/Private/HeterogeneousVolumes/HeterogeneousVolumesHardwareRayTracing.usf
@@ -557,11 +557,13 @@ RAY_TRACING_ENTRY_RAYGEN(RenderSingleScatteringWithPreshadingRGS)
 		const FDeferredLightData LightData = InitDeferredLightFromUniforms(LightType, VolumetricScatteringIntensity);
 		uint StepCount = CalcStepCount(RayMarchingContext);
 
+		float ScatterHitT = Payload.HitT[1];
 		RayMarchSingleScattering(
 			RayMarchingContext,
 			LightData,
 			LightType,
 			StepCount,
+			ScatterHitT,
 			Radiance,
 			Transmittance
 		);
diff --git a/Engine/Shaders/Private/HeterogeneousVolumes/HeterogeneousVolumesLiveShadingPipeline.usf b/Engine/Shaders/Private/HeterogeneousVolumes/HeterogeneousVolumesLiveShadingPipeline.usf
--- a/Engine/Shaders/Private/HeterogeneousVolumes/HeterogeneousVolumesLiveShadingPipeline.usf
+++ b/Engine/Shaders/Private/HeterogeneousVolumes/HeterogeneousVolumesLiveShadingPipeline.usf
@@ -74,6 +74,7 @@ int3 GroupCount;
 
 // Output
 RWTexture2D<float4> RWLightingTexture;
+RWTexture2D<float4> RWVelocityTexture;
 
 struct FDebugOutput
 {
@@ -314,6 +315,7 @@ void RenderSingleScatteringWithLiveShadingCS(
 
 	float3 Radiance = 0.0;
 	float3 Transmittance = 1.0 - RWLightingTexture[PixelCoord].a;
+	float PrevOpacity = Luminance(1.0 - Transmittance);
 
 	// Extract depth
 	float DeviceZ = SceneDepthTexture.Load(int3(PixelCoord, 0)).r;
@@ -371,21 +373,37 @@ void RenderSingleScatteringWithLiveShadingCS(
 		const FDeferredLightData LightData = InitDeferredLightFromUniforms(LightType, VolumetricScatteringIntensity);
 		uint StepCount = CalcStepCount(RayMarchingContext);
 
+		float ScatterHitT = HitT.y;
 		RayMarchSingleScattering(
 			RayMarchingContext,
 			LightData,
 			LightType,
 			StepCount,
+			ScatterHitT,
 			Radiance,
 			Transmittance
 		);
 
-		// Output..
+		// Radiance
 		RWLightingTexture[PixelCoord].rgb += Radiance * View.PreExposure;
 		if (bApplyEmissionAndTransmittance)
 		{
-			float3 Opacity = 1.0 - Transmittance;
-			RWLightingTexture[PixelCoord].a = Luminance(Opacity);
+			float Opacity = Luminance(1.0 - Transmittance);
+			RWLightingTexture[PixelCoord].a = Opacity;
+
+#if DIM_WRITE_VELOCITY
+			bool bWriteVelocity =
+				(ScatterHitT < HitT.y) &&
+				(PrevOpacity < 0.5) &&
+				(Opacity > 0.5);
+
+			if (bWriteVelocity)
+			{
+				float WorldScatterHitT = ScatterHitT * RayMarchingContext.LocalToWorldScale;
+				float3 Velocity = Calculate3DVelocityAtWorldPos(WorldRayOrigin, WorldRayDirection, WorldScatterHitT);
+				RWVelocityTexture[PixelCoord] = EncodeVelocityToTexture(Velocity);
+			}
+#endif // DIM_WRITE_VELOCITY
 		}
 
 		// Debug output
diff --git a/Engine/Shaders/Private/HeterogeneousVolumes/HeterogeneousVolumesPreshadingPipeline.usf b/Engine/Shaders/Private/HeterogeneousVolumes/HeterogeneousVolumesPreshadingPipeline.usf
--- a/Engine/Shaders/Private/HeterogeneousVolumes/HeterogeneousVolumesPreshadingPipeline.usf
+++ b/Engine/Shaders/Private/HeterogeneousVolumes/HeterogeneousVolumesPreshadingPipeline.usf
@@ -68,6 +68,7 @@ int MipLevel;
 
 // Output
 RWTexture2D<float4> RWLightingTexture;
+RWTexture2D<float4> RWVelocityTexture;
 
 int3 TextureResolution;
 Texture3D InputTexture;
@@ -517,6 +518,7 @@ void RenderSingleScatteringWithPreshadingCS(
 
 	float3 Radiance = 0.0;
 	float3 Transmittance = 1.0 - RWLightingTexture[PixelCoord].a;
+	float PrevOpacity = Luminance(1.0 - Transmittance);
 
 	// Create screen ray
 	if (any(PixelCoord.xy >= View.ViewSizeAndInvSize.xy))
@@ -562,6 +564,7 @@ void RenderSingleScatteringWithPreshadingCS(
 		return;
 	}
 
+	float WorldHitTMin = WorldRayLength;
 #if DIM_SPARSE_VOXEL_TRACING
 	float RayTMax = HitT.y;
 
@@ -599,14 +602,17 @@ void RenderSingleScatteringWithPreshadingCS(
 		const FDeferredLightData LightData = InitDeferredLightFromUniforms(LightType, VolumetricScatteringIntensity);
 		uint StepCount = CalcStepCount(RayMarchingContext);
 
+		float ScatterHitT = HitT.y;
 		RayMarchSingleScattering(
 			RayMarchingContext,
 			LightData,
 			LightType,
 			StepCount,
+			ScatterHitT,
 			Radiance,
 			Transmittance
 		);
+		WorldHitTMin = min(WorldHitTMin, ScatterHitT * RayMarchingContext.LocalToWorldScale);
 
 #if DIM_SPARSE_VOXEL_TRACING
 	#if DIM_VOXEL_CULLING
@@ -621,9 +627,22 @@ void RenderSingleScatteringWithPreshadingCS(
 #endif // DIM_SPARSE_VOXEL_TRACING
 	}
 
-	// Output..
-	float3 Opacity = bApplyEmissionAndTransmittance ? 1.0 - Transmittance : 0.0;
-	RWLightingTexture[PixelCoord] = float4(RWLightingTexture[PixelCoord].rgb + (Radiance * View.PreExposure), Luminance(Opacity));
+	// Radiance
+	float Opacity = bApplyEmissionAndTransmittance ? Luminance(1.0 - Transmittance) : 0.0;
+	RWLightingTexture[PixelCoord] = float4(RWLightingTexture[PixelCoord].rgb + (Radiance * View.PreExposure), Opacity);
+
+#if DIM_WRITE_VELOCITY
+	bool bWriteVelocity =
+		(WorldHitTMin < WorldRayLength) &&
+		(PrevOpacity < 0.5) &&
+		(Opacity > 0.5);
+
+	if (bWriteVelocity)
+	{
+		float3 Velocity = Calculate3DVelocityAtWorldPos(WorldRayOrigin, WorldRayDirection, WorldHitTMin);
+		RWVelocityTexture[PixelCoord] = EncodeVelocityToTexture(Velocity);
+	}
+#endif // DIM_WRITE_VELOCITY
 
 #if DIM_DEBUG
 	//if (!all(DispatchThreadId))
diff --git a/Engine/Shaders/Private/HeterogeneousVolumes/HeterogeneousVolumesRayMarchingUtils.ush b/Engine/Shaders/Private/HeterogeneousVolumes/HeterogeneousVolumesRayMarchingUtils.ush
--- a/Engine/Shaders/Private/HeterogeneousVolumes/HeterogeneousVolumesRayMarchingUtils.ush
+++ b/Engine/Shaders/Private/HeterogeneousVolumes/HeterogeneousVolumesRayMarchingUtils.ush
@@ -307,12 +307,13 @@ void RayMarchSingleScattering(
 	FDeferredLightData LightData,
 	uint LightType,
 	uint StepCount,
+	inout float LocalScatterTMin,
 	inout float3 Radiance,
 	inout float3 Transmittance
 )
 {
+	LocalScatterTMin = RayMarchingContext.LocalRayTMax;
 #if APPLY_FOG_INSCATTERING
-	float LocalScatterTMin = RayMarchingContext.LocalRayTMax;
 	float3 PrevTransmittance = Transmittance;
 #endif // #if APPLY_FOG_INSCATTERING
 
@@ -370,9 +371,7 @@ void RayMarchSingleScattering(
 				Radiance += Inscattering * ScatteringCoefficient * IsotropicPhase * Transmittance * StepSize;
 			}
 
-#if APPLY_FOG_INSCATTERING
 			LocalScatterTMin = min(LocalHitT, LocalScatterTMin);
-#endif 
 		}
 
 		// Accumulate transmittance for the next evaluation
diff --git a/Engine/Shaders/Private/HeterogeneousVolumes/HeterogeneousVolumesTracingUtils.ush b/Engine/Shaders/Private/HeterogeneousVolumes/HeterogeneousVolumesTracingUtils.ush
--- a/Engine/Shaders/Private/HeterogeneousVolumes/HeterogeneousVolumesTracingUtils.ush
+++ b/Engine/Shaders/Private/HeterogeneousVolumes/HeterogeneousVolumesTracingUtils.ush
@@ -2,6 +2,8 @@
 
 #pragma once
 
+#include "../VelocityCommon.ush"
+
 // Modified LineBoxIntersect() to operate at consistent RayDirection scale, within some [TMin, TMax] span
 float2 IntersectAABB(float3 RayOrigin, float3 RayDirection, float RayTMin, float RayTMax, float3 BoxMin, float3 BoxMax)
 {
@@ -23,3 +25,18 @@ float2 IntersectAABB(float3 RayOrigin, float3 RayDirection, float RayTMin, float
 	//clamp the intersections to be between RayTMin and RayTMax
 	return float2(max(BoxIntersections.x, RayTMin), min(BoxIntersections.y, RayTMax));
 }
+
+float3 Calculate3DVelocityAtWorldPos(float3 WorldRayOrigin, float3 WorldRayDirection, float WorldHitT)
+{
+	float3 WorldPos = WorldRayOrigin + WorldRayDirection * WorldHitT;
+	float3 TranslatedWorldPos = DFHackToFloat(DFFastAdd(WorldPos, PrimaryView.PreViewTranslation));
+	float4 ClipPos = mul(float4(TranslatedWorldPos, 1), View.TranslatedWorldToClip);
+	ClipPos /= ClipPos.w;
+
+	float3 PreViewTranslationOffset = DFFastLocalSubtractDemote(PrimaryView.PreViewTranslation, PrimaryView.PrevPreViewTranslation);
+	float3 PrevTranslatedWorldPos = mul(float4(ClipPos.xy * WorldHitT, WorldHitT, 1), View.PrevScreenToTranslatedWorld).xyz + PreViewTranslationOffset;
+	float4 PrevClipPos = mul(float4(PrevTranslatedWorldPos, 1), View.TranslatedWorldToClip);
+	PrevClipPos /= PrevClipPos.w;
+
+	return Calculate3DVelocity(ClipPos, PrevClipPos);
+}
\ No newline at end of file
diff --git a/Engine/Source/Runtime/Renderer/Private/HeterogeneousVolumes/HeterogeneousVolumes.cpp b/Engine/Source/Runtime/Renderer/Private/HeterogeneousVolumes/HeterogeneousVolumes.cpp
--- a/Engine/Source/Runtime/Renderer/Private/HeterogeneousVolumes/HeterogeneousVolumes.cpp
+++ b/Engine/Source/Runtime/Renderer/Private/HeterogeneousVolumes/HeterogeneousVolumes.cpp
@@ -210,6 +210,13 @@ static TAutoConsoleVariable<int32> CVarHeterogeneousVolumesApplyVolumetricFog(
 	ECVF_RenderThreadSafe
 );
 
+static TAutoConsoleVariable<int32> CVarHeterogeneousVolumesVelocity(
+	TEXT("r.HeterogeneousVolumes.Velocity"),
+	0,
+	TEXT("Writes Heterogeneous Volumes velocity to the feature buffer (Default = 0)"),
+	ECVF_RenderThreadSafe
+);
+
 DECLARE_GPU_STAT_NAMED(HeterogeneousVolumeShadowsStat, TEXT("HeterogeneousVolumeShadows"));
 DECLARE_GPU_STAT_NAMED(HeterogeneousVolumesStat, TEXT("HeterogeneousVolumes"));
 
@@ -413,6 +420,11 @@ namespace HeterogeneousVolumes
 		return CVarHeterogeneousVolumesApplyVolumetricFog.GetValueOnRenderThread() != 0;
 	}
 
+	bool ShouldWriteVelocity()
+	{
+		return CVarHeterogeneousVolumesVelocity.GetValueOnRenderThread() != 0;
+	}
+
 	// Convenience Utils
 	int GetVoxelCount(FIntVector VolumeResolution)
 	{
diff --git a/Engine/Source/Runtime/Renderer/Private/HeterogeneousVolumes/HeterogeneousVolumes.h b/Engine/Source/Runtime/Renderer/Private/HeterogeneousVolumes/HeterogeneousVolumes.h
--- a/Engine/Source/Runtime/Renderer/Private/HeterogeneousVolumes/HeterogeneousVolumes.h
+++ b/Engine/Source/Runtime/Renderer/Private/HeterogeneousVolumes/HeterogeneousVolumes.h
@@ -80,6 +80,7 @@ namespace HeterogeneousVolumes
 	bool UseAdaptiveVolumetricShadowMapForSelfShadowing();
 	bool ShouldApplyHeightFog();
 	bool ShouldApplyVolumetricFog();
+	bool ShouldWriteVelocity();
 
 	bool EnableIndirectionGrid();
 	bool EnableLinearInterpolation();
diff --git a/Engine/Source/Runtime/Renderer/Private/HeterogeneousVolumes/HeterogeneousVolumesLiveShadingPipeline.cpp b/Engine/Source/Runtime/Renderer/Private/HeterogeneousVolumes/HeterogeneousVolumesLiveShadingPipeline.cpp
--- a/Engine/Source/Runtime/Renderer/Private/HeterogeneousVolumes/HeterogeneousVolumesLiveShadingPipeline.cpp
+++ b/Engine/Source/Runtime/Renderer/Private/HeterogeneousVolumes/HeterogeneousVolumesLiveShadingPipeline.cpp
@@ -232,8 +232,9 @@ class FRenderSingleScatteringWithLiveShadingCS : public FMeshMaterialShader
 	class FUseTransmittanceVolume : SHADER_PERMUTATION_BOOL("DIM_USE_TRANSMITTANCE_VOLUME");
 	class FUseInscatteringVolume : SHADER_PERMUTATION_BOOL("DIM_USE_INSCATTERING_VOLUME");
 	class FUseLumenGI : SHADER_PERMUTATION_BOOL("DIM_USE_LUMEN_GI");
+	class FWriteVelocity : SHADER_PERMUTATION_BOOL("DIM_WRITE_VELOCITY");
 	class FUseAdaptiveVolumetricShadowMap : SHADER_PERMUTATION_BOOL("DIM_USE_ADAPTIVE_VOLUMETRIC_SHADOW_MAP");
-	using FPermutationDomain = TShaderPermutationDomain<FUseTransmittanceVolume, FUseInscatteringVolume, FUseLumenGI, FUseAdaptiveVolumetricShadowMap>;
+	using FPermutationDomain = TShaderPermutationDomain<FUseTransmittanceVolume, FUseInscatteringVolume, FUseLumenGI, FWriteVelocity, FUseAdaptiveVolumetricShadowMap>;
 
 	BEGIN_SHADER_PARAMETER_STRUCT(FParameters, )
 		// Scene data
@@ -289,6 +290,7 @@ class FRenderSingleScatteringWithLiveShadingCS : public FMeshMaterialShader
 
 		// Output
 		SHADER_PARAMETER_RDG_TEXTURE_UAV(RWTexture2D<float4>, RWLightingTexture)
+		SHADER_PARAMETER_RDG_TEXTURE_UAV(RWTexture2D<float4>, RWVelocityTexture)
 		//SHADER_PARAMETER_RDG_BUFFER_UAV(RWStructuredBuffer<Volumes::FDebugOutput>, RWDebugOutputBuffer)
 	END_SHADER_PARAMETER_STRUCT()
 
@@ -644,6 +646,7 @@ static void RenderSingleScatteringWithLiveShading(
 	}
 	TUniformBufferRef<FDeferredLightUniformStruct> DeferredLightUB = CreateUniformBufferImmediate(DeferredLightUniform, UniformBuffer_SingleDraw);
 
+	bool bWriteVelocity = HeterogeneousVolumes::ShouldWriteVelocity() && HasBeenProduced(SceneTextures.Velocity);
 	FRenderSingleScatteringWithLiveShadingCS::FParameters* PassParameters = GraphBuilder.AllocParameters<FRenderSingleScatteringWithLiveShadingCS::FParameters>();
 	{
 		// Scene data
@@ -745,6 +748,10 @@ static void RenderSingleScatteringWithLiveShading(
 
 		// Output
 		PassParameters->RWLightingTexture = GraphBuilder.CreateUAV(HeterogeneousVolumeTexture);
+		if (bWriteVelocity)
+		{
+			PassParameters->RWVelocityTexture = GraphBuilder.CreateUAV(SceneTextures.Velocity);
+		}
 		//PassParameters->RWDebugOutputBuffer = GraphBuilder.CreateUAV(DebugOutputBuffer);
 	}
 
@@ -765,6 +772,7 @@ static void RenderSingleScatteringWithLiveShading(
 	PermutationVector.Set<FRenderSingleScatteringWithLiveShadingCS::FUseTransmittanceVolume>(HeterogeneousVolumes::UseLightingCacheForTransmittance() && PassParameters->bApplyShadowTransmittance);
 	PermutationVector.Set<FRenderSingleScatteringWithLiveShadingCS::FUseInscatteringVolume>(HeterogeneousVolumes::UseLightingCacheForInscattering());
 	PermutationVector.Set<FRenderSingleScatteringWithLiveShadingCS::FUseLumenGI>(HeterogeneousVolumes::UseIndirectLighting() && View.GetLumenTranslucencyGIVolume().Texture0 != nullptr);
+	PermutationVector.Set<FRenderSingleScatteringWithLiveShadingCS::FWriteVelocity>(bWriteVelocity);
 	PermutationVector.Set<FRenderSingleScatteringWithLiveShadingCS::FUseAdaptiveVolumetricShadowMap>(HeterogeneousVolumes::UseAdaptiveVolumetricShadowMapForSelfShadowing());
 	TShaderRef<FRenderSingleScatteringWithLiveShadingCS> ComputeShader = Material.GetShader<FRenderSingleScatteringWithLiveShadingCS>(&FLocalVertexFactory::StaticType, PermutationVector, false);
 	if (!ComputeShader.IsNull())
diff --git a/Engine/Source/Runtime/Renderer/Private/HeterogeneousVolumes/HeterogeneousVolumesPreshadingPipeline.cpp b/Engine/Source/Runtime/Renderer/Private/HeterogeneousVolumes/HeterogeneousVolumesPreshadingPipeline.cpp
--- a/Engine/Source/Runtime/Renderer/Private/HeterogeneousVolumes/HeterogeneousVolumesPreshadingPipeline.cpp
+++ b/Engine/Source/Runtime/Renderer/Private/HeterogeneousVolumes/HeterogeneousVolumesPreshadingPipeline.cpp
@@ -294,8 +294,9 @@ class FRenderSingleScatteringWithPreshadingCS : public FGlobalShader
 	class FUseTransmittanceVolume : SHADER_PERMUTATION_BOOL("DIM_USE_TRANSMITTANCE_VOLUME");
 	class FUseInscatteringVolume : SHADER_PERMUTATION_BOOL("DIM_USE_INSCATTERING_VOLUME");
 	class FUseLumenGI : SHADER_PERMUTATION_BOOL("DIM_USE_LUMEN_GI");
+	class FWriteVelocity : SHADER_PERMUTATION_BOOL("DIM_WRITE_VELOCITY");
 	class FDebugDim : SHADER_PERMUTATION_BOOL("DIM_DEBUG");
-	using FPermutationDomain = TShaderPermutationDomain<FApplyShadowTransmittanceDim, FVoxelCullingDim, FSparseVoxelTracingDim, FUseTransmittanceVolume, FUseInscatteringVolume, FUseLumenGI, FDebugDim>;
+	using FPermutationDomain = TShaderPermutationDomain<FApplyShadowTransmittanceDim, FVoxelCullingDim, FSparseVoxelTracingDim, FUseTransmittanceVolume, FUseInscatteringVolume, FUseLumenGI, FWriteVelocity, FDebugDim>;
 
 	BEGIN_SHADER_PARAMETER_STRUCT(FParameters, )
 		// Scene data
@@ -343,6 +344,7 @@ class FRenderSingleScatteringWithPreshadingCS : public FGlobalShader
 
 		// Output
 		SHADER_PARAMETER_RDG_TEXTURE_UAV(RWTexture2D<float4>, RWLightingTexture)
+		SHADER_PARAMETER_RDG_TEXTURE_UAV(RWTexture2D<float4>, RWVelocityTexture)
 		SHADER_PARAMETER_RDG_BUFFER_UAV(RWStructuredBuffer<FVoxelDataPacked>, RWVoxelOutputBuffer)
 	END_SHADER_PARAMETER_STRUCT()
 
@@ -544,6 +546,7 @@ void RenderSingleScatteringWithPreshadingCompute(
 	}
 	TUniformBufferRef<FDeferredLightUniformStruct> DeferredLightUB = CreateUniformBufferImmediate(DeferredLightUniform, UniformBuffer_SingleDraw);
 
+	bool bWriteVelocity = HeterogeneousVolumes::ShouldWriteVelocity() && HasBeenProduced(SceneTextures.Velocity);
 	FRenderSingleScatteringWithPreshadingCS::FParameters* PassParameters = GraphBuilder.AllocParameters<FRenderSingleScatteringWithPreshadingCS::FParameters>();
 	{
 		// Scene data
@@ -626,6 +629,10 @@ void RenderSingleScatteringWithPreshadingCompute(
 
 		// Output
 		PassParameters->RWLightingTexture = GraphBuilder.CreateUAV(HeterogeneousVolumeTexture);
+		if (bWriteVelocity)
+		{
+			PassParameters->RWVelocityTexture = GraphBuilder.CreateUAV(SceneTextures.Velocity);
+		}
 		PassParameters->RWVoxelOutputBuffer = GraphBuilder.CreateUAV(VoxelOutputBuffer);
 	}
 
@@ -642,6 +649,7 @@ void RenderSingleScatteringWithPreshadingCompute(
 	PermutationVector.Set<FRenderSingleScatteringWithPreshadingCS::FUseTransmittanceVolume>(HeterogeneousVolumes::UseLightingCacheForTransmittance());
 	PermutationVector.Set<FRenderSingleScatteringWithPreshadingCS::FUseInscatteringVolume>(HeterogeneousVolumes::UseLightingCacheForInscattering());
 	PermutationVector.Set<FRenderSingleScatteringWithPreshadingCS::FUseLumenGI>(HeterogeneousVolumes::UseIndirectLighting() && View.GetLumenTranslucencyGIVolume().Texture0 != nullptr);
+	PermutationVector.Set<FRenderSingleScatteringWithPreshadingCS::FWriteVelocity>(bWriteVelocity);
 	PermutationVector.Set<FRenderSingleScatteringWithPreshadingCS::FDebugDim>(HeterogeneousVolumes::GetDebugMode() != 0);
 	TShaderRef<FRenderSingleScatteringWithPreshadingCS> ComputeShader = View.ShaderMap->GetShader<FRenderSingleScatteringWithPreshadingCS>(PermutationVector);
 	FComputeShaderUtils::AddPass(
