Clamp number of active scene captures for mobile
Deprecate calling GMaxNumReflectionCaptures and GMobileMaxNumReflectionCaptures directly to avoid errors in the future

diff --git a/Engine/Source/Runtime/Renderer/Private/IndirectLightRendering.cpp b/Engine/Source/Runtime/Renderer/Private/IndirectLightRendering.cpp
--- a/Engine/Source/Runtime/Renderer/Private/IndirectLightRendering.cpp
+++ b/Engine/Source/Runtime/Renderer/Private/IndirectLightRendering.cpp
@@ -383,7 +383,7 @@ class FReflectionEnvironmentSkyLightingPS : public FGlobalShader
 	static void ModifyCompilationEnvironment(const FGlobalShaderPermutationParameters& Parameters, FShaderCompilerEnvironment& OutEnvironment)
 	{
 		FGlobalShader::ModifyCompilationEnvironment(Parameters, OutEnvironment);
-		OutEnvironment.SetDefine(TEXT("MAX_CAPTURES"), GMaxNumReflectionCaptures);
+		OutEnvironment.SetDefine(TEXT("MAX_CAPTURES"), GetMaxNumReflectionCaptures(Parameters.Platform));
 		OutEnvironment.SetDefine(TEXT("SUPPORTS_ANISOTROPIC_MATERIALS"), FDataDrivenShaderPlatformInfo::GetSupportsAnisotropicMaterials(Parameters.Platform));
 		OutEnvironment.SetDefine(TEXT("USE_HAIR_COMPLEX_TRANSMITTANCE"), IsHairStrandsSupported(EHairStrandsShaderType::All, Parameters.Platform) ? 1u : 0u);
 		OutEnvironment.SetDefine(TEXT("SUBSTRATE_GLINTS_IS"), 0);
diff --git a/Engine/Source/Runtime/Renderer/Private/RendererScene.cpp b/Engine/Source/Runtime/Renderer/Private/RendererScene.cpp
--- a/Engine/Source/Runtime/Renderer/Private/RendererScene.cpp
+++ b/Engine/Source/Runtime/Renderer/Private/RendererScene.cpp
@@ -5274,7 +5274,7 @@ void UpdateReflectionSceneData(FScene* Scene)
 	ReflectionSceneData.NumSphereCaptures = 0;
 
 	const int32 MaxCubemaps = ReflectionSceneData.CubemapArray.GetMaxCubemaps();
-	int32_t PlatformMaxNumReflectionCaptures = FMath::Min(FMath::FloorToInt(GMaxTextureArrayLayers / 6.0f), GMaxNumReflectionCaptures);
+	int32_t PlatformMaxNumReflectionCaptures = FMath::Min(FMath::FloorToInt(GMaxTextureArrayLayers / 6.0f), GetMaxNumReflectionCaptures(Scene->GetShaderPlatform()));
 
 	// Pack visible reflection captures into the uniform buffer, each with an index to its cubemap array entry.
 	// GPUScene primitive data stores closest reflection capture as index into this buffer, so this index which must be invalidate every time OutSortData contents change.
diff --git a/Engine/Source/Runtime/Renderer/Private/SceneRendering.cpp b/Engine/Source/Runtime/Renderer/Private/SceneRendering.cpp
--- a/Engine/Source/Runtime/Renderer/Private/SceneRendering.cpp
+++ b/Engine/Source/Runtime/Renderer/Private/SceneRendering.cpp
@@ -166,10 +166,12 @@ bool IsDynamicInstancingEnabled(ERHIFeatureLevel::Type FeatureLevel)
 		&& UseGPUScene(GMaxRHIShaderPlatform, FeatureLevel);
 }
 
+PRAGMA_DISABLE_DEPRECATION_WARNINGS
 int32 GetMaxNumReflectionCaptures(EShaderPlatform ShaderPlatform)
 {
 	return IsMobilePlatform(ShaderPlatform) ? GMobileMaxNumReflectionCaptures : GMaxNumReflectionCaptures;
 }
+PRAGMA_ENABLE_DEPRECATION_WARNINGS
 
 int32 GDumpInstancingStats = 0;
 FAutoConsoleVariableRef CVarDumpInstancingStats(
diff --git a/Engine/Source/Runtime/Renderer/Private/SceneRendering.h b/Engine/Source/Runtime/Renderer/Private/SceneRendering.h
--- a/Engine/Source/Runtime/Renderer/Private/SceneRendering.h
+++ b/Engine/Source/Runtime/Renderer/Private/SceneRendering.h
@@ -885,9 +885,11 @@ struct FMeshDecalBatch
 };
 
 // DX11 maximum 2d texture array size is D3D11_REQ_TEXTURE2D_ARRAY_AXIS_DIMENSION = 2048, and 2048/6 = 341.33.
+UE_DEPRECATED(5.4, "Using GMaxNumReflectionCaptures directly is deprecated. Please use GetMaxNumReflectionCaptures(EShaderPlatform) instead")
 static const int32 GMaxNumReflectionCaptures = 341;
 
 /** Per-reflection capture data needed by the shader. */
+PRAGMA_DISABLE_DEPRECATION_WARNINGS
 BEGIN_GLOBAL_SHADER_PARAMETER_STRUCT(FReflectionCaptureShaderData,)
 	SHADER_PARAMETER_ARRAY(FVector4f,PositionHighAndRadius,[GMaxNumReflectionCaptures])
 	// W is unused
@@ -899,8 +901,12 @@ BEGIN_GLOBAL_SHADER_PARAMETER_STRUCT(FReflectionCaptureShaderData,)
 	SHADER_PARAMETER_ARRAY(FMatrix44f,BoxTransform,[GMaxNumReflectionCaptures])
 	SHADER_PARAMETER_ARRAY(FVector4f,BoxScales,[GMaxNumReflectionCaptures])
 END_GLOBAL_SHADER_PARAMETER_STRUCT()
+PRAGMA_ENABLE_DEPRECATION_WARNINGS
 
+UE_DEPRECATED(5.4, "Using GMobileMaxNumReflectionCaptures directly is deprecated. Please use GetMaxNumReflectionCaptures(EShaderPlatform) instead")
 static const int32 GMobileMaxNumReflectionCaptures = 100;
+
+PRAGMA_DISABLE_DEPRECATION_WARNINGS
 BEGIN_GLOBAL_SHADER_PARAMETER_STRUCT(FMobileReflectionCaptureShaderData, )
 	SHADER_PARAMETER_ARRAY(FVector4f,PositionHighAndRadius,[GMobileMaxNumReflectionCaptures])
 	// W is unused
@@ -912,6 +918,7 @@ BEGIN_GLOBAL_SHADER_PARAMETER_STRUCT(FMobileReflectionCaptureShaderData, )
 	SHADER_PARAMETER_ARRAY(FMatrix44f,BoxTransform,[GMobileMaxNumReflectionCaptures])
 	SHADER_PARAMETER_ARRAY(FVector4f,BoxScales,[GMobileMaxNumReflectionCaptures])
 END_GLOBAL_SHADER_PARAMETER_STRUCT()
+PRAGMA_ENABLE_DEPRECATION_WARNINGS
 
 extern int32 GetMaxNumReflectionCaptures(EShaderPlatform ShaderPlatform);
 
