Fix groom holdout composition

diff --git a/Engine/Shaders/Private/HairStrands/HairStrandsComposition.usf b/Engine/Shaders/Private/HairStrands/HairStrandsComposition.usf
--- a/Engine/Shaders/Private/HairStrands/HairStrandsComposition.usf
+++ b/Engine/Shaders/Private/HairStrands/HairStrandsComposition.usf
@@ -214,7 +214,7 @@ void HoldoutPS(
 	}
 
 	OutColor.rgb = 0;
-	OutColor.a = 1;  //PixelCoverage;
+	OutColor.a = PixelCoverage;
 	OutDepth     = HairDeviceZ;
 	if (bComposeDofDepth)
 	{
diff --git a/Engine/Source/Runtime/Renderer/Private/HairStrands/HairStrandsComposition.cpp b/Engine/Source/Runtime/Renderer/Private/HairStrands/HairStrandsComposition.cpp
--- a/Engine/Source/Runtime/Renderer/Private/HairStrands/HairStrandsComposition.cpp
+++ b/Engine/Source/Runtime/Renderer/Private/HairStrands/HairStrandsComposition.cpp
@@ -27,6 +27,9 @@ static FAutoConsoleVariableRef CVarHairWriteGBufferData(TEXT("r.HairStrands.Writ
 static int32 GHairStrandsComposeDOFDepth = 1;
 static FAutoConsoleVariableRef CVarHairStrandsComposeDOFDepth(TEXT("r.HairStrands.DOFDepth"), GHairStrandsComposeDOFDepth, TEXT("Compose hair with DOF by lerping hair depth based on its opacity."));
 
+static int32 GHairStrandsHoldoutMode = 0;
+static FAutoConsoleVariableRef CVarHairStrandsHoldoutMode(TEXT("r.HairStrands.HoldoutMode"), GHairStrandsHoldoutMode, TEXT("Change how sample are merged when rendering with holdout."));
+
 /////////////////////////////////////////////////////////////////////////////////////////
 
 float GetHairFastResolveVelocityThreshold(const FIntPoint& Resolution)
@@ -54,6 +57,7 @@ void InternalCommonDrawPass(
 	const FIntPoint Resolution,
 	const EHairStrandsCommonPassType Type,
 	const bool bWriteDepth,
+	const bool bHasHoldout,
 	const FHairStrandsTiles& TileData,
 	TPixelShader& PixelShader,
 	TPassParameter* PassParamters)
@@ -71,7 +75,7 @@ void InternalCommonDrawPass(
 		Forward<FRDGEventName>(EventName),
 		PassParamters,
 		ERDGPassFlags::Raster,
-		[PassParamters, TileVertexShader, PixelShader, Viewport, Resolution, Type, bWriteDepth, TileType](FRHICommandList& RHICmdList)
+		[PassParamters, TileVertexShader, PixelShader, Viewport, Resolution, Type, bWriteDepth, bHasHoldout, TileType](FRHICommandList& RHICmdList)
 	{
 		FHairStrandsTilePassVS::FParameters ParametersVS = PassParamters->TileData;
 
@@ -80,15 +84,33 @@ void InternalCommonDrawPass(
 
 		if (Type == EHairStrandsCommonPassType::Composition)
 		{
-			// Alpha usage/output is controlled with r.PostProcessing.PropagateAlpha. The value are:
-			// 0: disabled(default);
-			// 1: enabled in linear color space;
-			// 2: same as 1, but also enable it through the tonemapper.
-			//
-			// When enable (PorpagateAlpha is set to 1 or 2), the alpha value means:
-			// 0: valid pixel
-			// 1: invalid pixel (background)
-			GraphicsPSOInit.BlendState = TStaticBlendState<CW_RGBA, BO_Add, BF_One, BF_InverseSourceAlpha, BO_Add, BF_Zero, BF_InverseSourceAlpha>::GetRHI();
+			if (bHasHoldout)
+			{
+				// EHairStrandsCommonPassType::Composition and bHasHoldout, only color is composed. The alpha value is composed into a separate pass.
+				// Two modes:
+				// * 0: the background pixels are *not* weighted by the hair coverage
+				// * 1: the background pixels are       weighted by the hair coverage
+				if (GHairStrandsHoldoutMode == 0)
+				{
+					GraphicsPSOInit.BlendState = TStaticBlendState<CW_RGB, BO_Add, BF_One, BF_One, BO_Add, BF_Zero, BF_One>::GetRHI();
+				}
+				else
+				{
+					GraphicsPSOInit.BlendState = TStaticBlendState<CW_RGB, BO_Add, BF_One, BF_InverseSourceAlpha, BO_Add, BF_Zero, BF_One>::GetRHI();
+				}
+			}
+			else
+			{
+				// Alpha usage/output is controlled with r.PostProcessing.PropagateAlpha. The value are:
+				// 0: disabled(default);
+				// 1: enabled in linear color space;
+				// 2: same as 1, but also enable it through the tonemapper.
+				//
+				// When enable (PorpagateAlpha is set to 1 or 2), the alpha value means:
+				// 0: valid pixel
+				// 1: invalid pixel (background)
+				GraphicsPSOInit.BlendState = TStaticBlendState<CW_RGBA, BO_Add, BF_One, BF_InverseSourceAlpha, BO_Add, BF_Zero, BF_InverseSourceAlpha>::GetRHI();
+			}
 		}
 		else if (Type == EHairStrandsCommonPassType::Blit) //used to write data into the temporal input buffer
 		{
@@ -96,7 +118,7 @@ void InternalCommonDrawPass(
 		}
 		else if (Type == EHairStrandsCommonPassType::Holdout)
 		{
-			GraphicsPSOInit.BlendState = TStaticBlendState<CW_ALPHA, BO_Add, BF_Zero, BF_One, BO_Max, BF_One, BF_One>::GetRHI();
+			GraphicsPSOInit.BlendState = TStaticBlendState<CW_ALPHA, BO_Add, BF_Zero, BF_One, BO_Add, BF_One, BF_InverseSourceAlpha>::GetRHI();
 		}
 		else
 		{
@@ -208,7 +230,8 @@ static void AddHairHoldoutPass(
 		View,
 		OutColorTexture->Desc.Extent,
 		EHairStrandsCommonPassType::Holdout,
-		false,
+		false /*bWriteDepth*/,
+		true /*bHasHoldout*/,
 		VisibilityData.TileData,
 		PixelShader,
 		Parameters);
@@ -286,7 +309,8 @@ static void AddHairVisibilityComposeSamplePass(
 		View,
 		OutColorTexture->Desc.Extent,
 		EHairStrandsCommonPassType::Composition,
-		false,
+		false /*bWriteDepth*/,
+		bHasHoldout,
 		VisibilityData.TileData,
 		PixelShader,
 		Parameters);
@@ -353,7 +377,8 @@ static FRDGTextureRef AddHairDOFDepthPass(
 		View,
 		OutputResolution,
 		EHairStrandsCommonPassType::DOF,
-		false,
+		false /*bWriteDepth*/,
+		false /*bHasHoldout*/,
 		VisibilityData.TileData,
 		PixelShader,
 		Parameters);
@@ -417,7 +442,8 @@ static void AddHairVisibilityFastResolveMaskPass(
 		View,
 		Resolution,
 		EHairStrandsCommonPassType::TAAFastResolve,
-		false,
+		false /*bWriteDepth*/,
+		false /*bHasHoldout*/,
 		TileData,
 		PixelShader,
 		Parameters);
@@ -524,6 +550,7 @@ static void AddHairVisibilityGBufferWritePass(
 		OutGBufferATexture->Desc.Extent,
 		EHairStrandsCommonPassType::GBuffer,
 		bWriteDepth,
+		false /*bHasHoldout*/,
 		TileData,
 		PixelShader,
 		Parameters);
