Removed unused ERayTracingInstanceMaskType::SceneCapture and RAY_TRACING_MASK_SCENE_CAPTURE.

diff --git a/Engine/Shaders/Shared/RayTracingDefinitions.h b/Engine/Shaders/Shared/RayTracingDefinitions.h
--- a/Engine/Shaders/Shared/RayTracingDefinitions.h
+++ b/Engine/Shaders/Shared/RayTracingDefinitions.h
@@ -17,7 +17,6 @@
 #define RAY_TRACING_MASK_TRANSLUCENT				0x02    // Opaque and alpha tested meshes and particles (e.g. used by translucency tracing pass)
 #define RAY_TRACING_MASK_THIN_SHADOW				0x04    // Whether the thin geometry (e.g. hair) is visible for shadow rays
 #define RAY_TRACING_MASK_SHADOW						0x08    // Whether the geometry is visible for shadow rays
-#define RAY_TRACING_MASK_SCENE_CAPTURE				0x10    // Whether the geometry is visible when rendering scene capture views
 #define RAY_TRACING_MASK_FAR_FIELD					0x20
 #define RAY_TRACING_MASK_HAIR_STRANDS               0x40    // For primary ray tracing against hair
 #define RAY_TRACING_MASK_ALL						0xFF
diff --git a/Engine/Source/Runtime/Renderer/Private/RayTracing/RayTracingInstanceMask.cpp b/Engine/Source/Runtime/Renderer/Private/RayTracing/RayTracingInstanceMask.cpp
--- a/Engine/Source/Runtime/Renderer/Private/RayTracing/RayTracingInstanceMask.cpp
+++ b/Engine/Source/Runtime/Renderer/Private/RayTracing/RayTracingInstanceMask.cpp
@@ -38,9 +38,6 @@ uint8 ComputeRayTracingInstanceMask(ERayTracingInstanceMaskType MaskType, ERayTr
 		case ERayTracingInstanceMaskType::HairStrands:
 			Mask = RAY_TRACING_MASK_HAIR_STRANDS;
 			break;
-		case ERayTracingInstanceMaskType::SceneCapture:
-			Mask = RAY_TRACING_MASK_SCENE_CAPTURE;
-			break;
 		case ERayTracingInstanceMaskType::VisibleInPrimaryRay:
 			Mask = 0;
 			break;
@@ -75,9 +72,6 @@ uint8 ComputeRayTracingInstanceMask(ERayTracingInstanceMaskType MaskType, ERayTr
 		case ERayTracingInstanceMaskType::HairStrands:
 			Mask = PATHTRACER_MASK_HAIR_CAMERA | PATHTRACER_MASK_HAIR_INDIRECT;
 			break;
-		case ERayTracingInstanceMaskType::SceneCapture:
-			Mask = PATHTRACER_MASK_IGNORE;
-			break;
 		case ERayTracingInstanceMaskType::VisibleInPrimaryRay:
 			Mask = PATHTRACER_MASK_CAMERA | PATHTRACER_MASK_HAIR_CAMERA;
 			break;
diff --git a/Engine/Source/Runtime/Renderer/Private/RayTracing/RayTracingInstanceMask.h b/Engine/Source/Runtime/Renderer/Private/RayTracing/RayTracingInstanceMask.h
--- a/Engine/Source/Runtime/Renderer/Private/RayTracing/RayTracingInstanceMask.h
+++ b/Engine/Source/Runtime/Renderer/Private/RayTracing/RayTracingInstanceMask.h
@@ -19,9 +19,6 @@ enum class ERayTracingInstanceMaskType : uint8
 	FarField,
 	HairStrands,
 
-	// ray tracing specific mask type
-	SceneCapture,
-
 	// path tracing specific mask type
 	VisibleInPrimaryRay,
 	VisibleInIndirectRay
