Cleaned up SVT Viewer properties to be more in line with the HV component. SVT Viewer now applies frame transforms by default.

#rnx

diff --git a/Engine/Source/Runtime/Renderer/Private/SparseVolumeTexture/SparseVolumeTextureViewerComponent.cpp b/Engine/Source/Runtime/Renderer/Private/SparseVolumeTexture/SparseVolumeTextureViewerComponent.cpp
--- a/Engine/Source/Runtime/Renderer/Private/SparseVolumeTexture/SparseVolumeTextureViewerComponent.cpp
+++ b/Engine/Source/Runtime/Renderer/Private/SparseVolumeTexture/SparseVolumeTextureViewerComponent.cpp
@@ -29,9 +29,18 @@ constexpr int32 SVTViewerDefaultVolumeResolution = 128;
 USparseVolumeTextureViewerComponent::USparseVolumeTextureViewerComponent(const FObjectInitializer& ObjectInitializer)
 	: Super(ObjectInitializer)
 	, SparseVolumeTexturePreview(nullptr)
-	, bAnimate(false)
-	, AnimationFrame(0.0f)
+	, Frame(0.0f)
+	, FrameRate(24.0f)
+	, bPlaying(false)
+	, bLooping(false)
+	, bReversePlayback(false)
+	, bBlockingStreamingRequests(false)
+	, bApplyPerFrameTransforms(true)
+	, bPivotAtCentroid(false)
+	, VoxelSize(1.0f)
 	, PreviewAttribute(ESVTPA_AttributesA_R)
+	, MipLevel(0)
+	, Extinction(0.025f)
 	, SparseVolumeTextureViewerSceneProxy(nullptr)
 {
 	PrimaryComponentTick.bCanEverTick = true;
@@ -53,13 +62,22 @@ void USparseVolumeTextureViewerComponent::PostEditChangeProperty(FPropertyChange
 {
 	Super::PostEditChangeProperty(PropertyChangedEvent);
 
-	if (!bAnimate && SparseVolumeTexturePreview)
+	FName PropertyName;
+	if (PropertyChangedEvent.Property)
 	{
-		FrameIndex = int32(AnimationFrame * float(SparseVolumeTexturePreview->GetNumFrames()));
+		PropertyName = PropertyChangedEvent.Property->GetFName();
 	}
-	MarkRenderStateDirty();
 
-	SendRenderTransformCommand();
+	if (PropertyName == GET_MEMBER_NAME_CHECKED(USparseVolumeTextureViewerComponent, Frame))
+	{
+		if (SparseVolumeTexturePreview)
+		{
+			const int32 FrameCount = SparseVolumeTexturePreview->GetNumFrames();
+			Frame = FMath::Clamp(Frame, 0, FrameCount - 1);
+		}
+	}
+
+	MarkRenderStateDirty();
 }
 
 #endif // WITH_EDITOR
@@ -78,19 +96,26 @@ FBoxSphereBounds USparseVolumeTextureViewerComponent::CalcBounds(const FTransfor
 	}
 	
 	FBoxSphereBounds NewBounds;
-	if (bLocalOriginAtCorner)
+	FVector HalfVolumeResolution = FVector(VolumeResolution) * 0.5;
+	if (bPivotAtCentroid)
 	{
-		NewBounds.Origin = -VolumeResolution * 0.5;
-		NewBounds.BoxExtent = VolumeResolution;
+		NewBounds.Origin = FVector::ZeroVector;
 	}
 	else
 	{
-		NewBounds.Origin = FVector::ZeroVector;
-		NewBounds.BoxExtent = VolumeResolution * 0.5;
+		NewBounds.Origin = HalfVolumeResolution;
 	}
+	NewBounds.BoxExtent = HalfVolumeResolution;
 	NewBounds.SphereRadius = NewBounds.BoxExtent.Length();
 
-	return NewBounds.TransformBy(LocalToWorld);
+	if (SparseVolumeTextureFrame)
+	{
+		return NewBounds.TransformBy(SparseVolumeTextureFrame->GetFrameTransform() * LocalToWorld);
+	}
+	else
+	{
+		return NewBounds.TransformBy(LocalToWorld);
+	}
 }
 
 void USparseVolumeTextureViewerComponent::CreateRenderState_Concurrent(FRegisterComponentContext* Context)
@@ -111,7 +136,7 @@ void USparseVolumeTextureViewerComponent::CreateRenderState_Concurrent(FRegister
 		ShouldComponentAddToScene() && ShouldRender() && IsRegistered() && (GetOuter() == NULL || !GetOuter()->HasAnyFlags(RF_ClassDefaultObject)))
 	{
 		// Create the scene proxy.
-		SparseVolumeTextureViewerSceneProxy = new FSparseVolumeTextureViewerSceneProxy(this, FrameIndex);
+		SparseVolumeTextureViewerSceneProxy = new FSparseVolumeTextureViewerSceneProxy(this);
 		GetWorld()->Scene->AddSparseVolumeTextureViewer(SparseVolumeTextureViewerSceneProxy);
 		SendRenderTransformCommand();
 	}
@@ -163,7 +188,7 @@ void USparseVolumeTextureViewerComponent::SendRenderTransformCommand()
 
 		FSparseVolumeTextureViewerSceneProxy* SVTSceneProxy = SparseVolumeTextureViewerSceneProxy;
 		ENQUEUE_RENDER_COMMAND(FUpdateSparseVolumeTextureViewerProxyTransformCommand)(
-		[SVTSceneProxy, GlobalTransform, FrameTransform, VolumeRes3f, CompIdx = (uint32)PreviewAttribute, Ext = Extinction, Mip = MipLevel, VoxelSizeFactor = VoxelSize, bPivotAtCorner = (bool)bLocalOriginAtCorner](FRHICommandList& RHICmdList)
+		[SVTSceneProxy, GlobalTransform, FrameTransform, VolumeRes3f, CompIdx = (uint32)PreviewAttribute, Ext = Extinction, Mip = MipLevel, VoxelSizeFactor = VoxelSize, bCentroidPivot = (bool)bPivotAtCentroid](FRHICommandList& RHICmdList)
 		{
 			SVTSceneProxy->GlobalTransform = GlobalTransform;
 			SVTSceneProxy->FrameTransform = FrameTransform;
@@ -172,7 +197,7 @@ void USparseVolumeTextureViewerComponent::SendRenderTransformCommand()
 			SVTSceneProxy->ComponentToVisualize = CompIdx;
 			SVTSceneProxy->Extinction = Ext;
 			SVTSceneProxy->VoxelSizeFactor = VoxelSizeFactor;
-			SVTSceneProxy->bPivotAtCorner = bPivotAtCorner;
+			SVTSceneProxy->bPivotAtCentroid = bCentroidPivot;
 		});
 	}
 }
@@ -182,20 +207,23 @@ void USparseVolumeTextureViewerComponent::TickComponent(float DeltaTime, enum EL
 	if (SparseVolumeTexturePreview)
 	{
 		const int32 NumFrames = SparseVolumeTexturePreview->GetNumFrames();
-		float FrameIndexF = 0.0f;
-		if (bAnimate)
+
+		if (bPlaying)
 		{
-			const float AnimationDuration = (NumFrames / (FrameRate + UE_SMALL_NUMBER)) + UE_SMALL_NUMBER;
-			AnimationTime = FMath::Fmod(AnimationTime + AnimationDuration + (bReversePlayback ? -DeltaTime : DeltaTime), AnimationDuration);
-			FrameIndexF = AnimationTime * FrameRate;
+			Frame += DeltaTime * FrameRate * (bReversePlayback ? -1.0f : 1.0f);
+		}
+
+		if (bLooping)
+		{
+			// Simple way of dealing with looping when playing back in reverse: add NumFrames to the X input of Fmod.
+			Frame = FMath::Fmod(Frame + NumFrames, (float)NumFrames);
 		}
 		else
 		{
-			FrameIndexF = AnimationFrame * float(NumFrames);
+			Frame = FMath::Clamp(Frame, 0.0f, (float)(NumFrames - 1));
 		}
-		FrameIndexF = FMath::Clamp(FrameIndexF, 0, static_cast<float>(NumFrames - 1));
-		FrameIndex = FrameIndexF;
-		SparseVolumeTextureFrame = USparseVolumeTextureFrame::GetFrameAndIssueStreamingRequest(SparseVolumeTexturePreview, FrameIndexF, MipLevel, bBlockingStreamingRequests);
+
+		SparseVolumeTextureFrame = USparseVolumeTextureFrame::GetFrameAndIssueStreamingRequest(SparseVolumeTexturePreview, Frame, MipLevel, bBlockingStreamingRequests);
 	}
 	else
 	{
diff --git a/Engine/Source/Runtime/Renderer/Private/SparseVolumeTexture/SparseVolumeTextureViewerComponent.h b/Engine/Source/Runtime/Renderer/Private/SparseVolumeTexture/SparseVolumeTextureViewerComponent.h
--- a/Engine/Source/Runtime/Renderer/Private/SparseVolumeTexture/SparseVolumeTextureViewerComponent.h
+++ b/Engine/Source/Runtime/Renderer/Private/SparseVolumeTexture/SparseVolumeTextureViewerComponent.h
@@ -41,10 +41,19 @@ class USparseVolumeTextureViewerComponent : public UPrimitiveComponent
 	UPROPERTY(EditAnywhere, BlueprintReadWrite, Category= "Asset Preview")
 	TObjectPtr<class USparseVolumeTexture> SparseVolumeTexturePreview;
 
-	UPROPERTY(EditAnywhere, BlueprintReadOnly, Category = "Asset Preview")
-	uint32 bAnimate : 1;
+	UPROPERTY(EditAnywhere, BlueprintReadWrite, Category = Animation, meta = (EditCondition = "bPlaying == false"))
+	float Frame;
 
-	UPROPERTY(EditAnywhere, BlueprintReadOnly, Category = "Asset Preview")
+	UPROPERTY(EditAnywhere, BlueprintReadWrite, Category = Animation)
+	float FrameRate = 24.0f;
+
+	UPROPERTY(EditAnywhere, BlueprintReadWrite, Category = Animation)
+	uint32 bPlaying : 1;
+
+	UPROPERTY(EditAnywhere, BlueprintReadWrite, Category = Animation)
+	uint32 bLooping : 1;
+
+	UPROPERTY(EditAnywhere, BlueprintReadOnly, Category = Animation)
 	uint32 bReversePlayback : 1;
 
 	UPROPERTY(EditAnywhere, BlueprintReadOnly, Category = "Asset Preview")
@@ -54,20 +63,11 @@ class USparseVolumeTextureViewerComponent : public UPrimitiveComponent
 	uint32 bApplyPerFrameTransforms : 1;
 
 	UPROPERTY(EditAnywhere, BlueprintReadOnly, Category = "Asset Preview")
-	uint32 bLocalOriginAtCorner : 1;
+	uint32 bPivotAtCentroid : 1;
 
 	UPROPERTY(EditAnywhere, BlueprintReadWrite, Category = "Asset Preview")
 	float VoxelSize = 1.0f;
 
-	UPROPERTY(EditAnywhere, BlueprintReadWrite, Category = "Asset Preview", meta = (UIMin = 0.0, UIMax = 1.0, ClampMin = 0.0, ClampMax = 1.0, EditCondition = "!bAnimate"))
-	float AnimationFrame;
-
-	UPROPERTY(EditAnywhere, BlueprintReadWrite, Category = "Asset Preview", meta = (UIMin = 0.0, UIMax = 120.0, ClampMin = 0.0, ClampMax = 120.0, EditCondition = "bAnimate"))
-	float FrameRate = 24.0f;
-
-	UPROPERTY(VisibleAnywhere, Category = "Asset Preview", meta = (UIMin = 0.0, UIMax = 60.0))
-	float AnimationTime;
-
 	UPROPERTY(EditAnywhere, BlueprintReadWrite, Category = "Asset Preview")
 	TEnumAsByte<ESparseVolumeTexturePreviewAttribute> PreviewAttribute;
 
@@ -113,9 +113,6 @@ public:
 	//~ End UActorComponent Interface.
 
 private:
-	
-	int32 FrameIndex = 0;
-
 	void SendRenderTransformCommand();
 
 	FSparseVolumeTextureViewerSceneProxy* SparseVolumeTextureViewerSceneProxy;
diff --git a/Engine/Source/Runtime/Renderer/Private/SparseVolumeTexture/SparseVolumeTextureViewerRendering.cpp b/Engine/Source/Runtime/Renderer/Private/SparseVolumeTexture/SparseVolumeTextureViewerRendering.cpp
--- a/Engine/Source/Runtime/Renderer/Private/SparseVolumeTexture/SparseVolumeTextureViewerRendering.cpp
+++ b/Engine/Source/Runtime/Renderer/Private/SparseVolumeTexture/SparseVolumeTextureViewerRendering.cpp
@@ -136,7 +136,7 @@ void AddSparseVolumeTextureViewerRenderPass(FRDGBuilder& GraphBuilder, FSceneRen
 			UnitSpaceMat.SetColumn(0, FVector(RcpVolumeRes.X * 2.0, 0.0, 0.0));
 			UnitSpaceMat.SetColumn(1, FVector(0.0, RcpVolumeRes.Y * 2.0, 0.0));
 			UnitSpaceMat.SetColumn(2, FVector(0.0, 0.0, RcpVolumeRes.Z * 2.0));
-			UnitSpaceMat.SetOrigin(FVector(SVTSceneProxy->bPivotAtCorner ? -1.0 : 0.0));
+			UnitSpaceMat.SetOrigin(FVector(SVTSceneProxy->bPivotAtCentroid ? 0.0  : -1.0));
 
 			const FMatrix44f WorldToLocal = FMatrix44f(InvTransformMat * UnitSpaceMat);
 			const FMatrix44f WorldToLocalRot = FMatrix44f(InvRotation.ToMatrix());
diff --git a/Engine/Source/Runtime/Renderer/Private/SparseVolumeTexture/SparseVolumeTextureViewerSceneProxy.cpp b/Engine/Source/Runtime/Renderer/Private/SparseVolumeTexture/SparseVolumeTextureViewerSceneProxy.cpp
--- a/Engine/Source/Runtime/Renderer/Private/SparseVolumeTexture/SparseVolumeTextureViewerSceneProxy.cpp
+++ b/Engine/Source/Runtime/Renderer/Private/SparseVolumeTexture/SparseVolumeTextureViewerSceneProxy.cpp
@@ -36,7 +36,7 @@ void FScene::RemoveSparseVolumeTextureViewer(FSparseVolumeTextureViewerSceneProx
 
 ////////////////////////////////////////////////////////////////////////////////////////////////
 
-FSparseVolumeTextureViewerSceneProxy::FSparseVolumeTextureViewerSceneProxy(const USparseVolumeTextureViewerComponent* InComponent, int32 FrameIndex, FName ResourceName)
+FSparseVolumeTextureViewerSceneProxy::FSparseVolumeTextureViewerSceneProxy(const USparseVolumeTextureViewerComponent* InComponent, FName ResourceName)
 	: FPrimitiveSceneProxy((UPrimitiveComponent*)InComponent, ResourceName)
 	, TextureRenderResources(nullptr)
 {
diff --git a/Engine/Source/Runtime/Renderer/Private/SparseVolumeTexture/SparseVolumeTextureViewerSceneProxy.h b/Engine/Source/Runtime/Renderer/Private/SparseVolumeTexture/SparseVolumeTextureViewerSceneProxy.h
--- a/Engine/Source/Runtime/Renderer/Private/SparseVolumeTexture/SparseVolumeTextureViewerSceneProxy.h
+++ b/Engine/Source/Runtime/Renderer/Private/SparseVolumeTexture/SparseVolumeTextureViewerSceneProxy.h
@@ -21,7 +21,7 @@ namespace UE
 class FSparseVolumeTextureViewerSceneProxy : public FPrimitiveSceneProxy
 {
 public:
-	FSparseVolumeTextureViewerSceneProxy(const USparseVolumeTextureViewerComponent* InComponent, int32 FrameIndex, FName ResourceName = NAME_None);
+	FSparseVolumeTextureViewerSceneProxy(const USparseVolumeTextureViewerComponent* InComponent, FName ResourceName = NAME_None);
 	virtual ~FSparseVolumeTextureViewerSceneProxy() = default;
 
 	const UE::SVT::FTextureRenderResources* TextureRenderResources;
@@ -32,7 +32,7 @@ public:
 	uint32 ComponentToVisualize;
 	float Extinction;
 	float VoxelSizeFactor;
-	bool bPivotAtCorner;
+	bool bPivotAtCentroid;
 
 protected:
 
