Refactored how raytracing instance mask and flags is calculated for (non cached) static instances.
- Apply ERayTracingInstanceMaskType::Shadow in SetupRayTracingMeshCommandMaskAndStatus(...) when building individual mesh commands.
- In "GatherRayTracingRelevantPrimitives" simply combine instance masks from the various mesh commands and override instance mask for primitives assigned to far field.

diff --git a/Engine/Source/Runtime/Renderer/Private/RayTracing/RayTracing.cpp b/Engine/Source/Runtime/Renderer/Private/RayTracing/RayTracing.cpp
--- a/Engine/Source/Runtime/Renderer/Private/RayTracing/RayTracing.cpp
+++ b/Engine/Source/Runtime/Renderer/Private/RayTracing/RayTracing.cpp
@@ -182,20 +182,12 @@ namespace RayTracing
 			return Key ^ reinterpret_cast<uint64>(RayTracingGeometryRHI);
 		}
 
-		void UpdateMasks(const ERayTracingPrimitiveFlags Flags, ERayTracingViewMaskMode MaskMode)
+		void FinalizeInstanceMask(const ERayTracingPrimitiveFlags Flags, ERayTracingViewMaskMode MaskMode)
 		{
-			FRayTracingMeshCommand Command;
-			Command.InstanceMask = InstanceMask;
-			Command.bOpaque = bAllSegmentsOpaque;
-			Command.bCastRayTracedShadows = bAnySegmentsCastShadow;
-			Command.bDecal = bAnySegmentsDecal;
-			Command.bTwoSided = bTwoSided;
-			Command.bIsSky = bIsSky;
-			Command.bIsTranslucent = bAllSegmentsTranslucent;
-
-			UpdateRayTracingMeshCommandMasks(Command, Flags, MaskMode);
-
-			InstanceMask = Command.InstanceMask;
+			if (EnumHasAllFlags(Flags, ERayTracingPrimitiveFlags::FarField))
+			{
+				InstanceMask = ComputeRayTracingInstanceMask(ERayTracingInstanceMaskType::FarField, MaskMode);
+			}
 		}
 	};
 
@@ -568,6 +560,8 @@ namespace RayTracing
 								RelevantPrimitive->CachedRayTracingMeshCommandIndices = SceneInfo->CachedRayTracingMeshCommandIndicesPerLOD[LODIndex];
 								RelevantPrimitive->StateHash = SceneInfo->CachedRayTracingMeshCommandsHashPerLOD[LODIndex];
 
+								const ERayTracingViewMaskMode MaskMode = static_cast<ERayTracingViewMaskMode>(Scene.CachedRayTracingMeshCommandsMode);
+
 								// TODO: Cache these flags to avoid having to loop over the RayTracingMeshCommands
 								for (int32 CommandIndex : RelevantPrimitive->CachedRayTracingMeshCommandIndices)
 								{
@@ -592,9 +586,7 @@ namespace RayTracing
 									}
 								}
 
-								ERayTracingViewMaskMode MaskMode = static_cast<ERayTracingViewMaskMode>(Scene.CachedRayTracingMeshCommandsMode);
-
-								RelevantPrimitive->UpdateMasks(Flags, MaskMode);
+								RelevantPrimitive->FinalizeInstanceMask(Flags, MaskMode);
 							}
 						}
 					});
diff --git a/Engine/Source/Runtime/Renderer/Private/RayTracing/RayTracingInstanceMask.cpp b/Engine/Source/Runtime/Renderer/Private/RayTracing/RayTracingInstanceMask.cpp
--- a/Engine/Source/Runtime/Renderer/Private/RayTracing/RayTracingInstanceMask.cpp
+++ b/Engine/Source/Runtime/Renderer/Private/RayTracing/RayTracingInstanceMask.cpp
@@ -249,6 +249,8 @@ void SetupRayTracingMeshCommandMaskAndStatus(FRayTracingMeshCommand& MeshCommand
 
 	MeshCommand.InstanceMask = BlendModeToRayTracingInstanceMask(MaterialResource.GetBlendMode(), MaskMode);
 
+	MeshCommand.InstanceMask |= MeshCommand.bCastRayTracedShadows ? ComputeRayTracingInstanceMask(ERayTracingInstanceMaskType::Shadow, MaskMode) : 0;
+
 	if (!PrimitiveSceneProxy)
 	{
 		return;
@@ -256,6 +258,7 @@ void SetupRayTracingMeshCommandMaskAndStatus(FRayTracingMeshCommand& MeshCommand
 
 	FSceneProxyRayTracingMaskInfo MaskInfo = GetSceneProxyRayTracingMaskInfo(*PrimitiveSceneProxy, nullptr);
 
+	// TODO: This should be done once all mesh commands for a mesh are combined (similar to BuildRayTracingInstanceMaskAndFlags(...) above)
 	if (MaskMode == ERayTracingViewMaskMode::PathTracing || MaskMode == ERayTracingViewMaskMode::LightMapTracing)
 	{
 		if (!MaskInfo.bAffectsDynamicIndirectLighting)
@@ -285,16 +288,4 @@ void SetupRayTracingMeshCommandMaskAndStatus(FRayTracingMeshCommand& MeshCommand
 	}
 }
 
-void UpdateRayTracingMeshCommandMasks(FRayTracingMeshCommand& RayTracingCommand, const ERayTracingPrimitiveFlags Flags, ERayTracingViewMaskMode MaskMode)
-{
-	RayTracingCommand.InstanceMask |= RayTracingCommand.bCastRayTracedShadows ? ComputeRayTracingInstanceMask(ERayTracingInstanceMaskType::Shadow, MaskMode) : 0;
-
-	if (EnumHasAllFlags(Flags, ERayTracingPrimitiveFlags::FarField))
-	{
-		RayTracingCommand.InstanceMask = ComputeRayTracingInstanceMask(ERayTracingInstanceMaskType::FarField, MaskMode);
-	}
-}
-
-
-
 #endif
\ No newline at end of file
diff --git a/Engine/Source/Runtime/Renderer/Private/RayTracing/RayTracingInstanceMask.h b/Engine/Source/Runtime/Renderer/Private/RayTracing/RayTracingInstanceMask.h
--- a/Engine/Source/Runtime/Renderer/Private/RayTracing/RayTracingInstanceMask.h
+++ b/Engine/Source/Runtime/Renderer/Private/RayTracing/RayTracingInstanceMask.h
@@ -72,7 +72,4 @@ class FMaterial;
 void SetupRayTracingMeshCommandMaskAndStatus(FRayTracingMeshCommand& MeshCommand, const FMeshBatch& MeshBatch, const FPrimitiveSceneProxy* PrimitiveSceneProxy,
 	const FMaterial& MaterialResource, ERayTracingViewMaskMode MaskMode);
 
-void UpdateRayTracingMeshCommandMasks(FRayTracingMeshCommand& RayTracingMeshCommands, const ERayTracingPrimitiveFlags Flags, ERayTracingViewMaskMode MaskMode);
-
-
 #endif
\ No newline at end of file
