Split Shadow flag in Ray Tracing Instance Mask into Opaque and Translucent Shadow flags.
- makes it possible for shaders skip unecessary AHS evaluations when translucent shadows are not used.
- refactored logic that sets shadow flag so that the appropriate flag is set inside BlendModeToRayTracingInstanceMask(...)

diff --git a/Engine/Plugins/Experimental/GPULightmass/Source/GPULightmass/Private/LightmapRenderer.cpp b/Engine/Plugins/Experimental/GPULightmass/Source/GPULightmass/Private/LightmapRenderer.cpp
--- a/Engine/Plugins/Experimental/GPULightmass/Source/GPULightmass/Private/LightmapRenderer.cpp
+++ b/Engine/Plugins/Experimental/GPULightmass/Source/GPULightmass/Private/LightmapRenderer.cpp
@@ -36,8 +36,7 @@
 #include "PrimitiveUniformShaderParametersBuilder.h"
 #include "PrimitiveSceneShaderData.h"
 
-RENDERER_API uint8 BlendModeToRayTracingInstanceMask(const EBlendMode BlendMode, ERayTracingViewMaskMode MaskMode);
-RENDERER_API uint8 ComputeRayTracingInstanceShadowMask(ERayTracingViewMaskMode MaskMode);
+RENDERER_API uint8 BlendModeToRayTracingInstanceMask(const EBlendMode BlendMode, bool bCastShadow, ERayTracingViewMaskMode MaskMode);
 
 class FCopyConvergedLightmapTilesCS : public FGlobalShader
 {
@@ -628,20 +627,23 @@ struct FRayTracingMaskAndStatus
 
 	void UpdateInstanceMaskAndStatus(ERHIFeatureLevel::Type FeatureLevel, ERayTracingViewMaskMode MaskMode, TArray<FMeshBatch>& MeshBatches)
 	{
-
 		for (int32 SegmentIndex = 0; SegmentIndex < MeshBatches.Num(); SegmentIndex++)
 		{
-			const FMaterial& Material = MeshBatches[SegmentIndex].MaterialRenderProxy->GetIncompleteMaterialWithFallback(FeatureLevel);
+			const FMeshBatch& MeshBatch = MeshBatches[SegmentIndex];
+
+			const FMaterial& Material = MeshBatch.MaterialRenderProxy->GetIncompleteMaterialWithFallback(FeatureLevel);
+			const EBlendMode BlendMode = Material.GetBlendMode();
+
+			const bool bSegmentCastsShadow = MeshBatch.CastRayTracedShadow && Material.CastsRayTracedShadows(); // TODO: && BlendMode != BLEND_Additive;
 
-			bAllSegmentsUnlit &= Material.GetShadingModels().HasOnlyShadingModel(MSM_Unlit) || !MeshBatches[SegmentIndex].CastShadow;
+			bAllSegmentsUnlit &= Material.GetShadingModels().HasOnlyShadingModel(MSM_Unlit) || !MeshBatch.CastShadow;
 			bAllSegmentsOpaque &= Material.GetBlendMode() == EBlendMode::BLEND_Opaque;
-			bAnySegmentsCastShadow |= MeshBatches[SegmentIndex].CastRayTracedShadow && Material.CastsRayTracedShadows();
-			InstanceMask |= BlendModeToRayTracingInstanceMask(Material.GetBlendMode(), MaskMode);
+			bAnySegmentsCastShadow |= bSegmentCastsShadow;
+			InstanceMask |= BlendModeToRayTracingInstanceMask(Material.GetBlendMode(), bSegmentCastsShadow, MaskMode);
 		}
-
-		InstanceMask |= bAnySegmentsCastShadow ? ComputeRayTracingInstanceShadowMask(MaskMode) : 0;
 	}
 };
+
 void FCachedRayTracingSceneData::SetupFromSceneRenderState(FSceneRenderState& Scene)
 {
 #if RHI_RAYTRACING
diff --git a/Engine/Shaders/Shared/RayTracingDefinitions.h b/Engine/Shaders/Shared/RayTracingDefinitions.h
--- a/Engine/Shaders/Shared/RayTracingDefinitions.h
+++ b/Engine/Shaders/Shared/RayTracingDefinitions.h
@@ -15,12 +15,15 @@
 
 #define RAY_TRACING_MASK_OPAQUE						0x01    // Opaque and alpha tested meshes and particles (e.g. used by reflection, shadow, AO and GI tracing passes)
 #define RAY_TRACING_MASK_TRANSLUCENT				0x02    // Opaque and alpha tested meshes and particles (e.g. used by translucency tracing pass)
-#define RAY_TRACING_MASK_THIN_SHADOW				0x04    // Whether the thin geometry (e.g. hair) is visible for shadow rays
-#define RAY_TRACING_MASK_SHADOW						0x08    // Whether the geometry is visible for shadow rays
+#define RAY_TRACING_MASK_OPAQUE_SHADOW				0x04    // Opaque and alpha tested geometry visible for shadow rays
+#define RAY_TRACING_MASK_TRANSLUCENT_SHADOW			0x08    // Translucent geometry visible for shadow rays
+#define RAY_TRACING_MASK_THIN_SHADOW				0x10    // Whether the thin geometry (e.g. hair) is visible for shadow rays
 #define RAY_TRACING_MASK_FAR_FIELD					0x20
 #define RAY_TRACING_MASK_HAIR_STRANDS               0x40    // For primary ray tracing against hair
 #define RAY_TRACING_MASK_ALL						0xFF
 
+#define RAY_TRACING_MASK_SHADOW						(RAY_TRACING_MASK_OPAQUE_SHADOW | RAY_TRACING_MASK_TRANSLUCENT_SHADOW)
+
 #define RAY_TRACING_SHADER_SLOT_MATERIAL	0
 #define RAY_TRACING_SHADER_SLOT_SHADOW		1
 #define RAY_TRACING_NUM_SHADER_SLOTS		2
diff --git a/Engine/Source/Runtime/Renderer/Private/RayTracing/RayTracingInstanceMask.cpp b/Engine/Source/Runtime/Renderer/Private/RayTracing/RayTracingInstanceMask.cpp
--- a/Engine/Source/Runtime/Renderer/Private/RayTracing/RayTracingInstanceMask.cpp
+++ b/Engine/Source/Runtime/Renderer/Private/RayTracing/RayTracingInstanceMask.cpp
@@ -26,12 +26,15 @@ uint8 ComputeRayTracingInstanceMask(ERayTracingInstanceMaskType MaskType, ERayTr
 		case ERayTracingInstanceMaskType::Translucent:
 			Mask = RAY_TRACING_MASK_TRANSLUCENT;
 			break;
+		case ERayTracingInstanceMaskType::OpaqueShadow:
+			Mask = RAY_TRACING_MASK_OPAQUE_SHADOW;
+			break;
+		case ERayTracingInstanceMaskType::TranslucentShadow:
+			Mask = RAY_TRACING_MASK_TRANSLUCENT_SHADOW;
+			break;
 		case ERayTracingInstanceMaskType::ThinShadow:
 			Mask = RAY_TRACING_MASK_THIN_SHADOW;
 			break;
-		case ERayTracingInstanceMaskType::Shadow:
-			Mask = RAY_TRACING_MASK_SHADOW;
-			break;
 		case ERayTracingInstanceMaskType::FarField:
 			Mask = RAY_TRACING_MASK_FAR_FIELD;
 			break;
@@ -60,12 +63,15 @@ uint8 ComputeRayTracingInstanceMask(ERayTracingInstanceMaskType MaskType, ERayTr
 		case ERayTracingInstanceMaskType::Translucent:
 			Mask = PATHTRACER_MASK_CAMERA | PATHTRACER_MASK_INDIRECT;
 			break;
-		case ERayTracingInstanceMaskType::ThinShadow:
-			Mask = PATHTRACER_MASK_HAIR_SHADOW;
+		case ERayTracingInstanceMaskType::OpaqueShadow:
+			Mask = PATHTRACER_MASK_SHADOW;
 			break;
-		case ERayTracingInstanceMaskType::Shadow:
+		case ERayTracingInstanceMaskType::TranslucentShadow:
 			Mask = PATHTRACER_MASK_SHADOW;
 			break;
+		case ERayTracingInstanceMaskType::ThinShadow:
+			Mask = PATHTRACER_MASK_HAIR_SHADOW;
+			break;
 		case ERayTracingInstanceMaskType::FarField:
 			Mask = PATHTRACER_MASK_IGNORE;
 			break;
@@ -87,16 +93,22 @@ uint8 ComputeRayTracingInstanceMask(ERayTracingInstanceMaskType MaskType, ERayTr
 	return Mask;
 }
 
-uint8 ComputeRayTracingInstanceShadowMask(ERayTracingViewMaskMode MaskMode)
+uint8 BlendModeToRayTracingInstanceMask(const EBlendMode BlendMode, bool bCastShadow, ERayTracingViewMaskMode MaskMode)
 {
-	return ComputeRayTracingInstanceMask(ERayTracingInstanceMaskType::Shadow, MaskMode);
-}
+	uint8 InstanceMask = 0;
 
+	if (IsOpaqueOrMaskedBlendMode(BlendMode))
+	{
+		InstanceMask |= ComputeRayTracingInstanceMask(ERayTracingInstanceMaskType::Opaque, MaskMode);
+		InstanceMask |= bCastShadow ? ComputeRayTracingInstanceMask(ERayTracingInstanceMaskType::OpaqueShadow, MaskMode) : 0;
+	}
+	else
+	{
+		InstanceMask |= ComputeRayTracingInstanceMask(ERayTracingInstanceMaskType::Translucent, MaskMode);
+		InstanceMask |= bCastShadow ? ComputeRayTracingInstanceMask(ERayTracingInstanceMaskType::TranslucentShadow, MaskMode) : 0;
+	}
 
-uint8 BlendModeToRayTracingInstanceMask(const EBlendMode BlendMode, ERayTracingViewMaskMode MaskMode)
-{
-	ERayTracingInstanceMaskType Type = IsOpaqueOrMaskedBlendMode(BlendMode) ? ERayTracingInstanceMaskType::Opaque : ERayTracingInstanceMaskType::Translucent;
-	return ComputeRayTracingInstanceMask(Type, MaskMode);
+	return InstanceMask;
 }
 
 FSceneProxyRayTracingMaskInfo GetSceneProxyRayTracingMaskInfo(const FPrimitiveSceneProxy& PrimitiveSceneProxy, const FSceneViewFamily* SceneViewFamily)
@@ -154,9 +166,10 @@ FRayTracingMaskAndFlags BuildRayTracingInstanceMaskAndFlags(TArrayView<const FMe
 		{
 			const FMaterial& Material = MeshBatch.MaterialRenderProxy->GetIncompleteMaterialWithFallback(FeatureLevel);
 			const EBlendMode BlendMode = Material.GetBlendMode();
-			Result.Mask |= BlendModeToRayTracingInstanceMask(BlendMode, MaskMode);
-			bAllSegmentsOpaque &= BlendMode == BLEND_Opaque;
 			const bool bSegmentCastsShadow = MeshBatch.CastRayTracedShadow && Material.CastsRayTracedShadows() && BlendMode != BLEND_Additive;
+
+			Result.Mask |= BlendModeToRayTracingInstanceMask(BlendMode, bSegmentCastsShadow, MaskMode);
+			bAllSegmentsOpaque &= BlendMode == BLEND_Opaque;
 			bAnySegmentsCastShadow |= bSegmentCastsShadow;
 			bAllSegmentsCastShadow &= bSegmentCastsShadow;
 			bAnySegmentsDecal |= Material.IsDeferredDecal();
@@ -170,15 +183,15 @@ FRayTracingMaskAndFlags BuildRayTracingInstanceMaskAndFlags(TArrayView<const FMe
 	Result.bAnySegmentsDecal = bAnySegmentsDecal;
 	Result.bAllSegmentsDecal = bAllSegmentsDecal;
 
-	Result.Mask |= bAnySegmentsCastShadow ? ComputeRayTracingInstanceMask(ERayTracingInstanceMaskType::Shadow, MaskMode) : 0;
-
 	const bool bIsHairStrands = Result.Mask & ComputeRayTracingInstanceMask(ERayTracingInstanceMaskType::HairStrands, MaskMode);
 	if (bIsHairStrands)
 	{
 		// For hair strands, opaque/translucent mask should be cleared to make sure geometry is only in the hair group. 
 		// If any segment receives shadow, it should receive only thin shadow instead of shadow.
 
-		Result.Mask &= ~(ComputeRayTracingInstanceMask(ERayTracingInstanceMaskType::Shadow, MaskMode) |
+		Result.Mask &= ~(
+			ComputeRayTracingInstanceMask(ERayTracingInstanceMaskType::OpaqueShadow, MaskMode) |
+			ComputeRayTracingInstanceMask(ERayTracingInstanceMaskType::TranslucentShadow, MaskMode) |
 			ComputeRayTracingInstanceMask(ERayTracingInstanceMaskType::ThinShadow, MaskMode) |
 			ComputeRayTracingInstanceMask(ERayTracingInstanceMaskType::Translucent, MaskMode) |
 			ComputeRayTracingInstanceMask(ERayTracingInstanceMaskType::Opaque, MaskMode));
@@ -247,9 +260,7 @@ void SetupRayTracingMeshCommandMaskAndStatus(FRayTracingMeshCommand& MeshCommand
 	MeshCommand.bTwoSided = MaterialResource.IsTwoSided();
 	MeshCommand.bIsTranslucent = MaterialResource.GetBlendMode() == EBlendMode::BLEND_Translucent;
 
-	MeshCommand.InstanceMask = BlendModeToRayTracingInstanceMask(MaterialResource.GetBlendMode(), MaskMode);
-
-	MeshCommand.InstanceMask |= MeshCommand.bCastRayTracedShadows ? ComputeRayTracingInstanceMask(ERayTracingInstanceMaskType::Shadow, MaskMode) : 0;
+	MeshCommand.InstanceMask = BlendModeToRayTracingInstanceMask(MaterialResource.GetBlendMode(), MeshCommand.bCastRayTracedShadows, MaskMode);
 
 	if (!PrimitiveSceneProxy)
 	{
diff --git a/Engine/Source/Runtime/Renderer/Private/RayTracing/RayTracingInstanceMask.h b/Engine/Source/Runtime/Renderer/Private/RayTracing/RayTracingInstanceMask.h
--- a/Engine/Source/Runtime/Renderer/Private/RayTracing/RayTracingInstanceMask.h
+++ b/Engine/Source/Runtime/Renderer/Private/RayTracing/RayTracingInstanceMask.h
@@ -14,8 +14,9 @@ enum class ERayTracingInstanceMaskType : uint8
 	// General mask type
 	Opaque,
 	Translucent,
+	OpaqueShadow,
+	TranslucentShadow,
 	ThinShadow,
-	Shadow,
 	FarField,
 	HairStrands,
 
@@ -33,11 +34,9 @@ enum class ERayTracingViewMaskMode : uint8
 };
 
 uint8 ComputeRayTracingInstanceMask(ERayTracingInstanceMaskType MaskType, ERayTracingViewMaskMode MaskMode);
-RENDERER_API uint8 ComputeRayTracingInstanceShadowMask(ERayTracingViewMaskMode MaskMode);
-
 
 /** Compute the mask based on blend mode for different ray tracing mode*/
-RENDERER_API uint8 BlendModeToRayTracingInstanceMask(const EBlendMode BlendMode, ERayTracingViewMaskMode MaskMode);
+RENDERER_API uint8 BlendModeToRayTracingInstanceMask(const EBlendMode BlendMode, bool bCastShadow, ERayTracingViewMaskMode MaskMode);
 
 
 /** Util struct and function to derive mask related info from scene proxy*/
