Path Tracer: Fix light functions on directional lights

Directional lights did not setup their light function index correctly, resulting in directional lights not being able to use their assigned light function material. Reported via UDN.

diff --git a/Engine/Source/Runtime/Renderer/Private/PathTracing.cpp b/Engine/Source/Runtime/Renderer/Private/PathTracing.cpp
--- a/Engine/Source/Runtime/Renderer/Private/PathTracing.cpp
+++ b/Engine/Source/Runtime/Renderer/Private/PathTracing.cpp
@@ -2074,6 +2074,8 @@ void SetLightParameters(FRDGBuilder& GraphBuilder, FPathTracingRG::FParameters*
 		}
 	}
 
+	const FRayTracingLightFunctionMap* RayTracingLightFunctionMap = GraphBuilder.Blackboard.Get<FRayTracingLightFunctionMap>();
+
 	// Add directional lights next (all lights with infinite bounds should come first)
 	if (View.Family->EngineShowFlags.DirectionalLights)
 	{
@@ -2105,6 +2107,15 @@ void SetLightParameters(FRDGBuilder& GraphBuilder, FPathTracingRG::FParameters*
 			DestLight.IESAtlasIndex = INDEX_NONE;
 			DestLight.MissShaderIndex = 0;
 
+			if (RayTracingLightFunctionMap)
+			{
+				const int32* LightFunctionIndex = RayTracingLightFunctionMap->Find(Light.LightSceneInfo);
+				if (LightFunctionIndex)
+				{
+					DestLight.MissShaderIndex = *LightFunctionIndex;
+				}
+			}
+
 			// these mean roughly the same thing across all light types
 			DestLight.Color = FVector3f(LightParameters.Color) * LightParameters.GetLightExposureScale(View.GetLastEyeAdaptationExposure());
 			DestLight.TranslatedWorldPosition = FVector3f(LightParameters.WorldPosition + View.ViewMatrices.GetPreViewTranslation());
@@ -2135,7 +2146,6 @@ void SetLightParameters(FRDGBuilder& GraphBuilder, FPathTracingRG::FParameters*
 
 	int32 NextRectTextureIndex = 0;
 
-	const FRayTracingLightFunctionMap* RayTracingLightFunctionMap = GraphBuilder.Blackboard.Get<FRayTracingLightFunctionMap>();
 	for (const FLightSceneInfoCompact& Light : Scene->Lights)
 	{
 		ELightComponentType LightComponentType = (ELightComponentType)Light.LightSceneInfo->Proxy->GetLightType();
