Make r.CustomDepthTemporalAAJitter=0 work properly for Nanite rendering custom depth.

diff --git a/Engine/Source/Runtime/Renderer/Private/CustomDepthRendering.cpp b/Engine/Source/Runtime/Renderer/Private/CustomDepthRendering.cpp
--- a/Engine/Source/Runtime/Renderer/Private/CustomDepthRendering.cpp
+++ b/Engine/Source/Runtime/Renderer/Private/CustomDepthRendering.cpp
@@ -90,7 +90,7 @@ BEGIN_SHADER_PARAMETER_STRUCT(FCustomDepthPassParameters, )
 	RENDER_TARGET_BINDING_SLOTS()
 END_SHADER_PARAMETER_STRUCT()
 
-static FViewShaderParameters CreateViewShaderParametersWithoutJitter(const FViewInfo& View)
+static FViewShaderParameters CreateViewShaderParametersWithoutJitter(const FViewInfo& View, FViewUniformShaderParameters& OutParams)
 {
 	const auto SetupParameters = [](const FViewInfo& View, FViewUniformShaderParameters& Parameters)
 	{
@@ -105,6 +105,8 @@ static FViewShaderParameters CreateViewShaderParametersWithoutJitter(const FView
 	FViewUniformShaderParameters ViewUniformParameters;
 	SetupParameters(View, ViewUniformParameters);
 
+	OutParams = ViewUniformParameters;
+
 	FViewShaderParameters Parameters;
 	Parameters.View = TUniformBufferRef<FViewUniformShaderParameters>::CreateUniformBufferImmediate(ViewUniformParameters, UniformBuffer_SingleFrame);
 
@@ -159,10 +161,19 @@ bool FSceneRenderer::RenderCustomDepthPass(
 		return false;
 	}
 
+	struct FTempViewParams
+	{
+		FViewShaderParameters ViewParams;
+		Nanite::FPackedView NaniteView;
+		FNaniteCustomDepthDrawList NaniteDrawList;
+	};
+	TArray<FTempViewParams, FSceneRenderingArrayAllocator> TempViewParams;
+	TempViewParams.SetNum(Views.Num());
+
+	const bool bRemoveTAAJitter = CVarCustomDepthTemporalAAJitter.GetValueOnRenderThread() == 0;
+
 	// Determine if any of the views have custom depth and if any of them have Nanite that is rendering custom depth
 	bool bAnyCustomDepth = false;
-	TArray<FNaniteCustomDepthDrawList, SceneRenderingAllocator> NaniteDrawLists;
-	NaniteDrawLists.AddDefaulted(Views.Num());
 	uint32 TotalNaniteInstances = 0;
 	for (int32 ViewIndex = 0; ViewIndex < Views.Num(); ++ViewIndex)
 	{
@@ -171,13 +182,28 @@ bool FSceneRenderer::RenderCustomDepthPass(
 		{
 			if (PrimaryNaniteRasterResults.IsValidIndex(ViewIndex))
 			{
+				TempViewParams[ViewIndex].NaniteView = PrimaryNaniteViews[ViewIndex];
 				FNaniteVisibilityQuery* VisibilityQuery = PrimaryNaniteRasterResults[ViewIndex].VisibilityQuery;
 
 				// Get the Nanite instance draw list for this view. (NOTE: Always use view index 0 for now because we're not doing
 				// multi-view yet).
-				NaniteDrawLists[ViewIndex] = BuildNaniteCustomDepthDrawList(View, 0u, Nanite::GetVisibilityResults(VisibilityQuery));
+				TempViewParams[ViewIndex].NaniteDrawList = BuildNaniteCustomDepthDrawList(View, 0u, Nanite::GetVisibilityResults(VisibilityQuery));
 
-				TotalNaniteInstances += NaniteDrawLists[ViewIndex].Num();
+				TotalNaniteInstances += TempViewParams[ViewIndex].NaniteDrawList.Num();
+			}
+
+			// User requested jitter-free custom depth.
+			if (bRemoveTAAJitter && IsTemporalAccumulationBasedMethod(View.AntiAliasingMethod))
+			{
+				FViewUniformShaderParameters ShaderParams;
+				TempViewParams[ViewIndex].ViewParams = CreateViewShaderParametersWithoutJitter(View, ShaderParams);
+				TempViewParams[ViewIndex].NaniteView.TranslatedWorldToClip = ShaderParams.TranslatedWorldToClip;
+				TempViewParams[ViewIndex].NaniteView.ViewToClip = ShaderParams.ViewToClip;
+				TempViewParams[ViewIndex].NaniteView.ClipToRelativeWorld = ShaderParams.ClipToRelativeWorld;
+			}
+			else
+			{
+				TempViewParams[ViewIndex].ViewParams = View.GetShaderParameters();
 			}
 			bAnyCustomDepth = true;
 		}
@@ -206,16 +232,7 @@ bool FSceneRenderer::RenderCustomDepthPass(
 
 			FCustomDepthPassParameters* PassParameters = GraphBuilder.AllocParameters<FCustomDepthPassParameters>();
 			PassParameters->SceneTextures = SceneTextures;
-
-			// User requested jitter-free custom depth.
-			if (CVarCustomDepthTemporalAAJitter.GetValueOnRenderThread() == 0 && IsTemporalAccumulationBasedMethod(View.AntiAliasingMethod))
-			{
-				PassParameters->View = CreateViewShaderParametersWithoutJitter(View);
-			}
-			else
-			{
-				PassParameters->View = View.GetShaderParameters();
-			}
+			PassParameters->View = TempViewParams[ViewIndex].ViewParams;
 
 			const ERenderTargetLoadAction DepthLoadAction = GetLoadActionIfProduced(CustomDepthTextures.Depth, CustomDepthTextures.DepthAction);
 			const ERenderTargetLoadAction StencilLoadAction = (View.Family->ViewMode == VMI_VisualizeBuffer) ? ERenderTargetLoadAction::EClear : GetLoadActionIfProduced(CustomDepthTextures.Depth, CustomDepthTextures.StencilAction);
@@ -291,7 +308,7 @@ bool FSceneRenderer::RenderCustomDepthPass(
 
 			FViewInfo& View = Views[ViewIndex];
 
-			if (!View.ShouldRenderView() || NaniteDrawLists[ViewIndex].Num() == 0)
+			if (!View.ShouldRenderView() || TempViewParams[ViewIndex].NaniteDrawList.Num() == 0)
 			{
 				continue;
 			}
@@ -311,8 +328,8 @@ bool FSceneRenderer::RenderCustomDepthPass(
 			NaniteRenderer->DrawGeometry(
 				Scene->NaniteRasterPipelines[ENaniteMeshPass::BasePass],
 				PrimaryNaniteRasterResults[ViewIndex].VisibilityQuery,
-				*Nanite::FPackedViewArray::Create(GraphBuilder, PrimaryNaniteViews[ViewIndex]),
-				NaniteDrawLists[ViewIndex]
+				*Nanite::FPackedViewArray::Create(GraphBuilder, TempViewParams[ViewIndex].NaniteView),
+				TempViewParams[ViewIndex].NaniteDrawList
 			);
 
 			Nanite::FRasterResults RasterResults;
