Follow-on touchup change for the ISR visibility calculation fixes in CL 28402809.
Improves clarity and adds some checks.

diff --git a/Engine/Source/Runtime/Renderer/Private/SceneVisibility.cpp b/Engine/Source/Runtime/Renderer/Private/SceneVisibility.cpp
--- a/Engine/Source/Runtime/Renderer/Private/SceneVisibility.cpp
+++ b/Engine/Source/Runtime/Renderer/Private/SceneVisibility.cpp
@@ -3326,7 +3326,7 @@ FVisibilityViewPacket::FVisibilityViewPacket(FVisibilityTaskData& InTaskData, FS
 				Relevance.PrimaryViewCommandPipe->AddNumCommands(1);
 				Relevance.PrimaryViewCommandPipe->EnqueueCommand(PrimitiveIndexList);
 
-				Relevance.Context->AddPrimitives(MoveTemp(PrimitiveIndexList));
+				Relevance.Context->AddPrimitives(CopyTemp(PrimitiveIndexList));
 			});
 		}
 		else
@@ -3344,13 +3344,7 @@ FVisibilityViewPacket::FVisibilityViewPacket(FVisibilityTaskData& InTaskData, FS
 			{
 				Relevance.Context->Finish(Tasks.ComputeRelevance);
 
-				if (TaskData.DynamicMeshElements.CommandPipe)
-				{
-					// Release our reference on the dynamic mesh element pipe. We want to keep it alive until all relevance packets have been launched.
-					TaskData.DynamicMeshElements.CommandPipe->ReleaseNumCommands(1);
-				}
-
-				// Also release reference on the instanced primary pipe
+				// Release reference on the instanced primary pipe
 				Relevance.PrimaryViewCommandPipe->ReleaseNumCommands(1);
 			});
 		}
@@ -3360,6 +3354,7 @@ FVisibilityViewPacket::FVisibilityViewPacket(FVisibilityTaskData& InTaskData, FS
 			{
 				Relevance.Context->Finish(Tasks.ComputeRelevance);
 
+				// Only used in the Views.Num() == 1 case
 				if (TaskData.DynamicMeshElements.CommandPipe)
 				{
 					TaskData.DynamicMeshElements.CommandPipe->ReleaseNumCommands(1);
@@ -4045,6 +4040,10 @@ void FVisibilityTaskData::LaunchVisibilityTasks(const UE::Tasks::FTask& BeginIni
 	{
 		SceneRenderer.WaitOcclusionTests(RHICmdList);
 
+		// Parallel occlusion culling is not supported on mobile
+		check(!Views.IsEmpty())
+		checkf(!Views[0]->bIsMobileMultiViewEnabled, TEXT("This culling path was not tested with MMV"));
+
 		// In instanced stereo, we'll also redirect all primitives to the primary view's relevance command pipe, so secondary viewports need a reference
 		if (Views[0]->bIsMultiViewportEnabled)
 		{
