Fix crash in hair voxelization buffer allocation when having 16 macro groups or more, with large bounding boxes.

This causes the buffer allocation to goes beyond 2Gb allocation. This CL expose min/max page index buffer size, and reduce the max size from 256 to 64. This clamps the allocation in similar case to 67Mb.

diff --git a/Engine/Source/Runtime/Renderer/Private/HairStrands/HairStrandsVoxelization.cpp b/Engine/Source/Runtime/Renderer/Private/HairStrands/HairStrandsVoxelization.cpp
--- a/Engine/Source/Runtime/Renderer/Private/HairStrands/HairStrandsVoxelization.cpp
+++ b/Engine/Source/Runtime/Renderer/Private/HairStrands/HairStrandsVoxelization.cpp
@@ -95,9 +95,10 @@ static FAutoConsoleVariableRef CVarHairVirtualVoxel_VoxelPageResolution(TEXT("r.
 static FAutoConsoleVariableRef CVarHairVirtualVoxel_VoxelPageCountPerDim(TEXT("r.HairStrands.Voxelization.Virtual.VoxelPageCountPerDim"), GHairVirtualVoxel_PageCountPerDim, TEXT("Number of voxel pages per texture dimension. The voxel page memory is allocated with a 3D texture. This value provide the resolution of this texture."));
 
 static int32 GHairVirtualVoxelGPUDriven = 1;
-static int32 GHairVirtualVoxelGPUDrivenMaxPageIndexRes = 32;
-static int32 GHairVirtualVoxelGPUDrivenClampPageIndexRes = 256;
+static int32 GHairVirtualVoxelGPUDrivenMinPageIndexRes = 32;
+static int32 GHairVirtualVoxelGPUDrivenMaxPageIndexRes = 64;
 static FAutoConsoleVariableRef CVarHairVirtualVoxelGPUDriven(TEXT("r.HairStrands.Voxelization.GPUDriven"), GHairVirtualVoxelGPUDriven, TEXT("Enable GPU driven voxelization."));
+static FAutoConsoleVariableRef CVarHairVirtualVoxelGPUDrivenMinPageIndexRes(TEXT("r.HairStrands.Voxelization.GPUDriven.MinPageIndexResolution"), GHairVirtualVoxelGPUDrivenMinPageIndexRes, TEXT("Min resolution of the page index. This is used for allocating a conservative page index buffer when GPU driven allocation is enabled."));
 static FAutoConsoleVariableRef CVarHairVirtualVoxelGPUDrivenMaxPageIndexRes(TEXT("r.HairStrands.Voxelization.GPUDriven.MaxPageIndexResolution"), GHairVirtualVoxelGPUDrivenMaxPageIndexRes, TEXT("Max resolution of the page index. This is used for allocating a conservative page index buffer when GPU driven allocation is enabled."));
 
 static const FIntPoint GPUDrivenViewportResolution = FIntPoint(4096, 4096);
@@ -532,7 +533,7 @@ static void AddAllocateVoxelPagesPass(
 	{
 		// Use the max between the estimated size on CPU and a pseudo-conservative side driven by settings. The CPU estimation is no necessarily correct as the bounds are not reliable on skel. mesh.
 		const uint32 MinPageIndexCount = GHairVirtualVoxelGPUDrivenMaxPageIndexRes * GHairVirtualVoxelGPUDrivenMaxPageIndexRes * GHairVirtualVoxelGPUDrivenMaxPageIndexRes;
-		const uint32 MaxPageIndexCount = GHairVirtualVoxelGPUDrivenClampPageIndexRes * GHairVirtualVoxelGPUDrivenClampPageIndexRes * GHairVirtualVoxelGPUDrivenClampPageIndexRes;
+		const uint32 MaxPageIndexCount = GHairVirtualVoxelGPUDrivenMinPageIndexRes * GHairVirtualVoxelGPUDrivenMinPageIndexRes * GHairVirtualVoxelGPUDrivenMinPageIndexRes;
 		OutTotalPageIndexCount = FMath::Clamp(OutTotalPageIndexCount, MinPageIndexCount, MaxPageIndexCount);
 	}
 	check(OutTotalPageIndexCount > 0);
@@ -548,6 +549,9 @@ static void AddAllocateVoxelPagesPass(
 
 	FRDGBufferRef ReadBackBuffer = GraphBuilder.CreateBuffer(FRDGBufferDesc::CreateBufferDesc(sizeof(uint32), 1), TEXT("ReadBackAllocations"));
 
+	const uint32 PageIndexBufferSizeInBytes = 4 * sizeof(uint32) * OutTotalPageIndexCount;
+	const uint32 TotalPageIndexBufferSizeInBytes = PageIndexBufferSizeInBytes * PageIndexBufferSizeInBytes;
+
 	FRDGBufferRef PageIndexResolutionBuffer = GraphBuilder.CreateBuffer(FRDGBufferDesc::CreateBufferDesc(MacroGroupCount * 4 * sizeof(uint32), OutTotalPageIndexCount), TEXT("Hair.PageIndexResolutionBuffer"));
 	FRDGBufferRef PageIndexAllocationIndirectBufferArgs = GraphBuilder.CreateBuffer(FRDGBufferDesc::CreateIndirectDesc<FRHIDispatchIndirectParameters>(MacroGroupCount+1), TEXT("Hair.PageIndexAllocationIndirectBufferArgs"));
 	
