Fixed tracking of dirty instance hierarchy cells for the purposes of explicit bounds update.

#rnx

diff --git a/Engine/Source/Runtime/RenderCore/Public/RenderGraphUtils.h b/Engine/Source/Runtime/RenderCore/Public/RenderGraphUtils.h
--- a/Engine/Source/Runtime/RenderCore/Public/RenderGraphUtils.h
+++ b/Engine/Source/Runtime/RenderCore/Public/RenderGraphUtils.h
@@ -911,6 +911,26 @@ FORCEINLINE FRDGBufferRef CreateStructuredBuffer(
 	return CreateStructuredBuffer(GraphBuilder, Name, InitialData.GetTypeSize(), InitialData.Num(), InitialData.GetData(), InitialData.Num() * InitialData.GetTypeSize(), InitialDataFlags);
 }
 
+/**
+ * Helper to create a structured buffer with initial data from a TArray with move semantics, this can be cheaper as it guarantees the lifetimes of the data & permits copy-free upload.
+ */
+template <typename ElementType, typename AllocatorType>
+FORCEINLINE FRDGBufferRef CreateStructuredBuffer(
+	FRDGBuilder& GraphBuilder,
+	const TCHAR* Name,
+	TArray<ElementType, AllocatorType>&& InitialData)
+{
+	static const ElementType DummyElement = ElementType();
+	if (InitialData.Num() == 0)
+	{
+		return CreateStructuredBuffer(GraphBuilder, Name, InitialData.GetTypeSize(), 1, &DummyElement, InitialData.GetTypeSize(), ERDGInitialDataFlags::NoCopy);
+	}
+
+	// Create a move-initialized copy of the TArray with RDG lifetime & move the data there.
+	TArray<ElementType, AllocatorType>& UploadData = *GraphBuilder.AllocObject<TArray<ElementType, AllocatorType> >(MoveTemp(InitialData));
+	return CreateStructuredBuffer(GraphBuilder, Name, UploadData.GetTypeSize(), UploadData.Num(), UploadData.GetData(), UploadData.Num() * UploadData.GetTypeSize(), ERDGInitialDataFlags::NoCopy);
+}
+
 /**
  * Helper to create a structured buffer with initial data from a TConstArrayView.
  */
diff --git a/Engine/Source/Runtime/Renderer/Private/SceneCulling/SceneCulling.cpp b/Engine/Source/Runtime/Renderer/Private/SceneCulling/SceneCulling.cpp
--- a/Engine/Source/Runtime/Renderer/Private/SceneCulling/SceneCulling.cpp
+++ b/Engine/Source/Runtime/Renderer/Private/SceneCulling/SceneCulling.cpp
@@ -105,7 +105,7 @@ LLM_DEFINE_TAG(SceneCulling, NAME_None, NAME_None, GET_STATFNAME(STAT_SceneCulli
 
 static TAutoConsoleVariable<int32> CVarSceneCulling(
 	TEXT("r.SceneCulling"), 
-	0, 
+	1, 
 	TEXT("Enable/Disable scene culling.\n")
 	TEXT("  While enabled, it will only build the instance hierarchy if used by any system - currently that corresponds to Nanite being enabled.\n")
 	TEXT("  Forces a recreate of all render state since (at present) there is only an incremental update path."),
@@ -769,7 +769,6 @@ class FComputeExplicitCellBounds_CS : public FGlobalShader
 	}
 
 	BEGIN_SHADER_PARAMETER_STRUCT(FParameters, )
-		//SHADER_PARAMETER_STRUCT_INCLUDE(ShaderPrint::FShaderParameters, ShaderPrintUniformBuffer)
 		SHADER_PARAMETER_RDG_UNIFORM_BUFFER( FSceneUniformParameters, Scene )
 
 		SHADER_PARAMETER(uint32, NumCellsPerBlockLog2)
@@ -1357,7 +1356,10 @@ public:
 			}			
 			// No need to set the bits as they are maintained incrementally (or new and cleared)
 			SceneCulling.CellOccupancyMask.SetNum(NewMinSize, false);
-
+			if (SceneCulling.bUseExplictBounds)
+			{
+				DirtyCellBoundsMask.SetNum(NewMinSize, false);
+			}
 			BUILDER_LOG("Alloc Cells [%d, %d], Block %d", StartIndex, StartIndex + CellBlockSize, BlockId.GetIndex());
 		}
 
@@ -1603,6 +1605,8 @@ public:
 
 				// 3. Copy the list of chunk IDs
 				TempCell.OutputChunkIds(*this);
+
+				MarkCellBoundsDirty(TempCell.CellOffset);
 			}
 			LogCell(UE::HLSL::UnpackCellHeader(CellHeader));
 			CellHeaderUploader.Add(CellHeader, TempCell.CellOffset);
@@ -1879,6 +1883,18 @@ public:
 		MarkCellForRemove(CellIndex, NumInstances, UpdateFrequencyCategory);
 	}
 
+	SC_FORCEINLINE void MarkCellBoundsDirty(int32 CellIndex)
+	{
+		if (SceneCulling.bUseExplictBounds)
+		{
+			if (!DirtyCellBoundsMask[CellIndex])
+			{
+				DirtyCellBoundsIndices.Add(CellIndex);
+			}
+			DirtyCellBoundsMask[CellIndex] = true;
+		}
+	}
+
 	inline void RemovePrecomputed(FSceneCulling::FPrimitiveState &PrimitiveState)
 	{
 		BUILDER_LOG("FPrimitiveState::Precomputed");
@@ -1970,6 +1986,10 @@ public:
 					MarkForRemove(PrevCellIndex, InstanceId, 1, EUpdateFrequencyCategory::Dynamic);
 					bNeedAdd = true;
 				}
+				else
+				{
+					MarkCellBoundsDirty(PrevCellIndex);
+				}
 			}
 			if (bNeedAdd)
 			{
@@ -2050,6 +2070,7 @@ public:
 				{
 					// retain previous state.
 					NewPrimitiveState.Payload = PrevCellIndex;
+					MarkCellBoundsDirty(PrevCellIndex);
 				}
 			}
 		}
@@ -2214,7 +2235,7 @@ public:
 
 		bool bValidToCapture = CellHeaderUploader.GetNumScatters() > 0;
 #if !UE_BUILD_SHIPPING
-		RenderCaptureInterface::FScopedCapture RenderCapture(bValidToCapture && GCaptureNextSceneCullingUpdate-- == 0, GraphBuilder, TEXT("UploadToGPU"));
+		RenderCaptureInterface::FScopedCapture RenderCapture(bValidToCapture && GCaptureNextSceneCullingUpdate-- == 0, GraphBuilder, TEXT("SceneCulling.UploadToGPU"));
 		// Prevent overflow every 2B frames.
 		GCaptureNextSceneCullingUpdate = FMath::Max(-1, GCaptureNextSceneCullingUpdate);
 #endif
@@ -2236,11 +2257,10 @@ public:
 			bool bFullUpload = !SceneCulling.ExplicitCellBoundsBuffer.GetPooledBuffer().IsValid();
 			SceneCulling.ExplicitCellBoundsBuffer.ResizeBufferIfNeeded(GraphBuilder, SceneCulling.CellHeaders.Num() * 2);
 
-			if (UpdatedCellScatterInfo.NumScatters > 0 || bFullUpload)
+			if (!DirtyCellBoundsIndices.IsEmpty() || bFullUpload)
 			{
 				RDG_EVENT_SCOPE(GraphBuilder, "SceneCulling_ComputeExplicitCellBounds");
 				FComputeExplicitCellBounds_CS::FParameters* PassParameters = GraphBuilder.AllocParameters<FComputeExplicitCellBounds_CS::FParameters>();
-				//ShaderPrint::SetParameters(GraphBuilder, PassParameters->ShaderPrintUniformBuffer);
 				PassParameters->Scene = SceneUniformBuffer.GetBuffer(GraphBuilder);
 				PassParameters->NumCellsPerBlockLog2 = FSpatialHash::NumCellsPerBlockLog2;
 				PassParameters->CellBlockDimLog2 = FSpatialHash::CellBlockDimLog2;
@@ -2250,8 +2270,9 @@ public:
 				PassParameters->InstanceHierarchyCellHeaders = GraphBuilder.CreateSRV(CellHeadersRDG);
 				PassParameters->InstanceIds = GraphBuilder.CreateSRV(InstanceIdsRDG);
 				PassParameters->InstanceHierarchyItemChunks = GraphBuilder.CreateSRV(ItemChunksRDG);
-				PassParameters->UpdatedCellIds = GraphBuilder.CreateSRV(bFullUpload ? GSystemTextures.GetDefaultStructuredBuffer<uint32>(GraphBuilder) : UpdatedCellScatterInfo.ScatterOffsetsRDG);
-				int32 MaxCellCount = bFullUpload ? SceneCulling.CellHeaders.Num() : UpdatedCellScatterInfo.NumScatters;
+				int32 MaxCellCount = bFullUpload ? SceneCulling.CellHeaders.Num() : DirtyCellBoundsIndices.Num();
+				// Note: this Moves the DirtyCellBoundsIndices & thus implicitly clears it!
+				PassParameters->UpdatedCellIds = GraphBuilder.CreateSRV(bFullUpload ? GSystemTextures.GetDefaultStructuredBuffer<uint32>(GraphBuilder) : CreateStructuredBuffer(GraphBuilder, TEXT("SceneCulling.ModifiedCells"), MoveTemp(DirtyCellBoundsIndices)));
 				PassParameters->MaxCells = MaxCellCount;
 				PassParameters->OutExplicitCellBoundsBuffer = GraphBuilder.CreateUAV(SceneCulling.ExplicitCellBoundsBuffer.Register(GraphBuilder));
 
@@ -2410,11 +2431,11 @@ public:
 			SceneCulling.BlockLevelOccupancyMask[Item.Key.GetLevel()] = true;
 		}
 
-		// Make sure to release any references to data allocated through SceneRenderingAllocator
-		// TODO: maybe change to default allocator? Might make sense to keep the data allocated anyway.
+		// Releasing memory we're done with in the task saves some RT time
 		DirtyChunks.Empty();
 		DirtyBlocks.Empty();
 		RemovedInstanceFlags.Empty();
+		DirtyCellBoundsMask.Empty();
 	}
 
 	using FRemovedIntanceFlags = TBitArray<SceneRenderingAllocator>;
@@ -2437,7 +2458,10 @@ public:
 	int32 NumDirtyBlocks = 0;
 	TBitArray<SceneRenderingAllocator> DirtyBlocks;
 	TBitArray<SceneRenderingAllocator> DirtyChunks;
-	
+	// Dirty tracking for explicit cell bounds
+	TBitArray<SceneRenderingAllocator> DirtyCellBoundsMask;
+	TArray<int32, SceneRenderingAllocator> DirtyCellBoundsIndices;
+
 	TStructuredBufferScatterUploader<FCellBlockData> BlockDataUploader;
 	TStructuredBufferScatterUploader<uint32, INSTANCE_HIERARCHY_MAX_CHUNK_SIZE> ItemChunkDataUploader;
 	TStructuredBufferScatterUploader<FPackedCellHeader> CellHeaderUploader;
@@ -2603,6 +2627,12 @@ void FSceneCulling::FUpdater::OnPreSceneUpdate(FRDGBuilder& GraphBuilder, const
 		BUILDER_LOG_SCOPE("OnPreSceneUpdate: %d", ScenePreUpdateData.RemovedPrimitiveIds.Num());
 		CSV_SCOPED_TIMING_STAT(SceneCulling, PreSceneUpdate);
 
+		if (Implementation->SceneCulling.bUseExplictBounds)
+		{
+			Implementation->DirtyCellBoundsMask.Init(false, Implementation->SceneCulling.CellHeaders.Num());
+			Implementation->DirtyCellBoundsIndices.Reset(Implementation->SceneCulling.CellHeaders.Num());
+		}
+
 		check(DebugTaskCounter++ == 0);
 		for (int32 Index = 0; Index < ScenePreUpdateData.RemovedPrimitiveIds.Num(); ++Index)
 		{
