Substrate - fixed compilation issue leading to failed packaging.

diff --git a/Engine/Source/Runtime/Renderer/Private/Substrate/SubstrateVisualize.cpp b/Engine/Source/Runtime/Renderer/Private/Substrate/SubstrateVisualize.cpp
--- a/Engine/Source/Runtime/Renderer/Private/Substrate/SubstrateVisualize.cpp
+++ b/Engine/Source/Runtime/Renderer/Private/Substrate/SubstrateVisualize.cpp
@@ -26,7 +26,8 @@ static bool IsSubstrateDebugVisualizationSupported(EShaderPlatform InPlatform, b
 {
 	return 
 		Substrate::IsSubstrateEnabled() && 
-		GetMaxSupportedFeatureLevel(InPlatform) >= ERHIFeatureLevel::SM5 && 
+		GetMaxSupportedFeatureLevel(InPlatform) >= ERHIFeatureLevel::SM5 &&
+		!IsVulkanPlatform(InPlatform) && // Does not compile and fails to package games. Fixing this would require a deeper investigation.
 		(bIsEditorOnly ? (IsPCPlatform(InPlatform) || EnumHasAllFlags(Flags, EShaderPermutationFlags::HasEditorOnlyData)) : true);
 }
 
