Fixed race condition between RDG and the scene extensions when the scene is destroyed.

diff --git a/Engine/Source/Runtime/Renderer/Private/RendererScene.cpp b/Engine/Source/Runtime/Renderer/Private/RendererScene.cpp
--- a/Engine/Source/Runtime/Renderer/Private/RendererScene.cpp
+++ b/Engine/Source/Runtime/Renderer/Private/RendererScene.cpp
@@ -4596,6 +4596,9 @@ void FScene::Release()
 	ENQUEUE_RENDER_COMMAND(FReleaseCommand)(
 		[Scene](FRHICommandListImmediate& RHICmdList)
 		{
+			// Finish all pending scene rendering destruction work before deleting the scene.
+			FSceneRenderer::CleanUp(RHICmdList);
+
 			// Flush any remaining batched primitive update commands before deleting the scene.
 			FUpdateParameters UpdateParameters;
 			UpdateParameters.bDestruction = true;
