Replaced GetStaticRayTracingGeometryInstance() with GetValidStaticRayTracingGeometry() which returns the first valid RayTracingGeometry.
- fixes issue with instances missing from raytracing scene when the RayTracingGeometry of the target LOD is pending build but lower LODs RayTracingGeometry is already available and can be used as fallback.

diff --git a/Engine/Source/Runtime/Renderer/Private/PrimitiveSceneInfo.cpp b/Engine/Source/Runtime/Renderer/Private/PrimitiveSceneInfo.cpp
--- a/Engine/Source/Runtime/Renderer/Private/PrimitiveSceneInfo.cpp
+++ b/Engine/Source/Runtime/Renderer/Private/PrimitiveSceneInfo.cpp
@@ -398,6 +398,38 @@ bool FPrimitiveSceneInfo::IsCachedRayTracingGeometryValid() const
 	return false;
 }
 
+FRayTracingGeometry* FPrimitiveSceneInfo::GetStaticRayTracingGeometry(int32 LodLevel) const
+{
+	if (StaticRayTracingGeometries.Num() > LodLevel)
+	{
+		return StaticRayTracingGeometries[LodLevel];
+	}
+	else
+	{
+		return nullptr;
+	}
+}
+
+FRayTracingGeometry* FPrimitiveSceneInfo::GetValidStaticRayTracingGeometry(int32 MinLodLevel) const
+{
+	// TODO: Move HasPendingBuildRequest() / BoostBuildPriority() out of this function
+
+	for (int32 Index = MinLodLevel; Index < StaticRayTracingGeometries.Num(); ++Index)
+	{
+		if (StaticRayTracingGeometries[Index]->HasPendingBuildRequest())
+		{
+			ensure(StaticRayTracingGeometries[Index]->IsValid());
+			StaticRayTracingGeometries[Index]->BoostBuildPriority();
+		}
+		else if (StaticRayTracingGeometries[Index]->IsValid())
+		{
+			return StaticRayTracingGeometries[Index];
+		}
+	}
+
+	return nullptr;
+}
+
 FRHIRayTracingGeometry* FPrimitiveSceneInfo::GetStaticRayTracingGeometryInstance(int LodLevel) const
 {
 	if (StaticRayTracingGeometries.Num() > LodLevel)
diff --git a/Engine/Source/Runtime/Renderer/Private/RayTracing/RayTracing.cpp b/Engine/Source/Runtime/Renderer/Private/RayTracing/RayTracing.cpp
--- a/Engine/Source/Runtime/Renderer/Private/RayTracing/RayTracing.cpp
+++ b/Engine/Source/Runtime/Renderer/Private/RayTracing/RayTracing.cpp
@@ -540,8 +540,9 @@ namespace RayTracing
 						}
 						else
 						{
-							FRHIRayTracingGeometry* RayTracingGeometryInstance = SceneInfo->GetStaticRayTracingGeometryInstance(LODIndex);
-							if (RayTracingGeometryInstance == nullptr)
+							FRayTracingGeometry* RayTracingGeometry = SceneInfo->GetValidStaticRayTracingGeometry(LODIndex);
+
+							if (RayTracingGeometry == nullptr)
 							{
 								return;
 							}
@@ -555,7 +556,7 @@ namespace RayTracing
 								RelevantPrimitive->PersistentPrimitiveIndex = SceneInfo->GetPersistentIndex();
 
 								RelevantPrimitive->LODIndex = LODIndex;
-								RelevantPrimitive->RayTracingGeometryRHI = RayTracingGeometryInstance;
+								RelevantPrimitive->RayTracingGeometryRHI = RayTracingGeometry->RayTracingGeometryRHI;
 
 								RelevantPrimitive->CachedRayTracingMeshCommandIndices = SceneInfo->CachedRayTracingMeshCommandIndicesPerLOD[LODIndex];
 								RelevantPrimitive->StateHash = SceneInfo->CachedRayTracingMeshCommandsHashPerLOD[LODIndex];
diff --git a/Engine/Source/Runtime/Renderer/Public/PrimitiveSceneInfo.h b/Engine/Source/Runtime/Renderer/Public/PrimitiveSceneInfo.h
--- a/Engine/Source/Runtime/Renderer/Public/PrimitiveSceneInfo.h
+++ b/Engine/Source/Runtime/Renderer/Public/PrimitiveSceneInfo.h
@@ -548,6 +548,10 @@ public:
 
 	bool IsCachedRayTracingGeometryValid() const;
 
+	RENDERER_API FRayTracingGeometry* GetStaticRayTracingGeometry(int32 LodLevel) const;
+	RENDERER_API FRayTracingGeometry* GetValidStaticRayTracingGeometry(int32 MinLodLevel) const;
+
+	UE_DEPRECATED(5.4, "GetStaticRayTracingGeometryInstance has been deprecated. Use GetStaticRayTracingGeometry / GetValidStaticRayTracingGeometry instead.")
 	RENDERER_API FRHIRayTracingGeometry* GetStaticRayTracingGeometryInstance(int LodLevel) const;
 
 	int GetStaticRayTracingGeometryNum() const { return StaticRayTracingGeometries.Num(); }
