HeterogeneousVolumes: Fix crash when collecting shadow-casting primitives and selecting rendering bounds.

diff --git a/Engine/Source/Runtime/Engine/Private/Components/HeterogeneousVolumeComponent.cpp b/Engine/Source/Runtime/Engine/Private/Components/HeterogeneousVolumeComponent.cpp
--- a/Engine/Source/Runtime/Engine/Private/Components/HeterogeneousVolumeComponent.cpp
+++ b/Engine/Source/Runtime/Engine/Private/Components/HeterogeneousVolumeComponent.cpp
@@ -216,7 +216,8 @@ void FHeterogeneousVolumeSceneProxy::GetDynamicMeshElements(
 		{
 			if (VisibilityMap & (1 << ViewIndex))
 			{
-				RenderBounds(Collector.GetPDI(ViewIndex), ViewFamily.EngineShowFlags, GetBounds(), IsSelected());
+				const FSceneView* View = Views[ViewIndex];
+				RenderBounds(Collector.GetPDI(ViewIndex), View->Family->EngineShowFlags, GetBounds(), IsSelected());
 			}
 		}
 	}
diff --git a/Engine/Source/Runtime/Renderer/Private/ShadowSetup.cpp b/Engine/Source/Runtime/Renderer/Private/ShadowSetup.cpp
--- a/Engine/Source/Runtime/Renderer/Private/ShadowSetup.cpp
+++ b/Engine/Source/Runtime/Renderer/Private/ShadowSetup.cpp
@@ -2694,16 +2694,17 @@ bool FProjectedShadowInfo::GatherDynamicMeshElements(
 			ShadowDepthView->SetDynamicMeshElementsShadowCullFrustum(&CasterOuterFrustum);
 			bProcessedAllPrimitives &= GatherDynamicMeshElementsArray(MeshCollector, DynamicSubjectPrimitives, ReusedViewsArray, Renderer.ViewFamily, DynamicSubjectMeshElements, NumDynamicSubjectMeshElements, Pass);
 		}
+		MeshCollector.ClearViewMeshArrays();
 
 		ShadowDepthView->DrawDynamicFlags = EDrawDynamicFlags::None;
 
 		int32 NumDynamicSubjectTranslucentMeshElements = 0;
 		ShadowDepthView->SetDynamicMeshElementsShadowCullFrustum(&CasterOuterFrustum);
 		bProcessedAllPrimitives &= GatherDynamicMeshElementsArray(MeshCollector, SubjectTranslucentPrimitives, ReusedViewsArray, Renderer.ViewFamily, DynamicSubjectTranslucentMeshElements, NumDynamicSubjectTranslucentMeshElements, Pass);
+		MeshCollector.ClearViewMeshArrays();
 
 		int32 NumDynamicSubjectHeterogeneousVolumeMeshElements = 0;
 		bProcessedAllPrimitives &= GatherDynamicHeterogeneousVolumeMeshElementsArray(Renderer, MeshCollector, SubjectHeterogeneousVolumePrimitives, ReusedViewsArray, Renderer.ViewFamily, DynamicSubjectHeterogeneousVolumeMeshElements, NumDynamicSubjectHeterogeneousVolumeMeshElements, Pass);
-
 		MeshCollector.ClearViewMeshArrays();
 	}
 
@@ -2782,13 +2783,6 @@ bool FProjectedShadowInfo::GatherDynamicHeterogeneousVolumeMeshElementsArray(
 	int32& OutNumDynamicSubjectMeshElements,
 	EGatherDynamicMeshElementsPass Pass)
 {
-	int32 Zero = 0;
-	TArray<int32> MeshElementsPerView(&Zero, Views.Num());
-	for (int32 ViewIndex = 0; ViewIndex < Views.Num(); ++ViewIndex)
-	{
-		MeshElementsPerView[ViewIndex] = MeshCollector.GetMeshElementCount(ViewIndex);
-	}
-
 	bool bProcessedAllPrimitives = GatherDynamicMeshElementsArray(
 		MeshCollector,
 		Primitives,
@@ -2802,13 +2796,13 @@ bool FProjectedShadowInfo::GatherDynamicHeterogeneousVolumeMeshElementsArray(
 	for (int32 ViewIndex = 0; ViewIndex < Views.Num(); ++ViewIndex)
 	{
 		uint32 MeshElementCount = MeshCollector.GetMeshElementCount(ViewIndex);
-		for (uint32 Index = MeshElementsPerView[ViewIndex]; Index < MeshElementCount; ++Index)
+		for (uint32 Index = 0; Index < MeshElementCount; ++Index)
 		{
 			FMeshBatchAndRelevance MeshBatchAndRelevance = (*MeshCollector.MeshBatches[ViewIndex])[Index];
 			OutDynamicMeshElements.Add(MeshBatchAndRelevance);
 		}
 
-		OutNumDynamicSubjectMeshElements += MeshElementCount - MeshElementsPerView[ViewIndex];
+		OutNumDynamicSubjectMeshElements += MeshElementCount;
 	}
 
 	return bProcessedAllPrimitives;
