Fix Deep shadow in VR.

diff --git a/Engine/Source/Runtime/Renderer/Private/HairStrands/HairStrandsDeepShadow.cpp b/Engine/Source/Runtime/Renderer/Private/HairStrands/HairStrandsDeepShadow.cpp
--- a/Engine/Source/Runtime/Renderer/Private/HairStrands/HairStrandsDeepShadow.cpp
+++ b/Engine/Source/Runtime/Renderer/Private/HairStrands/HairStrandsDeepShadow.cpp
@@ -243,6 +243,13 @@ void RenderHairStrandsDeepShadows(
 	FHairStrandsDeepShadowResources& DeepShadowResources = View.HairStrandsViewData.DeepShadowResources;
 	FHairStrandsVoxelResources VirtualVoxelResources = View.HairStrandsViewData.VirtualVoxelResources;
 
+	// Reset view data
+	for (FHairStrandsMacroGroupData& MacroGroup : MacroGroupDatas)
+	{
+		MacroGroup.DeepShadowDatas.Empty();
+	}
+	DeepShadowResources = FHairStrandsDeepShadowResources();
+
 	{
 		if (!View.Family)
 		{
diff --git a/Engine/Source/Runtime/Renderer/Private/HairStrands/HairStrandsRendering.cpp b/Engine/Source/Runtime/Renderer/Private/HairStrands/HairStrandsRendering.cpp
--- a/Engine/Source/Runtime/Renderer/Private/HairStrands/HairStrandsRendering.cpp
+++ b/Engine/Source/Runtime/Renderer/Private/HairStrands/HairStrandsRendering.cpp
@@ -138,6 +138,11 @@ void RenderHairPrePass(
 				// voxelization feedback (only done for the first view in stereo) and Path-Tracer invalidation 
 				// (not supporting stereo)
 				Views[1].HairStrandsViewData = Views[0].HairStrandsViewData;
+
+				// Render DeepShadow for the second view, as for now the computations are view dependent. 
+				// This needs to be view independent to share result between eyes.
+				RenderHairStrandsDeepShadows(GraphBuilder, Scene, View, InstanceCullingManager);
+				GraphBuilder.AddDispatchHint();
 				return;
 			}
 		}
