Fix typo.

diff --git a/Engine/Source/Runtime/Renderer/Private/HairStrands/HairStrandsVoxelization.cpp b/Engine/Source/Runtime/Renderer/Private/HairStrands/HairStrandsVoxelization.cpp
--- a/Engine/Source/Runtime/Renderer/Private/HairStrands/HairStrandsVoxelization.cpp
+++ b/Engine/Source/Runtime/Renderer/Private/HairStrands/HairStrandsVoxelization.cpp
@@ -532,8 +532,8 @@ static void AddAllocateVoxelPagesPass(
 	if (bIsGPUDriven)
 	{
 		// Use the max between the estimated size on CPU and a pseudo-conservative side driven by settings. The CPU estimation is no necessarily correct as the bounds are not reliable on skel. mesh.
-		const uint32 MinPageIndexCount = GHairVirtualVoxelGPUDrivenMaxPageIndexRes * GHairVirtualVoxelGPUDrivenMaxPageIndexRes * GHairVirtualVoxelGPUDrivenMaxPageIndexRes;
-		const uint32 MaxPageIndexCount = GHairVirtualVoxelGPUDrivenMinPageIndexRes * GHairVirtualVoxelGPUDrivenMinPageIndexRes * GHairVirtualVoxelGPUDrivenMinPageIndexRes;
+		const uint32 MinPageIndexCount = GHairVirtualVoxelGPUDrivenMinPageIndexRes * GHairVirtualVoxelGPUDrivenMinPageIndexRes * GHairVirtualVoxelGPUDrivenMinPageIndexRes;
+		const uint32 MaxPageIndexCount = GHairVirtualVoxelGPUDrivenMaxPageIndexRes * GHairVirtualVoxelGPUDrivenMaxPageIndexRes * GHairVirtualVoxelGPUDrivenMaxPageIndexRes;
 		OutTotalPageIndexCount = FMath::Clamp(OutTotalPageIndexCount, MinPageIndexCount, MaxPageIndexCount);
 	}
 	check(OutTotalPageIndexCount > 0);
