Fix ensure in MDC stats by taking into account multiple views in the same culling pass (point light shadows for example)

#rnx

diff --git a/Engine/Source/Runtime/Renderer/Private/MeshDrawCommands.cpp b/Engine/Source/Runtime/Renderer/Private/MeshDrawCommands.cpp
--- a/Engine/Source/Runtime/Renderer/Private/MeshDrawCommands.cpp
+++ b/Engine/Source/Runtime/Renderer/Private/MeshDrawCommands.cpp
@@ -878,6 +878,10 @@ void CollectMeshDrawCommandPassStats(
 	check(VisibleMeshDrawCommands.Num() == InstanceCullingContext.MeshDrawCommandInfos.Num());
 		
 	PassStats->DrawData.SetNum(VisibleMeshDrawCommands.Num(), EAllowShrinking::No);
+
+	// See the InterlockedAdd on DrawIndirectArgsBufferOut in InstanceCullBuildInstanceIdBufferCS
+	int32 InstanceCountMultiplier = (InstanceCullingContext.InstanceCullingMode == EInstanceCullingMode::Stereo) ? 2 : InstanceCullingContext.ViewIds.Num();
+
 	for (int32 DrawCommandIndex = 0; DrawCommandIndex < VisibleMeshDrawCommands.Num(); ++DrawCommandIndex)
 	{
 		const FVisibleMeshDrawCommand& RESTRICT VisibleMeshDrawCommand = VisibleMeshDrawCommands[DrawCommandIndex];
@@ -931,6 +935,8 @@ void CollectMeshDrawCommandPassStats(
 				DrawData.TotalInstanceCount = MeshDrawCommand->NumInstances;
 			}
 
+			DrawData.TotalInstanceCount *= InstanceCountMultiplier;
+
 			// By default mark all invisible
 			DrawData.VisibleInstanceCount = 0;
 		}
