Create dynamic raytracing geometry using FRayTracingGeometry::CreateRayTracingGeometry() instead of calling RHICmdList.CreateRayTracingGeometry(Geometry.Initializer) manually.
- this way other state is correctly updated.

diff --git a/Engine/Source/Runtime/Renderer/Private/RayTracing/RayTracingDynamicGeometry.cpp b/Engine/Source/Runtime/Renderer/Private/RayTracing/RayTracingDynamicGeometry.cpp
--- a/Engine/Source/Runtime/Renderer/Private/RayTracing/RayTracingDynamicGeometry.cpp
+++ b/Engine/Source/Runtime/Renderer/Private/RayTracing/RayTracingDynamicGeometry.cpp
@@ -346,9 +346,8 @@ void FRayTracingDynamicGeometryCollection::AddDynamicMeshBatchForGeometryUpdate(
 
 	if (!bRefit)
 	{
-		checkf(Geometry.Initializer.OfflineData == nullptr, TEXT("Dynamic geometry is not expected to have offline acceleration structure data"));
-		Geometry.RayTracingGeometryRHI = RHICmdList.CreateRayTracingGeometry(Geometry.Initializer);
-		Geometry.SetRequiresBuild(true);
+		checkf(Geometry.RawData.IsEmpty() && Geometry.Initializer.OfflineData == nullptr, TEXT("Dynamic geometry is not expected to have offline acceleration structure data"));
+		Geometry.CreateRayTracingGeometry(RHICmdList, ERTAccelerationStructureBuildPriority::Skip);
 	}
 
 	FRayTracingGeometryBuildParams Params;
