Fix HairComplexTransmittance increasing occupancy on certain platform.

Add permutation for FDiffuseIndirectCompositePS and FReflectionEnvironmentSkyLightingPS to bypassing this part of the code. It is only used for hair cards/meshes.

diff --git a/Engine/Source/Runtime/Renderer/Private/IndirectLightRendering.cpp b/Engine/Source/Runtime/Renderer/Private/IndirectLightRendering.cpp
--- a/Engine/Source/Runtime/Renderer/Private/IndirectLightRendering.cpp
+++ b/Engine/Source/Runtime/Renderer/Private/IndirectLightRendering.cpp
@@ -123,9 +123,10 @@ class FDiffuseIndirectCompositePS : public FGlobalShader
 	class FScreenBentNormal : SHADER_PERMUTATION_BOOL("DIM_SCREEN_BENT_NORMAL");
 	class FSubstrateTileType : SHADER_PERMUTATION_INT("SUBSTRATE_TILETYPE", 4);
 	class FEnableDualSrcBlending : SHADER_PERMUTATION_BOOL("ENABLE_DUAL_SRC_BLENDING");
+	class FHairComplexTransmittance : SHADER_PERMUTATION_BOOL("USE_HAIR_COMPLEX_TRANSMITTANCE");
 	class FScreenSpaceReflectionTiledComposition : SHADER_PERMUTATION_BOOL("SSR_TILED_COMPOSITION");
 
-	using FPermutationDomain = TShaderPermutationDomain<FApplyDiffuseIndirectDim, FUpscaleDiffuseIndirectDim, FScreenBentNormal, FSubstrateTileType, FEnableDualSrcBlending, FScreenSpaceReflectionTiledComposition>;
+	using FPermutationDomain = TShaderPermutationDomain<FApplyDiffuseIndirectDim, FUpscaleDiffuseIndirectDim, FScreenBentNormal, FSubstrateTileType, FEnableDualSrcBlending, FScreenSpaceReflectionTiledComposition, FHairComplexTransmittance>;
 
 	static bool ShouldCompilePermutation(const FGlobalShaderPermutationParameters& Parameters)
 	{
@@ -211,7 +212,6 @@ class FDiffuseIndirectCompositePS : public FGlobalShader
 	static void ModifyCompilationEnvironment(const FGlobalShaderPermutationParameters& Parameters, FShaderCompilerEnvironment& OutEnvironment)
 	{
 		OutEnvironment.CompilerFlags.Add(CFLAG_ForceOptimization);
-		OutEnvironment.SetDefine(TEXT("USE_HAIR_COMPLEX_TRANSMITTANCE"), IsHairStrandsSupported(EHairStrandsShaderType::All, Parameters.Platform) ? 1u : 0u);
 		if (Substrate::IsSubstrateEnabled() && Substrate::IsOpaqueRoughRefractionEnabled())
 		{
 			OutEnvironment.CompilerFlags.Add(CFLAG_AllowTypedUAVLoads);
@@ -302,6 +302,7 @@ class FReflectionEnvironmentSkyLightingPS : public FGlobalShader
 	class FSkyShadowing : SHADER_PERMUTATION_BOOL("APPLY_SKY_SHADOWING");
 	class FReflectionSourceType : SHADER_PERMUTATION_ENUM_CLASS("REFLECTION_SOURCE_TYPE", EReflectionSourceType);
 	class FSubstrateTileType : SHADER_PERMUTATION_INT("SUBSTRATE_TILETYPE", 4);
+	class FHairComplexTransmittance : SHADER_PERMUTATION_BOOL("USE_HAIR_COMPLEX_TRANSMITTANCE");
 
 	using FPermutationDomain = TShaderPermutationDomain<
 		FHasBoxCaptures,
@@ -311,7 +312,8 @@ class FReflectionEnvironmentSkyLightingPS : public FGlobalShader
 		FDynamicSkyLight,
 		FSkyShadowing,
 		FReflectionSourceType,
-		FSubstrateTileType>;
+		FSubstrateTileType,
+		FHairComplexTransmittance>;
 
 	static EReflectionSourceType GetReflectionSourceType(bool bIsLumenStandalone, bool bIsTiledSSR)
 	{
@@ -349,7 +351,7 @@ class FReflectionEnvironmentSkyLightingPS : public FGlobalShader
 		return PermutationVector;
 	}
 
-	static FPermutationDomain BuildPermutationVector(const FViewInfo& View, bool bBoxCapturesOnly, bool bSphereCapturesOnly, bool bSupportDFAOIndirectOcclusion, bool bEnableSkyLight, bool bEnableDynamicSkyLight, bool bApplySkyShadowing, bool bLumenStandaloneReflections, bool bIsTiledSSR, ESubstrateTileType TileType)
+	static FPermutationDomain BuildPermutationVector(const FViewInfo& View, bool bBoxCapturesOnly, bool bSphereCapturesOnly, bool bSupportDFAOIndirectOcclusion, bool bEnableSkyLight, bool bEnableDynamicSkyLight, bool bApplySkyShadowing, bool bLumenStandaloneReflections, bool bIsTiledSSR, bool bHasHairCardsOrMeshes, ESubstrateTileType TileType)
 	{
 		FPermutationDomain PermutationVector;
 
@@ -361,6 +363,7 @@ class FReflectionEnvironmentSkyLightingPS : public FGlobalShader
 		PermutationVector.Set<FSkyShadowing>(bApplySkyShadowing);
 		PermutationVector.Set<FReflectionSourceType>(GetReflectionSourceType(bLumenStandaloneReflections,bIsTiledSSR));
 		PermutationVector.Set<FSubstrateTileType>(0);
+		PermutationVector.Set<FHairComplexTransmittance>(bHasHairCardsOrMeshes);
 		if (Substrate::IsSubstrateEnabled())
 		{
 			check(TileType <= ESubstrateTileType::EComplexSpecial);
@@ -385,7 +388,6 @@ class FReflectionEnvironmentSkyLightingPS : public FGlobalShader
 		FGlobalShader::ModifyCompilationEnvironment(Parameters, OutEnvironment);
 		OutEnvironment.SetDefine(TEXT("MAX_CAPTURES"), GetMaxNumReflectionCaptures(Parameters.Platform));
 		OutEnvironment.SetDefine(TEXT("SUPPORTS_ANISOTROPIC_MATERIALS"), FDataDrivenShaderPlatformInfo::GetSupportsAnisotropicMaterials(Parameters.Platform));
-		OutEnvironment.SetDefine(TEXT("USE_HAIR_COMPLEX_TRANSMITTANCE"), IsHairStrandsSupported(EHairStrandsShaderType::All, Parameters.Platform) ? 1u : 0u);
 		OutEnvironment.SetDefine(TEXT("SUBSTRATE_GLINTS_IS"), 0);
 		OutEnvironment.CompilerFlags.Add(CFLAG_StandardOptimization);
 		FForwardLightingParameters::ModifyCompilationEnvironment(Parameters.Platform, OutEnvironment);
@@ -1396,6 +1398,9 @@ void FDeferredShadingSceneRenderer::RenderDiffuseIndirectAndAmbientOcclusion(
 			PermutationVector.Set<FDiffuseIndirectCompositePS::FEnableDualSrcBlending>(!bEnableCopyPass);
 			PermutationVector.Set<FDiffuseIndirectCompositePS::FScreenSpaceReflectionTiledComposition>(bRunSSRTiledComposite);
 			
+			const bool bHasHairCardsOrMeshes = View.HairCardsMeshElements.Num() > 0;
+			PermutationVector.Set<FDiffuseIndirectCompositePS::FHairComplexTransmittance>(bHasHairCardsOrMeshes);
+
 			TShaderMapRef<FDiffuseIndirectCompositePS> PixelShader(View.ShaderMap, PermutationVector);
 
 			FRHIBlendState* BlendState;
@@ -1799,11 +1804,12 @@ static void AddSkyReflectionPass(
 
 	// Bind hair data
 	const bool bCheckerboardSubsurfaceRendering = IsSubsurfaceCheckerboardFormat(SceneColorTexture.Target->Desc.Format);
+	const bool bHasHairCardsOrMeshes = View.HairCardsMeshElements.Num() > 0;
 
 	auto PermutationVector = FReflectionEnvironmentSkyLightingPS::BuildPermutationVector(
 		View, bHasBoxCaptures, bHasSphereCaptures, DynamicBentNormalAO != 0.0f,
 		bSkyLight, bDynamicSkyLight, bApplySkyShadowing,
-		bLumenStandaloneReflections, bIsTiledSSR,
+		bLumenStandaloneReflections, bIsTiledSSR, bHasHairCardsOrMeshes,
 		SubstrateTileMaterialType);
 
 	TShaderMapRef<FReflectionEnvironmentSkyLightingPS> PixelShader(View.ShaderMap, PermutationVector);
