Fix firing assertion caused by calling RemoveAt with invalid index (even though count is 0)

diff --git a/Engine/Source/Runtime/Renderer/Private/ShadowSetup.cpp b/Engine/Source/Runtime/Renderer/Private/ShadowSetup.cpp
--- a/Engine/Source/Runtime/Renderer/Private/ShadowSetup.cpp
+++ b/Engine/Source/Runtime/Renderer/Private/ShadowSetup.cpp
@@ -6508,7 +6508,7 @@ void FSceneRenderer::FilterDynamicShadows(FDynamicShadowsTaskData& TaskData)
 		if (MobileMaxVisibleMovableSpotLightShadowsCVar)
 		{
 			int32 MobileMaxVisibleMovableSpotLightShadows = MobileMaxVisibleMovableSpotLightShadowsCVar->GetValueOnRenderThread();
-			MobileDynamicSpotlightShadows.RemoveAt(MobileMaxVisibleMovableSpotLightShadows, FMath::Max(MobileDynamicSpotlightShadows.Num() - MobileMaxVisibleMovableSpotLightShadows, 0), EAllowShrinking::No);
+			MobileDynamicSpotlightShadows.SetNum(FMath::Min(MobileDynamicSpotlightShadows.Num(), MobileMaxVisibleMovableSpotLightShadows), EAllowShrinking::No);
 		}
 
 		ShadowArrays.WholeSceneDirectionalShadows.Append(MobileDynamicSpotlightShadows);
