Fix for emissive clamp to only be applied on material shader for the deferred renderer. Tile renderer can specify other bound for render target rendering to be accurate w.r.t. what the user want to write.

diff --git a/Engine/Shaders/Private/BasePassPixelShader.usf b/Engine/Shaders/Private/BasePassPixelShader.usf
--- a/Engine/Shaders/Private/BasePassPixelShader.usf
+++ b/Engine/Shaders/Private/BasePassPixelShader.usf
@@ -2125,10 +2125,7 @@ void FPixelShaderInOut_MainPS(
 	}
 #endif // OIT_ENABLED
 
-	// Sky dome, or any emissive, materials can result in high luminance values, e.g. the sun disk. 
-	// This Min here is to we make sure pre-exposed luminance remains within the boundaries of fp10 and not cause NaN on some platforms.
-	// We also half that range to also make sure we have room for other additive elements such as bloom, clouds or particle visual effects.
-	Out.MRT[0].xyz = min(Out.MRT[0].xyz, Max10BitsFloat.xxx * 0.5f);
+	Out.MRT[0].xyz = min(Out.MRT[0].xyz, View.MaterialMaxEmissiveValue.xxx);
 
 #if MATERIAL_VIRTUALTEXTURE_FEEDBACK || LIGHTMAP_VT_ENABLED
 #if OUTPUT_PIXEL_DEPTH_OFFSET
diff --git a/Engine/Source/Runtime/Engine/Public/SceneView.h b/Engine/Source/Runtime/Engine/Public/SceneView.h
--- a/Engine/Source/Runtime/Engine/Public/SceneView.h
+++ b/Engine/Source/Runtime/Engine/Public/SceneView.h
@@ -996,6 +996,7 @@ enum ETranslucencyVolumeCascade
 	VIEW_UNIFORM_BUFFER_MEMBER_PER_VIEW(FVector4f, GlintLUTParameters0) \
 	VIEW_UNIFORM_BUFFER_MEMBER_PER_VIEW(FVector4f, GlintLUTParameters1) \
 	VIEW_UNIFORM_BUFFER_MEMBER(int, EnvironmentComponentsFlags) \
+	VIEW_UNIFORM_BUFFER_MEMBER(float, MaterialMaxEmissiveValue) \
 
 /** The uniform shader parameters associated with a view. */
 #define VIEW_UNIFORM_BUFFER_MEMBER(type, identifier) \
diff --git a/Engine/Source/Runtime/Renderer/Private/Renderer.cpp b/Engine/Source/Runtime/Renderer/Private/Renderer.cpp
--- a/Engine/Source/Runtime/Renderer/Private/Renderer.cpp
+++ b/Engine/Source/Runtime/Renderer/Private/Renderer.cpp
@@ -282,6 +282,34 @@ TRDGUniformBufferRef<FBatchedPrimitiveParameters> FRendererModule::CreateSingleP
 	return GraphBuilder.CreateUniformBuffer(BatchedPrimitiveParameters);
 }
 
+static float GetEmissiveMaxValueForPixelFormat(EPixelFormat PixelFormat)
+{
+	switch (PixelFormat)
+	{
+	// R11G11B10
+	case PF_FloatR11G11B10:
+	case PF_FloatRGB:
+		return 64512.0f; // Max10BitsFloat
+
+	// FP16
+	case PF_FloatRGBA:
+	case PF_G16R16F:
+	case PF_G16R16F_FILTER:
+	case PF_R16F:
+	case PF_R16F_FILTER:
+		return FFloat16::MaxF16Float;
+
+	// FP32
+	//case PF_R32_FLOAT:
+	//case PF_G32R32F:
+	//case PF_A32B32G32R32F:
+	//default:
+		// fall through
+	}
+
+	return UE_MAX_FLT; // default for FP32 and all other formats for now
+}
+
 void FRendererModule::DrawTileMesh(FCanvasRenderContext& RenderContext, FMeshPassProcessorRenderState& DrawRenderState, const FSceneView& SceneView, FMeshBatch& Mesh, bool bIsHitTesting, const FHitProxyId& HitProxyId, bool bUse128bitRT)
 {
 	if (!GUsingNullRHI)
@@ -295,6 +323,12 @@ void FRendererModule::DrawTileMesh(FCanvasRenderContext& RenderContext, FMeshPas
 		ViewFamily->AllViews.Add(&View);
 		View.Family = ViewFamily;
 
+		// When rendering tiles, this may be to render data in a URenderTargetTexture. In this case we should not clamp so that all the expected values setup by the artists go through.
+		if (RenderContext.GetRenderTarget())
+		{
+			View.MaterialMaxEmissiveValue = GetEmissiveMaxValueForPixelFormat(RenderContext.GetRenderTarget()->Desc.Format);
+		}
+
 		// Default init of SceneTexturesConfig will take extents from FSceneTextureExtentState.
 		// We want the view extents, so explicitly set that.
 		InitializeSceneTexturesConfig(ViewFamily->SceneTexturesConfig, *ViewFamily);
diff --git a/Engine/Source/Runtime/Renderer/Private/SceneRendering.cpp b/Engine/Source/Runtime/Renderer/Private/SceneRendering.cpp
--- a/Engine/Source/Runtime/Renderer/Private/SceneRendering.cpp
+++ b/Engine/Source/Runtime/Renderer/Private/SceneRendering.cpp
@@ -914,6 +914,12 @@ void FViewInfo::Init()
 	bUsesCustomDepth = false;
 	bUsesCustomStencil = false;
 
+	// Sky dome, or any emissive, materials can result in high luminance values, e.g. the sun disk. 
+	// This Min here is to we make sure pre-exposed luminance remains within the boundaries of fp10 and not cause NaN on some platforms.
+	// We also half that range to also make sure we have room for other additive elements such as bloom, clouds or particle visual effects.
+	const static float Max10BitsFloat = 64512.0f;
+	MaterialMaxEmissiveValue = Max10BitsFloat * 0.5f;
+
 	NumBoxReflectionCaptures = 0;
 	NumSphereReflectionCaptures = 0;
 	FurthestReflectionCaptureDistance = 0;
@@ -1479,6 +1485,8 @@ void FViewInfo::SetupUniformBufferParameters(
 		}
 	}
 
+	ViewUniformShaderParameters.MaterialMaxEmissiveValue = MaterialMaxEmissiveValue;
+
 	// This should probably be in SetupCommonViewUniformBufferParameters, but drags in too many dependencies
 	UpdateNoiseTextureParameters(ViewUniformShaderParameters);
 
diff --git a/Engine/Source/Runtime/Renderer/Private/SceneRendering.h b/Engine/Source/Runtime/Renderer/Private/SceneRendering.h
--- a/Engine/Source/Runtime/Renderer/Private/SceneRendering.h
+++ b/Engine/Source/Runtime/Renderer/Private/SceneRendering.h
@@ -1538,6 +1538,9 @@ public:
 	// Optional stencil dithering optimization during prepasses
 	bool bAllowStencilDither;
 
+	// Max emissive luminance ouput by any material for this view.
+	float MaterialMaxEmissiveValue;
+
 	/** Custom visibility query for view */
 	ICustomVisibilityQuery* CustomVisibilityQuery;
 
