Fixed crash due to raster setup cache not being cleared when a material is updated.

diff --git a/Engine/Source/Runtime/Renderer/Private/RendererScene.cpp b/Engine/Source/Runtime/Renderer/Private/RendererScene.cpp
--- a/Engine/Source/Runtime/Renderer/Private/RendererScene.cpp
+++ b/Engine/Source/Runtime/Renderer/Private/RendererScene.cpp
@@ -1276,10 +1276,7 @@ void FScene::RefreshNaniteRasterBins(FPrimitiveSceneInfo& PrimitiveSceneInfo)
 
 void FScene::ReloadNaniteFixedFunctionBins()
 {
-	for (int32 NanitePass = 0; NanitePass < ENaniteMeshPass::Num; ++NanitePass)
-	{
-		NaniteRasterPipelines[NanitePass].ReloadFixedFunctionBins();
-	}
+	bReloadNaniteFixedFunctionBins = true;
 }
 
 SIZE_T FScene::GetSizeBytes() const
@@ -6393,7 +6390,16 @@ void FScene::Update(FRDGBuilder& GraphBuilder, const FUpdateParameters& Paramete
 
 	{
 		SCOPED_NAMED_EVENT(UpdateStaticMeshes, FColor::Emerald);
-		
+
+		if (bReloadNaniteFixedFunctionBins)
+		{
+			for (int32 NanitePass = 0; NanitePass < ENaniteMeshPass::Num; ++NanitePass)
+			{
+				NaniteRasterPipelines[NanitePass].ReloadFixedFunctionBins();
+			}
+			bReloadNaniteFixedFunctionBins = false;
+		}
+
 		if (bScenesPrimitivesNeedStaticMeshElementUpdate || CachedDefaultBasePassDepthStencilAccess != DefaultBasePassDepthStencilAccess)
 		{
 			// Mark all primitives as needing an update
diff --git a/Engine/Source/Runtime/Renderer/Private/ScenePrivate.h b/Engine/Source/Runtime/Renderer/Private/ScenePrivate.h
--- a/Engine/Source/Runtime/Renderer/Private/ScenePrivate.h
+++ b/Engine/Source/Runtime/Renderer/Private/ScenePrivate.h
@@ -2865,6 +2865,7 @@ public:
 	/** Nanite raster and shading pipelines. These are stored on the scene as they are computed at FPrimitiveSceneInfo::AddToScene time. */
 	FNaniteRasterPipelines  NaniteRasterPipelines[ENaniteMeshPass::Num];
 	FNaniteShadingPipelines NaniteShadingPipelines[ENaniteMeshPass::Num];
+	bool bReloadNaniteFixedFunctionBins = false;
 
 	/** Nanite material visibility references. These are stored on the scene as they are computed at FPrimitiveSceneInfo::AddToScene time. */
 	FNaniteVisibility NaniteVisibility[ENaniteMeshPass::Num];
