Fixed assert when waiting for mesh pass setup tasks from the wrong thread by removing manual thread tagging.

diff --git a/Engine/Source/Runtime/Renderer/Private/MeshDrawCommands.cpp b/Engine/Source/Runtime/Renderer/Private/MeshDrawCommands.cpp
--- a/Engine/Source/Runtime/Renderer/Private/MeshDrawCommands.cpp
+++ b/Engine/Source/Runtime/Renderer/Private/MeshDrawCommands.cpp
@@ -1462,7 +1462,7 @@ void FParallelMeshDrawCommandPass::DispatchPassSetup(
 		// This work needs to be deferred until at least BuildRenderingCommands (to ensure the DynamicPrimitiveCollector is uploaded), so we use the async mechanism either way 
 		auto FinalizeInstanceCullingSetup = [this, Scene](FInstanceCullingContext& InstanceCullingContext)
 		{
-			WaitForMeshPassSetupTask(IsInActualRenderingThread() ? EWaitThread::Render : EWaitThread::Task);
+			WaitForMeshPassSetupTask();
 
 #if DO_CHECK
 			for (const FVisibleMeshDrawCommand& VisibleMeshDrawCommand : TaskContext.MeshDrawCommands)
@@ -1489,20 +1489,20 @@ bool FParallelMeshDrawCommandPass::IsOnDemandShaderCreationEnabled()
 		(GRHISupportsMultithreadedShaderCreation || (bIsMobileRenderer && (!GSupportsParallelRenderingTasksWithSeparateRHIThread && IsRunningRHIInSeparateThread())));
 }
 
-void FParallelMeshDrawCommandPass::WaitForMeshPassSetupTask(EWaitThread WaitThread) const
+void FParallelMeshDrawCommandPass::WaitForMeshPassSetupTask() const
 {
-	if (TaskEventRef.IsValid() && WaitThread != EWaitThread::TaskAlreadyWaited)
+	if (TaskEventRef.IsValid())
 	{
 		// Need to wait on GetRenderThread_Local, as mesh pass setup task can wait on rendering thread inside InitResourceFromPossiblyParallelRendering().
 		QUICK_SCOPE_CYCLE_COUNTER(STAT_WaitForMeshPassSetupTask);
-		FTaskGraphInterface::Get().WaitUntilTaskCompletes(TaskEventRef, WaitThread == EWaitThread::Render ? ENamedThreads::GetRenderThread_Local() : ENamedThreads::AnyThread);
+		TaskEventRef->Wait(IsInActualRenderingThread() ? ENamedThreads::GetRenderThread_Local() : ENamedThreads::AnyThread);
 	}
 }
 
-void FParallelMeshDrawCommandPass::WaitForTasksAndEmpty(EWaitThread WaitThread)
+void FParallelMeshDrawCommandPass::WaitForTasksAndEmpty()
 {
 	// Need to wait in case if someone dispatched sort and draw merge task, but didn't draw it.
-	WaitForMeshPassSetupTask(WaitThread);
+	WaitForMeshPassSetupTask();
 	TaskEventRef = nullptr;
 
 	DumpInstancingStats();
@@ -1702,7 +1702,7 @@ void FParallelMeshDrawCommandPass::DispatchDraw(FParallelCommandListSet* Paralle
 	{
 		QUICK_SCOPE_CYCLE_COUNTER(STAT_MeshPassDrawImmediate);
 
-		WaitForMeshPassSetupTask(IsInActualRenderingThread() ? EWaitThread::Render : EWaitThread::Task);
+		WaitForMeshPassSetupTask();
 
 		if (TaskContext.bUseGPUScene)
 		{
diff --git a/Engine/Source/Runtime/Renderer/Private/MeshDrawCommands.h b/Engine/Source/Runtime/Renderer/Private/MeshDrawCommands.h
--- a/Engine/Source/Runtime/Renderer/Private/MeshDrawCommands.h
+++ b/Engine/Source/Runtime/Renderer/Private/MeshDrawCommands.h
@@ -115,13 +115,6 @@ public:
 class FParallelMeshDrawCommandPass
 {
 public:
-	enum class EWaitThread
-	{
-		Render,
-		Task,
-		TaskAlreadyWaited,
-	};
-
 	FParallelMeshDrawCommandPass()
 		: bHasInstanceCullingDrawParameters(false)
 		, MaxNumDraws(0)
@@ -171,7 +164,7 @@ public:
 	 */
 	void DispatchDraw(FParallelCommandListSet* ParallelCommandListSet, FRHICommandList& RHICmdList, const FInstanceCullingDrawParams* InstanceCullingDrawParams = nullptr) const;
 
-	void WaitForTasksAndEmpty(EWaitThread WaitThread = EWaitThread::Render);
+	void WaitForTasksAndEmpty();
 	void SetDumpInstancingStats(const FString& InPassName);
 	bool HasAnyDraw() const { return MaxNumDraws > 0; }
 
@@ -202,7 +195,7 @@ private:
 	int32 MaxNumDraws;
 
 	void DumpInstancingStats() const;
-	void WaitForMeshPassSetupTask(EWaitThread WaitThread = EWaitThread::Render) const;
+	void WaitForMeshPassSetupTask() const;
 };
 
 RENDERER_API extern void SortAndMergeDynamicPassMeshDrawCommands(
diff --git a/Engine/Source/Runtime/Renderer/Private/SceneRendering.cpp b/Engine/Source/Runtime/Renderer/Private/SceneRendering.cpp
--- a/Engine/Source/Runtime/Renderer/Private/SceneRendering.cpp
+++ b/Engine/Source/Runtime/Renderer/Private/SceneRendering.cpp
@@ -968,11 +968,11 @@ void FViewInfo::Init()
 	ViewLumenSceneData = nullptr;
 }
 
-void FViewInfo::WaitForTasks(FParallelMeshDrawCommandPass::EWaitThread WaitThread)
+void FViewInfo::WaitForTasks()
 {
 	for (int32 MeshDrawIndex = 0; MeshDrawIndex < EMeshPass::Num; MeshDrawIndex++)
 	{
-		ParallelMeshDrawCommandPasses[MeshDrawIndex].WaitForTasksAndEmpty(WaitThread);
+		ParallelMeshDrawCommandPasses[MeshDrawIndex].WaitForTasksAndEmpty();
 	}
 }
 
@@ -2239,7 +2239,7 @@ FViewInfo* FViewInfo::CreateSnapshot() const
 	return Result;
 }
 
-void FViewInfo::DestroyAllSnapshots(FParallelMeshDrawCommandPass::EWaitThread WaitThread)
+void FViewInfo::DestroyAllSnapshots()
 {
 	QUICK_SCOPE_CYCLE_COUNTER(STAT_FViewInfo_DestroyAllSnapshots);
 
@@ -2261,7 +2261,7 @@ void FViewInfo::DestroyAllSnapshots(FParallelMeshDrawCommandPass::EWaitThread Wa
 
 		for (int32 Index = 0; Index < Snapshot->ParallelMeshDrawCommandPasses.Num(); ++Index)
 		{
-			Snapshot->ParallelMeshDrawCommandPasses[Index].WaitForTasksAndEmpty(WaitThread);
+			Snapshot->ParallelMeshDrawCommandPasses[Index].WaitForTasksAndEmpty();
 		}
 
 		for (int i = 0; i < EMeshPass::Num; i++)
@@ -4437,7 +4437,7 @@ static void FinishCleanUp(FRHICommandListImmediate& RHICmdList)
 	FRDGBuilder::WaitForAsyncDeleteTask();
 }
 
-static void DeleteSceneRenderers(const TArray<FSceneRenderer*>& SceneRenderers, FParallelMeshDrawCommandPass::EWaitThread WaitThread)
+static void DeleteSceneRenderers(const TArray<FSceneRenderer*>& SceneRenderers)
 {
 	SCOPED_NAMED_EVENT_TEXT("DeleteSceneRenderer", FColor::Red);
 
@@ -4446,16 +4446,16 @@ static void DeleteSceneRenderers(const TArray<FSceneRenderer*>& SceneRenderers,
 		// Wait for all dispatched shadow mesh draw tasks.
 		for (int32 PassIndex = 0; PassIndex < SceneRenderer->DispatchedShadowDepthPasses.Num(); ++PassIndex)
 		{
-			SceneRenderer->DispatchedShadowDepthPasses[PassIndex]->WaitForTasksAndEmpty(WaitThread);
+			SceneRenderer->DispatchedShadowDepthPasses[PassIndex]->WaitForTasksAndEmpty();
 		}
 
 		for (FViewInfo* View : SceneRenderer->AllViews)
 		{
-			View->WaitForTasks(WaitThread);
+			View->WaitForTasks();
 		}
 	}
 
-	FViewInfo::DestroyAllSnapshots(WaitThread);
+	FViewInfo::DestroyAllSnapshots();
 
 	for (FSceneRenderer* SceneRenderer : SceneRenderers)
 	{
@@ -4470,7 +4470,7 @@ static void WaitForTasksAndDeleteSceneRenderers(FRHICommandListImmediate& RHICmd
 		RHICmdList.ImmediateFlush(EImmediateFlushType::WaitForOutstandingTasksOnly);
 	}
 
-	DeleteSceneRenderers(SceneRenderers, FParallelMeshDrawCommandPass::EWaitThread::Render);
+	DeleteSceneRenderers(SceneRenderers);
 }
 
 void FSceneRenderer::RenderThreadBegin(FRHICommandListImmediate& RHICmdList)
@@ -4603,7 +4603,7 @@ void FSceneRenderer::RenderThreadEnd(FRHICommandListImmediate& RHICmdList, const
 
 			GSceneRenderCleanUpState.Task = FFunctionGraphTask::CreateAndDispatchWhenReady([LocalSceneRenderers = CopyTemp(SceneRenderers)]
 			{
-				DeleteSceneRenderers(LocalSceneRenderers, FParallelMeshDrawCommandPass::EWaitThread::TaskAlreadyWaited);
+				DeleteSceneRenderers(LocalSceneRenderers);
 			}, TStatId(), &CommandListTasks);
 		}
 	}
diff --git a/Engine/Source/Runtime/Renderer/Private/SceneRendering.h b/Engine/Source/Runtime/Renderer/Private/SceneRendering.h
--- a/Engine/Source/Runtime/Renderer/Private/SceneRendering.h
+++ b/Engine/Source/Runtime/Renderer/Private/SceneRendering.h
@@ -1791,9 +1791,9 @@ public:
 	FViewInfo* CreateSnapshot() const;
 
 	/** Destroy all snapshots before we wipe the scene allocator. */
-	static void DestroyAllSnapshots(FParallelMeshDrawCommandPass::EWaitThread WaitThread);
+	static void DestroyAllSnapshots();
 
-	void WaitForTasks(FParallelMeshDrawCommandPass::EWaitThread WaitThread);
+	void WaitForTasks();
 
 	// Get the range in DynamicMeshElements[] for a given PrimitiveIndex
 	// @return range (start is inclusive, end is exclusive)
