Fixed skyAtmosphere not writing its coverage over holdout meshes.

diff --git a/Engine/Shaders/Private/SkyAtmosphere.usf b/Engine/Shaders/Private/SkyAtmosphere.usf
--- a/Engine/Shaders/Private/SkyAtmosphere.usf
+++ b/Engine/Shaders/Private/SkyAtmosphere.usf
@@ -829,10 +829,10 @@ float4 PrepareOutput(float3 PreExposedLuminance, float3 Transmittance = float3(1
 	float4 LuminanceAlpha = float4(SafePreExposedLuminance, GreyScaleTransmittance);
 	
 #if POST_PROCESS_ALPHA
-	if (IsSkyAtmosphereHoldout(View.EnvironmentComponentsFlags) && View.RenderingReflectionCaptureMask == 0.0f)
+	if (View.RenderingReflectionCaptureMask == 0.0f)
 	{
-		LuminanceAlpha.rgb = 0.0f;
-		LuminanceAlpha.a = 1 - GreyScaleTransmittance;
+		LuminanceAlpha.rgb	= IsSkyAtmosphereHoldout(View.EnvironmentComponentsFlags) ? 0.0f : LuminanceAlpha.rgb;
+		LuminanceAlpha.a	= 1 - GreyScaleTransmittance;
 	}
 #endif
 	
@@ -955,7 +955,7 @@ void RenderSkyAtmosphereRayMarchingPS(
 
 		float3 SkyGreyTransmittance = 1.0f;
 	#if POST_PROCESS_ALPHA
-		if (IsSkyAtmosphereHoldout(View.EnvironmentComponentsFlags) && View.RenderingReflectionCaptureMask == 0.0f)
+		if (View.RenderingReflectionCaptureMask == 0.0f)
 		{
 			SkyGreyTransmittance = SkyLuminanceTransmittance.aaa;
 		}
diff --git a/Engine/Source/Runtime/Renderer/Private/SkyAtmosphereRendering.cpp b/Engine/Source/Runtime/Renderer/Private/SkyAtmosphereRendering.cpp
--- a/Engine/Source/Runtime/Renderer/Private/SkyAtmosphereRendering.cpp
+++ b/Engine/Source/Runtime/Renderer/Private/SkyAtmosphereRendering.cpp
@@ -1787,8 +1787,9 @@ void FSceneRenderer::RenderSkyAtmosphereInternal(
 		const bool bShouldSampleCloudShadow = SkyRC.bShouldSampleCloudShadow && (CloudShadowOnAtmosphereStrength0 > 0.0f || CloudShadowOnAtmosphereStrength1 > 0.0f);
 
 		const bool bFastAerialPerspectiveDepthTest = SkyRC.bFastAerialPerspectiveDepthTest;
-		const bool SkyAtmosphereAlphaHoldOut = IsPostProcessingWithAlphaChannelSupported() && SkyAtmosphereSceneProxy.IsHoldout();
-		const bool bRenderSkyPixel = SkyRC.bRenderSkyPixel || SkyAtmosphereAlphaHoldOut; // In this case we need to write alpha holdout values in the sky pixels.
+		const bool SkyAtmosphereOutputsAlpha = IsPostProcessingWithAlphaChannelSupported();
+		const bool SkyAtmosphereAlphaHoldOut = SkyAtmosphereOutputsAlpha && SkyAtmosphereSceneProxy.IsHoldout();
+		const bool bRenderSkyPixel = SkyRC.bRenderSkyPixel || SkyAtmosphereOutputsAlpha;	// In this case we need to write alpha holdout values in the sky pixels.
 
 		FRenderSkyAtmospherePS::FPermutationDomain PsPermutationVector;
 		PsPermutationVector.Set<FSampleCloudSkyAO>(SkyRC.bShouldSampleCloudSkyAO);
@@ -1866,7 +1867,7 @@ void FSceneRenderer::RenderSkyAtmosphereInternal(
 			{},
 			PsPassParameters,
 			ERDGPassFlags::Raster,
-			[PsPassParameters, VertexShader, PixelShader, Viewport, bFastAerialPerspectiveDepthTest, bRenderSkyPixel, bDisableBlending, StartDepthZ, SkyAtmosphereAlphaHoldOut](FRHICommandList& RHICmdListLambda)
+			[PsPassParameters, VertexShader, PixelShader, Viewport, bFastAerialPerspectiveDepthTest, bRenderSkyPixel, bDisableBlending, StartDepthZ, SkyAtmosphereOutputsAlpha, SkyAtmosphereAlphaHoldOut](FRHICommandList& RHICmdListLambda)
 		{
 			RHICmdListLambda.SetViewport(Viewport.Min.X, Viewport.Min.Y, 0.0f, Viewport.Max.X, Viewport.Max.Y, 1.0f);
 
@@ -1879,15 +1880,23 @@ void FSceneRenderer::RenderSkyAtmosphereInternal(
 			}
 			else
 			{
-				if (SkyAtmosphereAlphaHoldOut)
+				if (SkyAtmosphereOutputsAlpha)
 				{
-					// We might need to write alpha holdout data. Since this is the sky writing to 
-					GraphicsPSOInit.BlendState = TStaticBlendState<CW_RGBA, BO_Add, BF_One, BF_InverseSourceAlpha, BO_Add, BF_One, BF_InverseSourceAlpha>::GetRHI();
+					if (SkyAtmosphereAlphaHoldOut)
+					{
+						// We might need to write alpha holdout data. Since this is the sky writing to 
+						GraphicsPSOInit.BlendState = TStaticBlendState<CW_RGBA, BO_Add, BF_One, BF_InverseSourceAlpha, BO_Add, BF_One, BF_InverseSourceAlpha>::GetRHI();
+					}
+					else
+					{
+						// Maintain alpha untouched
+						GraphicsPSOInit.BlendState = TStaticBlendState<CW_RGBA, BO_Add, BF_One, BF_InverseSourceAlpha, BO_Add, BF_Zero, BF_InverseSourceAlpha>::GetRHI();
+					}
 				}
 				else
 				{
-					// Maintain alpha untouched
-					GraphicsPSOInit.BlendState = TStaticBlendState<CW_RGB, BO_Add, BF_One, BF_SourceAlpha>::GetRHI(); 
+					// Disable alpha writes 
+					GraphicsPSOInit.BlendState = TStaticBlendState<CW_RGB, BO_Add, BF_One, BF_SourceAlpha>::GetRHI();
 				}
 			}
 			if (bFastAerialPerspectiveDepthTest)
