Fixed cloud rdg event name.

diff --git a/Engine/Source/Runtime/Renderer/Private/VolumetricCloudRendering.cpp b/Engine/Source/Runtime/Renderer/Private/VolumetricCloudRendering.cpp
--- a/Engine/Source/Runtime/Renderer/Private/VolumetricCloudRendering.cpp
+++ b/Engine/Source/Runtime/Renderer/Private/VolumetricCloudRendering.cpp
@@ -2503,7 +2503,7 @@ bool FSceneRenderer::RenderVolumetricCloud(
 		FMaterialRenderProxy* CloudVolumeMaterialProxy = CloudSceneProxy.GetCloudVolumeMaterial()->GetRenderProxy();
 		if (CloudVolumeMaterialProxy->GetIncompleteMaterialWithFallback(ViewFamily.GetFeatureLevel()).GetMaterialDomain() == MD_Volume)
 		{
-			RDG_EVENT_SCOPE(GraphBuilder, "%s", bAccumulateAlphaHoldOut ? "VolumetricCloudAlphaHoldout" : "VolumetricCloud");
+			RDG_EVENT_SCOPE(GraphBuilder, "%s", bAccumulateAlphaHoldOut ? TEXT("VolumetricCloudAlphaHoldout") : TEXT("VolumetricCloud"));
 			RDG_GPU_STAT_SCOPE(GraphBuilder, VolumetricCloud);
 			SCOPED_NAMED_EVENT(VolumetricCloud, FColor::Emerald);
 
