VCloud - now account for bPrevTransformsReset.

diff --git a/Engine/Source/Runtime/Renderer/Private/VolumetricRenderTarget.cpp b/Engine/Source/Runtime/Renderer/Private/VolumetricRenderTarget.cpp
--- a/Engine/Source/Runtime/Renderer/Private/VolumetricRenderTarget.cpp
+++ b/Engine/Source/Runtime/Renderer/Private/VolumetricRenderTarget.cpp
@@ -439,7 +439,7 @@ void InitVolumetricRenderTargetForViews(FRDGBuilder& GraphBuilder, TArrayView<FV
 		FVolumetricRenderTargetViewStateData& VolumetricCloudRT = ViewInfo.ViewState->VolumetricCloudRenderTarget;
 
 		// Determine if we are initializing or we should reset the persistent state
-		const bool bCameraCut = ViewInfo.bCameraCut || ViewInfo.bForceCameraVisibilityReset;
+		const bool bCameraCut = ViewInfo.bCameraCut || ViewInfo.bForceCameraVisibilityReset || ViewInfo.bPrevTransformsReset;
 
 		FIntPoint ViewRect = ViewInfo.ViewRect.Size();
 		VolumetricCloudRT.Initialise(	// TODO this is going to reallocate a buffer each time dynamic resolution scaling is applied 
