Fix groom holdout composition.

Change default holdout mode to 1, which is the correct composition mode.

diff --git a/Engine/Source/Runtime/Renderer/Private/HairStrands/HairStrandsComposition.cpp b/Engine/Source/Runtime/Renderer/Private/HairStrands/HairStrandsComposition.cpp
--- a/Engine/Source/Runtime/Renderer/Private/HairStrands/HairStrandsComposition.cpp
+++ b/Engine/Source/Runtime/Renderer/Private/HairStrands/HairStrandsComposition.cpp
@@ -27,7 +27,7 @@ static FAutoConsoleVariableRef CVarHairWriteGBufferData(TEXT("r.HairStrands.Writ
 static int32 GHairStrandsComposeDOFDepth = 1;
 static FAutoConsoleVariableRef CVarHairStrandsComposeDOFDepth(TEXT("r.HairStrands.DOFDepth"), GHairStrandsComposeDOFDepth, TEXT("Compose hair with DOF by lerping hair depth based on its opacity."));
 
-static int32 GHairStrandsHoldoutMode = 0;
+static int32 GHairStrandsHoldoutMode = 1;
 static FAutoConsoleVariableRef CVarHairStrandsHoldoutMode(TEXT("r.HairStrands.HoldoutMode"), GHairStrandsHoldoutMode, TEXT("Change how sample are merged when rendering with holdout."));
 
 /////////////////////////////////////////////////////////////////////////////////////////
