Improve noise on hair transmittance and shadow.

This CL replace the jitter noise by a blue noise function which remove structural noise, and lead to more pleasing results.

diff --git a/Engine/Shaders/Private/HairStrands/HairStrandsDeepShadowMask.usf b/Engine/Shaders/Private/HairStrands/HairStrandsDeepShadowMask.usf
--- a/Engine/Shaders/Private/HairStrands/HairStrandsDeepShadowMask.usf
+++ b/Engine/Shaders/Private/HairStrands/HairStrandsDeepShadowMask.usf
@@ -410,6 +410,10 @@ void MainPS(
 #include "HairStrandsVoxelPageTraversal.ush"
 #endif
 
+#define USE_BLUE_NOISE 1
+#include "../BlueNoise.ush"
+DEFINE_HAIR_BLUE_NOISE_JITTER_FUNCTION
+
 float4				Voxel_TranslatedLightPosition_LightDirection;
 uint				Voxel_MacroGroupCount;
 uint				Voxel_MacroGroupId;
@@ -443,7 +447,11 @@ float3 GetShadowMaskRandom(uint2 PixelPosition)
 	}
 	else if (Voxel_RandomType == 2)
 	{
+		#if USE_BLUE_NOISE
+		Random = GetHairBlueNoiseJitter(PixelPosition, 0 /*SampleIndex*/, 1 /*MaxSampleCount*/, View.StateFrameIndexMod8/*TimeIndex*/).xyz;
+		#else
 		Random = GetHairVoxelJitter2(PixelPosition, View.StateFrameIndexMod8, VirtualVoxel.JitterMode).xyz;
+		#endif
 	}
 	return Random;
 }
diff --git a/Engine/Shaders/Private/HairStrands/HairStrandsDeepTransmittanceMask.usf b/Engine/Shaders/Private/HairStrands/HairStrandsDeepTransmittanceMask.usf
--- a/Engine/Shaders/Private/HairStrands/HairStrandsDeepTransmittanceMask.usf
+++ b/Engine/Shaders/Private/HairStrands/HairStrandsDeepTransmittanceMask.usf
@@ -274,6 +274,22 @@ float4				ShadowChannelMask;
 Buffer<uint>		IndirectArgsBuffer;
 RWBuffer<uint>		OutTransmittanceMask;
 
+#define USE_BLUE_NOISE 1
+#include "../BlueNoise.ush"
+DEFINE_HAIR_BLUE_NOISE_JITTER_FUNCTION
+
+float3 InternalGetHairVoxelJitter(uint2 PixelCoord, uint SampleIndex, uint MaxSampleCount, uint TimeIndex, uint JitterMode)
+{
+	// Blue noise is cheaper to compute and h
+	#if USE_BLUE_NOISE
+	float3 RandomSample = GetHairBlueNoiseJitter(PixelCoord, SampleIndex, MaxSampleCount, TimeIndex).xyz;
+	#else
+	float3 RandomSample = GetHairVoxelJitter(PixelCoord, View.StateFrameIndexMod8, JitterMode);
+	#endif
+
+	return JitterMode > 0 ? RandomSample : float3(0,0,0);
+}
+
 FHairTransmittanceMask ComputeTransmittanceVoxel(
 	const uint2 PixelCoord,
 	const float PixelRadius,
@@ -301,14 +317,11 @@ FHairTransmittanceMask ComputeTransmittanceVoxel(
 	const uint SampleCount = 8; 
 	for (uint SampleIt = 0; SampleIt < SampleCount; ++SampleIt)
 	#else
+	const uint SampleCount = 1; 
 	const uint SampleIt = 0;
 	#endif
 	{
-		#if PERMUTATION_SUPERSAMPLING
-		const float3 SampleRandom = GetHairVoxelJitter(PixelCoord, SampleIt, VirtualVoxel.JitterMode);
-		#else
-		const float3 SampleRandom = GetHairVoxelJitter(PixelCoord, View.StateFrameIndexMod8, VirtualVoxel.JitterMode);
-		#endif
+		float3 SampleRandom = InternalGetHairVoxelJitter(PixelCoord, SampleIt, SampleCount, View.StateFrameIndexMod8, VirtualVoxel.JitterMode);
 	
 		const float PositionBiasScale = 0.5f;
 		const float3 DepthBias = NodeDesc.VoxelWorldSize * (VirtualVoxel.DepthBiasScale_Transmittance * Settings.LightDirection + PositionBiasScale*(SampleRandom*2-1));
diff --git a/Engine/Shaders/Private/HairStrands/HairStrandsVoxelPageTraversal.ush b/Engine/Shaders/Private/HairStrands/HairStrandsVoxelPageTraversal.ush
--- a/Engine/Shaders/Private/HairStrands/HairStrandsVoxelPageTraversal.ush
+++ b/Engine/Shaders/Private/HairStrands/HairStrandsVoxelPageTraversal.ush
@@ -84,6 +84,17 @@ float4 GetHairVoxelJitter2(uint2 PixelCoord, uint Seed, uint JitterMode)
 	return JitterMode > 0 ? float4(Hammersley16(Seed, 8, Seed0), Hammersley16(Seed, 8, Seed1)) : float4(0, 0, 0, 0);
 }
 
+// Use define as this function use global resources not always available where this file is included
+#define DEFINE_HAIR_BLUE_NOISE_JITTER_FUNCTION float4 GetHairBlueNoiseJitter(uint2 PixelCoord, uint SampleIndex, uint MaxSampleCount, uint TimeIndex)\
+{\
+	int2 Offset1 = int2(R2Sequence(SampleIndex) * BlueNoise.Dimensions.xy);\
+	int2 Offset2 = int2(R2Sequence(SampleIndex + MaxSampleCount) * BlueNoise.Dimensions.xy);\
+	float4 RandomSample;\
+	RandomSample.xy = BlueNoiseVec2(PixelCoord + Offset1, TimeIndex);\
+	RandomSample.zw = BlueNoiseVec2(PixelCoord + Offset2, TimeIndex);\
+	return RandomSample;\
+}
+
 ////////////////////////////////////////////////////////////////////////////////////////////////////////////////////////////////
 
 bool ShouldStopTraversal(const FHairTraversalResult InResult, float InHairCountThreshold, bool InUseOpaqueVisibility=true)
diff --git a/Engine/Shaders/Private/VirtualShadowMaps/VirtualShadowMapProjection.usf b/Engine/Shaders/Private/VirtualShadowMaps/VirtualShadowMapProjection.usf
--- a/Engine/Shaders/Private/VirtualShadowMaps/VirtualShadowMapProjection.usf
+++ b/Engine/Shaders/Private/VirtualShadowMaps/VirtualShadowMapProjection.usf
@@ -25,6 +25,9 @@
 #define VOXEL_TRAVERSAL_DEBUG 0
 #include "../HairStrands/HairStrandsVoxelPageCommon.ush"
 #include "../HairStrands/HairStrandsVoxelPageTraversal.ush"
+#if USE_BLUE_NOISE
+DEFINE_HAIR_BLUE_NOISE_JITTER_FUNCTION
+#endif
 #endif
 
 #ifndef DIRECTIONAL_LIGHT
@@ -363,7 +366,12 @@ FVirtualShadowMapSampleResult ProjectLight(
 			float4 Random = 0;
 			if (Result.ShadowFactor > 0)
 			{
+				#if USE_BLUE_NOISE
+				Random = GetHairBlueNoiseJitter(PixelPos, 0 /*SampleIndex*/, 1 /*MaxSampleCount*/, View.StateFrameIndexMod8/*TimeIndex*/);
+				#else
 				Random = ComputeRandom4(PixelPos);
+				#endif
+
 				uint RayIndex = min(Random.w * SMRTRayCount, SMRTRayCount - 1);
 
 			#if DIRECTIONAL_LIGHT
diff --git a/Engine/Source/Runtime/Renderer/Private/HairStrands/HairStrandsTransmittance.cpp b/Engine/Source/Runtime/Renderer/Private/HairStrands/HairStrandsTransmittance.cpp
--- a/Engine/Source/Runtime/Renderer/Private/HairStrands/HairStrandsTransmittance.cpp
+++ b/Engine/Source/Runtime/Renderer/Private/HairStrands/HairStrandsTransmittance.cpp
@@ -199,6 +199,7 @@ class FHairStrandsVoxelTransmittanceMaskCS : public FGlobalShader
 		SHADER_PARAMETER_RDG_UNIFORM_BUFFER(FHairStrandsViewUniformParameters, HairStrands)
 		SHADER_PARAMETER_RDG_UNIFORM_BUFFER(FVirtualVoxelParameters, VirtualVoxel)
 		SHADER_PARAMETER_STRUCT_INCLUDE(FVirtualShadowMapSamplingParameters, VirtualShadowMap)
+		SHADER_PARAMETER_STRUCT_REF(FBlueNoise, BlueNoise)
 
 		SHADER_PARAMETER(FVector4f, TranslatedLightPosition_LightDirection)
 		SHADER_PARAMETER(FVector4f, ShadowChannelMask)
@@ -277,6 +278,9 @@ static FRDGBufferRef AddHairStrandsVoxelTransmittanceMaskPass(
 	Parameters->HairStrands = HairStrands::BindHairStrandsViewUniformParameters(View);
 	Parameters->VirtualVoxel = HairStrands::BindHairStrandsVoxelUniformParameters(View);
 
+	FBlueNoise BlueNoise = GetBlueNoiseGlobalParameters();
+	Parameters->BlueNoise = CreateUniformBufferImmediate(BlueNoise, EUniformBufferUsage::UniformBuffer_SingleDraw);
+
 	ShaderPrint::SetParameters(GraphBuilder, View.ShaderPrintData, Parameters->ShaderPrintParameters);
 
 	const bool bIsSuperSampled = GHairStrandsTransmittanceSuperSampling > 0;
@@ -459,6 +463,7 @@ class FHairStrandsVoxelShadowMaskPS : public FGlobalShader
 		SHADER_PARAMETER_SAMPLER(SamplerState, LinearSampler)
 		SHADER_PARAMETER_SAMPLER(SamplerState, ShadowSampler)
 
+		SHADER_PARAMETER_STRUCT_REF(FBlueNoise, BlueNoise)
 		SHADER_PARAMETER_STRUCT_REF(FViewUniformShaderParameters, ViewUniformBuffer)
 		SHADER_PARAMETER_RDG_UNIFORM_BUFFER(FVirtualVoxelParameters, VirtualVoxel)
 		SHADER_PARAMETER_RDG_UNIFORM_BUFFER(FHairStrandsViewUniformParameters, HairStrands)
@@ -515,6 +520,7 @@ static void AddHairStrandsVoxelShadowMaskPass(
 	const uint32 OutputChannel = bProjectingForForwardShading ? LightSceneInfo->GetDynamicShadowMapChannel() : ~0;
 	const FIntPoint Resolution = OutShadowMask->Desc.Extent;
 	
+	FBlueNoise BlueNoise = GetBlueNoiseGlobalParameters();
 	for (int32 GroupIt=0, GroupCount=View.HairStrandsViewData.MacroGroupDatas.Num(); GroupIt < GroupCount; ++GroupIt)
 	{
 		if (bOnePass && GroupIt > 0)
@@ -533,6 +539,7 @@ static void AddHairStrandsVoxelShadowMaskPass(
 		Parameters->VirtualVoxel = HairStrands::BindHairStrandsVoxelUniformParameters(View);
 		Parameters->LinearSampler = TStaticSamplerState<SF_Bilinear, AM_Clamp, AM_Clamp, AM_Clamp>::GetRHI();
 		Parameters->ShadowSampler = TStaticSamplerState<SF_Point, AM_Clamp, AM_Clamp, AM_Clamp>::GetRHI();
+		Parameters->BlueNoise = CreateUniformBufferImmediate(BlueNoise, EUniformBufferUsage::UniformBuffer_SingleDraw);
 		Parameters->Voxel_TranslatedLightPosition_LightDirection = GetLightTranslatedWorldPositionAndDirection(View, LightSceneInfo);
 		Parameters->Voxel_MacroGroupId	= MacroGroupData.MacroGroupId;
 		Parameters->Voxel_MacroGroupCount = GroupCount;
diff --git a/Engine/Source/Runtime/Renderer/Private/SceneRendering.cpp b/Engine/Source/Runtime/Renderer/Private/SceneRendering.cpp
--- a/Engine/Source/Runtime/Renderer/Private/SceneRendering.cpp
+++ b/Engine/Source/Runtime/Renderer/Private/SceneRendering.cpp
@@ -2596,7 +2596,8 @@ FSceneRenderer::FSceneRenderer(const FSceneViewFamily* InViewFamily, FHitProxyCo
 			ShouldRenderLumenReflections(*ViewInfo) || 
 			ShouldRenderVolumetricCloudWithBlueNoise_GameThread(Scene, *ViewInfo) || 
 			UseVirtualShadowMaps(Scene->GetShaderPlatform(), Scene->GetFeatureLevel()) ||
-			Substrate::IsGlintEnabled(ViewInfo->GetShaderPlatform()))
+			Substrate::IsGlintEnabled(ViewInfo->GetShaderPlatform()) ||
+			IsHairStrandsSupported(EHairStrandsShaderType::Strands, ViewInfo->GetShaderPlatform()))
 		{
 			GEngine->LoadBlueNoiseTexture();
 		}
