VCloud - added a support for a scale factor for cloud tracing mode 2.

diff --git a/Engine/Source/Runtime/Renderer/Private/VolumetricRenderTarget.cpp b/Engine/Source/Runtime/Renderer/Private/VolumetricRenderTarget.cpp
--- a/Engine/Source/Runtime/Renderer/Private/VolumetricRenderTarget.cpp
+++ b/Engine/Source/Runtime/Renderer/Private/VolumetricRenderTarget.cpp
@@ -36,6 +36,11 @@ static TAutoConsoleVariable<int32> CVarVolumetricRenderTargetUpsamplingMode(
 	TEXT("Used in compositing volumetric RT over the scene. [0] bilinear [1] bilinear + jitter [2] nearest + depth test [3] bilinear + jitter + keep closest [4] bilaterial upsampling"),
 	ECVF_RenderThreadSafe | ECVF_Scalability);
 
+static TAutoConsoleVariable<float> CVarVolumetricRenderTargetScale(
+	TEXT("r.VolumetricRenderTarget.Scale"), 1.0f,
+	TEXT("Scales volumetric render target size (1.0 = 100%). Supported by VRT mode 2 only."),
+	ECVF_RenderThreadSafe | ECVF_Scalability);
+
 static TAutoConsoleVariable<int32> CVarVolumetricRenderTargetPreferAsyncCompute(
 	TEXT("r.VolumetricRenderTarget.PreferAsyncCompute"), 0,
 	TEXT("Whether to prefer using async compute to generate volumetric cloud render targets."),
@@ -396,28 +401,41 @@ FRDGTextureRef FVolumetricRenderTargetViewStateData::GetOrCreateSrcVolumetricRec
 	return GraphBuilder.RegisterExternalTexture(VolumetricReconstructRTDepth[1u - CurrentRT]);
 }
 
+static float GetVolumetricBufferResolutionScale(uint32 VRTMode)
+{
+	if (VRTMode == 2) // Only valid for mode 2
+	{
+		return FMath::Clamp(CVarVolumetricRenderTargetScale.GetValueOnAnyThread(), 0.1f, 1.0f);
+	}
+	return 1.0f;
+}
+
 FUintVector4 FVolumetricRenderTargetViewStateData::GetTracingCoordToZbufferCoordScaleBias() const
 {
+	uint32 InvRenderTargetScale = (uint32)FMath::RoundToInt(1.0f / GetVolumetricBufferResolutionScale(Mode));
+	
 	if (Mode == 2 || Mode == 3)
 	{
 		// In this case, the source depth buffer full resolution depth buffer is the full resolution scene one
-		const uint32 CombinedDownsampleFactor = VolumetricReconstructRTDownsampleFactor * VolumetricTracingRTDownsampleFactor;
+		const uint32 CombinedDownsampleFactor = InvRenderTargetScale * VolumetricReconstructRTDownsampleFactor * VolumetricTracingRTDownsampleFactor;
 		return FUintVector4(CombinedDownsampleFactor, CombinedDownsampleFactor,																// Scale is the combined downsample factor
 			CurrentPixelOffset.X * VolumetricReconstructRTDownsampleFactor, CurrentPixelOffset.Y * VolumetricReconstructRTDownsampleFactor);// Each sample will then sample from full res according to reconstructed RT offset times its downsample factor
 	}
 
 	// Otherwise, a half resolution depth buffer is used
 	const uint32 SourceDepthBufferRTDownsampleFactor = 2;
-	const uint32 CombinedDownsampleFactor = VolumetricReconstructRTDownsampleFactor * VolumetricTracingRTDownsampleFactor / SourceDepthBufferRTDownsampleFactor;
-	return FUintVector4( CombinedDownsampleFactor, CombinedDownsampleFactor,									// Scale is the combined downsample factor
+	const uint32 CombinedDownsampleFactor = InvRenderTargetScale * VolumetricReconstructRTDownsampleFactor * VolumetricTracingRTDownsampleFactor / SourceDepthBufferRTDownsampleFactor;
+	return FUintVector4( CombinedDownsampleFactor, CombinedDownsampleFactor,										// Scale is the combined downsample factor
 		CurrentPixelOffset.X * VolumetricReconstructRTDownsampleFactor / VolumetricReconstructRTDownsampleFactor,	// Each sample will then sample from full res according to reconstructed RT offset times its downsample factor
 		CurrentPixelOffset.Y * VolumetricReconstructRTDownsampleFactor / VolumetricReconstructRTDownsampleFactor);
 }
 
 FUintVector4 FVolumetricRenderTargetViewStateData::GetTracingCoordToFullResPixelCoordScaleBias() const
 {
+	uint32 InvRenderTargetScale = (uint32)FMath::RoundToInt(1.0f / GetVolumetricBufferResolutionScale(Mode));
+
 	// In this case, the source depth buffer full resolution depth buffer is the full resolution scene one
-	const uint32 CombinedDownsampleFactor = VolumetricReconstructRTDownsampleFactor * VolumetricTracingRTDownsampleFactor;
+	const uint32 CombinedDownsampleFactor = InvRenderTargetScale * VolumetricReconstructRTDownsampleFactor * VolumetricTracingRTDownsampleFactor;
 	return FUintVector4(CombinedDownsampleFactor, CombinedDownsampleFactor,																// Scale is the combined downsample factor
 		CurrentPixelOffset.X * VolumetricReconstructRTDownsampleFactor, CurrentPixelOffset.Y * VolumetricReconstructRTDownsampleFactor);// Each sample will then sample from full res according to reconstructed RT offset times its downsample factor
 }
@@ -441,10 +459,15 @@ void InitVolumetricRenderTargetForViews(FRDGBuilder& GraphBuilder, TArrayView<FV
 		// Determine if we are initializing or we should reset the persistent state
 		const bool bCameraCut = ViewInfo.bCameraCut || ViewInfo.bForceCameraVisibilityReset || ViewInfo.bPrevTransformsReset;
 
+		int32 VRTMode = FMath::Clamp(CVarVolumetricRenderTargetMode.GetValueOnRenderThread(), 0, 3);
+
 		FIntPoint ViewRect = ViewInfo.ViewRect.Size();
+		const float VolumetricRenderTargetScale = GetVolumetricBufferResolutionScale(VRTMode);
+		ViewRect.X = FMath::RoundToInt((float)ViewRect.X * VolumetricRenderTargetScale);
+		ViewRect.Y = FMath::RoundToInt((float)ViewRect.Y * VolumetricRenderTargetScale);
 		VolumetricCloudRT.Initialise(	// TODO this is going to reallocate a buffer each time dynamic resolution scaling is applied 
 			ViewRect,
-			CVarVolumetricRenderTargetMode.GetValueOnRenderThread(),
+			VRTMode,
 			CVarVolumetricRenderTargetUpsamplingMode.GetValueOnAnyThread(),
 			bCameraCut);
 
@@ -813,7 +836,7 @@ void ComposeVolumetricRenderTargetOverScene(
 		PassParameters->VolumetricDepthTexture = VolumetricDepthTexture;
 		PassParameters->LinearTextureSampler = TStaticSamplerState<SF_Bilinear>::GetRHI();
 		PassParameters->UvOffsetSampleAcceptanceWeight = GetUvNoiseSampleAcceptanceWeight();
-		PassParameters->FullResolutionToVolumetricBufferResolutionScale = GetCompositionFullResolutionToVolumetricBufferResolutionScale(VRTMode);
+		PassParameters->FullResolutionToVolumetricBufferResolutionScale = GetCompositionFullResolutionToVolumetricBufferResolutionScale(VRTMode) * GetVolumetricBufferResolutionScale(VRTMode);
 		PassParameters->SceneTextures = SceneTextures.UniformBuffer;
 		PassParameters->ForwardShadingEnable = bForwardShading ? 1 : 0;
 		GetTextureSafeUvCoordBound(PassParameters->VolumetricTexture, PassParameters->VolumetricTextureValidCoordRect, PassParameters->VolumetricTextureValidUvRect);
@@ -891,7 +914,7 @@ void ComposeVolumetricRenderTargetOverSceneUnderWater(
 		PassParameters->LinearTextureSampler = TStaticSamplerState<SF_Bilinear>::GetRHI();
 		PassParameters->WaterLinearDepthSampler = TStaticSamplerState<SF_Point>::GetRHI();
 		PassParameters->UvOffsetSampleAcceptanceWeight = GetUvNoiseSampleAcceptanceWeight();
-		PassParameters->FullResolutionToVolumetricBufferResolutionScale = GetCompositionFullResolutionToVolumetricBufferResolutionScale(VRTMode);
+		PassParameters->FullResolutionToVolumetricBufferResolutionScale = GetCompositionFullResolutionToVolumetricBufferResolutionScale(VRTMode) * GetVolumetricBufferResolutionScale(VRTMode);
 		PassParameters->FullResolutionToWaterBufferScale = FVector2f(1.0f / WaterPassData.RefractionDownsampleFactor, WaterPassData.RefractionDownsampleFactor);
 		PassParameters->SceneWithoutSingleLayerWaterViewRect = FVector4f(WaterPassViewData.ViewRect.Min.X, WaterPassViewData.ViewRect.Min.Y,
 																		WaterPassViewData.ViewRect.Max.X, WaterPassViewData.ViewRect.Max.Y);
@@ -955,7 +978,7 @@ void ComposeVolumetricRenderTargetOverSceneForVisualization(
 		PassParameters->VolumetricDepthTexture = VolumetricDepthTexture;
 		PassParameters->LinearTextureSampler = TStaticSamplerState<SF_Bilinear>::GetRHI();
 		PassParameters->UvOffsetSampleAcceptanceWeight = GetUvNoiseSampleAcceptanceWeight();
-		PassParameters->FullResolutionToVolumetricBufferResolutionScale = GetCompositionFullResolutionToVolumetricBufferResolutionScale(VRTMode);
+		PassParameters->FullResolutionToVolumetricBufferResolutionScale = GetCompositionFullResolutionToVolumetricBufferResolutionScale(VRTMode) * GetVolumetricBufferResolutionScale(VRTMode);
 		PassParameters->SceneTextures = SceneTextures.UniformBuffer;
 		GetTextureSafeUvCoordBound(PassParameters->VolumetricTexture, PassParameters->VolumetricTextureValidCoordRect, PassParameters->VolumetricTextureValidUvRect);
 
