VCloud - improved snapping of the cloud shadow map to reduce shadow flickering.

diff --git a/Engine/Source/Runtime/Renderer/Private/VolumetricCloudRendering.cpp b/Engine/Source/Runtime/Renderer/Private/VolumetricCloudRendering.cpp
--- a/Engine/Source/Runtime/Renderer/Private/VolumetricCloudRendering.cpp
+++ b/Engine/Source/Runtime/Renderer/Private/VolumetricCloudRendering.cpp
@@ -146,6 +146,11 @@ static TAutoConsoleVariable<float> CVarVolumetricCloudShadowMapSnapLength(
 	TEXT("Snapping size in kilometers of the cloud shadowmap position to avoid flickering."),
 	ECVF_RenderThreadSafe | ECVF_Scalability);
 
+static TAutoConsoleVariable<int32> CVarVolumetricCloudShadowMapSnapToPixelGrid(
+	TEXT("r.VolumetricCloud.ShadowMap.SnapToPixelGrid"), 1,
+	TEXT("Snaps world origin to the shadow map pixel grid. Avoids shimmering with camera movement when r.VolumetricCloud.ShadowMap.SnapLength is small."),
+	ECVF_RenderThreadSafe | ECVF_Scalability);
+
 static TAutoConsoleVariable<float> CVarVolumetricCloudShadowMapRaySampleMaxCount(
 	TEXT("r.VolumetricCloud.ShadowMap.RaySampleMaxCount"), 128.0f,
 	TEXT("The maximum number of samples taken while ray marching shadow rays to evaluate the cloud shadow map."),
@@ -1686,7 +1691,36 @@ void FSceneRenderer::InitVolumetricCloudsForViews(FRDGBuilder& GraphBuilder, boo
 					const FVector3f LightPosition = LookAtPosition - AtmopshericLightDirection * SphereRadius;
 					FReversedZOrthoMatrix ShadowProjectionMatrix(SphereDiameter, SphereDiameter, ZScale, ZOffset);
 					FLookAtMatrix ShadowViewMatrix((FVector)LightPosition, (FVector)LookAtPosition, (FVector)UpVector);
-					CloudGlobalShaderParams.CloudShadowmapTranslatedWorldToLightClipMatrix[LightIndex] = FMatrix44f((TranslatedWorldToWorld * ShadowViewMatrix) * ShadowProjectionMatrix);
+					FMatrix44f ShadowViewProjectionMatrix = FMatrix44f( ShadowViewMatrix * ShadowProjectionMatrix );
+					if (CVarVolumetricCloudShadowMapSnapToPixelGrid.GetValueOnRenderThread() > 0)
+					{
+						// Make a point in the world (here the origin) always snap to pixel grid exactly, to prevent shimmering.
+						FVector4f WorldOrigin = FVector4f::Zero();
+						WorldOrigin.W = 1.0f;
+						FVector4f WorldOriginInClipSpace = ShadowViewProjectionMatrix.TransformPosition(WorldOrigin);
+
+						// Snap to align with pixel grid in clip (-1 to 1) space.
+						const auto SnapComponentToPixelGrid = [](float Component, float Resolution) -> float
+						{
+							float HalfResolution = Resolution * 0.5f;
+							HalfResolution *= 0.5f; // needed to avoid offsetting at half-pixel increments (perhaps RT isn't traced at 1:1 ratio?)
+							return FMath::RoundToFloat(Component * HalfResolution) / HalfResolution;
+						};
+
+						FVector4f SnappedWorldOriginInClipSpace = FVector4f::Zero();
+						SnappedWorldOriginInClipSpace.X = SnapComponentToPixelGrid(WorldOriginInClipSpace.X, CloudShadowmapResolution);
+						SnappedWorldOriginInClipSpace.Y = SnapComponentToPixelGrid(WorldOriginInClipSpace.Y, CloudShadowmapResolution);
+						SnappedWorldOriginInClipSpace.Z = WorldOriginInClipSpace.Z;
+						SnappedWorldOriginInClipSpace.W = 1.0f;
+
+						FVector4f SnapOffset = SnappedWorldOriginInClipSpace - WorldOriginInClipSpace;
+						FTranslationMatrix44f SnapToPixelGridMatrix(SnapOffset);
+
+						// Apply the snap offset
+						ShadowViewProjectionMatrix = ShadowViewProjectionMatrix * SnapToPixelGridMatrix;
+					}
+					CloudGlobalShaderParams.CloudShadowmapTranslatedWorldToLightClipMatrix[LightIndex] = FMatrix44f((FMatrix44f)TranslatedWorldToWorld * ShadowViewProjectionMatrix);
+
 					CloudGlobalShaderParams.CloudShadowmapTranslatedWorldToLightClipMatrixInv[LightIndex] = CloudGlobalShaderParams.CloudShadowmapTranslatedWorldToLightClipMatrix[LightIndex].Inverse();
 					CloudGlobalShaderParams.CloudShadowmapLightDir[LightIndex] = AtmopshericLightDirection;
 					CloudGlobalShaderParams.CloudShadowmapLightPos[LightIndex] = LightPosition;
