Fix mismatch between RayTracingGeometry and which LOD cached mesh commands are used when using GetValidStaticRayTracingGeometry(...) and the target LOD is not valid.
- Modified FPrimitiveSceneInfo::GetValidStaticRayTracingGeometry(...) to take LODIndex by reference and modify it to correspond to the valid RayTracingGeometry LOD index.
- Modified FStaticMeshRenderData to fill LODIndex member of each RayTracingGeometry.

diff --git a/Engine/Source/Runtime/Engine/Private/SkeletalRenderStatic.cpp b/Engine/Source/Runtime/Engine/Private/SkeletalRenderStatic.cpp
--- a/Engine/Source/Runtime/Engine/Private/SkeletalRenderStatic.cpp
+++ b/Engine/Source/Runtime/Engine/Private/SkeletalRenderStatic.cpp
@@ -112,6 +112,8 @@ void FSkeletalMeshObjectStatic::InitResources(USkinnedMeshComponent* InMeshCompo
 								Initializer.Type = ERayTracingGeometryInitializerType::StreamingDestination;
 							}
 
+							RayTracingGeometry.LODIndex = LODIndex;
+
 							RayTracingGeometry.SetInitializer(Initializer);
 							RayTracingGeometry.InitResource(RHICmdList);
 
diff --git a/Engine/Source/Runtime/Engine/Private/StaticMesh.cpp b/Engine/Source/Runtime/Engine/Private/StaticMesh.cpp
--- a/Engine/Source/Runtime/Engine/Private/StaticMesh.cpp
+++ b/Engine/Source/Runtime/Engine/Private/StaticMesh.cpp
@@ -1917,7 +1917,7 @@ void FStaticMeshRenderData::InitResources(ERHIFeatureLevel::Type InFeatureLevel,
 	}
 
 #if RHI_RAYTRACING
-	if (IsRayTracingAllowed())
+	if (IsRayTracingAllowed()) // TODO: Could move most of this to FStaticMeshLODResources::InitResources
 	{
 		ENQUEUE_RENDER_COMMAND(InitRayTracingGeometryForInlinedLODs)(
 			[this](FRHICommandListImmediate& RHICmdList)
@@ -1936,6 +1936,7 @@ void FStaticMeshRenderData::InitResources(ERHIFeatureLevel::Type InFeatureLevel,
 							LODResources[LODIndex].RayTracingGeometry.Initializer.Type = ERayTracingGeometryInitializerType::StreamingDestination;
 						}
 
+						LODResources[LODIndex].RayTracingGeometry.LODIndex = LODIndex;
 						LODResources[LODIndex].RayTracingGeometry.InitResource(RHICmdList);
 					}
 				}
diff --git a/Engine/Source/Runtime/Renderer/Private/PrimitiveSceneInfo.cpp b/Engine/Source/Runtime/Renderer/Private/PrimitiveSceneInfo.cpp
--- a/Engine/Source/Runtime/Renderer/Private/PrimitiveSceneInfo.cpp
+++ b/Engine/Source/Runtime/Renderer/Private/PrimitiveSceneInfo.cpp
@@ -398,11 +398,11 @@ bool FPrimitiveSceneInfo::IsCachedRayTracingGeometryValid() const
 	return false;
 }
 
-FRayTracingGeometry* FPrimitiveSceneInfo::GetStaticRayTracingGeometry(int32 LodLevel) const
+FRayTracingGeometry* FPrimitiveSceneInfo::GetStaticRayTracingGeometry(int8 LODIndex) const
 {
-	if (StaticRayTracingGeometries.Num() > LodLevel)
+	if (LODIndex < StaticRayTracingGeometries.Num())
 	{
-		return StaticRayTracingGeometries[LodLevel];
+		return StaticRayTracingGeometries[LODIndex];
 	}
 	else
 	{
@@ -410,20 +410,20 @@ FRayTracingGeometry* FPrimitiveSceneInfo::GetStaticRayTracingGeometry(int32 LodL
 	}
 }
 
-FRayTracingGeometry* FPrimitiveSceneInfo::GetValidStaticRayTracingGeometry(int32 MinLodLevel) const
+FRayTracingGeometry* FPrimitiveSceneInfo::GetValidStaticRayTracingGeometry(int8& InOutLODIndex) const
 {
 	// TODO: Move HasPendingBuildRequest() / BoostBuildPriority() out of this function
 
-	for (int32 Index = MinLodLevel; Index < StaticRayTracingGeometries.Num(); ++Index)
+	for (; InOutLODIndex < StaticRayTracingGeometries.Num(); ++InOutLODIndex)
 	{
-		if (StaticRayTracingGeometries[Index]->HasPendingBuildRequest())
+		if (StaticRayTracingGeometries[InOutLODIndex]->HasPendingBuildRequest())
 		{
-			ensure(StaticRayTracingGeometries[Index]->IsValid());
-			StaticRayTracingGeometries[Index]->BoostBuildPriority();
+			ensure(StaticRayTracingGeometries[InOutLODIndex]->IsValid());
+			StaticRayTracingGeometries[InOutLODIndex]->BoostBuildPriority();
 		}
-		else if (StaticRayTracingGeometries[Index]->IsValid())
+		else if (StaticRayTracingGeometries[InOutLODIndex]->IsValid())
 		{
-			return StaticRayTracingGeometries[Index];
+			return StaticRayTracingGeometries[InOutLODIndex];
 		}
 	}
 
diff --git a/Engine/Source/Runtime/Renderer/Private/RayTracing/RayTracing.cpp b/Engine/Source/Runtime/Renderer/Private/RayTracing/RayTracing.cpp
--- a/Engine/Source/Runtime/Renderer/Private/RayTracing/RayTracing.cpp
+++ b/Engine/Source/Runtime/Renderer/Private/RayTracing/RayTracing.cpp
@@ -548,6 +548,8 @@ namespace RayTracing
 								return;
 							}
 
+							check(RayTracingGeometry->LODIndex == LODIndex);
+
 							// Sometimes LODIndex is out of range because it is clamped by ClampToFirstLOD, like the requested LOD is being streamed in and hasn't been available
 							// According to InitViews, we should hide the static mesh instance
 							if (SceneInfo->CachedRayTracingMeshCommandIndicesPerLOD.IsValidIndex(LODIndex))
diff --git a/Engine/Source/Runtime/Renderer/Public/PrimitiveSceneInfo.h b/Engine/Source/Runtime/Renderer/Public/PrimitiveSceneInfo.h
--- a/Engine/Source/Runtime/Renderer/Public/PrimitiveSceneInfo.h
+++ b/Engine/Source/Runtime/Renderer/Public/PrimitiveSceneInfo.h
@@ -548,8 +548,8 @@ public:
 
 	bool IsCachedRayTracingGeometryValid() const;
 
-	RENDERER_API FRayTracingGeometry* GetStaticRayTracingGeometry(int32 LodLevel) const;
-	RENDERER_API FRayTracingGeometry* GetValidStaticRayTracingGeometry(int32 MinLodLevel) const;
+	RENDERER_API FRayTracingGeometry* GetStaticRayTracingGeometry(int8 LODIndex) const;
+	RENDERER_API FRayTracingGeometry* GetValidStaticRayTracingGeometry(int8& InOutLODIndex) const;
 
 	UE_DEPRECATED(5.4, "GetStaticRayTracingGeometryInstance has been deprecated. Use GetStaticRayTracingGeometry / GetValidStaticRayTracingGeometry instead.")
 	RENDERER_API FRHIRayTracingGeometry* GetStaticRayTracingGeometryInstance(int LodLevel) const;
