VCloud - fixed shadow frustum debug due to translated view space.

diff --git a/Engine/Source/Runtime/Renderer/Private/VolumetricCloudRendering.cpp b/Engine/Source/Runtime/Renderer/Private/VolumetricCloudRendering.cpp
--- a/Engine/Source/Runtime/Renderer/Private/VolumetricCloudRendering.cpp
+++ b/Engine/Source/Runtime/Renderer/Private/VolumetricCloudRendering.cpp
@@ -2925,6 +2925,8 @@ bool FSceneRenderer::RenderVolumetricCloud(
 						}
 					};
 
+					FMatrix44f TranslatedWorldToWorld = FMatrix44f(FTranslationMatrix(-ViewInfo.ViewMatrices.GetPreViewTranslation()));
+
 					if (DebugCloudShadowMap)
 					{
 						const int DebugLightIndex = 0;	// only debug atmospheric light 0 for now
@@ -2939,7 +2941,7 @@ bool FSceneRenderer::RenderVolumetricCloud(
 							Canvas.DrawShadowedString(0.05f, ViewPortHeight * 0.4f, *Text, GetStatsFont(), TextColor);
 						});
 
-						DrawFrustumWireframe(&ShadowFrustumPDI, FMatrix(VolumetricCloudParams.VolumetricCloud.CloudShadowmapTranslatedWorldToLightClipMatrixInv[DebugLightIndex]), FColor::Orange, 0);
+						DrawFrustumWireframe(&ShadowFrustumPDI, FMatrix(VolumetricCloudParams.VolumetricCloud.CloudShadowmapTranslatedWorldToLightClipMatrixInv[DebugLightIndex] * TranslatedWorldToWorld), FColor::Orange, 0);
 						FDrawDebugCloudShadowCS::FParameters* Parameters = GraphBuilder.AllocParameters<FDrawDebugCloudShadowCS::FParameters>();
 						Parameters->CloudTracedTexture = CloudRC.VolumetricCloudShadowTexture[DebugLightIndex];
 						Parameters->CloudTraceDirection = VolumetricCloudParams.VolumetricCloud.CloudShadowmapLightDir[DebugLightIndex];
@@ -2954,7 +2956,7 @@ bool FSceneRenderer::RenderVolumetricCloud(
 
 					if (DebugCloudSkyAO && ViewInfo.VolumetricCloudSkyAO != nullptr)
 					{
-						DrawFrustumWireframe(&ShadowFrustumPDI, FMatrix(VolumetricCloudParams.VolumetricCloud.CloudSkyAOTranslatedWorldToLightClipMatrixInv), FColor::Blue, 0);
+						DrawFrustumWireframe(&ShadowFrustumPDI, FMatrix(VolumetricCloudParams.VolumetricCloud.CloudSkyAOTranslatedWorldToLightClipMatrixInv * TranslatedWorldToWorld), FColor::Blue, 0);
 						FDrawDebugCloudShadowCS::FParameters* Parameters = GraphBuilder.AllocParameters<FDrawDebugCloudShadowCS::FParameters>();
 						Parameters->CloudTracedTexture = ViewInfo.VolumetricCloudSkyAO;
 						Parameters->CloudTextureSizeInvSize = VolumetricCloudParams.VolumetricCloud.CloudSkyAOSizeInvSize;
