VCloud - made the shadow frustum square by default + added ay to override the light shadow tracing start distance (shadow frustum depth).

diff --git a/Engine/Source/Runtime/Renderer/Private/VolumetricCloudRendering.cpp b/Engine/Source/Runtime/Renderer/Private/VolumetricCloudRendering.cpp
--- a/Engine/Source/Runtime/Renderer/Private/VolumetricCloudRendering.cpp
+++ b/Engine/Source/Runtime/Renderer/Private/VolumetricCloudRendering.cpp
@@ -156,6 +156,11 @@ static TAutoConsoleVariable<float> CVarVolumetricCloudShadowMapRaySampleMaxCount
 	TEXT("The maximum number of samples taken while ray marching shadow rays to evaluate the cloud shadow map."),
 	ECVF_RenderThreadSafe | ECVF_Scalability);
 
+static TAutoConsoleVariable<float> CVarVolumetricCloudShadowMapLightDistanceOverride(
+	TEXT("r.VolumetricCloud.ShadowMap.LightDistanceOverride"), 0.0f,
+	TEXT("When > 0, overrides the volumetric cloud shadow map light distance from the ground (km). This would usually be based on the shadow map extent."),
+	ECVF_RenderThreadSafe | ECVF_Scalability );
+
 static TAutoConsoleVariable<float> CVarVolumetricCloudShadowMapRaySampleHorizonMultiplier(
 	TEXT("r.VolumetricCloud.ShadowMap.RaySampleHorizonMultiplier"), 2.0f,
 	TEXT("The multipler on the sample count applied when the atmospheric light reach the horizon. Less pixels in the shadow map need to be traced, but rays need to travel a lot longer."),
@@ -1658,9 +1663,11 @@ void FSceneRenderer::InitVolumetricCloudsForViews(FRDGBuilder& GraphBuilder, boo
 					const FVector3f UpVector = FMath::Abs(FVector3f::DotProduct(AtmopshericLightDirection, FVector3f::UpVector)) > 0.99f ? FVector3f::ForwardVector : FVector3f::UpVector;
 
 					const float SphereRadius = GetVolumetricCloudShadowMapExtentKm(AtmosphericLight) * KilometersToCentimeters;
+					const float LightDistanceOverride = CVarVolumetricCloudShadowMapLightDistanceOverride.GetValueOnRenderThread();
+					const float LightDistance = LightDistanceOverride > 0.0f ? (LightDistanceOverride * KilometersToCentimeters) : SphereRadius * 2.0f;
 					const float SphereDiameter = SphereRadius * 2.0f;
 					const float NearPlane = 0.0f;
-					const float FarPlane = SphereDiameter;
+					const float FarPlane = LightDistance * 2.0f;
 					const float ZScale = 1.0f / (FarPlane - NearPlane);
 					const float ZOffset = -NearPlane;
 					FMatrix TranslatedWorldToWorld = FMatrix::Identity;
@@ -1678,7 +1685,7 @@ void FSceneRenderer::InitVolumetricCloudsForViews(FRDGBuilder& GraphBuilder, boo
 						PlanetToCameraNormUp = LookAtPosition;
 						LookAtPosition = (CloudGlobalShaderParams.CloudLayerCenterKm + LookAtPosition * PlanetRadiusKm) * KilometersToCentimeters;
 						// Light position is positioned away from the look at position in the light direction according to the shadowmap radius.
-						const FVector3f LightPosition = LookAtPosition - AtmopshericLightDirection * SphereRadius;
+						const FVector3f LightPosition = LookAtPosition - AtmopshericLightDirection * LightDistance;
 
 						float WorldSizeSnap = CVarVolumetricCloudShadowMapSnapLength.GetValueOnAnyThread() * KilometersToCentimeters;
 						LookAtPosition.X = (FMath::FloorToFloat((LookAtPosition.X + 0.5f * WorldSizeSnap) / WorldSizeSnap)) * WorldSizeSnap; // offset by 0.5 to not snap around origin
@@ -1688,7 +1695,7 @@ void FSceneRenderer::InitVolumetricCloudsForViews(FRDGBuilder& GraphBuilder, boo
 						TranslatedWorldToWorld = FTranslationMatrix(-View.ViewMatrices.GetPreViewTranslation());
 					}
 
-					const FVector3f LightPosition = LookAtPosition - AtmopshericLightDirection * SphereRadius;
+					const FVector3f LightPosition = LookAtPosition - AtmopshericLightDirection * LightDistance;
 					FReversedZOrthoMatrix ShadowProjectionMatrix(SphereDiameter, SphereDiameter, ZScale, ZOffset);
 					FLookAtMatrix ShadowViewMatrix((FVector)LightPosition, (FVector)LookAtPosition, (FVector)UpVector);
 					FMatrix44f ShadowViewProjectionMatrix = FMatrix44f( ShadowViewMatrix * ShadowProjectionMatrix );
