Add option to reduce amount of PSOs to precache:
- PrecacheAlphaColorChannel optional for base pass shaders (used during planar reflections and scene captures and generates PSO with different Render Target info)
- Projected Shadow PSOs optional
- Add PSO precaching mode for custom depth: fully disabled, use bCustomDepth or also precache for all primitives which render to depth (custom depth can be enabled any point in time without precaching again)
- PSO precaching with reverse culling optional for niagara assets: PSO precaching done at asset level and negative determinent can't be known then which influences the culling mode, but could double amount of Niagara PSOs to compile

diff --git a/Engine/Plugins/FX/Niagara/Source/Niagara/Private/NiagaraRendererProperties.cpp b/Engine/Plugins/FX/Niagara/Source/Niagara/Private/NiagaraRendererProperties.cpp
--- a/Engine/Plugins/FX/Niagara/Source/Niagara/Private/NiagaraRendererProperties.cpp
+++ b/Engine/Plugins/FX/Niagara/Source/Niagara/Private/NiagaraRendererProperties.cpp
@@ -24,6 +24,14 @@
 
 #define LOCTEXT_NAMESPACE "UNiagaraRendererProperties"
 
+static int GNiagaraPSOPrecacheReverseCulling = 1;
+static FAutoConsoleVariableRef CVarPSOPrecacheProjectedShadows(
+	TEXT("fx.Niagara.PSOPrecache.ReverseCulling"),
+	GNiagaraPSOPrecacheReverseCulling,
+	TEXT("Also Precache PSOs with with reverse culling set when not 2 sided. (default 1)"),
+	ECVF_Default
+);
+
 #if WITH_EDITORONLY_DATA
 int32 GNiagaraRendererCookOutStaticEnabledBinding = 1;
 static FAutoConsoleVariableRef CVarNiagaraRendererCookOutStaticEnabledBinding(
@@ -853,6 +861,39 @@ bool UNiagaraRendererProperties::PopulateRequiredBindings(FNiagaraParameterStore
 	return bAnyAdded;
 }
 
+void UNiagaraRendererProperties::CollectPSOPrecacheData(FNiagaraEmitterInstance* EmitterInstance, FMaterialInterfacePSOPrecacheParamsList& MaterialInterfacePSOPrecacheParamsList) const
+{
+	const FVertexFactoryType* VFType = GetVertexFactoryType();
+	if (VFType == nullptr)
+	{
+		return;
+	}
+
+	FMaterialInterfacePSOPrecacheParams NewEntry;
+	NewEntry.PSOPrecacheParams.SetMobility(EComponentMobility::Movable);
+	NewEntry.PSOPrecacheParams.bDisableBackFaceCulling = IsBackfaceCullingDisabled();
+
+	UNiagaraRendererProperties::FPSOPrecacheParamsList PSOPrecacheParamsList;
+	CollectPSOPrecacheData(EmitterInstance, PSOPrecacheParamsList);
+
+	for (UNiagaraRendererProperties::FPSOPrecacheParams& PSOPrecacheParams : PSOPrecacheParamsList)
+	{
+		NewEntry.MaterialInterface = PSOPrecacheParams.MaterialInterface;
+		NewEntry.VertexFactoryDataList = PSOPrecacheParams.VertexFactoryDataList;
+
+		NewEntry.PSOPrecacheParams.bReverseCulling = false;
+		AddMaterialInterfacePSOPrecacheParamsToList(NewEntry, MaterialInterfacePSOPrecacheParamsList);
+
+		// Also precache with reverse culling if not two sided because we don't know of the component using the asset will have negative determinant
+		if (!NewEntry.PSOPrecacheParams.bDisableBackFaceCulling && GNiagaraPSOPrecacheReverseCulling > 0)
+		{
+			NewEntry.PSOPrecacheParams.bReverseCulling = true;
+			AddMaterialInterfacePSOPrecacheParamsToList(NewEntry, MaterialInterfacePSOPrecacheParamsList);
+		}
+	}
+}
+
+
 bool UNiagaraRendererProperties::NeedsLoadForTargetPlatform(const ITargetPlatform* TargetPlatform) const
 {
 	// only keep enabled renderers that are parented to valid emitters
diff --git a/Engine/Plugins/FX/Niagara/Source/Niagara/Private/NiagaraSystem.cpp b/Engine/Plugins/FX/Niagara/Source/Niagara/Private/NiagaraSystem.cpp
--- a/Engine/Plugins/FX/Niagara/Source/Niagara/Private/NiagaraSystem.cpp
+++ b/Engine/Plugins/FX/Niagara/Source/Niagara/Private/NiagaraSystem.cpp
@@ -1339,12 +1339,10 @@ void UNiagaraSystem::PrecachePSOs()
 		return;
 	}
 
-	FMaterialInterfacePSOPrecacheParamsList MaterialInterfacePSOPrecacheParamsList;
-
-	FMaterialInterfacePSOPrecacheParams NewEntry;
-	NewEntry.PSOPrecacheParams.SetMobility(EComponentMobility::Movable);
-	NewEntry.PSOPrecacheParams.bRenderCustomDepth = true;
+	// Don't have an instance yet to retrieve the possible material override data from
+	FNiagaraEmitterInstance* EmitterInstance = nullptr;
 
+	FMaterialInterfacePSOPrecacheParamsList MaterialInterfacePSOPrecacheParamsList;
 	for (const FNiagaraEmitterHandle& EmitterHandle : GetEmitterHandles())
 	{
 		FVersionedNiagaraEmitterData* EmitterData = EmitterHandle.GetEmitterData();
@@ -1353,35 +1351,7 @@ void UNiagaraSystem::PrecachePSOs()
 			EmitterData->ForEachEnabledRenderer(
 				[&](UNiagaraRendererProperties* Properties)
 				{
-					const FVertexFactoryType* VFType = Properties->GetVertexFactoryType();
-					if (VFType == nullptr)
-					{
-						return;
-					}
-
-					NewEntry.PSOPrecacheParams.bDisableBackFaceCulling = Properties->IsBackfaceCullingDisabled();
-
-					// Don't have an instance yet to retrieve the possible material override data from
-					FNiagaraEmitterInstance* EmitterInstance = nullptr;
-
-					UNiagaraRendererProperties::FPSOPrecacheParamsList PSOPrecacheParamsList;
-					Properties->CollectPSOPrecacheData(EmitterInstance, PSOPrecacheParamsList);
-
-					for (UNiagaraRendererProperties::FPSOPrecacheParams& PSOPrecacheParams: PSOPrecacheParamsList)
-					{
-						NewEntry.MaterialInterface = PSOPrecacheParams.MaterialInterface;
-						NewEntry.VertexFactoryDataList = PSOPrecacheParams.VertexFactoryDataList;
-
-						NewEntry.PSOPrecacheParams.bReverseCulling = false;
-						AddMaterialInterfacePSOPrecacheParamsToList(NewEntry, MaterialInterfacePSOPrecacheParamsList);
-
-						// Also precache with reverse culling if not two sided because we don't know of the component using the asset will have negative determinant
-						if (!NewEntry.PSOPrecacheParams.bDisableBackFaceCulling)
-						{
-							NewEntry.PSOPrecacheParams.bReverseCulling = true;
-							AddMaterialInterfacePSOPrecacheParamsToList(NewEntry, MaterialInterfacePSOPrecacheParamsList);
-						}
-					}
+					Properties->CollectPSOPrecacheData(EmitterInstance, MaterialInterfacePSOPrecacheParamsList);
 				}
 			);
 		}
@@ -1390,7 +1360,6 @@ void UNiagaraSystem::PrecachePSOs()
 	LaunchPSOPrecaching(MaterialInterfacePSOPrecacheParamsList);
 }
 
-
 #if WITH_EDITORONLY_DATA
 
 UNiagaraEditorDataBase* UNiagaraSystem::GetEditorData()
diff --git a/Engine/Plugins/FX/Niagara/Source/Niagara/Private/NiagaraSystemInstanceController.cpp b/Engine/Plugins/FX/Niagara/Source/Niagara/Private/NiagaraSystemInstanceController.cpp
--- a/Engine/Plugins/FX/Niagara/Source/Niagara/Private/NiagaraSystemInstanceController.cpp
+++ b/Engine/Plugins/FX/Niagara/Source/Niagara/Private/NiagaraSystemInstanceController.cpp
@@ -109,39 +109,13 @@ void FNiagaraSystemInstanceController::CollectPSOPrecacheData(const FPSOPrecache
 
 	FMaterialInterfacePSOPrecacheParams NewEntry;
 	NewEntry.PSOPrecacheParams.SetMobility(EComponentMobility::Movable);
-	NewEntry.PSOPrecacheParams.bRenderCustomDepth = true;
 
 	for (const FNiagaraEmitterInstanceRef& EmitterInst : SystemInstance->GetEmitters())
 	{
 		EmitterInst->ForEachEnabledRenderer(
 			[&](const UNiagaraRendererProperties* Properties)
 			{	
-				const FVertexFactoryType* VFType = Properties->GetVertexFactoryType();
-				if (VFType == nullptr)
-				{
-					return;
-				}
-
-				NewEntry.PSOPrecacheParams.bDisableBackFaceCulling = Properties->IsBackfaceCullingDisabled();
-
-				UNiagaraRendererProperties::FPSOPrecacheParamsList PSOPrecacheParamsList;
-				Properties->CollectPSOPrecacheData(&EmitterInst.Get(), PSOPrecacheParamsList);
-
-				for (UNiagaraRendererProperties::FPSOPrecacheParams& PSOPrecacheParams : PSOPrecacheParamsList)
-				{
-					NewEntry.MaterialInterface = PSOPrecacheParams.MaterialInterface;
-					NewEntry.VertexFactoryDataList = PSOPrecacheParams.VertexFactoryDataList;
-
-					NewEntry.PSOPrecacheParams.bReverseCulling = false;
-					AddMaterialInterfacePSOPrecacheParamsToList(NewEntry, List);
-
-					// Also precache with reverse culling if not two sided because we don't know of the component using the asset will have negative determinant
-					if (!NewEntry.PSOPrecacheParams.bDisableBackFaceCulling)
-					{
-						NewEntry.PSOPrecacheParams.bReverseCulling = true;
-						AddMaterialInterfacePSOPrecacheParamsToList(NewEntry, List);
-					}
-				}
+				Properties->CollectPSOPrecacheData(&EmitterInst.Get(), List);
 			}
 		);
 	}
diff --git a/Engine/Plugins/FX/Niagara/Source/Niagara/Public/NiagaraRendererProperties.h b/Engine/Plugins/FX/Niagara/Source/Niagara/Public/NiagaraRendererProperties.h
--- a/Engine/Plugins/FX/Niagara/Source/Niagara/Public/NiagaraRendererProperties.h
+++ b/Engine/Plugins/FX/Niagara/Source/Niagara/Public/NiagaraRendererProperties.h
@@ -325,7 +325,9 @@ public:
 	/** In the case that we need parameters bound in that aren't Particle variables, these should be set up here so that the data is appropriately populated after the simulation.*/
 	NIAGARA_API virtual bool PopulateRequiredBindings(FNiagaraParameterStore& InParameterStore);
 
-	
+	/** Collect PSO precache data using the optional emitter instance */
+	NIAGARA_API void CollectPSOPrecacheData(FNiagaraEmitterInstance* EmitterInstance, FMaterialInterfacePSOPrecacheParamsList& MaterialInterfacePSOPrecacheParamsList) const;
+
 	/**
 	* Collect all the data required for PSO precaching 
 	*/
diff --git a/Engine/Source/Runtime/Engine/Private/Components/PrimitiveComponent.cpp b/Engine/Source/Runtime/Engine/Private/Components/PrimitiveComponent.cpp
--- a/Engine/Source/Runtime/Engine/Private/Components/PrimitiveComponent.cpp
+++ b/Engine/Source/Runtime/Engine/Private/Components/PrimitiveComponent.cpp
@@ -4655,7 +4655,7 @@ void UPrimitiveComponent::SetupPrecachePSOParams(FPSOPrecacheParams& Params)
 	Params.bCastShadow = CastShadow;
 	// Custom depth can be toggled at runtime with PSO precache call so assume it might be needed when depth pass is needed
 	// Ideally precache those with lower priority and don't wait on these (UE-174426)
-	Params.bRenderCustomDepth = bRenderInDepthPass;
+	Params.bRenderCustomDepth = bRenderCustomDepth;
 	Params.bCastShadowAsTwoSided = bCastShadowAsTwoSided;
 	Params.SetMobility(Mobility);	
 	Params.SetStencilWriteMask(FRendererStencilMaskEvaluation::ToStencilMask(CustomDepthStencilWriteMask));
diff --git a/Engine/Source/Runtime/Engine/Private/Particles/ParticleComponents.cpp b/Engine/Source/Runtime/Engine/Private/Particles/ParticleComponents.cpp
--- a/Engine/Source/Runtime/Engine/Private/Particles/ParticleComponents.cpp
+++ b/Engine/Source/Runtime/Engine/Private/Particles/ParticleComponents.cpp
@@ -2716,7 +2716,6 @@ void UParticleSystem::PrecachePSOs()
 
 	FMaterialInterfacePSOPrecacheParams NewEntry;
 	NewEntry.PSOPrecacheParams.SetMobility(EComponentMobility::Movable);
-	NewEntry.PSOPrecacheParams.bRenderCustomDepth = true;
 
 	// No per component emitter materials known at this point in time
 	TArray<UMaterialInterface*> EmptyEmitterMaterials;
diff --git a/Engine/Source/Runtime/Renderer/Private/BasePassRendering.cpp b/Engine/Source/Runtime/Renderer/Private/BasePassRendering.cpp
--- a/Engine/Source/Runtime/Renderer/Private/BasePassRendering.cpp
+++ b/Engine/Source/Runtime/Renderer/Private/BasePassRendering.cpp
@@ -107,6 +107,13 @@ static TAutoConsoleVariable<int32> CVarPSOPrecacheTranslucencyAllPass(
 	ECVF_ReadOnly
 );
 
+static TAutoConsoleVariable<int32> CVarPSOPrecacheAlphaColorChannel(
+	TEXT("r.PSOPrecache.PrecacheAlphaColorChannel"),
+	1,
+	TEXT("Also Precache PSOs with scene color alpha channel enabled. Planar reflections and scene captures use this for compositing into a different scene later."),
+	ECVF_ReadOnly
+);
+
 // Scene color alpha is used during scene captures and planar reflections.  1 indicates background should be shown, 0 indicates foreground is fully present.
 static const float kSceneColorClearAlpha = 1.0f;
 
@@ -1739,6 +1746,69 @@ void FBasePassMeshProcessor::CollectPSOInitializersForLMPolicy(
 	}	
 }
 
+template<typename PassShadersType>
+void FBasePassMeshProcessor::AddBasePassGraphicsPipelineStateInitializer(
+	ERHIFeatureLevel::Type InFeatureLevel,
+	const FPSOPrecacheVertexFactoryData& VertexFactoryData,
+	const FMaterial& RESTRICT MaterialResource,
+	const FMeshPassProcessorRenderState& RESTRICT DrawRenderState,
+	const FGraphicsPipelineRenderTargetsInfo& RESTRICT RenderTargetsInfo,
+	const PassShadersType& PassShaders,
+	ERasterizerFillMode MeshFillMode,
+	ERasterizerCullMode MeshCullMode,
+	EPrimitiveType PrimitiveType,
+	bool bPrecacheAlphaColorChannel,
+	int InPSOCollectorIndex,
+	TArray<FPSOPrecacheData>& PSOInitializers)
+{
+	AddGraphicsPipelineStateInitializer(
+		VertexFactoryData,
+		MaterialResource,
+		DrawRenderState,
+		RenderTargetsInfo,
+		PassShaders,
+		MeshFillMode,
+		MeshCullMode,
+		PrimitiveType,
+		EMeshPassFeatures::Default,
+		ESubpassHint::None,
+		0,
+		true /*bRequired*/,
+		InPSOCollectorIndex,
+		PSOInitializers);
+
+	// Planar reflections and scene captures use scene color alpha to keep track of where content has been rendered, for compositing into a different scene later
+	if (bPrecacheAlphaColorChannel && CVarPSOPrecacheAlphaColorChannel.GetValueOnAnyThread() > 0)
+	{
+		FGraphicsPipelineRenderTargetsInfo AlphaColorRenderTargetsInfo = RenderTargetsInfo;
+
+		bool bRequiresAlphaChannel = true;
+		ETextureCreateFlags ExtraSceneColorCreateFlags = ETextureCreateFlags::None;
+		EPixelFormat SceneColorFormatWithAlpha;
+		ETextureCreateFlags SceneColorCreateFlagsWithAlpha;
+		GetSceneColorFormatAndCreateFlags(InFeatureLevel, bRequiresAlphaChannel, ExtraSceneColorCreateFlags, RenderTargetsInfo.NumSamples, false, SceneColorFormatWithAlpha, SceneColorCreateFlagsWithAlpha);
+
+		AlphaColorRenderTargetsInfo.RenderTargetFormats[0] = SceneColorFormatWithAlpha;
+		AlphaColorRenderTargetsInfo.RenderTargetFlags[0] = SceneColorCreateFlagsWithAlpha;
+
+		AddGraphicsPipelineStateInitializer(
+			VertexFactoryData,
+			MaterialResource,
+			DrawRenderState,
+			AlphaColorRenderTargetsInfo,
+			PassShaders,
+			MeshFillMode,
+			MeshCullMode,
+			PrimitiveType,
+			EMeshPassFeatures::Default,
+			ESubpassHint::None,
+			0,
+			true,
+			InPSOCollectorIndex,
+			PSOInitializers);
+	}
+}
+
 template<typename LightMapPolicyType>
 bool FBasePassMeshProcessor::Process(
 	const FMeshBatch& RESTRICT MeshBatch,
diff --git a/Engine/Source/Runtime/Renderer/Private/BasePassRendering.h b/Engine/Source/Runtime/Renderer/Private/BasePassRendering.h
--- a/Engine/Source/Runtime/Renderer/Private/BasePassRendering.h
+++ b/Engine/Source/Runtime/Renderer/Private/BasePassRendering.h
@@ -917,69 +917,6 @@ private:
 	float AutoBeforeDOFTranslucencyBoundary = 0.0f;
 };
 
-template<typename PassShadersType>
-void FBasePassMeshProcessor::AddBasePassGraphicsPipelineStateInitializer(
-	ERHIFeatureLevel::Type InFeatureLevel,
-	const FPSOPrecacheVertexFactoryData& VertexFactoryData,
-	const FMaterial& RESTRICT MaterialResource,
-	const FMeshPassProcessorRenderState& RESTRICT DrawRenderState,
-	const FGraphicsPipelineRenderTargetsInfo& RESTRICT RenderTargetsInfo,
-	const PassShadersType& PassShaders,
-	ERasterizerFillMode MeshFillMode,
-	ERasterizerCullMode MeshCullMode,
-	EPrimitiveType PrimitiveType,
-	bool bPrecacheAlphaColorChannel,
-	int InPSOCollectorIndex,
-	TArray<FPSOPrecacheData>& PSOInitializers)
-{
-	AddGraphicsPipelineStateInitializer(
-		VertexFactoryData,
-		MaterialResource,
-		DrawRenderState,
-		RenderTargetsInfo,
-		PassShaders,
-		MeshFillMode,
-		MeshCullMode,
-		PrimitiveType,
-		EMeshPassFeatures::Default,
-		ESubpassHint::None,
-		0,
-		true /*bRequired*/,
-		InPSOCollectorIndex,
-		PSOInitializers);
-
-	// Planar reflections and scene captures use scene color alpha to keep track of where content has been rendered, for compositing into a different scene later
-	if (bPrecacheAlphaColorChannel)
-	{
-		FGraphicsPipelineRenderTargetsInfo AlphaColorRenderTargetsInfo = RenderTargetsInfo;
-
-		bool bRequiresAlphaChannel = true;
-		ETextureCreateFlags ExtraSceneColorCreateFlags = ETextureCreateFlags::None;
-		EPixelFormat SceneColorFormatWithAlpha;
-		ETextureCreateFlags SceneColorCreateFlagsWithAlpha;
-		GetSceneColorFormatAndCreateFlags(InFeatureLevel, bRequiresAlphaChannel, ExtraSceneColorCreateFlags, RenderTargetsInfo.NumSamples, false /*bMemorylessMSAA*/, SceneColorFormatWithAlpha, SceneColorCreateFlagsWithAlpha);
-
-		AlphaColorRenderTargetsInfo.RenderTargetFormats[0] = SceneColorFormatWithAlpha;
-		AlphaColorRenderTargetsInfo.RenderTargetFlags[0] = SceneColorCreateFlagsWithAlpha;
-
-		AddGraphicsPipelineStateInitializer(
-			VertexFactoryData,
-			MaterialResource,
-			DrawRenderState,
-			AlphaColorRenderTargetsInfo,
-			PassShaders,
-			MeshFillMode,
-			MeshCullMode,
-			PrimitiveType,
-			EMeshPassFeatures::Default,
-			ESubpassHint::None,
-			0,
-			true /*bRequired*/,
-			InPSOCollectorIndex,
-			PSOInitializers);
-	}
-}
-
 ENUM_CLASS_FLAGS(FBasePassMeshProcessor::EFlags);
 
 extern void SetupBasePassState(FExclusiveDepthStencil::Type BasePassDepthStencilAccess, const bool bShaderComplexity, FMeshPassProcessorRenderState& DrawRenderState);
diff --git a/Engine/Source/Runtime/Renderer/Private/CustomDepthRendering.cpp b/Engine/Source/Runtime/Renderer/Private/CustomDepthRendering.cpp
--- a/Engine/Source/Runtime/Renderer/Private/CustomDepthRendering.cpp
+++ b/Engine/Source/Runtime/Renderer/Private/CustomDepthRendering.cpp
@@ -32,6 +32,16 @@ static TAutoConsoleVariable<bool> CVarCustomDepthEnableFastClear(
 	TEXT("Enable HTile on the custom depth buffer (default:false).\n"),
 	ECVF_RenderThreadSafe);
 
+static TAutoConsoleVariable<int32> CVarPSOPrecacheCustomDepth(
+	TEXT("r.PSOPrecache.CustomDepth"),
+	1,
+	TEXT("Also Precache PSOs with for custom depth pass.")  \
+	TEXT(" 0: No PSOs are compiled for this pass.\n") \
+	TEXT(" 1: PSOs are compiled for all primitives which explicitly request custom depth rendering (default).\n") \
+	TEXT(" 2: PSOs are compiled for all primitives which also request regular depth rendering.\n"),
+	ECVF_ReadOnly
+);
+
 DECLARE_DWORD_COUNTER_STAT(TEXT("Nanite Custom Depth Instances"), STAT_NaniteCustomDepthInstances, STATGROUP_Nanite);
 
 DECLARE_GPU_DRAWCALL_STAT_NAMED(CustomDepth, TEXT("Custom Depth"));
@@ -649,6 +659,12 @@ bool FCustomDepthPassMeshProcessor::Process(
 
 void FCustomDepthPassMeshProcessor::CollectPSOInitializers(const FSceneTexturesConfig& SceneTexturesConfig, const FMaterial& Material, const FPSOPrecacheVertexFactoryData& VertexFactoryData, const FPSOPrecacheParams& PreCacheParams, TArray<FPSOPrecacheData>& PSOInitializers)
 {
+	int32 CustomDepthPrecacheMode = CVarPSOPrecacheCustomDepth.GetValueOnAnyThread();
+	if (CustomDepthPrecacheMode == 0)
+	{
+		return;
+	}
+
 	// Setup the depth stencil state to use
 	const bool bWriteCustomStencilValues = IsCustomDepthPassWritingStencil();
 	PassDrawRenderState.SetDepthStencilState(GetCustomDepthStencilState(bWriteCustomStencilValues, PreCacheParams.GetStencilWriteMask()));
@@ -676,8 +692,16 @@ void FCustomDepthPassMeshProcessor::CollectPSOInitializers(const FSceneTexturesC
 			EffectiveMaterial = UMaterial::GetDefaultMaterial(MD_Surface)->GetMaterialResource(FeatureLevel, ActiveQualityLevel);
 			bUseDefaultMaterial = false;
 		}		
+				
+		bool bPrecacheCustomDepth = PreCacheParams.bRenderCustomDepth;
+
+		// If requested precache for all primitives in depth pass as well
+		if (CustomDepthPrecacheMode == 2)
+		{
+			bPrecacheCustomDepth = bPrecacheCustomDepth || PreCacheParams.bRenderInDepthPass;
+		}
 
-		if (!bUseDefaultMaterial && PreCacheParams.bRenderCustomDepth)
+		if (!bUseDefaultMaterial && bPrecacheCustomDepth)
 		{
 			check(!bPositionOnly);
 
diff --git a/Engine/Source/Runtime/Renderer/Private/DepthRendering.cpp b/Engine/Source/Runtime/Renderer/Private/DepthRendering.cpp
--- a/Engine/Source/Runtime/Renderer/Private/DepthRendering.cpp
+++ b/Engine/Source/Runtime/Renderer/Private/DepthRendering.cpp
@@ -78,6 +78,15 @@ static TAutoConsoleVariable<int32> CVarPSOPrecacheDitheredLODFadingOutMaskPass(
 	ECVF_ReadOnly
 );
 
+static TAutoConsoleVariable<int32> CVarPSOPrecacheProjectedShadows(
+	TEXT("r.PSOPrecache.ProjectedShadows"),
+	1,
+	TEXT("Also Precache PSOs with for projected shadows.")  \
+	TEXT(" 0: No PSOs are compiled for this pass.\n") \
+	TEXT(" 1: PSOs are compiled for all primitives which render to depth pass (default).\n"),
+	ECVF_ReadOnly
+);
+
 extern bool IsHMDHiddenAreaMaskActive();
 
 FDepthPassInfo GetDepthPassInfo(const FScene* Scene)
@@ -883,6 +892,7 @@ void FDepthPassMeshProcessor::CollectPSOInitializersInternal(
 		PSOInitializers);
 
 	// Also cache with project shadow depth stencil state (see FProjectedShadowInfo::SetupMeshDrawCommandsForProjectionStenciling)
+	if (CVarPSOPrecacheProjectedShadows.GetValueOnAnyThread() > 0)
 	{
 		// Set stencil to one.
 		DrawRenderState.SetDepthStencilState(
