Fix missing translucency when using separatetranslucency with an enabled DOF that is skipped because of a blurring radius being too small.

diff --git a/Engine/Source/Runtime/Renderer/Private/PostProcess/DiaphragmDOF.cpp b/Engine/Source/Runtime/Renderer/Private/PostProcess/DiaphragmDOF.cpp
--- a/Engine/Source/Runtime/Renderer/Private/PostProcess/DiaphragmDOF.cpp
+++ b/Engine/Source/Runtime/Renderer/Private/PostProcess/DiaphragmDOF.cpp
@@ -1370,17 +1370,19 @@ bool DiaphragmDOF::IsEnabled(const FViewInfo& View)
 		((View.FinalPostProcessSettings.DepthOfFieldFstop > 0.f && View.FinalPostProcessSettings.DepthOfFieldFocalDistance > 0.f) || View.FinalPostProcessSettings.DepthOfFieldDepthBlurRadius > 0.f);
 }
 
-FRDGTextureRef DiaphragmDOF::AddPasses(
+bool DiaphragmDOF::AddPasses(
 	FRDGBuilder& GraphBuilder,
 	const FSceneTextureParameters& SceneTextures,
 	const FViewInfo& View,
 	FRDGTextureRef InputSceneColor,
-	const FTranslucencyPassResources& TranslucencyPassResources)
+	const FTranslucencyPassResources& TranslucencyPassResources,
+	FRDGTextureRef& OutputColor)
 {
 	if (View.Family->EngineShowFlags.VisualizeDOF)
 	{
 		// no need for this pass
-		return InputSceneColor;
+		OutputColor = InputSceneColor;
+		return false;
 	}
 
 	// Format of the scene color.
@@ -1527,7 +1529,8 @@ FRDGTextureRef DiaphragmDOF::AddPasses(
 	// If the max blurring radius is too small, do not wire any passes.
 	if (MaxBluringRadius < MinRequiredBlurringRadius)
 	{
-		return InputSceneColor;
+		OutputColor = InputSceneColor;
+		return false;
 	}
 
 	RDG_GPU_STAT_SCOPE(GraphBuilder, DepthOfField);
@@ -2773,5 +2776,6 @@ FRDGTextureRef DiaphragmDOF::AddPasses(
 			FComputeShaderUtils::GetGroupCount(PassViewRect.Size(), kDefaultGroupSize));
 	}
 
-	return NewSceneColor;
+	OutputColor = NewSceneColor;
+	return true;
 }
diff --git a/Engine/Source/Runtime/Renderer/Private/PostProcess/DiaphragmDOF.h b/Engine/Source/Runtime/Renderer/Private/PostProcess/DiaphragmDOF.h
--- a/Engine/Source/Runtime/Renderer/Private/PostProcess/DiaphragmDOF.h
+++ b/Engine/Source/Runtime/Renderer/Private/PostProcess/DiaphragmDOF.h
@@ -135,11 +135,12 @@ struct FBokehModel
 RENDERER_API bool IsSupported(const FStaticShaderPlatform ShaderPlatform);
 
 /** Wire all DOF's passes according to view settings and cvars to convolve the scene color. */
-RENDERER_API FRDGTextureRef AddPasses(
+RENDERER_API bool AddPasses(
 	FRDGBuilder& GraphBuilder,
 	const FSceneTextureParameters& SceneTextures,
 	const FViewInfo& View,
 	FRDGTextureRef InputSceneColor,
-	const FTranslucencyPassResources& TranslucencyViewResources);
+	const FTranslucencyPassResources& TranslucencyViewResources,
+	FRDGTextureRef& OutputColor);
 
 } // namespace DiaphragmDOF
diff --git a/Engine/Source/Runtime/Renderer/Private/PostProcess/PostProcessing.cpp b/Engine/Source/Runtime/Renderer/Private/PostProcess/PostProcessing.cpp
--- a/Engine/Source/Runtime/Renderer/Private/PostProcess/PostProcessing.cpp
+++ b/Engine/Source/Runtime/Renderer/Private/PostProcess/PostProcessing.cpp
@@ -708,13 +708,16 @@ void AddPostProcessingPasses(
 
 				bool bComposeTranslucency = PostDOFTranslucencyResources.IsValid() && !bComposeSeparateTranslucencyInTSR && MaterialChainTranslucencyAfterDOF.Num() == 0;
 
-				SceneColor.Texture = DiaphragmDOF::AddPasses(
+				if (DiaphragmDOF::AddPasses(
 					GraphBuilder,
 					SceneTextureParameters,
 					View,
-					SceneColor.Texture,
-					bComposeTranslucency ? PostDOFTranslucencyResources : DummyTranslucency);
-				bSceneColorHasPostDOFTranslucency = bComposeTranslucency;
+					InputSceneColorTexture,
+					bComposeTranslucency ? PostDOFTranslucencyResources : DummyTranslucency,
+					SceneColor.Texture))
+				{
+					bSceneColorHasPostDOFTranslucency = bComposeTranslucency;
+				}
 			}
 
 			if (GetHairStrandsComposition() == EHairStrandsCompositionType::AfterSeparateTranslucent)
@@ -3092,4 +3095,4 @@ static void AddGBufferPicking(FRDGBuilder& GraphBuilder, const FViewInfo& View,
 	TShaderMapRef<FGBufferPickingCS> ComputeShader(View.ShaderMap);
 	FComputeShaderUtils::AddPass(GraphBuilder, RDG_EVENT_NAME("Debug::GBufferPicking"), ComputeShader, Parameters, FIntVector(1,1,1));
 }
-#endif
\ No newline at end of file
+#endif
