Hierarchical instance culling fixes, to prevent out of range issues in the hierarchical instance culling & improve local light culling.
 - Fixed a bug where the convex culling volume was not set up for local lights (for hierarchical instance culling) leading to less efficient culling for large local lights.
 - Changed the local light footprint heuristic to match the intent put forth in the cvar.
 - Fixed clamping logic to prevent numerical overflow (leading to checks)

diff --git a/Engine/Source/Runtime/Renderer/Private/SceneCulling/SceneCulling.cpp b/Engine/Source/Runtime/Renderer/Private/SceneCulling/SceneCulling.cpp
--- a/Engine/Source/Runtime/Renderer/Private/SceneCulling/SceneCulling.cpp
+++ b/Engine/Source/Runtime/Renderer/Private/SceneCulling/SceneCulling.cpp
@@ -155,10 +155,11 @@ static TAutoConsoleVariable<int32> CVarTreatDynamicInstancedAsUncullable(
 	TEXT("  This significantly reduces the hierarchy update cost on the CPU and for scenes with a large proportion of static elements, does not increase the GPU cost."),
 	ECVF_RenderThreadSafe);
 
-static TAutoConsoleVariable<int32> CVarSmallFootprintCellCountThreshold(
-	TEXT("r.SceneCulling.SmallFootprintCellCountThreshold"), 
-	512, 
-	TEXT("Queries with a smaller footprint (in number of cells in the lowest level) go down the footprint based path."), 
+static TAutoConsoleVariable<int32> CVarSmallFootprintSideThreshold(
+	TEXT("r.SceneCulling.SmallFootprintSideThreshold"), 
+	16, 
+	TEXT("Queries with a smaller footprint (maximum) side (in number of cells in the lowest level) go down the footprint based path.\n") 
+	TEXT("  The default (16) <=> a footprint of 16x16x16 cells or 8 blocks"), 
 	ECVF_RenderThreadSafe);
 
 static TAutoConsoleVariable<int32> CVarValidateAllInstanceAllocations(
@@ -404,10 +405,12 @@ SC_FORCEINLINE FSceneCulling::FFootprint8 ToBlockLocal(const FSceneCulling::FFoo
 	FInt64Vector3 BlockMin = BlockLoc.Coord * FSceneCulling::FSpatialHash::CellBlockDim;
 	FInt64Vector3 BlockMax = BlockMin + FInt64Vector3(FSceneCulling::FSpatialHash::CellBlockDim - 1);
 
-	// This can be packed to very few bits if need be.
+	FInt64Vector3 BlockLocalMin = ClampDim(ObjFootprint.Min - BlockMin, 0ll, FSceneCulling::FSpatialHash::CellBlockDim - 1ll);
+	FInt64Vector3 BlockLocalMax = ClampDim(ObjFootprint.Max - BlockMin, 0ll, FSceneCulling::FSpatialHash::CellBlockDim - 1ll);
+
 	FSceneCulling::FFootprint8 LocalFp = {
-		ClampDim(FInt8Vector3(ObjFootprint.Min - BlockMin), int8(0), int8(FSceneCulling::FSpatialHash::CellBlockDim - 1)),
-		ClampDim(FInt8Vector3(ObjFootprint.Max - BlockMin), int8(0), int8(FSceneCulling::FSpatialHash::CellBlockDim - 1)),
+		FInt8Vector3(BlockLocalMin),
+		FInt8Vector3(BlockLocalMax),
 		ObjFootprint.Level
 	};
 	return LocalFp;
@@ -431,7 +434,6 @@ SC_FORCEINLINE FSceneCulling::FLocation8 ToBlockLocal(const FSceneCulling::FLoca
 	return LocalLoc;
 };
 
-
 inline FInt64Vector3 ToLevelRelative(const FInt64Vector3& Coord, int32 LevelDelta)
 {
 	if (LevelDelta > 0)
@@ -703,8 +705,7 @@ void FSceneCulling::Test(const FCullingVolume& CullingVolume, TArray<FCellDraw,
 		const float Level0CellSize = SpatialHash.GetCellSize(SpatialHash.GetFirstLevel());
 		FFootprint64 LightFootprint = SpatialHash.CalcFootprintSphere(SpatialHash.GetFirstLevel(), CullingVolume.Sphere.Center, CullingVolume.Sphere.W + (Level0CellSize * 0.5f));
 
-		// Diagonal length
-		if ((LightFootprint.Max - LightFootprint.Min).Size() < SmallFootprintCellCountThreshold)
+		if ((LightFootprint.Max - LightFootprint.Min).GetMax() <= int64(SmallFootprintCellSideThreshold))
 		{
 			TestSphere(CullingVolume.Sphere, OutCellDraws, ViewGroupId, MaxNumViews, OutNumInstanceGroups);
 			return;
@@ -1437,6 +1438,12 @@ public:
 		return TempCell.CellOffset;
 	}
 
+	// Clamp the cell location to prevent overflows
+	SC_FORCEINLINE FSceneCulling::FLocation64 ClampCellLoc(const FSceneCulling::FLocation64 &InLoc)
+	{
+		return FSceneCulling::FLocation64(ClampDim(InLoc.Coord, -FSceneCulling::FBlockTraits::MaxCellCoord, FSceneCulling::FBlockTraits::MaxCellCoord), InLoc.Level);
+	}
+
 	template <EUpdateFrequencyCategory::EType UpdateFrequencyCategory, typename HashLocationComputerType>
 	SC_FORCEINLINE void BuildInstanceRange(int32 InstanceDataOffset, int32 NumInstances, HashLocationComputerType HashLocationComputer, FSceneCulling::FCellIndexCacheEntry &CellIndexCacheEntry)
 	{
@@ -1452,7 +1459,7 @@ public:
 		{
 			const int32 InstanceId = InstanceDataOffset + InstanceIndex;
 
-			FSceneCulling::FLocation64 InstanceCellLoc = HashLocationComputer.CalcLoc(InstanceIndex);
+			FSceneCulling::FLocation64 InstanceCellLoc = ClampCellLoc(HashLocationComputer.CalcLoc(InstanceIndex));
 
 			bool bSameLoc = bCompressRLE && SameInstanceLocRunCount > 0 && PrevInstanceCellLoc == InstanceCellLoc;
 
@@ -1969,12 +1976,12 @@ public:
 	}
 
 	template <typename HashLocationComputerType>
-	inline void UpdateProcessDynamicInstances(HashLocationComputerType &HashLocationComputer, int32 InstanceDataOffset, int32 NumInstances, int32 PrevNumInstances, FSceneCulling::FCellIndexCacheEntry &CacheEntry)
+	SC_FORCEINLINE void UpdateProcessDynamicInstances(HashLocationComputerType &HashLocationComputer, int32 InstanceDataOffset, int32 NumInstances, int32 PrevNumInstances, FSceneCulling::FCellIndexCacheEntry &CacheEntry)
 	{
 		for (int32 InstanceIndex = 0; InstanceIndex < NumInstances; ++InstanceIndex)
 		{
 			const int32 InstanceId = InstanceDataOffset + InstanceIndex;
-			FSceneCulling::FLocation64 InstanceCellLoc = HashLocationComputer.CalcLoc(InstanceIndex);
+			FSceneCulling::FLocation64 InstanceCellLoc = ClampCellLoc(HashLocationComputer.CalcLoc(InstanceIndex));
 
 			bool bNeedAdd = InstanceIndex >= PrevNumInstances;
 			if (!bNeedAdd)
@@ -2590,10 +2597,9 @@ FSceneCulling::FUpdater &FSceneCulling::BeginUpdate(FRDGBuilder& GraphBuilder, F
 
 	bUseExplictBounds = CVarSceneCullingUseExplicitCellBounds.GetValueOnRenderThread() != 0;
 
-	SmallFootprintCellCountThreshold = CVarSmallFootprintCellCountThreshold.GetValueOnRenderThread();
+	SmallFootprintCellSideThreshold = CVarSmallFootprintSideThreshold.GetValueOnRenderThread();
 	bUseAsyncUpdate = CVarSceneCullingAsyncUpdate.GetValueOnRenderThread() != 0;
 	bUseAsyncQuery = CVarSceneCullingAsyncQuery.GetValueOnRenderThread() != 0;
-	SmallFootprintCellCountThreshold = CVarSmallFootprintCellCountThreshold.GetValueOnRenderThread();
 
 	if (bIsEnabled)
 	{
diff --git a/Engine/Source/Runtime/Renderer/Private/SceneCulling/SceneCulling.h b/Engine/Source/Runtime/Renderer/Private/SceneCulling/SceneCulling.h
--- a/Engine/Source/Runtime/Renderer/Private/SceneCulling/SceneCulling.h
+++ b/Engine/Source/Runtime/Renderer/Private/SceneCulling/SceneCulling.h
@@ -132,6 +132,11 @@ public:
 	{
 		static constexpr int32 CellBlockDimLog2 = 3; // (8x8x8)
 		using FBlockLoc = FBlockLoc;
+
+		// The FBlockLocAligned represents the block locations as 32-bit ints.
+		static constexpr int64 MaxCellBlockCoord = MAX_int32;
+		// The cell coordinate may be larger by the block dimension and still can fit into a signed 32-bit integer
+		static constexpr int64 MaxCellCoord = MaxCellBlockCoord << CellBlockDimLog2;
 	};
 
 	using FSpatialHash = THierarchicalSpatialHashGrid<FBlockTraits>;
@@ -274,10 +279,8 @@ private:
 	TArray<FPersistentPrimitiveIndex> UnCullablePrimitives;
 	int32 UncullableItemChunksOffset = INDEX_NONE;
 	int32 UncullableNumItemChunks = 0;
-	// Number of cells in the finest level under which a footprint is considered "small" and should go down the direct footprint path
-	// TODO: Maybe better to use some other metric? A sphere could also be tested for insideness, and that process would be efficient for a large
-	//       enough sphere, and able to trim more than a set of planes But perhaps marginal anyway?
-	int32 SmallFootprintCellCountThreshold = 0;
+	// Largest dimension length, in cells, at the finest level under which a footprint is considered "small" and should go down the direct footprint path
+	int32 SmallFootprintCellSideThreshold = 16;
 	bool bTestCellVsQueryBounds = true;
 	bool bUseAsyncUpdate = true;
 	bool bUseAsyncQuery = true;
diff --git a/Engine/Source/Runtime/Renderer/Private/VirtualShadowMaps/VirtualShadowMapArray.cpp b/Engine/Source/Runtime/Renderer/Private/VirtualShadowMaps/VirtualShadowMapArray.cpp
--- a/Engine/Source/Runtime/Renderer/Private/VirtualShadowMaps/VirtualShadowMapArray.cpp
+++ b/Engine/Source/Runtime/Renderer/Private/VirtualShadowMaps/VirtualShadowMapArray.cpp
@@ -2529,6 +2529,7 @@ static FCullingVolume GetCullingVolume(const FProjectedShadowInfo* ProjectedShad
 	else
 	{
 		CullingVolume.Sphere = ProjectedShadowInfo->GetLightSceneInfo().Proxy->GetBoundingSphere();
+		CullingVolume.ConvexVolume = ProjectedShadowInfo->CasterOuterFrustum;
 	}
 	return CullingVolume;
 }
