Add per-light VSM resolution bias from proxy to local light VSM setup.

diff --git a/Engine/Source/Runtime/Renderer/Private/Shadows/ShadowSceneRenderer.cpp b/Engine/Source/Runtime/Renderer/Private/Shadows/ShadowSceneRenderer.cpp
--- a/Engine/Source/Runtime/Renderer/Private/Shadows/ShadowSceneRenderer.cpp
+++ b/Engine/Source/Runtime/Renderer/Private/Shadows/ShadowSceneRenderer.cpp
@@ -182,12 +182,12 @@ UE::Renderer::Private::IShadowInvalidatingInstances *FShadowSceneRenderer::GetIn
 }
 
 
-static float GetResolutionLODBiasLocal(float LightMobilityFactor)
+static float GetResolutionLODBiasLocal(float LightMobilityFactor, float LightLODBias)
 {
 	return FVirtualShadowMapArray::InterpolateResolutionBias(
 		CVarResolutionLodBiasLocal.GetValueOnRenderThread(),
 		CVarResolutionLodBiasLocalMoving.GetValueOnRenderThread(),
-		LightMobilityFactor);
+		LightMobilityFactor) + LightLODBias;
 }
 
 FVirtualShadowMapProjectionShaderData FShadowSceneRenderer::GetLocalLightProjectionShaderData(
@@ -233,7 +233,8 @@ TSharedPtr<FVirtualShadowMapPerLightCacheEntry> FShadowSceneRenderer::AddLocalLi
 	LocalLightShadowFrameSetup.ProjectedShadowInfo = ProjectedShadowInfo;
 	LocalLightShadowFrameSetup.LightSceneInfo = LightSceneInfo;
 
-	const float ResolutionLODBiasLocal = GetResolutionLODBiasLocal(ShadowScene.GetLightMobilityFactor(LightSceneInfo->Id));
+	const FLightSceneProxy* LightSceneProxy = ProjectedShadowInfo->GetLightSceneInfo().Proxy;
+	const float ResolutionLODBiasLocal = GetResolutionLODBiasLocal(ShadowScene.GetLightMobilityFactor(LightSceneInfo->Id), LightSceneProxy->GetVSMResolutionLodBias());
 
 	// Single page res, at this point we force the VSM to be single page
 	// TODO: this computation does not match up with page marking logic super-well, particularly for long spot lights,
