Rewrites DOF's infinity background CoC radius to fix a mistake happening on the anamorphic lense

This fixes the DOF changing when the Squeeze parameter on the camera is set.  Note that this patch does not modify the Material Function CircleDOFParams. The math should be consistent with the previous logic as long the squeeze factor is 1.0

diff --git a/Engine/Source/Runtime/Renderer/Private/PathTracing.cpp b/Engine/Source/Runtime/Renderer/Private/PathTracing.cpp
--- a/Engine/Source/Runtime/Renderer/Private/PathTracing.cpp
+++ b/Engine/Source/Runtime/Renderer/Private/PathTracing.cpp
@@ -34,6 +34,7 @@ TAutoConsoleVariable<int32> CVarPathTracing(
 #include "SkyAtmosphereRendering.h"
 #include <limits>
 #include "PathTracingSpatialTemporalDenoising.h"
+#include "PostProcess/DiaphragmDOF.h"
 #include "EnvironmentComponentsFlags.h"
 
 TAutoConsoleVariable<int32> CVarPathTracingExperimental(
@@ -646,10 +647,10 @@ static void PreparePathTracingData(const FScene* Scene, const FViewInfo& View, F
 		PPV.DepthOfFieldFocalDistance > 0 &&
 		PPV.DepthOfFieldFstop > 0)
 	{
-		const float FocalLengthInCM = 0.05f * PPV.DepthOfFieldSensorWidth * View.ViewMatrices.GetProjectionMatrix().M[0][0];
-		PathTracingData.CameraFocusDistance = PPV.DepthOfFieldFocalDistance;
-		PathTracingData.CameraLensRadius.Y = 0.5f * FocalLengthInCM / PPV.DepthOfFieldFstop;
-		PathTracingData.CameraLensRadius.X = PathTracingData.CameraLensRadius.Y / FMath::Clamp(PPV.DepthOfFieldSqueezeFactor, 1.0f, 2.0f);
+		DiaphragmDOF::FPhysicalCocModel CocModel;
+		CocModel.Compile(View);
+		PathTracingData.CameraFocusDistance = CocModel.FocusDistance;
+		PathTracingData.CameraLensRadius    = CocModel.GetLensRadius();
 	}
 
 	// Merge all volume flags into one uint
diff --git a/Engine/Source/Runtime/Renderer/Private/PostProcess/DiaphragmDOF.h b/Engine/Source/Runtime/Renderer/Private/PostProcess/DiaphragmDOF.h
--- a/Engine/Source/Runtime/Renderer/Private/PostProcess/DiaphragmDOF.h
+++ b/Engine/Source/Runtime/Renderer/Private/PostProcess/DiaphragmDOF.h
@@ -22,14 +22,29 @@ namespace DiaphragmDOF
 // Whether DOF is enabled for the requested view.
 bool IsEnabled(const FViewInfo& View);
 
-float ComputeFocalLengthFromFov(const FSceneView& View);
 FVector4f CircleDofHalfCoc(const FViewInfo& View);
 
 /** Physically based circle of confusion computation model. */
 struct FPhysicalCocModel
 {
-	// Unclamped resolution less background coc radius.
+	// Size of the sensor, in unreal unit.
+	float SensorWidth;
+	float SensorHeight;
+
+	// Aspect ratio of the croped frame being rendered.
+	float RenderingAspectRatio;
+
+	// Focal length of the lens in unreal unit.
+	float VerticalFocalLength;
+	// float HorizontalFocalLength = VerticalFocalLength / Squeeze;
+
+	// Apperture diameter in fstop
+	float FStops; // = VerticalFocalLength / ApartureDiameter.
+
+	// Unclamped background vertical coc radius, in horizontal ViewportUV unit.
 	float InfinityBackgroundCocRadius;
+	// HorizontalInfinityBackgroundCocRadius = InfinityBackgroundCocRadius / Squeeze
+	// VerticalInfinityBackgroundCocRadius = InfinityBackgroundCocRadius
 
 	/** Indicates whether a dynamic offset dependent on scene depth should be computed for every pixel */
 	bool bEnableDynamicOffset;
@@ -46,7 +61,7 @@ struct FPhysicalCocModel
 	// Resolution less maximal background coc radius.
 	float MaxBackgroundCocRadius;
 
-	// Focus distance.
+	// Focus distance in unreal unit.
 	float FocusDistance;
 
 	// SqueezeFactor = VerticalFocalDistance / HorizontalFocalDistance
@@ -59,6 +74,9 @@ struct FPhysicalCocModel
 	/** Compile the coc model from a view. */
 	void Compile(const FViewInfo& View);
 	
+	/** Returns the lens radius in unreal unit from which the path traced ray should be traced from. */
+	FVector2f GetLensRadius() const;
+
 	/** Returns the CocRadius in half res pixels for given scene depth (in world unit).
 	 *
 	 * Notes: Matches Engine/Shaders/Private/DiaphragmDOF/Common.ush's SceneDepthToCocRadius().
diff --git a/Engine/Source/Runtime/Renderer/Private/PostProcess/DiaphragmDOFUtils.cpp b/Engine/Source/Runtime/Renderer/Private/PostProcess/DiaphragmDOFUtils.cpp
--- a/Engine/Source/Runtime/Renderer/Private/PostProcess/DiaphragmDOFUtils.cpp
+++ b/Engine/Source/Runtime/Renderer/Private/PostProcess/DiaphragmDOFUtils.cpp
@@ -24,7 +24,8 @@ TAutoConsoleVariable<float> CVarMaxBackgroundRadius(
 
 } // namespace
 
-float DiaphragmDOF::ComputeFocalLengthFromFov(const FSceneView& View)
+// TODO: delete.
+float ComputeFocalLengthFromFov(const FSceneView& View)
 {
 	// Convert FOV to focal length,
 	// 
@@ -44,6 +45,7 @@ float DiaphragmDOF::ComputeFocalLengthFromFov(const FSceneView& View)
 
 // Convert f-stop and focal distance into projected size in half resolution pixels.
 // Setup depth based blur.
+// TODO: This logic does not account for the Squeeze factor in the same way as the logic below. See JIRA UE-203727
 FVector4f DiaphragmDOF::CircleDofHalfCoc(const FViewInfo& View)
 {
 	static const auto CVar = IConsoleManager::Get().FindTConsoleVariableDataInt(TEXT("r.DepthOfFieldQuality"));
@@ -53,7 +55,7 @@ FVector4f DiaphragmDOF::CircleDofHalfCoc(const FViewInfo& View)
 
 	if(bDepthOfField)
 	{
-		float FocalLengthInMM = DiaphragmDOF::ComputeFocalLengthFromFov(View);
+		float FocalLengthInMM = ComputeFocalLengthFromFov(View); // TODO for Material.
 	 
 		// Convert focal distance in world position to mm (from cm to mm)
 		float FocalDistanceInMM = View.FinalPostProcessSettings.DepthOfFieldFocalDistance * 10.0f;
@@ -99,15 +101,34 @@ FVector4f DiaphragmDOF::CircleDofHalfCoc(const FViewInfo& View)
 
 void DiaphragmDOF::FPhysicalCocModel::Compile(const FViewInfo& View)
 {
-	// Fetches DOF settings.
+	// Fetches lens and filmback settings settings.
 	{
+		const float MMToUU = 0.1f;
+
 		FocusDistance = View.FinalPostProcessSettings.DepthOfFieldFocalDistance;
+		FStops = View.FinalPostProcessSettings.DepthOfFieldFstop;
 		Squeeze = FMath::Clamp(View.FinalPostProcessSettings.DepthOfFieldSqueezeFactor, 1.0f, 2.0f);
 
+		RenderingAspectRatio = float(View.UnscaledViewRect.Width()) / float(View.UnscaledViewRect.Height());
+		const float HorizontalHalfFOV = FMath::Atan(1.0f / View.ViewMatrices.GetProjectionMatrix().M[0][0]);
+		const float VerticalHalfFOV = FMath::Atan(1.0f / View.ViewMatrices.GetProjectionMatrix().M[1][1]);
+
+		// If the focal length isn't set, compute based of the sensor height and vertical FOV.
+		const float SensorAspectRatio = RenderingAspectRatio / Squeeze;
+		SensorWidth  = View.FinalPostProcessSettings.DepthOfFieldSensorWidth * MMToUU;
+		SensorHeight = SensorWidth / SensorAspectRatio;
+		VerticalFocalLength = 0.5f * SensorHeight * (1.0f / FMath::Tan(VerticalHalfFOV));
+	}
+
+	// Fetch the max bluring radius
+	{
 		// -because foreground Coc are negative.
 		MinForegroundCocRadius = -CVarMaxForegroundRadius.GetValueOnRenderThread();
 		MaxBackgroundCocRadius = CVarMaxBackgroundRadius.GetValueOnRenderThread();
+	}
 
+	// Fetch the depth blur.
+	{
 		MaxDepthBlurRadius = View.FinalPostProcessSettings.DepthOfFieldDepthBlurRadius / 1920.0f;
 
 		// Circle DOF was actually computing in this depth blur radius in half res.
@@ -117,29 +138,26 @@ void DiaphragmDOF::FPhysicalCocModel::Compile(const FViewInfo& View)
 	}
 
 	// Compile coc model equation.
-	if (View.FinalPostProcessSettings.DepthOfFieldFstop > 0.f && View.FinalPostProcessSettings.DepthOfFieldFocalDistance > 0.f)
+	if (FStops > 0.f && FocusDistance > 0.f)
 	{
-
-		float FocalLengthInMM = DiaphragmDOF::ComputeFocalLengthFromFov(View);
-
-		// Convert focal distance in world position to mm (from cm to mm)
-		float FocalDistanceInMM = View.FinalPostProcessSettings.DepthOfFieldFocalDistance * 10.0f;
-
-		// Convert mm to pixels.
-		float const SensorWidthInMM = View.FinalPostProcessSettings.DepthOfFieldSensorWidth;
-
 		// Convert f-stop, focal length, and focal distance to
-		// projected circle of confusion size at infinity in mm.
+		// projected circle of confusion size of infinity on the sensor in unreal unit.
 		//
 		// coc = f * f / (n * (d - f))
 		// where,
 		//   f = focal length
 		//   d = focal distance
 		//   n = fstop (where n is the "n" in "f/n")
-		float DiameterInMM = FMath::Square(FocalLengthInMM) / (View.FinalPostProcessSettings.DepthOfFieldFstop * (FocalDistanceInMM - FocalLengthInMM));
+		float VerticalDiameter = FMath::Square(VerticalFocalLength) / (FStops * (FocusDistance - VerticalFocalLength));
+
+		// Convert vertical diameter in unreal unit to radius on the filmback in vertical ViewportUV unit in uncropped viewport.
+		float UncroppedVerticalInfinityBackgroundCocRadius = VerticalDiameter * 0.5f / SensorHeight;
 
-		// Convert diameter in mm to resolution less radius on the filmback.
-		InfinityBackgroundCocRadius = DiameterInMM * 0.5f / SensorWidthInMM;
+		const float DesqueezedAspectRatio = SensorWidth / SensorHeight * Squeeze;
+		float VerticalInfinityBackgroundCocRadius = UncroppedVerticalInfinityBackgroundCocRadius * FMath::Max(RenderingAspectRatio / DesqueezedAspectRatio, 1.0);
+
+		// Convert diameter from vertical ViewportUV unit to horizontal ViewportUV unit.
+		InfinityBackgroundCocRadius = VerticalInfinityBackgroundCocRadius / RenderingAspectRatio;
 
 		if (View.InFocusDistance > 0)
 		{
@@ -166,6 +184,16 @@ void DiaphragmDOF::FPhysicalCocModel::Compile(const FViewInfo& View)
 	}
 }
 
+FVector2f DiaphragmDOF::FPhysicalCocModel::GetLensRadius() const
+{
+	// Size of the vertical aperture in unreal unit.
+	float ApertureDiameter = VerticalFocalLength / FStops;
+
+	float VerticalLensRadius = 0.5f * ApertureDiameter;
+	float HorizontalLensRadius = VerticalLensRadius / Squeeze;
+	return FVector2f(HorizontalLensRadius, VerticalLensRadius);
+}
+
 float DiaphragmDOF::FPhysicalCocModel::DepthToResCocRadius(float SceneDepth, float HorizontalResolution) const
 {
 	float InitialCocRadius = ((SceneDepth - FocusDistance) / SceneDepth) * InfinityBackgroundCocRadius;
diff --git a/Engine/Source/Runtime/Renderer/Private/PostProcess/PostProcessBokehDOF.cpp b/Engine/Source/Runtime/Renderer/Private/PostProcess/PostProcessBokehDOF.cpp
--- a/Engine/Source/Runtime/Renderer/Private/PostProcess/PostProcessBokehDOF.cpp
+++ b/Engine/Source/Runtime/Renderer/Private/PostProcess/PostProcessBokehDOF.cpp
@@ -117,13 +117,18 @@ FScreenPassTexture AddVisualizeDOFPass(FRDGBuilder& GraphBuilder, const FViewInf
 
 		const FVector2D Fov = View.ViewMatrices.ComputeHalfFieldOfViewPerAxis();
 
-		const float FocalLength = DiaphragmDOF::ComputeFocalLengthFromFov(View);
+		DiaphragmDOF::FPhysicalCocModel CocModel;
+		CocModel.Compile(View);
 
-		Line = FString::Printf(TEXT("Field Of View in deg. (computed): %.1f x %.1f"), FMath::RadiansToDegrees(Fov.X) * 2.0f, FMath::RadiansToDegrees(Fov.Y) * 2.0f);
+		const float UUToMM = 10.0f;
+
+		Line = FString::Printf(TEXT("Field Of View in deg.: %.1f x %.1f"), FMath::RadiansToDegrees(Fov.X) * 2.0f, FMath::RadiansToDegrees(Fov.Y) * 2.0f);
+		Canvas.DrawShadowedString(X, Y += YStep, *Line, GetStatsFont(), FLinearColor(0.5, 0.5, 1));
+		Line = FString::Printf(TEXT("Focal Length: %.1f mm"), UUToMM * CocModel.VerticalFocalLength);
+		Canvas.DrawShadowedString(X, Y += YStep, *Line, GetStatsFont(), FLinearColor(0.5, 0.5, 1));
+		Line = FString::Printf(TEXT("Squeeze Factor: %.1f"), CocModel.Squeeze);
 		Canvas.DrawShadowedString(X, Y += YStep, *Line, GetStatsFont(), FLinearColor(0.5, 0.5, 1));
-		Line = FString::Printf(TEXT("Focal Length (computed): %.1f"), FocalLength);
-		Canvas.DrawShadowedString(X, Y += YStep, *Line, GetStatsFont(), FLinearColor(0.5, 0.5, 10));
-		Line = FString::Printf(TEXT("Sensor: APS-C 24.576 mm sensor, crop-factor 1.61x"));
+		Line = FString::Printf(TEXT("Sensor: %.1f x %.1f mm"), UUToMM * CocModel.SensorWidth, UUToMM * CocModel.SensorHeight);
 		Canvas.DrawShadowedString(X, Y += YStep, *Line, GetStatsFont(), FLinearColor(0.5, 0.5, 1));
 	});
 
