Monochrome mode for distance field visualization
- useful for automated tests since colors are random.

diff --git a/Engine/Shaders/Private/DistanceFieldVisualization.usf b/Engine/Shaders/Private/DistanceFieldVisualization.usf
--- a/Engine/Shaders/Private/DistanceFieldVisualization.usf
+++ b/Engine/Shaders/Private/DistanceFieldVisualization.usf
@@ -249,7 +249,7 @@ void VisualizeMeshDistanceFieldCS(
 
 		Lighting += GetSkySHDiffuse(WorldHitNormal) * View.SkyLightColor.rgb;
 		float3 Albedo = .05f;
-#if !USE_GLOBAL_DISTANCE_FIELD
+#if !USE_GLOBAL_DISTANCE_FIELD && !MONOCHROME_VISUALIZATION
 		{
 			FDFObjectData DFObjectData = LoadDFObjectData(HitCulledObjectIndex);
 			Albedo = 0.1f * IntToColor(DFObjectData.GPUSceneInstanceIndex);
diff --git a/Engine/Source/Runtime/Renderer/Private/DistanceFieldVisualization.cpp b/Engine/Source/Runtime/Renderer/Private/DistanceFieldVisualization.cpp
--- a/Engine/Source/Runtime/Renderer/Private/DistanceFieldVisualization.cpp
+++ b/Engine/Source/Runtime/Renderer/Private/DistanceFieldVisualization.cpp
@@ -18,6 +18,14 @@
 #include "PostProcess/PostProcessSubsurface.h"
 #include "PipelineStateCache.h"
 
+int32 GDistanceFieldsVisualizeMonochrome = 0;
+FAutoConsoleVariableRef CVarDistanceFieldsVisualizeMonochrome(
+	TEXT("r.DistanceFields.Visualize.Monochrome"),
+	GDistanceFieldsVisualizeMonochrome,
+	TEXT("Whether to render the distance field visualization in monochrome"),
+	ECVF_RenderThreadSafe
+);
+
 class FVisualizeMeshDistanceFieldCS : public FGlobalShader
 {
 public:
@@ -25,10 +33,11 @@ public:
 	SHADER_USE_PARAMETER_STRUCT(FVisualizeMeshDistanceFieldCS, FGlobalShader);
 
 	class FUseGlobalDistanceFieldDim : SHADER_PERMUTATION_BOOL("USE_GLOBAL_DISTANCE_FIELD");
+	class FMonochromeVisualizationDim : SHADER_PERMUTATION_BOOL("MONOCHROME_VISUALIZATION");
 	class FCoverageBasedExpand : SHADER_PERMUTATION_BOOL("GLOBALSDF_USE_COVERAGE_BASED_EXPAND");
 	class FSimpleCoverageBasedExpand : SHADER_PERMUTATION_BOOL("GLOBALSDF_SIMPLE_COVERAGE_BASED_EXPAND");
 	class FOffsetDataStructure : SHADER_PERMUTATION_INT("OFFSET_DATA_STRUCT", 3);
-	using FPermutationDomain = TShaderPermutationDomain<FUseGlobalDistanceFieldDim, FCoverageBasedExpand, FSimpleCoverageBasedExpand, FOffsetDataStructure>;
+	using FPermutationDomain = TShaderPermutationDomain<FUseGlobalDistanceFieldDim, FMonochromeVisualizationDim, FCoverageBasedExpand, FSimpleCoverageBasedExpand, FOffsetDataStructure>;
 
 	BEGIN_SHADER_PARAMETER_STRUCT(FParameters, )
 		SHADER_PARAMETER_STRUCT_REF(FViewUniformShaderParameters, View)
@@ -160,6 +169,7 @@ void FSceneRenderer::RenderMeshDistanceFieldVisualization(FRDGBuilder& GraphBuil
 
 		FVisualizeMeshDistanceFieldCS::FPermutationDomain PermutationVector;
 		PermutationVector.Set<FVisualizeMeshDistanceFieldCS::FUseGlobalDistanceFieldDim>(bVisualizeGlobalDistanceField);
+		PermutationVector.Set<FVisualizeMeshDistanceFieldCS::FMonochromeVisualizationDim>(GDistanceFieldsVisualizeMonochrome != 0);
 		PermutationVector.Set<FVisualizeMeshDistanceFieldCS::FCoverageBasedExpand>(IsLumenEnabled(View));
 		PermutationVector.Set<FVisualizeMeshDistanceFieldCS::FSimpleCoverageBasedExpand>(IsLumenEnabled(View) && Lumen::UseGlobalSDFSimpleCoverageBasedExpand());
 		extern int32 GDistanceFieldOffsetDataStructure;
