Refactor contrast adaptive VRS to support single wide luminance inputs, the size of which can vary between views in a single family.
Change to outputting single wide luminance histories from the tonemapper when either view has post-processing steps after the tonemapper. Fixes crashes and other issues when one view in a family has post-processing steps after the tonemapper, and another doesn't. Removed ability to preview pre-reprojection shading rate images since they are no longer stored in a single easily providable texture.

diff --git a/Engine/Shaders/Private/VariableRateShading/VRSShadingRateReproject.usf b/Engine/Shaders/Private/VariableRateShading/VRSShadingRateReproject.usf
--- a/Engine/Shaders/Private/VariableRateShading/VRSShadingRateReproject.usf
+++ b/Engine/Shaders/Private/VariableRateShading/VRSShadingRateReproject.usf
@@ -4,12 +4,11 @@
 #include "../Common.ush"
 #include "../SceneTexturesCommon.ush"
 
+float2 InputSRIDimensions;
 float2 ScaledSRIDimensions;
 float2 TextureDimensions;
 float2 InvTextureDimensions;
 
-float2 SRIViewRectMin;
-float2 SRIViewRectMax;
 float2 ScaledUVOffset;
 float  InvDynamicResolutionScale;
 
@@ -71,8 +70,7 @@ void RescaleVariableRateShading(
 	const float epsilon = 0.0001;
 
 	const float2 HalfPixelOffset = float2(0.5, 0.5);
-	float2 UVOffset = ScaledUVOffset;
-	float2 CenterUV = UVOffset + (float2(DispatchThreadId.xy) + HalfPixelOffset) * pixelUVSize;
+	float2 CenterUV = (float2(DispatchThreadId.xy) + HalfPixelOffset) * pixelUVSize;
 	CenterUV = ReprojectBufferUV(DispatchThreadId, GroupThreadId, CenterUV, bValidThread); 
 
 	// FXC doesn't allow a thread to terminate before GroupMemoryBarrierWithGroupSync so 
@@ -89,14 +87,14 @@ void RescaleVariableRateShading(
 	// Convert from UV space to Texel space of the source (pre-scaled) image
 	int2 inputCoordStart = uvStart * InvDynamicResolutionScale * TextureDimensions.xy;
 	int2 inputCoordEnd = uvEnd * InvDynamicResolutionScale * TextureDimensions.xy;
+	
+	// Clamp start to top left corner of the input SRI.
+	inputCoordStart.x = max(inputCoordStart.x, 0.0f);
+	inputCoordStart.y = max(inputCoordStart.y, 0.0f);
 
-	// Clamp start to top left corner of the view rect.
-	inputCoordStart.x = max(inputCoordStart.x, SRIViewRectMin.x);
-	inputCoordStart.y = max(inputCoordStart.y, SRIViewRectMin.y);
-
-	// Clamp end to bottom right corner of the view rect.
-	inputCoordEnd.x = min(inputCoordEnd.x, SRIViewRectMax.x);
-	inputCoordEnd.y = min(inputCoordEnd.y, SRIViewRectMax.y);
+	// Clamp end to bottom right corner of the input SRI.
+	inputCoordEnd.x = min(inputCoordEnd.x, InputSRIDimensions.x);
+	inputCoordEnd.y = min(inputCoordEnd.y, InputSRIDimensions.y);
 
 	// start.x or y will be greater than end.x or y iff both corners were out of view on the same side.
 	// start is never clamped against the bottom-right and end is never clamped against the top-left
@@ -120,6 +118,6 @@ void RescaleVariableRateShading(
 		shadingRate = D3D12_SHADING_RATE_1X1 | (D3D12_SHADING_RATE_1X1 << CONSERVATIVE_SHADING_RATE_SHIFT);
 	}
 
-	ScaledSRITexture[DispatchThreadId.xy + UVOffset * TextureDimensions] = shadingRate & D3D12_SHADING_RATE_2X2;
-	ScaledConservativeSRITexture[DispatchThreadId.xy + UVOffset * TextureDimensions] = (shadingRate >> CONSERVATIVE_SHADING_RATE_SHIFT) & D3D12_SHADING_RATE_2X2;
+	ScaledSRITexture[DispatchThreadId.xy + ScaledUVOffset * TextureDimensions] = shadingRate & D3D12_SHADING_RATE_2X2;
+	ScaledConservativeSRITexture[DispatchThreadId.xy + ScaledUVOffset * TextureDimensions] = (shadingRate >> CONSERVATIVE_SHADING_RATE_SHIFT) & D3D12_SHADING_RATE_2X2;
 }
diff --git a/Engine/Source/Runtime/Renderer/Private/PostProcess/PostProcessTonemap.cpp b/Engine/Source/Runtime/Renderer/Private/PostProcess/PostProcessTonemap.cpp
--- a/Engine/Source/Runtime/Renderer/Private/PostProcess/PostProcessTonemap.cpp
+++ b/Engine/Source/Runtime/Renderer/Private/PostProcess/PostProcessTonemap.cpp
@@ -946,23 +946,30 @@ FScreenPassTexture AddTonemapPass(FRDGBuilder& GraphBuilder, const FViewInfo& Vi
 	{
 		// Due to the way split-screen shares the same render target for all views, make sure
 		// that we use the same texture for saving out the luminance as well to keep the memory
-		// footprint small
-		auto CachedOuputLuminance = GraphBuilder.Blackboard.Get<FOutputLuminance>();
-		if (CachedOuputLuminance == nullptr)
+		// footprint small - but only if we're outputting directly to the final render target for both views
+		// (that is, neither view has any post-processing steps after tonemapping)
+		const FOutputLuminance* CachedOutputLuminance = GraphBuilder.Blackboard.Get<FOutputLuminance>();
+		if (CachedOutputLuminance && Inputs.OverrideOutput.IsValid() &&
+			CachedOutputLuminance->Texture->Desc.Extent == Output.Texture->Desc.Extent)
 		{
-			auto NewOutputLuminance = &GraphBuilder.Blackboard.Create<FOutputLuminance>();
-			auto OutputSize = Output.Texture->Desc.GetSize();
+			OutputLuminance = CachedOutputLuminance->Texture;
+		}
+		else
+		{
+			FIntPoint OutputSize = Output.Texture->Desc.Extent;
 			const FIntPoint SDROutputSize = FIntPoint(OutputSize.X, OutputSize.Y);
 			FRDGTextureDesc Desc = FRDGTextureDesc::Create2D(
 				SDROutputSize,
 				PF_R8,
 				FClearValueBinding::Black,
 				bComputePass ? TexCreate_UAV : TexCreate_RenderTargetable);
-			OutputLuminance = NewOutputLuminance->Texture = GraphBuilder.CreateTexture(Desc, TEXT("Final Luminance"));
-		}
-		else
-		{
-			OutputLuminance = CachedOuputLuminance->Texture;
+			OutputLuminance = GraphBuilder.CreateTexture(Desc, TEXT("Final Luminance"));
+
+			if (CachedOutputLuminance == nullptr)
+			{
+				FOutputLuminance* NewOutputLuminance = &GraphBuilder.Blackboard.Create<FOutputLuminance>();
+				NewOutputLuminance->Texture = OutputLuminance;
+			}
 		}
 	}
 
diff --git a/Engine/Source/Runtime/Renderer/Private/VariableRateShading/ContrastAdaptiveImageGenerator.cpp b/Engine/Source/Runtime/Renderer/Private/VariableRateShading/ContrastAdaptiveImageGenerator.cpp
--- a/Engine/Source/Runtime/Renderer/Private/VariableRateShading/ContrastAdaptiveImageGenerator.cpp
+++ b/Engine/Source/Runtime/Renderer/Private/VariableRateShading/ContrastAdaptiveImageGenerator.cpp
@@ -63,15 +63,6 @@ TAutoConsoleVariable<int32> CVarCASPreview(
 	TEXT("0 - off, 1 - on (default)"),
 	ECVF_RenderThreadSafe);
 
-TAutoConsoleVariable<int32> CVarCASPreviewPreReprojection(
-	TEXT("r.VRS.ContrastAdaptiveShading.PreviewPreReprojection"),
-	0,
-	TEXT("Sets CAS preview to use the pre-reprojection SRI. Overrides full vs. conservative images.")
-	TEXT("0 - off (default), 1 - on"),
-	ECVF_RenderThreadSafe);
-
-
-
 /**
  * Shaders
  */
@@ -153,10 +144,9 @@ class FRescaleVariableRateShadingCS : public FGlobalShader
 		SHADER_PARAMETER_RDG_UNIFORM_BUFFER(FSceneTextureUniformParameters, SceneTextures)
 		SHADER_PARAMETER_STRUCT_REF(FViewUniformShaderParameters, View)
 		SHADER_PARAMETER_RDG_TEXTURE(Texture2D<uint>, InputSRITexture)
-		SHADER_PARAMETER(FVector2f, SRIViewRectMin)
-		SHADER_PARAMETER(FVector2f, SRIViewRectMax)
 		SHADER_PARAMETER(FVector2f, TextureDimensions)
 		SHADER_PARAMETER(FVector2f, InvTextureDimensions)
+		SHADER_PARAMETER(FVector2f, InputSRIDimensions)
 		SHADER_PARAMETER(FVector2f, ScaledSRIDimensions)
 		SHADER_PARAMETER(FVector2f, ScaledUVOffset)
 		SHADER_PARAMETER(float, InvDynamicResolutionScale)
@@ -175,8 +165,7 @@ class FRescaleVariableRateShadingCS : public FGlobalShader
 		const FMinimalSceneTextures& SceneTextures,
 		const FViewInfo& ViewInfo,
 		FRDGTextureRef InputSRITexture,
-		FVector2f ViewRectMin,
-		FVector2f ViewRectMax,
+		FVector2f InputSRIDimensions,
 		FVector2f ScaledSRIDimensions,
 		FVector2f TextureDimensions,
 		FVector2f ScaledUVOffset,
@@ -187,10 +176,9 @@ class FRescaleVariableRateShadingCS : public FGlobalShader
 		Parameters.SceneTextures = SceneTextures.UniformBuffer;
 		Parameters.View = ViewInfo.ViewUniformBuffer;
 		Parameters.InputSRITexture = InputSRITexture;
-		Parameters.SRIViewRectMin = ViewRectMin;
-		Parameters.SRIViewRectMax = ViewRectMax;
 		Parameters.TextureDimensions = TextureDimensions;
 		Parameters.InvTextureDimensions = FVector2f(1.0f / TextureDimensions.X, 1.0f / TextureDimensions.Y);
+		Parameters.InputSRIDimensions = InputSRIDimensions;
 		Parameters.ScaledSRIDimensions = ScaledSRIDimensions;
 		Parameters.ScaledUVOffset = ScaledUVOffset;
 		Parameters.InvDynamicResolutionScale = 1.0f / DynamicResolutionScale;
@@ -308,22 +296,34 @@ struct RENDERER_API FCASImageData
 		const FCASImageData* CASImageData = GraphBuilder.Blackboard.Get<FCASImageData>();
 		return CASImageData != nullptr;
 	}
-	void Create(FRDGBuilder& GraphBuilder, const FSceneViewFamily& ViewFamily, FIntRect PostProcessViewRect, bool bCreateHardwareImages, bool bCreateSoftwareImages)
+	void Create(FRDGBuilder& GraphBuilder, const FSceneViewFamily& ViewFamily, bool bCreateHardwareImages, bool bCreateSoftwareImages)
 	{
 		if (bCreateHardwareImages)
 		{
-			FRDGTextureDesc ConstructedSRIDesc = CreateSRIDesc(ViewFamily, false, PostProcessViewRect, false);
-			HardwareImages.ConstructedSRI = GraphBuilder.CreateTexture(ConstructedSRIDesc, ShadingRateTextureName);
-			FRDGTextureDesc ScaledSRIDesc = CreateSRIDesc(ViewFamily, true, PostProcessViewRect, false);
+			for (int Index = 0; Index < ViewFamily.Views.Num(); Index++)
+			{
+				const FViewInfo* ViewInfo = static_cast<const FViewInfo*>(ViewFamily.Views[Index]);
+				FIntRect LuminanceRect = ViewInfo->PrevViewInfo.LuminanceViewRectHistory;
+				FRDGTextureDesc ConstructedSRIDesc = CreateSRIDesc(ViewFamily, false, LuminanceRect, false);
+				HardwareImages.ConstructedSRIArray.Add(GraphBuilder.CreateTexture(ConstructedSRIDesc, ShadingRateTextureName));
+			}	
+
+			FRDGTextureDesc ScaledSRIDesc = CreateSRIDesc(ViewFamily, true, FIntRect(), false);
 			HardwareImages.ScaledSRI = GraphBuilder.CreateTexture(ScaledSRIDesc, ScaledShadingRateTextureName);
 			HardwareImages.ScaledConservativeSRI = GraphBuilder.CreateTexture(ScaledSRIDesc, ScaledConservativeShadingRateTextureName);
 		}
 		
 		if (bCreateSoftwareImages)
 		{
-			FRDGTextureDesc ConstructedSRIDesc = CreateSRIDesc(ViewFamily, false, PostProcessViewRect, true);
-			SoftwareImages.ConstructedSRI = GraphBuilder.CreateTexture(ConstructedSRIDesc, SoftwareShadingRateTextureName);
-			FRDGTextureDesc ScaledSRIDesc = CreateSRIDesc(ViewFamily, true, PostProcessViewRect, true);
+			for (int Index = 0; Index < ViewFamily.Views.Num(); Index++)
+			{
+				const FViewInfo* ViewInfo = static_cast<const FViewInfo*>(ViewFamily.Views[Index]);
+				FIntRect LuminanceRect = ViewInfo->PrevViewInfo.LuminanceViewRectHistory;
+				FRDGTextureDesc ConstructedSRIDesc = CreateSRIDesc(ViewFamily, false, LuminanceRect, true);
+				HardwareImages.ConstructedSRIArray.Add(GraphBuilder.CreateTexture(ConstructedSRIDesc, SoftwareShadingRateTextureName));
+			}
+
+			FRDGTextureDesc ScaledSRIDesc = CreateSRIDesc(ViewFamily, true, FIntRect(), true);
 			SoftwareImages.ScaledSRI = GraphBuilder.CreateTexture(ScaledSRIDesc, SoftwareScaledShadingRateTextureName);
 			SoftwareImages.ScaledConservativeSRI = GraphBuilder.CreateTexture(ScaledSRIDesc, SoftwareScaledConservativeShadingRateTextureName);
 		}
@@ -331,8 +331,8 @@ struct RENDERER_API FCASImageData
 
 	struct FCASImageSet
 	{
-		FRDGTextureRef ConstructedSRI;
-		FRDGTextureRef ScaledSRI;
+		TArray<FRDGTextureRef> ConstructedSRIArray; // Initial images are created per-view
+		FRDGTextureRef ScaledSRI; // Scaled and reprojected images are collated into a per-family texture matching the RT
 		FRDGTextureRef ScaledConservativeSRI;
 	};
 
@@ -375,52 +375,49 @@ static bool IsHDR10(const EDisplayOutputFormat& OutputFormat)
 		OutputFormat == EDisplayOutputFormat::HDR_ACES_2000nit_ST2084;
 }
 
-bool AddCreateShadingRateImagePass(
+void AddCreateShadingRateImagePass(
 	FRDGBuilder& GraphBuilder,
 	const FSceneViewFamily& ViewFamily,
-	FIntRect PostProcessViewRect,
-	bool bCreateHardwareImage, bool bCreateSoftwareImage)
+	bool bCreateHardwareImages, bool bCreateSoftwareImages)
 {
-	const FViewInfo& ViewInfo = *static_cast<const FViewInfo*>(ViewFamily.Views[0]);
+	FCASImageData& ImageData = GraphBuilder.Blackboard.Create<FCASImageData>();
+	ImageData.Create(GraphBuilder, ViewFamily, bCreateHardwareImages, bCreateSoftwareImages);
 
-	if (!ViewInfo.PrevViewInfo.LuminanceHistory)
+	for (int Index = 0; Index < ViewFamily.Views.Num(); Index++)
 	{
-		// Shading Rate Image unsupported
-		return false;
-	}
-	FRDGTextureRef Luminance = GraphBuilder.RegisterExternalTexture(ViewInfo.PrevViewInfo.LuminanceHistory);
-	const FCASImageData& ImageData = FCASImageData::Get(GraphBuilder);
+		const FViewInfo& ViewInfo = *static_cast<const FViewInfo*>(ViewFamily.Views[Index]);
+		FRDGTextureRef LuminanceTexture = GraphBuilder.RegisterExternalTexture(ViewInfo.PrevViewInfo.LuminanceHistory);
+		FIntRect LuminanceRect = ViewInfo.PrevViewInfo.LuminanceViewRectHistory;
 
-	{
-		// If not using HW VRS, use (8,8) groups. Otherwise, match HW tile size.
-		const FIntPoint TileSize = bCreateHardwareImage ? FVariableRateShadingImageManager::GetSRITileSize(/*bSoftwareVRS=*/false) : FIntPoint(8, 8);
-		check(TileSize.X == TileSize.Y);
-
-		FCalculateShadingRateImageCS::FPermutationDomain PermutationVector;
-		PermutationVector.Set<FCalculateShadingRateImageCS::FThreadGroupSizeXY>(TileSize.X);
-		PermutationVector.Set<FCalculateShadingRateImageCS::FOutputHardwareImage>(bCreateHardwareImage);
-		PermutationVector.Set<FCalculateShadingRateImageCS::FOutputSoftwareImage>(bCreateSoftwareImage);
-
-		TShaderMapRef<FCalculateShadingRateImageCS> ComputeShader(ViewInfo.ShaderMap, PermutationVector);
-		auto* PassParameters = GraphBuilder.AllocParameters<FCalculateShadingRateImageCS::FParameters>();
-		EDisplayOutputFormat OutputDisplayFormat = GetDisplayOutputFormat(ViewInfo);
-		FCalculateShadingRateImageCS::InitParameters(
-			*PassParameters,
-			Luminance,
-			PostProcessViewRect,
-			IsHDR10(OutputDisplayFormat),
-			bCreateHardwareImage ? GraphBuilder.CreateUAV(ImageData.HardwareImages.ConstructedSRI) : nullptr,
-			bCreateSoftwareImage ? GraphBuilder.CreateUAV(ImageData.SoftwareImages.ConstructedSRI) : nullptr);
-		FComputeShaderUtils::AddPass(
-			GraphBuilder,
-			RDG_EVENT_NAME("CreateShadingRateImage"),
-			ERDGPassFlags::AsyncCompute | ERDGPassFlags::NeverCull,
-			ComputeShader,
-			PassParameters,
-			FComputeShaderUtils::GetGroupCount(PostProcessViewRect.Size(), TileSize));
+		{
+			// If not using HW VRS, use (8,8) groups. Otherwise, match HW tile size.
+			const FIntPoint TileSize = bCreateHardwareImages ? FVariableRateShadingImageManager::GetSRITileSize(/*bSoftwareVRS=*/false) : FIntPoint(8, 8);
+			check(TileSize.X == TileSize.Y);
+
+			FCalculateShadingRateImageCS::FPermutationDomain PermutationVector;
+			PermutationVector.Set<FCalculateShadingRateImageCS::FThreadGroupSizeXY>(TileSize.X);
+			PermutationVector.Set<FCalculateShadingRateImageCS::FOutputHardwareImage>(bCreateHardwareImages);
+			PermutationVector.Set<FCalculateShadingRateImageCS::FOutputSoftwareImage>(bCreateSoftwareImages);
+
+			TShaderMapRef<FCalculateShadingRateImageCS> ComputeShader(ViewInfo.ShaderMap, PermutationVector);
+			auto* PassParameters = GraphBuilder.AllocParameters<FCalculateShadingRateImageCS::FParameters>();
+			EDisplayOutputFormat OutputDisplayFormat = GetDisplayOutputFormat(ViewInfo);
+			FCalculateShadingRateImageCS::InitParameters(
+				*PassParameters,
+				LuminanceTexture,
+				LuminanceRect,
+				IsHDR10(OutputDisplayFormat),
+				bCreateHardwareImages ? GraphBuilder.CreateUAV(ImageData.HardwareImages.ConstructedSRIArray[Index]) : nullptr,
+				bCreateSoftwareImages ? GraphBuilder.CreateUAV(ImageData.SoftwareImages.ConstructedSRIArray[Index]) : nullptr);
+			FComputeShaderUtils::AddPass(
+				GraphBuilder,
+				RDG_EVENT_NAME("CreateShadingRateImage"),
+				ERDGPassFlags::AsyncCompute | ERDGPassFlags::NeverCull,
+				ComputeShader,
+				PassParameters,
+				FComputeShaderUtils::GetGroupCount(LuminanceRect.Size(), TileSize));
+		}
 	}
-
-	return true;
 }
 
 void AddReprojectImageBasedVRSPass(
@@ -434,39 +431,24 @@ void AddReprojectImageBasedVRSPass(
 	const FCASImageData& CASImageData = FCASImageData::Get(GraphBuilder);
 	const FCASImageData::FCASImageSet& ImageSet = bReprojectSoftwareImages ? CASImageData.SoftwareImages : CASImageData.HardwareImages;
 
-	FRDGTextureRef VariableRateShadingImage = ImageSet.ConstructedSRI;
-
 	FIntPoint TileSize = FVariableRateShadingImageManager::GetSRITileSize(bReprojectSoftwareImages);
-
 	FIntPoint TextureSize = ImageSet.ScaledSRI->Desc.Extent;
 	FVector2f TextureDimensions(TextureSize.X, TextureSize.Y);
 
 	for (int32 ViewIndex = 0; ViewIndex < ViewFamily.Views.Num(); ViewIndex++)
 	{
-		const FSceneView* View = ViewFamily.Views[ViewIndex];
 		RDG_EVENT_SCOPE_CONDITIONAL(GraphBuilder, ViewFamily.Views.Num() > 1, "View%d", ViewIndex);
-		check(View->bIsViewInfo);
-		const FViewInfo* ViewInfo = static_cast<const FViewInfo*>(View);
-
-		TShaderMapRef<FRescaleVariableRateShadingCS> RescaleVariableRateShadingCS(ViewInfo->ShaderMap);
 
+		const FViewInfo* ViewInfo = static_cast<const FViewInfo*>(ViewFamily.Views[ViewIndex]);
+		
 		int32 ViewportWidth = ViewInfo->ViewRect.Width();
 		int32 ViewportHeight = ViewInfo->ViewRect.Height();
 
 		int32 ScaledTilesWide = FMath::DivideAndRoundUp(ViewportWidth, TileSize.X);
 		int32 ScaledTilesHigh = FMath::DivideAndRoundUp(ViewportHeight, TileSize.Y);
 		FVector2f ScaledSRIDimensions(ScaledTilesWide, ScaledTilesHigh);
-		
-		FIntRect SRIViewRect = ViewInfo->ViewRect.Scale(static_cast<float>(ViewInfo->PrevViewInfo.LuminanceViewRectHistory.Height()) / ViewportHeight);
-
-		FVector2f SRIViewRectMin(
-			static_cast<float>(FMath::DivideAndRoundDown(SRIViewRect.Min.X, TileSize.X)),
-			static_cast<float>(FMath::DivideAndRoundDown(SRIViewRect.Min.Y, TileSize.Y)));
-
-		FVector2f SRIViewRectMax(
-			static_cast<float>(FMath::DivideAndRoundUp(SRIViewRect.Max.X, TileSize.X)),
-			static_cast<float>(FMath::DivideAndRoundUp(SRIViewRect.Max.Y, TileSize.Y)));
 
+		// Only applies to the output SRI, since input SRIs generated in AddCreateShadingRateImagePass are always per-view and (0,0) aligned
 		FVector2f UVOffset(
 			static_cast<float>(ViewInfo->ViewRect.Min.X / TileSize.X) / TextureSize.X,
 			static_cast<float>(ViewInfo->ViewRect.Min.Y / TileSize.Y) / TextureSize.Y);
@@ -474,15 +456,14 @@ void AddReprojectImageBasedVRSPass(
 		check(ViewInfo->PrevViewInfo.LuminanceViewRectHistory.Width());
 		float DynamicResolutionScale = static_cast<float>(ViewportWidth) / (ViewInfo->PrevViewInfo.LuminanceViewRectHistory.Width());
 
+		TShaderMapRef<FRescaleVariableRateShadingCS> RescaleVariableRateShadingCS(ViewInfo->ShaderMap);
 		FRescaleVariableRateShadingCS::FParameters* PassParameters = GraphBuilder.AllocParameters<FRescaleVariableRateShadingCS::FParameters>();
-
 		FRescaleVariableRateShadingCS::InitParameters(
 			*PassParameters,
 			SceneTextures,
 			*ViewInfo,
-			VariableRateShadingImage,
-			SRIViewRectMin,
-			SRIViewRectMax,
+			ImageSet.ConstructedSRIArray[ViewIndex],
+			ImageSet.ConstructedSRIArray[ViewIndex]->Desc.Extent,
 			ScaledSRIDimensions,
 			TextureDimensions,
 			UVOffset,
@@ -539,12 +520,7 @@ void FContrastAdaptiveImageGenerator::PrepareImages(FRDGBuilder& GraphBuilder, c
 		}
 	}
 
-	FCASImageData& ImageData = GraphBuilder.Blackboard.Create<FCASImageData>();
-
-	FIntRect PostProcessViewRect = GetFullPostProcessOutputRect(ViewFamily);
-
-	ImageData.Create(GraphBuilder, ViewFamily, PostProcessViewRect, bPrepareHardwareImages, bPrepareSoftwareImages);
-	AddCreateShadingRateImagePass(GraphBuilder, ViewFamily, PostProcessViewRect, bPrepareHardwareImages, bPrepareSoftwareImages);
+	AddCreateShadingRateImagePass(GraphBuilder, ViewFamily, bPrepareHardwareImages, bPrepareSoftwareImages);
 
 	if (bPrepareHardwareImages)
 	{
@@ -577,13 +553,10 @@ FRDGTextureRef FContrastAdaptiveImageGenerator::GetDebugImage(FRDGBuilder& Graph
 	}
 
 	ESRIPreviewType::Type PreviewType = static_cast<ESRIPreviewType::Type>(ImageType);
-	if (CVarCASPreviewPreReprojection.GetValueOnRenderThread() && PreviewType != ESRIPreviewType::Off)
-	{
-		PreviewType = ESRIPreviewType::BeforeReprojection;
-	}
 
 	const FCASImageData& ImageData = FCASImageData::Get(GraphBuilder);
 	const FCASImageData::FCASImageSet& ImageSet = bGetSoftwareImage ? ImageData.SoftwareImages : ImageData.HardwareImages;
+
 	switch (PreviewType)
 	{
 	case ESRIPreviewType::Projected:
@@ -592,9 +565,9 @@ FRDGTextureRef FContrastAdaptiveImageGenerator::GetDebugImage(FRDGBuilder& Graph
 	case ESRIPreviewType::ProjectedConservative:
 		return ImageSet.ScaledConservativeSRI;
 		break;
-	case ESRIPreviewType::BeforeReprojection:
-		return ImageSet.ConstructedSRI;
-		break;
+	/*case ESRIPreviewType::BeforeReprojection:
+		return ImageSet.ConstructedSRI;  
+		break;*/
 	}
 
 	return nullptr;
