Fixed software shading rate image added to the wrong container (hardware instead of software)

#rnx

diff --git a/Engine/Source/Runtime/Renderer/Private/VariableRateShading/ContrastAdaptiveImageGenerator.cpp b/Engine/Source/Runtime/Renderer/Private/VariableRateShading/ContrastAdaptiveImageGenerator.cpp
--- a/Engine/Source/Runtime/Renderer/Private/VariableRateShading/ContrastAdaptiveImageGenerator.cpp
+++ b/Engine/Source/Runtime/Renderer/Private/VariableRateShading/ContrastAdaptiveImageGenerator.cpp
@@ -320,7 +320,7 @@ struct RENDERER_API FCASImageData
 				const FViewInfo* ViewInfo = static_cast<const FViewInfo*>(ViewFamily.Views[Index]);
 				FIntRect LuminanceRect = ViewInfo->PrevViewInfo.LuminanceViewRectHistory;
 				FRDGTextureDesc ConstructedSRIDesc = CreateSRIDesc(ViewFamily, false, LuminanceRect, true);
-				HardwareImages.ConstructedSRIArray.Add(GraphBuilder.CreateTexture(ConstructedSRIDesc, SoftwareShadingRateTextureName));
+				SoftwareImages.ConstructedSRIArray.Add(GraphBuilder.CreateTexture(ConstructedSRIDesc, SoftwareShadingRateTextureName));
 			}
 
 			FRDGTextureDesc ScaledSRIDesc = CreateSRIDesc(ViewFamily, true, FIntRect(), true);
