PR #11539: Fixed incorrect calculation of ContactShadowIntensityFadeOneOverLength in screen space shadows standalone code path.

diff --git a/Engine/Source/Runtime/Renderer/Private/Shadows/ScreenSpaceShadows.cpp b/Engine/Source/Runtime/Renderer/Private/Shadows/ScreenSpaceShadows.cpp
--- a/Engine/Source/Runtime/Renderer/Private/Shadows/ScreenSpaceShadows.cpp
+++ b/Engine/Source/Runtime/Renderer/Private/Shadows/ScreenSpaceShadows.cpp
@@ -417,7 +417,7 @@ void RenderScreenSpaceShadowsBend(
 			PassParameters->ContactShadowCastingIntensity = ContactShadowCastingIntensity;
 			PassParameters->ContactShadowNonCastingIntensity = ContactShadowNonCastingIntensity;
 			PassParameters->ContactShadowIntensityFadeStart = GContactShadowsIntensityFadeStart;
-			PassParameters->ContactShadowIntensityFadeOneOverLength = 1.0f / GContactShadowsIntensityFadeStart;
+			PassParameters->ContactShadowIntensityFadeOneOverLength = 1.0f / GContactShadowsIntensityFadeLength;
 			PassParameters->SurfaceThickness = GBendShadowsOverrideSurfaceThickness;
 
 			PassParameters->LightCoordinate = FVector4f(DispatchList.LightCoordinate_Shader[0], DispatchList.LightCoordinate_Shader[1], DispatchList.LightCoordinate_Shader[2], DispatchList.LightCoordinate_Shader[3]);
@@ -514,4 +514,4 @@ void RenderScreenSpaceShadows(
 			}
 		}
 	}
-}
\ No newline at end of file
+}
