PR #11574: Fix UE::FXRenderingUtils::CreateSceneUniformBuffer returning empty FSceneUniformBuffer

diff --git a/Engine/Source/Runtime/Renderer/Private/FXRenderingUtils.cpp b/Engine/Source/Runtime/Renderer/Private/FXRenderingUtils.cpp
--- a/Engine/Source/Runtime/Renderer/Private/FXRenderingUtils.cpp
+++ b/Engine/Source/Runtime/Renderer/Private/FXRenderingUtils.cpp
@@ -220,13 +220,12 @@ void UE::FXRenderingUtils::DistanceFields::SetupAtlasParameters(FRDGBuilder& Gra
 	}
 }
 
-FSceneUniformBuffer &UE::FXRenderingUtils::CreateSceneUniformBuffer(FRDGBuilder& GraphBuilder, const FSceneInterface* InScene)
+FSceneUniformBuffer& UE::FXRenderingUtils::CreateSceneUniformBuffer(FRDGBuilder& GraphBuilder, const FSceneInterface* InScene)
 {
 	FSceneUniformBuffer *Result = GraphBuilder.AllocObject<FSceneUniformBuffer>();
 	if (const FScene* Scene = InScene->GetRenderScene())
 	{
-		FSceneUniformBuffer SceneUniformBuffer;
-		Scene->GPUScene.FillSceneUniformBuffer(GraphBuilder, SceneUniformBuffer);
+		Scene->GPUScene.FillSceneUniformBuffer(GraphBuilder, *Result);
 	}
 	return *Result;
 }
