Made VSM stats accumulation able to record a fixed number of frames and set the output path to the Saved/Profiling directory.

#rnx

diff --git a/Engine/Source/Runtime/Renderer/Private/VirtualShadowMaps/VirtualShadowMapCacheManager.cpp b/Engine/Source/Runtime/Renderer/Private/VirtualShadowMaps/VirtualShadowMapCacheManager.cpp
--- a/Engine/Source/Runtime/Renderer/Private/VirtualShadowMaps/VirtualShadowMapCacheManager.cpp
+++ b/Engine/Source/Runtime/Renderer/Private/VirtualShadowMaps/VirtualShadowMapCacheManager.cpp
@@ -17,10 +17,12 @@
 
 CSV_DECLARE_CATEGORY_EXTERN(VSM);
 
-static TAutoConsoleVariable<int32> CVarAccumulateStats(
+static int32 GVSMAccumulateStats = 0;
+static FAutoConsoleVariableRef CVarAccumulateStats(
 	TEXT("r.Shadow.Virtual.AccumulateStats"),
-	0,
-	TEXT("When enabled, VSM stats will be collected over multiple frames and written to a CSV file"),
+	GVSMAccumulateStats,
+	TEXT("When nonzero, VSM stats will be collected over multiple frames and written to a CSV file output to the Saved/Profiling directory.\n")
+	TEXT("  If set to a number N > 0 it will auto disable and write the result after N frames, if < 0 it must be manually turned off by setting back to 0."),
 	ECVF_RenderThreadSafe
 );
 
@@ -620,6 +622,15 @@ FVirtualShadowMapArrayCacheManager::FVirtualShadowMapArrayCacheManager(FScene* I
 			}
 		}
 		TrimLoggingInfo();
+
+		if (GVSMAccumulateStats > 0)
+		{
+			OutMessages.Add(FCoreDelegates::EOnScreenMessageSeverity::Warning, FText::FromString(FString::Printf(TEXT("Virtual Shadow Map Stats Accumulation (%d frames left)"), GVSMAccumulateStats)));
+		}
+		else if (GVSMAccumulateStats < 0)
+		{
+			OutMessages.Add(FCoreDelegates::EOnScreenMessageSeverity::Warning, FText::FromString(FString::Printf(TEXT("Virtual Shadow Map Stats Accumulation Active. Set r.Shadow.Virtual.AccumulateStats to 0 to stop."))));
+		}
 	});
 #endif
 }
@@ -1011,6 +1022,12 @@ void FVirtualShadowMapArrayCacheManager::ExtractStats(FRDGBuilder& GraphBuilder,
 		AccumulatedStatsBufferRDG = GraphBuilder.RegisterExternalBuffer(AccumulatedStatsBuffer, TEXT("Shadow.Virtual.AccumulatedStatsBuffer"));
 	}
 
+	// Auto stop at zero, use -1 to record indefinitely
+	if (GVSMAccumulateStats > 0)
+	{
+		--GVSMAccumulateStats;
+	}
+
 	if (IsAccumulatingStats())
 	{
 		if (!AccumulatedStatsBuffer.IsValid())
@@ -1087,7 +1104,12 @@ void FVirtualShadowMapArrayCacheManager::ExtractStats(FRDGBuilder& GraphBuilder,
 			GPUBufferReadback = nullptr;
 		}
 
-		FString FileName = TEXT("VirtualShadowMapCacheStats.csv");// FString::Printf(TEXT("%s.csv"), *FileNameToUse);
+		FString FileName = FPaths::ProfilingDir() + FString::Printf(TEXT("VSMStats(%s).csv"), *FDateTime::Now().ToString(TEXT("%Y%m%d_%H%M%S")));
+
+		const uint32 NumRows = Tmp[0];
+
+		UE_LOG(LogRenderer, Log, TEXT("Writing VSM accumulated stats (%d frames) to file '%s'"), NumRows, *FileName);
+
 		FArchive * FileToLogTo = IFileManager::Get().CreateFileWriter(*FileName, false);
 		ensure(FileToLogTo);
 		if (FileToLogTo)
@@ -1136,8 +1158,7 @@ void FVirtualShadowMapArrayCacheManager::ExtractStats(FRDGBuilder& GraphBuilder,
 			StringToPrint += TEXT("\n");
 			FileToLogTo->Serialize(TCHAR_TO_ANSI(*StringToPrint), StringToPrint.Len());
 
-			uint32 Num = Tmp[0];
-			for (uint32 Ind = 0; Ind < Num; ++Ind)
+			for (uint32 Ind = 0; Ind < NumRows; ++Ind)
 			{
 				StringToPrint.Empty();
 
@@ -1163,7 +1184,7 @@ void FVirtualShadowMapArrayCacheManager::ExtractStats(FRDGBuilder& GraphBuilder,
 
 bool FVirtualShadowMapArrayCacheManager::IsAccumulatingStats()
 {
-	return CVarAccumulateStats.GetValueOnRenderThread() != 0;
+	return GVSMAccumulateStats != 0;
 }
 
 static uint32 GetPrimFlagsBufferSizeInDwords(int32 MaxPersistentPrimitiveIndex)
