Modified FSceneRenderer::ShouldPrepareGlobalDistanceField() to check whether materials rendered require global distance field matching behavior in UE 5.2.
- track whether any materials rendered require global distance field (store in FViewInfo).
- since FSceneRenderer::PrepareDistanceFieldScene(...) is called before visibility tasks for the current frame are finished, reuse bUsesGlobalDistanceField result from previous frame.
- this also matches 5.2 behavior w.r.t. materials requiring global distance field depending on r.DistanceFieldAO and r.AOQuality.

diff --git a/Engine/Source/Runtime/Renderer/Private/DistanceFieldAmbientOcclusion.cpp b/Engine/Source/Runtime/Renderer/Private/DistanceFieldAmbientOcclusion.cpp
--- a/Engine/Source/Runtime/Renderer/Private/DistanceFieldAmbientOcclusion.cpp
+++ b/Engine/Source/Runtime/Renderer/Private/DistanceFieldAmbientOcclusion.cpp
@@ -122,8 +122,6 @@ FAutoConsoleVariableRef CVarAOObjectDistanceField(
 	ECVF_RenderThreadSafe
 );
 
-// TODO: Find a proper way to check bUsesGlobalDistanceField material relevance and remove this cvar. Currently,
-// PrepareDistanceFieldScene is called before InitViews where we check material relevances. Maybe use last frame result?
 int32 GAOGlobalDistanceFieldDetailedNecessityCheck = 1;
 FAutoConsoleVariableRef CVarAOGlobalDistanceFieldDetailedCheck(
 	TEXT("r.AOGlobalDistanceField.DetailedNecessityCheck"),
@@ -774,6 +772,7 @@ bool FSceneRenderer::ShouldPrepareGlobalDistanceField() const
 	const bool bShouldPrepareForAO = SupportsDistanceFieldAO(Scene->GetFeatureLevel(), Scene->GetShaderPlatform())
 		&& (!GAOGlobalDistanceFieldDetailedNecessityCheck
 			|| ShouldPrepareForDistanceFieldAO()
+			|| Views[0].PrevViewInfo.bUsesGlobalDistanceField // use previous frame relevance because this frame relevance flags haven't been calculate yet
 			|| ((FXSystem != nullptr) && FXSystem->UsesGlobalDistanceField()));
 
 	const bool bShouldPrepareForLumen = IsLumenEnabled(Views[0]) && Lumen::UseGlobalSDFObjectGrid(ViewFamily);
diff --git a/Engine/Source/Runtime/Renderer/Private/SceneRendering.cpp b/Engine/Source/Runtime/Renderer/Private/SceneRendering.cpp
--- a/Engine/Source/Runtime/Renderer/Private/SceneRendering.cpp
+++ b/Engine/Source/Runtime/Renderer/Private/SceneRendering.cpp
@@ -861,6 +861,7 @@ void FViewInfo::Init()
 	bIsViewInfo = true;
 	
 	bStatePrevViewInfoIsReadOnly = true;
+	bUsesGlobalDistanceField = false;
 	bUsesLightingChannels = false;
 	bTranslucentSurfaceLighting = false;
 	bFogOnlyOnRenderedOpaque = false;
@@ -4793,6 +4794,15 @@ static void RenderViewFamilies_RenderThread(FRHICommandListImmediate& RHICmdList
 		}
 #endif
 
+		// Copy relevant data from ViewInfo to ViewState->PrevFrameViewInfo
+		for (FViewInfo& View : SceneRenderer->Views)
+		{
+			if (View.ViewState)
+			{
+				View.ViewState->PrevFrameViewInfo.bUsesGlobalDistanceField = View.bUsesGlobalDistanceField;
+			}
+		}
+
 		if (SceneRenderer->ViewFamily.ProfileSceneRenderTime)
 		{
 			*SceneRenderer->ViewFamily.ProfileSceneRenderTime = (float)FPlatformTime::ToSeconds64(FPlatformTime::Cycles64() - FamilyRenderStart);
diff --git a/Engine/Source/Runtime/Renderer/Private/SceneRendering.h b/Engine/Source/Runtime/Renderer/Private/SceneRendering.h
--- a/Engine/Source/Runtime/Renderer/Private/SceneRendering.h
+++ b/Engine/Source/Runtime/Renderer/Private/SceneRendering.h
@@ -1067,6 +1067,8 @@ struct FPreviousViewInfo
 	// Scene color's PreExposure.
 	float SceneColorPreExposure = 1.0f;
 
+	bool bUsesGlobalDistanceField = false;
+
 	// Depth buffer and Normals of the previous frame generating this history entry for bilateral kernel rejection.
 	TRefCountPtr<IPooledRenderTarget> DepthBuffer;
 	TRefCountPtr<IPooledRenderTarget> GBufferA;
@@ -1445,9 +1447,9 @@ public:
 	/** Whether we should submit new queries this frame. (used to disable occlusion queries completely. */
 	uint32 bDisableQuerySubmissions : 1;
 	/** Whether the view has any materials that use the global distance field. */
+	uint32 bUsesGlobalDistanceField : 1;
 	uint32 bUsesLightingChannels : 1;
 	uint32 bTranslucentSurfaceLighting : 1;
-	/** Whether the view has any materials that read from scene depth. */
 	uint32 bCustomDepthStencilValid : 1;
 	uint32 bUsesCustomDepth : 1;
 	uint32 bUsesCustomStencil : 1;
diff --git a/Engine/Source/Runtime/Renderer/Private/SceneVisibility.cpp b/Engine/Source/Runtime/Renderer/Private/SceneVisibility.cpp
--- a/Engine/Source/Runtime/Renderer/Private/SceneVisibility.cpp
+++ b/Engine/Source/Runtime/Renderer/Private/SceneVisibility.cpp
@@ -1204,6 +1204,7 @@ void FRelevancePacket::Finalize()
 #endif
 
 	WriteView.ShadingModelMaskInView |= CombinedShadingModelMask;
+	WriteView.bUsesGlobalDistanceField |= bUsesGlobalDistanceField;
 	WriteView.bUsesLightingChannels |= bUsesLightingChannels;
 	WriteView.bTranslucentSurfaceLighting |= bTranslucentSurfaceLighting;
 	WriteView.bSceneHasSkyMaterial |= bSceneHasSkyMaterial;
@@ -1298,6 +1299,7 @@ void FRelevancePacket::ComputeRelevance(FDynamicPrimitiveIndexList& DynamicPrimi
 	bSceneHasSkyMaterial = 0;
 	bHasSingleLayerWaterMaterial = 0;
 	bUsesSecondStageDepthPass = 0;
+	bUsesGlobalDistanceField = false;
 	bUsesLightingChannels = false;
 	bTranslucentSurfaceLighting = false;
 	const EShadingPath ShadingPath = GetFeatureLevelShadingPath(Scene.GetFeatureLevel());
@@ -1890,6 +1892,7 @@ void FRelevancePacket::ComputeRelevance(FDynamicPrimitiveIndexList& DynamicPrimi
 		SubstrateUintPerPixel = FMath::Max(SubstrateUintPerPixel, ViewRelevance.SubstrateUintPerPixel);
 		bUsesComplexSpecialRenderPath |= ViewRelevance.bUsesComplexSpecialRenderPath;
 		SubstrateClosureCountMask |= ViewRelevance.SubstrateClosureCountMask;
+		bUsesGlobalDistanceField |= ViewRelevance.bUsesGlobalDistanceField;
 		bUsesLightingChannels |= ViewRelevance.bUsesLightingChannels;
 		bTranslucentSurfaceLighting |= ViewRelevance.bTranslucentSurfaceLighting;
 		bUsesCustomDepth |= (ViewRelevance.CustomDepthStencilUsageMask & 1) > 0;
diff --git a/Engine/Source/Runtime/Renderer/Private/SceneVisibilityPrivate.h b/Engine/Source/Runtime/Renderer/Private/SceneVisibilityPrivate.h
--- a/Engine/Source/Runtime/Renderer/Private/SceneVisibilityPrivate.h
+++ b/Engine/Source/Runtime/Renderer/Private/SceneVisibilityPrivate.h
@@ -777,6 +777,7 @@ private:
 	bool bUsesComplexSpecialRenderPath = false;
 	bool bHasDistortionPrimitives = false;
 	bool bHasCustomDepthPrimitives = false;
+	bool bUsesGlobalDistanceField = false;
 	bool bUsesLightingChannels = false;
 	bool bTranslucentSurfaceLighting = false;
 	bool bUsesCustomDepth = false;
