Fix sampling of Light Function Atlas when Light Function Scale is non-uniform

diff --git a/Engine/Shaders/Private/LightFunctionAtlas/LightFunctionAtlasCommon.usf b/Engine/Shaders/Private/LightFunctionAtlas/LightFunctionAtlasCommon.usf
--- a/Engine/Shaders/Private/LightFunctionAtlas/LightFunctionAtlasCommon.usf
+++ b/Engine/Shaders/Private/LightFunctionAtlas/LightFunctionAtlasCommon.usf
@@ -35,11 +35,14 @@ FLightFunctionColor GetLocalLightFunctionCommon(float3 TranslatedWorldPosition,
 
 	float4 LightFunctionLightDataParameters = LightFunctionAtlas.LightInfoDataBuffer[BaseBufferIndex + 0];
 	const float  FadeDistance				= LightFunctionLightDataParameters.x;
-	const uint   PackedLightInfoDataParams	= asuint(LightFunctionLightDataParameters.y);
-	const float2 Slot_MinUV					= LightFunctionLightDataParameters.zw;
+	const uint   PackedLightInfoDataParams0 = asuint(LightFunctionLightDataParameters.y);
+	const uint   PackedLightInfoDataParams1 = asuint(LightFunctionLightDataParameters.z);
+	const float	 TanOuterAngle				= LightFunctionLightDataParameters.w;
 
-	const uint   LightType					= PackedLightInfoDataParams & 0xFF;
-	const float  DisabledBrightness			= f16tof32((PackedLightInfoDataParams >> 8) & 0xFFFF);
+	const uint   LightType					= PackedLightInfoDataParams0 & 0xFF;
+	const float  DisabledBrightness			= f16tof32((PackedLightInfoDataParams0 >> 8) & 0xFFFF);
+
+	const float2 Slot_MinUV					= float2(((PackedLightInfoDataParams1) & 0xFFFF) * rcp(65536.0f), ((PackedLightInfoDataParams1 >> 16) & 0xFFFF) * rcp(65536.0f));
 
 	float4x4 ActualLightFunctionMatrix = float4x4(
 		LightFunctionAtlas.LightInfoDataBuffer[BaseBufferIndex + 1],
@@ -48,22 +51,24 @@ FLightFunctionColor GetLocalLightFunctionCommon(float3 TranslatedWorldPosition,
 		LightFunctionAtlas.LightInfoDataBuffer[BaseBufferIndex + 4]);
 	float4 HomogeneousShadowPosition = mul(float4(TranslatedWorldPosition, 1), ActualLightFunctionMatrix);
 
-	float2 LightFunctionUV;
+	float3 LightVector = HomogeneousShadowPosition.zyx / HomogeneousShadowPosition.w;
+	float2 LightFunctionUV = LightVector.xy;
+
 	switch (LightType)
 	{
 
 	case LIGHT_TYPE_SPOT:
 	{
-		HomogeneousShadowPosition.xyz /= HomogeneousShadowPosition.www;
-		LightFunctionUV = HomogeneousShadowPosition.xy;
-		LightFunctionUV = 1 - LightFunctionUV;
+		LightFunctionUV = LightVector.xy / (LightVector.z * TanOuterAngle) * .5f + .5f;
 		break;
 	}
 
 	case LIGHT_TYPE_POINT:
 	{
-		float3 LightVec = normalize(HomogeneousShadowPosition).zyx;
-		LightFunctionUV = float2((atan2Fast(LightVec.y, LightVec.x) + PI) / (2 * PI), acosFast(LightVec.z) / PI);
+		float3 UnitLightVector = normalize(LightVector);
+		// Setup 2d UVs for a point light that map the texture using spherical coordinates, which is how max handles a 2d texture projector on a point light
+		// Artists should use a cubemap indexed by a light vector to get better quality
+		LightFunctionUV = float2((atan2(UnitLightVector.y, UnitLightVector.x) + PI) / (2 * PI), acos(UnitLightVector.z) / PI);
 		break;
 	}
 
@@ -71,12 +76,14 @@ FLightFunctionColor GetLocalLightFunctionCommon(float3 TranslatedWorldPosition,
 	//case LIGHT_TYPE_DIRECTIONAL:
 	default:
 	{
-		LightFunctionUV = frac(HomogeneousShadowPosition.zy);
+		// no op
 		break;
 	}
 
 	}
 
+	LightFunctionUV = frac(LightFunctionUV);
+
 	// Slot min uv and size already account for half texel offsets
 	float2 AtlasUV = Slot_MinUV + saturate(LightFunctionUV) * LightFunctionAtlas.Slot_UVSize.xx;
 
diff --git a/Engine/Source/Runtime/Renderer/Private/LightFunctionAtlas.cpp b/Engine/Source/Runtime/Renderer/Private/LightFunctionAtlas.cpp
--- a/Engine/Source/Runtime/Renderer/Private/LightFunctionAtlas.cpp
+++ b/Engine/Source/Runtime/Renderer/Private/LightFunctionAtlas.cpp
@@ -117,15 +117,18 @@ FAutoConsoleVariableRef CVarLumenUsesLightFunctionAtlas(
 
 //////////////////////////////////////////////////////////////////////////
 
+static const uint32 MAX_LIGHT_FUNCTION_ATLAS_SLOT_RESOLUTION = 256;
+static const uint32 MAX_LIGHT_FUNCTION_ATLAS_EDGE_SIZE = 16;
+
 static uint32 GetAtlasSlotResolution()
 {
-	const uint32 AtlasSlotResolution = FMath::Clamp(CVarLightFunctionAtlasSlotResolution.GetValueOnRenderThread(), 32, 256);
+	const uint32 AtlasSlotResolution = FMath::Clamp(CVarLightFunctionAtlasSlotResolution.GetValueOnRenderThread(), 32, MAX_LIGHT_FUNCTION_ATLAS_SLOT_RESOLUTION);
 	return AtlasSlotResolution;
 }
 
 static uint32 GetAtlasEdgeSize()
 {
-	const uint32 AtlasEdgeSize = FMath::Clamp(CVarLightFunctionAtlasSize.GetValueOnRenderThread(), 4, 16);// 16x16 is the maximum slot count of LIGHT_FUNCTION_ATLAS_MAX_LIGHT_FUNCTION_COUNT=256 we currently allow
+	const uint32 AtlasEdgeSize = FMath::Clamp(CVarLightFunctionAtlasSize.GetValueOnRenderThread(), 4, MAX_LIGHT_FUNCTION_ATLAS_EDGE_SIZE);// 16x16 is the maximum slot count of LIGHT_FUNCTION_ATLAS_MAX_LIGHT_FUNCTION_COUNT=256 we currently allow
 	return AtlasEdgeSize;
 }
 
@@ -551,7 +554,6 @@ void FLightFunctionAtlas::RenderLightFunctionAtlas(FRDGBuilder& GraphBuilder, TA
 		LightFunctionAtlasGlobalParameters->Slot_UVSize = (AtlasSlotResolution - 1.0f) / AtlasResolution; // -1.0 because we remove a bit more than half a texel at the border.
 
 		const FIntPoint LightFunctionResolution = FIntPoint(AtlasSlotResolution, AtlasSlotResolution);
-		const FMatrix TranslatedWorldToWorld = FTranslationMatrix(-View.ViewMatrices.GetPreViewTranslation());
 
 		// Write the light data needed to rotate and fade the light function in the world.
 		const uint32 InitialLightInfoDataLightCount = FMath::DivideAndRoundUp(uint32(EffectiveLocalLightSlotArray.Num()), 32u) * 32u;// Alloted with 32 lights step to better reuse shared buffers pool.
@@ -574,63 +576,40 @@ void FLightFunctionAtlas::RenderLightFunctionAtlas(FRDGBuilder& GraphBuilder, TA
 
 			float ShadowFadeFraction = 1.0f;
 
-			uint8 LightType = Proxy->GetLightType();
-			if (LightType == LightType_Spot)
-			{
-				// Spotlights needs to utilize their shadow projection when rendering the light function 
-				TArray<FWholeSceneProjectedShadowInitializer, TInlineAllocator<6>> Initializers;
-				Proxy->GetWholeSceneProjectedShadowInitializer(*View.Family, Initializers);
-
-				FWholeSceneProjectedShadowInitializer ShadowInitializer = Initializers[0];
-				FProjectedShadowInfo ProjectedShadowInfo;
-				ProjectedShadowInfo.SetupWholeSceneProjection(
-					LightSceneInfo,
-					&View,
-					ShadowInitializer,
-					LightFunctionResolution.X,
-					LightFunctionResolution.Y,
-					LightFunctionResolution.X,
-					LightFunctionResolution.Y,
-					0
-				);
-				//ShadowFadeFraction = ProjectedShadowInfo.FadeAlphas.Num() == 0 ? 1.0f : ProjectedShadowInfo.FadeAlphas[0];
-
-				FVector4f ShadowmapMinMaxValue;
-				LightInfoDataBufferPtr[LightIndex].Transform = FMatrix44f(TranslatedWorldToWorld * ProjectedShadowInfo.GetWorldToShadowMatrix(ShadowmapMinMaxValue, &LightFunctionResolution));
-			}
-			else if (LightType == LightType_Point)
-			{
-				LightInfoDataBufferPtr[LightIndex].Transform = FMatrix44f(TranslatedWorldToWorld * Proxy->GetWorldToLight());
-			}
-			else if (LightType == LightType_Rect)
+			FMatrix44f TranslatedWorldToLight;
+
 			{
 				const FVector Scale = Proxy->GetLightFunctionScale();
 				// Switch x and z so that z of the user specified scale affects the distance along the light direction
-				const FVector InverseScale = FVector(1.0 / Scale.Z, 1.0 / Scale.Y, 1.0 / Scale.X);
-				const FMatrix WorldToLight = Proxy->GetWorldToLight() * FScaleMatrix(InverseScale);
-
-				LightInfoDataBufferPtr[LightIndex].Transform = FMatrix44f(TranslatedWorldToWorld * WorldToLight);
+				const FVector InverseScale = FVector(1.f / Scale.Z, 1.f / Scale.Y, 1.f / Scale.X);
+				const FMatrix WorldToLight = Proxy->GetWorldToLight() * FScaleMatrix(FVector(InverseScale));
+				TranslatedWorldToLight = FMatrix44f(FTranslationMatrix(-View.ViewMatrices.GetPreViewTranslation()) * WorldToLight);
 			}
-			else if (LightType == LightType_Directional)
-			{
-				const FVector LightDirection = Proxy->GetDirection().GetSafeNormal();
 
-				const FVector Scale = LightSceneInfo->Proxy->GetLightFunctionScale();
-				// Switch x and z so that z of the user specified scale affects the distance along the light direction
-				const FVector InverseScale = FVector(1.0 / Scale.Z, 1.0 / Scale.Y, 1.0 / Scale.X);
+			LightInfoDataBufferPtr[LightIndex].Transform = TranslatedWorldToLight;
 
-				const FMatrix WorldToLight = LightSceneInfo->Proxy->GetWorldToLight() * FScaleMatrix(FVector(InverseScale));
-				LightInfoDataBufferPtr[LightIndex].Transform = FMatrix44f(TranslatedWorldToWorld * WorldToLight);
-			}
+			uint8 LightType = Proxy->GetLightType();
 			
 			const uint8 LightFunctionAtlasSlotIndex = LightSlot.LightFunctionAtlasSlotIndex;
 			FFloat16 PackedDisabledBrightness(Proxy->GetLightFunctionDisabledBrightness());
 			uint32 PackedLightInfoDataParams = uint32(LightType) | (uint32(PackedDisabledBrightness.Encoded) << 8);
 
-			EffectiveLightFunctionSlot& AtlasSlot = EffectiveLightFunctionSlotArray[LightFunctionAtlasSlotIndex];
+			const EffectiveLightFunctionSlot& AtlasSlot = EffectiveLightFunctionSlotArray[LightFunctionAtlasSlotIndex];
+
+			static_assert(MAX_LIGHT_FUNCTION_ATLAS_SLOT_RESOLUTION * MAX_LIGHT_FUNCTION_ATLAS_EDGE_SIZE <= 32 * 1024, 
+				"Unable to pack slot UVs into uint16 when atlas resolution is larger than 32K");
+
+			const uint32 PackedAtlasSlotMinU = uint32(round(AtlasSlot.MinU * 65536.0f));
+			const uint32 PackedAtlasSlotMinV = uint32(round(AtlasSlot.MinV * 65536.0f));
+			const uint32 PackedAtlasSlotMinUV = (PackedAtlasSlotMinU | (PackedAtlasSlotMinV << 16));
+
+			ensure(FMath::IsNearlyEqual((PackedAtlasSlotMinUV & 0xFFFF) / 65536.0f, AtlasSlot.MinU));
+			ensure(FMath::IsNearlyEqual(((PackedAtlasSlotMinUV >> 16) & 0xFFFF) / 65536.0f, AtlasSlot.MinV));
+
+			const float TanOuterAngle = LightType == LightType_Spot ? FMath::Tan(LightSceneInfo->Proxy->GetOuterConeAngle()) : 1.0f;
 
 			// ShadowFadeFraction is unused.
-			LightInfoDataBufferPtr[LightIndex].Parameters = FVector4f(Proxy->GetLightFunctionFadeDistance(), FMath::AsFloat(PackedLightInfoDataParams), AtlasSlot.MinU, AtlasSlot.MinV);
+			LightInfoDataBufferPtr[LightIndex].Parameters = FVector4f(Proxy->GetLightFunctionFadeDistance(), FMath::AsFloat(PackedLightInfoDataParams), FMath::AsFloat(PackedAtlasSlotMinUV), TanOuterAngle);
 
 			LightIndex++;
 		}
