Fix distortion rendering ensure when the velocity texture is not produced, replace it by a black dummy.

#tests Automation test in the jira

diff --git a/Engine/Source/Runtime/Renderer/Private/DistortionRendering.cpp b/Engine/Source/Runtime/Renderer/Private/DistortionRendering.cpp
--- a/Engine/Source/Runtime/Renderer/Private/DistortionRendering.cpp
+++ b/Engine/Source/Runtime/Renderer/Private/DistortionRendering.cpp
@@ -550,7 +550,9 @@ void FDeferredShadingSceneRenderer::RenderDistortion(
 			{
 				FTAAPassParameters TAASettings(View);
 				TAASettings.SceneDepthTexture = SceneDepthTexture;
-				TAASettings.SceneVelocityTexture = SceneVelocityTexture;
+
+				// We need a valid velocity buffer texture. Use black (no velocity) if it's not produced.
+				TAASettings.SceneVelocityTexture = GetIfProduced(SceneVelocityTexture, GraphBuilder.RegisterExternalTexture(GSystemTextures.BlackDummy));
 				TAASettings.Pass = ETAAPassConfig::Main;		// Reusing main config for now. We could add a ReoughRefraction config forcing 111110 format.
 				TAASettings.SceneColorInput = TAASceneColorTexture;
 				TAASettings.bOutputRenderTargetable = true;
