Optimize hair lighting for local lights

* Add OnePass VSM ShadowMaskBit support for DeferredLightPS with hair
* Add OnePass HairTransmittance support for DeferredLightPS with hair.
* Change voxel traversal to use Wave32 which speed up an extra 5% due to incoherence within wave.

This saves 15% on lighting on test content.

diff --git a/Engine/Shaders/Private/DeferredLightPixelShaders.usf b/Engine/Shaders/Private/DeferredLightPixelShaders.usf
--- a/Engine/Shaders/Private/DeferredLightPixelShaders.usf
+++ b/Engine/Shaders/Private/DeferredLightPixelShaders.usf
@@ -121,6 +121,7 @@ float GetExposure()
 // One pass projection
 Texture2D<uint4> ShadowMaskBits;
 int VirtualShadowMapId;
+int LightSceneId;
 #endif
 
 float4 GetLightAttenuationFromShadow(in FInputParams InputParams, float SceneDepth, float3 TranslatedWorldPosition)
@@ -456,21 +457,73 @@ void DeferredLightPixelMain(
 		else
 			L = LightData.Direction;
 
+		// Fetch precomputed transmittance
+		FHairTransmittanceMask TransmittanceMask = InitHairTransmittanceMask();
+
+		#if USE_VIRTUAL_SHADOW_MAP_MASK
+		{
+			const uint2 LocalPixelPos = PixelCoord.xy - uint2(View.ViewRectMinAndSize.xy);
+			const FCulledLightsGridData LightGridData = VirtualShadowMapGetLightsGrid(LocalPixelPos, SceneDepth);
+			uint LightCount = min(GetPackedShadowMaskMaxLightCount(), LightGridData.NumLights);
+
+			uint GridLightListIndex = 0;
+			for (; GridLightListIndex < LightCount; ++GridLightListIndex)
+			{
+				const FLocalLightData LightData = VirtualShadowMapGetLocalLightData(LightGridData, GridLightListIndex);
+				if (LightData.LightSceneId == LightSceneId)
+				{
+					break;
+				}
+			}
+			if (GridLightListIndex < GetPackedShadowMaskMaxLightCount())
+			{
+				const uint TotalSampleCount = NodeCount;
+				const uint TransmittanceSampleIndex = SampleIndex + GridLightListIndex * TotalSampleCount;
+				const uint TransmittanceSampleMaxCount = TotalSampleCount * GetPackedShadowMaskMaxLightCount();
+				if (TransmittanceSampleIndex < TransmittanceSampleMaxCount)
+				{
+					TransmittanceMask = UnpackTransmittanceMask(HairTransmittanceBuffer[TransmittanceSampleIndex]);
+				}
+			}
+		}
+		#else
+		if (SampleIndex < HairTransmittanceBufferMaxCount)
+		{
+			TransmittanceMask = UnpackTransmittanceMask(HairTransmittanceBuffer[SampleIndex]);
+		}
+		#endif
+
+
 		const float3 V = normalize(-CameraVector);
 		LightData.HairTransmittance = GetHairTransmittance(
 			V,
 			L,
 			HairScreenSpaceData.GBuffer,
-			SampleIndex,
-			HairTransmittanceBufferMaxCount,
-			HairTransmittanceBuffer,
+			TransmittanceMask,
 			View.HairScatteringLUTTexture,
 			View.HairScatteringLUTSampler,
 			View.HairComponents);
 	}
 
 	const float Dither = InterleavedGradientNoise(PixelCoord, View.StateFrameIndexMod8);
-	float4 LightAttenuation = HairShadowMaskValid ? ScreenShadowMaskSubPixelTexture.Load(uint3(PixelCoord,0)) : 1;
+	float4 LightAttenuation = 1.f;
+	#if USE_VIRTUAL_SHADOW_MAP_MASK
+	if (VirtualShadowMapId != INDEX_NONE)
+	{
+		LightAttenuation = GetVirtualShadowMapMaskForLight(
+			ShadowMaskBits[PixelCoord],
+			PixelCoord,
+			SceneDepth,
+			VirtualShadowMapId,
+			TranslatedWorldPosition);
+	}
+	#else
+	if (HairShadowMaskValid)
+	{
+		LightAttenuation = ScreenShadowMaskSubPixelTexture.Load(uint3(PixelCoord, 0));
+	}
+	#endif
+
 	if (any(ShadowChannelMask < 1.0f))
 	{
 		LightAttenuation = dot(LightAttenuation, ShadowChannelMask).xxxx;
diff --git a/Engine/Shaders/Private/HairStrands/HairStrandsDeepTransmittanceDualScattering.ush b/Engine/Shaders/Private/HairStrands/HairStrandsDeepTransmittanceDualScattering.ush
--- a/Engine/Shaders/Private/HairStrands/HairStrandsDeepTransmittanceDualScattering.ush
+++ b/Engine/Shaders/Private/HairStrands/HairStrandsDeepTransmittanceDualScattering.ush
@@ -146,32 +146,6 @@ FHairTransmittanceData GetHairTransmittance(
 	return Out;
 }
 
-FHairTransmittanceData GetHairTransmittance(
-	const float3 V,
-	const float3 L,
-	FGBufferData GBuffer,
-	const uint SampleOffset,
-	const uint HairTransmittanceBufferMaxCount,
-	Buffer<uint> HairTransmittanceBuffer,
-	Texture3D<float4> HairLUTTexture,
-	SamplerState HairLUTSampler,
-	const uint HairComponents)
-{
-	FHairTransmittanceMask TransmittanceMask = InitHairTransmittanceMask();
-	if (SampleOffset < HairTransmittanceBufferMaxCount)
-	{
-		TransmittanceMask = UnpackTransmittanceMask(HairTransmittanceBuffer[SampleOffset]);
-	}
-	return GetHairTransmittance(
-		V,
-		L,
-		GBuffer,
-		TransmittanceMask,
-		HairLUTTexture,
-		HairLUTSampler,
-		HairComponents);
-}
-
 float3 GetHairTransmittanceOnly(FHairTransmittanceMask TransmittanceMask, FGBufferData GBuffer, float SinLightAngle, Texture3D<float4> InHairScatteringLUTTexture, SamplerState InHairLUTSampler)
 {
 	// Always shift the hair count by one to remove self-occlusion/shadow aliasing and have smoother transition
diff --git a/Engine/Shaders/Private/HairStrands/HairStrandsDeepTransmittanceMask.usf b/Engine/Shaders/Private/HairStrands/HairStrandsDeepTransmittanceMask.usf
--- a/Engine/Shaders/Private/HairStrands/HairStrandsDeepTransmittanceMask.usf
+++ b/Engine/Shaders/Private/HairStrands/HairStrandsDeepTransmittanceMask.usf
@@ -266,7 +266,7 @@ uint				LightChannelMask;
 float4				TranslatedLightPosition_LightDirection;
 float				LightRadius;
 #if PERMUTATION_ONE_PASS
-Texture2D<uint> 	ShadowMaskBitsTexture;
+Texture2D<uint4> 	ShadowMaskBitsTexture;
 #else
 Texture2D<float4> 	RayMarchMaskTexture;
 float4				ShadowChannelMask;
@@ -385,8 +385,7 @@ void MainCS(uint2 DispatchThreadId : SV_DispatchThreadID)
 		TranslatedWorldPosition += VirtualVoxel.TranslatedWorldOffsetStereoCorrection;
 	}
 
-	// TODO: Need to update the shadow mask encoding/decoding functions here
-	//const uint4 PackedShadowMask = ~ShadowMaskBitsTexture[PixelCoord];
+	const uint4 PackedShadowMask = ShadowMaskBitsTexture[PixelCoord];
 	const float3 Random = GetHairVoxelJitter(PixelCoord.xy, View.StateFrameIndexMod8, VirtualVoxel.JitterMode);
 
 	const uint2 LocalPixelPos = PixelCoord;
@@ -403,27 +402,23 @@ void MainCS(uint2 DispatchThreadId : SV_DispatchThreadID)
 	{
 		const uint OutputIndex = SampleIndex + MaxVisibilityNodeCount * LightIndex;
 
-		// TODO: Add early out when the shadow mask bits are zero
-		#if 0
+		// The Virtual Shadow Remap stores directional lights first
+		const FLocalLightData LightData 	 = GetLocalLightData(CulledLightGridData.DataStartIndex + LightIndex, EyeIndex);
+		const float3 TranslatedLightPosition = LightData.LightPositionAndInvRadius.xyz; // LightPositionAndInvRadius is already in translated world space
+
+		// Early out when the shadow mask bits are zero
 		if (UnpackShadowMask(PackedShadowMask, LightIndex) == 0)
 		{
 			OutTransmittanceMask[OutputIndex] = PackTransmittanceMask(InitHairTransmittanceMask());
 			continue;
 		}
-		#endif
 
-		// TODO: add support for light channel
-		#if 0
-		if ((LightChannelMask & Sample.LightChannelMask) == 0)
+		// Early out when mismatching light channel
+		if ((Sample.LightChannelMask & UnpackLightingChannelMask(LightData)) == 0)
 		{
 			OutTransmittanceMask[OutputIndex] = InitNullPackedHairTransmittanceMask();
 			continue;
 		}
-		#endif	
-
-		// The Virtual Shadow Remap stores directional lights first
-		const FLocalLightData LightData 	 = GetLocalLightData(CulledLightGridData.DataStartIndex + LightIndex, EyeIndex);
-		const float3 TranslatedLightPosition = LightData.LightPositionAndInvRadius.xyz; // LightPositionAndInvRadius is already in translated world space
 
 		FTransmittanceSettings Settings;
 		Settings.bIsDirectional				= false;
diff --git a/Engine/Source/Runtime/Renderer/Private/DeferredShadingRenderer.h b/Engine/Source/Runtime/Renderer/Private/DeferredShadingRenderer.h
--- a/Engine/Source/Runtime/Renderer/Private/DeferredShadingRenderer.h
+++ b/Engine/Source/Runtime/Renderer/Private/DeferredShadingRenderer.h
@@ -973,7 +973,10 @@ private:
 		FRDGTextureRef ScreenShadowMaskSubPixelTexture,
 		FRDGTextureRef LightingChannelsTexture,
 		const FHairStrandsTransmittanceMaskData& InTransmittanceMaskData,
-		const bool bForwardRendering);
+		const bool bForwardRendering,
+		TRDGUniformBufferRef<FVirtualShadowMapUniformParameters> VirtualShadowMapUniformBuffer = nullptr,
+		FRDGTextureRef ShadowMaskBits = nullptr,
+		int32 VirtualShadowMapId = INDEX_NONE);
 
 	/** Renders an array of simple lights using standard deferred shading. */
 	void RenderSimpleLightsStandardDeferred(
diff --git a/Engine/Source/Runtime/Renderer/Private/HairStrands/HairStrandsTransmittance.cpp b/Engine/Source/Runtime/Renderer/Private/HairStrands/HairStrandsTransmittance.cpp
--- a/Engine/Source/Runtime/Renderer/Private/HairStrands/HairStrandsTransmittance.cpp
+++ b/Engine/Source/Runtime/Renderer/Private/HairStrands/HairStrandsTransmittance.cpp
@@ -227,6 +227,7 @@ public:
 		FGlobalShader::ModifyCompilationEnvironment(Parameters, OutEnvironment);
 		FForwardLightingParameters::ModifyCompilationEnvironment(Parameters.Platform, OutEnvironment);
 		OutEnvironment.SetDefine(TEXT("SHADER_TRANSMITTANCE_VOXEL"), 1);
+		OutEnvironment.CompilerFlags.Add(CFLAG_Wave32);
 	}
 };
 
@@ -477,6 +478,7 @@ public:
 		FGlobalShader::ModifyCompilationEnvironment(Parameters, OutEnvironment);
 		OutEnvironment.SetRenderTargetOutputFormat(0, PF_B8G8R8A8);
 		OutEnvironment.SetDefine(TEXT("SHADER_SHADOWMASK_VOXEL"), 1);
+		OutEnvironment.CompilerFlags.Add(CFLAG_Wave32);
 	}
 };
 
diff --git a/Engine/Source/Runtime/Renderer/Private/LightRendering.cpp b/Engine/Source/Runtime/Renderer/Private/LightRendering.cpp
--- a/Engine/Source/Runtime/Renderer/Private/LightRendering.cpp
+++ b/Engine/Source/Runtime/Renderer/Private/LightRendering.cpp
@@ -867,6 +867,7 @@ class FDeferredLightPS : public FGlobalShader
 		// For virtual shadow map mask
 		SHADER_PARAMETER_RDG_UNIFORM_BUFFER(FVirtualShadowMapUniformParameters, VirtualShadowMap)
 		SHADER_PARAMETER(int32, VirtualShadowMapId)
+		SHADER_PARAMETER(int32, LightSceneId)
 		SHADER_PARAMETER_RDG_TEXTURE(Texture2D, ShadowMaskBits)
 		// Heterogeneous Volume data
 		//SHADER_PARAMETER_RDG_UNIFORM_BUFFER(FOrthoVoxelGridUniformBufferParameters, OrthoGridUniformBuffer)
@@ -1583,13 +1584,13 @@ void FDeferredShadingSceneRenderer::RenderLights(
 						{
 							SharedScreenShadowMaskTexture = GraphBuilder.CreateTexture(SharedScreenShadowMaskTextureDesc, TEXT("ShadowMaskTexture"));
 						}
-						if (!SharedScreenShadowMaskSubPixelTexture && bUseHairLighting)
+						if (!SharedScreenShadowMaskSubPixelTexture && bUseHairLighting && !bElideScreenShadowMask)
 						{
 							SharedScreenShadowMaskSubPixelTexture = GraphBuilder.CreateTexture(SharedScreenShadowMaskTextureDesc, TEXT("ShadowMaskSubPixelTexture"));
 						}
 					}
 					ScreenShadowMaskTexture = bElideScreenShadowMask ? nullptr : SharedScreenShadowMaskTexture;
-					ScreenShadowMaskSubPixelTexture = SharedScreenShadowMaskSubPixelTexture;
+					ScreenShadowMaskSubPixelTexture = bElideScreenShadowMask ? nullptr : SharedScreenShadowMaskSubPixelTexture;
 				}
 
 				FString LightNameWithLevel;
@@ -2098,9 +2099,23 @@ void FDeferredShadingSceneRenderer::RenderLights(
 
 						if (HairStrands::HasViewHairStrandsData(View))
 						{
+							// If the light elided the screen space shadow mask, sample directly from the packed shadow mask
+							int32 VirtualShadowMapId = INDEX_NONE;
+							if (bElideScreenShadowMask)
+							{
+								INC_DWORD_STAT(STAT_VSMLocalProjectionOnePassFast);
+								VirtualShadowMapId = VisibleLightInfo.GetVirtualShadowMapId(&View);
+							}
+
 							FHairStrandsTransmittanceMaskData TransmittanceMaskData;
 							FRDGTextureRef HairShadowMask = nullptr;
-							if (bDrawHairShadow)
+							if (bDrawHairShadow && VirtualShadowMapId != INDEX_NONE)
+							{
+								TransmittanceMaskData.TransmittanceMask = ShadowSceneRenderer->HairTransmittanceMaskBits;
+								HairShadowMask = nullptr;
+								check(ScreenShadowMaskSubPixelTexture == nullptr);
+							}
+							else if (bDrawHairShadow)
 							{
 								TransmittanceMaskData = RenderHairStrandsTransmittanceMask(GraphBuilder, View, &LightSceneInfo, false, ScreenShadowMaskSubPixelTexture);
 								HairShadowMask = ScreenShadowMaskSubPixelTexture;
@@ -2110,13 +2125,12 @@ void FDeferredShadingSceneRenderer::RenderLights(
 								TransmittanceMaskData = DummyTransmittanceMaskData;
 							}
 
-							// TODO: One pass projection
-
 							// Note: ideally the light should still be evaluated for hair when not casting shadow, but for preserving the old behavior, and not adding 
 							// any perf. regression, we disable this light for hair rendering 
 							RenderLightForHair(
 								GraphBuilder, View, SceneTextures, &LightSceneInfo, 
-								HairShadowMask, LightingChannelsTexture, TransmittanceMaskData, false /*bForwardRendering*/);
+								VirtualShadowMapId != INDEX_NONE ? nullptr : HairShadowMask, LightingChannelsTexture, TransmittanceMaskData, false /*bForwardRendering*/,
+								VirtualShadowMapArray.GetUniformBuffer(), ShadowSceneRenderer->VirtualShadowMapMaskBitsHairStrands, VirtualShadowMapId);
 						}
 					}
 				}
@@ -2712,7 +2726,10 @@ void FDeferredShadingSceneRenderer::RenderLightForHair(
 	FRDGTextureRef HairShadowMaskTexture,
 	FRDGTextureRef LightingChannelsTexture,
 	const FHairStrandsTransmittanceMaskData& InTransmittanceMaskData,
-	const bool bForwardRendering)
+	const bool bForwardRendering,
+	TRDGUniformBufferRef<FVirtualShadowMapUniformParameters> VirtualShadowMapUniformBuffer,
+	FRDGTextureRef ShadowMaskBits,
+	int32 VirtualShadowMapId)
 {
 	// Ensure the light is valid for this view
 	const bool bHairRenderingEnabled = HairStrands::HasViewHairStrandsData(View);
@@ -2731,6 +2748,7 @@ void FDeferredShadingSceneRenderer::RenderLightForHair(
 
 	const bool bIsDirectional = LightSceneInfo->Proxy->GetLightType() == LightType_Directional;
 	const bool bCloudShadow   = bIsDirectional;
+	const bool bUseVirtualShadowMapMask = VirtualShadowMapId != INDEX_NONE && ShadowMaskBits;
 
 	FRenderLightForHairParameters* PassParameters = GraphBuilder.AllocParameters<FRenderLightForHairParameters>();
 	// VS - General parameters
@@ -2747,13 +2765,17 @@ void FDeferredShadingSceneRenderer::RenderLightForHair(
 		HairStrands::BindHairStrandsViewUniformParameters(View),
 		HairShadowMaskTexture,
 		LightingChannelsTexture,
-		bCloudShadow);
+		bCloudShadow,
+		VirtualShadowMapUniformBuffer, 
+		ShadowMaskBits, 
+		VirtualShadowMapId);
 
 	// PS - Hair parameters
 	const FIntPoint SampleLightingViewportResolution = View.HairStrandsViewData.VisibilityData.SampleLightingViewportResolution;
 	PassParameters->PS.HairTransmittanceBuffer = GraphBuilder.CreateSRV(InTransmittanceMaskData.TransmittanceMask, FHairStrandsTransmittanceMaskData::Format);
 	PassParameters->PS.HairTransmittanceBufferMaxCount = InTransmittanceMaskData.TransmittanceMask ? InTransmittanceMaskData.TransmittanceMask->Desc.NumElements : 0;
 	PassParameters->PS.ShadowChannelMask = FVector4f(1, 1, 1, 1);
+	PassParameters->PS.LightSceneId = LightSceneInfo->Id;
 	if (HairShadowMaskTexture)
 	{
 		PassParameters->PS.ScreenShadowMaskSubPixelTexture = HairShadowMaskTexture;
@@ -2775,6 +2797,7 @@ void FDeferredShadingSceneRenderer::RenderLightForHair(
 	PermutationVector.Set< FDeferredLightPS::FTransmissionDim >(false);
 	PermutationVector.Set< FDeferredLightPS::FHairLighting>(1);
 	PermutationVector.Set< FDeferredLightPS::FHairComplexTransmittance>(true);
+	PermutationVector.Set< FDeferredLightPS::FVirtualShadowMapMask >(bUseVirtualShadowMapMask);
 	PermutationVector.Set< FDeferredLightPS::FLightFunctionAtlasDim >(
 		LightFunctionAtlas::IsEnabled(View, ELightFunctionAtlasSystem::DeferredLighting) && LightSceneInfo->Proxy->HasValidLightFunctionAtlasSlot() &&
 		LightSceneInfo->Proxy->GetLightFunctionMaterial() != nullptr && !View.Family->EngineShowFlags.VisualizeLightCulling);
diff --git a/Engine/Source/Runtime/Renderer/Private/Shadows/ShadowSceneRenderer.cpp b/Engine/Source/Runtime/Renderer/Private/Shadows/ShadowSceneRenderer.cpp
--- a/Engine/Source/Runtime/Renderer/Private/Shadows/ShadowSceneRenderer.cpp
+++ b/Engine/Source/Runtime/Renderer/Private/Shadows/ShadowSceneRenderer.cpp
@@ -511,6 +511,7 @@ void FShadowSceneRenderer::RenderVirtualShadowMapProjectionMaskBits(
 
 			if (HairStrands::HasViewHairStrandsData(View))
 			{
+				// Shadow bits
 				RenderVirtualShadowMapProjectionOnePass(
 					GraphBuilder,
 					SceneTextures,
@@ -518,6 +519,9 @@ void FShadowSceneRenderer::RenderVirtualShadowMapProjectionMaskBits(
 					VirtualShadowMapArray,
 					EVirtualShadowMapProjectionInputType::HairStrands,
 					VirtualShadowMapMaskBitsHairStrands);
+
+				// Transmittance bits
+				HairTransmittanceMaskBits = RenderHairStrandsOnePassTransmittanceMask(GraphBuilder, View, VirtualShadowMapMaskBitsHairStrands, VirtualShadowMapArray).TransmittanceMask;
 			}
 		}
 	}
@@ -525,6 +529,7 @@ void FShadowSceneRenderer::RenderVirtualShadowMapProjectionMaskBits(
 	{
 		VirtualShadowMapMaskBits = nullptr;//Dummy;
 		VirtualShadowMapMaskBitsHairStrands = nullptr;//Dummy;
+		HairTransmittanceMaskBits = nullptr; //Dummy
 	}
 }
 
diff --git a/Engine/Source/Runtime/Renderer/Private/Shadows/ShadowSceneRenderer.h b/Engine/Source/Runtime/Renderer/Private/Shadows/ShadowSceneRenderer.h
--- a/Engine/Source/Runtime/Renderer/Private/Shadows/ShadowSceneRenderer.h
+++ b/Engine/Source/Runtime/Renderer/Private/Shadows/ShadowSceneRenderer.h
@@ -89,6 +89,7 @@ public:
 	// One pass projection stuff. Set up in RenderVitualShadowMapProjectionMaskBits
 	FRDGTextureRef VirtualShadowMapMaskBits = nullptr;
 	FRDGTextureRef VirtualShadowMapMaskBitsHairStrands = nullptr;
+	FRDGBufferRef HairTransmittanceMaskBits = nullptr;
 
 	bool UsePackedShadowMaskBits() const
 	{
