Fix high resolution screenshot with hair rendering.

diff --git a/Engine/Source/Runtime/Renderer/Private/HairStrands/HairStrandsUtils.cpp b/Engine/Source/Runtime/Renderer/Private/HairStrands/HairStrandsUtils.cpp
--- a/Engine/Source/Runtime/Renderer/Private/HairStrands/HairStrandsUtils.cpp
+++ b/Engine/Source/Runtime/Renderer/Private/HairStrands/HairStrandsUtils.cpp
@@ -121,12 +121,19 @@ float GetDeepShadowRasterizationScale()
 }
 
 FMinHairRadiusAtDepth1 ComputeMinStrandRadiusAtDepth1(
-	const FIntPoint& Resolution,
+	const FIntPoint& InResolution,
 	const float FOV,
 	const uint32 SampleCount,
 	const float OverrideStrandHairRasterizationScale,
 	const float OrthoWidth)
 {
+	FIntPoint Resolution = InResolution;
+	if (GIsHighResScreenshot)
+	{
+		Resolution.X = GScreenshotResolutionX;
+		Resolution.Y = GScreenshotResolutionY;
+	}
+
 	auto InternalMinRadiusAtDepth1 = [Resolution, FOV, SampleCount, OrthoWidth](float RasterizationScale)
 	{
 		const float DiameterToRadius = 0.5f;
@@ -141,8 +148,6 @@ FMinHairRadiusAtDepth1 ComputeMinStrandRadiusAtDepth1(
 			StrandDiameterAtDepth1 *= FMath::Tan(vFOV * 0.5f) / (0.5f * Resolution.Y);
 		}
 		return DiameterToRadius * RasterizationScale * StrandDiameterAtDepth1;
-
-		
 	};
 
 	FMinHairRadiusAtDepth1 Out;
