Followup change to CL 32039781 to add a nullptr check in case there is no dynamic resolution state.

diff --git a/Engine/Source/Runtime/Renderer/Private/SceneRendering.cpp b/Engine/Source/Runtime/Renderer/Private/SceneRendering.cpp
--- a/Engine/Source/Runtime/Renderer/Private/SceneRendering.cpp
+++ b/Engine/Source/Runtime/Renderer/Private/SceneRendering.cpp
@@ -2989,7 +2989,11 @@ void FSceneRenderer::PrepareViewRectsForRendering(FRHICommandListImmediate& RHIC
 	// If we support screen percentage, update the dynamic resolution state with our current temporal upscaler, which clamps the screen percentage to its supported range.
 	if (ViewFamily.SupportsScreenPercentage())
 	{
-		GEngine->GetDynamicResolutionState()->SetTemporalUpscaler(ViewFamily.GetTemporalUpscalerInterface());
+		IDynamicResolutionState* DynamicResolutionState = GEngine->GetDynamicResolutionState();
+		if (DynamicResolutionState)
+		{
+			DynamicResolutionState->SetTemporalUpscaler(ViewFamily.GetTemporalUpscalerInterface());
+		}
 	}
 
 	// Read the resolution data.
