Fully disable the "skip always visible" optimization under WITH_EDITOR for now (only cooked builds will run it)

diff --git a/Engine/Source/Runtime/Renderer/Private/RendererScene.cpp b/Engine/Source/Runtime/Renderer/Private/RendererScene.cpp
--- a/Engine/Source/Runtime/Renderer/Private/RendererScene.cpp
+++ b/Engine/Source/Runtime/Renderer/Private/RendererScene.cpp
@@ -6324,8 +6324,15 @@ void FScene::Update(FRDGBuilder& GraphBuilder, const FUpdateParameters& Paramete
 	// Determine the test visible vs. always visible primitive index ranges
 	PrimitivesAlwaysVisibleOffset = ~0u;
 
+#if WITH_EDITOR
+	// TODO: Support skip always visible in the editor (need to handle dynamic relevance)
+	const bool bSkipAlwaysVisible = false;
+#else
+	const bool bSkipAlwaysVisible = GVisibilitySkipAlwaysVisible != 0;
+#endif
+
 	// This optimization requires compute materials due to relevancy calculation
-	if (GVisibilitySkipAlwaysVisible != 0 && UseNaniteComputeMaterials())
+	if (bSkipAlwaysVisible && UseNaniteComputeMaterials())
 	{
 		uint32 NextTypeOffset = 0;
 		for (int32 TypeOffsetIndex = 0; TypeOffsetIndex < TypeOffsetTable.Num(); ++TypeOffsetIndex)
