Heterogeneous Volumes: Adding live-shading pipeline for generating adaptive volumetric shadow maps.

diff --git a/Engine/Shaders/Private/HeterogeneousVolumes/HeterogeneousVolumesLiveShadingShadows.usf b/Engine/Shaders/Private/HeterogeneousVolumes/HeterogeneousVolumesLiveShadingShadows.usf
new file mode 100644
--- /dev/null
+++ b/Engine/Shaders/Private/HeterogeneousVolumes/HeterogeneousVolumesLiveShadingShadows.usf
@@ -0,0 +1,442 @@
+// Copyright Epic Games, Inc. All Rights Reserved.
+
+#include "../Common.ush"
+
+#define SUPPORT_CONTACT_SHADOWS 0
+
+#define DYNAMICALLY_SHADOWED 1
+#define TREAT_MAXDEPTH_UNSHADOWED 1
+
+#define SHADOW_QUALITY 2
+#define NO_TRANSLUCENCY_AVAILABLE
+
+#include "/Engine/Generated/Material.ush"
+#include "HeterogeneousVolumesLiveShadingUtils.ush"
+#include "HeterogeneousVolumesTracingUtils.ush"
+#include "HeterogeneousVolumesAdaptiveVolumetricShadowMapUtils.ush"
+
+#ifndef THREADGROUP_SIZE_1D
+#define THREADGROUP_SIZE_1D 1
+#endif // THREADGROUP_SIZE_1D
+
+#ifndef THREADGROUP_SIZE_2D
+#define THREADGROUP_SIZE_2D 1
+#endif // THREADGROUP_SIZE_2D
+
+#ifndef THREADGROUP_SIZE_3D
+#define THREADGROUP_SIZE_3D 1
+#endif // THREADGROUP_SIZE_3D
+
+// Object data
+float4x4 LocalToWorld;
+float4x4 WorldToLocal;
+float3 LocalBoundsOrigin;
+float3 LocalBoundsExtent;
+int PrimitiveId;
+
+// Volume data
+int3 VoxelResolution;
+
+// Ray data
+float MaxTraceDistance;
+float MaxShadowTraceDistance;
+float StepSize;
+float StepFactor;
+float ShadowStepSize;
+float ShadowStepFactor;
+int MaxStepCount;
+int bJitter;
+
+// Shadow data
+int NumShadowMatrices;
+float4x4 TranslatedWorldToShadow[6];
+float4x4 ShadowToTranslatedWorld[6];
+float3 TranslatedWorldOrigin;
+
+int2 ShadowResolution;
+int MaxSampleCount;
+float AbsoluteErrorThreshold;
+float RelativeErrorThreshold;
+
+// Dispatch data
+int3 GroupCount;
+
+// Debug data
+int ShadowDebugTweak;
+
+// Output
+RWBuffer<int> RWVolumetricShadowLinkedListAllocatorBuffer;
+RWStructuredBuffer<uint2> RWVolumetricShadowLinkedListBuffer;
+RWTexture2D<float4> RWBeerShadowMapTexture;
+
+FPrimitiveSceneData GetPrimitiveData(FMaterialVertexParameters Parameters)
+{
+	return GetPrimitiveData(Parameters, LocalToWorld, WorldToLocal, LocalBoundsOrigin, LocalBoundsExtent);
+}
+
+FPrimitiveSceneData GetPrimitiveData(FMaterialPixelParameters Parameters)
+{
+	return GetPrimitiveData(Parameters, LocalToWorld, WorldToLocal, LocalBoundsOrigin, LocalBoundsExtent);
+}
+
+float4x4 GetLocalToWorld()
+{
+	return LocalToWorld;
+}
+
+float4x4 GetWorldToLocal()
+{
+	return WorldToLocal;
+}
+
+float3 GetLocalBoundsOrigin()
+{
+	return LocalBoundsOrigin;
+}
+
+float3 GetLocalBoundsExtent()
+{
+	return LocalBoundsExtent;
+}
+
+int3 GetVolumeResolution()
+{
+	return VoxelResolution;
+}
+
+float GetShadowStepSize()
+{
+	return ShadowStepSize;
+}
+
+float GetShadowStepFactor()
+{
+	return ShadowStepFactor;
+}
+
+float GetStepSize()
+{
+	return GetShadowStepSize();
+}
+
+float GetStepFactor()
+{
+	return GetShadowStepFactor();
+}
+
+float GetMaxTraceDistance()
+{
+	return MaxTraceDistance;
+}
+
+float GetMaxShadowTraceDistance()
+{
+	return MaxShadowTraceDistance;
+}
+
+struct FVolumeSampleContext
+{
+	FMaterialPixelParameters MaterialParameters;
+	FPixelMaterialInputs PixelMaterialInputs;
+};
+
+FVolumeSampleContext CreateVolumeSampleContext(float3 LocalPosition, float3 WorldPosition, float MipLevel)
+{
+	FVolumeSampleContext VolumeSampleContext;
+	VolumeSampleContext.MaterialParameters = MakeInitializedMaterialPixelParameters();
+	//VolumeSampleContext.MaterialParameters.PrimitiveId = PrimitiveId;
+	VolumeSampleContext.MaterialParameters.AbsoluteWorldPosition = DFPromote(WorldPosition);
+	VolumeSampleContext.MaterialParameters.LWCData.AbsoluteWorldPosition = WSPromote(WorldPosition);
+	// TODO: Add object centroid to LWC.ObjectWorldPosition
+	VolumeSampleContext.MaterialParameters.LWCData.LocalToWorld = WSPromote(GetLocalToWorld());
+	VolumeSampleContext.MaterialParameters.LWCData.WorldToLocal = WSPromoteInverse(GetWorldToLocal());
+
+	// Evaluate Material graph
+	CalcPixelMaterialInputs(VolumeSampleContext.MaterialParameters, VolumeSampleContext.PixelMaterialInputs);
+
+	return VolumeSampleContext;
+}
+
+float3 SampleExtinction(inout FVolumeSampleContext Context)
+{
+	float3 Extinction = SampleExtinctionCoefficients(Context.PixelMaterialInputs);
+	return Extinction;
+}
+
+uint GetLinearIndex(uint Face, uint2 PixelCoord)
+{
+	uint LinearIndex = Face * ShadowResolution.x * ShadowResolution.y + PixelCoord.y * ShadowResolution.x + PixelCoord.x;
+	return LinearIndex;
+}
+
+#include "HeterogeneousVolumesRayMarchingUtils.ush"
+
+int FindSampleIndexWithLowestError(
+	int PixelIndex
+)
+{
+	// Find sample, whose removal would mark the lowest overall geometric error
+	int LowestErrorIndex = 0;
+	float LowestError = 1.0e6;
+	for (int Index = 1; Index < MaxSampleCount - 1; ++Index)
+	{
+		int PrevIndex = AVSM_GetPrevPtr(RWVolumetricShadowLinkedListBuffer[PixelIndex + Index]);
+		int NextIndex = AVSM_GetNextPtr(RWVolumetricShadowLinkedListBuffer[PixelIndex + Index]);
+
+		FAVSMSampleData Data[3] = {
+			AVSM_UnpackSampleData(RWVolumetricShadowLinkedListBuffer[PixelIndex + PrevIndex]),
+			AVSM_UnpackSampleData(RWVolumetricShadowLinkedListBuffer[PixelIndex + Index]),
+			AVSM_UnpackSampleData(RWVolumetricShadowLinkedListBuffer[PixelIndex + NextIndex])
+		};
+
+		float Error = abs(
+			(Data[2].X - Data[0].X) * (Data[2].Tau - Data[1].Tau) -
+			(Data[2].X - Data[1].X) * (Data[2].Tau - Data[0].Tau)
+		);
+
+		if (Error < LowestError)
+		{
+			LowestError = Error;
+			LowestErrorIndex = Index;
+		}
+	}
+
+	return LowestErrorIndex;
+}
+
+void AddSample(FAVSMSampleData Data, int PixelIndex, inout int PrevSampleIndex, inout int SampleIndex)
+{
+	RWVolumetricShadowLinkedListBuffer[PixelIndex + SampleIndex].x = AVSM_PackSampleData(Data);
+	AVSM_SetPrevPtr(RWVolumetricShadowLinkedListBuffer[PixelIndex + SampleIndex], PrevSampleIndex);
+	AVSM_SetNextPtr(RWVolumetricShadowLinkedListBuffer[PixelIndex + SampleIndex], SampleIndex + 1);
+
+	PrevSampleIndex = SampleIndex;
+	SampleIndex = SampleIndex + 1;
+}
+
+void ReplaceSample(
+	int PixelIndex,
+	int SampleIndex,
+	FAVSMSampleData SampleData,
+	int PrevPtrIndex,
+	int NextPtrIndex
+)
+{
+	// Remove 
+	{
+		int PrevIndex = AVSM_GetPrevPtr(RWVolumetricShadowLinkedListBuffer[PixelIndex + SampleIndex]);
+		int NextIndex = AVSM_GetNextPtr(RWVolumetricShadowLinkedListBuffer[PixelIndex + SampleIndex]);
+
+		AVSM_SetNextPtr(RWVolumetricShadowLinkedListBuffer[PixelIndex + PrevIndex], NextIndex);
+		AVSM_SetPrevPtr(RWVolumetricShadowLinkedListBuffer[PixelIndex + NextIndex], PrevIndex);
+	}
+
+	// Replace
+	RWVolumetricShadowLinkedListBuffer[PixelIndex + SampleIndex].x = AVSM_PackSampleData(SampleData);
+	AVSM_SetPrevPtr(RWVolumetricShadowLinkedListBuffer[PixelIndex + SampleIndex], PrevPtrIndex);
+	AVSM_SetNextPtr(RWVolumetricShadowLinkedListBuffer[PixelIndex + SampleIndex], NextPtrIndex);
+
+	// Update neighborhood
+	AVSM_SetPrevPtr(RWVolumetricShadowLinkedListBuffer[PixelIndex + NextPtrIndex], SampleIndex);
+	AVSM_SetNextPtr(RWVolumetricShadowLinkedListBuffer[PixelIndex + PrevPtrIndex], SampleIndex);
+}
+
+void AddOrReplaceSample(FAVSMSampleData Data, int PixelIndex, inout int PrevSampleIndex, inout int SampleIndex)
+{
+	// Add
+	if (SampleIndex < MaxSampleCount - 1)
+	{
+		AddSample(Data, PixelIndex, PrevSampleIndex, SampleIndex);
+	}
+	// Or replace
+	else
+	{
+#if USE_AVSM_COMPRESSION
+		// Insert fictious sample at MaxSampleCount - 1 to compute error
+		RWVolumetricShadowLinkedListBuffer[PixelIndex + MaxSampleCount - 1].x = AVSM_PackSampleData(Data);
+		AVSM_SetNextPtr(RWVolumetricShadowLinkedListBuffer[PixelIndex + PrevSampleIndex], MaxSampleCount - 1);
+		AVSM_SetPrevPtr(RWVolumetricShadowLinkedListBuffer[PixelIndex + MaxSampleCount - 1], PrevSampleIndex);
+
+		int LowestErrorIndex = FindSampleIndexWithLowestError(PixelIndex);
+		if (LowestErrorIndex > 0)
+		{
+			if (LowestErrorIndex == PrevSampleIndex)
+			{
+				RWVolumetricShadowLinkedListBuffer[PixelIndex + PrevSampleIndex].x = AVSM_PackSampleData(Data);
+			}
+			else
+			{
+				ReplaceSample(PixelIndex, LowestErrorIndex, Data, PrevSampleIndex, MaxSampleCount - 1);
+			}
+
+			PrevSampleIndex = LowestErrorIndex;
+		}
+#else
+		RWVolumetricShadowLinkedListBuffer[PixelIndex + MaxSampleCount - 1].x = AVSM_PackSampleData(Data);
+		AVSM_SetPrevPtr(RWVolumetricShadowLinkedListBuffer[PixelIndex + MaxSampleCount - 1], PrevSampleIndex);
+
+		PrevSampleIndex = MaxSampleCount - 1;
+#endif
+	}
+}
+
+[numthreads(THREADGROUP_SIZE_2D, THREADGROUP_SIZE_2D, 1)]
+void RenderVolumetricShadowMapForLightWithLiveShadingCS(
+	uint2 GroupThreadId : SV_GroupThreadID,
+	uint3 DispatchThreadId : SV_DispatchThreadID
+)
+{
+	// Create screen ray
+	if (any(DispatchThreadId.xy >= ShadowResolution))
+	{
+		return;
+	}
+
+	uint Face = DispatchThreadId.z;
+	uint2 PixelCoord = DispatchThreadId.xy;
+	uint LinearIndex = GetLinearIndex(Face, PixelCoord);
+	//int PixelIndex = LinearIndex * MaxSampleCount;
+	int PixelIndex = LinearIndex * MaxSampleCount;
+
+	// Generate a light ray
+	float2 ShadowUV = float2(PixelCoord + 0.5f) / float2(ShadowResolution);
+	const float NearScreenZ = 1.0;
+	const float FarScreenZ = 0.0;
+
+	float4 LightRayStartAsFloat4 = mul(float4(ShadowUV, NearScreenZ, 1), ShadowToTranslatedWorld[Face]);
+	float4 LightRayEndAsFloat4 = mul(float4(ShadowUV, FarScreenZ, 1), ShadowToTranslatedWorld[Face]);
+	float3 LightRayStart = LightRayStartAsFloat4.xyz * rcp(LightRayStartAsFloat4.w);
+	float3 LightRayEnd = LightRayEndAsFloat4.xyz * rcp(LightRayEndAsFloat4.w);
+
+	// Intersect ray with bounding volume
+	// TODO: LWC integration..
+	float3 WorldRayOrigin = DFFastSubtractDemote(LightRayStart, PrimaryView.PreViewTranslation);
+	float3 WorldRayEnd = DFFastSubtractDemote(LightRayEnd, PrimaryView.PreViewTranslation);
+
+	float3 WorldRayDirection = WorldRayEnd - WorldRayOrigin;
+	float WorldRayLength = length(WorldRayDirection);
+	WorldRayDirection /= WorldRayLength;
+
+	float3 LocalRayOrigin = mul(float4(WorldRayOrigin, 1.0), WorldToLocal).xyz;
+	float3 LocalRayEnd = mul(float4(WorldRayEnd, 1.0), WorldToLocal).xyz;
+	float3 LocalRayDirection = LocalRayEnd - LocalRayOrigin;
+	float LocalRayLength = length(LocalRayDirection);
+	LocalRayDirection /= LocalRayLength;
+
+	float3 LocalBoundsMin = LocalBoundsOrigin - LocalBoundsExtent;
+	float3 LocalBoundsMax = LocalBoundsOrigin + LocalBoundsExtent;
+
+	float3 Transmittance = 1;
+
+	// Intersect bounding volume
+	float2 HitT = IntersectAABB(LocalRayOrigin, LocalRayDirection, 0.0, LocalRayLength, LocalBoundsMin, LocalBoundsMax);
+	float HitSpan = HitT.y - HitT.x;
+	if (HitSpan > 0.0)
+	{
+		float Jitter = bJitter ? InterleavedGradientNoise(PixelCoord, View.StateFrameIndexMod8) : 0.0;
+		int bApplyEmissionAndTransmittance = 0;
+		int bApplyDirectLighting = 0;
+		int bApplyShadowTransmittance = 0;
+
+		FRayMarchingContext RayMarchingContext = CreateRayMarchingContext(
+			// Local-space
+			LocalRayOrigin,
+			LocalRayDirection,
+			HitT.x,
+			HitT.y,
+			// World-space
+			WorldRayOrigin,
+			WorldRayDirection,
+			// Ray-step attributes
+			Jitter,
+			CalcShadowStepSize(WorldRayDirection),
+			MaxStepCount,
+			//MaxSampleCount, // Replace max-steps with max-samples..
+			bApplyEmissionAndTransmittance,
+			bApplyDirectLighting,
+			bApplyShadowTransmittance
+		);
+
+		RandomSequence RandSequence;
+		RandomSequence_Initialize(RandSequence, LinearIndex, View.StateFrameIndex);
+		RayMarchingContext.RandSequence = RandSequence;
+		uint StepCount = CalcStepCount(RayMarchingContext);
+
+		float3 ExpectedExtinction = 0.0;
+		float3 ExpectedTransmittance = 1.0;
+		int PrevElementIndex = AVSM_NULL_PTR;
+		int ElementIndex = 0;
+
+		float LocalHitT = HitT.x;
+		float WorldHitT = LocalHitT * RayMarchingContext.LocalToWorldScale;
+		for (uint StepIndex = 0; StepIndex < StepCount; ++StepIndex)
+		{
+			LocalHitT = min(RayMarchingContext.LocalRayTMin + RayMarchingContext.StepSize * (RayMarchingContext.Jitter + StepIndex), RayMarchingContext.LocalRayTMax);
+			WorldHitT = LocalHitT * RayMarchingContext.LocalToWorldScale;
+			float3 LocalPosition = RayMarchingContext.LocalRayOrigin + RayMarchingContext.LocalRayDirection * LocalHitT;
+			float3 WorldPosition = RayMarchingContext.WorldRayOrigin + RayMarchingContext.WorldRayDirection * WorldHitT;
+
+			FVolumeSampleContext SampleContext = CreateVolumeSampleContext(LocalPosition, WorldPosition, RayMarchingContext.MipLevel);
+			float3 Extinction = SampleExtinction(SampleContext);
+			Transmittance *= exp(-Extinction * RayMarchingContext.StepSize);
+
+			bool bFirstEntry = (ElementIndex == 0);
+			if (bFirstEntry && any(Extinction > 0.0))
+			{
+				FAVSMSampleData SampleData = AVSM_CreateSampleData(WorldHitT, Transmittance);
+				AddSample(SampleData, PixelIndex, PrevElementIndex, ElementIndex);
+
+				// Set expected values
+				ExpectedExtinction = Extinction;
+				ExpectedTransmittance = Transmittance;
+			}
+			else if (ElementIndex > 0 && any(Extinction > 0.0))
+			{
+				// If extrapolated transmittance is beyond relative error tolerance
+				ExpectedTransmittance *= exp(-ExpectedExtinction * RayMarchingContext.StepSize);
+				float AbsoluteError = length(Transmittance - ExpectedTransmittance);
+				float RelativeError = length((Transmittance - ExpectedTransmittance) / ExpectedTransmittance);
+				if ((AbsoluteError > AbsoluteErrorThreshold) && (RelativeError > RelativeErrorThreshold))
+				{
+					FAVSMSampleData SampleData = AVSM_CreateSampleData(WorldHitT, Transmittance);
+					AddOrReplaceSample(SampleData, PixelIndex, PrevElementIndex, ElementIndex);
+
+					// Reset expected values
+					ExpectedExtinction = Extinction;
+					ExpectedTransmittance = Transmittance;
+				}
+			}
+
+			// Early termination
+			float Epsilon = 1.0e-7;
+			if (all(Transmittance < Epsilon))
+			{
+				Transmittance = 0.0;
+				break;
+			}
+		}
+
+		// Stitch up the previous entry
+		if (PrevElementIndex != AVSM_NULL_PTR)
+		{
+			AVSM_SetNextPtr(RWVolumetricShadowLinkedListBuffer[PixelIndex + PrevElementIndex], ElementIndex);
+
+		}
+
+		// Last Entry
+		RWVolumetricShadowLinkedListBuffer[PixelIndex + ElementIndex].x = AVSM_PackSampleData(WorldHitT, Transmittance);
+		AVSM_SetPrevPtr(RWVolumetricShadowLinkedListBuffer[PixelIndex + ElementIndex], PrevElementIndex);
+		AVSM_SetNextPtr(RWVolumetricShadowLinkedListBuffer[PixelIndex + ElementIndex], AVSM_NULL_PTR);
+	}
+	else
+	{
+		RWVolumetricShadowLinkedListBuffer[PixelIndex].x = AVSM_PackSampleData(0.0, Transmittance);
+		AVSM_SetPrevPtr(RWVolumetricShadowLinkedListBuffer[PixelIndex], AVSM_NULL_PTR);
+		AVSM_SetNextPtr(RWVolumetricShadowLinkedListBuffer[PixelIndex], AVSM_NULL_PTR);
+	}
+
+	if (Face == ShadowDebugTweak)
+	{
+		RWBeerShadowMapTexture[PixelCoord] = float4(Transmittance, 1);
+	}
+}
diff --git a/Engine/Shaders/Private/HeterogeneousVolumes/HeterogeneousVolumesVoxelGridShadows.usf b/Engine/Shaders/Private/HeterogeneousVolumes/HeterogeneousVolumesVoxelGridShadows.usf
--- a/Engine/Shaders/Private/HeterogeneousVolumes/HeterogeneousVolumesVoxelGridShadows.usf
+++ b/Engine/Shaders/Private/HeterogeneousVolumes/HeterogeneousVolumesVoxelGridShadows.usf
@@ -210,6 +210,115 @@ uint GetLinearIndex(uint Face, uint2 PixelCoord)
 	return LinearIndex;
 }
 
+int FindSampleIndexWithLowestError(
+	int PixelIndex
+)
+{
+	// Find sample, whose removal would mark the lowest overall geometric error
+	int LowestErrorIndex = 0;
+	float LowestError = 1.0e6;
+	for (int Index = 1; Index < MaxSampleCount - 1; ++Index)
+	{
+		int PrevIndex = AVSM_GetPrevPtr(RWVolumetricShadowLinkedListBuffer[PixelIndex + Index]);
+		int NextIndex = AVSM_GetNextPtr(RWVolumetricShadowLinkedListBuffer[PixelIndex + Index]);
+
+		FAVSMSampleData Data[3] = {
+			AVSM_UnpackSampleData(RWVolumetricShadowLinkedListBuffer[PixelIndex + PrevIndex]),
+			AVSM_UnpackSampleData(RWVolumetricShadowLinkedListBuffer[PixelIndex + Index]),
+			AVSM_UnpackSampleData(RWVolumetricShadowLinkedListBuffer[PixelIndex + NextIndex])
+		};
+
+		float Error = abs(
+			(Data[2].X - Data[0].X) * (Data[2].Tau - Data[1].Tau) -
+			(Data[2].X - Data[1].X) * (Data[2].Tau - Data[0].Tau)
+		);
+
+		if (Error < LowestError)
+		{
+			LowestError = Error;
+			LowestErrorIndex = Index;
+		}
+	}
+
+	return LowestErrorIndex;
+}
+
+void AddSample(FAVSMSampleData Data, int PixelIndex, inout int PrevSampleIndex, inout int SampleIndex)
+{
+	RWVolumetricShadowLinkedListBuffer[PixelIndex + SampleIndex].x = AVSM_PackSampleData(Data);
+	AVSM_SetPrevPtr(RWVolumetricShadowLinkedListBuffer[PixelIndex + SampleIndex], PrevSampleIndex);
+	AVSM_SetNextPtr(RWVolumetricShadowLinkedListBuffer[PixelIndex + SampleIndex], SampleIndex + 1);
+
+	PrevSampleIndex = SampleIndex;
+	SampleIndex = SampleIndex + 1;
+}
+
+void ReplaceSample(
+	int PixelIndex,
+	int SampleIndex,
+	FAVSMSampleData SampleData,
+	int PrevPtrIndex,
+	int NextPtrIndex
+)
+{
+	// Remove 
+	{
+		int PrevIndex = AVSM_GetPrevPtr(RWVolumetricShadowLinkedListBuffer[PixelIndex + SampleIndex]);
+		int NextIndex = AVSM_GetNextPtr(RWVolumetricShadowLinkedListBuffer[PixelIndex + SampleIndex]);
+
+		AVSM_SetNextPtr(RWVolumetricShadowLinkedListBuffer[PixelIndex + PrevIndex], NextIndex);
+		AVSM_SetPrevPtr(RWVolumetricShadowLinkedListBuffer[PixelIndex + NextIndex], PrevIndex);
+	}
+
+	// Replace
+	RWVolumetricShadowLinkedListBuffer[PixelIndex + SampleIndex].x = AVSM_PackSampleData(SampleData);
+	AVSM_SetPrevPtr(RWVolumetricShadowLinkedListBuffer[PixelIndex + SampleIndex], PrevPtrIndex);
+	AVSM_SetNextPtr(RWVolumetricShadowLinkedListBuffer[PixelIndex + SampleIndex], NextPtrIndex);
+
+	// Update neighborhood
+	AVSM_SetNextPtr(RWVolumetricShadowLinkedListBuffer[PixelIndex + PrevPtrIndex], SampleIndex);
+	AVSM_SetPrevPtr(RWVolumetricShadowLinkedListBuffer[PixelIndex + NextPtrIndex], SampleIndex);
+}
+
+void AddOrReplaceSample(FAVSMSampleData Data, int PixelIndex, inout int PrevSampleIndex, inout int SampleIndex)
+{
+	// Add
+	if (SampleIndex < MaxSampleCount - 1)
+	{
+		AddSample(Data, PixelIndex, PrevSampleIndex, SampleIndex);
+	}
+	// Or replace
+	else
+	{
+#if USE_AVSM_COMPRESSION
+		// Insert fictious sample at MaxSampleCount - 1 to compute error
+		RWVolumetricShadowLinkedListBuffer[PixelIndex + MaxSampleCount - 1].x = AVSM_PackSampleData(Data);
+		AVSM_SetNextPtr(RWVolumetricShadowLinkedListBuffer[PixelIndex + PrevSampleIndex], MaxSampleCount - 1);
+		AVSM_SetPrevPtr(RWVolumetricShadowLinkedListBuffer[PixelIndex + MaxSampleCount - 1], PrevSampleIndex);
+
+		int LowestErrorIndex = FindSampleIndexWithLowestError(PixelIndex);
+		if (LowestErrorIndex > 0)
+		{
+			if (LowestErrorIndex == PrevSampleIndex)
+			{
+				RWVolumetricShadowLinkedListBuffer[PixelIndex + PrevSampleIndex].x = AVSM_PackSampleData(Data);
+			}
+			else
+			{
+				ReplaceSample(PixelIndex, LowestErrorIndex, Data, PrevSampleIndex, MaxSampleCount - 1);
+			}
+
+			PrevSampleIndex = LowestErrorIndex;
+		}
+#else
+		RWVolumetricShadowLinkedListBuffer[PixelIndex + MaxSampleCount - 1].x = AVSM_PackSampleData(Data);
+		AVSM_SetPrevPtr(RWVolumetricShadowLinkedListBuffer[PixelIndex + MaxSampleCount - 1], PrevSampleIndex);
+
+		PrevSampleIndex = MaxSampleCount - 1;
+#endif
+	}
+}
+
 [numthreads(THREADGROUP_SIZE_2D, THREADGROUP_SIZE_2D, 1)]
 void RenderVolumetricShadowMapForLightWithVoxelGridCS(
 	uint2 GroupThreadId : SV_GroupThreadID,
@@ -295,7 +404,6 @@ void RenderVolumetricShadowMapForLightWithVoxelGridCS(
 		RandomSequence_Initialize(RandSequence, LinearIndex, View.StateFrameIndex);
 		RayMarchingContext.RandSequence = RandSequence;
 		uint StepCount = CalcStepCount(RayMarchingContext);
-		int EntryCount = 1;
 
 		float3 ExpectedExtinction = 0.0;
 		float3 ExpectedTransmittance = 1.0;
@@ -313,25 +421,17 @@ void RenderVolumetricShadowMapForLightWithVoxelGridCS(
 			float3 Extinction = SampleExtinction(SampleContext);
 			Transmittance *= exp(-Extinction * RayMarchingContext.StepSize);
 
-			bool bFirstEntry = EntryCount == 1;
+			bool bFirstEntry = (ElementIndex == 0);
 			if (bFirstEntry && any(Extinction > 0.0))
 			{
-				RWVolumetricShadowLinkedListBuffer[PixelIndex + ElementIndex].x = AVSM_PackSampleData(WorldHitT, Transmittance);
-
-				// Allocate next entry
-				int NextElementIndex = ElementIndex + 1;
-				AVSM_SetPrevPtr(RWVolumetricShadowLinkedListBuffer[PixelIndex + ElementIndex], PrevElementIndex);
-				AVSM_SetNextPtr(RWVolumetricShadowLinkedListBuffer[PixelIndex + ElementIndex], NextElementIndex);
-
-				PrevElementIndex = ElementIndex;
-				ElementIndex = NextElementIndex;
-				EntryCount++;
+				FAVSMSampleData SampleData = AVSM_CreateSampleData(WorldHitT, Transmittance);
+				AddSample(SampleData, PixelIndex, PrevElementIndex, ElementIndex);
 
-				// Reset expected values
+				// Set expected values
 				ExpectedExtinction = Extinction;
 				ExpectedTransmittance = Transmittance;
 			}
-			else if (EntryCount > 1 && any(Extinction > 0.0))
+			else if (ElementIndex > 0 && any(Extinction > 0.0))
 			{
 				// If extrapolated transmittance is beyond relative error tolerance
 				ExpectedTransmittance *= exp(-ExpectedExtinction * RayMarchingContext.StepSize);
@@ -339,80 +439,12 @@ void RenderVolumetricShadowMapForLightWithVoxelGridCS(
 				float RelativeError = length((Transmittance - ExpectedTransmittance) / ExpectedTransmittance);
 				if ((AbsoluteError > AbsoluteErrorThreshold) && (RelativeError > RelativeErrorThreshold))
 				{
-					// Create vertex at current location
-					if (EntryCount < MaxSampleCount)
-					{
-						// Allocate next entry
-						int NextElementIndex = ElementIndex + 1;
-						RWVolumetricShadowLinkedListBuffer[PixelIndex + ElementIndex].x = AVSM_PackSampleData(WorldHitT, Transmittance);
-						AVSM_SetPrevPtr(RWVolumetricShadowLinkedListBuffer[PixelIndex + ElementIndex], PrevElementIndex);
-						AVSM_SetNextPtr(RWVolumetricShadowLinkedListBuffer[PixelIndex + ElementIndex], NextElementIndex);
-
-						PrevElementIndex = ElementIndex;
-						ElementIndex = NextElementIndex;
-						EntryCount++;
-
-						// Reset expected values
-						ExpectedExtinction = Extinction;
-						ExpectedTransmittance = Transmittance;
-					}
-#if USE_AVSM_COMPRESSION
-					// Or replace an existing vertex that contributes the lowest error to the approximation
-					else
-					{
-						int LowestErrorIndex = 0;
-						float LowestError = 1.0e6;
-						for (int Index = 1; Index < MaxSampleCount - 1; ++Index)
-						{
-							int PrevIndex = AVSM_GetPrevPtr(RWVolumetricShadowLinkedListBuffer[PixelIndex + Index]);
-							int NextIndex = AVSM_GetNextPtr(RWVolumetricShadowLinkedListBuffer[PixelIndex + Index]);
-
-							FAVSMSampleData Data[3] = {
-								AVSM_UnpackSampleData(RWVolumetricShadowLinkedListBuffer[PixelIndex + PrevIndex]),
-								AVSM_UnpackSampleData(RWVolumetricShadowLinkedListBuffer[PixelIndex + Index]),
-								AVSM_UnpackSampleData(RWVolumetricShadowLinkedListBuffer[PixelIndex + NextIndex])
-							};
-
-							float Error = abs(
-								(Data[2].X - Data[0].X) * (Data[2].Tau - Data[1].Tau) -
-								(Data[2].X - Data[1].X) * (Data[2].Tau - Data[0].Tau)
-							);
-
-							if (Error < LowestError)
-							{
-								LowestError = Error;
-								LowestErrorIndex = Index;
-							}
-						}
-
-						// Remove control point and replace with current control point
-						if (LowestErrorIndex > 0)
-						{
-							// Remove
-							int PrevIndex = AVSM_GetPrevPtr(RWVolumetricShadowLinkedListBuffer[PixelIndex + LowestErrorIndex]);
-							int NextIndex = AVSM_GetNextPtr(RWVolumetricShadowLinkedListBuffer[PixelIndex + LowestErrorIndex]);
-
-							AVSM_SetNextPtr(RWVolumetricShadowLinkedListBuffer[PixelIndex + PrevIndex], NextIndex);
-							AVSM_SetPrevPtr(RWVolumetricShadowLinkedListBuffer[PixelIndex + NextIndex], PrevIndex);
-
-							// Replace
-							RWVolumetricShadowLinkedListBuffer[PixelIndex + LowestErrorIndex].x = AVSM_PackSampleData(WorldHitT, Transmittance);
-							AVSM_SetPrevPtr(RWVolumetricShadowLinkedListBuffer[PixelIndex + LowestErrorIndex], PrevElementIndex);
-							AVSM_SetNextPtr(RWVolumetricShadowLinkedListBuffer[PixelIndex + LowestErrorIndex], MaxSampleCount - 1);
-
-							// Update neighborhood
-							AVSM_SetPrevPtr(RWVolumetricShadowLinkedListBuffer[PixelIndex + MaxSampleCount - 1], LowestErrorIndex);
-							AVSM_SetNextPtr(RWVolumetricShadowLinkedListBuffer[PixelIndex + PrevElementIndex], LowestErrorIndex);
-
-							PrevElementIndex = LowestErrorIndex;
-							ElementIndex = MaxSampleCount - 1;
-
-							// Reset expected values
-							ExpectedExtinction = Extinction;
-							ExpectedTransmittance = Transmittance;
-						}
-					}
-#endif // USE_AVSM_COMPRESSION
+					FAVSMSampleData SampleData = AVSM_CreateSampleData(WorldHitT, Transmittance);
+					AddOrReplaceSample(SampleData, PixelIndex, PrevElementIndex, ElementIndex);
+
+					// Reset expected values
+					ExpectedExtinction = Extinction;
+					ExpectedTransmittance = Transmittance;
 				}
 			}
 
@@ -425,6 +457,13 @@ void RenderVolumetricShadowMapForLightWithVoxelGridCS(
 			}
 		}
 
+		// Stitch up the previous entry
+		if (PrevElementIndex != AVSM_NULL_PTR)
+		{
+			AVSM_SetNextPtr(RWVolumetricShadowLinkedListBuffer[PixelIndex + PrevElementIndex], ElementIndex);
+
+		}
+
 		RWVolumetricShadowLinkedListBuffer[PixelIndex + ElementIndex].x = AVSM_PackSampleData(WorldHitT, Transmittance);
 		AVSM_SetPrevPtr(RWVolumetricShadowLinkedListBuffer[PixelIndex + ElementIndex], PrevElementIndex);
 		AVSM_SetNextPtr(RWVolumetricShadowLinkedListBuffer[PixelIndex + ElementIndex], AVSM_NULL_PTR);
@@ -436,13 +475,6 @@ void RenderVolumetricShadowMapForLightWithVoxelGridCS(
 		AVSM_SetNextPtr(RWVolumetricShadowLinkedListBuffer[PixelIndex], AVSM_NULL_PTR);
 	}
 
-#if 0
-	RWDebugBuffer[LinearIndex].LightRayStart = LightRayStart;
-	RWDebugBuffer[LinearIndex].LightRayEnd = LightRayEnd;
-	RWDebugBuffer[LinearIndex].RayOrigin = WorldRayOrigin;
-	RWDebugBuffer[LinearIndex].RayEnd = WorldRayEnd;
-	RWDebugBuffer[LinearIndex].HitSpan = float2(HitSpan, Depth);
-#endif
 	if (Face == ShadowDebugTweak)
 	{
 		RWBeerShadowMapTexture[PixelCoord] = float4(Transmittance, 1);
@@ -540,3 +572,85 @@ void CompressVolumetricShadowMapCS(
 		RWVolumetricShadowIndirectionBuffer[LinearIndex] = AVSM_PackIndirectionData(IndirectionData);
 	}
 }
+
+StructuredBuffer<uint2> VolumetricShadowLinkedListBuffer0;
+StructuredBuffer<uint2> VolumetricShadowLinkedListBuffer1;
+
+[numthreads(THREADGROUP_SIZE_2D, THREADGROUP_SIZE_2D, 1)]
+void CombineVolumetricShadowMapsCS(
+	uint2 GroupThreadId : SV_GroupThreadID,
+	uint3 DispatchThreadId : SV_DispatchThreadID,
+	uint GroupIndex : SV_GroupIndex
+)
+{
+	uint Face = DispatchThreadId.z;
+	uint2 PixelCoord = DispatchThreadId.xy;
+	if (any(PixelCoord >= ShadowResolution))
+	{
+		return;
+	}
+
+	uint LinearIndex = GetLinearIndex(Face, PixelCoord);
+	int PixelIndex = LinearIndex * MaxSampleCount;
+
+	int PrevSampleIndex = AVSM_NULL_PTR;
+	int SampleIndex = 0;
+
+	float Transmittance0 = 1.0;
+	float Transmittance1 = 1.0;
+	int IndexPtr0 = 0;
+	int IndexPtr1 = 0;
+
+	// Sorted sample insertion
+	while ((IndexPtr0 != AVSM_NULL_PTR) && (IndexPtr1 != AVSM_NULL_PTR))
+	{
+		FAVSMSampleData Data0 = AVSM_UnpackSampleData(VolumetricShadowLinkedListBuffer0[PixelIndex + IndexPtr0]);
+		FAVSMSampleData Data1 = AVSM_UnpackSampleData(VolumetricShadowLinkedListBuffer1[PixelIndex + IndexPtr1]);
+
+		FAVSMSampleData DataToInsert;
+		if (Data0.X < Data1.X)
+		{
+			DataToInsert = Data0;
+			DataToInsert.Tau *= Transmittance1;
+			Transmittance0 = Data0.Tau;
+
+			IndexPtr0 = AVSM_GetNextPtr(VolumetricShadowLinkedListBuffer0[PixelIndex + IndexPtr0]);
+		}
+		else
+		{
+			DataToInsert = Data1;
+			DataToInsert.Tau *= Transmittance0;
+			Transmittance1 = Data1.Tau;
+
+			IndexPtr1 = AVSM_GetNextPtr(VolumetricShadowLinkedListBuffer1[PixelIndex + IndexPtr1]);
+		}
+
+		AddOrReplaceSample(DataToInsert, PixelIndex, PrevSampleIndex, SampleIndex);
+	}
+
+	// Sample0 residual
+	while (IndexPtr0 != AVSM_NULL_PTR)
+	{
+		FAVSMSampleData DataToInsert = AVSM_UnpackSampleData(VolumetricShadowLinkedListBuffer0[PixelIndex + IndexPtr0]);
+		DataToInsert.Tau *= Transmittance1;
+
+		IndexPtr0 = AVSM_GetNextPtr(VolumetricShadowLinkedListBuffer0[PixelIndex + IndexPtr0]);
+		AddOrReplaceSample(DataToInsert, PixelIndex, PrevSampleIndex, SampleIndex);
+	}
+
+	// Sample1 residual
+	while ((IndexPtr1 != AVSM_NULL_PTR))
+	{
+		FAVSMSampleData DataToInsert = AVSM_UnpackSampleData(VolumetricShadowLinkedListBuffer1[PixelIndex + IndexPtr1]);
+		DataToInsert.Tau *= Transmittance0;
+
+		IndexPtr1 = AVSM_GetNextPtr(VolumetricShadowLinkedListBuffer1[PixelIndex + IndexPtr1]);
+		AddOrReplaceSample(DataToInsert, PixelIndex, PrevSampleIndex, SampleIndex);
+	}
+
+	// Mark the final sample as being terminal
+	if (PrevSampleIndex != AVSM_NULL_PTR)
+	{
+		AVSM_SetNextPtr(RWVolumetricShadowLinkedListBuffer[PixelIndex + PrevSampleIndex], AVSM_NULL_PTR);
+	}
+}
diff --git a/Engine/Source/Runtime/Renderer/Private/HeterogeneousVolumes/HeterogeneousVolumes.cpp b/Engine/Source/Runtime/Renderer/Private/HeterogeneousVolumes/HeterogeneousVolumes.cpp
--- a/Engine/Source/Runtime/Renderer/Private/HeterogeneousVolumes/HeterogeneousVolumes.cpp
+++ b/Engine/Source/Runtime/Renderer/Private/HeterogeneousVolumes/HeterogeneousVolumes.cpp
@@ -22,7 +22,8 @@ static TAutoConsoleVariable<int32> CVarHeterogeneousVolumesShadows(
 	TEXT("r.HeterogeneousVolumes.Shadows"),
 	0,
 	TEXT("Enables heterogeneous volume-casting shadows (default = 0)"),
-	ECVF_RenderThreadSafe | ECVF_ReadOnly);
+	ECVF_RenderThreadSafe | ECVF_ReadOnly
+);
 
 static TAutoConsoleVariable<int32> CVarTranslucencyHeterogeneousVolumes(
 	TEXT("r.Translucency.HeterogeneousVolumes"),
@@ -31,6 +32,14 @@ static TAutoConsoleVariable<int32> CVarTranslucencyHeterogeneousVolumes(
 	ECVF_RenderThreadSafe | ECVF_ReadOnly
 );
 
+static TAutoConsoleVariable<int32> CVarHeterogeneousVolumesShadowMode(
+	TEXT("r.HeterogeneousVolumes.Shadows.Mode"),
+	0,
+	TEXT("0: Live-Shading (Default)")
+	TEXT("1: Preshaded Voxel Grid"),
+	ECVF_RenderThreadSafe
+);
+
 static TAutoConsoleVariable<int32> CVarHeterogeneousVolumesDebug(
 	TEXT("r.HeterogeneousVolumes.Debug"),
 	0,
@@ -359,6 +368,11 @@ namespace HeterogeneousVolumes
 		return CVarHeterogeneousVolumesDebug.GetValueOnRenderThread();
 	}
 
+	EShadowMode GetShadowMode()
+	{
+		return static_cast<EShadowMode>(CVarHeterogeneousVolumesShadowMode.GetValueOnRenderThread());
+	}
+
 	bool UseSparseVoxelPipeline()
 	{
 		return CVarHeterogeneousVolumesSparseVoxel.GetValueOnRenderThread() != 0;
@@ -478,7 +492,7 @@ bool ShouldBuildVoxelGrids(const FScene* Scene)
 		return true;
 	}
 
-	if (HeterogeneousVolumes::GetDebugMode() != 0)
+	if (HeterogeneousVolumes::GetShadowMode() == HeterogeneousVolumes::EShadowMode::VoxelGrid)
 	{
 		return true;
 	}
@@ -494,7 +508,7 @@ bool ShouldCacheVoxelGrids(const FScene* Scene, FSceneViewState* ViewState)
 		return false;
 	}
 
-	if (HeterogeneousVolumes::GetDebugMode() != 0)
+	if (HeterogeneousVolumes::GetShadowMode() == HeterogeneousVolumes::EShadowMode::VoxelGrid)
 	{
 		return true;
 	}
@@ -518,18 +532,21 @@ void FDeferredShadingSceneRenderer::RenderHeterogeneousVolumeShadows(
 	RDG_GPU_STAT_SCOPE(GraphBuilder, HeterogeneousVolumeShadowsStat);
 	SCOPED_NAMED_EVENT(HeterogeneousVolumes, FColor::Emerald);
 
-	FVoxelGridBuildOptions BuildOptions;
-	BuildOptions.VoxelGridBuildMode = EVoxelGridBuildMode::Shadows;
-	BuildOptions.MinimumVoxelSizeOutsideFrustum = HeterogeneousVolumes::GetOutOfFrustumShadingRateForShadows();
-	BuildOptions.MinimumVoxelSizeInFrustum = HeterogeneousVolumes::GetShadingRateForShadows();
-	BuildOptions.bBuildOrthoGrid = true;
-	BuildOptions.bBuildFrustumGrid = false;
-	BuildOptions.bJitter = false;
-
 	TRDGUniformBufferRef<FOrthoVoxelGridUniformBufferParameters> OrthoGridUniformBuffer = nullptr;
 	TRDGUniformBufferRef<FFrustumVoxelGridUniformBufferParameters> FrustumGridUniformBuffer = nullptr;
-	BuildOrthoVoxelGrid(GraphBuilder, Scene, Views, VisibleLightInfos, BuildOptions, OrthoGridUniformBuffer);
-	BuildFrustumVoxelGrid(GraphBuilder, Scene, Views[0], BuildOptions, FrustumGridUniformBuffer);
+	if (HeterogeneousVolumes::GetShadowMode() == HeterogeneousVolumes::EShadowMode::VoxelGrid)
+	{
+		FVoxelGridBuildOptions BuildOptions;
+		BuildOptions.VoxelGridBuildMode = EVoxelGridBuildMode::Shadows;
+		BuildOptions.MinimumVoxelSizeOutsideFrustum = HeterogeneousVolumes::GetOutOfFrustumShadingRateForShadows();
+		BuildOptions.MinimumVoxelSizeInFrustum = HeterogeneousVolumes::GetShadingRateForShadows();
+		BuildOptions.bBuildOrthoGrid = true;
+		BuildOptions.bBuildFrustumGrid = false;
+		BuildOptions.bJitter = HeterogeneousVolumes::EnableJitterForShadows();
+
+		BuildOrthoVoxelGrid(GraphBuilder, Scene, Views, VisibleLightInfos, BuildOptions, OrthoGridUniformBuffer);
+		BuildFrustumVoxelGrid(GraphBuilder, Scene, Views[0], BuildOptions, FrustumGridUniformBuffer);
+	}
 
 	for (int32 ViewIndex = 0; ViewIndex < Views.Num(); ViewIndex++)
 	{
@@ -537,35 +554,64 @@ void FDeferredShadingSceneRenderer::RenderHeterogeneousVolumeShadows(
 
 		if (ShouldCompositeHeterogeneousVolumesWithTranslucency())
 		{
-			RenderAdaptiveVolumetricCameraMapWithVoxelGrid(
-				GraphBuilder,
-				// Scene data
-				SceneTextures,
-				Scene,
-				ViewFamily,
-				View,
-				// Volume data
-				OrthoGridUniformBuffer,
-				FrustumGridUniformBuffer
-			);
+			if (HeterogeneousVolumes::GetShadowMode() == HeterogeneousVolumes::EShadowMode::LiveShading)
+			{
+				RenderAdaptiveVolumetricCameraMapWithLiveShading(
+					GraphBuilder,
+					SceneTextures,
+					Scene,
+					ViewFamily,
+					View
+				);
+			}
+			else
+			{
+				RenderAdaptiveVolumetricCameraMapWithVoxelGrid(
+					GraphBuilder,
+					// Scene data
+					SceneTextures,
+					Scene,
+					ViewFamily,
+					View,
+					// Volume data
+					OrthoGridUniformBuffer,
+					FrustumGridUniformBuffer
+				);
+			}
 		}
 
 		if (ShouldHeterogeneousVolumesCastShadows())
 		{
-			RenderAdaptiveVolumetricShadowMapWithVoxelGrid(
-				GraphBuilder,
-				// Scene data
-				SceneTextures,
-				Scene,
-				ViewFamily,
-				View,
-				// Shadow Data
-				VisibleLightInfos,
-				VirtualShadowMapArray,
-				// Volume data
-				OrthoGridUniformBuffer,
-				FrustumGridUniformBuffer
-			);
+			if (HeterogeneousVolumes::GetShadowMode() == HeterogeneousVolumes::EShadowMode::LiveShading)
+			{
+				RenderAdaptiveVolumetricShadowMapWithLiveShading(
+					GraphBuilder,
+					// Scene data
+					SceneTextures,
+					Scene,
+					ViewFamily,
+					View,
+					// Light data
+					VisibleLightInfos
+				);
+			}
+			else
+			{
+				RenderAdaptiveVolumetricShadowMapWithVoxelGrid(
+					GraphBuilder,
+					// Scene data
+					SceneTextures,
+					Scene,
+					ViewFamily,
+					View,
+					// Shadow Data
+					VisibleLightInfos,
+					VirtualShadowMapArray,
+					// Volume data
+					OrthoGridUniformBuffer,
+					FrustumGridUniformBuffer
+				);
+			}
 		}
 	}
 
diff --git a/Engine/Source/Runtime/Renderer/Private/HeterogeneousVolumes/HeterogeneousVolumes.h b/Engine/Source/Runtime/Renderer/Private/HeterogeneousVolumes/HeterogeneousVolumes.h
--- a/Engine/Source/Runtime/Renderer/Private/HeterogeneousVolumes/HeterogeneousVolumes.h
+++ b/Engine/Source/Runtime/Renderer/Private/HeterogeneousVolumes/HeterogeneousVolumes.h
@@ -9,6 +9,7 @@
 class IHeterogeneousVolumeInterface;
 class FLightSceneInfo;
 class FPrimitiveSceneProxy;
+class FProjectedShadowInfo;
 class FRayTracingScene;
 class FRDGBuilder;
 class FScene;
@@ -23,6 +24,7 @@ struct FMaterialShaderParameters;
 struct FRDGTextureDesc;
 struct FSceneTextures;
 struct FPersistentPrimitiveIndex;
+struct FVolumetricMeshBatch;
 
 //
 // External API
@@ -59,8 +61,24 @@ namespace HeterogeneousVolumes
 	float GetMaxStepCount();
 	float GetMinimumVoxelSizeInFrustum();
 	float GetMinimumVoxelSizeOutsideFrustum();
+
+	// Shadow generation
+	enum class EShadowMode
+	{
+		LiveShading,
+		VoxelGrid
+	};
+	EShadowMode GetShadowMode();
+	FIntPoint GetShadowMapResolution();
 	float GetShadingRateForShadows();
 	float GetOutOfFrustumShadingRateForShadows();
+	bool EnableJitterForShadows();
+	float GetStepSizeForShadows();
+	uint32 GetShadowMaxSampleCount();
+	float GetShadowAbsoluteErrorThreshold();
+	float GetShadowRelativeErrorThreshold();
+	bool UseAVSMCompression();
+	float GetCameraDownsampleFactor();
 
 	int32 GetMipLevel();
 	int32 GetDebugMode();
@@ -89,8 +107,13 @@ namespace HeterogeneousVolumes
 	int GetVoxelCount(FIntVector VolumeResolution);
 	int GetVoxelCount(const FRDGTextureDesc& TextureDesc);
 	FIntVector GetMipVolumeResolution(FIntVector VolumeResolution, uint32 MipLevel);
+
+	const FProjectedShadowInfo* GetProjectedShadowInfo(const FVisibleLightInfo* VisibleLightInfo, int32 ShadowIndex);
+	bool IsDynamicShadow(const FVisibleLightInfo* VisibleLightInfo);
 }
 
+uint32 GetTypeHash(const FVolumetricMeshBatch& MeshBatch);
+
 struct FVoxelDataPacked
 {
 	uint32 LinearIndex;
@@ -208,6 +231,12 @@ END_UNIFORM_BUFFER_STRUCT()
 
 // Render specializations
 
+enum class EHeterogeneousVolumesShadowMode
+{
+	LiveShading,
+	VoxelGrid
+};
+
 enum class EVoxelGridBuildMode
 {
 	PathTracing,
@@ -398,6 +427,10 @@ namespace HeterogeneousVolumes {
 		FSceneViewState* ViewState
 	);
 
+	TRDGUniformBufferRef<FAdaptiveVolumetricShadowMapUniformBufferParameters> CreateEmptyAdaptiveVolumetricShadowMapUniformBuffer(
+		FRDGBuilder& GraphBuilder
+	);
+
 	TRDGUniformBufferRef<FFrustumVoxelGridUniformBufferParameters> GetFrustumVoxelGridUniformBuffer(
 		FRDGBuilder& GraphBuilder,
 		FSceneViewState* ViewState
@@ -480,6 +513,17 @@ void RenderAdaptiveVolumetricShadowMapWithVoxelGrid(
 	const TRDGUniformBufferRef<FFrustumVoxelGridUniformBufferParameters>& FrustumGridUniformBuffer
 );
 
+void RenderAdaptiveVolumetricShadowMapWithLiveShading(
+	FRDGBuilder& GraphBuilder,
+	// Scene data
+	const FSceneTextures& SceneTextures,
+	FScene* Scene,
+	const FSceneViewFamily& ViewFamily,
+	FViewInfo& View,
+	// Light data
+	TArray<FVisibleLightInfo, SceneRenderingAllocator>& VisibleLightInfos
+);
+
 void RenderAdaptiveVolumetricCameraMapWithVoxelGrid(
 	FRDGBuilder& GraphBuilder,
 	// Scene data
@@ -492,6 +536,57 @@ void RenderAdaptiveVolumetricCameraMapWithVoxelGrid(
 	const TRDGUniformBufferRef<FFrustumVoxelGridUniformBufferParameters>& FrustumGridUniformBuffer
 );
 
+void RenderAdaptiveVolumetricCameraMapWithLiveShading(
+	FRDGBuilder& GraphBuilder,
+	// Scene data
+	const FSceneTextures& SceneTextures,
+	FScene* Scene,
+	const FSceneViewFamily& ViewFamily,
+	FViewInfo& View
+);
+
+void CompressVolumetricShadowMap(
+	FRDGBuilder& GraphBuilder,
+	FViewInfo& View,
+	FIntVector GroupCount,
+	// Input
+	FIntPoint ShadowMapResolution,
+	uint32 MaxSampleCount,
+	FRDGBufferRef VolumetricShadowLinkedListBuffer,
+	// Output
+	FRDGBufferRef& VolumetricShadowIndirectionBuffer,
+	FRDGBufferRef& VolumetricShadowTransmittanceBuffer
+);
+
+void CombineVolumetricShadowMap(
+	FRDGBuilder& GraphBuilder,
+	FViewInfo& View,
+	FIntVector GroupCount,
+	// Input
+	uint32 LightType,
+	FIntPoint ShadowMapResolution,
+	uint32 MaxSampleCount,
+	FRDGBufferRef VolumetricShadowLinkedListBuffer0,
+	FRDGBufferRef VolumetricShadowLinkedListBuffer1,
+	// Output
+	FRDGBufferRef& VolumetricShadowLinkedListBuffer
+);
+
+void CreateAdaptiveVolumetricShadowMapUniformBuffer(
+	FRDGBuilder& GraphBuilder,
+	const FVector3f& TranslatedWorldOrigin,
+	const FVector4f& TranslatedWorldPlane,
+	const FMatrix44f* TranslatedWorldToShadow,
+	FIntPoint VolumetricShadowMapResolution,
+	int32 NumShadowMatrices,
+	uint32 VolumetricShadowMapMaxSampleCount,
+	bool bIsDirectionalLight,
+	FRDGBufferRef VolumetricShadowMapLinkedListBuffer,
+	FRDGBufferRef VolumetricShadowMapIndirectionBuffer,
+	FRDGBufferRef VolumetricShadowMapSampleBuffer,
+	TRDGUniformBufferRef<FAdaptiveVolumetricShadowMapUniformBufferParameters>& AdaptiveVolumetricShadowMapUniformBuffer
+);
+
 void RenderSingleScatteringWithVoxelGrid(
 	FRDGBuilder& GraphBuilder,
 	// Scene data
diff --git a/Engine/Source/Runtime/Renderer/Private/HeterogeneousVolumes/HeterogeneousVolumesLiveShadingPipeline.cpp b/Engine/Source/Runtime/Renderer/Private/HeterogeneousVolumes/HeterogeneousVolumesLiveShadingPipeline.cpp
--- a/Engine/Source/Runtime/Renderer/Private/HeterogeneousVolumes/HeterogeneousVolumesLiveShadingPipeline.cpp
+++ b/Engine/Source/Runtime/Renderer/Private/HeterogeneousVolumes/HeterogeneousVolumesLiveShadingPipeline.cpp
@@ -631,13 +631,6 @@ static void RenderSingleScatteringWithLiveShading(
 	uint32 GroupCountY = FMath::DivideAndRoundUp(View.ViewRect.Size().Y, FRenderSingleScatteringWithLiveShadingCS::GetThreadGroupSize2D());
 	FIntVector GroupCount = FIntVector(GroupCountX, GroupCountY, 1);
 
-#if 0
-	FRDGBufferRef DebugOutputBuffer = GraphBuilder.CreateBuffer(
-		FRDGBufferDesc::CreateStructuredDesc(sizeof(Volumes::FDebugOutput), GroupCountX * GroupCountY * FRenderSingleScatteringWithLiveShadingCS::GetThreadGroupSize2D()),
-		TEXT("Lumen.Reflections.TraceDataPacked"));
-	AddClearUAVPass(GraphBuilder, GraphBuilder.CreateUAV(DebugOutputBuffer), 0);
-#endif
-
 	// Note must be done in the same scope as we add the pass otherwise the UB lifetime will not be guaranteed
 	FDeferredLightUniformStruct DeferredLightUniform;
 	if (bApplyDirectLighting && (LightSceneInfo != nullptr))
@@ -1068,4 +1061,939 @@ void RenderWithLiveShading(
 			HeterogeneousVolumeRadiance
 		);
 	}
+}
+
+class FRenderVolumetricShadowMapForLightWithLiveShadingCS : public FMeshMaterialShader
+{
+	DECLARE_SHADER_TYPE(FRenderVolumetricShadowMapForLightWithLiveShadingCS, MeshMaterial);
+
+	class FUseAVSMCompression : SHADER_PERMUTATION_BOOL("USE_AVSM_COMPRESSION");
+	using FPermutationDomain = TShaderPermutationDomain<FUseAVSMCompression>;
+
+
+	BEGIN_SHADER_PARAMETER_STRUCT(FParameters, )
+		// Scene data
+		SHADER_PARAMETER_STRUCT_REF(FViewUniformShaderParameters, View)
+		SHADER_PARAMETER_STRUCT_INCLUDE(FSceneTextureParameters, SceneTextures)
+
+		// Volumetric Shadow Map data
+		SHADER_PARAMETER(FVector3f, TranslatedWorldOrigin)
+		SHADER_PARAMETER(FIntPoint, ShadowResolution)
+		SHADER_PARAMETER(int, MaxSampleCount)
+		SHADER_PARAMETER(float, AbsoluteErrorThreshold)
+		SHADER_PARAMETER(float, RelativeErrorThreshold)
+
+		SHADER_PARAMETER(int, NumShadowMatrices)
+		SHADER_PARAMETER_ARRAY(FMatrix44f, TranslatedWorldToShadow, [6])
+		SHADER_PARAMETER_ARRAY(FMatrix44f, ShadowToTranslatedWorld, [6])
+
+		// Object data
+		SHADER_PARAMETER(FMatrix44f, LocalToWorld)
+		SHADER_PARAMETER(FMatrix44f, WorldToLocal)
+		SHADER_PARAMETER(FVector3f, LocalBoundsOrigin)
+		SHADER_PARAMETER(FVector3f, LocalBoundsExtent)
+		SHADER_PARAMETER(int32, PrimitiveId)
+
+		SHADER_PARAMETER(FIntVector, VoxelResolution)
+
+		// Ray data
+		SHADER_PARAMETER(float, ShadowStepSize)
+		SHADER_PARAMETER(float, ShadowStepFactor)
+		SHADER_PARAMETER(float, MaxTraceDistance)
+		SHADER_PARAMETER(int, MaxStepCount)
+		SHADER_PARAMETER(int, bJitter)
+
+		// Dispatch data
+		SHADER_PARAMETER(FIntVector, GroupCount)
+		SHADER_PARAMETER(int, ShadowDebugTweak)
+
+		// Output
+		SHADER_PARAMETER_RDG_BUFFER_UAV(RWBuffer<int>, RWVolumetricShadowLinkedListAllocatorBuffer)
+		SHADER_PARAMETER_RDG_BUFFER_UAV(RWStructuredBuffer<int2>, RWVolumetricShadowLinkedListBuffer)
+		SHADER_PARAMETER_RDG_TEXTURE_UAV(RWTexture2D<float4>, RWBeerShadowMapTexture)
+
+		// Debug
+		//SHADER_PARAMETER_RDG_BUFFER_UAV(RWStructuredBuffer<FVolumetricShadowMapDebugData>, RWDebugBuffer)
+	END_SHADER_PARAMETER_STRUCT()
+
+	FRenderVolumetricShadowMapForLightWithLiveShadingCS() = default;
+
+	FRenderVolumetricShadowMapForLightWithLiveShadingCS(
+		const FMeshMaterialShaderType::CompiledShaderInitializerType& Initializer
+	)
+		: FMeshMaterialShader(Initializer)
+	{
+		Bindings.BindForLegacyShaderParameters(
+			this,
+			Initializer.PermutationId,
+			Initializer.ParameterMap,
+			*FParameters::FTypeInfo::GetStructMetadata(),
+			// Don't require full bindings, we use FMaterialShader::SetParameters
+			false
+		);
+		ShaderLooseBindings.Bind(Initializer.ParameterMap);
+	}
+
+	static bool ShouldCompilePermutation(
+		const FMaterialShaderPermutationParameters& Parameters
+	)
+	{
+		return DoesPlatformSupportHeterogeneousVolumes(Parameters.Platform)
+			&& DoesMaterialShaderSupportHeterogeneousVolumes(Parameters.MaterialParameters);
+	}
+
+	static void ModifyCompilationEnvironment(
+		const FMaterialShaderPermutationParameters& Parameters,
+		FShaderCompilerEnvironment& OutEnvironment
+	)
+	{
+		FMaterialShader::ModifyCompilationEnvironment(Parameters, OutEnvironment);
+		OutEnvironment.SetDefine(TEXT("THREADGROUP_SIZE_1D"), GetThreadGroupSize1D());
+		OutEnvironment.SetDefine(TEXT("THREADGROUP_SIZE_2D"), GetThreadGroupSize2D());
+		OutEnvironment.SetDefine(TEXT("THREADGROUP_SIZE_3D"), GetThreadGroupSize3D());
+
+		// Disable in-scattering features
+		OutEnvironment.SetDefine(TEXT("DIM_USE_TRANSMITTANCE_VOLUME"), 0);
+		OutEnvironment.SetDefine(TEXT("DIM_USE_INSCATTERING_VOLUME"), 0);
+		OutEnvironment.SetDefine(TEXT("DIM_USE_LUMEN_GI"), 0);
+
+		// This shader takes a very long time to compile with FXC, so we pre-compile it with DXC first and then forward the optimized HLSL to FXC.
+		OutEnvironment.CompilerFlags.Add(CFLAG_PrecompileWithDXC); // @lh-todo - Disabled to workaround SPIRV-Cross bug: StructuredBuffer<uint> is translated to ByteAddressBuffer in HLSL backend
+		OutEnvironment.CompilerFlags.Add(CFLAG_AllowTypedUAVLoads);
+
+		OutEnvironment.SetDefine(TEXT("GET_PRIMITIVE_DATA_OVERRIDE"), 1);
+	}
+
+	static int32 GetThreadGroupSize1D() { return GetThreadGroupSize3D() * GetThreadGroupSize3D() * GetThreadGroupSize3D(); }
+	static int32 GetThreadGroupSize2D() { return 8; }
+	static int32 GetThreadGroupSize3D() { return 4; }
+
+	LAYOUT_FIELD(FRenderLightingCacheLooseBindings, ShaderLooseBindings);
+};
+
+IMPLEMENT_MATERIAL_SHADER_TYPE(, FRenderVolumetricShadowMapForLightWithLiveShadingCS, TEXT("/Engine/Private/HeterogeneousVolumes/HeterogeneousVolumesLiveShadingShadows.usf"), TEXT("RenderVolumetricShadowMapForLightWithLiveShadingCS"), SF_Compute);
+
+void CollectHeterogeneousVolumeMeshBatchesForView(
+	const FViewInfo& View,
+	TSet<FVolumetricMeshBatch>& HeterogeneousVolumesMeshBatches,
+	FBoxSphereBounds& WorldBounds
+)
+{
+	for (int32 MeshBatchIndex = 0; MeshBatchIndex < View.HeterogeneousVolumesMeshBatches.Num(); ++MeshBatchIndex)
+	{
+		const FVolumetricMeshBatch& MeshBatch = View.HeterogeneousVolumesMeshBatches[MeshBatchIndex];
+		bool bIsShadowCast = MeshBatch.Proxy->IsShadowCast(&View);
+
+		// TODO: Is material determiniation too expensive?
+		const FMaterialRenderProxy* MaterialRenderProxy = nullptr;
+		const FMaterialRenderProxy* DefaultMaterialRenderProxy = MeshBatch.Mesh->MaterialRenderProxy;
+		const FMaterial& Material = DefaultMaterialRenderProxy->GetMaterialWithFallback(View.GetFeatureLevel(), MaterialRenderProxy);
+		MaterialRenderProxy = MaterialRenderProxy ? MaterialRenderProxy : DefaultMaterialRenderProxy;
+		bool bIsVolumeMaterial = Material.GetMaterialDomain() == MD_Volume;
+
+		if (bIsShadowCast && bIsVolumeMaterial)
+		{
+			HeterogeneousVolumesMeshBatches.FindOrAdd(FVolumetricMeshBatch(MeshBatch.Mesh, MeshBatch.Proxy));
+			WorldBounds = WorldBounds + MeshBatch.Proxy->GetBounds();
+		}
+	}
+}
+
+void CollectHeterogeneousVolumeMeshBatchesForLight(
+	const FLightSceneInfo* LightSceneInfo,
+	const FVisibleLightInfo* VisibleLightInfo,
+	const FViewInfo& View,
+	TSet<FVolumetricMeshBatch>& HeterogeneousVolumesMeshBatches,
+	FBoxSphereBounds& WorldBounds
+)
+{
+	check(LightSceneInfo);
+	check(VisibleLightInfo);
+
+	if (LightSceneInfo->Proxy->CastsVolumetricShadow())
+	{
+		//for (int32 ViewIndex = 0; ViewIndex < Views.Num(); ViewIndex++)
+		{
+			//const FViewInfo& View = Views[ViewIndex];
+			CollectHeterogeneousVolumeMeshBatchesForView(View, HeterogeneousVolumesMeshBatches, WorldBounds);
+		}
+
+		for (int32 ShadowIndex = 0; ShadowIndex < VisibleLightInfo->ShadowsToProject.Num(); ++ShadowIndex)
+		{
+			const FProjectedShadowInfo* ProjectedShadowInfo = HeterogeneousVolumes::GetProjectedShadowInfo(VisibleLightInfo, ShadowIndex);
+			if (ProjectedShadowInfo != nullptr)
+			{
+				const TArray<FMeshBatchAndRelevance, SceneRenderingAllocator>& MeshBatches = ProjectedShadowInfo->GetDynamicSubjectHeterogeneousVolumeMeshElements();
+				for (int32 MeshBatchIndex = 0; MeshBatchIndex < MeshBatches.Num(); ++MeshBatchIndex)
+				{
+					const FMeshBatchAndRelevance& MeshBatch = MeshBatches[MeshBatchIndex];
+					bool bIsShadowCast = MeshBatch.PrimitiveSceneProxy->IsShadowCast(ProjectedShadowInfo->ShadowDepthView);
+
+					// TODO: Is material determiniation too expensive?
+					const FMaterialRenderProxy* MaterialRenderProxy = nullptr;
+					const FMaterialRenderProxy* DefaultMaterialRenderProxy = MeshBatch.Mesh->MaterialRenderProxy;
+					const FMaterial& Material = DefaultMaterialRenderProxy->GetMaterialWithFallback(View.GetFeatureLevel(), MaterialRenderProxy);
+					MaterialRenderProxy = MaterialRenderProxy ? MaterialRenderProxy : DefaultMaterialRenderProxy;
+					bool bIsVolumeMaterial = Material.GetMaterialDomain() == MD_Volume;
+
+					if (bIsShadowCast && bIsVolumeMaterial)
+					{
+						HeterogeneousVolumesMeshBatches.FindOrAdd(FVolumetricMeshBatch(MeshBatch.Mesh, MeshBatch.PrimitiveSceneProxy));
+						WorldBounds = WorldBounds + MeshBatch.PrimitiveSceneProxy->GetBounds();
+					}
+				}
+			}
+		}
+	}
+}
+
+bool RenderVolumetricShadowMapForLightForHeterogeneousVolumeWithLiveShading(
+	FRDGBuilder& GraphBuilder,
+	// Scene data
+	const FSceneTextures& SceneTextures,
+	FScene* Scene,
+	const FSceneViewFamily& ViewFamily,
+	FViewInfo& View,
+	// Light data
+	uint32 LightType,
+	const FLightSceneInfo* LightSceneInfo,
+	const FVisibleLightInfo* VisibleLightInfo,
+	// Shadow data
+	const FVector3f& TranslatedWorldOrigin,
+	int32 NumShadowMatrices,
+	FMatrix44f* TranslatedWorldToShadow,
+	FMatrix44f* ShadowToTranslatedWorld,
+	FIntPoint ShadowMapResolution,
+	uint32 MaxSampleCount,
+	// Volume
+	const FVolumetricMeshBatch& VolumetricMeshBatch,
+	// Dispatch
+	FIntVector& GroupCount,
+	// Output
+	FRDGTextureRef& BeerShadowMapTexture,
+	FRDGBufferRef& VolumetricShadowLinkedListBuffer
+)
+{
+	// TODO: Understand how the default world material can be triggered here during a recompilation, but not elsewhere..
+	const FMaterialRenderProxy* MaterialRenderProxy = nullptr;
+	const FMaterialRenderProxy* DefaultMaterialRenderProxy = VolumetricMeshBatch.Mesh->MaterialRenderProxy;
+	const FMaterial& Material = DefaultMaterialRenderProxy->GetMaterialWithFallback(View.GetFeatureLevel(), MaterialRenderProxy);
+	MaterialRenderProxy = MaterialRenderProxy ? MaterialRenderProxy : DefaultMaterialRenderProxy;
+	if (Material.GetMaterialDomain() != MD_Volume)
+	{
+		return false;
+	}
+
+	FRDGBufferRef VolumetricShadowLinkedListAllocatorBuffer = GraphBuilder.CreateBuffer(
+		FRDGBufferDesc::CreateBufferDesc(sizeof(uint32), 1),
+		TEXT("HeterogeneousVolume.VolumetricShadowLinkedListAllocatorBuffer")
+	);
+	// Initialize allocator to contain 1-spp
+	AddClearUAVPass(GraphBuilder, GraphBuilder.CreateUAV(VolumetricShadowLinkedListAllocatorBuffer, PF_R32_UINT), ShadowMapResolution.X * ShadowMapResolution.Y);
+
+	FRenderVolumetricShadowMapForLightWithLiveShadingCS::FParameters* PassParameters = GraphBuilder.AllocParameters<FRenderVolumetricShadowMapForLightWithLiveShadingCS::FParameters>();
+	{
+		// Scene data
+		PassParameters->View = View.ViewUniformBuffer;
+		PassParameters->SceneTextures = GetSceneTextureParameters(GraphBuilder, SceneTextures);
+
+		// Shadow map data
+		PassParameters->TranslatedWorldOrigin = TranslatedWorldOrigin;
+		PassParameters->ShadowResolution = ShadowMapResolution;
+		PassParameters->MaxSampleCount = MaxSampleCount;
+		PassParameters->AbsoluteErrorThreshold = HeterogeneousVolumes::GetShadowAbsoluteErrorThreshold();
+		PassParameters->RelativeErrorThreshold = HeterogeneousVolumes::GetShadowRelativeErrorThreshold();
+
+		PassParameters->NumShadowMatrices = NumShadowMatrices;
+		for (int32 i = 0; i < PassParameters->NumShadowMatrices; ++i)
+		{
+			PassParameters->TranslatedWorldToShadow[i] = TranslatedWorldToShadow[i];
+			PassParameters->ShadowToTranslatedWorld[i] = ShadowToTranslatedWorld[i];
+		}
+
+		// TODO: Instancing support
+		int32 VolumeIndex = 0;
+
+		// Object data
+		const IHeterogeneousVolumeInterface* HeterogeneousVolumeInterface = (IHeterogeneousVolumeInterface*)VolumetricMeshBatch.Mesh->Elements[VolumeIndex].UserData;
+		FMatrix InstanceToLocal = HeterogeneousVolumeInterface->GetInstanceToLocal();
+		FMatrix LocalToWorld = HeterogeneousVolumeInterface->GetLocalToWorld();
+		PassParameters->LocalToWorld = FMatrix44f(InstanceToLocal * LocalToWorld);
+		PassParameters->WorldToLocal = PassParameters->LocalToWorld.Inverse();
+
+		FMatrix LocalToInstance = InstanceToLocal.Inverse();
+		FBoxSphereBounds InstanceBoxSphereBounds = HeterogeneousVolumeInterface->GetLocalBounds().TransformBy(LocalToInstance);
+		PassParameters->LocalBoundsOrigin = FVector3f(InstanceBoxSphereBounds.Origin);
+		PassParameters->LocalBoundsExtent = FVector3f(InstanceBoxSphereBounds.BoxExtent);
+		const FPrimitiveSceneInfo* PrimitiveSceneInfo = VolumetricMeshBatch.Proxy->GetPrimitiveSceneInfo();
+		check(PrimitiveSceneInfo);
+		PassParameters->PrimitiveId = PrimitiveSceneInfo->GetPersistentIndex().Index;
+
+		PassParameters->VoxelResolution = HeterogeneousVolumeInterface->GetVoxelResolution();
+
+		// Ray Data
+		PassParameters->ShadowStepSize = HeterogeneousVolumes::GetShadowStepSize();
+		PassParameters->ShadowStepFactor = HeterogeneousVolumeInterface->GetShadowStepFactor();
+		PassParameters->MaxTraceDistance = HeterogeneousVolumes::GetMaxTraceDistance();
+		PassParameters->MaxStepCount = HeterogeneousVolumes::GetMaxStepCount();
+		PassParameters->bJitter = HeterogeneousVolumes::ShouldJitter();
+
+		// Dispatch data
+		PassParameters->GroupCount = GroupCount;
+		//PassParameters->ShadowDebugTweak = CVarHeterogeneousVolumesShadowDebugTweak.GetValueOnRenderThread();
+		PassParameters->ShadowDebugTweak = 0;
+
+		// Output
+		PassParameters->RWVolumetricShadowLinkedListAllocatorBuffer = GraphBuilder.CreateUAV(VolumetricShadowLinkedListAllocatorBuffer, PF_R32_UINT);
+		PassParameters->RWVolumetricShadowLinkedListBuffer = GraphBuilder.CreateUAV(VolumetricShadowLinkedListBuffer);
+		PassParameters->RWBeerShadowMapTexture = GraphBuilder.CreateUAV(BeerShadowMapTexture);
+		//PassParameters->RWDebugBuffer = GraphBuilder.CreateUAV(DebugBuffer);
+	}
+
+	FString PassName;
+#if WANTS_DRAW_MESH_EVENTS
+	if (GetEmitDrawEvents())
+	{
+		FString LightName = TEXT("none");
+		if (LightSceneInfo != nullptr)
+		{
+			FSceneRenderer::GetLightNameForDrawEvent(LightSceneInfo->Proxy, LightName);
+		}
+		PassName = FString::Printf(TEXT("RenderVolumetricShadowMapForLightWithLiveShadingCS (Light = %s)"), *LightName);
+	}
+#endif // WANTS_DRAW_MESH_EVENTS
+
+	FRenderVolumetricShadowMapForLightWithLiveShadingCS::FPermutationDomain PermutationVector;
+	PermutationVector.Set<FRenderVolumetricShadowMapForLightWithLiveShadingCS::FUseAVSMCompression>(HeterogeneousVolumes::UseAVSMCompression());
+	TShaderRef<FRenderVolumetricShadowMapForLightWithLiveShadingCS> ComputeShader = Material.GetShader<FRenderVolumetricShadowMapForLightWithLiveShadingCS>(&FLocalVertexFactory::StaticType, PermutationVector, false);
+	if (!ComputeShader.IsNull())
+	{
+		GraphBuilder.AddPass(
+			RDG_EVENT_NAME("%s", *PassName),
+			PassParameters,
+			ERDGPassFlags::Compute,
+			[ComputeShader, PassParameters, Scene, MaterialRenderProxy, &Material, GroupCount](FRHIComputeCommandList& RHICmdList)
+			{
+#if 1
+				FMeshDrawShaderBindings ShaderBindings;
+		UE::MeshPassUtils::SetupComputeBindings(ComputeShader, Scene, Scene->GetFeatureLevel(), nullptr, *MaterialRenderProxy, Material, ShaderBindings);
+#else
+				FMeshMaterialShaderElementData ShaderElementData;
+		ShaderElementData.InitializeMeshMaterialData();
+
+		FMeshProcessorShaders PassShaders;
+		PassShaders.ComputeShader = ComputeShader;
+
+		FMeshDrawShaderBindings ShaderBindings;
+		ShaderBindings.Initialize(PassShaders);
+		{
+			FMeshDrawSingleShaderBindings SingleShaderBindings = ShaderBindings.GetSingleShaderBindings(SF_Compute);
+			ComputeShader->GetShaderBindings(Scene, Scene->GetFeatureLevel(), nullptr, *MaterialRenderProxy, Material, ShaderElementData, SingleShaderBindings);
+			//ComputeShader->ShaderLooseBindings.SetParameters(SingleShaderBindings, PassParameters);
+			SingleShaderBindings.AddTexture(
+				ComputeShader->ShaderLooseBindings.SceneDepthTextureBinding,
+				FShaderResourceParameter(),
+				TStaticSamplerState<SF_Point>::GetRHI(),
+				PassParameters->SceneTextures.SceneDepthTexture->GetRHI()
+			);
+			ShaderBindings.Finalize(&PassShaders);
+		}
+#endif
+
+		UE::MeshPassUtils::Dispatch(RHICmdList, ComputeShader, ShaderBindings, *PassParameters, GroupCount);
+			}
+		);
+	}
+
+	return true;
+}
+
+bool RenderVolumetricShadowMapForLightWithLiveShading(
+	FRDGBuilder& GraphBuilder,
+	// Scene data
+	const FSceneTextures& SceneTextures,
+	FScene* Scene,
+	const FSceneViewFamily& ViewFamily,
+	FViewInfo& View,
+	// Light data
+	uint32 LightType,
+	const FLightSceneInfo* LightSceneInfo,
+	const FVisibleLightInfo* VisibleLightInfo,
+	// Output
+	bool& bIsDirectionalLight,
+	FVector3f& TranslatedWorldOrigin,
+	FVector4f& TranslatedWorldPlane,
+	FMatrix44f* TranslatedWorldToShadow,
+	FIntVector& GroupCount,
+	int32& NumShadowMatrices,
+	FIntPoint& ShadowMapResolution,
+	uint32& MaxSampleCount,
+	FRDGTextureRef& BeerShadowMapTexture,
+	FRDGBufferRef& VolumetricShadowLinkedListBuffer
+)
+{
+	check(LightSceneInfo);
+	check(VisibleLightInfo);
+
+	const FProjectedShadowInfo* ProjectedShadowInfo = HeterogeneousVolumes::GetProjectedShadowInfo(VisibleLightInfo, 0);
+	check(ProjectedShadowInfo != NULL)
+
+	ShadowMapResolution = HeterogeneousVolumes::GetShadowMapResolution();
+
+	bool bIsMultiProjection = (LightType == LightType_Point) || (LightType == LightType_Rect);
+	GroupCount = FIntVector(1);
+	GroupCount.X = FMath::DivideAndRoundUp(ShadowMapResolution.X, FRenderVolumetricShadowMapForLightWithLiveShadingCS::GetThreadGroupSize2D());
+	GroupCount.Y = FMath::DivideAndRoundUp(ShadowMapResolution.Y, FRenderVolumetricShadowMapForLightWithLiveShadingCS::GetThreadGroupSize2D());
+	GroupCount.Z = bIsMultiProjection ? 6 : 1;
+
+	// Collect shadow-casting volumes
+	TSet<FVolumetricMeshBatch> HeterogeneousVolumesMeshBatches;
+	FBoxSphereBounds WorldVolumeBounds;
+	CollectHeterogeneousVolumeMeshBatchesForLight(LightSceneInfo, VisibleLightInfo, View, HeterogeneousVolumesMeshBatches, WorldVolumeBounds);
+	if (HeterogeneousVolumesMeshBatches.IsEmpty())
+	{
+		return false;
+	}
+
+	// Build shadow transform
+	NumShadowMatrices = ProjectedShadowInfo->OnePassShadowViewProjectionMatrices.Num();
+	FMatrix44f ShadowToTranslatedWorld[6];
+
+	if (NumShadowMatrices > 0)
+	{
+		FVector PreViewTranslation = View.ViewMatrices.GetPreViewTranslation();
+		FMatrix TranslatedWorldToWorldMatrix = FTranslationMatrix(-PreViewTranslation);
+		FVector LightPosition = LightSceneInfo->Proxy->GetPosition();
+		FMatrix WorldToLightMatrix = LightSceneInfo->Proxy->GetWorldToLight();
+
+		// Remove light rotation when building the RectLight projections..
+		FMatrix RotationalAdjustmentMatrix = FMatrix::Identity;
+		if (LightType == LIGHT_TYPE_RECT)
+		{
+			FVector LightDirection = LightSceneInfo->Proxy->GetDirection().GetSafeNormal();
+			RotationalAdjustmentMatrix = FRotationMatrix(LightDirection.Rotation());
+		}
+
+		FMatrix ViewMatrix[] = {
+			FLookFromMatrix(FVector::Zero(), FVector(-1, 0, 0), FVector(0, 0, 1)),
+			FLookFromMatrix(FVector::Zero(), FVector(1, 0, 0), FVector(0, 0, 1)),
+			FLookFromMatrix(FVector::Zero(), FVector(0, -1, 0), FVector(0, 0, 1)),
+			FLookFromMatrix(FVector::Zero(), FVector(0, 1, 0), FVector(0, 0, 1)),
+			FLookFromMatrix(FVector::Zero(), FVector(0, 0, -1), FVector(1, 0, 0)),
+			FLookFromMatrix(FVector::Zero(), FVector(0, 0, 1), FVector(1, 0, 0))
+		};
+
+		FMatrix PerspectiveMatrix = FPerspectiveMatrix(
+			PI / 4.0f,
+			ShadowMapResolution.X,
+			ShadowMapResolution.Y,
+			1.0,
+			LightSceneInfo->Proxy->GetRadius()
+		);
+
+		FMatrix ScreenMatrix = FScaleMatrix(FVector(0.5, -0.5, -0.5)) * FTranslationMatrix(FVector(0.5, 0.5, 0.5));
+
+		for (int32 i = 0; i < NumShadowMatrices; ++i)
+		{
+			FMatrix WorldToShadowMatrix = WorldToLightMatrix * RotationalAdjustmentMatrix * ViewMatrix[i] * PerspectiveMatrix * ScreenMatrix;
+			TranslatedWorldToShadow[i] = FMatrix44f(TranslatedWorldToWorldMatrix * WorldToShadowMatrix);
+			ShadowToTranslatedWorld[i] = TranslatedWorldToShadow[i].Inverse();
+		}
+		TranslatedWorldOrigin = FVector3f(PreViewTranslation + LightPosition);
+	}
+	else if (LightType == LightType_Directional)
+	{
+		bIsDirectionalLight = true;
+		// Build orthographic projection centered around volume..
+		FVector PreViewTranslation = View.ViewMatrices.GetPreViewTranslation();
+		FMatrix TranslatedWorldToWorldMatrix = FTranslationMatrix(-PreViewTranslation);
+
+		FMatrix TranslationMatrix = FTranslationMatrix(-WorldVolumeBounds.Origin);
+
+		FVector LightDirection = LightSceneInfo->Proxy->GetDirection().GetSafeNormal();
+		FMatrix RotationMatrix = FInverseRotationMatrix(LightDirection.Rotation());
+		FMatrix ScaleMatrix = FScaleMatrix(FVector(1.0 / WorldVolumeBounds.SphereRadius));
+
+		const FMatrix FaceMatrix(
+			FPlane(0, 0, 1, 0),
+			FPlane(0, 1, 0, 0),
+			FPlane(-1, 0, 0, 0),
+			FPlane(0, 0, 0, 1));
+
+		// Invert Z to match reverse-Z for the rest of the shadow types!
+		FMatrix ScreenMatrix = FScaleMatrix(FVector(0.5, -0.5, -0.5)) * FTranslationMatrix(FVector(0.5, 0.5, 0.5));
+		FMatrix WorldToShadowMatrix = TranslationMatrix * RotationMatrix * ScaleMatrix * FaceMatrix * ScreenMatrix;
+		FMatrix TranslatedWorldToShadowMatrix = TranslatedWorldToWorldMatrix * WorldToShadowMatrix;
+
+		NumShadowMatrices = 1;
+		TranslatedWorldToShadow[0] = FMatrix44f(TranslatedWorldToShadowMatrix);
+		ShadowToTranslatedWorld[0] = TranslatedWorldToShadow[0].Inverse();
+		TranslatedWorldOrigin = FVector3f(PreViewTranslation + WorldVolumeBounds.Origin - LightDirection * WorldVolumeBounds.SphereRadius);
+	}
+	else
+	{
+		FVector PreViewTranslation = View.ViewMatrices.GetPreViewTranslation();
+		FMatrix TranslatedWorldToWorldMatrix = FTranslationMatrix(-PreViewTranslation);
+		FVector4f ShadowmapMinMax = FVector4f::Zero();
+		FMatrix WorldToShadowMatrix = ProjectedShadowInfo->GetWorldToShadowMatrix(ShadowmapMinMax);
+		FMatrix TranslatedWorldToShadowMatrix = TranslatedWorldToWorldMatrix * WorldToShadowMatrix;
+
+		NumShadowMatrices = 1;
+		TranslatedWorldToShadow[0] = FMatrix44f(TranslatedWorldToShadowMatrix);
+		ShadowToTranslatedWorld[0] = TranslatedWorldToShadow[0].Inverse();
+		TranslatedWorldOrigin = FVector3f(View.ViewMatrices.GetPreViewTranslation() - ProjectedShadowInfo->PreShadowTranslation);
+	}
+
+	FVector LightDirection = LightSceneInfo->Proxy->GetDirection().GetSafeNormal();
+	float W = -FVector3f::DotProduct(TranslatedWorldOrigin, FVector3f(LightDirection));
+	TranslatedWorldPlane = FVector4f(LightDirection.X, LightDirection.Y, LightDirection.Z, W);
+
+	// Iterate over shadow-casting volumes
+	bool bHasShadowCastingVolume = false;
+	if (!HeterogeneousVolumesMeshBatches.IsEmpty())
+	{
+		auto VolumeMeshBatchItr = HeterogeneousVolumesMeshBatches.begin();
+
+		MaxSampleCount = HeterogeneousVolumes::GetShadowMaxSampleCount();
+		int32 VolumetricShadowLinkedListElementCount = ShadowMapResolution.X * ShadowMapResolution.Y * MaxSampleCount;
+		if (bIsMultiProjection)
+		{
+			VolumetricShadowLinkedListElementCount *= 6;
+		}
+		VolumetricShadowLinkedListBuffer = GraphBuilder.CreateBuffer(
+			FRDGBufferDesc::CreateStructuredDesc(sizeof(FAVSMLinkedListPackedData), VolumetricShadowLinkedListElementCount),
+			TEXT("HeterogeneousVolume.VolumetricShadowLinkedListBuffer")
+		);
+
+		RenderVolumetricShadowMapForLightForHeterogeneousVolumeWithLiveShading(
+			GraphBuilder,
+			SceneTextures,
+			Scene,
+			ViewFamily,
+			View,
+			// Light Info
+			LightType,
+			LightSceneInfo,
+			VisibleLightInfo,
+			// Shadow Info
+			TranslatedWorldOrigin,
+			NumShadowMatrices,
+			TranslatedWorldToShadow,
+			ShadowToTranslatedWorld,
+			ShadowMapResolution,
+			MaxSampleCount,
+			// Volume
+			*VolumeMeshBatchItr,
+			// Dispatch
+			GroupCount,
+			// Output
+			BeerShadowMapTexture,
+			VolumetricShadowLinkedListBuffer
+		);
+
+		++VolumeMeshBatchItr;
+		for (; VolumeMeshBatchItr != HeterogeneousVolumesMeshBatches.end(); ++VolumeMeshBatchItr)
+		{
+			FRDGBufferRef VolumetricShadowLinkedListBuffer1 = GraphBuilder.CreateBuffer(
+				FRDGBufferDesc::CreateStructuredDesc(sizeof(FAVSMLinkedListPackedData), VolumetricShadowLinkedListElementCount),
+				TEXT("HeterogeneousVolume.VolumetricShadowLinkedListBuffer1")
+			);
+
+			RenderVolumetricShadowMapForLightForHeterogeneousVolumeWithLiveShading(
+				GraphBuilder,
+				SceneTextures,
+				Scene,
+				ViewFamily,
+				View,
+				// Light Info
+				LightType,
+				LightSceneInfo,
+				VisibleLightInfo,
+				// Shadow Info
+				TranslatedWorldOrigin,
+				NumShadowMatrices,
+				TranslatedWorldToShadow,
+				ShadowToTranslatedWorld,
+				ShadowMapResolution,
+				MaxSampleCount,
+				// Volume
+				*VolumeMeshBatchItr,
+				// Dispatch
+				GroupCount,
+				// Output
+				BeerShadowMapTexture,
+				VolumetricShadowLinkedListBuffer1
+			);
+
+			FRDGBufferRef VolumetricShadowLinkedListBuffer2 = GraphBuilder.CreateBuffer(
+				FRDGBufferDesc::CreateStructuredDesc(sizeof(FAVSMLinkedListPackedData), VolumetricShadowLinkedListElementCount),
+				TEXT("HeterogeneousVolume.VolumetricShadowLinkedListBuffer2")
+			);
+
+			CombineVolumetricShadowMap(
+				GraphBuilder,
+				View,
+				GroupCount,
+				LightType,
+				ShadowMapResolution,
+				MaxSampleCount,
+				VolumetricShadowLinkedListBuffer,
+				VolumetricShadowLinkedListBuffer1,
+				VolumetricShadowLinkedListBuffer2
+			);
+
+			VolumetricShadowLinkedListBuffer = VolumetricShadowLinkedListBuffer2;
+		}
+	}
+
+	return true;
+}
+
+void RenderAdaptiveVolumetricShadowMapWithLiveShading(
+	FRDGBuilder& GraphBuilder,
+	// Scene data
+	const FSceneTextures& SceneTextures,
+	FScene* Scene,
+	const FSceneViewFamily& ViewFamily,
+	FViewInfo& View,
+	// Light data
+	TArray<FVisibleLightInfo, SceneRenderingAllocator>& VisibleLightInfos
+)
+{
+	RDG_EVENT_SCOPE(GraphBuilder, "Adaptive Volumetric Shadow Maps");
+	bool bShouldRenderShadowMaps = !View.ViewRect.IsEmpty();
+
+	// Light culling
+	TArray<FLightSceneInfoCompact, TInlineAllocator<64>> LightSceneInfoCompact;
+	for (auto LightIt = Scene->Lights.CreateConstIterator(); LightIt; ++LightIt)
+	{
+		// TODO: Use global bounds information..
+		//if (LightIt->AffectsPrimitive(HeterogeneousVolumeInterface->GetBounds(), HeterogeneousVolumeInterface->GetPrimitiveSceneProxy()))
+		{
+			LightSceneInfoCompact.Add(*LightIt);
+		}
+	}
+
+	// Light loop:
+	int32 NumPasses = LightSceneInfoCompact.Num();
+	for (int32 PassIndex = 0; PassIndex < NumPasses; ++PassIndex)
+	{
+		bool bApplyDirectLighting = !LightSceneInfoCompact.IsEmpty();
+		bool bApplyEmissionAndTransmittance = false;
+		bool bCastsVolumetricShadow = false;
+
+		uint32 LightType = 0;
+		FLightSceneInfo* LightSceneInfo = nullptr;
+		const FVisibleLightInfo* VisibleLightInfo = nullptr;
+		if (!LightSceneInfoCompact.IsEmpty())
+		{
+			LightType = LightSceneInfoCompact[PassIndex].LightType;
+			LightSceneInfo = LightSceneInfoCompact[PassIndex].LightSceneInfo;
+			check(LightSceneInfo != nullptr);
+
+			bool bDynamicallyShadowed = false;
+			if (LightSceneInfo)
+			{
+				VisibleLightInfo = &VisibleLightInfos[LightSceneInfo->Id];
+				bCastsVolumetricShadow = LightSceneInfo->Proxy && LightSceneInfo->Proxy->CastsVolumetricShadow();
+				bDynamicallyShadowed = HeterogeneousVolumes::IsDynamicShadow(VisibleLightInfo);
+			}
+
+			TRDGUniformBufferRef<FAdaptiveVolumetricShadowMapUniformBufferParameters> AdaptiveVolumetricShadowMapUniformBuffer;
+			bool bCreateShadowMap = bShouldRenderShadowMaps && bCastsVolumetricShadow && bDynamicallyShadowed && !ShouldRenderRayTracingShadowsForLight(LightSceneInfoCompact[PassIndex]);
+			if (bCreateShadowMap)
+			{
+				FString LightName;
+				FSceneRenderer::GetLightNameForDrawEvent(LightSceneInfo->Proxy, LightName);
+				RDG_EVENT_SCOPE(GraphBuilder, "%s", *LightName);
+
+				FRDGTextureDesc Desc = SceneTextures.Color.Target->Desc;
+				Desc.Format = PF_FloatRGBA;
+				Desc.Flags &= ~(TexCreate_FastVRAM);
+				FRDGTextureRef BeerShadowMapTexture = GraphBuilder.CreateTexture(Desc, TEXT("BeerShadowMapTexture"));
+				AddClearUAVPass(GraphBuilder, GraphBuilder.CreateUAV(BeerShadowMapTexture), FLinearColor::Transparent);
+
+				bool bIsDirectionalLight = false;
+				FVector3f TranslatedWorldOrigin = FVector3f::Zero();
+				FVector4f TranslatedWorldPlane = FVector4f::Zero();
+				FMatrix44f TranslatedWorldToShadow[] =
+				{
+					FMatrix44f::Identity,
+					FMatrix44f::Identity,
+					FMatrix44f::Identity,
+					FMatrix44f::Identity,
+					FMatrix44f::Identity,
+					FMatrix44f::Identity
+				};
+				FIntVector GroupCount = FIntVector::ZeroValue;
+				int32 NumShadowMatrices = 0;
+				FIntPoint VolumetricShadowMapResolution = FIntPoint::NoneValue;
+				uint32 VolumetricShadowMapMaxSampleCount = 0;
+				FRDGBufferRef VolumetricShadowMapLinkedListBuffer;
+				bool bIsCreated = RenderVolumetricShadowMapForLightWithLiveShading(
+					GraphBuilder,
+					// Scene data
+					SceneTextures,
+					Scene,
+					ViewFamily,
+					View,
+					// Light data
+					LightType,
+					LightSceneInfo,
+					VisibleLightInfo,
+					// Output
+					bIsDirectionalLight,
+					TranslatedWorldOrigin,
+					TranslatedWorldPlane,
+					TranslatedWorldToShadow,
+					GroupCount,
+					NumShadowMatrices,
+					VolumetricShadowMapResolution,
+					VolumetricShadowMapMaxSampleCount,
+					BeerShadowMapTexture,
+					VolumetricShadowMapLinkedListBuffer
+				);
+
+				if (bIsCreated)
+				{
+					FRDGBufferRef VolumetricShadowMapIndirectionBuffer;
+					FRDGBufferRef VolumetricShadowMapSampleBuffer;
+					CompressVolumetricShadowMap(
+						GraphBuilder,
+						View,
+						GroupCount,
+						VolumetricShadowMapResolution,
+						VolumetricShadowMapMaxSampleCount,
+						VolumetricShadowMapLinkedListBuffer,
+						VolumetricShadowMapIndirectionBuffer,
+						VolumetricShadowMapSampleBuffer
+					);
+
+					CreateAdaptiveVolumetricShadowMapUniformBuffer(
+						GraphBuilder,
+						TranslatedWorldOrigin,
+						TranslatedWorldPlane,
+						TranslatedWorldToShadow,
+						VolumetricShadowMapResolution,
+						NumShadowMatrices,
+						VolumetricShadowMapMaxSampleCount,
+						bIsDirectionalLight,
+						VolumetricShadowMapLinkedListBuffer,
+						VolumetricShadowMapIndirectionBuffer,
+						VolumetricShadowMapSampleBuffer,
+						AdaptiveVolumetricShadowMapUniformBuffer
+					);
+				}
+				else
+				{
+					AdaptiveVolumetricShadowMapUniformBuffer = HeterogeneousVolumes::CreateEmptyAdaptiveVolumetricShadowMapUniformBuffer(GraphBuilder);
+				}
+			}
+			else
+			{
+				AdaptiveVolumetricShadowMapUniformBuffer = HeterogeneousVolumes::CreateEmptyAdaptiveVolumetricShadowMapUniformBuffer(GraphBuilder);
+			}
+
+			if (View.ViewState)
+			{
+				TRDGUniformBufferRef<FAdaptiveVolumetricShadowMapUniformBufferParameters>& AdaptiveVolumetricShadowMap = View.ViewState->AdaptiveVolumetricShadowMapUniformBufferMap.FindOrAdd(LightSceneInfo->Id);
+				AdaptiveVolumetricShadowMap = AdaptiveVolumetricShadowMapUniformBuffer;
+			}
+		}
+	}
+}
+
+void RenderAdaptiveVolumetricCameraMapWithLiveShading(
+	FRDGBuilder& GraphBuilder,
+	// Scene data
+	const FSceneTextures& SceneTextures,
+	FScene* Scene,
+	const FSceneViewFamily& ViewFamily,
+	FViewInfo& View
+)
+{
+	if (View.ViewState == nullptr)
+	{
+		return;
+	}
+	RDG_EVENT_SCOPE(GraphBuilder, "Adaptive Volumetric Camera Map");
+
+	TRDGUniformBufferRef<FAdaptiveVolumetricShadowMapUniformBufferParameters> AdaptiveVolumetricShadowMapUniformBuffer;
+
+	// Collect all volumes for view
+	TSet<FVolumetricMeshBatch> HeterogeneousVolumesMeshBatches;
+	FBoxSphereBounds WorldVolumeBounds;
+	CollectHeterogeneousVolumeMeshBatchesForView(View, HeterogeneousVolumesMeshBatches, WorldVolumeBounds);
+
+	bool bShouldRenderCameraMap = !View.ViewRect.IsEmpty() && !HeterogeneousVolumesMeshBatches.IsEmpty();
+	if (bShouldRenderCameraMap)
+	{
+		// Resolution
+		FIntPoint ShadowMapResolution = FIntPoint(View.ViewRect.Width(), View.ViewRect.Height());
+		float DownsampleFactor = HeterogeneousVolumes::GetCameraDownsampleFactor();
+		ShadowMapResolution.X = FMath::Max(ShadowMapResolution.X / DownsampleFactor, 1);
+		ShadowMapResolution.Y = FMath::Max(ShadowMapResolution.Y / DownsampleFactor, 1);
+
+		// Transform
+		FMatrix ViewToClip = FPerspectiveMatrix(
+			FMath::DegreesToRadians(View.FOV * 0.5),
+			ShadowMapResolution.X,
+			ShadowMapResolution.Y,
+			1.0,
+			HeterogeneousVolumes::GetMaxTraceDistance()
+		);
+		FMatrix ClipToView = ViewToClip.Inverse();
+		FMatrix ScreenMatrix = FScaleMatrix(FVector(0.5, -0.5, -0.5)) * FTranslationMatrix(FVector(0.5, 0.5, 0.5));
+
+		int32 NumShadowMatrices = 1;
+		FMatrix44f TranslatedWorldToShadow[] = {
+			FMatrix44f(View.ViewMatrices.GetTranslatedViewMatrix() * ViewToClip * ScreenMatrix)
+		};
+		FMatrix44f ShadowToTranslatedWorld[] = {
+			TranslatedWorldToShadow[0].Inverse()
+		};
+		FVector3f TranslatedWorldOrigin = ShadowToTranslatedWorld[0].GetOrigin();
+
+		// Dispatch
+		FIntVector GroupCount = FIntVector(1);
+		GroupCount.X = FMath::DivideAndRoundUp(ShadowMapResolution.X, FRenderVolumetricShadowMapForLightWithLiveShadingCS::GetThreadGroupSize2D());
+		GroupCount.Y = FMath::DivideAndRoundUp(ShadowMapResolution.Y, FRenderVolumetricShadowMapForLightWithLiveShadingCS::GetThreadGroupSize2D());
+
+		// Visualization Texture
+		FRDGTextureDesc Desc = SceneTextures.Color.Target->Desc;
+		Desc.Format = PF_FloatRGBA;
+		Desc.Flags &= ~(TexCreate_FastVRAM);
+		FRDGTextureRef BeerShadowMapTexture = GraphBuilder.CreateTexture(Desc, TEXT("BeerShadowMapTexture"));
+		AddClearUAVPass(GraphBuilder, GraphBuilder.CreateUAV(BeerShadowMapTexture), FLinearColor::Transparent);
+
+		auto VolumeMeshBatchItr = HeterogeneousVolumesMeshBatches.begin();
+		int32 MaxSampleCount = HeterogeneousVolumes::GetShadowMaxSampleCount();
+		int32 VolumetricShadowLinkedListElementCount = ShadowMapResolution.X * ShadowMapResolution.Y * MaxSampleCount;
+
+		FRDGBufferRef VolumetricShadowLinkedListBuffer = GraphBuilder.CreateBuffer(
+			FRDGBufferDesc::CreateStructuredDesc(sizeof(FAVSMLinkedListPackedData), VolumetricShadowLinkedListElementCount),
+			TEXT("HeterogeneousVolume.VolumetricShadowLinkedListBuffer")
+		);
+
+		// Build a camera shadow for one volume
+		int32 LightType = 0;
+		FLightSceneInfo* LightSceneInfo = nullptr;
+		FVisibleLightInfo* VisibleLightInfo = nullptr;
+		RenderVolumetricShadowMapForLightForHeterogeneousVolumeWithLiveShading(
+			GraphBuilder,
+			SceneTextures,
+			Scene,
+			ViewFamily,
+			View,
+			// Light Info
+			LightType,
+			LightSceneInfo,
+			VisibleLightInfo,
+			// Shadow Info
+			TranslatedWorldOrigin,
+			NumShadowMatrices,
+			TranslatedWorldToShadow,
+			ShadowToTranslatedWorld,
+			ShadowMapResolution,
+			MaxSampleCount,
+			// Volume
+			*VolumeMeshBatchItr,
+			// Dispatch
+			GroupCount,
+			// Output
+			BeerShadowMapTexture,
+			VolumetricShadowLinkedListBuffer
+		);
+
+		// Iterate over volumes, combining each into the existing shadow map
+		++VolumeMeshBatchItr;
+		for (; VolumeMeshBatchItr != HeterogeneousVolumesMeshBatches.end(); ++VolumeMeshBatchItr)
+		{
+			FRDGBufferRef VolumetricShadowLinkedListBuffer1 = GraphBuilder.CreateBuffer(
+				FRDGBufferDesc::CreateStructuredDesc(sizeof(FAVSMLinkedListPackedData), VolumetricShadowLinkedListElementCount),
+				TEXT("HeterogeneousVolume.VolumetricShadowLinkedListBuffer1")
+			);
+
+			RenderVolumetricShadowMapForLightForHeterogeneousVolumeWithLiveShading(
+				GraphBuilder,
+				SceneTextures,
+				Scene,
+				ViewFamily,
+				View,
+				// Light Info
+				LightType,
+				LightSceneInfo,
+				VisibleLightInfo,
+				// Shadow Info
+				TranslatedWorldOrigin,
+				NumShadowMatrices,
+				TranslatedWorldToShadow,
+				ShadowToTranslatedWorld,
+				ShadowMapResolution,
+				MaxSampleCount,
+				// Volume
+				*VolumeMeshBatchItr,
+				// Dispatch
+				GroupCount,
+				// Output
+				BeerShadowMapTexture,
+				VolumetricShadowLinkedListBuffer1
+			);
+
+			FRDGBufferRef VolumetricShadowLinkedListBuffer2 = GraphBuilder.CreateBuffer(
+				FRDGBufferDesc::CreateStructuredDesc(sizeof(FAVSMLinkedListPackedData), VolumetricShadowLinkedListElementCount),
+				TEXT("HeterogeneousVolume.VolumetricShadowLinkedListBuffer2")
+			);
+
+			CombineVolumetricShadowMap(
+				GraphBuilder,
+				View,
+				GroupCount,
+				LightType,
+				ShadowMapResolution,
+				MaxSampleCount,
+				VolumetricShadowLinkedListBuffer,
+				VolumetricShadowLinkedListBuffer1,
+				VolumetricShadowLinkedListBuffer2
+			);
+
+			VolumetricShadowLinkedListBuffer = VolumetricShadowLinkedListBuffer2;
+		}
+
+		FRDGBufferRef VolumetricShadowIndirectionBuffer = GSystemTextures.GetDefaultStructuredBuffer(GraphBuilder, sizeof(FAVSMIndirectionPackedData));
+		FRDGBufferRef VolumetricShadowSampleBuffer = GSystemTextures.GetDefaultStructuredBuffer(GraphBuilder, sizeof(FAVSMSamplePackedData));
+		CompressVolumetricShadowMap(
+			GraphBuilder,
+			View,
+			GroupCount,
+			ShadowMapResolution,
+			MaxSampleCount,
+			VolumetricShadowLinkedListBuffer,
+			VolumetricShadowIndirectionBuffer,
+			VolumetricShadowSampleBuffer
+		);
+
+		FVector4f TranslatedWorldPlane = FVector4f::Zero();
+		bool bIsDirectionalLight = false;
+		CreateAdaptiveVolumetricShadowMapUniformBuffer(
+			GraphBuilder,
+			TranslatedWorldOrigin,
+			TranslatedWorldPlane,
+			TranslatedWorldToShadow,
+			ShadowMapResolution,
+			NumShadowMatrices,
+			MaxSampleCount,
+			bIsDirectionalLight,
+			VolumetricShadowLinkedListBuffer,
+			VolumetricShadowIndirectionBuffer,
+			VolumetricShadowSampleBuffer,
+			AdaptiveVolumetricShadowMapUniformBuffer
+		);
+	}
+	else
+	{
+		AdaptiveVolumetricShadowMapUniformBuffer = HeterogeneousVolumes::CreateEmptyAdaptiveVolumetricShadowMapUniformBuffer(GraphBuilder);
+	}
+
+	View.ViewState->AdaptiveVolumetricCameraMapUniformBuffer = AdaptiveVolumetricShadowMapUniformBuffer;
 }
\ No newline at end of file
diff --git a/Engine/Source/Runtime/Renderer/Private/HeterogeneousVolumes/HeterogeneousVolumesVoxelGridPipeline.cpp b/Engine/Source/Runtime/Renderer/Private/HeterogeneousVolumes/HeterogeneousVolumesVoxelGridPipeline.cpp
--- a/Engine/Source/Runtime/Renderer/Private/HeterogeneousVolumes/HeterogeneousVolumesVoxelGridPipeline.cpp
+++ b/Engine/Source/Runtime/Renderer/Private/HeterogeneousVolumes/HeterogeneousVolumesVoxelGridPipeline.cpp
@@ -219,8 +219,8 @@ static TAutoConsoleVariable<float> CVarHeterogeneousVolumesStepSizeForShadows(
 
 static TAutoConsoleVariable<int32> CVarHeterogeneousVolumesShadowMaxSampleCount(
 	TEXT("r.HeterogeneousVolumes.Shadows.MaxSampleCount"),
-	16,
-	TEXT("Maximum sample count when building volumetric shadow map (Default = 16)\n"),
+	8,
+	TEXT("Maximum sample count when building volumetric shadow map (Default = 8)\n"),
 	ECVF_RenderThreadSafe
 );
 
@@ -266,6 +266,13 @@ static TAutoConsoleVariable<int32> CVarHeterogeneousVolumesShadowOutOfFrustumSha
 	ECVF_RenderThreadSafe
 );
 
+static TAutoConsoleVariable<int32> CVarHeterogeneousVolumesShadowJitter(
+	TEXT("r.HeterogeneousVolumes.Shadows.Jitter"),
+	0,
+	TEXT("Enables jittering when constructing shadows (Default = 0)\n"),
+	ECVF_RenderThreadSafe
+);
+
 IMPLEMENT_UNIFORM_BUFFER_STRUCT(FOrthoVoxelGridUniformBufferParameters, "OrthoGridUniformBuffer");
 IMPLEMENT_UNIFORM_BUFFER_STRUCT(FFrustumVoxelGridUniformBufferParameters, "FrustumGridUniformBuffer");
 IMPLEMENT_UNIFORM_BUFFER_STRUCT(FAdaptiveVolumetricShadowMapUniformBufferParameters, "AVSM");
@@ -406,7 +413,7 @@ namespace HeterogeneousVolumes
 
 	FIntPoint GetShadowMapResolution()
 	{
-		return FIntPoint(FMath::Clamp(CVarHeterogeneousVolumesShadowResolution.GetValueOnRenderThread(), 1, 1024));
+		return FIntPoint(FMath::Clamp(CVarHeterogeneousVolumesShadowResolution.GetValueOnRenderThread(), 1, 2048));
 	}
 
 	float GetStepSizeForShadows()
@@ -448,6 +455,38 @@ namespace HeterogeneousVolumes
 	{
 		return FMath::Max(CVarHeterogeneousVolumesShadowOutOfFrustumShadingRate.GetValueOnRenderThread(), 0.1);
 	}
+
+	bool EnableJitterForShadows()
+	{
+		return CVarHeterogeneousVolumesShadowJitter.GetValueOnRenderThread() != 0;
+	}
+
+	bool IsDynamicShadow(const FVisibleLightInfo* VisibleLightInfo)
+	{
+		check(VisibleLightInfo != nullptr);
+		if (VisibleLightInfo)
+		{
+			return !VisibleLightInfo->ShadowsToProject.IsEmpty();
+		}
+
+		return false;
+	}
+
+	const FProjectedShadowInfo* GetProjectedShadowInfo(const FVisibleLightInfo* VisibleLightInfo, int32 ShadowIndex)
+	{
+		check(VisibleLightInfo != nullptr);
+		if (VisibleLightInfo && ShadowIndex < VisibleLightInfo->ShadowsToProject.Num())
+		{
+			return VisibleLightInfo->ShadowsToProject[ShadowIndex];
+		}
+
+		return nullptr;
+	}
+}
+
+uint32 GetTypeHash(const FVolumetricMeshBatch& MeshBatch)
+{
+	return HashCombineFast(GetTypeHash(MeshBatch.Mesh), GetTypeHash(MeshBatch.Proxy));
 }
 
 class FMarkTopLevelGridVoxelsForFrustumGrid : public FGlobalShader
@@ -2607,33 +2646,6 @@ void RasterizeVolumesIntoOrthoVoxelGrid(
 	}
 }
 
-uint32 GetTypeHash(const FVolumetricMeshBatch& MeshBatch)
-{
-	return HashCombineFast(GetTypeHash(MeshBatch.Mesh), GetTypeHash(MeshBatch.Proxy));
-}
-
-bool IsDynamicShadow(const FVisibleLightInfo* VisibleLightInfo)
-{
-	check(VisibleLightInfo != nullptr);
-	if (VisibleLightInfo)
-	{
-		return !VisibleLightInfo->ShadowsToProject.IsEmpty();
-	}
-
-	return false;
-}
-
-const FProjectedShadowInfo* GetProjectedShadowInfo(const FVisibleLightInfo* VisibleLightInfo, int32 ShadowIndex)
-{
-	check(VisibleLightInfo != nullptr);
-	if (VisibleLightInfo && ShadowIndex < VisibleLightInfo->ShadowsToProject.Num())
-	{
-		return VisibleLightInfo->ShadowsToProject[ShadowIndex];
-	}
-
-	return nullptr;
-}
-
 bool ShouldAddToVoxelGrid(const FPrimitiveSceneProxy* Proxy, const FViewInfo& View, const FVoxelGridBuildOptions& BuildOptions)
 {
 	switch (BuildOptions.VoxelGridBuildMode)
@@ -2690,23 +2702,23 @@ void CollectHeterogeneousVolumeMeshBatches(
 			check(VisibleLightInfo);
 			for (int32 ShadowIndex = 0; ShadowIndex < VisibleLightInfo->ShadowsToProject.Num(); ++ShadowIndex)
 			{
-				const FProjectedShadowInfo* ProjectedShadowInfo = GetProjectedShadowInfo(VisibleLightInfo, ShadowIndex);
-			if (ProjectedShadowInfo != nullptr)
-			{
-				const TArray<FMeshBatchAndRelevance, SceneRenderingAllocator>& MeshBatches = ProjectedShadowInfo->GetDynamicSubjectHeterogeneousVolumeMeshElements();
-				for (int32 MeshBatchIndex = 0; MeshBatchIndex < MeshBatches.Num(); ++MeshBatchIndex)
+				const FProjectedShadowInfo* ProjectedShadowInfo = HeterogeneousVolumes::GetProjectedShadowInfo(VisibleLightInfo, ShadowIndex);
+				if (ProjectedShadowInfo != nullptr)
 				{
-					const FMeshBatchAndRelevance& MeshBatch = MeshBatches[MeshBatchIndex];
-					if (ShouldAddToVoxelGrid(MeshBatch.PrimitiveSceneProxy, *ProjectedShadowInfo->ShadowDepthView, BuildOptions))
+					const TArray<FMeshBatchAndRelevance, SceneRenderingAllocator>& MeshBatches = ProjectedShadowInfo->GetDynamicSubjectHeterogeneousVolumeMeshElements();
+					for (int32 MeshBatchIndex = 0; MeshBatchIndex < MeshBatches.Num(); ++MeshBatchIndex)
 					{
-						HeterogeneousVolumesMeshBatches.FindOrAdd(FVolumetricMeshBatch(MeshBatch.Mesh, MeshBatch.PrimitiveSceneProxy));
+						const FMeshBatchAndRelevance& MeshBatch = MeshBatches[MeshBatchIndex];
+						if (ShouldAddToVoxelGrid(MeshBatch.PrimitiveSceneProxy, *ProjectedShadowInfo->ShadowDepthView, BuildOptions))
+						{
+							HeterogeneousVolumesMeshBatches.FindOrAdd(FVolumetricMeshBatch(MeshBatch.Mesh, MeshBatch.PrimitiveSceneProxy));
+						}
 					}
 				}
 			}
 		}
 	}
 }
-}
 
 void BuildOrthoVoxelGrid(
 	FRDGBuilder& GraphBuilder,
@@ -3063,6 +3075,61 @@ class FCompressVolumetricShadowMapCS : public FGlobalShader
 
 IMPLEMENT_GLOBAL_SHADER(FCompressVolumetricShadowMapCS, "/Engine/Private/HeterogeneousVolumes/HeterogeneousVolumesVoxelGridShadows.usf", "CompressVolumetricShadowMapCS", SF_Compute);
 
+class FCombineVolumetricShadowMapsCS : public FGlobalShader
+{
+	DECLARE_GLOBAL_SHADER(FCombineVolumetricShadowMapsCS);
+	SHADER_USE_PARAMETER_STRUCT(FCombineVolumetricShadowMapsCS, FGlobalShader);
+
+	class FUseAVSMCompression : SHADER_PERMUTATION_BOOL("USE_AVSM_COMPRESSION");
+	using FPermutationDomain = TShaderPermutationDomain<FUseAVSMCompression>;
+
+	BEGIN_SHADER_PARAMETER_STRUCT(FParameters, )
+		// Input
+		SHADER_PARAMETER(FIntPoint, ShadowResolution)
+		SHADER_PARAMETER_RDG_BUFFER_SRV(StructuredBuffer<uint2>, VolumetricShadowLinkedListBuffer0)
+		SHADER_PARAMETER_RDG_BUFFER_SRV(StructuredBuffer<uint2>, VolumetricShadowLinkedListBuffer1)
+
+		// Volumetric Shadow Map data
+		SHADER_PARAMETER(int, MaxSampleCount)
+		SHADER_PARAMETER(float, AbsoluteErrorThreshold)
+		SHADER_PARAMETER(float, RelativeErrorThreshold)
+
+		// Output
+		SHADER_PARAMETER_RDG_BUFFER_UAV(RWStructuredBuffer<int2>, RWVolumetricShadowLinkedListBuffer)
+	END_SHADER_PARAMETER_STRUCT()
+
+	static bool ShouldCompilePermutation(
+		const FGlobalShaderPermutationParameters& Parameters
+	)
+	{
+		return DoesPlatformSupportHeterogeneousVolumes(Parameters.Platform);
+	}
+
+	static void ModifyCompilationEnvironment(
+		const FGlobalShaderPermutationParameters& Parameters,
+		FShaderCompilerEnvironment& OutEnvironment
+	)
+	{
+		FMaterialShader::ModifyCompilationEnvironment(Parameters, OutEnvironment);
+		OutEnvironment.SetDefine(TEXT("THREADGROUP_SIZE_1D"), GetThreadGroupSize1D());
+		OutEnvironment.SetDefine(TEXT("THREADGROUP_SIZE_2D"), GetThreadGroupSize2D());
+		OutEnvironment.SetDefine(TEXT("THREADGROUP_SIZE_3D"), GetThreadGroupSize3D());
+
+		// Temporary disabling..
+		OutEnvironment.SetDefine(TEXT("DIM_USE_TRANSMITTANCE_VOLUME"), 0);
+		OutEnvironment.SetDefine(TEXT("DIM_USE_INSCATTERING_VOLUME"), 0);
+		OutEnvironment.SetDefine(TEXT("DIM_USE_LUMEN_GI"), 0);
+
+		OutEnvironment.CompilerFlags.Add(CFLAG_AllowTypedUAVLoads);
+	}
+
+	static int32 GetThreadGroupSize1D() { return GetThreadGroupSize3D() * GetThreadGroupSize3D() * GetThreadGroupSize3D(); }
+	static int32 GetThreadGroupSize2D() { return 8; }
+	static int32 GetThreadGroupSize3D() { return 4; }
+};
+
+IMPLEMENT_GLOBAL_SHADER(FCombineVolumetricShadowMapsCS, "/Engine/Private/HeterogeneousVolumes/HeterogeneousVolumesVoxelGridShadows.usf", "CombineVolumetricShadowMapsCS", SF_Compute);
+
 void CreateAdaptiveVolumetricShadowMapUniformBuffer(
 	FRDGBuilder& GraphBuilder,
 	const FVector3f& TranslatedWorldOrigin,
@@ -3270,6 +3337,66 @@ void CompressVolumetricShadowMap(
 	);
 }
 
+void CombineVolumetricShadowMap(
+	FRDGBuilder& GraphBuilder,
+	FViewInfo& View,
+	FIntVector GroupCount,
+	// Input
+	uint32 LightType,
+	FIntPoint ShadowMapResolution,
+	uint32 MaxSampleCount,
+	FRDGBufferRef VolumetricShadowLinkedListBuffer0,
+	FRDGBufferRef VolumetricShadowLinkedListBuffer1,
+	// Output
+	FRDGBufferRef& VolumetricShadowLinkedListBuffer
+)
+{
+	bool bIsMultiProjection = (LightType == LightType_Point) || (LightType == LightType_Rect);
+	GroupCount = FIntVector(1);
+	GroupCount.X = FMath::DivideAndRoundUp(ShadowMapResolution.X, FCombineVolumetricShadowMapsCS::GetThreadGroupSize2D());
+	GroupCount.Y = FMath::DivideAndRoundUp(ShadowMapResolution.Y, FCombineVolumetricShadowMapsCS::GetThreadGroupSize2D());
+	GroupCount.Z = bIsMultiProjection ? 6 : 1;
+
+	MaxSampleCount = HeterogeneousVolumes::GetShadowMaxSampleCount();
+	int32 VolumetricShadowLinkedListElementCount = ShadowMapResolution.X * ShadowMapResolution.Y * MaxSampleCount;
+	if (bIsMultiProjection)
+	{
+		VolumetricShadowLinkedListElementCount *= 6;
+	}
+	VolumetricShadowLinkedListBuffer = GraphBuilder.CreateBuffer(
+		FRDGBufferDesc::CreateStructuredDesc(sizeof(FAVSMLinkedListPackedData), VolumetricShadowLinkedListElementCount),
+		TEXT("HeterogeneousVolume.VolumetricShadowLinkedListBuffer")
+	);
+
+	FCombineVolumetricShadowMapsCS::FParameters* PassParameters = GraphBuilder.AllocParameters<FCombineVolumetricShadowMapsCS::FParameters>();
+	{
+		// Input
+		PassParameters->ShadowResolution = ShadowMapResolution;
+		PassParameters->VolumetricShadowLinkedListBuffer0 = GraphBuilder.CreateSRV(VolumetricShadowLinkedListBuffer0);
+		PassParameters->VolumetricShadowLinkedListBuffer1 = GraphBuilder.CreateSRV(VolumetricShadowLinkedListBuffer1);
+
+		PassParameters->MaxSampleCount = MaxSampleCount;
+		PassParameters->AbsoluteErrorThreshold = HeterogeneousVolumes::GetShadowAbsoluteErrorThreshold();
+		PassParameters->RelativeErrorThreshold = HeterogeneousVolumes::GetShadowRelativeErrorThreshold();
+
+		// Output
+		PassParameters->RWVolumetricShadowLinkedListBuffer = GraphBuilder.CreateUAV(VolumetricShadowLinkedListBuffer);
+	}
+
+	FCombineVolumetricShadowMapsCS::FPermutationDomain PermutationVector;
+	PermutationVector.Set<FCombineVolumetricShadowMapsCS::FUseAVSMCompression>(HeterogeneousVolumes::UseAVSMCompression());
+	TShaderRef<FCombineVolumetricShadowMapsCS> ComputeShader = View.ShaderMap->GetShader<FCombineVolumetricShadowMapsCS>(PermutationVector);
+	FComputeShaderUtils::AddPass(
+		GraphBuilder,
+		// Add Light name..
+		RDG_EVENT_NAME("CombineVolumetricShadowMapsCS"),
+		ERDGPassFlags::Compute | ERDGPassFlags::NeverCull,
+		ComputeShader,
+		PassParameters,
+		GroupCount
+	);
+}
+
 void RenderVolumetricShadowMapForLightWithVoxelGrid(
 	FRDGBuilder& GraphBuilder,
 	// Scene data
@@ -3305,7 +3432,7 @@ void RenderVolumetricShadowMapForLightWithVoxelGrid(
 	check(LightSceneInfo);
 	check(VisibleLightInfo);
 
-	const FProjectedShadowInfo* ProjectedShadowInfo = GetProjectedShadowInfo(VisibleLightInfo, 0);
+	const FProjectedShadowInfo* ProjectedShadowInfo = HeterogeneousVolumes::GetProjectedShadowInfo(VisibleLightInfo, 0);
 	check(ProjectedShadowInfo != NULL)
 
 	ShadowMapResolution = HeterogeneousVolumes::GetShadowMapResolution();
@@ -3677,7 +3804,7 @@ void RenderAdaptiveVolumetricShadowMapWithVoxelGrid(
 			{
 				VisibleLightInfo = &VisibleLightInfos[LightSceneInfo->Id];
 				bApplyShadowTransmittance = LightSceneInfo->Proxy && LightSceneInfo->Proxy->CastsVolumetricShadow();
-				bDynamicallyShadowed = IsDynamicShadow(VisibleLightInfo);
+				bDynamicallyShadowed = HeterogeneousVolumes::IsDynamicShadow(VisibleLightInfo);
 			}
 
 			TRDGUniformBufferRef<FAdaptiveVolumetricShadowMapUniformBufferParameters> AdaptiveVolumetricShadowMapUniformBuffer;
