Fixed the single-pass Nanite multi-view rendering for VR: the culling view volumes for the extra views were not passed to the instance culling query.

#rnx

diff --git a/Engine/Source/Runtime/Renderer/Private/DeferredShadingRenderer.cpp b/Engine/Source/Runtime/Renderer/Private/DeferredShadingRenderer.cpp
--- a/Engine/Source/Runtime/Renderer/Private/DeferredShadingRenderer.cpp
+++ b/Engine/Source/Runtime/Renderer/Private/DeferredShadingRenderer.cpp
@@ -1280,7 +1280,7 @@ void FDeferredShadingSceneRenderer::RenderNanite(FRDGBuilder& GraphBuilder, cons
 		const bool bDrawSceneViewsInOneNanitePass = InViews.Num() > 1 && Nanite::ShouldDrawSceneViewsInOneNanitePass(InViews[0]);
 
 		// creates one or more Nanite views (normally one per view unless drawing multiple views together - e.g. Stereo ISR views)
-		auto CreateNaniteViews = [bDrawSceneViewsInOneNanitePass, &InViews, &PrimaryNaniteViews, &GraphBuilder](const FViewInfo& View, int32 ViewIndex, const FIntPoint& RasterTextureSize, float MaxPixelsPerEdgeMultipler) -> Nanite::FPackedViewArray*
+		auto CreateNaniteViews = [bDrawSceneViewsInOneNanitePass, &InViews, &PrimaryNaniteViews, &GraphBuilder](const FViewInfo& View, int32 ViewIndex, const FIntPoint& RasterTextureSize, float MaxPixelsPerEdgeMultipler, TArray<FConvexVolume> &OutViewsCullingVolumes) -> Nanite::FPackedViewArray*
 		{
 			Nanite::FPackedViewArray::ArrayType OutViews;
 
@@ -1299,6 +1299,7 @@ void FDeferredShadingSceneRenderer::RenderNanite(FRDGBuilder& GraphBuilder, cons
 				MaxPixelsPerEdgeMultipler,
 				&HZBTestRect
 			);
+			OutViewsCullingVolumes.Add(View.ViewFrustum);
 			OutViews.Add(PackedView);
 			PrimaryNaniteViews.Add(PackedView);
 
@@ -1320,7 +1321,7 @@ void FDeferredShadingSceneRenderer::RenderNanite(FRDGBuilder& GraphBuilder, cons
 						MaxPixelsPerEdgeMultipler,
 						&SecondaryHZBTestRect
 					);
-
+					OutViewsCullingVolumes.Add(SecondaryViewInfo.ViewFrustum);
 					OutViews.Add(SecondaryPackedView);
 					PrimaryNaniteViews.Add(SecondaryPackedView);
 				}
@@ -1360,7 +1361,8 @@ void FDeferredShadingSceneRenderer::RenderNanite(FRDGBuilder& GraphBuilder, cons
 				MaxPixelsPerEdgeMultipler *= 1.0f / DynamicResolutionFractions[GDynamicNaniteScalingPrimary];
 			}
 
-			Nanite::FPackedViewArray* NaniteViewsToRender = CreateNaniteViews(View, ViewIndex, RasterTextureSize, MaxPixelsPerEdgeMultipler);
+			TArray<FConvexVolume> ViewsToRenderCullingVolumes;
+			Nanite::FPackedViewArray* NaniteViewsToRender = CreateNaniteViews(View, ViewIndex, RasterTextureSize, MaxPixelsPerEdgeMultipler, ViewsToRenderCullingVolumes);
 
 			TUniquePtr< Nanite::IRenderer > NaniteRenderer;
 
@@ -1383,7 +1385,7 @@ void FDeferredShadingSceneRenderer::RenderNanite(FRDGBuilder& GraphBuilder, cons
 					!bIsEarlyDepthComplete ? View.PrevViewInfo.NaniteHZB : View.PrevViewInfo.HZB
 				);
 
-				FSceneInstanceCullingQuery *SceneInstanceCullQuery = SceneCullingRenderer.CullInstances(GraphBuilder, View.ViewFrustum);
+				FSceneInstanceCullingQuery *SceneInstanceCullQuery = SceneCullingRenderer.CullInstances(GraphBuilder, ViewsToRenderCullingVolumes);
 				NaniteRenderer->DrawGeometry(
 					Scene->NaniteRasterPipelines[ENaniteMeshPass::BasePass],
 					RasterResults.VisibilityQuery,
