Fix deferred mobile renderer only rendering to the left eye with mobile multi view.

diff --git a/Engine/Source/Runtime/Renderer/Private/MobileShadingRenderer.cpp b/Engine/Source/Runtime/Renderer/Private/MobileShadingRenderer.cpp
--- a/Engine/Source/Runtime/Renderer/Private/MobileShadingRenderer.cpp
+++ b/Engine/Source/Runtime/Renderer/Private/MobileShadingRenderer.cpp
@@ -1824,7 +1824,13 @@ FRenderTargetBindingSlots FMobileSceneRenderer::InitRenderTargetBindings_Deferre
 	BasePassRenderTargets.SubpassHint = ESubpassHint::None;
 	BasePassRenderTargets.NumOcclusionQueries = 0u;
 	BasePassRenderTargets.ShadingRateTexture = nullptr;
-	BasePassRenderTargets.MultiViewCount = 0;
+	
+	static const auto CVarMobileMultiView = IConsoleManager::Get().FindTConsoleVariableDataInt(TEXT("vr.MobileMultiView"));
+	const bool bIsMultiViewApplication = (CVarMobileMultiView && CVarMobileMultiView->GetValueOnAnyThread() != 0);
+	
+	//if the scenecolor isn't multiview but the app is, need to render as a single-view multiview due to shaders
+	BasePassRenderTargets.MultiViewCount = Views[0].bIsMobileMultiViewEnabled ? 2 : (bIsMultiViewApplication ? 1 : 0);
+
 	return BasePassRenderTargets;
 }
 
diff --git a/Engine/Source/Runtime/Renderer/Private/SceneTextures.cpp b/Engine/Source/Runtime/Renderer/Private/SceneTextures.cpp
--- a/Engine/Source/Runtime/Renderer/Private/SceneTextures.cpp
+++ b/Engine/Source/Runtime/Renderer/Private/SceneTextures.cpp
@@ -550,31 +550,41 @@ void FSceneTextures::InitializeViewFamily(FRDGBuilder& GraphBuilder, FViewFamily
 
 		if (Bindings.GBufferA.Index >= 0)
 		{
-			const FRDGTextureDesc Desc(FRDGTextureDesc::Create2D(Config.Extent, Bindings.GBufferA.Format, FClearValueBinding::Transparent, Bindings.GBufferA.Flags | FlagsToAdd | GFastVRamConfig.GBufferA));
+			const FRDGTextureDesc Desc(Config.bRequireMultiView ?
+				FRDGTextureDesc::Create2DArray(Config.Extent, Bindings.GBufferA.Format, FClearValueBinding::Transparent, Bindings.GBufferA.Flags | FlagsToAdd | GFastVRamConfig.GBufferA, 2) :
+				FRDGTextureDesc::Create2D(Config.Extent, Bindings.GBufferA.Format, FClearValueBinding::Transparent, Bindings.GBufferA.Flags | FlagsToAdd | GFastVRamConfig.GBufferA));
 			SceneTextures.GBufferA = GraphBuilder.CreateTexture(Desc, TEXT("GBufferA"));
 		}
 
 		if (Bindings.GBufferB.Index >= 0)
 		{
-			const FRDGTextureDesc Desc(FRDGTextureDesc::Create2D(Config.Extent, Bindings.GBufferB.Format, FClearValueBinding::Transparent, Bindings.GBufferB.Flags | FlagsToAdd | GFastVRamConfig.GBufferB));
+			const FRDGTextureDesc Desc(Config.bRequireMultiView ?
+				FRDGTextureDesc::Create2DArray(Config.Extent, Bindings.GBufferB.Format, FClearValueBinding::Transparent, Bindings.GBufferB.Flags | FlagsToAdd | GFastVRamConfig.GBufferB, 2) :
+				FRDGTextureDesc::Create2D(Config.Extent, Bindings.GBufferB.Format, FClearValueBinding::Transparent, Bindings.GBufferB.Flags | FlagsToAdd | GFastVRamConfig.GBufferB));
 			SceneTextures.GBufferB = GraphBuilder.CreateTexture(Desc, TEXT("GBufferB"));
 		}
 
 		if (Bindings.GBufferC.Index >= 0)
 		{
-			const FRDGTextureDesc Desc(FRDGTextureDesc::Create2D(Config.Extent, Bindings.GBufferC.Format, FClearValueBinding::Transparent, Bindings.GBufferC.Flags | FlagsToAdd | GFastVRamConfig.GBufferC));
+			const FRDGTextureDesc Desc(Config.bRequireMultiView ?
+				FRDGTextureDesc::Create2DArray(Config.Extent, Bindings.GBufferC.Format, FClearValueBinding::Transparent, Bindings.GBufferC.Flags | FlagsToAdd | GFastVRamConfig.GBufferC, 2) :
+				FRDGTextureDesc::Create2D(Config.Extent, Bindings.GBufferC.Format, FClearValueBinding::Transparent, Bindings.GBufferC.Flags | FlagsToAdd | GFastVRamConfig.GBufferC));
 			SceneTextures.GBufferC = GraphBuilder.CreateTexture(Desc, TEXT("GBufferC"));
 		}
 
 		if (Bindings.GBufferD.Index >= 0)
 		{
-			const FRDGTextureDesc Desc(FRDGTextureDesc::Create2D(Config.Extent, Bindings.GBufferD.Format, FClearValueBinding::Transparent, Bindings.GBufferD.Flags | FlagsToAdd | GFastVRamConfig.GBufferD));
+			const FRDGTextureDesc Desc(Config.bRequireMultiView ?
+				FRDGTextureDesc::Create2DArray(Config.Extent, Bindings.GBufferD.Format, FClearValueBinding::Transparent, Bindings.GBufferD.Flags | FlagsToAdd | GFastVRamConfig.GBufferD, 2) :
+				FRDGTextureDesc::Create2D(Config.Extent, Bindings.GBufferD.Format, FClearValueBinding::Transparent, Bindings.GBufferD.Flags | FlagsToAdd | GFastVRamConfig.GBufferD));
 			SceneTextures.GBufferD = GraphBuilder.CreateTexture(Desc, TEXT("GBufferD"));
 		}
 
 		if (Bindings.GBufferE.Index >= 0)
 		{
-			const FRDGTextureDesc Desc(FRDGTextureDesc::Create2D(Config.Extent, Bindings.GBufferE.Format, FClearValueBinding::Transparent, Bindings.GBufferE.Flags | FlagsToAdd | GFastVRamConfig.GBufferE));
+			const FRDGTextureDesc Desc(Config.bRequireMultiView ?
+				FRDGTextureDesc::Create2DArray(Config.Extent, Bindings.GBufferE.Format, FClearValueBinding::Transparent, Bindings.GBufferE.Flags | FlagsToAdd | GFastVRamConfig.GBufferE, 2) :
+				FRDGTextureDesc::Create2D(Config.Extent, Bindings.GBufferE.Format, FClearValueBinding::Transparent, Bindings.GBufferE.Flags | FlagsToAdd | GFastVRamConfig.GBufferE));
 			SceneTextures.GBufferE = GraphBuilder.CreateTexture(Desc, TEXT("GBufferE"));
 		}
 
@@ -583,7 +593,9 @@ void FSceneTextures::InitializeViewFamily(FRDGBuilder& GraphBuilder, FViewFamily
 		{
 			ETextureCreateFlags GBufferFCreateFlags;
 			EPixelFormat GBufferFPixelFormat = GetGBufferFFormatAndCreateFlags(GBufferFCreateFlags);
-			const FRDGTextureDesc Desc = FRDGTextureDesc::Create2D(Config.Extent, GBufferFPixelFormat, FClearValueBinding({ 0.5f, 0.5f, 0.5f, 0.5f }), GBufferFCreateFlags | FlagsToAdd);
+			const FRDGTextureDesc Desc(Config.bRequireMultiView ?
+				FRDGTextureDesc::Create2DArray(Config.Extent, GBufferFPixelFormat, FClearValueBinding({ 0.5f, 0.5f, 0.5f, 0.5f }), GBufferFCreateFlags | FlagsToAdd, 2) :
+				FRDGTextureDesc::Create2D(Config.Extent, GBufferFPixelFormat, FClearValueBinding({ 0.5f, 0.5f, 0.5f, 0.5f }), GBufferFCreateFlags | FlagsToAdd));
 			SceneTextures.GBufferF = GraphBuilder.CreateTexture(Desc, TEXT("GBufferF"));
 		}
 	}
