AlphaHOldout - Make sure the component flag is initialise to 0 in any case.

diff --git a/Engine/Source/Runtime/Renderer/Private/SceneRendering.cpp b/Engine/Source/Runtime/Renderer/Private/SceneRendering.cpp
--- a/Engine/Source/Runtime/Renderer/Private/SceneRendering.cpp
+++ b/Engine/Source/Runtime/Renderer/Private/SceneRendering.cpp
@@ -1454,11 +1454,10 @@ void FViewInfo::SetupUniformBufferParameters(
 	ViewUniformShaderParameters.AtmosphereIrradianceTextureSampler = TStaticSamplerState<SF_Bilinear>::GetRHI();
 	ViewUniformShaderParameters.AtmosphereInscatterTextureSampler = TStaticSamplerState<SF_Bilinear>::GetRHI();
 
-	// Upload environment holdout flags
+	// Update environment holdout flags
+	ViewUniformShaderParameters.EnvironmentComponentsFlags = 0;
 	if (Scene)
 	{
-		ViewUniformShaderParameters.EnvironmentComponentsFlags = 0;
-
 		if (ShouldRenderSkyAtmosphere(Scene, Family->EngineShowFlags))
 		{
 			FSkyAtmosphereRenderSceneInfo* SkyAtmosphere = Scene->SkyAtmosphere;
