Added speculative crash fix for out of bounds instance buffer data access.

#rnx

diff --git a/Engine/Source/Runtime/Renderer/Private/SceneCulling/SceneCulling.cpp b/Engine/Source/Runtime/Renderer/Private/SceneCulling/SceneCulling.cpp
--- a/Engine/Source/Runtime/Renderer/Private/SceneCulling/SceneCulling.cpp
+++ b/Engine/Source/Runtime/Renderer/Private/SceneCulling/SceneCulling.cpp
@@ -825,20 +825,50 @@ SC_FORCEINLINE FVector4d TransformBounds(VectorRegister4f VecOrigin, VectorRegis
 	return Result;
 }
 
-struct FBoundsTransformerUniqueBounds
+struct FBoundsTransformerBase
 {
-	SC_FORCEINLINE FBoundsTransformerUniqueBounds(const FInstanceSceneDataBuffers& InInstanceSceneDataBuffers)
+	SC_FORCEINLINE FBoundsTransformerBase(const FInstanceSceneDataBuffers& InInstanceSceneDataBuffers)
 		: InstanceSceneDataBuffers(InInstanceSceneDataBuffers)
 	{
 		// Note: for reasons unknown VectorLoadFloat3 also does doubles...
 		PrimitiveToWorldTranslationVec = VectorLoadFloat3(&InstanceSceneDataBuffers.GetPrimitiveWorldSpaceOffset());
+
+
+		FInstanceSceneDataBuffers::FReadView InstanceDataView = InstanceSceneDataBuffers.GetReadView();
+		if (InstanceDataView.InstanceToPrimitiveRelative.IsEmpty())
+		{
+			// Set up with dummy data array
+			InstanceToPrimitiveRelativeArray = TConstArrayView<FRenderTransform>(&FRenderTransform::Identity, 1);
+		}
+		else
+		{
+			InstanceToPrimitiveRelativeArray = InstanceDataView.InstanceToPrimitiveRelative;
+		}
+	}
+
+	// returns the clamped instance transform, to cover for OOB accesses
+	SC_FORCEINLINE FRenderTransform GetInstanceToPrimitiveRelative(int32 InstanceIndex)
+	{
+		return InstanceToPrimitiveRelativeArray[FMath::Min(InstanceToPrimitiveRelativeArray.Num() - 1, InstanceIndex)];
+	}
+	VectorRegister4Double PrimitiveToWorldTranslationVec;
+
+	const FInstanceSceneDataBuffers &InstanceSceneDataBuffers;
+	TConstArrayView<FRenderTransform> InstanceToPrimitiveRelativeArray;
+};
+
+struct FBoundsTransformerUniqueBounds : public FBoundsTransformerBase
+{
+	SC_FORCEINLINE FBoundsTransformerUniqueBounds(const FInstanceSceneDataBuffers& InInstanceSceneDataBuffers)
+		: FBoundsTransformerBase(InInstanceSceneDataBuffers)
+	{
 	}
 
 	SC_FORCEINLINE FVector4d TransformBounds(int32 InstanceIndex)
 	{
 		const FRenderBounds InstanceBounds = InstanceSceneDataBuffers.GetInstanceLocalBounds(InstanceIndex);
 
-		FRenderTransform InstanceToPrimitiveRelative = InstanceSceneDataBuffers.GetInstanceToPrimitiveRelative(InstanceIndex);
+		FRenderTransform InstanceToPrimitiveRelative = GetInstanceToPrimitiveRelative(InstanceIndex);
 		const VectorRegister4f VecMin = VectorLoadFloat3(&InstanceBounds.Min);
 		const VectorRegister4f VecMax = VectorLoadFloat3(&InstanceBounds.Max);
 		const VectorRegister4f Half = VectorSetFloat1(0.5f); // VectorSetFloat1() can be faster than SetFloat3(0.5, 0.5, 0.5, 0.0). Okay if 4th element is 0.5, it's multiplied by 0.0 below and we discard W anyway.
@@ -847,18 +877,13 @@ struct FBoundsTransformerUniqueBounds
 
 		return ::TransformBounds(VecOrigin, VecExtent, InstanceToPrimitiveRelative, PrimitiveToWorldTranslationVec);
 	}
-
-	VectorRegister4Double PrimitiveToWorldTranslationVec;
-	const FInstanceSceneDataBuffers &InstanceSceneDataBuffers;
 };
 
-struct FBoundsTransformerSharedBounds
+struct FBoundsTransformerSharedBounds : public FBoundsTransformerBase
 {
 	SC_FORCEINLINE FBoundsTransformerSharedBounds(const FInstanceSceneDataBuffers& InInstanceSceneDataBuffers)
-		: InstanceSceneDataBuffers(InInstanceSceneDataBuffers)
+		: FBoundsTransformerBase(InInstanceSceneDataBuffers)
 	{
-		// Note: for reasons unknown VectorLoadFloat3 also does doubles...
-		PrimitiveToWorldTranslationVec = VectorLoadFloat3(&InstanceSceneDataBuffers.GetPrimitiveWorldSpaceOffset());
 
 		const FRenderBounds InstanceBounds = InstanceSceneDataBuffers.GetInstanceLocalBounds(0);
 		const VectorRegister4f VecMin = VectorLoadFloat3(&InstanceBounds.Min);
@@ -870,12 +895,10 @@ struct FBoundsTransformerSharedBounds
 
 	SC_FORCEINLINE FVector4d TransformBounds(int32 InstanceIndex)
 	{
-		FRenderTransform InstanceToPrimitiveRelative = InstanceSceneDataBuffers.GetInstanceToPrimitiveRelative(InstanceIndex);
+		FRenderTransform InstanceToPrimitiveRelative = GetInstanceToPrimitiveRelative(InstanceIndex);
 		return ::TransformBounds(VecOrigin, VecExtent, InstanceToPrimitiveRelative, PrimitiveToWorldTranslationVec);
 	}
 
-	const FInstanceSceneDataBuffers &InstanceSceneDataBuffers;
-	VectorRegister4Double PrimitiveToWorldTranslationVec;
 	VectorRegister4f VecOrigin;
 	VectorRegister4f VecExtent;
 };
