Fix overly aggressive static analysis warning about code that will never run (by design WITH_EDITOR)

diff --git a/Engine/Source/Runtime/Renderer/Private/RendererScene.cpp b/Engine/Source/Runtime/Renderer/Private/RendererScene.cpp
--- a/Engine/Source/Runtime/Renderer/Private/RendererScene.cpp
+++ b/Engine/Source/Runtime/Renderer/Private/RendererScene.cpp
@@ -6324,15 +6324,10 @@ void FScene::Update(FRDGBuilder& GraphBuilder, const FUpdateParameters& Paramete
 	// Determine the test visible vs. always visible primitive index ranges
 	PrimitivesAlwaysVisibleOffset = ~0u;
 
-#if WITH_EDITOR
-	// TODO: Support skip always visible in the editor (need to handle dynamic relevance)
-	const bool bSkipAlwaysVisible = false;
-#else
-	const bool bSkipAlwaysVisible = GVisibilitySkipAlwaysVisible != 0;
-#endif
-
+// TODO: Support skip always visible in the editor (need to handle dynamic relevance)
+#if !WITH_EDITOR
 	// This optimization requires compute materials due to relevancy calculation
-	if (bSkipAlwaysVisible && UseNaniteComputeMaterials())
+	if (GVisibilitySkipAlwaysVisible != 0 && UseNaniteComputeMaterials())
 	{
 		uint32 NextTypeOffset = 0;
 		for (int32 TypeOffsetIndex = 0; TypeOffsetIndex < TypeOffsetTable.Num(); ++TypeOffsetIndex)
@@ -6386,6 +6381,7 @@ void FScene::Update(FRDGBuilder& GraphBuilder, const FUpdateParameters& Paramete
 			PrimitivesAlwaysVisibleOffset = ~0u;
 		}
 	}
+#endif // !WITH_EDITOR
 
 	// Allocate all instance slots. Needs to happen after the instance data is updated since that may change the counts.
 	FPrimitiveSceneInfo::AllocateGPUSceneInstances(this, PendingAllocateInstanceIds);
