Many Lights - fixed missing temporal accumulation on the last row and column of pixels

diff --git a/Engine/Shaders/Private/ManyLights/ManyLightsDenoiserTemporal.usf b/Engine/Shaders/Private/ManyLights/ManyLightsDenoiserTemporal.usf
--- a/Engine/Shaders/Private/ManyLights/ManyLightsDenoiserTemporal.usf
+++ b/Engine/Shaders/Private/ManyLights/ManyLightsDenoiserTemporal.usf
@@ -14,6 +14,7 @@ Texture2D<UNORM float> NumFramesAccumulatedHistoryTexture;
 
 float4 HistoryScreenPositionScaleBias;
 float4 HistoryUVMinMax;
+float4 HistoryGatherUVMinMax;
 float PrevSceneColorPreExposureCorrection;
 float TemporalMaxFramesAccumulated;
 float TemporalNeighborhoodClampScale;
@@ -25,7 +26,6 @@ RWTexture2D<float4> RWDiffuseLightingAndSecondMoment;
 RWTexture2D<float4> RWSpecularLightingAndSecondMoment;
 RWTexture2D<UNORM float> RWNumFramesAccumulated;
 RWTexture2D<float> RWSceneDepth;
-RWTexture2D<float4> RWSceneColor;
 
 struct FNeighborhood
 {
@@ -142,9 +142,11 @@ void DenoiserTemporalCS(
 				float3 HistoryScreenPosition = GetHistoryScreenPosition(ScreenPosition, ScreenUV, DeviceZ);
 				float2 HistoryScreenUV = HistoryScreenPosition.xy * HistoryScreenPositionScaleBias.xy + HistoryScreenPositionScaleBias.wz;
 
-				bool bHistoryWasOnScreen = all(HistoryScreenUV >= HistoryUVMinMax.xy) && all (HistoryScreenUV <= HistoryUVMinMax.zw);
+				bool bHistoryWasOnScreen = all(HistoryScreenUV >= HistoryUVMinMax.xy) && all(HistoryScreenUV <= HistoryUVMinMax.zw);
 				if (bHistoryWasOnScreen)
 				{
+					HistoryScreenUV = clamp(HistoryScreenUV, HistoryGatherUVMinMax.xy + View.GeneralPurposeTweak * View.BufferSizeAndInvSize.zw, HistoryGatherUVMinMax.zw - View.GeneralPurposeTweak * View.BufferSizeAndInvSize.zw);
+
 					uint2 HistoryScreenCoord = floor(HistoryScreenUV * View.BufferSizeAndInvSize.xy - 0.5f);
 					float2 HistoryBilinearWeights = frac(HistoryScreenUV * View.BufferSizeAndInvSize.xy - 0.5f);
 					float2 HistoryGatherUV = (HistoryScreenCoord + 1.0f) * View.BufferSizeAndInvSize.zw;
diff --git a/Engine/Source/Runtime/Renderer/Private/ManyLights/ManyLights.cpp b/Engine/Source/Runtime/Renderer/Private/ManyLights/ManyLights.cpp
--- a/Engine/Source/Runtime/Renderer/Private/ManyLights/ManyLights.cpp
+++ b/Engine/Source/Runtime/Renderer/Private/ManyLights/ManyLights.cpp
@@ -48,7 +48,7 @@ static TAutoConsoleVariable<int32> CVarManyLightsTemporal(
 
 static TAutoConsoleVariable<int32> CVarManyLightsTemporalMaxFramesAccumulated(
 	TEXT("r.ManyLights.Temporal.MaxFramesAccumulated"),
-	8,
+	12,
 	TEXT("Max history length when accumulating frames. Lower values have less ghosting, but more noise."),
 	ECVF_Scalability | ECVF_RenderThreadSafe
 );
@@ -341,8 +341,6 @@ class FGenerateLightSamplesCS : public FGlobalShader
 		SHADER_PARAMETER_RDG_TEXTURE_UAV(RWTexture2D<uint>, RWLightSamples)
 		SHADER_PARAMETER_RDG_BUFFER_SRV(StructuredBuffer<uint>, DownsampledTileAllocator)
 		SHADER_PARAMETER_RDG_BUFFER_SRV(StructuredBuffer<uint>, DownsampledTileData)
-		SHADER_PARAMETER(FVector4f, HistoryUVMinMax)
-		SHADER_PARAMETER(FVector4f, HistoryScreenPositionScaleBias)
 	END_SHADER_PARAMETER_STRUCT()
 
 	class FTileType : SHADER_PERMUTATION_INT("TILE_TYPE", (int32)ManyLights::ETileType::SHADING_MAX);
@@ -350,7 +348,6 @@ class FGenerateLightSamplesCS : public FGlobalShader
 	class FLightFunctionAtlas : SHADER_PERMUTATION_BOOL("USE_LIGHT_FUNCTION_ATLAS");
 	class FTexturedRectLights : SHADER_PERMUTATION_BOOL("USE_SOURCE_TEXTURE");
 	class FNumSamplesPerPixel1d : SHADER_PERMUTATION_SPARSE_INT("NUM_SAMPLES_PER_PIXEL_1D", 1, 2, 4);
-
 	class FDebugMode : SHADER_PERMUTATION_BOOL("DEBUG_MODE");
 	using FPermutationDomain = TShaderPermutationDomain<FTileType, FIESProfile, FLightFunctionAtlas, FTexturedRectLights, FNumSamplesPerPixel1d, FDebugMode>;
 
@@ -544,14 +541,14 @@ class FDenoiserTemporalCS : public FGlobalShader
 		SHADER_PARAMETER_RDG_TEXTURE(Texture2D<float4>, SpecularLightingAndSecondMomentHistoryTexture)
 		SHADER_PARAMETER_RDG_TEXTURE(Texture2D<UNORM float>, NumFramesAccumulatedHistoryTexture)
 		SHADER_PARAMETER_RDG_TEXTURE(Texture2D<float>, ManyLightsDepthHistory)
-		SHADER_PARAMETER(FVector4f, HistoryUVMinMax)
 		SHADER_PARAMETER(FVector4f, HistoryScreenPositionScaleBias)
+		SHADER_PARAMETER(FVector4f, HistoryUVMinMax)
+		SHADER_PARAMETER(FVector4f, HistoryGatherUVMinMax)
 		SHADER_PARAMETER(float, PrevSceneColorPreExposureCorrection)
 		SHADER_PARAMETER_RDG_TEXTURE_UAV(RWTexture2D<float4>, RWDiffuseLightingAndSecondMoment)
 		SHADER_PARAMETER_RDG_TEXTURE_UAV(RWTexture2D<float4>, RWSpecularLightingAndSecondMoment)
 		SHADER_PARAMETER_RDG_TEXTURE_UAV(RWTexture2D<UNORM float>, RWNumFramesAccumulated)
 		SHADER_PARAMETER_RDG_TEXTURE_UAV(RWTexture2D<float>, RWSceneDepth)
-		SHADER_PARAMETER_RDG_TEXTURE_UAV(RWTexture2D<float4>, RWSceneColor)
 	END_SHADER_PARAMETER_STRUCT()
 
 	class FValidHistory : SHADER_PERMUTATION_BOOL("VALID_HISTORY");
@@ -691,6 +688,7 @@ void FDeferredShadingSceneRenderer::RenderManyLights(FRDGBuilder& GraphBuilder,
 		bool bTemporal = CVarManyLightsTemporal.GetValueOnRenderThread() != 0;
 		FVector4f HistoryScreenPositionScaleBias = FVector4f(0.0f, 0.0f, 0.0f, 0.0f);
 		FVector4f HistoryUVMinMax = FVector4f(0.0f, 0.0f, 0.0f, 0.0f);
+		FVector4f HistoryGatherUVMinMax = FVector4f(0.0f, 0.0f, 0.0f, 0.0f);
 		FRDGTextureRef DiffuseLightingAndSecondMomentHistory = nullptr;
 		FRDGTextureRef SpecularLightingAndSecondMomentHistory = nullptr;
 		FRDGTextureRef SceneDepthHistory = nullptr;
@@ -704,6 +702,7 @@ void FDeferredShadingSceneRenderer::RenderManyLights(FRDGBuilder& GraphBuilder,
 			{
 				HistoryScreenPositionScaleBias = LightingViewState.HistoryScreenPositionScaleBias;
 				HistoryUVMinMax = LightingViewState.HistoryUVMinMax;
+				HistoryGatherUVMinMax = LightingViewState.HistoryGatherUVMinMax;
 
 				if (LightingViewState.DiffuseLightingAndSecondMomentHistory
 					&& LightingViewState.SpecularLightingAndSecondMomentHistory
@@ -925,8 +924,6 @@ void FDeferredShadingSceneRenderer::RenderManyLights(FRDGBuilder& GraphBuilder,
 				PassParameters->RWCompositeTileData = CompositeTileDataUAV;
 				PassParameters->DownsampledTileAllocator = GraphBuilder.CreateSRV(DownsampledTileAllocator);
 				PassParameters->DownsampledTileData = GraphBuilder.CreateSRV(DownsampledTileData);
-				PassParameters->HistoryScreenPositionScaleBias = HistoryScreenPositionScaleBias;
-				PassParameters->HistoryUVMinMax = HistoryUVMinMax;
 
 				FGenerateLightSamplesCS::FPermutationDomain PermutationVector;
 				PermutationVector.Set<FGenerateLightSamplesCS::FTileType>(TileType);
@@ -1123,11 +1120,11 @@ void FDeferredShadingSceneRenderer::RenderManyLights(FRDGBuilder& GraphBuilder,
 			PassParameters->PrevSceneColorPreExposureCorrection = View.PreExposure / View.PrevViewInfo.SceneColorPreExposure;
 			PassParameters->HistoryScreenPositionScaleBias = HistoryScreenPositionScaleBias;
 			PassParameters->HistoryUVMinMax = HistoryUVMinMax;
+			PassParameters->HistoryGatherUVMinMax = HistoryGatherUVMinMax;
 			PassParameters->RWDiffuseLightingAndSecondMoment = GraphBuilder.CreateUAV(DiffuseLightingAndSecondMoment);
 			PassParameters->RWSpecularLightingAndSecondMoment = GraphBuilder.CreateUAV(SpecularLightingAndSecondMoment);
 			PassParameters->RWNumFramesAccumulated = GraphBuilder.CreateUAV(NumFramesAccumulated);
 			PassParameters->RWSceneDepth = GraphBuilder.CreateUAV(SceneDepthCopy);
-			PassParameters->RWSceneColor = GraphBuilder.CreateUAV(SceneTextures.Color.Target);
 
 			FDenoiserTemporalCS::FPermutationDomain PermutationVector;
 			PermutationVector.Set<FDenoiserTemporalCS::FValidHistory>(DiffuseLightingAndSecondMomentHistory != nullptr && bTemporal);
@@ -1175,13 +1172,20 @@ void FDeferredShadingSceneRenderer::RenderManyLights(FRDGBuilder& GraphBuilder,
 
 			LightingViewState.HistoryScreenPositionScaleBias = View.GetScreenPositionScaleBias(View.GetSceneTexturesConfig().Extent, View.ViewRect);
 
-			// Pull in the max UV to exclude the region which will read outside the viewport due to bilinear filtering
-			const FVector2D InvBufferSize(1.0f / SceneTextures.Config.Extent.X, 1.0f / SceneTextures.Config.Extent.Y);
+			const FVector2f InvBufferSize(1.0f / SceneTextures.Config.Extent.X, 1.0f / SceneTextures.Config.Extent.Y);
+
 			LightingViewState.HistoryUVMinMax = FVector4f(
-				(View.ViewRect.Min.X + 0.5f) * InvBufferSize.X,
-				(View.ViewRect.Min.Y + 0.5f) * InvBufferSize.Y,
-				(View.ViewRect.Max.X - 1.0f) * InvBufferSize.X,
-				(View.ViewRect.Max.Y - 1.0f) * InvBufferSize.Y);
+				View.ViewRect.Min.X * InvBufferSize.X,
+				View.ViewRect.Min.Y * InvBufferSize.Y,
+				View.ViewRect.Max.X * InvBufferSize.X,
+				View.ViewRect.Max.Y * InvBufferSize.Y);
+
+			// Clamp gather4 to a valid bilinear footprint in order to avoid sampling outside of valid bounds
+			LightingViewState.HistoryGatherUVMinMax = FVector4f(
+				(View.ViewRect.Min.X + 0.51f) * InvBufferSize.X,
+				(View.ViewRect.Min.Y + 0.51f) * InvBufferSize.Y,
+				(View.ViewRect.Max.X - 0.51f) * InvBufferSize.X,
+				(View.ViewRect.Max.Y - 0.51f) * InvBufferSize.Y);
 
 			if (DiffuseLightingAndSecondMoment && SpecularLightingAndSecondMoment && SceneDepthCopy && NumFramesAccumulated && bTemporal)
 			{
diff --git a/Engine/Source/Runtime/Renderer/Private/ManyLights/ManyLightsViewState.h b/Engine/Source/Runtime/Renderer/Private/ManyLights/ManyLightsViewState.h
--- a/Engine/Source/Runtime/Renderer/Private/ManyLights/ManyLightsViewState.h
+++ b/Engine/Source/Runtime/Renderer/Private/ManyLights/ManyLightsViewState.h
@@ -15,6 +15,7 @@ public:
 
 	FVector4f HistoryScreenPositionScaleBias = FVector4f(0.0f, 0.0f, 0.0f, 0.0f);
 	FVector4f HistoryUVMinMax = FVector4f(0.0f, 0.0f, 0.0f, 0.0f);
+	FVector4f HistoryGatherUVMinMax = FVector4f(0.0f, 0.0f, 0.0f, 0.0f);
 
 	void SafeRelease()
 	{
