Fix skeletal meshes not getting swapped to VSM dynamic caching in some cases

diff --git a/Engine/Source/Runtime/Renderer/Private/GPUScene.cpp b/Engine/Source/Runtime/Renderer/Private/GPUScene.cpp
--- a/Engine/Source/Runtime/Renderer/Private/GPUScene.cpp
+++ b/Engine/Source/Runtime/Renderer/Private/GPUScene.cpp
@@ -1610,7 +1610,8 @@ void FGPUScene::UploadDynamicPrimitiveShaderDataForViewInternal(FRDGBuilder& Gra
 			for (const FGPUScenePrimitiveCollector::FPrimitiveData& PrimitiveData : Collector.UploadData->PrimitiveData)
 			{
 				check(PrimitiveData.LocalInstanceSceneDataOffset != INDEX_NONE);
-				ShadowInvalidatingInstances->AddInstanceRange(PrimitiveData.LocalInstanceSceneDataOffset + uint32(InstanceIdStart), PrimitiveData.NumInstances);
+				ShadowInvalidatingInstances->AddInstanceRange(FPersistentPrimitiveIndex(),
+					PrimitiveData.LocalInstanceSceneDataOffset + uint32(InstanceIdStart), PrimitiveData.NumInstances);
 			}
 		}
 
diff --git a/Engine/Source/Runtime/Renderer/Private/VirtualShadowMaps/VirtualShadowMapCacheManager.cpp b/Engine/Source/Runtime/Renderer/Private/VirtualShadowMaps/VirtualShadowMapCacheManager.cpp
--- a/Engine/Source/Runtime/Renderer/Private/VirtualShadowMaps/VirtualShadowMapCacheManager.cpp
+++ b/Engine/Source/Runtime/Renderer/Private/VirtualShadowMaps/VirtualShadowMapCacheManager.cpp
@@ -295,13 +295,16 @@ bool FVirtualShadowMapPerLightCacheEntry::UpdateLocal(const FProjectedShadowInit
 
 void FVirtualShadowMapArrayCacheManager::FShadowInvalidatingInstancesImplementation::AddPrimitive(const FPrimitiveSceneInfo* PrimitiveSceneInfo)
 {
-	AddInstanceRange(PrimitiveSceneInfo->GetInstanceSceneDataOffset(), PrimitiveSceneInfo->GetNumInstanceSceneDataEntries());
+	AddInstanceRange(PrimitiveSceneInfo->GetPersistentIndex(), PrimitiveSceneInfo->GetInstanceSceneDataOffset(), PrimitiveSceneInfo->GetNumInstanceSceneDataEntries());
 }
 
-void FVirtualShadowMapArrayCacheManager::FShadowInvalidatingInstancesImplementation::AddInstanceRange(uint32 InstanceSceneDataOffset, uint32 NumInstanceSceneDataEntries)
+void FVirtualShadowMapArrayCacheManager::FShadowInvalidatingInstancesImplementation::AddInstanceRange(FPersistentPrimitiveIndex PersistentPrimitiveIndex, uint32 InstanceSceneDataOffset, uint32 NumInstanceSceneDataEntries)
 {
-	PrimitiveInstancesToInvalidate.Add(FVirtualShadowMapInstanceRange{int32(InstanceSceneDataOffset), int32(NumInstanceSceneDataEntries), false});
-}		
+	PrimitiveInstancesToInvalidate.Add(FVirtualShadowMapInstanceRange{
+		PersistentPrimitiveIndex,
+		int32(InstanceSceneDataOffset),
+		int32(NumInstanceSceneDataEntries)});
+}
 
 void FVirtualShadowMapPerLightCacheEntry::Invalidate()
 {
@@ -334,16 +337,26 @@ void FVirtualShadowMapArrayCacheManager::FInvalidatingPrimitiveCollector::AddPri
 		for (const auto& SmCacheEntry : CacheEntry.Value->ShadowMapEntries)
 		{
 			const uint32 Payload = EncodeInstanceInvalidationPayload(SmCacheEntry.CurrentVirtualShadowMapId);
-			const uint32 PayloadForceStatic = EncodeInstanceInvalidationPayload(SmCacheEntry.CurrentVirtualShadowMapId, VSM_INVALIDATION_PAYLOAD_FLAG_FORCE_STATIC);
 
+			// Global invalidations
 			for (const FVirtualShadowMapInstanceRange& Range : Manager.ShadowInvalidatingInstancesImplementation.PrimitiveInstancesToInvalidate)
 			{
-				Instances.Add(Range.InstanceSceneDataOffset, Range.NumInstanceSceneDataEntries, Range.bForceInvalidateStatic ? PayloadForceStatic : Payload);
+				Instances.Add(Range.InstanceSceneDataOffset, Range.NumInstanceSceneDataEntries, Payload);
+				if (Range.PersistentPrimitiveIndex.IsValid())
+				{
+					InvalidatedPrimitives[Range.PersistentPrimitiveIndex.Index] = true;
+				}
 			}
 
+			// Per-light invalidations
 			for (const FVirtualShadowMapInstanceRange& Range : CacheEntry.Value->PrimitiveInstancesToInvalidate)
 			{
-				Instances.Add(Range.InstanceSceneDataOffset, Range.NumInstanceSceneDataEntries, Range.bForceInvalidateStatic ? PayloadForceStatic : Payload);
+				Instances.Add(Range.InstanceSceneDataOffset, Range.NumInstanceSceneDataEntries, Payload);
+				check(Range.PersistentPrimitiveIndex.IsValid());		// Should always be valid currently in this path
+				if (Range.PersistentPrimitiveIndex.IsValid())
+				{
+					InvalidatedPrimitives[Range.PersistentPrimitiveIndex.Index] = true;
+				}
 			}
 		}
 
@@ -850,10 +863,10 @@ void FVirtualShadowMapPerLightCacheEntry::OnPrimitiveRendered(const FPrimitiveSc
 		// Skip if the invalidation mode is NOT auto (because Always will do it elsewhere & the others should prevent this).
 		if (PrimitiveSceneInfo->Proxy->HasDeformableMesh() && PrimitiveSceneInfo->Proxy->GetShadowCacheInvalidationBehavior() == EShadowCacheInvalidationBehavior::Auto)
 		{
-			PrimitiveInstancesToInvalidate.Add(FVirtualShadowMapInstanceRange{ 
+			PrimitiveInstancesToInvalidate.Add(FVirtualShadowMapInstanceRange{
+				PrimitiveSceneInfo->GetPersistentIndex(),
 				PrimitiveSceneInfo->GetInstanceSceneDataOffset(),
-				PrimitiveSceneInfo->GetNumInstanceSceneDataEntries(),
-				false
+				PrimitiveSceneInfo->GetNumInstanceSceneDataEntries()
 			});
 		}
 	}
@@ -1317,17 +1330,17 @@ void FVirtualShadowMapArrayCacheManager::UpdateCachePrimitiveAsDynamic(FInvalida
 			FPrimitiveSceneInfo* PrimitiveSceneInfo = Scene->GetPrimitiveSceneInfo(WrappedIndex);
 			if (PrimitiveSceneInfo)
 			{
-				FVirtualShadowMapInstanceRange InstanceRange{
-					PrimitiveSceneInfo->GetInstanceSceneDataOffset(),
-					PrimitiveSceneInfo->GetNumInstanceSceneDataEntries(),
-					true		// Always invalidate static since we are swapping between the two
-				};
-
-				// Add it to the invalidation list for every light
-				// NOTE: These will get added to the invalidation list by AddDynamicAndGPUPrimitives
+				// Add an invalidation for every light
 				for (auto& CacheEntry : CacheEntries)
 				{
-					CacheEntry.Value->PrimitiveInstancesToInvalidate.Add(InstanceRange);
+					for (const auto& SmCacheEntry : CacheEntry.Value->ShadowMapEntries)
+					{
+						const uint32 PayloadForceStatic = EncodeInstanceInvalidationPayload(SmCacheEntry.CurrentVirtualShadowMapId, VSM_INVALIDATION_PAYLOAD_FLAG_FORCE_STATIC);
+						InvalidatingPrimitiveCollector.Instances.Add(
+							PrimitiveSceneInfo->GetInstanceSceneDataOffset(),
+							PrimitiveSceneInfo->GetNumInstanceSceneDataEntries(),
+							PayloadForceStatic);
+					}
 				}
 			}
 			else
@@ -1355,9 +1368,12 @@ void FVirtualShadowMapArrayCacheManager::ProcessInvalidations(FRDGBuilder& Graph
 		// we update them as the shader needs to know the previous cache states for invalidation.
 		FInvalidationPassCommon InvalidationPassCommon = GetUniformParametersForInvalidation(GraphBuilder, SceneUniformBuffer);
 
+		// Add invalidations for skeletal meshes, CPU culling changes, dynamic primitives, etc.
+		InvalidatingPrimitiveCollector.AddPrimitivesToInvalidate();
+
 		// Check whether we want to swap any cache states and add any invalidations to that end as well
 		UpdateCachePrimitiveAsDynamic(InvalidatingPrimitiveCollector);
-		InvalidatingPrimitiveCollector.AddPrimitivesToInvalidate();
+
 		InvalidatingPrimitiveCollector.Instances.FinalizeBatches();
 
 		if (!InvalidatingPrimitiveCollector.Instances.IsEmpty())
diff --git a/Engine/Source/Runtime/Renderer/Private/VirtualShadowMaps/VirtualShadowMapCacheManager.h b/Engine/Source/Runtime/Renderer/Private/VirtualShadowMaps/VirtualShadowMapCacheManager.h
--- a/Engine/Source/Runtime/Renderer/Private/VirtualShadowMaps/VirtualShadowMapCacheManager.h
+++ b/Engine/Source/Runtime/Renderer/Private/VirtualShadowMaps/VirtualShadowMapCacheManager.h
@@ -19,9 +19,9 @@ namespace Nanite { struct FPackedViewParams; }
 
 struct FVirtualShadowMapInstanceRange
 {
+	FPersistentPrimitiveIndex PersistentPrimitiveIndex;
 	int32 InstanceSceneDataOffset;
 	int32 NumInstanceSceneDataEntries;
-	bool bForceInvalidateStatic;
 };
 
 #define VSM_LOG_INVALIDATIONS 0
@@ -355,7 +355,7 @@ private:
 	public:
 		FShadowInvalidatingInstancesImplementation(FVirtualShadowMapArrayCacheManager &InCacheManager) : CacheManager(InCacheManager) {}
 		virtual void AddPrimitive(const FPrimitiveSceneInfo *PrimitiveSceneInfo);
-		virtual void AddInstanceRange(uint32 InstanceSceneDataOffset, uint32 NumInstanceSceneDataEntries);
+		virtual void AddInstanceRange(FPersistentPrimitiveIndex PersistentPrimitiveIndex, uint32 InstanceSceneDataOffset, uint32 NumInstanceSceneDataEntries);
 
 		FVirtualShadowMapArrayCacheManager &CacheManager;
 		TArray<FVirtualShadowMapInstanceRange> PrimitiveInstancesToInvalidate;
diff --git a/Engine/Source/Runtime/Renderer/Public/SceneRendererInterface.h b/Engine/Source/Runtime/Renderer/Public/SceneRendererInterface.h
--- a/Engine/Source/Runtime/Renderer/Public/SceneRendererInterface.h
+++ b/Engine/Source/Runtime/Renderer/Public/SceneRendererInterface.h
@@ -33,7 +33,7 @@ public:
 	 * Add a range of instances for invalidation. 
 	 * Should be called from the render thread.
 	 */
-	virtual void AddInstanceRange(uint32 InstanceSceneDataOffset, uint32 NumInstanceSceneDataEntries) = 0;
+	virtual void AddInstanceRange(FPersistentPrimitiveIndex PersistentPrimitiveIndex, uint32 InstanceSceneDataOffset, uint32 NumInstanceSceneDataEntries) = 0;
 };
 
 }
