Fixing static analysis warnings caused by a potentially null pointer.

diff --git a/Engine/Source/Runtime/Renderer/Private/ViewDebug.cpp b/Engine/Source/Runtime/Renderer/Private/ViewDebug.cpp
--- a/Engine/Source/Runtime/Renderer/Private/ViewDebug.cpp
+++ b/Engine/Source/Runtime/Renderer/Private/ViewDebug.cpp
@@ -61,8 +61,12 @@ int32 FViewDebugInfo::FPrimitiveInfo::ComputeCurrentLODIndex(int32 PlayerIndex,
 		}
 	}
 	APlayerController* PlayerController = UGameplayStatics::GetPlayerController(Owner.Get(), PlayerIndex);
-	ULocalPlayer* LocalPlayer = IsValid(PlayerController) ? Cast<ULocalPlayer>(PlayerController->Player) : nullptr;
-	if (IsValid(LocalPlayer))
+	if (!IsValid(PlayerController))
+	{
+		return INDEX_NONE;
+	}
+	ULocalPlayer* LocalPlayer = Cast<ULocalPlayer>(PlayerController->Player);
+	if (IsValid(LocalPlayer) && IsValid(LocalPlayer->ViewportClient))
 	{
 		// see: AHUD::GetCoordinateOffset() and UGameViewportClient::Draw(FViewport* InViewport, FCanvas* SceneCanvas)
 		
