workaround for UB issue on some platforms. Followup JIRA has been create.

diff --git a/Engine/Shaders/Shared/EnvironmentComponentsFlags.h b/Engine/Shaders/Shared/EnvironmentComponentsFlags.h
--- a/Engine/Shaders/Shared/EnvironmentComponentsFlags.h
+++ b/Engine/Shaders/Shared/EnvironmentComponentsFlags.h
@@ -23,11 +23,11 @@
 
 
 // Using macro functions to avoid having to move the definition in cpp to avoid the multiple definition in different obj files.
-#define IsSkyAtmosphereHoldout(EnvironmentComponentsFlags)   ((EnvironmentComponentsFlags & ENVCOMP_FLAG_SKYATMOSPHERE_HOLDOUT)   > 0)
-#define IsVolumetricCloudHoldout(EnvironmentComponentsFlags) ((EnvironmentComponentsFlags & ENVCOMP_FLAG_VOLUMETRICCLOUD_HOLDOUT) > 0)
-#define IsExponentialFogHoldout(EnvironmentComponentsFlags)  ((EnvironmentComponentsFlags & ENVCOMP_FLAG_EXPONENTIALFOG_HOLDOUT)  > 0)
+#define IsSkyAtmosphereHoldout(EnvironmentComponentsFlags)   ((EnvironmentComponentsFlags[0] & ENVCOMP_FLAG_SKYATMOSPHERE_HOLDOUT)   > 0)
+#define IsVolumetricCloudHoldout(EnvironmentComponentsFlags) ((EnvironmentComponentsFlags[0] & ENVCOMP_FLAG_VOLUMETRICCLOUD_HOLDOUT) > 0)
+#define IsExponentialFogHoldout(EnvironmentComponentsFlags)  ((EnvironmentComponentsFlags[0] & ENVCOMP_FLAG_EXPONENTIALFOG_HOLDOUT)  > 0)
 
-#define IsSkyAtmosphereRenderedInMain(EnvironmentComponentsFlags)   ((EnvironmentComponentsFlags & ENVCOMP_FLAG_SKYATMOSPHERE_RENDERINMAIN)   > 0)
-#define IsVolumetricCloudRenderedInMain(EnvironmentComponentsFlags) ((EnvironmentComponentsFlags & ENVCOMP_FLAG_VOLUMETRICCLOUD_RENDERINMAIN) > 0)
-#define IsExponentialFogRenderedInMain(EnvironmentComponentsFlags)  ((EnvironmentComponentsFlags & ENVCOMP_FLAG_EXPONENTIALFOG_RENDERINMAIN)  > 0)
+#define IsSkyAtmosphereRenderedInMain(EnvironmentComponentsFlags)   ((EnvironmentComponentsFlags[0] & ENVCOMP_FLAG_SKYATMOSPHERE_RENDERINMAIN)   > 0)
+#define IsVolumetricCloudRenderedInMain(EnvironmentComponentsFlags) ((EnvironmentComponentsFlags[0] & ENVCOMP_FLAG_VOLUMETRICCLOUD_RENDERINMAIN) > 0)
+#define IsExponentialFogRenderedInMain(EnvironmentComponentsFlags)  ((EnvironmentComponentsFlags[0] & ENVCOMP_FLAG_EXPONENTIALFOG_RENDERINMAIN)  > 0)
 
diff --git a/Engine/Source/Runtime/Engine/Public/SceneView.h b/Engine/Source/Runtime/Engine/Public/SceneView.h
--- a/Engine/Source/Runtime/Engine/Public/SceneView.h
+++ b/Engine/Source/Runtime/Engine/Public/SceneView.h
@@ -995,8 +995,8 @@ enum ETranslucencyVolumeCascade
 	VIEW_UNIFORM_BUFFER_MEMBER_PER_VIEW(FVector4f, ScreenRayLengthMultiplier) \
 	VIEW_UNIFORM_BUFFER_MEMBER_PER_VIEW(FVector4f, GlintLUTParameters0) \
 	VIEW_UNIFORM_BUFFER_MEMBER_PER_VIEW(FVector4f, GlintLUTParameters1) \
-	VIEW_UNIFORM_BUFFER_MEMBER(int, EnvironmentComponentsFlags) \
 	VIEW_UNIFORM_BUFFER_MEMBER(float, MaterialMaxEmissiveValue) \
+	VIEW_UNIFORM_BUFFER_MEMBER(FIntVector4, EnvironmentComponentsFlags) \
 
 /** The uniform shader parameters associated with a view. */
 #define VIEW_UNIFORM_BUFFER_MEMBER(type, identifier) \
diff --git a/Engine/Source/Runtime/Renderer/Private/SceneRendering.cpp b/Engine/Source/Runtime/Renderer/Private/SceneRendering.cpp
--- a/Engine/Source/Runtime/Renderer/Private/SceneRendering.cpp
+++ b/Engine/Source/Runtime/Renderer/Private/SceneRendering.cpp
@@ -1454,17 +1454,18 @@ void FViewInfo::SetupUniformBufferParameters(
 	ViewUniformShaderParameters.AtmosphereIrradianceTextureSampler = TStaticSamplerState<SF_Bilinear>::GetRHI();
 	ViewUniformShaderParameters.AtmosphereInscatterTextureSampler = TStaticSamplerState<SF_Bilinear>::GetRHI();
 
-	// Update environment holdout flags
-	ViewUniformShaderParameters.EnvironmentComponentsFlags = 0;
+	// Upload environment holdout flags
+	ViewUniformShaderParameters.EnvironmentComponentsFlags = FIntVector4(EForceInit::ForceInitToZero);
 	if (Scene)
 	{
+		int32 Flags = 0;
 		if (ShouldRenderSkyAtmosphere(Scene, Family->EngineShowFlags))
 		{
 			FSkyAtmosphereRenderSceneInfo* SkyAtmosphere = Scene->SkyAtmosphere;
 			const FSkyAtmosphereSceneProxy& SkyAtmosphereSceneProxy = SkyAtmosphere->GetSkyAtmosphereSceneProxy();
 
-			ViewUniformShaderParameters.EnvironmentComponentsFlags |= SkyAtmosphereSceneProxy.IsHoldout() ? ENVCOMP_FLAG_SKYATMOSPHERE_HOLDOUT : 0;
-			ViewUniformShaderParameters.EnvironmentComponentsFlags |= SkyAtmosphereSceneProxy.IsRenderedInMainPass() ? ENVCOMP_FLAG_SKYATMOSPHERE_RENDERINMAIN : 0;
+			Flags |= SkyAtmosphereSceneProxy.IsHoldout() ? ENVCOMP_FLAG_SKYATMOSPHERE_HOLDOUT : 0;
+			Flags |= SkyAtmosphereSceneProxy.IsRenderedInMainPass() ? ENVCOMP_FLAG_SKYATMOSPHERE_RENDERINMAIN : 0;
 		}
 
 		if (ShouldRenderVolumetricCloud(Scene, Family->EngineShowFlags))
@@ -1472,17 +1473,19 @@ void FViewInfo::SetupUniformBufferParameters(
 			FVolumetricCloudRenderSceneInfo* VolumetricCloud = Scene->VolumetricCloud;
 			const FVolumetricCloudSceneProxy& VolumetricCloudSceneProxy = VolumetricCloud->GetVolumetricCloudSceneProxy();
 
-			ViewUniformShaderParameters.EnvironmentComponentsFlags |= VolumetricCloudSceneProxy.bHoldout ? ENVCOMP_FLAG_VOLUMETRICCLOUD_HOLDOUT : 0;
-			ViewUniformShaderParameters.EnvironmentComponentsFlags |= VolumetricCloudSceneProxy.bRenderInMainPass ? ENVCOMP_FLAG_VOLUMETRICCLOUD_RENDERINMAIN : 0;
+			Flags |= VolumetricCloudSceneProxy.bHoldout ? ENVCOMP_FLAG_VOLUMETRICCLOUD_HOLDOUT : 0;
+			Flags |= VolumetricCloudSceneProxy.bRenderInMainPass ? ENVCOMP_FLAG_VOLUMETRICCLOUD_RENDERINMAIN : 0;
 		}
 
 		if (Scene->ExponentialFogs.Num() > 0)
 		{
 			FExponentialHeightFogSceneInfo& Fog = Scene->ExponentialFogs[0];
 
-			ViewUniformShaderParameters.EnvironmentComponentsFlags |= Fog.bHoldout ? ENVCOMP_FLAG_EXPONENTIALFOG_HOLDOUT : 0;
-			ViewUniformShaderParameters.EnvironmentComponentsFlags |= Fog.bRenderInMainPass ? ENVCOMP_FLAG_EXPONENTIALFOG_RENDERINMAIN : 0;
+			Flags |= Fog.bHoldout ? ENVCOMP_FLAG_EXPONENTIALFOG_HOLDOUT : 0;
+			Flags |= Fog.bRenderInMainPass ? ENVCOMP_FLAG_EXPONENTIALFOG_RENDERINMAIN : 0;
 		}
+
+		ViewUniformShaderParameters.EnvironmentComponentsFlags.X = Flags;
 	}
 
 	ViewUniformShaderParameters.MaterialMaxEmissiveValue = MaterialMaxEmissiveValue;
