Disable mesh draw command caching for VT update when rendering standalone pages.
 - RenderPagesStandalone is called outside of a scene render. It's possible that render components have not been invalidated by World->SendAllEndOfFrameUpdates(), which results in a cached MDC crash.

diff --git a/Engine/Source/Runtime/Renderer/Private/VT/RuntimeVirtualTextureRender.cpp b/Engine/Source/Runtime/Renderer/Private/VT/RuntimeVirtualTextureRender.cpp
--- a/Engine/Source/Runtime/Renderer/Private/VT/RuntimeVirtualTextureRender.cpp
+++ b/Engine/Source/Runtime/Renderer/Private/VT/RuntimeVirtualTextureRender.cpp
@@ -979,7 +979,7 @@ namespace RuntimeVirtualTexture
 
 
 	/** Collect meshes to draw. */
-	void GatherMeshesToDraw(FDynamicPassMeshDrawListContext* DynamicMeshPassContext, FScene const* Scene, FViewInfo* View, ERuntimeVirtualTextureMaterialType MaterialType, uint32 RuntimeVirtualTextureMask, uint8 vLevel, uint8 MaxLevel)
+	void GatherMeshesToDraw(FDynamicPassMeshDrawListContext* DynamicMeshPassContext, FScene const* Scene, FViewInfo* View, ERuntimeVirtualTextureMaterialType MaterialType, uint32 RuntimeVirtualTextureMask, uint8 vLevel, uint8 MaxLevel, bool bAllowCachedMeshDrawCommands)
 	{
 		// Cached draw command collectors
 		const FCachedPassMeshDrawList& SceneDrawList = Scene->CachedDrawLists[EMeshPass::VirtualTexture];
@@ -1049,7 +1049,7 @@ namespace RuntimeVirtualTexture
 				if (StaticMeshRelevance.bRenderToVirtualTexture && StaticMeshRelevance.LODIndex == LodIndex && StaticMeshRelevance.RuntimeVirtualTextureMaterialType == (uint32)MaterialType)
 				{
 					bool bCachedDraw = false;
-					if (StaticMeshRelevance.bSupportsCachingMeshDrawCommands)
+					if (bAllowCachedMeshDrawCommands && StaticMeshRelevance.bSupportsCachingMeshDrawCommands)
 					{
 						// Use cached draw command
 						const int32 StaticMeshCommandInfoIndex = StaticMeshRelevance.GetStaticMeshCommandInfoIndex(EMeshPass::VirtualTexture);
@@ -1375,6 +1375,7 @@ namespace RuntimeVirtualTexture
 		ERuntimeVirtualTextureMaterialType MaterialType,
 		bool bClearTextures,
 		bool bIsThumbnails,
+		bool bAllowCachedMeshDrawCommands,
 		FRHITexture* OutputTexture0,		// todo[vt]: Only use IPooledRenderTarget or FRDGTextureRef, not raw RHI textures.
 		IPooledRenderTarget* OutputTarget0,
 		FBox2D const& DestBox0,
@@ -1471,9 +1472,9 @@ namespace RuntimeVirtualTexture
 			PassParameters->RenderTargets[2] = GraphSetup.RenderTexture2 ? FRenderTargetBinding(GraphSetup.RenderTexture2, LoadAction) : FRenderTargetBinding();
     		
 			AddSimpleMeshPass(GraphBuilder, PassParameters, Scene, *View, nullptr, RDG_EVENT_NAME("VirtualTextureDraw"), View->ViewRect,
-	    	[Scene, View, MaterialType, RuntimeVirtualTextureMask, vLevel, MaxLevel](FDynamicPassMeshDrawListContext* DynamicMeshPassContext)
+	    	[Scene, View, MaterialType, RuntimeVirtualTextureMask, vLevel, MaxLevel, bAllowCachedMeshDrawCommands](FDynamicPassMeshDrawListContext* DynamicMeshPassContext)
 	    	{
-	    		GatherMeshesToDraw(DynamicMeshPassContext, Scene, View, MaterialType, RuntimeVirtualTextureMask, vLevel, MaxLevel);
+	    		GatherMeshesToDraw(DynamicMeshPassContext, Scene, View, MaterialType, RuntimeVirtualTextureMask, vLevel, MaxLevel, bAllowCachedMeshDrawCommands);
 	    	});
 		}
 
@@ -1542,7 +1543,7 @@ namespace RuntimeVirtualTexture
 		}
 	}
 
-	void RenderPagesInternal(FRDGBuilder& GraphBuilder, FRenderPageBatchDesc const& InDesc, ISceneRenderer* SceneRenderer)
+	void RenderPagesInternal(FRDGBuilder& GraphBuilder, FRenderPageBatchDesc const& InDesc, ISceneRenderer* SceneRenderer, bool bAllowCachedMeshDrawCommands)
 	{
 		check(InDesc.NumPageDescs <= EMaxRenderPageBatch);
 
@@ -1563,6 +1564,7 @@ namespace RuntimeVirtualTexture
 					InDesc.MaterialType,
 					InDesc.bClearTextures,
 					InDesc.bIsThumbnails,
+					bAllowCachedMeshDrawCommands,
 					InDesc.Targets[0].Texture, InDesc.Targets[0].PooledRenderTarget, PageDesc.DestBox[0],
 					InDesc.Targets[1].Texture, InDesc.Targets[1].PooledRenderTarget, PageDesc.DestBox[1],
 					InDesc.Targets[2].Texture, InDesc.Targets[2].PooledRenderTarget, PageDesc.DestBox[2],
@@ -1580,7 +1582,10 @@ namespace RuntimeVirtualTexture
 	{
 		check(InDesc.Scene != nullptr);
 		FScenePrimitiveRenderingContextScopeHelper RenderingScope(GetRendererModule().BeginScenePrimitiveRendering(GraphBuilder, *InDesc.Scene));
-		RenderPagesInternal(GraphBuilder, InDesc, RenderingScope.ScenePrimitiveRenderingContext->GetSceneRenderer());
+
+		// Disable MDC caching because we can't guarantee that primitives associated with the scene have been recreated (e.g. by World->SendAllEndOfFrameUpdates()).
+		const bool bAllowCachedMeshDrawCommands = false;
+		RenderPagesInternal(GraphBuilder, InDesc, RenderingScope.ScenePrimitiveRenderingContext->GetSceneRenderer(), bAllowCachedMeshDrawCommands);
 	}
 
 	void RenderPages(FRDGBuilder& GraphBuilder, FRenderPageBatchDesc const& InDesc)
@@ -1602,7 +1607,8 @@ namespace RuntimeVirtualTexture
 			}
 			};
 			FSimpleRVTRenderer SimpleRenderer(GraphBuilder, InDesc);
-			RenderPagesInternal(GraphBuilder, InDesc, &SimpleRenderer);
+			const bool bAllowCachedMeshDrawCommands = true;
+			RenderPagesInternal(GraphBuilder, InDesc, &SimpleRenderer, bAllowCachedMeshDrawCommands);
 		}
 		else
 		{
