Join shadow tasks early if async lumen is enabled to avoid joining back to the graphics pipe.

diff --git a/Engine/Source/Runtime/Renderer/Private/DeferredShadingRenderer.cpp b/Engine/Source/Runtime/Renderer/Private/DeferredShadingRenderer.cpp
--- a/Engine/Source/Runtime/Renderer/Private/DeferredShadingRenderer.cpp
+++ b/Engine/Source/Runtime/Renderer/Private/DeferredShadingRenderer.cpp
@@ -2565,6 +2565,7 @@ void FDeferredShadingSceneRenderer::Render(FRDGBuilder& GraphBuilder)
 				SceneTextures,
 				InstanceCullingManager,
 				LumenFrameTemporaries,
+				InitViewTaskDatas.DynamicShadows,
 				LightingChannelsTexture,
 				/*bHasLumenLights*/ false,
 				AsyncLumenIndirectLightingOutputs);
diff --git a/Engine/Source/Runtime/Renderer/Private/DeferredShadingRenderer.h b/Engine/Source/Runtime/Renderer/Private/DeferredShadingRenderer.h
--- a/Engine/Source/Runtime/Renderer/Private/DeferredShadingRenderer.h
+++ b/Engine/Source/Runtime/Renderer/Private/DeferredShadingRenderer.h
@@ -666,6 +666,7 @@ private:
 		FSceneTextures& SceneTextures,
 		FInstanceCullingManager& InstanceCullingManager,
 		const FLumenSceneFrameTemporaries& LumenFrameTemporaries,
+		FDynamicShadowsTaskData* DynamicShadowsTaskData,
 		FRDGTextureRef LightingChannelsTexture,
 		bool bHasLumenLights,
 		FAsyncLumenIndirectLightingOutputs& Outputs);
diff --git a/Engine/Source/Runtime/Renderer/Private/IndirectLightRendering.cpp b/Engine/Source/Runtime/Renderer/Private/IndirectLightRendering.cpp
--- a/Engine/Source/Runtime/Renderer/Private/IndirectLightRendering.cpp
+++ b/Engine/Source/Runtime/Renderer/Private/IndirectLightRendering.cpp
@@ -863,6 +863,7 @@ void FDeferredShadingSceneRenderer::DispatchAsyncLumenIndirectLightingWork(
 	FSceneTextures& SceneTextures,
 	FInstanceCullingManager& InstanceCullingManager,
 	const FLumenSceneFrameTemporaries& LumenFrameTemporaries,
+	FDynamicShadowsTaskData* DynamicShadowsTaskData,
 	FRDGTextureRef LightingChannelsTexture,
 	bool bHasLumenLights,
 	FAsyncLumenIndirectLightingOutputs& Outputs)
@@ -885,6 +886,21 @@ void FDeferredShadingSceneRenderer::DispatchAsyncLumenIndirectLightingWork(
 	CompositionLighting.ProcessAfterBasePass(GraphBuilder, InstanceCullingManager, FCompositionLighting::EProcessAfterBasePassMode::OnlyBeforeLightingDecals);
 	Outputs.bHasDrawnBeforeLightingDecals = true;
 
+	if (DynamicShadowsTaskData)
+	{
+		for (const FViewInfo& View : Views)
+		{
+			const FPerViewPipelineState& ViewPipelineState = GetViewPipelineState(View);
+
+			// Write shadow GPU instance data prior to launching async compute tasks to avoid the early join back to graphics.
+			if (ViewPipelineState.DiffuseIndirectMethod == EDiffuseIndirectMethod::Lumen)
+			{
+				FinishDynamicShadowMeshPassSetup(GraphBuilder, DynamicShadowsTaskData);
+				break;
+			}
+		}
+	}
+
 	LLM_SCOPE_BYTAG(Lumen);
 	RDG_EVENT_SCOPE(GraphBuilder, "DiffuseIndirectAndAO");
 	SCOPED_NAMED_EVENT(DispatchAsyncLumenIndirectLightingWork, FColor::Emerald);
