[ScreenPercentage] Allow to override screen percentage on a per view basis instead of a viewfamily only setting

Problem: Screen percentage can be controlled per viewfamily but not changed on a single view which would be required to record reference views for quantitative algorithm comparisons.

Solution: Add an override setting to a view which is used instead of the default or dynamic screen resolution.

#rnx

diff --git a/Engine/Source/Runtime/Engine/Public/SceneView.h b/Engine/Source/Runtime/Engine/Public/SceneView.h
--- a/Engine/Source/Runtime/Engine/Public/SceneView.h
+++ b/Engine/Source/Runtime/Engine/Public/SceneView.h
@@ -244,6 +244,9 @@ struct FSceneViewInitOptions : public FSceneViewProjectionData
 	/** Whether this view is being used to render a planar reflection. */
 	bool bIsPlanarReflection;
 
+	/** If > 0, overrides the view's resolution fraction. */
+	float OverridePrimaryResolutionFraction;
+
 #if WITH_EDITOR
 	/** default to 0'th view index, which is a bitfield of 1 */
 	uint64 EditorViewBitflag;
@@ -279,6 +282,7 @@ struct FSceneViewInitOptions : public FSceneViewProjectionData
 		, bSceneCaptureUsesRayTracing(false)
 		, bIsReflectionCapture(false)
 		, bIsPlanarReflection(false)
+		, OverridePrimaryResolutionFraction(-1.0)
 #if WITH_EDITOR
 		, EditorViewBitflag(1)
 		, bDisableGameScreenPercentage(false)
diff --git a/Engine/Source/Runtime/Renderer/Private/SceneRendering.cpp b/Engine/Source/Runtime/Renderer/Private/SceneRendering.cpp
--- a/Engine/Source/Runtime/Renderer/Private/SceneRendering.cpp
+++ b/Engine/Source/Runtime/Renderer/Private/SceneRendering.cpp
@@ -2947,7 +2947,7 @@ FIntPoint FSceneRenderer::GetDesiredInternalBufferSize(const FSceneViewFamily& V
 	for (const FSceneView* View : ViewFamily.AllViews)
 	{
 		// Note: This ensures that custom passes (rendered with the main renderer) ignore screen percentage, like regular scene captures.
-		const float AdjustedResolutionFractionUpperBounds = View->CustomRenderPass ? 1.0f : ResolutionFractionUpperBound;
+		const float AdjustedResolutionFractionUpperBounds = View->CustomRenderPass ? 1.0f : (View->SceneViewInitOptions.OverridePrimaryResolutionFraction > 0.0 ? (View->SceneViewInitOptions.OverridePrimaryResolutionFraction * ViewFamily.SecondaryViewFraction)  : ResolutionFractionUpperBound);
 
 		FIntPoint ViewSize = ApplyResolutionFraction(ViewFamily, View->UnconstrainedViewRect.Size(), AdjustedResolutionFractionUpperBounds);
 		FIntPoint ViewRectMin = QuantizeViewRectMin(FIntPoint(
@@ -3077,10 +3077,12 @@ void FSceneRenderer::PrepareViewRectsForRendering(FRHICommandListImmediate& RHIC
 	{
 		FViewInfo& View = Views[i];
 
-		FIntPoint ViewSize = ApplyResolutionFraction(ViewFamily, View.UnscaledViewRect.Size(), ResolutionFraction);
+		float ViewResolutionFraction = View.SceneViewInitOptions.OverridePrimaryResolutionFraction > 0.0 ? (View.SceneViewInitOptions.OverridePrimaryResolutionFraction * ViewFamily.SecondaryViewFraction) : ResolutionFraction;
+
+		FIntPoint ViewSize = ApplyResolutionFraction(ViewFamily, View.UnscaledViewRect.Size(), ViewResolutionFraction);
 		FIntPoint ViewRectMin = QuantizeViewRectMin(FIntPoint(
-			FMath::CeilToInt(View.UnscaledViewRect.Min.X * ResolutionFraction),
-			FMath::CeilToInt(View.UnscaledViewRect.Min.Y * ResolutionFraction)));
+			FMath::CeilToInt(View.UnscaledViewRect.Min.X * ViewResolutionFraction),
+			FMath::CeilToInt(View.UnscaledViewRect.Min.Y * ViewResolutionFraction)));
 
 		// Use the bottom-left view rect if requested, instead of top-left
 		if (CVarViewRectUseScreenBottom.GetValueOnRenderThread())
