Fixing shader binding crash

diff --git a/Engine/Source/Runtime/Renderer/Private/LightRendering.cpp b/Engine/Source/Runtime/Renderer/Private/LightRendering.cpp
--- a/Engine/Source/Runtime/Renderer/Private/LightRendering.cpp
+++ b/Engine/Source/Runtime/Renderer/Private/LightRendering.cpp
@@ -2950,6 +2950,7 @@ static FSimpleLightsStandardDeferredParameters GetRenderLightSimpleParameters(
 	{
 		Out.PS.RenderTargets.DepthStencil = FDepthStencilBinding(SceneTextures.Depth.Target, ERenderTargetLoadAction::ELoad, ERenderTargetLoadAction::ELoad, FExclusiveDepthStencil::DepthRead_StencilWrite);
 	}
+	Out.PS.AVSM = HeterogeneousVolumes::CreateEmptyAdaptiveVolumetricShadowMapUniformBuffer(GraphBuilder);
 
 	// VS - General parameters (dummy geometry, as the geometry is setup within the pass light loop)
 	FSphere SphereLight;
