Adding setting for disabling one pass transmittance for hair lighting.

diff --git a/Engine/Source/Runtime/Renderer/Private/LightRendering.cpp b/Engine/Source/Runtime/Renderer/Private/LightRendering.cpp
--- a/Engine/Source/Runtime/Renderer/Private/LightRendering.cpp
+++ b/Engine/Source/Runtime/Renderer/Private/LightRendering.cpp
@@ -176,6 +176,12 @@ static TAutoConsoleVariable<float> CVarContactShadowsOverrideNonShadowCastingInt
 	TEXT("Should generally be left disabled outside of debugging."),
 	ECVF_RenderThreadSafe);
 
+static TAutoConsoleVariable<int32> CVarHairStrandsAllowOneTransmittancePass(
+	TEXT("r.HairStrands.Lighting.AllowOneTransmittancePass"),
+	0,
+	TEXT("Allows one transmittance pass for hair strands lighting to have better performance (experimental).\n"),
+	ECVF_RenderThreadSafe);
+
 #if ENABLE_DEBUG_DISCARD_PROP
 static float GDebugLightDiscardProp = 0.0f;
 static FAutoConsoleVariableRef CVarDebugLightDiscardProp(
@@ -1555,6 +1561,7 @@ void FDeferredShadingSceneRenderer::RenderLights(
 				bool bUsedShadowMaskTexture = false;
 
 				bool bElideScreenShadowMask = false;
+				bool bElideScreenShadowMaskSubPixel = false;
 
 				FScopeCycleCounter Context(LightSceneProxy.GetStatId());
 
@@ -1575,6 +1582,7 @@ void FDeferredShadingSceneRenderer::RenderLights(
 						SortedLightInfo.SortKey.Fields.LightType != LightType_Directional &&
 						VisibleLightInfo.VirtualShadowMapId != INDEX_NONE &&	// Not a directional light, so no per-view clipmaps
 						VisibleLightInfo.ContainsOnlyVirtualShadowMaps();
+					bElideScreenShadowMaskSubPixel = bElideScreenShadowMask && CVarHairStrandsAllowOneTransmittancePass.GetValueOnRenderThread() > 0;
 
 					if (!SharedScreenShadowMaskTexture || !SharedScreenShadowMaskSubPixelTexture)
 					{
@@ -1584,13 +1592,13 @@ void FDeferredShadingSceneRenderer::RenderLights(
 						{
 							SharedScreenShadowMaskTexture = GraphBuilder.CreateTexture(SharedScreenShadowMaskTextureDesc, TEXT("ShadowMaskTexture"));
 						}
-						if (!SharedScreenShadowMaskSubPixelTexture && bUseHairLighting && !bElideScreenShadowMask)
+						if (!SharedScreenShadowMaskSubPixelTexture && bUseHairLighting && !bElideScreenShadowMaskSubPixel)
 						{
 							SharedScreenShadowMaskSubPixelTexture = GraphBuilder.CreateTexture(SharedScreenShadowMaskTextureDesc, TEXT("ShadowMaskSubPixelTexture"));
 						}
 					}
 					ScreenShadowMaskTexture = bElideScreenShadowMask ? nullptr : SharedScreenShadowMaskTexture;
-					ScreenShadowMaskSubPixelTexture = bElideScreenShadowMask ? nullptr : SharedScreenShadowMaskSubPixelTexture;
+					ScreenShadowMaskSubPixelTexture = bElideScreenShadowMaskSubPixel ? nullptr : SharedScreenShadowMaskSubPixelTexture;
 				}
 
 				FString LightNameWithLevel;
@@ -2101,7 +2109,7 @@ void FDeferredShadingSceneRenderer::RenderLights(
 						{
 							// If the light elided the screen space shadow mask, sample directly from the packed shadow mask
 							int32 VirtualShadowMapId = INDEX_NONE;
-							if (bElideScreenShadowMask)
+							if (bElideScreenShadowMaskSubPixel)
 							{
 								INC_DWORD_STAT(STAT_VSMLocalProjectionOnePassFast);
 								VirtualShadowMapId = VisibleLightInfo.GetVirtualShadowMapId(&View);
