Fix a bug where all Nanite will render for a view with ShowOnly primitives enabled, but has an empty list.

diff --git a/Engine/Source/Runtime/Renderer/Private/BasePassRendering.cpp b/Engine/Source/Runtime/Renderer/Private/BasePassRendering.cpp
--- a/Engine/Source/Runtime/Renderer/Private/BasePassRendering.cpp
+++ b/Engine/Source/Runtime/Renderer/Private/BasePassRendering.cpp
@@ -1492,7 +1492,9 @@ void FDeferredShadingSceneRenderer::RenderBasePassInternal(
 					});
 				}
 
-				const bool bShouldRenderViewForNanite = bNaniteEnabled && (!bDrawSceneViewsInOneNanitePass || ViewIndex == 0); // when bDrawSceneViewsInOneNanitePass, the first view should cover all the other atlased ones
+				const bool bShouldRenderViewForNanite = bNaniteEnabled
+					&& !View.bHasNoVisiblePrimitive
+					&& (!bDrawSceneViewsInOneNanitePass || ViewIndex == 0); // when bDrawSceneViewsInOneNanitePass, the first view should cover all the other atlased ones
 				if (bShouldRenderViewForNanite)
 				{
 					// Should always have a full Z prepass with Nanite
diff --git a/Engine/Source/Runtime/Renderer/Private/DeferredShadingRenderer.cpp b/Engine/Source/Runtime/Renderer/Private/DeferredShadingRenderer.cpp
--- a/Engine/Source/Runtime/Renderer/Private/DeferredShadingRenderer.cpp
+++ b/Engine/Source/Runtime/Renderer/Private/DeferredShadingRenderer.cpp
@@ -1340,7 +1340,12 @@ void FDeferredShadingSceneRenderer::RenderNanite(FRDGBuilder& GraphBuilder, cons
 			const FViewInfo& View = InViews[ViewIndex];
 			// We don't check View.ShouldRenderView() since this is already taken care of by bDrawSceneViewsInOneNanitePass.
 			// If bDrawSceneViewsInOneNanitePass is false, we need to render the secondary view even if ShouldRenderView() is false
-
+			// NOTE: Except when there are no primitives to draw for the view
+			if (View.bHasNoVisiblePrimitive)
+			{
+				continue;
+			}
+			
 			RDG_EVENT_SCOPE_CONDITIONAL(GraphBuilder, InViews.Num() > 1 && !bDrawSceneViewsInOneNanitePass, "View%u", ViewIndex);
 			RDG_EVENT_SCOPE_CONDITIONAL(GraphBuilder, InViews.Num() > 1 && bDrawSceneViewsInOneNanitePass, "View%u (together with %d more)", ViewIndex, InViews.Num() - 1);
 
