Moved expensive view map initialization to a separate task. This reduces the critical path for complex scenes by starting frustum culling earlier.

diff --git a/Engine/Source/Runtime/Renderer/Private/SceneVisibility.cpp b/Engine/Source/Runtime/Renderer/Private/SceneVisibility.cpp
--- a/Engine/Source/Runtime/Renderer/Private/SceneVisibility.cpp
+++ b/Engine/Source/Runtime/Renderer/Private/SceneVisibility.cpp
@@ -1127,7 +1127,8 @@ FRelevancePacket::FRelevancePacket(
 	const FViewInfo& InView,
 	int32 InViewIndex,
 	const FFilterStaticMeshesForViewData& InViewData,
-	uint8* InMarkMasks)
+	uint8* InMarkMasks,
+	const UE::Tasks::FTask& InPrerequisitesTask)
 	: CurrentWorldTime(InView.Family->Time.GetWorldTimeSeconds())
 	, DeltaWorldTime(InView.Family->Time.GetDeltaWorldTimeSeconds())
 	, TaskData(InTaskData)
@@ -1139,6 +1140,7 @@ FRelevancePacket::FRelevancePacket(
 	, ViewData(InViewData)
 	, DynamicPrimitiveViewMasks(TaskData.DynamicMeshElements.PrimitiveViewMasks)
 	, MarkMasks(InMarkMasks)
+	, PrerequisitesTask(InPrerequisitesTask)
 	, Input(TaskConfig.Relevance.NumPrimitivesPerPacket)
 	, NotDrawRelevant(TaskConfig.Relevance.NumPrimitivesPerPacket)
 	, TranslucentSelfShadowPrimitives(TaskConfig.Relevance.NumPrimitivesPerPacket)
@@ -1180,7 +1182,7 @@ void FRelevancePacket::LaunchComputeRelevanceTask()
 				}
 			}
 
-		}, Scene.GetCacheMeshDrawCommandsTask(), TaskConfig.Relevance.ComputeRelevanceTaskPriority);
+		}, PrerequisitesTask, TaskConfig.Relevance.ComputeRelevanceTaskPriority);
 	}
 }
 
@@ -1949,7 +1951,7 @@ void FRelevancePacket::ComputeRelevance(FDynamicPrimitiveIndexList& DynamicPrimi
 
 ///////////////////////////////////////////////////////////////////////////////
 
-FComputeAndMarkRelevance::FComputeAndMarkRelevance(FVisibilityTaskData& InTaskData, FScene& InScene, FViewInfo& InView, uint8 InViewIndex)
+FComputeAndMarkRelevance::FComputeAndMarkRelevance(FVisibilityTaskData& InTaskData, FScene& InScene, FViewInfo& InView, uint8 InViewIndex, const UE::Tasks::FTask& InPrerequisitesTask)
 	: TaskData(InTaskData)
 	, Scene(InScene)
 	, View(InView)
@@ -1958,6 +1960,7 @@ FComputeAndMarkRelevance::FComputeAndMarkRelevance(FVisibilityTaskData& InTaskDa
 	, ViewData(View)
 	, NumMeshes(Scene.StaticMeshes.GetMaxIndex())
 	, NumPrimitivesPerPacket(InTaskData.TaskConfig.Relevance.NumPrimitivesPerPacket)
+	, PrerequisitesTask(InPrerequisitesTask)
 	, bLaunchOnAddPrimitive(TaskData.TaskConfig.Schedule == EVisibilityTaskSchedule::Parallel)
 	, bFinished(!bLaunchOnAddPrimitive)
 {
@@ -2046,6 +2049,7 @@ void FComputeAndMarkRelevance::Finalize()
 {
 	check(bFinished && !bFinalized);
 	bFinalized = true;
+	PrerequisitesTask.Wait();
 
 	if (!bLaunchOnAddPrimitive)
 	{
@@ -3500,24 +3504,35 @@ void FVisibilityViewPacket::BeginInitVisibility()
 	// Allocate the view's visibility maps.
 	View.PrimitiveVisibilityMap.Init(!bShouldVisibilityCull, Scene.Primitives.Num());
 	View.PrimitiveRayTracingVisibilityMap.Init(false, Scene.Primitives.Num());
-	View.DynamicMeshElementRanges.SetNumZeroed(Scene.Primitives.Num());
 	View.PotentiallyFadingPrimitiveMap.Init(false, Scene.Primitives.Num());
 	View.PrimitiveFadeUniformBuffers.AddZeroed(Scene.Primitives.Num());
 	View.PrimitiveFadeUniformBufferMap.Init(false, Scene.Primitives.Num());
-	View.StaticMeshVisibilityMap.Init(false, Scene.StaticMeshes.GetMaxIndex());
-	View.StaticMeshFadeOutDitheredLODMap.Init(false, Scene.StaticMeshes.GetMaxIndex());
-	View.StaticMeshFadeInDitheredLODMap.Init(false, Scene.StaticMeshes.GetMaxIndex());
-	View.PrimitivesLODMask.Init(FLODMask(), Scene.Primitives.Num());
 
-	View.PrimitiveViewRelevanceMap.Reset(Scene.Primitives.Num());
-	View.PrimitiveViewRelevanceMap.AddZeroed(Scene.Primitives.Num());
+	UE::Tasks::FTaskEvent RelevancePrereqs{ UE_SOURCE_LOCATION };
+	RelevancePrereqs.AddPrerequisites(Scene.GetCacheMeshDrawCommandsTask());
+
+	// Offload initialization of maps that are not touched by frustum culling / occlusion until the relevance phase. These can be quite expensive to zero out.
+	RelevancePrereqs.AddPrerequisites(UE::Tasks::Launch(UE_SOURCE_LOCATION, [this]
+	{
+		TRACE_CPUPROFILER_EVENT_SCOPE(SceneVisibility_InitMaps);
+		View.DynamicMeshElementRanges.SetNumZeroed(Scene.Primitives.Num());
+		View.PrimitivesLODMask.Init(FLODMask(), Scene.Primitives.Num());
+		View.StaticMeshVisibilityMap.Init(false, Scene.StaticMeshes.GetMaxIndex());
+		View.StaticMeshFadeOutDitheredLODMap.Init(false, Scene.StaticMeshes.GetMaxIndex());
+		View.StaticMeshFadeInDitheredLODMap.Init(false, Scene.StaticMeshes.GetMaxIndex());
+		View.PrimitiveViewRelevanceMap.Reset(Scene.Primitives.Num());
+		View.PrimitiveViewRelevanceMap.AddZeroed(Scene.Primitives.Num());
+
+	}, TaskConfig.TaskPriority));
+
+	RelevancePrereqs.Trigger();
 
 	if (View.ShowOnlyPrimitives.IsSet())
 	{
 		View.bHasNoVisiblePrimitive = View.ShowOnlyPrimitives->Num() == 0;
 	}
 
-	Relevance.Context = TaskData.Allocator.Create<FComputeAndMarkRelevance>(TaskData, Scene, View, ViewIndex);
+	Relevance.Context = TaskData.Allocator.Create<FComputeAndMarkRelevance>(TaskData, Scene, View, ViewIndex, RelevancePrereqs);
 
 	ClearStalePrimitiveFadingStates(View, ViewState);
 
@@ -3546,6 +3561,8 @@ void FVisibilityViewPacket::BeginInitVisibility()
 	// Most views use standard frustum culling.
 	if (bShouldVisibilityCull)
 	{
+		TRACE_CPUPROFILER_EVENT_SCOPE(HLODInit);
+
 		// Update HLOD transition/visibility states to allow use during distance culling
 		FLODSceneTree& HLODTree = Scene.SceneLODHierarchy;
 
@@ -5517,6 +5534,8 @@ void FDeferredShadingSceneRenderer::BeginInitViews(
 			// Initialize the view's RHI resources.
 			UpdateHairResources(GraphBuilder, View);
 			View.InitRHIResources();
+
+			View.PrevHZB = TryRegisterExternalTexture(GraphBuilder, View.PrevViewInfo.HZB);
 		}
 
 		for (FCustomRenderPassInfo& PassInfo : CustomRenderPassInfos)
diff --git a/Engine/Source/Runtime/Renderer/Private/SceneVisibilityPrivate.h b/Engine/Source/Runtime/Renderer/Private/SceneVisibilityPrivate.h
--- a/Engine/Source/Runtime/Renderer/Private/SceneVisibilityPrivate.h
+++ b/Engine/Source/Runtime/Renderer/Private/SceneVisibilityPrivate.h
@@ -705,7 +705,8 @@ public:
 		const FViewInfo& InView,
 		int32 InViewIndex,
 		const FFilterStaticMeshesForViewData& InViewData,
-		uint8* InMarkMasks);
+		uint8* InMarkMasks,
+		const UE::Tasks::FTask& PrerequisitesTask);
 
 	void LaunchComputeRelevanceTask();
 
@@ -728,6 +729,7 @@ private:
 	const FFilterStaticMeshesForViewData& ViewData;
 	FDynamicPrimitiveViewMasks* DynamicPrimitiveViewMasks;
 	uint8* RESTRICT MarkMasks;
+	UE::Tasks::FTask PrerequisitesTask;
 
 	FRelevancePrimSet<int32> Input;
 	FRelevancePrimSet<int32> NotDrawRelevant;
@@ -792,7 +794,7 @@ private:
 class FComputeAndMarkRelevance
 {
 public:
-	FComputeAndMarkRelevance(FVisibilityTaskData& InTaskData, FScene& InScene, FViewInfo& InView, uint8 InViewIndex);
+	FComputeAndMarkRelevance(FVisibilityTaskData& InTaskData, FScene& InScene, FViewInfo& InView, uint8 InViewIndex, const UE::Tasks::FTask& PreprequisitesTask);
 
 	~FComputeAndMarkRelevance()
 	{
@@ -813,7 +815,7 @@ private:
 	FRelevancePacket* CreateRelevancePacket()
 	{
 		check(!bLaunchOnAddPrimitive || !bFinished);
-		return Packets.Emplace_GetRef(new FRelevancePacket(TaskData, View, ViewIndex, ViewData, MarkMasks));
+		return Packets.Emplace_GetRef(new FRelevancePacket(TaskData, View, ViewIndex, ViewData, MarkMasks, PrerequisitesTask));
 	}
 
 	FVisibilityTaskData& TaskData;
@@ -827,6 +829,7 @@ private:
 	const uint32 NumPrimitivesPerPacket;
 	uint8* MarkMasks;
 	TArray<FRelevancePacket*, SceneRenderingAllocator> Packets;
+	UE::Tasks::FTask PrerequisitesTask;
 	const bool bLaunchOnAddPrimitive;
 	bool bFinished = false;
 	bool bFinalized = false;
