Fixed build break

diff --git a/Engine/Source/Runtime/Renderer/Private/SceneVisibility.cpp b/Engine/Source/Runtime/Renderer/Private/SceneVisibility.cpp
--- a/Engine/Source/Runtime/Renderer/Private/SceneVisibility.cpp
+++ b/Engine/Source/Runtime/Renderer/Private/SceneVisibility.cpp
@@ -5534,8 +5534,6 @@ void FDeferredShadingSceneRenderer::BeginInitViews(
 			// Initialize the view's RHI resources.
 			UpdateHairResources(GraphBuilder, View);
 			View.InitRHIResources();
-
-			View.PrevHZB = TryRegisterExternalTexture(GraphBuilder, View.PrevViewInfo.HZB);
 		}
 
 		for (FCustomRenderPassInfo& PassInfo : CustomRenderPassInfos)
