Early out decal PSO collection if shaders are not ready

diff --git a/Engine/Source/Runtime/Renderer/Private/CompositionLighting/PostProcessDeferredDecals.cpp b/Engine/Source/Runtime/Renderer/Private/CompositionLighting/PostProcessDeferredDecals.cpp
--- a/Engine/Source/Runtime/Renderer/Private/CompositionLighting/PostProcessDeferredDecals.cpp
+++ b/Engine/Source/Runtime/Renderer/Private/CompositionLighting/PostProcessDeferredDecals.cpp
@@ -509,7 +509,10 @@ void CollectDeferredDecalPassPSOInitializers(
 	GraphicsPSOInit.PrimitiveType = PT_TriangleList;		
 	GraphicsPSOInit.BlendState = DecalRendering::GetDecalBlendState(DecalBlendDesc, DecalRenderStage, DecalRenderTargetMode);
 
-	DecalRendering::SetupShaderState(FeatureLevel, Material, DecalRenderStage, GraphicsPSOInit.BoundShaderState);
+	if (!DecalRendering::SetupShaderState(FeatureLevel, Material, DecalRenderStage, GraphicsPSOInit.BoundShaderState))
+	{
+		return;
+	}
 
 	FGraphicsPipelineRenderTargetsInfo RenderTargetsInfo;
 	RenderTargetsInfo.NumSamples = 1;
