HeterogeneousVolumes: Qualify collection of mesh batches as being gathered for shadow-casting or not

diff --git a/Engine/Source/Runtime/Renderer/Private/HeterogeneousVolumes/HeterogeneousVolumesLiveShadingPipeline.cpp b/Engine/Source/Runtime/Renderer/Private/HeterogeneousVolumes/HeterogeneousVolumesLiveShadingPipeline.cpp
--- a/Engine/Source/Runtime/Renderer/Private/HeterogeneousVolumes/HeterogeneousVolumesLiveShadingPipeline.cpp
+++ b/Engine/Source/Runtime/Renderer/Private/HeterogeneousVolumes/HeterogeneousVolumesLiveShadingPipeline.cpp
@@ -1175,6 +1175,7 @@ IMPLEMENT_MATERIAL_SHADER_TYPE(, FRenderVolumetricShadowMapForLightWithLiveShadi
 
 void CollectHeterogeneousVolumeMeshBatchesForView(
 	const FViewInfo& View,
+	bool bCollectForShadowCasting,
 	TSet<FVolumetricMeshBatch>& HeterogeneousVolumesMeshBatches,
 	FBoxSphereBounds& WorldBounds
 )
@@ -1182,7 +1183,6 @@ void CollectHeterogeneousVolumeMeshBatchesForView(
 	for (int32 MeshBatchIndex = 0; MeshBatchIndex < View.HeterogeneousVolumesMeshBatches.Num(); ++MeshBatchIndex)
 	{
 		const FVolumetricMeshBatch& MeshBatch = View.HeterogeneousVolumesMeshBatches[MeshBatchIndex];
-		bool bIsShadowCast = MeshBatch.Proxy->IsShadowCast(&View);
 
 		// TODO: Is material determiniation too expensive?
 		const FMaterialRenderProxy* MaterialRenderProxy = nullptr;
@@ -1191,7 +1191,14 @@ void CollectHeterogeneousVolumeMeshBatchesForView(
 		MaterialRenderProxy = MaterialRenderProxy ? MaterialRenderProxy : DefaultMaterialRenderProxy;
 		bool bIsVolumeMaterial = Material.GetMaterialDomain() == MD_Volume;
 
-		if (bIsShadowCast && bIsVolumeMaterial)
+		bool bCollectMeshBatch = bIsVolumeMaterial;
+		if (bCollectForShadowCasting)
+		{
+			bool bIsShadowCast = MeshBatch.Proxy->IsShadowCast(&View);
+			bCollectMeshBatch = bCollectMeshBatch && bIsShadowCast;
+		}
+
+		if (bCollectMeshBatch)
 		{
 			HeterogeneousVolumesMeshBatches.FindOrAdd(FVolumetricMeshBatch(MeshBatch.Mesh, MeshBatch.Proxy));
 			WorldBounds = WorldBounds + MeshBatch.Proxy->GetBounds();
@@ -1215,7 +1222,8 @@ void CollectHeterogeneousVolumeMeshBatchesForLight(
 		//for (int32 ViewIndex = 0; ViewIndex < Views.Num(); ViewIndex++)
 		{
 			//const FViewInfo& View = Views[ViewIndex];
-			CollectHeterogeneousVolumeMeshBatchesForView(View, HeterogeneousVolumesMeshBatches, WorldBounds);
+			bool bCollectForShadowCasting = true;
+			CollectHeterogeneousVolumeMeshBatchesForView(View, bCollectForShadowCasting, HeterogeneousVolumesMeshBatches, WorldBounds);
 		}
 
 		for (int32 ShadowIndex = 0; ShadowIndex < VisibleLightInfo->ShadowsToProject.Num(); ++ShadowIndex)
@@ -1821,9 +1829,10 @@ void RenderAdaptiveVolumetricCameraMapWithLiveShading(
 	TRDGUniformBufferRef<FAdaptiveVolumetricShadowMapUniformBufferParameters> AdaptiveVolumetricShadowMapUniformBuffer;
 
 	// Collect all volumes for view
+	bool bCollectForShadowCasting = false;
 	TSet<FVolumetricMeshBatch> HeterogeneousVolumesMeshBatches;
 	FBoxSphereBounds WorldVolumeBounds;
-	CollectHeterogeneousVolumeMeshBatchesForView(View, HeterogeneousVolumesMeshBatches, WorldVolumeBounds);
+	CollectHeterogeneousVolumeMeshBatchesForView(View, bCollectForShadowCasting, HeterogeneousVolumesMeshBatches, WorldVolumeBounds);
 
 	bool bShouldRenderCameraMap = !View.ViewRect.IsEmpty() && !HeterogeneousVolumesMeshBatches.IsEmpty();
 	if (bShouldRenderCameraMap)
