HeterogeneousVolumes: Add option to change compositing order. Composite BEFORE Translucency by default.

r.HeterogeneousVolumes.Composition
Change the order of Heterogeneous Volumes composition (Default = 0)
0: Before Translucency
1: After Translucency

diff --git a/Engine/Source/Runtime/Renderer/Private/DeferredShadingRenderer.cpp b/Engine/Source/Runtime/Renderer/Private/DeferredShadingRenderer.cpp
--- a/Engine/Source/Runtime/Renderer/Private/DeferredShadingRenderer.cpp
+++ b/Engine/Source/Runtime/Renderer/Private/DeferredShadingRenderer.cpp
@@ -3012,6 +3012,13 @@ void FDeferredShadingSceneRenderer::Render(FRDGBuilder& GraphBuilder)
 			Nanite::DrawVisibleBricks( GraphBuilder, *Scene, View, SceneTextures );
 		}
 
+		// Composite Heterogeneous Volumes
+		if (!bHasRayTracedOverlay && ShouldRenderHeterogeneousVolumes(Scene) &&
+			(GetHeterogeneousVolumesComposition() == EHeterogeneousVolumesCompositionType::BeforeTranslucent))
+		{
+			CompositeHeterogeneousVolumes(GraphBuilder, SceneTextures);
+		}
+
 		// Draw translucency.
 		if (!bHasRayTracedOverlay && TranslucencyViewsToRender != ETranslucencyView::None)
 		{
@@ -3211,7 +3218,9 @@ void FDeferredShadingSceneRenderer::Render(FRDGBuilder& GraphBuilder)
 			RenderStationaryLightOverlap(GraphBuilder, SceneTextures, LightingChannelsTexture);
 		}
 
-		if (ShouldRenderHeterogeneousVolumes(Scene) && !bHasRayTracedOverlay)
+		// Composite Heterogeneous Volumes
+		if (!bHasRayTracedOverlay && ShouldRenderHeterogeneousVolumes(Scene) &&
+			(GetHeterogeneousVolumesComposition() == EHeterogeneousVolumesCompositionType::AfterTranslucent))
 		{
 			CompositeHeterogeneousVolumes(GraphBuilder, SceneTextures);
 		}
diff --git a/Engine/Source/Runtime/Renderer/Private/HeterogeneousVolumes/HeterogeneousVolumes.cpp b/Engine/Source/Runtime/Renderer/Private/HeterogeneousVolumes/HeterogeneousVolumes.cpp
--- a/Engine/Source/Runtime/Renderer/Private/HeterogeneousVolumes/HeterogeneousVolumes.cpp
+++ b/Engine/Source/Runtime/Renderer/Private/HeterogeneousVolumes/HeterogeneousVolumes.cpp
@@ -32,6 +32,16 @@ static TAutoConsoleVariable<int32> CVarTranslucencyHeterogeneousVolumes(
 	ECVF_RenderThreadSafe | ECVF_ReadOnly
 );
 
+static TAutoConsoleVariable<int32> CVarHeterogeneousVolumesComposition(
+	TEXT("r.HeterogeneousVolumes.Composition"),
+	0,
+	TEXT("Change the order of Heterogeneous Volumes composition (Default = 0)\n")
+	TEXT("0: Before Translucency\n")
+	TEXT("1: After Translucency\n")
+	TEXT("Requires enabling Heterogeneous Volumes Project Setting: 'Composite with Translucency'"),
+	ECVF_RenderThreadSafe
+);
+
 static TAutoConsoleVariable<int32> CVarHeterogeneousVolumesShadowMode(
 	TEXT("r.HeterogeneousVolumes.Shadows.Mode"),
 	0,
@@ -184,8 +194,8 @@ static TAutoConsoleVariable<int32> CVarHeterogeneousVolumesLightingCache(
 
 static TAutoConsoleVariable<int32> CVarHeterogeneousVolumesLightingCacheUseAVSM(
 	TEXT("r.HeterogeneousVolumes.LightingCache.UseAVSM"),
-	0,
-	TEXT("Enables use of AVSMs when evaluating self-shadowing (Default = 0)"),
+	1,
+	TEXT("Enables use of AVSMs when evaluating self-shadowing (Default = 1)"),
 	ECVF_RenderThreadSafe
 );
 
@@ -244,6 +254,30 @@ bool ShouldCompositeHeterogeneousVolumesWithTranslucency()
 	return CVarTranslucencyHeterogeneousVolumes.GetValueOnAnyThread() != 0;
 }
 
+EHeterogeneousVolumesCompositionType GetHeterogeneousVolumesCompositionType()
+{
+	int32 CompositionOrder = CVarHeterogeneousVolumesComposition.GetValueOnRenderThread();
+	switch (CompositionOrder)
+	{
+		case 0:
+		default:
+			return EHeterogeneousVolumesCompositionType::BeforeTranslucent;
+		case 1:
+			return EHeterogeneousVolumesCompositionType::AfterTranslucent;
+	}
+}
+
+EHeterogeneousVolumesCompositionType GetHeterogeneousVolumesComposition()
+{
+	// Composition order can only be modified if the Project Setting is enabled
+	if (!ShouldCompositeHeterogeneousVolumesWithTranslucency())
+	{
+		return EHeterogeneousVolumesCompositionType::AfterTranslucent;
+	}
+
+	return GetHeterogeneousVolumesCompositionType();
+}
+
 bool ShouldRenderHeterogeneousVolumes(
 	const FScene* Scene
 )
diff --git a/Engine/Source/Runtime/Renderer/Private/HeterogeneousVolumes/HeterogeneousVolumes.h b/Engine/Source/Runtime/Renderer/Private/HeterogeneousVolumes/HeterogeneousVolumes.h
--- a/Engine/Source/Runtime/Renderer/Private/HeterogeneousVolumes/HeterogeneousVolumes.h
+++ b/Engine/Source/Runtime/Renderer/Private/HeterogeneousVolumes/HeterogeneousVolumes.h
@@ -44,6 +44,13 @@ bool ShouldRenderMeshBatchWithHeterogeneousVolumes(
 bool ShouldCompositeHeterogeneousVolumesWithTranslucency();
 bool ShouldHeterogeneousVolumesCastShadows();
 
+enum class EHeterogeneousVolumesCompositionType : uint8
+{
+	BeforeTranslucent,
+	AfterTranslucent
+};
+EHeterogeneousVolumesCompositionType GetHeterogeneousVolumesComposition();
+
 //
 // Internal API
 //
