HeterogeneousVolumes: Selectively apply distortion when Camera AVSM is below user-defined threshold

This ability requires the Heterogeneous Volume Project Setting "Composite with Translucency" to be enabled.
r.HeterogeneousVolumes.Refraction.UseAVSM Enables AVSM lookup (Default = 1)
r.HeterogeneousVolumes.Refraction.TransmittanceThreshold Minimum transmittance threshold to apply distortion (Default 0.9)

diff --git a/Engine/Shaders/Private/DistortAccumulatePS.usf b/Engine/Shaders/Private/DistortAccumulatePS.usf
--- a/Engine/Shaders/Private/DistortAccumulatePS.usf
+++ b/Engine/Shaders/Private/DistortAccumulatePS.usf
@@ -22,6 +22,15 @@
 #include "/Engine/Generated/VertexFactory.ush"
 #include "DistortionCommon.ush"
 
+#ifndef ADAPTIVE_VOLUMETRIC_SHADOW_MAP
+#define ADAPTIVE_VOLUMETRIC_SHADOW_MAP 0
+#endif // ADAPTIVE_VOLUMETRIC_SHADOW_MAP
+
+#if ADAPTIVE_VOLUMETRIC_SHADOW_MAP
+#define AVSM DistortionPass.AVSM
+#include "HeterogeneousVolumes/HeterogeneousVolumesAdaptiveVolumetricShadowMapSampling.ush"
+#endif // ADAPTIVE_VOLUMETRIC_SHADOW_MAP
+
 // SUBSTRATE_TODO enable distortion on mobile.
 #if SUBSTRATE_ENABLED && (!MATERIAL_IS_SUBSTRATE || SHADING_PATH_MOBILE)
 #undef SUBSTRATE_ENABLED
@@ -218,6 +227,16 @@ void Main(
 
 	// output positives in R|G channels and negatives in B|A channels
 	OutColor = float4(PosOffset.x,PosOffset.y,NegOffset.x,NegOffset.y);
+
+#if ADAPTIVE_VOLUMETRIC_SHADOW_MAP
+	if (DistortionPass.UseAVSM)
+	{
+		float3 TranslatedWorldPosition = SvPositionToResolvedTranslatedWorld(SvPosition);
+		float Transmittance = saturate(AVSM_SampleTransmittance(TranslatedWorldPosition, 0));
+		float ApplyDistortion = (Transmittance >= DistortionPass.TransmittanceThreshold) ? 1.0 : 0.0;
+		OutColor *= ApplyDistortion;
+	}
+#endif
 }
 
 
diff --git a/Engine/Source/Runtime/Renderer/Private/DistortionRendering.cpp b/Engine/Source/Runtime/Renderer/Private/DistortionRendering.cpp
--- a/Engine/Source/Runtime/Renderer/Private/DistortionRendering.cpp
+++ b/Engine/Source/Runtime/Renderer/Private/DistortionRendering.cpp
@@ -18,6 +18,7 @@
 #include "ScreenRendering.h"
 #include "PostProcess/TemporalAA.h"
 #include "ShaderPlatformCachedIniValue.h"
+#include "HeterogeneousVolumes/HeterogeneousVolumes.h"
 
 DECLARE_GPU_DRAWCALL_STAT(Distortion);
 
@@ -57,10 +58,27 @@ static TAutoConsoleVariable<float> CVarRefractionBlurMaxExposedLuminance(
 	TEXT("Clamp scene pre-exposed luminance to this maximum value. It helps to reduce bright specular highlight flickering, even when r.Refraction.Blur.TemporalAA=1."),
 	ECVF_Scalability | ECVF_RenderThreadSafe);
 
+static TAutoConsoleVariable<int32> CVarRefractionUseAVSM(
+	TEXT("r.HeterogeneousVolumes.CompositeWithTranslucency.Refraction.UseAVSM"),
+	1,
+	TEXT("Enables AVSM lookup (Default = 1)\n")
+	TEXT("Requires enabling Heterogeneous Volumes Project Setting: 'Composite with Translucency'"),
+	ECVF_RenderThreadSafe);
+
+static TAutoConsoleVariable<float> CVarRefractionTransmittanceThreshold(
+	TEXT("r.HeterogeneousVolumes.CompositeWithTranslucency.Refraction.TransmittanceThreshold"),
+	0.9,
+	TEXT("Minimum transmittance threshold to apply distortion (Default = 0.9)\n")
+	TEXT("Requires enabling Heterogeneous Volumes Project Setting: 'Composite with Translucency'"),
+	ECVF_RenderThreadSafe);
+
 BEGIN_GLOBAL_SHADER_PARAMETER_STRUCT(FDistortionPassUniformParameters, RENDERER_API)
 	SHADER_PARAMETER_STRUCT(FSceneTextureUniformParameters, SceneTextures)
 	SHADER_PARAMETER_STRUCT(FSubstrateForwardPassUniformParameters, Substrate)
 	SHADER_PARAMETER(FVector4f, DistortionParams)
+	SHADER_PARAMETER_STRUCT(FAdaptiveVolumetricShadowMapUniformBufferParameters, AVSM)
+	SHADER_PARAMETER(int32, UseAVSM)
+	SHADER_PARAMETER(float, TransmittanceThreshold)
 END_GLOBAL_SHADER_PARAMETER_STRUCT()
 
 IMPLEMENT_STATIC_UNIFORM_BUFFER_STRUCT(FDistortionPassUniformParameters, "DistortionPass", SceneTextures);
@@ -110,6 +128,10 @@ TRDGUniformBufferRef<FDistortionPassUniformParameters> CreateDistortionPassUnifo
 	SetupSceneTextureUniformParameters(GraphBuilder, View.GetSceneTexturesChecked(), View.FeatureLevel, ESceneTextureSetupMode::All, Parameters->SceneTextures);
 	SetupDistortionParams(Parameters->DistortionParams, View);
 	Substrate::BindSubstrateForwardPasslUniformParameters(GraphBuilder, View, Parameters->Substrate);
+	Parameters->AVSM = HeterogeneousVolumes::GetAdaptiveVolumetricCameraMapParameters(GraphBuilder, View.ViewState);
+	Parameters->UseAVSM = CVarRefractionUseAVSM.GetValueOnRenderThread() != 0;
+	Parameters->TransmittanceThreshold = FMath::Clamp(CVarRefractionTransmittanceThreshold.GetValueOnRenderThread(), 0.0, 1.0);
+
 	return GraphBuilder.CreateUniformBuffer(Parameters);
 }
 
@@ -227,6 +249,7 @@ public:
 
 		// Skip the material clip if depth test should not be done
 		OutEnvironment.SetDefine(TEXT("MATERIAL_SHOULD_DISABLE_DEPTH_TEST"), Parameters.MaterialParameters.bShouldDisableDepthTest ? 1 : 0);
+		OutEnvironment.SetDefine(TEXT("ADAPTIVE_VOLUMETRIC_SHADOW_MAP"), ShouldCompositeHeterogeneousVolumesWithTranslucency());
 	}
 };
 
