Converted skylight distance illumination to a buffer to avoid requiring an extra sampler just for a single color sampe that do not need to be filtered.
Tested disabling distant code running, hiding the sky and mobile.

diff --git a/Engine/Shaders/Private/HeightFogCommon.ush b/Engine/Shaders/Private/HeightFogCommon.ush
--- a/Engine/Shaders/Private/HeightFogCommon.ush
+++ b/Engine/Shaders/Private/HeightFogCommon.ush
@@ -75,7 +75,7 @@ float3 ComputeInscatteringColor(float3 CameraToReceiver, float CameraToReceiverL
 #endif
 
 #if PROJECT_SUPPORT_SKY_ATMOSPHERE_AFFECTS_HEIGHFOG
-	Inscattering += FogStruct.SkyAtmosphereAmbientContributionColorScale.rgb * View.SkyAtmosphereHeightFogContribution.xxx * Texture2DSampleLevel(View.DistantSkyLightLutTexture, View.DistantSkyLightLutTextureSampler, float2(0.5f, 0.5f), 0.0f).rgb;
+	Inscattering += FogStruct.SkyAtmosphereAmbientContributionColorScale.rgb * View.SkyAtmosphereHeightFogContribution.xxx * View.DistantSkyLightLutBufferSRV[0].rgb;
 #endif
 
 	return Inscattering;
diff --git a/Engine/Shaders/Private/MaterialTemplate.ush b/Engine/Shaders/Private/MaterialTemplate.ush
--- a/Engine/Shaders/Private/MaterialTemplate.ush
+++ b/Engine/Shaders/Private/MaterialTemplate.ush
@@ -2236,8 +2236,7 @@ float4 MaterialExpressionSkyAtmosphereAerialPerspective(FMaterialPixelParameters
 float3 MaterialExpressionSkyAtmosphereDistantLightScatteredLuminance(FMaterialPixelParameters Parameters)
 {
 #if MATERIAL_SKY_ATMOSPHERE && PROJECT_SUPPORT_SKY_ATMOSPHERE
-	// TODO load on platforms supporting it
-	return Texture2DSampleLevel(View.DistantSkyLightLutTexture, View.DistantSkyLightLutTextureSampler, float2(0.5f, 0.5f), 0.0f).rgb;
+	return View.DistantSkyLightLutBufferSRV[0].rgb;
 #else
 	return float3(0.0f, 0.0f, 0.0f);
 #endif
diff --git a/Engine/Shaders/Private/SkyAtmosphere.usf b/Engine/Shaders/Private/SkyAtmosphere.usf
--- a/Engine/Shaders/Private/SkyAtmosphere.usf
+++ b/Engine/Shaders/Private/SkyAtmosphere.usf
@@ -1321,11 +1321,7 @@ void RenderSkyViewLutCS(uint3 ThreadId : SV_DispatchThreadID)
 ////////////////////////////////////////////////////////////
 
 // see SkyAtmosphereRendering.cpp GetSkyLutTextureFormat
-#if SHADING_PATH_MOBILE
-RWTexture2D<float4> DistantSkyLightLutUAV;
-#else
-RWTexture2D<float3> DistantSkyLightLutUAV;
-#endif
+RWBuffer<float4> DistantSkyLightLutBufferUAV;
 //Buffer<float4> UniformSphereSamplesBuffer;
 float4 AtmosphereLightDirection0;
 float4 AtmosphereLightDirection1;
@@ -1416,12 +1412,10 @@ void RenderDistantSkyLightLutCS(uint3 ThreadId : SV_DispatchThreadID)
 		const float SamplerSolidAngle = 4.0f * PI / float(THREADGROUP_SIZE * THREADGROUP_SIZE);
 		const float3 Illuminance = AccumulatedLuminanceSamples * SamplerSolidAngle;
 		const float3 UniformPhaseFunction = 1.0f / (4.0f * PI);
+		
 		// Luminance assuming scattering in a medium with a uniform phase function.
-#if SHADING_PATH_MOBILE
-		DistantSkyLightLutUAV[int2(PixPos)] = float4(Illuminance * UniformPhaseFunction, 0.0f);
-#else
-		DistantSkyLightLutUAV[int2(PixPos)] = Illuminance * UniformPhaseFunction;
-#endif
+		DistantSkyLightLutBufferUAV[0] = float4(Illuminance * UniformPhaseFunction, 0.0f);
+		
 		// Since this is ran once per scene, we do not have access to view data (VIEWDATA_AVAILABLE==0). 
 		// So this buffer is not pre-exposed today.
 	}
diff --git a/Engine/Shaders/Private/VolumetricCloud.usf b/Engine/Shaders/Private/VolumetricCloud.usf
--- a/Engine/Shaders/Private/VolumetricCloud.usf
+++ b/Engine/Shaders/Private/VolumetricCloud.usf
@@ -680,7 +680,7 @@ void MainCommon(in FMaterialPixelParameters MaterialParameters, in float4 SvPosi
 				// TODO Only works for ground view. This should be spatially varying for space viewsspace view. 
 				// TODO Have a illuminance texture to get that overall ambient contribution? 
 				// TODO An approximation would be to multiply it with sun transmittance?
-				DistantSkyLightLuminance = Texture2DSampleLevel(View.DistantSkyLightLutTexture, View.DistantSkyLightLutTextureSampler, float2(0.5f, 0.5f), 0.0f).rgb;
+				DistantSkyLightLuminance = View.DistantSkyLightLutBufferSRV[0].rgb;
 			}
 			else
 			{
diff --git a/Engine/Source/Runtime/Engine/Private/SceneManagement.cpp b/Engine/Source/Runtime/Engine/Private/SceneManagement.cpp
--- a/Engine/Source/Runtime/Engine/Private/SceneManagement.cpp
+++ b/Engine/Source/Runtime/Engine/Private/SceneManagement.cpp
@@ -1158,8 +1158,7 @@ FViewUniformShaderParameters::FViewUniformShaderParameters()
 	SkyViewLutTexture = GBlackTexture->TextureRHI;
 	SkyViewLutTextureSampler = TStaticSamplerState<SF_Bilinear>::GetRHI();
 
-	DistantSkyLightLutTexture = GBlackTexture->TextureRHI;
-	DistantSkyLightLutTextureSampler = TStaticSamplerState<SF_Point, AM_Wrap, AM_Wrap>::GetRHI();
+	DistantSkyLightLutBufferSRV = GBlackVertexBufferWithSRV->ShaderResourceViewRHI;
 
 	CameraAerialPerspectiveVolume = GBlackAlpha1VolumeTexture->TextureRHI;
 	CameraAerialPerspectiveVolumeSampler = TStaticSamplerState<SF_Bilinear>::GetRHI();
diff --git a/Engine/Source/Runtime/Engine/Public/SceneView.h b/Engine/Source/Runtime/Engine/Public/SceneView.h
--- a/Engine/Source/Runtime/Engine/Public/SceneView.h
+++ b/Engine/Source/Runtime/Engine/Public/SceneView.h
@@ -1084,14 +1084,13 @@ BEGIN_GLOBAL_SHADER_PARAMETER_STRUCT_WITH_CONSTRUCTOR(FViewUniformShaderParamete
 	SHADER_PARAMETER_SAMPLER(SamplerState, TransmittanceLutTextureSampler)
 	SHADER_PARAMETER_TEXTURE(Texture2D, SkyViewLutTexture)
 	SHADER_PARAMETER_SAMPLER(SamplerState, SkyViewLutTextureSampler)
-	SHADER_PARAMETER_TEXTURE(Texture2D, DistantSkyLightLutTexture)
-	SHADER_PARAMETER_SAMPLER(SamplerState, DistantSkyLightLutTextureSampler)
 	SHADER_PARAMETER_TEXTURE(Texture3D, CameraAerialPerspectiveVolume)
 	SHADER_PARAMETER_SAMPLER(SamplerState, CameraAerialPerspectiveVolumeSampler)
 	SHADER_PARAMETER_TEXTURE(Texture3D, CameraAerialPerspectiveVolumeMieOnly)
 	SHADER_PARAMETER_SAMPLER(SamplerState, CameraAerialPerspectiveVolumeMieOnlySampler)
 	SHADER_PARAMETER_TEXTURE(Texture3D, CameraAerialPerspectiveVolumeRayOnly)
 	SHADER_PARAMETER_SAMPLER(SamplerState, CameraAerialPerspectiveVolumeRayOnlySampler)
+	SHADER_PARAMETER_SRV(Buffer<float4>, DistantSkyLightLutBufferSRV)
 	// Hair
 	SHADER_PARAMETER_TEXTURE(Texture3D, HairScatteringLUTTexture)
 	SHADER_PARAMETER_SAMPLER(SamplerState, HairScatteringLUTSampler)
diff --git a/Engine/Source/Runtime/RenderCore/Private/GlobalRenderResources.cpp b/Engine/Source/Runtime/RenderCore/Private/GlobalRenderResources.cpp
--- a/Engine/Source/Runtime/RenderCore/Private/GlobalRenderResources.cpp
+++ b/Engine/Source/Runtime/RenderCore/Private/GlobalRenderResources.cpp
@@ -153,6 +153,27 @@ public:
 
 FVertexBufferWithSRV* GWhiteVertexBufferWithSRV = new TGlobalResource<FWhiteVertexBuffer, FRenderResource::EInitPhase::Pre>;
 
+class FBlackVertexBuffer : public FVertexBufferWithSRV
+{
+public:
+	virtual void InitRHI(FRHICommandListBase& RHICmdList) override
+	{
+		// Create the texture RHI.  		
+		FRHIResourceCreateInfo CreateInfo(TEXT("BlackVertexBuffer"));
+
+		VertexBufferRHI = RHICmdList.CreateVertexBuffer(sizeof(FVector4f), BUF_Static | BUF_ShaderResource, CreateInfo);
+
+		FVector4f* BufferData = (FVector4f*)RHICmdList.LockBuffer(VertexBufferRHI, 0, sizeof(FVector4f), RLM_WriteOnly);
+		*BufferData = FVector4f(0.0f, 0.0f, 0.0f, 0.0f);
+		RHICmdList.UnlockBuffer(VertexBufferRHI);
+
+		// Create a view of the buffer
+		ShaderResourceViewRHI = RHICmdList.CreateShaderResourceView(VertexBufferRHI, sizeof(FVector4f), PF_A32B32G32R32F);
+	}
+};
+
+FVertexBufferWithSRV* GBlackVertexBufferWithSRV = new TGlobalResource<FBlackVertexBuffer, FRenderResource::EInitPhase::Pre>;
+
 class FWhiteVertexBufferWithRDG : public FBufferWithRDG
 {
 public:
diff --git a/Engine/Source/Runtime/RenderCore/Public/GlobalRenderResources.h b/Engine/Source/Runtime/RenderCore/Public/GlobalRenderResources.h
--- a/Engine/Source/Runtime/RenderCore/Public/GlobalRenderResources.h
+++ b/Engine/Source/Runtime/RenderCore/Public/GlobalRenderResources.h
@@ -28,9 +28,12 @@ extern RENDERCORE_API FVertexBufferWithSRV* GEmptyVertexBufferWithUAV;
 // An empty structured buffer with a UAV
 extern RENDERCORE_API FVertexBufferWithSRV* GEmptyStructuredBufferWithUAV;
 
-// An empty vertex buffer with a SRV
+// An white vertex buffer with a SRV
 extern RENDERCORE_API FVertexBufferWithSRV* GWhiteVertexBufferWithSRV;
 
+// An black vertex buffer with a SRV
+extern RENDERCORE_API FVertexBufferWithSRV* GBlackVertexBufferWithSRV;
+
 // A white vertex buffer used with RenderGraph.
 extern RENDERCORE_API FBufferWithRDG* GWhiteVertexBufferWithRDG;
 
diff --git a/Engine/Source/Runtime/Renderer/Private/MobileShadingRenderer.cpp b/Engine/Source/Runtime/Renderer/Private/MobileShadingRenderer.cpp
--- a/Engine/Source/Runtime/Renderer/Private/MobileShadingRenderer.cpp
+++ b/Engine/Source/Runtime/Renderer/Private/MobileShadingRenderer.cpp
@@ -639,7 +639,7 @@ void FMobileSceneRenderer::InitViews(
 		// Initialise Sky/View resources before the view global uniform buffer is built.
 		if (ShouldRenderSkyAtmosphere(Scene, ViewFamily.EngineShowFlags))
 		{
-			InitSkyAtmosphereForViews(RHICmdList);
+			InitSkyAtmosphereForViews(RHICmdList, GraphBuilder);
 		}
 
 		if (bRequiresShadowProjections)
diff --git a/Engine/Source/Runtime/Renderer/Private/SceneRendering.cpp b/Engine/Source/Runtime/Renderer/Private/SceneRendering.cpp
--- a/Engine/Source/Runtime/Renderer/Private/SceneRendering.cpp
+++ b/Engine/Source/Runtime/Renderer/Private/SceneRendering.cpp
@@ -1286,7 +1286,7 @@ void FViewInfo::SetupUniformBufferParameters(
 	FRHITexture* CameraAerialPerspectiveVolumeFound = nullptr;
 	FRHITexture* CameraAerialPerspectiveVolumeMieOnlyFound = nullptr;
 	FRHITexture* CameraAerialPerspectiveVolumeRayOnlyFound = nullptr;
-	FRHITexture* DistantSkyLightLutTextureFound = nullptr;
+	FRHIShaderResourceView* DistantSkyLightLutBufferSRVFound = nullptr;
 	if (ShouldRenderSkyAtmosphere(Scene, Family->EngineShowFlags))
 	{
 		ViewUniformShaderParameters.SkyAtmospherePresentInScene = 1.0f;
@@ -1302,11 +1302,8 @@ void FViewInfo::SetupUniformBufferParameters(
 		{
 			TransmittanceLutTextureFound = PooledTransmittanceLutTexture->GetRHI();
 		}
-		const TRefCountPtr<IPooledRenderTarget>& PooledDistantSkyLightLutTexture = SkyAtmosphere->GetDistantSkyLightLutTexture();
-		if (PooledDistantSkyLightLutTexture.IsValid())
-		{
-			DistantSkyLightLutTextureFound = PooledDistantSkyLightLutTexture->GetRHI();
-		}
+
+		DistantSkyLightLutBufferSRVFound = SkyAtmosphere->GetDistantSkyLightLutBufferSRV();
 
 		if (this->SkyAtmosphereCameraAerialPerspectiveVolume.IsValid())
 		{
@@ -1435,8 +1432,11 @@ void FViewInfo::SetupUniformBufferParameters(
 
 	ViewUniformShaderParameters.TransmittanceLutTexture = OrWhite2DIfNull(TransmittanceLutTextureFound);
 	ViewUniformShaderParameters.TransmittanceLutTextureSampler = TStaticSamplerState<SF_Bilinear>::GetRHI();
-	ViewUniformShaderParameters.DistantSkyLightLutTexture = OrBlack2DIfNull(DistantSkyLightLutTextureFound);
-	ViewUniformShaderParameters.DistantSkyLightLutTextureSampler = TStaticSamplerState<SF_Point, AM_Wrap, AM_Wrap>::GetRHI();
+	ViewUniformShaderParameters.DistantSkyLightLutBufferSRV = GBlackVertexBufferWithSRV->ShaderResourceViewRHI.GetReference();
+	if(DistantSkyLightLutBufferSRVFound != nullptr)
+	{
+		ViewUniformShaderParameters.DistantSkyLightLutBufferSRV = DistantSkyLightLutBufferSRVFound;
+	}
 	ViewUniformShaderParameters.SkyViewLutTexture = OrBlack2DIfNull(SkyViewLutTextureFound);
 	ViewUniformShaderParameters.SkyViewLutTextureSampler = TStaticSamplerState<SF_Bilinear>::GetRHI();
 	ViewUniformShaderParameters.CameraAerialPerspectiveVolume = OrBlack3DAlpha1IfNull(CameraAerialPerspectiveVolumeFound);
diff --git a/Engine/Source/Runtime/Renderer/Private/SceneRendering.h b/Engine/Source/Runtime/Renderer/Private/SceneRendering.h
--- a/Engine/Source/Runtime/Renderer/Private/SceneRendering.h
+++ b/Engine/Source/Runtime/Renderer/Private/SceneRendering.h
@@ -2497,7 +2497,7 @@ protected:
 	void RenderPlanarReflection(class FPlanarReflectionSceneProxy* ReflectionSceneProxy);
 
 	/** Initialise sky atmosphere resources.*/
-	void InitSkyAtmosphereForViews(FRHICommandListImmediate& RHICmdList);
+	void InitSkyAtmosphereForViews(FRHICommandListImmediate& RHICmdList, FRDGBuilder& GraphBuilder);
 	
 	/** Render the sky atmosphere look up table needed for this frame.*/
 	void RenderSkyAtmosphereLookUpTables(FRDGBuilder& GraphBuilder, class FSkyAtmospherePendingRDGResources& PendingRDGResources);
diff --git a/Engine/Source/Runtime/Renderer/Private/SceneVisibility.cpp b/Engine/Source/Runtime/Renderer/Private/SceneVisibility.cpp
--- a/Engine/Source/Runtime/Renderer/Private/SceneVisibility.cpp
+++ b/Engine/Source/Runtime/Renderer/Private/SceneVisibility.cpp
@@ -5519,7 +5519,7 @@ void FDeferredShadingSceneRenderer::BeginInitViews(
 	// Initialise Sky/View resources before the view global uniform buffer is built.
 	if (ShouldRenderSkyAtmosphere(Scene, ViewFamily.EngineShowFlags))
 	{
-		InitSkyAtmosphereForViews(RHICmdList);
+		InitSkyAtmosphereForViews(RHICmdList, GraphBuilder);
 	}
 
 	{
diff --git a/Engine/Source/Runtime/Renderer/Private/SkyAtmosphereRendering.cpp b/Engine/Source/Runtime/Renderer/Private/SkyAtmosphereRendering.cpp
--- a/Engine/Source/Runtime/Renderer/Private/SkyAtmosphereRendering.cpp
+++ b/Engine/Source/Runtime/Renderer/Private/SkyAtmosphereRendering.cpp
@@ -496,13 +496,25 @@ FSkyAtmosphereRenderSceneInfo::~FSkyAtmosphereRenderSceneInfo()
 {
 }
 
-TRefCountPtr<IPooledRenderTarget>& FSkyAtmosphereRenderSceneInfo::GetDistantSkyLightLutTexture()
+void FSkyAtmosphereRenderSceneInfo::CreateDistantSkyLightLutBufferAndSRV(FRDGBuilder& GraphBuilder)
+{
+	DistantSkyLightLutBuffer = AllocatePooledBuffer(FRDGBufferDesc::CreateBufferDesc(sizeof(FVector4f), 1), TEXT("SkyAtmosphere.DistantSkyLightLutBuffer"));
+	DistantSkyLightLutBufferSRV = DistantSkyLightLutBuffer->GetOrCreateSRV(GraphBuilder.RHICmdList, FRHIBufferSRVCreateInfo(PF_A32B32G32R32F));
+}
+
+TRefCountPtr<FRDGPooledBuffer>& FSkyAtmosphereRenderSceneInfo::GetDistantSkyLightLutBuffer()
+{
+	check(CVarSkyAtmosphereDistantSkyLightLUT.GetValueOnRenderThread() > 0);
+	return DistantSkyLightLutBuffer;
+}
+
+FRHIShaderResourceView* FSkyAtmosphereRenderSceneInfo::GetDistantSkyLightLutBufferSRV()
 {
 	if (CVarSkyAtmosphereDistantSkyLightLUT.GetValueOnRenderThread() > 0)
 	{
-		return DistantSkyLightLutTexture;
+		return DistantSkyLightLutBufferSRV;
 	}
-	return GSystemTextures.BlackDummy;
+	return GBlackVertexBufferWithSRV->ShaderResourceViewRHI;
 }
 
 
@@ -825,7 +837,7 @@ public:
 	BEGIN_SHADER_PARAMETER_STRUCT(FParameters, )
 		SHADER_PARAMETER_STRUCT_REF(FAtmosphereUniformShaderParameters, Atmosphere)
 		SHADER_PARAMETER_STRUCT_REF(FSkyAtmosphereInternalCommonParameters, SkyAtmosphere)
-		SHADER_PARAMETER_RDG_TEXTURE_UAV(RWTexture2D<float3>, DistantSkyLightLutUAV)
+		SHADER_PARAMETER_RDG_BUFFER_UAV(RWBuffer<float4>, DistantSkyLightLutBufferUAV)
 		SHADER_PARAMETER_RDG_TEXTURE(Texture2D<float3>, TransmittanceLutTexture)
 		SHADER_PARAMETER_RDG_TEXTURE(Texture2D<float3>, MultiScatteredLuminanceLutTexture)
 		SHADER_PARAMETER_SAMPLER(SamplerState, TransmittanceLutTextureSampler)
@@ -1097,9 +1109,9 @@ TGlobalResource<FUniformSphereSamplesBuffer> GUniformSphereSamplesBuffer;
 	FSceneRenderer functions
 =============================================================================*/
 
-void FSceneRenderer::InitSkyAtmosphereForViews(FRHICommandListImmediate& RHICmdList)
+void FSceneRenderer::InitSkyAtmosphereForViews(FRHICommandListImmediate& RHICmdList, FRDGBuilder& GraphBuilder)
 {
-	InitSkyAtmosphereForScene(RHICmdList, Scene);
+	InitSkyAtmosphereForScene(RHICmdList, GraphBuilder, Scene);
 
 	for (int32 ViewIndex = 0; ViewIndex < Views.Num(); ViewIndex++)
 	{
@@ -1126,7 +1138,7 @@ static EPixelFormat GetSkyLutSmallTextureFormat()
 	return PF_R8G8B8A8;
 }
 
-void InitSkyAtmosphereForScene(FRHICommandListImmediate& RHICmdList, FScene* Scene)
+void InitSkyAtmosphereForScene(FRHICommandListImmediate& RHICmdList, FRDGBuilder& GraphBuilder, FScene* Scene)
 {
 	if (Scene)
 	{
@@ -1165,12 +1177,7 @@ void InitSkyAtmosphereForScene(FRHICommandListImmediate& RHICmdList, FScene* Sce
 
 		if (CVarSkyAtmosphereDistantSkyLightLUT.GetValueOnRenderThread() > 0)
 		{
-			TRefCountPtr<IPooledRenderTarget>& DistantSkyLightLutTexture = SkyInfo.GetDistantSkyLightLutTexture();
-			Desc = FPooledRenderTargetDesc::Create2DDesc(
-				FIntPoint(1, 1),
-				TextureLUTFormat, FClearValueBinding::None, TexCreate_None, TexCreate_ShaderResource | TexCreate_UAV, false);
-			GRenderTargetPool.FindFreeElement(RHICmdList, Desc, DistantSkyLightLutTexture, TEXT("SkyAtmosphere.DistantSkyLightLut"));
-		//	RHICmdList.Transition(FRHITransitionInfo(DistantSkyLightLutTexture->GetRHI(), ERHIAccess::Unknown, ERHIAccess::SRVMask)); // ERHIPipeline::All
+			SkyInfo.CreateDistantSkyLightLutBufferAndSRV(GraphBuilder);
 		}
 	}
 }
@@ -1369,8 +1376,8 @@ void FSceneRenderer::RenderSkyAtmosphereLookUpTables(FRDGBuilder& GraphBuilder,
 	// Distant Sky Light LUT
 	if(CVarSkyAtmosphereDistantSkyLightLUT.GetValueOnRenderThread() > 0)
 	{
-		FRDGTextureRef DistantSkyLightLut = GraphBuilder.RegisterExternalTexture(SkyInfo.GetDistantSkyLightLutTexture());
-		FRDGTextureUAVRef DistantSkyLightLutUAV = GraphBuilder.CreateUAV(FRDGTextureUAVDesc(DistantSkyLightLut, 0));
+		FRDGBufferRef DistantSkyLightLutBuffer = GraphBuilder.RegisterExternalBuffer(SkyInfo.GetDistantSkyLightLutBuffer());
+		FRDGBufferUAVRef DistantSkyLightLutBufferUAV = GraphBuilder.CreateUAV(FRDGBufferUAVDesc(DistantSkyLightLutBuffer, PF_A32B32G32R32F));
 
 		FRenderDistantSkyLightLutCS::FPermutationDomain PermutationVector;
 		PermutationVector.Set<FSecondAtmosphereLight>(bSecondAtmosphereLightEnabled);
@@ -1384,7 +1391,7 @@ void FSceneRenderer::RenderSkyAtmosphereLookUpTables(FRDGBuilder& GraphBuilder,
 		PassParameters->TransmittanceLutTexture = TransmittanceLut;
 		PassParameters->MultiScatteredLuminanceLutTexture = MultiScatteredLuminanceLut;
 		PassParameters->UniformSphereSamplesBuffer = GUniformSphereSamplesBuffer.UniformSphereSamplesBuffer.SRV;
-		PassParameters->DistantSkyLightLutUAV = DistantSkyLightLutUAV;
+		PassParameters->DistantSkyLightLutBufferUAV = DistantSkyLightLutBufferUAV;
 
 		FLightSceneInfo* Light0 = Scene->AtmosphereLights[0];
 		FLightSceneInfo* Light1 = Scene->AtmosphereLights[1];
@@ -1414,7 +1421,7 @@ void FSceneRenderer::RenderSkyAtmosphereLookUpTables(FRDGBuilder& GraphBuilder,
 		const FIntVector NumGroups = FIntVector::DivideAndRoundUp(TextureSize, FRenderDistantSkyLightLutCS::GroupSize);
 		FComputeShaderUtils::AddPass(GraphBuilder, RDG_EVENT_NAME("DistantSkyLightLut"), PassFlag, ComputeShader, PassParameters, NumGroups);
 
-		PendingRDGResources.DistantSkyLightLut = DistantSkyLightLut;
+		PendingRDGResources.DistantSkyLightLutBuffer = DistantSkyLightLutBuffer;
 	}
 
 	SkyAtmosphereLightShadowData LightShadowData;
@@ -1705,9 +1712,9 @@ void FSkyAtmospherePendingRDGResources::CommitToSceneAndViewUniformBuffers(FRDGB
 	FScene* Scene = SceneRenderer->Scene;
 	FSkyAtmosphereRenderSceneInfo& SkyInfo = *Scene->GetSkyAtmosphereSceneInfo();
 
-	if (DistantSkyLightLut)
+	if (DistantSkyLightLutBuffer && CVarSkyAtmosphereDistantSkyLightLUT.GetValueOnRenderThread() > 0)
 	{
-		SkyInfo.GetDistantSkyLightLutTexture() = ConvertToExternalAccessTexture(GraphBuilder, ExternalAccessQueue, DistantSkyLightLut, ERHIAccess::SRVMask, ERHIPipeline::All);
+		SkyInfo.GetDistantSkyLightLutBuffer() = ConvertToExternalAccessBuffer(GraphBuilder, ExternalAccessQueue, DistantSkyLightLutBuffer, ERHIAccess::SRVMask, ERHIPipeline::All);
 	}
 
 	if (RealTimeReflectionCaptureSkyAtmosphereViewLutTexture)
diff --git a/Engine/Source/Runtime/Renderer/Private/SkyAtmosphereRendering.h b/Engine/Source/Runtime/Renderer/Private/SkyAtmosphereRendering.h
--- a/Engine/Source/Runtime/Renderer/Private/SkyAtmosphereRendering.h
+++ b/Engine/Source/Runtime/Renderer/Private/SkyAtmosphereRendering.h
@@ -144,7 +144,10 @@ public:
 	const TUniformBufferRef<FAtmosphereUniformShaderParameters>& GetAtmosphereUniformBuffer() { return AtmosphereUniformBuffer; }
 	TRefCountPtr<IPooledRenderTarget>& GetTransmittanceLutTexture() { return TransmittanceLutTexture; }
 	TRefCountPtr<IPooledRenderTarget>& GetMultiScatteredLuminanceLutTexture() { return MultiScatteredLuminanceLutTexture; }
-	TRefCountPtr<IPooledRenderTarget>& GetDistantSkyLightLutTexture();
+
+	void CreateDistantSkyLightLutBufferAndSRV(FRDGBuilder& GraphBuilder);
+	TRefCountPtr<FRDGPooledBuffer>& GetDistantSkyLightLutBuffer();
+	FRHIShaderResourceView* GetDistantSkyLightLutBufferSRV();
 
 	FRDGTextureRef GetTransmittanceLutTexture(FRDGBuilder& GraphBuilder) const { return GraphBuilder.RegisterExternalTexture(TransmittanceLutTexture); }
 
@@ -165,7 +168,8 @@ private:
 
 	TRefCountPtr<IPooledRenderTarget> TransmittanceLutTexture;
 	TRefCountPtr<IPooledRenderTarget> MultiScatteredLuminanceLutTexture;
-	TRefCountPtr<IPooledRenderTarget> DistantSkyLightLutTexture;
+	TRefCountPtr<FRDGPooledBuffer> DistantSkyLightLutBuffer;
+	FRHIShaderResourceView* DistantSkyLightLutBufferSRV = nullptr;
 };
 
 /** Pending RDG resource to commit after the pre-pass / nanite rasterization so that RenderSkyAtmosphereLookUpTables() can overlap them on async compute. */
@@ -185,7 +189,7 @@ private:
 	FSceneRenderer* SceneRenderer = nullptr;
 	TArray<FViewRDGResources, TInlineAllocator<4>> ViewResources;
 
-	FRDGTextureRef DistantSkyLightLut = nullptr;
+	FRDGBufferRef DistantSkyLightLutBuffer = nullptr;
 	FRDGTextureRef RealTimeReflectionCaptureSkyAtmosphereViewLutTexture = nullptr;
 	FRDGTextureRef RealTimeReflectionCaptureCamera360APLutTexture = nullptr;
 	FRDGTextureRef TransmittanceLut = nullptr;
@@ -211,7 +215,7 @@ extern ESkyAtmospherePassLocation GetSkyAtmospherePassLocation();
 
 bool ShouldRenderSkyAtmosphere(const FScene* Scene, const FEngineShowFlags& EngineShowFlags);
 
-void InitSkyAtmosphereForScene(FRHICommandListImmediate& RHICmdList, FScene* Scene);
+void InitSkyAtmosphereForScene(FRHICommandListImmediate& RHICmdList, FRDGBuilder& GraphBuilder, FScene* Scene);
 void InitSkyAtmosphereForView(FRHICommandListImmediate& RHICmdList, const FScene* Scene, FViewInfo& View);
 
 extern void SetupSkyAtmosphereViewSharedUniformShaderParameters(const class FViewInfo& View, const FSkyAtmosphereSceneProxy& SkyAtmosphereProxy, FSkyAtmosphereViewSharedUniformShaderParameters& OutParameters);
