Changed visibility task checks to waits. Given the complexity of the task graph, it's possible for certain stages of tasks to complete before the full task event of the prior stage is marked as 'complete'. In this case, the mesh pass setup task was completing prior to the compute relevance task firing. This change simplifies things to wait for each stage in turn.

diff --git a/Engine/Source/Runtime/Renderer/Private/SceneVisibility.cpp b/Engine/Source/Runtime/Renderer/Private/SceneVisibility.cpp
--- a/Engine/Source/Runtime/Renderer/Private/SceneVisibility.cpp
+++ b/Engine/Source/Runtime/Renderer/Private/SceneVisibility.cpp
@@ -4604,13 +4604,11 @@ void FVisibilityTaskData::Finish()
 {
 	SCOPED_NAMED_EVENT(FinishVisibility, FColor::Magenta);
 
+	Tasks.ComputeRelevance.Wait();
+	Tasks.FinalizeRelevance.Wait();
+	Tasks.DynamicMeshElements.Wait();
 	Tasks.MeshPassSetup.Wait();
 
-	check(Tasks.ComputeRelevance.IsCompleted());
-	check(Tasks.FinalizeRelevance.IsCompleted());
-	check(Tasks.DynamicMeshElements.IsCompleted());
-	check(Tasks.MeshPassSetup.IsCompleted());
-
 	ViewPackets.Empty();
 	DynamicMeshElements.DynamicPrimitives.Empty();
 	Allocator.BulkDelete();
