Minor optimization for MobileDynamicSpotlightShadows array management

diff --git a/Engine/Source/Runtime/Renderer/Private/ShadowSetup.cpp b/Engine/Source/Runtime/Renderer/Private/ShadowSetup.cpp
--- a/Engine/Source/Runtime/Renderer/Private/ShadowSetup.cpp
+++ b/Engine/Source/Runtime/Renderer/Private/ShadowSetup.cpp
@@ -6497,7 +6497,7 @@ void FSceneRenderer::FilterDynamicShadows(FDynamicShadowsTaskData& TaskData)
 		WholeSceneDirectionalShadowsForLight.Reset();
 	}
 
-	if (bMobile)
+	if (MobileDynamicSpotlightShadows.Num() > 0)
 	{
 		// AllocateMobileCSMAndSpotLightShadowDepthTargets would only allocate a single large render target for all shadows, so if the requirement exceeds the MaxTextureSize, the rest of the shadows will not get space for rendering
 		// So we sort spotlight shadows and append them at the last to make sure csm will get space in any case.
