Add cached MDC code for MeshDecals to reduce overhead on rendering thread when rendering LOTs of decals

diff --git a/Engine/Source/Runtime/Renderer/Private/CompositionLighting/CompositionLighting.cpp b/Engine/Source/Runtime/Renderer/Private/CompositionLighting/CompositionLighting.cpp
--- a/Engine/Source/Runtime/Renderer/Private/CompositionLighting/CompositionLighting.cpp
+++ b/Engine/Source/Runtime/Renderer/Private/CompositionLighting/CompositionLighting.cpp
@@ -493,7 +493,7 @@ static FScreenPassTexture AddPostProcessingAmbientOcclusion(
 	return FinalOutput;
 }
 
-FCompositionLighting::FCompositionLighting(TArrayView<const FViewInfo> InViews, const FSceneTextures& InSceneTextures, TUniqueFunction<bool(int32)> RequestSSAOFunction)
+FCompositionLighting::FCompositionLighting(TArrayView<FViewInfo> InViews, const FSceneTextures& InSceneTextures, TUniqueFunction<bool(int32)> RequestSSAOFunction)
 	: Views(InViews)
 	, ViewFamily(*InViews[0].Family)
 	, SceneTextures(InSceneTextures)
@@ -576,7 +576,7 @@ void FCompositionLighting::ProcessBeforeBasePass(FRDGBuilder& GraphBuilder, FDBu
 
 	for (int32 ViewIndex = 0; ViewIndex < Views.Num(); ++ViewIndex)
 	{
-		const FViewInfo& View = Views[ViewIndex];
+		FViewInfo& View = Views[ViewIndex];
 		const FAOConfig& ViewConfig = ViewAOConfigs[ViewIndex];
 
 		const bool bEnableSSAO = ViewConfig.SSAOLocation == ESSAOLocation::BeforeBasePass;
@@ -626,7 +626,7 @@ void FCompositionLighting::ProcessAfterBasePass(FRDGBuilder& GraphBuilder, FInst
 
 	for (int32 ViewIndex = 0; ViewIndex < Views.Num(); ++ViewIndex)
 	{
-		const FViewInfo& View = Views[ViewIndex];
+		FViewInfo& View = Views[ViewIndex];
 		const FAOConfig& ViewConfig = ViewAOConfigs[ViewIndex];
 
 		RDG_GPU_MASK_SCOPE(GraphBuilder, View.GPUMask);
diff --git a/Engine/Source/Runtime/Renderer/Private/CompositionLighting/CompositionLighting.h b/Engine/Source/Runtime/Renderer/Private/CompositionLighting/CompositionLighting.h
--- a/Engine/Source/Runtime/Renderer/Private/CompositionLighting/CompositionLighting.h
+++ b/Engine/Source/Runtime/Renderer/Private/CompositionLighting/CompositionLighting.h
@@ -18,7 +18,7 @@
 class FCompositionLighting
 {
 public:
-	FCompositionLighting(TArrayView<const FViewInfo> InViews, const FSceneTextures& InSceneTextures, TUniqueFunction<bool(int32)> RequestSSAOFunction);
+	FCompositionLighting(TArrayView<FViewInfo> InViews, const FSceneTextures& InSceneTextures, TUniqueFunction<bool(int32)> RequestSSAOFunction);
 
 	void ProcessAfterOcclusion(FRDGBuilder& GraphBuilder);
 
@@ -36,7 +36,7 @@ public:
 private:
 	void TryInit();
 
-	const TArrayView<const FViewInfo> Views;
+	const TArrayView<FViewInfo> Views;
 	const FSceneViewFamily& ViewFamily;
 	const FSceneTextures& SceneTextures;
 
diff --git a/Engine/Source/Runtime/Renderer/Private/CompositionLighting/PostProcessDeferredDecals.cpp b/Engine/Source/Runtime/Renderer/Private/CompositionLighting/PostProcessDeferredDecals.cpp
--- a/Engine/Source/Runtime/Renderer/Private/CompositionLighting/PostProcessDeferredDecals.cpp
+++ b/Engine/Source/Runtime/Renderer/Private/CompositionLighting/PostProcessDeferredDecals.cpp
@@ -566,7 +566,7 @@ void CollectDeferredDecalPassPSOInitializers(
 
 void AddDeferredDecalPass(
 	FRDGBuilder& GraphBuilder,
-	const FViewInfo& View,
+	FViewInfo& View,
 	TConstArrayView<FTransientDecalRenderData> VisibleDecals,
 	const FDeferredDecalPassTextures& PassTextures,
 	FInstanceCullingManager& InstanceCullingManager,
@@ -604,7 +604,8 @@ void AddDeferredDecalPass(
 		INC_DWORD_STAT_BY(STAT_Decals, SortedDecalCount);
 	}
 
-	const bool bVisibleDecalsInView = MeshDecalCount > 0 || SortedDecalCount > 0;
+	const bool bHasAnyDrawCommandDecalCount = HasAnyDrawCommandDecalCount(DecalRenderStage, View);
+	const bool bVisibleDecalsInView = MeshDecalCount > 0 || SortedDecalCount > 0 || bHasAnyDrawCommandDecalCount;
 	const bool bShaderComplexity = View.Family->EngineShowFlags.ShaderComplexity;
 	const bool bStencilSizeThreshold = CVarStencilSizeThreshold.GetValueOnRenderThread() >= 0;
 
@@ -694,7 +695,7 @@ void AddDeferredDecalPass(
 	{
 		RDG_EVENT_SCOPE(GraphBuilder, "DeferredDecals %s", GetStageName(DecalRenderStage));
 
-		if (MeshDecalCount > 0 && (DecalRenderStage == EDecalRenderStage::BeforeBasePass || DecalRenderStage == EDecalRenderStage::BeforeLighting || DecalRenderStage == EDecalRenderStage::Emissive || DecalRenderStage == EDecalRenderStage::AmbientOcclusion))
+		if ((MeshDecalCount > 0 || bHasAnyDrawCommandDecalCount) && (DecalRenderStage == EDecalRenderStage::BeforeBasePass || DecalRenderStage == EDecalRenderStage::BeforeLighting || DecalRenderStage == EDecalRenderStage::Emissive || DecalRenderStage == EDecalRenderStage::AmbientOcclusion))
 		{
 			RenderMeshDecals(GraphBuilder, Scene, View, PassTextures, InstanceCullingManager, DecalRenderStage);
 		}
diff --git a/Engine/Source/Runtime/Renderer/Private/CompositionLighting/PostProcessDeferredDecals.h b/Engine/Source/Runtime/Renderer/Private/CompositionLighting/PostProcessDeferredDecals.h
--- a/Engine/Source/Runtime/Renderer/Private/CompositionLighting/PostProcessDeferredDecals.h
+++ b/Engine/Source/Runtime/Renderer/Private/CompositionLighting/PostProcessDeferredDecals.h
@@ -48,7 +48,7 @@ FDeferredDecalPassTextures GetDeferredDecalPassTextures(
 
 void AddDeferredDecalPass(
 	FRDGBuilder& GraphBuilder,
-	const FViewInfo& ViewInfo,
+	FViewInfo& ViewInfo,
 	TConstArrayView<FTransientDecalRenderData> VisibleDecals,
 	const FDeferredDecalPassTextures& Textures,
 	FInstanceCullingManager& InstanceCullingManager,
@@ -98,11 +98,17 @@ void GetDeferredDecalPassParameters(
 void RenderMeshDecals(
 	FRDGBuilder& GraphBuilder,
 	const FScene& Scene,
-	const FViewInfo& View,
+	FViewInfo& View,
 	const FDeferredDecalPassTextures& DecalPassTextures,
 	FInstanceCullingManager& InstanceCullingManager, 
 	EDecalRenderStage DecalRenderStage);
 
+bool HasAnyDrawCommandDecalCount(
+	EDecalRenderStage DecalRenderStage, 
+	FViewInfo& View);
+
+bool UseDecalCachedMDCs();
+
 void ExtractNormalsForNextFrameReprojection(
 	FRDGBuilder& GraphBuilder,
 	const FSceneTextures& SceneTextures,
diff --git a/Engine/Source/Runtime/Renderer/Private/CompositionLighting/PostProcessMeshDecals.cpp b/Engine/Source/Runtime/Renderer/Private/CompositionLighting/PostProcessMeshDecals.cpp
--- a/Engine/Source/Runtime/Renderer/Private/CompositionLighting/PostProcessMeshDecals.cpp
+++ b/Engine/Source/Runtime/Renderer/Private/CompositionLighting/PostProcessMeshDecals.cpp
@@ -21,6 +21,30 @@
 #include "DebugViewModeRendering.h"
 #include "MeshPassProcessor.inl"
 #include "SimpleMeshDrawCommandPass.h"
+#include "ComponentRecreateRenderStateContext.h"
+
+static TAutoConsoleVariable<int32> CVarParallelMeshDecal(
+	TEXT("r.ParallelMeshDecal"),
+	1,
+	TEXT("Toggles parallel mesh decal rendering. Parallel rendering must be enabled for this to have an effect."),
+	ECVF_RenderThreadSafe
+);
+
+static TAutoConsoleVariable<int32> CVarDecalUseCachedMDCs(
+	TEXT("r.Decal.UseCachedMDCs"),
+	0,
+	TEXT("Use cached MDCs for decal rendering."),
+	FConsoleVariableDelegate::CreateLambda([](IConsoleVariable* InVariable)
+	{
+		FGlobalComponentRecreateRenderStateContext Context;
+	}),
+	ECVF_RenderThreadSafe
+);
+
+bool UseDecalCachedMDCs()
+{
+	return CVarDecalUseCachedMDCs.GetValueOnRenderThread() != 0;
+}
 
 class FMeshDecalsVS : public FMeshMaterialShader
 {
@@ -127,9 +151,9 @@ class FMeshDecalMeshProcessor : public FMeshPassProcessor
 public:
 	FMeshDecalMeshProcessor(const FScene* Scene, 
 		ERHIFeatureLevel::Type FeatureLevel,
-		const FSceneView* InViewIfDynamicMeshCommand, 
-		EDecalRenderStage InPassDecalStage, 
+		const FSceneView* InViewIfDynamicMeshCommand,
 		EDecalRenderTargetMode InRenderTargetMode,
+		EShadingPath ShadingPath,
 		FMeshPassDrawListContext* InDrawListContext);
 
 	virtual void AddMeshBatch(const FMeshBatch& RESTRICT MeshBatch, uint64 BatchElementMask, const FPrimitiveSceneProxy* RESTRICT PrimitiveSceneProxy, int32 StaticMeshId = -1) override final;
@@ -177,16 +201,15 @@ private:
 IMPLEMENT_STATIC_UNIFORM_BUFFER_SLOT(DeferredDecals);
 IMPLEMENT_STATIC_UNIFORM_BUFFER_STRUCT(FDeferredDecalUniformParameters, "DeferredDecal", DeferredDecals);
 
-static const TCHAR* MeshDecalPassName = TEXT("MeshDecal");
 
 FMeshDecalMeshProcessor::FMeshDecalMeshProcessor(const FScene* Scene, 
 	ERHIFeatureLevel::Type InFeatureLevel,
 	const FSceneView* InViewIfDynamicMeshCommand, 
-	EDecalRenderStage InPassDecalStage, 
 	EDecalRenderTargetMode InRenderTargetMode,
+	EShadingPath ShadingPath,
 	FMeshPassDrawListContext* InDrawListContext)
-	: FMeshPassProcessor(MeshDecalPassName, Scene, InFeatureLevel, InViewIfDynamicMeshCommand, InDrawListContext)
-	, PassDecalStage(InPassDecalStage)
+	: FMeshPassProcessor(DecalRendering::GetMeshPassType(InRenderTargetMode), Scene, InFeatureLevel, InViewIfDynamicMeshCommand, InDrawListContext)
+	, PassDecalStage(DecalRendering::GetRenderStage(InRenderTargetMode, ShadingPath))
 	, RenderTargetMode(InRenderTargetMode)
 {
 	PassDrawRenderState.SetDepthStencilState(TStaticDepthStencilState<false, CF_DepthNearOrEqual>::GetRHI());
@@ -194,6 +217,12 @@ FMeshDecalMeshProcessor::FMeshDecalMeshProcessor(const FScene* Scene,
 
 void FMeshDecalMeshProcessor::AddMeshBatch(const FMeshBatch& RESTRICT MeshBatch, uint64 BatchElementMask, const FPrimitiveSceneProxy* RESTRICT PrimitiveSceneProxy, int32 StaticMeshId)
 {
+	// Early out if not using cached MDCs and there is no valid ViewIfDynamicMeshCommand (means we are caching the static MDCs)
+	if (CVarDecalUseCachedMDCs.GetValueOnRenderThread() == 0 && ViewIfDynamicMeshCommand == nullptr)
+	{
+		return;
+	}
+
 	if (MeshBatch.bUseForMaterial && MeshBatch.IsDecal(FeatureLevel))
 	{
 		const FMaterialRenderProxy* MaterialRenderProxy = MeshBatch.MaterialRenderProxy;
@@ -226,7 +255,7 @@ bool FMeshDecalMeshProcessor::TryAddMeshBatch(
 		// We have no special engine material for decals since we don't want to eat the compilation & memory cost, so just skip if it failed to compile
 		if (Material.GetRenderingThreadShaderMap())
 		{
-			const EShaderPlatform ShaderPlatform = ViewIfDynamicMeshCommand->GetShaderPlatform();
+			const EShaderPlatform ShaderPlatform = GetFeatureLevelShaderPlatform(FeatureLevel);
 			const FDecalBlendDesc DecalBlendDesc = DecalRendering::ComputeDecalBlendDesc(ShaderPlatform, Material);
 
 			const bool bShouldRender =
@@ -239,7 +268,7 @@ bool FMeshDecalMeshProcessor::TryAddMeshBatch(
 				ERasterizerFillMode MeshFillMode = ComputeMeshFillMode(Material, OverrideSettings);
 				ERasterizerCullMode MeshCullMode = ComputeMeshCullMode(Material, OverrideSettings);
 
-				if (ViewIfDynamicMeshCommand->Family->UseDebugViewPS())
+				if (ViewIfDynamicMeshCommand && ViewIfDynamicMeshCommand->Family->UseDebugViewPS())
 				{
 					// Deferred decals can only use translucent blend mode
 					if (ViewIfDynamicMeshCommand->Family->EngineShowFlags.ShaderComplexity)
@@ -429,25 +458,56 @@ void FMeshDecalMeshProcessor::CollectPSOInitializers(
 	}
 }
 
-IPSOCollector* CreateMeshDecalMeshProcessor(ERHIFeatureLevel::Type FeatureLevel)
+FMeshPassProcessor* CreateMeshDecalDBufferMeshProcessor(ERHIFeatureLevel::Type FeatureLevel, const FScene* Scene, const FSceneView* InViewIfDynamicMeshCommand, FMeshPassDrawListContext* InDrawListContext)
 {
-	if (DoesPlatformSupportNanite(GetFeatureLevelShaderPlatform(FeatureLevel)))
-	{
-		return new FMeshDecalMeshProcessor(nullptr, FeatureLevel, nullptr, EDecalRenderStage::None, EDecalRenderTargetMode::None, nullptr);
-	}
-	else
-	{
-		return nullptr;
-	}
+	return new FMeshDecalMeshProcessor(Scene, FeatureLevel, InViewIfDynamicMeshCommand, EDecalRenderTargetMode::DBuffer, EShadingPath::Deferred, InDrawListContext);
+}
+
+FMeshPassProcessor* CreateMeshDecalSceneColorAndGBufferMeshProcessor(ERHIFeatureLevel::Type FeatureLevel, const FScene* Scene, const FSceneView* InViewIfDynamicMeshCommand, FMeshPassDrawListContext* InDrawListContext)
+{
+	return new FMeshDecalMeshProcessor(Scene, FeatureLevel, InViewIfDynamicMeshCommand, EDecalRenderTargetMode::SceneColorAndGBuffer, EShadingPath::Deferred, InDrawListContext);
+}
+
+FMeshPassProcessor* CreateMeshDecalSceneColorAndGBufferNoNormalMeshProcessor(ERHIFeatureLevel::Type FeatureLevel, const FScene* Scene, const FSceneView* InViewIfDynamicMeshCommand, FMeshPassDrawListContext* InDrawListContext)
+{
+	return new FMeshDecalMeshProcessor(Scene, FeatureLevel, InViewIfDynamicMeshCommand, EDecalRenderTargetMode::SceneColorAndGBufferNoNormal, EShadingPath::Deferred, InDrawListContext);
+}
+
+FMeshPassProcessor* CreateMeshDecalSceneColorMeshProcessor(ERHIFeatureLevel::Type FeatureLevel, const FScene* Scene, const FSceneView* InViewIfDynamicMeshCommand, FMeshPassDrawListContext* InDrawListContext)
+{
+	return new FMeshDecalMeshProcessor(Scene, FeatureLevel, InViewIfDynamicMeshCommand, EDecalRenderTargetMode::SceneColor, EShadingPath::Deferred, InDrawListContext);
+}
+
+FMeshPassProcessor* CreateMeshDecalAmbientOcclusionProcessor(ERHIFeatureLevel::Type FeatureLevel, const FScene* Scene, const FSceneView* InViewIfDynamicMeshCommand, FMeshPassDrawListContext* InDrawListContext)
+{
+	return new FMeshDecalMeshProcessor(Scene, FeatureLevel, InViewIfDynamicMeshCommand, EDecalRenderTargetMode::AmbientOcclusion, EShadingPath::Deferred, InDrawListContext);
+}
+
+REGISTER_MESHPASSPROCESSOR_AND_PSOCOLLECTOR(MeshDecalPass_DBuffer, CreateMeshDecalDBufferMeshProcessor, EShadingPath::Deferred, EMeshPass::MeshDecal_DBuffer, EMeshPassFlags::CachedMeshCommands | EMeshPassFlags::MainView);
+REGISTER_MESHPASSPROCESSOR_AND_PSOCOLLECTOR(MeshDecal_SceneColorAndGBuffer, CreateMeshDecalSceneColorAndGBufferMeshProcessor, EShadingPath::Deferred, EMeshPass::MeshDecal_SceneColorAndGBuffer, EMeshPassFlags::CachedMeshCommands | EMeshPassFlags::MainView);
+REGISTER_MESHPASSPROCESSOR_AND_PSOCOLLECTOR(MeshDecal_SceneColorAndGBufferNoNormal, CreateMeshDecalSceneColorAndGBufferNoNormalMeshProcessor, EShadingPath::Deferred, EMeshPass::MeshDecal_SceneColorAndGBufferNoNormal, EMeshPassFlags::CachedMeshCommands | EMeshPassFlags::MainView);
+REGISTER_MESHPASSPROCESSOR_AND_PSOCOLLECTOR(MeshDecal_SceneColor, CreateMeshDecalSceneColorMeshProcessor, EShadingPath::Deferred, EMeshPass::MeshDecal_SceneColor, EMeshPassFlags::CachedMeshCommands | EMeshPassFlags::MainView);
+REGISTER_MESHPASSPROCESSOR_AND_PSOCOLLECTOR(MeshDecal_AmbientOcclusion, CreateMeshDecalAmbientOcclusionProcessor, EShadingPath::Deferred, EMeshPass::MeshDecal_AmbientOcclusion, EMeshPassFlags::CachedMeshCommands | EMeshPassFlags::MainView);
+
+FMeshPassProcessor* CreateMeshDecalSceneColorAndGBufferMeshProcessor_Mobile(ERHIFeatureLevel::Type FeatureLevel, const FScene* Scene, const FSceneView* InViewIfDynamicMeshCommand, FMeshPassDrawListContext* InDrawListContext)
+{
+	return new FMeshDecalMeshProcessor(Scene, FeatureLevel, InViewIfDynamicMeshCommand, EDecalRenderTargetMode::SceneColorAndGBuffer, EShadingPath::Mobile, InDrawListContext);
+}
+
+FMeshPassProcessor* CreateMeshDecalSceneColorMeshProcessor_Mobile(ERHIFeatureLevel::Type FeatureLevel, const FScene* Scene, const FSceneView* InViewIfDynamicMeshCommand, FMeshPassDrawListContext* InDrawListContext)
+{
+	return new FMeshDecalMeshProcessor(Scene, FeatureLevel, InViewIfDynamicMeshCommand, EDecalRenderTargetMode::SceneColor, EShadingPath::Mobile, InDrawListContext);
 }
 
-// Only register for PSO Collection
-FRegisterPSOCollectorCreateFunction RegisterPSOCollectorMeshDecal(&CreateMeshDecalMeshProcessor, EShadingPath::Deferred, MeshDecalPassName);
+REGISTER_MESHPASSPROCESSOR_AND_PSOCOLLECTOR(Mobile_MeshDecal_SceneColorAndGBuffer, CreateMeshDecalSceneColorAndGBufferMeshProcessor_Mobile, EShadingPath::Mobile, EMeshPass::MeshDecal_SceneColorAndGBuffer, EMeshPassFlags::CachedMeshCommands | EMeshPassFlags::MainView);
+REGISTER_MESHPASSPROCESSOR_AND_PSOCOLLECTOR(Mobile_MeshDecal_SceneColor, CreateMeshDecalSceneColorMeshProcessor_Mobile, EShadingPath::Mobile, EMeshPass::MeshDecal_SceneColor, EMeshPassFlags::CachedMeshCommands | EMeshPassFlags::MainView);
+
+DECLARE_CYCLE_STAT(TEXT("MeshDecal"), STAT_CLP_MeshDecal, STATGROUP_ParallelCommandListMarkers);
 
 void DrawDecalMeshCommands(
 	FRDGBuilder& GraphBuilder,
 	const FScene& Scene,
-	const FViewInfo& View,
+	FViewInfo& View,
 	const FDeferredDecalPassTextures& DecalPassTextures,
 	FInstanceCullingManager& InstanceCullingManager,
 	EDecalRenderStage DecalRenderStage,
@@ -456,41 +516,113 @@ void DrawDecalMeshCommands(
 	auto* PassParameters = GraphBuilder.AllocParameters<FDeferredDecalPassParameters>();
 	GetDeferredDecalPassParameters(GraphBuilder, View, DecalPassTextures, RenderTargetMode, *PassParameters);
 
-	AddSimpleMeshPass(
-		GraphBuilder, 
-		PassParameters, 
-		&Scene, 
-		View, 
-		&InstanceCullingManager, 
-		RDG_EVENT_NAME("MeshDecals"), 
-		View.ViewRect,
-		[&View, DecalRenderStage, RenderTargetMode](FDynamicPassMeshDrawListContext* DynamicMeshPassContext)
+	// If decal cached MDCs is enabled then mesh decal batches will be empty
+	// Cached decals are rendered below using the parallel mesh draw command passes
+	if (View.MeshDecalBatches.Num() > 0)
 	{
-		TRACE_CPUPROFILER_EVENT_SCOPE(MeshDecalCommands);
+		AddSimpleMeshPass(
+			GraphBuilder,
+			PassParameters,
+			&Scene,
+			View,
+			&InstanceCullingManager,
+			RDG_EVENT_NAME("MeshDecals"),
+			View.ViewRect,
+			[&View, DecalRenderStage, RenderTargetMode](FDynamicPassMeshDrawListContext* DynamicMeshPassContext)
+			{
+				TRACE_CPUPROFILER_EVENT_SCOPE(MeshDecalCommands);
 
-		FMeshDecalMeshProcessor PassMeshProcessor(
-			View.Family->Scene->GetRenderScene(),
-			View.GetFeatureLevel(),
-			&View,
-			DecalRenderStage,
-			RenderTargetMode,
-			DynamicMeshPassContext);
+				check(DecalRenderStage == DecalRendering::GetRenderStage(RenderTargetMode, EShadingPath::Deferred));
 
-		for (int32 MeshBatchIndex = 0; MeshBatchIndex < View.MeshDecalBatches.Num(); ++MeshBatchIndex)
-		{
-			const FMeshBatch* Mesh = View.MeshDecalBatches[MeshBatchIndex].Mesh;
-			const FPrimitiveSceneProxy* PrimitiveSceneProxy = View.MeshDecalBatches[MeshBatchIndex].Proxy;
-			const uint64 DefaultBatchElementMask = ~0ull;
+				FMeshDecalMeshProcessor PassMeshProcessor(
+					View.Family->Scene->GetRenderScene(),
+					View.GetFeatureLevel(),
+					&View,
+					RenderTargetMode,
+					EShadingPath::Deferred,
+					DynamicMeshPassContext);
 
-			PassMeshProcessor.AddMeshBatch(*Mesh, DefaultBatchElementMask, PrimitiveSceneProxy);
+				for (int32 MeshBatchIndex = 0; MeshBatchIndex < View.MeshDecalBatches.Num(); ++MeshBatchIndex)
+				{
+					const FMeshBatch* Mesh = View.MeshDecalBatches[MeshBatchIndex].Mesh;
+					const FPrimitiveSceneProxy* PrimitiveSceneProxy = View.MeshDecalBatches[MeshBatchIndex].Proxy;
+					const uint64 DefaultBatchElementMask = ~0ull;
+
+					PassMeshProcessor.AddMeshBatch(*Mesh, DefaultBatchElementMask, PrimitiveSceneProxy);
+				}
+			});
+	}
+
+	EMeshPass::Type DecalMeshPassType = DecalRendering::GetMeshPassType(RenderTargetMode);
+	if (View.ParallelMeshDrawCommandPasses[DecalMeshPassType].HasAnyDraw())
+	{
+		const bool bRenderInParallel = GRHICommandList.UseParallelAlgorithms() && CVarParallelMeshDecal.GetValueOnRenderThread() == 1;
+
+		View.ParallelMeshDrawCommandPasses[DecalMeshPassType].BuildRenderingCommands(GraphBuilder, Scene.GPUScene, PassParameters->InstanceCullingDrawParams);
+
+		if (bRenderInParallel)
+		{
+			GraphBuilder.AddPass(
+				RDG_EVENT_NAME("%s", GetMeshPassName(DecalMeshPassType)),
+				PassParameters,
+				ERDGPassFlags::Raster | ERDGPassFlags::SkipRenderPass,
+				[&View, PassParameters, DecalMeshPassType](const FRDGPass* InPass, FRHICommandListImmediate& RHICmdList)
+				{
+					FRDGParallelCommandListSet ParallelCommandListSet(InPass, RHICmdList, GET_STATID(STAT_CLP_MeshDecal), View, FParallelCommandListBindings(PassParameters));
+					View.ParallelMeshDrawCommandPasses[DecalMeshPassType].DispatchDraw(&ParallelCommandListSet, RHICmdList, &PassParameters->InstanceCullingDrawParams);
+				});
+		}
+		else
+		{
+			GraphBuilder.AddPass(
+				RDG_EVENT_NAME("%s", GetMeshPassName(DecalMeshPassType)),
+				PassParameters,
+				ERDGPassFlags::Raster,
+				[&View, PassParameters, DecalMeshPassType](FRHICommandList& RHICmdList)
+				{
+					FSceneRenderer::SetStereoViewport(RHICmdList, View, 1.0f);
+					View.ParallelMeshDrawCommandPasses[DecalMeshPassType].DispatchDraw(nullptr, RHICmdList, &PassParameters->InstanceCullingDrawParams);
+				});
 		}
-	});
+	}
+}
+
+bool HasAnyDrawCommandDecalCount(
+	EDecalRenderStage DecalRenderStage,
+	FViewInfo& View)
+{
+	switch (DecalRenderStage)
+	{
+	case EDecalRenderStage::BeforeBasePass:	
+		return View.ParallelMeshDrawCommandPasses[EMeshPass::MeshDecal_DBuffer].HasAnyDraw();
+
+	case EDecalRenderStage::BeforeLighting:
+		return View.ParallelMeshDrawCommandPasses[EMeshPass::MeshDecal_SceneColorAndGBuffer].HasAnyDraw() || View.ParallelMeshDrawCommandPasses[EMeshPass::MeshDecal_SceneColorAndGBufferNoNormal].HasAnyDraw();
+		break;
+
+	case EDecalRenderStage::Mobile:
+		return View.ParallelMeshDrawCommandPasses[EMeshPass::MeshDecal_SceneColor].HasAnyDraw();
+		break;
+
+	case EDecalRenderStage::MobileBeforeLighting:
+		return View.ParallelMeshDrawCommandPasses[EMeshPass::MeshDecal_SceneColorAndGBuffer].HasAnyDraw();
+		break;
+
+	case EDecalRenderStage::Emissive:
+		return View.ParallelMeshDrawCommandPasses[EMeshPass::MeshDecal_SceneColor].HasAnyDraw();
+		break;
+
+	case EDecalRenderStage::AmbientOcclusion:
+		return View.ParallelMeshDrawCommandPasses[EMeshPass::MeshDecal_AmbientOcclusion].HasAnyDraw();
+		break;
+	}
+	return false;
 }
 
 void RenderMeshDecals(
 	FRDGBuilder& GraphBuilder,
 	const FScene& Scene,
-	const FViewInfo& View,
+	FViewInfo& View,
 	const FDeferredDecalPassTextures& DecalPassTextures,
 	FInstanceCullingManager& InstanceCullingManager,
 	EDecalRenderStage DecalRenderStage)
@@ -526,7 +658,7 @@ void RenderMeshDecals(
 	}
 }
 
-void RenderMeshDecalsMobile(FRHICommandList& RHICmdList, const FViewInfo& View, EDecalRenderStage DecalRenderStage, EDecalRenderTargetMode RenderTargetMode)
+void RenderMeshDecalsMobile(FRHICommandList& RHICmdList, FViewInfo& View, EDecalRenderStage DecalRenderStage, EDecalRenderTargetMode RenderTargetMode)
 {
 	FGraphicsPipelineStateInitializer GraphicsPSOInit;
 	RHICmdList.SetViewport(View.ViewRect.Min.X, View.ViewRect.Min.Y, 0, View.ViewRect.Max.X, View.ViewRect.Max.Y, 1);
@@ -534,12 +666,14 @@ void RenderMeshDecalsMobile(FRHICommandList& RHICmdList, const FViewInfo& View,
 
 	DrawDynamicMeshPass(View, RHICmdList, [&View, DecalRenderStage, RenderTargetMode](FDynamicPassMeshDrawListContext* DynamicMeshPassContext)
 	{
+		check(DecalRenderStage == DecalRendering::GetRenderStage(RenderTargetMode, EShadingPath::Mobile));
+
 		FMeshDecalMeshProcessor PassMeshProcessor(
 			View.Family->Scene->GetRenderScene(),
 			View.GetFeatureLevel(),
 			&View,
-			DecalRenderStage,
 			RenderTargetMode,
+			EShadingPath::Mobile,
 			DynamicMeshPassContext);
 
 		for (int32 MeshBatchIndex = 0; MeshBatchIndex < View.MeshDecalBatches.Num(); ++MeshBatchIndex)
diff --git a/Engine/Source/Runtime/Renderer/Private/DecalRenderingCommon.cpp b/Engine/Source/Runtime/Renderer/Private/DecalRenderingCommon.cpp
--- a/Engine/Source/Runtime/Renderer/Private/DecalRenderingCommon.cpp
+++ b/Engine/Source/Runtime/Renderer/Private/DecalRenderingCommon.cpp
@@ -5,6 +5,10 @@
 #include "RenderUtils.h"
 #include "DataDrivenShaderPlatformInfo.h"
 #include "Substrate/Substrate.h"
+#include "MaterialShared.h"
+#include "StaticMeshBatch.h"
+
+static_assert(DecalRendering::DecalRenderTargetMode_NumBits == ((int32)EDecalRenderTargetMode::Num - 1));
 
 namespace DecalRendering
 {
@@ -181,6 +185,73 @@ namespace DecalRendering
 		return EDecalRenderStage::None;
 	}
 
+	uint8 GetDecalRenderTargetModeMask(const FMaterial& Material, ERHIFeatureLevel::Type FeatureLevel)
+	{
+		uint8 Mask = 0;
+
+		const EShaderPlatform ShaderPlatform = GetFeatureLevelShaderPlatform(FeatureLevel);
+		const FDecalBlendDesc DecalBlendDesc = ComputeDecalBlendDesc(ShaderPlatform, Material);
+
+		for (int32 DecalRenderStageNum = 0; DecalRenderStageNum < (int32)EDecalRenderStage::Num; ++DecalRenderStageNum)
+		{
+			EDecalRenderStage DecalRenderStage = (EDecalRenderStage)DecalRenderStageNum;
+			if (IsCompatibleWithRenderStage(DecalBlendDesc, DecalRenderStage))
+			{
+				EDecalRenderTargetMode DecalRenderTargetMode = GetRenderTargetMode(DecalBlendDesc, DecalRenderStage);
+				Mask |= (1 << (uint32)DecalRenderTargetMode);
+			}
+		}
+
+		return Mask;
+	}
+
+	bool IsCompatibleWithRenderTargetMode(uint8 DecalRenderTargetModeMask, EDecalRenderTargetMode DecalRenderTargetMode)
+	{
+		return (DecalRenderTargetModeMask & (1 << (uint32)DecalRenderTargetMode)) != 0;
+	}
+
+	EMeshPass::Type GetMeshPassType(EDecalRenderTargetMode RenderTargetMode)
+	{
+		switch (RenderTargetMode)
+		{
+		case EDecalRenderTargetMode::None:
+			return EMeshPass::Num;
+		case EDecalRenderTargetMode::DBuffer:
+			return EMeshPass::MeshDecal_DBuffer;
+		case EDecalRenderTargetMode::SceneColorAndGBuffer:
+			return EMeshPass::MeshDecal_SceneColorAndGBuffer;
+		case EDecalRenderTargetMode::SceneColorAndGBufferNoNormal:
+			return EMeshPass::MeshDecal_SceneColorAndGBufferNoNormal;
+		case EDecalRenderTargetMode::SceneColor:
+			return EMeshPass::MeshDecal_SceneColor;
+		case EDecalRenderTargetMode::AmbientOcclusion:
+			return EMeshPass::MeshDecal_AmbientOcclusion;
+		}
+
+		return EMeshPass::Num;
+	}
+
+	EDecalRenderStage GetRenderStage(EDecalRenderTargetMode RenderTargetMode, EShadingPath ShadingPath)
+	{
+		switch (RenderTargetMode)
+		{
+			case EDecalRenderTargetMode::None:
+				return EDecalRenderStage::None;
+			case EDecalRenderTargetMode::DBuffer:
+				return EDecalRenderStage::BeforeBasePass;
+			case EDecalRenderTargetMode::SceneColorAndGBuffer:
+				return (ShadingPath == EShadingPath::Deferred) ? EDecalRenderStage::BeforeLighting : EDecalRenderStage::MobileBeforeLighting;
+			case EDecalRenderTargetMode::SceneColorAndGBufferNoNormal:
+				return EDecalRenderStage::BeforeLighting;
+			case EDecalRenderTargetMode::SceneColor:
+				return (ShadingPath == EShadingPath::Deferred) ? EDecalRenderStage::Emissive : EDecalRenderStage::Mobile;
+			case EDecalRenderTargetMode::AmbientOcclusion:
+				return EDecalRenderStage::AmbientOcclusion;
+		}
+
+		return EDecalRenderStage::None;
+	}
+
 	EDecalRenderTargetMode GetRenderTargetMode(FDecalBlendDesc DecalBlendDesc, EDecalRenderStage DecalRenderStage)
 	{
 		switch(DecalRenderStage)
diff --git a/Engine/Source/Runtime/Renderer/Private/DecalRenderingCommon.h b/Engine/Source/Runtime/Renderer/Private/DecalRenderingCommon.h
--- a/Engine/Source/Runtime/Renderer/Private/DecalRenderingCommon.h
+++ b/Engine/Source/Runtime/Renderer/Private/DecalRenderingCommon.h
@@ -5,6 +5,7 @@
 #include "CoreMinimal.h"
 #include "Materials/Material.h"
 #include "RenderUtils.h"
+#include "MeshPassProcessor.h"
 
 class FMaterial;
 class FRHIBlendState;
@@ -66,6 +67,8 @@ enum class EDecalRenderTargetMode : uint8
 	SceneColorAndGBufferNoNormal = 3,
 	SceneColor = 4,
 	AmbientOcclusion = 5,
+
+	Num,
 };
 
 /** Enumeration of decal rasterization states. */
@@ -90,12 +93,24 @@ namespace DecalRendering
 	/** Returns true if a decal should be rendered in the render stage. */
 	bool IsCompatibleWithRenderStage(FDecalBlendDesc DecalBlendDesc, EDecalRenderStage DecalRenderStage);
 
+	/** Returns the main render stage from render target mode and shading path */
+	EDecalRenderStage GetRenderStage(EDecalRenderTargetMode RenderTargetMode, EShadingPath ShadingPath);
+
 	/** Returns the main render stage for a decal (does not include the emissive, and AO stages). Can return EDecalRenderStage::None if there is no valid main render stage. */
 	EDecalRenderStage GetBaseRenderStage(FDecalBlendDesc DecalBlendDesc);
 
 	/** Get the render target mode that a decal uses for a given stage. Can return EDecalRenderTargetMode::None if there is no valid render target mode. */
 	EDecalRenderTargetMode GetRenderTargetMode(FDecalBlendDesc DecalBlendDesc, EDecalRenderStage DecalRenderStage);
 
+	/** Setup a mask for each render target mode where the material should render to */
+	uint8 GetDecalRenderTargetModeMask(const FMaterial& Material, ERHIFeatureLevel::Type FeatureLevel);
+
+	/** Returns true if a decal should be rendered in the render target mode. */
+	bool IsCompatibleWithRenderTargetMode(uint8 DecalRenderTargetModeMask, EDecalRenderTargetMode DecalRenderTargetMode);
+
+	/** Get the mesh pass processer type to use for given render target mode */
+	EMeshPass::Type GetMeshPassType(EDecalRenderTargetMode RenderTargetMode);
+
 	/** Get render target count for the given render target mode. */
 	uint32 GetRenderTargetCount(FDecalBlendDesc DecalBlendDesc, EDecalRenderTargetMode RenderTargetMode);
 
diff --git a/Engine/Source/Runtime/Renderer/Private/MobileDecalRendering.cpp b/Engine/Source/Runtime/Renderer/Private/MobileDecalRendering.cpp
--- a/Engine/Source/Runtime/Renderer/Private/MobileDecalRendering.cpp
+++ b/Engine/Source/Runtime/Renderer/Private/MobileDecalRendering.cpp
@@ -20,7 +20,7 @@
 #include "CompositionLighting/PostProcessDeferredDecals.h"
 #include "DBufferTextures.h"
 
-void RenderMeshDecalsMobile(FRHICommandList& RHICmdList, const FViewInfo& View, EDecalRenderStage DecalRenderStage, EDecalRenderTargetMode RenderTargetMode);
+void RenderMeshDecalsMobile(FRHICommandList& RHICmdList, FViewInfo& View, EDecalRenderStage DecalRenderStage, EDecalRenderTargetMode RenderTargetMode);
 extern void RenderDeferredDecalsMobile(FRHICommandList& RHICmdList, const FScene& Scene, const FViewInfo& View, EDecalRenderStage DecalRenderStage, EDecalRenderTargetMode RenderTargetMode);
 
 static bool DoesPlatformSupportDecals(EShaderPlatform ShaderPlatform)
@@ -44,7 +44,7 @@ static bool DoesPlatformSupportDecals(EShaderPlatform ShaderPlatform)
 	return true;
 }
 
-void FMobileSceneRenderer::RenderDecals(FRHICommandList& RHICmdList, const FViewInfo& View, const FInstanceCullingDrawParams* InstanceCullingDrawParams)
+void FMobileSceneRenderer::RenderDecals(FRHICommandList& RHICmdList, FViewInfo& View, const FInstanceCullingDrawParams* InstanceCullingDrawParams)
 {
 	if (!DoesPlatformSupportDecals(View.GetShaderPlatform()) || !ViewFamily.EngineShowFlags.Decals || View.bIsPlanarReflection)
 	{
@@ -80,6 +80,13 @@ void FMobileSceneRenderer::RenderDecals(FRHICommandList& RHICmdList, const FView
 			RHICmdList.SetStaticUniformBuffer(BatchedPrimitiveSlot, BatchedPrimitiveBufferRHI);
 		}
 	}
+
+	EMeshPass::Type DecalMeshPassType = DecalRendering::GetMeshPassType(RenderTargetMode);
+	if (View.ParallelMeshDrawCommandPasses[DecalMeshPassType].HasAnyDraw())
+	{
+		RHICmdList.SetViewport(View.ViewRect.Min.X, View.ViewRect.Min.Y, 0, View.ViewRect.Max.X, View.ViewRect.Max.Y, 1);
+		View.ParallelMeshDrawCommandPasses[DecalMeshPassType].DispatchDraw(nullptr, RHICmdList, InstanceCullingDrawParams);
+	}
 }
 
 void RenderDeferredDecalsMobile(FRHICommandList& RHICmdList, const FScene& Scene, const FViewInfo& View, EDecalRenderStage DecalRenderStage, EDecalRenderTargetMode RenderTargetMode)
diff --git a/Engine/Source/Runtime/Renderer/Private/PrimitiveSceneInfo.cpp b/Engine/Source/Runtime/Renderer/Private/PrimitiveSceneInfo.cpp
--- a/Engine/Source/Runtime/Renderer/Private/PrimitiveSceneInfo.cpp
+++ b/Engine/Source/Runtime/Renderer/Private/PrimitiveSceneInfo.cpp
@@ -37,7 +37,7 @@
 #include "PrimitiveSceneDesc.h"
 #include "BasePassRendering.h" // TODO: Remove with later refactor (moving Nanite shading into its own files)
 #include "InstanceDataSceneProxy.h"
-
+#include "DecalRenderingCommon.h"
 
 extern int32 GGPUSceneInstanceClearList;
 
@@ -135,6 +135,12 @@ public:
 			bool bUseForWaterInfoTextureDepth = Mesh.bUseForWaterInfoTextureDepth;
 			bool bUseForLumenSceneCapture = Mesh.bUseForLumenSurfaceCacheCapture;
 
+			uint8 DecalRenderTargetModeMask = 0;
+			if (Mesh.IsDecal(FeatureLevel))
+			{
+				DecalRenderTargetModeMask = DecalRendering::GetDecalRenderTargetModeMask(Material, FeatureLevel);
+			}
+
 			FStaticMeshBatchRelevance* StaticMeshRelevance = new(PrimitiveSceneInfo->StaticMeshRelevances) FStaticMeshBatchRelevance(
 				*StaticMesh, 
 				ScreenSize, 
@@ -146,6 +152,7 @@ public:
 				bSupportsGPUScene,
 				bUseForWaterInfoTextureDepth,
 				bUseForLumenSceneCapture,
+				DecalRenderTargetModeMask,
 				FeatureLevel
 				);
 		}
diff --git a/Engine/Source/Runtime/Renderer/Private/SceneRendering.h b/Engine/Source/Runtime/Renderer/Private/SceneRendering.h
--- a/Engine/Source/Runtime/Renderer/Private/SceneRendering.h
+++ b/Engine/Source/Runtime/Renderer/Private/SceneRendering.h
@@ -2730,7 +2730,7 @@ protected:
 
 
 	/** Renders decals. */
-	void RenderDecals(FRHICommandList& RHICmdList, const FViewInfo& View, const FInstanceCullingDrawParams* InstanceCullingDrawParams);
+	void RenderDecals(FRHICommandList& RHICmdList, FViewInfo& View, const FInstanceCullingDrawParams* InstanceCullingDrawParams);
 	void RenderDBuffer(FRDGBuilder& GraphBuilder, FSceneTextures& SceneTextures, FDBufferTextures& DBufferTextures, FInstanceCullingManager& InstanceCullingManager);
 
 	/** Renders the atmospheric and height fog */
diff --git a/Engine/Source/Runtime/Renderer/Private/SceneVisibility.cpp b/Engine/Source/Runtime/Renderer/Private/SceneVisibility.cpp
--- a/Engine/Source/Runtime/Renderer/Private/SceneVisibility.cpp
+++ b/Engine/Source/Runtime/Renderer/Private/SceneVisibility.cpp
@@ -53,6 +53,8 @@
 #include "VT/VirtualTextureSystem.h"
 #include "NaniteSceneProxy.h"
 #include "ViewDebug.h"
+#include "DecalRenderingCommon.h"
+#include "CompositionLighting/PostProcessDeferredDecals.h"
 
 static float GWireframeCullThreshold = 5.0f;
 static FAutoConsoleVariableRef CVarWireframeCullThreshold(
@@ -1317,6 +1319,7 @@ void FRelevancePacket::ComputeRelevance(FDynamicPrimitiveIndexList& DynamicPrimi
 	const FHLODVisibilityState* const HLODState = bHLODActive && ViewState ? &ViewState->HLODVisibilityState : nullptr;
 	float MaxDrawDistanceScale = GetCachedScalabilityCVars().ViewDistanceScale;
 	MaxDrawDistanceScale *= GetCachedScalabilityCVars().CalculateFieldOfViewDistanceScale(View.DesiredFOV);
+	const bool bUseDecalCachedMDCs = UseDecalCachedMDCs();
 
 	const auto AddEditorDynamicPrimitive = [this, &DynamicPrimitiveIndexList](int32 PrimitiveIndex)
 	{
@@ -1726,11 +1729,25 @@ void FRelevancePacket::ComputeRelevance(FDynamicPrimitiveIndexList& DynamicPrimi
 						// Because ViewRelevance is a sum of all material relevances in the primitive
 						if (ViewRelevance.bRenderInMainPass && ViewRelevance.bDecal && StaticMeshRelevance.bUseForMaterial)
 						{
-							MeshDecalBatches.AddUninitialized(1);
-							FMeshDecalBatch& BatchAndProxy = MeshDecalBatches.Last();
-							BatchAndProxy.Mesh = &StaticMesh;
-							BatchAndProxy.Proxy = PrimitiveSceneProxy;
-							BatchAndProxy.SortKey = PrimitiveSceneProxy->GetTranslucencySortPriority();
+							if (bUseDecalCachedMDCs)
+							{
+								for (uint8 DecalRenderTargetMode = 0; DecalRenderTargetMode < (uint8)EDecalRenderTargetMode::Num; ++DecalRenderTargetMode)
+								{
+									if (DecalRendering::IsCompatibleWithRenderTargetMode(StaticMeshRelevance.DecalRenderTargetModeMask, (EDecalRenderTargetMode)DecalRenderTargetMode))
+									{
+										EMeshPass::Type DecalMeshPassType = DecalRendering::GetMeshPassType((EDecalRenderTargetMode)DecalRenderTargetMode);
+										DrawCommandPacket.AddCommandsForMesh(PrimitiveIndex, PrimitiveSceneInfo, StaticMeshRelevance, StaticMesh, CullingPayloadFlags, Scene, bCanCache, DecalMeshPassType);
+									}
+								}
+							}
+							else
+							{
+								MeshDecalBatches.AddUninitialized(1);
+								FMeshDecalBatch& BatchAndProxy = MeshDecalBatches.Last();
+								BatchAndProxy.Mesh = &StaticMesh;
+								BatchAndProxy.Proxy = PrimitiveSceneProxy;
+								BatchAndProxy.SortKey = PrimitiveSceneProxy->GetTranslucencySortPriority();
+							}
 						}
 					}
 
@@ -2395,11 +2412,29 @@ static void ComputeDynamicMeshRelevance(
 
 	if (ViewRelevance.bRenderInMainPass && ViewRelevance.bDecal)
 	{
-		View.MeshDecalBatches.AddUninitialized(1);
-		FMeshDecalBatch& BatchAndProxy = View.MeshDecalBatches.Last();
-		BatchAndProxy.Mesh = MeshBatch.Mesh;
-		BatchAndProxy.Proxy = MeshBatch.PrimitiveSceneProxy;
-		BatchAndProxy.SortKey = MeshBatch.PrimitiveSceneProxy->GetTranslucencySortPriority();
+		if (UseDecalCachedMDCs())
+		{
+			// DecalRenderTargetModeMask could be cached in FMeshBatchAndRelevance as well
+			const FMaterial& Material = MeshBatch.Mesh->MaterialRenderProxy->GetIncompleteMaterialWithFallback(View.FeatureLevel);
+			uint8 DecalRenderTargetModeMask = DecalRendering::GetDecalRenderTargetModeMask(Material, View.FeatureLevel);
+			for (uint8 DecalRenderTargetMode = 0; DecalRenderTargetMode < (uint8)EDecalRenderTargetMode::Num; ++DecalRenderTargetMode)
+			{
+				if (DecalRendering::IsCompatibleWithRenderTargetMode(DecalRenderTargetModeMask, (EDecalRenderTargetMode)DecalRenderTargetMode))
+				{
+					EMeshPass::Type DecalMeshPassType = DecalRendering::GetMeshPassType((EDecalRenderTargetMode)DecalRenderTargetMode);
+					PassMask.Set(DecalMeshPassType);
+					View.NumVisibleDynamicMeshElements[DecalMeshPassType] += NumElements;
+				}
+			}
+		}
+		else
+		{
+			View.MeshDecalBatches.AddUninitialized(1);
+			FMeshDecalBatch& BatchAndProxy = View.MeshDecalBatches.Last();
+			BatchAndProxy.Mesh = MeshBatch.Mesh;
+			BatchAndProxy.Proxy = MeshBatch.PrimitiveSceneProxy;
+			BatchAndProxy.SortKey = MeshBatch.PrimitiveSceneProxy->GetTranslucencySortPriority();
+		}
 	}
 
 	const bool bIsHairStrandsCompatible = ViewRelevance.bHairStrands && IsHairStrandsEnabled(EHairStrandsShaderType::All, View.GetShaderPlatform());
diff --git a/Engine/Source/Runtime/Renderer/Private/StaticMeshBatch.cpp b/Engine/Source/Runtime/Renderer/Private/StaticMeshBatch.cpp
--- a/Engine/Source/Runtime/Renderer/Private/StaticMeshBatch.cpp
+++ b/Engine/Source/Runtime/Renderer/Private/StaticMeshBatch.cpp
@@ -28,7 +28,19 @@ FStaticMeshBatch::~FStaticMeshBatch()
 }
 
 
-FStaticMeshBatchRelevance::FStaticMeshBatchRelevance(const FStaticMeshBatch& StaticMesh, float InScreenSize, bool InbSupportsCachingMeshDrawCommands, bool InbUseSkyMaterial, bool bInUseSingleLayerWaterMaterial, bool bInUseAnisotropy, bool bInSupportsNaniteRendering, bool bInSupportsGPUScene, bool bInUseForWaterInfoTextureDepth, bool bInUseForLumenSceneCapture, ERHIFeatureLevel::Type FeatureLevel)
+FStaticMeshBatchRelevance::FStaticMeshBatchRelevance(
+	const FStaticMeshBatch& StaticMesh, 
+	float InScreenSize, 
+	bool InbSupportsCachingMeshDrawCommands, 
+	bool InbUseSkyMaterial, 
+	bool bInUseSingleLayerWaterMaterial, 
+	bool bInUseAnisotropy, 
+	bool bInSupportsNaniteRendering, 
+	bool bInSupportsGPUScene, 
+	bool bInUseForWaterInfoTextureDepth, 
+	bool bInUseForLumenSceneCapture, 
+	uint8 InDecalRenderTargetModeMask,
+	ERHIFeatureLevel::Type FeatureLevel)
 	: Id(StaticMesh.Id)
 	, ScreenSize(InScreenSize)
 	, NumElements(StaticMesh.Elements.Num())
@@ -52,6 +64,7 @@ FStaticMeshBatchRelevance::FStaticMeshBatchRelevance(const FStaticMeshBatch& Sta
 	, bSupportsGPUScene(bInSupportsGPUScene)
 	, bUseForWaterInfoTextureDepth(bInUseForWaterInfoTextureDepth)
 	, bUseForLumenSceneCapture(bInUseForLumenSceneCapture)
+	, DecalRenderTargetModeMask(InDecalRenderTargetModeMask)
 {
 }
 
diff --git a/Engine/Source/Runtime/Renderer/Public/MeshPassProcessor.h b/Engine/Source/Runtime/Renderer/Public/MeshPassProcessor.h
--- a/Engine/Source/Runtime/Renderer/Public/MeshPassProcessor.h
+++ b/Engine/Source/Runtime/Renderer/Public/MeshPassProcessor.h
@@ -62,7 +62,11 @@ namespace EMeshPass
 		LumenFrontLayerTranslucencyGBuffer,
 		DitheredLODFadingOutMaskPass, /** A mini depth pass used to mark pixels with dithered LOD fading out. Currently only used by ray tracing shadows. */
 		NaniteMeshPass,
-		MeshDecal,
+		MeshDecal_DBuffer,
+		MeshDecal_SceneColorAndGBuffer,
+		MeshDecal_SceneColorAndGBufferNoNormal,
+		MeshDecal_SceneColor,
+		MeshDecal_AmbientOcclusion,
 		WaterInfoTextureDepthPass,
 		WaterInfoTexturePass,
 
@@ -113,7 +117,11 @@ inline const TCHAR* GetMeshPassName(EMeshPass::Type MeshPass)
 	case EMeshPass::LumenFrontLayerTranslucencyGBuffer: return TEXT("LumenFrontLayerTranslucencyGBuffer");
 	case EMeshPass::DitheredLODFadingOutMaskPass: return TEXT("DitheredLODFadingOutMaskPass");
 	case EMeshPass::NaniteMeshPass: return TEXT("NaniteMeshPass");
-	case EMeshPass::MeshDecal: return TEXT("MeshDecal");
+	case EMeshPass::MeshDecal_DBuffer: return TEXT("MeshDecal_DBuffer");
+	case EMeshPass::MeshDecal_SceneColorAndGBuffer: return TEXT("MeshDecal_SceneColorAndGBuffer");
+	case EMeshPass::MeshDecal_SceneColorAndGBufferNoNormal: return TEXT("MeshDecal_SceneColorAndGBufferNoNormal");
+	case EMeshPass::MeshDecal_SceneColor: return TEXT("MeshDecal_SceneColor");
+	case EMeshPass::MeshDecal_AmbientOcclusion: return TEXT("MeshDecal_AmbientOcclusion");
 	case EMeshPass::WaterInfoTextureDepthPass: return TEXT("WaterInfoTextureDepthPass");
 	case EMeshPass::WaterInfoTexturePass: return TEXT("WaterInfoTexturePass");
 #if WITH_EDITOR
@@ -125,9 +133,9 @@ inline const TCHAR* GetMeshPassName(EMeshPass::Type MeshPass)
 	}
 
 #if WITH_EDITOR
-	static_assert(EMeshPass::Num == 32 + 4, "Need to update switch(MeshPass) after changing EMeshPass"); // GUID to prevent incorrect auto-resolves, please change when changing the expression: {674D7D62-CFD8-4971-9A8D-CD91E5612CD8}
+	static_assert(EMeshPass::Num == 36 + 4, "Need to update switch(MeshPass) after changing EMeshPass"); // GUID to prevent incorrect auto-resolves, please change when changing the expression: {674D7D62-CFD8-4971-9A8D-CD91E5612CD8}
 #else
-	static_assert(EMeshPass::Num == 32, "Need to update switch(MeshPass) after changing EMeshPass"); // GUID to prevent incorrect auto-resolves, please change when changing the expression: {674D7D62-CFD8-4971-9A8D-CD91E5612CD8}
+	static_assert(EMeshPass::Num == 36, "Need to update switch(MeshPass) after changing EMeshPass"); // GUID to prevent incorrect auto-resolves, please change when changing the expression: {674D7D62-CFD8-4971-9A8D-CD91E5612CD8}
 #endif
 
 	checkf(0, TEXT("Missing case for EMeshPass %u"), (uint32)MeshPass);
diff --git a/Engine/Source/Runtime/Renderer/Public/StaticMeshBatch.h b/Engine/Source/Runtime/Renderer/Public/StaticMeshBatch.h
--- a/Engine/Source/Runtime/Renderer/Public/StaticMeshBatch.h
+++ b/Engine/Source/Runtime/Renderer/Public/StaticMeshBatch.h
@@ -31,6 +31,8 @@ private:
 	FStaticMeshBatch(const FStaticMeshBatch& InStaticMesh);
 };
 
+namespace DecalRendering { enum { DecalRenderTargetMode_NumBits = 5	}; }
+
 /**
  * FStaticMeshBatch data which is InitViews specific. Stored separately for cache efficiency.
  */
@@ -38,7 +40,19 @@ class FStaticMeshBatchRelevance
 {
 public:
 
-	FStaticMeshBatchRelevance(const FStaticMeshBatch& StaticMesh, float InScreenSize, bool InbSupportsCachingMeshDrawCommands, bool InbUseSkyMaterial, bool bInUseSingleLayerWaterMaterial, bool bInUseAnisotropy, bool bInSupportsNaniteRendering, bool bInSupportsGPUScene, bool bInUseForWaterInfoTextureDepth, bool bInUseForLumenSceneCapture, ERHIFeatureLevel::Type FeatureLevel);
+	FStaticMeshBatchRelevance(
+		const FStaticMeshBatch& StaticMesh, 
+		float InScreenSize, 
+		bool InbSupportsCachingMeshDrawCommands, 
+		bool InbUseSkyMaterial, 
+		bool bInUseSingleLayerWaterMaterial, 
+		bool bInUseAnisotropy, 
+		bool bInSupportsNaniteRendering, 
+		bool bInSupportsGPUScene, 
+		bool bInUseForWaterInfoTextureDepth, 
+		bool bInUseForLumenSceneCapture, 
+		uint8 InDecalRenderTargetModeMask,
+		ERHIFeatureLevel::Type FeatureLevel);
 
 	/** Starting offset into continuous array of command infos for this mesh in FPrimitiveSceneInfo::CachedMeshDrawCommandInfos. */
 	FMeshPassMask CommandInfosMask;
@@ -94,6 +108,9 @@ public:
 	/** Cached from lumen scene card capture */
 	uint8 bUseForLumenSceneCapture : 1;
 
+	/** What decal render target modes this mesh batch should be rendered to */
+	uint8 DecalRenderTargetModeMask : DecalRendering::DecalRenderTargetMode_NumBits;
+
 	/** Computes index of cached mesh draw command in FPrimitiveSceneInfo::CachedMeshDrawCommandInfos, for a given mesh pass. */
 	int32 GetStaticMeshCommandInfoIndex(EMeshPass::Type MeshPass) const;
 };
