Compile fixes after auto resolve issues from merge

diff --git a/Engine/Source/Runtime/Renderer/Private/CompositionLighting/PostProcessMeshDecals.cpp b/Engine/Source/Runtime/Renderer/Private/CompositionLighting/PostProcessMeshDecals.cpp
--- a/Engine/Source/Runtime/Renderer/Private/CompositionLighting/PostProcessMeshDecals.cpp
+++ b/Engine/Source/Runtime/Renderer/Private/CompositionLighting/PostProcessMeshDecals.cpp
@@ -502,8 +502,6 @@ FMeshPassProcessor* CreateMeshDecalSceneColorMeshProcessor_Mobile(ERHIFeatureLev
 REGISTER_MESHPASSPROCESSOR_AND_PSOCOLLECTOR(Mobile_MeshDecal_SceneColorAndGBuffer, CreateMeshDecalSceneColorAndGBufferMeshProcessor_Mobile, EShadingPath::Mobile, EMeshPass::MeshDecal_SceneColorAndGBuffer, EMeshPassFlags::CachedMeshCommands | EMeshPassFlags::MainView);
 REGISTER_MESHPASSPROCESSOR_AND_PSOCOLLECTOR(Mobile_MeshDecal_SceneColor, CreateMeshDecalSceneColorMeshProcessor_Mobile, EShadingPath::Mobile, EMeshPass::MeshDecal_SceneColor, EMeshPassFlags::CachedMeshCommands | EMeshPassFlags::MainView);
 
-DECLARE_CYCLE_STAT(TEXT("MeshDecal"), STAT_CLP_MeshDecal, STATGROUP_ParallelCommandListMarkers);
-
 void DrawDecalMeshCommands(
 	FRDGBuilder& GraphBuilder,
 	const FScene& Scene,
@@ -568,7 +566,7 @@ void DrawDecalMeshCommands(
 				ERDGPassFlags::Raster | ERDGPassFlags::SkipRenderPass,
 				[&View, PassParameters, DecalMeshPassType](const FRDGPass* InPass, FRHICommandListImmediate& RHICmdList)
 				{
-					FRDGParallelCommandListSet ParallelCommandListSet(InPass, RHICmdList, GET_STATID(STAT_CLP_MeshDecal), View, FParallelCommandListBindings(PassParameters));
+					FRDGParallelCommandListSet ParallelCommandListSet(InPass, RHICmdList, View, FParallelCommandListBindings(PassParameters));
 					View.ParallelMeshDrawCommandPasses[DecalMeshPassType].DispatchDraw(&ParallelCommandListSet, RHICmdList, &PassParameters->InstanceCullingDrawParams);
 				});
 		}
