Moved skylight distant lut to be a StructuredBuffer in order to use scalar load.
Some platform still use the regular non structured buffer when needed.

diff --git a/Engine/Shaders/Private/HeightFogCommon.ush b/Engine/Shaders/Private/HeightFogCommon.ush
--- a/Engine/Shaders/Private/HeightFogCommon.ush
+++ b/Engine/Shaders/Private/HeightFogCommon.ush
@@ -7,6 +7,7 @@
 #pragma once
 
 #include "/Engine/Shared/EnvironmentComponentsFlags.h"
+#include "/Engine/Private/SkyAtmosphereCommon.ush"
 
 #ifndef PERMUTATION_SUPPORT_FOG_INSCATTERING_TEXTURE
 #define PERMUTATION_SUPPORT_FOG_INSCATTERING_TEXTURE 1
@@ -75,7 +76,7 @@ float3 ComputeInscatteringColor(float3 CameraToReceiver, float CameraToReceiverL
 #endif
 
 #if PROJECT_SUPPORT_SKY_ATMOSPHERE_AFFECTS_HEIGHFOG
-	Inscattering += FogStruct.SkyAtmosphereAmbientContributionColorScale.rgb * View.SkyAtmosphereHeightFogContribution.xxx * View.DistantSkyLightLutBufferSRV[0].rgb;
+	Inscattering += FogStruct.SkyAtmosphereAmbientContributionColorScale.rgb * View.SkyAtmosphereHeightFogContribution.xxx * GetViewDistanceSkyLightColor();
 #endif
 
 	return Inscattering;
diff --git a/Engine/Shaders/Private/MaterialTemplate.ush b/Engine/Shaders/Private/MaterialTemplate.ush
--- a/Engine/Shaders/Private/MaterialTemplate.ush
+++ b/Engine/Shaders/Private/MaterialTemplate.ush
@@ -2236,7 +2236,7 @@ float4 MaterialExpressionSkyAtmosphereAerialPerspective(FMaterialPixelParameters
 float3 MaterialExpressionSkyAtmosphereDistantLightScatteredLuminance(FMaterialPixelParameters Parameters)
 {
 #if MATERIAL_SKY_ATMOSPHERE && PROJECT_SUPPORT_SKY_ATMOSPHERE
-	return View.DistantSkyLightLutBufferSRV[0].rgb;
+	return GetViewDistanceSkyLightColor();
 #else
 	return float3(0.0f, 0.0f, 0.0f);
 #endif
diff --git a/Engine/Shaders/Private/SkyAtmosphere.usf b/Engine/Shaders/Private/SkyAtmosphere.usf
--- a/Engine/Shaders/Private/SkyAtmosphere.usf
+++ b/Engine/Shaders/Private/SkyAtmosphere.usf
@@ -1321,7 +1321,8 @@ void RenderSkyViewLutCS(uint3 ThreadId : SV_DispatchThreadID)
 ////////////////////////////////////////////////////////////
 
 // see SkyAtmosphereRendering.cpp GetSkyLutTextureFormat
-RWBuffer<float4> DistantSkyLightLutBufferUAV;
+RWStructuredBuffer<float4> DistantSkyLightLutBufferUAV;
+RWBuffer<float4> MobileDistantSkyLightLutBufferUAV;
 //Buffer<float4> UniformSphereSamplesBuffer;
 float4 AtmosphereLightDirection0;
 float4 AtmosphereLightDirection1;
@@ -1414,7 +1415,9 @@ void RenderDistantSkyLightLutCS(uint3 ThreadId : SV_DispatchThreadID)
 		const float3 UniformPhaseFunction = 1.0f / (4.0f * PI);
 		
 		// Luminance assuming scattering in a medium with a uniform phase function.
-		DistantSkyLightLutBufferUAV[0] = float4(Illuminance * UniformPhaseFunction, 0.0f);
+		const float4 OutputResult			= float4(Illuminance * UniformPhaseFunction, 0.0f);
+		DistantSkyLightLutBufferUAV[0]		= OutputResult;
+		MobileDistantSkyLightLutBufferUAV[0]= OutputResult;
 		
 		// Since this is ran once per scene, we do not have access to view data (VIEWDATA_AVAILABLE==0). 
 		// So this buffer is not pre-exposed today.
diff --git a/Engine/Shaders/Private/SkyAtmosphereCommon.ush b/Engine/Shaders/Private/SkyAtmosphereCommon.ush
--- a/Engine/Shaders/Private/SkyAtmosphereCommon.ush
+++ b/Engine/Shaders/Private/SkyAtmosphereCommon.ush
@@ -333,3 +333,14 @@ MediumSampleRGB SampleAtmosphereMediumRGB(in float3 WorldPos)
 	return s;
 }
 #endif //AP_SUPPORT_SAMPLE_ATMOSHPERE
+
+float3 GetViewDistanceSkyLightColor()
+{
+#if SHADING_PATH_MOBILE && VERTEXSHADER
+	// Some mobile platforms cannot sample structured buffer in vertex shaders (height fog) so in this case a regular buffer is used.
+	return View.MobileDistantSkyLightLutBufferSRV[0].rgb;
+#else
+	// Use a structured buffer for scalar loads.
+	return View.DistantSkyLightLutBufferSRV[0].rgb; 
+#endif
+}
\ No newline at end of file
diff --git a/Engine/Shaders/Private/VolumetricCloud.usf b/Engine/Shaders/Private/VolumetricCloud.usf
--- a/Engine/Shaders/Private/VolumetricCloud.usf
+++ b/Engine/Shaders/Private/VolumetricCloud.usf
@@ -680,7 +680,7 @@ void MainCommon(in FMaterialPixelParameters MaterialParameters, in float4 SvPosi
 				// TODO Only works for ground view. This should be spatially varying for space viewsspace view. 
 				// TODO Have a illuminance texture to get that overall ambient contribution? 
 				// TODO An approximation would be to multiply it with sun transmittance?
-				DistantSkyLightLuminance = View.DistantSkyLightLutBufferSRV[0].rgb;
+				DistantSkyLightLuminance = GetViewDistanceSkyLightColor();
 			}
 			else
 			{
diff --git a/Engine/Source/Runtime/Engine/Private/SceneManagement.cpp b/Engine/Source/Runtime/Engine/Private/SceneManagement.cpp
--- a/Engine/Source/Runtime/Engine/Private/SceneManagement.cpp
+++ b/Engine/Source/Runtime/Engine/Private/SceneManagement.cpp
@@ -1159,6 +1159,7 @@ FViewUniformShaderParameters::FViewUniformShaderParameters()
 	SkyViewLutTextureSampler = TStaticSamplerState<SF_Bilinear>::GetRHI();
 
 	DistantSkyLightLutBufferSRV = GBlackVertexBufferWithSRV->ShaderResourceViewRHI;
+	MobileDistantSkyLightLutBufferSRV = GBlackStructuredBufferWithSRV->ShaderResourceViewRHI;
 
 	CameraAerialPerspectiveVolume = GBlackAlpha1VolumeTexture->TextureRHI;
 	CameraAerialPerspectiveVolumeSampler = TStaticSamplerState<SF_Bilinear>::GetRHI();
diff --git a/Engine/Source/Runtime/Engine/Public/SceneView.h b/Engine/Source/Runtime/Engine/Public/SceneView.h
--- a/Engine/Source/Runtime/Engine/Public/SceneView.h
+++ b/Engine/Source/Runtime/Engine/Public/SceneView.h
@@ -1090,7 +1090,8 @@ BEGIN_GLOBAL_SHADER_PARAMETER_STRUCT_WITH_CONSTRUCTOR(FViewUniformShaderParamete
 	SHADER_PARAMETER_SAMPLER(SamplerState, CameraAerialPerspectiveVolumeMieOnlySampler)
 	SHADER_PARAMETER_TEXTURE(Texture3D, CameraAerialPerspectiveVolumeRayOnly)
 	SHADER_PARAMETER_SAMPLER(SamplerState, CameraAerialPerspectiveVolumeRayOnlySampler)
-	SHADER_PARAMETER_SRV(Buffer<float4>, DistantSkyLightLutBufferSRV)
+	SHADER_PARAMETER_SRV(StructuredBuffer<float4>, DistantSkyLightLutBufferSRV)	
+	SHADER_PARAMETER_SRV(Buffer<float4>, MobileDistantSkyLightLutBufferSRV)		
 	// Hair
 	SHADER_PARAMETER_TEXTURE(Texture3D, HairScatteringLUTTexture)
 	SHADER_PARAMETER_SAMPLER(SamplerState, HairScatteringLUTSampler)
diff --git a/Engine/Source/Runtime/RenderCore/Private/GlobalRenderResources.cpp b/Engine/Source/Runtime/RenderCore/Private/GlobalRenderResources.cpp
--- a/Engine/Source/Runtime/RenderCore/Private/GlobalRenderResources.cpp
+++ b/Engine/Source/Runtime/RenderCore/Private/GlobalRenderResources.cpp
@@ -106,6 +106,27 @@ public:
 	}
 };
 
+class FBlackStructuredBufferWithSRV : public FVertexBufferWithSRV
+{
+public:
+	virtual void InitRHI(FRHICommandListBase& RHICmdList) override
+	{
+		// Create the buffer RHI.  		
+		FRHIResourceCreateInfo CreateInfo(TEXT("EmptyStructuredBuffer"));
+
+		const uint32 BufferSize = sizeof(float) * 4u;
+		VertexBufferRHI = RHICmdList.CreateStructuredBuffer(sizeof(float), BufferSize, BUF_Static | BUF_ShaderResource | BUF_UnorderedAccess, CreateInfo);
+
+		FVector4f* BufferData = (FVector4f*)RHICmdList.LockBuffer(VertexBufferRHI, 0, sizeof(FVector4f), RLM_WriteOnly);
+		*BufferData = FVector4f(0.0f, 0.0f, 0.0f, 0.0f);
+		RHICmdList.UnlockBuffer(VertexBufferRHI);
+
+		// Create a view of the buffer
+		ShaderResourceViewRHI = RHICmdList.CreateShaderResourceView(VertexBufferRHI);
+		UnorderedAccessViewRHI = RHICmdList.CreateUnorderedAccessView(VertexBufferRHI, false, false);
+	}
+};
+
 class FBlackTextureWithSRV : public FColoredTexture<0, 0, 0, 255>
 {
 public:
@@ -131,6 +152,7 @@ FTexture* GTransparentBlackTexture = GTransparentBlackTextureWithSRV;
 
 FVertexBufferWithSRV* GEmptyVertexBufferWithUAV = new TGlobalResource<FEmptyVertexBuffer, FRenderResource::EInitPhase::Pre>;
 FVertexBufferWithSRV* GEmptyStructuredBufferWithUAV = new TGlobalResource<FEmptyStructuredBuffer, FRenderResource::EInitPhase::Pre>;
+FVertexBufferWithSRV* GBlackStructuredBufferWithSRV = new TGlobalResource<FBlackStructuredBufferWithSRV, FRenderResource::EInitPhase::Pre>;
 
 class FWhiteVertexBuffer : public FVertexBufferWithSRV
 {
diff --git a/Engine/Source/Runtime/RenderCore/Public/GlobalRenderResources.h b/Engine/Source/Runtime/RenderCore/Public/GlobalRenderResources.h
--- a/Engine/Source/Runtime/RenderCore/Public/GlobalRenderResources.h
+++ b/Engine/Source/Runtime/RenderCore/Public/GlobalRenderResources.h
@@ -28,6 +28,9 @@ extern RENDERCORE_API FVertexBufferWithSRV* GEmptyVertexBufferWithUAV;
 // An empty structured buffer with a UAV
 extern RENDERCORE_API FVertexBufferWithSRV* GEmptyStructuredBufferWithUAV;
 
+// An empty structured buffer with a UAV
+extern RENDERCORE_API FVertexBufferWithSRV* GBlackStructuredBufferWithSRV;
+
 // An white vertex buffer with a SRV
 extern RENDERCORE_API FVertexBufferWithSRV* GWhiteVertexBufferWithSRV;
 
diff --git a/Engine/Source/Runtime/Renderer/Private/SceneRendering.cpp b/Engine/Source/Runtime/Renderer/Private/SceneRendering.cpp
--- a/Engine/Source/Runtime/Renderer/Private/SceneRendering.cpp
+++ b/Engine/Source/Runtime/Renderer/Private/SceneRendering.cpp
@@ -1287,6 +1287,7 @@ void FViewInfo::SetupUniformBufferParameters(
 	FRHITexture* CameraAerialPerspectiveVolumeMieOnlyFound = nullptr;
 	FRHITexture* CameraAerialPerspectiveVolumeRayOnlyFound = nullptr;
 	FRHIShaderResourceView* DistantSkyLightLutBufferSRVFound = nullptr;
+	FRHIShaderResourceView* MobileDistantSkyLightLutBufferSRVFound = nullptr;
 	if (ShouldRenderSkyAtmosphere(Scene, Family->EngineShowFlags))
 	{
 		ViewUniformShaderParameters.SkyAtmospherePresentInScene = 1.0f;
@@ -1304,6 +1305,7 @@ void FViewInfo::SetupUniformBufferParameters(
 		}
 
 		DistantSkyLightLutBufferSRVFound = SkyAtmosphere->GetDistantSkyLightLutBufferSRV();
+		MobileDistantSkyLightLutBufferSRVFound = SkyAtmosphere->GetMobileDistantSkyLightLutBufferSRV();
 
 		if (this->SkyAtmosphereCameraAerialPerspectiveVolume.IsValid())
 		{
@@ -1432,11 +1434,22 @@ void FViewInfo::SetupUniformBufferParameters(
 
 	ViewUniformShaderParameters.TransmittanceLutTexture = OrWhite2DIfNull(TransmittanceLutTextureFound);
 	ViewUniformShaderParameters.TransmittanceLutTextureSampler = TStaticSamplerState<SF_Bilinear>::GetRHI();
-	ViewUniformShaderParameters.DistantSkyLightLutBufferSRV = GBlackVertexBufferWithSRV->ShaderResourceViewRHI.GetReference();
 	if(DistantSkyLightLutBufferSRVFound != nullptr)
 	{
 		ViewUniformShaderParameters.DistantSkyLightLutBufferSRV = DistantSkyLightLutBufferSRVFound;
 	}
+	else
+	{
+		ViewUniformShaderParameters.DistantSkyLightLutBufferSRV = GBlackStructuredBufferWithSRV->ShaderResourceViewRHI.GetReference();
+	}
+	if(MobileDistantSkyLightLutBufferSRVFound != nullptr)
+	{
+		ViewUniformShaderParameters.MobileDistantSkyLightLutBufferSRV = MobileDistantSkyLightLutBufferSRVFound;
+	}
+	else
+	{
+		ViewUniformShaderParameters.MobileDistantSkyLightLutBufferSRV = GBlackVertexBufferWithSRV->ShaderResourceViewRHI.GetReference();
+	}
 	ViewUniformShaderParameters.SkyViewLutTexture = OrBlack2DIfNull(SkyViewLutTextureFound);
 	ViewUniformShaderParameters.SkyViewLutTextureSampler = TStaticSamplerState<SF_Bilinear>::GetRHI();
 	ViewUniformShaderParameters.CameraAerialPerspectiveVolume = OrBlack3DAlpha1IfNull(CameraAerialPerspectiveVolumeFound);
diff --git a/Engine/Source/Runtime/Renderer/Private/SkyAtmosphereRendering.cpp b/Engine/Source/Runtime/Renderer/Private/SkyAtmosphereRendering.cpp
--- a/Engine/Source/Runtime/Renderer/Private/SkyAtmosphereRendering.cpp
+++ b/Engine/Source/Runtime/Renderer/Private/SkyAtmosphereRendering.cpp
@@ -498,8 +498,11 @@ FSkyAtmosphereRenderSceneInfo::~FSkyAtmosphereRenderSceneInfo()
 
 void FSkyAtmosphereRenderSceneInfo::CreateDistantSkyLightLutBufferAndSRV(FRDGBuilder& GraphBuilder)
 {
-	DistantSkyLightLutBuffer = AllocatePooledBuffer(FRDGBufferDesc::CreateBufferDesc(sizeof(FVector4f), 1), TEXT("SkyAtmosphere.DistantSkyLightLutBuffer"));
+	DistantSkyLightLutBuffer = AllocatePooledBuffer(FRDGBufferDesc::CreateStructuredDesc(sizeof(FVector4f), 1), TEXT("SkyAtmosphere.DistantSkyLightLutBuffer"));
 	DistantSkyLightLutBufferSRV = DistantSkyLightLutBuffer->GetOrCreateSRV(GraphBuilder.RHICmdList, FRHIBufferSRVCreateInfo(PF_A32B32G32R32F));
+
+	MobileDistantSkyLightLutBuffer = AllocatePooledBuffer(FRDGBufferDesc::CreateBufferDesc(sizeof(FVector4f), 1), TEXT("SkyAtmosphere.MobileDistantSkyLightLutBuffer"));
+	MobileDistantSkyLightLutBufferSRV = MobileDistantSkyLightLutBuffer->GetOrCreateSRV(GraphBuilder.RHICmdList, FRHIBufferSRVCreateInfo(PF_A32B32G32R32F));
 }
 
 TRefCountPtr<FRDGPooledBuffer>& FSkyAtmosphereRenderSceneInfo::GetDistantSkyLightLutBuffer()
@@ -508,6 +511,12 @@ TRefCountPtr<FRDGPooledBuffer>& FSkyAtmosphereRenderSceneInfo::GetDistantSkyLigh
 	return DistantSkyLightLutBuffer;
 }
 
+TRefCountPtr<FRDGPooledBuffer>& FSkyAtmosphereRenderSceneInfo::GetMobileDistantSkyLightLutBuffer()
+{
+	check(CVarSkyAtmosphereDistantSkyLightLUT.GetValueOnRenderThread() > 0);
+	return MobileDistantSkyLightLutBuffer;
+}
+
 FRHIShaderResourceView* FSkyAtmosphereRenderSceneInfo::GetDistantSkyLightLutBufferSRV()
 {
 	if (CVarSkyAtmosphereDistantSkyLightLUT.GetValueOnRenderThread() > 0)
@@ -517,6 +526,15 @@ FRHIShaderResourceView* FSkyAtmosphereRenderSceneInfo::GetDistantSkyLightLutBuff
 	return GBlackVertexBufferWithSRV->ShaderResourceViewRHI;
 }
 
+FRHIShaderResourceView* FSkyAtmosphereRenderSceneInfo::GetMobileDistantSkyLightLutBufferSRV()
+{
+	if (CVarSkyAtmosphereDistantSkyLightLUT.GetValueOnRenderThread() > 0)
+	{
+		return MobileDistantSkyLightLutBufferSRV;
+	}
+	return GBlackStructuredBufferWithSRV->ShaderResourceViewRHI;
+}
+
 
 
 /*=============================================================================
@@ -837,7 +855,8 @@ public:
 	BEGIN_SHADER_PARAMETER_STRUCT(FParameters, )
 		SHADER_PARAMETER_STRUCT_REF(FAtmosphereUniformShaderParameters, Atmosphere)
 		SHADER_PARAMETER_STRUCT_REF(FSkyAtmosphereInternalCommonParameters, SkyAtmosphere)
-		SHADER_PARAMETER_RDG_BUFFER_UAV(RWBuffer<float4>, DistantSkyLightLutBufferUAV)
+		SHADER_PARAMETER_RDG_BUFFER_UAV(RWStructureBuffer<float4>, DistantSkyLightLutBufferUAV)
+		SHADER_PARAMETER_RDG_BUFFER_UAV(RWBuffer<float4>, MobileDistantSkyLightLutBufferUAV)
 		SHADER_PARAMETER_RDG_TEXTURE(Texture2D<float3>, TransmittanceLutTexture)
 		SHADER_PARAMETER_RDG_TEXTURE(Texture2D<float3>, MultiScatteredLuminanceLutTexture)
 		SHADER_PARAMETER_SAMPLER(SamplerState, TransmittanceLutTextureSampler)
@@ -1378,6 +1397,8 @@ void FSceneRenderer::RenderSkyAtmosphereLookUpTables(FRDGBuilder& GraphBuilder,
 	{
 		FRDGBufferRef DistantSkyLightLutBuffer = GraphBuilder.RegisterExternalBuffer(SkyInfo.GetDistantSkyLightLutBuffer());
 		FRDGBufferUAVRef DistantSkyLightLutBufferUAV = GraphBuilder.CreateUAV(FRDGBufferUAVDesc(DistantSkyLightLutBuffer, PF_A32B32G32R32F));
+		FRDGBufferRef MobileDistantSkyLightLutBuffer = GraphBuilder.RegisterExternalBuffer(SkyInfo.GetMobileDistantSkyLightLutBuffer());
+		FRDGBufferUAVRef MobileDistantSkyLightLutBufferUAV = GraphBuilder.CreateUAV(FRDGBufferUAVDesc(MobileDistantSkyLightLutBuffer, PF_A32B32G32R32F));
 
 		FRenderDistantSkyLightLutCS::FPermutationDomain PermutationVector;
 		PermutationVector.Set<FSecondAtmosphereLight>(bSecondAtmosphereLightEnabled);
@@ -1392,6 +1413,7 @@ void FSceneRenderer::RenderSkyAtmosphereLookUpTables(FRDGBuilder& GraphBuilder,
 		PassParameters->MultiScatteredLuminanceLutTexture = MultiScatteredLuminanceLut;
 		PassParameters->UniformSphereSamplesBuffer = GUniformSphereSamplesBuffer.UniformSphereSamplesBuffer.SRV;
 		PassParameters->DistantSkyLightLutBufferUAV = DistantSkyLightLutBufferUAV;
+		PassParameters->MobileDistantSkyLightLutBufferUAV = MobileDistantSkyLightLutBufferUAV;
 
 		FLightSceneInfo* Light0 = Scene->AtmosphereLights[0];
 		FLightSceneInfo* Light1 = Scene->AtmosphereLights[1];
@@ -1422,6 +1444,7 @@ void FSceneRenderer::RenderSkyAtmosphereLookUpTables(FRDGBuilder& GraphBuilder,
 		FComputeShaderUtils::AddPass(GraphBuilder, RDG_EVENT_NAME("DistantSkyLightLut"), PassFlag, ComputeShader, PassParameters, NumGroups);
 
 		PendingRDGResources.DistantSkyLightLutBuffer = DistantSkyLightLutBuffer;
+		PendingRDGResources.MobileDistantSkyLightLutBuffer = MobileDistantSkyLightLutBuffer;
 	}
 
 	SkyAtmosphereLightShadowData LightShadowData;
@@ -1716,6 +1739,10 @@ void FSkyAtmospherePendingRDGResources::CommitToSceneAndViewUniformBuffers(FRDGB
 	{
 		SkyInfo.GetDistantSkyLightLutBuffer() = ConvertToExternalAccessBuffer(GraphBuilder, ExternalAccessQueue, DistantSkyLightLutBuffer, ERHIAccess::SRVMask, ERHIPipeline::All);
 	}
+	if (MobileDistantSkyLightLutBuffer && CVarSkyAtmosphereDistantSkyLightLUT.GetValueOnRenderThread() > 0)
+	{
+		SkyInfo.GetMobileDistantSkyLightLutBuffer() = ConvertToExternalAccessBuffer(GraphBuilder, ExternalAccessQueue, MobileDistantSkyLightLutBuffer, ERHIAccess::SRVMask, ERHIPipeline::All);
+	}
 
 	if (RealTimeReflectionCaptureSkyAtmosphereViewLutTexture)
 	{
diff --git a/Engine/Source/Runtime/Renderer/Private/SkyAtmosphereRendering.h b/Engine/Source/Runtime/Renderer/Private/SkyAtmosphereRendering.h
--- a/Engine/Source/Runtime/Renderer/Private/SkyAtmosphereRendering.h
+++ b/Engine/Source/Runtime/Renderer/Private/SkyAtmosphereRendering.h
@@ -147,7 +147,9 @@ public:
 
 	void CreateDistantSkyLightLutBufferAndSRV(FRDGBuilder& GraphBuilder);
 	TRefCountPtr<FRDGPooledBuffer>& GetDistantSkyLightLutBuffer();
+	TRefCountPtr<FRDGPooledBuffer>& GetMobileDistantSkyLightLutBuffer();
 	FRHIShaderResourceView* GetDistantSkyLightLutBufferSRV();
+	FRHIShaderResourceView* GetMobileDistantSkyLightLutBufferSRV();
 
 	FRDGTextureRef GetTransmittanceLutTexture(FRDGBuilder& GraphBuilder) const { return GraphBuilder.RegisterExternalTexture(TransmittanceLutTexture); }
 
@@ -170,6 +172,8 @@ private:
 	TRefCountPtr<IPooledRenderTarget> MultiScatteredLuminanceLutTexture;
 	TRefCountPtr<FRDGPooledBuffer> DistantSkyLightLutBuffer;
 	FRHIShaderResourceView* DistantSkyLightLutBufferSRV = nullptr;
+	TRefCountPtr<FRDGPooledBuffer> MobileDistantSkyLightLutBuffer;
+	FRHIShaderResourceView* MobileDistantSkyLightLutBufferSRV = nullptr;
 };
 
 /** Pending RDG resource to commit after the pre-pass / nanite rasterization so that RenderSkyAtmosphereLookUpTables() can overlap them on async compute. */
@@ -190,6 +194,7 @@ private:
 	TArray<FViewRDGResources, TInlineAllocator<4>> ViewResources;
 
 	FRDGBufferRef DistantSkyLightLutBuffer = nullptr;
+	FRDGBufferRef MobileDistantSkyLightLutBuffer = nullptr;
 	FRDGTextureRef RealTimeReflectionCaptureSkyAtmosphereViewLutTexture = nullptr;
 	FRDGTextureRef RealTimeReflectionCaptureCamera360APLutTexture = nullptr;
 	FRDGTextureRef TransmittanceLut = nullptr;
