Enable cached mesh draw command decals by default on all platforms

diff --git a/Engine/Source/Runtime/Renderer/Private/CompositionLighting/PostProcessMeshDecals.cpp b/Engine/Source/Runtime/Renderer/Private/CompositionLighting/PostProcessMeshDecals.cpp
--- a/Engine/Source/Runtime/Renderer/Private/CompositionLighting/PostProcessMeshDecals.cpp
+++ b/Engine/Source/Runtime/Renderer/Private/CompositionLighting/PostProcessMeshDecals.cpp
@@ -32,7 +32,7 @@ static TAutoConsoleVariable<int32> CVarParallelMeshDecal(
 
 static TAutoConsoleVariable<int32> CVarDecalUseCachedMDCs(
 	TEXT("r.Decal.UseCachedMDCs"),
-	0,
+	1,
 	TEXT("Use cached MDCs for decal rendering."),
 	FConsoleVariableDelegate::CreateLambda([](IConsoleVariable* InVariable)
 	{
