Fixed incorrect sorting of AlwaysVisible primitives which was causing the two categories to become interleaved.

diff --git a/Engine/Source/Runtime/Renderer/Private/RendererScene.cpp b/Engine/Source/Runtime/Renderer/Private/RendererScene.cpp
--- a/Engine/Source/Runtime/Renderer/Private/RendererScene.cpp
+++ b/Engine/Source/Runtime/Renderer/Private/RendererScene.cpp
@@ -439,6 +439,12 @@ FSceneViewState::~FSceneViewState()
 	}
 }
 
+inline FScene::FPrimitiveSceneProxyType::FPrimitiveSceneProxyType(const FPrimitiveSceneProxy *PrimitiveSceneProxy) 
+	: ProxyTypeHash(PrimitiveSceneProxy->GetTypeHash())
+	, bIsAlwaysVisible(PrimitiveSceneProxy->IsAlwaysVisible())
+{
+}
+
 void FScene::RemoveViewLumenSceneData_RenderThread(FSceneViewStateInterface* ViewState)
 {
 	FLumenSceneDataKey ByViewKey = { ViewState->GetViewKey(), (uint32)INDEX_NONE };
@@ -1354,8 +1360,8 @@ void FScene::CheckPrimitiveArrays(int MaxTypeOffsetIndex)
 		for (uint32 Index = NextOffset; Index < Entry.Offset; Index++)
 		{
 			checkSlow(Primitives[Index]->Proxy == PrimitiveSceneProxies[Index]);
-			SIZE_T TypeHash = PrimitiveSceneProxies[Index]->GetTypeHash();
-			checkfSlow(TypeHash == Entry.PrimitiveSceneProxyType, TEXT("TypeHash: %i not matching, expected: %i"), TypeHash, Entry.PrimitiveSceneProxyType);
+			FPrimitiveSceneProxyType PrimitiveSceneProxyType = FPrimitiveSceneProxyType(PrimitiveSceneProxies[Index]);
+			checkfSlow(PrimitiveSceneProxyType == Entry.PrimitiveSceneProxyType, TEXT("TypeHash: %i not matching, expected: %i"), PrimitiveSceneProxyType.ProxyTypeHash, Entry.PrimitiveSceneProxyType.ProxyTypeHash);
 		}
 		NextOffset = Entry.Offset;
 	}
@@ -5636,9 +5642,9 @@ void FScene::Update(FRDGBuilder& GraphBuilder, const FUpdateParameters& Paramete
 		while (RemovedLocalPrimitiveSceneInfos.Num())
 		{
 			int32 StartIndex = RemovedLocalPrimitiveSceneInfos.Num() - 1;
-			SIZE_T InsertProxyHash = RemovedLocalPrimitiveSceneInfos[StartIndex]->Proxy->GetTypeHash();
+			FPrimitiveSceneProxyType RemovedProxyType = FPrimitiveSceneProxyType(RemovedLocalPrimitiveSceneInfos[StartIndex]->Proxy);
 
-			while (StartIndex > 0 && RemovedLocalPrimitiveSceneInfos[StartIndex - 1]->Proxy->GetTypeHash() == InsertProxyHash)
+			while (StartIndex > 0 && FPrimitiveSceneProxyType(RemovedLocalPrimitiveSceneInfos[StartIndex - 1]->Proxy) == RemovedProxyType)
 			{
 				StartIndex--;
 			}
@@ -5651,7 +5657,7 @@ void FScene::Update(FRDGBuilder& GraphBuilder, const FUpdateParameters& Paramete
 				// PrimitiveSceneProxies[0,0,0,6,6,6,6,6,2,2,2,2,1,1,1,7,4,8]
 				// TypeOffsetTable[3,8,12,15,16,17,18]
 
-				if (TypeOffsetTable[BroadIndex].PrimitiveSceneProxyType == InsertProxyHash)
+				if (TypeOffsetTable[BroadIndex].PrimitiveSceneProxyType == RemovedProxyType)
 				{
 					const int32 InsertionOffset = TypeOffsetTable[BroadIndex].Offset;
 					const int32 PrevOffset = BroadIndex > 0 ? TypeOffsetTable[BroadIndex - 1].Offset : 0;
@@ -5951,9 +5957,9 @@ void FScene::Update(FRDGBuilder& GraphBuilder, const FUpdateParameters& Paramete
 		while (AddedLocalPrimitiveSceneInfos.Num())
 		{
 			int32 StartIndex = AddedLocalPrimitiveSceneInfos.Num() - 1;
-			SIZE_T InsertProxyHash = AddedLocalPrimitiveSceneInfos[StartIndex]->Proxy->GetTypeHash();
+			FPrimitiveSceneProxyType InsertProxyType = FPrimitiveSceneProxyType(AddedLocalPrimitiveSceneInfos[StartIndex]->Proxy);
 
-			while (StartIndex > 0 && AddedLocalPrimitiveSceneInfos[StartIndex - 1]->Proxy->GetTypeHash() == InsertProxyHash)
+			while (StartIndex > 0 && FPrimitiveSceneProxyType(AddedLocalPrimitiveSceneInfos[StartIndex - 1]->Proxy) == InsertProxyType)
 			{
 				StartIndex--;
 			}
@@ -6010,7 +6016,7 @@ void FScene::Update(FRDGBuilder& GraphBuilder, const FUpdateParameters& Paramete
 				// PrimitiveSceneProxies[0,0,0,6,6,6,6,6,2,2,2,2,1,1,1,7,4,8]
 				// TypeOffsetTable[3,8,12,15,16,17,18]
 
-				if (TypeOffsetTable[BroadIndex].PrimitiveSceneProxyType == InsertProxyHash)
+				if (TypeOffsetTable[BroadIndex].PrimitiveSceneProxyType == InsertProxyType)
 				{
 					EntryFound = true;
 					break;
@@ -6037,12 +6043,12 @@ void FScene::Update(FRDGBuilder& GraphBuilder, const FUpdateParameters& Paramete
 					}
 
 					int32 PrevEntryOffset = BroadIndex > 0 ? TypeOffsetTable[BroadIndex - 1].Offset : 0;
-					TypeOffsetTable.Insert(FTypeOffsetTableEntry(InsertProxyHash, PrevEntryOffset), BroadIndex);
+					TypeOffsetTable.Insert(FTypeOffsetTableEntry(InsertProxyType, PrevEntryOffset), BroadIndex);
 				}
 				else
 				{
 					// Starting with an empty list and zero offset (offset will be incremented during the while loop)
-					TypeOffsetTable.Push(FTypeOffsetTableEntry(InsertProxyHash, 0));
+					TypeOffsetTable.Push(FTypeOffsetTableEntry(InsertProxyType, 0));
 				}
 			}
 
@@ -6342,8 +6348,8 @@ void FScene::Update(FRDGBuilder& GraphBuilder, const FUpdateParameters& Paramete
 			checkSlow(Primitives[NextTypeOffset]->Proxy == PrimitiveSceneProxies[NextTypeOffset]);
 
 			// Sanity check
-			const SIZE_T TypeHash = PrimitiveSceneProxies[NextTypeOffset]->GetTypeHash();
-			checkfSlow (TypeHash == TypeEntry.PrimitiveSceneProxyType, TEXT("TypeHash: %i not matching TypeOffsetTable, expected: %i"), TypeHash, TypeEntry.PrimitiveSceneProxyType);
+			const FPrimitiveSceneProxyType Type = FPrimitiveSceneProxyType(PrimitiveSceneProxies[NextTypeOffset]);
+			checkfSlow (Type == TypeEntry.PrimitiveSceneProxyType, TEXT("TypeHash: %i not matching TypeOffsetTable, expected: %i"), Type.ProxyTypeHash, TypeEntry.PrimitiveSceneProxyType.ProxyTypeHash);
 		#endif
 
 			if (PrimitiveSceneProxies[NextTypeOffset]->IsAlwaysVisible())
@@ -6365,7 +6371,6 @@ void FScene::Update(FRDGBuilder& GraphBuilder, const FUpdateParameters& Paramete
 			if (uint32(Test) < PrimitivesAlwaysVisibleOffset)
 			{
 				check(!AlwaysVisible);
-				check(!IsNanite);
 			}
 			else
 			{
diff --git a/Engine/Source/Runtime/Renderer/Private/ScenePrivate.h b/Engine/Source/Runtime/Renderer/Private/ScenePrivate.h
--- a/Engine/Source/Runtime/Renderer/Private/ScenePrivate.h
+++ b/Engine/Source/Runtime/Renderer/Private/ScenePrivate.h
@@ -2929,10 +2929,22 @@ public:
 
 	TArray<int32> PersistentPrimitiveIdToIndexMap;
 
+	/**
+	 * Defines a bucket "type" in the sorted order of the primitive arrays, as defined by the type-offset table.
+	 */
+	struct FPrimitiveSceneProxyType
+	{
+		FPrimitiveSceneProxyType(const FPrimitiveSceneProxy *PrimitiveSceneProxy);
+		bool operator ==(const FPrimitiveSceneProxyType&) const = default;
+
+		SIZE_T ProxyTypeHash; 
+		bool bIsAlwaysVisible;
+	};
+
 	struct FTypeOffsetTableEntry
 	{
-		FTypeOffsetTableEntry(SIZE_T InPrimitiveSceneProxyType, uint32 InOffset) : PrimitiveSceneProxyType(InPrimitiveSceneProxyType), Offset(InOffset) {}
-		SIZE_T PrimitiveSceneProxyType;
+		FTypeOffsetTableEntry(const FPrimitiveSceneProxyType &InPrimitiveSceneProxyType, uint32 InOffset) : PrimitiveSceneProxyType(InPrimitiveSceneProxyType), Offset(InOffset) {}
+		FPrimitiveSceneProxyType PrimitiveSceneProxyType;
 		uint32 Offset; //(e.g. prefix sum where the next type starts)
 	};
 	/* During insertion and deletion, used to skip large chunks of items of the same type */
