Enable r.Visibility.SkipAlwaysVisible by default now

diff --git a/Engine/Source/Runtime/Renderer/Private/RendererScene.cpp b/Engine/Source/Runtime/Renderer/Private/RendererScene.cpp
--- a/Engine/Source/Runtime/Renderer/Private/RendererScene.cpp
+++ b/Engine/Source/Runtime/Renderer/Private/RendererScene.cpp
@@ -139,8 +139,8 @@ static TAutoConsoleVariable<int32> CVarBasePassWriteDepthEvenWithFullPrepass(
 	TEXT("0 to allow a readonly base pass, which skips an MSAA depth resolve, and allows masked materials to get EarlyZ (writing to depth while doing clip() disables EarlyZ) (default)\n")
 	TEXT("1 to force depth writes in the base pass.  Useful for debugging when the prepass and base pass don't match what they render."));
 
-// TODO: Significant render thread optimization to heavy Nanite scenes - Off by default, pending extensive testing
-int32 GVisibilitySkipAlwaysVisible = 0;
+// Significant render thread (InitViews / Relevancy) optimization to Nanite heavy scenes
+int32 GVisibilitySkipAlwaysVisible = 1;
 static FAutoConsoleVariableRef CVarVisibilitySkipAlwaysVisible(
 	TEXT("r.Visibility.SkipAlwaysVisible"),
 	GVisibilitySkipAlwaysVisible,
