Fix incorrect rendering of materials that require global distance field when using Path Tracing.

diff --git a/Engine/Source/Runtime/Renderer/Private/DistanceFieldAmbientOcclusion.cpp b/Engine/Source/Runtime/Renderer/Private/DistanceFieldAmbientOcclusion.cpp
--- a/Engine/Source/Runtime/Renderer/Private/DistanceFieldAmbientOcclusion.cpp
+++ b/Engine/Source/Runtime/Renderer/Private/DistanceFieldAmbientOcclusion.cpp
@@ -734,15 +734,15 @@ bool FSceneRenderer::ShouldPrepareDistanceFieldScene() const
 		return false;
 	}
 
+	const bool bShouldPrepareForAO = ShouldPrepareForDistanceFieldAO();
+	const bool bShouldPrepareGlobalDistanceField = ShouldPrepareGlobalDistanceField();
+	const bool bShouldPrepareForDFInsetIndirectShadow = ShouldPrepareForDFInsetIndirectShadow();
+
 	if (ViewFamily.EngineShowFlags.PathTracing)
 	{
-		return false;
+		return bShouldPrepareGlobalDistanceField;
 	}
 
-	bool bShouldPrepareForAO = SupportsDistanceFieldAO(Scene->GetFeatureLevel(), Scene->GetShaderPlatform()) && ShouldPrepareForDistanceFieldAO();
-	bool bShouldPrepareGlobalDistanceField = ShouldPrepareGlobalDistanceField();
-	bool bShouldPrepareForDFInsetIndirectShadow = ShouldPrepareForDFInsetIndirectShadow();
-
 	// Prepare the distance field scene (object buffers and distance field atlas) if any feature needs it
 	return bShouldPrepareGlobalDistanceField 
 		|| bShouldPrepareForAO 
@@ -759,21 +759,23 @@ bool FSceneRenderer::ShouldPrepareGlobalDistanceField() const
 		return false;
 	}
 
-	if (!DoesProjectSupportDistanceFields())
+	if (!DoesProjectSupportDistanceFields() || !UseGlobalDistanceField())
 	{
 		return false;
 	}
 
+	const bool bShouldPrepareForMaterialsOrNiagara = 
+		!GAOGlobalDistanceFieldDetailedNecessityCheck
+		|| Views[0].PrevViewInfo.bUsesGlobalDistanceField // use previous frame relevance because this frame relevance flags haven't been calculate yet
+		|| ((FXSystem != nullptr) && FXSystem->UsesGlobalDistanceField());
+
 	if (ViewFamily.EngineShowFlags.PathTracing)
 	{
-		return false;
+		return bShouldPrepareForMaterialsOrNiagara;
 	}
 
 	const bool bShouldPrepareForAO = SupportsDistanceFieldAO(Scene->GetFeatureLevel(), Scene->GetShaderPlatform())
-		&& (!GAOGlobalDistanceFieldDetailedNecessityCheck
-			|| ShouldPrepareForDistanceFieldAO()
-			|| Views[0].PrevViewInfo.bUsesGlobalDistanceField // use previous frame relevance because this frame relevance flags haven't been calculate yet
-			|| ((FXSystem != nullptr) && FXSystem->UsesGlobalDistanceField()));
+		&& (ShouldPrepareForDistanceFieldAO() || bShouldPrepareForMaterialsOrNiagara);
 
 	const bool bShouldPrepareForLumen = IsLumenEnabled(Views[0]) && Lumen::UseGlobalSDFObjectGrid(ViewFamily);
 
@@ -782,7 +784,7 @@ bool FSceneRenderer::ShouldPrepareGlobalDistanceField() const
 
 	const bool bShouldPrepareForVisualization = ViewFamily.EngineShowFlags.VisualizeGlobalDistanceField;
 
-	return (bShouldPrepareForAO || bShouldPrepareForLumen || bShouldPrepareForManyLights || bShouldPrepareForVisualization) && UseGlobalDistanceField();
+	return (bShouldPrepareForAO || bShouldPrepareForLumen || bShouldPrepareForManyLights || bShouldPrepareForVisualization);
 }
 
 void FDeferredShadingSceneRenderer::RenderDFAOAsIndirectShadowing(
