Path tracing denoiser support update:
1. Change variance output for radiance and albedo from population std to the std of the sampling distribution of mean. Change that for normal to output the std of the mean length.
2. Stop path tracing variance buffer update in editor after denoising has been applied.

diff --git a/Engine/Shaders/Private/PathTracing/PathTracingSpatialTemporalDenoising.usf b/Engine/Shaders/Private/PathTracing/PathTracingSpatialTemporalDenoising.usf
--- a/Engine/Shaders/Private/PathTracing/PathTracingSpatialTemporalDenoising.usf
+++ b/Engine/Shaders/Private/PathTracing/PathTracingSpatialTemporalDenoising.usf
@@ -805,7 +805,7 @@ float4 GetMeanFromTextures(int2 Position)
 	//The length of the average normal indicates the variance
 	//1: All pointing to the same direction
 	//0: Uniform in all directions
-	Value.z = length2(NormalTexture[Position].rgb);
+	Value.z = length(NormalTexture[Position].rgb);
 	
 	// Special case handling:
 	// 1. Regions without normal has a zero length, like sky
@@ -871,12 +871,10 @@ void PrepassGenerateTextureCS(uint3 DT_ID : SV_DispatchThreadID)
 		int index = Position.x + Position.y * TargetViewport_Extent.x;
 		FPixelMaterialLightingFingerprint VarianceInfo = VarianceMap[index];
 
-		float4 StandDerivation= sqrt((VarianceInfo.Var) / (VarianceInfo.Mean.w + 1));
+		// VarianceInfo.Var = sigma^2 * n.
+		// Std of mean = sqrt(sigma^2/n) = sqrt(VarianceInfo.Var/n^2)
+		float4 StandDerivation= sqrt(VarianceInfo.Var) / (VarianceInfo.Mean.w + 1);
 
-#if VARIANCE_TYPE == RADIANCE_ALBEDO_NORMAL_SINGLECHANNEL
-		// The normal variance is indicated by 1 - sqrt for the lengh2(normal).
-		StandDerivation.z = 1 - sqrt(VarianceInfo.Mean.z);
-#endif
 		OutputTexture[Position] = float4(StandDerivation.xyz, 1.0f);
 	}
 }
@@ -1561,4 +1559,4 @@ void ConvertWorldSpaceNormalToCameraSpaceCS(uint3 DT_ID : SV_DispatchThreadID)
 		NormalTexture[Position] = float4(mul(Normal, View.TranslatedWorldToCameraView).xyz, 0.0f);
 	}
 }
-#endif
\ No newline at end of file
+#endif
diff --git a/Engine/Source/Runtime/Renderer/Private/PathTracing.cpp b/Engine/Source/Runtime/Renderer/Private/PathTracing.cpp
--- a/Engine/Source/Runtime/Renderer/Private/PathTracing.cpp
+++ b/Engine/Source/Runtime/Renderer/Private/PathTracing.cpp
@@ -3248,13 +3248,12 @@ void FDeferredShadingSceneRenderer::RenderPathTracing(
 			DenoisingContext.RadianceTexture = RadianceTexture;
 			DenoisingContext.AlbedoTexture = AlbedoTexture;
 			DenoisingContext.NormalTexture = NormalTexture;
+			DenoisingContext.VarianceBuffer = PathTracingState->VarianceBuffer ? 
+				GraphBuilder.RegisterExternalBuffer(PathTracingState->VarianceBuffer, TEXT("PathTracing.VarianceBuffer")) : nullptr;
+			DenoisingContext.LastVarianceBuffer = PathTracingState->LastVarianceBuffer ?
+				GraphBuilder.RegisterExternalBuffer(PathTracingState->LastVarianceBuffer, TEXT("PathTracing.LastVarianceBuffer")) : nullptr;
 
-			if (PathTracingState->VarianceBuffer)
-			{
-				DenoisingContext.VarianceBuffer = GraphBuilder.RegisterExternalBuffer(PathTracingState->VarianceBuffer, TEXT("PathTracing.VarianceBuffer"));
-			}
-
-			PathTracingSpatialTemporalDenoisingPrePass(GraphBuilder, View, Config.PathTracingData.Iteration, DenoisingContext);
+			PathTracingSpatialTemporalDenoisingPrePass(GraphBuilder, View, Config.PathTracingData.Iteration, MaxSPP, DenoisingContext);
 
 			CurrentVarianceBufer = DenoisingContext.VarianceBuffer;
 		}
@@ -3276,9 +3275,6 @@ void FDeferredShadingSceneRenderer::RenderPathTracing(
 					GraphBuilder.RegisterExternalTexture(PathTracingState->LastNormalRT, TEXT("PathTracing.LastNormalTexture"));
 				DenoisingContext.LastAlbedoTexture =
 					GraphBuilder.RegisterExternalTexture(PathTracingState->LastAlbedoRT, TEXT("PathTracing.LastAlbedoTexture"));
-
-				DenoisingContext.LastVarianceBuffer = PathTracingState->LastVarianceBuffer?
-					GraphBuilder.RegisterExternalBuffer(PathTracingState->LastVarianceBuffer, TEXT("PathTracing.LastVarianceBuffer")) : nullptr;
 			}
 
 			PathTracingSpatialTemporalDenoising(GraphBuilder,
diff --git a/Engine/Source/Runtime/Renderer/Private/PathTracingSpatialTemporalDenoising.cpp b/Engine/Source/Runtime/Renderer/Private/PathTracingSpatialTemporalDenoising.cpp
--- a/Engine/Source/Runtime/Renderer/Private/PathTracingSpatialTemporalDenoising.cpp
+++ b/Engine/Source/Runtime/Renderer/Private/PathTracingSpatialTemporalDenoising.cpp
@@ -57,8 +57,8 @@ namespace {
 	TAutoConsoleVariable<int32> CVarPathTracingDenoiserPrepassOutputVarianceTexture(
 		TEXT("r.PathTracing.Denoiser.Prepass.OutputVarianceTexture"),
 		1,
-		TEXT("0: Variance is used only in the denoiser")
-		TEXT("1: Output to the postprocess material, usually used by MRQ")
+		TEXT("0: No variance texture will be generated in Prepass")
+		TEXT("1: Output variance texture to denoisers, or the postprocess material usually used by MRQ")
 	);
 
 	TAutoConsoleVariable<int32> CVarPathTracingSpatialDenoiser(
@@ -519,10 +519,12 @@ static bool ShouldPrepassOutputVarianceTexture(const FViewInfo& View)
 		IConsoleManager::Get().FindTConsoleVariableDataInt(TEXT("r.PathTracing.OutputPostProcessResources"));
 	const bool bOutputPostProcessResources = CVarOutputPostProcessResources ?
 		(CVarOutputPostProcessResources->GetValueOnRenderThread() != 0) : false;
-
-	return CVarPathTracingDenoiserPrepassOutputVarianceTexture.GetValueOnRenderThread() != 0 &&
-		bOutputPostProcessResources && 
-		IsPathTracingVarianceTextureRequiredInPostProcessMaterial(View);
+	
+	// Variance texture will be available if post process requires and we allow output post process resource in path tracing
+	// or when we allow prepass to output so it can be accessed by denoisers based on variance.
+	return CVarPathTracingDenoiserPrepassOutputVarianceTexture.GetValueOnRenderThread() != 0 ||
+		(bOutputPostProcessResources && 
+		IsPathTracingVarianceTextureRequiredInPostProcessMaterial(View));
 }
 
 static constexpr uint32 kMipDiffDelta = 2;
@@ -964,6 +966,7 @@ static void PathTracingSpatialTemporalDenoiserPlugin(FRDGBuilder& GraphBuilder,
 	Inputs.ColorTex = InputTexture;
 	Inputs.AlbedoTex = AlbedoTexture;
 	Inputs.NormalTex = NormalTexture;
+	Inputs.VarianceTex = Context.VarianceTexture;
 	Inputs.OutputTex = OutputTexture;
 	Inputs.FlowTex = FlowTexture;
 	Inputs.PreviousOutputTex = PreviousOutputFrameTexture;
@@ -1684,23 +1687,26 @@ IMPLEMENT_GLOBAL_SHADER(FPrepassGenerateTextureCS, "/Engine/Private/PathTracing/
 
 void PathTracingSpatialTemporalDenoisingPrePass(FRDGBuilder& GraphBuilder, const FViewInfo& View,
 	int IterationNumber,
+	int MaxSPP,
 	FPathTracingSpatialTemporalDenoisingContext& SpatialTemporalDenoisingContext)
 {
 	bool bShouldPrepassOutputVarianceTexture = ShouldPrepassOutputVarianceTexture(View);
 	bool bShouldGenerateVarianceMap = ShouldGenerateVarianceMap() || bShouldPrepassOutputVarianceTexture;
 	if (bShouldGenerateVarianceMap)
 	{
-		bool bUpdateVarianceMap = (IterationNumber > 0);
+		bool bNeedToUpdateVariance = (IterationNumber < MaxSPP);
 		const FScreenPassTextureViewport TargetViewport(View.ViewRect);
 		const FScreenPassTextureViewportParameters TargetViewportParameters = GetScreenPassTextureViewportParameters(TargetViewport);
 
-		if (!SpatialTemporalDenoisingContext.VarianceBuffer)
-		{
-			SpatialTemporalDenoisingContext.VarianceBuffer = GraphBuilder.CreateBuffer(
-				FRDGBufferDesc::CreateStructuredDesc(sizeof(float) * 8, View.ViewRect.Area()), TEXT("PathTracing.VarianceBuffer"));
-		}
 
+		if (bNeedToUpdateVariance)
 		{
+			if (!SpatialTemporalDenoisingContext.VarianceBuffer)
+			{
+				SpatialTemporalDenoisingContext.VarianceBuffer = GraphBuilder.CreateBuffer(
+					FRDGBufferDesc::CreateStructuredDesc(sizeof(float) * 8, View.ViewRect.Area()), TEXT("PathTracing.VarianceBuffer"));
+			}
+
 			typedef FTemporalPrepassCS SHADER;
 			SHADER::FParameters* PassParameters = GraphBuilder.AllocParameters<SHADER::FParameters>();
 			{
@@ -1712,8 +1718,10 @@ void PathTracingSpatialTemporalDenoisingPrePass(FRDGBuilder& GraphBuilder, const
 				PassParameters->Iteration = IterationNumber;
 			}
 
+			bool bUpdateVarianceMapPhase = (IterationNumber > 0);
+
 			SHADER::FPermutationDomain ComputeShaderPermutationVector;
-			ComputeShaderPermutationVector.Set<SHADER::FPrepassPhase>(bUpdateVarianceMap);
+			ComputeShaderPermutationVector.Set<SHADER::FPrepassPhase>(bUpdateVarianceMapPhase);
 			ComputeShaderPermutationVector.Set<SHADER::FVarianceType>(SHADER::GetVarianceType());
 
 			TShaderMapRef<SHADER> ComputeShader(View.ShaderMap, ComputeShaderPermutationVector);
@@ -1727,8 +1735,14 @@ void PathTracingSpatialTemporalDenoisingPrePass(FRDGBuilder& GraphBuilder, const
 				PassParameters,
 				FComputeShaderUtils::GetGroupCount(TargetViewport.Extent, 8));
 		}
+		else
+		{
+			SpatialTemporalDenoisingContext.VarianceBuffer = SpatialTemporalDenoisingContext.LastVarianceBuffer;
+		}
+
+		FRDGBufferRef VarianceBuffer = SpatialTemporalDenoisingContext.VarianceBuffer;
 
-		if (bShouldPrepassOutputVarianceTexture)
+		if (bShouldPrepassOutputVarianceTexture && VarianceBuffer)
 		{
 			const FRDGTextureDesc TextureDescriptor = FRDGTextureDesc::Create2D(
 				TargetViewport.Extent,
@@ -1745,7 +1759,7 @@ void PathTracingSpatialTemporalDenoisingPrePass(FRDGBuilder& GraphBuilder, const
 			SHADER::FParameters* PassParameters = GraphBuilder.AllocParameters<SHADER::FParameters>();
 			{
 				PassParameters->OutputTexture = GraphBuilder.CreateUAV(SpatialTemporalDenoisingContext.VarianceTexture);
-				PassParameters->VarianceMap	  = GraphBuilder.CreateSRV(SpatialTemporalDenoisingContext.VarianceBuffer, EPixelFormat::PF_R32_FLOAT);
+				PassParameters->VarianceMap	  = GraphBuilder.CreateSRV(VarianceBuffer, EPixelFormat::PF_R32_FLOAT);
 				PassParameters->TargetViewport = TargetViewportParameters;
 			}
 
diff --git a/Engine/Source/Runtime/Renderer/Private/PathTracingSpatialTemporalDenoising.h b/Engine/Source/Runtime/Renderer/Private/PathTracingSpatialTemporalDenoising.h
--- a/Engine/Source/Runtime/Renderer/Private/PathTracingSpatialTemporalDenoising.h
+++ b/Engine/Source/Runtime/Renderer/Private/PathTracingSpatialTemporalDenoising.h
@@ -59,6 +59,7 @@ void PathTracingSpatialTemporalDenoisingPrePass(
 	FRDGBuilder& GraphBuilder,
 	const FViewInfo& View,
 	int IterationNumber,
+	int MaxSPP,
 	FPathTracingSpatialTemporalDenoisingContext& SpatialTemporalDenoisingContext
 );
 
diff --git a/Engine/Source/Runtime/Renderer/Public/PathTracingDenoiser.h b/Engine/Source/Runtime/Renderer/Public/PathTracingDenoiser.h
--- a/Engine/Source/Runtime/Renderer/Public/PathTracingDenoiser.h
+++ b/Engine/Source/Runtime/Renderer/Public/PathTracingDenoiser.h
@@ -61,6 +61,7 @@ public:
 		FRDGTextureRef ColorTex;
 		FRDGTextureRef AlbedoTex;
 		FRDGTextureRef NormalTex;
+		FRDGTextureRef VarianceTex;
 		FRDGTextureRef OutputTex;
 
 		FRDGTextureRef FlowTex;
