Enable r.Visibility.SkipAlwaysVisible by default now

diff --git a/Engine/Source/Runtime/Renderer/Private/RendererScene.cpp b/Engine/Source/Runtime/Renderer/Private/RendererScene.cpp
--- a/Engine/Source/Runtime/Renderer/Private/RendererScene.cpp
+++ b/Engine/Source/Runtime/Renderer/Private/RendererScene.cpp
@@ -139,8 +139,7 @@ static TAutoConsoleVariable<int32> CVarBasePassWriteDepthEvenWithFullPrepass(
 	TEXT("0 to allow a readonly base pass, which skips an MSAA depth resolve, and allows masked materials to get EarlyZ (writing to depth while doing clip() disables EarlyZ) (default)\n")
 	TEXT("1 to force depth writes in the base pass.  Useful for debugging when the prepass and base pass don't match what they render."));
 
-// TODO: Significant render thread optimization to heavy Nanite scenes - Off by default, pending extensive testing
-int32 GVisibilitySkipAlwaysVisible = 0;
+int32 GVisibilitySkipAlwaysVisible = 1;
 static FAutoConsoleVariableRef CVarVisibilitySkipAlwaysVisible(
 	TEXT("r.Visibility.SkipAlwaysVisible"),
 	GVisibilitySkipAlwaysVisible,
@@ -6383,10 +6382,13 @@ void FScene::Update(FRDGBuilder& GraphBuilder, const FUpdateParameters& Paramete
 		// Align up to next full dword - this is to avoid having a single dword spanning "tested" and "always visible" primitives,
 		// making the lockless parallel calculations much more efficient. This will push a few (<32) primitives from always visible
 		// into the tested path, but this is not a big deal.
-		PrimitivesAlwaysVisibleOffset = (PrimitivesAlwaysVisibleOffset + uint32(NumBitsPerDWORD) - 1u) & ~(uint32(NumBitsPerDWORD) - 1u);
-		if (int32(PrimitivesAlwaysVisibleOffset) >= Primitives.Num())
+		if (PrimitivesAlwaysVisibleOffset != ~0u)
 		{
-			PrimitivesAlwaysVisibleOffset = ~0u;
+			PrimitivesAlwaysVisibleOffset = (PrimitivesAlwaysVisibleOffset + uint32(NumBitsPerDWORD) - 1u) & ~(uint32(NumBitsPerDWORD) - 1u);
+			if (int32(PrimitivesAlwaysVisibleOffset) >= Primitives.Num())
+			{
+				PrimitivesAlwaysVisibleOffset = ~0u;
+			}
 		}
 	}
 #endif // !WITH_EDITOR
