Don't do LightPrimitiveInteractions for local lights on mobile when r.Mobile.Forward.LocalLightsSinglePermutation is enabled and you don't have local light dynamic shadows enabled

diff --git a/Engine/Source/Runtime/Renderer/Private/MobileBasePass.cpp b/Engine/Source/Runtime/Renderer/Private/MobileBasePass.cpp
--- a/Engine/Source/Runtime/Renderer/Private/MobileBasePass.cpp
+++ b/Engine/Source/Runtime/Renderer/Private/MobileBasePass.cpp
@@ -918,11 +918,10 @@ bool FMobileBasePassMeshProcessor::Process(
 	EMobileLocalLightSetting LocalLightSetting = EMobileLocalLightSetting::LOCAL_LIGHTS_DISABLED;
 	if (Scene && PrimitiveSceneProxy && ShadingModels.IsLit())
 	{
-		if (!bPassUsesDeferredShading && 
-			// we can choose to use a single permutation for opaque meshes regarless of local light state
+		if (!bPassUsesDeferredShading &&
+			// we can choose to use a single permutation regarless of local light state
 			// this is to avoid re-caching MDC on light state changes
-			// Translucency always use un-cached MDC so this option does not apply to it
-			((MobileLocalLightsUseSinglePermutation() && !bIsTranslucent) || PrimitiveSceneProxy->GetPrimitiveSceneInfo()->NumMobileDynamicLocalLights > 0))
+			(MobileLocalLightsUseSinglePermutation() || PrimitiveSceneProxy->GetPrimitiveSceneInfo()->NumMobileDynamicLocalLights > 0))
 		{
 			LocalLightSetting = GetMobileForwardLocalLightSetting(Scene->GetShaderPlatform());
 		}
diff --git a/Engine/Source/Runtime/Renderer/Private/RendererScene.cpp b/Engine/Source/Runtime/Renderer/Private/RendererScene.cpp
--- a/Engine/Source/Runtime/Renderer/Private/RendererScene.cpp
+++ b/Engine/Source/Runtime/Renderer/Private/RendererScene.cpp
@@ -3810,6 +3810,12 @@ void FScene::GetPrimitiveUniformShaderParameters_RenderThread(const FPrimitiveSc
 	SingleCaptureIndex = PrimitiveSceneInfo->CachedReflectionCaptureProxy ? PrimitiveSceneInfo->CachedReflectionCaptureProxy->SortedCaptureIndex : 0;
 }
 
+bool DoesPlatformNeedLocalLightPrimitiveInteraction(EShaderPlatform ShaderPlatform)
+{
+	extern bool MobileLocalLightsUseSinglePermutation();
+	return !IsMobilePlatform(ShaderPlatform) || !MobileLocalLightsUseSinglePermutation() || IsMobileMovableSpotlightShadowsEnabled(ShaderPlatform);
+}
+
 void FScene::UpdateLightTransform_RenderThread(int32 LightId, FLightSceneInfo* LightSceneInfo, const FUpdateLightCommand::FTransformParameters& Parameters)
 {
 	SCOPED_NAMED_EVENT(FScene_UpdateLightTransform_RenderThread, FColor::Yellow);
@@ -3835,7 +3841,7 @@ void FScene::UpdateLightTransform_RenderThread(int32 LightId, FLightSceneInfo* L
 		checkSlow(Lights[LightId].LightSceneInfo == LightSceneInfo);
 		Lights[LightId].Init(LightSceneInfo);
 
-		if (bUpdatePrimitiveInteractions)
+		if (bUpdatePrimitiveInteractions && DoesPlatformNeedLocalLightPrimitiveInteraction(GetShaderPlatform()))
 		{
 			using PrimitiveSceneInfoSet = TSet<FPrimitiveSceneInfo*, DefaultKeyFuncs<FPrimitiveSceneInfo*>, SceneRenderingSetAllocator>;
 			PrimitiveSceneInfoSet PrevPrimitivesInBounds;
@@ -6805,12 +6811,14 @@ void FScene::CreateLightPrimitiveInteractionsForPrimitive(FPrimitiveSceneInfo* P
 		const FBoxSphereBounds& Bounds = Proxy->GetBounds();
 		const FPrimitiveSceneInfoCompact PrimitiveSceneInfoCompact(PrimitiveInfo);
 
-		// Find local lights that affect the primitive in the light octree.
-		LocalShadowCastingLightOctree.FindElementsWithBoundsTest(Bounds.GetBox(), [&PrimitiveSceneInfoCompact](const FLightSceneInfoCompact& LightSceneInfoCompact)
+		if(DoesPlatformNeedLocalLightPrimitiveInteraction(GetShaderPlatform()))
 		{
-			LightSceneInfoCompact.LightSceneInfo->CreateLightPrimitiveInteraction(LightSceneInfoCompact, PrimitiveSceneInfoCompact);
-		});
-
+			// Find local lights that affect the primitive in the light octree.
+			LocalShadowCastingLightOctree.FindElementsWithBoundsTest(Bounds.GetBox(), [&PrimitiveSceneInfoCompact](const FLightSceneInfoCompact& LightSceneInfoCompact)
+			{
+				LightSceneInfoCompact.LightSceneInfo->CreateLightPrimitiveInteraction(LightSceneInfoCompact, PrimitiveSceneInfoCompact);
+			});
+		}
 		// Also loop through non-local (directional) shadow-casting lights
 		for (int32 LightID : DirectionalShadowCastingLightIDs)
 		{
