Split direct lighting loop for GBuffer & Hair inputs.

This avoids context switching between each light and improve performance overall, and give a 5% improvement.

diff --git a/Engine/Source/Runtime/Renderer/Private/LightRendering.cpp b/Engine/Source/Runtime/Renderer/Private/LightRendering.cpp
--- a/Engine/Source/Runtime/Renderer/Private/LightRendering.cpp
+++ b/Engine/Source/Runtime/Renderer/Private/LightRendering.cpp
@@ -1545,6 +1545,8 @@ void FDeferredShadingSceneRenderer::RenderLights(
 			FRDGTextureRef SharedScreenShadowMaskSubPixelTexture = nullptr;
 
 			// Draw shadowed and light function lights
+			auto UnbatchedLightsPass = [&](bool bIsHairPass)
+			{
 			for (int32 LightIndex = UnbatchedLightStart; LightIndex < ManyLightsLightStart; LightIndex++)
 			{
 				const FSortedLightSceneInfo& SortedLightInfo = SortedLights[LightIndex];
@@ -1558,11 +1560,20 @@ void FDeferredShadingSceneRenderer::RenderLights(
 				const bool bDrawPreviewIndicator = ViewFamily.EngineShowFlags.PreviewShadowsIndicator && !LightSceneInfo.IsPrecomputedLightingValid() && LightSceneProxy.HasStaticShadowing();
 				const bool bDrawHairShadow = bDrawShadows && bUseHairLighting;
 				const bool bUseHairDeepShadow = bDrawShadows && bUseHairLighting && LightSceneProxy.CastsHairStrandsDeepShadow();
+				const bool bRunHairLighting = bUseHairLighting && (bIsHairPass || OcclusionType == FLightOcclusionType::Raytraced);
+
 				bool bUsedShadowMaskTexture = false;
 
 				bool bElideScreenShadowMask = false;
 				bool bElideScreenShadowMaskSubPixel = false;
 
+				// Raytraced shadow light for hair are handled/interleaved with regular light, as the raytraced shadow masks 
+				// for gbuffer & hair are computed by the same single pass.
+				if (bIsHairPass && OcclusionType == FLightOcclusionType::Raytraced)
+				{
+					continue;
+				}
+
 				FScopeCycleCounter Context(LightSceneProxy.GetStatId());
 
 				FRDGTextureRef ScreenShadowMaskTexture = nullptr;
@@ -1592,7 +1603,7 @@ void FDeferredShadingSceneRenderer::RenderLights(
 						{
 							SharedScreenShadowMaskTexture = GraphBuilder.CreateTexture(SharedScreenShadowMaskTextureDesc, TEXT("ShadowMaskTexture"));
 						}
-						if (!SharedScreenShadowMaskSubPixelTexture && bUseHairLighting && !bElideScreenShadowMaskSubPixel)
+						if (!SharedScreenShadowMaskSubPixelTexture && bRunHairLighting && !bElideScreenShadowMaskSubPixel)
 						{
 							SharedScreenShadowMaskSubPixelTexture = GraphBuilder.CreateTexture(SharedScreenShadowMaskTextureDesc, TEXT("ShadowMaskSubPixelTexture"));
 						}
@@ -1835,7 +1846,7 @@ void FDeferredShadingSceneRenderer::RenderLights(
 						PreprocessedShadowMaskTextures[ShadowMaskIndex] = nullptr;
 
 						// Sub-pixel shadow for hair strands geometries
-						if (bUseHairLighting && ShadowMaskIndex < uint32(PreprocessedShadowMaskSubPixelTextures.Num()))
+						if (bRunHairLighting && ShadowMaskIndex < uint32(PreprocessedShadowMaskSubPixelTextures.Num()))
 						{
 							ScreenShadowMaskSubPixelTexture = PreprocessedShadowMaskSubPixelTextures[ShadowMaskIndex];
 							PreprocessedShadowMaskSubPixelTextures[ShadowMaskIndex] = nullptr;
@@ -1878,7 +1889,7 @@ void FDeferredShadingSceneRenderer::RenderLights(
 
 						FRDGTextureRef SubPixelRayTracingShadowMaskTexture = nullptr;
 						FRDGTextureUAV* SubPixelRayTracingShadowMaskUAV = nullptr;
-						if (bUseHairLighting)
+						if (bRunHairLighting)
 						{
 							FRDGTextureDesc Desc = FRDGTextureDesc::Create2D(
 								SceneTextures.Config.Extent,
@@ -2076,7 +2087,7 @@ void FDeferredShadingSceneRenderer::RenderLights(
 				}
 
 				// Render the light to the scene color buffer, conditionally using the attenuation buffer or a 1x1 white texture as input 
-				if (bDirectLighting)
+				if (bDirectLighting && !bIsHairPass)
 				{
 					for (int32 ViewIndex = 0, ViewCount = Views.Num(); ViewIndex < ViewCount; ++ViewIndex)
 					{
@@ -2099,7 +2110,7 @@ void FDeferredShadingSceneRenderer::RenderLights(
 					}
 				}
 
-				if (bUseHairLighting)
+				if (bRunHairLighting)
 				{
 					for (int32 ViewIndex = 0, ViewCount = Views.Num(); ViewIndex < ViewCount; ++ViewIndex)
 					{
@@ -2145,6 +2156,16 @@ void FDeferredShadingSceneRenderer::RenderLights(
 						}
 					}
 				}
+			}};
+
+			// Two seperate light loop:
+			// * For GBuffer inputs (and hair input for light having RT shadows)
+			// * For Hair inputs if any hair data are present
+			UnbatchedLightsPass(false/*bIsHairPass*/);
+			if (bUseHairLighting)
+			{
+				RDG_EVENT_SCOPE(GraphBuilder, "UnbatchedLights(Hair)");
+				UnbatchedLightsPass(true/*bIsHairPass*/);
 			}
 		}
 	}
