Heterogeneous Volumes: Modifying integration scheme with Fog.

Other small fixes:
1) Fixing a variable shadowing bug preventing calculating shadow bias
2) Changing HV Component defaults to align with new defaults to use AVSM if created
3) Updating ray marcher to use Blue Noise
4) Preventing fog double-contribution when disabling lighting cache

diff --git a/Engine/Shaders/Private/HeterogeneousVolumes/HeterogeneousVolumesLiveShadingPipeline.usf b/Engine/Shaders/Private/HeterogeneousVolumes/HeterogeneousVolumesLiveShadingPipeline.usf
--- a/Engine/Shaders/Private/HeterogeneousVolumes/HeterogeneousVolumesLiveShadingPipeline.usf
+++ b/Engine/Shaders/Private/HeterogeneousVolumes/HeterogeneousVolumesLiveShadingPipeline.usf
@@ -15,6 +15,7 @@
 #include "HeterogeneousVolumesLiveShadingUtils.ush"
 #include "HeterogeneousVolumesTracingUtils.ush"
 #include "HeterogeneousVolumesLightingUtils.ush"
+#include "../BlueNoise.ush"
 
 
 #if VIRTUAL_SHADOW_MAP
@@ -327,7 +328,7 @@ void RenderSingleScatteringWithLiveShadingCS(
 	float SceneDepth = min(ConvertFromDeviceZ(DeviceZ), GetMaxTraceDistance());
 	DeviceZ = ConvertToDeviceZ(SceneDepth);
 
-	float Jitter = bJitter ? InterleavedGradientNoise(PixelCoord, View.StateFrameIndexMod8) : 0.0;
+	float Jitter = bJitter ? BlueNoiseScalar(PixelCoord, View.StateFrameIndexMod8) : 0.0;
 
 	// Intersect ray with bounding volume
 	// TODO: LWC integration..
diff --git a/Engine/Shaders/Private/HeterogeneousVolumes/HeterogeneousVolumesLiveShadingShadows.usf b/Engine/Shaders/Private/HeterogeneousVolumes/HeterogeneousVolumesLiveShadingShadows.usf
--- a/Engine/Shaders/Private/HeterogeneousVolumes/HeterogeneousVolumesLiveShadingShadows.usf
+++ b/Engine/Shaders/Private/HeterogeneousVolumes/HeterogeneousVolumesLiveShadingShadows.usf
@@ -14,6 +14,7 @@
 #include "HeterogeneousVolumesLiveShadingUtils.ush"
 #include "HeterogeneousVolumesTracingUtils.ush"
 #include "HeterogeneousVolumesAdaptiveVolumetricShadowMapUtils.ush"
+#include "../BlueNoise.ush"
 
 #ifndef THREADGROUP_SIZE_1D
 #define THREADGROUP_SIZE_1D 1
@@ -333,7 +334,7 @@ void RenderVolumetricShadowMapForLightWithLiveShadingCS(
 	float HitSpan = HitT.y - HitT.x;
 	if (HitSpan > 0.0)
 	{
-		float Jitter = bJitter ? InterleavedGradientNoise(PixelCoord, View.StateFrameIndexMod8) : 0.0;
+		float Jitter = bJitter ? BlueNoiseScalar(PixelCoord, View.StateFrameIndexMod8) : 0.0;
 		int bApplyEmissionAndTransmittance = 0;
 		int bApplyDirectLighting = 0;
 		int bApplyShadowTransmittance = 0;
diff --git a/Engine/Shaders/Private/HeterogeneousVolumes/HeterogeneousVolumesRayMarchingUtils.ush b/Engine/Shaders/Private/HeterogeneousVolumes/HeterogeneousVolumesRayMarchingUtils.ush
--- a/Engine/Shaders/Private/HeterogeneousVolumes/HeterogeneousVolumesRayMarchingUtils.ush
+++ b/Engine/Shaders/Private/HeterogeneousVolumes/HeterogeneousVolumesRayMarchingUtils.ush
@@ -28,8 +28,12 @@
 #define DIM_USE_ADAPTIVE_VOLUMETRIC_SHADOW_MAP 0
 #endif // DIM_USE_ADAPTIVE_VOLUMETRIC_SHADOW_MAP
 
+#define APPLY_FOG_INSCATTERING_OFF 0
+#define APPLY_FOG_INSCATTERING_REFERENCE 1
+#define APPLY_FOG_INSCATTERING_STOCHASTIC 2
+
 #ifndef APPLY_FOG_INSCATTERING
-#define APPLY_FOG_INSCATTERING 0
+#define APPLY_FOG_INSCATTERING APPLY_FOG_INSCATTERING_OFF
 #endif // APPLY_FOG_INSCATTERING
 
 #if DIM_USE_ADAPTIVE_VOLUMETRIC_SHADOW_MAP
@@ -50,7 +54,7 @@ float CalcShadowBias()
 		float3 VoxelSize = VolumeBounds * rcp(LightingCacheResolutionAsFloat3);
 		float3 WorldVoxelSize = mul(float4(VoxelSize, 0), GetLocalToWorld()).xyz;
 		float VoxelDiagonal = length(WorldVoxelSize);
-		float WorldShadowBias = VoxelDiagonal * GetLightingCacheVoxelBias();
+		WorldShadowBias = VoxelDiagonal * GetLightingCacheVoxelBias();
 	}
 	return WorldShadowBias;
 }
@@ -313,13 +317,18 @@ void RayMarchSingleScattering(
 )
 {
 	LocalScatterTMin = RayMarchingContext.LocalRayTMax;
-#if APPLY_FOG_INSCATTERING
+#if APPLY_FOG_INSCATTERING != APPLY_FOG_INSCATTERING_OFF
+	float3 FogRadiance = 0;
 	float3 PrevTransmittance = Transmittance;
+	float4 PrevCombinedFog = 0;
 #endif // #if APPLY_FOG_INSCATTERING
 
 	float WorldShadowBias = CalcShadowBias();
+	//float WorldShadowBias = View.GeneralPurposeTweak;
 	float StepSize = RayMarchingContext.StepSize;
 	float JitterSize = StepSize * RayMarchingContext.Jitter;
+	// Reuse jitter as random sample
+	float RandomSample = RayMarchingContext.Jitter;
 
 	float LocalHitT = RayMarchingContext.LocalRayTMin + JitterSize;
 	for (uint StepIndex = 0; StepIndex < StepCount; ++StepIndex)
@@ -359,8 +368,11 @@ void RayMarchSingleScattering(
 		}
 
 		float3 Extinction = SampleExtinction(SampleContext);
-		if (RayMarchingContext.bApplyDirectLighting && (any(Extinction > 0.0)))
+		if (RayMarchingContext.bApplyDirectLighting && (any(Extinction > 0)))
 		{
+			float3 RadianceSample = 0;
+
+			// Evaluate in-scattering
 			float3 Albedo = SampleAlbedo(SampleContext);
 			if (any(Albedo > 0.0))
 			{
@@ -368,10 +380,48 @@ void RayMarchSingleScattering(
 
 				float3 ScatteringCoefficient = Albedo * Extinction;
 				float IsotropicPhase = 1.0 / (4.0 * PI);
-				Radiance += Inscattering * ScatteringCoefficient * IsotropicPhase * Transmittance * StepSize;
+				RadianceSample = Inscattering * ScatteringCoefficient * IsotropicPhase * Transmittance * StepSize;
+			}
+
+			// Stochastically decide to record scatter position
+			if (any(Transmittance < RandomSample))
+			{
+				LocalScatterTMin = min(LocalHitT, LocalScatterTMin);
+			}
+
+#if APPLY_FOG_INSCATTERING == APPLY_FOG_INSCATTERING_REFERENCE
+			{
+				float WorldRayT = LocalHitT * RayMarchingContext.LocalToWorldScale;
+				float3 CameraRelative_WorldPosition = RayMarchingContext.WorldRayDirection * WorldRayT;
+				float4 HeightFog = float4(0, 0, 0, 1);
+				if (ShouldApplyHeightFog())
+				{
+					HeightFog = CalculateHeightFog(CameraRelative_WorldPosition);
+				}
+
+				float4 CombinedFog = HeightFog;
+				if (ShouldApplyVolumetricFog() && FogStruct.ApplyVolumetricFog > 0)
+				{
+					float3 WorldPosition = CameraRelative_WorldPosition + RayMarchingContext.WorldRayOrigin;
+					const uint EyeIndex = 0;
+					float3 VolumeUV = ComputeVolumeUV_DEPRECATED(WorldPosition, DFHackToFloat(PrimaryView.WorldToClip));
+					CombinedFog = CombineVolumetricFog(HeightFog, VolumeUV, EyeIndex, WorldRayT);
+				}
+
+				// Attenuate RadianceSample
+				RadianceSample *= CombinedFog.a;
+
+				// Add FogRadiance
+				float3 DeltaTransmittance = clamp(Transmittance / PrevTransmittance, 0.0, 1.0);
+				float3 DeltaFogRadiance = CombinedFog.rgb - PrevCombinedFog.rgb;
+				FogRadiance += DeltaFogRadiance * DeltaTransmittance;
+
+				PrevCombinedFog = CombinedFog;
 			}
+#endif // APPLY_FOG_INSCATTERING_REFERENCE
 
-			LocalScatterTMin = min(LocalHitT, LocalScatterTMin);
+			// Accumulate radiance
+			Radiance += RadianceSample;
 		}
 
 		// Accumulate transmittance for the next evaluation
@@ -387,13 +437,11 @@ void RayMarchSingleScattering(
 		LocalHitT += StepSize;
 	}
 
-#if APPLY_FOG_INSCATTERING
-	bool bApplyFog = LocalScatterTMin < RayMarchingContext.LocalRayTMax;
-	if (bApplyFog)
+#if APPLY_FOG_INSCATTERING == APPLY_FOG_INSCATTERING_STOCHASTIC
+	if (LocalScatterTMin < RayMarchingContext.LocalRayTMax)
 	{
-		// TODO: Change in-scattering distance to be based on the first scattering event, instead of the
-		// intersection distance of the bounding box near-plane.
 		float WorldRayT = LocalScatterTMin * RayMarchingContext.LocalToWorldScale;
+
 		float3 CameraRelative_WorldPosition = RayMarchingContext.WorldRayDirection * WorldRayT;
 		float4 HeightFog = float4(0, 0, 0, 1);
 		if (ShouldApplyHeightFog())
@@ -410,16 +458,22 @@ void RayMarchSingleScattering(
 			CombinedFog = CombineVolumetricFog(HeightFog, VolumeUV, EyeIndex, WorldRayT);
 		}
 
-		// Simplifying terms to avoid vector division:
-		// float3 DeltaTransmittance = clamp(Transmittance / PrevTransmittance, 0.0, 1.0);
-		// float3 DeltaOpacity = 1.0 - DeltaTransmittance;
-		// float3 FogAlpha = PrevTransmittance * DeltaOpacity;
-		float3 FogAlpha = clamp(PrevTransmittance - Transmittance, 0.0, 1.0);
-
-		// Fogging in-scattering contribution is not currently attenuated by Heterogeneous Volumes
-		// for elements behind the subject. To prevent doubled contribution of in-scattering for 
-		// regions of low density (eg: silhouette edges), blending is interpreted as alpha-masking instead.
-		Radiance = Radiance * CombinedFog.a + CombinedFog.rgb * FogAlpha;
+		// Attenuate Radiance by fog
+		Radiance *= CombinedFog.a;
+		FogRadiance = CombinedFog.rgb;
 	}
 #endif // APPLY_FOG_INSCATTERING
+
+#if APPLY_FOG_INSCATTERING != APPLY_FOG_INSCATTERING_OFF
+	// Simplifying terms to avoid vector division:
+	// float3 DeltaTransmittance = clamp(Transmittance / PrevTransmittance, 0.0, 1.0);
+	// float3 DeltaOpacity = 1.0 - DeltaTransmittance;
+	// float3 FogAlpha = PrevTransmittance * DeltaOpacity;
+	float3 FogAlpha = clamp(PrevTransmittance - Transmittance, 0.0, 1.0);
+
+	// Fogging in-scattering contribution is not currently attenuated by Heterogeneous Volumes
+	// for elements behind the subject. To prevent doubled contribution of in-scattering for 
+	// regions of low density (eg: silhouette edges), blending is interpreted as alpha-masking instead.
+	Radiance += FogRadiance.rgb * FogAlpha;
+#endif // APPLY_FOG_INSCATTERING
 }
diff --git a/Engine/Source/Runtime/Engine/Private/Components/HeterogeneousVolumeComponent.cpp b/Engine/Source/Runtime/Engine/Private/Components/HeterogeneousVolumeComponent.cpp
--- a/Engine/Source/Runtime/Engine/Private/Components/HeterogeneousVolumeComponent.cpp
+++ b/Engine/Source/Runtime/Engine/Private/Components/HeterogeneousVolumeComponent.cpp
@@ -209,6 +209,8 @@ void FHeterogeneousVolumeSceneProxy::GetDynamicMeshElements(
 
 		Mesh.bCanApplyViewModeOverrides = true;
 		Mesh.bUseWireframeSelectionColoring = IsSelected();
+		Mesh.bUseSelectionOutline = false;
+		Mesh.bSelectable = false;
 
 		Collector.AddMesh(0, Mesh);
 	}
@@ -251,8 +253,8 @@ UHeterogeneousVolumeComponent::UHeterogeneousVolumeComponent(const FObjectInitia
 	StartFrame = 0;
 	EndFrame = 0;
 	StepFactor = 1.0f;
-	ShadowStepFactor = 8.0f;
-	ShadowBiasFactor = 0.0f;
+	ShadowStepFactor = 2.0f;
+	ShadowBiasFactor = 0.5f;
 	LightingDownsampleFactor = 2.0f;
 	MipLevel = 0;
 	bIssueBlockingRequests = false;
diff --git a/Engine/Source/Runtime/Renderer/Private/HeterogeneousVolumes/HeterogeneousVolumes.cpp b/Engine/Source/Runtime/Renderer/Private/HeterogeneousVolumes/HeterogeneousVolumes.cpp
--- a/Engine/Source/Runtime/Renderer/Private/HeterogeneousVolumes/HeterogeneousVolumes.cpp
+++ b/Engine/Source/Runtime/Renderer/Private/HeterogeneousVolumes/HeterogeneousVolumes.cpp
@@ -229,6 +229,16 @@ static TAutoConsoleVariable<int32> CVarHeterogeneousVolumesApplyVolumetricFog(
 	ECVF_RenderThreadSafe
 );
 
+static TAutoConsoleVariable<int32> CVarHeterogeneousVolumesApplyFogInscatteringMode(
+	TEXT("r.HeterogeneousVolumes.ApplyFogInscattering"),
+	1,
+	TEXT("Determines the method for applying fog in-scattering (default = 1)\n")
+	TEXT("0: Off\n")
+	TEXT("1: Reference (evaluated per ray-march step)\n")
+	TEXT("2: Stochastic (evaluated once per ray)\n"),
+	ECVF_RenderThreadSafe
+);
+
 static TAutoConsoleVariable<int32> CVarHeterogeneousVolumesVelocity(
 	TEXT("r.HeterogeneousVolumes.Velocity"),
 	0,
@@ -468,6 +478,11 @@ namespace HeterogeneousVolumes
 		return CVarHeterogeneousVolumesApplyVolumetricFog.GetValueOnRenderThread() != 0;
 	}
 
+	EFogMode GetApplyFogInscattering()
+	{
+		return static_cast<EFogMode>(FMath::Clamp(CVarHeterogeneousVolumesApplyFogInscatteringMode.GetValueOnRenderThread(), 0, 2));
+	}
+
 	bool ShouldWriteVelocity()
 	{
 		return CVarHeterogeneousVolumesVelocity.GetValueOnRenderThread() != 0;
diff --git a/Engine/Source/Runtime/Renderer/Private/HeterogeneousVolumes/HeterogeneousVolumes.h b/Engine/Source/Runtime/Renderer/Private/HeterogeneousVolumes/HeterogeneousVolumes.h
--- a/Engine/Source/Runtime/Renderer/Private/HeterogeneousVolumes/HeterogeneousVolumes.h
+++ b/Engine/Source/Runtime/Renderer/Private/HeterogeneousVolumes/HeterogeneousVolumes.h
@@ -105,6 +105,14 @@ namespace HeterogeneousVolumes
 	bool UseAdaptiveVolumetricShadowMapForSelfShadowing();
 	bool ShouldApplyHeightFog();
 	bool ShouldApplyVolumetricFog();
+
+	enum class EFogMode
+	{
+		Off,
+		Reference,
+		Stochastic
+	};
+	EFogMode GetApplyFogInscattering();
 	bool ShouldWriteVelocity();
 
 	bool EnableIndirectionGrid();
diff --git a/Engine/Source/Runtime/Renderer/Private/HeterogeneousVolumes/HeterogeneousVolumesLiveShadingPipeline.cpp b/Engine/Source/Runtime/Renderer/Private/HeterogeneousVolumes/HeterogeneousVolumesLiveShadingPipeline.cpp
--- a/Engine/Source/Runtime/Renderer/Private/HeterogeneousVolumes/HeterogeneousVolumesLiveShadingPipeline.cpp
+++ b/Engine/Source/Runtime/Renderer/Private/HeterogeneousVolumes/HeterogeneousVolumesLiveShadingPipeline.cpp
@@ -15,6 +15,7 @@
 #include "SceneManagement.h"
 #include "VolumeLighting.h"
 #include "VolumetricFog.h"
+#include "BlueNoise.h"
 
 static TAutoConsoleVariable<int32> CVarHeterogeneousLightingCacheBoundsCulling(
 	TEXT("r.HeterogeneousVolumes.LightingCache.BoundsCulling"),
@@ -234,13 +235,15 @@ class FRenderSingleScatteringWithLiveShadingCS : public FMeshMaterialShader
 	class FUseLumenGI : SHADER_PERMUTATION_BOOL("DIM_USE_LUMEN_GI");
 	class FWriteVelocity : SHADER_PERMUTATION_BOOL("DIM_WRITE_VELOCITY");
 	class FUseAdaptiveVolumetricShadowMap : SHADER_PERMUTATION_BOOL("DIM_USE_ADAPTIVE_VOLUMETRIC_SHADOW_MAP");
-	using FPermutationDomain = TShaderPermutationDomain<FUseTransmittanceVolume, FUseInscatteringVolume, FUseLumenGI, FWriteVelocity, FUseAdaptiveVolumetricShadowMap>;
+	class FApplyFogInscattering : SHADER_PERMUTATION_INT("APPLY_FOG_INSCATTERING", 3);
+	using FPermutationDomain = TShaderPermutationDomain<FUseTransmittanceVolume, FUseInscatteringVolume, FUseLumenGI, FWriteVelocity, FUseAdaptiveVolumetricShadowMap, FApplyFogInscattering>;
 
 	BEGIN_SHADER_PARAMETER_STRUCT(FParameters, )
 		// Scene data
 		SHADER_PARAMETER_STRUCT_REF(FViewUniformShaderParameters, View)
 		SHADER_PARAMETER_STRUCT_INCLUDE(FSceneTextureParameters, SceneTextures)
 		SHADER_PARAMETER_RDG_UNIFORM_BUFFER(FSceneUniformParameters, Scene)
+		SHADER_PARAMETER_STRUCT_REF(FBlueNoise, BlueNoise)
 
 		// Light data
 		SHADER_PARAMETER(int, bApplyEmissionAndTransmittance)
@@ -316,12 +319,24 @@ class FRenderSingleScatteringWithLiveShadingCS : public FMeshMaterialShader
 		const FMaterialShaderPermutationParameters& Parameters
 	)
 	{
+		FPermutationDomain PermutationVector(Parameters.PermutationId);
+		if (PermutationVector.Get<FApplyFogInscattering>() == 0)
+		{
+			return false;
+		}
+
 		return DoesPlatformSupportHeterogeneousVolumes(Parameters.Platform)
 			&& DoesMaterialShaderSupportHeterogeneousVolumes(Parameters.MaterialParameters);
 	}
 
 	static FPermutationDomain RemapPermutation(FPermutationDomain PermutationVector)
 	{
+		// Remap Off to Stochastic and turn off individual fog modes
+		if (PermutationVector.Get<FApplyFogInscattering>() == 0)
+		{
+			PermutationVector.Set<FApplyFogInscattering>(2);
+		}
+
 		return PermutationVector;
 	}
 
@@ -333,7 +348,6 @@ class FRenderSingleScatteringWithLiveShadingCS : public FMeshMaterialShader
 		FMaterialShader::ModifyCompilationEnvironment(Parameters, OutEnvironment);
 		OutEnvironment.SetDefine(TEXT("THREADGROUP_SIZE_1D"), GetThreadGroupSize1D());
 		OutEnvironment.SetDefine(TEXT("THREADGROUP_SIZE_2D"), GetThreadGroupSize2D());
-		OutEnvironment.SetDefine(TEXT("APPLY_FOG_INSCATTERING"), 1);
 		OutEnvironment.SetDefine(TEXT("FOG_MATERIALBLENDING_OVERRIDE"), 1);
 
 		bool bSupportVirtualShadowMap = IsFeatureLevelSupported(Parameters.Platform, ERHIFeatureLevel::SM5);
@@ -646,6 +660,8 @@ static void RenderSingleScatteringWithLiveShading(
 		PassParameters->View = View.ViewUniformBuffer;
 		PassParameters->SceneTextures = GetSceneTextureParameters(GraphBuilder, SceneTextures);
 		PassParameters->Scene = View.GetSceneUniforms().GetBuffer(GraphBuilder);
+		FBlueNoise BlueNoise = GetBlueNoiseGlobalParameters();
+		PassParameters->BlueNoise = CreateUniformBufferImmediate(BlueNoise, EUniformBufferUsage::UniformBuffer_SingleDraw);
 
 		// Light data
 		PassParameters->bApplyEmissionAndTransmittance = bApplyEmissionAndTransmittance;
@@ -714,8 +730,14 @@ static void RenderSingleScatteringWithLiveShading(
 
 		TRDGUniformBufferRef<FFogUniformParameters> FogBuffer = CreateFogUniformBuffer(GraphBuilder, View);
 		PassParameters->FogStruct = FogBuffer;
-		PassParameters->bApplyHeightFog = HeterogeneousVolumes::ShouldApplyHeightFog();
-		PassParameters->bApplyVolumetricFog = HeterogeneousVolumes::ShouldApplyVolumetricFog();
+		PassParameters->bApplyHeightFog = 0;
+		PassParameters->bApplyVolumetricFog = 0;
+		if (bApplyEmissionAndTransmittance &&
+			(HeterogeneousVolumes::GetApplyFogInscattering() != HeterogeneousVolumes::EFogMode::Off))
+		{
+			PassParameters->bApplyHeightFog = HeterogeneousVolumes::ShouldApplyHeightFog();
+			PassParameters->bApplyVolumetricFog = HeterogeneousVolumes::ShouldApplyVolumetricFog();
+		}
 
 		// Indirect lighting data
 		auto* LumenUniforms = GraphBuilder.AllocParameters<FLumenTranslucencyLightingUniforms>();
@@ -767,6 +789,8 @@ static void RenderSingleScatteringWithLiveShading(
 	PermutationVector.Set<FRenderSingleScatteringWithLiveShadingCS::FUseLumenGI>(HeterogeneousVolumes::UseIndirectLighting() && View.GetLumenTranslucencyGIVolume().Texture0 != nullptr);
 	PermutationVector.Set<FRenderSingleScatteringWithLiveShadingCS::FWriteVelocity>(bWriteVelocity);
 	PermutationVector.Set<FRenderSingleScatteringWithLiveShadingCS::FUseAdaptiveVolumetricShadowMap>(HeterogeneousVolumes::UseAdaptiveVolumetricShadowMapForSelfShadowing());
+	PermutationVector.Set<FRenderSingleScatteringWithLiveShadingCS::FApplyFogInscattering>(static_cast<int32>(HeterogeneousVolumes::GetApplyFogInscattering()));
+	PermutationVector = FRenderSingleScatteringWithLiveShadingCS::RemapPermutation(PermutationVector);
 	TShaderRef<FRenderSingleScatteringWithLiveShadingCS> ComputeShader = Material.GetShader<FRenderSingleScatteringWithLiveShadingCS>(&FLocalVertexFactory::StaticType, PermutationVector, false);
 	if (!ComputeShader.IsNull())
 	{
@@ -807,7 +831,8 @@ static void RenderWithTransmittanceVolumePipeline(
 	int32 NumPasses = FMath::Max(LightSceneInfoCompact.Num(), 1);
 	for (int32 PassIndex = 0; PassIndex < NumPasses; ++PassIndex)
 	{
-		bool bApplyEmissionAndTransmittance = (PassIndex == (NumPasses - 1));
+		bool bIsLastPass = (PassIndex == (NumPasses - 1));
+		bool bApplyEmissionAndTransmittance = bIsLastPass;
 		bool bApplyDirectLighting = !LightSceneInfoCompact.IsEmpty();
 		bool bApplyShadowTransmittance = false;
 
@@ -1075,6 +1100,7 @@ class FRenderVolumetricShadowMapForLightWithLiveShadingCS : public FMeshMaterial
 		// Scene data
 		SHADER_PARAMETER_STRUCT_REF(FViewUniformShaderParameters, View)
 		SHADER_PARAMETER_STRUCT_INCLUDE(FSceneTextureParameters, SceneTextures)
+		SHADER_PARAMETER_STRUCT_REF(FBlueNoise, BlueNoise)
 
 		// Volumetric Shadow Map data
 		SHADER_PARAMETER(FVector3f, TranslatedWorldOrigin)
@@ -1304,6 +1330,8 @@ bool RenderVolumetricShadowMapForLightForHeterogeneousVolumeWithLiveShading(
 		// Scene data
 		PassParameters->View = View.ViewUniformBuffer;
 		PassParameters->SceneTextures = GetSceneTextureParameters(GraphBuilder, SceneTextures);
+		FBlueNoise BlueNoise = GetBlueNoiseGlobalParameters();
+		PassParameters->BlueNoise = CreateUniformBufferImmediate(BlueNoise, EUniformBufferUsage::UniformBuffer_SingleDraw);
 
 		// Shadow map data
 		PassParameters->TranslatedWorldOrigin = TranslatedWorldOrigin;
