Put label rendering code in WITH_DEBUG_VIEW_MODES conditional block to fix compilation issue on non-editor platforms.

diff --git a/Engine/Source/Runtime/Renderer/Private/PostProcess/PostProcessVisualizeComplexity.cpp b/Engine/Source/Runtime/Renderer/Private/PostProcess/PostProcessVisualizeComplexity.cpp
--- a/Engine/Source/Runtime/Renderer/Private/PostProcess/PostProcessVisualizeComplexity.cpp
+++ b/Engine/Source/Runtime/Renderer/Private/PostProcess/PostProcessVisualizeComplexity.cpp
@@ -188,9 +188,11 @@ FScreenPassTexture AddVisualizeComplexityPass(FRDGBuilder& GraphBuilder, const F
 
 			if (View.Family->GetDebugViewShaderMode() == DVSM_LWCComplexity)
 			{
+#if WITH_DEBUG_VIEW_MODES
 				extern float GMaxLWCComplexity;
 				FString Line = FString::Printf(TEXT("r.ShaderComplexity.MaxLWCComplexity=%d"), (int32)GMaxLWCComplexity);
 				DrawString(Output.ViewRect.Max.X - 430, Output.ViewRect.Max.Y - 88, *Line);
+#endif
 			}
 			else
 			{
