HWRT: Use GetMaterialNoFallback when preparing RayTracingDynamicGeometry dispatch to not submit shader compile jobs and slow down render thread.

diff --git a/Engine/Source/Runtime/Renderer/Private/RayTracing/RayTracingDynamicGeometry.cpp b/Engine/Source/Runtime/Renderer/Private/RayTracing/RayTracingDynamicGeometry.cpp
--- a/Engine/Source/Runtime/Renderer/Private/RayTracing/RayTracingDynamicGeometry.cpp
+++ b/Engine/Source/Runtime/Renderer/Private/RayTracing/RayTracingDynamicGeometry.cpp
@@ -234,75 +234,85 @@ void FRayTracingDynamicGeometryCollection::AddDynamicMeshBatchForGeometryUpdate(
 			continue;
 		}
 
-		const FMaterialRenderProxy* FallbackMaterialRenderProxyPtr = nullptr;
-		const FMaterial& Material = MeshBatch.MaterialRenderProxy->GetMaterialWithFallback(Scene->GetFeatureLevel(), FallbackMaterialRenderProxyPtr);
-		auto* MaterialInterface = Material.GetMaterialInterface();
-		const FMaterialRenderProxy& MaterialRenderProxy = FallbackMaterialRenderProxyPtr ? *FallbackMaterialRenderProxyPtr : *MeshBatch.MaterialRenderProxy;
-
-		FMeshComputeDispatchCommand DispatchCmd;
-		
-		FMaterialShaderTypes ShaderTypes;
-		ShaderTypes.AddShaderType<FRayTracingDynamicGeometryConverterCS>();
-
-		FMaterialShaders MaterialShaders;
-		if (!Material.TryGetShaders(ShaderTypes, MeshBatch.VertexFactory->GetType(), MaterialShaders))
+		const FMaterialRenderProxy* MaterialRenderProxyPtr = MeshBatch.MaterialRenderProxy;
+		while (MaterialRenderProxyPtr)
 		{
-			continue;
-		}
+			const FMaterial* MaterialPtr = MaterialRenderProxyPtr->GetMaterialNoFallback(Scene->GetFeatureLevel());
+			if (MaterialPtr && MaterialPtr->GetRenderingThreadShaderMap())
+			{
+				const FMaterial& Material = *MaterialPtr;
+				const FMaterialRenderProxy& MaterialRenderProxy = *MaterialRenderProxyPtr;
 
-		TShaderRef<FRayTracingDynamicGeometryConverterCS> Shader;
-		MaterialShaders.TryGetShader(SF_Compute, Shader);
+				auto* MaterialInterface = Material.GetMaterialInterface();
 
-		FMeshProcessorShaders MeshProcessorShaders;
-		MeshProcessorShaders.ComputeShader = Shader;
+				FMeshComputeDispatchCommand DispatchCmd;
 
-		DispatchCmd.MaterialShader = Shader;
-		FMeshDrawShaderBindings& ShaderBindings = DispatchCmd.ShaderBindings;
-		ShaderBindings.Initialize(MeshProcessorShaders);
+				FMaterialShaderTypes ShaderTypes;
+				ShaderTypes.AddShaderType<FRayTracingDynamicGeometryConverterCS>();
 
-		FMeshMaterialShaderElementData ShaderElementData;
-		ShaderElementData.InitializeMeshMaterialData(View, PrimitiveSceneProxy, MeshBatch, -1, false);
+				FMaterialShaders MaterialShaders;
+				if (Material.TryGetShaders(ShaderTypes, MeshBatch.VertexFactory->GetType(), MaterialShaders))
+				{
+					TShaderRef<FRayTracingDynamicGeometryConverterCS> Shader;
+					MaterialShaders.TryGetShader(SF_Compute, Shader);
 
-		FMeshDrawSingleShaderBindings SingleShaderBindings = ShaderBindings.GetSingleShaderBindings(SF_Compute);
-		Shader->GetShaderBindings(Scene, Scene->GetFeatureLevel(), PrimitiveSceneProxy, MaterialRenderProxy, Material, ShaderElementData, SingleShaderBindings);
+					FMeshProcessorShaders MeshProcessorShaders;
+					MeshProcessorShaders.ComputeShader = Shader;
 
-		FVertexInputStreamArray DummyArray;
-		FMeshMaterialShader::GetElementShaderBindings(Shader, Scene, View, MeshBatch.VertexFactory, EVertexInputStreamType::Default, Scene->GetFeatureLevel(), PrimitiveSceneProxy, MeshBatch, MeshBatch.Elements[0], ShaderElementData, SingleShaderBindings, DummyArray);
+					DispatchCmd.MaterialShader = Shader;
+					FMeshDrawShaderBindings& ShaderBindings = DispatchCmd.ShaderBindings;
+					ShaderBindings.Initialize(MeshProcessorShaders);
 
-		DispatchCmd.TargetBuffer = RWBuffer;
-		DispatchCmd.NumMaxVertices = UpdateParams.NumVertices;
+					FMeshMaterialShaderElementData ShaderElementData;
+					ShaderElementData.InitializeMeshMaterialData(View, PrimitiveSceneProxy, MeshBatch, -1, false);
 
-		// Setup the loose parameters directly on the binding
-		uint32 OutputVertexBaseIndex = VertexBufferOffset / sizeof(float);
-		uint32 MinVertexIndex = MeshBatch.Elements[0].MinVertexIndex;
-		uint32 NumCPUVertices = UpdateParams.NumVertices;
-		if (MeshBatch.Elements[0].MinVertexIndex < MeshBatch.Elements[0].MaxVertexIndex)
-		{
-			NumCPUVertices = 1 + MeshBatch.Elements[0].MaxVertexIndex - MeshBatch.Elements[0].MinVertexIndex;
-		}
+					FMeshDrawSingleShaderBindings SingleShaderBindings = ShaderBindings.GetSingleShaderBindings(SF_Compute);
+					Shader->GetShaderBindings(Scene, Scene->GetFeatureLevel(), PrimitiveSceneProxy, MaterialRenderProxy, Material, ShaderElementData, SingleShaderBindings);
 
-		const uint32 VertexBufferNumElements = UpdateParams.VertexBufferSize / sizeof(FVector3f) - MinVertexIndex;
-		if (!ensureMsgf(NumCPUVertices <= VertexBufferNumElements, 
-			TEXT("Vertex buffer contains %d vertices, but RayTracingDynamicGeometryConverterCS dispatch command expects at least %d."),
-			VertexBufferNumElements, NumCPUVertices))
-		{
-			NumCPUVertices = VertexBufferNumElements;
-		}
+					FVertexInputStreamArray DummyArray;
+					FMeshMaterialShader::GetElementShaderBindings(Shader, Scene, View, MeshBatch.VertexFactory, EVertexInputStreamType::Default, Scene->GetFeatureLevel(), PrimitiveSceneProxy, MeshBatch, MeshBatch.Elements[0], ShaderElementData, SingleShaderBindings, DummyArray);
+
+					DispatchCmd.TargetBuffer = RWBuffer;
+					DispatchCmd.NumMaxVertices = UpdateParams.NumVertices;
+
+					// Setup the loose parameters directly on the binding
+					uint32 OutputVertexBaseIndex = VertexBufferOffset / sizeof(float);
+					uint32 MinVertexIndex = MeshBatch.Elements[0].MinVertexIndex;
+					uint32 NumCPUVertices = UpdateParams.NumVertices;
+					if (MeshBatch.Elements[0].MinVertexIndex < MeshBatch.Elements[0].MaxVertexIndex)
+					{
+						NumCPUVertices = 1 + MeshBatch.Elements[0].MaxVertexIndex - MeshBatch.Elements[0].MinVertexIndex;
+					}
+
+					const uint32 VertexBufferNumElements = UpdateParams.VertexBufferSize / sizeof(FVector3f) - MinVertexIndex;
+					if (!ensureMsgf(NumCPUVertices <= VertexBufferNumElements,
+						TEXT("Vertex buffer contains %d vertices, but RayTracingDynamicGeometryConverterCS dispatch command expects at least %d."),
+						VertexBufferNumElements, NumCPUVertices))
+					{
+						NumCPUVertices = VertexBufferNumElements;
+					}
 
-		SingleShaderBindings.Add(Shader->UsingIndirectDraw, bUsingIndirectDraw ? 1 : 0);
-		SingleShaderBindings.Add(Shader->NumVertices, NumCPUVertices);
-		SingleShaderBindings.Add(Shader->MinVertexIndex, MinVertexIndex);
-		SingleShaderBindings.Add(Shader->PrimitiveId, PrimitiveId);
-		SingleShaderBindings.Add(Shader->OutputVertexBaseIndex, OutputVertexBaseIndex);
-		SingleShaderBindings.Add(Shader->bApplyWorldPositionOffset, UpdateParams.bApplyWorldPositionOffset ? 1 : 0);
-		SingleShaderBindings.Add(Shader->InstanceId, UpdateParams.InstanceId);
-		SingleShaderBindings.Add(Shader->WorldToInstance, UpdateParams.WorldToInstance);
+					SingleShaderBindings.Add(Shader->UsingIndirectDraw, bUsingIndirectDraw ? 1 : 0);
+					SingleShaderBindings.Add(Shader->NumVertices, NumCPUVertices);
+					SingleShaderBindings.Add(Shader->MinVertexIndex, MinVertexIndex);
+					SingleShaderBindings.Add(Shader->PrimitiveId, PrimitiveId);
+					SingleShaderBindings.Add(Shader->OutputVertexBaseIndex, OutputVertexBaseIndex);
+					SingleShaderBindings.Add(Shader->bApplyWorldPositionOffset, UpdateParams.bApplyWorldPositionOffset ? 1 : 0);
+					SingleShaderBindings.Add(Shader->InstanceId, UpdateParams.InstanceId);
+					SingleShaderBindings.Add(Shader->WorldToInstance, UpdateParams.WorldToInstance);
 
 #if MESH_DRAW_COMMAND_DEBUG_DATA
-		ShaderBindings.Finalize(&MeshProcessorShaders);
+					ShaderBindings.Finalize(&MeshProcessorShaders);
 #endif
 
-		DispatchCommands.Add(DispatchCmd);
+					DispatchCommands.Add(DispatchCmd);
+
+					break;
+				}
+			}
+
+			MaterialRenderProxyPtr = MaterialRenderProxyPtr->GetFallback(Scene->GetFeatureLevel());
+		}
 	}
 
 	bool bRefit = true;
