Change VSM resolution by removing the old LODViewOrigin method, and replacing it with a mixture of ViewTarget location as the VSM world origin location, or defaulting it to the CameraViewOrigin.
Compensates by scaling the FirstLevel VSM depending on the OrthoWidth being used, which has the added benefit of reducing performance for VSM LODs that are not required for Ortho passes. Added a LineTrace system for estimating the location when a ViewTarget is not selected.
Also removes the PostProcess disable setting that previously blocked EyeAdaptation from being applied to Lumen, so SkyLight etc shadows now resolve appropriately, resolving bugs that appear in editor debug views + minimaps.
Also switches Editor debug views to use ortho auto plane calculations/fixes a crash when zooming out. CVars also added for disabling these settings, or forcing different LODBiases in VSMs.
Refactored NearPlane updates to an easier to follow logic path and added compensation for ViewOrigin to ViewTargets.

diff --git a/Engine/Source/Editor/UnrealEd/Private/EditorComponents.cpp b/Engine/Source/Editor/UnrealEd/Private/EditorComponents.cpp
--- a/Engine/Source/Editor/UnrealEd/Private/EditorComponents.cpp
+++ b/Engine/Source/Editor/UnrealEd/Private/EditorComponents.cpp
@@ -262,7 +262,7 @@ void FGridWidget::DrawNewGrid(const FSceneView* View, FPrimitiveDrawInterface* P
 
 	FMatrix ObjectToWorld = FMatrix::Identity;
 
-	FVector CameraPos = View->ViewMatrices.GetLODViewOrigin();
+	FVector CameraPos = View->ViewMatrices.GetViewOrigin();
 
 	FVector UVCameraPos = FVector(CameraPos.X, CameraPos.Y, 0);
 
diff --git a/Engine/Source/Editor/UnrealEd/Private/EditorViewportClient.cpp b/Engine/Source/Editor/UnrealEd/Private/EditorViewportClient.cpp
--- a/Engine/Source/Editor/UnrealEd/Private/EditorViewportClient.cpp
+++ b/Engine/Source/Editor/UnrealEd/Private/EditorViewportClient.cpp
@@ -894,7 +894,6 @@ void FEditorViewportClient::CenterViewportAtPoint(const FVector& NewLookAt, bool
 //////////////////////////////////////////////////////////////////////////
 //
 // Configures the specified FSceneView object with the view and projection matrices for this viewport.
-
 FSceneView* FEditorViewportClient::CalcSceneView(FSceneViewFamily* ViewFamily, const int32 StereoViewIndex)
 {
     const bool bStereoRendering = StereoViewIndex != INDEX_NONE;
@@ -1127,14 +1126,24 @@ FSceneView* FEditorViewportClient::CalcSceneView(FSceneViewFamily* ViewFamily, c
 		else
 		{
 			static_assert((bool)ERHIZBuffer::IsInverted, "Check all the Rotation Matrix transformations!");
-			FMatrix::FReal ZScale = 0.5f / UE_OLD_WORLD_MAX;
-			FMatrix::FReal ZOffset = UE_OLD_WORLD_MAX;
 
 			//The divisor for the matrix needs to match the translation code.
 			const float Zoom = GetOrthoUnitsPerPixel(Viewport);
+			const float OrthoZoom = FMath::Clamp(Zoom * 500.0f, 0.0f, UE_LARGE_HALF_WORLD_MAX);
+
+			float ViewportAspectRatio = float(ViewportSize.X) / (ViewportSize.Y > 0 ? ViewportSize.Y : 1);
 
-			float OrthoWidth = Zoom * ViewportSize.X / 2.0f;
-			float OrthoHeight = Zoom * ViewportSize.Y / 2.0f;
+			float OrthoWidth, OrthoHeight;
+			if(ViewportAspectRatio >= 1.0f)
+			{
+				OrthoWidth = OrthoZoom;
+				OrthoHeight = OrthoWidth / ViewportAspectRatio;
+			}
+			else
+			{
+				OrthoHeight = OrthoZoom;
+				OrthoWidth = OrthoHeight * ViewportAspectRatio;
+			}
 
 			if (EffectiveViewportType == LVT_OrthoXY)
 			{
@@ -1198,25 +1207,28 @@ FSceneView* FEditorViewportClient::CalcSceneView(FSceneViewFamily* ViewFamily, c
 				check(false);
 			}
 
+			FMinimalViewInfo CalculatePlanesViewInfo;
+			CalculatePlanesViewInfo.ProjectionMode = ECameraProjectionMode::Orthographic;
+			CalculatePlanesViewInfo.bAutoCalculateOrthoPlanes = true;
+			CalculatePlanesViewInfo.OrthoWidth = OrthoWidth;
+			CalculatePlanesViewInfo.Rotation = ViewInitOptions.ViewRotationMatrix.Rotator();
+			CalculatePlanesViewInfo.AspectRatio = ViewportAspectRatio;
+			CalculatePlanesViewInfo.bConstrainAspectRatio = false;
+			CalculatePlanesViewInfo.OrthoNearClipPlane = OrthoWidth * -CVarOrthoEditorDebugClipPlaneScale.GetValueOnAnyThread();
+			CalculatePlanesViewInfo.OrthoFarClipPlane = FarPlane - NearPlane + CalculatePlanesViewInfo.OrthoNearClipPlane;
+			CalculatePlanesViewInfo.AutoCalculateOrthoPlanes(ViewInitOptions);
+
+			ViewInitOptions.UpdateOrthoPlanes(CalculatePlanesViewInfo);
+
+			const float ZScale = 1.0f / (CalculatePlanesViewInfo.OrthoFarClipPlane - CalculatePlanesViewInfo.OrthoNearClipPlane);
+			const float ZOffset = -CalculatePlanesViewInfo.OrthoNearClipPlane;
+
 			ViewInitOptions.ProjectionMatrix = FReversedZOrthoMatrix(
 				OrthoWidth,
 				OrthoHeight,
 				ZScale,
 				ZOffset
-				);
-
-			if(!ViewFamily->EngineShowFlags.Wireframe)
-			{
-				/**
-				* Update the ortho near plane and view origin to a position that is proportional to the OrthoWidth
-				* and a User specified multiplier which allows adjusting of the NearPlane location. 
-				* This ensures lighting works more appropriately on meshes in Lit mode, but introduces more abrupt NearPlane clipping.
-				*
-				* Not needed for Wireframe view; NearPlane clipping is unnecessary outside of Lit modes.
-				*/
-				float OrthoNearPlane = OrthoWidth * -CVarOrthoEditorDebugClipPlaneScale.GetValueOnAnyThread();
-				ViewInitOptions.UpdateOrthoNearPlane(OrthoNearPlane, true);
-			}
+			);
 		}
 
 		if (bConstrainAspectRatio)
diff --git a/Engine/Source/Runtime/Engine/Classes/Camera/CameraTypes.h b/Engine/Source/Runtime/Engine/Classes/Camera/CameraTypes.h
--- a/Engine/Source/Runtime/Engine/Classes/Camera/CameraTypes.h
+++ b/Engine/Source/Runtime/Engine/Classes/Camera/CameraTypes.h
@@ -109,7 +109,7 @@ struct FMinimalViewInfo
 
 private:
 	// Only used for Ortho camera auto plane calculations, tells the Near plane of the extra distance that needs to be added.
-	float OrthoCameraArmLength;
+	FVector CameraToViewTarget;
 
 public:
 
@@ -129,7 +129,7 @@ public:
 		, ProjectionMode(ECameraProjectionMode::Perspective)
 		, PostProcessBlendWeight(0.0f)
 		, OffCenterProjectionOffset(ForceInitToZero)
-		, OrthoCameraArmLength(0.0f)
+		, CameraToViewTarget(FVector::ZeroVector)
 	{
 	}
 
@@ -161,12 +161,12 @@ public:
 	}
 
 	/** Automatically calculates the Near/Far plane values for an Ortho camera */
-	ENGINE_API bool AutoCalculateOrthoPlanes(const FSceneViewProjectionData& InOutProjectionData);
+	ENGINE_API bool AutoCalculateOrthoPlanes(FSceneViewProjectionData& InOutProjectionData);
 
 	/** Sets the camera distance from view target for AutoCalculateOrthoPlanes */
-	ENGINE_API inline void SetOrthoCameraArmLengthFromOwnerLocation(const FVector ActorLocation)
+	ENGINE_API inline void SetCameraToViewTarget(const FVector ActorLocation)
 	{
-		OrthoCameraArmLength = static_cast<float>((ActorLocation - Location).Length());
+		CameraToViewTarget = ActorLocation - Location;
 	}
 };
 
diff --git a/Engine/Source/Runtime/Engine/Private/Camera/CameraComponent.cpp b/Engine/Source/Runtime/Engine/Private/Camera/CameraComponent.cpp
--- a/Engine/Source/Runtime/Engine/Private/Camera/CameraComponent.cpp
+++ b/Engine/Source/Runtime/Engine/Private/Camera/CameraComponent.cpp
@@ -424,7 +424,7 @@ void UCameraComponent::GetCameraView(float DeltaTime, FMinimalViewInfo& DesiredV
 	{
 		if (const AActor* ViewTarget = GetOwner())
 		{
-			DesiredView.SetOrthoCameraArmLengthFromOwnerLocation(ViewTarget->GetActorLocation());
+			DesiredView.SetCameraToViewTarget(ViewTarget->GetActorLocation());
 		}
 	}
 
diff --git a/Engine/Source/Runtime/Engine/Private/Camera/CameraStackTypes.cpp b/Engine/Source/Runtime/Engine/Private/Camera/CameraStackTypes.cpp
--- a/Engine/Source/Runtime/Engine/Private/Camera/CameraStackTypes.cpp
+++ b/Engine/Source/Runtime/Engine/Private/Camera/CameraStackTypes.cpp
@@ -251,12 +251,11 @@ void FMinimalViewInfo::CalculateProjectionMatrixGivenViewRectangle(FMinimalViewI
 		// Enforce a particular aspect ratio for the render of the scene. 
 		// Results in black bars at top/bottom etc.
 		InOutProjectionData.SetConstrainedViewRectangle(ConstrainedViewRectangle);
-		InOutProjectionData.ProjectionMatrix = ViewInfo.CalculateProjectionMatrix();
-		if (bOrthographicNearPlaneCorrection)
+		if(bOrthographicNearPlaneCorrection)
 		{
-			float NearPlane = ViewInfo.OrthoNearClipPlane;
-			InOutProjectionData.UpdateOrthoNearPlane(NearPlane, true);
+			InOutProjectionData.UpdateOrthoPlanes(ViewInfo);
 		}
+		InOutProjectionData.ProjectionMatrix = ViewInfo.CalculateProjectionMatrix();
 	}
 	else
 	{
@@ -309,11 +308,11 @@ void FMinimalViewInfo::CalculateProjectionMatrixGivenViewRectangle(FMinimalViewI
 			const float OrthoWidth = ViewInfo.OrthoWidth / 2.0f * XAxisMultiplier;
 			const float OrthoHeight = (ViewInfo.OrthoWidth / 2.0f) / YAxisMultiplier;
 
-			const float FarPlane = ViewInfo.OrthoFarClipPlane;
+			float FarPlane = ViewInfo.OrthoFarClipPlane;
 			float NearPlane = ViewInfo.OrthoNearClipPlane;
 			if (bOrthographicNearPlaneCorrection)
 			{
-				InOutProjectionData.UpdateOrthoNearPlane(NearPlane);
+				InOutProjectionData.UpdateOrthoPlanes(NearPlane, FarPlane);
 			}
 
 			const float ZScale = 1.0f / (FarPlane - NearPlane);
@@ -357,7 +356,7 @@ void FMinimalViewInfo::CalculateProjectionMatrixGivenView(FMinimalViewInfo& View
 	CalculateProjectionMatrixGivenViewRectangle(ViewInfo, AspectRatioAxisConstraint, ViewExtents, InOutProjectionData);
 }
 
-bool FMinimalViewInfo::AutoCalculateOrthoPlanes(const FSceneViewProjectionData& InOutProjectionData)
+bool FMinimalViewInfo::AutoCalculateOrthoPlanes(FSceneViewProjectionData& InOutProjectionData)
 {	
 	if (ProjectionMode == ECameraProjectionMode::Orthographic && CVarOrthoAllowAutoPlanes.GetValueOnAnyThread() && bAutoCalculateOrthoPlanes)
 	{
@@ -380,7 +379,7 @@ bool FMinimalViewInfo::AutoCalculateOrthoPlanes(const FSceneViewProjectionData&
 		//Get the normalized view forward vector of the camera
 		const FRotationMatrix RotMat(Rotation);
 		FVector ViewForward = RotMat.GetColumn(2);
-		ViewForward.Normalize();	
+		ViewForward.Normalize();
 
 		/**
 		 * The CosAngle is the cosine of the angle between the ViewForward and camera down.
@@ -421,7 +420,7 @@ bool FMinimalViewInfo::AutoCalculateOrthoPlanes(const FSceneViewProjectionData&
 		* The camera arm length is adjusted depending on the CosAngle as the horizontal view typically has a significantly larger plane range, 
 		* so it becomes irrelevant, whereas it is necessary to account for in a top down view. Note: a small scene camera arm length will become irrelevant for a large ortho width.
 		*/
-		float CameraArmLength = OrthoCameraArmLength * CosAngle;
+		float CameraArmLength = static_cast<float>(CameraToViewTarget.Length()) * CosAngle;
 
 		/**
 		 * The NearPlane calculation is a scaled OrthoHeight depending on the camera angle, 
@@ -432,14 +431,15 @@ bool FMinimalViewInfo::AutoCalculateOrthoPlanes(const FSceneViewProjectionData&
 		 * We clamp this to remove the Near plane difference, and also max out at the previously set maximum FPValue. 
 		 * This setup should help for possible future implementations where we can increase the depth range (i.e. LWC + double float depth buffers).
 		 */
-		float SinAngle = FMath::Clamp(1.0f - CosAngle, 0.707f, 1.0f);
-		float NearPlane = OrthoHeight * (FMath::Clamp(CosAngle, 0.707f, 1.0f) - (1.0f/SinAngle)) - CameraArmLength;
+		float SinAngle = FMath::Clamp(1.0f - CosAngle, 0.707107f, 1.0f);
+		float NearPlane = FMath::Max(OrthoWidth, OrthoHeight) * FMath::Max((FMath::Clamp(CosAngle, 0.707107f, 1.0f) - (1.0f/SinAngle)), -0.5f) - CameraArmLength;
 		FarPlane = FMath::Clamp(FarPlane, OrthoHeight, MaxFPValue + NearPlane);
 
 		//The Planes can be scaled in the Z axis without restriction to ensure a user can capture their entire view.
 		const float AutoPlaneShift = CVarOrthoAutoPlaneShift.GetValueOnAnyThread();
 		OrthoNearClipPlane = NearPlane + AutoPlaneShift;
 		OrthoFarClipPlane = FarPlane + AutoPlaneShift;
+		InOutProjectionData.CameraToViewTarget = CameraToViewTarget;
 		return true;
 	}
 	return false;
diff --git a/Engine/Source/Runtime/Engine/Private/Components/SceneCaptureComponent.cpp b/Engine/Source/Runtime/Engine/Private/Components/SceneCaptureComponent.cpp
--- a/Engine/Source/Runtime/Engine/Private/Components/SceneCaptureComponent.cpp
+++ b/Engine/Source/Runtime/Engine/Private/Components/SceneCaptureComponent.cpp
@@ -706,7 +706,7 @@ void USceneCaptureComponent2D::GetCameraView(float DeltaTime, FMinimalViewInfo&
 	{
 		if(const AActor* ViewTarget = GetOwner())
 		{
-			OutMinimalViewInfo.SetOrthoCameraArmLengthFromOwnerLocation(ViewTarget->GetActorLocation());
+			OutMinimalViewInfo.SetCameraToViewTarget(ViewTarget->GetActorLocation());
 		}
 	}
 }
diff --git a/Engine/Source/Runtime/Engine/Private/PlayerCameraManager.cpp b/Engine/Source/Runtime/Engine/Private/PlayerCameraManager.cpp
--- a/Engine/Source/Runtime/Engine/Private/PlayerCameraManager.cpp
+++ b/Engine/Source/Runtime/Engine/Private/PlayerCameraManager.cpp
@@ -453,7 +453,7 @@ void APlayerCameraManager::UpdateViewTarget(FTViewTarget& OutVT, float DeltaTime
 	SetActorLocationAndRotation(OutVT.POV.Location, OutVT.POV.Rotation, false);
 	if (bAutoCalculateOrthoPlanes && OutVT.Target)
 	{
-		OutVT.POV.SetOrthoCameraArmLengthFromOwnerLocation(OutVT.Target->GetActorLocation());
+		OutVT.POV.SetCameraToViewTarget(OutVT.Target->GetActorLocation());
 	}
 
 	UpdateCameraLensEffects(OutVT);
diff --git a/Engine/Source/Runtime/Engine/Private/SceneManagement.cpp b/Engine/Source/Runtime/Engine/Private/SceneManagement.cpp
--- a/Engine/Source/Runtime/Engine/Private/SceneManagement.cpp
+++ b/Engine/Source/Runtime/Engine/Private/SceneManagement.cpp
@@ -87,7 +87,7 @@ void FTemporalLODState::UpdateTemporalLODTransition(const FSceneView& View, floa
 				TemporalLODViewOrigin[0] = TemporalLODViewOrigin[1];
 				TemporalLODTime[0] = TemporalLODTime[1];
 			}
-			TemporalLODViewOrigin[1] = View.ViewMatrices.GetLODViewOrigin();
+			TemporalLODViewOrigin[1] = View.ViewMatrices.GetViewOrigin();
 			TemporalLODTime[1] = LastRenderTime;
 			if (TemporalLODTime[1] <= TemporalLODTime[0])
 			{
@@ -97,8 +97,8 @@ void FTemporalLODState::UpdateTemporalLODTransition(const FSceneView& View, floa
 	}
 	if (!bOk)
 	{
-		TemporalLODViewOrigin[0] = View.ViewMatrices.GetLODViewOrigin();
-		TemporalLODViewOrigin[1] = View.ViewMatrices.GetLODViewOrigin();
+		TemporalLODViewOrigin[0] = View.ViewMatrices.GetViewOrigin();
+		TemporalLODViewOrigin[1] = View.ViewMatrices.GetViewOrigin();
 		TemporalLODTime[0] = LastRenderTime;
 		TemporalLODTime[1] = LastRenderTime;
 		TemporalLODLag = 0.0f;
diff --git a/Engine/Source/Runtime/Engine/Private/SceneView.cpp b/Engine/Source/Runtime/Engine/Private/SceneView.cpp
--- a/Engine/Source/Runtime/Engine/Private/SceneView.cpp
+++ b/Engine/Source/Runtime/Engine/Private/SceneView.cpp
@@ -24,6 +24,7 @@
 #include "StereoRenderUtils.h"
 #include "SceneRelativeViewMatrices.h"
 #include "Camera/CameraComponent.h"
+#include "Camera/CameraTypes.h"
 #include "UObject/Interface.h"
 
 DEFINE_LOG_CATEGORY(LogBufferVisualization);
@@ -320,7 +321,7 @@ static TAutoConsoleVariable<int32> CVarOrthoCalculateDepthThicknessScaling(
 	TEXT("r.Ortho.CalculateDepthThicknessScaling"),
 	1,
 	TEXT("Whether to automatically derive the depth thickness test scale from the Near/FarPlane difference.\n")
-	TEXT("0: Disabled (use scaling specified by r.Ortho.DepthTicknessScale)\n")
+	TEXT("0: Disabled (use scaling specified by r.Ortho.DepthThicknessScale)\n")
 	TEXT("1: Enabled (default)"),
 	ECVF_Scalability | ECVF_RenderThreadSafe);
 
@@ -340,6 +341,13 @@ static FAutoConsoleVariableRef CVarDefaultUpdateOrthoNearPlane(
 	ECVF_RenderThreadSafe
 );
 
+static TAutoConsoleVariable<bool> CVarOrthoCameraHeightAsViewTarget(
+	TEXT("r.Ortho.CameraHeightAsViewTarget"),
+	true,
+	TEXT("Sets whether to use the camera height as a pseudo camera to view target.\n")
+	TEXT("Primarily helps with VSM clipmap selection and avoids overcorrecting NearPlanes.\n"),
+	ECVF_Scalability | ECVF_RenderThreadSafe);
+
 int32 GVirtualTextureFeedbackFactor = 16;
 static FAutoConsoleVariableRef CVarVirtualTextureFeedbackFactor(
 	TEXT("r.vt.FeedbackFactor"),
@@ -552,51 +560,78 @@ FVector4f CreateInvDeviceZToWorldZTransform(const FMatrix& ProjMatrix)
 	}
 }
 
-bool FSceneViewProjectionData::UpdateOrthoNearPlane(FSceneViewProjectionData* InOutProjectionData, float& NearPlane, bool bUpdateOrthoProjectionMatrix)
+bool FSceneViewProjectionData::UpdateOrthoPlanes(FSceneViewProjectionData* InOutProjectionData, float& NearPlane, float& FarPlane)
 {
 	if (!InOutProjectionData)
 	{
 		return false;
 	}
-
-	//Store the original ViewOrigin for LOD location resolving regardless of if we early out.
-	InOutProjectionData->LODViewOrigin = InOutProjectionData->ViewOrigin;
-
-	if (NearPlane >= 0.0f)
-	{
-		return false;
-	}
 	
 	//Get the ViewForward vector from the RotationMatrix + ensure it is normalized.
 	FVector ViewForward = InOutProjectionData->ViewRotationMatrix.GetColumn(2);
 	ViewForward.Normalize();
 
-	//OrthoNearClipPlane is negative at this point + we are moving the theoretical camera position backwards.
-	InOutProjectionData->ViewOrigin += ViewForward * NearPlane;
-
-	//If required, recalculate the new projection matrix from the old one using the new NearClip value
-	if (bUpdateOrthoProjectionMatrix)
+	float PlaneDifference = FarPlane - NearPlane;
+	if(InOutProjectionData->CameraToViewTarget.Length() > 0)
 	{
-		const FMatrix InvProjectionMatrix = InOutProjectionData->ProjectionMatrix.InverseFast();
-		const float OrthoWidth = InvProjectionMatrix.M[0][0];
-		const float OrthoHeight = InvProjectionMatrix.M[1][1];
+		//Store the ViewTargetLocation to correct it for the repositioned ViewOrigin.
+		FVector ViewTargetLocation = InOutProjectionData->ViewOrigin + InOutProjectionData->CameraToViewTarget;
+		//OrthoNearClipPlane is negative at this point + we are moving the theoretical camera position backwards.
+		InOutProjectionData->ViewOrigin += InOutProjectionData->CameraToViewTarget +(ViewForward * NearPlane);
+		//Save out the new CameraToViewTarget vector.
+		InOutProjectionData->CameraToViewTarget = ViewTargetLocation - InOutProjectionData->ViewOrigin;
+	}
+	else
+	{	
+		/**
+		* Use the height of the camera as a sort of pseudo CameraToViewTarget to remove camera position as much as possible so we minimise the near clip plane distance.
+		* Depends on view direction being top down as when we get to a side view, the scale grows much larger and it becomes redundant.
+		*/
+		float CameraHeightAdjustment = CVarOrthoCameraHeightAsViewTarget.GetValueOnAnyThread() ? FMath::Abs(InOutProjectionData->ViewOrigin.Z) * FMath::Abs((ViewForward.Dot(FVector(0, 0, -1.0f)))) : 0.0f;
+		InOutProjectionData->ViewOrigin += ViewForward * (CameraHeightAdjustment + NearPlane);
+	}
+	NearPlane = GDefaultUpdateOrthoNearPlane;
+	FarPlane = NearPlane + PlaneDifference;
 
-		const float FarPlane = NearPlane - (InvProjectionMatrix.M[2][2]);
-		NearPlane = GDefaultUpdateOrthoNearPlane;
-		const float ZScale = 1.0f / (FarPlane - NearPlane);
+	return true;
+}
 
-		InOutProjectionData->ProjectionMatrix = FReversedZOrthoMatrix(
-			OrthoWidth,
-			OrthoHeight,
-			ZScale,
-			-NearPlane
-		);
+bool FSceneViewProjectionData::UpdateOrthoPlanes(FMinimalViewInfo& MinimalViewInfo)
+{
+	return UpdateOrthoPlanes(MinimalViewInfo.OrthoNearClipPlane, MinimalViewInfo.OrthoFarClipPlane);
+}
+
+bool FSceneViewProjectionData::UpdateOrthoPlanes()
+{
+	/**
+	* This function takes the existing projection matrix and moves the nearplane + view origin.
+	* Separating this step from the near plane calculation logic itself helps avoid applying the NearPlane correctiontwice /makes it easier to read.
+	*/
+	if(IsPerspectiveProjection() 
+		|| ProjectionMatrix.M[2][2] == 0.0 
+		|| ProjectionMatrix.M[2][2] == ProjectionMatrix.M[2][3])
+	{ 
+		return false;	
 	}
-	else
+
+	//Get existing Near and Far plane + their difference
+	float NearPlane = static_cast<float>(ProjectionMatrix.M[3][3] - ProjectionMatrix.M[3][2]) / (ProjectionMatrix.M[2][2] - ProjectionMatrix.M[2][3]);
+	float FarPlane = NearPlane - 1.0f/static_cast<float>(ProjectionMatrix.M[2][2]);
+	float PlaneDifference = FarPlane - NearPlane;
+	if (FarPlane - NearPlane == 0.0f)
 	{
-		NearPlane = GDefaultUpdateOrthoNearPlane;
+		return false;
 	}
 
+	UpdateOrthoPlanes(NearPlane, FarPlane);
+
+	const float ZScale = 1.0f / (FarPlane - NearPlane);
+	const float ZOffset = -NearPlane;
+
+	//Only the Near/Far plane elements need correcting, OrthoWidth/Height remains the same
+	ProjectionMatrix.M[2][2] = -ZScale;
+	ProjectionMatrix.M[2][3] = 1.0f - (ZOffset * ZScale);
+
 	return true;
 }
 
@@ -631,6 +666,7 @@ void FViewMatrices::Init(const FMinimalInitializer& Initializer)
 	// Translate world-space so its origin is at ViewOrigin for improved precision.
 	ViewOrigin = LocalViewOrigin;
 	PreViewTranslation = -LocalViewOrigin;
+	CameraToViewTarget = Initializer.CameraToViewTarget;
 
 	FMatrix LocalTranslatedViewMatrix = ViewRotationMatrix;
 	FMatrix LocalInvTranslatedViewMatrix = LocalTranslatedViewMatrix.GetTransposed();
@@ -654,7 +690,6 @@ void FViewMatrices::Init(const FMinimalInitializer& Initializer)
 		ProjectionScale.X = ScreenXScale * FMath::Abs(ProjectionMatrix.M[0][0]);
 		ProjectionScale.Y = FMath::Abs(ProjectionMatrix.M[1][1]);
 		PerProjectionDepthThicknessScale = 1.0f;
-		LODViewOrigin = ViewOrigin;
 	}
 	else
 	{
@@ -670,8 +705,6 @@ void FViewMatrices::Init(const FMinimalInitializer& Initializer)
 		{
 			PerProjectionDepthThicknessScale = GOrthographicDepthThicknessScale;
 		}
-
-		LODViewOrigin = Initializer.LODViewOrigin;
 	}
 	ScreenScale = FMath::Max(
 		Initializer.ConstrainedViewRect.Size().X * 0.5f * ProjectionScale.X,
@@ -686,7 +719,7 @@ FViewMatrices::FViewMatrices(const FSceneViewInitOptions& InitOptions) : FViewMa
 	Initializer.ViewRotationMatrix   = InitOptions.ViewRotationMatrix;
 	Initializer.ProjectionMatrix     = InitOptions.ProjectionMatrix;
 	Initializer.ViewOrigin           = InitOptions.ViewOrigin;
-	Initializer.LODViewOrigin        = InitOptions.LODViewOrigin;
+	Initializer.CameraToViewTarget	 = InitOptions.CameraToViewTarget;
 	Initializer.ConstrainedViewRect  = InitOptions.GetConstrainedViewRect();
 	Initializer.StereoPass           = InitOptions.StereoPass;
 
@@ -2657,7 +2690,8 @@ void FSceneView::SetupCommonViewUniformBufferParameters(
 	// FrameCounter is incremented once per engine tick, so multi views of the same frame have the same value.
 	ViewUniformShaderParameters.FrameCounter = Family->FrameCounter;
 	ViewUniformShaderParameters.WorldIsPaused = Family->bWorldIsPaused;
-	ViewUniformShaderParameters.CameraCut = bCameraCut ? 1 : 0;
+	//Set bCameraCut if we switch projection type to ensure histories are updated.
+	ViewUniformShaderParameters.CameraCut = bCameraCut ? 1 : (InViewMatrices.IsPerspectiveProjection() != InPrevViewMatrices.IsPerspectiveProjection());
 
 	ViewUniformShaderParameters.MinRoughness = FMath::Clamp(CVarGlobalMinRoughnessOverride.GetValueOnRenderThread(), 0.02f, 1.0f);
 
diff --git a/Engine/Source/Runtime/Engine/Private/ShowFlags.cpp b/Engine/Source/Runtime/Engine/Private/ShowFlags.cpp
--- a/Engine/Source/Runtime/Engine/Private/ShowFlags.cpp
+++ b/Engine/Source/Runtime/Engine/Private/ShowFlags.cpp
@@ -377,11 +377,6 @@ void ApplyViewMode(EViewModeIndex ViewModeIndex, bool bPerspective, FEngineShowF
 			break;
 	}
 
-	if (!bPerspective)
-	{
-		bPostProcessing = false;
-	}
-
 	// set the EngineShowFlags:
 
 	// Assigning the new state like this ensures we always set the same variables (they depend on the view mode)
diff --git a/Engine/Source/Runtime/Engine/Public/SceneView.h b/Engine/Source/Runtime/Engine/Public/SceneView.h
--- a/Engine/Source/Runtime/Engine/Public/SceneView.h
+++ b/Engine/Source/Runtime/Engine/Public/SceneView.h
@@ -41,6 +41,7 @@ class FSceneViewFamily;
 class FVolumetricFogViewResources;
 class ISceneRenderer;
 class ISpatialUpscaler;
+struct FMinimalViewInfo;
 
 namespace UE::Renderer::Private
 {
@@ -63,12 +64,12 @@ struct FSceneViewProjectionData
 	/** UE projection matrix projects such that clip space Z=1 is the near plane, and Z=0 is the infinite far plane. */
 	FMatrix ProjectionMatrix;
 
-	/** This can be specified for ortho views so that the min draw distance/LOD resolution etc, is controlled by camera location rather than the pseudo infinite viewport location of Ortho*/
-	FVector LODViewOrigin;
-
 	//The unconstrained (no aspect ratio bars applied) view rectangle (also unscaled)
 	FIntRect ViewRect;
 
+	//The vector (including distance) from the camera to it's viewtarget, if set. Primarily only used for Ortho views.
+	FVector CameraToViewTarget;
+
 protected:
 	// The constrained view rectangle (identical to UnconstrainedUnscaledViewRect if aspect ratio is not constrained)
 	FIntRect ConstrainedViewRect;
@@ -127,25 +128,13 @@ public:
 	}
 
 	// Function for correcting Ortho camera near plane locations to avoid artifacts behind camera view origin
-	static ENGINE_API bool UpdateOrthoNearPlane(FSceneViewProjectionData* InOutProjectionData, float& NearPlane, bool bUpdateOrthoProjectionMatrix = false);
-
-	static ENGINE_API bool UpdateOrthoNearPlane(FMatrix& ProjectionMatrix, float& NearPlane, bool bUpdateOrthoProjectionMatrix = false)
+	static ENGINE_API bool UpdateOrthoPlanes(FSceneViewProjectionData* InOutProjectionData, float& NearPlane, float& FarPlane);
+	ENGINE_API inline bool UpdateOrthoPlanes(float& NearPlane, float& FarPlane)
 	{
-		FSceneViewProjectionData InOutProjectionData;
-		InOutProjectionData.ProjectionMatrix = ProjectionMatrix;
-		InOutProjectionData.ViewOrigin = FVector::ZeroVector;
-		bool Result = UpdateOrthoNearPlane(&InOutProjectionData, NearPlane, bUpdateOrthoProjectionMatrix);
-		if(bUpdateOrthoProjectionMatrix && Result)
-		{
-			ProjectionMatrix = InOutProjectionData.ProjectionMatrix;
-		}
-		return Result;
-	}
-	
-	bool UpdateOrthoNearPlane(float& NearPlane, bool bUpdateOrthoProjectionMatrix = false)
-	{
-		return UpdateOrthoNearPlane(this, NearPlane, bUpdateOrthoProjectionMatrix);
+		return UpdateOrthoPlanes(this, NearPlane, FarPlane);
 	}
+	ENGINE_API bool UpdateOrthoPlanes(FMinimalViewInfo& MinimalViewInfo);
+	ENGINE_API bool UpdateOrthoPlanes();
 };
 
 /** Method used for primary screen percentage method. */
@@ -302,8 +291,8 @@ struct FViewMatrices
 		FMatrix ViewRotationMatrix = FMatrix::Identity;
 		FMatrix ProjectionMatrix = FMatrix::Identity;
 		FVector ViewOrigin = FVector::ZeroVector;
-		FVector LODViewOrigin = FVector::ZeroVector;
 		FIntRect ConstrainedViewRect = FIntRect(0, 0, 0, 0);
+		FVector CameraToViewTarget = FVector::ZeroVector;
 		EStereoscopicPass StereoPass = EStereoscopicPass::eSSP_FULL;
 	};
 
@@ -326,7 +315,7 @@ struct FViewMatrices
 		ScreenToClipMatrix.SetIdentity();
 		PreViewTranslation = FVector::ZeroVector;
 		ViewOrigin = FVector::ZeroVector;
-		LODViewOrigin = FVector::ZeroVector;
+		CameraToViewTarget = FVector::ZeroVector;
 		ProjectionScale = FVector2D::ZeroVector;
 		TemporalAAProjectionJitter = FVector2D::ZeroVector;
 		ScreenScale = 1.f;
@@ -373,8 +362,8 @@ private:
 	FVector		PreViewTranslation;
 	/** The camera/viewport location in world space */
 	FVector		ViewOrigin;
-	/** The camera/viewport location for resolving LOD sizes for Orthographic cameras */
-	FVector 	LODViewOrigin;
+	/** The current view target's location proportional to the camera location */
+	FVector		CameraToViewTarget;
 	/** Scale applied by the projection matrix in X and Y. */
 	FVector2D	ProjectionScale;
 	/** TemporalAA jitter offset currently stored in the projection matrix */
@@ -490,10 +479,9 @@ public:
 		return ViewOrigin;
 	}
 
-	inline const FVector& GetLODViewOrigin() const
+	inline const FVector& GetCameraToViewTarget() const
 	{
-		//Perspective should always set this to the ViewOrigin in Init, Ortho needs a corrected location for resolving LODs
-		return LODViewOrigin;
+		return CameraToViewTarget;
 	}
 
 	inline float GetScreenScale() const
@@ -635,18 +623,22 @@ public:
 	// @return in radians (horizontal,vertical)
 	const FVector2D ComputeHalfFieldOfViewPerAxis() const
 	{
-		const FMatrix ClipToView = ComputeInvProjectionNoAAMatrix();
+		if(IsPerspectiveProjection())
+		{
+			const FMatrix ClipToView = ComputeInvProjectionNoAAMatrix();
 
-		FVector VCenter = FVector(ClipToView.TransformPosition(FVector(0.0, 0.0, 0.0)));
-		FVector VUp = FVector(ClipToView.TransformPosition(FVector(0.0, 1.0, 0.0)));
-		FVector VRight = FVector(ClipToView.TransformPosition(FVector(1.0, 0.0, 0.0)));
+			FVector VCenter = FVector(ClipToView.TransformPosition(FVector(0.0, 0.0, 0.0)));
+			FVector VUp = FVector(ClipToView.TransformPosition(FVector(0.0, 1.0, 0.0)));
+			FVector VRight = FVector(ClipToView.TransformPosition(FVector(1.0, 0.0, 0.0)));
 
-		VCenter.Normalize();
-		VUp.Normalize();
-		VRight.Normalize();
+			VCenter.Normalize();
+			VUp.Normalize();
+			VRight.Normalize();
 
-		using ResultType = decltype(FVector2D::X);
-		return FVector2D((ResultType)FMath::Acos(VCenter | VRight), (ResultType)FMath::Acos(VCenter | VUp));
+			using ResultType = decltype(FVector2D::X);
+			return FVector2D((ResultType)FMath::Acos(VCenter | VRight), (ResultType)FMath::Acos(VCenter | VUp));
+		}
+		return FVector2D::Zero();
 	}
 
 	FMatrix::FReal ComputeNearPlane() const
@@ -654,6 +646,11 @@ public:
 		return ( ProjectionMatrix.M[3][3] - ProjectionMatrix.M[3][2] ) / ( ProjectionMatrix.M[2][2] - ProjectionMatrix.M[2][3] );
 	}
 
+	FMatrix::FReal ComputeFarPlane() const
+	{
+		return ComputeNearPlane() - InvProjectionMatrix.M[2][2];
+	}
+
 	void ApplyWorldOffset(const FVector& InOffset)
 	{
 		ViewOrigin+= InOffset;
diff --git a/Engine/Source/Runtime/Landscape/Private/LandscapeRender.cpp b/Engine/Source/Runtime/Landscape/Private/LandscapeRender.cpp
--- a/Engine/Source/Runtime/Landscape/Private/LandscapeRender.cpp
+++ b/Engine/Source/Runtime/Landscape/Private/LandscapeRender.cpp
@@ -4416,7 +4416,7 @@ float FLandscapeComponentSceneProxy::ComputeLODForView(const FSceneView& InView)
 	float ViewLODDistanceFactor = InView.LODDistanceFactor;
 	bool ViewEngineShowFlagCollisionPawn = InView.Family->EngineShowFlags.CollisionPawn;
 	bool ViewEngineShowFlagCollisionVisibility = InView.Family->EngineShowFlags.CollisionVisibility;
-	const FVector& ViewOrigin = GetLODView(InView).ViewMatrices.GetLODViewOrigin();
+	const FVector& ViewOrigin = GetLODView(InView).ViewMatrices.GetViewOrigin();
 	const FMatrix& ViewProjectionMatrix = GetLODView(InView).ViewMatrices.GetProjectionMatrix();
 
 	float LODScale = ViewLODDistanceFactor * CVarStaticMeshLODDistanceScale.GetValueOnRenderThread();
diff --git a/Engine/Source/Runtime/Renderer/Private/SceneCaptureRendering.cpp b/Engine/Source/Runtime/Renderer/Private/SceneCaptureRendering.cpp
--- a/Engine/Source/Runtime/Renderer/Private/SceneCaptureRendering.cpp
+++ b/Engine/Source/Runtime/Renderer/Private/SceneCaptureRendering.cpp
@@ -801,8 +801,7 @@ static FSceneRenderer* CreateSceneRendererForSceneCapture(
 	{
 		if (!SceneCaptureViewInfo.IsPerspectiveProjection() && !SceneCaptureComponent2D->bUseCustomProjectionMatrix)
 		{
-			float NearPlane = SceneCaptureViewInfo.GetNearPlaneFromProjectionMatrix();
-			SceneCaptureViewInfo.UpdateOrthoNearPlane(NearPlane, true);
+			SceneCaptureViewInfo.UpdateOrthoPlanes();
 		}
 	}
 
diff --git a/Engine/Source/Runtime/Renderer/Private/ShadowSetup.cpp b/Engine/Source/Runtime/Renderer/Private/ShadowSetup.cpp
--- a/Engine/Source/Runtime/Renderer/Private/ShadowSetup.cpp
+++ b/Engine/Source/Runtime/Renderer/Private/ShadowSetup.cpp
@@ -1946,7 +1946,7 @@ bool FProjectedShadowInfo::AddSubjectPrimitive(FDynamicShadowsTaskData& TaskData
 				if( CurrentView.IsPerspectiveProjection() )
 				{
 					// Compute the distance between the view and the primitive.
-					float DistanceSquared = (Proxy->GetBounds().Origin - CurrentView.ShadowViewMatrices.GetLODViewOrigin()).SizeSquared();
+					float DistanceSquared = (Proxy->GetBounds().Origin - CurrentView.ShadowViewMatrices.GetViewOrigin()).SizeSquared();
 
 					bool bIsDistanceCulled = CurrentView.IsDistanceCulled(
 						DistanceSquared,
@@ -2006,7 +2006,7 @@ bool FProjectedShadowInfo::AddSubjectPrimitive(FDynamicShadowsTaskData& TaskData
 					// Note: skip small-mesh culling for VSM since it needs it drawn for GPU-side caching.
 					if (bWholeSceneShadow && CacheMode != SDCM_StaticPrimitivesOnly && MeshPassTargetType != EMeshPass::VSMShadowDepth)
 					{
-						const float DistanceSquared = ( Bounds.Origin - CurrentView.ShadowViewMatrices.GetLODViewOrigin() ).SizeSquared();
+						const float DistanceSquared = ( Bounds.Origin - CurrentView.ShadowViewMatrices.GetViewOrigin() ).SizeSquared();
 						const float LODScaleSquared = FMath::Square(CurrentView.LODDistanceFactor);
 						const bool bDrawShadowDepth = FMath::Square(Bounds.SphereRadius) > FMath::Square(GMinScreenRadiusForShadowCaster) * DistanceSquared * LODScaleSquared;
 						if( !bDrawShadowDepth )
@@ -2163,7 +2163,7 @@ uint64 FProjectedShadowInfo::AddSubjectPrimitive_AnyThread(
 			{
 				bool bFadingIn;
 				// Compute the distance between the view and the primitive.
-				const float DistanceSquared = (PrimitiveSceneInfoCompact.Bounds.Origin - CurrentView->ShadowViewMatrices.GetLODViewOrigin()).SizeSquared();
+				const float DistanceSquared = (PrimitiveSceneInfoCompact.Bounds.Origin - CurrentView->ShadowViewMatrices.GetViewOrigin()).SizeSquared();
 
 				if (CurrentView->IsDistanceCulled_AnyThread(
 					DistanceSquared,
@@ -2228,7 +2228,7 @@ uint64 FProjectedShadowInfo::AddSubjectPrimitive_AnyThread(
 				if (bWholeSceneShadow)
 				{
 					const FCompactBoxSphereBounds& Bounds = PrimitiveSceneInfoCompact.Bounds;
-					const float DistanceSquared = (Bounds.Origin - CurrentView->ShadowViewMatrices.GetLODViewOrigin()).SizeSquared();
+					const float DistanceSquared = (Bounds.Origin - CurrentView->ShadowViewMatrices.GetViewOrigin()).SizeSquared();
 					const float LODScaleSquared = FMath::Square(CurrentView->LODDistanceFactor);
 					const bool bDrawShadowDepth = FMath::Square(Bounds.SphereRadius) > FMath::Square(GMinScreenRadiusForShadowCaster) * DistanceSquared * LODScaleSquared;
 					if (!bDrawShadowDepth)
@@ -3536,7 +3536,7 @@ void FSceneRenderer::CreatePerObjectProjectedShadow(
 				const float DistanceFromShadowCenterSquared = (WholeSceneShadow->ShadowBounds.Center - Bounds.Origin).SizeSquared();
 				//@todo - if view dependent whole scene shadows are ever supported in splitscreen, 
 				// We can only disable the preshadow at this point if it is inside a whole scene shadow for all views
-				const float DistanceFromViewSquared = ((FVector)WholeSceneShadow->DependentView->ShadowViewMatrices.GetLODViewOrigin() - Bounds.Origin).SizeSquared();
+				const float DistanceFromViewSquared = ((FVector)WholeSceneShadow->DependentView->ShadowViewMatrices.GetViewOrigin() - Bounds.Origin).SizeSquared();
 				// Mark the preshadow as inside the whole scene shadow if its bounding sphere is inside the near fade distance
 				if (DistanceFromShadowCenterSquared < FMath::Square(FMath::Max(WholeSceneShadow->ShadowBounds.W - Bounds.SphereRadius, 0.0f))
 					//@todo - why is this extra threshold required?
@@ -4909,7 +4909,7 @@ struct FGatherShadowPrimitivesPacket
 				check(ProjectedShadowInfo->DependentView);
 
 				{
-					const float DistanceSquared = (PrimitiveBounds.Origin - ProjectedShadowInfo->DependentView->ShadowViewMatrices.GetLODViewOrigin()).SizeSquared();
+					const float DistanceSquared = (PrimitiveBounds.Origin - ProjectedShadowInfo->DependentView->ShadowViewMatrices.GetViewOrigin()).SizeSquared();
 					const float LODScaleSquared = FMath::Square(ProjectedShadowInfo->DependentView->LODDistanceFactor);
 					bScreenSpaceSizeCulled = FMath::Square(PrimitiveBounds.SphereRadius) < FMath::Square(MinScreenRadiusForShadowCaster) * DistanceSquared * LODScaleSquared;
 				}
diff --git a/Engine/Source/Runtime/Renderer/Private/VelocityRendering.cpp b/Engine/Source/Runtime/Renderer/Private/VelocityRendering.cpp
--- a/Engine/Source/Runtime/Renderer/Private/VelocityRendering.cpp
+++ b/Engine/Source/Runtime/Renderer/Private/VelocityRendering.cpp
@@ -432,7 +432,7 @@ bool FVelocityMeshProcessor::PrimitiveHasVelocityForView(const FViewInfo& View,
 	}
 
 	const FBoxSphereBounds& PrimitiveBounds = PrimitiveSceneProxy->GetBounds();
-	const float LODFactorDistanceSquared = (PrimitiveBounds.Origin - View.ViewMatrices.GetLODViewOrigin()).SizeSquared() * FMath::Square(View.LODDistanceFactor);
+	const float LODFactorDistanceSquared = (PrimitiveBounds.Origin - View.ViewMatrices.GetViewOrigin()).SizeSquared() * FMath::Square(View.LODDistanceFactor);
 
 	// The minimum projected screen radius for a primitive to be drawn in the velocity pass, as a fraction of half the horizontal screen width.
 	float MinScreenRadiusForVelocityPass = View.FinalPostProcessSettings.MotionBlurPerObjectSize * (2.0f / 100.0f);
diff --git a/Engine/Source/Runtime/Renderer/Private/VirtualShadowMaps/VirtualShadowMapClipmap.cpp b/Engine/Source/Runtime/Renderer/Private/VirtualShadowMaps/VirtualShadowMapClipmap.cpp
--- a/Engine/Source/Runtime/Renderer/Private/VirtualShadowMaps/VirtualShadowMapClipmap.cpp
+++ b/Engine/Source/Runtime/Renderer/Private/VirtualShadowMaps/VirtualShadowMapClipmap.cpp
@@ -15,6 +15,9 @@ VirtualShadowMapClipmap.cpp
 #include "VirtualShadowMapArray.h"
 #include "VirtualShadowMapCacheManager.h"
 #include "VirtualShadowMapDefinitions.h"
+#include "CollisionQueryParams.h"
+#include "Engine/HitResult.h"
+#include "Engine/World.h"
 
 extern int32 GForceInvalidateDirectionalVSM;
 
@@ -90,6 +93,34 @@ static TAutoConsoleVariable<int32> CVarClipmapWPODisableDistanceLodBias(
 	ECVF_Scalability | ECVF_RenderThreadSafe
 );
 
+static TAutoConsoleVariable<bool> CVarOrthoVSMEstimateClipmapLevels(
+	TEXT("r.Ortho.VSM.EstimateClipmapLevels"),
+	true,
+	TEXT("Enable/Disable calculating the FirstLevel VSM based on the current camera OrthoWidth"),
+	ECVF_Scalability | ECVF_RenderThreadSafe
+);
+
+static TAutoConsoleVariable<int32> CVarOrthoVSMClipmapLODBias(
+	TEXT("r.Ortho.VSM.ClipmapLODBias"),
+	0,
+	TEXT("LOD setting for adjusting the VSM first level from it's OrthoWidth based value."),
+	ECVF_Scalability | ECVF_RenderThreadSafe
+);
+
+static TAutoConsoleVariable<bool> CVarOrthoVSMProjectViewOrigin(
+	TEXT("r.Ortho.VSM.ProjectViewOrigin"),
+	true,
+	TEXT("Enable/Disable moving the WorldOrigin of the VSM clipmaps to focus around the ViewTarget (if present)"),
+	ECVF_Scalability | ECVF_RenderThreadSafe
+);
+
+static TAutoConsoleVariable<bool> CVarOrthoVSMRayCastViewOrigin(
+	TEXT("r.Ortho.VSM.RayCastViewOrigin"),
+	true,
+	TEXT("Enable/Disable whether the ViewOrigin should be estimated with a raycast if the ViewTarget is not present (i.e. standalone camera)"),
+	ECVF_Scalability | ECVF_RenderThreadSafe
+);
+
 // "Virtual" clipmap level to clipmap radius
 // NOTE: This is the radius of around the clipmap origin that this level must cover
 // The actual clipmap dimensions will be larger due to snapping and other accomodations
@@ -125,9 +156,21 @@ FVirtualShadowMapClipmap::FVirtualShadowMapClipmap(
 	// Pure rotation matrix
 	FMatrix ViewToWorldRotationMatrix = WorldToLightViewRotationMatrix.GetTransposed();
 	
-	// Optionally clamp camera viewport to avoid excessively low resolution shadows with dynamic resolution
-	const int32 MinCameraViewportWidth = CVarVirtualShadowMapClipmapMinCameraViewportWidth.GetValueOnRenderThread();
-	const int32 CameraViewportWidth = FMath::Max(MinCameraViewportWidth, CameraViewRectSize.X);
+	// Optionally clamp camera viewport to avoid excessively low resolution shadows with dynamic resolution	
+	const int32 MinCameraViewportWidth = CVarVirtualShadowMapClipmapMinCameraViewportWidth.GetValueOnRenderThread();	
+	bool bIsOrthographicCamera = !CameraViewMatrices.IsPerspectiveProjection();
+
+	int32 CameraViewportWidth = CameraViewRectSize.X;
+	if (bIsOrthographicCamera)
+	{
+		/**
+		 * Orthographic cameras have uniform depth, so basing the LodScale on the width alone can cause issues when selecting the clipmap area to resolve
+		 * at larger scale views. Instead we use the OrthoWidth. This gives a larger area for the shadows to be drawn to and ensures shadows further away/in 
+		 * the corners of the view rect have the correct LOD resolution. We default to the viewport as a minimum.
+		 */	
+		CameraViewportWidth = FMath::Max(FMath::CeilToInt(CameraViewMatrices.GetInvProjectionMatrix().M[0][0] * 2.0f), CameraViewportWidth);
+	}
+	CameraViewportWidth = FMath::Max(MinCameraViewportWidth, CameraViewportWidth);
 
 	// NOTE: Rotational (roll) invariance of the directional light depends on square pixels so we just base everything on the camera X scales/resolution
 	// NOTE: 0.5 because we double the size of the clipmap region below to handle snapping
@@ -142,8 +185,64 @@ FVirtualShadowMapClipmap::FVirtualShadowMapClipmap(
 	// Clamp negative absolute resolution biases as they would exceed the maximum resolution/ranges allocated
 	ResolutionLodBias = FMath::Max(0.0f, ResolutionLodBias);
 
+	WorldOrigin = CameraViewMatrices.GetViewOrigin();
+	CameraToViewTarget = FVector::ZeroVector;
+	if (bIsOrthographicCamera && CVarOrthoVSMProjectViewOrigin.GetValueOnRenderThread())
+	{
+		/**
+		* If enabled, use the ViewTarget location as the WorldOrigin location, this helps with scaling VSMs in Ortho
+		* as the clipmaps emanate more evenly from the focus of the view.
+		* A ViewTarget is not always necessarily present, but there isn't really an alternative way to estimate the best
+		* WorldOrigin for this effect to work well right now without the ViewTarget set.
+		*/
+		CameraToViewTarget = CameraViewMatrices.GetCameraToViewTarget();
+		if (CameraToViewTarget.Length() == 0.0f && CVarOrthoVSMRayCastViewOrigin.GetValueOnRenderThread()
+			&& DependentView && DependentView->Family && DependentView->Family->Scene)
+		{
+			if (UWorld* World = DependentView->Family->Scene->GetWorld())
+			{
+				FVector ViewForward = CameraViewMatrices.GetViewMatrix().GetColumn(2);
+				FCollisionObjectQueryParams ObjectParams = FCollisionObjectQueryParams(FCollisionObjectQueryParams::InitType::AllObjects);
+				FCollisionQueryParams CollisionParams = FCollisionQueryParams(FName(TEXT("OrthoCamera_VSMTrace")), true, DependentView->ViewActor);
+				FHitResult Hit(ForceInit);
+
+				if (World->LineTraceSingleByObjectType(
+					Hit,		//result
+					WorldOrigin,	//start
+					WorldOrigin + ViewForward * FMath::Abs(CameraViewMatrices.GetInvProjectionMatrix().M[2][2]), //end
+					ObjectParams,//collision channel
+					CollisionParams
+				))
+				{
+					CameraToViewTarget = ViewForward * Hit.Distance;
+				}
+			}
+		}
+		WorldOrigin += CameraToViewTarget;
+	}
+
 	FirstLevel = GetFirstLevel();
-	int32 LastLevel = CVarVirtualShadowMapClipmapLastLevel.GetValueOnRenderThread();
+	int32 LastLevel = CVarVirtualShadowMapClipmapLastLevel.GetValueOnRenderThread();	
+	if (bIsOrthographicCamera && CVarOrthoVSMEstimateClipmapLevels.GetValueOnRenderThread())
+	{
+		/**
+		* For Ortho projections, this branch bases the first level VSM on the set OrthoWidth. This reduces the number of clipmaps generated
+		* and also scales the precision of the clipmaps depending on the scene.
+		* 
+		* To be on the safe side, we output -1 FirstLevel compared to what the full OrthoWidth would output. The InvProjectionMatrix outputs
+		* half the OrthoWidth in the [0][0] position, and as we are using Log2, we can just use that raw value, rather than multiplying it then
+		* subtracting a level.
+		*/
+		int32 OrthoFirstLevel = FMath::FloorToInt(FMath::Log2(static_cast<float>(CameraViewMatrices.GetInvProjectionMatrix().M[0][0])));
+		if (OrthoFirstLevel > FirstLevel)
+		{
+			//Only apply the ortho level if it above the desired minimum first level.
+			FirstLevel = OrthoFirstLevel;
+		}
+		//Allow manual correction using the Ortho only FirstLevel bias.
+		FirstLevel = FMath::Max(FirstLevel + CVarOrthoVSMClipmapLODBias.GetValueOnRenderThread(), 0);
+	}
+
 	LastLevel = FMath::Max(FirstLevel, LastLevel);
 	int32 LevelCount = LastLevel - FirstLevel + 1;
 
@@ -153,8 +252,6 @@ FVirtualShadowMapClipmap::FVirtualShadowMapClipmap(
 
 	VirtualShadowMapId = VirtualShadowMapArray.Allocate(false, LevelCount);
 
-	WorldOrigin = CameraViewMatrices.GetLODViewOrigin();
-
 	// TODO: We need a light/cache entry for every light/VSM now, but scene captures may not have persistent view state for indexing
 	// This is likely to all change as we refactor the multiple cache managers stuff anyways; for now this is hopefully safe since they do have separate copies
 	const uint32 UniqueViewKey = InDependentView->ViewState ? InDependentView->ViewState->GetViewKey() : 0U;
@@ -275,7 +372,7 @@ FVirtualShadowMapClipmap::FVirtualShadowMapClipmap(
 		ClipmapLevelEntry->ProjectionData = ComputeProjectionShaderData(Index);
 	}
 
-	ComputeBoundingVolumes(CameraViewMatrices);
+	ComputeBoundingVolumes(WorldOrigin);
 }
 
 const FVirtualShadowMapProjectionShaderData& FVirtualShadowMapClipmap::GetProjectionShaderData(int32 ClipmapIndex) const
@@ -283,7 +380,7 @@ const FVirtualShadowMapProjectionShaderData& FVirtualShadowMapClipmap::GetProjec
 	return PerLightCacheEntry->ShadowMapEntries[ClipmapIndex].ProjectionData;
 }
 
-void FVirtualShadowMapClipmap::ComputeBoundingVolumes(const FViewMatrices& CameraViewMatrices)
+void FVirtualShadowMapClipmap::ComputeBoundingVolumes(const FVector ViewOrigin)
 {
 	// We don't really do much CPU culling with clipmaps. After various testing the fact that we are culling
 	// a single frustum that goes out and basically the entire map, and we have to extrude towards (and away!) from
@@ -291,14 +388,10 @@ void FVirtualShadowMapClipmap::ComputeBoundingVolumes(const FViewMatrices& Camer
 	// into a page that then gets cached with incomplete geometry), in many situations there is effectively no
 	// culling that happens. For instance, as soon as the camera looks vaguely towards or away from the light direction,
 	// the extruded frustum effectively covers the whole world.
-
-	const FVector CameraOrigin = CameraViewMatrices.GetViewOrigin();
-	const FVector CameraDirection = CameraViewMatrices.GetViewMatrix().GetColumn(2);
-
 	// Thus we don't spend a lot of time trying to optimize for the easy cases and instead just pick an extremely
 	// conservative frustum.
 	ViewFrustumBounds = FConvexVolume();
-	BoundingSphere = FSphere(CameraOrigin, GetMaxRadius());
+	BoundingSphere = FSphere(ViewOrigin, GetMaxRadius());
 }
 
 float FVirtualShadowMapClipmap::GetMaxRadius() const
diff --git a/Engine/Source/Runtime/Renderer/Private/VirtualShadowMaps/VirtualShadowMapClipmap.h b/Engine/Source/Runtime/Renderer/Private/VirtualShadowMaps/VirtualShadowMapClipmap.h
--- a/Engine/Source/Runtime/Renderer/Private/VirtualShadowMaps/VirtualShadowMapClipmap.h
+++ b/Engine/Source/Runtime/Renderer/Private/VirtualShadowMaps/VirtualShadowMapClipmap.h
@@ -109,7 +109,7 @@ public:
 
 private:
 	FVirtualShadowMapProjectionShaderData ComputeProjectionShaderData(int32 ClipmapIndex) const;
-	void ComputeBoundingVolumes(const FViewMatrices& CameraViewMatrices);
+	void ComputeBoundingVolumes(const FVector CameraOrigin);
 
 	const FLightSceneInfo& LightSceneInfo;
 
@@ -124,6 +124,7 @@ private:
 	* Note that the centers of each of the levels can be different as they are snapped to page alignment at their respective scales
 	* */
 	FVector WorldOrigin;
+	FVector CameraToViewTarget;
 
 	FVector LightDirection;
 
