Fixed recorded instance invalidations building up over time and retaining old IDs that caused OOB indexing.

diff --git a/Engine/Source/Runtime/Renderer/Private/VirtualShadowMaps/VirtualShadowMapCacheManager.cpp b/Engine/Source/Runtime/Renderer/Private/VirtualShadowMaps/VirtualShadowMapCacheManager.cpp
--- a/Engine/Source/Runtime/Renderer/Private/VirtualShadowMaps/VirtualShadowMapCacheManager.cpp
+++ b/Engine/Source/Runtime/Renderer/Private/VirtualShadowMaps/VirtualShadowMapCacheManager.cpp
@@ -359,10 +359,9 @@ void FVirtualShadowMapArrayCacheManager::FInvalidatingPrimitiveCollector::AddPri
 				}
 			}
 		}
-
-		Manager.ShadowInvalidatingInstancesImplementation.PrimitiveInstancesToInvalidate.Reset();
 		CacheEntry.Value->PrimitiveInstancesToInvalidate.Reset();
 	}
+	Manager.ShadowInvalidatingInstancesImplementation.PrimitiveInstancesToInvalidate.Reset();
 }
 
 void FVirtualShadowMapArrayCacheManager::FInvalidatingPrimitiveCollector::AddInvalidation(FPrimitiveSceneInfo * PrimitiveSceneInfo, bool bRemovedPrimitive)
@@ -1381,6 +1380,15 @@ void FVirtualShadowMapArrayCacheManager::ProcessInvalidations(FRDGBuilder& Graph
 			ProcessInvalidations(GraphBuilder, InvalidationPassCommon, InvalidatingPrimitiveCollector.Instances);
 		}
 	}
+	else
+	{
+		// Clear any queued-up invalidations
+		ShadowInvalidatingInstancesImplementation.PrimitiveInstancesToInvalidate.Reset();
+		for (auto& CacheEntry : CacheEntries)
+		{
+			CacheEntry.Value->PrimitiveInstancesToInvalidate.Reset();
+		}
+	}
 }
 
 void FVirtualShadowMapArrayCacheManager::OnLightRemoved(int32 LightId)
