Pack data in FStaticMeshBatchRelevance from 32 to 24 bytes by storing LODIndex in 5 bits instead of int8

diff --git a/Engine/Source/Runtime/Renderer/Private/StaticMeshBatch.cpp b/Engine/Source/Runtime/Renderer/Private/StaticMeshBatch.cpp
--- a/Engine/Source/Runtime/Renderer/Private/StaticMeshBatch.cpp
+++ b/Engine/Source/Runtime/Renderer/Private/StaticMeshBatch.cpp
@@ -45,7 +45,8 @@ FStaticMeshBatchRelevance::FStaticMeshBatchRelevance(
 	, ScreenSize(InScreenSize)
 	, NumElements(StaticMesh.Elements.Num())
 	, CommandInfosBase(0)
-	, LODIndex(StaticMesh.LODIndex)
+	, UnsignedLODIndex(StaticMesh.LODIndex)
+	, bInvalidLODIndex(StaticMesh.LODIndex < 0)
 	, bDitheredLODTransition(StaticMesh.bDitheredLODTransition)
 	, bSelectable(StaticMesh.bSelectable)
 	, CastShadow(StaticMesh.CastShadow)
@@ -66,6 +67,7 @@ FStaticMeshBatchRelevance::FStaticMeshBatchRelevance(
 	, bUseForLumenSceneCapture(bInUseForLumenSceneCapture)
 	, DecalRenderTargetModeMask(InDecalRenderTargetModeMask)
 {
+	check(GetLODIndex() == StaticMesh.LODIndex);
 }
 
 int32 FStaticMeshBatchRelevance::GetStaticMeshCommandInfoIndex(EMeshPass::Type MeshPass) const
diff --git a/Engine/Source/Runtime/Renderer/Public/StaticMeshBatch.h b/Engine/Source/Runtime/Renderer/Public/StaticMeshBatch.h
--- a/Engine/Source/Runtime/Renderer/Public/StaticMeshBatch.h
+++ b/Engine/Source/Runtime/Renderer/Public/StaticMeshBatch.h
@@ -54,9 +54,7 @@ public:
 		uint8 InDecalRenderTargetModeMask,
 		ERHIFeatureLevel::Type FeatureLevel);
 
-	PRAGMA_DISABLE_DEPRECATION_WARNINGS
-	int8 GetLODIndex() const { return LODIndex; }
-	PRAGMA_ENABLE_DEPRECATION_WARNINGS
+	int8 GetLODIndex() const { return bInvalidLODIndex ? -1 : int8(UnsignedLODIndex); }
 
 	/** Starting offset into continuous array of command infos for this mesh in FPrimitiveSceneInfo::CachedMeshDrawCommandInfos. */
 	FMeshPassMask CommandInfosMask;
@@ -73,9 +71,19 @@ public:
 	/* Every bit corresponds to one MeshPass. If bit is set, then FPrimitiveSceneInfo::CachedMeshDrawCommandInfos contains this mesh pass. */
 	uint16 CommandInfosBase;
 
+private:
+
+	/** LOD index of the mesh, used for fading LOD transitions - unsigned and clamped for packing. */
+	uint8 UnsignedLODIndex : 4;
+
+	/** Original LOD index is negative */
+	uint8 bInvalidLODIndex : 1;
+
+public:
+
 	/** LOD index of the mesh, used for fading LOD transitions. */
-	UE_DEPRECATED(5.4, "Public LODIndex member is deprecated, use GetLODIndex() function instead.")
-	int8 LODIndex;
+	UE_DEPRECATED(5.4, "Public LODIndex member is deprecated and doesn't contain valid data anymore! Use GetLODIndex() function instead.")
+	int8 LODIndex : 1;
 
 	/** Whether the mesh batch should apply dithered LOD. */
 	uint8 bDitheredLODTransition : 1;
