- PVS Studio fix

diff --git a/Engine/Source/Runtime/Renderer/Private/HeterogeneousVolumes/HeterogeneousVolumesVoxelGridPipeline.cpp b/Engine/Source/Runtime/Renderer/Private/HeterogeneousVolumes/HeterogeneousVolumesVoxelGridPipeline.cpp
--- a/Engine/Source/Runtime/Renderer/Private/HeterogeneousVolumes/HeterogeneousVolumesVoxelGridPipeline.cpp
+++ b/Engine/Source/Runtime/Renderer/Private/HeterogeneousVolumes/HeterogeneousVolumesVoxelGridPipeline.cpp
@@ -3352,7 +3352,6 @@ void CombineVolumetricShadowMap(
 )
 {
 	bool bIsMultiProjection = (LightType == LightType_Point) || (LightType == LightType_Rect);
-	GroupCount = FIntVector(1);
 	GroupCount.X = FMath::DivideAndRoundUp(ShadowMapResolution.X, FCombineVolumetricShadowMapsCS::GetThreadGroupSize2D());
 	GroupCount.Y = FMath::DivideAndRoundUp(ShadowMapResolution.Y, FCombineVolumetricShadowMapsCS::GetThreadGroupSize2D());
 	GroupCount.Z = bIsMultiProjection ? 6 : 1;
