Add support for rendering debug primitives on mobile as part of the base pass when MobileHDR=false, they are normally rendered in post processing

diff --git a/Engine/Source/Runtime/Renderer/Private/MobileBasePassRendering.cpp b/Engine/Source/Runtime/Renderer/Private/MobileBasePassRendering.cpp
--- a/Engine/Source/Runtime/Renderer/Private/MobileBasePassRendering.cpp
+++ b/Engine/Source/Runtime/Renderer/Private/MobileBasePassRendering.cpp
@@ -556,3 +556,21 @@ void FMobileSceneRenderer::RenderMobileEditorPrimitives(FRHICommandList& RHICmdL
 		View.TopBatchedViewElements.Draw(RHICmdList, DrawRenderState, FeatureLevel, View, false);
 	}
 }
+
+#if UE_ENABLE_DEBUG_DRAWING
+void FMobileSceneRenderer::RenderMobileDebugPrimitives(FRHICommandList& RHICmdList, const FViewInfo& View)
+{
+	QUICK_SCOPE_CYCLE_COUNTER(STAT_DebugDynamicPrimitiveDrawTime);
+	SCOPED_DRAW_EVENT(RHICmdList, DynamicDebug);
+
+	FMeshPassProcessorRenderState DrawRenderState;
+	DrawRenderState.SetDepthStencilAccess(FExclusiveDepthStencil::DepthWrite_StencilWrite);
+	DrawRenderState.SetBlendState(TStaticBlendState<CW_RGBA, BO_Add, BF_One, BF_One, BO_Add, BF_One, BF_One>::GetRHI());
+		
+	DrawRenderState.SetDepthStencilState(TStaticDepthStencilState<true, CF_DepthNearOrEqual>::GetRHI());
+	View.DebugSimpleElementCollector.DrawBatchedElements(RHICmdList, DrawRenderState, View, EBlendModeFilter::OpaqueAndMasked, SDPG_World);
+
+	DrawRenderState.SetDepthStencilState(TStaticDepthStencilState<true, CF_Always>::GetRHI());
+	View.DebugSimpleElementCollector.DrawBatchedElements(RHICmdList, DrawRenderState, View, EBlendModeFilter::OpaqueAndMasked, SDPG_Foreground);
+}
+#endif
\ No newline at end of file
diff --git a/Engine/Source/Runtime/Renderer/Private/MobileShadingRenderer.cpp b/Engine/Source/Runtime/Renderer/Private/MobileShadingRenderer.cpp
--- a/Engine/Source/Runtime/Renderer/Private/MobileShadingRenderer.cpp
+++ b/Engine/Source/Runtime/Renderer/Private/MobileShadingRenderer.cpp
@@ -1580,6 +1580,14 @@ void FMobileSceneRenderer::RenderForwardSinglePass(FRDGBuilder& GraphBuilder, FM
 		}
 		// Draw translucency.
 		RenderTranslucency(RHICmdList, View);
+		
+#if UE_ENABLE_DEBUG_DRAWING
+		if ((!IsMobileHDR() || bTonemapSubpass) && FSceneRenderer::ShouldCompositeDebugPrimitivesInPostProcess(View))
+		{
+			// Draw debug primitives after translucency for LDR as we do not have a post processing pass
+			RenderMobileDebugPrimitives(RHICmdList, View);
+		}
+#endif
 
 		if (bDoOcclusionQueries)
 		{
diff --git a/Engine/Source/Runtime/Renderer/Private/SceneRendering.h b/Engine/Source/Runtime/Renderer/Private/SceneRendering.h
--- a/Engine/Source/Runtime/Renderer/Private/SceneRendering.h
+++ b/Engine/Source/Runtime/Renderer/Private/SceneRendering.h
@@ -2716,6 +2716,11 @@ protected:
 
 	void RenderMobileEditorPrimitives(FRHICommandList& RHICmdList, const FViewInfo& View, const FMeshPassProcessorRenderState& DrawRenderState, const FInstanceCullingDrawParams* InstanceCullingDrawParams);
 
+#if UE_ENABLE_DEBUG_DRAWING
+	/** Render debug primitives in the base pass, for MobileHDR=false */
+	void RenderMobileDebugPrimitives(FRHICommandList& RHICmdList, const FViewInfo& View);
+#endif
+
 	/** Renders the debug view pass for mobile. */
 	void RenderMobileDebugView(FRHICommandList& RHICmdList, const FViewInfo& View);
 
