Remove old non cached MDC decal rendering code

diff --git a/Engine/Source/Runtime/Renderer/Private/CompositionLighting/PostProcessDeferredDecals.cpp b/Engine/Source/Runtime/Renderer/Private/CompositionLighting/PostProcessDeferredDecals.cpp
--- a/Engine/Source/Runtime/Renderer/Private/CompositionLighting/PostProcessDeferredDecals.cpp
+++ b/Engine/Source/Runtime/Renderer/Private/CompositionLighting/PostProcessDeferredDecals.cpp
@@ -587,7 +587,6 @@ void AddDeferredDecalPass(
 	const FScene& Scene = *(FScene*)ViewFamily.Scene;
 	const EShaderPlatform ShaderPlatform = View.GetShaderPlatform();
 	const ERHIFeatureLevel::Type FeatureLevel = View.GetFeatureLevel();
-	const uint32 MeshDecalCount = View.MeshDecalBatches.Num();
 	const uint32 DecalCount = Scene.Decals.Num();
 	uint32 SortedDecalCount = 0;
 	FTransientDecalRenderDataList* SortedDecals = nullptr;
@@ -605,7 +604,7 @@ void AddDeferredDecalPass(
 	}
 
 	const bool bHasAnyDrawCommandDecalCount = HasAnyDrawCommandDecalCount(DecalRenderStage, View);
-	const bool bVisibleDecalsInView = MeshDecalCount > 0 || SortedDecalCount > 0 || bHasAnyDrawCommandDecalCount;
+	const bool bVisibleDecalsInView = SortedDecalCount > 0 || bHasAnyDrawCommandDecalCount;
 	const bool bShaderComplexity = View.Family->EngineShowFlags.ShaderComplexity;
 	const bool bStencilSizeThreshold = CVarStencilSizeThreshold.GetValueOnRenderThread() >= 0;
 
@@ -695,7 +694,7 @@ void AddDeferredDecalPass(
 	{
 		RDG_EVENT_SCOPE(GraphBuilder, "DeferredDecals %s", GetStageName(DecalRenderStage));
 
-		if ((MeshDecalCount > 0 || bHasAnyDrawCommandDecalCount) && (DecalRenderStage == EDecalRenderStage::BeforeBasePass || DecalRenderStage == EDecalRenderStage::BeforeLighting || DecalRenderStage == EDecalRenderStage::Emissive || DecalRenderStage == EDecalRenderStage::AmbientOcclusion))
+		if (bHasAnyDrawCommandDecalCount && (DecalRenderStage == EDecalRenderStage::BeforeBasePass || DecalRenderStage == EDecalRenderStage::BeforeLighting || DecalRenderStage == EDecalRenderStage::Emissive || DecalRenderStage == EDecalRenderStage::AmbientOcclusion))
 		{
 			RenderMeshDecals(GraphBuilder, Scene, View, PassTextures, InstanceCullingManager, DecalRenderStage);
 		}
diff --git a/Engine/Source/Runtime/Renderer/Private/CompositionLighting/PostProcessDeferredDecals.h b/Engine/Source/Runtime/Renderer/Private/CompositionLighting/PostProcessDeferredDecals.h
--- a/Engine/Source/Runtime/Renderer/Private/CompositionLighting/PostProcessDeferredDecals.h
+++ b/Engine/Source/Runtime/Renderer/Private/CompositionLighting/PostProcessDeferredDecals.h
@@ -107,8 +107,6 @@ bool HasAnyDrawCommandDecalCount(
 	EDecalRenderStage DecalRenderStage, 
 	FViewInfo& View);
 
-bool UseDecalCachedMDCs();
-
 void ExtractNormalsForNextFrameReprojection(
 	FRDGBuilder& GraphBuilder,
 	const FSceneTextures& SceneTextures,
diff --git a/Engine/Source/Runtime/Renderer/Private/CompositionLighting/PostProcessMeshDecals.cpp b/Engine/Source/Runtime/Renderer/Private/CompositionLighting/PostProcessMeshDecals.cpp
--- a/Engine/Source/Runtime/Renderer/Private/CompositionLighting/PostProcessMeshDecals.cpp
+++ b/Engine/Source/Runtime/Renderer/Private/CompositionLighting/PostProcessMeshDecals.cpp
@@ -30,22 +30,6 @@ static TAutoConsoleVariable<int32> CVarParallelMeshDecal(
 	ECVF_RenderThreadSafe
 );
 
-static TAutoConsoleVariable<int32> CVarDecalUseCachedMDCs(
-	TEXT("r.Decal.UseCachedMDCs"),
-	1,
-	TEXT("Use cached MDCs for decal rendering."),
-	FConsoleVariableDelegate::CreateLambda([](IConsoleVariable* InVariable)
-	{
-		FGlobalComponentRecreateRenderStateContext Context;
-	}),
-	ECVF_RenderThreadSafe
-);
-
-bool UseDecalCachedMDCs()
-{
-	return CVarDecalUseCachedMDCs.GetValueOnRenderThread() != 0;
-}
-
 class FMeshDecalsVS : public FMeshMaterialShader
 {
 public:
@@ -217,12 +201,6 @@ FMeshDecalMeshProcessor::FMeshDecalMeshProcessor(const FScene* Scene,
 
 void FMeshDecalMeshProcessor::AddMeshBatch(const FMeshBatch& RESTRICT MeshBatch, uint64 BatchElementMask, const FPrimitiveSceneProxy* RESTRICT PrimitiveSceneProxy, int32 StaticMeshId)
 {
-	// Early out if not using cached MDCs and there is no valid ViewIfDynamicMeshCommand (means we are caching the static MDCs)
-	if (CVarDecalUseCachedMDCs.GetValueOnRenderThread() == 0 && ViewIfDynamicMeshCommand == nullptr)
-	{
-		return;
-	}
-
 	if (MeshBatch.bUseForMaterial && MeshBatch.IsDecal(FeatureLevel))
 	{
 		const FMaterialRenderProxy* MaterialRenderProxy = MeshBatch.MaterialRenderProxy;
@@ -514,43 +492,6 @@ void DrawDecalMeshCommands(
 	auto* PassParameters = GraphBuilder.AllocParameters<FDeferredDecalPassParameters>();
 	GetDeferredDecalPassParameters(GraphBuilder, View, DecalPassTextures, RenderTargetMode, *PassParameters);
 
-	// If decal cached MDCs is enabled then mesh decal batches will be empty
-	// Cached decals are rendered below using the parallel mesh draw command passes
-	if (View.MeshDecalBatches.Num() > 0)
-	{
-		AddSimpleMeshPass(
-			GraphBuilder,
-			PassParameters,
-			&Scene,
-			View,
-			&InstanceCullingManager,
-			RDG_EVENT_NAME("MeshDecals"),
-			View.ViewRect,
-			[&View, DecalRenderStage, RenderTargetMode](FDynamicPassMeshDrawListContext* DynamicMeshPassContext)
-			{
-				TRACE_CPUPROFILER_EVENT_SCOPE(MeshDecalCommands);
-
-				check(DecalRenderStage == DecalRendering::GetRenderStage(RenderTargetMode, EShadingPath::Deferred));
-
-				FMeshDecalMeshProcessor PassMeshProcessor(
-					View.Family->Scene->GetRenderScene(),
-					View.GetFeatureLevel(),
-					&View,
-					RenderTargetMode,
-					EShadingPath::Deferred,
-					DynamicMeshPassContext);
-
-				for (int32 MeshBatchIndex = 0; MeshBatchIndex < View.MeshDecalBatches.Num(); ++MeshBatchIndex)
-				{
-					const FMeshBatch* Mesh = View.MeshDecalBatches[MeshBatchIndex].Mesh;
-					const FPrimitiveSceneProxy* PrimitiveSceneProxy = View.MeshDecalBatches[MeshBatchIndex].Proxy;
-					const uint64 DefaultBatchElementMask = ~0ull;
-
-					PassMeshProcessor.AddMeshBatch(*Mesh, DefaultBatchElementMask, PrimitiveSceneProxy);
-				}
-			});
-	}
-
 	EMeshPass::Type DecalMeshPassType = DecalRendering::GetMeshPassType(RenderTargetMode);
 	if (View.ParallelMeshDrawCommandPasses[DecalMeshPassType].HasAnyDraw())
 	{
@@ -665,32 +606,3 @@ void RenderMeshDecals(
 		break;
 	}
 }
-
-void RenderMeshDecalsMobile(FRHICommandList& RHICmdList, FViewInfo& View, EDecalRenderStage DecalRenderStage, EDecalRenderTargetMode RenderTargetMode)
-{
-	FGraphicsPipelineStateInitializer GraphicsPSOInit;
-	RHICmdList.SetViewport(View.ViewRect.Min.X, View.ViewRect.Min.Y, 0, View.ViewRect.Max.X, View.ViewRect.Max.Y, 1);
-	RHICmdList.ApplyCachedRenderTargets(GraphicsPSOInit);
-
-	DrawDynamicMeshPass(View, RHICmdList, [&View, DecalRenderStage, RenderTargetMode](FDynamicPassMeshDrawListContext* DynamicMeshPassContext)
-	{
-		check(DecalRenderStage == DecalRendering::GetRenderStage(RenderTargetMode, EShadingPath::Mobile));
-
-		FMeshDecalMeshProcessor PassMeshProcessor(
-			View.Family->Scene->GetRenderScene(),
-			View.GetFeatureLevel(),
-			&View,
-			RenderTargetMode,
-			EShadingPath::Mobile,
-			DynamicMeshPassContext);
-
-		for (int32 MeshBatchIndex = 0; MeshBatchIndex < View.MeshDecalBatches.Num(); ++MeshBatchIndex)
-		{
-			const FMeshBatch* Mesh = View.MeshDecalBatches[MeshBatchIndex].Mesh;
-			const FPrimitiveSceneProxy* PrimitiveSceneProxy = View.MeshDecalBatches[MeshBatchIndex].Proxy;
-			const uint64 DefaultBatchElementMask = ~0ull;
-
-			PassMeshProcessor.AddMeshBatch(*Mesh, DefaultBatchElementMask, PrimitiveSceneProxy);
-		}
-	}, true);
-}
\ No newline at end of file
diff --git a/Engine/Source/Runtime/Renderer/Private/MobileDecalRendering.cpp b/Engine/Source/Runtime/Renderer/Private/MobileDecalRendering.cpp
--- a/Engine/Source/Runtime/Renderer/Private/MobileDecalRendering.cpp
+++ b/Engine/Source/Runtime/Renderer/Private/MobileDecalRendering.cpp
@@ -20,7 +20,6 @@
 #include "CompositionLighting/PostProcessDeferredDecals.h"
 #include "DBufferTextures.h"
 
-void RenderMeshDecalsMobile(FRHICommandList& RHICmdList, FViewInfo& View, EDecalRenderStage DecalRenderStage, EDecalRenderTargetMode RenderTargetMode);
 extern void RenderDeferredDecalsMobile(FRHICommandList& RHICmdList, const FScene& Scene, const FViewInfo& View, EDecalRenderStage DecalRenderStage, EDecalRenderTargetMode RenderTargetMode);
 
 static bool DoesPlatformSupportDecals(EShaderPlatform ShaderPlatform)
@@ -65,22 +64,6 @@ void FMobileSceneRenderer::RenderDecals(FRHICommandList& RHICmdList, FViewInfo&
 		RenderDeferredDecalsMobile(RHICmdList, *Scene, View, DecalRenderStage, RenderTargetMode);
 	}
 
-	// Mesh decals
-	if (View.MeshDecalBatches.Num() > 0)
-	{
-		SCOPED_DRAW_EVENT(RHICmdList, MeshDecals);
-		RenderMeshDecalsMobile(RHICmdList, View, DecalRenderStage, RenderTargetMode);
-
-		// MeshDecals use DrawDynamicMeshPass which may change BatchedPrimitive binding, so we need to restore it
-		FUniformBufferStaticSlot BatchedPrimitiveSlot = FInstanceCullingContext::GetUniformBufferViewStaticSlot(View.GetShaderPlatform());
-		if (IsUniformBufferStaticSlotValid(BatchedPrimitiveSlot) && InstanceCullingDrawParams->BatchedPrimitive)
-		{
-			FRHIUniformBuffer* BatchedPrimitiveBufferRHI = InstanceCullingDrawParams->BatchedPrimitive.GetUniformBuffer()->GetRHI();
-			check(BatchedPrimitiveBufferRHI);
-			RHICmdList.SetStaticUniformBuffer(BatchedPrimitiveSlot, BatchedPrimitiveBufferRHI);
-		}
-	}
-
 	EMeshPass::Type DecalMeshPassType = DecalRendering::GetMeshPassType(RenderTargetMode);
 	if (View.ParallelMeshDrawCommandPasses[DecalMeshPassType].HasAnyDraw())
 	{
diff --git a/Engine/Source/Runtime/Renderer/Private/SceneRendering.h b/Engine/Source/Runtime/Renderer/Private/SceneRendering.h
--- a/Engine/Source/Runtime/Renderer/Private/SceneRendering.h
+++ b/Engine/Source/Runtime/Renderer/Private/SceneRendering.h
@@ -872,18 +872,6 @@ struct FSortedTrianglesMeshBatch
 	const FPrimitiveSceneProxy* Proxy = nullptr;
 };
 
-struct FMeshDecalBatch
-{
-	const FMeshBatch* Mesh;
-	const FPrimitiveSceneProxy* Proxy;
-	int16 SortKey;
-
-	FORCEINLINE bool operator<(const FMeshDecalBatch& rhs) const
-	{
-		return SortKey < rhs.SortKey;
-	}
-};
-
 // DX11 maximum 2d texture array size is D3D11_REQ_TEXTURE2D_ARRAY_AXIS_DIMENSION = 2048, and 2048/6 = 341.33.
 UE_DEPRECATED(5.4, "Using GMaxNumReflectionCaptures directly is deprecated. Please use GetMaxNumReflectionCaptures(EShaderPlatform) instead")
 static const int32 GMaxNumReflectionCaptures = 341;
@@ -1302,9 +1290,6 @@ public:
 	/** Get the GPU Scene instance ranges of visible Nanite primitives writing custom depth. */
 	TArray<FPrimitiveInstanceRange, SceneRenderingAllocator> NaniteCustomDepthInstances;
 
-	/** Mesh batches with for mesh decal rendering. */
-	TArray<FMeshDecalBatch, SceneRenderingAllocator> MeshDecalBatches;
-
 	/** Mesh batches with a volumetric material. */
 	TArray<FVolumetricMeshBatch, SceneRenderingAllocator> VolumetricMeshBatches;
 
diff --git a/Engine/Source/Runtime/Renderer/Private/SceneVisibility.cpp b/Engine/Source/Runtime/Renderer/Private/SceneVisibility.cpp
--- a/Engine/Source/Runtime/Renderer/Private/SceneVisibility.cpp
+++ b/Engine/Source/Runtime/Renderer/Private/SceneVisibility.cpp
@@ -1229,7 +1229,6 @@ void FRelevancePacket::Finalize()
 	WriteView.SubstrateViewData.bUsesComplexSpecialRenderPath |= bUsesComplexSpecialRenderPath;
 	DirtyIndirectLightingCacheBufferPrimitives.AppendTo(WriteView.DirtyIndirectLightingCacheBufferPrimitives);
 
-	WriteView.MeshDecalBatches.Append(MeshDecalBatches);
 	WriteView.VolumetricMeshBatches.Append(VolumetricMeshBatches);
 	WriteView.HeterogeneousVolumesMeshBatches.Append(HeterogeneousVolumesMeshBatches);
 	WriteView.SkyMeshBatches.Append(SkyMeshBatches);
@@ -1319,8 +1318,7 @@ void FRelevancePacket::ComputeRelevance(FDynamicPrimitiveIndexList& DynamicPrimi
 	const FHLODVisibilityState* const HLODState = bHLODActive && ViewState ? &ViewState->HLODVisibilityState : nullptr;
 	float MaxDrawDistanceScale = GetCachedScalabilityCVars().ViewDistanceScale;
 	MaxDrawDistanceScale *= GetCachedScalabilityCVars().CalculateFieldOfViewDistanceScale(View.DesiredFOV);
-	const bool bUseDecalCachedMDCs = UseDecalCachedMDCs();
-
+	
 	const auto AddEditorDynamicPrimitive = [this, &DynamicPrimitiveIndexList](int32 PrimitiveIndex)
 	{
 #if WITH_EDITOR
@@ -1730,25 +1728,14 @@ void FRelevancePacket::ComputeRelevance(FDynamicPrimitiveIndexList& DynamicPrimi
 						// Because ViewRelevance is a sum of all material relevances in the primitive
 						if (ViewRelevance.bRenderInMainPass && ViewRelevance.bDecal && StaticMeshRelevance.bUseForMaterial)
 						{
-							if (bUseDecalCachedMDCs)
+							for (uint8 DecalRenderTargetMode = 0; DecalRenderTargetMode < (uint8)EDecalRenderTargetMode::Num; ++DecalRenderTargetMode)
 							{
-								for (uint8 DecalRenderTargetMode = 0; DecalRenderTargetMode < (uint8)EDecalRenderTargetMode::Num; ++DecalRenderTargetMode)
+								if (DecalRendering::IsCompatibleWithRenderTargetMode(StaticMeshRelevance.DecalRenderTargetModeMask, (EDecalRenderTargetMode)DecalRenderTargetMode))
 								{
-									if (DecalRendering::IsCompatibleWithRenderTargetMode(StaticMeshRelevance.DecalRenderTargetModeMask, (EDecalRenderTargetMode)DecalRenderTargetMode))
-									{
-										EMeshPass::Type DecalMeshPassType = DecalRendering::GetMeshPassType((EDecalRenderTargetMode)DecalRenderTargetMode);
-										DrawCommandPacket.AddCommandsForMesh(PrimitiveIndex, PrimitiveSceneInfo, StaticMeshRelevance, StaticMesh, CullingPayloadFlags, Scene, bCanCache, DecalMeshPassType);
-									}
+									EMeshPass::Type DecalMeshPassType = DecalRendering::GetMeshPassType((EDecalRenderTargetMode)DecalRenderTargetMode);
+									DrawCommandPacket.AddCommandsForMesh(PrimitiveIndex, PrimitiveSceneInfo, StaticMeshRelevance, StaticMesh, CullingPayloadFlags, Scene, bCanCache, DecalMeshPassType);
 								}
 							}
-							else
-							{
-								MeshDecalBatches.AddUninitialized(1);
-								FMeshDecalBatch& BatchAndProxy = MeshDecalBatches.Last();
-								BatchAndProxy.Mesh = &StaticMesh;
-								BatchAndProxy.Proxy = PrimitiveSceneProxy;
-								BatchAndProxy.SortKey = PrimitiveSceneProxy->GetTranslucencySortPriority();
-							}
 						}
 					}
 
@@ -2413,29 +2400,18 @@ static void ComputeDynamicMeshRelevance(
 
 	if (ViewRelevance.bRenderInMainPass && ViewRelevance.bDecal)
 	{
-		if (UseDecalCachedMDCs())
+		// DecalRenderTargetModeMask could be cached in FMeshBatchAndRelevance as well
+		const FMaterial& Material = MeshBatch.Mesh->MaterialRenderProxy->GetIncompleteMaterialWithFallback(View.FeatureLevel);
+		uint8 DecalRenderTargetModeMask = DecalRendering::GetDecalRenderTargetModeMask(Material, View.FeatureLevel);
+		for (uint8 DecalRenderTargetMode = 0; DecalRenderTargetMode < (uint8)EDecalRenderTargetMode::Num; ++DecalRenderTargetMode)
 		{
-			// DecalRenderTargetModeMask could be cached in FMeshBatchAndRelevance as well
-			const FMaterial& Material = MeshBatch.Mesh->MaterialRenderProxy->GetIncompleteMaterialWithFallback(View.FeatureLevel);
-			uint8 DecalRenderTargetModeMask = DecalRendering::GetDecalRenderTargetModeMask(Material, View.FeatureLevel);
-			for (uint8 DecalRenderTargetMode = 0; DecalRenderTargetMode < (uint8)EDecalRenderTargetMode::Num; ++DecalRenderTargetMode)
+			if (DecalRendering::IsCompatibleWithRenderTargetMode(DecalRenderTargetModeMask, (EDecalRenderTargetMode)DecalRenderTargetMode))
 			{
-				if (DecalRendering::IsCompatibleWithRenderTargetMode(DecalRenderTargetModeMask, (EDecalRenderTargetMode)DecalRenderTargetMode))
-				{
-					EMeshPass::Type DecalMeshPassType = DecalRendering::GetMeshPassType((EDecalRenderTargetMode)DecalRenderTargetMode);
-					PassMask.Set(DecalMeshPassType);
-					View.NumVisibleDynamicMeshElements[DecalMeshPassType] += NumElements;
-				}
+				EMeshPass::Type DecalMeshPassType = DecalRendering::GetMeshPassType((EDecalRenderTargetMode)DecalRenderTargetMode);
+				PassMask.Set(DecalMeshPassType);
+				View.NumVisibleDynamicMeshElements[DecalMeshPassType] += NumElements;
 			}
 		}
-		else
-		{
-			View.MeshDecalBatches.AddUninitialized(1);
-			FMeshDecalBatch& BatchAndProxy = View.MeshDecalBatches.Last();
-			BatchAndProxy.Mesh = MeshBatch.Mesh;
-			BatchAndProxy.Proxy = MeshBatch.PrimitiveSceneProxy;
-			BatchAndProxy.SortKey = MeshBatch.PrimitiveSceneProxy->GetTranslucencySortPriority();
-		}
 	}
 
 	const bool bIsHairStrandsCompatible = ViewRelevance.bHairStrands && IsHairStrandsEnabled(EHairStrandsShaderType::All, View.GetShaderPlatform());
@@ -4468,7 +4444,6 @@ void FVisibilityTaskData::SetupMeshPasses(FExclusiveDepthStencil::Type BasePassD
 	for (int32 ViewIndex = 0; ViewIndex < Views.Num(); ViewIndex++)
 	{
 		FViewInfo& View = *Views[ViewIndex];
-		View.MeshDecalBatches.Sort();
 
 	#if WITH_EDITOR
 		{
diff --git a/Engine/Source/Runtime/Renderer/Private/SceneVisibilityPrivate.h b/Engine/Source/Runtime/Renderer/Private/SceneVisibilityPrivate.h
--- a/Engine/Source/Runtime/Renderer/Private/SceneVisibilityPrivate.h
+++ b/Engine/Source/Runtime/Renderer/Private/SceneVisibilityPrivate.h
@@ -746,7 +746,6 @@ private:
 	TArray<uint32> EditorSelectedNaniteHitProxyIds;
 #endif
 
-	TArray<FMeshDecalBatch, SceneRenderingAllocator> MeshDecalBatches;
 	TArray<FVolumetricMeshBatch, SceneRenderingAllocator> VolumetricMeshBatches;
 	TArray<FVolumetricMeshBatch, SceneRenderingAllocator> HeterogeneousVolumesMeshBatches;
 	TArray<FSkyMeshBatch, SceneRenderingAllocator> SkyMeshBatches;
