Heterogeneous Volumes: Accelerating rendering with continuous level-of-detail by applying LODFactor to voxel factors, lighting cache, and volumetric shadow map resolutions.

r.HeterogeneousVolumes.CLOD (Default = 1)
r.HeterogeneousVolumes.CLOD.Bias (Default = 0)

diff --git a/Engine/Source/Runtime/Engine/Private/Components/HeterogeneousVolumeComponent.cpp b/Engine/Source/Runtime/Engine/Private/Components/HeterogeneousVolumeComponent.cpp
--- a/Engine/Source/Runtime/Engine/Private/Components/HeterogeneousVolumeComponent.cpp
+++ b/Engine/Source/Runtime/Engine/Private/Components/HeterogeneousVolumeComponent.cpp
@@ -97,6 +97,7 @@ FHeterogeneousVolumeSceneProxy::FHeterogeneousVolumeSceneProxy(UHeterogeneousVol
 	HeterogeneousVolumeData.ShadowStepFactor = InComponent->ShadowStepFactor;
 	HeterogeneousVolumeData.ShadowBiasFactor = InComponent->ShadowBiasFactor;
 	HeterogeneousVolumeData.LightingDownsampleFactor = InComponent->LightingDownsampleFactor;
+	HeterogeneousVolumeData.MipBias = InComponent->StreamingMipBias;
 	HeterogeneousVolumeData.bPivotAtCentroid = InComponent->bPivotAtCentroid;
 
 	// Initialize vertex buffer data for a quad
diff --git a/Engine/Source/Runtime/Engine/Public/HeterogeneousVolumeInterface.h b/Engine/Source/Runtime/Engine/Public/HeterogeneousVolumeInterface.h
--- a/Engine/Source/Runtime/Engine/Public/HeterogeneousVolumeInterface.h
+++ b/Engine/Source/Runtime/Engine/Public/HeterogeneousVolumeInterface.h
@@ -32,6 +32,7 @@ public:
 	virtual float GetShadowStepFactor() const = 0;
 	virtual float GetShadowBiasFactor() const = 0;
 	virtual float GetLightingDownsampleFactor() const = 0;
+	virtual float GetMipBias() const = 0;
 
 	// Debug
 	virtual FString GetReadableName() const = 0;
@@ -51,6 +52,7 @@ public:
 		, ShadowStepFactor(8.0)
 		, ShadowBiasFactor(0.0)
 		, LightingDownsampleFactor(1.0)
+		, MipBias(0.0)
 	{}
 
 	FHeterogeneousVolumeData(const FPrimitiveSceneProxy* SceneProxy, FString Name)
@@ -62,6 +64,7 @@ public:
 		, ShadowStepFactor(8.0)
 		, ShadowBiasFactor(0.0)
 		, LightingDownsampleFactor(1.0)
+		, MipBias(0.0)
 		, bPivotAtCentroid(false)
 #if ACTOR_HAS_LABELS
 		, ReadableName(Name)
@@ -88,6 +91,7 @@ public:
 	virtual float GetShadowStepFactor() const { return ShadowStepFactor; }
 	virtual float GetShadowBiasFactor() const { return ShadowBiasFactor; }
 	virtual float GetLightingDownsampleFactor() const { return LightingDownsampleFactor; }
+	virtual float GetMipBias() const { return MipBias; }
 
 	const FPrimitiveSceneProxy* PrimitiveSceneProxy;
 	FMatrix InstanceToLocal;
@@ -97,6 +101,7 @@ public:
 	float ShadowStepFactor;
 	float ShadowBiasFactor;
 	float LightingDownsampleFactor;
+	float MipBias;
 	bool bPivotAtCentroid;
 
 #if ACTOR_HAS_LABELS
diff --git a/Engine/Source/Runtime/Renderer/Private/HeterogeneousVolumes/HeterogeneousVolumes.cpp b/Engine/Source/Runtime/Renderer/Private/HeterogeneousVolumes/HeterogeneousVolumes.cpp
--- a/Engine/Source/Runtime/Renderer/Private/HeterogeneousVolumes/HeterogeneousVolumes.cpp
+++ b/Engine/Source/Runtime/Renderer/Private/HeterogeneousVolumes/HeterogeneousVolumes.cpp
@@ -246,6 +246,22 @@ static TAutoConsoleVariable<int32> CVarHeterogeneousVolumesVelocity(
 	ECVF_RenderThreadSafe
 );
 
+static TAutoConsoleVariable<int32> CVarHeterogeneousVolumesCLOD(
+	TEXT("r.HeterogeneousVolumes.CLOD"),
+	1,
+	TEXT("Uses Continuous Level-of-Detail to accelerate rendering (Default = 1)"),
+	ECVF_RenderThreadSafe
+);
+
+static TAutoConsoleVariable<float> CVarHeterogeneousVolumesCLODBias(
+	TEXT("r.HeterogeneousVolumes.CLOD.Bias"),
+	0.0,
+	TEXT("Biases evaluation result when computing Continuous Level-of-Detail (Default = 0.0)\n")
+	TEXT("> 0: Coarser\n")
+	TEXT("< 0: Sharper\n"),
+	ECVF_RenderThreadSafe
+);
+
 DECLARE_GPU_STAT_NAMED(HeterogeneousVolumeShadowsStat, TEXT("HeterogeneousVolumeShadows"));
 DECLARE_GPU_STAT_NAMED(HeterogeneousVolumesStat, TEXT("HeterogeneousVolumes"));
 
@@ -488,6 +504,16 @@ namespace HeterogeneousVolumes
 		return CVarHeterogeneousVolumesVelocity.GetValueOnRenderThread() != 0;
 	}
 
+	bool UseContinuousLOD()
+	{
+		return CVarHeterogeneousVolumesCLOD.GetValueOnRenderThread() != 0;
+	}
+
+	float GetCLODBias()
+	{
+		return CVarHeterogeneousVolumesCLODBias.GetValueOnRenderThread();
+	}
+
 	// Convenience Utils
 	int GetVoxelCount(FIntVector VolumeResolution)
 	{
@@ -507,10 +533,57 @@ namespace HeterogeneousVolumes
 			FMath::Max(VolumeResolution.Z >> MipLevel, 1)
 		);
 	}
-	FIntVector GetLightingCacheResolution(const IHeterogeneousVolumeInterface* RenderInterface)
+	
+	float CalcLOD(const FSceneView& View, const IHeterogeneousVolumeInterface* HeterogeneousVolume)
+	{
+		if (!HeterogeneousVolumes::UseContinuousLOD())
+		{
+			return 0.0f;
+		}
+
+		FBoxSphereBounds WorldBounds = HeterogeneousVolume->GetBounds();
+		FIntVector VoxelResolution = HeterogeneousVolume->GetVoxelResolution();
+		float VoxelResolutionMin = VoxelResolution.GetMin();
+
+		float LODValue = FMath::Floor(FMath::Log2(VoxelResolutionMin));
+		if (View.ViewFrustum.IntersectBox(WorldBounds.Origin, WorldBounds.BoxExtent))
+		{
+			// Determine the pixel-width at the near-plane
+			float TanHalfFOV = FMath::Tan(FMath::DegreesToRadians(View.FOV * 0.5));
+			float HalfWidth = View.UnconstrainedViewRect.Width() * 0.5;
+			float PixelWidth = TanHalfFOV / HalfWidth;
+
+			// Project to nearest distance of volume bounds
+			FVector WorldCameraOrigin = View.ViewMatrices.GetViewOrigin();
+			float Distance = FMath::Max((WorldBounds.Origin - WorldCameraOrigin).Length() - WorldBounds.SphereRadius, View.NearClippingDistance);
+			float VoxelWidth = Distance * PixelWidth;
+
+			// MIP is defined as the log of the ratio of native voxel resolution to pixel-coverage of volume bounds
+			//float PixelWidthCoverage = (2.0 * WorldBounds.SphereRadius) / VoxelWidth;
+			float PixelWidthCoverage = (2.0 * WorldBounds.BoxExtent.GetMax()) / VoxelWidth;
+			float ViewLODValue = FMath::Log2(VoxelResolutionMin / PixelWidthCoverage) + HeterogeneousVolume->GetMipBias() + HeterogeneousVolumes::GetCLODBias();
+			ViewLODValue = FMath::Max(ViewLODValue, 0);
+
+			LODValue = FMath::Min(ViewLODValue, LODValue);
+		}
+
+		return LODValue;
+	}
+
+	float CalcLODFactor(float LODValue)
+	{
+		return FMath::Pow(2, LODValue);
+	}
+
+	float CalcLODFactor(const FSceneView& View, const IHeterogeneousVolumeInterface* HeterogeneousVolume)
+	{
+		return CalcLODFactor(CalcLOD(View, HeterogeneousVolume));
+	}
+	
+	FIntVector GetLightingCacheResolution(const IHeterogeneousVolumeInterface* RenderInterface, float LODFactor)
 	{
 		float OverrideDownsampleFactor = CVarHeterogeneousVolumesLightingCacheDownsampleFactor.GetValueOnRenderThread();
-		float DownsampleFactor = OverrideDownsampleFactor > 0.0 ? OverrideDownsampleFactor : RenderInterface->GetLightingDownsampleFactor();
+		float DownsampleFactor = OverrideDownsampleFactor > 0.0 ? OverrideDownsampleFactor : RenderInterface->GetLightingDownsampleFactor() * LODFactor;
 		DownsampleFactor = FMath::Max(DownsampleFactor, 0.125);
 
 		FVector VolumeResolution = FVector(GetVolumeResolution(RenderInterface));
@@ -791,7 +864,8 @@ void FDeferredShadingSceneRenderer::RenderHeterogeneousVolumes(
 					if (HeterogeneousVolumes::GetLightingCacheMode() != 0)
 					{
 						// TODO: Allow option for scalar transmittance to conserve bandwidth
-						FIntVector LightingCacheResolution = HeterogeneousVolumes::GetLightingCacheResolution(HeterogeneousVolume);
+						float LODFactor = HeterogeneousVolumes::CalcLODFactor(View, HeterogeneousVolume);
+						FIntVector LightingCacheResolution = HeterogeneousVolumes::GetLightingCacheResolution(HeterogeneousVolume, LODFactor);
 						uint32 NumMips = FMath::Log2(float(FMath::Min(FMath::Min(LightingCacheResolution.X, LightingCacheResolution.Y), LightingCacheResolution.Z))) + 1;
 						FRDGTextureDesc LightingCacheDesc = FRDGTextureDesc::Create3D(
 							LightingCacheResolution,
diff --git a/Engine/Source/Runtime/Renderer/Private/HeterogeneousVolumes/HeterogeneousVolumes.h b/Engine/Source/Runtime/Renderer/Private/HeterogeneousVolumes/HeterogeneousVolumes.h
--- a/Engine/Source/Runtime/Renderer/Private/HeterogeneousVolumes/HeterogeneousVolumes.h
+++ b/Engine/Source/Runtime/Renderer/Private/HeterogeneousVolumes/HeterogeneousVolumes.h
@@ -59,7 +59,7 @@ namespace HeterogeneousVolumes
 {
 	// CVars
 	FIntVector GetVolumeResolution(const IHeterogeneousVolumeInterface*);
-	FIntVector GetLightingCacheResolution(const IHeterogeneousVolumeInterface*);
+	FIntVector GetLightingCacheResolution(const IHeterogeneousVolumeInterface*, float LODFactor);
 
 	float GetShadowStepSize();
 	float GetMaxTraceDistance();
@@ -122,6 +122,9 @@ namespace HeterogeneousVolumes
 	int GetVoxelCount(FIntVector VolumeResolution);
 	int GetVoxelCount(const FRDGTextureDesc& TextureDesc);
 	FIntVector GetMipVolumeResolution(FIntVector VolumeResolution, uint32 MipLevel);
+	float CalcLOD(const FSceneView& View, const IHeterogeneousVolumeInterface* HeterogeneousVolume);
+	float CalcLODFactor(const FSceneView& View, const IHeterogeneousVolumeInterface* HeterogeneousVolume);
+	float CalcLODFactor(float LOD);
 
 	const FProjectedShadowInfo* GetProjectedShadowInfo(const FVisibleLightInfo* VisibleLightInfo, int32 ShadowIndex);
 	bool IsDynamicShadow(const FVisibleLightInfo* VisibleLightInfo);
diff --git a/Engine/Source/Runtime/Renderer/Private/HeterogeneousVolumes/HeterogeneousVolumesHardwareRayTracing.cpp b/Engine/Source/Runtime/Renderer/Private/HeterogeneousVolumes/HeterogeneousVolumesHardwareRayTracing.cpp
--- a/Engine/Source/Runtime/Renderer/Private/HeterogeneousVolumes/HeterogeneousVolumesHardwareRayTracing.cpp
+++ b/Engine/Source/Runtime/Renderer/Private/HeterogeneousVolumes/HeterogeneousVolumesHardwareRayTracing.cpp
@@ -612,7 +612,8 @@ void RenderLightingCacheWithPreshadingHardwareRayTracing(
 		PassParameters->SparseVoxelUniformBuffer = SparseVoxelUniformBuffer;
 
 		// Transmittance volume
-		PassParameters->LightingCache.LightingCacheResolution = HeterogeneousVolumes::GetLightingCacheResolution(HeterogeneousVolumeInterface);
+		float LODFactor = HeterogeneousVolumes::CalcLODFactor(View, HeterogeneousVolumeInterface);
+		PassParameters->LightingCache.LightingCacheResolution = HeterogeneousVolumes::GetLightingCacheResolution(HeterogeneousVolumeInterface, LODFactor);
 		PassParameters->LightingCache.LightingCacheVoxelBias = HeterogeneousVolumeInterface->GetShadowBiasFactor();
 		//PassParameters->LightingCache.LightingCacheTexture = GraphBuilder.CreateSRV(LightingCacheTexture);
 
@@ -666,7 +667,8 @@ void RenderLightingCacheWithPreshadingHardwareRayTracing(
 	FRenderLightingCacheWithPreshadingRGS::FPermutationDomain PermutationVector;
 	PermutationVector.Set<FRenderLightingCacheWithPreshadingRGS::FLightingCacheMode>(HeterogeneousVolumes::GetLightingCacheMode() - 1);
 	TShaderRef<FRenderLightingCacheWithPreshadingRGS> RayGenerationShader = View.ShaderMap->GetShader<FRenderLightingCacheWithPreshadingRGS>(PermutationVector);
-	FIntVector VolumeResolution = HeterogeneousVolumes::GetLightingCacheResolution(HeterogeneousVolumeInterface);
+	float LODFactor = HeterogeneousVolumes::CalcLODFactor(View, HeterogeneousVolumeInterface);
+	FIntVector VolumeResolution = HeterogeneousVolumes::GetLightingCacheResolution(HeterogeneousVolumeInterface, LODFactor);
 	FIntPoint DispatchResolution = FIntPoint(VolumeResolution.X, VolumeResolution.Y * VolumeResolution.Z);
 
 	GraphBuilder.AddPass(
@@ -794,7 +796,8 @@ void RenderSingleScatteringWithPreshadingHardwareRayTracing(
 		// Transmittance volume
 		if ((HeterogeneousVolumes::UseLightingCacheForTransmittance() && bApplyShadowTransmittance) || HeterogeneousVolumes::UseLightingCacheForInscattering())
 		{
-			PassParameters->LightingCache.LightingCacheResolution = HeterogeneousVolumes::GetLightingCacheResolution(HeterogeneousVolumeInterface);
+			float LODFactor = HeterogeneousVolumes::CalcLODFactor(View, HeterogeneousVolumeInterface);
+			PassParameters->LightingCache.LightingCacheResolution = HeterogeneousVolumes::GetLightingCacheResolution(HeterogeneousVolumeInterface, LODFactor);
 			PassParameters->LightingCache.LightingCacheVoxelBias = HeterogeneousVolumeInterface->GetShadowBiasFactor();
 			PassParameters->LightingCache.LightingCacheTexture = LightingCacheTexture;
 		}
diff --git a/Engine/Source/Runtime/Renderer/Private/HeterogeneousVolumes/HeterogeneousVolumesLiveShadingPipeline.cpp b/Engine/Source/Runtime/Renderer/Private/HeterogeneousVolumes/HeterogeneousVolumesLiveShadingPipeline.cpp
--- a/Engine/Source/Runtime/Renderer/Private/HeterogeneousVolumes/HeterogeneousVolumesLiveShadingPipeline.cpp
+++ b/Engine/Source/Runtime/Renderer/Private/HeterogeneousVolumes/HeterogeneousVolumesLiveShadingPipeline.cpp
@@ -452,6 +452,7 @@ static void RenderLightingCacheWithLiveShading(
 	// Note must be done in the same scope as we add the pass otherwise the UB lifetime will not be guaranteed
 	FDeferredLightUniformStruct DeferredLightUniform = GetDeferredLightParameters(View, *LightSceneInfo);
 	TUniformBufferRef<FDeferredLightUniformStruct> DeferredLightUB = CreateUniformBufferImmediate(DeferredLightUniform, UniformBuffer_SingleDraw);
+	float LODFactor = HeterogeneousVolumes::CalcLODFactor(View, HeterogeneousVolumeInterface);
 
 	FRenderLightingCacheWithLiveShadingCS::FParameters* PassParameters = GraphBuilder.AllocParameters<FRenderLightingCacheWithLiveShadingCS::FParameters>();
 	{
@@ -486,7 +487,7 @@ static void RenderLightingCacheWithLiveShading(
 
 		// Transmittance volume
 		PassParameters->VoxelResolution = HeterogeneousVolumeInterface->GetVoxelResolution();
-		PassParameters->LightingCache.LightingCacheResolution = HeterogeneousVolumes::GetLightingCacheResolution(HeterogeneousVolumeInterface);
+		PassParameters->LightingCache.LightingCacheResolution = HeterogeneousVolumes::GetLightingCacheResolution(HeterogeneousVolumeInterface, LODFactor);
 		PassParameters->LightingCache.LightingCacheVoxelBias = HeterogeneousVolumeInterface->GetShadowBiasFactor();
 		//PassParameters->LightingCache.LightingCacheTexture = GraphBuilder.CreateSRV(LightingCacheTexture);
 		PassParameters->LightingCache.LightingCacheTexture = FRDGSystemTextures::Get(GraphBuilder).VolumetricBlack;
@@ -495,9 +496,9 @@ static void RenderLightingCacheWithLiveShading(
 		PassParameters->MaxTraceDistance = HeterogeneousVolumes::GetMaxTraceDistance();
 		PassParameters->MaxShadowTraceDistance = HeterogeneousVolumes::GetMaxShadowTraceDistance();
 		PassParameters->StepSize = HeterogeneousVolumes::GetStepSize();
-		PassParameters->StepFactor = HeterogeneousVolumeInterface->GetStepFactor();
+		PassParameters->StepFactor = HeterogeneousVolumeInterface->GetStepFactor() * LODFactor;
 		PassParameters->ShadowStepSize = HeterogeneousVolumes::GetShadowStepSize();
-		PassParameters->ShadowStepFactor = HeterogeneousVolumeInterface->GetShadowStepFactor();
+		PassParameters->ShadowStepFactor = HeterogeneousVolumeInterface->GetShadowStepFactor() * LODFactor;
 		PassParameters->MaxStepCount = HeterogeneousVolumes::GetMaxStepCount();
 		PassParameters->bJitter = HeterogeneousVolumes::ShouldJitter();
 
@@ -544,7 +545,7 @@ static void RenderLightingCacheWithLiveShading(
 #endif // WANTS_DRAW_MESH_EVENTS
 
 	PassParameters->VoxelMin = FIntVector::ZeroValue;
-	PassParameters->VoxelMax = HeterogeneousVolumes::GetLightingCacheResolution(HeterogeneousVolumeInterface) - FIntVector(1);
+	PassParameters->VoxelMax = HeterogeneousVolumes::GetLightingCacheResolution(HeterogeneousVolumeInterface, LODFactor) - FIntVector(1);
 	
 	bool bShouldBoundsCull = HeterogeneousVolumes::ShouldBoundsCull();
 	if (LightType != LightType_Directional && bShouldBoundsCull)
@@ -664,6 +665,7 @@ static void RenderSingleScatteringWithLiveShading(
 		PassParameters->BlueNoise = CreateUniformBufferImmediate(BlueNoise, EUniformBufferUsage::UniformBuffer_SingleDraw);
 
 		// Light data
+		float LODFactor = HeterogeneousVolumes::CalcLODFactor(View, HeterogeneousVolumeInterface);
 		PassParameters->bApplyEmissionAndTransmittance = bApplyEmissionAndTransmittance;
 		PassParameters->bApplyDirectLighting = bApplyDirectLighting;
 		PassParameters->bApplyShadowTransmittance = bApplyShadowTransmittance;
@@ -678,7 +680,7 @@ static void RenderSingleScatteringWithLiveShading(
 		PassParameters->DeferredLight = DeferredLightUB;
 		PassParameters->LightType = LightType;
 		PassParameters->ShadowStepSize = HeterogeneousVolumes::GetShadowStepSize();
-		PassParameters->ShadowStepFactor = HeterogeneousVolumeInterface->GetShadowStepFactor();
+		PassParameters->ShadowStepFactor = HeterogeneousVolumeInterface->GetShadowStepFactor() * LODFactor;
 
 		// Object data
 		// TODO: Convert to relative-local space
@@ -701,7 +703,7 @@ static void RenderSingleScatteringWithLiveShading(
 		// Ray data
 		PassParameters->MaxTraceDistance = HeterogeneousVolumes::GetMaxTraceDistance();
 		PassParameters->StepSize = HeterogeneousVolumes::GetStepSize();
-		PassParameters->StepFactor = HeterogeneousVolumeInterface->GetStepFactor();
+		PassParameters->StepFactor = HeterogeneousVolumeInterface->GetStepFactor() * LODFactor;
 		PassParameters->MaxStepCount = HeterogeneousVolumes::GetMaxStepCount();
 		PassParameters->bJitter = HeterogeneousVolumes::ShouldJitter();
 
@@ -747,7 +749,7 @@ static void RenderSingleScatteringWithLiveShading(
 		// Volume data
 		if ((HeterogeneousVolumes::UseLightingCacheForTransmittance() && bApplyShadowTransmittance) || HeterogeneousVolumes::UseLightingCacheForInscattering())
 		{
-			PassParameters->LightingCache.LightingCacheResolution = HeterogeneousVolumes::GetLightingCacheResolution(HeterogeneousVolumeInterface);
+			PassParameters->LightingCache.LightingCacheResolution = HeterogeneousVolumes::GetLightingCacheResolution(HeterogeneousVolumeInterface, LODFactor);
 			PassParameters->LightingCache.LightingCacheVoxelBias = HeterogeneousVolumeInterface->GetShadowBiasFactor();
 			PassParameters->LightingCache.LightingCacheTexture = LightingCacheTexture;
 		}
@@ -1368,8 +1370,9 @@ bool RenderVolumetricShadowMapForLightForHeterogeneousVolumeWithLiveShading(
 		PassParameters->VoxelResolution = HeterogeneousVolumeInterface->GetVoxelResolution();
 
 		// Ray Data
+		float LODFactor = HeterogeneousVolumes::CalcLODFactor(View, HeterogeneousVolumeInterface);
 		PassParameters->ShadowStepSize = HeterogeneousVolumes::GetShadowStepSize();
-		PassParameters->ShadowStepFactor = HeterogeneousVolumeInterface->GetShadowStepFactor();
+		PassParameters->ShadowStepFactor = HeterogeneousVolumeInterface->GetShadowStepFactor() * LODFactor;
 		PassParameters->MaxTraceDistance = HeterogeneousVolumes::GetMaxTraceDistance();
 		PassParameters->MaxStepCount = HeterogeneousVolumes::GetMaxStepCount();
 		PassParameters->bJitter = HeterogeneousVolumes::ShouldJitter();
@@ -1491,6 +1494,20 @@ bool RenderVolumetricShadowMapForLightWithLiveShading(
 		return false;
 	}
 
+	// Adjust shadow resolution based on minimum MipLevel
+	float LODValue = FMath::CeilLogTwo(ShadowMapResolution.X);
+	for (auto VolumetricMeshBatch : HeterogeneousVolumesMeshBatches)
+	{
+		int32 VolumeCount = VolumetricMeshBatch.Mesh->Elements.Num();
+		for (int32 VolumeIndex = 0; VolumeIndex < VolumeCount; ++VolumeIndex)
+		{
+			const IHeterogeneousVolumeInterface* HeterogeneousVolumeInterface = (IHeterogeneousVolumeInterface*)VolumetricMeshBatch.Mesh->Elements[VolumeIndex].UserData;
+			LODValue = FMath::Min(LODValue, HeterogeneousVolumes::CalcLOD(View, HeterogeneousVolumeInterface));
+		}
+	}
+	float LODFactor = HeterogeneousVolumes::CalcLODFactor(LODValue);
+	ShadowMapResolution /= LODFactor;
+
 	// Build shadow transform
 	NumShadowMatrices = ProjectedShadowInfo->OnePassShadowViewProjectionMatrices.Num();
 	FMatrix44f ShadowToTranslatedWorld[6];
diff --git a/Engine/Source/Runtime/Renderer/Private/HeterogeneousVolumes/HeterogeneousVolumesPreshadingPipeline.cpp b/Engine/Source/Runtime/Renderer/Private/HeterogeneousVolumes/HeterogeneousVolumesPreshadingPipeline.cpp
--- a/Engine/Source/Runtime/Renderer/Private/HeterogeneousVolumes/HeterogeneousVolumesPreshadingPipeline.cpp
+++ b/Engine/Source/Runtime/Renderer/Private/HeterogeneousVolumes/HeterogeneousVolumesPreshadingPipeline.cpp
@@ -417,6 +417,8 @@ void RenderLightingCacheWithPreshadingCompute(
 	FDeferredLightUniformStruct DeferredLightUniform = GetDeferredLightParameters(View, *LightSceneInfo);
 	TUniformBufferRef<FDeferredLightUniformStruct> DeferredLightUB = CreateUniformBufferImmediate(DeferredLightUniform, UniformBuffer_SingleDraw);
 
+	float LODFactor = HeterogeneousVolumes::CalcLODFactor(View, HeterogeneousVolumeInterface);
+
 	FRenderLightingCacheWithPreshadingCS::FParameters* PassParameters = GraphBuilder.AllocParameters<FRenderLightingCacheWithPreshadingCS::FParameters>();
 	{
 		// Scene data
@@ -436,7 +438,7 @@ void RenderLightingCacheWithPreshadingCompute(
 		PassParameters->SparseVoxelUniformBuffer = SparseVoxelUniformBuffer;
 
 		// Transmittance volume
-		PassParameters->LightingCache.LightingCacheResolution = HeterogeneousVolumes::GetLightingCacheResolution(HeterogeneousVolumeInterface);
+		PassParameters->LightingCache.LightingCacheResolution = HeterogeneousVolumes::GetLightingCacheResolution(HeterogeneousVolumeInterface, LODFactor);
 		PassParameters->LightingCache.LightingCacheVoxelBias = HeterogeneousVolumeInterface->GetShadowBiasFactor();
 		PassParameters->LightingCache.LightingCacheTexture = LightingCacheTexture;
 
@@ -491,7 +493,7 @@ void RenderLightingCacheWithPreshadingCompute(
 	PermutationVector.Set<FRenderLightingCacheWithPreshadingCS::FLightingCacheMode>(HeterogeneousVolumes::GetLightingCacheMode() - 1);
 	TShaderRef<FRenderLightingCacheWithPreshadingCS> ComputeShader = View.ShaderMap->GetShader<FRenderLightingCacheWithPreshadingCS>(PermutationVector);
 
-	FIntVector GroupCount = HeterogeneousVolumes::GetLightingCacheResolution(HeterogeneousVolumeInterface);
+	FIntVector GroupCount = HeterogeneousVolumes::GetLightingCacheResolution(HeterogeneousVolumeInterface, LODFactor);
 	GroupCount.X = FMath::DivideAndRoundUp(GroupCount.X, FRenderLightingCacheWithPreshadingCS::GetThreadGroupSize3D());
 	GroupCount.Y = FMath::DivideAndRoundUp(GroupCount.Y, FRenderLightingCacheWithPreshadingCS::GetThreadGroupSize3D());
 	GroupCount.Z = FMath::DivideAndRoundUp(GroupCount.Z, FRenderLightingCacheWithPreshadingCS::GetThreadGroupSize3D());
@@ -603,7 +605,8 @@ void RenderSingleScatteringWithPreshadingCompute(
 		// Transmittance volume
 		if ((HeterogeneousVolumes::UseLightingCacheForTransmittance() && bApplyShadowTransmittance) || HeterogeneousVolumes::UseLightingCacheForInscattering())
 		{
-			PassParameters->LightingCache.LightingCacheResolution = HeterogeneousVolumes::GetLightingCacheResolution(HeterogeneousVolumeInterface);
+			float LODFactor = HeterogeneousVolumes::CalcLODFactor(View, HeterogeneousVolumeInterface);
+			PassParameters->LightingCache.LightingCacheResolution = HeterogeneousVolumes::GetLightingCacheResolution(HeterogeneousVolumeInterface, LODFactor);
 			PassParameters->LightingCache.LightingCacheVoxelBias = HeterogeneousVolumeInterface->GetShadowBiasFactor();
 			PassParameters->LightingCache.LightingCacheTexture = LightingCacheTexture;
 		}
