Heterogeneous Volumes: Compute self-shadowing for the active volume if the AVSM does not contain its shadowing term.

diff --git a/Engine/Source/Runtime/Renderer/Private/HeterogeneousVolumes/HeterogeneousVolumes.cpp b/Engine/Source/Runtime/Renderer/Private/HeterogeneousVolumes/HeterogeneousVolumes.cpp
--- a/Engine/Source/Runtime/Renderer/Private/HeterogeneousVolumes/HeterogeneousVolumes.cpp
+++ b/Engine/Source/Runtime/Renderer/Private/HeterogeneousVolumes/HeterogeneousVolumes.cpp
@@ -453,9 +453,11 @@ namespace HeterogeneousVolumes
 		return CVarHeterogeneousVolumesLightingCache.GetValueOnRenderThread();
 	}
 
-	bool UseAdaptiveVolumetricShadowMapForSelfShadowing()
+	bool UseAdaptiveVolumetricShadowMapForSelfShadowing(const FPrimitiveSceneProxy* PrimitiveSceneProxy)
 	{
-		return ShouldHeterogeneousVolumesCastShadows() && CVarHeterogeneousVolumesLightingCacheUseAVSM.GetValueOnRenderThread();
+		bool bUseAVSM = CVarHeterogeneousVolumesLightingCacheUseAVSM.GetValueOnRenderThread() != 0;
+		bool bPrimitiveCastsDynamicShadows = PrimitiveSceneProxy->CastsDynamicShadow();
+		return ShouldHeterogeneousVolumesCastShadows() && bUseAVSM && bPrimitiveCastsDynamicShadows;
 	}
 
 	bool UseLightingCacheForInscattering()
diff --git a/Engine/Source/Runtime/Renderer/Private/HeterogeneousVolumes/HeterogeneousVolumes.h b/Engine/Source/Runtime/Renderer/Private/HeterogeneousVolumes/HeterogeneousVolumes.h
--- a/Engine/Source/Runtime/Renderer/Private/HeterogeneousVolumes/HeterogeneousVolumes.h
+++ b/Engine/Source/Runtime/Renderer/Private/HeterogeneousVolumes/HeterogeneousVolumes.h
@@ -102,7 +102,7 @@ namespace HeterogeneousVolumes
 	bool UseSparseVoxelPerTileCulling();
 	bool UseLightingCacheForInscattering();
 	bool UseLightingCacheForTransmittance();
-	bool UseAdaptiveVolumetricShadowMapForSelfShadowing();
+	bool UseAdaptiveVolumetricShadowMapForSelfShadowing(const FPrimitiveSceneProxy* PrimitiveSceneProxy);
 	bool ShouldApplyHeightFog();
 	bool ShouldApplyVolumetricFog();
 
diff --git a/Engine/Source/Runtime/Renderer/Private/HeterogeneousVolumes/HeterogeneousVolumesLiveShadingPipeline.cpp b/Engine/Source/Runtime/Renderer/Private/HeterogeneousVolumes/HeterogeneousVolumesLiveShadingPipeline.cpp
--- a/Engine/Source/Runtime/Renderer/Private/HeterogeneousVolumes/HeterogeneousVolumesLiveShadingPipeline.cpp
+++ b/Engine/Source/Runtime/Renderer/Private/HeterogeneousVolumes/HeterogeneousVolumesLiveShadingPipeline.cpp
@@ -601,9 +601,11 @@ static void RenderLightingCacheWithLiveShading(
 	GroupCount.Y = FMath::DivideAndRoundUp(GroupCount.Y, FRenderLightingCacheWithLiveShadingCS::GetThreadGroupSize3D());
 	GroupCount.Z = FMath::DivideAndRoundUp(GroupCount.Z, FRenderLightingCacheWithLiveShadingCS::GetThreadGroupSize3D());
 
+	bool bUseAVSM = HeterogeneousVolumes::UseAdaptiveVolumetricShadowMapForSelfShadowing(HeterogeneousVolumeInterface->GetPrimitiveSceneProxy());
+
 	FRenderLightingCacheWithLiveShadingCS::FPermutationDomain PermutationVector;
 	PermutationVector.Set<FRenderLightingCacheWithLiveShadingCS::FLightingCacheMode>(HeterogeneousVolumes::GetLightingCacheMode() - 1);
-	PermutationVector.Set<FRenderLightingCacheWithLiveShadingCS::FUseAdaptiveVolumetricShadowMap>(HeterogeneousVolumes::UseAdaptiveVolumetricShadowMapForSelfShadowing());
+	PermutationVector.Set<FRenderLightingCacheWithLiveShadingCS::FUseAdaptiveVolumetricShadowMap>(bUseAVSM);
 	TShaderRef<FRenderLightingCacheWithLiveShadingCS> ComputeShader = Material.GetShader<FRenderLightingCacheWithLiveShadingCS>(&FLocalVertexFactory::StaticType, PermutationVector, false);
 	if (!ComputeShader.IsNull())
 	{
@@ -785,12 +787,14 @@ static void RenderSingleScatteringWithLiveShading(
 	}
 #endif // WANTS_DRAW_MESH_EVENTS
 
+	bool bUseAVSM = HeterogeneousVolumes::UseAdaptiveVolumetricShadowMapForSelfShadowing(HeterogeneousVolumeInterface->GetPrimitiveSceneProxy());
+
 	FRenderSingleScatteringWithLiveShadingCS::FPermutationDomain PermutationVector;
 	PermutationVector.Set<FRenderSingleScatteringWithLiveShadingCS::FUseTransmittanceVolume>(HeterogeneousVolumes::UseLightingCacheForTransmittance() && PassParameters->bApplyShadowTransmittance);
 	PermutationVector.Set<FRenderSingleScatteringWithLiveShadingCS::FUseInscatteringVolume>(HeterogeneousVolumes::UseLightingCacheForInscattering());
 	PermutationVector.Set<FRenderSingleScatteringWithLiveShadingCS::FUseLumenGI>(HeterogeneousVolumes::UseIndirectLighting() && View.GetLumenTranslucencyGIVolume().Texture0 != nullptr);
 	PermutationVector.Set<FRenderSingleScatteringWithLiveShadingCS::FWriteVelocity>(bWriteVelocity);
-	PermutationVector.Set<FRenderSingleScatteringWithLiveShadingCS::FUseAdaptiveVolumetricShadowMap>(HeterogeneousVolumes::UseAdaptiveVolumetricShadowMapForSelfShadowing());
+	PermutationVector.Set<FRenderSingleScatteringWithLiveShadingCS::FUseAdaptiveVolumetricShadowMap>(bUseAVSM);
 	PermutationVector.Set<FRenderSingleScatteringWithLiveShadingCS::FApplyFogInscattering>(static_cast<int32>(HeterogeneousVolumes::GetApplyFogInscattering()));
 	PermutationVector = FRenderSingleScatteringWithLiveShadingCS::RemapPermutation(PermutationVector);
 	TShaderRef<FRenderSingleScatteringWithLiveShadingCS> ComputeShader = Material.GetShader<FRenderSingleScatteringWithLiveShadingCS>(&FLocalVertexFactory::StaticType, PermutationVector, false);
