Removed dynamic string allocation when allocating RDG textures for translucent lighting.

diff --git a/Engine/Source/Runtime/Renderer/Private/TranslucentLighting.cpp b/Engine/Source/Runtime/Renderer/Private/TranslucentLighting.cpp
--- a/Engine/Source/Runtime/Renderer/Private/TranslucentLighting.cpp
+++ b/Engine/Source/Runtime/Renderer/Private/TranslucentLighting.cpp
@@ -831,16 +831,13 @@ void FTranslucencyLightingVolumeTextures::Init(FRDGBuilder& GraphBuilder, TArray
 				const uint32 TextureIndex = FTranslucencyLightingVolumeTextures::GetIndex(Views[ViewIndex], CascadeIndex);
 				check(TextureIndex <= NumViewsWithTextures * TVC_MAX);
 
-				const FRDGEventName& AmbientName = *GraphBuilder.AllocObject<FRDGEventName>(RDG_EVENT_NAME("TranslucentVolumeAmbient%d", TextureIndex));
-				const FRDGEventName& DirectionalName = *GraphBuilder.AllocObject<FRDGEventName>(RDG_EVENT_NAME("TranslucentVolumeDirectional%d", TextureIndex));
-
 				FRDGTextureRef AmbientTexture = GraphBuilder.CreateTexture(
 					FRDGTextureDesc::Create3D(
 						TranslucencyLightingVolumeDim,
 						PF_FloatRGBA,
 						FClearValueBinding::Transparent,
 						TranslucencyTargetFlags),
-					AmbientName.GetTCHAR());
+					TEXT("TranslucentVolumeAmbient"));
 
 				FRDGTextureRef DirectionalTexture = GraphBuilder.CreateTexture(
 					FRDGTextureDesc::Create3D(
@@ -848,7 +845,7 @@ void FTranslucencyLightingVolumeTextures::Init(FRDGBuilder& GraphBuilder, TArray
 						PF_FloatRGBA,
 						FClearValueBinding::Transparent,
 						TranslucencyTargetFlags),
-					DirectionalName.GetTCHAR());
+					TEXT("TranslucentVolumeDirectional"));
 
 				Ambient[TextureIndex] = AmbientTexture;
 				Directional[TextureIndex] = DirectionalTexture;
