- Update primitive last render time for always visible primitives
- This ensures Actors made up of always visible primitives will have a valid last render time

diff --git a/Engine/Source/Runtime/Renderer/Private/SceneVisibility.cpp b/Engine/Source/Runtime/Renderer/Private/SceneVisibility.cpp
--- a/Engine/Source/Runtime/Renderer/Private/SceneVisibility.cpp
+++ b/Engine/Source/Runtime/Renderer/Private/SceneVisibility.cpp
@@ -630,7 +630,7 @@ static void CullOctree(const FScene& Scene, FViewInfo& View, const FFrustumCulli
 		});
 }
 
-static void UpdateAlwaysVisible(const FScene& Scene, FViewInfo& View, FFrustumCullingFlags Flags, const FVisibilityTaskConfig& TaskConfig, int32 TaskIndex)
+static void UpdateAlwaysVisible(const FScene& Scene, FViewInfo& View, FFrustumCullingFlags Flags, const FVisibilityTaskConfig& TaskConfig, int32 TaskIndex, float CurrentWorldTime)
 {
 	QUICK_SCOPE_CYCLE_COUNTER(STAT_AlwaysVisible_Loop);
 
@@ -664,6 +664,12 @@ static void UpdateAlwaysVisible(const FScene& Scene, FViewInfo& View, FFrustumCu
 				RayTracingBits |= Mask;
 			}
 		#endif
+
+			FPrimitiveSceneInfo* PrimitiveSceneInfo = Scene.Primitives[Index];
+			PrimitiveSceneInfo->LastRenderTime = CurrentWorldTime;
+
+			const bool bUpdateLastRenderTimeOnScreen = true;
+			PrimitiveSceneInfo->UpdateComponentLastRenderTime(CurrentWorldTime, bUpdateLastRenderTimeOnScreen);
 		}
 
 		VisWords[StartWord + WordIndex] = VisBits;
@@ -3629,16 +3635,17 @@ void FVisibilityViewPacket::BeginInitVisibility()
 		const bool bHasAlwaysVisible = TaskConfig.NumVisiblePrimitives > 0;
 		if (bHasAlwaysVisible)
 		{
+			const float CurrentWorldTime = View.Family->Time.GetWorldTimeSeconds();
 			for (uint32 TaskIndex = 0; TaskIndex < TaskConfig.AlwaysVisible.NumTasks; ++TaskIndex)
 			{
 				Tasks.AlwaysVisible.AddPrerequisites(
-					UE::Tasks::Launch(UE_SOURCE_LOCATION, [this, Flags, TaskIndex]() mutable
+					UE::Tasks::Launch(UE_SOURCE_LOCATION, [this, Flags, TaskIndex, CurrentWorldTime]() mutable
 				{
 					TRACE_CPUPROFILER_EVENT_SCOPE(SceneVisibility_AlwaysVisible);
 					SCOPE_CYCLE_COUNTER(STAT_UpdateAlwaysVisible);
 
 					FOptionalTaskTagScope TaskTagScope(ETaskTag::EParallelRenderingThread);
-					UpdateAlwaysVisible(Scene, View, Flags, TaskConfig, TaskIndex);
+					UpdateAlwaysVisible(Scene, View, Flags, TaskConfig, TaskIndex, CurrentWorldTime);
 
 				}, PrerequisiteTask, TaskConfig.TaskPriority, UE::Tasks::EExtendedTaskPriority::None));
 			}
@@ -3721,11 +3728,12 @@ void FVisibilityViewPacket::BeginInitVisibility()
 
 		PrerequisiteTask.Wait();
 
-		ParallelFor(TaskConfig.AlwaysVisible.NumTasks, [this, Flags](int32 TaskIndex)
+		const float CurrentWorldTime = View.Family->Time.GetWorldTimeSeconds();
+		ParallelFor(TaskConfig.AlwaysVisible.NumTasks, [this, Flags, CurrentWorldTime](int32 TaskIndex)
 		{
 			TRACE_CPUPROFILER_EVENT_SCOPE(SceneVisibility_AlwaysVisible);
 			FOptionalTaskTagScope TaskTagScope(ETaskTag::EParallelRenderingThread);
-			UpdateAlwaysVisible(Scene, View, Flags, TaskConfig, TaskIndex);
+			UpdateAlwaysVisible(Scene, View, Flags, TaskConfig, TaskIndex, CurrentWorldTime);
 
 		}, bSingleThreaded);
 
