ColorManagement: Updated r.AutoExposure.LuminanceMethod option number two to working color space.

diff --git a/Engine/Source/Runtime/Renderer/Private/PostProcess/PostProcessEyeAdaptation.cpp b/Engine/Source/Runtime/Renderer/Private/PostProcess/PostProcessEyeAdaptation.cpp
--- a/Engine/Source/Runtime/Renderer/Private/PostProcess/PostProcessEyeAdaptation.cpp
+++ b/Engine/Source/Runtime/Renderer/Private/PostProcess/PostProcessEyeAdaptation.cpp
@@ -18,6 +18,7 @@
 #include "TextureResource.h"
 #include "PostProcessing.h"
 #include "PostProcessLocalExposure.h"
+#include "ColorManagement/ColorSpace.h"
 
 bool IsMobileEyeAdaptationEnabled(const FViewInfo& View);
 
@@ -83,7 +84,7 @@ namespace
 		0,
 		TEXT("0 - Uniform.\n")
 		TEXT("1 - NSTC.\n")
-		TEXT("2 - Rec709."),
+		TEXT("2 - Working Color Space."),
 		ECVF_RenderThreadSafe);
 
 	TAutoConsoleVariable<bool> CVarAutoExposureIgnoreMaterialsReconstructFromSceneColor(
@@ -545,13 +546,13 @@ FEyeAdaptationParameters GetEyeAdaptationParameters(const FViewInfo& View)
 	const int32 LuminanceMethod = CVarAutoExposureLuminanceMethod.GetValueOnRenderThread();
 	if (LuminanceMethod == 1)
 	{
-		// NTSC / match weights in Common.ush
+		// NTSC (deprecated legacy weights in Common.ush)
 		Parameters.LuminanceWeights = FVector3f(0.3f, 0.59f, 0.11f);
 	}
 	else if (LuminanceMethod == 2)
 	{
-		// Rec 709
-		Parameters.LuminanceWeights = FVector3f(0.2126f, 0.7152f, 0.0722f);
+		// Working color space (sRGB/Rec709 default)
+		Parameters.LuminanceWeights = FVector3f(UE::Color::FColorSpace::GetWorking().GetLuminanceFactors());
 	}
 	else
 	{
