HWRT: Fix RayTracingTestbed - missing BUF_MultiGPUAllocate flag on BLAS address buffer.

diff --git a/Engine/Source/Runtime/Renderer/Private/Tests/RayTracingTestbed.cpp b/Engine/Source/Runtime/Renderer/Private/Tests/RayTracingTestbed.cpp
--- a/Engine/Source/Runtime/Renderer/Private/Tests/RayTracingTestbed.cpp
+++ b/Engine/Source/Runtime/Renderer/Private/Tests/RayTracingTestbed.cpp
@@ -207,7 +207,7 @@ bool RunRayTracingTestbed_RenderThread(const FString& Parameters)
 	InstanceBuffer.Initialize(RHICmdList, TEXT("RayTracingTestBedInstanceBuffer"), GRHIRayTracingInstanceDescriptorSize, SceneInitializer.NumNativeInstancesPerLayer[0]);
 
 	FByteAddressBuffer AccelerationStructureAddressesBuffer;
-	AccelerationStructureAddressesBuffer.Initialize(RHICmdList, TEXT("RayTracingTestBedAccelerationStructureAddressesBuffer"), sizeof(FRayTracingAccelerationStructureAddress), BUF_Volatile);
+	AccelerationStructureAddressesBuffer.Initialize(RHICmdList, TEXT("RayTracingTestBedAccelerationStructureAddressesBuffer"), sizeof(FRayTracingAccelerationStructureAddress), BUF_Volatile | BUF_MultiGPUAllocate);
 
 	const uint32 InstanceUploadBufferSize = SceneInitializer.NumNativeInstancesPerLayer[0] * sizeof(FRayTracingInstanceDescriptorInput);
 	FBufferRHIRef InstanceUploadBuffer;
