Issue 446134:  Scene Capture static analysis warning fix (Lifetime of the lambda is greater than lifetime of the local variable).  Make copy of CopyInfos for lambda.

#rnx

diff --git a/Engine/Source/Runtime/Renderer/Private/SceneCaptureRendering.cpp b/Engine/Source/Runtime/Renderer/Private/SceneCaptureRendering.cpp
--- a/Engine/Source/Runtime/Renderer/Private/SceneCaptureRendering.cpp
+++ b/Engine/Source/Runtime/Renderer/Private/SceneCaptureRendering.cpp
@@ -407,7 +407,11 @@ static void UpdateSceneCaptureContentDeferred_RenderThread(
 		// as the scene textures, and no viewport remapping is required.
 		if (!CopyInfos[0].Size.IsZero() && !OutputTexture->Desc.IsTextureCube())
 		{
-			CopyCaptureToTargetSetViewportFn = [&CopyInfos](FRHICommandList& RHICmdList, int32 ViewIndex)
+			// Fix for static analysis warning -- lambda lifetime exceeds lifetime of CopyInfos.  Technically, the lambda is consumed in the
+			// scene render call below and not used afterwards, but static analysis doesn't know that, so we make a copy.
+			TArray<FRHICopyTextureInfo, TInlineAllocator<1>> CopyInfosLocal(CopyInfos);
+
+			CopyCaptureToTargetSetViewportFn = [CopyInfos = MoveTemp(CopyInfosLocal)](FRHICommandList& RHICmdList, int32 ViewIndex)
 			{
 				const FIntRect CopyDestRect = CopyInfos[ViewIndex].GetDestRect();
 
@@ -497,7 +501,11 @@ static void UpdateSceneCaptureContentMobile_RenderThread(
 		// as the scene textures, and no viewport remapping is required.
 		if (!CopyInfos[0].Size.IsZero() && !OutputTexture->Desc.IsTextureCube())
 		{
-			CopyCaptureToTargetSetViewportFn = [&CopyInfos](FRHICommandList& RHICmdList, int32 ViewIndex)
+			// Fix for static analysis warning -- lambda lifetime exceeds lifetime of CopyInfos.  Technically, the lambda is consumed in the
+			// scene render call below and not used afterwards, but static analysis doesn't know that, so we make a copy.
+			TArray<FRHICopyTextureInfo, TInlineAllocator<1>> CopyInfosLocal(CopyInfos);
+
+			CopyCaptureToTargetSetViewportFn = [CopyInfos = MoveTemp(CopyInfosLocal)](FRHICommandList& RHICmdList, int32 ViewIndex)
 			{
 				const FIntRect CopyDestRect = CopyInfos[ViewIndex].GetDestRect();
 
