Light Function Atlas  usage is now enabled by default for deferred shading.
- Directional lights are not rendered using the light function atlas when using deferred shading because lots of content could be broken (free form UVs will make the atlas tile non tilable)
- Materials can now tell if they are atlas compatible (not sampling depth or world position) and if so will be skipped by deferred shading only.
- All the other system light function atlas support is new so if used, light function material must comply with the light function compatibility.
- on levels with lots of lights using light function material, this can result in gain of 2ms

diff --git a/Engine/Shaders/Private/LightFunctionAtlas/LightFunctionAtlasRender.usf b/Engine/Shaders/Private/LightFunctionAtlas/LightFunctionAtlasRender.usf
--- a/Engine/Shaders/Private/LightFunctionAtlas/LightFunctionAtlasRender.usf
+++ b/Engine/Shaders/Private/LightFunctionAtlas/LightFunctionAtlasRender.usf
@@ -7,6 +7,9 @@
 // This define is for the light function common evaluation to only be based on UV stored in the light vector only.
 #define LIGHT_FUNCTION_ATLAS 1
 
+// Used to trigger shader recompiles, guid should be regenerated on merge conflicts
+#pragma message("UESHADERMETADATA_VERSION 4847F0A3-06B4-4A6D-8C31-3F98BEEF9480")
+
 #include "../Common.ush"
 #include "/Engine/Generated/Material.ush"
 #include "../LightFunctionCommon.ush"
diff --git a/Engine/Shaders/Private/LightFunctionPixelShader.usf b/Engine/Shaders/Private/LightFunctionPixelShader.usf
--- a/Engine/Shaders/Private/LightFunctionPixelShader.usf
+++ b/Engine/Shaders/Private/LightFunctionPixelShader.usf
@@ -9,6 +9,9 @@
 #include "LightFunctionCommon.ush"
 #include "DynamicLightingCommon.ush"
 
+// Used to trigger shader recompiles, guid should be regenerated on merge conflicts
+#pragma message("UESHADERMETADATA_VERSION BF889CEA-E940-4C56-8555-E5BEB172F300")
+
 // used like this: mul(float4(SvPosition.xyz, 1), SvPositionToLight);
 float4x4 SvPositionToLight;
 
diff --git a/Engine/Source/Runtime/Engine/Classes/Materials/MaterialRelevance.h b/Engine/Source/Runtime/Engine/Classes/Materials/MaterialRelevance.h
--- a/Engine/Source/Runtime/Engine/Classes/Materials/MaterialRelevance.h
+++ b/Engine/Source/Runtime/Engine/Classes/Materials/MaterialRelevance.h
@@ -45,6 +45,7 @@ struct FMaterialRelevance
 			uint8 bUsesDistanceCullFade : 1;
 			uint8 bDisableDepthTest : 1;
 			uint8 bUsesAnisotropy : 1;
+			uint8 bIsLightFunctionAtlasCompatible : 1;
 		};
 		uint64 Raw;
 	};
diff --git a/Engine/Source/Runtime/Engine/Private/Materials/HLSLMaterialTranslator.cpp b/Engine/Source/Runtime/Engine/Private/Materials/HLSLMaterialTranslator.cpp
--- a/Engine/Source/Runtime/Engine/Private/Materials/HLSLMaterialTranslator.cpp
+++ b/Engine/Source/Runtime/Engine/Private/Materials/HLSLMaterialTranslator.cpp
@@ -1700,6 +1700,8 @@ void FHLSLMaterialTranslator::DoTranslate()
 	bUsesAnisotropy = IsMaterialPropertyUsed(MP_Anisotropy, Chunk[MP_Anisotropy], FLinearColor(0, 0, 0, 0), 1);
 	MaterialCompilationOutput.bUsesAnisotropy = bUsesAnisotropy;
 
+	MaterialCompilationOutput.bIsLightFunctionAtlasCompatible = !bUsesVertexPosition && !bUsesSceneDepth;
+
 	EMaterialDecalResponse MDR = (EMaterialDecalResponse)Material->GetMaterialDecalResponse();
 	if (MDR == MDR_Color || MDR == MDR_ColorNormal || MDR == MDR_ColorRoughness || MDR == MDR_ColorNormalRoughness)
 	{
diff --git a/Engine/Source/Runtime/Engine/Private/Materials/MaterialInterface.cpp b/Engine/Source/Runtime/Engine/Private/Materials/MaterialInterface.cpp
--- a/Engine/Source/Runtime/Engine/Private/Materials/MaterialInterface.cpp
+++ b/Engine/Source/Runtime/Engine/Private/Materials/MaterialInterface.cpp
@@ -513,6 +513,7 @@ FMaterialRelevance UMaterialInterface::GetRelevance_Internal(const UMaterial* Ma
 		MaterialRelevance.SubstrateClosureCountMask = SubstrateBSDFCountMask;
 		MaterialRelevance.SubstrateUintPerPixel = SubstrateUintPerPixel;
 		MaterialRelevance.bUsesComplexSpecialRenderPath = MaterialResource->MaterialGetSubstrateUsesComplexSpecialRenderPath_GameThread();
+		MaterialRelevance.bIsLightFunctionAtlasCompatible = MaterialResource->MaterialIsLightFunctionAtlasCompatible_GameThread();
 
 		return MaterialRelevance;
 	}
diff --git a/Engine/Source/Runtime/Engine/Private/Materials/MaterialShared.cpp b/Engine/Source/Runtime/Engine/Private/Materials/MaterialShared.cpp
--- a/Engine/Source/Runtime/Engine/Private/Materials/MaterialShared.cpp
+++ b/Engine/Source/Runtime/Engine/Private/Materials/MaterialShared.cpp
@@ -1226,6 +1226,17 @@ bool FMaterial::MaterialUsesAnisotropy_RenderThread() const
 	return RenderingThreadShaderMap ? RenderingThreadShaderMap->UsesAnisotropy() : false;
 }
 
+bool FMaterial::MaterialIsLightFunctionAtlasCompatible_GameThread() const
+{
+	return GameThreadShaderMap ? GameThreadShaderMap->IsLightFunctionAtlasCompatible() : false;
+}
+
+bool FMaterial::MaterialIsLightFunctionAtlasCompatible_RenderThread() const
+{
+	check(IsInParallelRenderingThread());
+	return RenderingThreadShaderMap ? RenderingThreadShaderMap->IsLightFunctionAtlasCompatible() : false;
+}
+
 uint8 FMaterial::MaterialGetSubstrateMaterialType_GameThread() const
 {
 	return GameThreadShaderMap ? GameThreadShaderMap->GetSubstrateMaterialType() : false;
diff --git a/Engine/Source/Runtime/Engine/Public/MaterialShared.h b/Engine/Source/Runtime/Engine/Public/MaterialShared.h
--- a/Engine/Source/Runtime/Engine/Public/MaterialShared.h
+++ b/Engine/Source/Runtime/Engine/Public/MaterialShared.h
@@ -835,6 +835,9 @@ public:
 
 	/** true if the material uses non 0 anisotropy value */
 	LAYOUT_BITFIELD(uint8, bUsesAnisotropy, 1);
+
+	/** true if the material is compatible with the light function atlas */
+	LAYOUT_BITFIELD(uint8, bIsLightFunctionAtlasCompatible, 1);
 	
 	/** Whether the material uses scalar displacement. */
 	LAYOUT_BITFIELD(uint8, bUsesDisplacement, 1);
@@ -1490,6 +1493,7 @@ public:
 	bool UsesVelocitySceneTexture() const { return GetContent()->MaterialCompilationOutput.UsesVelocitySceneTexture(); }
 	bool UsesDistanceCullFade() const { return GetContent()->MaterialCompilationOutput.bUsesDistanceCullFade; }
 	bool UsesAnisotropy() const { return GetContent()->MaterialCompilationOutput.bUsesAnisotropy; }
+	bool IsLightFunctionAtlasCompatible() const { return GetContent()->MaterialCompilationOutput.bIsLightFunctionAtlasCompatible; }
 
 	const FSubstrateMaterialCompilationOutput& GetSubstrateMaterialCompilationOutput() const { return GetContent()->MaterialCompilationOutput.SubstrateMaterialCompilationOutput; }
 	uint8 GetSubstrateMaterialType() const { return GetSubstrateMaterialCompilationOutput().SubstrateMaterialType; }
@@ -2273,6 +2277,10 @@ public:
 	ENGINE_API bool MaterialUsesAnisotropy_GameThread() const;
 	ENGINE_API bool MaterialUsesAnisotropy_RenderThread() const;
 
+	/** Get whether or not the material is compatible with the light function atlas. */
+	ENGINE_API bool MaterialIsLightFunctionAtlasCompatible_GameThread() const;
+	ENGINE_API bool MaterialIsLightFunctionAtlasCompatible_RenderThread() const;
+
 	/** Get Substrate material type (single, single, complex slab). */
 	ENGINE_API uint8 MaterialGetSubstrateMaterialType_GameThread() const;
 	ENGINE_API uint8 MaterialGetSubstrateMaterialType_RenderThread() const;
diff --git a/Engine/Source/Runtime/Renderer/Private/DeferredShadingRenderer.h b/Engine/Source/Runtime/Renderer/Private/DeferredShadingRenderer.h
--- a/Engine/Source/Runtime/Renderer/Private/DeferredShadingRenderer.h
+++ b/Engine/Source/Runtime/Renderer/Private/DeferredShadingRenderer.h
@@ -975,6 +975,7 @@ private:
 		FRDGTextureRef LightingChannelsTexture,
 		const FHairStrandsTransmittanceMaskData& InTransmittanceMaskData,
 		const bool bForwardRendering,
+		const bool bCanLightUsesAtlasForUnbatchedLight,
 		TRDGUniformBufferRef<FVirtualShadowMapUniformParameters> VirtualShadowMapUniformBuffer = nullptr,
 		FRDGTextureRef ShadowMaskBits = nullptr,
 		int32 VirtualShadowMapId = INDEX_NONE);
diff --git a/Engine/Source/Runtime/Renderer/Private/LightFunctionAtlas.cpp b/Engine/Source/Runtime/Renderer/Private/LightFunctionAtlas.cpp
--- a/Engine/Source/Runtime/Renderer/Private/LightFunctionAtlas.cpp
+++ b/Engine/Source/Runtime/Renderer/Private/LightFunctionAtlas.cpp
@@ -99,7 +99,7 @@ FAutoConsoleVariableRef CVarVolumetricFogUsesLightFunctionAtlas(
 );
 
 // This deferred CVar includes deferred lights splatting (batched or not) as well as clustered lighting.
-int GDeferredUsesLightFunctionAtlas = 0;
+int GDeferredUsesLightFunctionAtlas = 1;
 FAutoConsoleVariableRef CVarDeferredLightsUsesLightFunctionAtlas(
 	TEXT("r.Deferred.UsesLightFunctionAtlas"),
 	GDeferredUsesLightFunctionAtlas,
diff --git a/Engine/Source/Runtime/Renderer/Private/LightRendering.cpp b/Engine/Source/Runtime/Renderer/Private/LightRendering.cpp
--- a/Engine/Source/Runtime/Renderer/Private/LightRendering.cpp
+++ b/Engine/Source/Runtime/Renderer/Private/LightRendering.cpp
@@ -33,6 +33,7 @@
 #include "ManyLights/ManyLights.h"
 #include "LightFunctionAtlas.h"
 #include "HeterogeneousVolumes/HeterogeneousVolumes.h"
+#include "Materials/MaterialRenderProxy.h"
 
 using namespace LightFunctionAtlas;
 
@@ -557,6 +558,7 @@ static void RenderLight(
 	FRDGTextureRef LightingChannelsTexture,
 	bool bRenderOverlap,
 	bool bCloudShadow,
+	const bool bCanLightUsesAtlasForUnbatchedLight = false,
 	TRDGUniformBufferRef<FVirtualShadowMapUniformParameters> VirtualShadowMapUniformBuffer = nullptr,
 	FRDGTextureRef ShadowMaskBits = nullptr,
 	int32 VirtualShadowMapId = INDEX_NONE);
@@ -1143,6 +1145,30 @@ bool FSceneRenderer::AllowSimpleLights() const
 	return bAllowSimpleLights == 1;
 }
 
+bool CanLightUsesAtlasForUnbatchedLight(ERHIFeatureLevel::Type FeatureLevel, const FLightSceneProxy* Proxy)
+{
+	// For now, we prevent directional light to use the light function atlas because atlas tiles needs to be repeatable.
+	// And if a texcoordinate node is not scale as a integer multiplier of the uv in [0,1] then the tile will become visible.
+	if (Proxy->GetLightType() == LightType_Directional)
+	{
+		return false;
+	}
+
+	// Material shoud also be compatible with light function atlas, i.e. not sample depth or world position.
+	// We do not check that for other systems (translucent, water, volume fog, clustered, Lumen), 
+	// because light functions were never available there before the atlas. So those LF aare still added into the atlas.
+	// => If a material is designed to be used with those systems, light function mateirla it must be made be compliant from the start.
+	const FMaterialRenderProxy* MaterialRenderProxy = Proxy ? Proxy->GetLightFunctionMaterial() : nullptr;
+	if (MaterialRenderProxy)
+	{
+		const FMaterial& LFMaterial = Proxy->GetLightFunctionMaterial()->GetIncompleteMaterialWithFallback(FeatureLevel);
+	
+		return LFMaterial.MaterialIsLightFunctionAtlasCompatible_RenderThread();
+	}
+
+	return false;
+}
+
 void FSceneRenderer::GatherAndSortLights(FSortedLightSetSceneInfo& OutSortedLights, bool bShadowedLightsInClustered)
 {
 	if (AllowSimpleLights())
@@ -1200,6 +1226,9 @@ void FSceneRenderer::GatherAndSortLights(FSortedLightSetSceneInfo& OutSortedLigh
 					SortedLightInfo->SortKey.Fields.bLightFunction = ViewFamily.EngineShowFlags.LightFunctions && CheckForLightFunction(LightSceneInfo);
 					SortedLightInfo->SortKey.Fields.bUsesLightingChannels = Views[ViewIndex].bUsesLightingChannels && LightSceneInfo->Proxy->GetLightingChannelMask() != GetDefaultLightingChannelMask();
 
+					// Cached once here and the material is never queried again later.
+					SortedLightInfo->bIsCompatibleWithLightFunctionAtlas = CanLightUsesAtlasForUnbatchedLight(FeatureLevel, LightSceneInfo->Proxy);
+
 					// These are not simple lights.
 					SortedLightInfo->SortKey.Fields.bIsNotSimpleLight = 1;
 
@@ -1280,8 +1309,10 @@ void FSceneRenderer::GatherAndSortLights(FSortedLightSetSceneInfo& OutSortedLigh
 		const bool bDrawShadows = SortedLightInfo.SortKey.Fields.bShadowed;
 		const bool bLightingChannels = SortedLightInfo.SortKey.Fields.bUsesLightingChannels;
 
-		// Do not schedule unbatched lights if the atlas is used and enabled
-		const bool bDrawLightFunction = SortedLightInfo.SortKey.Fields.bLightFunction && !bUseLightFunctionAtlas;
+		// Do not schedule unbatched lights if the atlas is used and enabled.
+		// Keep in mind that when clustered shading is enabled, lights with light function are clustered compatible so are pushed in to the clustered path (if a supported light type).
+		// Otherwise, those light are pushed in the non-shadowed non-lightfunction render path that still supports light function atlas sampling.
+		const bool bDrawLightFunction = SortedLightInfo.SortKey.Fields.bLightFunction && (!bUseLightFunctionAtlas || !SortedLightInfo.bIsCompatibleWithLightFunctionAtlas);
 
 		if (SortedLightInfo.SortKey.Fields.bHandledByManyLights && OutSortedLights.ManyLightsLightStart == SortedLights.Num())
 		{
@@ -1467,8 +1498,10 @@ void FDeferredShadingSceneRenderer::RenderLights(
 				for (int32 LightIndex = StandardDeferredStart; LightIndex < UnbatchedLightStart; LightIndex++)
 				{
 					// Render the light to the scene color buffer, using a 1x1 white texture as input
-					const FLightSceneInfo* LightSceneInfo = SortedLights[LightIndex].LightSceneInfo;
-					RenderLight(GraphBuilder, Scene, View, SceneTextures, LightSceneInfo, nullptr, LightingChannelsTexture, false /*bRenderOverlap*/, false /*bCloudShadow*/);
+					const FSortedLightSceneInfo& SortedLightInfo = SortedLights[LightIndex];
+					const FLightSceneInfo* LightSceneInfo = SortedLightInfo.LightSceneInfo;
+					RenderLight(GraphBuilder, Scene, View, SceneTextures, LightSceneInfo, nullptr, LightingChannelsTexture, 
+						false /*bRenderOverlap*/, false /*bCloudShadow*/, SortedLightInfo.bIsCompatibleWithLightFunctionAtlas);
 				}
 			}
 
@@ -1484,9 +1517,11 @@ void FDeferredShadingSceneRenderer::RenderLights(
 						// Draw non-shadowed non-light function lights without changing render targets between them
 						for (int32 LightIndex = StandardDeferredStart; LightIndex < UnbatchedLightStart; LightIndex++)
 						{
-							const FLightSceneInfo* LightSceneInfo = SortedLights[LightIndex].LightSceneInfo;
-							RenderLightForHair(GraphBuilder, View, SceneTextures, LightSceneInfo, 
-								NullScreenShadowMaskSubPixelTexture, LightingChannelsTexture, DummyTransmittanceMaskData, false /*bForwardRendering*/);
+							const FSortedLightSceneInfo& SortedLightInfo = SortedLights[LightIndex];
+							const FLightSceneInfo* LightSceneInfo = SortedLightInfo.LightSceneInfo;
+							RenderLightForHair(GraphBuilder, View, SceneTextures, LightSceneInfo,
+								NullScreenShadowMaskSubPixelTexture, LightingChannelsTexture, DummyTransmittanceMaskData,
+								false /*bForwardRendering*/, SortedLightInfo.bIsCompatibleWithLightFunctionAtlas);
 						}
 					}
 				}
@@ -1556,7 +1591,7 @@ void FDeferredShadingSceneRenderer::RenderLights(
 				const FLightOcclusionType OcclusionType = GetLightOcclusionType(LightSceneProxy);
 
 				const bool bDrawShadows = SortedLightInfo.SortKey.Fields.bShadowed;
-				const bool bDrawLightFunction = SortedLightInfo.SortKey.Fields.bLightFunction && !bUseLightFunctionAtlas;
+				const bool bDrawLightFunction = SortedLightInfo.SortKey.Fields.bLightFunction && (!bUseLightFunctionAtlas || !SortedLightInfo.bIsCompatibleWithLightFunctionAtlas);
 				const bool bDrawPreviewIndicator = ViewFamily.EngineShowFlags.PreviewShadowsIndicator && !LightSceneInfo.IsPrecomputedLightingValid() && LightSceneProxy.HasStaticShadowing();
 				const bool bDrawHairShadow = bDrawShadows && bUseHairLighting;
 				const bool bUseHairDeepShadow = bDrawShadows && bUseHairLighting && LightSceneProxy.CastsHairStrandsDeepShadow();
@@ -2105,7 +2140,7 @@ void FDeferredShadingSceneRenderer::RenderLights(
 						SCOPED_GPU_MASK(GraphBuilder.RHICmdList, View.GPUMask);
 						RenderLight(
 							GraphBuilder, Scene, View, SceneTextures, &LightSceneInfo, 
-							VirtualShadowMapId != INDEX_NONE ? nullptr : ScreenShadowMaskTexture, LightingChannelsTexture, false /*bRenderOverlap*/, true /*bCloudShadow*/, 
+							VirtualShadowMapId != INDEX_NONE ? nullptr : ScreenShadowMaskTexture, LightingChannelsTexture, false /*bRenderOverlap*/, true /*bCloudShadow*/, SortedLightInfo.bIsCompatibleWithLightFunctionAtlas,
 							VirtualShadowMapArray.GetUniformBuffer(), ShadowSceneRenderer->VirtualShadowMapMaskBits, VirtualShadowMapId);
 					}
 				}
@@ -2150,8 +2185,9 @@ void FDeferredShadingSceneRenderer::RenderLights(
 							// Note: ideally the light should still be evaluated for hair when not casting shadow, but for preserving the old behavior, and not adding 
 							// any perf. regression, we disable this light for hair rendering 
 							RenderLightForHair(
-								GraphBuilder, View, SceneTextures, &LightSceneInfo, 
-								VirtualShadowMapId != INDEX_NONE ? nullptr : HairShadowMask, LightingChannelsTexture, TransmittanceMaskData, false /*bForwardRendering*/,
+								GraphBuilder, View, SceneTextures, &LightSceneInfo,
+								VirtualShadowMapId != INDEX_NONE ? nullptr : HairShadowMask, LightingChannelsTexture, TransmittanceMaskData, 
+								false /*bForwardRendering*/, SortedLightInfo.bIsCompatibleWithLightFunctionAtlas,
 								VirtualShadowMapArray.GetUniformBuffer(), ShadowSceneRenderer->VirtualShadowMapMaskBitsHairStrands, VirtualShadowMapId);
 						}
 					}
@@ -2202,7 +2238,7 @@ static void RenderLightArrayForOverlapViewmode(
 		for (const FViewInfo& View : Views)
 		{
 			SCOPED_GPU_MASK(GraphBuilder.RHICmdList, View.GPUMask);
-			RenderLight(GraphBuilder, Scene, View, SceneTextures, LightSceneInfo, nullptr, LightingChannelsTexture, true /*bRenderOverlap*/, false /*bCloudShadow*/);
+			RenderLight(GraphBuilder, Scene, View, SceneTextures, LightSceneInfo, nullptr, LightingChannelsTexture, true /*bRenderOverlap*/, false /*bCloudShadow*/, false/*bIsCompatibleWithLightFunctionAtlas*/);
 		}
 	}
 }
@@ -2579,6 +2615,7 @@ static void RenderLight(
 	FRDGTextureRef LightingChannelsTexture,
 	bool bRenderOverlap,
 	bool bCloudShadow,
+	const bool bCanLightUsesAtlasForUnbatchedLight,
 	TRDGUniformBufferRef<FVirtualShadowMapUniformParameters> VirtualShadowMapUniformBuffer,
 	FRDGTextureRef ShadowMaskBits,
 	int32 VirtualShadowMapId)
@@ -2676,7 +2713,7 @@ static void RenderLight(
 		PermutationVector.Set< FDeferredLightPS::FHairComplexTransmittance >(bNeedComplexTransmittanceSupport);
 		PermutationVector.Set< FDeferredLightPS::FLightFunctionAtlasDim >(
 			LightFunctionAtlas::IsEnabled(View, ELightFunctionAtlasSystem::DeferredLighting) && LightSceneInfo->Proxy->HasValidLightFunctionAtlasSlot() &&
-			LightSceneInfo->Proxy->GetLightFunctionMaterial() != nullptr && !View.Family->EngineShowFlags.VisualizeLightCulling);
+			LightSceneInfo->Proxy->GetLightFunctionMaterial() != nullptr && !View.Family->EngineShowFlags.VisualizeLightCulling && bCanLightUsesAtlasForUnbatchedLight);
 
 		if (bIsRadial)
 		{
@@ -2759,6 +2796,7 @@ void FDeferredShadingSceneRenderer::RenderLightForHair(
 	FRDGTextureRef LightingChannelsTexture,
 	const FHairStrandsTransmittanceMaskData& InTransmittanceMaskData,
 	const bool bForwardRendering,
+	const bool bCanLightUsesAtlasForUnbatchedLight,
 	TRDGUniformBufferRef<FVirtualShadowMapUniformParameters> VirtualShadowMapUniformBuffer,
 	FRDGTextureRef ShadowMaskBits,
 	int32 VirtualShadowMapId)
@@ -2832,7 +2870,7 @@ void FDeferredShadingSceneRenderer::RenderLightForHair(
 	PermutationVector.Set< FDeferredLightPS::FVirtualShadowMapMask >(bUseVirtualShadowMapMask);
 	PermutationVector.Set< FDeferredLightPS::FLightFunctionAtlasDim >(
 		LightFunctionAtlas::IsEnabled(View, ELightFunctionAtlasSystem::DeferredLighting) && LightSceneInfo->Proxy->HasValidLightFunctionAtlasSlot() &&
-		LightSceneInfo->Proxy->GetLightFunctionMaterial() != nullptr && !View.Family->EngineShowFlags.VisualizeLightCulling);
+		LightSceneInfo->Proxy->GetLightFunctionMaterial() != nullptr && !View.Family->EngineShowFlags.VisualizeLightCulling && bCanLightUsesAtlasForUnbatchedLight);
 	if (bIsDirectional)
 	{
 		PermutationVector.Set< FDeferredLightPS::FSourceShapeDim >(ELightSourceShape::Directional);
@@ -2931,7 +2969,8 @@ void FDeferredShadingSceneRenderer::RenderLightsForHair(
 						ScreenShadowMaskSubPixelTexture,
 						LightingChannelsTexture,
 						TransmittanceMaskData,
-						true /*bForwardRendering*/);
+						true /*bForwardRendering*/,
+						SortedLightInfo.bIsCompatibleWithLightFunctionAtlas);
 				}
 			}
 		}
diff --git a/Engine/Source/Runtime/Renderer/Private/LightSceneInfo.h b/Engine/Source/Runtime/Renderer/Private/LightSceneInfo.h
--- a/Engine/Source/Runtime/Renderer/Private/LightSceneInfo.h
+++ b/Engine/Source/Runtime/Renderer/Private/LightSceneInfo.h
@@ -106,11 +106,13 @@ struct FSortedLightSceneInfo
 
 	const FLightSceneInfo* LightSceneInfo;
 	int32 SimpleLightIndex;
+	bool bIsCompatibleWithLightFunctionAtlas;
 
 	/** Initialization constructor. */
 	explicit FSortedLightSceneInfo(const FLightSceneInfo* InLightSceneInfo)
 		: LightSceneInfo(InLightSceneInfo),
-		SimpleLightIndex(-1)
+		SimpleLightIndex(-1),
+		bIsCompatibleWithLightFunctionAtlas(false)
 	{
 		SortKey.Packed = 0;
 		SortKey.Fields.bIsNotSimpleLight = 1;
@@ -118,7 +120,8 @@ struct FSortedLightSceneInfo
 
 	explicit FSortedLightSceneInfo(int32 InSimpleLightIndex)
 		: LightSceneInfo(nullptr),
-		SimpleLightIndex(InSimpleLightIndex)
+		SimpleLightIndex(InSimpleLightIndex),
+		bIsCompatibleWithLightFunctionAtlas(false)
 	{
 		SortKey.Packed = 0;
 		SortKey.Fields.bIsNotSimpleLight = 0;
