Fixed seemingly missing scene uniform buffer in volumetric cloud rendering compute.

diff --git a/Engine/Source/Runtime/Renderer/Private/VolumetricCloudRendering.cpp b/Engine/Source/Runtime/Renderer/Private/VolumetricCloudRendering.cpp
--- a/Engine/Source/Runtime/Renderer/Private/VolumetricCloudRendering.cpp
+++ b/Engine/Source/Runtime/Renderer/Private/VolumetricCloudRendering.cpp
@@ -26,6 +26,7 @@
 #include "SkyAtmosphereRendering.h"
 #include "BasePassRendering.h"
 #include "EnvironmentComponentsFlags.h"
+#include "SceneRendererInterface.h"
 
 //  If this is enabled, you also need to touch VolumetricCloud.usf for shaders to be recompiled.
 #define CLOUD_DEBUG_SAMPLES 0 /*!Never check in enabled!*/
@@ -840,6 +841,7 @@ class FRenderVolumetricCloudRenderViewCS : public FMeshMaterialShader
 		SHADER_PARAMETER(int32, TargetCubeFace)
 		SHADER_PARAMETER_RDG_UNIFORM_BUFFER(FRenderVolumetricCloudGlobalParameters, VolumetricCloudRenderViewParamsUB)
 		SHADER_PARAMETER_STRUCT_REF(FViewUniformShaderParameters, View)
+		SHADER_PARAMETER_RDG_UNIFORM_BUFFER(FSceneUniformParameters, Scene)
 		SHADER_PARAMETER_STRUCT_INCLUDE(FVirtualShadowMapSamplingParameters, VirtualShadowMap)
 		SHADER_PARAMETER(int32, VirtualShadowMapId0)
 		SHADER_PARAMETER_RDG_TEXTURE_UAV(RWTexture2D, OutCloudColor0)
@@ -2418,6 +2420,7 @@ void FSceneRenderer::RenderVolumetricCloudsInternal(FRDGBuilder& GraphBuilder, F
 		PassParameters->TargetCubeFace = CloudRC.RenderTargets.Output[0].GetArraySlice();
 		PassParameters->VolumetricCloudRenderViewParamsUB = CloudPassUniformBuffer;
 		PassParameters->View = ViewUniformBuffer;
+		PassParameters->Scene = GetSceneUniformBufferRef(GraphBuilder);
 		if (bSampleVirtualShadowMap)
 		{
 			PassParameters->VirtualShadowMap = VirtualShadowMapArray.GetSamplingParameters(GraphBuilder);
