VolumetricCloud - added scene UB to all material shaders in case it is required by some node at some point to not crash the editor.

diff --git a/Engine/Source/Runtime/Renderer/Private/VolumetricCloudRendering.cpp b/Engine/Source/Runtime/Renderer/Private/VolumetricCloudRendering.cpp
--- a/Engine/Source/Runtime/Renderer/Private/VolumetricCloudRendering.cpp
+++ b/Engine/Source/Runtime/Renderer/Private/VolumetricCloudRendering.cpp
@@ -732,6 +732,7 @@ BEGIN_SHADER_PARAMETER_STRUCT(FRenderVolumetricCloudRenderViewParametersPS, )
 	SHADER_PARAMETER_RDG_TEXTURE(Texture2D, CloudShadowTexture0)
 	SHADER_PARAMETER_RDG_TEXTURE(Texture2D, CloudShadowTexture1)
 	SHADER_PARAMETER_STRUCT_REF(FViewUniformShaderParameters, View)
+	SHADER_PARAMETER_RDG_UNIFORM_BUFFER(FSceneUniformParameters, Scene)
 	SHADER_PARAMETER_STRUCT_INCLUDE(FVirtualShadowMapSamplingParameters, VirtualShadowMap)
 	SHADER_PARAMETER_RDG_UNIFORM_BUFFER(FInstanceCullingGlobalUniforms, InstanceCulling)
 	RENDER_TARGET_BINDING_SLOTS()
@@ -932,6 +933,7 @@ class FRenderVolumetricCloudEmptySpaceSkippingCS : public FMeshMaterialShader
 		SHADER_PARAMETER(float, StartTracingSliceBias)
 		SHADER_PARAMETER_RDG_UNIFORM_BUFFER(FRenderVolumetricCloudGlobalParameters, VolumetricCloudRenderViewParamsUB)
 		SHADER_PARAMETER_STRUCT_REF(FViewUniformShaderParameters, View)
+		SHADER_PARAMETER_RDG_UNIFORM_BUFFER(FSceneUniformParameters, Scene)
 		SHADER_PARAMETER_RDG_TEXTURE_UAV(RWTexture2D, OutStartTracingDistanceTexture)
 //		SHADER_PARAMETER_STRUCT_INCLUDE(ShaderPrint::FShaderParameters, ShaderPrintParameters)
 	END_SHADER_PARAMETER_STRUCT()
@@ -1273,6 +1275,7 @@ private:
 
 BEGIN_SHADER_PARAMETER_STRUCT(FVolumetricCloudShadowParametersPS, )
 	SHADER_PARAMETER_STRUCT_INCLUDE(FViewShaderParameters, View)
+	SHADER_PARAMETER_RDG_UNIFORM_BUFFER(FSceneUniformParameters, Scene)
 	SHADER_PARAMETER_RDG_UNIFORM_BUFFER(FRenderVolumetricCloudGlobalParameters, TraceVolumetricCloudParamsUB)
 	SHADER_PARAMETER_RDG_UNIFORM_BUFFER(FInstanceCullingGlobalUniforms, InstanceCulling)
 	RENDER_TARGET_BINDING_SLOTS()
@@ -1850,6 +1853,7 @@ void FSceneRenderer::InitVolumetricCloudsForViews(FRDGBuilder& GraphBuilder, boo
 					{
 						FVolumetricCloudShadowParametersPS* CloudShadowParameters = GraphBuilder.AllocParameters<FVolumetricCloudShadowParametersPS>();
 						CloudShadowParameters->View = ViewInfo.GetShaderParameters();
+						CloudShadowParameters->Scene = GetSceneUniformBufferRef(GraphBuilder);
 						CloudShadowParameters->TraceVolumetricCloudParamsUB = TraceVolumetricCloudParamsUB;
 						CloudShadowParameters->InstanceCulling = InstanceCullingManager.GetDummyInstanceCullingUniformBuffer();
 						CloudShadowParameters->RenderTargets[0] = FRenderTargetBinding(CloudTextureTracedOutput, ERenderTargetLoadAction::ENoAction);
@@ -2352,6 +2356,7 @@ void FSceneRenderer::RenderVolumetricCloudsInternal(FRDGBuilder& GraphBuilder, F
 			FRenderVolumetricCloudEmptySpaceSkippingCS::FParameters* PassParameters = GraphBuilder.AllocParameters<FRenderVolumetricCloudEmptySpaceSkippingCS::FParameters>();
 			PassParameters->VolumetricCloudRenderViewParamsUB = CloudPassUniformBuffer;
 			PassParameters->View = ViewUniformBuffer;
+			PassParameters->Scene = GetSceneUniformBufferRef(GraphBuilder);
 			PassParameters->OutStartTracingDistanceTexture = GraphBuilder.CreateUAV(FRDGTextureUAVDesc(EmptySpaceSkippingTexture));
 			PassParameters->StartTracingDistanceTextureResolution = FVector2f(EmptySpaceSkippingTextureResolution.X, EmptySpaceSkippingTextureResolution.Y);
 			PassParameters->StartTracingSampleVolumeDepth = GetEmptySpaceSkippingVolumeDepth();
@@ -2491,6 +2496,7 @@ void FSceneRenderer::RenderVolumetricCloudsInternal(FRDGBuilder& GraphBuilder, F
 		RenderViewPassParameters->CloudShadowTexture0 = CloudRC.VolumetricCloudShadowTexture[0];
 		RenderViewPassParameters->CloudShadowTexture1 = CloudRC.VolumetricCloudShadowTexture[1];
 		RenderViewPassParameters->View = ViewUniformBuffer;
+		RenderViewPassParameters->Scene = GetSceneUniformBufferRef(GraphBuilder);
 		RenderViewPassParameters->VirtualShadowMap = VirtualShadowMapArray.GetSamplingParameters(GraphBuilder);
 		RenderViewPassParameters->InstanceCulling = InstanceCullingManager.GetDummyInstanceCullingUniformBuffer();
 		RenderViewPassParameters->RenderTargets = CloudRC.RenderTargets;
