Fix for static analysis valid warning.

diff --git a/Engine/Source/Runtime/Renderer/Private/LightRendering.cpp b/Engine/Source/Runtime/Renderer/Private/LightRendering.cpp
--- a/Engine/Source/Runtime/Renderer/Private/LightRendering.cpp
+++ b/Engine/Source/Runtime/Renderer/Private/LightRendering.cpp
@@ -1147,6 +1147,11 @@ bool FSceneRenderer::AllowSimpleLights() const
 
 bool CanLightUsesAtlasForUnbatchedLight(ERHIFeatureLevel::Type FeatureLevel, const FLightSceneProxy* Proxy)
 {
+	if (!Proxy)
+	{
+		return false;
+	}
+
 	// For now, we prevent directional light to use the light function atlas because atlas tiles needs to be repeatable.
 	// And if a texcoordinate node is not scale as a integer multiplier of the uv in [0,1] then the tile will become visible.
 	if (Proxy->GetLightType() == LightType_Directional)
@@ -1158,7 +1163,7 @@ bool CanLightUsesAtlasForUnbatchedLight(ERHIFeatureLevel::Type FeatureLevel, con
 	// We do not check that for other systems (translucent, water, volume fog, clustered, Lumen), 
 	// because light functions were never available there before the atlas. So those LF aare still added into the atlas.
 	// => If a material is designed to be used with those systems, light function mateirla it must be made be compliant from the start.
-	const FMaterialRenderProxy* MaterialRenderProxy = Proxy ? Proxy->GetLightFunctionMaterial() : nullptr;
+	const FMaterialRenderProxy* MaterialRenderProxy = Proxy->GetLightFunctionMaterial();
 	if (MaterialRenderProxy)
 	{
 		const FMaterial& LFMaterial = Proxy->GetLightFunctionMaterial()->GetIncompleteMaterialWithFallback(FeatureLevel);
