Set sub-pass information for a decal material PSO pre-caching

diff --git a/Engine/Source/Runtime/Renderer/Private/CompositionLighting/PostProcessDeferredDecals.cpp b/Engine/Source/Runtime/Renderer/Private/CompositionLighting/PostProcessDeferredDecals.cpp
--- a/Engine/Source/Runtime/Renderer/Private/CompositionLighting/PostProcessDeferredDecals.cpp
+++ b/Engine/Source/Runtime/Renderer/Private/CompositionLighting/PostProcessDeferredDecals.cpp
@@ -518,7 +518,14 @@ void CollectDeferredDecalPassPSOInitializers(
 	RenderTargetsInfo.NumSamples = 1;
 	GetDeferredDecalRenderTargetsInfo(SceneTexturesConfig, ShaderPlatform, DecalRenderTargetMode, RenderTargetsInfo);
 	ApplyTargetsInfo(GraphicsPSOInit, RenderTargetsInfo);
-	
+
+	if (FeatureLevel == ERHIFeatureLevel::ES3_1)
+	{
+		// subpass info set during the submission of the draws in a mobile renderer
+		GraphicsPSOInit.SubpassIndex = 1; // all decals use second sub-pass on mobile
+		GraphicsPSOInit.SubpassHint = SceneTexturesConfig.bIsUsingGBuffers ? ESubpassHint::DeferredShadingSubpass : ESubpassHint::DepthReadSubpass;
+	}
+		
 	const auto AddDeferredDecalPSO = [&](bool bInsideDecal,	bool bReverseHanded, bool bReverseCulling, bool bDecalUsesStencil)
 	{
 		const EDecalRasterizerState DecalRasterizerState = DecalRendering::GetDecalRasterizerState(bInsideDecal, bReverseHanded, bReverseCulling);
diff --git a/Engine/Source/Runtime/Renderer/Private/CompositionLighting/PostProcessMeshDecals.cpp b/Engine/Source/Runtime/Renderer/Private/CompositionLighting/PostProcessMeshDecals.cpp
--- a/Engine/Source/Runtime/Renderer/Private/CompositionLighting/PostProcessMeshDecals.cpp
+++ b/Engine/Source/Runtime/Renderer/Private/CompositionLighting/PostProcessMeshDecals.cpp
@@ -389,6 +389,15 @@ void FMeshDecalMeshProcessor::CollectDeferredDecalMeshPSOInitializers(
 	RenderTargetsInfo.NumSamples = 1;
 	GetDeferredDecalRenderTargetsInfo(SceneTexturesConfig, ShaderPlatform, LocalRenderTargetMode, RenderTargetsInfo);
 
+	uint8 SubpassIndex = 0;
+	ESubpassHint SubpassHint = ESubpassHint::None;
+	if (FeatureLevel == ERHIFeatureLevel::ES3_1)
+	{
+		// subpass info set during the submission of the draws in a mobile renderer
+		SubpassIndex = 1; // all decals use second sub-pass on mobile
+		SubpassHint = SceneTexturesConfig.bIsUsingGBuffers ? ESubpassHint::DeferredShadingSubpass : ESubpassHint::DepthReadSubpass;
+	}
+
 	AddGraphicsPipelineStateInitializer(
 		VertexFactoryData,
 		Material,
@@ -399,7 +408,10 @@ void FMeshDecalMeshProcessor::CollectDeferredDecalMeshPSOInitializers(
 		MeshCullMode,
 		PT_TriangleList,
 		EMeshPassFeatures::Default,
+		SubpassHint,
+		SubpassIndex,
 		true /*bRequired*/,
+		PSOCollectorIndex,
 		PSOInitializers);
 }
 
