Flag debug visualization mode permutation of FVirtualShadowMapProjectionCS as NotPrecached instead of NotUsed

diff --git a/Engine/Source/Runtime/Renderer/Private/VirtualShadowMaps/VirtualShadowMapProjection.cpp b/Engine/Source/Runtime/Renderer/Private/VirtualShadowMaps/VirtualShadowMapProjection.cpp
--- a/Engine/Source/Runtime/Renderer/Private/VirtualShadowMaps/VirtualShadowMapProjection.cpp
+++ b/Engine/Source/Runtime/Renderer/Private/VirtualShadowMaps/VirtualShadowMapProjection.cpp
@@ -322,7 +322,7 @@ class FVirtualShadowMapProjectionCS : public FGlobalShader
 		FPermutationDomain PermutationVector(Parameters.PermutationId);
 		if (PermutationVector.Get<FVisualizeOutputDim>())
 		{
-			return EShaderPermutationPrecacheRequest::NotUsed;
+			return EShaderPermutationPrecacheRequest::NotPrecached;
 		}
 
 		return EShaderPermutationPrecacheRequest::Precached;
