Fix for incorrect marking of unused permutation when culling type was not set as DFS_DirectionalLightScatterTileCulling

diff --git a/Engine/Source/Runtime/Renderer/Private/DistanceFieldShadowing.cpp b/Engine/Source/Runtime/Renderer/Private/DistanceFieldShadowing.cpp
--- a/Engine/Source/Runtime/Renderer/Private/DistanceFieldShadowing.cpp
+++ b/Engine/Source/Runtime/Renderer/Private/DistanceFieldShadowing.cpp
@@ -353,7 +353,7 @@ class FDistanceFieldShadowingCS : public FGlobalShader
 			return EShaderPermutationPrecacheRequest::NotUsed;
 		}
 
-		if (PermutationVector.Get<FCompactCulledObjects>() != (GDFShadowCompactCulledObjects  != 0))
+		if (PermutationVector.Get<FCompactCulledObjects>() != (GDFShadowCompactCulledObjects != 0 && (PermutationVector.Get<FCullingType>() == DFS_DirectionalLightScatterTileCulling)))
 		{
 			return EShaderPermutationPrecacheRequest::NotUsed;
 		}
