- Fix for early view init & stereo rendering

diff --git a/Engine/Source/Runtime/Renderer/Private/MobileShadingRenderer.cpp b/Engine/Source/Runtime/Renderer/Private/MobileShadingRenderer.cpp
--- a/Engine/Source/Runtime/Renderer/Private/MobileShadingRenderer.cpp
+++ b/Engine/Source/Runtime/Renderer/Private/MobileShadingRenderer.cpp
@@ -462,6 +462,12 @@ void FMobileSceneRenderer::InitViews(
 	{
 		// This is to init the ViewUniformBuffer before rendering for the Niagara compute shader.
 		// This needs to run before ComputeViewVisibility() is called, but the views normally initialize the ViewUniformBuffer after that (at the end of this method).
+
+		// during ISR, instanced view RHI resources need to be initialized first.
+		if (FViewInfo* InstancedView = const_cast<FViewInfo*>(Views[0].GetInstancedView()))
+		{
+			InstancedView->InitRHIResources();
+		}
 		Views[0].InitRHIResources();
 		FXSystem->PostInitViews(GraphBuilder, GetSceneViews(), !ViewFamily.EngineShowFlags.HitProxies);
 	}
