Fix or suppress Clang 18 issues, support Intel compiler 2024.1

diff --git a/Engine/Config/Windows/Windows_SDK.json b/Engine/Config/Windows/Windows_SDK.json
--- a/Engine/Config/Windows/Windows_SDK.json
+++ b/Engine/Config/Windows/Windows_SDK.json
@@ -35,9 +35,9 @@
 	"MinimumVisualCppVersion": "14.34.31933",
 
 	"PreferredIntelOneApiVersions": [
-		"2024.0.0-2024.9999"
+		"2024.1.0-2024.9999"
 	],
-	"MinimumIntelOneApiVersion": "2024.0.0",
+	"MinimumIntelOneApiVersion": "2024.1.0",
 
 	"VisualStudioSuggestedComponents": [
 		"Microsoft.Net.Component.4.6.2.TargetingPack",
diff --git a/Engine/Plugins/Experimental/TextureGraph/Source/TextureGraph/Public/TG_HelperFunctions.h b/Engine/Plugins/Experimental/TextureGraph/Source/TextureGraph/Public/TG_HelperFunctions.h
--- a/Engine/Plugins/Experimental/TextureGraph/Source/TextureGraph/Public/TG_HelperFunctions.h
+++ b/Engine/Plugins/Experimental/TextureGraph/Source/TextureGraph/Public/TG_HelperFunctions.h
@@ -7,6 +7,7 @@
 #include "TextureGraph.h"
 #include "Export/TextureExporter.h"
 #include "Data/Blob.h"
+#include "TG_Graph.h"
 #include "TG_Node.h"
 #include "2D/TextureHelper.h"
 
diff --git a/Engine/Plugins/Runtime/MassGameplay/Source/MassRepresentation/Public/MassRepresentationTypes.h b/Engine/Plugins/Runtime/MassGameplay/Source/MassRepresentation/Public/MassRepresentationTypes.h
--- a/Engine/Plugins/Runtime/MassGameplay/Source/MassRepresentation/Public/MassRepresentationTypes.h
+++ b/Engine/Plugins/Runtime/MassGameplay/Source/MassRepresentation/Public/MassRepresentationTypes.h
@@ -681,7 +681,7 @@ public:
 		}
 	}
 
-	FORCEINLINE void WriteCustomDataFloatsAtStartIndex(int32 StaticMeshIndex, const TArrayView<float>& CustomFloats, const float LODSignificance, const int32 FloatsPerInstance, const int32 FloatStartIndex, const float PrevLODSignificance = -1.0f)
+	void WriteCustomDataFloatsAtStartIndex(int32 StaticMeshIndex, const TArrayView<float>& CustomFloats, const float LODSignificance, const int32 FloatsPerInstance, const int32 FloatStartIndex, const float PrevLODSignificance = -1.0f)
 	{
 		if (FMassLODSignificanceRange* Range = GetLODSignificanceRange(LODSignificance))
 		{
diff --git a/Engine/Source/Programs/UnrealBuildTool/ToolChain/ClangWarnings.cs b/Engine/Source/Programs/UnrealBuildTool/ToolChain/ClangWarnings.cs
--- a/Engine/Source/Programs/UnrealBuildTool/ToolChain/ClangWarnings.cs
+++ b/Engine/Source/Programs/UnrealBuildTool/ToolChain/ClangWarnings.cs
@@ -52,6 +52,11 @@ namespace UnrealBuildTool
 					Arguments.Add("-Wno-invalid-unevaluated-string");   // https://clang.llvm.org/docs/DiagnosticsReference.html#winvalid-unevaluated-string			// new warning for clang 17
 				}
 			}
+			if (ClangVersion >= new VersionNumber(18))
+			{
+				Arguments.Add("-Wno-deprecated-this-capture");          // https://clang.llvm.org/docs/DiagnosticsReference.html#wdeprecated-this-capture
+				Arguments.Add("-Wno-enum-constexpr-conversion");        // https://clang.llvm.org/docs/DiagnosticsReference.html#wenum-constexpr-conversion
+			}
 
 			Arguments.Add("-Wno-gnu-string-literal-operator-template"); // https://clang.llvm.org/docs/DiagnosticsReference.html#wgnu-string-literal-operator-template	// We use this feature to allow static FNames.
 			Arguments.Add("-Wno-inconsistent-missing-override");        // https://clang.llvm.org/docs/DiagnosticsReference.html#winconsistent-missing-override			// ?? no reason given
diff --git a/Engine/Source/Runtime/D3D12RHI/Private/D3D12CommandContext.h b/Engine/Source/Runtime/D3D12RHI/Private/D3D12CommandContext.h
--- a/Engine/Source/Runtime/D3D12RHI/Private/D3D12CommandContext.h
+++ b/Engine/Source/Runtime/D3D12RHI/Private/D3D12CommandContext.h
@@ -718,10 +718,23 @@ public:
 		uint32 NumBindings, const FRayTracingLocalShaderBindings* Bindings, ERayTracingBindingType BindingType) final override;
 #endif // D3D12_RHI_RAYTRACING
 
+	template<typename TRHIType, typename TReturnType = typename TD3D12ResourceTraits<TRHIType>::TConcreteType>
+	static FORCEINLINE TReturnType* ResourceCast(TRHIType* Resource)
+	{
+		return static_cast<TReturnType*>(Resource);
+	}
+
+	template<typename TRHIType, typename TReturnType = typename TD3D12ResourceTraits<TRHIType>::TConcreteType>
+	static FORCEINLINE_DEBUGGABLE TReturnType* ResourceCast(TRHIType* Resource, uint32 GPUIndex)
+	{
+		TReturnType* Object = ResourceCast<TRHIType, TReturnType>(Resource);
+		return Object ? static_cast<TReturnType*>(Object->GetLinkedObject(GPUIndex)) : nullptr;
+	}
+
 	template<typename ObjectType, typename RHIType>
 	static FORCEINLINE_DEBUGGABLE ObjectType* RetrieveObject(RHIType* RHIObject, uint32 GPUIndex)
 	{
-		return FD3D12DynamicRHI::ResourceCast<RHIType, ObjectType>(RHIObject, GPUIndex);
+		return ResourceCast<RHIType, ObjectType>(RHIObject, GPUIndex);
 	}
 
 	template<typename ObjectType, typename RHIType>
diff --git a/Engine/Source/Runtime/Engine/Public/TraceFilter.h b/Engine/Source/Runtime/Engine/Public/TraceFilter.h
--- a/Engine/Source/Runtime/Engine/Public/TraceFilter.h
+++ b/Engine/Source/Runtime/Engine/Public/TraceFilter.h
@@ -84,7 +84,7 @@ struct FTraceFilter
 
 	/** Returns whether or not an object is eligible to be outputted (output trace data) */
 	template<bool bForceThreadSafe = true>
-	static bool IsObjectTraceable(const UObject* InObject);
+	static ENGINE_API bool IsObjectTraceable(const UObject* InObject);
 
 	static ENGINE_API void Lock();
 	static ENGINE_API void Unlock();
diff --git a/Engine/Source/Runtime/Renderer/Private/BasePassRendering.cpp b/Engine/Source/Runtime/Renderer/Private/BasePassRendering.cpp
--- a/Engine/Source/Runtime/Renderer/Private/BasePassRendering.cpp
+++ b/Engine/Source/Runtime/Renderer/Private/BasePassRendering.cpp
@@ -107,13 +107,6 @@ static TAutoConsoleVariable<int32> CVarPSOPrecacheTranslucencyAllPass(
 	ECVF_ReadOnly
 );
 
-static TAutoConsoleVariable<int32> CVarPSOPrecacheAlphaColorChannel(
-	TEXT("r.PSOPrecache.PrecacheAlphaColorChannel"),
-	1,
-	TEXT("Also Precache PSOs with scene color alpha channel enabled. Planar reflections and scene captures use this for compositing into a different scene later."),
-	ECVF_ReadOnly
-);
-
 // Scene color alpha is used during scene captures and planar reflections.  1 indicates background should be shown, 0 indicates foreground is fully present.
 static const float kSceneColorClearAlpha = 1.0f;
 
@@ -1776,69 +1769,6 @@ void FBasePassMeshProcessor::CollectPSOInitializersForLMPolicy(
 	}	
 }
 
-template<typename PassShadersType>
-void FBasePassMeshProcessor::AddBasePassGraphicsPipelineStateInitializer(
-	ERHIFeatureLevel::Type InFeatureLevel,
-	const FPSOPrecacheVertexFactoryData& VertexFactoryData,
-	const FMaterial& RESTRICT MaterialResource,
-	const FMeshPassProcessorRenderState& RESTRICT DrawRenderState,
-	const FGraphicsPipelineRenderTargetsInfo& RESTRICT RenderTargetsInfo,
-	const PassShadersType& PassShaders,
-	ERasterizerFillMode MeshFillMode,
-	ERasterizerCullMode MeshCullMode,
-	EPrimitiveType PrimitiveType,
-	bool bPrecacheAlphaColorChannel,
-	int InPSOCollectorIndex,
-	TArray<FPSOPrecacheData>& PSOInitializers)
-{
-	AddGraphicsPipelineStateInitializer(
-		VertexFactoryData,
-		MaterialResource,
-		DrawRenderState,
-		RenderTargetsInfo,
-		PassShaders,
-		MeshFillMode,
-		MeshCullMode,
-		PrimitiveType,
-		EMeshPassFeatures::Default,
-		ESubpassHint::None,
-		0,
-		true /*bRequired*/,
-		InPSOCollectorIndex,
-		PSOInitializers);
-
-	// Planar reflections and scene captures use scene color alpha to keep track of where content has been rendered, for compositing into a different scene later
-	if (bPrecacheAlphaColorChannel && CVarPSOPrecacheAlphaColorChannel.GetValueOnAnyThread() > 0)
-	{
-		FGraphicsPipelineRenderTargetsInfo AlphaColorRenderTargetsInfo = RenderTargetsInfo;
-
-		bool bRequiresAlphaChannel = true;
-		ETextureCreateFlags ExtraSceneColorCreateFlags = ETextureCreateFlags::None;
-		EPixelFormat SceneColorFormatWithAlpha;
-		ETextureCreateFlags SceneColorCreateFlagsWithAlpha;
-		GetSceneColorFormatAndCreateFlags(InFeatureLevel, bRequiresAlphaChannel, ExtraSceneColorCreateFlags, RenderTargetsInfo.NumSamples, false, SceneColorFormatWithAlpha, SceneColorCreateFlagsWithAlpha);
-
-		AlphaColorRenderTargetsInfo.RenderTargetFormats[0] = SceneColorFormatWithAlpha;
-		AlphaColorRenderTargetsInfo.RenderTargetFlags[0] = SceneColorCreateFlagsWithAlpha;
-
-		AddGraphicsPipelineStateInitializer(
-			VertexFactoryData,
-			MaterialResource,
-			DrawRenderState,
-			AlphaColorRenderTargetsInfo,
-			PassShaders,
-			MeshFillMode,
-			MeshCullMode,
-			PrimitiveType,
-			EMeshPassFeatures::Default,
-			ESubpassHint::None,
-			0,
-			true,
-			InPSOCollectorIndex,
-			PSOInitializers);
-	}
-}
-
 template<typename LightMapPolicyType>
 bool FBasePassMeshProcessor::Process(
 	const FMeshBatch& RESTRICT MeshBatch,
diff --git a/Engine/Source/Runtime/Renderer/Private/BasePassRendering.h b/Engine/Source/Runtime/Renderer/Private/BasePassRendering.h
--- a/Engine/Source/Runtime/Renderer/Private/BasePassRendering.h
+++ b/Engine/Source/Runtime/Renderer/Private/BasePassRendering.h
@@ -75,6 +75,13 @@ struct FForwardBasePassTextures
 	bool bIs24BitUnormDepthStencil = false;
 };
 
+static TAutoConsoleVariable<int32> CVarPSOPrecacheAlphaColorChannel(
+	TEXT("r.PSOPrecache.PrecacheAlphaColorChannel"),
+	1,
+	TEXT("Also Precache PSOs with scene color alpha channel enabled. Planar reflections and scene captures use this for compositing into a different scene later."),
+	ECVF_ReadOnly
+);
+
 BEGIN_GLOBAL_SHADER_PARAMETER_STRUCT(FSharedBasePassUniformParameters,)
 	SHADER_PARAMETER_STRUCT(FForwardLightData, Forward)
 	SHADER_PARAMETER_STRUCT(FForwardLightData, ForwardISR)
@@ -880,7 +887,55 @@ public:
 		EPrimitiveType PrimitiveType,
 		bool bPrecacheAlphaColorChannel,
 		int InPSOCollectorIndex,
-		TArray<FPSOPrecacheData>& PSOInitializers);
+		TArray<FPSOPrecacheData>& PSOInitializers)
+	{
+		AddGraphicsPipelineStateInitializer(
+			VertexFactoryData,
+			MaterialResource,
+			DrawRenderState,
+			RenderTargetsInfo,
+			PassShaders,
+			MeshFillMode,
+			MeshCullMode,
+			PrimitiveType,
+			EMeshPassFeatures::Default,
+			ESubpassHint::None,
+			0,
+			true /*bRequired*/,
+			InPSOCollectorIndex,
+			PSOInitializers);
+
+		// Planar reflections and scene captures use scene color alpha to keep track of where content has been rendered, for compositing into a different scene later
+		if (bPrecacheAlphaColorChannel && CVarPSOPrecacheAlphaColorChannel.GetValueOnAnyThread() > 0)
+		{
+			FGraphicsPipelineRenderTargetsInfo AlphaColorRenderTargetsInfo = RenderTargetsInfo;
+
+			bool bRequiresAlphaChannel = true;
+			ETextureCreateFlags ExtraSceneColorCreateFlags = ETextureCreateFlags::None;
+			EPixelFormat SceneColorFormatWithAlpha;
+			ETextureCreateFlags SceneColorCreateFlagsWithAlpha;
+			GetSceneColorFormatAndCreateFlags(InFeatureLevel, bRequiresAlphaChannel, ExtraSceneColorCreateFlags, RenderTargetsInfo.NumSamples, false, SceneColorFormatWithAlpha, SceneColorCreateFlagsWithAlpha);
+
+			AlphaColorRenderTargetsInfo.RenderTargetFormats[0] = SceneColorFormatWithAlpha;
+			AlphaColorRenderTargetsInfo.RenderTargetFlags[0] = SceneColorCreateFlagsWithAlpha;
+
+			AddGraphicsPipelineStateInitializer(
+				VertexFactoryData,
+				MaterialResource,
+				DrawRenderState,
+				AlphaColorRenderTargetsInfo,
+				PassShaders,
+				MeshFillMode,
+				MeshCullMode,
+				PrimitiveType,
+				EMeshPassFeatures::Default,
+				ESubpassHint::None,
+				0,
+				true,
+				InPSOCollectorIndex,
+				PSOInitializers);
+		}
+	}
 
 private:
 
diff --git a/Engine/Source/Runtime/Windows/D3D11RHI/Private/D3D11StateCachePrivate.h b/Engine/Source/Runtime/Windows/D3D11RHI/Private/D3D11StateCachePrivate.h
--- a/Engine/Source/Runtime/Windows/D3D11RHI/Private/D3D11StateCachePrivate.h
+++ b/Engine/Source/Runtime/Windows/D3D11RHI/Private/D3D11StateCachePrivate.h
@@ -383,42 +383,6 @@ public:
 		InternalSetSamplerState<ShaderFrequency>(SamplerState, SamplerIndex);
 	}
 
-	template <EShaderFrequency ShaderFrequency>
-	D3D11_STATE_CACHE_INLINE void GetSamplerState(uint32 StartSamplerIndex, uint32 NumSamplerIndexes, ID3D11SamplerState** SamplerStates)
-	{
-#if D3D11_ALLOW_STATE_CACHE
-		{
-			check(StartSamplerIndex + NumSamplerIndexes <= D3D11_COMMONSHADER_INPUT_RESOURCE_SLOT_COUNT);
-			for (uint32 StateLoop = 0; StateLoop < NumSamplerIndexes; StateLoop++)
-			{
-				SamplerStates[StateLoop] = CurrentShaderResourceViews[ShaderFrequency][StateLoop + StartSamplerIndex];
-				if (SamplerStates[StateLoop])
-				{
-					SamplerStates[StateLoop]->AddRef();
-				}
-			}
-		}
-#else
-		{
-			switch (ShaderFrequency)
-			{
-				case SF_Vertex:		
-					Direct3DDeviceIMContext->VSGetSamplers(StartSamplerIndex, NumSamplerIndexes, SamplerStates); 
-					break;
-				case SF_Geometry:	
-					Direct3DDeviceIMContext->GSGetSamplers(StartSamplerIndex, NumSamplerIndexes, SamplerStates); 
-					break;
-				case SF_Pixel:		
-					Direct3DDeviceIMContext->PSGetSamplers(StartSamplerIndex, NumSamplerIndexes, SamplerStates); 
-					break;
-				case SF_Compute:	
-					Direct3DDeviceIMContext->CSGetSamplers(StartSamplerIndex, NumSamplerIndexes, SamplerStates); 
-					break;
-
-			}
-		}
-#endif
-	}
 	template <EShaderFrequency ShaderFrequency>
 	D3D11_STATE_CACHE_INLINE void SetConstantBuffer(ID3D11Buffer* ConstantBuffer, uint32 SlotIndex)
 	{
