Removing CVar from header

#rnx

diff --git a/Engine/Source/Runtime/Renderer/Private/BasePassRendering.cpp b/Engine/Source/Runtime/Renderer/Private/BasePassRendering.cpp
--- a/Engine/Source/Runtime/Renderer/Private/BasePassRendering.cpp
+++ b/Engine/Source/Runtime/Renderer/Private/BasePassRendering.cpp
@@ -107,6 +107,13 @@ static TAutoConsoleVariable<int32> CVarPSOPrecacheTranslucencyAllPass(
 	ECVF_ReadOnly
 );
 
+static TAutoConsoleVariable<int32> CVarPSOPrecacheAlphaColorChannel(
+	TEXT("r.PSOPrecache.PrecacheAlphaColorChannel"),
+	1,
+	TEXT("Also Precache PSOs with scene color alpha channel enabled. Planar reflections and scene captures use this for compositing into a different scene later."),
+	ECVF_ReadOnly
+);
+
 // Scene color alpha is used during scene captures and planar reflections.  1 indicates background should be shown, 0 indicates foreground is fully present.
 static const float kSceneColorClearAlpha = 1.0f;
 
diff --git a/Engine/Source/Runtime/Renderer/Private/BasePassRendering.h b/Engine/Source/Runtime/Renderer/Private/BasePassRendering.h
--- a/Engine/Source/Runtime/Renderer/Private/BasePassRendering.h
+++ b/Engine/Source/Runtime/Renderer/Private/BasePassRendering.h
@@ -75,13 +75,6 @@ struct FForwardBasePassTextures
 	bool bIs24BitUnormDepthStencil = false;
 };
 
-static TAutoConsoleVariable<int32> CVarPSOPrecacheAlphaColorChannel(
-	TEXT("r.PSOPrecache.PrecacheAlphaColorChannel"),
-	1,
-	TEXT("Also Precache PSOs with scene color alpha channel enabled. Planar reflections and scene captures use this for compositing into a different scene later."),
-	ECVF_ReadOnly
-);
-
 BEGIN_GLOBAL_SHADER_PARAMETER_STRUCT(FSharedBasePassUniformParameters,)
 	SHADER_PARAMETER_STRUCT(FForwardLightData, Forward)
 	SHADER_PARAMETER_STRUCT(FForwardLightData, ForwardISR)
@@ -906,7 +899,8 @@ public:
 			PSOInitializers);
 
 		// Planar reflections and scene captures use scene color alpha to keep track of where content has been rendered, for compositing into a different scene later
-		if (bPrecacheAlphaColorChannel && CVarPSOPrecacheAlphaColorChannel.GetValueOnAnyThread() > 0)
+		static TConsoleVariableData<int32>* CVar = IConsoleManager::Get().FindTConsoleVariableDataInt(TEXT("r.PSOPrecache.PrecacheAlphaColorChannel"));
+		if (bPrecacheAlphaColorChannel && CVar && CVar->GetValueOnAnyThread() > 0)
 		{
 			FGraphicsPipelineRenderTargetsInfo AlphaColorRenderTargetsInfo = RenderTargetsInfo;
 
