Fixed a few issues with PSO pre-caching on platforms without MVF support
Fixed missing DepthAux setup for PSO pre-caching
Fixed missing PrimitiveID input for ISM PSO pre-caching

diff --git a/Engine/Source/Runtime/Engine/Private/InstancedStaticMesh.cpp b/Engine/Source/Runtime/Engine/Private/InstancedStaticMesh.cpp
--- a/Engine/Source/Runtime/Engine/Private/InstancedStaticMesh.cpp
+++ b/Engine/Source/Runtime/Engine/Private/InstancedStaticMesh.cpp
@@ -767,6 +767,12 @@ void FInstancedStaticMeshVertexFactory::GetVertexElements(
 {
 	FVertexStreamList VertexStreams;
 	GetVertexElements(FeatureLevel, InputStreamType, bSupportsManualVertexFetch, Data, InstanceData, Elements, VertexStreams);
+
+	if (UseGPUScene(GMaxRHIShaderPlatform, GMaxRHIFeatureLevel) 
+		&& !PlatformGPUSceneUsesUniformBufferView(GMaxRHIShaderPlatform))
+	{
+		Elements.Add(FVertexElement(VertexStreams.Num(), 0, VET_UInt, 13, sizeof(uint32), true));
+	}
 }
 
 void FInstancedStaticMeshVertexFactory::GetVertexElements(
diff --git a/Engine/Source/Runtime/Engine/Private/Materials/Material.cpp b/Engine/Source/Runtime/Engine/Private/Materials/Material.cpp
--- a/Engine/Source/Runtime/Engine/Private/Materials/Material.cpp
+++ b/Engine/Source/Runtime/Engine/Private/Materials/Material.cpp
@@ -582,8 +582,7 @@ void UMaterialInterface::InitDefaultMaterials()
 
 		// Now precache PSOs for all the default materials after the default materials are marked initialize
 		// PSO precaching can request default materials so they have to marked as initialized to avoid endless recursion
-		// Skip platforms that do not support MVF, non-MVF path needs mesh information for PSO 
-		if (bInitialized && PipelineStateCache::IsPSOPrecachingEnabled() && RHISupportsManualVertexFetch(GMaxRHIShaderPlatform))
+		if (bInitialized && PipelineStateCache::IsPSOPrecachingEnabled())
 		{
 			PrecacheDefaultMaterialPSOs();
 		}
@@ -592,6 +591,12 @@ void UMaterialInterface::InitDefaultMaterials()
 
 void UMaterialInterface::PrecacheDefaultMaterialPSOs()
 {
+	if (!GIsRHIInitialized || !RHISupportsManualVertexFetch(GMaxRHIShaderPlatform))
+	{
+		// Skip platforms that do not support MVF, non-MVF path needs mesh information for PSO 
+		return;
+	}
+		
 	TArray<FMaterialPSOPrecacheRequestID> MaterialPrecacheRequestIDs;
 
 	FPSOPrecacheParams PrecachePSOParams;
diff --git a/Engine/Source/Runtime/Engine/Private/SceneTexturesConfig.cpp b/Engine/Source/Runtime/Engine/Private/SceneTexturesConfig.cpp
--- a/Engine/Source/Runtime/Engine/Private/SceneTexturesConfig.cpp
+++ b/Engine/Source/Runtime/Engine/Private/SceneTexturesConfig.cpp
@@ -323,6 +323,11 @@ void FSceneTexturesConfig::Init(const FSceneTexturesConfigInitSettings& InitSett
 			GBufferBindings[Layout] = BindingCache.Bindings[Layout];
 		}
 	}
+
+	if (ShadingPath == EShadingPath::Mobile)
+	{
+		bRequiresDepthAux = MobileRequiresSceneDepthAux(ShaderPlatform);
+	}
 }
 
 void FSceneTexturesConfig::SetupMobileGBufferFlags(bool bRequiresMultiPass)
@@ -376,6 +381,12 @@ uint32 FSceneTexturesConfig::GetGBufferRenderTargetsInfo(FGraphicsPipelineRender
 		RenderTargetsInfo.RenderTargetFlags[RenderTargetCount++] = GBufferVelocity.Flags;
 	}
 
+	if (bRequiresDepthAux)
+	{
+		RenderTargetsInfo.RenderTargetFormats[RenderTargetCount] = bPreciseDepthAux ? PF_R32_FLOAT : PF_R16F;
+		RenderTargetsInfo.RenderTargetFlags[RenderTargetCount++] = TexCreate_RenderTargetable | TexCreate_ShaderResource | TexCreate_InputAttachmentRead;
+	}
+
 	// Store final number of render targets
 	RenderTargetsInfo.RenderTargetsEnabled = RenderTargetCount;
 
diff --git a/Engine/Source/Runtime/Engine/Private/UserInterface/Canvas.cpp b/Engine/Source/Runtime/Engine/Private/UserInterface/Canvas.cpp
--- a/Engine/Source/Runtime/Engine/Private/UserInterface/Canvas.cpp
+++ b/Engine/Source/Runtime/Engine/Private/UserInterface/Canvas.cpp
@@ -29,6 +29,7 @@
 #include "CanvasRendererItem.h"
 #include "RenderGraphUtils.h"
 #include "TextureResource.h"
+#include "DataDrivenShaderPlatformInfo.h"
 
 #include UE_INLINE_GENERATED_CPP_BY_NAME(Canvas)
 
@@ -2079,11 +2080,15 @@ void UCanvas::K2_DrawMaterial(UMaterialInterface* RenderMaterial, FVector2D Scre
 		// This is a user-facing function, so we'd rather make sure that shaders are ready by the time we render, in order to ensure we don't draw with a fallback material :
 		RenderMaterial->EnsureIsComplete();
 
-		// Should be moved earlier
-		FPSOPrecacheParams PSOPrecacheParams;
-		PSOPrecacheParams.bCanvasMaterial = true;
-		PSOPrecacheParams.BasePassPixelFormat = (Canvas->GetRenderTarget() && Canvas->GetRenderTarget()->GetRenderTargetTexture()) ? Canvas->GetRenderTarget()->GetRenderTargetTexture()->GetDesc().Format : PF_Unknown;
-		RenderMaterial->PrecachePSOs(&FLocalVertexFactory::StaticType, PSOPrecacheParams);
+		// TODO: need to pass a correct vertex declaration for non-MVF platforms
+		if (RHISupportsManualVertexFetch(GMaxRHIShaderPlatform))
+		{
+			// Should be moved earlier
+			FPSOPrecacheParams PSOPrecacheParams;
+			PSOPrecacheParams.bCanvasMaterial = true;
+			PSOPrecacheParams.BasePassPixelFormat = (Canvas->GetRenderTarget() && Canvas->GetRenderTarget()->GetRenderTargetTexture()) ? Canvas->GetRenderTarget()->GetRenderTargetTexture()->GetDesc().Format : PF_Unknown;
+			RenderMaterial->PrecachePSOs(&FLocalVertexFactory::StaticType, PSOPrecacheParams);
+		}
 
 		FCanvasTileItem TileItem(ScreenPosition, RenderMaterial->GetRenderProxy(), ScreenSize, CoordinatePosition, CoordinatePosition + CoordinateSize);
 		TileItem.Rotation = FRotator(0, Rotation, 0);
diff --git a/Engine/Source/Runtime/Renderer/Private/CompositionLighting/PostProcessDeferredDecals.cpp b/Engine/Source/Runtime/Renderer/Private/CompositionLighting/PostProcessDeferredDecals.cpp
--- a/Engine/Source/Runtime/Renderer/Private/CompositionLighting/PostProcessDeferredDecals.cpp
+++ b/Engine/Source/Runtime/Renderer/Private/CompositionLighting/PostProcessDeferredDecals.cpp
@@ -161,6 +161,17 @@ void GetDeferredDecalRenderTargetsInfo(
 		checkNoEntry();
 	}
 
+	if (Config.bRequiresDepthAux)
+	{
+		switch (RenderTargetMode)
+		{
+		case EDecalRenderTargetMode::SceneColorAndGBuffer:
+		case EDecalRenderTargetMode::SceneColorAndGBufferNoNormal:
+		case EDecalRenderTargetMode::SceneColor:
+			AddRenderTargetInfo(Config.bPreciseDepthAux ? PF_R32_FLOAT : PF_R16F, TexCreate_RenderTargetable | TexCreate_ShaderResource | TexCreate_InputAttachmentRead, RenderTargetsInfo);	
+		};
+	}
+
 	SetupDepthStencilInfo(PF_DepthStencil, Config.DepthCreateFlags, ERenderTargetLoadAction::ELoad,
 		ERenderTargetLoadAction::ELoad, FExclusiveDepthStencil::DepthRead_StencilWrite, RenderTargetsInfo);
 }
@@ -516,15 +527,20 @@ void CollectDeferredDecalPassPSOInitializers(
 
 	FGraphicsPipelineRenderTargetsInfo RenderTargetsInfo;
 	RenderTargetsInfo.NumSamples = 1;
-	GetDeferredDecalRenderTargetsInfo(SceneTexturesConfig, ShaderPlatform, DecalRenderTargetMode, RenderTargetsInfo);
-	ApplyTargetsInfo(GraphicsPSOInit, RenderTargetsInfo);
 
 	if (FeatureLevel == ERHIFeatureLevel::ES3_1)
 	{
 		// subpass info set during the submission of the draws in a mobile renderer
 		GraphicsPSOInit.SubpassIndex = 1; // all decals use second sub-pass on mobile
 		GraphicsPSOInit.SubpassHint = SceneTexturesConfig.bIsUsingGBuffers ? ESubpassHint::DeferredShadingSubpass : ESubpassHint::DepthReadSubpass;
+		if (DecalRenderTargetMode == EDecalRenderTargetMode::SceneColor)
+		{
+			RenderTargetsInfo.NumSamples = SceneTexturesConfig.NumSamples;
+		}
 	}
+
+	GetDeferredDecalRenderTargetsInfo(SceneTexturesConfig, ShaderPlatform, DecalRenderTargetMode, RenderTargetsInfo);
+	ApplyTargetsInfo(GraphicsPSOInit, RenderTargetsInfo);
 		
 	const auto AddDeferredDecalPSO = [&](bool bInsideDecal,	bool bReverseHanded, bool bReverseCulling, bool bDecalUsesStencil)
 	{
diff --git a/Engine/Source/Runtime/Renderer/Private/CompositionLighting/PostProcessMeshDecals.cpp b/Engine/Source/Runtime/Renderer/Private/CompositionLighting/PostProcessMeshDecals.cpp
--- a/Engine/Source/Runtime/Renderer/Private/CompositionLighting/PostProcessMeshDecals.cpp
+++ b/Engine/Source/Runtime/Renderer/Private/CompositionLighting/PostProcessMeshDecals.cpp
@@ -443,8 +443,12 @@ void FMeshDecalMeshProcessor::CollectPSOInitializers(
 		// Collect decal pass PSOs
 		CollectDeferredDecalPassPSOInitializers(PSOCollectorIndex, FeatureLevel, SceneTexturesConfig, Material, LocalDecalRenderStage, PSOInitializers);
 
-		// Collect decal mesh PSOs
-		CollectDeferredDecalMeshPSOInitializers(SceneTexturesConfig, VertexFactoryData, PreCacheParams, Material, DecalBlendDesc, LocalDecalRenderStage, PSOInitializers);
+		// TODO: need to pass a correct vertex declaration for non-MVF platforms
+		if (RHISupportsManualVertexFetch(ShaderPlatform))
+		{
+			// Collect decal mesh PSOs
+			CollectDeferredDecalMeshPSOInitializers(SceneTexturesConfig, VertexFactoryData, PreCacheParams, Material, DecalBlendDesc, LocalDecalRenderStage, PSOInitializers);
+		}
 	}
 }
 
