Add aniso support for legacy lighting with cluster shading.

#11487

diff --git a/Engine/Source/Runtime/Renderer/Private/ClusteredDeferredShadingPass.cpp b/Engine/Source/Runtime/Renderer/Private/ClusteredDeferredShadingPass.cpp
--- a/Engine/Source/Runtime/Renderer/Private/ClusteredDeferredShadingPass.cpp
+++ b/Engine/Source/Runtime/Renderer/Private/ClusteredDeferredShadingPass.cpp
@@ -29,6 +29,7 @@
 #include "PostProcess/PostProcessing.h"
 #include "Substrate/Substrate.h"
 #include "LightFunctionAtlas.h"
+#include "AnisotropyRendering.h"
 
 // This is used to switch on and off the clustered deferred shading implementation, that uses the light grid to perform shading.
 int32 GUseClusteredDeferredShading = 0;
@@ -99,7 +100,8 @@ class FClusteredShadingPS : public FGlobalShader
 	class FSubstrateTileType : SHADER_PERMUTATION_INT("SUBSTRATE_TILETYPE", 4);
 	class FLightFunctionAtlasDim : SHADER_PERMUTATION_BOOL("USE_LIGHT_FUNCTION_ATLAS");
 	class FRectLight : SHADER_PERMUTATION_BOOL("USE_RECT_LIGHT");
-	using FPermutationDomain = TShaderPermutationDomain<FVisualizeLightCullingDim, FHairStrandsLighting, FSubstrateTileType, FLightFunctionAtlasDim, FRectLight>;
+	class FAnistropicMaterials : SHADER_PERMUTATION_BOOL("SUPPORTS_ANISOTROPIC_MATERIALS");
+	using FPermutationDomain = TShaderPermutationDomain<FVisualizeLightCullingDim, FHairStrandsLighting, FSubstrateTileType, FLightFunctionAtlasDim, FRectLight, FAnistropicMaterials>;
 
 	BEGIN_SHADER_PARAMETER_STRUCT(FParameters, )
 		SHADER_PARAMETER(uint32, bHasLightChannels)
@@ -125,6 +127,11 @@ class FClusteredShadingPS : public FGlobalShader
 		{
 			return false;
 		}
+		else if (Substrate::IsSubstrateEnabled() && PermutationVector.Get<FAnistropicMaterials>())
+		{
+			return false;
+		}
+
 		// OLATODO: what level do we actually need for this?
 		return IsFeatureLevelSupported(Parameters.Platform, ERHIFeatureLevel::SM5);
 	}
@@ -168,7 +175,8 @@ static void InternalAddClusteredDeferredShadingPass(
 	const bool bHairStrands = InputType == EClusterPassInputType::HairStrands;
 	const bool bSubstrate = Substrate::IsSubstrateEnabled() && !bHairStrands;
 	const bool bHasRectLights = SortedLightsSet.bHasRectLights;
-	const bool bLightFunctionAtlas = LightFunctionAtlas::IsEnabled(View, ELightFunctionAtlasSystem::DeferredLighting);
+	const bool bLightFunctionAtlas = LightFunctionAtlas::IsEnabled(View, ELightFunctionAtlasSystem::DeferredLighting);	
+	const bool bSupportAnisotropyPermutation = ShouldRenderAnisotropyPass(View) && !Substrate::IsSubstrateEnabled(); // Strata managed anisotropy differently than legacy path. No need for special permutation.
 
 	FClusteredShadingPS::FParameters *PassParameters = GraphBuilder.AllocParameters<FClusteredShadingPS::FParameters>();
 	PassParameters->bHasLightChannels = SortedLightsSet.bHasLightChannels;
@@ -205,7 +213,7 @@ static void InternalAddClusteredDeferredShadingPass(
 		RDG_EVENT_NAME("Light::ClusteredDeferredShading(%s,Lights:%d%s%s%s%s)", bHairStrands ? TEXT("HairStrands") : (bSubstrate ? TEXT("Substrate") : TEXT("GBuffer")), SortedLightsSet.ClusteredSupportedEnd, bSubstrate ? TEXT(",Tile:") : TEXT(""), bSubstrate ? TileTypeName : TEXT(""), bLightFunctionAtlas ? TEXT(",LFAtlas") : TEXT(""), bHasRectLights ? TEXT(",RectLight") : TEXT("")),
 		PassParameters,
 		ERDGPassFlags::Raster,
-		[PassParameters, &View, SceneTextureExtent, bHasRectLights, bLightFunctionAtlas, bHairStrands, bSubstrate, TileType, PrimitiveType](FRHICommandListImmediate& InRHICmdList)
+		[PassParameters, &View, SceneTextureExtent, bHasRectLights, bLightFunctionAtlas, bHairStrands, bSubstrate, bSupportAnisotropyPermutation, TileType, PrimitiveType](FRHICommandListImmediate& InRHICmdList)
 	{
 		TShaderMapRef<FClusteredShadingVS> HairVertexShader(View.ShaderMap);
 		TShaderMapRef<FPostProcessVS> VertexShader(View.ShaderMap);
@@ -221,6 +229,7 @@ static void InternalAddClusteredDeferredShadingPass(
 		PermutationVector.Set<FClusteredShadingPS::FSubstrateTileType>(bSubstrate ? TileType : 0);
 		PermutationVector.Set<FClusteredShadingPS::FRectLight>(bHasRectLights);
 		PermutationVector.Set<FClusteredShadingPS::FLightFunctionAtlasDim>(bLightFunctionAtlas);
+		PermutationVector.Set<FClusteredShadingPS::FAnistropicMaterials>(bSupportAnisotropyPermutation);
 		TShaderMapRef<FClusteredShadingPS> PixelShader(View.ShaderMap, PermutationVector);
 		{
 			FGraphicsPipelineStateInitializer GraphicsPSOInit;
