VolumetricFog - fixed missing precached shader permutation happening when Lumen translucent volume cache is not rendered and FDistanceFieldSkyOcclusion should be used instead.

diff --git a/Engine/Source/Runtime/Renderer/Private/VolumetricFog.cpp b/Engine/Source/Runtime/Renderer/Private/VolumetricFog.cpp
--- a/Engine/Source/Runtime/Renderer/Private/VolumetricFog.cpp
+++ b/Engine/Source/Runtime/Renderer/Private/VolumetricFog.cpp
@@ -1064,10 +1064,8 @@ class FVolumetricFogLightScatteringCS : public FGlobalShader
 			return EShaderPermutationPrecacheRequest::NotUsed;
 		}
 
-		if (PermutationVector.Get<FDistanceFieldSkyOcclusion>() && bSupportsLumenGI)
-		{
-			return EShaderPermutationPrecacheRequest::NotUsed;
-		}
+		// We cannot skip the precaching of FDistanceFieldSkyOcclusion because the runtime logic to enable it is complex and not only based on project support & cvar values.
+		// It involves checking lumen support and the presence of GetLumenTranslucencyGIVolume().Texture0 which might not be present if no translucent are present.
 
 		if (PermutationVector.Get<FTemporalReprojection>() && (GVolumetricFogTemporalReprojection == 0))
 		{
