HWRT: Set ERayTracingInstanceFlags::TriangleCullReverse when it's set on a mesh batch.

diff --git a/Engine/Source/Runtime/Renderer/Private/PrimitiveSceneInfo.cpp b/Engine/Source/Runtime/Renderer/Private/PrimitiveSceneInfo.cpp
--- a/Engine/Source/Runtime/Renderer/Private/PrimitiveSceneInfo.cpp
+++ b/Engine/Source/Runtime/Renderer/Private/PrimitiveSceneInfo.cpp
@@ -1301,6 +1301,11 @@ void FPrimitiveSceneInfo::UpdateCachedRayTracingInstance(FPrimitiveSceneInfo* Sc
 			SceneInfo->CachedRayTracingInstance.Flags |= ERayTracingInstanceFlags::TriangleCullDisable;
 		}
 
+		if (InstanceMaskAndFlags.bReverseCulling)
+		{
+			SceneInfo->CachedRayTracingInstance.Flags |= ERayTracingInstanceFlags::TriangleCullReverse;
+		}
+
 		SceneInfo->bCachedRayTracingInstanceAnySegmentsDecal = InstanceMaskAndFlags.bAnySegmentsDecal;
 		SceneInfo->bCachedRayTracingInstanceAllSegmentsDecal = InstanceMaskAndFlags.bAllSegmentsDecal;
 
diff --git a/Engine/Source/Runtime/Renderer/Private/RayTracing/RayTracing.cpp b/Engine/Source/Runtime/Renderer/Private/RayTracing/RayTracing.cpp
--- a/Engine/Source/Runtime/Renderer/Private/RayTracing/RayTracing.cpp
+++ b/Engine/Source/Runtime/Renderer/Private/RayTracing/RayTracing.cpp
@@ -159,6 +159,7 @@ namespace RayTracing
 		bool bTwoSided : 1 = false;
 		bool bIsSky : 1 = false;
 		bool bAllSegmentsTranslucent : 1 = true;
+		bool bAllSegmentsReverseCulling : 1 = true;
 
 		const FRayTracingGeometryInstance* CachedRayTracingInstance = nullptr;
 		TArrayView<const int32> CachedRayTracingMeshCommandIndices; // Pointer to FPrimitiveSceneInfo::CachedRayTracingMeshCommandIndicesPerLOD data
@@ -180,6 +181,7 @@ namespace RayTracing
 			Key ^= bTwoSided ? 0x1ull << 45 : 0x0;
 			Key ^= bIsSky ? 0x1ull << 46 : 0x0;
 			Key ^= bAllSegmentsTranslucent ? 0x1ull << 47 : 0x0;
+			Key ^= bAllSegmentsReverseCulling ? 0x1ull << 48 : 0x0;
 			return Key ^ reinterpret_cast<uint64>(RayTracingGeometryRHI);
 		}
 
@@ -582,6 +584,7 @@ namespace RayTracing
 										RelevantPrimitive->bTwoSided |= RayTracingMeshCommand.bTwoSided;
 										RelevantPrimitive->bIsSky |= RayTracingMeshCommand.bIsSky;
 										RelevantPrimitive->bAllSegmentsTranslucent &= RayTracingMeshCommand.bIsTranslucent;
+										RelevantPrimitive->bAllSegmentsReverseCulling &= RayTracingMeshCommand.bReverseCulling;
 									}
 									else
 									{
@@ -901,6 +904,11 @@ namespace RayTracing
 								SceneInfo->CachedRayTracingInstance.Flags |= ERayTracingInstanceFlags::TriangleCullDisable;
 							}
 
+							if (InstanceMaskAndFlags.bReverseCulling)
+							{
+								SceneInfo->CachedRayTracingInstance.Flags |= ERayTracingInstanceFlags::TriangleCullReverse;
+							}
+
 							SceneInfo->bCachedRayTracingInstanceAnySegmentsDecal = InstanceMaskAndFlags.bAnySegmentsDecal;
 							SceneInfo->bCachedRayTracingInstanceAllSegmentsDecal = InstanceMaskAndFlags.bAllSegmentsDecal;
 
@@ -1253,6 +1261,10 @@ namespace RayTracing
 							{
 								RayTracingInstance.Flags |= ERayTracingInstanceFlags::TriangleCullDisable;
 							}
+							if (RelevantPrimitive.bAllSegmentsReverseCulling)
+							{
+								RayTracingInstance.Flags |= ERayTracingInstanceFlags::TriangleCullReverse;
+							}
 							AddDebugRayTracingInstanceFlags(RayTracingInstance.Flags);
 
 							RayTracingInstance.LayerIndex = (uint8)ERayTracingSceneLayer::Base;
diff --git a/Engine/Source/Runtime/Renderer/Private/RayTracing/RayTracingInstanceMask.cpp b/Engine/Source/Runtime/Renderer/Private/RayTracing/RayTracingInstanceMask.cpp
--- a/Engine/Source/Runtime/Renderer/Private/RayTracing/RayTracingInstanceMask.cpp
+++ b/Engine/Source/Runtime/Renderer/Private/RayTracing/RayTracingInstanceMask.cpp
@@ -154,6 +154,7 @@ FRayTracingMaskAndFlags BuildRayTracingInstanceMaskAndFlags(TArrayView<const FMe
 	bool bAnySegmentsDecal = false;
 	bool bAllSegmentsDecal = true;
 	bool bDoubleSided = false;
+	bool bAllSegmentsReverseCulling = true;
 	ERayTracingViewMaskMode MaskMode = SceneProxyRayTracingMaskInfo.MaskMode;
 	Result.Mask = ExtraMask;
 
@@ -175,6 +176,7 @@ FRayTracingMaskAndFlags BuildRayTracingInstanceMaskAndFlags(TArrayView<const FMe
 			bAnySegmentsDecal |= Material.IsDeferredDecal();
 			bAllSegmentsDecal &= Material.IsDeferredDecal();
 			bDoubleSided |= MeshBatch.bDisableBackfaceCulling || Material.IsTwoSided();
+			bAllSegmentsReverseCulling &= MeshBatch.ReverseCulling;
 		}
 	}
 
@@ -183,6 +185,7 @@ FRayTracingMaskAndFlags BuildRayTracingInstanceMaskAndFlags(TArrayView<const FMe
 	Result.bDoubleSided = bDoubleSided;	
 	Result.bAnySegmentsDecal = bAnySegmentsDecal;
 	Result.bAllSegmentsDecal = bAllSegmentsDecal;
+	Result.bReverseCulling = bAllSegmentsReverseCulling;
 
 	const bool bIsHairStrands = Result.Mask & ComputeRayTracingInstanceMask(ERayTracingInstanceMaskType::HairStrands, MaskMode);
 	if (bIsHairStrands)
@@ -261,6 +264,7 @@ void SetupRayTracingMeshCommandMaskAndStatus(FRayTracingMeshCommand& MeshCommand
 	MeshCommand.bIsSky = MaterialResource.IsSky();
 	MeshCommand.bTwoSided = MaterialResource.IsTwoSided();
 	MeshCommand.bIsTranslucent = MaterialResource.GetBlendMode() == EBlendMode::BLEND_Translucent;
+	MeshCommand.bReverseCulling = MeshBatch.ReverseCulling;
 
 	MeshCommand.InstanceMask = BlendModeToRayTracingInstanceMask(MaterialResource.GetBlendMode(), MeshCommand.bCastRayTracedShadows, MaskMode);
 
diff --git a/Engine/Source/Runtime/Renderer/Public/MeshPassProcessor.h b/Engine/Source/Runtime/Renderer/Public/MeshPassProcessor.h
--- a/Engine/Source/Runtime/Renderer/Public/MeshPassProcessor.h
+++ b/Engine/Source/Runtime/Renderer/Public/MeshPassProcessor.h
@@ -2354,6 +2354,7 @@ public:
 	bool bIsSky = false;
 	bool bIsTranslucent = false;
 	bool bTwoSided = false;
+	bool bReverseCulling = false;
 
 	RENDERER_API void SetRayTracingShaderBindingsForHitGroup(
 		FRayTracingLocalShaderBindingWriter* BindingWriter,
diff --git a/Engine/Source/Runtime/Renderer/Public/RayTracingInstance.h b/Engine/Source/Runtime/Renderer/Public/RayTracingInstance.h
--- a/Engine/Source/Runtime/Renderer/Public/RayTracingInstance.h
+++ b/Engine/Source/Runtime/Renderer/Public/RayTracingInstance.h
@@ -15,6 +15,7 @@ struct FRayTracingMaskAndFlags
 		: Mask(0xFF)
 		, bForceOpaque(false)
 		, bDoubleSided(false)
+		, bReverseCulling(false)
 		, bAnySegmentsDecal(false)
 		, bAllSegmentsDecal(false)
 	{}
@@ -27,6 +28,9 @@ struct FRayTracingMaskAndFlags
 
 	/** Whether ray hits should be registered for front and back faces. */
 	uint8 bDoubleSided : 1;
+	
+	/** Whether front and back facings should be reversed. */
+	uint8 bReverseCulling : 1;
 
 	/** Whether any or all of the segments in the instance are decals. */
 	uint8 bAnySegmentsDecal : 1;
