Light function atlas - improved visualisation

diff --git a/Engine/Source/Runtime/Renderer/Private/LightFunctionAtlas.cpp b/Engine/Source/Runtime/Renderer/Private/LightFunctionAtlas.cpp
--- a/Engine/Source/Runtime/Renderer/Private/LightFunctionAtlas.cpp
+++ b/Engine/Source/Runtime/Renderer/Private/LightFunctionAtlas.cpp
@@ -778,6 +778,9 @@ FScreenPassTexture FLightFunctionAtlas::AddDebugVisualizationPasses(FRDGBuilder&
 
 		Text = FString::Printf(TEXT("Local Lights sampling atlas: %d"), EffectiveLocalLightSlotArray.Num());
 		Canvas.DrawShadowedString(DrawPosX, DrawPosY, *Text, GEngine->GetLargeFont(), FLinearColor::White);
+		DrawPosY += 20.0f;
+
+		Canvas.DrawShadowedString(DrawPosX, DrawPosY, TEXT("Note: (X) means the material cannot be sampled from the atlas for the deferred lighting passes (due to using world pos, depth or manipulating UVs)"), GEngine->GetLargeFont(), FLinearColor::Yellow);
 		DrawPosY += 40.0f;
 
 		uint32 LightFunctionAtlasSlotIndex = 0;
@@ -786,9 +789,12 @@ FScreenPassTexture FLightFunctionAtlas::AddDebugVisualizationPasses(FRDGBuilder&
 			const FMaterialRenderProxy* LightFunctionMaterial = AtlasSlot.LightFunctionMaterial;
 
 			const FString& MaterialName = LightFunctionMaterial->GetMaterialName();
-			uint32 LFMaterialUniqueID = LightFunctionMaterial->GetMaterialInterface() ? MurmurFinalize32(MurmurFinalize32(LightFunctionMaterial->GetMaterialInterface()->GetUniqueID())) : 0xFFFFFFFF;
+			const uint32 LFMaterialUniqueID = LightFunctionMaterial->GetMaterialInterface() ? MurmurFinalize32(MurmurFinalize32(LightFunctionMaterial->GetMaterialInterface()->GetUniqueID())) : 0xFFFFFFFF;
+
+			const FMaterial& LFMaterial = LightFunctionMaterial->GetIncompleteMaterialWithFallback(View.GetFeatureLevel());
+			const bool bIsMaterialCompatibleWithAtlas = LFMaterial.MaterialIsLightFunctionAtlasCompatible_RenderThread();
 
-			FLinearColor MaterialColor(FColor(LFMaterialUniqueID & 0xFF, (LFMaterialUniqueID >> 8) & 0xFF, (LFMaterialUniqueID >> 16) & 0xFF));
+			const FLinearColor MaterialColor(FColor(LFMaterialUniqueID & 0xFF, (LFMaterialUniqueID >> 8) & 0xFF, (LFMaterialUniqueID >> 16) & 0xFF));
 
 			uint32 LightCountUsingThisMaterial = 0;
 			for (auto& LocalLight : EffectiveLocalLightSlotArray)
@@ -797,10 +803,12 @@ FScreenPassTexture FLightFunctionAtlas::AddDebugVisualizationPasses(FRDGBuilder&
 			}
 
 			// Draw the light function material
-			Canvas.DrawTile(DrawPosX - 10.0f, DrawPosY, 15.0f, 15.0f, 0.0f, 0.0f, 1.0f, 1.0f, MaterialColor, nullptr, false);
+			Canvas.DrawTile(DrawPosX - 20.0f, DrawPosY, 15.0f, 15.0f, 0.0f, 0.0f, 1.0f, 1.0f, MaterialColor, nullptr, false);
 
-			Text = FString::Printf(TEXT("%3d lights - %s"), LightCountUsingThisMaterial , *MaterialName);
+			Text = FString::Printf(TEXT("%2d"), LightCountUsingThisMaterial);
 			Canvas.DrawShadowedString(DrawPosX, DrawPosY, *Text, GEngine->GetLargeFont(), MaterialColor);
+			Text = FString::Printf(TEXT("lights - %s - %s"), bIsMaterialCompatibleWithAtlas ? TEXT("OK ") : TEXT("(X)"), *MaterialName);
+			Canvas.DrawShadowedString(DrawPosX + 20, DrawPosY, *Text, GEngine->GetLargeFont(), MaterialColor);
 
 			// Draw a line around the corresponding atlas tile
 			auto OutLineAtlasSlot = [&](float X0, float Y0, float X1, float Y1)
