SkyAtmosphere - does not disrupt engine rand stability anumore.

diff --git a/Engine/Source/Runtime/Renderer/Private/SkyAtmosphereRendering.cpp b/Engine/Source/Runtime/Renderer/Private/SkyAtmosphereRendering.cpp
--- a/Engine/Source/Runtime/Renderer/Private/SkyAtmosphereRendering.cpp
+++ b/Engine/Source/Runtime/Renderer/Private/SkyAtmosphereRendering.cpp
@@ -1092,18 +1092,18 @@ public:
 		UniformSphereSamplesBuffer.Initialize(RHICmdList, TEXT("UniformSphereSamplesBuffer"), sizeof(FVector4f), GroupSize * GroupSize, EPixelFormat::PF_A32B32G32R32F, BUF_Static);
 		FVector4f* Dest = (FVector4f*)RHICmdList.LockBuffer(UniformSphereSamplesBuffer.Buffer, 0, sizeof(FVector4f)*GroupSize*GroupSize, RLM_WriteOnly);
 
-		FMath::SRandInit(0xDE4DC0DE);
+		FRandomStream RandomStream(0xDE4DC0DE);
 		for (uint32 i = 0; i < GroupSize; ++i)
 		{
 			for (uint32 j = 0; j < GroupSize; ++j)
 			{
-				const float u0 = (float(i) + FMath::SRand()) * GroupSizeInv;
-				const float u1 = (float(j) + FMath::SRand()) * GroupSizeInv;
-
+				const float u0 = (float(i) + RandomStream.GetFraction()) * GroupSizeInv;
+				const float u1 = (float(j) + RandomStream.GetFraction()) * GroupSizeInv;
+		
 				const float a = 1.0f - 2.0f * u0;
 				const float b = FMath::Sqrt(1.0f - a*a);
 				const float phi = 2 * PI * u1;
-
+		
 				uint32 idx = j * GroupSize + i;
 				Dest[idx].X = b * FMath::Cos(phi);
 				Dest[idx].Y = b * FMath::Sin(phi);
