Fix issue building BLAS in HeterogeneousVolumesHardwareRayTracing and batch builds.
- GraphBuilder.AddPass lambda was capturing RayTracingGeometries by reference but the array is a local variable which caused undefined behaviour.
- refactored logic to batch raytracing geometry builds and "move" the relevant array allocations so BLAS are built correctly.
- allocate a shared BLAS scratch buffer.

diff --git a/Engine/Source/Runtime/Renderer/Private/HeterogeneousVolumes/HeterogeneousVolumesHardwareRayTracing.cpp b/Engine/Source/Runtime/Renderer/Private/HeterogeneousVolumes/HeterogeneousVolumesHardwareRayTracing.cpp
--- a/Engine/Source/Runtime/Renderer/Private/HeterogeneousVolumes/HeterogeneousVolumesHardwareRayTracing.cpp
+++ b/Engine/Source/Runtime/Renderer/Private/HeterogeneousVolumes/HeterogeneousVolumesHardwareRayTracing.cpp
@@ -157,8 +157,7 @@ void GenerateRayTracingGeometryInstance(
 	CreateSparseVoxelBLAS(GraphBuilder, View, SparseVoxelUniformBuffer, NumVoxelsBuffer, GraphBuilder.RegisterExternalBuffer(PooledVertexBuffer));
 
 	FRayTracingGeometryInitializer GeometryInitializer;
-	// TODO: REMOVE STRING ALLOCATION
-	//GeometryInitializer.DebugName = *PrimitiveSceneProxy->GetResourceName().ToString();// +TEXT(" (HeterogeneousVolume)");
+	GeometryInitializer.DebugName = TEXT(" (HeterogeneousVolume)"); // TODO: Include resource name ie: *PrimitiveSceneProxy->GetResourceName().ToString();
 	//GeometryInitializer.IndexBuffer = EmptyIndexBuffer;
 	GeometryInitializer.GeometryType = RTGT_Procedural;
 	GeometryInitializer.bFastBuild = false;
@@ -176,7 +175,7 @@ void GenerateRayTracingGeometryInstance(
 }
 
 BEGIN_SHADER_PARAMETER_STRUCT(FBuildBLASPassParams, )
-	SHADER_PARAMETER_RDG_BUFFER_UAV(RWStructuredBuffer, InstanceBuffer)
+	RDG_BUFFER_ACCESS(ScratchBuffer, ERHIAccess::UAVCompute)
 END_SHADER_PARAMETER_STRUCT()
 
 BEGIN_SHADER_PARAMETER_STRUCT(FBuildTLASPassParams, )
@@ -199,38 +198,59 @@ void GenerateRayTracingScene(
 {
 	RayTracingScene.Reset(false);
 
+	TArray<FRayTracingGeometryBuildParams> BuildParams;
+	uint32 BLASScratchSize = 0;
+
 	// Collect instances
 	for (int32 GeometryIndex = 0; GeometryIndex < RayTracingGeometries.Num(); ++GeometryIndex)
 	{
-		checkf(RayTracingGeometries[GeometryIndex], TEXT("RayTracingGeometryInstance not created."));
+		FRHIRayTracingGeometry* RayTracingGeometry = RayTracingGeometries[GeometryIndex];
+		checkf(RayTracingGeometry, TEXT("RayTracingGeometryInstance not created."));
+
+		FRayTracingGeometryBuildParams Params;
+		Params.Geometry = RayTracingGeometry;
+		Params.BuildMode = EAccelerationStructureBuildMode::Build;
+
+		BuildParams.Add(Params);
+
+		const FRayTracingGeometryInitializer& Initializer = RayTracingGeometry->GetInitializer();
+
+		FRayTracingAccelerationStructureSize SizeInfo = RHICalcRayTracingGeometrySize(Initializer);
+		BLASScratchSize = Align(BLASScratchSize + SizeInfo.BuildScratchSize, GRHIRayTracingScratchBufferAlignment);
+
 		FRayTracingGeometryInstance RayTracingGeometryInstance = {};
-		RayTracingGeometryInstance.GeometryRHI = RayTracingGeometries[GeometryIndex];
+		RayTracingGeometryInstance.GeometryRHI = RayTracingGeometry;
 		RayTracingGeometryInstance.NumTransforms = 1;
 		RayTracingGeometryInstance.Transforms = MakeArrayView(&RayTracingTransforms[GeometryIndex], 1);
 
 		RayTracingScene.AddInstance(RayTracingGeometryInstance);
 	}
 
+	FRDGBufferDesc ScratchBufferDesc;
+	ScratchBufferDesc.Usage = EBufferUsageFlags::RayTracingScratch | EBufferUsageFlags::StructuredBuffer;
+	ScratchBufferDesc.BytesPerElement = GRHIRayTracingScratchBufferAlignment;
+	ScratchBufferDesc.NumElements = FMath::DivideAndRoundUp(BLASScratchSize, GRHIRayTracingScratchBufferAlignment);
+
+	FRDGBufferRef ScratchBuffer = GraphBuilder.CreateBuffer(ScratchBufferDesc, TEXT("HeterogeneousVolumes.BLASSharedScratchBuffer"));
+
 	// Build instance BLAS
 	FBuildBLASPassParams* PassParamsBLAS = GraphBuilder.AllocParameters<FBuildBLASPassParams>();
+	PassParamsBLAS->ScratchBuffer = ScratchBuffer;
 
 	GraphBuilder.AddPass(
-		RDG_EVENT_NAME("BuildTLASInstanceBuffer"),
+		RDG_EVENT_NAME("BuildRayTracingGeometries"),
 		PassParamsBLAS,
 		ERDGPassFlags::Compute | ERDGPassFlags::NeverCull | ERDGPassFlags::NeverParallel,
 		[
 			PassParamsBLAS,
-			&RayTracingGeometries
+			BuildParams = MoveTemp(BuildParams)
 		](FRHICommandListImmediate& RHICmdList)
 		{
-			for (FRayTracingGeometryRHIRef RayTracingGeometry : RayTracingGeometries)
-			{
-				checkf(RayTracingGeometry, TEXT("RayTracingGeometry not created."));
-				RHICmdList.BuildAccelerationStructure(RayTracingGeometry);
+			FRHIBufferRange ScratchBufferRange;
+			ScratchBufferRange.Buffer = PassParamsBLAS->ScratchBuffer->GetRHI();
+			ScratchBufferRange.Offset = 0;
 
-				// #yuriy_todo: explicit transitions and state validation for BLAS
-				// RHICmdList.Transition(FRHITransitionInfo(RayTracingGeometry.GetReference(), ERHIAccess::BVHWrite, ERHIAccess::BVHRead));
-			}
+			RHICmdList.BuildAccelerationStructures(BuildParams, ScratchBufferRange);
 		}
 	);
 
