HeterogeneousVolumes: Binding scene and loose shader bindings for FRenderVolumetricShadowMapForLightWithLiveShadingCS

diff --git a/Engine/Source/Runtime/Renderer/Private/HeterogeneousVolumes/HeterogeneousVolumesLiveShadingPipeline.cpp b/Engine/Source/Runtime/Renderer/Private/HeterogeneousVolumes/HeterogeneousVolumesLiveShadingPipeline.cpp
--- a/Engine/Source/Runtime/Renderer/Private/HeterogeneousVolumes/HeterogeneousVolumesLiveShadingPipeline.cpp
+++ b/Engine/Source/Runtime/Renderer/Private/HeterogeneousVolumes/HeterogeneousVolumesLiveShadingPipeline.cpp
@@ -1094,6 +1094,31 @@ void RenderWithLiveShading(
 	}
 }
 
+class FRenderShadowMapLooseBindings
+{
+	DECLARE_TYPE_LAYOUT(FRenderShadowMapLooseBindings, NonVirtual);
+
+public:
+	void Bind(const FShaderParameterMap& ParameterMap)
+	{
+		SceneDepthTextureBinding.Bind(ParameterMap, TEXT("SceneDepthTexture"));
+	}
+
+	template<typename TPassParameters>
+	void SetParameters(FMeshDrawSingleShaderBindings& ShaderBindings, const TPassParameters* PassParameters)
+	{
+		ShaderBindings.AddTexture(
+			SceneDepthTextureBinding,
+			FShaderResourceParameter(),
+			TStaticSamplerState<SF_Point>::GetRHI(),
+			PassParameters->SceneTextures.SceneDepthTexture->GetRHI()
+		);
+	}
+
+	LAYOUT_FIELD(FShaderResourceParameter, SceneDepthTextureBinding);
+};
+IMPLEMENT_TYPE_LAYOUT(FRenderShadowMapLooseBindings);
+
 class FRenderVolumetricShadowMapForLightWithLiveShadingCS : public FMeshMaterialShader
 {
 	DECLARE_SHADER_TYPE(FRenderVolumetricShadowMapForLightWithLiveShadingCS, MeshMaterial);
@@ -1106,6 +1131,7 @@ class FRenderVolumetricShadowMapForLightWithLiveShadingCS : public FMeshMaterial
 		// Scene data
 		SHADER_PARAMETER_STRUCT_REF(FViewUniformShaderParameters, View)
 		SHADER_PARAMETER_STRUCT_INCLUDE(FSceneTextureParameters, SceneTextures)
+		SHADER_PARAMETER_RDG_UNIFORM_BUFFER(FSceneUniformParameters, Scene)
 		SHADER_PARAMETER_STRUCT_REF(FBlueNoise, BlueNoise)
 
 		// Volumetric Shadow Map data
@@ -1200,7 +1226,7 @@ class FRenderVolumetricShadowMapForLightWithLiveShadingCS : public FMeshMaterial
 	static int32 GetThreadGroupSize2D() { return 8; }
 	static int32 GetThreadGroupSize3D() { return 4; }
 
-	LAYOUT_FIELD(FRenderLightingCacheLooseBindings, ShaderLooseBindings);
+	LAYOUT_FIELD(FRenderShadowMapLooseBindings, ShaderLooseBindings);
 };
 
 IMPLEMENT_MATERIAL_SHADER_TYPE(, FRenderVolumetricShadowMapForLightWithLiveShadingCS, TEXT("/Engine/Private/HeterogeneousVolumes/HeterogeneousVolumesLiveShadingShadows.usf"), TEXT("RenderVolumetricShadowMapForLightWithLiveShadingCS"), SF_Compute);
@@ -1336,6 +1362,7 @@ bool RenderVolumetricShadowMapForLightForHeterogeneousVolumeWithLiveShading(
 		// Scene data
 		PassParameters->View = View.ViewUniformBuffer;
 		PassParameters->SceneTextures = GetSceneTextureParameters(GraphBuilder, SceneTextures);
+		PassParameters->Scene = View.GetSceneUniforms().GetBuffer(GraphBuilder);
 		FBlueNoise BlueNoise = GetBlueNoiseGlobalParameters();
 		PassParameters->BlueNoise = CreateUniformBufferImmediate(BlueNoise, EUniformBufferUsage::UniformBuffer_SingleDraw);
 
@@ -1417,33 +1444,22 @@ bool RenderVolumetricShadowMapForLightForHeterogeneousVolumeWithLiveShading(
 			ERDGPassFlags::Compute,
 			[ComputeShader, PassParameters, Scene, MaterialRenderProxy, &Material, GroupCount](FRHIComputeCommandList& RHICmdList)
 			{
-#if 1
-				FMeshDrawShaderBindings ShaderBindings;
-		UE::MeshPassUtils::SetupComputeBindings(ComputeShader, Scene, Scene->GetFeatureLevel(), nullptr, *MaterialRenderProxy, Material, ShaderBindings);
-#else
 				FMeshMaterialShaderElementData ShaderElementData;
-		ShaderElementData.InitializeMeshMaterialData();
+				ShaderElementData.InitializeMeshMaterialData();
 
-		FMeshProcessorShaders PassShaders;
-		PassShaders.ComputeShader = ComputeShader;
+				FMeshProcessorShaders PassShaders;
+				PassShaders.ComputeShader = ComputeShader;
 
-		FMeshDrawShaderBindings ShaderBindings;
-		ShaderBindings.Initialize(PassShaders);
-		{
-			FMeshDrawSingleShaderBindings SingleShaderBindings = ShaderBindings.GetSingleShaderBindings(SF_Compute);
-			ComputeShader->GetShaderBindings(Scene, Scene->GetFeatureLevel(), nullptr, *MaterialRenderProxy, Material, ShaderElementData, SingleShaderBindings);
-			//ComputeShader->ShaderLooseBindings.SetParameters(SingleShaderBindings, PassParameters);
-			SingleShaderBindings.AddTexture(
-				ComputeShader->ShaderLooseBindings.SceneDepthTextureBinding,
-				FShaderResourceParameter(),
-				TStaticSamplerState<SF_Point>::GetRHI(),
-				PassParameters->SceneTextures.SceneDepthTexture->GetRHI()
-			);
-			ShaderBindings.Finalize(&PassShaders);
-		}
-#endif
+				FMeshDrawShaderBindings ShaderBindings;
+				ShaderBindings.Initialize(PassShaders);
+				{
+					FMeshDrawSingleShaderBindings SingleShaderBindings = ShaderBindings.GetSingleShaderBindings(SF_Compute);
+					ComputeShader->GetShaderBindings(Scene, Scene->GetFeatureLevel(), nullptr, *MaterialRenderProxy, Material, ShaderElementData, SingleShaderBindings);
+					ComputeShader->ShaderLooseBindings.SetParameters(SingleShaderBindings, PassParameters);
+					ShaderBindings.Finalize(&PassShaders);
+				}
 
-		UE::MeshPassUtils::Dispatch(RHICmdList, ComputeShader, ShaderBindings, *PassParameters, GroupCount);
+				UE::MeshPassUtils::Dispatch(RHICmdList, ComputeShader, ShaderBindings, *PassParameters, GroupCount);
 			}
 		);
 	}
