LgithgFunctionAtlas - fix debug after window resizing.

diff --git a/Engine/Source/Runtime/Renderer/Private/LightFunctionAtlas.cpp b/Engine/Source/Runtime/Renderer/Private/LightFunctionAtlas.cpp
--- a/Engine/Source/Runtime/Renderer/Private/LightFunctionAtlas.cpp
+++ b/Engine/Source/Runtime/Renderer/Private/LightFunctionAtlas.cpp
@@ -768,6 +768,9 @@ FScreenPassTexture FLightFunctionAtlas::AddDebugVisualizationPasses(FRDGBuilder&
 		float DrawPosX = float(DstPoint.X + DstSize.X) + 30.0f;
 		float DrawPosY = float(DstPoint.Y) + 10.0f;
 
+		const float DPIScale = Canvas.GetDPIScale();
+		Canvas.SetBaseTransform(FMatrix(FScaleMatrix(DPIScale) * Canvas.CalcBaseTransform2D(Canvas.GetViewRect().Width(), Canvas.GetViewRect().Height())));
+
 		float DisplayResolutionRatio = float(DstSize.X) / float(AtlasResolution);
 
 		Canvas.DrawShadowedString(DstPoint.X + 180.0f, DstPoint.Y - 20.0f, TEXT("LIGHT FUNCTION ATLAS"), GEngine->GetLargeFont(), FLinearColor::White);
