Add an option to run non-tiled SSR using async compute

diff --git a/Engine/Shaders/Private/SSRT/SSRTReflections.usf b/Engine/Shaders/Private/SSRT/SSRTReflections.usf
--- a/Engine/Shaders/Private/SSRT/SSRTReflections.usf
+++ b/Engine/Shaders/Private/SSRT/SSRTReflections.usf
@@ -434,6 +434,35 @@ void ScreenSpaceReflectionsPS(
 	);
 }
 
+RWTexture2D<float4> SSRColorOutput;
+#if SSR_OUTPUT_FOR_DENOISER
+RWTexture2D<float> SSRHitDistanceOutput;
+#endif
+
+[numthreads(THREADGROUP_SIZE_X, THREADGROUP_SIZE_Y, 1)]
+void ScreenSpaceReflectionsCS(uint3 DispatchThreadId : SV_DispatchThreadID)
+{
+	if (any(DispatchThreadId.xy >= View.ViewRectMinAndSize.zw))
+	{
+		return;
+	}
+
+	float4 SvPosition = float4(View.ViewRectMin.xy + DispatchThreadId.xy + float2(0.5, 0.5), 0, 1); // .zw unused
+	float4 OutColor;
+	float4 OutClosestHitDistance;
+
+	ScreenSpaceReflections(SvPosition, OutColor
+#if SSR_OUTPUT_FOR_DENOISER
+		,OutClosestHitDistance
+#endif
+	);
+
+	SSRColorOutput[DispatchThreadId.xy] = OutColor;
+#if SSR_OUTPUT_FOR_DENOISER
+	SSRHitDistanceOutput[DispatchThreadId.xy] = OutClosestHitDistance.x;
+#endif
+}
+
 void VisualizeTiledScreenSpaceReflectionsPS(
 	float4 SvPosition : SV_POSITION
 	, out float4 OutColor : SV_Target0
diff --git a/Engine/Source/Runtime/Renderer/Private/ScreenSpaceRayTracing.cpp b/Engine/Source/Runtime/Renderer/Private/ScreenSpaceRayTracing.cpp
--- a/Engine/Source/Runtime/Renderer/Private/ScreenSpaceRayTracing.cpp
+++ b/Engine/Source/Runtime/Renderer/Private/ScreenSpaceRayTracing.cpp
@@ -46,6 +46,12 @@ static TAutoConsoleVariable<int32> CVarSSRStencil(
 	TEXT(" 0 is off (default), 1 is on"),
 	ECVF_RenderThreadSafe);
 
+static TAutoConsoleVariable<int32> CVarSSRCompute(
+	TEXT("r.SSR.Compute"), 0,
+	TEXT("Use compute for SSR if possible. Only available for non tiled SSR with no stencil prepass currently\n")
+	TEXT(" 0 is PS (default), 1 is sync compute, 2 is async compute"),
+	ECVF_RenderThreadSafe);
+
 static TAutoConsoleVariable<int32> CVarSSGILeakFreeReprojection(
 	TEXT("r.SSGI.LeakFreeReprojection"), 1,
 	TEXT("Whether use a more expensive but leak free reprojection of previous frame's scene color.\n"),
@@ -500,6 +506,41 @@ class FScreenSpaceReflectionsPS : public FGlobalShader
 	END_SHADER_PARAMETER_STRUCT()
 };
 
+class FScreenSpaceReflectionsCS : public FGlobalShader
+{
+	DECLARE_GLOBAL_SHADER(FScreenSpaceReflectionsCS);
+	SHADER_USE_PARAMETER_STRUCT(FScreenSpaceReflectionsCS, FGlobalShader);
+
+	using FPermutationDomain = TShaderPermutationDomain<FSSRQualityDim, FSSROutputForDenoiser>;
+
+	static bool ShouldCompilePermutation(const FGlobalShaderPermutationParameters& Parameters)
+	{
+		FPermutationDomain PermutationVector(Parameters.PermutationId);
+		return IsFeatureLevelSupported(Parameters.Platform, ERHIFeatureLevel::SM5);
+	}
+
+	static void ModifyCompilationEnvironment(const FGlobalShaderPermutationParameters& Parameters, FShaderCompilerEnvironment& OutEnvironment)
+	{
+		FGlobalShader::ModifyCompilationEnvironment(Parameters, OutEnvironment);
+		OutEnvironment.SetDefine(TEXT("SUPPORTS_ANISOTROPIC_MATERIALS"), FDataDrivenShaderPlatformInfo::GetSupportsAnisotropicMaterials(Parameters.Platform));
+		OutEnvironment.SetDefine(TEXT("THREADGROUP_SIZE_X"), GetThreadGroupSize().X);
+		OutEnvironment.SetDefine(TEXT("THREADGROUP_SIZE_Y"), GetThreadGroupSize().Y);
+	}
+
+	static FIntPoint GetThreadGroupSize()
+	{
+		return 8;
+	}
+
+	BEGIN_SHADER_PARAMETER_STRUCT(FParameters, )
+		SHADER_PARAMETER_STRUCT_INCLUDE(FSSRCommonParameters, CommonParameters)
+		SHADER_PARAMETER_STRUCT_INCLUDE(FSSRPassCommonParameters, SSRPassCommonParameter)
+		SHADER_PARAMETER_RDG_UNIFORM_BUFFER(FSubstrateGlobalUniformParameters, Substrate)
+		SHADER_PARAMETER_RDG_TEXTURE_UAV(RWTexture2D<float4>, SSRColorOutput)
+		SHADER_PARAMETER_RDG_TEXTURE_UAV(RWTexture2D<float>, SSRHitDistanceOutput)
+	END_SHADER_PARAMETER_STRUCT()
+};
+
 // This is duplicated from FWaterTileVS because vertex shader should share Parameters structure for everything to be registered correctly in a RDG pass.
 class FScreenSpaceReflectionsTileVS : public FGlobalShader
 {
@@ -602,6 +643,7 @@ class FScreenSpaceCastStandaloneRayCS : public FGlobalShader
 IMPLEMENT_GLOBAL_SHADER(FSSRTPrevFrameReductionCS, "/Engine/Private/SSRT/SSRTPrevFrameReduction.usf", "MainCS", SF_Compute);
 IMPLEMENT_GLOBAL_SHADER(FSSRTDiffuseTileClassificationCS, "/Engine/Private/SSRT/SSRTTileClassification.usf", "MainCS", SF_Compute);
 IMPLEMENT_GLOBAL_SHADER(FScreenSpaceReflectionsPS,        "/Engine/Private/SSRT/SSRTReflections.usf", "ScreenSpaceReflectionsPS", SF_Pixel);
+IMPLEMENT_GLOBAL_SHADER(FScreenSpaceReflectionsCS, "/Engine/Private/SSRT/SSRTReflections.usf", "ScreenSpaceReflectionsCS", SF_Compute);
 IMPLEMENT_GLOBAL_SHADER(FScreenSpaceReflectionsTileVS,    "/Engine/Private/SingleLayerWaterComposite.usf", "WaterTileVS", SF_Vertex);
 IMPLEMENT_GLOBAL_SHADER(FVisualizeTiledScreenSpaceReflectionsPS, "/Engine/Private/SSRT/SSRTReflections.usf", "VisualizeTiledScreenSpaceReflectionsPS", SF_Pixel)
 IMPLEMENT_GLOBAL_SHADER(FVisualizeTiledScreenSpaceReflectionsVS, "/Engine/Private/SingleLayerWaterComposite.usf", "WaterTileVS", SF_Vertex)
@@ -1167,18 +1209,43 @@ void RenderScreenSpaceReflections(
 		PassParameters->HZBSampler = TStaticSamplerState<SF_Point>::GetRHI();
 	};
 
-	FScreenSpaceReflectionsPS::FPermutationDomain PermutationVector;
-	PermutationVector.Set<FSSRQualityDim>(SSRQuality);
-	PermutationVector.Set<FSSROutputForDenoiser>(bDenoiser);
-		
-	FScreenSpaceReflectionsPS::FParameters* PassParameters = GraphBuilder.AllocParameters<FScreenSpaceReflectionsPS::FParameters>();
-	PassParameters->CommonParameters = CommonParameters;
-	SetSSRParameters(&PassParameters->SSRPassCommonParameter);
-	PassParameters->Substrate = Substrate::BindSubstrateGlobalUniformParameters(View);
-	PassParameters->RenderTargets = RenderTargets;
-	PassParameters->RenderTargets.ShadingRateTexture = GVRSImageManager.GetVariableRateShadingImage(GraphBuilder, View, FVariableRateShadingImageManager::EVRSPassType::SSR);
+	const int32 CVarSSRComputeValue = CVarSSRCompute.GetValueOnRenderThread();
+	// Compute path is only implemented for non-tiled SSR currently
+	const bool bCompute = CVarSSRComputeValue > 0 && !TiledScreenSpaceReflection && !SSRStencilPrePass && !bSingleLayerWater;
+	const bool bAsyncCompute = bCompute && CVarSSRComputeValue == 2;
+
+	FScreenSpaceReflectionsPS::FPermutationDomain PermutationVectorPS;
+	FScreenSpaceReflectionsCS::FPermutationDomain PermutationVectorCS;
+	FScreenSpaceReflectionsPS::FParameters* PassParametersPS = nullptr;
+	FScreenSpaceReflectionsCS::FParameters* PassParametersCS = nullptr;
 
-	TShaderMapRef<FScreenSpaceReflectionsPS> PixelShader(View.ShaderMap, PermutationVector);
+	if (bCompute)
+	{
+		PermutationVectorCS.Set<FSSRQualityDim>(SSRQuality);
+		PermutationVectorCS.Set<FSSROutputForDenoiser>(bDenoiser);
+
+		PassParametersCS = GraphBuilder.AllocParameters<FScreenSpaceReflectionsCS::FParameters>();
+		PassParametersCS->CommonParameters = CommonParameters;
+		SetSSRParameters(&PassParametersCS->SSRPassCommonParameter);
+		PassParametersCS->Substrate = Substrate::BindSubstrateGlobalUniformParameters(View);
+		PassParametersCS->SSRColorOutput = GraphBuilder.CreateUAV(DenoiserInputs->Color);
+		if (bDenoiser)
+		{
+			PassParametersCS->SSRHitDistanceOutput = GraphBuilder.CreateUAV(DenoiserInputs->RayHitDistance);
+		}
+	}
+	else
+	{
+		PermutationVectorPS.Set<FSSRQualityDim>(SSRQuality);
+		PermutationVectorPS.Set<FSSROutputForDenoiser>(bDenoiser);
+
+		PassParametersPS = GraphBuilder.AllocParameters<FScreenSpaceReflectionsPS::FParameters>();
+		PassParametersPS->CommonParameters = CommonParameters;
+		SetSSRParameters(&PassParametersPS->SSRPassCommonParameter);
+		PassParametersPS->Substrate = Substrate::BindSubstrateGlobalUniformParameters(View);
+		PassParametersPS->RenderTargets = RenderTargets;
+		PassParametersPS->RenderTargets.ShadingRateTexture = GVRSImageManager.GetVariableRateShadingImage(GraphBuilder, View, FVariableRateShadingImageManager::EVRSPassType::SSR);
+	}
 
 	RDG_GPU_STAT_SCOPE(GraphBuilder, ScreenSpaceReflections);
 
@@ -1187,31 +1254,50 @@ void RenderScreenSpaceReflections(
 
 	if (TiledScreenSpaceReflection == nullptr)
 	{
-		ClearUnusedGraphResources(PixelShader, PassParameters);
-
-		GraphBuilder.AddPass(
-			RDG_EVENT_NAME("SSR RayMarch(Quality=%d RayPerPixel=%d%s) %dx%d",
-				SSRQuality, RayTracingConfigs.RayCountPerPixel, bDenoiser ? TEXT(" DenoiserOutput") : TEXT(""),
-				View.ViewRect.Width(), View.ViewRect.Height()),
-			PassParameters,
-			ERDGPassFlags::Raster,
-			[PassParameters, &View, PixelShader, SSRStencilPrePass](FRHICommandList& RHICmdList)
+		if (bCompute)
 		{
-			RHICmdList.SetViewport(View.ViewRect.Min.X, View.ViewRect.Min.Y, 0.0f, View.ViewRect.Max.X, View.ViewRect.Max.Y, 1.0f);
-		
-			FGraphicsPipelineStateInitializer GraphicsPSOInit;
-			FPixelShaderUtils::InitFullscreenPipelineState(RHICmdList, View.ShaderMap, PixelShader, /* out */ GraphicsPSOInit);
-			if (SSRStencilPrePass)
+			TShaderMapRef<FScreenSpaceReflectionsCS> ComputeShader(View.ShaderMap, PermutationVectorCS);
+
+			FComputeShaderUtils::AddPass(
+				GraphBuilder,
+				RDG_EVENT_NAME("SSR RayMarchCS(Quality=%d RayPerPixel=%d%s) %dx%d",
+					SSRQuality, RayTracingConfigs.RayCountPerPixel, bDenoiser ? TEXT(" DenoiserOutput") : TEXT(""),
+					View.ViewRect.Width(), View.ViewRect.Height()),
+				bAsyncCompute ? ERDGPassFlags::AsyncCompute : ERDGPassFlags::Compute,
+				ComputeShader,
+				PassParametersCS,
+				FComputeShaderUtils::GetGroupCount(View.ViewRect.Size(), FScreenSpaceReflectionsCS::GetThreadGroupSize()));
+		}
+		else
+		{
+			TShaderMapRef<FScreenSpaceReflectionsPS> PixelShader(View.ShaderMap, PermutationVectorPS);
+
+			ClearUnusedGraphResources(PixelShader, PassParametersPS);
+
+			GraphBuilder.AddPass(
+				RDG_EVENT_NAME("SSR RayMarch(Quality=%d RayPerPixel=%d%s) %dx%d",
+					SSRQuality, RayTracingConfigs.RayCountPerPixel, bDenoiser ? TEXT(" DenoiserOutput") : TEXT(""),
+					View.ViewRect.Width(), View.ViewRect.Height()),
+				PassParametersPS,
+				ERDGPassFlags::Raster,
+				[PassParametersPS, &View, PixelShader, SSRStencilPrePass](FRHICommandList& RHICmdList)
 			{
-				// Clobers the stencil to pixel that should not compute SSR
-				GraphicsPSOInit.DepthStencilState = TStaticDepthStencilState<false, CF_Always, true, CF_Equal, SO_Keep, SO_Keep, SO_Keep>::GetRHI();
-			}
+				RHICmdList.SetViewport(View.ViewRect.Min.X, View.ViewRect.Min.Y, 0.0f, View.ViewRect.Max.X, View.ViewRect.Max.Y, 1.0f);
 
-			SetGraphicsPipelineState(RHICmdList, GraphicsPSOInit, 0x80);
-			SetShaderParameters(RHICmdList, PixelShader, PixelShader.GetPixelShader(), *PassParameters);
+				FGraphicsPipelineStateInitializer GraphicsPSOInit;
+				FPixelShaderUtils::InitFullscreenPipelineState(RHICmdList, View.ShaderMap, PixelShader, /* out */ GraphicsPSOInit);
+				if (SSRStencilPrePass)
+				{
+					// Clobers the stencil to pixel that should not compute SSR
+					GraphicsPSOInit.DepthStencilState = TStaticDepthStencilState<false, CF_Always, true, CF_Equal, SO_Keep, SO_Keep, SO_Keep>::GetRHI();
+				}
 
-			FPixelShaderUtils::DrawFullscreenTriangle(RHICmdList);
-		});
+				SetGraphicsPipelineState(RHICmdList, GraphicsPSOInit, 0x80);
+				SetShaderParameters(RHICmdList, PixelShader, PixelShader.GetPixelShader(), *PassParametersPS);
+
+				FPixelShaderUtils::DrawFullscreenTriangle(RHICmdList);
+			});
+		}
 	}
 	else if (!bVisualizeTiledScreenSpaceReflection)
 	{
@@ -1220,18 +1306,20 @@ void RenderScreenSpaceReflections(
 		FScreenSpaceReflectionsTileVS::FPermutationDomain VsPermutationVector;
 		TShaderMapRef<FScreenSpaceReflectionsTileVS> VertexShader(View.ShaderMap, VsPermutationVector);
 
-		PassParameters->TileListData = TiledScreenSpaceReflection->TileListDataBufferSRV;
-		PassParameters->IndirectDrawParameter = TiledScreenSpaceReflection->DrawIndirectParametersBuffer;
+		PassParametersPS->TileListData = TiledScreenSpaceReflection->TileListDataBufferSRV;
+		PassParametersPS->IndirectDrawParameter = TiledScreenSpaceReflection->DrawIndirectParametersBuffer;
+
+		TShaderMapRef<FScreenSpaceReflectionsPS> PixelShader(View.ShaderMap, PermutationVectorPS);
 
-		ClearUnusedGraphResources(VertexShader, PixelShader, PassParameters);
+		ClearUnusedGraphResources(VertexShader, PixelShader, PassParametersPS);
 
 		GraphBuilder.AddPass(
 			RDG_EVENT_NAME("SSR RayMarch(Quality=%d RayPerPixel=%d%s) %dx%d",
 				SSRQuality, RayTracingConfigs.RayCountPerPixel, bDenoiser ? TEXT(" DenoiserOutput") : TEXT(""),
 				View.ViewRect.Width(), View.ViewRect.Height()),
-			PassParameters,
+			PassParametersPS,
 			ERDGPassFlags::Raster,
-			[PassParameters, &View, VertexShader, PixelShader, SSRStencilPrePass](FRHICommandList& RHICmdList)
+			[PassParametersPS, &View, VertexShader, PixelShader, SSRStencilPrePass](FRHICommandList& RHICmdList)
 		{
 			RHICmdList.SetViewport(View.ViewRect.Min.X, View.ViewRect.Min.Y, 0.0f, View.ViewRect.Max.X, View.ViewRect.Max.Y, 1.0f);
 
@@ -1248,12 +1336,12 @@ void RenderScreenSpaceReflections(
 			GraphicsPSOInit.BoundShaderState.PixelShaderRHI = PixelShader.GetPixelShader();
 
 			SetGraphicsPipelineState(RHICmdList, GraphicsPSOInit, 0x80);
-			SetShaderParameters(RHICmdList, VertexShader, VertexShader.GetVertexShader(), *PassParameters);
-			SetShaderParameters(RHICmdList, PixelShader, PixelShader.GetPixelShader(), *PassParameters);
+			SetShaderParameters(RHICmdList, VertexShader, VertexShader.GetVertexShader(), *PassParametersPS);
+			SetShaderParameters(RHICmdList, PixelShader, PixelShader.GetPixelShader(), *PassParametersPS);
 
-			PassParameters->IndirectDrawParameter->MarkResourceAsUsed();
+			PassParametersPS->IndirectDrawParameter->MarkResourceAsUsed();
 
-			RHICmdList.DrawPrimitiveIndirect(PassParameters->IndirectDrawParameter->GetIndirectRHICallBuffer(), 0);
+			RHICmdList.DrawPrimitiveIndirect(PassParametersPS->IndirectDrawParameter->GetIndirectRHICallBuffer(), 0);
 		});
 	}
 	else
@@ -1272,7 +1360,7 @@ void RenderScreenSpaceReflections(
 
 		FVisualizeTiledScreenSpaceReflectionsPS::FParameters* VisualizePassParameters = GraphBuilder.AllocParameters<FVisualizeTiledScreenSpaceReflectionsPS::FParameters>();
 
-		VisualizePassParameters->CommonParameters = *PassParameters;
+		VisualizePassParameters->CommonParameters = *PassParametersPS;
 		VisualizePassParameters->CommonParameters.TileListData = TiledScreenSpaceReflection->TileListDataBufferSRV;
 		VisualizePassParameters->CommonParameters.IndirectDrawParameter = TiledScreenSpaceReflection->DrawIndirectParametersBuffer;
 
