Fixed a case where sky atmosphere would override sky pixel from a IsSky dome mesh when alpha hold out was enabled.

diff --git a/Engine/Source/Runtime/Renderer/Private/ReflectionEnvironmentRealTimeCapture.cpp b/Engine/Source/Runtime/Renderer/Private/ReflectionEnvironmentRealTimeCapture.cpp
--- a/Engine/Source/Runtime/Renderer/Private/ReflectionEnvironmentRealTimeCapture.cpp
+++ b/Engine/Source/Runtime/Renderer/Private/ReflectionEnvironmentRealTimeCapture.cpp
@@ -570,6 +570,7 @@ void FScene::AllocateAndCaptureFrameSkyEnvMap(
 				SkyRC.SceneUniformBuffer = SceneRenderer.GetSceneUniforms().GetBuffer(GraphBuilder);
 
 				SkyRC.ViewMatrices = &CubeViewMatrices;
+				SkyRC.bSceneHasSkyMaterial = MainView.bSceneHasSkyMaterial;
 
 				SkyRC.SkyAtmosphereViewLutTexture = BlackDummy2dTex;
 				SkyRC.SkyAtmosphereCameraAerialPerspectiveVolume = BlackDummy3dTex;
diff --git a/Engine/Source/Runtime/Renderer/Private/SkyAtmosphereRendering.cpp b/Engine/Source/Runtime/Renderer/Private/SkyAtmosphereRendering.cpp
--- a/Engine/Source/Runtime/Renderer/Private/SkyAtmosphereRendering.cpp
+++ b/Engine/Source/Runtime/Renderer/Private/SkyAtmosphereRendering.cpp
@@ -1823,7 +1823,7 @@ void FSceneRenderer::RenderSkyAtmosphereInternal(
 		const bool bFastAerialPerspectiveDepthTest = SkyRC.bFastAerialPerspectiveDepthTest;
 		const bool SkyAtmosphereOutputsAlpha = IsPostProcessingWithAlphaChannelSupported();
 		const bool SkyAtmosphereAlphaHoldOut = SkyAtmosphereOutputsAlpha && SkyAtmosphereSceneProxy.IsHoldout();
-		const bool bRenderSkyPixel = SkyRC.bRenderSkyPixel || SkyAtmosphereOutputsAlpha;	// In this case we need to write alpha holdout values in the sky pixels.
+		const bool bRenderSkyPixel = SkyRC.bRenderSkyPixel || (SkyAtmosphereOutputsAlpha && !SkyRC.bSceneHasSkyMaterial);	// In this case we need to write alpha holdout values in the sky pixels. If there is no IsSky dmoe meshes.
 
 		FRenderSkyAtmospherePS::FPermutationDomain PsPermutationVector;
 		PsPermutationVector.Set<FSampleCloudSkyAO>(SkyRC.bShouldSampleCloudSkyAO);
@@ -2020,6 +2020,7 @@ void FSceneRenderer::RenderSkyAtmosphere(FRDGBuilder& GraphBuilder, const FMinim
 		FViewInfo& View = Views[ViewIndex];
 		SkyRC.ViewMatrices = &View.ViewMatrices;
 		SkyRC.ViewUniformBuffer = View.ViewUniformBuffer;
+		SkyRC.bSceneHasSkyMaterial = View.bSceneHasSkyMaterial;
 
 		SkyRC.Viewport = View.ViewRect;
 		SkyRC.bLightDiskEnabled = !View.bIsReflectionCapture;
diff --git a/Engine/Source/Runtime/Renderer/Private/SkyAtmosphereRendering.h b/Engine/Source/Runtime/Renderer/Private/SkyAtmosphereRendering.h
--- a/Engine/Source/Runtime/Renderer/Private/SkyAtmosphereRendering.h
+++ b/Engine/Source/Runtime/Renderer/Private/SkyAtmosphereRendering.h
@@ -96,6 +96,7 @@ struct FSkyAtmosphereRenderContext
 	FViewMatrices* ViewMatrices;			// The actual view matrices we use to render the sky
 	TUniformBufferRef<FViewUniformShaderParameters> ViewUniformBuffer;
 	TRDGUniformBufferRef<FSceneUniformParameters> SceneUniformBuffer;
+	bool bSceneHasSkyMaterial;
 
 	FRenderTargetBindingSlots RenderTargets;
 
diff --git a/Engine/Source/Runtime/Renderer/Private/VolumetricCloudRendering.cpp b/Engine/Source/Runtime/Renderer/Private/VolumetricCloudRendering.cpp
--- a/Engine/Source/Runtime/Renderer/Private/VolumetricCloudRendering.cpp
+++ b/Engine/Source/Runtime/Renderer/Private/VolumetricCloudRendering.cpp
@@ -2846,6 +2846,7 @@ bool FSceneRenderer::RenderVolumetricCloud(
 					SkyRC.FeatureLevel = ViewInfo.FeatureLevel;
 
 					SkyRC.bRenderSkyPixel = false;
+					SkyRC.bSceneHasSkyMaterial = false;
 
 					if (ViewInfo.SkyAtmosphereViewLutTexture && ViewInfo.SkyAtmosphereCameraAerialPerspectiveVolume)
 					{
