Fix a shader compile error caused by undeclared defines

diff --git a/Engine/Shaders/Private/SSRT/SSRTReflections.usf b/Engine/Shaders/Private/SSRT/SSRTReflections.usf
--- a/Engine/Shaders/Private/SSRT/SSRTReflections.usf
+++ b/Engine/Shaders/Private/SSRT/SSRTReflections.usf
@@ -434,6 +434,8 @@ void ScreenSpaceReflectionsPS(
 	);
 }
 
+#if COMPUTE_SHADER
+
 RWTexture2D<float4> SSRColorOutput;
 #if SSR_OUTPUT_FOR_DENOISER
 RWTexture2D<float> SSRHitDistanceOutput;
@@ -463,6 +465,8 @@ void ScreenSpaceReflectionsCS(uint3 DispatchThreadId : SV_DispatchThreadID)
 #endif
 }
 
+#endif // COMPUTE_SHADER
+
 void VisualizeTiledScreenSpaceReflectionsPS(
 	float4 SvPosition : SV_POSITION
 	, out float4 OutColor : SV_Target0
diff --git a/Engine/Source/Runtime/Renderer/Private/ScreenSpaceRayTracing.cpp b/Engine/Source/Runtime/Renderer/Private/ScreenSpaceRayTracing.cpp
--- a/Engine/Source/Runtime/Renderer/Private/ScreenSpaceRayTracing.cpp
+++ b/Engine/Source/Runtime/Renderer/Private/ScreenSpaceRayTracing.cpp
@@ -525,6 +525,7 @@ class FScreenSpaceReflectionsCS : public FGlobalShader
 		OutEnvironment.SetDefine(TEXT("SUPPORTS_ANISOTROPIC_MATERIALS"), FDataDrivenShaderPlatformInfo::GetSupportsAnisotropicMaterials(Parameters.Platform));
 		OutEnvironment.SetDefine(TEXT("THREADGROUP_SIZE_X"), GetThreadGroupSize().X);
 		OutEnvironment.SetDefine(TEXT("THREADGROUP_SIZE_Y"), GetThreadGroupSize().Y);
+		OutEnvironment.SetDefine(TEXT("COMPUTE_SHADER"), 1);
 	}
 
 	static FIntPoint GetThreadGroupSize()
