Removed deprecated code from FRayTracingMaterialGatheringContext.
- also removed FDeferredShadingRayTracingMaterialGatheringContext since it doesn't currently provide any functionality over the base class.

diff --git a/Engine/Source/Runtime/Engine/Public/SceneManagement.h b/Engine/Source/Runtime/Engine/Public/SceneManagement.h
--- a/Engine/Source/Runtime/Engine/Public/SceneManagement.h
+++ b/Engine/Source/Runtime/Engine/Public/SceneManagement.h
@@ -2175,9 +2175,6 @@ struct FRayTracingMaterialGatheringContext
 		FGlobalDynamicReadBuffer& InGlobalDynamicReadBuffer);
 
 	ENGINE_API virtual ~FRayTracingMaterialGatheringContext();
-
-	UE_DEPRECATED(5.4, "InstanceMaskAndFlags is automatically built and cached in RayTracing.cpp")
-	virtual FRayTracingMaskAndFlags BuildInstanceMaskAndFlags(const FRayTracingInstance& Instance, const FPrimitiveSceneProxy& ScenePrimitive) = 0;
 };
 #endif
 
diff --git a/Engine/Source/Runtime/Renderer/Private/RayTracing/RayTracing.cpp b/Engine/Source/Runtime/Renderer/Private/RayTracing/RayTracing.cpp
--- a/Engine/Source/Runtime/Renderer/Private/RayTracing/RayTracing.cpp
+++ b/Engine/Source/Runtime/Renderer/Private/RayTracing/RayTracing.cpp
@@ -653,26 +653,6 @@ namespace RayTracing
 		}
 	}
 
-	// Class to implement build instance mask and flags so that rendering related mask build is maintained in any renderer module.
-	// BuildInstanceMaskAndFlags() will be called in the Engine module where it does not know specifics of the ray tracing instance
-	// masks used by the renderer (e.g., path tracer mask might be different from raytracing mask).
-	struct FDeferredShadingRayTracingMaterialGatheringContext : public FRayTracingMaterialGatheringContext
-	{
-		FDeferredShadingRayTracingMaterialGatheringContext(
-			const FScene* InScene,
-			const FSceneView* InReferenceView,
-			const FSceneViewFamily& InReferenceViewFamily,
-			FRDGBuilder& InGraphBuilder,
-			FRayTracingMeshResourceCollector& InRayTracingMeshResourceCollector,
-			FGlobalDynamicReadBuffer& InDynamicReadBuffer)
-			:FRayTracingMaterialGatheringContext(InScene, InReferenceView, InReferenceViewFamily, InGraphBuilder, InRayTracingMeshResourceCollector, InDynamicReadBuffer) {}
-
-		virtual FRayTracingMaskAndFlags BuildInstanceMaskAndFlags(const FRayTracingInstance& Instance, const FPrimitiveSceneProxy& ScenePrimitive) override
-		{
-			return BuildRayTracingInstanceMaskAndFlags(Instance, ScenePrimitive, &ReferenceViewFamily);
-		}
-	};
-
 	bool GatherWorldInstancesForView(
 		FRDGBuilder& GraphBuilder,
 		FScene& Scene,
@@ -703,7 +683,7 @@ namespace RayTracing
 
 		View.RayTracingCullingParameters.Init(View);
 
-		FDeferredShadingRayTracingMaterialGatheringContext MaterialGatheringContext
+		FRayTracingMaterialGatheringContext MaterialGatheringContext
 		(
 			&Scene,
 			&View,
