Re-implement MeshEdges showflag. This fixes issues with z-fighting, TSR ghosting, missing Nanite geometry and adds flexibility in the composition step.
The wireframe overlay is now rendered as a separate pre-pass, and composited into EditorPrimitives after the main pass.

diff --git a/Engine/Shaders/Private/MeshEdges.usf b/Engine/Shaders/Private/MeshEdges.usf
new file mode 100644
--- /dev/null
+++ b/Engine/Shaders/Private/MeshEdges.usf
@@ -0,0 +1,111 @@
+// Copyright Epic Games, Inc. All Rights Reserved.
+
+/*
+ * MeshEdges: Wireframe composited over other shaded views (e.g. Lit)
+ */
+
+#include "Common.ush"
+#include "DeferredShadingCommon.ush"
+#include "ScreenPass.ush"
+#include "PostProcessCommon.ush"
+
+#if MSAA_SAMPLE_COUNT > 1
+Texture2DMS<float4, MSAA_SAMPLE_COUNT> WireframeColorTexture;
+Texture2DMS<float, MSAA_SAMPLE_COUNT>  WireframeDepthTexture;
+#else
+Texture2D WireframeColorTexture;
+Texture2D<float> WireframeDepthTexture;
+#endif
+
+SCREEN_PASS_TEXTURE_VIEWPORT(Wireframe)
+
+Texture2D DepthTexture;
+SamplerState DepthSampler;
+SCREEN_PASS_TEXTURE_VIEWPORT(Depth)
+float2 DepthTextureJitter;
+float4 SampleOffsetArray[MSAA_SAMPLE_COUNT];
+
+SCREEN_PASS_TEXTURE_VIEWPORT(Output)
+
+float Opacity;
+
+// Return true if A is closer than B
+bool IsCloser(float A, float B)
+{
+	return HAS_INVERTED_Z_BUFFER
+		? (A > B)
+		: (A < B);
+}
+
+// Compose wireframe texture over shaded rendertarget, predicated on depth test with nearest shaded surface
+void ComposeMeshEdgesPS(
+	float4 SvPosition : SV_POSITION,
+	#if MSAA_SAMPLE_COUNT > 1
+	uint TargetSampleIndex : SV_SampleIndex,
+	#endif
+	out float4 OutColor : SV_Target0,
+	out float OutDepth : SV_Depth)
+{
+	float2 ViewportUV = (SvPosition.xy - Output_ViewportMin) * Output_ViewportSizeInverse;
+
+	float2 WireframeSamplePosition = (ViewportUV * Wireframe_ViewportSize) + Wireframe_ViewportMin;
+
+	float4 WireframeColor;
+	float WireframeDeviceZ;
+	float ShadedDeviceZ;
+	#if MSAA_SAMPLE_COUNT > 1
+	{
+		WireframeColor = float4(WireframeColorTexture.Load(WireframeSamplePosition, TargetSampleIndex).rgb, 1.0);
+		WireframeDeviceZ = WireframeDepthTexture.Load(WireframeSamplePosition, TargetSampleIndex).r;
+
+		const float2 SampleOffset = SampleOffsetArray[TargetSampleIndex].xy;
+
+		float2 ShadedDepthUV = (ViewportUV * Depth_ViewportSize + Depth_ViewportMin + SampleOffset - DepthTextureJitter) * Depth_ExtentInverse;
+		ShadedDepthUV = clamp(ShadedDepthUV, Depth_UVViewportBilinearMin, Depth_UVViewportBilinearMax);
+		ShadedDeviceZ = Texture2DSampleLevel(DepthTexture, DepthSampler, ShadedDepthUV, 0).r;
+	}
+	#else
+	{
+		WireframeColor = WireframeColorTexture.Load(int3(WireframeSamplePosition, 0));
+		WireframeDeviceZ = WireframeDepthTexture.Load(int3(WireframeSamplePosition, 0)).r;
+
+		float2 ShadedDepthUV = (ViewportUV * Depth_ViewportSize + Depth_ViewportMin - DepthTextureJitter) * Depth_ExtentInverse;
+		ShadedDepthUV = clamp(ShadedDepthUV, Depth_UVViewportBilinearMin, Depth_UVViewportBilinearMax);
+		ShadedDeviceZ = Texture2DSampleLevel(DepthTexture, DepthSampler, ShadedDepthUV, 0).r;
+	}
+	#endif
+
+	if (WireframeDeviceZ == ShadedDeviceZ && WireframeDeviceZ == FarDepthValue)
+	{
+		discard;
+	}
+
+	// Stop z-fighting by biasing wireframe depth towards camera.
+	// The bias factor is a manually tuned heuristic based on surface slope and the proximity of the wireframe to the surface.
+	float2 ShadedDepthSlope = abs(float2(ddx(ShadedDeviceZ), ddy(ShadedDeviceZ)) / ShadedDeviceZ);
+	float DepthRatio = WireframeDeviceZ / ShadedDeviceZ;
+	const float SlopeWeight = 2.0;
+	float DepthBiasFactor = lerp(0, 1, (ShadedDepthSlope.x + ShadedDepthSlope.y) * SlopeWeight) + max(1, DepthRatio);
+
+	WireframeDeviceZ *= DepthBiasFactor;
+
+	if (IsCloser(ShadedDeviceZ, WireframeDeviceZ))
+	{
+		discard;
+	}
+		
+	OutDepth = WireframeDeviceZ;
+	
+	// Fixes the gamma of editor primitives
+	{
+		// Bring out of premultiplied.
+		WireframeColor.rgb /= max(WireframeColor.a, 0.0001f);
+		// Fix gamma.
+		WireframeColor.rgb = pow(WireframeColor.rgb, 2.2f);
+		// Bring back to premultiplied
+		WireframeColor.rgb *= WireframeColor.a;
+	}
+
+	OutColor = WireframeColor;
+	OutColor.a *= Opacity;
+}
diff --git a/Engine/Source/Editor/UnrealEd/Private/EditorViewportClient.cpp b/Engine/Source/Editor/UnrealEd/Private/EditorViewportClient.cpp
--- a/Engine/Source/Editor/UnrealEd/Private/EditorViewportClient.cpp
+++ b/Engine/Source/Editor/UnrealEd/Private/EditorViewportClient.cpp
@@ -2887,7 +2887,11 @@ bool FEditorViewportClient::SupportsPreviewResolutionFraction() const
 	}
 
 	// Don't do preview screen percentage in certain cases.
-	if (EngineShowFlags.VisualizeBuffer || EngineShowFlags.VisualizeNanite || EngineShowFlags.VisualizeVirtualShadowMap || IsVisualizeCalibrationMaterialEnabled())
+	if (EngineShowFlags.VisualizeBuffer 
+		|| EngineShowFlags.MeshEdges
+		|| EngineShowFlags.VisualizeNanite 
+		|| EngineShowFlags.VisualizeVirtualShadowMap 
+		|| IsVisualizeCalibrationMaterialEnabled())
 	{
 		return false;
 	}
diff --git a/Engine/Source/Runtime/Engine/Private/PrimitiveDrawingUtils.cpp b/Engine/Source/Runtime/Engine/Private/PrimitiveDrawingUtils.cpp
--- a/Engine/Source/Runtime/Engine/Private/PrimitiveDrawingUtils.cpp
+++ b/Engine/Source/Runtime/Engine/Private/PrimitiveDrawingUtils.cpp
@@ -1681,45 +1681,6 @@ void ApplyViewModeOverrides(
 			Mesh.MaterialRenderProxy = InvalidSettingsMaterialInstance;
 			Collector.RegisterOneFrameMaterialProxy(InvalidSettingsMaterialInstance);
 		}
-
-		//Draw a wireframe overlay last, if requested
-		if (EngineShowFlags.MeshEdges)
-		{
-			FMeshBatch& MeshEdgeElement = Collector.AllocateMesh();
-			MeshEdgeElement = Mesh;
-			// Avoid infinite recursion
-			MeshEdgeElement.bCanApplyViewModeOverrides = false;
-
-			
-			// Draw the mesh's edges in blue, on top of the base geometry.
-			if (bMaterialModifiesMeshPosition)
-			{
-				// If the material is mesh-modifying, we cannot rely on substitution
-				auto WireframeMaterialInstance = new FOverrideSelectionColorMaterialRenderProxy(
-					MeshEdgeElement.MaterialRenderProxy,
-					PrimitiveSceneProxy->GetWireframeColor()
-				);
-
-				MeshEdgeElement.bWireframe = true;
-				MeshEdgeElement.MaterialRenderProxy = WireframeMaterialInstance;
-				Collector.RegisterOneFrameMaterialProxy(WireframeMaterialInstance);
-
-				Collector.AddMesh(ViewIndex, MeshEdgeElement);
-			}
-			else
-			{
-				auto WireframeMaterialInstance = new FColoredMaterialRenderProxy(
-					GEngine->WireframeMaterial->GetRenderProxy(),
-					PrimitiveSceneProxy->GetWireframeColor()
-				);
-
-				MeshEdgeElement.bWireframe = true;
-				MeshEdgeElement.MaterialRenderProxy = WireframeMaterialInstance;
-				Collector.RegisterOneFrameMaterialProxy(WireframeMaterialInstance);
-
-				Collector.AddMesh(ViewIndex, MeshEdgeElement);
-			}
-		}
 	}
 #endif
 }
diff --git a/Engine/Source/Runtime/Engine/Private/ShowFlags.cpp b/Engine/Source/Runtime/Engine/Private/ShowFlags.cpp
--- a/Engine/Source/Runtime/Engine/Private/ShowFlags.cpp
+++ b/Engine/Source/Runtime/Engine/Private/ShowFlags.cpp
@@ -622,6 +622,14 @@ void EngineShowFlagOverride(EShowFlagInitMode ShowFlagInitMode, EViewModeIndex V
 	{
 		EngineShowFlags.SetTonemapper(false);
 	}
+	
+	if (EngineShowFlags.MeshEdges)
+	{
+		// TAA-induced jittering causes a mismatch between the sampled positions for wireframe and lit-view pixels, resulting in z-fighting and flickering.
+		// Enabling jittering on the wireframe means also enabling TAA, which doesn't need it (MSAA) and causes ghosting artifacts.
+		EngineShowFlags.SetTemporalAA(false);
+	}
+
 
 	if (EngineShowFlags.Bones)
 	{
diff --git a/Engine/Source/Runtime/Engine/Public/SceneView.h b/Engine/Source/Runtime/Engine/Public/SceneView.h
--- a/Engine/Source/Runtime/Engine/Public/SceneView.h
+++ b/Engine/Source/Runtime/Engine/Public/SceneView.h
@@ -1970,6 +1970,9 @@ public:
 	 */
 	virtual DynamicRenderScaling::TMap<float> GetResolutionFractionsUpperBound() const = 0;
 
+	/** Create a new screen percentage interface for a new view family. */
+	virtual ISceneViewFamilyScreenPercentage* Fork_GameThread(const class FSceneViewFamily& ViewFamily) const = 0;
+
 protected:
 	/**
 	 * Setup view family's view's screen percentage on rendering thread.
@@ -1977,9 +1980,6 @@ protected:
 	 */
 	virtual DynamicRenderScaling::TMap<float> GetResolutionFractions_RenderThread() const = 0;
 
-	/** Create a new screen percentage interface for a new view family. */
-	virtual ISceneViewFamilyScreenPercentage* Fork_GameThread(const class FSceneViewFamily& ViewFamily) const = 0;
-
 	friend class FSceneViewFamily;
 	friend class FSceneRenderer;
 };
@@ -2334,20 +2334,25 @@ public:
 		check(SecondarySpatialUpscalerInterface == nullptr);
 	}
 
-	template<typename TExtensionData> TExtensionData* GetExtentionData()
+	template<typename TExtensionData> const TExtensionData* GetExtentionData() const
 	{
 		static_assert(TIsDerivedFrom<TExtensionData, ISceneViewFamilyExtentionData>::Value, "TExtensionData is not derived from ISceneViewFamilyExtentionData.");
 
-		for (TSharedRef<class ISceneViewFamilyExtentionData, ESPMode::ThreadSafe>& ViewExtensionData : ViewExtentionDatas)
+		for (const TSharedRef<class ISceneViewFamilyExtentionData, ESPMode::ThreadSafe>& ViewExtensionData : ViewExtentionDatas)
 		{
 			if (ViewExtensionData->GetSubclassIdentifier() == TExtensionData::GSubclassIdentifier)
 			{
-				return static_cast<TExtensionData*>(&ViewExtensionData.Get());
+				return static_cast<const TExtensionData*>(&ViewExtensionData.Get());
 			}
 		}
 		return nullptr;
 	}
 
+	template<typename TExtensionData> TExtensionData* GetExtentionData()
+	{
+		return const_cast<TExtensionData*>(AsConst(*this).GetExtentionData<TExtensionData>());
+	}
+
 	template<typename TExtensionData> TExtensionData* GetOrCreateExtentionData()
 	{
 		TExtensionData* ViewExtensionData = GetExtentionData<TExtensionData>();
diff --git a/Engine/Source/Runtime/Engine/Public/SceneViewExtension.h b/Engine/Source/Runtime/Engine/Public/SceneViewExtension.h
--- a/Engine/Source/Runtime/Engine/Public/SceneViewExtension.h
+++ b/Engine/Source/Runtime/Engine/Public/SceneViewExtension.h
@@ -140,6 +140,11 @@ public:
      */
     virtual void BeginRenderViewFamily(FSceneViewFamily& InViewFamily) = 0;
 
+	/**
+	* Called on game thread after the scene renderers have been created
+	*/
+	virtual void PostCreateSceneRenderer(const FSceneViewFamily& InViewFamily, ISceneRenderer* Renderer) {}
+
     /**
      * Called on render thread at the start of rendering.
      */
diff --git a/Engine/Source/Runtime/Renderer/Private/MeshEdgesRendering.cpp b/Engine/Source/Runtime/Renderer/Private/MeshEdgesRendering.cpp
new file mode 100644
--- /dev/null
+++ b/Engine/Source/Runtime/Renderer/Private/MeshEdgesRendering.cpp
@@ -0,0 +1,452 @@
+// Copyright Epic Games, Inc. All Rights Reserved.
+
+#include "MeshEdgesRendering.h"
+#include "DataDrivenShaderPlatformInfo.h"
+#include "Engine/Scene.h"
+#include "SceneInterface.h"
+#include "RenderingThread.h"
+#include "RHIStaticStates.h"
+#include "RendererInterface.h"
+#include "Camera/CameraTypes.h"
+#include "Shader.h"
+#include "TextureResource.h"
+#include "StaticBoundShaderState.h"
+#include "SceneUtils.h"
+#include "ScenePrivateBase.h"
+#include "PostProcess/SceneRenderTargets.h"
+#include "GlobalShader.h"
+#include "SceneRenderTargetParameters.h"
+#include "SceneRendering.h"
+#include "DeferredShadingRenderer.h"
+#include "ScenePrivate.h"
+#include "PostProcess/SceneFilterRendering.h"
+#include "PostProcess/PostProcessing.h"
+#include "LightRendering.h"
+#include "Materials/MaterialRenderProxy.h"
+#include "Components/SceneCaptureComponent.h"
+#include "Components/PlanarReflectionComponent.h"
+#include "Containers/ArrayView.h"
+#include "PipelineStateCache.h"
+#include "ClearQuad.h"
+#include "SceneTextureParameters.h"
+#include "SceneViewExtension.h"
+#include "PixelShaderUtils.h"
+#include "LegacyScreenPercentageDriver.h"
+#include "ScreenPass.h"
+#include "PostProcess/TemporalAA.h"
+
+class FComposeMeshEdgesPS : public FGlobalShader
+{
+	DECLARE_GLOBAL_SHADER(FComposeMeshEdgesPS);
+	SHADER_USE_PARAMETER_STRUCT(FComposeMeshEdgesPS, FGlobalShader)
+
+	static const uint32 kMSAASampleCountMaxLog2 = 3; // = log2(MSAASampleCountMax)
+	static const uint32 kMSAASampleCountMax = 1 << kMSAASampleCountMaxLog2;
+	class FSampleCountDimension : SHADER_PERMUTATION_RANGE_INT("MSAA_SAMPLE_COUNT_LOG2", 0, kMSAASampleCountMaxLog2 + 1);
+	using FPermutationDomain = TShaderPermutationDomain<FSampleCountDimension>;
+
+	static bool ShouldCompilePermutation(const FGlobalShaderPermutationParameters& Parameters)
+	{
+		return FDataDrivenShaderPlatformInfo::GetSupportsDebugViewShaders(Parameters.Platform);
+	}
+	
+	static void ModifyCompilationEnvironment(const FGlobalShaderPermutationParameters& Parameters, FShaderCompilerEnvironment& OutEnvironment)
+	{
+		const FPermutationDomain PermutationVector(Parameters.PermutationId);
+		const int32 SampleCount = 1 << PermutationVector.Get<FSampleCountDimension>();
+		OutEnvironment.SetDefine(TEXT("MSAA_SAMPLE_COUNT"), SampleCount);
+	}
+
+	BEGIN_SHADER_PARAMETER_STRUCT(FParameters, )
+		SHADER_PARAMETER_RDG_TEXTURE(Texture2D, WireframeColorTexture)
+		SHADER_PARAMETER_RDG_TEXTURE(Texture2D, WireframeDepthTexture)
+		SHADER_PARAMETER_STRUCT(FScreenPassTextureViewportParameters, Wireframe)
+
+		SHADER_PARAMETER_RDG_TEXTURE(Texture2D, DepthTexture)
+		SHADER_PARAMETER_SAMPLER(SamplerState, DepthSampler)
+		SHADER_PARAMETER_STRUCT(FScreenPassTextureViewportParameters, Depth)
+		SHADER_PARAMETER(FVector2f, DepthTextureJitter)
+		SHADER_PARAMETER_ARRAY(FVector4f, SampleOffsetArray, [FComposeMeshEdgesPS::kMSAASampleCountMax])
+
+		SHADER_PARAMETER_STRUCT(FScreenPassTextureViewportParameters, Output)
+		
+		SHADER_PARAMETER(float, Opacity)
+		
+		RENDER_TARGET_BINDING_SLOTS()
+	END_SHADER_PARAMETER_STRUCT()
+};
+
+IMPLEMENT_GLOBAL_SHADER(FComposeMeshEdgesPS, "/Engine/Private/MeshEdges.usf", "ComposeMeshEdgesPS", SF_Pixel);
+
+class FRenderTargetTexture : public FTexture, public FRenderTarget
+{
+public:
+
+	FRenderTargetTexture(FRHITextureCreateDesc InDesc)
+	: Desc(InDesc)
+	{}
+
+	virtual void InitRHI(FRHICommandListBase& RHICmdList)
+	{
+		// Create the sampler state RHI resource.
+		FSamplerStateInitializerRHI SamplerStateInitializer
+		(
+			SF_Bilinear,
+			AM_Clamp,
+			AM_Clamp,
+			AM_Clamp
+		);
+		SamplerStateRHI = GetOrCreateSamplerState(SamplerStateInitializer);
+
+		RenderTargetTextureRHI = TextureRHI = RHICreateTexture(Desc);
+	}
+
+	virtual uint32 GetSizeX() const
+	{
+		return Desc.GetSize().X;
+	}
+
+	virtual uint32 GetSizeY() const
+	{
+		return Desc.GetSize().Y;
+	}
+
+	virtual FIntPoint GetSizeXY() const { return FIntPoint(GetSizeX(), GetSizeY()); }
+
+	virtual float GetDisplayGamma() const { return 1.0f; }
+
+	virtual FString GetFriendlyName() const override { return Desc.DebugName; }
+
+private:
+	FRHITextureCreateDesc Desc;
+};
+
+class FMeshEdgesViewFamilyData : public ISceneViewFamilyExtentionData
+{
+public:
+	inline static const TCHAR* const GSubclassIdentifier = TEXT("FMeshEdgesViewFamilyData");
+	virtual const TCHAR* GetSubclassIdentifier() const override
+	{
+		return GSubclassIdentifier;
+	}
+
+	void CreateRenderTargets(ERHIFeatureLevel::Type FeatureLevel, FIntPoint DesiredBufferSize)
+	{
+		int NumMSAASamples = FSceneTexturesConfig::GetEditorPrimitiveNumSamples(FeatureLevel);
+
+		FRHITextureCreateDesc ColorDesc = FRHITextureCreateDesc::Create2D(TEXT("MeshEdgesRenderTarget"))
+			.SetExtent(DesiredBufferSize)
+			.SetFormat(PF_B8G8R8A8)
+			.SetClearValue(FClearValueBinding::Transparent)
+			.SetFlags(ETextureCreateFlags::RenderTargetable | ETextureCreateFlags::ShaderResource)
+			.SetInitialState(ERHIAccess::SRVMask)
+			.SetNumSamples(NumMSAASamples);
+
+		FRHITextureCreateDesc DepthDesc = FRHITextureCreateDesc::Create2D(TEXT("MeshEdgesDepthRenderTarget"))
+			.SetExtent(DesiredBufferSize)
+			.SetFormat(PF_DepthStencil)
+			.SetClearValue(FClearValueBinding::DepthFar)
+			.SetFlags(ETextureCreateFlags::DepthStencilTargetable | ETextureCreateFlags::ShaderResource)
+			.SetInitialState(ERHIAccess::SRVMask)
+			.SetNumSamples(NumMSAASamples);
+
+		WireframeColor = MakeUnique<FRenderTargetTexture>(ColorDesc);
+		WireframeDepth = MakeUnique<FRenderTargetTexture>(DepthDesc);
+	}
+
+	TUniquePtr<FRenderTargetTexture> WireframeColor;
+	TUniquePtr<FRenderTargetTexture> WireframeDepth;
+	TArray<FIntRect> ViewRects;
+	TArray<FMeshEdgesViewSettings> ViewSettings;
+};
+
+const FMeshEdgesViewSettings& GetMeshEdgesViewSettings(const FSceneViewFamily& ViewFamily, int ViewIndex)
+{
+	FMeshEdgesViewFamilyData* FamilyData = const_cast<FSceneViewFamily&>(ViewFamily).GetOrCreateExtentionData<FMeshEdgesViewFamilyData>();
+
+	if (FamilyData->ViewSettings.IsEmpty())
+	{
+		FamilyData->ViewSettings.Init(FMeshEdgesViewSettings{}, ViewFamily.Views.Num());
+	}
+
+	return FamilyData->ViewSettings[ViewIndex];
+}
+
+const FMeshEdgesViewSettings& GetMeshEdgesViewSettings(const FSceneView& View)
+{
+	check(View.Family);
+
+	int ViewIndex = 0;
+	for (; ViewIndex < View.Family->Views.Num(); ViewIndex++)
+	{
+		if (View.Family->Views[ViewIndex] == &View) break;
+	}
+
+	return GetMeshEdgesViewSettings(*View.Family, ViewIndex);
+}
+
+FMeshEdgesViewSettings& GetMeshEdgesViewSettings(FSceneView& View)
+{
+	return const_cast<FMeshEdgesViewSettings&>(GetMeshEdgesViewSettings(AsConst(View)));
+}
+
+void RenderMeshEdges(FSceneViewFamily& InViewFamily);
+
+class FMeshEdgesExtension : public FSceneViewExtensionBase
+{
+public:
+	FMeshEdgesExtension( const FAutoRegister& AutoRegister)
+	: FSceneViewExtensionBase( AutoRegister )
+	{}
+
+	virtual void SetupViewFamily(FSceneViewFamily& InViewFamily) override {}
+	virtual void SetupView(FSceneViewFamily& InViewFamily, FSceneView& InView) override {}
+	virtual void BeginRenderViewFamily(FSceneViewFamily& InViewFamily) override {}
+	virtual void PostCreateSceneRenderer(const FSceneViewFamily& InViewFamily, ISceneRenderer* Renderer) override
+	{
+		RenderMeshEdges(static_cast<FSceneRenderer*>(Renderer)->ViewFamily);
+	}
+
+	inline static TSharedPtr<FMeshEdgesExtension,ESPMode::ThreadSafe> Instance;
+
+private:
+	struct FStaticConstructor
+	{
+		FStaticConstructor()
+		{
+			FCoreDelegates::OnPostEngineInit.AddLambda([]()
+			{ 
+				Instance = FSceneViewExtensions::NewExtension<FMeshEdgesExtension>(); 
+			});
+			FCoreDelegates::OnEnginePreExit.AddLambda([]()
+			{ 
+				Instance = nullptr;
+			});
+		}
+	};
+	static inline FStaticConstructor StaticConstructor;
+};
+
+void CopyViewFamily(const FSceneViewFamily& SrcViewFamily, FSceneViewFamily& ViewFamily)
+{
+	ViewFamily.FrameNumber = SrcViewFamily.FrameNumber;
+	ViewFamily.FrameCounter = SrcViewFamily.FrameCounter;
+	ViewFamily.ViewExtensions = GEngine->ViewExtensions->GatherActiveExtensions(FSceneViewExtensionContext(SrcViewFamily.Scene));
+
+	for (int32 ViewIndex = 0; ViewIndex < SrcViewFamily.Views.Num(); ++ViewIndex)
+	{
+		FSceneViewInitOptions ViewInitOptions = SrcViewFamily.Views[ViewIndex]->SceneViewInitOptions;
+		ViewInitOptions.ViewFamily = &ViewFamily;
+
+		// Reset to avoid incorrect culling problems
+		ViewInitOptions.SceneViewStateInterface = FSceneViewInitOptions{}.SceneViewStateInterface;
+		
+		FSceneView* View = new FSceneView(ViewInitOptions);
+		ViewFamily.Views.Emplace(View);
+	}
+}
+
+void RenderMeshEdges_RenderThread(
+	FRHICommandListImmediate& RHICmdList,
+	FSceneRenderer* SceneRenderer,
+	FMeshEdgesViewFamilyData& MeshEdgesData,
+	bool bUseSceneColorTexture)
+{
+	MeshEdgesData.WireframeColor->InitResource(RHICmdList);
+	MeshEdgesData.WireframeDepth->InitResource(RHICmdList);
+
+	SceneRenderer->RenderThreadBegin(RHICmdList);
+
+	FUniformExpressionCacheAsyncUpdateScope AsyncUpdateScope;
+
+	// update any resources that needed a deferred update
+	FDeferredUpdateResource::UpdateResources(RHICmdList);
+
+	FRDGBuilder GraphBuilder(RHICmdList, RDG_EVENT_NAME("MeshEdges"), ERDGBuilderFlags::AllowParallelExecute);
+	{
+		RDG_EVENT_SCOPE(GraphBuilder, "RenderMeshEdges_RenderThread");
+
+		// Render the scene normally
+		{
+			RDG_RHI_EVENT_SCOPE(GraphBuilder, RenderScene);
+			SceneRenderer->Render(GraphBuilder);
+		}
+	}
+	GraphBuilder.Execute();
+	
+	for (const FViewInfo& ViewInfo : SceneRenderer->Views)
+	{
+		MeshEdgesData.ViewRects.Emplace(ViewInfo.ViewRect);
+	}
+
+	SceneRenderer->RenderThreadEnd(RHICmdList);
+}
+
+void RenderMeshEdges(FSceneViewFamily& ViewFamily)
+{
+	if (!ViewFamily.EngineShowFlags.MeshEdges || ViewFamily.EngineShowFlags.HitProxies)
+	{
+		return;
+	}
+
+	FMeshEdgesViewFamilyData* ViewFamilyData = ViewFamily.GetOrCreateExtentionData<FMeshEdgesViewFamilyData>();
+
+	ERHIFeatureLevel::Type FeatureLevel = ViewFamily.GetFeatureLevel();
+	FIntPoint DesiredBufferSize = FSceneRenderer::GetDesiredInternalBufferSize(ViewFamily);
+	ViewFamilyData->CreateRenderTargets(FeatureLevel, DesiredBufferSize);
+
+	FEngineShowFlags ShowFlags = ViewFamily.EngineShowFlags;
+	{
+		// Render a wireframe view
+		ShowFlags.SetWireframe(true);
+
+		// Copy the MSAA wireframe view only, don't copy other scene elements
+		ShowFlags.SetSceneCaptureCopySceneDepth(false);
+
+		// Disable rendering of elements that are not needed
+		ShowFlags.SetMeshEdges(false);
+		ShowFlags.SetLighting(false);
+		ShowFlags.SetLightFunctions(false);
+		ShowFlags.SetGlobalIllumination(false);
+		ShowFlags.SetLumenGlobalIllumination(false);
+		ShowFlags.SetLumenReflections(false);
+		ShowFlags.SetDynamicShadows(false);
+		ShowFlags.SetCapsuleShadows(false);
+		ShowFlags.SetDistanceFieldAO(false);
+		ShowFlags.SetFog(false);
+		ShowFlags.SetVolumetricFog(false);
+		ShowFlags.SetCloud(false);
+		ShowFlags.SetDecals(false);
+		ShowFlags.SetAtmosphere(false);
+		ShowFlags.SetPostProcessing(false);
+		ShowFlags.SetCompositeEditorPrimitives(false);
+		ShowFlags.SetGrid(false);
+		//ShowFlags.SetScreenPercentage(false);
+		//ShowFlags.SetTranslucency(false);
+	}
+	
+	FSceneViewFamilyContext CaptureViewFamily(
+		FSceneViewFamily::ConstructionValues(
+			ViewFamilyData->WireframeColor.Get(),
+			ViewFamily.Scene,
+			ShowFlags)
+		.SetRenderTargetDepth(ViewFamilyData->WireframeDepth.Get())
+		.SetResolveScene(true)
+		.SetRealtimeUpdate(true)
+		.SetTime(ViewFamily.Time)
+	);
+
+	{
+		CopyViewFamily(ViewFamily, CaptureViewFamily);
+
+		CaptureViewFamily.SceneCaptureSource = SCS_SceneColorSceneDepth;
+	
+		// Use the same resolution scale as main view, so the buffers align pixel-perfect.
+		// If the main view is low-res this affects the wireframe quality, so the main view should be 100% ideally
+		CaptureViewFamily.SetScreenPercentageInterface(ViewFamily.GetScreenPercentageInterface()->Fork_GameThread(CaptureViewFamily));
+	}
+
+	FSceneRenderer* SceneRenderer = FSceneRenderer::CreateSceneRenderer(&CaptureViewFamily, nullptr);
+
+	for (const FSceneViewExtensionRef& Extension : CaptureViewFamily.ViewExtensions)
+	{
+		Extension->SetupViewFamily(CaptureViewFamily);
+	}
+
+	for (int32 ViewIndex = 0; ViewIndex < SceneRenderer->Views.Num(); ++ViewIndex)
+	{
+		FViewInfo& ViewInfo = SceneRenderer->Views[ViewIndex];
+		ViewInfo.bAllowTemporalJitter = false;
+		ViewInfo.PrimaryScreenPercentageMethod = EPrimaryScreenPercentageMethod::RawOutput;
+
+		for (const FSceneViewExtensionRef& Extension : CaptureViewFamily.ViewExtensions)
+		{
+			Extension->SetupView(CaptureViewFamily, ViewInfo);
+		}
+	}
+
+	ENQUEUE_RENDER_COMMAND(CaptureCommand)(
+		[SceneRenderer, ViewFamilyData](FRHICommandListImmediate& RHICmdList)
+	{
+		RenderMeshEdges_RenderThread(RHICmdList, SceneRenderer, *ViewFamilyData, true);
+	});
+}
+
+void ComposeMeshEdges(FRDGBuilder& GraphBuilder,
+                        const FViewInfo& View,
+                        FScreenPassRenderTarget& EditorPrimitivesColor,
+                        FScreenPassRenderTarget& EditorPrimitivesDepth)
+{
+	const FSceneViewFamily& ViewFamily = *View.Family;
+	if (!ViewFamily.EngineShowFlags.MeshEdges)
+	{
+		return;
+	}
+
+	int ViewIndex = 0;
+	for (; ViewIndex < ViewFamily.Views.Num(); ViewIndex++)
+	{
+		if (ViewFamily.Views[ViewIndex] == &View) break;
+	}
+
+	const FMeshEdgesViewFamilyData* ViewFamilyData = ViewFamily.GetExtentionData<FMeshEdgesViewFamilyData>();
+	check(ViewFamilyData); // should have been created in RenderMeshEdges
+	const FMeshEdgesViewSettings& ViewSettings = GetMeshEdgesViewSettings(ViewFamily, ViewIndex);
+
+	FRenderTargetTexture& WireframeTextureColor = *ViewFamilyData->WireframeColor;
+	FRenderTargetTexture& WireframeTextureDepth = *ViewFamilyData->WireframeDepth;
+	const FIntRect& WireframeViewRect = ViewFamilyData->ViewRects[ViewIndex];
+	const FSceneTextures& SceneTextures = View.GetSceneTextures();
+	FScreenPassTexture SceneDepth(SceneTextures.Depth.Resolve, View.ViewRect);
+	const uint32 NumMSAASamples = SceneTextures.Config.EditorPrimitiveNumSamples;
+	FRHISamplerState* PointClampSampler = TStaticSamplerState<SF_Point, AM_Clamp, AM_Clamp, AM_Clamp>::GetRHI();
+
+	FComposeMeshEdgesPS::FParameters* PassParameters = GraphBuilder.AllocParameters<FComposeMeshEdgesPS::FParameters>();
+	
+	PassParameters->WireframeColorTexture = RegisterExternalTexture(GraphBuilder, WireframeTextureColor.TextureRHI, *WireframeTextureColor.GetFriendlyName());
+	PassParameters->WireframeDepthTexture = RegisterExternalTexture(GraphBuilder, WireframeTextureDepth.TextureRHI, *WireframeTextureDepth.GetFriendlyName());
+	PassParameters->Wireframe = GetScreenPassTextureViewportParameters(FScreenPassTextureViewport(WireframeViewRect));
+	PassParameters->Depth = GetScreenPassTextureViewportParameters(FScreenPassTextureViewport(SceneDepth));
+	PassParameters->Output = GetScreenPassTextureViewportParameters(FScreenPassTextureViewport(EditorPrimitivesColor));
+	PassParameters->DepthTexture = SceneDepth.Texture;
+	PassParameters->DepthSampler = PointClampSampler;
+	PassParameters->DepthTextureJitter = FVector2f(View.TemporalJitterPixels);
+	PassParameters->Opacity = ViewSettings.Opacity;
+
+	for (int32 i = 0; i < int32(NumMSAASamples); i++)
+	{
+		PassParameters->SampleOffsetArray[i].X = GetMSAASampleOffsets(NumMSAASamples, i).X;
+		PassParameters->SampleOffsetArray[i].Y = GetMSAASampleOffsets(NumMSAASamples, i).Y;
+	}
+
+	PassParameters->RenderTargets[0] = EditorPrimitivesColor.GetRenderTargetBinding();
+	PassParameters->RenderTargets.DepthStencil = FDepthStencilBinding(EditorPrimitivesDepth.Texture, EditorPrimitivesDepth.LoadAction, EditorPrimitivesDepth.LoadAction, FExclusiveDepthStencil::DepthWrite);
+
+	const int MSAASampleCountDim = FMath::FloorLog2(NumMSAASamples);
+
+	FComposeMeshEdgesPS::FPermutationDomain PermutationVector;
+	PermutationVector.Set<FComposeMeshEdgesPS::FSampleCountDimension>(MSAASampleCountDim);
+
+	const FGlobalShaderMap* GlobalShaderMap = View.ShaderMap;
+	const TShaderRef<FComposeMeshEdgesPS>& PixelShader = TShaderMapRef<FComposeMeshEdgesPS>(View.ShaderMap, PermutationVector);
+
+	FIntRect Viewport = EditorPrimitivesColor.ViewRect;
+
+	FRHIBlendState* BlendState = TStaticBlendState<CW_RGBA>::GetRHI();
+	FRHIDepthStencilState* DepthStencilState = TStaticDepthStencilState<true, CF_DepthNearOrEqual>::GetRHI();
+
+	FPixelShaderUtils::AddFullscreenPass(GraphBuilder, View.ShaderMap, RDG_EVENT_NAME("ComposeMeshEdges"), PixelShader, PassParameters, Viewport, BlendState, nullptr, DepthStencilState);
+	
+	if (View.IsLastInFamily())
+	{
+		GraphBuilder.AddPass(
+			Forward<FRDGEventName>(RDG_EVENT_NAME("Cleanup MeshEdges")),
+			ERDGPassFlags::NeverCull,
+			[&WireframeTextureColor, &WireframeTextureDepth](FRHICommandList& RHICmdList)
+			{
+				WireframeTextureColor.ReleaseResource();
+				WireframeTextureDepth.ReleaseResource();
+			});
+	}
+}
+
diff --git a/Engine/Source/Runtime/Renderer/Private/MeshEdgesRendering.h b/Engine/Source/Runtime/Renderer/Private/MeshEdgesRendering.h
new file mode 100644
--- /dev/null
+++ b/Engine/Source/Runtime/Renderer/Private/MeshEdgesRendering.h
@@ -0,0 +1,20 @@
+// Copyright Epic Games, Inc. All Rights Reserved.
+
+class FRDGBuilder;
+class FSceneView;
+class FViewInfo;
+struct FCompositePrimitiveInputs;
+struct FScreenPassRenderTarget;
+
+struct FMeshEdgesViewSettings
+{
+	float Opacity = 1.0;
+};
+
+const FMeshEdgesViewSettings& GetMeshEdgesViewSettings(const FSceneView& View);
+FMeshEdgesViewSettings& GetMeshEdgesViewSettings(FSceneView& View);
+
+void ComposeMeshEdges(FRDGBuilder& GraphBuilder,
+						const FViewInfo& View,
+						FScreenPassRenderTarget& EditorPrimitivesColor,
+						FScreenPassRenderTarget& EditorPrimitivesDepth);
diff --git a/Engine/Source/Runtime/Renderer/Private/PostProcess/PostProcessCompositeEditorPrimitives.cpp b/Engine/Source/Runtime/Renderer/Private/PostProcess/PostProcessCompositeEditorPrimitives.cpp
--- a/Engine/Source/Runtime/Renderer/Private/PostProcess/PostProcessCompositeEditorPrimitives.cpp
+++ b/Engine/Source/Runtime/Renderer/Private/PostProcess/PostProcessCompositeEditorPrimitives.cpp
@@ -9,6 +9,7 @@
 #include "MobileBasePassRendering.h"
 #include "PixelShaderUtils.h"
 #include "Substrate/Substrate.h"
+#include "MeshEdgesRendering.h"
 
 namespace
 {
@@ -279,7 +280,7 @@ FScreenPassTexture AddEditorPrimitivePass(
 
 	// The editor primitive composition pass is also used when rendering VMI_WIREFRAME in order to use MSAA.
 	// So we need to check whether the editor primitives are enabled inside this function.
-	if (View.Family->EngineShowFlags.CompositeEditorPrimitives)
+	if (View.Family->EngineShowFlags.CompositeEditorPrimitives || View.Family->EngineShowFlags.MeshEdges)
 	{
 		// Populate depth if a prior pass did not already do it.
 		if (!bProducedByPriorPass)
@@ -302,6 +303,10 @@ FScreenPassTexture AddEditorPrimitivePass(
 				SceneDepthJitter,
 				NumMSAASamples);
 		}	
+		
+		FScreenPassRenderTarget EditorPrimitiveColorRT(EditorPrimitiveColor, EditorPrimitivesViewport.Rect, ERenderTargetLoadAction::ELoad);
+		FScreenPassRenderTarget EditorPrimitiveDepthRT(EditorPrimitiveDepth, EditorPrimitivesViewport.Rect, ERenderTargetLoadAction::ELoad);
+		ComposeMeshEdges(GraphBuilder, View, EditorPrimitiveColorRT, EditorPrimitiveDepthRT);
 
 		// Draws the editors opaque primitives
 		{
@@ -402,7 +407,7 @@ FScreenPassTexture AddEditorPrimitivePass(
 		PassParameters->EditorPrimitivesDepth = EditorPrimitiveDepth;
 		PassParameters->EditorPrimitivesColor = EditorPrimitiveColor;
 		PassParameters->bOpaqueEditorGizmo = bOpaqueEditorGizmo;
-		PassParameters->bCompositeAnyNonNullDepth = bProducedByPriorPass;
+		PassParameters->bCompositeAnyNonNullDepth = bProducedByPriorPass && !View.Family->EngineShowFlags.MeshEdges;
 		PassParameters->DepthTextureJitter = SceneDepthJitter;
 
 		for (int32 i = 0; i < int32(NumMSAASamples); i++)
diff --git a/Engine/Source/Runtime/Renderer/Private/PostProcess/PostProcessing.cpp b/Engine/Source/Runtime/Renderer/Private/PostProcess/PostProcessing.cpp
--- a/Engine/Source/Runtime/Renderer/Private/PostProcess/PostProcessing.cpp
+++ b/Engine/Source/Runtime/Renderer/Private/PostProcess/PostProcessing.cpp
@@ -1501,8 +1501,10 @@ void AddPostProcessingPasses(
 	if (PassSequence.IsEnabled(EPass::EditorPrimitive))
 	{
 		FCompositePrimitiveInputs PassInputs;
-		PassSequence.AcceptOverrideIfLastPass(EPass::EditorPrimitive, PassInputs.OverrideOutput);
-		PassInputs.OverrideDepthOutput = ViewFamilyDepthOutput;
+		if (PassSequence.AcceptOverrideIfLastPass(EPass::EditorPrimitive, PassInputs.OverrideOutput))
+		{
+			PassInputs.OverrideDepthOutput = ViewFamilyDepthOutput;
+		}
 		PassInputs.SceneColor = SceneColor;
 		PassInputs.SceneDepth = SceneDepth;
 		PassInputs.BasePassType = FCompositePrimitiveInputs::EBasePassType::Deferred;
diff --git a/Engine/Source/Runtime/Renderer/Private/SceneRendering.cpp b/Engine/Source/Runtime/Renderer/Private/SceneRendering.cpp
--- a/Engine/Source/Runtime/Renderer/Private/SceneRendering.cpp
+++ b/Engine/Source/Runtime/Renderer/Private/SceneRendering.cpp
@@ -4286,6 +4286,11 @@ void FSceneRenderer::CreateSceneRenderers(TArrayView<const FSceneViewFamily*> In
 
 		OutSceneRenderers.Last()->bIsFirstSceneRenderer = (FamilyIndex == 0);
 		OutSceneRenderers.Last()->bIsLastSceneRenderer = (FamilyIndex == InViewFamilies.Num() - 1);
+
+		for (int ViewExt = 0; ViewExt < InViewFamily->ViewExtensions.Num(); ViewExt++)
+		{
+			InViewFamily->ViewExtensions[ViewExt]->PostCreateSceneRenderer(*InViewFamily, OutSceneRenderers.Last());
+		}
 	}
 
 #if RHI_RAYTRACING
@@ -4350,9 +4355,9 @@ bool FSceneRenderer::ShouldCompositeEditorPrimitives(const FViewInfo& View)
 		return false;
 	}
 
-	if (ShowFlags.Wireframe)
+	if (ShowFlags.Wireframe || ShowFlags.MeshEdges)
 	{
-		// We want wireframe view use MSAA if possible.
+		// Wireframe is drawn to EditorPrimitives buffer because it uses MSAA, and so it requires the composition step
 		return true;
 	}
 	else if (ShowFlags.CompositeEditorPrimitives)
