Batch primitive transform updates into a single render command.

diff --git a/Engine/Source/Runtime/Engine/Private/LevelTick.cpp b/Engine/Source/Runtime/Engine/Private/LevelTick.cpp
--- a/Engine/Source/Runtime/Engine/Private/LevelTick.cpp
+++ b/Engine/Source/Runtime/Engine/Private/LevelTick.cpp
@@ -1126,6 +1126,12 @@ void UWorld::SendAllEndOfFrameUpdates()
 		}
 	};
 
+	if (Scene)
+	{
+		// We do not know in advance how many transform updates will occur so need to resize array to total number of components to avoid synchronization later
+		Scene->StartUpdatePrimitiveTransform(LocalComponentsThatNeedEndOfFrameUpdate.Num() + ComponentsThatNeedEndOfFrameUpdate_OnGameThread.Num());
+	}
+
 	if (IsUsingParallelNotifyEvents)
 	{
 #if WITH_EDITOR
@@ -1143,6 +1149,11 @@ void UWorld::SendAllEndOfFrameUpdates()
 		ParallelFor(LocalComponentsThatNeedEndOfFrameUpdate.Num(), ParallelWork);
 	}
 	
+	if (Scene)
+	{
+		Scene->FinishUpdatePrimitiveTransform();
+	}
+
 	for (UMaterialParameterCollectionInstance* ParameterCollectionInstance : ParameterCollectionInstances)
 	{
 		if (ParameterCollectionInstance)
diff --git a/Engine/Source/Runtime/Engine/Public/SceneInterface.h b/Engine/Source/Runtime/Engine/Public/SceneInterface.h
--- a/Engine/Source/Runtime/Engine/Public/SceneInterface.h
+++ b/Engine/Source/Runtime/Engine/Public/SceneInterface.h
@@ -578,6 +578,9 @@ public:
 	}
 	virtual void UpdateSceneSettings(AWorldSettings* WorldSettings) {}
 
+	virtual void StartUpdatePrimitiveTransform(int32 NumPrimitives) = 0;
+	virtual void FinishUpdatePrimitiveTransform() = 0;
+
 	/**
 	* Gets the GPU Skin Cache system associated with the scene.
 	*/
diff --git a/Engine/Source/Runtime/Renderer/Private/RendererScene.cpp b/Engine/Source/Runtime/Renderer/Private/RendererScene.cpp
--- a/Engine/Source/Runtime/Renderer/Private/RendererScene.cpp
+++ b/Engine/Source/Runtime/Renderer/Private/RendererScene.cpp
@@ -2117,17 +2117,6 @@ void FScene::UpdatePrimitiveTransformInternal(T* Primitive)
 		{
 			FVector AttachmentRootPosition = Primitive->GetActorPositionForRenderer();
 
-			struct FPrimitiveUpdateParams
-			{
-				FScene* Scene;
-				FPrimitiveSceneProxy* PrimitiveSceneProxy;
-				FBoxSphereBounds WorldBounds;
-				FBoxSphereBounds LocalBounds;
-				FMatrix LocalToWorld;
-				TOptional<FTransform> PreviousTransform;
-				FVector AttachmentRootPosition;
-			};
-
 			FPrimitiveUpdateParams UpdateParams;
 			UpdateParams.Scene = this;
 			UpdateParams.PrimitiveSceneProxy = Primitive->GetSceneProxy();
@@ -2181,20 +2170,29 @@ void FScene::UpdatePrimitiveTransformInternal(T* Primitive)
 
 			if (bPerformUpdate)
 			{
-				ENQUEUE_RENDER_COMMAND(UpdateTransformCommand)(
-					[UpdateParams] (FRHICommandListBase&)
-					{
-						FScopeCycleCounter Context(UpdateParams.PrimitiveSceneProxy->GetStatId());
-						UpdateParams.Scene->UpdatePrimitiveTransform_RenderThread(
-							UpdateParams.PrimitiveSceneProxy,
-							UpdateParams.WorldBounds,
-							UpdateParams.LocalBounds,
-							UpdateParams.LocalToWorld,
-							UpdateParams.AttachmentRootPosition,
-							UpdateParams.PreviousTransform
-						);
-					}
-				);
+				// Accumulate all transform updates and enqueue them as once
+				if (bPrimitivesUpdateBatching)
+				{
+					int32 Index = PrimitiveUpdateIndex.fetch_add(1, std::memory_order_relaxed);
+					PrimitivesUpdates[Index] = MoveTemp(UpdateParams);
+				}
+				else
+				{
+					ENQUEUE_RENDER_COMMAND(UpdateTransformCommand)(
+						[UpdateParams](FRHICommandListBase&)
+						{
+							FScopeCycleCounter Context(UpdateParams.PrimitiveSceneProxy->GetStatId());
+							UpdateParams.Scene->UpdatePrimitiveTransform_RenderThread(
+								UpdateParams.PrimitiveSceneProxy,
+								UpdateParams.WorldBounds,
+								UpdateParams.LocalBounds,
+								UpdateParams.LocalToWorld,
+								UpdateParams.AttachmentRootPosition,
+								UpdateParams.PreviousTransform
+							);
+						}
+					);
+				}
 			}
 		}
 	}
@@ -2205,6 +2203,45 @@ void FScene::UpdatePrimitiveTransformInternal(T* Primitive)
 	}
 }
 
+void FScene::StartUpdatePrimitiveTransform(int32 NumPrimitives)
+{
+	if (NumPrimitives > 0)
+	{
+		PrimitivesUpdates.SetNum(NumPrimitives);
+		bPrimitivesUpdateBatching = true;
+	}
+}
+
+void FScene::FinishUpdatePrimitiveTransform()
+{
+	if (bPrimitivesUpdateBatching)
+	{
+		// Pass the collection and actual number of accumulated updates
+		ENQUEUE_RENDER_COMMAND(UpdateTransformCommand)(
+			[PrimitivesUpdates = MoveTemp(PrimitivesUpdates), NumUpdates = PrimitiveUpdateIndex.load(std::memory_order_relaxed)](FRHICommandListBase&)
+			{
+				for (int32 Index = 0; Index < NumUpdates; ++Index)
+				{
+					const auto& UpdateParams = PrimitivesUpdates[Index];
+					FScopeCycleCounter Context(UpdateParams.PrimitiveSceneProxy->GetStatId());
+					UpdateParams.Scene->UpdatePrimitiveTransform_RenderThread(
+						UpdateParams.PrimitiveSceneProxy,
+						UpdateParams.WorldBounds,
+						UpdateParams.LocalBounds,
+						UpdateParams.LocalToWorld,
+						UpdateParams.AttachmentRootPosition,
+						UpdateParams.PreviousTransform
+					);
+				}
+			}
+		);
+
+		PrimitivesUpdates.Reset();
+		PrimitiveUpdateIndex = 0;
+		bPrimitivesUpdateBatching = false;
+	}
+}
+
 void FScene::UpdatePrimitiveOcclusionBoundsSlack(UPrimitiveComponent* Primitive, float NewSlack)
 {
 	if (const FPrimitiveSceneProxy* SceneProxy = Primitive->GetSceneProxy())
@@ -7074,7 +7111,8 @@ public:
 	virtual void UpdateCustomPrimitiveData(FPrimitiveSceneDesc* Primitive, const FCustomPrimitiveData&) override {}
 	virtual void UpdatePrimitiveInstances(FInstancedStaticMeshSceneDesc* Primitive) override {};
 
-
+	virtual void StartUpdatePrimitiveTransform(int32 NumPrimitives) override {}
+	virtual void FinishUpdatePrimitiveTransform() override {}
 private:
 	UWorld* World;
 	class FFXSystemInterface* FXSystem;
diff --git a/Engine/Source/Runtime/Renderer/Private/ScenePrivate.h b/Engine/Source/Runtime/Renderer/Private/ScenePrivate.h
--- a/Engine/Source/Runtime/Renderer/Private/ScenePrivate.h
+++ b/Engine/Source/Runtime/Renderer/Private/ScenePrivate.h
@@ -2880,6 +2880,22 @@ public:
 	/** Nanite material visibility references. These are stored on the scene as they are computed at FPrimitiveSceneInfo::AddToScene time. */
 	FNaniteVisibility NaniteVisibility[ENaniteMeshPass::Num];
 
+	struct FPrimitiveUpdateParams
+	{
+		FScene* Scene;
+		FPrimitiveSceneProxy* PrimitiveSceneProxy;
+		FBoxSphereBounds WorldBounds;
+		FBoxSphereBounds LocalBounds;
+		FMatrix LocalToWorld;
+		TOptional<FTransform> PreviousTransform;
+		FVector AttachmentRootPosition;
+	};
+
+	/** DoDeferredRenderUpdates_Concurrent can be called not only inside SendAllEndOfFrameUpdates so we should be able to enqueue commands also in immediate mode. */
+	bool bPrimitivesUpdateBatching = false;
+	std::atomic_int32_t PrimitiveUpdateIndex = 0;
+	TArray<FPrimitiveUpdateParams> PrimitivesUpdates;
+
 	/**
 	 * The following arrays are densely packed primitive data needed by various
 	 * rendering passes. PrimitiveSceneInfo->PackedIndex maintains the index
@@ -3802,6 +3818,9 @@ private:
 
 	template<class T> 	
 	void UpdatePrimitiveTransformInternal(T* Primitive);
+
+	virtual void StartUpdatePrimitiveTransform(int32 NumPrimitives) override;
+	virtual void FinishUpdatePrimitiveTransform() override;
 	
 	void RemoveViewLumenSceneData_RenderThread(FSceneViewStateInterface* ViewState);
 	void RemoveViewState_RenderThread(FSceneViewStateInterface*);
