Updated light function processing and added more debug information when out of slots.
Fixed a bug wehre one slot was always marked as used, always removing 1 material form the atlas.

diff --git a/Engine/Source/Runtime/Renderer/Private/DeferredShadingRenderer.cpp b/Engine/Source/Runtime/Renderer/Private/DeferredShadingRenderer.cpp
--- a/Engine/Source/Runtime/Renderer/Private/DeferredShadingRenderer.cpp
+++ b/Engine/Source/Runtime/Renderer/Private/DeferredShadingRenderer.cpp
@@ -3110,6 +3110,8 @@ void FDeferredShadingSceneRenderer::Render(FRDGBuilder& GraphBuilder)
 			}
 		}
 
+		LightFunctionAtlas.RenderDebugInfo(GraphBuilder, Views);
+
 		if (VirtualShadowMapArray.IsEnabled())
 		{
 			VirtualShadowMapArray.RenderDebugInfo(GraphBuilder, Views);
diff --git a/Engine/Source/Runtime/Renderer/Private/LightFunctionAtlas.cpp b/Engine/Source/Runtime/Renderer/Private/LightFunctionAtlas.cpp
--- a/Engine/Source/Runtime/Renderer/Private/LightFunctionAtlas.cpp
+++ b/Engine/Source/Runtime/Renderer/Private/LightFunctionAtlas.cpp
@@ -59,14 +59,14 @@ static TAutoConsoleVariable<int32> CVarLightFunctionAtlas(
 static TAutoConsoleVariable<int32> CVarLightFunctionAtlasSlotResolution(
 	TEXT("r.LightFunctionAtlas.SlotResolution"),
 	128,
-	TEXT("Experimental: The resolution of each atlas slot."),
+	TEXT("Experimental: The resolution of each atlas slot. Maximum value is 256."),
 	ECVF_RenderThreadSafe);
 
 // We do not dynamically scale allocated slot resolution for now.
 static TAutoConsoleVariable<int32> CVarLightFunctionAtlasSize(
 	TEXT("r.LightFunctionAtlas.Size"),
 	4,
-	TEXT("Experimental: The default size in atlas slot count of the edge of the 2D texture atlas."),
+	TEXT("Experimental: The default size in atlas slot count of the edge of the 2D texture atlas. Maximum value is 16."),
 	ECVF_RenderThreadSafe);
 
 static TAutoConsoleVariable<int32> CVarLightFunctionAtlasMaxLightCount(
@@ -261,6 +261,14 @@ void FLightFunctionAtlas::BeginSceneFrame(const FViewFamilyInfo& ViewFamily, TAr
 
 	// Now lets check if we need to generate the atlas for this frame
 	bLightFunctionAtlasEnabled = CVarLightFunctionAtlas.GetValueOnRenderThread() > 0 && ViewFamily.EngineShowFlags.LightFunctions > 0;
+#if !UE_BUILD_SHIPPING
+	LightCountSkippedDueToMissingAtlasSlot = 0;
+	SkippedLightFunctionsSet.Empty();
+#endif
+#if WITH_EDITOR
+	LightCountWithLFMaterialsNotSamplingAtlas = 0;
+	NonCompatibleLightFunctionMaterials.Reset();
+#endif
 
 	// But only really enable the atlas generation if a system asks for it
 	bool bVolumetricFogRequestsLF = false;
@@ -315,6 +323,8 @@ void FLightFunctionAtlas::AllocateAtlasSlots(const TArray<FViewInfo>& Views)
 	{
 		return;
 	}
+	const FViewInfo& View = Views[0];
+	ERHIFeatureLevel::Type FeatureLevel = View.GetFeatureLevel();
 
 	//
 	// Sort the list of lights registered as having light function in order to keep directional lights first, then each lights closer to each views
@@ -376,7 +386,7 @@ void FLightFunctionAtlas::AllocateAtlasSlots(const TArray<FViewInfo>& Views)
 	check(AtlasEdgeSize * AtlasEdgeSize <= LIGHT_FUNCTION_ATLAS_MAX_LIGHT_FUNCTION_COUNT);
 	const uint32 AtlasMaxLightFunctionCount = AtlasEdgeSize * AtlasEdgeSize;
 
-	LightFunctionsSet.Reset();
+	EffectiveLightFunctionSlotArray.Reset();
 	EffectiveLightFunctionSlotArray.Reserve(AtlasMaxLightFunctionCount);
 
 	EffectiveLocalLightSlotArray.Reset();
@@ -385,9 +395,6 @@ void FLightFunctionAtlas::AllocateAtlasSlots(const TArray<FViewInfo>& Views)
 	LightFunctionsSet.Reset();
 	LightFunctionsSet.Reserve(AtlasMaxLightFunctionCount);
 
-	// Reserve the default slot at index 0 as the default identity slot
-	LightFunctionsSet.Add(FLightFunctionSlotKey());
-
 	uint32 NextAtlasSlotX = 0;
 	uint32 NextAtlasSlotY = 0;
 	auto AddAtlasSlot = [&](FLightFunctionSlotKey Key, const FMaterialRenderProxy* LightFunctionMaterial)
@@ -436,22 +443,52 @@ void FLightFunctionAtlas::AllocateAtlasSlots(const TArray<FViewInfo>& Views)
 	{
 		FLightSceneInfo* LightSceneInfo = RegisteredLights[SortedRegisteredLight.RegisteredLightIndex];
 		FLightSceneProxy* Proxy = LightSceneInfo->Proxy;
+		const FMaterialRenderProxy* LightFunctionMaterial = LightSceneInfo->Proxy->GetLightFunctionMaterial();
 
-		if (MaxLightCount >= 0 && int32(LocalLightWithLightFunctionCount) >= MaxLightCount)
+		if ((MaxLightCount >= 0 && int32(LocalLightWithLightFunctionCount) >= MaxLightCount) || LightFunctionMaterial == nullptr)
 		{
 			// We cannot register anymore light, so set them to no light function
 			Proxy->SetLightFunctionAtlasIndices(0);
 			continue;
 		}
 
+		if (LightFunctionMaterial != nullptr)
+		{
+			FMaterialRelevance MaterialRelevance;
+			const UMaterialInterface* LightFunctionMaterialInterface = LightFunctionMaterial->GetMaterialInterface();
+			if (LightFunctionMaterialInterface)
+			{
+				MaterialRelevance = LightFunctionMaterialInterface->GetRelevance(FeatureLevel);
+			}
+			if (!MaterialRelevance.bIsLightFunctionAtlasCompatible)
+			{
+			#if WITH_EDITOR
+				if (LightFunctionMaterialInterface)
+				{
+					NonCompatibleLightFunctionMaterials.Add(LightFunctionMaterialInterface->GetUniqueID(), LightFunctionMaterialInterface);
+				}
+				LightCountWithLFMaterialsNotSamplingAtlas++;
+			#endif
+				Proxy->SetLightFunctionAtlasIndices(0);
+				continue;
+			}
+		}
+
 		FLightFunctionSlotKey NewKey = FLightFunctionSlotKey(LightSceneInfo);
 
 		FLightFunctionSlotKey* ExistingKey = LightFunctionsSet.Find(NewKey);
 
-		if (ExistingKey == nullptr && (uint32(LightFunctionsSet.Num()) < AtlasMaxLightFunctionCount))
+	#if !UE_BUILD_SHIPPING
+		if (ExistingKey == nullptr && (uint32(LightFunctionsSet.Num()) >= AtlasMaxLightFunctionCount))
 		{
-			const FMaterialRenderProxy* LightFunctionMaterial = Proxy->GetLightFunctionMaterial();
+			// We cannot account for this light's required light function
+			LightCountSkippedDueToMissingAtlasSlot++;
+			SkippedLightFunctionsSet.Add(NewKey);
+		}
+	#endif
 
+		if (ExistingKey == nullptr && (uint32(LightFunctionsSet.Num()) < AtlasMaxLightFunctionCount))
+		{
 			// Allocate slots for the light and views
 			Proxy->SetLightFunctionAtlasIndices(AddLightSlot(LightSceneInfo, AddAtlasSlot(NewKey, LightFunctionMaterial)));
 
@@ -556,22 +593,24 @@ void FLightFunctionAtlas::RenderLightFunctionAtlas(FRDGBuilder& GraphBuilder, TA
 		const FIntPoint LightFunctionResolution = FIntPoint(AtlasSlotResolution, AtlasSlotResolution);
 
 		// Write the light data needed to rotate and fade the light function in the world.
-		const uint32 InitialLightInfoDataLightCount = FMath::DivideAndRoundUp(uint32(EffectiveLocalLightSlotArray.Num()), 32u) * 32u;// Alloted with 32 lights step to better reuse shared buffers pool.
+		// UVMinMax for pointed to atlas slot is also packed into the light structure.
+		const uint32 InitialLightInfoDataLightCount = FMath::DivideAndRoundUp(uint32(EffectiveLocalLightSlotArray.Num()), 32u) * 32u;	// Allocated with 32 lights step to better reuse shared buffers pool
 		const uint32 InitialLightInfoDataSize = InitialLightInfoDataLightCount * sizeof(FAtlasLightInfoData);
 		FAtlasLightInfoData* LightInfoDataBufferPtr = (FAtlasLightInfoData*)GraphBuilder.Alloc(InitialLightInfoDataSize, 16);
-		uint32 LightIndex = 0;
+		uint32 OutputBufferLightIndex = 0;
 		for (EffectiveLocalLightSlot& LightSlot : EffectiveLocalLightSlotArray)
 		{
 			FLightSceneInfo* LightSceneInfo = LightSlot.LightSceneInfo;
 			const bool bIsDefaultSlot = LightSceneInfo == nullptr;
-
 			if (bIsDefaultSlot)
 			{
-				LightInfoDataBufferPtr[LightIndex].Parameters = FVector4f(1.0f, 1.0f, 1.0f, 0.0f);
-				LightIndex++;
+				LightInfoDataBufferPtr[OutputBufferLightIndex].Parameters = FVector4f(1.0f, 1.0f, 1.0f, 0.0f);
+				OutputBufferLightIndex++;
 				continue;
 			}
 
+			const uint8 LightFunctionAtlasSlotIndex = LightSlot.LightFunctionAtlasSlotIndex;
+			check(LightSceneInfo);
 			const FLightSceneProxy* Proxy = LightSceneInfo->Proxy;
 
 			float ShadowFadeFraction = 1.0f;
@@ -586,11 +625,10 @@ void FLightFunctionAtlas::RenderLightFunctionAtlas(FRDGBuilder& GraphBuilder, TA
 				TranslatedWorldToLight = FMatrix44f(FTranslationMatrix(-View.ViewMatrices.GetPreViewTranslation()) * WorldToLight);
 			}
 
-			LightInfoDataBufferPtr[LightIndex].Transform = TranslatedWorldToLight;
+			LightInfoDataBufferPtr[OutputBufferLightIndex].Transform = TranslatedWorldToLight;
 
 			uint8 LightType = Proxy->GetLightType();
 			
-			const uint8 LightFunctionAtlasSlotIndex = LightSlot.LightFunctionAtlasSlotIndex;
 			FFloat16 PackedDisabledBrightness(Proxy->GetLightFunctionDisabledBrightness());
 			uint32 PackedLightInfoDataParams = uint32(LightType) | (uint32(PackedDisabledBrightness.Encoded) << 8);
 
@@ -609,9 +647,9 @@ void FLightFunctionAtlas::RenderLightFunctionAtlas(FRDGBuilder& GraphBuilder, TA
 			const float TanOuterAngle = LightType == LightType_Spot ? FMath::Tan(LightSceneInfo->Proxy->GetOuterConeAngle()) : 1.0f;
 
 			// ShadowFadeFraction is unused.
-			LightInfoDataBufferPtr[LightIndex].Parameters = FVector4f(Proxy->GetLightFunctionFadeDistance(), FMath::AsFloat(PackedLightInfoDataParams), FMath::AsFloat(PackedAtlasSlotMinUV), TanOuterAngle);
+			LightInfoDataBufferPtr[OutputBufferLightIndex].Parameters = FVector4f(Proxy->GetLightFunctionFadeDistance(), FMath::AsFloat(PackedLightInfoDataParams), FMath::AsFloat(PackedAtlasSlotMinUV), TanOuterAngle);
 
-			LightIndex++;
+			OutputBufferLightIndex++;
 		}
 
 		// Create the light instance data buffer SRV
@@ -670,16 +708,16 @@ void FLightFunctionAtlas::RenderAtlasSlots(FRDGBuilder& GraphBuilder, const TArr
 			uint32 AtlasEdgeSize = GetAtlasEdgeSize();
 			uint32 AtlasResolution = AtlasSlotResolution * AtlasEdgeSize;
 
+			// This always work because in this case we do not need anything from any view.
+			const FViewInfo& View = Views[0];
+
 			// Render all light functions and update light info
 			for (EffectiveLightFunctionSlot& Slot : EffectiveLightFunctionSlotArray)
 			{
-				// This always work because in this case we do not need anything from any view.
-				const FViewInfo& View = Views[0];
-
 				const FMaterialRenderProxy* MaterialProxyForRendering = Slot.LightFunctionMaterial;
 				if (MaterialProxyForRendering == nullptr)
 				{
-					continue;	// This is the default invalid slot at index 0 to notify disabled light function on local light data.
+					continue;	// This is to skip the unused slot at index 0 because index 0 disable the sampling from the light function atlas.
 				}
 				const FMaterial& Material = MaterialProxyForRendering->GetMaterialWithFallback(View.GetFeatureLevel(), MaterialProxyForRendering);
 
@@ -729,6 +767,53 @@ void FLightFunctionAtlas::RenderAtlasSlots(FRDGBuilder& GraphBuilder, const TArr
 	);
 }
 
+void FLightFunctionAtlas::RenderDebugInfo(FRDGBuilder& GraphBuilder, TArray<FViewInfo>& Views)
+{
+#if !UE_BUILD_SHIPPING
+	if (!IsLightFunctionAtlasEnabled())
+	{
+		return;
+	}
+
+	// In case we became out of budget, let's notify the game developers.
+	if (LightCountSkippedDueToMissingAtlasSlot > 0)
+	{
+		for (int32 ViewIndex = 0; ViewIndex < Views.Num(); ViewIndex++)
+		{
+			const FViewInfo& View = Views[ViewIndex];
+			const FMinimalSceneTextures& SceneTextures = View.GetSceneTextures();
+
+			AddDrawCanvasPass(GraphBuilder, {}, View, FScreenPassRenderTarget(SceneTextures.Color.Target, View.ViewRect, ERenderTargetLoadAction::ELoad),
+				[&View, this](FCanvas& Canvas)
+				{
+					FString Text;
+					const float DPIScale = Canvas.GetDPIScale();
+					Canvas.SetBaseTransform(FMatrix(FScaleMatrix(DPIScale) * Canvas.CalcBaseTransform2D(Canvas.GetViewRect().Width(), Canvas.GetViewRect().Height())));
+
+					const float ViewPortWidth = float(View.ViewRect.Width());
+					const float ViewPortHeight = float(View.ViewRect.Height());
+					FLinearColor TextColor(1.0f, 0.5f, 0.0f);
+
+					float DrawPosX = 40.0f;
+					float DrawPosY = 650.0f;
+
+					Text = FString::Printf(TEXT("Light Functions Atlas:"));
+					Canvas.DrawShadowedString(DrawPosX, DrawPosY, *Text, GetStatsFont(), TextColor);
+					DrawPosY += 20.0f;
+
+					Text = FString::Printf(TEXT("  - %d light(s) will skip light functions due to out of atlas slot (see r.LightFunctionAtlas.Size)."), LightCountSkippedDueToMissingAtlasSlot);
+					Canvas.DrawShadowedString(DrawPosX, DrawPosY, *Text, GetStatsFont(), TextColor);
+					DrawPosY += 20.0f;
+
+					Text = FString::Printf(TEXT("  - %d light function material(s) have been skipped."), SkippedLightFunctionsSet.Num());
+					Canvas.DrawShadowedString(DrawPosX, DrawPosY, *Text, GetStatsFont(), TextColor);
+					DrawPosY += 20.0f;
+				});
+		}
+	}
+#endif // !UE_BUILD_SHIPPING
+}
+
 FScreenPassTexture FLightFunctionAtlas::AddDebugVisualizationPasses(FRDGBuilder& GraphBuilder, const FViewInfo& View, FScreenPassTexture& ScreenPassSceneColor) const
 {
 #if WITH_EDITOR
@@ -781,22 +866,17 @@ FScreenPassTexture FLightFunctionAtlas::AddDebugVisualizationPasses(FRDGBuilder&
 
 		Text = FString::Printf(TEXT("Local Lights sampling atlas: %d"), EffectiveLocalLightSlotArray.Num());
 		Canvas.DrawShadowedString(DrawPosX, DrawPosY, *Text, GEngine->GetLargeFont(), FLinearColor::White);
-		DrawPosY += 20.0f;
-
-		Canvas.DrawShadowedString(DrawPosX, DrawPosY, TEXT("Note: (X) means the material cannot be sampled from the atlas for the deferred lighting passes (due to using world pos, depth or manipulating UVs)"), GEngine->GetLargeFont(), FLinearColor::Yellow);
 		DrawPosY += 40.0f;
 
 		uint32 LightFunctionAtlasSlotIndex = 0;
 		for (auto& AtlasSlot : EffectiveLightFunctionSlotArray)
 		{
 			const FMaterialRenderProxy* LightFunctionMaterial = AtlasSlot.LightFunctionMaterial;
+			check(LightFunctionMaterial);
 
 			const FString& MaterialName = LightFunctionMaterial->GetMaterialName();
 			const uint32 LFMaterialUniqueID = LightFunctionMaterial->GetMaterialInterface() ? MurmurFinalize32(MurmurFinalize32(LightFunctionMaterial->GetMaterialInterface()->GetUniqueID())) : 0xFFFFFFFF;
 
-			const FMaterial& LFMaterial = LightFunctionMaterial->GetIncompleteMaterialWithFallback(View.GetFeatureLevel());
-			const bool bIsMaterialCompatibleWithAtlas = LFMaterial.MaterialIsLightFunctionAtlasCompatible_RenderThread();
-
 			const FLinearColor MaterialColor(FColor(LFMaterialUniqueID & 0xFF, (LFMaterialUniqueID >> 8) & 0xFF, (LFMaterialUniqueID >> 16) & 0xFF));
 
 			uint32 LightCountUsingThisMaterial = 0;
@@ -810,7 +890,7 @@ FScreenPassTexture FLightFunctionAtlas::AddDebugVisualizationPasses(FRDGBuilder&
 
 			Text = FString::Printf(TEXT("%2d"), LightCountUsingThisMaterial);
 			Canvas.DrawShadowedString(DrawPosX, DrawPosY, *Text, GEngine->GetLargeFont(), MaterialColor);
-			Text = FString::Printf(TEXT("lights - %s - %s"), bIsMaterialCompatibleWithAtlas ? TEXT("OK ") : TEXT("(X)"), *MaterialName);
+			Text = FString::Printf(TEXT("lights - %s"), *MaterialName);
 			Canvas.DrawShadowedString(DrawPosX + 20, DrawPosY, *Text, GEngine->GetLargeFont(), MaterialColor);
 
 			// Draw a line around the corresponding atlas tile
@@ -831,6 +911,29 @@ FScreenPassTexture FLightFunctionAtlas::AddDebugVisualizationPasses(FRDGBuilder&
 			DrawPosY += 20.0f;
 			LightFunctionAtlasSlotIndex++;
 		}
+
+		DrawPosY += 50.0f;
+
+		// Now display incompatible materials
+		Text = FString::Printf(TEXT("Light functions not compatible with the atlas:   %d"), NonCompatibleLightFunctionMaterials.Num());
+		Canvas.DrawShadowedString(DrawPosX, DrawPosY, *Text, GEngine->GetLargeFont(), FLinearColor::White);
+		DrawPosY += 20.0f;
+
+		Text = FString::Printf(TEXT("Local Lights with light function not in the atlas: %d"), LightCountWithLFMaterialsNotSamplingAtlas);
+		Canvas.DrawShadowedString(DrawPosX, DrawPosY, *Text, GEngine->GetLargeFont(), FLinearColor::White);
+		DrawPosY += 40.0f;
+
+		for (auto& NotCompatibleMaterial : NonCompatibleLightFunctionMaterials)
+		{
+			const uint32 LFMaterialUniqueID = MurmurFinalize32(MurmurFinalize32(NotCompatibleMaterial.Value->GetUniqueID()));
+			const FLinearColor MaterialColor(FColor(LFMaterialUniqueID & 0xFF, (LFMaterialUniqueID >> 8) & 0xFF, (LFMaterialUniqueID >> 16) & 0xFF));
+
+			Canvas.DrawTile(DrawPosX - 20.0f, DrawPosY, 15.0f, 15.0f, 0.0f, 0.0f, 1.0f, 1.0f, MaterialColor, nullptr, false);
+
+			Text = FString::Printf(TEXT("lights - %s"), *NotCompatibleMaterial.Value->GetFullName());
+			Canvas.DrawShadowedString(DrawPosX + 20, DrawPosY, *Text, GEngine->GetLargeFont(), MaterialColor);
+			DrawPosY += 20.0f;
+		}
 	});
 
 #endif // WITH_EDITOR
diff --git a/Engine/Source/Runtime/Renderer/Private/LightFunctionAtlas.h b/Engine/Source/Runtime/Renderer/Private/LightFunctionAtlas.h
--- a/Engine/Source/Runtime/Renderer/Private/LightFunctionAtlas.h
+++ b/Engine/Source/Runtime/Renderer/Private/LightFunctionAtlas.h
@@ -143,6 +143,8 @@ struct FLightFunctionAtlas
 
 	void RenderLightFunctionAtlas(FRDGBuilder& GraphBuilder, TArray<FViewInfo>& Views);
 
+	void RenderDebugInfo(FRDGBuilder& GraphBuilder, TArray<FViewInfo>& Views);
+
 
 	FScreenPassTexture AddDebugVisualizationPasses(FRDGBuilder& GraphBuilder, const FViewInfo& View, FScreenPassTexture& ScreenPassSceneColor)  const;
 
@@ -187,13 +189,23 @@ private:
 		FLightSceneInfo*	LightSceneInfo = nullptr;
 		uint8				LightFunctionAtlasSlotIndex = 0;
 	};
-	TArray<EffectiveLocalLightSlot> EffectiveLocalLightSlotArray;
+	TArray<EffectiveLocalLightSlot>										EffectiveLocalLightSlotArray;
 
 	FLightFunctionAtlasGlobalParameters*								DefaultLightFunctionAtlasGlobalParameters = nullptr;
 	TRDGUniformBufferRef<FLightFunctionAtlasGlobalParameters>			DefaultLightFunctionAtlasGlobalParametersUB;
 
 	TArray<FLightFunctionAtlasGlobalParameters*>						ViewLightFunctionAtlasGlobalParametersArray;
 	TArray<TRDGUniformBufferRef<FLightFunctionAtlasGlobalParameters>>	ViewLightFunctionAtlasGlobalParametersUBArray;
+
+#if WITH_EDITOR
+	uint32																LightCountWithLFMaterialsNotSamplingAtlas = 0;
+	TMap<uint32, const UMaterialInterface*>								NonCompatibleLightFunctionMaterials;
+#endif
+
+#if !UE_BUILD_SHIPPING
+	uint32																LightCountSkippedDueToMissingAtlasSlot = 0;
+	TSet<FLightFunctionSlotKey>											SkippedLightFunctionsSet;
+#endif
 };
 
 
