Path Tracer: Initial volumetric Cloud support

Support VolumetricCloud actors when using the reference atmosphere mode in the path tracer. Two code paths are supported. The first, reference mode invokes the cloud material via callable shaders. Unfortunately there is a fairly substantial overhead to invoking callable shaders from such a deeply nested inner loop. The other option (enabled by default) is to get cloud properties through a specially constructed "cloud map" which is a 3d texture with angular density so that the resolution is highest near the viewer. The default size of this map is 512x512x64 which is 256Mb (uint4 format). All the main physical properties of the cloud (extinction,albedo,dual-lobed phase function) are packed into this uint4 format and can be driven by the shader.

To optimize ray marching, a secondary 2d structure is built which captures both min/max density and z values per cell. This structure is axis-aligned and traversed with a 2D DDA to efficiently skip empty regions and restrict cloud ray marching to areas which are known to contain clouds. The memory usage for this structure at the default resolution of 512x512 (float4 format) is just 4Mb. This structure is incrementally updated during rendering, so it progressively fills in more accurately as more samples are taken. While there is a slight bias (and potentially missing details) at very low sample counts, the renderer is guaranteed to eventually converge when more samples are taken.

Because this feature is still experimental, it is disabled by default, but can be enabled at runtime with the cvar:

r.PathTracing.EnableReferenceClouds 1

This only works when reference atmosphere is also enabled.

Quality can be tuned with the following cvars:

r.PathTracing.CloudRoughnessCutoff 1 (make the clouds visible for rays of all roughness values, this makes the clouds visible to self-scattering which means true-multiple scattering within the cloud is computed. This can be very expensive for large cloudscapes.)
r.PathTracing.CloudMapResolution (tune the voxel x/y resolution of the clouds - default 512)
r.PathTracing.CloudMapDepth (tune the voxel z resolution of the clouds - default 64)
r.PathTracing.CloudMapEnable 0 (disables the voxelized representation of the clouds and run the full shader)

Also, the "Tracing Start Max Distance" is the main parameter which controls how "far" the clouds can be seen. This directly impacts the effective resolution, so keeping it as low as possible is essential.

High density clouds can show some missing voxels if the maximum number of ray march steps is exceeded (the default has been raised to 768 to accomodate the default volumetric cloud shader)

diff --git a/Engine/Shaders/Private/PathTracing/Material/PathTracingMedium.ush b/Engine/Shaders/Private/PathTracing/Material/PathTracingMedium.ush
--- a/Engine/Shaders/Private/PathTracing/Material/PathTracingMedium.ush
+++ b/Engine/Shaders/Private/PathTracing/Material/PathTracingMedium.ush
@@ -2,8 +2,10 @@
 
 /*=============================================================================================
 PathTracingMedium.usf: Volumetric phase function
-   To support atmosphere, we represent a blend of two components: rayleigh scattering and an HG lobe
-   In the future, we may add additional HG lobe for cases like cloud rendering
+ * To support atmosphere, we represent a blend of two components: rayleigh scattering and an HG lobe.
+ * To support volumetric clouds we also include a dual HG lobe that we track seperately to simplify handling cases where multiple volumes overlap.
+ * See FVolumeShadedResult for the logic around how this struct is tracked during volume tracing.
+   
 ===============================================================================================*/
 #pragma once
 
@@ -17,27 +19,59 @@ FMaterialSample Medium_SampleMaterial(
 	float3 RandSample)
 {
 	const float3 RayleighWeight = Payload.GetBaseColor();
-	const float3 HGWeight = Payload.GetHGWeight();
-	const float G = Payload.GetHGPhaseG();
+	const float3 HG0 = Payload.GetHGWeight();
+	const float G0 = Payload.GetHGPhaseG();
+	const float3 DualHGPhaseData = Payload.GetDualHGPhaseData();
+	const float3 HG1 = Payload.GetDualHGWeight() - Payload.GetDualHGWeight() * DualHGPhaseData.z;
+	const float3 HG2 = Payload.GetDualHGWeight() * DualHGPhaseData.z;
+	const float G1 = DualHGPhaseData.x;
+	const float G2 = DualHGPhaseData.y;
+	const float Gb = DualHGPhaseData.z;
+
 
 #if USE_UNIFORM_PHASE_FUNCTION
 	const float4 Result = UniformSampleSphere(RandSample.xy);
 	// uniform scattering
-	return CreateMaterialSample(Result.xyz, RayleighWeight + HGWeight, Result.w, 1.0, 1.0, PATHTRACER_SCATTER_VOLUME);
+	return CreateMaterialSample(Result.xyz, RayleighWeight + HG0 + HG1 + HG2, Result.w, 1.0, 1.0, PATHTRACER_SCATTER_VOLUME);
 #else
-	const float LobeProb = LobeSelectionProb(RayleighWeight, HGWeight);
+	const float3 LobeCdf = LobeSelectionCdf(RayleighWeight, HG0, HG1, HG2);
+	const float4 LobePdf = LobeSelectionPdf(LobeCdf);
 
-	const float CosTheta = RandSample.x < LobeProb ?
-		RayleighPhaseInvertCdf(RescaleRandomNumber(RandSample.x, 0.0, LobeProb)) :
-		HenyeyGreensteinPhaseInvertCDF(RescaleRandomNumber(RandSample.x, LobeProb, 1.0), G);
+	float CosTheta = 0.0; 
+	if (RandSample.x < LobeCdf.x)
+	{
+		CosTheta = RayleighPhaseInvertCdf(RescaleRandomNumber(RandSample.x, 0.0, LobeCdf.x));
+	}
+	else
+	{
+		float G, R;
+		if (RandSample.x < LobeCdf.y)
+		{
+			R = RescaleRandomNumber(RandSample.x, LobeCdf.x, LobeCdf.y);
+			G = G0;
+		}
+		else if (RandSample.x < LobeCdf.z)
+		{
+			R = RescaleRandomNumber(RandSample.x, LobeCdf.y, LobeCdf.z);
+			G = G1;
+		}
+		else
+		{
+			R = RescaleRandomNumber(RandSample.x, LobeCdf.z, 1.0);
+			G = G2;
+		}
+		CosTheta = HenyeyGreensteinPhaseInvertCDF(R, G);
+	}
 	const float SinTheta = sqrt(saturate(1.0 - CosTheta * CosTheta));
 	const float Phi = RandSample.y * (2.0 * PI);
 
 	const float3 L_World = TangentToWorld(float3(cos(Phi) * SinTheta, sin(Phi) * SinTheta, CosTheta), -V_World);
 
 	FMaterialSample Result = CreateMaterialSample(L_World, 0.0, 0.0, 1.0, 1.0, PATHTRACER_SCATTER_VOLUME);
-	Result.AddLobeWithMIS(RayleighWeight, RayleighPhase(CosTheta), LobeProb);
-	Result.AddLobeWithMIS(HGWeight      , HenyeyGreensteinPhase(G, CosTheta )  , 1.0 - LobeProb);
+	Result.AddLobeWithMIS(RayleighWeight, RayleighPhase(CosTheta)            , LobePdf.x);
+	Result.AddLobeWithMIS(HG0           , HenyeyGreensteinPhase(G0, CosTheta), LobePdf.y);
+	Result.AddLobeWithMIS(HG1           , HenyeyGreensteinPhase(G1, CosTheta), LobePdf.z);
+	Result.AddLobeWithMIS(HG2           , HenyeyGreensteinPhase(G2, CosTheta), LobePdf.w);
 	return Result;
 #endif
 }
@@ -50,19 +84,28 @@ FMaterialEval Medium_EvalMaterial(
 )
 {
 	const float3 RayleighWeight = Payload.GetBaseColor();
-	const float3 HGWeight = Payload.GetHGWeight();
-	const float G = Payload.GetHGPhaseG();
+	const float3 HG0 = Payload.GetHGWeight();
+	const float G0 = Payload.GetHGPhaseG();
+	const float3 DualHGPhaseData = Payload.GetDualHGPhaseData();
+	const float3 HG1 = Payload.GetDualHGWeight() - Payload.GetDualHGWeight() * DualHGPhaseData.z;
+	const float3 HG2 = Payload.GetDualHGWeight() * DualHGPhaseData.z;
+	const float G1 = DualHGPhaseData.x;
+	const float G2 = DualHGPhaseData.y;
+	const float Gb = DualHGPhaseData.z;
 
 #if USE_UNIFORM_PHASE_FUNCTION
 	// simple omni-directional evaluation
-	return CreateMaterialEval((RayleighWeight + HGWeight) * DiffuseSpecularScale.x, 1.0 / (4 * PI));
+	return CreateMaterialEval((RayleighWeight + HG0 + HG1 + HG2) * DiffuseSpecularScale.x, 1.0 / (4 * PI));
 #else
-	const float LobeProb = LobeSelectionProb(RayleighWeight, HGWeight);
+	const float3 LobeCdf = LobeSelectionCdf(RayleighWeight, HG0, HG1, HG2);
+	const float4 LobePdf = LobeSelectionPdf(LobeCdf);
 
 	const float CosTheta = dot(V_World, L_World);
 	FMaterialEval Result = NullMaterialEval();
-	Result.AddLobeWithMIS(RayleighWeight, RayleighPhase(CosTheta), LobeProb);
-	Result.AddLobeWithMIS(HGWeight      , HenyeyGreensteinPhase(G, CosTheta), 1.0 - LobeProb);
+	Result.AddLobeWithMIS(RayleighWeight, RayleighPhase(CosTheta)            , LobePdf.x);
+	Result.AddLobeWithMIS(HG0           , HenyeyGreensteinPhase(G0, CosTheta), LobePdf.y);
+	Result.AddLobeWithMIS(HG1           , HenyeyGreensteinPhase(G1, CosTheta), LobePdf.z);
+	Result.AddLobeWithMIS(HG2           , HenyeyGreensteinPhase(G2, CosTheta), LobePdf.w);
 	Result.Weight *= DiffuseSpecularScale.x;
 	return Result;
 #endif
diff --git a/Engine/Shaders/Private/PathTracing/PathTracingBuildCloud.usf b/Engine/Shaders/Private/PathTracing/PathTracingBuildCloud.usf
new file mode 100644
--- /dev/null
+++ b/Engine/Shaders/Private/PathTracing/PathTracingBuildCloud.usf
@@ -0,0 +1,53 @@
+// Copyright Epic Games, Inc. All Rights Reserved.
+
+#define RANDSEQ_ERROR_DIFFUSION 1
+
+#include "../Common.ush"
+#include "../MonteCarlo.ush"
+#include "../SkyAtmosphereCommon.ush"
+#include "./Volume/PathTracingAtmosphereCommon.ush"
+#include "./Volume/PathTracingCloudsCommon.ush"
+#include "./Volume/PathTracingCloudsMaterialCommon.ush"
+#include "./Utilities/PathTracingRandomSequence.ush"
+#include "./PathTracingCommon.ush"
+
+RWTexture3D<uint4> CloudMap;
+uint TemporalSeed;
+uint Iteration;
+float BlendFactor;
+
+
+[numthreads(THREADGROUPSIZE_X, THREADGROUPSIZE_Y, THREADGROUPSIZE_Z)]
+void PathTracingBuildCloudCS(uint3 DispatchThreadId : SV_DispatchThreadID)
+{
+	if (any(DispatchThreadId >= uint3(CloudMapResolution)))
+	{
+		return;
+	}
+
+	RandomSequence RandSequence;
+	uint Seed = DispatchThreadId.x + CloudMapResolution.x * (DispatchThreadId.y + CloudMapResolution.y * DispatchThreadId.z); // unique seed per voxel
+	RandomSequence_Initialize(RandSequence, Seed, TemporalSeed);
+
+	float3 Jitter = RandomSequence_GenerateSample3D(RandSequence);
+	//Jitter = 0.5;
+	float3 UVW = float3(DispatchThreadId + Jitter) / float3(CloudMapResolution);
+	float3 WorldPos = GetCloudClipPosFromUVW(UVW);
+	float H = length(WorldPos);
+	WorldPos = mul(WorldPos, GetCloudClipBasis()) + CloudClipCenterKm; // TangentToWorld
+	FDFVector3 AbsoluteWorldPosition = DFMultiply(WorldPos, SKY_UNIT_TO_CM);
+	FSampleCloudMaterialResult Result = SampleCloudMaterial(AbsoluteWorldPosition, H, CloudLayerBotKm, CloudLayerTopKm);
+
+	if (Iteration > 0)
+	{
+		// if this is an incremental sample, blend with the previous result (so we get some averaging of all values in the voxel)
+		const FSampleCloudMaterialResult Previous = DecodeCloudResult(CloudMap[DispatchThreadId]);
+		Result.Extinction    = lerp(Previous.Extinction   , Result.Extinction   , BlendFactor);
+		Result.Albedo        = lerp(Previous.Albedo       , Result.Albedo       , BlendFactor);
+		Result.EmissiveColor = lerp(Previous.EmissiveColor, Result.EmissiveColor, BlendFactor);
+		Result.PhaseG1       = lerp(Previous.PhaseG1      , Result.PhaseG1      , BlendFactor);
+		Result.PhaseG2       = lerp(Previous.PhaseG2      , Result.PhaseG2      , BlendFactor);
+		Result.PhaseBlend    = lerp(Previous.PhaseBlend   , Result.PhaseBlend   , BlendFactor);
+	}
+	CloudMap[DispatchThreadId] = EncodeCloudResult(Result);
+}
diff --git a/Engine/Shaders/Private/PathTracing/PathTracingBuildCloudAccelerationMap.usf b/Engine/Shaders/Private/PathTracing/PathTracingBuildCloudAccelerationMap.usf
new file mode 100644
--- /dev/null
+++ b/Engine/Shaders/Private/PathTracing/PathTracingBuildCloudAccelerationMap.usf
@@ -0,0 +1,92 @@
+// Copyright Epic Games, Inc. All Rights Reserved.
+
+#include "../Common.ush"
+#include "../MonteCarlo.ush"
+#include "../SkyAtmosphereCommon.ush"
+#include "./Volume/PathTracingAtmosphereCommon.ush"
+#include "./Volume/PathTracingCloudsCommon.ush"
+#include "./Volume/PathTracingCloudsMaterialCommon.ush"
+#include "./Utilities/PathTracingRandomSequence.ush"
+
+uint NumSamples;
+uint TemporalSeed;
+uint Iteration;
+int UseCloudMap;
+
+Texture3D<uint4> CloudMap;
+RWTexture2D<float4> CloudAccelerationMap;
+
+[numthreads(THREADGROUPSIZE_X, THREADGROUPSIZE_Y, 1)]
+void PathTracingBuildCloudAccelerationMapCS(uint2 DispatchThreadId : SV_DispatchThreadID)
+{
+	if (any(DispatchThreadId >= CloudAccelMapResolution))
+	{
+		return;
+	}
+
+	// get min and max altitude of the clouds
+	const float T = CloudLayerTopKm;
+	const float B = CloudLayerBotKm;
+
+	// tracing view distance (around the origin -- can't go further than the top radius)
+	const float D = CloudClipDistKm;
+	
+	float MinDensity = POSITIVE_INFINITY;
+	float MaxDensity = 0;
+	float ZMin = POSITIVE_INFINITY;
+	float ZMax = NEGATIVE_INFINITY;
+	if (Iteration > 0)
+	{
+		float4 Data = CloudAccelerationMap[DispatchThreadId];
+		MaxDensity = Data.x;
+		MinDensity = Data.y;
+		ZMin = Data.z;
+		ZMax = Data.w;
+	}
+	for (uint Sample = 0; Sample < NumSamples; Sample++)
+	{
+		RandomSequence RandSequence;
+		RandomSequence_Initialize(RandSequence, DispatchThreadId.x + DispatchThreadId.y * 65536, Iteration * NumSamples + Sample);
+		float3 Jitter = RandomSequence_GenerateSample3D(RandSequence);
+		const float2 UV = 2.0 * ((DispatchThreadId + Jitter.xy) / float(CloudAccelMapResolution)) - 1.0;
+		// Project a 2d map onto the surface of a sphere
+		float3 ClipPos;
+		ClipPos.xy = D * UV;
+		float R2 = length2(ClipPos.xy);
+		float TopZ = sqrt(max(T * T - R2, 0.0));
+		float BotZ = sqrt(max(B * B - R2, 0.0));
+		ClipPos.z = lerp(BotZ, TopZ, Jitter.z);
+
+		float Z = ClipPos.z - CloudClipRadiusKm; // remember height of the point (before rotation, relative to center)
+		float H = length(ClipPos);
+
+		float3 WorldPos = mul(ClipPos, GetCloudClipBasis()) + CloudClipCenterKm; // TangentToWorld
+
+		uint4 Data = 0;
+		if (UseCloudMap)
+		{
+			const float3 UVW = GetUVWFromCloudClipPos(ClipPos);
+			const int3 GridCoord = int3(floor(float3(CloudMapResolution) * UVW));
+			Data = CloudMap.Load(int4(GridCoord, 0));
+		}
+		else
+		{
+			FDFVector3 AbsoluteWorldPosition = DFMultiply(WorldPos, SKY_UNIT_TO_CM);
+			// run the material and encode the result, which includes the effect that quantization will have during rendering
+			Data = EncodeCloudResult(SampleCloudMaterial(AbsoluteWorldPosition, H, B, T));
+
+		}
+		const float3 Extinction = DecodeCloudResult(Data).Extinction;
+		const float Density   = max3(Extinction.x, Extinction.y, Extinction.z);
+		const float DensityLo = min3(Extinction.x, Extinction.y, Extinction.z);
+
+		if (Density > 0)
+		{
+			ZMin = min(ZMin, Z);
+			ZMax = max(ZMax, Z);
+		}
+		MaxDensity = max(MaxDensity, Density);
+		MinDensity = min(MinDensity, DensityLo);
+	}
+	CloudAccelerationMap[DispatchThreadId] = float4(MaxDensity, MinDensity, ZMin, ZMax);
+}
diff --git a/Engine/Shaders/Private/PathTracing/PathTracingCommon.ush b/Engine/Shaders/Private/PathTracing/PathTracingCommon.ush
--- a/Engine/Shaders/Private/PathTracing/PathTracingCommon.ush
+++ b/Engine/Shaders/Private/PathTracing/PathTracingCommon.ush
@@ -9,6 +9,7 @@
 #include "Material/PathTracingMaterialCommon.ush"
 #include "../RayTracing/RayTracingCommon.ush"
 #include "../OctahedralCommon.ush"
+#include "../DoubleFloat.ush"
 #include "/Engine/Shared/PathTracingDefinitions.h"
 
 // These flags are analogous to RAY_TRACING_PAYLOAD_* flags. There are currently 9 bits available.
@@ -371,6 +372,36 @@ struct FPackedPathTracingPayload : FMinimalPayload
 		return asfloat(PackedData[14]);
 	}
 
+	// Volumetric callable shaders
+	void SetVolumetricCallableShaderInput(FDFVector3 AbsolutePosition, float CloudHeightKm)
+	{
+		PackedData[0] = asuint(AbsolutePosition.High.x);
+		PackedData[1] = asuint(AbsolutePosition.High.y);
+		PackedData[2] = asuint(AbsolutePosition.High.z);
+		PackedData[3] = asuint(AbsolutePosition.Low.x);
+		PackedData[4] = asuint(AbsolutePosition.Low.y);
+		PackedData[5] = asuint(AbsolutePosition.Low.z);
+		PackedData[6] = asuint(CloudHeightKm);
+	}
+	FDFVector3 GetVolumetricCallableShaderInputAbsolutePosition()
+	{
+		return MakeDFVector3(
+			float3(asfloat(PackedData[0]), asfloat(PackedData[1]), asfloat(PackedData[2])),
+			float3(asfloat(PackedData[3]), asfloat(PackedData[4]), asfloat(PackedData[5]))
+		);
+	}
+	float GetVolumetricCallableShaderInputCloudHeightKm() { return asfloat(PackedData[6]); }
+
+	void SetVolumetricCallableShaderOutput(uint4 Data)
+	{
+		PackedData[ 7] = Data.x;
+		PackedData[ 8] = Data.y;
+		PackedData[ 9] = Data.z;
+		PackedData[10] = Data.w;
+	}
+
+	uint4 GetVolumetricCallableShaderOutput() { return uint4(PackedData[7], PackedData[8], PackedData[9], PackedData[10]); }
+
 #endif // IS_PAYLOAD_ENABLED(RT_PAYLOAD_TYPE_GPULIGHTMASS)
 };
 
@@ -591,6 +622,15 @@ struct FPathTracingPayload : FMinimalPayload
 	float3 GetHGWeight() { return CustomData0.xyz; }
 	float GetHGPhaseG() { return CustomData0.w; }
 	void SetHG(float3 Weight, float G) { CustomData0 = float4(Weight, G); }
+	float3 GetDualHGWeight() { return CustomData1.xyz; }
+	float3 GetDualHGPhaseData() { return float3(Metallic, Specular, Roughness); }
+	void SetDualHG(float3 Weight, float3 PhaseData)
+	{
+		CustomData1.xyz = Weight;
+		Metallic = PhaseData.x;
+		Specular = PhaseData.y;
+		Roughness = PhaseData.z;
+	}
 #endif // PATHTRACING_SUBSTRATE_PAYLOAD
 
 	// Methods used by the integrator
@@ -760,6 +800,7 @@ FPackedPathTracingPayload PackPathTracingPayload(FPathTracingPayload Input)
 	return Output;
 }
 
+#if RAYGENSHADER
 FPathTracingPayload UnpackPathTracingPayload(FPackedPathTracingPayload Input, FRayDesc Ray)
 {
 	FPathTracingPayload Output = (FPathTracingPayload)0;
@@ -842,3 +883,4 @@ FPathTracingPayload UnpackPathTracingPayload(FPackedPathTracingPayload Input, FR
 #endif // IS_PAYLOAD_ENABLED(RT_PAYLOAD_TYPE_GPULIGHTMASS)
 	return Output;
 }
+#endif
diff --git a/Engine/Shaders/Private/PathTracing/PathTracingCore.ush b/Engine/Shaders/Private/PathTracing/PathTracingCore.ush
--- a/Engine/Shaders/Private/PathTracing/PathTracingCore.ush
+++ b/Engine/Shaders/Private/PathTracing/PathTracingCore.ush
@@ -485,7 +485,7 @@ FPackedPathTracingPayload TraceTransparentRay(inout FPathState PathState, int Bo
 			PathState.Ray.GetNativeDesc(),
 			PackedPayload);
 
-		FVolumeIntersectionList VolumeIsectList = VolumeIntersect(PathState.Ray.Origin, PathState.Ray.Direction, PathState.Ray.TMin, PathState.Ray.TMax);
+		FVolumeIntersectionList VolumeIsectList = VolumeIntersect(PathState.Ray.Origin, PathState.Ray.Direction, PathState.Ray.TMin, PathState.Ray.TMax, PathState.GetPathRoughness());
 		if (VolumeIsectList.HitBlocker() && (PackedPayload.IsMiss() || VolumeIsectList.BlockerHitT < PackedPayload.HitT))
 		{
 			// we didn't hit any real geometry, but we did hit the volume's blocker geometry
@@ -594,7 +594,7 @@ FPackedPathTracingPayload TraceTransparentRay(inout FPathState PathState, int Bo
 						// if this is not the last bounce, consider the possibility of scattering as giving up medium hits
 						if (!bLastBounce)
 						{
-							float3 SigmaS = min(Result.SigmaSRayleigh + Result.SigmaSHG, SigmaT);
+							float3 SigmaS = min(Result.SigmaSRayleigh + Result.SigmaSHG + Result.SigmaSDualHG, SigmaT);
 							// accumulate a signal for the denoiser
 							float3 Albedo = PathState.GetAlbedo();
 							AccumulateAlbedo(SigmaS, PathThroughput, PathState.GetPathRoughness(), Albedo);
@@ -627,7 +627,10 @@ FPackedPathTracingPayload TraceTransparentRay(inout FPathState PathState, int Bo
 				// restore interval start now that we are finished with ray marching
 				Interval.VolumeTMin = StartVolumeTMin;
 
-				const bool bUseAnalyticTransmittance = (VolumeFlags & PATH_TRACER_VOLUME_USE_ANALYTIC_TRANSMITTANCE) != 0;
+				// In order to use analytic transmittance, the user must have requested it _and_ we must be on an interval with no clouds or heterogeneous volumes
+				const bool bUseAnalyticTransmittance =
+					(VolumeFlags & PATH_TRACER_VOLUME_USE_ANALYTIC_TRANSMITTANCE) != 0 &&
+					(Interval.VolumeMask & (PATH_TRACER_VOLUME_ENABLE_CLOUDS | PATH_TRACER_VOLUME_ENABLE_HETEROGENEOUS_VOLUMES)) == 0;
 				if (bUseAnalyticTransmittance)
 				{
 					// analytically handle the transmittance to the next surface (which will be higher quality than the transmittance implicitly computed above and prevent noise with emissive surfaces in volumes)
@@ -782,7 +785,7 @@ FPackedPathTracingPayload TraceTransparentRay(inout FPathState PathState, int Bo
 	return Payload;
 }
 
-float3 TraceTransparentVisibilityRay(FRayDesc Ray, float PathRoughness, uint MissShaderIndex, bool bCastShadows, inout RandomSequence RandSequence)
+float3 TraceTransparentVisibilityRay(FRayDesc Ray, float PathRoughness, float PreviousPathRoughness, uint MissShaderIndex, bool bCastShadows, inout RandomSequence RandSequence)
 {
 	if (!bCastShadows && MissShaderIndex == 0)
 	{
@@ -822,7 +825,8 @@ float3 TraceTransparentVisibilityRay(FRayDesc Ray, float PathRoughness, uint Mis
 	float3 Throughput = PackedPayload.GetRayThroughput();
 	if (bCastShadows && any(Throughput > 0))
 	{
-		FVolumeIntersectionList VolumeIsectList = VolumeIntersect(Ray.Origin, Ray.Direction, Ray.TMin, Ray.TMax);
+		// NOTE: use the path roughness before scattering off the current vertex (so we can have clouds casting shadows onto diffuse surfaces)
+		FVolumeIntersectionList VolumeIsectList = VolumeIntersect(Ray.Origin, Ray.Direction, Ray.TMin, Ray.TMax, PreviousPathRoughness);
 		if (VolumeIsectList.HitBlocker())
 		{
 			// blockers are opaque -- done!
@@ -1464,7 +1468,7 @@ bool PathTracingKernel(inout FPathState PathState, int Bounce)
 			const float3 TranslatedWorldPos = Ro + EquiAngularT * Rd;
 			const FVolumeShadedResult Result = VolumeGetDensity(TranslatedWorldPos, VolumeSegment.Interval.VolumeMask, bIsCameraRay);
 			const float3 SigmaT = Result.SigmaT;
-			const float3 SigmaS = min(Result.SigmaSRayleigh + Result.SigmaSHG, SigmaT);
+			const float3 SigmaS = min(Result.SigmaSRayleigh + Result.SigmaSHG + Result.SigmaSDualHG, SigmaT);
 
 			if (any(SigmaS > 0))
 			{
@@ -1505,7 +1509,7 @@ bool PathTracingKernel(inout FPathState PathState, int Bounce)
 							{
 								LightContrib *= MISWeightPower(LightSample.Pdf, MaterialEval.Pdf);
 							}
-							LightContrib *= TraceTransparentVisibilityRay(LightRay, 1.0, MissShaderIndex, bCastShadows, PathState.RandSequence);
+							LightContrib *= TraceTransparentVisibilityRay(LightRay, 1.0, PathState.GetPathRoughness(), MissShaderIndex, bCastShadows, PathState.RandSequence);
 
 							AccumulateRadiance(PathState.Radiance, LightContrib, bIsCameraRay);
 						}
@@ -1531,7 +1535,7 @@ bool PathTracingKernel(inout FPathState PathState, int Bounce)
 								LightContrib *= MISWeightPower(MaterialSample.Pdf, LightResult.Pdf);
 							}
 							MaterialRay.TMax = LightResult.HitT;
-							LightContrib *= TraceTransparentVisibilityRay(MaterialRay, MaterialSample.Roughness, MissShaderIndex, bCastShadows, PathState.RandSequence);
+							LightContrib *= TraceTransparentVisibilityRay(MaterialRay, 1.0, PathState.GetPathRoughness(), MissShaderIndex, bCastShadows, PathState.RandSequence);
 							AccumulateRadiance(PathState.Radiance, LightContrib, bIsCameraRay);
 						}
 					}
@@ -1675,7 +1679,7 @@ bool PathTracingKernel(inout FPathState PathState, int Bounce)
 					}
 					const bool bCastShadows = bIsVolumeSample ? CastsVolumeShadow(LightId) : CastsShadow(LightId);
 					const uint MissShaderIndex = GetLightMissShaderIndex(LightId);
-					LightContrib *= TraceTransparentVisibilityRay(LightRay, AvgRoughness, MissShaderIndex, bCastShadows, PathState.RandSequence);
+					LightContrib *= TraceTransparentVisibilityRay(LightRay, AvgRoughness, PathState.GetPathRoughness(), MissShaderIndex, bCastShadows, PathState.RandSequence);
 
 					AccumulateRadiance(PathState.Radiance, LightContrib, bIsCameraRay);
 				}
@@ -1736,7 +1740,8 @@ bool PathTracingKernel(inout FPathState PathState, int Bounce)
 	ApplyRayBias(PathState.Ray, Payload.HitT, MaterialSample.PositionBiasSign * Payload.WorldGeoNormal);
 
 	// enlarge roughness based on the chosen lobe roughness
-	PathState.SetPathRoughness(max(PathState.GetPathRoughness(), MaterialSample.Roughness));
+	float PreviousPathRoughness = PathState.GetPathRoughness();
+	PathState.SetPathRoughness(max(PreviousPathRoughness, MaterialSample.Roughness));
 
 	// update the current extinction if we are crossing a boundary on glass or water
 	// summing the local extinction gives a rudimentary way of dealing with overlapping regions
@@ -1819,7 +1824,7 @@ bool PathTracingKernel(inout FPathState PathState, int Bounce)
 				const bool bCastShadows = bIsVolumeSample ? CastsVolumeShadow(LightId) : CastsShadow(LightId);
 				const uint MissShaderIndex = GetLightMissShaderIndex(LightId);
 				const float ShadowRayRoughness = ApproximateCaustics ? PathState.GetPathRoughness() : 0.0;
-				LightContrib *= TraceTransparentVisibilityRay(LightRay, ShadowRayRoughness, MissShaderIndex, bCastShadows, PathState.RandSequence);
+				LightContrib *= TraceTransparentVisibilityRay(LightRay, ShadowRayRoughness, PreviousPathRoughness, MissShaderIndex, bCastShadows, PathState.RandSequence);
 				AccumulateRadiance(PathState.Radiance, LightContrib, bIsCameraRay);
 			}
 		}
diff --git a/Engine/Shaders/Private/PathTracing/PathTracingVolumetricCloudMaterialShader.usf b/Engine/Shaders/Private/PathTracing/PathTracingVolumetricCloudMaterialShader.usf
new file mode 100644
--- /dev/null
+++ b/Engine/Shaders/Private/PathTracing/PathTracingVolumetricCloudMaterialShader.usf
@@ -0,0 +1,32 @@
+// Copyright Epic Games, Inc. All Rights Reserved.
+
+#define PATH_TRACING 1
+
+#include "/Engine/Private/Common.ush"
+#include "/Engine/Private/RayTracing/RayTracingCommon.ush"
+#include "/Engine/Private/PathTracing/PathTracingShaderUtils.ush"
+
+#if CLOUD_LAYER_PIXEL_SHADER
+#include "./Volume/PathTracingCloudsMaterialCommon.ush"
+#endif
+
+#include "./PathTracingCommon.ush"
+
+
+RAY_TRACING_ENTRY_CALLABLE(PathTracingVolumetricCloudMaterialShader, FPackedPathTracingPayload, PackedPayload)
+{
+#if CLOUD_LAYER_PIXEL_SHADER
+	ResolvedView = ResolveView();
+
+	FDFVector3 AbsoluteWorldPosition = PackedPayload.GetVolumetricCallableShaderInputAbsolutePosition();
+	float H = PackedPayload.GetVolumetricCallableShaderInputCloudHeightKm();
+	float B = PathTracingCloudParameters.CloudLayerBotKm;
+	float T = PathTracingCloudParameters.CloudLayerTopKm;
+
+	FSampleCloudMaterialResult Result = SampleCloudMaterial(AbsoluteWorldPosition, H, B, T);
+	PackedPayload.SetVolumetricCallableShaderOutput(EncodeCloudResult(Result));
+#else
+	// for default shader - just flush outputs
+	PackedPayload.SetVolumetricCallableShaderOutput(0);
+#endif
+}
diff --git a/Engine/Shaders/Private/PathTracing/Volume/PathTracingClouds.ush b/Engine/Shaders/Private/PathTracing/Volume/PathTracingClouds.ush
new file mode 100644
--- /dev/null
+++ b/Engine/Shaders/Private/PathTracing/Volume/PathTracingClouds.ush
@@ -0,0 +1,237 @@
+// Copyright Epic Games, Inc. All Rights Reserved.
+
+#pragma once
+
+#include "./PathTracingAtmosphereCommon.ush"
+#include "./PathTracingCloudsCommon.ush"
+#include "../../SkyAtmosphereCommon.ush"
+
+
+Texture2D<float4> CloudAccelerationMap;
+SamplerState CloudAccelerationMapSampler;
+Texture3D<uint4> CloudMap;
+
+// intersect the cloud layer with high precision
+float4 RayCloudIntersectPrecise(float3 Ro, float3 Rd)
+{
+	float4 Result = -1.0;
+
+	const FDFVector3 Center = MakeDFVector3(PlanetCenterTranslatedWorldHi, PlanetCenterTranslatedWorldLo);
+	const FDFScalar CloudTopRadius = DFMultiply(CloudLayerTopKm, SKY_UNIT_TO_CM);
+	const FDFScalar CloudBotRadius = DFMultiply(CloudLayerBotKm, SKY_UNIT_TO_CM);
+
+
+	const FDFVector3 Co = DFSubtract(Center, Ro);
+	const FDFScalar b = DFDot(Co, Rd);
+	const FDFVector3 H = DFSubtract(Co, DFMultiply(Rd, b));
+	const FDFScalar hd = DFLengthSqr(H);
+	const FDFScalar TR2 = DFSqr(CloudTopRadius);
+	FDFScalar ht = DFSubtract(TR2, hd);
+
+	if (DFGreater(ht, 0.0))
+	{
+		ht = DFSqrt(ht);
+		FDFScalar q;
+		if (DFGreater(b, 0.0))
+			q = DFAdd(b, ht);
+		else
+			q = DFSubtract(b, ht);
+		const float t0 = DFDemote(q);
+		const float t1 = DFFastDivideDemote(DFSubtract(DFLengthSqr(Co), TR2), q);
+		Result.x = min(t0, t1);
+		Result.y = max(t0, t1);
+	}
+	const FDFScalar BR2 = DFSqr(CloudBotRadius);
+	FDFScalar hb = DFSubtract(BR2, hd);
+	if (DFGreater(hb, 0.0))
+	{
+		hb = DFSqrt(hb);
+		FDFScalar q;
+		if (DFGreater(b, 0.0))
+			q = DFAdd(b, hb);
+		else
+			q = DFSubtract(b, hb);
+		const float t0 = DFDemote(q);
+		const float t1 = DFFastDivideDemote(DFSubtract(DFLengthSqr(Co), BR2), q);
+		Result.z = min(t0, t1);
+		Result.w = max(t0, t1);
+	}
+
+	return Result;
+}
+
+int2 PositionToVoxel(float2 P)
+{
+	float CloudVoxelWidth = 2.0 * SKY_UNIT_TO_CM * CloudClipDistKm / CloudAccelMapResolution;
+	int2 Idx = int2((P + CloudClipDistKm * SKY_UNIT_TO_CM) * rcp(CloudVoxelWidth));
+	return clamp(Idx, 0, CloudAccelMapResolution - 1);
+}
+
+float2 VoxelToPosition(int2 Idx)
+{
+	float CloudVoxelWidth = 2.0 * SKY_UNIT_TO_CM * CloudClipDistKm / CloudAccelMapResolution;
+	return (Idx * CloudVoxelWidth) - CloudClipDistKm * SKY_UNIT_TO_CM;
+}
+
+FVolumeIntersection CloudsIntersect(float3 Origin, float3 Direction, float TMin, float TMax, float PathRoughness)
+{
+	if (PathRoughness > CloudRoughnessCutoff)
+	{
+		return CreateEmptyVolumeIntersection();
+	}
+	float4 Result = RayCloudIntersectPrecise(Origin, Direction);
+	float VolumeTMin = max(Result.x, TMin);
+	float VolumeTMax = min(Result.y, TMax);
+	if (VolumeTMin < VolumeTMax)
+	{
+		// we intersected the top layer of clouds, refine to see if we hit the bottom layer
+		if (Result.z > 0)
+		{
+			// bottom layer start in front - act as the back side of our interval
+			VolumeTMax = min(Result.z, VolumeTMax);
+		}
+		else
+		{
+			// we are inside the bottom layer, or missed it entirely - clip the front segment
+			VolumeTMin = max(VolumeTMin, Result.w);
+		}
+
+		// isect clip slabs, the AABB our clouds are enclosed in
+		float3 Ro = mul(GetCloudClipBasis(), Origin);
+		float3 Rd = mul(GetCloudClipBasis(), Direction);
+		float3 InvRd = rcp(Rd);
+		float2 Slab0 = (-Ro.xy - CloudClipDistKm * SKY_UNIT_TO_CM) * InvRd.xy;
+		float2 Slab1 = (-Ro.xy + CloudClipDistKm * SKY_UNIT_TO_CM) * InvRd.xy;
+		VolumeTMin = max(VolumeTMin, max(min(Slab0.x, Slab1.x), min(Slab0.y, Slab1.y)));
+		VolumeTMax = min(VolumeTMax, min(max(Slab0.x, Slab1.x), max(Slab0.y, Slab1.y)));
+		if (VolumeTMin >= VolumeTMax)
+			return CreateEmptyVolumeIntersection();
+
+		VolumeTMax = min(VolumeTMax, VolumeTMin + CloudTracingMaxDistance);
+
+#if 1 // use cloud accel map to tighten intersection bounds
+
+		// now that we know the ray length, lets march through the grid to see if we can clip away anything
+		float VoxelWidth = (2 * SKY_UNIT_TO_CM * CloudClipDistKm) / CloudAccelMapResolution;
+		
+		float2 P = Ro.xy + VolumeTMin * Rd.xy;
+		int2 Idx = PositionToVoxel(P);
+		float2 NextCrossingT = VolumeTMin + (VoxelToPosition(Idx + int2(Rd.xy >= 0)) - P) * InvRd.xy;
+		float2 DeltaT = abs(VoxelWidth * InvRd.xy);
+		int2 Step = select(Rd.xy >= 0, 1, -1);
+		int2 Stop = select(Rd.xy >= 0, CloudAccelMapResolution, -1);
+		
+		// start with an empty interval
+		float ResultTMin = VolumeTMax;
+		float ResultTMax = VolumeTMin;
+		float MaxDensity = 0;
+		for (;;)
+		{
+			// visit current voxel
+			float4 Data = CloudAccelerationMap.Load(uint3(Idx.xy, 0));
+			bool bHasData = Data.x > 0;
+			if (bHasData)
+			{
+				float CellMinT = VolumeTMin;
+				float CellMaxT = min3(VolumeTMax, NextCrossingT.x, NextCrossingT.y);
+
+				float SlabT0 = (Data.z * SKY_UNIT_TO_CM - Ro.z) * InvRd.z;
+				float SlabT1 = (Data.w * SKY_UNIT_TO_CM - Ro.z) * InvRd.z;
+
+				CellMinT = max(CellMinT, min(SlabT0, SlabT1));
+				CellMaxT = min(CellMaxT, max(SlabT0, SlabT1));
+				// do we still have a valid interval?
+				if (CellMinT < CellMaxT)
+				{
+					// expand result bounds
+					ResultTMin = min(ResultTMin, CellMinT);
+					ResultTMax = max(ResultTMax, CellMaxT);
+					MaxDensity = max(Data.x, MaxDensity);
+				}
+			}
+
+			if (NextCrossingT.x < NextCrossingT.y)
+			{
+				// step x first
+				Idx.x += Step.x;
+				if (Idx.x == Stop.x || VolumeTMax < NextCrossingT.x)
+				{
+					break;
+				}
+				VolumeTMin = NextCrossingT.x;
+				NextCrossingT.x += DeltaT.x;
+			}
+			else
+			{
+				Idx.y += Step.y;
+				if (Idx.y == Stop.y || VolumeTMax < NextCrossingT.y)
+				{
+					break;
+				}
+				VolumeTMin = NextCrossingT.y;
+				NextCrossingT.y += DeltaT.y;
+			}
+		}
+		return CreateVolumeIntersectionWithDensity(ResultTMin, ResultTMax, MaxDensity);
+#else
+		// if we don't have the cloud accel map, we can't know the upper bound on density...
+		return CreateVolumeIntersectionWithDensity(VolumeTMin, VolumeTMax, CM_TO_SKY_UNIT);
+#endif
+	}
+	return CreateEmptyVolumeIntersection();
+}
+
+FVolumeShadedResult CloudsGetDensity(float3 TranslatedWorldPos)
+{
+	FVolumeShadedResult Result = (FVolumeShadedResult)0;
+
+	float3 WorldPos = DFMultiplyDemote(DFFastSubtract(TranslatedWorldPos, PrimaryView.PreViewTranslation), CM_TO_SKY_UNIT);
+
+	if ((VolumeFlags & PATH_TRACER_VOLUME_VIEW_CLOUD_MAP_ACCEL) != 0)
+	{
+		// debug preview path to visualize the cloud accel map
+		float3 P = mul(GetCloudClipBasis(), TranslatedWorldPos); // WorldToTangent
+		float2 UV = (P.xy * CM_TO_SKY_UNIT + CloudClipDistKm.xx) / (2 * CloudClipDistKm);
+		const int Level = 0;
+		float4 Data = CloudAccelerationMap.SampleLevel(CloudAccelerationMapSampler, UV, Level);
+	
+		float Z = P.z * CM_TO_SKY_UNIT;
+		float Density = (Z > Data.z && Z < Data.w) ? Data.x : 0.0;
+		Result.SigmaT = Density;
+		Result.SigmaSHG = 0.75 * Density;
+		return Result;
+	}
+
+	uint4 Data = 0;
+#if PLATFORM_SUPPORTS_CALLABLE_SHADERS
+	const bool bUseCloudMap = (VolumeFlags & PATH_TRACER_VOLUME_USE_CLOUD_MAP) != 0;
+	if (bUseCloudMap)
+#endif
+	{
+		// Main path -- via baked cloud map
+		const float3 P = mul(GetCloudClipBasis(), TranslatedWorldPos * CM_TO_SKY_UNIT); // WorldToTangent
+		const float3 UVW = GetUVWFromCloudClipPos(P + float3(0, 0, CloudClipRadiusKm));
+		const int3 GridCoord = int3(floor(float3(CloudMapResolution) * UVW));
+		Data = CloudMap.Load(int4(GridCoord, 0));
+	}
+#if PLATFORM_SUPPORTS_CALLABLE_SHADERS
+	else
+	{
+		// Main path -- invoke the actual callable shader
+		FPackedPathTracingPayload Payload = (FPackedPathTracingPayload)0;
+		Payload.SetVolumetricCallableShaderInput(
+			DFFastSubtract(TranslatedWorldPos, PrimaryView.PreViewTranslation),
+			GetSphericalPlanetParameters(TranslatedWorldPos, float3(0, 0, 1)).HeightKm
+		);
+		CallShader(CloudCallableShaderId, Payload);
+		Data = Payload.GetVolumetricCallableShaderOutput();
+#endif
+	}
+	const FSampleCloudMaterialResult CloudResult = DecodeCloudResult(Data);
+	Result.SigmaT        = CloudResult.Extinction;
+	Result.SigmaSDualHG  = CloudResult.Albedo * Result.SigmaT;
+	Result.Emission      = CloudResult.EmissiveColor;
+	Result.DualPhaseData = float3(CloudResult.PhaseG1, CloudResult.PhaseG2, CloudResult.PhaseBlend);
+	return Result;
+
+}
diff --git a/Engine/Shaders/Private/PathTracing/Volume/PathTracingCloudsCommon.ush b/Engine/Shaders/Private/PathTracing/Volume/PathTracingCloudsCommon.ush
new file mode 100644
--- /dev/null
+++ b/Engine/Shaders/Private/PathTracing/Volume/PathTracingCloudsCommon.ush
@@ -0,0 +1,163 @@
+// Copyright Epic Games, Inc. All Rights Reserved.
+
+#pragma once
+
+#include "../PathTracingCommon.ush"
+
+float3 CloudClipX;
+float3 CloudClipY;
+float3 CloudClipZ;
+
+float3 CloudClipCenterKm;
+
+float CloudLayerBotKm;
+float CloudLayerTopKm;
+float CloudClipDistKm;
+float CloudClipRadiusKm;
+
+float CloudTracingMaxDistance;
+float CloudRoughnessCutoff;
+int CloudAccelMapResolution;
+int CloudCallableShaderId;
+uint3 CloudMapResolution;
+
+float3x3 GetCloudClipBasis()
+{
+	return float3x3(CloudClipX, CloudClipY, CloudClipZ);
+}
+
+float GetLowestCloudZ()
+{
+	return sqrt(max(CloudLayerBotKm * CloudLayerBotKm - CloudClipDistKm * CloudClipDistKm * 2.0, 0.0));
+}
+
+float GetFarthestCornerDist()
+{
+	return sqrt(2 * CloudClipDistKm * CloudClipDistKm + max(Square(GetLowestCloudZ() - CloudClipRadiusKm), Square(CloudLayerTopKm - CloudClipRadiusKm)));
+}
+
+float2 GetRadiusBounds()
+{
+	float2 RBounds = float2(0, GetFarthestCornerDist());
+	float CloudLoKm = GetLowestCloudZ();
+	if (CloudClipRadiusKm < CloudLoKm)
+	{
+		RBounds.x = CloudLoKm - CloudClipRadiusKm;
+	}
+	else if (CloudClipRadiusKm > CloudLayerTopKm)
+	{
+		RBounds.x = CloudClipRadiusKm - CloudLayerTopKm;
+	}
+	return RBounds;
+}
+
+#define CLOUD_MAPPING_TYPE   2
+
+
+float3 GetCloudClipPosFromUVW(float3 UVW)
+{
+#if CLOUD_MAPPING_TYPE == 0
+	// super basic voxelization for debugging
+	const float3 Corner0 = float3(-CloudClipDistKm.xx, GetLowestCloudZ());
+	const float3 Corner1 = float3(+CloudClipDistKm.xx, CloudLayerTopKm  );
+	return lerp(Corner0, Corner1, UVW);
+#elif CLOUD_MAPPING_TYPE == 1 || CLOUD_MAPPING_TYPE == 2
+	const float3 Corner0 = float3(-CloudClipDistKm.xx, GetLowestCloudZ());
+	const float3 Corner1 = float3(+CloudClipDistKm.xx, CloudLayerTopKm  );
+
+
+	float SliceDist = abs(GetLowestCloudZ() - CloudClipRadiusKm);
+
+	float2 Angle0 = atan2(Corner0.xy, SliceDist);
+	float2 Angle1 = atan2(Corner1.xy, SliceDist);
+	float3 CloudClipPos = float3(SliceDist * tan(lerp(Angle0, Angle1, UVW.xy)), 0.0);
+
+#if CLOUD_MAPPING_TYPE == 1
+	// angular mapping
+	const float ColumnDist = length(CloudClipPos.xy);
+	const float AngleZ0 = atan2(Corner0.z, ColumnDist);
+	const float AngleZ1 = atan2(Corner1.z, ColumnDist);
+	CloudClipPos.z = ColumnDist * tan(lerp(AngleZ0, AngleZ1, UVW.z));
+#else
+	// linear mapping
+	CloudClipPos.z = lerp(Corner0.z, Corner1.z, UVW.z);
+#endif
+
+
+	return CloudClipPos;
+
+#else
+#error "Unknown cloud mapping type"
+#endif
+}
+
+float3 GetUVWFromCloudClipPos(float3 CloudClipPos)
+{
+#if CLOUD_MAPPING_TYPE == 0
+	// super basic voxelization for debugging
+	const float3 Corner0 = float3(-CloudClipDistKm.xx, GetLowestCloudZ());
+	const float3 Corner1 = float3(+CloudClipDistKm.xx, CloudLayerTopKm  );
+	return saturate((CloudClipPos - Corner0) / (Corner1 - Corner0)) * 0.99999994;
+#elif CLOUD_MAPPING_TYPE == 1 || CLOUD_MAPPING_TYPE == 2
+	const float3 Corner0 = float3(-CloudClipDistKm.xx, GetLowestCloudZ());
+	const float3 Corner1 = float3(+CloudClipDistKm.xx, CloudLayerTopKm  );
+
+	float SliceDist = abs(GetLowestCloudZ() - CloudClipRadiusKm);
+
+	float2 Angle0 = atan2(Corner0.xy, SliceDist);
+	float2 Angle1 = atan2(Corner1.xy, SliceDist);
+
+	float2 UV = (atan2(CloudClipPos.xy, SliceDist) - Angle0) / (Angle1 - Angle0);
+
+#if CLOUD_MAPPING_TYPE == 1
+	// angular mapping
+	const float ColumnDist = length(CloudClipPos.xy);
+	const float AngleZ0 = atan2(Corner0.z, ColumnDist);
+	const float AngleZ1 = atan2(Corner1.z, ColumnDist);
+	const float W = (atan2(CloudClipPos.z, ColumnDist) - AngleZ0) / (AngleZ1 - AngleZ0);
+#else
+	// linear mapping
+	const float W = (CloudClipPos.z - Corner0.z) / (Corner1.z - Corner0.z);
+#endif
+
+	return saturate(float3(UV, W));
+#else
+#error "Unknown cloud mapping type"
+#endif
+}
+
+
+
+// Simplified struct only used in the cloud case (to reduce the amount of data needed in the payload for the callable shader)
+struct FSampleCloudMaterialResult
+{
+	float3 Extinction;
+	float3 Albedo;
+	float3 EmissiveColor;
+	float PhaseG1;
+	float PhaseG2;
+	float PhaseBlend;
+};
+
+FSampleCloudMaterialResult DecodeCloudResult(uint4 Data)
+{
+	FSampleCloudMaterialResult Result = (FSampleCloudMaterialResult) 0;
+	Result.Extinction    = PayloadDecodeHDRColor(Data.x) * (1.0 / 64.0);
+	Result.Albedo        = PayloadDecodeLDRColor(Data.y);
+	Result.EmissiveColor = PayloadDecodeHDRColor(Data.z);
+	float3 PhaseData     = PayloadDecodeLDRColor(Data.w) * float3(2.0, 2.0, 1.0) - float3(1.0, 1.0, 0.0);
+	Result.PhaseG1       = PhaseData.x;
+	Result.PhaseG2       = PhaseData.y;
+	Result.PhaseBlend    = PhaseData.z;
+	return Result;
+}
+
+uint4 EncodeCloudResult(FSampleCloudMaterialResult Result)
+{
+	return uint4(
+		PayloadEncodeHDRColor(Result.Extinction * 64.0), // scale up so we don't loose any density when clipping bits and make better use of the higher range of the encoding
+		PayloadEncodeLDRColor(Result.Albedo),
+		PayloadEncodeHDRColor(Result.EmissiveColor),
+		PayloadEncodeLDRColor(float3(Result.PhaseG1, Result.PhaseG2, Result.PhaseBlend) * float3(0.5, 0.5, 1.0) + float3(0.5, 0.5, 0.0))
+	);
+}
diff --git a/Engine/Shaders/Private/PathTracing/Volume/PathTracingCloudsMaterialCommon.ush b/Engine/Shaders/Private/PathTracing/Volume/PathTracingCloudsMaterialCommon.ush
new file mode 100644
--- /dev/null
+++ b/Engine/Shaders/Private/PathTracing/Volume/PathTracingCloudsMaterialCommon.ush
@@ -0,0 +1,72 @@
+// Copyright Epic Games, Inc. All Rights Reserved.
+
+#pragma once
+
+// Functions related to evaluating volumetric cloud materials in the path tracer
+// This code is for sharing between the cloud acceleration map building and the callable shader version of the cloud material
+
+#include "../../Common.ush"
+#include "/Engine/Generated/Material.ush"
+#include "PathTracingVolumeCommon.ush"
+#include "PathTracingCloudsCommon.ush"
+#include "../../DoubleFloat.ush"
+#include "../../SkyAtmosphereCommon.ush"
+
+
+// For path tracing, we need to be able to evaluate the material starting from a single point in space
+// In all contexts, we have a slightly more accurate Height already available so pass that in instead of measuring it from the position
+FSampleCloudMaterialResult SampleCloudMaterial(FDFVector3 AbsoluteWorldPosition, float HeightKm, float CloudBotKm, float CloudTopKm)
+{
+	FMaterialPixelParameters MaterialParameters = MakeInitializedMaterialPixelParameters();
+	float4 SvPosition = 0; // TODO: could derive this from AbsoluteWorldPosition -- but is this needed? VolumetricCloud code doesn't update it per sample
+	FPixelMaterialInputs PixelMaterialInputs = (FPixelMaterialInputs)0;
+	CalcMaterialParameters(MaterialParameters, PixelMaterialInputs, SvPosition, true);
+
+	MaterialParameters.WorldPosition_CamRelative = MaterialParameters.WorldPosition_NoOffsets_CamRelative = DFAddDemote(AbsoluteWorldPosition, PrimaryView.PreViewTranslation);
+	MaterialParameters.LWCData = MakeMaterialLWCData(MaterialParameters);
+	MaterialParameters.AbsoluteWorldPosition = MaterialParameters.WorldPosition_NoOffsets = AbsoluteWorldPosition;
+	MaterialParameters.LWCData.AbsoluteWorldPosition = DFToWS(AbsoluteWorldPosition);
+	MaterialParameters.CameraVector = normalize(PrimaryView.TranslatedWorldCameraOrigin - MaterialParameters.WorldPosition_CamRelative);
+
+	MaterialParameters.CloudSampleAltitude = HeightKm * SKY_UNIT_TO_CM;
+	MaterialParameters.CloudSampleAltitudeInLayer = (HeightKm - CloudBotKm) * SKY_UNIT_TO_CM;
+	MaterialParameters.CloudSampleNormAltitudeInLayer = saturate((HeightKm - CloudBotKm) / (CloudTopKm - CloudBotKm));
+	
+	MaterialParameters.VolumeSampleConservativeDensity = 1.0f;
+#if MATERIAL_VOLUMETRIC_ADVANCED_CONSERVATIVE_DENSITY
+	MaterialParameters.VolumeSampleConservativeDensity = GetVolumetricAdvancedMaterialOutput6(MaterialParameters);	// Evaluate conservative density
+	if (MaterialParameters.VolumeSampleConservativeDensity.x <= 0.0f)
+	{
+		return (FSampleCloudMaterialResult) 0;
+	}
+#endif
+
+	CalcPixelMaterialInputs(MaterialParameters, PixelMaterialInputs);
+
+	FSampleCloudMaterialResult Result = (FSampleCloudMaterialResult)0;
+
+#if SUBSTRATE_ENABLED
+	FSubstrateBSDF BSDF = PixelMaterialInputs.FrontMaterial.InlinedBSDF;
+	Result.Extinction = VOLUMETRICFOGCLOUD_EXTINCTION(BSDF).rgb;
+	Result.EmissiveColor = BSDF_GETEMISSIVE(BSDF).rgb;
+	Result.Albedo = VOLUMETRICFOGCLOUD_ALBEDO(BSDF).rgb;
+#else
+#if !MATERIAL_SHADINGMODEL_UNLIT
+	Result.Extinction = GetMaterialSubsurfaceDataRaw(PixelMaterialInputs).rgb;
+	Result.Albedo = GetMaterialBaseColor(PixelMaterialInputs).rgb * View.DiffuseOverrideParameter.w + View.DiffuseOverrideParameter.xyz;
+#endif
+	Result.EmissiveColor = GetMaterialEmissiveRaw(PixelMaterialInputs).rgb;
+#endif
+
+	Result.Extinction = clamp(Result.Extinction * CENTIMETER_TO_METER, 0.0f, 65000.0f);
+	Result.EmissiveColor = clamp(Result.EmissiveColor, 0.0f, 6500.0f);
+	Result.Albedo = saturate(Result.Albedo);
+
+#if MATERIAL_VOLUMETRIC_ADVANCED
+	Result.PhaseG1 = GetVolumetricAdvancedMaterialOutput0(MaterialParameters);
+	Result.PhaseG2 = GetVolumetricAdvancedMaterialOutput1(MaterialParameters);
+	Result.PhaseBlend = GetVolumetricAdvancedMaterialOutput2(MaterialParameters);
+#endif
+
+	return Result;
+}
\ No newline at end of file
diff --git a/Engine/Shaders/Private/PathTracing/Volume/PathTracingVolume.ush b/Engine/Shaders/Private/PathTracing/Volume/PathTracingVolume.ush
--- a/Engine/Shaders/Private/PathTracing/Volume/PathTracingVolume.ush
+++ b/Engine/Shaders/Private/PathTracing/Volume/PathTracingVolume.ush
@@ -4,6 +4,7 @@
 
 #include "PathTracingVolumeCommon.ush"
 #include "PathTracingAtmosphere.ush"
+#include "PathTracingClouds.ush"
 #include "PathTracingFog.ush"
 #include "PathTracingHeterogeneousVolumes.ush"
 #include "PathTracingVolumeSampling.ush"
@@ -13,7 +14,7 @@ int MaxRaymarchSteps;
 
 
 // Given an input ray, figure out where the provided volume intersects it
-FVolumeIntersectionList VolumeIntersect(float3 Origin, float3 Direction, float TMin, float TMax)
+FVolumeIntersectionList VolumeIntersect(float3 Origin, float3 Direction, float TMin, float TMax, float PathRoughness)
 {
 	// NOTE: only this call needs to be guarded by enable flags at the moment since the others can't be reached if this returns nothing
 	FVolumeIntersectionList Result = CreateEmptyVolumeIntersectionList();
@@ -21,6 +22,10 @@ FVolumeIntersectionList VolumeIntersect(float3 Origin, float3 Direction, float T
 	{
 		Result.Add(VOLUMEID_ATMOSPHERE, AtmosphereIntersect(Origin, Direction, TMin, TMax));
 	}
+	if (VolumeFlags & PATH_TRACER_VOLUME_ENABLE_CLOUDS)
+	{
+		Result.Add(VOLUMEID_CLOUDS, CloudsIntersect(Origin, Direction, TMin, TMax, PathRoughness));
+	}
 	if (VolumeFlags & PATH_TRACER_VOLUME_ENABLE_FOG)
 	{
 		Result.Add(VOLUMEID_FOG, FogIntersect(Origin, Direction, TMin, TMax));
@@ -47,6 +52,11 @@ FVolumeDensityBounds VolumeGetDensityBounds(float3 Origin, float3 Direction, FVo
 	{
 		MergeVolumeDensityBounds(Result, AtmosphereGetDensityBounds(Origin, Direction, Interval.VolumeTMin, Interval.VolumeTMax));
 	}
+	if (Interval.VolumeMask & PATH_TRACER_VOLUME_ENABLE_CLOUDS)
+	{
+		// We tracked the max density during intersection, so it is directly available here
+		MergeVolumeDensityBounds(Result, CreateVolumeDensityBound(0, Interval.CloudDensity));
+	}
 	if (Interval.VolumeMask & PATH_TRACER_VOLUME_ENABLE_FOG)
 	{
 		MergeVolumeDensityBounds(Result, FogGetDensityBounds(Origin, Direction, Interval.VolumeTMin, Interval.VolumeTMax));
@@ -69,6 +79,10 @@ FVolumeShadedResult VolumeGetDensity(float3 TranslatedWorldPos, uint VolumeMask,
 	{
 		MergeVolumeShadedResult(Result, AtmosphereGetDensity(TranslatedWorldPos), bIsCameraRay && (VolumeFlags & PATH_TRACER_VOLUME_HOLDOUT_ATMOSPHERE) != 0);
 	}
+	if (VolumeMask & PATH_TRACER_VOLUME_ENABLE_CLOUDS)
+	{
+		MergeVolumeShadedResult(Result, CloudsGetDensity(TranslatedWorldPos), bIsCameraRay && (VolumeFlags & PATH_TRACER_VOLUME_HOLDOUT_CLOUDS) != 0);
+	}
 	if (VolumeMask & PATH_TRACER_VOLUME_ENABLE_FOG)
 	{
 		MergeVolumeShadedResult(Result, FogGetDensity(TranslatedWorldPos), bIsCameraRay && (VolumeFlags & PATH_TRACER_VOLUME_HOLDOUT_FOG) != 0);
@@ -80,13 +94,8 @@ FVolumeShadedResult VolumeGetDensity(float3 TranslatedWorldPos, uint VolumeMask,
 	return Result;
 }
 
-// Reference implementation of transmittance that uses only GetDensity and GetDensityBounds
-// express the inner loop here as a macro so we can stamp down a separate copy per volume ID and avoid a nested loop during ray marching
-#define PATH_TRACER_REFERENCE_TRANSMITTANCE_LOOP(GetDensityBoundsFunc, GetDensityFunc)				\
-	/* Limit number of steps to prevent timeouts // FIXME: This biases the result! */				\
-	for (int Step = 0; Step < MaxRaymarchSteps; Step++)												\
-	{																								\
-		float3 SigmaBar = GetDensityBoundsFunc(Origin, Direction, TMin, TMax).SigmaMax;				\
+
+#define PATH_TRACER_REFERENCE_TRANSMITTANCE_LOOP_BODY(GetDensityFunc)								\
 		/* take stochastic steps along the ray to estimate transmittance (null scattering) */		\
 		float3 ColorChannelPdf = Throughput;														\
 		/* Sample the distance to the next interaction */											\
@@ -116,8 +125,26 @@ FVolumeShadedResult VolumeGetDensity(float3 TranslatedWorldPos, uint VolumeMask,
 			Throughput *= EvaluateSpectralTransmittanceMiss(TMax - TMin, SigmaBar, ColorChannelPdf).xyz;	\
 			/* exit the ray marching loop */														\
 			break;																					\
-		}																							\
-	}																								\
+		}
+
+
+// Reference implementation of transmittance that uses only GetDensity and GetDensityBounds
+// express the inner loop here as a macro so we can stamp down a separate copy per volume ID and avoid a nested loop during ray marching
+#define PATH_TRACER_REFERENCE_TRANSMITTANCE_LOOP(GetDensityBoundsFunc, GetDensityFunc)				\
+	/* Limit number of steps to prevent timeouts // FIXME: This biases the result! */				\
+	for (int Step = 0; Step < MaxRaymarchSteps; Step++)												\
+	{																								\
+		float3 SigmaBar = GetDensityBoundsFunc(Origin, Direction, TMin, TMax).SigmaMax;				\
+		PATH_TRACER_REFERENCE_TRANSMITTANCE_LOOP_BODY(GetDensityFunc);								\
+	}
+
+// Same as above, but when SigmaBar is constant for the whole loop (and defined outside)
+#define PATH_TRACER_REFERENCE_TRANSMITTANCE_LOOP_CONSTANT_BOUNDS(GetDensityFunc)					\
+	/* Limit number of steps to prevent timeouts // FIXME: This biases the result! */				\
+	for (int Step = 0; Step < MaxRaymarchSteps; Step++)												\
+	{																								\
+		PATH_TRACER_REFERENCE_TRANSMITTANCE_LOOP_BODY(GetDensityFunc);								\
+	}
 
 
 // The returned value factors in the initial throughput. This is used to avoid sampling too much if throughput reaches 0.
@@ -138,6 +165,13 @@ float3 VolumeGetTransmittance(float3 StartThroughput, float3 Origin, float3 Dire
 			PATH_TRACER_REFERENCE_TRANSMITTANCE_LOOP(AtmosphereGetDensityBounds, AtmosphereGetDensity);
 		}
 	}
+	if (Interval.VolumeMask & PATH_TRACER_VOLUME_ENABLE_CLOUDS)
+	{
+		float TMin = Interval.VolumeTMin;
+		float TMax = Interval.VolumeTMax;
+		float SigmaBar = Interval.CloudDensity;
+		PATH_TRACER_REFERENCE_TRANSMITTANCE_LOOP_CONSTANT_BOUNDS(CloudsGetDensity);
+	}
 	if (Interval.VolumeMask & PATH_TRACER_VOLUME_ENABLE_FOG)
 	{
 		float TMin = Interval.VolumeTMin;
@@ -155,7 +189,6 @@ float3 VolumeGetTransmittance(float3 StartThroughput, float3 Origin, float3 Dire
 	{
 		float TMin = Interval.VolumeTMin;
 		float TMax = Interval.VolumeTMax;
-		// NOTE: there is analytic mode in this case as the presence of heterogeneous volumes disables analytic transmittance
 		Throughput = HeterogeneousVolumesGetTransmittance(Throughput, Origin, Direction, TMin, TMax, RandSequence);
 	}
 	return Throughput;
@@ -178,6 +211,13 @@ float3 VolumeGetTransmittance(float3 StartThroughput, float3 Origin, float3 Dire
 			PATH_TRACER_REFERENCE_TRANSMITTANCE_LOOP(AtmosphereGetDensityBounds, AtmosphereGetDensity);
 		}
 	}
+	if (IntersectionList.VolumeMask & PATH_TRACER_VOLUME_ENABLE_CLOUDS)
+	{
+		float TMin = IntersectionList.VolumeTMin[VOLUMEID_CLOUDS];
+		float TMax = IntersectionList.VolumeTMax[VOLUMEID_CLOUDS];
+		float SigmaBar = IntersectionList.CloudDensity;
+		PATH_TRACER_REFERENCE_TRANSMITTANCE_LOOP_CONSTANT_BOUNDS(CloudsGetDensity);
+	}
 	if (IntersectionList.VolumeMask & PATH_TRACER_VOLUME_ENABLE_FOG)
 	{
 		float TMin = IntersectionList.VolumeTMin[VOLUMEID_FOG];
@@ -195,7 +235,6 @@ float3 VolumeGetTransmittance(float3 StartThroughput, float3 Origin, float3 Dire
 	{
 		float TMin = IntersectionList.VolumeTMin[VOLUMEID_HETEROGENEOUS_VOLUMES];
 		float TMax = IntersectionList.VolumeTMax[VOLUMEID_HETEROGENEOUS_VOLUMES];
-		// NOTE: there is analytic mode in this case as the presence of heterogeneous volumes disables analytic transmittance
 		Throughput = HeterogeneousVolumesGetTransmittance(Throughput, Origin, Direction, TMin, TMax, RandSequence);
 	}
 	return Throughput;
@@ -222,7 +261,7 @@ FVolumeTrackingResult VolumeSampleDistance(float3 PathThroughput, float3 Origin,
 	}
 	else
 	{
-		// Interval is only crossing atmo/fog
+		// Interval is only crossing atmo/clouds/fog
 		float RandValue = RandomSequence_GenerateSample1D(RandSequence);
 		TrackingResult.Distance = SampleSpectralTransmittance(RandValue, SigmaBar, PathThroughput);
 		if (TrackingResult.Distance < 0.0)
diff --git a/Engine/Shaders/Private/PathTracing/Volume/PathTracingVolumeCommon.ush b/Engine/Shaders/Private/PathTracing/Volume/PathTracingVolumeCommon.ush
--- a/Engine/Shaders/Private/PathTracing/Volume/PathTracingVolumeCommon.ush
+++ b/Engine/Shaders/Private/PathTracing/Volume/PathTracingVolumeCommon.ush
@@ -7,6 +7,7 @@ struct FVolumeIntersection
 	float VolumeTMin;
 	float VolumeTMax;
 	float BlockerHitT;
+	float MaxDensity;
 
 	bool HitVolume()
 	{
@@ -19,7 +20,7 @@ struct FVolumeIntersection
 	}
 };
 
-FVolumeIntersection CreateVolumeIntersection(float TMin, float TMax, float BlockerT = -1.0)
+FVolumeIntersection CreateVolumeIntersection(float TMin, float TMax, float BlockerT = 0.0)
 {
 	FVolumeIntersection Result = (FVolumeIntersection)0;
 	Result.VolumeTMin = TMin;
@@ -28,6 +29,15 @@ FVolumeIntersection CreateVolumeIntersection(float TMin, float TMax, float Block
 	return Result;
 }
 
+FVolumeIntersection CreateVolumeIntersectionWithDensity(float TMin, float TMax, float MaxDensity)
+{
+	FVolumeIntersection Result = (FVolumeIntersection)0;
+	Result.VolumeTMin = TMin;
+	Result.VolumeTMax = TMax;
+	Result.MaxDensity = MaxDensity;
+	return Result;
+}
+
 FVolumeIntersection CreateEmptyVolumeIntersection()
 {
 	return (FVolumeIntersection)0;
@@ -40,12 +50,14 @@ struct FVolumeIntersectionInterval
 	float VolumeTMin;
 	float VolumeTMax;
 	uint VolumeMask;
+	float CloudDensity;
 };
 
 struct FVolumeIntersectionList
 {
-	float3 VolumeTMin;
-	float3 VolumeTMax;
+	float4 VolumeTMin;
+	float4 VolumeTMax;
+	float CloudDensity;
 	uint VolumeMask;
 
 	// only need to track one blocker (assume they are opaque)
@@ -68,6 +80,10 @@ struct FVolumeIntersectionList
 		{
 			VolumeTMin[ID] = VolumeIntersection.VolumeTMin;
 			VolumeTMax[ID] = VolumeIntersection.VolumeTMax;
+			if (ID == VOLUMEID_CLOUDS)
+			{
+				CloudDensity = VolumeIntersection.MaxDensity;
+			}
 			VolumeMask |= PATH_TRACER_VOLUME_ENABLE_BIT << ID;
 		}
 		if (VolumeIntersection.HitBlocker())
@@ -124,6 +140,7 @@ struct FVolumeIntersectionList
 		}
 		// now that we established the distance, figure out which IDs overlap with the resolved segment
 		Result.VolumeMask = 0;
+		Result.CloudDensity = CloudDensity;
 		for (int Index2 = 0; Index2 < PATH_TRACER_MAX_VOLUMES; Index2++)
 		{
 			uint Mask = PATH_TRACER_VOLUME_ENABLE_BIT << Index2;
@@ -163,6 +180,7 @@ struct FVolumeIntersectionList
 			}
 		}
 		Result.BlockerHitT = BlockerHitT;
+		Result.CloudDensity = CloudDensity;
 		return Result;
 	}
 };
@@ -196,17 +214,20 @@ void MergeVolumeDensityBounds(inout FVolumeDensityBounds Merged, FVolumeDensityB
 
 
 // The result of sampling a volume at a given point
+// See PathTracingMedium.usf for the implementation of the scattering model
 struct FVolumeShadedResult
 {
 	float3 SigmaT;         // extinction
 	float3 SigmaH;         // holdout density
 	float3 SigmaSRayleigh; // Rayleight scattering coefficient
-	float3 SigmaSHG;       // Henyey-Greenstein scattering coefficient
+	float3 SigmaSHG;       // Henyey-Greenstein scattering coefficient (for atmosphere, fog and heterogeneous volumes)
+	float3 SigmaSDualHG;   // Dual Henyey-Greenstein scattering coefficient (for clouds)
 	float3 Emission;
 	float PhaseG;          // 'g' term for HG lobe
-	// TODO: second HG lobe?
+	float3 DualPhaseData;  // g1,g2,blend for DualHG lobe
 };
 
+
 FPathTracingPayload CreateMediumHitPayload(float HitT, float3 TranslatedWorldPos, FVolumeShadedResult ShadedResult)
 {
 	FPathTracingPayload Result = (FPathTracingPayload)0; // clear all fields
@@ -216,31 +237,39 @@ FPathTracingPayload CreateMediumHitPayload(float HitT, float3 TranslatedWorldPos
 	Result.BSDFOpacity = 1.0;
 	Result.PrimitiveLightingChannelMask = 7;
 	Result.SetFrontFace();
-	const float3 SigmaS = ShadedResult.SigmaSRayleigh + ShadedResult.SigmaSHG;
-	const float3 RayleighWeight = select(SigmaS > 0.0, min(ShadedResult.SigmaSRayleigh, SigmaS) / SigmaS, 0.0);
-	const float3 HGWeight = select(SigmaS > 0, min(ShadedResult.SigmaSHG, SigmaS) / SigmaS, 0.0);
+	const float3 SigmaS = ShadedResult.SigmaSRayleigh + ShadedResult.SigmaSHG + ShadedResult.SigmaSDualHG;
+	const float3 RayleighWeight = select(SigmaS > 0.0, ShadedResult.SigmaSRayleigh / SigmaS, 0.0);
+	const float3 HGWeight       = select(SigmaS > 0.0, ShadedResult.SigmaSHG       / SigmaS, 0.0);
+	const float3 DualHGWeight   = select(SigmaS > 0.0, ShadedResult.SigmaSDualHG   / SigmaS, 0.0);
 	Result.SetBaseColor(RayleighWeight);
 	Result.SetHG(HGWeight, ShadedResult.PhaseG);
+	Result.SetDualHG(DualHGWeight, ShadedResult.DualPhaseData);
 	return Result;
 }
 
-void MergeVolumeShadedResult(inout FVolumeShadedResult Merged, FVolumeShadedResult Other, bool bAsHoldout)
+// Merge the shaded result "R" into "M"
+void MergeVolumeShadedResult(inout FVolumeShadedResult M, FVolumeShadedResult R, bool bAsHoldout)
 {
 	if (bAsHoldout)
 	{
-		Merged.SigmaT += Other.SigmaT;
-		Merged.SigmaH += Other.SigmaT;
+		M.SigmaT += R.SigmaT;
+		M.SigmaH += R.SigmaT;
 	}
 	else
 	{
-		Merged.SigmaT += Other.SigmaT;
-		Merged.SigmaSRayleigh += Other.SigmaSRayleigh;
-		Merged.SigmaSHG += Other.SigmaSHG;
-		Merged.Emission += Other.Emission;
-		// blend phase function anisotropy based on amount of each medium
-		float Qa = Other.SigmaSHG.x + Other.SigmaSHG.y + Other.SigmaSHG.z;
-		float Qs = Merged.SigmaSHG.x + Merged.SigmaSHG.y + Merged.SigmaSHG.z;
-		Merged.PhaseG = lerp(Merged.PhaseG, Other.PhaseG, Qs > 0 ? Qa / Qs : 0.0);
+		M.SigmaT         += R.SigmaT;
+		M.SigmaSRayleigh += R.SigmaSRayleigh;
+		M.SigmaSHG       += R.SigmaSHG;
+		M.SigmaSDualHG   += R.SigmaSDualHG;
+		M.Emission       += R.Emission;
+		// when combining across volumes, we track dual HG on its own to simplify the logic even though it technicaly create 3 distinct hg lobes
+		// A possible improvement we could make in the future is reduce this to only 2 HG lobes with improved blending here.
+		float Q1r = R.SigmaSHG.x + R.SigmaSHG.y + R.SigmaSHG.z;
+		float Q1m = M.SigmaSHG.x + M.SigmaSHG.y + M.SigmaSHG.z;
+		float Q2r = R.SigmaSDualHG.x + R.SigmaSDualHG.y + R.SigmaSDualHG.z;
+		float Q2m = M.SigmaSDualHG.x + M.SigmaSDualHG.y + M.SigmaSDualHG.z;
+		M.PhaseG        = lerp(M.PhaseG       , R.PhaseG       , Q1r / (Q1m + 1e-6));
+		M.DualPhaseData = lerp(M.DualPhaseData, R.DualPhaseData, Q2r / (Q2m + 1e-6));
 	}
 }
 
diff --git a/Engine/Shaders/Shared/PathTracingDefinitions.h b/Engine/Shaders/Shared/PathTracingDefinitions.h
--- a/Engine/Shaders/Shared/PathTracingDefinitions.h
+++ b/Engine/Shaders/Shared/PathTracingDefinitions.h
@@ -51,18 +51,23 @@
 
 // Constants related to volumetric support
 #define VOLUMEID_ATMOSPHERE				0
-#define VOLUMEID_FOG					1
-#define VOLUMEID_HETEROGENEOUS_VOLUMES	2
-#define PATH_TRACER_MAX_VOLUMES			3
+#define VOLUMEID_CLOUDS					1
+#define VOLUMEID_FOG					2
+#define VOLUMEID_HETEROGENEOUS_VOLUMES	3
+#define PATH_TRACER_MAX_VOLUMES			4
 
 #define PATH_TRACER_VOLUME_ENABLE_BIT						(1u)
 #define PATH_TRACER_VOLUME_ENABLE_ATMOSPHERE  				(PATH_TRACER_VOLUME_ENABLE_BIT << VOLUMEID_ATMOSPHERE)
+#define PATH_TRACER_VOLUME_ENABLE_CLOUDS                    (PATH_TRACER_VOLUME_ENABLE_BIT << VOLUMEID_CLOUDS)
 #define PATH_TRACER_VOLUME_ENABLE_FOG               		(PATH_TRACER_VOLUME_ENABLE_BIT << VOLUMEID_FOG)
 #define PATH_TRACER_VOLUME_ENABLE_HETEROGENEOUS_VOLUMES 	(PATH_TRACER_VOLUME_ENABLE_BIT << VOLUMEID_HETEROGENEOUS_VOLUMES)
 
 #define PATH_TRACER_VOLUME_HOLDOUT_BIT						(PATH_TRACER_VOLUME_ENABLE_BIT << PATH_TRACER_MAX_VOLUMES)
 #define PATH_TRACER_VOLUME_HOLDOUT_ATMOSPHERE  				(PATH_TRACER_VOLUME_HOLDOUT_BIT << VOLUMEID_ATMOSPHERE)
+#define PATH_TRACER_VOLUME_HOLDOUT_CLOUDS   				(PATH_TRACER_VOLUME_HOLDOUT_BIT << VOLUMEID_CLOUDS)
 #define PATH_TRACER_VOLUME_HOLDOUT_FOG               		(PATH_TRACER_VOLUME_HOLDOUT_BIT << VOLUMEID_FOG)
 #define PATH_TRACER_VOLUME_HOLDOUT_HETEROGENEOUS_VOLUMES	(PATH_TRACER_VOLUME_HOLDOUT_BIT << VOLUMEID_HETEROGENEOUS_VOLUMES)
 
 #define PATH_TRACER_VOLUME_USE_ANALYTIC_TRANSMITTANCE       (PATH_TRACER_VOLUME_HOLDOUT_BIT << PATH_TRACER_MAX_VOLUMES)
+#define PATH_TRACER_VOLUME_USE_CLOUD_MAP                    (PATH_TRACER_VOLUME_USE_ANALYTIC_TRANSMITTANCE << 1)
+#define PATH_TRACER_VOLUME_VIEW_CLOUD_MAP_ACCEL             (PATH_TRACER_VOLUME_USE_ANALYTIC_TRANSMITTANCE << 2)
diff --git a/Engine/Source/Runtime/Renderer/Private/DeferredShadingRenderer.cpp b/Engine/Source/Runtime/Renderer/Private/DeferredShadingRenderer.cpp
--- a/Engine/Source/Runtime/Renderer/Private/DeferredShadingRenderer.cpp
+++ b/Engine/Source/Runtime/Renderer/Private/DeferredShadingRenderer.cpp
@@ -1481,7 +1481,7 @@ void FDeferredShadingSceneRenderer::Render(FRDGBuilder& GraphBuilder)
 			// 1. Update sky atmosphere
 			// This needs to be done prior to start Lumen scene lighting to ensure directional light color is correct, as the sun color needs atmosphere transmittance
 			{
-				const bool bPathTracedAtmosphere = ViewFamily.EngineShowFlags.PathTracing && Views.Num() > 0 && Views[0].FinalPostProcessSettings.PathTracingEnableReferenceAtmosphere;
+				const bool bPathTracedAtmosphere = ViewFamily.EngineShowFlags.PathTracing && Views.Num() > 0 && PathTracing::UsesReferenceAtmosphere(Views[0]);
 				if (ShouldRenderSkyAtmosphere(Scene, ViewFamily.EngineShowFlags) && !bPathTracedAtmosphere)
 				{
 					for (int32 LightIndex = 0; LightIndex < NUM_ATMOSPHERE_LIGHTS; ++LightIndex)
@@ -1741,6 +1741,13 @@ void FDeferredShadingSceneRenderer::Render(FRDGBuilder& GraphBuilder)
 			}
 		}
 
+		if (ViewFamily.EngineShowFlags.PathTracing)
+		{
+			// If we might be path tracing the clouds -- call the path tracer's method for cloud callable shader setup
+			// this will skip work if cloud rendering is not being used
+			PreparePathTracingCloudMaterial(Scene, Views);
+		}
+
 		if (IsRayTracingEnabled(ViewFamily.GetShaderPlatform()) && RHISupportsRayTracingShaders(ViewFamily.GetShaderPlatform()))
 		{
 			// Nanite raytracing manager update must run before GPUScene update since it can modify primitive data
@@ -2147,8 +2154,8 @@ void FDeferredShadingSceneRenderer::Render(FRDGBuilder& GraphBuilder)
 	const bool bShouldRenderVolumetricCloudBase = ShouldRenderVolumetricCloud(Scene, ViewFamily.EngineShowFlags);
 	const bool bShouldRenderVolumetricCloud = bShouldRenderVolumetricCloudBase && (!ViewFamily.EngineShowFlags.VisualizeVolumetricCloudConservativeDensity && !ViewFamily.EngineShowFlags.VisualizeVolumetricCloudEmptySpaceSkipping);
 	const bool bShouldVisualizeVolumetricCloud = bShouldRenderVolumetricCloudBase && (!!ViewFamily.EngineShowFlags.VisualizeVolumetricCloudConservativeDensity || !!ViewFamily.EngineShowFlags.VisualizeVolumetricCloudEmptySpaceSkipping);
-	bool bAsyncComputeVolumetricCloud = IsVolumetricRenderTargetEnabled() && IsVolumetricRenderTargetAsyncCompute();
-	bool bVolumetricRenderTargetRequired = bShouldRenderVolumetricCloud && !bHasRayTracedOverlay;
+	const bool bAsyncComputeVolumetricCloud = IsVolumetricRenderTargetEnabled() && IsVolumetricRenderTargetAsyncCompute();
+	const bool bVolumetricRenderTargetRequired = bShouldRenderVolumetricCloud && !bHasRayTracedOverlay;
 
 	FRDGTextureRef ViewFamilyTexture = TryCreateViewFamilyTexture(GraphBuilder, ViewFamily);
 	FRDGTextureRef ViewFamilyDepthTexture = TryCreateViewFamilyDepthTexture(GraphBuilder, ViewFamily);
diff --git a/Engine/Source/Runtime/Renderer/Private/PathTracing.cpp b/Engine/Source/Runtime/Renderer/Private/PathTracing.cpp
--- a/Engine/Source/Runtime/Renderer/Private/PathTracing.cpp
+++ b/Engine/Source/Runtime/Renderer/Private/PathTracing.cpp
@@ -26,16 +26,18 @@ TAutoConsoleVariable<int32> CVarPathTracing(
 #include "RayTracing/RayTracingMaterialHitShaders.h"
 #include "RayTracing/RayTracingDecals.h"
 #include "DecalRenderingCommon.h"
+#include "VolumetricCloudProxy.h"
+#include "MeshPassUtils.h"
 #include "FogRendering.h"
 #include "GenerateMips.h"
 #include "HairStrands/HairStrandsData.h"
 #include "HeterogeneousVolumes/HeterogeneousVolumes.h"
 #include "Modules/ModuleManager.h"
 #include "SkyAtmosphereRendering.h"
-#include <limits>
 #include "PathTracingSpatialTemporalDenoising.h"
 #include "PostProcess/DiaphragmDOF.h"
 #include "EnvironmentComponentsFlags.h"
+#include <limits>
 
 TAutoConsoleVariable<int32> CVarPathTracingExperimental(
 	TEXT("r.PathTracing.Experimental"),
@@ -118,7 +120,7 @@ TAutoConsoleVariable<int32> CVarPathTracingVolumeMISMode(
 
 TAutoConsoleVariable<int32> CVarPathTracingMaxRaymarchSteps(
 	TEXT("r.PathTracing.MaxRaymarchSteps"),
-	256,
+	768,
 	TEXT("Upper limit on the number of ray marching steps in volumes. This limit should not be hit in most cases, but raising it can reduce bias in case it is. (default = 256)."),
 	ECVF_RenderThreadSafe
 );
@@ -192,6 +194,25 @@ TAutoConsoleVariable<int32> CVarPathTracingEnableCameraBackfaceCulling(
 	ECVF_RenderThreadSafe
 );
 
+TAutoConsoleVariable<int32> CVarPathTracingEnableReferenceAtmosphere(
+	TEXT("r.PathTracing.EnableReferenceAtmosphere"),
+	-1,
+	TEXT("Should the path tracer use a volumetric calculation to represent the sky atmosphere?\n")
+	TEXT("-1: Inherit from PostProcess settings (default)\n")
+	TEXT(" 0: Disabled\n")
+	TEXT(" 1: Enable\n"),
+	ECVF_RenderThreadSafe
+);
+
+TAutoConsoleVariable<int32> CVarPathTracingEnableReferenceClouds(
+	TEXT("r.PathTracing.EnableReferenceClouds"),
+	0,
+	TEXT("Should the path tracer use a volumetric calculation to represent volumetric clouds? (This requires Reference Atmosphere to be enabled)\n")
+	TEXT(" 0: Disabled (default)\n")
+	TEXT(" 1: Enable\n"),
+	ECVF_RenderThreadSafe
+);
+
 TAutoConsoleVariable<int32> CVarPathTracingAtmosphereOpticalDepthLutResolution(
 	TEXT("r.PathTracing.AtmosphereOpticalDepthLUTResolution"),
 	512,
@@ -206,6 +227,55 @@ TAutoConsoleVariable<int32> CVarPathTracingAtmosphereOpticalDepthLutNumSamples(
 	ECVF_RenderThreadSafe
 );
 
+TAutoConsoleVariable<int32> CVarPathTracingCloudAccelerationMapResolution(
+	TEXT("r.PathTracing.CloudAccelerationMapResolution"),
+	512,
+	TEXT("Size of the square texture used to accelerate cloud ray marching for the path tracer in reference atmosphere mode.  (default = 512)"),
+	ECVF_RenderThreadSafe
+);
+
+TAutoConsoleVariable<int32> CVarPathTracingCloudAccelerationMapNumSamples(
+	TEXT("r.PathTracing.CloudAccelerationMapNumSamples"),
+	64,
+	TEXT("Number of ray marching samples used when building the cloud acceleration map for the path tracer in reference atmosphere mode.  (default = 64)"),
+	ECVF_RenderThreadSafe
+);
+
+TAutoConsoleVariable<bool> CVarPathTracingCloudAccelerationMapVisualize(
+	TEXT("r.PathTracing.CloudAccelerationMapVisualize"),
+	false,
+	TEXT("If true, replace clouds with a visualization of the acceleration map to help visualize it and fine tune its resolution (default = false)"),
+	ECVF_RenderThreadSafe
+);
+
+TAutoConsoleVariable<float> CVarPathTracingCloudRoughnessCutoff(
+	TEXT("r.PathTracing.CloudRoughnessCutoff"),
+	0.05f,
+	TEXT("Do not evaluate volumetric clouds beyond this roughness level to improve performance (default=0.05)"),
+	ECVF_RenderThreadSafe
+);
+
+TAutoConsoleVariable<bool> CVarPathTracingCloudMapEnabled(
+	TEXT("r.PathTracing.CloudMapEnable"),
+	true,
+	TEXT("If true, clouds will be voxelized into a texture for faster evaluation. If false, the cloud material will be invoked via callable shaders. (default = true)"),
+	ECVF_RenderThreadSafe
+);
+
+TAutoConsoleVariable<int32> CVarPathTracingCloudMapResolution(
+	TEXT("r.PathTracing.CloudMapResolution"),
+	512,
+	TEXT("Size of the texture used to cache cloud appearance in reference atmosphere mode.  (default = 512)"),
+	ECVF_RenderThreadSafe
+);
+
+TAutoConsoleVariable<int32> CVarPathTracingCloudMapDepth(
+	TEXT("r.PathTracing.CloudMapDepth"),
+	64,
+	TEXT("Depth of the texture used to cache cloud appearance in reference atmosphere mode.  (default = 64)"),
+	ECVF_RenderThreadSafe
+);
+
 TAutoConsoleVariable<int32> CVarPathTracingFrameIndependentTemporalSeed(
 	TEXT("r.PathTracing.FrameIndependentTemporalSeed"),
 	1,
@@ -382,13 +452,12 @@ TAutoConsoleVariable<int32> CVarpathTracingOverrideDepth(
 	ECVF_RenderThreadSafe
 );
 
-TAutoConsoleVariable<int32> CVarPathTracingUseAnalyticTransmittance(
+TAutoConsoleVariable<bool> CVarPathTracingUseAnalyticTransmittance(
 	TEXT("r.PathTracing.UseAnalyticTransmittance"),
-	-1,
+	true,
 	TEXT("Determines use of analytical or null-tracking estimation when evaluating transmittance\n")
-	TEXT("-1: uses null-tracking estimation if heterogeneous volumes are present, or analytical estimation otherwise (default)\n")
-	TEXT("0: off (uses null-tracking estimation, instead)\n")
-	TEXT("1: on (uses analytical estimation)\n"),
+	TEXT("0: off (uses null-tracking estimation)\n")
+	TEXT("1: on (uses analytical estimation when possible) (default)\n"),
 	ECVF_RenderThreadSafe
 );
 
@@ -443,6 +512,7 @@ BEGIN_SHADER_PARAMETER_STRUCT(FPathTracingData, )
 	SHADER_PARAMETER(float, MaxPathIntensity)
 	SHADER_PARAMETER(float, MaxNormalBias)
 	SHADER_PARAMETER(float, FilterWidth)
+	SHADER_PARAMETER(float, CloudRoughnessCutoff)
 	SHADER_PARAMETER(float, DecalRoughnessCutoff)
 	SHADER_PARAMETER(float, MeshDecalRoughnessCutoff)
 	SHADER_PARAMETER(float, MeshDecalBias)
@@ -467,6 +537,10 @@ struct FPathTracingConfig
 	bool UseMultiGPU; // NOTE: Requires invalidation because the buffer layout changes
 	int DenoiserMode; // NOTE: does not require path tracing invalidation
 	float AdaptiveSamplingThreshold;
+	int CloudAccelerationMapNumSamples;
+	int CloudAccelerationMapResolution;
+	int CloudMapResolution;
+	int CloudMapDepth;
 
 	bool IsDifferent(const FPathTracingConfig& Other) const
 	{
@@ -490,6 +564,7 @@ struct FPathTracingConfig
 			PathTracingData.ApplyDiffuseSpecularOverrides != Other.PathTracingData.ApplyDiffuseSpecularOverrides ||
 			PathTracingData.EnabledDirectLightingContributions != Other.PathTracingData.EnabledDirectLightingContributions ||
 			PathTracingData.EnabledIndirectLightingContributions != Other.PathTracingData.EnabledIndirectLightingContributions ||
+			PathTracingData.CloudRoughnessCutoff != Other.PathTracingData.CloudRoughnessCutoff ||
 			PathTracingData.DecalRoughnessCutoff != Other.PathTracingData.DecalRoughnessCutoff ||
 			PathTracingData.MeshDecalRoughnessCutoff != Other.PathTracingData.MeshDecalRoughnessCutoff ||
 			PathTracingData.MeshDecalBias != Other.PathTracingData.MeshDecalBias ||
@@ -504,6 +579,10 @@ struct FPathTracingConfig
 			UseCameraMediumTracking != Other.UseCameraMediumTracking ||
 			UseAdaptiveSampling != Other.UseAdaptiveSampling ||
 			AdaptiveSamplingThreshold != Other.AdaptiveSamplingThreshold ||
+			CloudAccelerationMapNumSamples != Other.CloudAccelerationMapNumSamples ||
+			CloudAccelerationMapResolution != Other.CloudAccelerationMapResolution ||
+			CloudMapResolution != Other.CloudMapResolution ||
+			CloudMapDepth != Other.CloudMapDepth ||
 			UseMultiGPU != Other.UseMultiGPU;
 	}
 
@@ -561,6 +640,8 @@ struct FPathTracingState {
 	TRefCountPtr<IPooledRenderTarget> AlbedoRT;
 	TRefCountPtr<IPooledRenderTarget> NormalRT;
 	TRefCountPtr<FRDGPooledBuffer> VarianceBuffer;
+	TRefCountPtr<IPooledRenderTarget> CloudAccelerationMap;
+	TRefCountPtr<IPooledRenderTarget> CloudMap;
 
 	// Cache to improve the stability when frame denoising (SPP=r.pathtracing.SamplesPerPixel) is used in animation rendering
 	TRefCountPtr<IPooledRenderTarget> LastDenoisedRadianceRT;
@@ -597,17 +678,12 @@ namespace PathTracing
 	{
 		return ViewFamily.EngineShowFlags.Decals;
 	}
-}
 
-static bool EvalUseAnalyticTransmittance(const FViewInfo& View)
-{
-	int32 UseAnalyticTransmittance = CVarPathTracingUseAnalyticTransmittance.GetValueOnRenderThread();
-	if (UseAnalyticTransmittance < 0)
+	bool UsesReferenceAtmosphere(const FViewInfo& View)
 	{
-		UseAnalyticTransmittance = !ShouldRenderHeterogeneousVolumesForView(View);
+		const int32 EnableReferenceAtmosphereCVar = CVarPathTracingEnableReferenceAtmosphere.GetValueOnRenderThread();
+		return EnableReferenceAtmosphereCVar < 0 ? View.FinalPostProcessSettings.PathTracingEnableReferenceAtmosphere != 0 : EnableReferenceAtmosphereCVar != 0;
 	}
-
-	return UseAnalyticTransmittance != 0;
 }
 
 // This function prepares the portion of shader arguments that may involve invalidating the path traced state
@@ -653,12 +729,21 @@ static void PreparePathTracingData(const FScene* Scene, const FViewInfo& View, F
 		PathTracingData.CameraLensRadius    = CocModel.GetLensRadius();
 	}
 
+	const bool bUseReferenceAtmosphere = ShouldRenderSkyAtmosphere(Scene, ShowFlags) && 
+		View.SkyAtmosphereUniformShaderParameters != nullptr &&
+		PathTracing::UsesReferenceAtmosphere(View);
+	
+	const FVolumetricCloudRenderSceneInfo* CloudInfo = Scene->GetVolumetricCloudSceneInfo();
+	const bool bVolumeCloudMapEnabled = !RHISupportsRayTracingCallableShaders(View.GetShaderPlatform()) || CVarPathTracingCloudMapEnabled.GetValueOnRenderThread();
+	const bool bVolumeCloudsVisible = ShouldRenderVolumetricCloud(Scene, ShowFlags)
+		&& bUseReferenceAtmosphere
+		&& CVarPathTracingEnableReferenceClouds.GetValueOnRenderThread() != 0
+		&& (bVolumeCloudMapEnabled || uint32(View.PathTracingVolumetricCloudCallableShaderIndex) < Scene->RayTracingScene.NumCallableShaderSlots);
+
 	// Merge all volume flags into one uint
 	PathTracingData.VolumeFlags = 0;
-	PathTracingData.VolumeFlags |=
-		ShouldRenderSkyAtmosphere(Scene, ShowFlags) && 
-		View.SkyAtmosphereUniformShaderParameters != nullptr &&
-		PPV.PathTracingEnableReferenceAtmosphere != 0 ? PATH_TRACER_VOLUME_ENABLE_ATMOSPHERE : 0;
+	PathTracingData.VolumeFlags |= bUseReferenceAtmosphere ? PATH_TRACER_VOLUME_ENABLE_ATMOSPHERE : 0;
+	PathTracingData.VolumeFlags |= bVolumeCloudsVisible ? PATH_TRACER_VOLUME_ENABLE_CLOUDS : 0;
 	PathTracingData.VolumeFlags |= ShouldRenderFog(*View.Family)
 		&& Scene->ExponentialFogs.Num() > 0
 		&& Scene->ExponentialFogs[0].bEnableVolumetricFog
@@ -668,8 +753,13 @@ static void PreparePathTracingData(const FScene* Scene, const FViewInfo& View, F
 			Scene->ExponentialFogs[0].FogData[1].Density > 0) ? PATH_TRACER_VOLUME_ENABLE_FOG : 0;
 	PathTracingData.VolumeFlags |= ShouldRenderHeterogeneousVolumesForView(View) ? PATH_TRACER_VOLUME_ENABLE_HETEROGENEOUS_VOLUMES : 0;
 	PathTracingData.VolumeFlags |= View.SkyAtmosphereUniformShaderParameters != nullptr && IsSkyAtmosphereHoldout(View.CachedViewUniformShaderParameters->EnvironmentComponentsFlags) ? PATH_TRACER_VOLUME_HOLDOUT_ATMOSPHERE : 0;
+	PathTracingData.VolumeFlags |= bVolumeCloudsVisible && IsVolumetricCloudHoldout(View.CachedViewUniformShaderParameters->EnvironmentComponentsFlags) ? PATH_TRACER_VOLUME_HOLDOUT_CLOUDS : 0;
 	PathTracingData.VolumeFlags |= Scene->ExponentialFogs.Num() > 0 && IsExponentialFogHoldout(View.CachedViewUniformShaderParameters->EnvironmentComponentsFlags) ? PATH_TRACER_VOLUME_HOLDOUT_FOG : 0;
-	PathTracingData.VolumeFlags |= EvalUseAnalyticTransmittance(View) ? PATH_TRACER_VOLUME_USE_ANALYTIC_TRANSMITTANCE : 0;
+	PathTracingData.VolumeFlags |= CVarPathTracingUseAnalyticTransmittance.GetValueOnRenderThread() ? PATH_TRACER_VOLUME_USE_ANALYTIC_TRANSMITTANCE : 0;
+	PathTracingData.VolumeFlags |= bVolumeCloudMapEnabled ? PATH_TRACER_VOLUME_USE_CLOUD_MAP : 0;
+	PathTracingData.VolumeFlags |= CVarPathTracingCloudAccelerationMapVisualize.GetValueOnRenderThread() ? PATH_TRACER_VOLUME_VIEW_CLOUD_MAP_ACCEL : 0;
+
+	PathTracingData.CloudRoughnessCutoff = bVolumeCloudsVisible ? CVarPathTracingCloudRoughnessCutoff.GetValueOnRenderThread() : -1.0f;
 
 	PathTracingData.EnableDBuffer = CVarPathTracingUseDBuffer.GetValueOnRenderThread();
 
@@ -886,6 +976,32 @@ static FPathTracingFogParameters PrepareFogParameters(const FViewInfo& View, con
 	return Parameters;
 }
 
+
+BEGIN_SHADER_PARAMETER_STRUCT(FPathTracingCloudParameters, )
+	// coordinate frame for the cloud acceleration map
+	SHADER_PARAMETER(FVector3f, CloudClipX)	// Right
+	SHADER_PARAMETER(FVector3f, CloudClipY) // Forward
+	SHADER_PARAMETER(FVector3f, CloudClipZ) // Up
+
+	SHADER_PARAMETER(FVector3f, CloudClipCenterKm) // Planet center in Km
+
+	SHADER_PARAMETER(float, CloudLayerBotKm)
+	SHADER_PARAMETER(float, CloudLayerTopKm)
+	SHADER_PARAMETER(float, CloudClipDistKm) // distance in x,y in cloud clip space
+	SHADER_PARAMETER(float, CloudClipRadiusKm) // distance from origin to planet center in Km
+
+	SHADER_PARAMETER(float, CloudTracingMaxDistance) // limit ray lengths (to avoid slowing down when a ray crosses all clouds)
+	SHADER_PARAMETER(int32, CloudAccelMapResolution)
+	SHADER_PARAMETER(int32, CloudCallableShaderId)
+	SHADER_PARAMETER(FIntVector, CloudMapResolution)
+END_SHADER_PARAMETER_STRUCT()
+
+BEGIN_GLOBAL_SHADER_PARAMETER_STRUCT(FPathTracingCloudParameterGlobals,)
+	SHADER_PARAMETER_STRUCT_INCLUDE(FPathTracingCloudParameters, CloudParameters)
+END_GLOBAL_SHADER_PARAMETER_STRUCT()
+
+IMPLEMENT_GLOBAL_SHADER_PARAMETER_STRUCT(FPathTracingCloudParameterGlobals, "PathTracingCloudParameters");
+
 static uint32 GetPathtracingMaterialPayloadSize()
 {
 	// Substrate uses a slightly bigger payload as the basic slab contains more information
@@ -973,6 +1089,12 @@ class FPathTracingRG : public FGlobalShader
 		SHADER_PARAMETER(FVector3f, PlanetCenterTranslatedWorldHi)
 		SHADER_PARAMETER(FVector3f, PlanetCenterTranslatedWorldLo)
 
+		// clouds
+		SHADER_PARAMETER_STRUCT_INCLUDE(FPathTracingCloudParameters, CloudParameters)
+		SHADER_PARAMETER_RDG_TEXTURE(Texture2D, CloudAccelerationMap)
+		SHADER_PARAMETER_SAMPLER(SamplerState, CloudAccelerationMapSampler)
+		SHADER_PARAMETER_RDG_TEXTURE(Texture3D, CloudMap)
+
 		// exponential height fog
 		SHADER_PARAMETER_STRUCT_INCLUDE(FPathTracingFogParameters, FogParameters)
 
@@ -1082,6 +1204,293 @@ class FPathTracingBuildAtmosphereOpticalDepthLUTCS : public FGlobalShader
 IMPLEMENT_SHADER_TYPE(, FPathTracingBuildAtmosphereOpticalDepthLUTCS, TEXT("/Engine/Private/PathTracing/PathTracingBuildAtmosphereLUT.usf"), TEXT("PathTracingBuildAtmosphereOpticalDepthLUTCS"), SF_Compute);
 
 
+FPathTracingCloudParameters PrepareCloudParameters(const FScene* Scene, const FViewInfo& View, const FPathTracingConfig* Config)
+{
+	check(Scene != nullptr);;
+	check(Scene->GetVolumetricCloudSceneInfo() != nullptr);
+
+
+	const FSkyAtmosphereRenderSceneInfo* SkyInfo = Scene->GetSkyAtmosphereSceneInfo();
+	const FVolumetricCloudRenderSceneInfo* CloudInfo = Scene->GetVolumetricCloudSceneInfo();
+	const FVolumetricCloudSceneProxy& CloudProxy = CloudInfo->GetVolumetricCloudSceneProxy();
+
+	float PlanetRadiusKm = CloudProxy.PlanetRadiusKm;
+	FVector CloudCenterKm = FVector(0, 0, -CloudProxy.PlanetRadiusKm);
+	if (SkyInfo != nullptr)
+	{
+		const FAtmosphereSetup& AtmosphereSetup = SkyInfo->GetSkyAtmosphereSceneProxy().GetAtmosphereSetup();
+		PlanetRadiusKm = AtmosphereSetup.BottomRadiusKm;
+		CloudCenterKm = AtmosphereSetup.PlanetCenterKm;
+	}
+
+	FVector PlanetUp = View.ViewMatrices.GetViewOrigin() - CloudCenterKm * double(FAtmosphereSetup::SkyUnitToCm);
+	double ViewToPlanet = PlanetUp.Length();
+	PlanetUp.Normalize();
+
+	FPathTracingCloudParameters Params = {};
+
+	// Make a coordinate frame for our cloud acceleration map -- we want it to be stable with respect
+	// to camera rotations to minimize aliasing as the camera moves. PlanetUp will generally be quite
+	// stable when moving about the planet surface, so using this as the only input minimizes resampling artifacts
+	// TODO: Figure out a stable scheme for views from space ...
+	// See GetTangentBasis() in MonteCarlo.ush
+	// TODO: Should probably be turned into a utility on TVector?
+	{
+		const FVector TangentZ = PlanetUp;
+		const double Sign = TangentZ.Z >= 0 ? 1 : -1;
+		const double a = -1 / (Sign + TangentZ.Z);
+		const double b = TangentZ.X * TangentZ.Y * a;
+		FVector TangentX = { 1 + Sign * a * (TangentZ.X * TangentZ.X), Sign * b, -Sign * TangentZ.X };
+		FVector TangentY = { b,  Sign + a * (TangentZ.Y * TangentZ.Y), -TangentZ.Y };
+
+		Params.CloudClipX = FVector3f(TangentX);
+		Params.CloudClipY = FVector3f(TangentY);
+		Params.CloudClipZ = FVector3f(TangentZ);
+	}
+	Params.CloudClipCenterKm = FVector3f(CloudCenterKm); // LWC_TODO: Pass this in as high/low for better precision
+
+	Params.CloudLayerBotKm = PlanetRadiusKm + CloudProxy.LayerBottomAltitudeKm;
+	Params.CloudLayerTopKm = Params.CloudLayerBotKm + CloudProxy.LayerHeightKm;
+	if (CloudProxy.TracingMaxDistanceMode == 0)
+	{
+		Params.CloudClipDistKm = FMath::Min(CloudProxy.TracingStartMaxDistance, Params.CloudLayerTopKm);
+		Params.CloudTracingMaxDistance = CloudProxy.TracingMaxDistance * FAtmosphereSetup::SkyUnitToCm;
+	}
+	else
+	{
+		Params.CloudClipDistKm = FMath::Min(FMath::Min(CloudProxy.TracingStartMaxDistance, CloudProxy.TracingMaxDistance), Params.CloudLayerTopKm);
+		Params.CloudTracingMaxDistance = 2 * Params.CloudClipDistKm * FAtmosphereSetup::SkyUnitToCm; // full diagonal for this
+	}
+	
+	Params.CloudClipRadiusKm = float(ViewToPlanet * double(FAtmosphereSetup::CmToSkyUnit));
+	Params.CloudCallableShaderId = -1;
+	if (Config != nullptr)
+	{
+		Params.CloudAccelMapResolution = Config->CloudAccelerationMapResolution;
+		Params.CloudMapResolution = FIntVector(Config->CloudMapResolution, Config->CloudMapResolution, Config->CloudMapDepth);
+	}
+	return Params;
+}
+
+class FPathTracingBuildCloudAccelerationMapCS : public FMeshMaterialShader
+{
+	DECLARE_SHADER_TYPE(FPathTracingBuildCloudAccelerationMapCS, MeshMaterial)
+	SHADER_USE_PARAMETER_STRUCT_WITH_LEGACY_BASE(FPathTracingBuildCloudAccelerationMapCS, FMeshMaterialShader) // TODO: following pattern used in VolumetricCloudRender.cpp -- is there a proper modern alternative?
+
+	static bool ShouldCompilePermutation(const FMeshMaterialShaderPermutationParameters& Parameters)
+	{
+		return ShouldCompilePathTracingShadersForProject(Parameters.Platform) &&
+			Parameters.MaterialParameters.bIsUsedWithVolumetricCloud &&
+			Parameters.MaterialParameters.MaterialDomain == MD_Volume;
+	}
+
+	static void ModifyCompilationEnvironment(const FMaterialShaderPermutationParameters& Parameters, FShaderCompilerEnvironment& OutEnvironment)
+	{
+		FMaterialShader::ModifyCompilationEnvironment(Parameters, OutEnvironment);
+
+		OutEnvironment.SetDefine(TEXT("THREADGROUPSIZE_X"), FComputeShaderUtils::kGolden2DGroupSize);
+		OutEnvironment.SetDefine(TEXT("THREADGROUPSIZE_Y"), FComputeShaderUtils::kGolden2DGroupSize);
+		OutEnvironment.SetDefine(TEXT("CLOUD_LAYER_PIXEL_SHADER"), 1);
+	}
+
+	BEGIN_SHADER_PARAMETER_STRUCT(FParameters, )
+		SHADER_PARAMETER(uint32, NumSamples)
+		SHADER_PARAMETER(uint32, Iteration)
+		SHADER_PARAMETER(uint32, TemporalSeed)
+		SHADER_PARAMETER(int32 , UseCloudMap)
+		SHADER_PARAMETER_STRUCT_INCLUDE(FPathTracingCloudParameters, CloudParameters)
+		SHADER_PARAMETER_STRUCT_REF(FViewUniformShaderParameters, ViewUniformBuffer)
+		SHADER_PARAMETER_RDG_TEXTURE_UAV(RWTexture2D, CloudAccelerationMap)
+		SHADER_PARAMETER_RDG_TEXTURE(Texture3D, CloudMap)
+	END_SHADER_PARAMETER_STRUCT()
+};
+IMPLEMENT_MATERIAL_SHADER_TYPE(, FPathTracingBuildCloudAccelerationMapCS, TEXT("/Engine/Private/PathTracing/PathTracingBuildCloudAccelerationMap.usf"), TEXT("PathTracingBuildCloudAccelerationMapCS"), SF_Compute);
+
+class FPathTracingBuildCloudCS : public FMeshMaterialShader
+{
+	DECLARE_SHADER_TYPE(FPathTracingBuildCloudCS, MeshMaterial)
+	SHADER_USE_PARAMETER_STRUCT_WITH_LEGACY_BASE(FPathTracingBuildCloudCS, FMeshMaterialShader) // TODO: following pattern used in VolumetricCloudRender.cpp -- is there a proper modern alternative?
+
+	static bool ShouldCompilePermutation(const FMeshMaterialShaderPermutationParameters& Parameters)
+	{
+		return ShouldCompilePathTracingShadersForProject(Parameters.Platform) &&
+		Parameters.MaterialParameters.bIsUsedWithVolumetricCloud &&
+		Parameters.MaterialParameters.MaterialDomain == MD_Volume;
+	}
+
+	static void ModifyCompilationEnvironment(const FMaterialShaderPermutationParameters& Parameters, FShaderCompilerEnvironment& OutEnvironment)
+	{
+		FMaterialShader::ModifyCompilationEnvironment(Parameters, OutEnvironment);
+
+		OutEnvironment.SetDefine(TEXT("THREADGROUPSIZE_X"), 4);
+		OutEnvironment.SetDefine(TEXT("THREADGROUPSIZE_Y"), 4);
+		OutEnvironment.SetDefine(TEXT("THREADGROUPSIZE_Z"), 4);
+		OutEnvironment.SetDefine(TEXT("CLOUD_LAYER_PIXEL_SHADER"), 1);
+	}
+
+	BEGIN_SHADER_PARAMETER_STRUCT(FParameters, )
+		SHADER_PARAMETER(float , BlendFactor)
+		SHADER_PARAMETER(uint32, Iteration)
+		SHADER_PARAMETER(uint32, TemporalSeed)
+		SHADER_PARAMETER_STRUCT_INCLUDE(FPathTracingCloudParameters, CloudParameters)
+		SHADER_PARAMETER_STRUCT_REF(FViewUniformShaderParameters, ViewUniformBuffer)
+		SHADER_PARAMETER_RDG_TEXTURE_UAV(RWTexture3D, CloudMap)
+	END_SHADER_PARAMETER_STRUCT()
+};
+IMPLEMENT_MATERIAL_SHADER_TYPE(, FPathTracingBuildCloudCS, TEXT("/Engine/Private/PathTracing/PathTracingBuildCloud.usf"), TEXT("PathTracingBuildCloudCS"), SF_Compute);
+
+
+class FPathTracingVolumetricCloudMaterial : public FMaterialShader
+{
+	DECLARE_SHADER_TYPE(FPathTracingVolumetricCloudMaterial, Material);
+public:
+	FPathTracingVolumetricCloudMaterial() = default;
+
+	FPathTracingVolumetricCloudMaterial(const ShaderMetaType::CompiledShaderInitializerType& Initializer)
+	: FMaterialShader(Initializer)
+	{
+		CloudParameter.Bind(Initializer.ParameterMap, TEXT("PathTracingCloudParameters"));
+	}
+
+	static bool ShouldCompilePermutation(const FMaterialShaderPermutationParameters& Parameters)
+	{
+		return ShouldCompilePathTracingShadersForProject(Parameters.Platform) &&
+			ShouldCompileRayTracingCallableShadersForProject(Parameters.Platform) &&
+			Parameters.MaterialParameters.bIsUsedWithVolumetricCloud &&
+			Parameters.MaterialParameters.MaterialDomain == MD_Volume;
+	}
+
+	static void ModifyCompilationEnvironment(const FMaterialShaderPermutationParameters& Parameters, FShaderCompilerEnvironment& OutEnvironment)
+	{
+		FMeshMaterialShader::ModifyCompilationEnvironment(Parameters, OutEnvironment);
+		check(Parameters.MaterialParameters.MaterialDomain == MD_Volume);
+		OutEnvironment.SetDefine(TEXT("CLOUD_LAYER_PIXEL_SHADER"), 1);
+	}
+
+	static bool ValidateCompiledResult(EShaderPlatform Platform, const FShaderParameterMap& ParameterMap, TArray<FString>& OutError)
+	{
+		if (ParameterMap.ContainsParameterAllocation(FSceneTextureUniformParameters::FTypeInfo::GetStructMetadata()->GetShaderVariableName()))
+		{
+			OutError.Add(TEXT("Ray tracing callable shaders cannot read from the SceneTexturesStruct."));
+			return false;
+		}
+
+		for (const auto& It : ParameterMap.GetParameterMap())
+		{
+			const FParameterAllocation& ParamAllocation = It.Value;
+			if (ParamAllocation.Type != EShaderParameterType::UniformBuffer
+				&& ParamAllocation.Type != EShaderParameterType::LooseData)
+			{
+				OutError.Add(FString::Printf(TEXT("Invalid ray tracing shader parameter '%s'. Only uniform buffers and loose data parameters are supported."), *(It.Key)));
+				return false;
+			}
+		}
+
+		return true;
+	}
+
+	static ERayTracingPayloadType GetRayTracingPayloadType(const int32 PermutationId)
+	{
+		// TODO: This isn't the payload we use, but the logic in RayTracingMaterialHitShaders.cpp needs to assume a consistent payload ID for all callable shaders
+		// The simplest solution is probably to remove FDecalShaderPayload and use the PT payload everywhere
+		return ERayTracingPayloadType::Decals;
+	}
+
+	void GetShaderBindings(
+		const FScene* Scene,
+		ERHIFeatureLevel::Type FeatureLevel,
+		const FMaterialRenderProxy& MaterialRenderProxy,
+		const FMaterial& Material,
+		const FViewInfo& View,
+		const TUniformBufferRef<FPathTracingCloudParameterGlobals>& CloudParameters,
+		FMeshDrawSingleShaderBindings& ShaderBindings) const
+	{
+		FMaterialShader::GetShaderBindings(Scene, FeatureLevel, MaterialRenderProxy, Material, ShaderBindings);
+
+		ShaderBindings.Add(GetUniformBufferParameter<FViewUniformShaderParameters>(), View.ViewUniformBuffer);
+		// Use GIdentityPrimitiveUniformBuffer just like in the decal handling code
+		// We could potentially bind the actual primitive uniform buffer
+		ShaderBindings.Add(GetUniformBufferParameter<FPrimitiveUniformShaderParameters>(), GIdentityPrimitiveUniformBuffer);
+		ShaderBindings.Add(CloudParameter, CloudParameters);
+	}
+
+private:
+	LAYOUT_FIELD(FShaderUniformBufferParameter, CloudParameter);
+};
+IMPLEMENT_SHADER_TYPE(, FPathTracingVolumetricCloudMaterial, TEXT("/Engine/Private/PathTracing/PathTracingVolumetricCloudMaterialShader.usf"), TEXT("PathTracingVolumetricCloudMaterialShader"), SF_RayCallable);
+
+void PreparePathTracingCloudMaterial(FScene* Scene, TArrayView<FViewInfo> Views)
+{
+	// make sure all views have an invalid callable shader index (unless proven otherwise below)
+	for (FViewInfo& View : Views)
+	{
+		View.PathTracingVolumetricCloudCallableShaderIndex = -1;
+	}
+
+	// if we are using reference clouds, or if we are using the cloud map, no need to prepare the callable shader version
+	if (CVarPathTracingEnableReferenceClouds.GetValueOnRenderThread() == 0 ||
+		CVarPathTracingCloudMapEnabled.GetValueOnRenderThread())
+	{
+		return;
+	}
+
+	FVolumetricCloudRenderSceneInfo* CloudRenderSceneInfo = Scene->GetVolumetricCloudSceneInfo();
+	if (CloudRenderSceneInfo == nullptr)
+	{
+		return;
+	}
+
+	UMaterialInterface* CloudMaterialInterface = CloudRenderSceneInfo->GetVolumetricCloudSceneProxy().GetCloudVolumeMaterial();
+	if (CloudMaterialInterface == nullptr)
+	{
+		return;
+	}
+	const FMaterialRenderProxy* CloudVolumeMaterialProxy = CloudMaterialInterface->GetRenderProxy();
+	if (CloudVolumeMaterialProxy == nullptr)
+	{
+		return;
+	}
+	const FMaterial* MaterialResource = &CloudVolumeMaterialProxy->GetMaterialWithFallback(Scene->GetFeatureLevel(), CloudVolumeMaterialProxy);
+	if (MaterialResource->GetMaterialDomain() != MD_Volume)
+	{
+		return;
+	}
+
+	const FMaterialShaderMap* MaterialShaderMap = MaterialResource->GetRenderingThreadShaderMap();
+	auto CallableShader = MaterialShaderMap->GetShader<FPathTracingVolumetricCloudMaterial>();
+	if (!CallableShader.IsValid())
+	{
+		return;
+	}
+
+	const FVolumetricCloudRenderSceneInfo* CloudInfo = Scene->GetVolumetricCloudSceneInfo();
+
+	for (FViewInfo& View : Views)
+	{
+		uint32 BaseCallableSlotIndex = Scene->RayTracingScene.NumCallableShaderSlots;
+		FRayTracingShaderCommand& Command = Scene->RayTracingScene.CallableCommands.AddDefaulted_GetRef();
+
+		Command.SetShader(CallableShader);
+		Command.SlotInScene = BaseCallableSlotIndex;
+
+		View.PathTracingVolumetricCloudCallableShaderIndex = BaseCallableSlotIndex;
+
+		FPathTracingCloudParameterGlobals Params = {};
+		Params.CloudParameters = PrepareCloudParameters(Scene, View, nullptr);
+		TUniformBufferRef<FPathTracingCloudParameterGlobals> CloudParametersUB = CreateUniformBufferImmediate(
+			Params,
+			EUniformBufferUsage::UniformBuffer_SingleFrame
+		);
+		Scene->RayTracingScene.UniformBuffers.Add(CloudParametersUB); // Hold uniform buffer ref in RayTracingScene since FMeshDrawSingleShaderBindings doesn't
+
+		FMeshDrawSingleShaderBindings SingleShaderBindings = Command.ShaderBindings.GetSingleShaderBindings(SF_RayCallable);
+		CallableShader->GetShaderBindings(Scene, Scene->GetFeatureLevel(), *CloudVolumeMaterialProxy, *MaterialResource, View, CloudParametersUB, SingleShaderBindings);
+
+		Scene->RayTracingScene.NumCallableShaderSlots++;
+	}
+}
+
+
 class FPathTracingBuildAdaptiveErrorTextureCS : public FGlobalShader
 {
 	DECLARE_GLOBAL_SHADER(FPathTracingBuildAdaptiveErrorTextureCS)
@@ -2382,7 +2791,7 @@ void FSceneViewState::PathTracingInvalidate(bool InvalidateAnimationStates)
 	if (State)
 	{
 		
-		if(InvalidateAnimationStates)
+		if (InvalidateAnimationStates)
 		{
 			State->LastDenoisedRadianceRT.SafeRelease();
 			State->LastRadianceRT.SafeRelease();
@@ -2398,6 +2807,8 @@ void FSceneViewState::PathTracingInvalidate(bool InvalidateAnimationStates)
 		State->AlbedoRT.SafeRelease();
 		State->NormalRT.SafeRelease();
 		State->VarianceBuffer.SafeRelease();
+		State->CloudAccelerationMap.SafeRelease();
+		State->CloudMap.SafeRelease();
 		State->SampleIndex = 0;
 
 		State->AdaptiveFrustumGridParameterCache.TopLevelGridBuffer.SafeRelease();
@@ -2480,6 +2891,10 @@ void FDeferredShadingSceneRenderer::RenderPathTracing(
 	Config.UseCameraMediumTracking = CVarPathTracingCameraMediumTracking.GetValueOnRenderThread() != 0;
 	Config.UseAdaptiveSampling = bUseExperimental && CVarPathTracingAdaptiveSampling.GetValueOnAnyThread() != 0;
 	Config.AdaptiveSamplingThreshold = CVarPathTracingAdaptiveSamplingErrorThreshold.GetValueOnRenderThread();
+	Config.CloudAccelerationMapNumSamples = FMath::Clamp(CVarPathTracingCloudAccelerationMapNumSamples.GetValueOnRenderThread(), 1, 65536);
+	Config.CloudAccelerationMapResolution = FMath::Clamp(CVarPathTracingCloudAccelerationMapResolution.GetValueOnRenderThread(), 1, 4096);
+	Config.CloudMapResolution = FMath::Clamp(CVarPathTracingCloudMapResolution.GetValueOnRenderThread(), 1, 4096);
+	Config.CloudMapDepth = FMath::Clamp(CVarPathTracingCloudMapDepth.GetValueOnRenderThread(), 1, 256);
 
 	// compute an integer code of what show flags and booleans related to lights are currently enabled so we can detect changes
 	Config.LightShowFlags = 0;
@@ -2642,6 +3057,7 @@ void FDeferredShadingSceneRenderer::RenderPathTracing(
 	{
 		AtmosphereOpticalDepthLUT = GraphBuilder.RegisterExternalTexture(GSystemTextures.BlackDummy);
 	}
+
 #if WITH_MGPU
 	Config.UseMultiGPU = CVarPathTracingMultiGPU.GetValueOnRenderThread() != 0;
 	// TODO: Figure out how to support adaptive sampling in multi-gpu cases (this is complicated due to the swizzled layout of the variance texture)
@@ -2781,6 +3197,137 @@ void FDeferredShadingSceneRenderer::RenderPathTracing(
 				);
 			}
 
+			FRDGTexture* CloudAccelerationMap = nullptr;
+			FRDGTexture* CloudMap = nullptr;
+
+			const bool bEnableClouds   = (Config.PathTracingData.VolumeFlags & PATH_TRACER_VOLUME_ENABLE_CLOUDS) != 0;
+			const bool bEnableCloudMap = (Config.PathTracingData.VolumeFlags & PATH_TRACER_VOLUME_USE_CLOUD_MAP) != 0;
+
+			if (bEnableClouds)
+			{
+				// clouds are enabled, build an accel map (do this every frame as clouds are usually animating)
+				RDG_GPU_MASK_SCOPE(GraphBuilder, FRHIGPUMask::All());
+
+				const int32 NumSamples = Config.CloudAccelerationMapNumSamples;
+				const int32 Resolution = Config.CloudAccelerationMapResolution;
+				if (PathTracingState->CloudAccelerationMap)
+				{
+					// we already have a map from a previous iteration, re-use it
+					CloudAccelerationMap = GraphBuilder.RegisterExternalTexture(PathTracingState->CloudAccelerationMap, TEXT("PathTracer.CloudAccelerationMap"));
+				}
+				else
+				{
+					// need to create a new LUT
+					EPixelFormat CloudAccelerationMapFormat = PF_FloatRGBA; // 16 bit should be good enough for typical density/z ranges
+					FRDGTextureDesc Desc = FRDGTextureDesc::Create2D(
+						FIntPoint(Resolution, Resolution),
+						CloudAccelerationMapFormat,
+						FClearValueBinding::None,
+						TexCreate_ShaderResource | TexCreate_UAV);
+					CloudAccelerationMap = GraphBuilder.CreateTexture(Desc, TEXT("PathTracer.CloudAccelerationMap"), ERDGTextureFlags::None);
+				}
+				if (bEnableCloudMap)
+				{
+					if (PathTracingState->CloudMap)
+					{
+						CloudMap = GraphBuilder.RegisterExternalTexture(PathTracingState->CloudMap, TEXT("PathTracer.CloudMap"));
+					}
+					else
+					{
+						FRDGTextureDesc Desc = FRDGTextureDesc::Create3D(
+							FIntVector(Config.CloudMapResolution, Config.CloudMapResolution, Config.CloudMapDepth),
+							PF_R32G32B32A32_UINT,
+							FClearValueBinding::None,
+							TexCreate_ShaderResource | TexCreate_UAV);
+						CloudMap = GraphBuilder.CreateTexture(Desc, TEXT("PathTracer.CloudMap"), ERDGTextureFlags::None);
+					}
+				}
+				else
+				{
+					CloudMap = GraphBuilder.RegisterExternalTexture(GSystemTextures.VolumetricBlackUintDummy);
+				}
+				FVolumetricCloudRenderSceneInfo* CloudRenderSceneInfo = Scene->GetVolumetricCloudSceneInfo();
+				check(CloudRenderSceneInfo != nullptr);
+				UMaterialInterface* CloudMaterialInterface = CloudRenderSceneInfo->GetVolumetricCloudSceneProxy().GetCloudVolumeMaterial();
+				check(CloudMaterialInterface != nullptr);
+				const FMaterialRenderProxy* CloudVolumeMaterialProxy = CloudMaterialInterface->GetRenderProxy();
+				check(CloudVolumeMaterialProxy != nullptr);
+				const FMaterial* MaterialResource = &CloudVolumeMaterialProxy->GetMaterialWithFallback(Scene->GetFeatureLevel(), CloudVolumeMaterialProxy);
+				
+
+				const FPathTracingCloudParameters CloudParameters = PrepareCloudParameters(Scene, View, &Config);
+				// build cloud map
+				if (bEnableCloudMap)
+				{
+					typename FPathTracingBuildCloudCS::FPermutationDomain PermutationVector;
+					TShaderRef<FPathTracingBuildCloudCS> ComputeShader = MaterialResource->GetShader<FPathTracingBuildCloudCS>(&FLocalVertexFactory::StaticType, PermutationVector, false);
+
+					FPathTracingBuildCloudCS::FParameters* PassParameters = GraphBuilder.AllocParameters<FPathTracingBuildCloudCS::FParameters>();
+					PassParameters->BlendFactor = Config.PathTracingData.BlendFactor;;
+					PassParameters->Iteration = Config.PathTracingData.Iteration;
+					PassParameters->TemporalSeed = Config.PathTracingData.TemporalSeed;
+					PassParameters->CloudParameters = CloudParameters;
+					PassParameters->ViewUniformBuffer = View.ViewUniformBuffer;
+					PassParameters->CloudMap = GraphBuilder.CreateUAV(CloudMap);
+
+					check(CloudMap->Desc.GetSize() == CloudParameters.CloudMapResolution);
+
+					FIntVector GroupCount = FComputeShaderUtils::GetGroupCount(
+						CloudMap->Desc.GetSize(),
+						FIntVector(4, 4, 4)
+					);
+					GraphBuilder.AddPass(
+						RDG_EVENT_NAME("Path Tracing Cloud Build (%dx%dx%d)", CloudMap->Desc.GetSize().X, CloudMap->Desc.GetSize().Y, CloudMap->Desc.GetSize().Z),
+						PassParameters,
+						ERDGPassFlags::Compute,
+						[LocalScene = Scene, CloudVolumeMaterialProxy, MaterialResource, PassParameters, ComputeShader, GroupCount](FRHIComputeCommandList& RHICmdList)
+						{
+							FMeshDrawShaderBindings ShaderBindings;
+							UE::MeshPassUtils::SetupComputeBindings(ComputeShader, LocalScene, LocalScene->GetFeatureLevel(), nullptr, *CloudVolumeMaterialProxy, *MaterialResource, ShaderBindings);
+							UE::MeshPassUtils::Dispatch(RHICmdList, ComputeShader, ShaderBindings, *PassParameters, GroupCount);
+						});
+					GraphBuilder.QueueTextureExtraction(CloudMap, &PathTracingState->CloudMap);
+				}
+				// build cloud accel map
+				{
+					typename FPathTracingBuildCloudAccelerationMapCS::FPermutationDomain PermutationVector;
+					TShaderRef<FPathTracingBuildCloudAccelerationMapCS> ComputeShader = MaterialResource->GetShader<FPathTracingBuildCloudAccelerationMapCS>(&FLocalVertexFactory::StaticType, PermutationVector, false);
+
+					FPathTracingBuildCloudAccelerationMapCS::FParameters* PassParameters = GraphBuilder.AllocParameters<FPathTracingBuildCloudAccelerationMapCS::FParameters>();
+					PassParameters->NumSamples = NumSamples;
+					PassParameters->Iteration = PathTracingState->SampleIndex;
+					PassParameters->TemporalSeed = PathTracingState->FrameIndex;
+					PassParameters->CloudParameters = CloudParameters;
+					PassParameters->ViewUniformBuffer = View.ViewUniformBuffer;
+					PassParameters->CloudAccelerationMap = GraphBuilder.CreateUAV(CloudAccelerationMap);
+					PassParameters->UseCloudMap = bEnableCloudMap;
+					PassParameters->CloudMap = CloudMap;
+
+					FIntVector GroupCount = FComputeShaderUtils::GetGroupCount(
+						FIntPoint(Resolution, Resolution),
+						FIntPoint(FComputeShaderUtils::kGolden2DGroupSize, FComputeShaderUtils::kGolden2DGroupSize));
+					GraphBuilder.AddPass(
+						RDG_EVENT_NAME("Path Tracing Cloud Acceleration Map Build (Resolution=%u, NumSamples=%u)", Resolution, NumSamples),
+						PassParameters,
+						ERDGPassFlags::Compute,
+						[LocalScene = Scene, CloudVolumeMaterialProxy, MaterialResource, PassParameters, ComputeShader, GroupCount](FRHIComputeCommandList& RHICmdList)
+						{
+							FMeshDrawShaderBindings ShaderBindings;
+							UE::MeshPassUtils::SetupComputeBindings(ComputeShader, LocalScene, LocalScene->GetFeatureLevel(), nullptr, *CloudVolumeMaterialProxy, *MaterialResource, ShaderBindings);
+							UE::MeshPassUtils::Dispatch(RHICmdList, ComputeShader, ShaderBindings, *PassParameters, GroupCount);
+						});
+					GraphBuilder.QueueTextureExtraction(CloudAccelerationMap, &PathTracingState->CloudAccelerationMap);
+				}
+			}
+			if (CloudAccelerationMap == nullptr)
+			{
+				CloudAccelerationMap = GraphBuilder.RegisterExternalTexture(GSystemTextures.BlackDummy);
+			}
+			if (CloudMap == nullptr)
+			{
+				CloudMap = GraphBuilder.RegisterExternalTexture(GSystemTextures.VolumetricBlackUintDummy);
+			}
+
 			// We are writing to the texture, we'll need to extract it...
 			bNeedsTextureExtract = true;
 
@@ -2939,6 +3486,22 @@ void FDeferredShadingSceneRenderer::RenderPathTracing(
 							PassParameters->AtmosphereOpticalDepthLUT = AtmosphereOpticalDepthLUT;
 							PassParameters->AtmosphereOpticalDepthLUTSampler = TStaticSamplerState<SF_Bilinear, AM_Clamp, AM_Clamp, AM_Clamp>::GetRHI();
 
+							if (PreviousPassParameters != nullptr)
+							{
+								PassParameters->CloudParameters = PreviousPassParameters->CloudParameters;
+							} else if ((Config.PathTracingData.VolumeFlags & PATH_TRACER_VOLUME_ENABLE_CLOUDS) != 0)
+							{
+								PassParameters->CloudParameters = PrepareCloudParameters(Scene, View, &Config);
+							}
+							else
+							{
+								PassParameters->CloudParameters = FPathTracingCloudParameters{};
+							}
+							PassParameters->CloudParameters.CloudCallableShaderId = View.PathTracingVolumetricCloudCallableShaderIndex;
+							PassParameters->CloudAccelerationMap = CloudAccelerationMap;
+							PassParameters->CloudAccelerationMapSampler = TStaticSamplerState<SF_Point, AM_Clamp, AM_Clamp, AM_Clamp>::GetRHI();
+							PassParameters->CloudMap = CloudMap;
+
 							if ((Config.PathTracingData.VolumeFlags & PATH_TRACER_VOLUME_ENABLE_FOG) != 0)
 							{
 								PassParameters->FogParameters = PrepareFogParameters(View, Scene->ExponentialFogs[0]);
diff --git a/Engine/Source/Runtime/Renderer/Private/PathTracing.h b/Engine/Source/Runtime/Renderer/Private/PathTracing.h
--- a/Engine/Source/Runtime/Renderer/Private/PathTracing.h
+++ b/Engine/Source/Runtime/Renderer/Private/PathTracing.h
@@ -31,8 +31,10 @@ RENDERER_API FRHIRayTracingShader* GetGPULightmassDefaultOpaqueHitShader(const F
 RENDERER_API FRHIRayTracingShader* GetGPULightmassDefaultHiddenHitShader(const FGlobalShaderMap* ShaderMap);
 
 void PreparePathTracingRTPSO();
+void PreparePathTracingCloudMaterial(FScene* Scene, TArrayView<FViewInfo> Views);
 
 namespace PathTracing
 {
 	bool UsesDecals(const FSceneViewFamily& ViewFamily);
+	bool UsesReferenceAtmosphere(const FViewInfo& View);
 }
diff --git a/Engine/Source/Runtime/Renderer/Private/RayTracing/RayTracingMaterialHitShaders.cpp b/Engine/Source/Runtime/Renderer/Private/RayTracing/RayTracingMaterialHitShaders.cpp
--- a/Engine/Source/Runtime/Renderer/Private/RayTracing/RayTracingMaterialHitShaders.cpp
+++ b/Engine/Source/Runtime/Renderer/Private/RayTracing/RayTracingMaterialHitShaders.cpp
@@ -19,6 +19,7 @@
 #include "RayTracingLighting.h"
 #include "RayTracingDecals.h"
 #include "PathTracing.h"
+#include "RendererModule.h"
 
 int32 GEnableRayTracingMaterials = 1;
 static FAutoConsoleVariableRef CVarEnableRayTracingMaterials(
diff --git a/Engine/Source/Runtime/Renderer/Private/SceneRendering.h b/Engine/Source/Runtime/Renderer/Private/SceneRendering.h
--- a/Engine/Source/Runtime/Renderer/Private/SceneRendering.h
+++ b/Engine/Source/Runtime/Renderer/Private/SceneRendering.h
@@ -1600,6 +1600,8 @@ public:
 	TRDGUniformBufferRef<FRayTracingDecals>				RayTracingDecalUniformBuffer;
 	bool												bHasRayTracingDecals = false;
 
+	int32 PathTracingVolumetricCloudCallableShaderIndex = -1;
+
 	bool bHasAnyRayTracingPass = false;
 #endif // RHI_RAYTRACING
 
diff --git a/Engine/Source/Runtime/Renderer/Private/VolumetricCloudRendering.cpp b/Engine/Source/Runtime/Renderer/Private/VolumetricCloudRendering.cpp
--- a/Engine/Source/Runtime/Renderer/Private/VolumetricCloudRendering.cpp
+++ b/Engine/Source/Runtime/Renderer/Private/VolumetricCloudRendering.cpp
@@ -531,6 +531,7 @@ void FScene::AddVolumetricCloud(FVolumetricCloudSceneProxy* VolumetricCloudScene
 
 			// Use the most recently enabled VolumetricCloud
 			Scene->VolumetricCloud = VolumetricCloudSceneProxy->RenderSceneInfo;
+			Scene->InvalidatePathTracedOutput();
 		} );
 }
 
@@ -554,6 +555,7 @@ void FScene::RemoveVolumetricCloud(FVolumetricCloudSceneProxy* VolumetricCloudSc
 			{
 				Scene->VolumetricCloud = nullptr;
 			}
+			Scene->InvalidatePathTracedOutput();
 		} );
 }
 
@@ -1821,8 +1823,6 @@ void FSceneRenderer::InitVolumetricCloudsForViews(FRDGBuilder& GraphBuilder, boo
 			CloudInfo.GetVolumetricCloudCommonShaderParametersUB() = TUniformBufferRef<FVolumetricCloudCommonGlobalShaderParameters>::CreateUniformBufferImmediate(CloudGlobalShaderParamsUB, UniformBuffer_SingleFrame);
 		}
 
-
-
 		if (CloudProxy.GetCloudVolumeMaterial())
 		{
 			FMaterialRenderProxy* CloudVolumeMaterialProxy = CloudProxy.GetCloudVolumeMaterial()->GetRenderProxy();
