Fix undefined symbols on platforms with raytracing disabled

diff --git a/Engine/Source/Runtime/Renderer/Private/PathTracing.cpp b/Engine/Source/Runtime/Renderer/Private/PathTracing.cpp
--- a/Engine/Source/Runtime/Renderer/Private/PathTracing.cpp
+++ b/Engine/Source/Runtime/Renderer/Private/PathTracing.cpp
@@ -3974,4 +3974,19 @@ void FDeferredShadingSceneRenderer::RenderPathTracing(
 	}
 }
 
+#else // RHI_RAYTRACING == false
+
+namespace PathTracing
+{
+	bool UsesDecals(const FSceneViewFamily& ViewFamily)
+	{
+		return false;
+	}
+
+	bool UsesReferenceAtmosphere(const FViewInfo& View)
+	{
+		return false;
+	}
+}
+
 #endif
