Fix for crash when trying render a SparseVolumeTexture viewer in a cooked game.

#rnx

diff --git a/Engine/Source/Runtime/Renderer/Private/SparseVolumeTexture/SparseVolumeTextureViewerRendering.cpp b/Engine/Source/Runtime/Renderer/Private/SparseVolumeTexture/SparseVolumeTextureViewerRendering.cpp
--- a/Engine/Source/Runtime/Renderer/Private/SparseVolumeTexture/SparseVolumeTextureViewerRendering.cpp
+++ b/Engine/Source/Runtime/Renderer/Private/SparseVolumeTexture/SparseVolumeTextureViewerRendering.cpp
@@ -101,6 +101,8 @@ DECLARE_GPU_STAT(SparseVolumeTextureViewer);
 
 void AddSparseVolumeTextureViewerRenderPass(FRDGBuilder& GraphBuilder, FSceneRenderer& SceneRenderer, FSceneTextures& SceneTextures)
 {
+#if WITH_EDITORONLY_DATA // The SparseVolumeTexture viewer is only supported in editor builds
+
 	FScene* Scene = SceneRenderer.Scene;
 
 	if (Scene->SparseVolumeTextureViewers.Num() == 0)
@@ -216,4 +218,5 @@ void AddSparseVolumeTextureViewerRenderPass(FRDGBuilder& GraphBuilder, FSceneRen
 		}
 	}
 
+#endif // WITH_EDITORONLY_DATA
 }
