Ensure that skylight captures and other scene captures do not have a viewrect offset applied when the r.Test.ViewRectOffset debug setting is enabled.

diff --git a/Engine/Source/Runtime/Renderer/Private/SceneRendering.cpp b/Engine/Source/Runtime/Renderer/Private/SceneRendering.cpp
--- a/Engine/Source/Runtime/Renderer/Private/SceneRendering.cpp
+++ b/Engine/Source/Runtime/Renderer/Private/SceneRendering.cpp
@@ -3176,31 +3176,34 @@ void FSceneRenderer::PrepareViewRectsForRendering(FRHICommandListImmediate& RHIC
 		{
 			FViewInfo& View = Views[0];
 
-			FIntPoint DesiredBufferSize = GetDesiredInternalBufferSize(ViewFamily);
-			FIntPoint Offset = (DesiredBufferSize - View.ViewRect.Size()) / 2;
-			FIntPoint NewViewRectMin(0, 0);
-
-			switch (ViewRectOffset)
+			if (!View.bIsSceneCapture && !View.bIsReflectionCapture)
 			{
-			// Move to the center of the buffer.
-			case 1: NewViewRectMin = Offset; break;
+				FIntPoint DesiredBufferSize = GetDesiredInternalBufferSize(ViewFamily);
+				FIntPoint Offset = (DesiredBufferSize - View.ViewRect.Size()) / 2;
+				FIntPoint NewViewRectMin(0, 0);
 
-			// Move to top left.
-			case 2: break;
+				switch (ViewRectOffset)
+				{
+					// Move to the center of the buffer.
+				case 1: NewViewRectMin = Offset; break;
 
-			// Move to top right.
-			case 3: NewViewRectMin = FIntPoint(2 * Offset.X, 0); break;
+					// Move to top left.
+				case 2: break;
 
-			// Move to bottom right.
-			case 4: NewViewRectMin = FIntPoint(0, 2 * Offset.Y); break;
+					// Move to top right.
+				case 3: NewViewRectMin = FIntPoint(2 * Offset.X, 0); break;
 
-			// Move to bottom left.
-			case 5: NewViewRectMin = FIntPoint(2 * Offset.X, 2 * Offset.Y); break;
-			}
+					// Move to bottom right.
+				case 4: NewViewRectMin = FIntPoint(0, 2 * Offset.Y); break;
 
-			View.ViewRect += QuantizeViewRectMin(NewViewRectMin) - View.ViewRect.Min;
+					// Move to bottom left.
+				case 5: NewViewRectMin = FIntPoint(2 * Offset.X, 2 * Offset.Y); break;
+				}
 
-			check(View.VerifyMembersChecks());
+				View.ViewRect += QuantizeViewRectMin(NewViewRectMin) - View.ViewRect.Min;
+
+				check(View.VerifyMembersChecks());
+			}
 		}
 	}
 	#endif
