Disabled r.SceneCulling.ExplicitCellBounds to work around GPU hang caused by explicit cell bounds full rebuild touching uninitialized data. This reduces the improvements from hierarchical instance culling to some extent.

#rnx

diff --git a/Engine/Source/Runtime/Renderer/Private/SceneCulling/SceneCulling.cpp b/Engine/Source/Runtime/Renderer/Private/SceneCulling/SceneCulling.cpp
--- a/Engine/Source/Runtime/Renderer/Private/SceneCulling/SceneCulling.cpp
+++ b/Engine/Source/Runtime/Renderer/Private/SceneCulling/SceneCulling.cpp
@@ -170,7 +170,7 @@ static TAutoConsoleVariable<int32> CVarValidateAllInstanceAllocations(
 
 static TAutoConsoleVariable<int32> CVarSceneCullingUseExplicitCellBounds(
 	TEXT("r.SceneCulling.ExplicitCellBounds"), 
-	1, 
+	0, 
 	TEXT("Enable to to construct explicit cell bounds by processing the instance bounds as the scene is updated. Adds some GPU cost to the update but this is typically more than paid for by improved culling."),
 	ECVF_RenderThreadSafe);
 
