Fixed & re-enabled explicit cell bounds. Full rebuild no longer touches uninitialized data on GPU causing hangs but instead uploads list of valid cell IDs.

diff --git a/Engine/Source/Runtime/Renderer/Private/SceneCulling/SceneCulling.cpp b/Engine/Source/Runtime/Renderer/Private/SceneCulling/SceneCulling.cpp
--- a/Engine/Source/Runtime/Renderer/Private/SceneCulling/SceneCulling.cpp
+++ b/Engine/Source/Runtime/Renderer/Private/SceneCulling/SceneCulling.cpp
@@ -170,7 +170,7 @@ static TAutoConsoleVariable<int32> CVarValidateAllInstanceAllocations(
 
 static TAutoConsoleVariable<int32> CVarSceneCullingUseExplicitCellBounds(
 	TEXT("r.SceneCulling.ExplicitCellBounds"), 
-	0, 
+	1, 
 	TEXT("Enable to to construct explicit cell bounds by processing the instance bounds as the scene is updated. Adds some GPU cost to the update but this is typically more than paid for by improved culling."),
 	ECVF_RenderThreadSafe);
 
@@ -941,6 +941,15 @@ public:
 
 	using FBlockLoc = FSceneCulling::FBlockLoc;
 
+	enum class EExplicitBoundsUpdateMode
+	{
+		Disabled,
+		Incremental,
+		Full,
+	};
+
+	EExplicitBoundsUpdateMode ExplicitBoundsMode = EExplicitBoundsUpdateMode::Disabled;
+
 #if SC_ENABLE_DETAILED_BUILDER_STATS
 	// Detail Stats
 	struct FStats
@@ -1380,7 +1389,7 @@ public:
 			}			
 			// No need to set the bits as they are maintained incrementally (or new and cleared)
 			SceneCulling.CellOccupancyMask.SetNum(NewMinSize, false);
-			if (SceneCulling.bUseExplictBounds)
+			if (ExplicitBoundsMode == EExplicitBoundsUpdateMode::Incremental)
 			{
 				DirtyCellBoundsMask.SetNum(NewMinSize, false);
 			}
@@ -1915,7 +1924,7 @@ public:
 
 	SC_FORCEINLINE void MarkCellBoundsDirty(int32 CellIndex)
 	{
-		if (SceneCulling.bUseExplictBounds)
+		if (ExplicitBoundsMode == EExplicitBoundsUpdateMode::Incremental)
 		{
 			if (!DirtyCellBoundsMask[CellIndex])
 			{
@@ -2284,10 +2293,11 @@ public:
 		
 		if (SceneCulling.bUseExplictBounds)
 		{
-			bool bFullUpload = !SceneCulling.ExplicitCellBoundsBuffer.GetPooledBuffer().IsValid();
+			check(ExplicitBoundsMode != EExplicitBoundsUpdateMode::Disabled);
+
 			SceneCulling.ExplicitCellBoundsBuffer.ResizeBufferIfNeeded(GraphBuilder, SceneCulling.CellHeaders.Num() * 2);
 
-			if (!DirtyCellBoundsIndices.IsEmpty() || bFullUpload)
+			if (!DirtyCellBoundsIndices.IsEmpty())
 			{
 				RDG_EVENT_SCOPE(GraphBuilder, "SceneCulling_ComputeExplicitCellBounds");
 				FComputeExplicitCellBounds_CS::FParameters* PassParameters = GraphBuilder.AllocParameters<FComputeExplicitCellBounds_CS::FParameters>();
@@ -2300,14 +2310,14 @@ public:
 				PassParameters->InstanceHierarchyCellHeaders = GraphBuilder.CreateSRV(CellHeadersRDG);
 				PassParameters->InstanceIds = GraphBuilder.CreateSRV(InstanceIdsRDG);
 				PassParameters->InstanceHierarchyItemChunks = GraphBuilder.CreateSRV(ItemChunksRDG);
-				int32 MaxCellCount = bFullUpload ? SceneCulling.CellHeaders.Num() : DirtyCellBoundsIndices.Num();
+				int32 MaxCellCount = DirtyCellBoundsIndices.Num();
 				// Note: this Moves the DirtyCellBoundsIndices & thus implicitly clears it!
-				PassParameters->UpdatedCellIds = GraphBuilder.CreateSRV(bFullUpload ? GSystemTextures.GetDefaultStructuredBuffer<uint32>(GraphBuilder) : CreateStructuredBuffer(GraphBuilder, TEXT("SceneCulling.ModifiedCells"), MoveTemp(DirtyCellBoundsIndices)));
+				PassParameters->UpdatedCellIds = GraphBuilder.CreateSRV(CreateStructuredBuffer(GraphBuilder, TEXT("SceneCulling.DirtyCellBoundsIndices"), MoveTemp(DirtyCellBoundsIndices)));
 				PassParameters->MaxCells = MaxCellCount;
 				PassParameters->OutExplicitCellBoundsBuffer = GraphBuilder.CreateUAV(SceneCulling.ExplicitCellBoundsBuffer.Register(GraphBuilder));
 
 				FComputeExplicitCellBounds_CS::FPermutationDomain PermutationVector;
-				PermutationVector.Set<FComputeExplicitCellBounds_CS::FFullBuildDim>(bFullUpload);
+				PermutationVector.Set<FComputeExplicitCellBounds_CS::FFullBuildDim>(false /*TODO: actually remove the permutation after 5.4.1 */);
 
 				auto ComputeShader = GetGlobalShaderMap(SceneCulling.Scene.GetFeatureLevel())->GetShader<FComputeExplicitCellBounds_CS>(PermutationVector);
 
@@ -2362,6 +2372,32 @@ public:
 #endif
 	}
 
+	void ProcessPreSceneUpdate(const FScenePreUpdateChangeSet& ScenePreUpdateData)
+	{
+		if (SceneCulling.bUseExplictBounds)
+		{
+			bool bFullUpload = !SceneCulling.ExplicitCellBoundsBuffer.GetPooledBuffer().IsValid();
+			ExplicitBoundsMode = bFullUpload ? EExplicitBoundsUpdateMode::Full : EExplicitBoundsUpdateMode::Incremental;
+			
+			// No need for dirty tracking if we are going to upload all of them
+			if (ExplicitBoundsMode == EExplicitBoundsUpdateMode::Incremental)
+			{
+				DirtyCellBoundsMask.Init(false, SceneCulling.CellHeaders.Num());
+			}
+			// But we're still going to populate the index list
+			DirtyCellBoundsIndices.Reset(SceneCulling.CellHeaders.Num());
+		}
+		else
+		{
+			ExplicitBoundsMode = EExplicitBoundsUpdateMode::Disabled;
+		}
+
+		for (int32 Index = 0; Index < ScenePreUpdateData.RemovedPrimitiveIds.Num(); ++Index)
+		{
+			MarkInstancesForRemoval(ScenePreUpdateData.RemovedPrimitiveIds[Index], ScenePreUpdateData.RemovedPrimitiveSceneInfos[Index]);
+		}
+	}
+
 	void ProcessPostSceneUpdate(const FScenePostUpdateChangeSet& ScenePostUpdateData)
 	{
 		LLM_SCOPE_BYTAG(SceneCulling);
@@ -2460,7 +2496,22 @@ public:
 		{
 			SceneCulling.BlockLevelOccupancyMask[Item.Key.GetLevel()] = true;
 		}
+		
+		// Note: this can be put into an own task as it can be waited on by the data upload.
+		// If we're in full update mode, produce an array with all valid cell headers
+		if (ExplicitBoundsMode == EExplicitBoundsUpdateMode::Full)
+		{
+			check(DirtyCellBoundsIndices.IsEmpty());
+			for (int32 Index = 0; Index < SceneCulling.CellHeaders.Num(); ++Index)
+			{
+				if (IsValidCell(SceneCulling.CellHeaders[Index]))
+				{
+					DirtyCellBoundsIndices.Add(Index);
+				}
+			}
+		}
 
+		// Note: the below could be put in a separate task, to cut down wait time for subsequent tasks
 		// Releasing memory we're done with in the task saves some RT time
 		DirtyChunks.Empty();
 		DirtyBlocks.Empty();
@@ -2655,18 +2706,8 @@ void FSceneCulling::FUpdater::OnPreSceneUpdate(FRDGBuilder& GraphBuilder, const
 		SCOPE_CYCLE_COUNTER(STAT_SceneCulling_Update_Pre);
 		BUILDER_LOG_SCOPE("OnPreSceneUpdate: %d", ScenePreUpdateData.RemovedPrimitiveIds.Num());
 		CSV_SCOPED_TIMING_STAT(SceneCulling, PreSceneUpdate);
-
-		if (Implementation->SceneCulling.bUseExplictBounds)
-		{
-			Implementation->DirtyCellBoundsMask.Init(false, Implementation->SceneCulling.CellHeaders.Num());
-			Implementation->DirtyCellBoundsIndices.Reset(Implementation->SceneCulling.CellHeaders.Num());
-		}
-
 		check(DebugTaskCounter++ == 0);
-		for (int32 Index = 0; Index < ScenePreUpdateData.RemovedPrimitiveIds.Num(); ++Index)
-		{
-			Implementation->MarkInstancesForRemoval(ScenePreUpdateData.RemovedPrimitiveIds[Index], ScenePreUpdateData.RemovedPrimitiveSceneInfos[Index]);
-		}
+		Implementation->ProcessPreSceneUpdate(ScenePreUpdateData);
 		check(--DebugTaskCounter == 0);
 	}
 #if SC_ALLOW_ASYNC_TASKS
