Remove GPUTransformsSRV code path in ray tracing scene.
- GPU scene is used instead

diff --git a/Engine/Source/Runtime/RHI/Public/RHIResources.h b/Engine/Source/Runtime/RHI/Public/RHIResources.h
--- a/Engine/Source/Runtime/RHI/Public/RHIResources.h
+++ b/Engine/Source/Runtime/RHI/Public/RHIResources.h
@@ -3008,11 +3008,11 @@ struct FRayTracingGeometryInstance
 	TArrayView<const uint32> InstanceSceneDataOffsets;
 
 	// Optional buffer that stores GPU transforms. Used instead of CPU-side transform data.
+	UE_DEPRECATED(5.5, "GPUTransformsSRV has been deprecated. GPU Scene should be used instead.")
 	FShaderResourceViewRHIRef GPUTransformsSRV = nullptr;
 
 	// Conservative number of instances. Some of the actual instances may be made inactive if GPU transforms are used.
 	// Must be less or equal to number of entries in Transforms view if CPU transform data is used.
-	// Must be less or equal to number of entries in GPUTransformsSRV if it is non-null.
 	uint32 NumTransforms = 0;
 
 	// Each geometry copy can receive a user-provided integer, which can be used to retrieve extra shader parameters or customize appearance.
diff --git a/Engine/Source/Runtime/Renderer/Private/RayTracing/RayTracing.cpp b/Engine/Source/Runtime/Renderer/Private/RayTracing/RayTracing.cpp
--- a/Engine/Source/Runtime/Renderer/Private/RayTracing/RayTracing.cpp
+++ b/Engine/Source/Runtime/Renderer/Private/RayTracing/RayTracing.cpp
@@ -961,12 +961,7 @@ namespace RayTracing
 						RayTracingInstance.Flags = SceneInfo->CachedRayTracingInstance.Flags;
 						AddDebugRayTracingInstanceFlags(RayTracingInstance.Flags);
 
-						if (Instance.InstanceGPUTransformsSRV.IsValid())
-						{
-							RayTracingInstance.NumTransforms = Instance.NumTransforms;
-							RayTracingInstance.GPUTransformsSRV = Instance.InstanceGPUTransformsSRV;
-						}
-						else if(!Instance.GetTransforms().IsEmpty())
+						if(!Instance.GetTransforms().IsEmpty())
 						{
 							if (Instance.OwnsTransforms())
 							{
@@ -989,7 +984,7 @@ namespace RayTracing
 						}
 						else
 						{
-							// If neither InstanceGPUTransformsSRV nor array of transforms was provided, get the instance transforms from GPU Scene
+							// If array of transforms was not provided, get the instance transforms from GPU Scene
 							RayTracingInstance.NumTransforms = Instance.NumTransforms;
 							RayTracingInstance.BaseInstanceSceneDataOffset = InstanceSceneDataOffset;
 						}
diff --git a/Engine/Source/Runtime/Renderer/Private/RayTracing/RayTracingDebug.cpp b/Engine/Source/Runtime/Renderer/Private/RayTracing/RayTracingDebug.cpp
--- a/Engine/Source/Runtime/Renderer/Private/RayTracing/RayTracingDebug.cpp
+++ b/Engine/Source/Runtime/Renderer/Private/RayTracing/RayTracingDebug.cpp
@@ -1501,7 +1501,7 @@ void FDeferredShadingSceneRenderer::RenderRayTracingDebug(FRDGBuilder& GraphBuil
 		bRequiresBindings = true;
 	}
 
-	const uint32 NumInstances = View.RayTracingSceneInitData.NumNativeCPUInstances + View.RayTracingSceneInitData.NumNativeGPUInstances + View.RayTracingSceneInitData.NumNativeGPUSceneInstances;
+	const uint32 NumInstances = View.RayTracingSceneInitData.NumNativeCPUInstances + View.RayTracingSceneInitData.NumNativeGPUSceneInstances;
 
 	FRayTracingDebugRGS::FParameters* RayGenParameters = GraphBuilder.AllocParameters<FRayTracingDebugRGS::FParameters>();
 
diff --git a/Engine/Source/Runtime/Renderer/Private/RayTracing/RayTracingInstanceBufferUtil.cpp b/Engine/Source/Runtime/Renderer/Private/RayTracing/RayTracingInstanceBufferUtil.cpp
--- a/Engine/Source/Runtime/Renderer/Private/RayTracing/RayTracingInstanceBufferUtil.cpp
+++ b/Engine/Source/Runtime/Renderer/Private/RayTracing/RayTracingInstanceBufferUtil.cpp
@@ -37,7 +37,6 @@ FRayTracingSceneWithGeometryInstances CreateRayTracingSceneWithGeometryInstances
 	FRayTracingSceneWithGeometryInstances Output;
 	Output.NumNativeGPUSceneInstances = 0;
 	Output.NumNativeCPUInstances = 0;
-	Output.NumNativeGPUInstances = 0;
 	Output.InstanceGeometryIndices.SetNumUninitialized(NumSceneInstances);
 	Output.BaseUploadBufferOffsets.SetNumUninitialized(NumSceneInstances);
 
@@ -63,8 +62,7 @@ FRayTracingSceneWithGeometryInstances CreateRayTracingSceneWithGeometryInstances
 		const FRayTracingGeometryInstance& InstanceDesc = Instances[InstanceIndex];
 
 		const bool bGpuSceneInstance = InstanceDesc.BaseInstanceSceneDataOffset != -1 || !InstanceDesc.InstanceSceneDataOffsets.IsEmpty();
-		const bool bGpuInstance = InstanceDesc.GPUTransformsSRV != nullptr;
-		const bool bCpuInstance = !bGpuSceneInstance && !bGpuInstance;
+		const bool bCpuInstance = !bGpuSceneInstance;
 
 		checkf(!bGpuSceneInstance || InstanceDesc.BaseInstanceSceneDataOffset != -1 || InstanceDesc.NumTransforms <= uint32(InstanceDesc.InstanceSceneDataOffsets.Num()),
 			TEXT("Expected at least %d ray tracing geometry instance scene data offsets, but got %d."),
@@ -90,7 +88,7 @@ FRayTracingSceneWithGeometryInstances CreateRayTracingSceneWithGeometryInstances
 
 		if (bGpuSceneInstance)
 		{
-			check(InstanceDesc.GPUTransformsSRV == nullptr && InstanceDesc.Transforms.IsEmpty());
+			check(InstanceDesc.Transforms.IsEmpty());
 			Output.BaseUploadBufferOffsets[InstanceIndex] = Output.NumNativeGPUSceneInstances;
 			Output.NumNativeGPUSceneInstances += InstanceDesc.NumTransforms;
 		}
@@ -101,17 +99,7 @@ FRayTracingSceneWithGeometryInstances CreateRayTracingSceneWithGeometryInstances
 		}
 		else
 		{
-			if (InstanceDesc.NumTransforms)
-			{
-				FRayTracingGPUInstance GPUInstance;
-				GPUInstance.TransformSRV = InstanceDesc.GPUTransformsSRV;
-				GPUInstance.DescBufferOffset = Output.NumNativeGPUInstances;
-				GPUInstance.NumInstances = InstanceDesc.NumTransforms;
-				Output.GPUInstances.Add(GPUInstance);
-			}
-
-			Output.BaseUploadBufferOffsets[InstanceIndex] = Output.NumNativeGPUInstances;
-			Output.NumNativeGPUInstances += InstanceDesc.NumTransforms;
+			checkNoEntry();
 		}
 
 		checkf(InstanceDesc.LayerIndex < NumLayers, 
@@ -191,10 +179,9 @@ void FillRayTracingInstanceUploadBuffer(
 			const bool bUseUniqueUserData = SceneInstance.UserData.Num() != 0;
 
 			const bool bGpuSceneInstance = SceneInstance.BaseInstanceSceneDataOffset != -1 || !SceneInstance.InstanceSceneDataOffsets.IsEmpty();
-			const bool bGpuInstance = SceneInstance.GPUTransformsSRV != nullptr;
-			const bool bCpuInstance = !bGpuSceneInstance && !bGpuInstance;
+			const bool bCpuInstance = !bGpuSceneInstance;
 
-			checkf(bGpuSceneInstance + bGpuInstance + bCpuInstance == 1, TEXT("Instance can only get transforms from one of GPUScene, GPUTransformsSRV, or Transforms array."));
+			checkf(bGpuSceneInstance + bCpuInstance == 1, TEXT("Instance can only get transforms from one of GPUScene, or Transforms array."));
 
 			const uint32 AccelerationStructureIndex = InstanceGeometryIndices[SceneInstanceIndex];
 			const uint32 LayerBaseIndex = LayerBaseIndices[SceneInstance.LayerIndex];
diff --git a/Engine/Source/Runtime/Renderer/Private/RayTracing/RayTracingScene.cpp b/Engine/Source/Runtime/Renderer/Private/RayTracing/RayTracingScene.cpp
--- a/Engine/Source/Runtime/Renderer/Private/RayTracing/RayTracingScene.cpp
+++ b/Engine/Source/Runtime/Renderer/Private/RayTracing/RayTracingScene.cpp
@@ -74,7 +74,7 @@ void FRayTracingScene::CreateWithInitializationData(FRDGBuilder& GraphBuilder, c
 
 	const FRayTracingSceneInitializer2& SceneInitializer = RayTracingSceneRHI->GetInitializer();
 
-	const uint32 NumNativeInstances = SceneWithGeometryInstances.NumNativeGPUSceneInstances + SceneWithGeometryInstances.NumNativeGPUInstances + SceneWithGeometryInstances.NumNativeCPUInstances;
+	const uint32 NumNativeInstances = SceneWithGeometryInstances.NumNativeGPUSceneInstances + SceneWithGeometryInstances.NumNativeCPUInstances;
 	const uint32 NumNativeInstancesAligned = FMath::DivideAndRoundUp(FMath::Max(NumNativeInstances, 1U), AllocationGranularity) * AllocationGranularity;
 	const uint32 NumTransformsAligned = FMath::DivideAndRoundUp(FMath::Max(SceneWithGeometryInstances.NumNativeCPUInstances, 1U), AllocationGranularity) * AllocationGranularity;
 
diff --git a/Engine/Source/Runtime/Renderer/Public/RayTracingInstance.h b/Engine/Source/Runtime/Renderer/Public/RayTracingInstance.h
--- a/Engine/Source/Runtime/Renderer/Public/RayTracingInstance.h
+++ b/Engine/Source/Runtime/Renderer/Public/RayTracingInstance.h
@@ -130,6 +130,7 @@ struct FRayTracingInstance
 	}
 
 	/** When instance transforms are only available in GPU, this SRV holds them. */
+	UE_DEPRECATED(5.5, "InstanceGPUTransformsSRV has been deprecated. GPU Scene should be used instead.")
 	FShaderResourceViewRHIRef InstanceGPUTransformsSRV;
 
 	//disable deprecation warnings for default constructors
diff --git a/Engine/Source/Runtime/Renderer/Public/RayTracingInstanceBufferUtil.h b/Engine/Source/Runtime/Renderer/Public/RayTracingInstanceBufferUtil.h
--- a/Engine/Source/Runtime/Renderer/Public/RayTracingInstanceBufferUtil.h
+++ b/Engine/Source/Runtime/Renderer/Public/RayTracingInstanceBufferUtil.h
@@ -16,13 +16,12 @@ struct FDFVector3;
 * Each FRayTracingGeometryInstance can translate to multiple native TLAS instances (see FRayTracingGeometryInstance::NumTransforms).
 * 
 * The FRayTracingGeometryInstance array (ie: FRayTracingScene::Instances) used to create FRayTracingSceneRHI
-* can have mix of instances using GPUScene, CPU transforms or GPUTransformSRV.
+* can have mix of instances using GPUScene or CPU transforms.
 * In order to reduce the number of dispatches to build the native RayTracing Instance Buffer,
-* the upload buffer containing FRayTracingInstanceDescriptorInput is split in 3 sections, [GPUSceneInstances] [CPUInstances] [GPUInstances].
+* the upload buffer containing FRayTracingInstanceDescriptorInput is split in 2 sections, [GPUSceneInstances] [CPUInstances].
 * This way native GPUScene and CPU instance descriptors can be built in a single dispatch per type.
-* Followed by one dispatch per GPU instance (since GPU transforms of each GPU instance are stored in separate buffers).
 * 
-* If the raytracing scene contains multiple layers, the instance buffer is divided into multiple subranges as expected by the RHI.
+* If the ray tracing scene contains multiple layers, the instance buffer is divided into multiple subranges as expected by the RHI.
 * 
 */
 
@@ -49,7 +48,6 @@ struct FRayTracingSceneWithGeometryInstances
 	FRayTracingSceneRHIRef Scene;
 	uint32 NumNativeGPUSceneInstances;
 	uint32 NumNativeCPUInstances;
-	uint32 NumNativeGPUInstances;
 	// index of each instance geometry in FRayTracingSceneRHIRef ReferencedGeometries
 	TArray<uint32> InstanceGeometryIndices;
 	// base offset of each instance entries in the instance upload buffer
