Removed/deprecated more code related to GPUTransformsSRV.

diff --git a/Engine/Plugins/Experimental/GPULightmass/Source/GPULightmass/Private/LightmapRenderer.cpp b/Engine/Plugins/Experimental/GPULightmass/Source/GPULightmass/Private/LightmapRenderer.cpp
--- a/Engine/Plugins/Experimental/GPULightmass/Source/GPULightmass/Private/LightmapRenderer.cpp
+++ b/Engine/Plugins/Experimental/GPULightmass/Source/GPULightmass/Private/LightmapRenderer.cpp
@@ -1205,7 +1205,6 @@ bool FSceneRenderState::SetupRayTracingScene(FRDGBuilder& GraphBuilder, FSceneUn
 				TransformUploadSRV,
 				SceneWithGeometryInstances.NumNativeGPUSceneInstances,
 				SceneWithGeometryInstances.NumNativeCPUInstances,
-				{},
 				nullptr,
 				nullptr);
 
diff --git a/Engine/Source/Runtime/Renderer/Private/RayTracing/RayTracingInstanceBufferUtil.cpp b/Engine/Source/Runtime/Renderer/Private/RayTracing/RayTracingInstanceBufferUtil.cpp
--- a/Engine/Source/Runtime/Renderer/Private/RayTracing/RayTracingInstanceBufferUtil.cpp
+++ b/Engine/Source/Runtime/Renderer/Private/RayTracing/RayTracingInstanceBufferUtil.cpp
@@ -190,15 +190,10 @@ void FillRayTracingInstanceUploadBuffer(
 
 			uint32 BaseDescriptorIndex = BaseUploadBufferOffsets[SceneInstanceIndex];
 
-			// Upload buffer is split into 3 sections [GPUSceneInstances][CPUInstances][GPUInstances]
+			// Upload buffer is split into 2 sections [GPUSceneInstances][CPUInstances]
 			if (!bGpuSceneInstance)
 			{
 				BaseDescriptorIndex += NumNativeGPUSceneInstances;
-				
-				if (!bCpuInstance)
-				{
-					BaseDescriptorIndex += NumNativeCPUInstances;
-				}
 			}
 
 			int32 NumInactiveNativeInstancesThisSceneInstance = 0;
@@ -413,7 +408,6 @@ void BuildRayTracingInstanceBuffer(
 	FShaderResourceViewRHIRef CPUInstanceTransformSRV,
 	uint32 NumNativeGPUSceneInstances,
 	uint32 NumNativeCPUInstances,
-	TConstArrayView<FRayTracingGPUInstance> GPUInstances,
 	const FRayTracingCullingParameters* CullingParameters,
 	FUnorderedAccessViewRHIRef DebugInstanceGPUSceneIndexUAV)
 {
@@ -448,25 +442,6 @@ void BuildRayTracingInstanceBuffer(
 			nullptr,
 			nullptr);
 	}
-
-	for (const auto& GPUInstance : GPUInstances)
-	{
-		// GPU instance input descriptors are stored after CPU instances
-		const uint32 InputDescOffset = NumNativeGPUSceneInstances + NumNativeCPUInstances + GPUInstance.DescBufferOffset;
-
-		BuildRayTracingInstanceBuffer(
-			RHICmdList,
-			GPUScene,
-			PreViewTranslation,
-			GPUInstance.NumInstances,
-			InputDescOffset,
-			InstancesUAV,
-			InstanceUploadSRV,
-			AccelerationStructureAddressesSRV,
-			GPUInstance.TransformSRV,
-			CullingParameters,
-			DebugInstanceGPUSceneIndexUAV);
-	}
 }
 
 #endif //RHI_RAYTRACING
diff --git a/Engine/Source/Runtime/Renderer/Private/RayTracing/RayTracingScene.cpp b/Engine/Source/Runtime/Renderer/Private/RayTracing/RayTracingScene.cpp
--- a/Engine/Source/Runtime/Renderer/Private/RayTracing/RayTracingScene.cpp
+++ b/Engine/Source/Runtime/Renderer/Private/RayTracing/RayTracingScene.cpp
@@ -230,7 +230,6 @@ void FRayTracingScene::CreateWithInitializationData(FRDGBuilder& GraphBuilder, c
 			&SceneInitializer,
 			NumNativeGPUSceneInstances = SceneWithGeometryInstances.NumNativeGPUSceneInstances,
 			NumNativeCPUInstances = SceneWithGeometryInstances.NumNativeCPUInstances,
-			GPUInstances = MoveTemp(SceneWithGeometryInstances.GPUInstances),
 			CullingParameters = View.RayTracingCullingParameters
 			](FRHICommandListImmediate& RHICmdList)
 			{
@@ -273,7 +272,6 @@ void FRayTracingScene::CreateWithInitializationData(FRDGBuilder& GraphBuilder, c
 					TransformUploadSRV,
 					NumNativeGPUSceneInstances,
 					NumNativeCPUInstances,
-					GPUInstances,
 					CullingParameters.bUseInstanceCulling ? &CullingParameters : nullptr,
 					PassParams->DebugInstanceGPUSceneIndexBuffer ? PassParams->DebugInstanceGPUSceneIndexBuffer->GetRHI() : nullptr);
 			});
diff --git a/Engine/Source/Runtime/Renderer/Private/Tests/RayTracingTestbed.cpp b/Engine/Source/Runtime/Renderer/Private/Tests/RayTracingTestbed.cpp
--- a/Engine/Source/Runtime/Renderer/Private/Tests/RayTracingTestbed.cpp
+++ b/Engine/Source/Runtime/Renderer/Private/Tests/RayTracingTestbed.cpp
@@ -273,7 +273,6 @@ bool RunRayTracingTestbed_RenderThread(const FString& Parameters)
 		TransformUploadSRV,
 		RayTracingScene.NumNativeGPUSceneInstances,
 		RayTracingScene.NumNativeCPUInstances,
-		{},
 		nullptr,
 		nullptr);
 
diff --git a/Engine/Source/Runtime/Renderer/Public/RayTracingInstanceBufferUtil.h b/Engine/Source/Runtime/Renderer/Public/RayTracingInstanceBufferUtil.h
--- a/Engine/Source/Runtime/Renderer/Public/RayTracingInstanceBufferUtil.h
+++ b/Engine/Source/Runtime/Renderer/Public/RayTracingInstanceBufferUtil.h
@@ -52,6 +52,7 @@ struct FRayTracingSceneWithGeometryInstances
 	TArray<uint32> InstanceGeometryIndices;
 	// base offset of each instance entries in the instance upload buffer
 	TArray<uint32> BaseUploadBufferOffsets;
+	UE_DEPRECATED(5.5, "GPUInstances no longer supported. Use GPUSceneInstances instead.")
 	TArray<FRayTracingGPUInstance> GPUInstances;
 };
 
@@ -93,8 +94,29 @@ RENDERER_API void BuildRayTracingInstanceBuffer(
 	FShaderResourceViewRHIRef CPUInstanceTransformSRV,
 	uint32 NumNativeGPUSceneInstances,
 	uint32 NumNativeCPUInstances,
-	TConstArrayView<FRayTracingGPUInstance> GPUInstances,
 	const FRayTracingCullingParameters* CullingParameters,
 	FUnorderedAccessViewRHIRef DebugInstanceGPUSceneIndexUAV);
 
+UE_DEPRECATED(5.5, "GPUInstances no longer supported. Use GPUSceneInstances instead.")
+inline void BuildRayTracingInstanceBuffer(
+	FRHICommandList& RHICmdList,
+	const FGPUScene* GPUScene,
+	const FDFVector3& PreViewTranslation,
+	FUnorderedAccessViewRHIRef InstancesUAV,
+	FShaderResourceViewRHIRef InstanceUploadSRV,
+	FShaderResourceViewRHIRef AccelerationStructureAddressesSRV,
+	FShaderResourceViewRHIRef CPUInstanceTransformSRV,
+	uint32 NumNativeGPUSceneInstances,
+	uint32 NumNativeCPUInstances,
+	TConstArrayView<FRayTracingGPUInstance> GPUInstances,
+	const FRayTracingCullingParameters* CullingParameters,
+	FUnorderedAccessViewRHIRef DebugInstanceGPUSceneIndexUAV)
+{
+	BuildRayTracingInstanceBuffer(RHICmdList, GPUScene, PreViewTranslation,
+		InstancesUAV, InstanceUploadSRV, AccelerationStructureAddressesSRV, CPUInstanceTransformSRV,
+		NumNativeGPUSceneInstances, NumNativeCPUInstances,
+		CullingParameters,
+		DebugInstanceGPUSceneIndexUAV);
+}
+
 #endif // RHI_RAYTRACING
\ No newline at end of file
