Fix deprecation errors in RayTracingInstanceBufferUtil.h

diff --git a/Engine/Source/Runtime/Renderer/Public/RayTracingInstanceBufferUtil.h b/Engine/Source/Runtime/Renderer/Public/RayTracingInstanceBufferUtil.h
--- a/Engine/Source/Runtime/Renderer/Public/RayTracingInstanceBufferUtil.h
+++ b/Engine/Source/Runtime/Renderer/Public/RayTracingInstanceBufferUtil.h
@@ -56,6 +56,15 @@ struct FRayTracingSceneWithGeometryInstances
 	TArray<uint32> BaseInstancePrefixSum;
 	UE_DEPRECATED(5.5, "GPUInstances no longer supported. Use GPUSceneInstances instead.")
 	TArray<FRayTracingGPUInstance> GPUInstances;
+
+	PRAGMA_DISABLE_DEPRECATION_WARNINGS
+	FRayTracingSceneWithGeometryInstances() = default;
+	FRayTracingSceneWithGeometryInstances(const FRayTracingSceneWithGeometryInstances&) = default;
+	FRayTracingSceneWithGeometryInstances& operator=(const FRayTracingSceneWithGeometryInstances&) = default;
+	FRayTracingSceneWithGeometryInstances(FRayTracingSceneWithGeometryInstances&&) = default;
+	FRayTracingSceneWithGeometryInstances& operator=(FRayTracingSceneWithGeometryInstances&&) = default;
+	~FRayTracingSceneWithGeometryInstances() = default;
+	PRAGMA_ENABLE_DEPRECATION_WARNINGS
 };
 
 // Helper function to create FRayTracingSceneRHI using array of high level instances
