Split Screen:  Debug feature to split the main view into two split screen views, to trivially allow performance testing and render feature debugging in a split screen configuration without needing to set up code or content that puts the game in split screen.  Several console variables starting with "r.SplitScreenDebug" allow prototyping different configurations.

#rnx

diff --git a/Engine/Source/Editor/UnrealEd/Private/EditorViewportClient.cpp b/Engine/Source/Editor/UnrealEd/Private/EditorViewportClient.cpp
--- a/Engine/Source/Editor/UnrealEd/Private/EditorViewportClient.cpp
+++ b/Engine/Source/Editor/UnrealEd/Private/EditorViewportClient.cpp
@@ -4298,6 +4298,8 @@ void FEditorViewportClient::Draw(FViewport* InViewport, FCanvas* Canvas)
 		Canvas->Clear(FLinearColor::Black);
 	}
 
+	ViewFamily.bSplitScreenDebugAllowed = true;
+
 	// Draw the 3D scene
 	GetRendererModule().BeginRenderingViewFamily(Canvas,&ViewFamily);
 
diff --git a/Engine/Source/Runtime/Engine/Private/GameViewportClient.cpp b/Engine/Source/Runtime/Engine/Private/GameViewportClient.cpp
--- a/Engine/Source/Runtime/Engine/Private/GameViewportClient.cpp
+++ b/Engine/Source/Runtime/Engine/Private/GameViewportClient.cpp
@@ -1839,6 +1839,8 @@ void UGameViewportClient::Draw(FViewport* InViewport, FCanvas* SceneCanvas)
 
 	ViewFamily.bIsHDR = GetWindow().IsValid() ? GetWindow().Get()->GetIsHDR() : false;
 
+	ViewFamily.bSplitScreenDebugAllowed = true;
+
 	// Draw the player views.
 	if (!bDisableWorldRendering && PlayerViewMap.Num() > 0 && FSlateApplication::Get().GetPlatformApplication()->IsAllowedToRender()) //-V560
 	{
diff --git a/Engine/Source/Runtime/Engine/Public/SceneView.h b/Engine/Source/Runtime/Engine/Public/SceneView.h
--- a/Engine/Source/Runtime/Engine/Public/SceneView.h
+++ b/Engine/Source/Runtime/Engine/Public/SceneView.h
@@ -2179,7 +2179,10 @@ public:
 	bool bIsSceneTexturesInitialized = false;
 	bool bIsViewFamilyInfo = false;
 
-	/** 
+	/** Whether this view allows split screen debug -- we only want it for editor and game viewports, not other random scene renders */
+	bool bSplitScreenDebugAllowed = false;
+
+	/**
 	 * Which component of the scene rendering should be output to the final render target.
 	 * If SCS_FinalColorLDR this indicates do nothing.
 	 */
diff --git a/Engine/Source/Runtime/Renderer/Private/SceneRendering.cpp b/Engine/Source/Runtime/Renderer/Private/SceneRendering.cpp
--- a/Engine/Source/Runtime/Renderer/Private/SceneRendering.cpp
+++ b/Engine/Source/Runtime/Renderer/Private/SceneRendering.cpp
@@ -145,6 +145,79 @@ static TAutoConsoleVariable<int32> CVarCachedMeshDrawCommands(
 	TEXT("Whether to render from cached mesh draw commands (on vertex factories that support it), or to generate draw commands every frame."),
 	ECVF_RenderThreadSafe);
 
+#if !UE_BUILD_SHIPPING
+
+static TAutoConsoleVariable<int32> CVarSplitScreenDebugEnable(
+	TEXT("r.SplitScreenDebug.Enable"),
+	0,
+	TEXT("Debug feature to replace the main view with a pair of split screen views for testing purposes."),
+	FConsoleVariableDelegate::CreateLambda([](IConsoleVariable* CVar)
+	{
+		// We hit this assert on console in FInstanceCullingMergedContext::AddBatch when enabling split screen:
+		//
+		//		// Verify that each batch contains the same HZB if not null as we only support one
+		//		check(PrevHZB == nullptr || PrevHZB == GraphBuilder.RegisterExternalTexture(Context->PrevHZB));
+		//
+		// HZBs are not shared, so there is definitely a different one per view, making batched instance culling incompatible
+		// with split screen at the moment.  I'm not sure what the fix would be (Create an FInstanceCullingManager per view?
+		// Flush the instance culling manager and swap the HZB when switching views?  Reference multiple HZB and choose the
+		// correct one per batch?  Merge HZBs into a single resource?  Something else?), but I don't want the debug feature
+		// to just crash when used.  So disable this CVar as a workaround.
+		if (CVar->GetBool())
+		{
+			IConsoleVariable* InstanceCullBatchCVar = IConsoleManager::Get().FindConsoleVariable(TEXT("r.InstanceCulling.AllowBatchedBuildRenderingCommands"), false);
+			if (InstanceCullBatchCVar)
+			{
+				InstanceCullBatchCVar->Set(TEXT("0"));
+			}
+		}
+	}),
+	ECVF_Default);
+
+static TAutoConsoleVariable<int32> CVarSplitScreenDebugVertical(
+	TEXT("r.SplitScreenDebug.Vertical"),
+	0,
+	TEXT("Split screen debug use vertical split (two panes vertically stacked).  If false, uses horizontal split (two panes side by side)."),
+	ECVF_Default);
+
+static TAutoConsoleVariable<float> CVarSplitScreenDebugFOVZoom(
+	TEXT("r.SplitScreenDebug.FOVZoom"),
+	1.0f,
+	TEXT("Amount to zoom FOV.  Split screen expands the FOV for the new aspect.  This setting can counteract that expansion."),
+	ECVF_Default);
+
+static TAutoConsoleVariable<float> CVarSplitScreenDebugRotate0(
+	TEXT("r.SplitScreenDebug.Rotate0"),
+	0,
+	TEXT("Rotate first split screen view by this amount.  Values [-1..1] are rotations in view space by fraction of horizontal FOV, outside that range are yaw rotation in degrees."),
+	ECVF_Default);
+
+static TAutoConsoleVariable<float> CVarSplitScreenDebugRotate1(
+	TEXT("r.SplitScreenDebug.Rotate1"),
+	0,
+	TEXT("Rotate second split screen view by this amount.  Values [-1..1] are rotations in view space by fraction of horizontal FOV, outside that range are yaw rotation in degrees."),
+	ECVF_Default);
+
+static TAutoConsoleVariable<int32> CVarSplitScreenDebugOrbit(
+	TEXT("r.SplitScreenDebug.Orbit"),
+	1,
+	TEXT("When rotating by yaw, orbit around camera target actor, to keep third person character visible."),
+	ECVF_Default);
+
+static TAutoConsoleVariable<float> CVarSplitScreenDebugLetterbox(
+	TEXT("r.SplitScreenDebug.Letterbox"),
+	0,
+	TEXT("When non-zero, letterboxes away this percent of screen (rounds up to nearest multiple of 8 pixels, max 50%)."),
+	ECVF_Default);
+
+static TAutoConsoleVariable<int32> CVarSplitScreenDebugLumenScene(
+	TEXT("r.SplitScreenDebug.LumenScene"),
+	1,
+	TEXT("For split screen debugging, allocate a separate Lumen scene for the second view."),
+	ECVF_Default);
+
+#endif  // !UE_BUILD_SHIPPING
+
 bool UseCachedMeshDrawCommands()
 {
 	return CVarCachedMeshDrawCommands.GetValueOnRenderThread() > 0;
@@ -4256,6 +4329,163 @@ void FSceneRenderer::SetupMeshPass(FViewInfo& View, FExclusiveDepthStencil::Type
 	}
 }
 
+#if !UE_BUILD_SHIPPING
+FSceneViewFamily* FSceneRenderer::CreateSplitScreenDebugViewFamily(const FSceneViewFamily& InFamily)
+{
+	FSceneViewFamily* Family = new FSceneViewFamily(InFamily);
+	Family->ScreenPercentageInterface = InFamily.ScreenPercentageInterface->Fork_GameThread(InFamily);
+	Family->Views.SetNumZeroed(2);
+
+	FIntRect OriginalViewRect = InFamily.Views[0]->SceneViewInitOptions.ViewRect;
+
+	int32 SplitVertical = CVarSplitScreenDebugVertical.GetValueOnGameThread();
+	float FOVZoom = CVarSplitScreenDebugFOVZoom.GetValueOnGameThread();
+	float FOVScaleX = FOVZoom;
+	float FOVScaleY = FOVZoom;
+
+	float Letterbox = FMath::Clamp(CVarSplitScreenDebugLetterbox.GetValueOnGameThread(), 0.0f, 50.0f);
+	int32 LetterboxPixels;
+
+	if (SplitVertical)
+	{
+		// Double FOV X
+		FOVScaleX *= 0.5f;
+
+		// Convert letterbox from percentage to a multiple of 8 pixels, then reduce FOV by the relative pixel size
+		LetterboxPixels = FMath::CeilToInt((Letterbox / 100.0f) * OriginalViewRect.Size().X * 0.125f) * 8;
+		FOVScaleX = FOVScaleX * OriginalViewRect.Size().X / (OriginalViewRect.Size().X - LetterboxPixels);
+	}
+	else
+	{
+		// Double FOV Y
+		FOVScaleY *= 0.5f;
+
+		// Convert letterbox from percentage to a multiple of 8 pixels, then reduce FOV by the relative pixel size
+		LetterboxPixels = FMath::CeilToInt((Letterbox / 100.0f) * OriginalViewRect.Size().Y * 0.125f) * 8;
+		FOVScaleY = FOVScaleY * OriginalViewRect.Size().Y / (OriginalViewRect.Size().Y - LetterboxPixels);
+	}
+
+	for (int32 ViewIndex = 0; ViewIndex < 2; ViewIndex++)
+	{
+		FSceneViewInitOptions InitOptions = InFamily.Views[0]->SceneViewInitOptions;
+
+		// Adjust projection
+		InitOptions.ProjectionMatrix *= FMatrix(FVector(FOVScaleX, 0.0, 0.0), FVector(0.0, FOVScaleY, 0.0), FVector(0.0, 0.0, 1.0), FVector(0.0, 0.0, 0.0));
+
+		// Adjust view matrix rotation
+		double Rotate = ViewIndex == 0 ? CVarSplitScreenDebugRotate0.GetValueOnGameThread() : CVarSplitScreenDebugRotate1.GetValueOnGameThread();
+		if (Rotate)
+		{
+			if (FMath::Abs(Rotate) <= 1.0)
+			{
+				// Rotation in view space (post multiply) as a fraction of horizontal FOV.  This mode is useful for creating views
+				// that line up exactly along an edge with each other, without needing to do complex FOV calculations.  For example,
+				// setting the left pane to -0.5 and right pane to 0.5 rotates the views away from each other by half the FOV,
+				// producing a matching frustum edge at the middle of the screen (setting the right pane to 1.0 is another example).
+				double FOV = FMath::RadiansToDegrees(FMath::Atan(1.0 / InitOptions.ProjectionMatrix.M[0][0]) * 2.0);
+				Rotate *= FOV;
+
+				InitOptions.ViewRotationMatrix = InitOptions.ViewRotationMatrix * UE::Math::TRotationMatrix<double>::Make(FRotator(Rotate, 0.0, 0.0));
+			}
+			else
+			{
+				// Rotate by degrees in Yaw
+				FMatrix YawRotation = UE::Math::TRotationMatrix<double>::Make(FRotator(0.0, Rotate, 0.0));
+				InitOptions.ViewRotationMatrix = YawRotation * InitOptions.ViewRotationMatrix;
+
+				// And optionally orbit the position around the player
+				if (CVarSplitScreenDebugOrbit.GetValueOnGameThread() && InitOptions.ViewActor)
+				{
+					FVector TargetTranslation = InitOptions.ViewActor->GetTransform().GetTranslation();
+					InitOptions.ViewOrigin = YawRotation.GetTransposed().TransformVector(InitOptions.ViewOrigin - TargetTranslation) + TargetTranslation;
+					InitOptions.ViewLocation = InitOptions.ViewOrigin;
+				}
+			}
+
+			// Convert adjusted matrix back to a rotation
+			InitOptions.ViewRotation = InitOptions.ViewRotationMatrix.Rotator();
+		}
+		
+		// Make view rectangles half the width / height and adjust opposite dimension for letterbox
+		FIntRect ViewRect = OriginalViewRect;
+		if (SplitVertical)
+		{
+			if (ViewIndex == 0)
+			{
+				ViewRect.Max.Y = (ViewRect.Min.Y + ViewRect.Max.Y) / 2;
+			}
+			else
+			{
+				ViewRect.Min.Y = (ViewRect.Min.Y + ViewRect.Max.Y) / 2;
+			}
+			ViewRect.Min.X += LetterboxPixels / 2;
+			ViewRect.Max.X -= LetterboxPixels / 2;
+		}
+		else
+		{
+			if (ViewIndex == 0)
+			{
+				ViewRect.Max.X = (ViewRect.Min.X + ViewRect.Max.X) / 2;
+			}
+			else
+			{
+				ViewRect.Min.X = (ViewRect.Min.X + ViewRect.Max.X) / 2;
+			}
+			ViewRect.Min.Y += LetterboxPixels / 2;
+			ViewRect.Max.Y -= LetterboxPixels / 2;
+		}
+
+		InitOptions.SetViewRectangle(ViewRect);
+
+		// Set view family to dynamically allocated copy
+		InitOptions.ViewFamily = Family;
+
+		// Use new static view state for second view
+		if (ViewIndex == 1)
+		{
+			static FSceneViewState* GSecondViewState = nullptr;
+			if (!GSecondViewState)
+			{
+				GSecondViewState = new FSceneViewState(InFamily.GetFeatureLevel(), nullptr);
+			}
+
+			// Propagate this user writable field between FSceneViewState
+			const FSceneViewState* SourceViewState = InFamily.Views[0]->State->GetConcreteViewState();
+			GSecondViewState->SequencerState = SourceViewState->SequencerState;
+
+			// Add or remove optional Lumen scene for second view state
+			if (CVarSplitScreenDebugLumenScene.GetValueOnGameThread())
+			{
+				GSecondViewState->AddLumenSceneData(Family->Scene, 1.0f);
+			}
+			else
+			{
+				GSecondViewState->RemoveLumenSceneData(Family->Scene);
+			}
+
+			InitOptions.SceneViewStateInterface = GSecondViewState;
+		}
+
+		FSceneView* View = new FSceneView(InitOptions);
+		View->PrimaryViewIndex = ViewIndex;
+		View->FinalPostProcessSettings = InFamily.Views[0]->FinalPostProcessSettings;
+
+		Family->Views[ViewIndex] = View;
+	}
+
+	return Family;
+}
+
+void FSceneRenderer::DestroySplitScreenDebugViewFamily(FSceneViewFamily* Family)
+{
+	for (int32 ViewIndex = 0; ViewIndex < Family->Views.Num(); ViewIndex++)
+	{
+		delete Family->Views[ViewIndex];
+	}
+	delete Family;
+}
+#endif  // !UE_BUILD_SHIPPING
+
 void FSceneRenderer::CreateSceneRenderers(TArrayView<const FSceneViewFamily*> InViewFamilies, FHitProxyConsumer* HitProxyConsumer, TArray<FSceneRenderer*>& OutSceneRenderers)
 {
 	OutSceneRenderers.Empty(InViewFamilies.Num());
@@ -4265,6 +4495,18 @@ void FSceneRenderer::CreateSceneRenderers(TArrayView<const FSceneViewFamily*> In
 		return;
 	}
 
+#if !UE_BUILD_SHIPPING
+	bool bSplitScreenDebug = false;
+	if (CVarSplitScreenDebugEnable.GetValueOnGameThread() > 0)
+	{
+		if (InViewFamilies.Num() == 1 && InViewFamilies[0]->bSplitScreenDebugAllowed && InViewFamilies[0]->Views.Num() == 1)
+		{
+			InViewFamilies[0] = CreateSplitScreenDebugViewFamily(*InViewFamilies[0]);
+			bSplitScreenDebug = true;
+		}
+	}
+#endif
+
 	const FSceneInterface* Scene = InViewFamilies[0]->Scene;
 	check(Scene);
 
@@ -4333,6 +4575,13 @@ void FSceneRenderer::CreateSceneRenderers(TArrayView<const FSceneViewFamily*> In
 		});
 	}
 #endif  // RHI_RAYTRACING
+
+#if !UE_BUILD_SHIPPING
+	if (bSplitScreenDebug)
+	{
+		DestroySplitScreenDebugViewFamily(const_cast<FSceneViewFamily*>(InViewFamilies[0]));
+	}
+#endif
 }
 
 FSceneRenderer* FSceneRenderer::CreateSceneRenderer(const FSceneViewFamily* InViewFamily, FHitProxyConsumer* HitProxyConsumer)
diff --git a/Engine/Source/Runtime/Renderer/Private/SceneRendering.h b/Engine/Source/Runtime/Renderer/Private/SceneRendering.h
--- a/Engine/Source/Runtime/Renderer/Private/SceneRendering.h
+++ b/Engine/Source/Runtime/Renderer/Private/SceneRendering.h
@@ -2553,6 +2553,11 @@ protected:
 
 	void RenderMeshDistanceFieldVisualization(FRDGBuilder& GraphBuilder, const FMinimalSceneTextures& SceneTextures);
 
+#if !UE_BUILD_SHIPPING
+	static FSceneViewFamily* CreateSplitScreenDebugViewFamily(const FSceneViewFamily& InFamily);
+	static void DestroySplitScreenDebugViewFamily(FSceneViewFamily* Family);
+#endif
+
 protected:
 	FGPUSceneDynamicContext GPUSceneDynamicContext;
 
