Scene Capture:  support for triggering GPU captures or dumps from a transient component flag (in the Advanced section), especially useful for debugging captures triggered from blueprints or captures that aren't updated every frame (a flag is used instead of a button in the UX for that reason).  For captures updated every frame, also allows capture/dump of a scene capture without the rest of the frame's rendering, allowing quicker iteration.

diff --git a/Engine/Source/Runtime/Engine/Classes/Components/SceneCaptureComponent.h b/Engine/Source/Runtime/Engine/Classes/Components/SceneCaptureComponent.h
--- a/Engine/Source/Runtime/Engine/Classes/Components/SceneCaptureComponent.h
+++ b/Engine/Source/Runtime/Engine/Classes/Components/SceneCaptureComponent.h
@@ -86,6 +86,23 @@ class USceneCaptureComponent : public USceneComponent
 	UPROPERTY(EditAnywhere, BlueprintReadWrite, Category = SceneCapture)
 	uint8 bCaptureOnMovement : 1;
 
+	/** Capture a GPU frame for this scene capture, next time it renders (capture program must be connected). */
+	UPROPERTY(EditAnywhere, BlueprintReadWrite, Transient, DuplicateTransient, SkipSerialization, NonTransactional, AdvancedDisplay, Category = SceneCapture)
+	uint8 bCaptureGpuNextRender : 1;
+
+	/** Run DumpGPU for this scene capture, next time it renders. */
+	UPROPERTY(EditAnywhere, BlueprintReadWrite, Transient, DuplicateTransient, SkipSerialization, NonTransactional, AdvancedDisplay, Category = SceneCapture)
+	uint8 bDumpGpuNextRender : 1;
+
+	/**
+	 * Flag used to suppress bCaptureGpuNextRender or bDumpGpuNextRender on reregistration of the component.  Editing any property
+	 * of the component, including the capture/dump flags, forces it to be reregistered, which also triggers the capture to render.
+	 * The purpose of the flags is to allow a capture or dump to be queued and triggered when the next render occurs organically
+	 * (for example, on a blueprint event or movement of the actor), not based on the flag itself being set.  When a property change
+	 * event for one of the flags occurs, this is set to true, to skip the capture/dump on that automatic first render.
+	 */
+	uint8 bSuppressGpuCaptureOrDump : 1;
+
 	/** Whether to persist the rendering state even if bCaptureEveryFrame==false.  This allows velocities for Motion Blur and Temporal AA to be computed. */
 	UPROPERTY(EditAnywhere, BlueprintReadWrite, Category = SceneCapture, meta = (editcondition = "!bCaptureEveryFrame"))
 	bool bAlwaysPersistRenderingState;
diff --git a/Engine/Source/Runtime/Engine/Private/Components/SceneCaptureComponent.cpp b/Engine/Source/Runtime/Engine/Private/Components/SceneCaptureComponent.cpp
--- a/Engine/Source/Runtime/Engine/Private/Components/SceneCaptureComponent.cpp
+++ b/Engine/Source/Runtime/Engine/Private/Components/SceneCaptureComponent.cpp
@@ -171,6 +171,9 @@ USceneCaptureComponent::USceneCaptureComponent(const FObjectInitializer& ObjectI
 	CaptureSource = SCS_SceneColorHDR;
 	bCaptureEveryFrame = true;
 	bCaptureOnMovement = true;
+	bCaptureGpuNextRender = false;
+	bDumpGpuNextRender = false;
+	bSuppressGpuCaptureOrDump = false;
 	bAlwaysPersistRenderingState = false;
 	LODDistanceFactor = 1.0f;
 	MaxViewDistanceOverride = -1;
@@ -475,6 +478,14 @@ void USceneCaptureComponent::PostEditChangeProperty(struct FPropertyChangedEvent
 	{
 		UpdateShowFlags();
 	}
+
+	// The bSuppressGpuCaptureOrDump flag needs to be set to true when bCaptureGpuNextRender or bDumpGpuNextRender are changed, to suppress
+	// GPU capture or dump for the automatic render following component re-registration.  These two properties are the only ones in the
+	// component flagged as non-transactional (so they don't go in the undo/redo stream), so we can detect them based on that flag.
+	if (PropertyChangedEvent.Property && PropertyChangedEvent.Property->HasAnyPropertyFlags(CPF_NonTransactional))
+	{
+		bSuppressGpuCaptureOrDump = true;
+	}
 }
 #endif
 
@@ -735,6 +746,9 @@ void USceneCaptureComponent2D::CaptureScene()
 	UWorld* World = GetWorld();
 	if (World && World->Scene && IsVisible() && !IsCulledByDetailMode())
 	{
+		// Explicit CaptureScene blueprint command should always capture or dump (if capture or dump flag is set)
+		bSuppressGpuCaptureOrDump = false;
+
 		// We must push any deferred render state recreations before causing any rendering to happen, to make sure that deleted resource references are updated
 		World->SendAllEndOfFrameUpdates();
 		UpdateSceneCaptureContents(World->Scene);
@@ -1335,6 +1349,9 @@ void USceneCaptureComponentCube::CaptureScene()
 	UWorld* World = GetWorld();
 	if (World && World->Scene && IsVisible() && !IsCulledByDetailMode())
 	{
+		// Explicit CaptureScene blueprint command should always capture or dump (if capture or dump flag is set)
+		bSuppressGpuCaptureOrDump = false;
+
 		// We must push any deferred render state recreations before causing any rendering to happen, to make sure that deleted resource references are updated
 		World->SendAllEndOfFrameUpdates();
 		UpdateSceneCaptureContents(World->Scene);
diff --git a/Engine/Source/Runtime/RenderCore/Private/DumpGPU.cpp b/Engine/Source/Runtime/RenderCore/Private/DumpGPU.cpp
--- a/Engine/Source/Runtime/RenderCore/Private/DumpGPU.cpp
+++ b/Engine/Source/Runtime/RenderCore/Private/DumpGPU.cpp
@@ -2100,6 +2100,12 @@ FString FRDGBuilder::BeginResourceDump(const TCHAR* Cmd)
 		}
 	}
 
+	// Override frame count CVar, and just capture one frame.  Used for scene capture dumps, where they aren't going to run more than one frame.
+	if (Switches.Contains(TEXT("oneframe")))
+	{
+		NewResourceDumpContext->FrameCount = 1;
+	}
+
 	if (NewResourceDumpContext->bUpload)
 	{
 		if (GDumpGPUUploadCompressResources.GetValueOnGameThread() == 1)
diff --git a/Engine/Source/Runtime/Renderer/Private/SceneCaptureRendering.cpp b/Engine/Source/Runtime/Renderer/Private/SceneCaptureRendering.cpp
--- a/Engine/Source/Runtime/Renderer/Private/SceneCaptureRendering.cpp
+++ b/Engine/Source/Runtime/Renderer/Private/SceneCaptureRendering.cpp
@@ -43,6 +43,9 @@
 #include "RectLightTexture.h"
 #include "Materials/MaterialRenderProxy.h"
 #include "Rendering/CustomRenderPass.h"
+#include "DumpGPU.h"
+#include "IRenderCaptureProvider.h"
+#include "RenderCaptureInterface.h"
 
 bool GSceneCaptureAllowRenderInMainRenderer = true;
 static FAutoConsoleVariableRef CVarSceneCaptureAllowRenderInMainRenderer(
@@ -956,10 +959,74 @@ public:
 	FRenderTarget* SceneCaptureRenderTarget = nullptr;
 };
 
+static void BeginGpuCaptureOrDump(USceneCaptureComponent* CaptureComponent, bool& bCapturingGPU, bool& bDumpingGPU)
+{
+	bCapturingGPU = false;
+	bDumpingGPU = false;
+
+	if (CaptureComponent->bSuppressGpuCaptureOrDump)
+	{
+		CaptureComponent->bSuppressGpuCaptureOrDump = false;
+		return;
+	}
+	
+	bCapturingGPU = CaptureComponent->bCaptureGpuNextRender;
+	bDumpingGPU = CaptureComponent->bDumpGpuNextRender;
+
+	CaptureComponent->bCaptureGpuNextRender = false;
+	CaptureComponent->bDumpGpuNextRender = false;
+
+	// Clear capturing flag if it's not available
+	if (!IRenderCaptureProvider::IsAvailable())
+	{
+		bCapturingGPU = false;
+	}
+
+	// If user sets both capture and dump flags, prefer capturing over dumping (or clear flag if dumping is not available)
+	if (bCapturingGPU || !(WITH_ENGINE && WITH_DUMPGPU))
+	{
+		bDumpingGPU = false;
+	}
+
+#if WITH_ENGINE && WITH_DUMPGPU
+	if (bDumpingGPU)
+	{
+		// Don't try to start a dump if we are already dumping for some reason
+		if (FRDGBuilder::IsDumpingFrame())
+		{
+			bDumpingGPU = false;
+		}
+		else
+		{
+			// Pass "-oneframe" to override CVar that could enable multiple frames of capture
+			FRDGBuilder::BeginResourceDump(TEXT("-oneframe"));
+
+			// Tick the DumpGPU system, which will start the dump
+			UE::RenderCore::DumpGPU::TickEndFrame();
+		}
+	}
+#endif
+}
+
+static void EndGpuCaptureOrDump(bool bDumpingGPU)
+{
+#if WITH_ENGINE && WITH_DUMPGPU
+	if (bDumpingGPU)
+	{
+		// Tick the dump GPU system again, which will end the active dump, so it just includes the scene capture
+		UE::RenderCore::DumpGPU::TickEndFrame();
+	}
+#endif
+}
+
 void FScene::UpdateSceneCaptureContents(USceneCaptureComponent2D* CaptureComponent)
 {
 	check(CaptureComponent);
 
+	bool bCapturingGPU;
+	bool bDumpingGPU;
+	BeginGpuCaptureOrDump(CaptureComponent, bCapturingGPU, bDumpingGPU);
+
 	if (UTextureRenderTarget2D* TextureRenderTarget = CaptureComponent->TextureTarget)
 	{
 		FTransform Transform = CaptureComponent->GetComponentToWorld();
@@ -1203,8 +1270,10 @@ void FScene::UpdateSceneCaptureContents(USceneCaptureComponent2D* CaptureCompone
 		UE::RenderCommandPipe::FSyncScope SyncScope;
 
 		ENQUEUE_RENDER_COMMAND(CaptureCommand)(
-			[SceneRenderer, TextureRenderTargetResource, TexturePtrNotDeferenced, EventName, TargetName, bGenerateMips, GenerateMipsParams, GameViewportRT, bEnableOrthographicTiling, bIsCompositing, bOrthographicCamera, NumXTiles, NumYTiles, TileID, CaptureMemorySize](FRHICommandListImmediate& RHICmdList)
+			[SceneRenderer, TextureRenderTargetResource, TexturePtrNotDeferenced, EventName, TargetName, bGenerateMips, GenerateMipsParams, GameViewportRT, bEnableOrthographicTiling, bIsCompositing, bOrthographicCamera, NumXTiles, NumYTiles, TileID, CaptureMemorySize, bCapturingGPU](FRHICommandListImmediate& RHICmdList)
 			{
+				RenderCaptureInterface::FScopedCapture RenderCapture(bCapturingGPU, &RHICmdList, *FString::Format(TEXT("Scene Capture : {0}"), { EventName }));
+
 				if (GameViewportRT != nullptr)
 				{
 					const FRHIGPUMask GPUMask = GameViewportRT->GetGPUMask(RHICmdList);
@@ -1264,6 +1333,8 @@ void FScene::UpdateSceneCaptureContents(USceneCaptureComponent2D* CaptureCompone
 			}
 		);
 	}
+
+	EndGpuCaptureOrDump(bDumpingGPU);
 }
 
 // Split screen cube map faces are rendered as 3x2 tiles.
@@ -1279,6 +1350,10 @@ static const int32 GCubeFaceViewportOffsets[6][2] =
 
 void FScene::UpdateSceneCaptureContents(USceneCaptureComponentCube* CaptureComponent)
 {
+	bool bCapturingGPU;
+	bool bDumpingGPU;
+	BeginGpuCaptureOrDump(CaptureComponent, bCapturingGPU, bDumpingGPU);
+
 	struct FLocal
 	{
 		/** Creates a transformation for a cubemap face, following the D3D cubemap layout. */
@@ -1382,6 +1457,19 @@ void FScene::UpdateSceneCaptureContents(USceneCaptureComponentCube* CaptureCompo
 
 		if (GSceneCaptureCubeSinglePass == false)
 		{
+			// For GPU capture to work for multi-pass rendering, we need the capture scope to persist across all the scene render command lambdas,
+			// so we need to allocate a pointer on the heap, and let the last render command clean it up.
+			RenderCaptureInterface::FScopedCapture** ScopedCapturePtr;
+			if (bCapturingGPU)
+			{
+				ScopedCapturePtr = new RenderCaptureInterface::FScopedCapture*;
+				*ScopedCapturePtr = nullptr;
+			}
+			else
+			{
+				ScopedCapturePtr = nullptr;
+			}
+
 			for (int32 faceidx = 0; faceidx < (int32)ECubeFace::CubeFace_MAX; faceidx++)
 			{
 				const ECubeFace TargetFace = (ECubeFace)faceidx;
@@ -1420,8 +1508,13 @@ void FScene::UpdateSceneCaptureContents(USceneCaptureComponentCube* CaptureCompo
 				UE::RenderCommandPipe::FSyncScope SyncScope;
 
 				ENQUEUE_RENDER_COMMAND(CaptureCommand)(
-					[SceneRenderer, TextureRenderTarget, EventName, TargetFace, bGenerateMips, GenerateMipsParams](FRHICommandListImmediate& RHICmdList)
+					[SceneRenderer, TextureRenderTarget, EventName, TargetFace, bGenerateMips, GenerateMipsParams, ScopedCapturePtr](FRHICommandListImmediate& RHICmdList)
 					{
+						if (ScopedCapturePtr && (int32)TargetFace == 0)
+						{
+							*ScopedCapturePtr = new RenderCaptureInterface::FScopedCapture(true, &RHICmdList, *FString::Format(TEXT("Scene Capture : {0}"), { EventName }));
+						}
+
 #if WITH_EDITOR
 						// Scene renderer may be deleted in UpdateSceneCaptureContent_RenderThread, grab view state pointer first
 						const FSceneViewState* ViewState = SceneRenderer->Views[0].ViewState;
@@ -1450,6 +1543,13 @@ void FScene::UpdateSceneCaptureContents(USceneCaptureComponentCube* CaptureCompo
 							}
 						}
 #endif  // WITH_EDITOR
+
+						if (ScopedCapturePtr && (int32)TargetFace == (int32)ECubeFace::CubeFace_MAX - 1)
+						{
+							// Delete the scope, and the persistent pointer we allocated on the heap to share across the render command lambdas
+							delete *ScopedCapturePtr;
+							delete ScopedCapturePtr;
+						}
 					}
 				);
 
@@ -1601,8 +1701,10 @@ void FScene::UpdateSceneCaptureContents(USceneCaptureComponentCube* CaptureCompo
 			UE::RenderCommandPipe::FSyncScope SyncScope;
 
 			ENQUEUE_RENDER_COMMAND(CaptureAllCubeFaces)(
-				[SceneRenderer, CubeFaceTarget, TextureRenderTarget, EventName, CaptureSize, bGenerateMips, GenerateMipsParams](FRHICommandListImmediate& RHICmdList)
+				[SceneRenderer, CubeFaceTarget, TextureRenderTarget, EventName, CaptureSize, bGenerateMips, GenerateMipsParams, bCapturingGPU](FRHICommandListImmediate& RHICmdList)
 				{
+					RenderCaptureInterface::FScopedCapture RenderCapture(bCapturingGPU, &RHICmdList, *FString::Format(TEXT("Scene Capture : {0}"), { EventName }));
+
 					TStaticArray<FRHICopyTextureInfo, (int32)ECubeFace::CubeFace_MAX> CopyInfos;
 					for (int32 faceidx = 0; faceidx < (int32)ECubeFace::CubeFace_MAX; faceidx++)
 					{
@@ -1651,4 +1753,6 @@ void FScene::UpdateSceneCaptureContents(USceneCaptureComponentCube* CaptureCompo
 			);
 		}
 	}
+
+	EndGpuCaptureOrDump(bDumpingGPU);
 }
