PR#7554: MinDrawDistance  should calculate with MaxDrawDistanceScale

diff --git a/Engine/Source/Runtime/Renderer/Private/SceneVisibility.cpp b/Engine/Source/Runtime/Renderer/Private/SceneVisibility.cpp
--- a/Engine/Source/Runtime/Renderer/Private/SceneVisibility.cpp
+++ b/Engine/Source/Runtime/Renderer/Private/SceneVisibility.cpp
@@ -402,14 +402,15 @@ bool FViewInfo::IsDistanceCulled( float DistanceSquared, float MinDrawDistance,
 	return bDistanceCulled;
 }
 
-bool FViewInfo::IsDistanceCulled_AnyThread(float DistanceSquared, float MinDrawDistance, float InMaxDrawDistance, const FPrimitiveSceneInfo* PrimitiveSceneInfo, bool& bOutMayBeFading, bool& bOutFadingIn) const
+bool FViewInfo::IsDistanceCulled_AnyThread(float DistanceSquared, float InMinDrawDistance, float InMaxDrawDistance, const FPrimitiveSceneInfo* PrimitiveSceneInfo, bool& bOutMayBeFading, bool& bOutFadingIn) const
 {
 	const float MaxDrawDistanceScale = GetCachedScalabilityCVars().ViewDistanceScale;
 	const float FadeRadius = GDisableLODFade ? 0.0f : GDistanceFadeMaxTravel;
 	const float MaxDrawDistance = InMaxDrawDistance * MaxDrawDistanceScale;
+	const float MinDrawDistance = InMinDrawDistance * MaxDrawDistanceScale;
 
 	bool bHasMaxDrawDistance = InMaxDrawDistance != FLT_MAX;
-	bool bHasMinDrawDistance = MinDrawDistance > 0;
+	bool bHasMinDrawDistance = InMinDrawDistance > 0;
 	bOutMayBeFading = false;
 
 
@@ -813,7 +814,7 @@ static int32 FrustumCull(const FScene& Scene, FViewInfo& View, FFrustumCullingFl
 					if (bHasMaxDrawDistance || bHasMinDrawDistance)
 					{
 						float MaxDrawDistance = Bounds.MaxCullDistance * MaxDrawDistanceScale;
-						float MinDrawDistanceSq = FMath::Square(Bounds.MinDrawDistance);
+						float MinDrawDistanceSq = FMath::Square(Bounds.MinDrawDistance * MaxDrawDistanceScale);
 						float DistanceSquared = FVector::DistSquared(Bounds.BoxSphereBounds.Origin, ViewOriginForDistanceCulling);
 
 						// Always test the fade in distance.  If a primitive was set to always draw, it may need to be faded in.
