Fix animated decals moving in right eye.

diff --git a/Engine/Shaders/Private/DeferredDecal.usf b/Engine/Shaders/Private/DeferredDecal.usf
--- a/Engine/Shaders/Private/DeferredDecal.usf
+++ b/Engine/Shaders/Private/DeferredDecal.usf
@@ -34,6 +34,10 @@
 float4 DecalPositionHigh;
 // from SvPosition to decal space (for position), used like this: mul(float4(SvPosition.xyz, 1), SvPositionToDecal);
 float4x4 SvPositionToDecal;
+#if MOBILE_MULTI_VIEW || (INSTANCED_STEREO && MOBILE_MULTI_VIEW)
+// For stereo rendering.
+float4x4 RightEyeSvPositionToDecal;
+#endif
 // to transform position from decal space to world space, for normals use transpose of inverse
 float4x4 DecalToWorld;
 float3 DecalToWorldInvScale;
@@ -49,11 +53,31 @@ float4 DecalColorParam;
 
 // Global decal vertex shader
 void MainVS(
+#if INSTANCED_STEREO && MOBILE_MULTI_VIEW // Mobile multi view fallback path
+	in uint InstanceId : SV_InstanceID,
+	out nointerpolation uint EyeIndex : VIEW_ID,
+	out uint LayerIndex : SV_RenderTargetArrayIndex,
+#elif MOBILE_MULTI_VIEW
+	in nointerpolation uint ViewId : SV_ViewID,
+#endif
 	in float4 InPosition : ATTRIBUTE0,
 	out float4 OutPosition : SV_POSITION
 )
 {
-	OutPosition = mul(InPosition, FrustumComponentToClip);
+#if INSTANCED_STEREO && MOBILE_MULTI_VIEW
+	EyeIndex = GetEyeIndex(InstanceId);
+	LayerIndex = EyeIndex;
+	ResolvedView = ResolveView(EyeIndex);
+#elif MOBILE_MULTI_VIEW
+	ResolvedView = ResolveView(ViewId);
+#endif
+	
+#if MOBILE_MULTI_VIEW || (INSTANCED_STEREO && MOBILE_MULTI_VIEW)
+	float4x4 FrustumToClip = mul(FrustumComponentToClip, ResolvedView.MobileMultiviewDecalTransform);
+#else
+	float4x4 FrustumToClip = FrustumComponentToClip;
+#endif
+	OutPosition = mul(InPosition, FrustumToClip);
 }
 
 #if IS_DECAL && PIXELSHADER
@@ -80,8 +104,8 @@ float4 SvPositionToScreenPosition2(float4 SvPosition)
 
 	float4 Ret;
 
-	Ret.x = ((SvPosition.x - View.ViewRectMin.x) * View.ViewSizeAndInvSize.z) * 2 - 1;
-	Ret.y = ((SvPosition.y - View.ViewRectMin.y) * View.ViewSizeAndInvSize.w) * -2 + 1;
+	Ret.x = ((SvPosition.x - ResolvedView.ViewRectMin.x) * ResolvedView.ViewSizeAndInvSize.z) * 2 - 1;
+	Ret.y = ((SvPosition.y - ResolvedView.ViewRectMin.y) * ResolvedView.ViewSizeAndInvSize.w) * -2 + 1;
 	Ret.z = ConvertFromDeviceZ(SvPosition.z);
 	Ret.w = 1;
 
@@ -94,16 +118,26 @@ float4 SvPositionToScreenPosition2(float4 SvPosition)
 }
 
 // is called in MainPS() from PixelShaderOutputCommon.usf
-void FPixelShaderInOut_MainPS(inout FPixelShaderIn In, inout FPixelShaderOut Out)
+void FPixelShaderInOut_MainPS(inout FPixelShaderIn In, inout FPixelShaderOut Out, uint ArrayIndex)
 {
+#if (MOBILE_MULTI_VIEW || (INSTANCED_STEREO && MOBILE_MULTI_VIEW))
+	ResolvedView = ResolveView(ArrayIndex);
+#else
 	ResolvedView = ResolveView();
+#endif
 
 	float2 ScreenUV = SvPositionToBufferUV(In.SvPosition);
 
 	// make SvPosition appear to be rasterized with the depth from the depth buffer
+#if (MOBILE_MULTI_VIEW || (INSTANCED_STEREO && MOBILE_MULTI_VIEW))
+	In.SvPosition.z = LookupDeviceZ(ScreenUV, ArrayIndex);
+	float4x4 SvPositionToDecalForSlice = ArrayIndex ? RightEyeSvPositionToDecal : SvPositionToDecal;
+#else
 	In.SvPosition.z = LookupDeviceZ(ScreenUV);
-
-	float4 DecalVectorHom = mul(float4(In.SvPosition.xyz,1), SvPositionToDecal);
+	float4x4 SvPositionToDecalForSlice = SvPositionToDecal;
+#endif
+	
+	float4 DecalVectorHom = mul(float4(In.SvPosition.xyz,1), SvPositionToDecalForSlice);
 	float3 OSPosition = DecalVectorHom.xyz / DecalVectorHom.w;
 
 	// clip content outside the decal
@@ -113,7 +147,7 @@ void FPixelShaderInOut_MainPS(inout FPixelShaderIn In, inout FPixelShaderOut Out
 
 	float3 TranslatedWorldPosition = SvPositionToTranslatedWorld(In.SvPosition);
 
-	float3 CameraVector = normalize(PrimaryView.TranslatedWorldCameraOrigin - TranslatedWorldPosition);
+	float3 CameraVector = normalize(ResolvedView.TranslatedWorldCameraOrigin - TranslatedWorldPosition);
 
 	// can be optimized
 	float3 DecalVector = OSPosition * 0.5f + 0.5f;
@@ -210,6 +244,11 @@ void FPixelShaderInOut_MainPS(inout FPixelShaderIn In, inout FPixelShaderOut Out
 #endif
 }
 
+void FPixelShaderInOut_MainPS(inout FPixelShaderIn In, inout FPixelShaderOut Out)
+{
+	FPixelShaderInOut_MainPS(In, Out, 0);
+}
+
 //#define PIXELSHADEROUTPUT_MRT0 (DECAL_RENDERTARGET_COUNT > 0)
 //#define PIXELSHADEROUTPUT_MRT1 (DECAL_RENDERTARGET_COUNT > 1)
 //#define PIXELSHADEROUTPUT_MRT2 (DECAL_RENDERTARGET_COUNT > 2)
diff --git a/Engine/Shaders/Private/PixelShaderOutputCommon.ush b/Engine/Shaders/Private/PixelShaderOutputCommon.ush
--- a/Engine/Shaders/Private/PixelShaderOutputCommon.ush
+++ b/Engine/Shaders/Private/PixelShaderOutputCommon.ush
@@ -24,6 +24,12 @@ void MainPS
 		FMeshDecalInterpolants MeshDecalInterpolants,
 #endif
 
+#if INSTANCED_STEREO && MOBILE_MULTI_VIEW // Mobile multi view fallback path
+		in nointerpolation uint EyeIndex : VIEW_ID,
+#elif MOBILE_MULTI_VIEW
+		in nointerpolation uint ViewId : SV_ViewID,
+#endif
+
 		in INPUT_POSITION_QUALIFIERS float4 SvPosition : SV_Position		// after all interpolators
 
 		OPTIONAL_IsFrontFace
@@ -137,6 +143,13 @@ void MainPS
 	PixelShaderIn.SvPosition = SvPosition;
 	PixelShaderIn.bIsFrontFace = bIsFrontFace;
 
+uint ArrayIndex = 0;
+#if INSTANCED_STEREO && MOBILE_MULTI_VIEW // Mobile multi view fallback path
+	ArrayIndex = EyeIndex;
+#elif MOBILE_MULTI_VIEW
+	ArrayIndex = ViewId;
+#endif
+
 #if PIXELSHADEROUTPUT_BASEPASS
 	FPixelShaderInOut_MainPS(Interpolants, BasePassInterpolants, PixelShaderIn, PixelShaderOut);
 #elif PIXELSHADEROUTPUT_MESHDECALPASS
@@ -144,7 +157,7 @@ void MainPS
 #elif PIXELSHADEROUTPUT_INTERPOLANTS
 	FPixelShaderInOut_MainPS(Interpolants, PixelShaderIn, PixelShaderOut);
 #else
-	FPixelShaderInOut_MainPS(PixelShaderIn, PixelShaderOut);
+	FPixelShaderInOut_MainPS(PixelShaderIn, PixelShaderOut, ArrayIndex);
 #endif
 
 #if PIXELSHADEROUTPUT_MRT0
diff --git a/Engine/Source/Runtime/Engine/Private/SceneView.cpp b/Engine/Source/Runtime/Engine/Private/SceneView.cpp
--- a/Engine/Source/Runtime/Engine/Private/SceneView.cpp
+++ b/Engine/Source/Runtime/Engine/Private/SceneView.cpp
@@ -2657,6 +2657,9 @@ void FSceneView::SetupCommonViewUniformBufferParameters(
 	ViewUniformShaderParameters.MobileMultiviewShadowTransform = 
 		FMatrix44f(InViewMatrices.GetScreenToClipMatrix() * InViewMatrices.GetInvTranslatedViewProjectionMatrix() * FTranslationMatrix(-InViewMatrices.GetPreViewTranslation()));
 
+	ViewUniformShaderParameters.MobileMultiviewDecalTransform =
+		FMatrix44f(FTranslationMatrix(InViewMatrices.GetPreViewTranslation()) * InViewMatrices.GetTranslatedViewProjectionMatrix());
+
 	ViewUniformShaderParameters.PrevScreenToTranslatedWorld = 
 		FMatrix44f(InPrevViewMatrices.GetScreenToClipMatrix() * InPrevViewMatrices.GetInvTranslatedViewProjectionMatrix());
 
diff --git a/Engine/Source/Runtime/Engine/Public/SceneView.h b/Engine/Source/Runtime/Engine/Public/SceneView.h
--- a/Engine/Source/Runtime/Engine/Public/SceneView.h
+++ b/Engine/Source/Runtime/Engine/Public/SceneView.h
@@ -781,6 +781,7 @@ enum ETranslucencyVolumeCascade
 	VIEW_UNIFORM_BUFFER_MEMBER_PER_VIEW(FMatrix44f, ScreenToRelativeWorld) \
 	VIEW_UNIFORM_BUFFER_MEMBER_PER_VIEW(FMatrix44f, ScreenToTranslatedWorld) \
 	VIEW_UNIFORM_BUFFER_MEMBER_PER_VIEW(FMatrix44f, MobileMultiviewShadowTransform) \
+	VIEW_UNIFORM_BUFFER_MEMBER_PER_VIEW(FMatrix44f, MobileMultiviewDecalTransform) \
 	VIEW_UNIFORM_BUFFER_MEMBER(FVector3f, ViewOriginHigh) \
 	VIEW_UNIFORM_BUFFER_MEMBER_EX(FVector3f, ViewForward, EShaderPrecisionModifier::Half) \
 	VIEW_UNIFORM_BUFFER_MEMBER_EX(FVector3f, ViewUp, EShaderPrecisionModifier::Half) \
diff --git a/Engine/Source/Runtime/Renderer/Private/DecalRenderingShared.cpp b/Engine/Source/Runtime/Renderer/Private/DecalRenderingShared.cpp
--- a/Engine/Source/Runtime/Renderer/Private/DecalRenderingShared.cpp
+++ b/Engine/Source/Runtime/Renderer/Private/DecalRenderingShared.cpp
@@ -54,6 +54,7 @@ class FDeferredDecalVS : public FGlobalShader
 	SHADER_USE_PARAMETER_STRUCT(FDeferredDecalVS, FGlobalShader);
 
 	BEGIN_SHADER_PARAMETER_STRUCT(FParameters, )
+		SHADER_PARAMETER_STRUCT_INCLUDE(FViewShaderParameters, View)
 		SHADER_PARAMETER(FMatrix44f, FrustumComponentToClip)
 		SHADER_PARAMETER_STRUCT_REF(FPrimitiveUniformShaderParameters, PrimitiveUniformBuffer)
 	END_SHADER_PARAMETER_STRUCT()
@@ -87,6 +88,7 @@ public:
 	{
 		DecalPositionHigh.Bind(Initializer.ParameterMap, TEXT("DecalPositionHigh"));
 		SvPositionToDecal.Bind(Initializer.ParameterMap,TEXT("SvPositionToDecal"));
+		RightEyeSvPositionToDecal.Bind(Initializer.ParameterMap, TEXT("RightEyeSvPositionToDecal"));
 		DecalToWorld.Bind(Initializer.ParameterMap,TEXT("DecalToWorld"));
 		DecalToWorldInvScale.Bind(Initializer.ParameterMap, TEXT("DecalToWorldInvScale"));
 		DecalOrientation.Bind(Initializer.ParameterMap,TEXT("DecalOrientation"));
@@ -126,16 +128,29 @@ public:
 			float Ax = -1.0f - 2.0f * View.ViewRect.Min.X * InvViewSize.X;
 			float Ay = 1.0f + 2.0f * View.ViewRect.Min.Y * InvViewSize.Y;
 
-			// todo: we could use InvTranslatedViewProjectionMatrix and TranslatedWorldToComponent for better quality
-			const FMatrix44f SvPositionToDecalValue = FMatrix44f(										// LWC_TODO: Precision loss
-				FMatrix(
-					FPlane(Mx,  0,   0,  0),
-					FPlane( 0, My,   0,  0),
-					FPlane( 0,  0,   1,  0),
-					FPlane(Ax, Ay,   0,  1)
-				) * View.ViewMatrices.GetInvViewProjectionMatrix() * WorldToDecalMatrix);
+			const FMatrix SvPositionToDecalBase(
+				FPlane(Mx, 0, 0, 0),
+				FPlane(0, My, 0, 0),
+				FPlane(0, 0, 1, 0),
+				FPlane(Ax, Ay, 0, 1)
+			);
 
+			// todo: we could use InvTranslatedViewProjectionMatrix and TranslatedWorldToComponent for better quality
+			FMatrix44f SvPositionToDecalValue = FMatrix44f(										// LWC_TODO: Precision loss
+				SvPositionToDecalBase * View.ViewMatrices.GetInvViewProjectionMatrix() * WorldToDecalMatrix);
+			
 			SetShaderValue(BatchedParameters, SvPositionToDecal, SvPositionToDecalValue);
+
+			if (RightEyeSvPositionToDecal.IsBound())
+			{
+				const FViewInfo* InstancedView = View.GetInstancedView();
+				if (InstancedView)
+				{
+					FMatrix44f RightEyeSvPositionToDecalValue = FMatrix44f(										// LWC_TODO: Precision loss
+						SvPositionToDecalBase * InstancedView->ViewMatrices.GetInvViewProjectionMatrix() * WorldToDecalMatrix);
+					SetShaderValue(BatchedParameters, RightEyeSvPositionToDecal, RightEyeSvPositionToDecalValue);
+				}
+			}
 		}
 		if(DecalToWorld.IsBound())
 		{
@@ -164,6 +179,7 @@ public:
 
 private:
 	LAYOUT_FIELD(FShaderParameter, SvPositionToDecal);
+	LAYOUT_FIELD(FShaderParameter, RightEyeSvPositionToDecal);
 	LAYOUT_FIELD(FShaderParameter, DecalPositionHigh);
 	LAYOUT_FIELD(FShaderParameter, DecalToWorld);
 	LAYOUT_FIELD(FShaderParameter, DecalToWorldInvScale);
@@ -458,8 +474,18 @@ namespace DecalRendering
 
 	FMatrix ComputeComponentToClipMatrix(const FViewInfo& View, const FMatrix& DecalComponentToWorld)
 	{
-		FMatrix ComponentToWorldMatrixTrans = DecalComponentToWorld.ConcatTranslation(View.ViewMatrices.GetPreViewTranslation());
-		return ComponentToWorldMatrixTrans * View.ViewMatrices.GetTranslatedViewProjectionMatrix();
+		if (View.bIsMobileMultiViewEnabled || UE::StereoRenderUtils::FStereoShaderAspects(View.GetShaderPlatform()).IsMobileMultiViewEnabled())
+		{
+			// In multi view, the rest of the matrix that is multiplied with DecalComponentToWorld in the non-multi view
+			// case is split out in ViewUniformShaderParameters.MobileMultiviewDecalTransform so we can multiply
+			// it later in the shader.
+			return DecalComponentToWorld;
+		}
+		else
+		{
+			FMatrix ComponentToWorldMatrixTrans = DecalComponentToWorld.ConcatTranslation(View.ViewMatrices.GetPreViewTranslation());
+			return ComponentToWorldMatrixTrans * View.ViewMatrices.GetTranslatedViewProjectionMatrix();
+		}
 	}
 
 	bool TryGetDeferredDecalShaders(
@@ -559,6 +585,7 @@ namespace DecalRendering
 			FDeferredDecalVS::FParameters ShaderParameters;
 			ShaderParameters.FrustumComponentToClip = FMatrix44f(FrustumComponentToClip); // LWC_TODO: Precision loss?
 			ShaderParameters.PrimitiveUniformBuffer = GIdentityPrimitiveUniformBuffer.GetUniformBufferRef();
+			ShaderParameters.View = View.GetShaderParameters();
 			SetShaderParameters(RHICmdList, VertexShader, VertexShader.GetVertexShader(), ShaderParameters);
 		}
 
@@ -585,6 +612,7 @@ namespace DecalRendering
 			FDeferredDecalVS::FParameters ShaderParameters;
 			ShaderParameters.FrustumComponentToClip = FMatrix44f(FrustumComponentToClip); // LWC_TODO: Precision loss
 			ShaderParameters.PrimitiveUniformBuffer = GIdentityPrimitiveUniformBuffer.GetUniformBufferRef();
+			ShaderParameters.View = View.GetShaderParameters();
 			SetShaderParameters(RHICmdList, VertexShader, VertexShader.GetVertexShader(), ShaderParameters);
 		}
 	}
diff --git a/Engine/Source/Runtime/Renderer/Private/MobileDecalRendering.cpp b/Engine/Source/Runtime/Renderer/Private/MobileDecalRendering.cpp
--- a/Engine/Source/Runtime/Renderer/Private/MobileDecalRendering.cpp
+++ b/Engine/Source/Runtime/Renderer/Private/MobileDecalRendering.cpp
@@ -135,7 +135,7 @@ void RenderDeferredDecalsMobile(FRHICommandList& RHICmdList, const FScene& Scene
 			// Set shader params
 			DecalRendering::SetShader(RHICmdList, GraphicsPSOInit, 0, View, DecalData, DecalRenderStage, FrustumComponentToClip);
 
-			RHICmdList.DrawIndexedPrimitive(GetUnitCubeIndexBuffer(), 0, 0, 8, 0, UE_ARRAY_COUNT(GCubeIndices) / 3, 1);
+			RHICmdList.DrawIndexedPrimitive(GetUnitCubeIndexBuffer(), 0, 0, 8, 0, UE_ARRAY_COUNT(GCubeIndices) / 3, View.InstanceFactor);
 		}
 	}
 }
