Fixed shader print not mathcing buffer type of UAV it bound to. (was structured, should be typed)

diff --git a/Engine/Source/Runtime/Renderer/Private/ShaderPrint.cpp b/Engine/Source/Runtime/Renderer/Private/ShaderPrint.cpp
--- a/Engine/Source/Runtime/Renderer/Private/ShaderPrint.cpp
+++ b/Engine/Source/Runtime/Renderer/Private/ShaderPrint.cpp
@@ -187,7 +187,7 @@ namespace ShaderPrint
 	public:
 		void InitRHI(FRHICommandListBase&) override
 		{
-			Buffer = AllocatePooledBuffer(FRDGBufferDesc::CreateStructuredDesc(4, GetCountersUintSize()), TEXT("ShaderPrint.EmptyValueBuffer"));
+			Buffer = AllocatePooledBuffer(FRDGBufferDesc::CreateBufferDesc(4, GetCountersUintSize()), TEXT("ShaderPrint.EmptyValueBuffer"));
 		}
 	};
 
