Handle SceneInstance.bIncrementUserDataPerInstance which got incorrectly reverted in the previous revision.

diff --git a/Engine/Source/Runtime/Renderer/Private/RayTracing/RayTracingInstanceBufferUtil.cpp b/Engine/Source/Runtime/Renderer/Private/RayTracing/RayTracingInstanceBufferUtil.cpp
--- a/Engine/Source/Runtime/Renderer/Private/RayTracing/RayTracingInstanceBufferUtil.cpp
+++ b/Engine/Source/Runtime/Renderer/Private/RayTracing/RayTracingInstanceBufferUtil.cpp
@@ -218,9 +218,25 @@ void FillRayTracingInstanceUploadBuffer(
 					InstanceDesc.GPUSceneInstanceOrTransformIndex = BaseTransformIndex + TransformIndex;
 				}
 
+				uint32 UserData;
+
+				if (bUseUniqueUserData)
+				{
+					UserData = SceneInstance.UserData[TransformIndex];
+				}
+				else
+				{
+					UserData = SceneInstance.DefaultUserData;
+
+					if (SceneInstance.bIncrementUserDataPerInstance)
+					{
+						UserData += TransformIndex;
+					}
+				}
+
 				InstanceDesc.OutputDescriptorIndex = LayerBaseIndex + BaseInstanceIndex + TransformIndex;
 				InstanceDesc.AccelerationStructureIndex = AccelerationStructureIndex;
-				InstanceDesc.InstanceId = bUseUniqueUserData ? SceneInstance.UserData[TransformIndex] : SceneInstance.DefaultUserData;
+				InstanceDesc.InstanceId = UserData;
 				InstanceDesc.InstanceMaskAndFlags = SceneInstance.Mask | ((uint32)SceneInstance.Flags << 8);
 				InstanceDesc.InstanceContributionToHitGroupIndex = SceneInitializer.SegmentPrefixSum[SceneInstanceIndex] * SceneInitializer.ShaderSlotsPerGeometrySegment;
 				InstanceDesc.bApplyLocalBoundsTransform = SceneInstance.bApplyLocalBoundsTransform;
