Perform DF shadows distance-culling based on full-resolution depth. Previously, undersampling depth meant some pixels were not shadowed by DF shadows or CSM, causing glint-like artifacts in fully shadowed regions at the boundary between these techniques.

diff --git a/Engine/Shaders/Private/CapsuleShadowShaders.usf b/Engine/Shaders/Private/CapsuleShadowShaders.usf
--- a/Engine/Shaders/Private/CapsuleShadowShaders.usf
+++ b/Engine/Shaders/Private/CapsuleShadowShaders.usf
@@ -1206,7 +1206,7 @@ void CapsuleShadowingUpsamplePS(
 
 	float Output;
 	float SceneDepth;
-	UpsampleShadowFactors(SVPos, ScissorRectMinAndSize, 1.0f / DOWNSAMPLE_FACTOR, ShadowFactorsTexture, ShadowFactorsSampler, Output, SceneDepth);
+	UpsampleShadowFactors(SVPos, ScissorRectMinAndSize, 1.0f / DOWNSAMPLE_FACTOR, 0, POSITIVE_INFINITY, ShadowFactorsTexture, ShadowFactorsSampler, Output, SceneDepth);
 
 	if (OutputtingToLightAttenuation > 0)
 	{
diff --git a/Engine/Shaders/Private/Common.ush b/Engine/Shaders/Private/Common.ush
--- a/Engine/Shaders/Private/Common.ush
+++ b/Engine/Shaders/Private/Common.ush
@@ -113,6 +113,30 @@ FloatDeriv4 ConstructFloatDeriv4(float4 InValue, float4 InDdx, float4 InDdy)
 #define NearDepthValue (HAS_INVERTED_Z_BUFFER ? 1.0f : 0.0f)
 #define FarDepthValue  (HAS_INVERTED_Z_BUFFER ? 0.0f : 1.0f)
 
+float NearestDeviceDepth(float DepthA, 
+                         float DepthB,
+                         float DepthC = FarDepthValue
+)
+{
+#if HAS_INVERTED_Z_BUFFER
+	return max3(DepthA, DepthB, DepthC);
+#else
+	return min3(DepthA, DepthB, DepthC);
+#endif
+}
+
+float FarthestDeviceDepth(float DepthA, 
+                          float DepthB,
+                          float DepthC = NearDepthValue
+)
+{
+#if HAS_INVERTED_Z_BUFFER
+	return min3(DepthA, DepthB, DepthC);
+#else
+	return max3(DepthA, DepthB, DepthC);
+#endif
+}
+
 const static MaterialFloat PI = 3.1415926535897932f;
 const static float MaxHalfFloat = 65504.0f;
 const static float Max11BitsFloat = 65024.0f;
diff --git a/Engine/Shaders/Private/DistanceFieldShadowing.usf b/Engine/Shaders/Private/DistanceFieldShadowing.usf
--- a/Engine/Shaders/Private/DistanceFieldShadowing.usf
+++ b/Engine/Shaders/Private/DistanceFieldShadowing.usf
@@ -602,11 +602,30 @@ void DistanceFieldShadowingCS(
     uint3 GroupThreadId : SV_GroupThreadID) 
 {
 	uint ThreadIndex = GroupThreadId.y * THREADGROUP_SIZEX + GroupThreadId.x;
+	
+	float SceneDepth;
+	float FullResFurthestSceneDepth;
+	{
+		int2 TopLeftPixelPosition = DispatchThreadId.xy * DownsampleFactor + ScissorRectMinAndSize.xy;
+
+		if (DownsampleFactor == 2)
+		{
+			// cull shadow only if no full-resolution pixel falls within range, preventing edge artifacts
+			float4 SceneDepths = GatherDeviceZ(float2(TopLeftPixelPosition + 1) * View.BufferSizeAndInvSize.zw);
+			SceneDepth = ConvertFromDeviceZ(SceneDepths.x);
+			FullResFurthestSceneDepth = ConvertFromDeviceZ(FarthestDeviceDepth(FarthestDeviceDepth(SceneDepths.x, SceneDepths.y, SceneDepths.z), SceneDepths.w));
+		}
+		else
+		{
+			float SceneDepth00 = LookupDeviceZ(TopLeftPixelPosition);
+			SceneDepth = ConvertFromDeviceZ(SceneDepth00);
+			FullResFurthestSceneDepth = SceneDepth;
+		}
+	}
 
 	float2 ScreenUV = float2((DispatchThreadId.xy * DownsampleFactor + ScissorRectMinAndSize.xy + .5f) * View.BufferSizeAndInvSize.zw);
 	float2 ScreenPosition = (ScreenUV.xy - View.ScreenPositionScaleBias.wz) / View.ScreenPositionScaleBias.xy;
 
-	float SceneDepth = CalcSceneDepth(ScreenUV);
 	float3 OpaqueTranslatedWorldPosition = mul(float4(GetScreenPositionForProjectionType(ScreenPosition, SceneDepth), SceneDepth, 1), PrimaryView.ScreenToTranslatedWorld).xyz;
 
 	// Distance for directional lights to trace
@@ -614,7 +633,7 @@ void DistanceFieldShadowingCS(
 
 	uint NumIntersectingObjects = GetCulledNumObjects();
 	uint CulledDataParameter = 0;
-	bool bShouldComputeShadowing = SceneDepth > MinDepth && SceneDepth < MaxDepth;
+	bool bShouldComputeShadowing = FullResFurthestSceneDepth > MinDepth && SceneDepth < MaxDepth;
 
 #define USE_CULLING 1
 #if USE_CULLING
@@ -788,7 +807,7 @@ void DistanceFieldShadowingUpsamplePS(
 {
 	float Output;
 	float SceneDepth;
-	UpsampleShadowFactors(SVPos, ScissorRectMinAndSize, OneOverDownsampleFactor, ShadowFactorsTexture, ShadowFactorsSampler, Output, SceneDepth);
+	UpsampleShadowFactors(SVPos, ScissorRectMinAndSize, OneOverDownsampleFactor, MinDepth, MaxDepth, ShadowFactorsTexture, ShadowFactorsSampler, Output, SceneDepth);
 
 	float FarBlendFactor = 1.0f - saturate((SceneDepth - FadePlaneOffset) * InvFadePlaneLength);
 	Output = lerp(1, Output, FarBlendFactor);
diff --git a/Engine/Shaders/Private/SceneTexturesCommon.ush b/Engine/Shaders/Private/SceneTexturesCommon.ush
--- a/Engine/Shaders/Private/SceneTexturesCommon.ush
+++ b/Engine/Shaders/Private/SceneTexturesCommon.ush
@@ -86,6 +86,16 @@ float LookupDeviceZ(uint2 PixelPos)
 #endif
 }
 
+/** Returns 4 DeviceZ values, which is the z value stored in the depth buffer. */
+float4 GatherDeviceZ( float2 ScreenUV )
+{
+#if SCENE_TEXTURES_DISABLED
+	return (float4)SCENE_TEXTURES_DISABLED_SCENE_DEPTH_VALUE;
+#else
+	return SceneTexturesStruct.SceneDepthTexture.GatherRed(GlobalBilinearClampedSampler, ScreenUV);
+#endif
+}
+
 /** Returns clip space W, which is world space distance along the View Z axis. */
 float CalcSceneDepth(uint2 PixelPos)
 {
@@ -191,7 +201,7 @@ float4 CalcFullSceneColor(float2 ScreenUV)
 float LookupDeviceZ( float2 ScreenUV, float ArrayIndex )
 {
 #if	SCENE_TEXTURES_DISABLED
-	return FarDepthValue;
+	return SCENE_TEXTURES_DISABLED_SCENE_DEPTH_VALUE;
 #elif FORCE_DEPTH_TEXTURE_READS || PLATFORM_NEEDS_DEPTH_TEXTURE_READS
 	// native Depth buffer lookup
 	#if INSTANCED_STEREO || MOBILE_MULTI_VIEW
@@ -229,6 +239,36 @@ float LookupDeviceZ( float2 ScreenUV )
 	return LookupDeviceZ(ScreenUV, 0);
 }
 
+/** Returns 4 DeviceZ values, which is the z value stored in the depth buffer. */
+float4 GatherDeviceZ( float2 ScreenUV )
+{
+#if	SCENE_TEXTURES_DISABLED
+	return SCENE_TEXTURES_DISABLED_SCENE_DEPTH_VALUE;
+#elif FORCE_DEPTH_TEXTURE_READS || PLATFORM_NEEDS_DEPTH_TEXTURE_READS
+	// native Depth buffer lookup
+	#if INSTANCED_STEREO || MOBILE_MULTI_VIEW
+	return MobileSceneTextures.SceneDepthTextureArray.GatherRed(MobileSceneTextures.SceneDepthTextureSampler, float3(ScreenUV.xy,ArrayIndex));
+	#else
+		return MobileSceneTextures.SceneDepthTexture.GatherRed(MobileSceneTextures.SceneDepthTextureSampler, ScreenUV);
+	#endif
+#elif (USE_SCENE_DEPTH_AUX && !MOBILE_DEFERRED_SHADING)
+	// SceneDepth texture is discarded after BasePass (with forward shading) 
+	// instead fetch DeviceZ from SceneDepthAuxTexture
+	#if INSTANCED_STEREO || MOBILE_MULTI_VIEW
+		return MobileSceneTextures.SceneDepthAuxTextureArray.GatherRed(MobileSceneTextures.SceneDepthAuxTextureSampler, float3(ScreenUV.xy,ArrayIndex));
+	#else
+		return MobileSceneTextures.SceneDepthAuxTexture.GatherRed(MobileSceneTextures.SceneDepthAuxTextureSampler, ScreenUV);
+	#endif
+#else
+	// native Depth buffer lookup
+	#if INSTANCED_STEREO || MOBILE_MULTI_VIEW
+		return MobileSceneTextures.SceneDepthTextureArray.GatherRed(MobileSceneTextures.SceneDepthTextureSampler, float3(ScreenUV.xy,ArrayIndex));
+	#else
+		return MobileSceneTextures.SceneDepthTexture.GatherRed(MobileSceneTextures.SceneDepthTextureSampler, ScreenUV);
+	#endif
+#endif
+}
+
 /** Returns clip space W, which is world space distance along the View Z axis. Note if you need DeviceZ LookupDeviceZ() is the faster option */
 float CalcSceneDepth(float2 ScreenUV, float ArrayIndex)
 {
diff --git a/Engine/Shaders/Private/ScreenSpaceShadows.usf b/Engine/Shaders/Private/ScreenSpaceShadows.usf
--- a/Engine/Shaders/Private/ScreenSpaceShadows.usf
+++ b/Engine/Shaders/Private/ScreenSpaceShadows.usf
@@ -202,7 +202,7 @@ void ScreenSpaceShadowsUpsamplePS(
 {
 	float Output;
 	float SceneDepth;
-	UpsampleShadowFactors(SVPos, ScissorRectMinAndSize, OneOverDownsampleFactor, ShadowFactorsTexture, ShadowFactorsSampler, Output, SceneDepth);
+	UpsampleShadowFactors(SVPos, ScissorRectMinAndSize, OneOverDownsampleFactor, 0, POSITIVE_INFINITY, ShadowFactorsTexture, ShadowFactorsSampler, Output, SceneDepth);
 
 	OutColor = EncodeLightAttenuation(half4(Output, Output, Output, Output));
 }
diff --git a/Engine/Shaders/Private/ShadowFactorsUpsampleCommon.ush b/Engine/Shaders/Private/ShadowFactorsUpsampleCommon.ush
--- a/Engine/Shaders/Private/ShadowFactorsUpsampleCommon.ush
+++ b/Engine/Shaders/Private/ShadowFactorsUpsampleCommon.ush
@@ -11,7 +11,16 @@
 #include "Common.ush"
 #include "ComputeShaderUtils.ush"
 
-void UpsampleShadowFactors(float4 SVPos, int4 ScissorRectMinAndSize, float OneOverDownsampleFactor, Texture2D ShadowFactorsTexture, SamplerState ShadowFactorsSampler, out float ShadowFactor, out float SceneDepth)
+void UpsampleShadowFactors(
+	float4 SVPos, 
+	int4 ScissorRectMinAndSize, 
+	float OneOverDownsampleFactor, 
+	float ShadowFactorsMinDepth, 
+	float ShadowFactorsMaxDepth, 
+	Texture2D ShadowFactorsTexture, 
+	SamplerState ShadowFactorsSampler, 
+	out float ShadowFactor, 
+	out float SceneDepth)
 {
 	float2 BufferUV = SvPositionToBufferUV(SVPos);
 
@@ -50,16 +59,20 @@ void UpsampleShadowFactors(float4 SVPos, int4 ScissorRectMinAndSize, float OneOv
 
 	float Epsilon = .0001f;
 
+	float4 ShadowFactors = float4(TextureValues00.x, TextureValues10.x, TextureValues01.x, TextureValues11.x);
 	float4 CornerDepths = abs(float4(TextureValues00.y, TextureValues10.y, TextureValues01.y, TextureValues11.y));
+
+	// Ignore pixels marked as non-shadowed if they lie outside the DF shadows range.
+	// Shadowed pixels outside the range are marked as valid, these could be from heightfield traces.
+	float4 CornerDepthsValid = 
+		or(
+			and(ShadowFactorsMinDepth.xxxx < CornerDepths, CornerDepths < ShadowFactorsMaxDepth.xxxx),
+			ShadowFactors < 1.0
+		);
 	float4 DepthWeights = 1.0f / (abs(CornerDepths - SceneDepth.xxxx) + Epsilon);
-	float4 FinalWeights = CornerWeights * DepthWeights;
-
-	float InterpolatedResult = 
-		(FinalWeights.x * TextureValues00.x 
-		+ FinalWeights.y * TextureValues10.x
-		+ FinalWeights.z * TextureValues01.x 
-		+ FinalWeights.w * TextureValues11.x)
-		/ dot(FinalWeights, 1);
+	float4 FinalWeights = CornerWeights * CornerDepthsValid * DepthWeights;
+
+	float InterpolatedResult = dot(ShadowFactors, FinalWeights) / max(dot(FinalWeights, 1), Epsilon);
 
 	ShadowFactor = InterpolatedResult;
 #else
diff --git a/Engine/Source/Runtime/Renderer/Private/DistanceFieldShadowing.cpp b/Engine/Source/Runtime/Renderer/Private/DistanceFieldShadowing.cpp
--- a/Engine/Source/Runtime/Renderer/Private/DistanceFieldShadowing.cpp
+++ b/Engine/Source/Runtime/Renderer/Private/DistanceFieldShadowing.cpp
@@ -388,6 +388,8 @@ class FDistanceFieldShadowingUpsamplePS : public FGlobalShader
 		SHADER_PARAMETER(float, NearFadePlaneOffset)
 		SHADER_PARAMETER(float, InvNearFadePlaneLength)
 		SHADER_PARAMETER(float, OneOverDownsampleFactor)
+		SHADER_PARAMETER(float, MinDepth)
+		SHADER_PARAMETER(float, MaxDepth)
 	END_SHADER_PARAMETER_STRUCT()
 
 	class FUpsample : SHADER_PERMUTATION_BOOL("SHADOW_FACTORS_UPSAMPLE_REQUIRED");
@@ -759,6 +761,22 @@ bool FSceneRenderer::ShouldPrepareHeightFieldScene() const
 		&& SupportsHeightFieldShadows(Scene->GetFeatureLevel(), Scene->GetShaderPlatform());
 }
 
+void GetDistanceFieldShadowRange(const FProjectedShadowInfo* ProjectedShadowInfo, EDistanceFieldPrimitiveType PrimitiveType, float& OutMinDepth, float& OutMaxDepth)
+{
+	if (ProjectedShadowInfo->bDirectionalLight)
+	{
+		OutMinDepth = ProjectedShadowInfo->CascadeSettings.SplitNear - ProjectedShadowInfo->CascadeSettings.SplitNearFadeRegion;
+		OutMaxDepth = ProjectedShadowInfo->CascadeSettings.SplitFar;
+	}
+	else
+	{
+		check(PrimitiveType != DFPT_HeightField);
+		//@todo - set these up for point lights as well
+		OutMinDepth = 0.0f;
+		OutMaxDepth = HALF_WORLD_MAX;
+	}
+}
+
 void RayTraceShadows(
 	FRDGBuilder& GraphBuilder,
 	bool bAsyncCompute,
@@ -832,18 +850,7 @@ void RayTraceShadows(
 		PassParameters->TwoSidedMeshDistanceBiasScale = GDFShadowTwoSidedMeshDistanceBiasScale;
 		PassParameters->Substrate = Substrate::BindSubstrateGlobalUniformParameters(View);
 
-		if (ProjectedShadowInfo->bDirectionalLight)
-		{
-			PassParameters->MinDepth = ProjectedShadowInfo->CascadeSettings.SplitNear - ProjectedShadowInfo->CascadeSettings.SplitNearFadeRegion;
-			PassParameters->MaxDepth = ProjectedShadowInfo->CascadeSettings.SplitFar;
-		}
-		else
-		{
-			check(!bHeightfield);
-			//@todo - set these up for point lights as well
-			PassParameters->MinDepth = 0.0f;
-			PassParameters->MaxDepth = HALF_WORLD_MAX;
-		}
+		GetDistanceFieldShadowRange(ProjectedShadowInfo, PrimitiveType, PassParameters->MinDepth, PassParameters->MaxDepth);
 
 		PassParameters->DownsampleFactor = GetDFShadowDownsampleFactor();
 		const FIntPoint OutputBufferSize = OutputTexture->Desc.Extent;
@@ -1114,6 +1121,8 @@ void FProjectedShadowInfo::RenderRayTracedDistanceFieldProjection(
 		PassParameters->PS.ScissorRectMinAndSize = FIntRect(ScissorRect.Min, ScissorRect.Size());
 		PassParameters->PS.OneOverDownsampleFactor = 1.0f / GetDFShadowDownsampleFactor();
 
+		GetDistanceFieldShadowRange(this, DFPT_SignedDistanceField, PassParameters->PS.MinDepth, PassParameters->PS.MaxDepth);
+
 		if (bDirectionalLight && CascadeSettings.FadePlaneLength > 0)
 		{
 			PassParameters->PS.FadePlaneOffset = CascadeSettings.FadePlaneOffset;
