Refactored how FRayTracingScene releases resources and tracking of whether ray tracing scene was used in the frame.

diff --git a/Engine/Source/Runtime/Renderer/Private/DeferredShadingRenderer.cpp b/Engine/Source/Runtime/Renderer/Private/DeferredShadingRenderer.cpp
--- a/Engine/Source/Runtime/Renderer/Private/DeferredShadingRenderer.cpp
+++ b/Engine/Source/Runtime/Renderer/Private/DeferredShadingRenderer.cpp
@@ -923,11 +923,11 @@ bool FDeferredShadingSceneRenderer::DispatchRayTracingWorldUpdates(FRDGBuilder&
 	return true;
 }
 
-static void ReleaseRaytracingResources(FRDGBuilder& GraphBuilder, TArrayView<FViewInfo> Views, FRayTracingScene &RayTracingScene, bool bIsLastRenderer)
+static void ReleaseRaytracingResources(FRDGBuilder& GraphBuilder, FRayTracingScene &RayTracingScene, bool bIsLastRenderer)
 {
 	// Keep mask the same as what's already set (which will be the view mask) if TLAS updates should be masked to the view
 	RDG_GPU_MASK_SCOPE(GraphBuilder, GRayTracingMultiGpuTLASMask ? GraphBuilder.RHICmdList.GetGPUMask() : FRHIGPUMask::All());
-	AddPass(GraphBuilder, RDG_EVENT_NAME("ReleaseRayTracingResources"), [Views, &RayTracingScene, bIsLastRenderer](FRHICommandListImmediate& RHICmdList)
+	AddPass(GraphBuilder, RDG_EVENT_NAME("ReleaseRayTracingResources"), [&RayTracingScene, bIsLastRenderer](FRHICommandListImmediate& RHICmdList)
 	{
 		if (RayTracingScene.IsCreated())
 		{
@@ -936,13 +936,6 @@ static void ReleaseRaytracingResources(FRDGBuilder& GraphBuilder, TArrayView<FVi
 			{
 				RHICmdList.ClearRayTracingBindings(RayTracingScene.GetRHIRayTracingScene());
 			}
-
-			// Track if we ended up rendering anything this frame.  After rendering all view families, we'll release the
-			// ray tracing scene resources if nothing used ray tracing.
-			if (RayTracingScene.GetInstances().Num() > 0)
-			{
-				RayTracingScene.bUsedThisFrame = true;
-			}
 		}
 	});
 }
@@ -3369,7 +3362,7 @@ void FDeferredShadingSceneRenderer::Render(FRDGBuilder& GraphBuilder)
 		GetSceneExtensionsRenderers().PostRender(GraphBuilder);
 
 #if RHI_RAYTRACING
-		ReleaseRaytracingResources(GraphBuilder, Views, Scene->RayTracingScene, bIsLastSceneRenderer);
+		ReleaseRaytracingResources(GraphBuilder, Scene->RayTracingScene, bIsLastSceneRenderer);
 #endif //  RHI_RAYTRACING
 	}
 
diff --git a/Engine/Source/Runtime/Renderer/Private/RayTracing/RayTracingScene.cpp b/Engine/Source/Runtime/Renderer/Private/RayTracing/RayTracingScene.cpp
--- a/Engine/Source/Runtime/Renderer/Private/RayTracing/RayTracingScene.cpp
+++ b/Engine/Source/Runtime/Renderer/Private/RayTracing/RayTracingScene.cpp
@@ -63,6 +63,8 @@ void FRayTracingScene::CreateWithInitializationData(FRDGBuilder& GraphBuilder, c
 	// Make sure that all async tasks complete before we run initialization code again and create new tasks.
 	WaitForTasks();
 
+	bUsedThisFrame = true;
+
 	FRHICommandListBase& RHICmdList = GraphBuilder.RHICmdList;
 
 	static const uint8 NumLayers = uint8(ERayTracingSceneLayer::NUM);
@@ -393,15 +395,17 @@ void FRayTracingScene::Reset(bool bInInstanceDebugDataEnabled)
 
 	Instances.Reset();
 	InstancesDebugData.Reset();
-	NumMissShaderSlots = 1;
-	NumCallableShaderSlots = 0;
 	CallableCommands.Reset();
 	UniformBuffers.Reset();
 	GeometriesToBuild.Reset();
 	UsedCoarseMeshStreamingHandles.Reset();
 
+	NumMissShaderSlots = 1;
+	NumCallableShaderSlots = 0;
+
 	Allocator.Flush();
 
+	RayTracingSceneRHI = nullptr;
 	BuildScratchBuffer = nullptr;
 	InstanceDebugBuffer = nullptr;
 	DebugInstanceGPUSceneIndexBuffer = nullptr;
@@ -409,19 +413,25 @@ void FRayTracingScene::Reset(bool bInInstanceDebugDataEnabled)
 	bInstanceDebugDataEnabled = bInInstanceDebugDataEnabled;
 }
 
-void FRayTracingScene::ResetAndReleaseResources()
+void FRayTracingScene::EndFrame()
 {
 	Reset(false);
 
-	Instances.Empty();
-	InstancesDebugData.Empty();
-	CallableCommands.Empty();
-	UniformBuffers.Empty();
-	GeometriesToBuild.Empty();
-	UsedCoarseMeshStreamingHandles.Empty();
-	RayTracingSceneBuffer = nullptr;
-	RayTracingScenePooledBuffer = nullptr;
-	RayTracingSceneRHI = nullptr;
+	// Release the resources if ray tracing wasn't used
+	if (!bUsedThisFrame)
+	{
+		Instances.Empty();
+		InstancesDebugData.Empty();
+		CallableCommands.Empty();
+		UniformBuffers.Empty();
+		GeometriesToBuild.Empty();
+		UsedCoarseMeshStreamingHandles.Empty();
+
+		RayTracingSceneBuffer = nullptr;
+		RayTracingScenePooledBuffer = nullptr;
+	}
+
+	bUsedThisFrame = false;
 }
 
 #endif // RHI_RAYTRACING
diff --git a/Engine/Source/Runtime/Renderer/Private/RayTracing/RayTracingScene.h b/Engine/Source/Runtime/Renderer/Private/RayTracing/RayTracingScene.h
--- a/Engine/Source/Runtime/Renderer/Private/RayTracing/RayTracingScene.h
+++ b/Engine/Source/Runtime/Renderer/Private/RayTracing/RayTracingScene.h
@@ -69,8 +69,7 @@ public:
 	// Resets the instance list and reserves memory for this frame.
 	void Reset(bool bInstanceDebugDataEnabled);
 
-	// Similar to Reset(), but also releases any persistent CPU and GPU memory allocations.
-	void ResetAndReleaseResources();
+	void EndFrame();
 
 	// Allocates temporary memory that will be valid until the next Reset().
 	// Can be used to store temporary instance transforms, user data, etc.
@@ -110,7 +109,6 @@ public:
 
 	// Public members for initial refactoring step (previously were public members of FViewInfo).
 
-	bool bUsedThisFrame = false;
 	uint32 NumMissShaderSlots = 1; // we must have a default miss shader, so always include it from the start
 	uint32 NumCallableShaderSlots = 0;
 	TArray<FRayTracingShaderCommand> CallableCommands;
@@ -170,6 +168,8 @@ private:
 	TArray<FRayTracingInstanceDebugData> InstancesDebugData;
 
 	bool bInstanceDebugDataEnabled = false;
+
+	bool bUsedThisFrame = false;
 };
 
 #endif // RHI_RAYTRACING
diff --git a/Engine/Source/Runtime/Renderer/Private/SceneRendering.cpp b/Engine/Source/Runtime/Renderer/Private/SceneRendering.cpp
--- a/Engine/Source/Runtime/Renderer/Private/SceneRendering.cpp
+++ b/Engine/Source/Runtime/Renderer/Private/SceneRendering.cpp
@@ -4569,9 +4569,6 @@ void FSceneRenderer::CreateSceneRenderers(TArrayView<const FSceneViewFamily*> In
 					}
 				}
 			}
-
-			// Clear flag that tracks whether ray tracing was used this frame
-			SceneRenderers[0]->Scene->RayTracingScene.bUsedThisFrame = false;
 		});
 	}
 #endif  // RHI_RAYTRACING
@@ -5135,11 +5132,7 @@ static void RenderViewFamilies_RenderThread(FRHICommandListImmediate& RHICmdList
 	}
 
 #if RHI_RAYTRACING
-	// Release the ray tracing scene resources if ray tracing wasn't used
-	if (!Scene->RayTracingScene.bUsedThisFrame)
-	{
-		Scene->RayTracingScene.ResetAndReleaseResources();
-	}
+	Scene->RayTracingScene.EndFrame();
 #endif  // RHI_RAYTRACING
 
 #if STATS
