For virtual textures that lock root pages, don't map any other pages until the root pages are locked.

diff --git a/Engine/Source/Runtime/RenderCore/Public/VirtualTexturing.h b/Engine/Source/Runtime/RenderCore/Public/VirtualTexturing.h
--- a/Engine/Source/Runtime/RenderCore/Public/VirtualTexturing.h
+++ b/Engine/Source/Runtime/RenderCore/Public/VirtualTexturing.h
@@ -374,6 +374,7 @@ public:
 		, DepthInTiles(InDepthInTiles)
 		, FrameDeleted(0u)
 		, NumRefs(0)
+		, bIsWaitingToMap(false)
 		, PageTableFormat(EVTPageTableFormat::UInt32)
 		, SpaceID(~0u)
 		, MaxLevel(0u)
@@ -436,6 +437,7 @@ protected:
 	uint32 DepthInTiles;
 	uint32 FrameDeleted;
 	int32 NumRefs;
+	bool bIsWaitingToMap;
 
 	// should be set explicitly by derived class constructor
 	EVTPageTableFormat PageTableFormat;
diff --git a/Engine/Source/Runtime/Renderer/Private/VT/VirtualTextureSystem.cpp b/Engine/Source/Runtime/Renderer/Private/VT/VirtualTextureSystem.cpp
--- a/Engine/Source/Runtime/Renderer/Private/VT/VirtualTextureSystem.cpp
+++ b/Engine/Source/Runtime/Renderer/Private/VT/VirtualTextureSystem.cpp
@@ -687,6 +687,7 @@ IAllocatedVirtualTexture* FVirtualTextureSystem::AllocateVirtualTexture(FRHIComm
 	if (bAnyLayerProducerWantsPersistentHighestMip)
 	{
 		AllocatedVTsToMap.Add(AllocatedVT);
+		AllocatedVT->bIsWaitingToMap = true;
 	}
 
 	// Add to deterministic map that should apply across runs.
@@ -1593,6 +1594,12 @@ void FVirtualTextureSystem::GatherRequestsTask(const FGatherRequestsParameters&
 			continue;
 		}
 
+		if (AllocatedVT->bIsWaitingToMap)
+		{
+			// If the VT is still waiting to map locked pages, don't map other pages first
+			continue;
+		}
+
 		const uint32 MaxLevel = AllocatedVT->GetMaxLevel();
 
 		check(AllocatedVT->GetNumPageTableLayers() == Space->GetNumPageTableLayers());
@@ -2403,10 +2410,11 @@ void FVirtualTextureSystem::SubmitRequests(FRHICommandList& RHICmdList, ERHIFeat
 		uint32 Index = 0u;
 		while (Index < (uint32)AllocatedVTsToMap.Num())
 		{
-			const IAllocatedVirtualTexture* AllocatedVT = AllocatedVTsToMap[Index];
+			IAllocatedVirtualTexture* AllocatedVT = AllocatedVTsToMap[Index];
 			if (AllocatedVT->TryMapLockedTiles(this))
 			{
 				AllocatedVTsToMap.RemoveAtSwap(Index, EAllowShrinking::No);
+				AllocatedVT->bIsWaitingToMap = false;
 			}
 			else
 			{
@@ -2895,7 +2903,7 @@ float FVirtualTextureSystem::GetGlobalMipBias() const
 bool FVirtualTextureSystem::IsPendingRootPageMap(IAllocatedVirtualTexture* AllocatedVT) const
 {
 	UE::TScopeLock Lock(Mutex);
-	return AllocatedVTsToMap.Find(AllocatedVT) != INDEX_NONE;
+	return AllocatedVT->bIsWaitingToMap;
 }
 
 #if !UE_BUILD_SHIPPING
