Moved TLAS build pass into FRayTracingScene.

diff --git a/Engine/Source/Runtime/Renderer/Private/DeferredShadingRenderer.cpp b/Engine/Source/Runtime/Renderer/Private/DeferredShadingRenderer.cpp
--- a/Engine/Source/Runtime/Renderer/Private/DeferredShadingRenderer.cpp
+++ b/Engine/Source/Runtime/Renderer/Private/DeferredShadingRenderer.cpp
@@ -283,7 +283,6 @@ DECLARE_CYCLE_STAT(TEXT("DeferredShadingSceneRenderer RenderLightShaftBloom"), S
 DECLARE_CYCLE_STAT(TEXT("DeferredShadingSceneRenderer RenderFinish"), STAT_FDeferredShadingSceneRenderer_RenderFinish, STATGROUP_SceneRendering);
 
 DECLARE_GPU_STAT(RayTracingUpdate);
-DECLARE_GPU_STAT(RayTracingScene);
 DECLARE_GPU_STAT(RayTracingGeometry);
 DECLARE_GPU_STAT(RayTracingDynamicGeometry);
 
@@ -577,10 +576,7 @@ static void DeduplicateRayGenerationShaders(TArray< FRHIRayTracingShader*>& RayG
 }
 
 BEGIN_SHADER_PARAMETER_STRUCT(FBuildAccelerationStructurePassParams, )
-	RDG_BUFFER_ACCESS(RayTracingSceneScratchBuffer, ERHIAccess::UAVCompute)
 	RDG_BUFFER_ACCESS(DynamicGeometryScratchBuffer, ERHIAccess::UAVCompute)
-	RDG_BUFFER_ACCESS(RayTracingSceneInstanceBuffer, ERHIAccess::SRVCompute)
-	RDG_BUFFER_ACCESS(RayTracingSceneBuffer, ERHIAccess::BVHWrite)
 
 	SHADER_PARAMETER_STRUCT_REF(FViewUniformShaderParameters, View)
 	SHADER_PARAMETER_RDG_UNIFORM_BUFFER(FSceneUniformParameters, Scene)
@@ -850,8 +846,6 @@ bool FDeferredShadingSceneRenderer::DispatchRayTracingWorldUpdates(FRDGBuilder&
 			RDG_GPU_MASK_SCOPE(GraphBuilder, FRHIGPUMask::All());
 
 			FBuildAccelerationStructurePassParams* PassParams = GraphBuilder.AllocParameters<FBuildAccelerationStructurePassParams>();
-			PassParams->RayTracingSceneScratchBuffer = nullptr;
-			PassParams->RayTracingSceneInstanceBuffer = nullptr;
 			PassParams->View = ReferenceView.ViewUniformBuffer;
 			PassParams->Scene = GetSceneUniforms().GetBuffer(GraphBuilder);
 			PassParams->DynamicGeometryScratchBuffer = OutDynamicGeometryScratchBuffer;
@@ -871,41 +865,8 @@ bool FDeferredShadingSceneRenderer::DispatchRayTracingWorldUpdates(FRDGBuilder&
 			});
 		}
 
-		{
-			FBuildAccelerationStructurePassParams* PassParams = GraphBuilder.AllocParameters<FBuildAccelerationStructurePassParams>();
-			PassParams->RayTracingSceneScratchBuffer = Scene->RayTracingScene.BuildScratchBuffer;
-			PassParams->RayTracingSceneInstanceBuffer = Scene->RayTracingScene.InstanceBuffer;
-			PassParams->View = ReferenceView.ViewUniformBuffer;
-			PassParams->Scene = GetSceneUniforms().GetBuffer(GraphBuilder);
-			PassParams->DynamicGeometryScratchBuffer = OutDynamicGeometryScratchBuffer;
-			PassParams->LightGridPacked = nullptr;
-			PassParams->ClusterPageData = nullptr;
-			PassParams->HierarchyBuffer = nullptr;
-			PassParams->RayTracingDataBuffer = nullptr;
-			PassParams->RayTracingSceneBuffer = Scene->RayTracingScene.GetBufferChecked();
-
-			// Use ERDGPassFlags::NeverParallel here too -- see comment above on the previous pass
-			GraphBuilder.AddPass(RDG_EVENT_NAME("RayTracingUpdate"), PassParams, ComputePassFlags | ERDGPassFlags::NeverCull | ERDGPassFlags::NeverParallel,
-				[this, PassParams, bRayTracingAsyncBuild](FRHICommandListImmediate& RHICmdList)
-			{
-				SCOPED_GPU_STAT(RHICmdList, RayTracingScene);
-
-				FRHIRayTracingScene* RayTracingSceneRHI = Scene->RayTracingScene.GetRHIRayTracingSceneChecked();
-				FRHIBuffer* AccelerationStructureBuffer = PassParams->RayTracingSceneBuffer->GetRHI();
-				FRHIBuffer* ScratchBuffer = PassParams->RayTracingSceneScratchBuffer->GetRHI();
-				FRHIBuffer* InstanceBuffer = PassParams->RayTracingSceneInstanceBuffer->GetRHI();
-
-				FRayTracingSceneBuildParams BuildParams;
-				BuildParams.Scene = RayTracingSceneRHI;
-				BuildParams.ScratchBuffer = ScratchBuffer;
-				BuildParams.ScratchBufferOffset = 0;
-				BuildParams.InstanceBuffer = InstanceBuffer;
-				BuildParams.InstanceBufferOffset = 0;
-
-				RHICmdList.BindAccelerationStructureMemory(RayTracingSceneRHI, AccelerationStructureBuffer, 0);
-				RHICmdList.BuildAccelerationStructure(BuildParams);
-			});
-		}
+		// Use ERDGPassFlags::NeverParallel here too -- see comment above on the previous pass
+		RayTracingScene.Build(GraphBuilder, ComputePassFlags | ERDGPassFlags::NeverCull | ERDGPassFlags::NeverParallel, OutDynamicGeometryScratchBuffer);
 	}
 
 	AddPass(GraphBuilder, RDG_EVENT_NAME("RayTracingEndUpdate"), [this, bRayTracingAsyncBuild](FRHICommandListImmediate& RHICmdList)
@@ -979,7 +940,6 @@ void FDeferredShadingSceneRenderer::WaitForRayTracingScene(FRDGBuilder& GraphBui
 
 	FBuildAccelerationStructurePassParams* PassParams = GraphBuilder.AllocParameters<FBuildAccelerationStructurePassParams>();
 	PassParams->Scene = GetSceneUniformBufferRef(GraphBuilder);
-	PassParams->RayTracingSceneScratchBuffer = nullptr;
 	PassParams->DynamicGeometryScratchBuffer = nullptr;
 	PassParams->LightGridPacked = bIsPathTracing ? nullptr : ReferenceView.RayTracingLightGridUniformBuffer; // accessed by FRayTracingLightingMS // Is this needed for anything?
 	PassParams->LumenHardwareRayTracingUniformBuffer = ReferenceView.LumenHardwareRayTracingUniformBuffer;
diff --git a/Engine/Source/Runtime/Renderer/Private/HeterogeneousVolumes/HeterogeneousVolumesHardwareRayTracing.cpp b/Engine/Source/Runtime/Renderer/Private/HeterogeneousVolumes/HeterogeneousVolumesHardwareRayTracing.cpp
--- a/Engine/Source/Runtime/Renderer/Private/HeterogeneousVolumes/HeterogeneousVolumesHardwareRayTracing.cpp
+++ b/Engine/Source/Runtime/Renderer/Private/HeterogeneousVolumes/HeterogeneousVolumesHardwareRayTracing.cpp
@@ -258,43 +258,19 @@ void GenerateRayTracingScene(
 	const FGPUScene* EmptyGPUScene = nullptr;
 	RayTracingScene.Create(GraphBuilder, View, EmptyGPUScene);
 
-	// Build TLAS
-	FBuildTLASPassParams* PassParamsTLAS = GraphBuilder.AllocParameters<FBuildTLASPassParams>();
-	PassParamsTLAS->RayTracingSceneScratchBuffer = RayTracingScene.BuildScratchBuffer;
-	PassParamsTLAS->RayTracingSceneInstanceBuffer = RayTracingScene.InstanceBuffer;
-	PassParamsTLAS->RayTracingSceneBuffer = RayTracingScene.GetBufferChecked();
-
 	const bool bRayTracingAsyncBuild = false;//CVarRayTracingAsyncBuild.GetValueOnRenderThread() != 0 && GRHISupportsRayTracingAsyncBuildAccelerationStructure;
 	const ERDGPassFlags ComputePassFlags = bRayTracingAsyncBuild ? ERDGPassFlags::AsyncCompute : ERDGPassFlags::Compute;
-	GraphBuilder.AddPass(
-		RDG_EVENT_NAME("RayTracingScene"),
-		PassParamsTLAS,
-		ComputePassFlags | ERDGPassFlags::NeverCull | ERDGPassFlags::NeverParallel,
-		[
-			PassParamsTLAS,
-			&RayTracingScene,
-			bRayTracingAsyncBuild
-		](FRHIComputeCommandList& RHICmdList)
+	RayTracingScene.Build(GraphBuilder, ComputePassFlags | ERDGPassFlags::NeverCull | ERDGPassFlags::NeverParallel, nullptr);
+
+	AddPass(GraphBuilder, RDG_EVENT_NAME("RayTracingEndUpdate"), [&RayTracingScene, bRayTracingAsyncBuild](FRHICommandListImmediate& RHICmdList)
 		{
-			FRHIRayTracingScene* RayTracingSceneRHI = RayTracingScene.GetRHIRayTracingSceneChecked();
-			FRHIBuffer* AccelerationStructureBuffer = PassParamsTLAS->RayTracingSceneBuffer->GetRHI();
-
-			FRayTracingSceneBuildParams SceneBuildParams;
-			SceneBuildParams.Scene = RayTracingSceneRHI;
-			SceneBuildParams.ScratchBuffer = PassParamsTLAS->RayTracingSceneScratchBuffer->GetRHI();
-			SceneBuildParams.ScratchBufferOffset = 0;
-			SceneBuildParams.InstanceBuffer = PassParamsTLAS->RayTracingSceneInstanceBuffer->GetRHI();
-			SceneBuildParams.InstanceBufferOffset = 0;
-
-			RHICmdList.BindAccelerationStructureMemory(RayTracingSceneRHI, AccelerationStructureBuffer, 0);
-			RHICmdList.BuildAccelerationStructure(SceneBuildParams);
-			// Submit potentially expensive BVH build commands to the GPU as soon as possible.
-			// Avoids a GPU bubble in some CPU-limited cases.
-			RHICmdList.SubmitCommandsHint();
-
-			RHICmdList.Transition(FRHITransitionInfo(RayTracingSceneRHI, ERHIAccess::BVHWrite, ERHIAccess::BVHRead));
-		}
-	);
+			if (!bRayTracingAsyncBuild)
+			{
+				// Submit potentially expensive BVH build commands to the GPU as soon as possible.
+				// Avoids a GPU bubble in some CPU-limited cases.
+				RHICmdList.SubmitCommandsHint();
+			}
+		});
 }
 
 IMPLEMENT_RT_PAYLOAD_TYPE(ERayTracingPayloadType::SparseVoxel, 28);
diff --git a/Engine/Source/Runtime/Renderer/Private/RayTracing/RayTracingDebug.cpp b/Engine/Source/Runtime/Renderer/Private/RayTracing/RayTracingDebug.cpp
--- a/Engine/Source/Runtime/Renderer/Private/RayTracing/RayTracingDebug.cpp
+++ b/Engine/Source/Runtime/Renderer/Private/RayTracing/RayTracingDebug.cpp
@@ -802,6 +802,8 @@ void FDeferredShadingSceneRenderer::PrepareRayTracingDebug(const FSceneViewFamil
 
 static FRDGBufferRef RayTracingPerformPicking(FRDGBuilder& GraphBuilder, const FScene* Scene, const FViewInfo& View, FRayTracingPickingFeedback& PickingFeedback)
 {
+	const FRayTracingScene& RayTracingScene = Scene->RayTracingScene;
+
 	FGlobalShaderMap* ShaderMap = GetGlobalShaderMap(GMaxRHIFeatureLevel);
 	auto RayGenShader = ShaderMap->GetShader<FRayTracingPickingRGS>();
 
@@ -834,10 +836,10 @@ static FRDGBufferRef RayTracingPerformPicking(FRDGBuilder& GraphBuilder, const F
 	FRDGBufferRef PickingBuffer = GraphBuilder.CreateBuffer(PickingBufferDesc, TEXT("RayTracingDebug.PickingBuffer"));
 
 	FRayTracingPickingRGS::FParameters* RayGenParameters = GraphBuilder.AllocParameters<FRayTracingPickingRGS::FParameters>();
-	RayGenParameters->InstancesDebugData = GraphBuilder.CreateSRV(Scene->RayTracingScene.InstanceDebugBuffer);
-	RayGenParameters->TLAS = Scene->RayTracingScene.GetLayerView(ERayTracingSceneLayer::Base);
+	RayGenParameters->InstancesDebugData = GraphBuilder.CreateSRV(RayTracingScene.InstanceDebugBuffer);
+	RayGenParameters->TLAS = RayTracingScene.GetLayerView(ERayTracingSceneLayer::Base);
 	RayGenParameters->OpaqueOnly = CVarRayTracingDebugModeOpaqueOnly.GetValueOnRenderThread();
-	RayGenParameters->InstanceBuffer = GraphBuilder.CreateSRV(Scene->RayTracingScene.InstanceBuffer);
+	RayGenParameters->InstanceBuffer = GraphBuilder.CreateSRV(RayTracingScene.GetInstanceBuffer());
 	RayGenParameters->ViewUniformBuffer = View.ViewUniformBuffer;
 	RayGenParameters->SceneUniformBuffer = GetSceneUniformBufferRef(GraphBuilder, View); // TODO: use a separate params structure
 	RayGenParameters->PickingOutput = GraphBuilder.CreateUAV(PickingBuffer);
diff --git a/Engine/Source/Runtime/Renderer/Private/RayTracing/RayTracingScene.cpp b/Engine/Source/Runtime/Renderer/Private/RayTracing/RayTracingScene.cpp
--- a/Engine/Source/Runtime/Renderer/Private/RayTracing/RayTracingScene.cpp
+++ b/Engine/Source/Runtime/Renderer/Private/RayTracing/RayTracingScene.cpp
@@ -15,6 +15,8 @@
 #include "SceneRendering.h"
 #include "RayTracingInstanceCulling.h"
 
+DECLARE_GPU_STAT(RayTracingScene);
+
 BEGIN_SHADER_PARAMETER_STRUCT(FBuildInstanceBufferPassParams, )
 	SHADER_PARAMETER_RDG_BUFFER_UAV(RWStructuredBuffer, InstanceBuffer)
 	SHADER_PARAMETER_RDG_BUFFER_UAV(RWStructuredBuffer, DebugInstanceGPUSceneIndexBuffer)
@@ -224,7 +226,7 @@ void FRayTracingScene::CreateWithInitializationData(FRDGBuilder& GraphBuilder, c
 		PassParams->Scene = View.GetSceneUniforms().GetBuffer(GraphBuilder);
 
 		GraphBuilder.AddPass(
-			RDG_EVENT_NAME("BuildTLASInstanceBuffer"),
+			RDG_EVENT_NAME("RayTracingBuildInstanceBuffer"),
 			PassParams,
 			ERDGPassFlags::Compute,
 			[PassParams,
@@ -282,6 +284,39 @@ void FRayTracingScene::CreateWithInitializationData(FRDGBuilder& GraphBuilder, c
 	}
 }
 
+BEGIN_SHADER_PARAMETER_STRUCT(FRayTracingSceneBuildPassParams, )
+	RDG_BUFFER_ACCESS(ScratchBuffer, ERHIAccess::UAVCompute)
+	RDG_BUFFER_ACCESS(InstanceBuffer, ERHIAccess::SRVCompute)
+	RDG_BUFFER_ACCESS(TLASBuffer, ERHIAccess::BVHWrite)
+
+	RDG_BUFFER_ACCESS(DynamicGeometryScratchBuffer, ERHIAccess::UAVCompute)
+END_SHADER_PARAMETER_STRUCT()
+
+void FRayTracingScene::Build(FRDGBuilder& GraphBuilder, ERDGPassFlags ComputePassFlags, FRDGBufferRef DynamicGeometryScratchBuffer)
+{
+	FRayTracingSceneBuildPassParams* PassParams = GraphBuilder.AllocParameters<FRayTracingSceneBuildPassParams>();
+	PassParams->ScratchBuffer = BuildScratchBuffer;
+	PassParams->InstanceBuffer = InstanceBuffer;
+	PassParams->TLASBuffer = GetBufferChecked();
+	PassParams->DynamicGeometryScratchBuffer = DynamicGeometryScratchBuffer; // TODO: Is this necessary?
+
+	GraphBuilder.AddPass(RDG_EVENT_NAME("RayTracingBuildScene"), PassParams, ComputePassFlags,
+		[this, PassParams](FRHICommandListImmediate& RHICmdList)
+		{
+			SCOPED_GPU_STAT(RHICmdList, RayTracingScene);
+
+			FRayTracingSceneBuildParams BuildParams;
+			BuildParams.Scene = RayTracingSceneRHI;
+			BuildParams.ScratchBuffer = PassParams->ScratchBuffer->GetRHI();
+			BuildParams.ScratchBufferOffset = 0;
+			BuildParams.InstanceBuffer = PassParams->InstanceBuffer->GetRHI();
+			BuildParams.InstanceBufferOffset = 0;
+
+			RHICmdList.BindAccelerationStructureMemory(RayTracingSceneRHI, PassParams->TLASBuffer->GetRHI(), 0);
+			RHICmdList.BuildAccelerationStructure(BuildParams);
+		});
+}
+
 void FRayTracingScene::WaitForTasks() const
 {
 	if (FillInstanceUploadBufferTask.IsValid())
diff --git a/Engine/Source/Runtime/Renderer/Private/RayTracing/RayTracingScene.h b/Engine/Source/Runtime/Renderer/Private/RayTracing/RayTracingScene.h
--- a/Engine/Source/Runtime/Renderer/Private/RayTracing/RayTracingScene.h
+++ b/Engine/Source/Runtime/Renderer/Private/RayTracing/RayTracingScene.h
@@ -66,6 +66,8 @@ public:
 	// Backwards-compatible version of Create() which internally calls CreateRayTracingSceneWithGeometryInstances().
 	void Create(FRDGBuilder& GraphBuilder, const FViewInfo& View, const FGPUScene* GPUScene);
 
+	void Build(FRDGBuilder& GraphBuilder, ERDGPassFlags ComputePassFlags, FRDGBufferRef DynamicGeometryScratchBuffer);
+
 	// Resets the instance list and reserves memory for this frame.
 	void Reset(bool bInstanceDebugDataEnabled);
 
@@ -99,6 +101,8 @@ public:
 	// Returns RDG view of a layer. Can only be used on valid ray tracing scene.
 	RENDERER_API FRDGBufferSRVRef GetLayerView(ERayTracingSceneLayer Layer) const;
 
+	FRDGBufferRef GetInstanceBuffer() const { return InstanceBuffer; }
+
 	TArrayView<const FRayTracingGeometryInstance> GetInstances() const { return MakeArrayView(Instances); }
 
 	FRayTracingGeometryInstance& GetInstance(uint32 InstanceIndex) { return Instances[InstanceIndex]; }
@@ -122,9 +126,6 @@ public:
 	// Used coarse mesh streaming handles during the last TLAS build
 	TArray<Nanite::CoarseMeshStreamingHandle> UsedCoarseMeshStreamingHandles;
 
-	FRDGBufferRef InstanceBuffer;
-	FRDGBufferRef BuildScratchBuffer;
-
 	// Special data for debugging purposes
 	FRDGBufferRef InstanceDebugBuffer = nullptr;
 
@@ -136,7 +137,10 @@ public:
 private:
 	void WaitForTasks() const;
 
-	// RHI object that abstracts mesh instnaces in this scene
+	FRDGBufferRef InstanceBuffer;
+	FRDGBufferRef BuildScratchBuffer;
+
+	// RHI object that abstracts mesh instances in this scene
 	FRayTracingSceneRHIRef RayTracingSceneRHI;
 
 	// Persistently allocated buffer that holds the built TLAS
