LightFunctionAtlas - better atlas visualisation not showing flierking uninitialise memory.

diff --git a/Engine/Source/Runtime/Renderer/Private/LightFunctionAtlas.cpp b/Engine/Source/Runtime/Renderer/Private/LightFunctionAtlas.cpp
--- a/Engine/Source/Runtime/Renderer/Private/LightFunctionAtlas.cpp
+++ b/Engine/Source/Runtime/Renderer/Private/LightFunctionAtlas.cpp
@@ -819,28 +819,47 @@ FScreenPassTexture FLightFunctionAtlas::AddDebugVisualizationPasses(FRDGBuilder&
 		return ScreenPassSceneColor;
 	}
 
-	uint32 AtlasSlotResolution = GetAtlasSlotResolution();
-	uint32 AtlasEdgeSize = GetAtlasEdgeSize();
-	uint32 AtlasResolution = AtlasSlotResolution * AtlasEdgeSize;
+	const uint32 AtlasSlotResolution	= GetAtlasSlotResolution();
+	const uint32 AtlasEdgeSize			= GetAtlasEdgeSize();
+	const uint32 AtlasResolution		= AtlasSlotResolution * AtlasEdgeSize;
+
+	const FIntPoint SrcPoint			= FIntPoint::ZeroValue;
+	const FIntPoint SrcSize				= RDGAtlasTexture2D->Desc.Extent;
+	const FIntPoint DstPoint			= FIntPoint(100, 100);
+	const FIntPoint DstSize				= FIntPoint(512, 512);
 
-	const FIntPoint SrcPoint = FIntPoint::ZeroValue;
-	const FIntPoint SrcSize  = RDGAtlasTexture2D->Desc.Extent;
-	const FIntPoint DstPoint = FIntPoint(100, 100);
-	const FIntPoint DstSize  = FIntPoint(512, 512);
+	const float DisplayResolutionRatio	= float(DstSize.X) / float(AtlasResolution);
 
+	// Draw the atlas, first make it all grey,
 	AddDrawTexturePass(
 		GraphBuilder,
 		View,
-		RDGAtlasTexture2D,
+		GSystemTextures.GetMidGreyDummy(GraphBuilder),
 		ScreenPassSceneColor.Texture,
 		SrcPoint,
 		SrcSize,
 		DstPoint,
 		DstSize);
+	// then compose all tiles independentely.  (to not show flickering/uninitialised memory on unused tiles)
+	for (auto& AtlasSlot : EffectiveLightFunctionSlotArray)
+	{
+		const FIntPoint SlotDstPoint	= FIntPoint(DstPoint.X + float(AtlasSlot.Min.X) * DisplayResolutionRatio, DstPoint.Y + float(AtlasSlot.Min.Y) * DisplayResolutionRatio);
+		const FIntPoint SlotDstSize		= AtlasSlotResolution * DisplayResolutionRatio;
+		AddDrawTexturePass(
+			GraphBuilder,
+			View,
+			RDGAtlasTexture2D,
+			ScreenPassSceneColor.Texture,
+			AtlasSlot.Min,
+			AtlasSlot.Max-AtlasSlot.Min,
+			SlotDstPoint,
+			SlotDstSize);
+	}
+
 
 	// Now debug print
 	AddDrawCanvasPass(GraphBuilder, {}, View, FScreenPassRenderTarget(ScreenPassSceneColor, ERenderTargetLoadAction::ELoad),
-		[&View, this, DstPoint, DstSize, AtlasResolution](FCanvas& Canvas)
+		[&View, this, DstPoint, DstSize, AtlasResolution, DisplayResolutionRatio](FCanvas& Canvas)
 	{
 		FString Text;
 
@@ -852,8 +871,6 @@ FScreenPassTexture FLightFunctionAtlas::AddDebugVisualizationPasses(FRDGBuilder&
 		const float DPIScale = Canvas.GetDPIScale();
 		Canvas.SetBaseTransform(FMatrix(FScaleMatrix(DPIScale) * Canvas.CalcBaseTransform2D(Canvas.GetViewRect().Width(), Canvas.GetViewRect().Height())));
 
-		float DisplayResolutionRatio = float(DstSize.X) / float(AtlasResolution);
-
 		Canvas.DrawShadowedString(DstPoint.X + 180.0f, DstPoint.Y - 20.0f, TEXT("LIGHT FUNCTION ATLAS"), GEngine->GetLargeFont(), FLinearColor::White);
 
 		Text = FString::Printf(TEXT("Light Functions in atlas:         %d"), EffectiveLightFunctionSlotArray.Num());
