Move wait for RDG deletion tasks to happen when the scene is destroyed rather than during each scene update. This avoids a full sync of the RDG async deletion task with every frame.

diff --git a/Engine/Source/Runtime/Renderer/Private/RendererScene.cpp b/Engine/Source/Runtime/Renderer/Private/RendererScene.cpp
--- a/Engine/Source/Runtime/Renderer/Private/RendererScene.cpp
+++ b/Engine/Source/Runtime/Renderer/Private/RendererScene.cpp
@@ -4668,9 +4668,6 @@ void FScene::Release()
 	ENQUEUE_RENDER_COMMAND(FReleaseCommand)(
 		[Scene](FRHICommandListImmediate& RHICmdList)
 		{
-			// Finish all pending scene rendering destruction work before deleting the scene.
-			FSceneRenderer::CleanUp(RHICmdList);
-
 			// Flush any remaining batched primitive update commands before deleting the scene.
 			FUpdateParameters UpdateParameters;
 			UpdateParameters.bDestruction = true;
@@ -4682,6 +4679,9 @@ void FScene::Release()
 				GraphBuilder.Execute();
 			}
 
+			// Wait for RDG to complete async deletion as scene extensions can be allocated through RDG.
+			FRDGBuilder::WaitForAsyncDeleteTask();
+
 			delete Scene;
 		});
 }
diff --git a/Engine/Source/Runtime/Renderer/Private/SceneRendering.cpp b/Engine/Source/Runtime/Renderer/Private/SceneRendering.cpp
--- a/Engine/Source/Runtime/Renderer/Private/SceneRendering.cpp
+++ b/Engine/Source/Runtime/Renderer/Private/SceneRendering.cpp
@@ -4730,7 +4730,6 @@ static void FinishCleanUp(FRHICommandListImmediate& RHICmdList)
 	GlobalDynamicBuffer::GarbageCollect();
 	GPrimitiveIdVertexBufferPool.DiscardAll();
 	FGraphicsMinimalPipelineStateId::ResetLocalPipelineIdTableSize();
-	FRDGBuilder::WaitForAsyncDeleteTask();
 }
 
 static void DeleteSceneRenderers(const TArray<FSceneRenderer*>& SceneRenderers)
