Remove r.AllowCachedUniformExpressions cvar
Setting that to false isn't currently supported, since it requires getting the immediate command list on a worker thread. There's little value in fixing that since there's no real use case for disabling uniform caching.
Also converted ensure() to check() for the error condition of having an outdated shader map. The old fallback path would have hit a render thread check anyway.

diff --git a/Engine/Source/Runtime/Renderer/Private/ShaderBaseClasses.cpp b/Engine/Source/Runtime/Renderer/Private/ShaderBaseClasses.cpp
--- a/Engine/Source/Runtime/Renderer/Private/ShaderBaseClasses.cpp
+++ b/Engine/Source/Runtime/Renderer/Private/ShaderBaseClasses.cpp
@@ -18,15 +18,6 @@ IMPLEMENT_TYPE_LAYOUT(FMaterialShader);
 IMPLEMENT_TYPE_LAYOUT(FMeshMaterialShader);
 IMPLEMENT_TYPE_LAYOUT(FDebugUniformExpressionSet);
 
-/** If true, cached uniform expressions are allowed. */
-int32 FMaterialShader::bAllowCachedUniformExpressions = true;
-
-/** Console variable ref to toggle cached uniform expressions. */
-FAutoConsoleVariableRef FMaterialShader::CVarAllowCachedUniformExpressions(
-	TEXT("r.AllowCachedUniformExpressions"),
-	bAllowCachedUniformExpressions,
-	TEXT("Allow uniform expressions to be cached."),
-	ECVF_RenderThreadSafe);
 
 void FMeshMaterialShaderElementData::InitializeMeshMaterialData()
 {
@@ -274,36 +265,23 @@ void FMaterialShader::SetParameters(
 	bool bUniformExpressionCacheNeedsDelete = false;
 	//bool bForceExpressionEvaluation = false;
 
-#if !(UE_BUILD_TEST || UE_BUILD_SHIPPING || !WITH_EDITOR)
-	if (bAllowCachedUniformExpressions)
+#if DO_CHECK
+	if (UniformExpressionCache->CachedUniformExpressionShaderMap != ShaderMap)
 	{
-		// UE-46061 - Workaround for a rare crash with an outdated cached shader map
-		if (UniformExpressionCache->CachedUniformExpressionShaderMap != ShaderMap)
-		{
-			UMaterialInterface* MtlInterface = Material.GetMaterialInterface();
-			UMaterialInterface* ProxyInterface = MaterialRenderProxy->GetMaterialInterface();
-
-			const FShaderType* ShaderType = GetType(ShaderMap->GetPointerTable());
-			ensureMsgf(false,
-				TEXT("%s shader uniform expression set mismatched shader map for material %s/%s, forcing expression cache evaluation.\n")
-				TEXT("Material:  %s\n")
-				TEXT("Proxy:  %s\n"),
-				ShaderType->GetName(),
-				*MaterialRenderProxy->GetFriendlyName(), *Material.GetFriendlyName(),
-				MtlInterface ? *MtlInterface->GetFullName() : TEXT("nullptr"),
-				ProxyInterface ? *ProxyInterface->GetFullName() : TEXT("nullptr"));
-			//bForceExpressionEvaluation = true;
-		}
-	}
-#endif
+		UMaterialInterface* MtlInterface = Material.GetMaterialInterface();
+		UMaterialInterface* ProxyInterface = MaterialRenderProxy->GetMaterialInterface();
 
-	if (!bAllowCachedUniformExpressions || UniformExpressionCache->CachedUniformExpressionShaderMap != ShaderMap)
-	{
-		FMaterialRenderContext MaterialRenderContext(MaterialRenderProxy, Material, &View);
-		bUniformExpressionCacheNeedsDelete = true;
-		UniformExpressionCache = new FUniformExpressionCache();
-		MaterialRenderProxy->EvaluateUniformExpressions(FRHICommandListImmediate::Get(), *UniformExpressionCache, MaterialRenderContext);
+		const FShaderType* ShaderType = GetType(ShaderMap->GetPointerTable());
+		checkf(false,
+			TEXT("'%s' shader uniform expression set mismatched shader map for material '%s/%s'.\n")
+			TEXT("Material: '%s'\n")
+			TEXT("Proxy: '%s'\n"),
+			ShaderType->GetName(),
+			*MaterialRenderProxy->GetFriendlyName(), *Material.GetFriendlyName(),
+			MtlInterface ? *MtlInterface->GetFullName() : TEXT("nullptr"),
+			ProxyInterface ? *ProxyInterface->GetFullName() : TEXT("nullptr"));
 	}
+#endif
 
 	SetUniformBufferParameter(BatchedParameters, MaterialUniformBuffer, UniformExpressionCache->UniformBuffer);
 
