Ray Tracing - Run RenderDitheredLODFadingOutMask only when RT shadows are enabled. Apart from another render pass overhead, it writes to stencil interesting a sync point, which prevents the Lumen DiffuseIndirectAndAO from overlapping with Lights.

Replay on console: ~0.3ms win on average

Also:
* Use GetLightOcclusionType() == FLightOcclusionType::Raytraced instead of ShouldRenderRayTracingShadowsForLight for testing whether we should use RT shadows.
* Renamed bHasRayTracedLights to bHasLightsWithRayTracedShadows
As those conditions are specifically tied to RT shadows and shouldn?t be true for any other shadowing / direct lighting method.

diff --git a/Engine/Source/Runtime/Renderer/Private/DeferredShadingRenderer.cpp b/Engine/Source/Runtime/Renderer/Private/DeferredShadingRenderer.cpp
--- a/Engine/Source/Runtime/Renderer/Private/DeferredShadingRenderer.cpp
+++ b/Engine/Source/Runtime/Renderer/Private/DeferredShadingRenderer.cpp
@@ -2690,7 +2690,8 @@ void FDeferredShadingSceneRenderer::Render(FRDGBuilder& GraphBuilder)
 			}
 
 #if RHI_RAYTRACING
-			if (IsRayTracingEnabled())
+			// Only used by ray traced shadows
+			if (IsRayTracingEnabled() && Scene->bHasLightsWithRayTracedShadows)
 			{
 				RenderDitheredLODFadingOutMask(GraphBuilder, Views[0], SceneTextures.Depth.Target);
 			}
@@ -3378,7 +3379,7 @@ bool AnyRayTracingPassEnabled(const FScene* Scene, const FViewInfo& View)
 		|| ShouldRenderRayTracingTranslucency(View)
 		|| ShouldRenderRayTracingSkyLight(Scene->SkyLight, View.GetShaderPlatform())
 		|| ShouldRenderRayTracingShadows()
-		|| Scene->bHasRayTracedLights
+		|| Scene->bHasLightsWithRayTracedShadows
 		|| ShouldRenderPluginRayTracingGlobalIllumination(View)
         || Lumen::AnyLumenHardwareRayTracingPassEnabled(Scene, View)
 		|| ManyLights::UseHardwareRayTracing(*View.Family)
diff --git a/Engine/Source/Runtime/Renderer/Private/PrimitiveSceneInfo.cpp b/Engine/Source/Runtime/Renderer/Private/PrimitiveSceneInfo.cpp
--- a/Engine/Source/Runtime/Renderer/Private/PrimitiveSceneInfo.cpp
+++ b/Engine/Source/Runtime/Renderer/Private/PrimitiveSceneInfo.cpp
@@ -988,7 +988,7 @@ void FScene::UpdateRayTracedLights()
 	// Whether a light can use ray traced shadows depends on CVars that may be changed at runtime.
 	// It is not enough to check the light shadow mode when light is added to the scene. This must be done during rendering.
 
-	bHasRayTracedLights = false;
+	bHasLightsWithRayTracedShadows = false;
 
 	if (!IsRayTracingEnabled())
 	{
@@ -998,9 +998,9 @@ void FScene::UpdateRayTracedLights()
 	// We currently don't need a full list of RT lights, only whether there are any RT lights at all.
 	for (const FLightSceneInfoCompact& LightSceneInfoCompact : Lights)
 	{
-		if (ShouldRenderRayTracingShadowsForLight(LightSceneInfoCompact))
+		if (GetLightOcclusionType(LightSceneInfoCompact) == FLightOcclusionType::Raytraced)
 		{
-			bHasRayTracedLights = true;
+			bHasLightsWithRayTracedShadows = true;
 			break;
 		}
 	}
diff --git a/Engine/Source/Runtime/Renderer/Private/RayTracing/RayTracingShadows.cpp b/Engine/Source/Runtime/Renderer/Private/RayTracing/RayTracingShadows.cpp
--- a/Engine/Source/Runtime/Renderer/Private/RayTracing/RayTracingShadows.cpp
+++ b/Engine/Source/Runtime/Renderer/Private/RayTracing/RayTracingShadows.cpp
@@ -306,7 +306,7 @@ void FDeferredShadingSceneRenderer::PrepareRayTracingShadows(const FViewInfo& Vi
 	// We have to check if ray tracing is enabled on any of the scene lights. The Scene.bHasRayTracedLights is computed using ShouldRenderRayTracingShadowsForLight() helper, 
 	// which handles various override conditions.
 
-	if (Scene.bHasRayTracedLights == false)
+	if (Scene.bHasLightsWithRayTracedShadows == false)
 	{
 		return;
 	}
diff --git a/Engine/Source/Runtime/Renderer/Private/ScenePrivate.h b/Engine/Source/Runtime/Renderer/Private/ScenePrivate.h
--- a/Engine/Source/Runtime/Renderer/Private/ScenePrivate.h
+++ b/Engine/Source/Runtime/Renderer/Private/ScenePrivate.h
@@ -3121,7 +3121,7 @@ public:
 	FRayTracingScene RayTracingScene;
 	FRayTracingScene HeterogeneousVolumesRayTracingScene;
 
-	bool bHasRayTracedLights = false;
+	bool bHasLightsWithRayTracedShadows = false;
 	void UpdateRayTracedLights();
 #endif // RHI_RAYTRACING
 
diff --git a/Engine/Source/Runtime/Renderer/Private/SceneVisibility.cpp b/Engine/Source/Runtime/Renderer/Private/SceneVisibility.cpp
--- a/Engine/Source/Runtime/Renderer/Private/SceneVisibility.cpp
+++ b/Engine/Source/Runtime/Renderer/Private/SceneVisibility.cpp
@@ -1554,7 +1554,8 @@ void FRelevancePacket::ComputeRelevance(FDynamicPrimitiveIndexList& DynamicPrimi
 									}
 								}
 #if RHI_RAYTRACING
-								if (IsRayTracingEnabled())
+								// Only used by ray traced shadows
+								if (IsRayTracingEnabled() && Scene.bHasLightsWithRayTracedShadows)
 								{
 									if (MarkMask & EMarkMaskBits::StaticMeshFadeOutDitheredLODMapMask)
 									{
diff --git a/Engine/Source/Runtime/Renderer/Private/VolumetricFog.cpp b/Engine/Source/Runtime/Renderer/Private/VolumetricFog.cpp
--- a/Engine/Source/Runtime/Renderer/Private/VolumetricFog.cpp
+++ b/Engine/Source/Runtime/Renderer/Private/VolumetricFog.cpp
@@ -515,7 +515,7 @@ static void RenderRaytracedDirectionalShadowVolume(
 
 void FDeferredShadingSceneRenderer::PrepareRayTracingVolumetricFogShadows(const FViewInfo& View, const FScene& Scene, TArray<FRHIRayTracingShader*>& OutRayGenShaders)
 {
-	const bool bEnabled = Scene.bHasRayTracedLights && ::ShouldRenderVolumetricFog(&Scene, *View.Family) && GVolumetricFogInjectRaytracedLights;
+	const bool bEnabled = Scene.bHasLightsWithRayTracedShadows && ::ShouldRenderVolumetricFog(&Scene, *View.Family) && GVolumetricFogInjectRaytracedLights;
 	if (!bEnabled)
 	{
 		return;
@@ -564,7 +564,7 @@ const FProjectedShadowInfo* GetShadowForInjectionIntoVolumetricFog(const FVisibl
 
 bool LightHasRayTracedShadows(const FLightSceneInfo* LightSceneInfo)
 {
-	return ShouldRenderRayTracingShadowsForLight(*LightSceneInfo->Proxy) && GVolumetricFogInjectRaytracedLights;
+	return GetLightOcclusionType(*LightSceneInfo->Proxy) == FLightOcclusionType::Raytraced && GVolumetricFogInjectRaytracedLights;
 }
 
 bool LightNeedsSeparateInjectionIntoVolumetricFogForOpaqueShadow(const FViewInfo& View, const FLightSceneInfo* LightSceneInfo, const FVisibleLightInfo& VisibleLightInfo)
