PR #11530: Fixed Precomputed Visibility

#11591

diff --git a/Engine/Source/Runtime/Renderer/Private/SceneOcclusion.cpp b/Engine/Source/Runtime/Renderer/Private/SceneOcclusion.cpp
--- a/Engine/Source/Runtime/Renderer/Private/SceneOcclusion.cpp
+++ b/Engine/Source/Runtime/Renderer/Private/SceneOcclusion.cpp
@@ -150,7 +150,7 @@ static FGlobalBoundShaderState GOcclusionTestBoundShaderState;
  * This method decompresses data if necessary and caches it based on the bucket and chunk index in the view state.
  * InScene is passed in, as the Scene pointer in the class itself may be null, if it was allocated without a scene.
  */
-const uint8* FSceneViewState::GetPrecomputedVisibilityData(FViewInfo& View, const FScene* InScene)
+const uint8* FSceneViewState::ResolvePrecomputedVisibilityData(FViewInfo& View, const FScene* InScene)
 {
 	const uint8* PrecomputedVisibilityData = NULL;
 	if (InScene->PrecomputedVisibilityHandler && GAllowPrecomputedVisibility && View.Family->EngineShowFlags.PrecomputedVisibility)
@@ -176,12 +176,22 @@ const uint8* FSceneViewState::GetPrecomputedVisibilityData(FViewInfo& View, cons
 			}
 		}
 
+		//Determine view origin
+		FVector ViewOrigin = View.ViewMatrices.GetViewOrigin();
+#if !(UE_BUILD_SHIPPING || UE_BUILD_TEST)
+		if (const FViewMatrices* FrozenViewMatrices = GetFrozenViewMatrices())
+		{
+			// Use the frozen view for culling so we can test that it's working
+			ViewOrigin = FrozenViewMatrices->GetViewOrigin();
+		}
+#endif
+
 		// Calculate the bucket that ViewOrigin falls into
 		// Cells are hashed into buckets to reduce search time
-		const float FloatOffsetX = (View.ViewMatrices.GetViewOrigin().X - Handler.PrecomputedVisibilityCellBucketOriginXY.X) / Handler.PrecomputedVisibilityCellSizeXY;
+		const float FloatOffsetX = (ViewOrigin.X - Handler.PrecomputedVisibilityCellBucketOriginXY.X) / Handler.PrecomputedVisibilityCellSizeXY;
 		// FMath::TruncToInt rounds toward 0, we want to always round down
 		const int32 BucketIndexX = FMath::Abs((FMath::TruncToInt(FloatOffsetX) - (FloatOffsetX < 0.0f ? 1 : 0)) / Handler.PrecomputedVisibilityCellBucketSizeXY % Handler.PrecomputedVisibilityNumCellBuckets);
-		const float FloatOffsetY = (View.ViewMatrices.GetViewOrigin().Y -Handler.PrecomputedVisibilityCellBucketOriginXY.Y) / Handler.PrecomputedVisibilityCellSizeXY;
+		const float FloatOffsetY = (ViewOrigin.Y -Handler.PrecomputedVisibilityCellBucketOriginXY.Y) / Handler.PrecomputedVisibilityCellSizeXY;
 		const int32 BucketIndexY = FMath::Abs((FMath::TruncToInt(FloatOffsetY) - (FloatOffsetY < 0.0f ? 1 : 0)) / Handler.PrecomputedVisibilityCellBucketSizeXY % Handler.PrecomputedVisibilityNumCellBuckets);
 		const int32 PrecomputedVisibilityBucketIndex = BucketIndexY * Handler.PrecomputedVisibilityCellBucketSizeXY + BucketIndexX;
 
@@ -193,7 +203,7 @@ const uint8* FSceneViewState::GetPrecomputedVisibilityData(FViewInfo& View, cons
 			// Construct the cell's bounds
 			const FBox CellBounds(CurrentCell.Min, CurrentCell.Min + FVector(Handler.PrecomputedVisibilityCellSizeXY, Handler.PrecomputedVisibilityCellSizeXY, Handler.PrecomputedVisibilityCellSizeZ));
 			// Check if ViewOrigin is inside the current cell
-			if (CellBounds.IsInside(View.ViewMatrices.GetViewOrigin()))
+			if (CellBounds.IsInside(ViewOrigin))
 			{
 				// Reuse a cached decompressed chunk if possible
 				if (CachedVisibilityChunk
diff --git a/Engine/Source/Runtime/Renderer/Private/ScenePrivate.h b/Engine/Source/Runtime/Renderer/Private/ScenePrivate.h
--- a/Engine/Source/Runtime/Renderer/Private/ScenePrivate.h
+++ b/Engine/Source/Runtime/Renderer/Private/ScenePrivate.h
@@ -841,7 +841,7 @@ public:
 	/** Storage to which compressed visibility chunks are uncompressed at runtime. */
 	TArray<uint8> DecompressedVisibilityChunk;
 
-	/** Cached visibility data from the last call to GetPrecomputedVisibilityData. */
+	/** Cached visibility data from the last call to ResolvePrecomputedVisibilityData. */
 	const TArray<uint8>* CachedVisibilityChunk;
 	int32 CachedVisibilityHandlerId;
 	int32 CachedVisibilityBucketIndex;
@@ -1269,7 +1269,7 @@ public:
 	 * This method decompresses data if necessary and caches it based on the bucket and chunk index in the view state.
 	 * InScene is passed in, as the Scene pointer in the class itself may be null, if it was allocated without a scene.
 	 */
-	const uint8* GetPrecomputedVisibilityData(FViewInfo& View, const FScene* InScene);
+	const uint8* ResolvePrecomputedVisibilityData(FViewInfo& View, const FScene* InScene);
 
 	/**
 	 * Cleans out old entries from the primitive occlusion history, and resets unused pending occlusion queries.
diff --git a/Engine/Source/Runtime/Renderer/Private/SceneVisibility.cpp b/Engine/Source/Runtime/Renderer/Private/SceneVisibility.cpp
--- a/Engine/Source/Runtime/Renderer/Private/SceneVisibility.cpp
+++ b/Engine/Source/Runtime/Renderer/Private/SceneVisibility.cpp
@@ -4140,7 +4140,10 @@ void FVisibilityTaskData::LaunchVisibilityTasks(const UE::Tasks::FTask& BeginIni
 		if (ViewPacket.ViewState)
 		{
 			SCOPE_CYCLE_COUNTER(STAT_DecompressPrecomputedOcclusion);
-			ViewPacket.View.PrecomputedVisibilityData = ViewPacket.ViewState->GetPrecomputedVisibilityData(ViewPacket.View, &Scene);
+			if (ViewPacket.View.PrecomputedVisibilityData = ViewPacket.ViewState->ResolvePrecomputedVisibilityData(ViewPacket.View, &Scene))
+			{
+				SceneRenderer.bUsedPrecomputedVisibility = true;
+			}
 		}
 	}
 
@@ -4528,19 +4531,19 @@ void FVisibilityTaskData::ProcessRenderThreadTasks()
 
 		for (FVisibilityViewPacket& ViewPacket : ViewPackets)
 		{
+			SCOPED_NAMED_EVENT(OcclusionCull, FColor::Magenta);
+			const int32 NumCulledPrimitives = PrecomputedOcclusionCull(ViewPacket, PrimitiveRange);
+
 			if (ViewPacket.OcclusionCull.ContextIfSerial)
 			{
-				SCOPED_NAMED_EVENT(OcclusionCull, FColor::Magenta);
-				const int32 NumCulledPrimitives = PrecomputedOcclusionCull(ViewPacket, PrimitiveRange);
-
 				ViewPacket.OcclusionCull.ContextIfSerial->Map(RHICmdList);
 				ViewPacket.OcclusionCull.ContextIfSerial->AddPrimitives(PrimitiveRange);
 				ViewPacket.OcclusionCull.ContextIfSerial->Unmap(RHICmdList);
+			}
 
-				if (NumCulledPrimitives > 0)
-				{
-					TaskConfig.OcclusionCull.NumCulledPrimitives.fetch_add(NumCulledPrimitives, std::memory_order_relaxed);
-				}
+			if (NumCulledPrimitives > 0)
+			{
+				TaskConfig.OcclusionCull.NumCulledPrimitives.fetch_add(NumCulledPrimitives, std::memory_order_relaxed);
 			}
 
 			ViewPacket.Tasks.OcclusionCull.Trigger();
