Fixed build break.

diff --git a/Engine/Source/Runtime/Renderer/Private/SceneVisibility.cpp b/Engine/Source/Runtime/Renderer/Private/SceneVisibility.cpp
--- a/Engine/Source/Runtime/Renderer/Private/SceneVisibility.cpp
+++ b/Engine/Source/Runtime/Renderer/Private/SceneVisibility.cpp
@@ -4140,7 +4140,9 @@ void FVisibilityTaskData::LaunchVisibilityTasks(const UE::Tasks::FTask& BeginIni
 		if (ViewPacket.ViewState)
 		{
 			SCOPE_CYCLE_COUNTER(STAT_DecompressPrecomputedOcclusion);
-			if (ViewPacket.View.PrecomputedVisibilityData = ViewPacket.ViewState->ResolvePrecomputedVisibilityData(ViewPacket.View, &Scene))
+			ViewPacket.View.PrecomputedVisibilityData = ViewPacket.ViewState->ResolvePrecomputedVisibilityData(ViewPacket.View, &Scene);
+
+			if (ViewPacket.View.PrecomputedVisibilityData)
 			{
 				SceneRenderer.bUsedPrecomputedVisibility = true;
 			}
