PR #11591: Fixed crash when freezing mobile renderer

#11591

diff --git a/Engine/Source/Runtime/Renderer/Private/SceneOcclusion.cpp b/Engine/Source/Runtime/Renderer/Private/SceneOcclusion.cpp
--- a/Engine/Source/Runtime/Renderer/Private/SceneOcclusion.cpp
+++ b/Engine/Source/Runtime/Renderer/Private/SceneOcclusion.cpp
@@ -1213,11 +1213,6 @@ static void AllocateOcclusionTests(FViewOcclusionQueriesPerView& QueriesPerView,
 				}
 			}
 
-			// Don't do primitive occlusion if we have a view parent or are frozen - only applicable to Debug & Development.
-#if !(UE_BUILD_SHIPPING || UE_BUILD_TEST)
-			ViewQuery.bFlushQueries &= (!ViewState->bIsFrozen);
-#endif
-
 			bBatchedQueries |= (View.IndividualOcclusionQueries.HasBatches() || View.GroupedOcclusionQueries.HasBatches() || ViewQuery.bFlushQueries);
 		}
 	}
