Adds a cvar r.MotionBlur.VisualizeCheckerboard to toggle the checkerboard rendering when visualizing motion blur.  Turning this off when visualizing the buffer can make it easier to screenshot compare.

#tests Editor, enabling the visualize motion blur show flag and toggling r.MotionBlur.VisualizeCheckerboard

diff --git a/Engine/Shaders/Private/MotionBlur/MotionBlurVisualize.usf b/Engine/Shaders/Private/MotionBlur/MotionBlurVisualize.usf
--- a/Engine/Shaders/Private/MotionBlur/MotionBlurVisualize.usf
+++ b/Engine/Shaders/Private/MotionBlur/MotionBlurVisualize.usf
@@ -23,6 +23,8 @@ SamplerState ColorSampler;
 SamplerState VelocitySampler;
 SamplerState DepthSampler;
 
+int CheckerboardEnabled;
+
 #if USE_POST_MOTION_BLUR_TRANSLUCENCY
 
 Texture2D PostMotionBlurTranslucencyTexture;
@@ -128,6 +130,8 @@ void MainPS(
 	float2 ScreenPosAtTileCenter = (PixelPosAtTileCenter - Velocity_ScreenPosToViewportBias) / Velocity_ScreenPosToViewportScale;
 
 	// World aligned checkerboards
+	OutColor = float4(0, 0, 0, 1);
+	if (CheckerboardEnabled)
 	{
 		float DeviceZ = DepthTexture.SampleLevel(SharedDepthSampler, UV, 0).r;
 		float SceneDepth = ConvertFromDeviceZ(DeviceZ);
diff --git a/Engine/Source/Runtime/Renderer/Private/PostProcess/PostProcessMotionBlur.cpp b/Engine/Source/Runtime/Renderer/Private/PostProcess/PostProcessMotionBlur.cpp
--- a/Engine/Source/Runtime/Renderer/Private/PostProcess/PostProcessMotionBlur.cpp
+++ b/Engine/Source/Runtime/Renderer/Private/PostProcess/PostProcessMotionBlur.cpp
@@ -67,6 +67,15 @@ namespace
 		0,
 		TEXT("For projects with motion blur enabled, this allows motion blur to be enabled even while in VR."));
 
+	TAutoConsoleVariable<int32> CVarVisualizeMotionBlurEnableCheckerboard(
+		TEXT("r.MotionBlur.VisualizeCheckerboard"),
+		1,
+		TEXT("Enables checkerboard visualization when Visualize Motion Blur show flag is enabled.\n")
+		TEXT("0: off\n")
+		TEXT("1: on (default)\n"),
+		ECVF_RenderThreadSafe
+	);
+
 	FMatrix GetPreviousWorldToClipMatrix(const FViewInfo& View)
 	{
 		if (View.Family->EngineShowFlags.CameraInterpolation)
@@ -515,6 +524,8 @@ class FMotionBlurVisualizePS : public FMotionBlurShader
 		SHADER_PARAMETER_SAMPLER(SamplerState, VelocitySampler)
 		SHADER_PARAMETER_SAMPLER(SamplerState, DepthSampler)
 
+		SHADER_PARAMETER(int, CheckerboardEnabled)
+
 		RENDER_TARGET_BINDING_SLOTS()
 	END_SHADER_PARAMETER_STRUCT()
 };
@@ -1132,6 +1143,7 @@ FScreenPassTextureSlice AddVisualizeMotionBlurPass(FRDGBuilder& GraphBuilder, co
 	PassParameters->VelocitySampler = GetMotionBlurVelocitySampler();
 	PassParameters->DepthSampler = GetMotionBlurVelocitySampler();
 	PassParameters->RenderTargets[0] = Output.GetRenderTargetBinding();
+	PassParameters->CheckerboardEnabled = CVarVisualizeMotionBlurEnableCheckerboard.GetValueOnRenderThread();
 
 	TShaderMapRef<FMotionBlurVisualizePS> PixelShader(View.ShaderMap);
 
